#ttp 316639 - UE4: RENDERING: DrawDenormalizedQuad should reuse static index and vertexbuffer

#summary UV offsets are now correctly calculated. This was previously blocking shadow game

diff --git a/Engine/Shaders/AtmospherePrecompute.usf b/Engine/Shaders/AtmospherePrecompute.usf
--- a/Engine/Shaders/AtmospherePrecompute.usf
+++ b/Engine/Shaders/AtmospherePrecompute.usf
@@ -48,8 +48,7 @@ void MainVS(
 	out FAtmosphereVSOutput Out
 	)
 {
-	Out.OutPosition = InPosition;
-	Out.OutTexCoord = InTexCoord;
+	DrawRectangle(InPosition, InTexCoord, Out.OutPosition, Out.OutTexCoord );
 }
 
 // pixel shader entry point
diff --git a/Engine/Shaders/Common.usf b/Engine/Shaders/Common.usf
--- a/Engine/Shaders/Common.usf
+++ b/Engine/Shaders/Common.usf
@@ -928,6 +928,23 @@ struct FWriteToSliceGeometryOutput
 // because we are using an extra interpolator for this second world position, and in < SM4 there may not be enough
 #define USE_WORLD_POSITION_EXCLUDING_SHADER_OFFSETS	(NEEDS_WORLD_POSITION_EXCLUDING_SHADER_OFFSETS && FEATURE_LEVEL >= FEATURE_LEVEL_SM4)
 
+/** Used for calculating vertex positions and UVs when drawing with DrawRectangle */
+void DrawRectangle( in float4 InPosition, in float2 InTexCoord, out float4 OutPosition, out float2 OutTexCoord)
+{
+	OutPosition = InPosition;
+	OutPosition.xy = -1.0f + 2.0f * (DrawRectangleParameters.PosScaleBias.zw + (InPosition.xy * DrawRectangleParameters.PosScaleBias.xy)) * DrawRectangleParameters.InvTargetSizeAndTextureSize.xy;
+	OutPosition.xy *= float2( 1, -1 );
+	OutTexCoord.xy = (DrawRectangleParameters.UVScaleBias.zw + (InTexCoord.xy * DrawRectangleParameters.UVScaleBias.xy)) * DrawRectangleParameters.InvTargetSizeAndTextureSize.zw;
+}
+
+/** Used for calculating vertex positions when drawing with DrawRectangle */
+void DrawRectangle( in float4 InPosition, out float4 OutPosition)
+{
+	OutPosition = InPosition;
+	OutPosition.xy = -1.0f + 2.0f * (DrawRectangleParameters.PosScaleBias.zw + (InPosition.xy * DrawRectangleParameters.PosScaleBias.xy)) * DrawRectangleParameters.InvTargetSizeAndTextureSize.xy;
+	OutPosition.xy *= float2( 1, -1 );
+}
+
 //Since some platforms don't remove Nans in saturate calls, 
 //SafeSaturate function will remove nan/inf.    
 //Can be expensive, only call when there's a good reason to expect Nans.
diff --git a/Engine/Shaders/DeferredLightVertexShaders.usf b/Engine/Shaders/DeferredLightVertexShaders.usf
--- a/Engine/Shaders/DeferredLightVertexShaders.usf
+++ b/Engine/Shaders/DeferredLightVertexShaders.usf
@@ -20,9 +20,8 @@ void DirectionalVertexMain(
 	out float4 OutPosition : SV_POSITION
 	)
 {	
-	OutPosition = float4(InPosition, 0, 1);
-	OutTexCoord = InUV;
-	OutScreenVector = mul(float4(InPosition, 1, 0), View.ScreenToTranslatedWorld).xyz;
+	DrawRectangle(float4(InPosition.xy, 0, 1), InUV, OutPosition, OutTexCoord);
+	OutScreenVector = mul(float4(OutPosition.xy, 1, 0), View.ScreenToTranslatedWorld).xyz;
 }
 
 /** Vertex shader for rendering a point or spot light using approximate bounding geometry. */
diff --git a/Engine/Shaders/FXAAShader.usf b/Engine/Shaders/FXAAShader.usf
--- a/Engine/Shaders/FXAAShader.usf
+++ b/Engine/Shaders/FXAAShader.usf
@@ -51,9 +51,9 @@ void FxaaVS(
 	out float4 OutPosition : SV_POSITION
 	)
 {
-	OutPosition = InPosition;
-	OutCornersUV = float4(UV.xy - 0.5f * fxaaQualityRcpFrame, UV.xy + 0.5f * fxaaQualityRcpFrame);
-	OutCenterUV = UV.xy;
+	DrawRectangle(InPosition, UV, OutPosition, OutCenterUV);
+	OutCornersUV = float4(OutCenterUV.xy - 0.5f * fxaaQualityRcpFrame, OutCenterUV.xy + 0.5f * fxaaQualityRcpFrame);
+
 }
 
 
diff --git a/Engine/Shaders/FilterVertexShader.usf b/Engine/Shaders/FilterVertexShader.usf
--- a/Engine/Shaders/FilterVertexShader.usf
+++ b/Engine/Shaders/FilterVertexShader.usf
@@ -18,6 +18,8 @@ void Main(
 	out float4 OutPosition : SV_POSITION
 	)
 {
+	DrawRectangle( InPosition, UV, OutPosition, UV);
+
 	int SampleIndex;
 	for(SampleIndex = 0;SampleIndex < NUM_SAMPLES - 1;SampleIndex += 2)
 	{
@@ -31,7 +33,6 @@ void Main(
 		OutOffsetUVs[SampleIndex] = UV.xy + OffsetUVUV.xy;
 	}
 
-	OutPosition = InPosition;
 	OutUV = UV;
 }
 
@@ -45,11 +46,13 @@ void Main(
 	out float4 OutPosition : SV_POSITION
 	)
 {
+	DrawRectangle( InPosition, UV, OutPosition, UV);
+
 	for(int OffsetIndex = 0;OffsetIndex < (NUM_SAMPLES + 1) / 2;OffsetIndex++)
 	{
 		OutOffsetUVs[OffsetIndex] = UV.xyyx + SampleOffsets[OffsetIndex];
 	}
-	OutPosition = InPosition;
+
 	OutUV = UV;
 }
 
diff --git a/Engine/Shaders/GPUBenchmark.usf b/Engine/Shaders/GPUBenchmark.usf
--- a/Engine/Shaders/GPUBenchmark.usf
+++ b/Engine/Shaders/GPUBenchmark.usf
@@ -19,8 +19,7 @@ void MainBenchmarkVS(
 	out float4 OutPosition : SV_POSITION
 	) 
 {
-	OutPosition = InPosition;
-	OutUV = UV;
+	DrawRectangle(InPosition, UV, OutPosition, OutUV);
 }
 
 // pixel shader entry point
diff --git a/Engine/Shaders/GammaCorrection.usf b/Engine/Shaders/GammaCorrection.usf
--- a/Engine/Shaders/GammaCorrection.usf
+++ b/Engine/Shaders/GammaCorrection.usf
@@ -14,8 +14,7 @@ void MainVS(
 	out float4 OutPosition : SV_POSITION
 	)
 {
-	OutPosition = InPosition;
-	OutTexCoord = InTexCoord;
+	DrawRectangle(InPosition, InTexCoord, OutPosition, OutTexCoord);
 }
 
 half3 ColorScale;
diff --git a/Engine/Shaders/LightShaftShader.usf b/Engine/Shaders/LightShaftShader.usf
--- a/Engine/Shaders/LightShaftShader.usf
+++ b/Engine/Shaders/LightShaftShader.usf
@@ -14,10 +14,8 @@ void DownsampleLightShaftsVertexMain(
 	out float4 OutPosition : SV_POSITION
 	)
 {	
-	OutPosition = float4(InPosition, 0, 1);
-	OutTexCoord = InUV;
-
-	OutScreenVector = mul(float4(InPosition, 1, 0), View.ScreenToTranslatedWorld).xyz;
+	DrawRectangle(float4(InPosition.xy, 0, 1), InUV, OutPosition, OutTexCoord);
+	OutScreenVector = mul(float4(OutPosition.xy, 1, 0), View.ScreenToTranslatedWorld).xyz;
 }
 
 /** Origin in texture coordinates in xy. */
diff --git a/Engine/Shaders/OneColorShader.usf b/Engine/Shaders/OneColorShader.usf
--- a/Engine/Shaders/OneColorShader.usf
+++ b/Engine/Shaders/OneColorShader.usf
@@ -17,12 +17,21 @@ struct FOneColorVertexOutput
 	float4 Position : SV_POSITION;
 };
 
+
+#ifndef USING_NDC_POSITIONS
+	#define USING_NDC_POSITIONS 1
+#endif
+
 void MainVertexShader(
 	in float4 InPosition : ATTRIBUTE0,
 	out FOneColorVertexOutput Output
 	)
 {
 	Output.Position = InPosition;
+
+#if !USING_NDC_POSITIONS
+	DrawRectangle( InPosition, Output.Position);
+#endif
 }
 
 /*=============================================================================
diff --git a/Engine/Shaders/PostProcessBloom.usf b/Engine/Shaders/PostProcessBloom.usf
--- a/Engine/Shaders/PostProcessBloom.usf
+++ b/Engine/Shaders/PostProcessBloom.usf
@@ -15,8 +15,8 @@ void MainPostprocessCommonVS(
 	out float4 OutPosition : SV_POSITION
 	)
 {
-	OutPosition = InPosition;
-	OutTexCoord = float4(InTexCoord, InPosition.xy);
+	DrawRectangle(InPosition, InTexCoord, OutPosition, OutTexCoord.xy);
+	OutTexCoord.zw =  OutPosition.xy;
 }
 
 // vertex shader entry point
@@ -28,8 +28,7 @@ void MainVS(
 	out float4 OutPosition : SV_POSITION
 	)
 {
-	OutPosition = InPosition;
-	OutTexCoord = float4(InTexCoord, InPosition.xy);
+	MainPostprocessCommonVS(InPosition, InTexCoord, OutTexCoord, OutPosition);
 
 #if FEATURE_LEVEL >= FEATURE_LEVEL_SM5
 	// texture can be GWhiteTexture which is 1x1. It's important we don't read outside bounds.
diff --git a/Engine/Shaders/PostProcessDownsample.usf b/Engine/Shaders/PostProcessDownsample.usf
--- a/Engine/Shaders/PostProcessDownsample.usf
+++ b/Engine/Shaders/PostProcessDownsample.usf
@@ -15,8 +15,7 @@ void MainDownsampleVS(
 	out float4 OutPosition : SV_POSITION
 	) 
 {
-	OutPosition = InPosition;
-	OutUV = UV;
+	DrawRectangle(InPosition, UV, OutPosition, OutUV);
 }
 
 // pixel shader entry point
diff --git a/Engine/Shaders/PostProcessHMD.usf b/Engine/Shaders/PostProcessHMD.usf
--- a/Engine/Shaders/PostProcessHMD.usf
+++ b/Engine/Shaders/PostProcessHMD.usf
@@ -42,6 +42,20 @@ float2 HMDWarp(float2 in01)
 }
 #endif
 
+
+// vertex shader entry point
+void MainVS(
+	in float4 InPosition : ATTRIBUTE0,
+	in float2 InTexCoord : ATTRIBUTE1,
+	out float4 OutTexCoord : TEXCOORD0,
+	out float4 OutPosition : SV_POSITION
+	)
+{
+	OutPosition = InPosition;
+	OutTexCoord = float4(InTexCoord, InPosition.xy);
+}
+
+
 void MainPS(in float4 UVAndScreenPos : TEXCOORD0, out float4 OutColor : SV_Target0)
 {
     float2 InUV = UVAndScreenPos.xy;
diff --git a/Engine/Shaders/PostProcessMaterialShaders.usf b/Engine/Shaders/PostProcessMaterialShaders.usf
--- a/Engine/Shaders/PostProcessMaterialShaders.usf
+++ b/Engine/Shaders/PostProcessMaterialShaders.usf
@@ -13,8 +13,8 @@ void MainVS(
 	out float4 OutPosition : SV_POSITION
 	)
 {
-	OutPosition = InPosition;
-	OutScreenPosition = InPosition;
+	DrawRectangle(InPosition,  OutPosition);
+	OutScreenPosition = OutPosition;
 }
 
 void MainPS(
diff --git a/Engine/Shaders/PostProcessMobile.usf b/Engine/Shaders/PostProcessMobile.usf
--- a/Engine/Shaders/PostProcessMobile.usf
+++ b/Engine/Shaders/PostProcessMobile.usf
@@ -61,9 +61,8 @@ void SunMaskVS_ES2(
 	out float4 OutPosition : SV_POSITION
 	)
 {
-	OutPosition = InPosition;
-	OutUVPos.xy = InTexCoord;
-	OutUVPos.zw = InPosition.xy;
+	DrawRectangle(InPosition, InTexCoord, OutPosition, OutUVPos.xy);
+	OutUVPos.zw = OutPosition.xy;
 }
 
 void SunMaskPS_ES2(
@@ -114,11 +113,13 @@ void BloomVS_ES2(
 	out float4 OutPosition : SV_POSITION
 	)
 {
-	OutPosition = InPosition;
-	OutTexCoords[0] = InTexCoord + PostprocessInput0Size.zw * float2(-1, -1);
-	OutTexCoords[1] = InTexCoord + PostprocessInput0Size.zw * float2( 1, -1);
-	OutTexCoords[2] = InTexCoord + PostprocessInput0Size.zw * float2(-1,  1);
-	OutTexCoords[3] = InTexCoord + PostprocessInput0Size.zw * float2( 1,  1);
+	float2 TransformedUV;
+	DrawRectangle(InPosition, InTexCoord, OutPosition, TransformedUV);
+
+	OutTexCoords[0] = TransformedUV + PostprocessInput0Size.zw * float2(-1, -1);
+	OutTexCoords[1] = TransformedUV + PostprocessInput0Size.zw * float2( 1, -1);
+	OutTexCoords[2] = TransformedUV + PostprocessInput0Size.zw * float2(-1,  1);
+	OutTexCoords[3] = TransformedUV + PostprocessInput0Size.zw * float2( 1,  1);
 }
 
 void BloomPS_ES2(
@@ -193,7 +194,8 @@ void BloomSmallVS_ES2(
 	out float4 OutPosition : SV_POSITION
 	)
 {
-	OutPosition = InPosition;
+
+	DrawRectangle(InPosition, InTexCoord, OutPosition, InTexCoord);
 	
 	float Start = 2.0/14.0;
 	float Scale = 0.66 * 4.0 * 2.0;
@@ -280,7 +282,7 @@ void BloomDownVS_ES2(
 	out float4 OutPosition : SV_POSITION
 	)
 {
-	OutPosition = InPosition;
+	DrawRectangle(InPosition, InTexCoord, OutPosition, InTexCoord);
 
 	float Start = 2.0/14.0;
 	float Scale = BloomDownScale;
@@ -357,7 +359,7 @@ void BloomUpVS_ES2(
 	out float4 OutPosition : SV_POSITION
 	)
 {
-	OutPosition = InPosition;
+	DrawRectangle(InPosition, InTexCoord, OutPosition, InTexCoord);
 
 	float Start;
 	float Scale;
@@ -452,7 +454,8 @@ void DofNearVS_ES2(
 	out float4 OutPosition : SV_POSITION
 	)
 {
-	OutPosition = InPosition;
+	DrawRectangle(InPosition, InTexCoord, OutPosition, InTexCoord);
+
 	OutTexCoords2 = InTexCoord;
 	OutTexCoords4[0].xy = InTexCoord + PostprocessInput0Size.zw * float2(-0.5,-1.0);
 	OutTexCoords4[0].zw = InTexCoord + PostprocessInput0Size.zw * float2( 1.0,-0.5);
@@ -523,9 +526,10 @@ void DofDownVS_ES2(
 	out float4 OutPosition : SV_POSITION
 	)
 {
-	OutPosition = InPosition;
+	DrawRectangle(InPosition, InTexCoord, OutPosition, InTexCoord);
+
 	// Near position fixed to use UV based out output position.
-	OutTexCoords[0] = InPosition.xy * float2(0.5,-0.5) + 0.5;
+	OutTexCoords[0] = OutPosition.xy * float2(0.5,-0.5) + 0.5;
 	// Other source UVs based on possible non-full texture.
 	OutTexCoords[1] = InTexCoord + PostprocessInput0Size.zw * float2(-0.5, -0.5);
 	OutTexCoords[2] = InTexCoord + PostprocessInput0Size.zw * float2( 0.5, -0.5);
@@ -650,7 +654,8 @@ void DofBlurVS_ES2(
 	out float4 OutPosition : SV_POSITION
 	)
 {
-	OutPosition = InPosition;
+	DrawRectangle(InPosition, InTexCoord, OutPosition, InTexCoord);
+
     OutTexCoords[0] = InTexCoord.xy;
     OutTexCoords[1] = InTexCoord.xy + float2(DOF_1) * PostprocessInput0Size.zw;
     OutTexCoords[2] = InTexCoord.xy + float2(DOF_2) * PostprocessInput0Size.zw;
@@ -748,7 +753,8 @@ void SunAlphaVS_ES2(
 	out float4 OutPosition : SV_POSITION
 	)
 {
-	OutPosition = InPosition;
+	DrawRectangle(InPosition, InTexCoord, OutPosition, InTexCoord);
+
     OutTexCoords[0] = SunShaftRect(InPosition.xy, 1.0 - SUN_P0 * SUN_M);
     OutTexCoords[1] = SunShaftRect(InPosition.xy, 1.0 - SUN_P1 * SUN_M);
     OutTexCoords[2] = SunShaftRect(InPosition.xy, 1.0 - SUN_P2 * SUN_M);
@@ -805,7 +811,8 @@ void SunBlurVS_ES2(
 	out float4 OutPosition : SV_POSITION
 	)
 {
-	OutPosition = InPosition;
+	DrawRectangle(InPosition, InTexCoord, OutPosition, InTexCoord);
+
     OutTexCoords[0] = SunShaftRect(InPosition.xy, 1.0 - SUN_P0 * SUN_M);
     OutTexCoords[1] = SunShaftRect(InPosition.xy, 1.0 - SUN_P1 * SUN_M);
     OutTexCoords[2] = SunShaftRect(InPosition.xy, 1.0 - SUN_P2 * SUN_M);
@@ -849,9 +856,10 @@ void SunMergeVS_ES2(
 	out float4 OutPosition : SV_POSITION
 	)
 {
-	OutPosition = InPosition;
+	DrawRectangle(InPosition, InTexCoord, OutPosition, InTexCoord);
+
 	OutTexCoordVignette.xy = InTexCoord.xy;
-	OutTexCoordVignette.zw =  VignetteSpace(InPosition.xy);
+	OutTexCoordVignette.zw =  VignetteSpace(OutPosition.xy);
 
 	float Start;
 	float Scale;
@@ -959,9 +967,9 @@ void SunMergeSmallVS_ES2(
 	out float4 OutPosition : SV_POSITION
 	)
 {
-	OutPosition = InPosition;
+	DrawRectangle(InPosition, InTexCoord, OutPosition, InTexCoord);
 	OutTexCoordVignette.xy = InTexCoord.xy;
-	OutTexCoordVignette.zw =  VignetteSpace(InPosition.xy);
+	OutTexCoordVignette.zw =  VignetteSpace(OutPosition.xy);
 
 	float Start;
 	float Scale;
@@ -1039,8 +1047,7 @@ void SunAvgVS_ES2(
 	out float4 OutPosition : SV_POSITION
 	)
 {
-	OutPosition = InPosition;
-	OutTexCoord.xy = InTexCoord.xy;
+	DrawRectangle(InPosition, InTexCoord, OutPosition, OutTexCoord);
 }
 
 void SunAvgPS_ES2(
@@ -1068,7 +1075,7 @@ void AaVS_ES2(
 	out float4 OutPosition : SV_POSITION
 	)
 {
-	OutPosition = InPosition;
+	DrawRectangle(InPosition, InTexCoord, OutPosition, InTexCoord);
 
 	// AA texture coordinates.
 	float AaX = View.TemporalAAParams.x;
diff --git a/Engine/Shaders/PostProcessSceneColorFringe.usf b/Engine/Shaders/PostProcessSceneColorFringe.usf
--- a/Engine/Shaders/PostProcessSceneColorFringe.usf
+++ b/Engine/Shaders/PostProcessSceneColorFringe.usf
@@ -30,7 +30,7 @@ void MainVS(
 	out float4 OutPosition : SV_POSITION
 	)
 {
-	OutPosition = InPosition;
+	DrawRectangle(InPosition, InTexCoord, OutPosition, InTexCoord);
 
 	float2 ScreenPos = UVToScreenPos(InTexCoord);
 	 
diff --git a/Engine/Shaders/PostProcessTonemap.usf b/Engine/Shaders/PostProcessTonemap.usf
--- a/Engine/Shaders/PostProcessTonemap.usf
+++ b/Engine/Shaders/PostProcessTonemap.usf
@@ -115,8 +115,8 @@ void MainVS(
 	out float4 OutPosition : SV_POSITION
 	)
 {
-	OutPosition = InPosition;
-	OutTexCoord = float4(InTexCoord, InPosition.xy);
+	DrawRectangle(InPosition, InTexCoord, OutPosition, OutTexCoord.xy);
+	OutTexCoord = float4(OutTexCoord.xy, OutPosition.xy);
 
 #if FEATURE_LEVEL >= FEATURE_LEVEL_SM5
 	// texture can be GWhiteTexture which is 1x1. It's important we don't read outside bounds.
@@ -289,15 +289,15 @@ void MainVS_ES2(
 	out float4 OutPosition : SV_POSITION
 	)
 {
-	OutPosition = InPosition;
+	DrawRectangle(InPosition, InTexCoord, OutPosition, OutTexCoord.xy);
+	OutTexCoord = float4(OutTexCoord.xy, OutPosition.xy);
 
-	OutTexCoord = float4(InTexCoord, InPosition.xy);
 	// Avoiding a permutation.
 	if(GrainRandomFull.z > 0.0)
 	{
 		// Framebuffer fetch hardware uses the standard possibly non-full rectangle.
 		// Other hardware gets a full texture source.
-		OutTexCoord.xy = InPosition.xy * float2(0.5,-0.5) + 0.5;
+		OutTexCoord.xy = OutPosition.xy * float2(0.5,-0.5) + 0.5;
 	}
 
 	#if ES2_PROFILE && COMPILER_GLSL_ES2
@@ -308,7 +308,7 @@ void MainVS_ES2(
 	// Fine adjustment is inside the possible non-full viewport in the full resolution texture.
 	OutFineDofGrain.xy = OutTexCoord.xy + PostprocessInput0Size.zw * float2(-0.5,0.5);
 	// Want grain and a second UV based on the knowledge that the source texture has a full viewport.
-	OutFullViewUV.xy = InPosition.xy * float2(0.5,-0.5) + 0.5;
+	OutFullViewUV.xy = OutPosition.xy * float2(0.5,-0.5) + 0.5;
 	#if ES2_PROFILE && COMPILER_GLSL_ES2
 		// This is currently the last pass, so flip the texture on V to match D3D
 		OutFullViewUV.y = 1.0 - OutFullViewUV.y;
diff --git a/Engine/Shaders/ReflectionEnvironmentShaders.usf b/Engine/Shaders/ReflectionEnvironmentShaders.usf
--- a/Engine/Shaders/ReflectionEnvironmentShaders.usf
+++ b/Engine/Shaders/ReflectionEnvironmentShaders.usf
@@ -24,9 +24,8 @@ void CopyToCubeFaceVS(
 	out FCopyToCubeFaceVSOutput Out
 	)
 {	
-	Out.UV = InUV;
-	Out.Position = float4(InPosition, 0, 1);
-	Out.ScreenVector = mul(float4(InPosition, 1, 0), View.ScreenToTranslatedWorld).xyz;
+	DrawRectangle(float4(InPosition.xy, 0, 1), InUV, Out.Position, Out.UV);
+	Out.ScreenVector = mul(float4(Out.Position.xy, 1, 0), View.ScreenToTranslatedWorld).xyz;
 }
 
 struct FRouteToCubeFaceGSOut
diff --git a/Engine/Shaders/ScreenVertexShader.usf b/Engine/Shaders/ScreenVertexShader.usf
--- a/Engine/Shaders/ScreenVertexShader.usf
+++ b/Engine/Shaders/ScreenVertexShader.usf
@@ -12,8 +12,7 @@ void Main(
 	out FScreenVertexOutput Output
 	)
 {
-	Output.Position = float4( InPosition, 0, 1 );
-    Output.UV = InUV;
+	DrawRectangle( float4( InPosition, 0, 1 ), InUV, Output.Position, Output.UV);
 }
 
 
diff --git a/Engine/Shaders/ShaderComplexityApplyVertexShader.usf b/Engine/Shaders/ShaderComplexityApplyVertexShader.usf
--- a/Engine/Shaders/ShaderComplexityApplyVertexShader.usf
+++ b/Engine/Shaders/ShaderComplexityApplyVertexShader.usf
@@ -13,6 +13,5 @@ void Main(
 	out float4 OutPosition : SV_POSITION
 	)
 {
-	OutPosition = InPosition;
-	OutTexCoord = InTexCoord;
+	DrawRectangle(InPosition, InTexCoord, OutPosition, OutTexCoord);
 }
diff --git a/Engine/Source/Runtime/Engine/Private/OneColorShader.cpp b/Engine/Source/Runtime/Engine/Private/OneColorShader.cpp
--- a/Engine/Source/Runtime/Engine/Private/OneColorShader.cpp
+++ b/Engine/Source/Runtime/Engine/Private/OneColorShader.cpp
@@ -4,7 +4,8 @@
 #include "EnginePrivate.h"
 #include "OneColorShader.h"
 
-IMPLEMENT_SHADER_TYPE(,FOneColorVS,TEXT("OneColorShader"),TEXT("MainVertexShader"),SF_Vertex);
+IMPLEMENT_SHADER_TYPE(template<>,TOneColorVS<true>,TEXT("OneColorShader"),TEXT("MainVertexShader"),SF_Vertex);
+IMPLEMENT_SHADER_TYPE(template<>,TOneColorVS<false>,TEXT("OneColorShader"),TEXT("MainVertexShader"),SF_Vertex);
 IMPLEMENT_SHADER_TYPE(,FOneColorPS,TEXT("OneColorShader"),TEXT("MainPixelShader"),SF_Pixel);
 // Compiling a version for every number of MRT's
 // On AMD PC hardware, outputting to a color index in the shader without a matching render target set has a significant performance hit
diff --git a/Engine/Source/Runtime/Engine/Public/OneColorShader.h b/Engine/Source/Runtime/Engine/Public/OneColorShader.h
--- a/Engine/Source/Runtime/Engine/Public/OneColorShader.h
+++ b/Engine/Source/Runtime/Engine/Public/OneColorShader.h
@@ -10,14 +10,24 @@
 /**
  * Vertex shader for rendering a single, constant color.
  */
-class FOneColorVS : public FGlobalShader
+template<bool bUsingNDCPositions=true>
+class TOneColorVS : public FGlobalShader
 {
-	DECLARE_EXPORTED_SHADER_TYPE(FOneColorVS,Global,ENGINE_API);
+	DECLARE_EXPORTED_SHADER_TYPE(TOneColorVS, Global, ENGINE_API);
+
+	/** Default constructor. */
+	TOneColorVS() {}
+
 public:
-	FOneColorVS( )	{ }
-	FOneColorVS(const ShaderMetaType::CompiledShaderInitializerType& Initializer)
-	:	FGlobalShader( Initializer )
+
+	TOneColorVS(const ShaderMetaType::CompiledShaderInitializerType& Initializer)
+		: FGlobalShader(Initializer)
+	{}
+
+	static void ModifyCompilationEnvironment(EShaderPlatform Platform, FShaderCompilerEnvironment& OutEnvironment)
 	{
+		FGlobalShader::ModifyCompilationEnvironment(Platform, OutEnvironment);
+		OutEnvironment.SetDefine(TEXT("USING_NDC_POSITIONS"), (uint32)(bUsingNDCPositions ? 1 : 0));
 	}
 
 	static bool ShouldCache(EShaderPlatform Platform)
diff --git a/Engine/Source/Runtime/OpenGLDrv/Private/OpenGLDrv.cpp b/Engine/Source/Runtime/OpenGLDrv/Private/OpenGLDrv.cpp
--- a/Engine/Source/Runtime/OpenGLDrv/Private/OpenGLDrv.cpp
+++ b/Engine/Source/Runtime/OpenGLDrv/Private/OpenGLDrv.cpp
@@ -401,7 +401,7 @@ void FOpenGLDynamicRHI::IssueLongGPUTask()
 		RHISetDepthStencilState(TStaticDepthStencilState<false,CF_Always>::GetRHI(), 0);
 		RHISetRasterizerState(TStaticRasterizerState<FM_Solid,CM_None>::GetRHI());
 
-		TShaderMapRef<FOneColorVS> VertexShader(GetGlobalShaderMap());
+		TShaderMapRef<TOneColorVS<true> > VertexShader(GetGlobalShaderMap());
 		TShaderMapRef<FOpenGLRHILongGPUTaskPS> PixelShader(GetGlobalShaderMap());
 
 		SetGlobalBoundShaderState(LongGPUTaskBoundShaderState, GOpenGLVector4VertexDeclaration.VertexDeclarationRHI, *VertexShader, *PixelShader, 0);
diff --git a/Engine/Source/Runtime/Renderer/Private/AtmosphereRendering.cpp b/Engine/Source/Runtime/Renderer/Private/AtmosphereRendering.cpp
--- a/Engine/Source/Runtime/Renderer/Private/AtmosphereRendering.cpp
+++ b/Engine/Source/Runtime/Renderer/Private/AtmosphereRendering.cpp
@@ -1049,7 +1049,7 @@ FIntPoint FAtmosphericFogSceneInfo::GetTextureSize()
 	return AtmosphereTextures->AtmosphereTransmittance->GetDesc().Extent;
 }
 
-void FAtmosphericFogSceneInfo::DrawQuad(const FIntRect& ViewRect)
+void FAtmosphericFogSceneInfo::DrawQuad(const FIntRect& ViewRect, FShader* VertexShader)
 {
 	// Draw a quad mapping scene color to the view's render target
 	DrawRectangle( 
@@ -1059,6 +1059,7 @@ void FAtmosphericFogSceneInfo::DrawQuad(const FIntRect& ViewRect)
 		ViewRect.Width(), ViewRect.Height(),
 		ViewRect.Size(),
 		ViewRect.Size(),
+		VertexShader,
 		EDRF_UseTriangleOptimization);
 }
 
@@ -1091,7 +1092,7 @@ void FAtmosphericFogSceneInfo::RenderAtmosphereShaders(const FViewInfo& View, co
 			static FGlobalBoundShaderState BoundShaderState;
 			SetGlobalBoundShaderState( BoundShaderState, GFilterVertexDeclaration.VertexDeclarationRHI, *VertexShader, *PixelShader);
 			PixelShader->SetParameters(View);
-			DrawQuad(ViewRect);
+			DrawQuad(ViewRect, *VertexShader);
 
 			RHICopyToResolveTarget(DestRenderTarget.TargetableTexture, DestRenderTarget.ShaderResourceTexture, true, FResolveParams() );
 		}
@@ -1109,7 +1110,7 @@ void FAtmosphericFogSceneInfo::RenderAtmosphereShaders(const FViewInfo& View, co
 
 			PixelShader->SetParameters(AtmosphereTextures);
 
-			DrawQuad(ViewRect);
+			DrawQuad(ViewRect, *VertexShader);
 
 			RHICopyToResolveTarget(DestRenderTarget.TargetableTexture, DestRenderTarget.ShaderResourceTexture, true, FResolveParams() );
 		}
@@ -1136,7 +1137,7 @@ void FAtmosphericFogSceneInfo::RenderAtmosphereShaders(const FViewInfo& View, co
 				GetLayerValue(Layer, r, DhdH);
 				GeometryShader->SetParameters(Layer);
 				PixelShader->SetParameters(View, r, DhdH, AtmosphereTextures);
-				DrawQuad(ViewRect);
+				DrawQuad(ViewRect, *VertexShader);
 
 				if (Atmosphere3DTextureIndex == Component->PrecomputeParams.InscatterAltitudeSampleNum - 1)
 				{
@@ -1176,7 +1177,7 @@ void FAtmosphericFogSceneInfo::RenderAtmosphereShaders(const FViewInfo& View, co
 				GetLayerValue(Layer, r, DhdH);
 				GeometryShader->SetParameters(Layer);
 				PixelShader->SetParameters(View, r, DhdH, Layer, AtmosphereTextures);
-				DrawQuad(ViewRect);
+				DrawQuad(ViewRect, *VertexShader);
 
 				if (Atmosphere3DTextureIndex == Component->PrecomputeParams.InscatterAltitudeSampleNum - 1)
 				{
@@ -1204,7 +1205,7 @@ void FAtmosphericFogSceneInfo::RenderAtmosphereShaders(const FViewInfo& View, co
 				GetLayerValue(Layer, r, DhdH);
 				GeometryShader->SetParameters(Layer);
 				PixelShader->SetParameters(View, r, DhdH, AtmoshpereOrder == 2 ? 1.f : 0.f, AtmosphereTextures);
-				DrawQuad(ViewRect);
+				DrawQuad(ViewRect, *VertexShader);
 
 				if (Atmosphere3DTextureIndex == Component->PrecomputeParams.InscatterAltitudeSampleNum - 1)
 				{
@@ -1226,7 +1227,7 @@ void FAtmosphericFogSceneInfo::RenderAtmosphereShaders(const FViewInfo& View, co
 
 			PixelShader->SetParameters(View, AtmoshpereOrder == 2 ? 1.f : 0.f, AtmosphereTextures);
 
-			DrawQuad(ViewRect);
+			DrawQuad(ViewRect, *VertexShader);
 
 			RHICopyToResolveTarget(DestRenderTarget.TargetableTexture, DestRenderTarget.ShaderResourceTexture, true, FResolveParams() );
 		}
@@ -1251,7 +1252,7 @@ void FAtmosphericFogSceneInfo::RenderAtmosphereShaders(const FViewInfo& View, co
 				GetLayerValue(Layer, r, DhdH);
 				GeometryShader->SetParameters(Layer);
 				PixelShader->SetParameters(View, r, DhdH, AtmoshpereOrder == 2 ? 1.f : 0.f, AtmosphereTextures);
-				DrawQuad(ViewRect);
+				DrawQuad(ViewRect, *VertexShader);
 
 				if (Atmosphere3DTextureIndex == Component->PrecomputeParams.InscatterAltitudeSampleNum - 1)
 				{
@@ -1275,7 +1276,7 @@ void FAtmosphericFogSceneInfo::RenderAtmosphereShaders(const FViewInfo& View, co
 
 			PixelShader->SetParameters(AtmosphereTextures);
 
-			DrawQuad(ViewRect);
+			DrawQuad(ViewRect, *VertexShader);
 
 			RHICopyToResolveTarget(DestRenderTarget.TargetableTexture, DestRenderTarget.ShaderResourceTexture, true, FResolveParams() );
 
@@ -1304,7 +1305,7 @@ void FAtmosphericFogSceneInfo::RenderAtmosphereShaders(const FViewInfo& View, co
 				GetLayerValue(Layer, r, DhdH);
 				GeometryShader->SetParameters(Layer);
 				PixelShader->SetParameters(View, r, DhdH, Layer, AtmosphereTextures);
-				DrawQuad(ViewRect);
+				DrawQuad(ViewRect, *VertexShader);
 
 				if (Atmosphere3DTextureIndex == Component->PrecomputeParams.InscatterAltitudeSampleNum - 1)
 				{
@@ -1335,7 +1336,7 @@ void FAtmosphericFogSceneInfo::RenderAtmosphereShaders(const FViewInfo& View, co
 				GetLayerValue(Layer, r, DhdH);
 				GeometryShader->SetParameters(Layer);
 				PixelShader->SetParameters(r, DhdH, Layer, AtmosphereTextures);
-				DrawQuad(ViewRect);
+				DrawQuad(ViewRect, *VertexShader);
 
 				if (Atmosphere3DTextureIndex == Component->PrecomputeParams.InscatterAltitudeSampleNum - 1)
 				{
@@ -1363,7 +1364,7 @@ void FAtmosphericFogSceneInfo::RenderAtmosphereShaders(const FViewInfo& View, co
 				GetLayerValue(Layer, r, DhdH);
 				GeometryShader->SetParameters(Layer);
 				PixelShader->SetParameters(r, DhdH, Layer, AtmosphereTextures);
-				DrawQuad(ViewRect);
+				DrawQuad(ViewRect, *VertexShader);
 
 				if (Atmosphere3DTextureIndex == Component->PrecomputeParams.InscatterAltitudeSampleNum - 1)
 				{
diff --git a/Engine/Source/Runtime/Renderer/Private/AtmosphereRendering.h b/Engine/Source/Runtime/Renderer/Private/AtmosphereRendering.h
--- a/Engine/Source/Runtime/Renderer/Private/AtmosphereRendering.h
+++ b/Engine/Source/Runtime/Renderer/Private/AtmosphereRendering.h
@@ -69,7 +69,7 @@ public:
 private:
 	/** Atmosphere pre-computation related functions */
 	FIntPoint GetTextureSize();
-	inline void DrawQuad(const FIntRect& ViewRect);
+	inline void DrawQuad(const FIntRect& ViewRect, FShader* VertexShader);
 	void GetLayerValue(int Layer, float& AtmosphereR, FVector4& DhdH);
 	void RenderAtmosphereShaders(const FViewInfo& View, const FIntRect& ViewRect);
 	void PrecomputeAtmosphereData(const FViewInfo* View, FSceneViewFamily& ViewFamily);
diff --git a/Engine/Source/Runtime/Renderer/Private/ClearQuad.cpp b/Engine/Source/Runtime/Renderer/Private/ClearQuad.cpp
--- a/Engine/Source/Runtime/Renderer/Private/ClearQuad.cpp
+++ b/Engine/Source/Runtime/Renderer/Private/ClearQuad.cpp
@@ -48,7 +48,7 @@ void DrawClearQuadMRT(bool bClearColor,int32 NumClearColors,const FLinearColor*
 	RHISetDepthStencilState( DepthStencilStateRHI );
 
 	// Set the new shaders
-	TShaderMapRef<FOneColorVS> VertexShader(GetGlobalShaderMap());
+	TShaderMapRef<TOneColorVS<true> > VertexShader(GetGlobalShaderMap());
 
 	FOneColorPS* PixelShader = NULL;
 
diff --git a/Engine/Source/Runtime/Renderer/Private/CompositionLighting/PostProcessAmbient.cpp b/Engine/Source/Runtime/Renderer/Private/CompositionLighting/PostProcessAmbient.cpp
--- a/Engine/Source/Runtime/Renderer/Private/CompositionLighting/PostProcessAmbient.cpp
+++ b/Engine/Source/Runtime/Renderer/Private/CompositionLighting/PostProcessAmbient.cpp
@@ -117,6 +117,7 @@ void FRCPassPostProcessAmbient::Process(FRenderingCompositePassContext& Context)
 			View.ViewRect.Width(), View.ViewRect.Height(),
 			View.ViewRect.Size(),
 			GSceneRenderTargets.SceneColor->GetDesc().Extent,
+			*VertexShader,
 			EDRF_UseTriangleOptimization);
 	}
 
diff --git a/Engine/Source/Runtime/Renderer/Private/CompositionLighting/PostProcessAmbientOcclusion.cpp b/Engine/Source/Runtime/Renderer/Private/CompositionLighting/PostProcessAmbientOcclusion.cpp
--- a/Engine/Source/Runtime/Renderer/Private/CompositionLighting/PostProcessAmbientOcclusion.cpp
+++ b/Engine/Source/Runtime/Renderer/Private/CompositionLighting/PostProcessAmbientOcclusion.cpp
@@ -297,6 +297,8 @@ void FRCPassPostProcessAmbientOcclusionSetup::Process(FRenderingCompositePassCon
 		SetShaderSetupTempl<0>(Context);
 	}
 
+	TShaderMapRef<FPostProcessVS> VertexShader(GetGlobalShaderMap());
+
 	// Draw a quad mapping scene color to the view's render target
 	DrawRectangle( 
 		0, 0,
@@ -305,6 +307,7 @@ void FRCPassPostProcessAmbientOcclusionSetup::Process(FRenderingCompositePassCon
 		SrcRect.Width(), SrcRect.Height(),
 		DestRect.Size(),
 		GSceneRenderTargets.GetBufferSizeXY(),
+		*VertexShader,
 		EDRF_UseTriangleOptimization);
 
 	RHICopyToResolveTarget(DestRenderTarget.TargetableTexture, DestRenderTarget.ShaderResourceTexture, false, FResolveParams());
@@ -436,6 +439,7 @@ void FRCPassPostProcessAmbientOcclusion::Process(FRenderingCompositePassContext&
 		}
 	}
 
+	TShaderMapRef<FPostProcessVS> VertexShader(GetGlobalShaderMap());
 
 	// Draw a quad mapping scene color to the view's render target
 	DrawRectangle( 
@@ -445,6 +449,7 @@ void FRCPassPostProcessAmbientOcclusion::Process(FRenderingCompositePassContext&
 		ViewRect.Width(), ViewRect.Height(),
 		ViewRect.Size(),
 		TexSize,
+		*VertexShader,
 		EDRF_UseTriangleOptimization);
 
 
@@ -509,6 +514,7 @@ void FRCPassPostProcessBasePassAO::Process(FRenderingCompositePassContext& Conte
 		View.ViewRect.Width(), View.ViewRect.Height(),
 		View.ViewRect.Size(),
 		GSceneRenderTargets.GetBufferSizeXY(),
+		*VertexShader,
 		EDRF_UseTriangleOptimization);
 
 	RHICopyToResolveTarget(DestRenderTarget.TargetableTexture, DestRenderTarget.ShaderResourceTexture, false, FResolveParams());
diff --git a/Engine/Source/Runtime/Renderer/Private/CompositionLighting/PostProcessDeferredDecals.cpp b/Engine/Source/Runtime/Renderer/Private/CompositionLighting/PostProcessDeferredDecals.cpp
--- a/Engine/Source/Runtime/Renderer/Private/CompositionLighting/PostProcessDeferredDecals.cpp
+++ b/Engine/Source/Runtime/Renderer/Private/CompositionLighting/PostProcessDeferredDecals.cpp
@@ -371,6 +371,7 @@ void StencilDecalMask(const FSceneView& View)
 		View.ViewRect.Width(), View.ViewRect.Height(),
 		FIntPoint(View.ViewRect.Width(), View.ViewRect.Height()),
 		GSceneRenderTargets.GetBufferSizeXY(),
+		*ScreenVertexShader,
 		EDRF_UseTriangleOptimization);
 }
 
diff --git a/Engine/Source/Runtime/Renderer/Private/CompositionLighting/PostProcessLpvIndirect.cpp b/Engine/Source/Runtime/Renderer/Private/CompositionLighting/PostProcessLpvIndirect.cpp
--- a/Engine/Source/Runtime/Renderer/Private/CompositionLighting/PostProcessLpvIndirect.cpp
+++ b/Engine/Source/Runtime/Renderer/Private/CompositionLighting/PostProcessLpvIndirect.cpp
@@ -217,7 +217,8 @@ void FRCPassPostProcessLpvIndirect::Process(FRenderingCompositePassContext& Cont
 		View.ViewRect.Min.X, View.ViewRect.Min.Y, 
 		View.ViewRect.Width(), View.ViewRect.Height(),
 		View.ViewRect.Size(),
-		GSceneRenderTargets.SceneColor->GetDesc().Extent);
+		GSceneRenderTargets.SceneColor->GetDesc().Extent,
+		*VertexShader);
 
 	RHICopyToResolveTarget(DestColorRenderTarget.TargetableTexture, DestColorRenderTarget.ShaderResourceTexture, false, FResolveParams());
 }
diff --git a/Engine/Source/Runtime/Renderer/Private/CompositionLighting/PostProcessPassThrough.cpp b/Engine/Source/Runtime/Renderer/Private/CompositionLighting/PostProcessPassThrough.cpp
--- a/Engine/Source/Runtime/Renderer/Private/CompositionLighting/PostProcessPassThrough.cpp
+++ b/Engine/Source/Runtime/Renderer/Private/CompositionLighting/PostProcessPassThrough.cpp
@@ -132,6 +132,7 @@ void FRCPassPostProcessPassThrough::Process(FRenderingCompositePassContext& Cont
 		SrcRect.Width(), SrcRect.Height(),
 		DestSize,
 		SrcSize,
+		*VertexShader,
 		EDRF_UseTriangleOptimization);
 
 
diff --git a/Engine/Source/Runtime/Renderer/Private/DeferredShadingRenderer.cpp b/Engine/Source/Runtime/Renderer/Private/DeferredShadingRenderer.cpp
--- a/Engine/Source/Runtime/Renderer/Private/DeferredShadingRenderer.cpp
+++ b/Engine/Source/Runtime/Renderer/Private/DeferredShadingRenderer.cpp
@@ -127,7 +127,7 @@ void FDeferredShadingSceneRenderer::ClearGBufferAtMaxZ()
 
 	uint32 NumActiveRenderTargets = GSceneRenderTargets.GetNumGBufferTargets();
 	
-	TShaderMapRef<FOneColorVS> VertexShader(GetGlobalShaderMap());
+	TShaderMapRef<TOneColorVS<true> > VertexShader(GetGlobalShaderMap());
 	FOneColorPS* PixelShader = NULL; 
 
 	// Assume for now all code path supports SM4, otherwise render target numbers are changed
@@ -1035,6 +1035,7 @@ void FDeferredShadingSceneRenderer::UpdateDownsampledDepthSurface()
 				View.ViewRect.Width(), View.ViewRect.Height(),
 				FIntPoint(DownsampledSizeX, DownsampledSizeY),
 				GSceneRenderTargets.GetBufferSizeXY(),
+				*ScreenVertexShader,
 				EDRF_UseTriangleOptimization);
 		}
 	}
diff --git a/Engine/Source/Runtime/Renderer/Private/DistortionRendering.cpp b/Engine/Source/Runtime/Renderer/Private/DistortionRendering.cpp
--- a/Engine/Source/Runtime/Renderer/Private/DistortionRendering.cpp
+++ b/Engine/Source/Runtime/Renderer/Private/DistortionRendering.cpp
@@ -872,6 +872,7 @@ void FDeferredShadingSceneRenderer::RenderDistortion()
 					View.ViewRect.Width(), View.ViewRect.Height(),
 					View.ViewRect.Size(),
 					GSceneRenderTargets.SceneColor->GetDesc().Extent,
+					*VertexShader,
 					EDRF_UseTriangleOptimization);
 
 				RHICopyToResolveTarget(DestRenderTarget.TargetableTexture, DestRenderTarget.ShaderResourceTexture, false, FResolveParams());
diff --git a/Engine/Source/Runtime/Renderer/Private/ForwardTranslucentRendering.cpp b/Engine/Source/Runtime/Renderer/Private/ForwardTranslucentRendering.cpp
--- a/Engine/Source/Runtime/Renderer/Private/ForwardTranslucentRendering.cpp
+++ b/Engine/Source/Runtime/Renderer/Private/ForwardTranslucentRendering.cpp
@@ -74,6 +74,7 @@ void FForwardShadingSceneRenderer::CopySceneAlpha(void)
 		X, Y,
 		FIntPoint(X, Y),
 		GSceneRenderTargets.GetBufferSizeXY(),
+		*ScreenVertexShader,
 		EDRF_UseTriangleOptimization);
 
 	GSceneRenderTargets.FinishRenderingSceneAlphaCopy();
diff --git a/Engine/Source/Runtime/Renderer/Private/GPUBenchmark.cpp b/Engine/Source/Runtime/Renderer/Private/GPUBenchmark.cpp
--- a/Engine/Source/Runtime/Renderer/Private/GPUBenchmark.cpp
+++ b/Engine/Source/Runtime/Renderer/Private/GPUBenchmark.cpp
@@ -146,6 +146,7 @@ void RunBenchmarkShader(const FSceneView& View, TRefCountPtr<IPooledRenderTarget
 			GBenchmarkResolution, GBenchmarkResolution,
 			FIntPoint(GBenchmarkResolution, GBenchmarkResolution),
 			FIntPoint(GBenchmarkResolution, GBenchmarkResolution),
+			*VertexShader,
 			EDRF_Default);
 	}
 }
diff --git a/Engine/Source/Runtime/Renderer/Private/LightRendering.cpp b/Engine/Source/Runtime/Renderer/Private/LightRendering.cpp
--- a/Engine/Source/Runtime/Renderer/Private/LightRendering.cpp
+++ b/Engine/Source/Runtime/Renderer/Private/LightRendering.cpp
@@ -750,6 +750,7 @@ void FDeferredShadingSceneRenderer::RenderLight(const FLightSceneInfo* LightScen
 				View.ViewRect.Width(), View.ViewRect.Height(),
 				View.ViewRect.Size(),
 				GSceneRenderTargets.GetBufferSizeXY(),
+				*VertexShader,
 				EDRF_UseTriangleOptimization);
 		}
 		else
diff --git a/Engine/Source/Runtime/Renderer/Private/LightShaftRendering.cpp b/Engine/Source/Runtime/Renderer/Private/LightShaftRendering.cpp
--- a/Engine/Source/Runtime/Renderer/Private/LightShaftRendering.cpp
+++ b/Engine/Source/Runtime/Renderer/Private/LightShaftRendering.cpp
@@ -506,6 +506,7 @@ void DownsamplePass(const FViewInfo& View, const FLightSceneInfo* LightSceneInfo
 		View.ViewRect.Width(), View.ViewRect.Height(),
 		FIntPoint(DownsampledSizeX, DownsampledSizeY), 
 		BufferSize,
+		*DownsampleLightShaftsVertexShader,
 		EDRF_UseTriangleOptimization);
 
 	RHICopyToResolveTarget(LightShaftsDest->GetRenderTargetItem().TargetableTexture, LightShaftsDest->GetRenderTargetItem().ShaderResourceTexture, false, FResolveParams() );
@@ -625,6 +626,7 @@ void ApplyRadialBlurPasses(
 				DownSampledXY.X, DownSampledXY.Y, 
 				DownsampledSizeX, DownsampledSizeY,
 				FilterBufferSize, FilterBufferSize,
+				*ScreenVertexShader,
 				EDRF_UseTriangleOptimization);
 		}
 
@@ -666,6 +668,7 @@ void FinishOcclusionTerm(const FViewInfo& View, const FLightSceneInfo* const Lig
 			DownSampledXY.X, DownSampledXY.Y, 
 			DownsampledSizeX, DownsampledSizeY,
 			FilterBufferSize, FilterBufferSize,
+			*ScreenVertexShader,
 			EDRF_UseTriangleOptimization);
 	}
 
@@ -852,6 +855,7 @@ void ApplyLightShaftBloom(const FViewInfo& View, const FLightSceneInfo* const Li
 		DownSampledXY.X, DownSampledXY.Y, 
 		DownsampledSizeX, DownsampledSizeY,
 		FIntPoint(View.ViewRect.Width(), View.ViewRect.Height()), FilterBufferSize,
+		*ScreenVertexShader,
 		EDRF_UseTriangleOptimization);
 
 	GSceneRenderTargets.FinishRenderingSceneColor(false);
diff --git a/Engine/Source/Runtime/Renderer/Private/PostProcess/PostProcessAA.cpp b/Engine/Source/Runtime/Renderer/Private/PostProcess/PostProcessAA.cpp
--- a/Engine/Source/Runtime/Renderer/Private/PostProcess/PostProcessAA.cpp
+++ b/Engine/Source/Runtime/Renderer/Private/PostProcess/PostProcessAA.cpp
@@ -268,6 +268,8 @@ void FRCPassPostProcessAA::Process(FRenderingCompositePassContext& Context)
 	}
 	
 	// Draw a quad mapping scene color to the view's render target
+	TShaderMapRef<FFXAAVS> VertexShader(GetGlobalShaderMap());
+
 	DrawRectangle(
 		DestRect.Min.X, DestRect.Min.Y,
 		DestRect.Width(), DestRect.Height(),
@@ -275,6 +277,7 @@ void FRCPassPostProcessAA::Process(FRenderingCompositePassContext& Context)
 		SrcRect.Width(), SrcRect.Height(),
 		DestSize,
 		SrcSize,
+		*VertexShader,
 		EDRF_Default);
 
 	RHICopyToResolveTarget(DestRenderTarget.TargetableTexture, DestRenderTarget.ShaderResourceTexture, false, FResolveParams());
diff --git a/Engine/Source/Runtime/Renderer/Private/PostProcess/PostProcessBloomSetup.cpp b/Engine/Source/Runtime/Renderer/Private/PostProcess/PostProcessBloomSetup.cpp
--- a/Engine/Source/Runtime/Renderer/Private/PostProcess/PostProcessBloomSetup.cpp
+++ b/Engine/Source/Runtime/Renderer/Private/PostProcess/PostProcessBloomSetup.cpp
@@ -193,6 +193,7 @@ void FRCPassPostProcessBloomSetup::Process(FRenderingCompositePassContext& Conte
 		SrcRect.Width(), SrcRect.Height(),
 		DestSize,
 		SrcSize,
+		*VertexShader,
 		EDRF_UseTriangleOptimization);
 
 	RHICopyToResolveTarget(DestRenderTarget.TargetableTexture, DestRenderTarget.ShaderResourceTexture, false, FResolveParams());
diff --git a/Engine/Source/Runtime/Renderer/Private/PostProcess/PostProcessBokehDOF.cpp b/Engine/Source/Runtime/Renderer/Private/PostProcess/PostProcessBokehDOF.cpp
--- a/Engine/Source/Runtime/Renderer/Private/PostProcess/PostProcessBokehDOF.cpp
+++ b/Engine/Source/Runtime/Renderer/Private/PostProcess/PostProcessBokehDOF.cpp
@@ -182,8 +182,9 @@ void FRCPassPostProcessVisualizeDOF::Process(FRenderingCompositePassContext& Con
 	RHISetDepthStencilState(TStaticDepthStencilState<false,CF_Always>::GetRHI());
 
 	// setup shader
+	TShaderMapRef<FPostProcessVS> VertexShader(GetGlobalShaderMap());
+	
 	{
-		TShaderMapRef<FPostProcessVS> VertexShader(GetGlobalShaderMap());
 		TShaderMapRef<PostProcessVisualizeDOFPS> PixelShader(GetGlobalShaderMap());
 
 		static FGlobalBoundShaderState BoundShaderState;
@@ -202,6 +203,7 @@ void FRCPassPostProcessVisualizeDOF::Process(FRenderingCompositePassContext& Con
 		SrcRect.Width(), SrcRect.Height(),
 		DestSize,
 		SrcSize,
+		*VertexShader,
 		EDRF_UseTriangleOptimization);
 
 	RHICopyToResolveTarget(DestRenderTarget.TargetableTexture, DestRenderTarget.ShaderResourceTexture, false, FResolveParams());
@@ -338,8 +340,8 @@ void FRCPassPostProcessBokehDOFSetup::Process(FRenderingCompositePassContext& Co
 	RHISetDepthStencilState(TStaticDepthStencilState<false,CF_Always>::GetRHI());
 
 	// setup shader
+	TShaderMapRef<FPostProcessVS> VertexShader(GetGlobalShaderMap());
 	{
-		TShaderMapRef<FPostProcessVS> VertexShader(GetGlobalShaderMap());
 		TShaderMapRef<PostProcessBokehDOFSetupPS> PixelShader(GetGlobalShaderMap());
 
 		static FGlobalBoundShaderState BoundShaderState;
@@ -358,6 +360,7 @@ void FRCPassPostProcessBokehDOFSetup::Process(FRenderingCompositePassContext& Co
 		SrcRect.Width(), SrcRect.Height(),
 		DestSize,
 		SrcSize,
+		*VertexShader,
 		EDRF_UseTriangleOptimization);
 
 	RHICopyToResolveTarget(DestRenderTarget.TargetableTexture, DestRenderTarget.ShaderResourceTexture, false, FResolveParams());
diff --git a/Engine/Source/Runtime/Renderer/Private/PostProcess/PostProcessBokehDOFRecombine.cpp b/Engine/Source/Runtime/Renderer/Private/PostProcess/PostProcessBokehDOFRecombine.cpp
--- a/Engine/Source/Runtime/Renderer/Private/PostProcess/PostProcessBokehDOFRecombine.cpp
+++ b/Engine/Source/Runtime/Renderer/Private/PostProcess/PostProcessBokehDOFRecombine.cpp
@@ -156,6 +156,8 @@ void FRCPassPostProcessBokehDOFRecombine::Process(FRenderingCompositePassContext
 		SetShader<2>(Context);
 	}
 
+	TShaderMapRef<FPostProcessVS> VertexShader(GetGlobalShaderMap());
+
 	// Draw a quad mapping scene color to the view's render target
 	DrawRectangle(
 		0, 0,
@@ -164,6 +166,7 @@ void FRCPassPostProcessBokehDOFRecombine::Process(FRenderingCompositePassContext
 		HalfResViewRect.Width(), HalfResViewRect.Height(),
 		View.ViewRect.Size(),
 		TexSize,
+		*VertexShader,
 		EDRF_UseTriangleOptimization);
 
 	RHICopyToResolveTarget(DestRenderTarget.TargetableTexture, DestRenderTarget.ShaderResourceTexture, false, FResolveParams());
diff --git a/Engine/Source/Runtime/Renderer/Private/PostProcess/PostProcessBusyWait.cpp b/Engine/Source/Runtime/Renderer/Private/PostProcess/PostProcessBusyWait.cpp
--- a/Engine/Source/Runtime/Renderer/Private/PostProcess/PostProcessBusyWait.cpp
+++ b/Engine/Source/Runtime/Renderer/Private/PostProcess/PostProcessBusyWait.cpp
@@ -109,6 +109,7 @@ void FRCPassPostProcessBusyWait::Process(FRenderingCompositePassContext& Context
 		SrcRect.Width(), SrcRect.Height(),
 		DestRect.Size(),
 		SrcRect.Size(),
+		*VertexShader,
 		EDRF_UseTriangleOptimization);
 
 	RHICopyToResolveTarget(DestRenderTarget.TargetableTexture, DestRenderTarget.ShaderResourceTexture, false, FResolveParams());
diff --git a/Engine/Source/Runtime/Renderer/Private/PostProcess/PostProcessCombineLUTs.cpp b/Engine/Source/Runtime/Renderer/Private/PostProcess/PostProcessCombineLUTs.cpp
--- a/Engine/Source/Runtime/Renderer/Private/PostProcess/PostProcessCombineLUTs.cpp
+++ b/Engine/Source/Runtime/Renderer/Private/PostProcess/PostProcessCombineLUTs.cpp
@@ -406,6 +406,8 @@ void FRCPassPostProcessCombineLUTs::Process(FRenderingCompositePassContext& Cont
 	else
 	{
 		// use unwrapped 2d texture 256x16
+		TShaderMapRef<FPostProcessVS> VertexShader(GetGlobalShaderMap());
+
 		DrawRectangle( 
 			0, 0,						// XY
 			16 * 16, 16,				// SizeXY
@@ -413,6 +415,7 @@ void FRCPassPostProcessCombineLUTs::Process(FRenderingCompositePassContext& Cont
 			16 * 16, 16,				// SizeUV
 			FIntPoint(16 * 16, 16),		// TargetSize
 			FIntPoint(16 * 16, 16),		// TextureSize
+			*VertexShader,
 			EDRF_UseTriangleOptimization);
 	}
 
diff --git a/Engine/Source/Runtime/Renderer/Private/PostProcess/PostProcessCompositeEditorPrimitives.cpp b/Engine/Source/Runtime/Renderer/Private/PostProcess/PostProcessCompositeEditorPrimitives.cpp
--- a/Engine/Source/Runtime/Renderer/Private/PostProcess/PostProcessCompositeEditorPrimitives.cpp
+++ b/Engine/Source/Runtime/Renderer/Private/PostProcess/PostProcessCompositeEditorPrimitives.cpp
@@ -259,6 +259,8 @@ void FRCPassPostProcessCompositeEditorPrimitives::Process(FRenderingCompositePas
 		check(0);
 	}
 
+	TShaderMapRef<FPostProcessVS> VertexShader(GetGlobalShaderMap());
+
 	// Draw a quad mapping our render targets to the view's render target
 	DrawRectangle(
 		0, 0,
@@ -267,6 +269,7 @@ void FRCPassPostProcessCompositeEditorPrimitives::Process(FRenderingCompositePas
 		SrcRect.Width(), SrcRect.Height(),
 		DestRect.Size(),
 		SrcSize,
+		*VertexShader,
 		EDRF_UseTriangleOptimization);
 
 	RHICopyToResolveTarget(DestRenderTargetSurface, DestRenderTarget.ShaderResourceTexture, false, FResolveParams());
diff --git a/Engine/Source/Runtime/Renderer/Private/PostProcess/PostProcessDOF.cpp b/Engine/Source/Runtime/Renderer/Private/PostProcess/PostProcessDOF.cpp
--- a/Engine/Source/Runtime/Renderer/Private/PostProcess/PostProcessDOF.cpp
+++ b/Engine/Source/Runtime/Renderer/Private/PostProcess/PostProcessDOF.cpp
@@ -158,6 +158,7 @@ void FRCPassPostProcessDOFSetup::Process(FRenderingCompositePassContext& Context
 		SrcRect.Width() + 1, SrcRect.Height() + 1,
 		DestSize,
 		SrcSize,
+		*VertexShader,
 		EDRF_UseTriangleOptimization);
 
 	RHICopyToResolveTarget(DestRenderTarget0.TargetableTexture, DestRenderTarget0.ShaderResourceTexture, false, FResolveParams());
@@ -316,6 +317,7 @@ void FRCPassPostProcessDOFRecombine::Process(FRenderingCompositePassContext& Con
 		HalfResViewRect.Width(), HalfResViewRect.Height(),
 		View.ViewRect.Size(),
 		TexSize,
+		*VertexShader,
 		EDRF_UseTriangleOptimization);
 
 	RHICopyToResolveTarget(DestRenderTarget.TargetableTexture, DestRenderTarget.ShaderResourceTexture, false, FResolveParams());
diff --git a/Engine/Source/Runtime/Renderer/Private/PostProcess/PostProcessDownsample.cpp b/Engine/Source/Runtime/Renderer/Private/PostProcess/PostProcessDownsample.cpp
--- a/Engine/Source/Runtime/Renderer/Private/PostProcess/PostProcessDownsample.cpp
+++ b/Engine/Source/Runtime/Renderer/Private/PostProcess/PostProcessDownsample.cpp
@@ -218,6 +218,8 @@ void FRCPassPostProcessDownsample::Process(FRenderingCompositePassContext& Conte
 		bHasCleared = true;
 	}
 
+	TShaderMapRef<FPostProcessDownsampleVS> VertexShader(GetGlobalShaderMap());
+
 	switch (RectSource)
 	{
 		case EPostProcessRectSource::GBS_ViewRect:
@@ -239,6 +241,7 @@ void FRCPassPostProcessDownsample::Process(FRenderingCompositePassContext& Conte
 				SrcRect.Width(), SrcRect.Height(),
 				DestSize,
 				SrcSize,
+				*VertexShader,
 				EDRF_UseTriangleOptimization);
 		}
 		break;
@@ -280,6 +283,7 @@ void FRCPassPostProcessDownsample::Process(FRenderingCompositePassContext& Conte
 					SrcRect.Width(), SrcRect.Height(),
 					DestSize,
 					SrcSize,
+					*VertexShader,
 					EDRF_UseTriangleOptimization);
 			}
 		}
diff --git a/Engine/Source/Runtime/Renderer/Private/PostProcess/PostProcessEyeAdaptation.cpp b/Engine/Source/Runtime/Renderer/Private/PostProcess/PostProcessEyeAdaptation.cpp
--- a/Engine/Source/Runtime/Renderer/Private/PostProcess/PostProcessEyeAdaptation.cpp
+++ b/Engine/Source/Runtime/Renderer/Private/PostProcess/PostProcessEyeAdaptation.cpp
@@ -108,6 +108,7 @@ void FRCPassPostProcessEyeAdaptation::Process(FRenderingCompositePassContext& Co
 		DestSize.X, DestSize.Y,
 		DestSize,
 		DestSize,
+		*VertexShader,
 		EDRF_UseTriangleOptimization);
 
 	RHICopyToResolveTarget(EyeAdaptation->GetRenderTargetItem().TargetableTexture, EyeAdaptation->GetRenderTargetItem().ShaderResourceTexture, false, FResolveParams());
diff --git a/Engine/Source/Runtime/Renderer/Private/PostProcess/PostProcessGBufferHints.cpp b/Engine/Source/Runtime/Renderer/Private/PostProcess/PostProcessGBufferHints.cpp
--- a/Engine/Source/Runtime/Renderer/Private/PostProcess/PostProcessGBufferHints.cpp
+++ b/Engine/Source/Runtime/Renderer/Private/PostProcess/PostProcessGBufferHints.cpp
@@ -117,6 +117,7 @@ void FRCPassPostProcessGBufferHints::Process(FRenderingCompositePassContext& Con
 		SrcRect.Width(), SrcRect.Height(),
 		DestRect.Size(),
 		SrcSize,
+		*VertexShader,
 		EDRF_UseTriangleOptimization);
 
 
diff --git a/Engine/Source/Runtime/Renderer/Private/PostProcess/PostProcessHMD.cpp b/Engine/Source/Runtime/Renderer/Private/PostProcess/PostProcessHMD.cpp
--- a/Engine/Source/Runtime/Renderer/Private/PostProcess/PostProcessHMD.cpp
+++ b/Engine/Source/Runtime/Renderer/Private/PostProcess/PostProcessHMD.cpp
@@ -143,6 +143,41 @@ public:
 	}
 };
 
+
+/** Encapsulates the post processing vertex shader. */
+class FPostProcessHMDVS : public FGlobalShader
+{
+	DECLARE_SHADER_TYPE(FPostProcessHMDVS,Global);
+
+	static bool ShouldCache(EShaderPlatform Platform)
+	{
+		return true;
+	}
+
+	/** Default constructor. */
+	FPostProcessHMDVS() {}
+
+	/** to have a similar interface as all other shaders */
+	void SetParameters(const FRenderingCompositePassContext& Context)
+	{
+		FGlobalShader::SetParameters(GetVertexShader(), Context.View);
+	}
+
+	void SetParameters(const FSceneView& View)
+	{
+		FGlobalShader::SetParameters(GetVertexShader(), View);
+	}
+
+public:
+
+	/** Initialization constructor. */
+	FPostProcessHMDVS(const ShaderMetaType::CompiledShaderInitializerType& Initializer)
+		: FGlobalShader(Initializer)
+	{
+	}
+};
+
+IMPLEMENT_SHADER_TYPE(, FPostProcessHMDVS, TEXT("PostProcessHMD"), TEXT("MainVS"), SF_Vertex);
 IMPLEMENT_SHADER_TYPE(template<>, FPostProcessHMDPS<true>, TEXT("PostProcessHMD"), TEXT("MainPS"), SF_Pixel);
 IMPLEMENT_SHADER_TYPE(template<>, FPostProcessHMDPS<false>, TEXT("PostProcessHMD"), TEXT("MainPS"), SF_Pixel);
 
@@ -223,7 +258,7 @@ void FRCPassPostProcessHMD::Process(FRenderingCompositePassContext& Context)
 	RHISetRasterizerState(TStaticRasterizerState<>::GetRHI());
 	RHISetDepthStencilState(TStaticDepthStencilState<false,CF_Always>::GetRHI());
 
-	TShaderMapRef<FPostProcessVS> VertexShader(GetGlobalShaderMap());
+	TShaderMapRef<FPostProcessHMDVS> VertexShader(GetGlobalShaderMap());
 
 	FMatrix QuadTexTransform;
 	FMatrix QuadPosTransform = FMatrix::Identity;
diff --git a/Engine/Source/Runtime/Renderer/Private/PostProcess/PostProcessHistogramReduce.cpp b/Engine/Source/Runtime/Renderer/Private/PostProcess/PostProcessHistogramReduce.cpp
--- a/Engine/Source/Runtime/Renderer/Private/PostProcess/PostProcessHistogramReduce.cpp
+++ b/Engine/Source/Runtime/Renderer/Private/PostProcess/PostProcessHistogramReduce.cpp
@@ -135,6 +135,7 @@ void FRCPassPostProcessHistogramReduce::Process(FRenderingCompositePassContext&
 		SrcSize.X, 0,
 		DestSize,
 		SrcSize,
+		*VertexShader,
 		EDRF_UseTriangleOptimization);
 
 	RHICopyToResolveTarget(DestRenderTarget.TargetableTexture, DestRenderTarget.ShaderResourceTexture, false, FResolveParams());
diff --git a/Engine/Source/Runtime/Renderer/Private/PostProcess/PostProcessLensFlares.cpp b/Engine/Source/Runtime/Renderer/Private/PostProcess/PostProcessLensFlares.cpp
--- a/Engine/Source/Runtime/Renderer/Private/PostProcess/PostProcessLensFlares.cpp
+++ b/Engine/Source/Runtime/Renderer/Private/PostProcess/PostProcessLensFlares.cpp
@@ -174,6 +174,7 @@ void FRCPassPostProcessLensFlares::Process(FRenderingCompositePassContext& Conte
 			ViewSize1.X, ViewSize1.Y,
 			ViewSize1,
 			TexSize1,
+			*VertexShader,
 			EDRF_UseTriangleOptimization);
 	}
 
@@ -226,6 +227,7 @@ void FRCPassPostProcessLensFlares::Process(FRenderingCompositePassContext& Conte
 				ViewSize2.X, ViewSize2.Y,
 				ViewSize1,
 				TexSize2,
+				*VertexShader,
 				EDRF_Default);
 		}
 	}
diff --git a/Engine/Source/Runtime/Renderer/Private/PostProcess/PostProcessMaterial.cpp b/Engine/Source/Runtime/Renderer/Private/PostProcess/PostProcessMaterial.cpp
--- a/Engine/Source/Runtime/Renderer/Private/PostProcess/PostProcessMaterial.cpp
+++ b/Engine/Source/Runtime/Renderer/Private/PostProcess/PostProcessMaterial.cpp
@@ -164,6 +164,7 @@ void FRCPassPostProcessMaterial::Process(FRenderingCompositePassContext& Context
 		SrcRect.Width(), SrcRect.Height(),
 		DestRect.Size(),
 		SrcSize,
+		VertexShader,
 		EDRF_UseTriangleOptimization);
 
 	RHICopyToResolveTarget(DestRenderTarget.TargetableTexture, DestRenderTarget.ShaderResourceTexture, false, FResolveParams());
diff --git a/Engine/Source/Runtime/Renderer/Private/PostProcess/PostProcessMobile.cpp b/Engine/Source/Runtime/Renderer/Private/PostProcess/PostProcessMobile.cpp
--- a/Engine/Source/Runtime/Renderer/Private/PostProcess/PostProcessMobile.cpp
+++ b/Engine/Source/Runtime/Renderer/Private/PostProcess/PostProcessMobile.cpp
@@ -208,6 +208,8 @@ void FRCPassPostProcessBloomSetupES2::Process(FRenderingCompositePassContext& Co
 
 	SetShader(Context);
 
+	TShaderMapRef<FPostProcessBloomSetupVS_ES2> VertexShader(GetGlobalShaderMap());
+
 	DrawRectangle(
 		0, 0,
 		DstX, DstY,
@@ -215,6 +217,7 @@ void FRCPassPostProcessBloomSetupES2::Process(FRenderingCompositePassContext& Co
 		SrcRect.Width(), SrcRect.Height(),
 		DstSize,
 		SrcSize,
+		*VertexShader,
 		EDRF_UseTriangleOptimization);
 
 	RHICopyToResolveTarget(DestRenderTarget.TargetableTexture, DestRenderTarget.ShaderResourceTexture, false, FResolveParams());
@@ -405,6 +408,7 @@ void FRCPassPostProcessBloomSetupSmallES2::Process(FRenderingCompositePassContex
 
 	SetShader(Context);
 
+	TShaderMapRef<FPostProcessBloomSetupSmallVS_ES2> VertexShader(GetGlobalShaderMap());
 	DrawRectangle(
 		0, 0,
 		DstX, DstY,
@@ -412,6 +416,7 @@ void FRCPassPostProcessBloomSetupSmallES2::Process(FRenderingCompositePassContex
 		SrcRect.Width(), SrcRect.Height(),
 		DstSize,
 		SrcSize,
+		*VertexShader,
 		EDRF_UseTriangleOptimization);
 
 	RHICopyToResolveTarget(DestRenderTarget.TargetableTexture, DestRenderTarget.ShaderResourceTexture, false, FResolveParams());
@@ -572,6 +577,7 @@ void FRCPassPostProcessBloomDownES2::Process(FRenderingCompositePassContext& Con
 		DstX, DstY,
 		SrcDstSize,
 		SrcDstSize,
+		*VertexShader,
 		EDRF_UseTriangleOptimization);
 
 	RHICopyToResolveTarget(DestRenderTarget.TargetableTexture, DestRenderTarget.ShaderResourceTexture, false, FResolveParams());
@@ -740,6 +746,7 @@ void FRCPassPostProcessBloomUpES2::Process(FRenderingCompositePassContext& Conte
 		DstX, DstY,
 		SrcDstSize,
 		SrcDstSize,
+		*VertexShader,
 		EDRF_UseTriangleOptimization);
 
 	RHICopyToResolveTarget(DestRenderTarget.TargetableTexture, DestRenderTarget.ShaderResourceTexture, false, FResolveParams());
@@ -934,6 +941,9 @@ void FRCPassPostProcessSunMaskES2::Process(FRenderingCompositePassContext& Conte
 	FIntPoint SrcSize;
 	FIntRect SrcRect;
 	const FSceneView& View = Context.View;
+
+	TShaderMapRef<FPostProcessSunMaskVS_ES2> VertexShader(GetGlobalShaderMap());
+
 	if(bOnChip)
 	{
 		SrcSize = DstSize;
@@ -954,6 +964,7 @@ void FRCPassPostProcessSunMaskES2::Process(FRenderingCompositePassContext& Conte
 			SrcRect.Width(), SrcRect.Height(),
 			DstSize,
 			SrcSize,
+			*VertexShader,
 			EDRF_UseTriangleOptimization);
 	}
 	else
@@ -985,6 +996,7 @@ void FRCPassPostProcessSunMaskES2::Process(FRenderingCompositePassContext& Conte
 			SrcRect.Width(), SrcRect.Height(),
 			DstSize,
 			SrcSize,
+			*VertexShader,
 			EDRF_UseTriangleOptimization);
 
 		RHICopyToResolveTarget(DestRenderTarget.TargetableTexture, DestRenderTarget.ShaderResourceTexture, false, FResolveParams());
@@ -1158,6 +1170,7 @@ void FRCPassPostProcessSunAlphaES2::Process(FRenderingCompositePassContext& Cont
 	SetShader(Context);
 
 	FIntPoint SrcDstSize = PrePostSourceViewportSize / 4;
+	TShaderMapRef<FPostProcessSunAlphaVS_ES2> VertexShader(GetGlobalShaderMap());
 
 	DrawRectangle(
 		0, 0,
@@ -1166,6 +1179,7 @@ void FRCPassPostProcessSunAlphaES2::Process(FRenderingCompositePassContext& Cont
 		DstX, DstY,
 		SrcDstSize,
 		SrcDstSize,
+		*VertexShader,
 		EDRF_UseTriangleOptimization);
 
 	RHICopyToResolveTarget(DestRenderTarget.TargetableTexture, DestRenderTarget.ShaderResourceTexture, false, FResolveParams());
@@ -1328,6 +1342,7 @@ void FRCPassPostProcessSunBlurES2::Process(FRenderingCompositePassContext& Conte
 		DstX, DstY,
 		SrcDstSize,
 		SrcDstSize,
+		*VertexShader,
 		EDRF_UseTriangleOptimization);
 
 	RHICopyToResolveTarget(DestRenderTarget.TargetableTexture, DestRenderTarget.ShaderResourceTexture, false, FResolveParams());
@@ -1532,6 +1547,7 @@ void FRCPassPostProcessSunMergeES2::Process(FRenderingCompositePassContext& Cont
 	SetShader(Context);
 
 	FIntPoint SrcDstSize = PrePostSourceViewportSize / 4;
+	TShaderMapRef<FPostProcessSunMergeVS_ES2> VertexShader(GetGlobalShaderMap());
 
 	DrawRectangle(
 		0, 0,
@@ -1540,6 +1556,7 @@ void FRCPassPostProcessSunMergeES2::Process(FRenderingCompositePassContext& Cont
 		DstX, DstY,
 		SrcDstSize,
 		SrcDstSize,
+		*VertexShader,
 		EDRF_UseTriangleOptimization);
 
 	RHICopyToResolveTarget(DestRenderTarget.TargetableTexture, DestRenderTarget.ShaderResourceTexture, false, FResolveParams());
@@ -1723,6 +1740,7 @@ void FRCPassPostProcessSunMergeSmallES2::Process(FRenderingCompositePassContext&
 	SetShader(Context);
 
 	FIntPoint SrcDstSize = PrePostSourceViewportSize / 4;
+	TShaderMapRef<FPostProcessSunMergeSmallVS_ES2> VertexShader(GetGlobalShaderMap());
 
 	DrawRectangle(
 		0, 0,
@@ -1731,6 +1749,7 @@ void FRCPassPostProcessSunMergeSmallES2::Process(FRenderingCompositePassContext&
 		DstX, DstY,
 		SrcDstSize,
 		SrcDstSize,
+		*VertexShader,
 		EDRF_UseTriangleOptimization);
 
 	RHICopyToResolveTarget(DestRenderTarget.TargetableTexture, DestRenderTarget.ShaderResourceTexture, false, FResolveParams());
@@ -1943,6 +1962,8 @@ void FRCPassPostProcessDofDownES2::Process(FRenderingCompositePassContext& Conte
 
 	SetShader(Context);
 
+	TShaderMapRef<FPostProcessDofDownVS_ES2> VertexShader(GetGlobalShaderMap());
+
 	DrawRectangle(
 		0, 0,
 		DstX, DstY,
@@ -1950,6 +1971,7 @@ void FRCPassPostProcessDofDownES2::Process(FRenderingCompositePassContext& Conte
 		SrcRect.Width(), SrcRect.Height(),
 		DstSize,
 		SrcSize,
+		*VertexShader,
 		EDRF_UseTriangleOptimization);
 
 	RHICopyToResolveTarget(DestRenderTarget.TargetableTexture, DestRenderTarget.ShaderResourceTexture, false, FResolveParams());
@@ -2124,6 +2146,7 @@ void FRCPassPostProcessDofNearES2::Process(FRenderingCompositePassContext& Conte
 	SetShader(Context);
 
 	FIntPoint SrcDstSize = PrePostSourceViewportSize / 4;
+	TShaderMapRef<FPostProcessDofNearVS_ES2> VertexShader(GetGlobalShaderMap());
 
 	DrawRectangle(
 		0, 0,
@@ -2132,6 +2155,7 @@ void FRCPassPostProcessDofNearES2::Process(FRenderingCompositePassContext& Conte
 		DstX, DstY,
 		SrcDstSize,
 		SrcDstSize,
+		*VertexShader,
 		EDRF_UseTriangleOptimization);
 
 	RHICopyToResolveTarget(DestRenderTarget.TargetableTexture, DestRenderTarget.ShaderResourceTexture, false, FResolveParams());
@@ -2286,6 +2310,7 @@ void FRCPassPostProcessDofBlurES2::Process(FRenderingCompositePassContext& Conte
 		DstX, DstY,
 		SrcDstSize,
 		SrcDstSize,
+		*VertexShader,
 		EDRF_UseTriangleOptimization);
 
 	RHICopyToResolveTarget(DestRenderTarget.TargetableTexture, DestRenderTarget.ShaderResourceTexture, false, FResolveParams());
@@ -2447,6 +2472,7 @@ void FRCPassPostProcessSunAvgES2::Process(FRenderingCompositePassContext& Contex
 	SetShader(Context);
 
 	FIntPoint SrcDstSize = PrePostSourceViewportSize / 4;
+	TShaderMapRef<FPostProcessSunAvgVS_ES2> VertexShader(GetGlobalShaderMap());
 
 	DrawRectangle(
 		0, 0,
@@ -2455,6 +2481,7 @@ void FRCPassPostProcessSunAvgES2::Process(FRenderingCompositePassContext& Contex
 		DstX, DstY,
 		SrcDstSize,
 		SrcDstSize,
+		*VertexShader,
 		EDRF_UseTriangleOptimization);
 
 	RHICopyToResolveTarget(DestRenderTarget.TargetableTexture, DestRenderTarget.ShaderResourceTexture, false, FResolveParams());
@@ -2722,6 +2749,7 @@ void FRCPassPostProcessAaES2::Process(FRenderingCompositePassContext& Context)
 	SetShader(Context);
 
 	FIntPoint SrcDstSize = PrePostSourceViewportSize;
+	TShaderMapRef<FPostProcessAaVS_ES2> VertexShader(GetGlobalShaderMap());
 
 	DrawRectangle(
 		0, 0,
@@ -2730,6 +2758,7 @@ void FRCPassPostProcessAaES2::Process(FRenderingCompositePassContext& Context)
 		DstX, DstY,
 		SrcDstSize,
 		SrcDstSize,
+		*VertexShader,
 		EDRF_UseTriangleOptimization);
 
 	RHICopyToResolveTarget(DestRenderTarget.TargetableTexture, DestRenderTarget.ShaderResourceTexture, false, FResolveParams());
diff --git a/Engine/Source/Runtime/Renderer/Private/PostProcess/PostProcessMotionBlur.cpp b/Engine/Source/Runtime/Renderer/Private/PostProcess/PostProcessMotionBlur.cpp
--- a/Engine/Source/Runtime/Renderer/Private/PostProcess/PostProcessMotionBlur.cpp
+++ b/Engine/Source/Runtime/Renderer/Private/PostProcess/PostProcessMotionBlur.cpp
@@ -127,10 +127,10 @@ void FRCPassPostProcessMotionBlurSetup::Process(FRenderingCompositePassContext&
 	RHISetRasterizerState(TStaticRasterizerState<>::GetRHI());
 	RHISetDepthStencilState(TStaticDepthStencilState<false,CF_Always>::GetRHI());
 
-	{
-		TShaderMapRef<FPostProcessVS> VertexShader(GetGlobalShaderMap());
-		TShaderMapRef<FPostProcessMotionBlurSetupPS > PixelShader(GetGlobalShaderMap());
+	TShaderMapRef<FPostProcessVS> VertexShader(GetGlobalShaderMap());
 
+	{		
+		TShaderMapRef<FPostProcessMotionBlurSetupPS > PixelShader(GetGlobalShaderMap());
 		static FGlobalBoundShaderState BoundShaderState;
 
 		SetGlobalBoundShaderState(BoundShaderState, GFilterVertexDeclaration.VertexDeclarationRHI, *VertexShader, *PixelShader);
@@ -147,6 +147,7 @@ void FRCPassPostProcessMotionBlurSetup::Process(FRenderingCompositePassContext&
 		SrcRect.Width(), SrcRect.Height(),
 		DestSize,
 		SrcSize,
+		*VertexShader,
 		EDRF_UseTriangleOptimization);
 
 	RHICopyToResolveTarget(DestRenderTarget0.TargetableTexture, DestRenderTarget0.ShaderResourceTexture, false, FResolveParams());
@@ -442,6 +443,8 @@ void FRCPassPostProcessMotionBlur::Process(FRenderingCompositePassContext& Conte
 		SetMotionBlurShaderTempl<4>(Context);
 	}
 
+	TShaderMapRef<FPostProcessVS> VertexShader(GetGlobalShaderMap());
+
 	// Draw a quad mapping scene color to the view's render target
 	DrawRectangle(
 		0, 0,
@@ -450,6 +453,7 @@ void FRCPassPostProcessMotionBlur::Process(FRenderingCompositePassContext& Conte
 		SrcRect.Width(), SrcRect.Height(),
 		SrcRect.Size(),
 		SrcSize,
+		*VertexShader,
 		EDRF_UseTriangleOptimization);
 
 	RHICopyToResolveTarget(DestRenderTarget.TargetableTexture, DestRenderTarget.ShaderResourceTexture, false, FResolveParams());
@@ -573,6 +577,7 @@ void FRCPassPostProcessMotionBlurRecombine::Process(FRenderingCompositePassConte
 		SrcRect.Width(), SrcRect.Height(),
 		SrcRect.Size(),
 		SrcSize,
+		*VertexShader,
 		EDRF_UseTriangleOptimization);
 
 	RHICopyToResolveTarget(DestRenderTarget.TargetableTexture, DestRenderTarget.ShaderResourceTexture, false, FResolveParams());
@@ -637,6 +642,8 @@ void FRCPassPostProcessVisualizeMotionBlur::Process(FRenderingCompositePassConte
 	SetMotionBlurShaderTempl<0>(Context);
 
 	// Draw a quad mapping scene color to the view's render target
+	TShaderMapRef<FPostProcessVS> VertexShader(GetGlobalShaderMap());
+
 	DrawRectangle(
 		0, 0,
 		SrcRect.Width(), SrcRect.Height(),
@@ -644,6 +651,7 @@ void FRCPassPostProcessVisualizeMotionBlur::Process(FRenderingCompositePassConte
 		SrcRect.Width(), SrcRect.Height(),
 		SrcRect.Size(),
 		SrcSize,
+		*VertexShader,
 		EDRF_UseTriangleOptimization);
 
 	RHICopyToResolveTarget(DestRenderTarget.TargetableTexture, DestRenderTarget.ShaderResourceTexture, false, FResolveParams());
diff --git a/Engine/Source/Runtime/Renderer/Private/PostProcess/PostProcessNoiseBlur.cpp b/Engine/Source/Runtime/Renderer/Private/PostProcess/PostProcessNoiseBlur.cpp
--- a/Engine/Source/Runtime/Renderer/Private/PostProcess/PostProcessNoiseBlur.cpp
+++ b/Engine/Source/Runtime/Renderer/Private/PostProcess/PostProcessNoiseBlur.cpp
@@ -161,6 +161,7 @@ void FRCPassPostProcessNoiseBlur::Process(FRenderingCompositePassContext& Contex
 	}
 
 	// Draw a quad mapping scene color to the view's render target
+	TShaderMapRef<FPostProcessVS> VertexShader(GetGlobalShaderMap());
 	DrawRectangle(
 		DestRect.Min.X, DestRect.Min.Y,
 		DestRect.Width(), DestRect.Height(),
@@ -168,6 +169,7 @@ void FRCPassPostProcessNoiseBlur::Process(FRenderingCompositePassContext& Contex
 		SrcRect.Width(), SrcRect.Height(),
 		DestSize,
 		SrcSize,
+		*VertexShader,
 		EDRF_UseTriangleOptimization);
 
 	RHICopyToResolveTarget(DestRenderTarget.TargetableTexture, DestRenderTarget.ShaderResourceTexture, false, FResolveParams());
diff --git a/Engine/Source/Runtime/Renderer/Private/PostProcess/PostProcessSceneColorFringe.cpp b/Engine/Source/Runtime/Renderer/Private/PostProcess/PostProcessSceneColorFringe.cpp
--- a/Engine/Source/Runtime/Renderer/Private/PostProcess/PostProcessSceneColorFringe.cpp
+++ b/Engine/Source/Runtime/Renderer/Private/PostProcess/PostProcessSceneColorFringe.cpp
@@ -184,6 +184,7 @@ void FRCPassPostProcessSceneColorFringe::Process(FRenderingCompositePassContext&
 		View.ViewRect.Width(), View.ViewRect.Height(),
 		View.ViewRect.Size(),
 		GSceneRenderTargets.SceneColor->GetDesc().Extent,
+		*VertexShader,
 		EDRF_UseTriangleOptimization);
 
 	RHICopyToResolveTarget(DestRenderTarget.TargetableTexture, DestRenderTarget.ShaderResourceTexture, false, FResolveParams());
diff --git a/Engine/Source/Runtime/Renderer/Private/PostProcess/PostProcessSelectionOutline.cpp b/Engine/Source/Runtime/Renderer/Private/PostProcess/PostProcessSelectionOutline.cpp
--- a/Engine/Source/Runtime/Renderer/Private/PostProcess/PostProcessSelectionOutline.cpp
+++ b/Engine/Source/Runtime/Renderer/Private/PostProcess/PostProcessSelectionOutline.cpp
@@ -60,7 +60,7 @@ void FRCPassPostProcessSelectionOutlineColor::Process(FRenderingCompositePassCon
 			// 0 means no object, 1 means BSP so we start with 1
 			uint32 StencilValue = 2;
 
-			TShaderMapRef<FOneColorVS> VertexShader(GetGlobalShaderMap());
+			TShaderMapRef< TOneColorVS<false> > VertexShader(GetGlobalShaderMap());
 			TShaderMapRef<FOneColorPS> PixelShader(GetGlobalShaderMap());
 
 			RHISetRasterizerState(TStaticRasterizerState<>::GetRHI());
@@ -375,6 +375,7 @@ void FRCPassPostProcessSelectionOutline::Process(FRenderingCompositePassContext&
 	}
 
 	// Draw a quad mapping scene color to the view's render target
+	TShaderMapRef<FPostProcessVS> VertexShader(GetGlobalShaderMap());
 	DrawRectangle(
 		0, 0,
 		ViewRect.Width(), ViewRect.Height(),
@@ -382,6 +383,7 @@ void FRCPassPostProcessSelectionOutline::Process(FRenderingCompositePassContext&
 		ViewRect.Width(), ViewRect.Height(),
 		ViewRect.Size(),
 		SrcSize,
+		*VertexShader,
 		EDRF_UseTriangleOptimization);
 
 	RHICopyToResolveTarget(DestRenderTarget.TargetableTexture, DestRenderTarget.ShaderResourceTexture, false, FResolveParams());
diff --git a/Engine/Source/Runtime/Renderer/Private/PostProcess/PostProcessSubsurface.cpp b/Engine/Source/Runtime/Renderer/Private/PostProcess/PostProcessSubsurface.cpp
--- a/Engine/Source/Runtime/Renderer/Private/PostProcess/PostProcessSubsurface.cpp
+++ b/Engine/Source/Runtime/Renderer/Private/PostProcess/PostProcessSubsurface.cpp
@@ -171,6 +171,7 @@ void FRCPassPostProcessSubsurfaceSetup::Process(FRenderingCompositePassContext&
 	}
 
 	// Draw a quad mapping scene color to the view's render target
+	TShaderMapRef<FPostProcessVS> VertexShader(GetGlobalShaderMap());
 	DrawRectangle(
 		DestRect.Min.X, DestRect.Min.Y,
 		DestRect.Width(), DestRect.Height(),
@@ -178,6 +179,7 @@ void FRCPassPostProcessSubsurfaceSetup::Process(FRenderingCompositePassContext&
 		SrcRect.Width(), SrcRect.Height(),
 		DestSize,
 		SrcSize,
+		*VertexShader,
 		EDRF_UseTriangleOptimization);
 
 	RHICopyToResolveTarget(DestRenderTarget.TargetableTexture, DestRenderTarget.ShaderResourceTexture, false, FResolveParams());
@@ -330,6 +332,8 @@ void FRCPassPostProcessSubsurface::Process(FRenderingCompositePassContext& Conte
 	RHISetRasterizerState(TStaticRasterizerState<>::GetRHI());
 	RHISetDepthStencilState(TStaticDepthStencilState<false,CF_Always>::GetRHI());
 
+	TShaderMapRef<FPostProcessVS> VertexShader(GetGlobalShaderMap());
+
 	if(Pass == 0)
 	{
 		SCOPED_DRAW_EVENT(SubsurfacePass0, DEC_SCENE_ITEMS);
@@ -344,6 +348,7 @@ void FRCPassPostProcessSubsurface::Process(FRenderingCompositePassContext& Conte
 			SrcRect.Width(), SrcRect.Height(),
 			DestSize,
 			SrcSize,
+			*VertexShader,
 			EDRF_UseTriangleOptimization);
 
 		RHICopyToResolveTarget(DestRenderTarget.TargetableTexture, DestRenderTarget.ShaderResourceTexture, false, FResolveParams());
@@ -371,6 +376,7 @@ void FRCPassPostProcessSubsurface::Process(FRenderingCompositePassContext& Conte
 			SrcRect.Width(), SrcRect.Height(),
 			DestSize,
 			SrcSize,
+			*VertexShader,
 			EDRF_UseTriangleOptimization);
 
 		RHICopyToResolveTarget(DestRenderTarget.TargetableTexture, DestRenderTarget.ShaderResourceTexture, false, FResolveParams());
diff --git a/Engine/Source/Runtime/Renderer/Private/PostProcess/PostProcessTemporalAA.cpp b/Engine/Source/Runtime/Renderer/Private/PostProcess/PostProcessTemporalAA.cpp
--- a/Engine/Source/Runtime/Renderer/Private/PostProcess/PostProcessTemporalAA.cpp
+++ b/Engine/Source/Runtime/Renderer/Private/PostProcess/PostProcessTemporalAA.cpp
@@ -334,6 +334,7 @@ void FRCPassPostProcessSSRTemporalAA::Process(FRenderingCompositePassContext& Co
 		SrcRect.Width(), SrcRect.Height(),
 		SrcRect.Size(),
 		SrcSize,
+		*VertexShader,
 		EDRF_UseTriangleOptimization);
 
 	RHICopyToResolveTarget(DestRenderTarget.TargetableTexture, DestRenderTarget.ShaderResourceTexture, false, FResolveParams());
@@ -408,6 +409,7 @@ void FRCPassPostProcessDOFTemporalAA::Process(FRenderingCompositePassContext& Co
 		SrcRect.Width(), SrcRect.Height(),
 		SrcRect.Size(),
 		SrcSize,
+		*VertexShader,
 		EDRF_UseTriangleOptimization);
 
 	RHICopyToResolveTarget(DestRenderTarget.TargetableTexture, DestRenderTarget.ShaderResourceTexture, false, FResolveParams());
@@ -485,6 +487,7 @@ void FRCPassPostProcessLightShaftTemporalAA::Process(FRenderingCompositePassCont
 		SrcRect.Width(), SrcRect.Height(),
 		SrcRect.Size(),
 		SrcSize,
+		*VertexShader,
 		EDRF_UseTriangleOptimization);
 
 	RHICopyToResolveTarget(DestRenderTarget.TargetableTexture, DestRenderTarget.ShaderResourceTexture, false, FResolveParams());
@@ -543,13 +546,14 @@ void FRCPassPostProcessTemporalAA::Process(FRenderingCompositePassContext& Conte
 	uint32 Quality = FMath::Clamp(CVar->GetValueOnRenderThread(), 1, 6);
 	bool bUseFast = Quality == 3;
 
+	TShaderMapRef< FPostProcessTonemapVS >			VertexShader( GetGlobalShaderMap() );
+
 	{	// Normal temporal feedback
 		// Draw to pixels where stencil == 0
 		RHISetDepthStencilState(TStaticDepthStencilState<false,CF_Always,true,CF_Equal,SO_Keep,SO_Keep,SO_Keep>::GetRHI(), 0);
 
 		if(bUseFast)
 		{
-			TShaderMapRef< FPostProcessTonemapVS >			VertexShader( GetGlobalShaderMap() );
 			TShaderMapRef< FPostProcessTemporalAAPS<4,0> >	PixelShader( GetGlobalShaderMap() );
 
 			static FGlobalBoundShaderState BoundShaderState;
@@ -561,7 +565,6 @@ void FRCPassPostProcessTemporalAA::Process(FRenderingCompositePassContext& Conte
 		}
 		else
 		{
-			TShaderMapRef< FPostProcessTonemapVS >			VertexShader( GetGlobalShaderMap() );
 			TShaderMapRef< FPostProcessTemporalAAPS<1,0> >	PixelShader( GetGlobalShaderMap() );
 
 			static FGlobalBoundShaderState BoundShaderState;
@@ -580,16 +583,17 @@ void FRCPassPostProcessTemporalAA::Process(FRenderingCompositePassContext& Conte
 			SrcRect.Width(), SrcRect.Height(),
 			SrcRect.Size(),
 			SrcSize,
+			*VertexShader,
 			EDRF_UseTriangleOptimization);
 	}
 
 	{	// Responsive feedback for tagged pixels
 		// Draw to pixels where stencil != 0
 		RHISetDepthStencilState(TStaticDepthStencilState<false,CF_Always,true,CF_NotEqual,SO_Keep,SO_Keep,SO_Keep>::GetRHI(), 0);
-
+		
+		TShaderMapRef< FPostProcessTonemapVS >			VertexShader( GetGlobalShaderMap() );
 		if(bUseFast)
 		{
-			TShaderMapRef< FPostProcessTonemapVS >			VertexShader( GetGlobalShaderMap() );
 			TShaderMapRef< FPostProcessTemporalAAPS<4,1> >	PixelShader( GetGlobalShaderMap() );
 
 			static FGlobalBoundShaderState BoundShaderState;
@@ -601,7 +605,6 @@ void FRCPassPostProcessTemporalAA::Process(FRenderingCompositePassContext& Conte
 		}
 		else
 		{
-			TShaderMapRef< FPostProcessTonemapVS >			VertexShader( GetGlobalShaderMap() );
 			TShaderMapRef< FPostProcessTemporalAAPS<1,1> >	PixelShader( GetGlobalShaderMap() );
 
 			static FGlobalBoundShaderState BoundShaderState;
@@ -620,6 +623,7 @@ void FRCPassPostProcessTemporalAA::Process(FRenderingCompositePassContext& Conte
 			SrcRect.Width(), SrcRect.Height(),
 			SrcRect.Size(),
 			SrcSize,
+			*VertexShader,
 			EDRF_UseTriangleOptimization);
 	}
 
diff --git a/Engine/Source/Runtime/Renderer/Private/PostProcess/PostProcessTestImage.cpp b/Engine/Source/Runtime/Renderer/Private/PostProcess/PostProcessTestImage.cpp
--- a/Engine/Source/Runtime/Renderer/Private/PostProcess/PostProcessTestImage.cpp
+++ b/Engine/Source/Runtime/Renderer/Private/PostProcess/PostProcessTestImage.cpp
@@ -119,6 +119,7 @@ void FRCPassPostProcessTestImage::Process(FRenderingCompositePassContext& Contex
 		SrcRect.Width(), SrcRect.Height(),
 		DestRect.Size(),
 		GSceneRenderTargets.GetBufferSizeXY(),
+		*VertexShader,
 		EDRF_UseTriangleOptimization);
 
 	{
diff --git a/Engine/Source/Runtime/Renderer/Private/PostProcess/PostProcessTonemap.cpp b/Engine/Source/Runtime/Renderer/Private/PostProcess/PostProcessTonemap.cpp
--- a/Engine/Source/Runtime/Renderer/Private/PostProcess/PostProcessTonemap.cpp
+++ b/Engine/Source/Runtime/Renderer/Private/PostProcess/PostProcessTonemap.cpp
@@ -1037,6 +1037,8 @@ void FRCPassPostProcessTonemap::Process(FRenderingCompositePassContext& Context)
 	}
 
 	// Draw a quad mapping scene color to the view's render target
+	TShaderMapRef<FPostProcessTonemapVS> VertexShader(GetGlobalShaderMap());
+
 	DrawRectangle(
 		0, 0,
 		View.ViewRect.Width(), View.ViewRect.Height(),
@@ -1044,6 +1046,7 @@ void FRCPassPostProcessTonemap::Process(FRenderingCompositePassContext& Context)
 		View.ViewRect.Width(), View.ViewRect.Height(),
 		View.ViewRect.Size(),
 		GSceneRenderTargets.SceneColor->GetDesc().Extent,
+		*VertexShader,
 		EDRF_UseTriangleOptimization);
 
 	RHICopyToResolveTarget(DestRenderTarget.TargetableTexture, DestRenderTarget.ShaderResourceTexture, false, FResolveParams());
@@ -1391,6 +1394,8 @@ void FRCPassPostProcessTonemapES2::Process(FRenderingCompositePassContext& Conte
 	}
 
 	// Draw a quad mapping scene color to the view's render target
+	TShaderMapRef<FPostProcessTonemapVS_ES2> VertexShader(GetGlobalShaderMap());
+
 	DrawRectangle(
 		0, 0,
 		View.ViewRect.Width(), View.ViewRect.Height(),
@@ -1398,6 +1403,7 @@ void FRCPassPostProcessTonemapES2::Process(FRenderingCompositePassContext& Conte
 		View.ViewRect.Width(), View.ViewRect.Height(),
 		View.ViewRect.Size(),
 		GSceneRenderTargets.SceneColor->GetDesc().Extent,
+		*VertexShader,
 		EDRF_UseTriangleOptimization);
 
 	RHICopyToResolveTarget(DestRenderTarget.TargetableTexture, DestRenderTarget.ShaderResourceTexture, false, FResolveParams());
diff --git a/Engine/Source/Runtime/Renderer/Private/PostProcess/PostProcessUIBlur.cpp b/Engine/Source/Runtime/Renderer/Private/PostProcess/PostProcessUIBlur.cpp
--- a/Engine/Source/Runtime/Renderer/Private/PostProcess/PostProcessUIBlur.cpp
+++ b/Engine/Source/Runtime/Renderer/Private/PostProcess/PostProcessUIBlur.cpp
@@ -114,6 +114,7 @@ void FRCPassPostProcessUIBlur::Process(FRenderingCompositePassContext& Context)
 			SrcRect.Width(), SrcRect.Height(),
 			Dest->GetDesc().Extent,
 			Dest->GetDesc().Extent,
+			*VertexShader,
 			EDRF_Default);
 	}
 
diff --git a/Engine/Source/Runtime/Renderer/Private/PostProcess/PostProcessUpscale.cpp b/Engine/Source/Runtime/Renderer/Private/PostProcess/PostProcessUpscale.cpp
--- a/Engine/Source/Runtime/Renderer/Private/PostProcess/PostProcessUpscale.cpp
+++ b/Engine/Source/Runtime/Renderer/Private/PostProcess/PostProcessUpscale.cpp
@@ -168,6 +168,7 @@ void FRCPassPostProcessUpscale::Process(FRenderingCompositePassContext& Context)
 	}
 
 	// Draw a quad mapping scene color to the view's render target
+	TShaderMapRef<FPostProcessVS> VertexShader(GetGlobalShaderMap());
 	DrawRectangle(
 		0, 0,
 		DestRect.Width(), DestRect.Height(),
@@ -175,6 +176,7 @@ void FRCPassPostProcessUpscale::Process(FRenderingCompositePassContext& Context)
 		SrcRect.Width(), SrcRect.Height(),
 		DestRect.Size(),
 		SrcSize,
+		*VertexShader,
 		EDRF_UseTriangleOptimization);
 
 	RHICopyToResolveTarget(DestRenderTarget.TargetableTexture, DestRenderTarget.ShaderResourceTexture, false, FResolveParams());
diff --git a/Engine/Source/Runtime/Renderer/Private/PostProcess/PostProcessVisualizeBuffer.cpp b/Engine/Source/Runtime/Renderer/Private/PostProcess/PostProcessVisualizeBuffer.cpp
--- a/Engine/Source/Runtime/Renderer/Private/PostProcess/PostProcessVisualizeBuffer.cpp
+++ b/Engine/Source/Runtime/Renderer/Private/PostProcess/PostProcessVisualizeBuffer.cpp
@@ -153,6 +153,7 @@ void FRCPassPostProcessVisualizeBuffer::Process(FRenderingCompositePassContext&
 	SetShaderTempl<false>(Context);
 
 	// Draw a quad mapping scene color to the view's render target
+	TShaderMapRef<FPostProcessVS> VertexShader(GetGlobalShaderMap());
 	DrawRectangle(
 		0, 0,
 		DestRect.Width(), DestRect.Height(),
@@ -160,10 +161,10 @@ void FRCPassPostProcessVisualizeBuffer::Process(FRenderingCompositePassContext&
 		SrcRect.Width(), SrcRect.Height(),
 		DestRect.Size(),
 		SrcSize,
+		*VertexShader,
 		EDRF_UseTriangleOptimization);
 
 	// Now draw the requested tiles into the grid
-	TShaderMapRef<FPostProcessVS> VertexShader(GetGlobalShaderMap());
 	TShaderMapRef<FPostProcessVisualizeBufferPS<true> > PixelShader(GetGlobalShaderMap());
 
 	RHISetBlendState(TStaticBlendState<CW_RGB, BO_Add, BF_SourceAlpha, BF_InverseSourceAlpha>::GetRHI());
@@ -209,7 +210,9 @@ void FRCPassPostProcessVisualizeBuffer::Process(FRenderingCompositePassContext&
 				SrcRect.Width(), SrcRect.Height(),
 				DestRect.Size(),
 				SrcSize,
-				EDRF_Default);
+				*VertexShader,
+				EDRF_Default
+				);
 
 			Labels.Add(LabelRecord());
 			Labels.Last().Label = It->Name;
diff --git a/Engine/Source/Runtime/Renderer/Private/PostProcess/PostProcessVisualizeHDR.cpp b/Engine/Source/Runtime/Renderer/Private/PostProcess/PostProcessVisualizeHDR.cpp
--- a/Engine/Source/Runtime/Renderer/Private/PostProcess/PostProcessVisualizeHDR.cpp
+++ b/Engine/Source/Runtime/Renderer/Private/PostProcess/PostProcessVisualizeHDR.cpp
@@ -179,6 +179,7 @@ void FRCPassPostProcessVisualizeHDR::Process(FRenderingCompositePassContext& Con
 		SrcRect.Width(), SrcRect.Height(),
 		DestRect.Size(),
 		SrcSize,
+		*VertexShader,
 		EDRF_UseTriangleOptimization);
 
 
diff --git a/Engine/Source/Runtime/Renderer/Private/PostProcess/PostProcessWeightedSampleSum.cpp b/Engine/Source/Runtime/Renderer/Private/PostProcess/PostProcessWeightedSampleSum.cpp
--- a/Engine/Source/Runtime/Renderer/Private/PostProcess/PostProcessWeightedSampleSum.cpp
+++ b/Engine/Source/Runtime/Renderer/Private/PostProcess/PostProcessWeightedSampleSum.cpp
@@ -155,6 +155,7 @@ protected:
  * @param SampleOffsets - A pointer to an array of NumSamples UV offsets
  * @param SampleWeights - A pointer to an array of NumSamples 4-vector weights
  * @param NumSamples - The number of samples used by the filter.
+ * @param OutVertexShader - The vertex shader used for the filter
  */
 void SetFilterShaders(
 	FSamplerStateRHIParamRef SamplerStateRHI,
@@ -163,7 +164,8 @@ void SetFilterShaders(
 	uint32 CombineMethodInt,
 	FVector2D* SampleOffsets,
 	FLinearColor* SampleWeights,
-	uint32 NumSamples
+	uint32 NumSamples,
+	FShader** OutVertexShader
 	)
 {
 	check(CombineMethodInt <= 2);
@@ -173,6 +175,7 @@ void SetFilterShaders(
 	case NumSamples: \
 	{ \
 		TShaderMapRef<TFilterVS<NumSamples> > VertexShader(GetGlobalShaderMap()); \
+		*OutVertexShader = *VertexShader; \
 		if(CombineMethodInt == 0) \
 		{ \
 			TShaderMapRef<TFilterPS<NumSamples, 0> > PixelShader(GetGlobalShaderMap()); \
@@ -422,6 +425,7 @@ void FRCPassPostProcessWeightedSampleSum::Process(FRenderingCompositePassContext
 		}
 	}
 
+	FShader* VertexShader = nullptr;
 	SetFilterShaders(
 		TStaticSamplerState<SF_Bilinear,AM_Border,AM_Border,AM_Clamp>::GetRHI(),
 		FilterTexture,
@@ -429,7 +433,8 @@ void FRCPassPostProcessWeightedSampleSum::Process(FRenderingCompositePassContext
 		CombineMethodInt,
 		BlurOffsets,
 		BlurWeights,
-		NumSamples
+		NumSamples,
+		&VertexShader
 		);
 
 	const int NumOverrideRects = Context.View.UIBlurOverrideRectangles.Num();
@@ -450,7 +455,7 @@ void FRCPassPostProcessWeightedSampleSum::Process(FRenderingCompositePassContext
 			FIntRect SrcRect =  View.ViewRect / SrcScaleFactor;
 			FIntRect DestRect = View.ViewRect / DstScaleFactor;
 
-			DrawQuad(bDoFastBlur, SrcRect, DestRect, bRequiresClear, DestSize, SrcSize);
+			DrawQuad(bDoFastBlur, SrcRect, DestRect, bRequiresClear, DestSize, SrcSize, VertexShader);
 		}
 		break;
 		case EPostProcessRectSource::GBS_UIBlurRects:
@@ -470,7 +475,7 @@ void FRCPassPostProcessWeightedSampleSum::Process(FRenderingCompositePassContext
 				FIntRect SrcRect = ScaledQuad / SrcScaleFactor;
 				FIntRect DestRect = ScaledQuad / DstScaleFactor;
 
-				DrawQuad(bDoFastBlur, SrcRect, DestRect, bRequiresClear, DestSize, SrcSize);
+				DrawQuad(bDoFastBlur, SrcRect, DestRect, bRequiresClear, DestSize, SrcSize, VertexShader);
 			}
 		}
 		break;
@@ -554,7 +559,7 @@ bool FRCPassPostProcessWeightedSampleSum::DoFastBlur() const
 	return bRet;
 }
 
-void FRCPassPostProcessWeightedSampleSum::DrawQuad( bool bDoFastBlur, FIntRect SrcRect, FIntRect DestRect, bool bRequiresClear, FIntPoint DestSize, FIntPoint SrcSize ) const
+void FRCPassPostProcessWeightedSampleSum::DrawQuad( bool bDoFastBlur, FIntRect SrcRect, FIntRect DestRect, bool bRequiresClear, FIntPoint DestSize, FIntPoint SrcSize, FShader* VertexShader ) const
 {
 	if (bDoFastBlur)
 	{
@@ -583,6 +588,7 @@ void FRCPassPostProcessWeightedSampleSum::DrawQuad( bool bDoFastBlur, FIntRect S
 		SrcRect.Width(), SrcRect.Height(),
 		DestSize,
 		SrcSize,
+		VertexShader,
 		EDRF_UseTriangleOptimization);
 }
 
diff --git a/Engine/Source/Runtime/Renderer/Private/PostProcess/PostProcessWeightedSampleSum.h b/Engine/Source/Runtime/Renderer/Private/PostProcess/PostProcessWeightedSampleSum.h
--- a/Engine/Source/Runtime/Renderer/Private/PostProcess/PostProcessWeightedSampleSum.h
+++ b/Engine/Source/Runtime/Renderer/Private/PostProcess/PostProcessWeightedSampleSum.h
@@ -124,7 +124,7 @@ public:
 	static int GetIntegerKernelRadius(float KernelRadius);
 
 private:
-	void DrawQuad(bool bDoFastBlur, FIntRect SrcRect, FIntRect DestRect, bool bRequiresClear, FIntPoint DestSize, FIntPoint SrcSize) const;
+	void DrawQuad(bool bDoFastBlur, FIntRect SrcRect, FIntRect DestRect, bool bRequiresClear, FIntPoint DestSize, FIntPoint SrcSize, FShader* VertexShader) const;
 	static uint32 GetMaxNumSamples();
 
 	// e.g. EFS_Horiz or EFS_Vert
diff --git a/Engine/Source/Runtime/Renderer/Private/PostProcess/SceneFilterRendering.cpp b/Engine/Source/Runtime/Renderer/Private/PostProcess/SceneFilterRendering.cpp
--- a/Engine/Source/Runtime/Renderer/Private/PostProcess/SceneFilterRendering.cpp
+++ b/Engine/Source/Runtime/Renderer/Private/PostProcess/SceneFilterRendering.cpp
@@ -9,9 +9,79 @@
 #include "SceneFilterRendering.h"
 #include "PostProcessWeightedSampleSum.h"
 
-/** Vertex declaration for the light function fullscreen 2D quad. */
+/**
+* Static vertex and index buffer used for 2D screen rectangles.
+*/
+class FScreenRectangleVertexBuffer : public FVertexBuffer
+{
+public:
+	/** Initialize the RHI for this rendering resource */
+	void InitRHI() OVERRIDE
+	{
+		TResourceArray<FFilterVertex, VERTEXBUFFER_ALIGNMENT> Vertices;
+		Vertices.Init(6);
+
+		Vertices[0].Position = FVector4(1,  1,	0,	1);
+		Vertices[0].UV = FVector2D(1,	1);
+
+		Vertices[1].Position = FVector4(0,  1,	0,	1);
+		Vertices[1].UV = FVector2D(0,	1);
+
+		Vertices[2].Position = FVector4(1,	0,	0,	1);
+		Vertices[2].UV = FVector2D(1,	0);
+
+		Vertices[3].Position = FVector4(0,	0,	0,	1);
+		Vertices[3].UV = FVector2D(0,	0);
+
+		//The final two vertices are used for the triangle optimization (a single triangle spans the entire viewport )
+		Vertices[4].Position = FVector4(-1,  1,	0,	1);
+		Vertices[4].UV = FVector2D(-1,	1);
+
+		Vertices[5].Position = FVector4(1,  -1,	0,	1);
+		Vertices[5].UV = FVector2D(1, -1);
+
+		// Create vertex buffer. Fill buffer with initial data upon creation
+		VertexBufferRHI = RHICreateVertexBuffer(Vertices.GetResourceDataSize(), &Vertices, BUF_Static);
+	}
+};
+
+class FScreenRectangleIndexBuffer : public FIndexBuffer
+{
+public:
+	/** Initialize the RHI for this rendering resource */
+	void InitRHI() OVERRIDE
+	{
+		// Indices 0 - 5 are used for rendering a quad. Indices 6 - 8 are used for triangle optimization. 
+		const uint16 Indices[] = { 0, 1, 2, 2, 1, 3, 0, 4, 5 };
+		
+		TResourceArray<uint16, INDEXBUFFER_ALIGNMENT> IndexBuffer;
+		uint32 NumIndices = ARRAY_COUNT(Indices);
+		IndexBuffer.AddUninitialized(NumIndices);
+		FMemory::Memcpy(IndexBuffer.GetData(), Indices, NumIndices * sizeof(uint16));
+
+		// Create index buffer. Fill buffer with initial data upon creation
+		IndexBufferRHI = RHICreateIndexBuffer(sizeof(uint16), IndexBuffer.GetResourceDataSize(), &IndexBuffer, BUF_Static);
+	}
+};
+
+/** Global resource  */
+static TGlobalResource<FScreenRectangleVertexBuffer> GScreenRectangleVertexBuffer;
+static TGlobalResource<FScreenRectangleIndexBuffer> GScreenRectangleIndexBuffer;
+
+/** Vertex declaration for the 2D screen rectangle. */
 TGlobalResource<FFilterVertexDeclaration> GFilterVertexDeclaration;
 
+/** Uniform buffer for computing the vertex positional and UV adjustments in the vertex shader. */
+BEGIN_UNIFORM_BUFFER_STRUCT( FDrawRectangleParameters,)
+	DECLARE_UNIFORM_BUFFER_STRUCT_MEMBER( FVector4, PosScaleBias )
+	DECLARE_UNIFORM_BUFFER_STRUCT_MEMBER( FVector4, UVScaleBias )
+	DECLARE_UNIFORM_BUFFER_STRUCT_MEMBER( FVector4, InvTargetSizeAndTextureSize )
+END_UNIFORM_BUFFER_STRUCT( FDrawRectangleParameters )
+
+IMPLEMENT_UNIFORM_BUFFER_STRUCT(FDrawRectangleParameters,TEXT("DrawRectangleParameters"));
+
+typedef TUniformBufferRef<FDrawRectangleParameters> FDrawRectangleBufferRef;
+
 static void DoDrawRectangleFlagOverride(EDrawRectangleFlags& Flags)
 {
 #if !(UE_BUILD_SHIPPING || UE_BUILD_TEST)
@@ -27,7 +97,6 @@ static void DoDrawRectangleFlagOverride(EDrawRectangleFlags& Flags)
 #endif
 }
 
-
 void DrawRectangle(
 	float X,
 	float Y,
@@ -39,6 +108,7 @@ void DrawRectangle(
 	float SizeV,
 	FIntPoint TargetSize,
 	FIntPoint TextureSize,
+	FShader* VertexShader,
 	EDrawRectangleFlags Flags
 	)
 {
@@ -53,59 +123,50 @@ void DrawRectangle(
 		Flags = EDRF_Default;
 	}
 
-	if (Flags == EDRF_UseTriangleOptimization)
-	{
-		// A single triangle spans the entire viewport this results in a quad that fill the viewport. This can increase rasterization efficiency
-		// as we do not have a diagonal edge (through the center) for the rasterizer/span-dispatch. Although the actual benefit of this technique is dependent upon hardware.
-		FFilterVertex Vertices[3];
+	// Set up vertex uniform parameters for scaling and biasing the rectangle.
+	// Note: Use DrawRectangle in the vertex shader to calculate the correct vertex position and uv.
 
-		Vertices[0].Position = FVector4(X + SizeX,	Y - SizeY, ClipSpaceQuadZ,	1);
-		Vertices[1].Position = FVector4(X + SizeX,  Y + SizeY,	ClipSpaceQuadZ,	1);
-		Vertices[2].Position = FVector4(X - SizeX,  Y + SizeY ,	ClipSpaceQuadZ,	1);
+	FDrawRectangleParameters Parameters;
+	Parameters.PosScaleBias = FVector4(SizeX, SizeY, X, Y);
+	Parameters.UVScaleBias = FVector4(SizeU, SizeV, U, V);
 
-		Vertices[0].UV = FVector2D(U + SizeU, V - SizeV);
-		Vertices[1].UV = FVector2D(U + SizeU, V + SizeV);
-		Vertices[2].UV = FVector2D(U - SizeU, V + SizeV);
+	Parameters.InvTargetSizeAndTextureSize = FVector4( 
+		1.0f / TargetSize.X, 1.0f / TargetSize.Y, 
+		1.0f / TextureSize.X, 1.0f / TextureSize.Y);
 
-		for (int32 VertexIndex = 0; VertexIndex < 3; VertexIndex++)
-		{
-			Vertices[VertexIndex].Position.X = -1.0f + 2.0f * (Vertices[VertexIndex].Position.X - GPixelCenterOffset) / (float)TargetSize.X;
-			Vertices[VertexIndex].Position.Y = (+1.0f - 2.0f * (Vertices[VertexIndex].Position.Y - GPixelCenterOffset) / (float)TargetSize.Y) * GProjectionSignY;
+	SetUniformBufferParameterImmediate(VertexShader->GetVertexShader(), VertexShader->GetUniformBufferParameter<FDrawRectangleParameters>(), Parameters);
+	RHISetStreamSource(0, GScreenRectangleVertexBuffer.VertexBufferRHI, sizeof(FFilterVertex),0);
 
-			Vertices[VertexIndex].UV.X = Vertices[VertexIndex].UV.X / (float)TextureSize.X;
-			Vertices[VertexIndex].UV.Y = Vertices[VertexIndex].UV.Y / (float)TextureSize.Y;
-		}
-		
-		static const uint16 Indices[] = { 0, 1, 2 };
 
-		RHIDrawIndexedPrimitiveUP(PT_TriangleList, 0, 3, 1, Indices, sizeof(Indices[0]), Vertices, sizeof(Vertices[0]));
+	if (Flags == EDRF_UseTriangleOptimization)
+	{
+		// A single triangle spans the entire viewport this results in a quad that fill the viewport. This can increase rasterization efficiency
+		// as we do not have a diagonal edge (through the center) for the rasterizer/span-dispatch. Although the actual benefit of this technique is dependent upon hardware.
+
+		// We offset into the index buffer when using the triangle optimization to access the correct vertices.
+		RHIDrawIndexedPrimitive(
+			GScreenRectangleIndexBuffer.IndexBufferRHI,
+			PT_TriangleList,
+			/*BaseVertexIndex=*/ 0,
+			/*MinIndex=*/ 0,
+			/*NumVertices=*/ 3,
+			/*StartIndex=*/ 6,
+			/*NumPrimitives=*/ 1,
+			/*NumInstances=*/ 0
+			);
 	}
 	else
 	{
-		FFilterVertex Vertices[4];
-
-		Vertices[0].Position = FVector4(X,			Y,			ClipSpaceQuadZ,	1);
-		Vertices[1].Position = FVector4(X + SizeX,	Y,			ClipSpaceQuadZ,	1);
-		Vertices[2].Position = FVector4(X,			Y + SizeY,	ClipSpaceQuadZ,	1);
-		Vertices[3].Position = FVector4(X + SizeX,	Y + SizeY,	ClipSpaceQuadZ,	1);
-
-		Vertices[0].UV = FVector2D(U,			V);
-		Vertices[1].UV = FVector2D(U + SizeU,	V);
-		Vertices[2].UV = FVector2D(U,			V + SizeV);
-		Vertices[3].UV = FVector2D(U + SizeU,	V + SizeV);
-
-		for (int32 VertexIndex = 0; VertexIndex < 4; VertexIndex++)
-		{
-			Vertices[VertexIndex].Position.X = -1.0f + 2.0f * (Vertices[VertexIndex].Position.X - GPixelCenterOffset) / (float)TargetSize.X;
-			Vertices[VertexIndex].Position.Y = (+1.0f - 2.0f * (Vertices[VertexIndex].Position.Y - GPixelCenterOffset) / (float)TargetSize.Y) * GProjectionSignY;
-
-			Vertices[VertexIndex].UV.X = Vertices[VertexIndex].UV.X / (float)TextureSize.X;
-			Vertices[VertexIndex].UV.Y = Vertices[VertexIndex].UV.Y / (float)TextureSize.Y;
-		}
-
-		static const uint16 Indices[] = { 0, 1, 3, 0, 3, 2 };
-
-		RHIDrawIndexedPrimitiveUP(PT_TriangleList, 0, 4, 2, Indices, sizeof(Indices[0]), Vertices, sizeof(Vertices[0]));
+		RHIDrawIndexedPrimitive(
+			GScreenRectangleIndexBuffer.IndexBufferRHI,
+			PT_TriangleList,
+			/*BaseVertexIndex=*/ 0,
+			/*MinIndex=*/ 0,
+			/*NumVertices=*/ 4,
+			/*StartIndex=*/ 0,
+			/*NumPrimitives=*/ 2,
+			/*NumInstances=*/ 0
+			);
 	}
 }
 
diff --git a/Engine/Source/Runtime/Renderer/Private/PostProcess/SceneFilterRendering.h b/Engine/Source/Runtime/Renderer/Private/PostProcess/SceneFilterRendering.h
--- a/Engine/Source/Runtime/Renderer/Private/PostProcess/SceneFilterRendering.h
+++ b/Engine/Source/Runtime/Renderer/Private/PostProcess/SceneFilterRendering.h
@@ -54,12 +54,15 @@ enum EDrawRectangleFlags
  * Draws a quad with the given vertex positions and UVs in denormalized pixel/texel coordinates.
  * The platform-dependent mapping from pixels to texels is done automatically.
  * Note that the positions are affected by the current viewport.
+ * NOTE: DrawRectangle should be used in the vertex shader to calculate the correct position and uv for vertices.
+ *
  * X, Y							Position in screen pixels of the top left corner of the quad
  * SizeX, SizeY					Size in screen pixels of the quad
  * U, V							Position in texels of the top left corner of the quad's UV's
  * SizeU, SizeV					Size in texels of the quad's UV's
  * TargetSizeX, TargetSizeY		Size in screen pixels of the target surface
  * TextureSize                  Size in texels of the source texture
+ * VertexShader					The vertex shader used for rendering
  * Flags						see EDrawRectangleFlags
  */
 extern void DrawRectangle(
@@ -73,6 +76,7 @@ extern void DrawRectangle(
 	float SizeV,
 	FIntPoint TargetSize,
 	FIntPoint TextureSize,
+	FShader* VertexShader,
 	EDrawRectangleFlags Flags = EDRF_Default
 	);
 
diff --git a/Engine/Source/Runtime/Renderer/Private/PostProcess/ScreenSpaceReflections.cpp b/Engine/Source/Runtime/Renderer/Private/PostProcess/ScreenSpaceReflections.cpp
--- a/Engine/Source/Runtime/Renderer/Private/PostProcess/ScreenSpaceReflections.cpp
+++ b/Engine/Source/Runtime/Renderer/Private/PostProcess/ScreenSpaceReflections.cpp
@@ -221,6 +221,7 @@ void FRCPassPostProcessScreenSpaceReflections::Process(FRenderingCompositePassCo
 
 
 	// Draw a quad mapping scene color to the view's render target
+	TShaderMapRef< FPostProcessVS > VertexShader(GetGlobalShaderMap());
 	DrawRectangle( 
 		0, 0,
 		View.ViewRect.Width(), View.ViewRect.Height(),
@@ -228,6 +229,7 @@ void FRCPassPostProcessScreenSpaceReflections::Process(FRenderingCompositePassCo
 		View.ViewRect.Width(), View.ViewRect.Height(),
 		View.ViewRect.Size(),
 		GSceneRenderTargets.GetBufferSizeXY(),
+		*VertexShader,
 		EDRF_UseTriangleOptimization);
 
 	RHICopyToResolveTarget(DestRenderTarget.TargetableTexture, DestRenderTarget.ShaderResourceTexture, false, FResolveParams());
@@ -241,7 +243,6 @@ FPooledRenderTargetDesc FRCPassPostProcessScreenSpaceReflections::ComputeOutputD
 
 	return Ret;
 }
-
 void BuildHZB( const FViewInfo& View );
 
 void ScreenSpaceReflections( const FViewInfo& View, TRefCountPtr<IPooledRenderTarget>& SSROutput )
diff --git a/Engine/Source/Runtime/Renderer/Private/PostProcess/VisualizeTexture.cpp b/Engine/Source/Runtime/Renderer/Private/PostProcess/VisualizeTexture.cpp
--- a/Engine/Source/Runtime/Renderer/Private/PostProcess/VisualizeTexture.cpp
+++ b/Engine/Source/Runtime/Renderer/Private/PostProcess/VisualizeTexture.cpp
@@ -212,6 +212,7 @@ template<uint32 TextureType> void VisualizeTextureForTextureType(const FVisualiz
 		GSceneRenderTargets.GetBufferSizeXY(),
 		// TextureSize
 		FIntPoint(1, 1),
+		*VertexShader,
 		EDRF_UseTriangleOptimization);
 }
 
@@ -510,6 +511,7 @@ void FVisualizeTexture::PresentContent(const FSceneView& View)
 		VisualizeTextureRect.Width(), VisualizeTextureRect.Height(),
 		GSceneRenderTargets.GetBufferSizeXY(),
 		VisualizeTextureRect.Size(),
+		*VertexShader,
 		EDRF_Default);
 
 	// this is a helper class for FCanvas to be able to get screen size
diff --git a/Engine/Source/Runtime/Renderer/Private/ReflectionEnvironment.cpp b/Engine/Source/Runtime/Renderer/Private/ReflectionEnvironment.cpp
--- a/Engine/Source/Runtime/Renderer/Private/ReflectionEnvironment.cpp
+++ b/Engine/Source/Runtime/Renderer/Private/ReflectionEnvironment.cpp
@@ -606,6 +606,7 @@ void FDeferredShadingSceneRenderer::RenderDeferredReflections()
 						View.ViewRect.Width(), View.ViewRect.Height(),
 						FIntPoint(View.ViewRect.Width(), View.ViewRect.Height()),
 						GSceneRenderTargets.GetBufferSizeXY(),
+						*VertexShader,
 						EDRF_UseTriangleOptimization);
 				}
 			}
@@ -756,7 +757,8 @@ void FDeferredShadingSceneRenderer::RenderDeferredReflections()
 									View.ViewRect.Min.X, View.ViewRect.Min.Y, 
 									View.ViewRect.Width(), View.ViewRect.Height(),
 									FIntPoint(View.ViewRect.Width(), View.ViewRect.Height()),
-									GSceneRenderTargets.GetBufferSizeXY());
+									GSceneRenderTargets.GetBufferSizeXY(),
+									*VertexShader);
 							}
 						}
 					}
@@ -816,7 +818,8 @@ void FDeferredShadingSceneRenderer::RenderDeferredReflections()
 							View.ViewRect.Min.X, View.ViewRect.Min.Y, 
 							View.ViewRect.Width(), View.ViewRect.Height(),
 							FIntPoint(View.ViewRect.Width(), View.ViewRect.Height()),
-							GSceneRenderTargets.GetBufferSizeXY());
+							GSceneRenderTargets.GetBufferSizeXY(),
+							*VertexShader);
 					}
 				}
 			}
@@ -869,7 +872,8 @@ void FDeferredShadingSceneRenderer::RenderDeferredReflections()
 					View.ViewRect.Min.X, View.ViewRect.Min.Y, 
 					View.ViewRect.Width(), View.ViewRect.Height(),
 					FIntPoint(View.ViewRect.Width(), View.ViewRect.Height()),
-					GSceneRenderTargets.GetBufferSizeXY());
+					GSceneRenderTargets.GetBufferSizeXY(),
+					*VertexShader);
 			}
 		}
 		else
@@ -914,7 +918,8 @@ void FDeferredShadingSceneRenderer::RenderDeferredReflections()
 							View.ViewRect.Min.X, View.ViewRect.Min.Y, 
 							View.ViewRect.Width(), View.ViewRect.Height(),
 							FIntPoint(View.ViewRect.Width(), View.ViewRect.Height()),
-							GSceneRenderTargets.GetBufferSizeXY());
+							GSceneRenderTargets.GetBufferSizeXY(),
+							*VertexShader);
 					}
 				}
 			}
diff --git a/Engine/Source/Runtime/Renderer/Private/ReflectionEnvironmentCapture.cpp b/Engine/Source/Runtime/Renderer/Private/ReflectionEnvironmentCapture.cpp
--- a/Engine/Source/Runtime/Renderer/Private/ReflectionEnvironmentCapture.cpp
+++ b/Engine/Source/Runtime/Renderer/Private/ReflectionEnvironmentCapture.cpp
@@ -310,7 +310,8 @@ void ComputeAverageBrightness()
 		0, 0, 
 		1, 1,
 		FIntPoint(1, 1),
-		FIntPoint(1, 1));
+		FIntPoint(1, 1),
+		*VertexShader);
 }
 
 FVector4 FilterConeImportanceSamples[MAX_TEXTURE_MIP_COUNT][NumFilterSamples];
@@ -396,6 +397,7 @@ void FilterReflectionEnvironment(FSHVectorRGB3* OutIrradianceEnvironmentMap)
 					ViewRect.Width(), ViewRect.Height(),
 					FIntPoint(ViewRect.Width(), ViewRect.Height()),
 					FIntPoint(MipSize, MipSize),
+					*VertexShader,
 					EDRF_UseTriangleOptimization);
 
 				RHICopyToResolveTarget(EffectiveRT.TargetableTexture, EffectiveRT.ShaderResourceTexture, true, FResolveParams(FResolveRect(), (ECubeFace)CubeFace, MipIndex));
@@ -427,7 +429,7 @@ void FilterReflectionEnvironment(FSHVectorRGB3* OutIrradianceEnvironmentMap)
 						ViewRect.Width(), ViewRect.Height(),
 						FIntPoint(ViewRect.Width(), ViewRect.Height()),
 						FIntPoint(MipSize, MipSize),
-						EDRF_UseTriangleOptimization);
+						*VertexShader);
 
 					RHICopyToResolveTarget(EffectiveRT.TargetableTexture, EffectiveRT.ShaderResourceTexture, true, FResolveParams(FResolveRect(), (ECubeFace)CubeFace, MipIndex));
 				}
@@ -492,6 +494,7 @@ void FilterReflectionEnvironment(FSHVectorRGB3* OutIrradianceEnvironmentMap)
 					ViewRect.Width(), ViewRect.Height(),
 					FIntPoint(ViewRect.Width(), ViewRect.Height()),
 					FIntPoint(MipSize, MipSize),
+					*VertexShader,
 					EDRF_UseTriangleOptimization);
 
 				RHICopyToResolveTarget(EffectiveRT.TargetableTexture, EffectiveRT.ShaderResourceTexture, true, FResolveParams(FResolveRect(), (ECubeFace)CubeFace, MipIndex));
@@ -524,7 +527,7 @@ void FilterReflectionEnvironment(FSHVectorRGB3* OutIrradianceEnvironmentMap)
 						ViewRect.Width(), ViewRect.Height(),
 						FIntPoint(ViewRect.Width(), ViewRect.Height()),
 						FIntPoint(MipSize, MipSize),
-						EDRF_UseTriangleOptimization);
+						*VertexShader);
 
 					RHICopyToResolveTarget(EffectiveRT.TargetableTexture, EffectiveRT.ShaderResourceTexture, true, FResolveParams(FResolveRect(), (ECubeFace)CubeFace, MipIndex));
 				}
@@ -810,6 +813,7 @@ void CaptureSceneToScratchCubemap(FSceneRenderer* SceneRenderer, ECubeFace CubeF
 				ViewRect.Width() * GSupersampleCaptureFactor, ViewRect.Height() * GSupersampleCaptureFactor,
 				FIntPoint(ViewRect.Width(), ViewRect.Height()),
 				GSceneRenderTargets.GetBufferSizeXY(),
+				*VertexShader,
 				EDRF_UseTriangleOptimization);
 
 			RHICopyToResolveTarget(EffectiveColorRT.TargetableTexture, EffectiveColorRT.ShaderResourceTexture, true, FResolveParams(FResolveRect()));
@@ -839,7 +843,7 @@ void CaptureSceneToScratchCubemap(FSceneRenderer* SceneRenderer, ECubeFace CubeF
 					ViewRect.Width() * GSupersampleCaptureFactor, ViewRect.Height() * GSupersampleCaptureFactor,
 					FIntPoint(ViewRect.Width(), ViewRect.Height()),
 					GSceneRenderTargets.GetBufferSizeXY(),
-					EDRF_UseTriangleOptimization);
+					*VertexShader);
 
 				RHICopyToResolveTarget(EffectiveColorRT.TargetableTexture, EffectiveColorRT.ShaderResourceTexture, true, FResolveParams(FResolveRect(), CubeFace));
 			}
@@ -884,7 +888,8 @@ void CopyCubemapToScratchCubemap(UTextureCube* SourceCubemap, bool bIsSkyLight,
 				0, 0, 
 				SourceDimensions.X, SourceDimensions.Y,
 				FIntPoint(ViewRect.Width(), ViewRect.Height()),
-				SourceDimensions);
+				SourceDimensions,
+				*VertexShader);
 		}
 
 		RHICopyToResolveTarget(EffectiveColorRT.TargetableTexture, EffectiveColorRT.ShaderResourceTexture, true, FResolveParams(FResolveRect()));
@@ -916,7 +921,8 @@ void CopyCubemapToScratchCubemap(UTextureCube* SourceCubemap, bool bIsSkyLight,
 				0, 0, 
 				SourceDimensions.X, SourceDimensions.Y,
 				FIntPoint(ViewRect.Width(), ViewRect.Height()),
-				SourceDimensions);
+				SourceDimensions,
+				*VertexShader);
 
 			RHICopyToResolveTarget(EffectiveColorRT.TargetableTexture, EffectiveColorRT.ShaderResourceTexture, true, FResolveParams(FResolveRect(), (ECubeFace)CubeFace));
 		}
diff --git a/Engine/Source/Runtime/Renderer/Private/ReflectionEnvironmentDiffuseIrradiance.cpp b/Engine/Source/Runtime/Renderer/Private/ReflectionEnvironmentDiffuseIrradiance.cpp
--- a/Engine/Source/Runtime/Renderer/Private/ReflectionEnvironmentDiffuseIrradiance.cpp
+++ b/Engine/Source/Runtime/Renderer/Private/ReflectionEnvironmentDiffuseIrradiance.cpp
@@ -264,7 +264,8 @@ void ComputeDiffuseIrradiance(FTextureRHIRef LightingSource, int32 LightingSourc
 					ViewRect.Min.X, ViewRect.Min.Y, 
 					ViewRect.Width(), ViewRect.Height(),
 					FIntPoint(ViewRect.Width(), ViewRect.Height()),
-					FIntPoint(MipSize, MipSize));
+					FIntPoint(MipSize, MipSize),
+					*VertexShader);
 
 				RHICopyToResolveTarget(EffectiveRT.TargetableTexture, EffectiveRT.ShaderResourceTexture, true, FResolveParams(FResolveRect(), (ECubeFace)CubeFace, MipIndex));
 			}
@@ -310,7 +311,8 @@ void ComputeDiffuseIrradiance(FTextureRHIRef LightingSource, int32 LightingSourc
 						ViewRect.Min.X, ViewRect.Min.Y, 
 						ViewRect.Width(), ViewRect.Height(),
 						FIntPoint(ViewRect.Width(), ViewRect.Height()),
-						FIntPoint(MipSize, MipSize));
+						FIntPoint(MipSize, MipSize),
+						*VertexShader);
 
 					RHICopyToResolveTarget(EffectiveRT.TargetableTexture, EffectiveRT.ShaderResourceTexture, true, FResolveParams(FResolveRect(), (ECubeFace)CubeFace, MipIndex));
 				}
@@ -343,7 +345,8 @@ void ComputeDiffuseIrradiance(FTextureRHIRef LightingSource, int32 LightingSourc
 				0, 0, 
 				MipSize, MipSize,
 				FIntPoint(FSHVector3::MaxSHBasis, 1),
-				FIntPoint(MipSize, MipSize));
+				FIntPoint(MipSize, MipSize),
+				*VertexShader);
 
 			RHICopyToResolveTarget(EffectiveRT.TargetableTexture, EffectiveRT.ShaderResourceTexture, true, FResolveParams());
 		}
diff --git a/Engine/Source/Runtime/Renderer/Private/SceneCaptureRendering.cpp b/Engine/Source/Runtime/Renderer/Private/SceneCaptureRendering.cpp
--- a/Engine/Source/Runtime/Renderer/Private/SceneCaptureRendering.cpp
+++ b/Engine/Source/Runtime/Renderer/Private/SceneCaptureRendering.cpp
@@ -65,6 +65,7 @@ void UpdateSceneCaptureContent_RenderThread(FSceneRenderer* SceneRenderer, FText
 				ViewRect.Width(), ViewRect.Height(),
 				TargetSize,
 				GSceneRenderTargets.GetBufferSizeXY(),
+				*VertexShader,
 				EDRF_UseTriangleOptimization);
 		}
 
diff --git a/Engine/Source/Runtime/Renderer/Private/SceneOcclusion.cpp b/Engine/Source/Runtime/Renderer/Private/SceneOcclusion.cpp
--- a/Engine/Source/Runtime/Renderer/Private/SceneOcclusion.cpp
+++ b/Engine/Source/Runtime/Renderer/Private/SceneOcclusion.cpp
@@ -763,6 +763,7 @@ void FHZBOcclusionTester::Submit( const FViewInfo& View )
 			SizeX, SizeY,
 			FIntPoint( SizeX, SizeY ),
 			FIntPoint( SizeX, SizeY ),
+			*VertexShader,
 			EDRF_UseTriangleOptimization);
 	}
 
@@ -904,6 +905,7 @@ void BuildHZB( const FViewInfo& View )
 			View.ViewRect.Width(), View.ViewRect.Height(),
 			HZBSize,
 			GSceneRenderTargets.GetBufferSizeXY(),
+			*VertexShader,
 			EDRF_UseTriangleOptimization);
 	}
 
@@ -934,6 +936,7 @@ void BuildHZB( const FViewInfo& View )
 			SrcSize.X, SrcSize.Y,
 			DstSize,
 			SrcSize,
+			*VertexShader,
 			EDRF_UseTriangleOptimization);
 
 		RHICopyToResolveTarget( HZBRenderTarget.TargetableTexture, HZBRenderTarget.ShaderResourceTexture, false, FResolveParams(FResolveRect(), CubeFace_PosX, MipIndex) );
diff --git a/Engine/Source/Runtime/Renderer/Private/ScenePostProcessing.cpp b/Engine/Source/Runtime/Renderer/Private/ScenePostProcessing.cpp
--- a/Engine/Source/Runtime/Renderer/Private/ScenePostProcessing.cpp
+++ b/Engine/Source/Runtime/Renderer/Private/ScenePostProcessing.cpp
@@ -158,5 +158,6 @@ void FSceneRenderer::GammaCorrectToViewportRenderTarget(const FViewInfo* View, f
 		View->ViewRect.Width(),View->ViewRect.Height(),
 		ViewFamily.RenderTarget->GetSizeXY(),
 		GSceneRenderTargets.GetBufferSizeXY(),
+		*VertexShader,
 		EDRF_UseTriangleOptimization);
 }
diff --git a/Engine/Source/Runtime/Renderer/Private/ShaderComplexityRendering.cpp b/Engine/Source/Runtime/Renderer/Private/ShaderComplexityRendering.cpp
--- a/Engine/Source/Runtime/Renderer/Private/ShaderComplexityRendering.cpp
+++ b/Engine/Source/Runtime/Renderer/Private/ShaderComplexityRendering.cpp
@@ -201,6 +201,7 @@ void FRCPassPostProcessVisualizeComplexity::Process(FRenderingCompositePassConte
 		SrcRect.Width(), SrcRect.Height(),
 		DestRect.Size(),
 		SrcSize,
+		*VertexShader,
 		EDRF_UseTriangleOptimization);
 
 	RHICopyToResolveTarget(DestRenderTarget.TargetableTexture, DestRenderTarget.ShaderResourceTexture, false, FResolveParams());
diff --git a/Engine/Source/Runtime/Renderer/Private/TiledDeferredLightRendering.cpp b/Engine/Source/Runtime/Renderer/Private/TiledDeferredLightRendering.cpp
--- a/Engine/Source/Runtime/Renderer/Private/TiledDeferredLightRendering.cpp
+++ b/Engine/Source/Runtime/Renderer/Private/TiledDeferredLightRendering.cpp
@@ -352,6 +352,7 @@ void FDeferredShadingSceneRenderer::RenderTiledDeferredLighting(const TArray<FSo
 						View.ViewRect.Width(), View.ViewRect.Height(),
 						View.ViewRect.Size(),
 						GSceneRenderTargets.GetBufferSizeXY(),
+						*ScreenVertexShader,
 						EDRF_UseTriangleOptimization);
 				}
 			}
diff --git a/Engine/Source/Runtime/Renderer/Private/TranslucentRendering.cpp b/Engine/Source/Runtime/Renderer/Private/TranslucentRendering.cpp
--- a/Engine/Source/Runtime/Renderer/Private/TranslucentRendering.cpp
+++ b/Engine/Source/Runtime/Renderer/Private/TranslucentRendering.cpp
@@ -191,6 +191,7 @@ void FDeferredShadingSceneRenderer::CopySceneColor(const FViewInfo& View, const
 		View.ViewRect.Width(), View.ViewRect.Height(),
 		FIntPoint(View.ViewRect.Width(), View.ViewRect.Height()),
 		GSceneRenderTargets.GetBufferSizeXY(),
+		*ScreenVertexShader,
 		EDRF_UseTriangleOptimization);
 
 	GSceneRenderTargets.FinishRenderingLightAttenuation();
diff --git a/Engine/Source/Runtime/Windows/D3D11RHI/Private/D3D11Commands.cpp b/Engine/Source/Runtime/Windows/D3D11RHI/Private/D3D11Commands.cpp
--- a/Engine/Source/Runtime/Windows/D3D11RHI/Private/D3D11Commands.cpp
+++ b/Engine/Source/Runtime/Windows/D3D11RHI/Private/D3D11Commands.cpp
@@ -1601,7 +1601,7 @@ void FD3D11DynamicRHI::RHIClearMRT(bool bClearColor,int32 NumClearColors,const F
 		OriginalResourceState.ClearCurrentVertexResources(StateCache);		
 
 		// Set the new shaders
-		TShaderMapRef<FOneColorVS> VertexShader(GetGlobalShaderMap());
+		TShaderMapRef<TOneColorVS<true> > VertexShader(GetGlobalShaderMap());
 
 		FOneColorPS* PixelShader = NULL;
 
diff --git a/Engine/Source/Runtime/Windows/D3D11RHI/Private/D3D11RHI.cpp b/Engine/Source/Runtime/Windows/D3D11RHI/Private/D3D11RHI.cpp
--- a/Engine/Source/Runtime/Windows/D3D11RHI/Private/D3D11RHI.cpp
+++ b/Engine/Source/Runtime/Windows/D3D11RHI/Private/D3D11RHI.cpp
@@ -205,7 +205,7 @@ void FD3D11DynamicRHI::IssueLongGPUTask()
 			RHISetDepthStencilState(TStaticDepthStencilState<false,CF_Always>::GetRHI(), 0);
 			RHISetRasterizerState(TStaticRasterizerState<FM_Solid,CM_None>::GetRHI());
 
-			TShaderMapRef<FOneColorVS> VertexShader(GetGlobalShaderMap());
+			TShaderMapRef<TOneColorVS<true> > VertexShader(GetGlobalShaderMap());
 			TShaderMapRef<FLongGPUTaskPS> PixelShader(GetGlobalShaderMap());
 
 			SetGlobalBoundShaderState(LongGPUTaskBoundShaderState, GD3D11Vector4VertexDeclaration.VertexDeclarationRHI, *VertexShader, *PixelShader);
