#ue4 Initial Morpheus support.  Head tracking disabled on PS4.

diff --git a/Engine/Config/BaseEngine.ini b/Engine/Config/BaseEngine.ini
--- a/Engine/Config/BaseEngine.ini
+++ b/Engine/Config/BaseEngine.ini
@@ -1991,6 +1991,7 @@ DebugComponentKey=Quote
 +EnabledPlugins="PluginsEditor"
 +EnabledPlugins="EpicSurvey"
 +EnabledPlugins="OculusRift"
++EnabledPlugins="Morpheus"
 +EnabledPlugins="CustomMeshComponent"
 +EnabledPlugins="MessagingDebugger"
 +EnabledPlugins="PerforceSourceControl"
diff --git a/Engine/Shaders/MorpheusInclude.usf b/Engine/Shaders/MorpheusInclude.usf
new file mode 100644
--- /dev/null
+++ b/Engine/Shaders/MorpheusInclude.usf
@@ -0,0 +1,5 @@
+// Copyright 1998-2014 Epic Games, Inc. All Rights Reserved.
+
+#include "Common.usf"
+#include "PostProcessCommon.usf"
+#include "PS4/PostProcessHMDMorpheus.usf"
\ No newline at end of file
diff --git a/Engine/Source/Programs/UnrealBuildTool/Windows/UEBuildWindows.cs b/Engine/Source/Programs/UnrealBuildTool/Windows/UEBuildWindows.cs
--- a/Engine/Source/Programs/UnrealBuildTool/Windows/UEBuildWindows.cs
+++ b/Engine/Source/Programs/UnrealBuildTool/Windows/UEBuildWindows.cs
@@ -284,6 +284,14 @@ namespace UnrealBuildTool
 			}
 			InBuildTarget.GlobalCompileEnvironment.Config.Definitions.Add("PLATFORM_WINDOWS=1");
 
+			var ProjectSourceDirectory = Path.Combine( UnrealBuildTool.GetUProjectPath(), "Source" );
+
+			String MorpheusShaderPath = Path.Combine(BuildConfiguration.RelativeEnginePath, "Shaders/PS4/PostProcessHMDMorpheus.usf");
+			if (File.Exists(MorpheusShaderPath))
+			{
+				InBuildTarget.GlobalCompileEnvironment.Config.Definitions.Add("HAS_MORPHEUS=1");
+			}
+
 			if (InBuildTarget.Rules != null)
 			{
 				// Explicitly exclude the MS C++ runtime libraries we're not using, to ensure other libraries we link with use the same
diff --git a/Engine/Source/Runtime/HeadMountedDisplay/Public/IHeadMountedDisplay.h b/Engine/Source/Runtime/HeadMountedDisplay/Public/IHeadMountedDisplay.h
--- a/Engine/Source/Runtime/HeadMountedDisplay/Public/IHeadMountedDisplay.h
+++ b/Engine/Source/Runtime/HeadMountedDisplay/Public/IHeadMountedDisplay.h
@@ -1,4 +1,4 @@
-// Copyright 1998-2014 Epic Games, Inc. All Rights Reserved.
+// Copyright 1998-2012 Epic Games, Inc. All Rights Reserved.
 
 #pragma once
 
@@ -11,13 +11,15 @@ namespace EHMDDeviceType
 {
 	enum Type
 	{
-		DT_OculusRift
+		DT_OculusRift,
+		DT_Morpheus
 	};
 }
 
 /**
  * HMD device interface
  */
+
 class HEADMOUNTEDDISPLAY_API IHeadMountedDisplay : public IModuleInterface, public IStereoRendering
 {
 
@@ -45,6 +47,15 @@ public:
 		size_t  MonitorId;
 		int		DesktopX, DesktopY;
 		int		ResolutionX, ResolutionY;
+
+		MonitorInfo() : MonitorId(0)
+			, DesktopX(0)
+			, DesktopY(0)
+			, ResolutionX(0)
+			, ResolutionY(0)
+		{
+		}
+
 	};
 
     /**
@@ -70,7 +81,7 @@ public:
 	/**
 	 * If the HMD supports positional tracking via a camera, this returns the frustum properties (all in game-world space) of the tracking camera.
 	 */
-	virtual void	GetPositionalTrackingCameraProperties(FVector& OutOrigin, FRotator& OutOrientation, float& OutHFOV, float& OutVFOV, float& OutCameraDistance, float& OutNearPlane, float& OutFarPlane) const = 0;
+	virtual void	GetPositionalTrackingCameraProperties(FVector& OutOrigin, FRotator& OutOrientation, float& OutHFOV, float& OutVFOV, float& OutCameraDistance, float& OutNearPlane, float& OutFarPlane) const = 0;	
 
 	/**
 	 * Accessors to modify the interpupillary distance (meters)
@@ -156,7 +167,7 @@ public:
 	/**
 	 * Acquires color info for latency tester rendering. Returns true if rendering should be performed.
 	 */
-	virtual bool GetLatencyTesterColor_RenderThread(FColor& color, const FSceneView& view) =  0;
+	virtual bool GetLatencyTesterColor_RenderThread(FColor& Color, const FSceneView& View) {return false;}
 
 	/**
 	 * Returns true, if head tracking is allowed. Most common case: it returns true when GEngine->IsStereoscopic3D() is true,
@@ -191,6 +202,15 @@ public:
 	 * any drawing occurs. It is called at the beginning of UGameViewportClient::Draw() method.
 	 */
 	virtual void UpdateScreenSettings(const FViewport* InViewport) = 0;
+
+	/** 
+	 * Additional optional distorion rendering parameters
+	 * @todo:  Once we can move shaders into plugins, remove these!
+	 */
+	virtual void GetImageTranslation(float& x, float& y) const {}
+	virtual void GetDistortionCenterOffset(float& x, float& y) const {}
+	virtual const FTexture*  GetDistortionTexture() {return NULL;}	
+
 private:
 	/** Stores the dimensions of the window before we moved into fullscreen mode, so they can be restored */
 	FSlateRect PreFullScreenRect;
diff --git a/Engine/Source/Runtime/Renderer/Private/PostProcess/PostProcessMorpheus.cpp b/Engine/Source/Runtime/Renderer/Private/PostProcess/PostProcessMorpheus.cpp
new file mode 100644
--- /dev/null
+++ b/Engine/Source/Runtime/Renderer/Private/PostProcess/PostProcessMorpheus.cpp
@@ -0,0 +1,263 @@
+// Copyright 1998-2013 Epic Games, Inc. All Rights Reserved.
+
+/*=============================================================================
+	PostProcessMorpheus.cpp: Post processing for Sony Morpheus HMD device.
+=============================================================================*/
+
+#include "RendererPrivate.h"
+#include "ScenePrivate.h"
+#include "SceneFilterRendering.h"
+#include "PostProcessMorpheus.h"
+#include "PostProcessing.h"
+#include "PostProcessHistogram.h"
+#include "PostProcessEyeAdaptation.h"
+#include "IHeadMountedDisplay.h"
+
+DEFINE_LOG_CATEGORY_STATIC(LogMorpheusHMDPostProcess, All, All);
+
+#if HAS_MORPHEUS
+
+/** Encapsulates the post processing HMD distortion and correction pixel shader. */
+class FPostProcessMorpheusPS : public FGlobalShader
+{
+	DECLARE_SHADER_TYPE(FPostProcessMorpheusPS, Global);
+
+	static bool ShouldCache(EShaderPlatform Platform)
+	{
+		return true;
+	}
+
+	/** Default constructor. */
+	FPostProcessMorpheusPS()
+	{
+		DistortionTexture = NULL;
+	}
+
+public:
+	FPostProcessPassParameters PostprocessParameter;
+	FDeferredPixelShaderParameters DeferredParameters;
+	
+	// Distortion parameter values
+	FShaderParameter LensCenter;
+	FShaderParameter ScreenCenter;
+	FShaderParameter Scale;
+	FShaderParameter HMDWarpParam;
+	FShaderParameter CAWarpParam;
+	FShaderResourceParameter DistortionTextureParam; 
+	FShaderResourceParameter DistortionTextureSampler; 
+	const FTexture*		DistortionTexture;
+
+
+	/** Initialization constructor. */
+	FPostProcessMorpheusPS(const ShaderMetaType::CompiledShaderInitializerType& Initializer)
+		: FGlobalShader(Initializer)
+	{
+		PostprocessParameter.Bind(Initializer.ParameterMap);
+		DeferredParameters.Bind(Initializer.ParameterMap);
+
+		LensCenter.Bind(Initializer.ParameterMap, TEXT("LensCenter"));
+		ScreenCenter.Bind(Initializer.ParameterMap, TEXT("ScreenCenter"));
+		Scale.Bind(Initializer.ParameterMap, TEXT("Scale"));
+		HMDWarpParam.Bind(Initializer.ParameterMap, TEXT("HMDWarpParam"));
+		CAWarpParam.Bind(Initializer.ParameterMap, TEXT("CAWarpParam"));
+	
+		DistortionTextureParam.Bind(Initializer.ParameterMap, TEXT("DistortionTexture"));
+		check(DistortionTextureParam.IsBound());
+
+		DistortionTextureSampler.Bind(Initializer.ParameterMap, TEXT("DistortionTextureSampler"));
+		check(DistortionTextureSampler.IsBound());
+
+		DistortionTexture = NULL;
+	}
+
+
+	void SetPS(const FRenderingCompositePassContext& Context, FIntRect SrcRect, FIntPoint SrcBufferSize, EStereoscopicPass StereoPass, FMatrix& QuadTexTransform)
+	{
+		const FPixelShaderRHIParamRef ShaderRHI = GetPixelShader();
+		
+		FGlobalShader::SetParameters(ShaderRHI, Context.View);
+
+		PostprocessParameter.SetPS(ShaderRHI, Context, TStaticSamplerState<SF_Bilinear,AM_Clamp,AM_Clamp,AM_Clamp>::GetRHI());
+		DeferredParameters.Set(ShaderRHI, Context.View);
+
+		{
+			check(GEngine->HMDDevice.IsValid());
+
+            const FIntPoint SrcSize = SrcRect.Size();
+            const float BufferRatioX = float(SrcSize.X)/float(SrcBufferSize.X);
+            const float BufferRatioY = float(SrcSize.Y)/float(SrcBufferSize.Y);
+            const float w = float(SrcSize.X)/float(SrcBufferSize.X);
+            const float h = float(SrcSize.Y)/float(SrcBufferSize.Y);
+            const float x = float(SrcRect.Min.X)/float(SrcBufferSize.X);
+            const float y = float(SrcRect.Min.Y)/float(SrcBufferSize.Y);
+
+			float DistortionCenterOffsetX = 0, DistortionCenterOffsetY = 0;
+			GEngine->HMDDevice->GetDistortionCenterOffset(DistortionCenterOffsetX, DistortionCenterOffsetY);
+
+			const float XCenterOffset = (StereoPass == eSSP_RIGHT_EYE) ? -DistortionCenterOffsetX : DistortionCenterOffsetX;
+			const float YCenterOffset = DistortionCenterOffsetY;
+
+			const float AspectRatio = (float)SrcSize.X / (float)SrcSize.Y;
+			const float ScaleFactor = GEngine->HMDDevice->GetDistortionScalingFactor();
+
+			if(DistortionTexture == NULL)
+			{
+				check(IsInRenderingThread());
+				DistortionTexture = GEngine->HMDDevice->GetDistortionTexture();
+				check(DistortionTexture != NULL);
+			}
+			SetTextureParameter(ShaderRHI, DistortionTextureParam, DistortionTextureSampler, TStaticSamplerState<SF_Bilinear,AM_Clamp,AM_Clamp,AM_Clamp>::GetRHI(), DistortionTexture->TextureRHI);
+
+			// Shifts texture coordinates to the center of the distortion function around the center of the lens
+			FVector2D ViewLensCenter;
+			ViewLensCenter.X = x + (0.5f + XCenterOffset) * w;
+			ViewLensCenter.Y = y + (0.5f + YCenterOffset) * h;
+			SetShaderValue(ShaderRHI, LensCenter, ViewLensCenter);
+
+			// Texture coordinate for the center of the half scene texture, used to clamp sampling
+			FVector2D ViewScreenCenter;
+			ViewScreenCenter.X = x + w * 0.5f;
+			ViewScreenCenter.Y = y + h * 0.5f;
+			SetShaderValue(ShaderRHI, ScreenCenter, ViewScreenCenter);
+			
+			// Rescale output (sample) coordinates back to texture range and increase scale to support rendering outside the the screen
+			FVector2D ViewScale;
+			ViewScale.X = (w/2) * ScaleFactor;
+			ViewScale.Y = (h/2) * ScaleFactor * AspectRatio;
+            SetShaderValue(ShaderRHI, Scale, ViewScale);
+
+            // Distortion coefficients
+            FVector4 DistortionValues;
+            GEngine->HMDDevice->GetDistortionWarpValues(DistortionValues);
+            SetShaderValue(ShaderRHI, HMDWarpParam, DistortionValues);
+
+			// CNN - Morpheus changes
+			// CA correction values
+			FVector4 CAValues;
+			GEngine->HMDDevice->GetChromaAbCorrectionValues(CAValues);
+			SetShaderValue(ShaderRHI, CAWarpParam, CAValues);
+
+			// Rescale the texture coordinates to [-1,1] unit range and corrects aspect ratio
+			FVector2D ViewScaleIn;
+            ViewScaleIn.X = (2/w);
+            ViewScaleIn.Y = (2/h) / AspectRatio;
+            
+            QuadTexTransform = FMatrix::Identity;
+            QuadTexTransform *= FTranslationMatrix(FVector(-ViewLensCenter.X * SrcBufferSize.X, -ViewLensCenter.Y * SrcBufferSize.Y, 0));
+            QuadTexTransform *= FMatrix(FPlane(ViewScaleIn.X,0,0,0), FPlane(0,ViewScaleIn.Y,0,0), FPlane(0,0,0,0), FPlane(0,0,0,1));
+		}
+	}
+	
+	// FShader interface.
+	virtual bool Serialize(FArchive& Ar)
+	{
+		bool bShaderHasOutdatedParameters = FGlobalShader::Serialize(Ar);
+		Ar << PostprocessParameter << DeferredParameters << LensCenter << ScreenCenter << Scale << HMDWarpParam << CAWarpParam << DistortionTextureParam << DistortionTextureSampler;
+		return bShaderHasOutdatedParameters;
+	}
+};
+
+/** Encapsulates the post processing vertex shader. */
+class FPostProcessMorpheusVS : public FGlobalShader
+{
+	DECLARE_SHADER_TYPE(FPostProcessMorpheusVS,Global);
+
+	static bool ShouldCache(EShaderPlatform Platform)
+	{
+		return true;
+	}
+
+	/** Default constructor. */
+	FPostProcessMorpheusVS() {}
+
+	/** to have a similar interface as all other shaders */
+	void SetParameters(const FRenderingCompositePassContext& Context)
+	{
+		FGlobalShader::SetParameters(GetVertexShader(), Context.View);
+	}
+
+	void SetParameters(const FSceneView& View)
+	{
+		FGlobalShader::SetParameters(GetVertexShader(), View);
+	}
+
+public:
+
+	/** Initialization constructor. */
+	FPostProcessMorpheusVS(const ShaderMetaType::CompiledShaderInitializerType& Initializer)
+		: FGlobalShader(Initializer)
+	{
+	}
+};
+
+IMPLEMENT_SHADER_TYPE(, FPostProcessMorpheusVS, TEXT("MorpheusInclude"), TEXT("MainVS"), SF_Vertex);
+IMPLEMENT_SHADER_TYPE(, FPostProcessMorpheusPS,TEXT("MorpheusInclude"),TEXT("MainPS"),SF_Pixel);
+
+void FRCPassPostProcessMorpheus::Process(FRenderingCompositePassContext& Context)
+{
+	SCOPED_DRAW_EVENT(PostProcessMorpheus, DEC_SCENE_ITEMS);
+	const FPooledRenderTargetDesc* InputDesc = GetInputDesc(ePId_Input0);
+
+	if(!InputDesc)
+	{
+		// input is not hooked up correctly
+		return;
+	}
+
+	const FSceneView& View = Context.View;
+	const FSceneViewFamily& ViewFamily = *(View.Family);
+	
+	FIntRect SrcRect = View.ViewRect;
+	FIntRect DestRect = View.UnscaledViewRect;
+	FIntPoint SrcSize = InputDesc->Extent;
+
+    const FSceneRenderTargetItem& DestRenderTarget = PassOutputs[0].RequestSurface(Context);
+
+	// Set the view family's render target/viewport.
+	RHISetRenderTarget(DestRenderTarget.TargetableTexture, FTextureRHIRef());	
+
+	Context.SetViewportAndCallRHI(DestRect);
+
+	// set the state
+	RHISetBlendState(TStaticBlendState<>::GetRHI());
+	RHISetRasterizerState(TStaticRasterizerState<>::GetRHI());
+	RHISetDepthStencilState(TStaticDepthStencilState<false,CF_Always>::GetRHI());
+
+	TShaderMapRef<FPostProcessMorpheusVS> VertexShader(GetGlobalShaderMap());
+	TShaderMapRef<FPostProcessMorpheusPS> PixelShader(GetGlobalShaderMap());
+
+	static FGlobalBoundShaderState BoundShaderState;
+
+	SetGlobalBoundShaderState(BoundShaderState, GFilterVertexDeclaration.VertexDeclarationRHI, *VertexShader, *PixelShader);
+
+    FMatrix QuadTexTransform;
+    FMatrix QuadPosTransform = FMatrix::Identity;
+
+	PixelShader->SetPS(Context, SrcRect, SrcSize, View.StereoPass, QuadTexTransform);
+
+	// Draw a quad mapping scene color to the view's render target
+    DrawTransformedRectangle(
+        0, 0,
+        DestRect.Width(), DestRect.Height(),
+        QuadPosTransform,
+        SrcRect.Min.X, SrcRect.Min.Y,
+        SrcRect.Width(), SrcRect.Height(),
+        QuadTexTransform,
+        DestRect.Size(),
+        SrcSize
+        );
+
+	RHICopyToResolveTarget(DestRenderTarget.TargetableTexture, DestRenderTarget.ShaderResourceTexture, false, FResolveParams());
+}
+
+FPooledRenderTargetDesc FRCPassPostProcessMorpheus::ComputeOutputDesc(EPassOutputId InPassOutputId) const
+{
+	FPooledRenderTargetDesc Ret = PassOutputs[0].RenderTargetDesc;
+
+	Ret.NumSamples = 1;	// no MSAA
+	Ret.DebugName = TEXT("Morpheus");
+
+	return Ret;
+}
+
+#endif // HAS_MORPHEUS
\ No newline at end of file
diff --git a/Engine/Source/Runtime/Renderer/Private/PostProcess/PostProcessMorpheus.h b/Engine/Source/Runtime/Renderer/Private/PostProcess/PostProcessMorpheus.h
new file mode 100644
--- /dev/null
+++ b/Engine/Source/Runtime/Renderer/Private/PostProcess/PostProcessMorpheus.h
@@ -0,0 +1,24 @@
+// Copyright 1998-2013 Epic Games, Inc. All Rights Reserved.
+
+/*=============================================================================
+	PostProcessMorpheus.h: Post processing for Sony Morpheus HMD device 
+=============================================================================*/
+
+#pragma once
+
+#include "RenderingCompositionGraph.h"
+
+#if HAS_MORPHEUS
+
+// derives from TRenderingCompositePassBase<InputCount, OutputCount>
+// ePId_Input0: SceneColor
+class FRCPassPostProcessMorpheus : public TRenderingCompositePassBase<1, 1>
+{
+public:
+	// interface FRenderingCompositePass ---------
+	virtual void Process(FRenderingCompositePassContext& Context);
+	virtual void Release() { delete this; }
+	FPooledRenderTargetDesc ComputeOutputDesc(EPassOutputId InPassOutputId) const;
+};
+
+#endif
diff --git a/Engine/Source/Runtime/Renderer/Private/PostProcess/PostProcessing.cpp b/Engine/Source/Runtime/Renderer/Private/PostProcess/PostProcessing.cpp
--- a/Engine/Source/Runtime/Renderer/Private/PostProcess/PostProcessing.cpp
+++ b/Engine/Source/Runtime/Renderer/Private/PostProcess/PostProcessing.cpp
@@ -1,4 +1,4 @@
-// Copyright 1998-2014 Epic Games, Inc. All Rights Reserved.
+// Copyright 1998-2013 Epic Games, Inc. All Rights Reserved.
 
 /*=============================================================================
 	PostProcessing.cpp: The center for all post processing activities.
@@ -45,6 +45,9 @@
 #include "PostProcessUIBlur.h"
 #include "HighResScreenshot.h"
 
+#include "PostProcessMorpheus.h"
+
+
 /** The global center for all post processing activities. */
 FPostProcessing GPostProcessing;
 
@@ -1060,12 +1063,26 @@ void FPostProcessing::Process(const FViewInfo& View, TRefCountPtr<IPooledRenderT
 		AddGBufferVisualization(Context, SeparateTranslucency);
 
 		bool bStereoRenderingAndHMD = View.Family->EngineShowFlags.StereoRendering && View.Family->EngineShowFlags.HMDDistortion;
-
 		if (bStereoRenderingAndHMD)
 		{
-			FRenderingCompositePass* Node = Context.Graph.RegisterPass(new FRCPassPostProcessHMD());
-			Node->SetInput(ePId_Input0, FRenderingCompositeOutputRef(Context.FinalOutput));
-			Context.FinalOutput = FRenderingCompositeOutputRef(Node);
+			FRenderingCompositePass* Node = NULL;
+			const EHMDDeviceType::Type DeviceType = GEngine->HMDDevice->GetHMDDeviceType();
+			if(DeviceType == EHMDDeviceType::DT_OculusRift)
+			{
+				Node = Context.Graph.RegisterPass(new FRCPassPostProcessHMD());
+			}
+#if HAS_MORPHEUS
+			else if(DeviceType == EHMDDeviceType::DT_Morpheus)
+			{
+				Node = Context.Graph.RegisterPass(new FRCPassPostProcessMorpheus());
+			}
+#endif
+			
+			if(Node)
+			{
+				Node->SetInput(ePId_Input0, FRenderingCompositeOutputRef(Context.FinalOutput));
+				Context.FinalOutput = FRenderingCompositeOutputRef(Node);
+			}
 		}
 
 		if(View.Family->EngineShowFlags.VisualizeHDR && GRHIFeatureLevel >= ERHIFeatureLevel::SM5)
