Added missing include

diff --git a/Engine/Source/Runtime/Renderer/Private/PostProcess/PostProcessing.cpp b/Engine/Source/Runtime/Renderer/Private/PostProcess/PostProcessing.cpp
--- a/Engine/Source/Runtime/Renderer/Private/PostProcess/PostProcessing.cpp
+++ b/Engine/Source/Runtime/Renderer/Private/PostProcess/PostProcessing.cpp
@@ -46,6 +46,7 @@
 #include "HighResScreenshot.h"
 
 #include "PostProcessMorpheus.h"
+#include "IHeadMountedDisplay.h"
 
 
 /** The global center for all post processing activities. */
