#UE4:
* fixed emissive material with Light Propagation Volumes and secondary occlusion

diff --git a/Engine/Source/Runtime/Renderer/Private/ShadowRendering.cpp b/Engine/Source/Runtime/Renderer/Private/ShadowRendering.cpp
--- a/Engine/Source/Runtime/Renderer/Private/ShadowRendering.cpp
+++ b/Engine/Source/Runtime/Renderer/Private/ShadowRendering.cpp
@@ -431,6 +431,17 @@ public:
 				ProjectionMatrixParameter,
 				FTranslationMatrix(ShadowInfo->PreShadowTranslation - View.ViewMatrices.PreViewTranslation) * ShadowInfo->SubjectAndReceiverMatrix
 				);
+
+			const FSceneViewState* ViewState = (const FSceneViewState*)View.State;
+			if(ViewState)
+			{
+				const FLightPropagationVolume* Lpv = ViewState->GetLightPropagationVolume();
+
+				if(Lpv)
+				{
+					SetUniformBufferParameter(ShaderRHI, GetUniformBufferParameter<FLpvWriteUniformBufferParameters>(), Lpv->GetWriteUniformBuffer());
+				}
+			}
 		}
 	}
 
