#ue4 fix crash when users without platform access run builds compiled from a machine which does.

diff --git a/Engine/Source/Runtime/Renderer/Private/PostProcess/PostProcessMorpheus.cpp b/Engine/Source/Runtime/Renderer/Private/PostProcess/PostProcessMorpheus.cpp
--- a/Engine/Source/Runtime/Renderer/Private/PostProcess/PostProcessMorpheus.cpp
+++ b/Engine/Source/Runtime/Renderer/Private/PostProcess/PostProcessMorpheus.cpp
@@ -24,7 +24,14 @@ class FPostProcessMorpheusPS : public FGlobalShader
 
 	static bool ShouldCache(EShaderPlatform Platform)
 	{
-		return true;
+		// we must use a run time check for this because the builds the build machines create will have Morpheus defined,
+		// but a user will not necessarily have the Morpheus files
+		bool bEnableMorpheus = false;
+		if (GConfig->GetBool(TEXT("Morpheus.Settings"), TEXT("EnableMorpheus"), bEnableMorpheus, GEngineIni))
+		{
+			return bEnableMorpheus;
+		}
+		return false;
 	}
 
 	/** Default constructor. */
@@ -164,7 +171,14 @@ class FPostProcessMorpheusVS : public FGlobalShader
 
 	static bool ShouldCache(EShaderPlatform Platform)
 	{
-		return true;
+		// we must use a run time check for this because the builds the build machines create will have Morpheus defined,
+		// but a user will not necessarily have the Morpheus files
+		bool bEnableMorpheus = false;
+		if (GConfig->GetBool(TEXT("Morpheus.Settings"), TEXT("EnableMorpheus"), bEnableMorpheus, GEngineIni))
+		{
+			return bEnableMorpheus;
+		}
+		return false;
 	}
 
 	/** Default constructor. */
diff --git a/Engine/Source/Runtime/ShaderCore/Private/ShaderCore.cpp b/Engine/Source/Runtime/ShaderCore/Private/ShaderCore.cpp
--- a/Engine/Source/Runtime/ShaderCore/Private/ShaderCore.cpp
+++ b/Engine/Source/Runtime/ShaderCore/Private/ShaderCore.cpp
@@ -395,7 +395,7 @@ void GetShaderIncludes(const TCHAR* Filename, TArray<FString>& IncludeFilenames,
 			
 					// Some headers aren't required to be found (platforms that the user doesn't have access to)
 					// @todo: Is there some way to generalize this"
-					const bool bIsOptionalInclude = ExtractedIncludeFilename == TEXT("PS4/PS4Common.usf");
+					const bool bIsOptionalInclude = (ExtractedIncludeFilename == TEXT("PS4/PS4Common.usf") || ExtractedIncludeFilename == TEXT("PS4/PostProcessHMDMorpheus.usf"));
 					// ignore the header if it's optional and doesn't exist
 					if (bIsOptionalInclude)
 					{
