Fixed wireframe material rendering shadows

diff --git a/Engine/Source/Runtime/Renderer/Private/DepthRendering.cpp b/Engine/Source/Runtime/Renderer/Private/DepthRendering.cpp
--- a/Engine/Source/Runtime/Renderer/Private/DepthRendering.cpp
+++ b/Engine/Source/Runtime/Renderer/Private/DepthRendering.cpp
@@ -248,9 +248,10 @@ FPositionOnlyDepthDrawingPolicy::FPositionOnlyDepthDrawingPolicy(
 	const FVertexFactory* InVertexFactory,
 	const FMaterialRenderProxy* InMaterialRenderProxy,
 	const FMaterial& InMaterialResource,
-	bool bIsTwoSided
+	bool bIsTwoSided,
+	bool bIsWireframe
 	):
-	FMeshDrawingPolicy(InVertexFactory,InMaterialRenderProxy,InMaterialResource,false,/*bInTwoSidedOverride=*/ bIsTwoSided)
+	FMeshDrawingPolicy(InVertexFactory,InMaterialRenderProxy,InMaterialResource,false,bIsTwoSided,bIsWireframe)
 {
 	VertexShader = InMaterialResource.GetShader<TDepthOnlyVS<true> >(InVertexFactory->GetType());
 }
@@ -337,7 +338,8 @@ void FDepthDrawingPolicyFactory::AddStaticMesh(FScene* Scene,FStaticMesh* Static
 					StaticMesh->VertexFactory,
 					DefaultProxy,
 					*DefaultProxy->GetMaterial(GRHIFeatureLevel),
-					Material->IsTwoSided()
+					Material->IsTwoSided(),
+					Material->IsWireframe()
 					)
 				);
 		}
@@ -395,7 +397,7 @@ bool FDepthDrawingPolicyFactory::DrawMesh(
 		{
 			//render opaque primitives that support a separate position-only vertex buffer
 			const FMaterialRenderProxy* DefaultProxy = UMaterial::GetDefaultMaterial(MD_Surface)->GetRenderProxy(false);
-			FPositionOnlyDepthDrawingPolicy DrawingPolicy(Mesh.VertexFactory,DefaultProxy,*DefaultProxy->GetMaterial(GRHIFeatureLevel),Material->IsTwoSided());
+			FPositionOnlyDepthDrawingPolicy DrawingPolicy(Mesh.VertexFactory,DefaultProxy,*DefaultProxy->GetMaterial(GRHIFeatureLevel),Material->IsTwoSided(),Material->IsWireframe());
 			DrawingPolicy.DrawShared(&View,DrawingPolicy.CreateBoundShaderState());
 			int32 BatchElementIndex = 0;
 			uint64 Mask = BatchElementMask;
diff --git a/Engine/Source/Runtime/Renderer/Private/DepthRendering.h b/Engine/Source/Runtime/Renderer/Private/DepthRendering.h
--- a/Engine/Source/Runtime/Renderer/Private/DepthRendering.h
+++ b/Engine/Source/Runtime/Renderer/Private/DepthRendering.h
@@ -86,7 +86,8 @@ public:
 		const FVertexFactory* InVertexFactory,
 		const FMaterialRenderProxy* InMaterialRenderProxy,
 		const FMaterial& InMaterialResource,
-		bool bIsTwoSided
+		bool bIsTwoSided,
+		bool bIsWireframe
 		);
 
 	// FMeshDrawingPolicy interface.
diff --git a/Engine/Source/Runtime/Renderer/Private/DrawingPolicy.cpp b/Engine/Source/Runtime/Renderer/Private/DrawingPolicy.cpp
--- a/Engine/Source/Runtime/Renderer/Private/DrawingPolicy.cpp
+++ b/Engine/Source/Runtime/Renderer/Private/DrawingPolicy.cpp
@@ -21,12 +21,13 @@ FMeshDrawingPolicy::FMeshDrawingPolicy(
 	const FMaterialRenderProxy* InMaterialRenderProxy,
 	const FMaterial& InMaterialResource,
 	bool bInOverrideWithShaderComplexity,
-	bool bInTwoSidedOverride
+	bool bInTwoSidedOverride,
+	bool bInWireframeOverride
 	):
 	VertexFactory(InVertexFactory),
 	MaterialRenderProxy(InMaterialRenderProxy),
 	MaterialResource(&InMaterialResource),
-	bIsWireframeMaterial(InMaterialResource.IsWireframe()),
+	bIsWireframeMaterial(InMaterialResource.IsWireframe() || bInWireframeOverride),
 	//convert from signed bool to unsigned uint32
 	bOverrideWithShaderComplexity(bInOverrideWithShaderComplexity != false)
 {
diff --git a/Engine/Source/Runtime/Renderer/Private/DrawingPolicy.h b/Engine/Source/Runtime/Renderer/Private/DrawingPolicy.h
--- a/Engine/Source/Runtime/Renderer/Private/DrawingPolicy.h
+++ b/Engine/Source/Runtime/Renderer/Private/DrawingPolicy.h
@@ -30,7 +30,8 @@ public:
 		const FMaterialRenderProxy* InMaterialRenderProxy,
 		const FMaterial& InMaterialResource,
 		bool bInOverrideWithShaderComplexity = false,
-		bool bInTwoSidedOverride = false
+		bool bInTwoSidedOverride = false,
+		bool bInWireframeOverride = false
 		);
 
 	FMeshDrawingPolicy& operator = (const FMeshDrawingPolicy& Other)
