#ttp 332104 - CRITICAL: REGRESSION: PS4: SHOOTERGAME: Game hangs
#ttp 332107 - CRITICAL: REGRESSION: PS4: SHOOTERGAME: CRASH: GPU crash
#ttp 332132 - PS4: RENDERING: Bright green reflections
#summary - Temporary fix for running out of space in the LCUE resource buffer.
#change - Submits the command buffer after each cubemap reflection as there is currently no way to know when the LCUE has run out of resource buffer space.

diff --git a/Engine/Source/Runtime/Renderer/Private/ReflectionEnvironmentCapture.cpp b/Engine/Source/Runtime/Renderer/Private/ReflectionEnvironmentCapture.cpp
--- a/Engine/Source/Runtime/Renderer/Private/ReflectionEnvironmentCapture.cpp
+++ b/Engine/Source/Runtime/Renderer/Private/ReflectionEnvironmentCapture.cpp
@@ -779,6 +779,12 @@ void CaptureSceneToScratchCubemap(FSceneRenderer* SceneRenderer, ECubeFace CubeF
 		// Render the scene normally for one face of the cubemap
 		SceneRenderer->Render();
 
+#if PLATFORM_PS4 // @todo ps4 - this should be done a different way
+		// PS4 needs some code here to process the scene
+		extern void TEMP_PostReflectionCaptureRender();
+		TEMP_PostReflectionCaptureRender();
+#endif
+
 		const int32 EffectiveSize = GReflectionCaptureSize;
 		FSceneRenderTargetItem& EffectiveColorRT =  GSceneRenderTargets.ReflectionColorScratchCubemap[0]->GetRenderTargetItem();
 
@@ -1200,7 +1206,7 @@ void ClearScratchCubemaps()
 		if (GSupportsGSRenderTargetLayerSwitchingToMips)
 		{
 			RHISetRenderTarget(RT0.TargetableTexture, MipIndex, NULL);
-			RHIClear(true, FLinearColor(0, 10000, 0, 0), false, 0, false, 0, FIntRect());
+			RHIClear(true, FLinearColor(10000, 0, 0, 0), false, 0, false, 0, FIntRect());
 		}
 		else
 		{
@@ -1343,6 +1349,30 @@ void CaptureSceneIntoScratchCubemap(FScene* Scene, FVector CapturePosition, bool
 		FSceneViewInitOptions ViewInitOptions;
 		ViewInitOptions.ViewFamily = &ViewFamily;
 		ViewInitOptions.BackgroundColor = FLinearColor::Black;
+/*		switch(CubeFace )
+		{
+		case 0:
+			ViewInitOptions.BackgroundColor = FLinearColor::Blue;
+			break;
+		case 1:
+			ViewInitOptions.BackgroundColor = FLinearColor::Red;
+			break;
+		case 2:
+			ViewInitOptions.BackgroundColor = FLinearColor::Green;
+			break;
+		case 3:
+			ViewInitOptions.BackgroundColor = FLinearColor::Yellow;
+			break;
+		case 4:
+			ViewInitOptions.BackgroundColor = FLinearColor::White;
+			break;
+		case 5:
+			ViewInitOptions.BackgroundColor = FLinearColor::Gray;
+			break;
+		}
+		ViewFamily.EngineShowFlags.StaticMeshes = 0;
+		ViewFamily.EngineShowFlags.BSP = 0;
+*/
 		ViewInitOptions.OverlayColor = FLinearColor::Black;
 		ViewInitOptions.SetViewRectangle(FIntRect(0, 0, GReflectionCaptureSize * GSupersampleCaptureFactor, GReflectionCaptureSize * GSupersampleCaptureFactor));
 
