#summary Remove some reflection capture debug code that shoudnt have got checked in :-(

diff --git a/Engine/Source/Runtime/Renderer/Private/ReflectionEnvironmentCapture.cpp b/Engine/Source/Runtime/Renderer/Private/ReflectionEnvironmentCapture.cpp
--- a/Engine/Source/Runtime/Renderer/Private/ReflectionEnvironmentCapture.cpp
+++ b/Engine/Source/Runtime/Renderer/Private/ReflectionEnvironmentCapture.cpp
@@ -1206,7 +1206,7 @@ void ClearScratchCubemaps()
 		if (GSupportsGSRenderTargetLayerSwitchingToMips)
 		{
 			RHISetRenderTarget(RT0.TargetableTexture, MipIndex, NULL);
-			RHIClear(true, FLinearColor(10000, 0, 0, 0), false, 0, false, 0, FIntRect());
+			RHIClear(true, FLinearColor(0, 10000, 0, 0), false, 0, false, 0, FIntRect());
 		}
 		else
 		{
@@ -1349,30 +1349,6 @@ void CaptureSceneIntoScratchCubemap(FScene* Scene, FVector CapturePosition, bool
 		FSceneViewInitOptions ViewInitOptions;
 		ViewInitOptions.ViewFamily = &ViewFamily;
 		ViewInitOptions.BackgroundColor = FLinearColor::Black;
-/*		switch(CubeFace )
-		{
-		case 0:
-			ViewInitOptions.BackgroundColor = FLinearColor::Blue;
-			break;
-		case 1:
-			ViewInitOptions.BackgroundColor = FLinearColor::Red;
-			break;
-		case 2:
-			ViewInitOptions.BackgroundColor = FLinearColor::Green;
-			break;
-		case 3:
-			ViewInitOptions.BackgroundColor = FLinearColor::Yellow;
-			break;
-		case 4:
-			ViewInitOptions.BackgroundColor = FLinearColor::White;
-			break;
-		case 5:
-			ViewInitOptions.BackgroundColor = FLinearColor::Gray;
-			break;
-		}
-		ViewFamily.EngineShowFlags.StaticMeshes = 0;
-		ViewFamily.EngineShowFlags.BSP = 0;
-*/
 		ViewInitOptions.OverlayColor = FLinearColor::Black;
 		ViewInitOptions.SetViewRectangle(FIntRect(0, 0, GReflectionCaptureSize * GSupersampleCaptureFactor, GReflectionCaptureSize * GSupersampleCaptureFactor));
 
