Submitting for Keith to make RC deadline.

#ttp 331730 CRITICAL: Regression: XboxOne: QA-Postprocessing is washed out

#project Platforms

#branch UE4

#summary Make histogram calculations happen even if min/max brightness are the same. This stops QA-Postprocessing, Elemental and Infiltrator from looking extremely bright.

diff --git a/Engine/Source/Runtime/Renderer/Private/PostProcess/PostProcessing.cpp b/Engine/Source/Runtime/Renderer/Private/PostProcess/PostProcessing.cpp
--- a/Engine/Source/Runtime/Renderer/Private/PostProcess/PostProcessing.cpp
+++ b/Engine/Source/Runtime/Renderer/Private/PostProcess/PostProcessing.cpp
@@ -931,7 +931,7 @@ void FPostProcessing::Process(const FViewInfo& View, TRefCountPtr<IPooledRenderT
 				bool bHistogramNeeded = View.Family->EngineShowFlags.VisualizeHDR;
 
 				if(View.Family->EngineShowFlags.EyeAdaptation
-					&& View.FinalPostProcessSettings.AutoExposureMinBrightness < View.FinalPostProcessSettings.AutoExposureMaxBrightness
+					&& View.FinalPostProcessSettings.AutoExposureMinBrightness <= View.FinalPostProcessSettings.AutoExposureMaxBrightness
 					&& !View.bIsSceneCapture) // Eye adaption is not available for scene captures.
 				{
 					bHistogramNeeded = true;
