#ttp 304473 - UE4: RENDERING: Visualizetexture looks broken when using r.ScreenPercentage / FullscreenAndUpsample
#proj Engine
#branch UE4

#summary Fixed render target visualization when using r.screenpercentage

#change Make texture visualization code use the unscaled viewrect for its calculations
#change Give DrawRectangle function the size of the actual render target, not the scene rendertarget sizes

diff --git a/Engine/Source/Runtime/Renderer/Private/PostProcess/VisualizeTexture.cpp b/Engine/Source/Runtime/Renderer/Private/PostProcess/VisualizeTexture.cpp
--- a/Engine/Source/Runtime/Renderer/Private/PostProcess/VisualizeTexture.cpp
+++ b/Engine/Source/Runtime/Renderer/Private/PostProcess/VisualizeTexture.cpp
@@ -152,6 +152,7 @@ class FVisualizeTexturePresentPS : public FGlobalShader
 	static bool ShouldCache(EShaderPlatform Platform)
 	{
 		return IsFeatureLevelSupported(Platform, ERHIFeatureLevel::ES2);
+
 	}
 
 	/** Default constructor. */
@@ -482,8 +483,9 @@ void FVisualizeTexture::PresentContent(const FSceneView& View)
 
 	FPooledRenderTargetDesc Desc = VisualizeTextureDesc;
 
-	RHISetRenderTarget(View.Family->RenderTarget->GetRenderTargetTexture(),FTextureRHIRef());
-	RHISetViewport(0, 0, 0.0f, GSceneRenderTargets.GetBufferSizeXY().X, GSceneRenderTargets.GetBufferSizeXY().Y, 1.0f );
+	auto& RenderTarget = View.Family->RenderTarget->GetRenderTargetTexture();
+	RHISetRenderTarget(RenderTarget, FTextureRHIRef());
+	RHISetViewport(0, 0, 0.0f, RenderTarget->GetSizeX(), RenderTarget->GetSizeY(), 1.0f);
 
 	RHISetBlendState(TStaticBlendState<>::GetRHI());
 	RHISetRasterizerState(TStaticRasterizerState<>::GetRHI());
@@ -509,7 +511,7 @@ void FVisualizeTexture::PresentContent(const FSceneView& View)
 		VisualizeTextureRect.Width(), VisualizeTextureRect.Height(),
 		0, 0,
 		VisualizeTextureRect.Width(), VisualizeTextureRect.Height(),
-		GSceneRenderTargets.GetBufferSizeXY(),
+		FIntPoint(RenderTarget->GetSizeX(), RenderTarget->GetSizeY()),
 		VisualizeTextureRect.Size(),
 		*VertexShader,
 		EDRF_Default);
@@ -835,7 +837,7 @@ IPooledRenderTarget* FVisualizeTexture::GetObservedElement() const
 
 void FVisualizeTexture::OnStartFrame(const FSceneView& View)
 {
-	ViewRect = View.ViewRect;
+	ViewRect = View.UnscaledViewRect;
 
 #if !(UE_BUILD_SHIPPING || UE_BUILD_TEST)
 	// VisualizeTexture observed render target is set each frame
