#ue4
#ios
* fix for unlit character in SideScroller

diff --git a/Engine/Source/Runtime/Renderer/Private/IndirectLightingCache.cpp b/Engine/Source/Runtime/Renderer/Private/IndirectLightingCache.cpp
--- a/Engine/Source/Runtime/Renderer/Private/IndirectLightingCache.cpp
+++ b/Engine/Source/Runtime/Renderer/Private/IndirectLightingCache.cpp
@@ -135,7 +135,7 @@ bool IsIndirectLightingCacheAllowed()
 	static const auto AllowStaticLightingVar = IConsoleManager::Get().FindTConsoleVariableDataInt(TEXT("r.AllowStaticLighting"));
 	const bool bAllowStaticLighting = (!AllowStaticLightingVar || AllowStaticLightingVar->GetValueOnRenderThread() != 0);
 
-	return GIndirectLightingCache != 0 && GSupportsVolumeTextureRendering && bAllowStaticLighting;
+	return GIndirectLightingCache != 0 && (GSupportsVolumeTextureRendering || GRHIFeatureLevel == ERHIFeatureLevel::ES2) && bAllowStaticLighting;
 }
 
 bool CanIndirectLightingCacheUseVolumeTexture()
