#ttp 331314 Directional light selection for exponential fog inscattering is non-deterministic
#proj Editor
#branch UE4
#summary Use lights configured as "IsUsedAsAtmosphereSunlight".

diff --git a/Engine/Source/Runtime/Engine/Public/SceneManagement.h b/Engine/Source/Runtime/Engine/Public/SceneManagement.h
--- a/Engine/Source/Runtime/Engine/Public/SceneManagement.h
+++ b/Engine/Source/Runtime/Engine/Public/SceneManagement.h
@@ -1005,7 +1005,7 @@ protected:
 	/** Whether the light affects translucency or not.  Disabling this can save GPU time when there are many small lights. */
 	const uint32 bAffectTranslucentLighting : 1;
 
-	/** Whether to consider light as a sunlight for atmospheric scattering. */
+	/** Whether to consider light as a sunlight for atmospheric scattering and exponential height fog. */
 	const uint32 bUsedAsAtmosphereSunLight : 1;
 
 	/** Does the light have dynamic GI? */
diff --git a/Engine/Source/Runtime/Renderer/Private/FogRendering.cpp b/Engine/Source/Runtime/Renderer/Private/FogRendering.cpp
--- a/Engine/Source/Runtime/Renderer/Private/FogRendering.cpp
+++ b/Engine/Source/Runtime/Renderer/Private/FogRendering.cpp
@@ -240,7 +240,10 @@ void FSceneRenderer::InitFogConstants()
 				{
 					const FLightSceneInfoCompact& LightInfo = *It;
 
+					// This will find the first directional light that is set to be used as an atmospheric sun light of sufficient brightness.
+					// If you have more than one directional light with these properties then all subsequent lights will be ignored.
 					if (LightInfo.LightSceneInfo->Proxy->GetLightType() == LightType_Directional
+						&& LightInfo.LightSceneInfo->Proxy->IsUsedAsAtmosphereSunLight()
 						&& LightInfo.LightSceneInfo->Proxy->GetColor().ComputeLuminance() > KINDA_SMALL_NUMBER
 						&& FogInfo.DirectionalInscatteringColor.ComputeLuminance() > KINDA_SMALL_NUMBER)
 					{
