#ttp 329734 Mac: GDC2014: Default level is washed out
#branch UE4
#project Core
#summary Fix the sky light irradiance on Nvidia cards under OS X - they suffer from their inability to switch cubemap face in a geometry shader.
#change Add a code path for when GSupportsGSRenderTargetLayerSwitchingToMips is false in ComputeDiffuseIrradiance as is the case on OS X & is required for Nvidia cards which can't do this (though AMD & Intel cards can).

diff --git a/Engine/Source/Runtime/Renderer/Private/ReflectionEnvironmentDiffuseIrradiance.cpp b/Engine/Source/Runtime/Renderer/Private/ReflectionEnvironmentDiffuseIrradiance.cpp
--- a/Engine/Source/Runtime/Renderer/Private/ReflectionEnvironmentDiffuseIrradiance.cpp
+++ b/Engine/Source/Runtime/Renderer/Private/ReflectionEnvironmentDiffuseIrradiance.cpp
@@ -283,38 +283,72 @@ void ComputeDiffuseIrradiance(FTextureRHIRef LightingSource, int32 LightingSourc
 				FSceneRenderTargetItem& EffectiveRT = GetEffectiveDiffuseIrradianceRenderTarget(MipIndex);
 				FSceneRenderTargetItem& EffectiveSource = GetEffectiveDiffuseIrradianceSourceTexture(MipIndex);
 				check(EffectiveRT.TargetableTexture != EffectiveSource.ShaderResourceTexture);
-
-				RHISetRenderTarget(EffectiveRT.TargetableTexture, MipIndex, NULL);
-
-				const FIntRect ViewRect(0, 0, MipSize, MipSize);
-				RHISetViewport(0, 0, 0.0f, MipSize, MipSize, 1.0f);
-				RHISetRasterizerState(TStaticRasterizerState<FM_Solid,CM_None>::GetRHI());
-				RHISetDepthStencilState(TStaticDepthStencilState<false,CF_Always>::GetRHI());
-				RHISetBlendState(TStaticBlendState<>::GetRHI());
-
-				TShaderMapRef<FScreenVSForGS> VertexShader(GetGlobalShaderMap());
-				TShaderMapRef<FDownsampleGS> GeometryShader(GetGlobalShaderMap());
-				TShaderMapRef<FAccumulateDiffuseIrradiancePS> PixelShader(GetGlobalShaderMap());
-
-				SetGlobalBoundShaderState(DiffuseIrradianceAccumulateShaderState, GFilterVertexDeclaration.VertexDeclarationRHI, *VertexShader, *PixelShader, *GeometryShader);
-
-				// Draw each face with a geometry shader that routes to the correct render target slice
-				for (int32 CubeFace = 0; CubeFace < CubeFace_MAX; CubeFace++)
+				
+				if (GSupportsGSRenderTargetLayerSwitchingToMips)
 				{
-					GeometryShader->SetParameters(CubeFace);
-
-					PixelShader->SetParameters(CubeFace, NumMips, SourceMipIndex, CoefficientIndex, EffectiveSource.ShaderResourceTexture);
-
-					DrawRectangle( 
-						ViewRect.Min.X, ViewRect.Min.Y, 
-						ViewRect.Width(), ViewRect.Height(),
-						ViewRect.Min.X, ViewRect.Min.Y, 
-						ViewRect.Width(), ViewRect.Height(),
-						FIntPoint(ViewRect.Width(), ViewRect.Height()),
-						FIntPoint(MipSize, MipSize),
-						*VertexShader);
-
-					RHICopyToResolveTarget(EffectiveRT.TargetableTexture, EffectiveRT.ShaderResourceTexture, true, FResolveParams(FResolveRect(), (ECubeFace)CubeFace, MipIndex));
+					RHISetRenderTarget(EffectiveRT.TargetableTexture, MipIndex, NULL);
+					
+					const FIntRect ViewRect(0, 0, MipSize, MipSize);
+					RHISetViewport(0, 0, 0.0f, MipSize, MipSize, 1.0f);
+					RHISetRasterizerState(TStaticRasterizerState<FM_Solid,CM_None>::GetRHI());
+					RHISetDepthStencilState(TStaticDepthStencilState<false,CF_Always>::GetRHI());
+					RHISetBlendState(TStaticBlendState<>::GetRHI());
+					
+					TShaderMapRef<FScreenVSForGS> VertexShader(GetGlobalShaderMap());
+					TShaderMapRef<FDownsampleGS> GeometryShader(GetGlobalShaderMap());
+					TShaderMapRef<FAccumulateDiffuseIrradiancePS> PixelShader(GetGlobalShaderMap());
+					
+					SetGlobalBoundShaderState(DiffuseIrradianceAccumulateShaderState, GFilterVertexDeclaration.VertexDeclarationRHI, *VertexShader, *PixelShader, *GeometryShader);
+					
+					// Draw each face with a geometry shader that routes to the correct render target slice
+					for (int32 CubeFace = 0; CubeFace < CubeFace_MAX; CubeFace++)
+					{
+						GeometryShader->SetParameters(CubeFace);
+						
+						PixelShader->SetParameters(CubeFace, NumMips, SourceMipIndex, CoefficientIndex, EffectiveSource.ShaderResourceTexture);
+						
+						DrawRectangle(
+									  ViewRect.Min.X, ViewRect.Min.Y,
+									  ViewRect.Width(), ViewRect.Height(),
+									  ViewRect.Min.X, ViewRect.Min.Y,
+									  ViewRect.Width(), ViewRect.Height(),
+									  FIntPoint(ViewRect.Width(), ViewRect.Height()),
+									  FIntPoint(MipSize, MipSize),
+									  *VertexShader);
+						
+						RHICopyToResolveTarget(EffectiveRT.TargetableTexture, EffectiveRT.ShaderResourceTexture, true, FResolveParams(FResolveRect(), (ECubeFace)CubeFace, MipIndex));
+					}
+				}
+				else
+				{
+					for (int32 CubeFace = 0; CubeFace < CubeFace_MAX; CubeFace++)
+					{
+						RHISetRenderTarget(EffectiveRT.TargetableTexture, MipIndex, CubeFace, NULL);
+						
+						const FIntRect ViewRect(0, 0, MipSize, MipSize);
+						RHISetViewport(0, 0, 0.0f, MipSize, MipSize, 1.0f);
+						RHISetRasterizerState(TStaticRasterizerState<FM_Solid,CM_None>::GetRHI());
+						RHISetDepthStencilState(TStaticDepthStencilState<false,CF_Always>::GetRHI());
+						RHISetBlendState(TStaticBlendState<>::GetRHI());
+						
+						TShaderMapRef<FScreenVSForGS> VertexShader(GetGlobalShaderMap());
+						TShaderMapRef<FAccumulateDiffuseIrradiancePS> PixelShader(GetGlobalShaderMap());
+						
+						SetGlobalBoundShaderState(DiffuseIrradianceAccumulateShaderState, GFilterVertexDeclaration.VertexDeclarationRHI, *VertexShader, *PixelShader);
+						
+						PixelShader->SetParameters(CubeFace, NumMips, SourceMipIndex, CoefficientIndex, EffectiveSource.ShaderResourceTexture);
+						
+						DrawRectangle(
+									  ViewRect.Min.X, ViewRect.Min.Y,
+									  ViewRect.Width(), ViewRect.Height(),
+									  ViewRect.Min.X, ViewRect.Min.Y,
+									  ViewRect.Width(), ViewRect.Height(),
+									  FIntPoint(ViewRect.Width(), ViewRect.Height()),
+									  FIntPoint(MipSize, MipSize),
+									  *VertexShader);
+						
+						RHICopyToResolveTarget(EffectiveRT.TargetableTexture, EffectiveRT.ShaderResourceTexture, true, FResolveParams(FResolveRect(), (ECubeFace)CubeFace, MipIndex));
+					}
 				}
 			}
 		}
