#ttp 329734 Mac: GDC2014: Default level is washed out
#branch UE4
#project Renderer
#summary Further fix for sky light irradiance on Nvidia cards under OS X - missed a code path that also needed changing.
#change Add another code path for when GSupportsGSRenderTargetLayerSwitchingToMips is false in ComputeDiffuseIrradiance as is the case on OS X & is required for Nvidia cards which can't do this (though AMD & Intel cards can).

diff --git a/Engine/Source/Runtime/Renderer/Private/ReflectionEnvironmentDiffuseIrradiance.cpp b/Engine/Source/Runtime/Renderer/Private/ReflectionEnvironmentDiffuseIrradiance.cpp
--- a/Engine/Source/Runtime/Renderer/Private/ReflectionEnvironmentDiffuseIrradiance.cpp
+++ b/Engine/Source/Runtime/Renderer/Private/ReflectionEnvironmentDiffuseIrradiance.cpp
@@ -237,37 +237,71 @@ void ComputeDiffuseIrradiance(FTextureRHIRef LightingSource, int32 LightingSourc
 			const int32 MipIndex = 0;
 			const int32 MipSize = GDiffuseIrradianceCubemapSize;
 			FSceneRenderTargetItem& EffectiveRT = GetEffectiveDiffuseIrradianceRenderTarget(MipIndex);
-			RHISetRenderTarget(EffectiveRT.TargetableTexture, 0, NULL);
+			if (GSupportsGSRenderTargetLayerSwitchingToMips)
+			{
+				RHISetRenderTarget(EffectiveRT.TargetableTexture, 0, NULL);
 
-			const FIntRect ViewRect(0, 0, MipSize, MipSize);
-			RHISetViewport(0, 0, 0.0f, MipSize, MipSize, 1.0f);
-			RHISetRasterizerState(TStaticRasterizerState<FM_Solid,CM_None>::GetRHI());
-			RHISetDepthStencilState(TStaticDepthStencilState<false,CF_Always>::GetRHI());
-			RHISetBlendState(TStaticBlendState<>::GetRHI());
+				const FIntRect ViewRect(0, 0, MipSize, MipSize);
+				RHISetViewport(0, 0, 0.0f, MipSize, MipSize, 1.0f);
+				RHISetRasterizerState(TStaticRasterizerState<FM_Solid,CM_None>::GetRHI());
+				RHISetDepthStencilState(TStaticDepthStencilState<false,CF_Always>::GetRHI());
+				RHISetBlendState(TStaticBlendState<>::GetRHI());
 
-			TShaderMapRef<FScreenVSForGS> VertexShader(GetGlobalShaderMap());
-			TShaderMapRef<FDownsampleGS> GeometryShader(GetGlobalShaderMap());
-			TShaderMapRef<FCopyDiffuseIrradiancePS> PixelShader(GetGlobalShaderMap());
+				TShaderMapRef<FScreenVSForGS> VertexShader(GetGlobalShaderMap());
+				TShaderMapRef<FDownsampleGS> GeometryShader(GetGlobalShaderMap());
+				TShaderMapRef<FCopyDiffuseIrradiancePS> PixelShader(GetGlobalShaderMap());
 
-			SetGlobalBoundShaderState(CopyDiffuseIrradianceShaderState, GFilterVertexDeclaration.VertexDeclarationRHI, *VertexShader, *PixelShader, *GeometryShader);
+				SetGlobalBoundShaderState(CopyDiffuseIrradianceShaderState, GFilterVertexDeclaration.VertexDeclarationRHI, *VertexShader, *PixelShader, *GeometryShader);
 
-			// Draw each face with a geometry shader that routes to the correct render target slice
-			for (int32 CubeFace = 0; CubeFace < CubeFace_MAX; CubeFace++)
-			{
-				GeometryShader->SetParameters(CubeFace);
+				// Draw each face with a geometry shader that routes to the correct render target slice
+				for (int32 CubeFace = 0; CubeFace < CubeFace_MAX; CubeFace++)
+				{
+					GeometryShader->SetParameters(CubeFace);
 
-				PixelShader->SetParameters(CubeFace, LightingSourceMipIndex, CoefficientIndex, MipSize, LightingSource);
+					PixelShader->SetParameters(CubeFace, LightingSourceMipIndex, CoefficientIndex, MipSize, LightingSource);
 
-				DrawRectangle( 
-					ViewRect.Min.X, ViewRect.Min.Y, 
-					ViewRect.Width(), ViewRect.Height(),
-					ViewRect.Min.X, ViewRect.Min.Y, 
-					ViewRect.Width(), ViewRect.Height(),
-					FIntPoint(ViewRect.Width(), ViewRect.Height()),
-					FIntPoint(MipSize, MipSize),
-					*VertexShader);
+					DrawRectangle( 
+						ViewRect.Min.X, ViewRect.Min.Y, 
+						ViewRect.Width(), ViewRect.Height(),
+						ViewRect.Min.X, ViewRect.Min.Y, 
+						ViewRect.Width(), ViewRect.Height(),
+						FIntPoint(ViewRect.Width(), ViewRect.Height()),
+						FIntPoint(MipSize, MipSize),
+						*VertexShader);
 
-				RHICopyToResolveTarget(EffectiveRT.TargetableTexture, EffectiveRT.ShaderResourceTexture, true, FResolveParams(FResolveRect(), (ECubeFace)CubeFace, MipIndex));
+					RHICopyToResolveTarget(EffectiveRT.TargetableTexture, EffectiveRT.ShaderResourceTexture, true, FResolveParams(FResolveRect(), (ECubeFace)CubeFace, MipIndex));
+				}
+			}
+			else
+			{
+				for (int32 CubeFace = 0; CubeFace < CubeFace_MAX; CubeFace++)
+				{
+					RHISetRenderTarget(EffectiveRT.TargetableTexture, 0, CubeFace, NULL);
+					
+					const FIntRect ViewRect(0, 0, MipSize, MipSize);
+					RHISetViewport(0, 0, 0.0f, MipSize, MipSize, 1.0f);
+					RHISetRasterizerState(TStaticRasterizerState<FM_Solid,CM_None>::GetRHI());
+					RHISetDepthStencilState(TStaticDepthStencilState<false,CF_Always>::GetRHI());
+					RHISetBlendState(TStaticBlendState<>::GetRHI());
+					
+					TShaderMapRef<FScreenVSForGS> VertexShader(GetGlobalShaderMap());
+					TShaderMapRef<FCopyDiffuseIrradiancePS> PixelShader(GetGlobalShaderMap());
+					
+					SetGlobalBoundShaderState(CopyDiffuseIrradianceShaderState, GFilterVertexDeclaration.VertexDeclarationRHI, *VertexShader, *PixelShader);
+					
+					PixelShader->SetParameters(CubeFace, LightingSourceMipIndex, CoefficientIndex, MipSize, LightingSource);
+					
+					DrawRectangle(
+								  ViewRect.Min.X, ViewRect.Min.Y,
+								  ViewRect.Width(), ViewRect.Height(),
+								  ViewRect.Min.X, ViewRect.Min.Y,
+								  ViewRect.Width(), ViewRect.Height(),
+								  FIntPoint(ViewRect.Width(), ViewRect.Height()),
+								  FIntPoint(MipSize, MipSize),
+								  *VertexShader);
+					
+					RHICopyToResolveTarget(EffectiveRT.TargetableTexture, EffectiveRT.ShaderResourceTexture, true, FResolveParams(FResolveRect(), (ECubeFace)CubeFace, MipIndex));
+				}
 			}
 		}
 
