Fix d3ddebug warnings when rendering deferred lights.

diff --git a/Engine/Source/Runtime/Renderer/Private/LightRendering.cpp b/Engine/Source/Runtime/Renderer/Private/LightRendering.cpp
--- a/Engine/Source/Runtime/Renderer/Private/LightRendering.cpp
+++ b/Engine/Source/Runtime/Renderer/Private/LightRendering.cpp
@@ -643,6 +643,12 @@ void SetBoundingGeometryRasterizerAndDepthState(const FViewInfo& View, const FSp
 		);
 }
 
+template <bool bRadialAttenuation>
+static FVertexDeclarationRHIParamRef GetDeferredLightingVertexDeclaration()
+{
+	return bRadialAttenuation ? GetVertexDeclarationFVector4() : GFilterVertexDeclaration.VertexDeclarationRHI;
+}
+
 template<bool bUseIESProfile, bool bRadialAttenuation, bool bInverseSquaredFalloff>
 static void SetShaderTemplLighting(
 	const FSceneView& View, 
@@ -652,13 +658,13 @@ static void SetShaderTemplLighting(
 	if(View.Family->EngineShowFlags.VisualizeLightCulling)
 	{
 		TShaderMapRef<TDeferredLightPS<false, bRadialAttenuation, false, true> > PixelShader(GetGlobalShaderMap());
-		SetGlobalBoundShaderState(PixelShader->GetBoundShaderState(), GFilterVertexDeclaration.VertexDeclarationRHI, VertexShader, *PixelShader);
+		SetGlobalBoundShaderState(PixelShader->GetBoundShaderState(), GetDeferredLightingVertexDeclaration<bRadialAttenuation>(), VertexShader, *PixelShader);
 		PixelShader->SetParameters(View, LightSceneInfo);
 	}
 	else
 	{
 		TShaderMapRef<TDeferredLightPS<bUseIESProfile, bRadialAttenuation, bInverseSquaredFalloff, false> > PixelShader(GetGlobalShaderMap());
-		SetGlobalBoundShaderState(PixelShader->GetBoundShaderState(), GFilterVertexDeclaration.VertexDeclarationRHI, VertexShader, *PixelShader);
+		SetGlobalBoundShaderState(PixelShader->GetBoundShaderState(), GetDeferredLightingVertexDeclaration<bRadialAttenuation>(), VertexShader, *PixelShader);
 		PixelShader->SetParameters(View, LightSceneInfo);
 	}
 }
@@ -672,13 +678,13 @@ static void SetShaderTemplLightingSimple(
 	if(View.Family->EngineShowFlags.VisualizeLightCulling)
 	{
 		TShaderMapRef<TDeferredLightPS<false, bRadialAttenuation, false, true> > PixelShader(GetGlobalShaderMap());
-		SetGlobalBoundShaderState(PixelShader->GetBoundShaderState(), GFilterVertexDeclaration.VertexDeclarationRHI, VertexShader, *PixelShader);
+		SetGlobalBoundShaderState(PixelShader->GetBoundShaderState(), GetDeferredLightingVertexDeclaration<bRadialAttenuation>(), VertexShader, *PixelShader);
 		PixelShader->SetParametersSimpleLight(View, SimpleLight);
 	}
 	else
 	{
 		TShaderMapRef<TDeferredLightPS<bUseIESProfile, bRadialAttenuation, bInverseSquaredFalloff, false> > PixelShader(GetGlobalShaderMap());
-		SetGlobalBoundShaderState(PixelShader->GetBoundShaderState(), GFilterVertexDeclaration.VertexDeclarationRHI, VertexShader, *PixelShader);
+		SetGlobalBoundShaderState(PixelShader->GetBoundShaderState(), GetDeferredLightingVertexDeclaration<bRadialAttenuation>(), VertexShader, *PixelShader);
 		PixelShader->SetParametersSimpleLight(View, SimpleLight);
 	}}
 
@@ -725,7 +731,7 @@ void FDeferredShadingSceneRenderer::RenderLight(const FLightSceneInfo* LightScen
 			if (bRenderOverlap)
 			{
 				TShaderMapRef<TDeferredLightOverlapPS<false> > PixelShader(GetGlobalShaderMap());
-				SetGlobalBoundShaderState(PixelShader->GetBoundShaderState(), GFilterVertexDeclaration.VertexDeclarationRHI, *VertexShader, *PixelShader);
+				SetGlobalBoundShaderState(PixelShader->GetBoundShaderState(), GetDeferredLightingVertexDeclaration<false>(), *VertexShader, *PixelShader);
 				PixelShader->SetParameters(View, LightSceneInfo);
 			}
 			else
@@ -762,7 +768,7 @@ void FDeferredShadingSceneRenderer::RenderLight(const FLightSceneInfo* LightScen
 			if (bRenderOverlap)
 			{
 				TShaderMapRef<TDeferredLightOverlapPS<true> > PixelShader(GetGlobalShaderMap());
-				SetGlobalBoundShaderState(PixelShader->GetBoundShaderState(), GetVertexDeclarationFVector4(), *VertexShader, *PixelShader);
+				SetGlobalBoundShaderState(PixelShader->GetBoundShaderState(), GetDeferredLightingVertexDeclaration<true>(), *VertexShader, *PixelShader);
 				PixelShader->SetParameters(View, LightSceneInfo);
 			}
 			else
