#ue4 Port reference Morpheus distortion shader.  Fix D3D11 crash when creating texture arrays.

diff --git a/Engine/Source/Runtime/HeadMountedDisplay/Public/IHeadMountedDisplay.h b/Engine/Source/Runtime/HeadMountedDisplay/Public/IHeadMountedDisplay.h
--- a/Engine/Source/Runtime/HeadMountedDisplay/Public/IHeadMountedDisplay.h
+++ b/Engine/Source/Runtime/HeadMountedDisplay/Public/IHeadMountedDisplay.h
@@ -206,10 +206,13 @@ public:
 	/** 
 	 * Additional optional distorion rendering parameters
 	 * @todo:  Once we can move shaders into plugins, remove these!
-	 */
-	virtual void GetImageTranslation(float& x, float& y) const {}
-	virtual void GetDistortionCenterOffset(float& x, float& y) const {}
-	virtual const FTexture*  GetDistortionTexture() {return NULL;}	
+	 */	
+	virtual FTexture* GetDistortionTextureLeft() const {return NULL;}
+	virtual FTexture* GetDistortionTextureRight() const {return NULL;}
+	virtual FVector2D GetTextureOffsetLeft() const {return FVector2D::ZeroVector;}
+	virtual FVector2D GetTextureOffsetRight() const {return FVector2D::ZeroVector;}
+	virtual FVector2D GetTextureScaleLeft() const {return FVector2D::ZeroVector;}
+	virtual FVector2D GetTextureScaleRight() const {return FVector2D::ZeroVector;}
 
 private:
 	/** Stores the dimensions of the window before we moved into fullscreen mode, so they can be restored */
diff --git a/Engine/Source/Runtime/Renderer/Private/PostProcess/PostProcessMorpheus.cpp b/Engine/Source/Runtime/Renderer/Private/PostProcess/PostProcessMorpheus.cpp
--- a/Engine/Source/Runtime/Renderer/Private/PostProcess/PostProcessMorpheus.cpp
+++ b/Engine/Source/Runtime/Renderer/Private/PostProcess/PostProcessMorpheus.cpp
@@ -36,24 +36,20 @@ class FPostProcessMorpheusPS : public FGlobalShader
 
 	/** Default constructor. */
 	FPostProcessMorpheusPS()
-	{
-		DistortionTexture = NULL;
+	{		
 	}
 
 public:
 	FPostProcessPassParameters PostprocessParameter;
 	FDeferredPixelShaderParameters DeferredParameters;
 	
-	// Distortion parameter values
-	FShaderParameter LensCenter;
-	FShaderParameter ScreenCenter;
-	FShaderParameter Scale;
-	FShaderParameter HMDWarpParam;
-	FShaderParameter CAWarpParam;
-	FShaderResourceParameter DistortionTextureParam; 
-	FShaderResourceParameter DistortionTextureSampler; 
-	const FTexture*		DistortionTexture;
+	// Distortion parameter values	
+	FShaderParameter TextureScale;
+	FShaderParameter TextureOffset;
+	FShaderParameter TextureUVOffset;
 
+	FShaderResourceParameter DistortionTextureParam; 	
+	FShaderResourceParameter DistortionTextureSampler; 
 
 	/** Initialization constructor. */
 	FPostProcessMorpheusPS(const ShaderMetaType::CompiledShaderInitializerType& Initializer)
@@ -62,19 +58,20 @@ public:
 		PostprocessParameter.Bind(Initializer.ParameterMap);
 		DeferredParameters.Bind(Initializer.ParameterMap);
 
-		LensCenter.Bind(Initializer.ParameterMap, TEXT("LensCenter"));
-		ScreenCenter.Bind(Initializer.ParameterMap, TEXT("ScreenCenter"));
-		Scale.Bind(Initializer.ParameterMap, TEXT("Scale"));
-		HMDWarpParam.Bind(Initializer.ParameterMap, TEXT("HMDWarpParam"));
-		CAWarpParam.Bind(Initializer.ParameterMap, TEXT("CAWarpParam"));
-	
-		DistortionTextureParam.Bind(Initializer.ParameterMap, TEXT("DistortionTexture"));
-		check(DistortionTextureParam.IsBound());
+		TextureScale.Bind(Initializer.ParameterMap, TEXT("TextureScale"));
+		//check(TextureScaleLeft.IsBound());		
 
-		DistortionTextureSampler.Bind(Initializer.ParameterMap, TEXT("DistortionTextureSampler"));
-		check(DistortionTextureSampler.IsBound());
+		TextureOffset.Bind(Initializer.ParameterMap, TEXT("TextureOffset"));
+		//check(TextureOffsetRight.IsBound());
+
+		TextureUVOffset.Bind(Initializer.ParameterMap, TEXT("TextureUVOffset"));
+		//check(TextureUVOffset.IsBound());
+			
+		DistortionTextureParam.Bind(Initializer.ParameterMap, TEXT("DistortionTextureArray"));
+		//check(DistortionTextureLeftParam.IsBound());		
 
-		DistortionTexture = NULL;
+		DistortionTextureSampler.Bind(Initializer.ParameterMap, TEXT("DistortionTextureSampler"));
+		//check(DistortionTextureSampler.IsBound());		
 	}
 
 
@@ -89,69 +86,27 @@ public:
 
 		{
 			check(GEngine->HMDDevice.IsValid());
+			TSharedPtr< class IHeadMountedDisplay > HMDDevice = GEngine->HMDDevice;
 
-            const FIntPoint SrcSize = SrcRect.Size();
-            const float BufferRatioX = float(SrcSize.X)/float(SrcBufferSize.X);
-            const float BufferRatioY = float(SrcSize.Y)/float(SrcBufferSize.Y);
-            const float w = float(SrcSize.X)/float(SrcBufferSize.X);
-            const float h = float(SrcSize.Y)/float(SrcBufferSize.Y);
-            const float x = float(SrcRect.Min.X)/float(SrcBufferSize.X);
-            const float y = float(SrcRect.Min.Y)/float(SrcBufferSize.Y);
-
-			float DistortionCenterOffsetX = 0, DistortionCenterOffsetY = 0;
-			GEngine->HMDDevice->GetDistortionCenterOffset(DistortionCenterOffsetX, DistortionCenterOffsetY);
-
-			const float XCenterOffset = (StereoPass == eSSP_RIGHT_EYE) ? -DistortionCenterOffsetX : DistortionCenterOffsetX;
-			const float YCenterOffset = DistortionCenterOffsetY;
-
-			const float AspectRatio = (float)SrcSize.X / (float)SrcSize.Y;
-			const float ScaleFactor = GEngine->HMDDevice->GetDistortionScalingFactor();
-
-			if(DistortionTexture == NULL)
+			check (StereoPass != eSSP_FULL);
+			if (StereoPass == eSSP_LEFT_EYE)
 			{
-				check(IsInRenderingThread());
-				DistortionTexture = GEngine->HMDDevice->GetDistortionTexture();
-				check(DistortionTexture != NULL);
+				FTexture* TextureLeft = HMDDevice->GetDistortionTextureLeft();
+				SetTextureParameter(ShaderRHI, DistortionTextureParam, DistortionTextureSampler, TextureLeft->SamplerStateRHI, TextureLeft->TextureRHI);
+				SetShaderValue(ShaderRHI, TextureScale, HMDDevice->GetTextureScaleLeft());
+				SetShaderValue(ShaderRHI, TextureOffset, HMDDevice->GetTextureOffsetLeft());
+				SetShaderValue(ShaderRHI, TextureUVOffset, 0.0f);
 			}
-			SetTextureParameter(ShaderRHI, DistortionTextureParam, DistortionTextureSampler, TStaticSamplerState<SF_Bilinear,AM_Clamp,AM_Clamp,AM_Clamp>::GetRHI(), DistortionTexture->TextureRHI);
-
-			// Shifts texture coordinates to the center of the distortion function around the center of the lens
-			FVector2D ViewLensCenter;
-			ViewLensCenter.X = x + (0.5f + XCenterOffset) * w;
-			ViewLensCenter.Y = y + (0.5f + YCenterOffset) * h;
-			SetShaderValue(ShaderRHI, LensCenter, ViewLensCenter);
-
-			// Texture coordinate for the center of the half scene texture, used to clamp sampling
-			FVector2D ViewScreenCenter;
-			ViewScreenCenter.X = x + w * 0.5f;
-			ViewScreenCenter.Y = y + h * 0.5f;
-			SetShaderValue(ShaderRHI, ScreenCenter, ViewScreenCenter);
-			
-			// Rescale output (sample) coordinates back to texture range and increase scale to support rendering outside the the screen
-			FVector2D ViewScale;
-			ViewScale.X = (w/2) * ScaleFactor;
-			ViewScale.Y = (h/2) * ScaleFactor * AspectRatio;
-            SetShaderValue(ShaderRHI, Scale, ViewScale);
-
-            // Distortion coefficients
-            FVector4 DistortionValues;
-            GEngine->HMDDevice->GetDistortionWarpValues(DistortionValues);
-            SetShaderValue(ShaderRHI, HMDWarpParam, DistortionValues);
-
-			// CNN - Morpheus changes
-			// CA correction values
-			FVector4 CAValues;
-			GEngine->HMDDevice->GetChromaAbCorrectionValues(CAValues);
-			SetShaderValue(ShaderRHI, CAWarpParam, CAValues);
-
-			// Rescale the texture coordinates to [-1,1] unit range and corrects aspect ratio
-			FVector2D ViewScaleIn;
-            ViewScaleIn.X = (2/w);
-            ViewScaleIn.Y = (2/h) / AspectRatio;
-            
-            QuadTexTransform = FMatrix::Identity;
-            QuadTexTransform *= FTranslationMatrix(FVector(-ViewLensCenter.X * SrcBufferSize.X, -ViewLensCenter.Y * SrcBufferSize.Y, 0));
-            QuadTexTransform *= FMatrix(FPlane(ViewScaleIn.X,0,0,0), FPlane(0,ViewScaleIn.Y,0,0), FPlane(0,0,0,0), FPlane(0,0,0,1));
+			else
+			{
+				FTexture* TextureRight = HMDDevice->GetDistortionTextureRight();
+				SetTextureParameter(ShaderRHI, DistortionTextureParam, DistortionTextureSampler, TextureRight->SamplerStateRHI, TextureRight->TextureRHI);
+				SetShaderValue(ShaderRHI, TextureScale, HMDDevice->GetTextureScaleRight());
+				SetShaderValue(ShaderRHI, TextureOffset, HMDDevice->GetTextureOffsetRight());
+				SetShaderValue(ShaderRHI, TextureUVOffset, -0.5f);
+			}				
+			      
+            QuadTexTransform = FMatrix::Identity;            
 		}
 	}
 	
@@ -159,7 +114,7 @@ public:
 	virtual bool Serialize(FArchive& Ar)
 	{
 		bool bShaderHasOutdatedParameters = FGlobalShader::Serialize(Ar);
-		Ar << PostprocessParameter << DeferredParameters << LensCenter << ScreenCenter << Scale << HMDWarpParam << CAWarpParam << DistortionTextureParam << DistortionTextureSampler;
+		Ar << PostprocessParameter << DeferredParameters << TextureScale << TextureOffset << TextureUVOffset << DistortionTextureParam << DistortionTextureSampler;
 		return bShaderHasOutdatedParameters;
 	}
 };
@@ -216,7 +171,7 @@ void FRCPassPostProcessMorpheus::Process(FRenderingCompositePassContext& Context
 	{
 		// input is not hooked up correctly
 		return;
-	}
+	}	
 
 	const FSceneView& View = Context.View;
 	const FSceneViewFamily& ViewFamily = *(View.Family);
diff --git a/Engine/Source/Runtime/Windows/D3D11RHI/Private/D3D11Texture.cpp b/Engine/Source/Runtime/Windows/D3D11RHI/Private/D3D11Texture.cpp
--- a/Engine/Source/Runtime/Windows/D3D11RHI/Private/D3D11Texture.cpp
+++ b/Engine/Source/Runtime/Windows/D3D11RHI/Private/D3D11Texture.cpp
@@ -632,17 +632,29 @@ TD3D11Texture2D<BaseResourceType>* FD3D11DynamicRHI::CreateD3D11Texture2D(uint32
 		if (BulkData)
 		{
 			uint8* Data = (uint8*)BulkData->GetResourceBulkData();
-			SubResourceData.AddZeroed(NumMips);
-			uint32 MipOffset = 0;
-			for(uint32 MipIndex = 0;MipIndex < NumMips;++MipIndex)
-			{
-				uint32 NumBlocksX = FMath::Max<uint32>(1,(SizeX >> MipIndex) / GPixelFormats[Format].BlockSizeX);
-				uint32 NumBlocksY = FMath::Max<uint32>(1,(SizeY >> MipIndex) / GPixelFormats[Format].BlockSizeY);
 
-				SubResourceData[MipIndex].pSysMem = &Data[MipOffset];
-				SubResourceData[MipIndex].SysMemPitch      =  NumBlocksX * GPixelFormats[Format].BlockBytes;
-				SubResourceData[MipIndex].SysMemSlicePitch =  NumBlocksX * NumBlocksY * SubResourceData[MipIndex].SysMemPitch;
-				MipOffset                                  += SubResourceData[MipIndex].SysMemSlicePitch;
+			// each mip of each array slice counts as a subresource
+			SubResourceData.AddZeroed(NumMips * SizeZ);
+
+			uint32 SliceOffset = 0;
+			for (uint32 ArraySliceIndex = 0; ArraySliceIndex < SizeZ; ++ArraySliceIndex)
+			{			
+				uint32 MipOffset = 0;
+				for(uint32 MipIndex = 0;MipIndex < NumMips;++MipIndex)
+				{
+					uint32 DataOffset = SliceOffset + MipOffset;
+					uint32 SubResourceIndex = ArraySliceIndex * NumMips + MipIndex;
+
+					uint32 NumBlocksX = FMath::Max<uint32>(1,(SizeX >> MipIndex) / GPixelFormats[Format].BlockSizeX);
+					uint32 NumBlocksY = FMath::Max<uint32>(1,(SizeY >> MipIndex) / GPixelFormats[Format].BlockSizeY);
+
+					SubResourceData[SubResourceIndex].pSysMem = &Data[DataOffset];
+					SubResourceData[SubResourceIndex].SysMemPitch      =  NumBlocksX * GPixelFormats[Format].BlockBytes;
+					SubResourceData[SubResourceIndex].SysMemSlicePitch =  NumBlocksX * NumBlocksY * SubResourceData[MipIndex].SysMemPitch;
+
+					MipOffset                                  += NumBlocksY * SubResourceData[MipIndex].SysMemPitch;
+				}
+				SliceOffset += MipOffset;
 			}
 		}
 
