#UE4:
* added r.DumpRenderTargetPoolEvents to visualize RenderTargetPool usage over a frame, added mouse hover support

diff --git a/Engine/Source/Editor/UnrealEd/Private/EditorViewportClient.cpp b/Engine/Source/Editor/UnrealEd/Private/EditorViewportClient.cpp
--- a/Engine/Source/Editor/UnrealEd/Private/EditorViewportClient.cpp
+++ b/Engine/Source/Editor/UnrealEd/Private/EditorViewportClient.cpp
@@ -209,6 +209,7 @@ FEditorViewportClient::FEditorViewportClient(FPreviewScene* InPreviewScene)
 	, LastMouseY(0)
 	, CachedMouseX(0)
 	, CachedMouseY(0)
+	, CurrentMousePos(-1, -1)
 	, bIsRealtime(false)
 	, bStoredRealtime(false)
 	, bStoredShowFPS(false)
@@ -610,6 +611,7 @@ FSceneView* FEditorViewportClient::CalcSceneView(FSceneViewFamily* ViewFamily)
 	ViewInitOptions.bUseFauxOrthoViewPos = true;
 
 	ViewInitOptions.OverrideFarClippingPlaneDistance = FarPlane;
+	ViewInitOptions.CursorPos = CurrentMousePos;
 
 	FSceneView* View = new FSceneView(ViewInitOptions);
 
@@ -1929,8 +1931,6 @@ void FEditorViewportClient::ProcessClickInViewport( const FInputEventState& Inpu
 		const int32	HitX = InputStateViewport->GetMouseX();
 		const int32	HitY = InputStateViewport->GetMouseY();
 		
-		FIntPoint CurrentMousePos(HitX,HitY);
-
 		// Calc the raw delta from the mouse to detect if there was any movement
 		FVector RawMouseDelta = MouseDeltaTracker->GetScreenDelta();
 
@@ -3656,6 +3656,17 @@ void FEditorViewportClient::OnJoystickPlugged(const uint32 InControllerID, const
 void FEditorViewportClient::MouseMove(FViewport* InViewport,int32 x, int32 y)
 {
 	check(IsInGameThread());
+
+	CurrentMousePos = FIntPoint(x, y);
+}
+
+void FEditorViewportClient::MouseLeave( FViewport* Viewport )
+{
+	check(IsInGameThread());
+
+	CurrentMousePos = FIntPoint(-1, -1);
+
+	FCommonViewportClient::MouseLeave(Viewport);
 }
 
 void FEditorViewportClient::CapturedMouseMove( FViewport* InViewport, int32 InMouseX, int32 InMouseY )
diff --git a/Engine/Source/Editor/UnrealEd/Public/EditorViewportClient.h b/Engine/Source/Editor/UnrealEd/Public/EditorViewportClient.h
--- a/Engine/Source/Editor/UnrealEd/Public/EditorViewportClient.h
+++ b/Engine/Source/Editor/UnrealEd/Public/EditorViewportClient.h
@@ -410,6 +410,7 @@ public:
 	virtual void ReceivedFocus(FViewport* Viewport) OVERRIDE;
 	virtual void OnJoystickPlugged(const uint32 InControllerID, const uint32 InType, const uint32 bInConnected)  OVERRIDE;
 	virtual void MouseMove(FViewport* Viewport,int32 x, int32 y) OVERRIDE;
+	virtual void MouseLeave( FViewport* Viewport ) OVERRIDE;
 	virtual EMouseCursor::Type GetCursor(FViewport* Viewport,int32 X,int32 Y) OVERRIDE;
 	virtual void CapturedMouseMove( FViewport* InViewport, int32 InMouseX, int32 InMouseY ) OVERRIDE;
 	virtual bool IsOrtho() const OVERRIDE;
@@ -1164,13 +1165,17 @@ protected:
 	bool bUseControllingActorViewInfo;
 	FMinimalViewInfo ControllingActorViewInfo;
 
+	/* Updated on each mouse drag start */
 	uint32 LastMouseX;
 	uint32 LastMouseY;
 
-	/** Represents the last known mouse position. */
+	/** Represents the last known mouse position. If the mouse stops moving it's not the current but the last position before the current location. */
 	uint32 CachedMouseX;
 	uint32 CachedMouseY;
 
+	// -1, -1 if not set
+	FIntPoint CurrentMousePos;
+
 	/**
 	 * true when within a FMouseDeltaTracker::StartTracking/EndTracking block.
 	 */
diff --git a/Engine/Source/Runtime/Engine/Private/LocalPlayer.cpp b/Engine/Source/Runtime/Engine/Private/LocalPlayer.cpp
--- a/Engine/Source/Runtime/Engine/Private/LocalPlayer.cpp
+++ b/Engine/Source/Runtime/Engine/Private/LocalPlayer.cpp
@@ -371,6 +371,7 @@ FSceneView* ULocalPlayer::CalcSceneView( class FSceneViewFamily* ViewFamily,
 	ViewInitOptions.LODDistanceFactor = PlayerController->LocalPlayerCachedLODDistanceFactor;
 	ViewInitOptions.StereoPass = StereoPass;
 	ViewInitOptions.WorldToMetersScale = PlayerController->GetWorldSettings()->WorldToMeters;
+	ViewInitOptions.CursorPos = Viewport->HasMouseCapture() ? FIntPoint(-1, -1) : FIntPoint(Viewport->GetMouseX(), Viewport->GetMouseY());
 	PlayerController->BuildHiddenComponentList(OutViewLocation, /*out*/ ViewInitOptions.HiddenPrimitives);
 
 	FSceneView* const View = new FSceneView(ViewInitOptions);
diff --git a/Engine/Source/Runtime/Engine/Private/SceneView.cpp b/Engine/Source/Runtime/Engine/Private/SceneView.cpp
--- a/Engine/Source/Runtime/Engine/Private/SceneView.cpp
+++ b/Engine/Source/Runtime/Engine/Private/SceneView.cpp
@@ -151,6 +151,7 @@ FSceneView::FSceneView(const FSceneViewInitOptions& InitOptions)
 	, HiddenPrimitives(InitOptions.HiddenPrimitives)
 	, LODDistanceFactor(InitOptions.LODDistanceFactor)
 	, bCameraCut(InitOptions.bInCameraCut)
+	, CursorPos(InitOptions.CursorPos)
 	, bIsGameView(false)
 	, bForceShowMaterials(false)
 	, bIsViewInfo(false)
diff --git a/Engine/Source/Runtime/Engine/Private/Slate/SceneViewport.cpp b/Engine/Source/Runtime/Engine/Private/Slate/SceneViewport.cpp
--- a/Engine/Source/Runtime/Engine/Private/Slate/SceneViewport.cpp
+++ b/Engine/Source/Runtime/Engine/Private/Slate/SceneViewport.cpp
@@ -15,6 +15,7 @@ extern EWindowMode::Type GetWindowModeType(EWindowMode::Type WindowMode);
 
 FSceneViewport::FSceneViewport( FViewportClient* InViewportClient, TSharedPtr<SViewport> InViewportWidget )
 	: FViewport( InViewportClient )
+	, CachedMousePos(-1, -1)
 	, CurrentReplyState( FReply::Unhandled() )
 	, PreCaptureMousePos( -1, -1 )
 	, SoftwareCursorPosition( 0, 0 )
@@ -383,6 +384,7 @@ void FSceneViewport::OnMouseEnter( const FGeometry& MyGeometry, const FPointerEv
 
 void FSceneViewport::OnMouseLeave( const FPointerEvent& MouseEvent )
 {
+	CachedMousePos = FIntPoint(-1, -1);
 	ViewportClient->MouseLeave( this );
 }
 
diff --git a/Engine/Source/Runtime/Engine/Public/SceneManagement.h b/Engine/Source/Runtime/Engine/Public/SceneManagement.h
--- a/Engine/Source/Runtime/Engine/Public/SceneManagement.h
+++ b/Engine/Source/Runtime/Engine/Public/SceneManagement.h
@@ -2475,6 +2475,8 @@ struct FSceneViewInitOptions : public FSceneViewProjectionData
 
 	TSet<FPrimitiveComponentId> HiddenPrimitives;
 
+	// -1,-1 if not setup
+	FIntPoint CursorPos;
 
 	float LODDistanceFactor;
 
@@ -2505,6 +2507,7 @@ struct FSceneViewInitOptions : public FSceneViewProjectionData
 		, ColorScale(FLinearColor::White)
 		, StereoPass(eSSP_FULL)
 		, WorldToMetersScale(100.f)
+		, CursorPos(-1, -1)
 		, LODDistanceFactor(1.0f)
 		, OverrideFarClippingPlaneDistance(-1.0f)
 		, bInCameraCut(false)
@@ -2628,6 +2631,9 @@ public:
 	/** Whether we did a camera cut for this view this frame. */
 	bool bCameraCut;
 
+	// -1,-1 if not setup
+	FIntPoint CursorPos;
+
 	/** True if this scene was created from a game world. */
 	bool bIsGameView;
 
diff --git a/Engine/Source/Runtime/Engine/Public/Slate/SceneViewport.h b/Engine/Source/Runtime/Engine/Public/Slate/SceneViewport.h
--- a/Engine/Source/Runtime/Engine/Public/Slate/SceneViewport.h
+++ b/Engine/Source/Runtime/Engine/Public/Slate/SceneViewport.h
@@ -255,7 +255,7 @@ private:
 	FReply CurrentReplyState;
 	/** A mapping of key names to their pressed state */
 	TMap<FKey,bool> KeyStateMap;
-	/** The last known mouse position in local space */
+	/** The last known mouse position in local space, -1, -1 if unknown */
 	FIntPoint CachedMousePos;
 	/** The last known geometry info */
 	FGeometry CachedGeometry;
diff --git a/Engine/Source/Runtime/Renderer/Private/DeferredShadingRenderer.cpp b/Engine/Source/Runtime/Renderer/Private/DeferredShadingRenderer.cpp
--- a/Engine/Source/Runtime/Renderer/Private/DeferredShadingRenderer.cpp
+++ b/Engine/Source/Runtime/Renderer/Private/DeferredShadingRenderer.cpp
@@ -509,6 +509,8 @@ void FDeferredShadingSceneRenderer::Render()
 		Scene->FXSystem->PreRender();
 	}
 
+	GRenderTargetPool.AddPhaseEvent(TEXT("EarlyZPass"));
+
 	// Draw the scene pre-pass / early z pass, populating the scene depth buffer and HiZ
 	RenderPrePass();
 	
@@ -557,6 +559,8 @@ void FDeferredShadingSceneRenderer::Render()
 		GSceneRenderTargets.BeginRenderingSceneColor(true);
 	}
 
+	GRenderTargetPool.AddPhaseEvent(TEXT("BasePass"));
+
 	RenderBasePass();
 
 	if(ViewFamily.EngineShowFlags.VisualizeLightCulling)
@@ -617,6 +621,8 @@ void FDeferredShadingSceneRenderer::Render()
 		&& ViewFamily.EngineShowFlags.DeferredLighting
 		)
 	{
+		GRenderTargetPool.AddPhaseEvent(TEXT("Lighting"));
+
 		// Pre-lighting composition lighting stage
 		// e.g. deferred decals, blurred GBuffer
 		for(int32 ViewIndex = 0;ViewIndex < Views.Num();ViewIndex++)
@@ -694,6 +700,8 @@ void FDeferredShadingSceneRenderer::Render()
 		}
 	}
 
+	GRenderTargetPool.AddPhaseEvent(TEXT("Fog"));
+
 	// Draw fog.
 	if(ShouldRenderFog(ViewFamily))
 	{
@@ -703,6 +711,8 @@ void FDeferredShadingSceneRenderer::Render()
 	// No longer needed, release
 	LightShaftOutput.LightShaftOcclusion = NULL;
 
+	GRenderTargetPool.AddPhaseEvent(TEXT("Translucency"));
+
 	// Draw translucency.
 	if(ViewFamily.EngineShowFlags.Translucency)
 	{
diff --git a/Engine/Source/Runtime/Renderer/Private/PostProcess/PostProcessing.cpp b/Engine/Source/Runtime/Renderer/Private/PostProcess/PostProcessing.cpp
--- a/Engine/Source/Runtime/Renderer/Private/PostProcess/PostProcessing.cpp
+++ b/Engine/Source/Runtime/Renderer/Private/PostProcess/PostProcessing.cpp
@@ -709,6 +709,7 @@ void FPostProcessing::Process(const FViewInfo& View, TRefCountPtr<IPooledRenderT
 
 	bool bSimpleDynamicLighting = IsSimpleDynamicLightingEnabled();
 
+	GRenderTargetPool.AddPhaseEvent(TEXT("PostProcessing"));
 
 	// This page: https://udn.epicgames.com/Three/RenderingOverview#Rendering%20state%20defaults 
 	// describes what state a pass can expect and to what state it need to be set back.
diff --git a/Engine/Source/Runtime/Renderer/Private/PostProcess/RenderTargetPool.cpp b/Engine/Source/Runtime/Renderer/Private/PostProcess/RenderTargetPool.cpp
--- a/Engine/Source/Runtime/Renderer/Private/PostProcess/RenderTargetPool.cpp
+++ b/Engine/Source/Runtime/Renderer/Private/PostProcess/RenderTargetPool.cpp
@@ -8,12 +8,88 @@
 #include "ScenePrivate.h"
 #include "RenderTargetPool.h"
 
-
 /** The global render targets pool. */
 TGlobalResource<FRenderTargetPool> GRenderTargetPool;
 
 DEFINE_LOG_CATEGORY_STATIC(LogRenderTargetPool, Warning, All);
 
+static void DumpRenderTargetPoolMemory(FOutputDevice& OutputDevice)
+{
+	GRenderTargetPool.DumpMemoryUsage(OutputDevice);
+}
+static FAutoConsoleCommandWithOutputDevice GDumpRenderTargetPoolMemoryCmd(
+	TEXT("r.DumpRenderTargetPoolMemory"),
+	TEXT("Dump allocation information for the render target pool."),
+	FConsoleCommandWithOutputDeviceDelegate::CreateStatic(DumpRenderTargetPoolMemory)
+	);
+
+static void DumpRenderTargetPoolEvents(const TArray<FString>& Args)
+{
+	uint32 SizeInKBThreshold = -1;
+	if(Args.Num() && Args[0].IsNumeric())
+	{
+		SizeInKBThreshold = FCString::Atof(*Args[0]);
+	}
+
+	if(SizeInKBThreshold != -1)
+	{
+		UE_LOG(LogRenderTargetPool, Display, TEXT("r.DumpRenderTargetPoolEvents is now enabled, use r.DumpRenderTargetPoolEvents ? for help"));
+
+		GRenderTargetPool.EnableEventRecording(SizeInKBThreshold);
+	}
+	else
+	{
+		GRenderTargetPool.DisableEventDisplay();
+
+		UE_LOG(LogRenderTargetPool, Display, TEXT("r.DumpRenderTargetPoolEvents is now disabled, use r.DumpRenderTargetPoolEvents <SizeInKB> to enable or r.DumpRenderTargetPoolEvents ? for help"));
+	}
+}
+static FAutoConsoleCommand GDumpRenderTargetPoolEventsCmd(
+	TEXT("r.DumpRenderTargetPoolEvents"),
+	TEXT("Dump for the render target pool events. Optional parameter defines threshold in KB (default is 1024 for 1MB)."),
+	FConsoleCommandWithArgsDelegate::CreateStatic(DumpRenderTargetPoolEvents)
+	);
+
+bool FRenderTargetPool::IsEventRecordingEnabled() const
+{
+#if !(UE_BUILD_SHIPPING || UE_BUILD_TEST)
+	return bEventRecording; 
+#else
+	return false;
+#endif
+}
+
+IPooledRenderTarget* FRenderTargetPoolEvent::GetValidatedPointer() const
+{
+	int32 Index = GRenderTargetPool.FindIndex(Pointer);
+
+	if(Index >= 0)
+	{
+		return Pointer;
+	}
+
+	return 0;
+}
+
+bool FRenderTargetPoolEvent::NeedsDeallocEvent()
+{
+	if(GetEventType() == ERTPE_Alloc)
+	{
+		if(Pointer)
+		{
+			IPooledRenderTarget* ValidPointer = GetValidatedPointer();
+			if(ValidPointer->IsFree())
+			{
+				Pointer = 0;
+				return true;
+			}
+		}
+	}
+
+	return false;
+}
+
+
 static uint32 ComputeSizeInKB(FPooledRenderTarget& Element)
 {
 	return (Element.ComputeMemorySize() + 1023) / 1024;
@@ -22,6 +98,10 @@ static uint32 ComputeSizeInKB(FPooledRenderTarget& Element)
 FRenderTargetPool::FRenderTargetPool()
 	: AllocationLevelInKB(0)
 	, bCurrentlyOverBudget(false)
+	, bEventRecording(false)
+	, bEventRecordingTrigger(false)
+	, SizeInKBThreshold(0)
+	, CurrentTimeStep(0)
 {
 }
 
@@ -55,23 +135,28 @@ bool FRenderTargetPool::FindFreeElement(const FPooledRenderTargetDesc& Desc, TRe
 
 			if(Current->IsFree())
 			{
-				// Adding this to the pool would be a waste, this seems to be a reallocation
-				// where we want to change the size or format
-				PooledRenderTargets.Remove(Current);
+				int32 Index = FindIndex(Current);
+
+				check(Index >= 0);
+
+				// we don't use Remove() to not shuffle around the elements for better transparency on RenderTargetPoolEvents
+				PooledRenderTargets[Index] = 0;
 			}
 		}
 	}
 
 	FPooledRenderTarget* Found = 0;
+	uint32 FoundIndex = -1;
 
 	// try to find a suitable element in the pool
 	for(uint32 i = 0; i < (uint32)PooledRenderTargets.Num(); ++i)
 	{
 		FPooledRenderTarget* Element = PooledRenderTargets[i];
 
-		if(Element->IsFree() && Element->GetDesc() == Desc)
+		if(Element && Element->IsFree() && Element->GetDesc() == Desc)
 		{
 			Found = Element;
+			FoundIndex = i;
 			break;
 		}
 	}
@@ -203,6 +288,8 @@ bool FRenderTargetPool::FindFreeElement(const FPooledRenderTargetDesc& Desc, TRe
 
 		AllocationLevelInKB += ComputeSizeInKB(*Found);
 		VerifyAllocationLevel();
+
+		FoundIndex = PooledRenderTargets.Num() - 1;
 	}
 
 #if !(UE_BUILD_SHIPPING || UE_BUILD_TEST)
@@ -236,6 +323,8 @@ bool FRenderTargetPool::FindFreeElement(const FPooledRenderTargetDesc& Desc, TRe
 	Found->Desc.DebugName = InDebugName;
 	Found->UnusedForNFrames = 0;
 
+	AddAllocEvent(FoundIndex, Found);
+
 	// assign to the reference counted variable
 	Out = Found;
 
@@ -269,23 +358,569 @@ void FRenderTargetPool::GetStats(uint32& OutWholeCount, uint32& OutWholePoolInKB
 	{
 		FPooledRenderTarget* Element = PooledRenderTargets[i];
 
-		uint32 SizeInKB = ComputeSizeInKB(*Element);
-		
-		OutWholePoolInKB += SizeInKB;
-
-		if(!Element->IsFree())
+		if(Element)
 		{
-			OutUsedInKB += SizeInKB;
+			uint32 SizeInKB = ComputeSizeInKB(*Element);
+
+			OutWholePoolInKB += SizeInKB;
+
+			if(!Element->IsFree())
+			{
+				OutUsedInKB += SizeInKB;
+			}
 		}
 	}
 
 	check(AllocationLevelInKB == OutWholePoolInKB);
 }
 
+void FRenderTargetPool::AddPhaseEvent(const TCHAR *InPhaseName)
+{
+	if(IsEventRecordingEnabled())
+	{
+		AddDeallocEvents();
+
+		const FString* LastName = GetLastEventPhaseName();
+
+		if(!LastName || *LastName != InPhaseName)
+		{
+			if(CurrentTimeStep)
+			{
+				// put a break to former data
+				++CurrentTimeStep;
+			}
+
+			FRenderTargetPoolEvent NewEvent(InPhaseName, CurrentTimeStep);
+
+			RenderTargetPoolEvents.Add(NewEvent);
+		}
+	}
+}
+
+// helper class to get a consistent layout in multiple functions
+// MaxX and Y are the output value that can be requested during or after iteration
+// Examples usages:
+//    FRenderTargetPoolEventIterator It(RenderTargetPoolEvents, OptionalStartIndex);
+//    while(FRenderTargetPoolEvent* Event = It.Iterate()) {}
+struct FRenderTargetPoolEventIterator
+{
+	int32 Index;
+	TArray<FRenderTargetPoolEvent>& RenderTargetPoolEvents;
+	bool bLineContent;
+	uint32 TotalWidth;
+	int32 Y;
+
+	// constructor
+	FRenderTargetPoolEventIterator(TArray<FRenderTargetPoolEvent>& InRenderTargetPoolEvents, int32 InIndex = 0)
+		: Index(InIndex)
+		, RenderTargetPoolEvents(InRenderTargetPoolEvents)
+		, bLineContent(false)
+		, TotalWidth(1)
+		, Y(0)
+	{
+		Touch();
+	}
+
+	FRenderTargetPoolEvent* operator*()
+	{
+		if(Index < RenderTargetPoolEvents.Num())
+		{
+			return &RenderTargetPoolEvents[Index];
+		}
+
+		return 0;
+	}
+
+	// @return 0 if end was reached
+	FRenderTargetPoolEventIterator& operator++()
+	{
+		if(Index < RenderTargetPoolEvents.Num())
+		{
+			++Index;
+		}
+
+		Touch();
+
+		return *this;
+	}
+	
+	int32 FindClosingEventY() const
+	{
+		FRenderTargetPoolEventIterator It = *this;
+
+		const ERenderTargetPoolEventType StartType = (*It)->GetEventType();
+
+		if(StartType == ERTPE_Alloc)
+		{
+			int32 PoolEntryId = RenderTargetPoolEvents[Index].GetPoolEntryId();
+
+			++It;
+
+			// search for next Dealloc of the same PoolEntryId
+			for(; *It; ++It)
+			{
+				FRenderTargetPoolEvent* Event = *It;
+
+				if(Event->GetEventType() == ERTPE_Dealloc && Event->GetPoolEntryId() == PoolEntryId)
+				{
+					break;
+				}
+			}
+		}
+		else if(StartType == ERTPE_Phase)
+		{
+			++It;
+
+			// search for next Phase
+			for(; *It; ++It)
+			{
+				FRenderTargetPoolEvent* Event = *It;
+
+				if(Event->GetEventType() == ERTPE_Phase)
+				{
+					break;
+				}
+			}
+		}
+		else
+		{
+			check(0);
+		}
+
+		return It.Y;
+	}
+
+private:
+
+	void Touch()
+	{
+		if(Index < RenderTargetPoolEvents.Num())
+		{
+			const FRenderTargetPoolEvent& Event = RenderTargetPoolEvents[Index];
+
+			const ERenderTargetPoolEventType Type = Event.GetEventType();
+
+			if(Type == ERTPE_Alloc)
+			{
+				// for now they are all equal width
+				TotalWidth = FMath::Max(TotalWidth, Event.GetColumnX() + Event.GetColumnSize());
+			}
+			Y = Event.GetTimeStep();
+		}
+	}
+};
+
+FIntPoint FRenderTargetPool::ComputeEventDisplayExtent()
+{
+	FRenderTargetPoolEventIterator It(RenderTargetPoolEvents);
+
+	for(; *It; ++It)
+	{
+	}
+
+	return FIntPoint(It.TotalWidth, It.Y + 1);
+}
+
+const FString* FRenderTargetPool::GetLastEventPhaseName()
+{
+	// could be optimized but this is a debug view
+
+	// start from the end for better performance
+	for(int32 i = RenderTargetPoolEvents.Num() - 1; i >= 0; --i)
+	{
+		const FRenderTargetPoolEvent* Event = &RenderTargetPoolEvents[i];
+
+		if(Event->GetEventType() == ERTPE_Phase)
+		{
+			return &Event->GetPhaseName();
+		}
+	}
+
+	return 0;
+}
+
+FRenderTargetPool::SMemoryStats FRenderTargetPool::ComputeView()
+{
+	SMemoryStats MemoryStats;
+
+#if !(UE_BUILD_SHIPPING || UE_BUILD_TEST)
+	{
+		struct FRTPColumn
+		{
+			// index into the column, -1 if this is no valid column
+			uint32 PoolEntryId;
+			// for sorting
+			uint64 SizeInBytes;
+			// for sorting
+			bool bVRam;
+
+			// default constructor
+			FRTPColumn()
+				: PoolEntryId(-1)
+				, SizeInBytes(0)
+			{
+			}
+
+			// constructor
+			FRTPColumn(const FRenderTargetPoolEvent& Event)
+				: PoolEntryId(Event.GetPoolEntryId())
+				, SizeInBytes(Event.GetSizeInBytes())
+				, bVRam((Event.GetDesc().Flags & TexCreate_FastVRAM) != 0)
+			{
+			}
+
+			// sort criteria
+			bool operator <(const FRTPColumn& rhs) const
+			{
+				// sort VRam first (only matters on XboxOne but nice to alwasy see it)
+				if(bVRam != rhs.bVRam) return bVRam > rhs.bVRam;
+
+				// we want the large ones first
+				return SizeInBytes > rhs.SizeInBytes;
+			}
+		};
+
+		TArray<FRTPColumn> Colums;
+
+		// generate Colums
+		for(int32 i = 0, Num = RenderTargetPoolEvents.Num(); i < Num; i++)
+		{
+			FRenderTargetPoolEvent* Event = &RenderTargetPoolEvents[i];
+
+			if(Event->GetEventType() == ERTPE_Alloc)
+			{
+				uint32 PoolEntryId = Event->GetPoolEntryId();
+
+				if(PoolEntryId >= (uint32)Colums.Num())
+				{
+					Colums.SetNum(PoolEntryId + 1);
+				}
+
+				Colums[PoolEntryId] = FRTPColumn(*Event);
+			}
+		}
+
+		Colums.Sort();
+
+		{
+			uint32 ColumnX = 0;
+
+			for(int32 ColumnIndex = 0, Num = Colums.Num(); ColumnIndex < Num; ++ColumnIndex)
+			{
+				const FRTPColumn& RTPColumn = Colums[ColumnIndex];
+
+				uint32 ColumnSize = RTPColumn.SizeInBytes;
+
+				// hide columns that are too small to make a difference (e.g. <1 MB)
+				if(RTPColumn.SizeInBytes <= SizeInKBThreshold * 1024)
+				{
+					ColumnSize = 0;
+				}
+				else
+				{
+					MemoryStats.DisplayedUsageInBytes += RTPColumn.SizeInBytes;
+
+					// give an entry some size to be more UI friendly (if we get mouse UI for zooming in we might not want that any more)
+					ColumnSize = FMath::Max((uint32)(1024 * 1024), ColumnSize);
+				}
+
+				MemoryStats.TotalUsageInBytes += RTPColumn.SizeInBytes;
+				
+				for(int32 EventIndex = 0, Num = RenderTargetPoolEvents.Num(); EventIndex < Num; EventIndex++)
+				{
+					FRenderTargetPoolEvent* Event = &RenderTargetPoolEvents[EventIndex];
+
+					if(Event->GetEventType() != ERTPE_Phase)
+					{
+						uint32 PoolEntryId = Event->GetPoolEntryId();
+
+						if(RTPColumn.PoolEntryId == PoolEntryId)
+						{
+							Event->SetColumn(ColumnIndex, ColumnX, ColumnSize);
+						}
+					}
+				}
+				ColumnX += ColumnSize;
+			}
+		}
+	}
+#endif
+
+	return MemoryStats;
+}
+
+inline void DrawBorder(FCanvas& Canvas, const FIntRect Rect, FLinearColor Color)
+{
+	// top
+	Canvas.DrawTile(Rect.Min.X, Rect.Min.Y, Rect.Max.X - Rect.Min.X, 1, 0, 0, 1, 1, Color);
+	// bottom
+	Canvas.DrawTile(Rect.Min.X, Rect.Max.Y - 1, Rect.Max.X - Rect.Min.X, 1, 0, 0, 1, 1, Color);
+	// left
+	Canvas.DrawTile(Rect.Min.X, Rect.Min.Y + 1, 1, Rect.Max.Y - Rect.Min.Y - 2, 0, 0, 1, 1, Color);
+	// right
+	Canvas.DrawTile(Rect.Max.X - 1, Rect.Min.Y + 1, 1, Rect.Max.Y - Rect.Min.Y - 2, 0, 0, 1, 1, Color);
+}
+
+void FRenderTargetPool::PresentContent(const FSceneView& View)
+{
+	if(RenderTargetPoolEvents.Num())
+	{
+		FIntPoint TotalExtent = ComputeEventDisplayExtent();
+		FIntPoint DisplayLeftTop(20, 50);
+		// on the right we leave more space to make the mouse tooltip readable
+		FIntPoint DisplayExtent(View.ViewRect.Width() - DisplayLeftTop.X * 2 - 140, View.ViewRect.Height() - DisplayLeftTop.Y * 2);
+
+		// if the area is not too small
+		if(DisplayExtent.X > 50 && DisplayExtent.Y > 50)
+		{
+			SMemoryStats MemoryStats = ComputeView();
+
+			RHISetRenderTarget(View.Family->RenderTarget->GetRenderTargetTexture(),FTextureRHIRef());
+			RHISetViewport(0, 0, 0.0f, GSceneRenderTargets.GetBufferSizeXY().X, GSceneRenderTargets.GetBufferSizeXY().Y, 1.0f );
+
+			RHISetBlendState(TStaticBlendState<>::GetRHI());
+			RHISetRasterizerState(TStaticRasterizerState<>::GetRHI());
+			RHISetDepthStencilState(TStaticDepthStencilState<false,CF_Always>::GetRHI());
+
+			// this is a helper class for FCanvas to be able to get screen size
+			class FRenderTargetTemp : public FRenderTarget
+			{
+			public:
+				const FSceneView& View;
+
+				FRenderTargetTemp(const FSceneView& InView) : View(InView)
+				{
+				}
+				virtual FIntPoint GetSizeXY() const
+				{
+					return View.UnscaledViewRect.Size();
+				};
+				virtual const FTexture2DRHIRef& GetRenderTargetTexture() const
+				{
+					return View.Family->RenderTarget->GetRenderTargetTexture();
+				}
+			} TempRenderTarget(View);
+
+			FCanvas Canvas(&TempRenderTarget, NULL, View.Family->CurrentRealTime, View.Family->CurrentWorldTime, View.Family->DeltaWorldTime);
+
+			// TinyFont property
+			const int32 FontHeight = 12;
+
+			FIntPoint MousePos = View.CursorPos;
+			//		FVector2D MousePos = FSlateApplication::Get().GetCursorPos();
+			//		FVector2D MousePos = SceneViewport ? SceneViewport->GetSoftwareCursorPosition() : FVector2D(-1, -1);
+
+			FLinearColor BackgroundColor = FLinearColor(0.0f, 0.0f, 0.0f, 0.7f);
+			FLinearColor PhaseColor = FLinearColor(0.2f, 0.1f, 0.05f, 0.8f);
+			FLinearColor ElementColor = FLinearColor(0.3f, 0.3f, 0.3f, 0.9f);
+			FLinearColor ElementColorVRam = FLinearColor(0.4f, 0.25f, 0.25f, 0.9f);
+
+			UTexture2D* GradientTexture = UCanvas::StaticClass()->GetDefaultObject<UCanvas>()->GradientTexture0;
+
+			// background rectangle
+			Canvas.DrawTile(DisplayLeftTop.X, DisplayLeftTop.Y - 1 * FontHeight - 1, DisplayExtent.X, DisplayExtent.Y + FontHeight, 0, 0, 1, 1, BackgroundColor);
+
+			{
+				uint32 MB = 1024 * 1024;
+				uint32 MBm1 = MB - 1;
+
+				FString Headline = *FString::Printf(TEXT("RenderTargetPool elements(x) over time(y) >= %dKB, Displayed/Total:%d/%dMB"), 
+					SizeInKBThreshold, 
+					(uint32)((MemoryStats.DisplayedUsageInBytes + MBm1) / MB),
+					(uint32)((MemoryStats.TotalUsageInBytes + MBm1) / MB));
+				Canvas.DrawShadowedString(DisplayLeftTop.X, DisplayLeftTop.Y - 1 * FontHeight - 1, *Headline, GEngine->GetTinyFont(), FLinearColor(1, 1, 1));
+			}
+
+			float ScaleX = DisplayExtent.X / (float)TotalExtent.X;
+			float ScaleY = DisplayExtent.Y / (float)TotalExtent.Y;
+
+			// 0 if none
+			FRenderTargetPoolEvent* HighlightedEvent = 0;
+			FIntRect HighlightedRect;
+
+			// Phase events
+			for(FRenderTargetPoolEventIterator It(RenderTargetPoolEvents); *It; ++It)
+			{
+				FRenderTargetPoolEvent* Event = *It;
+
+				if(Event->GetEventType() == ERTPE_Phase)
+				{
+					int32 Y0 = It.Y;
+					int32 Y1 = It.FindClosingEventY();
+
+					FIntPoint PixelLeftTop((int32)(DisplayLeftTop.X), (int32)(DisplayLeftTop.Y + ScaleY * Y0));
+					FIntPoint PixelRightBottom((int32)(DisplayLeftTop.X + DisplayExtent.X), (int32)(DisplayLeftTop.Y + ScaleY * Y1));
+
+					bool bHighlight = MousePos.X >= PixelLeftTop.X && MousePos.X < PixelRightBottom.X && MousePos.Y >= PixelLeftTop.Y && MousePos.Y <= PixelRightBottom.Y;
+
+					if(bHighlight)
+					{
+						HighlightedEvent = Event;
+						HighlightedRect = FIntRect(PixelLeftTop, PixelRightBottom);
+					}
+
+					// UMax is 0.9f to avoid getting some wrap texture leaking in at the bottom
+					Canvas.DrawTile(PixelLeftTop.X, PixelLeftTop.Y, PixelRightBottom.X - PixelLeftTop.X, PixelRightBottom.Y - PixelLeftTop.Y, 0, 0, 1, 0.9f, PhaseColor, GradientTexture->Resource);
+				}
+			}
+
+			// Alloc / Dealloc events
+			for(FRenderTargetPoolEventIterator It(RenderTargetPoolEvents); *It; ++It)
+			{
+				FRenderTargetPoolEvent* Event = *It;
+
+				if(Event->GetEventType() == ERTPE_Alloc && Event->GetColumnSize())
+				{
+					int32 Y0 = It.Y;
+					int32 Y1 = It.FindClosingEventY();
+
+					int32 X0 = Event->GetColumnX();
+					// for now they are all equal width
+					int32 X1 = X0 + Event->GetColumnSize();
+
+					FIntPoint PixelLeftTop((int32)(DisplayLeftTop.X + ScaleX * X0), (int32)(DisplayLeftTop.Y + ScaleY * Y0));
+					FIntPoint PixelRightBottom((int32)(DisplayLeftTop.X + ScaleX * X1), (int32)(DisplayLeftTop.Y + ScaleY * Y1));
+
+					bool bHighlight = MousePos.X >= PixelLeftTop.X && MousePos.X < PixelRightBottom.X && MousePos.Y >= PixelLeftTop.Y && MousePos.Y <= PixelRightBottom.Y;
+
+					if(bHighlight)
+					{
+						HighlightedEvent = Event;
+						HighlightedRect = FIntRect(PixelLeftTop, PixelRightBottom);
+					}
+
+					FLinearColor Color = ElementColor;
+
+					// Highlight EDRAM/FastVRAM usage
+					if(Event->GetDesc().Flags & TexCreate_FastVRAM)
+					{
+						Color = ElementColorVRam;
+					}
+
+					Canvas.DrawTile(
+						PixelLeftTop.X, PixelLeftTop.Y,
+						PixelRightBottom.X - PixelLeftTop.X - 1, PixelRightBottom.Y - PixelLeftTop.Y - 1,
+						0, 0, 1, 1, Color);
+				}
+			}
+
+			if(HighlightedEvent)
+			{
+				DrawBorder(Canvas, HighlightedRect, FLinearColor(0.8f, 0 , 0, 0.5f));
+
+				// Offset to not intersect with crosshair (in editor) or arrow (in game).
+				FIntPoint Pos = MousePos + FIntPoint(12, 4);
+
+				if(HighlightedEvent->GetEventType() == ERTPE_Phase)
+				{
+					Canvas.DrawShadowedString(Pos.X, Pos.Y + 0 * FontHeight, *HighlightedEvent->GetPhaseName(), GEngine->GetTinyFont(), FLinearColor(0.5f, 0.5f, 1));
+				}
+				else
+				{
+					FString SizeString = FString::Printf(TEXT("%d KB"), (HighlightedEvent->GetSizeInBytes() + 1024) / 1024);
+
+					Canvas.DrawShadowedString(Pos.X, Pos.Y + 0 * FontHeight, HighlightedEvent->GetDesc().DebugName, GEngine->GetTinyFont(), FLinearColor(1, 1, 0));
+					Canvas.DrawShadowedString(Pos.X, Pos.Y + 1 * FontHeight, *HighlightedEvent->GetDesc().GenerateInfoString(), GEngine->GetTinyFont(), FLinearColor(1, 1, 0));
+					Canvas.DrawShadowedString(Pos.X, Pos.Y + 2 * FontHeight, *SizeString, GEngine->GetTinyFont(), FLinearColor(1, 1, 0));
+				}
+			}
+
+			Canvas.Flush();
+		}
+	}
+
+	VisualizeTexture.PresentContent(View);
+}
+
+void FRenderTargetPool::AddDeallocEvents()
+{
+	check(IsInRenderingThread());
+
+	bool bWorkWasDone = false;
+
+	for(uint32 i = 0, Num = (uint32)RenderTargetPoolEvents.Num(); i < Num; ++i)
+	{
+		FRenderTargetPoolEvent& Event = RenderTargetPoolEvents[i];
+
+		if(Event.NeedsDeallocEvent())
+		{
+			FRenderTargetPoolEvent NewEvent(Event.GetPoolEntryId(), CurrentTimeStep);
+
+			// for convenience - is actually redundant
+			NewEvent.SetDesc(Event.GetDesc());
+
+			RenderTargetPoolEvents.Add(NewEvent);
+			bWorkWasDone = true;
+		}
+	}
+
+	if(bWorkWasDone)
+	{
+		++CurrentTimeStep;
+	}
+}
+
+void FRenderTargetPool::AddAllocEvent(uint32 InPoolEntryId, FPooledRenderTarget* In)
+{
+	check(In);
+
+	if(IsEventRecordingEnabled())
+	{
+		AddDeallocEvents();
+
+		check(IsInRenderingThread());
+
+		FRenderTargetPoolEvent NewEvent(InPoolEntryId, CurrentTimeStep++, In);
+
+		RenderTargetPoolEvents.Add(NewEvent);
+	}
+}
+
+
+void FRenderTargetPool::AddAllocEventsFromCurrentState()
+{
+	if(!IsEventRecordingEnabled())
+	{
+		return;
+	}
+
+	check(IsInRenderingThread());
+
+	bool bWorkWasDone = false;
+
+	for(uint32 i = 0; i < (uint32)PooledRenderTargets.Num(); ++i)
+	{
+		FPooledRenderTarget* Element = PooledRenderTargets[i];
+
+		if(Element && !Element->IsFree())
+		{
+			FRenderTargetPoolEvent NewEvent(i, CurrentTimeStep, Element);
+
+			RenderTargetPoolEvents.Add(NewEvent);
+			bWorkWasDone = true;
+		}
+	}
+
+	if(bWorkWasDone)
+	{
+		++CurrentTimeStep;
+	}
+}
+
 void FRenderTargetPool::TickPoolElements()
 {
 	check(IsInRenderingThread());
 
+	bEventRecording = false;
+
+	if(bEventRecordingTrigger)
+	{
+		bEventRecordingTrigger = false;
+		bEventRecording = true;
+		CurrentTimeStep = 0;
+		RenderTargetPoolEvents.Empty();
+	}
+
 	uint32 MinimumPoolSizeInKB;
 	{
 		static const auto CVar = IConsoleManager::Get().FindTConsoleVariableDataInt(TEXT("r.RenderTargetPoolMin"));
@@ -297,9 +932,12 @@ void FRenderTargetPool::TickPoolElements()
 	{
 		FPooledRenderTarget* Element = PooledRenderTargets[i];
 
-		Element->OnFrameStart();
+		if(Element)
+		{
+			Element->OnFrameStart();
+		}
 	}
-			
+
 	// we need to release something, take the oldest ones first
 	while(AllocationLevelInKB > MinimumPoolSizeInKB)
 	{
@@ -307,11 +945,11 @@ void FRenderTargetPool::TickPoolElements()
 		int32 OldestElementIndex = -1;
 
 		// find oldest element we can remove
-		for(uint32 i = 0; i < (uint32)PooledRenderTargets.Num(); ++i)
+		for(uint32 i = 0, Num = (uint32)PooledRenderTargets.Num(); i < Num; ++i)
 		{
 			FPooledRenderTarget* Element = PooledRenderTargets[i];
 
-			if(Element->UnusedForNFrames > 2)
+			if(Element && Element->UnusedForNFrames > 2)
 			{
 				if(OldestElementIndex != -1)
 				{
@@ -332,7 +970,8 @@ void FRenderTargetPool::TickPoolElements()
 			AllocationLevelInKB -= ComputeSizeInKB(*PooledRenderTargets[OldestElementIndex]);
 
 			// we assume because of reference counting the resource gets released when not needed any more
-			PooledRenderTargets.RemoveAt(OldestElementIndex);
+			// we don't use Remove() to not shuffle around the elements for better transparency on RenderTargetPoolEvents
+			PooledRenderTargets[OldestElementIndex] = 0;
 
 			VerifyAllocationLevel();
 		}
@@ -369,53 +1008,74 @@ void FRenderTargetPool::TickPoolElements()
 			bCurrentlyOverBudget = false;
 		}
 	}
+
+//	CompactEventArray();
+
+	AddPhaseEvent(TEXT("FromLastFrame"));
+	AddAllocEventsFromCurrentState();
+	AddPhaseEvent(TEXT("Rendering"));
 }
 
-void FRenderTargetPool::FreeUnusedResource(TRefCountPtr<IPooledRenderTarget>& In)
+
+int32 FRenderTargetPool::FindIndex(IPooledRenderTarget* In) const
 {
 	check(IsInRenderingThread());
 
-	if(!IsValidRef(In))
+	if(In)
 	{
-		// no work needed
-		return;
+		for(uint32 i = 0, Num = (uint32)PooledRenderTargets.Num(); i < Num; ++i)
+		{
+			const FPooledRenderTarget* Element = PooledRenderTargets[i];
+
+			if(Element == In)
+			{
+				return i;
+			}
+		}
 	}
 
-	for(uint32 i = 0; i < (uint32)PooledRenderTargets.Num(); ++i)
+	// not found
+	return -1;
+}
+
+void FRenderTargetPool::FreeUnusedResource(TRefCountPtr<IPooledRenderTarget>& In)
+{
+	check(IsInRenderingThread());
+	
+	int32 Index = FindIndex(In);
+
+	if(Index != -1)
 	{
-		FPooledRenderTarget* Element = PooledRenderTargets[i];
+		FPooledRenderTarget* Element = PooledRenderTargets[Index];
 
-		if(Element == In)
+		if(Element)
 		{
 			AllocationLevelInKB -= ComputeSizeInKB(*Element);
 			// we assume because of reference counting the resource gets released when not needed any more
-			PooledRenderTargets.RemoveAt(i);
-			break;
-		}
-	}
+			// we don't use Remove() to not shuffle around the elements for better transparency on RenderTargetPoolEvents
+			PooledRenderTargets[Index] = 0;
 
-	In.SafeRelease();
+			In.SafeRelease();
 
-	VerifyAllocationLevel();
+			VerifyAllocationLevel();
+		}
+	}
 }
 
 void FRenderTargetPool::FreeUnusedResources()
 {
 	check(IsInRenderingThread());
 
-	for(uint32 i = 0; i < (uint32)PooledRenderTargets.Num();)
+	for(uint32 i = 0, Num = (uint32)PooledRenderTargets.Num(); i < Num; ++i)
 	{
 		FPooledRenderTarget* Element = PooledRenderTargets[i];
 
-		if(Element->IsFree())
+		if(Element && Element->IsFree())
 		{
 			AllocationLevelInKB -= ComputeSizeInKB(*Element);
 			// we assume because of reference counting the resource gets released when not needed any more
-			PooledRenderTargets.RemoveAt(i);
-		}
-		else
-		{
-			++i;
+			// we don't use Remove() to not shuffle around the elements for better transparency on RenderTargetPoolEvents
+			PooledRenderTargets[i] = 0;
 		}
 	}
 
@@ -429,17 +1089,21 @@ void FRenderTargetPool::DumpMemoryUsage(FOutputDevice& OutputDevice)
 	for(int32 i = 0; i < PooledRenderTargets.Num(); ++i)
 	{
 		FPooledRenderTarget* Element = PooledRenderTargets[i];
-		OutputDevice.Logf(
-			TEXT("  %6.3fMB %4dx%4d%s%s %2dmip(s) %s (%s)"),
-			ComputeSizeInKB(*Element) / 1024.0f,
-			Element->Desc.Extent.X,
-			Element->Desc.IsCubemap() ? Element->Desc.Extent.X : Element->Desc.Extent.Y,
-			Element->Desc.Depth > 1 ? *FString::Printf(TEXT("x%3d"), Element->Desc.Depth) : (Element->Desc.IsCubemap() ? TEXT("cube") : TEXT("    ")),
-			Element->Desc.bIsArray ? *FString::Printf(TEXT("[%3d]"), Element->Desc.ArraySize) : TEXT("     "),
-			Element->Desc.NumMips,
-			Element->Desc.DebugName,
-			GPixelFormats[Element->Desc.Format].Name
-			);
+
+		if(Element)
+		{
+			OutputDevice.Logf(
+				TEXT("  %6.3fMB %4dx%4d%s%s %2dmip(s) %s (%s)"),
+				ComputeSizeInKB(*Element) / 1024.0f,
+				Element->Desc.Extent.X,
+				Element->Desc.IsCubemap() ? Element->Desc.Extent.X : Element->Desc.Extent.Y,
+				Element->Desc.Depth > 1 ? *FString::Printf(TEXT("x%3d"), Element->Desc.Depth) : (Element->Desc.IsCubemap() ? TEXT("cube") : TEXT("    ")),
+				Element->Desc.bIsArray ? *FString::Printf(TEXT("[%3d]"), Element->Desc.ArraySize) : TEXT("     "),
+				Element->Desc.NumMips,
+				Element->Desc.DebugName,
+				GPixelFormats[Element->Desc.Format].Name
+				);
+		}
 	}
 	uint32 NumTargets=0;
 	uint32 UsedKB=0;
@@ -448,17 +1112,6 @@ void FRenderTargetPool::DumpMemoryUsage(FOutputDevice& OutputDevice)
 	OutputDevice.Logf(TEXT("%.3fMB total, %.3fMB used, %d render targets"), PoolKB / 1024.f, UsedKB / 1024.f, NumTargets);
 }
 
-static void DumpRenderTargetPoolMemory(FOutputDevice& OutputDevice)
-{
-	GRenderTargetPool.DumpMemoryUsage(OutputDevice);
-}
-
-static FAutoConsoleCommandWithOutputDevice GDumpRenderTargetPoolMemoryCmd(
-	TEXT("r.DumpRenderTargetPoolMemory"),
-	TEXT("Dump allocation information for the render target pool."),
-	FConsoleCommandWithOutputDeviceDelegate::CreateStatic(DumpRenderTargetPoolMemory)
-	);
-
 uint32 FPooledRenderTarget::AddRef() const
 {
 	return uint32(++NumRefs);
@@ -526,6 +1179,20 @@ void FRenderTargetPool::VerifyAllocationLevel() const
 */
 }
 
+void FRenderTargetPool::CompactPool()
+{
+	for(uint32 i = 0, Num = (uint32)PooledRenderTargets.Num(); i < Num; ++i)
+	{
+		FPooledRenderTarget* Element = PooledRenderTargets[i];
+
+		if(!Element)
+		{
+			PooledRenderTargets.RemoveAtSwap(i);
+			--Num;
+		}
+	}
+}
+
 bool FPooledRenderTarget::OnFrameStart()
 {
 	check(IsInRenderingThread());
diff --git a/Engine/Source/Runtime/Renderer/Private/PostProcess/RenderTargetPool.h b/Engine/Source/Runtime/Renderer/Private/PostProcess/RenderTargetPool.h
--- a/Engine/Source/Runtime/Renderer/Private/PostProcess/RenderTargetPool.h
+++ b/Engine/Source/Runtime/Renderer/Private/PostProcess/RenderTargetPool.h
@@ -6,6 +6,103 @@
 
 #pragma once
 
+enum ERenderTargetPoolEventType
+{
+	ERTPE_Alloc,
+	ERTPE_Dealloc,
+	ERTPE_Phase
+};
+
+struct FRenderTargetPoolEvent
+{
+	// constructor for ERTPE_Alloc
+	FRenderTargetPoolEvent(uint32 InPoolEntryId, uint32 InTimeStep, IPooledRenderTarget* InPointer)
+		: PoolEntryId(InPoolEntryId)
+		, TimeStep(InTimeStep)
+		, Pointer(InPointer)
+		, EventType(ERTPE_Alloc)
+		, ColumnIndex(-1)
+		, ColumnX(0)
+		, ColumnSize(0)
+	{
+		Desc = Pointer->GetDesc();
+		SizeInBytes = Pointer->ComputeMemorySize();
+	}
+
+	// constructor for ERTPE_Alloc
+	FRenderTargetPoolEvent(uint32 InPoolEntryId, uint32 InTimeStep)
+		: PoolEntryId(InPoolEntryId)
+		, TimeStep(InTimeStep)
+		, Pointer(0)
+		, EventType(ERTPE_Dealloc)
+		, SizeInBytes(0)
+		, ColumnIndex(-1)
+		, ColumnX(0)
+		, ColumnSize(0)
+	{
+	}
+
+	// constructor for ERTPE_Alloc
+	// @param PhaseName must not be 0
+	FRenderTargetPoolEvent(const FString& InPhaseName, uint32 InTimeStep)
+		: PoolEntryId(-1)
+		, TimeStep(InTimeStep)
+		, Pointer(0)
+		, PhaseName(InPhaseName)
+		, EventType(ERTPE_Phase)
+		, SizeInBytes(0)
+		, ColumnIndex(-1)
+		, ColumnX(0)
+		, ColumnSize(0)
+	{
+	}
+
+	// @return Pointer if the object is still in the pool
+	IPooledRenderTarget* GetValidatedPointer() const;
+
+	ERenderTargetPoolEventType GetEventType() const { return EventType; }
+
+	const int32 GetPoolEntryId() const { check(EventType == ERTPE_Alloc || EventType == ERTPE_Dealloc); return PoolEntryId; }
+	const FString& GetPhaseName() const { check(EventType == ERTPE_Phase); return PhaseName; }
+	const FPooledRenderTargetDesc& GetDesc() const { check(EventType == ERTPE_Alloc || EventType == ERTPE_Dealloc); return Desc; }
+	const uint32 GetTimeStep() const { return TimeStep; }
+	const uint64 GetSizeInBytes() const { check(EventType == ERTPE_Alloc); return SizeInBytes; }
+	const void SetPoolEntryId(uint32 InPoolEntryId) { PoolEntryId = InPoolEntryId; }
+	const void SetColumn(uint32 InColumnIndex, uint32 InColumnX, uint32 InColumnSize) { check(EventType == ERTPE_Alloc || EventType == ERTPE_Dealloc); ColumnIndex = InColumnIndex; ColumnX = InColumnX; ColumnSize = InColumnSize; }
+	const uint32 GetColumnX() const { check(EventType == ERTPE_Alloc || EventType == ERTPE_Dealloc); return ColumnX; }
+	const uint32 GetColumnSize() const { check(EventType == ERTPE_Alloc || EventType == ERTPE_Dealloc); return ColumnSize; }
+	bool IsVisible() const { return EventType == ERTPE_Phase || ColumnSize > 0; }
+
+	void SetDesc(const FPooledRenderTargetDesc &InDesc) { Desc = InDesc; }
+
+	bool NeedsDeallocEvent();
+
+private:
+
+	// valid if EventType==ERTPE_Alloc || EventType==ERTPE_Dealloc, -1 if not set, was index into PooledRenderTargets[]
+	uint32 PoolEntryId;
+	//
+	uint32 TimeStep;
+	// valid EventType==ERTPE_Alloc, 0 if not set
+	IPooledRenderTarget* Pointer;
+	// valid if EventType==ERTPE_Phase TEXT("") if not set
+	FString PhaseName;
+	// valid if EventType==ERTPE_Alloc || EventType==ERTPE_Dealloc
+	FPooledRenderTargetDesc Desc;
+	// valid if EventType==ERTPE_Alloc 0 if unknown
+	uint64 SizeInBytes;
+	// e.g. ERTPE_Alloc
+	ERenderTargetPoolEventType EventType;
+
+	// for display, computed by ComputeView()
+
+	// valid if EventType==ERTPE_Alloc || EventType==ERTPE_Dealloc, -1 if not defined yet
+	uint32 ColumnIndex;
+	//
+	uint32 ColumnX;
+	//
+	uint32 ColumnSize;
+};
 
 
 /** The reference to a pooled render target, use like this: TRefCountPtr<IPooledRenderTarget> */
@@ -92,7 +189,10 @@ public:
 	// for debugging purpose, assumes you call FlushRenderingCommands() be
 	// @return 0 if outside of the range
 	FPooledRenderTarget* GetElementById(uint32 Id) const;
-	
+
+	// @return -1 if not found
+	int32 FindIndex(IPooledRenderTarget* In) const;
+
 	// Render visualization
 	void RenderVisualizeTexture(class FDeferredShadingSceneRenderer& Scene);
 
@@ -101,11 +201,22 @@ public:
 	// Logs out usage information.
 	void DumpMemoryUsage(FOutputDevice& OutputDevice);
 
+	void EnableEventRecording(uint32 InSizeInKBThreshold) { SizeInKBThreshold = InSizeInKBThreshold; bEventRecordingTrigger = true; }
+	//
+	void DisableEventDisplay() { RenderTargetPoolEvents.Empty(); }
+	//
+	bool IsEventRecordingEnabled() const;
+
+	void AddPhaseEvent(const TCHAR* InPhaseName);
+
+	/** renders the VisualizeTextureContent to the current render target */
+	void PresentContent(const FSceneView& View);
+
 	FVisualizeTexture VisualizeTexture;
 
 private:
 
-	/** Elements must not be 0 */
+	/** Elements can be 0, we compact the buffer later. */
 	TArray< TRefCountPtr<FPooledRenderTarget> > PooledRenderTargets;
 
 	// redundant, can always be computed with GetStats(), to debug "out of memory" situations and used for r.RenderTargetPoolMin
@@ -114,10 +225,55 @@ private:
 	// to avoid log spam
 	bool bCurrentlyOverBudget;
 
-	friend struct FPooledRenderTarget;
-
 	// for debugging purpose
 	void VerifyAllocationLevel() const;
+
+	// could be done on the fly but that makes the RenderTargetPoolEvents harder to read
+	void CompactPool();
+
+	// the following is used for Event recording --------------------------------
+
+	struct SMemoryStats
+	{
+		SMemoryStats()
+			: DisplayedUsageInBytes(0)
+			, TotalUsageInBytes(0)
+		{
+		}
+		uint64 DisplayedUsageInBytes;
+		uint64 TotalUsageInBytes;
+	};
+
+	// if next frame we want to run with bEventRecording=true
+	bool bEventRecordingTrigger;
+	// e..g 1MB = 1024, 0 to display all
+	uint32 SizeInKBThreshold;
+	// true if enabled
+	bool bEventRecording;
+	// only used if bEventRecording
+	TArray<FRenderTargetPoolEvent> RenderTargetPoolEvents;
+	//
+	uint32 CurrentTimeStep;
+
+	//
+	void AddDeallocEvents();
+
+	// @param In must no be 0
+	void AddAllocEvent(uint32 InPoolEntryId, FPooledRenderTarget* In);
+
+	//
+	void AddAllocEventsFromCurrentState();
+
+	//
+	FIntPoint ComputeEventDisplayExtent();
+
+	// @return 0 if none was found
+	const FString* GetLastEventPhaseName();
+
+	// sorted by size
+	SMemoryStats ComputeView();
+
+	friend struct FPooledRenderTarget;
 };
 
 /** The global render targets for easy shading. */
diff --git a/Engine/Source/Runtime/Renderer/Private/PostProcess/RenderingCompositionGraph.cpp b/Engine/Source/Runtime/Renderer/Private/PostProcess/RenderingCompositionGraph.cpp
--- a/Engine/Source/Runtime/Renderer/Private/PostProcess/RenderingCompositionGraph.cpp
+++ b/Engine/Source/Runtime/Renderer/Private/PostProcess/RenderingCompositionGraph.cpp
@@ -317,6 +317,11 @@ void FRenderingCompositionGraph::RecursivelyProcess(const FRenderingCompositeOut
 				// to track down an issue, should never happen
 				check(OutputRefIt->GetPass());
 
+				if(GRenderTargetPool.IsEventRecordingEnabled() && Pass != Context.Root)
+				{
+					GRenderTargetPool.AddPhaseEvent(*Pass->ConstructDebugName());
+				}
+
 				Context.Pass = Pass;
 				RecursivelyProcess(*OutputRefIt, Context);
 			}
diff --git a/Engine/Source/Runtime/Renderer/Private/SceneRendering.cpp b/Engine/Source/Runtime/Renderer/Private/SceneRendering.cpp
--- a/Engine/Source/Runtime/Renderer/Private/SceneRendering.cpp
+++ b/Engine/Source/Runtime/Renderer/Private/SceneRendering.cpp
@@ -767,7 +767,7 @@ void FSceneRenderer::RenderFinish()
 				continue;
 			}
 
-			GRenderTargetPool.VisualizeTexture.PresentContent(View);
+			GRenderTargetPool.PresentContent(View);
 		}
 	}
 
diff --git a/Engine/Source/Runtime/Renderer/Public/RendererInterface.h b/Engine/Source/Runtime/Renderer/Public/RendererInterface.h
--- a/Engine/Source/Runtime/Renderer/Public/RendererInterface.h
+++ b/Engine/Source/Runtime/Renderer/Public/RendererInterface.h
@@ -339,7 +339,10 @@ struct IPooledRenderTarget : public IRefCountedObject
 	virtual const FPooledRenderTargetDesc& GetDesc() const = 0;
 	/** @param InName must not be 0 */
 	virtual void SetDebugName(const TCHAR *InName) = 0;
-	/** Only for debugging purpose */
+	/**
+	 * Only for debugging purpose
+	 * @return in bytes
+	 **/
 	virtual uint32 ComputeMemorySize() const = 0;
 	/** Get the low level internals (texture/surface) */
 	inline FSceneRenderTargetItem& GetRenderTargetItem() { return RenderTargetItem; }
