#UE4:
* removed not needed GWorld usage

diff --git a/Engine/Source/Runtime/Renderer/Private/Renderer.cpp b/Engine/Source/Runtime/Renderer/Private/Renderer.cpp
--- a/Engine/Source/Runtime/Renderer/Private/Renderer.cpp
+++ b/Engine/Source/Runtime/Renderer/Private/Renderer.cpp
@@ -92,8 +92,8 @@ void FRendererModule::DebugLogOnCrash()
 		{
 			void Thread()
 			{
-				GEngine->Exec(GWorld, TEXT("Mem FromReport"), *GLog);
-				GEngine->Exec(GWorld, TEXT("rhi.DumpMemory"), *GLog);
+				GEngine->Exec(NULL, TEXT("Mem FromReport"), *GLog);
+				GEngine->Exec(NULL, TEXT("rhi.DumpMemory"), *GLog);
 			}
 		} Test;
 
