#UE4:
* fixed VisualizeTexture to show up items when using r.SetRes
* fixed compacting of RenderTarget pool to avoid performance issue happening over time

diff --git a/Engine/Source/Runtime/Renderer/Private/PostProcess/RenderTargetPool.cpp b/Engine/Source/Runtime/Renderer/Private/PostProcess/RenderTargetPool.cpp
--- a/Engine/Source/Runtime/Renderer/Private/PostProcess/RenderTargetPool.cpp
+++ b/Engine/Source/Runtime/Renderer/Private/PostProcess/RenderTargetPool.cpp
@@ -572,8 +572,8 @@ FRenderTargetPool::SMemoryStats FRenderTargetPool::ComputeView()
 			// sort criteria
 			bool operator <(const FRTPColumn& rhs) const
 			{
-				// sort VRam first (only matters on XboxOne but nice to alwasy see it)
-				if(bVRam != rhs.bVRam) return bVRam > rhs.bVRam;
+				// sort VRam first (only matters on XboxOne but nice to always see it)
+// sorting only useful for XboxOne		if(bVRam != rhs.bVRam) return bVRam > rhs.bVRam;
 
 				// we want the large ones first
 				return SizeInBytes > rhs.SizeInBytes;
@@ -928,6 +928,8 @@ void FRenderTargetPool::TickPoolElements()
 		MinimumPoolSizeInKB = FMath::Clamp(CVar->GetValueOnRenderThread(), 0, 2000) * 1024;
 	}
 
+	CompactPool();
+
 	for(uint32 i = 0; i < (uint32)PooledRenderTargets.Num(); ++i)
 	{
 		FPooledRenderTarget* Element = PooledRenderTargets[i];
diff --git a/Engine/Source/Runtime/Renderer/Private/PostProcess/RenderTargetPool.h b/Engine/Source/Runtime/Renderer/Private/PostProcess/RenderTargetPool.h
--- a/Engine/Source/Runtime/Renderer/Private/PostProcess/RenderTargetPool.h
+++ b/Engine/Source/Runtime/Renderer/Private/PostProcess/RenderTargetPool.h
@@ -187,9 +187,11 @@ public:
 	void FreeUnusedResources();
 
 	// for debugging purpose, assumes you call FlushRenderingCommands() be
-	// @return 0 if outside of the range
+	// @return can be 0, that doesn't mean iteration is done
 	FPooledRenderTarget* GetElementById(uint32 Id) const;
 
+	uint32 GetElementCount() const { return PooledRenderTargets.Num(); }
+
 	// @return -1 if not found
 	int32 FindIndex(IPooledRenderTarget* In) const;
 
diff --git a/Engine/Source/Runtime/Renderer/Private/PostProcess/VisualizeTexture.cpp b/Engine/Source/Runtime/Renderer/Private/PostProcess/VisualizeTexture.cpp
--- a/Engine/Source/Runtime/Renderer/Private/PostProcess/VisualizeTexture.cpp
+++ b/Engine/Source/Runtime/Renderer/Private/PostProcess/VisualizeTexture.cpp
@@ -671,13 +671,13 @@ void FVisualizeTexture::DebugLog(bool bExtended)
 	{
 		TArray<FSortedLines> SortedLines;
 
-		for(uint32 i = 0;; ++i)
+		for(uint32 i = 0, Num = GRenderTargetPool.GetElementCount(); i < Num; ++i)
 		{
 			FPooledRenderTarget* RT = GRenderTargetPool.GetElementById(i);
 
 			if(!RT)
 			{
-				break;
+				continue;
 			}
 
 			FPooledRenderTargetDesc Desc = RT->GetDesc();
@@ -861,12 +861,13 @@ void FVisualizeTexture::OnStartFrame(const FSceneView& View)
 
 void FVisualizeTexture::QueryInfo( FQueryVisualizeTexureInfo& Out )
 {
-	for(uint32 i = 0; ; ++i)
+	for(uint32 i = 0, Num = GRenderTargetPool.GetElementCount(); i < Num; ++i)
 	{
 		FPooledRenderTarget* RT = GRenderTargetPool.GetElementById(i);
-		if (!RT)
+
+		if(!RT)
 		{
-			break;
+			continue;
 		}
 
 		FPooledRenderTargetDesc Desc = RT->GetDesc();
