#UE4:
* renamed cvar: r.RenderTargetPool.Events, made it updating each frame

diff --git a/Engine/Source/Runtime/Renderer/Private/PostProcess/RenderTargetPool.cpp b/Engine/Source/Runtime/Renderer/Private/PostProcess/RenderTargetPool.cpp
--- a/Engine/Source/Runtime/Renderer/Private/PostProcess/RenderTargetPool.cpp
+++ b/Engine/Source/Runtime/Renderer/Private/PostProcess/RenderTargetPool.cpp
@@ -23,7 +23,7 @@ static FAutoConsoleCommandWithOutputDevice GDumpRenderTargetPoolMemoryCmd(
 	FConsoleCommandWithOutputDeviceDelegate::CreateStatic(DumpRenderTargetPoolMemory)
 	);
 
-static void DumpRenderTargetPoolEvents(const TArray<FString>& Args)
+static void RenderTargetPoolEvents(const TArray<FString>& Args)
 {
 	uint32 SizeInKBThreshold = -1;
 	if(Args.Num() && Args[0].IsNumeric())
@@ -44,10 +44,11 @@ static void DumpRenderTargetPoolEvents(const TArray<FString>& Args)
 		UE_LOG(LogRenderTargetPool, Display, TEXT("r.DumpRenderTargetPoolEvents is now disabled, use r.DumpRenderTargetPoolEvents <SizeInKB> to enable or r.DumpRenderTargetPoolEvents ? for help"));
 	}
 }
-static FAutoConsoleCommand GDumpRenderTargetPoolEventsCmd(
-	TEXT("r.DumpRenderTargetPoolEvents"),
-	TEXT("Dump for the render target pool events. Optional parameter defines threshold in KB (default is 1024 for 1MB)."),
-	FConsoleCommandWithArgsDelegate::CreateStatic(DumpRenderTargetPoolEvents)
+static FAutoConsoleCommand GRenderTargetPoolEventsCmd(
+	TEXT("r.RenderTargetPool.Events"),
+	TEXT("Visualize the render target pool events over time in one frame. Optional parameter defines threshold in KB.\n")
+	TEXT("To disable the view use the command without any parameter"),
+	FConsoleCommandWithArgsDelegate::CreateStatic(RenderTargetPoolEvents)
 	);
 
 bool FRenderTargetPool::IsEventRecordingEnabled() const
@@ -99,9 +100,9 @@ FRenderTargetPool::FRenderTargetPool()
 	: AllocationLevelInKB(0)
 	, bCurrentlyOverBudget(false)
 	, bEventRecordingTrigger(false)
-	, SizeInKBThreshold(0)
+	, EventRecordingSizeThreshold(0)
 	, bEventRecording(false)
-	, CurrentTimeStep(0)
+	, CurrentEventRecordingTime(0)
 {
 }
 
@@ -384,13 +385,13 @@ void FRenderTargetPool::AddPhaseEvent(const TCHAR *InPhaseName)
 
 		if(!LastName || *LastName != InPhaseName)
 		{
-			if(CurrentTimeStep)
+			if(CurrentEventRecordingTime)
 			{
 				// put a break to former data
-				++CurrentTimeStep;
+				++CurrentEventRecordingTime;
 			}
 
-			FRenderTargetPoolEvent NewEvent(InPhaseName, CurrentTimeStep);
+			FRenderTargetPoolEvent NewEvent(InPhaseName, CurrentEventRecordingTime);
 
 			RenderTargetPoolEvents.Add(NewEvent);
 		}
@@ -510,7 +511,7 @@ private:
 	}
 };
 
-FIntPoint FRenderTargetPool::ComputeEventDisplayExtent()
+uint32 FRenderTargetPool::ComputeEventDisplayHeight()
 {
 	FRenderTargetPoolEventIterator It(RenderTargetPoolEvents);
 
@@ -518,7 +519,7 @@ FIntPoint FRenderTargetPool::ComputeEventDisplayExtent()
 	{
 	}
 
-	return FIntPoint(It.TotalWidth, It.Y + 1);
+	return It.Y + 1;
 }
 
 const FString* FRenderTargetPool::GetLastEventPhaseName()
@@ -564,9 +565,9 @@ FRenderTargetPool::SMemoryStats FRenderTargetPool::ComputeView()
 			// constructor
 			FRTPColumn(const FRenderTargetPoolEvent& Event)
 				: PoolEntryId(Event.GetPoolEntryId())
-				, SizeInBytes(Event.GetSizeInBytes())
 				, bVRam((Event.GetDesc().Flags & TexCreate_FastVRAM) != 0)
 			{
+				 SizeInBytes = Event.GetSizeInBytes();
 			}
 
 			// sort criteria
@@ -612,7 +613,7 @@ FRenderTargetPool::SMemoryStats FRenderTargetPool::ComputeView()
 				uint32 ColumnSize = RTPColumn.SizeInBytes;
 
 				// hide columns that are too small to make a difference (e.g. <1 MB)
-				if(RTPColumn.SizeInBytes <= SizeInKBThreshold * 1024)
+				if(RTPColumn.SizeInBytes <= EventRecordingSizeThreshold * 1024)
 				{
 					ColumnSize = 0;
 				}
@@ -624,6 +625,7 @@ FRenderTargetPool::SMemoryStats FRenderTargetPool::ComputeView()
 					ColumnSize = FMath::Max((uint32)(1024 * 1024), ColumnSize);
 				}
 
+				MemoryStats.TotalColumnSize += ColumnSize;
 				MemoryStats.TotalUsageInBytes += RTPColumn.SizeInBytes;
 				
 				for(int32 EventIndex = 0, Num = RenderTargetPoolEvents.Num(); EventIndex < Num; EventIndex++)
@@ -665,7 +667,6 @@ void FRenderTargetPool::PresentContent(const FSceneView& View)
 {
 	if(RenderTargetPoolEvents.Num())
 	{
-		FIntPoint TotalExtent = ComputeEventDisplayExtent();
 		FIntPoint DisplayLeftTop(20, 50);
 		// on the right we leave more space to make the mouse tooltip readable
 		FIntPoint DisplayExtent(View.ViewRect.Width() - DisplayLeftTop.X * 2 - 140, View.ViewRect.Height() - DisplayLeftTop.Y * 2);
@@ -725,14 +726,16 @@ void FRenderTargetPool::PresentContent(const FSceneView& View)
 				uint32 MBm1 = MB - 1;
 
 				FString Headline = *FString::Printf(TEXT("RenderTargetPool elements(x) over time(y) >= %dKB, Displayed/Total:%d/%dMB"), 
-					SizeInKBThreshold, 
+					EventRecordingSizeThreshold, 
 					(uint32)((MemoryStats.DisplayedUsageInBytes + MBm1) / MB),
 					(uint32)((MemoryStats.TotalUsageInBytes + MBm1) / MB));
 				Canvas.DrawShadowedString(DisplayLeftTop.X, DisplayLeftTop.Y - 1 * FontHeight - 1, *Headline, GEngine->GetTinyFont(), FLinearColor(1, 1, 1));
 			}
 
-			float ScaleX = DisplayExtent.X / (float)TotalExtent.X;
-			float ScaleY = DisplayExtent.Y / (float)TotalExtent.Y;
+			uint32 EventDisplayHeight = ComputeEventDisplayHeight();
+
+			float ScaleX = DisplayExtent.X / (float)MemoryStats.TotalColumnSize;
+			float ScaleY = DisplayExtent.Y / (float)EventDisplayHeight;
 
 			// 0 if none
 			FRenderTargetPoolEvent* HighlightedEvent = 0;
@@ -826,6 +829,9 @@ void FRenderTargetPool::PresentContent(const FSceneView& View)
 			}
 
 			Canvas.Flush();
+
+			CurrentEventRecordingTime = 0;
+			RenderTargetPoolEvents.Empty();
 		}
 	}
 
@@ -844,7 +850,7 @@ void FRenderTargetPool::AddDeallocEvents()
 
 		if(Event.NeedsDeallocEvent())
 		{
-			FRenderTargetPoolEvent NewEvent(Event.GetPoolEntryId(), CurrentTimeStep);
+			FRenderTargetPoolEvent NewEvent(Event.GetPoolEntryId(), CurrentEventRecordingTime);
 
 			// for convenience - is actually redundant
 			NewEvent.SetDesc(Event.GetDesc());
@@ -856,7 +862,7 @@ void FRenderTargetPool::AddDeallocEvents()
 
 	if(bWorkWasDone)
 	{
-		++CurrentTimeStep;
+		++CurrentEventRecordingTime;
 	}
 }
 
@@ -870,7 +876,7 @@ void FRenderTargetPool::AddAllocEvent(uint32 InPoolEntryId, FPooledRenderTarget*
 
 		check(IsInRenderingThread());
 
-		FRenderTargetPoolEvent NewEvent(InPoolEntryId, CurrentTimeStep++, In);
+		FRenderTargetPoolEvent NewEvent(InPoolEntryId, CurrentEventRecordingTime++, In);
 
 		RenderTargetPoolEvents.Add(NewEvent);
 	}
@@ -894,7 +900,7 @@ void FRenderTargetPool::AddAllocEventsFromCurrentState()
 
 		if(Element && !Element->IsFree())
 		{
-			FRenderTargetPoolEvent NewEvent(i, CurrentTimeStep, Element);
+			FRenderTargetPoolEvent NewEvent(i, CurrentEventRecordingTime, Element);
 
 			RenderTargetPoolEvents.Add(NewEvent);
 			bWorkWasDone = true;
@@ -903,7 +909,7 @@ void FRenderTargetPool::AddAllocEventsFromCurrentState()
 
 	if(bWorkWasDone)
 	{
-		++CurrentTimeStep;
+		++CurrentEventRecordingTime;
 	}
 }
 
@@ -911,14 +917,12 @@ void FRenderTargetPool::TickPoolElements()
 {
 	check(IsInRenderingThread());
 
-	bEventRecording = false;
+//	bEventRecording = false;
 
 	if(bEventRecordingTrigger)
 	{
 		bEventRecordingTrigger = false;
 		bEventRecording = true;
-		CurrentTimeStep = 0;
-		RenderTargetPoolEvents.Empty();
 	}
 
 	uint32 MinimumPoolSizeInKB;
diff --git a/Engine/Source/Runtime/Renderer/Private/PostProcess/RenderTargetPool.h b/Engine/Source/Runtime/Renderer/Private/PostProcess/RenderTargetPool.h
--- a/Engine/Source/Runtime/Renderer/Private/PostProcess/RenderTargetPool.h
+++ b/Engine/Source/Runtime/Renderer/Private/PostProcess/RenderTargetPool.h
@@ -203,9 +203,9 @@ public:
 	// Logs out usage information.
 	void DumpMemoryUsage(FOutputDevice& OutputDevice);
 
-	void EnableEventRecording(uint32 InSizeInKBThreshold) { SizeInKBThreshold = InSizeInKBThreshold; bEventRecordingTrigger = true; }
+	void EnableEventRecording(uint32 InSizeInKBThreshold) { EventRecordingSizeThreshold = InSizeInKBThreshold; bEventRecordingTrigger = true; }
 	//
-	void DisableEventDisplay() { RenderTargetPoolEvents.Empty(); }
+	void DisableEventDisplay() { RenderTargetPoolEvents.Empty(); bEventRecording = false; }
 	//
 	bool IsEventRecordingEnabled() const;
 
@@ -240,22 +240,27 @@ private:
 		SMemoryStats()
 			: DisplayedUsageInBytes(0)
 			, TotalUsageInBytes(0)
+			, TotalColumnSize(0)
 		{
 		}
+		// for statistics
 		uint64 DisplayedUsageInBytes;
+		// for statistics
 		uint64 TotalUsageInBytes;
+		// for display purpose, to normalize the view width
+		uint64 TotalColumnSize;
 	};
 
 	// if next frame we want to run with bEventRecording=true
 	bool bEventRecordingTrigger;
-	// e..g 1MB = 1024, 0 to display all
-	uint32 SizeInKBThreshold;
+	// in KB, e.g. 1MB = 1024, 0 to display all
+	uint32 EventRecordingSizeThreshold;
 	// true if enabled
 	bool bEventRecording;
 	// only used if bEventRecording
 	TArray<FRenderTargetPoolEvent> RenderTargetPoolEvents;
 	//
-	uint32 CurrentTimeStep;
+	uint32 CurrentEventRecordingTime;
 
 	//
 	void AddDeallocEvents();
@@ -266,8 +271,7 @@ private:
 	//
 	void AddAllocEventsFromCurrentState();
 
-	//
-	FIntPoint ComputeEventDisplayExtent();
+	uint32 ComputeEventDisplayHeight();
 
 	// @return 0 if none was found
 	const FString* GetLastEventPhaseName();
