#UE4:
* cleanup

diff --git a/Engine/Source/Runtime/Renderer/Private/PostProcess/RenderTargetPool.cpp b/Engine/Source/Runtime/Renderer/Private/PostProcess/RenderTargetPool.cpp
--- a/Engine/Source/Runtime/Renderer/Private/PostProcess/RenderTargetPool.cpp
+++ b/Engine/Source/Runtime/Renderer/Private/PostProcess/RenderTargetPool.cpp
@@ -651,6 +651,7 @@ FRenderTargetPool::SMemoryStats FRenderTargetPool::ComputeView()
 	return MemoryStats;
 }
 
+// draw a single pixel sized rectangle using 4 sub elements
 inline void DrawBorder(FCanvas& Canvas, const FIntRect Rect, FLinearColor Color)
 {
 	// top
@@ -708,8 +709,6 @@ void FRenderTargetPool::PresentContent(const FSceneView& View)
 			const int32 FontHeight = 12;
 
 			FIntPoint MousePos = View.CursorPos;
-			//		FVector2D MousePos = FSlateApplication::Get().GetCursorPos();
-			//		FVector2D MousePos = SceneViewport ? SceneViewport->GetSoftwareCursorPosition() : FVector2D(-1, -1);
 
 			FLinearColor BackgroundColor = FLinearColor(0.0f, 0.0f, 0.0f, 0.7f);
 			FLinearColor PhaseColor = FLinearColor(0.2f, 0.1f, 0.05f, 0.8f);
