* More useful debug info for render targets

diff --git a/Engine/Source/Runtime/Renderer/Public/RendererInterface.h b/Engine/Source/Runtime/Renderer/Public/RendererInterface.h
--- a/Engine/Source/Runtime/Renderer/Public/RendererInterface.h
+++ b/Engine/Source/Runtime/Renderer/Public/RendererInterface.h
@@ -231,6 +231,11 @@ public:
 			FlagsString += TEXT(" UAV");
 		}
 
+		if(LocalFlags & TexCreate_FastVRAM)
+		{
+			FlagsString += TEXT(" VRam");
+		}
+
 		FString ArrayString;
 
 		if(IsArray())
