Enabled render material in canvas for ES2 mode

diff --git a/Engine/Source/Runtime/Renderer/Private/BasePassRendering.h b/Engine/Source/Runtime/Renderer/Private/BasePassRendering.h
--- a/Engine/Source/Runtime/Renderer/Private/BasePassRendering.h
+++ b/Engine/Source/Runtime/Renderer/Private/BasePassRendering.h
@@ -102,7 +102,6 @@ public:
 	{
 		bool bShouldCache = TBasePassVertexShaderBaseType<LightMapPolicyType>::ShouldCache(Platform, Material, VertexFactoryType);
 		return bShouldCache 
-			&& IsFeatureLevelSupported(Platform, ERHIFeatureLevel::SM3)
 			&& (!bEnableAtmosphericFog || IsFeatureLevelSupported(Platform, ERHIFeatureLevel::SM4));
 	}
 
@@ -422,7 +421,6 @@ public:
 		const bool bCacheShaders = !bEnableSkyLight || (Material->GetLightingModel() != MLM_Unlit);
 
 		return bCacheShaders
-			&& IsFeatureLevelSupported(Platform, ERHIFeatureLevel::SM3)
 			&& TBasePassPixelShaderBaseType<LightMapPolicyType>::ShouldCache(Platform, Material, VertexFactoryType);
 	}
 
diff --git a/Engine/Source/Runtime/Renderer/Private/Renderer.cpp b/Engine/Source/Runtime/Renderer/Private/Renderer.cpp
--- a/Engine/Source/Runtime/Renderer/Private/Renderer.cpp
+++ b/Engine/Source/Runtime/Renderer/Private/Renderer.cpp
@@ -44,12 +44,19 @@ void FRendererModule::DrawTileMesh(const FSceneView& SceneView, const FMeshBatch
 	//get the blend mode of the material
 	const EBlendMode MaterialBlendMode = Material->GetBlendMode();
 
-	if (GRHIFeatureLevel >= ERHIFeatureLevel::SM3 && !GUsingNullRHI)
+	if (!GUsingNullRHI)
 	{
 		// handle translucent material blend modes
 		if (IsTranslucentBlendMode(MaterialBlendMode))
 		{
-			FTranslucencyDrawingPolicyFactory::DrawDynamicMesh(View, FTranslucencyDrawingPolicyFactory::ContextType(), Mesh, false, false, NULL, HitProxyId);
+			if (GRHIFeatureLevel >= ERHIFeatureLevel::SM3)
+			{
+				FTranslucencyDrawingPolicyFactory::DrawDynamicMesh(View, FTranslucencyDrawingPolicyFactory::ContextType(), Mesh, false, false, NULL, HitProxyId);
+			}
+			else
+			{
+				FTranslucencyForwardShadingDrawingPolicyFactory::DrawDynamicMesh(View, FTranslucencyForwardShadingDrawingPolicyFactory::ContextType(), Mesh, false, false, NULL, HitProxyId);
+			}
 		}
 		// handle opaque materials
 		else
@@ -64,7 +71,14 @@ void FRendererModule::DrawTileMesh(const FSceneView& SceneView, const FMeshBatch
 			}
 			else
 			{
-				FBasePassOpaqueDrawingPolicyFactory::DrawDynamicMesh(View, FBasePassOpaqueDrawingPolicyFactory::ContextType(false, ESceneRenderTargetsMode::SetTextures), Mesh, false, false, NULL, HitProxyId);
+				if (GRHIFeatureLevel >= ERHIFeatureLevel::SM3)
+				{
+					FBasePassOpaqueDrawingPolicyFactory::DrawDynamicMesh(View, FBasePassOpaqueDrawingPolicyFactory::ContextType(false, ESceneRenderTargetsMode::SetTextures), Mesh, false, false, NULL, HitProxyId);
+				}
+				else
+				{
+					FBasePassForwardOpaqueDrawingPolicyFactory::DrawDynamicMesh(View, FBasePassForwardOpaqueDrawingPolicyFactory::ContextType(ESceneRenderTargetsMode::SetTextures), Mesh, false, false, NULL, HitProxyId);
+				}
 			}
 		}	
 	}
