* fixed potential shader error

diff --git a/Engine/Source/Runtime/Renderer/Private/PostProcess/PostProcessVisualizeBuffer.cpp b/Engine/Source/Runtime/Renderer/Private/PostProcess/PostProcessVisualizeBuffer.cpp
--- a/Engine/Source/Runtime/Renderer/Private/PostProcess/PostProcessVisualizeBuffer.cpp
+++ b/Engine/Source/Runtime/Renderer/Private/PostProcess/PostProcessVisualizeBuffer.cpp
@@ -109,7 +109,7 @@ IMPLEMENT_SHADER_TYPE2(FPostProcessVisualizeBufferPS<true>, SF_Pixel);
 IMPLEMENT_SHADER_TYPE2(FPostProcessVisualizeBufferPS<false>, SF_Pixel);
 
 template <bool bDrawingTile>
-void FRCPassPostProcessVisualizeBuffer::SetShaderTempl(const FRenderingCompositePassContext& Context)
+FShader* FRCPassPostProcessVisualizeBuffer::SetShaderTempl(const FRenderingCompositePassContext& Context)
 {
 	TShaderMapRef<FPostProcessVS> VertexShader(GetGlobalShaderMap());
 	TShaderMapRef<FPostProcessVisualizeBufferPS<bDrawingTile> > PixelShader(GetGlobalShaderMap());
@@ -119,6 +119,8 @@ void FRCPassPostProcessVisualizeBuffer::SetShaderTempl(const FRenderingComposite
 	SetGlobalBoundShaderState(BoundShaderState, GFilterVertexDeclaration.VertexDeclarationRHI, *VertexShader, *PixelShader);
 
 	PixelShader->SetPS(Context);
+
+	return *VertexShader;
 }
 
 void FRCPassPostProcessVisualizeBuffer::Process(FRenderingCompositePassContext& Context)
@@ -150,26 +152,27 @@ void FRCPassPostProcessVisualizeBuffer::Process(FRenderingCompositePassContext&
 	RHISetRasterizerState(TStaticRasterizerState<>::GetRHI());
 	RHISetDepthStencilState(TStaticDepthStencilState<false,CF_Always>::GetRHI());
 
-	SetShaderTempl<false>(Context);
+	{
+		FShader* VertexShader = SetShaderTempl<false>(Context);
+
+		// Draw a quad mapping scene color to the view's render target
+		DrawRectangle(
+			0, 0,
+			DestRect.Width(), DestRect.Height(),
+			SrcRect.Min.X, SrcRect.Min.Y,
+			SrcRect.Width(), SrcRect.Height(),
+			DestRect.Size(),
+			SrcSize,
+			VertexShader,
+			EDRF_UseTriangleOptimization);
+	}
+
+	RHISetBlendState(TStaticBlendState<CW_RGB, BO_Add, BF_SourceAlpha, BF_InverseSourceAlpha>::GetRHI());
 
-	// Draw a quad mapping scene color to the view's render target
 	TShaderMapRef<FPostProcessVS> VertexShader(GetGlobalShaderMap());
-	DrawRectangle(
-		0, 0,
-		DestRect.Width(), DestRect.Height(),
-		SrcRect.Min.X, SrcRect.Min.Y,
-		SrcRect.Width(), SrcRect.Height(),
-		DestRect.Size(),
-		SrcSize,
-		*VertexShader,
-		EDRF_UseTriangleOptimization);
-
-	// Now draw the requested tiles into the grid
 	TShaderMapRef<FPostProcessVisualizeBufferPS<true> > PixelShader(GetGlobalShaderMap());
 
-	RHISetBlendState(TStaticBlendState<CW_RGB, BO_Add, BF_SourceAlpha, BF_InverseSourceAlpha>::GetRHI());
 	static FGlobalBoundShaderState BoundShaderState;
-
 	SetGlobalBoundShaderState(BoundShaderState, GFilterVertexDeclaration.VertexDeclarationRHI, *VertexShader, *PixelShader);
 
 	PixelShader->SetPS(Context);
diff --git a/Engine/Source/Runtime/Renderer/Private/PostProcess/PostProcessVisualizeBuffer.h b/Engine/Source/Runtime/Renderer/Private/PostProcess/PostProcessVisualizeBuffer.h
--- a/Engine/Source/Runtime/Renderer/Private/PostProcess/PostProcessVisualizeBuffer.h
+++ b/Engine/Source/Runtime/Renderer/Private/PostProcess/PostProcessVisualizeBuffer.h
@@ -40,8 +40,9 @@ public:
 
 private:
 
+	// @return VertexShader
 	template <bool bDrawTileTemplate>
-	void SetShaderTempl(const FRenderingCompositePassContext& Context);
+	FShader* SetShaderTempl(const FRenderingCompositePassContext& Context);
 
 	TArray<TileData> Tiles;
 };
