World Position Offset Fix

* MeshModifiesMaterialPosition now uses the result of the material translator to determine whether WPO is used.
* Added MaterialModifiesMeshPosition.  Only intended for use in deciding whether or not shader permutations are required. Uses a conservative estimate - if the WPO pin is hooked up or the material uses material attributes input
* Updated shader DDC versions.
* bUsesEyeAdaptation in now correctly serialized for FMaterialCompilationOutput

diff --git a/Engine/Source/Runtime/Engine/Private/Materials/HLSLMaterialTranslator.h b/Engine/Source/Runtime/Engine/Private/Materials/HLSLMaterialTranslator.h
--- a/Engine/Source/Runtime/Engine/Private/Materials/HLSLMaterialTranslator.h
+++ b/Engine/Source/Runtime/Engine/Private/Materials/HLSLMaterialTranslator.h
@@ -349,6 +349,29 @@ public:
 
 			check(FunctionStack.Num() == 1);
 
+		
+			int32 WPOFreq = (int32)GetMaterialPropertyShaderFrequency(MP_WorldPositionOffset);
+			FShaderCodeChunk& WPOChunk = CodeChunks[MP_WorldPositionOffset][WPOFreq][Chunk[MP_WorldPositionOffset]];
+			FMaterialRenderContext DummyContext( NULL, *Material, NULL );
+			
+			// Determine whether the world position offset is used. 
+			// If the output chunk has a uniform expression, it is constant, and GetNumberValue returns the default property value then WPO isn't used.
+			MaterialCompilationOutput.bModifiesMeshPosition = true;
+			if( WPOChunk.UniformExpression && WPOChunk.UniformExpression->IsConstant() )
+			{
+				float DummyFloat;
+				FColor DummyColor;
+				FVector DefaultVector;
+				GetDefaultForMaterialProperty(MP_WorldPositionOffset, DummyFloat,DummyColor,DefaultVector);
+
+				FLinearColor WPOValue;
+				WPOChunk.UniformExpression->GetNumberValue(DummyContext, WPOValue);
+				if( FVector(WPOValue) == DefaultVector)
+				{
+					MaterialCompilationOutput.bModifiesMeshPosition = false;
+				}
+			}
+
 			if (Material->GetBlendMode() == BLEND_Modulate && Material->GetLightingModel() != MLM_Unlit && !Material->IsUsedWithDeferredDecal())
 			{
 				Errorf(TEXT("Dynamically lit translucency is not supported for BLEND_Modulate materials."));
@@ -1834,7 +1857,7 @@ protected:
 			// be exactly the same as the offset one, so there is no point bringing in the extra code.
 			// Also, we can't access the full offset world position in anything other than the pixel shader, because it won't have
 			// been calculated yet
-			bool bNonOffsetWorldPositionAvailable = Material->MaterialModifiesMeshPosition() && ShaderFrequency == SF_Pixel;
+			bool bNonOffsetWorldPositionAvailable = Material->MaterialMayModifyMeshPosition() && ShaderFrequency == SF_Pixel;
 
 			switch (WorldPositionIncludedOffsets)
 			{
diff --git a/Engine/Source/Runtime/Engine/Private/Materials/MaterialShared.cpp b/Engine/Source/Runtime/Engine/Private/Materials/MaterialShared.cpp
--- a/Engine/Source/Runtime/Engine/Private/Materials/MaterialShared.cpp
+++ b/Engine/Source/Runtime/Engine/Private/Materials/MaterialShared.cpp
@@ -318,6 +318,10 @@ void FMaterialCompilationOutput::Serialize(FArchive& Ar)
 	}
 
 	Ar << bNeedsSceneTextures;
+
+	Ar << bUsesEyeAdaptation;
+
+	Ar << bModifiesMeshPosition;
 }
 
 void FMaterial::GetShaderMapId(EShaderPlatform Platform, FMaterialShaderMapId& OutId) const
@@ -451,7 +455,17 @@ bool FMaterial::UsesEyeAdaptation() const
 
 bool FMaterial::MaterialModifiesMeshPosition() const 
 { 
-	return HasVertexPositionOffsetConnected() || GetTessellationMode() != MTM_NoTessellation;
+	FMaterialShaderMap* ShaderMap = IsInRenderingThread() ? RenderingThreadShaderMap : GameThreadShaderMap.GetReference();
+	bool bUsesWPO = ShaderMap ? ShaderMap->ModifiesMeshPosition() : false;
+
+	return bUsesWPO || GetTessellationMode() != MTM_NoTessellation;
+}
+
+bool FMaterial::MaterialMayModifyMeshPosition() const
+{
+	// Conservative estimate when called before material translation has occurred. 
+    // This function is only intended for use in deciding whether or not shader permutations are required.
+	return HasVertexPositionOffsetConnected() || HasMaterialAttributesConnected() || GetTessellationMode() != MTM_NoTessellation;
 }
 
 FMaterialShaderMap* FMaterial::GetRenderingThreadShaderMap() const 
@@ -657,7 +671,8 @@ bool FMaterialResource::IsLightFunction() const { return Material->MaterialDomai
 bool FMaterialResource::IsUsedWithEditorCompositing() const { return Material->bUsedWithEditorCompositing; }
 bool FMaterialResource::IsUsedWithDeferredDecal() const { return Material->MaterialDomain == MD_DeferredDecal; }
 bool FMaterialResource::IsSpecialEngineMaterial() const { return Material->bUsedAsSpecialEngineMaterial; }
-bool FMaterialResource::HasVertexPositionOffsetConnected() const { return Material->WorldPositionOffset.Expression != NULL; }
+bool FMaterialResource::HasVertexPositionOffsetConnected() const { return !Material->bUseMaterialAttributes && Material->WorldPositionOffset.Expression != NULL; }
+bool FMaterialResource::HasMaterialAttributesConnected() const { return Material->bUseMaterialAttributes && Material->MaterialAttributes.Expression != NULL; }
 FString FMaterialResource::GetBaseMaterialPathName() const { return Material->GetPathName(); }
 
 bool FMaterialResource::IsUsedWithSkeletalMesh() const
diff --git a/Engine/Source/Runtime/Engine/Private/ShaderDerivedDataVersion.h b/Engine/Source/Runtime/Engine/Private/ShaderDerivedDataVersion.h
--- a/Engine/Source/Runtime/Engine/Private/ShaderDerivedDataVersion.h
+++ b/Engine/Source/Runtime/Engine/Private/ShaderDerivedDataVersion.h
@@ -9,5 +9,5 @@
 // In case of merge conflicts with DDC versions, you MUST generate a new GUID and set this new
 // guid as version
 
-#define GLOBALSHADERMAP_DERIVEDDATA_VER			TEXT("7c7fc3753f5f40b092cb861a0d798d11")
-#define MATERIALSHADERMAP_DERIVEDDATA_VER		TEXT("9b5ebb87518144f08d55ce1c2d2096b6")
+#define GLOBALSHADERMAP_DERIVEDDATA_VER			TEXT("af1606fcec114373844d2b6cf9df71c8")
+#define MATERIALSHADERMAP_DERIVEDDATA_VER		TEXT("e1c1e89ee9ae4f3ca7b8e0f1c49b34eb")
diff --git a/Engine/Source/Runtime/Engine/Public/MaterialShared.h b/Engine/Source/Runtime/Engine/Public/MaterialShared.h
--- a/Engine/Source/Runtime/Engine/Public/MaterialShared.h
+++ b/Engine/Source/Runtime/Engine/Public/MaterialShared.h
@@ -257,7 +257,8 @@ public:
 	FMaterialCompilationOutput() :
 		bUsesSceneColor(false),
 		bNeedsSceneTextures(false),
-		bUsesEyeAdaptation(false)
+		bUsesEyeAdaptation(false),
+		bModifiesMeshPosition(false)
 	{}
 
 	ENGINE_API void Serialize(FArchive& Ar);
@@ -274,6 +275,9 @@ public:
 
 	/** true if the material uses the EyeAdaptationLookup */
 	bool bUsesEyeAdaptation;
+
+	/** true if the material modifies the the mesh position. */
+	bool bModifiesMeshPosition;
 };
 
 /**
@@ -796,6 +800,7 @@ public:
 	bool UsesSceneColor() const { return MaterialCompilationOutput.bUsesSceneColor; }
 	bool NeedsSceneTextures() const { return MaterialCompilationOutput.bNeedsSceneTextures; }
 	bool UsesEyeAdaptation() const { return MaterialCompilationOutput.bUsesEyeAdaptation; }
+	bool ModifiesMeshPosition() const { return MaterialCompilationOutput.bModifiesMeshPosition; }
 
 	bool IsValidForRendering() const
 	{
@@ -1111,6 +1116,7 @@ public:
 	virtual bool UseTranslucencyVertexFog() const { return false; }
 	virtual FString GetFriendlyName() const = 0;
 	virtual bool HasVertexPositionOffsetConnected() const { return false; }
+	virtual bool HasMaterialAttributesConnected() const { return false; }
 	virtual uint32 GetDecalBlendMode() const { return 0; }
 	virtual uint32 GetMaterialDecalResponse() const { return 0; }
 	virtual bool HasNormalConnected() const { return false; }
@@ -1154,8 +1160,13 @@ public:
 	ENGINE_API bool UsesSceneColor() const;
 	ENGINE_API bool NeedsSceneTextures() const;
 	ENGINE_API bool UsesEyeAdaptation() const;	
+
+	/** Does the material modify the mesh position. */
 	ENGINE_API bool MaterialModifiesMeshPosition() const;
 
+	/** Note: This function is only intended for use in deciding whether or not shader permutations are required before material translation occurs. */
+	ENGINE_API bool MaterialMayModifyMeshPosition() const;
+
 	class FMaterialShaderMap* GetGameThreadShaderMap() const 
 	{ 
 		checkSlow(IsInGameThread());
@@ -1786,6 +1797,7 @@ protected:
 	/** Entry point for compiling a specific material property.  This must call SetMaterialProperty. */
 	ENGINE_API virtual int32 CompileProperty(EMaterialProperty Property,EShaderFrequency InShaderFrequency,class FMaterialCompiler* Compiler) const;
 	ENGINE_API virtual bool HasVertexPositionOffsetConnected() const;
+	ENGINE_API virtual bool HasMaterialAttributesConnected() const;
 	/** Useful for debugging. */
 	ENGINE_API virtual FString GetBaseMaterialPathName() const;
 };
diff --git a/Engine/Source/Runtime/Renderer/Private/DepthRendering.cpp b/Engine/Source/Runtime/Renderer/Private/DepthRendering.cpp
--- a/Engine/Source/Runtime/Renderer/Private/DepthRendering.cpp
+++ b/Engine/Source/Runtime/Renderer/Private/DepthRendering.cpp
@@ -33,7 +33,7 @@ public:
 		}
 
 		// Only compile for the default material and masked materials
-		return (Material->IsSpecialEngineMaterial() || Material->IsMasked() || Material->MaterialModifiesMeshPosition())
+		return (Material->IsSpecialEngineMaterial() || Material->IsMasked() || Material->MaterialMayModifyMeshPosition())
 			 && IsFeatureLevelSupported(Platform, ERHIFeatureLevel::SM3);
 	}
 
diff --git a/Engine/Source/Runtime/Renderer/Private/LightMapDensityRendering.h b/Engine/Source/Runtime/Renderer/Private/LightMapDensityRendering.h
--- a/Engine/Source/Runtime/Renderer/Private/LightMapDensityRendering.h
+++ b/Engine/Source/Runtime/Renderer/Private/LightMapDensityRendering.h
@@ -18,7 +18,7 @@ public:
 
 	static bool ShouldCache(EShaderPlatform Platform,const FMaterial* Material,const FVertexFactoryType* VertexFactoryType)
 	{
-		return (Material->IsSpecialEngineMaterial() || Material->IsMasked() || Material->MaterialModifiesMeshPosition())
+		return (Material->IsSpecialEngineMaterial() || Material->IsMasked() || Material->MaterialMayModifyMeshPosition())
 				&& LightMapPolicyType::ShouldCache(Platform,Material,VertexFactoryType)
 				&& IsFeatureLevelSupported(Platform, ERHIFeatureLevel::SM3);
 	}
@@ -131,7 +131,7 @@ public:
 
 	static bool ShouldCache(EShaderPlatform Platform,const FMaterial* Material,const FVertexFactoryType* VertexFactoryType)
 	{
-		return (Material->IsSpecialEngineMaterial() || Material->IsMasked() || Material->MaterialModifiesMeshPosition()) 
+		return (Material->IsSpecialEngineMaterial() || Material->IsMasked() || Material->MaterialMayModifyMeshPosition())
 				&& LightMapPolicyType::ShouldCache(Platform,Material,VertexFactoryType)
 				&& IsFeatureLevelSupported(Platform, ERHIFeatureLevel::SM3);
 	}
diff --git a/Engine/Source/Runtime/Renderer/Private/SceneHitProxyRendering.cpp b/Engine/Source/Runtime/Renderer/Private/SceneHitProxyRendering.cpp
--- a/Engine/Source/Runtime/Renderer/Private/SceneHitProxyRendering.cpp
+++ b/Engine/Source/Runtime/Renderer/Private/SceneHitProxyRendering.cpp
@@ -39,7 +39,7 @@ public:
 		return IsPCPlatform(Platform) && 
 			IsFeatureLevelSupported(Platform, ERHIFeatureLevel::SM3) &&
 			// and only compile for the default material or materials that are masked.
-			(Material->IsSpecialEngineMaterial() || Material->IsMasked() || Material->MaterialModifiesMeshPosition() || Material->IsTwoSided());
+			(Material->IsSpecialEngineMaterial() || Material->IsMasked() || Material->MaterialMayModifyMeshPosition() || Material->IsTwoSided());
 	}
 
 protected:
@@ -112,7 +112,7 @@ public:
 		return IsPCPlatform(Platform) 
 			&& IsFeatureLevelSupported(Platform, ERHIFeatureLevel::SM3) 
 			// and only compile for default materials or materials that are masked.
-			&& (Material->IsSpecialEngineMaterial() || Material->IsMasked() || Material->MaterialModifiesMeshPosition() || Material->IsTwoSided());
+			&& (Material->IsSpecialEngineMaterial() || Material->IsMasked() || Material->MaterialMayModifyMeshPosition() || Material->IsTwoSided());
 	}
 
 	FHitProxyPS(const ShaderMetaType::CompiledShaderInitializerType& Initializer):
@@ -257,7 +257,7 @@ void FHitProxyDrawingPolicyFactory::AddStaticMesh(FScene* Scene,FStaticMesh* Sta
 	// Add the static mesh to the DPG's hit proxy draw list.
 	const FMaterialRenderProxy* MaterialRenderProxy = StaticMesh->MaterialRenderProxy;
 	const FMaterial* Material = MaterialRenderProxy->GetMaterial(GRHIFeatureLevel);
-	if ( !Material->IsMasked() && !Material->IsTwoSided() && !Material->MaterialModifiesMeshPosition() )
+	if (!Material->IsMasked() && !Material->IsTwoSided() && !Material->MaterialModifiesMeshPosition())
 	{
 		// Default material doesn't handle masked, and doesn't have the correct bIsTwoSided setting.
 		MaterialRenderProxy = UMaterial::GetDefaultMaterial(MD_Surface)->GetRenderProxy(false);
@@ -305,7 +305,7 @@ bool FHitProxyDrawingPolicyFactory::DrawDynamicMesh(
 		if( View.bAllowTranslucentPrimitivesInHitProxy || !IsTranslucentBlendMode( Material->GetBlendMode() ) || ( HitProxy && HitProxy->AlwaysAllowsTranslucentPrimitives() ) )
 #endif
 		{
-			if ( !Material->IsMasked() && !Material->IsTwoSided() && !Material->MaterialModifiesMeshPosition() )
+			if (!Material->IsMasked() && !Material->IsTwoSided() && !Material->MaterialModifiesMeshPosition())
 			{
 				// Default material doesn't handle masked, and doesn't have the correct bIsTwoSided setting.
 				MaterialRenderProxy = UMaterial::GetDefaultMaterial(MD_Surface)->GetRenderProxy(false);
diff --git a/Engine/Source/Runtime/Renderer/Private/ShadowRendering.cpp b/Engine/Source/Runtime/Renderer/Private/ShadowRendering.cpp
--- a/Engine/Source/Runtime/Renderer/Private/ShadowRendering.cpp
+++ b/Engine/Source/Runtime/Renderer/Private/ShadowRendering.cpp
@@ -155,7 +155,7 @@ public:
 		{
 			return (Material->IsSpecialEngineMaterial()
 						// Masked and WPO materials need their shaders but cannot be used with a position only stream.
-						|| ((Material->IsMasked() || Material->MaterialModifiesMeshPosition()) && !bUsePositionOnlyStream))
+						|| ((Material->IsMasked() || Material->MaterialMayModifyMeshPosition()) && !bUsePositionOnlyStream))
 					// Only compile one pass point light shaders for feature levels >= SM4
 					&& (ShaderMode != VertexShadowDepth_OnePassPointLight || IsFeatureLevelSupported(Platform, ERHIFeatureLevel::SM4))
 					// Only compile position-only shaders for vertex factories that support it.
@@ -503,7 +503,7 @@ public:
 					// Only compile for masked or lit translucent materials
 					|| Material->IsMasked()
 					// Perspective correct rendering needs a pixel shader and WPO materials can't be overridden with default material.
-					|| (ShaderMode == PixelShadowDepth_PerspectiveCorrect && Material->MaterialModifiesMeshPosition()))
+					|| (ShaderMode == PixelShadowDepth_PerspectiveCorrect && Material->MaterialMayModifyMeshPosition()))
 				// Only compile one pass point light shaders for feature levels >= SM4
 				&& (ShaderMode != PixelShadowDepth_OnePassPointLight || IsFeatureLevelSupported(Platform, ERHIFeatureLevel::SM4))
 				// Don't render ShadowDepth for translucent unlit materials
diff --git a/Engine/Source/Runtime/Renderer/Private/VelocityRendering.cpp b/Engine/Source/Runtime/Renderer/Private/VelocityRendering.cpp
--- a/Engine/Source/Runtime/Renderer/Private/VelocityRendering.cpp
+++ b/Engine/Source/Runtime/Renderer/Private/VelocityRendering.cpp
@@ -39,7 +39,7 @@ public:
 			//or if the material is opaque and two-sided,
 			|| (Material->IsTwoSided() && !IsTranslucentBlendMode(Material->GetBlendMode()))
 			// or if the material modifies meshes
-			|| Material->MaterialModifiesMeshPosition()))
+			|| Material->MaterialMayModifyMeshPosition()))
 			&& IsFeatureLevelSupported(Platform, ERHIFeatureLevel::SM4);
 	}
 
@@ -127,7 +127,7 @@ public:
 			//or if the material is opaque and two-sided,
 			|| (Material->IsTwoSided() && !IsTranslucentBlendMode(Material->GetBlendMode()))
 			// or if the material modifies meshes
-			|| Material->MaterialModifiesMeshPosition()))
+			|| Material->MaterialMayModifyMeshPosition()))
 			&& IsFeatureLevelSupported(Platform, ERHIFeatureLevel::SM4);
 	}
 
@@ -352,7 +352,7 @@ void FVelocityDrawingPolicyFactory::AddStaticMesh(FScene* Scene, FStaticMesh* St
 	    EBlendMode BlendMode = Material->GetBlendMode();
 	    if(BlendMode == BLEND_Opaque || BlendMode == BLEND_Masked)
 	    {
-		    if(!Material->IsMasked() && !Material->IsTwoSided() && !Material->MaterialModifiesMeshPosition())
+			if (!Material->IsMasked() && !Material->IsTwoSided() && !Material->MaterialModifiesMeshPosition())
 		    {
 			    // Default material doesn't handle masked or mesh-mod, and doesn't have the correct bIsTwoSided setting.
 			    MaterialRenderProxy = UMaterial::GetDefaultMaterial(MD_Surface)->GetRenderProxy(false);
@@ -388,7 +388,7 @@ bool FVelocityDrawingPolicyFactory::DrawDynamicMesh(
 		// This should be enforced at a higher level
 		//@todo - figure out why this is failing and re-enable
 		//check(FVelocityDrawingPolicy::HasVelocity(View, PrimitiveSceneInfo));
-		if(!Material->IsMasked() && !Material->IsTwoSided() && !Material->MaterialModifiesMeshPosition())
+		if (!Material->IsMasked() && !Material->IsTwoSided() && !Material->MaterialModifiesMeshPosition())
 		{
 			// Default material doesn't handle masked, and doesn't have the correct bIsTwoSided setting.
 			MaterialRenderProxy = UMaterial::GetDefaultMaterial(MD_Surface)->GetRenderProxy(false);
