Fix for NaN issue due to Z being 0.

diff --git a/Engine/Source/Runtime/Renderer/Private/IndirectLightingCache.cpp b/Engine/Source/Runtime/Renderer/Private/IndirectLightingCache.cpp
--- a/Engine/Source/Runtime/Renderer/Private/IndirectLightingCache.cpp
+++ b/Engine/Source/Runtime/Renderer/Private/IndirectLightingCache.cpp
@@ -261,6 +261,20 @@ void FIndirectLightingCache::CalculateBlockScaleAndAdd(FIntVector InTexelMin, in
 	{
 		const float UVSize = AllocationTexelSize / (float)CacheSize;
 
+		// need to remove 0
+		if (InSize.X == 0.f)
+		{
+			InSize.X = 0.01f;
+		}
+		if(InSize.Y == 0.f)
+		{
+			InSize.Y = 0.01f;
+		}
+		if(InSize.Z == 0.f)
+		{
+			InSize.Z = 0.01f;
+		}
+
 		// Setup a scale and add to convert from world space position to volume texture UV
 		OutScale = FVector(UVSize) / InSize;
 		OutAdd = -InMin * UVSize / InSize + MinUV;
