* refactored code to be in onle place only
* optimized setting shader param

diff --git a/Engine/Source/Runtime/Core/Public/Math/Matrix.h b/Engine/Source/Runtime/Core/Public/Math/Matrix.h
--- a/Engine/Source/Runtime/Core/Public/Math/Matrix.h
+++ b/Engine/Source/Runtime/Core/Public/Math/Matrix.h
@@ -445,3 +445,40 @@ FORCEINLINE FLookAtMatrix::FLookAtMatrix(const FVector& EyePosition, const FVect
 
 
 template <> struct TIsPODType<FMatrix> { enum { Value = true }; };
+
+// very high quality 4x4 matrix inverse
+static inline void Inverse4x4( double* dst, const float* src )
+{
+	const double s0  = (double)(src[ 0]); const double s1  = (double)(src[ 1]); const double s2  = (double)(src[ 2]); const double s3  = (double)(src[ 3]);
+	const double s4  = (double)(src[ 4]); const double s5  = (double)(src[ 5]); const double s6  = (double)(src[ 6]); const double s7  = (double)(src[ 7]);
+	const double s8  = (double)(src[ 8]); const double s9  = (double)(src[ 9]); const double s10 = (double)(src[10]); const double s11 = (double)(src[11]);
+	const double s12 = (double)(src[12]); const double s13 = (double)(src[13]); const double s14 = (double)(src[14]); const double s15 = (double)(src[15]);
+
+	double inv[16];
+	inv[0]  =  s5 * s10 * s15 - s5 * s11 * s14 - s9 * s6 * s15 + s9 * s7 * s14 + s13 * s6 * s11 - s13 * s7 * s10;
+	inv[1]  = -s1 * s10 * s15 + s1 * s11 * s14 + s9 * s2 * s15 - s9 * s3 * s14 - s13 * s2 * s11 + s13 * s3 * s10;
+	inv[2]  =  s1 * s6  * s15 - s1 * s7  * s14 - s5 * s2 * s15 + s5 * s3 * s14 + s13 * s2 * s7  - s13 * s3 * s6;
+	inv[3]  = -s1 * s6  * s11 + s1 * s7  * s10 + s5 * s2 * s11 - s5 * s3 * s10 - s9  * s2 * s7  + s9  * s3 * s6;
+	inv[4]  = -s4 * s10 * s15 + s4 * s11 * s14 + s8 * s6 * s15 - s8 * s7 * s14 - s12 * s6 * s11 + s12 * s7 * s10;
+	inv[5]  =  s0 * s10 * s15 - s0 * s11 * s14 - s8 * s2 * s15 + s8 * s3 * s14 + s12 * s2 * s11 - s12 * s3 * s10;
+	inv[6]  = -s0 * s6  * s15 + s0 * s7  * s14 + s4 * s2 * s15 - s4 * s3 * s14 - s12 * s2 * s7  + s12 * s3 * s6;
+	inv[7]  =  s0 * s6  * s11 - s0 * s7  * s10 - s4 * s2 * s11 + s4 * s3 * s10 + s8  * s2 * s7  - s8  * s3 * s6;
+	inv[8]  =  s4 * s9  * s15 - s4 * s11 * s13 - s8 * s5 * s15 + s8 * s7 * s13 + s12 * s5 * s11 - s12 * s7 * s9;
+	inv[9]  = -s0 * s9  * s15 + s0 * s11 * s13 + s8 * s1 * s15 - s8 * s3 * s13 - s12 * s1 * s11 + s12 * s3 * s9;
+	inv[10] =  s0 * s5  * s15 - s0 * s7  * s13 - s4 * s1 * s15 + s4 * s3 * s13 + s12 * s1 * s7  - s12 * s3 * s5;
+	inv[11] = -s0 * s5  * s11 + s0 * s7  * s9  + s4 * s1 * s11 - s4 * s3 * s9  - s8  * s1 * s7  + s8  * s3 * s5;
+	inv[12] = -s4 * s9  * s14 + s4 * s10 * s13 + s8 * s5 * s14 - s8 * s6 * s13 - s12 * s5 * s10 + s12 * s6 * s9;
+	inv[13] =  s0 * s9  * s14 - s0 * s10 * s13 - s8 * s1 * s14 + s8 * s2 * s13 + s12 * s1 * s10 - s12 * s2 * s9;
+	inv[14] = -s0 * s5  * s14 + s0 * s6  * s13 + s4 * s1 * s14 - s4 * s2 * s13 - s12 * s1 * s6  + s12 * s2 * s5;
+	inv[15] =  s0 * s5  * s10 - s0 * s6  * s9  - s4 * s1 * s10 + s4 * s2 * s9  + s8  * s1 * s6  - s8  * s2 * s5;
+
+	double det = s0 * inv[0] + s1 * inv[4] + s2 * inv[8] + s3 * inv[12];
+	if( det != 0.0 )
+	{
+		det = 1.0 / det;
+	}
+	for( int i = 0; i < 16; i++ )
+	{
+		dst[i] = inv[i] * det;
+	}
+}
\ No newline at end of file
diff --git a/Engine/Source/Runtime/Renderer/Private/CompositionLighting/PostProcessAmbientOcclusion.cpp b/Engine/Source/Runtime/Renderer/Private/CompositionLighting/PostProcessAmbientOcclusion.cpp
--- a/Engine/Source/Runtime/Renderer/Private/CompositionLighting/PostProcessAmbientOcclusion.cpp
+++ b/Engine/Source/Runtime/Renderer/Private/CompositionLighting/PostProcessAmbientOcclusion.cpp
@@ -101,6 +101,85 @@ FArchive& operator<<(FArchive& Ar, FScreenSpaceAOandSSRShaderParameters& This)
 	return Ar;
 }
 
+// --------------------------------------------------------
+
+void FCameraMotionParameters::Bind(const FShaderParameterMap& ParameterMap)
+{
+	CameraMotion.Bind(ParameterMap, TEXT("CameraMotion"));
+}
+
+void FCameraMotionParameters::Set(const FSceneView& View, const FPixelShaderRHIParamRef ShaderRHI) const
+{
+	FSceneViewState* ViewState = (FSceneViewState*)View.State;
+
+	FMatrix Proj = View.ViewMatrices.ProjMatrix;
+	FMatrix PrevProj = ViewState->PrevViewMatrices.ProjMatrix;
+
+	// Remove jitter
+	Proj.M[2][0] = 0.0f;
+	Proj.M[2][1] = 0.0f;
+	PrevProj.M[2][0] = 0.0f;
+	PrevProj.M[2][1] = 0.0f;
+
+	FMatrix ViewProj = ( View.ViewMatrices.ViewMatrix * Proj ).GetTransposed();
+	FMatrix PrevViewProj = ( ViewState->PrevViewMatrices.ViewMatrix * PrevProj ).GetTransposed();
+
+	double InvViewProj[16];
+	Inverse4x4( InvViewProj, (float*)ViewProj.M );
+
+	const float* p = (float*)PrevViewProj.M;
+
+	const double cxx = InvViewProj[ 0]; const double cxy = InvViewProj[ 1]; const double cxz = InvViewProj[ 2]; const double cxw = InvViewProj[ 3];
+	const double cyx = InvViewProj[ 4]; const double cyy = InvViewProj[ 5]; const double cyz = InvViewProj[ 6]; const double cyw = InvViewProj[ 7];
+	const double czx = InvViewProj[ 8]; const double czy = InvViewProj[ 9]; const double czz = InvViewProj[10]; const double czw = InvViewProj[11];
+	const double cwx = InvViewProj[12]; const double cwy = InvViewProj[13]; const double cwz = InvViewProj[14]; const double cww = InvViewProj[15];
+
+	const double pxx = (double)(p[ 0]); const double pxy = (double)(p[ 1]); const double pxz = (double)(p[ 2]); const double pxw = (double)(p[ 3]);
+	const double pyx = (double)(p[ 4]); const double pyy = (double)(p[ 5]); const double pyz = (double)(p[ 6]); const double pyw = (double)(p[ 7]);
+	const double pwx = (double)(p[12]); const double pwy = (double)(p[13]); const double pwz = (double)(p[14]); const double pww = (double)(p[15]);
+
+	FVector4 CameraMotionValue[5];
+
+	CameraMotionValue[0] = FVector4(
+		(float)(4.0*(cwx*pww + cxx*pwx + cyx*pwy + czx*pwz)),
+		(float)((-4.0)*(cwy*pww + cxy*pwx + cyy*pwy + czy*pwz)),
+		(float)(2.0*(cwz*pww + cxz*pwx + cyz*pwy + czz*pwz)),
+		(float)(2.0*(cww*pww - cwx*pww + cwy*pww + (cxw - cxx + cxy)*pwx + (cyw - cyx + cyy)*pwy + (czw - czx + czy)*pwz)));
+
+	CameraMotionValue[1] = FVector4(
+		(float)(( 4.0)*(cwy*pww + cxy*pwx + cyy*pwy + czy*pwz)),
+		(float)((-2.0)*(cwz*pww + cxz*pwx + cyz*pwy + czz*pwz)),
+		(float)((-2.0)*(cww*pww + cwy*pww + cxw*pwx - 2.0*cxx*pwx + cxy*pwx + cyw*pwy - 2.0*cyx*pwy + cyy*pwy + czw*pwz - 2.0*czx*pwz + czy*pwz - cwx*(2.0*pww + pxw) - cxx*pxx - cyx*pxy - czx*pxz)),
+		(float)(-2.0*(cyy*pwy + czy*pwz + cwy*(pww + pxw) + cxy*(pwx + pxx) + cyy*pxy + czy*pxz)));
+
+	CameraMotionValue[2] = FVector4(
+		(float)((-4.0)*(cwx*pww + cxx*pwx + cyx*pwy + czx*pwz)),
+		(float)(cyz*pwy + czz*pwz + cwz*(pww + pxw) + cxz*(pwx + pxx) + cyz*pxy + czz*pxz),
+		(float)(cwy*pww + cwy*pxw + cww*(pww + pxw) - cwx*(pww + pxw) + (cxw - cxx + cxy)*(pwx + pxx) + (cyw - cyx + cyy)*(pwy + pxy) + (czw - czx + czy)*(pwz + pxz)),
+		(float)(0));
+
+	CameraMotionValue[3] = FVector4(
+		(float)((-4.0)*(cwx*pww + cxx*pwx + cyx*pwy + czx*pwz)),
+		(float)((-2.0)*(cwz*pww + cxz*pwx + cyz*pwy + czz*pwz)),
+		(float)(2.0*((-cww)*pww + cwx*pww - 2.0*cwy*pww - cxw*pwx + cxx*pwx - 2.0*cxy*pwx - cyw*pwy + cyx*pwy - 2.0*cyy*pwy - czw*pwz + czx*pwz - 2.0*czy*pwz + cwy*pyw + cxy*pyx + cyy*pyy + czy*pyz)),
+		(float)(2.0*(cyx*pwy + czx*pwz + cwx*(pww - pyw) + cxx*(pwx - pyx) - cyx*pyy - czx*pyz)));
+
+	CameraMotionValue[4] = FVector4(
+		(float)(4.0*(cwy*pww + cxy*pwx + cyy*pwy + czy*pwz)),
+		(float)(cyz*pwy + czz*pwz + cwz*(pww - pyw) + cxz*(pwx - pyx) - cyz*pyy - czz*pyz),
+		(float)(cwy*pww + cww*(pww - pyw) - cwy*pyw + cwx*((-pww) + pyw) + (cxw - cxx + cxy)*(pwx - pyx) + (cyw - cyx + cyy)*(pwy - pyy) + (czw - czx + czy)*(pwz - pyz)),
+		(float)(0));
+
+	SetShaderValueArray(ShaderRHI, CameraMotion, CameraMotionValue, 5);
+}
+
+FArchive& operator<<(FArchive& Ar, FCameraMotionParameters& This)
+{
+	Ar << This.CameraMotion;
+
+	return Ar;
+}
+
 /**
  * Encapsulates the post processing ambient occlusion pixel shader.
  * @param bAOSetupAsInput true:use AO setup instead of full resolution depth and normal
diff --git a/Engine/Source/Runtime/Renderer/Private/CompositionLighting/PostProcessAmbientOcclusion.h b/Engine/Source/Runtime/Renderer/Private/CompositionLighting/PostProcessAmbientOcclusion.h
--- a/Engine/Source/Runtime/Renderer/Private/CompositionLighting/PostProcessAmbientOcclusion.h
+++ b/Engine/Source/Runtime/Renderer/Private/CompositionLighting/PostProcessAmbientOcclusion.h
@@ -145,4 +145,19 @@ public:
 
 private:
 	FShaderParameter ScreenSpaceAOandSSRShaderParams;
-};
\ No newline at end of file
+};
+
+/** Pixel shader parameters needed for screen space TemporalAA and SSAO passes. */
+class FCameraMotionParameters
+{
+public:
+
+	void Bind(const FShaderParameterMap& ParameterMap);
+
+	void Set(const FSceneView& View, const FPixelShaderRHIParamRef ShaderRHI) const;
+
+	friend FArchive& operator<<(FArchive& Ar, FCameraMotionParameters& This);
+
+private:
+	FShaderParameter CameraMotion;
+};
diff --git a/Engine/Source/Runtime/Renderer/Private/PostProcess/PostProcessMobile.cpp b/Engine/Source/Runtime/Renderer/Private/PostProcess/PostProcessMobile.cpp
--- a/Engine/Source/Runtime/Renderer/Private/PostProcess/PostProcessMobile.cpp
+++ b/Engine/Source/Runtime/Renderer/Private/PostProcess/PostProcessMobile.cpp
@@ -2539,42 +2539,6 @@ public:
 		AaBlendAmount.Bind(Initializer.ParameterMap, TEXT("AaBlendAmount"));
 	}
 
-	static inline void Inverse4x4( double* dst, const float* src )
-	{
-		const double s0  = (double)(src[ 0]); const double s1  = (double)(src[ 1]); const double s2  = (double)(src[ 2]); const double s3  = (double)(src[ 3]);
-		const double s4  = (double)(src[ 4]); const double s5  = (double)(src[ 5]); const double s6  = (double)(src[ 6]); const double s7  = (double)(src[ 7]);
-		const double s8  = (double)(src[ 8]); const double s9  = (double)(src[ 9]); const double s10 = (double)(src[10]); const double s11 = (double)(src[11]);
-		const double s12 = (double)(src[12]); const double s13 = (double)(src[13]); const double s14 = (double)(src[14]); const double s15 = (double)(src[15]);
-
-		double inv[16];
-		inv[0]  =  s5 * s10 * s15 - s5 * s11 * s14 - s9 * s6 * s15 + s9 * s7 * s14 + s13 * s6 * s11 - s13 * s7 * s10;
-		inv[1]  = -s1 * s10 * s15 + s1 * s11 * s14 + s9 * s2 * s15 - s9 * s3 * s14 - s13 * s2 * s11 + s13 * s3 * s10;
-		inv[2]  =  s1 * s6  * s15 - s1 * s7  * s14 - s5 * s2 * s15 + s5 * s3 * s14 + s13 * s2 * s7  - s13 * s3 * s6;
-		inv[3]  = -s1 * s6  * s11 + s1 * s7  * s10 + s5 * s2 * s11 - s5 * s3 * s10 - s9  * s2 * s7  + s9  * s3 * s6;
-		inv[4]  = -s4 * s10 * s15 + s4 * s11 * s14 + s8 * s6 * s15 - s8 * s7 * s14 - s12 * s6 * s11 + s12 * s7 * s10;
-		inv[5]  =  s0 * s10 * s15 - s0 * s11 * s14 - s8 * s2 * s15 + s8 * s3 * s14 + s12 * s2 * s11 - s12 * s3 * s10;
-		inv[6]  = -s0 * s6  * s15 + s0 * s7  * s14 + s4 * s2 * s15 - s4 * s3 * s14 - s12 * s2 * s7  + s12 * s3 * s6;
-		inv[7]  =  s0 * s6  * s11 - s0 * s7  * s10 - s4 * s2 * s11 + s4 * s3 * s10 + s8  * s2 * s7  - s8  * s3 * s6;
-		inv[8]  =  s4 * s9  * s15 - s4 * s11 * s13 - s8 * s5 * s15 + s8 * s7 * s13 + s12 * s5 * s11 - s12 * s7 * s9;
-		inv[9]  = -s0 * s9  * s15 + s0 * s11 * s13 + s8 * s1 * s15 - s8 * s3 * s13 - s12 * s1 * s11 + s12 * s3 * s9;
-		inv[10] =  s0 * s5  * s15 - s0 * s7  * s13 - s4 * s1 * s15 + s4 * s3 * s13 + s12 * s1 * s7  - s12 * s3 * s5;
-		inv[11] = -s0 * s5  * s11 + s0 * s7  * s9  + s4 * s1 * s11 - s4 * s3 * s9  - s8  * s1 * s7  + s8  * s3 * s5;
-		inv[12] = -s4 * s9  * s14 + s4 * s10 * s13 + s8 * s5 * s14 - s8 * s6 * s13 - s12 * s5 * s10 + s12 * s6 * s9;
-		inv[13] =  s0 * s9  * s14 - s0 * s10 * s13 - s8 * s1 * s14 + s8 * s2 * s13 + s12 * s1 * s10 - s12 * s2 * s9;
-		inv[14] = -s0 * s5  * s14 + s0 * s6  * s13 + s4 * s1 * s14 - s4 * s2 * s13 - s12 * s1 * s6  + s12 * s2 * s5;
-		inv[15] =  s0 * s5  * s10 - s0 * s6  * s9  - s4 * s1 * s10 + s4 * s2 * s9  + s8  * s1 * s6  - s8  * s2 * s5;
-
-		double det = s0 * inv[0] + s1 * inv[4] + s2 * inv[8] + s3 * inv[12];
-		if( det != 0.0 )
-		{
-			det = 1.0 / det;
-		}
-		for( int i = 0; i < 16; i++ )
-		{
-			dst[i] = inv[i] * det;
-		}
-	}
-
 	void SetPS(const FRenderingCompositePassContext& Context)
 	{
 		const FPixelShaderRHIParamRef ShaderRHI = GetPixelShader();
diff --git a/Engine/Source/Runtime/Renderer/Private/PostProcess/PostProcessTemporalAA.cpp b/Engine/Source/Runtime/Renderer/Private/PostProcess/PostProcessTemporalAA.cpp
--- a/Engine/Source/Runtime/Renderer/Private/PostProcess/PostProcessTemporalAA.cpp
+++ b/Engine/Source/Runtime/Renderer/Private/PostProcess/PostProcessTemporalAA.cpp
@@ -8,45 +8,10 @@
 #include "ScenePrivate.h"
 #include "SceneFilterRendering.h"
 #include "PostProcessTemporalAA.h"
+#include "PostProcessAmbientOcclusion.h"
 #include "PostProcessTonemap.h"
 #include "PostProcessing.h"
 
-static inline void Inverse4x4( double* dst, const float* src )
-{
-	const double s0  = (double)(src[ 0]); const double s1  = (double)(src[ 1]); const double s2  = (double)(src[ 2]); const double s3  = (double)(src[ 3]);
-	const double s4  = (double)(src[ 4]); const double s5  = (double)(src[ 5]); const double s6  = (double)(src[ 6]); const double s7  = (double)(src[ 7]);
-	const double s8  = (double)(src[ 8]); const double s9  = (double)(src[ 9]); const double s10 = (double)(src[10]); const double s11 = (double)(src[11]);
-	const double s12 = (double)(src[12]); const double s13 = (double)(src[13]); const double s14 = (double)(src[14]); const double s15 = (double)(src[15]);
-
-	double inv[16];
-	inv[0]  =  s5 * s10 * s15 - s5 * s11 * s14 - s9 * s6 * s15 + s9 * s7 * s14 + s13 * s6 * s11 - s13 * s7 * s10;
-	inv[1]  = -s1 * s10 * s15 + s1 * s11 * s14 + s9 * s2 * s15 - s9 * s3 * s14 - s13 * s2 * s11 + s13 * s3 * s10;
-	inv[2]  =  s1 * s6  * s15 - s1 * s7  * s14 - s5 * s2 * s15 + s5 * s3 * s14 + s13 * s2 * s7  - s13 * s3 * s6;
-	inv[3]  = -s1 * s6  * s11 + s1 * s7  * s10 + s5 * s2 * s11 - s5 * s3 * s10 - s9  * s2 * s7  + s9  * s3 * s6;
-	inv[4]  = -s4 * s10 * s15 + s4 * s11 * s14 + s8 * s6 * s15 - s8 * s7 * s14 - s12 * s6 * s11 + s12 * s7 * s10;
-	inv[5]  =  s0 * s10 * s15 - s0 * s11 * s14 - s8 * s2 * s15 + s8 * s3 * s14 + s12 * s2 * s11 - s12 * s3 * s10;
-	inv[6]  = -s0 * s6  * s15 + s0 * s7  * s14 + s4 * s2 * s15 - s4 * s3 * s14 - s12 * s2 * s7  + s12 * s3 * s6;
-	inv[7]  =  s0 * s6  * s11 - s0 * s7  * s10 - s4 * s2 * s11 + s4 * s3 * s10 + s8  * s2 * s7  - s8  * s3 * s6;
-	inv[8]  =  s4 * s9  * s15 - s4 * s11 * s13 - s8 * s5 * s15 + s8 * s7 * s13 + s12 * s5 * s11 - s12 * s7 * s9;
-	inv[9]  = -s0 * s9  * s15 + s0 * s11 * s13 + s8 * s1 * s15 - s8 * s3 * s13 - s12 * s1 * s11 + s12 * s3 * s9;
-	inv[10] =  s0 * s5  * s15 - s0 * s7  * s13 - s4 * s1 * s15 + s4 * s3 * s13 + s12 * s1 * s7  - s12 * s3 * s5;
-	inv[11] = -s0 * s5  * s11 + s0 * s7  * s9  + s4 * s1 * s11 - s4 * s3 * s9  - s8  * s1 * s7  + s8  * s3 * s5;
-	inv[12] = -s4 * s9  * s14 + s4 * s10 * s13 + s8 * s5 * s14 - s8 * s6 * s13 - s12 * s5 * s10 + s12 * s6 * s9;
-	inv[13] =  s0 * s9  * s14 - s0 * s10 * s13 - s8 * s1 * s14 + s8 * s2 * s13 + s12 * s1 * s10 - s12 * s2 * s9;
-	inv[14] = -s0 * s5  * s14 + s0 * s6  * s13 + s4 * s1 * s14 - s4 * s2 * s13 - s12 * s1 * s6  + s12 * s2 * s5;
-	inv[15] =  s0 * s5  * s10 - s0 * s6  * s9  - s4 * s1 * s10 + s4 * s2 * s9  + s8  * s1 * s6  - s8  * s2 * s5;
-
-	double det = s0 * inv[0] + s1 * inv[4] + s2 * inv[8] + s3 * inv[12];
-	if( det != 0.0 )
-	{
-		det = 1.0 / det;
-	}
-	for( int i = 0; i < 16; i++ )
-	{
-		dst[i] = inv[i] * det;
-	}
-}
-
 static float TemporalHalton( int32 Index, int32 Base )
 {
 	float Result = 0.0f;
@@ -99,7 +64,7 @@ public:
 	FDeferredPixelShaderParameters DeferredParameters;
 	FShaderParameter SampleWeights;
 	FShaderParameter LowpassWeights;
-	FShaderParameter CameraMotion;
+	FCameraMotionParameters CameraMotionParams;
 	FShaderParameter RandomOffset;
 
 	/** Initialization constructor. */
@@ -110,7 +75,7 @@ public:
 		DeferredParameters.Bind(Initializer.ParameterMap);
 		SampleWeights.Bind(Initializer.ParameterMap, TEXT("SampleWeights"));
 		LowpassWeights.Bind(Initializer.ParameterMap, TEXT("LowpassWeights"));
-		CameraMotion.Bind(Initializer.ParameterMap, TEXT("CameraMotion"));
+		CameraMotionParams.Bind(Initializer.ParameterMap);
 		RandomOffset.Bind(Initializer.ParameterMap, TEXT("RandomOffset"));
 	}
 
@@ -118,7 +83,7 @@ public:
 	virtual bool Serialize(FArchive& Ar)
 	{
 		bool bShaderHasOutdatedParameters = FGlobalShader::Serialize(Ar);
-		Ar << PostprocessParameter << DeferredParameters << SampleWeights << LowpassWeights << CameraMotion << RandomOffset;
+		Ar << PostprocessParameter << DeferredParameters << SampleWeights << LowpassWeights << CameraMotionParams << RandomOffset;
 		return bShaderHasOutdatedParameters;
 	}
 
@@ -202,65 +167,7 @@ public:
 
 		}
 
-		{
-			FMatrix Proj = Context.View.ViewMatrices.ProjMatrix;
-			FMatrix PrevProj = ViewState->PrevViewMatrices.ProjMatrix;
-
-			// Remove jitter
-			Proj.M[2][0] = 0.0f;
-			Proj.M[2][1] = 0.0f;
-			PrevProj.M[2][0] = 0.0f;
-			PrevProj.M[2][1] = 0.0f;
-
-			FMatrix ViewProj = ( Context.View.ViewMatrices.ViewMatrix * Proj ).GetTransposed();
-			FMatrix PrevViewProj = ( ViewState->PrevViewMatrices.ViewMatrix * PrevProj ).GetTransposed();
-
-			double InvViewProj[16];
-			Inverse4x4( InvViewProj, (float*)ViewProj.M );
-
-			const float* p = (float*)PrevViewProj.M;
-
-			const double cxx = InvViewProj[ 0]; const double cxy = InvViewProj[ 1]; const double cxz = InvViewProj[ 2]; const double cxw = InvViewProj[ 3];
-			const double cyx = InvViewProj[ 4]; const double cyy = InvViewProj[ 5]; const double cyz = InvViewProj[ 6]; const double cyw = InvViewProj[ 7];
-			const double czx = InvViewProj[ 8]; const double czy = InvViewProj[ 9]; const double czz = InvViewProj[10]; const double czw = InvViewProj[11];
-			const double cwx = InvViewProj[12]; const double cwy = InvViewProj[13]; const double cwz = InvViewProj[14]; const double cww = InvViewProj[15];
-
-			const double pxx = (double)(p[ 0]); const double pxy = (double)(p[ 1]); const double pxz = (double)(p[ 2]); const double pxw = (double)(p[ 3]);
-			const double pyx = (double)(p[ 4]); const double pyy = (double)(p[ 5]); const double pyz = (double)(p[ 6]); const double pyw = (double)(p[ 7]);
-			const double pwx = (double)(p[12]); const double pwy = (double)(p[13]); const double pwz = (double)(p[14]); const double pww = (double)(p[15]);
-
-			FVector4 Motion;
-
-			Motion[0] = (float)(4.0*(cwx*pww + cxx*pwx + cyx*pwy + czx*pwz));
-			Motion[1] = (float)((-4.0)*(cwy*pww + cxy*pwx + cyy*pwy + czy*pwz));
-			Motion[2] = (float)(2.0*(cwz*pww + cxz*pwx + cyz*pwy + czz*pwz));
-			Motion[3] = (float)(2.0*(cww*pww - cwx*pww + cwy*pww + (cxw - cxx + cxy)*pwx + (cyw - cyx + cyy)*pwy + (czw - czx + czy)*pwz));
-			SetShaderValue( ShaderRHI, CameraMotion, Motion, 0 );
-
-			Motion[0] = (float)(( 4.0)*(cwy*pww + cxy*pwx + cyy*pwy + czy*pwz));
-			Motion[1] = (float)((-2.0)*(cwz*pww + cxz*pwx + cyz*pwy + czz*pwz));
-			Motion[2] = (float)((-2.0)*(cww*pww + cwy*pww + cxw*pwx - 2.0*cxx*pwx + cxy*pwx + cyw*pwy - 2.0*cyx*pwy + cyy*pwy + czw*pwz - 2.0*czx*pwz + czy*pwz - cwx*(2.0*pww + pxw) - cxx*pxx - cyx*pxy - czx*pxz));
-			Motion[3] = (float)(-2.0*(cyy*pwy + czy*pwz + cwy*(pww + pxw) + cxy*(pwx + pxx) + cyy*pxy + czy*pxz));
-			SetShaderValue( ShaderRHI, CameraMotion, Motion, 1 );
-
-			Motion[0] = (float)((-4.0)*(cwx*pww + cxx*pwx + cyx*pwy + czx*pwz));
-			Motion[1] = (float)(cyz*pwy + czz*pwz + cwz*(pww + pxw) + cxz*(pwx + pxx) + cyz*pxy + czz*pxz);
-			Motion[2] = (float)(cwy*pww + cwy*pxw + cww*(pww + pxw) - cwx*(pww + pxw) + (cxw - cxx + cxy)*(pwx + pxx) + (cyw - cyx + cyy)*(pwy + pxy) + (czw - czx + czy)*(pwz + pxz));
-			Motion[3] = (float)(0);
-			SetShaderValue( ShaderRHI, CameraMotion, Motion, 2 );
-
-			Motion[0] = (float)((-4.0)*(cwx*pww + cxx*pwx + cyx*pwy + czx*pwz));
-			Motion[1] = (float)((-2.0)*(cwz*pww + cxz*pwx + cyz*pwy + czz*pwz));
-			Motion[2] = (float)(2.0*((-cww)*pww + cwx*pww - 2.0*cwy*pww - cxw*pwx + cxx*pwx - 2.0*cxy*pwx - cyw*pwy + cyx*pwy - 2.0*cyy*pwy - czw*pwz + czx*pwz - 2.0*czy*pwz + cwy*pyw + cxy*pyx + cyy*pyy + czy*pyz));
-			Motion[3] = (float)(2.0*(cyx*pwy + czx*pwz + cwx*(pww - pyw) + cxx*(pwx - pyx) - cyx*pyy - czx*pyz));
-			SetShaderValue( ShaderRHI, CameraMotion, Motion, 3 );
-
-			Motion[0] = (float)(4.0*(cwy*pww + cxy*pwx + cyy*pwy + czy*pwz));
-			Motion[1] = (float)(cyz*pwy + czz*pwz + cwz*(pww - pyw) + cxz*(pwx - pyx) - cyz*pyy - czz*pyz);
-			Motion[2] = (float)(cwy*pww + cww*(pww - pyw) - cwy*pyw + cwx*((-pww) + pyw) + (cxw - cxx + cxy)*(pwx - pyx) + (cyw - cyx + cyy)*(pwy - pyy) + (czw - czx + czy)*(pwz - pyz));
-			Motion[3] = (float)(0);
-			SetShaderValue( ShaderRHI, CameraMotion, Motion, 4 );
-		}
+		CameraMotionParams.Set(Context.View, ShaderRHI);
 	}
 };
 
