* optimized and quality improved ScreenSpaceAmbientOcclusion using dither noise, using less samples, added r.AmbientOcclusionSampleSetQuality cvar
  now it  can be used with larger radius and less levels

diff --git a/Engine/Shaders/PostProcessAmbientOcclusion.usf b/Engine/Shaders/PostProcessAmbientOcclusion.usf
--- a/Engine/Shaders/PostProcessAmbientOcclusion.usf
+++ b/Engine/Shaders/PostProcessAmbientOcclusion.usf
@@ -8,6 +8,19 @@
 #include "PostProcessCommon.usf"				
 #include "DeferredShadingCommon.usf"
 
+// set by C++:
+//
+// 0:low / 1:high
+// SAMPLESET_QUALITY
+//
+// 0:no / 1:yes
+// USE_AO_SETUP_AS_INPUT
+//
+// 0:no / 1:yes
+// USE_UPSAMPLE
+
+
+
 // 0:does not use normals (slightly faster, lower quality, WIP = not finished), 1: uses normals (slower but higher quality)
 #define SSAO_NORMAL_AFFECTED 1
 
@@ -15,7 +28,8 @@
 // AO_SAMPLE_QUALITY = 0 : no AO sampling, only upsampling
 // AO_SAMPLE_QUALITY = 1 : no dither/per pixel randomization
 // AO_SAMPLE_QUALITY = 2 : efficient high frequency 4x4 pattern
-// AO_SAMPLE_QUALITY = 3 : extreme sample count
+// AO_SAMPLE_QUALITY = 3 : efficient high frequency 4x4 pattern changing over time
+// AO_SAMPLE_QUALITY = 4 : extreme sample count
 
 // [0]: .x:AmbientOcclusionPower, .y:AmbientOcclusionBias/BiasDistance, .z:1/AmbientOcclusionDistance, .w:AmbientOcclusionIntensity
 // [1]: .xy:ViewportUVToRandomUV, .z:AORadiusInShader, .w:Ratio
@@ -57,13 +71,22 @@ const static float Constant_Float16F_Scale =  4096.0f * 32.0f;
 	#define AO_UPSAMPLE_QUALITY 1
 #else
 	// full resolution is expensive, do lower quality
-	#define AO_SAMPLE_QUALITY 1
+	#define AO_SAMPLE_QUALITY 3
 	#define AO_UPSAMPLE_QUALITY 0
 #endif
 
 // 0: near the surface very large, softly fading out (6 points)
 // 1: more evenly spread (5 points - slightly faster, stronger effect, better with multiple levels?)
-#define USE_SAMPLESET 1
+// 2: 3 points
+// 3: 1 point (for testing)
+//#define USE_SAMPLESET 2
+
+#if SAMPLESET_QUALITY
+	#define USE_SAMPLESET 0
+#else
+	#define USE_SAMPLESET 2
+#endif
+
 
 // only for MainSetupPS()
 // .x:ScaleFactor(e.g. 4 if current RT is a quarter in size), .y:InvThreshold, .zw: unused
@@ -84,7 +107,7 @@ float4 NoiseScale;
 		float2(-0.711, 0.274), 
 		float2(0.060, 0.900) 
 	};
-#else // USE_SAMPLESET == 0
+#elif USE_SAMPLESET == 1
 	#define SAMPLESET_ARRAY_SIZE 5
 	static const float2 OcclusionSamplesOffsets[SAMPLESET_ARRAY_SIZE]=
 	{
@@ -95,7 +118,23 @@ float4 NoiseScale;
 		float2(-0.707107, -0.707107),
 		float2(-0.891006, 0.453991),
 	};
-#endif // USE_SAMPLESET == 0
+#elif USE_SAMPLESET == 2
+	#define SAMPLESET_ARRAY_SIZE 3
+	static const float2 OcclusionSamplesOffsets[SAMPLESET_ARRAY_SIZE]=
+	{
+		// 3 points distributed on the unit disc, spiral order and distance
+		float2(0, -1.0f) * 0.3f, 
+		float2(0.81f, 0.5) * 0.7f, 
+		float2(-0.81f, 0.5) 
+	};
+#else // USE_SAMPLESET == 3
+	#define SAMPLESET_ARRAY_SIZE 1
+	static const float2 OcclusionSamplesOffsets[SAMPLESET_ARRAY_SIZE]=
+	{
+		// 6 points distributed on the unit disc, spiral order and distance
+		float2(0.500, 0.500), 
+	};
+#endif // USE_SAMPLESET
 
 /** RGBA8 linear texture containing random normals */
 Texture2D RandomNormalTexture;
@@ -276,7 +315,7 @@ float acosApproxNormAngle(float x)
 {
 	// todo: expose
 	// 1: is a good linear approximation, 0.9f seems to look good
-	float4 ContrastTweak = 0.9f;
+	float ContrastTweak = 0.9f;
 
 	// correct: acos(x) / PI
 	// linear approximation: saturate((1 - x) * 0.5f);
@@ -339,7 +378,15 @@ void MainPS(in float4 UVAndScreenPos : TEXCOORD0, out float4 OutColor : SV_Targe
 	// extract one of 16 base vectors (rotation and scale) from a texture that repeats 4x4
 	float4 RandomBase = (Texture2DSample(RandomNormalTexture, RandomNormalTextureSampler, UV * ViewportUVToRandomUV) * 2 - 1 ) * float4(AORadiusInShader, AORadiusInShader);
 	{
-#else // AO_SAMPLE_QUALITY == 3
+#elif AO_SAMPLE_QUALITY == 3
+	// extract one of 16 base vectors (rotation and scale) from a texture that repeats 4x4, changing over time if TemporalAA is enabled
+	float2 TemporalOffset = (View.TemporalAAParams.r % 8) * float2(3.5, 7) / 64.0f;
+
+	float4 RandomNorm = Texture2DSample(RandomNormalTexture, RandomNormalTextureSampler, TemporalOffset + UV * ViewportUVToRandomUV);
+	float4 RandomBase = (RandomNorm * 2 - 1) * float4(AORadiusInShader, AORadiusInShader);
+
+	{
+#else // AO_SAMPLE_QUALITY == 4
 	// can be optimized
 	int2 PixelPos = int2(UVAndScreenPos.zw * ScreenPosToPixel.xy + ScreenPosToPixel.zw + 0.5f);
 	// repeat a cross pattern in a 5x5 tile
diff --git a/Engine/Source/Runtime/Renderer/Private/CompositionLighting/PostProcessAmbientOcclusion.cpp b/Engine/Source/Runtime/Renderer/Private/CompositionLighting/PostProcessAmbientOcclusion.cpp
--- a/Engine/Source/Runtime/Renderer/Private/CompositionLighting/PostProcessAmbientOcclusion.cpp
+++ b/Engine/Source/Runtime/Renderer/Private/CompositionLighting/PostProcessAmbientOcclusion.cpp
@@ -10,6 +10,14 @@
 #include "PostProcessing.h"
 #include "PostProcessAmbientOcclusion.h"
 
+static TAutoConsoleVariable<int32> CVarAmbientOcclusionSampleSetQuality(
+	TEXT("r.AmbientOcclusionSampleSetQuality"),
+	-1,
+	TEXT("Defines how many samples we use for ScreenSpace Ambient Occlusion\n")
+	TEXT("-1: sample count depends on post process settings (default)\n")
+	TEXT("0: low sample count (defined in shader, 3 * 2)\n")
+	TEXT("1: high sample count (defined in shader, 6 * 2)"),
+	ECVF_RenderThreadSafe);
 
 
 /** Encapsulates the post processing ambient occlusion pixel shader. */
@@ -185,7 +193,7 @@ FArchive& operator<<(FArchive& Ar, FCameraMotionParameters& This)
  * @param bAOSetupAsInput true:use AO setup instead of full resolution depth and normal
  * @param bDoUpsample true:we have lower resolution pass data we need to upsample, false otherwise
  */
-template<uint32 bTAOSetupAsInput, uint32 bDoUpsample>
+template<uint32 bTAOSetupAsInput, uint32 bDoUpsample, uint32 SampleSetQuality>
 class FPostProcessAmbientOcclusionPS : public FGlobalShader
 {
 	DECLARE_SHADER_TYPE(FPostProcessAmbientOcclusionPS, Global);
@@ -201,6 +209,7 @@ class FPostProcessAmbientOcclusionPS : public FGlobalShader
 
 		OutEnvironment.SetDefine(TEXT("USE_UPSAMPLE"), bDoUpsample);
 		OutEnvironment.SetDefine(TEXT("USE_AO_SETUP_AS_INPUT"), bTAOSetupAsInput);
+		OutEnvironment.SetDefine(TEXT("SAMPLESET_QUALITY"), SampleSetQuality);
 	}
 
 	/** Default constructor. */
@@ -263,11 +272,13 @@ public:
 
 
 // #define avoids a lot of code duplication
-#define VARIATION1(A)	VARIATION2(A, 0) VARIATION2(A, 1)
-#define VARIATION2(A, B) typedef FPostProcessAmbientOcclusionPS<A, B> FPostProcessAmbientOcclusionPS##A##B; \
-	IMPLEMENT_SHADER_TYPE2(FPostProcessAmbientOcclusionPS##A##B, SF_Pixel);
+#define VARIATION1(A, C)	VARIATION2(A, 0, C) VARIATION2(A, 1, C)
+#define VARIATION2(A, B, C) typedef FPostProcessAmbientOcclusionPS<A, B, C> FPostProcessAmbientOcclusionPS##A##B##C; \
+	IMPLEMENT_SHADER_TYPE2(FPostProcessAmbientOcclusionPS##A##B##C, SF_Pixel);
+
+	VARIATION1(0,0) VARIATION1(1,0)
+	VARIATION1(0,1) VARIATION1(1,1)
 
-	VARIATION1(0) VARIATION1(1)
 #undef VARIATION1
 #undef VARIATION2
 
@@ -436,11 +447,11 @@ bool FRCPassPostProcessAmbientOcclusionSetup::IsInitialPass() const
 
 
 // ---------------------------------
-template <uint32 bTAOSetupAsInput, uint32 bDoUpsample>
+template <uint32 bTAOSetupAsInput, uint32 bDoUpsample, uint32 SampleSetQuality>
 void FRCPassPostProcessAmbientOcclusion::SetShaderTempl(const FRenderingCompositePassContext& Context)
 {
 	TShaderMapRef<FPostProcessVS> VertexShader(GetGlobalShaderMap());
-	TShaderMapRef<FPostProcessAmbientOcclusionPS<bTAOSetupAsInput, bDoUpsample> > PixelShader(GetGlobalShaderMap());
+	TShaderMapRef<FPostProcessAmbientOcclusionPS<bTAOSetupAsInput, bDoUpsample, SampleSetQuality> > PixelShader(GetGlobalShaderMap());
 
 	static FGlobalBoundShaderState BoundShaderState;
 
@@ -490,8 +501,18 @@ void FRCPassPostProcessAmbientOcclusion::Process(FRenderingCompositePassContext&
 	RHISetRasterizerState(TStaticRasterizerState<>::GetRHI());
 	RHISetDepthStencilState(TStaticDepthStencilState<false,CF_Always>::GetRHI());
 
-	// 0..1
-	uint32 Quality = (int32)Context.View.FinalPostProcessSettings.AmbientOcclusionQuality > 50;
+	int32 SampleSetQuality = CVarAmbientOcclusionSampleSetQuality.GetValueOnRenderThread();
+
+	if(SampleSetQuality < 0)
+	{
+		SampleSetQuality = 0;
+
+		// SampleSetQuality depends on post process settings
+		if(Context.View.FinalPostProcessSettings.AmbientOcclusionQuality > 60.0f)
+		{
+			SampleSetQuality = 1;
+		}
+	}
 
 	bool bDoUpsample = (InputDesc2 != 0);
 
@@ -499,22 +520,22 @@ void FRCPassPostProcessAmbientOcclusion::Process(FRenderingCompositePassContext&
 	{
 		if(bDoUpsample)
 		{
-			SetShaderTempl<1, 1>(Context);
+			if(SampleSetQuality) SetShaderTempl<1, 1, 1>(Context); else SetShaderTempl<1, 1, 0>(Context);
 		}
 		else
 		{
-			SetShaderTempl<1, 0>(Context);
+			if(SampleSetQuality) SetShaderTempl<1, 0, 1>(Context); else SetShaderTempl<1, 0, 0>(Context);
 		}
 	}
 	else
 	{
 		if(bDoUpsample)
 		{
-			SetShaderTempl<0, 1>(Context);
+			if(SampleSetQuality) SetShaderTempl<0, 1, 1>(Context); else SetShaderTempl<0, 1, 0>(Context);
 		}
 		else
 		{
-			SetShaderTempl<0, 0>(Context);
+			if(SampleSetQuality) SetShaderTempl<0, 0, 1>(Context); else SetShaderTempl<0, 0, 0>(Context);
 		}
 	}
 
diff --git a/Engine/Source/Runtime/Renderer/Private/CompositionLighting/PostProcessAmbientOcclusion.h b/Engine/Source/Runtime/Renderer/Private/CompositionLighting/PostProcessAmbientOcclusion.h
--- a/Engine/Source/Runtime/Renderer/Private/CompositionLighting/PostProcessAmbientOcclusion.h
+++ b/Engine/Source/Runtime/Renderer/Private/CompositionLighting/PostProcessAmbientOcclusion.h
@@ -48,7 +48,7 @@ public:
 
 private:
 	
-	template <uint32 bAOSetupAsInput, uint32 bDoUpsample>
+	template <uint32 bAOSetupAsInput, uint32 bDoUpsample, uint32 SampleSetQuality>
 		void SetShaderTempl(const FRenderingCompositePassContext& Context);
 
 	bool bAOSetupAsInput;
