Atmospheric fog : Added SunDiscScale to control sun disk size

diff --git a/Engine/Shaders/AtmosphereCommon.usf b/Engine/Shaders/AtmosphereCommon.usf
--- a/Engine/Shaders/AtmosphereCommon.usf
+++ b/Engine/Shaders/AtmosphereCommon.usf
@@ -520,7 +520,7 @@ float3 GetSunColor(float3 X, float T, float3 V, float3 S, float Radius, float Mu
     }
 	else 
 	{
-        float SunIntensity = step(cos(PI / 180.0), dot(V, S)); // Lsun
+		float SunIntensity = step(cos(PI * View.AtmosphericFogSunDiscScale / 180.0), dot(V, S)); // Lsun
         return TransmittanceValue * SunIntensity; // Eq (9)
     }
 }
diff --git a/Engine/Source/Runtime/Engine/Classes/Atmosphere/AtmosphericFogComponent.h b/Engine/Source/Runtime/Engine/Classes/Atmosphere/AtmosphericFogComponent.h
--- a/Engine/Source/Runtime/Engine/Classes/Atmosphere/AtmosphericFogComponent.h
+++ b/Engine/Source/Runtime/Engine/Classes/Atmosphere/AtmosphericFogComponent.h
@@ -150,6 +150,10 @@ class UAtmosphericFogComponent : public USceneComponent
 	UPROPERTY(EditAnywhere, BlueprintReadWrite, interp, Category=Atmosphere)
 	float StartDistance;
 
+	/** Distance offset, in km (to handle large distance) */
+	UPROPERTY(EditAnywhere, BlueprintReadWrite, interp, Category = Atmosphere)
+	float SunDiscScale;
+
 	/** Default light brightness. Used when there is no sunlight placed in the level. Unit is lumens */
 	UPROPERTY(EditAnywhere, BlueprintReadOnly, interp, Category=Lighting)
 	float DefaultBrightness;
diff --git a/Engine/Source/Runtime/Engine/Private/Atmosphere/Atmosphere.cpp b/Engine/Source/Runtime/Engine/Private/Atmosphere/Atmosphere.cpp
--- a/Engine/Source/Runtime/Engine/Private/Atmosphere/Atmosphere.cpp
+++ b/Engine/Source/Runtime/Engine/Private/Atmosphere/Atmosphere.cpp
@@ -90,6 +90,7 @@ UAtmosphericFogComponent::UAtmosphericFogComponent(const class FPostConstructIni
 	DensityOffset = 0.f;
 	DistanceScale = 1.f;
 	AltitudeScale = 1.f;
+	SunDiscScale = 1.f;
 	StartDistance = 15000;
 	DistanceOffset = 0.f;
 	GroundOffset = -100000.f; // -1 km in default distance scale, 100K UU
@@ -506,6 +507,7 @@ void UAtmosphericFogComponent::PostEditChangeProperty(FPropertyChangedEvent& Pro
 		DensityOffset = FMath::Clamp(DensityOffset, -1.0f, 1.f);
 		DistanceScale = FMath::Clamp(DistanceScale, 0.1f, 1000.f);
 		AltitudeScale = FMath::Clamp(AltitudeScale, 0.1f, 1000.f);
+		SunDiscScale = FMath::Clamp(SunDiscScale, 0.1f, 1000.f);
 		GroundOffset = FMath::Clamp(GroundOffset, (float)-WORLD_MAX, (float)WORLD_MAX);
 		StartDistance = FMath::Clamp(StartDistance, 100.f, (float)WORLD_MAX);
 	}
diff --git a/Engine/Source/Runtime/Engine/Public/SceneManagement.h b/Engine/Source/Runtime/Engine/Public/SceneManagement.h
--- a/Engine/Source/Runtime/Engine/Public/SceneManagement.h
+++ b/Engine/Source/Runtime/Engine/Public/SceneManagement.h
@@ -2123,6 +2123,7 @@ BEGIN_UNIFORM_BUFFER_STRUCT(FViewUniformShaderParameters,ENGINE_API)
 	DECLARE_UNIFORM_BUFFER_STRUCT_MEMBER_EX(float,AtmosphericFogHeightScaleRayleigh, EShaderPrecisionModifier::Half)
 	DECLARE_UNIFORM_BUFFER_STRUCT_MEMBER_EX(float,AtmosphericFogStartDistance, EShaderPrecisionModifier::Half)
 	DECLARE_UNIFORM_BUFFER_STRUCT_MEMBER_EX(float,AtmosphericFogDistanceOffset, EShaderPrecisionModifier::Half)
+	DECLARE_UNIFORM_BUFFER_STRUCT_MEMBER_EX(float,AtmosphericFogSunDiscScale, EShaderPrecisionModifier::Half)
 	DECLARE_UNIFORM_BUFFER_STRUCT_MEMBER(uint32,AtmosphericFogRenderMask)
 	DECLARE_UNIFORM_BUFFER_STRUCT_MEMBER(uint32,AtmosphericFogInscatterAltitudeSampleNum)
 	DECLARE_UNIFORM_BUFFER_STRUCT_MEMBER(FLinearColor,AtmosphericFogSunColor)
diff --git a/Engine/Source/Runtime/Renderer/Private/AtmosphereRendering.cpp b/Engine/Source/Runtime/Renderer/Private/AtmosphereRendering.cpp
--- a/Engine/Source/Runtime/Renderer/Private/AtmosphereRendering.cpp
+++ b/Engine/Source/Runtime/Renderer/Private/AtmosphereRendering.cpp
@@ -1556,6 +1556,7 @@ FAtmosphericFogSceneInfo::FAtmosphericFogSceneInfo(UAtmosphericFogComponent* InC
 	, RHeight(InComponent->PrecomputeParams.DensityHeight * InComponent->PrecomputeParams.DensityHeight * InComponent->PrecomputeParams.DensityHeight * 64.f)
 	, StartDistance(InComponent->StartDistance)
 	, DistanceOffset(InComponent->DistanceOffset)
+	, SunDiscScale(InComponent->SunDiscScale)
 	, RenderFlag(EAtmosphereRenderFlag::E_EnableAll)
 	, InscatterAltitudeSampleNum(InComponent->PrecomputeParams.InscatterAltitudeSampleNum)
 #if WITH_EDITOR
diff --git a/Engine/Source/Runtime/Renderer/Private/AtmosphereRendering.h b/Engine/Source/Runtime/Renderer/Private/AtmosphereRendering.h
--- a/Engine/Source/Runtime/Renderer/Private/AtmosphereRendering.h
+++ b/Engine/Source/Runtime/Renderer/Private/AtmosphereRendering.h
@@ -39,6 +39,7 @@ public:
 	float RHeight;
 	float StartDistance;
 	float DistanceOffset;
+	float SunDiscScale;
 	FLinearColor DefaultSunColor;
 	FVector DefaultSunDirection;
 	uint32 RenderFlag;
diff --git a/Engine/Source/Runtime/Renderer/Private/SceneRendering.cpp b/Engine/Source/Runtime/Renderer/Private/SceneRendering.cpp
--- a/Engine/Source/Runtime/Renderer/Private/SceneRendering.cpp
+++ b/Engine/Source/Runtime/Renderer/Private/SceneRendering.cpp
@@ -306,6 +306,7 @@ TUniformBufferRef<FViewUniformShaderParameters> FViewInfo::CreateUniformBuffer(
 			ViewUniformShaderParameters.AtmosphericFogHeightScaleRayleigh =  Scene->AtmosphericFog->RHeight;
 			ViewUniformShaderParameters.AtmosphericFogStartDistance = Scene->AtmosphericFog->StartDistance;
 			ViewUniformShaderParameters.AtmosphericFogDistanceOffset = Scene->AtmosphericFog->DistanceOffset;
+			ViewUniformShaderParameters.AtmosphericFogSunDiscScale = Scene->AtmosphericFog->SunDiscScale;
 			ViewUniformShaderParameters.AtmosphericFogSunColor = Scene->SunLight ? Scene->SunLight->Proxy->GetColor() : Scene->AtmosphericFog->DefaultSunColor;
 			ViewUniformShaderParameters.AtmosphericFogSunDirection = Scene->SunLight ? -Scene->SunLight->Proxy->GetDirection() : -Scene->AtmosphericFog->DefaultSunDirection;
 			ViewUniformShaderParameters.AtmosphericFogRenderMask = Scene->AtmosphericFog->RenderFlag & (EAtmosphereRenderFlag::E_DisableGroundScattering | EAtmosphereRenderFlag::E_DisableSunDisk);
@@ -323,6 +324,7 @@ TUniformBufferRef<FViewUniformShaderParameters> FViewInfo::CreateUniformBuffer(
 			ViewUniformShaderParameters.AtmosphericFogHeightScaleRayleigh =  0.f;
 			ViewUniformShaderParameters.AtmosphericFogStartDistance = 0.f;
 			ViewUniformShaderParameters.AtmosphericFogDistanceOffset = 0.f;
+			ViewUniformShaderParameters.AtmosphericFogSunDiscScale = 1.f;
 			ViewUniformShaderParameters.AtmosphericFogSunColor = FLinearColor::Black;
 			ViewUniformShaderParameters.AtmosphericFogSunDirection = FVector::ZeroVector;
 			ViewUniformShaderParameters.AtmosphericFogRenderMask = EAtmosphereRenderFlag::E_EnableAll;
@@ -347,6 +349,7 @@ TUniformBufferRef<FViewUniformShaderParameters> FViewInfo::CreateUniformBuffer(
 		ViewUniformShaderParameters.AtmosphericFogHeightScaleRayleigh =  0.f;
 		ViewUniformShaderParameters.AtmosphericFogStartDistance = 0.f;
 		ViewUniformShaderParameters.AtmosphericFogDistanceOffset = 0.f;
+		ViewUniformShaderParameters.AtmosphericFogSunDiscScale = 1.f;
 		ViewUniformShaderParameters.AtmosphericFogSunColor = FLinearColor::Black;
 		ViewUniformShaderParameters.AtmosphericFogSunDirection = FVector::ZeroVector;
 		ViewUniformShaderParameters.AtmosphericFogRenderMask = EAtmosphereRenderFlag::E_EnableAll;
