Shadow clearing / resolving optimization
#change - Clear and resolve whole shadow targets rather than using subrects.
#change - RHISetRenderTargets sets the viewport to the size of the depth buffer if there is no colour target

diff --git a/Engine/Source/Runtime/OpenGLDrv/Private/OpenGLCommands.cpp b/Engine/Source/Runtime/OpenGLDrv/Private/OpenGLCommands.cpp
--- a/Engine/Source/Runtime/OpenGLDrv/Private/OpenGLCommands.cpp
+++ b/Engine/Source/Runtime/OpenGLDrv/Private/OpenGLCommands.cpp
@@ -1523,7 +1523,7 @@ void FOpenGLDynamicRHI::RHISetRenderTargets(
 
 	if (PendingState.FirstNonzeroRenderTarget != -1)
 	{
-		// Set viewport size to new render target size. (I see D3D11 doesn't set it in case of depth buffers with no render targets, so doing the same here).
+		// Set viewport size to new render target size.
 		PendingState.Viewport.Min.X = 0;
 		PendingState.Viewport.Min.Y = 0;
 
@@ -1577,6 +1577,13 @@ void FOpenGLDynamicRHI::RHISetRenderTargets(
 		PendingState.Viewport.Max.X = PendingState.RenderTargetWidth = Width;
 		PendingState.Viewport.Max.Y = PendingState.RenderTargetHeight = Height;
 	}
+	else if( NewDepthStencilTargetRHI )
+	{
+		// Set viewport size to new depth target size.
+		FOpenGLTexture2D* NewDepthTargetTexture = (FOpenGLTexture2D*)NewDepthStencilTargetRHI->GetTexture2D();
+		PendingState.Viewport.Max.X = NewDepthTargetTexture->GetSizeX();
+		PendingState.Viewport.Max.Y = NewDepthTargetTexture->GetSizeY();
+	}
 }
 
 void FOpenGLDynamicRHI::RHIDiscardRenderTargets(bool Depth, bool Stencil, uint32 ColorBitMask)
diff --git a/Engine/Source/Runtime/Renderer/Private/ShadowRendering.cpp b/Engine/Source/Runtime/Renderer/Private/ShadowRendering.cpp
--- a/Engine/Source/Runtime/Renderer/Private/ShadowRendering.cpp
+++ b/Engine/Source/Runtime/Renderer/Private/ShadowRendering.cpp
@@ -2853,39 +2853,19 @@ bool FDeferredShadingSceneRenderer::RenderProjectedShadows( const FLightSceneInf
 
 			GSceneRenderTargets.BeginRenderingShadowDepth();	
 
-			// keep track of the max RECT needed for resolving the depth surface	
-			FResolveRect ResolveRect;
-			ResolveRect.X1 = 0;
-			ResolveRect.Y1 = 0;
-			bool bResolveRectInit = false;
+			RHIClear(false,FColor(255,255,255),true,1.0f,false,0, FIntRect());
+
 			// render depth for each shadow
 			for (int32 ShadowIndex = 0; ShadowIndex < Shadows.Num(); ShadowIndex++)
 			{
 				FProjectedShadowInfo* ProjectedShadowInfo = Shadows[ShadowIndex];
 				if (ProjectedShadowInfo->bAllocated && !ProjectedShadowInfo->bTranslucentShadow)
 				{
-					ProjectedShadowInfo->ClearDepth(this);
-
 					ProjectedShadowInfo->RenderDepth(this);
-
-					// init values
-					if (!bResolveRectInit)
-					{
-						ResolveRect.X2 = ProjectedShadowInfo->X + ProjectedShadowInfo->ResolutionX + SHADOW_BORDER * 2;
-						ResolveRect.Y2 = ProjectedShadowInfo->Y + ProjectedShadowInfo->ResolutionY + SHADOW_BORDER * 2;
-						bResolveRectInit = true;
-					}
-					// keep track of max extents
-					else 
-					{
-						ResolveRect.X2 = FMath::Max<uint32>(ProjectedShadowInfo->X + ProjectedShadowInfo->ResolutionX + SHADOW_BORDER * 2, ResolveRect.X2);
-						ResolveRect.Y2 = FMath::Max<uint32>(ProjectedShadowInfo->Y + ProjectedShadowInfo->ResolutionY + SHADOW_BORDER * 2, ResolveRect.Y2);
-					}
 				}
 			}
 
-			// only resolve the portion of the shadow buffer that we rendered to
-			GSceneRenderTargets.FinishRenderingShadowDepth(ResolveRect);
+			GSceneRenderTargets.FinishRenderingShadowDepth();
 		}
 
 		// Render the shadow projections.
diff --git a/Engine/Source/Runtime/Windows/D3D11RHI/Private/D3D11Commands.cpp b/Engine/Source/Runtime/Windows/D3D11RHI/Private/D3D11Commands.cpp
--- a/Engine/Source/Runtime/Windows/D3D11RHI/Private/D3D11Commands.cpp
+++ b/Engine/Source/Runtime/Windows/D3D11RHI/Private/D3D11Commands.cpp
@@ -889,6 +889,15 @@ void FD3D11DynamicRHI::RHISetRenderTargets(
 		FRTVDesc RTTDesc = GetRenderTargetViewDesc(NewRenderTargetViews[0]);
 		RHISetViewport(0, 0, 0.0f, RTTDesc.Width, RTTDesc.Height, 1.0f);
 	}
+	else if( DepthStencilView )
+	{
+		TRefCountPtr<ID3D11Texture2D> DepthTargetTexture;
+		DepthStencilView->GetResource((ID3D11Resource**)DepthTargetTexture.GetInitReference());
+
+		D3D11_TEXTURE2D_DESC DTTDesc;
+		DepthTargetTexture->GetDesc(&DTTDesc);
+		RHISetViewport(0, 0, 0.0f, DTTDesc.Width, DTTDesc.Height, 1.0f);
+	}
 }
 
 void FD3D11DynamicRHI::RHIDiscardRenderTargets(bool Depth, bool Stencil, uint32 ColorBitMask)
