Static sky lights do not cause reflection captures to mask out the sky
* This allows the specular contribution of static sky lights to work in the ES2 rendering path

diff --git a/Engine/Shaders/ReflectionEnvironmentShaders.usf b/Engine/Shaders/ReflectionEnvironmentShaders.usf
--- a/Engine/Shaders/ReflectionEnvironmentShaders.usf
+++ b/Engine/Shaders/ReflectionEnvironmentShaders.usf
@@ -59,7 +59,11 @@ void RouteToCubeFaceGS(triangle FCopyToCubeFaceVSOutput Input[3], inout Triangle
 Texture2D InTexture;
 SamplerState InTextureSampler;
 
-/** X = 1 if skylight, 0 otherwise, Y = Sky distance threshold, Z = whether a skylight's lower hemisphere should be black. */
+/** 
+ * X = 0 if capturing sky light, 1 if capturing reflection capture with MaxDistance fade, 2 otherwise, 
+ * Y = Sky distance threshold, 
+ * Z = whether a skylight's lower hemisphere should be black. 
+ */
 float3 SkyLightCaptureParameters;
 
 void CopySceneColorToCubeFaceColorPS(
@@ -76,14 +80,14 @@ void CopySceneColorToCubeFaceColorPS(
 
 	float Alpha = 1;
 
-	if (SkyLightCaptureParameters.x > 0)
+	if (SkyLightCaptureParameters.x == 0)
 	{
 		// Assuming we're on a planet and no sky lighting is coming from below the horizon
 		// This is important to avoid leaking from below since we are integrating incoming lighting and shadowing separately
 		float AboveHorizonMask = Input.Vertex.ScreenVector.z > 0 || SkyLightCaptureParameters.z < 1;
 		SceneColor *= AboveHorizonMask;
 	}
-	else
+	else if (SkyLightCaptureParameters.x == 1)
 	{
 		float RadialDistance = length(WorldPosition - View.ViewOrigin.xyz);
 		float MaxDistance = SkyLightCaptureParameters.y;
diff --git a/Engine/Source/Runtime/Renderer/Private/ReflectionEnvironmentCapture.cpp b/Engine/Source/Runtime/Renderer/Private/ReflectionEnvironmentCapture.cpp
--- a/Engine/Source/Runtime/Renderer/Private/ReflectionEnvironmentCapture.cpp
+++ b/Engine/Source/Runtime/Renderer/Private/ReflectionEnvironmentCapture.cpp
@@ -637,12 +637,24 @@ public:
 			TStaticSamplerState<SF_Bilinear, AM_Clamp, AM_Clamp, AM_Clamp>::GetRHI(), 
 			GSceneRenderTargets.SceneColor->GetRenderTargetItem().ShaderResourceTexture);		
 
-		FVector SkyLightParametersValue(1, 0, bLowerHemisphereIsBlack ? 1.0f : 0.0f);
+		FVector SkyLightParametersValue = FVector::ZeroVector;
 		FScene* Scene = (FScene*)View.Family->Scene;
 
-		if (Scene->SkyLight && !bCapturingForSkyLight)
+		if (bCapturingForSkyLight)
 		{
-			SkyLightParametersValue = FVector(0, Scene->SkyLight->SkyDistanceThreshold, 0);
+			// When capturing reflection captures, support forcing all low hemisphere lighting to be black
+			SkyLightParametersValue = FVector(0, 0, bLowerHemisphereIsBlack ? 1.0f : 0.0f);
+		}
+		else if (Scene->SkyLight && !Scene->SkyLight->bHasStaticLighting)
+		{
+			// When capturing reflection captures and there's a stationary sky light, mask out any pixels whose depth classify it as part of the sky
+			// This will allow changing the stationary sky light at runtime
+			SkyLightParametersValue = FVector(1, Scene->SkyLight->SkyDistanceThreshold, 0);
+		}
+		else
+		{
+			// When capturing reflection captures and there's no sky light, or only a static sky light, capture all depth ranges
+			SkyLightParametersValue = FVector(2, 0, 0);
 		}
 
 		SetShaderValue(ShaderRHI, SkyLightCaptureParameters, SkyLightParametersValue);
