* added annotation to get better PIX captures / ProfileGPU

diff --git a/Engine/Source/Runtime/Renderer/Private/PostProcess/PostProcessBokehDOFRecombine.cpp b/Engine/Source/Runtime/Renderer/Private/PostProcess/PostProcessBokehDOFRecombine.cpp
--- a/Engine/Source/Runtime/Renderer/Private/PostProcess/PostProcessBokehDOFRecombine.cpp
+++ b/Engine/Source/Runtime/Renderer/Private/PostProcess/PostProcessBokehDOFRecombine.cpp
@@ -110,7 +110,25 @@ void FRCPassPostProcessBokehDOFRecombine::SetShader(const FRenderingCompositePas
 
 void FRCPassPostProcessBokehDOFRecombine::Process(FRenderingCompositePassContext& Context)
 {
-	SCOPED_DRAW_EVENT(BokehDOFRecombine, DEC_SCENE_ITEMS);
+	uint32 Method = 2;
+
+	if(GetInput(ePId_Input1)->GetPass())
+	{
+		if(GetInput(ePId_Input2)->GetPass())
+		{
+			Method = 3;
+		}
+		else
+		{
+			Method = 1;
+		}
+	}
+	else
+	{
+		check(GetInput(ePId_Input2)->GetPass());
+	}
+
+	SCOPED_DRAW_EVENTF(BokehDOFRecombine, DEC_SCENE_ITEMS, TEXT("BokehDOFRecombine#%d"), Method);
 
 	const FPooledRenderTargetDesc* InputDesc = GetInputDesc(ePId_Input1);
 	
@@ -138,22 +156,13 @@ void FRCPassPostProcessBokehDOFRecombine::Process(FRenderingCompositePassContext
 	RHISetRasterizerState(TStaticRasterizerState<>::GetRHI());
 	RHISetDepthStencilState(TStaticDepthStencilState<false,CF_Always>::GetRHI());
 
-	if(GetInput(ePId_Input1)->GetPass())
+	switch(Method)
 	{
-		if(GetInput(ePId_Input2)->GetPass())
-		{
-			SetShader<3>(Context);
-		}
-		else
-		{
-			SetShader<1>(Context);
-		}
-	}
-	else
-	{
-		check(GetInput(ePId_Input2)->GetPass());
-
-		SetShader<2>(Context);
+		case 1: SetShader<1>(Context); break;
+		case 2: SetShader<2>(Context); break;
+		case 3: SetShader<3>(Context); break;
+		default:
+			check(0);
 	}
 
 	TShaderMapRef<FPostProcessVS> VertexShader(GetGlobalShaderMap());
