Diffuse from captures.

Reflection environment refactor, optimization and clean up. Fixed blending bug with layered cubemaps. Added diffuse from captures as alternative to lightmaps. Still tied to cvar.

diff --git a/Engine/Shaders/BasePassPixelShader.usf b/Engine/Shaders/BasePassPixelShader.usf
--- a/Engine/Shaders/BasePassPixelShader.usf
+++ b/Engine/Shaders/BasePassPixelShader.usf
@@ -69,29 +69,28 @@ half3 GetImageBasedReflectionLighting(FMaterialPixelParameters MaterialParameter
 #if FEATURE_LEVEL >= FEATURE_LEVEL_SM5
 	half AbsoluteSpecularMip = ComputeReflectionCaptureMipFromRoughness(Roughness);
 	float4 SpecularIBL = TextureCubeArraySampleLevel(ReflectionCubemap, ReflectionCubemapSampler, MaterialParameters.ReflectionVector, CubemapArrayIndex, AbsoluteSpecularMip);
-	
+
 #if ENABLE_SKY_LIGHT
 	BRANCH
 	if (SkyLightParameters.y > 0 && SpecularIBL.a < .001f)
 	{
-		float3 SkyLighting = GetSkyLightReflection(MaterialParameters.ReflectionVector, Roughness);
+		float3 SkyLighting = GetSkyLightReflection(MaterialParameters.ReflectionVector, Roughness, View.UseLightmaps);
 		// Add in sky wherever reflection captures don't have coverage
-		SpecularIBL.rgb = lerp(SkyLighting, SpecularIBL.rgb, SpecularIBL.a);
+		SpecularIBL.rgb += (1 - SpecularIBL.a) * SkyLighting;
 	}
 #endif
-
-	half NoV = saturate(dot(MaterialParameters.WorldNormal, MaterialParameters.CameraVector));
-	half3 Reflectance = EnvBRDFApprox(SpecularColor, Roughness, NoV);
-
-	// Note: make sure this matches the lightmap mixing done on opaque (ReflectionEnvironmentTiledDeferredMain)
-	float MixAlpha = saturate(Roughness * 4 - 3);
-	float3 SmoothColor = Luminance(LowGlossIndirectSpecular) * (1 - MixAlpha);
-	float3 RoughColor = LowGlossIndirectSpecular * MixAlpha;
 	
-	// Combine with brightness from LowGlossIndirectSpecular
-	float3 MixedSpecularIBL = lerp(SpecularIBL.xyz, (SpecularIBL.rgb * SmoothColor + Luminance(SpecularIBL.rgb) * RoughColor), View.ReflectionLightmapMixingMask);
-	return MixedSpecularIBL * Reflectance;
+	BRANCH
+	if( View.UseLightmaps )
+	{
+		// Note: make sure this matches the lightmap mixing done on opaque (ReflectionEnvironmentTiledDeferredMain)
+		SpecularIBL.rgb *= Luminance( LowGlossIndirectSpecular );
+	}
+	
+	half NoV = saturate(dot(MaterialParameters.WorldNormal, MaterialParameters.CameraVector));
+	SpecularColor = EnvBRDFApprox(SpecularColor, Roughness, NoV);
 
+	return SpecularIBL.rgb * SpecularColor;
 #else
 	return 0;
 #endif
@@ -412,7 +411,22 @@ float3 GetPrecomputedIndirectLighting(
 	#endif
 
 	// Sky lighting must contribute to LowGlossIndirectSpecular for ReflectionEnvironment lightmap mixing
-	Lighting += GetSkyLighting(MaterialParameters.WorldNormal, SkyOcclusionUV);
+	float3 SkyLighting = GetSkyLighting(MaterialParameters.WorldNormal, SkyOcclusionUV);
+
+	BRANCH
+	if( View.UseLightmaps )
+	{
+		Lighting += SkyLighting;
+	}
+	else
+	{
+	#if TRANSLUCENCY_LIGHTING_SURFACE && FEATURE_LEVEL >= FEATURE_LEVEL_SM5
+		float DiffuseMip = ComputeReflectionCaptureMipFromRoughness(1);
+		float4 DiffuseIBL = TextureCubeArraySampleLevel(ReflectionCubemap, ReflectionCubemapSampler, MaterialParameters.WorldNormal, CubemapArrayIndex, DiffuseMip);
+		
+		Lighting += DiffuseIBL.rgb + (1 - DiffuseIBL.a) * SkyLighting;
+	#endif
+	}
 
 	#if HQ_TEXTURE_LIGHTMAP || LQ_TEXTURE_LIGHTMAP || CACHED_VOLUME_INDIRECT_LIGHTING || CACHED_POINT_INDIRECT_LIGHTING
 		Lighting *= View.IndirectLightingColorScale;
diff --git a/Engine/Shaders/ReflectionEnvironmentComputeShaders.usf b/Engine/Shaders/ReflectionEnvironmentComputeShaders.usf
--- a/Engine/Shaders/ReflectionEnvironmentComputeShaders.usf
+++ b/Engine/Shaders/ReflectionEnvironmentComputeShaders.usf
@@ -22,7 +22,7 @@ SamplerState ReflectionEnvironmentColorSampler;
 #if GL4_PROFILE
 #define ATOMIC_REDUCTION 0
 #else
-#define ATOMIC_REDUCTION 1
+#define ATOMIC_REDUCTION 0
 #endif
 
 uint NumCaptures;
@@ -290,8 +290,11 @@ void DoTileCulling(uint3 GroupId, uint ThreadIndex, float MinTileZ, float MaxTil
 	GroupMemoryBarrierWithGroupSync();
 }
 
+Texture2D ScreenSpaceReflections;
+Texture2D InSceneColor;
+
 /** Output HDR target. */
-RWTexture2D<float4> RWLightAccumulation;
+RWTexture2D<float4> RWOutSceneColor;
 
 [numthreads(THREADGROUP_SIZEX, THREADGROUP_SIZEY, 1)]
 void ReflectionEnvironmentTiledDeferredMain(
@@ -319,39 +322,34 @@ void ReflectionEnvironmentTiledDeferredMain(
 
 	// Lookup GBuffer properties once per pixel
 	FScreenSpaceData ScreenSpaceData = GetScreenSpaceDataUint(PixelPos);
-	FGBufferData InGBufferData = ScreenSpaceData.GBuffer;
+	FGBufferData GBuffer = ScreenSpaceData.GBuffer;
 	
 	float4 HomogeneousWorldPosition = mul(float4(ScreenPosition * SceneDepth, SceneDepth, 1), View.ScreenToWorld);
 	float3 WorldPosition = HomogeneousWorldPosition.xyz / HomogeneousWorldPosition.w;
 	float3 CameraToPixel = normalize(WorldPosition - View.ViewOrigin.xyz);
-	float3 ReflectionVector = reflect(CameraToPixel, InGBufferData.WorldNormal);
-	float3 SmoothColor = 0;
-	float3 RoughColor = 0;
+	float3 ReflectionVector = reflect(CameraToPixel, GBuffer.WorldNormal);
 	float Overdraw = 0;
 
 	uint NumCapturesAffectingTile = TileNumReflectionCaptures;
-	float4 CompositedLighting = float4(0, 0, 0, 1);
+
+	float4 DiffuseLighting = float4(0, 0, 0, 1);
+	float4 SpecularLighting = float4(0, 0, 0, 1);
 
 	#define VISUALIZE_NUM_CAPTURES 0
 	#if VISUALIZE_NUM_CAPTURES 
 
-		CompositedLighting = NumCapturesAffectingTile / 5.0f;
-		CompositedLighting.a = 0;
+		SpecularLighting = NumCapturesAffectingTile / 5.0f;
+		SpecularLighting.a = 0;
 
 	#else
 
 		BRANCH
 		// Only light pixels marked as using deferred shading
-		if (InGBufferData.LightingModelId > 0)
+		if (GBuffer.LightingModelId > 0)
 		{
-			// Compute GBuffer dependent intermediates for lit pixels
-			// Note: These are all combined into a few temporary registers (SmoothColor, RoughColor) to reduce temporary register pressure during the shading loop
-			// Increased temporary register pressure during shading was causing a performance cliff on AMD 7970
-			GetReflectionSmoothAndRoughColor(ScreenSpaceData, CameraToPixel, SmoothColor, RoughColor);
-
 			// Accumulate reflections from captures affecting this tile, applying largest captures first so that the smallest ones display on top
 			LOOP
-			for (uint TileCaptureIndex = 0; TileCaptureIndex < NumCapturesAffectingTile && CompositedLighting.a > .001f; TileCaptureIndex++) 
+			for (uint TileCaptureIndex = 0; TileCaptureIndex < NumCapturesAffectingTile; TileCaptureIndex++) 
 			{
 				#if SORT_CAPTURES
 					uint CaptureIndex = SortedTileReflectionCaptureIndices[TileCaptureIndex];
@@ -423,40 +421,43 @@ void ReflectionEnvironmentTiledDeferredMain(
 							else 
 						#endif
 						// Box
-						if (CaptureProperties.b > 0)
+						BRANCH if (CaptureProperties.b > 0)
 						{
-							float3 RayDirection = ReflectionVector * CapturePositionAndRadius.w * 2;
+							float3 RayDirection = ReflectionVector;
 
 							// Transform the ray into the local space of the box, where it is an AABB with mins at -1 and maxs at 1
-							float3 LocalRayStart = mul(float4(WorldPosition, 1), ReflectionCapture.BoxTransform[CaptureIndex]).xyz;
-							float3 LocalRayDirection = mul(RayDirection, (float3x3)ReflectionCapture.BoxTransform[CaptureIndex]);
-
-							// Intersections.y is the intersection with the far side of the box
-							float2 Intersections = LineBoxIntersect(LocalRayStart, LocalRayStart + LocalRayDirection, -1, 1);
+							float3 LocalRayStart		= mul(float4(WorldPosition, 1), ReflectionCapture.BoxTransform[CaptureIndex]).xyz;
+							float3 LocalRayDirection	= mul(float4(RayDirection,  0), ReflectionCapture.BoxTransform[CaptureIndex]).xyz;
 
-							//if (Intersections.y > Intersections.x)
-							{
-								// Compute the reprojected vector
-								float3 IntersectPosition = WorldPosition + Intersections.y * RayDirection;
-								ProjectedCaptureVector = IntersectPosition - CapturePositionAndRadius.xyz;
-
-								// Compute the distance from the receiving pixel to the box for masking
-								// Apply local to world scale to take scale into account without transforming back to world space
-								// Shrink the box by the transition distance (BoxScales.w) so that the fade happens inside the box influence area
-								float4 BoxScales = ReflectionCapture.BoxScales[CaptureIndex];
-								float BoxDistance = ComputeDistanceFromBoxToPoint(-(BoxScales.xyz - .5f * BoxScales.w), BoxScales.xyz - .5f * BoxScales.w, LocalRayStart * BoxScales.xyz);
-
-								// Setup a fade based on receiver distance to the box, hides the box influence shape
-								float BoxDistanceAlpha = 1.0 - smoothstep(0, .7f * BoxScales.w, BoxDistance);
-								// Setup a fade based on reflection ray intersection distance, hides the discontinuity between rays that just barely 
-								float RayDistanceAlpha = smoothstep(0, BoxScales.w, Intersections.y * CapturePositionAndRadius.w * 2);
-								DistanceAlpha = BoxDistanceAlpha * RayDistanceAlpha;
-							}
+							float3 InvRayDir = rcp(LocalRayDirection);
+	
+							//find the ray intersection with each of the 3 planes defined by the minimum extrema.
+							float3 FirstPlaneIntersections = -InvRayDir - LocalRayStart * InvRayDir;
+							//find the ray intersection with each of the 3 planes defined by the maximum extrema.
+							float3 SecondPlaneIntersections = InvRayDir - LocalRayStart * InvRayDir;
+							//get the furthest of these intersections along the ray
+							float3 FurthestPlaneIntersections = max(FirstPlaneIntersections, SecondPlaneIntersections);
+
+							//clamp the intersections to be between RayOrigin and RayEnd on the ray
+							float Intersection = min(FurthestPlaneIntersections.x, min(FurthestPlaneIntersections.y, FurthestPlaneIntersections.z));
+
+							// Compute the reprojected vector
+							float3 IntersectPosition = WorldPosition + Intersection * RayDirection;
+							ProjectedCaptureVector = IntersectPosition - CapturePositionAndRadius.xyz;
+
+							// Compute the distance from the receiving pixel to the box for masking
+							// Apply local to world scale to take scale into account without transforming back to world space
+							// Shrink the box by the transition distance (BoxScales.w) so that the fade happens inside the box influence area
+							float4 BoxScales = ReflectionCapture.BoxScales[CaptureIndex];
+							float BoxDistance = ComputeDistanceFromBoxToPoint(-(BoxScales.xyz - .5f * BoxScales.w), BoxScales.xyz - .5f * BoxScales.w, LocalRayStart * BoxScales.xyz);
+
+							// Setup a fade based on receiver distance to the box, hides the box influence shape
+							DistanceAlpha = 1.0 - smoothstep(0, .7f * BoxScales.w, BoxDistance);
 						}
 						// Sphere
 						else
 						{
-							float3 RayDirection = ReflectionVector * CapturePositionAndRadius.w;
+							float3 RayDirection = ReflectionVector;
 							float ProjectionSphereRadius = CapturePositionAndRadius.w * 1.2f;
 							float SphereRadiusSquared = ProjectionSphereRadius * ProjectionSphereRadius;
 
@@ -465,18 +466,16 @@ void ReflectionEnvironmentTiledDeferredMain(
 
 							// Find the intersection between the ray along the reflection vector and the capture's sphere
 							float3 QuadraticCoef;
-							QuadraticCoef.x = dot(RayDirection, RayDirection);
+							QuadraticCoef.x = 1;
 							QuadraticCoef.y = 2 * dot(RayDirection, ReceiverToSphereCenter);
 							QuadraticCoef.z = ReceiverToSphereCenterSq - SphereRadiusSquared;
 
-							float Determinant = QuadraticCoef.y * QuadraticCoef.y - 4 * QuadraticCoef.x * QuadraticCoef.z;
+							float Determinant = QuadraticCoef.y * QuadraticCoef.y - 4 * QuadraticCoef.z;
 							 
 							// Only continue if the ray intersects the sphere
 							if (Determinant >= 0)
-							{ 
-								float Epsilon = .000001f;
-								float InvTwoA = .5 / (QuadraticCoef.x + Epsilon);
-								float FarIntersection = (sqrt(Determinant) - QuadraticCoef.y) * InvTwoA;
+							{
+								float FarIntersection = (sqrt(Determinant) - QuadraticCoef.y) * 0.5;
 
 								float3 IntersectPosition = WorldPosition + FarIntersection * RayDirection;
 								ProjectedCaptureVector = IntersectPosition - CapturePositionAndRadius.xyz;
@@ -490,51 +489,113 @@ void ReflectionEnvironmentTiledDeferredMain(
 					#endif
 
 					float CaptureArrayIndex = CaptureProperties.g;
-					float AbsoluteSpecularMip = ComputeReflectionCaptureMipFromRoughness(InGBufferData.Roughness);
-					float4 Reflection = TextureCubeArraySampleLevel(ReflectionEnvironmentColorTexture, ReflectionEnvironmentColorSampler, ProjectedCaptureVector, CaptureArrayIndex, AbsoluteSpecularMip);
+					
+#if !USE_LIGHTMAPS
+					float DiffuseMip = ComputeReflectionCaptureMipFromRoughness(1);
+					float4 DiffuseIBL = TextureCubeArraySampleLevel(ReflectionEnvironmentColorTexture, ReflectionEnvironmentColorSampler, GBuffer.WorldNormal, CaptureArrayIndex, DiffuseMip);
+
+					DiffuseIBL.rgb *= CaptureProperties.r;
+					DiffuseIBL *= DistanceAlpha;
+
+					// Under operator (back to front)
+					DiffuseLighting.rgb += DiffuseIBL.rgb * DiffuseLighting.a;
+					DiffuseLighting.a *= 1 - DiffuseIBL.a;
+#endif
 
-					Reflection.rgb *= CaptureProperties.r;
+					float SpecularMip = ComputeReflectionCaptureMipFromRoughness(GBuffer.Roughness);
+					float4 SpecularIBL = TextureCubeArraySampleLevel(ReflectionEnvironmentColorTexture, ReflectionEnvironmentColorSampler, ProjectedCaptureVector, CaptureArrayIndex, SpecularMip);
+
+					SpecularIBL.rgb *= CaptureProperties.r;
+					SpecularIBL *= DistanceAlpha;
+
+					// Under operator (back to front)
+					SpecularLighting.rgb += SpecularIBL.rgb * SpecularLighting.a;
+					SpecularLighting.a *= 1 - SpecularIBL.a;
 
-					// Composite using the influence shape's alpha, which allows it to only override where it has valid data
-					CompositedLighting.rgb = lerp(CompositedLighting.rgb, Reflection.rgb, DistanceAlpha * CompositedLighting.a);
-					CompositedLighting.a *= saturate(1 - DistanceAlpha * Reflection.a);
 					Overdraw += .05f;
+
+#if USE_LIGHTMAPS
+					bool Covered = SpecularLighting.a < 0.001;
+#else
+					bool Covered = DiffuseLighting.a + SpecularLighting.a < 0.001;
+#endif
+					BRANCH if( Covered )
+					{
+						break;
+					}
 				}
 			}
 		}
 	#endif
 #else
-	CompositedLighting = abs(MaxTileZ - MinTileZ) * .0001f;
+	SpecularLighting = abs(MaxTileZ - MinTileZ) * .0001f;
 #endif
 
 #define APPLY_SKY_LIGHT 1
 #if APPLY_SKY_LIGHT
-	BRANCH
+
+#if USE_LIGHTMAPS
+	bool Covered = SpecularLighting.a < 0.001;
+#else
+	bool Covered = DiffuseLighting.a + SpecularLighting.a < 0.001;
+#endif
+
 	// Only light pixels marked as using deferred shading
-	if (InGBufferData.LightingModelId > 0 && SkyLightParameters.y > 0 && CompositedLighting.a > .001f)
+	BRANCH if (GBuffer.LightingModelId > 0 && SkyLightParameters.y > 0 && !Covered)
 	{
-		float3 SkyLighting = GetSkyLightReflection(ReflectionVector, InGBufferData.Roughness);
-		CompositedLighting.rgb = lerp(CompositedLighting.rgb, SkyLighting, CompositedLighting.a);
+#if !USE_LIGHTMAPS
+		DiffuseLighting.rgb += DiffuseLighting.a * GetSkyLightReflection(GBuffer.WorldNormal, 1, USE_LIGHTMAPS);
+#endif
+		SpecularLighting.rgb += SpecularLighting.a * GetSkyLightReflection(ReflectionVector, GBuffer.Roughness, USE_LIGHTMAPS);
 		Overdraw += .05f;
 	}
+
 #endif
 
 	// Only write to the buffer for threads inside the view
 	BRANCH
 	if (all(DispatchThreadId.xy < ViewDimensions.zw)) 
 	{
+		float4 OutColor = 0;
+
 		#define VISUALIZE_CAPTURE_OVERDRAW 0
 		#if VISUALIZE_CAPTURE_OVERDRAW
-			CompositedLighting.rgb = Overdraw / 4;
-		#else
-			// Optimization for temporaries
-			//float3 SmoothReflection = CompositedLighting.rgb * Luminance( InGBufferData.LowGlossIndirectSpecular );
-			//float3 RoughReflection = Luminance( CompositedLighting.rgb ) * InGBufferData.LowGlossIndirectSpecular;
-			//CompositedLighting.rgb = lerp( SmoothReflection, RoughReflection, MixAlpha );
-			CompositedLighting.rgb = CompositedLighting.rgb * SmoothColor + Luminance( CompositedLighting.rgb ) * RoughColor;
+			OutColor.rgb = Overdraw / 4;
 		#endif
+		
+		// Save GPRs by using R instead of V
+		float NoV = saturate( dot( GBuffer.WorldNormal, ReflectionVector ) );
+		float3 SpecularColor = EnvBRDF( GBuffer.SpecularColor, GBuffer.Roughness, NoV );
+		//float3 SpecularColor = EnvBRDFApprox( GBuffer.SpecularColor, GBuffer.Roughness, NoV );
+		
+		float AO = GBuffer.GBufferAO * ScreenSpaceData.AmbientOcclusion;
+		float SpecularOcclusion = saturate( Square( NoV + AO ) - 1 + AO );
+		SpecularColor *= SpecularOcclusion;
+
+#if USE_LIGHTMAPS
+		// We have high frequency directional data but low frequency spatial data in the envmap.
+		// We have high frequency spatial data but low frequency directional data in the lightmap.
+		// So, we combine the two for the best of both. This is done by removing the low spatial frequencies from the envmap and replacing them with the lightmap data.
+		// This is only done with luma so as to not get odd color shifting.
+		// Note: make sure this matches the lightmap mixing done for translucency (BasePassPixelShader.usf)
+		SpecularLighting.rgb *= Luminance( GBuffer.LowGlossIndirectSpecular );
+#else
+		OutColor.rgb += DiffuseLighting.rgb * GBuffer.DiffuseColor * AO;
+#endif
+
+#if 1//APPLY_SSR
+		float4 SSR = ScreenSpaceReflections.Load( int3(PixelPos, 0) );
+		SpecularLighting.rgb = SpecularLighting.rgb * (1 - SSR.a) + SSR.rgb;
+#endif
+
+		OutColor.rgb += SpecularLighting.rgb * SpecularColor;
+
+		// Transform NaNs to black, transform negative colors to black.
+		OutColor.rgb = -min(-OutColor.rgb, 0.0);
+
+		OutColor.rgb += InSceneColor.Load( int3(PixelPos, 0) ).rgb;
 
-		RWLightAccumulation[PixelPos.xy] = float4(CompositedLighting.rgb, 0);
+		RWOutSceneColor[PixelPos.xy] = OutColor;
 	}
 }
 
diff --git a/Engine/Shaders/ReflectionEnvironmentShaders.usf b/Engine/Shaders/ReflectionEnvironmentShaders.usf
--- a/Engine/Shaders/ReflectionEnvironmentShaders.usf
+++ b/Engine/Shaders/ReflectionEnvironmentShaders.usf
@@ -10,6 +10,7 @@
 #include "SHCommon.usf"
 #include "ReflectionEnvironmentShared.usf"
 #include "SkyLightShared.usf"
+#include "MonteCarlo.usf"
 
 struct FCopyToCubeFaceVSOutput
 {
@@ -229,12 +230,12 @@ void ComputeBrightnessMain(
 #define NUM_FILTER_SAMPLES 1
 #endif
 
-float4 ConeSamples[NUM_FILTER_SAMPLES];
+float FilterRoughness;
 
 Texture2D AverageBrightnessTexture;
 SamplerState AverageBrightnessSampler;
 
-void BlurPS(
+void FilterPS(
 	FDownsampleGSOut Input,
 	out float4 OutColor : SV_Target0
 	)
@@ -245,21 +246,28 @@ void BlurPS(
 #define FILTER_CUBEMAP 1
 #if FILTER_CUBEMAP
 
-	// Generate a coordinate system, where the Z axis is the direction from the cubemap to the current texel
-	float3 NormalizedCubeCoordinates = normalize(CubeCoordinates);
-	float3 XAxis = normalize(cross(NormalizedCubeCoordinates, float3(0, 1, 0)));
-	float3 YAxis = cross(NormalizedCubeCoordinates, XAxis);
-	float4 FilteredValue = 0;
+	float3 N = normalize(CubeCoordinates);
+	float Roughness = FilterRoughness;
 
-	// Generate samples in a cone and filter
-	for (int SampleIndex = 0; SampleIndex < NUM_FILTER_SAMPLES; SampleIndex++)
+	float4 FilteredColor = 0;
+	float Weight = 0;
+		
+	LOOP
+	for( uint i = 0; i < NUM_FILTER_SAMPLES; i++ )
 	{
-		float3 SampleVector = FilterData.Samples[SampleIndex].x * XAxis + FilterData.Samples[SampleIndex].y * YAxis + FilterData.Samples[SampleIndex].z * NormalizedCubeCoordinates;
-		float4 Sample = SampleCubemap(SampleVector);
-		FilteredValue += Sample;
+		float2 E = Hammersley( i, NUM_FILTER_SAMPLES, 0 );
+		float3 H = TangentToWorld( ImportanceSampleGGX( E, Roughness ).xyz, N );
+		float3 L = 2 * dot( N, H ) * H - N;
+
+		float NoL = saturate( dot( N, L ) );
+		if( NoL > 0 )
+		{
+			FilteredColor += SampleCubemap(L) * NoL;
+			Weight += NoL;
+		}
 	}
 
-	OutColor = FilteredValue / NUM_FILTER_SAMPLES;
+	OutColor = FilteredColor / max( Weight, 0.001 );
 
 #else
 
@@ -267,12 +275,12 @@ void BlurPS(
 
 #endif
 
+#if NORMALIZE
 	// Remove low frequency
 	// TODO match the frequency removed to the one added. In other words divide out 2nd order SH instead of 1st.
-	float LowFrequencyBrightness = Texture2DSample(AverageBrightnessTexture, AverageBrightnessSampler, float2(0, 0));
-
-	// Using too small a minimum value here will cause some GPUs to emit INF on OS X.
+	float LowFrequencyBrightness = Texture2DSample(AverageBrightnessTexture, AverageBrightnessSampler, float2(0, 0)).r;
 	OutColor.rgb /= max(LowFrequencyBrightness, 0.01f);
+#endif
 }
 
 float4 CoefficientMask0;
@@ -388,17 +396,13 @@ void StandardDeferredReflectionPS(
 		float3 CameraToPixel = normalize(WorldPosition - View.ViewOrigin.xyz);
 		float3 ReflectionVector = reflect(CameraToPixel, InGBufferData.WorldNormal);
 
-		float3 SmoothColor = 0;
-		float3 RoughColor = 0;
-		GetReflectionSmoothAndRoughColor(ScreenSpaceData, CameraToPixel, SmoothColor, RoughColor);
-
 		float3 ProjectedCaptureVector = ReflectionVector;
 		float DistanceAlpha = 0;
 
 		//@todo - find a way to share this code with the compute shader version
 
 		#if SPHERE_CAPTURE
-			float3 RayDirection = ReflectionVector * CapturePositionAndRadius.w;
+			float3 RayDirection = ReflectionVector;
 			float ProjectionSphereRadius = CapturePositionAndRadius.w * 1.2f;
 			float SphereRadiusSquared = ProjectionSphereRadius * ProjectionSphereRadius;
 
@@ -411,19 +415,17 @@ void StandardDeferredReflectionPS(
 
 			// Find the intersection between the ray along the reflection vector and the capture's sphere
 			float3 QuadraticCoef;
-			QuadraticCoef.x = dot(RayDirection, RayDirection);
+			QuadraticCoef.x = 1;
 			QuadraticCoef.y = 2 * dot(RayDirection, ReceiverToSphereCenter);
 			QuadraticCoef.z = ReceiverToSphereCenterSq - SphereRadiusSquared;
 
-			float Determinant = QuadraticCoef.y * QuadraticCoef.y - 4 * QuadraticCoef.x * QuadraticCoef.z;
-							 
+			float Determinant = QuadraticCoef.y * QuadraticCoef.y - 4 * QuadraticCoef.z;
+
 			BRANCH
 			// Only continue if the ray intersects the sphere
 			if (Determinant >= 0)
 			{ 
-				float Epsilon = .000001f;
-				float InvTwoA = .5 / (QuadraticCoef.x + Epsilon);
-				float FarIntersection = (sqrt(Determinant) - QuadraticCoef.y) * InvTwoA;
+				float FarIntersection = (sqrt(Determinant) - QuadraticCoef.y) * 0.5;
 
 				float3 IntersectPosition = WorldPosition + FarIntersection * RayDirection;
 				ProjectedCaptureVector = IntersectPosition - CapturePositionAndRadius.xyz;
@@ -469,8 +471,7 @@ void StandardDeferredReflectionPS(
 
 			// Composite using the influence shape's alpha, which allows it to only override where it has valid data
 			// Reflection texture was premultiplied by reflection alpha
-			OutColor.rgb = (Reflection.rgb * SmoothColor + Luminance(Reflection.rgb) * RoughColor) * DistanceAlpha;
-			OutColor.a = saturate(DistanceAlpha * Reflection.a);
+			OutColor = Reflection * DistanceAlpha;
 		}
 	}
 }
@@ -490,63 +491,52 @@ void ReflectionApplyPS(
 
 	FScreenSpaceData ScreenSpaceData = GetScreenSpaceData(UV);
 	FGBufferData GBuffer = ScreenSpaceData.GBuffer;
-		
-	float3 ReflectionEnv = Texture2DSample( ReflectionEnvTexture, ReflectionEnvSampler, UV ).rgb;
-	
-	// The reflection environment can contain NAN/INF on OS X so transform NaNs to black, transform negative colors to black.
-	ReflectionEnv = -min(-ReflectionEnv, 0.0);
-	
-	float4 SSR = Texture2DSample( ScreenSpaceReflectionsTexture, ScreenSpaceReflectionsSampler, UV );
-
-	// TODO if SpecularColor and Roughness were in same texture would remove 2 fetches
-	float3 V = -normalize( mul( float4(UVAndScreenPos.zw, 1, 0), View.ScreenToWorld ).xyz );
-	float NoV = saturate( dot( GBuffer.WorldNormal, V ) );
-	SSR.rgb *= EnvBRDF( GBuffer.SpecularColor, GBuffer.Roughness, NoV );
-	
-	// The Screen Space Reflection can contain NAN/INF on OS X so transform NaNs to black, transform negative colors to black.
-	SSR.rgba = -min(-SSR.rgba, 0.0);
-
-	OutColor.rgb = ReflectionEnv * (1 - SSR.a) + SSR.rgb;
-	OutColor.a = 0;
-}
 
-/** Used to apply sky light specular, with or without SSR. */
-void SkyLightSpecularPS(
-	in float4 UVAndScreenPos : TEXCOORD0, 
-	out float4 OutColor : SV_Target0)
-{
-	float2 UV = UVAndScreenPos.xy;
-
-	FScreenSpaceData ScreenSpaceData = GetScreenSpaceData(UV);
-	FGBufferData InGBufferData = ScreenSpaceData.GBuffer;
+	OutColor = 0;
 	
-	float3 Lighting = 0; 
-
 	BRANCH
 	// Only light pixels marked as lit
-	if (InGBufferData.LightingModelId > 0 && SkyLightParameters.y > 0)
+	if( GBuffer.LightingModelId > 0 )
 	{
-		float3 SmoothColor = 0;
-		float3 RoughColor = 0;
+		// TODO if SpecularColor and Roughness were in same texture would remove 2 fetches
+		float3 V = -normalize( mul( float4(UVAndScreenPos.zw, 1, 0), View.ScreenToWorld ).xyz );
+		float NoV = saturate( dot( GBuffer.WorldNormal, V ) );
+		float3 SpecularColor = EnvBRDF( GBuffer.SpecularColor, GBuffer.Roughness, NoV );
+		
+		float AO = GBuffer.GBufferAO * ScreenSpaceData.AmbientOcclusion;
+		float SpecularOcclusion = saturate( Square( NoV + AO ) - 1 + AO );
+		SpecularColor *= SpecularOcclusion;
 
-		float3 CameraToPixel = normalize(mul(float4(UVAndScreenPos.zw, 1, 0), View.ScreenToWorld).xyz);
-		GetReflectionSmoothAndRoughColor(ScreenSpaceData, CameraToPixel, SmoothColor, RoughColor);
+		float4 SpecularLighting = float4(0, 0, 0, 1);
+		
+#if APPLY_REFLECTION_ENV
+		SpecularLighting = Texture2DSample( ReflectionEnvTexture, ReflectionEnvSampler, UV );
+#endif
 
-		float3 ReflectionVector = reflect(CameraToPixel, InGBufferData.WorldNormal);
-		float3 SkyLighting = GetSkyLightReflection(ReflectionVector, InGBufferData.Roughness);
-		Lighting = SkyLighting * SmoothColor + Luminance(SkyLighting) * RoughColor;
+#if APPLY_SKYLIGHT
+		float3 ReflectionVector = reflect(-V, GBuffer.WorldNormal);
+		SpecularLighting.rgb += SpecularLighting.a * GetSkyLightReflection(ReflectionVector, GBuffer.Roughness, true);
+#endif
+		
+		BRANCH
+		if( View.UseLightmaps )
+		{
+			// We have high frequency directional data but low frequency spatial data in the envmap.
+			// We have high frequency spatial data but low frequency directional data in the lightmap.
+			// So, we combine the two for the best of both. This is done by removing the low spatial frequencies from the envmap and replacing them with the lightmap data.
+			// This is only done with luma so as to not get odd color shifting.
+			// Note: make sure this matches the lightmap mixing done for translucency (BasePassPixelShader.usf)
+			SpecularLighting.rgb *= Luminance( GBuffer.LowGlossIndirectSpecular );
+		}
 
 #if APPLY_SSR
 		float4 SSR = Texture2DSample( ScreenSpaceReflectionsTexture, ScreenSpaceReflectionsSampler, UV );
+		SpecularLighting.rgb = SpecularLighting.rgb * (1 - SSR.a) + SSR.rgb;
+#endif
 
-		// TODO if SpecularColor and Roughness were in same texture would remove 2 fetches
-		float3 V = -CameraToPixel;
-		float NoV = saturate( dot( InGBufferData.WorldNormal, V ) );
-		SSR.rgb *= EnvBRDF(InGBufferData.SpecularColor, InGBufferData.Roughness, NoV);
+		OutColor.rgb = SpecularLighting.rgb * SpecularColor;
 
-		Lighting = Lighting * (1 - SSR.a) + SSR.rgb;
-#endif
+		// Transform NaNs to black, transform negative colors to black.
+		OutColor.rgb = -min(-OutColor.rgb, 0.0);
 	}
-
-	OutColor = float4(Lighting, 1); 
 }
\ No newline at end of file
diff --git a/Engine/Shaders/ReflectionEnvironmentShared.usf b/Engine/Shaders/ReflectionEnvironmentShared.usf
--- a/Engine/Shaders/ReflectionEnvironmentShared.usf
+++ b/Engine/Shaders/ReflectionEnvironmentShared.usf
@@ -28,16 +28,15 @@ SamplerState SkyLightCubemapSampler;
 /** X = max mip, Y = 1 if sky light should be rendered, 0 otherwise. */
 float2 SkyLightParameters;
 
-float3 GetSkyLightReflection(float3 ReflectionVector, float Roughness)
+float3 GetSkyLightReflection(float3 ReflectionVector, float Roughness, bool Normalized)
 {
 	float AbsoluteSpecularMip = ComputeReflectionCaptureMipFromRoughness(Roughness);
 	float3 Reflection = TextureCubeSampleLevel(SkyLightCubemap, SkyLightCubemapSampler, ReflectionVector, AbsoluteSpecularMip).rgb;
 
 	FLATTEN
-	if (View.ReflectionLightmapMixingMask > 0)
+	if (Normalized)
 	{
 		// Sample the lowest resolution mip to get the average color
-		//@todo - put this in a shader constant
 		//@todo - can't normalize sky lighting and reflection capture lighting separately
 		float3 LowFrequencyReflection = TextureCubeSampleLevel(SkyLightCubemap, SkyLightCubemapSampler, ReflectionVector, SkyLightParameters.x).rgb;
 		float LowFrequencyBrightness = Luminance(LowFrequencyReflection);
diff --git a/Engine/Source/Runtime/Engine/Public/SceneManagement.h b/Engine/Source/Runtime/Engine/Public/SceneManagement.h
--- a/Engine/Source/Runtime/Engine/Public/SceneManagement.h
+++ b/Engine/Source/Runtime/Engine/Public/SceneManagement.h
@@ -2083,6 +2083,7 @@ BEGIN_UNIFORM_BUFFER_STRUCT(FViewUniformShaderParameters,ENGINE_API)
 	DECLARE_UNIFORM_BUFFER_STRUCT_MEMBER(float,RealTime)
 	DECLARE_UNIFORM_BUFFER_STRUCT_MEMBER(uint32,Random)
 	DECLARE_UNIFORM_BUFFER_STRUCT_MEMBER(uint32,FrameNumber)
+	DECLARE_UNIFORM_BUFFER_STRUCT_MEMBER(uint32,UseLightmaps)
 	DECLARE_UNIFORM_BUFFER_STRUCT_MEMBER_EX(float,UnlitViewmodeMask, EShaderPrecisionModifier::Half)
 	DECLARE_UNIFORM_BUFFER_STRUCT_MEMBER_EX(float,ReflectionLightmapMixingMask, EShaderPrecisionModifier::Half)
 	DECLARE_UNIFORM_BUFFER_STRUCT_MEMBER_EX(FLinearColor,DirectionalLightColor, EShaderPrecisionModifier::Half)
diff --git a/Engine/Source/Runtime/Renderer/Private/ReflectionEnvironment.cpp b/Engine/Source/Runtime/Renderer/Private/ReflectionEnvironment.cpp
--- a/Engine/Source/Runtime/Renderer/Private/ReflectionEnvironment.cpp
+++ b/Engine/Source/Runtime/Renderer/Private/ReflectionEnvironment.cpp
@@ -22,6 +22,15 @@ const int32 GReflectionEnvironmentTileSizeX = 16;
 const int32 GReflectionEnvironmentTileSizeY = 16;
 extern ENGINE_API int32 GReflectionCaptureSize;
 
+static int32 GDiffuseFromCaptures = 0;
+static FAutoConsoleVariableRef CVarDiffuseFromCaptures(
+	TEXT("r.DiffuseFromCaptures"),
+	GDiffuseFromCaptures,
+	TEXT("Apply indirect diffuse lighting from captures instead of lightmaps.\n")
+	TEXT(" 0 is off (default), 1 is on"),
+	ECVF_RenderThreadSafe
+	);
+
 // to avoid having direct access from many places
 static int GetReflectionEnvironmentCVar()
 {
@@ -154,7 +163,9 @@ public:
 		DeferredParameters.Bind(Initializer.ParameterMap);
 		ReflectionEnvironmentColorTexture.Bind(Initializer.ParameterMap,TEXT("ReflectionEnvironmentColorTexture"));
 		ReflectionEnvironmentColorSampler.Bind(Initializer.ParameterMap,TEXT("ReflectionEnvironmentColorSampler"));
-		LightAccumulation.Bind(Initializer.ParameterMap, TEXT("LightAccumulation"));
+		ScreenSpaceReflections.Bind(Initializer.ParameterMap, TEXT("ScreenSpaceReflections"));
+		InSceneColor.Bind(Initializer.ParameterMap, TEXT("InSceneColor"));
+		OutSceneColor.Bind(Initializer.ParameterMap, TEXT("OutSceneColor"));
 		NumCaptures.Bind(Initializer.ParameterMap, TEXT("NumCaptures"));
 		ViewDimensionsParameter.Bind(Initializer.ParameterMap, TEXT("ViewDimensions"));
 		PreIntegratedGF.Bind(Initializer.ParameterMap, TEXT("PreIntegratedGF"));
@@ -166,7 +177,7 @@ public:
 	{
 	}
 
-	void SetParameters(const FSceneView& View)
+	void SetParameters(const FSceneView& View, FTextureRHIParamRef SSRTexture, FUnorderedAccessViewRHIParamRef OutSceneColorUAV)
 	{
 		const FComputeShaderRHIParamRef ShaderRHI = GetComputeShader();
 
@@ -184,7 +195,10 @@ public:
 			// Set NULL when not valid to prevent d3d debug warnings
 			Scene->ReflectionSceneData.CubemapArray.GetRenderTarget().IsValid() ? CubemapArray.ShaderResourceTexture : NULL);
 
-		LightAccumulation.SetTexture(ShaderRHI, GSceneRenderTargets.GetLightAccumulationTexture(), GSceneRenderTargets.LightAccumulationUAV);
+		SetTextureParameter( ShaderRHI, ScreenSpaceReflections, SSRTexture );
+
+		SetTextureParameter( ShaderRHI, InSceneColor, GSceneRenderTargets.SceneColor->GetRenderTargetItem().ShaderResourceTexture );
+		OutSceneColor.SetTexture(ShaderRHI, NULL, OutSceneColorUAV);
 
 		SetShaderValue(ShaderRHI, ViewDimensionsParameter, View.ViewRect);
 
@@ -258,7 +272,7 @@ public:
 	{
 		const FComputeShaderRHIParamRef ShaderRHI = GetComputeShader();
 
-		LightAccumulation.UnsetUAV(ShaderRHI);
+		OutSceneColor.UnsetUAV(ShaderRHI);
 	}
 
 	virtual bool Serialize(FArchive& Ar)
@@ -267,7 +281,9 @@ public:
 		Ar << DeferredParameters;
 		Ar << ReflectionEnvironmentColorTexture;
 		Ar << ReflectionEnvironmentColorSampler;
-		Ar << LightAccumulation;
+		Ar << ScreenSpaceReflections;
+		Ar << InSceneColor;
+		Ar << OutSceneColor;
 		Ar << NumCaptures;
 		Ar << ViewDimensionsParameter;
 		Ar << PreIntegratedGF;
@@ -281,7 +297,9 @@ private:
 	FDeferredPixelShaderParameters DeferredParameters;
 	FShaderResourceParameter ReflectionEnvironmentColorTexture;
 	FShaderResourceParameter ReflectionEnvironmentColorSampler;
-	FRWShaderParameter LightAccumulation;
+	FShaderResourceParameter ScreenSpaceReflections;
+	FShaderResourceParameter InSceneColor;
+	FRWShaderParameter OutSceneColor;
 	FShaderParameter NumCaptures;
 	FShaderParameter ViewDimensionsParameter;
 	FShaderResourceParameter PreIntegratedGF;
@@ -289,9 +307,29 @@ private:
 	FSkyLightReflectionParameters SkyLightParameters;
 };
 
-IMPLEMENT_SHADER_TYPE(,FReflectionEnvironmentTiledDeferredCS,TEXT("ReflectionEnvironmentComputeShaders"),TEXT("ReflectionEnvironmentTiledDeferredMain"),SF_Compute);
+template< uint32 bUseLightmaps >
+class TReflectionEnvironmentTiledDeferredCS : public FReflectionEnvironmentTiledDeferredCS
+{
+	DECLARE_SHADER_TYPE(TReflectionEnvironmentTiledDeferredCS, Global);
+
+	/** Default constructor. */
+	TReflectionEnvironmentTiledDeferredCS() {}
+public:
+	TReflectionEnvironmentTiledDeferredCS(const ShaderMetaType::CompiledShaderInitializerType& Initializer)
+	: FReflectionEnvironmentTiledDeferredCS(Initializer)
+	{}
+
+	static void ModifyCompilationEnvironment(EShaderPlatform Platform, FShaderCompilerEnvironment& OutEnvironment)
+	{
+		FReflectionEnvironmentTiledDeferredCS::ModifyCompilationEnvironment(Platform, OutEnvironment);
+		OutEnvironment.SetDefine(TEXT("USE_LIGHTMAPS"), bUseLightmaps);
+	}
+};
 
+IMPLEMENT_SHADER_TYPE(template<>,TReflectionEnvironmentTiledDeferredCS<0>,TEXT("ReflectionEnvironmentComputeShaders"),TEXT("ReflectionEnvironmentTiledDeferredMain"),SF_Compute);
+IMPLEMENT_SHADER_TYPE(template<>,TReflectionEnvironmentTiledDeferredCS<1>,TEXT("ReflectionEnvironmentComputeShaders"),TEXT("ReflectionEnvironmentTiledDeferredMain"),SF_Compute);
 
+template< uint32 bSSR, uint32 bReflectionEnv, uint32 bSkylight >
 class FReflectionApplyPS : public FGlobalShader
 {
 	DECLARE_SHADER_TYPE(FReflectionApplyPS, Global);
@@ -301,11 +339,20 @@ class FReflectionApplyPS : public FGlobalShader
 		return IsFeatureLevelSupported(Platform, ERHIFeatureLevel::SM4);
 	}
 
+	static void ModifyCompilationEnvironment(EShaderPlatform Platform, FShaderCompilerEnvironment& OutEnvironment)
+	{
+		FGlobalShader::ModifyCompilationEnvironment(Platform, OutEnvironment);
+		OutEnvironment.SetDefine(TEXT("APPLY_SSR"), bSSR);
+		OutEnvironment.SetDefine(TEXT("APPLY_REFLECTION_ENV"), bReflectionEnv);
+		OutEnvironment.SetDefine(TEXT("APPLY_SKYLIGHT"), bSkylight);
+	}
+
 	/** Default constructor. */
 	FReflectionApplyPS() {}
 
 public:
 	FDeferredPixelShaderParameters DeferredParameters;
+	FSkyLightReflectionParameters SkyLightParameters;
 	FShaderResourceParameter ReflectionEnvTexture;
 	FShaderResourceParameter ReflectionEnvSampler;
 	FShaderResourceParameter ScreenSpaceReflectionsTexture;
@@ -318,6 +365,7 @@ public:
 		: FGlobalShader(Initializer)
 	{
 		DeferredParameters.Bind(Initializer.ParameterMap);
+		SkyLightParameters.Bind(Initializer.ParameterMap);
 		ReflectionEnvTexture.Bind(Initializer.ParameterMap,TEXT("ReflectionEnvTexture"));
 		ReflectionEnvSampler.Bind(Initializer.ParameterMap,TEXT("ReflectionEnvSampler"));
 		ScreenSpaceReflectionsTexture.Bind(Initializer.ParameterMap,TEXT("ScreenSpaceReflectionsTexture"));
@@ -332,6 +380,7 @@ public:
 
 		FGlobalShader::SetParameters(ShaderRHI, View);
 		DeferredParameters.Set(ShaderRHI, View);
+		SkyLightParameters.SetParameters(ShaderRHI, (FScene*)View.Family->Scene, true);
 		
 		SetTextureParameter( ShaderRHI, ReflectionEnvTexture, ReflectionEnvSampler, TStaticSamplerState<SF_Point>::GetRHI(), ReflectionEnv );
 		SetTextureParameter( ShaderRHI, ScreenSpaceReflectionsTexture, ScreenSpaceReflectionsSampler, TStaticSamplerState<SF_Point>::GetRHI(), ScreenSpaceReflections );
@@ -343,6 +392,7 @@ public:
 	{
 		bool bShaderHasOutdatedParameters = FGlobalShader::Serialize(Ar);
 		Ar << DeferredParameters;
+		Ar << SkyLightParameters;
 		Ar << ReflectionEnvTexture;
 		Ar << ReflectionEnvSampler;
 		Ar << ScreenSpaceReflectionsTexture;
@@ -353,7 +403,19 @@ public:
 	}
 };
 
-IMPLEMENT_SHADER_TYPE(,FReflectionApplyPS,TEXT("ReflectionEnvironmentShaders"),TEXT("ReflectionApplyPS"),SF_Pixel);
+// Typedef is necessary because the C preprocessor thinks the comma in the template parameter list is a comma in the macro parameter list.
+#define IMPLEMENT_REFLECTION_APPLY_PIXELSHADER_TYPE(A, B, C) \
+	typedef FReflectionApplyPS<A,B,C> FReflectionApplyPS##A##B##C; \
+	IMPLEMENT_SHADER_TYPE(template<>,FReflectionApplyPS##A##B##C,TEXT("ReflectionEnvironmentShaders"),TEXT("ReflectionApplyPS"),SF_Pixel);
+
+IMPLEMENT_REFLECTION_APPLY_PIXELSHADER_TYPE(0,0,0);
+IMPLEMENT_REFLECTION_APPLY_PIXELSHADER_TYPE(0,0,1);
+IMPLEMENT_REFLECTION_APPLY_PIXELSHADER_TYPE(0,1,0);
+IMPLEMENT_REFLECTION_APPLY_PIXELSHADER_TYPE(0,1,1);
+IMPLEMENT_REFLECTION_APPLY_PIXELSHADER_TYPE(1,0,0);
+IMPLEMENT_REFLECTION_APPLY_PIXELSHADER_TYPE(1,0,1);
+IMPLEMENT_REFLECTION_APPLY_PIXELSHADER_TYPE(1,1,0);
+IMPLEMENT_REFLECTION_APPLY_PIXELSHADER_TYPE(1,1,1);
 
 
 class FReflectionCaptureSpecularBouncePS : public FGlobalShader
@@ -429,8 +491,6 @@ public:
 		CaptureBoxScales.Bind(Initializer.ParameterMap, TEXT("CaptureBoxScales"));
 		ReflectionEnvironmentColorTexture.Bind(Initializer.ParameterMap, TEXT("ReflectionEnvironmentColorTexture"));
 		ReflectionEnvironmentColorSampler.Bind(Initializer.ParameterMap, TEXT("ReflectionEnvironmentColorSampler"));
-		PreIntegratedGF.Bind(Initializer.ParameterMap, TEXT("PreIntegratedGF"));
-		PreIntegratedGFSampler.Bind(Initializer.ParameterMap, TEXT("PreIntegratedGFSampler"));
 		DeferredParameters.Bind(Initializer.ParameterMap);
 	}
 
@@ -439,7 +499,6 @@ public:
 		const FPixelShaderRHIParamRef ShaderRHI = GetPixelShader();
 		FGlobalShader::SetParameters(ShaderRHI, View);
 
-		SetTextureParameter(ShaderRHI, PreIntegratedGF, PreIntegratedGFSampler, TStaticSamplerState<SF_Bilinear,AM_Clamp,AM_Clamp,AM_Clamp>::GetRHI(), GSystemTextures.PreintegratedGF->GetRenderTargetItem().ShaderResourceTexture);
 		SetTextureParameter(ShaderRHI, ReflectionEnvironmentColorTexture, ReflectionEnvironmentColorSampler, TStaticSamplerState<SF_Trilinear,AM_Clamp,AM_Clamp,AM_Clamp>::GetRHI(), SortData.SM4FullHDRCubemap->TextureRHI);
 
 		DeferredParameters.Set(ShaderRHI, View);
@@ -459,8 +518,6 @@ public:
 		Ar << CaptureBoxScales;
 		Ar << ReflectionEnvironmentColorTexture;
 		Ar << ReflectionEnvironmentColorSampler;
-		Ar << PreIntegratedGF;
-		Ar << PreIntegratedGFSampler;
 		Ar << DeferredParameters;
 		return bShaderHasOutdatedParameters;
 	}
@@ -473,172 +530,141 @@ private:
 	FShaderParameter CaptureBoxScales;
 	FShaderResourceParameter ReflectionEnvironmentColorTexture;
 	FShaderResourceParameter ReflectionEnvironmentColorSampler;
-	FShaderResourceParameter PreIntegratedGF;
-	FShaderResourceParameter PreIntegratedGFSampler;
 	FDeferredPixelShaderParameters DeferredParameters;
 };
 
 IMPLEMENT_SHADER_TYPE(template<>,TStandardDeferredReflectionPS<true>,TEXT("ReflectionEnvironmentShaders"),TEXT("StandardDeferredReflectionPS"),SF_Pixel);
 IMPLEMENT_SHADER_TYPE(template<>,TStandardDeferredReflectionPS<false>,TEXT("ReflectionEnvironmentShaders"),TEXT("StandardDeferredReflectionPS"),SF_Pixel);
 
-template<bool bApplySSR>
-class TSkyLightSpecularPS : public FGlobalShader
-{
-	DECLARE_SHADER_TYPE(TSkyLightSpecularPS, Global);
-public:
 
-	static bool ShouldCache(EShaderPlatform Platform)
+/** Renders deferred reflections for the scene. */
+void FDeferredShadingSceneRenderer::RenderDeferredReflections()
+{
+	if (IsSimpleDynamicLightingEnabled() || ViewFamily.EngineShowFlags.VisualizeLightCulling)
 	{
-		return IsFeatureLevelSupported(Platform, ERHIFeatureLevel::SM4);
+		return;
 	}
 
-	static void ModifyCompilationEnvironment(EShaderPlatform Platform, FShaderCompilerEnvironment& OutEnvironment)
+	bool bAnyViewIsReflectionCapture = false;
+	for (int32 ViewIndex = 0; ViewIndex < Views.Num(); ViewIndex++)
 	{
-		FGlobalShader::ModifyCompilationEnvironment(Platform, OutEnvironment);
-		OutEnvironment.SetDefine(TEXT("APPLY_SSR"), (uint32)bApplySSR);
+		const FViewInfo& View = Views[ViewIndex];
+		bAnyViewIsReflectionCapture = bAnyViewIsReflectionCapture || View.bIsReflectionCapture;
 	}
 
-	/** Default constructor. */
-	TSkyLightSpecularPS() {}
+	const uint32 bReflectionEnv = IsReflectionEnvironmentAvailable()
+		&& Scene->ReflectionSceneData.RegisteredReflectionCaptures.Num()
+		&& ViewFamily.EngineShowFlags.ReflectionEnvironment;
 
-	/** Initialization constructor. */
-	TSkyLightSpecularPS(const ShaderMetaType::CompiledShaderInitializerType& Initializer)
-		: FGlobalShader(Initializer)
-	{
-		PreIntegratedGF.Bind(Initializer.ParameterMap, TEXT("PreIntegratedGF"));
-		PreIntegratedGFSampler.Bind(Initializer.ParameterMap, TEXT("PreIntegratedGFSampler"));
-		SkyLightParameters.Bind(Initializer.ParameterMap);
-		DeferredParameters.Bind(Initializer.ParameterMap);
-		ScreenSpaceReflectionsTexture.Bind(Initializer.ParameterMap,TEXT("ScreenSpaceReflectionsTexture"));
-		ScreenSpaceReflectionsSampler.Bind(Initializer.ParameterMap,TEXT("ScreenSpaceReflectionsSampler"));
-	}
-
-	void SetParameters(const FSceneView& View, FTextureRHIParamRef ScreenSpaceReflections)
-	{
-		const FPixelShaderRHIParamRef ShaderRHI = GetPixelShader();
+	const uint32 bSkyLight = Scene->SkyLight
+		&& Scene->SkyLight->ProcessedTexture
+		&& ViewFamily.EngineShowFlags.SkyLighting;
 
-		FGlobalShader::SetParameters(ShaderRHI, View);
-		
-		SetTextureParameter( ShaderRHI, PreIntegratedGF, PreIntegratedGFSampler, TStaticSamplerState<SF_Bilinear,AM_Clamp,AM_Clamp,AM_Clamp>::GetRHI(), GSystemTextures.PreintegratedGF->GetRenderTargetItem().ShaderResourceTexture );
+	const uint32 bUseLightmaps = GDiffuseFromCaptures == 0;
 
-		SkyLightParameters.SetParameters(ShaderRHI, (FScene*)View.Family->Scene, true);
-		DeferredParameters.Set(ShaderRHI, View);
+	const bool bReflectionsWithCompute = GRHIFeatureLevel == ERHIFeatureLevel::SM5 && Scene->ReflectionSceneData.CubemapArray.IsValid();
 
-		SetTextureParameter( ShaderRHI, ScreenSpaceReflectionsTexture, ScreenSpaceReflectionsSampler, TStaticSamplerState<SF_Point>::GetRHI(), ScreenSpaceReflections );
-	}
-
-	// FShader interface.
-	virtual bool Serialize(FArchive& Ar)
+	TRefCountPtr<IPooledRenderTarget>	NewSceneColor;
+	FUnorderedAccessViewRHIRef			NewSceneColorUAV;
+	if( bReflectionEnv && bReflectionsWithCompute && !bAnyViewIsReflectionCapture )
 	{
-		bool bShaderHasOutdatedParameters = FGlobalShader::Serialize(Ar);
-		Ar << PreIntegratedGF;
-		Ar << PreIntegratedGFSampler;
-		Ar << SkyLightParameters;
-		Ar << DeferredParameters;
-		Ar << ScreenSpaceReflectionsTexture;
-		Ar << ScreenSpaceReflectionsSampler;
-		return bShaderHasOutdatedParameters;
-	}
-
-private:
-
-	FShaderResourceParameter PreIntegratedGF;
-	FShaderResourceParameter PreIntegratedGFSampler;
-	FSkyLightReflectionParameters SkyLightParameters;
-	FDeferredPixelShaderParameters DeferredParameters;
-	FShaderResourceParameter ScreenSpaceReflectionsTexture;
-	FShaderResourceParameter ScreenSpaceReflectionsSampler;
-};
+		GSceneRenderTargets.ResolveSceneColor(FResolveRect(0, 0, ViewFamily.FamilySizeX, ViewFamily.FamilySizeY));
 
-IMPLEMENT_SHADER_TYPE(template<>,TSkyLightSpecularPS<true>,TEXT("ReflectionEnvironmentShaders"),TEXT("SkyLightSpecularPS"),SF_Pixel);
-IMPLEMENT_SHADER_TYPE(template<>,TSkyLightSpecularPS<false>,TEXT("ReflectionEnvironmentShaders"),TEXT("SkyLightSpecularPS"),SF_Pixel);
+		FPooledRenderTargetDesc Desc = GSceneRenderTargets.SceneColor->GetDesc();
+		Desc.Flags |= TexCreate_UAV;
 
-FGlobalBoundShaderState SkySpecularBoundShaderState;
+		GRenderTargetPool.FindFreeElement( Desc, NewSceneColor, TEXT("SceneColor") );
+		NewSceneColorUAV = RHICreateUnorderedAccessView( NewSceneColor->GetRenderTargetItem().TargetableTexture );
+	}
 
-/** Renders deferred reflections for the scene. */
-void FDeferredShadingSceneRenderer::RenderDeferredReflections()
-{
-	if (!IsSimpleDynamicLightingEnabled() && !ViewFamily.EngineShowFlags.VisualizeLightCulling)
+	for (int32 ViewIndex = 0; ViewIndex < Views.Num(); ViewIndex++)
 	{
-		// todo: support multiple views
-		const bool bSSR = DoScreenSpaceReflections(Views[0]);
+		const FViewInfo& View = Views[ViewIndex];
 
-		if (IsReflectionEnvironmentAvailable()
-			&& Scene->ReflectionSceneData.RegisteredReflectionCaptures.Num()
-			&& ViewFamily.EngineShowFlags.ReflectionEnvironment)
+		if (bAnyViewIsReflectionCapture)
+		{
+			// We're currently capturing a reflection capture, output SpecularColor * IndirectDiffuseGI for metals so they are not black in reflections,
+			// Since we don't have multiple bounce specular reflections
+			GSceneRenderTargets.BeginRenderingSceneColor();
+
+			RHISetViewport( View.ViewRect.Min.X, View.ViewRect.Min.Y, 0.0f, View.ViewRect.Max.X, View.ViewRect.Max.Y, 1.0f );
+			RHISetRasterizerState( TStaticRasterizerState< FM_Solid, CM_None >::GetRHI() );
+			RHISetDepthStencilState( TStaticDepthStencilState< false, CF_Always >::GetRHI() );
+			RHISetBlendState( TStaticBlendState< CW_RGB, BO_Add, BF_One, BF_One >::GetRHI() );
+
+			TShaderMapRef< FPostProcessVS > VertexShader( GetGlobalShaderMap() );
+			TShaderMapRef< FReflectionCaptureSpecularBouncePS > PixelShader( GetGlobalShaderMap() );
+
+			static FGlobalBoundShaderState BoundShaderState;
+			SetGlobalBoundShaderState( BoundShaderState, GFilterVertexDeclaration.VertexDeclarationRHI, *VertexShader, *PixelShader );
+
+			PixelShader->SetParameters( View );
+
+			DrawRectangle( 
+				0, 0, 
+				View.ViewRect.Width(), View.ViewRect.Height(),
+				0, 0, 
+				View.ViewRect.Width(), View.ViewRect.Height(),
+				FIntPoint(View.ViewRect.Width(), View.ViewRect.Height()),
+				GSceneRenderTargets.GetBufferSizeXY(),
+				*VertexShader,
+				EDRF_UseTriangleOptimization);
+		}
+		else
 		{
-			bool bAnyViewIsReflectionCapture = false;
+			const uint32 bSSR = DoScreenSpaceReflections(Views[0]);
 
-			for (int32 ViewIndex = 0; ViewIndex < Views.Num(); ViewIndex++)
-			{
-				const FViewInfo& View = Views[ViewIndex];
-				bAnyViewIsReflectionCapture = bAnyViewIsReflectionCapture || View.bIsReflectionCapture;
-			}
+			bool bRequiresApply = bSSR || bReflectionEnv || bSkyLight;
 
-			if (bAnyViewIsReflectionCapture)
+			TRefCountPtr<IPooledRenderTarget> SSROutput = GSystemTextures.BlackDummy;
+			if( bSSR )
 			{
-				// We're currently capturing a reflection capture, output SpecularColor * IndirectDiffuseGI for metals so they are not black in reflections,
-				// Since we don't have multiple bounce specular reflections
-				GSceneRenderTargets.BeginRenderingSceneColor();
-
-				for( int32 ViewIndex = 0; ViewIndex < Views.Num(); ViewIndex++ )
-				{
-					const FViewInfo& View = Views[ViewIndex];
-
-					RHISetViewport( View.ViewRect.Min.X, View.ViewRect.Min.Y, 0.0f, View.ViewRect.Max.X, View.ViewRect.Max.Y, 1.0f );
-					RHISetRasterizerState( TStaticRasterizerState< FM_Solid, CM_None >::GetRHI() );
-					RHISetDepthStencilState( TStaticDepthStencilState< false, CF_Always >::GetRHI() );
-					RHISetBlendState( TStaticBlendState< CW_RGB, BO_Add, BF_One, BF_One >::GetRHI() );
-
-					TShaderMapRef< FPostProcessVS > VertexShader( GetGlobalShaderMap() );
-					TShaderMapRef< FReflectionCaptureSpecularBouncePS > PixelShader( GetGlobalShaderMap() );
-
-					static FGlobalBoundShaderState BoundShaderState;
-					SetGlobalBoundShaderState( BoundShaderState, GFilterVertexDeclaration.VertexDeclarationRHI, *VertexShader, *PixelShader );
-
-					PixelShader->SetParameters( View );
-
-					DrawRectangle( 
-						0, 0, 
-						View.ViewRect.Width(), View.ViewRect.Height(),
-						0, 0, 
-						View.ViewRect.Width(), View.ViewRect.Height(),
-						FIntPoint(View.ViewRect.Width(), View.ViewRect.Height()),
-						GSceneRenderTargets.GetBufferSizeXY(),
-						*VertexShader,
-						EDRF_UseTriangleOptimization);
-				}
+				ScreenSpaceReflections( View, SSROutput );
 			}
-			else
+
+			if( bReflectionEnv )
 			{
 				// If we are in SM5, use the compute shader gather method
-				if (GRHIFeatureLevel == ERHIFeatureLevel::SM5
-					&& Scene->ReflectionSceneData.CubemapArray.IsValid())
+				if (GRHIFeatureLevel == ERHIFeatureLevel::SM5 && Scene->ReflectionSceneData.CubemapArray.IsValid())
 				{
 					// Render the reflection environment with tiled deferred culling
 					SCOPED_DRAW_EVENT(ReflectionEnvironmentGather, DEC_SCENE_ITEMS);
-
+				
 					RHISetRenderTarget(NULL, NULL);
 
-					for (int32 ViewIndex = 0; ViewIndex < Views.Num(); ViewIndex++)
+					FReflectionEnvironmentTiledDeferredCS* ComputeShader = NULL;
+					if( bUseLightmaps )
+					{
+						ComputeShader = *TShaderMapRef< TReflectionEnvironmentTiledDeferredCS<1> >( GetGlobalShaderMap() );
+					}
+					else
 					{
-						const FViewInfo& View = Views[ViewIndex];
+						ComputeShader = *TShaderMapRef< TReflectionEnvironmentTiledDeferredCS<0> >( GetGlobalShaderMap() );
+					}
 
-						TShaderMapRef< FReflectionEnvironmentTiledDeferredCS> ComputeShader(GetGlobalShaderMap());
-						RHISetComputeShader(ComputeShader->GetComputeShader());
+					RHISetComputeShader(ComputeShader->GetComputeShader());
 
-						ComputeShader->SetParameters(View);
+					ComputeShader->SetParameters(View, SSROutput->GetRenderTargetItem().ShaderResourceTexture, NewSceneColorUAV);
+					//ComputeShader->SetParameters(View, DestRenderTarget.UAV);
 
-						uint32 GroupSizeX = (View.ViewRect.Size().X + GReflectionEnvironmentTileSizeX - 1) / GReflectionEnvironmentTileSizeX;
-						uint32 GroupSizeY = (View.ViewRect.Size().Y + GReflectionEnvironmentTileSizeY - 1) / GReflectionEnvironmentTileSizeY;
-						DispatchComputeShader(*ComputeShader, GroupSizeX, GroupSizeY, 1);
+					uint32 GroupSizeX = (View.ViewRect.Size().X + GReflectionEnvironmentTileSizeX - 1) / GReflectionEnvironmentTileSizeX;
+					uint32 GroupSizeY = (View.ViewRect.Size().Y + GReflectionEnvironmentTileSizeY - 1) / GReflectionEnvironmentTileSizeY;
+					DispatchComputeShader(ComputeShader, GroupSizeX, GroupSizeY, 1);
 
-						ComputeShader->UnsetParameters();
-					}
+					ComputeShader->UnsetParameters();
+
+					const FSceneRenderTargetItem& DestRenderTarget = NewSceneColor->GetRenderTargetItem();
+					RHICopyToResolveTarget( DestRenderTarget.TargetableTexture, DestRenderTarget.ShaderResourceTexture, false, FResolveParams() );
+
+					GSceneRenderTargets.SceneColor = NewSceneColor;
+					check(GSceneRenderTargets.SceneColor);
+
+					bRequiresApply = false;
 				}
 				// In SM4 use standard deferred shading to composite reflection capture contribution
 				else if (GRHIFeatureLevel == ERHIFeatureLevel::SM4)
 				{
+					// TODO sort once per view
 					static TArray<FReflectionCaptureSortData> SortData;
 					SortData.Reset(Scene->ReflectionSceneData.RegisteredReflectionCaptures.Num());
 
@@ -676,194 +702,103 @@ void FDeferredShadingSceneRenderer::RenderDeferredReflections()
 
 					SortData.Sort();
 
-					{
-						RHISetRenderTarget(GSceneRenderTargets.LightAccumulation->GetRenderTargetItem().TargetableTexture, NULL);
 
-						// Clear to no reflection contribution, alpha of 1 indicates full background contribution
-						RHIClear(true, FLinearColor(0, 0, 0, 1), false, 0, false, 0, FIntRect());
+					SCOPED_DRAW_EVENT(StandardDeferredReflectionEnvironment, DEC_SCENE_ITEMS);
+				
+					RHISetRenderTarget(GSceneRenderTargets.LightAccumulation->GetRenderTargetItem().TargetableTexture, NULL);
 
-						for (int32 ViewIndex = 0; ViewIndex < Views.Num(); ViewIndex++)
-						{
-							FViewInfo& View = Views[ViewIndex];
+					// Clear to no reflection contribution, alpha of 1 indicates full background contribution
+					RHIClear(true, FLinearColor(0, 0, 0, 1), false, 0, false, 0, View.ViewRect);
 
-							{
-								SCOPED_DRAW_EVENT(StandardDeferredReflectionEnvironment, DEC_SCENE_ITEMS);
-								RHISetViewport(View.ViewRect.Min.X, View.ViewRect.Min.Y, 0.0f, View.ViewRect.Max.X, View.ViewRect.Max.Y, 1.0f);
+					RHISetViewport(View.ViewRect.Min.X, View.ViewRect.Min.Y, 0.0f, View.ViewRect.Max.X, View.ViewRect.Max.Y, 1.0f);
 
-								// rgb accumulates reflection contribution front to back, alpha accumulates (1 - alpha0) * (1 - alpha 1)...
-								RHISetBlendState(TStaticBlendState<CW_RGBA, BO_Add, BF_DestAlpha, BF_One, BO_Add, BF_Zero, BF_InverseSourceAlpha>::GetRHI());
+					// rgb accumulates reflection contribution front to back, alpha accumulates (1 - alpha0) * (1 - alpha 1)...
+					RHISetBlendState(TStaticBlendState<CW_RGBA, BO_Add, BF_DestAlpha, BF_One, BO_Add, BF_Zero, BF_InverseSourceAlpha>::GetRHI());
 								 
-								for (int32 ReflectionCaptureIndex = 0; ReflectionCaptureIndex < SortData.Num(); ReflectionCaptureIndex++)
-								{
-									const FReflectionCaptureSortData& ReflectionCapture = SortData[ReflectionCaptureIndex];
-
-									if (ReflectionCapture.SM4FullHDRCubemap)
-									{
-										const FSphere LightBounds(ReflectionCapture.PositionAndRadius, ReflectionCapture.PositionAndRadius.W);
-
-										TShaderMapRef<TDeferredLightVS<true> > VertexShader(GetGlobalShaderMap());
-
-										// Use the appropriate shader for the capture shape
-										if (ReflectionCapture.CaptureProperties.Z == 0)
-										{
-											TShaderMapRef<TStandardDeferredReflectionPS<true> > PixelShader(GetGlobalShaderMap());
+					for (int32 ReflectionCaptureIndex = 0; ReflectionCaptureIndex < SortData.Num(); ReflectionCaptureIndex++)
+					{
+						const FReflectionCaptureSortData& ReflectionCapture = SortData[ReflectionCaptureIndex];
 
-											static FGlobalBoundShaderState BoundShaderState;
-											SetGlobalBoundShaderState(BoundShaderState, GFilterVertexDeclaration.VertexDeclarationRHI, *VertexShader, *PixelShader);
+						if (ReflectionCapture.SM4FullHDRCubemap)
+						{
+							const FSphere LightBounds(ReflectionCapture.PositionAndRadius, ReflectionCapture.PositionAndRadius.W);
 
-											PixelShader->SetParameters(View, ReflectionCapture);
-										}
-										else 
-										{
-											TShaderMapRef<TStandardDeferredReflectionPS<false> > PixelShader(GetGlobalShaderMap());
+							TShaderMapRef<TDeferredLightVS<true> > VertexShader(GetGlobalShaderMap());
 
-											static FGlobalBoundShaderState BoundShaderState;
-											SetGlobalBoundShaderState(BoundShaderState, GFilterVertexDeclaration.VertexDeclarationRHI, *VertexShader, *PixelShader);
+							// Use the appropriate shader for the capture shape
+							if (ReflectionCapture.CaptureProperties.Z == 0)
+							{
+								TShaderMapRef<TStandardDeferredReflectionPS<true> > PixelShader(GetGlobalShaderMap());
 
-											PixelShader->SetParameters(View, ReflectionCapture);
-										}
+								static FGlobalBoundShaderState BoundShaderState;
+								SetGlobalBoundShaderState(BoundShaderState, GFilterVertexDeclaration.VertexDeclarationRHI, *VertexShader, *PixelShader);
 
-										SetBoundingGeometryRasterizerAndDepthState(View, LightBounds);
-										VertexShader->SetSimpleLightParameters(View, LightBounds);
-										StencilingGeometry::DrawSphere();
-									}
-								}
+								PixelShader->SetParameters(View, ReflectionCapture);
 							}
-
-							if (Scene->SkyLight
-								&& Scene->SkyLight->ProcessedTexture
-								&& ViewFamily.EngineShowFlags.SkyLighting)
+							else 
 							{
-								SCOPED_DRAW_EVENT(SkyLightSpecularComposite, DEC_SCENE_ITEMS);
-
-								RHISetRasterizerState(TStaticRasterizerState<FM_Solid,CM_None>::GetRHI());
-								RHISetDepthStencilState(TStaticDepthStencilState<false,CF_Always>::GetRHI());
-								RHISetBlendState(TStaticBlendState<CW_RGBA, BO_Add, BF_DestAlpha, BF_One, BO_Add, BF_Zero, BF_InverseSourceAlpha>::GetRHI());
-
-								// Composite skylight specular, masked by accumulated reflection capture coverage
-								TShaderMapRef< FPostProcessVS >	 VertexShader( GetGlobalShaderMap() );
-								// Don't apply SSR in this sky light specular pass, that will happen in the reflection apply
-								TShaderMapRef< TSkyLightSpecularPS<false> > PixelShader( GetGlobalShaderMap() );
+								TShaderMapRef<TStandardDeferredReflectionPS<false> > PixelShader(GetGlobalShaderMap());
 
 								static FGlobalBoundShaderState BoundShaderState;
 								SetGlobalBoundShaderState(BoundShaderState, GFilterVertexDeclaration.VertexDeclarationRHI, *VertexShader, *PixelShader);
 
-								PixelShader->SetParameters(View, NULL);
-
-								DrawRectangle( 
-									0, 0, 
-									View.ViewRect.Width(), View.ViewRect.Height(),
-									View.ViewRect.Min.X, View.ViewRect.Min.Y, 
-									View.ViewRect.Width(), View.ViewRect.Height(),
-									FIntPoint(View.ViewRect.Width(), View.ViewRect.Height()),
-									GSceneRenderTargets.GetBufferSizeXY(),
-									*VertexShader);
+								PixelShader->SetParameters(View, ReflectionCapture);
 							}
-						}
-					}
-				}
-
-				if( bSSR )
-				{
-					GSceneRenderTargets.ResolveSceneColor(FResolveRect(0, 0, ViewFamily.FamilySizeX, ViewFamily.FamilySizeY));
-				}
-
-				for (int32 ViewIndex = 0; ViewIndex < Views.Num(); ViewIndex++)
-				{
-					const FViewInfo& View = Views[ViewIndex];
-
-					FTextureRHIParamRef SSRTexture = GSystemTextures.BlackDummy->GetRenderTargetItem().ShaderResourceTexture;
 
-					TRefCountPtr<IPooledRenderTarget> SSROutput;
-					if( bSSR )
-					{
-						SCOPED_DRAW_EVENT(ScreenSpaceReflections, DEC_SCENE_ITEMS);
-
-						ScreenSpaceReflections( View, SSROutput );
-						SSRTexture = SSROutput->GetRenderTargetItem().ShaderResourceTexture;
-					}
-
-					if (View.Family->EngineShowFlags.ReflectionEnvironment)
-					{
-						// Apply reflections to screen
-						SCOPED_DRAW_EVENT(ReflectionApply, DEC_SCENE_ITEMS);
-
-						GSceneRenderTargets.BeginRenderingSceneColor();
-
-						RHISetViewport(View.ViewRect.Min.X, View.ViewRect.Min.Y, 0.0f, View.ViewRect.Max.X, View.ViewRect.Max.Y, 1.0f);
-						RHISetRasterizerState(TStaticRasterizerState<FM_Solid,CM_None>::GetRHI());
-						RHISetDepthStencilState(TStaticDepthStencilState<false,CF_Always>::GetRHI());
-
-						if (GetReflectionEnvironmentCVar() == 2)
-						{
-							RHISetBlendState(TStaticBlendState<>::GetRHI());
-						}
-						else
-						{
-							RHISetBlendState(TStaticBlendState<CW_RGBA, BO_Add, BF_One, BF_One, BO_Add, BF_One, BF_One>::GetRHI());
+							SetBoundingGeometryRasterizerAndDepthState(View, LightBounds);
+							VertexShader->SetSimpleLightParameters(View, LightBounds);
+							StencilingGeometry::DrawSphere();
 						}
-
-						TShaderMapRef< FPostProcessVS >		VertexShader( GetGlobalShaderMap() );
-						TShaderMapRef< FReflectionApplyPS >	PixelShader( GetGlobalShaderMap() );
-
-						static FGlobalBoundShaderState BoundShaderState;
-						SetGlobalBoundShaderState(BoundShaderState, GFilterVertexDeclaration.VertexDeclarationRHI, *VertexShader, *PixelShader);
-
-						PixelShader->SetParameters( View, GSceneRenderTargets.GetLightAccumulationTexture(), SSRTexture );
-
-						DrawRectangle( 
-							0, 0, 
-							View.ViewRect.Width(), View.ViewRect.Height(),
-							View.ViewRect.Min.X, View.ViewRect.Min.Y, 
-							View.ViewRect.Width(), View.ViewRect.Height(),
-							FIntPoint(View.ViewRect.Width(), View.ViewRect.Height()),
-							GSceneRenderTargets.GetBufferSizeXY(),
-							*VertexShader);
 					}
-				}
-			}
-		}
-		else if (Scene->SkyLight
-			&& Scene->SkyLight->ProcessedTexture
-			&& ViewFamily.EngineShowFlags.SkyLighting
-			&& GRHIFeatureLevel >= ERHIFeatureLevel::SM4)
-		{
-			// Composite sky light specular by itself if the reflection environment pass is not going to be run
-			SCOPED_DRAW_EVENT(SkyLightSpecularApply, DEC_SCENE_ITEMS);
 
-			if( bSSR )
-			{
-				GSceneRenderTargets.ResolveSceneColor(FResolveRect(0, 0, ViewFamily.FamilySizeX, ViewFamily.FamilySizeY));
+					GRenderTargetPool.VisualizeTexture.SetCheckPoint( GSceneRenderTargets.LightAccumulation );
+				}
 			}
-
-			for (int32 ViewIndex = 0; ViewIndex < Views.Num(); ViewIndex++)
+			
+			if( bRequiresApply )
 			{
-				const FViewInfo& View = Views[ViewIndex];
-
-				FTextureRHIParamRef SSRTexture = GSystemTextures.BlackDummy->GetRenderTargetItem().ShaderResourceTexture;
-
-				TRefCountPtr<IPooledRenderTarget> SSROutput;
-				if( bSSR )
-				{
-					SCOPED_DRAW_EVENT(ScreenSpaceReflections, DEC_SCENE_ITEMS);
-
-					ScreenSpaceReflections( View, SSROutput );
-					SSRTexture = SSROutput->GetRenderTargetItem().ShaderResourceTexture;
-				}
+				// Apply reflections to screen
+				SCOPED_DRAW_EVENT(ReflectionApply, DEC_SCENE_ITEMS);
 
 				GSceneRenderTargets.BeginRenderingSceneColor();
 
 				RHISetViewport(View.ViewRect.Min.X, View.ViewRect.Min.Y, 0.0f, View.ViewRect.Max.X, View.ViewRect.Max.Y, 1.0f);
 				RHISetRasterizerState(TStaticRasterizerState<FM_Solid,CM_None>::GetRHI());
 				RHISetDepthStencilState(TStaticDepthStencilState<false,CF_Always>::GetRHI());
-				RHISetBlendState(TStaticBlendState<CW_RGB, BO_Add, BF_One, BF_One, BO_Add, BF_One, BF_One>::GetRHI());
 
-				TShaderMapRef< FPostProcessVS >	 VertexShader( GetGlobalShaderMap() );
-				TShaderMapRef< TSkyLightSpecularPS<true> > PixelShader( GetGlobalShaderMap() );
+				if (GetReflectionEnvironmentCVar() == 2)
+				{
+					RHISetBlendState(TStaticBlendState<>::GetRHI());
+				}
+				else
+				{
+					RHISetBlendState(TStaticBlendState<CW_RGBA, BO_Add, BF_One, BF_One, BO_Add, BF_One, BF_One>::GetRHI());
+				}
+
+				TShaderMapRef< FPostProcessVS > VertexShader( GetGlobalShaderMap() );
 
-				SetGlobalBoundShaderState(SkySpecularBoundShaderState, GFilterVertexDeclaration.VertexDeclarationRHI, *VertexShader, *PixelShader);
+				#define CASE(A,B,C) \
+					case ( (A << 2) | (B << 1) | C ): \
+					{ \
+						TShaderMapRef< FReflectionApplyPS<A,B,C> > PixelShader(GetGlobalShaderMap()); \
+						static FGlobalBoundShaderState BoundShaderState; \
+						SetGlobalBoundShaderState(BoundShaderState, GFilterVertexDeclaration.VertexDeclarationRHI, *VertexShader, *PixelShader); \
+						PixelShader->SetParameters( View, GSceneRenderTargets.GetLightAccumulationTexture(), SSROutput->GetRenderTargetItem().ShaderResourceTexture ); \
+					}; \
+					break
 
-				PixelShader->SetParameters(View, SSRTexture);
+				switch( (bSSR << 2) | (bReflectionEnv << 1) | bSkyLight )
+				{
+					CASE(0,0,0);
+					CASE(0,0,1);
+					CASE(0,1,0);
+					CASE(0,1,1);
+					CASE(1,0,0);
+					CASE(1,0,1);
+					CASE(1,1,0);
+					CASE(1,1,1);
+				}
+				#undef CASE
 
 				DrawRectangle( 
 					0, 0, 
@@ -875,53 +810,7 @@ void FDeferredShadingSceneRenderer::RenderDeferredReflections()
 					*VertexShader);
 			}
 		}
-		else
-		{
-			for (int32 ViewIndex = 0; ViewIndex < Views.Num(); ViewIndex++)
-			{
-				const FViewInfo& View = Views[ViewIndex];
-
-				FTextureRHIParamRef SSRTexture = GSystemTextures.BlackDummy->GetRenderTargetItem().ShaderResourceTexture;
-
-				TRefCountPtr<IPooledRenderTarget> SSROutput;
-				if( bSSR )
-				{
-					SCOPED_DRAW_EVENT(ScreenSpaceReflections, DEC_SCENE_ITEMS);
-
-					ScreenSpaceReflections( View, SSROutput );
-					SSRTexture = SSROutput->GetRenderTargetItem().ShaderResourceTexture;
-
-					{
-						// Apply reflections to screen
-						SCOPED_DRAW_EVENT(ReflectionApply, DEC_SCENE_ITEMS);
-
-						GSceneRenderTargets.BeginRenderingSceneColor();
-
-						RHISetViewport(View.ViewRect.Min.X, View.ViewRect.Min.Y, 0.0f, View.ViewRect.Max.X, View.ViewRect.Max.Y, 1.0f);
-						RHISetRasterizerState(TStaticRasterizerState<FM_Solid,CM_None>::GetRHI());
-						RHISetDepthStencilState(TStaticDepthStencilState<false,CF_Always>::GetRHI());
-
-						RHISetBlendState(TStaticBlendState<CW_RGBA, BO_Add, BF_One, BF_One, BO_Add, BF_One, BF_One>::GetRHI());
-
-						TShaderMapRef< FPostProcessVS >		VertexShader( GetGlobalShaderMap() );
-						TShaderMapRef< FReflectionApplyPS >	PixelShader( GetGlobalShaderMap() );
-
-						static FGlobalBoundShaderState BoundShaderState;
-						SetGlobalBoundShaderState(BoundShaderState, GFilterVertexDeclaration.VertexDeclarationRHI, *VertexShader, *PixelShader);
-
-						PixelShader->SetParameters( View, GSystemTextures.BlackDummy->GetRenderTargetItem().ShaderResourceTexture, SSRTexture );
-
-						DrawRectangle( 
-							0, 0, 
-							View.ViewRect.Width(), View.ViewRect.Height(),
-							View.ViewRect.Min.X, View.ViewRect.Min.Y, 
-							View.ViewRect.Width(), View.ViewRect.Height(),
-							FIntPoint(View.ViewRect.Width(), View.ViewRect.Height()),
-							GSceneRenderTargets.GetBufferSizeXY(),
-							*VertexShader);
-					}
-				}
-			}
-		}
 	}
+
+	NewSceneColorUAV.SafeRelease();
 }
diff --git a/Engine/Source/Runtime/Renderer/Private/ReflectionEnvironmentCapture.cpp b/Engine/Source/Runtime/Renderer/Private/ReflectionEnvironmentCapture.cpp
--- a/Engine/Source/Runtime/Renderer/Private/ReflectionEnvironmentCapture.cpp
+++ b/Engine/Source/Runtime/Renderer/Private/ReflectionEnvironmentCapture.cpp
@@ -76,37 +76,6 @@ FSceneRenderTargetItem& GetEffectiveSourceTexture(bool bDownsamplePass, int32 Ta
 
 IMPLEMENT_SHADER_TYPE(,FDownsampleGS,TEXT("ReflectionEnvironmentShaders"),TEXT("DownsampleGS"),SF_Geometry);
 
-const int32 NumFilterSamples = 1024;
-
-/** Generates a sample on a cone with probability equal to the value of the GGX specular distribution. */
-FVector GenerateImportanceSampleGGX(float Roughness, int32 Index, FVector2D* Samples)
-{
-	float m = Roughness * Roughness;
-	float m2 = m * m;
-
-	float E1 = Samples[Index].X;
-	float E2 = Samples[Index].Y;
-
-	// Only implemented for GGX
-	float Phi = 2 * PI * E1;
-	float CosPhi = FMath::Cos(Phi);
-	float SinPhi = FMath::Sin(Phi);
-	float CosTheta = FMath::Sqrt((1 - E2) / (1 + (m2 - 1) * E2));
-	float SinTheta = FMath::Sqrt(1 - CosTheta * CosTheta);
-
-	FVector HalfVector(SinTheta * FMath::Cos(Phi), SinTheta * FMath::Sin(Phi), CosTheta);
-	FVector View(0, 0, 1);
-	return 2 * (View | HalfVector) * HalfVector - View;
-}
-
-/** Uniform buffer for cubemap filtering data. */
-BEGIN_UNIFORM_BUFFER_STRUCT(FCubeFilterData,)
-	DECLARE_UNIFORM_BUFFER_STRUCT_MEMBER_ARRAY(FVector4,Samples,[NumFilterSamples])
-END_UNIFORM_BUFFER_STRUCT(FCubeFilterData)
-
-IMPLEMENT_UNIFORM_BUFFER_STRUCT(FCubeFilterData,TEXT("FilterData"));
-
-/** Pixel shader used for downsampling a mip or filtering a mip. */
 class FDownsamplePS : public FGlobalShader
 {
 	DECLARE_SHADER_TYPE(FDownsamplePS,Global);
@@ -117,12 +86,6 @@ public:
 		return IsFeatureLevelSupported(Platform, ERHIFeatureLevel::SM4);
 	}
 
-	static void ModifyCompilationEnvironment(EShaderPlatform Platform, FShaderCompilerEnvironment& OutEnvironment)
-	{
-		FGlobalShader::ModifyCompilationEnvironment(Platform, OutEnvironment);
-		OutEnvironment.SetDefine(TEXT("NUM_FILTER_SAMPLES"), NumFilterSamples);
-	}
-
 	FDownsamplePS(const ShaderMetaType::CompiledShaderInitializerType& Initializer):
 		FGlobalShader(Initializer)
 	{
@@ -130,40 +93,20 @@ public:
 		SourceMipIndex.Bind(Initializer.ParameterMap,TEXT("SourceMipIndex"));
 		SourceTexture.Bind(Initializer.ParameterMap,TEXT("SourceTexture"));
 		SourceTextureSampler.Bind(Initializer.ParameterMap,TEXT("SourceTextureSampler"));
-		AverageBrightnessTexture.Bind(Initializer.ParameterMap,TEXT("AverageBrightnessTexture"));
-		AverageBrightnessSampler.Bind(Initializer.ParameterMap,TEXT("AverageBrightnessSampler"));
 	}
 	FDownsamplePS() {}
 
-	void SetParameters(int32 CubeFaceValue, int32 SourceMipIndexValue, const FVector4* FilterSamples, FSceneRenderTargetItem& SourceTextureValue)
+	void SetParameters(int32 CubeFaceValue, int32 SourceMipIndexValue, FSceneRenderTargetItem& SourceTextureValue)
 	{
 		SetShaderValue(GetPixelShader(), CubeFace, CubeFaceValue);
 		SetShaderValue(GetPixelShader(), SourceMipIndex, SourceMipIndexValue);
 
-		FCubeFilterData FilterData;
-
-		// Copy the samples
-		// Can't use memcpy due to padding in the uniform buffer
-		for (int32 i = 0; i < NumFilterSamples; i++)
-		{
-			FilterData.Samples[i] = FilterSamples[i];
-		}
-
-		SetUniformBufferParameterImmediate(GetPixelShader(), GetUniformBufferParameter<FCubeFilterData>(), FilterData);
-
 		SetTextureParameter(
 			GetPixelShader(), 
 			SourceTexture, 
 			SourceTextureSampler, 
 			TStaticSamplerState<SF_Bilinear, AM_Clamp, AM_Clamp, AM_Clamp>::GetRHI(), 
 			SourceTextureValue.ShaderResourceTexture);
-
-		SetTextureParameter(
-			GetPixelShader(), 
-			AverageBrightnessTexture, 
-			AverageBrightnessSampler, 
-			TStaticSamplerState<SF_Bilinear, AM_Clamp, AM_Clamp, AM_Clamp>::GetRHI(), 
-			GSceneRenderTargets.ReflectionBrightness->GetRenderTargetItem().ShaderResourceTexture);
 	}
 
 	virtual bool Serialize(FArchive& Ar)
@@ -173,8 +116,6 @@ public:
 		Ar << SourceMipIndex;
 		Ar << SourceTexture;
 		Ar << SourceTextureSampler;
-		Ar << AverageBrightnessTexture;
-		Ar << AverageBrightnessSampler;
 		return bShaderHasOutdatedParameters;
 	}
 
@@ -183,41 +124,92 @@ private:
 	FShaderParameter SourceMipIndex;
 	FShaderResourceParameter SourceTexture;
 	FShaderResourceParameter SourceTextureSampler;
-	FShaderResourceParameter AverageBrightnessTexture;
-	FShaderResourceParameter AverageBrightnessSampler;
 };
 
-/** Used to generate permutations of FDownsamplePS with different defines. */
-template<bool bDownsample>
-class TDownsamplePS : public FDownsamplePS
+IMPLEMENT_SHADER_TYPE(,FDownsamplePS,TEXT("ReflectionEnvironmentShaders"),TEXT("DownsamplePS"),SF_Pixel);
+
+
+const int32 NumFilterSamples = 1024;
+
+/** Pixel shader used for filtering a mip. */
+class FFilterPS : public FDownsamplePS
 {
-	DECLARE_SHADER_TYPE(TDownsamplePS,Global);
+	DECLARE_SHADER_TYPE(FFilterPS,Global);
 public:
 
+	static bool ShouldCache(EShaderPlatform Platform)
+	{
+		return IsFeatureLevelSupported(Platform, ERHIFeatureLevel::SM4);
+	}
+
 	static void ModifyCompilationEnvironment(EShaderPlatform Platform, FShaderCompilerEnvironment& OutEnvironment)
 	{
 		FDownsamplePS::ModifyCompilationEnvironment(Platform, OutEnvironment);
-		OutEnvironment.SetDefine(TEXT("DOWNSAMPLE_PASS"), (uint32)(bDownsample ? 1 : 0));
+		OutEnvironment.SetDefine(TEXT("NUM_FILTER_SAMPLES"), NumFilterSamples);
 	}
 
-	TDownsamplePS(const ShaderMetaType::CompiledShaderInitializerType& Initializer):
+	FFilterPS(const ShaderMetaType::CompiledShaderInitializerType& Initializer):
 		FDownsamplePS(Initializer)
-	{}
+	{
+		FilterRoughness.Bind(Initializer.ParameterMap,TEXT("FilterRoughness"));
+		AverageBrightnessTexture.Bind(Initializer.ParameterMap,TEXT("AverageBrightnessTexture"));
+		AverageBrightnessSampler.Bind(Initializer.ParameterMap,TEXT("AverageBrightnessSampler"));
+	}
+	FFilterPS() {}
 
-	TDownsamplePS()
-	{}
+	void SetParameters(int32 CubeFaceValue, int32 SourceMipIndexValue, float RoughnessValue, FSceneRenderTargetItem& SourceTextureValue)
+	{
+		FDownsamplePS::SetParameters(CubeFaceValue, SourceMipIndexValue, SourceTextureValue);
+
+		SetShaderValue(GetPixelShader(), FilterRoughness, RoughnessValue);
+
+		SetTextureParameter(
+			GetPixelShader(), 
+			AverageBrightnessTexture, 
+			AverageBrightnessSampler, 
+			TStaticSamplerState<SF_Bilinear, AM_Clamp, AM_Clamp, AM_Clamp>::GetRHI(), 
+			GSceneRenderTargets.ReflectionBrightness->GetRenderTargetItem().ShaderResourceTexture);
+	}
+
+	virtual bool Serialize(FArchive& Ar)
+	{
+		bool bShaderHasOutdatedParameters = FDownsamplePS::Serialize(Ar);
+		Ar << FilterRoughness;
+		Ar << AverageBrightnessTexture;
+		Ar << AverageBrightnessSampler;
+		return bShaderHasOutdatedParameters;
+	}
+
+private:
+	FShaderParameter FilterRoughness;
+	FShaderResourceParameter AverageBrightnessTexture;
+	FShaderResourceParameter AverageBrightnessSampler;
 };
 
-#define IMPLEMENT_REFLECTION_FILTER_PIXELSHADER_TYPE(bDownsample, MainFunction) \
-	typedef TDownsamplePS<bDownsample> TDownsamplePS##bDownsample; \
-	IMPLEMENT_SHADER_TYPE(template<>,TDownsamplePS##bDownsample,TEXT("ReflectionEnvironmentShaders"),MainFunction,SF_Pixel);
+template< uint32 bNormalize >
+class TFilterPS : public FFilterPS
+{
+	DECLARE_SHADER_TYPE(TFilterPS,Global);
 
+public:
+	static void ModifyCompilationEnvironment(EShaderPlatform Platform, FShaderCompilerEnvironment& OutEnvironment)
+	{
+		FFilterPS::ModifyCompilationEnvironment(Platform, OutEnvironment);
+		OutEnvironment.SetDefine(TEXT("NORMALIZE"), bNormalize);
+	}
 
-IMPLEMENT_REFLECTION_FILTER_PIXELSHADER_TYPE(true, TEXT("DownsamplePS"));
-IMPLEMENT_REFLECTION_FILTER_PIXELSHADER_TYPE(false, TEXT("BlurPS"));
+	TFilterPS(const ShaderMetaType::CompiledShaderInitializerType& Initializer)
+	: FFilterPS(Initializer)
+	{}
 
-FGlobalBoundShaderState DownsampleBoundShaderState;
-FGlobalBoundShaderState BlurBoundShaderState;
+	TFilterPS() {}
+};
+
+IMPLEMENT_SHADER_TYPE(template<>,TFilterPS<0>,TEXT("ReflectionEnvironmentShaders"),TEXT("FilterPS"),SF_Pixel);
+IMPLEMENT_SHADER_TYPE(template<>,TFilterPS<1>,TEXT("ReflectionEnvironmentShaders"),TEXT("FilterPS"),SF_Pixel);
+
+static FGlobalBoundShaderState DownsampleBoundShaderState;
+static FGlobalBoundShaderState FilterBoundShaderState;
 
 /** Computes the average brightness of a 1x1 mip of a cubemap. */
 class FComputeBrightnessPS : public FGlobalShader
@@ -313,45 +305,18 @@ void ComputeAverageBrightness()
 		*VertexShader);
 }
 
-FVector4 FilterConeImportanceSamples[MAX_TEXTURE_MIP_COUNT][NumFilterSamples];
-
 /** Generates mips for glossiness and filters the cubemap for a given reflection. */
 void FilterReflectionEnvironment(FSHVectorRGB3* OutIrradianceEnvironmentMap)
 {
 	const int32 EffectiveTopMipSize = GReflectionCaptureSize;
 	const int32 NumMips = FMath::CeilLogTwo(EffectiveTopMipSize) + 1;
 
-	static bool bGeneratedSamples = false;
-
-	if (!bGeneratedSamples)
-	{
-		bGeneratedSamples = true;
-		FRandomStream RandomStream(12345);
-		const int32 NumSamplesOneDimension = FMath::Sqrt(NumFilterSamples);
-		FVector2D StratifiedUniformRandoms[NumFilterSamples];
-
-		for (int32 i = 0; i < NumFilterSamples; i++)
-		{
-			StratifiedUniformRandoms[i] = FVector2D(
-				(i % NumSamplesOneDimension + RandomStream.GetFraction()) / (float)NumSamplesOneDimension, 
-				(i / NumSamplesOneDimension + RandomStream.GetFraction()) / (float)NumSamplesOneDimension);
-		}
-
-		for (int32 MipIndex = 0; MipIndex < NumMips; MipIndex++)
-		{
-			// Note: this must match the logic in ComputeReflectionCaptureMipFromRoughness!
-			float Roughness = FMath::Pow(2, (MipIndex - NumMips + 1 + ReflectionCaptureRoughestMip) / ReflectionCaptureRoughnessMipScale);
-
-			for (int32 SampleIndex = 0; SampleIndex < NumFilterSamples; SampleIndex++)
-			{
-				FilterConeImportanceSamples[MipIndex][SampleIndex] = FVector4(GenerateImportanceSampleGGX(Roughness, SampleIndex, StratifiedUniformRandoms), 0);		
-			}
-		}
-	}
-
 	int32 DiffuseConvolutionSourceMip = INDEX_NONE;
 	FSceneRenderTargetItem* DiffuseConvolutionSource = NULL;
 
+	static const auto CVar = IConsoleManager::Get().FindConsoleVariable(TEXT("r.DiffuseFromCaptures"));
+	const bool bNormalize = CVar->GetInt() == 0;
+
 	{	
 		SCOPED_DRAW_EVENT(DownsampleCubeMips, DEC_SCENE_ITEMS);
 		// Downsample all the mips, each one reads from the mip above it
@@ -367,40 +332,40 @@ void FilterReflectionEnvironment(FSHVectorRGB3* OutIrradianceEnvironmentMap)
 
 			if (GSupportsGSRenderTargetLayerSwitchingToMips)
 			{
-			RHISetRenderTarget(EffectiveRT.TargetableTexture, MipIndex, NULL);
+				RHISetRenderTarget(EffectiveRT.TargetableTexture, MipIndex, NULL);
 
-			const FIntRect ViewRect(0, 0, MipSize, MipSize);
-			RHISetViewport(0, 0, 0.0f, MipSize, MipSize, 1.0f);
-			RHISetRasterizerState(TStaticRasterizerState<FM_Solid,CM_None>::GetRHI());
-			RHISetDepthStencilState(TStaticDepthStencilState<false,CF_Always>::GetRHI());
-			RHISetBlendState(TStaticBlendState<>::GetRHI());
+				const FIntRect ViewRect(0, 0, MipSize, MipSize);
+				RHISetViewport(0, 0, 0.0f, MipSize, MipSize, 1.0f);
+				RHISetRasterizerState(TStaticRasterizerState<FM_Solid,CM_None>::GetRHI());
+				RHISetDepthStencilState(TStaticDepthStencilState<false,CF_Always>::GetRHI());
+				RHISetBlendState(TStaticBlendState<>::GetRHI());
 
-			TShaderMapRef<FScreenVSForGS> VertexShader(GetGlobalShaderMap());
-			TShaderMapRef<FDownsampleGS> GeometryShader(GetGlobalShaderMap());
-			TShaderMapRef<TDownsamplePS<true> > PixelShader(GetGlobalShaderMap());
+				TShaderMapRef<FScreenVSForGS> VertexShader(GetGlobalShaderMap());
+				TShaderMapRef<FDownsampleGS> GeometryShader(GetGlobalShaderMap());
+				TShaderMapRef<FDownsamplePS> PixelShader(GetGlobalShaderMap());
 
-			SetGlobalBoundShaderState(DownsampleBoundShaderState, GFilterVertexDeclaration.VertexDeclarationRHI, *VertexShader, *PixelShader, *GeometryShader);
+				SetGlobalBoundShaderState(DownsampleBoundShaderState, GFilterVertexDeclaration.VertexDeclarationRHI, *VertexShader, *PixelShader, *GeometryShader);
 
-			// Draw each face with a geometry shader that routes to the correct render target slice
-			for (int32 CubeFace = 0; CubeFace < CubeFace_MAX; CubeFace++)
-			{
-				SCOPED_DRAW_EVENT(DownsampleCubeFace, DEC_SCENE_ITEMS);
-				GeometryShader->SetParameters(CubeFace);
+				// Draw each face with a geometry shader that routes to the correct render target slice
+				for (int32 CubeFace = 0; CubeFace < CubeFace_MAX; CubeFace++)
+				{
+					SCOPED_DRAW_EVENT(DownsampleCubeFace, DEC_SCENE_ITEMS);
+					GeometryShader->SetParameters(CubeFace);
 
-				PixelShader->SetParameters(CubeFace, SourceMipIndex, FilterConeImportanceSamples[SourceMipIndex], EffectiveSource);
+					PixelShader->SetParameters(CubeFace, SourceMipIndex, EffectiveSource);
 
-				DrawRectangle( 
-					ViewRect.Min.X, ViewRect.Min.Y, 
-					ViewRect.Width(), ViewRect.Height(),
-					ViewRect.Min.X, ViewRect.Min.Y, 
-					ViewRect.Width(), ViewRect.Height(),
-					FIntPoint(ViewRect.Width(), ViewRect.Height()),
-					FIntPoint(MipSize, MipSize),
-					*VertexShader,
-					EDRF_UseTriangleOptimization);
+					DrawRectangle( 
+						ViewRect.Min.X, ViewRect.Min.Y, 
+						ViewRect.Width(), ViewRect.Height(),
+						ViewRect.Min.X, ViewRect.Min.Y, 
+						ViewRect.Width(), ViewRect.Height(),
+						FIntPoint(ViewRect.Width(), ViewRect.Height()),
+						FIntPoint(MipSize, MipSize),
+						*VertexShader,
+						EDRF_UseTriangleOptimization);
 
-				RHICopyToResolveTarget(EffectiveRT.TargetableTexture, EffectiveRT.ShaderResourceTexture, true, FResolveParams(FResolveRect(), (ECubeFace)CubeFace, MipIndex));
-			}
+					RHICopyToResolveTarget(EffectiveRT.TargetableTexture, EffectiveRT.ShaderResourceTexture, true, FResolveParams(FResolveRect(), (ECubeFace)CubeFace, MipIndex));
+				}
 			}
 			else
 			{
@@ -415,11 +380,11 @@ void FilterReflectionEnvironment(FSHVectorRGB3* OutIrradianceEnvironmentMap)
 					RHISetBlendState(TStaticBlendState<>::GetRHI());
 
 					TShaderMapRef<FScreenVSForGS> VertexShader(GetGlobalShaderMap());
-					TShaderMapRef<TDownsamplePS<true> > PixelShader(GetGlobalShaderMap());
+					TShaderMapRef<FDownsamplePS> PixelShader(GetGlobalShaderMap());
 
 					SetGlobalBoundShaderState(DownsampleBoundShaderState, GFilterVertexDeclaration.VertexDeclarationRHI, *VertexShader, *PixelShader);
 
-					PixelShader->SetParameters(CubeFace, SourceMipIndex, FilterConeImportanceSamples[SourceMipIndex], EffectiveSource);
+					PixelShader->SetParameters(CubeFace, SourceMipIndex, EffectiveSource);
 
 					DrawRectangle( 
 						ViewRect.Min.X, ViewRect.Min.Y, 
@@ -462,43 +427,58 @@ void FilterReflectionEnvironment(FSHVectorRGB3* OutIrradianceEnvironmentMap)
 			check(EffectiveRT.TargetableTexture != EffectiveSource.ShaderResourceTexture);
 			const int32 MipSize = 1 << (NumMips - MipIndex - 1);
 
+			// Note: this must match the logic in ComputeReflectionCaptureMipFromRoughness!
+			float Roughness = FMath::Pow(2, (MipIndex - NumMips + 1 + ReflectionCaptureRoughestMip) / ReflectionCaptureRoughnessMipScale);
+
 			if (GSupportsGSRenderTargetLayerSwitchingToMips)
 			{
-			RHISetRenderTarget(EffectiveRT.TargetableTexture, MipIndex, NULL);
+				RHISetRenderTarget(EffectiveRT.TargetableTexture, MipIndex, NULL);
 
-			const FIntRect ViewRect(0, 0, MipSize, MipSize);
-			RHISetViewport(0, 0, 0.0f, MipSize, MipSize, 1.0f);
-			RHISetRasterizerState(TStaticRasterizerState<FM_Solid,CM_None>::GetRHI());
-			RHISetDepthStencilState(TStaticDepthStencilState<false,CF_Always>::GetRHI());
-			RHISetBlendState(TStaticBlendState<>::GetRHI());
+				const FIntRect ViewRect(0, 0, MipSize, MipSize);
+				RHISetViewport(0, 0, 0.0f, MipSize, MipSize, 1.0f);
+				RHISetRasterizerState(TStaticRasterizerState<FM_Solid,CM_None>::GetRHI());
+				RHISetDepthStencilState(TStaticDepthStencilState<false,CF_Always>::GetRHI());
+				RHISetBlendState(TStaticBlendState<>::GetRHI());
 
-			TShaderMapRef<FScreenVSForGS> VertexShader(GetGlobalShaderMap());
-			TShaderMapRef<FDownsampleGS> GeometryShader(GetGlobalShaderMap());
-			TShaderMapRef<TDownsamplePS<false> > CaptureCubemapArrayPixelShader(GetGlobalShaderMap());
-			FDownsamplePS* PixelShader = (FDownsamplePS*)*CaptureCubemapArrayPixelShader;
+				TShaderMapRef<FScreenVSForGS> VertexShader(GetGlobalShaderMap());
+				TShaderMapRef<FDownsampleGS> GeometryShader(GetGlobalShaderMap());
+				
+				FFilterPS* PixelShader;
+				if( bNormalize )
+				{
+					PixelShader = *TShaderMapRef< TFilterPS<1> >( GetGlobalShaderMap() );
+					static FGlobalBoundShaderState BoundShaderState;
+					SetGlobalBoundShaderState(BoundShaderState, GFilterVertexDeclaration.VertexDeclarationRHI, *VertexShader, PixelShader, *GeometryShader);
+				}
+				else
+				{
+					PixelShader = *TShaderMapRef< TFilterPS<0> >( GetGlobalShaderMap() );
+					static FGlobalBoundShaderState BoundShaderState;
+					SetGlobalBoundShaderState(BoundShaderState, GFilterVertexDeclaration.VertexDeclarationRHI, *VertexShader, PixelShader, *GeometryShader);
+				}
 
-			SetGlobalBoundShaderState(BlurBoundShaderState, GFilterVertexDeclaration.VertexDeclarationRHI, *VertexShader, PixelShader, *GeometryShader);
+				//SetGlobalBoundShaderState(FilterBoundShaderState, GFilterVertexDeclaration.VertexDeclarationRHI, *VertexShader, PixelShader, *GeometryShader);
 
-			for (int32 CubeFace = 0; CubeFace < CubeFace_MAX; CubeFace++)
-			{
-				SCOPED_DRAW_EVENT(FilterCubeFace, DEC_SCENE_ITEMS);
-				GeometryShader->SetParameters(CubeFace);
+				for (int32 CubeFace = 0; CubeFace < CubeFace_MAX; CubeFace++)
+				{
+					SCOPED_DRAW_EVENT(FilterCubeFace, DEC_SCENE_ITEMS);
+					GeometryShader->SetParameters(CubeFace);
 
-				PixelShader->SetParameters(CubeFace, MipIndex, FilterConeImportanceSamples[MipIndex], EffectiveSource);
+					PixelShader->SetParameters(CubeFace, MipIndex, Roughness, EffectiveSource);
 
-				DrawRectangle( 
-					ViewRect.Min.X, ViewRect.Min.Y, 
-					ViewRect.Width(), ViewRect.Height(),
-					ViewRect.Min.X, ViewRect.Min.Y, 
-					ViewRect.Width(), ViewRect.Height(),
-					FIntPoint(ViewRect.Width(), ViewRect.Height()),
-					FIntPoint(MipSize, MipSize),
-					*VertexShader,
-					EDRF_UseTriangleOptimization);
+					DrawRectangle( 
+						ViewRect.Min.X, ViewRect.Min.Y, 
+						ViewRect.Width(), ViewRect.Height(),
+						ViewRect.Min.X, ViewRect.Min.Y, 
+						ViewRect.Width(), ViewRect.Height(),
+						FIntPoint(ViewRect.Width(), ViewRect.Height()),
+						FIntPoint(MipSize, MipSize),
+						*VertexShader,
+						EDRF_UseTriangleOptimization);
 
-				RHICopyToResolveTarget(EffectiveRT.TargetableTexture, EffectiveRT.ShaderResourceTexture, true, FResolveParams(FResolveRect(), (ECubeFace)CubeFace, MipIndex));
+					RHICopyToResolveTarget(EffectiveRT.TargetableTexture, EffectiveRT.ShaderResourceTexture, true, FResolveParams(FResolveRect(), (ECubeFace)CubeFace, MipIndex));
+				}
 			}
-		}
 			else
 			{
 				for (int32 CubeFace = 0; CubeFace < CubeFace_MAX; CubeFace++)
@@ -512,12 +492,26 @@ void FilterReflectionEnvironment(FSHVectorRGB3* OutIrradianceEnvironmentMap)
 					RHISetBlendState(TStaticBlendState<>::GetRHI());
 
 					TShaderMapRef<FScreenVSForGS> VertexShader(GetGlobalShaderMap());
-					TShaderMapRef<TDownsamplePS<false> > CaptureCubemapArrayPixelShader(GetGlobalShaderMap());
-					FDownsamplePS* PixelShader = (FDownsamplePS*)*CaptureCubemapArrayPixelShader;
+					TShaderMapRef< TFilterPS<1> > CaptureCubemapArrayPixelShader(GetGlobalShaderMap());
+					//FFilterPS* PixelShader = (FFilterPS*)*CaptureCubemapArrayPixelShader;
 
-					SetGlobalBoundShaderState(BlurBoundShaderState, GFilterVertexDeclaration.VertexDeclarationRHI, *VertexShader, PixelShader);
+					FFilterPS* PixelShader;
+					if( bNormalize )
+					{
+						PixelShader = *TShaderMapRef< TFilterPS<1> >( GetGlobalShaderMap() );
+						static FGlobalBoundShaderState BoundShaderState;
+						SetGlobalBoundShaderState(BoundShaderState, GFilterVertexDeclaration.VertexDeclarationRHI, *VertexShader, PixelShader);
+					}
+					else
+					{
+						PixelShader = *TShaderMapRef< TFilterPS<0> >( GetGlobalShaderMap() );
+						static FGlobalBoundShaderState BoundShaderState;
+						SetGlobalBoundShaderState(BoundShaderState, GFilterVertexDeclaration.VertexDeclarationRHI, *VertexShader, PixelShader);
+					}
 
-					PixelShader->SetParameters(CubeFace, MipIndex, FilterConeImportanceSamples[MipIndex], EffectiveSource);
+					//SetGlobalBoundShaderState(FilterBoundShaderState, GFilterVertexDeclaration.VertexDeclarationRHI, *VertexShader, PixelShader);
+
+					PixelShader->SetParameters(CubeFace, MipIndex, Roughness, EffectiveSource);
 
 					DrawRectangle( 
 						ViewRect.Min.X, ViewRect.Min.Y, 
@@ -804,37 +798,37 @@ void CaptureSceneToScratchCubemap(FSceneRenderer* SceneRenderer, ECubeFace CubeF
 
 			if (GSupportsGSRenderTargetLayerSwitchingToMips)
 			{
-			// Copy the captured scene into the cubemap face
-			RHISetRenderTarget(EffectiveColorRT.TargetableTexture, NULL);
+				// Copy the captured scene into the cubemap face
+				RHISetRenderTarget(EffectiveColorRT.TargetableTexture, NULL);
 
-			const FIntRect ViewRect(0, 0, EffectiveSize, EffectiveSize);
-			RHISetViewport(0, 0, 0.0f, EffectiveSize, EffectiveSize, 1.0f);
-			RHISetRasterizerState(TStaticRasterizerState<FM_Solid,CM_None>::GetRHI());
-			RHISetDepthStencilState(TStaticDepthStencilState<false,CF_Always>::GetRHI());
-			RHISetBlendState(TStaticBlendState<>::GetRHI());
+				const FIntRect ViewRect(0, 0, EffectiveSize, EffectiveSize);
+				RHISetViewport(0, 0, 0.0f, EffectiveSize, EffectiveSize, 1.0f);
+				RHISetRasterizerState(TStaticRasterizerState<FM_Solid,CM_None>::GetRHI());
+				RHISetDepthStencilState(TStaticDepthStencilState<false,CF_Always>::GetRHI());
+				RHISetBlendState(TStaticBlendState<>::GetRHI());
 
-			TShaderMapRef<FCopyToCubeFaceVS> VertexShader(GetGlobalShaderMap());
-			TShaderMapRef<FRouteToCubeFaceGS> GeometryShader(GetGlobalShaderMap());
+				TShaderMapRef<FCopyToCubeFaceVS> VertexShader(GetGlobalShaderMap());
+				TShaderMapRef<FRouteToCubeFaceGS> GeometryShader(GetGlobalShaderMap());
 
-			TShaderMapRef<FCopySceneColorToCubeFacePS> PixelShader(GetGlobalShaderMap());
-			SetGlobalBoundShaderState(CopyColorCubemapBoundShaderState, GFilterVertexDeclaration.VertexDeclarationRHI, *VertexShader, *PixelShader, *GeometryShader);
+				TShaderMapRef<FCopySceneColorToCubeFacePS> PixelShader(GetGlobalShaderMap());
+				SetGlobalBoundShaderState(CopyColorCubemapBoundShaderState, GFilterVertexDeclaration.VertexDeclarationRHI, *VertexShader, *PixelShader, *GeometryShader);
 			
-			PixelShader->SetParameters(SceneRenderer->Views[0], bCapturingForSkyLight, bLowerHemisphereIsBlack);
-			VertexShader->SetParameters(SceneRenderer->Views[0]);
-			GeometryShader->SetParameters((uint32)CubeFace);
+				PixelShader->SetParameters(SceneRenderer->Views[0], bCapturingForSkyLight, bLowerHemisphereIsBlack);
+				VertexShader->SetParameters(SceneRenderer->Views[0]);
+				GeometryShader->SetParameters((uint32)CubeFace);
 
-			DrawRectangle( 
-				ViewRect.Min.X, ViewRect.Min.Y, 
-				ViewRect.Width(), ViewRect.Height(),
-				ViewRect.Min.X, ViewRect.Min.Y, 
-				ViewRect.Width() * GSupersampleCaptureFactor, ViewRect.Height() * GSupersampleCaptureFactor,
-				FIntPoint(ViewRect.Width(), ViewRect.Height()),
-				GSceneRenderTargets.GetBufferSizeXY(),
-				*VertexShader,
-				EDRF_UseTriangleOptimization);
+				DrawRectangle( 
+					ViewRect.Min.X, ViewRect.Min.Y, 
+					ViewRect.Width(), ViewRect.Height(),
+					ViewRect.Min.X, ViewRect.Min.Y, 
+					ViewRect.Width() * GSupersampleCaptureFactor, ViewRect.Height() * GSupersampleCaptureFactor,
+					FIntPoint(ViewRect.Width(), ViewRect.Height()),
+					GSceneRenderTargets.GetBufferSizeXY(),
+					*VertexShader,
+					EDRF_UseTriangleOptimization);
 
-			RHICopyToResolveTarget(EffectiveColorRT.TargetableTexture, EffectiveColorRT.ShaderResourceTexture, true, FResolveParams(FResolveRect()));
-		}
+				RHICopyToResolveTarget(EffectiveColorRT.TargetableTexture, EffectiveColorRT.ShaderResourceTexture, true, FResolveParams(FResolveRect()));
+			}
 			else
 			{
 				// Copy the captured scene into the cubemap face
diff --git a/Engine/Source/Runtime/Renderer/Private/SceneRendering.cpp b/Engine/Source/Runtime/Renderer/Private/SceneRendering.cpp
--- a/Engine/Source/Runtime/Renderer/Private/SceneRendering.cpp
+++ b/Engine/Source/Runtime/Renderer/Private/SceneRendering.cpp
@@ -365,6 +365,11 @@ TUniformBufferRef<FViewUniformShaderParameters> FViewInfo::CreateUniformBuffer(
 	ViewUniformShaderParameters.Random = FMath::Rand();
 	ViewUniformShaderParameters.FrameNumber = FrameNumber;
 
+	static const auto CVar = IConsoleManager::Get().FindConsoleVariable(TEXT("r.DiffuseFromCaptures"));
+	const bool bUseLightmaps = CVar->GetInt() == 0;
+
+	ViewUniformShaderParameters.UseLightmaps = bUseLightmaps;
+
 	if(State)
 	{
 		// safe to cast on the renderer side
