* fixed build error

diff --git a/Engine/Source/Runtime/Renderer/Private/PostProcess/VisualizeTexture.cpp b/Engine/Source/Runtime/Renderer/Private/PostProcess/VisualizeTexture.cpp
--- a/Engine/Source/Runtime/Renderer/Private/PostProcess/VisualizeTexture.cpp
+++ b/Engine/Source/Runtime/Renderer/Private/PostProcess/VisualizeTexture.cpp
@@ -731,7 +731,7 @@ void FVisualizeTexture::DebugLog(bool bExtended)
 				{
 					FRHIResourceInfo Info;
 
-					FTextureRHIParamRef Texture = RT->GetRenderTargetItem().ShaderResourceTexture;
+					FTextureRHIRef Texture = RT->GetRenderTargetItem().ShaderResourceTexture;
 
 					if(!IsValidRef(Texture))
 					{
