* refactor: abuse of SceneColor

diff --git a/Engine/Source/Runtime/Renderer/Private/CompositionLighting/PostProcessAmbientOcclusion.cpp b/Engine/Source/Runtime/Renderer/Private/CompositionLighting/PostProcessAmbientOcclusion.cpp
--- a/Engine/Source/Runtime/Renderer/Private/CompositionLighting/PostProcessAmbientOcclusion.cpp
+++ b/Engine/Source/Runtime/Renderer/Private/CompositionLighting/PostProcessAmbientOcclusion.cpp
@@ -488,7 +488,7 @@ void FRCPassPostProcessAmbientOcclusion::Process(FRenderingCompositePassContext&
 	FIntPoint TexSize = InputDesc0->Extent;
 
 	// usually 1, 2, 4 or 8
-	uint32 ScaleToFullRes = GSceneRenderTargets.SceneColor->GetDesc().Extent.X / TexSize.X;
+	uint32 ScaleToFullRes = GSceneRenderTargets.GetBufferSizeXY().X / TexSize.X;
 
 	FIntRect ViewRect = FIntRect::DivideAndRoundUp(View.ViewRect, ScaleToFullRes);
 
diff --git a/Engine/Source/Runtime/Renderer/Private/PostProcess/PostProcessBokehDOF.cpp b/Engine/Source/Runtime/Renderer/Private/PostProcess/PostProcessBokehDOF.cpp
--- a/Engine/Source/Runtime/Renderer/Private/PostProcess/PostProcessBokehDOF.cpp
+++ b/Engine/Source/Runtime/Renderer/Private/PostProcess/PostProcessBokehDOF.cpp
@@ -608,7 +608,7 @@ void FRCPassPostProcessBokehDOF::Process(FRenderingCompositePassContext& Context
 	FIntPoint TexSize = InputDesc->Extent;
 
 	// usually 1, 2, 4 or 8
-	uint32 ScaleToFullRes = GSceneRenderTargets.SceneColor->GetDesc().Extent.X / TexSize.X;
+	uint32 ScaleToFullRes = GSceneRenderTargets.GetBufferSizeXY().X / TexSize.X;
 
 	// don't use DivideAndRoundUp as this could cause cause lookups into areas we don't have setup 
 	FIntRect LocalViewRect = View.ViewRect / ScaleToFullRes;
diff --git a/Engine/Source/Runtime/Renderer/Private/PostProcess/PostProcessBokehDOFRecombine.cpp b/Engine/Source/Runtime/Renderer/Private/PostProcess/PostProcessBokehDOFRecombine.cpp
--- a/Engine/Source/Runtime/Renderer/Private/PostProcess/PostProcessBokehDOFRecombine.cpp
+++ b/Engine/Source/Runtime/Renderer/Private/PostProcess/PostProcessBokehDOFRecombine.cpp
@@ -130,14 +130,15 @@ void FRCPassPostProcessBokehDOFRecombine::Process(FRenderingCompositePassContext
 
 	SCOPED_DRAW_EVENTF(BokehDOFRecombine, DEC_SCENE_ITEMS, TEXT("BokehDOFRecombine#%d"), Method);
 
-	const FPooledRenderTargetDesc* InputDesc = GetInputDesc(ePId_Input1);
+	const FPooledRenderTargetDesc* InputDesc0 = GetInputDesc(ePId_Input0);
+	const FPooledRenderTargetDesc* InputDesc1 = GetInputDesc(ePId_Input1);
 	
 	const FSceneView& View = Context.View;
 
-	FIntPoint TexSize = InputDesc ? InputDesc->Extent : GSceneRenderTargets.SceneColor->GetDesc().Extent;
+	FIntPoint TexSize = InputDesc1 ? InputDesc1->Extent : InputDesc0->Extent;
 
 	// usually 1, 2, 4 or 8
-	uint32 ScaleToFullRes = GSceneRenderTargets.SceneColor->GetDesc().Extent.X / TexSize.X;
+	uint32 ScaleToFullRes = GSceneRenderTargets.GetBufferSizeXY().X / TexSize.X;
 
 	FIntRect HalfResViewRect = FIntRect::DivideAndRoundUp(View.ViewRect, ScaleToFullRes);
 
diff --git a/Engine/Source/Runtime/Renderer/Private/PostProcess/PostProcessDOF.cpp b/Engine/Source/Runtime/Renderer/Private/PostProcess/PostProcessDOF.cpp
--- a/Engine/Source/Runtime/Renderer/Private/PostProcess/PostProcessDOF.cpp
+++ b/Engine/Source/Runtime/Renderer/Private/PostProcess/PostProcessDOF.cpp
@@ -271,7 +271,7 @@ void FRCPassPostProcessDOFRecombine::Process(FRenderingCompositePassContext& Con
 	FIntPoint TexSize = InputDesc->Extent;
 
 	// usually 1, 2, 4 or 8
-	uint32 ScaleToFullRes = GSceneRenderTargets.SceneColor->GetDesc().Extent.X / TexSize.X;
+	uint32 ScaleToFullRes = GSceneRenderTargets.GetBufferSizeXY().X / TexSize.X;
 
 	FIntRect HalfResViewRect = View.ViewRect / ScaleToFullRes;
 
diff --git a/Engine/Source/Runtime/Renderer/Private/PostProcess/PostProcessLensBlur.cpp b/Engine/Source/Runtime/Renderer/Private/PostProcess/PostProcessLensBlur.cpp
--- a/Engine/Source/Runtime/Renderer/Private/PostProcess/PostProcessLensBlur.cpp
+++ b/Engine/Source/Runtime/Renderer/Private/PostProcess/PostProcessLensBlur.cpp
@@ -176,7 +176,7 @@ void FRCPassPostProcessLensBlur::Process(FRenderingCompositePassContext& Context
 	FIntPoint TexSize = InputDesc->Extent;
 
 	// usually 1, 2, 4 or 8
-	uint32 ScaleToFullRes = GSceneRenderTargets.SceneColor->GetDesc().Extent.X / TexSize.X;
+	uint32 ScaleToFullRes = GSceneRenderTargets.GetBufferSizeXY().X / TexSize.X;
 
 	FIntRect ViewRect = FIntRect::DivideAndRoundUp(View.ViewRect, ScaleToFullRes);
 	FIntPoint ViewSize = ViewRect.Size();
diff --git a/Engine/Source/Runtime/Renderer/Private/PostProcess/PostProcessLensFlares.cpp b/Engine/Source/Runtime/Renderer/Private/PostProcess/PostProcessLensFlares.cpp
--- a/Engine/Source/Runtime/Renderer/Private/PostProcess/PostProcessLensFlares.cpp
+++ b/Engine/Source/Runtime/Renderer/Private/PostProcess/PostProcessLensFlares.cpp
@@ -129,8 +129,8 @@ void FRCPassPostProcessLensFlares::Process(FRenderingCompositePassContext& Conte
 	FIntPoint TexSize1 = InputDesc1->Extent;
 	FIntPoint TexSize2 = InputDesc2->Extent;
 
-	uint32 ScaleToFullRes1 = GSceneRenderTargets.SceneColor->GetDesc().Extent.X / TexSize1.X;
-	uint32 ScaleToFullRes2 = GSceneRenderTargets.SceneColor->GetDesc().Extent.X / TexSize2.X;
+	uint32 ScaleToFullRes1 = GSceneRenderTargets.GetBufferSizeXY().X / TexSize1.X;
+	uint32 ScaleToFullRes2 = GSceneRenderTargets.GetBufferSizeXY().X / TexSize2.X;
 
 	FIntRect ViewRect1 = FIntRect::DivideAndRoundUp(View.ViewRect, ScaleToFullRes1);
 	FIntRect ViewRect2 = FIntRect::DivideAndRoundUp(View.ViewRect, ScaleToFullRes2);
