Fixed unity build name conflicts

diff --git a/Engine/Source/Runtime/Renderer/Private/ReflectionEnvironmentCapture.cpp b/Engine/Source/Runtime/Renderer/Private/ReflectionEnvironmentCapture.cpp
--- a/Engine/Source/Runtime/Renderer/Private/ReflectionEnvironmentCapture.cpp
+++ b/Engine/Source/Runtime/Renderer/Private/ReflectionEnvironmentCapture.cpp
@@ -132,9 +132,9 @@ IMPLEMENT_SHADER_TYPE(,FDownsamplePS,TEXT("ReflectionEnvironmentShaders"),TEXT("
 const int32 NumFilterSamples = 1024;
 
 /** Pixel shader used for filtering a mip. */
-class FFilterPS : public FDownsamplePS
+class FCubeFilterPS : public FDownsamplePS
 {
-	DECLARE_SHADER_TYPE(FFilterPS,Global);
+	DECLARE_SHADER_TYPE(FCubeFilterPS,Global);
 public:
 
 	static bool ShouldCache(EShaderPlatform Platform)
@@ -148,14 +148,14 @@ public:
 		OutEnvironment.SetDefine(TEXT("NUM_FILTER_SAMPLES"), NumFilterSamples);
 	}
 
-	FFilterPS(const ShaderMetaType::CompiledShaderInitializerType& Initializer):
+	FCubeFilterPS(const ShaderMetaType::CompiledShaderInitializerType& Initializer):
 		FDownsamplePS(Initializer)
 	{
 		FilterRoughness.Bind(Initializer.ParameterMap,TEXT("FilterRoughness"));
 		AverageBrightnessTexture.Bind(Initializer.ParameterMap,TEXT("AverageBrightnessTexture"));
 		AverageBrightnessSampler.Bind(Initializer.ParameterMap,TEXT("AverageBrightnessSampler"));
 	}
-	FFilterPS() {}
+	FCubeFilterPS() {}
 
 	void SetParameters(int32 CubeFaceValue, int32 SourceMipIndexValue, float RoughnessValue, FSceneRenderTargetItem& SourceTextureValue)
 	{
@@ -187,29 +187,28 @@ private:
 };
 
 template< uint32 bNormalize >
-class TFilterPS : public FFilterPS
+class TCubeFilterPS : public FCubeFilterPS
 {
-	DECLARE_SHADER_TYPE(TFilterPS,Global);
+	DECLARE_SHADER_TYPE(TCubeFilterPS,Global);
 
 public:
 	static void ModifyCompilationEnvironment(EShaderPlatform Platform, FShaderCompilerEnvironment& OutEnvironment)
 	{
-		FFilterPS::ModifyCompilationEnvironment(Platform, OutEnvironment);
+		FCubeFilterPS::ModifyCompilationEnvironment(Platform, OutEnvironment);
 		OutEnvironment.SetDefine(TEXT("NORMALIZE"), bNormalize);
 	}
 
-	TFilterPS(const ShaderMetaType::CompiledShaderInitializerType& Initializer)
-	: FFilterPS(Initializer)
+	TCubeFilterPS(const ShaderMetaType::CompiledShaderInitializerType& Initializer)
+	: FCubeFilterPS(Initializer)
 	{}
 
-	TFilterPS() {}
+	TCubeFilterPS() {}
 };
 
-IMPLEMENT_SHADER_TYPE(template<>,TFilterPS<0>,TEXT("ReflectionEnvironmentShaders"),TEXT("FilterPS"),SF_Pixel);
-IMPLEMENT_SHADER_TYPE(template<>,TFilterPS<1>,TEXT("ReflectionEnvironmentShaders"),TEXT("FilterPS"),SF_Pixel);
+IMPLEMENT_SHADER_TYPE(template<>,TCubeFilterPS<0>,TEXT("ReflectionEnvironmentShaders"),TEXT("FilterPS"),SF_Pixel);
+IMPLEMENT_SHADER_TYPE(template<>,TCubeFilterPS<1>,TEXT("ReflectionEnvironmentShaders"),TEXT("FilterPS"),SF_Pixel);
 
 static FGlobalBoundShaderState DownsampleBoundShaderState;
-static FGlobalBoundShaderState FilterBoundShaderState;
 
 /** Computes the average brightness of a 1x1 mip of a cubemap. */
 class FComputeBrightnessPS : public FGlobalShader
@@ -443,22 +442,20 @@ void FilterReflectionEnvironment(FSHVectorRGB3* OutIrradianceEnvironmentMap)
 				TShaderMapRef<FScreenVSForGS> VertexShader(GetGlobalShaderMap());
 				TShaderMapRef<FDownsampleGS> GeometryShader(GetGlobalShaderMap());
 				
-				FFilterPS* PixelShader;
+				FCubeFilterPS* PixelShader;
 				if( bNormalize )
 				{
-					PixelShader = *TShaderMapRef< TFilterPS<1> >( GetGlobalShaderMap() );
+					PixelShader = *TShaderMapRef< TCubeFilterPS<1> >( GetGlobalShaderMap() );
 					static FGlobalBoundShaderState BoundShaderState;
 					SetGlobalBoundShaderState(BoundShaderState, GFilterVertexDeclaration.VertexDeclarationRHI, *VertexShader, PixelShader, *GeometryShader);
 				}
 				else
 				{
-					PixelShader = *TShaderMapRef< TFilterPS<0> >( GetGlobalShaderMap() );
+					PixelShader = *TShaderMapRef< TCubeFilterPS<0> >( GetGlobalShaderMap() );
 					static FGlobalBoundShaderState BoundShaderState;
 					SetGlobalBoundShaderState(BoundShaderState, GFilterVertexDeclaration.VertexDeclarationRHI, *VertexShader, PixelShader, *GeometryShader);
 				}
 
-				//SetGlobalBoundShaderState(FilterBoundShaderState, GFilterVertexDeclaration.VertexDeclarationRHI, *VertexShader, PixelShader, *GeometryShader);
-
 				for (int32 CubeFace = 0; CubeFace < CubeFace_MAX; CubeFace++)
 				{
 					SCOPED_DRAW_EVENT(FilterCubeFace, DEC_SCENE_ITEMS);
@@ -492,25 +489,23 @@ void FilterReflectionEnvironment(FSHVectorRGB3* OutIrradianceEnvironmentMap)
 					RHISetBlendState(TStaticBlendState<>::GetRHI());
 
 					TShaderMapRef<FScreenVSForGS> VertexShader(GetGlobalShaderMap());
-					TShaderMapRef< TFilterPS<1> > CaptureCubemapArrayPixelShader(GetGlobalShaderMap());
+					TShaderMapRef< TCubeFilterPS<1> > CaptureCubemapArrayPixelShader(GetGlobalShaderMap());
 					//FFilterPS* PixelShader = (FFilterPS*)*CaptureCubemapArrayPixelShader;
 
-					FFilterPS* PixelShader;
+					FCubeFilterPS* PixelShader;
 					if( bNormalize )
 					{
-						PixelShader = *TShaderMapRef< TFilterPS<1> >( GetGlobalShaderMap() );
+						PixelShader = *TShaderMapRef< TCubeFilterPS<1> >( GetGlobalShaderMap() );
 						static FGlobalBoundShaderState BoundShaderState;
 						SetGlobalBoundShaderState(BoundShaderState, GFilterVertexDeclaration.VertexDeclarationRHI, *VertexShader, PixelShader);
 					}
 					else
 					{
-						PixelShader = *TShaderMapRef< TFilterPS<0> >( GetGlobalShaderMap() );
+						PixelShader = *TShaderMapRef< TCubeFilterPS<0> >( GetGlobalShaderMap() );
 						static FGlobalBoundShaderState BoundShaderState;
 						SetGlobalBoundShaderState(BoundShaderState, GFilterVertexDeclaration.VertexDeclarationRHI, *VertexShader, PixelShader);
 					}
 
-					//SetGlobalBoundShaderState(FilterBoundShaderState, GFilterVertexDeclaration.VertexDeclarationRHI, *VertexShader, PixelShader);
-
 					PixelShader->SetParameters(CubeFace, MipIndex, Roughness, EffectiveSource);
 
 					DrawRectangle( 
