* fixed build error

diff --git a/Engine/Source/Runtime/Renderer/Private/PostProcess/VisualizeTexture.cpp b/Engine/Source/Runtime/Renderer/Private/PostProcess/VisualizeTexture.cpp
--- a/Engine/Source/Runtime/Renderer/Private/PostProcess/VisualizeTexture.cpp
+++ b/Engine/Source/Runtime/Renderer/Private/PostProcess/VisualizeTexture.cpp
@@ -735,7 +735,7 @@ void FVisualizeTexture::DebugLog(bool bExtended)
 
 					if(!IsValidRef(Texture))
 					{
-						RT->GetRenderTargetItem().TargetableTexture;
+						Texture = RT->GetRenderTargetItem().TargetableTexture;
 					}
 
 					if(IsValidRef(Texture))
