Enable transluscent self-shadowing on ShaderModel 4, which works on Mac OpenGL and AFAICT on D3D SM4 & PC GL too. This fixes the translucent cloud with shadows in ContentExamples/Effects (example 1.8) and probably others. Fallback to non-self-shadowed rendering on < SM4 so that something is rendered.

diff --git a/Engine/Source/Runtime/Renderer/Private/BasePassRendering.h b/Engine/Source/Runtime/Renderer/Private/BasePassRendering.h
--- a/Engine/Source/Runtime/Renderer/Private/BasePassRendering.h
+++ b/Engine/Source/Runtime/Renderer/Private/BasePassRendering.h
@@ -848,21 +848,19 @@ void ProcessBasePassMesh(
 	// This happens for example when rendering a thumbnail of an opaque post process material that uses scene textures
 	if (!(Parameters.TextureMode == ESceneRenderTargetsMode::DontSet && bNeedsSceneTextures))
 	{
-		if (bIsLitMaterial && Action.UseTranslucentSelfShadowing())
+		// Render self-shadowing only for >= SM4 and fallback to non-shadowed for lesser shader models
+		if (bIsLitMaterial && Action.UseTranslucentSelfShadowing() && GRHIFeatureLevel >= ERHIFeatureLevel::SM4)
 		{
-			if (GRHIFeatureLevel >= ERHIFeatureLevel::SM5)
+			if (IsIndirectLightingCacheAllowed()
+				&& Action.AllowIndirectLightingCache()
+				&& Parameters.PrimitiveSceneProxy)
 			{
-				if (IsIndirectLightingCacheAllowed()
-					&& Action.AllowIndirectLightingCache()
-					&& Parameters.PrimitiveSceneProxy)
-				{
-					// Apply cached point indirect lighting as well as self shadowing if needed
-					Action.template Process<FSelfShadowedCachedPointIndirectLightingPolicy>(Parameters, FSelfShadowedCachedPointIndirectLightingPolicy(), FSelfShadowedTranslucencyPolicy::ElementDataType(Action.GetTranslucentSelfShadow()));
-				}
-				else
-				{
-					Action.template Process<FSelfShadowedTranslucencyPolicy>(Parameters, FSelfShadowedTranslucencyPolicy(), FSelfShadowedTranslucencyPolicy::ElementDataType(Action.GetTranslucentSelfShadow()));
-				}
+				// Apply cached point indirect lighting as well as self shadowing if needed
+				Action.template Process<FSelfShadowedCachedPointIndirectLightingPolicy>(Parameters, FSelfShadowedCachedPointIndirectLightingPolicy(), FSelfShadowedTranslucencyPolicy::ElementDataType(Action.GetTranslucentSelfShadow()));
+			}
+			else
+			{
+				Action.template Process<FSelfShadowedTranslucencyPolicy>(Parameters, FSelfShadowedTranslucencyPolicy(), FSelfShadowedTranslucencyPolicy::ElementDataType(Action.GetTranslucentSelfShadow()));
 			}
 		}
 		else
