Fix accidental change to FDeferredShadingSceneRenderer::RenderDeferredReflections that left garbage in the LightAccumulation buffer.

diff --git a/Engine/Source/Runtime/Renderer/Private/ReflectionEnvironment.cpp b/Engine/Source/Runtime/Renderer/Private/ReflectionEnvironment.cpp
--- a/Engine/Source/Runtime/Renderer/Private/ReflectionEnvironment.cpp
+++ b/Engine/Source/Runtime/Renderer/Private/ReflectionEnvironment.cpp
@@ -708,7 +708,7 @@ void FDeferredShadingSceneRenderer::RenderDeferredReflections()
 					RHISetRenderTarget(GSceneRenderTargets.LightAccumulation->GetRenderTargetItem().TargetableTexture, NULL);
 
 					// Clear to no reflection contribution, alpha of 1 indicates full background contribution
-					RHIClear(true, FLinearColor(0, 0, 0, 1), false, 0, false, 0, View.ViewRect);
+					RHIClear(true, FLinearColor(0, 0, 0, 1), false, 0, false, 0, FIntRect());
 
 					RHISetViewport(View.ViewRect.Min.X, View.ViewRect.Min.Y, 0.0f, View.ViewRect.Max.X, View.ViewRect.Max.Y, 1.0f);
 
