Particle lights now work with camera offset (per-view).

#branch UE4
#proj Engine, Renderer

#add Created a class FSimpleLightPerViewEntry and relocated the positional vector from FSimpleLightEntry.
#add Created FSimpleLightArray which contains a FSimpleLightEntry and FSimpleLightPerViewEntry TArray - this is just so to keep function parameter lists uncluttered when passing the data around.
#change ViewFamily is now passed to GatherSimpleLights and GatherParticleLightData allowing us to calculate the camera offset for each view. This also allows for future additions that are view-dependant.
#change Camera offset is calculated in GatherParticleLightData for each view. When camera offset isn't enabled FSimpleLightPerViewEntry is shared for all views in order to save memory

diff --git a/Engine/Source/Runtime/Engine/Private/Particles/ParticleSystemRender.cpp b/Engine/Source/Runtime/Engine/Private/Particles/ParticleSystemRender.cpp
--- a/Engine/Source/Runtime/Engine/Private/Particles/ParticleSystemRender.cpp
+++ b/Engine/Source/Runtime/Engine/Private/Particles/ParticleSystemRender.cpp
@@ -1195,7 +1195,7 @@ void FDynamicSpriteEmitterData::PreRenderView(FParticleSystemSceneProxy* Proxy,
 	}
 }
 
-void GatherParticleLightData(const FDynamicSpriteEmitterReplayDataBase& Source, const FMatrix& InLocalToWorld, TArray<FSimpleLightEntry, SceneRenderingAllocator>& OutParticleLights)
+void GatherParticleLightData(const FDynamicSpriteEmitterReplayDataBase& Source, const FMatrix& InLocalToWorld, const FSceneViewFamily& InViewFamily, FSimpleLightArray& OutParticleLights)
 {
 	if (Source.LightDataOffset != 0)
 	{
@@ -1208,8 +1208,18 @@ void GatherParticleLightData(const FDynamicSpriteEmitterReplayDataBase& Source,
 			ParticleCount = Source.MaxDrawCount;
 		}
 
-		OutParticleLights.Reserve(OutParticleLights.Num() + ParticleCount);
-
+		OutParticleLights.InstanceData.Reserve(OutParticleLights.InstanceData.Num() + ParticleCount);
+		
+		// Reserve memory for per-view data. If camera offset is not used then all views can share per-view data in order to save memory.
+		if(Source.CameraPayloadOffset != 0)
+		{
+			OutParticleLights.PerViewData.Reserve(OutParticleLights.PerViewData.Num() + ParticleCount * InViewFamily.Views.Num());
+		}
+		else
+		{
+			OutParticleLights.PerViewData.Reserve(OutParticleLights.PerViewData.Num() + ParticleCount);
+		}
+		
 		const uint8* ParticleData = Source.ParticleData.GetData();
 		const uint16* ParticleIndices = Source.ParticleIndices.GetData();
 
@@ -1229,6 +1239,19 @@ void GatherParticleLightData(const FDynamicSpriteEmitterReplayDataBase& Source,
 			if (LightPayload->bValid)
 			{
 				const FVector2D Size = GetParticleSize(Particle, Source);
+				
+				FSimpleLightEntry ParticleLight;
+				ParticleLight.Radius =  LightPayload->RadiusScale * (Size.X + Size.Y) / 2.0f;
+				ParticleLight.Color = FVector(Particle.Color) * Particle.Color.A * LightPayload->ColorScale;
+				ParticleLight.Exponent = LightPayload->LightExponent;
+				ParticleLight.bAffectTranslucency = LightPayload->bAffectsTranslucency;
+
+				// Early out if the light will have no visible contribution
+				if (ParticleLight.Radius <= KINDA_SMALL_NUMBER
+					&& ParticleLight.Color.GetMax() <= KINDA_SMALL_NUMBER)
+				{
+					continue;
+				}
 
 				FVector ParticlePosition = Particle.Location;
 				FOrbitChainModuleInstancePayload* LocalOrbitPayload = NULL;
@@ -1238,28 +1261,40 @@ void GatherParticleLightData(const FDynamicSpriteEmitterReplayDataBase& Source,
 
 				ApplyOrbitToPosition(Particle, Source, InLocalToWorld, LocalOrbitPayload, Unused, Unused2, ParticlePosition, Unused3);
 
-				FSimpleLightEntry ParticleLight;
 
 				FVector LightPosition = Source.bUseLocalSpace ? FVector(InLocalToWorld.TransformPosition(ParticlePosition)) : ParticlePosition;
-				ParticleLight.PositionAndRadius = FVector4(LightPosition, LightPayload->RadiusScale * (Size.X + Size.Y) / 2.0f);
-				ParticleLight.Color = FVector(Particle.Color) * Particle.Color.A * LightPayload->ColorScale;
-				ParticleLight.Exponent = LightPayload->LightExponent;
-				ParticleLight.bAffectTranslucency = LightPayload->bAffectsTranslucency;
+				
+				// Calculate light positions per-view if we are using a camera offset.
+				if (Source.CameraPayloadOffset != 0)
+				{
+					for(int32 ViewIndex = 0; ViewIndex < InViewFamily.Views.Num(); ++ViewIndex)
+					{
+						FSimpleLightPerViewEntry PerViewData;
+						const FVector& ViewOrigin = InViewFamily.Views[ViewIndex]->ViewMatrices.ViewOrigin;
+						FVector CameraOffset = GetCameraOffsetFromPayload(Source.CameraPayloadOffset, Particle, ParticlePosition, ViewOrigin);
+						PerViewData.Position = LightPosition + CameraOffset;
 
-				// Only create a light if it will have visible contribution
-				if (ParticleLight.PositionAndRadius.W > KINDA_SMALL_NUMBER
-					&& ParticleLight.Color.GetMax() > KINDA_SMALL_NUMBER)
+						OutParticleLights.PerViewData.Add(PerViewData);
+					}
+				}
+				else
 				{
-					OutParticleLights.Add(ParticleLight);
+					// When not using camera-offset, output one position for all views to share. 
+					FSimpleLightPerViewEntry PerViewData;
+					PerViewData.Position = LightPosition;
+					OutParticleLights.PerViewData.Add(PerViewData);
 				}
+
+				// Add an entry for the light instance.
+				OutParticleLights.InstanceData.Add(ParticleLight);
 			}
 		}
 	}
 }
 
-void FDynamicSpriteEmitterData::GatherSimpleLights(const FParticleSystemSceneProxy* Proxy, TArray<FSimpleLightEntry, SceneRenderingAllocator>& OutParticleLights) const
+void FDynamicSpriteEmitterData::GatherSimpleLights(const FParticleSystemSceneProxy* Proxy, const FSceneViewFamily& ViewFamily, FSimpleLightArray& OutParticleLights) const
 {
-	GatherParticleLightData(Source, Proxy->GetLocalToWorld(), OutParticleLights);
+	GatherParticleLightData(Source, Proxy->GetLocalToWorld(), ViewFamily, OutParticleLights);
 }
 
 /**
@@ -1938,9 +1973,9 @@ void FDynamicMeshEmitterData::PreRenderView(FParticleSystemSceneProxy* Proxy, co
 	}
 }
 
-void FDynamicMeshEmitterData::GatherSimpleLights(const FParticleSystemSceneProxy* Proxy, TArray<FSimpleLightEntry, SceneRenderingAllocator>& OutParticleLights) const
+void FDynamicMeshEmitterData::GatherSimpleLights(const FParticleSystemSceneProxy* Proxy, const FSceneViewFamily& ViewFamily, FSimpleLightArray& OutParticleLights) const
 {
-	GatherParticleLightData(Source, Proxy->GetLocalToWorld(), OutParticleLights);
+	GatherParticleLightData(Source, Proxy->GetLocalToWorld(), ViewFamily, OutParticleLights);
 }
 
 void FDynamicMeshEmitterData::GetParticleTransform(FBaseParticle& InParticle, const FVector& CameraPosition, const FVector& CameraFacingOpVector, const FQuat& PointToLockedAxis, 
@@ -7228,7 +7263,7 @@ void FParticleSystemSceneProxy::PreRenderView(const FSceneViewFamily* ViewFamily
 	}
 }
 
-void FParticleSystemSceneProxy::GatherSimpleLights(TArray<FSimpleLightEntry, SceneRenderingAllocator>& OutParticleLights) const
+void FParticleSystemSceneProxy::GatherSimpleLights(const FSceneViewFamily& ViewFamily, FSimpleLightArray& OutParticleLights) const
 {
 	if (DynamicData != NULL)
 	{
@@ -7237,7 +7272,7 @@ void FParticleSystemSceneProxy::GatherSimpleLights(TArray<FSimpleLightEntry, Sce
 			const FDynamicEmitterDataBase* DynamicEmitterData = DynamicData->DynamicEmitterDataArray[EmitterIndex];
 			if (DynamicEmitterData)
 			{
-				DynamicEmitterData->GatherSimpleLights(this, OutParticleLights);
+				DynamicEmitterData->GatherSimpleLights(this, ViewFamily, OutParticleLights);
 			}
 		}
 	}
diff --git a/Engine/Source/Runtime/Engine/Public/ParticleHelper.h b/Engine/Source/Runtime/Engine/Public/ParticleHelper.h
--- a/Engine/Source/Runtime/Engine/Public/ParticleHelper.h
+++ b/Engine/Source/Runtime/Engine/Public/ParticleHelper.h
@@ -1367,7 +1367,7 @@ struct FDynamicEmitterDataBase
 	virtual void PreRenderView(FParticleSystemSceneProxy* Proxy, const FSceneViewFamily* ViewFamily, const uint32 VisibilityMap, int32 FrameNumber) {}
 
 	/** Callback from the renderer to gather simple lights that this proxy wants renderered. */
-	virtual void GatherSimpleLights(const FParticleSystemSceneProxy* Proxy, TArray<FSimpleLightEntry, SceneRenderingAllocator>& OutParticleLights) const {}
+	virtual void GatherSimpleLights(const FParticleSystemSceneProxy* Proxy, const FSceneViewFamily& ViewFamily, FSimpleLightArray& OutParticleLights) const {}
 
 	/** Returns the source data for this particle system */
 	virtual const FDynamicEmitterReplayDataBase& GetSource() const = 0;
@@ -1759,7 +1759,7 @@ struct FDynamicSpriteEmitterData : public FDynamicSpriteEmitterDataBase
 	virtual void PreRenderView(FParticleSystemSceneProxy* Proxy, const FSceneViewFamily* ViewFamily, const uint32 VisibilityMap, int32 FrameNumber) OVERRIDE;
 
 	/** Gathers simple lights for this emitter. */
-	virtual void GatherSimpleLights(const FParticleSystemSceneProxy* Proxy, TArray<FSimpleLightEntry, SceneRenderingAllocator>& OutParticleLights) const OVERRIDE;
+	virtual void GatherSimpleLights(const FParticleSystemSceneProxy* Proxy, const FSceneViewFamily& ViewFamily, FSimpleLightArray& OutParticleLights) const OVERRIDE;
 
 	/**
 	 *	Render thread only draw call
@@ -1948,7 +1948,7 @@ struct FDynamicMeshEmitterData : public FDynamicSpriteEmitterDataBase
 	virtual void PreRenderView(FParticleSystemSceneProxy* Proxy, const FSceneViewFamily* ViewFamily, const uint32 VisibilityMap, int32 FrameNumber); 
 
 	/** Gathers simple lights for this emitter. */
-	virtual void GatherSimpleLights(const FParticleSystemSceneProxy* Proxy, TArray<FSimpleLightEntry, SceneRenderingAllocator>& OutParticleLights) const OVERRIDE;
+	virtual void GatherSimpleLights(const FParticleSystemSceneProxy* Proxy, const FSceneViewFamily& ViewFamily, FSimpleLightArray& OutParticleLights) const OVERRIDE;
 
 	/**
 	 *	Get the vertex stride for the dynamic rendering data
@@ -2610,7 +2610,7 @@ public:
 	virtual void PreRenderView(const FSceneViewFamily* ViewFamily, const uint32 VisibilityMap, int32 FrameNumber);
 
 	/** Gathers simple lights for this emitter. */
-	virtual void GatherSimpleLights(TArray<FSimpleLightEntry, SceneRenderingAllocator>& OutParticleLights) const OVERRIDE;
+	virtual void GatherSimpleLights(const FSceneViewFamily& ViewFamily, FSimpleLightArray& OutParticleLights) const OVERRIDE;
 
 	/**
 	 *	Called when the rendering thread adds the proxy to the scene.
diff --git a/Engine/Source/Runtime/Engine/Public/PrimitiveSceneProxy.h b/Engine/Source/Runtime/Engine/Public/PrimitiveSceneProxy.h
--- a/Engine/Source/Runtime/Engine/Public/PrimitiveSceneProxy.h
+++ b/Engine/Source/Runtime/Engine/Public/PrimitiveSceneProxy.h
@@ -13,12 +13,52 @@
 class FSimpleLightEntry
 {
 public:
-	FVector4 PositionAndRadius;
 	FVector Color;
+	float Radius;
 	float Exponent;
 	bool bAffectTranslucency;
 };
 
+/** Data for a simple dynamic light which could change per-view. */
+class FSimpleLightPerViewEntry
+{
+public:
+	FVector Position;
+};
+
+/** Data pertaining to a set of simple dynamic lights */
+class FSimpleLightArray
+{
+public:
+	/** Data per simple dynamic light instance, independent of view */
+	TArray<FSimpleLightEntry, SceneRenderingAllocator> InstanceData;
+	/** Per-view data for each light */
+	TArray<FSimpleLightPerViewEntry, SceneRenderingAllocator> PerViewData;
+
+public:
+
+	/**
+	* Returns the per-view data for a simple light entry.
+	* @param LightIndex - The index of the simple light
+	* @param ViewIndex - The index of the view
+	* @param NumViews - The number of views in the ViewFamily.
+	*/
+	inline const FSimpleLightPerViewEntry& GetViewDependentData(int32 LightIndex, int32 ViewIndex, int32 NumViews) const
+	{
+		// If InstanceData has an equal number of elements to PerViewData then all views share the same PerViewData.
+		if (InstanceData.Num() == PerViewData.Num())
+		{
+			return PerViewData[LightIndex];
+		}
+		else 
+		{
+			// Calculate strided per-view index.
+			const int32 PerViewDataIndex = (LightIndex * NumViews) + ViewIndex;
+			return PerViewData[PerViewDataIndex];
+		}
+	}
+};
+
 namespace EDrawDynamicFlags
 {
 	enum Type
@@ -125,7 +165,7 @@ public:
 	virtual void PreRenderView(const FSceneViewFamily* ViewFamily, const uint32 VisibilityMap, int32 FrameNumber) {}
 
 	/** Callback from the renderer to gather simple lights that this proxy wants renderered. */
-	virtual void GatherSimpleLights(TArray<FSimpleLightEntry, SceneRenderingAllocator>& OutSimpleLights) const {}
+	virtual void GatherSimpleLights(const FSceneViewFamily& ViewFamily, FSimpleLightArray& OutParticleLights) const {}
 
 	/**
 	 *	Determines the relevance of this primitive's elements to the given light.
diff --git a/Engine/Source/Runtime/Renderer/Private/DeferredShadingRenderer.h b/Engine/Source/Runtime/Renderer/Private/DeferredShadingRenderer.h
--- a/Engine/Source/Runtime/Renderer/Private/DeferredShadingRenderer.h
+++ b/Engine/Source/Runtime/Renderer/Private/DeferredShadingRenderer.h
@@ -183,7 +183,7 @@ private:
 	bool ShouldUseTiledDeferred(int32 NumUnshadowedLights, int32 NumSimpleLights) const;
 
 	/** Renders the lights in SortedLights in the range [0, NumUnshadowedLights) using tiled deferred shading. */
-	void RenderTiledDeferredLighting(const TArray<FSortedLightSceneInfo, SceneRenderingAllocator>& SortedLights, int32 NumUnshadowedLights, const TArray<FSimpleLightEntry, SceneRenderingAllocator>& SimpleLights);
+	void RenderTiledDeferredLighting(const TArray<FSortedLightSceneInfo, SceneRenderingAllocator>& SortedLights, int32 NumUnshadowedLights, const FSimpleLightArray& SimpleLights);
 
 	/** Renders the scene's lighting. */
 	void RenderLights();
@@ -294,7 +294,7 @@ private:
 	void RenderLight(const FLightSceneInfo* LightSceneInfo, bool bRenderOverlap, bool bIssueDrawEvent);
 
 	/** Renders an array of simple lights using standard deferred shading. */
-	void RenderSimpleLightsStandardDeferred(const TArray<FSimpleLightEntry, SceneRenderingAllocator>& SimpleLights);
+	void RenderSimpleLightsStandardDeferred(const FSimpleLightArray& SimpleLights);
 
 	/** Clears the translucency lighting volumes before light accumulation. */
 	void ClearTranslucentVolumeLighting();
@@ -318,7 +318,7 @@ private:
 	void InjectTranslucentVolumeLightingArray(const TArray<FSortedLightSceneInfo, SceneRenderingAllocator>& SortedLights, int32 NumLights);
 
 	/** Accumulates direct lighting for simple lights. */
-	void InjectSimpleTranslucentVolumeLightingArray(const TArray<FSimpleLightEntry, SceneRenderingAllocator>& SimpleLights);
+	void InjectSimpleTranslucentVolumeLightingArray(const FSimpleLightArray& SimpleLights);
 
 	/** Filters the translucency lighting volumes to reduce aliasing. */
 	void FilterTranslucentVolumeLighting();
diff --git a/Engine/Source/Runtime/Renderer/Private/LightRendering.cpp b/Engine/Source/Runtime/Renderer/Private/LightRendering.cpp
--- a/Engine/Source/Runtime/Renderer/Private/LightRendering.cpp
+++ b/Engine/Source/Runtime/Renderer/Private/LightRendering.cpp
@@ -65,11 +65,11 @@ public:
 		SetDeferredLightParameters(ShaderRHI, GetUniformBufferParameter<FDeferredLightUniformStruct>(), LightSceneInfo, View);
 	}
 
-	void SetParametersSimpleLight(const FSceneView& View, const FSimpleLightEntry& SimpleLight)
+	void SetParametersSimpleLight(const FSceneView& View, const FSimpleLightEntry& SimpleLight, const FSimpleLightPerViewEntry& SimpleLightPerViewData)
 	{
 		const FPixelShaderRHIParamRef ShaderRHI = GetPixelShader();
 		SetParametersBase(ShaderRHI, View, NULL);
-		SetSimpleDeferredLightParameters(ShaderRHI, GetUniformBufferParameter<FDeferredLightUniformStruct>(), SimpleLight, View);
+		SetSimpleDeferredLightParameters(ShaderRHI, GetUniformBufferParameter<FDeferredLightUniformStruct>(), SimpleLight, SimpleLightPerViewData, View);
 	}
 
 	virtual bool Serialize(FArchive& Ar)
@@ -219,7 +219,7 @@ IMPLEMENT_SHADER_TYPE(template<>, TDeferredLightOverlapPS<true>, TEXT("Stationar
 IMPLEMENT_SHADER_TYPE(template<>, TDeferredLightOverlapPS<false>, TEXT("StationaryLightOverlapShaders"), TEXT("OverlapDirectionalPixelMain"), SF_Pixel);
 
 /** Gathers simple lights from visible primtives in the passed in views. */
-void GatherSimpleLights(const TArray<FViewInfo>& Views, TArray<FSimpleLightEntry, SceneRenderingAllocator>& SimpleLights)
+void GatherSimpleLights(const FSceneViewFamily& ViewFamily, const TArray<FViewInfo>& Views, FSimpleLightArray& SimpleLights)
 {
 	TArray<const FPrimitiveSceneInfo*, SceneRenderingAllocator> PrimitivesWithSimpleLights;
 
@@ -245,7 +245,7 @@ void GatherSimpleLights(const TArray<FViewInfo>& Views, TArray<FSimpleLightEntry
 	for (int32 PrimitiveIndex = 0; PrimitiveIndex < PrimitivesWithSimpleLights.Num(); PrimitiveIndex++)
 	{
 		const FPrimitiveSceneInfo* Primitive = PrimitivesWithSimpleLights[PrimitiveIndex];
-		Primitive->Proxy->GatherSimpleLights(SimpleLights);
+		Primitive->Proxy->GatherSimpleLights(ViewFamily, SimpleLights);
 	}
 }
 
@@ -287,8 +287,8 @@ void FDeferredShadingSceneRenderer::RenderLights()
 	SCOPE_CYCLE_COUNTER(STAT_LightingDrawTime);
 	SCOPE_CYCLE_COUNTER(STAT_LightRendering);
 
-	TArray<FSimpleLightEntry, SceneRenderingAllocator> SimpleLights;
-	GatherSimpleLights(Views, SimpleLights);
+	FSimpleLightArray SimpleLights;
+	GatherSimpleLights(ViewFamily, Views, SimpleLights);
 
 	TArray<FSortedLightSceneInfo, SceneRenderingAllocator> SortedLights;
 	SortedLights.Empty(Scene->Lights.Num());
@@ -399,19 +399,19 @@ void FDeferredShadingSceneRenderer::RenderLights()
 				}
 
 				// Use tiled deferred shading on any unshadowed lights without a texture light profile
-				if (!ShouldUseTiledDeferred(SupportedByTiledDeferredLightEnd, SimpleLights.Num()) || bAnyUnsupportedByTiledDeferred || bAnyViewIsStereo)
+				if (!ShouldUseTiledDeferred(SupportedByTiledDeferredLightEnd, SimpleLights.InstanceData.Num()) || bAnyUnsupportedByTiledDeferred || bAnyViewIsStereo)
 				{
 					NumSortedLightsTiledDeferred = 0;
 				}
 
-				if (NumSortedLightsTiledDeferred > 0 || SimpleLights.Num() > 0)
+				if (NumSortedLightsTiledDeferred > 0 || SimpleLights.InstanceData.Num() > 0)
 				{
 					// Update the range that needs to be processed by standard deferred to exclude the lights done with tiled
 					StandardDeferredStart = NumSortedLightsTiledDeferred;
 					RenderTiledDeferredLighting(SortedLights, NumSortedLightsTiledDeferred, SimpleLights);
 				}
 			}
-			else if (SimpleLights.Num() > 0)
+			else if (SimpleLights.InstanceData.Num() > 0)
 			{
 				GSceneRenderTargets.BeginRenderingSceneColor();
 				RenderSimpleLightsStandardDeferred(SimpleLights);
@@ -443,7 +443,7 @@ void FDeferredShadingSceneRenderer::RenderLights()
 					InjectTranslucentVolumeLightingArray(SortedLights, AttenuationLightStart);
 				}
 				
-				if (SimpleLights.Num() > 0)
+				if (SimpleLights.InstanceData.Num() > 0)
 				{
 					SCOPED_DRAW_EVENT(InjectSimpleLightsTranslucentLighting, DEC_SCENE_ITEMS);
 					InjectSimpleTranslucentVolumeLightingArray(SimpleLights);
@@ -673,20 +673,22 @@ template<bool bUseIESProfile, bool bRadialAttenuation, bool bInverseSquaredFallo
 static void SetShaderTemplLightingSimple(
 	const FSceneView& View, 
 	FShader* VertexShader,
-	const FSimpleLightEntry& SimpleLight)
+	const FSimpleLightEntry& SimpleLight,
+	const FSimpleLightPerViewEntry& SimpleLightPerViewData)
 {
 	if(View.Family->EngineShowFlags.VisualizeLightCulling)
 	{
 		TShaderMapRef<TDeferredLightPS<false, bRadialAttenuation, false, true> > PixelShader(GetGlobalShaderMap());
 		SetGlobalBoundShaderState(PixelShader->GetBoundShaderState(), GetDeferredLightingVertexDeclaration<bRadialAttenuation>(), VertexShader, *PixelShader);
-		PixelShader->SetParametersSimpleLight(View, SimpleLight);
+		PixelShader->SetParametersSimpleLight(View, SimpleLight, SimpleLightPerViewData);
 	}
 	else
 	{
 		TShaderMapRef<TDeferredLightPS<bUseIESProfile, bRadialAttenuation, bInverseSquaredFalloff, false> > PixelShader(GetGlobalShaderMap());
 		SetGlobalBoundShaderState(PixelShader->GetBoundShaderState(), GetDeferredLightingVertexDeclaration<bRadialAttenuation>(), VertexShader, *PixelShader);
-		PixelShader->SetParametersSimpleLight(View, SimpleLight);
-	}}
+		PixelShader->SetParametersSimpleLight(View, SimpleLight, SimpleLightPerViewData);
+	}
+}
 
 /**
  * Used by RenderLights to render a light to the scene color buffer.
@@ -819,23 +821,25 @@ void FDeferredShadingSceneRenderer::RenderLight(const FLightSceneInfo* LightScen
 	}
 }
 
-void FDeferredShadingSceneRenderer::RenderSimpleLightsStandardDeferred(const TArray<FSimpleLightEntry, SceneRenderingAllocator>& SimpleLights)
+void FDeferredShadingSceneRenderer::RenderSimpleLightsStandardDeferred(const FSimpleLightArray& SimpleLights)
 {
 	SCOPE_CYCLE_COUNTER(STAT_DirectLightRenderingTime);
-	INC_DWORD_STAT_BY(STAT_NumLightsUsingStandardDeferred, SimpleLights.Num());
+	INC_DWORD_STAT_BY(STAT_NumLightsUsingStandardDeferred, SimpleLights.InstanceData.Num());
 	SCOPED_DRAW_EVENT(StandardDeferredSimpleLights, DEC_SCENE_ITEMS);
 	
 	// Use additive blending for color
 	RHISetBlendState(TStaticBlendState<CW_RGBA, BO_Add, BF_One, BF_One, BO_Add, BF_One, BF_One>::GetRHI());
 
-
-	for (int32 LightIndex = 0; LightIndex < SimpleLights.Num(); LightIndex++)
+	const int32 NumViews = Views.Num();
+	for (int32 LightIndex = 0; LightIndex < SimpleLights.InstanceData.Num(); LightIndex++)
 	{
-		const FSimpleLightEntry& SimpleLight = SimpleLights[LightIndex];
-		const FSphere LightBounds(SimpleLight.PositionAndRadius, SimpleLight.PositionAndRadius.W);
-
-		for (int32 ViewIndex = 0; ViewIndex < Views.Num(); ViewIndex++)
+		const FSimpleLightEntry& SimpleLight = SimpleLights.InstanceData[LightIndex];
+		
+		for (int32 ViewIndex = 0; ViewIndex < NumViews; ViewIndex++)
 		{
+			const FSimpleLightPerViewEntry& SimpleLightPerViewData = SimpleLights.GetViewDependentData(LightIndex, ViewIndex, NumViews);
+			const FSphere LightBounds(SimpleLightPerViewData.Position, SimpleLight.Radius);
+
 			FViewInfo& View = Views[ViewIndex];
 			// Set the device viewport for the view.
 			RHISetViewport(View.ViewRect.Min.X, View.ViewRect.Min.Y, 0.0f, View.ViewRect.Max.X, View.ViewRect.Max.Y, 1.0f);
@@ -847,12 +851,12 @@ void FDeferredShadingSceneRenderer::RenderSimpleLightsStandardDeferred(const TAr
 			if (SimpleLight.Exponent == 0)
 			{
 				// inverse squared
-				SetShaderTemplLightingSimple<false, true, true>(View, *VertexShader, SimpleLight);
+				SetShaderTemplLightingSimple<false, true, true>(View, *VertexShader, SimpleLight, SimpleLightPerViewData);
 			}
 			else
 			{
 				// light's exponent, not inverse squared
-				SetShaderTemplLightingSimple<false, true, false>(View, *VertexShader, SimpleLight);
+				SetShaderTemplLightingSimple<false, true, false>(View, *VertexShader, SimpleLight, SimpleLightPerViewData);
 			}
 
 			VertexShader->SetSimpleLightParameters(View, LightBounds);
diff --git a/Engine/Source/Runtime/Renderer/Private/LightRendering.h b/Engine/Source/Runtime/Renderer/Private/LightRendering.h
--- a/Engine/Source/Runtime/Renderer/Private/LightRendering.h
+++ b/Engine/Source/Runtime/Renderer/Private/LightRendering.h
@@ -98,11 +98,12 @@ void SetSimpleDeferredLightParameters(
 	const ShaderRHIParamRef ShaderRHI, 
 	const TShaderUniformBufferParameter<FDeferredLightUniformStruct>& DeferredLightUniformBufferParameter, 
 	const FSimpleLightEntry& SimpleLight,
+	const FSimpleLightPerViewEntry &SimpleLightPerViewData,
 	const FSceneView& View)
 {
 	FDeferredLightUniformStruct DeferredLightUniformsValue;
-	DeferredLightUniformsValue.LightPosition = SimpleLight.PositionAndRadius;
-	DeferredLightUniformsValue.LightInvRadius = 1.0f / SimpleLight.PositionAndRadius.W;
+	DeferredLightUniformsValue.LightPosition = SimpleLightPerViewData.Position;
+	DeferredLightUniformsValue.LightInvRadius = 1.0f / SimpleLight.Radius;
 	DeferredLightUniformsValue.LightColor = SimpleLight.Color;
 	DeferredLightUniformsValue.LightFalloffExponent = SimpleLight.Exponent;
 	DeferredLightUniformsValue.NormalizedLightDirection = FVector(1, 0, 0);
diff --git a/Engine/Source/Runtime/Renderer/Private/TiledDeferredLightRendering.cpp b/Engine/Source/Runtime/Renderer/Private/TiledDeferredLightRendering.cpp
--- a/Engine/Source/Runtime/Renderer/Private/TiledDeferredLightRendering.cpp
+++ b/Engine/Source/Runtime/Renderer/Private/TiledDeferredLightRendering.cpp
@@ -105,9 +105,11 @@ public:
 
 	void SetParameters(
 		const FSceneView& View, 
+		int32 ViewIndex,
+		int32 NumViews,
 		const TArray<FSortedLightSceneInfo, SceneRenderingAllocator>& SortedLights, 
 		int32 NumLightsToRenderInSortedLights, 
-		const TArray<FSimpleLightEntry, SceneRenderingAllocator>& SimpleLights, 
+		const FSimpleLightArray& SimpleLights, 
 		int32 StartIndex, 
 		int32 NumThisPass)
 	{
@@ -186,8 +188,10 @@ public:
 			}
 			else
 			{
-				const FSimpleLightEntry& SimpleLight = SimpleLights[StartIndex + LightIndex - NumLightsToRenderInSortedLights];
-				LightData.LightPositionAndInvRadius[LightIndex] = FVector4(FVector(SimpleLight.PositionAndRadius), 1.0f / SimpleLight.PositionAndRadius.W);
+				int32 SimpleLightIndex = StartIndex + LightIndex - NumLightsToRenderInSortedLights;
+				const FSimpleLightEntry& SimpleLight = SimpleLights.InstanceData[SimpleLightIndex];
+				const FSimpleLightPerViewEntry& SimpleLightPerViewData = SimpleLights.GetViewDependentData(SimpleLightIndex, ViewIndex, NumViews);
+				LightData.LightPositionAndInvRadius[LightIndex] = FVector4(SimpleLightPerViewData.Position, 1.0f / SimpleLight.Radius);
 				LightData.LightColorAndFalloffExponent[LightIndex] = FVector4(SimpleLight.Color, SimpleLight.Exponent);
 				LightData2.LightDirectionAndSpotlightMaskAndMinRoughness[LightIndex] = FVector4(FVector(1, 0, 0), 0);
 				LightData2.SpotAnglesAndSourceRadiusAndSimpleLighting[LightIndex] = FVector4(-2, 1, 0, 1);
@@ -259,17 +263,19 @@ bool FDeferredShadingSceneRenderer::ShouldUseTiledDeferred(int32 NumUnshadowedLi
 
 template <bool bVisualizeLightCulling>
 static void SetShaderTemplTiledLighting(
-	const FSceneView& View, 
-	const TArray<FSortedLightSceneInfo, SceneRenderingAllocator>& SortedLights, 
-	int32 NumLightsToRenderInSortedLights, 
-	const TArray<FSimpleLightEntry, SceneRenderingAllocator>& SimpleLights, 
+	const FSceneView& View,
+	int32 ViewIndex,
+	int32 NumViews,
+	const TArray<FSortedLightSceneInfo, SceneRenderingAllocator>& SortedLights,
+	int32 NumLightsToRenderInSortedLights,
+	const FSimpleLightArray& SimpleLights,
 	int32 StartIndex, 
 	int32 NumThisPass)
 {
 	TShaderMapRef<FTiledDeferredLightingCS<bVisualizeLightCulling> > ComputeShader(GetGlobalShaderMap());
 	RHISetComputeShader(ComputeShader->GetComputeShader());
 
-	ComputeShader->SetParameters(View, SortedLights, NumLightsToRenderInSortedLights, SimpleLights, StartIndex, NumThisPass);
+	ComputeShader->SetParameters(View, ViewIndex, NumViews, SortedLights, NumLightsToRenderInSortedLights, SimpleLights, StartIndex, NumThisPass);
 
 	uint32 GroupSizeX = (View.ViewRect.Size().X + GDeferredLightTileSizeX - 1) / GDeferredLightTileSizeX;
 	uint32 GroupSizeY = (View.ViewRect.Size().Y + GDeferredLightTileSizeY - 1) / GDeferredLightTileSizeY;
@@ -279,23 +285,22 @@ static void SetShaderTemplTiledLighting(
 }
 
 
-void FDeferredShadingSceneRenderer::RenderTiledDeferredLighting(const TArray<FSortedLightSceneInfo, SceneRenderingAllocator>& SortedLights, int32 NumUnshadowedLights, const TArray<FSimpleLightEntry, SceneRenderingAllocator>& SimpleLights)
+void FDeferredShadingSceneRenderer::RenderTiledDeferredLighting(const TArray<FSortedLightSceneInfo, SceneRenderingAllocator>& SortedLights, int32 NumUnshadowedLights, const FSimpleLightArray& SimpleLights)
 {
 	check(GUseTiledDeferredShading);
 	check(SortedLights.Num() >= NumUnshadowedLights);
 
-	const int32 NumLightsToRender = NumUnshadowedLights + SimpleLights.Num();
+	const int32 NumLightsToRender = NumUnshadowedLights + SimpleLights.InstanceData.Num();
 	const int32 NumLightsToRenderInSortedLights = NumUnshadowedLights;
 
 	if (NumLightsToRender > 0)
 	{
 		INC_DWORD_STAT_BY(STAT_NumLightsUsingTiledDeferred, NumLightsToRender);
-		INC_DWORD_STAT_BY(STAT_NumLightsUsingSimpleTiledDeferred, SimpleLights.Num());
+		INC_DWORD_STAT_BY(STAT_NumLightsUsingSimpleTiledDeferred, SimpleLights.InstanceData.Num());
 		SCOPE_CYCLE_COUNTER(STAT_DirectLightRenderingTime);
 
 		// Determine how many compute shader passes will be needed to process all the lights
 		const int32 NumPassesNeeded = FMath::DivideAndRoundUp(NumLightsToRender, GMaxNumTiledDeferredLights);
-
 		for (int32 PassIndex = 0; PassIndex < NumPassesNeeded; PassIndex++)
 		{
 			const int32 StartIndex = PassIndex * GMaxNumTiledDeferredLights;
@@ -313,11 +318,11 @@ void FDeferredShadingSceneRenderer::RenderTiledDeferredLighting(const TArray<FSo
 
 					if(View.Family->EngineShowFlags.VisualizeLightCulling)
 					{
-						SetShaderTemplTiledLighting<1>(View, SortedLights, NumLightsToRenderInSortedLights, SimpleLights, StartIndex, NumThisPass);
+						SetShaderTemplTiledLighting<1>(View, ViewIndex, Views.Num(), SortedLights, NumLightsToRenderInSortedLights, SimpleLights, StartIndex, NumThisPass);
 					}
 					else
 					{
-						SetShaderTemplTiledLighting<0>(View, SortedLights, NumLightsToRenderInSortedLights, SimpleLights, StartIndex, NumThisPass);
+						SetShaderTemplTiledLighting<0>(View, ViewIndex, Views.Num(), SortedLights, NumLightsToRenderInSortedLights, SimpleLights, StartIndex, NumThisPass);
 					}
 				}
 			}
diff --git a/Engine/Source/Runtime/Renderer/Private/TranslucentLighting.cpp b/Engine/Source/Runtime/Renderer/Private/TranslucentLighting.cpp
--- a/Engine/Source/Runtime/Renderer/Private/TranslucentLighting.cpp
+++ b/Engine/Source/Runtime/Renderer/Private/TranslucentLighting.cpp
@@ -1738,12 +1738,13 @@ public:
 	}
 	FSimpleLightTranslucentLightingInjectPS() {}
 
-	void SetParameters(const FViewInfo& View, const FSimpleLightEntry& SimpleLight, int32 VolumeCascadeIndexValue)
+	void SetParameters(const FViewInfo& View, const FSimpleLightEntry& SimpleLight, const FSimpleLightPerViewEntry& SimpleLightPerViewData, int32 VolumeCascadeIndexValue)
 	{
 		FGlobalShader::SetParameters(GetPixelShader(), View);
 
+		FVector4 PositionAndRadius(SimpleLightPerViewData.Position, SimpleLight.Radius);
 		SetShaderValue(GetPixelShader(), VolumeCascadeIndex, VolumeCascadeIndexValue);
-		SetShaderValue(GetPixelShader(), SimpleLightPositionAndRadius, SimpleLight.PositionAndRadius);
+		SetShaderValue(GetPixelShader(), SimpleLightPositionAndRadius, PositionAndRadius);
 		SetShaderValue(GetPixelShader(), SimpleLightColorAndExponent, FVector4(SimpleLight.Color, SimpleLight.Exponent));
 	}
 
@@ -1766,15 +1767,15 @@ IMPLEMENT_SHADER_TYPE(,FSimpleLightTranslucentLightingInjectPS,TEXT("Translucent
 
 FGlobalBoundShaderState InjectSimpleLightBoundShaderState;
 
-void FDeferredShadingSceneRenderer::InjectSimpleTranslucentVolumeLightingArray(const TArray<FSimpleLightEntry, SceneRenderingAllocator>& SimpleLights)
+void FDeferredShadingSceneRenderer::InjectSimpleTranslucentVolumeLightingArray(const FSimpleLightArray& SimpleLights)
 {
 	SCOPE_CYCLE_COUNTER(STAT_TranslucentInjectTime);
 
 	int32 NumLightsToInject = 0;
 
-	for (int32 LightIndex = 0; LightIndex < SimpleLights.Num(); LightIndex++)
+	for (int32 LightIndex = 0; LightIndex < SimpleLights.InstanceData.Num(); LightIndex++)
 	{
-		if (SimpleLights[LightIndex].bAffectTranslucency)
+		if (SimpleLights.InstanceData[LightIndex].bAffectTranslucency)
 		{
 			NumLightsToInject++;
 		}
@@ -1782,8 +1783,9 @@ void FDeferredShadingSceneRenderer::InjectSimpleTranslucentVolumeLightingArray(c
 
 	if (NumLightsToInject > 0)
 	{
-		//@todo - support multiple views
+		//@todo - support multiple views 
 		const FViewInfo& View = Views[0];
+		const int32 ViewIndex = 0;
 
 		INC_DWORD_STAT_BY(STAT_NumLightsInjectedIntoTranslucency, NumLightsToInject);
 
@@ -1806,13 +1808,16 @@ void FDeferredShadingSceneRenderer::InjectSimpleTranslucentVolumeLightingArray(c
 			RHISetRasterizerState(TStaticRasterizerState<FM_Solid,CM_None>::GetRHI());
 			RHISetDepthStencilState(TStaticDepthStencilState<false,CF_Always>::GetRHI());
 
-			for (int32 LightIndex = 0; LightIndex < SimpleLights.Num(); LightIndex++)
+		
+			for (int32 LightIndex = 0; LightIndex < SimpleLights.InstanceData.Num(); LightIndex++)
 			{
-				const FSimpleLightEntry& SimpleLight = SimpleLights[LightIndex];
+				const FSimpleLightEntry& SimpleLight = SimpleLights.InstanceData[LightIndex];
+				const FSimpleLightPerViewEntry& SimpleLightPerViewData = SimpleLights.GetViewDependentData(LightIndex, ViewIndex, Views.Num());
 
 				if (SimpleLight.bAffectTranslucency)
 				{
-					const FVolumeBounds VolumeBounds = CalculateLightVolumeBounds((const FSphere&)SimpleLight.PositionAndRadius, View, VolumeCascadeIndex, false);
+					const FSphere LightBounds(SimpleLightPerViewData.Position, SimpleLight.Radius);
+					const FVolumeBounds VolumeBounds = CalculateLightVolumeBounds(LightBounds, View, VolumeCascadeIndex, false);
 
 					if (VolumeBounds.IsValid())
 					{
@@ -1823,7 +1828,7 @@ void FDeferredShadingSceneRenderer::InjectSimpleTranslucentVolumeLightingArray(c
 
 						VertexShader->SetParameters(VolumeBounds, GTranslucencyLightingVolumeDim);
 						GeometryShader->SetParameters(VolumeBounds);
-						PixelShader->SetParameters(View, SimpleLight, VolumeCascadeIndex);
+						PixelShader->SetParameters(View, SimpleLight, SimpleLightPerViewData, VolumeCascadeIndex);
 
 						// Accumulate the contribution of multiple lights
 						RHISetBlendState(TStaticBlendState<
