Remove UAV flag from rendertargets that don't need it. Fixes debug/validated D3D driver assert.

diff --git a/Engine/Source/Runtime/Renderer/Public/RendererInterface.h b/Engine/Source/Runtime/Renderer/Public/RendererInterface.h
--- a/Engine/Source/Runtime/Renderer/Public/RendererInterface.h
+++ b/Engine/Source/Runtime/Renderer/Public/RendererInterface.h
@@ -266,6 +266,9 @@ public:
 	{
 		// Usually we don't want to propagate MSAA samples.
 		NumSamples = 1;
+
+		// Remove UAV flag for rendertargets that don't need it (some formats are incompatible)
+		Flags &= (~TexCreate_UAV);
 	}
 
 	/** In pixels, (0,0) if not set, (x,0) for cube maps, todo: make 3d int vector for volume textures */
