* minor cleanup

diff --git a/Engine/Source/Runtime/Renderer/Private/ScenePostProcessing.cpp b/Engine/Source/Runtime/Renderer/Private/ScenePostProcessing.cpp
--- a/Engine/Source/Runtime/Renderer/Private/ScenePostProcessing.cpp
+++ b/Engine/Source/Runtime/Renderer/Private/ScenePostProcessing.cpp
@@ -132,18 +132,20 @@ void FSceneRenderer::GammaCorrectToViewportRenderTarget(const FViewInfo* View, f
 		InvDisplayGamma = 1 / OverrideGamma;
 	}
 
+	const FPixelShaderRHIParamRef ShaderRHI = PixelShader->GetPixelShader();
+
 	SetShaderValue(
-		PixelShader->GetPixelShader(),
+		ShaderRHI,
 		PixelShader->InverseGamma,
 		InvDisplayGamma
 		);
-	SetShaderValue(PixelShader->GetPixelShader(),PixelShader->ColorScale,View->ColorScale);
-	SetShaderValue(PixelShader->GetPixelShader(),PixelShader->OverlayColor,View->OverlayColor);
+	SetShaderValue(ShaderRHI,PixelShader->ColorScale,View->ColorScale);
+	SetShaderValue(ShaderRHI,PixelShader->OverlayColor,View->OverlayColor);
 
-	const FTexture2DRHIRef DesiredSceneColorTexture = GSceneRenderTargets.GetSceneColorTexture();
+	const FTextureRHIRef DesiredSceneColorTexture = GSceneRenderTargets.GetSceneColorTexture();
 
 	SetTextureParameter(
-		PixelShader->GetPixelShader(),
+		ShaderRHI,
 		PixelShader->SceneTexture,
 		PixelShader->SceneTextureSampler,
 		TStaticSamplerState<SF_Bilinear>::GetRHI(),
