* fixed build error

diff --git a/Engine/Source/Runtime/Renderer/Private/ScenePostProcessing.cpp b/Engine/Source/Runtime/Renderer/Private/ScenePostProcessing.cpp
--- a/Engine/Source/Runtime/Renderer/Private/ScenePostProcessing.cpp
+++ b/Engine/Source/Runtime/Renderer/Private/ScenePostProcessing.cpp
@@ -142,7 +142,7 @@ void FSceneRenderer::GammaCorrectToViewportRenderTarget(const FViewInfo* View, f
 	SetShaderValue(ShaderRHI,PixelShader->ColorScale,View->ColorScale);
 	SetShaderValue(ShaderRHI,PixelShader->OverlayColor,View->OverlayColor);
 
-	const FTextureRHIRef DesiredSceneColorTexture = GSceneRenderTargets.GetSceneColorTexture();
+	const FTexture2DRHIRef DesiredSceneColorTexture = GSceneRenderTargets.GetSceneColorTexture();
 
 	SetTextureParameter(
 		ShaderRHI,
