* optimized out some memory allocations

diff --git a/Engine/Source/Runtime/Renderer/Private/PostProcess/RenderingCompositionGraph.h b/Engine/Source/Runtime/Renderer/Private/PostProcess/RenderingCompositionGraph.h
--- a/Engine/Source/Runtime/Renderer/Private/PostProcess/RenderingCompositionGraph.h
+++ b/Engine/Source/Runtime/Renderer/Private/PostProcess/RenderingCompositionGraph.h
@@ -230,7 +230,7 @@ struct FRenderingCompositePass
 	 * Allows access to dump filename for a given output
 	 * @return Filename for output dump
 	 */
-	virtual const TCHAR* GetOutputDumpFilename(EPassOutputId OutputId) const = 0;
+	virtual const FString& GetOutputDumpFilename(EPassOutputId OutputId) = 0;
 
 	/**
 	 * Allows setting of a dump filename for a given output
@@ -461,10 +461,10 @@ struct TRenderingCompositePassBase :public FRenderingCompositePass
 		return 0;
 	}
 
-	virtual const TCHAR* GetOutputDumpFilename(EPassOutputId OutputId) const
+	virtual const FString& GetOutputDumpFilename(EPassOutputId OutputId)
 	{
 		check (OutputId < OutputCount);
-		return *PassOutputDumpFilenames[OutputId];
+		return PassOutputDumpFilenames[OutputId];
 	}
 
 	virtual void SetOutputDumpFilename(EPassOutputId OutputId, const TCHAR* Filename)
