* supress potential crash in RenderTargetPool event visualizer

diff --git a/Engine/Source/Runtime/Renderer/Private/PostProcess/RenderTargetPool.cpp b/Engine/Source/Runtime/Renderer/Private/PostProcess/RenderTargetPool.cpp
--- a/Engine/Source/Runtime/Renderer/Private/PostProcess/RenderTargetPool.cpp
+++ b/Engine/Source/Runtime/Renderer/Private/PostProcess/RenderTargetPool.cpp
@@ -79,7 +79,7 @@ bool FRenderTargetPoolEvent::NeedsDeallocEvent()
 		if(Pointer)
 		{
 			IPooledRenderTarget* ValidPointer = GetValidatedPointer();
-			if(ValidPointer->IsFree())
+			if(!ValidPointer || ValidPointer->IsFree())
 			{
 				Pointer = 0;
 				return true;
