* fixed wrongly use of SceneColor

diff --git a/Engine/Source/Runtime/Renderer/Private/PostProcess/PostProcessSceneColorFringe.cpp b/Engine/Source/Runtime/Renderer/Private/PostProcess/PostProcessSceneColorFringe.cpp
--- a/Engine/Source/Runtime/Renderer/Private/PostProcess/PostProcessSceneColorFringe.cpp
+++ b/Engine/Source/Runtime/Renderer/Private/PostProcess/PostProcessSceneColorFringe.cpp
@@ -183,7 +183,7 @@ void FRCPassPostProcessSceneColorFringe::Process(FRenderingCompositePassContext&
 		View.ViewRect.Min.X, View.ViewRect.Min.Y, 
 		View.ViewRect.Width(), View.ViewRect.Height(),
 		View.ViewRect.Size(),
-		GSceneRenderTargets.SceneColor->GetDesc().Extent,
+		GSceneRenderTargets.GetBufferSizeXY(),
 		*VertexShader,
 		EDRF_UseTriangleOptimization);
 
