* refactored render target pool managentment, better sharing of rendertargets, can cause some different behavior (getting SceneColor) or might bump into ensure().
* saves GPU memory doing better sharing
* can be visualized using r.RenderTargetPool.Events

diff --git a/Engine/Source/Runtime/Renderer/Private/CompositionLighting/CompositionLighting.cpp b/Engine/Source/Runtime/Renderer/Private/CompositionLighting/CompositionLighting.cpp
--- a/Engine/Source/Runtime/Renderer/Private/CompositionLighting/CompositionLighting.cpp
+++ b/Engine/Source/Runtime/Renderer/Private/CompositionLighting/CompositionLighting.cpp
@@ -229,8 +229,10 @@ void FCompositionLighting::ProcessBeforeBasePass(const FViewInfo& View)
 
 		SCOPED_DRAW_EVENT(CompositionBeforeBasePass, DEC_SCENE_ITEMS);
 
-		Context.FinalOutput.GetOutput()->RenderTargetDesc = GSceneRenderTargets.SceneColor->GetDesc();
-		Context.FinalOutput.GetOutput()->PooledRenderTarget = GSceneRenderTargets.SceneColor;
+		TRefCountPtr<IPooledRenderTarget>& SceneColor = GSceneRenderTargets.GetSceneColor();
+
+		Context.FinalOutput.GetOutput()->RenderTargetDesc = GSceneRenderTargets.GetSceneColor()->GetDesc();
+		Context.FinalOutput.GetOutput()->PooledRenderTarget = SceneColor;
 
 		// you can add multiple dependencies
 		CompositeContext.Root->AddDependency(Context.FinalOutput);
@@ -244,10 +246,10 @@ void FCompositionLighting::ProcessAfterBasePass(const FViewInfo& View)
 	check(IsInRenderingThread());
 
 	// might get renamed to refracted or ...WithAO
-	GSceneRenderTargets.SceneColor->SetDebugName(TEXT("SceneColor"));
+	GSceneRenderTargets.GetSceneColor()->SetDebugName(TEXT("SceneColor"));
 	// to be able to observe results with VisualizeTexture
 
-	GRenderTargetPool.VisualizeTexture.SetCheckPoint(GSceneRenderTargets.SceneColor);
+	GRenderTargetPool.VisualizeTexture.SetCheckPoint(GSceneRenderTargets.GetSceneColor());
 	GRenderTargetPool.VisualizeTexture.SetCheckPoint(GSceneRenderTargets.GBufferA);
 	GRenderTargetPool.VisualizeTexture.SetCheckPoint(GSceneRenderTargets.GBufferB);
 	GRenderTargetPool.VisualizeTexture.SetCheckPoint(GSceneRenderTargets.GBufferC);
@@ -286,8 +288,11 @@ void FCompositionLighting::ProcessAfterBasePass(const FViewInfo& View)
 		// The graph setup should be finished before this line ----------------------------------------
     
 		SCOPED_DRAW_EVENT(LightCompositionTasks_PreLighting, DEC_SCENE_ITEMS);
-		Context.FinalOutput.GetOutput()->RenderTargetDesc = GSceneRenderTargets.SceneColor->GetDesc();
-		Context.FinalOutput.GetOutput()->PooledRenderTarget = GSceneRenderTargets.SceneColor;
+
+		TRefCountPtr<IPooledRenderTarget>& SceneColor = GSceneRenderTargets.GetSceneColor();
+
+		Context.FinalOutput.GetOutput()->RenderTargetDesc = SceneColor->GetDesc();
+		Context.FinalOutput.GetOutput()->PooledRenderTarget = SceneColor;
     
 		// you can add multiple dependencies
 		CompositeContext.Root->AddDependency(Context.FinalOutput);
@@ -322,12 +327,18 @@ void FCompositionLighting::ProcessLighting(const FViewInfo& View)
 
 		SCOPED_DRAW_EVENT(CompositionLighting, DEC_SCENE_ITEMS);
 
-		Context.FinalOutput.GetOutput()->RenderTargetDesc = GSceneRenderTargets.SceneColor->GetDesc();
-		Context.FinalOutput.GetOutput()->PooledRenderTarget = GSceneRenderTargets.SceneColor;
+		TRefCountPtr<IPooledRenderTarget>& SceneColor = GSceneRenderTargets.GetSceneColor();
+
+		Context.FinalOutput.GetOutput()->RenderTargetDesc = SceneColor->GetDesc();
+		Context.FinalOutput.GetOutput()->PooledRenderTarget = SceneColor;
 
 		// you can add multiple dependencies
 		CompositeContext.Root->AddDependency(Context.FinalOutput);
 
 		CompositeContext.Process(TEXT("CompositionLighting_Lighting"));
 	}
+
+	// The RT should be released as early as possible to allow shading of that memory for other purposes.
+	// This becomes even more important with some limited VRam (XBoxOne).
+	GSceneRenderTargets.SetLightAttenuation(0);
 }
diff --git a/Engine/Source/Runtime/Renderer/Private/CompositionLighting/PostProcessAmbient.cpp b/Engine/Source/Runtime/Renderer/Private/CompositionLighting/PostProcessAmbient.cpp
--- a/Engine/Source/Runtime/Renderer/Private/CompositionLighting/PostProcessAmbient.cpp
+++ b/Engine/Source/Runtime/Renderer/Private/CompositionLighting/PostProcessAmbient.cpp
@@ -75,7 +75,7 @@ void FRCPassPostProcessAmbient::Process(FRenderingCompositePassContext& Context)
 	FIntRect DestRect = View.ViewRect;
 	FIntPoint DestSize = DestRect.Size();
 
-	const FSceneRenderTargetItem& DestRenderTarget = GSceneRenderTargets.SceneColor->GetRenderTargetItem();
+	const FSceneRenderTargetItem& DestRenderTarget = GSceneRenderTargets.GetSceneColor()->GetRenderTargetItem();
 
 	// Set the view family's render target/viewport.
 	RHISetRenderTarget(DestRenderTarget.TargetableTexture, FTextureRHIRef());	
@@ -116,7 +116,7 @@ void FRCPassPostProcessAmbient::Process(FRenderingCompositePassContext& Context)
 			View.ViewRect.Min.X, View.ViewRect.Min.Y, 
 			View.ViewRect.Width(), View.ViewRect.Height(),
 			View.ViewRect.Size(),
-			GSceneRenderTargets.SceneColor->GetDesc().Extent,
+			GSceneRenderTargets.GetBufferSizeXY(),
 			*VertexShader,
 			EDRF_UseTriangleOptimization);
 	}
diff --git a/Engine/Source/Runtime/Renderer/Private/CompositionLighting/PostProcessAmbientOcclusion.cpp b/Engine/Source/Runtime/Renderer/Private/CompositionLighting/PostProcessAmbientOcclusion.cpp
--- a/Engine/Source/Runtime/Renderer/Private/CompositionLighting/PostProcessAmbientOcclusion.cpp
+++ b/Engine/Source/Runtime/Renderer/Private/CompositionLighting/PostProcessAmbientOcclusion.cpp
@@ -585,7 +585,7 @@ void FRCPassPostProcessBasePassAO::Process(FRenderingCompositePassContext& Conte
 
 	const FSceneView& View = Context.View;
 
-	const FSceneRenderTargetItem& DestRenderTarget = GSceneRenderTargets.SceneColor->GetRenderTargetItem();
+	const FSceneRenderTargetItem& DestRenderTarget = GSceneRenderTargets.GetSceneColor()->GetRenderTargetItem();
 
 	// Set the view family's render target/viewport.
 	RHISetRenderTarget(DestRenderTarget.TargetableTexture, FTextureRHIRef());	
diff --git a/Engine/Source/Runtime/Renderer/Private/CompositionLighting/PostProcessLpvIndirect.cpp b/Engine/Source/Runtime/Renderer/Private/CompositionLighting/PostProcessLpvIndirect.cpp
--- a/Engine/Source/Runtime/Renderer/Private/CompositionLighting/PostProcessLpvIndirect.cpp
+++ b/Engine/Source/Runtime/Renderer/Private/CompositionLighting/PostProcessLpvIndirect.cpp
@@ -142,7 +142,7 @@ void FRCPassPostProcessLpvIndirect::Process(FRenderingCompositePassContext& Cont
 
 	uint32 NumReflectionCaptures = ViewFamily.Scene->GetRenderScene()->ReflectionSceneData.RegisteredReflectionCaptures.Num();
 
-	const FSceneRenderTargetItem& DestColorRenderTarget = GSceneRenderTargets.SceneColor->GetRenderTargetItem();
+	const FSceneRenderTargetItem& DestColorRenderTarget = GSceneRenderTargets.GetSceneColor()->GetRenderTargetItem();
 
 	// Set the view family's render target/viewport.
 	FTextureRHIParamRef RenderTargets[] =
@@ -217,7 +217,7 @@ void FRCPassPostProcessLpvIndirect::Process(FRenderingCompositePassContext& Cont
 		View.ViewRect.Min.X, View.ViewRect.Min.Y, 
 		View.ViewRect.Width(), View.ViewRect.Height(),
 		View.ViewRect.Size(),
-		GSceneRenderTargets.SceneColor->GetDesc().Extent,
+		GSceneRenderTargets.GetBufferSizeXY(),
 		*VertexShader);
 
 	RHICopyToResolveTarget(DestColorRenderTarget.TargetableTexture, DestColorRenderTarget.ShaderResourceTexture, false, FResolveParams());
diff --git a/Engine/Source/Runtime/Renderer/Private/DeferredShadingRenderer.cpp b/Engine/Source/Runtime/Renderer/Private/DeferredShadingRenderer.cpp
--- a/Engine/Source/Runtime/Renderer/Private/DeferredShadingRenderer.cpp
+++ b/Engine/Source/Runtime/Renderer/Private/DeferredShadingRenderer.cpp
@@ -426,6 +426,11 @@ void FDeferredShadingSceneRenderer::RenderFinish()
 	{
 		Scene->MotionBlurInfoData.UpdateMotionBlurCache();
 	}
+
+	// The SceneColor RT should be released as early as possible to allow shading of that memory for other purposes.
+	// This becomes even more important with some limited VRam (XBoxOne).
+	GSceneRenderTargets.SetSceneColor(0);
+	GSceneRenderTargets.SetLightAttenuation(0);
 }
 
 /** 
diff --git a/Engine/Source/Runtime/Renderer/Private/DistortionRendering.cpp b/Engine/Source/Runtime/Renderer/Private/DistortionRendering.cpp
--- a/Engine/Source/Runtime/Renderer/Private/DistortionRendering.cpp
+++ b/Engine/Source/Runtime/Renderer/Private/DistortionRendering.cpp
@@ -33,12 +33,12 @@ public:
 	}
 	FDistortionApplyScreenPS() {}
 
-	void SetParameters(const FRenderingCompositePassContext& Context)
+	void SetParameters(const FRenderingCompositePassContext& Context, IPooledRenderTarget& DistortionRT)
 	{
 		const FPixelShaderRHIParamRef ShaderRHI = GetPixelShader();
 
-		FTextureRHIParamRef DistortionTextureValue = GSceneRenderTargets.LightAttenuation->GetRenderTargetItem().ShaderResourceTexture;
-		FTextureRHIParamRef SceneColorTextureValue = GSceneRenderTargets.SceneColor->GetRenderTargetItem().ShaderResourceTexture;
+		FTextureRHIParamRef DistortionTextureValue = DistortionRT.GetRenderTargetItem().ShaderResourceTexture;
+		FTextureRHIParamRef SceneColorTextureValue = GSceneRenderTargets.GetSceneColor()->GetRenderTargetItem().ShaderResourceTexture;
 
 		// Here we use SF_Point as in fullscreen the pixels are 1:1 mapped.
 		SetTextureParameter(
@@ -774,43 +774,64 @@ void FDeferredShadingSceneRenderer::RenderDistortion()
 
 	bool bDirty = false;
 
+	TRefCountPtr<IPooledRenderTarget> DistortionRT;
+
 	// Render accumulated distortion offsets
 	if( bRender)
 	{
 		SCOPED_DRAW_EVENT(DistortionAccum, DEC_SCENE_ITEMS);
 
-		GSceneRenderTargets.BeginRenderingDistortionAccumulation();
+		// Create a texture to store the resolved light attenuation values, and a render-targetable surface to hold the unresolved light attenuation values.
+		{
+			FPooledRenderTargetDesc Desc(FPooledRenderTargetDesc::Create2DDesc(GSceneRenderTargets.GetBufferSizeXY(), PF_B8G8R8A8, TexCreate_None, TexCreate_RenderTargetable, false));
+			Desc.Flags |= TexCreate_FastVRAM;
+			GRenderTargetPool.FindFreeElement(Desc, DistortionRT, TEXT("Distortion"));
+
+			// use RGBA8 light target for accumulating distortion offsets	
+			// R = positive X offset
+			// G = positive Y offset
+			// B = negative X offset
+			// A = negative Y offset
+		}
 
-		for(int32 ViewIndex = 0;ViewIndex < Views.Num();ViewIndex++)
+		// DistortionRT==0 should never happen but better we don't crash
+		if(DistortionRT)
 		{
-			SCOPED_CONDITIONAL_DRAW_EVENTF(EventView, Views.Num() > 1, DEC_SCENE_ITEMS, TEXT("View%d"), ViewIndex);
+			RHISetRenderTarget(DistortionRT->GetRenderTargetItem().TargetableTexture,GSceneRenderTargets.GetSceneDepthSurface());
 
-			FViewInfo& View = Views[ViewIndex];
-			// viewport to match view size
-			RHISetViewport(View.ViewRect.Min.X,View.ViewRect.Min.Y,0.0f,View.ViewRect.Max.X,View.ViewRect.Max.Y,1.0f);
+			for(int32 ViewIndex = 0;ViewIndex < Views.Num();ViewIndex++)
+			{
+				SCOPED_CONDITIONAL_DRAW_EVENTF(EventView, Views.Num() > 1, DEC_SCENE_ITEMS, TEXT("View%d"), ViewIndex);
 
-			// clear offsets to 0, stencil to 0
-			RHIClear(true,FLinearColor(0,0,0,0),false,0,true,0,FIntRect());
+				FViewInfo& View = Views[ViewIndex];
+				// viewport to match view size
+				RHISetViewport(View.ViewRect.Min.X,View.ViewRect.Min.Y,0.0f,View.ViewRect.Max.X,View.ViewRect.Max.Y,1.0f);
 
-			// enable depth test but disable depth writes
-			// Note, this is a reversed Z depth surface, using CF_GreaterEqual.
-			RHISetDepthStencilState(TStaticDepthStencilState<false,CF_GreaterEqual>::GetRHI());
+				// clear offsets to 0, stencil to 0
+				RHIClear(true,FLinearColor(0,0,0,0),false,0,true,0,FIntRect());
 
-			// additive blending of offsets (or complexity if the shader complexity viewmode is enabled)
-			RHISetBlendState(TStaticBlendState<CW_RGBA,BO_Add,BF_One,BF_One,BO_Add,BF_One,BF_One>::GetRHI());
+				// enable depth test but disable depth writes
+				// Note, this is a reversed Z depth surface, using CF_GreaterEqual.
+				RHISetDepthStencilState(TStaticDepthStencilState<false,CF_GreaterEqual>::GetRHI());
 
-			// draw only distortion meshes to accumulate their offsets
-			bDirty |= View.DistortionPrimSet.DrawAccumulatedOffsets(&View,false);
-		}
+				// additive blending of offsets (or complexity if the shader complexity viewmode is enabled)
+				RHISetBlendState(TStaticBlendState<CW_RGBA,BO_Add,BF_One,BF_One,BO_Add,BF_One,BF_One>::GetRHI());
 
-		if (bDirty)
-		{
-			// restore default stencil state
-			// Note, this is a reversed Z depth surface, using CF_GreaterEqual.
-			RHISetDepthStencilState(TStaticDepthStencilState<true,CF_GreaterEqual>::GetRHI());
+				// draw only distortion meshes to accumulate their offsets
+				bDirty |= View.DistortionPrimSet.DrawAccumulatedOffsets(&View,false);
+			}
 
-			// resolve using the current ResolveParams 
-			GSceneRenderTargets.FinishRenderingDistortionAccumulation();
+			if (bDirty)
+			{
+				// restore default stencil state
+				// Note, this is a reversed Z depth surface, using CF_GreaterEqual.
+				RHISetDepthStencilState(TStaticDepthStencilState<true,CF_GreaterEqual>::GetRHI());
+
+				// resolve using the current ResolveParams 
+				RHICopyToResolveTarget(DistortionRT->GetRenderTargetItem().TargetableTexture, DistortionRT->GetRenderTargetItem().ShaderResourceTexture, false, FResolveParams());
+				// to be able to observe results with VisualizeTexture
+				GRenderTargetPool.VisualizeTexture.SetCheckPoint(DistortionRT);
+			}
 		}
 	}
 
@@ -822,12 +843,12 @@ void FDeferredShadingSceneRenderer::RenderDistortion()
 
 // OCULUS BEGIN: select ONE render target for all views (eyes)
 		TRefCountPtr<IPooledRenderTarget> NewSceneColor;
-		GRenderTargetPool.FindFreeElement(GSceneRenderTargets.SceneColor->GetDesc(), NewSceneColor, TEXT("RefractedSceneColor"));
+		GRenderTargetPool.FindFreeElement(GSceneRenderTargets.GetSceneColor()->GetDesc(), NewSceneColor, TEXT("RefractedSceneColor"));
 		const FSceneRenderTargetItem& DestRenderTarget = NewSceneColor->GetRenderTargetItem();
 // OCULUS END
 
 		// Apply distortion as a full-screen pass		
-		for(int32 ViewIndex = 0;ViewIndex < Views.Num();ViewIndex++)
+		for(int32 ViewIndex = 0, Num = Views.Num(); ViewIndex < Num; ++ViewIndex)
 		{
 			SCOPED_CONDITIONAL_DRAW_EVENTF(EventView, Views.Num() > 1, DEC_SCENE_ITEMS, TEXT("View%d"), ViewIndex);
 
@@ -849,7 +870,7 @@ void FDeferredShadingSceneRenderer::RenderDistortion()
 
 				RHISetRenderTarget( DestRenderTarget.TargetableTexture, FTextureRHIParamRef());
 
-				// useful wehen we move this into the compositing graph
+				// useful when we move this into the compositing graph
 				FRenderingCompositePassContext Context(View);
 
 				// Set the view family's render target/viewport.
@@ -862,7 +883,7 @@ void FDeferredShadingSceneRenderer::RenderDistortion()
 				SetGlobalBoundShaderState(BoundShaderState, GFilterVertexDeclaration.VertexDeclarationRHI, *VertexShader, *PixelShader);
 
 				VertexShader->SetParameters(Context);
-				PixelShader->SetParameters(Context);
+				PixelShader->SetParameters(Context, *DistortionRT);
 
 				// Draw a quad mapping scene color to the view's render target
 				DrawRectangle(
@@ -871,15 +892,15 @@ void FDeferredShadingSceneRenderer::RenderDistortion()
 					View.ViewRect.Min.X, View.ViewRect.Min.Y, 
 					View.ViewRect.Width(), View.ViewRect.Height(),
 					View.ViewRect.Size(),
-					GSceneRenderTargets.SceneColor->GetDesc().Extent,
+					GSceneRenderTargets.GetBufferSizeXY(),
 					*VertexShader,
 					EDRF_UseTriangleOptimization);
 
 				RHICopyToResolveTarget(DestRenderTarget.TargetableTexture, DestRenderTarget.ShaderResourceTexture, false, FResolveParams());
 
 // OCULUS BEGIN
-// 				GSceneRenderTargets.SceneColor = NewSceneColor;
-// 				check(GSceneRenderTargets.SceneColor);
+// 				GSceneRenderTargets.SetSceneColor(NewSceneColor);
+// 				check(GSceneRenderTargets.GetSceneColor());
 // OCULUS END				
 
 			}
@@ -887,9 +908,10 @@ void FDeferredShadingSceneRenderer::RenderDistortion()
 		}
 
 // OCULUS BEGIN
-		GSceneRenderTargets.SceneColor = NewSceneColor;
-		check(GSceneRenderTargets.SceneColor);
+		GSceneRenderTargets.SetSceneColor(NewSceneColor);
+		check(GSceneRenderTargets.GetSceneColor());
 // OCULUS END				
+		// Distortions RT is no longer needed, buffer can be reused by the pool, see BeginRenderingDistortionAccumulation() call above
 		GSceneRenderTargets.FinishRenderingSceneColor(false);
 	}
 }
diff --git a/Engine/Source/Runtime/Renderer/Private/PostProcess/PostProcessBusyWait.cpp b/Engine/Source/Runtime/Renderer/Private/PostProcess/PostProcessBusyWait.cpp
--- a/Engine/Source/Runtime/Renderer/Private/PostProcess/PostProcessBusyWait.cpp
+++ b/Engine/Source/Runtime/Renderer/Private/PostProcess/PostProcessBusyWait.cpp
@@ -80,8 +80,10 @@ void FRCPassPostProcessBusyWait::Process(FRenderingCompositePassContext& Context
 	
 	FIntRect SrcRect = View.ViewRect;
 	FIntRect DestRect = View.UnscaledViewRect;
+
+	GSceneRenderTargets.BeginRenderingLightAttenuation();
 	
-	const FSceneRenderTargetItem& DestRenderTarget = GSceneRenderTargets.LightAttenuation->GetRenderTargetItem();
+	const FSceneRenderTargetItem& DestRenderTarget = GSceneRenderTargets.GetLightAttenuation()->GetRenderTargetItem();
 
 	// Set the view family's render target/viewport.
 	RHISetRenderTarget(DestRenderTarget.TargetableTexture, FTextureRHIRef());	
@@ -113,6 +115,8 @@ void FRCPassPostProcessBusyWait::Process(FRenderingCompositePassContext& Context
 		EDRF_UseTriangleOptimization);
 
 	RHICopyToResolveTarget(DestRenderTarget.TargetableTexture, DestRenderTarget.ShaderResourceTexture, false, FResolveParams());
+
+	GSceneRenderTargets.SetLightAttenuation(0);
 }
 
 FPooledRenderTargetDesc FRCPassPostProcessBusyWait::ComputeOutputDesc(EPassOutputId InPassOutputId) const
diff --git a/Engine/Source/Runtime/Renderer/Private/PostProcess/PostProcessTonemap.cpp b/Engine/Source/Runtime/Renderer/Private/PostProcess/PostProcessTonemap.cpp
--- a/Engine/Source/Runtime/Renderer/Private/PostProcess/PostProcessTonemap.cpp
+++ b/Engine/Source/Runtime/Renderer/Private/PostProcess/PostProcessTonemap.cpp
@@ -1045,11 +1045,15 @@ void FRCPassPostProcessTonemap::Process(FRenderingCompositePassContext& Context)
 		View.ViewRect.Min.X, View.ViewRect.Min.Y, 
 		View.ViewRect.Width(), View.ViewRect.Height(),
 		View.ViewRect.Size(),
-		GSceneRenderTargets.SceneColor->GetDesc().Extent,
+		GSceneRenderTargets.GetBufferSizeXY(),
 		*VertexShader,
 		EDRF_UseTriangleOptimization);
 
 	RHICopyToResolveTarget(DestRenderTarget.TargetableTexture, DestRenderTarget.ShaderResourceTexture, false, FResolveParams());
+
+	// The RT should be released as early as possible to allow shading of that memory for other purposes.
+	// This becomes even more important with some limited VRam (XBoxOne).
+	GSceneRenderTargets.SetSceneColor(0);
 }
 
 FPooledRenderTargetDesc FRCPassPostProcessTonemap::ComputeOutputDesc(EPassOutputId InPassOutputId) const
@@ -1402,7 +1406,7 @@ void FRCPassPostProcessTonemapES2::Process(FRenderingCompositePassContext& Conte
 		View.ViewRect.Min.X, View.ViewRect.Min.Y, 
 		View.ViewRect.Width(), View.ViewRect.Height(),
 		View.ViewRect.Size(),
-		GSceneRenderTargets.SceneColor->GetDesc().Extent,
+		GSceneRenderTargets.GetBufferSizeXY(),
 		*VertexShader,
 		EDRF_UseTriangleOptimization);
 
@@ -1417,6 +1421,8 @@ void FRCPassPostProcessTonemapES2::Process(FRenderingCompositePassContext& Conte
 			ViewState->MobileAaColor0 = PassOutputs[0].PooledRenderTarget;
 		}
 	}
+
+	GSceneRenderTargets.SetSceneColor(0);
 }
 
 FPooledRenderTargetDesc FRCPassPostProcessTonemapES2::ComputeOutputDesc(EPassOutputId InPassOutputId) const
diff --git a/Engine/Source/Runtime/Renderer/Private/PostProcess/PostProcessing.cpp b/Engine/Source/Runtime/Renderer/Private/PostProcess/PostProcessing.cpp
--- a/Engine/Source/Runtime/Renderer/Private/PostProcess/PostProcessing.cpp
+++ b/Engine/Source/Runtime/Renderer/Private/PostProcess/PostProcessing.cpp
@@ -69,7 +69,7 @@ FPostprocessContext::FPostprocessContext(class FRenderingCompositionGraph& InGra
 	: Graph(InGraph)
 	, View(InView)
 {
-	SceneColor = Graph.RegisterPass(new(FMemStack::Get()) FRCPassPostProcessInput(GSceneRenderTargets.SceneColor));
+	SceneColor = Graph.RegisterPass(new(FMemStack::Get()) FRCPassPostProcessInput(GSceneRenderTargets.GetSceneColor()));
 	SceneDepth = Graph.RegisterPass(new(FMemStack::Get()) FRCPassPostProcessInput(GSceneRenderTargets.SceneDepthZ));
 	GBufferA = Graph.RegisterPass(new(FMemStack::Get()) FRCPassPostProcessInput(GSceneRenderTargets.GBufferA));
 
@@ -459,7 +459,7 @@ static void AddTemporalAA( FPostprocessContext& Context, FRenderingCompositeOutp
 	else
 	{
 		// No history so use current as history
-		HistoryInput = Context.Graph.RegisterPass( new(FMemStack::Get()) FRCPassPostProcessInput( GSceneRenderTargets.SceneColor ) );
+		HistoryInput = Context.Graph.RegisterPass( new(FMemStack::Get()) FRCPassPostProcessInput( GSceneRenderTargets.GetSceneColor() ) );
 	}
 
 	FRenderingCompositePass* TemporalAAPass = Context.Graph.RegisterPass( new(FMemStack::Get()) FRCPassPostProcessTemporalAA );
@@ -1002,7 +1002,7 @@ void FPostProcessing::Process(const FViewInfo& View, TRefCountPtr<IPooledRenderT
 				VisualizeNode->SetInput(ePId_Input0, FRenderingCompositeOutputRef(Context.FinalOutput));
 
 				// PassThrough is needed to upscale the half res texture
-				FPooledRenderTargetDesc Desc = GSceneRenderTargets.SceneColor->GetDesc();
+				FPooledRenderTargetDesc Desc = GSceneRenderTargets.GetSceneColor()->GetDesc();
 				Desc.Format = PF_B8G8R8A8;
 				FRenderingCompositePass* NullPass = Context.Graph.RegisterPass(new(FMemStack::Get()) FRCPassPostProcessPassThrough(Desc));
 				NullPass->SetInput(ePId_Input0, FRenderingCompositeOutputRef(VisualizeNode));
diff --git a/Engine/Source/Runtime/Renderer/Private/PostProcess/SceneRenderTargets.cpp b/Engine/Source/Runtime/Renderer/Private/PostProcess/SceneRenderTargets.cpp
--- a/Engine/Source/Runtime/Renderer/Private/PostProcess/SceneRenderTargets.cpp
+++ b/Engine/Source/Runtime/Renderer/Private/PostProcess/SceneRenderTargets.cpp
@@ -213,6 +213,8 @@ void FSceneRenderTargets::BeginRenderingSceneColor( bool bGBufferPass )
 	{
 		bGBufferPass = false;
 	}
+	
+	AllocSceneColor();
 
 	// Set the scene color surface as the render target, and the scene depth surface as the depth-stencil target.
 	if (bGBufferPass && GRHIFeatureLevel >= ERHIFeatureLevel::SM4)
@@ -254,6 +256,145 @@ int32 FSceneRenderTargets::GetNumGBufferTargets() const
 	return NumGBufferTargets;
 }
 
+void FSceneRenderTargets::AllocSceneColor()
+{
+	if(SceneColor)
+	{
+		// no work needed
+		return;
+	}
+
+	// create SceneColor on demand so it can be shared with other pooled RT
+
+	EPixelFormat SceneColorBufferFormat = GetSceneColorFormat();
+
+	// Create the scene color.
+	{
+		FPooledRenderTargetDesc Desc(FPooledRenderTargetDesc::Create2DDesc(BufferSize, SceneColorBufferFormat, TexCreate_None, TexCreate_RenderTargetable, false));
+		Desc.Flags |= TexCreate_FastVRAM;
+		// to allow better sharing with later elements
+		Desc.Flags |= TexCreate_UAV;
+		GRenderTargetPool.FindFreeElement(Desc, SceneColor, TEXT("SceneColor"));
+	}
+
+	// otherwise we have a severe problem
+	check(SceneColor);
+}
+
+void FSceneRenderTargets::AllocLightAttenuation()
+{
+	if(LightAttenuation)
+	{
+		// no work needed
+		return;
+	}
+
+	// create LightAttenuation on demand so it can be shared with other pooled RT
+
+	// Create a texture to store the resolved light attenuation values, and a render-targetable surface to hold the unresolved light attenuation values.
+	{
+		FPooledRenderTargetDesc Desc(FPooledRenderTargetDesc::Create2DDesc(BufferSize, PF_B8G8R8A8, TexCreate_None, TexCreate_RenderTargetable, false));
+		Desc.Flags |= TexCreate_FastVRAM;
+		GRenderTargetPool.FindFreeElement(Desc, LightAttenuation, TEXT("LightAttenuation"));
+
+		// the channel assignment is documented in ShadowRendering.cpp (look for Light Attenuation channel assignment)
+	}
+
+	// otherwise we have a severe problem
+	check(LightAttenuation);
+}
+
+const TRefCountPtr<IPooledRenderTarget>& FSceneRenderTargets::GetSceneColor() const
+{
+	if(!SceneColor)
+	{
+		// to avoid log/ensure spam
+		static bool bFirst = true;
+		if(bFirst)
+		{
+			bFirst = false;
+
+			// the first called should be AllocSceneColor(), contact MartinM if that happens
+			ensure(SceneColor);
+		}
+
+		return GSystemTextures.BlackDummy;
+	}
+
+	return SceneColor;
+}
+
+TRefCountPtr<IPooledRenderTarget>& FSceneRenderTargets::GetSceneColor()
+{
+	if(!SceneColor)
+	{
+		// to avoid log/ensure spam
+		static bool bFirst = true;
+		if(bFirst)
+		{
+			bFirst = false;
+
+			// the first called should be AllocSceneColor(), contact MartinM if that happens
+			ensure(SceneColor);
+		}
+
+		return GSystemTextures.BlackDummy;
+	}
+
+	return SceneColor;
+}
+
+void FSceneRenderTargets::SetSceneColor(IPooledRenderTarget* In)
+{
+	SceneColor = In;
+}
+
+void FSceneRenderTargets::SetLightAttenuation(IPooledRenderTarget* In)
+{
+	LightAttenuation = In;
+}
+
+const TRefCountPtr<IPooledRenderTarget>& FSceneRenderTargets::GetLightAttenuation() const
+{
+	if(!LightAttenuation)
+	{
+		// to avoid log/ensure spam
+		static bool bFirst = true;
+		if(bFirst)
+		{
+			bFirst = false;
+
+			// First we need to call AllocLightAttenuation(), contact MartinM if that happens
+			ensure(LightAttenuation);
+		}
+
+		return GSystemTextures.WhiteDummy;
+	}
+
+	return LightAttenuation;
+}
+
+TRefCountPtr<IPooledRenderTarget>& FSceneRenderTargets::GetLightAttenuation()
+{
+	if(!LightAttenuation)
+	{
+		// to avoid log/ensure spam
+		static bool bFirst = true;
+		if(bFirst)
+		{
+			bFirst = false;
+
+			// the first called should be AllocLightAttenuation(), contact MartinM if that happens
+			ensure(LightAttenuation);
+		}
+
+		return GSystemTextures.WhiteDummy;
+	}
+
+	return LightAttenuation;
+}
+
+
 void FSceneRenderTargets::FinishRenderingSceneColor(bool bKeepChanges, const FResolveRect& ResolveRect)
 {
 	SCOPED_DRAW_EVENT(FinishRenderingSceneColor, DEC_SCENE_ITEMS);
@@ -324,7 +465,7 @@ void FSceneRenderTargets::BeginRenderingPrePass()
 
 	// Set the scene depth surface and a DUMMY buffer as color buffer
 	// (as long as it's the same dimension as the depth buffer),
-	RHISetRenderTarget( GetLightAttenuationSurface(), GetSceneDepthSurface());
+	RHISetRenderTarget( FTextureRHIRef(), GetSceneDepthSurface());
 }
 
 void FSceneRenderTargets::FinishRenderingPrePass()
@@ -406,7 +547,9 @@ void FSceneRenderTargets::BeginRenderingLightAttenuation()
 {
 	SCOPED_DRAW_EVENT(BeginRenderingLightAttenuation, DEC_SCENE_ITEMS);
 
-	GRenderTargetPool.VisualizeTexture.SetCheckPoint(GSceneRenderTargets.LightAttenuation);
+	AllocLightAttenuation();
+
+	GRenderTargetPool.VisualizeTexture.SetCheckPoint(GSceneRenderTargets.GetLightAttenuation());
 
 	// Set the light attenuation surface as the render target, and the scene depth buffer as the depth-stencil surface.
 	RHISetRenderTarget(GetLightAttenuationSurface(),GetSceneDepthSurface());
@@ -419,7 +562,7 @@ void FSceneRenderTargets::FinishRenderingLightAttenuation()
 	// Resolve the light attenuation surface.
 	RHICopyToResolveTarget(GetLightAttenuationSurface(), LightAttenuation->GetRenderTargetItem().ShaderResourceTexture, false, FResolveParams(FResolveRect()));
 	
-	GRenderTargetPool.VisualizeTexture.SetCheckPoint(GSceneRenderTargets.LightAttenuation);
+	GRenderTargetPool.VisualizeTexture.SetCheckPoint(GSceneRenderTargets.GetLightAttenuation());
 }
 
 void FSceneRenderTargets::BeginRenderingTranslucency(const FViewInfo& View)
@@ -475,29 +618,6 @@ void FSceneRenderTargets::FinishRenderingSeparateTranslucency(const FViewInfo& V
 	}
 }
 
-void FSceneRenderTargets::BeginRenderingDistortionAccumulation()
-{
-	SCOPED_DRAW_EVENT(BeginRenderingDistortionAccumulation, DEC_SCENE_ITEMS);
-
-	// use RGBA8 light target for accumulating distortion offsets	
-	// R = positive X offset
-	// G = positive Y offset
-	// B = negative X offset
-	// A = negative Y offset
-
-	RHISetRenderTarget(GetLightAttenuationSurface(),GetSceneDepthSurface());
-}
-
-void FSceneRenderTargets::FinishRenderingDistortionAccumulation()
-{
-	SCOPED_DRAW_EVENT(FinishRenderingDistortionAccumulation, DEC_SCENE_ITEMS);
-
-	RHICopyToResolveTarget(GetLightAttenuationSurface(), GetLightAttenuationTexture(), false, FResolveParams());
-
-	// to be able to observe results with VisualizeTexture
-	GRenderTargetPool.VisualizeTexture.SetCheckPoint(LightAttenuation);
-}
-
 void FSceneRenderTargets::ResolveSceneDepthTexture()
 {
 	SCOPED_DRAW_EVENT(ResolveSceneDepthTexture, DEC_SCENE_ITEMS);
@@ -518,16 +638,6 @@ void FSceneRenderTargets::ResolveSceneDepthToAuxiliaryTexture()
 	}
 }
 
-void FSceneRenderTargets::BeginRenderingHitProxies()
-{
-	RHISetRenderTarget(GetHitProxySurface(),GetSceneDepthSurface());
-}
-
-void FSceneRenderTargets::FinishRenderingHitProxies()
-{
-	RHICopyToResolveTarget(GetHitProxySurface(), GetHitProxyTexture(), false, FResolveParams());
-}
-
 void FSceneRenderTargets::CleanUpEditorPrimitiveTargets()
 {
 	EditorPrimitivesDepth.SafeRelease();
@@ -667,18 +777,14 @@ void FSceneRenderTargets::AllocateForwardShadingPathRenderTargets()
 	{
 		// Create a texture to store the resolved scene depth, and a render-targetable surface to hold the unresolved scene depth.
 		FPooledRenderTargetDesc Desc(FPooledRenderTargetDesc::Create2DDesc(BufferSize, PF_DepthStencil, TexCreate_None, TexCreate_DepthStencilTargetable, false));
+		Desc.Flags |= TexCreate_FastVRAM;
 		GRenderTargetPool.FindFreeElement(Desc, SceneDepthZ, TEXT("SceneDepthZ"));
 	}
 
 	EPixelFormat Format = GSupportsRenderTargetFormat_PF_FloatRGBA ? PF_FloatRGBA : PF_B8G8R8A8;
 	if (!IsMobileHDR() || IsMobileHDR32bpp()) 
 	{
-			Format = PF_B8G8R8A8;
-	}
-	// Create the scene color.
-	{
-		FPooledRenderTargetDesc Desc(FPooledRenderTargetDesc::Create2DDesc(BufferSize, Format, TexCreate_None, TexCreate_RenderTargetable, false));
-		GRenderTargetPool.FindFreeElement(Desc, SceneColor, TEXT("SceneColor"));
+		Format = PF_B8G8R8A8;
 	}
 
 	// For 64-bit ES2 without framebuffer fetch, create extra render target for copy of alpha channel.
@@ -723,7 +829,7 @@ void FSceneRenderTargets::AllocateDeferredShadingPathRenderTargets()
 		Desc.Flags |= TexCreate_FastVRAM;
 		GRenderTargetPool.FindFreeElement(Desc, SceneDepthZ, TEXT("SceneDepthZ"));
 	}
-		
+
 	// When targeting DX Feature Level 10, create an auxiliary texture to store the resolved scene depth, and a render-targetable surface to hold the unresolved scene depth.
 	if(!GSupportsDepthFetchDuringDepthTest)
 	{
@@ -731,26 +837,6 @@ void FSceneRenderTargets::AllocateDeferredShadingPathRenderTargets()
 		GRenderTargetPool.FindFreeElement(Desc, AuxiliarySceneDepthZ, TEXT("AuxiliarySceneDepthZ"));
 	}
 
-	// Potentially allocate an alpha channel in the scene color texture to store the resolved scene depth.
-	EPixelFormat SceneColorBufferFormat = PF_FloatRGBA;
-
-	switch(CurrentSceneColorFormat)
-	{
-		case 0:
-			SceneColorBufferFormat = PF_R8G8B8A8; break;
-		case 1:
-			SceneColorBufferFormat = PF_A2B10G10R10; break;
-		case 2:	
-			SceneColorBufferFormat = PF_FloatR11G11B10; break;
-		case 3:	
-			SceneColorBufferFormat = PF_FloatRGB; break;
-		case 4:
-			// default
-			break;
-		case 5:
-			SceneColorBufferFormat = PF_A32B32G32R32F; break;
-	}
-
 	// Create a quarter-sized version of the scene depth.
 	{
 		FIntPoint SmallDepthZSize(FMath::Max<uint32>(BufferSize.X / SmallColorDepthDownsampleFactor, 1), FMath::Max<uint32>(BufferSize.Y / SmallColorDepthDownsampleFactor, 1));
@@ -759,22 +845,6 @@ void FSceneRenderTargets::AllocateDeferredShadingPathRenderTargets()
 		GRenderTargetPool.FindFreeElement(Desc, SmallDepthZ, TEXT("SmallDepthZ"));
 	}
 
-	// Create the scene color.
-	{
-		FPooledRenderTargetDesc Desc(FPooledRenderTargetDesc::Create2DDesc(BufferSize, SceneColorBufferFormat, TexCreate_None, TexCreate_RenderTargetable, false));
-		Desc.Flags |= TexCreate_FastVRAM;
-		GRenderTargetPool.FindFreeElement(Desc, SceneColor, TEXT("SceneColor"));
-	}
-
-	// Create a texture to store the resolved light attenuation values, and a render-targetable surface to hold the unresolved light attenuation values.
-	{
-		FPooledRenderTargetDesc Desc(FPooledRenderTargetDesc::Create2DDesc(BufferSize, PF_B8G8R8A8, TexCreate_None, TexCreate_RenderTargetable, false));
-		Desc.Flags |= TexCreate_FastVRAM;
-		GRenderTargetPool.FindFreeElement(Desc, LightAttenuation, TEXT("LightAttenuation"));
-
-		// the channel assignment is documented in ShadowRendering.cpp (look for Light Attenuation channel assignment)
-	}
-
 	// Set up quarter size scene color shared texture
 	const FIntPoint ShadowBufferResolution = GetShadowDepthTextureResolution();
 
@@ -964,6 +1034,8 @@ void FSceneRenderTargets::AllocateDeferredShadingPathRenderTargets()
 	}
 
 	{
+		EPixelFormat SceneColorBufferFormat = GetSceneColorFormat();
+
 		uint32 LightAccumulationUAVFlag = GRHIFeatureLevel == ERHIFeatureLevel::SM5 ? TexCreate_UAV : 0;
 		FPooledRenderTargetDesc Desc(FPooledRenderTargetDesc::Create2DDesc(BufferSize, SceneColorBufferFormat, TexCreate_None, TexCreate_RenderTargetable | LightAccumulationUAVFlag, false));
 		GRenderTargetPool.FindFreeElement(Desc, LightAccumulation, TEXT("LightAccumulation"));
@@ -994,6 +1066,31 @@ void FSceneRenderTargets::AllocateDeferredShadingPathRenderTargets()
 	}
 }
 
+EPixelFormat FSceneRenderTargets::GetSceneColorFormat() const
+{
+	// Potentially allocate an alpha channel in the scene color texture to store the resolved scene depth.
+	EPixelFormat SceneColorBufferFormat = PF_FloatRGBA;
+
+	switch(CurrentSceneColorFormat)
+	{
+		case 0:
+			SceneColorBufferFormat = PF_R8G8B8A8; break;
+		case 1:
+			SceneColorBufferFormat = PF_A2B10G10R10; break;
+		case 2:	
+			SceneColorBufferFormat = PF_FloatR11G11B10; break;
+		case 3:	
+			SceneColorBufferFormat = PF_FloatRGB; break;
+		case 4:
+			// default
+			break;
+		case 5:
+			SceneColorBufferFormat = PF_A32B32G32R32F; break;
+	}
+
+	return SceneColorBufferFormat;
+}
+
 void FSceneRenderTargets::InitDynamicRHI()
 {
 	if(BufferSize.X > 0 && BufferSize.Y > 0)
@@ -1123,6 +1220,26 @@ static int32 GetCustomDepthCVar()
 	return CVar->GetValueOnRenderThread();
 }
 
+const FTextureRHIRef& FSceneRenderTargets::GetSceneColorSurface() const							
+{
+	if(!SceneColor)
+	{
+		return GBlackTexture->TextureRHI;
+	}
+
+	return (const FTextureRHIRef&)SceneColor->GetRenderTargetItem().TargetableTexture; 
+}
+
+const FTextureRHIRef& FSceneRenderTargets::GetSceneColorTexture() const
+{
+	if(!SceneColor)
+	{
+		return GBlackTexture->TextureRHI;
+	}
+
+	return (const FTextureRHIRef&)GetSceneColor()->GetRenderTargetItem().ShaderResourceTexture; 
+}
+
 IPooledRenderTarget* FSceneRenderTargets::RequestCustomDepth(bool bPrimitives)
 {
 	int Value = GetCustomDepthCVar();
diff --git a/Engine/Source/Runtime/Renderer/Private/PostProcess/SceneRenderTargets.h b/Engine/Source/Runtime/Renderer/Private/PostProcess/SceneRenderTargets.h
--- a/Engine/Source/Runtime/Renderer/Private/PostProcess/SceneRenderTargets.h
+++ b/Engine/Source/Runtime/Renderer/Private/PostProcess/SceneRenderTargets.h
@@ -171,9 +171,6 @@ public:
 	void FinishRenderingSeparateTranslucency(const FViewInfo& View);
 	void FreeSeparateTranslucency();
 
-	void BeginRenderingDistortionAccumulation();
-	void FinishRenderingDistortionAccumulation();
-
 	void ResolveSceneDepthTexture();
 	void ResolveSceneDepthToAuxiliaryTexture();
 
@@ -186,9 +183,6 @@ public:
 	void BeginRenderingLightAttenuation();
 	void FinishRenderingLightAttenuation();
 
-	void BeginRenderingHitProxies();
-	void FinishRenderingHitProxies();
-
 	/**
 	 * Cleans up editor primitive targets that we no longer need
 	 */
@@ -221,7 +215,7 @@ public:
 
 	// Texture Accessors -----------
 
-	const FTexture2DRHIRef& GetSceneColorTexture() const { return (const FTexture2DRHIRef&)SceneColor->GetRenderTargetItem().ShaderResourceTexture; }
+	const FTextureRHIRef& GetSceneColorTexture() const;
 	const FTexture2DRHIRef& GetSceneAlphaCopyTexture() const { return (const FTexture2DRHIRef&)SceneAlphaCopy->GetRenderTargetItem().ShaderResourceTexture; }
 	bool HasSceneAlphaCopyTexture() const { return SceneAlphaCopy.GetReference() != 0; }
 	const FTexture2DRHIRef& GetSceneDepthTexture() const { return (const FTexture2DRHIRef&)SceneDepthZ->GetRenderTargetItem().ShaderResourceTexture; }
@@ -267,7 +261,6 @@ public:
 	{ 
 		return (const FTextureCubeRHIRef&)CubeShadowDepthZ[GetCubeShadowDepthZIndex(ShadowResolution)]->GetRenderTargetItem().ShaderResourceTexture; 
 	}
-	const FTexture2DRHIRef& GetHitProxyTexture() const { return (const FTexture2DRHIRef&)LightAttenuation->GetRenderTargetItem().ShaderResourceTexture; }
 	const FTexture2DRHIRef& GetGBufferATexture() const { return (const FTexture2DRHIRef&)GBufferA->GetRenderTargetItem().ShaderResourceTexture; }
 
 	/** 
@@ -288,14 +281,14 @@ public:
 	}
 	const FTextureRHIRef& GetLightAttenuationTexture() const
 	{
-		return *(FTextureRHIRef*)&LightAttenuation->GetRenderTargetItem().ShaderResourceTexture;
+		return *(FTextureRHIRef*)&GetLightAttenuation()->GetRenderTargetItem().ShaderResourceTexture;
 	}
 	const FTextureRHIRef& GetLightAccumulationTexture() const
 	{
 		return *(FTextureRHIRef*)&LightAccumulation->GetRenderTargetItem().ShaderResourceTexture;
 	}
 
-	const FTexture2DRHIRef& GetSceneColorSurface() const							{ return (const FTexture2DRHIRef&)SceneColor->GetRenderTargetItem().TargetableTexture; }
+	const FTextureRHIRef& GetSceneColorSurface() const;
 	const FTexture2DRHIRef& GetSceneAlphaCopySurface() const						{ return (const FTexture2DRHIRef&)SceneAlphaCopy->GetRenderTargetItem().TargetableTexture; }
 	const FTexture2DRHIRef& GetSceneDepthSurface() const							{ return (const FTexture2DRHIRef&)SceneDepthZ->GetRenderTargetItem().TargetableTexture; }
 	const FTexture2DRHIRef& GetSmallDepthSurface() const							{ return (const FTexture2DRHIRef&)SmallDepthZ->GetRenderTargetItem().TargetableTexture; }
@@ -312,8 +305,7 @@ public:
 	{ 
 		return (const FTextureCubeRHIRef&)CubeShadowDepthZ[GetCubeShadowDepthZIndex(ShadowResolution)]->GetRenderTargetItem().TargetableTexture; 
 	}
-	const FTexture2DRHIRef& GetLightAttenuationSurface() const					{ return (const FTexture2DRHIRef&)LightAttenuation->GetRenderTargetItem().TargetableTexture; }
-	const FTexture2DRHIRef& GetHitProxySurface() const							{ return (const FTexture2DRHIRef&)LightAttenuation->GetRenderTargetItem().TargetableTexture; }
+	const FTexture2DRHIRef& GetLightAttenuationSurface() const					{ return (const FTexture2DRHIRef&)GetLightAttenuation()->GetRenderTargetItem().TargetableTexture; }
 	const FTexture2DRHIRef& GetAuxiliarySceneDepthSurface() const 
 	{	
 		check(!GSupportsDepthFetchDuringDepthTest); 
@@ -359,10 +351,38 @@ public:
 	int32 GetNumGBufferTargets() const;
 
 	// ---
+
+	// needs to be called between AllocSceneColor() and ReleaseSceneColor()
+	const TRefCountPtr<IPooledRenderTarget>& GetSceneColor() const;
+
+	TRefCountPtr<IPooledRenderTarget>& GetSceneColor();
+
+	void SetSceneColor(IPooledRenderTarget* In);
+
+	// ---
+
+	void SetLightAttenuation(IPooledRenderTarget* In);
+
+	// needs to be called between AllocSceneColor() and SetSceneColor(0)
+	const TRefCountPtr<IPooledRenderTarget>& GetLightAttenuation() const;
+
+	TRefCountPtr<IPooledRenderTarget>& GetLightAttenuation();
+
+	// ---
+
+private: // use SceneColor() or LightAttenuation() instead of direct access
+
+	// 0 before BeginRenderingSceneColor and after tone mapping
 	TRefCountPtr<IPooledRenderTarget> SceneColor;
-	TRefCountPtr<IPooledRenderTarget> SceneAlphaCopy; // Mobile without framebuffer fetch (to get depth from alpha).
+	// also used as LDR scene color
+	TRefCountPtr<IPooledRenderTarget> LightAttenuation;
+public:
+
+	//
 	TRefCountPtr<IPooledRenderTarget> SceneDepthZ;
-	// Auxillary scene depth target. The scene depth is resolved to this surface when targeting SM4. 
+	// Mobile without frame buffer fetch (to get depth from alpha).
+	TRefCountPtr<IPooledRenderTarget> SceneAlphaCopy;
+	// Auxiliary scene depth target. The scene depth is resolved to this surface when targeting SM4. 
 	TRefCountPtr<IPooledRenderTarget> AuxiliarySceneDepthZ;
 	// Render target for a quarter-sized version of the scene depths.
 	TRefCountPtr<IPooledRenderTarget> SmallDepthZ;
@@ -381,8 +401,6 @@ public:
 
 	// for AmbientOcclusion, only valid for a short time during the frame to allow reuse
 	TRefCountPtr<IPooledRenderTarget> ScreenSpaceAO;
-	// also used for hit proxy and as LDR scene color
-	TRefCountPtr<IPooledRenderTarget> LightAttenuation;
 	// used by the CustomDepth material feature, is allocated on demand or if r.CustomDepth is 2
 	TRefCountPtr<IPooledRenderTarget> CustomDepth;
 	// Used for accumulating tiled deferred VPL indirect lighting
@@ -461,6 +479,12 @@ private:
 	/** Determine the appropriate render target dimensions. */
 	FIntPoint GetSceneRenderTargetSize(const FSceneViewFamily & ViewFamily) const;
 
+	void AllocSceneColor();
+
+	void AllocLightAttenuation();
+
+	EPixelFormat GetSceneColorFormat() const;
+
 private:
 
 	/** size of the back buffer, in editor this has to be >= than the biggest view port */
@@ -562,7 +586,7 @@ public:
 					);
 			}
 
-			if (SceneDepthTextureParameter.IsBound())
+			if(SceneDepthTextureParameter.IsBound() || SceneDepthTextureParameterSampler.IsBound())
 			{
 				const FTexture2DRHIRef* DepthTexture = GSceneRenderTargets.GetActualDepthTexture();
 				SetTextureParameter(
@@ -582,13 +606,25 @@ public:
 			{
 				if(GSupportsDepthFetchDuringDepthTest)
 				{
-					SetTextureParameter(ShaderRHI, SceneDepthSurfaceParameter, GSceneRenderTargets.GetSceneDepthSurface());
-					SetTextureParameter(ShaderRHI, SceneDepthTextureNonMS, GSceneRenderTargets.GetSceneDepthTexture());
+					if(SceneDepthSurfaceParameter.IsBound())
+					{
+						SetTextureParameter(ShaderRHI, SceneDepthSurfaceParameter, GSceneRenderTargets.GetSceneDepthSurface());
+					}
+					if(SceneDepthTextureNonMS.IsBound())
+					{
+						SetTextureParameter(ShaderRHI, SceneDepthTextureNonMS, GSceneRenderTargets.GetSceneDepthTexture());
+					}
 				}
 				else
 				{
-					SetTextureParameter(ShaderRHI, SceneDepthSurfaceParameter, GSceneRenderTargets.GetAuxiliarySceneDepthSurface());
-					SetTextureParameter(ShaderRHI, SceneDepthTextureNonMS, GSceneRenderTargets.GetAuxiliarySceneDepthSurface());
+					if(SceneDepthSurfaceParameter.IsBound())
+					{
+						SetTextureParameter(ShaderRHI, SceneDepthSurfaceParameter, GSceneRenderTargets.GetAuxiliarySceneDepthSurface());
+					}
+					if(SceneDepthTextureNonMS.IsBound())
+					{
+						SetTextureParameter(ShaderRHI, SceneDepthTextureNonMS, GSceneRenderTargets.GetAuxiliarySceneDepthSurface());
+					}
 				}
 			}
 		}
diff --git a/Engine/Source/Runtime/Renderer/Private/PostProcess/ScreenSpaceReflections.cpp b/Engine/Source/Runtime/Renderer/Private/PostProcess/ScreenSpaceReflections.cpp
--- a/Engine/Source/Runtime/Renderer/Private/PostProcess/ScreenSpaceReflections.cpp
+++ b/Engine/Source/Runtime/Renderer/Private/PostProcess/ScreenSpaceReflections.cpp
@@ -254,7 +254,7 @@ void ScreenSpaceReflections( const FViewInfo& View, TRefCountPtr<IPooledRenderTa
 
 	FSceneViewState* ViewState = (FSceneViewState*)Context.View.State;
 
-	FRenderingCompositePass* SceneColorInput = Context.Graph.RegisterPass( new FRCPassPostProcessInput( GSceneRenderTargets.SceneColor ) );
+	FRenderingCompositePass* SceneColorInput = Context.Graph.RegisterPass( new FRCPassPostProcessInput( GSceneRenderTargets.GetSceneColor() ) );
 	FRenderingCompositePass* HZBInput = Context.Graph.RegisterPass( new FRCPassPostProcessInput( ViewState->HZB.Texture ) );
 
 	bool bPrevFrame = 0;
diff --git a/Engine/Source/Runtime/Renderer/Private/ReflectionEnvironment.cpp b/Engine/Source/Runtime/Renderer/Private/ReflectionEnvironment.cpp
--- a/Engine/Source/Runtime/Renderer/Private/ReflectionEnvironment.cpp
+++ b/Engine/Source/Runtime/Renderer/Private/ReflectionEnvironment.cpp
@@ -197,7 +197,7 @@ public:
 
 		SetTextureParameter( ShaderRHI, ScreenSpaceReflections, SSRTexture );
 
-		SetTextureParameter( ShaderRHI, InSceneColor, GSceneRenderTargets.SceneColor->GetRenderTargetItem().ShaderResourceTexture );
+		SetTextureParameter( ShaderRHI, InSceneColor, GSceneRenderTargets.GetSceneColor()->GetRenderTargetItem().ShaderResourceTexture );
 		OutSceneColor.SetTexture(ShaderRHI, NULL, OutSceneColorUAV);
 
 		SetShaderValue(ShaderRHI, ViewDimensionsParameter, View.ViewRect);
@@ -570,7 +570,7 @@ void FDeferredShadingSceneRenderer::RenderDeferredReflections()
 	{
 		GSceneRenderTargets.ResolveSceneColor(FResolveRect(0, 0, ViewFamily.FamilySizeX, ViewFamily.FamilySizeY));
 
-		FPooledRenderTargetDesc Desc = GSceneRenderTargets.SceneColor->GetDesc();
+		FPooledRenderTargetDesc Desc = GSceneRenderTargets.GetSceneColor()->GetDesc();
 		Desc.Flags |= TexCreate_UAV;
 
 		GRenderTargetPool.FindFreeElement( Desc, NewSceneColor, TEXT("SceneColor") );
@@ -656,8 +656,8 @@ void FDeferredShadingSceneRenderer::RenderDeferredReflections()
 					const FSceneRenderTargetItem& DestRenderTarget = NewSceneColor->GetRenderTargetItem();
 					RHICopyToResolveTarget( DestRenderTarget.TargetableTexture, DestRenderTarget.ShaderResourceTexture, false, FResolveParams() );
 
-					GSceneRenderTargets.SceneColor = NewSceneColor;
-					check(GSceneRenderTargets.SceneColor);
+					GSceneRenderTargets.SetSceneColor(NewSceneColor);
+					check(GSceneRenderTargets.GetSceneColor());
 
 					bRequiresApply = false;
 				}
diff --git a/Engine/Source/Runtime/Renderer/Private/ReflectionEnvironmentCapture.cpp b/Engine/Source/Runtime/Renderer/Private/ReflectionEnvironmentCapture.cpp
--- a/Engine/Source/Runtime/Renderer/Private/ReflectionEnvironmentCapture.cpp
+++ b/Engine/Source/Runtime/Renderer/Private/ReflectionEnvironmentCapture.cpp
@@ -624,7 +624,7 @@ public:
 			InTexture, 
 			InTextureSampler, 
 			TStaticSamplerState<SF_Bilinear, AM_Clamp, AM_Clamp, AM_Clamp>::GetRHI(), 
-			GSceneRenderTargets.SceneColor->GetRenderTargetItem().ShaderResourceTexture);		
+			GSceneRenderTargets.GetSceneColor()->GetRenderTargetItem().ShaderResourceTexture);		
 
 		FVector SkyLightParametersValue = FVector::ZeroVector;
 		FScene* Scene = (FScene*)View.Family->Scene;
diff --git a/Engine/Source/Runtime/Renderer/Private/SceneHitProxyRendering.cpp b/Engine/Source/Runtime/Renderer/Private/SceneHitProxyRendering.cpp
--- a/Engine/Source/Runtime/Renderer/Private/SceneHitProxyRendering.cpp
+++ b/Engine/Source/Runtime/Renderer/Private/SceneHitProxyRendering.cpp
@@ -332,8 +332,22 @@ void FDeferredShadingSceneRenderer::RenderHitProxies()
 	// Allocate the maximum scene render target space for the current view family.
 	GSceneRenderTargets.Allocate(ViewFamily);
 
-	// Write to the hit proxy render target.
-	GSceneRenderTargets.BeginRenderingHitProxies();
+	TRefCountPtr<IPooledRenderTarget> HitProxyRT;
+
+	// Create a texture to store the resolved light attenuation values, and a render-targetable surface to hold the unresolved light attenuation values.
+	{
+		FPooledRenderTargetDesc Desc(FPooledRenderTargetDesc::Create2DDesc(GSceneRenderTargets.GetBufferSizeXY(), PF_B8G8R8A8, TexCreate_None, TexCreate_RenderTargetable, false));
+		Desc.Flags |= TexCreate_FastVRAM;
+		GRenderTargetPool.FindFreeElement(Desc, HitProxyRT, TEXT("HitProxy"));
+	}
+
+	if(!HitProxyRT)
+	{
+		// HitProxyRT==0 should never happen but better we don't crash
+		return;
+	}
+
+	RHISetRenderTarget(HitProxyRT->GetRenderTargetItem().TargetableTexture, GSceneRenderTargets.GetSceneDepthSurface());
 
 	// Clear color for each view.
 	for(int32 ViewIndex = 0;ViewIndex < Views.Num();ViewIndex++)
@@ -427,7 +441,9 @@ void FDeferredShadingSceneRenderer::RenderHitProxies()
 	}
 
 	// Finish drawing to the hit proxy render target.
-	GSceneRenderTargets.FinishRenderingHitProxies();
+	RHICopyToResolveTarget(HitProxyRT->GetRenderTargetItem().TargetableTexture, HitProxyRT->GetRenderTargetItem().ShaderResourceTexture, false, FResolveParams());
+	// to be able to observe results with VisualizeTexture
+	GRenderTargetPool.VisualizeTexture.SetCheckPoint(HitProxyRT);
 
 	// After scene rendering, disable the depth buffer.
 	RHISetDepthStencilState(TStaticDepthStencilState<false,CF_Always>::GetRHI());
@@ -438,7 +454,7 @@ void FDeferredShadingSceneRenderer::RenderHitProxies()
 	
 	// Set up a FTexture that is used to draw the hit proxy buffer to the view family's render target.
 	FTexture HitProxyRenderTargetTexture;
-	HitProxyRenderTargetTexture.TextureRHI = GSceneRenderTargets.GetHitProxyTexture();
+	HitProxyRenderTargetTexture.TextureRHI = HitProxyRT->GetRenderTargetItem().ShaderResourceTexture;
 	HitProxyRenderTargetTexture.SamplerStateRHI = TStaticSamplerState<>::GetRHI();
 
 	// Generate the vertices and triangles mapping the hit proxy RT pixels into the view family's RT pixels.
diff --git a/Engine/Source/Runtime/Renderer/Private/ScenePostProcessing.cpp b/Engine/Source/Runtime/Renderer/Private/ScenePostProcessing.cpp
--- a/Engine/Source/Runtime/Renderer/Private/ScenePostProcessing.cpp
+++ b/Engine/Source/Runtime/Renderer/Private/ScenePostProcessing.cpp
@@ -142,7 +142,7 @@ void FSceneRenderer::GammaCorrectToViewportRenderTarget(const FViewInfo* View, f
 	SetShaderValue(ShaderRHI,PixelShader->ColorScale,View->ColorScale);
 	SetShaderValue(ShaderRHI,PixelShader->OverlayColor,View->OverlayColor);
 
-	const FTexture2DRHIRef DesiredSceneColorTexture = GSceneRenderTargets.GetSceneColorTexture();
+	const FTextureRHIRef DesiredSceneColorTexture = GSceneRenderTargets.GetSceneColorTexture();
 
 	SetTextureParameter(
 		ShaderRHI,
diff --git a/Engine/Source/Runtime/Renderer/Private/ShaderBaseClasses.h b/Engine/Source/Runtime/Renderer/Private/ShaderBaseClasses.h
--- a/Engine/Source/Runtime/Renderer/Private/ShaderBaseClasses.h
+++ b/Engine/Source/Runtime/Renderer/Private/ShaderBaseClasses.h
@@ -278,12 +278,15 @@ public:
 
 		if (GRHIFeatureLevel >= ERHIFeatureLevel::SM3)
 		{
-			SetTextureParameter(
-				ShaderRHI,
-				LightAttenuation,
-				LightAttenuationSampler,
-				TStaticSamplerState<SF_Bilinear,AM_Clamp,AM_Clamp,AM_Clamp>::GetRHI(),
-				GSceneRenderTargets.GetLightAttenuationTexture());
+			if(LightAttenuation.IsBound())
+			{
+				SetTextureParameter(
+					ShaderRHI,
+					LightAttenuation,
+					LightAttenuationSampler,
+					TStaticSamplerState<SF_Bilinear,AM_Clamp,AM_Clamp,AM_Clamp>::GetRHI(),
+					GSceneRenderTargets.GetLightAttenuationTexture());
+			}
 		}
 
 		// if we are in a postprocessing pass
