Improve support for Mac <= GL 3.3 drivers that don't support Volume Texture Rendering properly.
On Mac SM4, cache shaders for indirect point lighting, so that they can be used for rendering on GPUs without VTR.
This removes the workaround for indirect lighting that had been required for ES2 as the runtime test is now in CanIndirectLightingCacheUseVolumeTexture where it should always have been.
Make all Macs on <= SM4 use a 2D tonemapping color LUT so that they can share the same shaders.
This makes all Macs share the same binary shader cache which should fix the binary UE4 builds in the simplest way.

diff --git a/Engine/Source/Runtime/Renderer/Private/IndirectLightingCache.cpp b/Engine/Source/Runtime/Renderer/Private/IndirectLightingCache.cpp
--- a/Engine/Source/Runtime/Renderer/Private/IndirectLightingCache.cpp
+++ b/Engine/Source/Runtime/Renderer/Private/IndirectLightingCache.cpp
@@ -135,12 +135,13 @@ bool IsIndirectLightingCacheAllowed()
 	static const auto AllowStaticLightingVar = IConsoleManager::Get().FindTConsoleVariableDataInt(TEXT("r.AllowStaticLighting"));
 	const bool bAllowStaticLighting = (!AllowStaticLightingVar || AllowStaticLightingVar->GetValueOnRenderThread() != 0);
 
-	return GIndirectLightingCache != 0 && (GSupportsVolumeTextureRendering || GRHIFeatureLevel == ERHIFeatureLevel::ES2) && bAllowStaticLighting;
+	return GIndirectLightingCache != 0 && bAllowStaticLighting;
 }
 
 bool CanIndirectLightingCacheUseVolumeTexture()
 {
-	return GRHIFeatureLevel >= ERHIFeatureLevel::SM3;
+	// @todo Mac OS X/OpenGL: For OpenGL devices which don't support volume-texture rendering we need to use the simpler point indirect lighting shaders.
+	return GRHIFeatureLevel >= ERHIFeatureLevel::SM3 && GSupportsVolumeTextureRendering;
 }
 
 FIndirectLightingCache::FIndirectLightingCache()
diff --git a/Engine/Source/Runtime/Renderer/Private/LightMapRendering.h b/Engine/Source/Runtime/Renderer/Private/LightMapRendering.h
--- a/Engine/Source/Runtime/Renderer/Private/LightMapRendering.h
+++ b/Engine/Source/Runtime/Renderer/Private/LightMapRendering.h
@@ -718,9 +718,11 @@ public:
 	static bool ShouldCache(EShaderPlatform Platform,const FMaterial* Material,const FVertexFactoryType* VertexFactoryType)
 	{
 		static const auto AllowStaticLightingVar = IConsoleManager::Get().FindTConsoleVariableDataInt(TEXT("r.AllowStaticLighting"));
-
+		
+		// @todo Mac OS X: For GL 3.3 devices which don't support volume-texture rendering we need to cache the simpler point indirect lighting shaders.
 		return Material->GetLightingModel() != MLM_Unlit 
-			&& (IsTranslucentBlendMode(Material->GetBlendMode()) || GetMaxSupportedFeatureLevel(Platform) == ERHIFeatureLevel::ES2) 
+			&& (IsTranslucentBlendMode(Material->GetBlendMode()) || GetMaxSupportedFeatureLevel(Platform) == ERHIFeatureLevel::ES2
+				|| (PLATFORM_MAC && GetMaxSupportedFeatureLevel(Platform) == ERHIFeatureLevel::SM4))
 			&& (!AllowStaticLightingVar || AllowStaticLightingVar->GetValueOnAnyThread() != 0);
 	}
 
diff --git a/Engine/Source/Runtime/Renderer/Private/PostProcess/PostProcessCombineLUTs.cpp b/Engine/Source/Runtime/Renderer/Private/PostProcess/PostProcessCombineLUTs.cpp
--- a/Engine/Source/Runtime/Renderer/Private/PostProcess/PostProcessCombineLUTs.cpp
+++ b/Engine/Source/Runtime/Renderer/Private/PostProcess/PostProcessCombineLUTs.cpp
@@ -13,9 +13,10 @@
 
 // false:use 256x16 texture / true:use volume texture (faster, requires geometry shader)
 // USE_VOLUME_LUT: needs to be the same for C++ and HLSL
-static bool UseVolumeTextureLUT() 
-{	
-	return (GRHIFeatureLevel >= ERHIFeatureLevel::SM4 && GSupportsVolumeTextureRendering);
+static bool UseVolumeTextureLUT(EShaderPlatform Platform) 
+{
+	// @todo Mac OS X: in order to share precompiled shaders between GL 3.3 & GL 4.1 devices we mustn't use volume-texture rendering as it isn't universally supported.
+	return (IsFeatureLevelSupported(Platform,ERHIFeatureLevel::SM4) && GSupportsVolumeTextureRendering && !PLATFORM_MAC);
 }
 
 // including the neutral one at index 0
@@ -143,7 +144,7 @@ public:
 		FGlobalShader::ModifyCompilationEnvironment(Platform, OutEnvironment);
 
 		OutEnvironment.SetDefine(TEXT("BLENDCOUNT"), BlendCount);
-		OutEnvironment.SetDefine(TEXT("USE_VOLUME_LUT"), UseVolumeTextureLUT());
+		OutEnvironment.SetDefine(TEXT("USE_VOLUME_LUT"), UseVolumeTextureLUT(Platform));
 	}
 
 	virtual bool Serialize(FArchive& Ar)
@@ -211,7 +212,7 @@ void SetLUTBlenderShader(FRenderingCompositePassContext& Context, uint32 BlendCo
 	}
 #undef CASE_COUNT
 	check(LocalBoundShaderState != NULL);
-	if(UseVolumeTextureLUT())
+	if(UseVolumeTextureLUT(GRHIShaderPlatform))
 	{
 		TShaderMapRef<FWriteToSliceVS> VertexShader(GetGlobalShaderMap());
 		TShaderMapRef<FWriteToSliceGS> GeometryShader(GetGlobalShaderMap());
@@ -381,7 +382,7 @@ void FRCPassPostProcessCombineLUTs::Process(FRenderingCompositePassContext& Cont
 	}
 
 	// for a 3D texture, the viewport is 16x16 (per slice), for a 2D texture, it's unwrapped to 256x16
-	FIntPoint DestSize(UseVolumeTextureLUT() ? 16 : 256, 16);
+	FIntPoint DestSize(UseVolumeTextureLUT(GRHIShaderPlatform) ? 16 : 256, 16);
 
 	const FSceneRenderTargetItem& DestRenderTarget = PassOutputs[0].RequestSurface(Context);
 
@@ -398,7 +399,7 @@ void FRCPassPostProcessCombineLUTs::Process(FRenderingCompositePassContext& Cont
 
 	SetLUTBlenderShader(Context, LocalCount, LocalTextures, LocalWeights, VolumeBounds);
 
-	if(UseVolumeTextureLUT())
+	if(UseVolumeTextureLUT(GRHIShaderPlatform))
 	{
 		// use volume texture 16x16x16
 		RasterizeToVolumeTexture(VolumeBounds);
@@ -426,7 +427,7 @@ FPooledRenderTargetDesc FRCPassPostProcessCombineLUTs::ComputeOutputDesc(EPassOu
 {
 	FPooledRenderTargetDesc Ret = FPooledRenderTargetDesc::Create2DDesc(FIntPoint(256, 16), PF_B8G8R8A8, TexCreate_None, TexCreate_RenderTargetable | TexCreate_ShaderResource, false);
 
-	if(UseVolumeTextureLUT())
+	if(UseVolumeTextureLUT(GRHIShaderPlatform))
 	{
 		Ret.Extent = FIntPoint(16, 16);
 		Ret.Depth = 16;
diff --git a/Engine/Source/Runtime/Renderer/Private/PostProcess/PostProcessTonemap.cpp b/Engine/Source/Runtime/Renderer/Private/PostProcess/PostProcessTonemap.cpp
--- a/Engine/Source/Runtime/Renderer/Private/PostProcess/PostProcessTonemap.cpp
+++ b/Engine/Source/Runtime/Renderer/Private/PostProcess/PostProcessTonemap.cpp
@@ -747,7 +747,8 @@ class FPostProcessTonemapPS : public FGlobalShader
 		OutEnvironment.SetDefine(TEXT("USE_GRAIN_QUANTIZATION"), TonemapperIsDefined(ConfigBitmask, TonemapperGrainQuantization));
 		OutEnvironment.SetDefine(TEXT("USE_VIGNETTE"),           TonemapperIsDefined(ConfigBitmask, TonemapperVignette));
 		OutEnvironment.SetDefine(TEXT("USE_VIGNETTE_COLOR"),     TonemapperIsDefined(ConfigBitmask, TonemapperVignetteColor));
-		OutEnvironment.SetDefine(TEXT("USE_VOLUME_LUT"),		 (GRHIFeatureLevel >= ERHIFeatureLevel::SM4 && GSupportsVolumeTextureRendering));
+		// @todo Mac OS X: in order to share precompiled shaders between GL 3.3 & GL 4.1 devices we mustn't use volume-texture rendering as it isn't universally supported.
+		OutEnvironment.SetDefine(TEXT("USE_VOLUME_LUT"),		 (IsFeatureLevelSupported(Platform,ERHIFeatureLevel::SM4) && GSupportsVolumeTextureRendering && !PLATFORM_MAC));
 
 		if( !IsFeatureLevelSupported(Platform,ERHIFeatureLevel::SM5) )
 		{
