* improved rendering of r.RenderTargetPool.Events

diff --git a/Engine/Source/Runtime/Renderer/Private/PostProcess/RenderTargetPool.cpp b/Engine/Source/Runtime/Renderer/Private/PostProcess/RenderTargetPool.cpp
--- a/Engine/Source/Runtime/Renderer/Private/PostProcess/RenderTargetPool.cpp
+++ b/Engine/Source/Runtime/Renderer/Private/PostProcess/RenderTargetPool.cpp
@@ -519,7 +519,7 @@ uint32 FRenderTargetPool::ComputeEventDisplayHeight()
 	{
 	}
 
-	return It.Y + 1;
+	return It.Y;
 }
 
 const FString* FRenderTargetPool::GetLastEventPhaseName()
@@ -668,6 +668,8 @@ void FRenderTargetPool::PresentContent(const FSceneView& View)
 {
 	if(RenderTargetPoolEvents.Num())
 	{
+		AddPhaseEvent(TEXT("FrameEnd"));
+
 		FIntPoint DisplayLeftTop(20, 50);
 		// on the right we leave more space to make the mouse tooltip readable
 		FIntPoint DisplayExtent(View.ViewRect.Width() - DisplayLeftTop.X * 2 - 140, View.ViewRect.Height() - DisplayLeftTop.Y * 2);
