* nicer printout

diff --git a/Engine/Source/Runtime/Renderer/Private/PostProcess/RenderTargetPool.cpp b/Engine/Source/Runtime/Renderer/Private/PostProcess/RenderTargetPool.cpp
--- a/Engine/Source/Runtime/Renderer/Private/PostProcess/RenderTargetPool.cpp
+++ b/Engine/Source/Runtime/Renderer/Private/PostProcess/RenderTargetPool.cpp
@@ -817,7 +817,9 @@ void FRenderTargetPool::PresentContent(const FSceneView& View)
 
 				if(HighlightedEvent->GetEventType() == ERTPE_Phase)
 				{
-					Canvas.DrawShadowedString(Pos.X, Pos.Y + 0 * FontHeight, *HighlightedEvent->GetPhaseName(), GEngine->GetTinyFont(), FLinearColor(0.5f, 0.5f, 1));
+					FString PhaseText = *FString::Printf(TEXT("Phase: %s"), *HighlightedEvent->GetPhaseName()); 
+
+					Canvas.DrawShadowedString(Pos.X, Pos.Y + 0 * FontHeight, *PhaseText, GEngine->GetTinyFont(), FLinearColor(0.5f, 0.5f, 1));
 				}
 				else
 				{
