* fixed ensure when capturing HDR scene color

diff --git a/Engine/Source/Runtime/Renderer/Private/DeferredShadingRenderer.cpp b/Engine/Source/Runtime/Renderer/Private/DeferredShadingRenderer.cpp
--- a/Engine/Source/Runtime/Renderer/Private/DeferredShadingRenderer.cpp
+++ b/Engine/Source/Runtime/Renderer/Private/DeferredShadingRenderer.cpp
@@ -427,9 +427,9 @@ void FDeferredShadingSceneRenderer::RenderFinish()
 		Scene->MotionBlurInfoData.UpdateMotionBlurCache();
 	}
 
-	// The SceneColor RT should be released as early as possible to allow shading of that memory for other purposes.
+	// Some RT should be released as early as possible to allow shading of that memory for other purposes.
+	// SceneColor should be released in tone mapping.
 	// This becomes even more important with some limited VRam (XBoxOne).
-	GSceneRenderTargets.SetSceneColor(0);
 	GSceneRenderTargets.SetLightAttenuation(0);
 }
 
