* nicer comment

diff --git a/Engine/Source/Runtime/Renderer/Private/DeferredShadingRenderer.cpp b/Engine/Source/Runtime/Renderer/Private/DeferredShadingRenderer.cpp
--- a/Engine/Source/Runtime/Renderer/Private/DeferredShadingRenderer.cpp
+++ b/Engine/Source/Runtime/Renderer/Private/DeferredShadingRenderer.cpp
@@ -427,8 +427,8 @@ void FDeferredShadingSceneRenderer::RenderFinish()
 		Scene->MotionBlurInfoData.UpdateMotionBlurCache();
 	}
 
-	// Some RT should be released as early as possible to allow shading of that memory for other purposes.
-	// SceneColor should be released in tone mapping.
+	// Some RT should be released as early as possible to allow sharing of that memory for other purposes.
+	// SceneColor is be released in tone mapping, if not we want to get access to the HDR scene color after this pass so we keep it.
 	// This becomes even more important with some limited VRam (XBoxOne).
 	GSceneRenderTargets.SetLightAttenuation(0);
 }
