* cleaner, more efficient RT UAV sharing

diff --git a/Engine/Source/Runtime/Renderer/Private/PostProcess/RenderTargetPool.cpp b/Engine/Source/Runtime/Renderer/Private/PostProcess/RenderTargetPool.cpp
--- a/Engine/Source/Runtime/Renderer/Private/PostProcess/RenderTargetPool.cpp
+++ b/Engine/Source/Runtime/Renderer/Private/PostProcess/RenderTargetPool.cpp
@@ -170,6 +170,9 @@ bool FRenderTargetPool::FindFreeElement(const FPooledRenderTargetDesc& Desc, TRe
 		Found = new FPooledRenderTarget(Desc);
 
 		PooledRenderTargets.Add(Found);
+		
+		// TexCreate_UAV should be used on Desc.TargetableFlags
+		check(!(Desc.Flags & TexCreate_UAV));
 
 		if(Desc.TargetableFlags & (TexCreate_RenderTargetable | TexCreate_DepthStencilTargetable | TexCreate_UAV))
 		{
diff --git a/Engine/Source/Runtime/Renderer/Private/PostProcess/SceneRenderTargets.cpp b/Engine/Source/Runtime/Renderer/Private/PostProcess/SceneRenderTargets.cpp
--- a/Engine/Source/Runtime/Renderer/Private/PostProcess/SceneRenderTargets.cpp
+++ b/Engine/Source/Runtime/Renderer/Private/PostProcess/SceneRenderTargets.cpp
@@ -273,7 +273,7 @@ void FSceneRenderTargets::AllocSceneColor()
 		FPooledRenderTargetDesc Desc(FPooledRenderTargetDesc::Create2DDesc(BufferSize, SceneColorBufferFormat, TexCreate_None, TexCreate_RenderTargetable, false));
 		Desc.Flags |= TexCreate_FastVRAM;
 		// to allow better sharing with later elements
-		Desc.Flags |= TexCreate_UAV;
+		Desc.TargetableFlags |= TexCreate_UAV;
 		GRenderTargetPool.FindFreeElement(Desc, SceneColor, TEXT("SceneColor"));
 	}
 
diff --git a/Engine/Source/Runtime/Renderer/Private/ReflectionEnvironment.cpp b/Engine/Source/Runtime/Renderer/Private/ReflectionEnvironment.cpp
--- a/Engine/Source/Runtime/Renderer/Private/ReflectionEnvironment.cpp
+++ b/Engine/Source/Runtime/Renderer/Private/ReflectionEnvironment.cpp
@@ -565,16 +565,14 @@ void FDeferredShadingSceneRenderer::RenderDeferredReflections()
 	const bool bReflectionsWithCompute = GRHIFeatureLevel == ERHIFeatureLevel::SM5 && Scene->ReflectionSceneData.CubemapArray.IsValid();
 
 	TRefCountPtr<IPooledRenderTarget>	NewSceneColor;
-	FUnorderedAccessViewRHIRef			NewSceneColorUAV;
 	if( bReflectionEnv && bReflectionsWithCompute && !bAnyViewIsReflectionCapture )
 	{
 		GSceneRenderTargets.ResolveSceneColor(FResolveRect(0, 0, ViewFamily.FamilySizeX, ViewFamily.FamilySizeY));
 
 		FPooledRenderTargetDesc Desc = GSceneRenderTargets.GetSceneColor()->GetDesc();
-		Desc.Flags |= TexCreate_UAV;
+		Desc.TargetableFlags |= TexCreate_UAV;
 
-		GRenderTargetPool.FindFreeElement( Desc, NewSceneColor, TEXT("SceneColor") );
-		NewSceneColorUAV = RHICreateUnorderedAccessView( NewSceneColor->GetRenderTargetItem().TargetableTexture );
+		GRenderTargetPool.FindFreeElement( Desc, NewSceneColor, TEXT("SceneColorEnv") );
 	}
 
 	for (int32 ViewIndex = 0; ViewIndex < Views.Num(); ViewIndex++)
@@ -644,7 +642,7 @@ void FDeferredShadingSceneRenderer::RenderDeferredReflections()
 
 					RHISetComputeShader(ComputeShader->GetComputeShader());
 
-					ComputeShader->SetParameters(View, SSROutput->GetRenderTargetItem().ShaderResourceTexture, NewSceneColorUAV);
+					ComputeShader->SetParameters(View, SSROutput->GetRenderTargetItem().ShaderResourceTexture, NewSceneColor->GetRenderTargetItem().UAV);
 					//ComputeShader->SetParameters(View, DestRenderTarget.UAV);
 
 					uint32 GroupSizeX = (View.ViewRect.Size().X + GReflectionEnvironmentTileSizeX - 1) / GReflectionEnvironmentTileSizeX;
@@ -811,6 +809,4 @@ void FDeferredShadingSceneRenderer::RenderDeferredReflections()
 			}
 		}
 	}
-
-	NewSceneColorUAV.SafeRelease();
 }
diff --git a/Engine/Source/Runtime/Renderer/Public/RendererInterface.h b/Engine/Source/Runtime/Renderer/Public/RendererInterface.h
--- a/Engine/Source/Runtime/Renderer/Public/RendererInterface.h
+++ b/Engine/Source/Runtime/Renderer/Public/RendererInterface.h
@@ -268,7 +268,7 @@ public:
 		NumSamples = 1;
 
 		// Remove UAV flag for rendertargets that don't need it (some formats are incompatible)
-		Flags &= (~TexCreate_UAV);
+		TargetableFlags &= (~TexCreate_UAV);
 	}
 
 	/** In pixels, (0,0) if not set, (x,0) for cube maps, todo: make 3d int vector for volume textures */
