* added comments

diff --git a/Engine/Shaders/BasePassPixelShader.usf b/Engine/Shaders/BasePassPixelShader.usf
--- a/Engine/Shaders/BasePassPixelShader.usf
+++ b/Engine/Shaders/BasePassPixelShader.usf
@@ -47,7 +47,8 @@ float4 PointSkyBentNormal;
 
 #define EDITOR_ALPHA2COVERAGE (EDITOR_PRIMITIVE_MATERIAL && FEATURE_LEVEL >= FEATURE_LEVEL_SM5 && !COMPILER_GLSL)
 
-// for now we only use it on EDITOR_PRIMITIVE_MATERIAL (LevelGrid bias can be much lower) but later we might want to use it for all materials
+// For now we only use it on EDITOR_PRIMITIVE_MATERIAL (LevelGrid bias can be much lower) but later we might want to use it for all materials.
+// According to GCNPerformanceTweets.pdf Tip 38 it might have a performance cost on GCN.
 // On consoles we don't need the editor rendering features
 #if EDITOR_PRIMITIVE_MATERIAL && FEATURE_LEVEL >= FEATURE_LEVEL_SM5 && !PS4_PROFILE && !XBOXONE_PROFILE
 	// not available on mobile, if not available we have to reconstruct it from the world position interpolator which is much worse quality (noticable for depth bias)
diff --git a/Engine/Source/Runtime/Renderer/Private/PostProcess/SceneRenderTargets.cpp b/Engine/Source/Runtime/Renderer/Private/PostProcess/SceneRenderTargets.cpp
--- a/Engine/Source/Runtime/Renderer/Private/PostProcess/SceneRenderTargets.cpp
+++ b/Engine/Source/Runtime/Renderer/Private/PostProcess/SceneRenderTargets.cpp
@@ -271,9 +271,13 @@ void FSceneRenderTargets::AllocSceneColor()
 	// Create the scene color.
 	{
 		FPooledRenderTargetDesc Desc(FPooledRenderTargetDesc::Create2DDesc(BufferSize, SceneColorBufferFormat, TexCreate_None, TexCreate_RenderTargetable, false));
+
 		Desc.Flags |= TexCreate_FastVRAM;
-		// to allow better sharing with later elements
+
+		// with TexCreate_UAV it would allow better sharing with later elements but it might come at a high cost:
+		// GCNPerformanceTweets.pdf Tip 37: Warning: Causes additional synchronization between draw calls when using a render target allocated with this flag, use sparingly
 		Desc.TargetableFlags |= TexCreate_UAV;
+
 		GRenderTargetPool.FindFreeElement(Desc, SceneColor, TEXT("SceneColor"));
 	}
 
