* fixed Mac compiling
* nicer RT name

diff --git a/Engine/Source/Runtime/Renderer/Private/PostProcess/SceneRenderTargets.cpp b/Engine/Source/Runtime/Renderer/Private/PostProcess/SceneRenderTargets.cpp
--- a/Engine/Source/Runtime/Renderer/Private/PostProcess/SceneRenderTargets.cpp
+++ b/Engine/Source/Runtime/Renderer/Private/PostProcess/SceneRenderTargets.cpp
@@ -276,7 +276,10 @@ void FSceneRenderTargets::AllocSceneColor()
 
 		// with TexCreate_UAV it would allow better sharing with later elements but it might come at a high cost:
 		// GCNPerformanceTweets.pdf Tip 37: Warning: Causes additional synchronization between draw calls when using a render target allocated with this flag, use sparingly
-		Desc.TargetableFlags |= TexCreate_UAV;
+		if(GRHIFeatureLevel >= ERHIFeatureLevel::SM5)
+		{
+			Desc.TargetableFlags |= TexCreate_UAV;
+		}
 
 		GRenderTargetPool.FindFreeElement(Desc, SceneColor, TEXT("SceneColor"));
 	}
@@ -825,6 +828,23 @@ static TCHAR* const GetVolumeName(uint32 Id, bool bDirectional)
 	return (TCHAR*)TEXT("InvalidName");
 }
 
+
+// for easier use of "VisualizeTexture"
+static TCHAR* const GetTranslucencyShadowTransmissionName(uint32 Id)
+{
+	// (TCHAR*) for non VisualStudio
+	switch(Id)
+	{
+		case 0: return (TCHAR*)TEXT("TranslucencyShadowTransmission0");
+		case 1: return (TCHAR*)TEXT("TranslucencyShadowTransmission1");
+
+		default:
+			check(0);
+	}
+	return (TCHAR*)TEXT("InvalidName");
+}
+
+
 void FSceneRenderTargets::AllocateDeferredShadingPathRenderTargets()
 {
 	{
@@ -862,7 +882,7 @@ void FSceneRenderTargets::AllocateDeferredShadingPathRenderTargets()
 			{
 				// Using PF_FloatRGBA because Fourier coefficients used by Fourier opacity maps have a large range and can be negative
 				FPooledRenderTargetDesc Desc(FPooledRenderTargetDesc::Create2DDesc(TranslucentShadowBufferResolution, PF_FloatRGBA, TexCreate_None, TexCreate_RenderTargetable, false));
-				GRenderTargetPool.FindFreeElement(Desc, TranslucencyShadowTransmission[SurfaceIndex], TEXT("TranslucencyShadowTransmission"));
+				GRenderTargetPool.FindFreeElement(Desc, TranslucencyShadowTransmission[SurfaceIndex], GetTranslucencyShadowTransmissionName(SurfaceIndex));
 			}
 		}
 	}
