* nicer comments

diff --git a/Engine/Source/Runtime/Renderer/Private/PostProcess/SceneRenderTargets.h b/Engine/Source/Runtime/Renderer/Private/PostProcess/SceneRenderTargets.h
--- a/Engine/Source/Runtime/Renderer/Private/PostProcess/SceneRenderTargets.h
+++ b/Engine/Source/Runtime/Renderer/Private/PostProcess/SceneRenderTargets.h
@@ -370,7 +370,7 @@ public:
 
 	// ---
 
-private: // use SceneColor() or LightAttenuation() instead of direct access
+private: // Get...() methods instead of direct access
 
 	// 0 before BeginRenderingSceneColor and after tone mapping
 	TRefCountPtr<IPooledRenderTarget> SceneColor;
@@ -384,7 +384,7 @@ public:
 	TRefCountPtr<IPooledRenderTarget> SceneAlphaCopy;
 	// Auxiliary scene depth target. The scene depth is resolved to this surface when targeting SM4. 
 	TRefCountPtr<IPooledRenderTarget> AuxiliarySceneDepthZ;
-	// Render target for a quarter-sized version of the scene depths.
+	// Quarter-sized version of the scene depths
 	TRefCountPtr<IPooledRenderTarget> SmallDepthZ;
 
 	// GBuffer: Geometry Buffer rendered in base pass for deferred shading
