* reduce complexity

diff --git a/Engine/Source/Runtime/Renderer/Private/PostProcess/SceneRenderTargets.cpp b/Engine/Source/Runtime/Renderer/Private/PostProcess/SceneRenderTargets.cpp
--- a/Engine/Source/Runtime/Renderer/Private/PostProcess/SceneRenderTargets.cpp
+++ b/Engine/Source/Runtime/Renderer/Private/PostProcess/SceneRenderTargets.cpp
@@ -1062,13 +1062,17 @@ void FSceneRenderTargets::AllocateDeferredShadingPathRenderTargets()
 
 		uint32 LightAccumulationUAVFlag = GRHIFeatureLevel == ERHIFeatureLevel::SM5 ? TexCreate_UAV : 0;
 		FPooledRenderTargetDesc Desc(FPooledRenderTargetDesc::Create2DDesc(BufferSize, SceneColorBufferFormat, TexCreate_None, TexCreate_RenderTargetable | LightAccumulationUAVFlag, false));
+
+		if (GRHIFeatureLevel >= ERHIFeatureLevel::SM5)
+		{
+			Desc.TargetableFlags |= TexCreate_UAV;
+		}
+
 		GRenderTargetPool.FindFreeElement(Desc, LightAccumulation, TEXT("LightAccumulation"));
 	}
 
-	if (GRHIFeatureLevel == ERHIFeatureLevel::SM5)
+	if (GRHIFeatureLevel >= ERHIFeatureLevel::SM5)
 	{
-		LightAccumulationUAV = RHICreateUnorderedAccessView(LightAccumulation->GetRenderTargetItem().TargetableTexture);
-
 		// Create the reflective shadow map textures for LightPropagationVolume feature
 		if(bCurrentLightPropagationVolume)
 		{
@@ -1164,7 +1168,6 @@ void FSceneRenderTargets::ReleaseDynamicRHI()
 
 	ShadowDepthZ.SafeRelease();
 	PreShadowCacheDepthZ.SafeRelease();
-	LightAccumulationUAV.SafeRelease();
 	
 	for(int32 Index = 0; Index < NumCubeShadowDepthSurfaces; ++Index)
 	{
diff --git a/Engine/Source/Runtime/Renderer/Private/PostProcess/SceneRenderTargets.h b/Engine/Source/Runtime/Renderer/Private/PostProcess/SceneRenderTargets.h
--- a/Engine/Source/Runtime/Renderer/Private/PostProcess/SceneRenderTargets.h
+++ b/Engine/Source/Runtime/Renderer/Private/PostProcess/SceneRenderTargets.h
@@ -432,9 +432,6 @@ public:
 	/** Temporary storage, used during reflection capture filtering. */
 	TRefCountPtr<IPooledRenderTarget> ReflectionBrightness;
 
-	/** Used to accumulate HDR values in screenspace. */
-	FUnorderedAccessViewRHIRef LightAccumulationUAV;
-
 	/** Volume textures used for lighting translucency. */
 	TRefCountPtr<IPooledRenderTarget> TranslucencyLightingVolumeAmbient[NumTranslucentVolumeRenderTargetSets];
 	TRefCountPtr<IPooledRenderTarget> TranslucencyLightingVolumeDirectional[NumTranslucentVolumeRenderTargetSets];
diff --git a/Engine/Source/Runtime/Renderer/Private/TiledDeferredLightRendering.cpp b/Engine/Source/Runtime/Renderer/Private/TiledDeferredLightRendering.cpp
--- a/Engine/Source/Runtime/Renderer/Private/TiledDeferredLightRendering.cpp
+++ b/Engine/Source/Runtime/Renderer/Private/TiledDeferredLightRendering.cpp
@@ -117,7 +117,7 @@ public:
 
 		FGlobalShader::SetParameters(ShaderRHI, View);
 		DeferredParameters.Set(ShaderRHI, View);
-		LightAccumulation.SetTexture(ShaderRHI, GSceneRenderTargets.GetLightAccumulationTexture(), GSceneRenderTargets.LightAccumulationUAV);
+		LightAccumulation.SetTexture(ShaderRHI, GSceneRenderTargets.GetLightAccumulationTexture(), GSceneRenderTargets.LightAccumulation->GetRenderTargetItem().UAV);
 
 		SetShaderValue(ShaderRHI, ViewDimensions, View.ViewRect);
 
