* minor cleanup

diff --git a/Engine/Source/Runtime/Renderer/Private/PostProcess/SceneRenderTargets.h b/Engine/Source/Runtime/Renderer/Private/PostProcess/SceneRenderTargets.h
--- a/Engine/Source/Runtime/Renderer/Private/PostProcess/SceneRenderTargets.h
+++ b/Engine/Source/Runtime/Renderer/Private/PostProcess/SceneRenderTargets.h
@@ -110,12 +110,6 @@ public:
 	 * If the allocated render targets are too small, they are reallocated.
 	 */
 	void Allocate(const FSceneViewFamily& ViewFamily);
-	/**
-	 * Takes the requested buffer size and quantizes it to an appropriate size for the rest of the
-	 * rendering pipeline. Currently ensures that sizes are multiples of 8 so that they can safely
-	 * be halved in size several times.
-	 */
-	void QuantizeBufferSize(int32& InOutBufferSizeX, int32& InOutBufferSizeY) const;
 	/**
 	 *
 	 */
@@ -283,10 +277,6 @@ public:
 	{
 		return *(FTextureRHIRef*)&GetLightAttenuation()->GetRenderTargetItem().ShaderResourceTexture;
 	}
-	const FTextureRHIRef& GetLightAccumulationTexture() const
-	{
-		return *(FTextureRHIRef*)&LightAccumulation->GetRenderTargetItem().ShaderResourceTexture;
-	}
 
 	const FTextureRHIRef& GetSceneColorSurface() const;
 	const FTexture2DRHIRef& GetSceneAlphaCopySurface() const						{ return (const FTexture2DRHIRef&)SceneAlphaCopy->GetRenderTargetItem().TargetableTexture; }
@@ -457,6 +447,13 @@ public:
 	bool bPreshadowCacheNewlyAllocated;
 
 private:
+	/**
+	 * Takes the requested buffer size and quantizes it to an appropriate size for the rest of the
+	 * rendering pipeline. Currently ensures that sizes are multiples of 8 so that they can safely
+	 * be halved in size several times.
+	 */
+	void QuantizeBufferSize(int32& InOutBufferSizeX, int32& InOutBufferSizeY) const;
+
 	/**
 	 * Initializes the editor primitive color render target
 	 */
diff --git a/Engine/Source/Runtime/Renderer/Private/ReflectionEnvironment.cpp b/Engine/Source/Runtime/Renderer/Private/ReflectionEnvironment.cpp
--- a/Engine/Source/Runtime/Renderer/Private/ReflectionEnvironment.cpp
+++ b/Engine/Source/Runtime/Renderer/Private/ReflectionEnvironment.cpp
@@ -781,7 +781,7 @@ void FDeferredShadingSceneRenderer::RenderDeferredReflections()
 						TShaderMapRef< FReflectionApplyPS<A,B,C> > PixelShader(GetGlobalShaderMap()); \
 						static FGlobalBoundShaderState BoundShaderState; \
 						SetGlobalBoundShaderState(BoundShaderState, GFilterVertexDeclaration.VertexDeclarationRHI, *VertexShader, *PixelShader); \
-						PixelShader->SetParameters( View, GSceneRenderTargets.GetLightAccumulationTexture(), SSROutput->GetRenderTargetItem().ShaderResourceTexture ); \
+						PixelShader->SetParameters( View, GSceneRenderTargets.LightAccumulation->GetRenderTargetItem().ShaderResourceTexture, SSROutput->GetRenderTargetItem().ShaderResourceTexture ); \
 					}; \
 					break
 
diff --git a/Engine/Source/Runtime/Renderer/Private/TiledDeferredLightRendering.cpp b/Engine/Source/Runtime/Renderer/Private/TiledDeferredLightRendering.cpp
--- a/Engine/Source/Runtime/Renderer/Private/TiledDeferredLightRendering.cpp
+++ b/Engine/Source/Runtime/Renderer/Private/TiledDeferredLightRendering.cpp
@@ -111,13 +111,14 @@ public:
 		int32 NumLightsToRenderInSortedLights, 
 		const FSimpleLightArray& SimpleLights, 
 		int32 StartIndex, 
-		int32 NumThisPass)
+		int32 NumThisPass,
+		IPooledRenderTarget& InOutTexture)
 	{
 		FComputeShaderRHIParamRef ShaderRHI = GetComputeShader();
 
 		FGlobalShader::SetParameters(ShaderRHI, View);
 		DeferredParameters.Set(ShaderRHI, View);
-		LightAccumulation.SetTexture(ShaderRHI, GSceneRenderTargets.GetLightAccumulationTexture(), GSceneRenderTargets.LightAccumulation->GetRenderTargetItem().UAV);
+		LightAccumulation.SetTexture(ShaderRHI, InOutTexture.GetRenderTargetItem().ShaderResourceTexture, InOutTexture.GetRenderTargetItem().UAV);
 
 		SetShaderValue(ShaderRHI, ViewDimensions, View.ViewRect);
 
@@ -270,12 +271,13 @@ static void SetShaderTemplTiledLighting(
 	int32 NumLightsToRenderInSortedLights,
 	const FSimpleLightArray& SimpleLights,
 	int32 StartIndex, 
-	int32 NumThisPass)
+	int32 NumThisPass,
+	IPooledRenderTarget& InOutTexture)
 {
 	TShaderMapRef<FTiledDeferredLightingCS<bVisualizeLightCulling> > ComputeShader(GetGlobalShaderMap());
 	RHISetComputeShader(ComputeShader->GetComputeShader());
 
-	ComputeShader->SetParameters(View, ViewIndex, NumViews, SortedLights, NumLightsToRenderInSortedLights, SimpleLights, StartIndex, NumThisPass);
+	ComputeShader->SetParameters(View, ViewIndex, NumViews, SortedLights, NumLightsToRenderInSortedLights, SimpleLights, StartIndex, NumThisPass, InOutTexture);
 
 	uint32 GroupSizeX = (View.ViewRect.Size().X + GDeferredLightTileSizeX - 1) / GDeferredLightTileSizeX;
 	uint32 GroupSizeY = (View.ViewRect.Size().Y + GDeferredLightTileSizeY - 1) / GDeferredLightTileSizeY;
@@ -293,6 +295,8 @@ void FDeferredShadingSceneRenderer::RenderTiledDeferredLighting(const TArray<FSo
 	const int32 NumLightsToRender = NumUnshadowedLights + SimpleLights.InstanceData.Num();
 	const int32 NumLightsToRenderInSortedLights = NumUnshadowedLights;
 
+	IPooledRenderTarget& InOutTexture = *GSceneRenderTargets.LightAccumulation;
+
 	if (NumLightsToRender > 0)
 	{
 		INC_DWORD_STAT_BY(STAT_NumLightsUsingTiledDeferred, NumLightsToRender);
@@ -318,11 +322,11 @@ void FDeferredShadingSceneRenderer::RenderTiledDeferredLighting(const TArray<FSo
 
 					if(View.Family->EngineShowFlags.VisualizeLightCulling)
 					{
-						SetShaderTemplTiledLighting<1>(View, ViewIndex, Views.Num(), SortedLights, NumLightsToRenderInSortedLights, SimpleLights, StartIndex, NumThisPass);
+						SetShaderTemplTiledLighting<1>(View, ViewIndex, Views.Num(), SortedLights, NumLightsToRenderInSortedLights, SimpleLights, StartIndex, NumThisPass, InOutTexture);
 					}
 					else
 					{
-						SetShaderTemplTiledLighting<0>(View, ViewIndex, Views.Num(), SortedLights, NumLightsToRenderInSortedLights, SimpleLights, StartIndex, NumThisPass);
+						SetShaderTemplTiledLighting<0>(View, ViewIndex, Views.Num(), SortedLights, NumLightsToRenderInSortedLights, SimpleLights, StartIndex, NumThisPass, InOutTexture);
 					}
 				}
 			}
@@ -348,7 +352,7 @@ void FDeferredShadingSceneRenderer::RenderTiledDeferredLighting(const TArray<FSo
 					
 					static FGlobalBoundShaderState BoundShaderState;
 					SetGlobalBoundShaderState(BoundShaderState, GFilterVertexDeclaration.VertexDeclarationRHI, *ScreenVertexShader, *ScreenPixelShader);
-					ScreenPixelShader->SetParameters(TStaticSamplerState<SF_Point>::GetRHI(), GSceneRenderTargets.GetLightAccumulationTexture());
+					ScreenPixelShader->SetParameters(TStaticSamplerState<SF_Point>::GetRHI(), InOutTexture.GetRenderTargetItem().ShaderResourceTexture);
 
 					DrawRectangle( 
 						0, 0,
