* optimized RenderTarget usage, optimized out a pass for TileDeferredLighting, one less consumer for LightAccumulation

diff --git a/Engine/Shaders/TiledDeferredLightShaders.usf b/Engine/Shaders/TiledDeferredLightShaders.usf
--- a/Engine/Shaders/TiledDeferredLightShaders.usf
+++ b/Engine/Shaders/TiledDeferredLightShaders.usf
@@ -31,7 +31,8 @@ groupshared uint TileNumSimpleLights;
 groupshared uint TileSimpleLightIndices[MAX_LIGHTS];
 
 /** Output HDR target. */
-RWTexture2D<float4> RWLightAccumulation;
+Texture2D<float4> InTexture;
+RWTexture2D<float4> RWOutTexture;
 
 [numthreads(THREADGROUP_SIZEX, THREADGROUP_SIZEY, 1)]
 void TiledDeferredLightingMain(
@@ -229,6 +230,7 @@ void TiledDeferredLightingMain(
 	BRANCH
     if (all(DispatchThreadId.xy < ViewDimensions.zw)) 
 	{
-		RWLightAccumulation[PixelPos.xy] = CompositedLighting;
+		// One some hardware we can read and write from the same UAV with a 32 bit format. We don't do that yet.
+		RWOutTexture[PixelPos.xy] = InTexture[PixelPos.xy] + CompositedLighting;
     }
 }
diff --git a/Engine/Source/Runtime/Renderer/Private/ReflectionEnvironment.cpp b/Engine/Source/Runtime/Renderer/Private/ReflectionEnvironment.cpp
--- a/Engine/Source/Runtime/Renderer/Private/ReflectionEnvironment.cpp
+++ b/Engine/Source/Runtime/Renderer/Private/ReflectionEnvironment.cpp
@@ -651,9 +651,6 @@ void FDeferredShadingSceneRenderer::RenderDeferredReflections()
 
 					ComputeShader->UnsetParameters();
 
-					const FSceneRenderTargetItem& DestRenderTarget = NewSceneColor->GetRenderTargetItem();
-					RHICopyToResolveTarget( DestRenderTarget.TargetableTexture, DestRenderTarget.ShaderResourceTexture, false, FResolveParams() );
-
 					GSceneRenderTargets.SetSceneColor(NewSceneColor);
 					check(GSceneRenderTargets.GetSceneColor());
 
diff --git a/Engine/Source/Runtime/Renderer/Private/TiledDeferredLightRendering.cpp b/Engine/Source/Runtime/Renderer/Private/TiledDeferredLightRendering.cpp
--- a/Engine/Source/Runtime/Renderer/Private/TiledDeferredLightRendering.cpp
+++ b/Engine/Source/Runtime/Renderer/Private/TiledDeferredLightRendering.cpp
@@ -92,7 +92,8 @@ public:
 		: FGlobalShader(Initializer)
 	{
 		DeferredParameters.Bind(Initializer.ParameterMap);
-		LightAccumulation.Bind(Initializer.ParameterMap, TEXT("LightAccumulation"));
+		InTexture.Bind(Initializer.ParameterMap, TEXT("InTexture"));
+		OutTexture.Bind(Initializer.ParameterMap, TEXT("OutTexture"));
 		NumLights.Bind(Initializer.ParameterMap, TEXT("NumLights"));
 		ViewDimensions.Bind(Initializer.ParameterMap, TEXT("ViewDimensions"));
 		PreIntegratedBRDF.Bind(Initializer.ParameterMap, TEXT("PreIntegratedBRDF"));
@@ -112,14 +113,16 @@ public:
 		const FSimpleLightArray& SimpleLights, 
 		int32 StartIndex, 
 		int32 NumThisPass,
-		IPooledRenderTarget& InOutTexture)
+		IPooledRenderTarget& InTextureValue,
+		IPooledRenderTarget& OutTextureValue)
 	{
 		FComputeShaderRHIParamRef ShaderRHI = GetComputeShader();
 
 		FGlobalShader::SetParameters(ShaderRHI, View);
 		DeferredParameters.Set(ShaderRHI, View);
-		LightAccumulation.SetTexture(ShaderRHI, InOutTexture.GetRenderTargetItem().ShaderResourceTexture, InOutTexture.GetRenderTargetItem().UAV);
-
+		SetTextureParameter(ShaderRHI, InTexture, InTextureValue.GetRenderTargetItem().ShaderResourceTexture);
+		OutTexture.SetTexture(ShaderRHI, 0, OutTextureValue.GetRenderTargetItem().UAV);
+	
 		SetShaderValue(ShaderRHI, ViewDimensions, View.ViewRect);
 
 		SetTextureParameter(
@@ -207,14 +210,15 @@ public:
 
 	void UnsetParameters()
 	{
-		LightAccumulation.UnsetUAV(GetComputeShader());
+		OutTexture.UnsetUAV(GetComputeShader());
 	}
 
 	virtual bool Serialize(FArchive& Ar)
 	{		
 		bool bShaderHasOutdatedParameters = FGlobalShader::Serialize(Ar);
 		Ar << DeferredParameters;
-		Ar << LightAccumulation;
+		Ar << OutTexture;
+		Ar << InTexture;
 		Ar << NumLights;
 		Ar << ViewDimensions;
 		Ar << PreIntegratedBRDF;
@@ -235,7 +239,8 @@ public:
 private:
 
 	FDeferredPixelShaderParameters DeferredParameters;
-	FRWShaderParameter LightAccumulation;
+	FShaderResourceParameter InTexture;
+	FRWShaderParameter OutTexture;
 	FShaderParameter NumLights;
 	FShaderParameter ViewDimensions;
 	FShaderResourceParameter PreIntegratedBRDF;
@@ -272,12 +277,13 @@ static void SetShaderTemplTiledLighting(
 	const FSimpleLightArray& SimpleLights,
 	int32 StartIndex, 
 	int32 NumThisPass,
-	IPooledRenderTarget& InOutTexture)
+	IPooledRenderTarget& InTexture,
+	IPooledRenderTarget& OutTexture)
 {
 	TShaderMapRef<FTiledDeferredLightingCS<bVisualizeLightCulling> > ComputeShader(GetGlobalShaderMap());
 	RHISetComputeShader(ComputeShader->GetComputeShader());
 
-	ComputeShader->SetParameters(View, ViewIndex, NumViews, SortedLights, NumLightsToRenderInSortedLights, SimpleLights, StartIndex, NumThisPass, InOutTexture);
+	ComputeShader->SetParameters(View, ViewIndex, NumViews, SortedLights, NumLightsToRenderInSortedLights, SimpleLights, StartIndex, NumThisPass, InTexture, OutTexture);
 
 	uint32 GroupSizeX = (View.ViewRect.Size().X + GDeferredLightTileSizeX - 1) / GDeferredLightTileSizeX;
 	uint32 GroupSizeY = (View.ViewRect.Size().Y + GDeferredLightTileSizeY - 1) / GDeferredLightTileSizeY;
@@ -295,14 +301,14 @@ void FDeferredShadingSceneRenderer::RenderTiledDeferredLighting(const TArray<FSo
 	const int32 NumLightsToRender = NumUnshadowedLights + SimpleLights.InstanceData.Num();
 	const int32 NumLightsToRenderInSortedLights = NumUnshadowedLights;
 
-	IPooledRenderTarget& InOutTexture = *GSceneRenderTargets.LightAccumulation;
-
 	if (NumLightsToRender > 0)
 	{
 		INC_DWORD_STAT_BY(STAT_NumLightsUsingTiledDeferred, NumLightsToRender);
 		INC_DWORD_STAT_BY(STAT_NumLightsUsingSimpleTiledDeferred, SimpleLights.InstanceData.Num());
 		SCOPE_CYCLE_COUNTER(STAT_DirectLightRenderingTime);
 
+		RHISetRenderTarget(NULL, NULL);
+
 		// Determine how many compute shader passes will be needed to process all the lights
 		const int32 NumPassesNeeded = FMath::DivideAndRoundUp(NumLightsToRender, GMaxNumTiledDeferredLights);
 		for (int32 PassIndex = 0; PassIndex < NumPassesNeeded; PassIndex++)
@@ -311,10 +317,21 @@ void FDeferredShadingSceneRenderer::RenderTiledDeferredLighting(const TArray<FSo
 			const int32 NumThisPass = (PassIndex == NumPassesNeeded - 1) ? NumLightsToRender - StartIndex : GMaxNumTiledDeferredLights;
 			check(NumThisPass > 0);
 
+			// One some hardware we can read and write from the same UAV with a 32 bit format. We don't do that yet.
+			TRefCountPtr<IPooledRenderTarget> OutTexture;
+			{
+				GSceneRenderTargets.ResolveSceneColor(FResolveRect(0, 0, ViewFamily.FamilySizeX, ViewFamily.FamilySizeY));
+
+				FPooledRenderTargetDesc Desc = GSceneRenderTargets.GetSceneColor()->GetDesc();
+				Desc.TargetableFlags |= TexCreate_UAV;
+
+				GRenderTargetPool.FindFreeElement( Desc, OutTexture, TEXT("SceneColorTiled") );
+			}
+
 			{
 				SCOPED_DRAW_EVENT(TiledDeferredLighting, DEC_SCENE_ITEMS);
 
-				RHISetRenderTarget(NULL, NULL);
+				IPooledRenderTarget& InTexture = *GSceneRenderTargets.GetSceneColor();
 
 				for (int32 ViewIndex = 0; ViewIndex < Views.Num(); ViewIndex++)
 				{
@@ -322,49 +339,17 @@ void FDeferredShadingSceneRenderer::RenderTiledDeferredLighting(const TArray<FSo
 
 					if(View.Family->EngineShowFlags.VisualizeLightCulling)
 					{
-						SetShaderTemplTiledLighting<1>(View, ViewIndex, Views.Num(), SortedLights, NumLightsToRenderInSortedLights, SimpleLights, StartIndex, NumThisPass, InOutTexture);
+						SetShaderTemplTiledLighting<1>(View, ViewIndex, Views.Num(), SortedLights, NumLightsToRenderInSortedLights, SimpleLights, StartIndex, NumThisPass, InTexture, *OutTexture);
 					}
 					else
 					{
-						SetShaderTemplTiledLighting<0>(View, ViewIndex, Views.Num(), SortedLights, NumLightsToRenderInSortedLights, SimpleLights, StartIndex, NumThisPass, InOutTexture);
+						SetShaderTemplTiledLighting<0>(View, ViewIndex, Views.Num(), SortedLights, NumLightsToRenderInSortedLights, SimpleLights, StartIndex, NumThisPass, InTexture, *OutTexture);
 					}
 				}
 			}
 
-			{
-				// Composite the compute shader results onto scene color
-				//@todo - swap scene color render targets instead and write to scene color directly from the compute shader
-				SCOPED_DRAW_EVENT(CompositeTiledDeferredLighting, DEC_SCENE_ITEMS);
-
-				GSceneRenderTargets.BeginRenderingSceneColor();
-
-				for (int32 ViewIndex = 0; ViewIndex < Views.Num(); ViewIndex++)
-				{
-					const FViewInfo& View = Views[ViewIndex];
-
-					RHISetViewport(View.ViewRect.Min.X, View.ViewRect.Min.Y, 0.0f, View.ViewRect.Max.X, View.ViewRect.Max.Y, 1.0f);
-					RHISetRasterizerState(TStaticRasterizerState<FM_Solid,CM_None>::GetRHI());
-					RHISetDepthStencilState(TStaticDepthStencilState<false,CF_Always>::GetRHI());
-					RHISetBlendState(TStaticBlendState<CW_RGBA, BO_Add, BF_One, BF_One, BO_Add, BF_One, BF_One>::GetRHI());
-
-					TShaderMapRef<FScreenVS> ScreenVertexShader(GetGlobalShaderMap());
-					TShaderMapRef<FScreenPS> ScreenPixelShader(GetGlobalShaderMap());
-					
-					static FGlobalBoundShaderState BoundShaderState;
-					SetGlobalBoundShaderState(BoundShaderState, GFilterVertexDeclaration.VertexDeclarationRHI, *ScreenVertexShader, *ScreenPixelShader);
-					ScreenPixelShader->SetParameters(TStaticSamplerState<SF_Point>::GetRHI(), InOutTexture.GetRenderTargetItem().ShaderResourceTexture);
-
-					DrawRectangle( 
-						0, 0,
-						View.ViewRect.Width(), View.ViewRect.Height(),
-						View.ViewRect.Min.X, View.ViewRect.Min.Y, 
-						View.ViewRect.Width(), View.ViewRect.Height(),
-						View.ViewRect.Size(),
-						GSceneRenderTargets.GetBufferSizeXY(),
-						*ScreenVertexShader,
-						EDRF_UseTriangleOptimization);
-				}
-			}
+			// swap with the former SceneColor
+			GSceneRenderTargets.SetSceneColor(OutTexture);
 		}
 	}
 }
