* fixed issue with splitscreen where SceneColor was released too early

diff --git a/Engine/Source/Runtime/Renderer/Private/CompositionLighting/CompositionLighting.cpp b/Engine/Source/Runtime/Renderer/Private/CompositionLighting/CompositionLighting.cpp
--- a/Engine/Source/Runtime/Renderer/Private/CompositionLighting/CompositionLighting.cpp
+++ b/Engine/Source/Runtime/Renderer/Private/CompositionLighting/CompositionLighting.cpp
@@ -338,7 +338,7 @@ void FCompositionLighting::ProcessLighting(const FViewInfo& View)
 		CompositeContext.Process(TEXT("CompositionLighting_Lighting"));
 	}
 
-	// The RT should be released as early as possible to allow shading of that memory for other purposes.
+	// The RT should be released as early as possible to allow sharing of that memory for other purposes.
 	// This becomes even more important with some limited VRam (XBoxOne).
 	GSceneRenderTargets.SetLightAttenuation(0);
 }
diff --git a/Engine/Source/Runtime/Renderer/Private/PostProcess/PostProcessTonemap.cpp b/Engine/Source/Runtime/Renderer/Private/PostProcess/PostProcessTonemap.cpp
--- a/Engine/Source/Runtime/Renderer/Private/PostProcess/PostProcessTonemap.cpp
+++ b/Engine/Source/Runtime/Renderer/Private/PostProcess/PostProcessTonemap.cpp
@@ -1052,9 +1052,13 @@ void FRCPassPostProcessTonemap::Process(FRenderingCompositePassContext& Context)
 
 	RHICopyToResolveTarget(DestRenderTarget.TargetableTexture, DestRenderTarget.ShaderResourceTexture, false, FResolveParams());
 
-	// The RT should be released as early as possible to allow shading of that memory for other purposes.
-	// This becomes even more important with some limited VRam (XBoxOne).
-	GSceneRenderTargets.SetSceneColor(0);
+	// We only release the SceneColor after the last view was processed (SplitScreen)
+	if(Context.View.Family->Views[Context.View.Family->Views.Num() - 1] == &Context.View)
+	{
+		// The RT should be released as early as possible to allow sharing of that memory for other purposes.
+		// This becomes even more important with some limited VRam (XBoxOne).
+		GSceneRenderTargets.SetSceneColor(0);
+	}
 }
 
 FPooledRenderTargetDesc FRCPassPostProcessTonemap::ComputeOutputDesc(EPassOutputId InPassOutputId) const
@@ -1422,8 +1426,6 @@ void FRCPassPostProcessTonemapES2::Process(FRenderingCompositePassContext& Conte
 			ViewState->MobileAaColor0 = PassOutputs[0].PooledRenderTarget;
 		}
 	}
-
-	GSceneRenderTargets.SetSceneColor(0);
 }
 
 FPooledRenderTargetDesc FRCPassPostProcessTonemapES2::ComputeOutputDesc(EPassOutputId InPassOutputId) const
