* fixed issue with splitscreen where SceneColor was released too early

diff --git a/Engine/Source/Runtime/Renderer/Private/CompositionLighting/CompositionLighting.cpp b/Engine/Source/Runtime/Renderer/Private/CompositionLighting/CompositionLighting.cpp
--- a/Engine/Source/Runtime/Renderer/Private/CompositionLighting/CompositionLighting.cpp
+++ b/Engine/Source/Runtime/Renderer/Private/CompositionLighting/CompositionLighting.cpp
@@ -338,7 +338,11 @@ void FCompositionLighting::ProcessLighting(const FViewInfo& View)
 		CompositeContext.Process(TEXT("CompositionLighting_Lighting"));
 	}
 
-	// The RT should be released as early as possible to allow sharing of that memory for other purposes.
-	// This becomes even more important with some limited VRam (XBoxOne).
-	GSceneRenderTargets.SetLightAttenuation(0);
+	// We only release the after the last view was processed (SplitScreen)
+	if(View.Family->Views[View.Family->Views.Num() - 1] == &View)
+	{
+		// The RT should be released as early as possible to allow sharing of that memory for other purposes.
+		// This becomes even more important with some limited VRam (XBoxOne).
+		GSceneRenderTargets.SetLightAttenuation(0);
+	}
 }
