First pass of changes to remove GRHIFeatureLevel, allowing for multiple scenes with different feature levels to exist. First step towards dynamic mobile preview in the editor.

#ttp 331829 - Mobile Preview in Editor
#proj Engine
#branch UE4

#summary First pass of changes to remove GRHIFeatureLevel, allowing for multiple scenes with different feature levels to exist. First step towards dynamic mobile preview in the editor.

#change Added GMaxRHIFeature level which contains the highest supported feature level on this machine
#change Added a cvar "r.FeatureLevelPreview" which, when set to 1, enables a feature level selection in the quick settings menu. Not wired up to anything at the minute
#change Added a feature level variable to FScene, which is currently initialized to whatever GRHIFeatureLevel is. Accessed with GetFeatureLevel()
#change Added a helper accessor GetFeatureLevel() to FSceneView which gets the feature level from the relevant scene. Purely to reduce the amount of typing needed to get to a feature level
#change Many changes across the renderer to try and read the current feature level from the view or scene. Not everything is done yet, but this is almost all of the "easy" stuff. As there is still just a global feature level in the engine, nothing should change.

diff --git a/Engine/Source/Editor/Cascade/Cascade.Build.cs b/Engine/Source/Editor/Cascade/Cascade.Build.cs
--- a/Engine/Source/Editor/Cascade/Cascade.Build.cs
+++ b/Engine/Source/Editor/Cascade/Cascade.Build.cs
@@ -23,7 +23,8 @@ public class Cascade : ModuleRules
 				"SlateCore",
                 "EditorStyle",
 				"DistCurveEditor",
-				"UnrealEd"
+				"UnrealEd",
+				"RHI"
 			}
 		);
 
diff --git a/Engine/Source/Editor/Cascade/Private/Cascade.cpp b/Engine/Source/Editor/Cascade/Private/Cascade.cpp
--- a/Engine/Source/Editor/Cascade/Private/Cascade.cpp
+++ b/Engine/Source/Editor/Cascade/Private/Cascade.cpp
@@ -183,7 +183,7 @@ void FCascade::InitCascade(const EToolkitMode::Type Mode, const TSharedPtr< clas
 	}
 
 	// Construct an FX system for this preview.
-	FXSystem = FFXSystemInterface::Create();
+	FXSystem = FFXSystemInterface::Create(GRHIFeatureLevel);
 
 	ParticleSystemComponent = ConstructObject<UCascadeParticleSystemComponent>(UCascadeParticleSystemComponent::StaticClass());
 
diff --git a/Engine/Source/Editor/DetailCustomizations/DetailCustomizations.Build.cs b/Engine/Source/Editor/DetailCustomizations/DetailCustomizations.Build.cs
--- a/Engine/Source/Editor/DetailCustomizations/DetailCustomizations.Build.cs
+++ b/Engine/Source/Editor/DetailCustomizations/DetailCustomizations.Build.cs
@@ -34,7 +34,8 @@ public class DetailCustomizations : ModuleRules
 				"MoviePlayer",
 				"SourceControl",
                 "InternationalizationSettings",
-				"SourceCodeAccess"
+				"SourceCodeAccess",
+				"RHI"
 			}
 		);
 
diff --git a/Engine/Source/Editor/DetailCustomizations/Private/AnimMontageSegmentDetails.cpp b/Engine/Source/Editor/DetailCustomizations/Private/AnimMontageSegmentDetails.cpp
--- a/Engine/Source/Editor/DetailCustomizations/Private/AnimMontageSegmentDetails.cpp
+++ b/Engine/Source/Editor/DetailCustomizations/Private/AnimMontageSegmentDetails.cpp
@@ -148,7 +148,7 @@ void SAnimationSegmentViewport::Construct(const FArguments& InArgs)
 {
 	TargetSkeleton = InArgs._Skeleton;
 	AnimRef = InArgs._AnimRef;
-	FXSystem = FFXSystemInterface::Create();
+	FXSystem = FFXSystemInterface::Create(GRHIFeatureLevel);
 	
 	AnimRefPropertyHandle = InArgs._AnimRefPropertyHandle;
 	StartTimePropertyHandle = InArgs._StartTimePropertyHandle;
diff --git a/Engine/Source/Editor/DetailCustomizations/Private/AnimSequenceDetails.cpp b/Engine/Source/Editor/DetailCustomizations/Private/AnimSequenceDetails.cpp
--- a/Engine/Source/Editor/DetailCustomizations/Private/AnimSequenceDetails.cpp
+++ b/Engine/Source/Editor/DetailCustomizations/Private/AnimSequenceDetails.cpp
@@ -337,7 +337,7 @@ SAnimationRefPoseViewport::SAnimationRefPoseViewport()
 BEGIN_SLATE_FUNCTION_BUILD_OPTIMIZATION
 void SAnimationRefPoseViewport::Construct(const FArguments& InArgs)
 {
-	FXSystem = FFXSystemInterface::Create();
+	FXSystem = FFXSystemInterface::Create(GRHIFeatureLevel);
 
 	TargetSkeleton = InArgs._Skeleton;
 	AnimRefPropertyHandle = InArgs._AnimRefPropertyHandle;
diff --git a/Engine/Source/Editor/LevelEditor/LevelEditor.Build.cs b/Engine/Source/Editor/LevelEditor/LevelEditor.Build.cs
--- a/Engine/Source/Editor/LevelEditor/LevelEditor.Build.cs
+++ b/Engine/Source/Editor/LevelEditor/LevelEditor.Build.cs
@@ -43,7 +43,8 @@ public class LevelEditor : ModuleRules
 				"DeviceProfileServices",
 				"ContentBrowser",
                 "SceneOutliner",
-                "ActorPickerMode"
+                "ActorPickerMode",
+                "RHI"
 			}
 		);
 
diff --git a/Engine/Source/Editor/LevelEditor/Private/LevelEditor.cpp b/Engine/Source/Editor/LevelEditor/Private/LevelEditor.cpp
--- a/Engine/Source/Editor/LevelEditor/Private/LevelEditor.cpp
+++ b/Engine/Source/Editor/LevelEditor/Private/LevelEditor.cpp
@@ -1511,6 +1511,16 @@ void FLevelEditorModule::BindGlobalLevelEditorCommands()
 		FExecuteAction::CreateStatic( &FLevelEditorActionCallbacks::SetMaterialQualityLevel, (EMaterialQualityLevel::Type)EMaterialQualityLevel::High ),
 		FCanExecuteAction(),
 		FIsActionChecked::CreateStatic( &FLevelEditorActionCallbacks::IsMaterialQualityLevelChecked, (EMaterialQualityLevel::Type)EMaterialQualityLevel::High ) );
+
+
+	for (int32 i = 0; i < ERHIFeatureLevel::Num; ++i)
+	{
+		ActionList.MapAction(
+			Commands.FeatureLevelPreview[i],
+			FExecuteAction::CreateStatic(&FLevelEditorActionCallbacks::SetFeatureLevelPreview, (ERHIFeatureLevel::Type)i),
+			FCanExecuteAction(),
+			FIsActionChecked::CreateStatic(&FLevelEditorActionCallbacks::IsFeatureLevelPreviewChecked, (ERHIFeatureLevel::Type)i));
+	}
 }
 	
 #undef LOCTEXT_NAMESPACE
diff --git a/Engine/Source/Editor/LevelEditor/Private/LevelEditorActions.cpp b/Engine/Source/Editor/LevelEditor/Private/LevelEditorActions.cpp
--- a/Engine/Source/Editor/LevelEditor/Private/LevelEditorActions.cpp
+++ b/Engine/Source/Editor/LevelEditor/Private/LevelEditorActions.cpp
@@ -443,6 +443,17 @@ bool FLevelEditorActionCallbacks::IsMaterialQualityLevelChecked( EMaterialQualit
 	return TestQualityLevel == MaterialQualityLevel;
 }
 
+void FLevelEditorActionCallbacks::SetFeatureLevelPreview(ERHIFeatureLevel::Type InPreviewFeatureLevel)
+{
+	// Here be dragons...
+}
+
+bool FLevelEditorActionCallbacks::IsFeatureLevelPreviewChecked(ERHIFeatureLevel::Type InPreviewFeatureLevel)
+{
+	// For now, we just stay stuck on the system max feature level
+	return InPreviewFeatureLevel == GRHIFeatureLevel;
+}
+
 void FLevelEditorActionCallbacks::ConfigureLightingBuildOptions( const FLightingBuildOptions& Options )
 {
 	GConfig->SetBool( TEXT("LightingBuildOptions"), TEXT("OnlyBuildSelected"),		Options.bOnlyBuildSelected,			GEditorUserSettingsIni );
@@ -2754,6 +2765,22 @@ void FLevelEditorCommands::RegisterCommands()
 
 	UI_COMMAND( MaterialQualityLevel_Low, "Low", "Sets material quality in the scene to low.", EUserInterfaceActionType::RadioButton, FInputGesture() );
 	UI_COMMAND( MaterialQualityLevel_High, "High", "Sets material quality in the scene to high.", EUserInterfaceActionType::RadioButton, FInputGesture() );
+
+	for (int32 i = 0; i < ERHIFeatureLevel::Num; ++i)
+	{
+		FName Name;
+		GetFeatureLevelName((ERHIFeatureLevel::Type)i, Name);
+
+		FeatureLevelPreview[i] =
+			FUICommandInfoDecl(
+			this->AsShared(),
+			Name,
+			//FText::Format(NSLOCTEXT("LevelEditorCommands", "FeatureLevelPreview", "Preview scenes using feature level {0}"), TEXT("Thing")),
+			FText::FromName(Name),
+			FText::Format(NSLOCTEXT("LevelEditorCommands", "FeatureLevelPreviewTooltip", "Preview scenes using feature level {0}"), FText::FromName(Name)))
+			.UserInterfaceType(EUserInterfaceActionType::Button)
+			.DefaultGesture(FInputGesture());
+	}
 }
 
 PRAGMA_ENABLE_OPTIMIZATION
diff --git a/Engine/Source/Editor/LevelEditor/Private/LevelEditorToolBar.cpp b/Engine/Source/Editor/LevelEditor/Private/LevelEditorToolBar.cpp
--- a/Engine/Source/Editor/LevelEditor/Private/LevelEditorToolBar.cpp
+++ b/Engine/Source/Editor/LevelEditor/Private/LevelEditorToolBar.cpp
@@ -15,6 +15,12 @@
 #include "Editor/ContentBrowser/Public/ContentBrowserModule.h"
 #include "EngineBuildSettings.h"
 
+static TAutoConsoleVariable<int32> CVarFeatureLevelPreview(
+	TEXT("r.FeatureLevelPreview"),
+	0,
+	TEXT("If 1 the quick settings menu will contain an option to enable feature level preview modes"),
+	ECVF_RenderThreadSafe);
+
 /**
  * Static: Creates a widget for the level editor tool bar
  *
@@ -466,6 +472,18 @@ static void MakeMaterialQualityLevelMenu( FMenuBuilder& MenuBuilder )
 	MenuBuilder.EndSection();
 }
 
+static void MakeShaderModelPreviewMenu(FMenuBuilder& MenuBuilder)
+{
+	MenuBuilder.BeginSection("LevelEditorShaderModelPreview", NSLOCTEXT("LevelToolBarViewMenu", "ShaderModelPreviewHeading", "Shader Model Preview"));
+	{
+		for (int32 i = GMaxRHIFeatureLevel; i >= 0; --i)
+		{
+			MenuBuilder.AddMenuEntry(FLevelEditorCommands::Get().FeatureLevelPreview[i]);
+		}
+	}
+	MenuBuilder.EndSection();
+}
+
 static void MakeScalabilityMenu( FMenuBuilder& MenuBuilder )
 {
 	MenuBuilder.AddWidget(SNew(SScalabilitySettings), FText(), true);
@@ -525,6 +543,14 @@ TSharedRef< SWidget > FLevelEditorToolBar::GenerateQuickSettingsMenu( TSharedRef
 			LOCTEXT( "MaterialQualityLevelSubMenu", "Material Quality Level" ),
 			LOCTEXT( "MaterialQualityLevelSubMenu_ToolTip", "Sets the value of the CVar \"r.MaterialQualityLevel\" (low=0, high=1). This affects materials via the QualitySwitch material expression." ),
 			FNewMenuDelegate::CreateStatic( &MakeMaterialQualityLevelMenu ) );
+
+		if (CVarFeatureLevelPreview.GetValueOnGameThread() != 0)
+		{
+			MenuBuilder.AddSubMenu(
+				LOCTEXT("ShaderModelPreviewSubMenu", "Shader Model Preview"),
+				LOCTEXT("ShaderModelPreviewSubMenu_ToolTip", "Sets the shader model preview mode"),
+				FNewMenuDelegate::CreateStatic(&MakeShaderModelPreviewMenu));
+		}
 	}
 	MenuBuilder.EndSection();
 
diff --git a/Engine/Source/Editor/LevelEditor/Public/LevelEditorActions.h b/Engine/Source/Editor/LevelEditor/Public/LevelEditorActions.h
--- a/Engine/Source/Editor/LevelEditor/Public/LevelEditorActions.h
+++ b/Engine/Source/Editor/LevelEditor/Public/LevelEditorActions.h
@@ -536,6 +536,8 @@ public:
 	TSharedPtr< FUICommandInfo > MaterialQualityLevel_Low;
 	TSharedPtr< FUICommandInfo > MaterialQualityLevel_High;
 
+	TSharedPtr< FUICommandInfo > FeatureLevelPreview[ERHIFeatureLevel::Num];
+	
 	///**
 	// * Mode Commands                   
 	// */
@@ -715,6 +717,8 @@ public:
 	static void AttachToSocketSelection(FName SocketName, AActor* ParentActorPtr);
 	static void SetMaterialQualityLevel( EMaterialQualityLevel::Type NewQualityLevel );
 	static bool IsMaterialQualityLevelChecked( EMaterialQualityLevel::Type TestQualityLevel );
+	static void SetFeatureLevelPreview(ERHIFeatureLevel::Type InFeatureLevel);
+	static bool IsFeatureLevelPreviewChecked(ERHIFeatureLevel::Type InFeatureLevel);
 	
 	/**
 	 * Called when the Scene Stats button is clicked.  Invokes the Primitive Stats dialog.
diff --git a/Engine/Source/Editor/UnrealEd/Private/ThumbnailHelpers.cpp b/Engine/Source/Editor/UnrealEd/Private/ThumbnailHelpers.cpp
--- a/Engine/Source/Editor/UnrealEd/Private/ThumbnailHelpers.cpp
+++ b/Engine/Source/Editor/UnrealEd/Private/ThumbnailHelpers.cpp
@@ -134,7 +134,7 @@ FParticleSystemThumbnailScene::FParticleSystemThumbnailScene()
 	bForceAllUsedMipsResident = false;
 	PartComponent = NULL;
 
-	ThumbnailFXSystem = FFXSystemInterface::Create();
+	ThumbnailFXSystem = FFXSystemInterface::Create(GRHIFeatureLevel);
 	GetScene()->SetFXSystem( ThumbnailFXSystem );
 }
 
diff --git a/Engine/Source/Runtime/Engine/Private/NiagaraComponent.cpp b/Engine/Source/Runtime/Engine/Private/NiagaraComponent.cpp
--- a/Engine/Source/Runtime/Engine/Private/NiagaraComponent.cpp
+++ b/Engine/Source/Runtime/Engine/Private/NiagaraComponent.cpp
@@ -154,7 +154,7 @@ private:
 		FMaterialRenderProxy* MaterialRenderProxy = Material->GetRenderProxy(IsSelected(),IsHovered());
 
 		if (DynamicVertexAllocation.IsValid()
-			&& (bIsWireframe || !PDI->IsMaterialIgnored(MaterialRenderProxy)))
+			&& (bIsWireframe || !PDI->IsMaterialIgnored(MaterialRenderProxy, View->GetFeatureLevel())))
 		{
 			FMeshBatch MeshBatch;
 			MeshBatch.VertexFactory = &VertexFactory;
diff --git a/Engine/Source/Runtime/Engine/Private/Particles/FXSystem.cpp b/Engine/Source/Runtime/Engine/Private/Particles/FXSystem.cpp
--- a/Engine/Source/Runtime/Engine/Private/Particles/FXSystem.cpp
+++ b/Engine/Source/Runtime/Engine/Private/Particles/FXSystem.cpp
@@ -18,9 +18,9 @@
 	External FX system interface.
 -----------------------------------------------------------------------------*/
 
-FFXSystemInterface* FFXSystemInterface::Create()
+FFXSystemInterface* FFXSystemInterface::Create(ERHIFeatureLevel::Type InFeatureLevel)
 {
-	return new FFXSystem();
+	return new FFXSystem(InFeatureLevel);
 }
 
 void FFXSystemInterface::Destroy( FFXSystemInterface* FXSystem )
@@ -139,11 +139,12 @@ namespace FXConsoleVariables
 	FX system.
 ------------------------------------------------------------------------------*/
 
-FFXSystem::FFXSystem()
+FFXSystem::FFXSystem(ERHIFeatureLevel::Type InFeatureLevel)
 	: ParticleSimulationResources(NULL)
 #if WITH_EDITOR
 	, bSuspended(false)
 #endif // #if WITH_EDITOR
+	, FeatureLevel(InFeatureLevel)
 {
 	InitGPUSimulation();
 }
@@ -155,7 +156,7 @@ FFXSystem::~FFXSystem()
 
 void FFXSystem::Tick(float DeltaSeconds)
 {
-	if (CurrentRHISupportsGPUParticles())
+	if (RHISupportsGPUParticles(FeatureLevel))
 	{
 		// Test GPU sorting if requested.
 		if (FXConsoleVariables::TestGPUSort.GetValueOnGameThread() != 0)
@@ -175,7 +176,7 @@ void FFXSystem::Tick(float DeltaSeconds)
 #if WITH_EDITOR
 void FFXSystem::Suspend()
 {
-	if (!bSuspended && CurrentRHISupportsGPUParticles())
+	if (!bSuspended && RHISupportsGPUParticles(FeatureLevel))
 	{
 		ReleaseGPUResources();
 		bSuspended = true;
@@ -184,7 +185,7 @@ void FFXSystem::Suspend()
 
 void FFXSystem::Resume()
 {
-	if (bSuspended && CurrentRHISupportsGPUParticles())
+	if (bSuspended && RHISupportsGPUParticles(FeatureLevel))
 	{
 		bSuspended = false;
 		InitGPUResources();
@@ -198,7 +199,7 @@ void FFXSystem::Resume()
 
 void FFXSystem::AddVectorField( UVectorFieldComponent* VectorFieldComponent )
 {
-	if (CurrentRHISupportsGPUParticles())
+	if (RHISupportsGPUParticles(FeatureLevel))
 	{
 		check( VectorFieldComponent->VectorFieldInstance == NULL );
 		check( VectorFieldComponent->FXSystem == this );
@@ -228,7 +229,7 @@ void FFXSystem::AddVectorField( UVectorFieldComponent* VectorFieldComponent )
 
 void FFXSystem::RemoveVectorField( UVectorFieldComponent* VectorFieldComponent )
 {
-	if (CurrentRHISupportsGPUParticles())
+	if (RHISupportsGPUParticles(FeatureLevel))
 	{
 		check( VectorFieldComponent->FXSystem == this );
 
@@ -254,7 +255,7 @@ void FFXSystem::RemoveVectorField( UVectorFieldComponent* VectorFieldComponent )
 
 void FFXSystem::UpdateVectorField( UVectorFieldComponent* VectorFieldComponent )
 {
-	if (CurrentRHISupportsGPUParticles())
+	if (RHISupportsGPUParticles(FeatureLevel))
 	{
 		check( VectorFieldComponent->FXSystem == this );
 
@@ -298,7 +299,7 @@ void FFXSystem::UpdateVectorField( UVectorFieldComponent* VectorFieldComponent )
 void FFXSystem::DrawDebug( FCanvas* Canvas )
 {
 	if (FXConsoleVariables::VisualizeGPUSimulation > 0
-		&& CurrentRHISupportsGPUParticles())
+		&& RHISupportsGPUParticles(FeatureLevel))
 	{
 		VisualizeGPUParticles(Canvas);
 	}
@@ -306,7 +307,7 @@ void FFXSystem::DrawDebug( FCanvas* Canvas )
 
 void FFXSystem::PreInitViews()
 {
-	if (CurrentRHISupportsGPUParticles())
+	if (RHISupportsGPUParticles(FeatureLevel))
 	{
 		ResetSimulationPhases();
 		ResetSortedGPUParticles();
@@ -315,7 +316,7 @@ void FFXSystem::PreInitViews()
 
 void FFXSystem::PreRender()
 {
-	if (CurrentRHISupportsGPUParticles())
+	if (RHISupportsGPUParticles(FeatureLevel))
 	{
 		SimulateGPUParticles(EParticleSimulatePhase::Main,NULL,FTexture2DRHIParamRef(),FTexture2DRHIParamRef());
 	}
@@ -323,7 +324,7 @@ void FFXSystem::PreRender()
 
 void FFXSystem::PostRenderOpaque(const class FSceneView* CollisionView, FTexture2DRHIParamRef SceneDepthTexture, FTexture2DRHIParamRef GBufferATexture)
 {
-	if (CurrentRHISupportsGPUParticles())
+	if (RHISupportsGPUParticles(FeatureLevel))
 	{
 		SimulateGPUParticles(EParticleSimulatePhase::Collision,CollisionView,SceneDepthTexture,GBufferATexture);
 		SortGPUParticles();
diff --git a/Engine/Source/Runtime/Engine/Private/Particles/FXSystemPrivate.h b/Engine/Source/Runtime/Engine/Private/Particles/FXSystemPrivate.h
--- a/Engine/Source/Runtime/Engine/Private/Particles/FXSystemPrivate.h
+++ b/Engine/Source/Runtime/Engine/Private/Particles/FXSystemPrivate.h
@@ -48,7 +48,7 @@ class FFXSystem : public FFXSystemInterface
 public:
 
 	/** Default constructoer. */
-	FFXSystem();
+	FFXSystem(ERHIFeatureLevel::Type InFeatureLevel);
 
 	/** Destructor. */
 	virtual ~FFXSystem();
@@ -176,6 +176,8 @@ private:
 	TSparseArray<FParticleSimulationGPU*> GPUSimulations;
 	/** Particle render resources. */
 	FParticleSimulationResources* ParticleSimulationResources;
+	/** Feature level of this effects system */
+	ERHIFeatureLevel::Type FeatureLevel;
 
 #if WITH_EDITOR
 	/** true if the system has been suspended. */
diff --git a/Engine/Source/Runtime/Engine/Private/Particles/ParticleSystemRender.cpp b/Engine/Source/Runtime/Engine/Private/Particles/ParticleSystemRender.cpp
--- a/Engine/Source/Runtime/Engine/Private/Particles/ParticleSystemRender.cpp
+++ b/Engine/Source/Runtime/Engine/Private/Particles/ParticleSystemRender.cpp
@@ -1308,7 +1308,9 @@ int32 FDynamicSpriteEmitterData::Render(FParticleSystemSceneProxy* Proxy, FPrimi
 {
 	SCOPE_CYCLE_COUNTER(STAT_SpriteRenderingTime);
 
-	const bool bInstanced = GRHIFeatureLevel >= ERHIFeatureLevel::SM3;
+	const auto FeatureLevel = View->GetFeatureLevel();
+
+	const bool bInstanced = FeatureLevel >= ERHIFeatureLevel::SM3;
 
 	if (bValid == false)
 	{
@@ -1324,7 +1326,7 @@ int32 FDynamicSpriteEmitterData::Render(FParticleSystemSceneProxy* Proxy, FPrimi
 	{
 		// Don't render if the material will be ignored
 		const bool bIsWireframe = View->Family->EngineShowFlags.Wireframe;
-		if (PDI->IsMaterialIgnored(MaterialResource[bSelected]) && !bIsWireframe)
+		if (PDI->IsMaterialIgnored(MaterialResource[bSelected], FeatureLevel) && !bIsWireframe)
 		{
 			return 0;
 		}
@@ -1724,7 +1726,7 @@ int32 FDynamicMeshEmitterData::Render(FParticleSystemSceneProxy* Proxy, FPrimiti
 
 				// If the material is ignored by the PDI (or not there at all...), 
 				// do not add it to the list of valid elements.
-				if ((MaterialProxy && !PDI->IsMaterialIgnored(MaterialProxy)) || View->Family->EngineShowFlags.Wireframe)
+				if ((MaterialProxy && !PDI->IsMaterialIgnored(MaterialProxy, View->GetFeatureLevel())) || View->Family->EngineShowFlags.Wireframe)
 				{
 					ValidSectionIndices.Add(ElementIndex);
 					bNoValidElements = false;
@@ -2660,7 +2662,7 @@ int32 FDynamicBeam2EmitterData::Render(FParticleSystemSceneProxy* Proxy, FPrimit
 	}
 
 	const bool bIsWireframe = View->Family->EngineShowFlags.Wireframe;
-	bool bMaterialIgnored = PDI->IsMaterialIgnored(MaterialResource[bSelected]);
+	bool bMaterialIgnored = PDI->IsMaterialIgnored(MaterialResource[bSelected], View->GetFeatureLevel());
 	if (bMaterialIgnored && !bIsWireframe)
 	{
 		return 0;
@@ -5446,7 +5448,7 @@ int32 FDynamicTrailsEmitterData::Render(FParticleSystemSceneProxy* Proxy, FPrimi
 	bool const bIsWireframe = View->Family->EngineShowFlags.Wireframe;
 
 	// Don't render if the material will be ignored
-	if (PDI->IsMaterialIgnored(MaterialResource[bSelected]) && !bIsWireframe)
+	if (PDI->IsMaterialIgnored(MaterialResource[bSelected], View->GetFeatureLevel()) && !bIsWireframe)
 	{
 		return 0;
 	}
diff --git a/Engine/Source/Runtime/Engine/Private/ReflectionCaptureComponent.cpp b/Engine/Source/Runtime/Engine/Private/ReflectionCaptureComponent.cpp
--- a/Engine/Source/Runtime/Engine/Private/ReflectionCaptureComponent.cpp
+++ b/Engine/Source/Runtime/Engine/Private/ReflectionCaptureComponent.cpp
@@ -1282,7 +1282,7 @@ void UReflectionCaptureComponent::UpdateReflectionCaptureContents(UWorld* WorldT
 			}
 		}
 
-		if (GRHIFeatureLevel == ERHIFeatureLevel::SM4)
+		if (WorldToUpdate->Scene->GetFeatureLevel() == ERHIFeatureLevel::SM4)
 		{
 			for (int32 CaptureIndex = 0; CaptureIndex < WorldCombinedCaptures.Num(); CaptureIndex++)
 			{
diff --git a/Engine/Source/Runtime/Engine/Private/RendererSupport.cpp b/Engine/Source/Runtime/Engine/Private/RendererSupport.cpp
--- a/Engine/Source/Runtime/Engine/Private/RendererSupport.cpp
+++ b/Engine/Source/Runtime/Engine/Private/RendererSupport.cpp
@@ -125,11 +125,11 @@ void RestoreReferencesToRendererModuleClasses(
 	{
 		UWorld* World = It.Key();
 
-		World->Scene = RendererModule.AllocateScene(World, World->RequiresHitProxies());
+		World->Scene = RendererModule.AllocateScene(World, World->RequiresHitProxies(), GRHIFeatureLevel);
 
 		if (It.Value())
 		{
-			World->FXSystem = FFXSystemInterface::Create();
+			World->FXSystem = FFXSystemInterface::Create(GRHIFeatureLevel);
 			World->Scene->SetFXSystem( World->FXSystem );
 		}
 
diff --git a/Engine/Source/Runtime/Engine/Private/SceneView.cpp b/Engine/Source/Runtime/Engine/Private/SceneView.cpp
--- a/Engine/Source/Runtime/Engine/Private/SceneView.cpp
+++ b/Engine/Source/Runtime/Engine/Private/SceneView.cpp
@@ -1104,6 +1104,11 @@ void FSceneView::ConfigureBufferVisualizationSettings()
 	}
 }
 
+ERHIFeatureLevel::Type FSceneView::GetFeatureLevel() const
+{ 
+	return Family->Scene->GetFeatureLevel();
+}
+
 FSceneViewFamily::FSceneViewFamily( const ConstructionValues& CVS )
 	:
 	FamilySizeX(0),
diff --git a/Engine/Source/Runtime/Engine/Private/World.cpp b/Engine/Source/Runtime/Engine/Private/World.cpp
--- a/Engine/Source/Runtime/Engine/Private/World.cpp
+++ b/Engine/Source/Runtime/Engine/Private/World.cpp
@@ -559,10 +559,10 @@ void UWorld::InitWorld(const InitializationValues IVS)
 		
 		// Save off the value used to create the scene, so this UWorld can recreate its scene later
 		bRequiresHitProxies = IVS.bRequiresHitProxies;
-		Scene = GetRendererModule().AllocateScene( this, bRequiresHitProxies );
+		Scene = GetRendererModule().AllocateScene( this, bRequiresHitProxies, GRHIFeatureLevel );
 		if ( !IsRunningDedicatedServer() && !IsRunningCommandlet() )
 		{
-			FXSystem = FFXSystemInterface::Create();
+			FXSystem = FFXSystemInterface::Create(GRHIFeatureLevel);
 			Scene->SetFXSystem( FXSystem );
 		}
 		else
diff --git a/Engine/Source/Runtime/Engine/Public/FXSystem.h b/Engine/Source/Runtime/Engine/Public/FXSystem.h
--- a/Engine/Source/Runtime/Engine/Public/FXSystem.h
+++ b/Engine/Source/Runtime/Engine/Public/FXSystem.h
@@ -78,12 +78,17 @@ inline bool SupportsGPUParticles(EShaderPlatform Platform)
 	return IsFeatureLevelSupported(Platform, ERHIFeatureLevel::SM4);
 }
 
+inline bool RHISupportsGPUParticles(ERHIFeatureLevel::Type InFeatureLevel)
+{
+	return InFeatureLevel >= ERHIFeatureLevel::SM4;
+}
+
 /**
  * Returns true if the current RHI supports GPU particles.
  */
 inline bool CurrentRHISupportsGPUParticles()
 {
-	return GRHIFeatureLevel >= ERHIFeatureLevel::SM4;
+	return RHISupportsGPUParticles(GRHIFeatureLevel);
 }
 
 /*-----------------------------------------------------------------------------
@@ -100,7 +105,7 @@ public:
 	/**
 	 * Create an effects system instance.
 	 */
-	ENGINE_API static FFXSystemInterface* Create();
+	ENGINE_API static FFXSystemInterface* Create(ERHIFeatureLevel::Type InFeatureLevel);
 
 	/**
 	 * Destroy an effects system instance.
diff --git a/Engine/Source/Runtime/Engine/Public/SceneManagement.h b/Engine/Source/Runtime/Engine/Public/SceneManagement.h
--- a/Engine/Source/Runtime/Engine/Public/SceneManagement.h
+++ b/Engine/Source/Runtime/Engine/Public/SceneManagement.h
@@ -1466,9 +1466,10 @@ public:
 	/**
 	 * Determines whether a particular material will be ignored in this context.
 	 * @param MaterialRenderProxy - The render proxy of the material to check.
+	 * @param InFeatureLevel - The feature level we are currently rendering at
 	 * @return true if meshes using the material will be ignored in this context.
 	 */
-	virtual bool IsMaterialIgnored(const FMaterialRenderProxy* MaterialRenderProxy) const
+	virtual bool IsMaterialIgnored(const FMaterialRenderProxy* MaterialRenderProxy, ERHIFeatureLevel::Type InFeatureLevel) const
 	{
 		return false;
 	}
@@ -1705,7 +1706,7 @@ public:
 	/**
 	 * Creates any needed motion blur infos if needed and saves the transforms of the frame we just completed
 	 */
-	void UpdateMotionBlurCache();
+	void UpdateMotionBlurCache(class FScene* InScene);
 
 	/**
 	 * Call if you want to keep the existing motionblur
@@ -2034,6 +2035,8 @@ public:
 
 	virtual bool IsEditorScene() const { return false; }
 
+	virtual ERHIFeatureLevel::Type GetFeatureLevel() const { return GRHIFeatureLevel; }
+
 protected:
 	virtual ~FSceneInterface() {}
 };
@@ -2522,7 +2525,6 @@ struct FSceneViewInitOptions : public FSceneViewProjectionData
 	}
 };
 
-
 /**
  * A projection from scene space into a 2D screen region.
  */
@@ -2766,6 +2768,9 @@ public:
 
 	/** Configure post process settings for the buffer visualization system */
 	void ConfigureBufferVisualizationSettings();
+
+	/** Get the feature level for this view **/
+	ERHIFeatureLevel::Type GetFeatureLevel() const;
 };
 
 
diff --git a/Engine/Source/Runtime/OpenGLDrv/Private/OpenGLDevice.cpp b/Engine/Source/Runtime/OpenGLDrv/Private/OpenGLDevice.cpp
--- a/Engine/Source/Runtime/OpenGLDrv/Private/OpenGLDevice.cpp
+++ b/Engine/Source/Runtime/OpenGLDrv/Private/OpenGLDevice.cpp
@@ -602,7 +602,7 @@ static void InitRHICapabilitiesForGL()
 	const GLint MinorVersion = FOpenGL::GetMinorVersion();
 
 	// Shader platform & RHI feature level
-	GRHIFeatureLevel = FOpenGL::GetFeatureLevel();
+	SetMaxRHIFeatureLevel(FOpenGL::GetFeatureLevel());
 	GRHIShaderPlatform = FOpenGL::GetShaderPlatform();
 
 	FString FeatureLevelName;
@@ -646,7 +646,7 @@ static void InitRHICapabilitiesForGL()
 	GProjectionSignY = 1.0f;
 
 	// Disable texture streaming on ES2 unless we have the GL_APPLE_copy_texture_levels extension
-	if(GRHIFeatureLevel == ERHIFeatureLevel::ES2 && !FOpenGL::SupportsCopyTextureLevels() )
+	if (GRHIFeatureLevel == ERHIFeatureLevel::ES2 && !FOpenGL::SupportsCopyTextureLevels())
 	{
 		GRHISupportsTextureStreaming = false;
 	}
diff --git a/Engine/Source/Runtime/RHI/Private/RHI.cpp b/Engine/Source/Runtime/RHI/Private/RHI.cpp
--- a/Engine/Source/Runtime/RHI/Private/RHI.cpp
+++ b/Engine/Source/Runtime/RHI/Private/RHI.cpp
@@ -168,8 +168,13 @@ void RHIPrivateBeginFrame()
 
 RHI_API EShaderPlatform GRHIShaderPlatform = SP_PCD3D_SM5;
 
-/** The current feature level supported. */
-RHI_API ERHIFeatureLevel::Type GRHIFeatureLevel = ERHIFeatureLevel::SM5;
+/** The maximum feature level supported on this machine */
+ERHIFeatureLevel::Type GMaxRHIFeatureLevel = ERHIFeatureLevel::SM5;
+
+RHI_API void SetMaxRHIFeatureLevel(ERHIFeatureLevel::Type InType)
+{
+	GMaxRHIFeatureLevel = InType;
+}
 
 FName FeatureLevelNames[] = 
 {
diff --git a/Engine/Source/Runtime/RHI/Public/RHI.h b/Engine/Source/Runtime/RHI/Public/RHI.h
--- a/Engine/Source/Runtime/RHI/Public/RHI.h
+++ b/Engine/Source/Runtime/RHI/Public/RHI.h
@@ -313,8 +313,18 @@ extern RHI_API void GetFeatureLevelName(ERHIFeatureLevel::Type InFeatureLevel, F
 /** Creates an FName for the given feature level. */
 extern RHI_API void GetFeatureLevelName(ERHIFeatureLevel::Type InFeatureLevel, FName& OutName);
 
+extern RHI_API ERHIFeatureLevel::Type GMaxRHIFeatureLevel;
+
 /** The current feature level supported by the RHI. */
-extern RHI_API ERHIFeatureLevel::Type GRHIFeatureLevel;
+inline ERHIFeatureLevel::Type GetRHIFeatureLevel()
+{
+	return GMaxRHIFeatureLevel;
+}
+
+#define GRHIFeatureLevel GetRHIFeatureLevel()
+
+/** Set the current feature level supported by the RHI. */
+extern RHI_API void SetMaxRHIFeatureLevel(ERHIFeatureLevel::Type InType);
 
 /** Table for finding out which shader platform corresponds to a given feature level for this RHI. */
 extern RHI_API EShaderPlatform GShaderPlatformForFeatureLevel[ERHIFeatureLevel::Num];
diff --git a/Engine/Source/Runtime/Renderer/Private/AtmosphereRendering.cpp b/Engine/Source/Runtime/Renderer/Private/AtmosphereRendering.cpp
--- a/Engine/Source/Runtime/Renderer/Private/AtmosphereRendering.cpp
+++ b/Engine/Source/Runtime/Renderer/Private/AtmosphereRendering.cpp
@@ -171,7 +171,7 @@ public:
 	void SetParameters(const FSceneView& View, TRefCountPtr<IPooledRenderTarget>& LightShaftOcclusion)
 	{
 		FGlobalShader::SetParameters(GetPixelShader(), View);
-		SceneTextureParameters.Set(GetPixelShader());
+		SceneTextureParameters.Set(GetPixelShader(), View);
 		AtmosphereTextureParameters.Set(GetPixelShader(), View);
 
 		if (OcclusionTextureParameter.IsBound())
@@ -390,7 +390,7 @@ void SetAtmosphericFogShaders(FScene* Scene, const FViewInfo& View, TRefCountPtr
 void FDeferredShadingSceneRenderer::RenderAtmosphere(FLightShaftsOutput LightShaftsOutput)
 {
 	// Atmospheric fog?
-	if (GRHIFeatureLevel >= ERHIFeatureLevel::SM4 && Scene->HasAtmosphericFog())
+	if (Scene->GetFeatureLevel() >= ERHIFeatureLevel::SM4 && Scene->HasAtmosphericFog())
 	{
 		SCOPED_DRAW_EVENT(Fog, DEC_SCENE_ITEMS);
 
diff --git a/Engine/Source/Runtime/Renderer/Private/BasePassRendering.cpp b/Engine/Source/Runtime/Renderer/Private/BasePassRendering.cpp
--- a/Engine/Source/Runtime/Renderer/Private/BasePassRendering.cpp
+++ b/Engine/Source/Runtime/Renderer/Private/BasePassRendering.cpp
@@ -184,7 +184,8 @@ public:
 				Parameters.TextureMode,
 				Parameters.LightingModel != MLM_Unlit && Scene && Scene->SkyLight,
 				IsTranslucentBlendMode(Parameters.BlendMode) && Scene->HasAtmosphericFog()
-				)
+				),
+			Scene->GetFeatureLevel()
 			);
 	}
 };
@@ -192,7 +193,7 @@ public:
 void FBasePassOpaqueDrawingPolicyFactory::AddStaticMesh(FScene* Scene,FStaticMesh* StaticMesh)
 {
 	// Determine the mesh's material and blend mode.
-	const FMaterial* Material = StaticMesh->MaterialRenderProxy->GetMaterial(GRHIFeatureLevel);
+	const FMaterial* Material = StaticMesh->MaterialRenderProxy->GetMaterial(Scene->GetFeatureLevel());
 	const EBlendMode BlendMode = Material->GetBlendMode();
 
 	// Don't composite static meshes
@@ -209,7 +210,8 @@ void FBasePassOpaqueDrawingPolicyFactory::AddStaticMesh(FScene* Scene,FStaticMes
 				false,
 				bEditorCompositeDepthTest,
 				ESceneRenderTargetsMode::DontSet),
-			FDrawBasePassStaticMeshAction(Scene,StaticMesh)
+			FDrawBasePassStaticMeshAction(Scene,StaticMesh),
+			Scene->GetFeatureLevel()
 			);
 	}
 }
@@ -287,7 +289,7 @@ public:
 			);
 		DrawingPolicy.DrawShared(
 			&View,
-			DrawingPolicy.CreateBoundShaderState()
+			DrawingPolicy.CreateBoundShaderState(View.GetFeatureLevel())
 			);
 
 		for( int32 BatchElementIndex = 0, Num = Parameters.Mesh.Elements.Num(); BatchElementIndex < Num; BatchElementIndex++ )
@@ -324,7 +326,7 @@ bool FBasePassOpaqueDrawingPolicyFactory::DrawDynamicMesh(
 	)
 {
 	// Determine the mesh's material and blend mode.
-	const FMaterial* Material = Mesh.MaterialRenderProxy->GetMaterial(GRHIFeatureLevel);
+	const FMaterial* Material = Mesh.MaterialRenderProxy->GetMaterial(View.GetFeatureLevel());
 	const EBlendMode BlendMode = Material->GetBlendMode();
 
 	// Only draw opaque materials.
@@ -343,7 +345,8 @@ bool FBasePassOpaqueDrawingPolicyFactory::DrawDynamicMesh(
 				View,
 				bBackFace,
 				HitProxyId
-				)
+				),
+			View.GetFeatureLevel()
 			);
 		return true;
 	}
diff --git a/Engine/Source/Runtime/Renderer/Private/BasePassRendering.h b/Engine/Source/Runtime/Renderer/Private/BasePassRendering.h
--- a/Engine/Source/Runtime/Renderer/Private/BasePassRendering.h
+++ b/Engine/Source/Runtime/Renderer/Private/BasePassRendering.h
@@ -334,7 +334,7 @@ public:
 
 		FMeshMaterialShader::SetParameters(ShaderRHI, MaterialRenderProxy, MaterialResource, *View, TextureMode);
 
-		if (GRHIFeatureLevel >= ERHIFeatureLevel::SM4)
+		if (View->GetFeatureLevel() >= ERHIFeatureLevel::SM4)
 		{
 			if (IsTranslucentBlendMode(BlendMode))
 			{
@@ -377,7 +377,7 @@ public:
 
 	void SetMesh(const FVertexFactory* VertexFactory,const FSceneView& View,const FPrimitiveSceneProxy* Proxy,const FMeshBatchElement& BatchElement, EBlendMode BlendMode)
 	{
-		if (GRHIFeatureLevel >= ERHIFeatureLevel::SM4
+		if (View.GetFeatureLevel() >= ERHIFeatureLevel::SM4
 			&& IsTranslucentBlendMode(BlendMode))
 		{
 			TranslucentLightingParameters.SetMesh(this, Proxy);
@@ -578,7 +578,7 @@ public:
 			TShaderMapRef<FShaderComplexityAccumulatePS> ShaderComplexityPixelShader(GetGlobalShaderMap());
 			const uint32 NumPixelShaderInstructions = PixelShader->GetNumInstructions();
 			const uint32 NumVertexShaderInstructions = VertexShader->GetNumInstructions();
-			ShaderComplexityPixelShader->SetParameters(NumVertexShaderInstructions,NumPixelShaderInstructions);
+			ShaderComplexityPixelShader->SetParameters(NumVertexShaderInstructions, NumPixelShaderInstructions, View->GetFeatureLevel());
 		}
 		else
 #endif
@@ -617,7 +617,7 @@ public:
 	* @param DynamicStride - optional stride for dynamic vertex data
 	* @return new bound shader state object
 	*/
-	FBoundShaderStateRHIRef CreateBoundShaderState()
+	FBoundShaderStateRHIRef CreateBoundShaderState(ERHIFeatureLevel::Type InFeatureLevel)
 	{
 		FPixelShaderRHIParamRef PixelShaderRHIRef = PixelShader->GetPixelShader();
 
@@ -684,7 +684,7 @@ public:
 			TShaderMapRef<FShaderComplexityAccumulatePS> ShaderComplexityPixelShader(GetGlobalShaderMap());
 			const uint32 NumPixelShaderInstructions = PixelShader->GetNumInstructions();
 			const uint32 NumVertexShaderInstructions = VertexShader->GetNumInstructions();
-			ShaderComplexityPixelShader->SetParameters(NumVertexShaderInstructions,NumPixelShaderInstructions);
+			ShaderComplexityPixelShader->SetParameters(NumVertexShaderInstructions,NumPixelShaderInstructions, View.GetFeatureLevel());
 		}
 		else
 #endif
@@ -762,11 +762,11 @@ public:
 		const FPrimitiveSceneProxy* PrimitiveSceneProxy,
 		FHitProxyId HitProxyId
 		);
-	static bool IsMaterialIgnored(const FMaterialRenderProxy* MaterialRenderProxy)
+	static bool IsMaterialIgnored(const FMaterialRenderProxy* MaterialRenderProxy, ERHIFeatureLevel::Type InFeatureLevel)
 	{
 		// Ignore non-opaque materials in the opaque base pass.
 		// Note: blend mode does not depend on the feature level.
-		return MaterialRenderProxy && IsTranslucentBlendMode(MaterialRenderProxy->GetMaterial(GRHIFeatureLevel)->GetBlendMode());
+		return MaterialRenderProxy && IsTranslucentBlendMode(MaterialRenderProxy->GetMaterial(InFeatureLevel)->GetBlendMode());
 	}
 };
 
@@ -834,7 +834,8 @@ public:
 template<typename ProcessActionType>
 void ProcessBasePassMesh(
 	const FProcessBasePassMeshParameters& Parameters,
-	const ProcessActionType& Action
+	const ProcessActionType& Action,
+	ERHIFeatureLevel::Type InFeatureLevel
 	)
 {
 	// Check for a cached light-map.
@@ -849,9 +850,9 @@ void ProcessBasePassMesh(
 	if (!(Parameters.TextureMode == ESceneRenderTargetsMode::DontSet && bNeedsSceneTextures))
 	{
 		// Render self-shadowing only for >= SM4 and fallback to non-shadowed for lesser shader models
-		if (bIsLitMaterial && Action.UseTranslucentSelfShadowing() && GRHIFeatureLevel >= ERHIFeatureLevel::SM4)
+		if (bIsLitMaterial && Action.UseTranslucentSelfShadowing() && InFeatureLevel >= ERHIFeatureLevel::SM4)
 		{
-			if (IsIndirectLightingCacheAllowed()
+			if (IsIndirectLightingCacheAllowed(InFeatureLevel)
 				&& Action.AllowIndirectLightingCache()
 				&& Parameters.PrimitiveSceneProxy)
 			{
@@ -907,7 +908,7 @@ void ProcessBasePassMesh(
 							{
 								Action.template Process<FSimpleDynamicLightingPolicy>(Parameters, FSimpleDynamicLightingPolicy(), FSimpleDynamicLightingPolicy::ElementDataType());
 							}
-							else if (IsIndirectLightingCacheAllowed()
+							else if (IsIndirectLightingCacheAllowed(InFeatureLevel)
 								&& Action.AllowIndirectLightingCache()
 								&& Parameters.PrimitiveSceneProxy
 								// Use the indirect lighting cache shaders if the object has a cache allocation
@@ -917,7 +918,7 @@ void ProcessBasePassMesh(
 									// Use the indirect lighting cache shaders if the object is movable, it may not have a cache allocation yet because that is done in InitViews
 									|| Parameters.PrimitiveSceneProxy->IsMovable()))
 							{
-								if (CanIndirectLightingCacheUseVolumeTexture() && Action.AllowIndirectLightingCacheVolumeTexture())
+								if (CanIndirectLightingCacheUseVolumeTexture(InFeatureLevel) && Action.AllowIndirectLightingCacheVolumeTexture())
 								{
 									// Use a lightmap policy that supports reading indirect lighting from a volume texture for dynamic objects
 									Action.template Process<FCachedVolumeIndirectLightingPolicy>(Parameters,FCachedVolumeIndirectLightingPolicy(),FCachedVolumeIndirectLightingPolicy::ElementDataType());
diff --git a/Engine/Source/Runtime/Renderer/Private/CompositionLighting/CompositionLighting.cpp b/Engine/Source/Runtime/Renderer/Private/CompositionLighting/CompositionLighting.cpp
--- a/Engine/Source/Runtime/Renderer/Private/CompositionLighting/CompositionLighting.cpp
+++ b/Engine/Source/Runtime/Renderer/Private/CompositionLighting/CompositionLighting.cpp
@@ -70,7 +70,7 @@ static bool IsReflectionEnvironmentActive(FPostprocessContext& Context)
 	bool HasReflectionCaptures = (Scene->ReflectionSceneData.RegisteredReflectionCaptures.Num() > 0);
 	bool HasSSR = Context.View.Family->EngineShowFlags.ScreenSpaceReflections;
 
-	return (GRHIFeatureLevel == ERHIFeatureLevel::SM5 && IsReflectingEnvironment && (HasReflectionCaptures || HasSSR) && !IsSimpleDynamicLightingEnabled() );
+	return (Scene->GetFeatureLevel() == ERHIFeatureLevel::SM5 && IsReflectingEnvironment && (HasReflectionCaptures || HasSSR) && !IsSimpleDynamicLightingEnabled());
 }
 
 static bool IsBasePassAmbientOcclusionRequired(FPostprocessContext& Context)
diff --git a/Engine/Source/Runtime/Renderer/Private/CompositionLighting/PostProcessDeferredDecals.cpp b/Engine/Source/Runtime/Renderer/Private/CompositionLighting/PostProcessDeferredDecals.cpp
--- a/Engine/Source/Runtime/Renderer/Private/CompositionLighting/PostProcessDeferredDecals.cpp
+++ b/Engine/Source/Runtime/Renderer/Private/CompositionLighting/PostProcessDeferredDecals.cpp
@@ -125,11 +125,11 @@ struct FTransientDecalRenderData
 	const FDeferredDecalProxy* DecalProxy;
 	bool bHasNormal;
 
-	FTransientDecalRenderData(FDeferredDecalProxy* InDecalProxy)
+	FTransientDecalRenderData(const FScene& InScene, FDeferredDecalProxy* InDecalProxy)
 		: DecalProxy(InDecalProxy)
 	{
 		MaterialProxy = InDecalProxy->DecalMaterial->GetRenderProxy(InDecalProxy->bOwnerSelected);
-		MaterialResource = MaterialProxy->GetMaterial(GRHIFeatureLevel);
+		MaterialResource = MaterialProxy->GetMaterial(InScene.GetFeatureLevel());
 		bHasNormal = MaterialResource->HasNormalConnected();
 		DecalBlendMode = ComputeFinalDecalBlendMode((EDecalBlendMode)MaterialResource->GetDecalBlendMode(), bHasNormal);
 		check(MaterialProxy && MaterialResource);
@@ -469,7 +469,7 @@ public:
 	{
 		const FPixelShaderRHIParamRef ShaderRHI = GetPixelShader();
 
-		FMaterialShader::SetParameters(ShaderRHI, MaterialProxy, *MaterialProxy->GetMaterial(GRHIFeatureLevel), View, true, ESceneRenderTargetsMode::SetTextures);
+		FMaterialShader::SetParameters(ShaderRHI, MaterialProxy, *MaterialProxy->GetMaterial(View.GetFeatureLevel()), View, true, ESceneRenderTargetsMode::SetTextures);
 
 		FTransform ComponentTrans = DecalProxy.ComponentTrans;
 
@@ -806,7 +806,7 @@ void FRCPassPostProcessDeferredDecals::Process(FRenderingCompositePassContext& C
 
 		if (bIsShown)
 		{
-			FTransientDecalRenderData Data(DecalProxy);
+			FTransientDecalRenderData Data(Scene, DecalProxy);
 
 			uint32 DecalRenderStage = ComputeRenderStage(Data.DecalBlendMode);
 
diff --git a/Engine/Source/Runtime/Renderer/Private/DeferredShadingRenderer.cpp b/Engine/Source/Runtime/Renderer/Private/DeferredShadingRenderer.cpp
--- a/Engine/Source/Runtime/Renderer/Private/DeferredShadingRenderer.cpp
+++ b/Engine/Source/Runtime/Renderer/Private/DeferredShadingRenderer.cpp
@@ -424,7 +424,7 @@ void FDeferredShadingSceneRenderer::RenderFinish()
 	//grab the new transform out of the proxies for next frame
 	if(ViewFamily.EngineShowFlags.MotionBlur || ViewFamily.EngineShowFlags.TemporalAA)
 	{
-		Scene->MotionBlurInfoData.UpdateMotionBlurCache();
+		Scene->MotionBlurInfoData.UpdateMotionBlurCache(Scene);
 	}
 
 	// Some RT should be released as early as possible to allow sharing of that memory for other purposes.
@@ -445,6 +445,8 @@ void FDeferredShadingSceneRenderer::Render()
 
 	SCOPED_DRAW_EVENT(Scene,DEC_SCENE_ITEMS);
 
+	const auto FeatureLevel = Scene->GetFeatureLevel();
+
 	// Initialize global system textures (pass-through if already initialized).
 	GSystemTextures.InitializeTextures();
 
@@ -455,7 +457,6 @@ void FDeferredShadingSceneRenderer::Render()
 	InitViews();
 
 	const bool bIsWireframe = ViewFamily.EngineShowFlags.Wireframe;
-
 	static const auto ClearMethodCVar = IConsoleManager::Get().FindTConsoleVariableDataInt(TEXT("r.ClearSceneMethod"));
 	bool bRequiresRHIClear = true;
 	bool bRequiresFarZQuadClear = false;
@@ -622,7 +623,7 @@ void FDeferredShadingSceneRenderer::Render()
 	
 	// Render lighting.
 	if (ViewFamily.EngineShowFlags.Lighting
-		&& GRHIFeatureLevel >= ERHIFeatureLevel::SM4
+		&& FeatureLevel >= ERHIFeatureLevel::SM4
 		&& ViewFamily.EngineShowFlags.DeferredLighting
 		)
 	{
@@ -667,7 +668,7 @@ void FDeferredShadingSceneRenderer::Render()
 	}
 
 	if( ViewFamily.EngineShowFlags.StationaryLightOverlap &&
-		GRHIFeatureLevel >= ERHIFeatureLevel::SM4)
+		FeatureLevel >= ERHIFeatureLevel::SM4)
 	{
 		RenderStationaryLightOverlap();
 	}
@@ -987,7 +988,7 @@ public:
 		SetShaderValue(GetPixelShader(), SourceTexelOffsets01, Offsets01);
 		const FVector4 Offsets23(0.0f, 1.0f / DownsampledBufferSizeY, 1.0f / DownsampledBufferSizeX, 1.0f / DownsampledBufferSizeY);
 		SetShaderValue(GetPixelShader(), SourceTexelOffsets23, Offsets23);
-		SceneTextureParameters.Set(GetPixelShader());
+		SceneTextureParameters.Set(GetPixelShader(), View);
 	}
 
 	virtual bool Serialize(FArchive& Ar)
diff --git a/Engine/Source/Runtime/Renderer/Private/DepthRendering.cpp b/Engine/Source/Runtime/Renderer/Private/DepthRendering.cpp
--- a/Engine/Source/Runtime/Renderer/Private/DepthRendering.cpp
+++ b/Engine/Source/Runtime/Renderer/Private/DepthRendering.cpp
@@ -196,7 +196,7 @@ void FDepthDrawingPolicy::DrawShared(const FSceneView* View,FBoundShaderStateRHI
 * as well as the shaders needed to draw the mesh
 * @return new bound shader state object
 */
-FBoundShaderStateRHIRef FDepthDrawingPolicy::CreateBoundShaderState()
+FBoundShaderStateRHIRef FDepthDrawingPolicy::CreateBoundShaderState(ERHIFeatureLevel::Type InFeatureLevel)
 {
 	return RHICreateBoundShaderState(
 		FMeshDrawingPolicy::GetVertexDeclaration(), 
@@ -273,12 +273,12 @@ void FPositionOnlyDepthDrawingPolicy::DrawShared(const FSceneView* View,FBoundSh
 * as well as the shaders needed to draw the mesh
 * @return new bound shader state object
 */
-FBoundShaderStateRHIRef FPositionOnlyDepthDrawingPolicy::CreateBoundShaderState()
+FBoundShaderStateRHIRef FPositionOnlyDepthDrawingPolicy::CreateBoundShaderState(ERHIFeatureLevel::Type InFeatureLevel)
 {
 	FVertexDeclarationRHIParamRef VertexDeclaration;
 	VertexDeclaration = VertexFactory->GetPositionDeclaration();
 
-	checkSlow(MaterialRenderProxy->GetMaterial(GRHIFeatureLevel)->GetBlendMode() == BLEND_Opaque);
+	checkSlow(MaterialRenderProxy->GetMaterial(InFeatureLevel)->GetBlendMode() == BLEND_Opaque);
 	return RHICreateBoundShaderState(VertexDeclaration, VertexShader->GetVertexShader(), FHullShaderRHIRef(), FDomainShaderRHIRef(), FPixelShaderRHIRef(), FGeometryShaderRHIRef());
 }
 
@@ -307,7 +307,7 @@ int32 CompareDrawingPolicy(const FPositionOnlyDepthDrawingPolicy& A,const FPosit
 void FDepthDrawingPolicyFactory::AddStaticMesh(FScene* Scene,FStaticMesh* StaticMesh)
 {
 	const FMaterialRenderProxy* MaterialRenderProxy = StaticMesh->MaterialRenderProxy;
-	const FMaterial* Material = MaterialRenderProxy->GetMaterial(GRHIFeatureLevel);
+	const FMaterial* Material = MaterialRenderProxy->GetMaterial(Scene->GetFeatureLevel());
 	const EBlendMode BlendMode = Material->GetBlendMode();
 
 	if (Material->IsMasked())
@@ -321,7 +321,8 @@ void FDepthDrawingPolicyFactory::AddStaticMesh(FScene* Scene,FStaticMesh* Static
 				MaterialRenderProxy,
 				*Material,
 				Material->IsTwoSided()
-				)
+				),
+			Scene->GetFeatureLevel()
 			);
 	}
 	else
@@ -337,10 +338,11 @@ void FDepthDrawingPolicyFactory::AddStaticMesh(FScene* Scene,FStaticMesh* Static
 				FPositionOnlyDepthDrawingPolicy(
 					StaticMesh->VertexFactory,
 					DefaultProxy,
-					*DefaultProxy->GetMaterial(GRHIFeatureLevel),
+					*DefaultProxy->GetMaterial(Scene->GetFeatureLevel()),
 					Material->IsTwoSided(),
 					Material->IsWireframe()
-					)
+					),
+				Scene->GetFeatureLevel()
 				);
 		}
 		else
@@ -358,9 +360,10 @@ void FDepthDrawingPolicyFactory::AddStaticMesh(FScene* Scene,FStaticMesh* Static
 				FDepthDrawingPolicy(
 					StaticMesh->VertexFactory,
 					MaterialRenderProxy,
-					*MaterialRenderProxy->GetMaterial(GRHIFeatureLevel),
+					*MaterialRenderProxy->GetMaterial(Scene->GetFeatureLevel()),
 					Material->IsTwoSided()
-					)
+					),
+				Scene->GetFeatureLevel()
 				);
 		}
 	}
@@ -384,7 +387,7 @@ bool FDepthDrawingPolicyFactory::DrawMesh(
 	if (Mesh.bUseAsOccluder)
 	{
 		const FMaterialRenderProxy* MaterialRenderProxy = Mesh.MaterialRenderProxy;
-		const FMaterial* Material = MaterialRenderProxy->GetMaterial(GRHIFeatureLevel);
+		const FMaterial* Material = MaterialRenderProxy->GetMaterial(View.GetFeatureLevel());
 		const EBlendMode BlendMode = Material->GetBlendMode();
 
 		// Check to see if the primitive is currently fading in or out using the screen door effect.  If it is,
@@ -397,8 +400,8 @@ bool FDepthDrawingPolicyFactory::DrawMesh(
 		{
 			//render opaque primitives that support a separate position-only vertex buffer
 			const FMaterialRenderProxy* DefaultProxy = UMaterial::GetDefaultMaterial(MD_Surface)->GetRenderProxy(false);
-			FPositionOnlyDepthDrawingPolicy DrawingPolicy(Mesh.VertexFactory,DefaultProxy,*DefaultProxy->GetMaterial(GRHIFeatureLevel),Material->IsTwoSided(),Material->IsWireframe());
-			DrawingPolicy.DrawShared(&View,DrawingPolicy.CreateBoundShaderState());
+			FPositionOnlyDepthDrawingPolicy DrawingPolicy(Mesh.VertexFactory, DefaultProxy, *DefaultProxy->GetMaterial(View.GetFeatureLevel()), Material->IsTwoSided(), Material->IsWireframe());
+			DrawingPolicy.DrawShared(&View,DrawingPolicy.CreateBoundShaderState(View.GetFeatureLevel()));
 			int32 BatchElementIndex = 0;
 			uint64 Mask = BatchElementMask;
 			do
@@ -439,8 +442,8 @@ bool FDepthDrawingPolicyFactory::DrawMesh(
 					MaterialRenderProxy = UMaterial::GetDefaultMaterial(MD_Surface)->GetRenderProxy(false);
 				}
 
-				FDepthDrawingPolicy DrawingPolicy(Mesh.VertexFactory,MaterialRenderProxy,*MaterialRenderProxy->GetMaterial(GRHIFeatureLevel),Material->IsTwoSided());
-				DrawingPolicy.DrawShared(&View,DrawingPolicy.CreateBoundShaderState());
+				FDepthDrawingPolicy DrawingPolicy(Mesh.VertexFactory, MaterialRenderProxy, *MaterialRenderProxy->GetMaterial(View.GetFeatureLevel()), Material->IsTwoSided());
+				DrawingPolicy.DrawShared(&View,DrawingPolicy.CreateBoundShaderState(View.GetFeatureLevel()));
 
 				int32 BatchElementIndex = 0;
 				uint64 Mask = BatchElementMask;
@@ -497,7 +500,7 @@ bool FDepthDrawingPolicyFactory::DrawStaticMesh(
 {
 	bool bDirty = false;
 
-	const FMaterial* Material = StaticMesh.MaterialRenderProxy->GetMaterial(GRHIFeatureLevel);
+	const FMaterial* Material = StaticMesh.MaterialRenderProxy->GetMaterial(View.GetFeatureLevel());
 	const EMaterialLightingModel LightingModel = Material->GetLightingModel();
 	bDirty |= DrawMesh(
 		View,
@@ -513,7 +516,7 @@ bool FDepthDrawingPolicyFactory::DrawStaticMesh(
 	return bDirty;
 }
 
-bool FDepthDrawingPolicyFactory::IsMaterialIgnored(const FMaterialRenderProxy* MaterialRenderProxy)
+bool FDepthDrawingPolicyFactory::IsMaterialIgnored(const FMaterialRenderProxy* MaterialRenderProxy, ERHIFeatureLevel::Type InFeatureLevel)
 {
-	return IsTranslucentBlendMode(MaterialRenderProxy->GetMaterial(GRHIFeatureLevel)->GetBlendMode());
+	return IsTranslucentBlendMode(MaterialRenderProxy->GetMaterial(InFeatureLevel)->GetBlendMode());
 }
diff --git a/Engine/Source/Runtime/Renderer/Private/DepthRendering.h b/Engine/Source/Runtime/Renderer/Private/DepthRendering.h
--- a/Engine/Source/Runtime/Renderer/Private/DepthRendering.h
+++ b/Engine/Source/Runtime/Renderer/Private/DepthRendering.h
@@ -52,7 +52,7 @@ public:
 	* @param DynamicStride - optional stride for dynamic vertex data
 	* @return new bound shader state object
 	*/
-	FBoundShaderStateRHIRef CreateBoundShaderState();
+	FBoundShaderStateRHIRef CreateBoundShaderState(ERHIFeatureLevel::Type InFeatureLevel);
 
 	void SetMeshRenderState(
 		const FSceneView& View,
@@ -103,7 +103,7 @@ public:
 	* as well as the shaders needed to draw the mesh
 	* @return new bound shader state object
 	*/
-	FBoundShaderStateRHIRef CreateBoundShaderState();
+	FBoundShaderStateRHIRef CreateBoundShaderState(ERHIFeatureLevel::Type InFeatureLevel);
 
 	void SetMeshRenderState(
 		const FSceneView& View,
@@ -158,7 +158,7 @@ public:
 		FHitProxyId HitProxyId
 		);
 
-	static bool IsMaterialIgnored(const FMaterialRenderProxy* MaterialRenderProxy);
+	static bool IsMaterialIgnored(const FMaterialRenderProxy* MaterialRenderProxy, ERHIFeatureLevel::Type InFeatureLevel);
 
 private:
 	/**
diff --git a/Engine/Source/Runtime/Renderer/Private/DistortionRendering.cpp b/Engine/Source/Runtime/Renderer/Private/DistortionRendering.cpp
--- a/Engine/Source/Runtime/Renderer/Private/DistortionRendering.cpp
+++ b/Engine/Source/Runtime/Renderer/Private/DistortionRendering.cpp
@@ -125,7 +125,7 @@ public:
 	
 	void SetParameters(const FVertexFactory* VertexFactory,const FMaterialRenderProxy* MaterialRenderProxy,const FSceneView* View)
 	{
-		FMeshMaterialShader::SetParameters(GetVertexShader(),MaterialRenderProxy,*MaterialRenderProxy->GetMaterial(GRHIFeatureLevel),*View,ESceneRenderTargetsMode::SetTextures);
+		FMeshMaterialShader::SetParameters(GetVertexShader(), MaterialRenderProxy, *MaterialRenderProxy->GetMaterial(View->GetFeatureLevel()), *View, ESceneRenderTargetsMode::SetTextures);
 	}
 
 	void SetMesh(const FVertexFactory* VertexFactory,const FSceneView& View,const FPrimitiveSceneProxy* Proxy,const FMeshBatchElement& BatchElement)
@@ -219,8 +219,7 @@ public:
 		const FSceneView& View
 		)
 	{
-		FMeshMaterialShader::SetParameters(GetPixelShader(),MaterialRenderProxy,*MaterialRenderProxy->GetMaterial(GRHIFeatureLevel),View,ESceneRenderTargetsMode::SetTextures);
-
+		FMeshMaterialShader::SetParameters(GetPixelShader(), MaterialRenderProxy, *MaterialRenderProxy->GetMaterial(View.GetFeatureLevel()), View, ESceneRenderTargetsMode::SetTextures);
 
 		float Ratio = View.UnscaledViewRect.Width() / (float)View.UnscaledViewRect.Height();
 		float Params[4];
@@ -302,7 +301,7 @@ public:
 	* as well as the shaders needed to draw the mesh
 	* @return new bound shader state object
 	*/
-	FBoundShaderStateRHIRef CreateBoundShaderState();
+	FBoundShaderStateRHIRef CreateBoundShaderState(ERHIFeatureLevel::Type InFeatureLevel);
 
 	/**
 	* Sets the render states for drawing a mesh.
@@ -448,7 +447,7 @@ void TDistortionMeshDrawingPolicy<DistortMeshPolicy>::DrawShared(
 * @return new bound shader state object
 */
 template<class DistortMeshPolicy>
-FBoundShaderStateRHIRef TDistortionMeshDrawingPolicy<DistortMeshPolicy>::CreateBoundShaderState()
+FBoundShaderStateRHIRef TDistortionMeshDrawingPolicy<DistortMeshPolicy>::CreateBoundShaderState(ERHIFeatureLevel::Type InFeatureLevel)
 {
 	FPixelShaderRHIParamRef PixelShaderRHIRef = NULL;
 
@@ -571,7 +570,7 @@ public:
 		FHitProxyId HitProxyId
 		);
 
-	static bool IsMaterialIgnored(const FMaterialRenderProxy* Material);
+	static bool IsMaterialIgnored(const FMaterialRenderProxy* Material, ERHIFeatureLevel::Type InFeatureLevel);
 };
 
 /**
@@ -589,7 +588,8 @@ bool TDistortionMeshDrawingPolicyFactory<DistortMeshPolicy>::DrawDynamicMesh(
 	FHitProxyId HitProxyId
 	)
 {
-	bool bDistorted = Mesh.MaterialRenderProxy && Mesh.MaterialRenderProxy->GetMaterial(GRHIFeatureLevel)->IsDistorted();
+	const auto FeatureLevel = View.GetFeatureLevel();
+	bool bDistorted = Mesh.MaterialRenderProxy && Mesh.MaterialRenderProxy->GetMaterial(FeatureLevel)->IsDistorted();
 
 	if(bDistorted && !bBackFace)
 	{
@@ -597,11 +597,11 @@ bool TDistortionMeshDrawingPolicyFactory<DistortMeshPolicy>::DrawDynamicMesh(
 		TDistortionMeshDrawingPolicy<DistortMeshPolicy> DrawingPolicy(
 			Mesh.VertexFactory,
 			Mesh.MaterialRenderProxy,
-			*Mesh.MaterialRenderProxy->GetMaterial(GRHIFeatureLevel),
+			*Mesh.MaterialRenderProxy->GetMaterial(FeatureLevel),
 			bInitializeOffsets,
 			View.Family->EngineShowFlags.ShaderComplexity
 			);
-		DrawingPolicy.DrawShared(&View,DrawingPolicy.CreateBoundShaderState());
+		DrawingPolicy.DrawShared(&View,DrawingPolicy.CreateBoundShaderState(View.GetFeatureLevel()));
 		for( int32 BatchElementIndex=0;BatchElementIndex < Mesh.Elements.Num();BatchElementIndex++ )
 		{
 			DrawingPolicy.SetMeshRenderState(View,PrimitiveSceneProxy,Mesh,BatchElementIndex,bBackFace,typename TDistortionMeshDrawingPolicy<DistortMeshPolicy>::ElementDataType());
@@ -631,7 +631,8 @@ bool TDistortionMeshDrawingPolicyFactory<DistortMeshPolicy>::DrawStaticMesh(
 	FHitProxyId HitProxyId
 	)
 {
-	bool bDistorted = StaticMesh.MaterialRenderProxy && StaticMesh.MaterialRenderProxy->GetMaterial(GRHIFeatureLevel)->IsDistorted();
+	const auto FeatureLevel = View->GetFeatureLevel();
+	bool bDistorted = StaticMesh.MaterialRenderProxy && StaticMesh.MaterialRenderProxy->GetMaterial(FeatureLevel)->IsDistorted();
 	
 	if(bDistorted && !bBackFace)
 	{
@@ -639,11 +640,11 @@ bool TDistortionMeshDrawingPolicyFactory<DistortMeshPolicy>::DrawStaticMesh(
 		TDistortionMeshDrawingPolicy<DistortMeshPolicy> DrawingPolicy(
 			StaticMesh.VertexFactory,
 			StaticMesh.MaterialRenderProxy,
-			*StaticMesh.MaterialRenderProxy->GetMaterial(GRHIFeatureLevel),
+			*StaticMesh.MaterialRenderProxy->GetMaterial(FeatureLevel),
 			bInitializeOffsets,
 			View->Family->EngineShowFlags.ShaderComplexity
 			);
-		DrawingPolicy.DrawShared(View,DrawingPolicy.CreateBoundShaderState());
+		DrawingPolicy.DrawShared(View,DrawingPolicy.CreateBoundShaderState(FeatureLevel));
 		int32 BatchElementIndex = 0;
 		do
 		{
@@ -665,10 +666,10 @@ bool TDistortionMeshDrawingPolicyFactory<DistortMeshPolicy>::DrawStaticMesh(
 }
 
 template<typename DistortMeshPolicy>
-bool TDistortionMeshDrawingPolicyFactory<DistortMeshPolicy>::IsMaterialIgnored(const FMaterialRenderProxy* MaterialRenderProxy)
+bool TDistortionMeshDrawingPolicyFactory<DistortMeshPolicy>::IsMaterialIgnored(const FMaterialRenderProxy* MaterialRenderProxy, ERHIFeatureLevel::Type InFeatureLevel)
 {
 	// Non-distorted materials are ignored when rendering distortion.
-	return MaterialRenderProxy && !MaterialRenderProxy->GetMaterial(GRHIFeatureLevel)->IsDistorted();
+	return MaterialRenderProxy && !MaterialRenderProxy->GetMaterial(InFeatureLevel)->IsDistorted();
 }
 
 /*-----------------------------------------------------------------------------
diff --git a/Engine/Source/Runtime/Renderer/Private/DynamicPrimitiveDrawing.h b/Engine/Source/Runtime/Renderer/Private/DynamicPrimitiveDrawing.h
--- a/Engine/Source/Runtime/Renderer/Private/DynamicPrimitiveDrawing.h
+++ b/Engine/Source/Runtime/Renderer/Private/DynamicPrimitiveDrawing.h
@@ -54,7 +54,7 @@ public:
 	virtual void SetHitProxy(HHitProxy* HitProxy) OVERRIDE;
 	virtual bool IsRenderingSelectionOutline() const OVERRIDE { return bIsSelectionOutlineRendering; }
 	virtual void RegisterDynamicResource(FDynamicPrimitiveResource* DynamicResource) OVERRIDE;
-	virtual bool IsMaterialIgnored(const FMaterialRenderProxy* MaterialRenderProxy) const OVERRIDE;
+	virtual bool IsMaterialIgnored(const FMaterialRenderProxy* MaterialRenderProxy, ERHIFeatureLevel::Type InFeatureLevel) const OVERRIDE;
 	virtual int32 DrawMesh(const FMeshBatch& Mesh) OVERRIDE;
 	virtual void AddReserveLines(uint8 DepthPriorityGroup, int32 NumLines, bool bDepthBiased = false) OVERRIDE;
 	virtual void DrawSprite(
diff --git a/Engine/Source/Runtime/Renderer/Private/DynamicPrimitiveDrawing.inl b/Engine/Source/Runtime/Renderer/Private/DynamicPrimitiveDrawing.inl
--- a/Engine/Source/Runtime/Renderer/Private/DynamicPrimitiveDrawing.inl
+++ b/Engine/Source/Runtime/Renderer/Private/DynamicPrimitiveDrawing.inl
@@ -89,9 +89,9 @@ void TDynamicPrimitiveDrawer<DrawingPolicyFactoryType>::RegisterDynamicResource(
 }
 
 template<typename DrawingPolicyFactoryType>
-bool TDynamicPrimitiveDrawer<DrawingPolicyFactoryType>::IsMaterialIgnored(const FMaterialRenderProxy* MaterialRenderProxy) const
+bool TDynamicPrimitiveDrawer<DrawingPolicyFactoryType>::IsMaterialIgnored(const FMaterialRenderProxy* MaterialRenderProxy, ERHIFeatureLevel::Type InFeatureLevel) const
 {
-	return DrawingPolicyFactoryType::IsMaterialIgnored(MaterialRenderProxy);
+	return DrawingPolicyFactoryType::IsMaterialIgnored(MaterialRenderProxy, InFeatureLevel);
 }
 
 template<typename DrawingPolicyFactoryType>
@@ -223,9 +223,10 @@ bool DrawViewElements(
 	for(int32 MeshIndex = 0;MeshIndex < ViewMeshElementList.Num();MeshIndex++)
 	{
 		const FHitProxyMeshPair& Mesh = ViewMeshElementList[MeshIndex];
+		const auto FeatureLevel = View.GetFeatureLevel();
 		check(Mesh.MaterialRenderProxy);
-		check(Mesh.MaterialRenderProxy->GetMaterial(GRHIFeatureLevel));
-		const bool bIsTwoSided = Mesh.MaterialRenderProxy->GetMaterial(GRHIFeatureLevel)->IsTwoSided();
+		check(Mesh.MaterialRenderProxy->GetMaterial(FeatureLevel));
+		const bool bIsTwoSided = Mesh.MaterialRenderProxy->GetMaterial(FeatureLevel)->IsTwoSided();
 		int32 bBackFace = bIsTwoSided ? 1 : 0;
 		do
 		{
diff --git a/Engine/Source/Runtime/Renderer/Private/FogRendering.cpp b/Engine/Source/Runtime/Renderer/Private/FogRendering.cpp
--- a/Engine/Source/Runtime/Renderer/Private/FogRendering.cpp
+++ b/Engine/Source/Runtime/Renderer/Private/FogRendering.cpp
@@ -127,7 +127,7 @@ public:
 	void SetParameters(const FViewInfo& View, FLightShaftsOutput LightShaftsOutput)
 	{
 		FGlobalShader::SetParameters(GetPixelShader(), View);
-		SceneTextureParameters.Set(GetPixelShader());
+		SceneTextureParameters.Set(GetPixelShader(), View);
 		ExponentialParameters.Set(GetPixelShader(), &View);
 
 		if (LightShaftsOutput.bRendered)
diff --git a/Engine/Source/Runtime/Renderer/Private/ForwardBasePassRendering.cpp b/Engine/Source/Runtime/Renderer/Private/ForwardBasePassRendering.cpp
--- a/Engine/Source/Runtime/Renderer/Private/ForwardBasePassRendering.cpp
+++ b/Engine/Source/Runtime/Renderer/Private/ForwardBasePassRendering.cpp
@@ -79,7 +79,8 @@ public:
 				LightMapPolicy,
 				Parameters.BlendMode,
 				Parameters.TextureMode
-				)
+				),
+			Scene->GetFeatureLevel()
 			);
 	}
 };
@@ -87,7 +88,7 @@ public:
 void FBasePassForwardOpaqueDrawingPolicyFactory::AddStaticMesh(FScene* Scene,FStaticMesh* StaticMesh)
 {
 	// Determine the mesh's material and blend mode.
-	const FMaterial* Material = StaticMesh->MaterialRenderProxy->GetMaterial(GRHIFeatureLevel);
+	const FMaterial* Material = StaticMesh->MaterialRenderProxy->GetMaterial(Scene->GetFeatureLevel());
 	const EBlendMode BlendMode = Material->GetBlendMode();
 
 	// Only draw opaque materials.
@@ -167,7 +168,7 @@ public:
 			);
 		DrawingPolicy.DrawShared(
 			&View,
-			DrawingPolicy.CreateBoundShaderState()
+			DrawingPolicy.CreateBoundShaderState(View.GetFeatureLevel())
 			);
 
 		for( int32 BatchElementIndex=0;BatchElementIndex<Parameters.Mesh.Elements.Num();BatchElementIndex++ )
@@ -205,7 +206,7 @@ bool FBasePassForwardOpaqueDrawingPolicyFactory::DrawDynamicMesh(
 	)
 {
 	// Determine the mesh's material and blend mode.
-	const FMaterial* Material = Mesh.MaterialRenderProxy->GetMaterial(GRHIFeatureLevel);
+	const FMaterial* Material = Mesh.MaterialRenderProxy->GetMaterial(View.GetFeatureLevel());
 	const EBlendMode BlendMode = Material->GetBlendMode();
 
 	// Only draw opaque materials.
diff --git a/Engine/Source/Runtime/Renderer/Private/ForwardBasePassRendering.h b/Engine/Source/Runtime/Renderer/Private/ForwardBasePassRendering.h
--- a/Engine/Source/Runtime/Renderer/Private/ForwardBasePassRendering.h
+++ b/Engine/Source/Runtime/Renderer/Private/ForwardBasePassRendering.h
@@ -286,7 +286,7 @@ public:
 			TShaderMapRef<FShaderComplexityAccumulatePS> ShaderComplexityPixelShader(GetGlobalShaderMap());
 			const uint32 NumPixelShaderInstructions = PixelShader->GetNumInstructions();
 			const uint32 NumVertexShaderInstructions = VertexShader->GetNumInstructions();
-			ShaderComplexityPixelShader->SetParameters(NumVertexShaderInstructions,NumPixelShaderInstructions);
+			ShaderComplexityPixelShader->SetParameters(NumVertexShaderInstructions,NumPixelShaderInstructions, View->GetFeatureLevel());
 		}
 		else
 #endif
@@ -325,7 +325,7 @@ public:
 	* as well as the shaders needed to draw the mesh
 	* @return new bound shader state object
 	*/
-	FBoundShaderStateRHIRef CreateBoundShaderState()
+	FBoundShaderStateRHIRef CreateBoundShaderState(ERHIFeatureLevel::Type InFeatureLevel)
 	{
 		FPixelShaderRHIParamRef PixelShaderRHIRef = PixelShader->GetPixelShader();
 
@@ -385,7 +385,7 @@ public:
 			TShaderMapRef<FShaderComplexityAccumulatePS> ShaderComplexityPixelShader(GetGlobalShaderMap());
 			const uint32 NumPixelShaderInstructions = PixelShader->GetNumInstructions();
 			const uint32 NumVertexShaderInstructions = VertexShader->GetNumInstructions();
-			ShaderComplexityPixelShader->SetParameters(NumVertexShaderInstructions,NumPixelShaderInstructions);
+			ShaderComplexityPixelShader->SetParameters(NumVertexShaderInstructions,NumPixelShaderInstructions, View.GetFeatureLevel());
 		}
 		else
 #endif
@@ -443,11 +443,11 @@ public:
 		const FPrimitiveSceneProxy* PrimitiveSceneProxy,
 		FHitProxyId HitProxyId
 		);
-	static bool IsMaterialIgnored(const FMaterialRenderProxy* MaterialRenderProxy)
+	static bool IsMaterialIgnored(const FMaterialRenderProxy* MaterialRenderProxy, ERHIFeatureLevel::Type InFeatureLevel)
 	{
 		// Ignore non-opaque materials in the opaque base pass.
 		// Note: blend mode does not depend on the feature level.
-		return MaterialRenderProxy && IsTranslucentBlendMode(MaterialRenderProxy->GetMaterial(GRHIFeatureLevel)->GetBlendMode());
+		return MaterialRenderProxy && IsTranslucentBlendMode(MaterialRenderProxy->GetMaterial(InFeatureLevel)->GetBlendMode());
 	}
 };
 
@@ -469,8 +469,8 @@ void ProcessBasePassMeshForForwardShading(
 
 	if (LightMapInteraction.GetType() == LMIT_Texture)
 	{
-		const FShadowMapInteraction ShadowMapInteraction = (Parameters.Mesh.LCI && bIsLitMaterial) 
-			? Parameters.Mesh.LCI->GetShadowMapInteraction() 
+		const FShadowMapInteraction ShadowMapInteraction = (Parameters.Mesh.LCI && bIsLitMaterial)
+			? Parameters.Mesh.LCI->GetShadowMapInteraction()
 			: FShadowMapInteraction();
 
 		if (ShadowMapInteraction.GetType() == SMIT_Texture)
@@ -482,11 +482,11 @@ void ProcessBasePassMeshForForwardShading(
 		}
 		else
 		{
-			Action.template Process< TLightMapPolicy<LQ_LIGHTMAP> >( Parameters, TLightMapPolicy<LQ_LIGHTMAP>(), LightMapInteraction );
+			Action.template Process< TLightMapPolicy<LQ_LIGHTMAP> >(Parameters, TLightMapPolicy<LQ_LIGHTMAP>(), LightMapInteraction);
 		}
 	}
-	else if (bIsLitMaterial 
-		&& IsIndirectLightingCacheAllowed()
+	else if (bIsLitMaterial
+		&& IsIndirectLightingCacheAllowed(GRHIFeatureLevel)
 		&& Parameters.PrimitiveSceneProxy
 		// Movable objects need to get their GI from the indirect lighting cache
 		&& Parameters.PrimitiveSceneProxy->IsMovable())
diff --git a/Engine/Source/Runtime/Renderer/Private/ForwardShadingRenderer.cpp b/Engine/Source/Runtime/Renderer/Private/ForwardShadingRenderer.cpp
--- a/Engine/Source/Runtime/Renderer/Private/ForwardShadingRenderer.cpp
+++ b/Engine/Source/Runtime/Renderer/Private/ForwardShadingRenderer.cpp
@@ -45,6 +45,8 @@ void FForwardShadingSceneRenderer::Render()
 		return;
 	}
 
+	auto FeatureLevel = ViewFamily.Scene->GetFeatureLevel();
+
 	// Initialize global system textures (pass-through if already initialized).
 	GSystemTextures.InitializeTextures();
 
diff --git a/Engine/Source/Runtime/Renderer/Private/ForwardTranslucentRendering.cpp b/Engine/Source/Runtime/Renderer/Private/ForwardTranslucentRendering.cpp
--- a/Engine/Source/Runtime/Renderer/Private/ForwardTranslucentRendering.cpp
+++ b/Engine/Source/Runtime/Renderer/Private/ForwardTranslucentRendering.cpp
@@ -25,9 +25,9 @@ public:
 	}
 	FForwardCopySceneAlphaPS() {}
 
-	void SetParameters()
+	void SetParameters(const FSceneView& View)
 	{
-		SceneTextureParameters.Set(GetPixelShader());
+		SceneTextureParameters.Set(GetPixelShader(), View);
 	}
 
 	virtual bool Serialize(FArchive& Ar)
@@ -45,7 +45,7 @@ IMPLEMENT_SHADER_TYPE(,FForwardCopySceneAlphaPS,TEXT("TranslucentLightingShaders
 
 FGlobalBoundShaderState ForwardCopySceneAlphaBoundShaderState;
 
-void FForwardShadingSceneRenderer::CopySceneAlpha(void)
+void FForwardShadingSceneRenderer::CopySceneAlpha(const FSceneView& View)
 {
 	SCOPED_DRAW_EVENTF(EventCopy, DEC_SCENE_ITEMS, TEXT("CopySceneAlpha"));
 	RHISetRasterizerState(TStaticRasterizerState<FM_Solid,CM_None>::GetRHI());
@@ -65,7 +65,7 @@ void FForwardShadingSceneRenderer::CopySceneAlpha(void)
 	TShaderMapRef<FForwardCopySceneAlphaPS> PixelShader(GetGlobalShaderMap());
 	SetGlobalBoundShaderState(ForwardCopySceneAlphaBoundShaderState, GFilterVertexDeclaration.VertexDeclarationRHI, *ScreenVertexShader, *PixelShader);
 
-	PixelShader->SetParameters();
+	PixelShader->SetParameters(View);
 
 	DrawRectangle( 
 		0, 0, 
@@ -139,7 +139,7 @@ public:
 
 		DrawingPolicy.DrawShared(
 			&View,
-			DrawingPolicy.CreateBoundShaderState()
+			DrawingPolicy.CreateBoundShaderState(View.GetFeatureLevel())
 			);
 
 		for (int32 BatchElementIndex = 0; BatchElementIndex<Parameters.Mesh.Elements.Num(); BatchElementIndex++)
@@ -174,7 +174,7 @@ bool FTranslucencyForwardShadingDrawingPolicyFactory::DrawDynamicMesh(
 	bool bDirty = false;
 
 	// Determine the mesh's material and blend mode.
-	const FMaterial* Material = Mesh.MaterialRenderProxy->GetMaterial(GRHIFeatureLevel);
+	const FMaterial* Material = Mesh.MaterialRenderProxy->GetMaterial(View.GetFeatureLevel());
 	const EBlendMode BlendMode = Material->GetBlendMode();
 
 	// Only render translucent materials.
@@ -215,7 +215,7 @@ bool FTranslucencyForwardShadingDrawingPolicyFactory::DrawStaticMesh(
 	)
 {
 	bool bDirty = false;
-	const FMaterial* Material = StaticMesh.MaterialRenderProxy->GetMaterial(GRHIFeatureLevel);
+	const FMaterial* Material = StaticMesh.MaterialRenderProxy->GetMaterial(View.GetFeatureLevel());
 
 	bDirty |= DrawDynamicMesh(
 		View,
diff --git a/Engine/Source/Runtime/Renderer/Private/IndirectLightingCache.cpp b/Engine/Source/Runtime/Renderer/Private/IndirectLightingCache.cpp
--- a/Engine/Source/Runtime/Renderer/Private/IndirectLightingCache.cpp
+++ b/Engine/Source/Runtime/Renderer/Private/IndirectLightingCache.cpp
@@ -130,7 +130,7 @@ static FAutoConsoleVariableRef CVarLightingCacheUnbuiltPreviewAllocationSize(
 	ECVF_ReadOnly
 	);
 
-bool IsIndirectLightingCacheAllowed() 
+bool IsIndirectLightingCacheAllowed(ERHIFeatureLevel::Type InFeatureLevel) 
 {
 	static const auto AllowStaticLightingVar = IConsoleManager::Get().FindTConsoleVariableDataInt(TEXT("r.AllowStaticLighting"));
 	const bool bAllowStaticLighting = (!AllowStaticLightingVar || AllowStaticLightingVar->GetValueOnRenderThread() != 0);
@@ -138,7 +138,7 @@ bool IsIndirectLightingCacheAllowed()
 	return GIndirectLightingCache != 0 && bAllowStaticLighting;
 }
 
-bool CanIndirectLightingCacheUseVolumeTexture()
+bool CanIndirectLightingCacheUseVolumeTexture(ERHIFeatureLevel::Type InFeatureLevel)
 {
 	// @todo Mac OS X/OpenGL: For OpenGL devices which don't support volume-texture rendering we need to use the simpler point indirect lighting shaders.
 	return GRHIFeatureLevel >= ERHIFeatureLevel::SM3 && GSupportsVolumeTextureRendering;
@@ -153,7 +153,7 @@ FIndirectLightingCache::FIndirectLightingCache()
 
 void FIndirectLightingCache::InitDynamicRHI()
 {
-	if (CanIndirectLightingCacheUseVolumeTexture())
+	if (CanIndirectLightingCacheUseVolumeTexture(GRHIFeatureLevel))
 	{
 		uint32 Flags = TexCreate_ShaderResource | TexCreate_NoTiling;
 
@@ -344,7 +344,7 @@ FIndirectLightingCacheAllocation* FIndirectLightingCache::FindPrimitiveAllocatio
 
 void FIndirectLightingCache::UpdateCache(FScene* Scene, FSceneRenderer& Renderer, bool bAllowUnbuiltPreview)
 {
-	if (IsIndirectLightingCacheAllowed())
+	if (IsIndirectLightingCacheAllowed(Scene->GetFeatureLevel()))
 	{
 		bool bAnyViewAllowsIndirectLightingCache = false;
 
@@ -608,7 +608,7 @@ void FIndirectLightingCache::UpdateBlock(FScene* Scene, FViewInfo* DebugDrawingV
 	float DirectionalShadowing = 1;
 	FVector SkyBentNormal(0, 0, 1);
 
-	if (CanIndirectLightingCacheUseVolumeTexture() && BlockInfo.Allocation->bOpaqueRelevance)
+	if (CanIndirectLightingCacheUseVolumeTexture(Scene->GetFeatureLevel()) && BlockInfo.Allocation->bOpaqueRelevance)
 	{
 		static TArray<float> AccumulatedWeight;
 		AccumulatedWeight.Reset(NumSamplesPerBlock);
diff --git a/Engine/Source/Runtime/Renderer/Private/LightFunctionRendering.cpp b/Engine/Source/Runtime/Renderer/Private/LightFunctionRendering.cpp
--- a/Engine/Source/Runtime/Renderer/Private/LightFunctionRendering.cpp
+++ b/Engine/Source/Runtime/Renderer/Private/LightFunctionRendering.cpp
@@ -97,7 +97,7 @@ public:
 	{
 		const FPixelShaderRHIParamRef ShaderRHI = GetPixelShader();
 
-		FMaterialShader::SetParameters(ShaderRHI, MaterialProxy, *MaterialProxy->GetMaterial(GRHIFeatureLevel), *View, true, ESceneRenderTargetsMode::SetTextures);
+		FMaterialShader::SetParameters(ShaderRHI, MaterialProxy, *MaterialProxy->GetMaterial(View->GetFeatureLevel()), *View, true, ESceneRenderTargetsMode::SetTextures);
 
 		// Set the transform from screen space to light space.
 		if ( ScreenToLight.IsBound() )
@@ -184,7 +184,7 @@ bool FDeferredShadingSceneRenderer::CheckForLightFunction( const FLightSceneInfo
 {
 	// NOTE: The extra check is necessary because there could be something wrong with the material.
 	if( LightSceneInfo->Proxy->GetLightFunctionMaterial() && 
-		LightSceneInfo->Proxy->GetLightFunctionMaterial()->GetMaterial(GRHIFeatureLevel)->IsLightFunction() )
+		LightSceneInfo->Proxy->GetLightFunctionMaterial()->GetMaterial(Scene->GetFeatureLevel())->IsLightFunction())
 	{
 		FSphere LightBounds = LightSceneInfo->Proxy->GetBoundingSphere();
 		for (int32 ViewIndex = 0;ViewIndex < Views.Num();ViewIndex++)
@@ -236,13 +236,13 @@ bool FDeferredShadingSceneRenderer::RenderLightFunctionForMaterial(const FLightS
 {
 	bool bRenderedLightFunction = false;
 
-	if (MaterialProxy && MaterialProxy->GetMaterial(GRHIFeatureLevel)->IsLightFunction())
+	if (MaterialProxy && MaterialProxy->GetMaterial(Scene->GetFeatureLevel())->IsLightFunction())
 	{
 		GSceneRenderTargets.BeginRenderingLightAttenuation();
 
 		bRenderedLightFunction = true;
 
-		const FMaterial* Material = MaterialProxy->GetMaterial(GRHIFeatureLevel);
+		const FMaterial* Material = MaterialProxy->GetMaterial(Scene->GetFeatureLevel());
 		SCOPED_DRAW_EVENTF(LightFunction, DEC_SCENE_ITEMS, TEXT("LightFunction Material=%s"), *Material->GetFriendlyName());
 
 		const FMaterialShaderMap* MaterialShaderMap = Material->GetRenderingThreadShaderMap();
diff --git a/Engine/Source/Runtime/Renderer/Private/LightMapDensityRendering.cpp b/Engine/Source/Runtime/Renderer/Private/LightMapDensityRendering.cpp
--- a/Engine/Source/Runtime/Renderer/Private/LightMapDensityRendering.cpp
+++ b/Engine/Source/Runtime/Renderer/Private/LightMapDensityRendering.cpp
@@ -38,7 +38,7 @@ bool FDeferredShadingSceneRenderer::RenderLightMapDensities()
 {
 	bool bDirty=0;
 
-	if (GRHIFeatureLevel >= ERHIFeatureLevel::SM3)
+	if (Scene->GetFeatureLevel() >= ERHIFeatureLevel::SM3)
 	{
 		SCOPED_DRAW_EVENT(LightMapDensity, DEC_SCENE_ITEMS);
 
@@ -104,7 +104,7 @@ bool FLightMapDensityDrawingPolicyFactory::DrawDynamicMesh(
 	bool bDirty = false;
 
 	const FMaterialRenderProxy* MaterialRenderProxy = Mesh.MaterialRenderProxy;
-	const FMaterial* Material = MaterialRenderProxy->GetMaterial(GRHIFeatureLevel);
+	const FMaterial* Material = MaterialRenderProxy->GetMaterial(View.GetFeatureLevel());
 	const EBlendMode BlendMode = Material->GetBlendMode();
 
 	const bool bMaterialMasked = Material->IsMasked();
@@ -135,8 +135,8 @@ bool FLightMapDensityDrawingPolicyFactory::DrawDynamicMesh(
 		{
 			if (bAllowHighQualityLightMaps)
 			{
-				TLightMapDensityDrawingPolicy< TLightMapPolicy<HQ_LIGHTMAP> > DrawingPolicy( Mesh.VertexFactory, MaterialRenderProxy, TLightMapPolicy<HQ_LIGHTMAP>(), BlendMode );
-				DrawingPolicy.DrawShared(&View,DrawingPolicy.CreateBoundShaderState());
+				TLightMapDensityDrawingPolicy< TLightMapPolicy<HQ_LIGHTMAP> > DrawingPolicy(View, Mesh.VertexFactory, MaterialRenderProxy, TLightMapPolicy<HQ_LIGHTMAP>(), BlendMode);
+				DrawingPolicy.DrawShared(&View,DrawingPolicy.CreateBoundShaderState(View.GetFeatureLevel()));
 				for( int32 BatchElementIndex = 0; BatchElementIndex < Mesh.Elements.Num(); BatchElementIndex++)
 				{
 					DrawingPolicy.SetMeshRenderState(View,PrimitiveSceneProxy,Mesh,BatchElementIndex,bBackFace,
@@ -147,8 +147,8 @@ bool FLightMapDensityDrawingPolicyFactory::DrawDynamicMesh(
 			}
 			else
 			{
-				TLightMapDensityDrawingPolicy< TLightMapPolicy<LQ_LIGHTMAP> > DrawingPolicy( Mesh.VertexFactory, MaterialRenderProxy, TLightMapPolicy<LQ_LIGHTMAP>(), BlendMode );
-				DrawingPolicy.DrawShared(&View,DrawingPolicy.CreateBoundShaderState());
+				TLightMapDensityDrawingPolicy< TLightMapPolicy<LQ_LIGHTMAP> > DrawingPolicy(View, Mesh.VertexFactory, MaterialRenderProxy, TLightMapPolicy<LQ_LIGHTMAP>(), BlendMode);
+				DrawingPolicy.DrawShared(&View, DrawingPolicy.CreateBoundShaderState(View.GetFeatureLevel()));
 				for( int32 BatchElementIndex = 0; BatchElementIndex < Mesh.Elements.Num(); BatchElementIndex++)
 				{
 					DrawingPolicy.SetMeshRenderState(View,PrimitiveSceneProxy,Mesh,BatchElementIndex,bBackFace,
@@ -160,8 +160,8 @@ bool FLightMapDensityDrawingPolicyFactory::DrawDynamicMesh(
 		}
 		else
 		{
-			TLightMapDensityDrawingPolicy<FDummyLightMapPolicy> DrawingPolicy(Mesh.VertexFactory,MaterialRenderProxy,FDummyLightMapPolicy(),BlendMode);
-			DrawingPolicy.DrawShared(&View,DrawingPolicy.CreateBoundShaderState());
+			TLightMapDensityDrawingPolicy<FDummyLightMapPolicy> DrawingPolicy(View, Mesh.VertexFactory, MaterialRenderProxy, FDummyLightMapPolicy(), BlendMode);
+			DrawingPolicy.DrawShared(&View, DrawingPolicy.CreateBoundShaderState(View.GetFeatureLevel()));
 			for( int32 BatchElementIndex = 0; BatchElementIndex < Mesh.Elements.Num(); BatchElementIndex++)
 			{
 				DrawingPolicy.SetMeshRenderState(View,PrimitiveSceneProxy,Mesh,BatchElementIndex,bBackFace,
@@ -173,8 +173,8 @@ bool FLightMapDensityDrawingPolicyFactory::DrawDynamicMesh(
 	}
 	else
 	{
-		TLightMapDensityDrawingPolicy<FNoLightMapPolicy> DrawingPolicy(Mesh.VertexFactory,MaterialRenderProxy,FNoLightMapPolicy(),BlendMode);
-	 	DrawingPolicy.DrawShared(&View,DrawingPolicy.CreateBoundShaderState());
+		TLightMapDensityDrawingPolicy<FNoLightMapPolicy> DrawingPolicy(View, Mesh.VertexFactory, MaterialRenderProxy, FNoLightMapPolicy(), BlendMode);
+		DrawingPolicy.DrawShared(&View, DrawingPolicy.CreateBoundShaderState(View.GetFeatureLevel()));
 		for( int32 BatchElementIndex = 0; BatchElementIndex < Mesh.Elements.Num(); BatchElementIndex++)
 		{
 			DrawingPolicy.SetMeshRenderState(View,PrimitiveSceneProxy,Mesh,BatchElementIndex,bBackFace,TLightMapDensityDrawingPolicy<FNoLightMapPolicy>::ElementDataType());
diff --git a/Engine/Source/Runtime/Renderer/Private/LightMapDensityRendering.h b/Engine/Source/Runtime/Renderer/Private/LightMapDensityRendering.h
--- a/Engine/Source/Runtime/Renderer/Private/LightMapDensityRendering.h
+++ b/Engine/Source/Runtime/Renderer/Private/LightMapDensityRendering.h
@@ -49,7 +49,7 @@ public:
 		const FSceneView& View
 		)
 	{
-		FMeshMaterialShader::SetParameters(GetVertexShader(),MaterialRenderProxy,*MaterialRenderProxy->GetMaterial(GRHIFeatureLevel),View,ESceneRenderTargetsMode::SetTextures);
+		FMeshMaterialShader::SetParameters(GetVertexShader(), MaterialRenderProxy, *MaterialRenderProxy->GetMaterial(View.GetFeatureLevel()), View, ESceneRenderTargetsMode::SetTextures);
 	}
 
 	void SetMesh(const FVertexFactory* VertexFactory,const FSceneView& View,const FPrimitiveSceneProxy* Proxy,const FMeshBatchElement& BatchElement)
@@ -160,7 +160,7 @@ public:
 
 	void SetParameters(const FMaterialRenderProxy* MaterialRenderProxy,const FSceneView* View)
 	{
-		FMeshMaterialShader::SetParameters(GetPixelShader(),MaterialRenderProxy,*MaterialRenderProxy->GetMaterial(GRHIFeatureLevel),*View,ESceneRenderTargetsMode::SetTextures);
+		FMeshMaterialShader::SetParameters(GetPixelShader(), MaterialRenderProxy, *MaterialRenderProxy->GetMaterial(View->GetFeatureLevel()), *View, ESceneRenderTargetsMode::SetTextures);
 
 		if (GridTexture.IsBound())
 		{
@@ -257,12 +257,13 @@ public:
 
 	/** Initialization constructor. */
 	TLightMapDensityDrawingPolicy(
+		const FViewInfo& View,
 		const FVertexFactory* InVertexFactory,
 		const FMaterialRenderProxy* InMaterialRenderProxy,
 		LightMapPolicyType InLightMapPolicy,
 		EBlendMode InBlendMode
 		):
-		FMeshDrawingPolicy(InVertexFactory,InMaterialRenderProxy, *InMaterialRenderProxy->GetMaterial(GRHIFeatureLevel)),
+		FMeshDrawingPolicy(InVertexFactory, InMaterialRenderProxy, *InMaterialRenderProxy->GetMaterial(View.GetFeatureLevel())),
 		LightMapPolicy(InLightMapPolicy),
 		BlendMode(InBlendMode)
 	{
@@ -395,7 +396,7 @@ public:
 	 * as well as the shaders needed to draw the mesh
 	 * @return new bound shader state object
 	 */
-	FBoundShaderStateRHIRef CreateBoundShaderState()
+	FBoundShaderStateRHIRef CreateBoundShaderState(ERHIFeatureLevel::Type InFeatureLevel)
 	{
 		FBoundShaderStateRHIRef BoundShaderState;
 
@@ -449,10 +450,10 @@ public:
 		const FPrimitiveSceneProxy* PrimitiveSceneProxy,
 		FHitProxyId HitProxyId
 		);
-	static bool IsMaterialIgnored(const FMaterialRenderProxy* MaterialRenderProxy)
+	static bool IsMaterialIgnored(const FMaterialRenderProxy* MaterialRenderProxy, ERHIFeatureLevel::Type InFeatureLevel)
 	{
 		// Note: MaterialModifiesMeshPosition may depend on feature level!
-		const FMaterial* Material = MaterialRenderProxy ? MaterialRenderProxy->GetMaterial(GRHIFeatureLevel) : NULL;
+		const FMaterial* Material = MaterialRenderProxy ? MaterialRenderProxy->GetMaterial(InFeatureLevel) : NULL;
 		return (Material && 
 				!(	Material->IsSpecialEngineMaterial() || 
 					Material->IsMasked() ||
diff --git a/Engine/Source/Runtime/Renderer/Private/LightPropagationVolume.cpp b/Engine/Source/Runtime/Renderer/Private/LightPropagationVolume.cpp
--- a/Engine/Source/Runtime/Renderer/Private/LightPropagationVolume.cpp
+++ b/Engine/Source/Runtime/Renderer/Private/LightPropagationVolume.cpp
@@ -792,6 +792,8 @@ FLightPropagationVolume::~FLightPropagationVolume()
 */
 void FLightPropagationVolume::InitSettings( const FSceneView& View )
 	{
+	check(View.GetFeatureLevel() >= ERHIFeatureLevel::SM5);
+
 #if LPV_VOLUME_TEXTURE
 	if ( !bInitialized )
 	{
@@ -927,7 +929,6 @@ void FLightPropagationVolume::Clear()
 	}
 
 	SCOPED_DRAW_EVENT(LpvClear, DEC_LIGHT);
-	check(GRHIFeatureLevel == ERHIFeatureLevel::SM5);
 
 	if ( !LpvWriteUniformBuffer.IsInitialized() )
 	{
@@ -1040,7 +1041,6 @@ void FLightPropagationVolume::InjectDirectionalLightRSM(
 	const FLightSceneProxy* LightProxy = ProjectedShadowInfo.LightSceneInfo->Proxy;
 	{
 		SCOPED_DRAW_EVENT(LpvInjectDirectionalLightRSM, DEC_LIGHT);
-		check(GRHIFeatureLevel == ERHIFeatureLevel::SM5);
 
 		SetVplInjectionConstants( ProjectedShadowInfo, LightProxy );
 
@@ -1096,9 +1096,6 @@ void FLightPropagationVolume::Propagate()
 		return;
 	}
 
-	check(GRHIFeatureLevel == ERHIFeatureLevel::SM5);
-
-
 	// Inject the VPLs into the LPV
 	if ( mInjectedLightCount )
 	{
@@ -1205,8 +1202,6 @@ void FLightPropagationVolume::InjectLightDirect( const FLightSceneProxy& Light )
 		return;
 	}
 
-	check(GRHIFeatureLevel == ERHIFeatureLevel::SM5);
-
 	// A geometry volume is required for direct light injection. This currently requires a directional light to be injected
 	//@TODO: Add support for generating a GV when there's no directional light
 	if ( GeometryVolumeGenerated )
diff --git a/Engine/Source/Runtime/Renderer/Private/LightPropagationVolumeVisualisation.cpp b/Engine/Source/Runtime/Renderer/Private/LightPropagationVolumeVisualisation.cpp
--- a/Engine/Source/Runtime/Renderer/Private/LightPropagationVolumeVisualisation.cpp
+++ b/Engine/Source/Runtime/Renderer/Private/LightPropagationVolumeVisualisation.cpp
@@ -247,7 +247,7 @@ IMPLEMENT_SHADER_TYPE(,FLpvVisualisePS,TEXT("LPVVisualise"),TEXT("PShader"),SF_P
 void FLightPropagationVolume::Visualise( const FSceneView& View ) const
 {
 	SCOPED_DRAW_EVENT(LpvVisualise, DEC_LIGHT);
-	check(GRHIFeatureLevel == ERHIFeatureLevel::SM5);
+	check(View.GetFeatureLevel() == ERHIFeatureLevel::SM5);
 
 	TShaderMapRef<FLpvVisualiseVS> VertexShader(GetGlobalShaderMap());
 	TShaderMapRef<FLpvVisualiseGS> GeometryShader(GetGlobalShaderMap());
diff --git a/Engine/Source/Runtime/Renderer/Private/LightShaftRendering.cpp b/Engine/Source/Runtime/Renderer/Private/LightShaftRendering.cpp
--- a/Engine/Source/Runtime/Renderer/Private/LightShaftRendering.cpp
+++ b/Engine/Source/Runtime/Renderer/Private/LightShaftRendering.cpp
@@ -299,7 +299,7 @@ public:
 		const FIntPoint BufferSize = GSceneRenderTargets.GetBufferSizeXY();
 		FVector2D SampleOffsets(1.0f / BufferSize.X, 1.0f / BufferSize.Y);
 		SetShaderValue(GetPixelShader(),SampleOffsetsParameter,SampleOffsets);
-		SceneTextureParams.Set(GetPixelShader());
+		SceneTextureParams.Set(GetPixelShader(), View);
 	}
 
 private:
diff --git a/Engine/Source/Runtime/Renderer/Private/PostProcess/PostProcessDownsample.cpp b/Engine/Source/Runtime/Renderer/Private/PostProcess/PostProcessDownsample.cpp
--- a/Engine/Source/Runtime/Renderer/Private/PostProcess/PostProcessDownsample.cpp
+++ b/Engine/Source/Runtime/Renderer/Private/PostProcess/PostProcessDownsample.cpp
@@ -212,7 +212,7 @@ void FRCPassPostProcessDownsample::Process(FRenderingCompositePassContext& Conte
 
 	// check if we have to clear the whole surface.
 	// Otherwise perform the clear when the dest rectangle has been computed.
-	if (bHasMultipleQuads || GRHIFeatureLevel == ERHIFeatureLevel::ES2)
+	if (bHasMultipleQuads || Context.View.GetFeatureLevel() == ERHIFeatureLevel::ES2)
 	{
 		RHIClear(true, FLinearColor(0, 0, 0, 0), false, 1.0f, false, 0, FIntRect());
 		bHasCleared = true;
@@ -248,7 +248,7 @@ void FRCPassPostProcessDownsample::Process(FRenderingCompositePassContext& Conte
 		case EPostProcessRectSource::GBS_UIBlurRects:
 		{
 			static auto* ICVar = IConsoleManager::Get().FindTConsoleVariableDataFloat(TEXT("UI.BlurRadius"));
-			int32 IntegerKernelRadius = FRCPassPostProcessWeightedSampleSum::GetIntegerKernelRadius(ICVar->GetValueOnRenderThread());
+			int32 IntegerKernelRadius = FRCPassPostProcessWeightedSampleSum::GetIntegerKernelRadius(Context.View.GetFeatureLevel(), ICVar->GetValueOnRenderThread());
 
 			// Add Downsample's own inflation (PS linear+HW filter) to maximum required by Gaussian blur
 			InflateSize += IntegerKernelRadius;
diff --git a/Engine/Source/Runtime/Renderer/Private/PostProcess/PostProcessEyeAdaptation.cpp b/Engine/Source/Runtime/Renderer/Private/PostProcess/PostProcessEyeAdaptation.cpp
--- a/Engine/Source/Runtime/Renderer/Private/PostProcess/PostProcessEyeAdaptation.cpp
+++ b/Engine/Source/Runtime/Renderer/Private/PostProcess/PostProcessEyeAdaptation.cpp
@@ -126,7 +126,7 @@ void FRCPassPostProcessEyeAdaptation::ComputeEyeAdaptationParamsValue(const FVie
 	float LocalOverrideExposure = FLT_MAX;
 
 	// Eye adaptation is disabled except for highend right now because the histogram is not computed.
-	if(!EngineShowFlags.EyeAdaptation || GRHIFeatureLevel != ERHIFeatureLevel::SM5)
+	if (!EngineShowFlags.EyeAdaptation || View.GetFeatureLevel() != ERHIFeatureLevel::SM5)
 	{
 		LocalOverrideExposure = 0;
 	}
diff --git a/Engine/Source/Runtime/Renderer/Private/PostProcess/PostProcessMaterial.cpp b/Engine/Source/Runtime/Renderer/Private/PostProcess/PostProcessMaterial.cpp
--- a/Engine/Source/Runtime/Renderer/Private/PostProcess/PostProcessMaterial.cpp
+++ b/Engine/Source/Runtime/Renderer/Private/PostProcess/PostProcessMaterial.cpp
@@ -74,7 +74,7 @@ public:
 	{
 		const FPixelShaderRHIParamRef ShaderRHI = GetPixelShader();
 
-		FMaterialShader::SetParameters(ShaderRHI, Material, *Material->GetMaterial(GRHIFeatureLevel), Context.View, true, ESceneRenderTargetsMode::SetTextures);
+		FMaterialShader::SetParameters(ShaderRHI, Material, *Material->GetMaterial(Context.View.GetFeatureLevel()), Context.View, true, ESceneRenderTargetsMode::SetTextures);
 		PostprocessParameter.SetPS(ShaderRHI, Context, TStaticSamplerState<SF_Point,AM_Clamp,AM_Clamp,AM_Clamp>::GetRHI());
 	}
 
@@ -102,7 +102,7 @@ void FRCPassPostProcessMaterial::Process(FRenderingCompositePassContext& Context
 
 	check(Proxy);
 
-	const FMaterial* Material = Proxy->GetMaterial(GRHIFeatureLevel);
+	const FMaterial* Material = Proxy->GetMaterial(Context.View.GetFeatureLevel());
 	
 	check(Material);
 
diff --git a/Engine/Source/Runtime/Renderer/Private/PostProcess/PostProcessWeightedSampleSum.cpp b/Engine/Source/Runtime/Renderer/Private/PostProcess/PostProcessWeightedSampleSum.cpp
--- a/Engine/Source/Runtime/Renderer/Private/PostProcess/PostProcessWeightedSampleSum.cpp
+++ b/Engine/Source/Runtime/Renderer/Private/PostProcess/PostProcessWeightedSampleSum.cpp
@@ -269,10 +269,10 @@ static float NormalDistributionUnscaled(float X,float Mean,float Variance)
  * @return NumSamples >0
  */
 
-static uint32 Compute1DGaussianFilterKernel(float KernelRadius, FVector2D OutOffsetAndWeight[MAX_FILTER_SAMPLES], uint32 MaxFilterSamples)
+static uint32 Compute1DGaussianFilterKernel(ERHIFeatureLevel::Type InFeatureLevel, float KernelRadius, FVector2D OutOffsetAndWeight[MAX_FILTER_SAMPLES], uint32 MaxFilterSamples)
 {
-	float ClampedKernelRadius = FRCPassPostProcessWeightedSampleSum::GetClampedKernelRadius( KernelRadius );
-	int32 IntegerKernelRadius = FRCPassPostProcessWeightedSampleSum::GetIntegerKernelRadius( KernelRadius );
+	float ClampedKernelRadius = FRCPassPostProcessWeightedSampleSum::GetClampedKernelRadius( InFeatureLevel, KernelRadius );
+	int32 IntegerKernelRadius = FRCPassPostProcessWeightedSampleSum::GetIntegerKernelRadius( InFeatureLevel, KernelRadius );
 
 	// smallest IntegerKernelRadius will be 1
 
@@ -362,10 +362,12 @@ void FRCPassPostProcessWeightedSampleSum::Process(FRenderingCompositePassContext
 	FLinearColor BlurWeights[MAX_FILTER_SAMPLES];
 	FVector2D OffsetAndWeight[MAX_FILTER_SAMPLES];
 
+	const auto FeatureLevel = Context.View.GetFeatureLevel();
+
 	// compute 1D filtered samples
-	uint32 MaxNumSamples = GetMaxNumSamples();
+	uint32 MaxNumSamples = GetMaxNumSamples(FeatureLevel);
 
-	uint32 NumSamples = Compute1DGaussianFilterKernel(EffectiveBlurRadius, OffsetAndWeight, MaxNumSamples);
+	uint32 NumSamples = Compute1DGaussianFilterKernel(FeatureLevel, EffectiveBlurRadius, OffsetAndWeight, MaxNumSamples);
 
 	// compute weights as weighted contributions of the TintValue
 	for(uint32 i = 0; i < NumSamples; ++i)
@@ -442,7 +444,7 @@ void FRCPassPostProcessWeightedSampleSum::Process(FRenderingCompositePassContext
 	bool bRequiresClear = true;
 	// check if we have to clear the whole surface.
 	// Otherwise perform the clear when the dest rectangle has been computed.
-	if (bHasMultipleQuads || GRHIFeatureLevel == ERHIFeatureLevel::ES2)
+	if (bHasMultipleQuads || Context.View.GetFeatureLevel() == ERHIFeatureLevel::ES2)
 	{
 		RHIClear(true, FLinearColor(0, 0, 0, 0), false, 1.0f, false, 0, FIntRect());
 		bRequiresClear = false;
@@ -592,14 +594,14 @@ void FRCPassPostProcessWeightedSampleSum::DrawQuad( bool bDoFastBlur, FIntRect S
 		EDRF_UseTriangleOptimization);
 }
 
-uint32 FRCPassPostProcessWeightedSampleSum::GetMaxNumSamples()
+uint32 FRCPassPostProcessWeightedSampleSum::GetMaxNumSamples(ERHIFeatureLevel::Type InFeatureLevel)
 {
 	uint32 MaxNumSamples;
-	if (GRHIFeatureLevel >= ERHIFeatureLevel::SM5)
+	if (InFeatureLevel >= ERHIFeatureLevel::SM5)
 	{
 		MaxNumSamples = MAX_FILTER_SAMPLES;
 	}
-	else if (GRHIFeatureLevel >= ERHIFeatureLevel::SM3)
+	else if (InFeatureLevel >= ERHIFeatureLevel::SM3)
 	{
 		MaxNumSamples = 16;
 	}
@@ -610,12 +612,12 @@ uint32 FRCPassPostProcessWeightedSampleSum::GetMaxNumSamples()
 	return MaxNumSamples;
 }
 
-float FRCPassPostProcessWeightedSampleSum::GetClampedKernelRadius( float KernelRadius )
+float FRCPassPostProcessWeightedSampleSum::GetClampedKernelRadius(ERHIFeatureLevel::Type InFeatureLevel, float KernelRadius)
 {
-	return FMath::Clamp<float>(KernelRadius, DELTA, GetMaxNumSamples() - 1);
+	return FMath::Clamp<float>(KernelRadius, DELTA, GetMaxNumSamples(InFeatureLevel) - 1);
 }
 
-int FRCPassPostProcessWeightedSampleSum::GetIntegerKernelRadius( float KernelRadius )
+int FRCPassPostProcessWeightedSampleSum::GetIntegerKernelRadius(ERHIFeatureLevel::Type InFeatureLevel, float KernelRadius)
 {
-	return FMath::Min<int32>(FMath::CeilToInt(GetClampedKernelRadius(KernelRadius)), GetMaxNumSamples() - 1);
+	return FMath::Min<int32>(FMath::CeilToInt(GetClampedKernelRadius(InFeatureLevel, KernelRadius)), GetMaxNumSamples(InFeatureLevel) - 1);
 }
diff --git a/Engine/Source/Runtime/Renderer/Private/PostProcess/PostProcessWeightedSampleSum.h b/Engine/Source/Runtime/Renderer/Private/PostProcess/PostProcessWeightedSampleSum.h
--- a/Engine/Source/Runtime/Renderer/Private/PostProcess/PostProcessWeightedSampleSum.h
+++ b/Engine/Source/Runtime/Renderer/Private/PostProcess/PostProcessWeightedSampleSum.h
@@ -120,12 +120,12 @@ public:
 	virtual FPooledRenderTargetDesc ComputeOutputDesc(EPassOutputId InPassOutputId) const;
 
 	// retrieve runtime filter kernel properties.
-	static float GetClampedKernelRadius(float KernelRadius);
-	static int GetIntegerKernelRadius(float KernelRadius);
+	static float GetClampedKernelRadius(ERHIFeatureLevel::Type InFeatureLevel, float KernelRadius);
+	static int GetIntegerKernelRadius(ERHIFeatureLevel::Type InFeatureLevel, float KernelRadius);
 
 private:
 	void DrawQuad(bool bDoFastBlur, FIntRect SrcRect, FIntRect DestRect, bool bRequiresClear, FIntPoint DestSize, FIntPoint SrcSize, FShader* VertexShader) const;
-	static uint32 GetMaxNumSamples();
+	static uint32 GetMaxNumSamples(ERHIFeatureLevel::Type InFeatureLevel);
 
 	// e.g. EFS_Horiz or EFS_Vert
 	EFilterShape FilterShape;
diff --git a/Engine/Source/Runtime/Renderer/Private/PostProcess/PostProcessing.cpp b/Engine/Source/Runtime/Renderer/Private/PostProcess/PostProcessing.cpp
--- a/Engine/Source/Runtime/Renderer/Private/PostProcess/PostProcessing.cpp
+++ b/Engine/Source/Runtime/Renderer/Private/PostProcess/PostProcessing.cpp
@@ -709,6 +709,8 @@ void FPostProcessing::Process(const FViewInfo& View, TRefCountPtr<IPooledRenderT
 
 	bool bSimpleDynamicLighting = IsSimpleDynamicLightingEnabled();
 
+	const auto FeatureLevel = View.GetFeatureLevel();
+
 	GRenderTargetPool.AddPhaseEvent(TEXT("PostProcessing"));
 
 	// This page: https://udn.epicgames.com/Three/RenderingOverview#Rendering%20state%20defaults 
@@ -755,7 +757,7 @@ void FPostProcessing::Process(const FViewInfo& View, TRefCountPtr<IPooledRenderT
 		
 		// add the passes we want to add to the graph (commenting a line means the pass is not inserted into the graph) ---------
 
-		if( View.Family->EngineShowFlags.PostProcessing && GRHIFeatureLevel >= ERHIFeatureLevel::SM4 )
+		if (View.Family->EngineShowFlags.PostProcessing && FeatureLevel >= ERHIFeatureLevel::SM4)
 		{
 			// Screen Space Subsurface Scattering
 			{
@@ -943,7 +945,7 @@ void FPostProcessing::Process(const FViewInfo& View, TRefCountPtr<IPooledRenderT
 					bHistogramNeeded = false;
 				}
 
-				if(!GIsHighResScreenshot && bHistogramNeeded && GRHIFeatureLevel == ERHIFeatureLevel::SM5 && StereoPass != eSSP_RIGHT_EYE)
+				if (!GIsHighResScreenshot && bHistogramNeeded && FeatureLevel >= ERHIFeatureLevel::SM5 && StereoPass != eSSP_RIGHT_EYE)
 				{
 					Histogram = AddPostProcessHistogram(Context, SceneColorHalfRes);
 				}
@@ -968,7 +970,7 @@ void FPostProcessing::Process(const FViewInfo& View, TRefCountPtr<IPooledRenderT
 
 			HDRColor = Context.FinalOutput;
 
-			if(bAllowTonemapper && GRHIFeatureLevel >= ERHIFeatureLevel::SM4)
+			if(bAllowTonemapper && FeatureLevel >= ERHIFeatureLevel::SM4)
 			{
 				FRenderingCompositeOutputRef CombinedLUT;
 
@@ -1087,7 +1089,7 @@ void FPostProcessing::Process(const FViewInfo& View, TRefCountPtr<IPooledRenderT
 			}
 		}
 
-		if(View.Family->EngineShowFlags.VisualizeHDR && GRHIFeatureLevel >= ERHIFeatureLevel::SM5)
+		if(View.Family->EngineShowFlags.VisualizeHDR && FeatureLevel >= ERHIFeatureLevel::SM5)
 		{
 			FRenderingCompositePass* Node = Context.Graph.RegisterPass(new(FMemStack::Get()) FRCPassPostProcessVisualizeHDR());
 			Node->SetInput(ePId_Input0, FRenderingCompositeOutputRef(Context.FinalOutput));
@@ -1098,7 +1100,7 @@ void FPostProcessing::Process(const FViewInfo& View, TRefCountPtr<IPooledRenderT
 
 		AddUIBlur(Context);
 
-		if(View.Family->EngineShowFlags.TestImage && GRHIFeatureLevel >= ERHIFeatureLevel::SM5)
+		if(View.Family->EngineShowFlags.TestImage && FeatureLevel >= ERHIFeatureLevel::SM5)
 		{
 			FRenderingCompositePass* Node = Context.Graph.RegisterPass(new(FMemStack::Get()) FRCPassPostProcessTestImage());
 			Node->SetInput(ePId_Input0, FRenderingCompositeOutputRef(Context.FinalOutput));
diff --git a/Engine/Source/Runtime/Renderer/Private/PostProcess/RenderingCompositionGraph.cpp b/Engine/Source/Runtime/Renderer/Private/PostProcess/RenderingCompositionGraph.cpp
--- a/Engine/Source/Runtime/Renderer/Private/PostProcess/RenderingCompositionGraph.cpp
+++ b/Engine/Source/Runtime/Renderer/Private/PostProcess/RenderingCompositionGraph.cpp
@@ -121,6 +121,11 @@ FRenderingCompositePassContext::~FRenderingCompositePassContext()
 	Graph.Free();
 }
 
+ERHIFeatureLevel::Type FRenderingCompositePassContext::GetFeatureLevel() const
+{
+	return View.GetFeatureLevel();
+}
+
 void FRenderingCompositePassContext::Process(const TCHAR *GraphDebugName)
 {
 	if(Root)
diff --git a/Engine/Source/Runtime/Renderer/Private/PostProcess/RenderingCompositionGraph.h b/Engine/Source/Runtime/Renderer/Private/PostProcess/RenderingCompositionGraph.h
--- a/Engine/Source/Runtime/Renderer/Private/PostProcess/RenderingCompositionGraph.h
+++ b/Engine/Source/Runtime/Renderer/Private/PostProcess/RenderingCompositionGraph.h
@@ -170,6 +170,8 @@ struct FRenderingCompositePassContext
 		return ViewPortRect.Min != ViewPortRect.Max;
 	}
 
+	ERHIFeatureLevel::Type GetFeatureLevel() const;
+
 	FRenderingCompositePass* Root;
 
 	FRenderingCompositionGraph Graph;
@@ -625,6 +627,8 @@ public:
 				break;
 			}
 
+			const auto FeatureLevel = Context.GetFeatureLevel();
+
 			FRenderingCompositeOutput* Input = OutputRef->GetOutput();
 
 			TRefCountPtr<IPooledRenderTarget> InputPooledElement;
diff --git a/Engine/Source/Runtime/Renderer/Private/PostProcess/SceneRenderTargets.h b/Engine/Source/Runtime/Renderer/Private/PostProcess/SceneRenderTargets.h
--- a/Engine/Source/Runtime/Renderer/Private/PostProcess/SceneRenderTargets.h
+++ b/Engine/Source/Runtime/Renderer/Private/PostProcess/SceneRenderTargets.h
@@ -530,7 +530,7 @@ public:
 
 	/** Sets the scene texture parameters for the given view. */
 	template<typename TParamRef>
-	void Set(const TParamRef& ShaderRHI, ESceneRenderTargetsMode::Type TextureMode = ESceneRenderTargetsMode::SetTextures, ESamplerFilter ColorFilter = SF_Point) const
+	void Set(const TParamRef& ShaderRHI, const FSceneView& View, ESceneRenderTargetsMode::Type TextureMode = ESceneRenderTargetsMode::SetTextures, ESamplerFilter ColorFilter = SF_Point) const
 	{
 		if (TextureMode == ESceneRenderTargetsMode::SetTextures)
 		{
@@ -596,7 +596,7 @@ public:
 			{
 				SetTextureParameter(ShaderRHI, SceneColorSurfaceParameter, GSceneRenderTargets.GetSceneColorSurface());
 			}
-			if(GRHIFeatureLevel >= ERHIFeatureLevel::SM4)
+			if (GRHIFeatureLevel >= ERHIFeatureLevel::SM4)
 			{
 				if(GSupportsDepthFetchDuringDepthTest)
 				{
@@ -684,7 +684,7 @@ public:
 	void Set(const ShaderRHIParamRef ShaderRHI, const FSceneView& View, ESceneRenderTargetsMode::Type TextureMode = ESceneRenderTargetsMode::SetTextures) const
 	{
 		// This is needed on PC ES2 for SceneAlphaCopy, probably should be refactored for performance.
-		SceneTextureParameters.Set(ShaderRHI, TextureMode, SF_Point);
+		SceneTextureParameters.Set(ShaderRHI, View, TextureMode, SF_Point);
 
 		// if() is purely an optimization and could be removed
 		if(IsDBufferEnabled())
diff --git a/Engine/Source/Runtime/Renderer/Private/PostProcess/VisualizeTexture.cpp b/Engine/Source/Runtime/Renderer/Private/PostProcess/VisualizeTexture.cpp
--- a/Engine/Source/Runtime/Renderer/Private/PostProcess/VisualizeTexture.cpp
+++ b/Engine/Source/Runtime/Renderer/Private/PostProcess/VisualizeTexture.cpp
@@ -267,6 +267,7 @@ FVisualizeTexture::FVisualizeTexture()
 	bOutputStencil = false;
 	bFullList = false;
 	SortOrder = -1;
+	bEnabled = true;
 }
 
 FIntRect FVisualizeTexture::ComputeVisualizeTextureRect(FIntPoint InputTextureSize) const
@@ -482,7 +483,7 @@ void FVisualizeTexture::PresentContent(const FSceneView& View)
 
 	if(!VisualizeTextureContent 
 		|| !IsValidRef(RenderTargetTexture)
-		|| GRHIFeatureLevel < ERHIFeatureLevel::SM3)
+		|| !bEnabled)
 	{
 		// visualize feature is deactivated
 		return;
@@ -617,7 +618,7 @@ void FVisualizeTexture::SetObserveTarget(const FString& InObservedDebugName, uin
 void FVisualizeTexture::SetCheckPoint(const IPooledRenderTarget* PooledRenderTarget)
 {
 #if !(UE_BUILD_SHIPPING || UE_BUILD_TEST)
-	if (!PooledRenderTarget || GRHIFeatureLevel == ERHIFeatureLevel::ES2)
+	if (!PooledRenderTarget || !bEnabled)
 	{
 		// Don't checkpoint on ES2 to avoid TMap alloc/reallocations
 		return;
@@ -873,6 +874,7 @@ IPooledRenderTarget* FVisualizeTexture::GetObservedElement() const
 
 void FVisualizeTexture::OnStartFrame(const FSceneView& View)
 {
+	bEnabled = View.GetFeatureLevel() >= ERHIFeatureLevel::SM3;
 	ViewRect = View.UnscaledViewRect;
 	AspectRatioConstrainedViewRect = View.Family->EngineShowFlags.CameraAspectRatioBars ? View.CameraConstrainedViewRect : ViewRect;
 
diff --git a/Engine/Source/Runtime/Renderer/Private/PostProcess/VisualizeTexture.h b/Engine/Source/Runtime/Renderer/Private/PostProcess/VisualizeTexture.h
--- a/Engine/Source/Runtime/Renderer/Private/PostProcess/VisualizeTexture.h
+++ b/Engine/Source/Runtime/Renderer/Private/PostProcess/VisualizeTexture.h
@@ -83,6 +83,9 @@ private:
 	// the texture visualization image inside this area
 	FIntRect AspectRatioConstrainedViewRect;
 
+	// Flag to determine whether texture visualization is enabled, currently based on the feature level we are rendering with
+	bool bEnabled;
+
 	// is called by FPooledRenderTarget
 
 	void GenerateContent(const FSceneRenderTargetItem& RenderTargetItem, const FPooledRenderTargetDesc& Desc);
@@ -108,5 +111,5 @@ struct FVisualizeTextureData
 	bool bSaturateInsteadOfFrac;
 	int32 InputValueMapping;
 	int32 ArrayIndex;
-	int32 CustomMip;	
+	int32 CustomMip;
 };
\ No newline at end of file
diff --git a/Engine/Source/Runtime/Renderer/Private/PrimitiveSceneInfo.cpp b/Engine/Source/Runtime/Renderer/Private/PrimitiveSceneInfo.cpp
--- a/Engine/Source/Runtime/Renderer/Private/PrimitiveSceneInfo.cpp
+++ b/Engine/Source/Runtime/Renderer/Private/PrimitiveSceneInfo.cpp
@@ -157,7 +157,7 @@ void FPrimitiveSceneInfo::AddToScene(bool bUpdateStaticDrawLists)
 	// Allocate space in the indirect lighting cache so that it can be used for previewing indirect lighting
 	if (Proxy->HasStaticLighting() 
 		&& Proxy->NeedsUnbuiltPreviewLighting() 
-		&& IsIndirectLightingCacheAllowed())
+		&& IsIndirectLightingCacheAllowed(Scene->GetFeatureLevel()))
 	{
 		FIndirectLightingCacheAllocation* PrimitiveAllocation = Scene->IndirectLightingCache.FindPrimitiveAllocation(PrimitiveComponentId);
 
diff --git a/Engine/Source/Runtime/Renderer/Private/ReflectionEnvironment.cpp b/Engine/Source/Runtime/Renderer/Private/ReflectionEnvironment.cpp
--- a/Engine/Source/Runtime/Renderer/Private/ReflectionEnvironment.cpp
+++ b/Engine/Source/Runtime/Renderer/Private/ReflectionEnvironment.cpp
@@ -44,12 +44,12 @@ static int GetReflectionEnvironmentCVar()
 	return 1;
 }
 
-bool IsReflectionEnvironmentAvailable()
+bool IsReflectionEnvironmentAvailable(ERHIFeatureLevel::Type InFeatureLevel)
 {
 	static const auto AllowStaticLightingVar = IConsoleManager::Get().FindTConsoleVariableDataInt(TEXT("r.AllowStaticLighting"));
 	const bool bAllowStaticLighting = (!AllowStaticLightingVar || AllowStaticLightingVar->GetValueOnAnyThread() != 0);
 
-	return (GRHIFeatureLevel >= ERHIFeatureLevel::SM4) && (GetReflectionEnvironmentCVar() != 0) && bAllowStaticLighting;
+	return (InFeatureLevel >= ERHIFeatureLevel::SM4) && (GetReflectionEnvironmentCVar() != 0) && bAllowStaticLighting;
 }
 
 void FReflectionEnvironmentCubemapArray::InitDynamicRHI()
@@ -263,7 +263,7 @@ public:
 		SetUniformBufferParameterImmediate(ShaderRHI, GetUniformBufferParameter<FReflectionCaptureData>(), SamplePositionsBuffer);
 		SetShaderValue(ShaderRHI, NumCaptures, SortData.Num());
 
-		SetTextureParameter(ShaderRHI, PreIntegratedGF, PreIntegratedGFSampler, TStaticSamplerState<SF_Bilinear,AM_Clamp,AM_Clamp,AM_Clamp>::GetRHI(), GSystemTextures.PreintegratedGF->GetRenderTargetItem().ShaderResourceTexture);
+		SetTextureParameter(ShaderRHI, PreIntegratedGF, PreIntegratedGFSampler, TStaticSamplerState<SF_Bilinear, AM_Clamp, AM_Clamp, AM_Clamp>::GetRHI(), GSystemTextures.PreintegratedGF->GetRenderTargetItem().ShaderResourceTexture);
 	
 		SkyLightParameters.SetParameters(ShaderRHI, Scene, View.Family->EngineShowFlags.SkyLighting);
 	}
@@ -384,7 +384,7 @@ public:
 		
 		SetTextureParameter( ShaderRHI, ReflectionEnvTexture, ReflectionEnvSampler, TStaticSamplerState<SF_Point>::GetRHI(), ReflectionEnv );
 		SetTextureParameter( ShaderRHI, ScreenSpaceReflectionsTexture, ScreenSpaceReflectionsSampler, TStaticSamplerState<SF_Point>::GetRHI(), ScreenSpaceReflections );
-		SetTextureParameter( ShaderRHI, PreIntegratedGF, PreIntegratedGFSampler, TStaticSamplerState<SF_Bilinear,AM_Clamp,AM_Clamp,AM_Clamp>::GetRHI(), GSystemTextures.PreintegratedGF->GetRenderTargetItem().ShaderResourceTexture );
+		SetTextureParameter(ShaderRHI, PreIntegratedGF, PreIntegratedGFSampler, TStaticSamplerState<SF_Bilinear, AM_Clamp, AM_Clamp, AM_Clamp>::GetRHI(), GSystemTextures.PreintegratedGF->GetRenderTargetItem().ShaderResourceTexture);
 	}
 
 	// FShader interface.
@@ -552,7 +552,9 @@ void FDeferredShadingSceneRenderer::RenderDeferredReflections()
 		bAnyViewIsReflectionCapture = bAnyViewIsReflectionCapture || View.bIsReflectionCapture;
 	}
 
-	const uint32 bReflectionEnv = IsReflectionEnvironmentAvailable()
+	const ERHIFeatureLevel::Type FeatureLevel = Scene->GetFeatureLevel();
+
+	const uint32 bReflectionEnv = IsReflectionEnvironmentAvailable(FeatureLevel)
 		&& Scene->ReflectionSceneData.RegisteredReflectionCaptures.Num()
 		&& ViewFamily.EngineShowFlags.ReflectionEnvironment;
 
@@ -562,7 +564,7 @@ void FDeferredShadingSceneRenderer::RenderDeferredReflections()
 
 	const uint32 bUseLightmaps = GDiffuseFromCaptures == 0;
 
-	const bool bReflectionsWithCompute = GRHIFeatureLevel == ERHIFeatureLevel::SM5 && Scene->ReflectionSceneData.CubemapArray.IsValid();
+	const bool bReflectionsWithCompute = FeatureLevel == ERHIFeatureLevel::SM5 && Scene->ReflectionSceneData.CubemapArray.IsValid();
 
 	TRefCountPtr<IPooledRenderTarget>	NewSceneColor;
 	if( bReflectionEnv && bReflectionsWithCompute && !bAnyViewIsReflectionCapture )
@@ -623,7 +625,7 @@ void FDeferredShadingSceneRenderer::RenderDeferredReflections()
 			if( bReflectionEnv )
 			{
 				// If we are in SM5, use the compute shader gather method
-				if (GRHIFeatureLevel == ERHIFeatureLevel::SM5 && Scene->ReflectionSceneData.CubemapArray.IsValid())
+				if (FeatureLevel == ERHIFeatureLevel::SM5 && Scene->ReflectionSceneData.CubemapArray.IsValid())
 				{
 					// Render the reflection environment with tiled deferred culling
 					SCOPED_DRAW_EVENT(ReflectionEnvironmentGather, DEC_SCENE_ITEMS);
@@ -657,7 +659,7 @@ void FDeferredShadingSceneRenderer::RenderDeferredReflections()
 					bRequiresApply = false;
 				}
 				// In SM4 use standard deferred shading to composite reflection capture contribution
-				else if (GRHIFeatureLevel == ERHIFeatureLevel::SM4)
+				else if (FeatureLevel == ERHIFeatureLevel::SM4)
 				{
 					// TODO sort once per view
 					static TArray<FReflectionCaptureSortData> SortData;
diff --git a/Engine/Source/Runtime/Renderer/Private/ReflectionEnvironment.h b/Engine/Source/Runtime/Renderer/Private/ReflectionEnvironment.h
--- a/Engine/Source/Runtime/Renderer/Private/ReflectionEnvironment.h
+++ b/Engine/Source/Runtime/Renderer/Private/ReflectionEnvironment.h
@@ -13,5 +13,5 @@
  */
 static const int32 GMaxNumReflectionCaptures = 341;
 
-extern bool IsReflectionEnvironmentAvailable();
+extern bool IsReflectionEnvironmentAvailable(ERHIFeatureLevel::Type InFeatureLevel);
 extern bool IsReflectionCaptureAvailable();
diff --git a/Engine/Source/Runtime/Renderer/Private/ReflectionEnvironmentCapture.cpp b/Engine/Source/Runtime/Renderer/Private/ReflectionEnvironmentCapture.cpp
--- a/Engine/Source/Runtime/Renderer/Private/ReflectionEnvironmentCapture.cpp
+++ b/Engine/Source/Runtime/Renderer/Private/ReflectionEnvironmentCapture.cpp
@@ -943,7 +943,7 @@ void FScene::AllocateReflectionCaptures(const TArray<UReflectionCaptureComponent
 {
 	if (NewCaptures.Num() > 0)
 	{
-		if (GRHIFeatureLevel == ERHIFeatureLevel::SM5)
+		if (FeatureLevel == ERHIFeatureLevel::SM5)
 		{
 			for (int32 CaptureIndex = 0; CaptureIndex < NewCaptures.Num(); CaptureIndex++)
 			{
@@ -1013,7 +1013,7 @@ void FScene::AllocateReflectionCaptures(const TArray<UReflectionCaptureComponent
 				}
 			}
 		}
-		else if (GRHIFeatureLevel == ERHIFeatureLevel::SM4)
+		else if (FeatureLevel == ERHIFeatureLevel::SM4)
 		{
 			for (int32 ComponentIndex = 0; ComponentIndex < NewCaptures.Num(); ComponentIndex++)
 			{
@@ -1042,7 +1042,7 @@ void FScene::AllocateReflectionCaptures(const TArray<UReflectionCaptureComponent
 /** Updates the contents of all reflection captures in the scene.  Must be called from the game thread. */
 void FScene::UpdateAllReflectionCaptures()
 {
-	if (IsReflectionEnvironmentAvailable())
+	if (IsReflectionEnvironmentAvailable(FeatureLevel))
 	{
 		ENQUEUE_UNIQUE_RENDER_COMMAND_ONEPARAMETER( 
 			CaptureCommand,
@@ -1139,7 +1139,7 @@ void GetReflectionCaptureData_RenderingThread(FScene* Scene, const UReflectionCa
 
 void FScene::GetReflectionCaptureData(UReflectionCaptureComponent* Component, FReflectionCaptureFullHDRDerivedData& OutDerivedData) 
 {
-	check(GRHIFeatureLevel == ERHIFeatureLevel::SM5);
+	check(FeatureLevel == ERHIFeatureLevel::SM5);
 
 	ENQUEUE_UNIQUE_RENDER_COMMAND_THREEPARAMETER(
 		GetReflectionDataCommand,
@@ -1444,7 +1444,7 @@ void CopyToComponentTexture(FScene* Scene, FReflectionCaptureProxy* ReflectionPr
  */
 void FScene::UpdateReflectionCaptureContents(UReflectionCaptureComponent* CaptureComponent)
 {
-	if (IsReflectionEnvironmentAvailable())
+	if (IsReflectionEnvironmentAvailable(FeatureLevel))
 	{
 		const FReflectionCaptureFullHDRDerivedData* DerivedData = CaptureComponent->GetCachedFullHDRDerivedData();
 
@@ -1452,7 +1452,7 @@ void FScene::UpdateReflectionCaptureContents(UReflectionCaptureComponent* Captur
 		if (DerivedData && DerivedData->CompressedCapturedData.Num() > 0)
 		{
 			// For other feature levels the reflection textures are stored on the component instead of in a scene-wide texture array
-			if (GRHIFeatureLevel == ERHIFeatureLevel::SM5)
+			if (FeatureLevel == ERHIFeatureLevel::SM5)
 			{
 				ENQUEUE_UNIQUE_RENDER_COMMAND_THREEPARAMETER( 
 					UploadCaptureCommand,
@@ -1485,16 +1485,17 @@ void FScene::UpdateReflectionCaptureContents(UReflectionCaptureComponent* Captur
 			// We can't use the component's SceneProxy here because the component may not be registered with the scene
 			FReflectionCaptureProxy* ReflectionProxy = new FReflectionCaptureProxy(CaptureComponent);
 
-			ENQUEUE_UNIQUE_RENDER_COMMAND_TWOPARAMETER( 
+			ENQUEUE_UNIQUE_RENDER_COMMAND_THREEPARAMETER( 
 				CopyCommand,
 				FScene*, Scene, this,
 				FReflectionCaptureProxy*, ReflectionProxy, ReflectionProxy,
+				ERHIFeatureLevel::Type, FeatureLevel, FeatureLevel,
 			{
-				if (GRHIFeatureLevel == ERHIFeatureLevel::SM5)
+				if (FeatureLevel == ERHIFeatureLevel::SM5)
 				{
 					CopyToSceneArray(Scene, ReflectionProxy);
 				}
-				else if (GRHIFeatureLevel == ERHIFeatureLevel::SM4)
+				else if (FeatureLevel == ERHIFeatureLevel::SM4)
 				{
 					CopyToComponentTexture(Scene, ReflectionProxy);
 				}
@@ -1526,7 +1527,7 @@ void CopyToSkyTexture(FScene* Scene, FTexture* ProcessedTexture)
 
 void FScene::UpdateSkyCaptureContents(const USkyLightComponent* CaptureComponent, bool bCaptureEmissiveOnly, FTexture* OutProcessedTexture, FSHVectorRGB3& OutIrradianceEnvironmentMap)
 {
-	if (IsReflectionEnvironmentAvailable())
+	if (IsReflectionEnvironmentAvailable(FeatureLevel))
 	{
 		ENQUEUE_UNIQUE_RENDER_COMMAND( 
 			ClearCommand,
diff --git a/Engine/Source/Runtime/Renderer/Private/Renderer.cpp b/Engine/Source/Runtime/Renderer/Private/Renderer.cpp
--- a/Engine/Source/Runtime/Renderer/Private/Renderer.cpp
+++ b/Engine/Source/Runtime/Renderer/Private/Renderer.cpp
@@ -39,7 +39,9 @@ void FRendererModule::DrawTileMesh(const FSceneView& SceneView, const FMeshBatch
 	FViewInfo View(&SceneView);
 	View.InitRHIResources();
 
-	const FMaterial* Material = Mesh.MaterialRenderProxy->GetMaterial(GRHIFeatureLevel);	
+	const auto FeatureLevel = View.GetFeatureLevel();
+
+	const FMaterial* Material = Mesh.MaterialRenderProxy->GetMaterial(FeatureLevel);
 
 	//get the blend mode of the material
 	const EBlendMode MaterialBlendMode = Material->GetBlendMode();
@@ -49,7 +51,7 @@ void FRendererModule::DrawTileMesh(const FSceneView& SceneView, const FMeshBatch
 		// handle translucent material blend modes
 		if (IsTranslucentBlendMode(MaterialBlendMode))
 		{
-			if (GRHIFeatureLevel >= ERHIFeatureLevel::SM3)
+			if (FeatureLevel >= ERHIFeatureLevel::SM3)
 			{
 				FTranslucencyDrawingPolicyFactory::DrawDynamicMesh(View, FTranslucencyDrawingPolicyFactory::ContextType(), Mesh, false, false, NULL, HitProxyId);
 			}
@@ -71,7 +73,7 @@ void FRendererModule::DrawTileMesh(const FSceneView& SceneView, const FMeshBatch
 			}
 			else
 			{
-				if (GRHIFeatureLevel >= ERHIFeatureLevel::SM3)
+				if (FeatureLevel >= ERHIFeatureLevel::SM3)
 				{
 					FBasePassOpaqueDrawingPolicyFactory::DrawDynamicMesh(View, FBasePassOpaqueDrawingPolicyFactory::ContextType(false, ESceneRenderTargetsMode::SetTextures), Mesh, false, false, NULL, HitProxyId);
 				}
diff --git a/Engine/Source/Runtime/Renderer/Private/RendererPrivate.h b/Engine/Source/Runtime/Renderer/Private/RendererPrivate.h
--- a/Engine/Source/Runtime/Renderer/Private/RendererPrivate.h
+++ b/Engine/Source/Runtime/Renderer/Private/RendererPrivate.h
@@ -21,7 +21,7 @@ public:
 	virtual bool SupportsDynamicReloading() OVERRIDE { return true; }
 
 	virtual void BeginRenderingViewFamily(FCanvas* Canvas,const FSceneViewFamily* ViewFamily) OVERRIDE;
-	virtual FSceneInterface* AllocateScene(UWorld* World, bool bInRequiresHitProxies) OVERRIDE;
+	virtual FSceneInterface* AllocateScene(UWorld* World, bool bInRequiresHitProxies, ERHIFeatureLevel::Type InFeatureLevel) OVERRIDE;
 	virtual void RemoveScene(FSceneInterface* Scene) OVERRIDE;
 	virtual void UpdateStaticDrawListsForMaterials(const TArray<const FMaterial*>& Materials) OVERRIDE;
 	virtual FSceneViewStateInterface* AllocateViewState() OVERRIDE;
diff --git a/Engine/Source/Runtime/Renderer/Private/RendererScene.cpp b/Engine/Source/Runtime/Renderer/Private/RendererScene.cpp
--- a/Engine/Source/Runtime/Renderer/Private/RendererScene.cpp
+++ b/Engine/Source/Runtime/Renderer/Private/RendererScene.cpp
@@ -139,7 +139,7 @@ FORCEINLINE static void VerifyProperPIEScene(UPrimitiveComponent* Component, UWo
 #endif
 }
 
-FScene::FScene(UWorld* InWorld, bool bInRequiresHitProxies, bool bInIsEditorScene)
+FScene::FScene(UWorld* InWorld, bool bInRequiresHitProxies, bool bInIsEditorScene, ERHIFeatureLevel::Type InFeatureLevel)
 :	World(InWorld)
 ,	FXSystem(NULL)
 ,	bStaticDrawListsMobileHDR(false)
@@ -162,6 +162,7 @@ FScene::FScene(UWorld* InWorld, bool bInRequiresHitProxies, bool bInIsEditorScen
 ,	LowerDynamicSkylightColor(FLinearColor::Black)
 ,	NumVisibleLights(0)
 ,	bHasSkyLight(false)
+,	FeatureLevel(InFeatureLevel)
 {
 	check(World);
 
@@ -1465,7 +1466,7 @@ void FScene::SetShaderMapsOnMaterialResources_RenderThread(const FMaterialsToUpd
 	for (TSet<FMaterialRenderProxy*>::TConstIterator It(FMaterialRenderProxy::GetMaterialRenderProxyMap()); It; ++It)
 	{
 		FMaterialRenderProxy* MaterialProxy = *It;
-		FMaterial* Material = MaterialProxy->GetMaterialNoFallback(GRHIFeatureLevel);
+		FMaterial* Material = MaterialProxy->GetMaterialNoFallback(FeatureLevel);
 
 		if (Material && MaterialsToUpdate.Contains(Material))
 		{
@@ -1474,11 +1475,11 @@ void FScene::SetShaderMapsOnMaterialResources_RenderThread(const FMaterialsToUpd
 			MaterialProxy->CacheUniformExpressions();
 			bFoundAnyInitializedMaterials = true;
 
-			const FMaterial& MaterialForRendering = *MaterialProxy->GetMaterial(GRHIFeatureLevel);
+			const FMaterial& MaterialForRendering = *MaterialProxy->GetMaterial(FeatureLevel);
 			check(MaterialForRendering.GetRenderingThreadShaderMap());
 
-			check(!MaterialProxy->UniformExpressionCache[GRHIFeatureLevel].bUpToDate 
-				|| MaterialProxy->UniformExpressionCache[GRHIFeatureLevel].CachedUniformExpressionShaderMap == MaterialForRendering.GetRenderingThreadShaderMap());
+			check(!MaterialProxy->UniformExpressionCache[FeatureLevel].bUpToDate
+				|| MaterialProxy->UniformExpressionCache[FeatureLevel].CachedUniformExpressionShaderMap == MaterialForRendering.GetRenderingThreadShaderMap());
 
 			check(MaterialForRendering.GetRenderingThreadShaderMap()->IsValidForRendering());
 		}
@@ -2026,14 +2027,14 @@ private:
 	class FFXSystemInterface* FXSystem;
 };
 
-FSceneInterface* FRendererModule::AllocateScene(UWorld* World, bool bInRequiresHitProxies)
+FSceneInterface* FRendererModule::AllocateScene(UWorld* World, bool bInRequiresHitProxies, ERHIFeatureLevel::Type InFeatureLevel)
 {
 	check(IsInGameThread());
 
 	// Create a full fledged scene if we have something to render.
 	if( GIsClient && !IsRunningCommandlet() && !GUsingNullRHI )
 	{
-		FScene* NewScene = new FScene(World,bInRequiresHitProxies,GIsEditor && !World->IsGameWorld());
+		FScene* NewScene = new FScene(World, bInRequiresHitProxies, GIsEditor && !World->IsGameWorld(), InFeatureLevel);
 		AllocatedScenes.Add(NewScene);
 		return NewScene;
 	}
@@ -2244,11 +2245,11 @@ void FMotionBlurInfoData::RestoreForPausedMotionBlur()
 	}
 }
 
-void FMotionBlurInfoData::UpdateMotionBlurCache()
+void FMotionBlurInfoData::UpdateMotionBlurCache(FScene* InScene)
 {
 	check(IsInRenderingThread());
 
-	if (GRHIFeatureLevel >= ERHIFeatureLevel::SM4)
+	if (InScene->GetFeatureLevel() >= ERHIFeatureLevel::SM4)
 	{
 		if(bShouldClearMotionBlurInfo)
 		{
diff --git a/Engine/Source/Runtime/Renderer/Private/SceneCaptureRendering.cpp b/Engine/Source/Runtime/Renderer/Private/SceneCaptureRendering.cpp
--- a/Engine/Source/Runtime/Renderer/Private/SceneCaptureRendering.cpp
+++ b/Engine/Source/Runtime/Renderer/Private/SceneCaptureRendering.cpp
@@ -227,7 +227,7 @@ void FScene::UpdateSceneCaptureContents(USceneCaptureComponentCube* CaptureCompo
 
 	check(CaptureComponent);
 
-	if (GRHIFeatureLevel >= ERHIFeatureLevel::SM4 && CaptureComponent->TextureTarget)
+	if (FeatureLevel >= ERHIFeatureLevel::SM4 && CaptureComponent->TextureTarget)
 	{
 		const float FOV = 90 * (float)PI / 360.0f;
 		for (int32 faceidx = 0; faceidx < (int32)ECubeFace::CubeFace_MAX; faceidx++)
diff --git a/Engine/Source/Runtime/Renderer/Private/SceneCore.cpp b/Engine/Source/Runtime/Renderer/Private/SceneCore.cpp
--- a/Engine/Source/Runtime/Renderer/Private/SceneCore.cpp
+++ b/Engine/Source/Runtime/Renderer/Private/SceneCore.cpp
@@ -276,7 +276,7 @@ void FStaticMesh::UnlinkDrawList(FStaticMesh::FDrawListElementLink* Link)
 
 void FStaticMesh::AddToDrawLists(FScene* Scene)
 {
-	if (GRHIFeatureLevel >= ERHIFeatureLevel::SM3)
+	if (Scene->GetFeatureLevel() >= ERHIFeatureLevel::SM3)
 	{
 		if (CastShadow)
 		{
diff --git a/Engine/Source/Runtime/Renderer/Private/SceneCore.h b/Engine/Source/Runtime/Renderer/Private/SceneCore.h
--- a/Engine/Source/Runtime/Renderer/Private/SceneCore.h
+++ b/Engine/Source/Runtime/Renderer/Private/SceneCore.h
@@ -207,7 +207,7 @@ public:
 };
 
 /** Returns true if the indirect lighting cache can be used at all. */
-extern bool IsIndirectLightingCacheAllowed();
+extern bool IsIndirectLightingCacheAllowed(ERHIFeatureLevel::Type InFeatureLevel);
 
 /** Returns true if the indirect lighting cache can use the volume texture atlas on this feature level. */
-extern bool CanIndirectLightingCacheUseVolumeTexture();
+extern bool CanIndirectLightingCacheUseVolumeTexture(ERHIFeatureLevel::Type InFeatureLevel);
diff --git a/Engine/Source/Runtime/Renderer/Private/SceneHitProxyRendering.cpp b/Engine/Source/Runtime/Renderer/Private/SceneHitProxyRendering.cpp
--- a/Engine/Source/Runtime/Renderer/Private/SceneHitProxyRendering.cpp
+++ b/Engine/Source/Runtime/Renderer/Private/SceneHitProxyRendering.cpp
@@ -19,7 +19,7 @@ public:
 
 	void SetParameters(const FMaterialRenderProxy* MaterialRenderProxy,const FSceneView& View)
 	{
-		FMeshMaterialShader::SetParameters(GetVertexShader(),MaterialRenderProxy,*MaterialRenderProxy->GetMaterial(GRHIFeatureLevel),View,ESceneRenderTargetsMode::SetTextures);
+		FMeshMaterialShader::SetParameters(GetVertexShader(), MaterialRenderProxy, *MaterialRenderProxy->GetMaterial(View.GetFeatureLevel()), View, ESceneRenderTargetsMode::SetTextures);
 	}
 
 	void SetMesh(const FVertexFactory* VertexFactory,const FSceneView& View,const FPrimitiveSceneProxy* Proxy,const FMeshBatchElement& BatchElement)
@@ -125,7 +125,7 @@ public:
 
 	void SetParameters(const FMaterialRenderProxy* MaterialRenderProxy,const FSceneView& View)
 	{
-		FMeshMaterialShader::SetParameters(GetPixelShader(),MaterialRenderProxy,*MaterialRenderProxy->GetMaterial(GRHIFeatureLevel),View,ESceneRenderTargetsMode::SetTextures);
+		FMeshMaterialShader::SetParameters(GetPixelShader(), MaterialRenderProxy, *MaterialRenderProxy->GetMaterial(View.GetFeatureLevel()), View, ESceneRenderTargetsMode::SetTextures);
 	}
 
 	void SetMesh(const FVertexFactory* VertexFactory,const FSceneView& View,const FPrimitiveSceneProxy* Proxy,const FMeshBatchElement& BatchElement)
@@ -153,9 +153,10 @@ IMPLEMENT_MATERIAL_SHADER_TYPE(,FHitProxyPS,TEXT("HitProxyPixelShader"),TEXT("Ma
 
 FHitProxyDrawingPolicy::FHitProxyDrawingPolicy(
 	const FVertexFactory* InVertexFactory,
-	const FMaterialRenderProxy* InMaterialRenderProxy
-	):
-	FMeshDrawingPolicy(InVertexFactory, InMaterialRenderProxy, *InMaterialRenderProxy->GetMaterial(GRHIFeatureLevel))
+	const FMaterialRenderProxy* InMaterialRenderProxy,
+	ERHIFeatureLevel::Type InFeatureLevel
+	) :
+	FMeshDrawingPolicy(InVertexFactory, InMaterialRenderProxy, *InMaterialRenderProxy->GetMaterial(InFeatureLevel))
 {
 	HullShader = NULL;
 	DomainShader = NULL;
@@ -196,7 +197,7 @@ void FHitProxyDrawingPolicy::DrawShared(const FSceneView* View, FBoundShaderStat
 * as well as the shaders needed to draw the mesh
 * @return new bound shader state object
 */
-FBoundShaderStateRHIRef FHitProxyDrawingPolicy::CreateBoundShaderState()
+FBoundShaderStateRHIRef FHitProxyDrawingPolicy::CreateBoundShaderState(ERHIFeatureLevel::Type InFeatureLevel)
 {
 	FBoundShaderStateRHIRef BoundShaderState;
 
@@ -256,7 +257,7 @@ void FHitProxyDrawingPolicyFactory::AddStaticMesh(FScene* Scene,FStaticMesh* Sta
 
 	// Add the static mesh to the DPG's hit proxy draw list.
 	const FMaterialRenderProxy* MaterialRenderProxy = StaticMesh->MaterialRenderProxy;
-	const FMaterial* Material = MaterialRenderProxy->GetMaterial(GRHIFeatureLevel);
+	const FMaterial* Material = MaterialRenderProxy->GetMaterial(Scene->GetFeatureLevel());
 	if (!Material->IsMasked() && !Material->IsTwoSided() && !Material->MaterialModifiesMeshPosition())
 	{
 		// Default material doesn't handle masked, and doesn't have the correct bIsTwoSided setting.
@@ -266,7 +267,8 @@ void FHitProxyDrawingPolicyFactory::AddStaticMesh(FScene* Scene,FStaticMesh* Sta
 	Scene->HitProxyDrawList.AddMesh(
 		StaticMesh,
 		StaticMesh->HitProxyId,
-		FHitProxyDrawingPolicy(StaticMesh->VertexFactory,MaterialRenderProxy)
+		FHitProxyDrawingPolicy(StaticMesh->VertexFactory, MaterialRenderProxy, Scene->GetFeatureLevel()),
+		Scene->GetFeatureLevel()
 		);
 
 #if WITH_EDITOR
@@ -278,7 +280,8 @@ void FHitProxyDrawingPolicyFactory::AddStaticMesh(FScene* Scene,FStaticMesh* Sta
 		Scene->HitProxyDrawList_OpaqueOnly.AddMesh(
 			StaticMesh,
 			StaticMesh->HitProxyId,
-			FHitProxyDrawingPolicy(StaticMesh->VertexFactory,MaterialRenderProxy)
+			FHitProxyDrawingPolicy(StaticMesh->VertexFactory, MaterialRenderProxy, Scene->GetFeatureLevel()),
+			Scene->GetFeatureLevel()
 			);
 	}
 #endif
@@ -298,7 +301,7 @@ bool FHitProxyDrawingPolicyFactory::DrawDynamicMesh(
 	if (!PrimitiveSceneProxy || PrimitiveSceneProxy->IsSelectable())
 	{
 		const FMaterialRenderProxy* MaterialRenderProxy = Mesh.MaterialRenderProxy;
-		const FMaterial* Material = MaterialRenderProxy->GetMaterial(GRHIFeatureLevel);
+		const FMaterial* Material = MaterialRenderProxy->GetMaterial(View.GetFeatureLevel());
 #if WITH_EDITOR
 		const HHitProxy* HitProxy = GetHitProxyById( HitProxyId );
 		// Only draw translucent primitives to the hit proxy if the user wants those objects to be selectable
@@ -310,8 +313,8 @@ bool FHitProxyDrawingPolicyFactory::DrawDynamicMesh(
 				// Default material doesn't handle masked, and doesn't have the correct bIsTwoSided setting.
 				MaterialRenderProxy = UMaterial::GetDefaultMaterial(MD_Surface)->GetRenderProxy(false);
 			}
-			FHitProxyDrawingPolicy DrawingPolicy( Mesh.VertexFactory, MaterialRenderProxy );
-			DrawingPolicy.DrawShared(&View, DrawingPolicy.CreateBoundShaderState());
+			FHitProxyDrawingPolicy DrawingPolicy( Mesh.VertexFactory, MaterialRenderProxy, View.GetFeatureLevel() );
+			DrawingPolicy.DrawShared(&View, DrawingPolicy.CreateBoundShaderState(View.GetFeatureLevel()));
 			for( int32 BatchElementIndex = 0; BatchElementIndex < Mesh.Elements.Num(); BatchElementIndex++ )
 			{
 				DrawingPolicy.SetMeshRenderState(View, PrimitiveSceneProxy, Mesh, BatchElementIndex, bBackFace, HitProxyId);
@@ -326,6 +329,9 @@ bool FHitProxyDrawingPolicyFactory::DrawDynamicMesh(
 void FDeferredShadingSceneRenderer::RenderHitProxies()
 {
 #if WITH_EDITOR
+
+	auto FeatureLevel = ViewFamily.Scene->GetFeatureLevel();
+
 	// Initialize global system textures (pass-through if already initialized).
 	GSystemTextures.InitializeTextures();
 
diff --git a/Engine/Source/Runtime/Renderer/Private/SceneHitProxyRendering.h b/Engine/Source/Runtime/Renderer/Private/SceneHitProxyRendering.h
--- a/Engine/Source/Runtime/Renderer/Private/SceneHitProxyRendering.h
+++ b/Engine/Source/Runtime/Renderer/Private/SceneHitProxyRendering.h
@@ -17,7 +17,8 @@ public:
 
 	FHitProxyDrawingPolicy(
 		const FVertexFactory* InVertexFactory,
-		const FMaterialRenderProxy* InMaterialRenderProxy
+		const FMaterialRenderProxy* InMaterialRenderProxy,
+		ERHIFeatureLevel::Type InFeatureLevel
 		);
 
 	// FMeshDrawingPolicy interface.
@@ -36,7 +37,7 @@ public:
 	* as well as the shaders needed to draw the mesh
 	* @return new bound shader state object
 	*/
-	FBoundShaderStateRHIRef CreateBoundShaderState();
+	FBoundShaderStateRHIRef CreateBoundShaderState(ERHIFeatureLevel::Type InFeatureLevel);
 
 	void SetMeshRenderState(
 		const FSceneView& View,
@@ -74,7 +75,7 @@ public:
 		const FPrimitiveSceneProxy* PrimitiveSceneProxy,
 		FHitProxyId HitProxyId
 		);
-	static bool IsMaterialIgnored(const FMaterialRenderProxy* MaterialRenderProxy)
+	static bool IsMaterialIgnored(const FMaterialRenderProxy* MaterialRenderProxy, ERHIFeatureLevel::Type InFeatureLevel)
 	{
 		return false;
 	}
diff --git a/Engine/Source/Runtime/Renderer/Private/SceneOcclusion.cpp b/Engine/Source/Runtime/Renderer/Private/SceneOcclusion.cpp
--- a/Engine/Source/Runtime/Renderer/Private/SceneOcclusion.cpp
+++ b/Engine/Source/Runtime/Renderer/Private/SceneOcclusion.cpp
@@ -816,7 +816,7 @@ public:
 		const FVector2D InvSize( 1.0f / Size.X, 1.0f / Size.Y );
 		SetShaderValue( ShaderRHI, InvSizeParameter, InvSize );
 		
-		SceneTextureParameters.Set( ShaderRHI );
+		SceneTextureParameters.Set( ShaderRHI, View );
 	}
 
 	void SetParameters( const FSceneView& View, const FIntPoint& Size, FShaderResourceViewRHIParamRef ShaderResourceView )
diff --git a/Engine/Source/Runtime/Renderer/Private/ScenePrivate.h b/Engine/Source/Runtime/Renderer/Private/ScenePrivate.h
--- a/Engine/Source/Runtime/Renderer/Private/ScenePrivate.h
+++ b/Engine/Source/Runtime/Renderer/Private/ScenePrivate.h
@@ -1282,7 +1282,7 @@ public:
 	TMap<FGuid, FUniformBufferRHIRef> ParameterCollections;
 
 	/** Initialization constructor. */
-	FScene(UWorld* InWorld,bool bInRequiresHitProxies,bool bInIsEditorScene);
+	FScene(UWorld* InWorld, bool bInRequiresHitProxies,bool bInIsEditorScene, ERHIFeatureLevel::Type InFeatureLevel);
 
 	virtual ~FScene();
 
@@ -1429,6 +1429,8 @@ public:
 
 	virtual bool IsEditorScene() const OVERRIDE { return bIsEditorScene; }
 
+	virtual ERHIFeatureLevel::Type GetFeatureLevel() const OVERRIDE { return FeatureLevel; }
+
 private:
 
 	/**
@@ -1539,6 +1541,11 @@ private:
 	 * Note: This is tracked on the game thread!
 	 */
 	bool bHasSkyLight;
+
+	/** 
+	 * Feature level that this scene should render using
+	 */
+	ERHIFeatureLevel::Type FeatureLevel;
 };
 
 #endif // __SCENEPRIVATE_H__
diff --git a/Engine/Source/Runtime/Renderer/Private/SceneRendering.cpp b/Engine/Source/Runtime/Renderer/Private/SceneRendering.cpp
--- a/Engine/Source/Runtime/Renderer/Private/SceneRendering.cpp
+++ b/Engine/Source/Runtime/Renderer/Private/SceneRendering.cpp
@@ -786,7 +786,9 @@ void FSceneRenderer::RenderFinish()
 
 FSceneRenderer* FSceneRenderer::CreateSceneRenderer(const FSceneViewFamily* InViewFamily, FHitProxyConsumer* HitProxyConsumer)
 {
-	if (GRHIFeatureLevel == ERHIFeatureLevel::ES2)
+	ERHIFeatureLevel::Type FeatureLevel = InViewFamily->Scene->GetFeatureLevel();
+
+	if (FeatureLevel == ERHIFeatureLevel::ES2)
 	{
 		return new FForwardShadingSceneRenderer(InViewFamily, HitProxyConsumer);
 	}
@@ -897,7 +899,7 @@ static void RenderViewFamily_RenderThread( FSceneRenderer* SceneRenderer )
 
     {
 		SCOPE_CYCLE_COUNTER(STAT_TotalSceneRenderingTime);
-
+		
 		if(SceneRenderer->ViewFamily.EngineShowFlags.HitProxies)
 		{
 			// Render the scene's hit proxies.
diff --git a/Engine/Source/Runtime/Renderer/Private/SceneRendering.h b/Engine/Source/Runtime/Renderer/Private/SceneRendering.h
--- a/Engine/Source/Runtime/Renderer/Private/SceneRendering.h
+++ b/Engine/Source/Runtime/Renderer/Private/SceneRendering.h
@@ -412,6 +412,8 @@ public:
 	FLinearColor LightShaftColorApply;
 	bool bLightShaftUse;
 
+	ERHIFeatureLevel::Type FeatureLevel;
+
 	/** 
 	 * Initialization constructor. Passes all parameters to FSceneView constructor
 	 */
@@ -581,7 +583,7 @@ protected:
 	void RenderForwardShadingBasePass();
 
 	/** Makes a copy of scene alpha so PC can emulate ES2 framebuffer fetch. */
-	void CopySceneAlpha(void);
+	void CopySceneAlpha(const FSceneView& View);
 
 	/** Renders the base pass for translucency. */
 	void RenderTranslucency();
diff --git a/Engine/Source/Runtime/Renderer/Private/SceneVisibility.cpp b/Engine/Source/Runtime/Renderer/Private/SceneVisibility.cpp
--- a/Engine/Source/Runtime/Renderer/Private/SceneVisibility.cpp
+++ b/Engine/Source/Runtime/Renderer/Private/SceneVisibility.cpp
@@ -348,7 +348,7 @@ static int32 OcclusionCull(const FScene* Scene, FViewInfo& View)
 	float CurrentRealTime = View.Family->CurrentRealTime;
 	if (ViewState)
 	{
-		if (GRHIFeatureLevel >= ERHIFeatureLevel::SM3)
+		if (Scene->GetFeatureLevel() >= ERHIFeatureLevel::SM3)
 		{
 			bool bClearQueries = !View.Family->EngineShowFlags.HitProxies;
 			bool bSubmitQueries = !View.bDisableQuerySubmissions;
@@ -731,7 +731,7 @@ static void ComputeRelevanceForView(
 		if (PrimitiveSceneInfo->bNeedsCachedReflectionCaptureUpdate
 			// In ES2, the per-object reflection is used for everything
 			// Otherwise it is just used on translucency
-			&& (GRHIFeatureLevel == ERHIFeatureLevel::ES2 || bTranslucentRelevance))
+			&& (Scene->GetFeatureLevel() == ERHIFeatureLevel::ES2 || bTranslucentRelevance))
 		{
 			PrimitiveSceneInfo->CachedReflectionCaptureProxy = Scene->FindClosestReflectionCapture(Scene->PrimitiveBounds[BitIt.GetIndex()].Origin);
 			PrimitiveSceneInfo->bNeedsCachedReflectionCaptureUpdate = false;
@@ -994,7 +994,7 @@ void FSceneRenderer::PreVisibilityFrameSetup()
 			{
 				float SampleX, SampleY;
 
-				if( GRHIFeatureLevel < ERHIFeatureLevel::SM4 )
+				if (Scene->GetFeatureLevel() < ERHIFeatureLevel::SM4)
 				{
 					// Only support 2 samples for mobile temporal AA.
 					TemporalAASamples = 2;
@@ -1409,7 +1409,7 @@ void FSceneRenderer::PostVisibilityFrameSetup()
 	}
 
 	bool bCheckLightShafts = false;
-	if(GRHIFeatureLevel <= ERHIFeatureLevel::ES2)
+	if (Scene->GetFeatureLevel() <= ERHIFeatureLevel::ES2)
 	{
 		// Clear the mobile light shaft data.
 		for (int32 ViewIndex = 0; ViewIndex < Views.Num(); ViewIndex++)
diff --git a/Engine/Source/Runtime/Renderer/Private/ShaderBaseClasses.h b/Engine/Source/Runtime/Renderer/Private/ShaderBaseClasses.h
--- a/Engine/Source/Runtime/Renderer/Private/ShaderBaseClasses.h
+++ b/Engine/Source/Runtime/Renderer/Private/ShaderBaseClasses.h
@@ -137,7 +137,7 @@ public:
 		bool bDeferredPass, 
 		ESceneRenderTargetsMode::Type TextureMode)
 	{
-		ERHIFeatureLevel::Type FeatureLevel = GRHIFeatureLevel;
+		ERHIFeatureLevel::Type FeatureLevel = View.GetFeatureLevel();
 		FUniformExpressionCache TempUniformExpressionCache;
 		const FUniformExpressionCache* UniformExpressionCache = &MaterialRenderProxy->UniformExpressionCache[FeatureLevel];
 
@@ -276,7 +276,7 @@ public:
 
 		AtmosphericFogTextureParameters.Set(ShaderRHI, View);
 
-		if (GRHIFeatureLevel >= ERHIFeatureLevel::SM3)
+		if (FeatureLevel >= ERHIFeatureLevel::SM3)
 		{
 			if(LightAttenuation.IsBound())
 			{
@@ -468,7 +468,7 @@ public:
 
 	void SetParameters(const FMaterialRenderProxy* MaterialRenderProxy,const FSceneView& View)
 	{
-		FMeshMaterialShader::SetParameters(GetHullShader(),MaterialRenderProxy,*MaterialRenderProxy->GetMaterial(GRHIFeatureLevel),View,ESceneRenderTargetsMode::SetTextures);
+		FMeshMaterialShader::SetParameters(GetHullShader(), MaterialRenderProxy, *MaterialRenderProxy->GetMaterial(View.GetFeatureLevel()), View, ESceneRenderTargetsMode::SetTextures);
 	}
 
 	void SetMesh(const FVertexFactory* VertexFactory,const FSceneView& View,const FPrimitiveSceneProxy* Proxy,const FMeshBatchElement& BatchElement)
@@ -514,7 +514,7 @@ public:
 
 	void SetParameters(const FMaterialRenderProxy* MaterialRenderProxy,const FSceneView& View)
 	{
-		FMeshMaterialShader::SetParameters(GetDomainShader(),MaterialRenderProxy,*MaterialRenderProxy->GetMaterial(GRHIFeatureLevel),View,ESceneRenderTargetsMode::SetTextures);
+		FMeshMaterialShader::SetParameters(GetDomainShader(), MaterialRenderProxy, *MaterialRenderProxy->GetMaterial(View.GetFeatureLevel()), View, ESceneRenderTargetsMode::SetTextures);
 	}
 
 	void SetMesh(const FVertexFactory* VertexFactory,const FSceneView& View,const FPrimitiveSceneProxy* Proxy,const FMeshBatchElement& BatchElement)
diff --git a/Engine/Source/Runtime/Renderer/Private/ShaderComplexityRendering.cpp b/Engine/Source/Runtime/Renderer/Private/ShaderComplexityRendering.cpp
--- a/Engine/Source/Runtime/Renderer/Private/ShaderComplexityRendering.cpp
+++ b/Engine/Source/Runtime/Renderer/Private/ShaderComplexityRendering.cpp
@@ -12,17 +12,18 @@ ShaderComplexityRendering.cpp: Contains definitions for rendering the shader com
 /**
  * Gets the maximum shader complexity count from the ini settings.
  */
-float GetMaxShaderComplexityCount()
+float GetMaxShaderComplexityCount(ERHIFeatureLevel::Type InFeatureType)
 {
-	return GRHIFeatureLevel == ERHIFeatureLevel::ES2 ? GEngine->MaxES2PixelShaderAdditiveComplexityCount : GEngine->MaxPixelShaderAdditiveComplexityCount;
+	return InFeatureType == ERHIFeatureLevel::ES2 ? GEngine->MaxES2PixelShaderAdditiveComplexityCount : GEngine->MaxPixelShaderAdditiveComplexityCount;
 }
 
 void FShaderComplexityAccumulatePS::SetParameters(
 	uint32 NumVertexInstructions, 
-	uint32 NumPixelInstructions)
+	uint32 NumPixelInstructions,
+	ERHIFeatureLevel::Type InFeatureLevel)
 {
 	//normalize the complexity so we can fit it in a low precision scene color which is necessary on some platforms
-	const float NormalizedComplexityValue = float(NumPixelInstructions) / GetMaxShaderComplexityCount();
+	const float NormalizedComplexityValue = float(NumPixelInstructions) / GetMaxShaderComplexityCount(InFeatureLevel);
 
 	SetShaderValue( GetPixelShader(), NormalizedComplexity, NormalizedComplexityValue);
 }
@@ -95,9 +96,9 @@ public:
 		const TArray<FLinearColor>& Colors
 		)
 	{
-		FGlobalShader::SetParameters(GetPixelShader(),*View);
+		FGlobalShader::SetParameters(GetPixelShader(), *View);
 
-		SceneTextureParams.Set(GetPixelShader());
+		SceneTextureParams.Set(GetPixelShader(), *View);
 
 		//Make sure there are at least NumShaderComplexityColors colors specified in the ini.
 		//If there are more than NumShaderComplexityColors they will be ignored.
diff --git a/Engine/Source/Runtime/Renderer/Private/ShaderComplexityRendering.h b/Engine/Source/Runtime/Renderer/Private/ShaderComplexityRendering.h
--- a/Engine/Source/Runtime/Renderer/Private/ShaderComplexityRendering.h
+++ b/Engine/Source/Runtime/Renderer/Private/ShaderComplexityRendering.h
@@ -28,8 +28,9 @@ public:
 	FShaderComplexityAccumulatePS() {}
 
 	void SetParameters(
-		uint32 NumVertexInstructions, 
-		uint32 NumPixelInstructions);
+		uint32 NumVertexInstructions,
+		uint32 NumPixelInstructions,
+		ERHIFeatureLevel::Type InFeatureLevel);
 
 	virtual bool Serialize(FArchive& Ar)
 	{
diff --git a/Engine/Source/Runtime/Renderer/Private/ShadowRendering.cpp b/Engine/Source/Runtime/Renderer/Private/ShadowRendering.cpp
--- a/Engine/Source/Runtime/Renderer/Private/ShadowRendering.cpp
+++ b/Engine/Source/Runtime/Renderer/Private/ShadowRendering.cpp
@@ -634,7 +634,8 @@ IMPLEMENT_SHADER_TYPE(template<>,TOnePassPointShadowProjectionPS<5>,TEXT("Shadow
 void OverrideWithDefaultMaterialForShadowDepth(
 	const FMaterialRenderProxy*& InOutMaterialRenderProxy, 
 	const FMaterial*& InOutMaterialResource,
-	bool bReflectiveShadowmap)
+	bool bReflectiveShadowmap,
+	ERHIFeatureLevel::Type InFeatureLevel)
 {
 	// Override with the default material when possible.
 	if (!InOutMaterialResource->IsMasked() &&						// Don't override masked materials.
@@ -642,7 +643,7 @@ void OverrideWithDefaultMaterialForShadowDepth(
 		!bReflectiveShadowmap)										// Don't override when rendering reflective shadow maps.
 	{
 		const FMaterialRenderProxy* DefaultProxy = UMaterial::GetDefaultMaterial(MD_Surface)->GetRenderProxy(false);
-		const FMaterial* DefaultMaterialResource = DefaultProxy->GetMaterial(GRHIFeatureLevel);
+		const FMaterial* DefaultMaterialResource = DefaultProxy->GetMaterial(InFeatureLevel);
 
 		// Override with the default material for opaque materials that don't modify mesh position.
 		InOutMaterialRenderProxy = DefaultProxy;
@@ -873,7 +874,7 @@ void FShadowDepthDrawingPolicy<bRenderingReflectiveShadowMaps>::DrawShared(const
  * @return new bound shader state object
  */
 template <bool bRenderingReflectiveShadowMaps>
-FBoundShaderStateRHIRef FShadowDepthDrawingPolicy<bRenderingReflectiveShadowMaps>::CreateBoundShaderState()
+FBoundShaderStateRHIRef FShadowDepthDrawingPolicy<bRenderingReflectiveShadowMaps>::CreateBoundShaderState(ERHIFeatureLevel::Type InFeatureLevel)
 {
 	FVertexDeclarationRHIRef VertexDeclaration;
 	if (bUsePositionOnlyVS)
@@ -951,7 +952,7 @@ void FShadowDepthDrawingPolicyFactory::AddStaticMesh(FScene* Scene,FStaticMesh*
 	if (StaticMesh->CastShadow)
 	{
 		const FMaterialRenderProxy* MaterialRenderProxy = StaticMesh->MaterialRenderProxy;
-		const FMaterial* Material = MaterialRenderProxy->GetMaterial(GRHIFeatureLevel);
+		const FMaterial* Material = MaterialRenderProxy->GetMaterial(Scene->GetFeatureLevel());
 		const EBlendMode BlendMode = Material->GetBlendMode();
 		const EMaterialLightingModel LightingModel = Material->GetLightingModel();
 
@@ -974,13 +975,14 @@ void FShadowDepthDrawingPolicyFactory::AddStaticMesh(FScene* Scene,FStaticMesh*
 						StaticMesh->VertexFactory,
 						MaterialRenderProxy,
 						bTwoSided,
-						StaticMesh->ReverseCulling)
+						StaticMesh->ReverseCulling),
+					Scene->GetFeatureLevel()
 					);
 			}
 		}
 		if ( bLitOpaque )
 		{
-			OverrideWithDefaultMaterialForShadowDepth(MaterialRenderProxy, Material, false); 
+			OverrideWithDefaultMaterialForShadowDepth(MaterialRenderProxy, Material, false, Scene->GetFeatureLevel()); 
 
 			// Add the static mesh to the shadow's subject draw list.
 			Scene->WholeSceneShadowDepthDrawList.AddMesh(
@@ -994,7 +996,8 @@ void FShadowDepthDrawingPolicyFactory::AddStaticMesh(FScene* Scene,FStaticMesh*
 					StaticMesh->VertexFactory,
 					MaterialRenderProxy,
 					bTwoSided,
-					StaticMesh->ReverseCulling)
+					StaticMesh->ReverseCulling),
+				Scene->GetFeatureLevel()
 				);
 		}
 	}
@@ -1016,7 +1019,7 @@ bool FShadowDepthDrawingPolicyFactory::DrawDynamicMesh(
 	if (Mesh.CastShadow)
 	{
 		const FMaterialRenderProxy* MaterialRenderProxy = Mesh.MaterialRenderProxy;
-		const FMaterial* Material = MaterialRenderProxy->GetMaterial(GRHIFeatureLevel);
+		const FMaterial* Material = MaterialRenderProxy->GetMaterial(View.GetFeatureLevel());
 		const EBlendMode BlendMode = Material->GetBlendMode();
 		const EMaterialLightingModel LightingModel = Material->GetLightingModel();
 
@@ -1029,12 +1032,12 @@ bool FShadowDepthDrawingPolicyFactory::DrawDynamicMesh(
 			const bool bPreShadow = Context.ShadowInfo->bPreShadow;
 			const bool bTwoSided = Material->IsTwoSided() || PrimitiveSceneProxy->CastsShadowAsTwoSided();
 
-			OverrideWithDefaultMaterialForShadowDepth(MaterialRenderProxy, Material, bReflectiveShadowmap);
+			OverrideWithDefaultMaterialForShadowDepth(MaterialRenderProxy, Material, bReflectiveShadowmap, View.GetFeatureLevel());
 
 			if(bReflectiveShadowmap)
 			{
 				FShadowDepthDrawingPolicy<true> DrawingPolicy(
-					MaterialRenderProxy->GetMaterial(GRHIFeatureLevel),
+					MaterialRenderProxy->GetMaterial(View.GetFeatureLevel()),
 					bDirectionalLight,
 					bOnePassPointLightShadow,
 					bPreShadow,
@@ -1044,7 +1047,7 @@ bool FShadowDepthDrawingPolicyFactory::DrawDynamicMesh(
 					Mesh.ReverseCulling
 					);
 
-				DrawingPolicy.DrawShared(&View,DrawingPolicy.CreateBoundShaderState());
+				DrawingPolicy.DrawShared(&View,DrawingPolicy.CreateBoundShaderState(View.GetFeatureLevel()));
 				for( int32 BatchElementIndex = 0, Num = Mesh.Elements.Num(); BatchElementIndex < Num; BatchElementIndex++)
 				{
 					DrawingPolicy.SetMeshRenderState(View,PrimitiveSceneProxy,Mesh,BatchElementIndex,bBackFace,FMeshDrawingPolicy::ElementDataType());
@@ -1054,7 +1057,7 @@ bool FShadowDepthDrawingPolicyFactory::DrawDynamicMesh(
 			else
 			{
 				FShadowDepthDrawingPolicy<false> DrawingPolicy(
-					MaterialRenderProxy->GetMaterial(GRHIFeatureLevel),
+					MaterialRenderProxy->GetMaterial(View.GetFeatureLevel()),
 					bDirectionalLight,
 					bOnePassPointLightShadow,
 					bPreShadow,
@@ -1064,7 +1067,7 @@ bool FShadowDepthDrawingPolicyFactory::DrawDynamicMesh(
 					Mesh.ReverseCulling
 					);
 
-				DrawingPolicy.DrawShared(&View,DrawingPolicy.CreateBoundShaderState());
+				DrawingPolicy.DrawShared(&View, DrawingPolicy.CreateBoundShaderState(View.GetFeatureLevel()));
 				for( int32 BatchElementIndex = 0; BatchElementIndex < Mesh.Elements.Num(); BatchElementIndex++)
 				{
 					DrawingPolicy.SetMeshRenderState(View,PrimitiveSceneProxy,Mesh,BatchElementIndex,bBackFace,FMeshDrawingPolicy::ElementDataType());
@@ -1084,11 +1087,11 @@ bool FShadowDepthDrawingPolicyFactory::DrawDynamicMesh(
 	FProjectedShadowInfo
 -----------------------------------------------------------------------------*/
 
-static void CheckShadowDepthMaterials(const FMaterialRenderProxy* InRenderProxy, const FMaterial* InMaterial, bool bReflectiveShadowmap)
+static void CheckShadowDepthMaterials(const FMaterialRenderProxy* InRenderProxy, const FMaterial* InMaterial, bool bReflectiveShadowmap, ERHIFeatureLevel::Type InFeatureLevel)
 {
 	const FMaterialRenderProxy* RenderProxy = InRenderProxy;
 	const FMaterial* Material = InMaterial;
-	OverrideWithDefaultMaterialForShadowDepth(RenderProxy, Material, bReflectiveShadowmap);
+	OverrideWithDefaultMaterialForShadowDepth(RenderProxy, Material, bReflectiveShadowmap, InFeatureLevel);
 	check(RenderProxy == InRenderProxy);
 	check(Material == InMaterial);
 }
@@ -1154,7 +1157,7 @@ void DrawMeshElements(FShadowDepthDrawingPolicy<bReflectiveShadowmap>& SharedDra
 #if UE_BUILD_DEBUG
 	// During shadow setup we should have already overridden materials with default material where needed.
 	// Make sure of it!
-	CheckShadowDepthMaterials(State.RenderProxy, State.MaterialResource, bReflectiveShadowmap);
+	CheckShadowDepthMaterials(State.RenderProxy, State.MaterialResource, bReflectiveShadowmap, View.GetFeatureLevel());
 #endif
 
 #if UE_BUILD_DEBUG
@@ -1222,7 +1225,7 @@ void DrawShadowMeshElements(const FViewInfo& View, const FProjectedShadowInfo& S
 			OldState = CurrentState;
 
 			SharedDrawingPolicy.UpdateElementState(CurrentState);
-			SharedDrawingPolicy.DrawShared(&View, SharedDrawingPolicy.CreateBoundShaderState());
+			SharedDrawingPolicy.DrawShared(&View, SharedDrawingPolicy.CreateBoundShaderState(View.GetFeatureLevel()));
 		}
 
 		DrawMeshElements(SharedDrawingPolicy, OldState, View, ShadowMesh.Mesh);
@@ -2912,7 +2915,7 @@ bool FDeferredShadingSceneRenderer::RenderProjectedShadows( const FLightSceneInf
 
 	bAttenuationBufferDirty |= RenderCachedPreshadows(LightSceneInfo);
 
-	if (GRHIFeatureLevel >= ERHIFeatureLevel::SM4)
+	if (Scene->GetFeatureLevel() >= ERHIFeatureLevel::SM4)
 	{
 		bAttenuationBufferDirty |= RenderOnePassPointLightShadows(LightSceneInfo, bRenderedTranslucentObjectShadows, bInjectedTranslucentVolume);
 	}
diff --git a/Engine/Source/Runtime/Renderer/Private/ShadowRendering.h b/Engine/Source/Runtime/Renderer/Private/ShadowRendering.h
--- a/Engine/Source/Runtime/Renderer/Private/ShadowRendering.h
+++ b/Engine/Source/Runtime/Renderer/Private/ShadowRendering.h
@@ -313,7 +313,7 @@ public:
 	 * as well as the shaders needed to draw the mesh
 	 * @return new bound shader state object
 	 */
-	FBoundShaderStateRHIRef CreateBoundShaderState();
+	FBoundShaderStateRHIRef CreateBoundShaderState(ERHIFeatureLevel::Type InFeatureLevel);
 
 	void SetMeshRenderState(
 		const FSceneView& View,
@@ -394,7 +394,7 @@ public:
 		FHitProxyId HitProxyId
 		);
 
-	static bool IsMaterialIgnored(const FMaterialRenderProxy* MaterialRenderProxy)
+	static bool IsMaterialIgnored(const FMaterialRenderProxy* MaterialRenderProxy, ERHIFeatureLevel::Type InFeatureLevel)
 	{
 		return false;
 	}
@@ -408,7 +408,8 @@ public:
 void OverrideWithDefaultMaterialForShadowDepth(
 	const FMaterialRenderProxy*& InOutMaterialRenderProxy, 
 	const FMaterial*& InOutMaterialResource,
-	bool bReflectiveShadowmap
+	bool bReflectiveShadowmap,
+	ERHIFeatureLevel::Type InFeatureLevel
 	);
 
 /** A single static mesh element for shadow depth rendering. */
diff --git a/Engine/Source/Runtime/Renderer/Private/ShadowSetup.cpp b/Engine/Source/Runtime/Renderer/Private/ShadowSetup.cpp
--- a/Engine/Source/Runtime/Renderer/Private/ShadowSetup.cpp
+++ b/Engine/Source/Runtime/Renderer/Private/ShadowSetup.cpp
@@ -889,7 +889,7 @@ void FProjectedShadowInfo::AddSubjectPrimitive(FPrimitiveSceneInfo* PrimitiveSce
 							if(((!IsTranslucentBlendMode(BlendMode)) && LightingModel != MLM_Unlit) || (bReflectiveShadowmap && Material->ShouldInjectEmissiveIntoLPV())) 
 							{
 								const bool bTwoSided = Material->IsTwoSided() || PrimitiveSceneInfo->Proxy->CastsShadowAsTwoSided();
-								OverrideWithDefaultMaterialForShadowDepth(MaterialRenderProxy, Material, bReflectiveShadowmap);
+								OverrideWithDefaultMaterialForShadowDepth(MaterialRenderProxy, Material, bReflectiveShadowmap, GRHIFeatureLevel);
 								SubjectMeshElements.Add(FShadowStaticMeshElement(MaterialRenderProxy, Material, &StaticMesh,bTwoSided));
 							}
 						}
@@ -1349,7 +1349,7 @@ void FDeferredShadingSceneRenderer::CreatePerObjectProjectedShadow(
 			}
 
 			if (bTranslucentRelevance
-				&& GRHIFeatureLevel >= ERHIFeatureLevel::SM4
+				&& Scene->GetFeatureLevel() >= ERHIFeatureLevel::SM4
 				&& bCreateTranslucentObjectShadow 
 				&& (bTranslucentShadowIsVisibleThisFrame || bShadowIsPotentiallyVisibleNextFrame))
 			{
diff --git a/Engine/Source/Runtime/Renderer/Private/StaticMeshDrawList.h b/Engine/Source/Runtime/Renderer/Private/StaticMeshDrawList.h
--- a/Engine/Source/Runtime/Renderer/Private/StaticMeshDrawList.h
+++ b/Engine/Source/Runtime/Renderer/Private/StaticMeshDrawList.h
@@ -166,6 +166,7 @@ private:
 		TArray<FElement>			Elements;
 		DrawingPolicyType			DrawingPolicy;
 		FBoundShaderStateRHIRef		BoundShaderState;
+		ERHIFeatureLevel::Type		FeatureLevel;
 
 		/** Used when sorting policy links */
 		FSphere						CachedBoundingSphere;
@@ -176,8 +177,9 @@ private:
 		TStaticMeshDrawList* DrawList;
 
 		/** Initialization constructor. */
-		FDrawingPolicyLink(TStaticMeshDrawList* InDrawList,const DrawingPolicyType& InDrawingPolicy):
+		FDrawingPolicyLink(TStaticMeshDrawList* InDrawList, const DrawingPolicyType& InDrawingPolicy, ERHIFeatureLevel::Type InFeatureLevel) :
 			DrawingPolicy(InDrawingPolicy),
+			FeatureLevel(InFeatureLevel),
 			DrawList(InDrawList)
 		{
 			CreateBoundShaderState();
@@ -195,7 +197,7 @@ private:
 
 		void CreateBoundShaderState()
 		{
-			BoundShaderState = DrawingPolicy.CreateBoundShaderState();
+			BoundShaderState = DrawingPolicy.CreateBoundShaderState(FeatureLevel);
 		}
 	};
 
@@ -235,11 +237,13 @@ public:
 	 * @param Mesh - The mesh to add.
 	 * @param PolicyData - The drawing policy data for the mesh.
 	 * @param InDrawingPolicy - The drawing policy to use to draw the mesh.
+	 * @param InFeatureLevel - The feature level of the scene we're rendering
 	 */
 	void AddMesh(
 		FStaticMesh* Mesh,
 		const ElementPolicyDataType& PolicyData,
-		const DrawingPolicyType& InDrawingPolicy
+		const DrawingPolicyType& InDrawingPolicy,
+		ERHIFeatureLevel::Type InFeatureLevel
 		);
 
 	/**
diff --git a/Engine/Source/Runtime/Renderer/Private/StaticMeshDrawList.inl b/Engine/Source/Runtime/Renderer/Private/StaticMeshDrawList.inl
--- a/Engine/Source/Runtime/Renderer/Private/StaticMeshDrawList.inl
+++ b/Engine/Source/Runtime/Renderer/Private/StaticMeshDrawList.inl
@@ -109,7 +109,8 @@ template<typename DrawingPolicyType>
 void TStaticMeshDrawList<DrawingPolicyType>::AddMesh(
 	FStaticMesh* Mesh,
 	const ElementPolicyDataType& PolicyData,
-	const DrawingPolicyType& InDrawingPolicy
+	const DrawingPolicyType& InDrawingPolicy,
+	ERHIFeatureLevel::Type InFeatureLevel
 	)
 {
 	// Check for an existing drawing policy matching the mesh's drawing policy.
@@ -117,7 +118,7 @@ void TStaticMeshDrawList<DrawingPolicyType>::AddMesh(
 	if(!DrawingPolicyLink)
 	{
 		// If no existing drawing policy matches the mesh, create a new one.
-		const FSetElementId DrawingPolicyLinkId = DrawingPolicySet.Add(FDrawingPolicyLink(this,InDrawingPolicy));
+		const FSetElementId DrawingPolicyLinkId = DrawingPolicySet.Add(FDrawingPolicyLink(this,InDrawingPolicy, InFeatureLevel));
 
 		DrawingPolicyLink = &DrawingPolicySet[DrawingPolicyLinkId];
 		DrawingPolicyLink->SetId = DrawingPolicyLinkId;
diff --git a/Engine/Source/Runtime/Renderer/Private/TiledDeferredLightRendering.cpp b/Engine/Source/Runtime/Renderer/Private/TiledDeferredLightRendering.cpp
--- a/Engine/Source/Runtime/Renderer/Private/TiledDeferredLightRendering.cpp
+++ b/Engine/Source/Runtime/Renderer/Private/TiledDeferredLightRendering.cpp
@@ -257,7 +257,7 @@ VARIATION1(0)			VARIATION1(1)
 
 bool FDeferredShadingSceneRenderer::CanUseTiledDeferred() const
 {
-	return GUseTiledDeferredShading != 0 && GRHIFeatureLevel >= ERHIFeatureLevel::SM5;
+	return GUseTiledDeferredShading != 0 && Scene->GetFeatureLevel() >= ERHIFeatureLevel::SM5;
 }
 
 bool FDeferredShadingSceneRenderer::ShouldUseTiledDeferred(int32 NumUnshadowedLights, int32 NumSimpleLights) const
diff --git a/Engine/Source/Runtime/Renderer/Private/TranslucentLighting.cpp b/Engine/Source/Runtime/Renderer/Private/TranslucentLighting.cpp
--- a/Engine/Source/Runtime/Renderer/Private/TranslucentLighting.cpp
+++ b/Engine/Source/Runtime/Renderer/Private/TranslucentLighting.cpp
@@ -187,7 +187,7 @@ public:
 		const FProjectedShadowInfo* ShadowInfo
 		)
 	{
-		FMeshMaterialShader::SetParameters(GetVertexShader(), MaterialRenderProxy, *MaterialRenderProxy->GetMaterial(GRHIFeatureLevel), View, ESceneRenderTargetsMode::DontSet);
+		FMeshMaterialShader::SetParameters(GetVertexShader(), MaterialRenderProxy, *MaterialRenderProxy->GetMaterial(View.GetFeatureLevel()), View, ESceneRenderTargetsMode::DontSet);
 		ShadowParameters.SetVertexShader(this, View, ShadowInfo, MaterialRenderProxy);
 	}
 
@@ -259,16 +259,17 @@ public:
 		)
 	{
 		const FPixelShaderRHIParamRef ShaderRHI = GetPixelShader();
+		const auto FeatureLevel = View.GetFeatureLevel();
 
 		//@todo - scene depth can be bound by the material for use in depth fades
 		// This is incorrect when rendering a shadowmap as it's not from the camera's POV
 		// Set the scene depth texture to something safe when rendering shadow depths
-		FMeshMaterialShader::SetParameters(ShaderRHI,MaterialRenderProxy,*MaterialRenderProxy->GetMaterial(GRHIFeatureLevel),View,ESceneRenderTargetsMode::NonSceneAlignedPass);
+		FMeshMaterialShader::SetParameters(ShaderRHI, MaterialRenderProxy, *MaterialRenderProxy->GetMaterial(FeatureLevel), View, ESceneRenderTargetsMode::NonSceneAlignedPass);
 
 		SetShaderValue(ShaderRHI, TranslInvMaxSubjectDepth, ShadowInfo->InvMaxSubjectDepth);
 
 		const float LocalToWorldScale = ShadowInfo->ParentSceneInfo->Proxy->GetLocalToWorld().GetScaleVector().GetMax();
-		const float TranslucentShadowStartOffsetValue = MaterialRenderProxy->GetMaterial(GRHIFeatureLevel)->GetTranslucentShadowStartOffset() * LocalToWorldScale;
+		const float TranslucentShadowStartOffsetValue = MaterialRenderProxy->GetMaterial(FeatureLevel)->GetTranslucentShadowStartOffset() * LocalToWorldScale;
 		SetShaderValue(ShaderRHI,TranslucentShadowStartOffset, TranslucentShadowStartOffsetValue / (ShadowInfo->MaxSubjectZ - ShadowInfo->MinSubjectZ));
 		TranslucencyProjectionParameters.Set(this);
 	}
@@ -362,7 +363,7 @@ public:
 	* as well as the shaders needed to draw the mesh
 	* @return new bound shader state object
 	*/
-	FBoundShaderStateRHIRef CreateBoundShaderState()
+	FBoundShaderStateRHIRef CreateBoundShaderState(ERHIFeatureLevel::Type InFeatureLevel)
 	{
 		return RHICreateBoundShaderState(
 			FMeshDrawingPolicy::GetVertexDeclaration(), 
@@ -423,17 +424,18 @@ public:
 		)
 	{
 		bool bDirty = false;
+		const auto FeatureLevel = View.GetFeatureLevel();
 
 		if (Mesh.CastShadow)
 		{
 			const FMaterialRenderProxy* MaterialRenderProxy = Mesh.MaterialRenderProxy;
-			const EBlendMode BlendMode = MaterialRenderProxy->GetMaterial(GRHIFeatureLevel)->GetBlendMode();
+			const EBlendMode BlendMode = MaterialRenderProxy->GetMaterial(FeatureLevel)->GetBlendMode();
 
 			// Only render translucent meshes into the Fourier opacity maps
 			if (IsTranslucentBlendMode(BlendMode))
 			{
-				FTranslucencyShadowDepthDrawingPolicy DrawingPolicy(Mesh.VertexFactory, MaterialRenderProxy, *MaterialRenderProxy->GetMaterial(GRHIFeatureLevel), DrawingContext.bDirectionalLight);
-				DrawingPolicy.DrawShared(&View, DrawingPolicy.CreateBoundShaderState());
+				FTranslucencyShadowDepthDrawingPolicy DrawingPolicy(Mesh.VertexFactory, MaterialRenderProxy, *MaterialRenderProxy->GetMaterial(FeatureLevel), DrawingContext.bDirectionalLight);
+				DrawingPolicy.DrawShared(&View, DrawingPolicy.CreateBoundShaderState(FeatureLevel));
 
 				for (int32 BatchElementIndex = 0; BatchElementIndex < Mesh.Elements.Num(); BatchElementIndex++)
 				{
@@ -471,9 +473,9 @@ public:
 		return bDirty;
 	}
 
-	static bool IsMaterialIgnored(const FMaterialRenderProxy* MaterialRenderProxy) 
+	static bool IsMaterialIgnored(const FMaterialRenderProxy* MaterialRenderProxy, ERHIFeatureLevel::Type InFeatureLevel)
 	{ 
-		return !IsTranslucentBlendMode(MaterialRenderProxy->GetMaterial(GRHIFeatureLevel)->GetBlendMode());
+		return !IsTranslucentBlendMode(MaterialRenderProxy->GetMaterial(InFeatureLevel)->GetBlendMode());
 	}
 };
 
@@ -831,7 +833,7 @@ public:
 		
 		const FPixelShaderRHIParamRef ShaderRHI = GetPixelShader();
 
-		FMaterialShader::SetParameters(ShaderRHI, MaterialProxy, *MaterialProxy->GetMaterial(GRHIFeatureLevel), View, false, ESceneRenderTargetsMode::SetTextures);
+		FMaterialShader::SetParameters(ShaderRHI, MaterialProxy, *MaterialProxy->GetMaterial(View.GetFeatureLevel()), View, false, ESceneRenderTargetsMode::SetTextures);
 		
 		FSphere ShadowBounds = ShadowMap ? ShadowMap->ShadowBounds : FSphere(FVector(0,0,0), HALF_WORLD_MAX);
 		SetShaderValue(ShaderRHI, CascadeBounds, ShadowBounds);
@@ -1456,7 +1458,7 @@ void SetInjectionShader(
 {
 	check(ShadowMap || !bShadowed);
 
-	const FMaterialShaderMap* MaterialShaderMap = MaterialProxy->GetMaterial(GRHIFeatureLevel)->GetRenderingThreadShaderMap();
+	const FMaterialShaderMap* MaterialShaderMap = MaterialProxy->GetMaterial(View.GetFeatureLevel())->GetRenderingThreadShaderMap();
 	FMaterialShader* PixelShader = NULL;
 
 	const bool Directional = InjectionType == LightType_Directional;
@@ -1572,16 +1574,18 @@ static void AddLightForInjection(
 	{
 		const FVisibleLightInfo& VisibleLightInfo = SceneRenderer.VisibleLightInfos[LightSceneInfo.Id];
 
+		const ERHIFeatureLevel::Type FeatureLevel = SceneRenderer.Scene->GetFeatureLevel();
+
 		const bool bApplyLightFunction = (SceneRenderer.ViewFamily.EngineShowFlags.LightFunctions &&
 			LightSceneInfo.Proxy->GetLightFunctionMaterial() && 
-			LightSceneInfo.Proxy->GetLightFunctionMaterial()->GetMaterial(GRHIFeatureLevel)->IsLightFunction());
+			LightSceneInfo.Proxy->GetLightFunctionMaterial()->GetMaterial(FeatureLevel)->IsLightFunction());
 
 		const FMaterialRenderProxy* MaterialProxy = bApplyLightFunction ? 
 			LightSceneInfo.Proxy->GetLightFunctionMaterial() : 
 			UMaterial::GetDefaultMaterial(MD_LightFunction)->GetRenderProxy(false);
 
 		// Skip rendering if the DefaultLightFunctionMaterial isn't compiled yet
-		if (MaterialProxy->GetMaterial(GRHIFeatureLevel)->IsLightFunction())
+		if (MaterialProxy->GetMaterial(FeatureLevel)->IsLightFunction())
 		{
 			FTranslucentLightInjectionData InjectionData;
 			InjectionData.LightSceneInfo = &LightSceneInfo;
diff --git a/Engine/Source/Runtime/Renderer/Private/TranslucentRendering.cpp b/Engine/Source/Runtime/Renderer/Private/TranslucentRendering.cpp
--- a/Engine/Source/Runtime/Renderer/Private/TranslucentRendering.cpp
+++ b/Engine/Source/Runtime/Renderer/Private/TranslucentRendering.cpp
@@ -146,9 +146,9 @@ public:
 	}
 	FCopySceneColorPS() {}
 
-	void SetParameters()
+	void SetParameters(const FViewInfo& View)
 	{
-		SceneTextureParameters.Set(GetPixelShader());
+		SceneTextureParameters.Set(GetPixelShader(), View);
 	}
 
 	virtual bool Serialize(FArchive& Ar)
@@ -182,7 +182,7 @@ void FDeferredShadingSceneRenderer::CopySceneColor(const FViewInfo& View, const
 	TShaderMapRef<FCopySceneColorPS> PixelShader(GetGlobalShaderMap());
 	SetGlobalBoundShaderState(CopySceneColorBoundShaderState, GFilterVertexDeclaration.VertexDeclarationRHI, *ScreenVertexShader, *PixelShader);
 
-	PixelShader->SetParameters();
+	PixelShader->SetParameters(View);
 
 	DrawRectangle( 
 		0, 0, 
@@ -271,7 +271,7 @@ public:
 			);
 		DrawingPolicy.DrawShared(
 			&View,
-			DrawingPolicy.CreateBoundShaderState()
+			DrawingPolicy.CreateBoundShaderState(View.GetFeatureLevel())
 			);
 
 		int32 BatchElementIndex = 0;
@@ -315,7 +315,7 @@ bool FTranslucencyDrawingPolicyFactory::DrawMesh(
 	bool bDirty = false;
 
 	// Determine the mesh's material and blend mode.
-	const FMaterial* Material = Mesh.MaterialRenderProxy->GetMaterial(GRHIFeatureLevel);
+	const FMaterial* Material = Mesh.MaterialRenderProxy->GetMaterial(View.GetFeatureLevel());
 	const EBlendMode BlendMode = Material->GetBlendMode();
 	const EMaterialLightingModel LightingModel = Material->GetLightingModel();
 
@@ -360,7 +360,8 @@ bool FTranslucencyDrawingPolicyFactory::DrawMesh(
 				HitProxyId,
 				DrawingContext.TranslucentSelfShadow,
 				PrimitiveSceneProxy && PrimitiveSceneProxy->CastsVolumetricTranslucentShadow()
-			)
+			),
+			View.GetFeatureLevel()
 		);
 
 		if (bDisableDepthTest || bEnableResponsiveAA)
@@ -520,7 +521,7 @@ void FTranslucentPrimSet::RenderPrimitive(
 				if (View.StaticMeshVisibilityMap[StaticMesh.Id]
 					// Only render static mesh elements using translucent materials
 					&& StaticMesh.IsTranslucent() 
-					&& (StaticMesh.MaterialRenderProxy->GetMaterial(GRHIFeatureLevel)->IsSeparateTranslucencyEnabled() == bSeparateTranslucencyPass))
+					&& (StaticMesh.MaterialRenderProxy->GetMaterial(View.GetFeatureLevel())->IsSeparateTranslucencyEnabled() == bSeparateTranslucencyPass))
 				{
 					FTranslucencyDrawingPolicyFactory::DrawStaticMesh(
 						View,
@@ -549,8 +550,10 @@ void FTranslucentPrimSet::AddScenePrimitive(FPrimitiveSceneInfo* PrimitiveSceneI
 	SortKey = (PrimitiveSceneInfo->Proxy->GetBounds().Origin - ViewInfo.ViewMatrices.ViewOrigin).Size();
 	// UE4_TODO: also account for DPG in the sort key.
 
+	const auto FeatureLevel = ViewInfo.GetFeatureLevel();
+
 	if(bUseSeparateTranslucency 
-		&& GRHIFeatureLevel >= ERHIFeatureLevel::SM3)
+		&& FeatureLevel >= ERHIFeatureLevel::SM3)
 	{
 		// add to list of translucent prims that use scene color
 		new(SortedSeparateTranslucencyPrims) FSortedPrim(PrimitiveSceneInfo,SortKey,PrimitiveSceneInfo->Proxy->GetTranslucencySortPriority());
@@ -558,7 +561,7 @@ void FTranslucentPrimSet::AddScenePrimitive(FPrimitiveSceneInfo* PrimitiveSceneI
 
 	if (bUseNormalTranslucency 
 		// Force separate translucency to be rendered normally if the feature level does not support separate translucency
-		|| (bUseSeparateTranslucency && GRHIFeatureLevel < ERHIFeatureLevel::SM3))
+		|| (bUseSeparateTranslucency && FeatureLevel < ERHIFeatureLevel::SM3))
 	{
 		// add to list of translucent prims
 		new(SortedPrims) FSortedPrim(PrimitiveSceneInfo,SortKey,PrimitiveSceneInfo->Proxy->GetTranslucencySortPriority());
diff --git a/Engine/Source/Runtime/Renderer/Private/TranslucentRendering.h b/Engine/Source/Runtime/Renderer/Private/TranslucentRendering.h
--- a/Engine/Source/Runtime/Renderer/Private/TranslucentRendering.h
+++ b/Engine/Source/Runtime/Renderer/Private/TranslucentRendering.h
@@ -51,10 +51,10 @@ public:
 		FHitProxyId HitProxyId
 		);
 
-	static bool IsMaterialIgnored(const FMaterialRenderProxy* MaterialRenderProxy)
+	static bool IsMaterialIgnored(const FMaterialRenderProxy* MaterialRenderProxy, ERHIFeatureLevel::Type InFeatureLevel)
 	{
 		// Note: blend mode does not depend on the feature level.
-		return !IsTranslucentBlendMode(MaterialRenderProxy->GetMaterial(GRHIFeatureLevel)->GetBlendMode());
+		return !IsTranslucentBlendMode(MaterialRenderProxy->GetMaterial(InFeatureLevel)->GetBlendMode());
 	}
 
 
@@ -115,10 +115,10 @@ public:
 		FHitProxyId HitProxyId
 		);
 
-	static bool IsMaterialIgnored(const FMaterialRenderProxy* MaterialRenderProxy)
+	static bool IsMaterialIgnored(const FMaterialRenderProxy* MaterialRenderProxy, ERHIFeatureLevel::Type InFeatureLevel)
 	{
 		// Note: blend mode does not depend on the feature level.
-		return !IsTranslucentBlendMode(MaterialRenderProxy->GetMaterial(GRHIFeatureLevel)->GetBlendMode());
+		return !IsTranslucentBlendMode(MaterialRenderProxy->GetMaterial(InFeatureLevel)->GetBlendMode());
 	}
 };
 
diff --git a/Engine/Source/Runtime/Renderer/Private/VelocityRendering.cpp b/Engine/Source/Runtime/Renderer/Private/VelocityRendering.cpp
--- a/Engine/Source/Runtime/Renderer/Private/VelocityRendering.cpp
+++ b/Engine/Source/Runtime/Renderer/Private/VelocityRendering.cpp
@@ -17,7 +17,7 @@ public:
 
 	void SetParameters(const FVertexFactory* VertexFactory,const FMaterialRenderProxy* MaterialRenderProxy,const FViewInfo& View)
 	{
-		FMeshMaterialShader::SetParameters(GetVertexShader(),MaterialRenderProxy,*MaterialRenderProxy->GetMaterial(GRHIFeatureLevel),View,ESceneRenderTargetsMode::DontSet);
+		FMeshMaterialShader::SetParameters(GetVertexShader(), MaterialRenderProxy, *MaterialRenderProxy->GetMaterial(View.GetFeatureLevel()), View, ESceneRenderTargetsMode::DontSet);
 	}
 
 	void SetMesh(const FVertexFactory* VertexFactory, const FMeshBatch& Mesh, int32 BatchElementIndex, const FViewInfo& View, const FPrimitiveSceneProxy* Proxy, const FMatrix& InPreviousLocalToWorld)
@@ -93,7 +93,7 @@ public:
 
 	void SetParameters(const FMaterialRenderProxy* MaterialRenderProxy,const FViewInfo& View)
 	{
-		FMeshMaterialShader::SetParameters(GetDomainShader(), MaterialRenderProxy, *MaterialRenderProxy->GetMaterial(GRHIFeatureLevel), View, ESceneRenderTargetsMode::DontSet);
+		FMeshMaterialShader::SetParameters(GetDomainShader(), MaterialRenderProxy, *MaterialRenderProxy->GetMaterial(View.GetFeatureLevel()), View, ESceneRenderTargetsMode::DontSet);
 	}
 
 protected:
@@ -146,7 +146,7 @@ public:
 
 	void SetParameters(const FVertexFactory* VertexFactory,const FMaterialRenderProxy* MaterialRenderProxy,const FViewInfo& View)
 	{
-		FMeshMaterialShader::SetParameters(GetPixelShader(),MaterialRenderProxy,*MaterialRenderProxy->GetMaterial(GRHIFeatureLevel),View,ESceneRenderTargetsMode::DontSet);
+		FMeshMaterialShader::SetParameters(GetPixelShader(), MaterialRenderProxy, *MaterialRenderProxy->GetMaterial(View.GetFeatureLevel()), View, ESceneRenderTargetsMode::DontSet);
 	}
 
 	void SetMesh(const FVertexFactory* VertexFactory, const FMeshBatch& Mesh,int32 BatchElementIndex,const FViewInfo& View, const FPrimitiveSceneProxy* Proxy, bool bBackFace)
@@ -312,7 +312,7 @@ bool FVelocityDrawingPolicy::HasVelocity(const FViewInfo& View, const FPrimitive
 	return true;
 }
 
-FBoundShaderStateRHIRef FVelocityDrawingPolicy::CreateBoundShaderState()
+FBoundShaderStateRHIRef FVelocityDrawingPolicy::CreateBoundShaderState(ERHIFeatureLevel::Type InFeatureLevel)
 {
 	FBoundShaderStateRHIRef BoundShaderState;
 
@@ -347,8 +347,9 @@ void FVelocityDrawingPolicyFactory::AddStaticMesh(FScene* Scene, FStaticMesh* St
 	// Velocity only needs to be directly rendered for movable meshes.
 	if(StaticMesh->PrimitiveSceneInfo->Proxy->IsMovable())
 	{
+		const auto FeatureLevel = Scene->GetFeatureLevel();
 	    const FMaterialRenderProxy* MaterialRenderProxy = StaticMesh->MaterialRenderProxy;
-	    const FMaterial* Material = MaterialRenderProxy->GetMaterial(GRHIFeatureLevel);
+		const FMaterial* Material = MaterialRenderProxy->GetMaterial(FeatureLevel);
 	    EBlendMode BlendMode = Material->GetBlendMode();
 	    if(BlendMode == BLEND_Opaque || BlendMode == BLEND_Masked)
 	    {
@@ -358,11 +359,11 @@ void FVelocityDrawingPolicyFactory::AddStaticMesh(FScene* Scene, FStaticMesh* St
 			    MaterialRenderProxy = UMaterial::GetDefaultMaterial(MD_Surface)->GetRenderProxy(false);
 		    }
 
-			FVelocityDrawingPolicy DrawingPolicy(StaticMesh->VertexFactory, MaterialRenderProxy,*MaterialRenderProxy->GetMaterial(GRHIFeatureLevel));
+			FVelocityDrawingPolicy DrawingPolicy(StaticMesh->VertexFactory, MaterialRenderProxy, *MaterialRenderProxy->GetMaterial(FeatureLevel));
 			if (DrawingPolicy.SupportsVelocity())
 			{
 				// Add the static mesh to the depth-only draw list.
-				Scene->VelocityDrawList.AddMesh(StaticMesh, FVelocityDrawingPolicy::ElementDataType(), DrawingPolicy);
+				Scene->VelocityDrawList.AddMesh(StaticMesh, FVelocityDrawingPolicy::ElementDataType(), DrawingPolicy, Scene->GetFeatureLevel());
 			}
 	    }
 	}
@@ -379,8 +380,9 @@ bool FVelocityDrawingPolicyFactory::DrawDynamicMesh(
 	)
 {
 	// Only draw opaque materials in the depth pass.
+	const auto FeatureLevel = View.GetFeatureLevel();
 	const FMaterialRenderProxy* MaterialRenderProxy = Mesh.MaterialRenderProxy;
-	const FMaterial* Material = MaterialRenderProxy->GetMaterial(GRHIFeatureLevel);
+	const FMaterial* Material = MaterialRenderProxy->GetMaterial(FeatureLevel);
 	EBlendMode BlendMode = Material->GetBlendMode();
 
 	if(BlendMode == BLEND_Opaque || BlendMode == BLEND_Masked)
@@ -393,10 +395,10 @@ bool FVelocityDrawingPolicyFactory::DrawDynamicMesh(
 			// Default material doesn't handle masked, and doesn't have the correct bIsTwoSided setting.
 			MaterialRenderProxy = UMaterial::GetDefaultMaterial(MD_Surface)->GetRenderProxy(false);
 		}
-		FVelocityDrawingPolicy DrawingPolicy(Mesh.VertexFactory, MaterialRenderProxy, *MaterialRenderProxy->GetMaterial(GRHIFeatureLevel));
+		FVelocityDrawingPolicy DrawingPolicy(Mesh.VertexFactory, MaterialRenderProxy, *MaterialRenderProxy->GetMaterial(FeatureLevel));
 		if(DrawingPolicy.SupportsVelocity())
 		{			
-			DrawingPolicy.DrawShared(&View,DrawingPolicy.CreateBoundShaderState());
+			DrawingPolicy.DrawShared(&View,DrawingPolicy.CreateBoundShaderState(FeatureLevel));
 			for(int32 BatchElementIndex = 0; BatchElementIndex < Mesh.Elements.Num(); ++BatchElementIndex)
 			{
 				DrawingPolicy.SetMeshRenderState(View, PrimitiveSceneProxy, Mesh, BatchElementIndex, bBackFace, FMeshDrawingPolicy::ElementDataType());
@@ -425,7 +427,7 @@ int32 GetMotionBlurQualityFromCVar()
 
 bool IsMotionBlurEnabled(const FViewInfo& View)
 {
-	if(!(GRHIFeatureLevel >= ERHIFeatureLevel::SM4))
+	if (!(View.GetFeatureLevel() >= ERHIFeatureLevel::SM4))
 	{
 		return false; 
 	}
diff --git a/Engine/Source/Runtime/Renderer/Private/VelocityRendering.h b/Engine/Source/Runtime/Renderer/Private/VelocityRendering.h
--- a/Engine/Source/Runtime/Renderer/Private/VelocityRendering.h
+++ b/Engine/Source/Runtime/Renderer/Private/VelocityRendering.h
@@ -37,7 +37,7 @@ public:
 		const ElementDataType& ElementData
 		) const;
 
-	FBoundShaderStateRHIRef CreateBoundShaderState();
+	FBoundShaderStateRHIRef CreateBoundShaderState(ERHIFeatureLevel::Type InFeatureLevel);
 
 	friend int32 Compare(const FVelocityDrawingPolicy& A, const FVelocityDrawingPolicy& B);
 
diff --git a/Engine/Source/Runtime/Renderer/Public/RendererInterface.h b/Engine/Source/Runtime/Renderer/Public/RendererInterface.h
--- a/Engine/Source/Runtime/Renderer/Public/RendererInterface.h
+++ b/Engine/Source/Runtime/Renderer/Public/RendererInterface.h
@@ -384,7 +384,7 @@ public:
 	 * @param World - An optional world to associate with the scene.
 	 * @param bInRequiresHitProxies - Indicates that hit proxies should be rendered in the scene.
 	 */
-	virtual FSceneInterface* AllocateScene(UWorld* World, bool bInRequiresHitProxies) = 0;
+	virtual FSceneInterface* AllocateScene(UWorld* World, bool bInRequiresHitProxies, ERHIFeatureLevel::Type InFeatureLevel) = 0;
 	
 	virtual void RemoveScene(FSceneInterface* Scene) = 0;
 
diff --git a/Engine/Source/Runtime/Windows/D3D11RHI/Private/D3D11Device.cpp b/Engine/Source/Runtime/Windows/D3D11RHI/Private/D3D11Device.cpp
--- a/Engine/Source/Runtime/Windows/D3D11RHI/Private/D3D11Device.cpp
+++ b/Engine/Source/Runtime/Windows/D3D11RHI/Private/D3D11Device.cpp
@@ -83,17 +83,17 @@ FD3D11DynamicRHI::FD3D11DynamicRHI(IDXGIFactory* InDXGIFactory,D3D_FEATURE_LEVEL
 	// ES2 feature level emulation in D3D11
 	if (FParse::Param(FCommandLine::Get(), TEXT("FeatureLevelES2")))
 	{
-		GRHIFeatureLevel = ERHIFeatureLevel::ES2;
+		SetMaxRHIFeatureLevel(ERHIFeatureLevel::ES2);
 		GRHIShaderPlatform = SP_PCD3D_ES2;
 	}
 	else if(FeatureLevel == D3D_FEATURE_LEVEL_11_0)
 	{
-		GRHIFeatureLevel = ERHIFeatureLevel::SM5;
+		SetMaxRHIFeatureLevel(ERHIFeatureLevel::SM5);
 		GRHIShaderPlatform = SP_PCD3D_SM5;
 	}
 	else if(FeatureLevel == D3D_FEATURE_LEVEL_10_0)
 	{
-		GRHIFeatureLevel = ERHIFeatureLevel::SM4;
+		SetMaxRHIFeatureLevel(ERHIFeatureLevel::SM4);
 		GRHIShaderPlatform = SP_PCD3D_SM4;
 	}
 
@@ -304,7 +304,7 @@ void FD3D11DynamicRHI::RHIGetSupportedResolution( uint32 &Width, uint32 &Height
 void FD3D11DynamicRHI::GetBestSupportedMSAASetting( DXGI_FORMAT PlatformFormat, uint32 MSAACount, uint32& OutBestMSAACount, uint32& OutMSAAQualityLevels )
 {
 	//  We disable MSAA for Feature level 10
-	if(GRHIFeatureLevel == ERHIFeatureLevel::SM4)
+	if (GRHIFeatureLevel == ERHIFeatureLevel::SM4)
 	{
 		OutBestMSAACount = 1;
 		OutMSAAQualityLevels = 0;
