* removed LightAccumulation for SceneRenderTargets, alloc on demand
* fixed Splitscreen ReflectionEnv rendering

diff --git a/Engine/Source/Runtime/Renderer/Private/DeferredShadingRenderer.h b/Engine/Source/Runtime/Renderer/Private/DeferredShadingRenderer.h
--- a/Engine/Source/Runtime/Renderer/Private/DeferredShadingRenderer.h
+++ b/Engine/Source/Runtime/Renderer/Private/DeferredShadingRenderer.h
@@ -323,5 +323,16 @@ private:
 	/** Filters the translucency lighting volumes to reduce aliasing. */
 	void FilterTranslucentVolumeLighting();
 
+	/** Output SpecularColor * IndirectDiffuseGI for metals so they are not black in reflections */
+	void RenderReflectionCaptureSpecularBounceForAllViews();
+
+	/** Render image based reflections (SSR, Env, SkyLight) without compute shaders */
+	void RenderImageBasedReflectionsSM4ForAllViews();
+
+	/** Render image based reflections (SSR, Env, SkyLight) with compute shaders */
+	void RenderImageBasedReflectionsSM5ForAllViews();
+
+	bool ShouldDoReflectionEnvironment() const;
+
 	friend class FTranslucentPrimSet;
 };
diff --git a/Engine/Source/Runtime/Renderer/Private/PostProcess/SceneRenderTargets.cpp b/Engine/Source/Runtime/Renderer/Private/PostProcess/SceneRenderTargets.cpp
--- a/Engine/Source/Runtime/Renderer/Private/PostProcess/SceneRenderTargets.cpp
+++ b/Engine/Source/Runtime/Renderer/Private/PostProcess/SceneRenderTargets.cpp
@@ -1057,20 +1057,6 @@ void FSceneRenderTargets::AllocateDeferredShadingPathRenderTargets()
 		GRenderTargetPool.FindFreeElement(Desc, SkySHIrradianceMap, TEXT("SkySHIrradianceMap"));
 	}
 
-	{
-		EPixelFormat SceneColorBufferFormat = GetSceneColorFormat();
-
-		uint32 LightAccumulationUAVFlag = GRHIFeatureLevel == ERHIFeatureLevel::SM5 ? TexCreate_UAV : 0;
-		FPooledRenderTargetDesc Desc(FPooledRenderTargetDesc::Create2DDesc(BufferSize, SceneColorBufferFormat, TexCreate_None, TexCreate_RenderTargetable | LightAccumulationUAVFlag, false));
-
-		if (GRHIFeatureLevel >= ERHIFeatureLevel::SM5)
-		{
-			Desc.TargetableFlags |= TexCreate_UAV;
-		}
-
-		GRenderTargetPool.FindFreeElement(Desc, LightAccumulation, TEXT("LightAccumulation"));
-	}
-
 	if (GRHIFeatureLevel >= ERHIFeatureLevel::SM5)
 	{
 		// Create the reflective shadow map textures for LightPropagationVolume feature
@@ -1156,7 +1142,6 @@ void FSceneRenderTargets::ReleaseDynamicRHI()
 	ScreenSpaceAO.SafeRelease();
 	LightAttenuation.SafeRelease();
 	CustomDepth.SafeRelease();
-	LightAccumulation.SafeRelease();
 	ReflectiveShadowMapNormal.SafeRelease();
 	ReflectiveShadowMapDiffuse.SafeRelease();
 	ReflectiveShadowMapDepth.SafeRelease();
diff --git a/Engine/Source/Runtime/Renderer/Private/PostProcess/SceneRenderTargets.h b/Engine/Source/Runtime/Renderer/Private/PostProcess/SceneRenderTargets.h
--- a/Engine/Source/Runtime/Renderer/Private/PostProcess/SceneRenderTargets.h
+++ b/Engine/Source/Runtime/Renderer/Private/PostProcess/SceneRenderTargets.h
@@ -393,8 +393,6 @@ public:
 	TRefCountPtr<IPooledRenderTarget> ScreenSpaceAO;
 	// used by the CustomDepth material feature, is allocated on demand or if r.CustomDepth is 2
 	TRefCountPtr<IPooledRenderTarget> CustomDepth;
-	// Used for accumulating tiled deferred VPL indirect lighting
-	TRefCountPtr<IPooledRenderTarget> LightAccumulation;
 	// Render target for per-object shadow depths.
 	TRefCountPtr<IPooledRenderTarget> ShadowDepthZ;
 	// Cache of preshadow depths
diff --git a/Engine/Source/Runtime/Renderer/Private/ReflectionEnvironment.cpp b/Engine/Source/Runtime/Renderer/Private/ReflectionEnvironment.cpp
--- a/Engine/Source/Runtime/Renderer/Private/ReflectionEnvironment.cpp
+++ b/Engine/Source/Runtime/Renderer/Private/ReflectionEnvironment.cpp
@@ -22,14 +22,12 @@ const int32 GReflectionEnvironmentTileSizeX = 16;
 const int32 GReflectionEnvironmentTileSizeY = 16;
 extern ENGINE_API int32 GReflectionCaptureSize;
 
-static int32 GDiffuseFromCaptures = 0;
-static FAutoConsoleVariableRef CVarDiffuseFromCaptures(
+static TAutoConsoleVariable<int32> CVarDiffuseFromCaptures(
 	TEXT("r.DiffuseFromCaptures"),
-	GDiffuseFromCaptures,
+	0,
 	TEXT("Apply indirect diffuse lighting from captures instead of lightmaps.\n")
 	TEXT(" 0 is off (default), 1 is on"),
-	ECVF_RenderThreadSafe
-	);
+	ECVF_RenderThreadSafe);
 
 // to avoid having direct access from many places
 static int GetReflectionEnvironmentCVar()
@@ -185,6 +183,10 @@ public:
 		DeferredParameters.Set(ShaderRHI, View);
 
 		FScene* Scene = (FScene*)View.Family->Scene;
+
+		check(Scene->ReflectionSceneData.CubemapArray.IsValid());
+		check(Scene->ReflectionSceneData.CubemapArray.GetRenderTarget().IsValid());
+
 		FSceneRenderTargetItem& CubemapArray = Scene->ReflectionSceneData.CubemapArray.GetRenderTarget();
 
 		SetTextureParameter(
@@ -192,8 +194,7 @@ public:
 			ReflectionEnvironmentColorTexture, 
 			ReflectionEnvironmentColorSampler, 
 			TStaticSamplerState<SF_Trilinear, AM_Clamp, AM_Clamp, AM_Clamp>::GetRHI(), 
-			// Set NULL when not valid to prevent d3d debug warnings
-			Scene->ReflectionSceneData.CubemapArray.GetRenderTarget().IsValid() ? CubemapArray.ShaderResourceTexture : NULL);
+			CubemapArray.ShaderResourceTexture);
 
 		SetTextureParameter( ShaderRHI, ScreenSpaceReflections, SSRTexture );
 
@@ -205,7 +206,6 @@ public:
 		static TArray<FReflectionCaptureSortData> SortData;
 		SortData.Reset(Scene->ReflectionSceneData.RegisteredReflectionCaptures.Num());
 
-		check(Scene->ReflectionSceneData.CubemapArray.GetRenderTarget().IsValid());
 		const int32 MaxCubemaps = Scene->ReflectionSceneData.CubemapArray.GetMaxCubemaps();
 
 		// Pack only visible reflection captures into the uniform buffer, each with an index to its cubemap array entry
@@ -384,7 +384,7 @@ public:
 		
 		SetTextureParameter( ShaderRHI, ReflectionEnvTexture, ReflectionEnvSampler, TStaticSamplerState<SF_Point>::GetRHI(), ReflectionEnv );
 		SetTextureParameter( ShaderRHI, ScreenSpaceReflectionsTexture, ScreenSpaceReflectionsSampler, TStaticSamplerState<SF_Point>::GetRHI(), ScreenSpaceReflections );
-		SetTextureParameter(ShaderRHI, PreIntegratedGF, PreIntegratedGFSampler, TStaticSamplerState<SF_Bilinear, AM_Clamp, AM_Clamp, AM_Clamp>::GetRHI(), GSystemTextures.PreintegratedGF->GetRenderTargetItem().ShaderResourceTexture);
+		SetTextureParameter( ShaderRHI, PreIntegratedGF, PreIntegratedGFSampler, TStaticSamplerState<SF_Bilinear, AM_Clamp, AM_Clamp, AM_Clamp>::GetRHI(), GSystemTextures.PreintegratedGF->GetRenderTargetItem().ShaderResourceTexture );
 	}
 
 	// FShader interface.
@@ -536,38 +536,57 @@ private:
 IMPLEMENT_SHADER_TYPE(template<>,TStandardDeferredReflectionPS<true>,TEXT("ReflectionEnvironmentShaders"),TEXT("StandardDeferredReflectionPS"),SF_Pixel);
 IMPLEMENT_SHADER_TYPE(template<>,TStandardDeferredReflectionPS<false>,TEXT("ReflectionEnvironmentShaders"),TEXT("StandardDeferredReflectionPS"),SF_Pixel);
 
+void FDeferredShadingSceneRenderer::RenderReflectionCaptureSpecularBounceForAllViews()
+{		
+	// We're currently capturing a reflection capture, output SpecularColor * IndirectDiffuseGI for metals so they are not black in reflections,
+	// Since we don't have multiple bounce specular reflections
+	GSceneRenderTargets.BeginRenderingSceneColor();
+	RHISetRasterizerState( TStaticRasterizerState< FM_Solid, CM_None >::GetRHI() );
+	RHISetDepthStencilState( TStaticDepthStencilState< false, CF_Always >::GetRHI() );
+	RHISetBlendState( TStaticBlendState< CW_RGB, BO_Add, BF_One, BF_One >::GetRHI() );
 
-/** Renders deferred reflections for the scene. */
-void FDeferredShadingSceneRenderer::RenderDeferredReflections()
-{
-	if (IsSimpleDynamicLightingEnabled() || ViewFamily.EngineShowFlags.VisualizeLightCulling)
-	{
-		return;
-	}
+	TShaderMapRef< FPostProcessVS > VertexShader( GetGlobalShaderMap() );
+	TShaderMapRef< FReflectionCaptureSpecularBouncePS > PixelShader( GetGlobalShaderMap() );
 
-	bool bAnyViewIsReflectionCapture = false;
-	for (int32 ViewIndex = 0; ViewIndex < Views.Num(); ViewIndex++)
+	static FGlobalBoundShaderState BoundShaderState;
+	SetGlobalBoundShaderState( BoundShaderState, GFilterVertexDeclaration.VertexDeclarationRHI, *VertexShader, *PixelShader );
+
+	for (int32 ViewIndex = 0, Num = Views.Num(); ViewIndex < Num; ViewIndex++)
 	{
 		const FViewInfo& View = Views[ViewIndex];
-		bAnyViewIsReflectionCapture = bAnyViewIsReflectionCapture || View.bIsReflectionCapture;
+
+		RHISetViewport( View.ViewRect.Min.X, View.ViewRect.Min.Y, 0.0f, View.ViewRect.Max.X, View.ViewRect.Max.Y, 1.0f );
+
+		PixelShader->SetParameters( View );
+
+		DrawRectangle( 
+			0, 0, 
+			View.ViewRect.Width(), View.ViewRect.Height(),
+			0, 0, 
+			View.ViewRect.Width(), View.ViewRect.Height(),
+			FIntPoint(View.ViewRect.Width(), View.ViewRect.Height()),
+			GSceneRenderTargets.GetBufferSizeXY(),
+			*VertexShader,
+			EDRF_UseTriangleOptimization);
 	}
+}
 
+bool FDeferredShadingSceneRenderer::ShouldDoReflectionEnvironment() const
+{
 	const ERHIFeatureLevel::Type FeatureLevel = Scene->GetFeatureLevel();
 
-	const uint32 bReflectionEnv = IsReflectionEnvironmentAvailable(FeatureLevel)
+	return IsReflectionEnvironmentAvailable(FeatureLevel)
 		&& Scene->ReflectionSceneData.RegisteredReflectionCaptures.Num()
 		&& ViewFamily.EngineShowFlags.ReflectionEnvironment;
+}
 
-	const uint32 bSkyLight = Scene->SkyLight
-		&& Scene->SkyLight->ProcessedTexture
-		&& ViewFamily.EngineShowFlags.SkyLighting;
-
-	const uint32 bUseLightmaps = GDiffuseFromCaptures == 0;
+void FDeferredShadingSceneRenderer::RenderImageBasedReflectionsSM5ForAllViews()
+{
+	const uint32 bUseLightmaps = CVarDiffuseFromCaptures.GetValueOnRenderThread() == 0;
 
-	const bool bReflectionsWithCompute = FeatureLevel == ERHIFeatureLevel::SM5 && Scene->ReflectionSceneData.CubemapArray.IsValid();
+	const uint32 bReflectionEnv = ShouldDoReflectionEnvironment();
 
-	TRefCountPtr<IPooledRenderTarget>	NewSceneColor;
-	if( bReflectionEnv && bReflectionsWithCompute && !bAnyViewIsReflectionCapture )
+	TRefCountPtr<IPooledRenderTarget> NewSceneColor;
 	{
 		GSceneRenderTargets.ResolveSceneColor(FResolveRect(0, 0, ViewFamily.FamilySizeX, ViewFamily.FamilySizeY));
 
@@ -577,181 +596,180 @@ void FDeferredShadingSceneRenderer::RenderDeferredReflections()
 		GRenderTargetPool.FindFreeElement( Desc, NewSceneColor, TEXT("SceneColorEnv") );
 	}
 
-	for (int32 ViewIndex = 0; ViewIndex < Views.Num(); ViewIndex++)
+	// If we are in SM5, use the compute shader gather method
+	for (int32 ViewIndex = 0, Num = Views.Num(); ViewIndex < Num; ViewIndex++)
 	{
 		const FViewInfo& View = Views[ViewIndex];
 
-		if (bAnyViewIsReflectionCapture)
+		const uint32 bSSR = DoScreenSpaceReflections(Views[ViewIndex]);
+
+		TRefCountPtr<IPooledRenderTarget> SSROutput = GSystemTextures.BlackDummy;
+		if( bSSR )
 		{
-			// We're currently capturing a reflection capture, output SpecularColor * IndirectDiffuseGI for metals so they are not black in reflections,
-			// Since we don't have multiple bounce specular reflections
-			GSceneRenderTargets.BeginRenderingSceneColor();
-
-			RHISetViewport( View.ViewRect.Min.X, View.ViewRect.Min.Y, 0.0f, View.ViewRect.Max.X, View.ViewRect.Max.Y, 1.0f );
-			RHISetRasterizerState( TStaticRasterizerState< FM_Solid, CM_None >::GetRHI() );
-			RHISetDepthStencilState( TStaticDepthStencilState< false, CF_Always >::GetRHI() );
-			RHISetBlendState( TStaticBlendState< CW_RGB, BO_Add, BF_One, BF_One >::GetRHI() );
-
-			TShaderMapRef< FPostProcessVS > VertexShader( GetGlobalShaderMap() );
-			TShaderMapRef< FReflectionCaptureSpecularBouncePS > PixelShader( GetGlobalShaderMap() );
-
-			static FGlobalBoundShaderState BoundShaderState;
-			SetGlobalBoundShaderState( BoundShaderState, GFilterVertexDeclaration.VertexDeclarationRHI, *VertexShader, *PixelShader );
-
-			PixelShader->SetParameters( View );
-
-			DrawRectangle( 
-				0, 0, 
-				View.ViewRect.Width(), View.ViewRect.Height(),
-				0, 0, 
-				View.ViewRect.Width(), View.ViewRect.Height(),
-				FIntPoint(View.ViewRect.Width(), View.ViewRect.Height()),
-				GSceneRenderTargets.GetBufferSizeXY(),
-				*VertexShader,
-				EDRF_UseTriangleOptimization);
+			ScreenSpaceReflections( View, SSROutput );
 		}
-		else
+
+		if(bReflectionEnv)
 		{
-			const uint32 bSSR = DoScreenSpaceReflections(Views[0]);
+			// Render the reflection environment with tiled deferred culling
+			SCOPED_DRAW_EVENT(ReflectionEnvironmentGather, DEC_SCENE_ITEMS);
 
-			bool bRequiresApply = bSSR || bReflectionEnv || bSkyLight;
+			RHISetRenderTarget(NULL, NULL);
 
-			TRefCountPtr<IPooledRenderTarget> SSROutput = GSystemTextures.BlackDummy;
-			if( bSSR )
+			FReflectionEnvironmentTiledDeferredCS* ComputeShader = NULL;
+			if( bUseLightmaps )
 			{
-				ScreenSpaceReflections( View, SSROutput );
+				ComputeShader = *TShaderMapRef< TReflectionEnvironmentTiledDeferredCS<1> >( GetGlobalShaderMap() );
 			}
-
-			if( bReflectionEnv )
+			else
 			{
-				// If we are in SM5, use the compute shader gather method
-				if (FeatureLevel == ERHIFeatureLevel::SM5 && Scene->ReflectionSceneData.CubemapArray.IsValid())
-				{
-					// Render the reflection environment with tiled deferred culling
-					SCOPED_DRAW_EVENT(ReflectionEnvironmentGather, DEC_SCENE_ITEMS);
-				
-					RHISetRenderTarget(NULL, NULL);
+				ComputeShader = *TShaderMapRef< TReflectionEnvironmentTiledDeferredCS<0> >( GetGlobalShaderMap() );
+			}
 
-					FReflectionEnvironmentTiledDeferredCS* ComputeShader = NULL;
-					if( bUseLightmaps )
-					{
-						ComputeShader = *TShaderMapRef< TReflectionEnvironmentTiledDeferredCS<1> >( GetGlobalShaderMap() );
-					}
-					else
-					{
-						ComputeShader = *TShaderMapRef< TReflectionEnvironmentTiledDeferredCS<0> >( GetGlobalShaderMap() );
-					}
+			RHISetComputeShader(ComputeShader->GetComputeShader());
+
+			ComputeShader->SetParameters(View, SSROutput->GetRenderTargetItem().ShaderResourceTexture, NewSceneColor->GetRenderTargetItem().UAV);
+			//ComputeShader->SetParameters(View, DestRenderTarget.UAV);
+
+			uint32 GroupSizeX = (View.ViewRect.Size().X + GReflectionEnvironmentTileSizeX - 1) / GReflectionEnvironmentTileSizeX;
+			uint32 GroupSizeY = (View.ViewRect.Size().Y + GReflectionEnvironmentTileSizeY - 1) / GReflectionEnvironmentTileSizeY;
+			DispatchComputeShader(ComputeShader, GroupSizeX, GroupSizeY, 1);
+
+			ComputeShader->UnsetParameters();
+		}
+	}
+
+	GSceneRenderTargets.SetSceneColor(NewSceneColor);
+	check(GSceneRenderTargets.GetSceneColor());
+}
+
+void FDeferredShadingSceneRenderer::RenderImageBasedReflectionsSM4ForAllViews()
+{
+	const bool bSkyLight = Scene->SkyLight
+		&& Scene->SkyLight->ProcessedTexture
+		&& ViewFamily.EngineShowFlags.SkyLighting;
+
+	const uint32 bReflectionEnv = ShouldDoReflectionEnvironment();
+
+	TRefCountPtr<IPooledRenderTarget> LightAccumulation;
+	{
+		const ERHIFeatureLevel::Type FeatureLevel = Scene->GetFeatureLevel();
+
+		uint32 LightAccumulationUAVFlag = (FeatureLevel == ERHIFeatureLevel::SM5) ? TexCreate_UAV : 0;
+		FPooledRenderTargetDesc Desc = GSceneRenderTargets.GetSceneColor()->GetDesc();
+
+		GRenderTargetPool.FindFreeElement(Desc, LightAccumulation, TEXT("LightAccumulation"));
+	}
+
+	// In SM4 use standard deferred shading to composite reflection capture contribution
+	for (int32 ViewIndex = 0, Num = Views.Num(); ViewIndex < Num; ViewIndex++)
+	{
+		const FViewInfo& View = Views[ViewIndex];
+
+		bool bRequiresApply = bSkyLight;
+
+		const bool bSSR = DoScreenSpaceReflections(View);
+
+		TRefCountPtr<IPooledRenderTarget> SSROutput = GSystemTextures.BlackDummy;
+		if( bSSR )
+		{
+			ScreenSpaceReflections( View, SSROutput );
+			bRequiresApply = true;
+		}
 
-					RHISetComputeShader(ComputeShader->GetComputeShader());
+		if( bReflectionEnv )
+		{
+			bRequiresApply = true;
 
-					ComputeShader->SetParameters(View, SSROutput->GetRenderTargetItem().ShaderResourceTexture, NewSceneColor->GetRenderTargetItem().UAV);
-					//ComputeShader->SetParameters(View, DestRenderTarget.UAV);
+			// TODO sort once per view
+			static TArray<FReflectionCaptureSortData> SortData;
+			SortData.Reset(Scene->ReflectionSceneData.RegisteredReflectionCaptures.Num());
 
-					uint32 GroupSizeX = (View.ViewRect.Size().X + GReflectionEnvironmentTileSizeX - 1) / GReflectionEnvironmentTileSizeX;
-					uint32 GroupSizeY = (View.ViewRect.Size().Y + GReflectionEnvironmentTileSizeY - 1) / GReflectionEnvironmentTileSizeY;
-					DispatchComputeShader(ComputeShader, GroupSizeX, GroupSizeY, 1);
+			// Gather visible reflection capture data
+			for (int32 ReflectionProxyIndex = 0; ReflectionProxyIndex < Scene->ReflectionSceneData.RegisteredReflectionCaptures.Num() && SortData.Num() < GMaxNumReflectionCaptures; ReflectionProxyIndex++)
+			{
+				FReflectionCaptureProxy* CurrentCapture = Scene->ReflectionSceneData.RegisteredReflectionCaptures[ReflectionProxyIndex];
+				FReflectionCaptureSortData NewSortEntry;
 
-					ComputeShader->UnsetParameters();
+				NewSortEntry.SM4FullHDRCubemap = CurrentCapture->SM4FullHDRCubemap;
+				NewSortEntry.Guid = CurrentCapture->Guid;
+				NewSortEntry.PositionAndRadius = FVector4(CurrentCapture->Position, CurrentCapture->InfluenceRadius);
+				float ShapeTypeValue = (float)CurrentCapture->Shape;
+				NewSortEntry.CaptureProperties = FVector4(CurrentCapture->Brightness, 0, ShapeTypeValue, 0);
 
-					GSceneRenderTargets.SetSceneColor(NewSceneColor);
-					check(GSceneRenderTargets.GetSceneColor());
+				if (CurrentCapture->Shape == EReflectionCaptureShape::Plane)
+				{
+					NewSortEntry.BoxTransform = FMatrix(
+						FPlane(CurrentCapture->ReflectionPlane), 
+						FPlane(CurrentCapture->ReflectionXAxisAndYScale), 
+						FPlane(0, 0, 0, 0), 
+						FPlane(0, 0, 0, 0));
 
-					bRequiresApply = false;
+					NewSortEntry.BoxScales = FVector4(0);
 				}
-				// In SM4 use standard deferred shading to composite reflection capture contribution
-				else if (FeatureLevel == ERHIFeatureLevel::SM4)
+				else
 				{
-					// TODO sort once per view
-					static TArray<FReflectionCaptureSortData> SortData;
-					SortData.Reset(Scene->ReflectionSceneData.RegisteredReflectionCaptures.Num());
+					NewSortEntry.BoxTransform = CurrentCapture->BoxTransform;
 
-					// Gather visible reflection capture data
-					for (int32 ReflectionProxyIndex = 0; ReflectionProxyIndex < Scene->ReflectionSceneData.RegisteredReflectionCaptures.Num() && SortData.Num() < GMaxNumReflectionCaptures; ReflectionProxyIndex++)
-					{
-						FReflectionCaptureProxy* CurrentCapture = Scene->ReflectionSceneData.RegisteredReflectionCaptures[ReflectionProxyIndex];
-						FReflectionCaptureSortData NewSortEntry;
-
-						NewSortEntry.SM4FullHDRCubemap = CurrentCapture->SM4FullHDRCubemap;
-						NewSortEntry.Guid = CurrentCapture->Guid;
-						NewSortEntry.PositionAndRadius = FVector4(CurrentCapture->Position, CurrentCapture->InfluenceRadius);
-						float ShapeTypeValue = (float)CurrentCapture->Shape;
-						NewSortEntry.CaptureProperties = FVector4(CurrentCapture->Brightness, 0, ShapeTypeValue, 0);
-
-						if (CurrentCapture->Shape == EReflectionCaptureShape::Plane)
-						{
-							NewSortEntry.BoxTransform = FMatrix(
-								FPlane(CurrentCapture->ReflectionPlane), 
-								FPlane(CurrentCapture->ReflectionXAxisAndYScale), 
-								FPlane(0, 0, 0, 0), 
-								FPlane(0, 0, 0, 0));
-
-							NewSortEntry.BoxScales = FVector4(0);
-						}
-						else
-						{
-							NewSortEntry.BoxTransform = CurrentCapture->BoxTransform;
-
-							NewSortEntry.BoxScales = FVector4(CurrentCapture->BoxScales, CurrentCapture->BoxTransitionDistance);
-						}
-
-						SortData.Add(NewSortEntry);
-					}
+					NewSortEntry.BoxScales = FVector4(CurrentCapture->BoxScales, CurrentCapture->BoxTransitionDistance);
+				}
 
-					SortData.Sort();
+				SortData.Add(NewSortEntry);
+			}
 
+			SortData.Sort();
 
-					SCOPED_DRAW_EVENT(StandardDeferredReflectionEnvironment, DEC_SCENE_ITEMS);
-				
-					RHISetRenderTarget(GSceneRenderTargets.LightAccumulation->GetRenderTargetItem().TargetableTexture, NULL);
 
-					// Clear to no reflection contribution, alpha of 1 indicates full background contribution
-					RHIClear(true, FLinearColor(0, 0, 0, 1), false, 0, false, 0, FIntRect());
+			SCOPED_DRAW_EVENT(StandardDeferredReflectionEnvironment, DEC_SCENE_ITEMS);
 
-					RHISetViewport(View.ViewRect.Min.X, View.ViewRect.Min.Y, 0.0f, View.ViewRect.Max.X, View.ViewRect.Max.Y, 1.0f);
+			RHISetRenderTarget(LightAccumulation->GetRenderTargetItem().TargetableTexture, NULL);
 
-					// rgb accumulates reflection contribution front to back, alpha accumulates (1 - alpha0) * (1 - alpha 1)...
-					RHISetBlendState(TStaticBlendState<CW_RGBA, BO_Add, BF_DestAlpha, BF_One, BO_Add, BF_Zero, BF_InverseSourceAlpha>::GetRHI());
-								 
-					for (int32 ReflectionCaptureIndex = 0; ReflectionCaptureIndex < SortData.Num(); ReflectionCaptureIndex++)
-					{
-						const FReflectionCaptureSortData& ReflectionCapture = SortData[ReflectionCaptureIndex];
+			// Clear to no reflection contribution, alpha of 1 indicates full background contribution
+			RHIClear(true, FLinearColor(0, 0, 0, 1), false, 0, false, 0, FIntRect());
 
-						if (ReflectionCapture.SM4FullHDRCubemap)
-						{
-							const FSphere LightBounds(ReflectionCapture.PositionAndRadius, ReflectionCapture.PositionAndRadius.W);
+			RHISetViewport(View.ViewRect.Min.X, View.ViewRect.Min.Y, 0.0f, View.ViewRect.Max.X, View.ViewRect.Max.Y, 1.0f);
 
-							TShaderMapRef<TDeferredLightVS<true> > VertexShader(GetGlobalShaderMap());
+			// rgb accumulates reflection contribution front to back, alpha accumulates (1 - alpha0) * (1 - alpha 1)...
+			RHISetBlendState(TStaticBlendState<CW_RGBA, BO_Add, BF_DestAlpha, BF_One, BO_Add, BF_Zero, BF_InverseSourceAlpha>::GetRHI());
 
-							// Use the appropriate shader for the capture shape
-							if (ReflectionCapture.CaptureProperties.Z == 0)
-							{
-								TShaderMapRef<TStandardDeferredReflectionPS<true> > PixelShader(GetGlobalShaderMap());
+			for (int32 ReflectionCaptureIndex = 0; ReflectionCaptureIndex < SortData.Num(); ReflectionCaptureIndex++)
+			{
+				const FReflectionCaptureSortData& ReflectionCapture = SortData[ReflectionCaptureIndex];
 
-								static FGlobalBoundShaderState BoundShaderState;
-								SetGlobalBoundShaderState(BoundShaderState, GFilterVertexDeclaration.VertexDeclarationRHI, *VertexShader, *PixelShader);
+				if (ReflectionCapture.SM4FullHDRCubemap)
+				{
+					const FSphere LightBounds(ReflectionCapture.PositionAndRadius, ReflectionCapture.PositionAndRadius.W);
 
-								PixelShader->SetParameters(View, ReflectionCapture);
-							}
-							else 
-							{
-								TShaderMapRef<TStandardDeferredReflectionPS<false> > PixelShader(GetGlobalShaderMap());
+					TShaderMapRef<TDeferredLightVS<true> > VertexShader(GetGlobalShaderMap());
 
-								static FGlobalBoundShaderState BoundShaderState;
-								SetGlobalBoundShaderState(BoundShaderState, GFilterVertexDeclaration.VertexDeclarationRHI, *VertexShader, *PixelShader);
+					// Use the appropriate shader for the capture shape
+					if (ReflectionCapture.CaptureProperties.Z == 0)
+					{
+						TShaderMapRef<TStandardDeferredReflectionPS<true> > PixelShader(GetGlobalShaderMap());
 
-								PixelShader->SetParameters(View, ReflectionCapture);
-							}
+						static FGlobalBoundShaderState BoundShaderState;
+						SetGlobalBoundShaderState(BoundShaderState, GFilterVertexDeclaration.VertexDeclarationRHI, *VertexShader, *PixelShader);
 
-							SetBoundingGeometryRasterizerAndDepthState(View, LightBounds);
-							VertexShader->SetSimpleLightParameters(View, LightBounds);
-							StencilingGeometry::DrawSphere();
-						}
+						PixelShader->SetParameters(View, ReflectionCapture);
 					}
+					else 
+					{
+						TShaderMapRef<TStandardDeferredReflectionPS<false> > PixelShader(GetGlobalShaderMap());
+
+						static FGlobalBoundShaderState BoundShaderState;
+						SetGlobalBoundShaderState(BoundShaderState, GFilterVertexDeclaration.VertexDeclarationRHI, *VertexShader, *PixelShader);
 
-					GRenderTargetPool.VisualizeTexture.SetCheckPoint( GSceneRenderTargets.LightAccumulation );
+						PixelShader->SetParameters(View, ReflectionCapture);
+					}
+
+					SetBoundingGeometryRasterizerAndDepthState(View, LightBounds);
+					VertexShader->SetSimpleLightParameters(View, LightBounds);
+					StencilingGeometry::DrawSphere();
 				}
 			}
-			
+
+			GRenderTargetPool.VisualizeTexture.SetCheckPoint( LightAccumulation );
+
 			if( bRequiresApply )
 			{
 				// Apply reflections to screen
@@ -765,26 +783,28 @@ void FDeferredShadingSceneRenderer::RenderDeferredReflections()
 
 				if (GetReflectionEnvironmentCVar() == 2)
 				{
+					// override scene color for debugging
 					RHISetBlendState(TStaticBlendState<>::GetRHI());
 				}
 				else
 				{
+					// additive to scene color
 					RHISetBlendState(TStaticBlendState<CW_RGBA, BO_Add, BF_One, BF_One, BO_Add, BF_One, BF_One>::GetRHI());
 				}
 
 				TShaderMapRef< FPostProcessVS > VertexShader( GetGlobalShaderMap() );
 
-				#define CASE(A,B,C) \
-					case ( (A << 2) | (B << 1) | C ): \
-					{ \
-						TShaderMapRef< FReflectionApplyPS<A,B,C> > PixelShader(GetGlobalShaderMap()); \
-						static FGlobalBoundShaderState BoundShaderState; \
-						SetGlobalBoundShaderState(BoundShaderState, GFilterVertexDeclaration.VertexDeclarationRHI, *VertexShader, *PixelShader); \
-						PixelShader->SetParameters( View, GSceneRenderTargets.LightAccumulation->GetRenderTargetItem().ShaderResourceTexture, SSROutput->GetRenderTargetItem().ShaderResourceTexture ); \
-					}; \
-					break
-
-				switch( (bSSR << 2) | (bReflectionEnv << 1) | bSkyLight )
+#define CASE(A,B,C) \
+				case ( (A << 2) | (B << 1) | C ): \
+				{ \
+				TShaderMapRef< FReflectionApplyPS<A,B,C> > PixelShader(GetGlobalShaderMap()); \
+				static FGlobalBoundShaderState BoundShaderState; \
+				SetGlobalBoundShaderState(BoundShaderState, GFilterVertexDeclaration.VertexDeclarationRHI, *VertexShader, *PixelShader); \
+				PixelShader->SetParameters( View, LightAccumulation->GetRenderTargetItem().ShaderResourceTexture, SSROutput->GetRenderTargetItem().ShaderResourceTexture ); \
+				}; \
+				break
+
+				switch( ((uint32)bSSR << 2) | ((uint32)bReflectionEnv << 1) | (uint32)bSkyLight )
 				{
 					CASE(0,0,0);
 					CASE(0,0,1);
@@ -795,7 +815,7 @@ void FDeferredShadingSceneRenderer::RenderDeferredReflections()
 					CASE(1,1,0);
 					CASE(1,1,1);
 				}
-				#undef CASE
+#undef CASE
 
 				DrawRectangle( 
 					0, 0, 
@@ -809,3 +829,39 @@ void FDeferredShadingSceneRenderer::RenderDeferredReflections()
 		}
 	}
 }
+
+void FDeferredShadingSceneRenderer::RenderDeferredReflections()
+{
+	if (IsSimpleDynamicLightingEnabled() || ViewFamily.EngineShowFlags.VisualizeLightCulling)
+	{
+		return;
+	}
+
+	bool bAnyViewIsReflectionCapture = false;
+	for (int32 ViewIndex = 0, Num = Views.Num(); ViewIndex < Num; ViewIndex++)
+	{
+		const FViewInfo& View = Views[ViewIndex];
+		bAnyViewIsReflectionCapture = bAnyViewIsReflectionCapture || View.bIsReflectionCapture;
+	}
+
+	if (bAnyViewIsReflectionCapture)
+	{
+		RenderReflectionCaptureSpecularBounceForAllViews();
+	}
+	else
+	{
+		const ERHIFeatureLevel::Type FeatureLevel = Scene->GetFeatureLevel();
+
+		bool bReflectionsWithCompute = (FeatureLevel >= ERHIFeatureLevel::SM5) && Scene->ReflectionSceneData.CubemapArray.IsValid();
+
+		if (bReflectionsWithCompute)
+		{
+			RenderImageBasedReflectionsSM5ForAllViews();
+		}
+		else
+		{
+			// to test this code path run with -SM4
+			RenderImageBasedReflectionsSM4ForAllViews();
+		}
+	}
+}
\ No newline at end of file
