* added r.OptimizeForUAVPerformance for perf tests
 Allows to profile if hardware has a performance cost due to rendertarget reuse (more info: search for GCNPerformanceTweets.pdf Tip 37)

diff --git a/Engine/Source/Runtime/Renderer/Private/PostProcess/SceneRenderTargets.cpp b/Engine/Source/Runtime/Renderer/Private/PostProcess/SceneRenderTargets.cpp
--- a/Engine/Source/Runtime/Renderer/Private/PostProcess/SceneRenderTargets.cpp
+++ b/Engine/Source/Runtime/Renderer/Private/PostProcess/SceneRenderTargets.cpp
@@ -36,6 +36,18 @@ static TAutoConsoleVariable<int32> CVarSceneTargetsResizingMethod(
 	ECVF_RenderThreadSafe
 	);
 
+
+static TAutoConsoleVariable<int32> CVarOptimizeForUAVPerformance(
+	TEXT("r.OptimizeForUAVPerformance"),
+	0,
+	TEXT("Allows to profile if hardware has a performance cost due to rendertarget reuse (more info: search for GCNPerformanceTweets.pdf Tip 37)\n")
+	TEXT("If we see a noticable difference on some hardware we can add another option like -1 (meaning auto) and make it the new default.\n")
+	TEXT("0: Optimize for memory savings and reuse of rendertargets (default)\n")
+	TEXT("1: Optimize for less GPU perf loss due to reuse of render targets (can render faster but can require more GPU memory)"),
+	ECVF_RenderThreadSafe
+	);
+
+
 /** The global render targets used for scene rendering. */
 TGlobalResource<FSceneRenderTargets> GSceneRenderTargets;
 
@@ -274,9 +286,11 @@ void FSceneRenderTargets::AllocSceneColor()
 
 		Desc.Flags |= TexCreate_FastVRAM;
 
+		int32 OptimizeForUAVPerformance = CVarOptimizeForUAVPerformance.GetValueOnRenderThread();
+
 		// with TexCreate_UAV it would allow better sharing with later elements but it might come at a high cost:
 		// GCNPerformanceTweets.pdf Tip 37: Warning: Causes additional synchronization between draw calls when using a render target allocated with this flag, use sparingly
-		if(GRHIFeatureLevel >= ERHIFeatureLevel::SM5)
+		if(GRHIFeatureLevel >= ERHIFeatureLevel::SM5 && !OptimizeForUAVPerformance)
 		{
 			Desc.TargetableFlags |= TexCreate_UAV;
 		}
