* nicer comment/help

diff --git a/Engine/Source/Runtime/Renderer/Private/PostProcess/SceneRenderTargets.cpp b/Engine/Source/Runtime/Renderer/Private/PostProcess/SceneRenderTargets.cpp
--- a/Engine/Source/Runtime/Renderer/Private/PostProcess/SceneRenderTargets.cpp
+++ b/Engine/Source/Runtime/Renderer/Private/PostProcess/SceneRenderTargets.cpp
@@ -40,10 +40,10 @@ static TAutoConsoleVariable<int32> CVarSceneTargetsResizingMethod(
 static TAutoConsoleVariable<int32> CVarOptimizeForUAVPerformance(
 	TEXT("r.OptimizeForUAVPerformance"),
 	0,
-	TEXT("Allows to profile if hardware has a performance cost due to rendertarget reuse (more info: search for GCNPerformanceTweets.pdf Tip 37)\n")
-	TEXT("If we see a noticable difference on some hardware we can add another option like -1 (meaning auto) and make it the new default.\n")
-	TEXT("0: Optimize for memory savings and reuse of rendertargets (default)\n")
-	TEXT("1: Optimize for less GPU perf loss due to reuse of render targets (can render faster but can require more GPU memory)"),
+	TEXT("Allows to profile if hardware has a performance cost due to render target reuse (more info: search for GCNPerformanceTweets.pdf Tip 37)\n")
+	TEXT("If we see a noticeable difference on some hardware we can add another option like -1 (meaning auto) and make it the new default.\n")
+	TEXT("0: Optimize for GPU memory savings and reuse render targets (default)\n")
+	TEXT("1: Optimize for GPU performance (might render faster but can require more GPU memory)"),
 	ECVF_RenderThreadSafe
 	);
 
