* alloc and free GBufferTargets on demand to allow reuse, still need to release earlier than postprocessing if possible

diff --git a/Engine/Source/Runtime/Core/Private/HAL/ConsoleManager.cpp b/Engine/Source/Runtime/Core/Private/HAL/ConsoleManager.cpp
--- a/Engine/Source/Runtime/Core/Private/HAL/ConsoleManager.cpp
+++ b/Engine/Source/Runtime/Core/Private/HAL/ConsoleManager.cpp
@@ -1923,13 +1923,6 @@ void CreateConsoleVariables()
 		TEXT("Prevents the engine from taking any CPU or GPU time while not the foreground app."),
 		ECVF_Cheat);
 	
-	IConsoleManager::Get().RegisterConsoleVariable(TEXT("r.CustomDepth"),
-		1,
-		TEXT("0: feature is disabled\n")
-		TEXT("1: feature is enabled, texture is created on demand\n")
-		TEXT("2: feature is enabled, texture is not released until required (should be the project setting if the feature should not stall)"),
-		ECVF_RenderThreadSafe);
-
 	IConsoleManager::Get().RegisterConsoleVariable(TEXT("r.TemporalAASamples"),
 		8,
 		TEXT("Number of jittered positions for temporal AA (4, 8=default, 16, 32, 64)."),
diff --git a/Engine/Source/Runtime/Renderer/Private/DeferredShadingRenderer.cpp b/Engine/Source/Runtime/Renderer/Private/DeferredShadingRenderer.cpp
--- a/Engine/Source/Runtime/Renderer/Private/DeferredShadingRenderer.cpp
+++ b/Engine/Source/Runtime/Renderer/Private/DeferredShadingRenderer.cpp
@@ -519,6 +519,8 @@ void FDeferredShadingSceneRenderer::Render()
 
 	// Draw the scene pre-pass / early z pass, populating the scene depth buffer and HiZ
 	RenderPrePass();
+
+	GSceneRenderTargets.AllocGBufferTargets();
 	
 	// Clear scene color buffer if necessary.
 	if ( bRequiresRHIClear )
diff --git a/Engine/Source/Runtime/Renderer/Private/PostProcess/PostProcessing.cpp b/Engine/Source/Runtime/Renderer/Private/PostProcess/PostProcessing.cpp
--- a/Engine/Source/Runtime/Renderer/Private/PostProcess/PostProcessing.cpp
+++ b/Engine/Source/Runtime/Renderer/Private/PostProcess/PostProcessing.cpp
@@ -1146,6 +1146,15 @@ void FPostProcessing::Process(const FViewInfo& View, TRefCountPtr<IPooledRenderT
 			CompositeContext.Process(TEXT("PostProcessing"));
 		}
 	}
+
+	GRenderTargetPool.AddPhaseEvent(TEXT("AfterPostprocessing"));
+
+	// todo: We should move this to after lighting, but we also have to fix lookups in post process materials reading from it.
+	// We only release the GBuffers after the last view was processed (SplitScreen)
+	if(View.Family->Views[View.Family->Views.Num() - 1] == &View)
+	{
+		GSceneRenderTargets.FreeGBufferTargets();
+	}
 }
 
 
diff --git a/Engine/Source/Runtime/Renderer/Private/PostProcess/SceneRenderTargets.cpp b/Engine/Source/Runtime/Renderer/Private/PostProcess/SceneRenderTargets.cpp
--- a/Engine/Source/Runtime/Renderer/Private/PostProcess/SceneRenderTargets.cpp
+++ b/Engine/Source/Runtime/Renderer/Private/PostProcess/SceneRenderTargets.cpp
@@ -36,7 +36,6 @@ static TAutoConsoleVariable<int32> CVarSceneTargetsResizingMethod(
 	ECVF_RenderThreadSafe
 	);
 
-
 static TAutoConsoleVariable<int32> CVarOptimizeForUAVPerformance(
 	TEXT("r.OptimizeForUAVPerformance"),
 	0,
@@ -47,6 +46,14 @@ static TAutoConsoleVariable<int32> CVarOptimizeForUAVPerformance(
 	ECVF_RenderThreadSafe
 	);
 
+static TAutoConsoleVariable<int32> CVarCustomDepth(
+	TEXT("r.CustomDepth"),
+	1,
+	TEXT("0: feature is disabled\n")
+	TEXT("1: feature is enabled, texture is created on demand\n")
+	TEXT("2: feature is enabled, texture is not released until required (should be the project setting if the feature should not stall)"),
+	ECVF_RenderThreadSafe
+	);
 
 /** The global render targets used for scene rendering. */
 TGlobalResource<FSceneRenderTargets> GSceneRenderTargets;
@@ -325,6 +332,81 @@ void FSceneRenderTargets::AllocLightAttenuation()
 	check(LightAttenuation);
 }
 
+void FSceneRenderTargets::FreeGBufferTargets()
+{
+	GBufferA.SafeRelease();
+	GBufferB.SafeRelease();
+	GBufferC.SafeRelease();
+	GBufferD.SafeRelease();
+	GBufferE.SafeRelease();
+}
+
+void FSceneRenderTargets::AllocGBufferTargets()
+{
+	if(GBufferA)
+	{
+		// no work needed
+		return;
+	}
+
+	// create GBuffer on demand so it can be shared with other pooled RT
+
+	// good to see the quality loss due to precision in the gbuffer
+	const bool bHighPrecisionGBuffers = (CurrentGBufferFormat >= 5);
+	// good to profile the impact of non 8 bit formats
+	const bool bEnforce8BitPerChannel = (CurrentGBufferFormat == 0);
+
+	// Create the world-space normal g-buffer.
+	{
+		EPixelFormat NormalGBufferFormat = bHighPrecisionGBuffers ? PF_FloatRGBA : PF_A2B10G10R10;
+
+		if(bEnforce8BitPerChannel)
+		{
+			NormalGBufferFormat = PF_B8G8R8A8;
+		}
+
+		FPooledRenderTargetDesc Desc(FPooledRenderTargetDesc::Create2DDesc(BufferSize, NormalGBufferFormat, TexCreate_None, TexCreate_RenderTargetable, false));
+		GRenderTargetPool.FindFreeElement(Desc, GBufferA, TEXT("GBufferA"));
+	}
+
+	// Create the specular color and power g-buffer.
+	{
+		const EPixelFormat SpecularGBufferFormat = bHighPrecisionGBuffers ? PF_FloatRGBA : PF_B8G8R8A8;
+
+		FPooledRenderTargetDesc Desc(FPooledRenderTargetDesc::Create2DDesc(BufferSize, SpecularGBufferFormat, TexCreate_None, TexCreate_RenderTargetable, false));
+		GRenderTargetPool.FindFreeElement(Desc, GBufferB, TEXT("GBufferB"));
+	}
+
+	// Create the diffuse color g-buffer.
+	{
+		const EPixelFormat DiffuseGBufferFormat = bHighPrecisionGBuffers ? PF_FloatRGBA : PF_B8G8R8A8;
+		uint32 DiffuseGBufferFlags = TexCreate_SRGB;
+
+#if PLATFORM_MAC // @todo: remove once Apple fixes radr://16754329 AMD Cards don't always perform FRAMEBUFFER_SRGB if the draw FBO has mixed sRGB & non-SRGB colour attachments
+		static const auto CVar = IConsoleManager::Get().FindTConsoleVariableDataInt(TEXT("r.Mac.UseFrameBufferSRGB"));
+		DiffuseGBufferFlags = CVar && CVar->GetValueOnRenderThread() ? TexCreate_SRGB : TexCreate_None;
+#endif
+
+		FPooledRenderTargetDesc Desc(FPooledRenderTargetDesc::Create2DDesc(BufferSize, DiffuseGBufferFormat, DiffuseGBufferFlags, TexCreate_RenderTargetable, false));
+		GRenderTargetPool.FindFreeElement(Desc, GBufferC, TEXT("GBufferC"));
+	}
+
+	// Create the mask g-buffer (e.g. SSAO, subsurface scattering, wet surface mask, skylight mask, ...).
+	{
+		FPooledRenderTargetDesc Desc(FPooledRenderTargetDesc::Create2DDesc(BufferSize, PF_B8G8R8A8, TexCreate_None, TexCreate_RenderTargetable, false));
+		GRenderTargetPool.FindFreeElement(Desc, GBufferD, TEXT("GBufferD"));
+	}
+
+	if (bAllowStaticLighting)
+	{
+		FPooledRenderTargetDesc Desc(FPooledRenderTargetDesc::Create2DDesc(BufferSize, PF_B8G8R8A8, TexCreate_None, TexCreate_RenderTargetable, false));
+		GRenderTargetPool.FindFreeElement(Desc, GBufferE, TEXT("GBufferE"));
+	}
+
+	// otherwise we have a severe problem
+	check(GBufferA);
+}
+
 const TRefCountPtr<IPooledRenderTarget>& FSceneRenderTargets::GetSceneColor() const
 {
 	if(!SceneColor)
@@ -929,58 +1011,6 @@ void FSceneRenderTargets::AllocateDeferredShadingPathRenderTargets()
 	// Create the required render targets if running Highend.
 	if (GRHIFeatureLevel >= ERHIFeatureLevel::SM4)
 	{
-		// good to see the quality loss due to precision in the gbuffer
-		bool bHighPrecisionGBuffers = (CurrentGBufferFormat >= 5);
-		// good to profile the impact of non 8 bit formats
-		bool bEnforce8BitPerChannel = (CurrentGBufferFormat == 0);
-
-		// Create the world-space normal g-buffer.
-		{
-			EPixelFormat NormalGBufferFormat = bHighPrecisionGBuffers ? PF_FloatRGBA : PF_A2B10G10R10;
-
-			if(bEnforce8BitPerChannel)
-			{
-				NormalGBufferFormat = PF_B8G8R8A8;
-			}
-
-			FPooledRenderTargetDesc Desc(FPooledRenderTargetDesc::Create2DDesc(BufferSize, NormalGBufferFormat, TexCreate_None, TexCreate_RenderTargetable, false));
-			GRenderTargetPool.FindFreeElement(Desc, GBufferA, TEXT("GBufferA"));
-		}
-
-		// Create the specular color and power g-buffer.
-		{
-			const EPixelFormat SpecularGBufferFormat = bHighPrecisionGBuffers ? PF_FloatRGBA : PF_B8G8R8A8;
-
-			FPooledRenderTargetDesc Desc(FPooledRenderTargetDesc::Create2DDesc(BufferSize, SpecularGBufferFormat, TexCreate_None, TexCreate_RenderTargetable, false));
-			GRenderTargetPool.FindFreeElement(Desc, GBufferB, TEXT("GBufferB"));
-		}
-
-		// Create the diffuse color g-buffer.
-		{
-			const EPixelFormat DiffuseGBufferFormat = bHighPrecisionGBuffers ? PF_FloatRGBA : PF_B8G8R8A8;
-			uint32 DiffuseGBufferFlags = TexCreate_SRGB;
-
-#if PLATFORM_MAC // @todo: remove once Apple fixes radr://16754329 AMD Cards don't always perform FRAMEBUFFER_SRGB if the draw FBO has mixed sRGB & non-SRGB colour attachments
-			static const auto CVar = IConsoleManager::Get().FindTConsoleVariableDataInt(TEXT("r.Mac.UseFrameBufferSRGB"));
-			DiffuseGBufferFlags = CVar && CVar->GetValueOnRenderThread() ? TexCreate_SRGB : TexCreate_None;
-#endif
-
-			FPooledRenderTargetDesc Desc(FPooledRenderTargetDesc::Create2DDesc(BufferSize, DiffuseGBufferFormat, DiffuseGBufferFlags, TexCreate_RenderTargetable, false));
-			GRenderTargetPool.FindFreeElement(Desc, GBufferC, TEXT("GBufferC"));
-		}
-
-		// Create the mask g-buffer (e.g. SSAO, subsurface scattering, wet surface mask, skylight mask, ...).
-		{
-			FPooledRenderTargetDesc Desc(FPooledRenderTargetDesc::Create2DDesc(BufferSize, PF_B8G8R8A8, TexCreate_None, TexCreate_RenderTargetable, false));
-			GRenderTargetPool.FindFreeElement(Desc, GBufferD, TEXT("GBufferD"));
-		}
-
-		if (bAllowStaticLighting)
-		{
-			FPooledRenderTargetDesc Desc(FPooledRenderTargetDesc::Create2DDesc(BufferSize, PF_B8G8R8A8, TexCreate_None, TexCreate_RenderTargetable, false));
-			GRenderTargetPool.FindFreeElement(Desc, GBufferE, TEXT("GBufferE"));
-		}
-
 		// Create the screen space ambient occlusion buffer
 		{
 			FPooledRenderTargetDesc Desc(FPooledRenderTargetDesc::Create2DDesc(BufferSize, PF_G8, TexCreate_None, TexCreate_RenderTargetable, false));
@@ -1139,17 +1169,14 @@ void FSceneRenderTargets::InitDynamicRHI()
 }
 
 void FSceneRenderTargets::ReleaseDynamicRHI()
-{ 
+{
+	FreeGBufferTargets();
+
 	SceneColor.SafeRelease();
 	SceneAlphaCopy.SafeRelease();
 	SceneDepthZ.SafeRelease();
 	AuxiliarySceneDepthZ.SafeRelease();
 	SmallDepthZ.SafeRelease();
-	GBufferA.SafeRelease();
-	GBufferB.SafeRelease();
-	GBufferC.SafeRelease();
-	GBufferD.SafeRelease();
-	GBufferE.SafeRelease();
 	DBufferA.SafeRelease();
 	DBufferB.SafeRelease();
 	DBufferC.SafeRelease();
@@ -1233,19 +1260,6 @@ FIntPoint FSceneRenderTargets::GetTranslucentShadowDepthTextureResolution() cons
 	return ShadowDepthResolution;
 }
 
-bool FSceneRenderTargets::ShouldDoMSAAInDeferredPasses() const
-{
-	return false;
-}
-
-static int32 GetCustomDepthCVar()
-{
-	static const auto CVar = IConsoleManager::Get().FindTConsoleVariableDataInt(TEXT("r.CustomDepth"));
-	
-	check(CVar);
-	return CVar->GetValueOnRenderThread();
-}
-
 const FTextureRHIRef& FSceneRenderTargets::GetSceneColorSurface() const							
 {
 	if(!SceneColor)
@@ -1268,7 +1282,7 @@ const FTextureRHIRef& FSceneRenderTargets::GetSceneColorTexture() const
 
 IPooledRenderTarget* FSceneRenderTargets::RequestCustomDepth(bool bPrimitives)
 {
-	int Value = GetCustomDepthCVar();
+	int Value = CVarCustomDepth.GetValueOnRenderThread();
 
 	if((Value == 1  && bPrimitives) || Value == 2)
 	{
@@ -1277,7 +1291,6 @@ IPooledRenderTarget* FSceneRenderTargets::RequestCustomDepth(bool bPrimitives)
 		return CustomDepth;
 	}
 
-	//
 	return 0;
 }
 
diff --git a/Engine/Source/Runtime/Renderer/Private/PostProcess/SceneRenderTargets.h b/Engine/Source/Runtime/Renderer/Private/PostProcess/SceneRenderTargets.h
--- a/Engine/Source/Runtime/Renderer/Private/PostProcess/SceneRenderTargets.h
+++ b/Engine/Source/Runtime/Renderer/Private/PostProcess/SceneRenderTargets.h
@@ -336,8 +336,6 @@ public:
 	/** Returns the size of the RSM buffer, taking into account platform limitations and game specific resolution limits. */
 	FIntPoint GetReflectiveShadowMapTextureResolution() const;
 
-	bool ShouldDoMSAAInDeferredPasses() const;
-
 	int32 GetNumGBufferTargets() const;
 
 	// ---
@@ -360,6 +358,10 @@ public:
 
 	// ---
 
+	void FreeGBufferTargets();
+
+	void AllocGBufferTargets();
+
 private: // Get...() methods instead of direct access
 
 	// 0 before BeginRenderingSceneColor and after tone mapping
@@ -377,7 +379,7 @@ public:
 	// Quarter-sized version of the scene depths
 	TRefCountPtr<IPooledRenderTarget> SmallDepthZ;
 
-	// GBuffer: Geometry Buffer rendered in base pass for deferred shading
+	// GBuffer: Geometry Buffer rendered in base pass for deferred shading, only available between AllocGBufferTargets() and FreeGBufferTargets()
 	TRefCountPtr<IPooledRenderTarget> GBufferA;
 	TRefCountPtr<IPooledRenderTarget> GBufferB;
 	TRefCountPtr<IPooledRenderTarget> GBufferC;
@@ -721,31 +723,38 @@ public:
 
 			if (GRHIFeatureLevel >= ERHIFeatureLevel::SM4)
 			{
-				SetTextureParameter(ShaderRHI, GBufferATexture, GBufferATextureSampler, TStaticSamplerState<>::GetRHI(), GSceneRenderTargets.GBufferA->GetRenderTargetItem().ShaderResourceTexture);
-				SetTextureParameter(ShaderRHI, GBufferBTexture, GBufferBTextureSampler, TStaticSamplerState<>::GetRHI(), GSceneRenderTargets.GBufferB->GetRenderTargetItem().ShaderResourceTexture);
-				SetTextureParameter(ShaderRHI, GBufferCTexture, GBufferCTextureSampler, TStaticSamplerState<>::GetRHI(), GSceneRenderTargets.GBufferC->GetRenderTargetItem().ShaderResourceTexture);
-				SetTextureParameter(ShaderRHI, GBufferDTexture, GBufferDTextureSampler, TStaticSamplerState<>::GetRHI(), GSceneRenderTargets.GBufferD->GetRenderTargetItem().ShaderResourceTexture);
-				SetTextureParameter(ShaderRHI, ScreenSpaceAOTexture, ScreenSpaceAOTextureSampler, TStaticSamplerState<>::GetRHI(), ScreenSpaceAO->GetRenderTargetItem().ShaderResourceTexture);
-				SetTextureParameter(ShaderRHI, CustomDepthTexture, CustomDepthTextureSampler, TStaticSamplerState<>::GetRHI(), CustomDepth->GetRenderTargetItem().ShaderResourceTexture);
+				IPooledRenderTarget* GBufferA = GSceneRenderTargets.GBufferA ? GSceneRenderTargets.GBufferA : GSystemTextures.BlackDummy;
+				IPooledRenderTarget* GBufferB = GSceneRenderTargets.GBufferB ? GSceneRenderTargets.GBufferB : GSystemTextures.BlackDummy;
+				IPooledRenderTarget* GBufferC = GSceneRenderTargets.GBufferC ? GSceneRenderTargets.GBufferC : GSystemTextures.BlackDummy;
+				IPooledRenderTarget* GBufferD = GSceneRenderTargets.GBufferD ? GSceneRenderTargets.GBufferD : GSystemTextures.BlackDummy;
+
+				SetTextureParameter(ShaderRHI, GBufferATexture, GBufferATextureSampler, TStaticSamplerState<>::GetRHI(), GBufferA->GetRenderTargetItem().ShaderResourceTexture);
+				SetTextureParameter(ShaderRHI, GBufferBTexture, GBufferBTextureSampler, TStaticSamplerState<>::GetRHI(), GBufferB->GetRenderTargetItem().ShaderResourceTexture);
+				SetTextureParameter(ShaderRHI, GBufferCTexture, GBufferCTextureSampler, TStaticSamplerState<>::GetRHI(), GBufferC->GetRenderTargetItem().ShaderResourceTexture);
+				SetTextureParameter(ShaderRHI, GBufferDTexture, GBufferDTextureSampler, TStaticSamplerState<>::GetRHI(), GBufferD->GetRenderTargetItem().ShaderResourceTexture);
+				SetTextureParameter(ShaderRHI, GBufferATextureMS, GBufferA->GetRenderTargetItem().TargetableTexture);
+				SetTextureParameter(ShaderRHI, GBufferBTextureMS, GBufferB->GetRenderTargetItem().TargetableTexture);
+				SetTextureParameter(ShaderRHI, GBufferCTextureMS, GBufferC->GetRenderTargetItem().TargetableTexture);
+				SetTextureParameter(ShaderRHI, GBufferDTextureMS, GBufferD->GetRenderTargetItem().TargetableTexture);
+				SetTextureParameter(ShaderRHI, GBufferATextureNonMS, GBufferA->GetRenderTargetItem().ShaderResourceTexture);
+				SetTextureParameter(ShaderRHI, GBufferBTextureNonMS, GBufferB->GetRenderTargetItem().ShaderResourceTexture);
+				SetTextureParameter(ShaderRHI, GBufferCTextureNonMS, GBufferC->GetRenderTargetItem().ShaderResourceTexture);
+				SetTextureParameter(ShaderRHI, GBufferDTextureNonMS, GBufferD->GetRenderTargetItem().ShaderResourceTexture);
 
-				SetTextureParameter(ShaderRHI, GBufferATextureMS, GSceneRenderTargets.GBufferA->GetRenderTargetItem().TargetableTexture);
-				SetTextureParameter(ShaderRHI, GBufferBTextureMS, GSceneRenderTargets.GBufferB->GetRenderTargetItem().TargetableTexture);
-				SetTextureParameter(ShaderRHI, GBufferCTextureMS, GSceneRenderTargets.GBufferC->GetRenderTargetItem().TargetableTexture);
-				SetTextureParameter(ShaderRHI, GBufferDTextureMS, GSceneRenderTargets.GBufferD->GetRenderTargetItem().TargetableTexture);
+				SetTextureParameter(ShaderRHI, ScreenSpaceAOTexture, ScreenSpaceAOTextureSampler, TStaticSamplerState<>::GetRHI(), ScreenSpaceAO->GetRenderTargetItem().ShaderResourceTexture);
 				SetTextureParameter(ShaderRHI, ScreenSpaceAOTextureMS, ScreenSpaceAO->GetRenderTargetItem().TargetableTexture);
-
-				SetTextureParameter(ShaderRHI, GBufferATextureNonMS, GSceneRenderTargets.GBufferA->GetRenderTargetItem().ShaderResourceTexture);
-				SetTextureParameter(ShaderRHI, GBufferBTextureNonMS, GSceneRenderTargets.GBufferB->GetRenderTargetItem().ShaderResourceTexture);
-				SetTextureParameter(ShaderRHI, GBufferCTextureNonMS, GSceneRenderTargets.GBufferC->GetRenderTargetItem().ShaderResourceTexture);
-				SetTextureParameter(ShaderRHI, GBufferDTextureNonMS, GSceneRenderTargets.GBufferD->GetRenderTargetItem().ShaderResourceTexture);
 				SetTextureParameter(ShaderRHI, ScreenSpaceAOTextureNonMS, ScreenSpaceAO->GetRenderTargetItem().ShaderResourceTexture);
+
+				SetTextureParameter(ShaderRHI, CustomDepthTexture, CustomDepthTextureSampler, TStaticSamplerState<>::GetRHI(), CustomDepth->GetRenderTargetItem().ShaderResourceTexture);
 				SetTextureParameter(ShaderRHI, CustomDepthTextureNonMS, CustomDepth->GetRenderTargetItem().ShaderResourceTexture);
 
 				if (GSceneRenderTargets.IsStaticLightingAllowed())
 				{
-					SetTextureParameter(ShaderRHI, GBufferETexture, GBufferETextureSampler, TStaticSamplerState<>::GetRHI(), GSceneRenderTargets.GBufferE->GetRenderTargetItem().ShaderResourceTexture);
-					SetTextureParameter(ShaderRHI, GBufferETextureMS, GSceneRenderTargets.GBufferE->GetRenderTargetItem().TargetableTexture);
-					SetTextureParameter(ShaderRHI, GBufferETextureNonMS, GSceneRenderTargets.GBufferE->GetRenderTargetItem().ShaderResourceTexture);
+					IPooledRenderTarget* GBufferE = GSceneRenderTargets.GBufferE ? GSceneRenderTargets.GBufferE : GSystemTextures.BlackDummy;
+
+					SetTextureParameter(ShaderRHI, GBufferETexture, GBufferETextureSampler, TStaticSamplerState<>::GetRHI(), GBufferE->GetRenderTargetItem().ShaderResourceTexture);
+					SetTextureParameter(ShaderRHI, GBufferETextureMS, GBufferE->GetRenderTargetItem().TargetableTexture);
+					SetTextureParameter(ShaderRHI, GBufferETextureNonMS, GBufferE->GetRenderTargetItem().ShaderResourceTexture);
 				}
 			}
 		}
