* fixed rendering when reflectionenv is disabled
* fixed SSR being applied in SM4 path

diff --git a/Engine/Source/Runtime/Renderer/Private/DeferredShadingRenderer.h b/Engine/Source/Runtime/Renderer/Private/DeferredShadingRenderer.h
--- a/Engine/Source/Runtime/Renderer/Private/DeferredShadingRenderer.h
+++ b/Engine/Source/Runtime/Renderer/Private/DeferredShadingRenderer.h
@@ -327,7 +327,7 @@ private:
 	void RenderReflectionCaptureSpecularBounceForAllViews();
 
 	/** Render image based reflections (SSR, Env, SkyLight) without compute shaders */
-	void RenderImageBasedReflectionsSM4ForAllViews();
+	void RenderImageBasedReflectionsSM4ForAllViews(bool bReflectionEnv);
 
 	/** Render image based reflections (SSR, Env, SkyLight) with compute shaders */
 	void RenderImageBasedReflectionsSM5ForAllViews();
diff --git a/Engine/Source/Runtime/Renderer/Private/ReflectionEnvironment.cpp b/Engine/Source/Runtime/Renderer/Private/ReflectionEnvironment.cpp
--- a/Engine/Source/Runtime/Renderer/Private/ReflectionEnvironment.cpp
+++ b/Engine/Source/Runtime/Renderer/Private/ReflectionEnvironment.cpp
@@ -584,8 +584,6 @@ void FDeferredShadingSceneRenderer::RenderImageBasedReflectionsSM5ForAllViews()
 {
 	const uint32 bUseLightmaps = CVarDiffuseFromCaptures.GetValueOnRenderThread() == 0;
 
-	const uint32 bReflectionEnv = ShouldDoReflectionEnvironment();
-
 	TRefCountPtr<IPooledRenderTarget> NewSceneColor;
 	{
 		GSceneRenderTargets.ResolveSceneColor(FResolveRect(0, 0, ViewFamily.FamilySizeX, ViewFamily.FamilySizeY));
@@ -609,7 +607,7 @@ void FDeferredShadingSceneRenderer::RenderImageBasedReflectionsSM5ForAllViews()
 			ScreenSpaceReflections( View, SSROutput );
 		}
 
-		if(bReflectionEnv)
+		// ReflectionEnv is assumed to be on when going into this method
 		{
 			// Render the reflection environment with tiled deferred culling
 			SCOPED_DRAW_EVENT(ReflectionEnvironmentGather, DEC_SCENE_ITEMS);
@@ -643,14 +641,12 @@ void FDeferredShadingSceneRenderer::RenderImageBasedReflectionsSM5ForAllViews()
 	check(GSceneRenderTargets.GetSceneColor());
 }
 
-void FDeferredShadingSceneRenderer::RenderImageBasedReflectionsSM4ForAllViews()
+void FDeferredShadingSceneRenderer::RenderImageBasedReflectionsSM4ForAllViews(bool bReflectionEnv)
 {
 	const bool bSkyLight = Scene->SkyLight
 		&& Scene->SkyLight->ProcessedTexture
 		&& ViewFamily.EngineShowFlags.SkyLighting;
 
-	const uint32 bReflectionEnv = ShouldDoReflectionEnvironment();
-
 	TRefCountPtr<IPooledRenderTarget> LightAccumulation;
 	{
 		const ERHIFeatureLevel::Type FeatureLevel = Scene->GetFeatureLevel();
@@ -661,6 +657,47 @@ void FDeferredShadingSceneRenderer::RenderImageBasedReflectionsSM4ForAllViews()
 		GRenderTargetPool.FindFreeElement(Desc, LightAccumulation, TEXT("LightAccumulation"));
 	}
 
+	static TArray<FReflectionCaptureSortData> SortData;
+
+	if(bReflectionEnv)
+	{
+		SortData.Reset(Scene->ReflectionSceneData.RegisteredReflectionCaptures.Num());
+
+		// Gather visible reflection capture data
+		for (int32 ReflectionProxyIndex = 0; ReflectionProxyIndex < Scene->ReflectionSceneData.RegisteredReflectionCaptures.Num() && SortData.Num() < GMaxNumReflectionCaptures; ReflectionProxyIndex++)
+		{
+			FReflectionCaptureProxy* CurrentCapture = Scene->ReflectionSceneData.RegisteredReflectionCaptures[ReflectionProxyIndex];
+			FReflectionCaptureSortData NewSortEntry;
+
+			NewSortEntry.SM4FullHDRCubemap = CurrentCapture->SM4FullHDRCubemap;
+			NewSortEntry.Guid = CurrentCapture->Guid;
+			NewSortEntry.PositionAndRadius = FVector4(CurrentCapture->Position, CurrentCapture->InfluenceRadius);
+			float ShapeTypeValue = (float)CurrentCapture->Shape;
+			NewSortEntry.CaptureProperties = FVector4(CurrentCapture->Brightness, 0, ShapeTypeValue, 0);
+
+			if (CurrentCapture->Shape == EReflectionCaptureShape::Plane)
+			{
+				NewSortEntry.BoxTransform = FMatrix(
+					FPlane(CurrentCapture->ReflectionPlane), 
+					FPlane(CurrentCapture->ReflectionXAxisAndYScale), 
+					FPlane(0, 0, 0, 0), 
+					FPlane(0, 0, 0, 0));
+
+				NewSortEntry.BoxScales = FVector4(0);
+			}
+			else
+			{
+				NewSortEntry.BoxTransform = CurrentCapture->BoxTransform;
+
+				NewSortEntry.BoxScales = FVector4(CurrentCapture->BoxScales, CurrentCapture->BoxTransitionDistance);
+			}
+
+			SortData.Add(NewSortEntry);
+		}
+
+		SortData.Sort();
+	}
+
 	// In SM4 use standard deferred shading to composite reflection capture contribution
 	for (int32 ViewIndex = 0, Num = Views.Num(); ViewIndex < Num; ViewIndex++)
 	{
@@ -681,45 +718,6 @@ void FDeferredShadingSceneRenderer::RenderImageBasedReflectionsSM4ForAllViews()
 		{
 			bRequiresApply = true;
 
-			// TODO sort once per view
-			static TArray<FReflectionCaptureSortData> SortData;
-			SortData.Reset(Scene->ReflectionSceneData.RegisteredReflectionCaptures.Num());
-
-			// Gather visible reflection capture data
-			for (int32 ReflectionProxyIndex = 0; ReflectionProxyIndex < Scene->ReflectionSceneData.RegisteredReflectionCaptures.Num() && SortData.Num() < GMaxNumReflectionCaptures; ReflectionProxyIndex++)
-			{
-				FReflectionCaptureProxy* CurrentCapture = Scene->ReflectionSceneData.RegisteredReflectionCaptures[ReflectionProxyIndex];
-				FReflectionCaptureSortData NewSortEntry;
-
-				NewSortEntry.SM4FullHDRCubemap = CurrentCapture->SM4FullHDRCubemap;
-				NewSortEntry.Guid = CurrentCapture->Guid;
-				NewSortEntry.PositionAndRadius = FVector4(CurrentCapture->Position, CurrentCapture->InfluenceRadius);
-				float ShapeTypeValue = (float)CurrentCapture->Shape;
-				NewSortEntry.CaptureProperties = FVector4(CurrentCapture->Brightness, 0, ShapeTypeValue, 0);
-
-				if (CurrentCapture->Shape == EReflectionCaptureShape::Plane)
-				{
-					NewSortEntry.BoxTransform = FMatrix(
-						FPlane(CurrentCapture->ReflectionPlane), 
-						FPlane(CurrentCapture->ReflectionXAxisAndYScale), 
-						FPlane(0, 0, 0, 0), 
-						FPlane(0, 0, 0, 0));
-
-					NewSortEntry.BoxScales = FVector4(0);
-				}
-				else
-				{
-					NewSortEntry.BoxTransform = CurrentCapture->BoxTransform;
-
-					NewSortEntry.BoxScales = FVector4(CurrentCapture->BoxScales, CurrentCapture->BoxTransitionDistance);
-				}
-
-				SortData.Add(NewSortEntry);
-			}
-
-			SortData.Sort();
-
-
 			SCOPED_DRAW_EVENT(StandardDeferredReflectionEnvironment, DEC_SCENE_ITEMS);
 
 			RHISetRenderTarget(LightAccumulation->GetRenderTargetItem().TargetableTexture, NULL);
@@ -769,63 +767,63 @@ void FDeferredShadingSceneRenderer::RenderImageBasedReflectionsSM4ForAllViews()
 			}
 
 			GRenderTargetPool.VisualizeTexture.SetCheckPoint( LightAccumulation );
+		}
 
-			if( bRequiresApply )
-			{
-				// Apply reflections to screen
-				SCOPED_DRAW_EVENT(ReflectionApply, DEC_SCENE_ITEMS);
+		if( bRequiresApply )
+		{
+			// Apply reflections to screen
+			SCOPED_DRAW_EVENT(ReflectionApply, DEC_SCENE_ITEMS);
 
-				GSceneRenderTargets.BeginRenderingSceneColor();
+			GSceneRenderTargets.BeginRenderingSceneColor();
 
-				RHISetViewport(View.ViewRect.Min.X, View.ViewRect.Min.Y, 0.0f, View.ViewRect.Max.X, View.ViewRect.Max.Y, 1.0f);
-				RHISetRasterizerState(TStaticRasterizerState<FM_Solid,CM_None>::GetRHI());
-				RHISetDepthStencilState(TStaticDepthStencilState<false,CF_Always>::GetRHI());
+			RHISetViewport(View.ViewRect.Min.X, View.ViewRect.Min.Y, 0.0f, View.ViewRect.Max.X, View.ViewRect.Max.Y, 1.0f);
+			RHISetRasterizerState(TStaticRasterizerState<FM_Solid,CM_None>::GetRHI());
+			RHISetDepthStencilState(TStaticDepthStencilState<false,CF_Always>::GetRHI());
 
-				if (GetReflectionEnvironmentCVar() == 2)
-				{
-					// override scene color for debugging
-					RHISetBlendState(TStaticBlendState<>::GetRHI());
-				}
-				else
-				{
-					// additive to scene color
-					RHISetBlendState(TStaticBlendState<CW_RGBA, BO_Add, BF_One, BF_One, BO_Add, BF_One, BF_One>::GetRHI());
-				}
+			if (GetReflectionEnvironmentCVar() == 2)
+			{
+				// override scene color for debugging
+				RHISetBlendState(TStaticBlendState<>::GetRHI());
+			}
+			else
+			{
+				// additive to scene color
+				RHISetBlendState(TStaticBlendState<CW_RGBA, BO_Add, BF_One, BF_One, BO_Add, BF_One, BF_One>::GetRHI());
+			}
 
-				TShaderMapRef< FPostProcessVS > VertexShader( GetGlobalShaderMap() );
+			TShaderMapRef< FPostProcessVS > VertexShader( GetGlobalShaderMap() );
 
 #define CASE(A,B,C) \
-				case ( (A << 2) | (B << 1) | C ): \
-				{ \
-				TShaderMapRef< FReflectionApplyPS<A,B,C> > PixelShader(GetGlobalShaderMap()); \
-				static FGlobalBoundShaderState BoundShaderState; \
-				SetGlobalBoundShaderState(BoundShaderState, GFilterVertexDeclaration.VertexDeclarationRHI, *VertexShader, *PixelShader); \
-				PixelShader->SetParameters( View, LightAccumulation->GetRenderTargetItem().ShaderResourceTexture, SSROutput->GetRenderTargetItem().ShaderResourceTexture ); \
-				}; \
-				break
-
-				switch( ((uint32)bSSR << 2) | ((uint32)bReflectionEnv << 1) | (uint32)bSkyLight )
-				{
-					CASE(0,0,0);
-					CASE(0,0,1);
-					CASE(0,1,0);
-					CASE(0,1,1);
-					CASE(1,0,0);
-					CASE(1,0,1);
-					CASE(1,1,0);
-					CASE(1,1,1);
-				}
+			case ( (A << 2) | (B << 1) | C ): \
+			{ \
+			TShaderMapRef< FReflectionApplyPS<A,B,C> > PixelShader(GetGlobalShaderMap()); \
+			static FGlobalBoundShaderState BoundShaderState; \
+			SetGlobalBoundShaderState(BoundShaderState, GFilterVertexDeclaration.VertexDeclarationRHI, *VertexShader, *PixelShader); \
+			PixelShader->SetParameters( View, LightAccumulation->GetRenderTargetItem().ShaderResourceTexture, SSROutput->GetRenderTargetItem().ShaderResourceTexture ); \
+			}; \
+			break
+
+			switch( ((uint32)bSSR << 2) | ((uint32)bReflectionEnv << 1) | (uint32)bSkyLight )
+			{
+				CASE(0,0,0);
+				CASE(0,0,1);
+				CASE(0,1,0);
+				CASE(0,1,1);
+				CASE(1,0,0);
+				CASE(1,0,1);
+				CASE(1,1,0);
+				CASE(1,1,1);
+			}
 #undef CASE
 
-				DrawRectangle( 
-					0, 0, 
-					View.ViewRect.Width(), View.ViewRect.Height(),
-					View.ViewRect.Min.X, View.ViewRect.Min.Y, 
-					View.ViewRect.Width(), View.ViewRect.Height(),
-					FIntPoint(View.ViewRect.Width(), View.ViewRect.Height()),
-					GSceneRenderTargets.GetBufferSizeXY(),
-					*VertexShader);
-			}
+			DrawRectangle( 
+				0, 0, 
+				View.ViewRect.Width(), View.ViewRect.Height(),
+				View.ViewRect.Min.X, View.ViewRect.Min.Y, 
+				View.ViewRect.Width(), View.ViewRect.Height(),
+				FIntPoint(View.ViewRect.Width(), View.ViewRect.Height()),
+				GSceneRenderTargets.GetBufferSizeXY(),
+				*VertexShader);
 		}
 	}
 }
@@ -852,16 +850,19 @@ void FDeferredShadingSceneRenderer::RenderDeferredReflections()
 	{
 		const ERHIFeatureLevel::Type FeatureLevel = Scene->GetFeatureLevel();
 
-		bool bReflectionsWithCompute = (FeatureLevel >= ERHIFeatureLevel::SM5) && Scene->ReflectionSceneData.CubemapArray.IsValid();
+		const bool bReflectionEnv = ShouldDoReflectionEnvironment();
+
+		bool bReflectionsWithCompute = (FeatureLevel >= ERHIFeatureLevel::SM5) && bReflectionEnv && Scene->ReflectionSceneData.CubemapArray.IsValid();
 
 		if (bReflectionsWithCompute)
 		{
+			check(bReflectionEnv);
 			RenderImageBasedReflectionsSM5ForAllViews();
 		}
 		else
 		{
 			// to test this code path run with -SM4
-			RenderImageBasedReflectionsSM4ForAllViews();
+			RenderImageBasedReflectionsSM4ForAllViews(bReflectionEnv);
 		}
 	}
 }
\ No newline at end of file
