* cleanup

diff --git a/Engine/Source/Runtime/Renderer/Private/ReflectionEnvironment.cpp b/Engine/Source/Runtime/Renderer/Private/ReflectionEnvironment.cpp
--- a/Engine/Source/Runtime/Renderer/Private/ReflectionEnvironment.cpp
+++ b/Engine/Source/Runtime/Renderer/Private/ReflectionEnvironment.cpp
@@ -661,6 +661,8 @@ void FDeferredShadingSceneRenderer::RenderImageBasedReflectionsSM4ForAllViews(bo
 
 	if(bReflectionEnv)
 	{
+		// shared for multiple views
+
 		SortData.Reset(Scene->ReflectionSceneData.RegisteredReflectionCaptures.Num());
 
 		// Gather visible reflection capture data
@@ -710,8 +712,9 @@ void FDeferredShadingSceneRenderer::RenderImageBasedReflectionsSM4ForAllViews(bo
 		TRefCountPtr<IPooledRenderTarget> SSROutput = GSystemTextures.BlackDummy;
 		if( bSSR )
 		{
-			ScreenSpaceReflections( View, SSROutput );
 			bRequiresApply = true;
+
+			ScreenSpaceReflections( View, SSROutput );
 		}
 
 		if( bReflectionEnv )
