Bring back old Morpheus distortion handling.  New in-development one is available via define 'NEW_MORPHEUS_DISTORTION'

diff --git a/Engine/Source/Runtime/HeadMountedDisplay/Public/IHeadMountedDisplay.h b/Engine/Source/Runtime/HeadMountedDisplay/Public/IHeadMountedDisplay.h
--- a/Engine/Source/Runtime/HeadMountedDisplay/Public/IHeadMountedDisplay.h
+++ b/Engine/Source/Runtime/HeadMountedDisplay/Public/IHeadMountedDisplay.h
@@ -4,6 +4,10 @@
 
 #include "StereoRendering.h"
 
+// depending on your kit and SDK, you may want to use this.
+// new distortion handling still in development.
+#define NEW_MORPHEUS_DISTORTION 0
+
 /**
  * The family of HMD device.  Register a new class of device here if you need to branch code for PostProcessing until 
  */
@@ -207,12 +211,18 @@ public:
 	 * Additional optional distorion rendering parameters
 	 * @todo:  Once we can move shaders into plugins, remove these!
 	 */	
+#if NEW_MORPHEUS_DISTORTION
 	virtual FTexture* GetDistortionTextureLeft() const {return NULL;}
 	virtual FTexture* GetDistortionTextureRight() const {return NULL;}
 	virtual FVector2D GetTextureOffsetLeft() const {return FVector2D::ZeroVector;}
 	virtual FVector2D GetTextureOffsetRight() const {return FVector2D::ZeroVector;}
 	virtual FVector2D GetTextureScaleLeft() const {return FVector2D::ZeroVector;}
 	virtual FVector2D GetTextureScaleRight() const {return FVector2D::ZeroVector;}
+#else
+	virtual void GetImageTranslation(float& x, float& y) const {}
+	virtual void GetDistortionCenterOffset(float& x, float& y) const {}
+	virtual const FTexture*  GetDistortionTexture() {return NULL;}	
+#endif
 
 private:
 	/** Stores the dimensions of the window before we moved into fullscreen mode, so they can be restored */
diff --git a/Engine/Source/Runtime/Renderer/Private/PostProcess/PostProcessMorpheus.cpp b/Engine/Source/Runtime/Renderer/Private/PostProcess/PostProcessMorpheus.cpp
--- a/Engine/Source/Runtime/Renderer/Private/PostProcess/PostProcessMorpheus.cpp
+++ b/Engine/Source/Runtime/Renderer/Private/PostProcess/PostProcessMorpheus.cpp
@@ -34,19 +34,39 @@ class FPostProcessMorpheusPS : public FGlobalShader
 		return false;
 	}
 
+	static void ModifyCompilationEnvironment(EShaderPlatform Platform, FShaderCompilerEnvironment& OutEnvironment)
+	{
+		FGlobalShader::ModifyCompilationEnvironment(Platform, OutEnvironment);
+#if NEW_MORPHEUS_DISTORTION
+		OutEnvironment.SetDefine(TEXT("NEW_MORPHEUS_DISTORTION"), TEXT("1"));
+#endif
+	}
+
 	/** Default constructor. */
 	FPostProcessMorpheusPS()
-	{		
+	{
+#if !NEW_MORPHEUS_DISTORTION
+		DistortionTexture = NULL;
+#endif
 	}
 
 public:
 	FPostProcessPassParameters PostprocessParameter;
 	FDeferredPixelShaderParameters DeferredParameters;
 	
+#if NEW_MORPHEUS_DISTORTION
 	// Distortion parameter values	
 	FShaderParameter TextureScale;
 	FShaderParameter TextureOffset;
 	FShaderParameter TextureUVOffset;
+#else
+	FShaderParameter LensCenter;
+	FShaderParameter ScreenCenter;
+	FShaderParameter Scale;
+	FShaderParameter HMDWarpParam;
+	FShaderParameter CAWarpParam;
+	const FTexture*		DistortionTexture;
+#endif
 
 	FShaderResourceParameter DistortionTextureParam; 	
 	FShaderResourceParameter DistortionTextureSampler; 
@@ -58,6 +78,7 @@ public:
 		PostprocessParameter.Bind(Initializer.ParameterMap);
 		DeferredParameters.Bind(Initializer.ParameterMap);
 
+#if NEW_MORPHEUS_DISTORTION
 		TextureScale.Bind(Initializer.ParameterMap, TEXT("TextureScale"));
 		//check(TextureScaleLeft.IsBound());		
 
@@ -72,6 +93,21 @@ public:
 
 		DistortionTextureSampler.Bind(Initializer.ParameterMap, TEXT("DistortionTextureSampler"));
 		//check(DistortionTextureSampler.IsBound());		
+#else
+		LensCenter.Bind(Initializer.ParameterMap, TEXT("LensCenter"));
+		ScreenCenter.Bind(Initializer.ParameterMap, TEXT("ScreenCenter"));
+		Scale.Bind(Initializer.ParameterMap, TEXT("Scale"));
+		HMDWarpParam.Bind(Initializer.ParameterMap, TEXT("HMDWarpParam"));
+		CAWarpParam.Bind(Initializer.ParameterMap, TEXT("CAWarpParam"));
+
+		DistortionTextureParam.Bind(Initializer.ParameterMap, TEXT("DistortionTexture"));
+		check(DistortionTextureParam.IsBound());
+
+		DistortionTextureSampler.Bind(Initializer.ParameterMap, TEXT("DistortionTextureSampler"));
+		check(DistortionTextureSampler.IsBound());
+
+		DistortionTexture = NULL;
+#endif
 	}
 
 
@@ -88,6 +124,8 @@ public:
 			check(GEngine->HMDDevice.IsValid());
 			TSharedPtr< class IHeadMountedDisplay > HMDDevice = GEngine->HMDDevice;
 
+#if NEW_MORPHEUS_DISTORTION
+
 			check (StereoPass != eSSP_FULL);
 			if (StereoPass == eSSP_LEFT_EYE)
 			{
@@ -107,6 +145,70 @@ public:
 			}				
 			      
             QuadTexTransform = FMatrix::Identity;            
+#else
+			const FIntPoint SrcSize = SrcRect.Size();
+			const float BufferRatioX = float(SrcSize.X)/float(SrcBufferSize.X);
+			const float BufferRatioY = float(SrcSize.Y)/float(SrcBufferSize.Y);
+			const float w = float(SrcSize.X)/float(SrcBufferSize.X);
+			const float h = float(SrcSize.Y)/float(SrcBufferSize.Y);
+			const float x = float(SrcRect.Min.X)/float(SrcBufferSize.X);
+			const float y = float(SrcRect.Min.Y)/float(SrcBufferSize.Y);
+
+			float DistortionCenterOffsetX = 0, DistortionCenterOffsetY = 0;
+			GEngine->HMDDevice->GetDistortionCenterOffset(DistortionCenterOffsetX, DistortionCenterOffsetY);
+
+			const float XCenterOffset = (StereoPass == eSSP_RIGHT_EYE) ? -DistortionCenterOffsetX : DistortionCenterOffsetX;
+			const float YCenterOffset = DistortionCenterOffsetY;
+
+			const float AspectRatio = (float)SrcSize.X / (float)SrcSize.Y;
+			const float ScaleFactor = GEngine->HMDDevice->GetDistortionScalingFactor();
+
+			if(DistortionTexture == NULL)
+			{
+				check(IsInRenderingThread());
+				DistortionTexture = GEngine->HMDDevice->GetDistortionTexture();
+				check(DistortionTexture != NULL);
+			}
+			SetTextureParameter(ShaderRHI, DistortionTextureParam, DistortionTextureSampler, TStaticSamplerState<SF_Bilinear,AM_Clamp,AM_Clamp,AM_Clamp>::GetRHI(), DistortionTexture->TextureRHI);
+
+			// Shifts texture coordinates to the center of the distortion function around the center of the lens
+			FVector2D ViewLensCenter;
+			ViewLensCenter.X = x + (0.5f + XCenterOffset) * w;
+			ViewLensCenter.Y = y + (0.5f + YCenterOffset) * h;
+			SetShaderValue(ShaderRHI, LensCenter, ViewLensCenter);
+
+			// Texture coordinate for the center of the half scene texture, used to clamp sampling
+			FVector2D ViewScreenCenter;
+			ViewScreenCenter.X = x + w * 0.5f;
+			ViewScreenCenter.Y = y + h * 0.5f;
+			SetShaderValue(ShaderRHI, ScreenCenter, ViewScreenCenter);
+
+			// Rescale output (sample) coordinates back to texture range and increase scale to support rendering outside the the screen
+			FVector2D ViewScale;
+			ViewScale.X = (w/2) * ScaleFactor;
+			ViewScale.Y = (h/2) * ScaleFactor * AspectRatio;
+			SetShaderValue(ShaderRHI, Scale, ViewScale);
+
+			// Distortion coefficients
+			FVector4 DistortionValues;
+			GEngine->HMDDevice->GetDistortionWarpValues(DistortionValues);
+			SetShaderValue(ShaderRHI, HMDWarpParam, DistortionValues);
+
+			// CNN - Morpheus changes
+			// CA correction values
+			FVector4 CAValues;
+			GEngine->HMDDevice->GetChromaAbCorrectionValues(CAValues);
+			SetShaderValue(ShaderRHI, CAWarpParam, CAValues);
+
+			// Rescale the texture coordinates to [-1,1] unit range and corrects aspect ratio
+			FVector2D ViewScaleIn;
+			ViewScaleIn.X = (2/w);
+			ViewScaleIn.Y = (2/h) / AspectRatio;
+
+			QuadTexTransform = FMatrix::Identity;
+			QuadTexTransform *= FTranslationMatrix(FVector(-ViewLensCenter.X * SrcBufferSize.X, -ViewLensCenter.Y * SrcBufferSize.Y, 0));
+			QuadTexTransform *= FMatrix(FPlane(ViewScaleIn.X,0,0,0), FPlane(0,ViewScaleIn.Y,0,0), FPlane(0,0,0,0), FPlane(0,0,0,1));
+#endif
 		}
 	}
 	
@@ -114,7 +216,11 @@ public:
 	virtual bool Serialize(FArchive& Ar)
 	{
 		bool bShaderHasOutdatedParameters = FGlobalShader::Serialize(Ar);
+#if NEW_MORPHEUS_DISTORTION
 		Ar << PostprocessParameter << DeferredParameters << TextureScale << TextureOffset << TextureUVOffset << DistortionTextureParam << DistortionTextureSampler;
+#else
+		Ar << PostprocessParameter << DeferredParameters << LensCenter << ScreenCenter << Scale << HMDWarpParam << CAWarpParam << DistortionTextureParam << DistortionTextureSampler;
+#endif
 		return bShaderHasOutdatedParameters;
 	}
 };
@@ -136,6 +242,14 @@ class FPostProcessMorpheusVS : public FGlobalShader
 		return false;
 	}
 
+	static void ModifyCompilationEnvironment(EShaderPlatform Platform, FShaderCompilerEnvironment& OutEnvironment)
+	{
+		FGlobalShader::ModifyCompilationEnvironment(Platform, OutEnvironment);
+#if NEW_MORPHEUS_DISTORTION
+		OutEnvironment.SetDefine(TEXT("NEW_MORPHEUS_DISTORTION"), TEXT("1"));
+#endif
+	}
+
 	/** Default constructor. */
 	FPostProcessMorpheusVS() {}
 
