* added comments

diff --git a/Engine/Source/Runtime/Renderer/Private/LightSceneInfo.h b/Engine/Source/Runtime/Renderer/Private/LightSceneInfo.h
--- a/Engine/Source/Runtime/Renderer/Private/LightSceneInfo.h
+++ b/Engine/Source/Runtime/Renderer/Private/LightSceneInfo.h
@@ -18,11 +18,13 @@ public:
 
 	// must not be 0
 	FLightSceneInfo* LightSceneInfo;
+	// XYZ: origin, W:sphere radius
 	VectorRegister BoundingSphereVector;
 	FLinearColor Color;
 	uint32 bCastDynamicShadow : 1;
 	uint32 bCastStaticShadow : 1;
 	uint32 bStaticLighting : 1;
+	// e.g. LightType_Directional, LightType_Point or LightType_Spot
 	uint32 LightType : LightType_NumBits;	
 
 	/** Initializes the compact scene info from the light's full scene info. */
