Added MSAA support to the mobile path.

Only d3d featureleveles2 tested to work.

diff --git a/Engine/Source/Runtime/Engine/Private/SceneView.cpp b/Engine/Source/Runtime/Engine/Private/SceneView.cpp
--- a/Engine/Source/Runtime/Engine/Private/SceneView.cpp
+++ b/Engine/Source/Runtime/Engine/Private/SceneView.cpp
@@ -948,12 +948,14 @@ void FSceneView::EndFinalPostprocessSettings()
 	{
 		static const auto CVar = IConsoleManager::Get().FindTConsoleVariableDataInt(TEXT("r.PostProcessAAQuality")); 
 		static auto* MobileHDRCvar = IConsoleManager::Get().FindTConsoleVariableDataInt(TEXT("r.MobileHDR"));
+		static auto* MobileMSAACvar = IConsoleManager::Get().FindTConsoleVariableDataInt(TEXT("r.MobileMSAA"));
 
 		int32 Quality = FMath::Clamp(CVar->GetValueOnGameThread(), 0, 6);
 
 		if( !Family->EngineShowFlags.PostProcessing || !Family->EngineShowFlags.AntiAliasing || Quality <= 0
 			// Disable antialiasing in GammaLDR mode to avoid jittering.
-			|| (GRHIFeatureLevel == ERHIFeatureLevel::ES2 && MobileHDRCvar->GetValueOnGameThread() == 0))
+			|| (GRHIFeatureLevel == ERHIFeatureLevel::ES2 && MobileHDRCvar->GetValueOnGameThread() == 0)
+			|| (GRHIFeatureLevel == ERHIFeatureLevel::ES2 && MobileMSAACvar->GetValueOnGameThread()) )
 		{
 			FinalPostProcessSettings.AntiAliasingMethod = AAM_None;
 		}
diff --git a/Engine/Source/Runtime/RHI/Public/RHIUtilities.h b/Engine/Source/Runtime/RHI/Public/RHIUtilities.h
--- a/Engine/Source/Runtime/RHI/Public/RHIUtilities.h
+++ b/Engine/Source/Runtime/RHI/Public/RHIUtilities.h
@@ -204,12 +204,6 @@ inline void RHICreateTargetableShaderResource2D(
 		bForceSeparateTargetAndShaderResource = true;
 	}
 
-	// ES2 doesn't support resolve operations.
-	if (GRHIFeatureLevel <= ERHIFeatureLevel::ES2)
-	{
-		bForceSeparateTargetAndShaderResource = false;
-	}
-
 	if(!bForceSeparateTargetAndShaderResource/* && GSupportsRenderDepthTargetableShaderResources*/)
 	{
 		// Create a single texture that has both TargetableTextureFlags and TexCreate_ShaderResource set.
diff --git a/Engine/Source/Runtime/Renderer/Private/ForwardShadingRenderer.cpp b/Engine/Source/Runtime/Renderer/Private/ForwardShadingRenderer.cpp
--- a/Engine/Source/Runtime/Renderer/Private/ForwardShadingRenderer.cpp
+++ b/Engine/Source/Runtime/Renderer/Private/ForwardShadingRenderer.cpp
@@ -85,8 +85,6 @@ void FForwardShadingSceneRenderer::Render()
 
 	RenderForwardShadingBasePass();
 
-	RenderCustomDepthPass();
-
 	// Notify the FX system that opaque primitives have been rendered.
 	if (Scene->FXSystem)
 	{
diff --git a/Engine/Source/Runtime/Renderer/Private/PostProcess/SceneRenderTargets.cpp b/Engine/Source/Runtime/Renderer/Private/PostProcess/SceneRenderTargets.cpp
--- a/Engine/Source/Runtime/Renderer/Private/PostProcess/SceneRenderTargets.cpp
+++ b/Engine/Source/Runtime/Renderer/Private/PostProcess/SceneRenderTargets.cpp
@@ -55,6 +55,13 @@ static TAutoConsoleVariable<int32> CVarCustomDepth(
 	ECVF_RenderThreadSafe
 	);
 
+TAutoConsoleVariable<int32> CVarMobileMSAA(
+	TEXT("r.MobileMSAA"),
+	0,
+	TEXT("Whether to use MSAA instead of Temporal AA on mobile"),
+	ECVF_RenderThreadSafe
+	);
+
 /** The global render targets used for scene rendering. */
 TGlobalResource<FSceneRenderTargets> GSceneRenderTargets;
 
@@ -160,6 +167,9 @@ void FSceneRenderTargets::Allocate(const FSceneViewFamily& ViewFamily)
 
 	uint32 Mobile32bpp = !IsMobileHDR() || IsMobileHDR32bpp();
 
+	static const auto CVarMobileMSAA = IConsoleManager::Get().FindTConsoleVariableDataInt(TEXT("r.MobileMSAA"));
+	int32 MobileMSAA = CVarMobileMSAA->GetValueOnRenderThread();
+
 	bool bLightPropagationVolume = UseLightPropagationVolumeRT();
 
 	uint32 MinShadowResolution;
@@ -178,6 +188,7 @@ void FSceneRenderTargets::Allocate(const FSceneViewFamily& ViewFamily)
  		(CurrentMaxShadowResolution != MaxShadowResolution) ||
 		(CurrentTranslucencyLightingVolumeDim != TranslucencyLightingVolumeDim) ||
 		(CurrentMobile32bpp != Mobile32bpp) ||
+		(CurrentMobileMSAA != MobileMSAA) ||
 		(bCurrentLightPropagationVolume != bLightPropagationVolume) ||
 		(CurrentMinShadowResolution != MinShadowResolution))
 	{
@@ -188,6 +199,7 @@ void FSceneRenderTargets::Allocate(const FSceneViewFamily& ViewFamily)
 		CurrentMaxShadowResolution = MaxShadowResolution;
 		CurrentTranslucencyLightingVolumeDim = TranslucencyLightingVolumeDim;
 		CurrentMobile32bpp = Mobile32bpp;
+		CurrentMobileMSAA = MobileMSAA;
 		CurrentMinShadowResolution = MinShadowResolution;
 		bCurrentLightPropagationVolume = bLightPropagationVolume;
 		
@@ -260,7 +272,7 @@ void FSceneRenderTargets::BeginRenderingSceneColor( bool bGBufferPass )
 	}
 	else
 	{
-		RHISetRenderTarget(GetSceneColorTexture(), GetSceneDepthTexture());
+		RHISetRenderTarget(GetSceneColorSurface(), GetSceneDepthSurface());
 	}
 } 
 
@@ -877,10 +889,12 @@ uint8 Quantize8SignedByte(float x)
 
 void FSceneRenderTargets::AllocateForwardShadingPathRenderTargets()
 {
+	uint16 NumSamples = CVarMobileMSAA.GetValueOnRenderThread() ? 4 : 1;
+
 	{
 		// Create a texture to store the resolved scene depth, and a render-targetable surface to hold the unresolved scene depth.
 		FPooledRenderTargetDesc Desc(FPooledRenderTargetDesc::Create2DDesc(BufferSize, PF_DepthStencil, TexCreate_None, TexCreate_DepthStencilTargetable, false));
-		Desc.Flags |= TexCreate_FastVRAM;
+		Desc.NumSamples = NumSamples;
 		GRenderTargetPool.FindFreeElement(Desc, SceneDepthZ, TEXT("SceneDepthZ"));
 	}
 
@@ -889,6 +903,12 @@ void FSceneRenderTargets::AllocateForwardShadingPathRenderTargets()
 	{
 		Format = PF_B8G8R8A8;
 	}
+	// Create the scene color.
+	{
+		FPooledRenderTargetDesc Desc(FPooledRenderTargetDesc::Create2DDesc(BufferSize, Format, TexCreate_None, TexCreate_RenderTargetable, false));
+		Desc.NumSamples = NumSamples;
+		GRenderTargetPool.FindFreeElement(Desc, SceneColor, TEXT("SceneColor"));
+	}
 
 	// For 64-bit ES2 without framebuffer fetch, create extra render target for copy of alpha channel.
 	if((Format == PF_FloatRGBA) && (GSupportsShaderFramebufferFetch == false)) 
diff --git a/Engine/Source/Runtime/Renderer/Private/PostProcess/SceneRenderTargets.h b/Engine/Source/Runtime/Renderer/Private/PostProcess/SceneRenderTargets.h
--- a/Engine/Source/Runtime/Renderer/Private/PostProcess/SceneRenderTargets.h
+++ b/Engine/Source/Runtime/Renderer/Private/PostProcess/SceneRenderTargets.h
@@ -501,6 +501,8 @@ private:
 	int32 CurrentTranslucencyLightingVolumeDim;
 	/** To detect a change of the CVar r.MobileHDR / r.MobileHDR32bpp */
 	int32 CurrentMobile32bpp;
+	/** To detect a change of the CVar r.MobileMSAA */
+	int32 CurrentMobileMSAA;
 	/** To detect a change of the CVar r.Shadow.MinResolution */
 	int32 CurrentMinShadowResolution;
 	/** To detect a change of the CVar r.LightPropagationVolume */
diff --git a/Engine/Source/Runtime/Renderer/Public/RendererInterface.h b/Engine/Source/Runtime/Renderer/Public/RendererInterface.h
--- a/Engine/Source/Runtime/Renderer/Public/RendererInterface.h
+++ b/Engine/Source/Runtime/Renderer/Public/RendererInterface.h
@@ -267,6 +267,8 @@ public:
 		// Usually we don't want to propagate MSAA samples.
 		NumSamples = 1;
 
+		bForceSeparateTargetAndShaderResource = false;
+
 		// Remove UAV flag for rendertargets that don't need it (some formats are incompatible)
 		TargetableFlags &= (~TexCreate_UAV);
 	}
