Better fix for tonemapping going overbright when min/max brightness are the same. Instead of calculating the histogram regardless, we make the EyeAdaptation not occur if the input is invalid.

diff --git a/Engine/Source/Runtime/Renderer/Private/PostProcess/PostProcessEyeAdaptation.cpp b/Engine/Source/Runtime/Renderer/Private/PostProcess/PostProcessEyeAdaptation.cpp
--- a/Engine/Source/Runtime/Renderer/Private/PostProcess/PostProcessEyeAdaptation.cpp
+++ b/Engine/Source/Runtime/Renderer/Private/PostProcess/PostProcessEyeAdaptation.cpp
@@ -72,6 +72,14 @@ void FRCPassPostProcessEyeAdaptation::Process(FRenderingCompositePassContext& Co
 {
 	SCOPED_DRAW_EVENT(PostProcessEyeAdaptation, DEC_SCENE_ITEMS);
 
+	const FPooledRenderTargetDesc* InputDesc = GetInputDesc(ePId_Input0);
+
+	if(!InputDesc)
+	{
+		// input is not hooked up correctly
+		return;
+	}
+
 	const FSceneView& View = Context.View;
 	const FSceneViewFamily& ViewFamily = *(View.Family);
 
diff --git a/Engine/Source/Runtime/Renderer/Private/PostProcess/PostProcessing.cpp b/Engine/Source/Runtime/Renderer/Private/PostProcess/PostProcessing.cpp
--- a/Engine/Source/Runtime/Renderer/Private/PostProcess/PostProcessing.cpp
+++ b/Engine/Source/Runtime/Renderer/Private/PostProcess/PostProcessing.cpp
@@ -934,7 +934,7 @@ void FPostProcessing::Process(const FViewInfo& View, TRefCountPtr<IPooledRenderT
 				bool bHistogramNeeded = View.Family->EngineShowFlags.VisualizeHDR;
 
 				if(View.Family->EngineShowFlags.EyeAdaptation
-					&& View.FinalPostProcessSettings.AutoExposureMinBrightness <= View.FinalPostProcessSettings.AutoExposureMaxBrightness
+					&& View.FinalPostProcessSettings.AutoExposureMinBrightness < View.FinalPostProcessSettings.AutoExposureMaxBrightness
 					&& !View.bIsSceneCapture) // Eye adaption is not available for scene captures.
 				{
 					bHistogramNeeded = true;
