Fixed static sky lights still being applied in the base pass with no visible resuilt for static draw lists and translucency

diff --git a/Engine/Source/Runtime/Renderer/Private/BasePassRendering.cpp b/Engine/Source/Runtime/Renderer/Private/BasePassRendering.cpp
--- a/Engine/Source/Runtime/Renderer/Private/BasePassRendering.cpp
+++ b/Engine/Source/Runtime/Renderer/Private/BasePassRendering.cpp
@@ -182,7 +182,7 @@ public:
 				LightMapPolicy,
 				Parameters.BlendMode,
 				Parameters.TextureMode,
-				Parameters.LightingModel != MLM_Unlit && Scene && Scene->SkyLight,
+				Parameters.LightingModel != MLM_Unlit && Scene && Scene->SkyLight && !Scene->SkyLight->bHasStaticLighting,
 				IsTranslucentBlendMode(Parameters.BlendMode) && Scene->HasAtmosphericFog()
 				),
 			Scene->GetFeatureLevel()
diff --git a/Engine/Source/Runtime/Renderer/Private/TranslucentRendering.cpp b/Engine/Source/Runtime/Renderer/Private/TranslucentRendering.cpp
--- a/Engine/Source/Runtime/Renderer/Private/TranslucentRendering.cpp
+++ b/Engine/Source/Runtime/Renderer/Private/TranslucentRendering.cpp
@@ -264,7 +264,7 @@ public:
 			Parameters.BlendMode,
 			// Translucent meshes need scene render targets set as textures
 			ESceneRenderTargetsMode::SetTextures,
-			bIsLitMaterial && Scene && Scene->SkyLight,
+			bIsLitMaterial && Scene && Scene->SkyLight && !Scene->SkyLight->bHasStaticLighting,
 			Scene && Scene->HasAtmosphericFog() && View.Family->EngineShowFlags.Atmosphere,
 			View.Family->EngineShowFlags.ShaderComplexity,
 			Parameters.bAllowFog
