* cleanup and optimization

diff --git a/Engine/Source/Runtime/Renderer/Private/SceneVisibility.cpp b/Engine/Source/Runtime/Renderer/Private/SceneVisibility.cpp
--- a/Engine/Source/Runtime/Renderer/Private/SceneVisibility.cpp
+++ b/Engine/Source/Runtime/Renderer/Private/SceneVisibility.cpp
@@ -777,7 +777,8 @@ static void MarkRelevantStaticMeshesForView(
 	
 	static const auto* StaticMeshLODDistanceScale = IConsoleManager::Get().FindTConsoleVariableDataFloat(TEXT("r.StaticMeshLODDistanceScale"));
 	check(StaticMeshLODDistanceScale);
-	const float InvLODScale = 1.0f / StaticMeshLODDistanceScale->GetValueOnRenderThread();
+	const float LODScale = StaticMeshLODDistanceScale->GetValueOnRenderThread();
+	const float InvLODScale = 1.0f / LODScale;
 
 	const float MinScreenRadiusForCSMDepthSquared = GMinScreenRadiusForCSMDepth * GMinScreenRadiusForCSMDepth;
 	const float MinScreenRadiusForDepthPrepassSquared = GMinScreenRadiusForDepthPrepass * GMinScreenRadiusForDepthPrepass;
@@ -795,14 +796,10 @@ static void MarkRelevantStaticMeshesForView(
 		const FPrimitiveBounds & Bounds = Scene->PrimitiveBounds[PrimitiveIndex];
 		const FPrimitiveViewRelevance & ViewRelevance = View.PrimitiveViewRelevanceMap[PrimitiveIndex];
 		
-		static const auto* StaticMeshLODDistanceScale = IConsoleManager::Get().FindTConsoleVariableDataFloat(TEXT("r.StaticMeshLODDistanceScale"));
-		check(StaticMeshLODDistanceScale);
-		const float LODScale = StaticMeshLODDistanceScale->GetValueOnRenderThread();
-
 		int8 LODToRender = ComputeLODForMeshes( PrimitiveSceneInfo->StaticMeshes, View, Bounds.Origin, Bounds.SphereRadius, ForcedLODLevel, LODScale);
 
 		float DistanceSquared = (Bounds.Origin - ViewOrigin).SizeSquared();
-		const float LODFactorDistanceSquared = DistanceSquared * FMath::Square(View.LODDistanceFactor * (1.0f / LODScale));
+		const float LODFactorDistanceSquared = DistanceSquared * FMath::Square(View.LODDistanceFactor * InvLODScale);
 		const bool bDrawShadowDepth = FMath::Square(Bounds.SphereRadius) > MinScreenRadiusForCSMDepthSquared * LODFactorDistanceSquared;
 		const bool bDrawDepthOnly = bForceEarlyZPass || FMath::Square(Bounds.SphereRadius) > GMinScreenRadiusForDepthPrepass * GMinScreenRadiusForDepthPrepass * LODFactorDistanceSquared;
 		
