* fixed shadercomplexity view mode

diff --git a/Engine/Shaders/ShaderComplexityApplyPixelShader.usf b/Engine/Shaders/ShaderComplexityApplyPixelShader.usf
--- a/Engine/Shaders/ShaderComplexityApplyPixelShader.usf
+++ b/Engine/Shaders/ShaderComplexityApplyPixelShader.usf
@@ -5,15 +5,17 @@
 =============================================================================*/
 
 #include "Common.usf"
+#include "PostProcessCommon.usf"
 
 float3 ShaderComplexityColors[NUM_COMPLEXITY_COLORS];
 
 void Main(
-	in float2 ScreenPosition : TEXCOORD0,
+	in float2 UV : TEXCOORD0,
 	out float4 OutColor : SV_Target0
 	)
 {
-	float3 SceneColor = clamp(CalcSceneColor(ScreenPosition), 0.0f, .999f);
+	float3 SceneColor = clamp(Texture2DSampleLevel(PostprocessInput0, PostprocessInput0Sampler, UV, 0), 0.0f, .999f);
+
 	//expand the stored complexity from [0, 1/3] into [0, NUM_COMPLEXITY_COLORS - 3] 
 	//the first third of the maximum complexity gets NUM_COMPLEXITY_COLORS - 3 unique colors to blend between,
 	//and ShaderComplexityColors[6] represents 1/3rd of the maximum complexity.
diff --git a/Engine/Source/Runtime/Renderer/Private/PostProcess/PostProcessing.cpp b/Engine/Source/Runtime/Renderer/Private/PostProcess/PostProcessing.cpp
--- a/Engine/Source/Runtime/Renderer/Private/PostProcess/PostProcessing.cpp
+++ b/Engine/Source/Runtime/Renderer/Private/PostProcess/PostProcessing.cpp
@@ -1023,14 +1023,14 @@ void FPostProcessing::Process(const FViewInfo& View, TRefCountPtr<IPooledRenderT
 		if(View.Family->EngineShowFlags.StationaryLightOverlap)
 		{
 			FRenderingCompositePass* Node = Context.Graph.RegisterPass(new(FMemStack::Get()) FRCPassPostProcessVisualizeComplexity(GEngine->StationaryLightOverlapColors));
-			Node->SetInput(ePId_Input0, FRenderingCompositeOutputRef(Context.FinalOutput));
+			Node->SetInput(ePId_Input0, FRenderingCompositeOutputRef(Context.SceneColor));
 			Context.FinalOutput = FRenderingCompositeOutputRef(Node);
 		}
 
 		if(View.Family->EngineShowFlags.ShaderComplexity)
 		{
 			FRenderingCompositePass* Node = Context.Graph.RegisterPass(new(FMemStack::Get()) FRCPassPostProcessVisualizeComplexity(GEngine->ShaderComplexityColors));
-			Node->SetInput(ePId_Input0, FRenderingCompositeOutputRef(Context.FinalOutput));
+			Node->SetInput(ePId_Input0, FRenderingCompositeOutputRef(Context.SceneColor));
 			Context.FinalOutput = FRenderingCompositeOutputRef(Node);
 		}
 
diff --git a/Engine/Source/Runtime/Renderer/Private/ShaderComplexityRendering.cpp b/Engine/Source/Runtime/Renderer/Private/ShaderComplexityRendering.cpp
--- a/Engine/Source/Runtime/Renderer/Private/ShaderComplexityRendering.cpp
+++ b/Engine/Source/Runtime/Renderer/Private/ShaderComplexityRendering.cpp
@@ -78,7 +78,7 @@ public:
 	FShaderComplexityApplyPS(const ShaderMetaType::CompiledShaderInitializerType& Initializer)
 		:	FGlobalShader(Initializer)
 	{
-		SceneTextureParams.Bind(Initializer.ParameterMap);
+		PostprocessParameter.Bind(Initializer.ParameterMap);
 		ShaderComplexityColors.Bind(Initializer.ParameterMap,TEXT("ShaderComplexityColors"));
 	}
 
@@ -86,19 +86,16 @@ public:
 	{
 	}
 
-	/**
-	* Sets the current pixel shader params
-	* @param View - current view
-	* @param ShadowInfo - projected shadow info for a single light
-	*/
 	virtual void SetParameters(
-		const FSceneView* View,
+		const FRenderingCompositePassContext& Context,
 		const TArray<FLinearColor>& Colors
 		)
 	{
-		FGlobalShader::SetParameters(GetPixelShader(), *View);
+		const FPixelShaderRHIParamRef ShaderRHI = GetPixelShader();
 
-		SceneTextureParams.Set(GetPixelShader(), *View);
+		FGlobalShader::SetParameters(ShaderRHI, Context.View);
+
+		PostprocessParameter.SetPS(ShaderRHI, Context, TStaticSamplerState<SF_Bilinear,AM_Clamp,AM_Clamp,AM_Clamp>::GetRHI());
 
 		//Make sure there are at least NumShaderComplexityColors colors specified in the ini.
 		//If there are more than NumShaderComplexityColors they will be ignored.
@@ -107,14 +104,8 @@ public:
 		//pass the complexity -> color mapping into the pixel shader
 		for(int32 ColorIndex = 0; ColorIndex < NumShaderComplexityColors; ColorIndex ++)
 		{
-			SetShaderValue(
-				GetPixelShader(),
-				ShaderComplexityColors,
-				Colors[ColorIndex],
-				ColorIndex
-				);
+			SetShaderValue(ShaderRHI, ShaderComplexityColors, Colors[ColorIndex], ColorIndex);
 		}
-
 	}
 
 	/**
@@ -139,14 +130,14 @@ public:
 	bool Serialize(FArchive& Ar)
 	{
 		bool bShaderHasOutdatedParameters = FGlobalShader::Serialize(Ar);
-		Ar << SceneTextureParams;
+		Ar << PostprocessParameter;
 		Ar << ShaderComplexityColors;
 		return bShaderHasOutdatedParameters;
 	}
 
 private:
 
-	FSceneTextureShaderParameters SceneTextureParams;
+	FPostProcessPassParameters PostprocessParameter;
 	FShaderParameter ShaderComplexityColors;
 };
 
@@ -193,7 +184,7 @@ void FRCPassPostProcessVisualizeComplexity::Process(FRenderingCompositePassConte
 	static FGlobalBoundShaderState ShaderComplexityBoundShaderState;
 	SetGlobalBoundShaderState(ShaderComplexityBoundShaderState, GFilterVertexDeclaration.VertexDeclarationRHI, *VertexShader, *PixelShader);
 
-	PixelShader->SetParameters(&View, Colors);
+	PixelShader->SetParameters(Context, Colors);
 	
 	DrawRectangle(
 		0, 0,
