* release GBuffers early during the frame to allow reuse, PostProcessMaterial can hold on to it longer
  granularity is all GBuffers, could be per texture but might not be worth the tracking
* disabled RenderTargetPoolEvents tracking for editor rendering to get a more readable output

diff --git a/Engine/Shaders/ShaderVersion.usf b/Engine/Shaders/ShaderVersion.usf
--- a/Engine/Shaders/ShaderVersion.usf
+++ b/Engine/Shaders/ShaderVersion.usf
@@ -2,4 +2,4 @@
 // Each time the console command is executed it generates a new GUID. As this file is included
 // in common.usf (which should be included in any shader) it allows to invalidate the shader DDC.
 // 
-// GUID = {000C0006-0014-0005-0035-02555806A618}
\ No newline at end of file
+// GUID = {000C0006-0014-0005-0035-02555806A618} test2
\ No newline at end of file
diff --git a/Engine/Source/Runtime/Engine/Private/Materials/HLSLMaterialTranslator.h b/Engine/Source/Runtime/Engine/Private/Materials/HLSLMaterialTranslator.h
--- a/Engine/Source/Runtime/Engine/Private/Materials/HLSLMaterialTranslator.h
+++ b/Engine/Source/Runtime/Engine/Private/Materials/HLSLMaterialTranslator.h
@@ -291,7 +291,7 @@ public:
 			Material->ErrorExpressions.Empty();
 
 			MaterialCompilationOutput.bUsesSceneColor = false;
-			MaterialCompilationOutput.bNeedsSceneTextures = false;
+			check(!MaterialCompilationOutput.bNeedsSceneTextures);
 			MaterialCompilationOutput.bUsesEyeAdaptation = false;
 
 			// Add a state item for the root level
@@ -2202,12 +2202,9 @@ protected:
 			return NonPixelShaderExpressionError();
 		}
 		
-		MaterialCompilationOutput.bNeedsSceneTextures = true;
-		bNeedsSceneTexturePostProcessInputs = bNeedsSceneTexturePostProcessInputs || ((InSceneTextureId >= PPI_PostProcessInput0 && InSceneTextureId <= PPI_PostProcessInput6) || InSceneTextureId == PPI_SceneColor);
-
 		ESceneTextureId SceneTextureId = (ESceneTextureId)InSceneTextureId;
 
-		MaterialCompilationOutput.bUsesSceneColor |= SceneTextureId == PPI_SceneColor;
+		UseSceneTextureId(SceneTextureId, true);
 
 		FString DefaultScreenAligned(TEXT("MaterialFloat2(ScreenAlignedPosition(Parameters.ScreenPosition).xy)"));
 		FString TexCoordCode((UV != INDEX_NONE) ? CoerceParameter(UV, MCT_Float2) : DefaultScreenAligned);
@@ -2228,10 +2225,10 @@ protected:
 			return NonPixelShaderExpressionError();
 		}
 
-		MaterialCompilationOutput.bNeedsSceneTextures = true;
-
 		ESceneTextureId SceneTextureId = (ESceneTextureId)InSceneTextureId;
 
+		UseSceneTextureId(SceneTextureId, false);
+
 		if(SceneTextureId >= PPI_PostProcessInput0 && SceneTextureId <= PPI_PostProcessInput6)
 		{
 			int Index = SceneTextureId - PPI_PostProcessInput0;
@@ -2268,10 +2265,10 @@ protected:
 			return NonPixelShaderExpressionError();
 		}
 
-		MaterialCompilationOutput.bNeedsSceneTextures = true;
-
 		ESceneTextureId SceneTextureId = (ESceneTextureId)InSceneTextureId;
 
+		UseSceneTextureId(SceneTextureId, false);
+
 		if(SceneTextureId >= PPI_PostProcessInput0 && SceneTextureId <= PPI_PostProcessInput6)
 		{
 			int Index = SceneTextureId - PPI_PostProcessInput0;
@@ -2292,10 +2289,10 @@ protected:
 			return NonPixelShaderExpressionError();
 		}
 
-		MaterialCompilationOutput.bNeedsSceneTextures = true;
-
 		ESceneTextureId SceneTextureId = (ESceneTextureId)InSceneTextureId;
 
+		UseSceneTextureId(SceneTextureId, false);
+
 		if(SceneTextureId >= PPI_PostProcessInput0 && SceneTextureId <= PPI_PostProcessInput6)
 		{
 			int Index = SceneTextureId - PPI_PostProcessInput0;
@@ -2308,6 +2305,39 @@ protected:
 		}
 	}
 
+	// @param bTextureLookup true: texture, false:no texture lookup, usually to get the size
+	void UseSceneTextureId(ESceneTextureId SceneTextureId, bool bTextureLookup)
+	{
+		MaterialCompilationOutput.bNeedsSceneTextures = true;
+
+		if(bTextureLookup)
+		{
+			bNeedsSceneTexturePostProcessInputs = bNeedsSceneTexturePostProcessInputs
+				|| ((SceneTextureId >= PPI_PostProcessInput0 && SceneTextureId <= PPI_PostProcessInput6)
+				|| SceneTextureId == PPI_SceneColor);
+		}
+
+		MaterialCompilationOutput.bUsesSceneColor = MaterialCompilationOutput.bUsesSceneColor
+			|| SceneTextureId == PPI_SceneColor;
+
+		MaterialCompilationOutput.bNeedsGBuffer = MaterialCompilationOutput.bNeedsGBuffer
+			|| SceneTextureId == PPI_DiffuseColor 
+			|| SceneTextureId == PPI_SpecularColor
+			|| SceneTextureId == PPI_SubsurfaceColor
+			|| SceneTextureId == PPI_BaseColor
+			|| SceneTextureId == PPI_Specular
+			|| SceneTextureId == PPI_Metallic
+			|| SceneTextureId == PPI_WorldNormal
+			|| SceneTextureId == PPI_Opacity
+			|| SceneTextureId == PPI_Roughness
+			|| SceneTextureId == PPI_MaterialAO
+			|| SceneTextureId == PPI_DecalMask
+			|| SceneTextureId == PPI_LightingModel;
+
+		// not yet tracked:
+		//   PPI_SeparateTranslucency, PPI_CustomDepth, PPI_AmbientOcclusion
+	}
+
 	virtual int32 SceneColor(int32 Offset, int32 UV, bool bUseOffset) OVERRIDE
 	{
 		if (Offset == INDEX_NONE && bUseOffset)
diff --git a/Engine/Source/Runtime/Engine/Private/Materials/MaterialShared.cpp b/Engine/Source/Runtime/Engine/Private/Materials/MaterialShared.cpp
--- a/Engine/Source/Runtime/Engine/Private/Materials/MaterialShared.cpp
+++ b/Engine/Source/Runtime/Engine/Private/Materials/MaterialShared.cpp
@@ -318,10 +318,9 @@ void FMaterialCompilationOutput::Serialize(FArchive& Ar)
 	}
 
 	Ar << bNeedsSceneTextures;
-
 	Ar << bUsesEyeAdaptation;
-
 	Ar << bModifiesMeshPosition;
+	Ar << bNeedsGBuffer;
 }
 
 void FMaterial::GetShaderMapId(EShaderPlatform Platform, FMaterialShaderMapId& OutId) const
@@ -441,6 +440,19 @@ bool FMaterial::NeedsSceneTextures() const
 	return false;
 }
 
+bool FMaterial::NeedsGBuffer() const
+{
+	checkSlow(IsInRenderingThread());
+
+	if (RenderingThreadShaderMap)
+	{
+		return RenderingThreadShaderMap->NeedsGBuffer();
+	}
+
+	return 0;
+}
+
+
 bool FMaterial::UsesEyeAdaptation() const 
 {
 	checkSlow(IsInRenderingThread());
diff --git a/Engine/Source/Runtime/Engine/Public/MaterialShared.h b/Engine/Source/Runtime/Engine/Public/MaterialShared.h
--- a/Engine/Source/Runtime/Engine/Public/MaterialShared.h
+++ b/Engine/Source/Runtime/Engine/Public/MaterialShared.h
@@ -258,7 +258,8 @@ public:
 		bUsesSceneColor(false),
 		bNeedsSceneTextures(false),
 		bUsesEyeAdaptation(false),
-		bModifiesMeshPosition(false)
+		bModifiesMeshPosition(false),
+		bNeedsGBuffer(false)
 	{}
 
 	ENGINE_API void Serialize(FArchive& Ar);
@@ -278,6 +279,9 @@ public:
 
 	/** true if the material modifies the the mesh position. */
 	bool bModifiesMeshPosition;
+
+	/** true if the material uses any GBuffer textures */
+	bool bNeedsGBuffer;
 };
 
 /**
@@ -799,6 +803,7 @@ public:
 	const FString& GetDebugDescription() const { return DebugDescription; }
 	bool UsesSceneColor() const { return MaterialCompilationOutput.bUsesSceneColor; }
 	bool NeedsSceneTextures() const { return MaterialCompilationOutput.bNeedsSceneTextures; }
+	bool NeedsGBuffer() const { return MaterialCompilationOutput.bNeedsGBuffer; }
 	bool UsesEyeAdaptation() const { return MaterialCompilationOutput.bUsesEyeAdaptation; }
 	bool ModifiesMeshPosition() const { return MaterialCompilationOutput.bModifiesMeshPosition; }
 
@@ -1159,6 +1164,7 @@ public:
 	}
 	ENGINE_API bool UsesSceneColor() const;
 	ENGINE_API bool NeedsSceneTextures() const;
+	ENGINE_API bool NeedsGBuffer() const;
 	ENGINE_API bool UsesEyeAdaptation() const;	
 
 	/** Does the material modify the mesh position. */
diff --git a/Engine/Source/Runtime/Renderer/Private/CompositionLighting/CompositionLighting.cpp b/Engine/Source/Runtime/Renderer/Private/CompositionLighting/CompositionLighting.cpp
--- a/Engine/Source/Runtime/Renderer/Private/CompositionLighting/CompositionLighting.cpp
+++ b/Engine/Source/Runtime/Renderer/Private/CompositionLighting/CompositionLighting.cpp
@@ -79,18 +79,27 @@ static bool IsBasePassAmbientOcclusionRequired(FPostprocessContext& Context)
 	return Context.View.FinalPostProcessSettings.AmbientOcclusionStaticFraction >= 1 / 100.0f && !IsSimpleDynamicLightingEnabled();
 }
 
-static bool IsAmbientOcclusionPassRequired(FPostprocessContext& Context)
+// @return 0:off, 0..4
+static uint32 ComputeAmbientOcclusionPassCount(FPostprocessContext& Context)
 {
-	if(IsLpvIndirectPassRequired(Context) )
+	uint32 Ret = 0;
+
+	if(IsLpvIndirectPassRequired(Context))		// this seems wrong - we need to revisit that
 	{
-		return true;
+		bool bEnabled = 
+			Context.View.FinalPostProcessSettings.AmbientOcclusionIntensity > 0 
+			&& Context.View.FinalPostProcessSettings.AmbientOcclusionRadius >= 0.1f 
+			&& (IsBasePassAmbientOcclusionRequired(Context) || IsAmbientCubemapPassRequired(Context) || IsReflectionEnvironmentActive(Context) || Context.View.Family->EngineShowFlags.VisualizeBuffer )
+			&& !IsSimpleDynamicLightingEnabled();
+
+		if(bEnabled)
+		{
+			static const auto ICVar = IConsoleManager::Get().FindTConsoleVariableDataInt(TEXT("r.AmbientOcclusionLevels"));
+			Ret = FMath::Clamp(ICVar->GetValueOnRenderThread(), 0, 4);
+		}
 	}
 
-	return
-		Context.View.FinalPostProcessSettings.AmbientOcclusionIntensity > 0 
-		&& Context.View.FinalPostProcessSettings.AmbientOcclusionRadius >= 0.1f 
-		&& (IsBasePassAmbientOcclusionRequired(Context) || IsAmbientCubemapPassRequired(Context) || IsReflectionEnvironmentActive(Context) || Context.View.Family->EngineShowFlags.VisualizeBuffer )
-		&& !IsSimpleDynamicLightingEnabled();
+	return Ret;
 }
 
 static void AddPostProcessingAmbientCubemap(FPostprocessContext& Context, FRenderingCompositeOutputRef AmbientOcclusion)
@@ -102,19 +111,10 @@ static void AddPostProcessingAmbientCubemap(FPostprocessContext& Context, FRende
 	Context.FinalOutput = FRenderingCompositeOutputRef(Pass);
 }
 
-static FRenderingCompositeOutputRef AddPostProcessingAmbientOcclusion(FPostprocessContext& Context)
+// @param Levels 0..4, how many different resolution levels we want to render
+static FRenderingCompositeOutputRef AddPostProcessingAmbientOcclusion(FPostprocessContext& Context, uint32 Levels)
 {
-	uint32 AmbientOcclusionLevels;
-	{
-		static const auto ICVar = IConsoleManager::Get().FindTConsoleVariableDataInt(TEXT("r.AmbientOcclusionLevels"));
-		AmbientOcclusionLevels = FMath::Clamp(ICVar->GetValueOnRenderThread(), 0, 4);
-	}
-
-	if(AmbientOcclusionLevels <= 0)
-	{
-		// off
-		return FRenderingCompositeOutputRef();
-	}
+	check(Levels >= 0 && Levels <= 4);
 
 	FRenderingCompositePass* AmbientOcclusionInMip1 = 0;
 	FRenderingCompositePass* AmbientOcclusionInMip2 = 0;
@@ -128,13 +128,13 @@ static FRenderingCompositeOutputRef AddPostProcessingAmbientOcclusion(FPostproce
 	AmbientOcclusionInMip1 = Context.Graph.RegisterPass(new(FMemStack::Get()) FRCPassPostProcessAmbientOcclusionSetup());
 	AmbientOcclusionInMip1->SetInput(ePId_Input0, Context.SceneDepth);
 
-	if(AmbientOcclusionLevels >= 3)
+	if(Levels >= 3)
 	{
 		AmbientOcclusionInMip2 = Context.Graph.RegisterPass(new(FMemStack::Get()) FRCPassPostProcessAmbientOcclusionSetup());
 		AmbientOcclusionInMip2->SetInput(ePId_Input1, FRenderingCompositeOutputRef(AmbientOcclusionInMip1, ePId_Output0));
 	}
 
-	if(AmbientOcclusionLevels >= 4)
+	if(Levels >= 4)
 	{
 		AmbientOcclusionInMip3 = Context.Graph.RegisterPass(new(FMemStack::Get()) FRCPassPostProcessAmbientOcclusionSetup());
 		AmbientOcclusionInMip3->SetInput(ePId_Input1, FRenderingCompositeOutputRef(AmbientOcclusionInMip2, ePId_Output0));
@@ -142,14 +142,14 @@ static FRenderingCompositeOutputRef AddPostProcessingAmbientOcclusion(FPostproce
 
 	// upsample from lower resolution
 
-	if(AmbientOcclusionLevels >= 4)
+	if(Levels >= 4)
 	{
 		AmbientOcclusionPassMip3 = Context.Graph.RegisterPass(new(FMemStack::Get()) FRCPassPostProcessAmbientOcclusion());
 		AmbientOcclusionPassMip3->SetInput(ePId_Input0, AmbientOcclusionInMip3);
 		AmbientOcclusionPassMip3->SetInput(ePId_Input1, AmbientOcclusionInMip3);
 	}
 
-	if(AmbientOcclusionLevels >= 3)
+	if(Levels >= 3)
 	{
 		AmbientOcclusionPassMip2 = Context.Graph.RegisterPass(new(FMemStack::Get()) FRCPassPostProcessAmbientOcclusion());
 		AmbientOcclusionPassMip2->SetInput(ePId_Input0, AmbientOcclusionInMip2);
@@ -157,7 +157,7 @@ static FRenderingCompositeOutputRef AddPostProcessingAmbientOcclusion(FPostproce
 		AmbientOcclusionPassMip2->SetInput(ePId_Input2, AmbientOcclusionPassMip3);
 	}
 
-	if(AmbientOcclusionLevels >= 2)
+	if(Levels >= 2)
 	{
 		AmbientOcclusionPassMip1 = Context.Graph.RegisterPass(new(FMemStack::Get()) FRCPassPostProcessAmbientOcclusion());
 		AmbientOcclusionPassMip1->SetInput(ePId_Input0, AmbientOcclusionInMip1);
@@ -165,10 +165,12 @@ static FRenderingCompositeOutputRef AddPostProcessingAmbientOcclusion(FPostproce
 		AmbientOcclusionPassMip1->SetInput(ePId_Input2, AmbientOcclusionPassMip2);
 	}
 
+	FRenderingCompositePass* GBufferA = Context.Graph.RegisterPass(new(FMemStack::Get()) FRCPassPostProcessInput(GSceneRenderTargets.GBufferA));
+
 	// finally full resolution
 
 	FRenderingCompositePass* AmbientOcclusionPassMip0 = Context.Graph.RegisterPass(new(FMemStack::Get()) FRCPassPostProcessAmbientOcclusion(false));
-	AmbientOcclusionPassMip0->SetInput(ePId_Input0, Context.GBufferA);
+	AmbientOcclusionPassMip0->SetInput(ePId_Input0, GBufferA);
 	AmbientOcclusionPassMip0->SetInput(ePId_Input1, AmbientOcclusionInMip1);
 	AmbientOcclusionPassMip0->SetInput(ePId_Input2, AmbientOcclusionPassMip1);
 
@@ -275,9 +277,9 @@ void FCompositionLighting::ProcessAfterBasePass(const FViewInfo& View)
 
 		FRenderingCompositeOutputRef AmbientOcclusion;
 
-		if(IsAmbientOcclusionPassRequired(Context))
+		if(uint32 Levels = ComputeAmbientOcclusionPassCount(Context))
 		{
-			AmbientOcclusion = AddPostProcessingAmbientOcclusion(Context);
+			AmbientOcclusion = AddPostProcessingAmbientOcclusion(Context, Levels);
 		}
 
 		if(IsAmbientCubemapPassRequired(Context))
diff --git a/Engine/Source/Runtime/Renderer/Private/PostProcess/PostProcessMaterial.cpp b/Engine/Source/Runtime/Renderer/Private/PostProcess/PostProcessMaterial.cpp
--- a/Engine/Source/Runtime/Renderer/Private/PostProcess/PostProcessMaterial.cpp
+++ b/Engine/Source/Runtime/Renderer/Private/PostProcess/PostProcessMaterial.cpp
@@ -168,6 +168,11 @@ void FRCPassPostProcessMaterial::Process(FRenderingCompositePassContext& Context
 		EDRF_UseTriangleOptimization);
 
 	RHICopyToResolveTarget(DestRenderTarget.TargetableTexture, DestRenderTarget.ShaderResourceTexture, false, FResolveParams());
+
+	if(Material->NeedsGBuffer())
+	{
+		GSceneRenderTargets.AdjustGBufferRefCount(-1);
+	}
 }
 
 FPooledRenderTargetDesc FRCPassPostProcessMaterial::ComputeOutputDesc(EPassOutputId InPassOutputId) const
diff --git a/Engine/Source/Runtime/Renderer/Private/PostProcess/PostProcessing.cpp b/Engine/Source/Runtime/Renderer/Private/PostProcess/PostProcessing.cpp
--- a/Engine/Source/Runtime/Renderer/Private/PostProcess/PostProcessing.cpp
+++ b/Engine/Source/Runtime/Renderer/Private/PostProcess/PostProcessing.cpp
@@ -71,7 +71,6 @@ FPostprocessContext::FPostprocessContext(class FRenderingCompositionGraph& InGra
 {
 	SceneColor = Graph.RegisterPass(new(FMemStack::Get()) FRCPassPostProcessInput(GSceneRenderTargets.GetSceneColor()));
 	SceneDepth = Graph.RegisterPass(new(FMemStack::Get()) FRCPassPostProcessInput(GSceneRenderTargets.SceneDepthZ));
-	GBufferA = Graph.RegisterPass(new(FMemStack::Get()) FRCPassPostProcessInput(GSceneRenderTargets.GBufferA));
 
 	FinalOutput = FRenderingCompositeOutputRef(SceneColor);
 }
@@ -517,7 +516,23 @@ static void AddPostProcessMaterial(FPostprocessContext& Context, EBlendableLocat
 	{
 		FPostProcessMaterialNode* Data = PPNodes[i];
 
-		FRenderingCompositePass* Node = Context.Graph.RegisterPass(new(FMemStack::Get()) FRCPassPostProcessMaterial(Data->MID));
+		UMaterialInterface* MaterialInterface = Data->MID;
+
+		FMaterialRenderProxy* Proxy = MaterialInterface->GetRenderProxy(false);
+
+		check(Proxy);
+
+		const FMaterial* Material = Proxy->GetMaterial(Context.View.GetFeatureLevel());
+
+		check(Material);
+
+		if(Material->NeedsGBuffer())
+		{
+			// AdjustGBufferRefCount(-1) call is done when the pass gets executed
+			GSceneRenderTargets.AdjustGBufferRefCount(1);
+		}
+
+		FRenderingCompositePass* Node = Context.Graph.RegisterPass(new(FMemStack::Get()) FRCPassPostProcessMaterial(MaterialInterface));
 		Node->SetInput(ePId_Input0, FRenderingCompositeOutputRef(Context.FinalOutput));
 
 		// We are binding separate translucency here because the post process SceneTexture node can reference 
@@ -1119,6 +1134,17 @@ void FPostProcessing::Process(const FViewInfo& View, TRefCountPtr<IPooledRenderT
 			Context.FinalOutput = FRenderingCompositeOutputRef(Node);
 		}
 
+		// After the graph is built but before the graph is processed.
+		// If a postprocess material is using a GBuffer it adds the refcount int FRCPassPostProcessMaterial::Process()
+		// and when it gets processed it removes the refcount
+		// We only release the GBuffers after the last view was processed (SplitScreen)
+		if(View.Family->Views[View.Family->Views.Num() - 1] == &View)
+		{
+			// Generally we no longer need the GBuffers, anyone that wants to keep the GBuffers for longer should have called AdjustGBufferRefCount(1) to keep it for longer
+			// and call AdjustGBufferRefCount(-1) once it's consumed. This needs to happen each frame. PostProcessMaterial do that automatically
+			GSceneRenderTargets.AdjustGBufferRefCount(-1);
+		}
+
 		// The graph setup should be finished before this line ----------------------------------------
 
 		{
@@ -1147,13 +1173,6 @@ void FPostProcessing::Process(const FViewInfo& View, TRefCountPtr<IPooledRenderT
 	}
 
 	GRenderTargetPool.AddPhaseEvent(TEXT("AfterPostprocessing"));
-
-	// todo: We should move this to after lighting, but we also have to fix lookups in post process materials reading from it.
-	// We only release the GBuffers after the last view was processed (SplitScreen)
-	if(View.Family->Views[View.Family->Views.Num() - 1] == &View)
-	{
-		GSceneRenderTargets.FreeGBufferTargets();
-	}
 }
 
 
diff --git a/Engine/Source/Runtime/Renderer/Private/PostProcess/PostProcessing.h b/Engine/Source/Runtime/Renderer/Private/PostProcess/PostProcessing.h
--- a/Engine/Source/Runtime/Renderer/Private/PostProcess/PostProcessing.h
+++ b/Engine/Source/Runtime/Renderer/Private/PostProcess/PostProcessing.h
@@ -22,7 +22,6 @@ public:
 
 	FRenderingCompositePass* SceneColor;
 	FRenderingCompositePass* SceneDepth;
-	FRenderingCompositePass* GBufferA;
 
 	FRenderingCompositeOutputRef FinalOutput;
 };
diff --git a/Engine/Source/Runtime/Renderer/Private/PostProcess/RenderTargetPool.cpp b/Engine/Source/Runtime/Renderer/Private/PostProcess/RenderTargetPool.cpp
--- a/Engine/Source/Runtime/Renderer/Private/PostProcess/RenderTargetPool.cpp
+++ b/Engine/Source/Runtime/Renderer/Private/PostProcess/RenderTargetPool.cpp
@@ -35,7 +35,8 @@ static void RenderTargetPoolEvents(const TArray<FString>& Args)
 	{
 		UE_LOG(LogRenderTargetPool, Display, TEXT("r.DumpRenderTargetPoolEvents is now enabled, use r.DumpRenderTargetPoolEvents ? for help"));
 
-		GRenderTargetPool.EnableEventRecording(SizeInKBThreshold);
+		GRenderTargetPool.EventRecordingSizeThreshold = SizeInKBThreshold;
+		GRenderTargetPool.bStartEventRecordingNextTick = true;
 	}
 	else
 	{
@@ -54,7 +55,7 @@ static FAutoConsoleCommand GRenderTargetPoolEventsCmd(
 bool FRenderTargetPool::IsEventRecordingEnabled() const
 {
 #if !(UE_BUILD_SHIPPING || UE_BUILD_TEST)
-	return bEventRecording; 
+	return bEventRecordingStarted && bEventRecordingActive; 
 #else
 	return false;
 #endif
@@ -99,9 +100,10 @@ static uint32 ComputeSizeInKB(FPooledRenderTarget& Element)
 FRenderTargetPool::FRenderTargetPool()
 	: AllocationLevelInKB(0)
 	, bCurrentlyOverBudget(false)
-	, bEventRecordingTrigger(false)
+	, bStartEventRecordingNextTick(false)
 	, EventRecordingSizeThreshold(0)
-	, bEventRecording(false)
+	, bEventRecordingActive(false)
+	, bEventRecordingStarted(false)
 	, CurrentEventRecordingTime(0)
 {
 }
@@ -923,12 +925,10 @@ void FRenderTargetPool::TickPoolElements()
 {
 	check(IsInRenderingThread());
 
-//	bEventRecording = false;
-
-	if(bEventRecordingTrigger)
+	if(bStartEventRecordingNextTick)
 	{
-		bEventRecordingTrigger = false;
-		bEventRecording = true;
+		bStartEventRecordingNextTick = false;
+		bEventRecordingStarted = true;
 	}
 
 	uint32 MinimumPoolSizeInKB;
diff --git a/Engine/Source/Runtime/Renderer/Private/PostProcess/RenderTargetPool.h b/Engine/Source/Runtime/Renderer/Private/PostProcess/RenderTargetPool.h
--- a/Engine/Source/Runtime/Renderer/Private/PostProcess/RenderTargetPool.h
+++ b/Engine/Source/Runtime/Renderer/Private/PostProcess/RenderTargetPool.h
@@ -209,9 +209,11 @@ public:
 	// Logs out usage information.
 	void DumpMemoryUsage(FOutputDevice& OutputDevice);
 
-	void EnableEventRecording(uint32 InSizeInKBThreshold) { EventRecordingSizeThreshold = InSizeInKBThreshold; bEventRecordingTrigger = true; }
+	// to not have event recording for some time during rendering (e.g. thumbnail rendering)
+	bool SetEventRecordingActive(bool bValue) { bEventRecordingActive = bValue; }
+
 	//
-	void DisableEventDisplay() { RenderTargetPoolEvents.Empty(); bEventRecording = false; }
+	void DisableEventDisplay() { RenderTargetPoolEvents.Empty(); bEventRecordingStarted = false; }
 	//
 	bool IsEventRecordingEnabled() const;
 
@@ -224,6 +226,8 @@ public:
 
 private:
 
+	friend static void RenderTargetPoolEvents(const TArray<FString>& Args);
+
 	/** Elements can be 0, we compact the buffer later. */
 	TArray< TRefCountPtr<FPooledRenderTarget> > PooledRenderTargets;
 
@@ -258,11 +262,13 @@ private:
 	};
 
 	// if next frame we want to run with bEventRecording=true
-	bool bEventRecordingTrigger;
+	bool bStartEventRecordingNextTick;
 	// in KB, e.g. 1MB = 1024, 0 to display all
 	uint32 EventRecordingSizeThreshold;
-	// true if enabled
-	bool bEventRecording;
+	// true if active, to not have the event recording for some time during rendering (e.g. thumbnail rendering)
+	bool bEventRecordingActive;
+	// true meaning someone used r.RenderTargetPool.Events to start it
+	bool bEventRecordingStarted;
 	// only used if bEventRecording
 	TArray<FRenderTargetPoolEvent> RenderTargetPoolEvents;
 	//
diff --git a/Engine/Source/Runtime/Renderer/Private/PostProcess/SceneRenderTargets.cpp b/Engine/Source/Runtime/Renderer/Private/PostProcess/SceneRenderTargets.cpp
--- a/Engine/Source/Runtime/Renderer/Private/PostProcess/SceneRenderTargets.cpp
+++ b/Engine/Source/Runtime/Renderer/Private/PostProcess/SceneRenderTargets.cpp
@@ -355,6 +355,9 @@ void FSceneRenderTargets::FreeGBufferTargets()
 
 void FSceneRenderTargets::AllocGBufferTargets()
 {
+	// AdjustGBufferRefCount +1 doesn't match -1 (within the same frame)
+	ensure(GBufferRefCount == 0);
+
 	if(GBufferA)
 	{
 		// no work needed
@@ -417,6 +420,9 @@ void FSceneRenderTargets::AllocGBufferTargets()
 
 	// otherwise we have a severe problem
 	check(GBufferA);
+
+	// so that 
+	GBufferRefCount = 1;
 }
 
 const TRefCountPtr<IPooledRenderTarget>& FSceneRenderTargets::GetSceneColor() const
@@ -509,6 +515,19 @@ TRefCountPtr<IPooledRenderTarget>& FSceneRenderTargets::GetLightAttenuation()
 	return LightAttenuation;
 }
 
+void FSceneRenderTargets::AdjustGBufferRefCount(int Delta)
+{ 
+	GBufferRefCount += Delta; 
+	
+	if(Delta > 0 && GBufferRefCount == 1)
+	{
+		AllocGBufferTargets();
+	}
+	else if(GBufferRefCount == 0)
+	{
+		FreeGBufferTargets(); 
+	}
+}
 
 void FSceneRenderTargets::FinishRenderingSceneColor(bool bKeepChanges, const FResolveRect& ResolveRect)
 {
diff --git a/Engine/Source/Runtime/Renderer/Private/PostProcess/SceneRenderTargets.h b/Engine/Source/Runtime/Renderer/Private/PostProcess/SceneRenderTargets.h
--- a/Engine/Source/Runtime/Renderer/Private/PostProcess/SceneRenderTargets.h
+++ b/Engine/Source/Runtime/Renderer/Private/PostProcess/SceneRenderTargets.h
@@ -88,6 +88,7 @@ public:
 protected:
 	/** Constructor */
 	FSceneRenderTargets(): 
+		GBufferRefCount(0),
 		bScreenSpaceAOIsValid(false),
 		bCustomDepthIsValid(false),
 		bPreshadowCacheNewlyAllocated(false),
@@ -358,8 +359,11 @@ public:
 
 	// ---
 
-	void FreeGBufferTargets();
+	// allows to release the GBuffer once post process materials no longer need it
+	// @param 1: add a reference, -1: remove a reference
+	void AdjustGBufferRefCount(int Delta);
 
+	//
 	void AllocGBufferTargets();
 
 private: // Get...() methods instead of direct access
@@ -447,6 +451,9 @@ public:
 	bool bPreshadowCacheNewlyAllocated;
 
 private:
+	/** used by AdjustGBufferRefCount */
+	int32 GBufferRefCount;
+
 	/**
 	 * Takes the requested buffer size and quantizes it to an appropriate size for the rest of the
 	 * rendering pipeline. Currently ensures that sizes are multiples of 8 so that they can safely
@@ -477,6 +484,9 @@ private:
 
 	void AllocLightAttenuation();
 
+	// internal method, used by AdjustGBufferRefCount()
+	void FreeGBufferTargets();
+
 	EPixelFormat GetSceneColorFormat() const;
 
 private:
diff --git a/Engine/Source/Runtime/Renderer/Private/SceneRendering.cpp b/Engine/Source/Runtime/Renderer/Private/SceneRendering.cpp
--- a/Engine/Source/Runtime/Renderer/Private/SceneRendering.cpp
+++ b/Engine/Source/Runtime/Renderer/Private/SceneRendering.cpp
@@ -897,6 +897,11 @@ static void RenderViewFamily_RenderThread( FSceneRenderer* SceneRenderer )
         }
     }
 
+	if(SceneRenderer->ViewFamily.EngineShowFlags.OnScreenDebug)
+	{
+		GRenderTargetPool.SetEventRecordingActive(true);
+	}
+
     {
 		SCOPE_CYCLE_COUNTER(STAT_TotalSceneRenderingTime);
 		
@@ -929,6 +934,8 @@ static void RenderViewFamily_RenderThread( FSceneRenderer* SceneRenderer )
 		SET_MEMORY_STAT(STAT_LightInteractionMemory, FLightPrimitiveInteraction::GetMemoryPoolSize());
 #endif
 
+		GRenderTargetPool.SetEventRecordingActive(false);
+
         // Delete the scene renderer.
 		delete SceneRenderer;
 	}
