* fixed compiler errors

diff --git a/Engine/Source/Runtime/Renderer/Private/PostProcess/RenderTargetPool.cpp b/Engine/Source/Runtime/Renderer/Private/PostProcess/RenderTargetPool.cpp
--- a/Engine/Source/Runtime/Renderer/Private/PostProcess/RenderTargetPool.cpp
+++ b/Engine/Source/Runtime/Renderer/Private/PostProcess/RenderTargetPool.cpp
@@ -23,7 +23,7 @@ static FAutoConsoleCommandWithOutputDevice GDumpRenderTargetPoolMemoryCmd(
 	FConsoleCommandWithOutputDeviceDelegate::CreateStatic(DumpRenderTargetPoolMemory)
 	);
 
-static void RenderTargetPoolEvents(const TArray<FString>& Args)
+void RenderTargetPoolEvents(const TArray<FString>& Args)
 {
 	uint32 SizeInKBThreshold = -1;
 	if(Args.Num() && Args[0].IsNumeric())
diff --git a/Engine/Source/Runtime/Renderer/Private/PostProcess/RenderTargetPool.h b/Engine/Source/Runtime/Renderer/Private/PostProcess/RenderTargetPool.h
--- a/Engine/Source/Runtime/Renderer/Private/PostProcess/RenderTargetPool.h
+++ b/Engine/Source/Runtime/Renderer/Private/PostProcess/RenderTargetPool.h
@@ -210,7 +210,7 @@ public:
 	void DumpMemoryUsage(FOutputDevice& OutputDevice);
 
 	// to not have event recording for some time during rendering (e.g. thumbnail rendering)
-	bool SetEventRecordingActive(bool bValue) { bEventRecordingActive = bValue; }
+	void SetEventRecordingActive(bool bValue) { bEventRecordingActive = bValue; }
 
 	//
 	void DisableEventDisplay() { RenderTargetPoolEvents.Empty(); bEventRecordingStarted = false; }
@@ -226,7 +226,7 @@ public:
 
 private:
 
-	friend static void RenderTargetPoolEvents(const TArray<FString>& Args);
+	friend void RenderTargetPoolEvents(const TArray<FString>& Args);
 
 	/** Elements can be 0, we compact the buffer later. */
 	TArray< TRefCountPtr<FPooledRenderTarget> > PooledRenderTargets;
