* refactored some LPV/shadow code, should not cause any change

diff --git a/Engine/Source/Runtime/Engine/Private/DirectionalLightComponent.cpp b/Engine/Source/Runtime/Engine/Private/DirectionalLightComponent.cpp
--- a/Engine/Source/Runtime/Engine/Private/DirectionalLightComponent.cpp
+++ b/Engine/Source/Runtime/Engine/Private/DirectionalLightComponent.cpp
@@ -244,7 +244,7 @@ public:
 	virtual bool GetViewDependentRsmWholeSceneProjectedShadowInitializer(
 		const class FSceneView& View, 
 		const FBox& LightPropagationVolumeBounds,
-		class FRsmWholeSceneProjectedShadowInitializer& OutInitializer ) const
+		class FWholeSceneProjectedShadowInitializer& OutInitializer ) const
 	{
 		const FMatrix& WorldToLight = GetWorldToLight();
 
@@ -259,7 +259,6 @@ public:
 		OutInitializer.WAxis = FVector4(0,0,0,1);
 		OutInitializer.MinLightW = -HALF_WORLD_MAX;
 		OutInitializer.MaxDistanceToCastInLightW = HALF_WORLD_MAX / 8.0f;
-		OutInitializer.LightPropagationVolumeBounds = LightPropagationVolumeBounds;
 
 		// Compute the RSM bounds
 		{
diff --git a/Engine/Source/Runtime/Engine/Public/SceneManagement.h b/Engine/Source/Runtime/Engine/Public/SceneManagement.h
--- a/Engine/Source/Runtime/Engine/Public/SceneManagement.h
+++ b/Engine/Source/Runtime/Engine/Public/SceneManagement.h
@@ -674,19 +674,6 @@ public:
 
 };
 
-/** Information needed to create a whole scene reflective shadow map projected shadow. */
-class ENGINE_API FRsmWholeSceneProjectedShadowInitializer : public FProjectedShadowInitializer
-{
-public:
-	FRsmWholeSceneProjectedShadowInitializer()
-	 : FProjectedShadowInitializer() 
-	 , LightPropagationVolumeBounds(ForceInitToZero)
-	{}
-
-	FBox LightPropagationVolumeBounds;
-	FShadowCascadeSettings	CascadeSettings;
-};
-
 /** Information needed to create a whole scene projected shadow. */
 class ENGINE_API FWholeSceneProjectedShadowInitializer : public FProjectedShadowInitializer
 {
@@ -698,9 +685,6 @@ public:
 	/** Whether the shadow is a point light shadow that renders all faces of a cubemap in one pass. */
 	uint32 bOnePassPointLightShadow : 1;
 
-	/** Culling planes used by CSM. */
-	TArray<FPlane, TInlineAllocator<8> > FrustumCullPlanes;
-
 	FWholeSceneProjectedShadowInitializer()
 	:	SplitIndex(INDEX_NONE)
 	,	bOnePassPointLightShadow(false)
@@ -840,7 +824,7 @@ public:
 	virtual bool GetViewDependentRsmWholeSceneProjectedShadowInitializer(
 		const class FSceneView& View, 
 		const FBox& LightPropagationVolumeBounds,
-		class FRsmWholeSceneProjectedShadowInitializer& OutInitializer ) const
+		class FWholeSceneProjectedShadowInitializer& OutInitializer ) const
 	{
 		return false;
 	}
diff --git a/Engine/Source/Runtime/Renderer/Private/ShadowRendering.h b/Engine/Source/Runtime/Renderer/Private/ShadowRendering.h
--- a/Engine/Source/Runtime/Renderer/Private/ShadowRendering.h
+++ b/Engine/Source/Runtime/Renderer/Private/ShadowRendering.h
@@ -516,9 +516,6 @@ public:
 
 	FSphere ShadowBounds;
 
-	/** The extruded cascade frustum's culling planes, if a whole scene directional shadow. */
-	TArray<FPlane> FrustumCullPlanes;
-
 	FShadowCascadeSettings CascadeSettings;
 
 	/** X and Y position of the shadow in the appropriate depth buffer.  These are only initialized after the shadow has been allocated. */
@@ -628,16 +625,7 @@ public:
 		const FWholeSceneProjectedShadowInitializer& Initializer,
 		uint32 InResolutionX,
 		uint32 InResolutionY,
-		const TArray<float, TInlineAllocator<2> >& InFadeAlphas
-		);
-
-	/** Initialization constructor for a whole-scene reflective shadow map */
-	FProjectedShadowInfo(
-		FLightSceneInfo* InLightSceneInfo,
-		FViewInfo* InDependentView,
-		const FRsmWholeSceneProjectedShadowInitializer& Initializer,
-		uint32 InResolutionX,
-		uint32 InResolutionY
+		bool bInReflectiveShadowMap
 		);
 
 	float GetShaderDepthBias() const { return ShaderDepthBias; }
diff --git a/Engine/Source/Runtime/Renderer/Private/ShadowSetup.cpp b/Engine/Source/Runtime/Renderer/Private/ShadowSetup.cpp
--- a/Engine/Source/Runtime/Renderer/Private/ShadowSetup.cpp
+++ b/Engine/Source/Runtime/Renderer/Private/ShadowSetup.cpp
@@ -470,14 +470,14 @@ FProjectedShadowInfo::FProjectedShadowInfo(
 	UpdateShaderDepthBias();
 }
 
-/** Initialization constructor for a whole-scene shadow. e.g. directional light cascade or point light */
+/** Initialization constructor for a whole-scene shadow. e.g. reflective shadow map, directional light cascade or point light */
 FProjectedShadowInfo::FProjectedShadowInfo(
 	FLightSceneInfo* InLightSceneInfo,
 	FViewInfo* InDependentView,
 	const FWholeSceneProjectedShadowInitializer& Initializer,
 	uint32 InResolutionX,
 	uint32 InResolutionY,
-	const TArray<float, TInlineAllocator<2> >& InFadeAlphas
+	bool bInReflectiveShadowMap
 	)
 :	LightSceneInfo(InLightSceneInfo)
 ,	LightSceneInfoCompact(InLightSceneInfo)
@@ -485,14 +485,12 @@ FProjectedShadowInfo::FProjectedShadowInfo(
 ,	DependentView(InDependentView)
 ,	ShadowId(INDEX_NONE)
 ,	PreShadowTranslation(Initializer.PreShadowTranslation)
-,	FrustumCullPlanes(Initializer.FrustumCullPlanes)
 ,	CascadeSettings(Initializer.CascadeSettings)
 ,	X(0)
 ,	Y(0)
 ,	ResolutionX(InResolutionX)
 ,	ResolutionY(InResolutionY)
 ,	MaxScreenPercent(1.0f)
-,	FadeAlphas(InFadeAlphas)
 ,	SplitIndex(Initializer.SplitIndex)
 ,	bAllocated(false)
 ,	bAllocatedInTranslucentLayout(false)
@@ -502,7 +500,7 @@ FProjectedShadowInfo::FProjectedShadowInfo(
 ,	bDirectionalLight(Initializer.bDirectionalLight)
 ,	bWholeSceneShadow(true)
 ,	bOnePassPointLightShadow(Initializer.bOnePassPointLightShadow)
-,	bReflectiveShadowmap(false) 
+,	bReflectiveShadowmap(bInReflectiveShadowMap) 
 ,	bTranslucentShadow(false)
 ,	bPreShadow(false)
 ,	bValidTransform(true)
@@ -514,170 +512,105 @@ FProjectedShadowInfo::FProjectedShadowInfo(
 
 	MaxSubjectZ = WorldToFace.TransformPosition(Initializer.SubjectBounds.Origin).Z + Initializer.SubjectBounds.SphereRadius;
 	MinSubjectZ = FMath::Max(MaxSubjectZ - Initializer.SubjectBounds.SphereRadius * 2,Initializer.MinLightW);
-	
-	if(Initializer.bDirectionalLight)
-	{
-		// Limit how small the depth range can be for smaller cascades
-		// This is needed for shadow modes like subsurface shadows which need depth information outside of the smaller cascade depth range
-		//@todo - expose this value to the ini
-		const float DepthRangeClamp = 5000;
-		MaxSubjectZ = FMath::Max(MaxSubjectZ, DepthRangeClamp);
-		MinSubjectZ = FMath::Min(MinSubjectZ, -DepthRangeClamp);
-
-		const FIntPoint ShadowBufferResolution = GSceneRenderTargets.GetShadowDepthTextureResolution();
-		const uint32 ShadowDepthBufferSizeX = ShadowBufferResolution.X - SHADOW_BORDER * 2;
-		const uint32 ShadowDepthBufferSizeY = ShadowBufferResolution.Y - SHADOW_BORDER * 2;
-		// Transform the shadow's position into shadowmap space
-		const FVector TransformedPosition = WorldToFace.TransformPosition(-PreShadowTranslation);
-
-		// Largest amount that the shadowmap will be downsampled to during sampling
-		// We need to take this into account when snapping to get a stable result
-		// This corresponds to the maximum kernel filter size used by subsurface shadows in ShadowProjectionPixelShader.usf
-		const int32 MaxDownsampleFactor = 4;
-		// Determine the distance necessary to snap the shadow's position to the nearest texel
-		const float SnapX = FMath::Fmod(TransformedPosition.X, 2.0f * MaxDownsampleFactor / ShadowDepthBufferSizeX);
-		const float SnapY = FMath::Fmod(TransformedPosition.Y, 2.0f * MaxDownsampleFactor / ShadowDepthBufferSizeY);
-		// Snap the shadow's position and transform it back into world space
-		// This snapping prevents sub-texel camera movements which removes view dependent aliasing from the final shadow result
-		// This only maintains stable shadows under camera translation and rotation
-		const FVector SnappedWorldPosition = WorldToFace.Inverse().TransformPosition(TransformedPosition - FVector(SnapX, SnapY, 0.0f));
-		PreShadowTranslation = -SnappedWorldPosition;
-	}
 
-	check(MaxSubjectZ > MinSubjectZ);
-
-	const float ClampedMaxLightW = FMath::Min(MinSubjectZ + Initializer.MaxDistanceToCastInLightW, (float)HALF_WORLD_MAX);
-	MinPreSubjectZ = Initializer.MinLightW;
-
-	const FMatrix SubjectMatrix = WorldToFace * FShadowProjectionMatrix(MinSubjectZ,MaxSubjectZ,Initializer.WAxis);
-	const FMatrix PostSubjectMatrix = WorldToFace * FShadowProjectionMatrix(MinSubjectZ, ClampedMaxLightW, Initializer.WAxis);
-
-	float MaxSubjectDepth = SubjectMatrix.TransformPosition(
-		Initializer.SubjectBounds.Origin
-		+ WorldToLightScaled.Inverse().TransformVector(Initializer.FaceDirection) * Initializer.SubjectBounds.SphereRadius
-		).Z;
-
-	if (Initializer.bOnePassPointLightShadow)
+	if(bInReflectiveShadowMap)
 	{
-		MaxSubjectDepth = Initializer.SubjectBounds.SphereRadius;
-	}
+		check(!bOnePassPointLightShadow);
 
-	InvMaxSubjectDepth = 1.0f / MaxSubjectDepth;
+		SplitIndex = 0;
 
-	// Store the view matrix
-	// Reorder the vectors to match the main view, since ShadowViewMatrix will be used to override the main view's view matrix during shadow depth rendering
-	ShadowViewMatrix = Initializer.WorldToLight * 
-		FMatrix(
-			FPlane(0,	0,	1,	0),
-			FPlane(1,	0,	0,	0),
-			FPlane(0,	1,	0,	0),
-			FPlane(0,	0,	0,	1));
+		// Quantise the RSM in shadow texel space
+		static bool bQuantize = true;
+		if ( bQuantize )
+		{
+			const FIntPoint ShadowBufferResolution = GSceneRenderTargets.GetReflectiveShadowMapTextureResolution();
+			uint32 ShadowDepthBufferSizeX = ShadowBufferResolution.X;
+			uint32 ShadowDepthBufferSizeY = ShadowBufferResolution.Y;
+			// Transform the shadow's position into shadowmap space
+			const FVector TransformedPosition = WorldToFace.TransformPosition(-PreShadowTranslation);
+
+			// Largest amount that the shadowmap will be downsampled to during sampling
+			// We need to take this into account when snapping to get a stable result
+			// This corresponds to the maximum kernel filter size used by subsurface shadows in ShadowProjectionPixelShader.usf
+			static int32 MaxDownsampleFactor = 4;
+			// Determine the distance necessary to snap the shadow's position to the nearest texel
+			const float SnapX = FMath::Fmod(TransformedPosition.X, 2.0f * MaxDownsampleFactor / ShadowDepthBufferSizeX);
+			const float SnapY = FMath::Fmod(TransformedPosition.Y, 2.0f * MaxDownsampleFactor / ShadowDepthBufferSizeY);
+			// Snap the shadow's position and transform it back into world space
+			// This snapping prevents sub-texel camera movements which removes view dependent aliasing from the final shadow result
+			// This only maintains stable shadows under camera translation and rotation
+			const FVector SnappedWorldPosition = WorldToFace.Inverse().TransformPosition(TransformedPosition - FVector(SnapX, SnapY, 0.0f));
+			PreShadowTranslation = -SnappedWorldPosition;
+		}
 
-	SubjectAndReceiverMatrix = SubjectMatrix,
-	ReceiverMatrix = PostSubjectMatrix;
-	InvReceiverMatrix = ReceiverMatrix.Inverse();
+		ShadowBounds = FSphere(-PreShadowTranslation, Initializer.SubjectBounds.SphereRadius);
 
-	if (Initializer.SplitIndex >= 0 && Initializer.bDirectionalLight)
-	{
-		checkSlow(InDependentView);
-		ShadowBounds = InLightSceneInfo->Proxy->GetShadowSplitBounds(*InDependentView, SplitIndex, 0);
+		GetViewFrustumBounds(CasterFrustum, SubjectAndReceiverMatrix, true);
 	}
 	else
 	{
-		ShadowBounds = FSphere(-Initializer.PreShadowTranslation, Initializer.SubjectBounds.SphereRadius);
-	}
-
-	// Any meshes between the light and the subject can cast shadows, also any meshes inside the subject region
-	const FMatrix CasterMatrix = WorldToFace * FShadowProjectionMatrix(Initializer.MinLightW,MaxSubjectZ,Initializer.WAxis);
-	GetViewFrustumBounds(CasterFrustum,CasterMatrix,true);
-	GetViewFrustumBounds(ReceiverFrustum,ReceiverMatrix,true);
-
-	UpdateShaderDepthBias();
-}
-
-FProjectedShadowInfo::FProjectedShadowInfo(
-	FLightSceneInfo* InLightSceneInfo,
-	FViewInfo* InDependentView,
-	const FRsmWholeSceneProjectedShadowInitializer& Initializer,
-	uint32 InResolutionX,
-	uint32 InResolutionY
-	)
-	:	LightSceneInfo(InLightSceneInfo)
-	,	LightSceneInfoCompact(InLightSceneInfo)
-	,	ParentSceneInfo(NULL)
-	,	DependentView(InDependentView)
-	,	ShadowId(INDEX_NONE)
-	,	PreShadowTranslation(Initializer.PreShadowTranslation)
-	,	CascadeSettings(Initializer.CascadeSettings)
-	,	X(0)
-	,	Y(0)
-	,	ResolutionX(InResolutionX)
-	,	ResolutionY(InResolutionY)
-	,	MaxScreenPercent(1.0f)
-	,   SplitIndex(0)
-	,	bAllocated(false)
-	,	bAllocatedInTranslucentLayout(false)
-	,	bRendered(false)
-	,	bAllocatedInPreshadowCache(false)
-	,	bDepthsCached(false)
-	,	bDirectionalLight(Initializer.bDirectionalLight)
-	,	bWholeSceneShadow(true)
-	,	bOnePassPointLightShadow(false)
-	,	bReflectiveShadowmap(true) 
-	,	bTranslucentShadow(false)
-	,	bPreShadow(false)
-	,	bValidTransform(true)
-{
-	FVector	XAxis, YAxis;
-	Initializer.FaceDirection.FindBestAxisVectors(XAxis,YAxis);
-	const FMatrix WorldToLightScaled = Initializer.WorldToLight * FScaleMatrix(Initializer.Scales);
-	const FMatrix WorldToFace = WorldToLightScaled * FBasisVectorMatrix(-XAxis,YAxis,Initializer.FaceDirection.SafeNormal(),FVector::ZeroVector);
-
-	MaxSubjectZ = WorldToFace.TransformPosition(Initializer.SubjectBounds.Origin).Z + Initializer.SubjectBounds.SphereRadius;
+		if(Initializer.bDirectionalLight)
+		{
+			// Limit how small the depth range can be for smaller cascades
+			// This is needed for shadow modes like subsurface shadows which need depth information outside of the smaller cascade depth range
+			//@todo - expose this value to the ini
+			const float DepthRangeClamp = 5000;
+			MaxSubjectZ = FMath::Max(MaxSubjectZ, DepthRangeClamp);
+			MinSubjectZ = FMath::Min(MinSubjectZ, -DepthRangeClamp);
+
+			const FIntPoint ShadowBufferResolution = GSceneRenderTargets.GetShadowDepthTextureResolution();
+			const uint32 ShadowDepthBufferSizeX = ShadowBufferResolution.X - SHADOW_BORDER * 2;
+			const uint32 ShadowDepthBufferSizeY = ShadowBufferResolution.Y - SHADOW_BORDER * 2;
+			// Transform the shadow's position into shadowmap space
+			const FVector TransformedPosition = WorldToFace.TransformPosition(-PreShadowTranslation);
+
+			// Largest amount that the shadowmap will be downsampled to during sampling
+			// We need to take this into account when snapping to get a stable result
+			// This corresponds to the maximum kernel filter size used by subsurface shadows in ShadowProjectionPixelShader.usf
+			const int32 MaxDownsampleFactor = 4;
+			// Determine the distance necessary to snap the shadow's position to the nearest texel
+			const float SnapX = FMath::Fmod(TransformedPosition.X, 2.0f * MaxDownsampleFactor / ShadowDepthBufferSizeX);
+			const float SnapY = FMath::Fmod(TransformedPosition.Y, 2.0f * MaxDownsampleFactor / ShadowDepthBufferSizeY);
+			// Snap the shadow's position and transform it back into world space
+			// This snapping prevents sub-texel camera movements which removes view dependent aliasing from the final shadow result
+			// This only maintains stable shadows under camera translation and rotation
+			const FVector SnappedWorldPosition = WorldToFace.Inverse().TransformPosition(TransformedPosition - FVector(SnapX, SnapY, 0.0f));
+			PreShadowTranslation = -SnappedWorldPosition;
+		}
 
-	MinSubjectZ = FMath::Max(MaxSubjectZ - Initializer.SubjectBounds.SphereRadius * 2,Initializer.MinLightW);
+		if (Initializer.SplitIndex >= 0 && Initializer.bDirectionalLight)
+		{
+			checkSlow(InDependentView);
+			ShadowBounds = InLightSceneInfo->Proxy->GetShadowSplitBounds(*InDependentView, SplitIndex, 0);
+		}
+		else
+		{
+			ShadowBounds = FSphere(-Initializer.PreShadowTranslation, Initializer.SubjectBounds.SphereRadius);
+		}
 
-	static float Maxz = 0.0f;
-	static float Minz = 0.0f;
-	if ( Minz != 0.0f ) MinSubjectZ = Minz;
-	if ( Maxz != 0.0f ) MaxSubjectZ = Maxz;
+		// Any meshes between the light and the subject can cast shadows, also any meshes inside the subject region
+		const FMatrix CasterMatrix = WorldToFace * FShadowProjectionMatrix(Initializer.MinLightW, MaxSubjectZ, Initializer.WAxis);
+		GetViewFrustumBounds(CasterFrustum, CasterMatrix, true);
+	}
 
+	check(MaxSubjectZ > MinSubjectZ);
 
 	const float ClampedMaxLightW = FMath::Min(MinSubjectZ + Initializer.MaxDistanceToCastInLightW, (float)HALF_WORLD_MAX);
 	MinPreSubjectZ = Initializer.MinLightW;
 
-	const FMatrix SubjectMatrix = WorldToFace * FShadowProjectionMatrix(MinSubjectZ,MaxSubjectZ,Initializer.WAxis);
-	const FMatrix PostSubjectMatrix = WorldToFace * FShadowProjectionMatrix(MinSubjectZ, ClampedMaxLightW, Initializer.WAxis);
-
-	// Quantise the RSM in shadow texel space
-	static bool quantize = true;
-	if ( quantize )
-	{
-		const FIntPoint ShadowBufferResolution = GSceneRenderTargets.GetReflectiveShadowMapTextureResolution();
-		uint32 ShadowDepthBufferSizeX = ShadowBufferResolution.X ;
-		uint32 ShadowDepthBufferSizeY = ShadowBufferResolution.Y ;
-		// Transform the shadow's position into shadowmap space
-		const FVector TransformedPosition = WorldToFace.TransformPosition(-PreShadowTranslation);
-
-		// Largest amount that the shadowmap will be downsampled to during sampling
-		// We need to take this into account when snapping to get a stable result
-		// This corresponds to the maximum kernel filter size used by subsurface shadows in ShadowProjectionPixelShader.usf
-		static int32 MaxDownsampleFactor = 4;
-		// Determine the distance necessary to snap the shadow's position to the nearest texel
-		const float SnapX = FMath::Fmod(TransformedPosition.X, 2.0f * MaxDownsampleFactor / ShadowDepthBufferSizeX);
-		const float SnapY = FMath::Fmod(TransformedPosition.Y, 2.0f * MaxDownsampleFactor / ShadowDepthBufferSizeY);
-		// Snap the shadow's position and transform it back into world space
-		// This snapping prevents sub-texel camera movements which removes view dependent aliasing from the final shadow result
-		// This only maintains stable shadows under camera translation and rotation
-		const FVector SnappedWorldPosition = WorldToFace.Inverse().TransformPosition(TransformedPosition - FVector(SnapX, SnapY, 0.0f));
-		PreShadowTranslation = -SnappedWorldPosition;
-	}
+	SubjectAndReceiverMatrix = WorldToFace * FShadowProjectionMatrix(MinSubjectZ, MaxSubjectZ, Initializer.WAxis);
+	ReceiverMatrix = WorldToFace * FShadowProjectionMatrix(MinSubjectZ, ClampedMaxLightW, Initializer.WAxis);
 
-	float MaxSubjectDepth = SubjectMatrix.TransformPosition(
+	float MaxSubjectDepth = SubjectAndReceiverMatrix.TransformPosition(
 		Initializer.SubjectBounds.Origin
 		+ WorldToLightScaled.Inverse().TransformVector(Initializer.FaceDirection) * Initializer.SubjectBounds.SphereRadius
 		).Z;
 
+	if (Initializer.bOnePassPointLightShadow)
+	{
+		MaxSubjectDepth = Initializer.SubjectBounds.SphereRadius;
+	}
+
 	InvMaxSubjectDepth = 1.0f / MaxSubjectDepth;
 
 	// Store the view matrix
@@ -689,14 +622,9 @@ FProjectedShadowInfo::FProjectedShadowInfo(
 		FPlane(0,	1,	0,	0),
 		FPlane(0,	0,	0,	1));
 
-	SubjectAndReceiverMatrix = SubjectMatrix,
-	ReceiverMatrix = PostSubjectMatrix;
 	InvReceiverMatrix = ReceiverMatrix.Inverse();
 
-	ShadowBounds = FSphere(-PreShadowTranslation, Initializer.SubjectBounds.SphereRadius);
-	
-	GetViewFrustumBounds(CasterFrustum, SubjectAndReceiverMatrix, true);
-	GetViewFrustumBounds(ReceiverFrustum,ReceiverMatrix,true);
+	GetViewFrustumBounds(ReceiverFrustum, ReceiverMatrix, true);
 
 	UpdateShaderDepthBias();
 }
@@ -1558,8 +1486,10 @@ void FDeferredShadingSceneRenderer::CreateWholeSceneProjectedShadow(FLightSceneI
 					ProjectedShadowInitializer,
 					SizeX,
 					SizeY,
-					FadeAlphas
+					false	// no RSM
 					);
+				ProjectedShadowInfo->FadeAlphas = FadeAlphas;
+
 				VisibleLightInfo.MemStackProjectedShadows.Add(ProjectedShadowInfo);
 				VisibleLightInfo.AllProjectedShadows.Add(ProjectedShadowInfo);
 
@@ -2070,9 +2000,11 @@ void FDeferredShadingSceneRenderer::InitDynamicShadows()
 											//@todo - remove the shadow border for whole scene shadows
 											ShadowBufferResolution.X - SHADOW_BORDER * 2,
 											ShadowBufferResolution.Y - SHADOW_BORDER * 2,
-											FadeAlphas
+											false	// no RSM
 											);
 
+										ProjectedShadowInfo->FadeAlphas = FadeAlphas;
+
 										FVisibleLightInfo& LightViewInfo = VisibleLightInfos[LightSceneInfo->Id];
 										VisibleLightInfo.MemStackProjectedShadows.Add(ProjectedShadowInfo);
 										VisibleLightInfo.AllProjectedShadows.Add(ProjectedShadowInfo);
@@ -2087,7 +2019,7 @@ void FDeferredShadingSceneRenderer::InitDynamicShadows()
 									if (LightPropagationVolume && View.FinalPostProcessSettings.LPVIntensity > 0 )
 									{
 										// Generate the RSM shadow info
-										FRsmWholeSceneProjectedShadowInitializer ProjectedShadowInitializer;
+										FWholeSceneProjectedShadowInitializer ProjectedShadowInitializer;
 										FLightPropagationVolume& Lpv = *LightPropagationVolume;
 
 										if (LightSceneInfo->Proxy->GetViewDependentRsmWholeSceneProjectedShadowInitializer( View, Lpv.GetBoundingBox(), ProjectedShadowInitializer ))
@@ -2100,7 +2032,8 @@ void FDeferredShadingSceneRenderer::InitDynamicShadows()
 												&View,
 												ProjectedShadowInitializer,
 												ShadowBufferResolution.X,
-												ShadowBufferResolution.Y );
+												ShadowBufferResolution.Y,
+												true);
 
 											FVisibleLightInfo& LightViewInfo = VisibleLightInfos[LightSceneInfo->Id];
 											VisibleLightInfo.MemStackProjectedShadows.Add(ProjectedShadowInfo);
