#UE4 - Fix ensures in scene render target stuff when enabling buffer visualization and on startup in vehiclegame

diff --git a/Engine/Source/Runtime/Renderer/Private/DeferredShadingRenderer.cpp b/Engine/Source/Runtime/Renderer/Private/DeferredShadingRenderer.cpp
--- a/Engine/Source/Runtime/Renderer/Private/DeferredShadingRenderer.cpp
+++ b/Engine/Source/Runtime/Renderer/Private/DeferredShadingRenderer.cpp
@@ -760,6 +760,11 @@ void FDeferredShadingSceneRenderer::Render()
 			FinishRenderViewTarget(&Views[ViewIndex], ViewIndex == (Views.Num() - 1));
 		}
 	}
+	else
+	{
+		// Release the original reference on the scene render targets
+		GSceneRenderTargets.AdjustGBufferRefCount(-1);
+	}
 
 	RenderFinish();
 }
diff --git a/Engine/Source/Runtime/Renderer/Private/PostProcess/PostProcessing.cpp b/Engine/Source/Runtime/Renderer/Private/PostProcess/PostProcessing.cpp
--- a/Engine/Source/Runtime/Renderer/Private/PostProcess/PostProcessing.cpp
+++ b/Engine/Source/Runtime/Renderer/Private/PostProcess/PostProcessing.cpp
@@ -605,13 +605,22 @@ static void AddGBufferVisualization(FPostprocessContext& Context, FRenderingComp
 			// Loop over materials, creating stages for generation and downsampling of the tiles.			
 			for (TArray<UMaterialInterface*>::TConstIterator It = Context.View.FinalPostProcessSettings.BufferVisualizationOverviewMaterials.CreateConstIterator(); It; ++It)
 			{
-				if (*It)
+				auto MaterialInterface = *It;
+				if (MaterialInterface)
 				{
 					// Apply requested material
-					FRenderingCompositePass* MaterialPass = Context.Graph.RegisterPass(new(FMemStack::Get()) FRCPassPostProcessMaterial(*It));
+					FRenderingCompositePass* MaterialPass = Context.Graph.RegisterPass(new(FMemStack::Get()) FRCPassPostProcessMaterial(MaterialInterface));
 					MaterialPass->SetInput(ePId_Input0, FRenderingCompositeOutputRef(IncomingStage));
 					MaterialPass->SetInput(ePId_Input1, FRenderingCompositeOutputRef(SeparateTranslucencyInput));
 
+					auto Proxy = MaterialInterface->GetRenderProxy(false);
+					const FMaterial* Material = Proxy->GetMaterial(Context.View.GetFeatureLevel());
+					if (Material->NeedsGBuffer())
+					{
+						// AdjustGBufferRefCount(-1) call is done when the pass gets executed
+						GSceneRenderTargets.AdjustGBufferRefCount(1);
+					}
+
 					if (BaseFilename.Len())
 					{
 						FString MaterialFilename = BaseFilename + TEXT("_") + (*It)->GetName();
@@ -656,6 +665,12 @@ static void AddGBufferVisualization(FPostprocessContext& Context, FRenderingComp
 			MaterialPass->SetInput(ePId_Input0, FRenderingCompositeOutputRef(Context.FinalOutput));
 			MaterialPass->SetInput(ePId_Input1, FRenderingCompositeOutputRef(SeparateTranslucencyInput));
 			Context.FinalOutput = FRenderingCompositeOutputRef(MaterialPass);
+
+			if (Material->GetMaterialResource(Context.View.GetFeatureLevel())->NeedsGBuffer())
+			{
+				// AdjustGBufferRefCount(-1) call is done when the pass gets executed
+				GSceneRenderTargets.AdjustGBufferRefCount(1);
+			}
 		}
 	}
 }
