UE4 HTML5  Build fix

diff --git a/Engine/Source/Runtime/Renderer/Private/ForwardTranslucentRendering.cpp b/Engine/Source/Runtime/Renderer/Private/ForwardTranslucentRendering.cpp
--- a/Engine/Source/Runtime/Renderer/Private/ForwardTranslucentRendering.cpp
+++ b/Engine/Source/Runtime/Renderer/Private/ForwardTranslucentRendering.cpp
@@ -303,16 +303,6 @@ void FForwardShadingSceneRenderer::RenderTranslucency()
 		const bool bGammaSpace = !IsMobileHDR();
 		const bool bLinearHDR64 = !bGammaSpace && !IsMobileHDR32bpp();
 
-#if PLATFORM_HTML5
-		// Copy the view so emulation of framebuffer fetch works for alpha=depth.
-		// Possible optimization: this copy shouldn't be needed unless something uses fetch of depth.
-		if(bLinearHDR64 && GSupportsRenderTargetFormat_PF_FloatRGBA && (GSupportsShaderFramebufferFetch == false) && (!IsPCPlatform(GRHIShaderPlatform)))
-		{
-			CopySceneAlpha();
-		}
-#endif 
-
-
 		SCOPED_DRAW_EVENT(Translucency, DEC_SCENE_ITEMS);
 
 		for (int32 ViewIndex = 0; ViewIndex < Views.Num(); ViewIndex++)
@@ -321,6 +311,16 @@ void FForwardShadingSceneRenderer::RenderTranslucency()
 
 			const FViewInfo& View = Views[ViewIndex];
 
+			#if PLATFORM_HTML5
+			// Copy the view so emulation of framebuffer fetch works for alpha=depth.
+			// Possible optimization: this copy shouldn't be needed unless something uses fetch of depth.
+			if(bLinearHDR64 && GSupportsRenderTargetFormat_PF_FloatRGBA && (GSupportsShaderFramebufferFetch == false) && (!IsPCPlatform(GRHIShaderPlatform)))
+			{
+				CopySceneAlpha(View);
+			}
+			#endif 
+
+
 			if (!bGammaSpace)
 			{
 				GSceneRenderTargets.BeginRenderingTranslucency(View);
