* minor optimization

diff --git a/Engine/Source/Runtime/Renderer/Private/PostProcess/RenderTargetPool.cpp b/Engine/Source/Runtime/Renderer/Private/PostProcess/RenderTargetPool.cpp
--- a/Engine/Source/Runtime/Renderer/Private/PostProcess/RenderTargetPool.cpp
+++ b/Engine/Source/Runtime/Renderer/Private/PostProcess/RenderTargetPool.cpp
@@ -152,7 +152,7 @@ bool FRenderTargetPool::FindFreeElement(const FPooledRenderTargetDesc& Desc, TRe
 	uint32 FoundIndex = -1;
 
 	// try to find a suitable element in the pool
-	for(uint32 i = 0; i < (uint32)PooledRenderTargets.Num(); ++i)
+	for(uint32 i = 0, Num = (uint32)PooledRenderTargets.Num(); i < Num; ++i)
 	{
 		FPooledRenderTarget* Element = PooledRenderTargets[i];
 
