* fixed TTP 336027 CRITICAL: Regression: RENDERING: Crash: Enabling Depth of Field Layers in a viewport

diff --git a/Engine/Source/Runtime/Renderer/Private/CompositionLighting/PostProcessPassThrough.cpp b/Engine/Source/Runtime/Renderer/Private/CompositionLighting/PostProcessPassThrough.cpp
--- a/Engine/Source/Runtime/Renderer/Private/CompositionLighting/PostProcessPassThrough.cpp
+++ b/Engine/Source/Runtime/Renderer/Private/CompositionLighting/PostProcessPassThrough.cpp
@@ -155,6 +155,7 @@ FPooledRenderTargetDesc FRCPassPostProcessPassThrough::ComputeOutputDesc(EPassOu
 		}
 	}
 
+	Ret.Reset();
 	Ret.DebugName = TEXT("PassThrough");
 
 	return Ret;
diff --git a/Engine/Source/Runtime/Renderer/Private/PostProcess/PostProcessing.cpp b/Engine/Source/Runtime/Renderer/Private/PostProcess/PostProcessing.cpp
--- a/Engine/Source/Runtime/Renderer/Private/PostProcess/PostProcessing.cpp
+++ b/Engine/Source/Runtime/Renderer/Private/PostProcess/PostProcessing.cpp
@@ -1175,6 +1175,7 @@ void FPostProcessing::Process(const FViewInfo& View, TRefCountPtr<IPooledRenderT
 			// todo: this should come from View.Family->RenderTarget
 			Desc.Format = PF_B8G8R8A8;
 			Desc.NumMips = 1;
+			Desc.DebugName = TEXT("FinalPostProcessColor");
 
 			GRenderTargetPool.CreateUntrackedElement(Desc, Temp, Item);
 
diff --git a/Engine/Source/Runtime/Renderer/Private/PostProcess/VisualizeTexture.cpp b/Engine/Source/Runtime/Renderer/Private/PostProcess/VisualizeTexture.cpp
--- a/Engine/Source/Runtime/Renderer/Private/PostProcess/VisualizeTexture.cpp
+++ b/Engine/Source/Runtime/Renderer/Private/PostProcess/VisualizeTexture.cpp
@@ -882,6 +882,7 @@ void FVisualizeTexture::OnStartFrame(const FSceneView& View)
 	// VisualizeTexture observed render target is set each frame
 	VisualizeTextureContent.SafeRelease();
 	VisualizeTextureDesc = FPooledRenderTargetDesc();
+	VisualizeTextureDesc.DebugName = TEXT("VisualizeTexture");
 
 	ObservedDebugNameReusedCurrent = 0;
 
