Fix regression crash on Android devices that don't support floating point render targets

diff --git a/Engine/Source/Runtime/Renderer/Private/PostProcess/SceneRenderTargets.cpp b/Engine/Source/Runtime/Renderer/Private/PostProcess/SceneRenderTargets.cpp
--- a/Engine/Source/Runtime/Renderer/Private/PostProcess/SceneRenderTargets.cpp
+++ b/Engine/Source/Runtime/Renderer/Private/PostProcess/SceneRenderTargets.cpp
@@ -1185,6 +1185,15 @@ EPixelFormat FSceneRenderTargets::GetSceneColorFormat() const
 			SceneColorBufferFormat = PF_A32B32G32R32F; break;
 	}
 
+	if (GRHIFeatureLevel <= ERHIFeatureLevel::ES2)
+	{
+		if ((!GSupportsRenderTargetFormat_PF_FloatRGBA && (SceneColorBufferFormat == PF_FloatRGBA || SceneColorBufferFormat == PF_FloatRGB))
+			|| SceneColorBufferFormat == PF_A16B16G16R16)
+		{
+			SceneColorBufferFormat = PF_B8G8R8A8;
+		}
+	}
+
 	return SceneColorBufferFormat;
 }
 
