* AllocSceneColor should not be called for ES2, issue was fixed by Jack with change in SceRenderTaregts.cpp but it shouldn't have gone there.
TTP 336087 BLOCKER: Regression: Android: Mobile Temple worked on 4.1, 4.2 crashes on Mali Galaxy Note 10.1

diff --git a/Engine/Source/Runtime/Renderer/Private/PostProcess/SceneRenderTargets.cpp b/Engine/Source/Runtime/Renderer/Private/PostProcess/SceneRenderTargets.cpp
--- a/Engine/Source/Runtime/Renderer/Private/PostProcess/SceneRenderTargets.cpp
+++ b/Engine/Source/Runtime/Renderer/Private/PostProcess/SceneRenderTargets.cpp
@@ -289,6 +289,12 @@ int32 FSceneRenderTargets::GetNumGBufferTargets() const
 
 void FSceneRenderTargets::AllocSceneColor()
 {
+	if(GRHIFeatureLevel == ERHIFeatureLevel::ES2)
+	{
+		// ES2 is forward and not yet releasing the SceneColor during the frame (is not calling FRCPassPostProcessTonemap::Process())
+		return;
+	}
+
 	if(SceneColor)
 	{
 		// no work needed
@@ -1165,6 +1171,8 @@ void FSceneRenderTargets::AllocateDeferredShadingPathRenderTargets()
 
 EPixelFormat FSceneRenderTargets::GetSceneColorFormat() const
 {
+	check(GRHIFeatureLevel != ERHIFeatureLevel::ES2);
+
 	// Potentially allocate an alpha channel in the scene color texture to store the resolved scene depth.
 	EPixelFormat SceneColorBufferFormat = PF_FloatRGBA;
 
diff --git a/Engine/Source/Runtime/Renderer/Private/PostProcess/SceneRenderTargets.h b/Engine/Source/Runtime/Renderer/Private/PostProcess/SceneRenderTargets.h
--- a/Engine/Source/Runtime/Renderer/Private/PostProcess/SceneRenderTargets.h
+++ b/Engine/Source/Runtime/Renderer/Private/PostProcess/SceneRenderTargets.h
@@ -368,7 +368,7 @@ public:
 
 private: // Get...() methods instead of direct access
 
-	// 0 before BeginRenderingSceneColor and after tone mapping
+	// 0 before BeginRenderingSceneColor and after tone mapping, for ES2 it's not released during the frame
 	TRefCountPtr<IPooledRenderTarget> SceneColor;
 	// also used as LDR scene color
 	TRefCountPtr<IPooledRenderTarget> LightAttenuation;
