* fixed TTP 336087 BLOCKER: Regression: Android: Mobile Temple worked on 4.1, 4.2 crashes on Mali Galaxy Note 10.1
  in UE4.2 this is already changed

diff --git a/Engine/Source/Runtime/Renderer/Private/PostProcess/SceneRenderTargets.cpp b/Engine/Source/Runtime/Renderer/Private/PostProcess/SceneRenderTargets.cpp
--- a/Engine/Source/Runtime/Renderer/Private/PostProcess/SceneRenderTargets.cpp
+++ b/Engine/Source/Runtime/Renderer/Private/PostProcess/SceneRenderTargets.cpp
@@ -923,17 +923,10 @@ void FSceneRenderTargets::AllocateForwardShadingPathRenderTargets()
 		GRenderTargetPool.FindFreeElement(Desc, SceneDepthZ, TEXT("SceneDepthZ"));
 	}
 
-	EPixelFormat Format = GSupportsRenderTargetFormat_PF_FloatRGBA ? PF_FloatRGBA : PF_B8G8R8A8;
-	if (!IsMobileHDR() || IsMobileHDR32bpp()) 
-	{
-		Format = PF_B8G8R8A8;
-	}
-	// Create the scene color.
-	{
-		FPooledRenderTargetDesc Desc(FPooledRenderTargetDesc::Create2DDesc(BufferSize, Format, TexCreate_None, TexCreate_RenderTargetable, false));
-		Desc.NumSamples = NumSamples;
-		GRenderTargetPool.FindFreeElement(Desc, SceneColor, TEXT("SceneColor"));
-	}
+	// on ES2 we don't do on demand allocation of SceneColor yet (in non ES2 it's released in the Tonemapper Process())
+	AllocSceneColor();
+
+	EPixelFormat Format = SceneColor->GetDesc().Format;
 
 	// For 64-bit ES2 without framebuffer fetch, create extra render target for copy of alpha channel.
 	if((Format == PF_FloatRGBA) && (GSupportsShaderFramebufferFetch == false)) 
@@ -1173,34 +1166,36 @@ EPixelFormat FSceneRenderTargets::GetSceneColorFormat() const
 {
 	check(GRHIFeatureLevel != ERHIFeatureLevel::ES2);
 
-	// Potentially allocate an alpha channel in the scene color texture to store the resolved scene depth.
 	EPixelFormat SceneColorBufferFormat = PF_FloatRGBA;
 
-	switch(CurrentSceneColorFormat)
-	{
-		case 0:
-			SceneColorBufferFormat = PF_R8G8B8A8; break;
-		case 1:
-			SceneColorBufferFormat = PF_A2B10G10R10; break;
-		case 2:	
-			SceneColorBufferFormat = PF_FloatR11G11B10; break;
-		case 3:	
-			SceneColorBufferFormat = PF_FloatRGB; break;
-		case 4:
-			// default
-			break;
-		case 5:
-			SceneColorBufferFormat = PF_A32B32G32R32F; break;
-	}
-
 	if (GRHIFeatureLevel <= ERHIFeatureLevel::ES2)
 	{
-		if ((!GSupportsRenderTargetFormat_PF_FloatRGBA && (SceneColorBufferFormat == PF_FloatRGBA || SceneColorBufferFormat == PF_FloatRGB))
-			|| SceneColorBufferFormat == PF_A16B16G16R16)
+		// Potentially allocate an alpha channel in the scene color texture to store the resolved scene depth.
+		SceneColorBufferFormat = GSupportsRenderTargetFormat_PF_FloatRGBA ? PF_FloatRGBA : PF_B8G8R8A8;
+		if (!IsMobileHDR() || IsMobileHDR32bpp()) 
 		{
 			SceneColorBufferFormat = PF_B8G8R8A8;
 		}
 	}
+	else
+	{
+		switch(CurrentSceneColorFormat)
+		{
+			case 0:
+				SceneColorBufferFormat = PF_R8G8B8A8; break;
+			case 1:
+				SceneColorBufferFormat = PF_A2B10G10R10; break;
+			case 2:	
+				SceneColorBufferFormat = PF_FloatR11G11B10; break;
+			case 3:	
+				SceneColorBufferFormat = PF_FloatRGB; break;
+			case 4:
+				// default
+				break;
+			case 5:
+				SceneColorBufferFormat = PF_A32B32G32R32F; break;
+		}
+	}
 
 	return SceneColorBufferFormat;
 }
