* missing change for this:
  * fixed TTP 336087 BLOCKER: Regression: Android: Mobile Temple worked on 4.1, 4.2 crashes on Mali Galaxy Note 10.1
    in UE4.2 this is already changed

diff --git a/Engine/Source/Runtime/Renderer/Private/PostProcess/SceneRenderTargets.cpp b/Engine/Source/Runtime/Renderer/Private/PostProcess/SceneRenderTargets.cpp
--- a/Engine/Source/Runtime/Renderer/Private/PostProcess/SceneRenderTargets.cpp
+++ b/Engine/Source/Runtime/Renderer/Private/PostProcess/SceneRenderTargets.cpp
@@ -289,12 +289,6 @@ int32 FSceneRenderTargets::GetNumGBufferTargets() const
 
 void FSceneRenderTargets::AllocSceneColor()
 {
-	if(GRHIFeatureLevel == ERHIFeatureLevel::ES2)
-	{
-		// ES2 is forward and not yet releasing the SceneColor during the frame (is not calling FRCPassPostProcessTonemap::Process())
-		return;
-	}
-
 	if(SceneColor)
 	{
 		// no work needed
