* missing change for this:
  * fixed TTP 336087 BLOCKER: Regression: Android: Mobile Temple worked on 4.1, 4.2 crashes on Mali Galaxy Note 10.1
    in UE4.2 this is already changed

diff --git a/Engine/Source/Runtime/Renderer/Private/PostProcess/SceneRenderTargets.cpp b/Engine/Source/Runtime/Renderer/Private/PostProcess/SceneRenderTargets.cpp
--- a/Engine/Source/Runtime/Renderer/Private/PostProcess/SceneRenderTargets.cpp
+++ b/Engine/Source/Runtime/Renderer/Private/PostProcess/SceneRenderTargets.cpp
@@ -1158,13 +1158,11 @@ void FSceneRenderTargets::AllocateDeferredShadingPathRenderTargets()
 
 EPixelFormat FSceneRenderTargets::GetSceneColorFormat() const
 {
-	check(GRHIFeatureLevel != ERHIFeatureLevel::ES2);
-
 	EPixelFormat SceneColorBufferFormat = PF_FloatRGBA;
 
 	if (GRHIFeatureLevel <= ERHIFeatureLevel::ES2)
 	{
-		// Potentially allocate an alpha channel in the scene color texture to store the resolved scene depth.
+		// Potentially allocate an alpha channel in th -fe scene color texture to store the resolved scene depth.
 		SceneColorBufferFormat = GSupportsRenderTargetFormat_PF_FloatRGBA ? PF_FloatRGBA : PF_B8G8R8A8;
 		if (!IsMobileHDR() || IsMobileHDR32bpp()) 
 		{
