* fixed MSAA for mobile

diff --git a/Engine/Source/Runtime/Renderer/Private/PostProcess/SceneRenderTargets.cpp b/Engine/Source/Runtime/Renderer/Private/PostProcess/SceneRenderTargets.cpp
--- a/Engine/Source/Runtime/Renderer/Private/PostProcess/SceneRenderTargets.cpp
+++ b/Engine/Source/Runtime/Renderer/Private/PostProcess/SceneRenderTargets.cpp
@@ -314,6 +314,11 @@ void FSceneRenderTargets::AllocSceneColor()
 			Desc.TargetableFlags |= TexCreate_UAV;
 		}
 
+		if (GRHIFeatureLevel == ERHIFeatureLevel::ES2)
+		{
+			Desc.NumSamples = CVarMobileMSAA.GetValueOnRenderThread() ? 4 : 1;
+		}
+
 		GRenderTargetPool.FindFreeElement(Desc, SceneColor, TEXT("SceneColor"));
 	}
 
@@ -908,20 +913,18 @@ uint8 Quantize8SignedByte(float x)
 
 void FSceneRenderTargets::AllocateForwardShadingPathRenderTargets()
 {
-	uint16 NumSamples = CVarMobileMSAA.GetValueOnRenderThread() ? 4 : 1;
+	// on ES2 we don't do on demand allocation of SceneColor yet (in non ES2 it's released in the Tonemapper Process())
+	AllocSceneColor();
+
+	EPixelFormat Format = SceneColor->GetDesc().Format;
 
 	{
 		// Create a texture to store the resolved scene depth, and a render-targetable surface to hold the unresolved scene depth.
 		FPooledRenderTargetDesc Desc(FPooledRenderTargetDesc::Create2DDesc(BufferSize, PF_DepthStencil, TexCreate_None, TexCreate_DepthStencilTargetable, false));
-		Desc.NumSamples = NumSamples;
+		Desc.NumSamples = SceneColor->GetDesc().NumSamples;
 		GRenderTargetPool.FindFreeElement(Desc, SceneDepthZ, TEXT("SceneDepthZ"));
 	}
 
-	// on ES2 we don't do on demand allocation of SceneColor yet (in non ES2 it's released in the Tonemapper Process())
-	AllocSceneColor();
-
-	EPixelFormat Format = SceneColor->GetDesc().Format;
-
 	// For 64-bit ES2 without framebuffer fetch, create extra render target for copy of alpha channel.
 	if((Format == PF_FloatRGBA) && (GSupportsShaderFramebufferFetch == false)) 
 	{
