#UE #HTML5

     HTML5 Build fix for a merge issue from CL. 2082256

diff --git a/Engine/Source/Runtime/Renderer/Private/ForwardTranslucentRendering.cpp b/Engine/Source/Runtime/Renderer/Private/ForwardTranslucentRendering.cpp
--- a/Engine/Source/Runtime/Renderer/Private/ForwardTranslucentRendering.cpp
+++ b/Engine/Source/Runtime/Renderer/Private/ForwardTranslucentRendering.cpp
@@ -67,10 +67,10 @@ void FForwardShadingSceneRenderer::CopySceneAlpha(const FSceneView& View)
 
 	PixelShader->SetParameters(View);
 
-	DrawRectangle( 
-		0, 0, 
-		X, Y, 
-		0, 0, 
+	DrawRectangle(
+		0, 0,
+		X, Y,
+		0, 0,
 		X, Y,
 		FIntPoint(X, Y),
 		GSceneRenderTargets.GetBufferSizeXY(),
@@ -105,15 +105,15 @@ public:
 		bBackFace(bInBackFace),
 		HitProxyId(InHitProxyId)
 	{}
-	
+
 	inline bool ShouldPackAmbientSH() const
 	{
 		// So shader code can read a single constant to get the ambient term
 		return true;
 	}
 
-	const FLightSceneInfo* GetSimpleDirectionalLight() const 
-	{ 
+	const FLightSceneInfo* GetSimpleDirectionalLight() const
+	{
 		return ((FScene*)View.Family->Scene)->SimpleDirectionalLight;
 	}
 
@@ -247,8 +247,8 @@ void FTranslucentPrimSet::DrawPrimitivesForForwardShading(const FViewInfo& View,
 }
 
 void FTranslucentPrimSet::RenderPrimitiveForForwardShading(
-	const FViewInfo& View, 
-	FPrimitiveSceneInfo* PrimitiveSceneInfo, 
+	const FViewInfo& View,
+	FPrimitiveSceneInfo* PrimitiveSceneInfo,
 	const FPrimitiveViewRelevance& ViewRelevance) const
 {
 	checkSlow(ViewRelevance.HasTranslucency());
@@ -316,9 +316,9 @@ void FForwardShadingSceneRenderer::RenderTranslucency()
 			// Possible optimization: this copy shouldn't be needed unless something uses fetch of depth.
 			if(bLinearHDR64 && GSupportsRenderTargetFormat_PF_FloatRGBA && (GSupportsShaderFramebufferFetch == false) && (!IsPCPlatform(GRHIShaderPlatform)))
 			{
-				CopySceneAlpha();
+				CopySceneAlpha(View);
 			}
-#endif 
+#endif
 
 			if (!bGammaSpace)
 			{
