Mobile Rendering : Enabled Distortion rendering for ES2

diff --git a/Engine/Shaders/DistortAccumulatePS.usf b/Engine/Shaders/DistortAccumulatePS.usf
--- a/Engine/Shaders/DistortAccumulatePS.usf
+++ b/Engine/Shaders/DistortAccumulatePS.usf
@@ -17,7 +17,7 @@ float4 DistortionParams;
 /** output distortion offsets as color so they can be accumulated (via blending) */
 void Main(
 	FVertexFactoryInterpolantsVSToPS Interpolants,
-	float4	PixelPosition	: TEXCOORD5,
+	float4	PixelPosition	: TEXCOORD6,
 	OPTIONAL_IsFrontFace,
 	out float4 OutColor		: SV_Target0
 	)
@@ -59,7 +59,11 @@ void Main(
 
 	// Sample depth at distortion offset
 	float2 DistortScreenUV = ScreenUV + Distortion;
+#if COMPILER_GLSL_ES2
+	float DistortSceneDepth = Texture2DSample(SceneColorTexture, SceneColorTextureSampler, DistortScreenUV).w;
+#else
 	float DistortSceneDepth = CalcSceneDepth(DistortScreenUV);
+#endif
 
 	// Soft thresholding 
 	float Bias = -Refraction.y;
@@ -68,7 +72,7 @@ void Main(
 	float ZCompare = MaterialParameters.ScreenPosition.w;
 	float InvWidth = 1.0f / max(1.0f, Range );
 	Distortion *= saturate((Z - ZCompare) * InvWidth + Bias);
-	
+
 	//Scale up for better precision in low/subtle refractions at the expense of artefacts at higher refraction.
 	Distortion*= DistortionScaleBias;
 
diff --git a/Engine/Shaders/DistortAccumulateVS.usf b/Engine/Shaders/DistortAccumulateVS.usf
--- a/Engine/Shaders/DistortAccumulateVS.usf
+++ b/Engine/Shaders/DistortAccumulateVS.usf
@@ -11,7 +11,7 @@
 struct FDistortAccumulateVSToPS
 {
 	FVertexFactoryInterpolantsVSToPS FactoryInterpolants;
-	float4	PixelPosition	: TEXCOORD5;
+	float4	PixelPosition	: TEXCOORD6;
 	float4	Position		: SV_POSITION;
 };
 
diff --git a/Engine/Shaders/LocalVertexFactoryCommon.usf b/Engine/Shaders/LocalVertexFactoryCommon.usf
--- a/Engine/Shaders/LocalVertexFactoryCommon.usf
+++ b/Engine/Shaders/LocalVertexFactoryCommon.usf
@@ -109,7 +109,7 @@ void SetTangents(inout FVertexFactoryInterpolantsVSToPS Interpolants, float3 InT
 
 #define NUM_PACKED_UVS (NUM_MATERIAL_TEXCOORDS > 1 ? 3 : 1)
 
-#define USE_INTERPOLATOR_345 ((!MATERIAL_LIGHTINGMODEL_UNLIT) || USES_TRANSFORM_VECTOR)
+#define USE_INTERPOLATOR_345 ((!MATERIAL_LIGHTINGMODEL_UNLIT) || USES_TRANSFORM_VECTOR || USES_DISTORTION)
 
 /** Interpolants packed for ES2.  All UVs are in the xy to ensure independent texture fetches. */
 struct FVertexFactoryInterpolantsVSToPS
diff --git a/Engine/Source/Runtime/Engine/Private/Materials/HLSLMaterialTranslator.h b/Engine/Source/Runtime/Engine/Private/Materials/HLSLMaterialTranslator.h
--- a/Engine/Source/Runtime/Engine/Private/Materials/HLSLMaterialTranslator.h
+++ b/Engine/Source/Runtime/Engine/Private/Materials/HLSLMaterialTranslator.h
@@ -585,6 +585,8 @@ public:
 		OutEnvironment.SetDefine(TEXT("INTERPOLATE_VERTEX_COLOR"), bUsesVertexColor ? TEXT("1") : TEXT("0")); 
 		OutEnvironment.SetDefine(TEXT("NEEDS_PARTICLE_COLOR"), bUsesParticleColor ? TEXT("1") : TEXT("0")); 
 		OutEnvironment.SetDefine(TEXT("USES_TRANSFORM_VECTOR"), bUsesTransformVector ? TEXT("1") : TEXT("0")); 
+		// Distortion uses tangent space transform 
+		OutEnvironment.SetDefine(TEXT("USES_DISTORTION"), Material->IsDistorted() ? TEXT("1") : TEXT("0")); 
 
 		OutEnvironment.SetDefine(TEXT("ENABLE_TRANSLUCENCY_VERTEX_FOG"), Material->UseTranslucencyVertexFog() ? TEXT("1") : TEXT("0"));
 
diff --git a/Engine/Source/Runtime/Renderer/Private/DeferredShadingRenderer.h b/Engine/Source/Runtime/Renderer/Private/DeferredShadingRenderer.h
--- a/Engine/Source/Runtime/Renderer/Private/DeferredShadingRenderer.h
+++ b/Engine/Source/Runtime/Renderer/Private/DeferredShadingRenderer.h
@@ -194,9 +194,6 @@ private:
 	/** Render stationary light overlap as complexity to scene color. */
 	void RenderStationaryLightOverlap();
 	
-	/** Renders the scene's distortion */
-	void RenderDistortion();
-
 	/** 
 	 * Renders the scene's translucency.
 	 */
diff --git a/Engine/Source/Runtime/Renderer/Private/DistortionRendering.cpp b/Engine/Source/Runtime/Renderer/Private/DistortionRendering.cpp
--- a/Engine/Source/Runtime/Renderer/Private/DistortionRendering.cpp
+++ b/Engine/Source/Runtime/Renderer/Private/DistortionRendering.cpp
@@ -20,7 +20,7 @@ class FDistortionApplyScreenPS : public FGlobalShader
 	DECLARE_SHADER_TYPE(FDistortionApplyScreenPS,Global);
 public:
 
-	static bool ShouldCache(EShaderPlatform Platform) { return IsFeatureLevelSupported(Platform, ERHIFeatureLevel::SM3); }
+	static bool ShouldCache(EShaderPlatform Platform) { return true; }
 
 	FDistortionApplyScreenPS(const ShaderMetaType::CompiledShaderInitializerType& Initializer)
 		: FGlobalShader(Initializer)
@@ -118,7 +118,7 @@ protected:
 
 	static bool ShouldCache(EShaderPlatform Platform,const FMaterial* Material,const FVertexFactoryType* VertexFactoryType)
 	{
-		return IsFeatureLevelSupported(Platform, ERHIFeatureLevel::SM3) && DistortMeshPolicy::ShouldCache(Platform,Material,VertexFactoryType);
+		return DistortMeshPolicy::ShouldCache(Platform,Material,VertexFactoryType);
 	}
 
 public:
@@ -154,7 +154,6 @@ protected:
 	static bool ShouldCache(EShaderPlatform Platform,const FMaterial* Material,const FVertexFactoryType* VertexFactoryType)
 	{
 		return FBaseHS::ShouldCache(Platform, Material, VertexFactoryType)
-			&& IsFeatureLevelSupported(Platform, ERHIFeatureLevel::SM3)
 			&& DistortMeshPolicy::ShouldCache(Platform, Material, VertexFactoryType);
 	}
 };
@@ -178,7 +177,6 @@ protected:
 	static bool ShouldCache(EShaderPlatform Platform,const FMaterial* Material,const FVertexFactoryType* VertexFactoryType)
 	{
 		return FBaseDS::ShouldCache(Platform, Material, VertexFactoryType)
-			&& IsFeatureLevelSupported(Platform, ERHIFeatureLevel::SM3)
 			&& DistortMeshPolicy::ShouldCache(Platform, Material, VertexFactoryType);
 	}
 };
@@ -200,7 +198,7 @@ class TDistortionMeshPS : public FMeshMaterialShader
 public:
 	static bool ShouldCache(EShaderPlatform Platform,const FMaterial* Material,const FVertexFactoryType* VertexFactoryType)
 	{
-		return IsFeatureLevelSupported(Platform, ERHIFeatureLevel::SM3) && DistortMeshPolicy::ShouldCache(Platform,Material,VertexFactoryType);
+		return DistortMeshPolicy::ShouldCache(Platform,Material,VertexFactoryType);
 	}
 
 	TDistortionMeshPS(const ShaderMetaType::CompiledShaderInitializerType& Initializer)
@@ -750,14 +748,10 @@ void FDistortionPrimSet::AddScenePrimitive(FPrimitiveSceneProxy* PrimitiveSceneP
 	Prims.Add(PrimitiveSceneProxy);
 }
 
-/*-----------------------------------------------------------------------------
-	FDeferredShadingSceneRenderer
------------------------------------------------------------------------------*/
-
 /** 
  * Renders the scene's distortion 
  */
-void FDeferredShadingSceneRenderer::RenderDistortion()
+void FSceneRenderer::RenderDistortion()
 {
 	SCOPED_DRAW_EVENT(Distortion, DEC_SCENE_ITEMS);
 
@@ -916,3 +910,4 @@ void FDeferredShadingSceneRenderer::RenderDistortion()
 		GSceneRenderTargets.FinishRenderingSceneColor(false);
 	}
 }
+
diff --git a/Engine/Source/Runtime/Renderer/Private/ForwardShadingRenderer.cpp b/Engine/Source/Runtime/Renderer/Private/ForwardShadingRenderer.cpp
--- a/Engine/Source/Runtime/Renderer/Private/ForwardShadingRenderer.cpp
+++ b/Engine/Source/Runtime/Renderer/Private/ForwardShadingRenderer.cpp
@@ -96,6 +96,11 @@ void FForwardShadingSceneRenderer::Render()
 	{
 		SCOPE_CYCLE_COUNTER(STAT_TranslucencyDrawTime);
 
+		if (ViewFamily.EngineShowFlags.Refraction)
+		{
+			// to apply refraction effect by distorting the scene color
+			RenderDistortion();
+		}
 		RenderTranslucency();
 	}
 
diff --git a/Engine/Source/Runtime/Renderer/Private/SceneRendering.h b/Engine/Source/Runtime/Renderer/Private/SceneRendering.h
--- a/Engine/Source/Runtime/Renderer/Private/SceneRendering.h
+++ b/Engine/Source/Runtime/Renderer/Private/SceneRendering.h
@@ -559,6 +559,9 @@ protected:
 	void RenderCustomDepthPass();
 
 	void OnStartFrame();
+
+	/** Renders the scene's distortion */
+	void RenderDistortion();
 };
 
 
