Distortion rendering : added r.RefractionQuality to control for low-perf device, especially Mobile Devices (Disabled by default for iOS, Android, HTML5)

diff --git a/Engine/Config/BaseDeviceProfiles.ini b/Engine/Config/BaseDeviceProfiles.ini
--- a/Engine/Config/BaseDeviceProfiles.ini
+++ b/Engine/Config/BaseDeviceProfiles.ini
@@ -42,6 +42,7 @@ BaseProfileName=
 +CVars=r.BloomQuality=0
 +CVars=r.DepthOfFieldQuality=0
 +CVars=r.LightShaftQuality=0
++CVars=r.RefractionQuality=0
 
 [iPad2 DeviceProfile]
 DeviceType=IOS
@@ -112,6 +113,7 @@ BaseProfileName=
 +CVars=r.BloomQuality=0
 +CVars=r.DepthOfFieldQuality=0
 +CVars=r.LightShaftQuality=0
++CVars=r.RefractionQuality=0
 
 [Android_Low DeviceProfile]
 DeviceType=Android
@@ -172,6 +174,7 @@ BaseProfileName=
 [HTML5 DeviceProfile]
 DeviceType=HTML5
 BaseProfileName=
++CVars=r.RefractionQuality=0
 
 [Mac DeviceProfile]
 DeviceType=Mac
diff --git a/Engine/Source/Runtime/Renderer/Private/DistortionRendering.cpp b/Engine/Source/Runtime/Renderer/Private/DistortionRendering.cpp
--- a/Engine/Source/Runtime/Renderer/Private/DistortionRendering.cpp
+++ b/Engine/Source/Runtime/Renderer/Private/DistortionRendering.cpp
@@ -748,11 +748,24 @@ void FDistortionPrimSet::AddScenePrimitive(FPrimitiveSceneProxy* PrimitiveSceneP
 	Prims.Add(PrimitiveSceneProxy);
 }
 
+static TAutoConsoleVariable<int32> CVarRefractionQuality(
+	TEXT("r.RefractionQuality"),
+	1,
+	TEXT("Defines the refraction quality.\n")
+	TEXT("  0: off\n")
+	TEXT("  1: on (default)\n"),
+	ECVF_Scalability | ECVF_RenderThreadSafe);
+
 /** 
  * Renders the scene's distortion 
  */
 void FSceneRenderer::RenderDistortion()
 {
+	if (CVarRefractionQuality.GetValueOnRenderThread() == 0)
+	{
+		return; // Don't render distortion
+	}
+
 	SCOPED_DRAW_EVENT(Distortion, DEC_SCENE_ITEMS);
 
 	// do we need to render the distortion pass?
