* fixed SSAO (wasn't applied any more)

diff --git a/Engine/Source/Runtime/Renderer/Private/CompositionLighting/CompositionLighting.cpp b/Engine/Source/Runtime/Renderer/Private/CompositionLighting/CompositionLighting.cpp
--- a/Engine/Source/Runtime/Renderer/Private/CompositionLighting/CompositionLighting.cpp
+++ b/Engine/Source/Runtime/Renderer/Private/CompositionLighting/CompositionLighting.cpp
@@ -84,19 +84,20 @@ static uint32 ComputeAmbientOcclusionPassCount(FPostprocessContext& Context)
 {
 	uint32 Ret = 0;
 
-	if(IsLpvIndirectPassRequired(Context))		// this seems wrong - we need to revisit that
+	bool bEnabled = true;
+
+	if(!IsLpvIndirectPassRequired(Context))
 	{
-		bool bEnabled = 
-			Context.View.FinalPostProcessSettings.AmbientOcclusionIntensity > 0 
+		bEnabled = Context.View.FinalPostProcessSettings.AmbientOcclusionIntensity > 0 
 			&& Context.View.FinalPostProcessSettings.AmbientOcclusionRadius >= 0.1f 
 			&& (IsBasePassAmbientOcclusionRequired(Context) || IsAmbientCubemapPassRequired(Context) || IsReflectionEnvironmentActive(Context) || Context.View.Family->EngineShowFlags.VisualizeBuffer )
 			&& !IsSimpleDynamicLightingEnabled();
+	}
 
-		if(bEnabled)
-		{
-			static const auto ICVar = IConsoleManager::Get().FindTConsoleVariableDataInt(TEXT("r.AmbientOcclusionLevels"));
-			Ret = FMath::Clamp(ICVar->GetValueOnRenderThread(), 0, 4);
-		}
+	if(bEnabled)
+	{
+		static const auto ICVar = IConsoleManager::Get().FindTConsoleVariableDataInt(TEXT("r.AmbientOcclusionLevels"));
+		Ret = FMath::Clamp(ICVar->GetValueOnRenderThread(), 0, 4);
 	}
 
 	return Ret;
