* set default of LensFlareIntensity to be 0 (off) for better performance and a clean look out of the box , can be set to 1.0 if you want the old look
* reduced default quality of screen space reflections (for better performance), can be changed in postprocess settings to be like before.

diff --git a/Engine/Source/Runtime/Engine/Classes/Engine/Scene.h b/Engine/Source/Runtime/Engine/Classes/Engine/Scene.h
--- a/Engine/Source/Runtime/Engine/Classes/Engine/Scene.h
+++ b/Engine/Source/Runtime/Engine/Classes/Engine/Scene.h
@@ -828,7 +828,8 @@ struct FPostProcessSettings
 		AutoExposureSpeedDown = 1.0f;
 		HistogramLogMin = -8.0f;
 		HistogramLogMax = 4.0f;
-		LensFlareIntensity = 1.0f;
+		// Note: We changed this from 1.0f to 0.0f for a faster default.
+		LensFlareIntensity = 0.0f;
 		LensFlareTint = FLinearColor(1.0f, 1.0f, 1.0f);
 		LensFlareBokehSize = 3.0f;
 		LensFlareThreshold = 8.0f;
diff --git a/Engine/Source/Runtime/Renderer/Private/PostProcess/ScreenSpaceReflections.cpp b/Engine/Source/Runtime/Renderer/Private/PostProcess/ScreenSpaceReflections.cpp
--- a/Engine/Source/Runtime/Renderer/Private/PostProcess/ScreenSpaceReflections.cpp
+++ b/Engine/Source/Runtime/Renderer/Private/PostProcess/ScreenSpaceReflections.cpp
@@ -162,13 +162,13 @@ static int32 ComputeSSRQuality(float Quality)
 {
 	int32 Ret;
 
-	if(Quality >= 50.0f)
+	if(Quality >= 60.0f)
 	{
-		Ret = (Quality >= 75.0f) ? 4 : 3;
+		Ret = (Quality >= 80.0f) ? 4 : 3;
 	}
 	else
 	{
-		Ret = (Quality >= 25.0f) ? 2 : 1;
+		Ret = (Quality >= 40.0f) ? 2 : 1;
 	}
 
 	int SSRQualityCVar = FMath::Max(0, CVarSSRQuality.GetValueOnRenderThread());
