Mobile editor changes:
- FeatureLevelPreview menu option now instructs world to switch feature level. This triggers our WIP switching code which works a bit...
- Converted a few more GRHIFeatureLevel references over to point at the nearest convenient scene/sceneview
- Made GetRHIFeatureLevel a non-inlined function (breakpoints wouldn't stop on the inlined version) and #define'd GRHIFeatureLevel to point at the GMaxRHIFeatureLevel (i.e. back to old behaviour)
- Changes to FSystemTextures so that the global object gets initialized in a stack-y kind of way. You initialize at a feature level and it creates all textures needed up to that point. Subsequent initializes at a higher feature level will allocate the new ones needed but not the existing ones. Only behaviour change from this is that there was a single 1x1 texture created only for ES2 that I just let get created for every feature level to keep the "feature level stack" nature of things.

diff --git a/Engine/Source/Editor/LevelEditor/Private/LevelEditorActions.cpp b/Engine/Source/Editor/LevelEditor/Private/LevelEditorActions.cpp
--- a/Engine/Source/Editor/LevelEditor/Private/LevelEditorActions.cpp
+++ b/Engine/Source/Editor/LevelEditor/Private/LevelEditorActions.cpp
@@ -535,7 +535,14 @@ bool FLevelEditorActionCallbacks::IsMaterialQualityLevelChecked( EMaterialQualit
 
 void FLevelEditorActionCallbacks::SetFeatureLevelPreview(ERHIFeatureLevel::Type InPreviewFeatureLevel)
 {
-	// Here be dragons...
+	// try iterating over all worlds here
+	for (TObjectIterator<UWorld> It; It; ++It)
+	{
+		UWorld* World = *It;
+		World->ChangeFeatureLevel(InPreviewFeatureLevel);
+	}
+	
+	UWorld::ForceFeatureLevelUpdate(InPreviewFeatureLevel);
 }
 
 bool FLevelEditorActionCallbacks::IsFeatureLevelPreviewChecked(ERHIFeatureLevel::Type InPreviewFeatureLevel)
diff --git a/Engine/Source/Editor/LevelEditor/Private/LevelEditorToolBar.cpp b/Engine/Source/Editor/LevelEditor/Private/LevelEditorToolBar.cpp
--- a/Engine/Source/Editor/LevelEditor/Private/LevelEditorToolBar.cpp
+++ b/Engine/Source/Editor/LevelEditor/Private/LevelEditorToolBar.cpp
@@ -476,9 +476,11 @@ static void MakeMaterialQualityLevelMenu( FMenuBuilder& MenuBuilder )
 
 static void MakeShaderModelPreviewMenu(FMenuBuilder& MenuBuilder)
 {
-	MenuBuilder.BeginSection("LevelEditorShaderModelPreview", NSLOCTEXT("LevelToolBarViewMenu", "ShaderModelPreviewHeading", "Shader Model Preview"));
+	MenuBuilder.BeginSection("LevelEditorShaderModelPreview", NSLOCTEXT("LevelToolBarViewMenu", "FeatureLevelPreviewHeading", "Feature Level Preview"));
 	{
-		for (int32 i = GMaxRHIFeatureLevel; i >= 0; --i)
+		static const ERHIFeatureLevel::Type MaxFeatureLevel = GMaxRHIFeatureLevel;
+
+		for (int32 i = MaxFeatureLevel; i >= 0; --i)
 		{
 			MenuBuilder.AddMenuEntry(FLevelEditorCommands::Get().FeatureLevelPreview[i]);
 		}
@@ -549,8 +551,8 @@ TSharedRef< SWidget > FLevelEditorToolBar::GenerateQuickSettingsMenu( TSharedRef
 		if (CVarFeatureLevelPreview.GetValueOnGameThread() != 0)
 		{
 			MenuBuilder.AddSubMenu(
-				LOCTEXT("ShaderModelPreviewSubMenu", "Shader Model Preview"),
-				LOCTEXT("ShaderModelPreviewSubMenu_ToolTip", "Sets the shader model preview mode"),
+				LOCTEXT("FeatureLevelPreviewSubMenu", "Feature Level Preview"),
+				LOCTEXT("FeatureLevelPreviewSubMenu_ToolTip", "Sets the feature level preview mode"),
 				FNewMenuDelegate::CreateStatic(&MakeShaderModelPreviewMenu));
 		}
 	}
diff --git a/Engine/Source/Editor/Persona/Persona.Build.cs b/Engine/Source/Editor/Persona/Persona.Build.cs
--- a/Engine/Source/Editor/Persona/Persona.Build.cs
+++ b/Engine/Source/Editor/Persona/Persona.Build.cs
@@ -33,6 +33,7 @@ public class Persona : ModuleRules
                 "PropertyEditor",
 				"EditorWidgets",
                 "BlueprintGraph",
+                "RHI",
 			}
 		);
 
diff --git a/Engine/Source/Editor/Persona/Private/PersonaMeshDetails.cpp b/Engine/Source/Editor/Persona/Private/PersonaMeshDetails.cpp
--- a/Engine/Source/Editor/Persona/Private/PersonaMeshDetails.cpp
+++ b/Engine/Source/Editor/Persona/Private/PersonaMeshDetails.cpp
@@ -830,7 +830,7 @@ FReply FPersonaMeshDetails::OnOpenClothingFileClicked(IDetailLayoutBuilder* Deta
 					FClothingAssetData& AssetData = SkelMesh->ClothingAssets[AssetIndex];
 					int32 NumLODs = ApexClothingUtils::GetNumLODs(AssetData.ApexClothingAsset->GetAsset());
 
-					uint32 MaxClothVertices = ApexClothingUtils::GetMaxClothSimulVertices();
+					uint32 MaxClothVertices = ApexClothingUtils::GetMaxClothSimulVertices(GRHIFeatureLevel);
 
 					// check whether there are sub-meshes which have over MaxClothVertices or not
 					// this checking will be removed after changing implementation way for supporting over MaxClothVertices
diff --git a/Engine/Source/Editor/UnrealEd/Private/ApexClothingUtils.cpp b/Engine/Source/Editor/UnrealEd/Private/ApexClothingUtils.cpp
--- a/Engine/Source/Editor/UnrealEd/Private/ApexClothingUtils.cpp
+++ b/Engine/Source/Editor/UnrealEd/Private/ApexClothingUtils.cpp
@@ -144,9 +144,9 @@ int32 FindSectionByMaterialIndex( USkeletalMesh* SkelMesh, uint32 LODIndex, uint
 	return -1;
 }
 
-uint32 GetMaxClothSimulVertices()
+uint32 GetMaxClothSimulVertices(ERHIFeatureLevel::Type InFeatureLevel)
 {
-	if (GRHIFeatureLevel >= ERHIFeatureLevel::SM4)
+	if (InFeatureLevel >= ERHIFeatureLevel::SM4)
 	{					
 		return MAX_APEXCLOTH_VERTICES_FOR_VB;
 	}
diff --git a/Engine/Source/Editor/UnrealEd/Private/EditorComponents.cpp b/Engine/Source/Editor/UnrealEd/Private/EditorComponents.cpp
--- a/Engine/Source/Editor/UnrealEd/Private/EditorComponents.cpp
+++ b/Engine/Source/Editor/UnrealEd/Private/EditorComponents.cpp
@@ -50,11 +50,11 @@ static TAutoConsoleVariable<float> CVarEditor2DSnapScale(
 	TEXT("Tweak to define the grid rendering in 2D viewports."),
 	ECVF_RenderThreadSafe);
 
-static bool IsEditorCompositingMSAAEnabled()
+static bool IsEditorCompositingMSAAEnabled(ERHIFeatureLevel::Type InFeatureLevel)
 {
 	bool Ret = false;
 
-	if (GRHIFeatureLevel >= ERHIFeatureLevel::SM5)
+	if (InFeatureLevel >= ERHIFeatureLevel::SM5)
 	{
 		// only supported on SM5 yet
 		static const auto CVar = IConsoleManager::Get().FindTConsoleVariableDataInt(TEXT("r.MSAA.CompositingSampleCount"));
@@ -127,7 +127,7 @@ void FGridWidget::DrawNewGrid(const FSceneView* View, FPrimitiveDrawInterface* P
 		return;
 	}
 
-	bool bMSAA = IsEditorCompositingMSAAEnabled();
+	bool bMSAA = IsEditorCompositingMSAAEnabled(View->GetFeatureLevel());
 	bool bIsPerspective = ( View->ViewMatrices.ProjMatrix.M[3][3] < 1.0f );
 
 	// in unreal units
diff --git a/Engine/Source/Editor/UnrealEd/Public/ApexClothingUtils.h b/Engine/Source/Editor/UnrealEd/Public/ApexClothingUtils.h
--- a/Engine/Source/Editor/UnrealEd/Public/ApexClothingUtils.h
+++ b/Engine/Source/Editor/UnrealEd/Public/ApexClothingUtils.h
@@ -77,7 +77,7 @@ namespace ApexClothingUtils
 	// Find mesh section index in a StaticLODModel by material index
 	UNREALED_API int32 FindSectionByMaterialIndex( USkeletalMesh* SkelMesh, uint32 LODIndex, uint32 MaterialIndex );
 
-	UNREALED_API uint32 GetMaxClothSimulVertices();
+	UNREALED_API uint32 GetMaxClothSimulVertices(ERHIFeatureLevel::Type InFeatureLevel);
 
 #if WITH_APEX_CLOTHING
 
diff --git a/Engine/Source/Runtime/Engine/Classes/Engine/World.h b/Engine/Source/Runtime/Engine/Classes/Engine/World.h
--- a/Engine/Source/Runtime/Engine/Classes/Engine/World.h
+++ b/Engine/Source/Runtime/Engine/Classes/Engine/World.h
@@ -623,6 +623,11 @@ public:
 	/** Return the array of objects currently bieng debugged. */
 	const FBlueprintToDebuggedObjectMap& GetBlueprintObjectsBeingDebugged() const{ return BlueprintObjectsBeingDebugged; };
 
+	/** Change the feature level that this world is current rendering with */
+	void ChangeFeatureLevel(ERHIFeatureLevel::Type InFeatureLevel);
+
+	static void ForceFeatureLevelUpdate(ERHIFeatureLevel::Type InFeatureLevel);
+
 private:
 	/** List of all the controllers in the world. */
 	TArray<TAutoWeakObjectPtr<class AController> >	ControllerList;
diff --git a/Engine/Source/Runtime/Engine/Private/GPUSkinVertexFactory.cpp b/Engine/Source/Runtime/Engine/Private/GPUSkinVertexFactory.cpp
--- a/Engine/Source/Runtime/Engine/Private/GPUSkinVertexFactory.cpp
+++ b/Engine/Source/Runtime/Engine/Private/GPUSkinVertexFactory.cpp
@@ -429,6 +429,7 @@ public:
 	{
 		if(Shader->GetVertexShader())
 		{
+			const auto FeatureLevel = View.GetFeatureLevel();
 			const FGPUBaseSkinVertexFactory::ShaderDataType& ShaderData = ((const FGPUBaseSkinVertexFactory*)VertexFactory)->GetShaderData();
 
 			SetShaderValue(
@@ -442,7 +443,7 @@ public:
 				ShaderData.MeshExtension
 				);
 			
-			if (GRHIFeatureLevel >= ERHIFeatureLevel::SM4)
+			if (FeatureLevel >= ERHIFeatureLevel::SM4)
 			{
 				if(BoneMatrices.IsBound())
 				{
@@ -460,7 +461,7 @@ public:
 
 			bool bLocalPerBoneMotionBlur = false;
 			
-			if (GRHIFeatureLevel >= ERHIFeatureLevel::SM4 && GPrevPerBoneMotionBlur.IsLocked())
+			if (FeatureLevel >= ERHIFeatureLevel::SM4 && GPrevPerBoneMotionBlur.IsLocked())
 			{
 				// we are in the velocity rendering pass
 
diff --git a/Engine/Source/Runtime/Engine/Private/Landscape/LandscapeRender.cpp b/Engine/Source/Runtime/Engine/Private/Landscape/LandscapeRender.cpp
--- a/Engine/Source/Runtime/Engine/Private/Landscape/LandscapeRender.cpp
+++ b/Engine/Source/Runtime/Engine/Private/Landscape/LandscapeRender.cpp
@@ -551,7 +551,8 @@ FLandscapeComponentSceneProxy::FLandscapeComponentSceneProxy(ULandscapeComponent
 ,	LODBias(InComponent->LODBias)
 ,	LODFalloff(InComponent->GetLandscapeProxy()->LODFalloff)
 {
-	if (GRHIFeatureLevel == ERHIFeatureLevel::ES2)
+	const auto FeatureLevel = GetScene()->GetFeatureLevel();
+	if (FeatureLevel == ERHIFeatureLevel::ES2)
 	{
 		HeightmapTexture = NULL;
 		HeightmapSubsectionOffsetU = 0;
@@ -1014,12 +1015,12 @@ void FLandscapeComponentSceneProxy::OnTransformChanged()
 
 namespace
 {
-	inline bool RequiresAdjacencyInformation(FMaterialRenderProxy* MaterialRenderProxy) // Assumes VertexFactory supports tessellation, and rendering thread with this function
+	inline bool RequiresAdjacencyInformation(FMaterialRenderProxy* MaterialRenderProxy, ERHIFeatureLevel::Type InFeatureLevel) // Assumes VertexFactory supports tessellation, and rendering thread with this function
 	{
 		if (RHISupportsTessellation(GRHIShaderPlatform) && MaterialRenderProxy )
 		{
 			check ( IsInRenderingThread() );
-			const FMaterial* MaterialResource = MaterialRenderProxy->GetMaterial(GRHIFeatureLevel);
+			const FMaterial* MaterialResource = MaterialRenderProxy->GetMaterial(InFeatureLevel);
 			check( MaterialResource );
 			EMaterialTessellationMode TessellationMode = MaterialResource->GetTessellationMode();
 			bool bEnableCrackFreeDisplacement = MaterialResource->IsCrackFreeDisplacementEnabled();
@@ -1053,7 +1054,7 @@ void FLandscapeComponentSceneProxy::DrawStaticElements(FStaticPrimitiveDrawInter
 	FMaterialRenderProxy* RenderProxy = MaterialInterface->GetRenderProxy(false);
 
 	// Could be different from bRequiresAdjacencyInformation during shader compilation
-	bool bCurrentRequiresAdjacencyInformation = RequiresAdjacencyInformation( RenderProxy );
+	bool bCurrentRequiresAdjacencyInformation = RequiresAdjacencyInformation( RenderProxy, GetScene()->GetFeatureLevel() );
 
 	MeshBatch.VertexFactory = VertexFactory;
 	MeshBatch.MaterialRenderProxy = RenderProxy;
@@ -1356,7 +1357,7 @@ void FLandscapeComponentSceneProxy::DrawDynamicElements(FPrimitiveDrawInterface*
 
 	// Could be different from bRequiresAdjacencyInformation during shader compilation
 	FMaterialRenderProxy* RenderProxy = MaterialInterface->GetRenderProxy(false);
-	bool bCurrentRequiresAdjacencyInformation = RequiresAdjacencyInformation( RenderProxy );
+	bool bCurrentRequiresAdjacencyInformation = RequiresAdjacencyInformation( RenderProxy, View->GetFeatureLevel() );
 	DynamicMesh.Type = bCurrentRequiresAdjacencyInformation ? PT_12_ControlPointPatchList : PT_TriangleList;
 
 	// Setup the LOD parameters in DynamicMesh.
diff --git a/Engine/Source/Runtime/Engine/Private/Landscape/LandscapeSplines.cpp b/Engine/Source/Runtime/Engine/Private/Landscape/LandscapeSplines.cpp
--- a/Engine/Source/Runtime/Engine/Private/Landscape/LandscapeSplines.cpp
+++ b/Engine/Source/Runtime/Engine/Private/Landscape/LandscapeSplines.cpp
@@ -2373,7 +2373,7 @@ FPrimitiveSceneProxy* USplineMeshComponent::CreateSceneProxy()
 		StaticMesh &&
 		StaticMesh->HasValidRenderData();
 
-	if (bMeshIsValid && (GRHIFeatureLevel >= ERHIFeatureLevel::SM3))
+	if (bMeshIsValid && (GetScene()->GetFeatureLevel() >= ERHIFeatureLevel::SM3))
 	{
 		return ::new FSplineMeshSceneProxy(this);
 	}
diff --git a/Engine/Source/Runtime/Engine/Private/MeshParticleVertexFactory.cpp b/Engine/Source/Runtime/Engine/Private/MeshParticleVertexFactory.cpp
--- a/Engine/Source/Runtime/Engine/Private/MeshParticleVertexFactory.cpp
+++ b/Engine/Source/Runtime/Engine/Private/MeshParticleVertexFactory.cpp
@@ -37,7 +37,7 @@ public:
 
 	virtual void SetMesh(FShader* Shader,const FVertexFactory* VertexFactory,const FSceneView& View,const FMeshBatchElement& BatchElement,uint32 DataFlags) const OVERRIDE
 	{
-		const bool bInstanced = GRHIFeatureLevel >= ERHIFeatureLevel::SM3;
+		const bool bInstanced = View.GetFeatureLevel() >= ERHIFeatureLevel::SM3;
 		FMeshParticleVertexFactory* MeshParticleVF = (FMeshParticleVertexFactory*)VertexFactory;
 		FVertexShaderRHIParamRef VertexShaderRHI = Shader->GetVertexShader();
 		SetUniformBufferParameter( VertexShaderRHI, Shader->GetUniformBufferParameter<FMeshParticleUniformParameters>(), MeshParticleVF->GetUniformBuffer() );
diff --git a/Engine/Source/Runtime/Engine/Private/Particles/FXSystemPrivate.h b/Engine/Source/Runtime/Engine/Private/Particles/FXSystemPrivate.h
--- a/Engine/Source/Runtime/Engine/Private/Particles/FXSystemPrivate.h
+++ b/Engine/Source/Runtime/Engine/Private/Particles/FXSystemPrivate.h
@@ -72,6 +72,10 @@ public:
 	/*--------------------------------------------------------------------------
 		Internal interface for GPU simulation.
 	--------------------------------------------------------------------------*/
+	/**
+	 * Retrieve feature level that this FXSystem was created for
+	 */
+	ERHIFeatureLevel::Type GetFeatureLevel() const { return FeatureLevel; }
 
 	/**
 	 * Add a new GPU simulation to the system.
diff --git a/Engine/Source/Runtime/Engine/Private/Particles/ParticleGpuSimulation.cpp b/Engine/Source/Runtime/Engine/Private/Particles/ParticleGpuSimulation.cpp
--- a/Engine/Source/Runtime/Engine/Private/Particles/ParticleGpuSimulation.cpp
+++ b/Engine/Source/Runtime/Engine/Private/Particles/ParticleGpuSimulation.cpp
@@ -2494,7 +2494,7 @@ struct FGPUSpriteDynamicEmitterData : FDynamicEmitterDataBase
 	 */
 	virtual void CreateRenderThreadResources(const FParticleSystemSceneProxy* InOwnerProxy)
 	{
-		if (CurrentRHISupportsGPUParticles())
+		if (RHISupportsGPUParticles(InOwnerProxy->GetScene()->GetFeatureLevel()))
 		{
 			Simulation->PendingDeltaSeconds = PendingDeltaSeconds;
 
@@ -2544,7 +2544,9 @@ struct FGPUSpriteDynamicEmitterData : FDynamicEmitterDataBase
 	 */
 	virtual void PreRenderView(FParticleSystemSceneProxy* Proxy, const FSceneViewFamily* ViewFamily, const uint32 VisibilityMap, int32 FrameNumber) OVERRIDE
 	{
-		if (CurrentRHISupportsGPUParticles())
+		auto FeatureLevel = ViewFamily->Scene->GetFeatureLevel();
+
+		if (RHISupportsGPUParticles(FeatureLevel))
 		{
 			SCOPE_CYCLE_COUNTER(STAT_GPUSpritePreRenderTime);
 
@@ -2556,11 +2558,11 @@ struct FGPUSpriteDynamicEmitterData : FDynamicEmitterDataBase
 			if (Simulation->SimulationIndex != INDEX_NONE
 				&& Simulation->VertexBuffer.ParticleCount > 0)
 			{
-				EBlendMode BlendMode = Material->GetRenderProxy(false)->GetMaterial(GRHIFeatureLevel)->GetBlendMode();
+				EBlendMode BlendMode = Material->GetRenderProxy(false)->GetMaterial(FeatureLevel)->GetBlendMode();
 				const bool bOpaque = (BlendMode == BLEND_Opaque || BlendMode == BLEND_Masked);
 
 				const bool bAllowSorting = FXConsoleVariables::bAllowGPUSorting
-					&& GRHIFeatureLevel == ERHIFeatureLevel::SM5
+					&& FeatureLevel == ERHIFeatureLevel::SM5
 					&& !bOpaque;
 
 				// If the simulation wants to collide against the depth buffer
@@ -2615,7 +2617,7 @@ struct FGPUSpriteDynamicEmitterData : FDynamicEmitterDataBase
 	 */
 	virtual int32 Render(FParticleSystemSceneProxy* Proxy, FPrimitiveDrawInterface* PDI,const FSceneView* View)
 	{
-		if (CurrentRHISupportsGPUParticles())
+		if (RHISupportsGPUParticles(View->GetFeatureLevel()))
 		{
 			check(Simulation);
 
@@ -2935,7 +2937,7 @@ public:
 		const bool bSimulateGPUParticles = 
 			FXConsoleVariables::bFreezeGPUSimulation == false &&
 			FXConsoleVariables::bFreezeParticleSimulation == false &&
-			CurrentRHISupportsGPUParticles();
+			RHISupportsGPUParticles(FXSystem->GetFeatureLevel());
 
 		if (bSimulateGPUParticles)
 		{
@@ -3045,7 +3047,7 @@ public:
 
 		if (FXConsoleVariables::bFreezeGPUSimulation ||
 			FXConsoleVariables::bFreezeParticleSimulation ||
-			!CurrentRHISupportsGPUParticles())
+			!RHISupportsGPUParticles(FXSystem->GetFeatureLevel()))
 		{
 			return;
 		}
@@ -3826,7 +3828,7 @@ void FFXSystem::DestroyGPUSimulation()
 
 void FFXSystem::InitGPUResources()
 {
-	if (CurrentRHISupportsGPUParticles())
+	if (RHISupportsGPUParticles(FeatureLevel))
 	{
 		check(ParticleSimulationResources);
 		ParticleSimulationResources->Init();
@@ -3835,7 +3837,7 @@ void FFXSystem::InitGPUResources()
 
 void FFXSystem::ReleaseGPUResources()
 {
-	if (CurrentRHISupportsGPUParticles())
+	if (RHISupportsGPUParticles(FeatureLevel))
 	{
 		check(ParticleSimulationResources);
 		ParticleSimulationResources->Release();
diff --git a/Engine/Source/Runtime/Engine/Private/Particles/ParticleModules.cpp b/Engine/Source/Runtime/Engine/Private/Particles/ParticleModules.cpp
--- a/Engine/Source/Runtime/Engine/Private/Particles/ParticleModules.cpp
+++ b/Engine/Source/Runtime/Engine/Private/Particles/ParticleModules.cpp
@@ -4574,7 +4574,7 @@ FParticleEmitterInstance* UParticleModuleTypeDataGpu::CreateInstance(UParticleEm
 		InComponent->Template != NULL ? *InComponent->Template->GetName() : TEXT("NULL"));
 
 	FParticleEmitterInstance* Instance = NULL;
-	if (CurrentRHISupportsGPUParticles())
+	if (RHISupportsGPUParticles(World->Scene->GetFeatureLevel()))
 	{
 		check( InComponent && InComponent->FXSystem );
 		Instance = InComponent->FXSystem->CreateGPUSpriteEmitterInstance( EmitterInfo );
diff --git a/Engine/Source/Runtime/Engine/Private/Particles/ParticleSystemRender.cpp b/Engine/Source/Runtime/Engine/Private/Particles/ParticleSystemRender.cpp
--- a/Engine/Source/Runtime/Engine/Private/Particles/ParticleSystemRender.cpp
+++ b/Engine/Source/Runtime/Engine/Private/Particles/ParticleSystemRender.cpp
@@ -1130,7 +1130,7 @@ void FDynamicSpriteEmitterData::PreRenderView(FParticleSystemSceneProxy* Proxy,
 	{
 		// Determine how many vertices and indices are needed to render.
 		const int32 ParticleCount = SourceData->ActiveParticleCount;
-		const int32 VertexSize = GetDynamicVertexStride();
+		const int32 VertexSize = GetDynamicVertexStride(ViewFamily->Scene->GetFeatureLevel());
 		const int32 DynamicParameterVertexSize = sizeof(FParticleVertexDynamicParameter);
 		const int32 VertexPerParticle = bInstanced ? 1 : 4;
 		const int32 ViewCount = ViewFamily->Views.Num();
@@ -1369,7 +1369,7 @@ int32 FDynamicSpriteEmitterData::Render(FParticleSystemSceneProxy* Proxy, FPrimi
 			// Set the sprite uniform buffer for this view.
 			check( SpriteVertexFactory );
 			SpriteVertexFactory->SetSpriteUniformBuffer( PerViewUniformBuffers[ViewIndex] );
-			SpriteVertexFactory->SetInstanceBuffer( Allocation->VertexBuffer, Allocation->VertexOffset, GetDynamicVertexStride(), bInstanced );
+			SpriteVertexFactory->SetInstanceBuffer( Allocation->VertexBuffer, Allocation->VertexOffset, GetDynamicVertexStride(View->GetFeatureLevel()), bInstanced );
 			SpriteVertexFactory->SetDynamicParameterBuffer( DynamicParameterAllocation ? DynamicParameterAllocation->VertexBuffer : NULL, DynamicParameterAllocation->VertexOffset, GetDynamicParameterVertexStride(), bInstanced );
 
 			// Construct the mesh element to render.
@@ -1740,7 +1740,7 @@ int32 FDynamicMeshEmitterData::Render(FParticleSystemSceneProxy* Proxy, FPrimiti
 			check(!bUsesDynamicParameter || DynamicParameterDataAllocations[InstanceBufferIndex].IsValid());
 
 			// Set the particle instance data. This buffer is generated in PreRenderView.
-			MeshVertexFactory->SetInstanceBuffer(InstanceDataAllocations[InstanceBufferIndex].VertexBuffer, InstanceDataAllocations[InstanceBufferIndex].VertexOffset , GetDynamicVertexStride());
+			MeshVertexFactory->SetInstanceBuffer(InstanceDataAllocations[InstanceBufferIndex].VertexBuffer, InstanceDataAllocations[InstanceBufferIndex].VertexOffset, GetDynamicVertexStride(View->GetFeatureLevel()));
 			MeshVertexFactory->SetDynamicParameterBuffer(DynamicParameterDataAllocations[InstanceBufferIndex].VertexBuffer, DynamicParameterDataAllocations[InstanceBufferIndex].VertexOffset , GetDynamicParameterVertexStride());
 		}
 
@@ -1799,7 +1799,7 @@ void FDynamicMeshEmitterData::PreRenderView(FParticleSystemSceneProxy* Proxy, co
 		ParticleCount = Source.MaxDrawCount;
 	}
 
-	const int32 InstanceVertexStride  = GetDynamicVertexStride();
+	const int32 InstanceVertexStride  = GetDynamicVertexStride(ViewFamily->Scene->GetFeatureLevel());
 	const int32 DynamicParameterVertexStride  = GetDynamicParameterVertexStride();
 	const int32 ViewCount = ViewFamily->Views.Num();
 	if(bInstanced)
@@ -2685,7 +2685,7 @@ int32 FDynamicBeam2EmitterData::Render(FParticleSystemSceneProxy* Proxy, FPrimit
 
 			if( Allocation && IndexAllocation && Allocation->IsValid() && IndexAllocation->IsValid() )
 			{
-				BeamTrailVertexFactory->SetVertexBuffer( Allocation->VertexBuffer, Allocation->VertexOffset, GetDynamicVertexStride() );
+				BeamTrailVertexFactory->SetVertexBuffer( Allocation->VertexBuffer, Allocation->VertexOffset, GetDynamicVertexStride(View->GetFeatureLevel()) );
 				BeamTrailVertexFactory->SetDynamicParameterBuffer( NULL, 0, 0 );
 
 				FMeshBatch Mesh;
@@ -5480,7 +5480,7 @@ int32 FDynamicTrailsEmitterData::Render(FParticleSystemSceneProxy* Proxy, FPrimi
 			// Create and set the uniform buffer for this emitter.
 			FParticleBeamTrailVertexFactory* BeamTrailVertexFactory = (FParticleBeamTrailVertexFactory*)VertexFactory;
 			BeamTrailVertexFactory->SetBeamTrailUniformBuffer( CreateBeamTrailUniformBuffer( Proxy, SourcePointer, View ) );
-			BeamTrailVertexFactory->SetVertexBuffer( Allocation->VertexBuffer, Allocation->VertexOffset, GetDynamicVertexStride() );
+			BeamTrailVertexFactory->SetVertexBuffer( Allocation->VertexBuffer, Allocation->VertexOffset, GetDynamicVertexStride(View->GetFeatureLevel()) );
 			BeamTrailVertexFactory->SetDynamicParameterBuffer( DynamicParameterAllocation ? DynamicParameterAllocation->VertexBuffer : NULL, DynamicParameterAllocation ? DynamicParameterAllocation->VertexOffset : 0, GetDynamicParameterVertexStride() );
 
 			FMeshBatch Mesh;
@@ -5575,7 +5575,7 @@ void FDynamicTrailsEmitterData::PreRenderView(FParticleSystemSceneProxy* Proxy,
 	// Only need to do this once per-view
 	if (LastFramePreRendered < FrameNumber)
 	{
-		int32 VertexStride = GetDynamicVertexStride();
+		int32 VertexStride = GetDynamicVertexStride(ViewFamily->Scene->GetFeatureLevel());
 		int32 DynamicParameterVertexStride = 0;
 		if (bUsesDynamicParameter == true)
 		{
diff --git a/Engine/Source/Runtime/Engine/Private/PointLightComponent.cpp b/Engine/Source/Runtime/Engine/Private/PointLightComponent.cpp
--- a/Engine/Source/Runtime/Engine/Private/PointLightComponent.cpp
+++ b/Engine/Source/Runtime/Engine/Private/PointLightComponent.cpp
@@ -69,7 +69,7 @@ public:
 	 */
 	virtual bool GetWholeSceneProjectedShadowInitializer(const FSceneViewFamily& ViewFamily, TArray<FWholeSceneProjectedShadowInitializer, TInlineAllocator<6> >& OutInitializers) const
 	{
-		if (GRHIFeatureLevel >= ERHIFeatureLevel::SM4
+		if (ViewFamily.Scene->GetFeatureLevel() >= ERHIFeatureLevel::SM4
 			&& GAllowPointLightCubemapShadows != 0)
 		{
 			FWholeSceneProjectedShadowInitializer& OutInitializer = *new(OutInitializers) FWholeSceneProjectedShadowInitializer;
diff --git a/Engine/Source/Runtime/Engine/Private/PrimitiveDrawingUtils.cpp b/Engine/Source/Runtime/Engine/Private/PrimitiveDrawingUtils.cpp
--- a/Engine/Source/Runtime/Engine/Private/PrimitiveDrawingUtils.cpp
+++ b/Engine/Source/Runtime/Engine/Private/PrimitiveDrawingUtils.cpp
@@ -15,7 +15,7 @@ void EmitMeshDrawEvents(const FPrimitiveSceneProxy* PrimitiveSceneProxy, const F
 	{
 		// Only show material name at the top level
 		// Note: this is the parent's material name, not the material instance
-		SCOPED_DRAW_EVENTF(MaterialEvent, DEC_SCENE_ITEMS, *Mesh.MaterialRenderProxy->GetMaterial(GRHIFeatureLevel)->GetFriendlyName());
+		SCOPED_DRAW_EVENTF(MaterialEvent, DEC_SCENE_ITEMS, *Mesh.MaterialRenderProxy->GetMaterial(PrimitiveSceneProxy->GetScene()->GetFeatureLevel())->GetFriendlyName());
 		if (PrimitiveSceneProxy)
 		{
 			// Show Actor, level and resource name inside the material name
@@ -1218,6 +1218,7 @@ int32 DrawRichMesh(
 		return PDI->DrawMesh( Mesh );
 	}
 
+	const auto FeatureLevel = PDI->View->GetFeatureLevel();
 	const FEngineShowFlags& EngineShowFlags = PDI->View->Family->EngineShowFlags;
 
 	if(bDrawInWireframe)
@@ -1234,7 +1235,7 @@ int32 DrawRichMesh(
 			BaseColor = LevelColor;
 		}
 
-		if (Mesh.MaterialRenderProxy->GetMaterial(GRHIFeatureLevel)->MaterialModifiesMeshPosition())
+		if (Mesh.MaterialRenderProxy->GetMaterial(FeatureLevel)->MaterialModifiesMeshPosition())
 		{
 			// If the material is mesh-modifying, we cannot rely on substitution
 			const FOverrideSelectionColorMaterialRenderProxy WireframeMaterialInstance(
@@ -1262,7 +1263,7 @@ int32 DrawRichMesh(
 	else if(PDI->View->Family->EngineShowFlags.LightComplexity)
 	{
 		// Don't render unlit translucency when in 'light complexity' viewmode.
-		if (!Mesh.IsTranslucent() || Mesh.MaterialRenderProxy->GetMaterial(GRHIFeatureLevel)->GetShadingModel() != MSM_Unlit)
+		if (!Mesh.IsTranslucent(FeatureLevel) || Mesh.MaterialRenderProxy->GetMaterial(FeatureLevel)->GetShadingModel() != MSM_Unlit)
 		{
 			// Count the number of lights interacting with this primitive.
 			int32 NumDynamicLights = GetRendererModule().GetNumDynamicLightsAffectingPrimitive(PrimitiveSceneProxy->GetPrimitiveSceneInfo(),Mesh.LCI);
@@ -1286,11 +1287,11 @@ int32 DrawRichMesh(
 	else if(!EngineShowFlags.Materials && !PDI->View->bForceShowMaterials)
 	{
 		// Don't render unlit translucency when in 'lighting only' viewmode.
-		if (Mesh.MaterialRenderProxy->GetMaterial(GRHIFeatureLevel)->GetShadingModel() != MSM_Unlit
+		if (Mesh.MaterialRenderProxy->GetMaterial(FeatureLevel)->GetShadingModel() != MSM_Unlit
 			// Don't render translucency in 'lighting only', since the viewmode works by overriding with an opaque material
 			// This would cause a mismatch of the material's blend mode with the primitive's view relevance,
 			// And make faint particles block the view
-			&& !IsTranslucentBlendMode(Mesh.MaterialRenderProxy->GetMaterial(GRHIFeatureLevel)->GetBlendMode()))
+			&& !IsTranslucentBlendMode(Mesh.MaterialRenderProxy->GetMaterial(FeatureLevel)->GetBlendMode()))
 		{
 			// When materials aren't shown, apply the same basic material to all meshes.
 			FMeshBatch ModifiedMesh = Mesh;
@@ -1403,7 +1404,7 @@ int32 DrawRichMesh(
 		if(EngineShowFlags.MeshEdges)
 		{
 			// Draw the mesh's edges in blue, on top of the base geometry.
-			if (Mesh.MaterialRenderProxy->GetMaterial(GRHIFeatureLevel)->MaterialModifiesMeshPosition())
+			if (Mesh.MaterialRenderProxy->GetMaterial(FeatureLevel)->MaterialModifiesMeshPosition())
 			{
 				const FSceneView* View = PDI->View;
 
diff --git a/Engine/Source/Runtime/Engine/Private/SceneView.cpp b/Engine/Source/Runtime/Engine/Private/SceneView.cpp
--- a/Engine/Source/Runtime/Engine/Private/SceneView.cpp
+++ b/Engine/Source/Runtime/Engine/Private/SceneView.cpp
@@ -951,11 +951,12 @@ void FSceneView::EndFinalPostprocessSettings()
 		static auto* MobileMSAACvar = IConsoleManager::Get().FindTConsoleVariableDataInt(TEXT("r.MobileMSAA"));
 
 		int32 Quality = FMath::Clamp(CVar->GetValueOnGameThread(), 0, 6);
+		const auto FeatureLevel = GetFeatureLevel();
 
 		if( !Family->EngineShowFlags.PostProcessing || !Family->EngineShowFlags.AntiAliasing || Quality <= 0
 			// Disable antialiasing in GammaLDR mode to avoid jittering.
-			|| (GRHIFeatureLevel == ERHIFeatureLevel::ES2 && MobileHDRCvar->GetValueOnGameThread() == 0)
-			|| (GRHIFeatureLevel == ERHIFeatureLevel::ES2 && MobileMSAACvar->GetValueOnGameThread()) )
+			|| (FeatureLevel == ERHIFeatureLevel::ES2 && MobileHDRCvar->GetValueOnGameThread() == 0)
+			|| (FeatureLevel == ERHIFeatureLevel::ES2 && MobileMSAACvar->GetValueOnGameThread()))
 		{
 			FinalPostProcessSettings.AntiAliasingMethod = AAM_None;
 		}
diff --git a/Engine/Source/Runtime/Engine/Private/World.cpp b/Engine/Source/Runtime/Engine/Private/World.cpp
--- a/Engine/Source/Runtime/Engine/Private/World.cpp
+++ b/Engine/Source/Runtime/Engine/Private/World.cpp
@@ -5004,6 +5004,27 @@ void UWorld::GetLightMapsAndShadowMaps(TArray<UTexture2D*>& OutLightMapsAndShado
 	}
 }
 
+void UWorld::ChangeFeatureLevel(ERHIFeatureLevel::Type InFeatureLevel)
+{
+	if (InFeatureLevel != Scene->GetFeatureLevel())
+	{
+		Scene->ChangeFeatureLevel(InFeatureLevel);
+		
+		FFXSystemInterface::Destroy(FXSystem);
+		FXSystem = FFXSystemInterface::Create(InFeatureLevel);
+		Scene->SetFXSystem(FXSystem);
+	}
+}
+
+void UWorld::ForceFeatureLevelUpdate(ERHIFeatureLevel::Type InFeatureLevel)
+{
+	FGlobalComponentReregisterContext RecreateComponents;
+	UMaterial::AllMaterialsCacheResourceShadersForRendering();
+	UMaterialInstance::AllMaterialsCacheResourceShadersForRendering();
+	GetGlobalShaderMap(GShaderPlatformForFeatureLevel[InFeatureLevel], true);
+	RecompileGlobalShaders();
+}
+
 /**
 * Dump visible actors in current world.
 */
diff --git a/Engine/Source/Runtime/Engine/Public/ParticleHelper.h b/Engine/Source/Runtime/Engine/Public/ParticleHelper.h
--- a/Engine/Source/Runtime/Engine/Public/ParticleHelper.h
+++ b/Engine/Source/Runtime/Engine/Public/ParticleHelper.h
@@ -1498,7 +1498,7 @@ struct FDynamicSpriteEmitterDataBase : public FDynamicEmitterDataBase
 	/**
 	 *	Get the vertex stride for the dynamic rendering data
 	 */
-	virtual int32 GetDynamicVertexStride() const
+	virtual int32 GetDynamicVertexStride(ERHIFeatureLevel::Type /*InFeatureLevel*/) const
 	{
 		checkf(0, TEXT("GetDynamicVertexStride MUST be overridden"));
 		return 0;
@@ -1697,9 +1697,9 @@ struct FDynamicSpriteEmitterData : public FDynamicSpriteEmitterDataBase
 	/**
 	 *	Get the vertex stride for the dynamic rendering data
 	 */
-	virtual int32 GetDynamicVertexStride() const OVERRIDE
+	virtual int32 GetDynamicVertexStride(ERHIFeatureLevel::Type InFeatureLevel) const OVERRIDE
 	{
-		const bool bInstanced = GRHIFeatureLevel >= ERHIFeatureLevel::SM3;
+		const bool bInstanced = InFeatureLevel >= ERHIFeatureLevel::SM3;
 		return bInstanced ? sizeof(FParticleSpriteVertex) : sizeof(FParticleSpriteVertexNonInstanced);
 	}
 
@@ -1954,7 +1954,7 @@ struct FDynamicMeshEmitterData : public FDynamicSpriteEmitterDataBase
 	/**
 	 *	Get the vertex stride for the dynamic rendering data
 	 */
-	virtual int32 GetDynamicVertexStride() const OVERRIDE
+	virtual int32 GetDynamicVertexStride(ERHIFeatureLevel::Type /*InFeatureLevel*/) const OVERRIDE
 	{
 		return sizeof(FMeshParticleInstanceVertex);
 	}
@@ -2251,7 +2251,7 @@ struct FDynamicBeam2EmitterData : public FDynamicSpriteEmitterDataBase
 	/**
 	 *	Get the vertex stride for the dynamic rendering data
 	 */
-	virtual int32 GetDynamicVertexStride() const OVERRIDE
+	virtual int32 GetDynamicVertexStride(ERHIFeatureLevel::Type /*InFeatureLevel*/) const OVERRIDE
 	{
 		return sizeof(FParticleBeamTrailVertex);
 	}
@@ -2434,7 +2434,7 @@ struct FDynamicTrailsEmitterData : public FDynamicSpriteEmitterDataBase
 	/**
 	 *	Get the vertex stride for the dynamic rendering data
 	 */
-	virtual int32 GetDynamicVertexStride() const OVERRIDE
+	virtual int32 GetDynamicVertexStride(ERHIFeatureLevel::Type /*InFeatureLevel*/) const OVERRIDE
 	{
 		return sizeof(FParticleBeamTrailVertex);
 	}
diff --git a/Engine/Source/Runtime/Engine/Public/SceneManagement.h b/Engine/Source/Runtime/Engine/Public/SceneManagement.h
--- a/Engine/Source/Runtime/Engine/Public/SceneManagement.h
+++ b/Engine/Source/Runtime/Engine/Public/SceneManagement.h
@@ -1321,16 +1321,16 @@ struct FMeshBatch
 	const FMaterialRenderProxy* MaterialRenderProxy;
 
 
-	FORCEINLINE bool IsTranslucent() const
+	FORCEINLINE bool IsTranslucent(ERHIFeatureLevel::Type InFeatureLevel) const
 	{
 		// Note: blend mode does not depend on the feature level we are actually rendering in.
-		return MaterialRenderProxy && IsTranslucentBlendMode(MaterialRenderProxy->GetMaterial(GRHIFeatureLevel)->GetBlendMode());
+		return MaterialRenderProxy && IsTranslucentBlendMode(MaterialRenderProxy->GetMaterial(InFeatureLevel)->GetBlendMode());
 	}
 
-	FORCEINLINE bool IsMasked() const
+	FORCEINLINE bool IsMasked(ERHIFeatureLevel::Type InFeatureLevel) const
 	{
 		// Note: blend mode does not depend on the feature level we are actually rendering in.
-		return MaterialRenderProxy && MaterialRenderProxy->GetMaterial(GRHIFeatureLevel)->IsMasked();
+		return MaterialRenderProxy && MaterialRenderProxy->GetMaterial(InFeatureLevel)->IsMasked();
 	}
 
 	/** Converts from an int32 index into a int8 */
@@ -1343,7 +1343,7 @@ struct FMeshBatch
 	/** 
 	* @return vertex stride specified for the mesh. 0 if not dynamic
 	*/
-	FORCEINLINE uint32 GetDynamicVertexStride() const
+	FORCEINLINE uint32 GetDynamicVertexStride(ERHIFeatureLevel::Type /*InFeatureLevel*/) const
 	{
 		if (UseDynamicData && DynamicVertexData)
 		{
@@ -2033,6 +2033,7 @@ public:
 	virtual bool IsEditorScene() const { return false; }
 
 	virtual ERHIFeatureLevel::Type GetFeatureLevel() const { return GRHIFeatureLevel; }
+	virtual void ChangeFeatureLevel(ERHIFeatureLevel::Type InFeatureLevel) {}
 
 protected:
 	virtual ~FSceneInterface() {}
diff --git a/Engine/Source/Runtime/RHI/Private/RHI.cpp b/Engine/Source/Runtime/RHI/Private/RHI.cpp
--- a/Engine/Source/Runtime/RHI/Private/RHI.cpp
+++ b/Engine/Source/Runtime/RHI/Private/RHI.cpp
@@ -171,6 +171,11 @@ RHI_API EShaderPlatform GRHIShaderPlatform = SP_PCD3D_SM5;
 /** The maximum feature level supported on this machine */
 ERHIFeatureLevel::Type GMaxRHIFeatureLevel = ERHIFeatureLevel::SM5;
 
+RHI_API ERHIFeatureLevel::Type GetRHIFeatureLevel()
+{
+	return GMaxRHIFeatureLevel;
+}
+
 RHI_API void SetMaxRHIFeatureLevel(ERHIFeatureLevel::Type InType)
 {
 	GMaxRHIFeatureLevel = InType;
diff --git a/Engine/Source/Runtime/RHI/Public/RHI.h b/Engine/Source/Runtime/RHI/Public/RHI.h
--- a/Engine/Source/Runtime/RHI/Public/RHI.h
+++ b/Engine/Source/Runtime/RHI/Public/RHI.h
@@ -316,12 +316,10 @@ extern RHI_API void GetFeatureLevelName(ERHIFeatureLevel::Type InFeatureLevel, F
 extern RHI_API ERHIFeatureLevel::Type GMaxRHIFeatureLevel;
 
 /** The current feature level supported by the RHI. */
-inline ERHIFeatureLevel::Type GetRHIFeatureLevel()
-{
-	return GMaxRHIFeatureLevel;
-}
+extern RHI_API ERHIFeatureLevel::Type GetRHIFeatureLevel();
 
-#define GRHIFeatureLevel GetRHIFeatureLevel()
+//#define GRHIFeatureLevel GetRHIFeatureLevel()
+#define GRHIFeatureLevel GMaxRHIFeatureLevel
 
 /** Set the current feature level supported by the RHI. */
 extern RHI_API void SetMaxRHIFeatureLevel(ERHIFeatureLevel::Type InType);
diff --git a/Engine/Source/Runtime/Renderer/Private/BasePassRendering.cpp b/Engine/Source/Runtime/Renderer/Private/BasePassRendering.cpp
--- a/Engine/Source/Runtime/Renderer/Private/BasePassRendering.cpp
+++ b/Engine/Source/Runtime/Renderer/Private/BasePassRendering.cpp
@@ -162,7 +162,7 @@ public:
 	{
 		FScene::EBasePassDrawListType DrawType = FScene::EBasePass_Default;		
  
-		if (StaticMesh->IsMasked())
+		if (StaticMesh->IsMasked(Parameters.FeatureLevel))
 		{
 			DrawType = FScene::EBasePass_Masked;	
 		}
@@ -209,9 +209,9 @@ void FBasePassOpaqueDrawingPolicyFactory::AddStaticMesh(FScene* Scene,FStaticMes
 				StaticMesh->PrimitiveSceneInfo->Proxy,
 				false,
 				bEditorCompositeDepthTest,
-				ESceneRenderTargetsMode::DontSet),
-			FDrawBasePassStaticMeshAction(Scene,StaticMesh),
-			Scene->GetFeatureLevel()
+				ESceneRenderTargetsMode::DontSet,
+				Scene->GetFeatureLevel()),
+			FDrawBasePassStaticMeshAction(Scene,StaticMesh)
 			);
 	}
 }
@@ -339,14 +339,14 @@ bool FBasePassOpaqueDrawingPolicyFactory::DrawDynamicMesh(
 				PrimitiveSceneProxy,
 				!bPreFog,
 				DrawingContext.bEditorCompositeDepthTest,
-				DrawingContext.TextureMode
+				DrawingContext.TextureMode,
+				View.GetFeatureLevel()
 				),
 			FDrawBasePassDynamicMeshAction(
 				View,
 				bBackFace,
 				HitProxyId
-				),
-			View.GetFeatureLevel()
+				)
 			);
 		return true;
 	}
diff --git a/Engine/Source/Runtime/Renderer/Private/BasePassRendering.h b/Engine/Source/Runtime/Renderer/Private/BasePassRendering.h
--- a/Engine/Source/Runtime/Renderer/Private/BasePassRendering.h
+++ b/Engine/Source/Runtime/Renderer/Private/BasePassRendering.h
@@ -785,6 +785,7 @@ public:
 	/** Whether or not to perform depth test in the pixel shader */
 	const bool bEditorCompositeDepthTest;
 	ESceneRenderTargetsMode::Type TextureMode;
+	ERHIFeatureLevel::Type FeatureLevel;
 
 	/** Initialization constructor. */
 	FProcessBasePassMeshParameters(
@@ -793,7 +794,8 @@ public:
 		const FPrimitiveSceneProxy* InPrimitiveSceneProxy,
 		bool InbAllowFog,
 		bool bInEditorCompositeDepthTest,
-		ESceneRenderTargetsMode::Type InTextureMode
+		ESceneRenderTargetsMode::Type InTextureMode,
+		ERHIFeatureLevel::Type InFeatureLevel
 		):
 		Mesh(InMesh),
 		BatchElementMask(Mesh.Elements.Num()==1 ? 1 : (1<<Mesh.Elements.Num())-1), // 1 bit set for each mesh element
@@ -803,7 +805,8 @@ public:
 		ShadingModel(InMaterial->GetShadingModel()),
 		bAllowFog(InbAllowFog),
 		bEditorCompositeDepthTest(bInEditorCompositeDepthTest),
-		TextureMode(InTextureMode)
+		TextureMode(InTextureMode),
+		FeatureLevel(InFeatureLevel)
 	{
 	}
 
@@ -815,8 +818,9 @@ public:
 		const FPrimitiveSceneProxy* InPrimitiveSceneProxy,
 		bool InbAllowFog,
 		bool bInEditorCompositeDepthTest,
-		ESceneRenderTargetsMode::Type InTextureMode
-		):
+		ESceneRenderTargetsMode::Type InTextureMode,
+		ERHIFeatureLevel::Type InFeatureLevel
+		) :
 		Mesh(InMesh),
 		BatchElementMask(InBatchElementMask),
 		Material(InMaterial),
@@ -825,7 +829,8 @@ public:
 		ShadingModel(InMaterial->GetShadingModel()),
 		bAllowFog(InbAllowFog),
 		bEditorCompositeDepthTest(bInEditorCompositeDepthTest),
-		TextureMode(InTextureMode)
+		TextureMode(InTextureMode),
+		FeatureLevel(InFeatureLevel)
 	{
 	}
 };
@@ -834,8 +839,7 @@ public:
 template<typename ProcessActionType>
 void ProcessBasePassMesh(
 	const FProcessBasePassMeshParameters& Parameters,
-	const ProcessActionType& Action,
-	ERHIFeatureLevel::Type InFeatureLevel
+	const ProcessActionType& Action
 	)
 {
 	// Check for a cached light-map.
@@ -850,9 +854,9 @@ void ProcessBasePassMesh(
 	if (!(Parameters.TextureMode == ESceneRenderTargetsMode::DontSet && bNeedsSceneTextures))
 	{
 		// Render self-shadowing only for >= SM4 and fallback to non-shadowed for lesser shader models
-		if (bIsLitMaterial && Action.UseTranslucentSelfShadowing() && InFeatureLevel >= ERHIFeatureLevel::SM4)
+		if (bIsLitMaterial && Action.UseTranslucentSelfShadowing() && Parameters.FeatureLevel >= ERHIFeatureLevel::SM4)
 		{
-			if (IsIndirectLightingCacheAllowed(InFeatureLevel)
+			if (IsIndirectLightingCacheAllowed(Parameters.FeatureLevel)
 				&& Action.AllowIndirectLightingCache()
 				&& Parameters.PrimitiveSceneProxy)
 			{
@@ -908,7 +912,7 @@ void ProcessBasePassMesh(
 							{
 								Action.template Process<FSimpleDynamicLightingPolicy>(Parameters, FSimpleDynamicLightingPolicy(), FSimpleDynamicLightingPolicy::ElementDataType());
 							}
-							else if (IsIndirectLightingCacheAllowed(InFeatureLevel)
+							else if (IsIndirectLightingCacheAllowed(Parameters.FeatureLevel)
 								&& Action.AllowIndirectLightingCache()
 								&& Parameters.PrimitiveSceneProxy
 								// Use the indirect lighting cache shaders if the object has a cache allocation
@@ -918,7 +922,7 @@ void ProcessBasePassMesh(
 									// Use the indirect lighting cache shaders if the object is movable, it may not have a cache allocation yet because that is done in InitViews
 									|| Parameters.PrimitiveSceneProxy->IsMovable()))
 							{
-								if (CanIndirectLightingCacheUseVolumeTexture(InFeatureLevel) && Action.AllowIndirectLightingCacheVolumeTexture())
+								if (CanIndirectLightingCacheUseVolumeTexture(Parameters.FeatureLevel) && Action.AllowIndirectLightingCacheVolumeTexture())
 								{
 									// Use a lightmap policy that supports reading indirect lighting from a volume texture for dynamic objects
 									Action.template Process<FCachedVolumeIndirectLightingPolicy>(Parameters,FCachedVolumeIndirectLightingPolicy(),FCachedVolumeIndirectLightingPolicy::ElementDataType());
diff --git a/Engine/Source/Runtime/Renderer/Private/DeferredShadingRenderer.cpp b/Engine/Source/Runtime/Renderer/Private/DeferredShadingRenderer.cpp
--- a/Engine/Source/Runtime/Renderer/Private/DeferredShadingRenderer.cpp
+++ b/Engine/Source/Runtime/Renderer/Private/DeferredShadingRenderer.cpp
@@ -448,7 +448,7 @@ void FDeferredShadingSceneRenderer::Render()
 	const auto FeatureLevel = Scene->GetFeatureLevel();
 
 	// Initialize global system textures (pass-through if already initialized).
-	GSystemTextures.InitializeTextures();
+	GSystemTextures.InitializeTextures(FeatureLevel);
 
 	// Allocate the maximum scene render target space for the current view family.
 	GSceneRenderTargets.Allocate(ViewFamily);
diff --git a/Engine/Source/Runtime/Renderer/Private/DistortionRendering.cpp b/Engine/Source/Runtime/Renderer/Private/DistortionRendering.cpp
--- a/Engine/Source/Runtime/Renderer/Private/DistortionRendering.cpp
+++ b/Engine/Source/Runtime/Renderer/Private/DistortionRendering.cpp
@@ -721,7 +721,7 @@ bool FDistortionPrimSet::DrawAccumulatedOffsets(const FViewInfo* ViewInfo,bool b
 						FStaticMesh& StaticMesh = PrimitiveSceneProxy->GetPrimitiveSceneInfo()->StaticMeshes[StaticMeshIdx];
 						if( ViewInfo->StaticMeshVisibilityMap[StaticMesh.Id]
 							// Only render static mesh elements using translucent materials
-							&& StaticMesh.IsTranslucent() )
+							&& StaticMesh.IsTranslucent(ViewInfo->GetFeatureLevel()) )
 						{
 							bDirty |= TDistortionMeshDrawingPolicyFactory<FDistortMeshAccumulatePolicy>::DrawStaticMesh(
 								ViewInfo,
diff --git a/Engine/Source/Runtime/Renderer/Private/ForwardBasePassRendering.cpp b/Engine/Source/Runtime/Renderer/Private/ForwardBasePassRendering.cpp
--- a/Engine/Source/Runtime/Renderer/Private/ForwardBasePassRendering.cpp
+++ b/Engine/Source/Runtime/Renderer/Private/ForwardBasePassRendering.cpp
@@ -59,7 +59,7 @@ public:
 	{
 		FScene::EBasePassDrawListType DrawType = FScene::EBasePass_Default;		
  
-		if (StaticMesh->IsMasked())
+		if (StaticMesh->IsMasked(Parameters.FeatureLevel))
 		{
 			DrawType = FScene::EBasePass_Masked;	
 		}
@@ -88,7 +88,8 @@ public:
 void FBasePassForwardOpaqueDrawingPolicyFactory::AddStaticMesh(FScene* Scene,FStaticMesh* StaticMesh)
 {
 	// Determine the mesh's material and blend mode.
-	const FMaterial* Material = StaticMesh->MaterialRenderProxy->GetMaterial(Scene->GetFeatureLevel());
+	const auto FeatureLevel = Scene->GetFeatureLevel();
+	const FMaterial* Material = StaticMesh->MaterialRenderProxy->GetMaterial(FeatureLevel);
 	const EBlendMode BlendMode = Material->GetBlendMode();
 
 	// Only draw opaque materials.
@@ -101,7 +102,8 @@ void FBasePassForwardOpaqueDrawingPolicyFactory::AddStaticMesh(FScene* Scene,FSt
 				StaticMesh->PrimitiveSceneInfo->Proxy,
 				true,
 				false,
-				ESceneRenderTargetsMode::DontSet
+				ESceneRenderTargetsMode::DontSet,
+				FeatureLevel
 				),
 			FDrawBasePassForwardShadingStaticMeshAction(Scene,StaticMesh)
 			);
@@ -206,7 +208,8 @@ bool FBasePassForwardOpaqueDrawingPolicyFactory::DrawDynamicMesh(
 	)
 {
 	// Determine the mesh's material and blend mode.
-	const FMaterial* Material = Mesh.MaterialRenderProxy->GetMaterial(View.GetFeatureLevel());
+	const auto FeatureLevel = View.GetFeatureLevel();
+	const FMaterial* Material = Mesh.MaterialRenderProxy->GetMaterial(FeatureLevel);
 	const EBlendMode BlendMode = Material->GetBlendMode();
 
 	// Only draw opaque materials.
@@ -219,7 +222,8 @@ bool FBasePassForwardOpaqueDrawingPolicyFactory::DrawDynamicMesh(
 				PrimitiveSceneProxy,
 				true,
 				false,
-				DrawingContext.TextureMode
+				DrawingContext.TextureMode,
+				FeatureLevel
 				),
 			FDrawBasePassForwardShadingDynamicMeshAction(
 				View,
diff --git a/Engine/Source/Runtime/Renderer/Private/ForwardBasePassRendering.h b/Engine/Source/Runtime/Renderer/Private/ForwardBasePassRendering.h
--- a/Engine/Source/Runtime/Renderer/Private/ForwardBasePassRendering.h
+++ b/Engine/Source/Runtime/Renderer/Private/ForwardBasePassRendering.h
@@ -486,7 +486,7 @@ void ProcessBasePassMeshForForwardShading(
 		}
 	}
 	else if (bIsLitMaterial
-		&& IsIndirectLightingCacheAllowed(GRHIFeatureLevel)
+		&& IsIndirectLightingCacheAllowed(Parameters.FeatureLevel)
 		&& Parameters.PrimitiveSceneProxy
 		// Movable objects need to get their GI from the indirect lighting cache
 		&& Parameters.PrimitiveSceneProxy->IsMovable())
diff --git a/Engine/Source/Runtime/Renderer/Private/ForwardShadingRenderer.cpp b/Engine/Source/Runtime/Renderer/Private/ForwardShadingRenderer.cpp
--- a/Engine/Source/Runtime/Renderer/Private/ForwardShadingRenderer.cpp
+++ b/Engine/Source/Runtime/Renderer/Private/ForwardShadingRenderer.cpp
@@ -48,7 +48,7 @@ void FForwardShadingSceneRenderer::Render()
 	auto FeatureLevel = ViewFamily.Scene->GetFeatureLevel();
 
 	// Initialize global system textures (pass-through if already initialized).
-	GSystemTextures.InitializeTextures();
+	GSystemTextures.InitializeTextures(FeatureLevel);
 
 	// Allocate the maximum scene render target space for the current view family.
 	GSceneRenderTargets.Allocate(ViewFamily);
diff --git a/Engine/Source/Runtime/Renderer/Private/ForwardTranslucentRendering.cpp b/Engine/Source/Runtime/Renderer/Private/ForwardTranslucentRendering.cpp
--- a/Engine/Source/Runtime/Renderer/Private/ForwardTranslucentRendering.cpp
+++ b/Engine/Source/Runtime/Renderer/Private/ForwardTranslucentRendering.cpp
@@ -174,7 +174,8 @@ bool FTranslucencyForwardShadingDrawingPolicyFactory::DrawDynamicMesh(
 	bool bDirty = false;
 
 	// Determine the mesh's material and blend mode.
-	const FMaterial* Material = Mesh.MaterialRenderProxy->GetMaterial(View.GetFeatureLevel());
+	const auto FeatureLevel = View.GetFeatureLevel();
+	const FMaterial* Material = Mesh.MaterialRenderProxy->GetMaterial(FeatureLevel);
 	const EBlendMode BlendMode = Material->GetBlendMode();
 
 	// Only render translucent materials.
@@ -187,7 +188,8 @@ bool FTranslucencyForwardShadingDrawingPolicyFactory::DrawDynamicMesh(
 				PrimitiveSceneProxy,
 				true,
 				false,
-				ESceneRenderTargetsMode::SetTextures
+				ESceneRenderTargetsMode::SetTextures,
+				FeatureLevel
 				),
 			FDrawTranslucentMeshForwardShadingAction(
 				View,
@@ -280,7 +282,7 @@ void FTranslucentPrimSet::RenderPrimitiveForForwardShading(
 				FStaticMesh& StaticMesh = PrimitiveSceneInfo->StaticMeshes[StaticMeshIdx];
 				if (View.StaticMeshVisibilityMap[StaticMesh.Id]
 					// Only render static mesh elements using translucent materials
-					&& StaticMesh.IsTranslucent() )
+					&& StaticMesh.IsTranslucent(View.GetFeatureLevel()) )
 				{
 					FTranslucencyForwardShadingDrawingPolicyFactory::DrawStaticMesh(
 						View,
diff --git a/Engine/Source/Runtime/Renderer/Private/IndirectLightingCache.cpp b/Engine/Source/Runtime/Renderer/Private/IndirectLightingCache.cpp
--- a/Engine/Source/Runtime/Renderer/Private/IndirectLightingCache.cpp
+++ b/Engine/Source/Runtime/Renderer/Private/IndirectLightingCache.cpp
@@ -141,7 +141,7 @@ bool IsIndirectLightingCacheAllowed(ERHIFeatureLevel::Type InFeatureLevel)
 bool CanIndirectLightingCacheUseVolumeTexture(ERHIFeatureLevel::Type InFeatureLevel)
 {
 	// @todo Mac OS X/OpenGL: For OpenGL devices which don't support volume-texture rendering we need to use the simpler point indirect lighting shaders.
-	return GRHIFeatureLevel >= ERHIFeatureLevel::SM3 && GSupportsVolumeTextureRendering;
+	return InFeatureLevel >= ERHIFeatureLevel::SM3 && GSupportsVolumeTextureRendering;
 }
 
 FIndirectLightingCache::FIndirectLightingCache()
diff --git a/Engine/Source/Runtime/Renderer/Private/RendererScene.cpp b/Engine/Source/Runtime/Renderer/Private/RendererScene.cpp
--- a/Engine/Source/Runtime/Renderer/Private/RendererScene.cpp
+++ b/Engine/Source/Runtime/Renderer/Private/RendererScene.cpp
@@ -1933,6 +1933,13 @@ void FScene::OnLevelAddedToWorld_RenderThread(FName InLevelName)
 	}
 }
 
+void FScene::ChangeFeatureLevel(ERHIFeatureLevel::Type InFeatureLevel)
+{
+	FeatureLevel = InFeatureLevel;
+	GRHIShaderPlatform = GShaderPlatformForFeatureLevel[InFeatureLevel];
+	SetMaxRHIFeatureLevel(InFeatureLevel);
+}
+
 /**
  * Dummy NULL scene interface used by dedicated servers.
  */
diff --git a/Engine/Source/Runtime/Renderer/Private/SceneCore.cpp b/Engine/Source/Runtime/Renderer/Private/SceneCore.cpp
--- a/Engine/Source/Runtime/Renderer/Private/SceneCore.cpp
+++ b/Engine/Source/Runtime/Renderer/Private/SceneCore.cpp
@@ -276,7 +276,8 @@ void FStaticMesh::UnlinkDrawList(FStaticMesh::FDrawListElementLink* Link)
 
 void FStaticMesh::AddToDrawLists(FScene* Scene)
 {
-	if (Scene->GetFeatureLevel() >= ERHIFeatureLevel::SM3)
+	const auto FeatureLevel = Scene->GetFeatureLevel();
+	if (FeatureLevel >= ERHIFeatureLevel::SM3)
 	{
 		if (CastShadow)
 		{
@@ -293,7 +294,7 @@ void FStaticMesh::AddToDrawLists(FScene* Scene)
 				FHitProxyDrawingPolicyFactory::AddStaticMesh(Scene,this);
 			}
 
-			if(!IsTranslucent())
+			if (!IsTranslucent(FeatureLevel))
 			{
 				extern TAutoConsoleVariable<int32> CVarEarlyZPass;
 				int32 EarlyZPass = CVarEarlyZPass.GetValueOnRenderThread();
@@ -302,7 +303,7 @@ void FStaticMesh::AddToDrawLists(FScene* Scene)
 
 				// Render non-masked materials in the depth only pass
 				if (PrimitiveSceneInfo->Proxy->ShouldUseAsOccluder() 
-					&& (!IsMasked() || EarlyZPass == 2)
+					&& (!IsMasked(FeatureLevel) || EarlyZPass == 2)
 					&& (!PrimitiveSceneInfo->Proxy->IsMovable() || GEarlyZPassMovable))
 				{
 					FDepthDrawingPolicyFactory::AddStaticMesh(Scene,this);
@@ -317,7 +318,7 @@ void FStaticMesh::AddToDrawLists(FScene* Scene)
 	}
 	else
 	{
-		if (!bShadowOnly && !IsTranslucent())
+		if (!bShadowOnly && !IsTranslucent(FeatureLevel))
 		{
 			// Add the static mesh to the DPG's base pass draw list.
 			FBasePassForwardOpaqueDrawingPolicyFactory::AddStaticMesh(Scene,this);
diff --git a/Engine/Source/Runtime/Renderer/Private/SceneHitProxyRendering.cpp b/Engine/Source/Runtime/Renderer/Private/SceneHitProxyRendering.cpp
--- a/Engine/Source/Runtime/Renderer/Private/SceneHitProxyRendering.cpp
+++ b/Engine/Source/Runtime/Renderer/Private/SceneHitProxyRendering.cpp
@@ -333,7 +333,7 @@ void FDeferredShadingSceneRenderer::RenderHitProxies()
 	auto FeatureLevel = ViewFamily.Scene->GetFeatureLevel();
 
 	// Initialize global system textures (pass-through if already initialized).
-	GSystemTextures.InitializeTextures();
+	GSystemTextures.InitializeTextures(FeatureLevel);
 
 	// Allocate the maximum scene render target space for the current view family.
 	GSceneRenderTargets.Allocate(ViewFamily);
diff --git a/Engine/Source/Runtime/Renderer/Private/ScenePrivate.h b/Engine/Source/Runtime/Renderer/Private/ScenePrivate.h
--- a/Engine/Source/Runtime/Renderer/Private/ScenePrivate.h
+++ b/Engine/Source/Runtime/Renderer/Private/ScenePrivate.h
@@ -1430,6 +1430,7 @@ public:
 	virtual bool IsEditorScene() const OVERRIDE { return bIsEditorScene; }
 
 	virtual ERHIFeatureLevel::Type GetFeatureLevel() const OVERRIDE { return FeatureLevel; }
+	virtual void ChangeFeatureLevel(ERHIFeatureLevel::Type InFeatureLevel) OVERRIDE;
 
 private:
 
diff --git a/Engine/Source/Runtime/Renderer/Private/ShadowSetup.cpp b/Engine/Source/Runtime/Renderer/Private/ShadowSetup.cpp
--- a/Engine/Source/Runtime/Renderer/Private/ShadowSetup.cpp
+++ b/Engine/Source/Runtime/Renderer/Private/ShadowSetup.cpp
@@ -658,6 +658,7 @@ void FProjectedShadowInfo::AddSubjectPrimitive(FPrimitiveSceneInfo* PrimitiveSce
 		bool bNeedsPreRenderView = false;
 		uint32 ViewMask = 0;
 		int32 PrimitiveId = PrimitiveSceneInfo->GetIndex();
+		const auto FeatureLevel = PrimitiveSceneInfo->Scene->GetFeatureLevel();
 
 		for (int32 ViewIndex = 0, Num = Views.Num(); ViewIndex < Num; ViewIndex++)
 		{
@@ -810,14 +811,14 @@ void FProjectedShadowInfo::AddSubjectPrimitive(FPrimitiveSceneInfo* PrimitiveSce
 						if (StaticMesh.CastShadow)
 						{
 							const FMaterialRenderProxy* MaterialRenderProxy = StaticMesh.MaterialRenderProxy;
-							const FMaterial* Material = MaterialRenderProxy->GetMaterial(GRHIFeatureLevel);
+							const FMaterial* Material = MaterialRenderProxy->GetMaterial(FeatureLevel);
 							const EBlendMode BlendMode = Material->GetBlendMode();
 							const EMaterialShadingModel ShadingModel = Material->GetShadingModel();
 
 							if(((!IsTranslucentBlendMode(BlendMode)) && ShadingModel != MSM_Unlit) || (bReflectiveShadowmap && Material->ShouldInjectEmissiveIntoLPV())) 
 							{
 								const bool bTwoSided = Material->IsTwoSided() || PrimitiveSceneInfo->Proxy->CastsShadowAsTwoSided();
-								OverrideWithDefaultMaterialForShadowDepth(MaterialRenderProxy, Material, bReflectiveShadowmap, GRHIFeatureLevel);
+								OverrideWithDefaultMaterialForShadowDepth(MaterialRenderProxy, Material, bReflectiveShadowmap, FeatureLevel);
 								SubjectMeshElements.Add(FShadowStaticMeshElement(MaterialRenderProxy, Material, &StaticMesh,bTwoSided));
 							}
 						}
diff --git a/Engine/Source/Runtime/Renderer/Private/SystemTextures.cpp b/Engine/Source/Runtime/Renderer/Private/SystemTextures.cpp
--- a/Engine/Source/Runtime/Renderer/Private/SystemTextures.cpp
+++ b/Engine/Source/Runtime/Renderer/Private/SystemTextures.cpp
@@ -16,71 +16,86 @@ TGlobalResource<FSystemTextures> GSystemTextures;
 
 uint8 Quantize8SignedByte(float x);
 
-void FSystemTextures::InitializeTextures()
+void FSystemTextures::InitializeTextures(ERHIFeatureLevel::Type InFeatureLevel)
 {
-	// Do not double allocate
-	if (bTexturesInitialized)
+	if (bTexturesInitialized && FeatureLevelInitializedTo >= InFeatureLevel)
 	{
+		// Already initialized up to at least the feature level we need, so do nothing
 		return;
 	}
 
 	// start with a defined state for the scissor rect (D3D11 was returning (0,0,0,0) which caused a clear to not execute correctly)
 	// todo: move this to an earlier place (for dx9 is has to be after device creation which is after window creation)
 	RHISetScissorRect(false, 0, 0, 0, 0);
-	
-	// Create a dummy white texture.
-	{
-		FPooledRenderTargetDesc Desc(FPooledRenderTargetDesc::Create2DDesc(FIntPoint(1, 1), PF_B8G8R8A8, TexCreate_HideInVisualizeTexture, TexCreate_RenderTargetable, false));
-		GRenderTargetPool.FindFreeElement(Desc, WhiteDummy, TEXT("WhiteDummy"));
-
-		RHISetRenderTarget(WhiteDummy->GetRenderTargetItem().TargetableTexture, FTextureRHIRef());
-		RHIClear(true, FLinearColor(1, 1, 1, 1), false, 0, false, 0, FIntRect());
-		RHICopyToResolveTarget(WhiteDummy->GetRenderTargetItem().TargetableTexture, WhiteDummy->GetRenderTargetItem().ShaderResourceTexture, true, FResolveParams());
-	}
 
-	// Create a dummy black texture.
+	// First initialize textures that are common to all feature levels. This is always done the first time we come into this function, as doesn't care about the
+	// requested feature level
+	if (!bTexturesInitialized)
 	{
-		FPooledRenderTargetDesc Desc(FPooledRenderTargetDesc::Create2DDesc(FIntPoint(1, 1), PF_B8G8R8A8, TexCreate_HideInVisualizeTexture, TexCreate_RenderTargetable, false));
-		GRenderTargetPool.FindFreeElement(Desc, BlackDummy, TEXT("BlackDummy"));
+		// Create a dummy white texture.
+		{
+			FPooledRenderTargetDesc Desc(FPooledRenderTargetDesc::Create2DDesc(FIntPoint(1, 1), PF_B8G8R8A8, TexCreate_HideInVisualizeTexture, TexCreate_RenderTargetable, false));
+			GRenderTargetPool.FindFreeElement(Desc, WhiteDummy, TEXT("WhiteDummy"));
 
-		RHISetRenderTarget(BlackDummy->GetRenderTargetItem().TargetableTexture, FTextureRHIRef());
-		RHIClear(true, FLinearColor(0, 0, 0, 0), false, 0, false, 0, FIntRect());
-		RHICopyToResolveTarget(BlackDummy->GetRenderTargetItem().TargetableTexture, BlackDummy->GetRenderTargetItem().ShaderResourceTexture, true, FResolveParams());
-	}
+			RHISetRenderTarget(WhiteDummy->GetRenderTargetItem().TargetableTexture, FTextureRHIRef());
+			RHIClear(true, FLinearColor(1, 1, 1, 1), false, 0, false, 0, FIntRect());
+			RHICopyToResolveTarget(WhiteDummy->GetRenderTargetItem().TargetableTexture, WhiteDummy->GetRenderTargetItem().ShaderResourceTexture, true, FResolveParams());
+		}
 
-	// Create the PerlinNoiseGradient texture
-	{
-		FPooledRenderTargetDesc Desc(FPooledRenderTargetDesc::Create2DDesc(FIntPoint(128, 128), PF_B8G8R8A8, TexCreate_HideInVisualizeTexture, TexCreate_None, false));
-		GRenderTargetPool.FindFreeElement(Desc, PerlinNoiseGradient, TEXT("PerlinNoiseGradient"));
-		// Write the contents of the texture.
-		uint32 DestStride;
-		uint8* DestBuffer = (uint8*)RHILockTexture2D((FTexture2DRHIRef&)PerlinNoiseGradient->GetRenderTargetItem().ShaderResourceTexture,0,RLM_WriteOnly,DestStride,false);
-		// seed the pseudo random stream with a good value
-		FRandomStream RandomStream(12345);
-		// Values represent float3 values in the -1..1 range.
-		// The vectors are the edge mid point of a cube from -1 .. 1
-		static uint32 gradtable[] = 
+		// Create a dummy black texture.
 		{
-			0x88ffff, 0xff88ff, 0xffff88,
-			0x88ff00, 0xff8800, 0xff0088,
-			0x8800ff, 0x0088ff, 0x00ff88,
-			0x880000, 0x008800, 0x000088,
-		};
-		for(int32 y = 0; y < Desc.Extent.Y; ++y)
+			FPooledRenderTargetDesc Desc(FPooledRenderTargetDesc::Create2DDesc(FIntPoint(1, 1), PF_B8G8R8A8, TexCreate_HideInVisualizeTexture, TexCreate_RenderTargetable, false));
+			GRenderTargetPool.FindFreeElement(Desc, BlackDummy, TEXT("BlackDummy"));
+
+			RHISetRenderTarget(BlackDummy->GetRenderTargetItem().TargetableTexture, FTextureRHIRef());
+			RHIClear(true, FLinearColor(0, 0, 0, 0), false, 0, false, 0, FIntRect());
+			RHICopyToResolveTarget(BlackDummy->GetRenderTargetItem().TargetableTexture, BlackDummy->GetRenderTargetItem().ShaderResourceTexture, true, FResolveParams());
+		}
+
+		// Create the PerlinNoiseGradient texture
 		{
-			for(int32 x = 0; x < Desc.Extent.X; ++x)
+			FPooledRenderTargetDesc Desc(FPooledRenderTargetDesc::Create2DDesc(FIntPoint(128, 128), PF_B8G8R8A8, TexCreate_HideInVisualizeTexture, TexCreate_None, false));
+			GRenderTargetPool.FindFreeElement(Desc, PerlinNoiseGradient, TEXT("PerlinNoiseGradient"));
+			// Write the contents of the texture.
+			uint32 DestStride;
+			uint8* DestBuffer = (uint8*)RHILockTexture2D((FTexture2DRHIRef&)PerlinNoiseGradient->GetRenderTargetItem().ShaderResourceTexture, 0, RLM_WriteOnly, DestStride, false);
+			// seed the pseudo random stream with a good value
+			FRandomStream RandomStream(12345);
+			// Values represent float3 values in the -1..1 range.
+			// The vectors are the edge mid point of a cube from -1 .. 1
+			static uint32 gradtable[] =
+			{
+				0x88ffff, 0xff88ff, 0xffff88,
+				0x88ff00, 0xff8800, 0xff0088,
+				0x8800ff, 0x0088ff, 0x00ff88,
+				0x880000, 0x008800, 0x000088,
+			};
+			for (int32 y = 0; y < Desc.Extent.Y; ++y)
 			{
-				uint32* Dest = (uint32*)(DestBuffer + x * sizeof(uint32) + y * DestStride);
+				for (int32 x = 0; x < Desc.Extent.X; ++x)
+				{
+					uint32* Dest = (uint32*)(DestBuffer + x * sizeof(uint32)+y * DestStride);
 
-				// pick a random direction (hacky way to overcome the quality issues FRandomStream has)
-				*Dest = gradtable[(uint32)(RandomStream.GetFraction() * 11.9999999f)];
+					// pick a random direction (hacky way to overcome the quality issues FRandomStream has)
+					*Dest = gradtable[(uint32)(RandomStream.GetFraction() * 11.9999999f)];
+				}
 			}
+			RHIUnlockTexture2D((FTexture2DRHIRef&)PerlinNoiseGradient->GetRenderTargetItem().ShaderResourceTexture, 0, false);
+		}
+
+		if (!GSupportsShaderFramebufferFetch && GPixelFormats[PF_FloatRGBA].Supported)
+		{
+			FPooledRenderTargetDesc Desc(FPooledRenderTargetDesc::Create2DDesc(FIntPoint(1, 1), PF_FloatRGBA, TexCreate_HideInVisualizeTexture, TexCreate_RenderTargetable, false));
+			GRenderTargetPool.FindFreeElement(Desc, MaxFP16Depth, TEXT("MaxFP16Depth"));
+
+			RHISetRenderTarget(MaxFP16Depth->GetRenderTargetItem().TargetableTexture, FTextureRHIRef());
+			RHIClear(true, FLinearColor(65000.0f, 65000.0f, 65000.0f, 65000.0f), false, 0, false, 0, FIntRect());
+			RHICopyToResolveTarget(MaxFP16Depth->GetRenderTargetItem().TargetableTexture, MaxFP16Depth->GetRenderTargetItem().ShaderResourceTexture, true, FResolveParams());
 		}
-		RHIUnlockTexture2D((FTexture2DRHIRef&)PerlinNoiseGradient->GetRenderTargetItem().ShaderResourceTexture,0,false);
 	}
 
 	// Create the PerlinNoise3D texture (similar to http://prettyprocs.wordpress.com/2012/10/20/fast-perlin-noise/)
-	if(GRHIFeatureLevel >= ERHIFeatureLevel::SM5)
+	if (FeatureLevelInitializedTo < ERHIFeatureLevel::SM5 && InFeatureLevel >= ERHIFeatureLevel::SM5)
 	{
 		uint32 Extent = 16;
 
@@ -191,7 +206,7 @@ void FSystemTextures::InitializeTextures()
 	}
 
 	// Create the SSAO randomization texture
-	if (GRHIFeatureLevel >= ERHIFeatureLevel::SM3)
+	if (FeatureLevelInitializedTo < ERHIFeatureLevel::SM3 && InFeatureLevel >= ERHIFeatureLevel::SM3)
 	{
 		float g_AngleOff1 = 127;
 		float g_AngleOff2 = 198;
@@ -234,119 +249,107 @@ void FSystemTextures::InitializeTextures()
 				*Dest = Bases[Index];
 			}
 		}
-		RHIUnlockTexture2D((FTexture2DRHIRef&)SSAORandomization->GetRenderTargetItem().ShaderResourceTexture,0,false);
-	}
-
-
-	// Create preintegrated BRDF texture
-	if (GRHIFeatureLevel >= ERHIFeatureLevel::SM3)
-	{
-		// for testing, with 128x128 R8G8 we are very close to the reference (if lower res is needed we might have to add an offset to counter the 0.5f texel shift)
-		const bool bReference = false;
+		RHIUnlockTexture2D((FTexture2DRHIRef&)SSAORandomization->GetRenderTargetItem().ShaderResourceTexture, 0, false);
 
-		FPooledRenderTargetDesc Desc(FPooledRenderTargetDesc::Create2DDesc(FIntPoint(128, 32), PF_R8G8, TexCreate_HideInVisualizeTexture | TexCreate_FastVRAM, TexCreate_None, false));
-
-		if(bReference)
 		{
-			Desc.Extent.X = 128;
-			Desc.Extent.Y = 128;
-			Desc.Format =  PF_G16R16;
-		}
+			// for testing, with 128x128 R8G8 we are very close to the reference (if lower res is needed we might have to add an offset to counter the 0.5f texel shift)
+			const bool bReference = false;
 
-		GRenderTargetPool.FindFreeElement(Desc, PreintegratedGF, TEXT("PreintegratedGF"));
-		// Write the contents of the texture.
-		uint32 DestStride;
-		uint8* DestBuffer = (uint8*)RHILockTexture2D((FTexture2DRHIRef&)PreintegratedGF->GetRenderTargetItem().ShaderResourceTexture,0,RLM_WriteOnly,DestStride,false);
-		
-		// x is NoV, y is roughness
-		for( int32 y = 0; y < Desc.Extent.Y; y++ )
-		{
-			float Roughness = (float)( y + 0.5f ) / Desc.Extent.Y;
-			float m = Roughness * Roughness;
-			float m2 = m*m;
-			float k = m * 0.5f;
-			
-			for( int32 x = 0; x < Desc.Extent.X; x++ )
+			FPooledRenderTargetDesc Desc(FPooledRenderTargetDesc::Create2DDesc(FIntPoint(128, 32), PF_R8G8, TexCreate_HideInVisualizeTexture | TexCreate_FastVRAM, TexCreate_None, false));
+
+			if (bReference)
 			{
-				float NoV = (float)( x + 0.5f ) / Desc.Extent.X;
-				float G_SchlickV = 1.0f / ( NoV * (1.0f - k) + k );
-				float G_SmithV = 2.0f / ( NoV + FMath::Sqrt( (NoV*NoV) * (1.0f - m2) + m2 ) );
+				Desc.Extent.X = 128;
+				Desc.Extent.Y = 128;
+				Desc.Format = PF_G16R16;
+			}
 
-				FVector V;
-				V.X = FMath::Sqrt( 1.0f - NoV * NoV );	// sin
-				V.Y = 0.0f;
-				V.Z = NoV;								// cos
+			GRenderTargetPool.FindFreeElement(Desc, PreintegratedGF, TEXT("PreintegratedGF"));
+			// Write the contents of the texture.
+			uint32 DestStride;
+			uint8* DestBuffer = (uint8*)RHILockTexture2D((FTexture2DRHIRef&)PreintegratedGF->GetRenderTargetItem().ShaderResourceTexture, 0, RLM_WriteOnly, DestStride, false);
 
-				float A = 0.0f;
-				float B = 0.0f;
+			// x is NoV, y is roughness
+			for (int32 y = 0; y < Desc.Extent.Y; y++)
+			{
+				float Roughness = (float)(y + 0.5f) / Desc.Extent.Y;
+				float m = Roughness * Roughness;
+				float m2 = m*m;
+				float k = m * 0.5f;
 
-				const uint32 NumSamples = 128;
-				for( uint32 i = 0; i < NumSamples; i++ )
+				for (int32 x = 0; x < Desc.Extent.X; x++)
 				{
-					float E1 = (float)i / NumSamples;
-					float E2 = (double)ReverseBits(i) / (double)0x100000000LL;
+					float NoV = (float)(x + 0.5f) / Desc.Extent.X;
+					float G_SchlickV = 1.0f / (NoV * (1.0f - k) + k);
+					float G_SmithV = 2.0f / (NoV + FMath::Sqrt((NoV*NoV) * (1.0f - m2) + m2));
 
-					float Phi = 2.0f * PI * E1;
-					float CosPhi = FMath::Cos( Phi );
-					float SinPhi = FMath::Sin( Phi );
-					float CosTheta = FMath::Sqrt( (1.0f - E2) / ( 1.0f + (m2 - 1.0f) * E2 ) );
-					float SinTheta = FMath::Sqrt( 1.0f - CosTheta * CosTheta );
+					FVector V;
+					V.X = FMath::Sqrt(1.0f - NoV * NoV);	// sin
+					V.Y = 0.0f;
+					V.Z = NoV;								// cos
 
-					FVector H( SinTheta * FMath::Cos( Phi ), SinTheta * FMath::Sin( Phi ), CosTheta );
+					float A = 0.0f;
+					float B = 0.0f;
 
-					FVector L = 2.0f * ( V | H ) * H - V;
+					const uint32 NumSamples = 128;
+					for (uint32 i = 0; i < NumSamples; i++)
+					{
+						float E1 = (float)i / NumSamples;
+						float E2 = (double)ReverseBits(i) / (double)0x100000000LL;
+
+						float Phi = 2.0f * PI * E1;
+						float CosPhi = FMath::Cos(Phi);
+						float SinPhi = FMath::Sin(Phi);
+						float CosTheta = FMath::Sqrt((1.0f - E2) / (1.0f + (m2 - 1.0f) * E2));
+						float SinTheta = FMath::Sqrt(1.0f - CosTheta * CosTheta);
+
+						FVector H(SinTheta * FMath::Cos(Phi), SinTheta * FMath::Sin(Phi), CosTheta);
+
+						FVector L = 2.0f * (V | H) * H - V;
+
+						float NoL = FMath::Max(L.Z, 0.0f);
+						float NoH = FMath::Max(H.Z, 0.0f);
+						float VoH = FMath::Max(V | H, 0.0f);
+
+						if (NoL > 0.0f)
+						{
+							float G_SchlickL = 1.0f / (NoL * (1.0f - k) + k);
+							float G_SmithL = 2.0f / (NoL + FMath::Sqrt((NoL*NoL) * (1.0f - m2) + m2));
+
+							float Vis_G = VoH * (NoL / NoH) * G_SchlickL;
+							float Fc = 1.0f - VoH;
+							Fc *= FMath::Square(Fc*Fc);
+							A += Vis_G * (1.0f - Fc);
+							B += Vis_G * Fc;
+						}
+					}
+					A /= NumSamples;
+					B /= NumSamples;
 
-					float NoL = FMath::Max( L.Z, 0.0f );
-					float NoH = FMath::Max( H.Z, 0.0f );
-					float VoH = FMath::Max( V | H, 0.0f );
+					A *= G_SchlickV;
+					B *= G_SchlickV;
 
-					if( NoL > 0.0f )
+					if (bReference)
 					{
-						float G_SchlickL = 1.0f / ( NoL * (1.0f - k) + k );
-						float G_SmithL = 2.0f / ( NoL + FMath::Sqrt( (NoL*NoL) * (1.0f - m2) + m2 ) );
-
-						float Vis_G = VoH * (NoL / NoH) * G_SchlickL;
-						float Fc = 1.0f - VoH;
-						Fc *= FMath::Square( Fc*Fc );
-						A += Vis_G * (1.0f - Fc);
-						B += Vis_G * Fc;
+						uint16* Dest = (uint16*)(DestBuffer + x * sizeof(uint32)+y * DestStride);
+						Dest[0] = (int32)(FMath::Clamp(A, 0.0f, 1.0f) * 65535.0f + 0.5f);
+						Dest[1] = (int32)(FMath::Clamp(B, 0.0f, 1.0f) * 65535.0f + 0.5f);
+					}
+					else
+					{
+						uint8* Dest = (uint8*)(DestBuffer + x * sizeof(uint16)+y * DestStride);
+						Dest[0] = (int32)(FMath::Clamp(A, 0.0f, 1.0f) * 255.9999f);
+						Dest[1] = (int32)(FMath::Clamp(B, 0.0f, 1.0f) * 255.9999f);
 					}
-				}
-				A /= NumSamples;
-				B /= NumSamples;
-				
-				A *= G_SchlickV;
-				B *= G_SchlickV;
-
-				if(bReference)
-				{
-					uint16* Dest = (uint16*)(DestBuffer + x * sizeof(uint32) + y * DestStride);
-					Dest[0] = (int32)( FMath::Clamp( A, 0.0f, 1.0f ) * 65535.0f + 0.5f );
-					Dest[1] = (int32)( FMath::Clamp( B, 0.0f, 1.0f ) * 65535.0f + 0.5f );
-				}
-				else
-				{
-					uint8* Dest = (uint8*)(DestBuffer + x * sizeof(uint16) + y * DestStride);
-					Dest[0] = (int32)( FMath::Clamp( A, 0.0f, 1.0f ) * 255.9999f );
-					Dest[1] = (int32)( FMath::Clamp( B, 0.0f, 1.0f ) * 255.9999f );
 				}
 			}
+			RHIUnlockTexture2D((FTexture2DRHIRef&)PreintegratedGF->GetRenderTargetItem().ShaderResourceTexture, 0, false);
 		}
-		RHIUnlockTexture2D((FTexture2DRHIRef&)PreintegratedGF->GetRenderTargetItem().ShaderResourceTexture,0,false);
-	}
-
-	if (GRHIFeatureLevel == ERHIFeatureLevel::ES2 && !GSupportsShaderFramebufferFetch && GPixelFormats[PF_FloatRGBA].Supported)
-	{
-		FPooledRenderTargetDesc Desc(FPooledRenderTargetDesc::Create2DDesc(FIntPoint(1, 1), PF_FloatRGBA, TexCreate_HideInVisualizeTexture, TexCreate_RenderTargetable, false));
-		GRenderTargetPool.FindFreeElement(Desc, MaxFP16Depth, TEXT("MaxFP16Depth"));
-
-		RHISetRenderTarget(MaxFP16Depth->GetRenderTargetItem().TargetableTexture, FTextureRHIRef());
-		RHIClear(true, FLinearColor(65000.0f, 65000.0f, 65000.0f, 65000.0f), false, 0, false, 0, FIntRect());
-		RHICopyToResolveTarget(MaxFP16Depth->GetRenderTargetItem().TargetableTexture, MaxFP16Depth->GetRenderTargetItem().ShaderResourceTexture, true, FResolveParams());
 	}
 
 	// Initialize textures only once.
 	bTexturesInitialized = true;
+	FeatureLevelInitializedTo = InFeatureLevel;
 }
 
 void FSystemTextures::ReleaseDynamicRHI()
diff --git a/Engine/Source/Runtime/Renderer/Private/SystemTextures.h b/Engine/Source/Runtime/Renderer/Private/SystemTextures.h
--- a/Engine/Source/Runtime/Renderer/Private/SystemTextures.h
+++ b/Engine/Source/Runtime/Renderer/Private/SystemTextures.h
@@ -17,13 +17,14 @@ class FSystemTextures : public FRenderResource
 public:
 	FSystemTextures()
 	: FRenderResource()
+	, FeatureLevelInitializedTo(ERHIFeatureLevel::ES2)
 	, bTexturesInitialized(false)
 	{}
 
 	/**
 	 * Initialize/allocate textures if not already.
 	 */
-	void InitializeTextures();
+	void InitializeTextures(ERHIFeatureLevel::Type InFeatureLevel);
 
 	// FRenderResource interface.
 	/**
@@ -48,6 +49,8 @@ public:
 	/** Texture that holds a single value containing the maximum depth that can be stored as FP16. */
 	TRefCountPtr<IPooledRenderTarget> MaxFP16Depth;
 
+	/** Maximum feature level that the textures have been initialized up to */
+	ERHIFeatureLevel::Type FeatureLevelInitializedTo;
 	bool bTexturesInitialized;
 };
 
diff --git a/Engine/Source/Runtime/Renderer/Private/TranslucentRendering.cpp b/Engine/Source/Runtime/Renderer/Private/TranslucentRendering.cpp
--- a/Engine/Source/Runtime/Renderer/Private/TranslucentRendering.cpp
+++ b/Engine/Source/Runtime/Renderer/Private/TranslucentRendering.cpp
@@ -313,9 +313,10 @@ bool FTranslucencyDrawingPolicyFactory::DrawMesh(
 	)
 {
 	bool bDirty = false;
+	const auto FeatureLevel = View.GetFeatureLevel();
 
 	// Determine the mesh's material and blend mode.
-	const FMaterial* Material = Mesh.MaterialRenderProxy->GetMaterial(View.GetFeatureLevel());
+	const FMaterial* Material = Mesh.MaterialRenderProxy->GetMaterial(FeatureLevel);
 	const EBlendMode BlendMode = Material->GetBlendMode();
 	const EMaterialShadingModel ShadingModel = Material->GetShadingModel();
 
@@ -352,7 +353,8 @@ bool FTranslucencyDrawingPolicyFactory::DrawMesh(
 				PrimitiveSceneProxy, 
 				!bPreFog,
 				bEditorCompositeDepthTest,
-				ESceneRenderTargetsMode::SetTextures
+				ESceneRenderTargetsMode::SetTextures,
+				FeatureLevel
 			),
 			FDrawTranslucentMeshAction(
 				View,
@@ -360,8 +362,7 @@ bool FTranslucencyDrawingPolicyFactory::DrawMesh(
 				HitProxyId,
 				DrawingContext.TranslucentSelfShadow,
 				PrimitiveSceneProxy && PrimitiveSceneProxy->CastsVolumetricTranslucentShadow()
-			),
-			View.GetFeatureLevel()
+			)
 		);
 
 		if (bDisableDepthTest || bEnableResponsiveAA)
@@ -492,6 +493,7 @@ void FTranslucentPrimSet::RenderPrimitive(
 	bool bSeparateTranslucencyPass) const
 {
 	checkSlow(ViewRelevance.HasTranslucency());
+	auto FeatureLevel = View.GetFeatureLevel();
 
 	if(ViewRelevance.bDrawRelevance)
 	{
@@ -520,8 +522,8 @@ void FTranslucentPrimSet::RenderPrimitive(
 				FStaticMesh& StaticMesh = PrimitiveSceneInfo->StaticMeshes[StaticMeshIdx];
 				if (View.StaticMeshVisibilityMap[StaticMesh.Id]
 					// Only render static mesh elements using translucent materials
-					&& StaticMesh.IsTranslucent() 
-					&& (StaticMesh.MaterialRenderProxy->GetMaterial(View.GetFeatureLevel())->IsSeparateTranslucencyEnabled() == bSeparateTranslucencyPass))
+					&& StaticMesh.IsTranslucent(FeatureLevel)
+					&& (StaticMesh.MaterialRenderProxy->GetMaterial(FeatureLevel)->IsSeparateTranslucencyEnabled() == bSeparateTranslucencyPass))
 				{
 					FTranslucencyDrawingPolicyFactory::DrawStaticMesh(
 						View,
