Whole Scene Directional Shadow GPU optimization - Uses depth bounds tests (if available) rather than the stencil
Saves approx 0.1ms of GPU on GTX780

diff --git a/Engine/Source/Runtime/OpenGLDrv/Private/OpenGLDevice.cpp b/Engine/Source/Runtime/OpenGLDrv/Private/OpenGLDevice.cpp
--- a/Engine/Source/Runtime/OpenGLDrv/Private/OpenGLDevice.cpp
+++ b/Engine/Source/Runtime/OpenGLDrv/Private/OpenGLDevice.cpp
@@ -625,6 +625,7 @@ static void InitRHICapabilitiesForGL()
 	GSupportsVertexTextureFetch = true;
 	GSupportsRenderTargetFormat_PF_G8 = true;
 	GSupportsSeparateRenderTargetBlendState = FOpenGL::SupportsSeparateAlphaBlend();
+	GSupportsDepthBoundsTest = FOpenGL::SupportsDepthBoundsTest();
 
 	GSupportsRenderTargetFormat_PF_FloatRGBA = FOpenGL::SupportsColorBufferHalfFloat();
 
diff --git a/Engine/Source/Runtime/RHI/Private/RHI.cpp b/Engine/Source/Runtime/RHI/Private/RHI.cpp
--- a/Engine/Source/Runtime/RHI/Private/RHI.cpp
+++ b/Engine/Source/Runtime/RHI/Private/RHI.cpp
@@ -134,6 +134,7 @@ bool GSupportsEmulatedVertexInstancing = false;
 bool GVertexElementsCanShareStreamOffset = true;
 bool GTriggerGPUProfile = false;
 bool GRHISupportsTextureStreaming = false;
+bool GSupportsDepthBoundsTest = false;
 
 /** Whether we are profiling GPU hitches. */
 bool GTriggerGPUHitchProfile = false;
diff --git a/Engine/Source/Runtime/RHI/Public/RHI.h b/Engine/Source/Runtime/RHI/Public/RHI.h
--- a/Engine/Source/Runtime/RHI/Public/RHI.h
+++ b/Engine/Source/Runtime/RHI/Public/RHI.h
@@ -115,6 +115,9 @@ extern RHI_API bool GSupportsGSRenderTargetLayerSwitchingToMips;
 /** True if the RHI supports separate blend states per render target. */
 extern RHI_API bool GSupportsSeparateRenderTargetBlendState;
 
+/** True if the RHI supports depth bounds testing */
+extern RHI_API bool GSupportsDepthBoundsTest;
+
 /** The offset from the upper left corner of a pixel to the position which is used to sample textures for fragments of that pixel. */
 extern RHI_API float GPixelCenterOffset;
 
diff --git a/Engine/Source/Runtime/Renderer/Private/ShadowRendering.cpp b/Engine/Source/Runtime/Renderer/Private/ShadowRendering.cpp
--- a/Engine/Source/Runtime/Renderer/Private/ShadowRendering.cpp
+++ b/Engine/Source/Runtime/Renderer/Private/ShadowRendering.cpp
@@ -21,6 +21,12 @@ static TAutoConsoleVariable<float> CVarCSMSplitPenumbraScale(
 	TEXT("Scale applied to the penumbra size of Cascaded Shadow Map splits, useful for minimizing the transition between splits"),
 	ECVF_RenderThreadSafe);
 
+static TAutoConsoleVariable<int32> CVarCSMDepthBoundsTest(
+	TEXT("r.Shadow.CSMDepthBoundsTest"),
+	1,
+	TEXT("Whether to use depth bounds tests rather than stencil tests for the CSM bounds"),
+	ECVF_RenderThreadSafe);
+
 static TAutoConsoleVariable<float> CVarSpotLightShadowTransitionScale(
 	TEXT("r.Shadow.SpotLightTransitionScale"),
 	60.0f,
@@ -1607,6 +1613,8 @@ void FProjectedShadowInfo::RenderProjection(int32 ViewIndex, const FViewInfo* Vi
 	RHISetDepthStencilState(TStaticDepthStencilState<false,CF_GreaterEqual>::GetRHI());
 	RHISetBlendState(TStaticBlendState<CW_NONE>::GetRHI());
 	
+	bool bDepthBoundsTestEnabled = false;
+
 	// If this is a preshadow, mask the projection by the receiver primitives.
 	if (bPreShadow)
 	{
@@ -1665,53 +1673,79 @@ void FProjectedShadowInfo::RenderProjection(int32 ViewIndex, const FViewInfo* Vi
 	}
 	else if (IsWholeSceneDirectionalShadow())
 	{
-		// Increment stencil on front-facing zfail, decrement on back-facing zfail.
-		// Note, this is a reversed Z depth surface, using CF_GreaterEqual.
-		RHISetDepthStencilState(TStaticDepthStencilState<
-			false,CF_GreaterEqual,
-			true,CF_Always,SO_Keep,SO_Increment,SO_Keep,
-			true,CF_Always,SO_Keep,SO_Decrement,SO_Keep,
-			0xff,0xff
-			>::GetRHI());
+		if( GSupportsDepthBoundsTest && CVarCSMDepthBoundsTest.GetValueOnRenderThread() != 0 )
+		{
+			FVector4 Near = View->ViewMatrices.ProjMatrix.TransformFVector4(FVector4(0, 0, CascadeSettings.SplitNear));
+			FVector4 Far = View->ViewMatrices.ProjMatrix.TransformFVector4(FVector4(0, 0, CascadeSettings.SplitFar));
+			float DepthNear = Near.Z / Near.W;
+			float DepthFar = Far.Z / Far.W;
 
-		RHISetRasterizerState(TStaticRasterizerState<FM_Solid,CM_None>::GetRHI());
+			if( DepthFar > 1 ) DepthFar = 1.f;
+			else if( DepthFar < 0 ) DepthFar = 0.f;
 
-		checkSlow(SplitIndex >= 0);
-		checkSlow(bDirectionalLight);
+			if( DepthNear > 1 ) DepthNear = 1.f;
+			else if( DepthNear < 0 ) DepthNear = 0.f;
 
-		// Draw 2 fullscreen planes, front facing one at the near subfrustum plane, and back facing one at the far.
+			if( DepthNear <= DepthFar )
+			{
+				DepthNear = 1.0f;
+				DepthFar = 0.0f;
+			}
 
-		// Find the projection shaders.
-		TShaderMapRef<FShadowProjectionNoTransformVS> VertexShaderNoTransform(GetGlobalShaderMap());
-		VertexShaderNoTransform->SetParameters(*View);
-		SetGlobalBoundShaderState(MaskBoundShaderState[0], GetVertexDeclarationFVector4(), *VertexShaderNoTransform, nullptr);
-
-		FVector4 Near = View->ViewMatrices.ProjMatrix.TransformFVector4(FVector4(0, 0, CascadeSettings.SplitNear));
-		FVector4 Far = View->ViewMatrices.ProjMatrix.TransformFVector4(FVector4(0, 0, CascadeSettings.SplitFar));
-		float StencilNear = Near.Z / Near.W;
-		float StencilFar = Far.Z / Far.W;
-
-		FVector4 Verts[] =
-		{
-			// Far Plane
-			FVector4( 1,  1,  StencilFar),
-			FVector4(-1,  1,  StencilFar),
-			FVector4( 1, -1,  StencilFar),
-			FVector4( 1, -1,  StencilFar),
-			FVector4(-1,  1,  StencilFar),
-			FVector4(-1, -1,  StencilFar),
-
-			// Near Plane
-			FVector4(-1,  1, StencilNear),
-			FVector4( 1,  1, StencilNear),
-			FVector4(-1, -1, StencilNear),
-			FVector4(-1, -1, StencilNear),
-			FVector4( 1,  1, StencilNear),
-			FVector4( 1, -1, StencilNear),
-		};
+			// Note, using a reversed z depth surface
+			RHIEnableDepthBoundsTest( true, DepthFar, DepthNear );
+			bDepthBoundsTestEnabled = true;
+		}
+		else
+		{
+			// Increment stencil on front-facing zfail, decrement on back-facing zfail.
+			// Note, this is a reversed Z depth surface, using CF_GreaterEqual.
+			RHISetDepthStencilState(TStaticDepthStencilState<
+				false,CF_GreaterEqual,
+				true,CF_Always,SO_Keep,SO_Increment,SO_Keep,
+				true,CF_Always,SO_Keep,SO_Decrement,SO_Keep,
+				0xff,0xff
+				>::GetRHI());
+
+			RHISetRasterizerState(TStaticRasterizerState<FM_Solid,CM_None>::GetRHI());
+
+			checkSlow(SplitIndex >= 0);
+			checkSlow(bDirectionalLight);
+
+			// Draw 2 fullscreen planes, front facing one at the near subfrustum plane, and back facing one at the far.
 
-		// Only draw the near plane if this is not the nearest split
-		RHIDrawPrimitiveUP(PT_TriangleList, (SplitIndex > 0) ? 4 : 2, Verts, sizeof(FVector4));
+			// Find the projection shaders.
+			TShaderMapRef<FShadowProjectionNoTransformVS> VertexShaderNoTransform(GetGlobalShaderMap());
+			VertexShaderNoTransform->SetParameters(*View);
+			SetGlobalBoundShaderState(MaskBoundShaderState[0], GetVertexDeclarationFVector4(), *VertexShaderNoTransform, nullptr);
+
+			FVector4 Near = View->ViewMatrices.ProjMatrix.TransformFVector4(FVector4(0, 0, CascadeSettings.SplitNear));
+			FVector4 Far = View->ViewMatrices.ProjMatrix.TransformFVector4(FVector4(0, 0, CascadeSettings.SplitFar));
+			float StencilNear = Near.Z / Near.W;
+			float StencilFar = Far.Z / Far.W;
+
+			FVector4 Verts[] =
+			{
+				// Far Plane
+				FVector4( 1,  1,  StencilFar),
+				FVector4(-1,  1,  StencilFar),
+				FVector4( 1, -1,  StencilFar),
+				FVector4( 1, -1,  StencilFar),
+				FVector4(-1,  1,  StencilFar),
+				FVector4(-1, -1,  StencilFar),
+
+				// Near Plane
+				FVector4(-1,  1, StencilNear),
+				FVector4( 1,  1, StencilNear),
+				FVector4(-1, -1, StencilNear),
+				FVector4(-1, -1, StencilNear),
+				FVector4( 1,  1, StencilNear),
+				FVector4( 1, -1, StencilNear),
+			};
+
+			// Only draw the near plane if this is not the nearest split
+			RHIDrawPrimitiveUP(PT_TriangleList, (SplitIndex > 0) ? 4 : 2, Verts, sizeof(FVector4));
+		}
 	}
 	// Not a preshadow, mask the projection to any pixels inside the frustum.
 	else
@@ -1759,17 +1793,25 @@ void FProjectedShadowInfo::RenderProjection(int32 ViewIndex, const FViewInfo* Vi
 		RHIDrawIndexedPrimitiveUP(PT_TriangleList, 0, 8, 12, GCubeIndices, sizeof(uint16), FrustumVertices, sizeof(FVector4));
 	}
 
-	// no depth test or writes, solid rasterization w/ back-face culling.
-	// Test stencil for non-zero.
+	// solid rasterization w/ back-face culling.
 	RHISetRasterizerState(
 		XOR(View->bReverseCulling, IsWholeSceneDirectionalShadow()) ? TStaticRasterizerState<FM_Solid,CM_CCW>::GetRHI() : TStaticRasterizerState<FM_Solid,CM_CW>::GetRHI());
 
-	RHISetDepthStencilState(TStaticDepthStencilState<
-		false,CF_Always,
-		true,CF_NotEqual,SO_Keep,SO_Keep,SO_Keep,
-		false,CF_Always,SO_Keep,SO_Keep,SO_Keep,
-		0xff,0xff
-		>::GetRHI());
+	if( bDepthBoundsTestEnabled )
+	{
+		// no depth test or writes
+		RHISetDepthStencilState(TStaticDepthStencilState<false, CF_Always>::GetRHI());
+	}
+	else
+	{
+		// no depth test or writes, Test stencil for non-zero.
+		RHISetDepthStencilState( TStaticDepthStencilState<
+			false, CF_Always,
+			true, CF_NotEqual, SO_Keep, SO_Keep, SO_Keep,
+			false, CF_Always, SO_Keep, SO_Keep, SO_Keep,
+			0xff, 0xff
+		>::GetRHI() );
+	}
 
 	// Light Attenuation channel assignment:
 	//  R:     WholeSceneShadows, non SSS
@@ -1871,8 +1913,16 @@ void FProjectedShadowInfo::RenderProjection(int32 ViewIndex, const FViewInfo* Vi
 		RHIDrawIndexedPrimitiveUP( PT_TriangleList, 0, 8, 12, GCubeIndices, sizeof(uint16), FrustumVertices, sizeof(FVector4));
 	}
 
-	// Clear the stencil buffer to 0.
-	RHIClear(false,FColor(0,0,0),false,0,true,0, FIntRect());
+	if( bDepthBoundsTestEnabled )
+	{
+		// Disable depth bounds testing
+		RHIEnableDepthBoundsTest(false, 0, 1);
+	}
+	else
+	{
+		// Clear the stencil buffer to 0.
+		RHIClear(false,FColor(0,0,0),false,0,true,0, FIntRect());
+	}
 }
 
 
diff --git a/Engine/Source/Runtime/Windows/D3D11RHI/Private/Windows/WindowsD3D11Device.cpp b/Engine/Source/Runtime/Windows/D3D11RHI/Private/Windows/WindowsD3D11Device.cpp
--- a/Engine/Source/Runtime/Windows/D3D11RHI/Private/Windows/WindowsD3D11Device.cpp
+++ b/Engine/Source/Runtime/Windows/D3D11RHI/Private/Windows/WindowsD3D11Device.cpp
@@ -576,6 +576,13 @@ void FD3D11DynamicRHI::InitD3DDevice()
 			GRHISupportsAsyncTextureCreation = false;
 		}
 
+#if PLATFORM_DESKTOP
+		if( IsRHIDeviceNVIDIA() )
+		{
+			GSupportsDepthBoundsTest = true;
+		}
+#endif
+
 		SetupAfterDeviceCreation();
 
 		// Notify all initialized FRenderResources that there's a valid RHI device to create their RHI resources for now.
