* finished outstanding shadow bias refactoring, now it's using a single value
  will not pickup old bias as it works differently but if old bias was not altered it's no difference.

diff --git a/Engine/Source/Runtime/Engine/Classes/Components/LightComponent.h b/Engine/Source/Runtime/Engine/Classes/Components/LightComponent.h
--- a/Engine/Source/Runtime/Engine/Classes/Components/LightComponent.h
+++ b/Engine/Source/Runtime/Engine/Classes/Components/LightComponent.h
@@ -34,14 +34,6 @@ class ENGINE_API ULightComponent : public ULightComponentBase
 	UPROPERTY(EditAnywhere, BlueprintReadOnly, Category=Light, AdvancedDisplay, meta=(UIMin = "0.08", UIMax = "1.0"))
 	float MinRoughness;
 
-	/** 
-	 * Controls how accurate self shadowing of whole scene shadows from this light are.  
-	 * At close to 0, shadows will start far from their caster, and there won't be self shadowing artifacts.
-	 * At close to 1, shadows will start very close to their caster, but there will be many self shadowing artifacts.
-	 */
-	UPROPERTY(EditAnywhere, BlueprintReadOnly, Category=Light, AdvancedDisplay, meta=(UIMin = "0", UIMax = "1"))
-	float SelfShadowingAccuracy;
-
 	/** 
 	 * Controls how accurate self shadowing of whole scene shadows from this light are.  
 	 * At 0, shadows will start at the their caster surface, but there will be many self shadowing artifacts.
diff --git a/Engine/Source/Runtime/Engine/Private/DirectionalLightComponent.cpp b/Engine/Source/Runtime/Engine/Private/DirectionalLightComponent.cpp
--- a/Engine/Source/Runtime/Engine/Private/DirectionalLightComponent.cpp
+++ b/Engine/Source/Runtime/Engine/Private/DirectionalLightComponent.cpp
@@ -306,16 +306,6 @@ public:
 		return true;
 	}
 
-	virtual float GetDepthBiasScale() const 
-	{ 
-		return FMath::Lerp(2.0f, 0.0f, SelfShadowingAccuracy);
-	}
-
-	virtual float GetShadowTransitionScale() const 
-	{ 
-		return FMath::Lerp(0.0f, 2.0f, SelfShadowingAccuracy);
-	}
-
 private:
 
 	void UpdateLightShaftOverrideDirection_RenderThread(FVector NewLightShaftOverrideDirection)
@@ -606,7 +596,6 @@ void UDirectionalLightComponent::PostEditChangeProperty(FPropertyChangedEvent& P
 	CascadeDistributionExponent = FMath::Clamp(CascadeDistributionExponent, .1f, 10.0f);
 	CascadeTransitionFraction = FMath::Clamp(CascadeTransitionFraction, 0.0f, 0.3f);
 	ShadowDistanceFadeoutFraction = FMath::Clamp(ShadowDistanceFadeoutFraction, 0.0f, 1.0f);
-	SelfShadowingAccuracy = FMath::Clamp(SelfShadowingAccuracy, 0.01f, 1.0f);
 	// max range is larger than UI
 	ShadowBias = FMath::Clamp(ShadowBias, 0.0f, 10.0f);
 
diff --git a/Engine/Source/Runtime/Engine/Private/LightComponent.cpp b/Engine/Source/Runtime/Engine/Private/LightComponent.cpp
--- a/Engine/Source/Runtime/Engine/Private/LightComponent.cpp
+++ b/Engine/Source/Runtime/Engine/Private/LightComponent.cpp
@@ -144,7 +144,6 @@ FBoxSphereBounds ULightComponentBase::GetPlacementExtent() const
 FLightSceneProxy::FLightSceneProxy(const ULightComponent* InLightComponent)
 	: LightComponent(InLightComponent)
 	, IndirectLightingScale(InLightComponent->IndirectLightingIntensity)
-	, SelfShadowingAccuracy(InLightComponent->SelfShadowingAccuracy)
 	, ShadowBias(InLightComponent->ShadowBias)
 	, ShadowSharpen(InLightComponent->ShadowSharpen)
 	, MinRoughness(InLightComponent->MinRoughness)
@@ -242,7 +241,6 @@ ULightComponent::ULightComponent(const class FPostConstructInitializeProperties&
 	ShadowMapChannel = INDEX_NONE;
 	PreviewShadowMapChannel = INDEX_NONE;
 	IndirectLightingIntensity = 1.0f;
-	SelfShadowingAccuracy = 0.5f;
 	ShadowBias = 0.5f;
 	ShadowSharpen = 0.0f;
 	bUseIESBrightness = (GetLinkerUE4Version() < VER_UE4_LIGHTS_USE_IES_BRIGHTNESS_DEFAULT_CHANGED);
@@ -433,7 +431,6 @@ void ULightComponent::PostEditChangeProperty(FPropertyChangedEvent& PropertyChan
 		PropertyName != GET_MEMBER_NAME_STRING_CHECKED(ULightComponent, LightFunctionScale) &&
 		PropertyName != GET_MEMBER_NAME_STRING_CHECKED(ULightComponent, LightFunctionFadeDistance) &&
 		PropertyName != GET_MEMBER_NAME_STRING_CHECKED(ULightComponent, DisabledBrightness) &&
-		PropertyName != GET_MEMBER_NAME_STRING_CHECKED(ULightComponent, SelfShadowingAccuracy) &&
 		PropertyName != GET_MEMBER_NAME_STRING_CHECKED(ULightComponent, ShadowBias) &&
 		PropertyName != GET_MEMBER_NAME_STRING_CHECKED(ULightComponent, ShadowSharpen) &&
 		PropertyName != GET_MEMBER_NAME_STRING_CHECKED(ULightComponent, bEnableLightShaftBloom) &&
diff --git a/Engine/Source/Runtime/Engine/Private/PointLightSceneProxy.h b/Engine/Source/Runtime/Engine/Private/PointLightSceneProxy.h
--- a/Engine/Source/Runtime/Engine/Private/PointLightSceneProxy.h
+++ b/Engine/Source/Runtime/Engine/Private/PointLightSceneProxy.h
@@ -233,24 +233,6 @@ public:
 		return true;
 	}
 
-	virtual float GetDepthBiasScale() const 
-	{ 
-		return FMath::Lerp(2.0f, 0.0f, SelfShadowingAccuracy);
-	}
-
-	virtual float GetShadowTransitionScale() const 
-	{ 
-		if (SelfShadowingAccuracy < .5f)
-		{
-			return FMath::Lerp(0.01f, 1.0f, SelfShadowingAccuracy * 2);
-		}
-		else
-		{
-			// Center an accuracy of .5 at 1
-			return FMath::Lerp(-10.0f, 12.0f, SelfShadowingAccuracy);
-		}
-	}
-
 private:
 
 	/** Updates the light scene info's radius from the component. */
diff --git a/Engine/Source/Runtime/Engine/Public/SceneManagement.h b/Engine/Source/Runtime/Engine/Public/SceneManagement.h
--- a/Engine/Source/Runtime/Engine/Public/SceneManagement.h
+++ b/Engine/Source/Runtime/Engine/Public/SceneManagement.h
@@ -760,16 +760,6 @@ public:
 		return FVector2D::ZeroVector;
 	}
 
-	virtual float GetDepthBiasScale() const 
-	{ 
-		return 1;
-	}
-
-	virtual float GetShadowTransitionScale() const 
-	{ 
-		return 1;
-	}
-
 	/** Accesses parameters needed for rendering the light. */
 	virtual void GetParameters(FVector4& LightPositionAndInvRadius, FVector4& LightColorAndFalloffExponent, FVector& NormalizedLightDirection, FVector2D& SpotAngles, float& LightSourceRadius, float& LightSourceLength, float& LightMinRoughness) const {}
 
@@ -926,13 +916,6 @@ protected:
 	/** Scale for indirect lighting from this light.  When 0, indirect lighting is disabled. */
 	float IndirectLightingScale;
 
-	/** 
-	 * Controls how accurate self shadowing of whole scene shadows from this light are.  
-	 * At close to 0, shadows will start far from their caster, and there won't be self shadowing artifacts.
-	 * At close to 1, shadows will start very close to their caster, but there will be many self shadowing artifacts.
-	 */
-	float SelfShadowingAccuracy;
-
 	/** User setting from light component, 0:no bias, 0.5:reasonable, larger object might appear to float */
 	float ShadowBias;
 
diff --git a/Engine/Source/Runtime/Renderer/Private/ShadowRendering.cpp b/Engine/Source/Runtime/Renderer/Private/ShadowRendering.cpp
--- a/Engine/Source/Runtime/Renderer/Private/ShadowRendering.cpp
+++ b/Engine/Source/Runtime/Renderer/Private/ShadowRendering.cpp
@@ -2141,7 +2141,8 @@ void FProjectedShadowInfo::UpdateShaderDepthBias()
 	if (IsWholeScenePointLightShadow())
 	{
 		DepthBias = CVarPointLightShadowDepthBias.GetValueOnRenderThread() * 512.0f / FMath::Max(ResolutionX, ResolutionY);
-		DepthBias *= LightSceneInfo->Proxy->GetDepthBiasScale();
+		// * 2.0f to be compatible with the system we had before ShadowBias
+		DepthBias *= 2.0f * LightSceneInfo->Proxy->GetUserShadowBias();
 	}
 	else if (IsWholeSceneDirectionalShadow())
 	{
@@ -2180,7 +2181,8 @@ void FProjectedShadowInfo::UpdateShaderDepthBias()
 			// spot lights (old code, might need to be improved)
 			const float LightTypeDepthBias = CVarSpotLightShadowDepthBias.GetValueOnRenderThread();
 			DepthBias = LightTypeDepthBias * 512.0f / ((MaxSubjectZ - MinSubjectZ) * FMath::Max(ResolutionX, ResolutionY));
-			DepthBias *= LightSceneInfo->Proxy->GetDepthBiasScale();
+			// * 2.0f to be compatible with the system we had before ShadowBias
+			DepthBias *= 2.0f * LightSceneInfo->Proxy->GetUserShadowBias();
 		}
 	}
 
@@ -2195,7 +2197,8 @@ float FProjectedShadowInfo::ComputeTransitionSize() const
 	{
 		// todo: optimize
 		TransitionSize = bDirectionalLight ? (1.0f / CVarShadowTransitionScale.GetValueOnRenderThread()) : (1.0f / CVarSpotLightShadowTransitionScale.GetValueOnRenderThread());
-		TransitionSize /= LightSceneInfo->Proxy->GetShadowTransitionScale();
+		// * 2.0f to be compatible with the system we had before ShadowBias
+		TransitionSize *= 2.0f * LightSceneInfo->Proxy->GetUserShadowBias();
 	}
 	else if (IsWholeSceneDirectionalShadow())
 	{
@@ -2221,7 +2224,8 @@ float FProjectedShadowInfo::ComputeTransitionSize() const
 	{
 		// todo: optimize
 		TransitionSize = bDirectionalLight ? (1.0f / CVarShadowTransitionScale.GetValueOnRenderThread()) : (1.0f / CVarSpotLightShadowTransitionScale.GetValueOnRenderThread());
-		TransitionSize /= LightSceneInfo->Proxy->GetShadowTransitionScale();
+		// * 2.0f to be compatible with the system we had before ShadowBias
+		TransitionSize *= 2.0f * LightSceneInfo->Proxy->GetUserShadowBias();
 	}
 
 	return TransitionSize;
