Tidied up code with some Clamps.

diff --git a/Engine/Source/Runtime/Renderer/Private/ShadowRendering.cpp b/Engine/Source/Runtime/Renderer/Private/ShadowRendering.cpp
--- a/Engine/Source/Runtime/Renderer/Private/ShadowRendering.cpp
+++ b/Engine/Source/Runtime/Renderer/Private/ShadowRendering.cpp
@@ -1680,11 +1680,8 @@ void FProjectedShadowInfo::RenderProjection(int32 ViewIndex, const FViewInfo* Vi
 			float DepthNear = Near.Z / Near.W;
 			float DepthFar = Far.Z / Far.W;
 
-			if( DepthFar > 1 ) DepthFar = 1.f;
-			else if( DepthFar < 0 ) DepthFar = 0.f;
-
-			if( DepthNear > 1 ) DepthNear = 1.f;
-			else if( DepthNear < 0 ) DepthNear = 0.f;
+			DepthFar = FMath::Clamp( DepthFar, 0.0f, 1.0f );
+			DepthNear = FMath::Clamp( DepthNear, 0.0f, 1.0f );
 
 			if( DepthNear <= DepthFar )
 			{
