Distance field AO using a surface cache

diff --git a/Engine/Source/Runtime/Renderer/Private/DistanceFieldSurfaceCacheLighting.cpp b/Engine/Source/Runtime/Renderer/Private/DistanceFieldSurfaceCacheLighting.cpp
new file mode 100644
--- /dev/null
+++ b/Engine/Source/Runtime/Renderer/Private/DistanceFieldSurfaceCacheLighting.cpp
@@ -0,0 +1,3085 @@
+// Copyright 1998-2013 Epic Games, Inc. All Rights Reserved.
+
+/*=============================================================================
+	
+=============================================================================*/
+
+#include "RendererPrivate.h"
+#include "ScenePrivate.h"
+#include "UniformBuffer.h"
+#include "ShaderParameters.h"
+#include "PostProcessing.h"
+#include "SceneFilterRendering.h"
+#include "ScreenRendering.h"
+#include "DistanceFieldSurfaceCacheLighting.h"
+#include "DistanceFieldLightingShared.h"
+#include "PostProcessAmbientOcclusion.h"
+
+int32 GDistanceFieldAO = 1;
+FAutoConsoleVariableRef CVarDistanceFieldAO(
+	TEXT("r.DistanceFieldAO"),
+	GDistanceFieldAO,
+	TEXT("Whether the distance field AO feature is allowed, which is used to implement shadows of Movable sky lights from static meshes."),
+	ECVF_Cheat | ECVF_RenderThreadSafe
+	);
+
+int32 GAOMinLevel = 1;
+FAutoConsoleVariableRef CVarAOMinLevel(
+	TEXT("r.AOMinLevel"),
+	GAOMinLevel,
+	TEXT("Smallest downsample power of 4 to use for surface cache population.\n")
+	TEXT("The default is 1, which means every 8 full resolution pixels (BaseDownsampleFactor(2) * 4^1) will be checked for a valid interpolation from the cache or shaded.\n")
+	TEXT("Going down to 0 gives less aliasing, and removes the need for gap filling, but costs a lot."),
+	ECVF_Cheat | ECVF_RenderThreadSafe
+	);
+
+int32 GAOMaxLevel = GAOMaxSupportedLevel;
+FAutoConsoleVariableRef CVarAOMaxLevel(
+	TEXT("r.AOMaxLevel"),
+	GAOMaxLevel,
+	TEXT("Largest downsample power of 4 to use for surface cache population."),
+	ECVF_Cheat | ECVF_RenderThreadSafe
+	);
+
+int32 GAOReuseAcrossFrames = 1;
+FAutoConsoleVariableRef CVarAOReuseAcrossFrames(
+	TEXT("r.AOReuseAcrossFrames"),
+	GAOReuseAcrossFrames,
+	TEXT("Whether to allow reusing surface cache results across frames."),
+	ECVF_Cheat | ECVF_RenderThreadSafe
+	);
+
+float GAOTrimOldRecordsFraction = .1f;
+FAutoConsoleVariableRef CVarAOTrimOldRecordsFraction(
+	TEXT("r.AOTrimOldRecordsFraction"),
+	GAOTrimOldRecordsFraction,
+	TEXT("When r.AOReuseAcrossFrames is enabled, this is the fraction of the last frame's surface cache records that will not be reused.\n")
+	TEXT("Low settings provide better performance, while values closer to 1 give faster lighting updates when dynamic scene changes are happening."),
+	ECVF_Cheat | ECVF_RenderThreadSafe
+	);
+
+int32 GAOFillGaps = 1;
+FAutoConsoleVariableRef CVarAOFillGaps(
+	TEXT("r.AOFillGaps"),
+	GAOFillGaps,
+	TEXT("Whether to fill in pixels using a screen space filter that had no valid world space interpolation weight from surface cache samples.\n")
+	TEXT("This is needed whenever r.AOMinLevel is not 0."),
+	ECVF_Cheat | ECVF_RenderThreadSafe
+	);
+
+int32 GAOUseHistory = 1;
+FAutoConsoleVariableRef CVarAOUseHistory(
+	TEXT("r.AOUseHistory"),
+	GAOUseHistory,
+	TEXT("Whether to apply a temporal filter to the distance field AO, which reduces flickering but also adds trails when occluders are moving."),
+	ECVF_Cheat | ECVF_RenderThreadSafe
+	);
+
+float GAOHistoryWeight = .7f;
+FAutoConsoleVariableRef CVarAOHistoryWeight(
+	TEXT("r.AOHistoryWeight"),
+	GAOHistoryWeight,
+	TEXT("Amount of last frame's AO to lerp into the final result.  Higher values increase stability, lower values have less streaking under occluder movement."),
+	ECVF_Cheat | ECVF_RenderThreadSafe
+	);
+
+float GAORecordRadiusScale = .3f;
+FAutoConsoleVariableRef CVarAORecordRadiusScale(
+	TEXT("r.AORecordRadiusScale"),
+	GAORecordRadiusScale,
+	TEXT("Scale applied to the minimum occluder distance to produce the record radius.  This effectively controls how dense shading samples are."),
+	ECVF_Cheat | ECVF_RenderThreadSafe
+	);
+
+float GAOInterpolationRadiusScale = 1.5f;
+FAutoConsoleVariableRef CVarAOInterpolationRadiusScale(
+	TEXT("r.AOInterpolationRadiusScale"),
+	GAOInterpolationRadiusScale,
+	TEXT("Scale applied to record radii during the final interpolation pass.  Values larger than 1 result in world space smoothing."),
+	ECVF_Cheat | ECVF_RenderThreadSafe
+	);
+
+float GAOMinPointBehindPlaneAngle = 4;
+FAutoConsoleVariableRef CVarAOMinPointBehindPlaneAngle(
+	TEXT("r.AOMinPointBehindPlaneAngle"),
+	GAOMinPointBehindPlaneAngle,
+	TEXT("Minimum angle that a point can lie behind a record and still be considered valid.\n")
+	TEXT("This threshold helps reduce leaking that happens when interpolating records in front of the shading point, ignoring occluders in between."),
+	ECVF_Cheat | ECVF_RenderThreadSafe
+	);
+
+float GAOInterpolationMaxAngle = 10;
+FAutoConsoleVariableRef CVarAOInterpolationMaxAngle(
+	TEXT("r.AOInterpolationMaxAngle"),
+	GAOInterpolationMaxAngle,
+	TEXT("Largest angle allowed between the shading point's normal and a nearby record's normal."),
+	ECVF_Cheat | ECVF_RenderThreadSafe
+	);
+
+float GAOInterpolationAngleScale = 1.5f;
+FAutoConsoleVariableRef CVarAOInterpolationAngleScale(
+	TEXT("r.AOInterpolationAngleScale"),
+	GAOInterpolationAngleScale,
+	TEXT("Scale applied to angle error during the final interpolation pass.  Values larger than 1 result in smoothing."),
+	ECVF_Cheat | ECVF_RenderThreadSafe
+	);
+
+float GAOLargestSampleOffset = 400;
+FAutoConsoleVariableRef CVarAOLargestSampleOffset(
+	TEXT("r.AOLargestSampleOffset"),
+	GAOLargestSampleOffset,
+	TEXT("Distance from the shading point to place the furthest AO samples.\n")
+	TEXT("This is the biggest factor in how far an occluder can affect.  Larger values cost significantly more and reduce detail AO quality."),
+	ECVF_Cheat | ECVF_RenderThreadSafe
+	);
+
+float GAOStepExponentScale = .5f;
+FAutoConsoleVariableRef CVarAOStepExponentScale(
+	TEXT("r.AOStepExponentScale"),
+	GAOStepExponentScale,
+	TEXT("Exponent used to distribute AO samples along a cone direction."),
+	ECVF_Cheat | ECVF_RenderThreadSafe
+	);
+
+int32 GAOCPUFrustumCull = 1;
+FAutoConsoleVariableRef CVarAOCPUFrustumCull(
+	TEXT("r.AOCPUFrustumCull"),
+	GAOCPUFrustumCull,
+	TEXT("Whether to perform CPU frustum culling of occluder objects."),
+	ECVF_Cheat | ECVF_RenderThreadSafe
+	);
+
+float GAOMaxViewDistance = 10000;
+FAutoConsoleVariableRef CVarAOMaxViewDistance(
+	TEXT("r.AOMaxViewDistance"),
+	GAOMaxViewDistance,
+	TEXT("The maximum distance that AO will be computed at."),
+	ECVF_Cheat | ECVF_RenderThreadSafe
+	);
+
+int32 GAOScatterTileCulling = 1;
+FAutoConsoleVariableRef CVarAOScatterTileCulling(
+	TEXT("r.AOScatterTileCulling"),
+	GAOScatterTileCulling,
+	TEXT("Whether to use the rasterizer for binning occluder objects into screenspace tiles."),
+	ECVF_Cheat | ECVF_RenderThreadSafe
+	);
+
+int32 GAOComputeShaderNormalCalculation = 1;
+FAutoConsoleVariableRef CVarAOComputeShaderNormalCalculation(
+	TEXT("r.AOComputeShaderNormalCalculation"),
+	GAOComputeShaderNormalCalculation,
+	TEXT("Whether to use the compute shader version of the distance field normal computation."),
+	ECVF_Cheat | ECVF_RenderThreadSafe
+	);
+
+IMPLEMENT_UNIFORM_BUFFER_STRUCT(FAOSampleData2,TEXT("AOSamples2"));
+
+FIntPoint GetBufferSizeForAO()
+{
+	return FIntPoint(
+		GAODownsampleFactor * (GSceneRenderTargets.GetBufferSizeXY().X / GAODownsampleFactor),
+		GAODownsampleFactor * (GSceneRenderTargets.GetBufferSizeXY().Y / GAODownsampleFactor));
+}
+
+bool DoesPlatformSupportDistanceFieldAO(EShaderPlatform Platform)
+{
+	return Platform == SP_PCD3D_SM5;
+}
+
+// Cone half angle derived from each cone covering an equal solid angle
+float GAOConeHalfAngle = FMath::Acos(1 - 1.0f / (float)ARRAY_COUNT(SpacedVectors9));
+
+// Number of AO sample positions along each cone
+// Must match shader code
+uint32 GAONumConeSteps = 10;
+
+TGlobalResource<FCircleVertexBuffer> GCircleVertexBuffer;
+
+class FAOObjectData
+{
+public:
+
+	FVertexBufferRHIRef Bounds;
+	FVertexBufferRHIRef Data;
+	FVertexBufferRHIRef Data2;
+	FShaderResourceViewRHIRef BoundsSRV;
+	FShaderResourceViewRHIRef DataSRV;
+	FShaderResourceViewRHIRef Data2SRV;
+};
+
+class FAOObjectBuffers : public FRenderResource
+{
+public:
+
+	// In float4's
+	static int32 ObjectDataStride;
+	static int32 ObjectData2Stride;
+
+	int32 MaxObjects;
+	FAOObjectData ObjectData;
+
+	FAOObjectBuffers()
+	{
+		MaxObjects = 0;
+	}
+
+	virtual void InitDynamicRHI() 
+	{
+		if (MaxObjects > 0)
+		{
+			FRHIResourceCreateInfo CreateInfo;
+			ObjectData.Bounds = RHICreateVertexBuffer(MaxObjects * sizeof(FVector4), BUF_Volatile | BUF_ShaderResource, CreateInfo);
+			ObjectData.Data = RHICreateVertexBuffer(MaxObjects * sizeof(FVector4) * ObjectDataStride, BUF_Volatile | BUF_ShaderResource, CreateInfo);
+			ObjectData.Data2 = RHICreateVertexBuffer(MaxObjects * sizeof(FVector4) * ObjectData2Stride, BUF_Volatile | BUF_ShaderResource, CreateInfo);
+
+			ObjectData.BoundsSRV = RHICreateShaderResourceView(ObjectData.Bounds, sizeof(FVector4), PF_A32B32G32R32F);
+			ObjectData.DataSRV = RHICreateShaderResourceView(ObjectData.Data, sizeof(FVector4), PF_A32B32G32R32F);
+			ObjectData.Data2SRV = RHICreateShaderResourceView(ObjectData.Data2, sizeof(FVector4), PF_A32B32G32R32F);
+		}
+	}
+
+	virtual void ReleaseDynamicRHI()
+	{
+		ObjectData.Bounds.SafeRelease();
+		ObjectData.Data.SafeRelease();
+		ObjectData.Data2.SafeRelease();
+		ObjectData.BoundsSRV.SafeRelease(); 
+		ObjectData.DataSRV.SafeRelease(); 
+		ObjectData.Data2SRV.SafeRelease(); 
+	}
+};
+
+// Only 7 needed atm, pad to cache line size
+// Must match equivalent shader defines
+int32 FAOObjectBuffers::ObjectDataStride = 8;
+int32 FAOObjectBuffers::ObjectData2Stride = 4;
+
+TGlobalResource<FAOObjectBuffers> GAOObjectBuffers;
+
+void FSceneViewState::DestroyAOTileResources()
+{
+	if (AOTileIntersectionResources) 
+	{
+		ENQUEUE_UNIQUE_RENDER_COMMAND_ONEPARAMETER(
+			DeleteTileResources,
+			FTileIntersectionResources*, AOTileIntersectionResources, AOTileIntersectionResources,
+			{
+				AOTileIntersectionResources->ReleaseResource();
+				delete AOTileIntersectionResources;
+			}
+		);
+	}
+	AOTileIntersectionResources = NULL;
+}
+
+void OnClearIrradianceCache(UWorld* InWorld)
+{
+	FlushRenderingCommands();
+
+	FScene* Scene = (FScene*)InWorld->Scene;
+
+	if (Scene && Scene->SurfaceCacheResources)
+	{
+		Scene->SurfaceCacheResources->bClearedResources = false;
+	}
+}
+
+FAutoConsoleCommandWithWorld ClearCacheConsoleCommand(
+	TEXT("r.AOClearCache"),
+	TEXT(""),
+	FConsoleCommandWithWorldDelegate::CreateStatic(OnClearIrradianceCache)
+	);
+
+class FAOObjectBufferParameters
+{
+public:
+	void Bind(const FShaderParameterMap& ParameterMap)
+	{
+		ObjectBounds.Bind(ParameterMap, TEXT("ObjectBounds"));
+		ObjectData.Bind(ParameterMap, TEXT("ObjectData"));
+		ObjectData2.Bind(ParameterMap, TEXT("ObjectData2"));
+		NumObjects.Bind(ParameterMap, TEXT("NumObjects"));
+	}
+
+	template<typename TParamRef>
+	void Set(const TParamRef& ShaderRHI, int32 NumObjectsValue)
+	{
+		SetSRVParameter(ShaderRHI, ObjectBounds, GAOObjectBuffers.ObjectData.BoundsSRV);
+		SetSRVParameter(ShaderRHI, ObjectData, GAOObjectBuffers.ObjectData.DataSRV);
+		SetSRVParameter(ShaderRHI, ObjectData2, GAOObjectBuffers.ObjectData.Data2SRV);
+
+		SetShaderValue(ShaderRHI, NumObjects, NumObjectsValue);
+	}
+
+	friend FArchive& operator<<(FArchive& Ar, FAOObjectBufferParameters& P)
+	{
+		Ar << P.ObjectBounds << P.ObjectData << P.ObjectData2 << P.NumObjects;
+		return Ar;
+	}
+
+private:
+	FShaderResourceParameter ObjectBounds;
+	FShaderResourceParameter ObjectData;
+	FShaderResourceParameter ObjectData2;
+	FShaderParameter NumObjects;
+};
+
+
+/**  */
+class FBuildTileConesCS : public FGlobalShader
+{
+	DECLARE_SHADER_TYPE(FBuildTileConesCS,Global)
+public:
+
+	static bool ShouldCache(EShaderPlatform Platform)
+	{
+		return IsFeatureLevelSupported(Platform, ERHIFeatureLevel::SM5) && DoesPlatformSupportDistanceFieldAO(Platform);
+	}
+
+	static void ModifyCompilationEnvironment(EShaderPlatform Platform, FShaderCompilerEnvironment& OutEnvironment)
+	{
+		FGlobalShader::ModifyCompilationEnvironment(Platform,OutEnvironment);
+		OutEnvironment.SetDefine(TEXT("THREADGROUP_SIZEX"), GDistanceFieldAOTileSizeX);
+		OutEnvironment.SetDefine(TEXT("THREADGROUP_SIZEY"), GDistanceFieldAOTileSizeY);
+		OutEnvironment.SetDefine(TEXT("NUM_SAMPLES"), NumConeSampleDirections);
+		OutEnvironment.SetDefine(TEXT("DOWNSAMPLE_FACTOR"), GAODownsampleFactor);
+		OutEnvironment.SetDefine(TEXT("MAX_OBJECTS_PER_TILE"), GMaxNumObjectsPerTile);
+		
+		// To reduce shader compile time of compute shaders with shared memory, doesn't have an impact on generated code with current compiler (June 2010 DX SDK)
+		OutEnvironment.CompilerFlags.Add(CFLAG_StandardOptimization);
+	}
+
+	FBuildTileConesCS(const ShaderMetaType::CompiledShaderInitializerType& Initializer)
+		: FGlobalShader(Initializer)
+	{
+		DeferredParameters.Bind(Initializer.ParameterMap);
+		AOParameters.Bind(Initializer.ParameterMap);
+		TileConeAxisAndCos.Bind(Initializer.ParameterMap, TEXT("TileConeAxisAndCos"));
+		TileConeDepthRanges.Bind(Initializer.ParameterMap, TEXT("TileConeDepthRanges"));
+		TileHeadDataUnpacked.Bind(Initializer.ParameterMap, TEXT("TileHeadDataUnpacked"));
+		NumGroups.Bind(Initializer.ParameterMap, TEXT("NumGroups"));
+		ViewDimensionsParameter.Bind(Initializer.ParameterMap, TEXT("ViewDimensions"));
+	}
+
+	FBuildTileConesCS()
+	{
+	}
+	void SetParameters(const FSceneView& View, FScene* Scene, FVector2D NumGroupsValue)
+	{
+		FComputeShaderRHIParamRef ShaderRHI = GetComputeShader();
+
+		FGlobalShader::SetParameters(ShaderRHI, View);
+		DeferredParameters.Set(ShaderRHI, View);
+		AOParameters.Set(ShaderRHI);
+
+		FTileIntersectionResources* TileIntersectionResources = ((FSceneViewState*)View.State)->AOTileIntersectionResources;
+
+		TileConeAxisAndCos.SetBuffer(ShaderRHI, TileIntersectionResources->TileConeAxisAndCos);
+		TileConeDepthRanges.SetBuffer(ShaderRHI, TileIntersectionResources->TileConeDepthRanges);
+		TileHeadDataUnpacked.SetBuffer(ShaderRHI, TileIntersectionResources->TileHeadDataUnpacked);
+
+		SetShaderValue(ShaderRHI, ViewDimensionsParameter, View.ViewRect);
+
+		SetShaderValue(ShaderRHI, NumGroups, NumGroupsValue);
+	}
+
+	void UnsetParameters()
+	{
+		TileConeAxisAndCos.UnsetUAV(GetComputeShader());
+		TileConeDepthRanges.UnsetUAV(GetComputeShader());
+		TileHeadDataUnpacked.UnsetUAV(GetComputeShader());
+	}
+
+	virtual bool Serialize(FArchive& Ar)
+	{		
+		bool bShaderHasOutdatedParameters = FGlobalShader::Serialize(Ar);
+		Ar << DeferredParameters;
+		Ar << AOParameters;
+		Ar << TileConeAxisAndCos;
+		Ar << TileConeDepthRanges;
+		Ar << TileHeadDataUnpacked;
+		Ar << NumGroups;
+		Ar << ViewDimensionsParameter;
+		return bShaderHasOutdatedParameters;
+	}
+
+private:
+
+	FDeferredPixelShaderParameters DeferredParameters;
+	FAOParameters AOParameters;
+	FRWShaderParameter TileConeAxisAndCos;
+	FRWShaderParameter TileConeDepthRanges;
+	FRWShaderParameter TileHeadDataUnpacked;
+	FShaderParameter ViewDimensionsParameter;
+	FShaderParameter NumGroups;
+};
+
+IMPLEMENT_SHADER_TYPE(,FBuildTileConesCS,TEXT("DistanceFieldSurfaceCacheLightingCompute"),TEXT("BuildTileConesMain"),SF_Compute);
+
+
+/**  */
+class FObjectCullVS : public FGlobalShader
+{
+	DECLARE_SHADER_TYPE(FObjectCullVS,Global);
+public:
+
+	static bool ShouldCache(EShaderPlatform Platform) 
+	{ 
+		return IsFeatureLevelSupported(Platform, ERHIFeatureLevel::SM5) && DoesPlatformSupportDistanceFieldAO(Platform); 
+	}
+
+	static void ModifyCompilationEnvironment(EShaderPlatform Platform, FShaderCompilerEnvironment& OutEnvironment)
+	{
+		FGlobalShader::ModifyCompilationEnvironment(Platform,OutEnvironment);
+	}
+
+	FObjectCullVS(const ShaderMetaType::CompiledShaderInitializerType& Initializer):
+		FGlobalShader(Initializer)
+	{
+		ObjectParameters.Bind(Initializer.ParameterMap);
+		AOParameters.Bind(Initializer.ParameterMap);
+		ConservativeRadiusScale.Bind(Initializer.ParameterMap, TEXT("ConservativeRadiusScale"));
+	}
+
+	FObjectCullVS() {}
+
+	void SetParameters(const FSceneView& View, int32 NumObjects)
+	{
+		const FVertexShaderRHIParamRef ShaderRHI = GetVertexShader();
+		FGlobalShader::SetParameters(ShaderRHI, View);
+		
+		ObjectParameters.Set(ShaderRHI, NumObjects);
+		AOParameters.Set(ShaderRHI);
+
+		const int32 NumRings = StencilingGeometry::GLowPolyStencilSphereVertexBuffer.GetNumRings();
+		const float RadiansPerRingSegment = PI / (float)NumRings;
+
+		// Boost the effective radius so that the edges of the sphere approximation lie on the sphere, instead of the vertices
+		const float ConservativeRadiusScaleValue = 1.0f / FMath::Cos(RadiansPerRingSegment);
+
+		SetShaderValue(ShaderRHI, ConservativeRadiusScale, ConservativeRadiusScaleValue);
+	}
+
+	virtual bool Serialize(FArchive& Ar)
+	{
+		bool bShaderHasOutdatedParameters = FGlobalShader::Serialize(Ar);
+		Ar << ObjectParameters;
+		Ar << AOParameters;
+		Ar << ConservativeRadiusScale;
+		return bShaderHasOutdatedParameters;
+	}
+
+private:
+	FAOObjectBufferParameters ObjectParameters;
+	FAOParameters AOParameters;
+	FShaderParameter ConservativeRadiusScale;
+};
+
+IMPLEMENT_SHADER_TYPE(,FObjectCullVS,TEXT("DistanceFieldSurfaceCacheLighting"),TEXT("ObjectCullVS"),SF_Vertex);
+
+class FObjectCullPS : public FGlobalShader
+{
+	DECLARE_SHADER_TYPE(FObjectCullPS, Global);
+public:
+
+	static bool ShouldCache(EShaderPlatform Platform)
+	{
+		return IsFeatureLevelSupported(Platform, ERHIFeatureLevel::SM5) && DoesPlatformSupportDistanceFieldAO(Platform);
+	}
+
+	static void ModifyCompilationEnvironment(EShaderPlatform Platform, FShaderCompilerEnvironment& OutEnvironment)
+	{
+		FGlobalShader::ModifyCompilationEnvironment(Platform,OutEnvironment);
+		
+		OutEnvironment.SetDefine(TEXT("DOWNSAMPLE_FACTOR"), GAODownsampleFactor);
+		OutEnvironment.SetDefine(TEXT("MAX_OBJECTS_PER_TILE"), GMaxNumObjectsPerTile);
+	}
+
+	/** Default constructor. */
+	FObjectCullPS() {}
+
+	/** Initialization constructor. */
+	FObjectCullPS(const ShaderMetaType::CompiledShaderInitializerType& Initializer)
+		: FGlobalShader(Initializer)
+	{
+		AOParameters.Bind(Initializer.ParameterMap);
+		TileHeadDataUnpacked.Bind(Initializer.ParameterMap, TEXT("TileHeadDataUnpacked"));
+		TileArrayData.Bind(Initializer.ParameterMap, TEXT("TileArrayData"));
+		TileConeAxisAndCos.Bind(Initializer.ParameterMap, TEXT("TileConeAxisAndCos"));
+		TileConeDepthRanges.Bind(Initializer.ParameterMap, TEXT("TileConeDepthRanges"));
+		NumGroups.Bind(Initializer.ParameterMap, TEXT("NumGroups"));
+	}
+
+	void SetParameters(const FSceneView& View, FVector2D NumGroupsValue)
+	{
+		const FPixelShaderRHIParamRef ShaderRHI = GetPixelShader();
+		FGlobalShader::SetParameters(ShaderRHI, View);
+		AOParameters.Set(ShaderRHI);
+
+		FTileIntersectionResources* TileIntersectionResources = ((FSceneViewState*)View.State)->AOTileIntersectionResources;
+
+		SetSRVParameter(ShaderRHI, TileConeAxisAndCos, TileIntersectionResources->TileConeAxisAndCos.SRV);
+		SetSRVParameter(ShaderRHI, TileConeDepthRanges, TileIntersectionResources->TileConeDepthRanges.SRV);
+
+		SetShaderValue(ShaderRHI, NumGroups, NumGroupsValue);
+	}
+
+	void GetUAVs(const FSceneView& View, TArray<FUnorderedAccessViewRHIParamRef>& UAVs)
+	{
+		FTileIntersectionResources* TileIntersectionResources = ((FSceneViewState*)View.State)->AOTileIntersectionResources;
+
+		int32 MaxIndex = FMath::Max(TileHeadDataUnpacked.GetUAVIndex(), TileArrayData.GetUAVIndex());
+		UAVs.AddZeroed(MaxIndex + 1);
+
+		if (TileHeadDataUnpacked.IsBound())
+		{
+			UAVs[TileHeadDataUnpacked.GetUAVIndex()] = TileIntersectionResources->TileHeadDataUnpacked.UAV;
+		}
+
+		if (TileArrayData.IsBound())
+		{
+			UAVs[TileArrayData.GetUAVIndex()] = TileIntersectionResources->TileArrayData.UAV;
+		}
+
+		check(UAVs.Num() > 0);
+	}
+
+	// FShader interface.
+	virtual bool Serialize(FArchive& Ar)
+	{
+		bool bShaderHasOutdatedParameters = FGlobalShader::Serialize(Ar);
+		Ar << AOParameters;
+		Ar << TileHeadDataUnpacked;
+		Ar << TileArrayData;
+		Ar << TileConeAxisAndCos;
+		Ar << TileConeDepthRanges;
+		Ar << NumGroups;
+		return bShaderHasOutdatedParameters;
+	}
+
+private:
+	FAOParameters AOParameters;
+	FRWShaderParameter TileHeadDataUnpacked;
+	FRWShaderParameter TileArrayData;
+	FShaderResourceParameter TileConeAxisAndCos;
+	FShaderResourceParameter TileConeDepthRanges;
+	FShaderParameter NumGroups;
+};
+
+IMPLEMENT_SHADER_TYPE(,FObjectCullPS,TEXT("DistanceFieldSurfaceCacheLighting"),TEXT("ObjectCullPS"),SF_Pixel);
+
+
+/**  */
+class FDistanceFieldBuildTileListCS : public FGlobalShader
+{
+	DECLARE_SHADER_TYPE(FDistanceFieldBuildTileListCS,Global)
+public:
+
+	static bool ShouldCache(EShaderPlatform Platform)
+	{
+		return IsFeatureLevelSupported(Platform, ERHIFeatureLevel::SM5) && DoesPlatformSupportDistanceFieldAO(Platform);
+	}
+
+	static void ModifyCompilationEnvironment(EShaderPlatform Platform, FShaderCompilerEnvironment& OutEnvironment)
+	{
+		FGlobalShader::ModifyCompilationEnvironment(Platform,OutEnvironment);
+		OutEnvironment.SetDefine(TEXT("THREADGROUP_SIZEX"), GDistanceFieldAOTileSizeX);
+		OutEnvironment.SetDefine(TEXT("THREADGROUP_SIZEY"), GDistanceFieldAOTileSizeY);
+		OutEnvironment.SetDefine(TEXT("NUM_SAMPLES"), NumConeSampleDirections);
+		OutEnvironment.SetDefine(TEXT("DOWNSAMPLE_FACTOR"), GAODownsampleFactor);
+		OutEnvironment.SetDefine(TEXT("MAX_OBJECTS_PER_TILE"), GMaxNumObjectsPerTile);
+		// To reduce shader compile time of compute shaders with shared memory, doesn't have an impact on generated code with current compiler (June 2010 DX SDK)
+		OutEnvironment.CompilerFlags.Add(CFLAG_StandardOptimization);
+	}
+
+	FDistanceFieldBuildTileListCS(const ShaderMetaType::CompiledShaderInitializerType& Initializer)
+		: FGlobalShader(Initializer)
+	{
+		DeferredParameters.Bind(Initializer.ParameterMap);
+		ObjectParameters.Bind(Initializer.ParameterMap);
+		AOParameters.Bind(Initializer.ParameterMap);
+		TileHeadData.Bind(Initializer.ParameterMap, TEXT("TileHeadData"));
+		TileArrayData.Bind(Initializer.ParameterMap, TEXT("TileArrayData"));
+		TileArrayNextAllocation.Bind(Initializer.ParameterMap, TEXT("TileArrayNextAllocation"));
+		NumGroups.Bind(Initializer.ParameterMap, TEXT("NumGroups"));
+		ViewDimensionsParameter.Bind(Initializer.ParameterMap, TEXT("ViewDimensions"));
+	}
+
+	FDistanceFieldBuildTileListCS()
+	{
+	}
+	void SetParameters(const FSceneView& View, FScene* Scene, FVector2D NumGroupsValue, int32 NumObjects)
+	{
+		FComputeShaderRHIParamRef ShaderRHI = GetComputeShader();
+
+		FGlobalShader::SetParameters(ShaderRHI, View);
+		DeferredParameters.Set(ShaderRHI, View);
+		ObjectParameters.Set(ShaderRHI, NumObjects);
+		AOParameters.Set(ShaderRHI);
+
+		FTileIntersectionResources* TileIntersectionResources = ((FSceneViewState*)View.State)->AOTileIntersectionResources;
+
+		TileHeadData.SetBuffer(ShaderRHI, TileIntersectionResources->TileHeadData);
+		TileArrayData.SetBuffer(ShaderRHI, TileIntersectionResources->TileArrayData);
+		TileArrayNextAllocation.SetBuffer(ShaderRHI, TileIntersectionResources->TileArrayNextAllocation);
+
+		SetShaderValue(ShaderRHI, ViewDimensionsParameter, View.ViewRect);
+
+		SetShaderValue(ShaderRHI, NumGroups, NumGroupsValue);
+	}
+
+	void UnsetParameters()
+	{
+		TileHeadData.UnsetUAV(GetComputeShader());
+		TileArrayData.UnsetUAV(GetComputeShader());
+		TileArrayNextAllocation.UnsetUAV(GetComputeShader());
+	}
+
+	virtual bool Serialize(FArchive& Ar)
+	{		
+		bool bShaderHasOutdatedParameters = FGlobalShader::Serialize(Ar);
+		Ar << DeferredParameters;
+		Ar << ObjectParameters;
+		Ar << AOParameters;
+		Ar << TileHeadData;
+		Ar << TileArrayData;
+		Ar << TileArrayNextAllocation;
+		Ar << NumGroups;
+		Ar << ViewDimensionsParameter;
+		return bShaderHasOutdatedParameters;
+	}
+
+private:
+
+	FDeferredPixelShaderParameters DeferredParameters;
+	FAOObjectBufferParameters ObjectParameters;
+	FAOParameters AOParameters;
+	FRWShaderParameter TileHeadData;
+	FRWShaderParameter TileArrayData;
+	FRWShaderParameter TileArrayNextAllocation;
+	FShaderParameter ViewDimensionsParameter;
+	FShaderParameter NumGroups;
+};
+
+IMPLEMENT_SHADER_TYPE(,FDistanceFieldBuildTileListCS,TEXT("DistanceFieldSurfaceCacheLightingCompute"),TEXT("DistanceFieldAOBuildTileListMain"),SF_Compute);
+
+class FAOLevelParameters
+{
+public:
+	void Bind(const FShaderParameterMap& ParameterMap)
+	{
+		CurrentLevelDownsampleFactor.Bind(ParameterMap, TEXT("CurrentLevelDownsampleFactor"));
+		AOBufferSize.Bind(ParameterMap, TEXT("AOBufferSize"));
+		AOScreenPositionScaleBias.Bind(ParameterMap, TEXT("AOScreenPositionScaleBias"));
+		DownsampleFactorToBaseLevel.Bind(ParameterMap, TEXT("DownsampleFactorToBaseLevel"));
+		BaseLevelTexelSize.Bind(ParameterMap, TEXT("BaseLevelTexelSize"));
+	}
+
+	template<typename TParamRef>
+	void Set(const TParamRef& ShaderRHI, const FSceneView& View, int32 CurrentLevelDownsampleFactorValue)
+	{
+		SetShaderValue(ShaderRHI, CurrentLevelDownsampleFactor, CurrentLevelDownsampleFactorValue);
+
+		// Round up, to match render target allocation
+		const FVector2D AOBufferSizeValue = FIntPoint::DivideAndRoundUp(GSceneRenderTargets.GetBufferSizeXY(), CurrentLevelDownsampleFactorValue);
+		SetShaderValue(ShaderRHI, AOBufferSize, AOBufferSizeValue);
+
+		const FVector2D InvAOBufferSize(1.0f / AOBufferSizeValue.X, 1.0f / AOBufferSizeValue.Y);
+
+		// Round down, to match how the view was set
+		const FIntPoint AOViewRectMin = FIntPoint(View.ViewRect.Min.X / CurrentLevelDownsampleFactorValue, View.ViewRect.Min.Y / CurrentLevelDownsampleFactorValue);
+		const FIntPoint AOViewRectSize = FIntPoint::DivideAndRoundUp(View.ViewRect.Size(), CurrentLevelDownsampleFactorValue);
+
+		/*
+		const FVector4 AOScreenPositionScaleBiasValue(
+			AOViewRectSize.X * InvAOBufferSize.X / +2.0f,
+			AOViewRectSize.Y * InvAOBufferSize.Y / (-2.0f * GProjectionSignY),
+			(AOViewRectSize.X / 2.0f + GPixelCenterOffset + AOViewRectMin.X) * InvAOBufferSize.X,
+			(AOViewRectSize.Y / 2.0f + GPixelCenterOffset + AOViewRectMin.Y) * InvAOBufferSize.Y
+			);*/
+
+		const FIntPoint RoundedDownViewDimensions = AOViewRectSize * CurrentLevelDownsampleFactorValue;
+
+		const FIntPoint BufferSize = GSceneRenderTargets.GetBufferSizeXY();
+		const float InvBufferSizeX = 1.0f / BufferSize.X;
+		const float InvBufferSizeY = 1.0f / BufferSize.Y;
+		const FVector4 ScreenPositionScaleBias(
+			RoundedDownViewDimensions.X * InvBufferSizeX / +2.0f,
+			RoundedDownViewDimensions.Y * InvBufferSizeY / (-2.0f * GProjectionSignY),
+			(RoundedDownViewDimensions.Y / 2.0f + GPixelCenterOffset + AOViewRectMin.Y) * InvBufferSizeY,
+			(RoundedDownViewDimensions.X / 2.0f + GPixelCenterOffset + AOViewRectMin.X) * InvBufferSizeX
+			);
+
+		SetShaderValue(ShaderRHI, AOScreenPositionScaleBias, ScreenPositionScaleBias);
+
+		SetShaderValue(ShaderRHI, DownsampleFactorToBaseLevel, CurrentLevelDownsampleFactorValue / GAODownsampleFactor);
+
+		const FIntPoint DownsampledBufferSize(GSceneRenderTargets.GetBufferSizeXY() / FIntPoint(GAODownsampleFactor, GAODownsampleFactor));
+		const FVector2D BaseLevelBufferSizeValue(1.0f / DownsampledBufferSize.X, 1.0f / DownsampledBufferSize.Y);
+		SetShaderValue(ShaderRHI, BaseLevelTexelSize, BaseLevelBufferSizeValue);
+	}
+
+	friend FArchive& operator<<(FArchive& Ar,FAOLevelParameters& P)
+	{
+		Ar << P.CurrentLevelDownsampleFactor << P.AOBufferSize << P.AOScreenPositionScaleBias << P.DownsampleFactorToBaseLevel << P.BaseLevelTexelSize;
+		return Ar;
+	}
+
+private:
+	FShaderParameter CurrentLevelDownsampleFactor;
+	FShaderParameter AOBufferSize;
+	FShaderParameter AOScreenPositionScaleBias;
+	FShaderParameter DownsampleFactorToBaseLevel;
+	FShaderParameter BaseLevelTexelSize;
+};
+
+class FComputeDistanceFieldNormalPS : public FGlobalShader
+{
+	DECLARE_SHADER_TYPE(FComputeDistanceFieldNormalPS, Global);
+public:
+
+	static bool ShouldCache(EShaderPlatform Platform)
+	{
+		return IsFeatureLevelSupported(Platform, ERHIFeatureLevel::SM5) && DoesPlatformSupportDistanceFieldAO(Platform);
+	}
+
+	static void ModifyCompilationEnvironment(EShaderPlatform Platform, FShaderCompilerEnvironment& OutEnvironment)
+	{
+		OutEnvironment.SetDefine(TEXT("DOWNSAMPLE_FACTOR"), GAODownsampleFactor);
+		OutEnvironment.SetDefine(TEXT("THREADGROUP_SIZEX"), GDistanceFieldAOTileSizeX);
+		OutEnvironment.SetDefine(TEXT("THREADGROUP_SIZEY"), GDistanceFieldAOTileSizeY);
+		OutEnvironment.SetDefine(TEXT("MAX_OBJECTS_PER_TILE"), GMaxNumObjectsPerTile);
+	}
+
+	/** Default constructor. */
+	FComputeDistanceFieldNormalPS() {}
+
+	/** Initialization constructor. */
+	FComputeDistanceFieldNormalPS(const ShaderMetaType::CompiledShaderInitializerType& Initializer)
+		: FGlobalShader(Initializer)
+	{
+		ObjectParameters.Bind(Initializer.ParameterMap);
+		DeferredParameters.Bind(Initializer.ParameterMap);
+		UVToTileHead.Bind(Initializer.ParameterMap, TEXT("UVToTileHead"));
+		TileHeadData.Bind(Initializer.ParameterMap, TEXT("TileHeadData"));
+		TileHeadDataUnpacked.Bind(Initializer.ParameterMap, TEXT("TileHeadDataUnpacked"));
+		TileArrayData.Bind(Initializer.ParameterMap, TEXT("TileArrayData"));
+		TileListGroupSize.Bind(Initializer.ParameterMap, TEXT("TileListGroupSize"));
+		AOParameters.Bind(Initializer.ParameterMap);
+	}
+
+	void SetParameters(const FSceneView& View, int32 NumObjects)
+	{
+		const FPixelShaderRHIParamRef ShaderRHI = GetPixelShader();
+
+		FGlobalShader::SetParameters(ShaderRHI, View);
+		ObjectParameters.Set(ShaderRHI, NumObjects);
+		AOParameters.Set(ShaderRHI);
+		DeferredParameters.Set(ShaderRHI, View);
+
+		SetShaderValue(ShaderRHI, UVToTileHead, FVector2D(View.ViewRect.Size().X / (float)(GDistanceFieldAOTileSizeX * GAODownsampleFactor), View.ViewRect.Size().Y / (float)(GDistanceFieldAOTileSizeY * GAODownsampleFactor)));
+	
+		FTileIntersectionResources* TileIntersectionResources = ((FSceneViewState*)View.State)->AOTileIntersectionResources;
+
+		SetSRVParameter(ShaderRHI, TileHeadData, TileIntersectionResources->TileHeadData.SRV);
+		SetSRVParameter(ShaderRHI, TileHeadDataUnpacked, TileIntersectionResources->TileHeadDataUnpacked.SRV);
+		SetSRVParameter(ShaderRHI, TileArrayData, TileIntersectionResources->TileArrayData.SRV);
+
+		SetShaderValue(ShaderRHI, TileListGroupSize, TileIntersectionResources->TileDimensions);
+	}
+	// FShader interface.
+	virtual bool Serialize(FArchive& Ar)
+	{
+		bool bShaderHasOutdatedParameters = FGlobalShader::Serialize(Ar);
+		Ar << ObjectParameters;
+		Ar << UVToTileHead;
+		Ar << DeferredParameters;
+		Ar << TileHeadData;
+		Ar << TileHeadDataUnpacked;
+		Ar << TileArrayData;
+		Ar << TileListGroupSize;
+		Ar << AOParameters;
+		return bShaderHasOutdatedParameters;
+	}
+
+private:
+
+	FAOObjectBufferParameters ObjectParameters;
+	FShaderParameter UVToTileHead;
+	FDeferredPixelShaderParameters DeferredParameters;
+	FShaderResourceParameter TileHeadData;
+	FShaderResourceParameter TileHeadDataUnpacked;
+	FShaderResourceParameter TileArrayData;
+	FShaderParameter TileListGroupSize;
+	FAOParameters AOParameters;
+};
+
+IMPLEMENT_SHADER_TYPE(,FComputeDistanceFieldNormalPS,TEXT("DistanceFieldSurfaceCacheLighting"),TEXT("ComputeDistanceFieldNormalPS"),SF_Pixel);
+
+
+class FComputeDistanceFieldNormalCS : public FGlobalShader
+{
+	DECLARE_SHADER_TYPE(FComputeDistanceFieldNormalCS, Global);
+public:
+
+	static bool ShouldCache(EShaderPlatform Platform)
+	{
+		return IsFeatureLevelSupported(Platform, ERHIFeatureLevel::SM5) && DoesPlatformSupportDistanceFieldAO(Platform);
+	}
+
+	static void ModifyCompilationEnvironment(EShaderPlatform Platform, FShaderCompilerEnvironment& OutEnvironment)
+	{
+		OutEnvironment.SetDefine(TEXT("DOWNSAMPLE_FACTOR"), GAODownsampleFactor);
+		OutEnvironment.SetDefine(TEXT("THREADGROUP_SIZEX"), GDistanceFieldAOTileSizeX);
+		OutEnvironment.SetDefine(TEXT("THREADGROUP_SIZEY"), GDistanceFieldAOTileSizeY);
+		OutEnvironment.SetDefine(TEXT("MAX_OBJECTS_PER_TILE"), GMaxNumObjectsPerTile);
+	}
+
+	/** Default constructor. */
+	FComputeDistanceFieldNormalCS() {}
+
+	/** Initialization constructor. */
+	FComputeDistanceFieldNormalCS(const ShaderMetaType::CompiledShaderInitializerType& Initializer)
+		: FGlobalShader(Initializer)
+	{
+		DistanceFieldNormal.Bind(Initializer.ParameterMap, TEXT("DistanceFieldNormal"));
+		ObjectParameters.Bind(Initializer.ParameterMap);
+		DeferredParameters.Bind(Initializer.ParameterMap);
+		UVToTileHead.Bind(Initializer.ParameterMap, TEXT("UVToTileHead"));
+		TileHeadData.Bind(Initializer.ParameterMap, TEXT("TileHeadData"));
+		TileHeadDataUnpacked.Bind(Initializer.ParameterMap, TEXT("TileHeadDataUnpacked"));
+		TileArrayData.Bind(Initializer.ParameterMap, TEXT("TileArrayData"));
+		TileListGroupSize.Bind(Initializer.ParameterMap, TEXT("TileListGroupSize"));
+		AOParameters.Bind(Initializer.ParameterMap);
+	}
+
+	void SetParameters(const FSceneView& View, FSceneRenderTargetItem& DistanceFieldNormalValue, int32 NumObjects)
+	{
+		const FComputeShaderRHIParamRef ShaderRHI = GetComputeShader();
+
+		FGlobalShader::SetParameters(ShaderRHI, View);
+
+		DistanceFieldNormal.SetTexture(ShaderRHI, DistanceFieldNormalValue.ShaderResourceTexture, DistanceFieldNormalValue.UAV);
+
+		ObjectParameters.Set(ShaderRHI, NumObjects);
+		AOParameters.Set(ShaderRHI);
+		DeferredParameters.Set(ShaderRHI, View);
+
+		SetShaderValue(ShaderRHI, UVToTileHead, FVector2D(View.ViewRect.Size().X / (float)(GDistanceFieldAOTileSizeX * GAODownsampleFactor), View.ViewRect.Size().Y / (float)(GDistanceFieldAOTileSizeY * GAODownsampleFactor)));
+	
+		FTileIntersectionResources* TileIntersectionResources = ((FSceneViewState*)View.State)->AOTileIntersectionResources;
+
+		SetSRVParameter(ShaderRHI, TileHeadData, TileIntersectionResources->TileHeadData.SRV);
+		SetSRVParameter(ShaderRHI, TileHeadDataUnpacked, TileIntersectionResources->TileHeadDataUnpacked.SRV);
+		SetSRVParameter(ShaderRHI, TileArrayData, TileIntersectionResources->TileArrayData.SRV);
+
+		SetShaderValue(ShaderRHI, TileListGroupSize, TileIntersectionResources->TileDimensions);
+	}
+
+	void UnsetParameters()
+	{
+		DistanceFieldNormal.UnsetUAV(GetComputeShader());
+	}
+
+	// FShader interface.
+	virtual bool Serialize(FArchive& Ar)
+	{
+		bool bShaderHasOutdatedParameters = FGlobalShader::Serialize(Ar);
+		Ar << DistanceFieldNormal;
+		Ar << ObjectParameters;
+		Ar << UVToTileHead;
+		Ar << DeferredParameters;
+		Ar << TileHeadData;
+		Ar << TileHeadDataUnpacked;
+		Ar << TileArrayData;
+		Ar << TileListGroupSize;
+		Ar << AOParameters;
+		return bShaderHasOutdatedParameters;
+	}
+
+private:
+
+	FRWShaderParameter DistanceFieldNormal;
+	FAOObjectBufferParameters ObjectParameters;
+	FShaderParameter UVToTileHead;
+	FDeferredPixelShaderParameters DeferredParameters;
+	FShaderResourceParameter TileHeadData;
+	FShaderResourceParameter TileHeadDataUnpacked;
+	FShaderResourceParameter TileArrayData;
+	FShaderParameter TileListGroupSize;
+	FAOParameters AOParameters;
+};
+
+IMPLEMENT_SHADER_TYPE(,FComputeDistanceFieldNormalCS,TEXT("DistanceFieldSurfaceCacheLighting"),TEXT("ComputeDistanceFieldNormalCS"),SF_Compute);
+
+void ComputeDistanceFieldNormal(const TArray<FViewInfo>& Views, FSceneRenderTargetItem& DistanceFieldNormal, int32 NumObjects)
+{
+	if (GAOComputeShaderNormalCalculation)
+	{
+		RHISetRenderTarget(NULL, NULL);
+
+		for (int32 ViewIndex = 0; ViewIndex < Views.Num(); ViewIndex++)
+		{
+			const FViewInfo& View = Views[ViewIndex];
+
+			uint32 GroupSizeX = FMath::DivideAndRoundUp(View.ViewRect.Size().X / GAODownsampleFactor, GDistanceFieldAOTileSizeX);
+			uint32 GroupSizeY = FMath::DivideAndRoundUp(View.ViewRect.Size().Y / GAODownsampleFactor, GDistanceFieldAOTileSizeY);
+
+			{
+				SCOPED_DRAW_EVENT(ComputeNormalCS, DEC_SCENE_ITEMS);
+				TShaderMapRef<FComputeDistanceFieldNormalCS> ComputeShader(GetGlobalShaderMap());
+
+				RHISetComputeShader(ComputeShader->GetComputeShader());	
+				ComputeShader->SetParameters(View, DistanceFieldNormal, NumObjects);
+				DispatchComputeShader(*ComputeShader, GroupSizeX, GroupSizeY, 1);
+
+				ComputeShader->UnsetParameters();
+			}
+		}
+	}
+	else
+	{
+		RHISetRenderTarget(DistanceFieldNormal.TargetableTexture, NULL);
+
+		for (int32 ViewIndex = 0; ViewIndex < Views.Num(); ViewIndex++)
+		{
+			const FViewInfo& View = Views[ViewIndex];
+
+			SCOPED_DRAW_EVENT(ComputeNormal, DEC_SCENE_ITEMS);
+
+			RHISetViewport(View.ViewRect.Min.X / GAODownsampleFactor, View.ViewRect.Min.Y / GAODownsampleFactor, 0.0f, View.ViewRect.Max.X / GAODownsampleFactor, View.ViewRect.Max.Y / GAODownsampleFactor, 1.0f);
+			RHISetRasterizerState(TStaticRasterizerState<FM_Solid,CM_None>::GetRHI());
+			RHISetDepthStencilState(TStaticDepthStencilState<false,CF_Always>::GetRHI());
+			RHISetBlendState(TStaticBlendState<>::GetRHI());
+
+			TShaderMapRef<FPostProcessVS> VertexShader( GetGlobalShaderMap() );
+			TShaderMapRef<FComputeDistanceFieldNormalPS> PixelShader( GetGlobalShaderMap() );
+
+			static FGlobalBoundShaderState BoundShaderState;
+			SetGlobalBoundShaderState(BoundShaderState, GFilterVertexDeclaration.VertexDeclarationRHI, *VertexShader, *PixelShader);
+
+			PixelShader->SetParameters(View, NumObjects);
+
+			DrawRectangle( 
+				0, 0, 
+				View.ViewRect.Width() / GAODownsampleFactor, View.ViewRect.Height() / GAODownsampleFactor,
+				View.ViewRect.Min.X, View.ViewRect.Min.Y, 
+				View.ViewRect.Width(), View.ViewRect.Height(),
+				FIntPoint(View.ViewRect.Width() / GAODownsampleFactor, View.ViewRect.Height() / GAODownsampleFactor),
+				GSceneRenderTargets.GetBufferSizeXY(),
+				*VertexShader);
+		}
+	}
+}
+
+class FSetupCopyIndirectArgumentsCS : public FGlobalShader
+{
+	DECLARE_SHADER_TYPE(FSetupCopyIndirectArgumentsCS,Global)
+public:
+
+	static bool ShouldCache(EShaderPlatform Platform)
+	{
+		return IsFeatureLevelSupported(Platform, ERHIFeatureLevel::SM5) && DoesPlatformSupportDistanceFieldAO(Platform);
+	}
+
+	static void ModifyCompilationEnvironment(EShaderPlatform Platform, FShaderCompilerEnvironment& OutEnvironment)
+	{
+		FGlobalShader::ModifyCompilationEnvironment(Platform,OutEnvironment);
+		OutEnvironment.CompilerFlags.Add(CFLAG_StandardOptimization);
+	}
+
+	FSetupCopyIndirectArgumentsCS(const ShaderMetaType::CompiledShaderInitializerType& Initializer)
+		: FGlobalShader(Initializer)
+	{
+		DrawParameters.Bind(Initializer.ParameterMap, TEXT("DrawParameters"));
+		DispatchParameters.Bind(Initializer.ParameterMap, TEXT("DispatchParameters"));
+		TrimFraction.Bind(Initializer.ParameterMap, TEXT("TrimFraction"));
+	}
+
+	FSetupCopyIndirectArgumentsCS()
+	{
+	}
+
+	void SetParameters(const FSceneView& View, int32 DepthLevel)
+	{
+		FComputeShaderRHIParamRef ShaderRHI = GetComputeShader();
+		FGlobalShader::SetParameters(ShaderRHI, View);
+
+		const FScene* Scene = (const FScene*)View.Family->Scene;
+		FSurfaceCacheResources& SurfaceCacheResources = *Scene->SurfaceCacheResources;
+
+		SetSRVParameter(ShaderRHI, DrawParameters, SurfaceCacheResources.Level[DepthLevel]->ScatterDrawParameters.SRV);
+
+		DispatchParameters.SetBuffer(ShaderRHI, SurfaceCacheResources.DispatchParameters);
+
+		SetShaderValue(ShaderRHI, TrimFraction, GAOTrimOldRecordsFraction);
+	}
+
+	void UnsetParameters()
+	{
+		DispatchParameters.UnsetUAV(GetComputeShader());
+	}
+
+	virtual bool Serialize(FArchive& Ar)
+	{		
+		bool bShaderHasOutdatedParameters = FGlobalShader::Serialize(Ar);
+		Ar << DrawParameters;
+		Ar << DispatchParameters;
+		Ar << TrimFraction;
+		return bShaderHasOutdatedParameters;
+	}
+
+private:
+
+	FShaderResourceParameter DrawParameters;
+	FRWShaderParameter DispatchParameters;
+	FShaderParameter TrimFraction;
+};
+
+IMPLEMENT_SHADER_TYPE(,FSetupCopyIndirectArgumentsCS,TEXT("DistanceFieldSurfaceCacheLightingCompute"),TEXT("SetupCopyIndirectArgumentsCS"),SF_Compute);
+
+
+class FCopyIrradianceCacheSamplesCS : public FGlobalShader
+{
+	DECLARE_SHADER_TYPE(FCopyIrradianceCacheSamplesCS,Global)
+public:
+
+	static bool ShouldCache(EShaderPlatform Platform)
+	{
+		return IsFeatureLevelSupported(Platform, ERHIFeatureLevel::SM5) && DoesPlatformSupportDistanceFieldAO(Platform);
+	}
+
+	static void ModifyCompilationEnvironment(EShaderPlatform Platform, FShaderCompilerEnvironment& OutEnvironment)
+	{
+		FGlobalShader::ModifyCompilationEnvironment(Platform,OutEnvironment);
+		OutEnvironment.CompilerFlags.Add(CFLAG_StandardOptimization);
+	}
+
+	FCopyIrradianceCacheSamplesCS(const ShaderMetaType::CompiledShaderInitializerType& Initializer)
+		: FGlobalShader(Initializer)
+	{
+		IrradianceCachePositionRadius.Bind(Initializer.ParameterMap, TEXT("IrradianceCachePositionRadius"));
+		IrradianceCacheNormal.Bind(Initializer.ParameterMap, TEXT("IrradianceCacheNormal"));
+		OccluderRadius.Bind(Initializer.ParameterMap, TEXT("OccluderRadius"));
+		IrradianceCacheBentNormal.Bind(Initializer.ParameterMap, TEXT("IrradianceCacheBentNormal"));
+		IrradianceCacheTileCoordinate.Bind(Initializer.ParameterMap, TEXT("IrradianceCacheTileCoordinate"));
+		DrawParameters.Bind(Initializer.ParameterMap, TEXT("DrawParameters"));
+		CopyIrradianceCachePositionRadius.Bind(Initializer.ParameterMap, TEXT("CopyIrradianceCachePositionRadius"));
+		CopyIrradianceCacheNormal.Bind(Initializer.ParameterMap, TEXT("CopyIrradianceCacheNormal"));
+		CopyOccluderRadius.Bind(Initializer.ParameterMap, TEXT("CopyOccluderRadius"));
+		CopyIrradianceCacheBentNormal.Bind(Initializer.ParameterMap, TEXT("CopyIrradianceCacheBentNormal"));
+		CopyIrradianceCacheTileCoordinate.Bind(Initializer.ParameterMap, TEXT("CopyIrradianceCacheTileCoordinate"));
+		ScatterDrawParameters.Bind(Initializer.ParameterMap, TEXT("ScatterDrawParameters"));
+		TrimFraction.Bind(Initializer.ParameterMap, TEXT("TrimFraction"));
+	}
+
+	FCopyIrradianceCacheSamplesCS()
+	{
+	}
+
+	void SetParameters(const FSceneView& View, int32 DepthLevel)
+	{
+		FComputeShaderRHIParamRef ShaderRHI = GetComputeShader();
+		FGlobalShader::SetParameters(ShaderRHI, View);
+
+		const FScene* Scene = (const FScene*)View.Family->Scene;
+		FSurfaceCacheResources& SurfaceCacheResources = *Scene->SurfaceCacheResources;
+
+		SetSRVParameter(ShaderRHI, IrradianceCachePositionRadius, SurfaceCacheResources.Level[DepthLevel]->PositionAndRadius.SRV);
+		SetSRVParameter(ShaderRHI, IrradianceCacheNormal, SurfaceCacheResources.Level[DepthLevel]->Normal.SRV);
+		SetSRVParameter(ShaderRHI, OccluderRadius, SurfaceCacheResources.Level[DepthLevel]->OccluderRadius.SRV);
+		SetSRVParameter(ShaderRHI, IrradianceCacheBentNormal, SurfaceCacheResources.Level[DepthLevel]->BentNormal.SRV);
+		SetSRVParameter(ShaderRHI, IrradianceCacheTileCoordinate, SurfaceCacheResources.Level[DepthLevel]->TileCoordinate.SRV);
+		SetSRVParameter(ShaderRHI, DrawParameters, SurfaceCacheResources.Level[DepthLevel]->ScatterDrawParameters.SRV);
+
+		CopyIrradianceCachePositionRadius.SetBuffer(ShaderRHI, SurfaceCacheResources.TempResources->PositionAndRadius);
+		CopyIrradianceCacheNormal.SetBuffer(ShaderRHI, SurfaceCacheResources.TempResources->Normal);
+		CopyOccluderRadius.SetBuffer(ShaderRHI, SurfaceCacheResources.TempResources->OccluderRadius);
+		CopyIrradianceCacheBentNormal.SetBuffer(ShaderRHI, SurfaceCacheResources.TempResources->BentNormal);
+		CopyIrradianceCacheTileCoordinate.SetBuffer(ShaderRHI, SurfaceCacheResources.TempResources->TileCoordinate);
+		ScatterDrawParameters.SetBuffer(ShaderRHI, SurfaceCacheResources.TempResources->ScatterDrawParameters);
+
+		SetShaderValue(ShaderRHI, TrimFraction, GAOTrimOldRecordsFraction);
+	}
+
+	void UnsetParameters()
+	{
+		CopyIrradianceCachePositionRadius.UnsetUAV(GetComputeShader());
+		CopyIrradianceCacheNormal.UnsetUAV(GetComputeShader());
+		CopyOccluderRadius.UnsetUAV(GetComputeShader());
+		CopyIrradianceCacheBentNormal.UnsetUAV(GetComputeShader());
+		CopyIrradianceCacheTileCoordinate.UnsetUAV(GetComputeShader());
+		ScatterDrawParameters.UnsetUAV(GetComputeShader());
+	}
+
+	virtual bool Serialize(FArchive& Ar)
+	{		
+		bool bShaderHasOutdatedParameters = FGlobalShader::Serialize(Ar);
+		Ar << IrradianceCachePositionRadius;
+		Ar << IrradianceCacheNormal;
+		Ar << OccluderRadius;
+		Ar << IrradianceCacheBentNormal;
+		Ar << IrradianceCacheTileCoordinate;
+		Ar << DrawParameters;
+		Ar << CopyIrradianceCachePositionRadius;
+		Ar << CopyIrradianceCacheNormal;
+		Ar << CopyOccluderRadius;
+		Ar << CopyIrradianceCacheBentNormal;
+		Ar << CopyIrradianceCacheTileCoordinate;
+		Ar << ScatterDrawParameters;
+		Ar << TrimFraction;
+		return bShaderHasOutdatedParameters;
+	}
+
+private:
+
+	FShaderResourceParameter IrradianceCachePositionRadius;
+	FShaderResourceParameter IrradianceCacheNormal;
+	FShaderResourceParameter OccluderRadius;
+	FShaderResourceParameter IrradianceCacheBentNormal;
+	FShaderResourceParameter IrradianceCacheTileCoordinate;
+	FShaderResourceParameter DrawParameters;
+	FRWShaderParameter CopyIrradianceCachePositionRadius;
+	FRWShaderParameter CopyIrradianceCacheNormal;
+	FRWShaderParameter CopyOccluderRadius;
+	FRWShaderParameter CopyIrradianceCacheBentNormal;
+	FRWShaderParameter CopyIrradianceCacheTileCoordinate;
+	FRWShaderParameter ScatterDrawParameters;
+	FShaderParameter TrimFraction;
+};
+
+IMPLEMENT_SHADER_TYPE(,FCopyIrradianceCacheSamplesCS,TEXT("DistanceFieldSurfaceCacheLightingCompute"),TEXT("CopyIrradianceCacheSamplesCS"),SF_Compute);
+
+
+class FSaveStartIndexCS : public FGlobalShader
+{
+	DECLARE_SHADER_TYPE(FSaveStartIndexCS,Global)
+public:
+
+	static bool ShouldCache(EShaderPlatform Platform)
+	{
+		return IsFeatureLevelSupported(Platform, ERHIFeatureLevel::SM5) && DoesPlatformSupportDistanceFieldAO(Platform);
+	}
+
+	FSaveStartIndexCS(const ShaderMetaType::CompiledShaderInitializerType& Initializer)
+		: FGlobalShader(Initializer)
+	{
+		DrawParameters.Bind(Initializer.ParameterMap, TEXT("DrawParameters"));
+		SavedStartIndex.Bind(Initializer.ParameterMap, TEXT("SavedStartIndex"));
+	}
+
+	FSaveStartIndexCS()
+	{
+	}
+
+	void SetParameters(const FSceneView& View, int32 DepthLevel)
+	{
+		FComputeShaderRHIParamRef ShaderRHI = GetComputeShader();
+		FGlobalShader::SetParameters(ShaderRHI, View);
+
+		const FScene* Scene = (const FScene*)View.Family->Scene;
+		FSurfaceCacheResources& SurfaceCacheResources = *Scene->SurfaceCacheResources;
+
+		SetSRVParameter(ShaderRHI, DrawParameters, SurfaceCacheResources.Level[DepthLevel]->ScatterDrawParameters.SRV);
+		SavedStartIndex.SetBuffer(ShaderRHI, SurfaceCacheResources.Level[DepthLevel]->SavedStartIndex);
+	}
+
+	void UnsetParameters()
+	{
+		SavedStartIndex.UnsetUAV(GetComputeShader());
+	}
+
+	virtual bool Serialize(FArchive& Ar)
+	{		
+		bool bShaderHasOutdatedParameters = FGlobalShader::Serialize(Ar);
+		Ar << DrawParameters;
+		Ar << SavedStartIndex;
+		return bShaderHasOutdatedParameters;
+	}
+
+private:
+
+	FShaderResourceParameter DrawParameters;
+	FRWShaderParameter SavedStartIndex;
+};
+
+IMPLEMENT_SHADER_TYPE(,FSaveStartIndexCS,TEXT("DistanceFieldSurfaceCacheLightingCompute"),TEXT("SaveStartIndexCS"),SF_Compute);
+
+class FPopulateCacheCS : public FGlobalShader
+{
+	DECLARE_SHADER_TYPE(FPopulateCacheCS,Global)
+public:
+
+	static bool ShouldCache(EShaderPlatform Platform)
+	{
+		return IsFeatureLevelSupported(Platform, ERHIFeatureLevel::SM5) && DoesPlatformSupportDistanceFieldAO(Platform);
+	}
+
+	static void ModifyCompilationEnvironment(EShaderPlatform Platform, FShaderCompilerEnvironment& OutEnvironment)
+	{
+		FGlobalShader::ModifyCompilationEnvironment(Platform,OutEnvironment);
+		OutEnvironment.SetDefine(TEXT("THREADGROUP_SIZEX"), GDistanceFieldAOTileSizeX);
+		OutEnvironment.SetDefine(TEXT("THREADGROUP_SIZEY"), GDistanceFieldAOTileSizeY);
+		OutEnvironment.SetDefine(TEXT("NUM_SAMPLES"), NumConeSampleDirections);
+		OutEnvironment.SetDefine(TEXT("DOWNSAMPLE_FACTOR"), GAODownsampleFactor);
+
+		// To reduce shader compile time of compute shaders with shared memory, doesn't have an impact on generated code with current compiler (June 2010 DX SDK)
+		OutEnvironment.CompilerFlags.Add(CFLAG_StandardOptimization);
+	}
+
+	FPopulateCacheCS(const ShaderMetaType::CompiledShaderInitializerType& Initializer)
+		: FGlobalShader(Initializer)
+	{
+		DeferredParameters.Bind(Initializer.ParameterMap);
+		AOParameters.Bind(Initializer.ParameterMap);
+		IrradianceCachePositionRadius.Bind(Initializer.ParameterMap, TEXT("IrradianceCachePositionRadius"));
+		IrradianceCacheNormal.Bind(Initializer.ParameterMap, TEXT("IrradianceCacheNormal"));
+		IrradianceCacheTileCoordinate.Bind(Initializer.ParameterMap, TEXT("IrradianceCacheTileCoordinate"));
+		ScatterDrawParameters.Bind(Initializer.ParameterMap, TEXT("ScatterDrawParameters"));
+		ViewDimensionsParameter.Bind(Initializer.ParameterMap, TEXT("ViewDimensions"));
+		ThreadToCulledTile.Bind(Initializer.ParameterMap, TEXT("ThreadToCulledTile"));
+		TileListGroupSize.Bind(Initializer.ParameterMap, TEXT("TileListGroupSize"));
+		ConeHalfAngle.Bind(Initializer.ParameterMap, TEXT("ConeHalfAngle"));
+		BentNormalNormalizeFactor.Bind(Initializer.ParameterMap, TEXT("BentNormalNormalizeFactor"));
+		DistanceFieldAtlasTexelSize.Bind(Initializer.ParameterMap, TEXT("DistanceFieldAtlasTexelSize"));
+		DistanceFieldNormalTexture.Bind(Initializer.ParameterMap, TEXT("DistanceFieldNormalTexture"));
+		DistanceFieldNormalSampler.Bind(Initializer.ParameterMap, TEXT("DistanceFieldNormalSampler"));
+		IrradianceCacheSplatTexture.Bind(Initializer.ParameterMap, TEXT("IrradianceCacheSplatTexture"));
+		IrradianceCacheSplatSampler.Bind(Initializer.ParameterMap, TEXT("IrradianceCacheSplatSampler"));
+		AOLevelParameters.Bind(Initializer.ParameterMap);
+		RecordRadiusScale.Bind(Initializer.ParameterMap, TEXT("RecordRadiusScale"));
+	}
+
+	FPopulateCacheCS()
+	{
+	}
+
+	void SetParameters(const FSceneView& View, FSceneRenderTargetItem& DistanceFieldAOIrradianceCacheSplat, FSceneRenderTargetItem& DistanceFieldNormal, int32 DownsampleFactorValue, int32 DepthLevel, FIntPoint TileListGroupSizeValue)
+	{
+		FComputeShaderRHIParamRef ShaderRHI = GetComputeShader();
+		FGlobalShader::SetParameters(ShaderRHI, View);
+		DeferredParameters.Set(ShaderRHI, View);
+		AOParameters.Set(ShaderRHI);
+		AOLevelParameters.Set(ShaderRHI, View, DownsampleFactorValue);
+
+		const FScene* Scene = (const FScene*)View.Family->Scene;
+		FSurfaceCacheResources& SurfaceCacheResources = *Scene->SurfaceCacheResources;
+
+		IrradianceCachePositionRadius.SetBuffer(ShaderRHI, SurfaceCacheResources.Level[DepthLevel]->PositionAndRadius);
+		IrradianceCacheNormal.SetBuffer(ShaderRHI, SurfaceCacheResources.Level[DepthLevel]->Normal);
+		IrradianceCacheTileCoordinate.SetBuffer(ShaderRHI, SurfaceCacheResources.Level[DepthLevel]->TileCoordinate);
+		ScatterDrawParameters.SetBuffer(ShaderRHI, SurfaceCacheResources.Level[DepthLevel]->ScatterDrawParameters);
+
+		SetShaderValue(ShaderRHI, ViewDimensionsParameter, View.ViewRect);
+
+		static bool bInitialized = false;
+		static FAOSampleData2 AOSampleData;
+
+		if (!bInitialized)
+		{
+			for (int32 SampleIndex = 0; SampleIndex < NumConeSampleDirections; SampleIndex++)
+			{
+				AOSampleData.SampleDirections[SampleIndex] = FVector4(SpacedVectors9[SampleIndex]);
+			}
+
+			bInitialized = true;
+		} 
+
+		SetUniformBufferParameterImmediate(ShaderRHI, GetUniformBufferParameter<FAOSampleData2>(), AOSampleData);
+
+		FVector2D ThreadToCulledTileValue(DownsampleFactorValue / (float)(GAODownsampleFactor * GDistanceFieldAOTileSizeX), DownsampleFactorValue / (float)(GAODownsampleFactor * GDistanceFieldAOTileSizeY));
+		SetShaderValue(ShaderRHI, ThreadToCulledTile, ThreadToCulledTileValue);
+
+		SetShaderValue(ShaderRHI, TileListGroupSize, TileListGroupSizeValue);
+
+		SetShaderValue(ShaderRHI, ConeHalfAngle, GAOConeHalfAngle);
+
+		FVector UnoccludedVector(0);
+
+		for (int32 SampleIndex = 0; SampleIndex < NumConeSampleDirections; SampleIndex++)
+		{
+			UnoccludedVector += SpacedVectors9[SampleIndex];
+		}
+
+		float BentNormalNormalizeFactorValue = 1.0f / (UnoccludedVector / NumConeSampleDirections).Size();
+		SetShaderValue(ShaderRHI, BentNormalNormalizeFactor, BentNormalNormalizeFactorValue);
+
+		const int32 NumTexelsOneDimX = GDistanceFieldVolumeTextureAtlas.GetSizeX();
+		const int32 NumTexelsOneDimY = GDistanceFieldVolumeTextureAtlas.GetSizeY();
+		const int32 NumTexelsOneDimZ = GDistanceFieldVolumeTextureAtlas.GetSizeZ();
+		const FVector InvTextureDim(1.0f / NumTexelsOneDimX, 1.0f / NumTexelsOneDimY, 1.0f / NumTexelsOneDimZ);
+		SetShaderValue(ShaderRHI, DistanceFieldAtlasTexelSize, InvTextureDim);
+
+		SetTextureParameter(
+			ShaderRHI,
+			DistanceFieldNormalTexture,
+			DistanceFieldNormalSampler,
+			TStaticSamplerState<SF_Point,AM_Clamp,AM_Clamp,AM_Clamp>::GetRHI(),
+			DistanceFieldNormal.ShaderResourceTexture
+			);
+
+		SetTextureParameter(
+			ShaderRHI,
+			IrradianceCacheSplatTexture,
+			IrradianceCacheSplatSampler,
+			TStaticSamplerState<SF_Point,AM_Clamp,AM_Clamp,AM_Clamp>::GetRHI(),
+			DistanceFieldAOIrradianceCacheSplat.ShaderResourceTexture
+			);
+
+		SetShaderValue(ShaderRHI, RecordRadiusScale, GAORecordRadiusScale);
+	}
+
+	void UnsetParameters()
+	{
+		IrradianceCachePositionRadius.UnsetUAV(GetComputeShader());
+		IrradianceCacheNormal.UnsetUAV(GetComputeShader());
+		IrradianceCacheTileCoordinate.UnsetUAV(GetComputeShader());
+		ScatterDrawParameters.UnsetUAV(GetComputeShader());
+	}
+
+	virtual bool Serialize(FArchive& Ar)
+	{		
+		bool bShaderHasOutdatedParameters = FGlobalShader::Serialize(Ar);
+		Ar << DeferredParameters;
+		Ar << AOParameters;
+		Ar << AOLevelParameters;
+		Ar << IrradianceCachePositionRadius;
+		Ar << IrradianceCacheNormal;
+		Ar << IrradianceCacheTileCoordinate;
+		Ar << ScatterDrawParameters;
+		Ar << ViewDimensionsParameter;
+		Ar << ThreadToCulledTile;
+		Ar << TileListGroupSize;
+		Ar << ConeHalfAngle;
+		Ar << BentNormalNormalizeFactor;
+		Ar << DistanceFieldAtlasTexelSize;
+		Ar << DistanceFieldNormalTexture;
+		Ar << DistanceFieldNormalSampler;
+		Ar << IrradianceCacheSplatTexture;
+		Ar << IrradianceCacheSplatSampler;
+		Ar << RecordRadiusScale;
+		return bShaderHasOutdatedParameters;
+	}
+
+private:
+
+	FDeferredPixelShaderParameters DeferredParameters;
+	FAOParameters AOParameters;
+	FAOLevelParameters AOLevelParameters;
+	FRWShaderParameter IrradianceCachePositionRadius;
+	FRWShaderParameter IrradianceCacheNormal;
+	FRWShaderParameter IrradianceCacheTileCoordinate;
+	FRWShaderParameter ScatterDrawParameters;
+	FShaderParameter ViewDimensionsParameter;
+	FShaderParameter ThreadToCulledTile;
+	FShaderParameter TileListGroupSize;
+	FShaderParameter ConeHalfAngle;
+	FShaderParameter BentNormalNormalizeFactor;
+	FShaderParameter DistanceFieldAtlasTexelSize;
+	FShaderResourceParameter DistanceFieldNormalTexture;
+	FShaderResourceParameter DistanceFieldNormalSampler;
+	FShaderResourceParameter IrradianceCacheSplatTexture;
+	FShaderResourceParameter IrradianceCacheSplatSampler;
+	FShaderParameter RecordRadiusScale;
+};
+
+IMPLEMENT_SHADER_TYPE(,FPopulateCacheCS,TEXT("DistanceFieldSurfaceCacheLightingCompute"),TEXT("PopulateCacheCS"),SF_Compute);
+
+class FSetupFinalGatherIndirectArgumentsCS : public FGlobalShader
+{
+	DECLARE_SHADER_TYPE(FSetupFinalGatherIndirectArgumentsCS,Global)
+public:
+
+	static bool ShouldCache(EShaderPlatform Platform)
+	{
+		return IsFeatureLevelSupported(Platform, ERHIFeatureLevel::SM5) && DoesPlatformSupportDistanceFieldAO(Platform);
+	}
+
+	static void ModifyCompilationEnvironment(EShaderPlatform Platform, FShaderCompilerEnvironment& OutEnvironment)
+	{
+		FGlobalShader::ModifyCompilationEnvironment(Platform,OutEnvironment);
+
+		// To reduce shader compile time of compute shaders with shared memory, doesn't have an impact on generated code with current compiler (June 2010 DX SDK)
+		OutEnvironment.CompilerFlags.Add(CFLAG_StandardOptimization);
+	}
+
+	FSetupFinalGatherIndirectArgumentsCS(const ShaderMetaType::CompiledShaderInitializerType& Initializer)
+		: FGlobalShader(Initializer)
+	{
+		DrawParameters.Bind(Initializer.ParameterMap, TEXT("DrawParameters"));
+		DispatchParameters.Bind(Initializer.ParameterMap, TEXT("DispatchParameters"));
+		SavedStartIndex.Bind(Initializer.ParameterMap, TEXT("SavedStartIndex"));
+	}
+
+	FSetupFinalGatherIndirectArgumentsCS()
+	{
+	}
+
+	void SetParameters(const FSceneView& View, int32 DepthLevel)
+	{
+		FComputeShaderRHIParamRef ShaderRHI = GetComputeShader();
+		FGlobalShader::SetParameters(ShaderRHI, View);
+
+		const FScene* Scene = (const FScene*)View.Family->Scene;
+		FSurfaceCacheResources& SurfaceCacheResources = *Scene->SurfaceCacheResources;
+
+		SetSRVParameter(ShaderRHI, DrawParameters, SurfaceCacheResources.Level[DepthLevel]->ScatterDrawParameters.SRV);
+		SetSRVParameter(ShaderRHI, SavedStartIndex, SurfaceCacheResources.Level[DepthLevel]->SavedStartIndex.SRV);
+
+		DispatchParameters.SetBuffer(ShaderRHI, SurfaceCacheResources.DispatchParameters);
+	}
+
+	void UnsetParameters()
+	{
+		DispatchParameters.UnsetUAV(GetComputeShader());
+	}
+
+	virtual bool Serialize(FArchive& Ar)
+	{		
+		bool bShaderHasOutdatedParameters = FGlobalShader::Serialize(Ar);
+		Ar << DrawParameters;
+		Ar << SavedStartIndex;
+		Ar << DispatchParameters;
+		return bShaderHasOutdatedParameters;
+	}
+
+private:
+
+	FShaderResourceParameter DrawParameters;
+	FShaderResourceParameter SavedStartIndex;
+	FRWShaderParameter DispatchParameters;
+};
+
+IMPLEMENT_SHADER_TYPE(,FSetupFinalGatherIndirectArgumentsCS,TEXT("DistanceFieldSurfaceCacheLightingCompute"),TEXT("SetupFinalGatherIndirectArgumentsCS"),SF_Compute);
+
+class FFinalGatherCS : public FGlobalShader
+{
+	DECLARE_SHADER_TYPE(FFinalGatherCS,Global)
+public:
+
+	static bool ShouldCache(EShaderPlatform Platform)
+	{
+		return IsFeatureLevelSupported(Platform, ERHIFeatureLevel::SM5) && DoesPlatformSupportDistanceFieldAO(Platform);
+	}
+
+	static void ModifyCompilationEnvironment(EShaderPlatform Platform, FShaderCompilerEnvironment& OutEnvironment)
+	{
+		FGlobalShader::ModifyCompilationEnvironment(Platform,OutEnvironment);
+		OutEnvironment.SetDefine(TEXT("NUM_SAMPLES"), NumConeSampleDirections);
+		OutEnvironment.SetDefine(TEXT("DOWNSAMPLE_FACTOR"), GAODownsampleFactor);
+		OutEnvironment.SetDefine(TEXT("MAX_OBJECTS_PER_TILE"), GMaxNumObjectsPerTile);
+
+		// To reduce shader compile time of compute shaders with shared memory, doesn't have an impact on generated code with current compiler (June 2010 DX SDK)
+		OutEnvironment.CompilerFlags.Add(CFLAG_StandardOptimization);
+	}
+
+	FFinalGatherCS(const ShaderMetaType::CompiledShaderInitializerType& Initializer)
+		: FGlobalShader(Initializer)
+	{
+		DeferredParameters.Bind(Initializer.ParameterMap);
+		ObjectParameters.Bind(Initializer.ParameterMap);
+		AOParameters.Bind(Initializer.ParameterMap);
+		IrradianceCachePositionRadius.Bind(Initializer.ParameterMap, TEXT("IrradianceCachePositionRadius"));
+		IrradianceCacheNormal.Bind(Initializer.ParameterMap, TEXT("IrradianceCacheNormal"));
+		OccluderRadius.Bind(Initializer.ParameterMap, TEXT("OccluderRadius"));
+		IrradianceCacheBentNormal.Bind(Initializer.ParameterMap, TEXT("IrradianceCacheBentNormal"));
+		IrradianceCacheTileCoordinate.Bind(Initializer.ParameterMap, TEXT("IrradianceCacheTileCoordinate"));
+		ScatterDrawParameters.Bind(Initializer.ParameterMap, TEXT("ScatterDrawParameters"));
+		SavedStartIndex.Bind(Initializer.ParameterMap, TEXT("SavedStartIndex"));
+		ViewDimensionsParameter.Bind(Initializer.ParameterMap, TEXT("ViewDimensions"));
+		ThreadToCulledTile.Bind(Initializer.ParameterMap, TEXT("ThreadToCulledTile"));
+		TileHeadData.Bind(Initializer.ParameterMap, TEXT("TileHeadData"));
+		TileHeadDataUnpacked.Bind(Initializer.ParameterMap, TEXT("TileHeadDataUnpacked"));
+		TileArrayData.Bind(Initializer.ParameterMap, TEXT("TileArrayData"));
+		TileListGroupSize.Bind(Initializer.ParameterMap, TEXT("TileListGroupSize"));
+		ConeHalfAngle.Bind(Initializer.ParameterMap, TEXT("ConeHalfAngle"));
+		BentNormalNormalizeFactor.Bind(Initializer.ParameterMap, TEXT("BentNormalNormalizeFactor"));
+		DistanceFieldAtlasTexelSize.Bind(Initializer.ParameterMap, TEXT("DistanceFieldAtlasTexelSize"));
+		DistanceFieldNormalTexture.Bind(Initializer.ParameterMap, TEXT("DistanceFieldNormalTexture"));
+		DistanceFieldNormalSampler.Bind(Initializer.ParameterMap, TEXT("DistanceFieldNormalSampler"));
+		IrradianceCacheSplatTexture.Bind(Initializer.ParameterMap, TEXT("IrradianceCacheSplatTexture"));
+		IrradianceCacheSplatSampler.Bind(Initializer.ParameterMap, TEXT("IrradianceCacheSplatSampler"));
+		AOLevelParameters.Bind(Initializer.ParameterMap);
+		RecordRadiusScale.Bind(Initializer.ParameterMap, TEXT("RecordRadiusScale"));
+	}
+
+	FFinalGatherCS()
+	{
+	}
+
+	void SetParameters(const FSceneView& View, FSceneRenderTargetItem& DistanceFieldAOIrradianceCacheSplat, FSceneRenderTargetItem& DistanceFieldNormal, int32 NumObjects, int32 DownsampleFactorValue, int32 DepthLevel, FIntPoint TileListGroupSizeValue)
+	{
+		FComputeShaderRHIParamRef ShaderRHI = GetComputeShader();
+		FGlobalShader::SetParameters(ShaderRHI, View);
+		DeferredParameters.Set(ShaderRHI, View);
+		ObjectParameters.Set(ShaderRHI, NumObjects);
+		AOParameters.Set(ShaderRHI);
+		AOLevelParameters.Set(ShaderRHI, View, DownsampleFactorValue);
+
+		const FScene* Scene = (const FScene*)View.Family->Scene;
+		FSurfaceCacheResources& SurfaceCacheResources = *Scene->SurfaceCacheResources;
+
+		SetSRVParameter(ShaderRHI, IrradianceCacheNormal, SurfaceCacheResources.Level[DepthLevel]->Normal.SRV);
+		SetSRVParameter(ShaderRHI, ScatterDrawParameters, SurfaceCacheResources.Level[DepthLevel]->ScatterDrawParameters.SRV);
+		SetSRVParameter(ShaderRHI, SavedStartIndex, SurfaceCacheResources.Level[DepthLevel]->SavedStartIndex.SRV);
+		SetSRVParameter(ShaderRHI, IrradianceCachePositionRadius, SurfaceCacheResources.Level[DepthLevel]->PositionAndRadius.SRV);
+
+		OccluderRadius.SetBuffer(ShaderRHI, SurfaceCacheResources.Level[DepthLevel]->OccluderRadius);
+		IrradianceCacheBentNormal.SetBuffer(ShaderRHI, SurfaceCacheResources.Level[DepthLevel]->BentNormal);
+
+		SetShaderValue(ShaderRHI, ViewDimensionsParameter, View.ViewRect);
+
+		static bool bInitialized = false;
+		static FAOSampleData2 AOSampleData;
+
+		if (!bInitialized)
+		{
+			for (int32 SampleIndex = 0; SampleIndex < NumConeSampleDirections; SampleIndex++)
+			{
+				AOSampleData.SampleDirections[SampleIndex] = FVector4(SpacedVectors9[SampleIndex]);
+			}
+
+			bInitialized = true;
+		} 
+
+		SetUniformBufferParameterImmediate(ShaderRHI, GetUniformBufferParameter<FAOSampleData2>(), AOSampleData);
+
+		FVector2D ThreadToCulledTileValue(DownsampleFactorValue / (float)(GAODownsampleFactor * GDistanceFieldAOTileSizeX), DownsampleFactorValue / (float)(GAODownsampleFactor * GDistanceFieldAOTileSizeY));
+		SetShaderValue(ShaderRHI, ThreadToCulledTile, ThreadToCulledTileValue);
+
+		FTileIntersectionResources* TileIntersectionResources = ((FSceneViewState*)View.State)->AOTileIntersectionResources;
+
+		SetSRVParameter(ShaderRHI, TileHeadDataUnpacked, TileIntersectionResources->TileHeadDataUnpacked.SRV);
+		SetSRVParameter(ShaderRHI, TileHeadData, TileIntersectionResources->TileHeadData.SRV);
+		SetSRVParameter(ShaderRHI, TileArrayData, TileIntersectionResources->TileArrayData.SRV);
+		SetSRVParameter(ShaderRHI, IrradianceCacheTileCoordinate, SurfaceCacheResources.Level[DepthLevel]->TileCoordinate.SRV);
+
+		SetShaderValue(ShaderRHI, TileListGroupSize, TileListGroupSizeValue);
+
+		SetShaderValue(ShaderRHI, ConeHalfAngle, GAOConeHalfAngle);
+
+		FVector UnoccludedVector(0);
+
+		for (int32 SampleIndex = 0; SampleIndex < NumConeSampleDirections; SampleIndex++)
+		{
+			UnoccludedVector += SpacedVectors9[SampleIndex];
+		}
+
+		float BentNormalNormalizeFactorValue = 1.0f / (UnoccludedVector / NumConeSampleDirections).Size();
+		SetShaderValue(ShaderRHI, BentNormalNormalizeFactor, BentNormalNormalizeFactorValue);
+
+		const int32 NumTexelsOneDimX = GDistanceFieldVolumeTextureAtlas.GetSizeX();
+		const int32 NumTexelsOneDimY = GDistanceFieldVolumeTextureAtlas.GetSizeY();
+		const int32 NumTexelsOneDimZ = GDistanceFieldVolumeTextureAtlas.GetSizeZ();
+		const FVector InvTextureDim(1.0f / NumTexelsOneDimX, 1.0f / NumTexelsOneDimY, 1.0f / NumTexelsOneDimZ);
+		SetShaderValue(ShaderRHI, DistanceFieldAtlasTexelSize, InvTextureDim);
+
+		SetTextureParameter(
+			ShaderRHI,
+			DistanceFieldNormalTexture,
+			DistanceFieldNormalSampler,
+			TStaticSamplerState<SF_Point,AM_Clamp,AM_Clamp,AM_Clamp>::GetRHI(),
+			DistanceFieldNormal.ShaderResourceTexture
+			);
+
+		SetTextureParameter(
+			ShaderRHI,
+			IrradianceCacheSplatTexture,
+			IrradianceCacheSplatSampler,
+			TStaticSamplerState<SF_Point,AM_Clamp,AM_Clamp,AM_Clamp>::GetRHI(),
+			DistanceFieldAOIrradianceCacheSplat.ShaderResourceTexture
+			);
+
+		SetShaderValue(ShaderRHI, RecordRadiusScale, GAORecordRadiusScale);
+	}
+
+	void UnsetParameters()
+	{
+		OccluderRadius.UnsetUAV(GetComputeShader());
+		IrradianceCacheBentNormal.UnsetUAV(GetComputeShader());
+	}
+
+	virtual bool Serialize(FArchive& Ar)
+	{		
+		bool bShaderHasOutdatedParameters = FGlobalShader::Serialize(Ar);
+		Ar << DeferredParameters;
+		Ar << ObjectParameters;
+		Ar << AOParameters;
+		Ar << AOLevelParameters;
+		Ar << IrradianceCachePositionRadius;
+		Ar << IrradianceCacheNormal;
+		Ar << OccluderRadius;
+		Ar << IrradianceCacheBentNormal;
+		Ar << IrradianceCacheTileCoordinate;
+		Ar << ScatterDrawParameters;
+		Ar << SavedStartIndex;
+		Ar << ViewDimensionsParameter;
+		Ar << ThreadToCulledTile;
+		Ar << TileHeadDataUnpacked;
+		Ar << TileHeadData;
+		Ar << TileArrayData;
+		Ar << TileListGroupSize;
+		Ar << ConeHalfAngle;
+		Ar << BentNormalNormalizeFactor;
+		Ar << DistanceFieldAtlasTexelSize;
+		Ar << DistanceFieldNormalTexture;
+		Ar << DistanceFieldNormalSampler;
+		Ar << IrradianceCacheSplatTexture;
+		Ar << IrradianceCacheSplatSampler;
+		Ar << RecordRadiusScale;
+		return bShaderHasOutdatedParameters;
+	}
+
+private:
+
+	FDeferredPixelShaderParameters DeferredParameters;
+	FAOObjectBufferParameters ObjectParameters;
+	FAOParameters AOParameters;
+	FAOLevelParameters AOLevelParameters;
+	FShaderResourceParameter IrradianceCacheNormal;
+	FShaderResourceParameter IrradianceCachePositionRadius;
+	FRWShaderParameter OccluderRadius;
+	FRWShaderParameter IrradianceCacheBentNormal;
+	FShaderResourceParameter ScatterDrawParameters;
+	FShaderResourceParameter SavedStartIndex;
+	FShaderResourceParameter IrradianceCacheTileCoordinate;
+	FShaderParameter ViewDimensionsParameter;
+	FShaderParameter ThreadToCulledTile;
+	FShaderResourceParameter TileHeadDataUnpacked;
+	FShaderResourceParameter TileHeadData;
+	FShaderResourceParameter TileArrayData;
+	FShaderParameter TileListGroupSize;
+	FShaderParameter ConeHalfAngle;
+	FShaderParameter BentNormalNormalizeFactor;
+	FShaderParameter DistanceFieldAtlasTexelSize;
+	FShaderResourceParameter DistanceFieldNormalTexture;
+	FShaderResourceParameter DistanceFieldNormalSampler;
+	FShaderResourceParameter IrradianceCacheSplatTexture;
+	FShaderResourceParameter IrradianceCacheSplatSampler;
+	FShaderParameter RecordRadiusScale;
+};
+
+IMPLEMENT_SHADER_TYPE(,FFinalGatherCS,TEXT("DistanceFieldSurfaceCacheLightingCompute"),TEXT("FinalGatherCS"),SF_Compute);
+
+/**  */
+template<bool bFinalInterpolationPass>
+class TIrradianceCacheSplatVS : public FGlobalShader
+{
+	DECLARE_SHADER_TYPE(TIrradianceCacheSplatVS,Global);
+public:
+
+	static bool ShouldCache(EShaderPlatform Platform) 
+	{ 
+		return IsFeatureLevelSupported(Platform, ERHIFeatureLevel::SM5) && DoesPlatformSupportDistanceFieldAO(Platform); 
+	}
+
+	static void ModifyCompilationEnvironment(EShaderPlatform Platform, FShaderCompilerEnvironment& OutEnvironment)
+	{
+		FGlobalShader::ModifyCompilationEnvironment(Platform,OutEnvironment);
+
+		OutEnvironment.SetDefine(TEXT("FINAL_INTERPOLATION_PASS"), (uint32)bFinalInterpolationPass);
+	}
+
+	TIrradianceCacheSplatVS(const ShaderMetaType::CompiledShaderInitializerType& Initializer):
+		FGlobalShader(Initializer)
+	{
+		IrradianceCachePositionRadius.Bind(Initializer.ParameterMap, TEXT("IrradianceCachePositionRadius"));
+		IrradianceCacheOccluderRadius.Bind(Initializer.ParameterMap, TEXT("IrradianceCacheOccluderRadius"));
+		IrradianceCacheNormal.Bind(Initializer.ParameterMap, TEXT("IrradianceCacheNormal"));
+		IrradianceCacheBentNormal.Bind(Initializer.ParameterMap, TEXT("IrradianceCacheBentNormal"));
+		InterpolationRadiusScale.Bind(Initializer.ParameterMap, TEXT("InterpolationRadiusScale"));
+		NormalizedOffsetToPixelCenter.Bind(Initializer.ParameterMap, TEXT("NormalizedOffsetToPixelCenter"));
+		HackExpand.Bind(Initializer.ParameterMap, TEXT("HackExpand"));
+	}
+
+	TIrradianceCacheSplatVS() {}
+
+	void SetParameters(const FSceneView& View, int32 DepthLevel, int32 CurrentLevelDownsampleFactorValue, FVector2D NormalizedOffsetToPixelCenterValue)
+	{
+		const FVertexShaderRHIParamRef ShaderRHI = GetVertexShader();
+		FGlobalShader::SetParameters(ShaderRHI, View);
+		
+		const FScene* Scene = (const FScene*)View.Family->Scene;
+		FSurfaceCacheResources& SurfaceCacheResources = *Scene->SurfaceCacheResources;
+
+		SetSRVParameter(ShaderRHI, IrradianceCachePositionRadius, SurfaceCacheResources.Level[DepthLevel]->PositionAndRadius.SRV);
+		SetSRVParameter(ShaderRHI, IrradianceCacheOccluderRadius, SurfaceCacheResources.Level[DepthLevel]->OccluderRadius.SRV);
+		SetSRVParameter(ShaderRHI, IrradianceCacheNormal, SurfaceCacheResources.Level[DepthLevel]->Normal.SRV);
+		SetSRVParameter(ShaderRHI, IrradianceCacheBentNormal, SurfaceCacheResources.Level[DepthLevel]->BentNormal.SRV);
+
+		SetShaderValue(ShaderRHI, InterpolationRadiusScale, (bFinalInterpolationPass ? GAOInterpolationRadiusScale : 1.0f));
+		SetShaderValue(ShaderRHI, NormalizedOffsetToPixelCenter, NormalizedOffsetToPixelCenterValue);
+
+		const FIntPoint AOViewRectSize = FIntPoint::DivideAndRoundUp(View.ViewRect.Size(), CurrentLevelDownsampleFactorValue);
+		const FVector2D HackExpandValue(.5f / AOViewRectSize.X, .5f / AOViewRectSize.Y);
+		SetShaderValue(ShaderRHI, HackExpand, HackExpandValue);
+	}
+
+	virtual bool Serialize(FArchive& Ar)
+	{
+		bool bShaderHasOutdatedParameters = FGlobalShader::Serialize(Ar);
+		Ar << IrradianceCachePositionRadius;
+		Ar << IrradianceCacheOccluderRadius;
+		Ar << IrradianceCacheNormal;
+		Ar << IrradianceCacheBentNormal;
+		Ar << InterpolationRadiusScale;
+		Ar << NormalizedOffsetToPixelCenter;
+		Ar << HackExpand;
+		return bShaderHasOutdatedParameters;
+	}
+
+private:
+	FShaderResourceParameter IrradianceCachePositionRadius;
+	FShaderResourceParameter IrradianceCacheOccluderRadius;
+	FShaderResourceParameter IrradianceCacheNormal;
+	FShaderResourceParameter IrradianceCacheBentNormal;
+	FShaderParameter InterpolationRadiusScale;
+	FShaderParameter NormalizedOffsetToPixelCenter;
+	FShaderParameter HackExpand;
+};
+
+IMPLEMENT_SHADER_TYPE(template<>,TIrradianceCacheSplatVS<true>,TEXT("DistanceFieldSurfaceCacheLighting"),TEXT("IrradianceCacheSplatVS"),SF_Vertex);
+IMPLEMENT_SHADER_TYPE(template<>,TIrradianceCacheSplatVS<false>,TEXT("DistanceFieldSurfaceCacheLighting"),TEXT("IrradianceCacheSplatVS"),SF_Vertex);
+
+template<bool bFinalInterpolationPass>
+class TIrradianceCacheSplatPS : public FGlobalShader
+{
+	DECLARE_SHADER_TYPE(TIrradianceCacheSplatPS, Global);
+public:
+
+	static bool ShouldCache(EShaderPlatform Platform)
+	{
+		return IsFeatureLevelSupported(Platform, ERHIFeatureLevel::SM5) && DoesPlatformSupportDistanceFieldAO(Platform);
+	}
+
+	static void ModifyCompilationEnvironment(EShaderPlatform Platform, FShaderCompilerEnvironment& OutEnvironment)
+	{
+		FGlobalShader::ModifyCompilationEnvironment(Platform,OutEnvironment);
+		
+		OutEnvironment.SetDefine(TEXT("DOWNSAMPLE_FACTOR"), GAODownsampleFactor);
+		OutEnvironment.SetDefine(TEXT("FINAL_INTERPOLATION_PASS"), (uint32)bFinalInterpolationPass);
+	}
+
+	/** Default constructor. */
+	TIrradianceCacheSplatPS() {}
+
+	/** Initialization constructor. */
+	TIrradianceCacheSplatPS(const ShaderMetaType::CompiledShaderInitializerType& Initializer)
+		: FGlobalShader(Initializer)
+	{
+		AOParameters.Bind(Initializer.ParameterMap);
+		DeferredParameters.Bind(Initializer.ParameterMap);
+		AOLevelParameters.Bind(Initializer.ParameterMap);
+		DistanceFieldNormalTexture.Bind(Initializer.ParameterMap, TEXT("DistanceFieldNormalTexture"));
+		DistanceFieldNormalSampler.Bind(Initializer.ParameterMap, TEXT("DistanceFieldNormalSampler"));
+		InterpolationAngleNormalization.Bind(Initializer.ParameterMap, TEXT("InterpolationAngleNormalization"));
+		InvMinCosPointBehindPlane.Bind(Initializer.ParameterMap, TEXT("InvMinCosPointBehindPlane"));
+	}
+
+	void SetParameters(const FSceneView& View, FSceneRenderTargetItem& DistanceFieldNormal, int32 DestLevelDownsampleFactor)
+	{
+		const FPixelShaderRHIParamRef ShaderRHI = GetPixelShader();
+		FGlobalShader::SetParameters(ShaderRHI, View);
+		AOParameters.Set(ShaderRHI);
+		DeferredParameters.Set(ShaderRHI, View);
+		AOLevelParameters.Set(ShaderRHI, View, DestLevelDownsampleFactor);
+
+		SetTextureParameter(
+			ShaderRHI,
+			DistanceFieldNormalTexture,
+			DistanceFieldNormalSampler,
+			TStaticSamplerState<SF_Point,AM_Clamp,AM_Clamp,AM_Clamp>::GetRHI(),
+			DistanceFieldNormal.ShaderResourceTexture
+			);
+
+		const float EffectiveMaxAngle = bFinalInterpolationPass ? GAOInterpolationMaxAngle * GAOInterpolationAngleScale : GAOInterpolationMaxAngle;
+		const float InterpolationAngleNormalizationValue = 1.0f / FMath::Sqrt(1.0f - FMath::Cos(EffectiveMaxAngle * PI / 180.0f));
+		SetShaderValue(ShaderRHI, InterpolationAngleNormalization, InterpolationAngleNormalizationValue);
+
+		const float MinCosPointBehindPlaneValue = FMath::Cos((GAOMinPointBehindPlaneAngle + 90.0f) * PI / 180.0f);
+		SetShaderValue(ShaderRHI, InvMinCosPointBehindPlane, 1.0f / MinCosPointBehindPlaneValue);
+	}
+
+	// FShader interface.
+	virtual bool Serialize(FArchive& Ar)
+	{
+		bool bShaderHasOutdatedParameters = FGlobalShader::Serialize(Ar);
+		Ar << AOParameters;
+		Ar << AOLevelParameters;
+		Ar << DeferredParameters;
+		Ar << DistanceFieldNormalTexture;
+		Ar << DistanceFieldNormalSampler;
+		Ar << InterpolationAngleNormalization;
+		Ar << InvMinCosPointBehindPlane;
+		return bShaderHasOutdatedParameters;
+	}
+
+private:
+	FAOParameters AOParameters;
+	FAOLevelParameters AOLevelParameters;
+	FDeferredPixelShaderParameters DeferredParameters;
+	FShaderResourceParameter DistanceFieldNormalTexture;
+	FShaderResourceParameter DistanceFieldNormalSampler;
+	FShaderParameter InterpolationAngleNormalization;
+	FShaderParameter InvMinCosPointBehindPlane;
+};
+
+IMPLEMENT_SHADER_TYPE(template<>,TIrradianceCacheSplatPS<true>,TEXT("DistanceFieldSurfaceCacheLighting"),TEXT("IrradianceCacheSplatPS"),SF_Pixel);
+IMPLEMENT_SHADER_TYPE(template<>,TIrradianceCacheSplatPS<false>,TEXT("DistanceFieldSurfaceCacheLighting"),TEXT("IrradianceCacheSplatPS"),SF_Pixel);
+
+class FDistanceFieldAOCombinePS2 : public FGlobalShader
+{
+	DECLARE_SHADER_TYPE(FDistanceFieldAOCombinePS2, Global);
+public:
+
+	static bool ShouldCache(EShaderPlatform Platform)
+	{
+		return IsFeatureLevelSupported(Platform, ERHIFeatureLevel::SM5) && DoesPlatformSupportDistanceFieldAO(Platform);
+	}
+
+	static void ModifyCompilationEnvironment(EShaderPlatform Platform, FShaderCompilerEnvironment& OutEnvironment)
+	{
+		OutEnvironment.SetDefine(TEXT("DOWNSAMPLE_FACTOR"), GAODownsampleFactor);
+	}
+
+	/** Default constructor. */
+	FDistanceFieldAOCombinePS2() {}
+
+	/** Initialization constructor. */
+	FDistanceFieldAOCombinePS2(const ShaderMetaType::CompiledShaderInitializerType& Initializer)
+		: FGlobalShader(Initializer)
+	{
+		DeferredParameters.Bind(Initializer.ParameterMap);
+		AOParameters.Bind(Initializer.ParameterMap);
+		AOLevelParameters.Bind(Initializer.ParameterMap);
+		BentNormalAOTexture3.Bind(Initializer.ParameterMap,TEXT("BentNormalAOTexture3"));
+		BentNormalAOSampler3.Bind(Initializer.ParameterMap,TEXT("BentNormalAOSampler3"));
+	}
+
+	void SetParameters(const FSceneView& View, FSceneRenderTargetItem& InBentNormalTexture)
+	{
+		const FPixelShaderRHIParamRef ShaderRHI = GetPixelShader();
+
+		FGlobalShader::SetParameters(ShaderRHI, View);
+		DeferredParameters.Set(ShaderRHI, View);
+		AOParameters.Set(ShaderRHI);
+		AOLevelParameters.Set(ShaderRHI, View, GAODownsampleFactor);
+
+		SetTextureParameter(ShaderRHI, BentNormalAOTexture3, BentNormalAOSampler3, TStaticSamplerState<SF_Point>::GetRHI(), InBentNormalTexture.ShaderResourceTexture);	
+	}
+	// FShader interface.
+	virtual bool Serialize(FArchive& Ar)
+	{
+		bool bShaderHasOutdatedParameters = FGlobalShader::Serialize(Ar);
+		Ar << DeferredParameters;
+		Ar << AOParameters;
+		Ar << AOLevelParameters;
+		Ar << BentNormalAOTexture3;
+		Ar << BentNormalAOSampler3;
+		return bShaderHasOutdatedParameters;
+	}
+
+private:
+
+	FDeferredPixelShaderParameters DeferredParameters;
+	FAOParameters AOParameters;
+	FAOLevelParameters AOLevelParameters;
+	FShaderResourceParameter BentNormalAOTexture3;
+	FShaderResourceParameter BentNormalAOSampler3;
+};
+
+IMPLEMENT_SHADER_TYPE(,FDistanceFieldAOCombinePS2,TEXT("DistanceFieldSurfaceCacheLighting"),TEXT("AOCombinePS2"),SF_Pixel);
+
+class FFillGapsPS : public FGlobalShader
+{
+	DECLARE_SHADER_TYPE(FFillGapsPS, Global);
+public:
+
+	static bool ShouldCache(EShaderPlatform Platform)
+	{
+		return IsFeatureLevelSupported(Platform, ERHIFeatureLevel::SM5) && DoesPlatformSupportDistanceFieldAO(Platform);
+	}
+
+	static void ModifyCompilationEnvironment(EShaderPlatform Platform, FShaderCompilerEnvironment& OutEnvironment)
+	{
+		OutEnvironment.SetDefine(TEXT("DOWNSAMPLE_FACTOR"), GAODownsampleFactor);
+	}
+
+	/** Default constructor. */
+	FFillGapsPS() {}
+
+	/** Initialization constructor. */
+	FFillGapsPS(const ShaderMetaType::CompiledShaderInitializerType& Initializer)
+		: FGlobalShader(Initializer)
+	{
+		BentNormalAOTexture.Bind(Initializer.ParameterMap, TEXT("BentNormalAOTexture"));
+		BentNormalAOSampler.Bind(Initializer.ParameterMap, TEXT("BentNormalAOSampler"));
+		BentNormalAOTexelSize.Bind(Initializer.ParameterMap, TEXT("BentNormalAOTexelSize"));
+	}
+
+	void SetParameters(const FSceneView& View, FSceneRenderTargetItem& DistanceFieldAOBentNormal2)
+	{
+		const FPixelShaderRHIParamRef ShaderRHI = GetPixelShader();
+
+		FGlobalShader::SetParameters(ShaderRHI, View);
+
+		SetTextureParameter(
+			ShaderRHI,
+			BentNormalAOTexture,
+			BentNormalAOSampler,
+			TStaticSamplerState<SF_Bilinear,AM_Clamp,AM_Clamp,AM_Clamp>::GetRHI(),
+			DistanceFieldAOBentNormal2.ShaderResourceTexture
+			);
+
+		const FIntPoint DownsampledBufferSize(GSceneRenderTargets.GetBufferSizeXY() / FIntPoint(GAODownsampleFactor, GAODownsampleFactor));
+		const FVector2D BaseLevelTexelSizeValue(1.0f / DownsampledBufferSize.X, 1.0f / DownsampledBufferSize.Y);
+		SetShaderValue(ShaderRHI, BentNormalAOTexelSize, BaseLevelTexelSizeValue);
+	}
+	// FShader interface.
+	virtual bool Serialize(FArchive& Ar)
+	{
+		bool bShaderHasOutdatedParameters = FGlobalShader::Serialize(Ar);
+		Ar << BentNormalAOTexture;
+		Ar << BentNormalAOSampler;
+		Ar << BentNormalAOTexelSize;
+		return bShaderHasOutdatedParameters;
+	}
+
+private:
+
+	FShaderResourceParameter BentNormalAOTexture;
+	FShaderResourceParameter BentNormalAOSampler;
+	FShaderParameter BentNormalAOTexelSize;
+};
+
+IMPLEMENT_SHADER_TYPE(,FFillGapsPS,TEXT("DistanceFieldSurfaceCacheLighting"),TEXT("FillGapsPS"),SF_Pixel);
+
+
+class FUpdateHistoryPS : public FGlobalShader
+{
+	DECLARE_SHADER_TYPE(FUpdateHistoryPS, Global);
+public:
+
+	static bool ShouldCache(EShaderPlatform Platform)
+	{
+		return IsFeatureLevelSupported(Platform, ERHIFeatureLevel::SM5) && DoesPlatformSupportDistanceFieldAO(Platform);
+	}
+
+	static void ModifyCompilationEnvironment(EShaderPlatform Platform, FShaderCompilerEnvironment& OutEnvironment)
+	{
+		OutEnvironment.SetDefine(TEXT("DOWNSAMPLE_FACTOR"), GAODownsampleFactor);
+	}
+
+	/** Default constructor. */
+	FUpdateHistoryPS() {}
+
+	/** Initialization constructor. */
+	FUpdateHistoryPS(const ShaderMetaType::CompiledShaderInitializerType& Initializer)
+		: FGlobalShader(Initializer)
+	{
+		DeferredParameters.Bind(Initializer.ParameterMap);
+		BentNormalAOTexture.Bind(Initializer.ParameterMap, TEXT("BentNormalAOTexture"));
+		BentNormalAOSampler.Bind(Initializer.ParameterMap, TEXT("BentNormalAOSampler"));
+		HistoryTexture.Bind(Initializer.ParameterMap, TEXT("HistoryTexture"));
+		HistorySampler.Bind(Initializer.ParameterMap, TEXT("HistorySampler"));
+		CameraMotion.Bind(Initializer.ParameterMap);
+		HistoryWeight.Bind(Initializer.ParameterMap, TEXT("HistoryWeight"));
+	}
+
+	void SetParameters(const FSceneView& View, FSceneRenderTargetItem& HistoryTextureValue, FSceneRenderTargetItem& DistanceFieldAOBentNormal)
+	{
+		const FPixelShaderRHIParamRef ShaderRHI = GetPixelShader();
+
+		FGlobalShader::SetParameters(ShaderRHI, View);
+		DeferredParameters.Set(ShaderRHI, View);
+
+		SetTextureParameter(
+			ShaderRHI,
+			BentNormalAOTexture,
+			BentNormalAOSampler,
+			TStaticSamplerState<SF_Bilinear,AM_Clamp,AM_Clamp,AM_Clamp>::GetRHI(),
+			DistanceFieldAOBentNormal.ShaderResourceTexture
+			);
+
+		SetTextureParameter(
+			ShaderRHI,
+			HistoryTexture,
+			HistorySampler,
+			TStaticSamplerState<SF_Bilinear,AM_Clamp,AM_Clamp,AM_Clamp>::GetRHI(),
+			HistoryTextureValue.ShaderResourceTexture
+			);
+
+		CameraMotion.Set(View, ShaderRHI);
+
+		SetShaderValue(ShaderRHI, HistoryWeight, GAOHistoryWeight);
+	}
+	// FShader interface.
+	virtual bool Serialize(FArchive& Ar)
+	{
+		bool bShaderHasOutdatedParameters = FGlobalShader::Serialize(Ar);
+		Ar << DeferredParameters;
+		Ar << BentNormalAOTexture;
+		Ar << BentNormalAOSampler;
+		Ar << HistoryTexture;
+		Ar << HistorySampler;
+		Ar << CameraMotion;
+		Ar << HistoryWeight;
+		return bShaderHasOutdatedParameters;
+	}
+
+private:
+
+	FDeferredPixelShaderParameters DeferredParameters;
+	FShaderResourceParameter BentNormalAOTexture;
+	FShaderResourceParameter BentNormalAOSampler;
+	FShaderResourceParameter HistoryTexture;
+	FShaderResourceParameter HistorySampler;
+	FCameraMotionParameters CameraMotion;
+	FShaderParameter HistoryWeight;
+};
+
+IMPLEMENT_SHADER_TYPE(,FUpdateHistoryPS,TEXT("DistanceFieldSurfaceCacheLighting"),TEXT("UpdateHistoryPS"),SF_Pixel);
+
+void AllocateOrReuseAORenderTarget(TRefCountPtr<IPooledRenderTarget>& Target, const TCHAR* Name)
+{
+	if (!Target)
+	{
+		extern FIntPoint GetBufferSizeForAO();
+		FIntPoint ExpandedBufferSize = GetBufferSizeForAO();
+
+		FPooledRenderTargetDesc Desc(FPooledRenderTargetDesc::Create2DDesc(ExpandedBufferSize / GAODownsampleFactor, PF_FloatRGBA, TexCreate_None, TexCreate_RenderTargetable | TexCreate_UAV, false));
+		GRenderTargetPool.FindFreeElement(Desc, Target, Name);
+	}
+}
+
+void UpdateHistory(
+	FViewInfo& View, 
+	const TCHAR* HistoryRTName,
+	/** Contains last frame's history, if non-NULL.  This will be updated with the new frame's history. */
+	TRefCountPtr<IPooledRenderTarget>* HistoryState,
+	/** Source */
+	TRefCountPtr<IPooledRenderTarget>& AOSource, 
+	/** Output of Temporal Reprojection for the next step in the pipeline. */
+	TRefCountPtr<IPooledRenderTarget>& HistoryOutput)
+{
+	if (HistoryState)
+	{
+		if (*HistoryState && !View.bCameraCut && !View.bPrevTransformsReset)
+		{
+			// Reuse a render target from the pool with a consistent name, for vis purposes
+			TRefCountPtr<IPooledRenderTarget> NewHistory;
+			AllocateOrReuseAORenderTarget(NewHistory, HistoryRTName);
+
+			{
+				SCOPED_DRAW_EVENT(UpdateHistory, DEC_SCENE_ITEMS);
+				RHISetRenderTarget(NewHistory->GetRenderTargetItem().TargetableTexture, NULL);
+
+				RHISetViewport(View.ViewRect.Min.X / GAODownsampleFactor, View.ViewRect.Min.Y / GAODownsampleFactor, 0.0f, View.ViewRect.Max.X / GAODownsampleFactor, View.ViewRect.Max.Y / GAODownsampleFactor, 1.0f);
+				RHISetRasterizerState(TStaticRasterizerState<FM_Solid,CM_None>::GetRHI());
+				RHISetDepthStencilState(TStaticDepthStencilState<false,CF_Always>::GetRHI());
+				RHISetBlendState(TStaticBlendState<>::GetRHI());
+
+				TShaderMapRef<FPostProcessVS> VertexShader( GetGlobalShaderMap() );
+				TShaderMapRef<FUpdateHistoryPS> PixelShader( GetGlobalShaderMap() );
+
+				static FGlobalBoundShaderState BoundShaderState;
+				SetGlobalBoundShaderState(BoundShaderState, GFilterVertexDeclaration.VertexDeclarationRHI, *VertexShader, *PixelShader);
+
+				PixelShader->SetParameters(View, (*HistoryState)->GetRenderTargetItem(), AOSource->GetRenderTargetItem());
+
+				DrawRectangle( 
+					0, 0, 
+					View.ViewRect.Width() / GAODownsampleFactor, View.ViewRect.Height() / GAODownsampleFactor,
+					View.ViewRect.Min.X / GAODownsampleFactor, View.ViewRect.Min.Y / GAODownsampleFactor, 
+					View.ViewRect.Width() / GAODownsampleFactor, View.ViewRect.Height() / GAODownsampleFactor,
+					FIntPoint(View.ViewRect.Width() / GAODownsampleFactor, View.ViewRect.Height() / GAODownsampleFactor),
+					GSceneRenderTargets.GetBufferSizeXY() / FIntPoint(GAODownsampleFactor, GAODownsampleFactor),
+					*VertexShader);
+			}
+
+			// Update the view state's render target reference with the new history
+			*HistoryState = NewHistory;
+			HistoryOutput = NewHistory;
+		}
+		else
+		{
+			// Use the current frame's mask for next frame's history
+			*HistoryState = AOSource;
+			HistoryOutput = AOSource;
+			AOSource = NULL;
+		}
+	}
+	else
+	{
+		// Temporal reprojection is disabled or there is no view state - pass through
+		HistoryOutput = AOSource;
+	}
+}
+
+void PostProcessBentNormalAO(TArray<FViewInfo>& Views, FSceneRenderTargetItem& IrradianceCacheInterpolation, TRefCountPtr<IPooledRenderTarget>& AOOutput)
+{
+	TRefCountPtr<IPooledRenderTarget> DistanceFieldAOBentNormal;
+	AllocateOrReuseAORenderTarget(DistanceFieldAOBentNormal, TEXT("DistanceFieldBentNormalAO"));
+
+	TRefCountPtr<IPooledRenderTarget> DistanceFieldAOBentNormal2;
+
+	if (GAOFillGaps)
+	{
+		AllocateOrReuseAORenderTarget(DistanceFieldAOBentNormal2, TEXT("DistanceFieldBentNormalAO2"));
+	}
+
+	{
+		SCOPED_DRAW_EVENT(AOCombine, DEC_SCENE_ITEMS);
+
+		RHISetRenderTarget(GAOFillGaps 
+			? DistanceFieldAOBentNormal2->GetRenderTargetItem().TargetableTexture
+			: DistanceFieldAOBentNormal->GetRenderTargetItem().TargetableTexture, 
+			NULL);
+
+		for (int32 ViewIndex = 0; ViewIndex < Views.Num(); ViewIndex++)
+		{
+			FViewInfo& View = Views[ViewIndex];
+			RHISetViewport(View.ViewRect.Min.X / GAODownsampleFactor, View.ViewRect.Min.Y / GAODownsampleFactor, 0.0f, View.ViewRect.Max.X / GAODownsampleFactor, View.ViewRect.Max.Y / GAODownsampleFactor, 1.0f);
+			RHISetRasterizerState(TStaticRasterizerState<FM_Solid,CM_None>::GetRHI());
+			RHISetDepthStencilState(TStaticDepthStencilState<false,CF_Always>::GetRHI());
+			RHISetBlendState(TStaticBlendState<>::GetRHI());
+
+			TShaderMapRef<FPostProcessVS> VertexShader(GetGlobalShaderMap());
+			TShaderMapRef<FDistanceFieldAOCombinePS2> PixelShader(GetGlobalShaderMap());
+
+			static FGlobalBoundShaderState BoundShaderState;
+			SetGlobalBoundShaderState(BoundShaderState, GFilterVertexDeclaration.VertexDeclarationRHI, *VertexShader, *PixelShader);
+
+			PixelShader->SetParameters(View, IrradianceCacheInterpolation);
+
+			//@todo - get the round up correct
+			DrawRectangle( 
+				0, 0, 
+				View.ViewRect.Width() / GAODownsampleFactor, View.ViewRect.Height() / GAODownsampleFactor,
+				View.ViewRect.Min.X / GAODownsampleFactor, View.ViewRect.Min.Y / GAODownsampleFactor, 
+				View.ViewRect.Width() / GAODownsampleFactor, View.ViewRect.Height() / GAODownsampleFactor,
+				FIntPoint(View.ViewRect.Width() / GAODownsampleFactor, View.ViewRect.Height() / GAODownsampleFactor),
+				GetBufferSizeForAO() / GAODownsampleFactor,
+				*VertexShader);
+		}
+	}
+
+	if (GAOFillGaps)
+	{
+		SCOPED_DRAW_EVENT(FillGaps, DEC_SCENE_ITEMS);
+		RHISetRenderTarget(DistanceFieldAOBentNormal->GetRenderTargetItem().TargetableTexture, NULL);
+
+		for (int32 ViewIndex = 0; ViewIndex < Views.Num(); ViewIndex++)
+		{
+			const FViewInfo& View = Views[ViewIndex];
+
+			RHISetViewport(View.ViewRect.Min.X / GAODownsampleFactor, View.ViewRect.Min.Y / GAODownsampleFactor, 0.0f, View.ViewRect.Max.X / GAODownsampleFactor, View.ViewRect.Max.Y / GAODownsampleFactor, 1.0f);
+			RHISetRasterizerState(TStaticRasterizerState<FM_Solid,CM_None>::GetRHI());
+			RHISetDepthStencilState(TStaticDepthStencilState<false,CF_Always>::GetRHI());
+			RHISetBlendState(TStaticBlendState<>::GetRHI());
+
+			TShaderMapRef<FPostProcessVS> VertexShader( GetGlobalShaderMap() );
+			TShaderMapRef<FFillGapsPS> PixelShader( GetGlobalShaderMap() );
+
+			static FGlobalBoundShaderState BoundShaderState;
+			SetGlobalBoundShaderState(BoundShaderState, GFilterVertexDeclaration.VertexDeclarationRHI, *VertexShader, *PixelShader);
+
+			VertexShader->SetParameters(View);
+			PixelShader->SetParameters(View, DistanceFieldAOBentNormal2->GetRenderTargetItem());
+
+			DrawRectangle( 
+				0, 0, 
+				View.ViewRect.Width() / GAODownsampleFactor, View.ViewRect.Height() / GAODownsampleFactor,
+				View.ViewRect.Min.X / GAODownsampleFactor, View.ViewRect.Min.Y / GAODownsampleFactor, 
+				View.ViewRect.Width() / GAODownsampleFactor, View.ViewRect.Height() / GAODownsampleFactor,
+				FIntPoint(View.ViewRect.Width() / GAODownsampleFactor, View.ViewRect.Height() / GAODownsampleFactor),
+				GSceneRenderTargets.GetBufferSizeXY() / FIntPoint(GAODownsampleFactor, GAODownsampleFactor),
+				*VertexShader);
+		}
+	}
+
+	FSceneViewState* ViewState = (FSceneViewState*)Views[0].State;
+	TRefCountPtr<IPooledRenderTarget>* HistoryState = ViewState ? &ViewState->DistanceFieldAOHistoryRT : NULL;
+	AOOutput = DistanceFieldAOBentNormal;
+
+	if (GAOUseHistory)
+	{
+		FViewInfo& View = Views[0];
+
+		UpdateHistory(
+			View, 
+			TEXT("DistanceFieldAOHistory"),
+			HistoryState,
+			DistanceFieldAOBentNormal, 
+			AOOutput);
+	}
+}
+
+class FDistanceFieldAOUpsamplePS : public FGlobalShader
+{
+	DECLARE_SHADER_TYPE(FDistanceFieldAOUpsamplePS, Global);
+public:
+
+	static bool ShouldCache(EShaderPlatform Platform)
+	{
+		return IsFeatureLevelSupported(Platform, ERHIFeatureLevel::SM5) && DoesPlatformSupportDistanceFieldAO(Platform);
+	}
+
+	static void ModifyCompilationEnvironment(EShaderPlatform Platform, FShaderCompilerEnvironment& OutEnvironment)
+	{
+		OutEnvironment.SetDefine(TEXT("DOWNSAMPLE_FACTOR"), GAODownsampleFactor);
+	}
+
+	/** Default constructor. */
+	FDistanceFieldAOUpsamplePS() {}
+
+	/** Initialization constructor. */
+	FDistanceFieldAOUpsamplePS(const ShaderMetaType::CompiledShaderInitializerType& Initializer)
+		: FGlobalShader(Initializer)
+	{
+		DeferredParameters.Bind(Initializer.ParameterMap);
+		BentNormalAOTexture.Bind(Initializer.ParameterMap,TEXT("BentNormalAOTexture"));
+		BentNormalAOSampler.Bind(Initializer.ParameterMap,TEXT("BentNormalAOSampler"));
+	}
+
+	void SetParameters(const FSceneView& View, TRefCountPtr<IPooledRenderTarget>& DistanceFieldAOBentNormal)
+	{
+		const FPixelShaderRHIParamRef ShaderRHI = GetPixelShader();
+
+		FGlobalShader::SetParameters(ShaderRHI, View);
+		DeferredParameters.Set(ShaderRHI, View);
+
+		SetTextureParameter(ShaderRHI, BentNormalAOTexture, BentNormalAOSampler, TStaticSamplerState<SF_Bilinear>::GetRHI(), DistanceFieldAOBentNormal->GetRenderTargetItem().ShaderResourceTexture);	
+	}
+	// FShader interface.
+	virtual bool Serialize(FArchive& Ar)
+	{
+		bool bShaderHasOutdatedParameters = FGlobalShader::Serialize(Ar);
+		Ar << DeferredParameters;
+		Ar << BentNormalAOTexture;
+		Ar << BentNormalAOSampler;
+		return bShaderHasOutdatedParameters;
+	}
+
+private:
+
+	FDeferredPixelShaderParameters DeferredParameters;
+	FShaderResourceParameter BentNormalAOTexture;
+	FShaderResourceParameter BentNormalAOSampler;
+};
+
+IMPLEMENT_SHADER_TYPE(,FDistanceFieldAOUpsamplePS,TEXT("DistanceFieldSurfaceCacheLighting"),TEXT("AOUpsamplePS"),SF_Pixel);
+
+void UpsampleBentNormalAO(const TArray<FViewInfo>& Views, TRefCountPtr<IPooledRenderTarget>& DistanceFieldAOBentNormal)
+{
+	for (int32 ViewIndex = 0; ViewIndex < Views.Num(); ViewIndex++)
+	{
+		const FViewInfo& View = Views[ViewIndex];
+
+		SCOPED_DRAW_EVENT(UpsampleAO, DEC_SCENE_ITEMS);
+
+		RHISetViewport(View.ViewRect.Min.X, View.ViewRect.Min.Y, 0.0f, View.ViewRect.Max.X, View.ViewRect.Max.Y, 1.0f);
+		RHISetRasterizerState(TStaticRasterizerState<FM_Solid,CM_None>::GetRHI());
+		RHISetDepthStencilState(TStaticDepthStencilState<false,CF_Always>::GetRHI());
+		RHISetBlendState(TStaticBlendState<>::GetRHI());
+
+		TShaderMapRef<FPostProcessVS> VertexShader( GetGlobalShaderMap() );
+		TShaderMapRef<FDistanceFieldAOUpsamplePS> PixelShader( GetGlobalShaderMap() );
+
+		static FGlobalBoundShaderState BoundShaderState;
+		SetGlobalBoundShaderState(BoundShaderState, GFilterVertexDeclaration.VertexDeclarationRHI, *VertexShader, *PixelShader);
+
+		PixelShader->SetParameters(View, DistanceFieldAOBentNormal);
+
+		DrawRectangle( 
+			0, 0, 
+			View.ViewRect.Width(), View.ViewRect.Height(),
+			View.ViewRect.Min.X / GAODownsampleFactor, View.ViewRect.Min.Y / GAODownsampleFactor, 
+			View.ViewRect.Width() / GAODownsampleFactor, View.ViewRect.Height() / GAODownsampleFactor,
+			FIntPoint(View.ViewRect.Width(), View.ViewRect.Height()),
+			GetBufferSizeForAO() / GAODownsampleFactor,
+			*VertexShader);
+	}
+}
+
+/** Generates a pseudo-random position inside the unit sphere, uniformly distributed over the volume of the sphere. */
+FVector GetUnitPosition2(FRandomStream& RandomStream)
+{
+	FVector Result;
+	// Use rejection sampling to generate a valid sample
+	do
+	{
+		Result.X = RandomStream.GetFraction() * 2 - 1;
+		Result.Y = RandomStream.GetFraction() * 2 - 1;
+		Result.Z = RandomStream.GetFraction() * 2 - 1;
+	} while( Result.SizeSquared() > 1.f );
+	return Result;
+}
+
+/** Generates a pseudo-random unit vector, uniformly distributed over all directions. */
+FVector GetUnitVector2(FRandomStream& RandomStream)
+{
+	return GetUnitPosition2(RandomStream).UnsafeNormal();
+}
+
+void GenerateBestSpacedVectors()
+{
+	if (false)
+	{
+		FVector BestSpacedVectors9[9];
+		float BestCoverage = 0;
+		// Each cone covers an area of ConeSolidAngle = HemisphereSolidAngle / NumCones
+		// HemisphereSolidAngle = 2 * PI
+		// ConeSolidAngle = 2 * PI * (1 - cos(ConeHalfAngle))
+		// cos(ConeHalfAngle) = 1 - 1 / NumCones
+		float CosHalfAngle = 1 - 1.0f / (float)ARRAY_COUNT(BestSpacedVectors9);
+		// Prevent self-intersection in sample set
+		float MinAngle = FMath::Acos(CosHalfAngle);
+		float MinZ = FMath::Sin(MinAngle);
+		FRandomStream RandomStream(1234567);
+
+		// Super slow random brute force search
+		for (int i = 0; i < 1000000; i++)
+		{
+			FVector CandidateSpacedVectors[ARRAY_COUNT(BestSpacedVectors9)];
+
+			for (int j = 0; j < ARRAY_COUNT(CandidateSpacedVectors); j++)
+			{
+				FVector NewSample;
+
+				// Reject invalid directions until we get a valid one
+				do 
+				{
+					NewSample = GetUnitVector2(RandomStream);
+				} 
+				while (NewSample.Z <= MinZ);
+
+				CandidateSpacedVectors[j] = NewSample;
+			}
+
+			float Coverage = 0;
+			int NumSamples = 10000;
+
+			// Determine total cone coverage with monte carlo estimation
+			for (int sample = 0; sample < NumSamples; sample++)
+			{
+				FVector NewSample;
+
+				do 
+				{
+					NewSample = GetUnitVector2(RandomStream);
+				} 
+				while (NewSample.Z <= 0);
+
+				bool bIntersects = false;
+
+				for (int j = 0; j < ARRAY_COUNT(CandidateSpacedVectors); j++)
+				{
+					if (FVector::DotProduct(CandidateSpacedVectors[j], NewSample) > CosHalfAngle)
+					{
+						bIntersects = true;
+						break;
+					}
+				}
+
+				Coverage += bIntersects ? 1 / (float)NumSamples : 0;
+			}
+
+			if (Coverage > BestCoverage)
+			{
+				BestCoverage = Coverage;
+
+				for (int j = 0; j < ARRAY_COUNT(CandidateSpacedVectors); j++)
+				{
+					BestSpacedVectors9[j] = CandidateSpacedVectors[j];
+				}
+			}
+		}
+
+		int32 temp = 0;
+	}
+}
+
+FIntPoint BuildTileObjectLists(FScene* Scene, TArray<FViewInfo>& Views, int32 NumObjects)
+{
+	SCOPED_DRAW_EVENT(BuildTileList, DEC_SCENE_ITEMS);
+	RHISetRenderTarget(NULL, NULL);
+
+	FIntPoint TileListGroupSize;
+
+	if (GAOScatterTileCulling)
+	{
+		for (int32 ViewIndex = 0; ViewIndex < Views.Num(); ViewIndex++)
+		{
+			const FViewInfo& View = Views[ViewIndex];
+
+			uint32 GroupSizeX = FMath::DivideAndRoundUp(View.ViewRect.Size().X / GAODownsampleFactor, GDistanceFieldAOTileSizeX);
+			uint32 GroupSizeY = FMath::DivideAndRoundUp(View.ViewRect.Size().Y / GAODownsampleFactor, GDistanceFieldAOTileSizeY);
+			TileListGroupSize = FIntPoint(GroupSizeX, GroupSizeY);
+
+			FTileIntersectionResources*& TileIntersectionResources = ((FSceneViewState*)View.State)->AOTileIntersectionResources;
+
+			if (!TileIntersectionResources || TileIntersectionResources->TileDimensions != TileListGroupSize)
+			{
+				if (TileIntersectionResources)
+				{
+					TileIntersectionResources->ReleaseResource();
+				}
+				else
+				{
+					TileIntersectionResources = new FTileIntersectionResources();
+				}
+
+				TileIntersectionResources->TileDimensions = TileListGroupSize;
+
+				TileIntersectionResources->InitResource();
+			}
+
+			{
+				SCOPED_DRAW_EVENT(BuildTileCones, DEC_SCENE_ITEMS);
+				TShaderMapRef<FBuildTileConesCS> ComputeShader(GetGlobalShaderMap());
+
+				RHISetComputeShader(ComputeShader->GetComputeShader());	
+				ComputeShader->SetParameters(View, Scene, FVector2D(GroupSizeX, GroupSizeY));
+				DispatchComputeShader(*ComputeShader, GroupSizeX, GroupSizeY, 1);
+
+				ComputeShader->UnsetParameters();
+			}
+
+			{
+				SCOPED_DRAW_EVENT(CullObjects, DEC_SCENE_ITEMS);
+
+				TShaderMapRef<FObjectCullVS> VertexShader(GetGlobalShaderMap());
+				TShaderMapRef<FObjectCullPS> PixelShader(GetGlobalShaderMap());
+
+				TArray<FUnorderedAccessViewRHIParamRef> UAVs;
+				PixelShader->GetUAVs(Views[0], UAVs);
+				RHISetRenderTargets(0, (const FRHIRenderTargetView *)NULL, NULL, UAVs.Num(), UAVs.GetData());
+
+				RHISetViewport(0, 0, 0.0f, GroupSizeX, GroupSizeY, 1.0f);
+
+				// Render backfaces since camera may intersect
+				RHISetRasterizerState(View.bReverseCulling ? TStaticRasterizerState<FM_Solid,CM_CW>::GetRHI() : TStaticRasterizerState<FM_Solid,CM_CCW>::GetRHI());
+				RHISetDepthStencilState(TStaticDepthStencilState<false,CF_Always>::GetRHI());
+				RHISetBlendState(TStaticBlendState<>::GetRHI());
+
+				static FGlobalBoundShaderState BoundShaderState;
+				SetGlobalBoundShaderState(BoundShaderState, GetVertexDeclarationFVector4(), *VertexShader, *PixelShader);
+
+				VertexShader->SetParameters(View, NumObjects);
+				PixelShader->SetParameters(View, FVector2D(GroupSizeX, GroupSizeY));
+
+				RHISetStreamSource(0, StencilingGeometry::GLowPolyStencilSphereVertexBuffer.VertexBufferRHI, sizeof(FVector4), 0);
+
+				RHIDrawIndexedPrimitive(
+					StencilingGeometry::GLowPolyStencilSphereIndexBuffer.IndexBufferRHI, 
+					PT_TriangleList, 
+					0,
+					0, 
+					StencilingGeometry::GLowPolyStencilSphereVertexBuffer.GetVertexCount(), 
+					0, 
+					StencilingGeometry::GLowPolyStencilSphereIndexBuffer.GetIndexCount() / 3, 
+					NumObjects);
+			}
+		}
+	}
+	else
+	{
+		for (int32 ViewIndex = 0; ViewIndex < Views.Num(); ViewIndex++)
+		{
+			const FViewInfo& View = Views[ViewIndex];
+
+			uint32 GroupSizeX = (View.ViewRect.Size().X / GAODownsampleFactor + GDistanceFieldAOTileSizeX - 1) / GDistanceFieldAOTileSizeX;
+			uint32 GroupSizeY = (View.ViewRect.Size().Y / GAODownsampleFactor + GDistanceFieldAOTileSizeY - 1) / GDistanceFieldAOTileSizeY;
+			TileListGroupSize = FIntPoint(GroupSizeX, GroupSizeY);
+
+			FTileIntersectionResources*& TileIntersectionResources = ((FSceneViewState*)View.State)->AOTileIntersectionResources;
+
+			if (!TileIntersectionResources || TileIntersectionResources->TileDimensions != TileListGroupSize)
+			{
+				if (TileIntersectionResources)
+				{
+					TileIntersectionResources->ReleaseResource();
+				}
+				else
+				{
+					TileIntersectionResources = new FTileIntersectionResources();
+				}
+
+				TileIntersectionResources->TileDimensions = TileListGroupSize;
+
+				TileIntersectionResources->InitResource();
+			}
+
+			// Indicates the clear value for each channel of the UAV format
+			uint32 ClearValues[4] = {0};
+			RHIClearUAV(TileIntersectionResources->TileArrayNextAllocation.UAV, ClearValues);
+			RHIClearUAV(TileIntersectionResources->TileHeadData.UAV, ClearValues);
+
+			TShaderMapRef<FDistanceFieldBuildTileListCS > ComputeShader(GetGlobalShaderMap());
+
+			RHISetComputeShader(ComputeShader->GetComputeShader());	
+			ComputeShader->SetParameters(View, Scene, FVector2D(GroupSizeX, GroupSizeY), NumObjects);
+			DispatchComputeShader(*ComputeShader, GroupSizeX, GroupSizeY, 1);
+
+			ComputeShader->UnsetParameters();
+		}
+	}
+
+	return TileListGroupSize;
+}
+
+void RenderIrradianceCacheInterpolation(TArray<FViewInfo>& Views, IPooledRenderTarget* InterpolationTarget, FSceneRenderTargetItem& DistanceFieldNormal, int32 DepthLevel, int32 DestLevelDownsampleFactor, bool bFinalInterpolation)
+{
+	check(!(bFinalInterpolation && DepthLevel != 0));
+
+	{
+		SCOPED_DRAW_EVENT(IrradianceCacheSplat, DEC_SCENE_ITEMS);
+		const int32 MaxAllowedLevel = GAOMaxSupportedLevel;
+
+		RHISetRenderTarget(InterpolationTarget->GetRenderTargetItem().TargetableTexture, NULL);
+		RHIClear(true, FLinearColor(0, 0, 0, 0), false, 0, false, 0, FIntRect());
+
+		for (int32 ViewIndex = 0; ViewIndex < Views.Num(); ViewIndex++)
+		{
+			FViewInfo& View = Views[ViewIndex];
+
+			const FScene* Scene = (const FScene*)View.Family->Scene;
+			FSurfaceCacheResources& SurfaceCacheResources = *Scene->SurfaceCacheResources;
+
+			RHISetViewport(View.ViewRect.Min.X / DestLevelDownsampleFactor, View.ViewRect.Min.Y / DestLevelDownsampleFactor, 0.0f, FMath::DivideAndRoundUp(View.ViewRect.Max.X, DestLevelDownsampleFactor), FMath::DivideAndRoundUp(View.ViewRect.Max.Y, DestLevelDownsampleFactor), 1.0f);
+			//RHISetViewport(View.ViewRect.Min.X / DestLevelDownsampleFactor, View.ViewRect.Min.Y / DestLevelDownsampleFactor, 0.0f, View.ViewRect.Max.X / DestLevelDownsampleFactor, View.ViewRect.Max.Y / DestLevelDownsampleFactor, 1.0f);
+
+			RHISetRasterizerState(TStaticRasterizerState<FM_Solid,CM_None>::GetRHI());
+			//@todo - render front faces with depth testing, requires depth buffer
+			RHISetDepthStencilState(TStaticDepthStencilState<false,CF_Always>::GetRHI());
+			RHISetBlendState(TStaticBlendState<CW_RGBA, BO_Add, BF_One, BF_One, BO_Add, BF_One, BF_One>::GetRHI());
+
+			//FVector2D NormalizedOffsetToPixelCenter(-(-DestLevelDownsampleFactor / 2 + .5f) / (float)View.ViewRect.Width(), -(-DestLevelDownsampleFactor / 2 + .5f) / (float)View.ViewRect.Height());
+			uint32 DownsampledViewSizeX = FMath::DivideAndRoundUp(View.ViewRect.Width(), DestLevelDownsampleFactor);
+			uint32 DownsampledViewSizeY = FMath::DivideAndRoundUp(View.ViewRect.Height(), DestLevelDownsampleFactor);
+			FVector2D NormalizedOffsetToPixelCenter2(.5f / DownsampledViewSizeX - .5f / (DownsampledViewSizeX * DestLevelDownsampleFactor), .5f / DownsampledViewSizeY - .5f / (DownsampledViewSizeY * DestLevelDownsampleFactor));
+
+			// Convert to NDC 
+			NormalizedOffsetToPixelCenter2 = NormalizedOffsetToPixelCenter2 * FVector2D(2, -2);
+
+			const FIntPoint BufferSize = GSceneRenderTargets.GetBufferSizeXY();
+			const float InvBufferSizeX = 1.0f / BufferSize.X;
+			const float InvBufferSizeY = 1.0f / BufferSize.Y;
+
+			const FVector2D ScreenPositionScale(
+				View.ViewRect.Width() * InvBufferSizeX / +2.0f,
+				View.ViewRect.Height() * InvBufferSizeY / (-2.0f * GProjectionSignY));
+
+			const FVector2D ScreenPositionBias(
+				(View.ViewRect.Width() / 2.0f + GPixelCenterOffset + View.ViewRect.Min.X) * InvBufferSizeX,
+				(View.ViewRect.Height() / 2.0f + GPixelCenterOffset + View.ViewRect.Min.Y) * InvBufferSizeY);
+
+			FVector2D OffsetToLowResCorner = (FVector2D(.5f, .5f) / FVector2D(DownsampledViewSizeX, DownsampledViewSizeY) - FVector2D(.5f, .5f)) * FVector2D(2, -2);
+			FVector2D OffsetToTopResPixel = (FVector2D(.5f, .5f) / BufferSize - ScreenPositionBias) / ScreenPositionScale;
+			FVector2D NormalizedOffsetToPixelCenter = OffsetToLowResCorner - OffsetToTopResPixel;
+
+			RHISetStreamSource(0, GCircleVertexBuffer.VertexBufferRHI, sizeof(FScreenVertex), 0);
+
+			if (bFinalInterpolation)
+			{
+				TShaderMapRef<TIrradianceCacheSplatVS<true> > VertexShader(GetGlobalShaderMap());
+				TShaderMapRef<TIrradianceCacheSplatPS<true> > PixelShader(GetGlobalShaderMap());
+
+				for (int32 SplatSourceDepthLevel = MaxAllowedLevel; SplatSourceDepthLevel >= DepthLevel; SplatSourceDepthLevel--)
+				{
+					static FGlobalBoundShaderState BoundShaderState;
+					SetGlobalBoundShaderState(BoundShaderState, GScreenVertexDeclaration.VertexDeclarationRHI, *VertexShader, *PixelShader);
+
+					VertexShader->SetParameters(View, SplatSourceDepthLevel, DestLevelDownsampleFactor, NormalizedOffsetToPixelCenter);
+					PixelShader->SetParameters(View, DistanceFieldNormal, DestLevelDownsampleFactor);
+
+					RHIDrawPrimitiveIndirect(PT_TriangleList, SurfaceCacheResources.Level[SplatSourceDepthLevel]->ScatterDrawParameters.Buffer, 0);
+				}
+			}
+			else
+			{
+				TShaderMapRef<TIrradianceCacheSplatVS<false> > VertexShader(GetGlobalShaderMap());
+				TShaderMapRef<TIrradianceCacheSplatPS<false> > PixelShader(GetGlobalShaderMap());
+
+				for (int32 SplatSourceDepthLevel = MaxAllowedLevel; SplatSourceDepthLevel >= DepthLevel; SplatSourceDepthLevel--)
+				{
+					static FGlobalBoundShaderState BoundShaderState;
+					SetGlobalBoundShaderState(BoundShaderState, GScreenVertexDeclaration.VertexDeclarationRHI, *VertexShader, *PixelShader);
+
+					VertexShader->SetParameters(View, SplatSourceDepthLevel, DestLevelDownsampleFactor, NormalizedOffsetToPixelCenter);
+					PixelShader->SetParameters(View, DistanceFieldNormal, DestLevelDownsampleFactor);
+
+					RHIDrawPrimitiveIndirect(PT_TriangleList, SurfaceCacheResources.Level[SplatSourceDepthLevel]->ScatterDrawParameters.Buffer, 0);
+				}
+			}
+		}
+	}
+}
+
+void UpdateVisibleObjectBuffers(const FScene* Scene, const FViewInfo& View, int32& NumObjects)
+{
+	//@todo - scene rendering allocator
+	TArray<FVector4> ObjectBoundsData;
+	TArray<FVector4> ObjectData;
+	TArray<FVector4> ObjectData2;
+
+	{
+		QUICK_SCOPE_CYCLE_COUNTER(STAT_GatherObjectData);
+
+		const float AOMaxDistance = GetAOMaxDistance();
+		const FVector ViewOriginForDistanceCulling = View.ViewMatrices.ViewOrigin;
+		const int32 NumTexelsOneDimX = GDistanceFieldVolumeTextureAtlas.GetSizeX();
+		const int32 NumTexelsOneDimY = GDistanceFieldVolumeTextureAtlas.GetSizeY();
+		const int32 NumTexelsOneDimZ = GDistanceFieldVolumeTextureAtlas.GetSizeZ();
+		const FVector InvTextureDim(1.0f / NumTexelsOneDimX, 1.0f / NumTexelsOneDimY, 1.0f / NumTexelsOneDimZ);
+
+		ObjectBoundsData.Empty(Scene->Primitives.Num() / 2);
+		ObjectData.Empty(FAOObjectBuffers::ObjectDataStride * Scene->Primitives.Num() / 2);
+		ObjectData2.Empty(FAOObjectBuffers::ObjectData2Stride * Scene->Primitives.Num() / 2);
+
+		//@todo - dense data-coherent pass like frustum culling in InitViews
+		for (int32 PrimitiveIndex = 0; PrimitiveIndex < Scene->Primitives.Num(); PrimitiveIndex++)
+		{
+			const FPrimitiveSceneInfo* PrimitiveSceneInfo = Scene->Primitives[PrimitiveIndex];
+
+			if (PrimitiveSceneInfo->Proxy->CastsDynamicShadow() || PrimitiveSceneInfo->Proxy->CastsStaticShadow())
+			{
+				FBox LocalVolumeBounds;
+				FIntVector BlockMin;
+				FIntVector BlockSize;
+				PrimitiveSceneInfo->Proxy->GetDistancefieldAtlasData(LocalVolumeBounds, BlockMin, BlockSize);
+
+				if (BlockMin.X >= 0 && BlockMin.Y >= 0 && BlockMin.Z >= 0)
+				{
+					const FMatrix LocalToWorld = PrimitiveSceneInfo->Proxy->GetLocalToWorld();
+
+					const FMatrix VolumeToWorld = FScaleMatrix(LocalVolumeBounds.GetExtent()) 
+						* FTranslationMatrix(LocalVolumeBounds.GetCenter())
+						* LocalToWorld;
+
+					const FVector4 ObjectBoundingSphere(VolumeToWorld.GetOrigin(), VolumeToWorld.GetScaleVector().Size());
+					const float DistanceToViewSq = (ViewOriginForDistanceCulling - (FVector)ObjectBoundingSphere).SizeSquared();
+
+					if (!GAOCPUFrustumCull 
+						|| DistanceToViewSq < FMath::Square(GAOMaxViewDistance + ObjectBoundingSphere.W)
+						&& View.ViewFrustum.IntersectSphere((FVector)ObjectBoundingSphere, ObjectBoundingSphere.W + AOMaxDistance))
+					{
+						ObjectBoundsData.Add(ObjectBoundingSphere);
+
+						const float MaxExtent = LocalVolumeBounds.GetExtent().GetMax();
+
+						const FMatrix UniformScaleVolumeToWorld = FScaleMatrix(MaxExtent) 
+							* FTranslationMatrix(LocalVolumeBounds.GetCenter())
+							* LocalToWorld;
+
+						const FVector InvBlockSize(1.0f / BlockSize.X, 1.0f / BlockSize.Y, 1.0f / BlockSize.Z);
+
+						//float3 VolumeUV = (VolumePosition / LocalPositionExtent * .5f * UVScale + .5f * UVScale + UVAdd;
+						const FVector LocalPositionExtent = LocalVolumeBounds.GetExtent() / FVector(MaxExtent);
+						const FVector UVScale = FVector(BlockSize) * InvTextureDim;
+						const float VolumeScale = UniformScaleVolumeToWorld.GetMaximumAxisScale();
+						ObjectData2.Add(*(FVector4*)&UniformScaleVolumeToWorld.M[0]);
+						ObjectData2.Add(*(FVector4*)&UniformScaleVolumeToWorld.M[1]);
+						ObjectData2.Add(*(FVector4*)&UniformScaleVolumeToWorld.M[2]);
+						ObjectData2.Add(*(FVector4*)&UniformScaleVolumeToWorld.M[3]);
+
+						const FMatrix WorldToVolume = UniformScaleVolumeToWorld.Inverse();
+						// WorldToVolume
+						ObjectData.Add(*(FVector4*)&WorldToVolume.M[0]);
+						ObjectData.Add(*(FVector4*)&WorldToVolume.M[1]);
+						ObjectData.Add(*(FVector4*)&WorldToVolume.M[2]);
+						ObjectData.Add(*(FVector4*)&WorldToVolume.M[3]);
+
+						// Clamp to texel center by subtracting a half texel in the [-1,1] position space
+						// LocalPositionExtent
+						ObjectData.Add(FVector4(LocalPositionExtent - InvBlockSize, 0));
+
+						// UVScale
+						ObjectData.Add(FVector4(FVector(BlockSize) * InvTextureDim * .5f / LocalPositionExtent, VolumeScale));
+
+						// UVAdd
+						ObjectData.Add(FVector4(FVector(BlockMin) * InvTextureDim + .5f * UVScale, 0));
+
+						// Padding
+						ObjectData.Add(FVector4(0, 0, 0, 0));
+
+						checkSlow(ObjectData.Num() % FAOObjectBuffers::ObjectDataStride == 0);
+						checkSlow(ObjectData2.Num() % FAOObjectBuffers::ObjectData2Stride == 0);
+
+						NumObjects++;
+					}
+				}
+			}
+		}
+	}
+
+	if (NumObjects > 0)
+	{
+		QUICK_SCOPE_CYCLE_COUNTER(STAT_UpdateObjectBuffers);
+
+		if (NumObjects > GAOObjectBuffers.MaxObjects)
+		{
+			// Allocate with slack
+			GAOObjectBuffers.MaxObjects = NumObjects * 5 / 4;
+			GAOObjectBuffers.UpdateRHI();
+		}
+
+		checkSlow(NumObjects < MAX_uint16);
+
+		void* LockedBuffer = RHILockVertexBuffer(GAOObjectBuffers.ObjectData.Bounds, 0, GAOObjectBuffers.ObjectData.Bounds->GetSize(), RLM_WriteOnly);
+		FPlatformMemory::Memcpy(LockedBuffer, ObjectBoundsData.GetData(), ObjectBoundsData.GetTypeSize() * ObjectBoundsData.Num());
+		RHIUnlockVertexBuffer(GAOObjectBuffers.ObjectData.Bounds);
+
+		LockedBuffer = RHILockVertexBuffer(GAOObjectBuffers.ObjectData.Data, 0, GAOObjectBuffers.ObjectData.Data->GetSize(), RLM_WriteOnly);
+		FPlatformMemory::Memcpy(LockedBuffer, ObjectData.GetData(), ObjectData.GetTypeSize() * ObjectData.Num());
+		RHIUnlockVertexBuffer(GAOObjectBuffers.ObjectData.Data);
+
+		LockedBuffer = RHILockVertexBuffer(GAOObjectBuffers.ObjectData.Data2, 0, GAOObjectBuffers.ObjectData.Data2->GetSize(), RLM_WriteOnly);
+		FPlatformMemory::Memcpy(LockedBuffer, ObjectData2.GetData(), ObjectData2.GetTypeSize() * ObjectData2.Num());
+		RHIUnlockVertexBuffer(GAOObjectBuffers.ObjectData.Data2);
+	}
+}
+
+bool FDeferredShadingSceneRenderer::RenderDistanceFieldAOSurfaceCache(FSceneRenderTargetItem& OutBentNormalAO, bool bApplyToSceneColor)
+{
+	//@todo - support multiple views
+	const FViewInfo& View = Views[0];
+
+	if (GDistanceFieldAO 
+		&& GRHIFeatureLevel >= ERHIFeatureLevel::SM5
+		&& DoesPlatformSupportDistanceFieldAO(GRHIShaderPlatform)
+		&& Views.Num() == 1
+		// ViewState is used to cache tile intersection resources which have to be sized based on the view
+		&& View.State)
+	{
+		QUICK_SCOPE_CYCLE_COUNTER(STAT_RenderDistanceFieldAOSurfaceCache);
+		SCOPED_DRAW_EVENT(DistanceFieldAO, DEC_SCENE_ITEMS);
+
+		// Update the global distance field atlas
+		GDistanceFieldVolumeTextureAtlas.UpdateAllocations();
+
+		if (GDistanceFieldVolumeTextureAtlas.VolumeTextureRHI)
+		{
+			GenerateBestSpacedVectors();
+
+			// Create surface cache state that will persist across frames if needed
+			if (!Scene->SurfaceCacheResources)
+			{
+				Scene->SurfaceCacheResources = new FSurfaceCacheResources();
+				Scene->SurfaceCacheResources->InitResource();
+			}
+
+			uint32 ClearValues[4] = {0};
+			const int32 MaxAllowedLevel = GAOMaxSupportedLevel;
+			FSurfaceCacheResources& SurfaceCacheResources = *Scene->SurfaceCacheResources;
+
+			if (!SurfaceCacheResources.bClearedResources || !GAOReuseAcrossFrames)
+			{
+				// Reset the number of active cache records to 0
+				for (int32 DepthLevel = MaxAllowedLevel; DepthLevel >= 0; DepthLevel--)
+				{
+					RHIClearUAV(SurfaceCacheResources.Level[DepthLevel]->ScatterDrawParameters.UAV, ClearValues);
+				}
+
+				SurfaceCacheResources.bClearedResources = true;
+			}
+
+			int32 NumObjects = 0;
+
+			// Update object buffers with objects that have a valid distance field representation and may affect pixels on the screen
+			UpdateVisibleObjectBuffers(Scene, View, NumObjects);
+
+			if (NumObjects > 0)
+			{
+				QUICK_SCOPE_CYCLE_COUNTER(STAT_AOIssueGPUWork);
+
+				// Intersect objects with screen tiles, build lists
+				FIntPoint TileListGroupSize = BuildTileObjectLists(Scene, Views, NumObjects);
+
+				TRefCountPtr<IPooledRenderTarget> DistanceFieldNormal;
+
+				{
+					const FIntPoint ExpandedBufferSize = GetBufferSizeForAO();
+					FPooledRenderTargetDesc Desc(FPooledRenderTargetDesc::Create2DDesc(ExpandedBufferSize / GAODownsampleFactor, PF_FloatRGBA, TexCreate_None, TexCreate_RenderTargetable | TexCreate_UAV, false));
+					GRenderTargetPool.FindFreeElement(Desc, DistanceFieldNormal, TEXT("DistanceFieldNormal"));
+				}
+
+				// Compute the distance field normal, this is used for surface caching instead of the GBuffer normal because it matches the occluding geometry
+				ComputeDistanceFieldNormal(Views, DistanceFieldNormal->GetRenderTargetItem(), NumObjects);
+
+				GRenderTargetPool.VisualizeTexture.SetCheckPoint(DistanceFieldNormal);
+
+				if (GAOReuseAcrossFrames && GAOTrimOldRecordsFraction > 0)
+				{
+					SCOPED_DRAW_EVENT(TrimRecords, DEC_SCENE_ITEMS);
+
+					// Copy and trim last frame's surface cache samples
+					for (int32 DepthLevel = FMath::Min(GAOMaxLevel, MaxAllowedLevel); DepthLevel >= FMath::Max(GAOMinLevel, 0); DepthLevel--)
+					{
+						{	
+							TShaderMapRef<FSetupCopyIndirectArgumentsCS> ComputeShader(GetGlobalShaderMap());
+
+							RHISetComputeShader(ComputeShader->GetComputeShader());	
+							ComputeShader->SetParameters(View, DepthLevel);
+							DispatchComputeShader(*ComputeShader, 1, 1, 1);
+
+							ComputeShader->UnsetParameters();
+						}
+
+						{
+							TShaderMapRef<FCopyIrradianceCacheSamplesCS> ComputeShader(GetGlobalShaderMap());
+
+							RHISetComputeShader(ComputeShader->GetComputeShader());	
+							ComputeShader->SetParameters(View, DepthLevel);
+							DispatchIndirectComputeShader(*ComputeShader, SurfaceCacheResources.DispatchParameters.Buffer, 0);
+
+							ComputeShader->UnsetParameters();
+						}
+
+						Swap(SurfaceCacheResources.Level[DepthLevel], SurfaceCacheResources.TempResources);
+					}
+				}
+
+				// Start from the highest depth, which will consider the fewest potential shading points
+				// Each level potentially prevents the next higher resolution level from doing expensive shading work
+				// This is the core of the surface cache algorithm
+				for (int32 DepthLevel = FMath::Min(GAOMaxLevel, MaxAllowedLevel); DepthLevel >= FMath::Max(GAOMinLevel, 0); DepthLevel--)
+				{
+					int32 DestLevelDownsampleFactor = GAODownsampleFactor * (1 << (DepthLevel * 2));
+
+					SCOPED_DRAW_EVENTF(Level, DEC_SCENE_ITEMS, TEXT("Level_%d"), DepthLevel);
+
+					TRefCountPtr<IPooledRenderTarget> DistanceFieldAOIrradianceCacheSplat;
+
+					{
+						FIntPoint AOBufferSize = FIntPoint::DivideAndRoundUp(GSceneRenderTargets.GetBufferSizeXY(), DestLevelDownsampleFactor);
+						FPooledRenderTargetDesc Desc(FPooledRenderTargetDesc::Create2DDesc(AOBufferSize, PF_FloatRGBA, TexCreate_None, TexCreate_RenderTargetable | TexCreate_UAV, false));
+						GRenderTargetPool.FindFreeElement(Desc, DistanceFieldAOIrradianceCacheSplat, TEXT("DistanceFieldAOIrradianceCacheSplat"));
+					}
+
+					// Splat / interpolate the surface cache records onto the buffer sized for the current depth level
+					RenderIrradianceCacheInterpolation(Views, DistanceFieldAOIrradianceCacheSplat, DistanceFieldNormal->GetRenderTargetItem(), DepthLevel, DestLevelDownsampleFactor, false);
+
+					{
+						SCOPED_DRAW_EVENT(PopulateIrradianceCache, DEC_SCENE_ITEMS);
+						RHISetRenderTarget(NULL, NULL);
+
+						// Save off the current record count before adding any more
+						{	
+							TShaderMapRef<FSaveStartIndexCS> ComputeShader(GetGlobalShaderMap());
+
+							RHISetComputeShader(ComputeShader->GetComputeShader());	
+							ComputeShader->SetParameters(View, DepthLevel);
+							DispatchComputeShader(*ComputeShader, 1, 1, 1);
+
+							ComputeShader->UnsetParameters();
+						}
+
+						// Create new records which haven't been shaded yet for shading points which don't have a valid interpolation from existing records
+						for (int32 ViewIndex = 0; ViewIndex < Views.Num(); ViewIndex++)
+						{
+							const FViewInfo& View = Views[ViewIndex];
+
+							FIntPoint DownsampledViewSize = FIntPoint::DivideAndRoundUp(View.ViewRect.Size(), DestLevelDownsampleFactor);
+							uint32 GroupSizeX = FMath::DivideAndRoundUp(DownsampledViewSize.X, GDistanceFieldAOTileSizeX);
+							uint32 GroupSizeY = FMath::DivideAndRoundUp(DownsampledViewSize.Y, GDistanceFieldAOTileSizeY);
+
+							TShaderMapRef<FPopulateCacheCS> ComputeShader(GetGlobalShaderMap());
+
+							RHISetComputeShader(ComputeShader->GetComputeShader());	
+							ComputeShader->SetParameters(View, DistanceFieldAOIrradianceCacheSplat->GetRenderTargetItem(), DistanceFieldNormal->GetRenderTargetItem(), DestLevelDownsampleFactor, DepthLevel, TileListGroupSize);
+							DispatchComputeShader(*ComputeShader, GroupSizeX, GroupSizeY, 1);
+
+							ComputeShader->UnsetParameters();
+						}
+					}
+
+					{
+						SCOPED_DRAW_EVENT(ShadeIrradianceCache, DEC_SCENE_ITEMS);
+
+						{	
+							TShaderMapRef<FSetupFinalGatherIndirectArgumentsCS> ComputeShader(GetGlobalShaderMap());
+
+							RHISetComputeShader(ComputeShader->GetComputeShader());	
+							ComputeShader->SetParameters(View, DepthLevel);
+							DispatchComputeShader(*ComputeShader, 1, 1, 1);
+
+							ComputeShader->UnsetParameters();
+						}
+
+						// Compute lighting for the new surface cache records by cone-stepping through the object distance fields
+						{
+							TShaderMapRef<FFinalGatherCS> ComputeShader(GetGlobalShaderMap());
+
+							RHISetComputeShader(ComputeShader->GetComputeShader());	
+							ComputeShader->SetParameters(View, DistanceFieldAOIrradianceCacheSplat->GetRenderTargetItem(), DistanceFieldNormal->GetRenderTargetItem(), NumObjects, DestLevelDownsampleFactor, DepthLevel, TileListGroupSize);
+							DispatchIndirectComputeShader(*ComputeShader, SurfaceCacheResources.DispatchParameters.Buffer, 0);
+
+							ComputeShader->UnsetParameters();
+						}
+					}
+				}
+
+				TRefCountPtr<IPooledRenderTarget> IrradianceCacheAccumulation;
+
+				{
+					FIntPoint ExpandedBufferSize = GetBufferSizeForAO();
+					FPooledRenderTargetDesc Desc(FPooledRenderTargetDesc::Create2DDesc(ExpandedBufferSize / GAODownsampleFactor, PF_FloatRGBA, TexCreate_None, TexCreate_RenderTargetable | TexCreate_UAV, false));
+					GRenderTargetPool.FindFreeElement(Desc, IrradianceCacheAccumulation, TEXT("DistanceFieldAOIrradianceCacheAccumulation"));
+				}
+
+				// Splat the surface cache records onto the opaque pixels, using less strict weighting so the lighting is smoothed in world space
+				{
+					SCOPED_DRAW_EVENT(FinalIrradianceCacheSplat, DEC_SCENE_ITEMS);
+					RenderIrradianceCacheInterpolation(Views, IrradianceCacheAccumulation, DistanceFieldNormal->GetRenderTargetItem(), 0, GAODownsampleFactor, true);
+				}
+
+				// Post process the AO to cover over artifacts
+				TRefCountPtr<IPooledRenderTarget> AOOutput;
+				PostProcessBentNormalAO(Views, IrradianceCacheAccumulation->GetRenderTargetItem(), AOOutput);
+
+				if (bApplyToSceneColor)
+				{
+					GSceneRenderTargets.BeginRenderingSceneColor(false);
+				}
+				else
+				{
+					RHISetRenderTarget(OutBentNormalAO.TargetableTexture, NULL);
+				}
+
+				// Upsample to full resolution, write to output
+				UpsampleBentNormalAO(Views, AOOutput);
+
+				return true;
+			}
+		}
+	}
+
+	return false;
+}
+
+class FDynamicSkyLightDiffusePS : public FGlobalShader
+{
+	DECLARE_SHADER_TYPE(FDynamicSkyLightDiffusePS, Global);
+public:
+
+	static bool ShouldCache(EShaderPlatform Platform)
+	{
+		return IsFeatureLevelSupported(Platform, ERHIFeatureLevel::SM4);
+	}
+
+	static void ModifyCompilationEnvironment(EShaderPlatform Platform, FShaderCompilerEnvironment& OutEnvironment)
+	{
+		FGlobalShader::ModifyCompilationEnvironment(Platform,OutEnvironment);
+		OutEnvironment.SetDefine(TEXT("DOWNSAMPLE_FACTOR"), GAODownsampleFactor);
+	}
+
+	/** Default constructor. */
+	FDynamicSkyLightDiffusePS() {}
+
+	/** Initialization constructor. */
+	FDynamicSkyLightDiffusePS(const ShaderMetaType::CompiledShaderInitializerType& Initializer)
+		: FGlobalShader(Initializer)
+	{
+		DeferredParameters.Bind(Initializer.ParameterMap);
+		LightAccumulationTexture.Bind(Initializer.ParameterMap, TEXT("BentNormalAOTexture"));
+		LightAccumulationSampler.Bind(Initializer.ParameterMap, TEXT("BentNormalAOSampler"));
+		BentNormalUVScale.Bind(Initializer.ParameterMap, TEXT("BentNormalUVScale"));
+		ApplyShadowing.Bind(Initializer.ParameterMap, TEXT("ApplyShadowing"));
+	}
+
+	void SetParameters(const FSceneView& View, FTextureRHIParamRef DynamicBentNormalAO, bool bApplyShadowing)
+	{
+		const FPixelShaderRHIParamRef ShaderRHI = GetPixelShader();
+		FGlobalShader::SetParameters(ShaderRHI, View);
+		DeferredParameters.Set(ShaderRHI, View);
+
+		SetTextureParameter(ShaderRHI, LightAccumulationTexture, LightAccumulationSampler, TStaticSamplerState<SF_Point>::GetRHI(), DynamicBentNormalAO);
+
+		SetShaderValue(ShaderRHI, BentNormalUVScale, 1.0f);
+		SetShaderValue(ShaderRHI, ApplyShadowing, bApplyShadowing ? 1.0f : 0.0f);
+	}
+
+	// FShader interface.
+	virtual bool Serialize(FArchive& Ar)
+	{
+		bool bShaderHasOutdatedParameters = FGlobalShader::Serialize(Ar);
+		Ar << DeferredParameters;
+		Ar << LightAccumulationTexture;
+		Ar << LightAccumulationSampler;
+		Ar << BentNormalUVScale;
+		Ar << ApplyShadowing;
+		return bShaderHasOutdatedParameters;
+	}
+
+private:
+
+	FDeferredPixelShaderParameters DeferredParameters;
+	FShaderResourceParameter LightAccumulationTexture;
+	FShaderResourceParameter LightAccumulationSampler;
+	FShaderParameter BentNormalUVScale;
+	FShaderParameter ApplyShadowing;
+};
+
+IMPLEMENT_SHADER_TYPE(,FDynamicSkyLightDiffusePS,TEXT("DistanceFieldSurfaceCacheLighting"),TEXT("SkyLightDiffusePS"),SF_Pixel);
+
+void FDeferredShadingSceneRenderer::RenderDynamicSkyLighting()
+{
+	if (Scene->SkyLight
+		&& Scene->SkyLight->ProcessedTexture
+		&& !Scene->SkyLight->bWantsStaticShadowing
+		&& !Scene->SkyLight->bHasStaticLighting
+		&& ViewFamily.EngineShowFlags.SkyLighting
+		&& GRHIFeatureLevel >= ERHIFeatureLevel::SM4
+		&& !IsSimpleDynamicLightingEnabled() 
+		&& !ViewFamily.EngineShowFlags.VisualizeLightCulling)
+	{
+		SCOPED_DRAW_EVENT(SkyLightDiffuse, DEC_SCENE_ITEMS);
+
+		bool bApplyShadowing = false;
+
+		TRefCountPtr<IPooledRenderTarget> LightAccumulation;
+		{
+			FPooledRenderTargetDesc Desc = GSceneRenderTargets.GetSceneColor()->GetDesc();
+			GRenderTargetPool.FindFreeElement(Desc, LightAccumulation, TEXT("LightAccumulation"));
+		}
+
+		if (Scene->SkyLight->bCastShadows
+			&& ViewFamily.EngineShowFlags.DistanceFieldAO 
+			&& !ViewFamily.EngineShowFlags.VisualizeDistanceFieldAO)
+		{
+			bApplyShadowing = RenderDistanceFieldAOSurfaceCache(LightAccumulation->GetRenderTargetItem(), false);
+		}
+
+		GSceneRenderTargets.BeginRenderingSceneColor();
+
+		for( int32 ViewIndex = 0; ViewIndex < Views.Num(); ViewIndex++ )
+		{
+			const FViewInfo& View = Views[ViewIndex];
+
+			RHISetViewport( View.ViewRect.Min.X, View.ViewRect.Min.Y, 0.0f, View.ViewRect.Max.X, View.ViewRect.Max.Y, 1.0f );
+			RHISetRasterizerState( TStaticRasterizerState< FM_Solid, CM_None >::GetRHI() );
+			RHISetDepthStencilState( TStaticDepthStencilState< false, CF_Always >::GetRHI() );
+			RHISetBlendState( TStaticBlendState< CW_RGB, BO_Add, BF_One, BF_One >::GetRHI() );
+
+			TShaderMapRef< FPostProcessVS > VertexShader( GetGlobalShaderMap() );
+			TShaderMapRef< FDynamicSkyLightDiffusePS > PixelShader( GetGlobalShaderMap() );
+
+			static FGlobalBoundShaderState BoundShaderState;
+			SetGlobalBoundShaderState( BoundShaderState, GFilterVertexDeclaration.VertexDeclarationRHI, *VertexShader, *PixelShader );
+
+			PixelShader->SetParameters( View, LightAccumulation->GetRenderTargetItem().ShaderResourceTexture, bApplyShadowing );
+
+			DrawRectangle( 
+				0, 0, 
+				View.ViewRect.Width(), View.ViewRect.Height(),
+				View.ViewRect.Min.X, View.ViewRect.Min.Y, 
+				View.ViewRect.Width(), View.ViewRect.Height(),
+				FIntPoint(View.ViewRect.Width(), View.ViewRect.Height()),
+				GSceneRenderTargets.GetBufferSizeXY(),
+				*VertexShader);
+		}
+	}
+}
\ No newline at end of file
