Distance field AO using a surface cache

diff --git a/Engine/Source/Runtime/Renderer/Private/LightFunctionRendering.cpp b/Engine/Source/Runtime/Renderer/Private/LightFunctionRendering.cpp
--- a/Engine/Source/Runtime/Renderer/Private/LightFunctionRendering.cpp
+++ b/Engine/Source/Runtime/Renderer/Private/LightFunctionRendering.cpp
@@ -46,7 +46,7 @@ public:
 		}
 
 		FVector4 StencilingSpherePosAndScale;
-		StencilingGeometry::CalcTransform(StencilingSpherePosAndScale, LightBounds, View->ViewMatrices.PreViewTranslation);
+		StencilingGeometry::GStencilSphereVertexBuffer.CalcTransform(StencilingSpherePosAndScale, LightBounds, View->ViewMatrices.PreViewTranslation);
 		StencilingGeometryParameters.Set(this, StencilingSpherePosAndScale);
 	}
 
diff --git a/Engine/Source/Runtime/Renderer/Private/LightPropagationVolume.cpp b/Engine/Source/Runtime/Renderer/Private/LightPropagationVolume.cpp
--- a/Engine/Source/Runtime/Renderer/Private/LightPropagationVolume.cpp
+++ b/Engine/Source/Runtime/Renderer/Private/LightPropagationVolume.cpp
@@ -470,6 +470,11 @@ public:
 
 		LinearTextureSampler.Bind(	Initializer.ParameterMap, TEXT("LinearSampler") );
 		PointTextureSampler.Bind(	Initializer.ParameterMap, TEXT("PointSampler") );
+
+		GvListBuffer.Bind(Initializer.ParameterMap,TEXT("GvListBuffer"));
+		GvListHeadBuffer.Bind(Initializer.ParameterMap,TEXT("GvListHeadBuffer"));
+		VplListBuffer.Bind(Initializer.ParameterMap,TEXT("VplListBuffer"));
+		VplListHeadBuffer.Bind(Initializer.ParameterMap,TEXT("VplListHeadBuffer"));
 	}
 
 	void SetParameters(
@@ -505,6 +510,10 @@ public:
 		Ar << RsmDepthTexture;
 		Ar << LinearTextureSampler;
 		Ar << PointTextureSampler;
+		Ar << GvListBuffer;
+		Ar << GvListHeadBuffer;
+		Ar << VplListBuffer;
+		Ar << VplListHeadBuffer;
 		return rv;
 	}
 protected:
@@ -514,6 +523,11 @@ protected:
 
 	FShaderResourceParameter LinearTextureSampler;
 	FShaderResourceParameter PointTextureSampler;
+
+	FRWShaderParameter GvListBuffer;
+	FRWShaderParameter GvListHeadBuffer;
+	FRWShaderParameter VplListBuffer;
+	FRWShaderParameter VplListHeadBuffer;
 };
 IMPLEMENT_SHADER_TYPE(,FLpvInject_GenerateVplListsCS,TEXT("LPVInject_GenerateVplLists"),TEXT("CSGenerateVplLists_LightDirectional"),SF_Compute);
 
diff --git a/Engine/Source/Runtime/Renderer/Private/LightRendering.h b/Engine/Source/Runtime/Renderer/Private/LightRendering.h
--- a/Engine/Source/Runtime/Renderer/Private/LightRendering.h
+++ b/Engine/Source/Runtime/Renderer/Private/LightRendering.h
@@ -182,7 +182,7 @@ public:
 		FGlobalShader::SetParameters(GetVertexShader(),View);
 
 		FVector4 StencilingSpherePosAndScale;
-		StencilingGeometry::CalcTransform(StencilingSpherePosAndScale, LightBounds, View.ViewMatrices.PreViewTranslation);
+		StencilingGeometry::GStencilSphereVertexBuffer.CalcTransform(StencilingSpherePosAndScale, LightBounds, View.ViewMatrices.PreViewTranslation);
 		StencilingGeometryParameters.Set(this, StencilingSpherePosAndScale);
 	}
 
diff --git a/Engine/Source/Runtime/Renderer/Private/LightSceneInfo.cpp b/Engine/Source/Runtime/Renderer/Private/LightSceneInfo.cpp
--- a/Engine/Source/Runtime/Renderer/Private/LightSceneInfo.cpp
+++ b/Engine/Source/Runtime/Renderer/Private/LightSceneInfo.cpp
@@ -171,7 +171,7 @@ bool FLightSceneInfo::ShouldRenderLight(const FViewInfo& View) const
 
 	return bLocalVisible
 		// Only render lights with static shadowing for reflection captures, since they are only captured at edit time
-		&& (!View.bIsReflectionCapture || Proxy->HasStaticShadowing());
+		&& (!View.bStaticSceneOnly || Proxy->HasStaticShadowing());
 }
 
 void FLightSceneInfo::ReleaseRHI()
diff --git a/Engine/Source/Runtime/Renderer/Private/ReflectionEnvironmentCapture.cpp b/Engine/Source/Runtime/Renderer/Private/ReflectionEnvironmentCapture.cpp
--- a/Engine/Source/Runtime/Renderer/Private/ReflectionEnvironmentCapture.cpp
+++ b/Engine/Source/Runtime/Renderer/Private/ReflectionEnvironmentCapture.cpp
@@ -1311,7 +1311,7 @@ private:
 
 TGlobalResource<FCaptureRenderTarget> GReflectionCaptureRenderTarget;
 
-void CaptureSceneIntoScratchCubemap(FScene* Scene, FVector CapturePosition, bool bCapturingForSkyLight, float SkyLightNearPlane, bool bLowerHemisphereIsBlack, bool bCaptureEmissiveOnly)
+void CaptureSceneIntoScratchCubemap(FScene* Scene, FVector CapturePosition, bool bCapturingForSkyLight, bool bStaticSceneOnly, float SkyLightNearPlane, bool bLowerHemisphereIsBlack, bool bCaptureEmissiveOnly)
 {
 	for (int32 CubeFace = 0; CubeFace < CubeFace_MAX; CubeFace++)
 	{
@@ -1342,6 +1342,7 @@ void CaptureSceneIntoScratchCubemap(FScene* Scene, FVector CapturePosition, bool
 		ViewFamily.EngineShowFlags.CompositeEditorPrimitives = 0;
 		// These are highly dynamic and can't be captured effectively
 		ViewFamily.EngineShowFlags.LightShafts = 0;
+		ViewFamily.EngineShowFlags.VisualizeDistanceFieldAO = 0;
 
 		// Don't apply sky lighting diffuse when capturing the sky light source, or we would have feedback
 		ViewFamily.EngineShowFlags.SkyLighting = !bCapturingForSkyLight;
@@ -1377,6 +1378,7 @@ void CaptureSceneIntoScratchCubemap(FScene* Scene, FVector CapturePosition, bool
 		}
 
 		View->bIsReflectionCapture = true;
+		View->bStaticSceneOnly = bStaticSceneOnly;
 		View->StartFinalPostprocessSettings(CapturePosition);
 		View->EndFinalPostprocessSettings();
 
@@ -1472,7 +1474,7 @@ void FScene::UpdateReflectionCaptureContents(UReflectionCaptureComponent* Captur
 				ClearScratchCubemaps();
 			});
 			
-			CaptureSceneIntoScratchCubemap(this, CaptureComponent->GetComponentLocation(), false, 0, false, false);
+			CaptureSceneIntoScratchCubemap(this, CaptureComponent->GetComponentLocation(), false, true, 0, false, false);
 
 			ENQUEUE_UNIQUE_RENDER_COMMAND( 
 				FilterCommand,
@@ -1527,7 +1529,7 @@ void CopyToSkyTexture(FScene* Scene, FTexture* ProcessedTexture)
 
 void FScene::UpdateSkyCaptureContents(const USkyLightComponent* CaptureComponent, bool bCaptureEmissiveOnly, FTexture* OutProcessedTexture, FSHVectorRGB3& OutIrradianceEnvironmentMap)
 {
-	if (IsReflectionEnvironmentAvailable(FeatureLevel))
+	if (FeatureLevel >= ERHIFeatureLevel::SM4)
 	{
 		ENQUEUE_UNIQUE_RENDER_COMMAND( 
 			ClearCommand,
@@ -1538,7 +1540,8 @@ void FScene::UpdateSkyCaptureContents(const USkyLightComponent* CaptureComponent
 
 		if (CaptureComponent->SourceType == SLS_CapturedScene)
 		{
-			CaptureSceneIntoScratchCubemap(this, CaptureComponent->GetComponentLocation(), true, CaptureComponent->SkyDistanceThreshold, CaptureComponent->bLowerHemisphereIsBlack, bCaptureEmissiveOnly);
+			bool bStaticSceneOnly = CaptureComponent->Mobility != EComponentMobility::Movable;
+			CaptureSceneIntoScratchCubemap(this, CaptureComponent->GetComponentLocation(), true, bStaticSceneOnly, CaptureComponent->SkyDistanceThreshold, CaptureComponent->bLowerHemisphereIsBlack, bCaptureEmissiveOnly);
 		}
 		else if (CaptureComponent->SourceType == SLS_SpecifiedCubemap)
 		{
diff --git a/Engine/Source/Runtime/Renderer/Private/RendererScene.cpp b/Engine/Source/Runtime/Renderer/Private/RendererScene.cpp
--- a/Engine/Source/Runtime/Renderer/Private/RendererScene.cpp
+++ b/Engine/Source/Runtime/Renderer/Private/RendererScene.cpp
@@ -8,6 +8,7 @@
 #include "ScenePrivate.h"
 #include "ShaderCompiler.h"
 #include "StaticMeshResources.h"
+#include "DistanceFieldSurfaceCacheLighting.h"
 
 
 // Enable this define to do slow checks for components being added to the wrong
@@ -23,6 +24,46 @@ TGlobalResource< FGlobalDistanceCullFadeUniformBuffer > GDistanceCullFadedInUnif
 
 SIZE_T FStaticMeshDrawListBase::TotalBytesUsed = 0;
 
+/** Default constructor. */
+FSceneViewState::FSceneViewState()
+	: OcclusionQueryPool(RQT_Occlusion)
+{
+	LastRenderTime = -FLT_MAX;
+	LastRenderTimeDelta = 0.0f;
+	MotionBlurTimeScale = 1.0f;
+	PrevViewMatrixForOcclusionQuery.SetIdentity();
+	PrevViewOriginForOcclusionQuery = FVector::ZeroVector;
+#if !(UE_BUILD_SHIPPING || UE_BUILD_TEST)
+	bIsFreezing = false;
+	bIsFrozen = false;
+#endif
+	// Register this object as a resource, so it will receive device reset notifications.
+	if ( IsInGameThread() )
+	{
+		BeginInitResource(this);
+	}
+	else
+	{
+		InitResource();
+	}
+	CachedVisibilityChunk = NULL;
+	CachedVisibilityHandlerId = INDEX_NONE;
+	CachedVisibilityBucketIndex = INDEX_NONE;
+	CachedVisibilityChunkIndex = INDEX_NONE;
+	MIDUsedCount = 0;
+	TemporalAASampleIndex = 0;
+	TemporalAASampleCount = 1;
+	AOTileIntersectionResources = NULL;
+	bBokehDOFHistory = true;
+
+	LightPropagationVolume = NULL; 
+
+	for (int32 CascadeIndex = 0; CascadeIndex < ARRAY_COUNT(TranslucencyLightingCacheAllocations); CascadeIndex++)
+	{
+		TranslucencyLightingCacheAllocations[CascadeIndex] = NULL;
+	}
+}
+
 /**
  * Sets the FX system associated with the scene.
  */
@@ -151,6 +192,7 @@ FScene::FScene(UWorld* InWorld, bool bInRequiresHitProxies, bool bInIsEditorScen
 ,	SimpleDirectionalLight(NULL)
 ,	SunLight(NULL)
 ,	IndirectLightingCache()
+,	SurfaceCacheResources(NULL)
 ,	PreshadowCacheLayout(0, 0, 0, 0, false, false)
 ,	AtmosphericFog(NULL)
 ,	PrecomputedVisibilityHandler(NULL)
@@ -192,6 +234,13 @@ FScene::~FScene()
 	ReflectionSceneData.CubemapArray.ReleaseResource();
 	IndirectLightingCache.ReleaseResource();
 
+	if (SurfaceCacheResources)
+	{
+		SurfaceCacheResources->ReleaseResource();
+		delete SurfaceCacheResources;
+		SurfaceCacheResources = NULL;
+	}
+
 	if (AtmosphericFog)
 	{
 		delete AtmosphericFog;
diff --git a/Engine/Source/Runtime/Renderer/Private/SceneOcclusion.h b/Engine/Source/Runtime/Renderer/Private/SceneOcclusion.h
--- a/Engine/Source/Runtime/Renderer/Private/SceneOcclusion.h
+++ b/Engine/Source/Runtime/Renderer/Private/SceneOcclusion.h
@@ -33,7 +33,7 @@ public:
 		FGlobalShader::SetParameters(GetVertexShader(), View);
 
 		FVector4 StencilingSpherePosAndScale;
-		StencilingGeometry::CalcTransform(StencilingSpherePosAndScale, BoundingSphere, View.ViewMatrices.PreViewTranslation);
+		StencilingGeometry::GStencilSphereVertexBuffer.CalcTransform(StencilingSpherePosAndScale, BoundingSphere, View.ViewMatrices.PreViewTranslation);
 		StencilingGeometryParameters.Set(this, StencilingSpherePosAndScale);
 	}
 
diff --git a/Engine/Source/Runtime/Renderer/Private/ScenePrivate.h b/Engine/Source/Runtime/Renderer/Private/ScenePrivate.h
--- a/Engine/Source/Runtime/Renderer/Private/ScenePrivate.h
+++ b/Engine/Source/Runtime/Renderer/Private/ScenePrivate.h
@@ -475,6 +475,7 @@ public:
 	TRefCountPtr<IPooledRenderTarget> LightShaftOcclusionHistoryRT;
 	// Temporal AA result for light shafts of last frame
 	TMap<const ULightComponent*, TRefCountPtr<IPooledRenderTarget> > LightShaftBloomHistoryRTs;
+	TRefCountPtr<IPooledRenderTarget> DistanceFieldAOHistoryRT;
 	// Mobile temporal AA surfaces.
 	TRefCountPtr<IPooledRenderTarget> MobileAaBloomSunVignette0;
 	TRefCountPtr<IPooledRenderTarget> MobileAaBloomSunVignette1;
@@ -485,6 +486,9 @@ public:
 	FTextureRHIRef SelectionOutlineCacheKey;
 	TRefCountPtr<FRHIShaderResourceView> SelectionOutlineCacheValue;
 
+	/** Distance field AO tile intersection GPU resources.  Last frame's state is not used, but they must be sized exactly to the view so stored here. */
+	class FTileIntersectionResources* AOTileIntersectionResources;
+
 	// Is DOFHistoryRT set from Bokeh DOF?
 	bool bBokehDOFHistory;
 
@@ -532,44 +536,9 @@ public:
 	FLightPropagationVolume* GetLightPropagationVolume() const { return LightPropagationVolume; }
 
 	/** Default constructor. */
-    FSceneViewState()
-		: OcclusionQueryPool(RQT_Occlusion)
-    {
-		LastRenderTime = -FLT_MAX;
-		LastRenderTimeDelta = 0.0f;
-		MotionBlurTimeScale = 1.0f;
-		PrevViewMatrixForOcclusionQuery.SetIdentity();
-		PrevViewOriginForOcclusionQuery = FVector::ZeroVector;
-#if !(UE_BUILD_SHIPPING || UE_BUILD_TEST)
-		bIsFreezing = false;
-		bIsFrozen = false;
-#endif
-		// Register this object as a resource, so it will receive device reset notifications.
-		if ( IsInGameThread() )
-		{
-			BeginInitResource(this);
-		}
-		else
-		{
-			InitResource();
-		}
-		CachedVisibilityChunk = NULL;
-		CachedVisibilityHandlerId = INDEX_NONE;
-		CachedVisibilityBucketIndex = INDEX_NONE;
-		CachedVisibilityChunkIndex = INDEX_NONE;
-		MIDUsedCount = 0;
-		TemporalAASampleIndex = 0;
-		TemporalAASampleCount = 1;
-		bBokehDOFHistory = true;
-
-		LightPropagationVolume = NULL; 
-
-		for (int32 CascadeIndex = 0; CascadeIndex < ARRAY_COUNT(TranslucencyLightingCacheAllocations); CascadeIndex++)
-		{
-			TranslucencyLightingCacheAllocations[CascadeIndex] = NULL;
-		}
-    }
+    FSceneViewState();
 
+	void DestroyAOTileResources();
 	void DestroyLightPropagationVolume();
 
 	virtual ~FSceneViewState()
@@ -581,6 +550,7 @@ public:
 			delete TranslucencyLightingCacheAllocations[CascadeIndex];
 		}
 
+		DestroyAOTileResources();
 		DestroyLightPropagationVolume();
 	}
 
@@ -656,6 +626,7 @@ public:
 		SSRHistoryRT.SafeRelease();
 		LightShaftOcclusionHistoryRT.SafeRelease();
 		LightShaftBloomHistoryRTs.Empty();
+		DistanceFieldAOHistoryRT.SafeRelease();
 		MobileAaBloomSunVignette0.SafeRelease();
 		MobileAaBloomSunVignette1.SafeRelease();
 		MobileAaColor0.SafeRelease();
@@ -1230,6 +1201,9 @@ public:
 	/** Interpolates and caches indirect lighting for dynamic objects. */
 	FIndirectLightingCache IndirectLightingCache;
 
+	/** Scene state of distance field AO.  NULL if the scene doesn't use the feature. */
+	class FSurfaceCacheResources* SurfaceCacheResources;
+
 	/** Preshadows that are currently cached in the PreshadowCache render target. */
 	TArray<TRefCountPtr<FProjectedShadowInfo> > CachedPreshadows;
 
diff --git a/Engine/Source/Runtime/Renderer/Private/SceneVisibility.cpp b/Engine/Source/Runtime/Renderer/Private/SceneVisibility.cpp
--- a/Engine/Source/Runtime/Renderer/Private/SceneVisibility.cpp
+++ b/Engine/Source/Runtime/Renderer/Private/SceneVisibility.cpp
@@ -1324,7 +1324,7 @@ void FSceneRenderer::ComputeViewVisibility()
 			}
 		}
 
-		if (View.bIsReflectionCapture)
+		if (View.bStaticSceneOnly)
 		{
 			for (FSceneSetBitIterator BitIt(View.PrimitiveVisibilityMap); BitIt; ++BitIt)
 			{
diff --git a/Engine/Source/Runtime/Renderer/Private/ShadowRendering.cpp b/Engine/Source/Runtime/Renderer/Private/ShadowRendering.cpp
--- a/Engine/Source/Runtime/Renderer/Private/ShadowRendering.cpp
+++ b/Engine/Source/Runtime/Renderer/Private/ShadowRendering.cpp
@@ -410,6 +410,10 @@ public:
 		TransmissionStrength.Bind(Initializer.ParameterMap,TEXT("TransmissionStrength"));
 		ReflectiveShadowMapTextureResolution.Bind(Initializer.ParameterMap,TEXT("ReflectiveShadowMapTextureResolution"));
 		ProjectionMatrixParameter.Bind(Initializer.ParameterMap,TEXT("ProjectionMatrix"));
+		GvListBuffer.Bind(Initializer.ParameterMap,TEXT("GvListBuffer"));
+		GvListHeadBuffer.Bind(Initializer.ParameterMap,TEXT("GvListHeadBuffer"));
+		VplListBuffer.Bind(Initializer.ParameterMap,TEXT("VplListBuffer"));
+		VplListHeadBuffer.Bind(Initializer.ParameterMap,TEXT("VplListHeadBuffer"));
 	}
 
 	TShadowDepthBasePS() {}
@@ -465,6 +469,10 @@ public:
 		Ar << TransmissionStrength;
 		Ar << ReflectiveShadowMapTextureResolution;
 		Ar << ProjectionMatrixParameter;
+		Ar << GvListBuffer;
+		Ar << GvListHeadBuffer;
+		Ar << VplListBuffer;
+		Ar << VplListHeadBuffer;
 
 		return bShaderHasOutdatedParameters;
 	}
@@ -476,6 +484,11 @@ private:
 	FShaderParameter TransmissionStrength;
 	FShaderParameter ReflectiveShadowMapTextureResolution;
 	FShaderParameter ProjectionMatrixParameter;
+
+	FRWShaderParameter GvListBuffer;
+	FRWShaderParameter GvListHeadBuffer;
+	FRWShaderParameter VplListBuffer;
+	FRWShaderParameter VplListHeadBuffer;
 };
 
 enum EShadowDepthPixelShaderMode
@@ -549,10 +562,13 @@ IMPLEMENT_SHADOWDEPTHPASS_PIXELSHADER_TYPE(PixelShadowDepth_OnePassPointLight, t
 IMPLEMENT_SHADOWDEPTHPASS_PIXELSHADER_TYPE(PixelShadowDepth_OnePassPointLight, false);
 
 /** The stencil sphere vertex buffer. */
-TGlobalResource<StencilingGeometry::FStencilSphereVertexBuffer> StencilingGeometry::GStencilSphereVertexBuffer;
+TGlobalResource<StencilingGeometry::TStencilSphereVertexBuffer<18, 12> > StencilingGeometry::GStencilSphereVertexBuffer;
 
 /** The stencil sphere index buffer. */
-TGlobalResource<StencilingGeometry::FStencilSphereIndexBuffer> StencilingGeometry::GStencilSphereIndexBuffer;
+TGlobalResource<StencilingGeometry::TStencilSphereIndexBuffer<18, 12> > StencilingGeometry::GStencilSphereIndexBuffer;
+
+TGlobalResource<StencilingGeometry::TStencilSphereVertexBuffer<4, 4> > StencilingGeometry::GLowPolyStencilSphereVertexBuffer;
+TGlobalResource<StencilingGeometry::TStencilSphereIndexBuffer<4, 4> > StencilingGeometry::GLowPolyStencilSphereIndexBuffer;
 
 /** The (dummy) stencil cone vertex buffer. */
 TGlobalResource<StencilingGeometry::FStencilConeVertexBuffer> StencilingGeometry::GStencilConeVertexBuffer;
@@ -1415,22 +1431,6 @@ void FProjectedShadowInfo::RenderDepth(FDeferredShadingSceneRenderer* SceneRende
 	FoundView->ViewMatrices.ViewMatrix = OriginalViewMatrix;
 }
 
-void StencilingGeometry::CalcTransform(FVector4& OutPosAndScale, const FSphere& Sphere, const FVector& PreViewTranslation, bool bConservativelyBoundSphere)
-{
-	float Radius = Sphere.W;
-	if (bConservativelyBoundSphere)
-	{
-		const int32 NumRings = StencilingGeometry::NUM_SPHERE_RINGS;
-		const float RadiansPerRingSegment = PI / (float)NumRings;
-
-		// Boost the effective radius so that the edges of the sphere approximation lie on the sphere, instead of the vertices
-		Radius /= FMath::Cos(RadiansPerRingSegment);
-	}
-
-	const FVector Translate(Sphere.Center + PreViewTranslation);
-	OutPosAndScale = FVector4(Translate, Radius);
-}
-
 void StencilingGeometry::DrawSphere()
 {
 	RHISetStreamSource(0, StencilingGeometry::GStencilSphereVertexBuffer.VertexBufferRHI, sizeof(FVector4), 0);
diff --git a/Engine/Source/Runtime/Renderer/Private/ShadowRendering.h b/Engine/Source/Runtime/Renderer/Private/ShadowRendering.h
--- a/Engine/Source/Runtime/Renderer/Private/ShadowRendering.h
+++ b/Engine/Source/Runtime/Renderer/Private/ShadowRendering.h
@@ -13,19 +13,6 @@ class FProjectedShadowInfo;
 /** Utility functions for drawing a sphere */
 namespace StencilingGeometry
 {
-	const int32 NUM_SPHERE_SIDES = 18;
-	const int32 NUM_SPHERE_RINGS = 12;
-
-	/** 
-	* Calculates the world transform for a sphere.
-	* @param OutTransform - The output world transform.
-	* @param Sphere - The sphere to generate the transform for.
-	* @param PreViewTranslation - The pre-view translation to apply to the transform.
-	* @param bConservativelyBoundSphere - when true, the sphere that is drawn will contain all positions in the analytical sphere,
-	*		 Otherwise the sphere vertices will lie on the analytical sphere and the positions on the faces will lie inside the sphere.
-	*/
-	extern void CalcTransform(FVector4& OutPosAndScale, const FSphere& Sphere, const FVector& PreViewTranslation, bool bConservativelyBoundSphere = true);
-
 	/**
 	* Draws a sphere using RHIDrawIndexedPrimitive, useful as approximate bounding geometry for deferred passes.
 	* Note: The sphere will be of unit size unless transformed by the shader. 
@@ -37,16 +24,23 @@ namespace StencilingGeometry
 	/** 
 	* Vertex buffer for a sphere of unit size. Used for drawing a sphere as approximate bounding geometry for deferred passes.
 	*/
-	class FStencilSphereVertexBuffer : public FVertexBuffer
+	template<int32 NumSphereSides, int32 NumSphereRings>
+	class TStencilSphereVertexBuffer : public FVertexBuffer
 	{
 	public:
+
+		int32 GetNumRings() const
+		{
+			return NumSphereRings;
+		}
+
 		/** 
 		* Initialize the RHI for this rendering resource 
 		*/
 		void InitRHI() OVERRIDE
 		{
-			const int32 NumSides = StencilingGeometry::NUM_SPHERE_SIDES;
-			const int32 NumRings = StencilingGeometry::NUM_SPHERE_RINGS;
+			const int32 NumSides = NumSphereSides;
+			const int32 NumRings = NumSphereRings;
 			const int32 NumVerts = (NumSides + 1) * (NumRings + 1);
 
 			const float RadiansPerRingSegment = PI / (float)NumRings;
@@ -87,6 +81,30 @@ namespace StencilingGeometry
 
 		int32 GetVertexCount() const { return NumSphereVerts; }
 
+	/** 
+	* Calculates the world transform for a sphere.
+	* @param OutTransform - The output world transform.
+	* @param Sphere - The sphere to generate the transform for.
+	* @param PreViewTranslation - The pre-view translation to apply to the transform.
+	* @param bConservativelyBoundSphere - when true, the sphere that is drawn will contain all positions in the analytical sphere,
+	*		 Otherwise the sphere vertices will lie on the analytical sphere and the positions on the faces will lie inside the sphere.
+	*/
+	void CalcTransform(FVector4& OutPosAndScale, const FSphere& Sphere, const FVector& PreViewTranslation, bool bConservativelyBoundSphere = true)
+	{
+		float Radius = Sphere.W;
+		if (bConservativelyBoundSphere)
+		{
+			const int32 NumRings = NumSphereRings;
+			const float RadiansPerRingSegment = PI / (float)NumRings;
+
+			// Boost the effective radius so that the edges of the sphere approximation lie on the sphere, instead of the vertices
+			Radius /= FMath::Cos(RadiansPerRingSegment);
+		}
+
+		const FVector Translate(Sphere.Center + PreViewTranslation);
+		OutPosAndScale = FVector4(Translate, Radius);
+	}
+
 	private:
 		int32 NumSphereVerts;
 	};
@@ -94,7 +112,8 @@ namespace StencilingGeometry
 	/** 
 	* Stenciling sphere index buffer
 	*/
-	class FStencilSphereIndexBuffer : public FIndexBuffer
+	template<int32 NumSphereSides, int32 NumSphereRings>
+	class TStencilSphereIndexBuffer : public FIndexBuffer
 	{
 	public:
 		/** 
@@ -102,8 +121,8 @@ namespace StencilingGeometry
 		*/
 		void InitRHI() OVERRIDE
 		{
-			const int32 NumSides = StencilingGeometry::NUM_SPHERE_SIDES;
-			const int32 NumRings = StencilingGeometry::NUM_SPHERE_RINGS;
+			const int32 NumSides = NumSphereSides;
+			const int32 NumRings = NumSphereRings;
 			TResourceArray<uint16, INDEXBUFFER_ALIGNMENT> Indices;
 
 			// Add triangles for all the vertices generated
@@ -239,8 +258,10 @@ namespace StencilingGeometry
 		int32 GetVertexCount() const { return NumVerts; }
 	};
 
-	extern TGlobalResource<FStencilSphereVertexBuffer> GStencilSphereVertexBuffer;
-	extern TGlobalResource<FStencilSphereIndexBuffer> GStencilSphereIndexBuffer;
+	extern TGlobalResource<TStencilSphereVertexBuffer<18, 12> > GStencilSphereVertexBuffer;
+	extern TGlobalResource<TStencilSphereIndexBuffer<18, 12> > GStencilSphereIndexBuffer;
+	extern TGlobalResource<TStencilSphereVertexBuffer<4, 4> > GLowPolyStencilSphereVertexBuffer;
+	extern TGlobalResource<TStencilSphereIndexBuffer<4, 4> > GLowPolyStencilSphereIndexBuffer;
 	extern TGlobalResource<FStencilConeVertexBuffer> GStencilConeVertexBuffer;
 	extern TGlobalResource<FStencilConeIndexBuffer> GStencilConeIndexBuffer;
 
@@ -800,7 +821,7 @@ public:
 		FVector4 GeometryPosAndScale;
 		if(LightSceneInfo->Proxy->GetLightType() == LightType_Point)
 		{
-			StencilingGeometry::CalcTransform(GeometryPosAndScale, LightSceneInfo->Proxy->GetBoundingSphere(), View.ViewMatrices.PreViewTranslation);
+			StencilingGeometry::GStencilSphereVertexBuffer.CalcTransform(GeometryPosAndScale, LightSceneInfo->Proxy->GetBoundingSphere(), View.ViewMatrices.PreViewTranslation);
 			SetShaderValue(Shader->GetVertexShader(), StencilGeometryPosAndScale, GeometryPosAndScale);
 			SetShaderValue(Shader->GetVertexShader(), StencilConeParameters, FVector4(0.0f, 0.0f, 0.0f, 0.0f));
 		}
diff --git a/Engine/Source/Runtime/Renderer/Private/ShadowSetup.cpp b/Engine/Source/Runtime/Renderer/Private/ShadowSetup.cpp
--- a/Engine/Source/Runtime/Renderer/Private/ShadowSetup.cpp
+++ b/Engine/Source/Runtime/Renderer/Private/ShadowSetup.cpp
@@ -1026,7 +1026,7 @@ bool FDeferredShadingSceneRenderer::ShouldCreateObjectShadowForStationaryLight(c
 void FDeferredShadingSceneRenderer::SetupInteractionShadows(
 	FLightPrimitiveInteraction* Interaction, 
 	FVisibleLightInfo& VisibleLightInfo, 
-	bool bReflectionCaptureScene,
+	bool bStaticSceneOnly,
 	const TArray<FProjectedShadowInfo*,SceneRenderingAllocator>& ViewDependentWholeSceneShadows,
 	TArray<FProjectedShadowInfo*,SceneRenderingAllocator>& PreShadows)
 {
@@ -1051,7 +1051,7 @@ void FDeferredShadingSceneRenderer::SetupInteractionShadows(
 
 		if (Interaction->HasShadow() 
 			// Only render shadows from objects that use static lighting during a reflection capture, since the reflection capture doesn't update at runtime
-			&& (!bReflectionCaptureScene || PrimitiveSceneInfo->Proxy->HasStaticLighting())
+			&& (!bStaticSceneOnly || PrimitiveSceneInfo->Proxy->HasStaticLighting())
 			&& (bCreateTranslucentObjectShadow || bCreateInsetObjectShadow || bCreateObjectShadowForStationaryLight))
 		{
 			// Create projected shadow infos
@@ -1424,7 +1424,7 @@ void FDeferredShadingSceneRenderer::CreateWholeSceneProjectedShadow(FLightSceneI
 		uint32 MaxUnclampedResolution	= 0;
 		TArray<float, TInlineAllocator<2> > FadeAlphas;
 		float MaxFadeAlpha = 0;
-		bool bReflectionCaptureScene = false;
+		bool bStaticSceneOnly = false;
 
 		for(int32 ViewIndex = 0;ViewIndex < Views.Num();ViewIndex++)
 		{
@@ -1455,7 +1455,7 @@ void FDeferredShadingSceneRenderer::CreateWholeSceneProjectedShadow(FLightSceneI
 					)
 				);
 
-			bReflectionCaptureScene = bReflectionCaptureScene || View.bIsReflectionCapture;
+			bStaticSceneOnly = bStaticSceneOnly || View.bStaticSceneOnly;
 
 			const float FadeAlpha = CalculateShadowFadeAlpha( MaxUnclampedResolution, ShadowFadeResolution, MinShadowResolution );
 			MaxFadeAlpha = FMath::Max(MaxFadeAlpha, FadeAlpha);
@@ -1539,7 +1539,7 @@ void FDeferredShadingSceneRenderer::CreateWholeSceneProjectedShadow(FLightSceneI
 					Interaction = Interaction->GetNextPrimitive())
 				{
 					if (Interaction->HasShadow() 
-						&& (!bReflectionCaptureScene || Interaction->GetPrimitiveSceneInfo()->Proxy->HasStaticLighting()))
+						&& (!bStaticSceneOnly || Interaction->GetPrimitiveSceneInfo()->Proxy->HasStaticLighting()))
 					{
 						ProjectedShadowInfo->AddSubjectPrimitive(Interaction->GetPrimitiveSceneInfo(), &Views);
 					}
@@ -1687,7 +1687,7 @@ inline void FDeferredShadingSceneRenderer::GatherShadowsForPrimitiveInner(
 	const FPrimitiveSceneInfoCompact& PrimitiveSceneInfoCompact, 
 	const TArray<FProjectedShadowInfo*,SceneRenderingAllocator>& PreShadows,
 	const TArray<FProjectedShadowInfo*,SceneRenderingAllocator>& ViewDependentWholeSceneShadows,
-	bool bReflectionCaptureScene)
+	bool bStaticSceneOnly)
 {
 	if(PrimitiveSceneInfoCompact.bCastDynamicShadow)
 	{
@@ -1776,7 +1776,7 @@ inline void FDeferredShadingSceneRenderer::GatherShadowsForPrimitiveInner(
 							// Exclude primitives that will create a per-object shadow from a stationary light
 							&& !ShouldCreateObjectShadowForStationaryLight(ProjectedShadowInfo->LightSceneInfo, PrimitiveSceneInfo->Proxy, true)
 							// Only render shadows from objects that use static lighting during a reflection capture, since the reflection capture doesn't update at runtime
-							&& (!bReflectionCaptureScene || PrimitiveProxy->HasStaticLighting()) 
+							&& (!bStaticSceneOnly || PrimitiveProxy->HasStaticLighting()) 
 							&& !bScreenSpaceSizeCulled )
 						{
 							// Add this primitive to the shadow.
@@ -1792,7 +1792,7 @@ inline void FDeferredShadingSceneRenderer::GatherShadowsForPrimitiveInner(
 void FDeferredShadingSceneRenderer::GatherShadowPrimitives(
 	const TArray<FProjectedShadowInfo*,SceneRenderingAllocator>& PreShadows,
 	const TArray<FProjectedShadowInfo*,SceneRenderingAllocator>& ViewDependentWholeSceneShadows,
-	bool bReflectionCaptureScene
+	bool bStaticSceneOnly
 	)
 {
 	SCOPE_CYCLE_COUNTER(STAT_GatherShadowPrimitivesTime);
@@ -1878,7 +1878,7 @@ void FDeferredShadingSceneRenderer::GatherShadowPrimitives(
 			for(FScenePrimitiveOctree::ElementConstIt NodePrimitiveIt(PrimitiveOctreeNode.GetElementIt());NodePrimitiveIt;++NodePrimitiveIt)
 			{
 				// gather the shadows for this one primitive
-				GatherShadowsForPrimitiveInner(*NodePrimitiveIt, PreShadows, ViewDependentWholeSceneShadows, bReflectionCaptureScene);
+				GatherShadowsForPrimitiveInner(*NodePrimitiveIt, PreShadows, ViewDependentWholeSceneShadows, bStaticSceneOnly);
 			}
 		}
 
@@ -1896,12 +1896,12 @@ void FDeferredShadingSceneRenderer::InitDynamicShadows()
 {
 	SCOPE_CYCLE_COUNTER(STAT_DynamicShadowSetupTime);
 
-	bool bReflectionCaptureScene = false;
+	bool bStaticSceneOnly = false;
 
 	for (int32 ViewIndex = 0; ViewIndex < Views.Num(); ViewIndex++)
 	{
 		FViewInfo& View = Views[ViewIndex];
-		bReflectionCaptureScene = bReflectionCaptureScene || View.bIsReflectionCapture;
+		bStaticSceneOnly = bStaticSceneOnly || View.bStaticSceneOnly;
 	}
 
 	TArray<FProjectedShadowInfo*,SceneRenderingAllocator> PreShadows;
@@ -2053,7 +2053,7 @@ void FDeferredShadingSceneRenderer::InitDynamicShadows()
 							Interaction = Interaction->GetNextPrimitive()
 							)
 						{
-							SetupInteractionShadows(Interaction, VisibleLightInfo, bReflectionCaptureScene, ViewDependentWholeSceneShadows, PreShadows);
+							SetupInteractionShadows(Interaction, VisibleLightInfo, bStaticSceneOnly, ViewDependentWholeSceneShadows, PreShadows);
 						}
 					}
 				}
@@ -2068,5 +2068,5 @@ void FDeferredShadingSceneRenderer::InitDynamicShadows()
 	UpdatePreshadowCache();
 
 	// Gathers the list of primitives used to draw various shadow types
-	GatherShadowPrimitives(PreShadows, ViewDependentWholeSceneShadows, bReflectionCaptureScene);
+	GatherShadowPrimitives(PreShadows, ViewDependentWholeSceneShadows, bStaticSceneOnly);
 }
diff --git a/Engine/Source/Runtime/ShaderCore/Public/ShaderParameters.h b/Engine/Source/Runtime/ShaderCore/Public/ShaderParameters.h
--- a/Engine/Source/Runtime/ShaderCore/Public/ShaderParameters.h
+++ b/Engine/Source/Runtime/ShaderCore/Public/ShaderParameters.h
@@ -411,6 +411,11 @@ public:
 		return SRVParameter.IsBound() || UAVParameter.IsBound();
 	}
 
+	uint32 GetUAVIndex() const
+	{
+		return UAVParameter.GetBaseIndex();
+	}
+
 	friend FArchive& operator<<(FArchive& Ar,FRWShaderParameter& Parameter)
 	{
 		return Ar << Parameter.SRVParameter << Parameter.UAVParameter;
diff --git a/Engine/Source/Runtime/Windows/D3D11RHI/Private/D3D11Device.cpp b/Engine/Source/Runtime/Windows/D3D11RHI/Private/D3D11Device.cpp
--- a/Engine/Source/Runtime/Windows/D3D11RHI/Private/D3D11Device.cpp
+++ b/Engine/Source/Runtime/Windows/D3D11RHI/Private/D3D11Device.cpp
@@ -159,6 +159,8 @@ FD3D11DynamicRHI::FD3D11DynamicRHI(IDXGIFactory* InDXGIFactory,D3D_FEATURE_LEVEL
 	GPixelFormats[ PF_R5G6B5_UNORM	].PlatformFormat	= DXGI_FORMAT_B5G6R5_UNORM;
 	GPixelFormats[ PF_R8G8B8A8		].PlatformFormat	= DXGI_FORMAT_R8G8B8A8_TYPELESS;
 	GPixelFormats[ PF_R8G8			].PlatformFormat	= DXGI_FORMAT_R8G8_UNORM;
+	GPixelFormats[ PF_R32G32B32A32_UINT].PlatformFormat = DXGI_FORMAT_R32G32B32A32_UINT;
+	GPixelFormats[ PF_R16G16_UINT].PlatformFormat = DXGI_FORMAT_R16G16_UINT;
 
 	if (FeatureLevel >= D3D_FEATURE_LEVEL_11_0)
 	{
diff --git a/Engine/Source/Runtime/Windows/D3D11RHI/Private/D3D11Util.cpp b/Engine/Source/Runtime/Windows/D3D11RHI/Private/D3D11Util.cpp
--- a/Engine/Source/Runtime/Windows/D3D11RHI/Private/D3D11Util.cpp
+++ b/Engine/Source/Runtime/Windows/D3D11RHI/Private/D3D11Util.cpp
@@ -172,10 +172,7 @@ void VerifyD3D11Result(HRESULT D3DResult,const ANSICHAR* Code,const ANSICHAR* Fi
 	// this is to track down a rarely happening crash
 	if(D3DResult == E_OUTOFMEMORY)
 	{
-		if (IsInGameThread())
-		{
-			FMessageDialog::Open(EAppMsgType::Ok, NSLOCTEXT("D3D11RHI", "OutOfMemory", "Out of video memory trying to allocate a rendering resource."));
-		}
+		FPlatformMisc::MessageBoxExt(EAppMsgType::Ok, *NSLOCTEXT("D3D11RHI", "OutOfMemory", "Out of video memory trying to allocate a rendering resource.").ToString(), TEXT("Error"));
 #if STATS
 		GetRendererModule().DebugLogOnCrash();
 #endif
