Added missing file

diff --git a/Engine/Source/Runtime/Renderer/Private/DistanceFieldSurfaceCacheLighting.h b/Engine/Source/Runtime/Renderer/Private/DistanceFieldSurfaceCacheLighting.h
new file mode 100644
--- /dev/null
+++ b/Engine/Source/Runtime/Renderer/Private/DistanceFieldSurfaceCacheLighting.h
@@ -0,0 +1,145 @@
+// Copyright 1998-2013 Epic Games, Inc. All Rights Reserved.
+
+/*=============================================================================
+	
+=============================================================================*/
+
+#pragma once
+
+const static int32 GAOMaxSupportedLevel = 4;
+const static uint32 GMaxIrradianceCacheSamples = 100000;
+static const int32 GMaxNumObjectsPerTile = 512;
+
+/**  */
+class FRefinementLevelResources
+{
+public:
+	void InitDynamicRHI() 
+	{
+		//@todo - handle max exceeded
+		PositionAndRadius.Initialize(sizeof(float) * 4, GMaxIrradianceCacheSamples, PF_A32B32G32R32F, BUF_Static);
+		OccluderRadius.Initialize(sizeof(float), GMaxIrradianceCacheSamples, PF_R32_FLOAT, BUF_Static);
+		Normal.Initialize(sizeof(FFloat16Color), GMaxIrradianceCacheSamples, PF_FloatRGBA, BUF_Static);
+		BentNormal.Initialize(sizeof(FFloat16Color), GMaxIrradianceCacheSamples, PF_FloatRGBA, BUF_Static);
+		ScatterDrawParameters.Initialize(sizeof(uint32), 4, PF_R32_UINT, BUF_Static | BUF_DrawIndirect);
+		SavedStartIndex.Initialize(sizeof(uint32), 1, PF_R32_UINT, BUF_Static);
+		TileCoordinate.Initialize(sizeof(uint16) * 2, GMaxIrradianceCacheSamples, PF_R16G16_UINT, BUF_Static);
+	}
+
+	void ReleaseDynamicRHI()
+	{
+		PositionAndRadius.Release();
+		OccluderRadius.Release();
+		Normal.Release();
+		BentNormal.Release();
+		ScatterDrawParameters.Release();
+		SavedStartIndex.Release();
+		TileCoordinate.Release();
+	}
+
+	FRWBuffer PositionAndRadius;
+	FRWBuffer OccluderRadius;
+	FRWBuffer Normal;
+	FRWBuffer BentNormal;
+	FRWBuffer ScatterDrawParameters;
+	FRWBuffer SavedStartIndex;
+	FRWBuffer TileCoordinate;
+};
+
+/**  */
+class FSurfaceCacheResources : public FRenderResource
+{
+public:
+
+	FSurfaceCacheResources()
+	{
+		for (int32 i = 0; i < ARRAY_COUNT(Level); i++)
+		{
+			Level[i] = new FRefinementLevelResources();
+		}
+
+		TempResources = new FRefinementLevelResources();
+	}
+
+	~FSurfaceCacheResources()
+	{
+		for (int32 i = 0; i < ARRAY_COUNT(Level); i++)
+		{
+			delete Level[i];
+		}
+
+		delete TempResources;
+	}
+
+	virtual void InitDynamicRHI() 
+	{
+		DispatchParameters.Initialize(sizeof(uint32), 3, PF_R32_UINT, BUF_Static | BUF_DrawIndirect);
+
+		for (int32 i = 0; i < ARRAY_COUNT(Level); i++)
+		{
+			Level[i]->InitDynamicRHI();
+		}
+
+		TempResources->InitDynamicRHI();
+
+		bClearedResources = false;
+	}
+
+	virtual void ReleaseDynamicRHI()
+	{
+		DispatchParameters.Release();
+
+		for (int32 i = 0; i < ARRAY_COUNT(Level); i++)
+		{
+			Level[i]->ReleaseDynamicRHI();
+		}
+
+		TempResources->ReleaseDynamicRHI();
+	}
+
+	FRWBuffer DispatchParameters;
+
+	FRefinementLevelResources* Level[GAOMaxSupportedLevel + 1];
+
+	/** This is for temporary storage when copying from last frame's state */
+	FRefinementLevelResources* TempResources;
+
+	bool bClearedResources;
+};
+
+/**  */
+class FTileIntersectionResources : public FRenderResource
+{
+public:
+	virtual void InitDynamicRHI() 
+	{
+		TileConeAxisAndCos.Initialize(sizeof(float) * 4, TileDimensions.X * TileDimensions.Y, PF_A32B32G32R32F, BUF_Static);
+		TileConeDepthRanges.Initialize(sizeof(float) * 4, TileDimensions.X * TileDimensions.Y, PF_A32B32G32R32F, BUF_Static);
+
+		TileHeadDataUnpacked.Initialize(sizeof(uint32), TileDimensions.X * TileDimensions.Y * 4, PF_R32_UINT, BUF_Static);
+		TileHeadData.Initialize(sizeof(uint32) * 4, TileDimensions.X * TileDimensions.Y, PF_R32G32B32A32_UINT, BUF_Static);
+
+		//@todo - handle max exceeded
+		TileArrayData.Initialize(sizeof(uint16), GMaxNumObjectsPerTile * TileDimensions.X * TileDimensions.Y * 3, PF_R16_UINT, BUF_Static);
+		TileArrayNextAllocation.Initialize(sizeof(uint32), 1, PF_R32_UINT, BUF_Static);
+	}
+
+	virtual void ReleaseDynamicRHI()
+	{
+		TileConeAxisAndCos.Release();
+		TileConeDepthRanges.Release();
+		TileHeadDataUnpacked.Release();
+		TileHeadData.Release();
+		TileArrayData.Release();
+		TileArrayNextAllocation.Release();
+	}
+
+	FIntPoint TileDimensions;
+
+	FRWBuffer TileConeAxisAndCos;
+	FRWBuffer TileConeDepthRanges;
+	FRWBuffer TileHeadDataUnpacked;
+	FRWBuffer TileHeadData;
+	FRWBuffer TileArrayData;
+	FRWBuffer TileArrayNextAllocation;
+};
