Clang workaround

diff --git a/Engine/Source/Runtime/Renderer/Private/DistanceFieldSurfaceCacheLighting.cpp b/Engine/Source/Runtime/Renderer/Private/DistanceFieldSurfaceCacheLighting.cpp
--- a/Engine/Source/Runtime/Renderer/Private/DistanceFieldSurfaceCacheLighting.cpp
+++ b/Engine/Source/Runtime/Renderer/Private/DistanceFieldSurfaceCacheLighting.cpp
@@ -2663,8 +2663,8 @@ void UpdateVisibleObjectBuffers(const FScene* Scene, const FViewInfo& View, int3
 					const float DistanceToViewSq = (ViewOriginForDistanceCulling - (FVector)ObjectBoundingSphere).SizeSquared();
 
 					if (!GAOCPUFrustumCull 
-						|| DistanceToViewSq < FMath::Square(GAOMaxViewDistance + ObjectBoundingSphere.W)
-						&& View.ViewFrustum.IntersectSphere((FVector)ObjectBoundingSphere, ObjectBoundingSphere.W + AOMaxDistance))
+						|| (DistanceToViewSq < FMath::Square(GAOMaxViewDistance + ObjectBoundingSphere.W)
+						&& View.ViewFrustum.IntersectSphere((FVector)ObjectBoundingSphere, ObjectBoundingSphere.W + AOMaxDistance)))
 					{
 						ObjectBoundsData.Add(ObjectBoundingSphere);
 
