Brush wireframe is now rendered stippled when behind objects when being moved.

Added new editor relevance flag to separate editor primitives that want to render in the stippled 'behind' mode.
Added pre/post edit & move callbacks to communicate whether we are currently interactively editing a brushes transform or parameters.

TTP# 329784 - EDITOR: BSP: Unable to see the geometry Wireframe when scaling down BSP.

diff --git a/Engine/Source/Runtime/Engine/Classes/Engine/Brush.h b/Engine/Source/Runtime/Engine/Classes/Engine/Brush.h
--- a/Engine/Source/Runtime/Engine/Classes/Engine/Brush.h
+++ b/Engine/Source/Runtime/Engine/Classes/Engine/Brush.h
@@ -116,6 +116,10 @@ class ABrush : public AActor
 	class UBrushBuilder* BrushBuilder;
 #endif
 
+	/** Flag set when we are in a manipulation (scaling, translation, brush builder param change etc.) */
+	UPROPERTY()
+	uint32 bInManipulation:1;
+
 	/**
 	 * Stores selection information from geometry mode.  This is the only information that we can't
 	 * regenerate by looking at the source brushes following an undo operation.
@@ -137,6 +141,7 @@ class ABrush : public AActor
 	// Begin UObject interface.
 	ENGINE_API virtual void PostLoad() OVERRIDE;
 #if WITH_EDITOR
+	ENGINE_API virtual void PostEditMove(bool bFinished) OVERRIDE;
 	ENGINE_API virtual void PostEditChangeProperty(FPropertyChangedEvent& PropertyChangedEvent) OVERRIDE;
 #endif // WITH_EDITOR
 	ENGINE_API virtual bool Modify(bool bAlwaysMarkDirty = false) OVERRIDE;
diff --git a/Engine/Source/Runtime/Engine/Private/Brush.cpp b/Engine/Source/Runtime/Engine/Private/Brush.cpp
--- a/Engine/Source/Runtime/Engine/Private/Brush.cpp
+++ b/Engine/Source/Runtime/Engine/Private/Brush.cpp
@@ -34,6 +34,18 @@ ABrush::ABrush(const class FPostConstructInitializeProperties& PCIP)
 }
 
 #if WITH_EDITOR
+void ABrush::PostEditMove(bool bFinished)
+{
+	bInManipulation = !bFinished;
+
+	if( BrushComponent.IsValid() )
+	{
+		BrushComponent->ReregisterComponent();
+	}
+
+	Super::PostEditMove(bFinished);
+}
+
 void ABrush::PostEditChangeProperty(FPropertyChangedEvent& PropertyChangedEvent)
 {
 	if(Brush)
diff --git a/Engine/Source/Runtime/Engine/Private/BrushComponent.cpp b/Engine/Source/Runtime/Engine/Private/BrushComponent.cpp
--- a/Engine/Source/Runtime/Engine/Private/BrushComponent.cpp
+++ b/Engine/Source/Runtime/Engine/Private/BrushComponent.cpp
@@ -133,6 +133,7 @@ public:
 		bVolume(false),
 		bBuilder(false),
 		bSolidWhenSelected(false),
+		bInManipulation(false),
 		BrushColor(GEngine->C_BrushWire),
 		LevelColor(255,255,255),
 		PropertyColor(255,255,255),
@@ -154,6 +155,7 @@ public:
 			bBuilder = FActorEditorUtils::IsABuilderBrush( Owner );
 			BrushColor = Owner->GetWireColor();
 			bSolidWhenSelected = Owner->bSolidWhenSelected;
+			bInManipulation = Owner->bInManipulation;
 
 			// Builder brushes should be unaffected by level coloration, so if this is a builder brush, use
 			// the brush color as the level color.
@@ -385,6 +387,10 @@ public:
 		Result.bDrawRelevance = bVisible;
 		Result.bDynamicRelevance = true;
 		Result.bShadowRelevance = IsShadowCast(View);
+		if(bInManipulation)
+		{
+			Result.bEditorNoDepthTestPrimitiveRelevance = true;
+		}
 
 		// Don't render on top in 'collision view' modes
 		if(!bInCollisionView)
@@ -418,6 +424,7 @@ private:
 	uint32 bVolume : 1;
 	uint32 bBuilder : 1;	
 	uint32 bSolidWhenSelected : 1;
+	uint32 bInManipulation : 1;
 
 	FColor BrushColor;
 	FColor LevelColor;
diff --git a/Engine/Source/Runtime/Engine/Public/SceneManagement.h b/Engine/Source/Runtime/Engine/Public/SceneManagement.h
--- a/Engine/Source/Runtime/Engine/Public/SceneManagement.h
+++ b/Engine/Source/Runtime/Engine/Public/SceneManagement.h
@@ -1523,6 +1523,8 @@ struct FPrimitiveViewRelevance
 	uint32 bRenderInMainPass : 1;
 	/** The primitive is drawn only in the editor and composited onto the scene after post processing */
 	uint32 bEditorPrimitiveRelevance : 1;
+	/** The primitive is drawn only in the editor and composited onto the scene after post processing using no depth testing */
+	uint32 bEditorNoDepthTestPrimitiveRelevance : 1;
 	/** The primitive needs PreRenderView to be called before rendering. */
 	uint32 bNeedsPreRenderView : 1;
 	/** The primitive should have GatherSimpleLights called on the proxy when gathering simple lights. */
@@ -1562,6 +1564,7 @@ struct FPrimitiveViewRelevance
 		bRenderCustomDepth(false),
 		bRenderInMainPass(true),
 		bEditorPrimitiveRelevance(false),
+		bEditorNoDepthTestPrimitiveRelevance(false),
 		bNeedsPreRenderView(false),
 		bHasSimpleLights(false),
 		bOpaqueRelevance(true),
@@ -1585,6 +1588,8 @@ struct FPrimitiveViewRelevance
 		bDistortionRelevance |= B.bDistortionRelevance != 0;
 		bRenderCustomDepth |= B.bRenderCustomDepth != 0;
 		bRenderInMainPass |= B.bRenderInMainPass !=0;
+		bEditorPrimitiveRelevance |= B.bEditorPrimitiveRelevance !=0;
+		bEditorNoDepthTestPrimitiveRelevance |= B.bEditorNoDepthTestPrimitiveRelevance !=0;
 		bSceneColorRelevance |= B.bSceneColorRelevance != 0;
 		bNeedsPreRenderView |= B.bNeedsPreRenderView != 0;
 		bHasSimpleLights |= B.bHasSimpleLights != 0;
diff --git a/Engine/Source/Runtime/Renderer/Private/PostProcess/PostProcessCompositeEditorPrimitives.cpp b/Engine/Source/Runtime/Renderer/Private/PostProcess/PostProcessCompositeEditorPrimitives.cpp
--- a/Engine/Source/Runtime/Renderer/Private/PostProcess/PostProcessCompositeEditorPrimitives.cpp
+++ b/Engine/Source/Runtime/Renderer/Private/PostProcess/PostProcessCompositeEditorPrimitives.cpp
@@ -303,8 +303,12 @@ void FRCPassPostProcessCompositeEditorPrimitives::RenderPrimitivesToComposite(co
 	if( View.VisibleEditorPrimitives.Num() > 0 )
 	{
 		// Draw the dynamic non-occluded primitives using a base pass drawing policy.
-		TDynamicPrimitiveDrawer<FBasePassOpaqueDrawingPolicyFactory> Drawer(
+		TDynamicPrimitiveDrawer<FBasePassOpaqueDrawingPolicyFactory> DepthTestDrawer(
 			&View,FBasePassOpaqueDrawingPolicyFactory::ContextType(bDepthTest, ESceneRenderTargetsMode::SetTextures),true,false,false,bDepthTest);
+
+		TDynamicPrimitiveDrawer<FBasePassOpaqueDrawingPolicyFactory> NoDepthTestDrawer(
+			&View,FBasePassOpaqueDrawingPolicyFactory::ContextType(!bDepthTest, ESceneRenderTargetsMode::SetTextures),true,false,false,!bDepthTest);
+
 		for(int32 PrimitiveIndex = 0;PrimitiveIndex < View.VisibleEditorPrimitives.Num();PrimitiveIndex++)
 		{
 			const FPrimitiveSceneInfo* PrimitiveSceneInfo = View.VisibleEditorPrimitives[PrimitiveIndex];
@@ -316,13 +320,25 @@ void FRCPassPostProcessCompositeEditorPrimitives::RenderPrimitivesToComposite(co
 			// Only draw the primitive if it's visible
 			if( bVisible && 
 				// only draw opaque and masked primitives if wireframe is disabled
-				(PrimitiveViewRelevance.bOpaqueRelevance || View.Family->EngineShowFlags.Wireframe) )
+				(PrimitiveViewRelevance.bOpaqueRelevance || View.Family->EngineShowFlags.Wireframe))
 			{
-				Drawer.SetPrimitive(PrimitiveSceneInfo->Proxy);
-				PrimitiveSceneInfo->Proxy->DrawDynamicElements(
-					&Drawer,
-					&View			
-					);
+				// Distinguish between depth tested & non-depth tested primitives
+				if(PrimitiveViewRelevance.bEditorNoDepthTestPrimitiveRelevance)
+				{
+					NoDepthTestDrawer.SetPrimitive(PrimitiveSceneInfo->Proxy);
+					PrimitiveSceneInfo->Proxy->DrawDynamicElements(
+						&NoDepthTestDrawer,
+						&View			
+						);
+				}
+				else
+				{
+					DepthTestDrawer.SetPrimitive(PrimitiveSceneInfo->Proxy);
+					PrimitiveSceneInfo->Proxy->DrawDynamicElements(
+						&DepthTestDrawer,
+						&View			
+						);
+				}
 			}
 		}
 	}
