Generated a new sample set for Distance Field AO by relaxing the original set through repulsion, improves quality indoors quite a bit

diff --git a/Engine/Shaders/DeferredShadingCommon.usf b/Engine/Shaders/DeferredShadingCommon.usf
--- a/Engine/Shaders/DeferredShadingCommon.usf
+++ b/Engine/Shaders/DeferredShadingCommon.usf
@@ -248,6 +248,11 @@ float Decode71(float Scalar, out int Mask)
 	return (Scalar - 0.5f * Mask) * 2.0f;
 }
 
+uint DecodeShadingModelId(float InPackedChannel)
+{
+	return (uint)( InPackedChannel * 255.999 );
+}
+
 #define SHADINGMODELID_UNLIT				0
 #define SHADINGMODELID_DEFAULT_LIT			1
 #define SHADINGMODELID_SUBSURFACE			2
@@ -396,7 +401,7 @@ FGBufferData DecodeGBufferData(
 	uint PackedAlpha = (uint)( InGBufferA.a * 3.999 );
 	ret.DecalMask = PackedAlpha & 2;
 	ret.HasDistanceFieldRepresentation = PackedAlpha & 1;
-	ret.ShadingModelID	= (uint)( InGBufferB.a * 255.999 );
+	ret.ShadingModelID = DecodeShadingModelId(InGBufferB.a);
 	
 	bool bSubsurface = (ret.ShadingModelID == SHADINGMODELID_SUBSURFACE || ret.ShadingModelID == SHADINGMODELID_PREINTEGRATED_SKIN);
 
@@ -636,6 +641,13 @@ FGBufferData GetGBufferData(float2 UV, bool bGetNormalizedNormal = true)
 	return DecodeGBufferData(GBufferA, GBufferB, GBufferC, GBufferD, GBufferE, CustomNativeDepth, SceneDepth, bGetNormalizedNormal);
 }
 
+// Minimal path for just the lighting model, used to branch around unlit pixels (skybox)
+uint GetShadingModelId(float2 UV)
+{
+	float4 GBufferB = Texture2DSampleLevel(GBufferBTexture, GBufferBTextureSampler, UV, 0);
+	return DecodeShadingModelId(GBufferB.a);
+}
+
 // @param UV - UV space in the DBuffer textures (BufferSize resolution)
 FDBufferData GetDBufferData(float2 UV)
 {
diff --git a/Engine/Shaders/DistanceFieldSurfaceCacheLighting.usf b/Engine/Shaders/DistanceFieldSurfaceCacheLighting.usf
--- a/Engine/Shaders/DistanceFieldSurfaceCacheLighting.usf
+++ b/Engine/Shaders/DistanceFieldSurfaceCacheLighting.usf
@@ -36,7 +36,6 @@ void ObjectCullVS(
 	Output.ObjectIndex = ObjectIndex;
 } 
 
-#if FEATURE_LEVEL >= FEATURE_LEVEL_SM5
 RWBuffer<uint> RWTileHeadDataUnpacked;
 RWBuffer<uint> RWTileArrayData;
 
@@ -101,7 +100,6 @@ void ObjectCullPS(
 		}
 	}
 }
-#endif
 
 
 float2 UVToTileHead;
@@ -197,8 +195,6 @@ void ComputeDistanceFieldNormalPS(
 	OutColor = float4(EncodeNormalForAO(WorldNormal), DepthSign * SceneDepth);
 }
 
-#if FEATURE_LEVEL >= FEATURE_LEVEL_SM5
-
 #define THREADS_PER_PIXEL 1
 
 #if THREADS_PER_PIXEL != 1
@@ -339,7 +335,6 @@ void ComputeDistanceFieldNormalCS(
 		RWDistanceFieldNormal[PixelCoordinate] = OutValue;
 	}
 }
-#endif
 
 // For some reason gives innaccurate results at lower resolutions
 #define USE_SCREENVECTOR_WORLD_POSITION FINAL_INTERPOLATION_PASS
@@ -478,9 +473,10 @@ void IrradianceCacheSplatPS(
 		OutColor.a = Weight;
 	}
 
-	#if FINAL_INTERPOLATION_PASS
+	#define VISUALIZE_RECORD_POINTS 0
+	#if VISUALIZE_RECORD_POINTS && FINAL_INTERPOLATION_PASS
 		// Visualize cache record points
-		//OutColor.a = (DistanceError < .05f) * .1f;
+		OutColor.a = (DistanceError < .05f) * .1f;
 	#endif
 }
 
@@ -495,7 +491,13 @@ void AOCombinePS2(
 	float2 ScreenUV = float2((((uint2)SVPos.xy) * DOWNSAMPLE_FACTOR + .5f) * View.ViewSizeAndSceneTexelSize.zw);
 	float SceneDepth = CalcSceneDepth(ScreenUV);
 
-#if 1
+#define VISUALIZE_ACCUMULATED_WEIGHTS 0
+#if VISUALIZE_ACCUMULATED_WEIGHTS
+
+	float3 BentNormalAO = Texture2DSampleLevel(BentNormalAOTexture3, BentNormalAOSampler3, UVAndScreenPos.xy, 0).aaa;
+	OutColor = float4(BentNormalAO, SceneDepth); 
+
+#else
 
 	float4 IrradianceCacheAccumulation = Texture2DSampleLevel(BentNormalAOTexture3, BentNormalAOSampler3, UVAndScreenPos.xy, 0);
 	float3 BentNormalAO = float3(0, 0, 0);
@@ -522,9 +524,6 @@ void AOCombinePS2(
 		OutColor.a = abs(OutColor.a);
 	}
 
-#else
-	float3 BentNormalAO = Texture2DSampleLevel(BentNormalAOTexture3, BentNormalAOSampler3, UVAndScreenPos.xy, 0).aaa;
-	OutColor = float4(BentNormalAO, SceneDepth); 
 #endif
 }
 
@@ -652,69 +651,3 @@ void AOUpsamplePS(in float4 UVAndScreenPos : TEXCOORD0, out float4 OutColor : SV
 
 	OutColor = float4(BentNormal, 1);
 }
-
-float BentNormalUVScale;
-float ApplyShadowing;
-
-void SkyLightDiffusePS(
-	in float4 UVAndScreenPos : TEXCOORD0, 
-	out float4 OutColor : SV_Target0)
-{
-	float2 UV = UVAndScreenPos.xy;
-
-	FScreenSpaceData ScreenSpaceData = GetScreenSpaceData(UV);
-	FGBufferData InGBufferData = ScreenSpaceData.GBuffer;
-	
-	float3 Lighting = 0; 
-
-	BRANCH
-	// Only light pixels marked as lit
-	if (InGBufferData.ShadingModelID > 0)
-	{
-		float SkyVisibility = 1;
-		float DotProductFactor = 1;
-		float3 SkyLightingNormal = InGBufferData.WorldNormal;
-
-		BRANCH
-		if (ApplyShadowing > 0)
-		{
-			float3 BentNormal = Texture2DSampleLevel(BentNormalAOTexture, BentNormalAOSampler, UV * BentNormalUVScale, 0).xyz;
-			SkyVisibility = BentNormal.x;
-
-			// Disabled for now as it needs tweaking
-			#define USE_DYNAMIC_OCCLUSION_BENT_NORMAL 0
-			#if USE_DYNAMIC_OCCLUSION_BENT_NORMAL
-				SkyVisibility = length(BentNormal);
-				float3 NormalizedBentNormal = BentNormal / (max(SkyVisibility, .00001f));
-
-				// Weight toward the material normal to increase directionality
-				float BentNormalWeightFactor = 1 - (1 - SkyVisibility) * (1 - SkyVisibility);
-
-				// We are lerping between the inputs of two lighting scenarios based on occlusion
-				// In the mostly unoccluded case, evaluate sky lighting with the material normal, because it has higher detail
-				// In the mostly occluded case, evaluate sky lighting with the bent normal, because it is a better representation of the incoming lighting
-				// Then treat the lighting evaluated along the bent normal as an area light, so we must apply the lambert term
-				SkyLightingNormal = lerp(NormalizedBentNormal, InGBufferData.WorldNormal, BentNormalWeightFactor);
-				DotProductFactor = lerp(dot(NormalizedBentNormal, InGBufferData.WorldNormal), 1, BentNormalWeightFactor);
-			#endif
-
-			if (InGBufferData.HasDistanceFieldRepresentation == 0)
-			{
-				SkyVisibility = .5f;
-			}
-		}
-			
-		// Compute the preconvolved incoming lighting with the bent normal direction
-		float3 DiffuseLookup = GetSkySHDiffuse(SkyLightingNormal) * View.SkyLightColor.rgb;
-
-		float ScalarFactors = SkyVisibility * DotProductFactor * InGBufferData.GBufferAO * ScreenSpaceData.AmbientOcclusion;
-
-		// Apply AO to the sky diffuse and account for darkening due to the geometry term
-		Lighting = ScalarFactors * DiffuseLookup * InGBufferData.DiffuseColor;
-
-		// For debugging
-		//Lighting = Texture2DSampleLevel(BentNormalAOTexture, BentNormalAOSampler, UV * BentNormalUVScale, 0).xyz;
-	}
-
-	OutColor = float4(Lighting, 0);
-}
\ No newline at end of file
diff --git a/Engine/Shaders/SkyLighting.usf b/Engine/Shaders/SkyLighting.usf
new file mode 100644
--- /dev/null
+++ b/Engine/Shaders/SkyLighting.usf
@@ -0,0 +1,77 @@
+// Copyright 1998-2013 Epic Games, Inc. All Rights Reserved.
+
+/*=============================================================================
+
+=============================================================================*/
+
+#include "Common.usf"
+#include "DeferredShadingCommon.usf"
+#include "ReflectionEnvironmentShared.usf"
+
+float BentNormalUVScale;
+float ApplyShadowing;
+Texture2D BentNormalAOTexture;
+SamplerState BentNormalAOSampler;
+
+void SkyLightDiffusePS(
+	in float4 UVAndScreenPos : TEXCOORD0, 
+	out float4 OutColor : SV_Target0)
+{
+	float2 UV = UVAndScreenPos.xy;
+	uint ShadingModelId = GetShadingModelId(UV);
+	float3 Lighting = 0; 
+
+	BRANCH
+	// Only light pixels marked as lit
+	if (ShadingModelId > 0)
+	{
+		FScreenSpaceData ScreenSpaceData = GetScreenSpaceData(UV);
+		FGBufferData InGBufferData = ScreenSpaceData.GBuffer;
+
+		float SkyVisibility = 1;
+		float DotProductFactor = 1;
+		float3 SkyLightingNormal = InGBufferData.WorldNormal;
+
+		BRANCH
+		if (ApplyShadowing > 0)
+		{
+			float3 BentNormal = Texture2DSampleLevel(BentNormalAOTexture, BentNormalAOSampler, UV * BentNormalUVScale, 0).xyz;
+			SkyVisibility = BentNormal.x;
+
+			// Disabled for now as it needs tweaking
+			#define USE_DYNAMIC_OCCLUSION_BENT_NORMAL 0
+			#if USE_DYNAMIC_OCCLUSION_BENT_NORMAL
+				SkyVisibility = length(BentNormal);
+				float3 NormalizedBentNormal = BentNormal / (max(SkyVisibility, .00001f));
+
+				// Weight toward the material normal to increase directionality
+				float BentNormalWeightFactor = 1 - (1 - SkyVisibility) * (1 - SkyVisibility);
+
+				// We are lerping between the inputs of two lighting scenarios based on occlusion
+				// In the mostly unoccluded case, evaluate sky lighting with the material normal, because it has higher detail
+				// In the mostly occluded case, evaluate sky lighting with the bent normal, because it is a better representation of the incoming lighting
+				// Then treat the lighting evaluated along the bent normal as an area light, so we must apply the lambert term
+				SkyLightingNormal = lerp(NormalizedBentNormal, InGBufferData.WorldNormal, BentNormalWeightFactor);
+				DotProductFactor = lerp(dot(NormalizedBentNormal, InGBufferData.WorldNormal), 1, BentNormalWeightFactor);
+			#endif
+
+			if (InGBufferData.HasDistanceFieldRepresentation == 0)
+			{
+				SkyVisibility = .5f;
+			}
+		}
+			
+		// Compute the preconvolved incoming lighting with the bent normal direction
+		float3 DiffuseLookup = GetSkySHDiffuse(SkyLightingNormal) * View.SkyLightColor.rgb;
+
+		float ScalarFactors = SkyVisibility * DotProductFactor * InGBufferData.GBufferAO * ScreenSpaceData.AmbientOcclusion;
+
+		// Apply AO to the sky diffuse and account for darkening due to the geometry term
+		Lighting = ScalarFactors * DiffuseLookup * InGBufferData.DiffuseColor;
+
+		// For debugging
+		//Lighting = Texture2DSampleLevel(BentNormalAOTexture, BentNormalAOSampler, UV * BentNormalUVScale, 0).xyz;
+	}
+
+	OutColor = float4(Lighting, 0);
+}
\ No newline at end of file
diff --git a/Engine/Source/Runtime/Renderer/Private/DistanceFieldLightingShared.h b/Engine/Source/Runtime/Renderer/Private/DistanceFieldLightingShared.h
--- a/Engine/Source/Runtime/Renderer/Private/DistanceFieldLightingShared.h
+++ b/Engine/Source/Runtime/Renderer/Private/DistanceFieldLightingShared.h
@@ -22,34 +22,7 @@ const int32 NumConeSampleDirections = 9;
 /** Generates unit length, stratified and uniformly distributed direction samples in a hemisphere. */
 extern void GenerateStratifiedUniformHemisphereSamples2(int32 NumThetaSteps, int32 NumPhiSteps, FRandomStream& RandomStream, TArray<FVector4>& Samples);
 
-// Optimized with GA, Best =  -2.764536 after 117591 Steps
-const FVector RandomVectors9[] = 
-{
-	FVector(-0.134453, 0.570145, 0.810467),
-	FVector(0.600695, 0.732059, 0.321333),
-	FVector(0.591856, -0.692181, 0.413028),
-	FVector(0.065080, -0.293309, 0.953800),
-	FVector(-0.952426, -0.251775, 0.171737),
-	FVector(0.047865, -0.998847, 0.003783),
-	FVector(-0.476063, -0.743138, 0.470224),
-	FVector(-0.837584, 0.480720, 0.259541),
-	FVector(0.780739, 0.035320, 0.623858)
-};
-
-// Sample set restricted to not self-intersect a surface based on cone angle .475882232
-// Coverage of hemisphere = 0.755312979
-const FVector SpacedVectors9[] = 
-{
-	FVector(-0.573257625, 0.625250816, 0.529563010),
-	FVector(0.253354192, -0.840093017, 0.479640961),
-	FVector(-0.421664953, -0.718063235, 0.553700149),
-	FVector(0.249163717, 0.796005428, 0.551627457),
-	FVector(0.375082791, 0.295851320, 0.878512800),
-	FVector(-0.217619032, 0.00193520682, 0.976031899),
-	FVector(-0.852834642, 0.0111727007, 0.522061586),
-	FVector(0.745701790, 0.239393353, 0.621787369),
-	FVector(-0.151036426, -0.465937436, 0.871831656)
-};
+extern void GetSpacedVectors(TArray<TInlineAllocator<9> >& OutVectors);
 
 BEGIN_UNIFORM_BUFFER_STRUCT(FAOSampleData2,)
 	DECLARE_UNIFORM_BUFFER_STRUCT_MEMBER_ARRAY(FVector4,SampleDirections,[NumConeSampleDirections])
diff --git a/Engine/Source/Runtime/Renderer/Private/DistanceFieldSurfaceCacheLighting.cpp b/Engine/Source/Runtime/Renderer/Private/DistanceFieldSurfaceCacheLighting.cpp
--- a/Engine/Source/Runtime/Renderer/Private/DistanceFieldSurfaceCacheLighting.cpp
+++ b/Engine/Source/Runtime/Renderer/Private/DistanceFieldSurfaceCacheLighting.cpp
@@ -173,6 +173,14 @@ FAutoConsoleVariableRef CVarAOComputeShaderNormalCalculation(
 	ECVF_Cheat | ECVF_RenderThreadSafe
 	);
 
+int32 GAOSampleSet = 1;
+FAutoConsoleVariableRef CVarAOSampleSet(
+	TEXT("r.AOSampleSet"),
+	GAOSampleSet,
+	TEXT("0 = Original set, 1 = Relaxed set"),
+	ECVF_Cheat | ECVF_RenderThreadSafe
+	);
+
 IMPLEMENT_UNIFORM_BUFFER_STRUCT(FAOSampleData2,TEXT("AOSamples2"));
 
 FIntPoint GetBufferSizeForAO()
@@ -187,6 +195,55 @@ bool DoesPlatformSupportDistanceFieldAO(EShaderPlatform Platform)
 	return Platform == SP_PCD3D_SM5;
 }
 
+// Sample set restricted to not self-intersect a surface based on cone angle .475882232
+// Coverage of hemisphere = 0.755312979
+const FVector SpacedVectors9[] = 
+{
+	FVector(-0.573257625, 0.625250816, 0.529563010),
+	FVector(0.253354192, -0.840093017, 0.479640961),
+	FVector(-0.421664953, -0.718063235, 0.553700149),
+	FVector(0.249163717, 0.796005428, 0.551627457),
+	FVector(0.375082791, 0.295851320, 0.878512800),
+	FVector(-0.217619032, 0.00193520682, 0.976031899),
+	FVector(-0.852834642, 0.0111727007, 0.522061586),
+	FVector(0.745701790, 0.239393353, 0.621787369),
+	FVector(-0.151036426, -0.465937436, 0.871831656)
+};
+
+// Generated from SpacedVectors9 by applying repulsion forces until convergence
+const FVector RelaxedSpacedVectors9[] = 
+{
+	FVector(-0.467612, 0.739424, 0.484347),
+	FVector(0.517459, -0.705440, 0.484346),
+	FVector(-0.419848, -0.767551, 0.484347),
+	FVector(0.343077, 0.804802, 0.484347),
+	FVector(0.364239, 0.244290, 0.898695),
+	FVector(-0.381547, 0.185815, 0.905481),
+	FVector(-0.870176, -0.090559, 0.484347),
+	FVector(0.874448, 0.027390, 0.484346),
+	FVector(0.032967, -0.435625, 0.899524)
+};
+
+void GetSpacedVectors(TArray<FVector, TInlineAllocator<9> >& OutVectors)
+{
+	OutVectors.Empty(ARRAY_COUNT(SpacedVectors9));
+
+	if (GAOSampleSet == 0)
+	{
+		for (int32 i = 0; i < ARRAY_COUNT(SpacedVectors9); i++)
+		{
+			OutVectors.Add(SpacedVectors9[i]);
+		}
+	}
+	else
+	{
+		for (int32 i = 0; i < ARRAY_COUNT(RelaxedSpacedVectors9); i++)
+		{
+			OutVectors.Add(RelaxedSpacedVectors9[i]);
+		}
+	}
+}
+
 // Cone half angle derived from each cone covering an equal solid angle
 float GAOConeHalfAngle = FMath::Acos(1 - 1.0f / (float)ARRAY_COUNT(SpacedVectors9));
 
@@ -1273,18 +1330,15 @@ public:
 
 		SetShaderValue(ShaderRHI, ViewDimensionsParameter, View.ViewRect);
 
-		static bool bInitialized = false;
-		static FAOSampleData2 AOSampleData;
+		FAOSampleData2 AOSampleData;
 
-		if (!bInitialized)
-		{
-			for (int32 SampleIndex = 0; SampleIndex < NumConeSampleDirections; SampleIndex++)
-			{
-				AOSampleData.SampleDirections[SampleIndex] = FVector4(SpacedVectors9[SampleIndex]);
-			}
+		TArray<FVector, TInlineAllocator<9> > SampleDirections;
+		GetSpacedVectors(SampleDirections);
 
-			bInitialized = true;
-		} 
+		for (int32 SampleIndex = 0; SampleIndex < NumConeSampleDirections; SampleIndex++)
+		{
+			AOSampleData.SampleDirections[SampleIndex] = FVector4(SampleDirections[SampleIndex]);
+		}
 
 		SetUniformBufferParameterImmediate(ShaderRHI, GetUniformBufferParameter<FAOSampleData2>(), AOSampleData);
 
@@ -1299,7 +1353,7 @@ public:
 
 		for (int32 SampleIndex = 0; SampleIndex < NumConeSampleDirections; SampleIndex++)
 		{
-			UnoccludedVector += SpacedVectors9[SampleIndex];
+			UnoccludedVector += SampleDirections[SampleIndex];
 		}
 
 		float BentNormalNormalizeFactorValue = 1.0f / (UnoccludedVector / NumConeSampleDirections).Size();
@@ -1530,18 +1584,15 @@ public:
 
 		SetShaderValue(ShaderRHI, ViewDimensionsParameter, View.ViewRect);
 
-		static bool bInitialized = false;
-		static FAOSampleData2 AOSampleData;
+		FAOSampleData2 AOSampleData;
 
-		if (!bInitialized)
-		{
-			for (int32 SampleIndex = 0; SampleIndex < NumConeSampleDirections; SampleIndex++)
-			{
-				AOSampleData.SampleDirections[SampleIndex] = FVector4(SpacedVectors9[SampleIndex]);
-			}
+		TArray<FVector, TInlineAllocator<9> > SampleDirections;
+		GetSpacedVectors(SampleDirections);
 
-			bInitialized = true;
-		} 
+		for (int32 SampleIndex = 0; SampleIndex < NumConeSampleDirections; SampleIndex++)
+		{
+			AOSampleData.SampleDirections[SampleIndex] = FVector4(SampleDirections[SampleIndex]);
+		}
 
 		SetUniformBufferParameterImmediate(ShaderRHI, GetUniformBufferParameter<FAOSampleData2>(), AOSampleData);
 
@@ -1563,7 +1614,7 @@ public:
 
 		for (int32 SampleIndex = 0; SampleIndex < NumConeSampleDirections; SampleIndex++)
 		{
-			UnoccludedVector += SpacedVectors9[SampleIndex];
+			UnoccludedVector += SampleDirections[SampleIndex];
 		}
 
 		float BentNormalNormalizeFactorValue = 1.0f / (UnoccludedVector / NumConeSampleDirections).Size();
@@ -2325,7 +2376,66 @@ FVector GetUnitVector2(FRandomStream& RandomStream)
 
 void GenerateBestSpacedVectors()
 {
-	if (false)
+	static bool bGenerated = false;
+	bool bApplyRepulsion = false;
+
+	if (bApplyRepulsion && !bGenerated)
+	{
+		bGenerated = true;
+
+		FVector OriginalSpacedVectors9[ARRAY_COUNT(SpacedVectors9)];
+
+		for (int32 i = 0; i < ARRAY_COUNT(OriginalSpacedVectors9); i++)
+		{
+			OriginalSpacedVectors9[i] = SpacedVectors9[i];
+		}
+
+		float CosHalfAngle = 1 - 1.0f / (float)ARRAY_COUNT(OriginalSpacedVectors9);
+		// Used to prevent self-shadowing on a plane
+		float AngleBias = .03f;
+		float MinAngle = FMath::Acos(CosHalfAngle) + AngleBias;
+		float MinZ = FMath::Sin(MinAngle);
+
+		// Relaxation iterations by repulsion
+		for (int32 Iteration = 0; Iteration < 10000; Iteration++)
+		{
+			for (int32 i = 0; i < ARRAY_COUNT(OriginalSpacedVectors9); i++)
+			{
+				FVector Force(0.0f, 0.0f, 0.0f);
+
+				for (int32 j = 0; j < ARRAY_COUNT(OriginalSpacedVectors9); j++)
+				{
+					if (i != j)
+					{
+						FVector Distance = OriginalSpacedVectors9[i] - OriginalSpacedVectors9[j];
+						float Dot = OriginalSpacedVectors9[i] | OriginalSpacedVectors9[j];
+
+						if (Dot > 0)
+						{
+							// Repulsion force
+							Force += .001f * Distance.SafeNormal() * Dot * Dot * Dot * Dot;
+						}
+					}
+				}
+
+				FVector NewPosition = OriginalSpacedVectors9[i] + Force;
+				NewPosition.Z = FMath::Max(NewPosition.Z, MinZ);
+				NewPosition = NewPosition.SafeNormal();
+				OriginalSpacedVectors9[i] = NewPosition;
+			}
+		}
+
+		for (int32 i = 0; i < ARRAY_COUNT(OriginalSpacedVectors9); i++)
+		{
+			UE_LOG(LogTemp, Log, TEXT("FVector(%f, %f, %f),"), OriginalSpacedVectors9[i].X, OriginalSpacedVectors9[i].Y, OriginalSpacedVectors9[i].Z);
+		}
+
+		int32 temp = 0;
+	}
+
+	bool bBruteForceGenerateConeDirections = false;
+
+	if (bBruteForceGenerateConeDirections)
 	{
 		FVector BestSpacedVectors9[9];
 		float BestCoverage = 0;
@@ -2337,7 +2447,7 @@ void GenerateBestSpacedVectors()
 		// Prevent self-intersection in sample set
 		float MinAngle = FMath::Acos(CosHalfAngle);
 		float MinZ = FMath::Sin(MinAngle);
-		FRandomStream RandomStream(1234567);
+		FRandomStream RandomStream(123567);
 
 		// Super slow random brute force search
 		for (int i = 0; i < 1000000; i++)
@@ -3023,7 +3133,7 @@ private:
 	FShaderParameter ApplyShadowing;
 };
 
-IMPLEMENT_SHADER_TYPE(,FDynamicSkyLightDiffusePS,TEXT("DistanceFieldSurfaceCacheLighting"),TEXT("SkyLightDiffusePS"),SF_Pixel);
+IMPLEMENT_SHADER_TYPE(,FDynamicSkyLightDiffusePS,TEXT("SkyLighting"),TEXT("SkyLightDiffusePS"),SF_Pixel);
 
 void FDeferredShadingSceneRenderer::RenderDynamicSkyLighting()
 {
