Fixed Movable skylights on translucency

diff --git a/Engine/Source/Runtime/Renderer/Private/SceneRendering.cpp b/Engine/Source/Runtime/Renderer/Private/SceneRendering.cpp
--- a/Engine/Source/Runtime/Renderer/Private/SceneRendering.cpp
+++ b/Engine/Source/Runtime/Renderer/Private/SceneRendering.cpp
@@ -464,11 +464,20 @@ TUniformBufferRef<FViewUniformShaderParameters> FViewInfo::CreateUniformBuffer(
 								(((float)ViewRect.Max.Y / BufferSize.Y) - OneScenePixelUVSize.Y) );
 	ViewUniformShaderParameters.SceneTextureMinMax = SceneTexMinMax;
 
-	if (Family->Scene && ((FScene*)Family->Scene)->SkyLight)
+	FScene* Scene = (FScene*)Family->Scene;
+
+	if (Scene && Scene->SkyLight)
 	{
-		FScene* Scene = (FScene*)Family->Scene;
-		ViewUniformShaderParameters.SkyLightColor = Scene->SkyLight->LightColor;
-		ViewUniformShaderParameters.SkyLightParameters = Scene->SkyLight->bCastShadows && Scene->SkyLight->bPrecomputedLightingIsValid ? 1 : 0;
+		FSkyLightSceneProxy* SkyLight = Scene->SkyLight;
+
+		ViewUniformShaderParameters.SkyLightColor = SkyLight->LightColor;
+
+		bool bApplyPrecomputedBentNormalShadowing = 
+			SkyLight->bCastShadows 
+			&& SkyLight->bWantsStaticShadowing
+			&& SkyLight->bPrecomputedLightingIsValid;
+
+		ViewUniformShaderParameters.SkyLightParameters = bApplyPrecomputedBentNormalShadowing ? 1 : 0;
 	}
 	else
 	{
