UE4 - merge first pass at parallel rendering to main

diff --git a/Engine/Source/Runtime/RHI/Public/RHIDefinitions.h b/Engine/Source/Runtime/RHI/Public/RHIDefinitions.h
new file mode 100644
--- /dev/null
+++ b/Engine/Source/Runtime/RHI/Public/RHIDefinitions.h
@@ -0,0 +1,638 @@
+// Copyright 1998-2014 Epic Games, Inc. All Rights Reserved.
+
+/*=============================================================================
+	RHIDefinitions.h: Render Hardware Interface definitions
+		(that don't require linking).
+=============================================================================*/
+
+#ifndef __RHIDEFINITIONS_H__
+#define __RHIDEFINITIONS_H__
+
+#include "Core.h"
+
+enum EShaderFrequency
+{
+	SF_Vertex			= 0,
+	SF_Hull				= 1,
+	SF_Domain			= 2,
+	SF_Pixel			= 3,
+	SF_Geometry			= 4,
+	SF_Compute			= 5,
+
+	SF_NumFrequencies	= 6,
+	SF_NumBits			= 3,
+};
+static_assert(SF_NumFrequencies <= (1 << SF_NumBits), "SF_NumFrequencies will not fit on SF_NumBits");
+
+/** @warning: update *LegacyShaderPlatform* when the below changes */
+enum EShaderPlatform
+{
+	SP_PCD3D_SM5		= 0,
+	SP_OPENGL_SM4		= 1,
+	SP_PS4				= 2,
+	/** Used when running in Feature Level ES2 in OpenGL. */
+	SP_OPENGL_PCES2		= 3,
+	SP_XBOXONE			= 4,
+	SP_PCD3D_SM4		= 5,
+	SP_OPENGL_SM5		= 6,
+	/** Used when running in Feature Level ES2 in D3D11. */
+	SP_PCD3D_ES2		= 7,
+	SP_OPENGL_ES2		= 8,
+	SP_OPENGL_ES2_WEBGL = 9, 
+	SP_OPENGL_ES2_IOS	= 10,
+
+	SP_NumPlatforms		= 11,
+	SP_NumBits			= 4,
+};
+static_assert(SP_NumPlatforms <= (1 << SP_NumBits), "SP_NumPlatforms will not fit on SP_NumBits");
+
+enum ERenderQueryType
+{
+	// e.g. WaitForFrameEventCompletion()
+	RQT_Undefined,
+	// Result is the number of samples that are not culled (divide by MSAACount to get pixels)
+	RQT_Occlusion,
+	// Result is time in micro seconds = 1/1000 ms = 1/1000000 sec
+	RQT_AbsoluteTime,
+};
+
+/** Maximum number of miplevels in a texture. */
+enum { MAX_TEXTURE_MIP_COUNT = 14 };
+
+/** The maximum number of vertex elements which can be used by a vertex declaration. */
+enum { MaxVertexElementCount = 16 };
+
+/** The alignment in bytes between elements of array shader parameters. */
+enum { ShaderArrayElementAlignBytes = 16 };
+
+/** The number of render-targets that may be simultaneously written to. */
+enum { MaxSimultaneousRenderTargets = 8 };
+
+/** The number of UAVs that may be simultaneously bound to a shader. */
+enum { MaxSimultaneousUAVs = 8 };
+
+/**
+ * The RHI's feature level indicates what level of support can be relied upon.
+ * Note: these are named after graphics API's like ES2 but a feature level can be used with a different API (eg ERHIFeatureLevel::ES2 on D3D11)
+ * As long as the graphics API supports all the features of the feature level (eg no ERHIFeatureLevel::SM5 on OpenGL ES2)
+ */
+namespace ERHIFeatureLevel
+{
+	enum Type
+	{
+		/** Feature level defined by the core capabilities of OpenGL ES2. */
+		ES2,
+		/** Feature level defined by the capabilities of DX9 Shader Model 3. */
+		SM3,
+		/** Feature level defined by the capabilities of DX10 Shader Model 4. */
+		SM4,
+		/** Feature level defined by the capabilities of DX11 Shader Model 5. */
+		SM5,
+		Num
+	};
+};
+
+enum ESamplerFilter
+{
+	SF_Point,
+	SF_Bilinear,
+	SF_Trilinear,
+	SF_AnisotropicPoint,
+	SF_AnisotropicLinear,
+};
+
+enum ESamplerAddressMode
+{
+	AM_Wrap,
+	AM_Clamp,
+	AM_Mirror,
+	/** Not supported on all platforms */
+	AM_Border
+};
+
+enum ESamplerCompareFunction
+{
+	SCF_Never,
+	SCF_Less
+};
+
+enum ERasterizerFillMode
+{
+	FM_Point,
+	FM_Wireframe,
+	FM_Solid
+};
+
+enum ERasterizerCullMode
+{
+	CM_None,
+	CM_CW,
+	CM_CCW
+};
+
+enum EColorWriteMask
+{
+	CW_RED   = 0x01,
+	CW_GREEN = 0x02,
+	CW_BLUE  = 0x04,
+	CW_ALPHA = 0x08,
+
+	CW_NONE  = 0,
+	CW_RGB   = CW_RED | CW_GREEN | CW_BLUE,
+	CW_RGBA  = CW_RED | CW_GREEN | CW_BLUE | CW_ALPHA,
+	CW_RG    = CW_RED | CW_GREEN,
+	CW_BA    = CW_BLUE | CW_ALPHA,
+};
+
+enum ECompareFunction
+{
+	CF_Less,
+	CF_LessEqual,
+	CF_Greater,
+	CF_GreaterEqual,
+	CF_Equal,
+	CF_NotEqual,
+	CF_Never,
+	CF_Always
+};
+
+enum EStencilOp
+{
+	SO_Keep,
+	SO_Zero,
+	SO_Replace,
+	SO_SaturatedIncrement,
+	SO_SaturatedDecrement,
+	SO_Invert,
+	SO_Increment,
+	SO_Decrement
+};
+
+enum EBlendOperation
+{
+	BO_Add,
+	BO_Subtract,
+	BO_Min,
+	BO_Max,
+	BO_ReverseSubtract,
+};
+
+enum EBlendFactor
+{
+	BF_Zero,
+	BF_One,
+	BF_SourceColor,
+	BF_InverseSourceColor,
+	BF_SourceAlpha,
+	BF_InverseSourceAlpha,
+	BF_DestAlpha,
+	BF_InverseDestAlpha,
+	BF_DestColor,
+	BF_InverseDestColor,
+	BF_ConstantBlendFactor,
+	BF_InverseConstantBlendFactor
+};
+
+enum EVertexElementType
+{
+	VET_None,
+	VET_Float1,
+	VET_Float2,
+	VET_Float3,
+	VET_Float4,
+	VET_PackedNormal,	// FPackedNormal
+	VET_UByte4,
+	VET_UByte4N,
+	VET_Color,
+	VET_Short2,
+	VET_Short4,
+	VET_Short2N,		// 16 bit word normalized to (value/32767.0,value/32767.0,0,0,1)
+	VET_Half2,			// 16 bit float using 1 bit sign, 5 bit exponent, 10 bit mantissa 
+	VET_Half4,
+	VET_MAX
+};
+
+enum ECubeFace
+{
+	CubeFace_PosX=0,
+	CubeFace_NegX,
+	CubeFace_PosY,
+	CubeFace_NegY,
+	CubeFace_PosZ,
+	CubeFace_NegZ,
+	CubeFace_MAX
+};
+
+enum EUniformBufferUsage
+{
+	// the uniform buffer is temporary, used for a single draw call then discarded
+	UniformBuffer_SingleDraw = 0,
+	// the uniform buffer is used for multiple draw calls but only for the current frame
+	UniformBuffer_SingleFrame,
+	// the uniform buffer is used for multiple draw calls, possibly across multiple frames
+	UniformBuffer_MultiFrame,
+};
+
+/** The base type of a value in a uniform buffer. */
+enum EUniformBufferBaseType
+{
+	UBMT_INVALID,
+	UBMT_BOOL,
+	UBMT_INT32,
+	UBMT_UINT32,
+	UBMT_FLOAT32,
+	UBMT_STRUCT,
+	UBMT_SRV,
+	UBMT_UAV,
+	UBMT_SAMPLER,
+	UBMT_TEXTURE
+};
+
+struct FRHIResourceTableEntry
+{
+public:
+	static CONSTEXPR uint32 GetEndOfStreamToken()
+	{
+		return 0xffffffff;
+	}
+
+	static uint32 Create(uint16 UniformBufferIndex, uint16 ResourceIndex, uint16 BindIndex)
+	{
+		return ((UniformBufferIndex & RTD_Mask_UniformBufferIndex) << RTD_Shift_UniformBufferIndex) |
+			((ResourceIndex & RTD_Mask_ResourceIndex) << RTD_Shift_ResourceIndex) |
+			((BindIndex & RTD_Mask_BindIndex) << RTD_Shift_BindIndex);
+	}
+
+	static inline uint16 GetUniformBufferIndex(uint32 Data)
+	{
+		return (Data >> RTD_Shift_UniformBufferIndex) & RTD_Mask_UniformBufferIndex;
+	}
+
+	static inline uint16 GetResourceIndex(uint32 Data)
+	{
+		return (Data >> RTD_Shift_ResourceIndex) & RTD_Mask_ResourceIndex;
+	}
+
+	static inline uint16 GetBindIndex(uint32 Data)
+	{
+		return (Data >> RTD_Shift_BindIndex) & RTD_Mask_BindIndex;
+	}
+
+private:
+	enum EResourceTableDefinitions
+	{
+		RTD_NumBits_UniformBufferIndex	= 8,
+		RTD_NumBits_ResourceIndex		= 16,
+		RTD_NumBits_BindIndex			= 8,
+
+		RTD_Mask_UniformBufferIndex		= (1 << RTD_NumBits_UniformBufferIndex) - 1,
+		RTD_Mask_ResourceIndex			= (1 << RTD_NumBits_ResourceIndex) - 1,
+		RTD_Mask_BindIndex				= (1 << RTD_NumBits_BindIndex) - 1,
+
+		RTD_Shift_BindIndex				= 0,
+		RTD_Shift_ResourceIndex			= RTD_Shift_BindIndex + RTD_NumBits_BindIndex,
+		RTD_Shift_UniformBufferIndex	= RTD_Shift_ResourceIndex + RTD_NumBits_ResourceIndex,
+	};
+	static_assert(RTD_NumBits_UniformBufferIndex + RTD_NumBits_ResourceIndex + RTD_NumBits_BindIndex <= sizeof(uint32) * 8, "RTD_* values must fit in 32 bits");
+};
+
+enum EResourceLockMode
+{
+	RLM_ReadOnly,
+	RLM_WriteOnly,
+	RLM_Num
+};
+
+/** limited to 8 types in FReadSurfaceDataFlags */
+enum ERangeCompressionMode
+{
+	// 0 .. 1
+	RCM_UNorm,
+	// -1 .. 1
+	RCM_SNorm,
+	// 0 .. 1 unless there are smaller values than 0 or bigger values than 1, then the range is extended to the minimum or the maximum of the values
+	RCM_MinMaxNorm,
+	// minimum .. maximum (each channel independent)
+	RCM_MinMax,
+};
+
+enum EPrimitiveType
+{
+	PT_TriangleList,
+	PT_TriangleStrip,
+	PT_LineList,
+	PT_QuadList,
+	PT_PointList,
+	PT_1_ControlPointPatchList,
+	PT_2_ControlPointPatchList,
+	PT_3_ControlPointPatchList,
+	PT_4_ControlPointPatchList,
+	PT_5_ControlPointPatchList,
+	PT_6_ControlPointPatchList,
+	PT_7_ControlPointPatchList,
+	PT_8_ControlPointPatchList,
+	PT_9_ControlPointPatchList,
+	PT_10_ControlPointPatchList,
+	PT_11_ControlPointPatchList,
+	PT_12_ControlPointPatchList,
+	PT_13_ControlPointPatchList,
+	PT_14_ControlPointPatchList,
+	PT_15_ControlPointPatchList,
+	PT_16_ControlPointPatchList,
+	PT_17_ControlPointPatchList,
+	PT_18_ControlPointPatchList,
+	PT_19_ControlPointPatchList,
+	PT_20_ControlPointPatchList,
+	PT_21_ControlPointPatchList,
+	PT_22_ControlPointPatchList,
+	PT_23_ControlPointPatchList,
+	PT_24_ControlPointPatchList,
+	PT_25_ControlPointPatchList,
+	PT_26_ControlPointPatchList,
+	PT_27_ControlPointPatchList,
+	PT_28_ControlPointPatchList,
+	PT_29_ControlPointPatchList,
+	PT_30_ControlPointPatchList,
+	PT_31_ControlPointPatchList,
+	PT_32_ControlPointPatchList,
+	PT_Num,
+	PT_NumBits = 6
+};
+
+static_assert(PT_Num <= (1 << 8), "EPrimitiveType doesn't fit in a byte");
+static_assert(PT_Num <= (1 << PT_NumBits), "PT_NumBits is too small");
+
+/**
+ *	Resource usage flags - for vertex and index buffers.
+ */
+enum EBufferUsageFlags
+{
+	// Mutually exclusive write-frequency flags
+	BUF_Static            = 0x0001, // The buffer will be written to once.
+	BUF_Dynamic           = 0x0002, // The buffer will be written to occasionally.
+	BUF_Volatile          = 0x0004, // The buffer will be written to frequently.
+
+	// Mutually exclusive bind flags.
+	BUF_UnorderedAccess   = 0x0008, // Allows an unordered access view to be created for the buffer.
+
+	/** Create a byte address buffer, which is basically a structured buffer with a uint32 type. */
+	BUF_ByteAddressBuffer = 0x0020,
+	/** Create a structured buffer with an atomic UAV counter. */
+	BUF_UAVCounter        = 0x0040,
+	/** Create a buffer that can be bound as a stream output target. */
+	BUF_StreamOutput      = 0x0080,
+	/** Create a buffer which contains the arguments used by DispatchIndirect or DrawIndirect. */
+	BUF_DrawIndirect      = 0x0100,
+	/** 
+	 * Create a buffer that can be bound as a shader resource. 
+	 * This is only needed for buffer types which wouldn't ordinarily be used as a shader resource, like a vertex buffer.
+	 */
+	BUF_ShaderResource    = 0x0200,
+
+	/**
+	 * Request that this buffer is directly CPU accessible
+	 * (@todo josh: this is probably temporary and will go away in a few months)
+	 */
+	BUF_KeepCPUAccessible = 0x0400,
+
+	/**
+	 * Provide information that this buffer will contain only one vertex, which should be delivered to every primitive drawn.
+	 * This is necessary for OpenGL implementations, which need to handle this case very differently (and can't handle GL_HALF_FLOAT in such vertices at all).
+	 */
+	BUF_ZeroStride        = 0x0800,
+
+	/** Buffer should go in fast vram (hint only) */
+	BUF_FastVRAM          = 0x1000,
+
+	// Helper bit-masks
+	BUF_AnyDynamic      = (BUF_Dynamic|BUF_Volatile),
+};
+
+#define ENUM_RHI_RESOURCE_TYPES(EnumerationMacro) \
+	EnumerationMacro(SamplerState,None) \
+	EnumerationMacro(RasterizerState,None) \
+	EnumerationMacro(DepthStencilState,None) \
+	EnumerationMacro(BlendState,None) \
+	EnumerationMacro(VertexDeclaration,None) \
+	EnumerationMacro(VertexShader,None) \
+	EnumerationMacro(HullShader,None) \
+	EnumerationMacro(DomainShader,None) \
+	EnumerationMacro(PixelShader,None) \
+	EnumerationMacro(GeometryShader,None) \
+	EnumerationMacro(ComputeShader,None) \
+	EnumerationMacro(BoundShaderState,None) \
+	EnumerationMacro(UniformBuffer,None) \
+	EnumerationMacro(IndexBuffer,None) \
+	EnumerationMacro(VertexBuffer,None) \
+	EnumerationMacro(StructuredBuffer,None) \
+	EnumerationMacro(Texture,None) \
+	EnumerationMacro(Texture2D,Texture) \
+	EnumerationMacro(Texture2DArray,Texture) \
+	EnumerationMacro(Texture3D,Texture) \
+	EnumerationMacro(TextureCube,Texture) \
+	EnumerationMacro(TextureReference,Texture) \
+	EnumerationMacro(RenderQuery,None) \
+	EnumerationMacro(Viewport,None) \
+	EnumerationMacro(UnorderedAccessView,None) \
+	EnumerationMacro(ShaderResourceView,None)
+
+/** An enumeration of the different RHI reference types. */
+enum ERHIResourceType
+{
+	RRT_None,
+
+#define DECLARE_RESOURCETYPE_ENUM(Type,ParentType) RRT_##Type,
+	ENUM_RHI_RESOURCE_TYPES(DECLARE_RESOURCETYPE_ENUM)
+#undef DECLARE_RESOURCETYPE_ENUM
+
+	RRT_Num
+};
+
+/** Flags used for texture creation */
+enum ETextureCreateFlags
+{
+	TexCreate_None					= 0,
+
+	// Texture can be used as a render target
+	TexCreate_RenderTargetable		= 1<<0,
+	// Texture can be used as a resolve target
+	TexCreate_ResolveTargetable		= 1<<1,
+	// Texture can be used as a depth-stencil target.
+	TexCreate_DepthStencilTargetable= 1<<2,
+	// Texture can be used as a shader resource.
+	TexCreate_ShaderResource		= 1<<3,
+
+	// Texture is encoded in sRGB gamma space
+	TexCreate_SRGB					= 1<<4,
+	// Texture will be created without a packed miptail
+	TexCreate_NoMipTail				= 1<<5,
+	// Texture will be created with an un-tiled format
+	TexCreate_NoTiling				= 1<<6,
+	// Texture that may be updated every frame
+	TexCreate_Dynamic				= 1<<8,
+	// Allow silent texture creation failure
+	// @warning:	When you update this, you must update FTextureAllocations::FindTextureType() in Core/Private/UObject/TextureAllocations.cpp
+	TexCreate_AllowFailure			= 1<<9,
+	// Disable automatic defragmentation if the initial texture memory allocation fails.
+	// @warning:	When you update this, you must update FTextureAllocations::FindTextureType() in Core/Private/UObject/TextureAllocations.cpp
+	TexCreate_DisableAutoDefrag		= 1<<10,
+	// Create the texture with automatic -1..1 biasing
+	TexCreate_BiasNormalMap			= 1<<11,
+	// Create the texture with the flag that allows mip generation later, only applicable to D3D11
+	TexCreate_GenerateMipCapable	= 1<<12,
+	// UnorderedAccessView (DX11 only)
+	// Warning: Causes additional synchronization between draw calls when using a render target allocated with this flag, use sparingly
+	// See: GCNPerformanceTweets.pdf Tip 37
+	TexCreate_UAV					= 1<<16,
+	// Render target texture that will be displayed on screen (back buffer)
+	TexCreate_Presentable			= 1<<17,
+	// Texture data is accessible by the CPU
+	TexCreate_CPUReadback			= 1<<18,
+	// Texture was processed offline (via a texture conversion process for the current platform)
+	TexCreate_OfflineProcessed		= 1<<19,
+	// Texture needs to go in fast VRAM if available (HINT only)
+	TexCreate_FastVRAM				= 1<<20,
+	// by default the texture is not showing up in the list - this is to reduce clutter, using the FULL option this can be ignored
+	TexCreate_HideInVisualizeTexture= 1<<21,
+	// Texture should be created in virtual memory, with no physical memory allocation made
+	// You must make further calls to RHIVirtualTextureSetFirstMipInMemory to allocate physical memory
+	// and RHIVirtualTextureSetFirstMipVisible to map the first mip visible to the GPU
+	TexCreate_Virtual				= 1<<22,
+	// Creates a RenderTargetView for each array slice of the texture
+	// Warning: if this was specified when the resource was created, you can't use SV_RenderTargetArrayIndex to route to other slices!
+	TexCreate_TargetArraySlicesIndependently	= 1<<23,
+	// Texture that may be shared with DX9 or other devices
+	TexCreate_Shared = 1 << 24,
+};
+
+/**
+ * Async texture reallocation status, returned by RHIGetReallocateTexture2DStatus().
+ */
+enum ETextureReallocationStatus
+{
+	TexRealloc_Succeeded = 0,
+	TexRealloc_Failed,
+	TexRealloc_InProgress,
+};
+
+inline bool IsPCPlatform(const EShaderPlatform Platform)
+{
+	return Platform == SP_PCD3D_SM5 || Platform == SP_PCD3D_SM4 || Platform == SP_PCD3D_ES2 || Platform ==  SP_OPENGL_SM4 || Platform == SP_OPENGL_SM5 || Platform == SP_OPENGL_PCES2;
+}
+
+/** Whether the shader platform corresponds to the ES2 feature level. */
+inline bool IsES2Platform(const EShaderPlatform Platform)
+{
+	return Platform == SP_PCD3D_ES2 || Platform == SP_OPENGL_PCES2 || Platform == SP_OPENGL_ES2 || Platform == SP_OPENGL_ES2_WEBGL || Platform == SP_OPENGL_ES2_IOS; 
+}
+
+inline bool IsOpenGLPlatform(const EShaderPlatform Platform)
+{
+	return Platform == SP_OPENGL_SM4 || Platform == SP_OPENGL_SM5 || Platform == SP_OPENGL_PCES2 || Platform == SP_OPENGL_ES2 || Platform == SP_OPENGL_ES2_WEBGL || Platform == SP_OPENGL_ES2_IOS;
+}
+
+inline bool IsConsolePlatform(const EShaderPlatform Platform)
+{
+	return Platform == SP_PS4 || Platform == SP_XBOXONE;
+}
+
+inline ERHIFeatureLevel::Type GetMaxSupportedFeatureLevel(EShaderPlatform InShaderPlatform)
+{
+	switch(InShaderPlatform)
+	{
+	case SP_PCD3D_SM5:
+	case SP_OPENGL_SM5:
+	case SP_PS4:
+	case SP_XBOXONE:
+		return ERHIFeatureLevel::SM5;
+	case SP_PCD3D_SM4:
+	case SP_OPENGL_SM4:
+		return ERHIFeatureLevel::SM4;
+	case SP_PCD3D_ES2:
+	case SP_OPENGL_PCES2:
+	case SP_OPENGL_ES2:
+	case SP_OPENGL_ES2_WEBGL:
+	case  SP_OPENGL_ES2_IOS:
+		return ERHIFeatureLevel::ES2;
+	default:
+		check(0);
+		return ERHIFeatureLevel::Num;
+	}	
+}
+
+/** Returns true if the feature level is supported by the shader platform. */
+inline bool IsFeatureLevelSupported(EShaderPlatform InShaderPlatform, ERHIFeatureLevel::Type InFeatureLevel)
+{
+	switch(InShaderPlatform)
+	{
+	case SP_PCD3D_SM5:
+		return InFeatureLevel <= ERHIFeatureLevel::SM5;
+	case SP_PCD3D_SM4:
+		return InFeatureLevel <= ERHIFeatureLevel::SM4;
+	case SP_PCD3D_ES2:
+		return InFeatureLevel <= ERHIFeatureLevel::ES2;
+	case SP_OPENGL_PCES2:
+		return InFeatureLevel <= ERHIFeatureLevel::ES2;
+	case SP_OPENGL_ES2:
+		return InFeatureLevel <= ERHIFeatureLevel::ES2;
+	case SP_OPENGL_ES2_WEBGL: 
+		return InFeatureLevel <= ERHIFeatureLevel::ES2;
+	case SP_OPENGL_ES2_IOS: 
+		return InFeatureLevel <= ERHIFeatureLevel::ES2;
+	case SP_OPENGL_SM4:
+		return InFeatureLevel <= ERHIFeatureLevel::SM4;
+	case SP_OPENGL_SM5:
+		return InFeatureLevel <= ERHIFeatureLevel::SM5;
+	case SP_PS4:
+		return InFeatureLevel <= ERHIFeatureLevel::SM5;
+	case SP_XBOXONE:
+		return InFeatureLevel <= ERHIFeatureLevel::SM5;
+	default:
+		return false;
+	}	
+}
+
+inline bool RHISupportsTessellation(const EShaderPlatform Platform)
+{
+	if (IsFeatureLevelSupported(Platform, ERHIFeatureLevel::SM5))
+	{
+		return ((Platform == SP_PCD3D_SM5) || (Platform == SP_XBOXONE) || (Platform == SP_OPENGL_SM5));
+	}
+	return false;
+}
+
+inline uint32 GetFeatureLevelMaxTextureSamplers(ERHIFeatureLevel::Type FeatureLevel)
+{
+	if (FeatureLevel == ERHIFeatureLevel::ES2)
+	{
+		return 8;
+	}
+	else
+	{
+		return 16;
+	}
+}
+
+inline int32 GetFeatureLevelMaxNumberOfBones(ERHIFeatureLevel::Type FeatureLevel)
+{
+	switch (FeatureLevel)
+	{
+	case ERHIFeatureLevel::ES2:
+	case ERHIFeatureLevel::SM3:
+		return 75;
+	case ERHIFeatureLevel::SM4:
+	case ERHIFeatureLevel::SM5:
+		return 256;
+	default:
+		check(0);
+	}
+
+	return 0;
+}
+
+inline bool IsUniformBufferResourceType(EUniformBufferBaseType BaseType)
+{
+	return BaseType == UBMT_SRV || BaseType == UBMT_UAV || BaseType == UBMT_SAMPLER || BaseType == UBMT_TEXTURE;
+}
+
+#endif	// __RHIDEFINITIONS_H__
diff --git a/Engine/Source/Runtime/RHI/Public/RHIMethods.h b/Engine/Source/Runtime/RHI/Public/RHIMethods.h
--- a/Engine/Source/Runtime/RHI/Public/RHIMethods.h
+++ b/Engine/Source/Runtime/RHI/Public/RHIMethods.h
@@ -228,9 +228,9 @@ DEFINE_RHIMETHOD_2(
 DEFINE_RHIMETHOD_3(
 	FUniformBufferRHIRef,RHICreateUniformBuffer,
 	const void*,Contents,
-	uint32,NumBytes,
+	const FRHIUniformBufferLayout&,Layout,
 	EUniformBufferUsage,Usage,
-	return,return new FRHIUniformBuffer(NumBytes);
+	return,return new FRHIUniformBuffer(Layout);
 	);
 
 
@@ -455,6 +455,18 @@ DEFINE_RHIMETHOD_5(
 	return,return false
 	);
 
+DEFINE_RHIMETHOD_0(
+	FTextureReferenceRHIRef,RHICreateTextureReference,
+	return, return new FRHITextureReferenceNullImpl();
+	);
+
+DEFINE_RHIMETHOD_2(
+	void,RHIUpdateTextureReference,
+	FTextureReferenceRHIParamRef,TextureRef,
+	FTextureRHIParamRef,NewTexture,
+	,if (TextureRef) { ((FRHITextureReferenceNullImpl*)TextureRef)->SetReferencedTexture(NewTexture); }
+	);
+
 /**
 * Creates a 2D RHI texture resource
 * @param SizeX - width of the texture to create
@@ -589,7 +601,7 @@ DEFINE_RHIMETHOD_4(
 	uint8, NumMipLevels,
 	uint8, Format,		
 	return,return new FRHIShaderResourceView();
-	);
+);
 
 /**
 * Generates mip maps for a texture.
@@ -794,24 +806,24 @@ DEFINE_RHIMETHOD_5(
 	);
 
 
-/**
-* Creates a Cube Array RHI texture resource
-* @param Size - width/height of the texture to create
-* @param ArraySize - number of array elements of the texture to create
-* @param Format - EPixelFormat texture format
-* @param NumMips - number of mips to generate or 0 for full mip pyramid
-* @param Flags - ETextureCreateFlags creation flags
-*/
-DEFINE_RHIMETHOD_6(
-	FTextureCubeRHIRef,RHICreateTextureCubeArray,
-	uint32,Size,
-	uint32,ArraySize,
-	uint8,Format,
-	uint32,NumMips,
-	uint32,Flags,
+	/**
+	* Creates a Cube Array RHI texture resource
+	* @param Size - width/height of the texture to create
+	* @param ArraySize - number of array elements of the texture to create
+	* @param Format - EPixelFormat texture format
+	* @param NumMips - number of mips to generate or 0 for full mip pyramid
+	* @param Flags - ETextureCreateFlags creation flags
+	*/
+	DEFINE_RHIMETHOD_6(
+		FTextureCubeRHIRef,RHICreateTextureCubeArray,
+		uint32,Size,
+		uint32,ArraySize,
+		uint8,Format,
+		uint32,NumMips,
+		uint32,Flags,
 	FRHIResourceCreateInfo&,CreateInfo,
-	return,return new FRHITextureCube(Size,NumMips,(EPixelFormat)Format,Flags);
-);
+		return,return new FRHITextureCube(Size,NumMips,(EPixelFormat)Format,Flags);
+	);
 
 
 /**
@@ -992,10 +1004,19 @@ DEFINE_RHIMETHOD_0(
 	return,return;
 	);
 
+/**
+ * Signals the beginning of scene rendering. The RHI makes certain caching assumptions between
+ * calls to BeginScene/EndScene. Currently the only restriction is that you can't update texture
+ * references.
+ */
 DEFINE_RHIMETHOD_0(
 	void,RHIBeginScene,
 	return,return;
 	);
+
+/**
+ * Signals the end of scene rendering. See RHIBeginScene.
+ */
 DEFINE_RHIMETHOD_0(
 	void,RHIEndScene,
 	return,return;
@@ -1681,6 +1702,12 @@ DEFINE_RHIMETHOD_2(
 	,
 	);
 
+DEFINE_RHIMETHOD_1(
+	void,RHIExecuteCommandList,
+	FRHICommandList*,RHICmdList,
+	,
+	);
+
 /**
  * Enabled/Disables Depth Bounds Testing with the given min/max depth.
  * @param bEnable	Enable(non-zero)/disable(zero) the depth bounds test
diff --git a/Engine/Source/Runtime/RHI/Public/RHIResources.h b/Engine/Source/Runtime/RHI/Public/RHIResources.h
--- a/Engine/Source/Runtime/RHI/Public/RHIResources.h
+++ b/Engine/Source/Runtime/RHI/Public/RHIResources.h
@@ -3,8 +3,9 @@
 
 #pragma once
 
-#include "RHI.h"
+#include "RHIDefinitions.h"
 #include "RefCounting.h"
+#include "Runtime/Engine/Public/PixelFormat.h" // for EPixelFormat
 
 /** The base type of RHI resources. */
 class RHI_API FRHIResource : public FRefCountedObject
@@ -44,20 +45,77 @@ class FRHIComputeShader : public FRHIResource {};
 // Buffers
 //
 
+/** The layout of a uniform buffer in memory. */
+struct FRHIUniformBufferLayout
+{
+	/** The size of the constant buffer in bytes. */
+	uint32 ConstantBufferSize;
+	/** The offset to the beginning of the resource table. */
+	uint32 ResourceOffset;
+	/** The type of each resource (EUniformBufferBaseType). */
+	TArray<uint8> Resources;
+
+	uint32 GetHash() const
+	{
+		if (!bComputedHash)
+		{
+			uint32 TmpHash = ConstantBufferSize;
+			TmpHash ^= ResourceOffset;
+			uint32 N = Resources.Num();
+			while (N >= 4)
+			{
+				TmpHash ^= (Resources[--N] << 0);
+				TmpHash ^= (Resources[--N] << 8);
+				TmpHash ^= (Resources[--N] << 16);
+				TmpHash ^= (Resources[--N] << 24);
+			}
+			while (N > 0)
+			{
+				TmpHash ^= Resources[--N];
+			}
+			Hash = TmpHash;
+			bComputedHash = true;
+		}
+		return Hash;
+	}
+
+	FRHIUniformBufferLayout() :
+		ConstantBufferSize(0),
+		ResourceOffset(0),
+		Hash(0),
+		bComputedHash(false)
+	{
+	}
+
+private:
+	mutable uint32 Hash;
+	mutable bool bComputedHash;
+};
+
+/** Compare two uniform buffer layouts. */
+inline bool operator==(const FRHIUniformBufferLayout& A, const FRHIUniformBufferLayout& B)
+{
+	return A.ConstantBufferSize == B.ConstantBufferSize
+		&& A.ResourceOffset == B.ResourceOffset
+		&& A.Resources == B.Resources;
+}
+
 class FRHIUniformBuffer : public FRHIResource
 {
 public:
 
 	/** Initialization constructor. */
-	FRHIUniformBuffer(uint32 InSize)
-	: Size(InSize)
+	FRHIUniformBuffer(const FRHIUniformBufferLayout& InLayout)
+	: Layout(&InLayout)
 	{}
 
 	/** @return The number of bytes in the uniform buffer. */
-	uint32 GetSize() const { return Size; }
+	uint32 GetSize() const { return Layout->ConstantBufferSize; }
+	const FRHIUniformBufferLayout& GetLayout() const { return *Layout; }
 
 private:
-	uint32 Size;
+	/** Layout of the uniform buffer. */
+	const FRHIUniformBufferLayout* Layout;
 };
 
 class FRHIIndexBuffer : public FRHIResource
@@ -146,6 +204,7 @@ public:
 	, NumSamples(InNumSamples)
 	, Format(InFormat)
 	, Flags(InFlags)
+	, LastCachedTime(-FLT_MAX)
 	{}
 
 	// Dynamic cast methods.
@@ -153,6 +212,7 @@ public:
 	virtual class FRHITexture2DArray* GetTexture2DArray() { return NULL; }
 	virtual class FRHITexture3D* GetTexture3D() { return NULL; }
 	virtual class FRHITextureCube* GetTextureCube() { return NULL; }
+	virtual class FRHITextureReference* GetTextureReference() { return NULL; }
 	
 	/**
 	 * Returns access to the platform-specific native resource pointer.  This is designed to be used to provide plugins with access
@@ -179,15 +239,25 @@ public:
 	uint32 GetNumSamples() const { return NumSamples; }
 
 	/** @return Whether the texture is multi sampled. */
-	bool IsMultisampled() const { return NumSamples > 1; }
+	bool IsMultisampled() const { return NumSamples > 1; }		
 
 	FRHIResourceInfo ResourceInfo;
 
+	/** @return the last time this texture was cached in a resource table. */
+	inline float GetLastCachedTime() const { return LastCachedTime; }
+
+	/** sets the last time this texture was cached in a resource table. */
+	void SetLastCachedTime(float InLastCachedTime)
+	{
+		LastCachedTime = InLastCachedTime;
+	}
+
 private:
 	uint32 NumMips;
 	uint32 NumSamples;
 	EPixelFormat Format;
 	uint32 Flags;
+	float LastCachedTime;
 };
 
 class RHI_API FRHITexture2D : public FRHITexture
@@ -299,6 +369,39 @@ private:
 	uint32 Size;
 };
 
+class RHI_API FRHITextureReference : public FRHITexture
+{
+public:
+	explicit FRHITextureReference()
+		: FRHITexture(0,0,PF_Unknown,0)
+	{}
+
+	virtual FRHITextureReference* GetTextureReference() OVERRIDE { return this; }
+	inline FRHITexture* GetReferencedTexture() const { return ReferencedTexture.GetReference(); }
+
+protected:
+	void SetReferencedTexture(FRHITexture* InTexture)
+	{
+		ReferencedTexture = InTexture;
+	}
+
+private:
+	TRefCountPtr<FRHITexture> ReferencedTexture;
+};
+
+class RHI_API FRHITextureReferenceNullImpl : public FRHITextureReference
+{
+public:
+	FRHITextureReferenceNullImpl()
+		: FRHITextureReference()
+	{}
+
+	void SetReferencedTexture(FRHITexture* InTexture)
+	{
+		FRHITextureReference::SetReferencedTexture(InTexture);
+	}
+};
+
 //
 // Misc
 //
@@ -345,4 +448,25 @@ public:
 		MipIndex(InMipIndex),
 		ArraySliceIndex(InArraySliceIndex)
 	{}
-};
\ No newline at end of file
+};
+
+// Template magic to convert an FRHI*Shader to its enum
+template<typename TRHIShader> struct TRHIShaderToEnum {};
+template<> struct TRHIShaderToEnum<FRHIVertexShader>	{ enum { ShaderFrequency = SF_Vertex }; };
+template<> struct TRHIShaderToEnum<FRHIHullShader>		{ enum { ShaderFrequency = SF_Hull }; };
+template<> struct TRHIShaderToEnum<FRHIDomainShader>	{ enum { ShaderFrequency = SF_Domain }; };
+template<> struct TRHIShaderToEnum<FRHIPixelShader>		{ enum { ShaderFrequency = SF_Pixel }; };
+template<> struct TRHIShaderToEnum<FRHIGeometryShader>	{ enum { ShaderFrequency = SF_Geometry }; };
+template<> struct TRHIShaderToEnum<FRHIComputeShader>	{ enum { ShaderFrequency = SF_Compute }; };
+template<> struct TRHIShaderToEnum<FVertexShaderRHIParamRef>	{ enum { ShaderFrequency = SF_Vertex }; };
+template<> struct TRHIShaderToEnum<FHullShaderRHIParamRef>		{ enum { ShaderFrequency = SF_Hull }; };
+template<> struct TRHIShaderToEnum<FDomainShaderRHIParamRef>	{ enum { ShaderFrequency = SF_Domain }; };
+template<> struct TRHIShaderToEnum<FPixelShaderRHIParamRef>		{ enum { ShaderFrequency = SF_Pixel }; };
+template<> struct TRHIShaderToEnum<FGeometryShaderRHIParamRef>	{ enum { ShaderFrequency = SF_Geometry }; };
+template<> struct TRHIShaderToEnum<FComputeShaderRHIParamRef>	{ enum { ShaderFrequency = SF_Compute }; };
+template<> struct TRHIShaderToEnum<FVertexShaderRHIRef>		{ enum { ShaderFrequency = SF_Vertex }; };
+template<> struct TRHIShaderToEnum<FHullShaderRHIRef>		{ enum { ShaderFrequency = SF_Hull }; };
+template<> struct TRHIShaderToEnum<FDomainShaderRHIRef>		{ enum { ShaderFrequency = SF_Domain }; };
+template<> struct TRHIShaderToEnum<FPixelShaderRHIRef>		{ enum { ShaderFrequency = SF_Pixel }; };
+template<> struct TRHIShaderToEnum<FGeometryShaderRHIRef>	{ enum { ShaderFrequency = SF_Geometry }; };
+template<> struct TRHIShaderToEnum<FComputeShaderRHIRef>	{ enum { ShaderFrequency = SF_Compute }; };
diff --git a/Engine/Source/Runtime/RenderCore/Public/UniformBuffer.h b/Engine/Source/Runtime/RenderCore/Public/UniformBuffer.h
--- a/Engine/Source/Runtime/RenderCore/Public/UniformBuffer.h
+++ b/Engine/Source/Runtime/RenderCore/Public/UniformBuffer.h
@@ -50,7 +50,7 @@ public:
 	// FRenderResource interface.
 	virtual void InitDynamicRHI()
 	{
-		UniformBufferRHI = RHICreateUniformBuffer(Contents,sizeof(TBufferStruct),UniformBuffer_MultiUse);
+		UniformBufferRHI = RHICreateUniformBuffer(Contents,TBufferStruct::StaticStruct.GetLayout(),UniformBuffer_MultiFrame);
 	}
 	virtual void ReleaseDynamicRHI()
 	{
@@ -85,7 +85,7 @@ public:
 	static TUniformBufferRef<TBufferStruct> CreateUniformBufferImmediate(const TBufferStruct& Value, EUniformBufferUsage Usage)
 	{
 		check(IsInRenderingThread());
-		return TUniformBufferRef<TBufferStruct>(RHICreateUniformBuffer(&Value,sizeof(Value),Usage));
+		return TUniformBufferRef<TBufferStruct>(RHICreateUniformBuffer(&Value,TBufferStruct::StaticStruct.GetLayout(),Usage));
 	}
 
 private:
@@ -117,15 +117,15 @@ void BeginSetUniformBufferContents(
 		TBufferStruct,Struct);
 }
 
-/** The base type of a value in a uniform buffer. */
-enum EUniformBufferBaseType
+/** Each entry in a resource table is provided to the shader compiler for creating mappings. */
+struct FResourceTableEntry
 {
-	UBMT_INVALID,
-	UBMT_BOOL,
-	UBMT_INT32,
-	UBMT_UINT32,
-	UBMT_FLOAT32,
-	UBMT_STRUCT
+	/** The name of the uniform buffer in which this resource exists. */
+	FString UniformBufferName;
+	/** The type of the resource (EUniformBufferBaseType). */
+	uint16 Type;
+	/** The index of the resource in the table. */
+	uint16 ResourceIndex;
 };
 
 /** A uniform buffer struct. */
@@ -141,6 +141,7 @@ public:
 		/** Initialization constructor. */
 		FMember(
 			const TCHAR* InName,
+			const TCHAR* InShaderType,
 			uint32 InOffset,
 			EUniformBufferBaseType InBaseType,
 			EShaderPrecisionModifier::Type InPrecision,
@@ -150,6 +151,7 @@ public:
 			const FUniformBufferStruct* InStruct
 			)
 		:	Name(InName)
+		,	ShaderType(InShaderType)
 		,	Offset(InOffset)
 		,	BaseType(InBaseType)
 		,	Precision(InPrecision)
@@ -160,6 +162,7 @@ public:
 		{}
 	
 		const TCHAR* GetName() const { return Name; }
+		const TCHAR* GetShaderType() const { return ShaderType; }
 		uint32 GetOffset() const { return Offset; }
 		EUniformBufferBaseType GetBaseType() const { return BaseType; }
 		EShaderPrecisionModifier::Type GetPrecision() const { return Precision; }
@@ -171,6 +174,7 @@ public:
 	private:
 
 		const TCHAR* Name;
+		const TCHAR* ShaderType;
 		uint32 Offset;
 		EUniformBufferBaseType BaseType;
 		EShaderPrecisionModifier::Type Precision;
@@ -191,6 +195,25 @@ public:
 	,	Members(InMembers)
 	,	GlobalListLink(this)
 	{
+		bool bHasDeclaredResources = false;
+		Layout.ConstantBufferSize = InSize;
+		Layout.ResourceOffset = 0;
+		for (int32 i = 0; i < Members.Num(); ++i)
+		{
+			bool bIsResource = IsUniformBufferResourceType(Members[i].GetBaseType());
+			checkf(!bHasDeclaredResources || bIsResource, TEXT("Constructing invalid uniform buffer struct: member '%s' is not a resource but has been declared after a resource member."), Members[i].GetName());
+			if (bIsResource)
+			{
+				if (!bHasDeclaredResources)
+				{
+					Layout.ConstantBufferSize = (i == 0) ? 0 : Align(Members[i].GetOffset(),UNIFORM_BUFFER_STRUCT_ALIGNMENT);
+					Layout.ResourceOffset = Members[i].GetOffset();
+				}
+				Layout.Resources.Add(Members[i].GetBaseType());
+			}
+			bHasDeclaredResources |= bIsResource;
+		}
+
 		if (bRegisterForAutoBinding)
 		{
 			GlobalListLink.Link(GetStructList());
@@ -202,9 +225,30 @@ public:
 		GlobalListLink.Unlink();
 	}
 
+	void AddResourceTableEntries(TMap<FString,FResourceTableEntry>& ResourceTableMap, TMap<FString,uint32>& ResourceTableLayoutHashes) const
+	{
+		uint16 ResourceIndex = 0;
+		for (int32 MemberIndex = 0; MemberIndex < Members.Num(); ++MemberIndex)
+		{
+			const FMember& Member = Members[MemberIndex];
+			if (IsUniformBufferResourceType(Member.GetBaseType()))
+			{
+				FResourceTableEntry& Entry = ResourceTableMap.FindOrAdd(FString::Printf(TEXT("%s_%s"), ShaderVariableName, Member.GetName()));
+				if (Entry.UniformBufferName.IsEmpty())
+				{
+					Entry.UniformBufferName = ShaderVariableName;
+					Entry.Type = Member.GetBaseType();
+					Entry.ResourceIndex = ResourceIndex++;
+				}
+			}
+		}
+		ResourceTableLayoutHashes.Add(ShaderVariableName,GetLayout().GetHash());
+	}
+
 	const TCHAR* GetStructTypeName() const { return StructTypeName; }
 	const TCHAR* GetShaderVariableName() const { return ShaderVariableName; }
 	const uint32 GetSize() const { return Size; }
+	const FRHIUniformBufferLayout& GetLayout() const { return Layout; }
 	const TArray<FMember>& GetMembers() const { return Members; }
 	FShaderUniformBufferParameter* ConstructTypedParameter() const { return (*ConstructUniformBufferParameterRef)(); }
 
@@ -215,6 +259,7 @@ private:
 	const TCHAR* ShaderVariableName;
 	ConstructUniformBufferParameterType ConstructUniformBufferParameterRef;
 	uint32 Size;
+	FRHIUniformBufferLayout Layout;
 	TArray<FMember> Members;
 	TLinkedList<FUniformBufferStruct*> GlobalListLink;
 };
@@ -253,6 +298,7 @@ private:
 		enum { NumColumns = 0 };
 		enum { NumElements = 0 };
 		enum { Alignment = 1 };
+		enum { IsResource = 0 };
 		static const FUniformBufferStruct* GetStruct() { return NULL; }
 	};
 
@@ -265,6 +311,7 @@ private:
 		enum { NumColumns = 1 };
 		enum { NumElements = 0 };
 		enum { Alignment = 4 };
+		enum { IsResource = 0 };
 		typedef TAlignedTypedef<bool,Alignment>::TAlignedType TAlignedType;
 		static const FUniformBufferStruct* GetStruct() { return NULL; }
 	};
@@ -278,6 +325,7 @@ private:
 		enum { NumColumns = 1 };
 		enum { NumElements = 0 };
 		enum { Alignment = 4 };
+		enum { IsResource = 0 };
 		typedef uint32 TAlignedType;
 		static const FUniformBufferStruct* GetStruct() { return NULL; }
 	};
@@ -292,6 +340,7 @@ private:
 		enum { NumColumns = 1 };
 		enum { NumElements = 0 };
 		enum { Alignment = 4 };
+		enum { IsResource = 0 };
 		typedef int32 TAlignedType;
 		static const FUniformBufferStruct* GetStruct() { return NULL; }
 	};
@@ -305,6 +354,7 @@ private:
 		enum { NumColumns = 1 };
 		enum { NumElements = 0 };
 		enum { Alignment = 4 };
+		enum { IsResource = 0 };
 		typedef float TAlignedType;
 		static const FUniformBufferStruct* GetStruct() { return NULL; }
 	};
@@ -318,6 +368,7 @@ private:
 		enum { NumColumns = 2 };
 		enum { NumElements = 0 };
 		enum { Alignment = 8 };
+		enum { IsResource = 0 };
 		typedef TAlignedTypedef<FVector2D,Alignment>::TAlignedType TAlignedType;
 		static const FUniformBufferStruct* GetStruct() { return NULL; }
 	};
@@ -331,6 +382,7 @@ private:
 		enum { NumColumns = 3 };
 		enum { NumElements = 0 };
 		enum { Alignment = 16 };
+		enum { IsResource = 0 };
 		typedef TAlignedTypedef<FVector,Alignment>::TAlignedType TAlignedType;
 		static const FUniformBufferStruct* GetStruct() { return NULL; }
 	};
@@ -344,6 +396,7 @@ private:
 		enum { NumColumns = 4 };
 		enum { NumElements = 0 };
 		enum { Alignment = 16 };
+		enum { IsResource = 0 };
 		typedef TAlignedTypedef<FVector4,Alignment>::TAlignedType TAlignedType;
 		static const FUniformBufferStruct* GetStruct() { return NULL; }
 	};
@@ -357,6 +410,7 @@ private:
 		enum { NumColumns = 4 };
 		enum { NumElements = 0 };
 		enum { Alignment = 16 };
+		enum { IsResource = 0 };
 		typedef TAlignedTypedef<FLinearColor,Alignment>::TAlignedType TAlignedType;
 		static const FUniformBufferStruct* GetStruct() { return NULL; }
 	};
@@ -370,6 +424,7 @@ private:
 		enum { NumColumns = 2 };
 		enum { NumElements = 0 };
 		enum { Alignment = 8 };
+		enum { IsResource = 0 };
 		typedef TAlignedTypedef<FIntPoint,Alignment>::TAlignedType TAlignedType;
 		static const FUniformBufferStruct* GetStruct() { return NULL; }
 	};
@@ -383,6 +438,7 @@ private:
 		enum { NumColumns = 3 };
 		enum { NumElements = 0 };
 		enum { Alignment = 16 };
+		enum { IsResource = 0 };
 		typedef TAlignedTypedef<FIntVector,Alignment>::TAlignedType TAlignedType;
 		static const FUniformBufferStruct* GetStruct() { return NULL; }
 	};
@@ -396,6 +452,7 @@ private:
 		enum { NumColumns = 4 };
 		enum { NumElements = 0 };
 		enum { Alignment = 16 };
+		enum { IsResource = 0 };
 		typedef TAlignedTypedef<FIntRect,Alignment>::TAlignedType TAlignedType;
 		static const FUniformBufferStruct* GetStruct() { return NULL; }
 	};
@@ -409,6 +466,7 @@ private:
 		enum { NumColumns = 4 };
 		enum { NumElements = 0 };
 		enum { Alignment = 16 };
+		enum { IsResource = 0 };
 		typedef TAlignedTypedef<FMatrix,Alignment>::TAlignedType TAlignedType;
 		static const FUniformBufferStruct* GetStruct() { return NULL; }
 	};
@@ -422,6 +480,7 @@ private:
 		enum { NumColumns = TUniformBufferTypeInfo<T>::NumColumns };
 		enum { NumElements = InNumElements };
 		enum { Alignment = TUniformBufferTypeInfo<T>::Alignment };
+		enum { IsResource = TUniformBufferTypeInfo<T>::IsResource };
 		typedef TStaticArray<T,InNumElements,16> TAlignedType;
 		static const FUniformBufferStruct* GetStruct() { return TUniformBufferTypeInfo<T>::GetStruct(); }
 	};
@@ -435,10 +494,75 @@ private:
 		enum { NumColumns = TUniformBufferTypeInfo<T>::NumColumns };
 		enum { NumElements = InNumElements };
 		enum { Alignment = TUniformBufferTypeInfo<T>::Alignment };
+		enum { IsResource = TUniformBufferTypeInfo<T>::IsResource };
 		typedef TStaticArray<T,InNumElements,16> TAlignedType;
 		static const FUniformBufferStruct* GetStruct() { return TUniformBufferTypeInfo<T>::GetStruct(); }
 	};
 
+	template<>
+	class TUniformBufferTypeInfo<FShaderResourceViewRHIParamRef>
+	{
+	public:
+		enum { BaseType = UBMT_SRV };
+		enum { NumRows = 1 };
+		enum { NumColumns = 1 };
+		enum { NumElements = 0 };
+		enum { Alignment = sizeof(void*) };
+		enum { IsResource = 1 };
+		typedef TAlignedTypedef<FShaderResourceViewRHIParamRef,Alignment>::TAlignedType TAlignedType;
+		static const FUniformBufferStruct* GetStruct() { return NULL; }
+	};
+	checkAtCompileTime(sizeof(FShaderResourceViewRHIParamRef) == sizeof(void*),SRVParamRef_NotPointerSized);
+	static_assert(sizeof(TUniformBufferTypeInfo<FShaderResourceViewRHIParamRef>::TAlignedType) == sizeof(void*), "SRV UniformBufferParam is not aligned to pointer type");
+
+	template<>
+	class TUniformBufferTypeInfo<FUnorderedAccessViewRHIParamRef>
+	{
+	public:
+		enum { BaseType = UBMT_UAV };
+		enum { NumRows = 1 };
+		enum { NumColumns = 1 };
+		enum { NumElements = 0 };
+		enum { Alignment = sizeof(void*) };
+		enum { IsResource = 1 };
+		typedef TAlignedTypedef<FUnorderedAccessViewRHIParamRef,Alignment>::TAlignedType TAlignedType;
+		static const FUniformBufferStruct* GetStruct() { return NULL; }
+	};
+	checkAtCompileTime(sizeof(FUnorderedAccessViewRHIParamRef) == sizeof(void*),UAVParamRef_NotPointerSized);
+	static_assert(sizeof(TUniformBufferTypeInfo<FUnorderedAccessViewRHIParamRef>::TAlignedType) == sizeof(void*), "UAV UniformBufferParam is not aligned to pointer type");
+
+	template<>
+	class TUniformBufferTypeInfo<FSamplerStateRHIParamRef>
+	{
+	public:
+		enum { BaseType = UBMT_SAMPLER };
+		enum { NumRows = 1 };
+		enum { NumColumns = 1 };
+		enum { NumElements = 0 };
+		enum { Alignment = sizeof(void*) };
+		enum { IsResource = 1 };
+		typedef TAlignedTypedef<FSamplerStateRHIParamRef,Alignment>::TAlignedType TAlignedType;
+		static const FUniformBufferStruct* GetStruct() { return NULL; }
+	};
+	checkAtCompileTime(sizeof(FSamplerStateRHIParamRef) == sizeof(void*),SamplerStateParamRef_NotPointerSized);
+	static_assert(sizeof(TUniformBufferTypeInfo<FSamplerStateRHIParamRef>::TAlignedType) == sizeof(void*), "SamplerState UniformBufferParam is not aligned to pointer type");
+
+	template<>
+	class TUniformBufferTypeInfo<FTextureRHIParamRef>
+	{
+	public:
+		enum { BaseType = UBMT_TEXTURE };
+		enum { NumRows = 1 };
+		enum { NumColumns = 1 };
+		enum { NumElements = 0 };
+		enum { Alignment = sizeof(void*) };
+		enum { IsResource = 1 };
+		typedef TAlignedTypedef<FTextureRHIParamRef,Alignment>::TAlignedType TAlignedType;
+		static const FUniformBufferStruct* GetStruct() { return NULL; }
+	};
+	checkAtCompileTime(sizeof(FTextureRHIParamRef) == sizeof(void*),TextureParamRef_NotPointerSized);
+	static_assert(sizeof(TUniformBufferTypeInfo<FTextureRHIParamRef>::TAlignedType) == sizeof(void*), "Texture UniformBufferParam is not aligned to pointer type");
+
 //
 // Macros for declaring uniform buffer structures.
 //
@@ -459,9 +583,13 @@ private:
 	public: \
 		static FUniformBufferStruct StaticStruct; \
 		static FShaderUniformBufferParameter* ConstructUniformBufferParameter() { return new TShaderUniformBufferParameter<StructTypeName>(); } \
+		static FUniformBufferRHIRef CreateUniformBuffer(const StructTypeName& InContents, EUniformBufferUsage InUsage) \
+		{ \
+			return RHICreateUniformBuffer(&InContents,StaticStruct.GetLayout(),InUsage); \
+		} \
 	private: \
 		typedef StructTypeName zzTThisStruct; \
-		struct zzFirstMemberId {}; \
+		struct zzFirstMemberId { enum { HasDeclaredResource = 0 }; }; \
 		static TArray<FUniformBufferStruct::FMember> zzGetMembersBefore(zzFirstMemberId) \
 		{ \
 			return TArray<FUniformBufferStruct::FMember>(); \
@@ -469,22 +597,26 @@ private:
 		typedef zzFirstMemberId
 
 /** Declares a member of a uniform buffer struct. */
-#define DECLARE_UNIFORM_BUFFER_STRUCT_MEMBER_ARRAY_EX(MemberType,MemberName,ArrayDecl,Precision) \
+#define DECLARE_UNIFORM_BUFFER_STRUCT_MEMBER_EXPLICIT(MemberType,MemberName,ArrayDecl,Precision,OptionalShaderType) \
 		zzMemberId##MemberName; \
 	public: \
 		typedef MemberType zzA##MemberName ArrayDecl; \
 		typedef TUniformBufferTypeInfo<zzA##MemberName>::TAlignedType zzT##MemberName; \
 		zzT##MemberName MemberName; \
-		checkAtCompileTime(TUniformBufferTypeInfo<zzA##MemberName>::BaseType != UBMT_INVALID,InvalidMemberType##MemberName); \
+		checkAtCompileTime(TUniformBufferTypeInfo<zzA##MemberName>::BaseType != UBMT_INVALID,InvalidMemberType_##MemberName_##MemberType); \
+		checkAtCompileTime(TUniformBufferTypeInfo<zzA##MemberName>::BaseType != UBMT_UAV,UAVsNotYetSupportedInResourceTables_##MemberName_##MemberType); \
 	private: \
-		struct zzNextMemberId##MemberName {}; \
+		struct zzNextMemberId##MemberName { enum { HasDeclaredResource = zzMemberId##MemberName::HasDeclaredResource || TUniformBufferTypeInfo<zzT##MemberName>::IsResource }; }; \
 		static TArray<FUniformBufferStruct::FMember> zzGetMembersBefore(zzNextMemberId##MemberName) \
 		{ \
+			checkAtCompileTime(TUniformBufferTypeInfo<zzT##MemberName>::IsResource == 1 || zzMemberId##MemberName::HasDeclaredResource == 0, AllResourcesMustBeDeclaredLast_##MemberName); \
+			checkAtCompileTime(TUniformBufferTypeInfo<zzT##MemberName>::IsResource == 0 || IS_TCHAR_ARRAY(OptionalShaderType), NoShaderTypeFor_##MemberName); \
 			/* Route the member enumeration on to the function for the member following this. */ \
 			TArray<FUniformBufferStruct::FMember> OutMembers = zzGetMembersBefore(zzMemberId##MemberName()); \
 			/* Add this member. */ \
 			OutMembers.Add(FUniformBufferStruct::FMember( \
 				TEXT(#MemberName), \
+				OptionalShaderType, \
 				STRUCT_OFFSET(zzTThisStruct,MemberName), \
 				(EUniformBufferBaseType)TUniformBufferTypeInfo<zzA##MemberName>::BaseType, \
 				Precision, \
@@ -501,11 +633,22 @@ private:
 		} \
 		typedef zzNextMemberId##MemberName
 
-#define DECLARE_UNIFORM_BUFFER_STRUCT_MEMBER_ARRAY(MemberType,MemberName,ArrayDecl) DECLARE_UNIFORM_BUFFER_STRUCT_MEMBER_ARRAY_EX(MemberType,MemberName,ArrayDecl,EShaderPrecisionModifier::Float)
+#define DECLARE_UNIFORM_BUFFER_STRUCT_MEMBER_ARRAY_EX(MemberType,MemberName,ArrayDecl,Precision) DECLARE_UNIFORM_BUFFER_STRUCT_MEMBER_EXPLICIT(MemberType,MemberName,ArrayDecl,Precision,TEXT(""))
+
+#define DECLARE_UNIFORM_BUFFER_STRUCT_MEMBER_ARRAY(MemberType,MemberName,ArrayDecl) DECLARE_UNIFORM_BUFFER_STRUCT_MEMBER_EXPLICIT(MemberType,MemberName,ArrayDecl,EShaderPrecisionModifier::Float,TEXT(""))
+
+#define DECLARE_UNIFORM_BUFFER_STRUCT_MEMBER(MemberType,MemberName) DECLARE_UNIFORM_BUFFER_STRUCT_MEMBER_EXPLICIT(MemberType,MemberName,,EShaderPrecisionModifier::Float,TEXT(""))
+
+#define DECLARE_UNIFORM_BUFFER_STRUCT_MEMBER_EX(MemberType,MemberName,Precision) DECLARE_UNIFORM_BUFFER_STRUCT_MEMBER_EXPLICIT(MemberType,MemberName,,Precision,TEXT(""))
+
+#define DECLARE_UNIFORM_BUFFER_STRUCT_MEMBER_SRV(ShaderType,MemberName) DECLARE_UNIFORM_BUFFER_STRUCT_MEMBER_EXPLICIT(FShaderResourceViewRHIParamRef,MemberName,,EShaderPrecisionModifier::Float,TEXT(#ShaderType))
+
+// NOT SUPPORTED YET
+//#define DECLARE_UNIFORM_BUFFER_STRUCT_MEMBER_UAV(ShaderType,MemberName) DECLARE_UNIFORM_BUFFER_STRUCT_MEMBER_EXPLICIT(FUnorderedAccessViewRHIParamRef,MemberName,,EShaderPrecisionModifier::Float,TEXT(#ShaderType))
 
-#define DECLARE_UNIFORM_BUFFER_STRUCT_MEMBER(MemberType,MemberName) DECLARE_UNIFORM_BUFFER_STRUCT_MEMBER_ARRAY_EX(MemberType,MemberName,,EShaderPrecisionModifier::Float)
+#define DECLARE_UNIFORM_BUFFER_STRUCT_MEMBER_SAMPLER(ShaderType,MemberName) DECLARE_UNIFORM_BUFFER_STRUCT_MEMBER_EXPLICIT(FSamplerStateRHIParamRef,MemberName,,EShaderPrecisionModifier::Float,TEXT(#ShaderType))
 
-#define DECLARE_UNIFORM_BUFFER_STRUCT_MEMBER_EX(MemberType,MemberName,Precision) DECLARE_UNIFORM_BUFFER_STRUCT_MEMBER_ARRAY_EX(MemberType,MemberName,,Precision)
+#define DECLARE_UNIFORM_BUFFER_STRUCT_MEMBER_TEXTURE(ShaderType,MemberName) DECLARE_UNIFORM_BUFFER_STRUCT_MEMBER_EXPLICIT(FTextureRHIParamRef,MemberName,,EShaderPrecisionModifier::Float,TEXT(#ShaderType))
 
 /** Ends a uniform buffer struct declaration. */
 #define END_UNIFORM_BUFFER_STRUCT(Name) \
diff --git a/Engine/Source/Runtime/Renderer/Private/AtmosphereRendering.cpp b/Engine/Source/Runtime/Renderer/Private/AtmosphereRendering.cpp
--- a/Engine/Source/Runtime/Renderer/Private/AtmosphereRendering.cpp
+++ b/Engine/Source/Runtime/Renderer/Private/AtmosphereRendering.cpp
@@ -62,20 +62,20 @@ public:
 	}
 
 	template<typename ShaderRHIParamRef>
-	void Set(const ShaderRHIParamRef ShaderRHI, uint32 TextureIdx, const FTextureRHIRef Texture) const
+	void Set(FRHICommandList* RHICmdList, const ShaderRHIParamRef ShaderRHI, uint32 TextureIdx, const FTextureRHIRef Texture) const
 	{
 		if (TextureIdx >= 4)
 		{
 			return;
 		}
 
-		SetTextureParameter(ShaderRHI, AtmosphereTexture[TextureIdx], AtmosphereTextureSampler[TextureIdx], 
+		SetTextureParameter(RHICmdList, ShaderRHI, AtmosphereTexture[TextureIdx], AtmosphereTextureSampler[TextureIdx], 
 			TStaticSamplerState<SF_Bilinear>::GetRHI(), 
 			Texture);
 	}
 
 	template<typename ShaderRHIParamRef>
-	void Set(const ShaderRHIParamRef ShaderRHI, uint32 TextureIdx, TexType TextureType, const FAtmosphereTextures* AtmosphereTextures) const
+	void Set(FRHICommandList* RHICmdList, const ShaderRHIParamRef ShaderRHI, uint32 TextureIdx, TexType TextureType, const FAtmosphereTextures* AtmosphereTextures) const
 	{
 		if (TextureIdx >= 4 || TextureType >= Type_MAX || !AtmosphereTextures)
 		{
@@ -85,37 +85,37 @@ public:
 		switch(TextureType)
 		{
 		case Transmittance:
-			SetTextureParameter(ShaderRHI, AtmosphereTexture[TextureIdx], AtmosphereTextureSampler[TextureIdx], 
+			SetTextureParameter(RHICmdList, ShaderRHI, AtmosphereTexture[TextureIdx], AtmosphereTextureSampler[TextureIdx], 
 				TStaticSamplerState<SF_Bilinear>::GetRHI(), 
 				AtmosphereTextures->AtmosphereTransmittance->GetRenderTargetItem().ShaderResourceTexture);
 			break;
 		case Irradiance:
-			SetTextureParameter(ShaderRHI, AtmosphereTexture[TextureIdx], AtmosphereTextureSampler[TextureIdx], 
+			SetTextureParameter(RHICmdList, ShaderRHI, AtmosphereTexture[TextureIdx], AtmosphereTextureSampler[TextureIdx], 
 				TStaticSamplerState<SF_Bilinear>::GetRHI(), 
 				AtmosphereTextures->AtmosphereIrradiance->GetRenderTargetItem().ShaderResourceTexture);
 			break;
 		case DeltaE:
-			SetTextureParameter(ShaderRHI, AtmosphereTexture[TextureIdx], AtmosphereTextureSampler[TextureIdx], 
+			SetTextureParameter(RHICmdList, ShaderRHI, AtmosphereTexture[TextureIdx], AtmosphereTextureSampler[TextureIdx], 
 				TStaticSamplerState<SF_Bilinear>::GetRHI(), 
 				AtmosphereTextures->AtmosphereDeltaE->GetRenderTargetItem().ShaderResourceTexture);
 			break;
 		case Inscatter:
-			SetTextureParameter(ShaderRHI, AtmosphereTexture[TextureIdx], AtmosphereTextureSampler[TextureIdx], 
+			SetTextureParameter(RHICmdList, ShaderRHI, AtmosphereTexture[TextureIdx], AtmosphereTextureSampler[TextureIdx], 
 				TStaticSamplerState<SF_Bilinear, AM_Clamp, AM_Clamp, AM_Clamp>::GetRHI(), 
 				AtmosphereTextures->AtmosphereInscatter->GetRenderTargetItem().ShaderResourceTexture);
 			break;
 		case DeltaSR:
-			SetTextureParameter(ShaderRHI, AtmosphereTexture[TextureIdx], AtmosphereTextureSampler[TextureIdx], 
+			SetTextureParameter(RHICmdList, ShaderRHI, AtmosphereTexture[TextureIdx], AtmosphereTextureSampler[TextureIdx], 
 				TStaticSamplerState<SF_Bilinear, AM_Clamp, AM_Clamp, AM_Clamp>::GetRHI(), 
 				AtmosphereTextures->AtmosphereDeltaSR->GetRenderTargetItem().ShaderResourceTexture);
 			break;
 		case DeltaSM:
-			SetTextureParameter(ShaderRHI, AtmosphereTexture[TextureIdx], AtmosphereTextureSampler[TextureIdx], 
+			SetTextureParameter(RHICmdList, ShaderRHI, AtmosphereTexture[TextureIdx], AtmosphereTextureSampler[TextureIdx], 
 				TStaticSamplerState<SF_Bilinear, AM_Clamp, AM_Clamp, AM_Clamp>::GetRHI(), 
 				AtmosphereTextures->AtmosphereDeltaSM->GetRenderTargetItem().ShaderResourceTexture);
 			break;
 		case DeltaJ:
-			SetTextureParameter(ShaderRHI, AtmosphereTexture[TextureIdx], AtmosphereTextureSampler[TextureIdx], 
+			SetTextureParameter(RHICmdList, ShaderRHI, AtmosphereTexture[TextureIdx], AtmosphereTextureSampler[TextureIdx], 
 				TStaticSamplerState<SF_Bilinear, AM_Clamp, AM_Clamp, AM_Clamp>::GetRHI(), 
 				AtmosphereTextures->AtmosphereDeltaJ->GetRenderTargetItem().ShaderResourceTexture);
 			break;
@@ -168,17 +168,18 @@ public:
 		OcclusionTextureSamplerParameter.Bind(Initializer.ParameterMap, TEXT("OcclusionTextureSampler"));
 	}
 
-	void SetParameters(const FSceneView& View, TRefCountPtr<IPooledRenderTarget>& LightShaftOcclusion)
+	void SetParameters(FRHICommandList* RHICmdList, const FSceneView& View, TRefCountPtr<IPooledRenderTarget>& LightShaftOcclusion)
 	{
-		FGlobalShader::SetParameters(GetPixelShader(), View);
-		SceneTextureParameters.Set(GetPixelShader(), View);
-		AtmosphereTextureParameters.Set(GetPixelShader(), View);
+		FGlobalShader::SetParameters(RHICmdList, GetPixelShader(), View);
+		SceneTextureParameters.Set(RHICmdList, GetPixelShader(), View);
+		AtmosphereTextureParameters.Set(RHICmdList, GetPixelShader(), View);
 
 		if (OcclusionTextureParameter.IsBound())
 		{
 			if (LightShaftOcclusion)
 			{
 				SetTextureParameter(
+					RHICmdList, 
 					GetPixelShader(),
 					OcclusionTextureParameter, OcclusionTextureSamplerParameter,
 					TStaticSamplerState<SF_Bilinear,AM_Clamp,AM_Clamp,AM_Clamp>::GetRHI(),
@@ -188,6 +189,7 @@ public:
 			else
 			{
 				SetTextureParameter(
+					RHICmdList, 
 					GetPixelShader(),
 					OcclusionTextureParameter, OcclusionTextureSamplerParameter,
 					TStaticSamplerState<SF_Bilinear,AM_Clamp,AM_Clamp,AM_Clamp>::GetRHI(),
@@ -290,9 +292,9 @@ public:
 	{
 	}
 
-	void SetParameters(const FViewInfo& View)
+	void SetParameters(FRHICommandList* RHICmdList, const FViewInfo& View)
 	{
-		FGlobalShader::SetParameters(GetVertexShader(),View);
+		FGlobalShader::SetParameters(RHICmdList, GetVertexShader(),View);
 	}
 
 	virtual bool Serialize(FArchive& Ar)
@@ -349,7 +351,7 @@ void FSceneRenderer::InitAtmosphereConstants()
 
 FGlobalBoundShaderState AtmosphereBoundShaderState[EAtmosphereRenderFlag::E_RenderFlagMax];
 
-void SetAtmosphericFogShaders(FScene* Scene, const FViewInfo& View, TRefCountPtr<IPooledRenderTarget>& LightShaftOcclusion)
+void SetAtmosphericFogShaders(FRHICommandList* RHICmdList, FScene* Scene, const FViewInfo& View, TRefCountPtr<IPooledRenderTarget>& LightShaftOcclusion)
 {
 	TShaderMapRef<FAtmosphericVS> VertexShader(GetGlobalShaderMap());
 	FAtmosphericFogPS* PixelShader = NULL;
@@ -383,8 +385,8 @@ void SetAtmosphericFogShaders(FScene* Scene, const FViewInfo& View, TRefCountPtr
 	}
 	
 	SetGlobalBoundShaderState(AtmosphereBoundShaderState[Scene->AtmosphericFog->RenderFlag], GAtmophereVertexDeclaration.VertexDeclarationRHI, *VertexShader, PixelShader);
-	VertexShader->SetParameters(View);
-	PixelShader->SetParameters(View, LightShaftOcclusion);
+	VertexShader->SetParameters(RHICmdList, View);
+	PixelShader->SetParameters(RHICmdList, View, LightShaftOcclusion);
 }
 
 void FDeferredShadingSceneRenderer::RenderAtmosphere(FLightShaftsOutput LightShaftsOutput)
@@ -407,6 +409,9 @@ void FDeferredShadingSceneRenderer::RenderAtmosphere(FLightShaftsOutput LightSha
 			0, 2, 3
 		};
 
+		//@todo-rco: RHIPacketList
+		FRHICommandList* RHICmdList = nullptr;
+
 		GSceneRenderTargets.BeginRenderingSceneColor();
 		for(int32 ViewIndex = 0;ViewIndex < Views.Num();ViewIndex++)
 		{
@@ -422,7 +427,7 @@ void FDeferredShadingSceneRenderer::RenderAtmosphere(FLightShaftsOutput LightSha
 
 			RHISetDepthStencilState(TStaticDepthStencilState<false,CF_Always>::GetRHI());
 
-			SetAtmosphericFogShaders(Scene, View, LightShaftsOutput.LightShaftOcclusion);
+			SetAtmosphericFogShaders(RHICmdList, Scene, View, LightShaftsOutput.LightShaftOcclusion);
 
 			// Draw a quad covering the view.
 			RHIDrawIndexedPrimitiveUP(
@@ -469,9 +474,9 @@ public:
 		return bShaderHasOutdatedParameters;
 	}
 
-	void SetParameters(const FViewInfo& View)
+	void SetParameters(FRHICommandList* RHICmdList, const FViewInfo& View)
 	{
-		FGlobalShader::SetParameters(GetPixelShader(), View);
+		FGlobalShader::SetParameters(RHICmdList, GetPixelShader(), View);
 	}
 };
 
@@ -505,9 +510,9 @@ public:
 		return bShaderHasOutdatedParameters;
 	}
 
-	void SetParameters(const FAtmosphereTextures* Textures)
+	void SetParameters(FRHICommandList* RHICmdList, const FAtmosphereTextures* Textures)
 	{
-		AtmosphereParameters.Set(GetPixelShader(), 0, FAtmosphereShaderPrecomputeTextureParameters::Transmittance, Textures);
+		AtmosphereParameters.Set(RHICmdList, GetPixelShader(), 0, FAtmosphereShaderPrecomputeTextureParameters::Transmittance, Textures);
 	}
 };
 
@@ -545,13 +550,13 @@ public:
 		return bShaderHasOutdatedParameters;
 	}
 
-	void SetParameters(const FViewInfo& View, float FirstOrder, const FAtmosphereTextures* Textures)
+	void SetParameters(FRHICommandList* RHICmdList, const FViewInfo& View, float FirstOrder, const FAtmosphereTextures* Textures)
 	{
-		FGlobalShader::SetParameters(GetPixelShader(), View);
-		AtmosphereParameters.Set(GetPixelShader(), 0, FAtmosphereShaderPrecomputeTextureParameters::Transmittance, Textures);
-		AtmosphereParameters.Set(GetPixelShader(), 1, FAtmosphereShaderPrecomputeTextureParameters::DeltaSR, Textures);
-		AtmosphereParameters.Set(GetPixelShader(), 2, FAtmosphereShaderPrecomputeTextureParameters::DeltaSM, Textures);
-		SetShaderValue(GetPixelShader(), FirstOrderParameter, FirstOrder);
+		FGlobalShader::SetParameters(RHICmdList, GetPixelShader(), View);
+		AtmosphereParameters.Set(RHICmdList, GetPixelShader(), 0, FAtmosphereShaderPrecomputeTextureParameters::Transmittance, Textures);
+		AtmosphereParameters.Set(RHICmdList, GetPixelShader(), 1, FAtmosphereShaderPrecomputeTextureParameters::DeltaSR, Textures);
+		AtmosphereParameters.Set(RHICmdList, GetPixelShader(), 2, FAtmosphereShaderPrecomputeTextureParameters::DeltaSM, Textures);
+		SetShaderValue(RHICmdList, GetPixelShader(), FirstOrderParameter, FirstOrder);
 	}
 };
 
@@ -585,9 +590,9 @@ public:
 		return bShaderHasOutdatedParameters;
 	}
 
-	void SetParameters(const FAtmosphereTextures* Textures)
+	void SetParameters(FRHICommandList* RHICmdList, const FAtmosphereTextures* Textures)
 	{
-		AtmosphereParameters.Set(GetPixelShader(), 0, FAtmosphereShaderPrecomputeTextureParameters::DeltaE, Textures);
+		AtmosphereParameters.Set(RHICmdList, GetPixelShader(), 0, FAtmosphereShaderPrecomputeTextureParameters::DeltaE, Textures);
 	}
 };
 
@@ -611,9 +616,9 @@ public:
 		AtmosphereLayerParameter.Bind(Initializer.ParameterMap, TEXT("AtmosphereLayer"));
 	}
 
-	void SetParameters(int AtmosphereLayer)
+	void SetParameters(FRHICommandList* RHICmdList, int AtmosphereLayer)
 	{
-		SetShaderValue(GetGeometryShader(), AtmosphereLayerParameter, AtmosphereLayer);
+		SetShaderValue(RHICmdList, GetGeometryShader(), AtmosphereLayerParameter, AtmosphereLayer);
 	}
 
 	virtual bool Serialize(FArchive& Ar)
@@ -658,12 +663,12 @@ public:
 		return bShaderHasOutdatedParameters;
 	}
 
-	void SetParameters(const FViewInfo& View, float AtmosphereR, const FVector4& DhdH, const FAtmosphereTextures* Textures)
+	void SetParameters(FRHICommandList* RHICmdList, const FViewInfo& View, float AtmosphereR, const FVector4& DhdH, const FAtmosphereTextures* Textures)
 	{
-		FGlobalShader::SetParameters(GetPixelShader(), View);
-		AtmosphereParameters.Set(GetPixelShader(), 0, FAtmosphereShaderPrecomputeTextureParameters::Transmittance, Textures);
-		SetShaderValue(GetPixelShader(), dhdHParameter, DhdH);
-		SetShaderValue(GetPixelShader(), AtmosphereRParameter, AtmosphereR);
+		FGlobalShader::SetParameters(RHICmdList, GetPixelShader(), View);
+		AtmosphereParameters.Set(RHICmdList, GetPixelShader(), 0, FAtmosphereShaderPrecomputeTextureParameters::Transmittance, Textures);
+		SetShaderValue(RHICmdList, GetPixelShader(), dhdHParameter, DhdH);
+		SetShaderValue(RHICmdList, GetPixelShader(), AtmosphereRParameter, AtmosphereR);
 	}
 };
 
@@ -704,14 +709,14 @@ public:
 		return bShaderHasOutdatedParameters;
 	}
 
-	void SetParameters(const FViewInfo& View, float AtmosphereR, const FVector4& DhdH, int AtmosphereLayer, const FAtmosphereTextures* Textures)
+	void SetParameters(FRHICommandList* RHICmdList, const FViewInfo& View, float AtmosphereR, const FVector4& DhdH, int AtmosphereLayer, const FAtmosphereTextures* Textures)
 	{
-		FGlobalShader::SetParameters(GetPixelShader(), View);
-		AtmosphereParameters.Set(GetPixelShader(), 0, FAtmosphereShaderPrecomputeTextureParameters::DeltaSR, Textures);
-		AtmosphereParameters.Set(GetPixelShader(), 1, FAtmosphereShaderPrecomputeTextureParameters::DeltaSM, Textures);
-		SetShaderValue(GetPixelShader(), AtmosphereRParameter, AtmosphereR);
-		SetShaderValue(GetPixelShader(), dhdHParameter, DhdH);
-		SetShaderValue(GetPixelShader(), AtmosphereLayerParameter, AtmosphereLayer);
+		FGlobalShader::SetParameters(RHICmdList, GetPixelShader(), View);
+		AtmosphereParameters.Set(RHICmdList, GetPixelShader(), 0, FAtmosphereShaderPrecomputeTextureParameters::DeltaSR, Textures);
+		AtmosphereParameters.Set(RHICmdList, GetPixelShader(), 1, FAtmosphereShaderPrecomputeTextureParameters::DeltaSM, Textures);
+		SetShaderValue(RHICmdList, GetPixelShader(), AtmosphereRParameter, AtmosphereR);
+		SetShaderValue(RHICmdList, GetPixelShader(), dhdHParameter, DhdH);
+		SetShaderValue(RHICmdList, GetPixelShader(), AtmosphereLayerParameter, AtmosphereLayer);
 	}
 };
 
@@ -751,13 +756,13 @@ public:
 		return bShaderHasOutdatedParameters;
 	}
 
-	void SetParameters(const FViewInfo& View, float AtmosphereR, const FVector4& DhdH, int AtmosphereLayer, const FAtmosphereTextures* Textures)
+	void SetParameters(FRHICommandList* RHICmdList, const FViewInfo& View, float AtmosphereR, const FVector4& DhdH, int AtmosphereLayer, const FAtmosphereTextures* Textures)
 	{
-		FGlobalShader::SetParameters(GetPixelShader(), View);
-		AtmosphereParameters.Set(GetPixelShader(), 0, FAtmosphereShaderPrecomputeTextureParameters::DeltaSR, Textures);
-		SetShaderValue(GetPixelShader(), AtmosphereRParameter, AtmosphereR);
-		SetShaderValue(GetPixelShader(), dhdHParameter, DhdH);
-		SetShaderValue(GetPixelShader(), AtmosphereLayerParameter, AtmosphereLayer);
+		FGlobalShader::SetParameters(RHICmdList, GetPixelShader(), View);
+		AtmosphereParameters.Set(RHICmdList, GetPixelShader(), 0, FAtmosphereShaderPrecomputeTextureParameters::DeltaSR, Textures);
+		SetShaderValue(RHICmdList, GetPixelShader(), AtmosphereRParameter, AtmosphereR);
+		SetShaderValue(RHICmdList, GetPixelShader(), dhdHParameter, DhdH);
+		SetShaderValue(RHICmdList, GetPixelShader(), AtmosphereLayerParameter, AtmosphereLayer);
 	}
 };
 
@@ -800,16 +805,16 @@ public:
 		return bShaderHasOutdatedParameters;
 	}
 
-	void SetParameters(const FViewInfo& View, float AtmosphereR, const FVector4& DhdH, float FirstOrder, const FAtmosphereTextures* Textures)
+	void SetParameters(FRHICommandList* RHICmdList, const FViewInfo& View, float AtmosphereR, const FVector4& DhdH, float FirstOrder, const FAtmosphereTextures* Textures)
 	{
-		FGlobalShader::SetParameters(GetPixelShader(), View);
-		AtmosphereParameters.Set(GetPixelShader(), 0, FAtmosphereShaderPrecomputeTextureParameters::Transmittance, Textures);
-		AtmosphereParameters.Set(GetPixelShader(), 1, FAtmosphereShaderPrecomputeTextureParameters::DeltaE, Textures);
-		AtmosphereParameters.Set(GetPixelShader(), 2, FAtmosphereShaderPrecomputeTextureParameters::DeltaSR, Textures);
-		AtmosphereParameters.Set(GetPixelShader(), 3, FAtmosphereShaderPrecomputeTextureParameters::DeltaSM, Textures);
-		SetShaderValue(GetPixelShader(), AtmosphereRParameter, AtmosphereR);
-		SetShaderValue(GetPixelShader(), dhdHParameter, DhdH);
-		SetShaderValue(GetPixelShader(), FirstOrderParameter, FirstOrder);
+		FGlobalShader::SetParameters(RHICmdList, GetPixelShader(), View);
+		AtmosphereParameters.Set(RHICmdList, GetPixelShader(), 0, FAtmosphereShaderPrecomputeTextureParameters::Transmittance, Textures);
+		AtmosphereParameters.Set(RHICmdList, GetPixelShader(), 1, FAtmosphereShaderPrecomputeTextureParameters::DeltaE, Textures);
+		AtmosphereParameters.Set(RHICmdList, GetPixelShader(), 2, FAtmosphereShaderPrecomputeTextureParameters::DeltaSR, Textures);
+		AtmosphereParameters.Set(RHICmdList, GetPixelShader(), 3, FAtmosphereShaderPrecomputeTextureParameters::DeltaSM, Textures);
+		SetShaderValue(RHICmdList, GetPixelShader(), AtmosphereRParameter, AtmosphereR);
+		SetShaderValue(RHICmdList, GetPixelShader(), dhdHParameter, DhdH);
+		SetShaderValue(RHICmdList, GetPixelShader(), FirstOrderParameter, FirstOrder);
 	}
 };
 
@@ -850,14 +855,14 @@ public:
 		return bShaderHasOutdatedParameters;
 	}
 
-	void SetParameters(const FViewInfo& View, float AtmosphereR, const FVector4& DhdH, float FirstOrder, const FAtmosphereTextures* Textures)
+	void SetParameters(FRHICommandList* RHICmdList, const FViewInfo& View, float AtmosphereR, const FVector4& DhdH, float FirstOrder, const FAtmosphereTextures* Textures)
 	{
-		FGlobalShader::SetParameters(GetPixelShader(), View);
-		AtmosphereParameters.Set(GetPixelShader(), 0, FAtmosphereShaderPrecomputeTextureParameters::Transmittance, Textures);
-		AtmosphereParameters.Set(GetPixelShader(), 1, FAtmosphereShaderPrecomputeTextureParameters::DeltaJ, Textures);
-		SetShaderValue(GetPixelShader(), AtmosphereRParameter, AtmosphereR);
-		SetShaderValue(GetPixelShader(), dhdHParameter, DhdH);
-		SetShaderValue(GetPixelShader(), FirstOrderParameter, FirstOrder);
+		FGlobalShader::SetParameters(RHICmdList, GetPixelShader(), View);
+		AtmosphereParameters.Set(RHICmdList, GetPixelShader(), 0, FAtmosphereShaderPrecomputeTextureParameters::Transmittance, Textures);
+		AtmosphereParameters.Set(RHICmdList, GetPixelShader(), 1, FAtmosphereShaderPrecomputeTextureParameters::DeltaJ, Textures);
+		SetShaderValue(RHICmdList, GetPixelShader(), AtmosphereRParameter, AtmosphereR);
+		SetShaderValue(RHICmdList, GetPixelShader(), dhdHParameter, DhdH);
+		SetShaderValue(RHICmdList, GetPixelShader(), FirstOrderParameter, FirstOrder);
 	}
 };
 
@@ -925,12 +930,12 @@ public:
 		return bShaderHasOutdatedParameters;
 	}
 
-	void SetParameters(float AtmosphereR, const FVector4& DhdH, int AtmosphereLayer, const FAtmosphereTextures* Textures)
+	void SetParameters(FRHICommandList* RHICmdList, float AtmosphereR, const FVector4& DhdH, int AtmosphereLayer, const FAtmosphereTextures* Textures)
 	{
-		AtmosphereParameters.Set(GetPixelShader(), 0, FAtmosphereShaderPrecomputeTextureParameters::Inscatter, Textures);
-		SetShaderValue(GetPixelShader(), AtmosphereRParameter, AtmosphereR);
-		SetShaderValue(GetPixelShader(), dhdHParameter, DhdH);
-		SetShaderValue(GetPixelShader(), AtmosphereLayerParameter, AtmosphereLayer);
+		AtmosphereParameters.Set(RHICmdList, GetPixelShader(), 0, FAtmosphereShaderPrecomputeTextureParameters::Inscatter, Textures);
+		SetShaderValue(RHICmdList, GetPixelShader(), AtmosphereRParameter, AtmosphereR);
+		SetShaderValue(RHICmdList, GetPixelShader(), dhdHParameter, DhdH);
+		SetShaderValue(RHICmdList, GetPixelShader(), AtmosphereLayerParameter, AtmosphereLayer);
 	}
 };
 
@@ -970,12 +975,12 @@ public:
 		return bShaderHasOutdatedParameters;
 	}
 
-	void SetParameters(float AtmosphereR, const FVector4& DhdH, int AtmosphereLayer, const FAtmosphereTextures* Textures)
+	void SetParameters(FRHICommandList* RHICmdList, float AtmosphereR, const FVector4& DhdH, int AtmosphereLayer, const FAtmosphereTextures* Textures)
 	{
-		AtmosphereParameters.Set(GetPixelShader(), 0, FAtmosphereShaderPrecomputeTextureParameters::DeltaSR, Textures);
-		SetShaderValue(GetPixelShader(), AtmosphereRParameter, AtmosphereR);
-		SetShaderValue(GetPixelShader(), dhdHParameter, DhdH);
-		SetShaderValue(GetPixelShader(), AtmosphereLayerParameter, AtmosphereLayer);
+		AtmosphereParameters.Set(RHICmdList, GetPixelShader(), 0, FAtmosphereShaderPrecomputeTextureParameters::DeltaSR, Textures);
+		SetShaderValue(RHICmdList, GetPixelShader(), AtmosphereRParameter, AtmosphereR);
+		SetShaderValue(RHICmdList, GetPixelShader(), dhdHParameter, DhdH);
+		SetShaderValue(RHICmdList, GetPixelShader(), AtmosphereLayerParameter, AtmosphereLayer);
 	}
 };
 
@@ -1078,6 +1083,9 @@ void FAtmosphericFogSceneInfo::GetLayerValue(int Layer, float& AtmosphereR, FVec
 
 void FAtmosphericFogSceneInfo::RenderAtmosphereShaders(const FViewInfo& View, const FIntRect& ViewRect)
 {
+	//@todo-rco: RHIPacketList
+	FRHICommandList* RHICmdList = nullptr;
+
 	check(Component != NULL);
 	switch(AtmospherePhase)
 	{
@@ -1091,7 +1099,7 @@ void FAtmosphericFogSceneInfo::RenderAtmosphereShaders(const FViewInfo& View, co
 
 			static FGlobalBoundShaderState BoundShaderState;
 			SetGlobalBoundShaderState( BoundShaderState, GFilterVertexDeclaration.VertexDeclarationRHI, *VertexShader, *PixelShader);
-			PixelShader->SetParameters(View);
+			PixelShader->SetParameters(RHICmdList, View);
 			DrawQuad(ViewRect, *VertexShader);
 
 			RHICopyToResolveTarget(DestRenderTarget.TargetableTexture, DestRenderTarget.ShaderResourceTexture, true, FResolveParams() );
@@ -1108,7 +1116,7 @@ void FAtmosphericFogSceneInfo::RenderAtmosphereShaders(const FViewInfo& View, co
 			static FGlobalBoundShaderState BoundShaderState;
 			SetGlobalBoundShaderState( BoundShaderState, GFilterVertexDeclaration.VertexDeclarationRHI, *VertexShader, *PixelShader);
 
-			PixelShader->SetParameters(AtmosphereTextures);
+			PixelShader->SetParameters(RHICmdList, AtmosphereTextures);
 
 			DrawQuad(ViewRect, *VertexShader);
 
@@ -1135,8 +1143,8 @@ void FAtmosphericFogSceneInfo::RenderAtmosphereShaders(const FViewInfo& View, co
 				float r;
 				FVector4 DhdH;
 				GetLayerValue(Layer, r, DhdH);
-				GeometryShader->SetParameters(Layer);
-				PixelShader->SetParameters(View, r, DhdH, AtmosphereTextures);
+				GeometryShader->SetParameters(RHICmdList, Layer);
+				PixelShader->SetParameters(RHICmdList, View, r, DhdH, AtmosphereTextures);
 				DrawQuad(ViewRect, *VertexShader);
 
 				if (Atmosphere3DTextureIndex == Component->PrecomputeParams.InscatterAltitudeSampleNum - 1)
@@ -1175,8 +1183,8 @@ void FAtmosphericFogSceneInfo::RenderAtmosphereShaders(const FViewInfo& View, co
 				float r;
 				FVector4 DhdH;
 				GetLayerValue(Layer, r, DhdH);
-				GeometryShader->SetParameters(Layer);
-				PixelShader->SetParameters(View, r, DhdH, Layer, AtmosphereTextures);
+				GeometryShader->SetParameters(RHICmdList, Layer);
+				PixelShader->SetParameters(RHICmdList, View, r, DhdH, Layer, AtmosphereTextures);
 				DrawQuad(ViewRect, *VertexShader);
 
 				if (Atmosphere3DTextureIndex == Component->PrecomputeParams.InscatterAltitudeSampleNum - 1)
@@ -1203,8 +1211,8 @@ void FAtmosphericFogSceneInfo::RenderAtmosphereShaders(const FViewInfo& View, co
 				float r;
 				FVector4 DhdH;
 				GetLayerValue(Layer, r, DhdH);
-				GeometryShader->SetParameters(Layer);
-				PixelShader->SetParameters(View, r, DhdH, AtmoshpereOrder == 2 ? 1.f : 0.f, AtmosphereTextures);
+				GeometryShader->SetParameters(RHICmdList, Layer);
+				PixelShader->SetParameters(RHICmdList, View, r, DhdH, AtmoshpereOrder == 2 ? 1.f : 0.f, AtmosphereTextures);
 				DrawQuad(ViewRect, *VertexShader);
 
 				if (Atmosphere3DTextureIndex == Component->PrecomputeParams.InscatterAltitudeSampleNum - 1)
@@ -1225,7 +1233,7 @@ void FAtmosphericFogSceneInfo::RenderAtmosphereShaders(const FViewInfo& View, co
 			static FGlobalBoundShaderState BoundShaderState;
 			SetGlobalBoundShaderState( BoundShaderState, GFilterVertexDeclaration.VertexDeclarationRHI, *VertexShader, *PixelShader);
 
-			PixelShader->SetParameters(View, AtmoshpereOrder == 2 ? 1.f : 0.f, AtmosphereTextures);
+			PixelShader->SetParameters(RHICmdList, View, AtmoshpereOrder == 2 ? 1.f : 0.f, AtmosphereTextures);
 
 			DrawQuad(ViewRect, *VertexShader);
 
@@ -1250,8 +1258,8 @@ void FAtmosphericFogSceneInfo::RenderAtmosphereShaders(const FViewInfo& View, co
 				float r;
 				FVector4 DhdH;
 				GetLayerValue(Layer, r, DhdH);
-				GeometryShader->SetParameters(Layer);
-				PixelShader->SetParameters(View, r, DhdH, AtmoshpereOrder == 2 ? 1.f : 0.f, AtmosphereTextures);
+				GeometryShader->SetParameters(RHICmdList, Layer);
+				PixelShader->SetParameters(RHICmdList, View, r, DhdH, AtmoshpereOrder == 2 ? 1.f : 0.f, AtmosphereTextures);
 				DrawQuad(ViewRect, *VertexShader);
 
 				if (Atmosphere3DTextureIndex == Component->PrecomputeParams.InscatterAltitudeSampleNum - 1)
@@ -1274,7 +1282,7 @@ void FAtmosphericFogSceneInfo::RenderAtmosphereShaders(const FViewInfo& View, co
 			static FGlobalBoundShaderState BoundShaderState;
 			SetGlobalBoundShaderState( BoundShaderState, GFilterVertexDeclaration.VertexDeclarationRHI, *VertexShader, *PixelShader);
 
-			PixelShader->SetParameters(AtmosphereTextures);
+			PixelShader->SetParameters(RHICmdList, AtmosphereTextures);
 
 			DrawQuad(ViewRect, *VertexShader);
 
@@ -1303,8 +1311,8 @@ void FAtmosphericFogSceneInfo::RenderAtmosphereShaders(const FViewInfo& View, co
 				float r;
 				FVector4 DhdH;
 				GetLayerValue(Layer, r, DhdH);
-				GeometryShader->SetParameters(Layer);
-				PixelShader->SetParameters(View, r, DhdH, Layer, AtmosphereTextures);
+				GeometryShader->SetParameters(RHICmdList, Layer);
+				PixelShader->SetParameters(RHICmdList, View, r, DhdH, Layer, AtmosphereTextures);
 				DrawQuad(ViewRect, *VertexShader);
 
 				if (Atmosphere3DTextureIndex == Component->PrecomputeParams.InscatterAltitudeSampleNum - 1)
@@ -1334,8 +1342,8 @@ void FAtmosphericFogSceneInfo::RenderAtmosphereShaders(const FViewInfo& View, co
 				float r;
 				FVector4 DhdH;
 				GetLayerValue(Layer, r, DhdH);
-				GeometryShader->SetParameters(Layer);
-				PixelShader->SetParameters(r, DhdH, Layer, AtmosphereTextures);
+				GeometryShader->SetParameters(RHICmdList, Layer);
+				PixelShader->SetParameters(RHICmdList, r, DhdH, Layer, AtmosphereTextures);
 				DrawQuad(ViewRect, *VertexShader);
 
 				if (Atmosphere3DTextureIndex == Component->PrecomputeParams.InscatterAltitudeSampleNum - 1)
@@ -1362,8 +1370,8 @@ void FAtmosphericFogSceneInfo::RenderAtmosphereShaders(const FViewInfo& View, co
 				float r;
 				FVector4 DhdH;
 				GetLayerValue(Layer, r, DhdH);
-				GeometryShader->SetParameters(Layer);
-				PixelShader->SetParameters(r, DhdH, Layer, AtmosphereTextures);
+				GeometryShader->SetParameters(RHICmdList, Layer);
+				PixelShader->SetParameters(RHICmdList, r, DhdH, Layer, AtmosphereTextures);
 				DrawQuad(ViewRect, *VertexShader);
 
 				if (Atmosphere3DTextureIndex == Component->PrecomputeParams.InscatterAltitudeSampleNum - 1)
diff --git a/Engine/Source/Runtime/Renderer/Private/AtmosphereRendering.h b/Engine/Source/Runtime/Renderer/Private/AtmosphereRendering.h
--- a/Engine/Source/Runtime/Renderer/Private/AtmosphereRendering.h
+++ b/Engine/Source/Runtime/Renderer/Private/AtmosphereRendering.h
@@ -95,15 +95,15 @@ public:
 	}
 
 	template<typename ShaderRHIParamRef>
-	void Set(const ShaderRHIParamRef ShaderRHI, const FSceneView& View) const
+	void Set(FRHICommandList* RHICmdList, const ShaderRHIParamRef ShaderRHI, const FSceneView& View) const
 	{
 		if (TransmittanceTexture.IsBound() || IrradianceTexture.IsBound() || InscatterTexture.IsBound())
 		{
-			SetTextureParameter(ShaderRHI, TransmittanceTexture, TransmittanceTextureSampler, 
+			SetTextureParameter(RHICmdList, ShaderRHI, TransmittanceTexture, TransmittanceTextureSampler, 
 				TStaticSamplerState<SF_Bilinear>::GetRHI(), View.AtmosphereTransmittanceTexture);
-			SetTextureParameter(ShaderRHI, IrradianceTexture, IrradianceTextureSampler, 
+			SetTextureParameter(RHICmdList, ShaderRHI, IrradianceTexture, IrradianceTextureSampler, 
 				TStaticSamplerState<SF_Bilinear>::GetRHI(), View.AtmosphereIrradianceTexture);
-			SetTextureParameter(ShaderRHI, InscatterTexture, InscatterTextureSampler, 
+			SetTextureParameter(RHICmdList, ShaderRHI, InscatterTexture, InscatterTextureSampler, 
 				TStaticSamplerState<SF_Bilinear>::GetRHI(), View.AtmosphereInscatterTexture);
 		}
 	}
diff --git a/Engine/Source/Runtime/Renderer/Private/BasePassRendering.cpp b/Engine/Source/Runtime/Renderer/Private/BasePassRendering.cpp
--- a/Engine/Source/Runtime/Renderer/Private/BasePassRendering.cpp
+++ b/Engine/Source/Runtime/Renderer/Private/BasePassRendering.cpp
@@ -70,9 +70,10 @@ void FSkyLightReflectionParameters::GetSkyParametersFromScene(const FScene* Scen
 	}
 }
 
-void FTranslucentLightingParameters::Set(FShader* Shader, const FSceneView* View)
+void FTranslucentLightingParameters::Set(FRHICommandList* RHICmdList, FShader* Shader, const FSceneView* View)
 {
 	SetTextureParameter(
+		RHICmdList, 
 		Shader->GetPixelShader(), 
 		TranslucencyLightingVolumeAmbientInner, 
 		TranslucencyLightingVolumeAmbientInnerSampler, 
@@ -80,6 +81,7 @@ void FTranslucentLightingParameters::Set(FShader* Shader, const FSceneView* View
 		GSceneRenderTargets.GetTranslucencyVolumeAmbient(TVC_Inner)->GetRenderTargetItem().ShaderResourceTexture);
 
 	SetTextureParameter(
+		RHICmdList, 
 		Shader->GetPixelShader(), 
 		TranslucencyLightingVolumeAmbientOuter, 
 		TranslucencyLightingVolumeAmbientOuterSampler, 
@@ -87,6 +89,7 @@ void FTranslucentLightingParameters::Set(FShader* Shader, const FSceneView* View
 		GSceneRenderTargets.GetTranslucencyVolumeAmbient(TVC_Outer)->GetRenderTargetItem().ShaderResourceTexture);
 
 	SetTextureParameter(
+		RHICmdList, 
 		Shader->GetPixelShader(), 
 		TranslucencyLightingVolumeDirectionalInner, 
 		TranslucencyLightingVolumeDirectionalInnerSampler, 
@@ -94,16 +97,17 @@ void FTranslucentLightingParameters::Set(FShader* Shader, const FSceneView* View
 		GSceneRenderTargets.GetTranslucencyVolumeDirectional(TVC_Inner)->GetRenderTargetItem().ShaderResourceTexture);
 
 	SetTextureParameter(
+		RHICmdList, 
 		Shader->GetPixelShader(), 
 		TranslucencyLightingVolumeDirectionalOuter, 
 		TranslucencyLightingVolumeDirectionalOuterSampler, 
 		TStaticSamplerState<SF_Bilinear, AM_Clamp, AM_Clamp, AM_Clamp>::GetRHI(), 
 		GSceneRenderTargets.GetTranslucencyVolumeDirectional(TVC_Outer)->GetRenderTargetItem().ShaderResourceTexture);
 
-	SkyLightReflectionParameters.SetParameters(Shader->GetPixelShader(), (const FScene*)(View->Family->Scene), true);
+	SkyLightReflectionParameters.SetParameters(RHICmdList, Shader->GetPixelShader(), (const FScene*)(View->Family->Scene), true);
 }
 
-void FTranslucentLightingParameters::SetMesh(FShader* Shader, const FPrimitiveSceneProxy* Proxy)
+void FTranslucentLightingParameters::SetMesh(FRHICommandList* RHICmdList, FShader* Shader, const FPrimitiveSceneProxy* Proxy)
 {
 	// Note: GBlackCubeArrayTexture has an alpha of 0, which is needed to represent invalid data so the sky cubemap can still be applied
 	FTextureRHIParamRef CubeArrayTexture = GBlackCubeArrayTexture->TextureRHI;
@@ -116,13 +120,14 @@ void FTranslucentLightingParameters::SetMesh(FShader* Shader, const FPrimitiveSc
 	}
 
 	SetTextureParameter(
+		RHICmdList, 
 		Shader->GetPixelShader(), 
 		ReflectionCubemap, 
 		ReflectionCubemapSampler, 
 		TStaticSamplerState<SF_Trilinear, AM_Clamp, AM_Clamp, AM_Clamp>::GetRHI(), 
 		CubeArrayTexture);
 
-	SetShaderValue(Shader->GetPixelShader(), CubemapArrayIndex, ArrayIndex);
+	SetShaderValue(RHICmdList, Shader->GetPixelShader(), CubemapArrayIndex, ArrayIndex);
 }
 
 /** The action used to draw a base pass static mesh element. */
@@ -155,6 +160,7 @@ public:
 	/** Draws the translucent mesh with a specific light-map type, and fog volume type */
 	template<typename LightMapPolicyType>
 	void Process(
+		FRHICommandList* RHICmdList, 
 		const FProcessBasePassMeshParameters& Parameters,
 		const LightMapPolicyType& LightMapPolicy,
 		const typename LightMapPolicyType::ElementDataType& LightMapElementData
@@ -202,7 +208,10 @@ void FBasePassOpaqueDrawingPolicyFactory::AddStaticMesh(FScene* Scene,FStaticMes
 	// Only draw opaque materials.
 	if( !IsTranslucentBlendMode(BlendMode) )
 	{
+		//@todo-rco: RHIPacketList
+		FRHICommandList* RHICmdList = nullptr;
 		ProcessBasePassMesh(
+			RHICmdList, 
 			FProcessBasePassMeshParameters(
 				*StaticMesh,
 				Material,
@@ -253,6 +262,7 @@ public:
 	/** Draws the translucent mesh with a specific light-map type, and shader complexity predicate. */
 	template<typename LightMapPolicyType>
 	void Process(
+		FRHICommandList* RHICmdList, 
 		const FProcessBasePassMeshParameters& Parameters,
 		const LightMapPolicyType& LightMapPolicy,
 		const typename LightMapPolicyType::ElementDataType& LightMapElementData
@@ -288,6 +298,7 @@ public:
 			Parameters.bEditorCompositeDepthTest
 			);
 		DrawingPolicy.DrawShared(
+			RHICmdList, 
 			&View,
 			DrawingPolicy.CreateBoundShaderState(View.GetFeatureLevel())
 			);
@@ -295,6 +306,7 @@ public:
 		for( int32 BatchElementIndex = 0, Num = Parameters.Mesh.Elements.Num(); BatchElementIndex < Num; BatchElementIndex++ )
 		{
 			DrawingPolicy.SetMeshRenderState(
+				RHICmdList, 
 				View,
 				Parameters.PrimitiveSceneProxy,
 				Parameters.Mesh,
@@ -302,7 +314,7 @@ public:
 				bBackFace,
 				typename TBasePassDrawingPolicy<LightMapPolicyType>::ElementDataType(LightMapElementData)
 				);
-			DrawingPolicy.DrawMesh(Parameters.Mesh, BatchElementIndex);
+			DrawingPolicy.DrawMesh(RHICmdList, Parameters.Mesh, BatchElementIndex);
 		}
 
 #if !(UE_BUILD_SHIPPING || UE_BUILD_TEST)
@@ -316,6 +328,7 @@ public:
 };
 
 bool FBasePassOpaqueDrawingPolicyFactory::DrawDynamicMesh(
+	FRHICommandList* RHICmdList, 
 	const FSceneView& View,
 	ContextType DrawingContext,
 	const FMeshBatch& Mesh,
@@ -333,6 +346,7 @@ bool FBasePassOpaqueDrawingPolicyFactory::DrawDynamicMesh(
 	if(!IsTranslucentBlendMode(BlendMode))
 	{
 		ProcessBasePassMesh(
+			RHICmdList, 
 			FProcessBasePassMeshParameters(
 				Mesh,
 				Material,
@@ -357,6 +371,7 @@ bool FBasePassOpaqueDrawingPolicyFactory::DrawDynamicMesh(
 }
 
 void FCachedVolumeIndirectLightingPolicy::Set(
+	FRHICommandList* RHICmdList, 
 	const VertexParametersType* VertexShaderParameters,
 	const PixelParametersType* PixelShaderParameters,
 	FShader* VertexShader,
@@ -366,7 +381,7 @@ void FCachedVolumeIndirectLightingPolicy::Set(
 	const FSceneView* View
 	) const
 {
-	FNoLightMapPolicy::Set(VertexShaderParameters, PixelShaderParameters, VertexShader, PixelShader, VertexFactory, MaterialRenderProxy, View);
+	FNoLightMapPolicy::Set(RHICmdList, VertexShaderParameters, PixelShaderParameters, VertexShader, PixelShader, VertexFactory, MaterialRenderProxy, View);
 
 	if (PixelShaderParameters)
 	{
@@ -390,6 +405,7 @@ void FCachedVolumeIndirectLightingPolicy::Set(
 		const FPixelShaderRHIParamRef PixelShaderRHI = PixelShader->GetPixelShader();
 
 		SetTextureParameter(
+			RHICmdList, 
 			PixelShaderRHI,
 			PixelShaderParameters->IndirectLightingCacheTexture0,
 			PixelShaderParameters->IndirectLightingCacheSampler0,
@@ -398,6 +414,7 @@ void FCachedVolumeIndirectLightingPolicy::Set(
 			);
 
 		SetTextureParameter(
+			RHICmdList, 
 			PixelShaderRHI,
 			PixelShaderParameters->IndirectLightingCacheTexture1,
 			PixelShaderParameters->IndirectLightingCacheSampler1,
@@ -406,6 +423,7 @@ void FCachedVolumeIndirectLightingPolicy::Set(
 			);
 
 		SetTextureParameter(
+			RHICmdList, 
 			PixelShaderRHI,
 			PixelShaderParameters->IndirectLightingCacheTexture2,
 			PixelShaderParameters->IndirectLightingCacheSampler2,
@@ -416,6 +434,7 @@ void FCachedVolumeIndirectLightingPolicy::Set(
 }
 
 void FCachedVolumeIndirectLightingPolicy::SetMesh(
+	FRHICommandList* RHICmdList, 
 	const FSceneView& View,
 	const FPrimitiveSceneProxy* PrimitiveSceneProxy,
 	const VertexParametersType* VertexShaderParameters,
@@ -448,15 +467,16 @@ void FCachedVolumeIndirectLightingPolicy::SetMesh(
 		}
 		
 		const FPixelShaderRHIParamRef PixelShaderRHI = PixelShader->GetPixelShader();
-		SetShaderValue(PixelShaderRHI, PixelShaderParameters->IndirectlightingCachePrimitiveAdd, AllocationAdd);
-		SetShaderValue(PixelShaderRHI, PixelShaderParameters->IndirectlightingCachePrimitiveScale, AllocationScale);
-		SetShaderValue(PixelShaderRHI, PixelShaderParameters->IndirectlightingCacheMinUV, MinUV);
-		SetShaderValue(PixelShaderRHI, PixelShaderParameters->IndirectlightingCacheMaxUV, MaxUV);
-		SetShaderValue(PixelShaderRHI, PixelShaderParameters->PointSkyBentNormal, SkyBentNormal);
+		SetShaderValue(RHICmdList, PixelShaderRHI, PixelShaderParameters->IndirectlightingCachePrimitiveAdd, AllocationAdd);
+		SetShaderValue(RHICmdList, PixelShaderRHI, PixelShaderParameters->IndirectlightingCachePrimitiveScale, AllocationScale);
+		SetShaderValue(RHICmdList, PixelShaderRHI, PixelShaderParameters->IndirectlightingCacheMinUV, MinUV);
+		SetShaderValue(RHICmdList, PixelShaderRHI, PixelShaderParameters->IndirectlightingCacheMaxUV, MaxUV);
+		SetShaderValue(RHICmdList, PixelShaderRHI, PixelShaderParameters->PointSkyBentNormal, SkyBentNormal);
 	}
 }
 
 void FCachedPointIndirectLightingPolicy::SetMesh(
+	FRHICommandList* RHICmdList, 
 	const FSceneView& View,
 	const FPrimitiveSceneProxy* PrimitiveSceneProxy,
 	const VertexParametersType* VertexShaderParameters,
@@ -485,6 +505,7 @@ void FCachedPointIndirectLightingPolicy::SetMesh(
 					* FSHVector2::ConstantBasisIntegral * .5f;
 
 				SetShaderValue(
+					RHICmdList, 
 					PixelShader->GetPixelShader(), 
 					PixelShaderParameters->IndirectLightingSHCoefficients, 
 					SwizzledAmbientTerm);
@@ -492,6 +513,7 @@ void FCachedPointIndirectLightingPolicy::SetMesh(
 			else
 			{
 				SetShaderValueArray(
+					RHICmdList, 
 					PixelShader->GetPixelShader(), 
 					PixelShaderParameters->IndirectLightingSHCoefficients, 
 					&LightingAllocation.SingleSamplePacked, 
@@ -499,11 +521,13 @@ void FCachedPointIndirectLightingPolicy::SetMesh(
 			}
 
 			SetShaderValue(
+				RHICmdList, 
 				PixelShader->GetPixelShader(), 
 				PixelShaderParameters->PointSkyBentNormal, 
 				LightingAllocation.CurrentSkyBentNormal);
 
 			SetShaderValue(
+				RHICmdList, 
 				PixelShader->GetPixelShader(), 
 				PixelShaderParameters->DirectionalLightShadowing, 
 				LightingAllocation.CurrentDirectionalShadowing);
@@ -513,17 +537,20 @@ void FCachedPointIndirectLightingPolicy::SetMesh(
 			const FVector4 ZeroArray[sizeof(FSHVectorRGB2) / sizeof(FVector4)] = {FVector4(0, 0, 0, 0), FVector4(0, 0, 0, 0), FVector4(0, 0, 0, 0)};
 
 			SetShaderValueArray(
+				RHICmdList, 
 				PixelShader->GetPixelShader(), 
 				PixelShaderParameters->IndirectLightingSHCoefficients, 
 				&ZeroArray, 
 				ARRAY_COUNT(ZeroArray));
 
 			SetShaderValue(
+				RHICmdList, 
 				PixelShader->GetPixelShader(), 
 				PixelShaderParameters->PointSkyBentNormal, 
 				FVector::ZeroVector);
 
 			SetShaderValue(
+				RHICmdList, 
 				PixelShader->GetPixelShader(), 
 				PixelShaderParameters->DirectionalLightShadowing, 
 				1);
@@ -532,6 +559,7 @@ void FCachedPointIndirectLightingPolicy::SetMesh(
 }
 
 void FSelfShadowedCachedPointIndirectLightingPolicy::SetMesh(
+	FRHICommandList* RHICmdList, 
 	const FSceneView& View,
 	const FPrimitiveSceneProxy* PrimitiveSceneProxy,
 	const VertexParametersType* VertexShaderParameters,
@@ -552,12 +580,14 @@ void FSelfShadowedCachedPointIndirectLightingPolicy::SetMesh(
 			const FIndirectLightingCacheAllocation& LightingAllocation = *PrimitiveSceneProxy->GetPrimitiveSceneInfo()->IndirectLightingCacheAllocation;
 
 			SetShaderValueArray(
+				RHICmdList, 
 				PixelShader->GetPixelShader(), 
 				PixelShaderParameters->IndirectLightingSHCoefficients, 
 				&LightingAllocation.SingleSamplePacked, 
 				ARRAY_COUNT(LightingAllocation.SingleSamplePacked));
 
 			SetShaderValue(
+				RHICmdList, 
 				PixelShader->GetPixelShader(), 
 				PixelShaderParameters->PointSkyBentNormal, 
 				LightingAllocation.CurrentSkyBentNormal);
@@ -567,12 +597,14 @@ void FSelfShadowedCachedPointIndirectLightingPolicy::SetMesh(
 			const FVector4 ZeroArray[sizeof(FSHVectorRGB2) / sizeof(FVector4)] = {FVector4(0, 0, 0, 0), FVector4(0, 0, 0, 0), FVector4(0, 0, 0, 0)};
 
 			SetShaderValueArray(
+				RHICmdList, 
 				PixelShader->GetPixelShader(), 
 				PixelShaderParameters->IndirectLightingSHCoefficients, 
 				&ZeroArray, 
 				ARRAY_COUNT(ZeroArray));
 
 			SetShaderValue(
+				RHICmdList, 
 				PixelShader->GetPixelShader(), 
 				PixelShaderParameters->PointSkyBentNormal, 
 				FVector::ZeroVector);
@@ -580,6 +612,7 @@ void FSelfShadowedCachedPointIndirectLightingPolicy::SetMesh(
 	}
 
 	FSelfShadowedTranslucencyPolicy::SetMesh(
+		RHICmdList, 
 		View, 
 		PrimitiveSceneProxy, 
 		VertexShaderParameters,
