UE4 - merge first pass at parallel rendering to main

diff --git a/Engine/Source/Runtime/Renderer/Private/BasePassRendering.h b/Engine/Source/Runtime/Renderer/Private/BasePassRendering.h
--- a/Engine/Source/Runtime/Renderer/Private/BasePassRendering.h
+++ b/Engine/Source/Runtime/Renderer/Private/BasePassRendering.h
@@ -54,6 +54,7 @@ public:
 	}
 
 	void SetParameters(
+		FRHICommandList* RHICmdList, 
 		const FMaterialRenderProxy* MaterialRenderProxy,
 		const FVertexFactory* VertexFactory,
 		const FMaterial& InMaterialResource,
@@ -62,18 +63,18 @@ public:
 		ESceneRenderTargetsMode::Type TextureMode
 		)
 	{
-		FMeshMaterialShader::SetParameters(GetVertexShader(),MaterialRenderProxy,InMaterialResource,View,TextureMode);
+		FMeshMaterialShader::SetParameters(RHICmdList, GetVertexShader(),MaterialRenderProxy,InMaterialResource,View,TextureMode);
 
 		if (bAllowGlobalFog)
 		{
-			HeightFogParameters.Set(GetVertexShader(), &View);
-			AtmosphericFogTextureParameters.Set(GetVertexShader(), View);
+			HeightFogParameters.Set(RHICmdList, GetVertexShader(), &View);
+			AtmosphericFogTextureParameters.Set(RHICmdList, GetVertexShader(), View);
 		}
 	}
 
-	void SetMesh(const FVertexFactory* VertexFactory,const FSceneView& View,const FPrimitiveSceneProxy* Proxy,const FMeshBatchElement& BatchElement)
+	void SetMesh(FRHICommandList* RHICmdList, const FVertexFactory* VertexFactory,const FSceneView& View,const FPrimitiveSceneProxy* Proxy,const FMeshBatchElement& BatchElement)
 	{
-		FMeshMaterialShader::SetMesh(GetVertexShader(),VertexFactory,View,Proxy,BatchElement);
+		FMeshMaterialShader::SetMesh(RHICmdList, GetVertexShader(),VertexFactory,View,Proxy,BatchElement);
 	}
 
 private:
@@ -180,12 +181,13 @@ public:
 	}
 
 	void SetParameters(
+		FRHICommandList* RHICmdList, 
 		const FMaterialRenderProxy* MaterialRenderProxy,
 		const FVertexFactory* VertexFactory,
 		const FSceneView& View
 		)
 	{
-		FBaseDS::SetParameters(MaterialRenderProxy, View);
+		FBaseDS::SetParameters(RHICmdList, MaterialRenderProxy, View);
 	}
 };
 
@@ -202,7 +204,7 @@ public:
 	}
 
 	template<typename TParamRef>
-	void SetParameters(const TParamRef& ShaderRHI, const FScene* Scene, bool bApplySkyLight)
+	void SetParameters(FRHICommandList* RHICmdList, const TParamRef& ShaderRHI, const FScene* Scene, bool bApplySkyLight)
 	{
 		FTexture* SkyLightTextureResource = GBlackTextureCube;
 		float ApplySkyLightMask = 0;
@@ -210,9 +212,9 @@ public:
 
 		GetSkyParametersFromScene(Scene, bApplySkyLight, SkyLightTextureResource, ApplySkyLightMask, SkyMipCount);
 
-		SetTextureParameter(ShaderRHI, SkyLightCubemap, SkyLightCubemapSampler, SkyLightTextureResource);
+		SetTextureParameter(RHICmdList, ShaderRHI, SkyLightCubemap, SkyLightCubemapSampler, SkyLightTextureResource);
 		const FVector2D SkyParametersValue(SkyMipCount - 1.0f, ApplySkyLightMask);
-		SetShaderValue(ShaderRHI, SkyLightParameters, SkyParametersValue);
+		SetShaderValue(RHICmdList, ShaderRHI, SkyLightParameters, SkyParametersValue);
 	}
 
 	friend FArchive& operator<<(FArchive& Ar,FSkyLightReflectionParameters& P)
@@ -251,9 +253,9 @@ public:
 		SkyLightReflectionParameters.Bind(ParameterMap);
 	}
 
-	void Set(FShader* Shader, const FSceneView* View);
+	void Set(FRHICommandList* RHICmdList, FShader* Shader, const FSceneView* View);
 
-	void SetMesh(FShader* Shader, const FPrimitiveSceneProxy* Proxy);
+	void SetMesh(FRHICommandList* RHICmdList, FShader* Shader, const FPrimitiveSceneProxy* Proxy);
 
 	/** Serializer. */
 	friend FArchive& operator<<(FArchive& Ar,FTranslucentLightingParameters& P)
@@ -323,6 +325,7 @@ public:
 	TBasePassPixelShaderBaseType() {}
 
 	void SetParameters(
+		FRHICommandList* RHICmdList, 
 		const FMaterialRenderProxy* MaterialRenderProxy, 
 		const FMaterial& MaterialResource, 
 		const FSceneView* View, 
@@ -332,13 +335,13 @@ public:
 	{
 		const FPixelShaderRHIParamRef ShaderRHI = GetPixelShader();
 
-		FMeshMaterialShader::SetParameters(ShaderRHI, MaterialRenderProxy, MaterialResource, *View, TextureMode);
+		FMeshMaterialShader::SetParameters(RHICmdList, ShaderRHI, MaterialRenderProxy, MaterialResource, *View, TextureMode);
 
 		if (View->GetFeatureLevel() >= ERHIFeatureLevel::SM4)
 		{
 			if (IsTranslucentBlendMode(BlendMode))
 			{
-				TranslucentLightingParameters.Set(this, View);
+				TranslucentLightingParameters.Set(RHICmdList, this, View);
 			}
 
 #if WITH_EDITOR
@@ -360,6 +363,7 @@ public:
 					const FTexture2DRHIRef* DepthTexture = GSceneRenderTargets.GetActualDepthTexture();
 					check(DepthTexture != NULL);
 					SetTextureParameter(
+						RHICmdList,
 						ShaderRHI,
 						FilteredSceneDepthTexture,
 						FilteredSceneDepthTextureSampler,
@@ -368,22 +372,22 @@ public:
 						);
 				}
 
-				SetShaderValue(ShaderRHI, EditorCompositeDepthTestParameter, bEnableEditorPrimitveDepthTest );
-				SetShaderValue(ShaderRHI, MSAASampleCount, GSceneRenderTargets.EditorPrimitivesColor ? GSceneRenderTargets.EditorPrimitivesColor->GetDesc().NumSamples : 0 );
+				SetShaderValue(RHICmdList, ShaderRHI, EditorCompositeDepthTestParameter, bEnableEditorPrimitveDepthTest );
+				SetShaderValue(RHICmdList, ShaderRHI, MSAASampleCount, GSceneRenderTargets.EditorPrimitivesColor ? GSceneRenderTargets.EditorPrimitivesColor->GetDesc().NumSamples : 0 );
 			}
 #endif
 		}
 	}
 
-	void SetMesh(const FVertexFactory* VertexFactory,const FSceneView& View,const FPrimitiveSceneProxy* Proxy,const FMeshBatchElement& BatchElement, EBlendMode BlendMode)
+	void SetMesh(FRHICommandList* RHICmdList, const FVertexFactory* VertexFactory,const FSceneView& View,const FPrimitiveSceneProxy* Proxy,const FMeshBatchElement& BatchElement, EBlendMode BlendMode)
 	{
 		if (View.GetFeatureLevel() >= ERHIFeatureLevel::SM4
 			&& IsTranslucentBlendMode(BlendMode))
 		{
-			TranslucentLightingParameters.SetMesh(this, Proxy);
+			TranslucentLightingParameters.SetMesh(RHICmdList, this, Proxy);
 		}
 
-		FMeshMaterialShader::SetMesh(GetPixelShader(),VertexFactory,View,Proxy,BatchElement);
+		FMeshMaterialShader::SetMesh(RHICmdList, GetPixelShader(),VertexFactory,View,Proxy,BatchElement);
 	}
 
 	virtual bool Serialize(FArchive& Ar)
@@ -545,24 +549,23 @@ public:
 			LightMapPolicy == Other.LightMapPolicy;
 	}
 
-	void DrawShared(const FSceneView* View,FBoundShaderStateRHIParamRef BoundShaderState) const
+	void DrawShared(FRHICommandList* RHICmdList, const FSceneView* View,FBoundShaderStateRHIParamRef BoundShaderState) const
 	{
-		// Set the actual shader & vertex declaration state
-		RHISetBoundShaderState( BoundShaderState);
+		FRHICommandList::SetBoundShaderState(RHICmdList, BoundShaderState);
 
 		// Set the light-map policy.
-		LightMapPolicy.Set(VertexShader,bOverrideWithShaderComplexity ? NULL : PixelShader,VertexShader,PixelShader,VertexFactory,MaterialRenderProxy,View);
+		LightMapPolicy.Set(RHICmdList, VertexShader,bOverrideWithShaderComplexity ? NULL : PixelShader,VertexShader,PixelShader,VertexFactory,MaterialRenderProxy,View);
 
-		VertexShader->SetParameters(MaterialRenderProxy, VertexFactory, *MaterialResource, *View, bAllowGlobalFog, SceneTextureMode);
+		VertexShader->SetParameters(RHICmdList, MaterialRenderProxy, VertexFactory, *MaterialResource, *View, bAllowGlobalFog, SceneTextureMode);
 
 		if(HullShader)
 		{
-			HullShader->SetParameters(MaterialRenderProxy, *View);
+			HullShader->SetParameters(RHICmdList, MaterialRenderProxy, *View);
 		}
 
 		if (DomainShader)
 		{
-			DomainShader->SetParameters(MaterialRenderProxy, VertexFactory, *View);
+			DomainShader->SetParameters(RHICmdList, MaterialRenderProxy, VertexFactory, *View);
 		}
 
 #if !(UE_BUILD_SHIPPING || UE_BUILD_TEST)
@@ -571,19 +574,18 @@ public:
 			// If we are in the translucent pass then override the blend mode, otherwise maintain additive blending.
 			if (IsTranslucentBlendMode(BlendMode))
 			{
-				// Add complexity to existing
-				RHISetBlendState(TStaticBlendState<CW_RGBA,BO_Add,BF_One,BF_One,BO_Add,BF_Zero,BF_One>::GetRHI());
+				FRHICommandList::SetBlendState(RHICmdList, TStaticBlendState<CW_RGBA, BO_Add, BF_One, BF_One, BO_Add, BF_Zero, BF_One>::GetRHI());
 			}
 
 			TShaderMapRef<FShaderComplexityAccumulatePS> ShaderComplexityPixelShader(GetGlobalShaderMap());
 			const uint32 NumPixelShaderInstructions = PixelShader->GetNumInstructions();
 			const uint32 NumVertexShaderInstructions = VertexShader->GetNumInstructions();
-			ShaderComplexityPixelShader->SetParameters(NumVertexShaderInstructions, NumPixelShaderInstructions, View->GetFeatureLevel());
+			ShaderComplexityPixelShader->SetParameters(RHICmdList, NumVertexShaderInstructions, NumPixelShaderInstructions, View->GetFeatureLevel());
 		}
 		else
 #endif
 		{
-			PixelShader->SetParameters(MaterialRenderProxy,*MaterialResource,View,BlendMode,bEnableEditorPrimitiveDepthTest,SceneTextureMode);
+			PixelShader->SetParameters(RHICmdList, MaterialRenderProxy,*MaterialResource,View,BlendMode,bEnableEditorPrimitiveDepthTest,SceneTextureMode);
 
 			switch(BlendMode)
 			{
@@ -596,16 +598,17 @@ public:
 				break;
 			case BLEND_Translucent:
 				// Alpha channel is only needed for SeparateTranslucency, before this was preserving the alpha channel but we no longer store depth in the alpha channel so it's no problem
-				RHISetBlendState(TStaticBlendState<CW_RGBA, BO_Add,BF_SourceAlpha,BF_InverseSourceAlpha,BO_Add,BF_Zero,BF_InverseSourceAlpha>::GetRHI());
+
+				FRHICommandList::SetBlendState(RHICmdList, TStaticBlendState<CW_RGBA, BO_Add, BF_SourceAlpha, BF_InverseSourceAlpha, BO_Add, BF_Zero, BF_InverseSourceAlpha>::GetRHI());
 				break;
 			case BLEND_Additive:
 				// Add to the existing scene color
 				// Alpha channel is only needed for SeparateTranslucency, before this was preserving the alpha channel but we no longer store depth in the alpha channel so it's no problem
-				RHISetBlendState(TStaticBlendState<CW_RGBA, BO_Add,BF_One,BF_One,BO_Add,BF_Zero,BF_InverseSourceAlpha>::GetRHI());
+				FRHICommandList::SetBlendState(RHICmdList, TStaticBlendState<CW_RGBA, BO_Add, BF_One, BF_One, BO_Add, BF_Zero, BF_InverseSourceAlpha>::GetRHI());
 				break;
 			case BLEND_Modulate:
 				// Modulate with the existing scene color, preserve destination alpha.
-				RHISetBlendState(TStaticBlendState<CW_RGB,BO_Add,BF_DestColor,BF_Zero>::GetRHI());
+				FRHICommandList::SetBlendState(RHICmdList, TStaticBlendState<CW_RGB, BO_Add, BF_DestColor, BF_Zero>::GetRHI());
 				break;
 			};
 		}
@@ -642,6 +645,7 @@ public:
 	}
 
 	void SetMeshRenderState(
+		FRHICommandList* RHICmdList, 
 		const FSceneView& View,
 		const FPrimitiveSceneProxy* PrimitiveSceneProxy,
 		const FMeshBatch& Mesh,
@@ -652,6 +656,7 @@ public:
 	{
 		// Set the light-map policy's mesh-specific settings.
 		LightMapPolicy.SetMesh(
+			RHICmdList, 
 			View,
 			PrimitiveSceneProxy,
 			VertexShader,
@@ -663,12 +668,12 @@ public:
 			ElementData.LightMapElementData);
 
 		const FMeshBatchElement& BatchElement = Mesh.Elements[BatchElementIndex];
-		VertexShader->SetMesh(VertexFactory,View,PrimitiveSceneProxy,BatchElement);
+		VertexShader->SetMesh(RHICmdList, VertexFactory,View,PrimitiveSceneProxy,BatchElement);
 		
 		if(HullShader && DomainShader)
 		{
-			HullShader->SetMesh(VertexFactory,View,PrimitiveSceneProxy,BatchElement);
-			DomainShader->SetMesh(VertexFactory,View,PrimitiveSceneProxy,BatchElement);
+			HullShader->SetMesh(RHICmdList, VertexFactory,View,PrimitiveSceneProxy,BatchElement);
+			DomainShader->SetMesh(RHICmdList, VertexFactory,View,PrimitiveSceneProxy,BatchElement);
 		}
 
 #if !(UE_BUILD_SHIPPING || UE_BUILD_TEST)
@@ -678,21 +683,22 @@ public:
 			if (BlendMode != BLEND_Opaque)
 			{
 				// Add complexity to existing, keep alpha
+				check(!RHICmdList);
 				RHISetBlendState(TStaticBlendState<CW_RGB,BO_Add,BF_One,BF_One>::GetRHI());
 			}
 
 			TShaderMapRef<FShaderComplexityAccumulatePS> ShaderComplexityPixelShader(GetGlobalShaderMap());
 			const uint32 NumPixelShaderInstructions = PixelShader->GetNumInstructions();
 			const uint32 NumVertexShaderInstructions = VertexShader->GetNumInstructions();
-			ShaderComplexityPixelShader->SetParameters(NumVertexShaderInstructions,NumPixelShaderInstructions, View.GetFeatureLevel());
+			ShaderComplexityPixelShader->SetParameters(RHICmdList, NumVertexShaderInstructions,NumPixelShaderInstructions, View.GetFeatureLevel());
 		}
 		else
 #endif
 		{
-			PixelShader->SetMesh(VertexFactory,View,PrimitiveSceneProxy,BatchElement,BlendMode);
+			PixelShader->SetMesh(RHICmdList, VertexFactory,View,PrimitiveSceneProxy,BatchElement,BlendMode);
 		}
 
-		FMeshDrawingPolicy::SetMeshRenderState(View,PrimitiveSceneProxy,Mesh,BatchElementIndex,bBackFace,FMeshDrawingPolicy::ElementDataType());
+		FMeshDrawingPolicy::SetMeshRenderState(RHICmdList, View,PrimitiveSceneProxy,Mesh,BatchElementIndex,bBackFace,FMeshDrawingPolicy::ElementDataType());
 	}
 
 	friend int32 CompareDrawingPolicy(const TBasePassDrawingPolicy& A,const TBasePassDrawingPolicy& B)
@@ -754,6 +760,7 @@ public:
 
 	static void AddStaticMesh(FScene* Scene,FStaticMesh* StaticMesh);
 	static bool DrawDynamicMesh(
+		FRHICommandList* RHICmdList, 
 		const FSceneView& View,
 		ContextType DrawingContext,
 		const FMeshBatch& Mesh,
@@ -838,6 +845,7 @@ public:
 /** Processes a base pass mesh using an unknown light map policy, and unknown fog density policy. */
 template<typename ProcessActionType>
 void ProcessBasePassMesh(
+	FRHICommandList* RHICmdList, 
 	const FProcessBasePassMeshParameters& Parameters,
 	const ProcessActionType& Action
 	)
@@ -857,17 +865,17 @@ void ProcessBasePassMesh(
 		if (bIsLitMaterial && Action.UseTranslucentSelfShadowing() && Parameters.FeatureLevel >= ERHIFeatureLevel::SM4)
 		{
 			if (IsIndirectLightingCacheAllowed(Parameters.FeatureLevel)
-				&& Action.AllowIndirectLightingCache()
-				&& Parameters.PrimitiveSceneProxy)
-			{
-				// Apply cached point indirect lighting as well as self shadowing if needed
-				Action.template Process<FSelfShadowedCachedPointIndirectLightingPolicy>(Parameters, FSelfShadowedCachedPointIndirectLightingPolicy(), FSelfShadowedTranslucencyPolicy::ElementDataType(Action.GetTranslucentSelfShadow()));
-			}
-			else
-			{
-				Action.template Process<FSelfShadowedTranslucencyPolicy>(Parameters, FSelfShadowedTranslucencyPolicy(), FSelfShadowedTranslucencyPolicy::ElementDataType(Action.GetTranslucentSelfShadow()));
+					&& Action.AllowIndirectLightingCache()
+					&& Parameters.PrimitiveSceneProxy)
+				{
+					// Apply cached point indirect lighting as well as self shadowing if needed
+					Action.template Process<FSelfShadowedCachedPointIndirectLightingPolicy>(RHICmdList, Parameters, FSelfShadowedCachedPointIndirectLightingPolicy(), FSelfShadowedTranslucencyPolicy::ElementDataType(Action.GetTranslucentSelfShadow()));
+				}
+				else
+				{
+					Action.template Process<FSelfShadowedTranslucencyPolicy>(RHICmdList, Parameters, FSelfShadowedTranslucencyPolicy(), FSelfShadowedTranslucencyPolicy::ElementDataType(Action.GetTranslucentSelfShadow()));
+				}
 			}
-		}
 		else
 		{
 			const FLightMapInteraction LightMapInteraction = (bAllowStaticLighting && Parameters.Mesh.LCI && bIsLitMaterial) 
@@ -889,18 +897,19 @@ void ProcessBasePassMesh(
 						if (ShadowMapInteraction.GetType() == SMIT_Texture)
 						{
 							Action.template Process< TDistanceFieldShadowsAndLightMapPolicy<HQ_LIGHTMAP> >(
+								RHICmdList, 
 								Parameters,
 								TDistanceFieldShadowsAndLightMapPolicy<HQ_LIGHTMAP>(),
 								TDistanceFieldShadowsAndLightMapPolicy<HQ_LIGHTMAP>::ElementDataType(ShadowMapInteraction, LightMapInteraction));
 						}
 						else
 						{
-							Action.template Process< TLightMapPolicy<HQ_LIGHTMAP> >( Parameters, TLightMapPolicy<HQ_LIGHTMAP>(), LightMapInteraction );
+							Action.template Process< TLightMapPolicy<HQ_LIGHTMAP> >(RHICmdList, Parameters, TLightMapPolicy<HQ_LIGHTMAP>(), LightMapInteraction );
 						}
 					} 
 					else 
 					{ 
-						Action.template Process< TLightMapPolicy<LQ_LIGHTMAP> >( Parameters, TLightMapPolicy<LQ_LIGHTMAP>(), LightMapInteraction );
+						Action.template Process< TLightMapPolicy<LQ_LIGHTMAP> >(RHICmdList, Parameters, TLightMapPolicy<LQ_LIGHTMAP>(), LightMapInteraction );
 					} 
 					break;
 				default:
@@ -910,7 +919,7 @@ void ProcessBasePassMesh(
 						{
 							if (IsSimpleDynamicLightingEnabled())
 							{
-								Action.template Process<FSimpleDynamicLightingPolicy>(Parameters, FSimpleDynamicLightingPolicy(), FSimpleDynamicLightingPolicy::ElementDataType());
+								Action.template Process<FSimpleDynamicLightingPolicy>(RHICmdList, Parameters, FSimpleDynamicLightingPolicy(), FSimpleDynamicLightingPolicy::ElementDataType());
 							}
 							else if (IsIndirectLightingCacheAllowed(Parameters.FeatureLevel)
 								&& Action.AllowIndirectLightingCache()
@@ -925,22 +934,22 @@ void ProcessBasePassMesh(
 								if (CanIndirectLightingCacheUseVolumeTexture(Parameters.FeatureLevel) && Action.AllowIndirectLightingCacheVolumeTexture())
 								{
 									// Use a lightmap policy that supports reading indirect lighting from a volume texture for dynamic objects
-									Action.template Process<FCachedVolumeIndirectLightingPolicy>(Parameters,FCachedVolumeIndirectLightingPolicy(),FCachedVolumeIndirectLightingPolicy::ElementDataType());
+									Action.template Process<FCachedVolumeIndirectLightingPolicy>(RHICmdList, Parameters,FCachedVolumeIndirectLightingPolicy(),FCachedVolumeIndirectLightingPolicy::ElementDataType());
 								}
 								else
 								{
 									// Use a lightmap policy that supports reading indirect lighting from a single SH sample
-									Action.template Process<FCachedPointIndirectLightingPolicy>(Parameters,FCachedPointIndirectLightingPolicy(),FCachedPointIndirectLightingPolicy::ElementDataType(false));
+									Action.template Process<FCachedPointIndirectLightingPolicy>(RHICmdList, Parameters,FCachedPointIndirectLightingPolicy(),FCachedPointIndirectLightingPolicy::ElementDataType(false));
 								}
 							}
 							else
 							{
-								Action.template Process<FNoLightMapPolicy>(Parameters,FNoLightMapPolicy(),FNoLightMapPolicy::ElementDataType());
+								Action.template Process<FNoLightMapPolicy>(RHICmdList, Parameters,FNoLightMapPolicy(),FNoLightMapPolicy::ElementDataType());
 							}
 						}
 						else
 						{
-							Action.template Process<FNoLightMapPolicy>(Parameters,FNoLightMapPolicy(),FNoLightMapPolicy::ElementDataType());
+							Action.template Process<FNoLightMapPolicy>(RHICmdList, Parameters,FNoLightMapPolicy(),FNoLightMapPolicy::ElementDataType());
 						}
 					}
 					break;
diff --git a/Engine/Source/Runtime/Renderer/Private/ClearQuad.cpp b/Engine/Source/Runtime/Renderer/Private/ClearQuad.cpp
--- a/Engine/Source/Runtime/Renderer/Private/ClearQuad.cpp
+++ b/Engine/Source/Runtime/Renderer/Private/ClearQuad.cpp
@@ -23,7 +23,10 @@ void DrawClearQuadMRT(bool bClearColor,int32 NumClearColors,const FLinearColor*
 			? TStaticBlendState<>::GetRHI()
 			: TStaticBlendStateWriteMask<CW_NONE,CW_NONE,CW_NONE,CW_NONE,CW_NONE,CW_NONE,CW_NONE,CW_NONE>::GetRHI();
 	}
-		
+	
+	//@todo-rco: RHIPacketList
+	FRHICommandList* RHICmdList = nullptr;
+
 	const FDepthStencilStateRHIParamRef DepthStencilStateRHI = 
 		(bClearDepth && bClearStencil)
 			? TStaticDepthStencilState<
@@ -104,7 +107,7 @@ void DrawClearQuadMRT(bool bClearColor,int32 NumClearColors,const FLinearColor*
 		ShaderClearColors[i] = ClearColorArray[i];
 	}
 
-	SetShaderValueArray(PixelShader->GetPixelShader(),PixelShader->ColorParameter,ShaderClearColors,NumClearColors);
+	SetShaderValueArray(RHICmdList, PixelShader->GetPixelShader(),PixelShader->ColorParameter,ShaderClearColors,NumClearColors);
 		
 	{
 		// Draw a fullscreen quad
diff --git a/Engine/Source/Runtime/Renderer/Private/CompositionLighting/PostProcessAmbient.cpp b/Engine/Source/Runtime/Renderer/Private/CompositionLighting/PostProcessAmbient.cpp
--- a/Engine/Source/Runtime/Renderer/Private/CompositionLighting/PostProcessAmbient.cpp
+++ b/Engine/Source/Runtime/Renderer/Private/CompositionLighting/PostProcessAmbient.cpp
@@ -20,6 +20,12 @@ class FPostProcessAmbientPS : public FGlobalShader
 		return IsFeatureLevelSupported(Platform, ERHIFeatureLevel::SM3);
 	}
 
+	static void ModifyCompilationEnvironment(EShaderPlatform Platform, FShaderCompilerEnvironment& OutEnvironment)
+	{
+		FGlobalShader::ModifyCompilationEnvironment(Platform,OutEnvironment);
+		FDeferredPixelShaderParameters::ModifyCompilationEnvironment(Platform,OutEnvironment);
+	}
+
 	/** Default constructor. */
 	FPostProcessAmbientPS() {}
 
@@ -41,15 +47,15 @@ public:
 		PreIntegratedGFSampler.Bind(Initializer.ParameterMap, TEXT("PreIntegratedGFSampler"));
 	}
 
-	void SetParameters(const FRenderingCompositePassContext& Context, const FFinalPostProcessSettings::FCubemapEntry& Entry)
+	void SetParameters(FRHICommandList* RHICmdList, const FRenderingCompositePassContext& Context, const FFinalPostProcessSettings::FCubemapEntry& Entry)
 	{
 		const FPixelShaderRHIParamRef ShaderRHI = GetPixelShader();
 
-		FGlobalShader::SetParameters(ShaderRHI, Context.View);
-		PostprocessParameter.SetPS(ShaderRHI, Context, TStaticSamplerState<SF_Bilinear,AM_Clamp,AM_Clamp,AM_Clamp>::GetRHI());
-		DeferredParameters.Set(ShaderRHI, Context.View);
+		FGlobalShader::SetParameters(RHICmdList, ShaderRHI, Context.View);
+		PostprocessParameter.SetPS(RHICmdList, ShaderRHI, Context, TStaticSamplerState<SF_Bilinear,AM_Clamp,AM_Clamp,AM_Clamp>::GetRHI());
+		DeferredParameters.Set(RHICmdList, ShaderRHI, Context.View);
 		CubemapShaderParameters.SetParameters(ShaderRHI, Entry);
-		SetTextureParameter(ShaderRHI, PreIntegratedGF, PreIntegratedGFSampler, TStaticSamplerState<SF_Bilinear, AM_Clamp, AM_Clamp, AM_Clamp>::GetRHI(), GSystemTextures.PreintegratedGF->GetRenderTargetItem().ShaderResourceTexture);
+		SetTextureParameter(RHICmdList, ShaderRHI, PreIntegratedGF, PreIntegratedGFSampler, TStaticSamplerState<SF_Bilinear,AM_Clamp,AM_Clamp,AM_Clamp>::GetRHI(), GSystemTextures.PreintegratedGF->GetRenderTargetItem().ShaderResourceTexture );
 	}
 	
 	// FShader interface.
@@ -107,7 +113,9 @@ void FRCPassPostProcessAmbient::Process(FRenderingCompositePassContext& Context)
 			SetGlobalBoundShaderState(BoundShaderState, GFilterVertexDeclaration.VertexDeclarationRHI, *VertexShader, *PixelShader);
 		}
 
-		PixelShader->SetParameters(Context, Context.View.FinalPostProcessSettings.ContributingCubemaps[i]);
+		//@todo-rco: RHIPacketList
+		FRHICommandList* RHICmdList = nullptr;
+		PixelShader->SetParameters(nullptr, Context, Context.View.FinalPostProcessSettings.ContributingCubemaps[i]);
 
 		// Draw a quad mapping scene color to the view's render target
 		DrawRectangle( 
@@ -159,6 +167,9 @@ void FCubemapShaderParameters::SetParameters(const FComputeShaderRHIParamRef Sha
 template<typename TShaderRHIRef>
 void FCubemapShaderParameters::SetParametersTemplate(const TShaderRHIRef& ShaderRHI, const FFinalPostProcessSettings::FCubemapEntry& Entry) const
 {
+	//@todo-rco: RHIPacketList
+	FRHICommandList* RHICmdList = nullptr;
+
 	// floats to render the cubemap
 	{
 		float MipCount = 0;
@@ -172,7 +183,7 @@ void FCubemapShaderParameters::SetParametersTemplate(const TShaderRHIRef& Shader
 		{
 			FLinearColor AmbientCubemapTintMulScaleValue = Entry.AmbientCubemapTintMulScaleValue;
 
-			SetShaderValue(ShaderRHI, AmbientCubemapColor, AmbientCubemapTintMulScaleValue);
+			SetShaderValue(RHICmdList, ShaderRHI, AmbientCubemapColor, AmbientCubemapTintMulScaleValue);
 		}
 
 		{
@@ -183,7 +194,7 @@ void FCubemapShaderParameters::SetParametersTemplate(const TShaderRHIRef& Shader
 			AmbientCubemapMipAdjustValue.Z = MipCount - GDiffuseConvolveMipLevel;
 			AmbientCubemapMipAdjustValue.W = MipCount;
 
-			SetShaderValue(ShaderRHI, AmbientCubemapMipAdjust, AmbientCubemapMipAdjustValue);
+			SetShaderValue(RHICmdList, ShaderRHI, AmbientCubemapMipAdjust, AmbientCubemapMipAdjustValue);
 		}
 	}
 
@@ -191,7 +202,7 @@ void FCubemapShaderParameters::SetParametersTemplate(const TShaderRHIRef& Shader
 	{
 		FTexture* InternalSpec = Entry.AmbientCubemap ? Entry.AmbientCubemap->Resource : GBlackTextureCube;
 
-		SetTextureParameter(ShaderRHI, AmbientCubemap, AmbientCubemapSampler, InternalSpec);
+		SetTextureParameter(RHICmdList, ShaderRHI, AmbientCubemap, AmbientCubemapSampler, InternalSpec);
 	}
 }
 
diff --git a/Engine/Source/Runtime/Renderer/Private/CompositionLighting/PostProcessAmbientOcclusion.cpp b/Engine/Source/Runtime/Renderer/Private/CompositionLighting/PostProcessAmbientOcclusion.cpp
--- a/Engine/Source/Runtime/Renderer/Private/CompositionLighting/PostProcessAmbientOcclusion.cpp
+++ b/Engine/Source/Runtime/Renderer/Private/CompositionLighting/PostProcessAmbientOcclusion.cpp
@@ -33,7 +33,8 @@ class FPostProcessAmbientOcclusionSetupPS : public FGlobalShader
 
 	static void ModifyCompilationEnvironment(EShaderPlatform Platform, FShaderCompilerEnvironment& OutEnvironment)
 	{
-		FGlobalShader::ModifyCompilationEnvironment(Platform, OutEnvironment);
+		FGlobalShader::ModifyCompilationEnvironment(Platform,OutEnvironment);
+		FDeferredPixelShaderParameters::ModifyCompilationEnvironment(Platform,OutEnvironment);
 		OutEnvironment.SetDefine(TEXT("INITIAL_PASS"), bInitialPass);
 	}
 
@@ -54,22 +55,22 @@ public:
 		AmbientOcclusionSetupParams.Bind(Initializer.ParameterMap, TEXT("AmbientOcclusionSetupParams"));
 	}
 
-	void SetParameters(const FRenderingCompositePassContext& Context)
+	void SetParameters(FRHICommandList* RHICmdList, const FRenderingCompositePassContext& Context)
 	{
 		const FFinalPostProcessSettings& Settings = Context.View.FinalPostProcessSettings;
 		const FPixelShaderRHIParamRef ShaderRHI = GetPixelShader();
 
-		FGlobalShader::SetParameters(ShaderRHI, Context.View);
+		FGlobalShader::SetParameters(RHICmdList, ShaderRHI, Context.View);
 
-		PostprocessParameter.SetPS(ShaderRHI, Context, TStaticSamplerState<SF_Bilinear,AM_Clamp,AM_Clamp,AM_Clamp>::GetRHI());
-		DeferredParameters.Set(ShaderRHI, Context.View);
+		PostprocessParameter.SetPS(RHICmdList, ShaderRHI, Context, TStaticSamplerState<SF_Bilinear,AM_Clamp,AM_Clamp,AM_Clamp>::GetRHI());
+		DeferredParameters.Set(RHICmdList, ShaderRHI, Context.View);
 
 		// e.g. 4 means the input texture is 4x smaller than the buffer size
 		uint32 ScaleToFullRes = GSceneRenderTargets.GetBufferSizeXY().X / Context.Pass->GetOutput(ePId_Output0)->RenderTargetDesc.Extent.X;
 
 		// /1000 to be able to define the value in that distance
 		FVector4 AmbientOcclusionSetupParamsValue = FVector4(ScaleToFullRes, Settings.AmbientOcclusionMipThreshold / ScaleToFullRes, 0, 0);
-		SetShaderValue(ShaderRHI, AmbientOcclusionSetupParams, AmbientOcclusionSetupParamsValue);
+		SetShaderValue(RHICmdList, ShaderRHI, AmbientOcclusionSetupParams, AmbientOcclusionSetupParamsValue);
 	}
 
 	// FShader interface.
@@ -116,7 +117,7 @@ void FCameraMotionParameters::Bind(const FShaderParameterMap& ParameterMap)
 	CameraMotion.Bind(ParameterMap, TEXT("CameraMotion"));
 }
 
-void FCameraMotionParameters::Set(const FSceneView& View, const FPixelShaderRHIParamRef ShaderRHI) const
+void FCameraMotionParameters::Set(FRHICommandList* RHICmdList, const FSceneView& View, const FPixelShaderRHIParamRef ShaderRHI) const
 {
 	FSceneViewState* ViewState = (FSceneViewState*)View.State;
 
@@ -178,7 +179,7 @@ void FCameraMotionParameters::Set(const FSceneView& View, const FPixelShaderRHIP
 		(float)(cwy*pww + cww*(pww - pyw) - cwy*pyw + cwx*((-pww) + pyw) + (cxw - cxx + cxy)*(pwx - pyx) + (cyw - cyx + cyy)*(pwy - pyy) + (czw - czx + czy)*(pwz - pyz)),
 		(float)(0));
 
-	SetShaderValueArray(ShaderRHI, CameraMotion, CameraMotionValue, 5);
+	SetShaderValueArray(RHICmdList, ShaderRHI, CameraMotion, CameraMotionValue, 5);
 }
 
 FArchive& operator<<(FArchive& Ar, FCameraMotionParameters& This)
@@ -205,7 +206,8 @@ class FPostProcessAmbientOcclusionPS : public FGlobalShader
 
 	static void ModifyCompilationEnvironment(EShaderPlatform Platform, FShaderCompilerEnvironment& OutEnvironment)
 	{
-		FGlobalShader::ModifyCompilationEnvironment(Platform, OutEnvironment);
+		FGlobalShader::ModifyCompilationEnvironment(Platform,OutEnvironment);
+		FDeferredPixelShaderParameters::ModifyCompilationEnvironment(Platform,OutEnvironment);
 
 		OutEnvironment.SetDefine(TEXT("USE_UPSAMPLE"), bDoUpsample);
 		OutEnvironment.SetDefine(TEXT("USE_AO_SETUP_AS_INPUT"), bTAOSetupAsInput);
@@ -233,22 +235,22 @@ public:
 		RandomNormalTextureSampler.Bind(Initializer.ParameterMap, TEXT("RandomNormalTextureSampler"));
 	}
 
-	void SetParameters(const FRenderingCompositePassContext& Context, FIntPoint InputTextureSize)
+	void SetParameters(FRHICommandList* RHICmdList, const FRenderingCompositePassContext& Context, FIntPoint InputTextureSize)
 	{
 		const FFinalPostProcessSettings& Settings = Context.View.FinalPostProcessSettings;
 		const FPixelShaderRHIParamRef ShaderRHI = GetPixelShader();
 
-		FGlobalShader::SetParameters(ShaderRHI, Context.View);
+		FGlobalShader::SetParameters(RHICmdList, ShaderRHI, Context.View);
 
 		// SF_Point is better than bilinear to avoid halos around objects
-		PostprocessParameter.SetPS(ShaderRHI, Context, TStaticSamplerState<SF_Point,AM_Clamp,AM_Clamp,AM_Clamp>::GetRHI());
-		DeferredParameters.Set(ShaderRHI, Context.View);
+		PostprocessParameter.SetPS(RHICmdList, ShaderRHI, Context, TStaticSamplerState<SF_Point,AM_Clamp,AM_Clamp,AM_Clamp>::GetRHI());
+		DeferredParameters.Set(RHICmdList, ShaderRHI, Context.View);
 
 		const FSceneRenderTargetItem& SSAORandomization = GSystemTextures.SSAORandomization->GetRenderTargetItem();
 
-		SetTextureParameter(ShaderRHI, RandomNormalTexture, RandomNormalTextureSampler, TStaticSamplerState<SF_Point,AM_Wrap,AM_Wrap,AM_Wrap>::GetRHI(), SSAORandomization.ShaderResourceTexture);
+		SetTextureParameter(RHICmdList, ShaderRHI, RandomNormalTexture, RandomNormalTextureSampler, TStaticSamplerState<SF_Point,AM_Wrap,AM_Wrap,AM_Wrap>::GetRHI(), SSAORandomization.ShaderResourceTexture);
 
-		ScreenSpaceAOandSSRShaderParams.Set(Context.View, ShaderRHI, InputTextureSize);
+		ScreenSpaceAOandSSRShaderParams.Set(RHICmdList, Context.View, ShaderRHI, InputTextureSize);
 	}
 	
 	// FShader interface.
@@ -296,6 +298,12 @@ class FPostProcessBasePassAOPS : public FGlobalShader
 		return IsFeatureLevelSupported(Platform, ERHIFeatureLevel::SM4);
 	}
 
+	static void ModifyCompilationEnvironment(EShaderPlatform Platform, FShaderCompilerEnvironment& OutEnvironment)
+	{
+		FGlobalShader::ModifyCompilationEnvironment(Platform,OutEnvironment);
+		FDeferredPixelShaderParameters::ModifyCompilationEnvironment(Platform,OutEnvironment);
+	}
+
 	/** Default constructor. */
 	FPostProcessBasePassAOPS() {}
 
@@ -313,15 +321,15 @@ public:
 		ScreenSpaceAOandSSRShaderParams.Bind(Initializer.ParameterMap);
 	}
 
-	void SetParameters(const FRenderingCompositePassContext& Context, FIntPoint InputTextureSize)
+	void SetParameters(FRHICommandList* RHICmdList, const FRenderingCompositePassContext& Context, FIntPoint InputTextureSize)
 	{
 		const FFinalPostProcessSettings& Settings = Context.View.FinalPostProcessSettings;
 		const FPixelShaderRHIParamRef ShaderRHI = GetPixelShader();
 
-		FGlobalShader::SetParameters(ShaderRHI, Context.View);
-		PostprocessParameter.SetPS(ShaderRHI, Context, TStaticSamplerState<SF_Point,AM_Clamp,AM_Clamp,AM_Clamp>::GetRHI());
-		DeferredParameters.Set(ShaderRHI, Context.View);
-		ScreenSpaceAOandSSRShaderParams.Set(Context.View, ShaderRHI, InputTextureSize);
+		FGlobalShader::SetParameters(RHICmdList, ShaderRHI, Context.View);
+		PostprocessParameter.SetPS(RHICmdList, ShaderRHI, Context, TStaticSamplerState<SF_Point,AM_Clamp,AM_Clamp,AM_Clamp>::GetRHI());
+		DeferredParameters.Set(RHICmdList, ShaderRHI, Context.View);
+		ScreenSpaceAOandSSRShaderParams.Set(RHICmdList, Context.View, ShaderRHI, InputTextureSize);
 	}
 
 	// FShader interface.
@@ -339,7 +347,7 @@ IMPLEMENT_SHADER_TYPE(,FPostProcessBasePassAOPS,TEXT("PostProcessAmbientOcclusio
 
 
 template <uint32 bInitialSetup>
-void FRCPassPostProcessAmbientOcclusionSetup::SetShaderSetupTempl(const FRenderingCompositePassContext& Context)
+void FRCPassPostProcessAmbientOcclusionSetup::SetShaderSetupTempl(FRHICommandList* RHICmdList, const FRenderingCompositePassContext& Context)
 {
 	TShaderMapRef<FPostProcessVS> VertexShader(GetGlobalShaderMap());
 	TShaderMapRef<FPostProcessAmbientOcclusionSetupPS<bInitialSetup> > PixelShader(GetGlobalShaderMap());
@@ -348,8 +356,8 @@ void FRCPassPostProcessAmbientOcclusionSetup::SetShaderSetupTempl(const FRenderi
 
 	SetGlobalBoundShaderState(BoundShaderState, GFilterVertexDeclaration.VertexDeclarationRHI, *VertexShader, *PixelShader);
 
-	VertexShader->SetParameters(Context);
-	PixelShader->SetParameters(Context);
+	VertexShader->SetParameters(RHICmdList, Context);
+	PixelShader->SetParameters(RHICmdList, Context);
 }
 
 void FRCPassPostProcessAmbientOcclusionSetup::Process(FRenderingCompositePassContext& Context)
@@ -378,13 +386,15 @@ void FRCPassPostProcessAmbientOcclusionSetup::Process(FRenderingCompositePassCon
 	RHISetRasterizerState(TStaticRasterizerState<>::GetRHI());
 	RHISetDepthStencilState(TStaticDepthStencilState<false,CF_Always>::GetRHI());
 
+	//@todo-rco: RHIPacketList
+	FRHICommandList* RHICmdList = nullptr;
 	if(IsInitialPass())
 	{
-		SetShaderSetupTempl<1>(Context);
+		SetShaderSetupTempl<1>(RHICmdList, Context);
 	}
 	else
 	{
-		SetShaderSetupTempl<0>(Context);
+		SetShaderSetupTempl<0>(RHICmdList, Context);
 	}
 
 	TShaderMapRef<FPostProcessVS> VertexShader(GetGlobalShaderMap());
@@ -459,8 +469,10 @@ void FRCPassPostProcessAmbientOcclusion::SetShaderTempl(const FRenderingComposit
 
 	const FPooledRenderTargetDesc* InputDesc0 = GetInputDesc(ePId_Input0);
 
-	VertexShader->SetParameters(Context);
-	PixelShader->SetParameters(Context, InputDesc0->Extent);
+	//@todo-rco: RHIPacketList
+	FRHICommandList* RHICmdList = nullptr;
+	VertexShader->SetParameters(RHICmdList, Context);
+	PixelShader->SetParameters(RHICmdList, Context, InputDesc0->Extent);
 }
 
 void FRCPassPostProcessAmbientOcclusion::Process(FRenderingCompositePassContext& Context)
@@ -603,8 +615,10 @@ void FRCPassPostProcessBasePassAO::Process(FRenderingCompositePassContext& Conte
 
 	SetGlobalBoundShaderState(BoundShaderState, GFilterVertexDeclaration.VertexDeclarationRHI, *VertexShader, *PixelShader);
 
-	VertexShader->SetParameters(Context);
-	PixelShader->SetParameters(Context, GSceneRenderTargets.GetBufferSizeXY());
+	//@todo-rco: RHIPacketList
+	FRHICommandList* RHICmdList = nullptr;
+	VertexShader->SetParameters(RHICmdList, Context);
+	PixelShader->SetParameters(RHICmdList, Context, GSceneRenderTargets.GetBufferSizeXY());
 
 	// Draw a quad mapping scene color to the view's render target
 	DrawRectangle( 
diff --git a/Engine/Source/Runtime/Renderer/Private/CompositionLighting/PostProcessAmbientOcclusion.h b/Engine/Source/Runtime/Renderer/Private/CompositionLighting/PostProcessAmbientOcclusion.h
--- a/Engine/Source/Runtime/Renderer/Private/CompositionLighting/PostProcessAmbientOcclusion.h
+++ b/Engine/Source/Runtime/Renderer/Private/CompositionLighting/PostProcessAmbientOcclusion.h
@@ -26,7 +26,7 @@ private:
 	bool IsInitialPass() const;
 
 	template <uint32 bInitialSetup>
-	void SetShaderSetupTempl(const FRenderingCompositePassContext& Context);
+	void SetShaderSetupTempl(FRHICommandList* RHICmdList, const FRenderingCompositePassContext& Context);
 };
 
 // ePId_Input0: defines the resolution we compute AO and provides the normal
@@ -78,7 +78,7 @@ public:
 	}
 
 	template<typename ShaderRHIParamRef>
-	void Set(const FSceneView& View, const ShaderRHIParamRef ShaderRHI, FIntPoint InputTextureSize) const
+	void Set(FRHICommandList* RHICmdList, const FSceneView& View, const ShaderRHIParamRef ShaderRHI, FIntPoint InputTextureSize) const
 	{
 		const FFinalPostProcessSettings& Settings = View.FinalPostProcessSettings;
 
@@ -138,7 +138,7 @@ public:
 		Value[3] = FVector4(0, 0, StaticFraction, InvTanHalfFov);
 		Value[4] = FVector4(InvFadeRadius, -(Settings.AmbientOcclusionFadeDistance - FadeRadius) * InvFadeRadius, 0, 0);
 
-		SetShaderValueArray(ShaderRHI, ScreenSpaceAOandSSRShaderParams, Value, 5);
+		SetShaderValueArray(RHICmdList, ShaderRHI, ScreenSpaceAOandSSRShaderParams, Value, 5);
 	}
 
 	friend FArchive& operator<<(FArchive& Ar, FScreenSpaceAOandSSRShaderParameters& This);
@@ -154,7 +154,7 @@ public:
 
 	void Bind(const FShaderParameterMap& ParameterMap);
 
-	void Set(const FSceneView& View, const FPixelShaderRHIParamRef ShaderRHI) const;
+	void Set(FRHICommandList* RHICmdList, const FSceneView& View, const FPixelShaderRHIParamRef ShaderRHI) const;
 
 	friend FArchive& operator<<(FArchive& Ar, FCameraMotionParameters& This);
 
diff --git a/Engine/Source/Runtime/Renderer/Private/CompositionLighting/PostProcessDeferredDecals.cpp b/Engine/Source/Runtime/Renderer/Private/CompositionLighting/PostProcessDeferredDecals.cpp
--- a/Engine/Source/Runtime/Renderer/Private/CompositionLighting/PostProcessDeferredDecals.cpp
+++ b/Engine/Source/Runtime/Renderer/Private/CompositionLighting/PostProcessDeferredDecals.cpp
@@ -314,6 +314,12 @@ public:
 		return IsFeatureLevelSupported(Platform, ERHIFeatureLevel::SM4); 
 	}
 
+	static void ModifyCompilationEnvironment(EShaderPlatform Platform, FShaderCompilerEnvironment& OutEnvironment)
+	{
+		FGlobalShader::ModifyCompilationEnvironment(Platform,OutEnvironment);
+		FDeferredPixelShaderParameters::ModifyCompilationEnvironment(Platform,OutEnvironment);
+	}
+
 	FStencilDecalMaskPS(const ShaderMetaType::CompiledShaderInitializerType& Initializer):
 		FGlobalShader(Initializer)
 	{
@@ -321,12 +327,12 @@ public:
 	}
 	FStencilDecalMaskPS() {}
 
-	void SetParameters(const FSceneView& View)
+	void SetParameters(FRHICommandList* RHICmdList, const FSceneView& View)
 	{
 		const FPixelShaderRHIParamRef ShaderRHI = GetPixelShader();
 
-		FGlobalShader::SetParameters(ShaderRHI, View);
-		DeferredParameters.Set(ShaderRHI, View);
+		FGlobalShader::SetParameters(RHICmdList, ShaderRHI, View);
+		DeferredParameters.Set(RHICmdList, ShaderRHI, View);
 	}
 
 	virtual bool Serialize(FArchive& Ar)
@@ -361,7 +367,9 @@ void StencilDecalMask(const FSceneView& View)
 	
 	SetGlobalBoundShaderState(StencilDecalMaskBoundShaderState, GFilterVertexDeclaration.VertexDeclarationRHI, *ScreenVertexShader, *PixelShader);
 
-	PixelShader->SetParameters(View);
+	//@todo-rco: RHIPacketList
+	FRHICommandList* RHICmdList = nullptr;
+	PixelShader->SetParameters(RHICmdList, View);
 
 	DrawRectangle( 
 		0, 0, 
@@ -399,10 +407,10 @@ public:
 		FrustumComponentToClip.Bind(Initializer.ParameterMap, TEXT("FrustumComponentToClip"));
 	}
 
-	void SetParameters(const FSceneView& View, const FMatrix& InFrustumComponentToClip)
+	void SetParameters(FRHICommandList* RHICmdList, const FSceneView& View, const FMatrix& InFrustumComponentToClip)
 	{
-		FMaterialShader::SetParameters(GetVertexShader(), View);
-		SetShaderValue(GetVertexShader(), FrustumComponentToClip, InFrustumComponentToClip);
+		FMaterialShader::SetParameters(RHICmdList, GetVertexShader(), View);
+		SetShaderValue(RHICmdList, GetVertexShader(), FrustumComponentToClip, InFrustumComponentToClip);
 	}
 
 	virtual bool Serialize(FArchive& Ar)
@@ -465,11 +473,11 @@ public:
 		DecalToWorld.Bind(Initializer.ParameterMap,TEXT("DecalToWorld"));
 	}
 
-	void SetParameters(const FSceneView& View, const FMaterialRenderProxy* MaterialProxy, const FDeferredDecalProxy& DecalProxy)
+	void SetParameters(FRHICommandList* RHICmdList, const FSceneView& View, const FMaterialRenderProxy* MaterialProxy, const FDeferredDecalProxy& DecalProxy)
 	{
 		const FPixelShaderRHIParamRef ShaderRHI = GetPixelShader();
 
-		FMaterialShader::SetParameters(ShaderRHI, MaterialProxy, *MaterialProxy->GetMaterial(View.GetFeatureLevel()), View, true, ESceneRenderTargetsMode::SetTextures);
+		FMaterialShader::SetParameters(RHICmdList, ShaderRHI, MaterialProxy, *MaterialProxy->GetMaterial(View.GetFeatureLevel()), View, true, ESceneRenderTargetsMode::SetTextures);
 
 		FTransform ComponentTrans = DecalProxy.ComponentTrans;
 
@@ -489,7 +497,7 @@ public:
 					FPlane(0,0,View.ViewMatrices.ProjMatrix.M[3][2],0)
 				) * View.InvViewProjectionMatrix * WorldToComponent;
 
-			SetShaderValue(ShaderRHI, ScreenToDecal, ScreenToDecalValue);
+			SetShaderValue(RHICmdList, ShaderRHI, ScreenToDecal, ScreenToDecalValue);
 		}
 
 		// Set the transform from light space to world space (only for normals)
@@ -500,7 +508,7 @@ public:
 			// 1,1,1 requires no scale
 			//			DecalToWorldValue = DecalToWorldValue.GetScaled(GDefaultDecalSize);
 
-			SetShaderValue(ShaderRHI, DecalToWorld, DecalToWorldValue);
+			SetShaderValue(RHICmdList, ShaderRHI, DecalToWorld, DecalToWorldValue);
 		}
 	}
 
@@ -596,13 +604,15 @@ private:
 static TGlobalResource<FUnitCubeVertexBuffer> GUnitCubeVertexBuffer;
 static TGlobalResource<FUnitCubeIndexBuffer> GUnitCubeIndexBuffer;
 
-void SetShader(const FRenderingCompositePassContext& Context, const FTransientDecalRenderData& DecalData, FShader* VertexShader)
+void SetShader(FRHICommandList* RHICmdList, const FRenderingCompositePassContext& Context, const FTransientDecalRenderData& DecalData, FShader* VertexShader)
 {
 	const FSceneView& View = Context.View;
 	
 	const FMaterialShaderMap* MaterialShaderMap = DecalData.MaterialResource->GetRenderingThreadShaderMap();
 	FDeferredDecalPS* PixelShader = MaterialShaderMap->GetShader<FDeferredDecalPS>();
 
+	check(!RHICmdList);
+
 	// This was cached but when changing the material (e.g. editor) it wasn't updated.
 	// This will change with upcoming multi threaded rendering changes.
 	FBoundShaderStateRHIRef BoundShaderState;
@@ -616,10 +626,10 @@ void SetShader(const FRenderingCompositePassContext& Context, const FTransientDe
 
 	RHISetBoundShaderState(BoundShaderState);
 
-	PixelShader->SetParameters(View, DecalData.MaterialProxy, *DecalData.DecalProxy);
+	PixelShader->SetParameters(RHICmdList, View, DecalData.MaterialProxy, *DecalData.DecalProxy);
 }
 
-bool RenderPreStencil(FRenderingCompositePassContext& Context, const FMaterialShaderMap* MaterialShaderMap, const FMatrix& ComponentToWorldMatrix, const FMatrix& FrustumComponentToClip)
+bool RenderPreStencil(FRHICommandList* RHICmdList, FRenderingCompositePassContext& Context, const FMaterialShaderMap* MaterialShaderMap, const FMatrix& ComponentToWorldMatrix, const FMatrix& FrustumComponentToClip)
 {
 	SCOPED_DRAW_EVENT(RenderPreStencil, DEC_SCENE_ITEMS);
 
@@ -644,6 +654,8 @@ bool RenderPreStencil(FRenderingCompositePassContext& Context, const FMaterialSh
 
 	FDeferredDecalVS* VertexShader = MaterialShaderMap->GetShader<FDeferredDecalVS>();
 	
+	check(!RHICmdList);
+
 	// This was cached but when changing the material (e.g. editor) it wasn't updated.
 	// This will change with upcoming multi threaded rendering changes.
 	FBoundShaderStateRHIRef BoundShaderState;
@@ -653,7 +665,7 @@ bool RenderPreStencil(FRenderingCompositePassContext& Context, const FMaterialSh
 
 	RHISetBoundShaderState(BoundShaderState);
 
-	VertexShader->SetParameters(View, FrustumComponentToClip);
+	VertexShader->SetParameters(RHICmdList, View, FrustumComponentToClip);
 
 	// Set states, the state cache helps us avoiding redundant sets
 	RHISetRasterizerState(TStaticRasterizerState<FM_Solid,CM_None>::GetRHI());
@@ -696,6 +708,9 @@ FRCPassPostProcessDeferredDecals::FRCPassPostProcessDeferredDecals(uint32 InRend
 
 void FRCPassPostProcessDeferredDecals::Process(FRenderingCompositePassContext& Context)
 {
+	//@todo-rco: RHIPacketList
+	FRHICommandList* RHICmdList = nullptr;
+
 	bool bDBuffer = IsDBufferEnabled();
 	float bDecalPreStencil = CVarStencilSizeThreshold.GetValueOnRenderThread() >= 0;
 
@@ -897,7 +912,7 @@ void FRCPassPostProcessDeferredDecals::Process(FRenderingCompositePassContext& C
 			{
 				if(bDecalPreStencil)
 				{
-					bThisDecalUsesStencil = RenderPreStencil(Context, MaterialShaderMap, ComponentToWorldMatrix, FrustumComponentToClip);
+					bThisDecalUsesStencil = RenderPreStencil(RHICmdList, Context, MaterialShaderMap, ComponentToWorldMatrix, FrustumComponentToClip);
 					WasInsideDecal = -1;
 					LastDecalBlendMode = -1;
 				}
@@ -973,9 +988,9 @@ void FRCPassPostProcessDeferredDecals::Process(FRenderingCompositePassContext& C
 
 			{
 				FDeferredDecalVS* VertexShader = MaterialShaderMap->GetShader<FDeferredDecalVS>();
-				SetShader(Context, DecalData, VertexShader);
+				SetShader(RHICmdList, Context, DecalData, VertexShader);
 
-				VertexShader->SetParameters(View, FrustumComponentToClip);
+				VertexShader->SetParameters(RHICmdList, View, FrustumComponentToClip);
 
 				const int32 IsInsideDecal = ((FVector)View.ViewMatrices.ViewOrigin - ComponentToWorldMatrix.GetOrigin()).SizeSquared() < FMath::Square(ConservativeRadius * 1.05f + View.NearClippingDistance * 2.0f) + ( bThisDecalUsesStencil ) ? 2 : 0;
 				if ( WasInsideDecal != IsInsideDecal )
@@ -1047,7 +1062,7 @@ void FRCPassPostProcessDeferredDecals::Process(FRenderingCompositePassContext& C
 					}
 				}
 
-				SetShader(Context, DecalData, VertexShader);
+				SetShader(RHICmdList, Context, DecalData, VertexShader);
 
 				RHIDrawIndexedPrimitive(GUnitCubeIndexBuffer.IndexBufferRHI, PT_TriangleList, 0, 0, 8, 0, GUnitCubeIndexBuffer.GetIndexCount() / 3, 0);
 			}
diff --git a/Engine/Source/Runtime/Renderer/Private/CompositionLighting/PostProcessLpvIndirect.cpp b/Engine/Source/Runtime/Renderer/Private/CompositionLighting/PostProcessLpvIndirect.cpp
--- a/Engine/Source/Runtime/Renderer/Private/CompositionLighting/PostProcessLpvIndirect.cpp
+++ b/Engine/Source/Runtime/Renderer/Private/CompositionLighting/PostProcessLpvIndirect.cpp
@@ -67,6 +67,7 @@ public:
 	}
 
 	void SetParameters(	
+		FRHICommandList* RHICmdList, 
 #if LPV_VOLUME_TEXTURE
 		FTextureRHIParamRef* LpvBufferSRVsIn, 
 #else
@@ -78,7 +79,9 @@ public:
 	{
 		const FPixelShaderRHIParamRef ShaderRHI = GetPixelShader();
 
-		SetUniformBufferParameter( ShaderRHI, GetUniformBufferParameter<FLpvReadUniformBufferParameters>(), LpvUniformBuffer );
+		check(!RHICmdList);
+
+		SetUniformBufferParameter(RHICmdList, ShaderRHI, GetUniformBufferParameter<FLpvReadUniformBufferParameters>(), LpvUniformBuffer );
 
 #if LPV_VOLUME_TEXTURE
 		for ( int i=0; i<7; i++ )
@@ -86,7 +89,7 @@ public:
 			if ( LpvBufferSRVParameters[i].IsBound() )
 			{
 				RHISetShaderTexture( ShaderRHI, LpvBufferSRVParameters[i].GetBaseIndex(), LpvBufferSRVsIn[i] );
-				SetTextureParameter( ShaderRHI, LpvBufferSRVParameters[i], LpvVolumeTextureSampler, TStaticSamplerState<SF_Bilinear,AM_Clamp,AM_Clamp,AM_Clamp>::GetRHI(), LpvBufferSRVsIn[i] );
+				SetTextureParameter(RHICmdList, ShaderRHI, LpvBufferSRVParameters[i], LpvVolumeTextureSampler, TStaticSamplerState<SF_Bilinear,AM_Clamp,AM_Clamp,AM_Clamp>::GetRHI(), LpvBufferSRVsIn[i] );
 			}
 		}
 #else
@@ -95,10 +98,10 @@ public:
 			RHISetShaderResourceViewParameter( ShaderRHI, LpvBufferSRV.GetBaseIndex(), LpvBufferSRVIn );
 		}
 #endif
-		FGlobalShader::SetParameters(ShaderRHI, Context.View);
-		PostprocessParameter.SetPS(ShaderRHI, Context, TStaticSamplerState<SF_Bilinear,AM_Clamp,AM_Clamp,AM_Clamp>::GetRHI());
-		DeferredParameters.Set(ShaderRHI, Context.View);
-		SetTextureParameter(ShaderRHI, PreIntegratedGF, PreIntegratedGFSampler, TStaticSamplerState<SF_Bilinear, AM_Clamp, AM_Clamp, AM_Clamp>::GetRHI(), GSystemTextures.PreintegratedGF->GetRenderTargetItem().ShaderResourceTexture);
+		FGlobalShader::SetParameters(RHICmdList, ShaderRHI, Context.View);
+		PostprocessParameter.SetPS(RHICmdList, ShaderRHI, Context, TStaticSamplerState<SF_Bilinear,AM_Clamp,AM_Clamp,AM_Clamp>::GetRHI());
+		DeferredParameters.Set(RHICmdList, ShaderRHI, Context.View);
+		SetTextureParameter(RHICmdList, ShaderRHI, PreIntegratedGF, PreIntegratedGFSampler, TStaticSamplerState<SF_Bilinear,AM_Clamp,AM_Clamp,AM_Clamp>::GetRHI(), GSystemTextures.PreintegratedGF->GetRenderTargetItem().ShaderResourceTexture );
 	}
 	
 	// FShader interface.
@@ -150,6 +153,10 @@ void FRCPassPostProcessLpvIndirect::Process(FRenderingCompositePassContext& Cont
 		DestColorRenderTarget.TargetableTexture,
 	};
 
+	//@todo-rco: RHIPacketList
+	FRHICommandList* RHICmdList = nullptr;
+	check(!RHICmdList);
+
 	// Set the view family's render target/viewport.
 	RHISetRenderTargets(1, RenderTargets, GSceneRenderTargets.GetSceneDepthSurface(), 0, NULL);
 
@@ -196,7 +203,7 @@ void FRCPassPostProcessLpvIndirect::Process(FRenderingCompositePassContext& Cont
 	FLpvReadUniformBufferRef LpvReadUniformBuffer;
 
 	LpvReadUniformBufferParams = Lpv->GetReadUniformBufferParams();
-	LpvReadUniformBuffer = FLpvReadUniformBufferRef::CreateUniformBufferImmediate( LpvReadUniformBufferParams, UniformBuffer_SingleUse ); 
+	LpvReadUniformBuffer = FLpvReadUniformBufferRef::CreateUniformBufferImmediate( LpvReadUniformBufferParams, UniformBuffer_SingleDraw ); 
 
 #if LPV_VOLUME_TEXTURE
 	FTextureRHIParamRef LpvBufferSrvs[7];
@@ -204,10 +211,10 @@ void FRCPassPostProcessLpvIndirect::Process(FRenderingCompositePassContext& Cont
 	{
 		LpvBufferSrvs[i] = Lpv->GetLpvBufferSrv(i);
 	}
-	PixelShader->SetParameters(LpvBufferSrvs, LpvReadUniformBuffer, Context );
+	PixelShader->SetParameters(RHICmdList, LpvBufferSrvs, LpvReadUniformBuffer, Context );
 #else
 	FShaderResourceViewRHIParamRef LpvBufferSrv = Lpv->GetLpvBufferSrv();
-	PixelShader->SetParameters(LpvBufferSrv, LpvReadUniformBuffer, Context );
+	PixelShader->SetParameters(RHICmdList, LpvBufferSrv, LpvReadUniformBuffer, Context );
 #endif
 
 	// Draw a quad mapping scene color to the view's render target
diff --git a/Engine/Source/Runtime/Renderer/Private/CompositionLighting/PostProcessPassThrough.cpp b/Engine/Source/Runtime/Renderer/Private/CompositionLighting/PostProcessPassThrough.cpp
--- a/Engine/Source/Runtime/Renderer/Private/CompositionLighting/PostProcessPassThrough.cpp
+++ b/Engine/Source/Runtime/Renderer/Private/CompositionLighting/PostProcessPassThrough.cpp
@@ -45,9 +45,11 @@ public:
 	{
 		const FPixelShaderRHIParamRef ShaderRHI = GetPixelShader();
 
-		FGlobalShader::SetParameters(ShaderRHI, Context.View);
+		//@todo-rco: RHIPacketList
+		FRHICommandList* RHICmdList = nullptr;
+		FGlobalShader::SetParameters(RHICmdList, ShaderRHI, Context.View);
 
-		PostprocessParameter.SetPS(ShaderRHI, Context, TStaticSamplerState<SF_Point,AM_Clamp,AM_Clamp,AM_Clamp>::GetRHI());
+		PostprocessParameter.SetPS(RHICmdList, ShaderRHI, Context, TStaticSamplerState<SF_Point,AM_Clamp,AM_Clamp,AM_Clamp>::GetRHI());
 	}
 };
 
@@ -121,7 +123,9 @@ void FRCPassPostProcessPassThrough::Process(FRenderingCompositePassContext& Cont
 
 	SetGlobalBoundShaderState(BoundShaderState, GFilterVertexDeclaration.VertexDeclarationRHI, *VertexShader, *PixelShader);
 
-	VertexShader->SetParameters(Context);
+	//@todo-rco: RHIPacketList
+	FRHICommandList* RHICmdList = nullptr;
+	VertexShader->SetParameters(RHICmdList, Context);
 	PixelShader->SetParameters(Context);
 
 	// Draw a quad mapping scene color to the view's render target
diff --git a/Engine/Source/Runtime/Renderer/Private/CustomDepthRendering.cpp b/Engine/Source/Runtime/Renderer/Private/CustomDepthRendering.cpp
--- a/Engine/Source/Runtime/Renderer/Private/CustomDepthRendering.cpp
+++ b/Engine/Source/Runtime/Renderer/Private/CustomDepthRendering.cpp
@@ -16,7 +16,7 @@
 	FCustomDepthPrimSet
 -----------------------------------------------------------------------------*/
 
-bool FCustomDepthPrimSet::DrawPrims(const FViewInfo* ViewInfo,bool bInitializeOffsets)
+bool FCustomDepthPrimSet::DrawPrims(FRHICommandList* RHICmdList, const FViewInfo* ViewInfo,bool bInitializeOffsets)
 {
 	bool bDirty=false;
 
@@ -49,6 +49,7 @@ bool FCustomDepthPrimSet::DrawPrims(const FViewInfo* ViewInfo,bool bInitializeOf
 						if (ViewInfo->StaticMeshVisibilityMap[StaticMesh.Id])
 						{
 							bDirty |= FDepthDrawingPolicyFactory::DrawStaticMesh(
+								RHICmdList, 
 								*ViewInfo,
 								FDepthDrawingPolicyFactory::ContextType(DDM_AllOccluders),
 								StaticMesh,
diff --git a/Engine/Source/Runtime/Renderer/Private/CustomDepthRendering.h b/Engine/Source/Runtime/Renderer/Private/CustomDepthRendering.h
--- a/Engine/Source/Runtime/Renderer/Private/CustomDepthRendering.h
+++ b/Engine/Source/Runtime/Renderer/Private/CustomDepthRendering.h
@@ -18,7 +18,7 @@ public:
 	* @param ViewInfo - current view used to draw items
 	* @return true if anything was drawn
 	*/
-	bool DrawPrims(const class FViewInfo* ViewInfo,bool bInitializeOffsets);
+	bool DrawPrims(FRHICommandList* RHICmdList, const class FViewInfo* ViewInfo,bool bInitializeOffsets);
 
 	/**
 	* Add a new primitive to the list of prims
diff --git a/Engine/Source/Runtime/Renderer/Private/DeferredShadingRenderer.cpp b/Engine/Source/Runtime/Renderer/Private/DeferredShadingRenderer.cpp
--- a/Engine/Source/Runtime/Renderer/Private/DeferredShadingRenderer.cpp
+++ b/Engine/Source/Runtime/Renderer/Private/DeferredShadingRenderer.cpp
@@ -154,6 +154,9 @@ void FDeferredShadingSceneRenderer::ClearGBufferAtMaxZ()
 		break;
 	}
 
+	//@todo-rco: RHIPacketList
+	FRHICommandList* RHICmdList = nullptr;
+
 	SetGlobalBoundShaderState(GClearMRTBoundShaderState[NumActiveRenderTargets - 1], GetVertexDeclarationFVector4(), *VertexShader, PixelShader);
 
 	// Opaque rendering, depth test but no depth writes
@@ -173,14 +176,14 @@ void FDeferredShadingSceneRenderer::ClearGBufferAtMaxZ()
 		RHISetViewport(View.ViewRect.Min.X, View.ViewRect.Min.Y, 0, View.ViewRect.Max.X, View.ViewRect.Max.Y, 1);
 
 		// Setup PS
-		SetShaderValueArray(PixelShader->GetPixelShader(),PixelShader->ColorParameter, ClearColors, NumActiveRenderTargets);
+		SetShaderValueArray(RHICmdList, PixelShader->GetPixelShader(),PixelShader->ColorParameter, ClearColors, NumActiveRenderTargets);
 
 		// Render quad
 		RHIDrawPrimitiveUP(PT_TriangleStrip, 2, ClearQuadVertices, sizeof(ClearQuadVertices[0]) );
 	}
 }
 
-bool FDeferredShadingSceneRenderer::RenderBasePassStaticDataMasked(FViewInfo& View)
+bool FDeferredShadingSceneRenderer::RenderBasePassStaticDataMasked(FRHICommandList* RHICmdList, FViewInfo& View)
 {
 	bool bDirty = false;
 
@@ -189,23 +192,23 @@ bool FDeferredShadingSceneRenderer::RenderBasePassStaticDataMasked(FViewInfo& Vi
 		FScene::EBasePassDrawListType MaskedDrawType = FScene::EBasePass_Masked;
 		{
 			SCOPED_DRAW_EVENT(StaticMaskedNoLightmap, DEC_SCENE_ITEMS);
-			bDirty |= Scene->BasePassNoLightMapDrawList[MaskedDrawType].DrawVisible(View,View.StaticMeshVisibilityMap,View.StaticMeshBatchVisibility);
-			bDirty |= Scene->BasePassSimpleDynamicLightingDrawList[MaskedDrawType].DrawVisible(View,View.StaticMeshVisibilityMap,View.StaticMeshBatchVisibility);
-			bDirty |= Scene->BasePassCachedVolumeIndirectLightingDrawList[MaskedDrawType].DrawVisible(View,View.StaticMeshVisibilityMap,View.StaticMeshBatchVisibility);
-			bDirty |= Scene->BasePassCachedPointIndirectLightingDrawList[MaskedDrawType].DrawVisible(View,View.StaticMeshVisibilityMap,View.StaticMeshBatchVisibility);
+			bDirty |= Scene->BasePassNoLightMapDrawList[MaskedDrawType].DrawVisible(RHICmdList, View,View.StaticMeshVisibilityMap,View.StaticMeshBatchVisibility);
+			bDirty |= Scene->BasePassSimpleDynamicLightingDrawList[MaskedDrawType].DrawVisible(RHICmdList, View,View.StaticMeshVisibilityMap,View.StaticMeshBatchVisibility);
+			bDirty |= Scene->BasePassCachedVolumeIndirectLightingDrawList[MaskedDrawType].DrawVisible(RHICmdList, View,View.StaticMeshVisibilityMap,View.StaticMeshBatchVisibility);
+			bDirty |= Scene->BasePassCachedPointIndirectLightingDrawList[MaskedDrawType].DrawVisible(RHICmdList, View,View.StaticMeshVisibilityMap,View.StaticMeshBatchVisibility);
 		}
 		{
 			SCOPED_DRAW_EVENT(StaticMaskedLightmapped, DEC_SCENE_ITEMS);
-			bDirty |= Scene->BasePassHighQualityLightMapDrawList[MaskedDrawType].DrawVisible(View,View.StaticMeshVisibilityMap,View.StaticMeshBatchVisibility);
-			bDirty |= Scene->BasePassDistanceFieldShadowMapLightMapDrawList[MaskedDrawType].DrawVisible(View,View.StaticMeshVisibilityMap,View.StaticMeshBatchVisibility);
-			bDirty |= Scene->BasePassLowQualityLightMapDrawList[MaskedDrawType].DrawVisible(View,View.StaticMeshVisibilityMap,View.StaticMeshBatchVisibility);
+			bDirty |= Scene->BasePassHighQualityLightMapDrawList[MaskedDrawType].DrawVisible(RHICmdList, View,View.StaticMeshVisibilityMap,View.StaticMeshBatchVisibility);
+			bDirty |= Scene->BasePassDistanceFieldShadowMapLightMapDrawList[MaskedDrawType].DrawVisible(RHICmdList, View,View.StaticMeshVisibilityMap,View.StaticMeshBatchVisibility);
+			bDirty |= Scene->BasePassLowQualityLightMapDrawList[MaskedDrawType].DrawVisible(RHICmdList, View,View.StaticMeshVisibilityMap,View.StaticMeshBatchVisibility);
 		}
 	}
 
 	return bDirty;
 }
 
-bool FDeferredShadingSceneRenderer::RenderBasePassStaticDataDefault(FViewInfo& View)
+bool FDeferredShadingSceneRenderer::RenderBasePassStaticDataDefault(FRHICommandList* RHICmdList, FViewInfo& View)
 {
 	bool bDirty = false;
 
@@ -213,16 +216,16 @@ bool FDeferredShadingSceneRenderer::RenderBasePassStaticDataDefault(FViewInfo& V
 		FScene::EBasePassDrawListType OpaqueDrawType = FScene::EBasePass_Default;
 		{
 			SCOPED_DRAW_EVENT(StaticOpaqueNoLightmap, DEC_SCENE_ITEMS);
-			bDirty |= Scene->BasePassNoLightMapDrawList[OpaqueDrawType].DrawVisible(View,View.StaticMeshVisibilityMap,View.StaticMeshBatchVisibility);
-			bDirty |= Scene->BasePassSimpleDynamicLightingDrawList[OpaqueDrawType].DrawVisible(View,View.StaticMeshVisibilityMap,View.StaticMeshBatchVisibility);
-			bDirty |= Scene->BasePassCachedVolumeIndirectLightingDrawList[OpaqueDrawType].DrawVisible(View,View.StaticMeshVisibilityMap,View.StaticMeshBatchVisibility);
-			bDirty |= Scene->BasePassCachedPointIndirectLightingDrawList[OpaqueDrawType].DrawVisible(View,View.StaticMeshVisibilityMap,View.StaticMeshBatchVisibility);
+			bDirty |= Scene->BasePassNoLightMapDrawList[OpaqueDrawType].DrawVisible(RHICmdList, View,View.StaticMeshVisibilityMap,View.StaticMeshBatchVisibility);
+			bDirty |= Scene->BasePassSimpleDynamicLightingDrawList[OpaqueDrawType].DrawVisible(RHICmdList, View,View.StaticMeshVisibilityMap,View.StaticMeshBatchVisibility);
+			bDirty |= Scene->BasePassCachedVolumeIndirectLightingDrawList[OpaqueDrawType].DrawVisible(RHICmdList, View,View.StaticMeshVisibilityMap,View.StaticMeshBatchVisibility);
+			bDirty |= Scene->BasePassCachedPointIndirectLightingDrawList[OpaqueDrawType].DrawVisible(RHICmdList, View,View.StaticMeshVisibilityMap,View.StaticMeshBatchVisibility);
 		}
 		{
 			SCOPED_DRAW_EVENT(StaticOpaqueLightmapped, DEC_SCENE_ITEMS);
-			bDirty |= Scene->BasePassHighQualityLightMapDrawList[OpaqueDrawType].DrawVisible(View,View.StaticMeshVisibilityMap,View.StaticMeshBatchVisibility);
-			bDirty |= Scene->BasePassDistanceFieldShadowMapLightMapDrawList[OpaqueDrawType].DrawVisible(View,View.StaticMeshVisibilityMap,View.StaticMeshBatchVisibility);
-			bDirty |= Scene->BasePassLowQualityLightMapDrawList[OpaqueDrawType].DrawVisible(View,View.StaticMeshVisibilityMap,View.StaticMeshBatchVisibility);
+			bDirty |= Scene->BasePassHighQualityLightMapDrawList[OpaqueDrawType].DrawVisible(RHICmdList, View,View.StaticMeshVisibilityMap,View.StaticMeshBatchVisibility);
+			bDirty |= Scene->BasePassDistanceFieldShadowMapLightMapDrawList[OpaqueDrawType].DrawVisible(RHICmdList, View,View.StaticMeshVisibilityMap,View.StaticMeshBatchVisibility);
+			bDirty |= Scene->BasePassLowQualityLightMapDrawList[OpaqueDrawType].DrawVisible(RHICmdList, View,View.StaticMeshVisibilityMap,View.StaticMeshBatchVisibility);
 		}
 	}
 
@@ -255,7 +258,7 @@ void FDeferredShadingSceneRenderer::SortBasePassStaticData(FVector ViewPosition)
 *
 * @return true if anything was rendered to scene color
 */
-bool FDeferredShadingSceneRenderer::RenderBasePassStaticData(FViewInfo& View)
+bool FDeferredShadingSceneRenderer::RenderBasePassStaticData(FRHICommandList* RHICmdList, FViewInfo& View)
 {
 	bool bDirty = false;
 
@@ -267,16 +270,16 @@ bool FDeferredShadingSceneRenderer::RenderBasePassStaticData(FViewInfo& View)
 	// rendering the default opaque geometry afterward.
 	if (EarlyZPassMode != DDM_None)
 	{
-		bDirty |= RenderBasePassStaticDataMasked(View);
-		bDirty |= RenderBasePassStaticDataDefault(View);
+		bDirty |= RenderBasePassStaticDataMasked(RHICmdList, View);
+		bDirty |= RenderBasePassStaticDataDefault(RHICmdList, View);
 	}
 	else
 	{
 		// Otherwise, in the case where we're not using a depth-only pre-pass, there
 		// is an advantage to rendering default opaque first to help cull the more
 		// expensive masked geometry.
-		bDirty |= RenderBasePassStaticDataDefault(View);
-		bDirty |= RenderBasePassStaticDataMasked(View);
+		bDirty |= RenderBasePassStaticDataDefault(RHICmdList, View);
+		bDirty |= RenderBasePassStaticDataMasked(RHICmdList, View);
 	}
 
 	return bDirty;
@@ -322,7 +325,15 @@ bool FDeferredShadingSceneRenderer::RenderBasePass(FViewInfo& View)
 	bool bDirty=0;
 
 	// Render the base pass static data
-	bDirty |= RenderBasePassStaticData(View);
+	{
+		//@todo-rco: Very temp code!!!
+		SCOPE_CYCLE_COUNTER(STAT_RHICounterTEMP);
+		static IConsoleVariable* RHICmdListCVar = IConsoleManager::Get().FindConsoleVariable(TEXT("r.RHICmd"));
+		bool bUseRHICmdList = RHICmdListCVar->GetInt() != 0;
+		auto* RHICmdList = bUseRHICmdList ? GRHICommandList.CreateList() : nullptr;
+		bDirty |= RenderBasePassStaticData(RHICmdList, View);
+		GRHICommandList.ExecuteAndFreeList(RHICmdList);
+	}
 
 	{
 		SCOPE_CYCLE_COUNTER(STAT_DynamicPrimitiveDrawTime);
@@ -503,6 +514,10 @@ void FDeferredShadingSceneRenderer::Render()
 		const bool bIsViewFrozen = Views[0].State && ((FSceneViewState*)Views[0].State)->bIsFrozen;
 		const bool bHasViewParent = Views[0].State && ((FSceneViewState*)Views[0].State)->HasViewParent();
 	#endif
+
+	//@todo-rco: RHIPacketList
+	FRHICommandList* RHICmdList = nullptr;
+
 	const bool bIsOcclusionTesting = DoOcclusionQueries() && (!bIsWireframe || bIsViewFrozen || bHasViewParent);
 
 	// Dynamic vertex and index buffers need to be committed before rendering.
@@ -519,7 +534,7 @@ void FDeferredShadingSceneRenderer::Render()
 
 	// Draw the scene pre-pass / early z pass, populating the scene depth buffer and HiZ
 	RenderPrePass();
-
+	
 	GSceneRenderTargets.AllocGBufferTargets();
 	
 	// Clear scene color buffer if necessary.
@@ -534,7 +549,7 @@ void FDeferredShadingSceneRenderer::Render()
 	// Clear LPVs for all views
 	if ( IsFeatureLevelSupported(GRHIShaderPlatform, ERHIFeatureLevel::SM5) )
 	{
-		ClearLPVs();
+		ClearLPVs(RHICmdList);
 	}
 
 	// only temporarily available after early z pass and until base pass
@@ -642,13 +657,13 @@ void FDeferredShadingSceneRenderer::Render()
 		RenderDynamicSkyLighting();
 
 		// Clear the translucent lighting volumes before we accumulate
-		ClearTranslucentVolumeLighting();
+		ClearTranslucentVolumeLighting(RHICmdList);
 
-		RenderLights();
+		RenderLights(RHICmdList);
 
-		InjectAmbientCubemapTranslucentVolumeLighting();
+		InjectAmbientCubemapTranslucentVolumeLighting(RHICmdList);
 
-		CompositeIndirectTranslucentVolumeLighting();
+		CompositeIndirectTranslucentVolumeLighting(RHICmdList);
 
 		// Filter the translucency lighting volume now that it is complete
 		FilterTranslucentVolumeLighting();
@@ -660,7 +675,7 @@ void FDeferredShadingSceneRenderer::Render()
 		}
 
 		// Render reflections that only operate on opaque pixels
-		RenderDeferredReflections();
+		RenderDeferredReflections(RHICmdList);
 
 		// Post-lighting composition lighting stage
 		// e.g. ambient cubemaps, ambient occlusion, LPV indirect
@@ -731,7 +746,7 @@ void FDeferredShadingSceneRenderer::Render()
 		if(ViewFamily.EngineShowFlags.Refraction)
 		{
 			// to apply refraction effect by distorting the scene color
-			RenderDistortion();
+			RenderDistortion(RHICmdList);
 		}
 		RenderTranslucency();
 	}
@@ -787,6 +802,9 @@ bool FDeferredShadingSceneRenderer::RenderPrePass()
 
 	GSceneRenderTargets.BeginRenderingPrePass();
 
+	//@todo-rco: RHIPacketList
+	FRHICommandList* RHICmdList = nullptr;
+
 	// Clear the depth buffer.
 	// Note, this is a reversed Z depth surface, so 0.0f is the far plane.
 	RHIClear(false,FLinearColor::Black,true,0.0f,true,0, FIntRect());
@@ -812,19 +830,19 @@ bool FDeferredShadingSceneRenderer::RenderPrePass()
 				// vertex buffer to minimize vertex fetch bandwidth, which is
 				// often the bottleneck during the depth only pass.
 				SCOPED_DRAW_EVENT(PosOnlyOpaque, DEC_SCENE_ITEMS);
-				bDirty |= Scene->PositionOnlyDepthDrawList.DrawVisible(View,View.StaticMeshOccluderMap,View.StaticMeshBatchVisibility);
+				bDirty |= Scene->PositionOnlyDepthDrawList.DrawVisible(RHICmdList, View,View.StaticMeshOccluderMap,View.StaticMeshBatchVisibility);
 			}
 			{
 				// Draw opaque occluders, using double speed z where supported.
 				SCOPED_DRAW_EVENT(Opaque, DEC_SCENE_ITEMS);
-				bDirty |= Scene->DepthDrawList.DrawVisible(View,View.StaticMeshOccluderMap,View.StaticMeshBatchVisibility);
+				bDirty |= Scene->DepthDrawList.DrawVisible(RHICmdList, View,View.StaticMeshOccluderMap,View.StaticMeshBatchVisibility);
 			}
 
 			if(EarlyZPassMode >= DDM_AllOccluders)
 			{
 				// Draw opaque occluders with masked materials
 				SCOPED_DRAW_EVENT(Opaque, DEC_SCENE_ITEMS);
-				bDirty |= Scene->MaskedDepthDrawList.DrawVisible(View,View.StaticMeshOccluderMap,View.StaticMeshBatchVisibility);
+				bDirty |= Scene->MaskedDepthDrawList.DrawVisible(RHICmdList, View,View.StaticMeshOccluderMap,View.StaticMeshBatchVisibility);
 			}
 
 			// Draw the dynamic occluder primitives using a depth drawing policy.
@@ -880,6 +898,9 @@ bool FDeferredShadingSceneRenderer::RenderBasePass()
 {
 	bool bDirty = false;
 
+	//@todo-rco: RHIPacketList
+	FRHICommandList* RHICmdList = nullptr;
+
 	if(ViewFamily.EngineShowFlags.LightMapDensity && AllowDebugViewmodes())
 	{
 		// Override the base pass with the lightmap density pass if the viewmode is enabled.
@@ -919,7 +940,7 @@ bool FDeferredShadingSceneRenderer::RenderBasePass()
 	return bDirty;
 }
 
-void FDeferredShadingSceneRenderer::ClearLPVs()
+void FDeferredShadingSceneRenderer::ClearLPVs(FRHICommandList* RHICmdList)
 {
 	// clear light propagation volumes
 
@@ -937,7 +958,7 @@ void FDeferredShadingSceneRenderer::ClearLPVs()
 				SCOPED_DRAW_EVENT(ClearLPVs, DEC_SCENE_ITEMS);
 				SCOPE_CYCLE_COUNTER(STAT_UpdateLPVs);
 				LightPropagationVolume->InitSettings(Views[ViewIndex]);
-				LightPropagationVolume->Clear();
+				LightPropagationVolume->Clear(RHICmdList);
 			}
 		}
 	}
@@ -958,7 +979,8 @@ void FDeferredShadingSceneRenderer::PropagateLPVs()
 			{
 				SCOPED_DRAW_EVENT(UpdateLPVs, DEC_SCENE_ITEMS);
 				SCOPE_CYCLE_COUNTER(STAT_UpdateLPVs);
-				LightPropagationVolume->Propagate();
+				//@todo-rco: RHIPacketList
+				LightPropagationVolume->Propagate(nullptr);
 			}
 		}
 	}
@@ -985,13 +1007,13 @@ public:
 	}
 	FDownsampleSceneDepthPS() {}
 
-	void SetParameters(const FSceneView& View)
+	void SetParameters(FRHICommandList* RHICmdList, const FSceneView& View)
 	{
-		FGlobalShader::SetParameters(GetPixelShader(), View);
+		FGlobalShader::SetParameters(RHICmdList, GetPixelShader(), View);
 
 		// Used to remap view space Z (which is stored in scene color alpha) into post projection z and w so we can write z/w into the downsized depth buffer
 		const FVector2D ProjectionScaleBiasValue(View.ViewMatrices.ProjMatrix.M[2][2], View.ViewMatrices.ProjMatrix.M[3][2]);
-		SetShaderValue(GetPixelShader(), ProjectionScaleBias, ProjectionScaleBiasValue);
+		SetShaderValue(RHICmdList, GetPixelShader(), ProjectionScaleBias, ProjectionScaleBiasValue);
 
 		FIntPoint BufferSize = GSceneRenderTargets.GetBufferSizeXY();
 
@@ -1000,10 +1022,10 @@ public:
 
 		// Offsets of the four full resolution pixels corresponding with a low resolution pixel
 		const FVector4 Offsets01(0.0f, 0.0f, 1.0f / DownsampledBufferSizeX, 0.0f);
-		SetShaderValue(GetPixelShader(), SourceTexelOffsets01, Offsets01);
+		SetShaderValue(RHICmdList, GetPixelShader(), SourceTexelOffsets01, Offsets01);
 		const FVector4 Offsets23(0.0f, 1.0f / DownsampledBufferSizeY, 1.0f / DownsampledBufferSizeX, 1.0f / DownsampledBufferSizeY);
-		SetShaderValue(GetPixelShader(), SourceTexelOffsets23, Offsets23);
-		SceneTextureParameters.Set(GetPixelShader(), View);
+		SetShaderValue(RHICmdList, GetPixelShader(), SourceTexelOffsets23, Offsets23);
+		SceneTextureParameters.Set(RHICmdList, GetPixelShader(), View);
 	}
 
 	virtual bool Serialize(FArchive& Ar)
@@ -1031,6 +1053,9 @@ void FDeferredShadingSceneRenderer::UpdateDownsampledDepthSurface()
 {
 	if (GSceneRenderTargets.UseDownsizedOcclusionQueries() && IsFeatureLevelSupported(GRHIShaderPlatform, ERHIFeatureLevel::SM3))
 	{
+		//@todo-rco: RHIPacketList
+		FRHICommandList* RHICmdList = nullptr;
+
 		RHISetRenderTarget(NULL, GSceneRenderTargets.GetSmallDepthSurface());
 
 		SCOPED_DRAW_EVENT(DownsampleDepth, DEC_SCENE_ITEMS);
@@ -1050,7 +1075,7 @@ void FDeferredShadingSceneRenderer::UpdateDownsampledDepthSurface()
 			RHISetRasterizerState(TStaticRasterizerState<FM_Solid,CM_None>::GetRHI());
 			RHISetDepthStencilState(TStaticDepthStencilState<true,CF_Always>::GetRHI());
 
-			PixelShader->SetParameters(View);
+			PixelShader->SetParameters(RHICmdList, View);
 
 			const uint32 DownsampledX = FMath::TruncToInt(View.ViewRect.Min.X / GSceneRenderTargets.GetSmallColorDepthDownsampleFactor());
 			const uint32 DownsampledY = FMath::TruncToInt(View.ViewRect.Min.Y / GSceneRenderTargets.GetSmallColorDepthDownsampleFactor());
diff --git a/Engine/Source/Runtime/Renderer/Private/DeferredShadingRenderer.h b/Engine/Source/Runtime/Renderer/Private/DeferredShadingRenderer.h
--- a/Engine/Source/Runtime/Renderer/Private/DeferredShadingRenderer.h
+++ b/Engine/Source/Runtime/Renderer/Private/DeferredShadingRenderer.h
@@ -45,15 +45,15 @@ public:
     void ClearGBufferAtMaxZ();
 
 	/** Clears LPVs for all views */
-	void ClearLPVs();
+	void ClearLPVs(FRHICommandList* RHICmdList);
 
 	/** Propagates LPVs for all views */
 	void PropagateLPVs();
 
     /** Renders the basepass for the static data of a given View. */
-    bool RenderBasePassStaticData(FViewInfo& View);
-	bool RenderBasePassStaticDataMasked(FViewInfo& View);
-	bool RenderBasePassStaticDataDefault(FViewInfo& View);
+    bool RenderBasePassStaticData(FRHICommandList* RHICmdList, FViewInfo& View);
+	bool RenderBasePassStaticDataMasked(FRHICommandList* RHICmdList, FViewInfo& View);
+	bool RenderBasePassStaticDataDefault(FRHICommandList* RHICmdList, FViewInfo& View);
 
 	/** Sorts base pass draw lists front to back for improved GPU culling. */
 	void SortBasePassStaticData(FVector ViewPosition);
@@ -174,7 +174,7 @@ private:
 	void RenderAtmosphere(FLightShaftsOutput LightShaftsOutput);
 
 	/** Renders reflections that can be done in a deferred pass. */
-	void RenderDeferredReflections();
+	void RenderDeferredReflections(FRHICommandList* RHICmdList);
 
 	/** Render dynamic sky lighting from Movable sky lights. */
 	void RenderDynamicSkyLighting();
@@ -189,10 +189,10 @@ private:
 	bool ShouldUseTiledDeferred(int32 NumUnshadowedLights, int32 NumSimpleLights) const;
 
 	/** Renders the lights in SortedLights in the range [0, NumUnshadowedLights) using tiled deferred shading. */
-	void RenderTiledDeferredLighting(const TArray<FSortedLightSceneInfo, SceneRenderingAllocator>& SortedLights, int32 NumUnshadowedLights, const FSimpleLightArray& SimpleLights);
+	void RenderTiledDeferredLighting(FRHICommandList* RHICmdList, const TArray<FSortedLightSceneInfo, SceneRenderingAllocator>& SortedLights, int32 NumUnshadowedLights, const FSimpleLightArray& SimpleLights);
 
 	/** Renders the scene's lighting. */
-	void RenderLights();
+	void RenderLights(FRHICommandList* RHICmdList);
 
 	/** Renders an array of lights for the stationary light overlap viewmode. */
 	void RenderLightArrayForOverlapViewmode(const TSparseArray<FLightSceneInfoCompact>& LightArray);
@@ -240,19 +240,20 @@ private:
 	bool CheckForProjectedShadows( const FLightSceneInfo* LightSceneInfo ) const;
 
 	/** Renders one pass point light shadows. */
-	bool RenderOnePassPointLightShadows(const FLightSceneInfo* LightSceneInfo, bool bRenderedTranslucentObjectShadows, bool& bInjectedTranslucentVolume);
+	bool RenderOnePassPointLightShadows(FRHICommandList* RHICmdList, const FLightSceneInfo* LightSceneInfo, bool bRenderedTranslucentObjectShadows, bool& bInjectedTranslucentVolume);
 
 	/** Renders the projections of the given Shadows to the appropriate color render target. */
 	void RenderProjections(
+		FRHICommandList* RHICmdList, 
 		const FLightSceneInfo* LightSceneInfo, 
 		const TArray<FProjectedShadowInfo*,SceneRenderingAllocator>& Shadows
 		);
 
 	/** Renders the shadowmaps of translucent shadows and their projections onto opaque surfaces. */
-	bool RenderTranslucentProjectedShadows( const FLightSceneInfo* LightSceneInfo );
+	bool RenderTranslucentProjectedShadows(FRHICommandList* RHICmdList, const FLightSceneInfo* LightSceneInfo );
 
 	/** Renders reflective shadowmaps for LPVs */
-	bool RenderReflectiveShadowMaps( const FLightSceneInfo* LightSceneInfo );
+	bool RenderReflectiveShadowMaps(FRHICommandList* RHICmdList, const FLightSceneInfo* LightSceneInfo );
 
 	/**
 	  * Used by RenderLights to render projected shadows to the attenuation buffer.
@@ -260,10 +261,10 @@ private:
 	  * @param LightSceneInfo Represents the current light
 	  * @return true if anything got rendered
 	  */
-	bool RenderProjectedShadows( const FLightSceneInfo* LightSceneInfo, bool bRenderedTranslucentObjectShadows, bool& bInjectedTranslucentVolume );
+	bool RenderProjectedShadows(FRHICommandList* RHICmdList, const FLightSceneInfo* LightSceneInfo, bool bRenderedTranslucentObjectShadows, bool& bInjectedTranslucentVolume );
 
 	/** Caches the depths of any preshadows that should be cached, and renders their projections. */
-	bool RenderCachedPreshadows(const FLightSceneInfo* LightSceneInfo);
+	bool RenderCachedPreshadows(FRHICommandList* RHICmdList, const FLightSceneInfo* LightSceneInfo);
 
 	/**
 	  * Used by RenderLights to figure out if light functions need to be rendered to the attenuation buffer.
@@ -279,13 +280,13 @@ private:
 	  * @param LightSceneInfo Represents the current light
 	  * @param LightIndex The light's index into FScene::Lights
 	  */
-	bool RenderLightFunction(const FLightSceneInfo* LightSceneInfo, bool bLightAttenuationCleared);
+	bool RenderLightFunction(FRHICommandList* RHICmdList, const FLightSceneInfo* LightSceneInfo, bool bLightAttenuationCleared);
 
 	/** Renders a light function indicating that whole scene shadowing being displayed is for previewing only, and will go away in game. */
-	bool RenderPreviewShadowsIndicator(const FLightSceneInfo* LightSceneInfo, bool bLightAttenuationCleared);
+	bool RenderPreviewShadowsIndicator(FRHICommandList* RHICmdList, const FLightSceneInfo* LightSceneInfo, bool bLightAttenuationCleared);
 
 	/** Renders a light function with the given material. */
-	bool RenderLightFunctionForMaterial(const FLightSceneInfo* LightSceneInfo, const FMaterialRenderProxy* MaterialProxy, bool bLightAttenuationCleared, bool bRenderingPreviewShadowsIndicator);
+	bool RenderLightFunctionForMaterial(FRHICommandList* RHICmdList, const FLightSceneInfo* LightSceneInfo, const FMaterialRenderProxy* MaterialProxy, bool bLightAttenuationCleared, bool bRenderingPreviewShadowsIndicator);
 
 	/**
 	  * Used by RenderLights to render a light to the scene color buffer.
@@ -300,13 +301,13 @@ private:
 	void RenderSimpleLightsStandardDeferred(const FSimpleLightArray& SimpleLights);
 
 	/** Clears the translucency lighting volumes before light accumulation. */
-	void ClearTranslucentVolumeLighting();
+	void ClearTranslucentVolumeLighting(FRHICommandList* RHICmdList);
 
 	/** Add AmbientCubemap to the lighting volumes. */
-	void InjectAmbientCubemapTranslucentVolumeLighting();
+	void InjectAmbientCubemapTranslucentVolumeLighting(FRHICommandList* RHICmdList);
 
 	/** Renders indirect lighting into the translucent lighting volumes. */
-	void CompositeIndirectTranslucentVolumeLighting();
+	void CompositeIndirectTranslucentVolumeLighting(FRHICommandList* RHICmdList);
 
 	/** Clears the volume texture used to accumulate per object shadows for translucency. */
 	void ClearTranslucentVolumePerObjectShadowing();
@@ -330,10 +331,10 @@ private:
 	void RenderReflectionCaptureSpecularBounceForAllViews();
 
 	/** Render image based reflections (SSR, Env, SkyLight) without compute shaders */
-	void RenderImageBasedReflectionsSM4ForAllViews(bool bReflectionEnv);
+	void RenderImageBasedReflectionsSM4ForAllViews(FRHICommandList* RHICmdList, bool bReflectionEnv);
 
 	/** Render image based reflections (SSR, Env, SkyLight) with compute shaders */
-	void RenderImageBasedReflectionsSM5ForAllViews();
+	void RenderImageBasedReflectionsSM5ForAllViews(FRHICommandList* RHICmdList);
 
 	bool ShouldDoReflectionEnvironment() const;
 
diff --git a/Engine/Source/Runtime/Renderer/Private/DepthRendering.cpp b/Engine/Source/Runtime/Renderer/Private/DepthRendering.cpp
--- a/Engine/Source/Runtime/Renderer/Private/DepthRendering.cpp
+++ b/Engine/Source/Runtime/Renderer/Private/DepthRendering.cpp
@@ -37,14 +37,14 @@ public:
 			 && IsFeatureLevelSupported(Platform, ERHIFeatureLevel::SM3);
 	}
 
-	void SetParameters(const FMaterialRenderProxy* MaterialRenderProxy,const FMaterial& MaterialResource,const FSceneView& View)
+	void SetParameters(FRHICommandList* RHICmdList, const FMaterialRenderProxy* MaterialRenderProxy,const FMaterial& MaterialResource,const FSceneView& View)
 	{
-		FMeshMaterialShader::SetParameters(GetVertexShader(),MaterialRenderProxy,MaterialResource,View,ESceneRenderTargetsMode::DontSet);
+		FMeshMaterialShader::SetParameters(RHICmdList, GetVertexShader(),MaterialRenderProxy,MaterialResource,View,ESceneRenderTargetsMode::DontSet);
 	}
 
-	void SetMesh(const FVertexFactory* VertexFactory,const FSceneView& View,const FPrimitiveSceneProxy* Proxy,const FMeshBatchElement& BatchElement)
+	void SetMesh(FRHICommandList* RHICmdList, const FVertexFactory* VertexFactory,const FSceneView& View,const FPrimitiveSceneProxy* Proxy,const FMeshBatchElement& BatchElement)
 	{
-		FMeshMaterialShader::SetMesh(GetVertexShader(),VertexFactory,View,Proxy,BatchElement);
+		FMeshMaterialShader::SetMesh(RHICmdList, GetVertexShader(),VertexFactory,View,Proxy,BatchElement);
 	}
 };
 
@@ -115,14 +115,14 @@ public:
 
 	FDepthOnlyPS() {}
 
-	void SetParameters(const FMaterialRenderProxy* MaterialRenderProxy,const FMaterial& MaterialResource,const FSceneView* View)
+	void SetParameters(FRHICommandList* RHICmdList, const FMaterialRenderProxy* MaterialRenderProxy,const FMaterial& MaterialResource,const FSceneView* View)
 	{
-		FMeshMaterialShader::SetParameters(GetPixelShader(),MaterialRenderProxy,MaterialResource,*View,ESceneRenderTargetsMode::DontSet);
+		FMeshMaterialShader::SetParameters(RHICmdList, GetPixelShader(),MaterialRenderProxy,MaterialResource,*View,ESceneRenderTargetsMode::DontSet);
 	}
 
-	void SetMesh(const FVertexFactory* VertexFactory,const FSceneView& View,const FPrimitiveSceneProxy* Proxy,const FMeshBatchElement& BatchElement)
+	void SetMesh(FRHICommandList* RHICmdList, const FVertexFactory* VertexFactory,const FSceneView& View,const FPrimitiveSceneProxy* Proxy,const FMeshBatchElement& BatchElement)
 	{
-		FMeshMaterialShader::SetMesh(GetPixelShader(),VertexFactory,View,Proxy,BatchElement);
+		FMeshMaterialShader::SetMesh(RHICmdList, GetPixelShader(),VertexFactory,View,Proxy,BatchElement);
 	}
 };
 
@@ -169,26 +169,26 @@ FDepthDrawingPolicy::FDepthDrawingPolicy(
 
 }
 
-void FDepthDrawingPolicy::DrawShared(const FSceneView* View,FBoundShaderStateRHIParamRef BoundShaderState) const
+void FDepthDrawingPolicy::DrawShared(FRHICommandList* RHICmdList, const FSceneView* View,FBoundShaderStateRHIParamRef BoundShaderState) const
 {
 	// Set the actual shader & vertex declaration state
 	RHISetBoundShaderState(BoundShaderState);
 
 	// Set the depth-only shader parameters for the material.
-	VertexShader->SetParameters(MaterialRenderProxy,*MaterialResource,*View);
+	VertexShader->SetParameters(RHICmdList, MaterialRenderProxy,*MaterialResource,*View);
 	if(HullShader && DomainShader)
 	{
-		HullShader->SetParameters(MaterialRenderProxy,*View);
-		DomainShader->SetParameters(MaterialRenderProxy,*View);
+		HullShader->SetParameters(RHICmdList, MaterialRenderProxy,*View);
+		DomainShader->SetParameters(RHICmdList, MaterialRenderProxy,*View);
 	}
 
 	if (bNeedsPixelShader)
 	{
-		PixelShader->SetParameters(MaterialRenderProxy,*MaterialResource,View);
+		PixelShader->SetParameters(RHICmdList, MaterialRenderProxy,*MaterialResource,View);
 	}
 
 	// Set the shared mesh resources.
-	FMeshDrawingPolicy::DrawShared(View);
+	FMeshDrawingPolicy::DrawShared(RHICmdList, View);
 }
 
 /** 
@@ -208,6 +208,7 @@ FBoundShaderStateRHIRef FDepthDrawingPolicy::CreateBoundShaderState(ERHIFeatureL
 }
 
 void FDepthDrawingPolicy::SetMeshRenderState(
+	FRHICommandList* RHICmdList, 
 	const FSceneView& View,
 	const FPrimitiveSceneProxy* PrimitiveSceneProxy,
 	const FMeshBatch& Mesh,
@@ -217,18 +218,18 @@ void FDepthDrawingPolicy::SetMeshRenderState(
 	) const
 {
 	const FMeshBatchElement& BatchElement = Mesh.Elements[BatchElementIndex];
-	VertexShader->SetMesh(VertexFactory,View,PrimitiveSceneProxy,BatchElement);
+	VertexShader->SetMesh(RHICmdList, VertexFactory,View,PrimitiveSceneProxy,BatchElement);
 	if(HullShader && DomainShader)
 	{
-		HullShader->SetMesh(VertexFactory,View,PrimitiveSceneProxy,BatchElement);
-		DomainShader->SetMesh(VertexFactory,View,PrimitiveSceneProxy,BatchElement);
+		HullShader->SetMesh(RHICmdList, VertexFactory,View,PrimitiveSceneProxy,BatchElement);
+		DomainShader->SetMesh(RHICmdList, VertexFactory,View,PrimitiveSceneProxy,BatchElement);
 	}
 
 	if (bNeedsPixelShader)
 	{
-		PixelShader->SetMesh(VertexFactory,View,PrimitiveSceneProxy,BatchElement);
+		PixelShader->SetMesh(RHICmdList, VertexFactory,View,PrimitiveSceneProxy,BatchElement);
 	}
-	FMeshDrawingPolicy::SetMeshRenderState(View,PrimitiveSceneProxy,Mesh,BatchElementIndex,bBackFace,ElementData);
+	FMeshDrawingPolicy::SetMeshRenderState(RHICmdList, View,PrimitiveSceneProxy,Mesh,BatchElementIndex,bBackFace,ElementData);
 }
 
 int32 CompareDrawingPolicy(const FDepthDrawingPolicy& A,const FDepthDrawingPolicy& B)
@@ -256,13 +257,15 @@ FPositionOnlyDepthDrawingPolicy::FPositionOnlyDepthDrawingPolicy(
 	VertexShader = InMaterialResource.GetShader<TDepthOnlyVS<true> >(InVertexFactory->GetType());
 }
 
-void FPositionOnlyDepthDrawingPolicy::DrawShared(const FSceneView* View,FBoundShaderStateRHIParamRef BoundShaderState) const
+void FPositionOnlyDepthDrawingPolicy::DrawShared(FRHICommandList* RHICmdList, const FSceneView* View,FBoundShaderStateRHIParamRef BoundShaderState) const
 {
+	check(!RHICmdList);
+
 	// Set the actual shader & vertex declaration state
 	RHISetBoundShaderState( BoundShaderState);
 
 	// Set the depth-only shader parameters for the material.
-	VertexShader->SetParameters(MaterialRenderProxy,*MaterialResource,*View);
+	VertexShader->SetParameters(RHICmdList, MaterialRenderProxy,*MaterialResource,*View);
 
 	// Set the shared mesh resources.
 	VertexFactory->SetPositionStream();
@@ -283,6 +286,7 @@ FBoundShaderStateRHIRef FPositionOnlyDepthDrawingPolicy::CreateBoundShaderState(
 }
 
 void FPositionOnlyDepthDrawingPolicy::SetMeshRenderState(
+	FRHICommandList* RHICmdList, 
 	const FSceneView& View,
 	const FPrimitiveSceneProxy* PrimitiveSceneProxy,
 	const FMeshBatch& Mesh,
@@ -291,8 +295,8 @@ void FPositionOnlyDepthDrawingPolicy::SetMeshRenderState(
 	const ElementDataType& ElementData
 	) const
 {
-	VertexShader->SetMesh(VertexFactory,View,PrimitiveSceneProxy,Mesh.Elements[BatchElementIndex]);
-	FMeshDrawingPolicy::SetMeshRenderState(View,PrimitiveSceneProxy,Mesh,BatchElementIndex,bBackFace,ElementData);
+	VertexShader->SetMesh(RHICmdList, VertexFactory,View,PrimitiveSceneProxy,Mesh.Elements[BatchElementIndex]);
+	FMeshDrawingPolicy::SetMeshRenderState(RHICmdList, View,PrimitiveSceneProxy,Mesh,BatchElementIndex,bBackFace,ElementData);
 }
 
 int32 CompareDrawingPolicy(const FPositionOnlyDepthDrawingPolicy& A,const FPositionOnlyDepthDrawingPolicy& B)
@@ -370,6 +374,7 @@ void FDepthDrawingPolicyFactory::AddStaticMesh(FScene* Scene,FStaticMesh* Static
 }
 
 bool FDepthDrawingPolicyFactory::DrawMesh(
+	FRHICommandList* RHICmdList, 
 	const FViewInfo& View,
 	ContextType DrawingContext,
 	const FMeshBatch& Mesh,
@@ -401,15 +406,15 @@ bool FDepthDrawingPolicyFactory::DrawMesh(
 			//render opaque primitives that support a separate position-only vertex buffer
 			const FMaterialRenderProxy* DefaultProxy = UMaterial::GetDefaultMaterial(MD_Surface)->GetRenderProxy(false);
 			FPositionOnlyDepthDrawingPolicy DrawingPolicy(Mesh.VertexFactory, DefaultProxy, *DefaultProxy->GetMaterial(View.GetFeatureLevel()), Material->IsTwoSided(), Material->IsWireframe());
-			DrawingPolicy.DrawShared(&View,DrawingPolicy.CreateBoundShaderState(View.GetFeatureLevel()));
+			DrawingPolicy.DrawShared(RHICmdList, &View,DrawingPolicy.CreateBoundShaderState(View.GetFeatureLevel()));
 			int32 BatchElementIndex = 0;
 			uint64 Mask = BatchElementMask;
 			do
 			{
 				if(Mask & 1)
 				{
-					DrawingPolicy.SetMeshRenderState(View,PrimitiveSceneProxy,Mesh,BatchElementIndex,bBackFace,FPositionOnlyDepthDrawingPolicy::ElementDataType());
-					DrawingPolicy.DrawMesh(Mesh,BatchElementIndex);
+					DrawingPolicy.SetMeshRenderState(RHICmdList, View,PrimitiveSceneProxy,Mesh,BatchElementIndex,bBackFace,FPositionOnlyDepthDrawingPolicy::ElementDataType());
+					DrawingPolicy.DrawMesh(RHICmdList, Mesh,BatchElementIndex);
 				}
 				Mask >>= 1;
 				BatchElementIndex++;
@@ -443,7 +448,7 @@ bool FDepthDrawingPolicyFactory::DrawMesh(
 				}
 
 				FDepthDrawingPolicy DrawingPolicy(Mesh.VertexFactory, MaterialRenderProxy, *MaterialRenderProxy->GetMaterial(View.GetFeatureLevel()), Material->IsTwoSided());
-				DrawingPolicy.DrawShared(&View,DrawingPolicy.CreateBoundShaderState(View.GetFeatureLevel()));
+				DrawingPolicy.DrawShared(RHICmdList, &View,DrawingPolicy.CreateBoundShaderState(View.GetFeatureLevel()));
 
 				int32 BatchElementIndex = 0;
 				uint64 Mask = BatchElementMask;
@@ -451,8 +456,8 @@ bool FDepthDrawingPolicyFactory::DrawMesh(
 				{
 					if(Mask & 1)
 					{
-						DrawingPolicy.SetMeshRenderState(View,PrimitiveSceneProxy,Mesh,BatchElementIndex,bBackFace,FMeshDrawingPolicy::ElementDataType());
-						DrawingPolicy.DrawMesh(Mesh,BatchElementIndex);
+						DrawingPolicy.SetMeshRenderState(RHICmdList, View,PrimitiveSceneProxy,Mesh,BatchElementIndex,bBackFace,FMeshDrawingPolicy::ElementDataType());
+						DrawingPolicy.DrawMesh(RHICmdList, Mesh,BatchElementIndex);
 					}
 					Mask >>= 1;
 					BatchElementIndex++;
@@ -467,6 +472,7 @@ bool FDepthDrawingPolicyFactory::DrawMesh(
 }
 
 bool FDepthDrawingPolicyFactory::DrawDynamicMesh(
+	FRHICommandList* RHICmdList, 
 	const FViewInfo& View,
 	ContextType DrawingContext,
 	const FMeshBatch& Mesh,
@@ -477,6 +483,7 @@ bool FDepthDrawingPolicyFactory::DrawDynamicMesh(
 	)
 {
 	return DrawMesh(
+		RHICmdList, 
 		View,
 		DrawingContext,
 		Mesh,
@@ -489,6 +496,7 @@ bool FDepthDrawingPolicyFactory::DrawDynamicMesh(
 
 
 bool FDepthDrawingPolicyFactory::DrawStaticMesh(
+	FRHICommandList* RHICmdList, 
 	const FViewInfo& View,
 	ContextType DrawingContext,
 	const FStaticMesh& StaticMesh,
@@ -503,6 +511,7 @@ bool FDepthDrawingPolicyFactory::DrawStaticMesh(
 	const FMaterial* Material = StaticMesh.MaterialRenderProxy->GetMaterial(View.GetFeatureLevel());
 	const EMaterialShadingModel ShadingModel = Material->GetShadingModel();
 	bDirty |= DrawMesh(
+		RHICmdList, 
 		View,
 		DrawingContext,
 		StaticMesh,
diff --git a/Engine/Source/Runtime/Renderer/Private/DepthRendering.h b/Engine/Source/Runtime/Renderer/Private/DepthRendering.h
--- a/Engine/Source/Runtime/Renderer/Private/DepthRendering.h
+++ b/Engine/Source/Runtime/Renderer/Private/DepthRendering.h
@@ -44,7 +44,7 @@ public:
 			&& PixelShader == Other.PixelShader;
 	}
 
-	void DrawShared(const FSceneView* View,FBoundShaderStateRHIParamRef BoundShaderState) const;
+	void DrawShared(FRHICommandList* RHICmdList, const FSceneView* View,FBoundShaderStateRHIParamRef BoundShaderState) const;
 
 	/** 
 	* Create bound shader state using the vertex decl from the mesh draw policy
@@ -55,6 +55,7 @@ public:
 	FBoundShaderStateRHIRef CreateBoundShaderState(ERHIFeatureLevel::Type InFeatureLevel);
 
 	void SetMeshRenderState(
+		FRHICommandList* RHICmdList, 
 		const FSceneView& View,
 		const FPrimitiveSceneProxy* PrimitiveSceneProxy,
 		const FMeshBatch& Mesh,
@@ -96,7 +97,7 @@ public:
 		return FMeshDrawingPolicy::Matches(Other) && VertexShader == Other.VertexShader;
 	}
 
-	void DrawShared(const FSceneView* View,FBoundShaderStateRHIParamRef BoundShaderState) const;
+	void DrawShared(FRHICommandList* RHICmdList, const FSceneView* View,FBoundShaderStateRHIParamRef BoundShaderState) const;
 
 	/** 
 	* Create bound shader state using the vertex decl from the mesh draw policy
@@ -106,6 +107,7 @@ public:
 	FBoundShaderStateRHIRef CreateBoundShaderState(ERHIFeatureLevel::Type InFeatureLevel);
 
 	void SetMeshRenderState(
+		FRHICommandList* RHICmdList, 
 		const FSceneView& View,
 		const FPrimitiveSceneProxy* PrimitiveSceneProxy,
 		const FMeshBatch& Mesh,
@@ -139,6 +141,7 @@ public:
 
 	static void AddStaticMesh(FScene* Scene,FStaticMesh* StaticMesh);
 	static bool DrawDynamicMesh(
+		FRHICommandList* RHICmdList, 
 		const FViewInfo& View,
 		ContextType DrawingContext,
 		const FMeshBatch& Mesh,
@@ -149,6 +152,7 @@ public:
 		);
 
 	static bool DrawStaticMesh(
+		FRHICommandList* RHICmdList, 
 		const FViewInfo& View,
 		ContextType DrawingContext,
 		const FStaticMesh& StaticMesh,
@@ -166,6 +170,7 @@ private:
 	* @return true if the mesh rendered
 	*/
 	static bool DrawMesh(
+		FRHICommandList* RHICmdList, 
 		const FViewInfo& View,
 		ContextType DrawingContext,
 		const FMeshBatch& Mesh,
diff --git a/Engine/Source/Runtime/Renderer/Private/DistanceFieldLightingShared.h b/Engine/Source/Runtime/Renderer/Private/DistanceFieldLightingShared.h
--- a/Engine/Source/Runtime/Renderer/Private/DistanceFieldLightingShared.h
+++ b/Engine/Source/Runtime/Renderer/Private/DistanceFieldLightingShared.h
@@ -63,23 +63,24 @@ public:
 	}
 
 	template<typename ShaderRHIParamRef>
-	void Set(const ShaderRHIParamRef ShaderRHI)
+	void Set(FRHICommandList* RHICmdList, const ShaderRHIParamRef ShaderRHI)
 	{
 		float AOMaxDistanceValue = GetAOMaxDistance();
-		SetShaderValue(ShaderRHI, AOMaxDistance, AOMaxDistanceValue);
+		SetShaderValue(RHICmdList, ShaderRHI, AOMaxDistance, AOMaxDistanceValue);
 
 		extern float GAOLargestSampleOffset;
 		extern float GAOStepExponentScale;
 		extern uint32 GAONumConeSteps;
 		float AOStepScaleValue = GAOLargestSampleOffset / FMath::Pow(2.0f, GAOStepExponentScale * (GAONumConeSteps - 1));
-		SetShaderValue(ShaderRHI, AOStepScale, AOStepScaleValue);
+		SetShaderValue(RHICmdList, ShaderRHI, AOStepScale, AOStepScaleValue);
 
-		SetShaderValue(ShaderRHI, AOStepExponentScale, GAOStepExponentScale);
+		SetShaderValue(RHICmdList, ShaderRHI, AOStepExponentScale, GAOStepExponentScale);
 
 		extern float GAOMaxViewDistance;
-		SetShaderValue(ShaderRHI, AOMaxViewDistance, GAOMaxViewDistance);
+		SetShaderValue(RHICmdList, ShaderRHI, AOMaxViewDistance, GAOMaxViewDistance);
 
 		SetTextureParameter(
+			RHICmdList,
 			ShaderRHI,
 			DistanceFieldTexture,
 			DistanceFieldSampler,
@@ -91,7 +92,7 @@ public:
 		const int32 NumTexelsOneDimY = GDistanceFieldVolumeTextureAtlas.GetSizeY();
 		const int32 NumTexelsOneDimZ = GDistanceFieldVolumeTextureAtlas.GetSizeZ();
 		const FVector InvTextureDim(1.0f / NumTexelsOneDimX, 1.0f / NumTexelsOneDimY, 1.0f / NumTexelsOneDimZ);
-		SetShaderValue(ShaderRHI, DistanceFieldAtlasTexelSize, InvTextureDim);
+		SetShaderValue(RHICmdList, ShaderRHI, DistanceFieldAtlasTexelSize, InvTextureDim);
 	}
 
 private:
