UE4 - merge first pass at parallel rendering to main

diff --git a/Engine/Source/Runtime/Renderer/Private/DistanceFieldSurfaceCacheLighting.cpp b/Engine/Source/Runtime/Renderer/Private/DistanceFieldSurfaceCacheLighting.cpp
--- a/Engine/Source/Runtime/Renderer/Private/DistanceFieldSurfaceCacheLighting.cpp
+++ b/Engine/Source/Runtime/Renderer/Private/DistanceFieldSurfaceCacheLighting.cpp
@@ -14,6 +14,7 @@
 #include "DistanceFieldSurfaceCacheLighting.h"
 #include "DistanceFieldLightingShared.h"
 #include "PostProcessAmbientOcclusion.h"
+#include "RHICommandList.h"
 
 int32 GDistanceFieldAO = 1;
 FAutoConsoleVariableRef CVarDistanceFieldAO(
@@ -360,13 +361,13 @@ public:
 	}
 
 	template<typename TParamRef>
-	void Set(const TParamRef& ShaderRHI, int32 NumObjectsValue)
+	void Set(FRHICommandList* RHICmdList, const TParamRef& ShaderRHI, int32 NumObjectsValue)
 	{
-		SetSRVParameter(ShaderRHI, ObjectBounds, GAOObjectBuffers.ObjectData.BoundsSRV);
-		SetSRVParameter(ShaderRHI, ObjectData, GAOObjectBuffers.ObjectData.DataSRV);
-		SetSRVParameter(ShaderRHI, ObjectData2, GAOObjectBuffers.ObjectData.Data2SRV);
+		SetSRVParameter(RHICmdList, ShaderRHI, ObjectBounds, GAOObjectBuffers.ObjectData.BoundsSRV);
+		SetSRVParameter(RHICmdList, ShaderRHI, ObjectData, GAOObjectBuffers.ObjectData.DataSRV);
+		SetSRVParameter(RHICmdList, ShaderRHI, ObjectData2, GAOObjectBuffers.ObjectData.Data2SRV);
 
-		SetShaderValue(ShaderRHI, NumObjects, NumObjectsValue);
+		SetShaderValue(RHICmdList, ShaderRHI, NumObjects, NumObjectsValue);
 	}
 
 	friend FArchive& operator<<(FArchive& Ar, FAOObjectBufferParameters& P)
@@ -391,7 +392,7 @@ public:
 
 	static bool ShouldCache(EShaderPlatform Platform)
 	{
-		return IsFeatureLevelSupported(Platform, ERHIFeatureLevel::SM5) && DoesPlatformSupportDistanceFieldAO(Platform);
+		return IsFeatureLevelSupported(Platform, ERHIFeatureLevel::SM5) && DoesPlatformSupportDistanceFieldAO(Platform) && !IsOpenGLPlatform(Platform);
 	}
 
 	static void ModifyCompilationEnvironment(EShaderPlatform Platform, FShaderCompilerEnvironment& OutEnvironment)
@@ -422,30 +423,30 @@ public:
 	FBuildTileConesCS()
 	{
 	}
-	void SetParameters(const FSceneView& View, FScene* Scene, FVector2D NumGroupsValue)
+	void SetParameters(FRHICommandList* RHICmdList, const FSceneView& View, FScene* Scene, FVector2D NumGroupsValue)
 	{
 		FComputeShaderRHIParamRef ShaderRHI = GetComputeShader();
 
-		FGlobalShader::SetParameters(ShaderRHI, View);
-		DeferredParameters.Set(ShaderRHI, View);
-		AOParameters.Set(ShaderRHI);
+		FGlobalShader::SetParameters(RHICmdList, ShaderRHI, View);
+		DeferredParameters.Set(RHICmdList, ShaderRHI, View);
+		AOParameters.Set(RHICmdList, ShaderRHI);
 
 		FTileIntersectionResources* TileIntersectionResources = ((FSceneViewState*)View.State)->AOTileIntersectionResources;
 
-		TileConeAxisAndCos.SetBuffer(ShaderRHI, TileIntersectionResources->TileConeAxisAndCos);
-		TileConeDepthRanges.SetBuffer(ShaderRHI, TileIntersectionResources->TileConeDepthRanges);
-		TileHeadDataUnpacked.SetBuffer(ShaderRHI, TileIntersectionResources->TileHeadDataUnpacked);
+		TileConeAxisAndCos.SetBuffer(RHICmdList, ShaderRHI, TileIntersectionResources->TileConeAxisAndCos);
+		TileConeDepthRanges.SetBuffer(RHICmdList, ShaderRHI, TileIntersectionResources->TileConeDepthRanges);
+		TileHeadDataUnpacked.SetBuffer(RHICmdList, ShaderRHI, TileIntersectionResources->TileHeadDataUnpacked);
 
-		SetShaderValue(ShaderRHI, ViewDimensionsParameter, View.ViewRect);
+		SetShaderValue(RHICmdList, ShaderRHI, ViewDimensionsParameter, View.ViewRect);
 
-		SetShaderValue(ShaderRHI, NumGroups, NumGroupsValue);
+		SetShaderValue(RHICmdList, ShaderRHI, NumGroups, NumGroupsValue);
 	}
 
-	void UnsetParameters()
+	void UnsetParameters(FRHICommandList* RHICmdList)
 	{
-		TileConeAxisAndCos.UnsetUAV(GetComputeShader());
-		TileConeDepthRanges.UnsetUAV(GetComputeShader());
-		TileHeadDataUnpacked.UnsetUAV(GetComputeShader());
+		TileConeAxisAndCos.UnsetUAV(RHICmdList, GetComputeShader());
+		TileConeDepthRanges.UnsetUAV(RHICmdList, GetComputeShader());
+		TileHeadDataUnpacked.UnsetUAV(RHICmdList, GetComputeShader());
 	}
 
 	virtual bool Serialize(FArchive& Ar)
@@ -483,7 +484,7 @@ public:
 
 	static bool ShouldCache(EShaderPlatform Platform) 
 	{ 
-		return IsFeatureLevelSupported(Platform, ERHIFeatureLevel::SM5) && DoesPlatformSupportDistanceFieldAO(Platform); 
+		return IsFeatureLevelSupported(Platform, ERHIFeatureLevel::SM5) && DoesPlatformSupportDistanceFieldAO(Platform) && !IsOpenGLPlatform(Platform); 
 	}
 
 	static void ModifyCompilationEnvironment(EShaderPlatform Platform, FShaderCompilerEnvironment& OutEnvironment)
@@ -501,13 +502,13 @@ public:
 
 	FObjectCullVS() {}
 
-	void SetParameters(const FSceneView& View, int32 NumObjects)
+	void SetParameters(FRHICommandList* RHICmdList, const FSceneView& View, int32 NumObjects)
 	{
 		const FVertexShaderRHIParamRef ShaderRHI = GetVertexShader();
-		FGlobalShader::SetParameters(ShaderRHI, View);
+		FGlobalShader::SetParameters(RHICmdList, ShaderRHI, View);
 		
-		ObjectParameters.Set(ShaderRHI, NumObjects);
-		AOParameters.Set(ShaderRHI);
+		ObjectParameters.Set(RHICmdList, ShaderRHI, NumObjects);
+		AOParameters.Set(RHICmdList, ShaderRHI);
 
 		const int32 NumRings = StencilingGeometry::GLowPolyStencilSphereVertexBuffer.GetNumRings();
 		const float RadiansPerRingSegment = PI / (float)NumRings;
@@ -515,7 +516,7 @@ public:
 		// Boost the effective radius so that the edges of the sphere approximation lie on the sphere, instead of the vertices
 		const float ConservativeRadiusScaleValue = 1.0f / FMath::Cos(RadiansPerRingSegment);
 
-		SetShaderValue(ShaderRHI, ConservativeRadiusScale, ConservativeRadiusScaleValue);
+		SetShaderValue(RHICmdList, ShaderRHI, ConservativeRadiusScale, ConservativeRadiusScaleValue);
 	}
 
 	virtual bool Serialize(FArchive& Ar)
@@ -542,7 +543,7 @@ public:
 
 	static bool ShouldCache(EShaderPlatform Platform)
 	{
-		return IsFeatureLevelSupported(Platform, ERHIFeatureLevel::SM5) && DoesPlatformSupportDistanceFieldAO(Platform);
+		return IsFeatureLevelSupported(Platform, ERHIFeatureLevel::SM5) && DoesPlatformSupportDistanceFieldAO(Platform) && !IsOpenGLPlatform(Platform);
 	}
 
 	static void ModifyCompilationEnvironment(EShaderPlatform Platform, FShaderCompilerEnvironment& OutEnvironment)
@@ -568,18 +569,18 @@ public:
 		NumGroups.Bind(Initializer.ParameterMap, TEXT("NumGroups"));
 	}
 
-	void SetParameters(const FSceneView& View, FVector2D NumGroupsValue)
+	void SetParameters(FRHICommandList* RHICmdList, const FSceneView& View, FVector2D NumGroupsValue)
 	{
 		const FPixelShaderRHIParamRef ShaderRHI = GetPixelShader();
-		FGlobalShader::SetParameters(ShaderRHI, View);
-		AOParameters.Set(ShaderRHI);
+		FGlobalShader::SetParameters(RHICmdList, ShaderRHI, View);
+		AOParameters.Set(RHICmdList, ShaderRHI);
 
 		FTileIntersectionResources* TileIntersectionResources = ((FSceneViewState*)View.State)->AOTileIntersectionResources;
 
-		SetSRVParameter(ShaderRHI, TileConeAxisAndCos, TileIntersectionResources->TileConeAxisAndCos.SRV);
-		SetSRVParameter(ShaderRHI, TileConeDepthRanges, TileIntersectionResources->TileConeDepthRanges.SRV);
+		SetSRVParameter(RHICmdList, ShaderRHI, TileConeAxisAndCos, TileIntersectionResources->TileConeAxisAndCos.SRV);
+		SetSRVParameter(RHICmdList, ShaderRHI, TileConeDepthRanges, TileIntersectionResources->TileConeDepthRanges.SRV);
 
-		SetShaderValue(ShaderRHI, NumGroups, NumGroupsValue);
+		SetShaderValue(RHICmdList, ShaderRHI, NumGroups, NumGroupsValue);
 	}
 
 	void GetUAVs(const FSceneView& View, TArray<FUnorderedAccessViewRHIParamRef>& UAVs)
@@ -635,7 +636,7 @@ public:
 
 	static bool ShouldCache(EShaderPlatform Platform)
 	{
-		return IsFeatureLevelSupported(Platform, ERHIFeatureLevel::SM5) && DoesPlatformSupportDistanceFieldAO(Platform);
+		return IsFeatureLevelSupported(Platform, ERHIFeatureLevel::SM5) && DoesPlatformSupportDistanceFieldAO(Platform) && !IsOpenGLPlatform(Platform);
 	}
 
 	static void ModifyCompilationEnvironment(EShaderPlatform Platform, FShaderCompilerEnvironment& OutEnvironment)
@@ -666,31 +667,31 @@ public:
 	FDistanceFieldBuildTileListCS()
 	{
 	}
-	void SetParameters(const FSceneView& View, FScene* Scene, FVector2D NumGroupsValue, int32 NumObjects)
+	void SetParameters(FRHICommandList* RHICmdList, const FSceneView& View, FScene* Scene, FVector2D NumGroupsValue, int32 NumObjects)
 	{
 		FComputeShaderRHIParamRef ShaderRHI = GetComputeShader();
 
-		FGlobalShader::SetParameters(ShaderRHI, View);
-		DeferredParameters.Set(ShaderRHI, View);
-		ObjectParameters.Set(ShaderRHI, NumObjects);
-		AOParameters.Set(ShaderRHI);
+		FGlobalShader::SetParameters(RHICmdList, ShaderRHI, View);
+		DeferredParameters.Set(RHICmdList, ShaderRHI, View);
+		ObjectParameters.Set(RHICmdList, ShaderRHI, NumObjects);
+		AOParameters.Set(RHICmdList, ShaderRHI);
 
 		FTileIntersectionResources* TileIntersectionResources = ((FSceneViewState*)View.State)->AOTileIntersectionResources;
 
-		TileHeadData.SetBuffer(ShaderRHI, TileIntersectionResources->TileHeadData);
-		TileArrayData.SetBuffer(ShaderRHI, TileIntersectionResources->TileArrayData);
-		TileArrayNextAllocation.SetBuffer(ShaderRHI, TileIntersectionResources->TileArrayNextAllocation);
+		TileHeadData.SetBuffer(RHICmdList, ShaderRHI, TileIntersectionResources->TileHeadData);
+		TileArrayData.SetBuffer(RHICmdList, ShaderRHI, TileIntersectionResources->TileArrayData);
+		TileArrayNextAllocation.SetBuffer(RHICmdList, ShaderRHI, TileIntersectionResources->TileArrayNextAllocation);
 
-		SetShaderValue(ShaderRHI, ViewDimensionsParameter, View.ViewRect);
+		SetShaderValue(RHICmdList, ShaderRHI, ViewDimensionsParameter, View.ViewRect);
 
-		SetShaderValue(ShaderRHI, NumGroups, NumGroupsValue);
+		SetShaderValue(RHICmdList, ShaderRHI, NumGroups, NumGroupsValue);
 	}
 
-	void UnsetParameters()
+	void UnsetParameters(FRHICommandList* RHICmdList)
 	{
-		TileHeadData.UnsetUAV(GetComputeShader());
-		TileArrayData.UnsetUAV(GetComputeShader());
-		TileArrayNextAllocation.UnsetUAV(GetComputeShader());
+		TileHeadData.UnsetUAV(RHICmdList, GetComputeShader());
+		TileArrayData.UnsetUAV(RHICmdList, GetComputeShader());
+		TileArrayNextAllocation.UnsetUAV(RHICmdList, GetComputeShader());
 	}
 
 	virtual bool Serialize(FArchive& Ar)
@@ -734,13 +735,13 @@ public:
 	}
 
 	template<typename TParamRef>
-	void Set(const TParamRef& ShaderRHI, const FSceneView& View, int32 CurrentLevelDownsampleFactorValue)
+	void Set(FRHICommandList* RHICmdList, const TParamRef& ShaderRHI, const FSceneView& View, int32 CurrentLevelDownsampleFactorValue)
 	{
-		SetShaderValue(ShaderRHI, CurrentLevelDownsampleFactor, CurrentLevelDownsampleFactorValue);
+		SetShaderValue(RHICmdList, ShaderRHI, CurrentLevelDownsampleFactor, CurrentLevelDownsampleFactorValue);
 
 		// Round up, to match render target allocation
 		const FVector2D AOBufferSizeValue = FIntPoint::DivideAndRoundUp(GSceneRenderTargets.GetBufferSizeXY(), CurrentLevelDownsampleFactorValue);
-		SetShaderValue(ShaderRHI, AOBufferSize, AOBufferSizeValue);
+		SetShaderValue(RHICmdList, ShaderRHI, AOBufferSize, AOBufferSizeValue);
 
 		const FVector2D InvAOBufferSize(1.0f / AOBufferSizeValue.X, 1.0f / AOBufferSizeValue.Y);
 
@@ -768,13 +769,13 @@ public:
 			(RoundedDownViewDimensions.X / 2.0f + GPixelCenterOffset + AOViewRectMin.X) * InvBufferSizeX
 			);
 
-		SetShaderValue(ShaderRHI, AOScreenPositionScaleBias, ScreenPositionScaleBias);
+		SetShaderValue(RHICmdList, ShaderRHI, AOScreenPositionScaleBias, ScreenPositionScaleBias);
 
-		SetShaderValue(ShaderRHI, DownsampleFactorToBaseLevel, CurrentLevelDownsampleFactorValue / GAODownsampleFactor);
+		SetShaderValue(RHICmdList, ShaderRHI, DownsampleFactorToBaseLevel, CurrentLevelDownsampleFactorValue / GAODownsampleFactor);
 
 		const FIntPoint DownsampledBufferSize(GSceneRenderTargets.GetBufferSizeXY() / FIntPoint(GAODownsampleFactor, GAODownsampleFactor));
 		const FVector2D BaseLevelBufferSizeValue(1.0f / DownsampledBufferSize.X, 1.0f / DownsampledBufferSize.Y);
-		SetShaderValue(ShaderRHI, BaseLevelTexelSize, BaseLevelBufferSizeValue);
+		SetShaderValue(RHICmdList, ShaderRHI, BaseLevelTexelSize, BaseLevelBufferSizeValue);
 	}
 
 	friend FArchive& operator<<(FArchive& Ar,FAOLevelParameters& P)
@@ -798,7 +799,7 @@ public:
 
 	static bool ShouldCache(EShaderPlatform Platform)
 	{
-		return IsFeatureLevelSupported(Platform, ERHIFeatureLevel::SM5) && DoesPlatformSupportDistanceFieldAO(Platform);
+		return IsFeatureLevelSupported(Platform, ERHIFeatureLevel::SM5) && DoesPlatformSupportDistanceFieldAO(Platform) && !IsOpenGLPlatform(Platform);
 	}
 
 	static void ModifyCompilationEnvironment(EShaderPlatform Platform, FShaderCompilerEnvironment& OutEnvironment)
@@ -826,24 +827,24 @@ public:
 		AOParameters.Bind(Initializer.ParameterMap);
 	}
 
-	void SetParameters(const FSceneView& View, int32 NumObjects)
+	void SetParameters(FRHICommandList* RHICmdList, const FSceneView& View, int32 NumObjects)
 	{
 		const FPixelShaderRHIParamRef ShaderRHI = GetPixelShader();
 
-		FGlobalShader::SetParameters(ShaderRHI, View);
-		ObjectParameters.Set(ShaderRHI, NumObjects);
-		AOParameters.Set(ShaderRHI);
-		DeferredParameters.Set(ShaderRHI, View);
+		FGlobalShader::SetParameters(RHICmdList, ShaderRHI, View);
+		ObjectParameters.Set(RHICmdList, ShaderRHI, NumObjects);
+		AOParameters.Set(RHICmdList, ShaderRHI);
+		DeferredParameters.Set(RHICmdList, ShaderRHI, View);
 
-		SetShaderValue(ShaderRHI, UVToTileHead, FVector2D(View.ViewRect.Size().X / (float)(GDistanceFieldAOTileSizeX * GAODownsampleFactor), View.ViewRect.Size().Y / (float)(GDistanceFieldAOTileSizeY * GAODownsampleFactor)));
+		SetShaderValue(RHICmdList, ShaderRHI, UVToTileHead, FVector2D(View.ViewRect.Size().X / (float)(GDistanceFieldAOTileSizeX * GAODownsampleFactor), View.ViewRect.Size().Y / (float)(GDistanceFieldAOTileSizeY * GAODownsampleFactor)));
 	
 		FTileIntersectionResources* TileIntersectionResources = ((FSceneViewState*)View.State)->AOTileIntersectionResources;
 
-		SetSRVParameter(ShaderRHI, TileHeadData, TileIntersectionResources->TileHeadData.SRV);
-		SetSRVParameter(ShaderRHI, TileHeadDataUnpacked, TileIntersectionResources->TileHeadDataUnpacked.SRV);
-		SetSRVParameter(ShaderRHI, TileArrayData, TileIntersectionResources->TileArrayData.SRV);
+		SetSRVParameter(RHICmdList, ShaderRHI, TileHeadData, TileIntersectionResources->TileHeadData.SRV);
+		SetSRVParameter(RHICmdList, ShaderRHI, TileHeadDataUnpacked, TileIntersectionResources->TileHeadDataUnpacked.SRV);
+		SetSRVParameter(RHICmdList, ShaderRHI, TileArrayData, TileIntersectionResources->TileArrayData.SRV);
 
-		SetShaderValue(ShaderRHI, TileListGroupSize, TileIntersectionResources->TileDimensions);
+		SetShaderValue(RHICmdList, ShaderRHI, TileListGroupSize, TileIntersectionResources->TileDimensions);
 	}
 	// FShader interface.
 	virtual bool Serialize(FArchive& Ar)
@@ -882,7 +883,7 @@ public:
 
 	static bool ShouldCache(EShaderPlatform Platform)
 	{
-		return IsFeatureLevelSupported(Platform, ERHIFeatureLevel::SM5) && DoesPlatformSupportDistanceFieldAO(Platform);
+		return IsFeatureLevelSupported(Platform, ERHIFeatureLevel::SM5) && DoesPlatformSupportDistanceFieldAO(Platform) && !IsOpenGLPlatform(Platform);
 	}
 
 	static void ModifyCompilationEnvironment(EShaderPlatform Platform, FShaderCompilerEnvironment& OutEnvironment)
@@ -911,32 +912,32 @@ public:
 		AOParameters.Bind(Initializer.ParameterMap);
 	}
 
-	void SetParameters(const FSceneView& View, FSceneRenderTargetItem& DistanceFieldNormalValue, int32 NumObjects)
+	void SetParameters(FRHICommandList* RHICmdList, const FSceneView& View, FSceneRenderTargetItem& DistanceFieldNormalValue, int32 NumObjects)
 	{
 		const FComputeShaderRHIParamRef ShaderRHI = GetComputeShader();
 
-		FGlobalShader::SetParameters(ShaderRHI, View);
+		FGlobalShader::SetParameters(RHICmdList, ShaderRHI, View);
 
-		DistanceFieldNormal.SetTexture(ShaderRHI, DistanceFieldNormalValue.ShaderResourceTexture, DistanceFieldNormalValue.UAV);
+		DistanceFieldNormal.SetTexture(RHICmdList, ShaderRHI, DistanceFieldNormalValue.ShaderResourceTexture, DistanceFieldNormalValue.UAV);
 
-		ObjectParameters.Set(ShaderRHI, NumObjects);
-		AOParameters.Set(ShaderRHI);
-		DeferredParameters.Set(ShaderRHI, View);
+		ObjectParameters.Set(RHICmdList, ShaderRHI, NumObjects);
+		AOParameters.Set(RHICmdList, ShaderRHI);
+		DeferredParameters.Set(RHICmdList, ShaderRHI, View);
 
-		SetShaderValue(ShaderRHI, UVToTileHead, FVector2D(View.ViewRect.Size().X / (float)(GDistanceFieldAOTileSizeX * GAODownsampleFactor), View.ViewRect.Size().Y / (float)(GDistanceFieldAOTileSizeY * GAODownsampleFactor)));
+		SetShaderValue(RHICmdList, ShaderRHI, UVToTileHead, FVector2D(View.ViewRect.Size().X / (float)(GDistanceFieldAOTileSizeX * GAODownsampleFactor), View.ViewRect.Size().Y / (float)(GDistanceFieldAOTileSizeY * GAODownsampleFactor)));
 	
 		FTileIntersectionResources* TileIntersectionResources = ((FSceneViewState*)View.State)->AOTileIntersectionResources;
 
-		SetSRVParameter(ShaderRHI, TileHeadData, TileIntersectionResources->TileHeadData.SRV);
-		SetSRVParameter(ShaderRHI, TileHeadDataUnpacked, TileIntersectionResources->TileHeadDataUnpacked.SRV);
-		SetSRVParameter(ShaderRHI, TileArrayData, TileIntersectionResources->TileArrayData.SRV);
+		SetSRVParameter(RHICmdList, ShaderRHI, TileHeadData, TileIntersectionResources->TileHeadData.SRV);
+		SetSRVParameter(RHICmdList, ShaderRHI, TileHeadDataUnpacked, TileIntersectionResources->TileHeadDataUnpacked.SRV);
+		SetSRVParameter(RHICmdList, ShaderRHI, TileArrayData, TileIntersectionResources->TileArrayData.SRV);
 
-		SetShaderValue(ShaderRHI, TileListGroupSize, TileIntersectionResources->TileDimensions);
+		SetShaderValue(RHICmdList, ShaderRHI, TileListGroupSize, TileIntersectionResources->TileDimensions);
 	}
 
-	void UnsetParameters()
+	void UnsetParameters(FRHICommandList* RHICmdList)
 	{
-		DistanceFieldNormal.UnsetUAV(GetComputeShader());
+		DistanceFieldNormal.UnsetUAV(RHICmdList, GetComputeShader());
 	}
 
 	// FShader interface.
@@ -972,6 +973,9 @@ IMPLEMENT_SHADER_TYPE(,FComputeDistanceFieldNormalCS,TEXT("DistanceFieldSurfaceC
 
 void ComputeDistanceFieldNormal(const TArray<FViewInfo>& Views, FSceneRenderTargetItem& DistanceFieldNormal, int32 NumObjects)
 {
+	//@todo-rco: RHIPacketList
+	FRHICommandList* RHICmdList = nullptr;
+
 	if (GAOComputeShaderNormalCalculation)
 	{
 		RHISetRenderTarget(NULL, NULL);
@@ -988,10 +992,10 @@ void ComputeDistanceFieldNormal(const TArray<FViewInfo>& Views, FSceneRenderTarg
 				TShaderMapRef<FComputeDistanceFieldNormalCS> ComputeShader(GetGlobalShaderMap());
 
 				RHISetComputeShader(ComputeShader->GetComputeShader());	
-				ComputeShader->SetParameters(View, DistanceFieldNormal, NumObjects);
-				DispatchComputeShader(*ComputeShader, GroupSizeX, GroupSizeY, 1);
+				ComputeShader->SetParameters(RHICmdList, View, DistanceFieldNormal, NumObjects);
+				DispatchComputeShader(RHICmdList, *ComputeShader, GroupSizeX, GroupSizeY, 1);
 
-				ComputeShader->UnsetParameters();
+				ComputeShader->UnsetParameters(RHICmdList);
 			}
 		}
 	}
@@ -1016,7 +1020,7 @@ void ComputeDistanceFieldNormal(const TArray<FViewInfo>& Views, FSceneRenderTarg
 			static FGlobalBoundShaderState BoundShaderState;
 			SetGlobalBoundShaderState(BoundShaderState, GFilterVertexDeclaration.VertexDeclarationRHI, *VertexShader, *PixelShader);
 
-			PixelShader->SetParameters(View, NumObjects);
+			PixelShader->SetParameters(RHICmdList, View, NumObjects);
 
 			DrawRectangle( 
 				0, 0, 
@@ -1037,7 +1041,7 @@ public:
 
 	static bool ShouldCache(EShaderPlatform Platform)
 	{
-		return IsFeatureLevelSupported(Platform, ERHIFeatureLevel::SM5) && DoesPlatformSupportDistanceFieldAO(Platform);
+		return IsFeatureLevelSupported(Platform, ERHIFeatureLevel::SM5) && DoesPlatformSupportDistanceFieldAO(Platform) && !IsOpenGLPlatform(Platform);
 	}
 
 	static void ModifyCompilationEnvironment(EShaderPlatform Platform, FShaderCompilerEnvironment& OutEnvironment)
@@ -1058,24 +1062,24 @@ public:
 	{
 	}
 
-	void SetParameters(const FSceneView& View, int32 DepthLevel)
+	void SetParameters(FRHICommandList* RHICmdList, const FSceneView& View, int32 DepthLevel)
 	{
 		FComputeShaderRHIParamRef ShaderRHI = GetComputeShader();
-		FGlobalShader::SetParameters(ShaderRHI, View);
+		FGlobalShader::SetParameters(RHICmdList, ShaderRHI, View);
 
 		const FScene* Scene = (const FScene*)View.Family->Scene;
 		FSurfaceCacheResources& SurfaceCacheResources = *Scene->SurfaceCacheResources;
 
-		SetSRVParameter(ShaderRHI, DrawParameters, SurfaceCacheResources.Level[DepthLevel]->ScatterDrawParameters.SRV);
+		SetSRVParameter(RHICmdList, ShaderRHI, DrawParameters, SurfaceCacheResources.Level[DepthLevel]->ScatterDrawParameters.SRV);
 
-		DispatchParameters.SetBuffer(ShaderRHI, SurfaceCacheResources.DispatchParameters);
+		DispatchParameters.SetBuffer(RHICmdList, ShaderRHI, SurfaceCacheResources.DispatchParameters);
 
-		SetShaderValue(ShaderRHI, TrimFraction, GAOTrimOldRecordsFraction);
+		SetShaderValue(RHICmdList, ShaderRHI, TrimFraction, GAOTrimOldRecordsFraction);
 	}
 
-	void UnsetParameters()
+	void UnsetParameters(FRHICommandList* RHICmdList)
 	{
-		DispatchParameters.UnsetUAV(GetComputeShader());
+		DispatchParameters.UnsetUAV(RHICmdList, GetComputeShader());
 	}
 
 	virtual bool Serialize(FArchive& Ar)
@@ -1104,7 +1108,7 @@ public:
 
 	static bool ShouldCache(EShaderPlatform Platform)
 	{
-		return IsFeatureLevelSupported(Platform, ERHIFeatureLevel::SM5) && DoesPlatformSupportDistanceFieldAO(Platform);
+		return IsFeatureLevelSupported(Platform, ERHIFeatureLevel::SM5) && DoesPlatformSupportDistanceFieldAO(Platform) && !IsOpenGLPlatform(Platform);
 	}
 
 	static void ModifyCompilationEnvironment(EShaderPlatform Platform, FShaderCompilerEnvironment& OutEnvironment)
@@ -1135,39 +1139,39 @@ public:
 	{
 	}
 
-	void SetParameters(const FSceneView& View, int32 DepthLevel)
+	void SetParameters(FRHICommandList* RHICmdList, const FSceneView& View, int32 DepthLevel)
 	{
 		FComputeShaderRHIParamRef ShaderRHI = GetComputeShader();
-		FGlobalShader::SetParameters(ShaderRHI, View);
+		FGlobalShader::SetParameters(RHICmdList, ShaderRHI, View);
 
 		const FScene* Scene = (const FScene*)View.Family->Scene;
 		FSurfaceCacheResources& SurfaceCacheResources = *Scene->SurfaceCacheResources;
 
-		SetSRVParameter(ShaderRHI, IrradianceCachePositionRadius, SurfaceCacheResources.Level[DepthLevel]->PositionAndRadius.SRV);
-		SetSRVParameter(ShaderRHI, IrradianceCacheNormal, SurfaceCacheResources.Level[DepthLevel]->Normal.SRV);
-		SetSRVParameter(ShaderRHI, OccluderRadius, SurfaceCacheResources.Level[DepthLevel]->OccluderRadius.SRV);
-		SetSRVParameter(ShaderRHI, IrradianceCacheBentNormal, SurfaceCacheResources.Level[DepthLevel]->BentNormal.SRV);
-		SetSRVParameter(ShaderRHI, IrradianceCacheTileCoordinate, SurfaceCacheResources.Level[DepthLevel]->TileCoordinate.SRV);
-		SetSRVParameter(ShaderRHI, DrawParameters, SurfaceCacheResources.Level[DepthLevel]->ScatterDrawParameters.SRV);
+		SetSRVParameter(RHICmdList, ShaderRHI, IrradianceCachePositionRadius, SurfaceCacheResources.Level[DepthLevel]->PositionAndRadius.SRV);
+		SetSRVParameter(RHICmdList, ShaderRHI, IrradianceCacheNormal, SurfaceCacheResources.Level[DepthLevel]->Normal.SRV);
+		SetSRVParameter(RHICmdList, ShaderRHI, OccluderRadius, SurfaceCacheResources.Level[DepthLevel]->OccluderRadius.SRV);
+		SetSRVParameter(RHICmdList, ShaderRHI, IrradianceCacheBentNormal, SurfaceCacheResources.Level[DepthLevel]->BentNormal.SRV);
+		SetSRVParameter(RHICmdList, ShaderRHI, IrradianceCacheTileCoordinate, SurfaceCacheResources.Level[DepthLevel]->TileCoordinate.SRV);
+		SetSRVParameter(RHICmdList, ShaderRHI, DrawParameters, SurfaceCacheResources.Level[DepthLevel]->ScatterDrawParameters.SRV);
 
-		CopyIrradianceCachePositionRadius.SetBuffer(ShaderRHI, SurfaceCacheResources.TempResources->PositionAndRadius);
-		CopyIrradianceCacheNormal.SetBuffer(ShaderRHI, SurfaceCacheResources.TempResources->Normal);
-		CopyOccluderRadius.SetBuffer(ShaderRHI, SurfaceCacheResources.TempResources->OccluderRadius);
-		CopyIrradianceCacheBentNormal.SetBuffer(ShaderRHI, SurfaceCacheResources.TempResources->BentNormal);
-		CopyIrradianceCacheTileCoordinate.SetBuffer(ShaderRHI, SurfaceCacheResources.TempResources->TileCoordinate);
-		ScatterDrawParameters.SetBuffer(ShaderRHI, SurfaceCacheResources.TempResources->ScatterDrawParameters);
+		CopyIrradianceCachePositionRadius.SetBuffer(RHICmdList, ShaderRHI, SurfaceCacheResources.TempResources->PositionAndRadius);
+		CopyIrradianceCacheNormal.SetBuffer(RHICmdList, ShaderRHI, SurfaceCacheResources.TempResources->Normal);
+		CopyOccluderRadius.SetBuffer(RHICmdList, ShaderRHI, SurfaceCacheResources.TempResources->OccluderRadius);
+		CopyIrradianceCacheBentNormal.SetBuffer(RHICmdList, ShaderRHI, SurfaceCacheResources.TempResources->BentNormal);
+		CopyIrradianceCacheTileCoordinate.SetBuffer(RHICmdList, ShaderRHI, SurfaceCacheResources.TempResources->TileCoordinate);
+		ScatterDrawParameters.SetBuffer(RHICmdList, ShaderRHI, SurfaceCacheResources.TempResources->ScatterDrawParameters);
 
-		SetShaderValue(ShaderRHI, TrimFraction, GAOTrimOldRecordsFraction);
+		SetShaderValue(RHICmdList, ShaderRHI, TrimFraction, GAOTrimOldRecordsFraction);
 	}
 
-	void UnsetParameters()
+	void UnsetParameters(FRHICommandList* RHICmdList)
 	{
-		CopyIrradianceCachePositionRadius.UnsetUAV(GetComputeShader());
-		CopyIrradianceCacheNormal.UnsetUAV(GetComputeShader());
-		CopyOccluderRadius.UnsetUAV(GetComputeShader());
-		CopyIrradianceCacheBentNormal.UnsetUAV(GetComputeShader());
-		CopyIrradianceCacheTileCoordinate.UnsetUAV(GetComputeShader());
-		ScatterDrawParameters.UnsetUAV(GetComputeShader());
+		CopyIrradianceCachePositionRadius.UnsetUAV(RHICmdList, GetComputeShader());
+		CopyIrradianceCacheNormal.UnsetUAV(RHICmdList, GetComputeShader());
+		CopyOccluderRadius.UnsetUAV(RHICmdList, GetComputeShader());
+		CopyIrradianceCacheBentNormal.UnsetUAV(RHICmdList, GetComputeShader());
+		CopyIrradianceCacheTileCoordinate.UnsetUAV(RHICmdList, GetComputeShader());
+		ScatterDrawParameters.UnsetUAV(RHICmdList, GetComputeShader());
 	}
 
 	virtual bool Serialize(FArchive& Ar)
@@ -1216,7 +1220,7 @@ public:
 
 	static bool ShouldCache(EShaderPlatform Platform)
 	{
-		return IsFeatureLevelSupported(Platform, ERHIFeatureLevel::SM5) && DoesPlatformSupportDistanceFieldAO(Platform);
+		return IsFeatureLevelSupported(Platform, ERHIFeatureLevel::SM5) && DoesPlatformSupportDistanceFieldAO(Platform) && !IsOpenGLPlatform(Platform);
 	}
 
 	FSaveStartIndexCS(const ShaderMetaType::CompiledShaderInitializerType& Initializer)
@@ -1230,21 +1234,21 @@ public:
 	{
 	}
 
-	void SetParameters(const FSceneView& View, int32 DepthLevel)
+	void SetParameters(FRHICommandList* RHICmdList, const FSceneView& View, int32 DepthLevel)
 	{
 		FComputeShaderRHIParamRef ShaderRHI = GetComputeShader();
-		FGlobalShader::SetParameters(ShaderRHI, View);
+		FGlobalShader::SetParameters(RHICmdList, ShaderRHI, View);
 
 		const FScene* Scene = (const FScene*)View.Family->Scene;
 		FSurfaceCacheResources& SurfaceCacheResources = *Scene->SurfaceCacheResources;
 
-		SetSRVParameter(ShaderRHI, DrawParameters, SurfaceCacheResources.Level[DepthLevel]->ScatterDrawParameters.SRV);
-		SavedStartIndex.SetBuffer(ShaderRHI, SurfaceCacheResources.Level[DepthLevel]->SavedStartIndex);
+		SetSRVParameter(RHICmdList, ShaderRHI, DrawParameters, SurfaceCacheResources.Level[DepthLevel]->ScatterDrawParameters.SRV);
+		SavedStartIndex.SetBuffer(RHICmdList, ShaderRHI, SurfaceCacheResources.Level[DepthLevel]->SavedStartIndex);
 	}
 
-	void UnsetParameters()
+	void UnsetParameters(FRHICommandList* RHICmdList)
 	{
-		SavedStartIndex.UnsetUAV(GetComputeShader());
+		SavedStartIndex.UnsetUAV(RHICmdList, GetComputeShader());
 	}
 
 	virtual bool Serialize(FArchive& Ar)
@@ -1270,7 +1274,7 @@ public:
 
 	static bool ShouldCache(EShaderPlatform Platform)
 	{
-		return IsFeatureLevelSupported(Platform, ERHIFeatureLevel::SM5) && DoesPlatformSupportDistanceFieldAO(Platform);
+		return IsFeatureLevelSupported(Platform, ERHIFeatureLevel::SM5) && DoesPlatformSupportDistanceFieldAO(Platform) && !IsOpenGLPlatform(Platform);
 	}
 
 	static void ModifyCompilationEnvironment(EShaderPlatform Platform, FShaderCompilerEnvironment& OutEnvironment)
@@ -1312,23 +1316,23 @@ public:
 	{
 	}
 
-	void SetParameters(const FSceneView& View, FSceneRenderTargetItem& DistanceFieldAOIrradianceCacheSplat, FSceneRenderTargetItem& DistanceFieldNormal, int32 DownsampleFactorValue, int32 DepthLevel, FIntPoint TileListGroupSizeValue)
+	void SetParameters(FRHICommandList* RHICmdList, const FSceneView& View, FSceneRenderTargetItem& DistanceFieldAOIrradianceCacheSplat, FSceneRenderTargetItem& DistanceFieldNormal, int32 DownsampleFactorValue, int32 DepthLevel, FIntPoint TileListGroupSizeValue)
 	{
 		FComputeShaderRHIParamRef ShaderRHI = GetComputeShader();
-		FGlobalShader::SetParameters(ShaderRHI, View);
-		DeferredParameters.Set(ShaderRHI, View);
-		AOParameters.Set(ShaderRHI);
-		AOLevelParameters.Set(ShaderRHI, View, DownsampleFactorValue);
+		FGlobalShader::SetParameters(RHICmdList, ShaderRHI, View);
+		DeferredParameters.Set(RHICmdList, ShaderRHI, View);
+		AOParameters.Set(RHICmdList, ShaderRHI);
+		AOLevelParameters.Set(RHICmdList, ShaderRHI, View, DownsampleFactorValue);
 
 		const FScene* Scene = (const FScene*)View.Family->Scene;
 		FSurfaceCacheResources& SurfaceCacheResources = *Scene->SurfaceCacheResources;
 
-		IrradianceCachePositionRadius.SetBuffer(ShaderRHI, SurfaceCacheResources.Level[DepthLevel]->PositionAndRadius);
-		IrradianceCacheNormal.SetBuffer(ShaderRHI, SurfaceCacheResources.Level[DepthLevel]->Normal);
-		IrradianceCacheTileCoordinate.SetBuffer(ShaderRHI, SurfaceCacheResources.Level[DepthLevel]->TileCoordinate);
-		ScatterDrawParameters.SetBuffer(ShaderRHI, SurfaceCacheResources.Level[DepthLevel]->ScatterDrawParameters);
+		IrradianceCachePositionRadius.SetBuffer(RHICmdList, ShaderRHI, SurfaceCacheResources.Level[DepthLevel]->PositionAndRadius);
+		IrradianceCacheNormal.SetBuffer(RHICmdList, ShaderRHI, SurfaceCacheResources.Level[DepthLevel]->Normal);
+		IrradianceCacheTileCoordinate.SetBuffer(RHICmdList, ShaderRHI, SurfaceCacheResources.Level[DepthLevel]->TileCoordinate);
+		ScatterDrawParameters.SetBuffer(RHICmdList, ShaderRHI, SurfaceCacheResources.Level[DepthLevel]->ScatterDrawParameters);
 
-		SetShaderValue(ShaderRHI, ViewDimensionsParameter, View.ViewRect);
+		SetShaderValue(RHICmdList, ShaderRHI, ViewDimensionsParameter, View.ViewRect);
 
 		FAOSampleData2 AOSampleData;
 
@@ -1340,14 +1344,14 @@ public:
 			AOSampleData.SampleDirections[SampleIndex] = FVector4(SampleDirections[SampleIndex]);
 		}
 
-		SetUniformBufferParameterImmediate(ShaderRHI, GetUniformBufferParameter<FAOSampleData2>(), AOSampleData);
+		SetUniformBufferParameterImmediate(RHICmdList, ShaderRHI, GetUniformBufferParameter<FAOSampleData2>(), AOSampleData);
 
 		FVector2D ThreadToCulledTileValue(DownsampleFactorValue / (float)(GAODownsampleFactor * GDistanceFieldAOTileSizeX), DownsampleFactorValue / (float)(GAODownsampleFactor * GDistanceFieldAOTileSizeY));
-		SetShaderValue(ShaderRHI, ThreadToCulledTile, ThreadToCulledTileValue);
+		SetShaderValue(RHICmdList, ShaderRHI, ThreadToCulledTile, ThreadToCulledTileValue);
 
-		SetShaderValue(ShaderRHI, TileListGroupSize, TileListGroupSizeValue);
+		SetShaderValue(RHICmdList, ShaderRHI, TileListGroupSize, TileListGroupSizeValue);
 
-		SetShaderValue(ShaderRHI, ConeHalfAngle, GAOConeHalfAngle);
+		SetShaderValue(RHICmdList, ShaderRHI, ConeHalfAngle, GAOConeHalfAngle);
 
 		FVector UnoccludedVector(0);
 
@@ -1357,15 +1361,16 @@ public:
 		}
 
 		float BentNormalNormalizeFactorValue = 1.0f / (UnoccludedVector / NumConeSampleDirections).Size();
-		SetShaderValue(ShaderRHI, BentNormalNormalizeFactor, BentNormalNormalizeFactorValue);
+		SetShaderValue(RHICmdList, ShaderRHI, BentNormalNormalizeFactor, BentNormalNormalizeFactorValue);
 
 		const int32 NumTexelsOneDimX = GDistanceFieldVolumeTextureAtlas.GetSizeX();
 		const int32 NumTexelsOneDimY = GDistanceFieldVolumeTextureAtlas.GetSizeY();
 		const int32 NumTexelsOneDimZ = GDistanceFieldVolumeTextureAtlas.GetSizeZ();
 		const FVector InvTextureDim(1.0f / NumTexelsOneDimX, 1.0f / NumTexelsOneDimY, 1.0f / NumTexelsOneDimZ);
-		SetShaderValue(ShaderRHI, DistanceFieldAtlasTexelSize, InvTextureDim);
+		SetShaderValue(RHICmdList, ShaderRHI, DistanceFieldAtlasTexelSize, InvTextureDim);
 
 		SetTextureParameter(
+			RHICmdList,
 			ShaderRHI,
 			DistanceFieldNormalTexture,
 			DistanceFieldNormalSampler,
@@ -1374,6 +1379,7 @@ public:
 			);
 
 		SetTextureParameter(
+			RHICmdList,
 			ShaderRHI,
 			IrradianceCacheSplatTexture,
 			IrradianceCacheSplatSampler,
@@ -1381,15 +1387,15 @@ public:
 			DistanceFieldAOIrradianceCacheSplat.ShaderResourceTexture
 			);
 
-		SetShaderValue(ShaderRHI, RecordRadiusScale, GAORecordRadiusScale);
+		SetShaderValue(RHICmdList, ShaderRHI, RecordRadiusScale, GAORecordRadiusScale);
 	}
 
-	void UnsetParameters()
+	void UnsetParameters(FRHICommandList* RHICmdList)
 	{
-		IrradianceCachePositionRadius.UnsetUAV(GetComputeShader());
-		IrradianceCacheNormal.UnsetUAV(GetComputeShader());
-		IrradianceCacheTileCoordinate.UnsetUAV(GetComputeShader());
-		ScatterDrawParameters.UnsetUAV(GetComputeShader());
+		IrradianceCachePositionRadius.UnsetUAV(RHICmdList, GetComputeShader());
+		IrradianceCacheNormal.UnsetUAV(RHICmdList, GetComputeShader());
+		IrradianceCacheTileCoordinate.UnsetUAV(RHICmdList, GetComputeShader());
+		ScatterDrawParameters.UnsetUAV(RHICmdList, GetComputeShader());
 	}
 
 	virtual bool Serialize(FArchive& Ar)
@@ -1447,7 +1453,7 @@ public:
 
 	static bool ShouldCache(EShaderPlatform Platform)
 	{
-		return IsFeatureLevelSupported(Platform, ERHIFeatureLevel::SM5) && DoesPlatformSupportDistanceFieldAO(Platform);
+		return IsFeatureLevelSupported(Platform, ERHIFeatureLevel::SM5) && DoesPlatformSupportDistanceFieldAO(Platform) && !IsOpenGLPlatform(Platform);
 	}
 
 	static void ModifyCompilationEnvironment(EShaderPlatform Platform, FShaderCompilerEnvironment& OutEnvironment)
@@ -1470,23 +1476,23 @@ public:
 	{
 	}
 
-	void SetParameters(const FSceneView& View, int32 DepthLevel)
+	void SetParameters(FRHICommandList* RHICmdList, const FSceneView& View, int32 DepthLevel)
 	{
 		FComputeShaderRHIParamRef ShaderRHI = GetComputeShader();
-		FGlobalShader::SetParameters(ShaderRHI, View);
+		FGlobalShader::SetParameters(RHICmdList, ShaderRHI, View);
 
 		const FScene* Scene = (const FScene*)View.Family->Scene;
 		FSurfaceCacheResources& SurfaceCacheResources = *Scene->SurfaceCacheResources;
 
-		SetSRVParameter(ShaderRHI, DrawParameters, SurfaceCacheResources.Level[DepthLevel]->ScatterDrawParameters.SRV);
-		SetSRVParameter(ShaderRHI, SavedStartIndex, SurfaceCacheResources.Level[DepthLevel]->SavedStartIndex.SRV);
+		SetSRVParameter(RHICmdList, ShaderRHI, DrawParameters, SurfaceCacheResources.Level[DepthLevel]->ScatterDrawParameters.SRV);
+		SetSRVParameter(RHICmdList, ShaderRHI, SavedStartIndex, SurfaceCacheResources.Level[DepthLevel]->SavedStartIndex.SRV);
 
-		DispatchParameters.SetBuffer(ShaderRHI, SurfaceCacheResources.DispatchParameters);
+		DispatchParameters.SetBuffer(RHICmdList, ShaderRHI, SurfaceCacheResources.DispatchParameters);
 	}
 
-	void UnsetParameters()
+	void UnsetParameters(FRHICommandList* RHICmdList)
 	{
-		DispatchParameters.UnsetUAV(GetComputeShader());
+		DispatchParameters.UnsetUAV(RHICmdList, GetComputeShader());
 	}
 
 	virtual bool Serialize(FArchive& Ar)
@@ -1514,7 +1520,7 @@ public:
 
 	static bool ShouldCache(EShaderPlatform Platform)
 	{
-		return IsFeatureLevelSupported(Platform, ERHIFeatureLevel::SM5) && DoesPlatformSupportDistanceFieldAO(Platform);
+		return IsFeatureLevelSupported(Platform, ERHIFeatureLevel::SM5) && DoesPlatformSupportDistanceFieldAO(Platform) && !IsOpenGLPlatform(Platform);
 	}
 
 	static void ModifyCompilationEnvironment(EShaderPlatform Platform, FShaderCompilerEnvironment& OutEnvironment)
@@ -1562,27 +1568,27 @@ public:
 	{
 	}
 
-	void SetParameters(const FSceneView& View, FSceneRenderTargetItem& DistanceFieldAOIrradianceCacheSplat, FSceneRenderTargetItem& DistanceFieldNormal, int32 NumObjects, int32 DownsampleFactorValue, int32 DepthLevel, FIntPoint TileListGroupSizeValue)
+	void SetParameters(FRHICommandList* RHICmdList, const FSceneView& View, FSceneRenderTargetItem& DistanceFieldAOIrradianceCacheSplat, FSceneRenderTargetItem& DistanceFieldNormal, int32 NumObjects, int32 DownsampleFactorValue, int32 DepthLevel, FIntPoint TileListGroupSizeValue)
 	{
 		FComputeShaderRHIParamRef ShaderRHI = GetComputeShader();
-		FGlobalShader::SetParameters(ShaderRHI, View);
-		DeferredParameters.Set(ShaderRHI, View);
-		ObjectParameters.Set(ShaderRHI, NumObjects);
-		AOParameters.Set(ShaderRHI);
-		AOLevelParameters.Set(ShaderRHI, View, DownsampleFactorValue);
+		FGlobalShader::SetParameters(RHICmdList, ShaderRHI, View);
+		DeferredParameters.Set(RHICmdList, ShaderRHI, View);
+		ObjectParameters.Set(RHICmdList, ShaderRHI, NumObjects);
+		AOParameters.Set(RHICmdList, ShaderRHI);
+		AOLevelParameters.Set(RHICmdList, ShaderRHI, View, DownsampleFactorValue);
 
 		const FScene* Scene = (const FScene*)View.Family->Scene;
 		FSurfaceCacheResources& SurfaceCacheResources = *Scene->SurfaceCacheResources;
 
-		SetSRVParameter(ShaderRHI, IrradianceCacheNormal, SurfaceCacheResources.Level[DepthLevel]->Normal.SRV);
-		SetSRVParameter(ShaderRHI, ScatterDrawParameters, SurfaceCacheResources.Level[DepthLevel]->ScatterDrawParameters.SRV);
-		SetSRVParameter(ShaderRHI, SavedStartIndex, SurfaceCacheResources.Level[DepthLevel]->SavedStartIndex.SRV);
-		SetSRVParameter(ShaderRHI, IrradianceCachePositionRadius, SurfaceCacheResources.Level[DepthLevel]->PositionAndRadius.SRV);
+		SetSRVParameter(RHICmdList, ShaderRHI, IrradianceCacheNormal, SurfaceCacheResources.Level[DepthLevel]->Normal.SRV);
+		SetSRVParameter(RHICmdList, ShaderRHI, ScatterDrawParameters, SurfaceCacheResources.Level[DepthLevel]->ScatterDrawParameters.SRV);
+		SetSRVParameter(RHICmdList, ShaderRHI, SavedStartIndex, SurfaceCacheResources.Level[DepthLevel]->SavedStartIndex.SRV);
+		SetSRVParameter(RHICmdList, ShaderRHI, IrradianceCachePositionRadius, SurfaceCacheResources.Level[DepthLevel]->PositionAndRadius.SRV);
 
-		OccluderRadius.SetBuffer(ShaderRHI, SurfaceCacheResources.Level[DepthLevel]->OccluderRadius);
-		IrradianceCacheBentNormal.SetBuffer(ShaderRHI, SurfaceCacheResources.Level[DepthLevel]->BentNormal);
+		OccluderRadius.SetBuffer(RHICmdList, ShaderRHI, SurfaceCacheResources.Level[DepthLevel]->OccluderRadius);
+		IrradianceCacheBentNormal.SetBuffer(RHICmdList, ShaderRHI, SurfaceCacheResources.Level[DepthLevel]->BentNormal);
 
-		SetShaderValue(ShaderRHI, ViewDimensionsParameter, View.ViewRect);
+		SetShaderValue(RHICmdList, ShaderRHI, ViewDimensionsParameter, View.ViewRect);
 
 		FAOSampleData2 AOSampleData;
 
@@ -1594,21 +1600,21 @@ public:
 			AOSampleData.SampleDirections[SampleIndex] = FVector4(SampleDirections[SampleIndex]);
 		}
 
-		SetUniformBufferParameterImmediate(ShaderRHI, GetUniformBufferParameter<FAOSampleData2>(), AOSampleData);
+		SetUniformBufferParameterImmediate(RHICmdList, ShaderRHI, GetUniformBufferParameter<FAOSampleData2>(), AOSampleData);
 
 		FVector2D ThreadToCulledTileValue(DownsampleFactorValue / (float)(GAODownsampleFactor * GDistanceFieldAOTileSizeX), DownsampleFactorValue / (float)(GAODownsampleFactor * GDistanceFieldAOTileSizeY));
-		SetShaderValue(ShaderRHI, ThreadToCulledTile, ThreadToCulledTileValue);
+		SetShaderValue(RHICmdList, ShaderRHI, ThreadToCulledTile, ThreadToCulledTileValue);
 
 		FTileIntersectionResources* TileIntersectionResources = ((FSceneViewState*)View.State)->AOTileIntersectionResources;
 
-		SetSRVParameter(ShaderRHI, TileHeadDataUnpacked, TileIntersectionResources->TileHeadDataUnpacked.SRV);
-		SetSRVParameter(ShaderRHI, TileHeadData, TileIntersectionResources->TileHeadData.SRV);
-		SetSRVParameter(ShaderRHI, TileArrayData, TileIntersectionResources->TileArrayData.SRV);
-		SetSRVParameter(ShaderRHI, IrradianceCacheTileCoordinate, SurfaceCacheResources.Level[DepthLevel]->TileCoordinate.SRV);
+		SetSRVParameter(RHICmdList, ShaderRHI, TileHeadDataUnpacked, TileIntersectionResources->TileHeadDataUnpacked.SRV);
+		SetSRVParameter(RHICmdList, ShaderRHI, TileHeadData, TileIntersectionResources->TileHeadData.SRV);
+		SetSRVParameter(RHICmdList, ShaderRHI, TileArrayData, TileIntersectionResources->TileArrayData.SRV);
+		SetSRVParameter(RHICmdList, ShaderRHI, IrradianceCacheTileCoordinate, SurfaceCacheResources.Level[DepthLevel]->TileCoordinate.SRV);
 
-		SetShaderValue(ShaderRHI, TileListGroupSize, TileListGroupSizeValue);
+		SetShaderValue(RHICmdList, ShaderRHI, TileListGroupSize, TileListGroupSizeValue);
 
-		SetShaderValue(ShaderRHI, ConeHalfAngle, GAOConeHalfAngle);
+		SetShaderValue(RHICmdList, ShaderRHI, ConeHalfAngle, GAOConeHalfAngle);
 
 		FVector UnoccludedVector(0);
 
@@ -1618,15 +1624,16 @@ public:
 		}
 
 		float BentNormalNormalizeFactorValue = 1.0f / (UnoccludedVector / NumConeSampleDirections).Size();
-		SetShaderValue(ShaderRHI, BentNormalNormalizeFactor, BentNormalNormalizeFactorValue);
+		SetShaderValue(RHICmdList, ShaderRHI, BentNormalNormalizeFactor, BentNormalNormalizeFactorValue);
 
 		const int32 NumTexelsOneDimX = GDistanceFieldVolumeTextureAtlas.GetSizeX();
 		const int32 NumTexelsOneDimY = GDistanceFieldVolumeTextureAtlas.GetSizeY();
 		const int32 NumTexelsOneDimZ = GDistanceFieldVolumeTextureAtlas.GetSizeZ();
 		const FVector InvTextureDim(1.0f / NumTexelsOneDimX, 1.0f / NumTexelsOneDimY, 1.0f / NumTexelsOneDimZ);
-		SetShaderValue(ShaderRHI, DistanceFieldAtlasTexelSize, InvTextureDim);
+		SetShaderValue(RHICmdList, ShaderRHI, DistanceFieldAtlasTexelSize, InvTextureDim);
 
 		SetTextureParameter(
+			RHICmdList,
 			ShaderRHI,
 			DistanceFieldNormalTexture,
 			DistanceFieldNormalSampler,
@@ -1635,6 +1642,7 @@ public:
 			);
 
 		SetTextureParameter(
+			RHICmdList,
 			ShaderRHI,
 			IrradianceCacheSplatTexture,
 			IrradianceCacheSplatSampler,
@@ -1642,13 +1650,13 @@ public:
 			DistanceFieldAOIrradianceCacheSplat.ShaderResourceTexture
 			);
 
-		SetShaderValue(ShaderRHI, RecordRadiusScale, GAORecordRadiusScale);
+		SetShaderValue(RHICmdList, ShaderRHI, RecordRadiusScale, GAORecordRadiusScale);
 	}
 
-	void UnsetParameters()
+	void UnsetParameters(FRHICommandList* RHICmdList)
 	{
-		OccluderRadius.UnsetUAV(GetComputeShader());
-		IrradianceCacheBentNormal.UnsetUAV(GetComputeShader());
+		OccluderRadius.UnsetUAV(RHICmdList, GetComputeShader());
+		IrradianceCacheBentNormal.UnsetUAV(RHICmdList, GetComputeShader());
 	}
 
 	virtual bool Serialize(FArchive& Ar)
@@ -1722,7 +1730,7 @@ public:
 
 	static bool ShouldCache(EShaderPlatform Platform) 
 	{ 
-		return IsFeatureLevelSupported(Platform, ERHIFeatureLevel::SM5) && DoesPlatformSupportDistanceFieldAO(Platform); 
+		return IsFeatureLevelSupported(Platform, ERHIFeatureLevel::SM5) && DoesPlatformSupportDistanceFieldAO(Platform) && !IsOpenGLPlatform(Platform);
 	}
 
 	static void ModifyCompilationEnvironment(EShaderPlatform Platform, FShaderCompilerEnvironment& OutEnvironment)
@@ -1746,25 +1754,25 @@ public:
 
 	TIrradianceCacheSplatVS() {}
 
-	void SetParameters(const FSceneView& View, int32 DepthLevel, int32 CurrentLevelDownsampleFactorValue, FVector2D NormalizedOffsetToPixelCenterValue)
+	void SetParameters(FRHICommandList* RHICmdList, const FSceneView& View, int32 DepthLevel, int32 CurrentLevelDownsampleFactorValue, FVector2D NormalizedOffsetToPixelCenterValue)
 	{
 		const FVertexShaderRHIParamRef ShaderRHI = GetVertexShader();
-		FGlobalShader::SetParameters(ShaderRHI, View);
+		FGlobalShader::SetParameters(RHICmdList, ShaderRHI, View);
 		
 		const FScene* Scene = (const FScene*)View.Family->Scene;
 		FSurfaceCacheResources& SurfaceCacheResources = *Scene->SurfaceCacheResources;
 
-		SetSRVParameter(ShaderRHI, IrradianceCachePositionRadius, SurfaceCacheResources.Level[DepthLevel]->PositionAndRadius.SRV);
-		SetSRVParameter(ShaderRHI, IrradianceCacheOccluderRadius, SurfaceCacheResources.Level[DepthLevel]->OccluderRadius.SRV);
-		SetSRVParameter(ShaderRHI, IrradianceCacheNormal, SurfaceCacheResources.Level[DepthLevel]->Normal.SRV);
-		SetSRVParameter(ShaderRHI, IrradianceCacheBentNormal, SurfaceCacheResources.Level[DepthLevel]->BentNormal.SRV);
+		SetSRVParameter(RHICmdList, ShaderRHI, IrradianceCachePositionRadius, SurfaceCacheResources.Level[DepthLevel]->PositionAndRadius.SRV);
+		SetSRVParameter(RHICmdList, ShaderRHI, IrradianceCacheOccluderRadius, SurfaceCacheResources.Level[DepthLevel]->OccluderRadius.SRV);
+		SetSRVParameter(RHICmdList, ShaderRHI, IrradianceCacheNormal, SurfaceCacheResources.Level[DepthLevel]->Normal.SRV);
+		SetSRVParameter(RHICmdList, ShaderRHI, IrradianceCacheBentNormal, SurfaceCacheResources.Level[DepthLevel]->BentNormal.SRV);
 
-		SetShaderValue(ShaderRHI, InterpolationRadiusScale, (bFinalInterpolationPass ? GAOInterpolationRadiusScale : 1.0f));
-		SetShaderValue(ShaderRHI, NormalizedOffsetToPixelCenter, NormalizedOffsetToPixelCenterValue);
+		SetShaderValue(RHICmdList, ShaderRHI, InterpolationRadiusScale, (bFinalInterpolationPass ? GAOInterpolationRadiusScale : 1.0f));
+		SetShaderValue(RHICmdList, ShaderRHI, NormalizedOffsetToPixelCenter, NormalizedOffsetToPixelCenterValue);
 
 		const FIntPoint AOViewRectSize = FIntPoint::DivideAndRoundUp(View.ViewRect.Size(), CurrentLevelDownsampleFactorValue);
 		const FVector2D HackExpandValue(.5f / AOViewRectSize.X, .5f / AOViewRectSize.Y);
-		SetShaderValue(ShaderRHI, HackExpand, HackExpandValue);
+		SetShaderValue(RHICmdList, ShaderRHI, HackExpand, HackExpandValue);
 	}
 
 	virtual bool Serialize(FArchive& Ar)
@@ -1801,7 +1809,7 @@ public:
 
 	static bool ShouldCache(EShaderPlatform Platform)
 	{
-		return IsFeatureLevelSupported(Platform, ERHIFeatureLevel::SM5) && DoesPlatformSupportDistanceFieldAO(Platform);
+		return IsFeatureLevelSupported(Platform, ERHIFeatureLevel::SM5) && DoesPlatformSupportDistanceFieldAO(Platform) && !IsOpenGLPlatform(Platform);
 	}
 
 	static void ModifyCompilationEnvironment(EShaderPlatform Platform, FShaderCompilerEnvironment& OutEnvironment)
@@ -1828,15 +1836,16 @@ public:
 		InvMinCosPointBehindPlane.Bind(Initializer.ParameterMap, TEXT("InvMinCosPointBehindPlane"));
 	}
 
-	void SetParameters(const FSceneView& View, FSceneRenderTargetItem& DistanceFieldNormal, int32 DestLevelDownsampleFactor)
+	void SetParameters(FRHICommandList* RHICmdList, const FSceneView& View, FSceneRenderTargetItem& DistanceFieldNormal, int32 DestLevelDownsampleFactor)
 	{
 		const FPixelShaderRHIParamRef ShaderRHI = GetPixelShader();
-		FGlobalShader::SetParameters(ShaderRHI, View);
-		AOParameters.Set(ShaderRHI);
-		DeferredParameters.Set(ShaderRHI, View);
-		AOLevelParameters.Set(ShaderRHI, View, DestLevelDownsampleFactor);
+		FGlobalShader::SetParameters(RHICmdList, ShaderRHI, View);
+		AOParameters.Set(RHICmdList, ShaderRHI);
+		DeferredParameters.Set(RHICmdList, ShaderRHI, View);
+		AOLevelParameters.Set(RHICmdList, ShaderRHI, View, DestLevelDownsampleFactor);
 
 		SetTextureParameter(
+			RHICmdList,
 			ShaderRHI,
 			DistanceFieldNormalTexture,
 			DistanceFieldNormalSampler,
@@ -1846,10 +1855,10 @@ public:
 
 		const float EffectiveMaxAngle = bFinalInterpolationPass ? GAOInterpolationMaxAngle * GAOInterpolationAngleScale : GAOInterpolationMaxAngle;
 		const float InterpolationAngleNormalizationValue = 1.0f / FMath::Sqrt(1.0f - FMath::Cos(EffectiveMaxAngle * PI / 180.0f));
-		SetShaderValue(ShaderRHI, InterpolationAngleNormalization, InterpolationAngleNormalizationValue);
+		SetShaderValue(RHICmdList, ShaderRHI, InterpolationAngleNormalization, InterpolationAngleNormalizationValue);
 
 		const float MinCosPointBehindPlaneValue = FMath::Cos((GAOMinPointBehindPlaneAngle + 90.0f) * PI / 180.0f);
-		SetShaderValue(ShaderRHI, InvMinCosPointBehindPlane, 1.0f / MinCosPointBehindPlaneValue);
+		SetShaderValue(RHICmdList, ShaderRHI, InvMinCosPointBehindPlane, 1.0f / MinCosPointBehindPlaneValue);
 	}
 
 	// FShader interface.
@@ -1886,7 +1895,7 @@ public:
 
 	static bool ShouldCache(EShaderPlatform Platform)
 	{
-		return IsFeatureLevelSupported(Platform, ERHIFeatureLevel::SM5) && DoesPlatformSupportDistanceFieldAO(Platform);
+		return IsFeatureLevelSupported(Platform, ERHIFeatureLevel::SM5) && DoesPlatformSupportDistanceFieldAO(Platform) && !IsOpenGLPlatform(Platform);
 	}
 
 	static void ModifyCompilationEnvironment(EShaderPlatform Platform, FShaderCompilerEnvironment& OutEnvironment)
@@ -1908,16 +1917,16 @@ public:
 		BentNormalAOSampler3.Bind(Initializer.ParameterMap,TEXT("BentNormalAOSampler3"));
 	}
 
-	void SetParameters(const FSceneView& View, FSceneRenderTargetItem& InBentNormalTexture)
+	void SetParameters(FRHICommandList* RHICmdList, const FSceneView& View, FSceneRenderTargetItem& InBentNormalTexture)
 	{
 		const FPixelShaderRHIParamRef ShaderRHI = GetPixelShader();
 
-		FGlobalShader::SetParameters(ShaderRHI, View);
-		DeferredParameters.Set(ShaderRHI, View);
-		AOParameters.Set(ShaderRHI);
-		AOLevelParameters.Set(ShaderRHI, View, GAODownsampleFactor);
+		FGlobalShader::SetParameters(RHICmdList, ShaderRHI, View);
+		DeferredParameters.Set(RHICmdList, ShaderRHI, View);
+		AOParameters.Set(RHICmdList, ShaderRHI);
+		AOLevelParameters.Set(RHICmdList, ShaderRHI, View, GAODownsampleFactor);
 
-		SetTextureParameter(ShaderRHI, BentNormalAOTexture3, BentNormalAOSampler3, TStaticSamplerState<SF_Point>::GetRHI(), InBentNormalTexture.ShaderResourceTexture);	
+		SetTextureParameter(RHICmdList, ShaderRHI, BentNormalAOTexture3, BentNormalAOSampler3, TStaticSamplerState<SF_Point>::GetRHI(), InBentNormalTexture.ShaderResourceTexture);
 	}
 	// FShader interface.
 	virtual bool Serialize(FArchive& Ar)
@@ -1949,7 +1958,7 @@ public:
 
 	static bool ShouldCache(EShaderPlatform Platform)
 	{
-		return IsFeatureLevelSupported(Platform, ERHIFeatureLevel::SM5) && DoesPlatformSupportDistanceFieldAO(Platform);
+		return IsFeatureLevelSupported(Platform, ERHIFeatureLevel::SM5) && DoesPlatformSupportDistanceFieldAO(Platform) && !IsOpenGLPlatform(Platform);
 	}
 
 	static void ModifyCompilationEnvironment(EShaderPlatform Platform, FShaderCompilerEnvironment& OutEnvironment)
@@ -1969,13 +1978,14 @@ public:
 		BentNormalAOTexelSize.Bind(Initializer.ParameterMap, TEXT("BentNormalAOTexelSize"));
 	}
 
-	void SetParameters(const FSceneView& View, FSceneRenderTargetItem& DistanceFieldAOBentNormal2)
+	void SetParameters(FRHICommandList* RHICmdList, const FSceneView& View, FSceneRenderTargetItem& DistanceFieldAOBentNormal2)
 	{
 		const FPixelShaderRHIParamRef ShaderRHI = GetPixelShader();
 
-		FGlobalShader::SetParameters(ShaderRHI, View);
+		FGlobalShader::SetParameters(RHICmdList, ShaderRHI, View);
 
 		SetTextureParameter(
+			RHICmdList,
 			ShaderRHI,
 			BentNormalAOTexture,
 			BentNormalAOSampler,
@@ -1985,7 +1995,7 @@ public:
 
 		const FIntPoint DownsampledBufferSize(GSceneRenderTargets.GetBufferSizeXY() / FIntPoint(GAODownsampleFactor, GAODownsampleFactor));
 		const FVector2D BaseLevelTexelSizeValue(1.0f / DownsampledBufferSize.X, 1.0f / DownsampledBufferSize.Y);
-		SetShaderValue(ShaderRHI, BentNormalAOTexelSize, BaseLevelTexelSizeValue);
+		SetShaderValue(RHICmdList, ShaderRHI, BentNormalAOTexelSize, BaseLevelTexelSizeValue);
 	}
 	// FShader interface.
 	virtual bool Serialize(FArchive& Ar)
@@ -2014,7 +2024,7 @@ public:
 
 	static bool ShouldCache(EShaderPlatform Platform)
 	{
-		return IsFeatureLevelSupported(Platform, ERHIFeatureLevel::SM5) && DoesPlatformSupportDistanceFieldAO(Platform);
+		return IsFeatureLevelSupported(Platform, ERHIFeatureLevel::SM5) && DoesPlatformSupportDistanceFieldAO(Platform) && !IsOpenGLPlatform(Platform);
 	}
 
 	static void ModifyCompilationEnvironment(EShaderPlatform Platform, FShaderCompilerEnvironment& OutEnvironment)
@@ -2038,14 +2048,15 @@ public:
 		HistoryWeight.Bind(Initializer.ParameterMap, TEXT("HistoryWeight"));
 	}
 
-	void SetParameters(const FSceneView& View, FSceneRenderTargetItem& HistoryTextureValue, FSceneRenderTargetItem& DistanceFieldAOBentNormal)
+	void SetParameters(FRHICommandList* RHICmdList, const FSceneView& View, FSceneRenderTargetItem& HistoryTextureValue, FSceneRenderTargetItem& DistanceFieldAOBentNormal)
 	{
 		const FPixelShaderRHIParamRef ShaderRHI = GetPixelShader();
 
-		FGlobalShader::SetParameters(ShaderRHI, View);
-		DeferredParameters.Set(ShaderRHI, View);
+		FGlobalShader::SetParameters(RHICmdList, ShaderRHI, View);
+		DeferredParameters.Set(RHICmdList, ShaderRHI, View);
 
 		SetTextureParameter(
+			RHICmdList,
 			ShaderRHI,
 			BentNormalAOTexture,
 			BentNormalAOSampler,
@@ -2054,6 +2065,7 @@ public:
 			);
 
 		SetTextureParameter(
+			RHICmdList,
 			ShaderRHI,
 			HistoryTexture,
 			HistorySampler,
@@ -2061,9 +2073,9 @@ public:
 			HistoryTextureValue.ShaderResourceTexture
 			);
 
-		CameraMotion.Set(View, ShaderRHI);
+		CameraMotion.Set(RHICmdList, View, ShaderRHI);
 
-		SetShaderValue(ShaderRHI, HistoryWeight, GAOHistoryWeight);
+		SetShaderValue(RHICmdList, ShaderRHI, HistoryWeight, GAOHistoryWeight);
 	}
 	// FShader interface.
 	virtual bool Serialize(FArchive& Ar)
@@ -2114,6 +2126,9 @@ void UpdateHistory(
 	/** Output of Temporal Reprojection for the next step in the pipeline. */
 	TRefCountPtr<IPooledRenderTarget>& HistoryOutput)
 {
+	//@todo-rco: RHIPacketList
+	FRHICommandList* RHICmdList = nullptr;
+
 	if (HistoryState)
 	{
 		if (*HistoryState && !View.bCameraCut && !View.bPrevTransformsReset)
@@ -2137,7 +2152,7 @@ void UpdateHistory(
 				static FGlobalBoundShaderState BoundShaderState;
 				SetGlobalBoundShaderState(BoundShaderState, GFilterVertexDeclaration.VertexDeclarationRHI, *VertexShader, *PixelShader);
 
-				PixelShader->SetParameters(View, (*HistoryState)->GetRenderTargetItem(), AOSource->GetRenderTargetItem());
+				PixelShader->SetParameters(RHICmdList, View, (*HistoryState)->GetRenderTargetItem(), AOSource->GetRenderTargetItem());
 
 				DrawRectangle( 
 					0, 0, 
@@ -2170,6 +2185,9 @@ void UpdateHistory(
 
 void PostProcessBentNormalAO(TArray<FViewInfo>& Views, FSceneRenderTargetItem& IrradianceCacheInterpolation, TRefCountPtr<IPooledRenderTarget>& AOOutput)
 {
+	//@todo-rco: RHIPacketList
+	FRHICommandList* RHICmdList = nullptr;
+
 	TRefCountPtr<IPooledRenderTarget> DistanceFieldAOBentNormal;
 	AllocateOrReuseAORenderTarget(DistanceFieldAOBentNormal, TEXT("DistanceFieldBentNormalAO"));
 
@@ -2202,7 +2220,7 @@ void PostProcessBentNormalAO(TArray<FViewInfo>& Views, FSceneRenderTargetItem& I
 			static FGlobalBoundShaderState BoundShaderState;
 			SetGlobalBoundShaderState(BoundShaderState, GFilterVertexDeclaration.VertexDeclarationRHI, *VertexShader, *PixelShader);
 
-			PixelShader->SetParameters(View, IrradianceCacheInterpolation);
+			PixelShader->SetParameters(RHICmdList, View, IrradianceCacheInterpolation);
 
 			//@todo - get the round up correct
 			DrawRectangle( 
@@ -2236,8 +2254,8 @@ void PostProcessBentNormalAO(TArray<FViewInfo>& Views, FSceneRenderTargetItem& I
 			static FGlobalBoundShaderState BoundShaderState;
 			SetGlobalBoundShaderState(BoundShaderState, GFilterVertexDeclaration.VertexDeclarationRHI, *VertexShader, *PixelShader);
 
-			VertexShader->SetParameters(View);
-			PixelShader->SetParameters(View, DistanceFieldAOBentNormal2->GetRenderTargetItem());
+			VertexShader->SetParameters(RHICmdList, View);
+			PixelShader->SetParameters(RHICmdList, View, DistanceFieldAOBentNormal2->GetRenderTargetItem());
 
 			DrawRectangle( 
 				0, 0, 
@@ -2274,7 +2292,7 @@ public:
 
 	static bool ShouldCache(EShaderPlatform Platform)
 	{
-		return IsFeatureLevelSupported(Platform, ERHIFeatureLevel::SM5) && DoesPlatformSupportDistanceFieldAO(Platform);
+		return IsFeatureLevelSupported(Platform, ERHIFeatureLevel::SM5) && DoesPlatformSupportDistanceFieldAO(Platform) && !IsOpenGLPlatform(Platform);
 	}
 
 	static void ModifyCompilationEnvironment(EShaderPlatform Platform, FShaderCompilerEnvironment& OutEnvironment)
@@ -2294,14 +2312,14 @@ public:
 		BentNormalAOSampler.Bind(Initializer.ParameterMap,TEXT("BentNormalAOSampler"));
 	}
 
-	void SetParameters(const FSceneView& View, TRefCountPtr<IPooledRenderTarget>& DistanceFieldAOBentNormal)
+	void SetParameters(FRHICommandList* RHICmdList, const FSceneView& View, TRefCountPtr<IPooledRenderTarget>& DistanceFieldAOBentNormal)
 	{
 		const FPixelShaderRHIParamRef ShaderRHI = GetPixelShader();
 
-		FGlobalShader::SetParameters(ShaderRHI, View);
-		DeferredParameters.Set(ShaderRHI, View);
+		FGlobalShader::SetParameters(RHICmdList, ShaderRHI, View);
+		DeferredParameters.Set(RHICmdList, ShaderRHI, View);
 
-		SetTextureParameter(ShaderRHI, BentNormalAOTexture, BentNormalAOSampler, TStaticSamplerState<SF_Bilinear>::GetRHI(), DistanceFieldAOBentNormal->GetRenderTargetItem().ShaderResourceTexture);	
+		SetTextureParameter(RHICmdList, ShaderRHI, BentNormalAOTexture, BentNormalAOSampler, TStaticSamplerState<SF_Bilinear>::GetRHI(), DistanceFieldAOBentNormal->GetRenderTargetItem().ShaderResourceTexture);
 	}
 	// FShader interface.
 	virtual bool Serialize(FArchive& Ar)
@@ -2324,6 +2342,9 @@ IMPLEMENT_SHADER_TYPE(,FDistanceFieldAOUpsamplePS,TEXT("DistanceFieldSurfaceCach
 
 void UpsampleBentNormalAO(const TArray<FViewInfo>& Views, TRefCountPtr<IPooledRenderTarget>& DistanceFieldAOBentNormal)
 {
+	//@todo-rco: RHIPacketList
+	FRHICommandList* RHICmdList = nullptr;
+
 	for (int32 ViewIndex = 0; ViewIndex < Views.Num(); ViewIndex++)
 	{
 		const FViewInfo& View = Views[ViewIndex];
@@ -2341,7 +2362,7 @@ void UpsampleBentNormalAO(const TArray<FViewInfo>& Views, TRefCountPtr<IPooledRe
 		static FGlobalBoundShaderState BoundShaderState;
 		SetGlobalBoundShaderState(BoundShaderState, GFilterVertexDeclaration.VertexDeclarationRHI, *VertexShader, *PixelShader);
 
-		PixelShader->SetParameters(View, DistanceFieldAOBentNormal);
+		PixelShader->SetParameters(RHICmdList, View, DistanceFieldAOBentNormal);
 
 		DrawRectangle( 
 			0, 0, 
@@ -2513,6 +2534,9 @@ void GenerateBestSpacedVectors()
 
 FIntPoint BuildTileObjectLists(FScene* Scene, TArray<FViewInfo>& Views, int32 NumObjects)
 {
+	//@todo-rco: RHIPacketList
+	FRHICommandList* RHICmdList = nullptr;
+
 	SCOPED_DRAW_EVENT(BuildTileList, DEC_SCENE_ITEMS);
 	RHISetRenderTarget(NULL, NULL);
 
@@ -2551,10 +2575,10 @@ FIntPoint BuildTileObjectLists(FScene* Scene, TArray<FViewInfo>& Views, int32 Nu
 				TShaderMapRef<FBuildTileConesCS> ComputeShader(GetGlobalShaderMap());
 
 				RHISetComputeShader(ComputeShader->GetComputeShader());	
-				ComputeShader->SetParameters(View, Scene, FVector2D(GroupSizeX, GroupSizeY));
-				DispatchComputeShader(*ComputeShader, GroupSizeX, GroupSizeY, 1);
+				ComputeShader->SetParameters(RHICmdList, View, Scene, FVector2D(GroupSizeX, GroupSizeY));
+				DispatchComputeShader(RHICmdList, *ComputeShader, GroupSizeX, GroupSizeY, 1);
 
-				ComputeShader->UnsetParameters();
+				ComputeShader->UnsetParameters(RHICmdList);
 			}
 
 			{
@@ -2577,8 +2601,8 @@ FIntPoint BuildTileObjectLists(FScene* Scene, TArray<FViewInfo>& Views, int32 Nu
 				static FGlobalBoundShaderState BoundShaderState;
 				SetGlobalBoundShaderState(BoundShaderState, GetVertexDeclarationFVector4(), *VertexShader, *PixelShader);
 
-				VertexShader->SetParameters(View, NumObjects);
-				PixelShader->SetParameters(View, FVector2D(GroupSizeX, GroupSizeY));
+				VertexShader->SetParameters(RHICmdList, View, NumObjects);
+				PixelShader->SetParameters(RHICmdList, View, FVector2D(GroupSizeX, GroupSizeY));
 
 				RHISetStreamSource(0, StencilingGeometry::GLowPolyStencilSphereVertexBuffer.VertexBufferRHI, sizeof(FVector4), 0);
 
@@ -2630,10 +2654,10 @@ FIntPoint BuildTileObjectLists(FScene* Scene, TArray<FViewInfo>& Views, int32 Nu
 			TShaderMapRef<FDistanceFieldBuildTileListCS > ComputeShader(GetGlobalShaderMap());
 
 			RHISetComputeShader(ComputeShader->GetComputeShader());	
-			ComputeShader->SetParameters(View, Scene, FVector2D(GroupSizeX, GroupSizeY), NumObjects);
-			DispatchComputeShader(*ComputeShader, GroupSizeX, GroupSizeY, 1);
+			ComputeShader->SetParameters(RHICmdList, View, Scene, FVector2D(GroupSizeX, GroupSizeY), NumObjects);
+			DispatchComputeShader(RHICmdList, *ComputeShader, GroupSizeX, GroupSizeY, 1);
 
-			ComputeShader->UnsetParameters();
+			ComputeShader->UnsetParameters(RHICmdList);
 		}
 	}
 
@@ -2642,6 +2666,9 @@ FIntPoint BuildTileObjectLists(FScene* Scene, TArray<FViewInfo>& Views, int32 Nu
 
 void RenderIrradianceCacheInterpolation(TArray<FViewInfo>& Views, IPooledRenderTarget* InterpolationTarget, FSceneRenderTargetItem& DistanceFieldNormal, int32 DepthLevel, int32 DestLevelDownsampleFactor, bool bFinalInterpolation)
 {
+	//@todo-rco: RHIPacketList
+	FRHICommandList* RHICmdList = nullptr;
+
 	check(!(bFinalInterpolation && DepthLevel != 0));
 
 	{
@@ -2702,8 +2729,8 @@ void RenderIrradianceCacheInterpolation(TArray<FViewInfo>& Views, IPooledRenderT
 					static FGlobalBoundShaderState BoundShaderState;
 					SetGlobalBoundShaderState(BoundShaderState, GScreenVertexDeclaration.VertexDeclarationRHI, *VertexShader, *PixelShader);
 
-					VertexShader->SetParameters(View, SplatSourceDepthLevel, DestLevelDownsampleFactor, NormalizedOffsetToPixelCenter);
-					PixelShader->SetParameters(View, DistanceFieldNormal, DestLevelDownsampleFactor);
+					VertexShader->SetParameters(RHICmdList, View, SplatSourceDepthLevel, DestLevelDownsampleFactor, NormalizedOffsetToPixelCenter);
+					PixelShader->SetParameters(RHICmdList, View, DistanceFieldNormal, DestLevelDownsampleFactor);
 
 					RHIDrawPrimitiveIndirect(PT_TriangleList, SurfaceCacheResources.Level[SplatSourceDepthLevel]->ScatterDrawParameters.Buffer, 0);
 				}
@@ -2718,8 +2745,8 @@ void RenderIrradianceCacheInterpolation(TArray<FViewInfo>& Views, IPooledRenderT
 					static FGlobalBoundShaderState BoundShaderState;
 					SetGlobalBoundShaderState(BoundShaderState, GScreenVertexDeclaration.VertexDeclarationRHI, *VertexShader, *PixelShader);
 
-					VertexShader->SetParameters(View, SplatSourceDepthLevel, DestLevelDownsampleFactor, NormalizedOffsetToPixelCenter);
-					PixelShader->SetParameters(View, DistanceFieldNormal, DestLevelDownsampleFactor);
+					VertexShader->SetParameters(RHICmdList, View, SplatSourceDepthLevel, DestLevelDownsampleFactor, NormalizedOffsetToPixelCenter);
+					PixelShader->SetParameters(RHICmdList, View, DistanceFieldNormal, DestLevelDownsampleFactor);
 
 					RHIDrawPrimitiveIndirect(PT_TriangleList, SurfaceCacheResources.Level[SplatSourceDepthLevel]->ScatterDrawParameters.Buffer, 0);
 				}
@@ -2854,6 +2881,9 @@ void UpdateVisibleObjectBuffers(const FScene* Scene, const FViewInfo& View, int3
 
 bool FDeferredShadingSceneRenderer::RenderDistanceFieldAOSurfaceCache(FSceneRenderTargetItem& OutBentNormalAO, bool bApplyToSceneColor)
 {
+	//@todo-rco: RHIPacketList
+	FRHICommandList* RHICmdList = nullptr;
+
 	//@todo - support multiple views
 	const FViewInfo& View = Views[0];
 
@@ -2932,20 +2962,20 @@ bool FDeferredShadingSceneRenderer::RenderDistanceFieldAOSurfaceCache(FSceneRend
 							TShaderMapRef<FSetupCopyIndirectArgumentsCS> ComputeShader(GetGlobalShaderMap());
 
 							RHISetComputeShader(ComputeShader->GetComputeShader());	
-							ComputeShader->SetParameters(View, DepthLevel);
-							DispatchComputeShader(*ComputeShader, 1, 1, 1);
+							ComputeShader->SetParameters(RHICmdList, View, DepthLevel);
+							DispatchComputeShader(RHICmdList, *ComputeShader, 1, 1, 1);
 
-							ComputeShader->UnsetParameters();
+							ComputeShader->UnsetParameters(RHICmdList);
 						}
 
 						{
 							TShaderMapRef<FCopyIrradianceCacheSamplesCS> ComputeShader(GetGlobalShaderMap());
 
 							RHISetComputeShader(ComputeShader->GetComputeShader());	
-							ComputeShader->SetParameters(View, DepthLevel);
-							DispatchIndirectComputeShader(*ComputeShader, SurfaceCacheResources.DispatchParameters.Buffer, 0);
+							ComputeShader->SetParameters(RHICmdList, View, DepthLevel);
+							DispatchIndirectComputeShader(RHICmdList, *ComputeShader, SurfaceCacheResources.DispatchParameters.Buffer, 0);
 
-							ComputeShader->UnsetParameters();
+							ComputeShader->UnsetParameters(RHICmdList);
 						}
 
 						Swap(SurfaceCacheResources.Level[DepthLevel], SurfaceCacheResources.TempResources);
@@ -2981,10 +3011,10 @@ bool FDeferredShadingSceneRenderer::RenderDistanceFieldAOSurfaceCache(FSceneRend
 							TShaderMapRef<FSaveStartIndexCS> ComputeShader(GetGlobalShaderMap());
 
 							RHISetComputeShader(ComputeShader->GetComputeShader());	
-							ComputeShader->SetParameters(View, DepthLevel);
-							DispatchComputeShader(*ComputeShader, 1, 1, 1);
+							ComputeShader->SetParameters(RHICmdList, View, DepthLevel);
+							DispatchComputeShader(RHICmdList, *ComputeShader, 1, 1, 1);
 
-							ComputeShader->UnsetParameters();
+							ComputeShader->UnsetParameters(RHICmdList);
 						}
 
 						// Create new records which haven't been shaded yet for shading points which don't have a valid interpolation from existing records
@@ -2999,10 +3029,10 @@ bool FDeferredShadingSceneRenderer::RenderDistanceFieldAOSurfaceCache(FSceneRend
 							TShaderMapRef<FPopulateCacheCS> ComputeShader(GetGlobalShaderMap());
 
 							RHISetComputeShader(ComputeShader->GetComputeShader());	
-							ComputeShader->SetParameters(View, DistanceFieldAOIrradianceCacheSplat->GetRenderTargetItem(), DistanceFieldNormal->GetRenderTargetItem(), DestLevelDownsampleFactor, DepthLevel, TileListGroupSize);
-							DispatchComputeShader(*ComputeShader, GroupSizeX, GroupSizeY, 1);
+							ComputeShader->SetParameters(RHICmdList, View, DistanceFieldAOIrradianceCacheSplat->GetRenderTargetItem(), DistanceFieldNormal->GetRenderTargetItem(), DestLevelDownsampleFactor, DepthLevel, TileListGroupSize);
+							DispatchComputeShader(RHICmdList, *ComputeShader, GroupSizeX, GroupSizeY, 1);
 
-							ComputeShader->UnsetParameters();
+							ComputeShader->UnsetParameters(RHICmdList);
 						}
 					}
 
@@ -3013,10 +3043,10 @@ bool FDeferredShadingSceneRenderer::RenderDistanceFieldAOSurfaceCache(FSceneRend
 							TShaderMapRef<FSetupFinalGatherIndirectArgumentsCS> ComputeShader(GetGlobalShaderMap());
 
 							RHISetComputeShader(ComputeShader->GetComputeShader());	
-							ComputeShader->SetParameters(View, DepthLevel);
-							DispatchComputeShader(*ComputeShader, 1, 1, 1);
+							ComputeShader->SetParameters(RHICmdList, View, DepthLevel);
+							DispatchComputeShader(RHICmdList, *ComputeShader, 1, 1, 1);
 
-							ComputeShader->UnsetParameters();
+							ComputeShader->UnsetParameters(RHICmdList);
 						}
 
 						// Compute lighting for the new surface cache records by cone-stepping through the object distance fields
@@ -3024,10 +3054,10 @@ bool FDeferredShadingSceneRenderer::RenderDistanceFieldAOSurfaceCache(FSceneRend
 							TShaderMapRef<FFinalGatherCS> ComputeShader(GetGlobalShaderMap());
 
 							RHISetComputeShader(ComputeShader->GetComputeShader());	
-							ComputeShader->SetParameters(View, DistanceFieldAOIrradianceCacheSplat->GetRenderTargetItem(), DistanceFieldNormal->GetRenderTargetItem(), NumObjects, DestLevelDownsampleFactor, DepthLevel, TileListGroupSize);
-							DispatchIndirectComputeShader(*ComputeShader, SurfaceCacheResources.DispatchParameters.Buffer, 0);
+							ComputeShader->SetParameters(RHICmdList, View, DistanceFieldAOIrradianceCacheSplat->GetRenderTargetItem(), DistanceFieldNormal->GetRenderTargetItem(), NumObjects, DestLevelDownsampleFactor, DepthLevel, TileListGroupSize);
+							DispatchIndirectComputeShader(RHICmdList, *ComputeShader, SurfaceCacheResources.DispatchParameters.Buffer, 0);
 
-							ComputeShader->UnsetParameters();
+							ComputeShader->UnsetParameters(RHICmdList);
 						}
 					}
 				}
@@ -3077,7 +3107,7 @@ public:
 
 	static bool ShouldCache(EShaderPlatform Platform)
 	{
-		return IsFeatureLevelSupported(Platform, ERHIFeatureLevel::SM4);
+		return IsFeatureLevelSupported(Platform, ERHIFeatureLevel::SM4) && !IsOpenGLPlatform(Platform);
 	}
 
 	static void ModifyCompilationEnvironment(EShaderPlatform Platform, FShaderCompilerEnvironment& OutEnvironment)
@@ -3100,16 +3130,16 @@ public:
 		ApplyShadowing.Bind(Initializer.ParameterMap, TEXT("ApplyShadowing"));
 	}
 
-	void SetParameters(const FSceneView& View, FTextureRHIParamRef DynamicBentNormalAO, bool bApplyShadowing)
+	void SetParameters(FRHICommandList* RHICmdList, const FSceneView& View, FTextureRHIParamRef DynamicBentNormalAO, bool bApplyShadowing)
 	{
 		const FPixelShaderRHIParamRef ShaderRHI = GetPixelShader();
-		FGlobalShader::SetParameters(ShaderRHI, View);
-		DeferredParameters.Set(ShaderRHI, View);
+		FGlobalShader::SetParameters(RHICmdList, ShaderRHI, View);
+		DeferredParameters.Set(RHICmdList, ShaderRHI, View);
 
-		SetTextureParameter(ShaderRHI, LightAccumulationTexture, LightAccumulationSampler, TStaticSamplerState<SF_Point>::GetRHI(), DynamicBentNormalAO);
+		SetTextureParameter(RHICmdList, ShaderRHI, LightAccumulationTexture, LightAccumulationSampler, TStaticSamplerState<SF_Point>::GetRHI(), DynamicBentNormalAO);
 
-		SetShaderValue(ShaderRHI, BentNormalUVScale, 1.0f);
-		SetShaderValue(ShaderRHI, ApplyShadowing, bApplyShadowing ? 1.0f : 0.0f);
+		SetShaderValue(RHICmdList, ShaderRHI, BentNormalUVScale, 1.0f);
+		SetShaderValue(RHICmdList, ShaderRHI, ApplyShadowing, bApplyShadowing ? 1.0f : 0.0f);
 	}
 
 	// FShader interface.
@@ -3137,6 +3167,9 @@ IMPLEMENT_SHADER_TYPE(,FDynamicSkyLightDiffusePS,TEXT("SkyLighting"),TEXT("SkyLi
 
 void FDeferredShadingSceneRenderer::RenderDynamicSkyLighting()
 {
+	//@todo-rco: RHIPacketList
+	FRHICommandList* RHICmdList = nullptr;
+
 	if (Scene->SkyLight
 		&& Scene->SkyLight->ProcessedTexture
 		&& !Scene->SkyLight->bWantsStaticShadowing
@@ -3180,7 +3213,7 @@ void FDeferredShadingSceneRenderer::RenderDynamicSkyLighting()
 			static FGlobalBoundShaderState BoundShaderState;
 			SetGlobalBoundShaderState( BoundShaderState, GFilterVertexDeclaration.VertexDeclarationRHI, *VertexShader, *PixelShader );
 
-			PixelShader->SetParameters( View, LightAccumulation->GetRenderTargetItem().ShaderResourceTexture, bApplyShadowing );
+			PixelShader->SetParameters(RHICmdList, View, LightAccumulation->GetRenderTargetItem().ShaderResourceTexture, bApplyShadowing);
 
 			DrawRectangle( 
 				0, 0, 
diff --git a/Engine/Source/Runtime/Renderer/Private/DistortionRendering.cpp b/Engine/Source/Runtime/Renderer/Private/DistortionRendering.cpp
--- a/Engine/Source/Runtime/Renderer/Private/DistortionRendering.cpp
+++ b/Engine/Source/Runtime/Renderer/Private/DistortionRendering.cpp
@@ -33,7 +33,7 @@ public:
 	}
 	FDistortionApplyScreenPS() {}
 
-	void SetParameters(const FRenderingCompositePassContext& Context, IPooledRenderTarget& DistortionRT)
+	void SetParameters(FRHICommandList* RHICmdList, const FRenderingCompositePassContext& Context, IPooledRenderTarget& DistortionRT)
 	{
 		const FPixelShaderRHIParamRef ShaderRHI = GetPixelShader();
 
@@ -42,6 +42,7 @@ public:
 
 		// Here we use SF_Point as in fullscreen the pixels are 1:1 mapped.
 		SetTextureParameter(
+			RHICmdList, 
 			ShaderRHI,
 			DistortionTexture,
 			DistortionTextureSampler,
@@ -50,6 +51,7 @@ public:
 			);
 
 		SetTextureParameter(
+			RHICmdList, 
 			ShaderRHI,
 			SceneColorTexture,
 			SceneColorTextureSampler,
@@ -63,7 +65,7 @@ public:
 												(float)ViewportRect.Min.Y/SceneBufferSize.Y,
 												(float)ViewportRect.Max.X/SceneBufferSize.X,
 												(float)ViewportRect.Max.Y/SceneBufferSize.Y);
-		SetShaderValue(ShaderRHI, SceneColorRect, SceneColorRectValue);
+		SetShaderValue(RHICmdList, ShaderRHI, SceneColorRect, SceneColorRectValue);
 	}
 
 	virtual bool Serialize(FArchive& Ar)
@@ -123,14 +125,14 @@ protected:
 
 public:
 	
-	void SetParameters(const FVertexFactory* VertexFactory,const FMaterialRenderProxy* MaterialRenderProxy,const FSceneView* View)
+	void SetParameters(FRHICommandList* RHICmdList, const FVertexFactory* VertexFactory,const FMaterialRenderProxy* MaterialRenderProxy,const FSceneView* View)
 	{
-		FMeshMaterialShader::SetParameters(GetVertexShader(), MaterialRenderProxy, *MaterialRenderProxy->GetMaterial(View->GetFeatureLevel()), *View, ESceneRenderTargetsMode::SetTextures);
+		FMeshMaterialShader::SetParameters(RHICmdList, GetVertexShader(), MaterialRenderProxy, *MaterialRenderProxy->GetMaterial(View->GetFeatureLevel()), *View, ESceneRenderTargetsMode::SetTextures);
 	}
 
-	void SetMesh(const FVertexFactory* VertexFactory,const FSceneView& View,const FPrimitiveSceneProxy* Proxy,const FMeshBatchElement& BatchElement)
+	void SetMesh(FRHICommandList* RHICmdList, const FVertexFactory* VertexFactory,const FSceneView& View,const FPrimitiveSceneProxy* Proxy,const FMeshBatchElement& BatchElement)
 	{
-		FMeshMaterialShader::SetMesh(GetVertexShader(),VertexFactory,View,Proxy,BatchElement);
+		FMeshMaterialShader::SetMesh(RHICmdList, GetVertexShader(),VertexFactory,View,Proxy,BatchElement);
 	}
 };
 
@@ -213,11 +215,12 @@ public:
 	}
 
 	void SetParameters(
+		FRHICommandList* RHICmdList, 
 		const FMaterialRenderProxy* MaterialRenderProxy,
 		const FSceneView& View
 		)
 	{
-		FMeshMaterialShader::SetParameters(GetPixelShader(), MaterialRenderProxy, *MaterialRenderProxy->GetMaterial(View.GetFeatureLevel()), View, ESceneRenderTargetsMode::SetTextures);
+		FMeshMaterialShader::SetParameters(RHICmdList, GetPixelShader(), MaterialRenderProxy, *MaterialRenderProxy->GetMaterial(View.GetFeatureLevel()), View, ESceneRenderTargetsMode::SetTextures);
 
 		float Ratio = View.UnscaledViewRect.Width() / (float)View.UnscaledViewRect.Height();
 		float Params[4];
@@ -226,13 +229,13 @@ public:
 		Params[2] = (float)View.UnscaledViewRect.Width();
 		Params[3] = (float)View.UnscaledViewRect.Height();
 
-		SetShaderValue(GetPixelShader(), DistortionParams, Params);
+		SetShaderValue(RHICmdList, GetPixelShader(), DistortionParams, Params);
 
 	}
 
-	void SetMesh(const FVertexFactory* VertexFactory,const FSceneView& View,const FPrimitiveSceneProxy* Proxy,const FMeshBatchElement& BatchElement)
+	void SetMesh(FRHICommandList* RHICmdList, const FVertexFactory* VertexFactory,const FSceneView& View,const FPrimitiveSceneProxy* Proxy,const FMeshBatchElement& BatchElement)
 	{
-		FMeshMaterialShader::SetMesh(GetPixelShader(),VertexFactory,View,Proxy,BatchElement);
+		FMeshMaterialShader::SetMesh(RHICmdList, GetPixelShader(),VertexFactory,View,Proxy,BatchElement);
 	}
 
 	virtual bool Serialize(FArchive& Ar)
@@ -292,7 +295,7 @@ public:
 	* @param CI - The command interface to execute the draw commands on.
 	* @param View - The view of the scene being drawn.
 	*/
-	void DrawShared(const FSceneView* View,FBoundShaderStateRHIParamRef BoundShaderState) const;
+	void DrawShared(FRHICommandList* RHICmdList, const FSceneView* View,FBoundShaderStateRHIParamRef BoundShaderState) const;
 	
 	/** 
 	* Create bound shader state using the vertex decl from the mesh draw policy
@@ -308,6 +311,7 @@ public:
 	* @param ElementData - context specific data for mesh rendering
 	*/
 	void SetMeshRenderState(
+		FRHICommandList* RHICmdList, 
 		const FSceneView& View,
 		const FPrimitiveSceneProxy* PrimitiveSceneProxy,
 		const FMeshBatch& Mesh,
@@ -400,23 +404,26 @@ bool TDistortionMeshDrawingPolicy<DistortMeshPolicy>::Matches(
 */
 template<class DistortMeshPolicy>
 void TDistortionMeshDrawingPolicy<DistortMeshPolicy>::DrawShared(
+	FRHICommandList* RHICmdList, 
 	const FSceneView* View,
 	FBoundShaderStateRHIParamRef BoundShaderState
 	) const
 {
+	check(!RHICmdList);
+
 	// Set shared mesh resources
-	FMeshDrawingPolicy::DrawShared(View);
+	FMeshDrawingPolicy::DrawShared(RHICmdList, View);
 
 	// Set the actual shader & vertex declaration state
 	RHISetBoundShaderState(BoundShaderState);
 
 	// Set the translucent shader parameters for the material instance
-	VertexShader->SetParameters(VertexFactory,MaterialRenderProxy,View);
+	VertexShader->SetParameters(RHICmdList, VertexFactory,MaterialRenderProxy,View);
 
 	if(HullShader && DomainShader)
 	{
-		HullShader->SetParameters(MaterialRenderProxy,*View);
-		DomainShader->SetParameters(MaterialRenderProxy,*View);
+		HullShader->SetParameters(RHICmdList, MaterialRenderProxy,*View);
+		DomainShader->SetParameters(RHICmdList, MaterialRenderProxy,*View);
 	}
 
 	if (bOverrideWithShaderComplexity)
@@ -434,7 +441,7 @@ void TDistortionMeshDrawingPolicy<DistortMeshPolicy>::DrawShared(
 		}
 		else
 		{
-			DistortPixelShader->SetParameters(MaterialRenderProxy,*View);
+			DistortPixelShader->SetParameters(RHICmdList, MaterialRenderProxy,*View);
 		}
 	}
 }
@@ -486,6 +493,7 @@ FBoundShaderStateRHIRef TDistortionMeshDrawingPolicy<DistortMeshPolicy>::CreateB
 */
 template<class DistortMeshPolicy>
 void TDistortionMeshDrawingPolicy<DistortMeshPolicy>::SetMeshRenderState(
+	FRHICommandList* RHICmdList, 
 	const FSceneView& View,
 	const FPrimitiveSceneProxy* PrimitiveSceneProxy,
 	const FMeshBatch& Mesh,
@@ -494,17 +502,19 @@ void TDistortionMeshDrawingPolicy<DistortMeshPolicy>::SetMeshRenderState(
 	const ElementDataType& ElementData
 	) const
 {
+	check(!RHICmdList);
+
 	const FMeshBatchElement& BatchElement = Mesh.Elements[BatchElementIndex];
 
 	EmitMeshDrawEvents(PrimitiveSceneProxy, Mesh);
 
 	// Set transforms
-	VertexShader->SetMesh(VertexFactory,View,PrimitiveSceneProxy,BatchElement);
+	VertexShader->SetMesh(RHICmdList, VertexFactory,View,PrimitiveSceneProxy,BatchElement);
 
 	if(HullShader && DomainShader)
 	{
-		HullShader->SetMesh(VertexFactory,View,PrimitiveSceneProxy,BatchElement);
-		DomainShader->SetMesh(VertexFactory,View,PrimitiveSceneProxy,BatchElement);
+		HullShader->SetMesh(RHICmdList, VertexFactory,View,PrimitiveSceneProxy,BatchElement);
+		DomainShader->SetMesh(RHICmdList, VertexFactory,View,PrimitiveSceneProxy,BatchElement);
 	}
 	
 
@@ -515,7 +525,7 @@ void TDistortionMeshDrawingPolicy<DistortMeshPolicy>::SetMeshRenderState(
 	{
 		if (!bInitializeOffsets)
 		{
-			DistortPixelShader->SetMesh(VertexFactory,View,PrimitiveSceneProxy,BatchElement);
+			DistortPixelShader->SetMesh(RHICmdList, VertexFactory,View,PrimitiveSceneProxy,BatchElement);
 		}
 	}
 	
@@ -545,6 +555,7 @@ public:
 	* @return true if the mesh rendered
 	*/
 	static bool DrawDynamicMesh(
+		FRHICommandList* RHICmdList, 
 		const FSceneView& View,
 		ContextType DrawingContext,
 		const FMeshBatch& Mesh,
@@ -559,6 +570,7 @@ public:
 	* @return true if the mesh rendered
 	*/
 	static bool DrawStaticMesh(
+		FRHICommandList* RHICmdList, 
 		const FSceneView* View,
 		ContextType DrawingContext,
 		const FStaticMesh& StaticMesh,
@@ -577,6 +589,7 @@ public:
 */
 template<class DistortMeshPolicy>
 bool TDistortionMeshDrawingPolicyFactory<DistortMeshPolicy>::DrawDynamicMesh(
+	FRHICommandList* RHICmdList, 
 	const FSceneView& View,
 	ContextType bInitializeOffsets,
 	const FMeshBatch& Mesh,
@@ -599,11 +612,11 @@ bool TDistortionMeshDrawingPolicyFactory<DistortMeshPolicy>::DrawDynamicMesh(
 			bInitializeOffsets,
 			View.Family->EngineShowFlags.ShaderComplexity
 			);
-		DrawingPolicy.DrawShared(&View,DrawingPolicy.CreateBoundShaderState(View.GetFeatureLevel()));
+		DrawingPolicy.DrawShared(RHICmdList, &View,DrawingPolicy.CreateBoundShaderState(View.GetFeatureLevel()));
 		for( int32 BatchElementIndex=0;BatchElementIndex < Mesh.Elements.Num();BatchElementIndex++ )
 		{
-			DrawingPolicy.SetMeshRenderState(View,PrimitiveSceneProxy,Mesh,BatchElementIndex,bBackFace,typename TDistortionMeshDrawingPolicy<DistortMeshPolicy>::ElementDataType());
-			DrawingPolicy.DrawMesh(Mesh,BatchElementIndex);
+			DrawingPolicy.SetMeshRenderState(RHICmdList, View,PrimitiveSceneProxy,Mesh,BatchElementIndex,bBackFace,typename TDistortionMeshDrawingPolicy<DistortMeshPolicy>::ElementDataType());
+			DrawingPolicy.DrawMesh(RHICmdList, Mesh,BatchElementIndex);
 		}
 
 		return true;
@@ -620,6 +633,7 @@ bool TDistortionMeshDrawingPolicyFactory<DistortMeshPolicy>::DrawDynamicMesh(
 */
 template<class DistortMeshPolicy>
 bool TDistortionMeshDrawingPolicyFactory<DistortMeshPolicy>::DrawStaticMesh(
+	FRHICommandList* RHICmdList, 
 	const FSceneView* View,
 	ContextType bInitializeOffsets,
 	const FStaticMesh& StaticMesh,
@@ -642,14 +656,14 @@ bool TDistortionMeshDrawingPolicyFactory<DistortMeshPolicy>::DrawStaticMesh(
 			bInitializeOffsets,
 			View->Family->EngineShowFlags.ShaderComplexity
 			);
-		DrawingPolicy.DrawShared(View,DrawingPolicy.CreateBoundShaderState(FeatureLevel));
+		DrawingPolicy.DrawShared(RHICmdList, View,DrawingPolicy.CreateBoundShaderState(FeatureLevel));
 		int32 BatchElementIndex = 0;
 		do
 		{
 			if(BatchElementMask & 1)
 			{
-				DrawingPolicy.SetMeshRenderState(*View,PrimitiveSceneProxy,StaticMesh,BatchElementIndex,bBackFace,typename TDistortionMeshDrawingPolicy<DistortMeshPolicy>::ElementDataType());
-				DrawingPolicy.DrawMesh(StaticMesh, BatchElementIndex);
+				DrawingPolicy.SetMeshRenderState(RHICmdList, *View,PrimitiveSceneProxy,StaticMesh,BatchElementIndex,bBackFace,typename TDistortionMeshDrawingPolicy<DistortMeshPolicy>::ElementDataType());
+				DrawingPolicy.DrawMesh(RHICmdList, StaticMesh, BatchElementIndex);
 			}
 			BatchElementMask >>= 1;
 			BatchElementIndex++;
@@ -678,6 +692,9 @@ bool FDistortionPrimSet::DrawAccumulatedOffsets(const FViewInfo* ViewInfo,bool b
 {
 	bool bDirty=false;
 
+	//@todo-rco: RHIPacketList
+	FRHICommandList* RHICmdList = nullptr;
+
 	// Draw the view's elements with the translucent drawing policy.
 	bDirty |= DrawViewElements<TDistortionMeshDrawingPolicyFactory<FDistortMeshAccumulatePolicy> >(
 		*ViewInfo,
@@ -724,6 +741,7 @@ bool FDistortionPrimSet::DrawAccumulatedOffsets(const FViewInfo* ViewInfo,bool b
 							&& StaticMesh.IsTranslucent(ViewInfo->GetFeatureLevel()) )
 						{
 							bDirty |= TDistortionMeshDrawingPolicyFactory<FDistortMeshAccumulatePolicy>::DrawStaticMesh(
+								RHICmdList, 
 								ViewInfo,
 								bInitializeOffsets,
 								StaticMesh,
@@ -759,7 +777,7 @@ static TAutoConsoleVariable<int32> CVarRefractionQuality(
 /** 
  * Renders the scene's distortion 
  */
-void FSceneRenderer::RenderDistortion()
+void FSceneRenderer::RenderDistortion(FRHICommandList* RHICmdList)
 {
 	if (CVarRefractionQuality.GetValueOnRenderThread() == 0)
 	{
@@ -890,8 +908,8 @@ void FSceneRenderer::RenderDistortion()
 				static FGlobalBoundShaderState BoundShaderState;
 				SetGlobalBoundShaderState(BoundShaderState, GFilterVertexDeclaration.VertexDeclarationRHI, *VertexShader, *PixelShader);
 
-				VertexShader->SetParameters(Context);
-				PixelShader->SetParameters(Context, *DistortionRT);
+				VertexShader->SetParameters(RHICmdList, Context);
+				PixelShader->SetParameters(RHICmdList, Context, *DistortionRT);
 
 				// Draw a quad mapping scene color to the view's render target
 				DrawRectangle(
diff --git a/Engine/Source/Runtime/Renderer/Private/DrawingPolicy.cpp b/Engine/Source/Runtime/Renderer/Private/DrawingPolicy.cpp
--- a/Engine/Source/Runtime/Renderer/Private/DrawingPolicy.cpp
+++ b/Engine/Source/Runtime/Renderer/Private/DrawingPolicy.cpp
@@ -42,6 +42,7 @@ FMeshDrawingPolicy::FMeshDrawingPolicy(
 }
 
 void FMeshDrawingPolicy::SetMeshRenderState(
+	FRHICommandList* RHICmdList, 
 	const FSceneView& View,
 	const FPrimitiveSceneProxy* PrimitiveSceneProxy,
 	const FMeshBatch& Mesh,
@@ -52,14 +53,14 @@ void FMeshDrawingPolicy::SetMeshRenderState(
 {
 	EmitMeshDrawEvents(PrimitiveSceneProxy, Mesh);
 
-	// Use bitwise logic ops to avoid branches
-	RHISetRasterizerState(GetStaticRasterizerState<true>(
-		( Mesh.bWireframe | IsWireframe() ) ? FM_Wireframe : FM_Solid, ( ( IsTwoSided() & !NeedsBackfacePass() ) | Mesh.bDisableBackfaceCulling ) ? CM_None :
+		// Use bitwise logic ops to avoid branches
+	FRHICommandList::SetRasterizerState(RHICmdList, GetStaticRasterizerState<true>(
+			( Mesh.bWireframe | IsWireframe() ) ? FM_Wireframe : FM_Solid, ( ( IsTwoSided() & !NeedsBackfacePass() ) | Mesh.bDisableBackfaceCulling ) ? CM_None :
 			( ( (View.bReverseCulling ^ bBackFace) ^ Mesh.ReverseCulling ) ? CM_CCW : CM_CW )
-		));
+			));
 }
 
-void FMeshDrawingPolicy::DrawMesh(const FMeshBatch& Mesh, int32 BatchElementIndex) const
+void FMeshDrawingPolicy::DrawMesh(FRHICommandList* RHICmdList, const FMeshBatch& Mesh, int32 BatchElementIndex) const
 {
 	INC_DWORD_STAT(STAT_MeshDrawCalls);
 	SCOPED_CONDITIONAL_DRAW_EVENTF(MeshEvent, GEmitMeshDrawEvent != 0, DEC_SCENE_ITEMS, TEXT("Mesh Draw"));
@@ -72,6 +73,8 @@ void FMeshDrawingPolicy::DrawMesh(const FMeshBatch& Mesh, int32 BatchElementInde
 
 		if (BatchElement.DynamicIndexData)
 		{
+			check(!RHICmdList);
+
 			RHIDrawIndexedPrimitiveUP(
 				Mesh.Type,
 				BatchElement.MinVertexIndex,
@@ -85,6 +88,8 @@ void FMeshDrawingPolicy::DrawMesh(const FMeshBatch& Mesh, int32 BatchElementInde
 		}
 		else
 		{
+			check(!RHICmdList);
+
 			RHIDrawPrimitiveUP(
 				Mesh.Type,
 				BatchElement.NumPrimitives,
@@ -98,33 +103,35 @@ void FMeshDrawingPolicy::DrawMesh(const FMeshBatch& Mesh, int32 BatchElementInde
 		if(BatchElement.IndexBuffer)
 		{
 			check(BatchElement.IndexBuffer->IsInitialized());
-			RHIDrawIndexedPrimitive(
-				BatchElement.IndexBuffer->IndexBufferRHI,
-				Mesh.Type,
-				0,
-				BatchElement.MinVertexIndex,
-				BatchElement.MaxVertexIndex - BatchElement.MinVertexIndex + 1,
-				BatchElement.FirstIndex,
-				BatchElement.NumPrimitives,
-				BatchElement.NumInstances
-				);
+			FRHICommandList::DrawIndexedPrimitive(
+					RHICmdList,
+					BatchElement.IndexBuffer->IndexBufferRHI,
+					Mesh.Type,
+					0,
+					BatchElement.MinVertexIndex,
+					BatchElement.MaxVertexIndex - BatchElement.MinVertexIndex + 1,
+					BatchElement.FirstIndex,
+					BatchElement.NumPrimitives,
+					BatchElement.NumInstances
+					);
 		}
 		else
 		{
-			RHIDrawPrimitive(
-				Mesh.Type,
-				BatchElement.FirstIndex,
-				BatchElement.NumPrimitives,
-				BatchElement.NumInstances
-				);
+			FRHICommandList::DrawPrimitive(
+					RHICmdList,
+					Mesh.Type,
+					BatchElement.FirstIndex,
+					BatchElement.NumPrimitives,
+					BatchElement.NumInstances
+					);
 		}
 	}
 }
 
-void FMeshDrawingPolicy::DrawShared(const FSceneView* View) const
+void FMeshDrawingPolicy::DrawShared(FRHICommandList* RHICmdList, const FSceneView* View) const
 {
 	check(VertexFactory && VertexFactory->IsInitialized());
-	VertexFactory->Set();
+	VertexFactory->Set(RHICmdList);
 }
 
 /**
diff --git a/Engine/Source/Runtime/Renderer/Private/DrawingPolicy.h b/Engine/Source/Runtime/Renderer/Private/DrawingPolicy.h
--- a/Engine/Source/Runtime/Renderer/Private/DrawingPolicy.h
+++ b/Engine/Source/Runtime/Renderer/Private/DrawingPolicy.h
@@ -7,6 +7,8 @@
 #ifndef __DRAWINGPOLICY_H__
 #define __DRAWINGPOLICY_H__
 
+#include "RHICommandList.h"
+
 /**
  * A macro to compare members of two drawing policies(A and B), and return based on the result.
  * If the members are the same, the macro continues execution rather than returning to the caller.
@@ -65,6 +67,7 @@ public:
 	 * @param PrimitiveSceneProxy - The primitive drawing the dynamic mesh.  If this is a view element, this will be NULL.
 	 */
 	void SetMeshRenderState(
+		FRHICommandList* RHICmdList, 
 		const FSceneView& View,
 		const FPrimitiveSceneProxy* PrimitiveSceneProxy,
 		const FMeshBatch& Mesh,
@@ -76,14 +79,14 @@ public:
 	/**
 	 * Executes the draw commands for a mesh.
 	 */
-	void DrawMesh(const FMeshBatch& Mesh, int32 BatchElementIndex) const;
+	void DrawMesh(FRHICommandList* RHICmdList, const FMeshBatch& Mesh, int32 BatchElementIndex) const;
 
 	/**
 	 * Executes the draw commands which can be shared between any meshes using this drawer.
 	 * @param CI - The command interface to execute the draw commands on.
 	 * @param View - The view of the scene being drawn.
 	 */
-	void DrawShared(const FSceneView* View) const;
+	void DrawShared(FRHICommandList* RHICmdList, const FSceneView* View) const;
 
 	/**
 	* Get the decl for this mesh policy type and vertexfactory
diff --git a/Engine/Source/Runtime/Renderer/Private/DynamicPrimitiveDrawing.inl b/Engine/Source/Runtime/Renderer/Private/DynamicPrimitiveDrawing.inl
--- a/Engine/Source/Runtime/Renderer/Private/DynamicPrimitiveDrawing.inl
+++ b/Engine/Source/Runtime/Renderer/Private/DynamicPrimitiveDrawing.inl
@@ -104,8 +104,11 @@ int32 TDynamicPrimitiveDrawer<DrawingPolicyFactoryType>::DrawMesh(const FMeshBat
 #if DO_CHECK
 	Mesh.CheckUniformBuffers();
 #endif
+	//@todo-rco: RHIPacketList
+	FRHICommandList* RHICmdList = nullptr;
 	INC_DWORD_STAT_BY(STAT_DynamicPathMeshDrawCalls,Mesh.Elements.Num());
 	const bool DrawDirty = DrawingPolicyFactoryType::DrawDynamicMesh(
+		RHICmdList,
 		*View,
 		DrawingContext,
 		Mesh,
@@ -217,6 +220,9 @@ bool DrawViewElements(
 	bool bPreFog
 	)
 {
+	//@todo-rco: RHIPacketList
+	FRHICommandList* RHICmdList = nullptr;
+
 	// Get the correct element list based on dpg index
 	const TIndirectArray<FHitProxyMeshPair>& ViewMeshElementList = ( DPGIndex == SDPG_Foreground ? View.TopViewMeshElements : View.ViewMeshElements );
 	// Draw the view's mesh elements.
@@ -231,6 +237,7 @@ bool DrawViewElements(
 		do
 		{
 			DrawingPolicyFactoryType::DrawDynamicMesh(
+				RHICmdList, 
 				View,
 				DrawingContext,
 				Mesh,
diff --git a/Engine/Source/Runtime/Renderer/Private/FogRendering.cpp b/Engine/Source/Runtime/Renderer/Private/FogRendering.cpp
--- a/Engine/Source/Runtime/Renderer/Private/FogRendering.cpp
+++ b/Engine/Source/Runtime/Renderer/Private/FogRendering.cpp
@@ -60,9 +60,9 @@ public:
 		FogStartZ.Bind(Initializer.ParameterMap,TEXT("FogStartZ"));
 	}
 
-	void SetParameters(const FViewInfo& View)
+	void SetParameters(FRHICommandList* RHICmdList, const FViewInfo& View)
 	{
-		FGlobalShader::SetParameters(GetVertexShader(),View);
+		FGlobalShader::SetParameters(RHICmdList, GetVertexShader(),View);
 
 		{
 			// The fog can be set to start at a certain euclidean distance.
@@ -86,7 +86,7 @@ public:
 
 			float FogClipSpaceZ = ClipSpaceMaxDistance.Z / ClipSpaceMaxDistance.W;
 
-			SetShaderValue(GetVertexShader(),FogStartZ, FogClipSpaceZ);
+			SetShaderValue(RHICmdList, GetVertexShader(),FogStartZ, FogClipSpaceZ);
 		}
 	}
 
@@ -124,15 +124,16 @@ public:
 		SceneTextureParameters.Bind(Initializer.ParameterMap);
 	}
 
-	void SetParameters(const FViewInfo& View, FLightShaftsOutput LightShaftsOutput)
+	void SetParameters(FRHICommandList* RHICmdList, const FViewInfo& View, FLightShaftsOutput LightShaftsOutput)
 	{
-		FGlobalShader::SetParameters(GetPixelShader(), View);
-		SceneTextureParameters.Set(GetPixelShader(), View);
-		ExponentialParameters.Set(GetPixelShader(), &View);
+		FGlobalShader::SetParameters(RHICmdList, GetPixelShader(), View);
+		SceneTextureParameters.Set(RHICmdList, GetPixelShader(), View);
+		ExponentialParameters.Set(RHICmdList, GetPixelShader(), &View);
 
 		if (LightShaftsOutput.bRendered)
 		{
 			SetTextureParameter(
+				RHICmdList, 
 				GetPixelShader(),
 				OcclusionTexture, OcclusionSampler,
 				TStaticSamplerState<SF_Bilinear,AM_Clamp,AM_Clamp,AM_Clamp>::GetRHI(),
@@ -142,6 +143,7 @@ public:
 		else
 		{
 			SetTextureParameter(
+				RHICmdList, 
 				GetPixelShader(),
 				OcclusionTexture, OcclusionSampler,
 				TStaticSamplerState<SF_Bilinear,AM_Clamp,AM_Clamp,AM_Clamp>::GetRHI(),
@@ -261,7 +263,7 @@ void FSceneRenderer::InitFogConstants()
 FGlobalBoundShaderState ExponentialBoundShaderState;
 
 /** Sets the bound shader state for either the per-pixel or per-sample fog pass. */
-void SetFogShaders(FScene* Scene,const FViewInfo& View,FLightShaftsOutput LightShaftsOutput)
+void SetFogShaders(FRHICommandList* RHICmdList, FScene* Scene,const FViewInfo& View,FLightShaftsOutput LightShaftsOutput)
 {
 	if (Scene->ExponentialFogs.Num() > 0)
 	{
@@ -269,8 +271,8 @@ void SetFogShaders(FScene* Scene,const FViewInfo& View,FLightShaftsOutput LightS
 		TShaderMapRef<FExponentialHeightFogPS> ExponentialHeightFogPixelShader(GetGlobalShaderMap());
 
 		SetGlobalBoundShaderState(ExponentialBoundShaderState, GFogVertexDeclaration.VertexDeclarationRHI, *VertexShader, *ExponentialHeightFogPixelShader);
-		VertexShader->SetParameters(View);
-		ExponentialHeightFogPixelShader->SetParameters(View, LightShaftsOutput);
+		VertexShader->SetParameters(RHICmdList, View);
+		ExponentialHeightFogPixelShader->SetParameters(RHICmdList, View, LightShaftsOutput);
 	}
 }
 
@@ -293,6 +295,9 @@ bool FDeferredShadingSceneRenderer::RenderFog(FLightShaftsOutput LightShaftsOutp
 			0, 2, 3
 		};
 
+		//@todo-rco: RHIPacketList
+		FRHICommandList* RHICmdList = nullptr;
+
 		GSceneRenderTargets.BeginRenderingSceneColor();
 		for(int32 ViewIndex = 0;ViewIndex < Views.Num();ViewIndex++)
 		{
@@ -313,7 +318,7 @@ bool FDeferredShadingSceneRenderer::RenderFog(FLightShaftsOutput LightShaftsOutp
 
 			RHISetDepthStencilState(TStaticDepthStencilState<false,CF_Always>::GetRHI());
 
-			SetFogShaders(Scene,View,LightShaftsOutput);
+			SetFogShaders(RHICmdList, Scene,View,LightShaftsOutput);
 
 			// Draw a quad covering the view.
 			RHIDrawIndexedPrimitiveUP(
