UE4 - merge first pass at parallel rendering to main

diff --git a/Engine/Source/Runtime/Renderer/Private/ForwardBasePassRendering.cpp b/Engine/Source/Runtime/Renderer/Private/ForwardBasePassRendering.cpp
--- a/Engine/Source/Runtime/Renderer/Private/ForwardBasePassRendering.cpp
+++ b/Engine/Source/Runtime/Renderer/Private/ForwardBasePassRendering.cpp
@@ -52,6 +52,7 @@ public:
 	/** Draws the translucent mesh with a specific light-map type, and fog volume type */
 	template<typename LightMapPolicyType>
 	void Process(
+		FRHICommandList* RHICmdList, 
 		const FProcessBasePassMeshParameters& Parameters,
 		const LightMapPolicyType& LightMapPolicy,
 		const typename LightMapPolicyType::ElementDataType& LightMapElementData
@@ -143,11 +144,14 @@ public:
 	/** Draws the translucent mesh with a specific light-map type, and shader complexity predicate. */
 	template<typename LightMapPolicyType>
 	void Process(
+		FRHICommandList* RHICmdList, 
 		const FProcessBasePassMeshParameters& Parameters,
 		const LightMapPolicyType& LightMapPolicy,
 		const typename LightMapPolicyType::ElementDataType& LightMapElementData
 		) const
 	{
+		check(!RHICmdList);
+
 #if !(UE_BUILD_SHIPPING || UE_BUILD_TEST)
 		// Treat masked materials as if they don't occlude in shader complexity, which is PVR behavior
 		if(Parameters.BlendMode == BLEND_Masked && View.Family->EngineShowFlags.ShaderComplexity)
@@ -169,6 +173,7 @@ public:
 			View.Family->EngineShowFlags.ShaderComplexity
 			);
 		DrawingPolicy.DrawShared(
+			RHICmdList, 
 			&View,
 			DrawingPolicy.CreateBoundShaderState(View.GetFeatureLevel())
 			);
@@ -176,6 +181,7 @@ public:
 		for( int32 BatchElementIndex=0;BatchElementIndex<Parameters.Mesh.Elements.Num();BatchElementIndex++ )
 		{
 			DrawingPolicy.SetMeshRenderState(
+				RHICmdList, 
 				View,
 				Parameters.PrimitiveSceneProxy,
 				Parameters.Mesh,
@@ -183,7 +189,7 @@ public:
 				bBackFace,
 				typename TBasePassForForwardShadingDrawingPolicy<LightMapPolicyType>::ElementDataType(LightMapElementData)
 				);
-			DrawingPolicy.DrawMesh(Parameters.Mesh, BatchElementIndex);
+			DrawingPolicy.DrawMesh(RHICmdList, Parameters.Mesh, BatchElementIndex);
 		}
 
 		// Restore
@@ -198,6 +204,7 @@ public:
 };
 
 bool FBasePassForwardOpaqueDrawingPolicyFactory::DrawDynamicMesh(
+	FRHICommandList* RHICmdList, 
 	const FSceneView& View,
 	ContextType DrawingContext,
 	const FMeshBatch& Mesh,
@@ -285,7 +292,7 @@ EBasePassSort::Type GetSortMode()
 	}
 }
 
-void FForwardShadingSceneRenderer::RenderForwardShadingBasePass()
+void FForwardShadingSceneRenderer::RenderForwardShadingBasePass(FRHICommandList* RHICmdList)
 {
 	SCOPED_DRAW_EVENT(BasePass, DEC_SCENE_ITEMS);
 	SCOPE_CYCLE_COUNTER(STAT_BasePassDrawTime);
@@ -322,22 +329,44 @@ void FForwardShadingSceneRenderer::RenderForwardShadingBasePass()
 		RHISetDepthStencilState(TStaticDepthStencilState<true, CF_GreaterEqual>::GetRHI());
 		RHISetViewport(View.ViewRect.Min.X, View.ViewRect.Min.Y, 0, View.ViewRect.Max.X, View.ViewRect.Max.Y, 1);
 
-		// Render the base pass static data
-		if (SortMode == EBasePassSort::SortPerMesh)
-		{
-			SCOPE_CYCLE_COUNTER(STAT_StaticDrawListDrawTime);
-			MaxDraws -= Scene->BasePassForForwardShadingLowQualityLightMapDrawList[FScene::EBasePass_Default].DrawVisibleFrontToBack(View,View.StaticMeshVisibilityMap,View.StaticMeshBatchVisibility,MaxDraws);
-			MaxDraws -= Scene->BasePassForForwardShadingDistanceFieldShadowMapLightMapDrawList[FScene::EBasePass_Default].DrawVisibleFrontToBack(View,View.StaticMeshVisibilityMap,View.StaticMeshBatchVisibility,MaxDraws);
-			MaxDraws -= Scene->BasePassForForwardShadingDirectionalLightAndSHIndirectDrawList[FScene::EBasePass_Default].DrawVisibleFrontToBack(View,View.StaticMeshVisibilityMap,View.StaticMeshBatchVisibility,MaxDraws);
-			MaxDraws -= Scene->BasePassForForwardShadingNoLightMapDrawList[FScene::EBasePass_Default].DrawVisibleFrontToBack(View,View.StaticMeshVisibilityMap,View.StaticMeshBatchVisibility,MaxDraws);
-		}
-		else
 		{
-			SCOPE_CYCLE_COUNTER(STAT_StaticDrawListDrawTime);
-			Scene->BasePassForForwardShadingLowQualityLightMapDrawList[FScene::EBasePass_Default].DrawVisible(View,View.StaticMeshVisibilityMap,View.StaticMeshBatchVisibility);
-			Scene->BasePassForForwardShadingDistanceFieldShadowMapLightMapDrawList[FScene::EBasePass_Default].DrawVisible(View,View.StaticMeshVisibilityMap,View.StaticMeshBatchVisibility);
-			Scene->BasePassForForwardShadingDirectionalLightAndSHIndirectDrawList[FScene::EBasePass_Default].DrawVisible(View,View.StaticMeshVisibilityMap,View.StaticMeshBatchVisibility);
-			Scene->BasePassForForwardShadingNoLightMapDrawList[FScene::EBasePass_Default].DrawVisible(View,View.StaticMeshVisibilityMap,View.StaticMeshBatchVisibility);
+			SCOPE_CYCLE_COUNTER(STAT_RHICounterTEMP);
+			auto* PrevRHICmdList = RHICmdList;
+			if (!RHICmdList)
+			{
+				//@todo-rco: Very temp code!!!
+				static IConsoleVariable* RHICmdListCVar = IConsoleManager::Get().FindConsoleVariable(TEXT("r.RHICmd"));
+				bool bUseRHICmdList = RHICmdListCVar->GetInt() != 0;
+				RHICmdList = bUseRHICmdList ? GRHICommandList.CreateList() : nullptr;
+			}
+
+			// Render the base pass static data
+			if (SortMode == EBasePassSort::SortPerMesh)
+			{
+				SCOPE_CYCLE_COUNTER(STAT_StaticDrawListDrawTime);
+				MaxDraws -= Scene->BasePassForForwardShadingLowQualityLightMapDrawList[FScene::EBasePass_Default].DrawVisibleFrontToBack(RHICmdList, View,View.StaticMeshVisibilityMap,View.StaticMeshBatchVisibility,MaxDraws);
+				MaxDraws -= Scene->BasePassForForwardShadingDistanceFieldShadowMapLightMapDrawList[FScene::EBasePass_Default].DrawVisibleFrontToBack(RHICmdList, View,View.StaticMeshVisibilityMap,View.StaticMeshBatchVisibility,MaxDraws);
+				MaxDraws -= Scene->BasePassForForwardShadingDirectionalLightAndSHIndirectDrawList[FScene::EBasePass_Default].DrawVisibleFrontToBack(RHICmdList, View,View.StaticMeshVisibilityMap,View.StaticMeshBatchVisibility,MaxDraws);
+				MaxDraws -= Scene->BasePassForForwardShadingNoLightMapDrawList[FScene::EBasePass_Default].DrawVisibleFrontToBack(RHICmdList, View,View.StaticMeshVisibilityMap,View.StaticMeshBatchVisibility,MaxDraws);
+			}
+			else
+			{
+				SCOPE_CYCLE_COUNTER(STAT_StaticDrawListDrawTime);
+				Scene->BasePassForForwardShadingLowQualityLightMapDrawList[FScene::EBasePass_Default].DrawVisible(RHICmdList, View,View.StaticMeshVisibilityMap,View.StaticMeshBatchVisibility);
+				Scene->BasePassForForwardShadingDistanceFieldShadowMapLightMapDrawList[FScene::EBasePass_Default].DrawVisible(RHICmdList, View,View.StaticMeshVisibilityMap,View.StaticMeshBatchVisibility);
+				Scene->BasePassForForwardShadingDirectionalLightAndSHIndirectDrawList[FScene::EBasePass_Default].DrawVisible(RHICmdList, View,View.StaticMeshVisibilityMap,View.StaticMeshBatchVisibility);
+				Scene->BasePassForForwardShadingNoLightMapDrawList[FScene::EBasePass_Default].DrawVisible(RHICmdList, View,View.StaticMeshVisibilityMap,View.StaticMeshBatchVisibility);
+			}
+
+			if (PrevRHICmdList != RHICmdList)
+			{
+				if (RHICmdList)
+				{
+					GRHICommandList.ExecuteAndFreeList(RHICmdList);
+				}
+
+				RHICmdList = PrevRHICmdList;
+			}
 		}
 
 		{
@@ -388,18 +417,18 @@ void FForwardShadingSceneRenderer::RenderForwardShadingBasePass()
 		if (SortMode == EBasePassSort::SortPerMesh)
 		{
 			SCOPE_CYCLE_COUNTER(STAT_StaticDrawListDrawTime);
-			MaxDraws -= Scene->BasePassForForwardShadingNoLightMapDrawList[FScene::EBasePass_Masked].DrawVisibleFrontToBack(View,View.StaticMeshVisibilityMap,View.StaticMeshBatchVisibility,MaxDraws);
-			MaxDraws -= Scene->BasePassForForwardShadingLowQualityLightMapDrawList[FScene::EBasePass_Masked].DrawVisibleFrontToBack(View,View.StaticMeshVisibilityMap,View.StaticMeshBatchVisibility,MaxDraws);
-			MaxDraws -= Scene->BasePassForForwardShadingDistanceFieldShadowMapLightMapDrawList[FScene::EBasePass_Masked].DrawVisibleFrontToBack(View,View.StaticMeshVisibilityMap,View.StaticMeshBatchVisibility,MaxDraws);
-			MaxDraws -= Scene->BasePassForForwardShadingDirectionalLightAndSHIndirectDrawList[FScene::EBasePass_Masked].DrawVisibleFrontToBack(View,View.StaticMeshVisibilityMap,View.StaticMeshBatchVisibility,MaxDraws);
+			MaxDraws -= Scene->BasePassForForwardShadingNoLightMapDrawList[FScene::EBasePass_Masked].DrawVisibleFrontToBack(RHICmdList, View,View.StaticMeshVisibilityMap,View.StaticMeshBatchVisibility,MaxDraws);
+			MaxDraws -= Scene->BasePassForForwardShadingLowQualityLightMapDrawList[FScene::EBasePass_Masked].DrawVisibleFrontToBack(RHICmdList, View,View.StaticMeshVisibilityMap,View.StaticMeshBatchVisibility,MaxDraws);
+			MaxDraws -= Scene->BasePassForForwardShadingDistanceFieldShadowMapLightMapDrawList[FScene::EBasePass_Masked].DrawVisibleFrontToBack(RHICmdList, View,View.StaticMeshVisibilityMap,View.StaticMeshBatchVisibility,MaxDraws);
+			MaxDraws -= Scene->BasePassForForwardShadingDirectionalLightAndSHIndirectDrawList[FScene::EBasePass_Masked].DrawVisibleFrontToBack(RHICmdList, View,View.StaticMeshVisibilityMap,View.StaticMeshBatchVisibility,MaxDraws);
 		}
 		else
 		{
 			SCOPE_CYCLE_COUNTER(STAT_StaticDrawListDrawTime);
-			Scene->BasePassForForwardShadingNoLightMapDrawList[FScene::EBasePass_Masked].DrawVisible(View,View.StaticMeshVisibilityMap,View.StaticMeshBatchVisibility);
-			Scene->BasePassForForwardShadingLowQualityLightMapDrawList[FScene::EBasePass_Masked].DrawVisible(View,View.StaticMeshVisibilityMap,View.StaticMeshBatchVisibility);
-			Scene->BasePassForForwardShadingDistanceFieldShadowMapLightMapDrawList[FScene::EBasePass_Masked].DrawVisible(View,View.StaticMeshVisibilityMap,View.StaticMeshBatchVisibility);
-			Scene->BasePassForForwardShadingDirectionalLightAndSHIndirectDrawList[FScene::EBasePass_Masked].DrawVisible(View,View.StaticMeshVisibilityMap,View.StaticMeshBatchVisibility);
+			Scene->BasePassForForwardShadingNoLightMapDrawList[FScene::EBasePass_Masked].DrawVisible(RHICmdList, View,View.StaticMeshVisibilityMap,View.StaticMeshBatchVisibility);
+			Scene->BasePassForForwardShadingLowQualityLightMapDrawList[FScene::EBasePass_Masked].DrawVisible(RHICmdList, View,View.StaticMeshVisibilityMap,View.StaticMeshBatchVisibility);
+			Scene->BasePassForForwardShadingDistanceFieldShadowMapLightMapDrawList[FScene::EBasePass_Masked].DrawVisible(RHICmdList, View,View.StaticMeshVisibilityMap,View.StaticMeshBatchVisibility);
+			Scene->BasePassForForwardShadingDirectionalLightAndSHIndirectDrawList[FScene::EBasePass_Masked].DrawVisible(RHICmdList, View,View.StaticMeshVisibilityMap,View.StaticMeshBatchVisibility);
 		}
 	}
 }
diff --git a/Engine/Source/Runtime/Renderer/Private/ForwardBasePassRendering.h b/Engine/Source/Runtime/Renderer/Private/ForwardBasePassRendering.h
--- a/Engine/Source/Runtime/Renderer/Private/ForwardBasePassRendering.h
+++ b/Engine/Source/Runtime/Renderer/Private/ForwardBasePassRendering.h
@@ -56,6 +56,7 @@ public:
 	}
 
 	void SetParameters(
+		FRHICommandList* RHICmdList, 
 		const FMaterialRenderProxy* MaterialRenderProxy,
 		const FVertexFactory* VertexFactory,
 		const FMaterial& InMaterialResource,
@@ -63,13 +64,13 @@ public:
 		ESceneRenderTargetsMode::Type TextureMode
 		)
 	{
-		HeightFogParameters.Set(GetVertexShader(), &View);
-		FMeshMaterialShader::SetParameters(GetVertexShader(),MaterialRenderProxy,InMaterialResource,View,TextureMode);
+		HeightFogParameters.Set(RHICmdList, GetVertexShader(), &View);
+		FMeshMaterialShader::SetParameters(RHICmdList, GetVertexShader(),MaterialRenderProxy,InMaterialResource,View,TextureMode);
 	}
 
-	void SetMesh(const FVertexFactory* VertexFactory,const FSceneView& View,const FPrimitiveSceneProxy* Proxy,const FMeshBatchElement& BatchElement)
+	void SetMesh(FRHICommandList* RHICmdList, const FVertexFactory* VertexFactory,const FSceneView& View,const FPrimitiveSceneProxy* Proxy,const FMeshBatchElement& BatchElement)
 	{
-		FMeshMaterialShader::SetMesh(GetVertexShader(),VertexFactory,View,Proxy,BatchElement);
+		FMeshMaterialShader::SetMesh(RHICmdList, GetVertexShader(),VertexFactory,View,Proxy,BatchElement);
 	}
 
 private:
@@ -129,12 +130,12 @@ public:
 	}
 	TBasePassForForwardShadingPSBaseType() {}
 
-	void SetParameters(const FMaterialRenderProxy* MaterialRenderProxy,const FMaterial& MaterialResource,const FSceneView* View,ESceneRenderTargetsMode::Type TextureMode)
+	void SetParameters(FRHICommandList* RHICmdList, const FMaterialRenderProxy* MaterialRenderProxy,const FMaterial& MaterialResource,const FSceneView* View,ESceneRenderTargetsMode::Type TextureMode)
 	{
-		FMeshMaterialShader::SetParameters(GetPixelShader(),MaterialRenderProxy,MaterialResource,*View,TextureMode);
+		FMeshMaterialShader::SetParameters(RHICmdList, GetPixelShader(),MaterialRenderProxy,MaterialResource,*View,TextureMode);
 	}
 
-	void SetMesh(const FVertexFactory* VertexFactory,const FSceneView& View,const FPrimitiveSceneProxy* Proxy,const FMeshBatchElement& BatchElement)
+	void SetMesh(FRHICommandList* RHICmdList, const FVertexFactory* VertexFactory,const FSceneView& View,const FPrimitiveSceneProxy* Proxy,const FMeshBatchElement& BatchElement)
 	{
 		if (ReflectionCubemap.IsBound())
 		{
@@ -150,10 +151,10 @@ public:
 			}
 
 			// Set the reflection cubemap
-			SetTextureParameter(GetPixelShader(), ReflectionCubemap, ReflectionSampler, ReflectionTexture);
+			SetTextureParameter(RHICmdList, GetPixelShader(), ReflectionCubemap, ReflectionSampler, ReflectionTexture);
 		}
 
-		FMeshMaterialShader::SetMesh(GetPixelShader(),VertexFactory,View,Proxy,BatchElement);
+		FMeshMaterialShader::SetMesh(RHICmdList, GetPixelShader(),VertexFactory,View,Proxy,BatchElement);
 	}
 
 	virtual bool Serialize(FArchive& Ar)
@@ -263,35 +264,35 @@ public:
 			SceneTextureMode == Other.SceneTextureMode;
 	}
 
-	void DrawShared(const FSceneView* View,FBoundShaderStateRHIParamRef BoundShaderState) const
+	void DrawShared(FRHICommandList* RHICmdList, const FSceneView* View,FBoundShaderStateRHIParamRef BoundShaderState) const
 	{
 		// Set the actual shader & vertex declaration state
-		RHISetBoundShaderState(BoundShaderState);
+		FRHICommandList::SetBoundShaderState(RHICmdList, BoundShaderState);
 
-		VertexShader->SetParameters(MaterialRenderProxy, VertexFactory, *MaterialResource, *View, SceneTextureMode);
+		VertexShader->SetParameters(RHICmdList, MaterialRenderProxy, VertexFactory, *MaterialResource, *View, SceneTextureMode);
 
 #if !(UE_BUILD_SHIPPING || UE_BUILD_TEST)
 		if (bOverrideWithShaderComplexity)
 		{
 			if (BlendMode == BLEND_Opaque)
 			{
-				RHISetBlendState(TStaticBlendStateWriteMask<CW_RGBA>::GetRHI());
+				FRHICommandList::SetBlendState(RHICmdList, TStaticBlendStateWriteMask<CW_RGBA>::GetRHI());
 			}
 			else
 			{
 				// Add complexity to existing
-				RHISetBlendState(TStaticBlendState<CW_RGBA,BO_Add,BF_One,BF_One,BO_Add,BF_Zero,BF_One>::GetRHI());
+				FRHICommandList::SetBlendState(RHICmdList, TStaticBlendState<CW_RGBA,BO_Add,BF_One,BF_One,BO_Add,BF_Zero,BF_One>::GetRHI());
 			}
 
 			TShaderMapRef<FShaderComplexityAccumulatePS> ShaderComplexityPixelShader(GetGlobalShaderMap());
 			const uint32 NumPixelShaderInstructions = PixelShader->GetNumInstructions();
 			const uint32 NumVertexShaderInstructions = VertexShader->GetNumInstructions();
-			ShaderComplexityPixelShader->SetParameters(NumVertexShaderInstructions,NumPixelShaderInstructions, View->GetFeatureLevel());
+			ShaderComplexityPixelShader->SetParameters(RHICmdList, NumVertexShaderInstructions,NumPixelShaderInstructions, View->GetFeatureLevel());
 		}
 		else
 #endif
 		{
-			PixelShader->SetParameters(MaterialRenderProxy, *MaterialResource, View, SceneTextureMode);
+			PixelShader->SetParameters(RHICmdList, MaterialRenderProxy, *MaterialResource, View, SceneTextureMode);
 
 			switch(BlendMode)
 			{
@@ -303,21 +304,21 @@ public:
 				// Masked materials are rendered together in the base pass, where the blend state is set at a higher level
 				break;
 			case BLEND_Translucent:
-				RHISetBlendState(TStaticBlendState<CW_RGB, BO_Add,BF_SourceAlpha,BF_InverseSourceAlpha,BO_Add,BF_Zero,BF_InverseSourceAlpha>::GetRHI());
+				FRHICommandList::SetBlendState(RHICmdList, TStaticBlendState<CW_RGB, BO_Add,BF_SourceAlpha,BF_InverseSourceAlpha,BO_Add,BF_Zero,BF_InverseSourceAlpha>::GetRHI());
 				break;
 			case BLEND_Additive:
 				// Add to the existing scene color
-				RHISetBlendState(TStaticBlendState<CW_RGB, BO_Add,BF_One,BF_One,BO_Add,BF_Zero,BF_InverseSourceAlpha>::GetRHI());
+				FRHICommandList::SetBlendState(RHICmdList, TStaticBlendState<CW_RGB, BO_Add,BF_One,BF_One,BO_Add,BF_Zero,BF_InverseSourceAlpha>::GetRHI());
 				break;
 			case BLEND_Modulate:
 				// Modulate with the existing scene color
-				RHISetBlendState(TStaticBlendState<CW_RGB,BO_Add,BF_DestColor,BF_Zero>::GetRHI());
+				FRHICommandList::SetBlendState(RHICmdList, TStaticBlendState<CW_RGB,BO_Add,BF_DestColor,BF_Zero>::GetRHI());
 				break;
 			};
 		}
 		
 		// Set the light-map policy.
-		LightMapPolicy.Set(VertexShader,bOverrideWithShaderComplexity ? NULL : PixelShader,VertexShader,PixelShader,VertexFactory,MaterialRenderProxy,View);		
+		LightMapPolicy.Set(RHICmdList, VertexShader,bOverrideWithShaderComplexity ? NULL : PixelShader,VertexShader,PixelShader,VertexFactory,MaterialRenderProxy,View);		
 	}
 
 	/** 
@@ -349,6 +350,7 @@ public:
 	}
 
 	void SetMeshRenderState(
+		FRHICommandList* RHICmdList, 
 		const FSceneView& View,
 		const FPrimitiveSceneProxy* PrimitiveSceneProxy,
 		const FMeshBatch& Mesh,
@@ -359,6 +361,7 @@ public:
 	{
 		// Set the light-map policy's mesh-specific settings.
 		LightMapPolicy.SetMesh(
+			RHICmdList, 
 			View,
 			PrimitiveSceneProxy,
 			VertexShader,
@@ -370,7 +373,7 @@ public:
 			ElementData.LightMapElementData);
 
 		const FMeshBatchElement& BatchElement = Mesh.Elements[BatchElementIndex];
-		VertexShader->SetMesh(VertexFactory,View,PrimitiveSceneProxy,BatchElement);
+		VertexShader->SetMesh(RHICmdList, VertexFactory,View,PrimitiveSceneProxy,BatchElement);
 
 #if !(UE_BUILD_SHIPPING || UE_BUILD_TEST)
 		if (bOverrideWithShaderComplexity)
@@ -385,15 +388,15 @@ public:
 			TShaderMapRef<FShaderComplexityAccumulatePS> ShaderComplexityPixelShader(GetGlobalShaderMap());
 			const uint32 NumPixelShaderInstructions = PixelShader->GetNumInstructions();
 			const uint32 NumVertexShaderInstructions = VertexShader->GetNumInstructions();
-			ShaderComplexityPixelShader->SetParameters(NumVertexShaderInstructions,NumPixelShaderInstructions, View.GetFeatureLevel());
+			ShaderComplexityPixelShader->SetParameters(RHICmdList, NumVertexShaderInstructions,NumPixelShaderInstructions, View.GetFeatureLevel());
 		}
 		else
 #endif
 		{
-			PixelShader->SetMesh(VertexFactory,View,PrimitiveSceneProxy,BatchElement);
+			PixelShader->SetMesh(RHICmdList, VertexFactory,View,PrimitiveSceneProxy,BatchElement);
 		}
 
-		FMeshDrawingPolicy::SetMeshRenderState(View,PrimitiveSceneProxy,Mesh,BatchElementIndex,bBackFace,FMeshDrawingPolicy::ElementDataType());
+		FMeshDrawingPolicy::SetMeshRenderState(RHICmdList, View,PrimitiveSceneProxy,Mesh,BatchElementIndex,bBackFace,FMeshDrawingPolicy::ElementDataType());
 	}
 
 	friend int32 CompareDrawingPolicy(const TBasePassForForwardShadingDrawingPolicy& A,const TBasePassForForwardShadingDrawingPolicy& B)
@@ -435,6 +438,7 @@ public:
 
 	static void AddStaticMesh(FScene* Scene,FStaticMesh* StaticMesh);
 	static bool DrawDynamicMesh(
+		FRHICommandList* RHICmdList, 
 		const FSceneView& View,
 		ContextType DrawingContext,
 		const FMeshBatch& Mesh,
@@ -467,34 +471,38 @@ void ProcessBasePassMeshForForwardShading(
 
 	check(!AllowHighQualityLightmaps());
 
+	//@todo-rco: RHIPacketList
+	FRHICommandList* RHICmdList = nullptr;
+
 	if (LightMapInteraction.GetType() == LMIT_Texture)
 	{
-		const FShadowMapInteraction ShadowMapInteraction = (Parameters.Mesh.LCI && bIsLitMaterial)
-			? Parameters.Mesh.LCI->GetShadowMapInteraction()
+		const FShadowMapInteraction ShadowMapInteraction = (Parameters.Mesh.LCI && bIsLitMaterial) 
+			? Parameters.Mesh.LCI->GetShadowMapInteraction() 
 			: FShadowMapInteraction();
 
 		if (ShadowMapInteraction.GetType() == SMIT_Texture)
 		{
 			Action.template Process< TDistanceFieldShadowsAndLightMapPolicy<LQ_LIGHTMAP> >(
+				RHICmdList, 
 				Parameters,
 				TDistanceFieldShadowsAndLightMapPolicy<LQ_LIGHTMAP>(),
 				TDistanceFieldShadowsAndLightMapPolicy<LQ_LIGHTMAP>::ElementDataType(ShadowMapInteraction, LightMapInteraction));
 		}
 		else
 		{
-			Action.template Process< TLightMapPolicy<LQ_LIGHTMAP> >(Parameters, TLightMapPolicy<LQ_LIGHTMAP>(), LightMapInteraction);
+			Action.template Process< TLightMapPolicy<LQ_LIGHTMAP> >(RHICmdList, Parameters, TLightMapPolicy<LQ_LIGHTMAP>(), LightMapInteraction );
 		}
 	}
-	else if (bIsLitMaterial
+	else if (bIsLitMaterial 
 		&& IsIndirectLightingCacheAllowed(Parameters.FeatureLevel)
 		&& Parameters.PrimitiveSceneProxy
 		// Movable objects need to get their GI from the indirect lighting cache
 		&& Parameters.PrimitiveSceneProxy->IsMovable())
 	{
-		Action.template Process<FSimpleDirectionalLightAndSHIndirectPolicy>(Parameters,FSimpleDirectionalLightAndSHIndirectPolicy(),FSimpleDirectionalLightAndSHIndirectPolicy::ElementDataType(Action.ShouldPackAmbientSH())); 
+		Action.template Process<FSimpleDirectionalLightAndSHIndirectPolicy>(RHICmdList, Parameters,FSimpleDirectionalLightAndSHIndirectPolicy(),FSimpleDirectionalLightAndSHIndirectPolicy::ElementDataType(Action.ShouldPackAmbientSH())); 
 	}
 	else
 	{
-		Action.template Process<FNoLightMapPolicy>(Parameters,FNoLightMapPolicy(),FNoLightMapPolicy::ElementDataType());
+		Action.template Process<FNoLightMapPolicy>(RHICmdList, Parameters,FNoLightMapPolicy(),FNoLightMapPolicy::ElementDataType());
 	}
 }
diff --git a/Engine/Source/Runtime/Renderer/Private/ForwardShadingRenderer.cpp b/Engine/Source/Runtime/Renderer/Private/ForwardShadingRenderer.cpp
--- a/Engine/Source/Runtime/Renderer/Private/ForwardShadingRenderer.cpp
+++ b/Engine/Source/Runtime/Renderer/Private/ForwardShadingRenderer.cpp
@@ -83,7 +83,9 @@ void FForwardShadingSceneRenderer::Render()
 	// Note, this is a reversed Z depth surface, so 0.0f is the far plane.
 	RHIClear(true, FLinearColor::Black, true, 0.0f, true, 0, FIntRect());
 
-	RenderForwardShadingBasePass();
+	//@todo-rco: RHIPacketList
+	FRHICommandList* RHICmdList = nullptr;
+	RenderForwardShadingBasePass(RHICmdList);
 
 	// Notify the FX system that opaque primitives have been rendered.
 	if (Scene->FXSystem)
@@ -99,7 +101,7 @@ void FForwardShadingSceneRenderer::Render()
 		if (ViewFamily.EngineShowFlags.Refraction)
 		{
 			// to apply refraction effect by distorting the scene color
-			RenderDistortion();
+			RenderDistortion(RHICmdList);
 		}
 		RenderTranslucency();
 	}
diff --git a/Engine/Source/Runtime/Renderer/Private/ForwardTranslucentRendering.cpp b/Engine/Source/Runtime/Renderer/Private/ForwardTranslucentRendering.cpp
--- a/Engine/Source/Runtime/Renderer/Private/ForwardTranslucentRendering.cpp
+++ b/Engine/Source/Runtime/Renderer/Private/ForwardTranslucentRendering.cpp
@@ -25,9 +25,9 @@ public:
 	}
 	FForwardCopySceneAlphaPS() {}
 
-	void SetParameters(const FSceneView& View)
+	void SetParameters(FRHICommandList* RHICmdList, const FSceneView& View)
 	{
-		SceneTextureParameters.Set(GetPixelShader(), View);
+		SceneTextureParameters.Set(RHICmdList, GetPixelShader(), View);
 	}
 
 	virtual bool Serialize(FArchive& Ar)
@@ -45,7 +45,7 @@ IMPLEMENT_SHADER_TYPE(,FForwardCopySceneAlphaPS,TEXT("TranslucentLightingShaders
 
 FGlobalBoundShaderState ForwardCopySceneAlphaBoundShaderState;
 
-void FForwardShadingSceneRenderer::CopySceneAlpha(const FSceneView& View)
+void FForwardShadingSceneRenderer::CopySceneAlpha(FRHICommandList* RHICmdList, const FSceneView& View)
 {
 	SCOPED_DRAW_EVENTF(EventCopy, DEC_SCENE_ITEMS, TEXT("CopySceneAlpha"));
 	RHISetRasterizerState(TStaticRasterizerState<FM_Solid,CM_None>::GetRHI());
@@ -65,12 +65,12 @@ void FForwardShadingSceneRenderer::CopySceneAlpha(const FSceneView& View)
 	TShaderMapRef<FForwardCopySceneAlphaPS> PixelShader(GetGlobalShaderMap());
 	SetGlobalBoundShaderState(ForwardCopySceneAlphaBoundShaderState, GFilterVertexDeclaration.VertexDeclarationRHI, *ScreenVertexShader, *PixelShader);
 
-	PixelShader->SetParameters(View);
+	PixelShader->SetParameters(RHICmdList, View);
 
-	DrawRectangle(
-		0, 0,
-		X, Y,
-		0, 0,
+	DrawRectangle( 
+		0, 0, 
+		X, Y, 
+		0, 0, 
 		X, Y,
 		FIntPoint(X, Y),
 		GSceneRenderTargets.GetBufferSizeXY(),
@@ -105,21 +105,22 @@ public:
 		bBackFace(bInBackFace),
 		HitProxyId(InHitProxyId)
 	{}
-
+	
 	inline bool ShouldPackAmbientSH() const
 	{
 		// So shader code can read a single constant to get the ambient term
 		return true;
 	}
 
-	const FLightSceneInfo* GetSimpleDirectionalLight() const
-	{
+	const FLightSceneInfo* GetSimpleDirectionalLight() const 
+	{ 
 		return ((FScene*)View.Family->Scene)->SimpleDirectionalLight;
 	}
 
 	/** Draws the translucent mesh with a specific light-map type, and fog volume type */
 	template<typename LightMapPolicyType>
 	void Process(
+		FRHICommandList* RHICmdList, 
 		const FProcessBasePassMeshParameters& Parameters,
 		const LightMapPolicyType& LightMapPolicy,
 		const typename LightMapPolicyType::ElementDataType& LightMapElementData
@@ -138,6 +139,7 @@ public:
 			);
 
 		DrawingPolicy.DrawShared(
+			RHICmdList, 
 			&View,
 			DrawingPolicy.CreateBoundShaderState(View.GetFeatureLevel())
 			);
@@ -145,6 +147,7 @@ public:
 		for (int32 BatchElementIndex = 0; BatchElementIndex<Parameters.Mesh.Elements.Num(); BatchElementIndex++)
 		{
 			DrawingPolicy.SetMeshRenderState(
+				RHICmdList, 
 				View,
 				Parameters.PrimitiveSceneProxy,
 				Parameters.Mesh,
@@ -152,7 +155,7 @@ public:
 				bBackFace,
 				typename TBasePassForForwardShadingDrawingPolicy<LightMapPolicyType>::ElementDataType(LightMapElementData)
 				);
-			DrawingPolicy.DrawMesh(Parameters.Mesh, BatchElementIndex);
+			DrawingPolicy.DrawMesh(RHICmdList, Parameters.Mesh, BatchElementIndex);
 		}
 	}
 };
@@ -162,6 +165,7 @@ public:
  * @return true if the mesh rendered
  */
 bool FTranslucencyForwardShadingDrawingPolicyFactory::DrawDynamicMesh(
+	FRHICommandList* RHICmdList, 
 	const FViewInfo& View,
 	ContextType DrawingContext,
 	const FMeshBatch& Mesh,
@@ -208,6 +212,7 @@ bool FTranslucencyForwardShadingDrawingPolicyFactory::DrawDynamicMesh(
  * @return true if the mesh rendered
  */
 bool FTranslucencyForwardShadingDrawingPolicyFactory::DrawStaticMesh(
+	FRHICommandList* RHICmdList, 
 	const FViewInfo& View,
 	ContextType DrawingContext,
 	const FStaticMesh& StaticMesh,
@@ -220,6 +225,7 @@ bool FTranslucencyForwardShadingDrawingPolicyFactory::DrawStaticMesh(
 	const FMaterial* Material = StaticMesh.MaterialRenderProxy->GetMaterial(View.GetFeatureLevel());
 
 	bDirty |= DrawDynamicMesh(
+		RHICmdList, 
 		View,
 		DrawingContext,
 		StaticMesh,
@@ -249,8 +255,8 @@ void FTranslucentPrimSet::DrawPrimitivesForForwardShading(const FViewInfo& View,
 }
 
 void FTranslucentPrimSet::RenderPrimitiveForForwardShading(
-	const FViewInfo& View,
-	FPrimitiveSceneInfo* PrimitiveSceneInfo,
+	const FViewInfo& View, 
+	FPrimitiveSceneInfo* PrimitiveSceneInfo, 
 	const FPrimitiveViewRelevance& ViewRelevance) const
 {
 	checkSlow(ViewRelevance.HasTranslucency());
@@ -273,6 +279,9 @@ void FTranslucentPrimSet::RenderPrimitiveForForwardShading(
 				);
 		}
 
+		//@todo-rco: RHIPacketList
+		FRHICommandList* RHICmdList = nullptr;
+
 		// Render static scene prim
 		if( ViewRelevance.bStaticRelevance )
 		{
@@ -285,6 +294,7 @@ void FTranslucentPrimSet::RenderPrimitiveForForwardShading(
 					&& StaticMesh.IsTranslucent(View.GetFeatureLevel()) )
 				{
 					FTranslucencyForwardShadingDrawingPolicyFactory::DrawStaticMesh(
+						RHICmdList, 
 						View,
 						FTranslucencyForwardShadingDrawingPolicyFactory::ContextType(),
 						StaticMesh,
@@ -302,6 +312,9 @@ void FForwardShadingSceneRenderer::RenderTranslucency()
 {
 	if (ShouldRenderTranslucency())
 	{
+		//@todo-rco: RHIPacketList
+		FRHICommandList* RHICmdList = nullptr;
+
 		const bool bGammaSpace = !IsMobileHDR();
 		const bool bLinearHDR64 = !bGammaSpace && !IsMobileHDR32bpp();
 
@@ -313,14 +326,14 @@ void FForwardShadingSceneRenderer::RenderTranslucency()
 
 			const FViewInfo& View = Views[ViewIndex];
 
-#if PLATFORM_HTML5
+			#if PLATFORM_HTML5
 			// Copy the view so emulation of framebuffer fetch works for alpha=depth.
 			// Possible optimization: this copy shouldn't be needed unless something uses fetch of depth.
 			if(bLinearHDR64 && GSupportsRenderTargetFormat_PF_FloatRGBA && (GSupportsShaderFramebufferFetch == false) && (!IsPCPlatform(GRHIShaderPlatform)))
 			{
-				CopySceneAlpha(View);
+				CopySceneAlpha(RHICmdList, View);
 			}
-#endif
+			#endif 
 
 			if (!bGammaSpace)
 			{
diff --git a/Engine/Source/Runtime/Renderer/Private/GPUBenchmark.cpp b/Engine/Source/Runtime/Renderer/Private/GPUBenchmark.cpp
--- a/Engine/Source/Runtime/Renderer/Private/GPUBenchmark.cpp
+++ b/Engine/Source/Runtime/Renderer/Private/GPUBenchmark.cpp
@@ -57,13 +57,13 @@ public:
 		return bShaderHasOutdatedParameters;
 	}
 
-	void SetParameters(const FSceneView& View, TRefCountPtr<IPooledRenderTarget>& Src)
+	void SetParameters(FRHICommandList* RHICmdList, const FSceneView& View, TRefCountPtr<IPooledRenderTarget>& Src)
 	{
 		const FPixelShaderRHIParamRef ShaderRHI = GetPixelShader();
 
-		FGlobalShader::SetParameters(ShaderRHI, View);
+		FGlobalShader::SetParameters(RHICmdList, ShaderRHI, View);
 
-		SetTextureParameter(ShaderRHI, InputTexture, InputTextureSampler, TStaticSamplerState<>::GetRHI(), Src->GetRenderTargetItem().ShaderResourceTexture);
+		SetTextureParameter(RHICmdList, ShaderRHI, InputTexture, InputTextureSampler, TStaticSamplerState<>::GetRHI(), Src->GetRenderTargetItem().ShaderResourceTexture);
 	}
 
 	static const TCHAR* GetSourceFilename()
@@ -113,11 +113,11 @@ public:
 		return bShaderHasOutdatedParameters;
 	}
 
-	void SetParameters(const FSceneView& View)
+	void SetParameters(FRHICommandList* RHICmdList, const FSceneView& View)
 	{
 		const FVertexShaderRHIParamRef ShaderRHI = GetVertexShader();
 		
-		FGlobalShader::SetParameters(ShaderRHI, View);
+		FGlobalShader::SetParameters(RHICmdList, ShaderRHI, View);
 	}
 };
 
@@ -125,7 +125,7 @@ IMPLEMENT_SHADER_TYPE(,FPostProcessBenchmarkVS,TEXT("GPUBenchmark"),TEXT("MainBe
 
 
 template <uint32 Method>
-void RunBenchmarkShader(const FSceneView& View, TRefCountPtr<IPooledRenderTarget>& Src, uint32 Count)
+void RunBenchmarkShader(FRHICommandList* RHICmdList, const FSceneView& View, TRefCountPtr<IPooledRenderTarget>& Src, uint32 Count)
 {
 	TShaderMapRef<FPostProcessBenchmarkVS> VertexShader(GetGlobalShaderMap());
 	TShaderMapRef<FPostProcessBenchmarkPS<Method> > PixelShader(GetGlobalShaderMap());
@@ -134,8 +134,8 @@ void RunBenchmarkShader(const FSceneView& View, TRefCountPtr<IPooledRenderTarget
 
 	SetGlobalBoundShaderState(BoundShaderState, GFilterVertexDeclaration.VertexDeclarationRHI, *VertexShader, *PixelShader);
 
-	PixelShader->SetParameters(View, Src);
-	VertexShader->SetParameters(View);
+	PixelShader->SetParameters(RHICmdList, View, Src);
+	VertexShader->SetParameters(RHICmdList, View);
 
 	for(uint32 i = 0; i < Count; ++i)
 	{
@@ -151,17 +151,17 @@ void RunBenchmarkShader(const FSceneView& View, TRefCountPtr<IPooledRenderTarget
 	}
 }
 
-void RunBenchmarkShader(const FSceneView& View, uint32 MethodId, TRefCountPtr<IPooledRenderTarget>& Src, uint32 Count)
+void RunBenchmarkShader(FRHICommandList* RHICmdList, const FSceneView& View, uint32 MethodId, TRefCountPtr<IPooledRenderTarget>& Src, uint32 Count)
 {
 	SCOPED_DRAW_EVENTF(Benchmark, DEC_SCENE_ITEMS, TEXT("Benchmark Method:%d"), MethodId);
 
 	switch(MethodId)
 	{
-		case 0: RunBenchmarkShader<0>(View, Src, Count); return;
-		case 1: RunBenchmarkShader<1>(View, Src, Count); return;
-		case 2: RunBenchmarkShader<2>(View, Src, Count); return;
-		case 3: RunBenchmarkShader<3>(View, Src, Count); return;
-		case 4: RunBenchmarkShader<4>(View, Src, Count); return;
+		case 0: RunBenchmarkShader<0>(RHICmdList, View, Src, Count); return;
+		case 1: RunBenchmarkShader<1>(RHICmdList, View, Src, Count); return;
+		case 2: RunBenchmarkShader<2>(RHICmdList, View, Src, Count); return;
+		case 3: RunBenchmarkShader<3>(RHICmdList, View, Src, Count); return;
+		case 4: RunBenchmarkShader<4>(RHICmdList, View, Src, Count); return;
 
 		default:
 			check(0);
@@ -346,6 +346,9 @@ private:
 
 void RendererGPUBenchmark(FSynthBenchmarkResults& InOut, const FSceneView& View, uint32 WorkScale, bool bDebugOut)
 {
+	//@todo-rco: RHIPacketList
+	FRHICommandList* RHICmdList = nullptr;
+
 	check(IsInRenderingThread());
 	
 	// two RT to ping pong so we force the GPU to flush it's pipeline
@@ -431,7 +434,7 @@ void RendererGPUBenchmark(FSynthBenchmarkResults& InOut, const FSceneView& View,
 				// decide how much work we do in this pass
 				PassCount[Iteration] = (Iteration / 10 + 1) * WorkScale;
 
-				RunBenchmarkShader(View, MethodId, RTItems[SrcRTIndex], PassCount[Iteration]);
+				RunBenchmarkShader(RHICmdList, View, MethodId, RTItems[SrcRTIndex], PassCount[Iteration]);
 
 				RHICopyToResolveTarget(RTItems[DestRTIndex]->GetRenderTargetItem().TargetableTexture, RTItems[DestRTIndex]->GetRenderTargetItem().ShaderResourceTexture, false, FResolveParams());
 
diff --git a/Engine/Source/Runtime/Renderer/Private/LightFunctionRendering.cpp b/Engine/Source/Runtime/Renderer/Private/LightFunctionRendering.cpp
--- a/Engine/Source/Runtime/Renderer/Private/LightFunctionRendering.cpp
+++ b/Engine/Source/Runtime/Renderer/Private/LightFunctionRendering.cpp
@@ -33,9 +33,9 @@ public:
 		StencilingGeometryParameters.Bind(Initializer.ParameterMap);
 	}
 
-	void SetParameters( const FSceneView* View, const FLightSceneInfo* LightSceneInfo )
+	void SetParameters(FRHICommandList* RHICmdList, const FSceneView* View, const FLightSceneInfo* LightSceneInfo )
 	{
-		FMaterialShader::SetParameters(GetVertexShader(), *View);
+		FMaterialShader::SetParameters(RHICmdList, GetVertexShader(), *View);
 		
 		// Light functions are projected using a bounding sphere.
 		// Calculate transform for bounding stencil sphere.
@@ -47,7 +47,7 @@ public:
 
 		FVector4 StencilingSpherePosAndScale;
 		StencilingGeometry::GStencilSphereVertexBuffer.CalcTransform(StencilingSpherePosAndScale, LightBounds, View->ViewMatrices.PreViewTranslation);
-		StencilingGeometryParameters.Set(this, StencilingSpherePosAndScale);
+		StencilingGeometryParameters.Set(RHICmdList, this, StencilingSpherePosAndScale);
 	}
 
 	// Begin FShader interface
@@ -93,11 +93,11 @@ public:
 		DeferredParameters.Bind(Initializer.ParameterMap);
 	}
 
-	void SetParameters( const FSceneView* View, const FLightSceneInfo* LightSceneInfo, const FMaterialRenderProxy* MaterialProxy, bool bRenderingPreviewShadowIndicator, float ShadowFadeFraction )
+	void SetParameters(FRHICommandList* RHICmdList, const FSceneView* View, const FLightSceneInfo* LightSceneInfo, const FMaterialRenderProxy* MaterialProxy, bool bRenderingPreviewShadowIndicator, float ShadowFadeFraction )
 	{
 		const FPixelShaderRHIParamRef ShaderRHI = GetPixelShader();
 
-		FMaterialShader::SetParameters(ShaderRHI, MaterialProxy, *MaterialProxy->GetMaterial(View->GetFeatureLevel()), *View, true, ESceneRenderTargetsMode::SetTextures);
+		FMaterialShader::SetParameters(RHICmdList, ShaderRHI, MaterialProxy, *MaterialProxy->GetMaterial(View->GetFeatureLevel()), *View, true, ESceneRenderTargetsMode::SetTextures);
 
 		// Set the transform from screen space to light space.
 		if ( ScreenToLight.IsBound() )
@@ -114,17 +114,17 @@ public:
 					FPlane(0,0,View->ViewMatrices.ProjMatrix.M[3][2],0))
 				* View->InvViewProjectionMatrix * WorldToLight;
 
-			SetShaderValue( ShaderRHI, ScreenToLight, ScreenToLightValue );
+			SetShaderValue(RHICmdList, ShaderRHI, ScreenToLight, ScreenToLightValue );
 		}
 
-		LightFunctionParameters.Set(ShaderRHI, LightSceneInfo, ShadowFadeFraction);
+		LightFunctionParameters.Set(RHICmdList, ShaderRHI, LightSceneInfo, ShadowFadeFraction);
 
-		SetShaderValue(ShaderRHI, LightFunctionParameters2, FVector(
+		SetShaderValue(RHICmdList, ShaderRHI, LightFunctionParameters2, FVector(
 			LightSceneInfo->Proxy->GetLightFunctionFadeDistance(), 
 			LightSceneInfo->Proxy->GetLightFunctionDisabledBrightness(),
 			bRenderingPreviewShadowIndicator ? 1.0f : 0.0f));
 
-		DeferredParameters.Set(ShaderRHI, *View);
+		DeferredParameters.Set(RHICmdList, ShaderRHI, *View);
 	}
 
 	virtual bool Serialize(FArchive& Ar)
@@ -212,27 +212,27 @@ bool FDeferredShadingSceneRenderer::CheckForLightFunction( const FLightSceneInfo
  *
  * @param LightSceneInfo Represents the current light
  */
-bool FDeferredShadingSceneRenderer::RenderLightFunction(const FLightSceneInfo* LightSceneInfo, bool bLightAttenuationCleared)
+bool FDeferredShadingSceneRenderer::RenderLightFunction(FRHICommandList* RHICmdList, const FLightSceneInfo* LightSceneInfo, bool bLightAttenuationCleared)
 {
 	if (ViewFamily.EngineShowFlags.LightFunctions)
 	{
-		return RenderLightFunctionForMaterial(LightSceneInfo, LightSceneInfo->Proxy->GetLightFunctionMaterial(), bLightAttenuationCleared, false);
+		return RenderLightFunctionForMaterial(RHICmdList, LightSceneInfo, LightSceneInfo->Proxy->GetLightFunctionMaterial(), bLightAttenuationCleared, false);
 	}
 	
 	return false;
 }
 
-bool FDeferredShadingSceneRenderer::RenderPreviewShadowsIndicator(const FLightSceneInfo* LightSceneInfo, bool bLightAttenuationCleared)
+bool FDeferredShadingSceneRenderer::RenderPreviewShadowsIndicator(FRHICommandList* RHICmdList, const FLightSceneInfo* LightSceneInfo, bool bLightAttenuationCleared)
 {
 	if (GEngine->PreviewShadowsIndicatorMaterial)
 	{
-		return RenderLightFunctionForMaterial(LightSceneInfo, GEngine->PreviewShadowsIndicatorMaterial->GetRenderProxy(false), bLightAttenuationCleared, true);
+		return RenderLightFunctionForMaterial(RHICmdList, LightSceneInfo, GEngine->PreviewShadowsIndicatorMaterial->GetRenderProxy(false), bLightAttenuationCleared, true);
 	}
 
 	return false;
 }
 
-bool FDeferredShadingSceneRenderer::RenderLightFunctionForMaterial(const FLightSceneInfo* LightSceneInfo, const FMaterialRenderProxy* MaterialProxy, bool bLightAttenuationCleared, bool bRenderingPreviewShadowsIndicator)
+bool FDeferredShadingSceneRenderer::RenderLightFunctionForMaterial(FRHICommandList* RHICmdList, const FLightSceneInfo* LightSceneInfo, const FMaterialRenderProxy* MaterialProxy, bool bLightAttenuationCleared, bool bRenderingPreviewShadowsIndicator)
 {
 	bool bRenderedLightFunction = false;
 
@@ -248,6 +248,7 @@ bool FDeferredShadingSceneRenderer::RenderLightFunctionForMaterial(const FLightS
 		const FMaterialShaderMap* MaterialShaderMap = Material->GetRenderingThreadShaderMap();
 		FLightFunctionVS* VertexShader = MaterialShaderMap->GetShader<FLightFunctionVS>();
 		FLightFunctionPS* PixelShader = MaterialShaderMap->GetShader<FLightFunctionPS>();
+		check(!RHICmdList);
 
 		// This was cached but when changing the material (e.g. editor) it wasn't updated.
 		// This will change with upcoming multi threaded rendering changes.
@@ -328,8 +329,8 @@ bool FDeferredShadingSceneRenderer::RenderLightFunctionForMaterial(const FLightS
 
 					// Render a bounding light sphere.
 					RHISetBoundShaderState(LightFunctionBoundShaderState);
-					VertexShader->SetParameters(&View, LightSceneInfo);
-					PixelShader->SetParameters(&View, LightSceneInfo, MaterialProxy, bRenderingPreviewShadowsIndicator, FadeAlpha);
+					VertexShader->SetParameters(RHICmdList, &View, LightSceneInfo);
+					PixelShader->SetParameters(RHICmdList, &View, LightSceneInfo, MaterialProxy, bRenderingPreviewShadowsIndicator, FadeAlpha);
 
 					// Project the light function using a sphere around the light
 					//@todo - could use a cone for spotlights
diff --git a/Engine/Source/Runtime/Renderer/Private/LightMapDensityRendering.cpp b/Engine/Source/Runtime/Renderer/Private/LightMapDensityRendering.cpp
--- a/Engine/Source/Runtime/Renderer/Private/LightMapDensityRendering.cpp
+++ b/Engine/Source/Runtime/Renderer/Private/LightMapDensityRendering.cpp
@@ -92,6 +92,7 @@ bool FDeferredShadingSceneRenderer::RenderLightMapDensities()
 }
 
 bool FLightMapDensityDrawingPolicyFactory::DrawDynamicMesh(
+	FRHICommandList* RHICmdList, 
 	const FViewInfo& View,
 	ContextType DrawingContext,
 	const FMeshBatch& Mesh,
@@ -136,24 +137,24 @@ bool FLightMapDensityDrawingPolicyFactory::DrawDynamicMesh(
 			if (bAllowHighQualityLightMaps)
 			{
 				TLightMapDensityDrawingPolicy< TLightMapPolicy<HQ_LIGHTMAP> > DrawingPolicy(View, Mesh.VertexFactory, MaterialRenderProxy, TLightMapPolicy<HQ_LIGHTMAP>(), BlendMode);
-				DrawingPolicy.DrawShared(&View,DrawingPolicy.CreateBoundShaderState(View.GetFeatureLevel()));
+				DrawingPolicy.DrawShared(RHICmdList, &View,DrawingPolicy.CreateBoundShaderState(View.GetFeatureLevel()));
 				for( int32 BatchElementIndex = 0; BatchElementIndex < Mesh.Elements.Num(); BatchElementIndex++)
 				{
-					DrawingPolicy.SetMeshRenderState(View,PrimitiveSceneProxy,Mesh,BatchElementIndex,bBackFace,
+					DrawingPolicy.SetMeshRenderState(RHICmdList, View,PrimitiveSceneProxy,Mesh,BatchElementIndex,bBackFace,
 						TLightMapDensityDrawingPolicy< TLightMapPolicy<HQ_LIGHTMAP> >::ElementDataType(LightMapInteraction));
-					DrawingPolicy.DrawMesh(Mesh,BatchElementIndex);
+					DrawingPolicy.DrawMesh(RHICmdList, Mesh,BatchElementIndex);
 				}
 				bDirty = true;
 			}
 			else
 			{
 				TLightMapDensityDrawingPolicy< TLightMapPolicy<LQ_LIGHTMAP> > DrawingPolicy(View, Mesh.VertexFactory, MaterialRenderProxy, TLightMapPolicy<LQ_LIGHTMAP>(), BlendMode);
-				DrawingPolicy.DrawShared(&View, DrawingPolicy.CreateBoundShaderState(View.GetFeatureLevel()));
+				DrawingPolicy.DrawShared(RHICmdList, &View, DrawingPolicy.CreateBoundShaderState(View.GetFeatureLevel()));
 				for( int32 BatchElementIndex = 0; BatchElementIndex < Mesh.Elements.Num(); BatchElementIndex++)
 				{
-					DrawingPolicy.SetMeshRenderState(View,PrimitiveSceneProxy,Mesh,BatchElementIndex,bBackFace,
+					DrawingPolicy.SetMeshRenderState(RHICmdList, View,PrimitiveSceneProxy,Mesh,BatchElementIndex,bBackFace,
 						TLightMapDensityDrawingPolicy< TLightMapPolicy<LQ_LIGHTMAP> >::ElementDataType(LightMapInteraction));
-					DrawingPolicy.DrawMesh(Mesh,BatchElementIndex);
+					DrawingPolicy.DrawMesh(RHICmdList, Mesh,BatchElementIndex);
 				}
 				bDirty = true;
 			}
@@ -161,12 +162,12 @@ bool FLightMapDensityDrawingPolicyFactory::DrawDynamicMesh(
 		else
 		{
 			TLightMapDensityDrawingPolicy<FDummyLightMapPolicy> DrawingPolicy(View, Mesh.VertexFactory, MaterialRenderProxy, FDummyLightMapPolicy(), BlendMode);
-			DrawingPolicy.DrawShared(&View, DrawingPolicy.CreateBoundShaderState(View.GetFeatureLevel()));
+			DrawingPolicy.DrawShared(RHICmdList, &View, DrawingPolicy.CreateBoundShaderState(View.GetFeatureLevel()));
 			for( int32 BatchElementIndex = 0; BatchElementIndex < Mesh.Elements.Num(); BatchElementIndex++)
 			{
-				DrawingPolicy.SetMeshRenderState(View,PrimitiveSceneProxy,Mesh,BatchElementIndex,bBackFace,
+				DrawingPolicy.SetMeshRenderState(RHICmdList, View,PrimitiveSceneProxy,Mesh,BatchElementIndex,bBackFace,
 					TLightMapDensityDrawingPolicy<FDummyLightMapPolicy>::ElementDataType(LightMapInteraction));
-				DrawingPolicy.DrawMesh(Mesh,BatchElementIndex);
+				DrawingPolicy.DrawMesh(RHICmdList, Mesh,BatchElementIndex);
 			}
 			bDirty = true;
 		}
@@ -174,11 +175,11 @@ bool FLightMapDensityDrawingPolicyFactory::DrawDynamicMesh(
 	else
 	{
 		TLightMapDensityDrawingPolicy<FNoLightMapPolicy> DrawingPolicy(View, Mesh.VertexFactory, MaterialRenderProxy, FNoLightMapPolicy(), BlendMode);
-		DrawingPolicy.DrawShared(&View, DrawingPolicy.CreateBoundShaderState(View.GetFeatureLevel()));
+		DrawingPolicy.DrawShared(RHICmdList, &View, DrawingPolicy.CreateBoundShaderState(View.GetFeatureLevel()));
 		for( int32 BatchElementIndex = 0; BatchElementIndex < Mesh.Elements.Num(); BatchElementIndex++)
 		{
-			DrawingPolicy.SetMeshRenderState(View,PrimitiveSceneProxy,Mesh,BatchElementIndex,bBackFace,TLightMapDensityDrawingPolicy<FNoLightMapPolicy>::ElementDataType());
-			DrawingPolicy.DrawMesh(Mesh,BatchElementIndex);
+			DrawingPolicy.SetMeshRenderState(RHICmdList, View,PrimitiveSceneProxy,Mesh,BatchElementIndex,bBackFace,TLightMapDensityDrawingPolicy<FNoLightMapPolicy>::ElementDataType());
+			DrawingPolicy.DrawMesh(RHICmdList, Mesh,BatchElementIndex);
 		}
 		bDirty = true;
 	}
diff --git a/Engine/Source/Runtime/Renderer/Private/LightMapDensityRendering.h b/Engine/Source/Runtime/Renderer/Private/LightMapDensityRendering.h
--- a/Engine/Source/Runtime/Renderer/Private/LightMapDensityRendering.h
+++ b/Engine/Source/Runtime/Renderer/Private/LightMapDensityRendering.h
@@ -45,16 +45,17 @@ public:
 	}
 
 	void SetParameters(		
+		FRHICommandList* RHICmdList, 
 		const FMaterialRenderProxy* MaterialRenderProxy,
 		const FSceneView& View
 		)
 	{
-		FMeshMaterialShader::SetParameters(GetVertexShader(), MaterialRenderProxy, *MaterialRenderProxy->GetMaterial(View.GetFeatureLevel()), View, ESceneRenderTargetsMode::SetTextures);
+		FMeshMaterialShader::SetParameters(RHICmdList, GetVertexShader(), MaterialRenderProxy, *MaterialRenderProxy->GetMaterial(View.GetFeatureLevel()), View, ESceneRenderTargetsMode::SetTextures);
 	}
 
-	void SetMesh(const FVertexFactory* VertexFactory,const FSceneView& View,const FPrimitiveSceneProxy* Proxy,const FMeshBatchElement& BatchElement)
+	void SetMesh(FRHICommandList* RHICmdList, const FVertexFactory* VertexFactory,const FSceneView& View,const FPrimitiveSceneProxy* Proxy,const FMeshBatchElement& BatchElement)
 	{
-		FMeshMaterialShader::SetMesh(GetVertexShader(),VertexFactory,View,Proxy,BatchElement);
+		FMeshMaterialShader::SetMesh(RHICmdList, GetVertexShader(),VertexFactory,View,Proxy,BatchElement);
 	}
 };
 
@@ -158,13 +159,14 @@ public:
 	}
 	TLightMapDensityPS() {}
 
-	void SetParameters(const FMaterialRenderProxy* MaterialRenderProxy,const FSceneView* View)
+	void SetParameters(FRHICommandList* RHICmdList, const FMaterialRenderProxy* MaterialRenderProxy,const FSceneView* View)
 	{
-		FMeshMaterialShader::SetParameters(GetPixelShader(), MaterialRenderProxy, *MaterialRenderProxy->GetMaterial(View->GetFeatureLevel()), *View, ESceneRenderTargetsMode::SetTextures);
+		FMeshMaterialShader::SetParameters(RHICmdList, GetPixelShader(), MaterialRenderProxy, *MaterialRenderProxy->GetMaterial(View->GetFeatureLevel()), *View, ESceneRenderTargetsMode::SetTextures);
 
 		if (GridTexture.IsBound())
 		{
 			SetTextureParameter(
+				RHICmdList, 
 				GetPixelShader(),
 				GridTexture,
 				GridTextureSampler,
@@ -173,9 +175,9 @@ public:
 		}
 	}
 
-	void SetMesh(const FVertexFactory* VertexFactory,const FPrimitiveSceneProxy* PrimitiveSceneProxy,const FMeshBatchElement& BatchElement,const FSceneView& View,bool bBackFace,FVector& InBuiltLightingAndSelectedFlags,FVector2D& InLightMapResolutionScale, bool bTextureMapped)
+	void SetMesh(FRHICommandList* RHICmdList, const FVertexFactory* VertexFactory,const FPrimitiveSceneProxy* PrimitiveSceneProxy,const FMeshBatchElement& BatchElement,const FSceneView& View,bool bBackFace,FVector& InBuiltLightingAndSelectedFlags,FVector2D& InLightMapResolutionScale, bool bTextureMapped)
 	{
-		FMeshMaterialShader::SetMesh(GetPixelShader(),VertexFactory,View,PrimitiveSceneProxy,BatchElement);
+		FMeshMaterialShader::SetMesh(RHICmdList, GetPixelShader(),VertexFactory,View,PrimitiveSceneProxy,BatchElement);
 
 		if (LightMapDensity.IsBound())
 		{
@@ -184,11 +186,11 @@ public:
 				GEngine->MinLightMapDensity * GEngine->MinLightMapDensity,
 				GEngine->IdealLightMapDensity * GEngine->IdealLightMapDensity,
 				GEngine->MaxLightMapDensity * GEngine->MaxLightMapDensity );
-			SetShaderValue(GetPixelShader(),LightMapDensity,DensityParameters,0);
+			SetShaderValue(RHICmdList, GetPixelShader(),LightMapDensity,DensityParameters,0);
 		}
-		SetShaderValue(GetPixelShader(),BuiltLightingAndSelectedFlags,InBuiltLightingAndSelectedFlags,0);
-		SetShaderValue(GetPixelShader(),DensitySelectedColor,GEngine->LightMapDensitySelectedColor,0);
-		SetShaderValue(GetPixelShader(),LightMapResolutionScale,InLightMapResolutionScale,0);
+		SetShaderValue(RHICmdList, GetPixelShader(),BuiltLightingAndSelectedFlags,InBuiltLightingAndSelectedFlags,0);
+		SetShaderValue(RHICmdList, GetPixelShader(),DensitySelectedColor,GEngine->LightMapDensitySelectedColor,0);
+		SetShaderValue(RHICmdList, GetPixelShader(),LightMapResolutionScale,InLightMapResolutionScale,0);
 		if (LightMapDensityDisplayOptions.IsBound())
 		{
 			FVector4 OptionsParameter(
@@ -197,9 +199,9 @@ public:
 				(bTextureMapped == true) ? 1.0f : 0.0f,
 				(bTextureMapped == false) ? 1.0f : 0.0f
 				);
-			SetShaderValue(GetPixelShader(),LightMapDensityDisplayOptions,OptionsParameter,0);
+			SetShaderValue(RHICmdList, GetPixelShader(),LightMapDensityDisplayOptions,OptionsParameter,0);
 		}
-		SetShaderValue(GetPixelShader(),VertexMappedColor,GEngine->LightMapDensityVertexMappedColor,0);
+		SetShaderValue(RHICmdList, GetPixelShader(),VertexMappedColor,GEngine->LightMapDensityVertexMappedColor,0);
 	}
 
 	virtual bool Serialize(FArchive& Ar)
@@ -294,28 +296,31 @@ public:
 			LightMapPolicy == Other.LightMapPolicy;
 	}
 
-	void DrawShared( const FSceneView* View, FBoundShaderStateRHIRef ShaderState ) const
+	void DrawShared(FRHICommandList* RHICmdList, const FSceneView* View, FBoundShaderStateRHIRef ShaderState ) const
 	{
+		check(!RHICmdList);
+
 		// Set the actual shader & vertex declaration state
 		RHISetBoundShaderState(ShaderState);
 
 		// Set the base pass shader parameters for the material.
-		VertexShader->SetParameters(MaterialRenderProxy,*View);
-		PixelShader->SetParameters(MaterialRenderProxy,View);
+		VertexShader->SetParameters(RHICmdList, MaterialRenderProxy,*View);
+		PixelShader->SetParameters(RHICmdList, MaterialRenderProxy,View);
 
 		if(HullShader && DomainShader)
 		{
-			HullShader->SetParameters(MaterialRenderProxy,*View);
-			DomainShader->SetParameters(MaterialRenderProxy,*View);
+			HullShader->SetParameters(RHICmdList, MaterialRenderProxy,*View);
+			DomainShader->SetParameters(RHICmdList, MaterialRenderProxy,*View);
 		}
 
 		RHISetBlendState(TStaticBlendState<>::GetRHI());
 
 		// Set the light-map policy.
-		LightMapPolicy.Set(VertexShader,PixelShader,VertexShader,PixelShader,VertexFactory,MaterialRenderProxy,View);
+		LightMapPolicy.Set(RHICmdList, VertexShader,PixelShader,VertexShader,PixelShader,VertexFactory,MaterialRenderProxy,View);
 	}
 
 	void SetMeshRenderState(
+		FRHICommandList* RHICmdList, 
 		const FSceneView& View,
 		const FPrimitiveSceneProxy* PrimitiveSceneProxy,
 		const FMeshBatch& Mesh,
@@ -326,16 +331,16 @@ public:
 	{
 		const FMeshBatchElement& BatchElement = Mesh.Elements[BatchElementIndex];
 
-		VertexShader->SetMesh(VertexFactory,View,PrimitiveSceneProxy,BatchElement);
+		VertexShader->SetMesh(RHICmdList, VertexFactory,View,PrimitiveSceneProxy,BatchElement);
 
 		if(HullShader && DomainShader)
 		{
-			HullShader->SetMesh(VertexFactory,View,PrimitiveSceneProxy,BatchElement);
-			DomainShader->SetMesh(VertexFactory,View,PrimitiveSceneProxy,BatchElement);
+			HullShader->SetMesh(RHICmdList, VertexFactory,View,PrimitiveSceneProxy,BatchElement);
+			DomainShader->SetMesh(RHICmdList, VertexFactory,View,PrimitiveSceneProxy,BatchElement);
 		}
 
 		// Set the light-map policy's mesh-specific settings.
-		LightMapPolicy.SetMesh(View,PrimitiveSceneProxy,VertexShader,PixelShader,VertexShader,PixelShader,VertexFactory,MaterialRenderProxy,ElementData.LightMapElementData);
+		LightMapPolicy.SetMesh(RHICmdList, View,PrimitiveSceneProxy,VertexShader,PixelShader,VertexShader,PixelShader,VertexFactory,MaterialRenderProxy,ElementData.LightMapElementData);
 
 		// BuiltLightingAndSelectedFlags informs the shader is lighting is built or not for this primitive
 		FVector BuiltLightingAndSelectedFlags(0.0f,0.0f,0.0f);
@@ -387,8 +392,8 @@ public:
 
 		// Adjust for the grid texture being 2x2 repeating pattern...
 		LMResolutionScale *= 0.5f;
-		PixelShader->SetMesh(VertexFactory,PrimitiveSceneProxy, BatchElement, View, bBackFace, BuiltLightingAndSelectedFlags, LMResolutionScale, bTextureMapped);	
-		FMeshDrawingPolicy::SetMeshRenderState(View,PrimitiveSceneProxy,Mesh,BatchElementIndex,bBackFace,FMeshDrawingPolicy::ElementDataType());
+		PixelShader->SetMesh(RHICmdList, VertexFactory,PrimitiveSceneProxy, BatchElement, View, bBackFace, BuiltLightingAndSelectedFlags, LMResolutionScale, bTextureMapped);	
+		FMeshDrawingPolicy::SetMeshRenderState(RHICmdList, View,PrimitiveSceneProxy,Mesh,BatchElementIndex,bBackFace,FMeshDrawingPolicy::ElementDataType());
 	}
 
 	/** 
@@ -442,6 +447,7 @@ public:
 	struct ContextType {};
 
 	static bool DrawDynamicMesh(
+		FRHICommandList* RHICmdList, 
 		const FViewInfo& View,
 		ContextType DrawingContext,
 		const FMeshBatch& Mesh,
diff --git a/Engine/Source/Runtime/Renderer/Private/LightMapRendering.h b/Engine/Source/Runtime/Renderer/Private/LightMapRendering.h
--- a/Engine/Source/Runtime/Renderer/Private/LightMapRendering.h
+++ b/Engine/Source/Runtime/Renderer/Private/LightMapRendering.h
@@ -42,6 +42,7 @@ public:
 	{}
 
 	void Set(
+		FRHICommandList* RHICmdList, 
 		const VertexParametersType* VertexShaderParameters,
 		const PixelParametersType* PixelShaderParameters,
 		FShader* VertexShader,
@@ -52,10 +53,11 @@ public:
 		) const
 	{
 		check(VertexFactory);
-		VertexFactory->Set();
+		VertexFactory->Set(RHICmdList);
 	}
 
 	void SetMesh(
+		FRHICommandList* RHICmdList,
 		const FSceneView& View,
 		const FPrimitiveSceneProxy* PrimitiveSceneProxy,
 		const VertexParametersType* VertexShaderParameters,
@@ -103,11 +105,11 @@ public:
 			LightMapCoordinateScaleBiasParameter.Bind(ParameterMap,TEXT("LightMapCoordinateScaleBias"));
 		}
 
-		void SetLightMapScale(FShader* VertexShader,const FLightMapInteraction& LightMapInteraction) const
+		void SetLightMapScale(FRHICommandList* RHICmdList, FShader* VertexShader,const FLightMapInteraction& LightMapInteraction) const
 		{
 			const FVector2D LightmapCoordinateScale = LightMapInteraction.GetCoordinateScale();
 			const FVector2D LightmapCoordinateBias = LightMapInteraction.GetCoordinateBias();
-			SetShaderValue(VertexShader->GetVertexShader(),LightMapCoordinateScaleBiasParameter,FVector4(
+			SetShaderValue(RHICmdList, VertexShader->GetVertexShader(),LightMapCoordinateScaleBiasParameter,FVector4(
 				LightmapCoordinateScale.X,
 				LightmapCoordinateScale.Y,
 				LightmapCoordinateBias.X,
@@ -137,7 +139,7 @@ public:
 			LightMapAddParameter.Bind(ParameterMap,TEXT("LightMapAdd"));
 		}
 
-		void SetLightMapTexture(FShader* PixelShader, const UTexture2D* LightMapTexture) const
+		void SetLightMapTexture(FRHICommandList* RHICmdList, FShader* PixelShader, const UTexture2D* LightMapTexture) const
 		{
 			FTexture* TextureResource = GBlackTexture;
 
@@ -147,6 +149,7 @@ public:
 			}
 
 			SetTextureParameter(
+				RHICmdList, 
 				PixelShader->GetPixelShader(),
 				LightMapTextureParameter,
 				LightMapSamplerParameter,
@@ -154,7 +157,7 @@ public:
 				);
 		}
 
-		void SetSkyOcclusionTexture(FShader* PixelShader, const UTexture2D* SkyOcclusionTextureValue) const
+		void SetSkyOcclusionTexture(FRHICommandList* RHICmdList, FShader* PixelShader, const UTexture2D* SkyOcclusionTextureValue) const
 		{
 			FTexture* TextureResource = GWhiteTexture;
 
@@ -164,6 +167,7 @@ public:
 			}
 
 			SetTextureParameter(
+				RHICmdList, 
 				PixelShader->GetPixelShader(),
 				SkyOcclusionTexture,
 				SkyOcclusionSampler,
@@ -171,10 +175,10 @@ public:
 				);
 		}
 
-		void SetLightMapScale(FShader* PixelShader,const FLightMapInteraction& LightMapInteraction) const
+		void SetLightMapScale(FRHICommandList* RHICmdList, FShader* PixelShader,const FLightMapInteraction& LightMapInteraction) const
 		{
-			SetShaderValueArray(PixelShader->GetPixelShader(),LightMapScaleParameter,LightMapInteraction.GetScaleArray(),LightMapInteraction.GetNumLightmapCoefficients());
-			SetShaderValueArray(PixelShader->GetPixelShader(),LightMapAddParameter,LightMapInteraction.GetAddArray(),LightMapInteraction.GetNumLightmapCoefficients());
+			SetShaderValueArray(RHICmdList, PixelShader->GetPixelShader(),LightMapScaleParameter,LightMapInteraction.GetScaleArray(),LightMapInteraction.GetNumLightmapCoefficients());
+			SetShaderValueArray(RHICmdList, PixelShader->GetPixelShader(),LightMapAddParameter,LightMapInteraction.GetAddArray(),LightMapInteraction.GetNumLightmapCoefficients());
 		}
 
 		void Serialize(FArchive& Ar)
@@ -228,6 +232,7 @@ public:
 	}
 
 	void Set(
+		FRHICommandList* RHICmdList, 
 		const VertexParametersType* VertexShaderParameters,
 		const PixelParametersType* PixelShaderParameters,
 		FShader* VertexShader,
@@ -238,10 +243,11 @@ public:
 		) const
 	{
 		check(VertexFactory);
-		VertexFactory->Set();
+		VertexFactory->Set(RHICmdList);
 	}
 
 	void SetMesh(
+		FRHICommandList* RHICmdList, 
 		const FSceneView& View,
 		const FPrimitiveSceneProxy* PrimitiveSceneProxy,
 		const VertexParametersType* VertexShaderParameters,
@@ -255,14 +261,14 @@ public:
 	{
 		if(VertexShaderParameters)
 		{
-			VertexShaderParameters->SetLightMapScale(VertexShader,LightMapInteraction);
+			VertexShaderParameters->SetLightMapScale(RHICmdList, VertexShader,LightMapInteraction);
 		}
 
 		if(PixelShaderParameters)
 		{
-			PixelShaderParameters->SetLightMapScale(PixelShader,LightMapInteraction);
-			PixelShaderParameters->SetLightMapTexture(PixelShader, LightMapInteraction.GetTexture());
-			PixelShaderParameters->SetSkyOcclusionTexture(PixelShader, LightMapInteraction.GetSkyOcclusionTexture());
+			PixelShaderParameters->SetLightMapScale(RHICmdList, PixelShader,LightMapInteraction);
+			PixelShaderParameters->SetLightMapTexture(RHICmdList, PixelShader, LightMapInteraction.GetTexture());
+			PixelShaderParameters->SetSkyOcclusionTexture(RHICmdList, PixelShader, LightMapInteraction.GetSkyOcclusionTexture());
 		}
 	}
 
@@ -318,9 +324,9 @@ public:
 			Super::VertexParametersType::Bind(ParameterMap);
 		}
 
-		void SetMesh(FShader* VertexShader, const FShadowMapInteraction& ShadowMapInteraction) const
+		void SetMesh(FRHICommandList* RHICmdList, FShader* VertexShader, const FShadowMapInteraction& ShadowMapInteraction) const
 		{
-			SetShaderValue(VertexShader->GetVertexShader(), ShadowMapCoordinateScaleBias, FVector4(ShadowMapInteraction.GetCoordinateScale(), ShadowMapInteraction.GetCoordinateBias()));
+			SetShaderValue(RHICmdList, VertexShader->GetVertexShader(), ShadowMapCoordinateScaleBias, FVector4(ShadowMapInteraction.GetCoordinateScale(), ShadowMapInteraction.GetCoordinateBias()));
 		}
 
 		void Serialize(FArchive& Ar)
@@ -347,11 +353,11 @@ public:
 			Super::PixelParametersType::Bind(ParameterMap);
 		}
 
-		void SetMesh(FShader* PixelShader, const FShadowMapInteraction& ShadowMapInteraction, const FVector4& DistanceFieldValues) const
+		void SetMesh(FRHICommandList* RHICmdList, FShader* PixelShader, const FShadowMapInteraction& ShadowMapInteraction, const FVector4& DistanceFieldValues) const
 		{
-			SetShaderValue(PixelShader->GetPixelShader(), DistanceFieldParameters, DistanceFieldValues);
+			SetShaderValue(RHICmdList, PixelShader->GetPixelShader(), DistanceFieldParameters, DistanceFieldValues);
 			
-			SetShaderValue(PixelShader->GetPixelShader(), StaticShadowMapMasks, FVector4(
+			SetShaderValue(RHICmdList, PixelShader->GetPixelShader(), StaticShadowMapMasks, FVector4(
 				ShadowMapInteraction.GetChannelValid(0),
 				ShadowMapInteraction.GetChannelValid(1),
 				ShadowMapInteraction.GetChannelValid(2),
@@ -359,6 +365,7 @@ public:
 				));
 
 			SetTextureParameter(
+				RHICmdList, 
 				PixelShader->GetPixelShader(),
 				StaticShadowTexture,
 				StaticShadowSampler,
@@ -392,6 +399,7 @@ public:
 	TDistanceFieldShadowsAndLightMapPolicy() {}
 
 	void SetMesh(
+		FRHICommandList* RHICmdList, 
 		const FSceneView& View,
 		const FPrimitiveSceneProxy* PrimitiveSceneProxy,
 		const VertexParametersType* VertexShaderParameters,
@@ -405,15 +413,15 @@ public:
 	{
 		if (VertexShaderParameters)
 		{
-			VertexShaderParameters->SetMesh(VertexShader, ElementData.ShadowMapInteraction);
+			VertexShaderParameters->SetMesh(RHICmdList, VertexShader, ElementData.ShadowMapInteraction);
 		}
 
 		if (PixelShaderParameters)
 		{
-			PixelShaderParameters->SetMesh(PixelShader, ElementData.ShadowMapInteraction, ElementData.DistanceFieldValues);
+			PixelShaderParameters->SetMesh(RHICmdList, PixelShader, ElementData.ShadowMapInteraction, ElementData.DistanceFieldValues);
 		}
 
-		Super::SetMesh(View, PrimitiveSceneProxy, VertexShaderParameters, PixelShaderParameters, VertexShader, PixelShader, VertexFactory, MaterialRenderProxy, ElementData.SuperElementData);
+		Super::SetMesh(RHICmdList, View, PrimitiveSceneProxy, VertexShaderParameters, PixelShaderParameters, VertexShader, PixelShader, VertexFactory, MaterialRenderProxy, ElementData.SuperElementData);
 	}
 };
 
@@ -430,6 +438,7 @@ public:
 	}
 
 	void Set(
+		FRHICommandList* RHICmdList, 
 		const VertexParametersType* VertexShaderParameters,
 		const PixelParametersType* PixelShaderParameters,
 		FShader* VertexShader,
@@ -440,10 +449,11 @@ public:
 		) const
 	{
 		check(VertexFactory);
-		VertexFactory->Set();
+		VertexFactory->Set(RHICmdList);
 	}
 
 	void SetMesh(
+		FRHICommandList* RHICmdList, 
 		const FSceneView& View,
 		const FPrimitiveSceneProxy* PrimitiveSceneProxy,
 		const VertexParametersType* VertexShaderParameters,
@@ -457,12 +467,12 @@ public:
 	{
 		if(VertexShaderParameters)
 		{
-			VertexShaderParameters->SetLightMapScale(VertexShader,LightMapInteraction);
+			VertexShaderParameters->SetLightMapScale(RHICmdList, VertexShader,LightMapInteraction);
 		}
 
 		if(PixelShaderParameters)
 		{
-			PixelShaderParameters->SetLightMapScale(PixelShader,LightMapInteraction);
+			PixelShaderParameters->SetLightMapScale(RHICmdList, PixelShader,LightMapInteraction);
 		}
 	}
 
@@ -541,6 +551,7 @@ public:
 	FSelfShadowedTranslucencyPolicy() {}
 
 	void SetMesh(
+		FRHICommandList* RHICmdList,
 		const FSceneView& View,
 		const FPrimitiveSceneProxy* PrimitiveSceneProxy,
 		const VertexParametersType* VertexShaderParameters,
@@ -555,25 +566,25 @@ public:
 		if (PixelShaderParameters)
 		{
 			// Set these even if ElementData.TranslucentSelfShadow is NULL to avoid a d3d debug error from the shader expecting texture SRV's when a different type are bound
-			PixelShaderParameters->TranslucencyShadowParameters.Set(PixelShader);
+			PixelShaderParameters->TranslucencyShadowParameters.Set(RHICmdList, PixelShader);
 
 			if (ElementData.TranslucentSelfShadow)
 			{
 				FVector4 ShadowmapMinMax;
 				FMatrix WorldToShadowMatrixValue = ElementData.TranslucentSelfShadow->GetWorldToShadowMatrix(ShadowmapMinMax);
 
-				SetShaderValue(PixelShader->GetPixelShader(), PixelShaderParameters->WorldToShadowMatrix, WorldToShadowMatrixValue);
-				SetShaderValue(PixelShader->GetPixelShader(), PixelShaderParameters->ShadowUVMinMax, ShadowmapMinMax);
-				SetShaderValue(PixelShader->GetPixelShader(), PixelShaderParameters->DirectionalLightDirection, ElementData.TranslucentSelfShadow->LightSceneInfo->Proxy->GetDirection());
+				SetShaderValue(RHICmdList, PixelShader->GetPixelShader(), PixelShaderParameters->WorldToShadowMatrix, WorldToShadowMatrixValue);
+				SetShaderValue(RHICmdList, PixelShader->GetPixelShader(), PixelShaderParameters->ShadowUVMinMax, ShadowmapMinMax);
+				SetShaderValue(RHICmdList, PixelShader->GetPixelShader(), PixelShaderParameters->DirectionalLightDirection, ElementData.TranslucentSelfShadow->LightSceneInfo->Proxy->GetDirection());
 				//@todo - support fading from both views
 				const float FadeAlpha = ElementData.TranslucentSelfShadow->FadeAlphas[0];
 				// Incorporate the diffuse scale of 1 / PI into the light color
 				const FVector4 DirectionalLightColorValue(FVector(ElementData.TranslucentSelfShadow->LightSceneInfo->Proxy->GetColor() * FadeAlpha / PI), FadeAlpha);
-				SetShaderValue(PixelShader->GetPixelShader(), PixelShaderParameters->DirectionalLightColor, DirectionalLightColorValue);
+				SetShaderValue(RHICmdList, PixelShader->GetPixelShader(), PixelShaderParameters->DirectionalLightColor, DirectionalLightColorValue);
 			}
 			else
 			{
-				SetShaderValue(PixelShader->GetPixelShader(), PixelShaderParameters->DirectionalLightColor, FVector4(0, 0, 0, 0));
+				SetShaderValue(RHICmdList, PixelShader->GetPixelShader(), PixelShaderParameters->DirectionalLightColor, FVector4(0, 0, 0, 0));
 			}
 		}
 	}
@@ -654,6 +665,7 @@ public:
 	FCachedVolumeIndirectLightingPolicy() {}
 
 	void Set(
+		FRHICommandList* RHICmdList, 
 		const VertexParametersType* VertexShaderParameters,
 		const PixelParametersType* PixelShaderParameters,
 		FShader* VertexShader,
@@ -664,6 +676,7 @@ public:
 		) const;
 
 	void SetMesh(
+		FRHICommandList* RHICmdList, 
 		const FSceneView& View,
 		const FPrimitiveSceneProxy* PrimitiveSceneProxy,
 		const VertexParametersType* VertexShaderParameters,
@@ -736,6 +749,7 @@ public:
 	FCachedPointIndirectLightingPolicy() {}
 
 	void SetMesh(
+		FRHICommandList* RHICmdList, 
 		const FSceneView& View,
 		const FPrimitiveSceneProxy* PrimitiveSceneProxy,
 		const VertexParametersType* VertexShaderParameters,
@@ -797,6 +811,7 @@ public:
 	FSelfShadowedCachedPointIndirectLightingPolicy() {}
 
 	void SetMesh(
+		FRHICommandList* RHICmdList, 
 		const FSceneView& View,
 		const FPrimitiveSceneProxy* PrimitiveSceneProxy,
 		const VertexParametersType* VertexShaderParameters,
@@ -896,6 +911,7 @@ public:
 	}
 
 	void Set(
+		FRHICommandList* RHICmdList, 
 		const VertexParametersType* VertexShaderParameters,
 		const PixelParametersType* PixelShaderParameters,
 		FShader* VertexShader,
@@ -906,10 +922,11 @@ public:
 		) const
 	{
 		check(VertexFactory);
-		VertexFactory->Set();
+		VertexFactory->Set(RHICmdList);
 	}
 
 	void SetMesh(
+		FRHICommandList* RHICmdList, 
 		const FSceneView& View,
 		const FPrimitiveSceneProxy* PrimitiveSceneProxy,
 		const VertexParametersType* VertexShaderParameters,
@@ -924,6 +941,7 @@ public:
 		if (VertexShaderParameters && PixelShaderParameters)
 		{
 			CachedPointIndirectLightingPolicy.SetMesh(
+				RHICmdList, 
 				View, 
 				PrimitiveSceneProxy, 
 				&VertexShaderParameters->CachedPointIndirectParameters, 
diff --git a/Engine/Source/Runtime/Renderer/Private/LightPropagationVolume.cpp b/Engine/Source/Runtime/Renderer/Private/LightPropagationVolume.cpp
--- a/Engine/Source/Runtime/Renderer/Private/LightPropagationVolume.cpp
+++ b/Engine/Source/Runtime/Renderer/Private/LightPropagationVolume.cpp
@@ -177,10 +177,11 @@ public:
 		return bShaderHasOutdatedParameters;
 	}
 
-	void SetParameters( const FLpvBaseWriteShaderParams& Params )
+	void SetParameters(FRHICommandList* RHICmdList, const FLpvBaseWriteShaderParams& Params )
 	{
 		FComputeShaderRHIParamRef ShaderRHI = GetComputeShader();
-		SetUniformBufferParameter( ShaderRHI, GetUniformBufferParameter<FLpvWriteUniformBufferParameters>(), Params.UniformBuffer );
+		SetUniformBufferParameter(RHICmdList, ShaderRHI, GetUniformBufferParameter<FLpvWriteUniformBufferParameters>(), Params.UniformBuffer );
+		check(!RHICmdList);
 #if LPV_VOLUME_TEXTURE
 		for(int i  =0; i < 7; i++)
 		{
@@ -193,7 +194,7 @@ public:
 				RHISetUAVParameter( ShaderRHI, LpvBufferUAVs[i].GetBaseIndex(), Params.LpvBufferUAVs[i] );
 			}
 
-			SetTextureParameter( ShaderRHI, LpvBufferSRVParameters[i], LpvVolumeTextureSampler, TStaticSamplerState<SF_Bilinear,AM_Clamp,AM_Clamp,AM_Clamp>::GetRHI(), Params.LpvBufferSRVs[i] );
+			SetTextureParameter(RHICmdList, ShaderRHI, LpvBufferSRVParameters[i], LpvVolumeTextureSampler, TStaticSamplerState<SF_Bilinear,AM_Clamp,AM_Clamp,AM_Clamp>::GetRHI(), Params.LpvBufferSRVs[i] );
 		}
 #else
 		if ( LpvBufferSRV.IsBound() )
@@ -233,7 +234,7 @@ public:
 				RHISetUAVParameter( ShaderRHI, GvBufferUAVs[i].GetBaseIndex(), Params.GvBufferUAVs[i] );
 			}
 
-			SetTextureParameter( ShaderRHI, GvBufferSRVParameters[i], LpvVolumeTextureSampler, TStaticSamplerState<SF_Bilinear,AM_Clamp,AM_Clamp,AM_Clamp>::GetRHI(), Params.GvBufferSRVs[i] );
+			SetTextureParameter(RHICmdList, ShaderRHI, GvBufferSRVParameters[i], LpvVolumeTextureSampler, TStaticSamplerState<SF_Bilinear,AM_Clamp,AM_Clamp,AM_Clamp>::GetRHI(), Params.GvBufferSRVs[i] );
 		}
 #else
 		if ( GvBufferSRV.IsBound() )
@@ -264,8 +265,9 @@ public:
 	}
 
 	// Unbinds any buffers that have been bound.
-	void UnbindBuffers()
+	void UnbindBuffers(FRHICommandList* RHICmdList)
 	{
+		check(!RHICmdList);
 		FComputeShaderRHIParamRef ShaderRHI = GetComputeShader();
 #if LPV_VOLUME_TEXTURE
 		for ( int i = 0; i < 7; i++ )
@@ -478,27 +480,28 @@ public:
 	}
 
 	void SetParameters(
+		FRHICommandList* RHICmdList, 
 		FLpvBaseWriteShaderParams& BaseParams,
 		FTextureRHIParamRef RsmDiffuseTextureRHI,
 		FTextureRHIParamRef RsmNormalTextureRHI,
 		FTextureRHIParamRef RsmDepthTextureRHI )
 	{
 		FComputeShaderRHIParamRef ShaderRHI = GetComputeShader();
-		FLpvWriteShaderCSBase::SetParameters( BaseParams );
+		FLpvWriteShaderCSBase::SetParameters(RHICmdList, BaseParams );
 
 		FSamplerStateRHIParamRef SamplerStateLinear  = TStaticSamplerState<SF_Bilinear,AM_Clamp,AM_Clamp,AM_Clamp>::GetRHI();
 		FSamplerStateRHIParamRef SamplerStatePoint   = TStaticSamplerState<SF_Point,AM_Clamp,AM_Clamp,AM_Clamp>::GetRHI();
 
 		// FIXME: Why do we have to bind a samplerstate to a sampler here? Presumably this is for legacy reasons... 
-		SetTextureParameter( ShaderRHI, RsmDiffuseTexture, LinearTextureSampler, SamplerStateLinear, RsmDiffuseTextureRHI );
-		SetTextureParameter( ShaderRHI, RsmNormalTexture, LinearTextureSampler, SamplerStateLinear, RsmNormalTextureRHI );
-		SetTextureParameter( ShaderRHI, RsmDepthTexture, PointTextureSampler, SamplerStatePoint, RsmDepthTextureRHI );
+		SetTextureParameter(RHICmdList, ShaderRHI, RsmDiffuseTexture, LinearTextureSampler, SamplerStateLinear, RsmDiffuseTextureRHI );
+		SetTextureParameter(RHICmdList, ShaderRHI, RsmNormalTexture, LinearTextureSampler, SamplerStateLinear, RsmNormalTextureRHI );
+		SetTextureParameter(RHICmdList, ShaderRHI, RsmDepthTexture, PointTextureSampler, SamplerStatePoint, RsmDepthTextureRHI );
 	}
 
 	// Unbinds any buffers that have been bound.
-	void UnbindBuffers()
+	void UnbindBuffers(FRHICommandList* RHICmdList)
 	{
-		FLpvWriteShaderCSBase::UnbindBuffers();
+		FLpvWriteShaderCSBase::UnbindBuffers(RHICmdList);
 	}
 
 
@@ -655,13 +658,13 @@ class FLpvInjectShader_Base : public FLpvWriteShaderCSBase
 {
 public:
 	void SetParameters(
+		FRHICommandList* RHICmdList, 
 		FLpvBaseWriteShaderParams& BaseParams,
 		FDirectLightInjectBufferRef& InjectUniformBuffer )
 	{
-		FLpvWriteShaderCSBase::SetParameters( BaseParams );
+		FLpvWriteShaderCSBase::SetParameters(RHICmdList, BaseParams );
 		FComputeShaderRHIParamRef ComputeShaderRHI = GetComputeShader();
-
-		SetUniformBufferParameter( ComputeShaderRHI, GetUniformBufferParameter<FLpvDirectLightInjectParameters>(), InjectUniformBuffer );
+		SetUniformBufferParameter(RHICmdList, ComputeShaderRHI, GetUniformBufferParameter<FLpvDirectLightInjectParameters>(), InjectUniformBuffer );
 	}
 
 	FLpvInjectShader_Base()	{	}
@@ -935,7 +938,7 @@ void FLightPropagationVolume::InitSettings( const FSceneView& View )
 /**
 * Clears the LPV
 */
-void FLightPropagationVolume::Clear()
+void FLightPropagationVolume::Clear(FRHICommandList* RHICmdList)
 {
 	if ( !bEnabled )
 	{
@@ -959,9 +962,9 @@ void FLightPropagationVolume::Clear()
 
 		FLpvBaseWriteShaderParams ShaderParams;
 		GetShaderParams( ShaderParams );
-		Shader->SetParameters( ShaderParams );
-		DispatchComputeShader(*Shader, LPV_GRIDRES/4, LPV_GRIDRES/4, LPV_GRIDRES/4 );
-		Shader->UnbindBuffers();
+		Shader->SetParameters(RHICmdList, ShaderParams );
+		DispatchComputeShader(RHICmdList, *Shader, LPV_GRIDRES/4, LPV_GRIDRES/4, LPV_GRIDRES/4 );
+		Shader->UnbindBuffers(RHICmdList);
 	}
 
 	// Clear the LPV (or fade, if REFINE_OVER_TIME is enabled)
@@ -971,9 +974,9 @@ void FLightPropagationVolume::Clear()
 
 		FLpvBaseWriteShaderParams ShaderParams;
 		GetShaderParams( ShaderParams );
-		Shader->SetParameters( ShaderParams );
-		DispatchComputeShader(*Shader, LPV_GRIDRES/4, LPV_GRIDRES/4, LPV_GRIDRES/4 );
-		Shader->UnbindBuffers();
+		Shader->SetParameters(RHICmdList, ShaderParams );
+		DispatchComputeShader(RHICmdList, *Shader, LPV_GRIDRES/4, LPV_GRIDRES/4, LPV_GRIDRES/4 );
+		Shader->UnbindBuffers(RHICmdList);
 	}
 
 	// Clear the geometry volume if necessary
@@ -984,13 +987,15 @@ void FLightPropagationVolume::Clear()
 
 		FLpvBaseWriteShaderParams ShaderParams;
 		GetShaderParams( ShaderParams );
-		Shader->SetParameters( ShaderParams );
-		DispatchComputeShader(*Shader, LPV_GRIDRES/4, LPV_GRIDRES/4, LPV_GRIDRES/4 );
-	Shader->UnbindBuffers();
+		Shader->SetParameters(RHICmdList, ShaderParams );
+		DispatchComputeShader(RHICmdList, *Shader, LPV_GRIDRES/4, LPV_GRIDRES/4, LPV_GRIDRES/4 );
+		Shader->UnbindBuffers(RHICmdList);
 	}
 	//RHIBatchComputeShaderEnd();
 
 	// Clear the list UAV counters
+	check(!RHICmdList);
+
 	RHISetUAVParameter( FComputeShaderRHIRef(), 7, mVplListBuffer->UAV, 0 );
 	RHISetUAVParameter( FComputeShaderRHIRef(), 7, GvListBuffer->UAV, 0 );
 	RHISetUAVParameter( FComputeShaderRHIRef(), 7, FUnorderedAccessViewRHIParamRef(), 0 );
@@ -1046,6 +1051,7 @@ void FLightPropagationVolume::SetVplInjectionConstants(
 * Injects a Directional light into the LPV
 */
 void FLightPropagationVolume::InjectDirectionalLightRSM(
+	FRHICommandList* RHICmdList, 
 	const FTexture2DRHIRef&		RsmDiffuseTex, 
 	const FTexture2DRHIRef&		RsmNormalTex, 
 	const FTexture2DRHIRef&		RsmDepthTex, 
@@ -1056,7 +1062,7 @@ void FLightPropagationVolume::InjectDirectionalLightRSM(
 	{
 		SCOPED_DRAW_EVENT(LpvInjectDirectionalLightRSM, DEC_LIGHT);
 
-		SetVplInjectionConstants( ProjectedShadowInfo, LightProxy );
+		SetVplInjectionConstants(ProjectedShadowInfo, LightProxy );
 
 		TShaderMapRef<FLpvInject_GenerateVplListsCS> Shader(GetGlobalShaderMap());
 		RHISetComputeShader(Shader->GetComputeShader());
@@ -1064,11 +1070,11 @@ void FLightPropagationVolume::InjectDirectionalLightRSM(
 		// Clear the list counter the first time this function is called in a frame
 		FLpvBaseWriteShaderParams ShaderParams;
 		GetShaderParams( ShaderParams );
-		Shader->SetParameters( ShaderParams, RsmDiffuseTex, RsmNormalTex, RsmDepthTex );
+		Shader->SetParameters(RHICmdList, ShaderParams, RsmDiffuseTex, RsmNormalTex, RsmDepthTex );
 
-		DispatchComputeShader(*Shader, 256/8, 256/8, 1 ); 
+		DispatchComputeShader(RHICmdList, *Shader, 256/8, 256/8, 1 ); 
 
-		Shader->UnbindBuffers(); 
+		Shader->UnbindBuffers(RHICmdList); 
 	}
 
 	// If this is the first directional light, build the geometry volume with it
@@ -1090,11 +1096,11 @@ void FLightPropagationVolume::InjectDirectionalLightRSM(
 
 		FLpvBaseWriteShaderParams ShaderParams;
 		GetShaderParams( ShaderParams );
-		Shader->SetParameters( ShaderParams );
+		Shader->SetParameters(RHICmdList, ShaderParams );
 
-		DispatchComputeShader(*Shader, LPV_GRIDRES/4, LPV_GRIDRES/4, LPV_GRIDRES/4 );
+		DispatchComputeShader(RHICmdList, *Shader, LPV_GRIDRES/4, LPV_GRIDRES/4, LPV_GRIDRES/4 );
 
-		Shader->UnbindBuffers(); 
+		Shader->UnbindBuffers(RHICmdList);
 	}
 
 	mInjectedLightCount++;
@@ -1103,7 +1109,7 @@ void FLightPropagationVolume::InjectDirectionalLightRSM(
 /**
 * Propagates light in the LPV 
 */
-void FLightPropagationVolume::Propagate()
+void FLightPropagationVolume::Propagate(FRHICommandList* RHICmdList)
 {
 	if ( !bEnabled )
 	{
@@ -1123,11 +1129,11 @@ void FLightPropagationVolume::Propagate()
 
 		FLpvBaseWriteShaderParams ShaderParams;
 		GetShaderParams( ShaderParams );
-		Shader->SetParameters( ShaderParams );
+		Shader->SetParameters(RHICmdList, ShaderParams );
 
-		DispatchComputeShader(*Shader, LPV_GRIDRES/4, LPV_GRIDRES/4, LPV_GRIDRES/4 );
+		DispatchComputeShader(RHICmdList, *Shader, LPV_GRIDRES/4, LPV_GRIDRES/4, LPV_GRIDRES/4 );
 
-		Shader->UnbindBuffers(); 
+		Shader->UnbindBuffers(RHICmdList);
 	}
 
 	// Propagate lighting, ping-ponging between the two buffers
@@ -1152,11 +1158,11 @@ void FLightPropagationVolume::Propagate()
 
 			FLpvBaseWriteShaderParams ShaderParams;
 			GetShaderParams( ShaderParams );
-			Shader->SetParameters( ShaderParams );
+			Shader->SetParameters(RHICmdList, ShaderParams );
 
-			DispatchComputeShader( Shader, LPV_GRIDRES/4, LPV_GRIDRES/4, LPV_GRIDRES/4 );
+			DispatchComputeShader(RHICmdList, Shader, LPV_GRIDRES/4, LPV_GRIDRES/4, LPV_GRIDRES/4 );
 
-			Shader->UnbindBuffers(); 
+			Shader->UnbindBuffers(RHICmdList); 
 		}
 	}
 
@@ -1203,7 +1209,7 @@ void FLightPropagationVolume::GetShaderParams( FLpvBaseWriteShaderParams& OutPar
 	OutParams.UniformBuffer = LpvWriteUniformBuffer;
 }
 
-void FLightPropagationVolume::InjectLightDirect( const FLightSceneProxy& Light )
+void FLightPropagationVolume::InjectLightDirect(FRHICommandList* RHICmdList, const FLightSceneProxy& Light )
 {
 	if(!bEnabled)
 	{
@@ -1246,7 +1252,7 @@ void FLightPropagationVolume::InjectLightDirect( const FLightSceneProxy& Light )
 		}
 		RHISetComputeShader(Shader->GetComputeShader());
 
-		FDirectLightInjectBufferRef InjectUniformBuffer = FDirectLightInjectBufferRef::CreateUniformBufferImmediate(InjectUniformBufferParams, UniformBuffer_SingleUse);
+		FDirectLightInjectBufferRef InjectUniformBuffer = FDirectLightInjectBufferRef::CreateUniformBufferImmediate(InjectUniformBufferParams, UniformBuffer_SingleDraw);
 
 		mWriteBufferIndex = 1 - mWriteBufferIndex; // Swap buffers with each iteration
 
@@ -1255,9 +1261,9 @@ void FLightPropagationVolume::InjectLightDirect( const FLightSceneProxy& Light )
 
 		LpvWriteUniformBuffer.SetContents( *LpvWriteUniformBufferParams );
 
-		Shader->SetParameters( ShaderParams, InjectUniformBuffer );
-		DispatchComputeShader(Shader, LPV_GRIDRES / 4, LPV_GRIDRES / 4, LPV_GRIDRES / 4 );
-		Shader->UnbindBuffers(); 
+		Shader->SetParameters(RHICmdList, ShaderParams, InjectUniformBuffer );
+		DispatchComputeShader(RHICmdList, Shader, LPV_GRIDRES / 4, LPV_GRIDRES / 4, LPV_GRIDRES / 4 );
+		Shader->UnbindBuffers(RHICmdList); 
 	}
 
 }
diff --git a/Engine/Source/Runtime/Renderer/Private/LightPropagationVolume.h b/Engine/Source/Runtime/Renderer/Private/LightPropagationVolume.h
--- a/Engine/Source/Runtime/Renderer/Private/LightPropagationVolume.h
+++ b/Engine/Source/Runtime/Renderer/Private/LightPropagationVolume.h
@@ -126,22 +126,23 @@ public:
 
 	void InitSettings( const FSceneView& View );
 
-	void Clear();
+	void Clear(FRHICommandList* RHICmdList);
 
 	void SetVplInjectionConstants(
 		const FProjectedShadowInfo&	ProjectedShadowInfo,
 		const FLightSceneProxy* LightProxy );
 
-	void InjectDirectionalLightRSM(		
+	void InjectDirectionalLightRSM(
+		FRHICommandList* RHICmdList, 
 		const FTexture2DRHIRef&		RsmDiffuseTex, 
 		const FTexture2DRHIRef&		RsmNormalTex, 
 		const FTexture2DRHIRef&		RsmDepthTex, 
 		const FProjectedShadowInfo&	ProjectedShadowInfo,
 		const FLinearColor&			LightColour );
 
-	void InjectLightDirect( const FLightSceneProxy& Light );
+	void InjectLightDirect(FRHICommandList* RHICmdList, const FLightSceneProxy& Light );
 
-	void Propagate();
+	void Propagate(FRHICommandList* RHICmdList);
 
 	void Visualise( const FSceneView& View ) const;
 
diff --git a/Engine/Source/Runtime/Renderer/Private/LightPropagationVolumeVisualisation.cpp b/Engine/Source/Runtime/Renderer/Private/LightPropagationVolumeVisualisation.cpp
--- a/Engine/Source/Runtime/Renderer/Private/LightPropagationVolumeVisualisation.cpp
+++ b/Engine/Source/Runtime/Renderer/Private/LightPropagationVolumeVisualisation.cpp
@@ -56,10 +56,11 @@ public:
 	static void ModifyCompilationEnvironment( EShaderPlatform Platform, FShaderCompilerEnvironment& OutEnvironment )				{ FLpvVisualiseBase::ModifyCompilationEnvironment( Platform, OutEnvironment ); }
 
 	void SetParameters(
+		FRHICommandList* RHICmdList, 
 		const FSceneView& View )
 	{
 		FGeometryShaderRHIParamRef ShaderRHI = GetGeometryShader();
-		FGlobalShader::SetParameters(ShaderRHI, View); 
+		FGlobalShader::SetParameters(RHICmdList, ShaderRHI, View); 
 	}
 };
 
@@ -77,10 +78,11 @@ public:
 	static void ModifyCompilationEnvironment( EShaderPlatform Platform, FShaderCompilerEnvironment& OutEnvironment )				{ FLpvVisualiseBase::ModifyCompilationEnvironment( Platform, OutEnvironment ); }
 
 	void SetParameters(
+		FRHICommandList* RHICmdList, 
 		const FSceneView& View )
 	{
 		FVertexShaderRHIParamRef ShaderRHI = GetVertexShader();
-		FGlobalShader::SetParameters(ShaderRHI, View); 
+		FGlobalShader::SetParameters(RHICmdList, ShaderRHI, View); 
 	}
 };
 
@@ -120,11 +122,13 @@ public:
 	static void ModifyCompilationEnvironment( EShaderPlatform Platform, FShaderCompilerEnvironment& OutEnvironment )				{ FLpvVisualiseBase::ModifyCompilationEnvironment( Platform, OutEnvironment ); }
 
 	void SetParameters(
+		FRHICommandList* RHICmdList, 
 		const FLightPropagationVolume* LPV,
 		const FSceneView& View )
 	{
 		FPixelShaderRHIParamRef ShaderRHI = GetPixelShader();
-		FGlobalShader::SetParameters(ShaderRHI, View); 
+		FGlobalShader::SetParameters(RHICmdList, ShaderRHI, View); 
+		check(!RHICmdList);
 
 #if LPV_VOLUME_TEXTURE
 		for ( int i=0; i<7; i++ )
@@ -134,7 +138,7 @@ public:
 			{
 				RHISetShaderTexture( ShaderRHI, LpvBufferSRVParameters[i].GetBaseIndex(), LpvBufferSrv );
 			}
-			SetTextureParameter( ShaderRHI, LpvBufferSRVParameters[i], LpvVolumeTextureSampler, TStaticSamplerState<SF_Bilinear,AM_Clamp,AM_Clamp,AM_Clamp>::GetRHI(), LpvBufferSrv );
+			SetTextureParameter(RHICmdList, ShaderRHI, LpvBufferSRVParameters[i], LpvVolumeTextureSampler, TStaticSamplerState<SF_Bilinear,AM_Clamp,AM_Clamp,AM_Clamp>::GetRHI(), LpvBufferSrv );
 		}
 #else
 		if ( InLpvBuffer.IsBound() ) 
@@ -151,7 +155,7 @@ public:
 			{
 				RHISetShaderTexture( ShaderRHI, GvBufferSRVParameters[i].GetBaseIndex(), GvBufferSrv );
 			}
-			SetTextureParameter( ShaderRHI, GvBufferSRVParameters[i], LpvVolumeTextureSampler, TStaticSamplerState<SF_Bilinear,AM_Clamp,AM_Clamp,AM_Clamp>::GetRHI(), GvBufferSrv );
+			SetTextureParameter(RHICmdList, ShaderRHI, GvBufferSRVParameters[i], LpvVolumeTextureSampler, TStaticSamplerState<SF_Bilinear,AM_Clamp,AM_Clamp,AM_Clamp>::GetRHI(), GvBufferSrv );
 		}
 
 #else
@@ -257,11 +261,12 @@ void FLightPropagationVolume::Visualise( const FSceneView& View ) const
 	RHISetRasterizerState( TStaticRasterizerState<FM_Solid,CM_None>::GetRHI() );
 	RHISetBlendState(TStaticBlendState<CW_RGB,BO_Add,BF_One,BF_One>::GetRHI());
 
+	//@todo-rco: RHIPacketList
 	SetGlobalBoundShaderState(LpvVisBoundShaderState, GSimpleElementVertexDeclaration.VertexDeclarationRHI, *VertexShader, *PixelShader, *GeometryShader);
 
-	VertexShader->SetParameters( View );
-	GeometryShader->SetParameters( View );
-	PixelShader->SetParameters( this, View );
+	VertexShader->SetParameters(nullptr, View );
+	GeometryShader->SetParameters(nullptr, View );
+	PixelShader->SetParameters(nullptr, this, View );
 
 	RHISetStreamSource(0, NULL, 0, 0);
 	RHIDrawPrimitive( PT_PointList, 0, 1, 32*3 );
diff --git a/Engine/Source/Runtime/Renderer/Private/LightRendering.cpp b/Engine/Source/Runtime/Renderer/Private/LightRendering.cpp
--- a/Engine/Source/Runtime/Renderer/Private/LightRendering.cpp
+++ b/Engine/Source/Runtime/Renderer/Private/LightRendering.cpp
@@ -35,6 +35,7 @@ public:
 	static void ModifyCompilationEnvironment(EShaderPlatform Platform, FShaderCompilerEnvironment& OutEnvironment)
 	{
 		FGlobalShader::ModifyCompilationEnvironment(Platform, OutEnvironment);
+		FDeferredPixelShaderParameters::ModifyCompilationEnvironment(Platform,OutEnvironment);
 		OutEnvironment.SetDefine(TEXT("USE_IES_PROFILE"), (uint32)bUseIESProfile);
 		OutEnvironment.SetDefine(TEXT("RADIAL_ATTENUATION"), (uint32)bRadialAttenuation);
 		OutEnvironment.SetDefine(TEXT("INVERSE_SQUARED_FALLOFF"), (uint32)bInverseSquaredFalloff);
@@ -58,18 +59,18 @@ public:
 	{
 	}
 
-	void SetParameters(const FSceneView& View, const FLightSceneInfo* LightSceneInfo)
+	void SetParameters(FRHICommandList* RHICmdList, const FSceneView& View, const FLightSceneInfo* LightSceneInfo)
 	{
 		const FPixelShaderRHIParamRef ShaderRHI = GetPixelShader();
-		SetParametersBase(ShaderRHI, View, LightSceneInfo->Proxy->GetIESTextureResource());
-		SetDeferredLightParameters(ShaderRHI, GetUniformBufferParameter<FDeferredLightUniformStruct>(), LightSceneInfo, View);
+		SetParametersBase(RHICmdList, ShaderRHI, View, LightSceneInfo->Proxy->GetIESTextureResource());
+		SetDeferredLightParameters(RHICmdList, ShaderRHI, GetUniformBufferParameter<FDeferredLightUniformStruct>(), LightSceneInfo, View);
 	}
 
-	void SetParametersSimpleLight(const FSceneView& View, const FSimpleLightEntry& SimpleLight, const FSimpleLightPerViewEntry& SimpleLightPerViewData)
+	void SetParametersSimpleLight(FRHICommandList* RHICmdList, const FSceneView& View, const FSimpleLightEntry& SimpleLight, const FSimpleLightPerViewEntry& SimpleLightPerViewData)
 	{
 		const FPixelShaderRHIParamRef ShaderRHI = GetPixelShader();
-		SetParametersBase(ShaderRHI, View, NULL);
-		SetSimpleDeferredLightParameters(ShaderRHI, GetUniformBufferParameter<FDeferredLightUniformStruct>(), SimpleLight, SimpleLightPerViewData, View);
+		SetParametersBase(RHICmdList, ShaderRHI, View, NULL);
+		SetSimpleDeferredLightParameters(RHICmdList, ShaderRHI, GetUniformBufferParameter<FDeferredLightUniformStruct>(), SimpleLight, SimpleLightPerViewData, View);
 	}
 
 	virtual bool Serialize(FArchive& Ar)
@@ -94,14 +95,15 @@ public:
 
 private:
 
-	void SetParametersBase(const FPixelShaderRHIParamRef ShaderRHI, const FSceneView& View, FTexture* IESTextureResource)
+	void SetParametersBase(FRHICommandList* RHICmdList, const FPixelShaderRHIParamRef ShaderRHI, const FSceneView& View, FTexture* IESTextureResource)
 	{
-		FGlobalShader::SetParameters(ShaderRHI,View);
-		DeferredParameters.Set(ShaderRHI, View);
+		FGlobalShader::SetParameters(RHICmdList, ShaderRHI,View);
+		DeferredParameters.Set(RHICmdList, ShaderRHI, View);
 
 		if(LightAttenuationTexture.IsBound())
 		{
 			SetTextureParameter(
+				RHICmdList, 
 				ShaderRHI,
 				LightAttenuationTexture,
 				LightAttenuationTextureSampler,
@@ -111,6 +113,7 @@ private:
 		}
 
 		SetTextureParameter(
+			RHICmdList, 
 			ShaderRHI,
 			PreIntegratedBRDF,
 			PreIntegratedBRDFSampler,
@@ -122,6 +125,7 @@ private:
 			FTextureRHIParamRef TextureRHI = IESTextureResource ? IESTextureResource->TextureRHI : GSystemTextures.WhiteDummy->GetRenderTargetItem().TargetableTexture;
 
 			SetTextureParameter(
+				RHICmdList, 
 				ShaderRHI,
 				IESTexture,
 				IESTextureSampler,
@@ -171,6 +175,7 @@ public:
 	static void ModifyCompilationEnvironment(EShaderPlatform Platform, FShaderCompilerEnvironment& OutEnvironment)
 	{
 		FGlobalShader::ModifyCompilationEnvironment(Platform, OutEnvironment);
+		FDeferredPixelShaderParameters::ModifyCompilationEnvironment(Platform,OutEnvironment);
 		OutEnvironment.SetDefine(TEXT("RADIAL_ATTENUATION"), (uint32)bRadialAttenuation);
 	}
 
@@ -185,14 +190,14 @@ public:
 	{
 	}
 
-	void SetParameters(const FSceneView& View, const FLightSceneInfo* LightSceneInfo)
+	void SetParameters(FRHICommandList* RHICmdList, const FSceneView& View, const FLightSceneInfo* LightSceneInfo)
 	{
 		const FPixelShaderRHIParamRef ShaderRHI = GetPixelShader();
-		FGlobalShader::SetParameters(ShaderRHI,View);
+		FGlobalShader::SetParameters(RHICmdList, ShaderRHI,View);
 		const float HasValidChannelValue = LightSceneInfo->Proxy->GetPreviewShadowMapChannel() == INDEX_NONE ? 0.0f : 1.0f;
-		SetShaderValue(ShaderRHI, HasValidChannel, HasValidChannelValue);
-		DeferredParameters.Set(ShaderRHI, View);
-		SetDeferredLightParameters(ShaderRHI, GetUniformBufferParameter<FDeferredLightUniformStruct>(), LightSceneInfo, View);
+		SetShaderValue(RHICmdList, ShaderRHI, HasValidChannel, HasValidChannelValue);
+		DeferredParameters.Set(RHICmdList, ShaderRHI, View);
+		SetDeferredLightParameters(RHICmdList, ShaderRHI, GetUniformBufferParameter<FDeferredLightUniformStruct>(), LightSceneInfo, View);
 	}
 
 	virtual bool Serialize(FArchive& Ar)
@@ -273,8 +278,10 @@ void GetLightNameForDrawEvent(const FLightSceneProxy* LightProxy, FString& Light
 uint32 GetShadowQuality();
 
 /** Renders the scene's lighting. */
-void FDeferredShadingSceneRenderer::RenderLights()
+void FDeferredShadingSceneRenderer::RenderLights(FRHICommandList* RHICmdList)
 {
+	check(!RHICmdList);
+
 	SCOPED_DRAW_EVENT(Lights, DEC_SCENE_ITEMS);
 
 	if(IsSimpleDynamicLightingEnabled())
@@ -408,7 +415,7 @@ void FDeferredShadingSceneRenderer::RenderLights()
 				{
 					// Update the range that needs to be processed by standard deferred to exclude the lights done with tiled
 					StandardDeferredStart = NumSortedLightsTiledDeferred;
-					RenderTiledDeferredLighting(SortedLights, NumSortedLightsTiledDeferred, SimpleLights);
+					RenderTiledDeferredLighting(RHICmdList, SortedLights, NumSortedLightsTiledDeferred, SimpleLights);
 				}
 			}
 			else if (SimpleLights.InstanceData.Num() > 0)
@@ -479,9 +486,9 @@ void FDeferredShadingSceneRenderer::RenderLights()
 				GSceneRenderTargets.BeginRenderingLightAttenuation();
 				RHIClear(true, FLinearColor::White, false, 0, false, 0, FIntRect());
 
-				bool bRenderedTranslucentObjectShadows = RenderTranslucentProjectedShadows( &LightSceneInfo );
+				bool bRenderedTranslucentObjectShadows = RenderTranslucentProjectedShadows(RHICmdList, &LightSceneInfo );
 				// Render non-modulated projected shadows to the attenuation buffer.
-				RenderProjectedShadows( &LightSceneInfo, bRenderedTranslucentObjectShadows, bInjectedTranslucentVolume );
+				RenderProjectedShadows(RHICmdList, &LightSceneInfo, bRenderedTranslucentObjectShadows, bInjectedTranslucentVolume );
 			}
 
 			// Render any reflective shadow maps (if necessary)
@@ -490,21 +497,21 @@ void FDeferredShadingSceneRenderer::RenderLights()
 				if ( LightSceneInfo.Proxy->HasReflectiveShadowMap() )
 				{
 					INC_DWORD_STAT(STAT_NumReflectiveShadowMapLights);
-					RenderReflectiveShadowMaps( &LightSceneInfo );
+					RenderReflectiveShadowMaps(RHICmdList, &LightSceneInfo );
 				}
 			}
 			
 			// Render light function to the attenuation buffer.
 			if (bDirectLighting)
 			{
-				const bool bLightFunctionRendered = RenderLightFunction(&LightSceneInfo, bDrawShadows);
+				const bool bLightFunctionRendered = RenderLightFunction(RHICmdList, &LightSceneInfo, bDrawShadows);
 
 				if (ViewFamily.EngineShowFlags.PreviewShadowsIndicator
 					&& !LightSceneInfo.bPrecomputedLightingIsValid 
 					&& LightSceneInfo.Proxy->HasStaticShadowing())
 				{
 					const bool bLightAttenuationCleared = bDrawShadows || bLightFunctionRendered;
-					RenderPreviewShadowsIndicator(&LightSceneInfo, bLightAttenuationCleared);
+					RenderPreviewShadowsIndicator(RHICmdList, &LightSceneInfo, bLightAttenuationCleared);
 				}
 
 				if (!bDrawShadows)
@@ -561,7 +568,8 @@ void FDeferredShadingSceneRenderer::RenderLights()
 							FLightPropagationVolume* Lpv = ViewState->GetLightPropagationVolume();
 							if ( Lpv && LightSceneInfo->Proxy )
 							{
-								Lpv->InjectLightDirect( *LightSceneInfo->Proxy );
+								//@todo-rco: RHIPacketList
+								Lpv->InjectLightDirect(nullptr, *LightSceneInfo->Proxy );
 							}
 						}
 					}					
@@ -651,6 +659,7 @@ static FVertexDeclarationRHIParamRef GetDeferredLightingVertexDeclaration()
 
 template<bool bUseIESProfile, bool bRadialAttenuation, bool bInverseSquaredFalloff>
 static void SetShaderTemplLighting(
+	FRHICommandList* RHICmdList, 
 	const FSceneView& View, 
 	FShader* VertexShader,
 	const FLightSceneInfo* LightSceneInfo)
@@ -659,18 +668,19 @@ static void SetShaderTemplLighting(
 	{
 		TShaderMapRef<TDeferredLightPS<false, bRadialAttenuation, false, true> > PixelShader(GetGlobalShaderMap());
 		SetGlobalBoundShaderState(PixelShader->GetBoundShaderState(), GetDeferredLightingVertexDeclaration<bRadialAttenuation>(), VertexShader, *PixelShader);
-		PixelShader->SetParameters(View, LightSceneInfo);
+		PixelShader->SetParameters(RHICmdList, View, LightSceneInfo);
 	}
 	else
 	{
 		TShaderMapRef<TDeferredLightPS<bUseIESProfile, bRadialAttenuation, bInverseSquaredFalloff, false> > PixelShader(GetGlobalShaderMap());
 		SetGlobalBoundShaderState(PixelShader->GetBoundShaderState(), GetDeferredLightingVertexDeclaration<bRadialAttenuation>(), VertexShader, *PixelShader);
-		PixelShader->SetParameters(View, LightSceneInfo);
+		PixelShader->SetParameters(RHICmdList, View, LightSceneInfo);
 	}
 }
 
 template<bool bUseIESProfile, bool bRadialAttenuation, bool bInverseSquaredFalloff>
 static void SetShaderTemplLightingSimple(
+	FRHICommandList* RHICmdList, 
 	const FSceneView& View, 
 	FShader* VertexShader,
 	const FSimpleLightEntry& SimpleLight,
@@ -680,13 +690,13 @@ static void SetShaderTemplLightingSimple(
 	{
 		TShaderMapRef<TDeferredLightPS<false, bRadialAttenuation, false, true> > PixelShader(GetGlobalShaderMap());
 		SetGlobalBoundShaderState(PixelShader->GetBoundShaderState(), GetDeferredLightingVertexDeclaration<bRadialAttenuation>(), VertexShader, *PixelShader);
-		PixelShader->SetParametersSimpleLight(View, SimpleLight, SimpleLightPerViewData);
+		PixelShader->SetParametersSimpleLight(RHICmdList, View, SimpleLight, SimpleLightPerViewData);
 	}
 	else
 	{
 		TShaderMapRef<TDeferredLightPS<bUseIESProfile, bRadialAttenuation, bInverseSquaredFalloff, false> > PixelShader(GetGlobalShaderMap());
 		SetGlobalBoundShaderState(PixelShader->GetBoundShaderState(), GetDeferredLightingVertexDeclaration<bRadialAttenuation>(), VertexShader, *PixelShader);
-		PixelShader->SetParametersSimpleLight(View, SimpleLight, SimpleLightPerViewData);
+		PixelShader->SetParametersSimpleLight(RHICmdList, View, SimpleLight, SimpleLightPerViewData);
 	}
 }
 
@@ -720,6 +730,9 @@ void FDeferredShadingSceneRenderer::RenderLight(const FLightSceneInfo* LightScen
 			bUseIESTexture = (LightSceneInfo->Proxy->GetIESTextureResource() != 0);
 		}
 
+		//@todo-rco: RHIPacketList
+		FRHICommandList* RHICmdList = nullptr;
+
 		// Set the device viewport for the view.
 		RHISetViewport(View.ViewRect.Min.X, View.ViewRect.Min.Y, 0.0f, View.ViewRect.Max.X, View.ViewRect.Max.Y, 1.0f);
 
@@ -734,21 +747,21 @@ void FDeferredShadingSceneRenderer::RenderLight(const FLightSceneInfo* LightScen
 			{
 				TShaderMapRef<TDeferredLightOverlapPS<false> > PixelShader(GetGlobalShaderMap());
 				SetGlobalBoundShaderState(PixelShader->GetBoundShaderState(), GetDeferredLightingVertexDeclaration<false>(), *VertexShader, *PixelShader);
-				PixelShader->SetParameters(View, LightSceneInfo);
+				PixelShader->SetParameters(RHICmdList, View, LightSceneInfo);
 			}
 			else
 			{
 				if(bUseIESTexture)
 				{
-					SetShaderTemplLighting<true, false, false>(View, *VertexShader, LightSceneInfo);
+					SetShaderTemplLighting<true, false, false>(RHICmdList, View, *VertexShader, LightSceneInfo);
 				}
 				else
 				{
-					SetShaderTemplLighting<false, false, false>(View, *VertexShader, LightSceneInfo);
+					SetShaderTemplLighting<false, false, false>(RHICmdList, View, *VertexShader, LightSceneInfo);
 				}
 			}
 
-			VertexShader->SetParameters(View, LightSceneInfo);
+			VertexShader->SetParameters(RHICmdList, View, LightSceneInfo);
 
 			// Apply the directional light as a full screen quad
 			DrawRectangle( 
@@ -771,7 +784,7 @@ void FDeferredShadingSceneRenderer::RenderLight(const FLightSceneInfo* LightScen
 			{
 				TShaderMapRef<TDeferredLightOverlapPS<true> > PixelShader(GetGlobalShaderMap());
 				SetGlobalBoundShaderState(PixelShader->GetBoundShaderState(), GetDeferredLightingVertexDeclaration<true>(), *VertexShader, *PixelShader);
-				PixelShader->SetParameters(View, LightSceneInfo);
+				PixelShader->SetParameters(RHICmdList, View, LightSceneInfo);
 			}
 			else
 			{
@@ -779,27 +792,27 @@ void FDeferredShadingSceneRenderer::RenderLight(const FLightSceneInfo* LightScen
 				{
 					if(bUseIESTexture)
 					{
-						SetShaderTemplLighting<true, true, true>(View, *VertexShader, LightSceneInfo);
+						SetShaderTemplLighting<true, true, true>(RHICmdList, View, *VertexShader, LightSceneInfo);
 					}
 					else
 					{
-						SetShaderTemplLighting<false, true, true>(View, *VertexShader, LightSceneInfo);
+						SetShaderTemplLighting<false, true, true>(RHICmdList, View, *VertexShader, LightSceneInfo);
 					}
 				}
 				else
 				{
 					if(bUseIESTexture)
 					{
-						SetShaderTemplLighting<true, true, false>(View, *VertexShader, LightSceneInfo);
+						SetShaderTemplLighting<true, true, false>(RHICmdList, View, *VertexShader, LightSceneInfo);
 					}
 					else
 					{
-						SetShaderTemplLighting<false, true, false>(View, *VertexShader, LightSceneInfo);
+						SetShaderTemplLighting<false, true, false>(RHICmdList, View, *VertexShader, LightSceneInfo);
 					}
 				}
 			}
 
-			VertexShader->SetParameters(View, LightSceneInfo);
+			VertexShader->SetParameters(RHICmdList, View, LightSceneInfo);
 
 			if (LightSceneInfo->Proxy->GetLightType() == LightType_Point)
 			{
@@ -827,6 +840,9 @@ void FDeferredShadingSceneRenderer::RenderSimpleLightsStandardDeferred(const FSi
 	INC_DWORD_STAT_BY(STAT_NumLightsUsingStandardDeferred, SimpleLights.InstanceData.Num());
 	SCOPED_DRAW_EVENT(StandardDeferredSimpleLights, DEC_SCENE_ITEMS);
 	
+	//@todo-rco: RHIPacketList
+	FRHICommandList* RHICmdList = nullptr;
+
 	// Use additive blending for color
 	RHISetBlendState(TStaticBlendState<CW_RGBA, BO_Add, BF_One, BF_One, BO_Add, BF_One, BF_One>::GetRHI());
 
@@ -834,9 +850,9 @@ void FDeferredShadingSceneRenderer::RenderSimpleLightsStandardDeferred(const FSi
 	for (int32 LightIndex = 0; LightIndex < SimpleLights.InstanceData.Num(); LightIndex++)
 	{
 		const FSimpleLightEntry& SimpleLight = SimpleLights.InstanceData[LightIndex];
-		
+
 		for (int32 ViewIndex = 0; ViewIndex < NumViews; ViewIndex++)
-		{
+	{
 			const FSimpleLightPerViewEntry& SimpleLightPerViewData = SimpleLights.GetViewDependentData(LightIndex, ViewIndex, NumViews);
 			const FSphere LightBounds(SimpleLightPerViewData.Position, SimpleLight.Radius);
 
@@ -851,15 +867,15 @@ void FDeferredShadingSceneRenderer::RenderSimpleLightsStandardDeferred(const FSi
 			if (SimpleLight.Exponent == 0)
 			{
 				// inverse squared
-				SetShaderTemplLightingSimple<false, true, true>(View, *VertexShader, SimpleLight, SimpleLightPerViewData);
+				SetShaderTemplLightingSimple<false, true, true>(RHICmdList, View, *VertexShader, SimpleLight, SimpleLightPerViewData);
 			}
 			else
 			{
 				// light's exponent, not inverse squared
-				SetShaderTemplLightingSimple<false, true, false>(View, *VertexShader, SimpleLight, SimpleLightPerViewData);
+				SetShaderTemplLightingSimple<false, true, false>(RHICmdList, View, *VertexShader, SimpleLight, SimpleLightPerViewData);
 			}
 
-			VertexShader->SetSimpleLightParameters(View, LightBounds);
+			VertexShader->SetSimpleLightParameters(RHICmdList, View, LightBounds);
 
 			// Apply the point or spot light with some approximately bounding geometry, 
 			// So we can get speedups from depth testing and not processing pixels outside of the light's influence.
