UE4 - merge first pass at parallel rendering to main

diff --git a/Engine/Source/Runtime/Renderer/Private/LightRendering.h b/Engine/Source/Runtime/Renderer/Private/LightRendering.h
--- a/Engine/Source/Runtime/Renderer/Private/LightRendering.h
+++ b/Engine/Source/Runtime/Renderer/Private/LightRendering.h
@@ -28,6 +28,7 @@ extern float GMinScreenRadiusForLights;
 
 template<typename ShaderRHIParamRef>
 void SetDeferredLightParameters(
+	FRHICommandList* RHICmdList, 
 	const ShaderRHIParamRef ShaderRHI, 
 	const TShaderUniformBufferParameter<FDeferredLightUniformStruct>& DeferredLightUniformBufferParameter, 
 	const FLightSceneInfo* LightSceneInfo,
@@ -90,11 +91,12 @@ void SetDeferredLightParameters(
 		DeferredLightUniformsValue.LightColor *= Fade;
 	}
 
-	SetUniformBufferParameterImmediate(ShaderRHI,DeferredLightUniformBufferParameter,DeferredLightUniformsValue);
+	SetUniformBufferParameterImmediate(RHICmdList, ShaderRHI,DeferredLightUniformBufferParameter,DeferredLightUniformsValue);
 }
 
 template<typename ShaderRHIParamRef>
 void SetSimpleDeferredLightParameters(
+	FRHICommandList* RHICmdList, 
 	const ShaderRHIParamRef ShaderRHI, 
 	const TShaderUniformBufferParameter<FDeferredLightUniformStruct>& DeferredLightUniformBufferParameter, 
 	const FSimpleLightEntry& SimpleLight,
@@ -114,7 +116,7 @@ void SetSimpleDeferredLightParameters(
 	DeferredLightUniformsValue.SourceRadius = 0;
 	DeferredLightUniformsValue.SourceLength = 0;
 
-	SetUniformBufferParameterImmediate(ShaderRHI, DeferredLightUniformBufferParameter, DeferredLightUniformsValue);
+	SetUniformBufferParameterImmediate(RHICmdList, ShaderRHI, DeferredLightUniformBufferParameter, DeferredLightUniformsValue);
 }
 
 /** Shader parameters needed to render a light function. */
@@ -128,13 +130,14 @@ public:
 	}
 
 	template<typename ShaderRHIParamRef>
-	void Set(const ShaderRHIParamRef ShaderRHI, const FLightSceneInfo* LightSceneInfo, float ShadowFadeFraction) const
+	void Set(FRHICommandList* RHICmdList, const ShaderRHIParamRef ShaderRHI, const FLightSceneInfo* LightSceneInfo, float ShadowFadeFraction) const
 	{
 		const bool bIsSpotLight = LightSceneInfo->Proxy->GetLightType() == LightType_Spot;
 		const bool bIsPointLight = LightSceneInfo->Proxy->GetLightType() == LightType_Point;
 		const float TanOuterAngle = bIsSpotLight ? FMath::Tan(LightSceneInfo->Proxy->GetOuterConeAngle()) : 1.0f;
 
 		SetShaderValue( 
+			RHICmdList, 
 			ShaderRHI, 
 			LightFunctionParameters, 
 			FVector4(TanOuterAngle, ShadowFadeFraction, bIsSpotLight ? 1.0f : 0.0f, bIsPointLight ? 1.0f : 0.0f));
@@ -171,19 +174,19 @@ public:
 		StencilingGeometryParameters.Bind(Initializer.ParameterMap);
 	}
 
-	void SetParameters(const FViewInfo& View, const FLightSceneInfo* LightSceneInfo)
+	void SetParameters(FRHICommandList* RHICmdList, const FViewInfo& View, const FLightSceneInfo* LightSceneInfo)
 	{
-		FGlobalShader::SetParameters(GetVertexShader(),View);
-		StencilingGeometryParameters.Set(this, View, LightSceneInfo);
+		FGlobalShader::SetParameters(RHICmdList, GetVertexShader(),View);
+		StencilingGeometryParameters.Set(RHICmdList, this, View, LightSceneInfo);
 	}
 
-	void SetSimpleLightParameters(const FViewInfo& View, const FSphere& LightBounds)
+	void SetSimpleLightParameters(FRHICommandList* RHICmdList, const FViewInfo& View, const FSphere& LightBounds)
 	{
-		FGlobalShader::SetParameters(GetVertexShader(),View);
+		FGlobalShader::SetParameters(RHICmdList, GetVertexShader(),View);
 
 		FVector4 StencilingSpherePosAndScale;
 		StencilingGeometry::GStencilSphereVertexBuffer.CalcTransform(StencilingSpherePosAndScale, LightBounds, View.ViewMatrices.PreViewTranslation);
-		StencilingGeometryParameters.Set(this, StencilingSpherePosAndScale);
+		StencilingGeometryParameters.Set(RHICmdList, this, StencilingSpherePosAndScale);
 	}
 
 	// Begin FShader Interface
diff --git a/Engine/Source/Runtime/Renderer/Private/LightShaftRendering.cpp b/Engine/Source/Runtime/Renderer/Private/LightShaftRendering.cpp
--- a/Engine/Source/Runtime/Renderer/Private/LightShaftRendering.cpp
+++ b/Engine/Source/Runtime/Renderer/Private/LightShaftRendering.cpp
@@ -105,7 +105,7 @@ public:
 	}
 
 	template<typename ShaderRHIParamRef>
-	void SetParameters(const ShaderRHIParamRef Shader, const FLightSceneInfo* LightSceneInfo, const FSceneView& View, TRefCountPtr<IPooledRenderTarget>& PassSource)
+	void SetParameters(FRHICommandList* RHICmdList, const ShaderRHIParamRef Shader, const FLightSceneInfo* LightSceneInfo, const FSceneView& View, TRefCountPtr<IPooledRenderTarget>& PassSource)
 	{
 		const uint32 DownsampleFactor = GetLightShaftDownsampleFactor();
 		FIntPoint DownSampledViewSize(FMath::FloorToInt(View.ViewRect.Width() / DownsampleFactor), FMath::FloorToInt(View.ViewRect.Height() / DownsampleFactor));
@@ -118,7 +118,7 @@ public:
 			1.0f / ViewRatioOfBuffer.X, 
 			(float)FilterBufferSize.Y / (FilterBufferSize.X * ViewRatioOfBuffer.Y));
 
-		SetShaderValue(Shader, AspectRatioAndInvAspectRatioParameter, AspectRatioAndInvAspectRatio);
+		SetShaderValue(RHICmdList, Shader, AspectRatioAndInvAspectRatioParameter, AspectRatioAndInvAspectRatio);
 
 		const FVector WorldSpaceBlurOrigin = LightSceneInfo->Proxy->GetLightPositionForLightShafts(View.ViewMatrices.ViewOrigin);
 		// Transform into texture coordinates
@@ -135,20 +135,20 @@ public:
 		ScreenSpaceBlurOrigin.Y *= InvBufferSizeY;
 		FVector2D TextureSpaceBlurOrigin(ScreenSpaceBlurOrigin * FVector2D(AspectRatioAndInvAspectRatio.Z, AspectRatioAndInvAspectRatio.W));
 
-		SetShaderValue(Shader, TextureSpaceBlurOriginParameter, TextureSpaceBlurOrigin);
+		SetShaderValue(RHICmdList, Shader, TextureSpaceBlurOriginParameter, TextureSpaceBlurOrigin);
 
-		SetShaderValue(Shader, WorldSpaceBlurOriginAndRadiusParameter, FVector4(WorldSpaceBlurOrigin, LightSceneInfo->Proxy->GetRadius()));
-		SetShaderValue(Shader, LightSourceRadius, LightSceneInfo->Proxy->GetSourceRadius());
+		SetShaderValue(RHICmdList, Shader, WorldSpaceBlurOriginAndRadiusParameter, FVector4(WorldSpaceBlurOrigin, LightSceneInfo->Proxy->GetRadius()));
+		SetShaderValue(RHICmdList, Shader, LightSourceRadius, LightSceneInfo->Proxy->GetSourceRadius());
 
 		const bool bIsSpotLight = LightSceneInfo->Proxy->GetLightType() == LightType_Spot;
 		if (bIsSpotLight)
 		{
-			SetShaderValue(Shader, WorldSpaceSpotDirectionParameter, LightSceneInfo->Proxy->GetDirection());
-			SetShaderValue(Shader, SpotAnglesParameter, LightSceneInfo->Proxy->GetLightShaftConeParams());
+			SetShaderValue(RHICmdList, Shader, WorldSpaceSpotDirectionParameter, LightSceneInfo->Proxy->GetDirection());
+			SetShaderValue(RHICmdList, Shader, SpotAnglesParameter, LightSceneInfo->Proxy->GetLightShaftConeParams());
 		}
 
 		const float DistanceFromLight = (View.ViewMatrices.ViewOrigin - WorldSpaceBlurOrigin).Size() + PointLightFadeDistanceIncrease;
-		SetShaderValue(Shader, WorldSpaceCameraPositionParameter, FVector4(View.ViewMatrices.ViewOrigin, DistanceFromLight));
+		SetShaderValue(RHICmdList, Shader, WorldSpaceCameraPositionParameter, FVector4(View.ViewMatrices.ViewOrigin, DistanceFromLight));
 
 		const FIntPoint DownSampledXY = View.ViewRect.Min / DownsampleFactor;
 		const uint32 DownsampledSizeX = View.ViewRect.Width() / DownsampleFactor;
@@ -167,17 +167,17 @@ public:
 		}
 
 		FVector4 UVMinMax( MinU, MinV, MaxU, MaxV );
-		SetShaderValue(Shader, UVMinMaxParameter, UVMinMax);
+		SetShaderValue(RHICmdList, Shader, UVMinMaxParameter, UVMinMax);
 
 		const FLinearColor BloomTint = LightSceneInfo->BloomTint;
-		SetShaderValue(Shader, BloomTintAndThresholdParameter, FVector4(BloomTint.R, BloomTint.G, BloomTint.B, LightSceneInfo->BloomThreshold));
+		SetShaderValue(RHICmdList, Shader, BloomTintAndThresholdParameter, FVector4(BloomTint.R, BloomTint.G, BloomTint.B, LightSceneInfo->BloomThreshold));
 
 		float OcclusionMaskDarkness;
 		float OcclusionDepthRange;
 		LightSceneInfo->Proxy->GetLightShaftOcclusionParameters(OcclusionMaskDarkness, OcclusionDepthRange);
 
 		const FVector4 LightShaftParameterValues(1.0f / OcclusionDepthRange, LightSceneInfo->BloomScale, 1, OcclusionMaskDarkness);
-		SetShaderValue(Shader, LightShaftParameters, LightShaftParameterValues);
+		SetShaderValue(RHICmdList, Shader, LightShaftParameters, LightShaftParameterValues);
 
 		float DistanceFade = 0.0f;
 		if (LightSceneInfo->Proxy->GetLightType() != LightType_Directional)
@@ -185,11 +185,12 @@ public:
 			DistanceFade = FMath::Clamp(DistanceFromLight / (LightSceneInfo->Proxy->GetRadius() * PointLightRadiusFadeFactor), 0.0f, 1.0f);
 		}
 
-		SetShaderValue(Shader, DistanceFadeParameter, DistanceFade);
+		SetShaderValue(RHICmdList, Shader, DistanceFadeParameter, DistanceFade);
 
 		if (IsValidRef(PassSource))
 		{
 			SetTextureParameter(
+				RHICmdList, 
 				Shader,
 				SourceTextureParameter, SourceTextureSamplerParameter,
 				TStaticSamplerState<SF_Bilinear,AM_Clamp,AM_Clamp,AM_Clamp>::GetRHI(),
@@ -237,9 +238,9 @@ public:
 	}
 
 	/** Sets shader parameter values */
-	void SetParameters(const FViewInfo& View)
+	void SetParameters(FRHICommandList* RHICmdList, const FViewInfo& View)
 	{
-		FGlobalShader::SetParameters(GetVertexShader(), View);
+		FGlobalShader::SetParameters(RHICmdList, GetVertexShader(), View);
 	}
 };
 
@@ -291,15 +292,15 @@ public:
 	}
 
 	/** Sets shader parameter values */
-	void SetParameters(const FLightSceneInfo* LightSceneInfo, const FViewInfo& View, TRefCountPtr<IPooledRenderTarget>& PassSource)
+	void SetParameters(FRHICommandList* RHICmdList, const FLightSceneInfo* LightSceneInfo, const FViewInfo& View, TRefCountPtr<IPooledRenderTarget>& PassSource)
 	{
-		FGlobalShader::SetParameters(GetPixelShader(), View);
-		LightShaftParameters.SetParameters(GetPixelShader(), LightSceneInfo, View, PassSource);
+		FGlobalShader::SetParameters(RHICmdList, GetPixelShader(), View);
+		LightShaftParameters.SetParameters(RHICmdList, GetPixelShader(), LightSceneInfo, View, PassSource);
 
 		const FIntPoint BufferSize = GSceneRenderTargets.GetBufferSizeXY();
 		FVector2D SampleOffsets(1.0f / BufferSize.X, 1.0f / BufferSize.Y);
-		SetShaderValue(GetPixelShader(),SampleOffsetsParameter,SampleOffsets);
-		SceneTextureParams.Set(GetPixelShader(), View);
+		SetShaderValue(RHICmdList, GetPixelShader(),SampleOffsetsParameter,SampleOffsets);
+		SceneTextureParams.Set(RHICmdList, GetPixelShader(), View);
 	}
 
 private:
@@ -360,13 +361,13 @@ public:
 	}
 
 	/** Sets shader parameter values */
-	void SetParameters(const FLightSceneInfo* LightSceneInfo, const FViewInfo& View, int32 PassIndex, TRefCountPtr<IPooledRenderTarget>& PassSource)
+	void SetParameters(FRHICommandList* RHICmdList, const FLightSceneInfo* LightSceneInfo, const FViewInfo& View, int32 PassIndex, TRefCountPtr<IPooledRenderTarget>& PassSource)
 	{
-		FGlobalShader::SetParameters(GetPixelShader(), View);
-		LightShaftParameters.SetParameters(GetPixelShader(), LightSceneInfo, View, PassSource);
+		FGlobalShader::SetParameters(RHICmdList, GetPixelShader(), View);
+		LightShaftParameters.SetParameters(RHICmdList, GetPixelShader(), LightSceneInfo, View, PassSource);
 
 		const FVector4 Parameters(GLightShaftBlurNumSamples, GLightShaftFirstPassDistance, PassIndex);
-		SetShaderValue(GetPixelShader(), RadialBlurParameters, Parameters);
+		SetShaderValue(RHICmdList, GetPixelShader(), RadialBlurParameters, Parameters);
 	}
 
 private:
@@ -410,10 +411,10 @@ public:
 	}
 
 	/** Sets shader parameter values */
-	void SetParameters(const FLightSceneInfo* LightSceneInfo, const FViewInfo& View, TRefCountPtr<IPooledRenderTarget>& PassSource)
+	void SetParameters(FRHICommandList* RHICmdList, const FLightSceneInfo* LightSceneInfo, const FViewInfo& View, TRefCountPtr<IPooledRenderTarget>& PassSource)
 	{
-		FGlobalShader::SetParameters(GetPixelShader(), View);
-		LightShaftParameters.SetParameters(GetPixelShader(), LightSceneInfo, View, PassSource);
+		FGlobalShader::SetParameters(RHICmdList, GetPixelShader(), View);
+		LightShaftParameters.SetParameters(RHICmdList, GetPixelShader(), LightSceneInfo, View, PassSource);
 	}
 
 private:
@@ -458,6 +459,9 @@ void DownsamplePass(const FViewInfo& View, const FLightSceneInfo* LightSceneInfo
 	const uint32 DownsampledSizeX = View.ViewRect.Width() / DownsampleFactor;
 	const uint32 DownsampledSizeY = View.ViewRect.Height() / DownsampleFactor;
 
+	//@todo-rco: RHIPacketList
+	FRHICommandList* RHICmdList = nullptr;
+
 	RHISetRenderTarget(LightShaftsDest->GetRenderTargetItem().TargetableTexture, FTextureRHIRef());
 	RHISetViewport(DownSampledXY.X, DownSampledXY.Y, 0.0f, DownSampledXY.X + DownsampledSizeX, DownSampledXY.Y + DownsampledSizeY, 1.0f);
 
@@ -471,14 +475,14 @@ void DownsamplePass(const FViewInfo& View, const FLightSceneInfo* LightSceneInfo
 		{
 			TShaderMapRef<TDownsampleLightShaftsPixelShader<LightType_Directional, bDownsampleOcclusion> > DownsampleLightShaftsPixelShader(GetGlobalShaderMap());
 			SetGlobalBoundShaderState(DownsampleDirectionalLightShaftsBoundShaderState[bDownsampleOcclusion], GFilterVertexDeclaration.VertexDeclarationRHI, *DownsampleLightShaftsVertexShader, *DownsampleLightShaftsPixelShader);
-			DownsampleLightShaftsPixelShader->SetParameters(LightSceneInfo, View, UnusedRT);
+			DownsampleLightShaftsPixelShader->SetParameters(RHICmdList, LightSceneInfo, View, UnusedRT);
 		}
 		break;
 	case LightType_Spot:
 		{
 			TShaderMapRef<TDownsampleLightShaftsPixelShader<LightType_Spot, bDownsampleOcclusion> > DownsampleLightShaftsPixelShader(GetGlobalShaderMap());
 			SetGlobalBoundShaderState(DownsampleSpotLightShaftsBoundShaderState[bDownsampleOcclusion], GFilterVertexDeclaration.VertexDeclarationRHI, *DownsampleLightShaftsVertexShader, *DownsampleLightShaftsPixelShader);
-			DownsampleLightShaftsPixelShader->SetParameters(LightSceneInfo, View, UnusedRT);
+			DownsampleLightShaftsPixelShader->SetParameters(RHICmdList, LightSceneInfo, View, UnusedRT);
 		}
 		break;
 	default:
@@ -486,12 +490,12 @@ void DownsamplePass(const FViewInfo& View, const FLightSceneInfo* LightSceneInfo
 		{
 			TShaderMapRef<TDownsampleLightShaftsPixelShader<LightType_Point, bDownsampleOcclusion> > DownsampleLightShaftsPixelShader(GetGlobalShaderMap());
 			SetGlobalBoundShaderState(DownsamplePointLightShaftsBoundShaderState[bDownsampleOcclusion], GFilterVertexDeclaration.VertexDeclarationRHI, *DownsampleLightShaftsVertexShader, *DownsampleLightShaftsPixelShader);
-			DownsampleLightShaftsPixelShader->SetParameters(LightSceneInfo, View, UnusedRT);
+			DownsampleLightShaftsPixelShader->SetParameters(RHICmdList, LightSceneInfo, View, UnusedRT);
 		}
 		break;
 	}
 
-	DownsampleLightShaftsVertexShader->SetParameters(View);
+	DownsampleLightShaftsVertexShader->SetParameters(RHICmdList, View);
 
 	// No depth tests, no backface culling.
 	RHISetBlendState(TStaticBlendState<>::GetRHI());
@@ -615,7 +619,8 @@ void ApplyRadialBlurPasses(
 		SetGlobalBoundShaderState(BlurLightShaftsBoundShaderState, GFilterVertexDeclaration.VertexDeclarationRHI, *ScreenVertexShader, *BlurLightShaftsPixelShader);
 
 		TRefCountPtr<IPooledRenderTarget>& EffectiveSource = PassIndex == 0 ? FirstPassSource : LightShaftsSource;
-		BlurLightShaftsPixelShader->SetParameters(LightSceneInfo, View, PassIndex, EffectiveSource);
+		/// ?
+		BlurLightShaftsPixelShader->SetParameters(nullptr, LightSceneInfo, View, PassIndex, EffectiveSource);
 
 		{
 			SCOPED_DRAW_EVENT(RadialBlur, DEC_SCENE_ITEMS);
@@ -657,7 +662,8 @@ void FinishOcclusionTerm(const FViewInfo& View, const FLightSceneInfo* const Lig
 
 	TShaderMapRef<FFinishOcclusionPixelShader> MaskOcclusionTermPixelShader(GetGlobalShaderMap());
 	SetGlobalBoundShaderState(AccumulateTermBoundShaderState, GFilterVertexDeclaration.VertexDeclarationRHI, *ScreenVertexShader, *MaskOcclusionTermPixelShader);
-	MaskOcclusionTermPixelShader->SetParameters(LightSceneInfo, View, LightShaftsSource);
+	/// ?
+	MaskOcclusionTermPixelShader->SetParameters(nullptr, LightSceneInfo, View, LightShaftsSource);
 
 	{
 		SCOPED_DRAW_EVENT(FinishOcclusion, DEC_SCENE_ITEMS);
@@ -805,11 +811,12 @@ public:
 	}
 
 	/** Sets shader parameter values */
-	void SetParameters(const FViewInfo& View, TRefCountPtr<IPooledRenderTarget>& LightShaftOcclusion)
+	void SetParameters(FRHICommandList* RHICmdList, const FViewInfo& View, TRefCountPtr<IPooledRenderTarget>& LightShaftOcclusion)
 	{
-		FGlobalShader::SetParameters(GetPixelShader(), View);
+		FGlobalShader::SetParameters(RHICmdList, GetPixelShader(), View);
 
 		SetTextureParameter(
+			RHICmdList, 
 			GetPixelShader(),
 			SourceTextureParameter, SourceTextureSamplerParameter,
 			TStaticSamplerState<SF_Bilinear,AM_Clamp,AM_Clamp,AM_Clamp>::GetRHI(),
@@ -840,7 +847,8 @@ void ApplyLightShaftBloom(const FViewInfo& View, const FLightSceneInfo* const Li
 
 	SetGlobalBoundShaderState(ApplyLightShaftsBoundShaderState, GFilterVertexDeclaration.VertexDeclarationRHI, *ScreenVertexShader, *ApplyLightShaftsPixelShader);
 
-	ApplyLightShaftsPixelShader->SetParameters(View, LightShaftsSource);
+	/// ?
+	ApplyLightShaftsPixelShader->SetParameters(nullptr, View, LightShaftsSource);
 
 	const FIntPoint BufferSize = GSceneRenderTargets.GetBufferSizeXY();
 	const uint32 DownsampleFactor	= GetLightShaftDownsampleFactor();
diff --git a/Engine/Source/Runtime/Renderer/Private/PostProcess/PostProcessAA.cpp b/Engine/Source/Runtime/Renderer/Private/PostProcess/PostProcessAA.cpp
--- a/Engine/Source/Runtime/Renderer/Private/PostProcess/PostProcessAA.cpp
+++ b/Engine/Source/Runtime/Renderer/Private/PostProcess/PostProcessAA.cpp
@@ -70,15 +70,15 @@ public:
 		return bShaderHasOutdatedParameters;
 	}
 
-	void SetParameters(const FRenderingCompositePassContext& Context)
+	void SetParameters(FRHICommandList* RHICmdList, const FRenderingCompositePassContext& Context)
 	{
 		const FPixelShaderRHIParamRef ShaderRHI = GetPixelShader();
 
 		const FSceneView& View = Context.View;
 
-		FGlobalShader::SetParameters(ShaderRHI, View);
+		FGlobalShader::SetParameters(RHICmdList, ShaderRHI, View);
 		
-		PostprocessParameter.SetPS(ShaderRHI, Context, TStaticSamplerState<SF_Bilinear,AM_Clamp,AM_Clamp,AM_Clamp>::GetRHI());
+		PostprocessParameter.SetPS(RHICmdList, ShaderRHI, Context, TStaticSamplerState<SF_Bilinear,AM_Clamp,AM_Clamp,AM_Clamp>::GetRHI());
 
 		const FPooledRenderTargetDesc* InputDesc = Context.Pass->GetInputDesc(ePId_Input0);
 
@@ -90,48 +90,48 @@ public:
 
 		FVector2D InvExtent = FVector2D(1.0f / InputDesc->Extent.X, 1.0f / InputDesc->Extent.Y);
 
-		SetShaderValue(ShaderRHI, fxaaQualityRcpFrame, InvExtent);
+		SetShaderValue(RHICmdList, ShaderRHI, fxaaQualityRcpFrame, InvExtent);
 
 		{
 			float N = 0.5f;
 			FVector4 Value(-N * InvExtent.X, -N * InvExtent.Y, N * InvExtent.X, N * InvExtent.Y);
-			SetShaderValue(ShaderRHI, fxaaConsoleRcpFrameOpt, Value);
+			SetShaderValue(RHICmdList, ShaderRHI, fxaaConsoleRcpFrameOpt, Value);
 		}
 
 		{
 			float N = 2.0f;
 			FVector4 Value(-N * InvExtent.X, -N * InvExtent.Y, N * InvExtent.X, N * InvExtent.Y);
-			SetShaderValue(ShaderRHI, fxaaConsoleRcpFrameOpt2, Value);
+			SetShaderValue(RHICmdList, ShaderRHI, fxaaConsoleRcpFrameOpt2, Value);
 		}
 
 		{
 			float Value = 0.75f;
-			SetShaderValue(ShaderRHI, fxaaQualitySubpix, Value);
+			SetShaderValue(RHICmdList, ShaderRHI, fxaaQualitySubpix, Value);
 		}
 
 		{
 			float Value = 0.166f;
-			SetShaderValue(ShaderRHI, fxaaQualityEdgeThreshold, Value);
+			SetShaderValue(RHICmdList, ShaderRHI, fxaaQualityEdgeThreshold, Value);
 		}
 
 		{
 			float Value = 0.0833f;
-			SetShaderValue(ShaderRHI, fxaaQualityEdgeThresholdMin, Value);
+			SetShaderValue(RHICmdList, ShaderRHI, fxaaQualityEdgeThresholdMin, Value);
 		}
 
 		{
 			float Value = 8.0f;
-			SetShaderValue(ShaderRHI, fxaaConsoleEdgeSharpness, Value);
+			SetShaderValue(RHICmdList, ShaderRHI, fxaaConsoleEdgeSharpness, Value);
 		}
 
 		{
 			float Value = 0.125f;
-			SetShaderValue(ShaderRHI, fxaaConsoleEdgeThreshold, Value);
+			SetShaderValue(RHICmdList, ShaderRHI, fxaaConsoleEdgeThreshold, Value);
 		}
 
 		{
 			float Value = 0.05f;
-			SetShaderValue(ShaderRHI, fxaaConsoleEdgeThresholdMin, Value);
+			SetShaderValue(RHICmdList, ShaderRHI, fxaaConsoleEdgeThresholdMin, Value);
 		}
 	}
 	
@@ -186,11 +186,11 @@ public:
 		return bShaderHasOutdatedParameters;
 	}
 
-	void SetParameters(const FRenderingCompositePassContext& Context)
+	void SetParameters(FRHICommandList* RHICmdList, const FRenderingCompositePassContext& Context)
 	{
 		const FVertexShaderRHIParamRef ShaderRHI = GetVertexShader();
 
-		FGlobalShader::SetParameters(ShaderRHI, Context.View);
+		FGlobalShader::SetParameters(RHICmdList, ShaderRHI, Context.View);
 
 		const FPooledRenderTargetDesc* InputDesc = Context.Pass->GetInputDesc(ePId_Input0);
 
@@ -202,7 +202,7 @@ public:
 
 		FVector2D InvExtent = FVector2D(1.0f / InputDesc->Extent.X, 1.0f / InputDesc->Extent.Y);
 
-		SetShaderValue(ShaderRHI, fxaaQualityRcpFrame, InvExtent);
+		SetShaderValue(RHICmdList, ShaderRHI, fxaaQualityRcpFrame, InvExtent);
 	}
 
 	FShaderParameter fxaaQualityRcpFrame;
@@ -213,7 +213,7 @@ IMPLEMENT_SHADER_TYPE(,FFXAAVS,TEXT("FXAAShader"),TEXT("FxaaVS"),SF_Vertex);
 
 // @param Quality 1..6
 template <uint32 Quality>
-static void SetShaderTemplAA(const FRenderingCompositePassContext& Context)
+static void SetShaderTemplAA(FRHICommandList* RHICmdList, const FRenderingCompositePassContext& Context)
 {
 	TShaderMapRef<FFXAAVS> VertexShader(GetGlobalShaderMap());
 	TShaderMapRef<FPostProcessAntiAliasingPS<Quality> > PixelShader(GetGlobalShaderMap());
@@ -222,8 +222,8 @@ static void SetShaderTemplAA(const FRenderingCompositePassContext& Context)
 
 	SetGlobalBoundShaderState(BoundShaderState, GFilterVertexDeclaration.VertexDeclarationRHI, *VertexShader, *PixelShader);
 
-	PixelShader->SetParameters(Context);
-	VertexShader->SetParameters(Context);
+	PixelShader->SetParameters(RHICmdList, Context);
+	VertexShader->SetParameters(RHICmdList, Context);
 }
 
 void FRCPassPostProcessAA::Process(FRenderingCompositePassContext& Context)
@@ -248,6 +248,10 @@ void FRCPassPostProcessAA::Process(FRenderingCompositePassContext& Context)
 
 	const FSceneRenderTargetItem& DestRenderTarget = PassOutputs[0].RequestSurface(Context);
 
+	//@todo-rco: RHIPacketList
+	FRHICommandList* RHICmdList = nullptr;
+	check(!RHICmdList);
+
 	// Set the view family's render target/viewport.
 	RHISetRenderTarget(DestRenderTarget.TargetableTexture, FTextureRHIRef());	
 	Context.SetViewportAndCallRHI(0, 0, 0.0f, DestSize.X, DestSize.Y, 1.0f );
@@ -259,12 +263,12 @@ void FRCPassPostProcessAA::Process(FRenderingCompositePassContext& Context)
 
 	switch(Quality)
 	{
-		case 1:  SetShaderTemplAA<1>(Context); break;
-		case 2:  SetShaderTemplAA<2>(Context); break;
-		case 3:  SetShaderTemplAA<3>(Context); break;
-		case 4:  SetShaderTemplAA<4>(Context); break;
-		case 5:  SetShaderTemplAA<5>(Context); break;
-		default: SetShaderTemplAA<6>(Context); break;
+		case 1:  SetShaderTemplAA<1>(RHICmdList, Context); break;
+		case 2:  SetShaderTemplAA<2>(RHICmdList, Context); break;
+		case 3:  SetShaderTemplAA<3>(RHICmdList, Context); break;
+		case 4:  SetShaderTemplAA<4>(RHICmdList, Context); break;
+		case 5:  SetShaderTemplAA<5>(RHICmdList, Context); break;
+		default: SetShaderTemplAA<6>(RHICmdList, Context); break;
 	}
 	
 	// Draw a quad mapping scene color to the view's render target
diff --git a/Engine/Source/Runtime/Renderer/Private/PostProcess/PostProcessBloomSetup.cpp b/Engine/Source/Runtime/Renderer/Private/PostProcess/PostProcessBloomSetup.cpp
--- a/Engine/Source/Runtime/Renderer/Private/PostProcess/PostProcessBloomSetup.cpp
+++ b/Engine/Source/Runtime/Renderer/Private/PostProcess/PostProcessBloomSetup.cpp
@@ -48,20 +48,20 @@ public:
 		BloomThreshold.Bind(Initializer.ParameterMap, TEXT("BloomThreshold"));
 	}
 
-	void SetPS(const FRenderingCompositePassContext& Context)
+	void SetPS(FRHICommandList* RHICmdList, const FRenderingCompositePassContext& Context)
 	{
 		const FPixelShaderRHIParamRef ShaderRHI = GetPixelShader();
 
 		const FPostProcessSettings& Settings = Context.View.FinalPostProcessSettings;
 
-		FGlobalShader::SetParameters(ShaderRHI, Context.View);
+		FGlobalShader::SetParameters(RHICmdList, ShaderRHI, Context.View);
 
-		PostprocessParameter.SetPS(ShaderRHI, Context, TStaticSamplerState<SF_Bilinear,AM_Clamp,AM_Clamp,AM_Clamp>::GetRHI());
+		PostprocessParameter.SetPS(RHICmdList, ShaderRHI, Context, TStaticSamplerState<SF_Bilinear,AM_Clamp,AM_Clamp,AM_Clamp>::GetRHI());
 
 		float ExposureScale = FRCPassPostProcessEyeAdaptation::ComputeExposureScaleValue(Context.View);
 
 		FVector4 BloomThresholdValue(Settings.BloomThreshold, 0, 0, ExposureScale);
-		SetShaderValue(ShaderRHI, BloomThreshold, BloomThresholdValue);
+		SetShaderValue(RHICmdList, ShaderRHI, BloomThreshold, BloomThresholdValue);
 	}
 	
 	// FShader interface.
@@ -102,13 +102,13 @@ public:
 		EyeAdaptation.Bind(Initializer.ParameterMap, TEXT("EyeAdaptation"));
 	}
 	
-	void SetVS(const FRenderingCompositePassContext& Context)
+	void SetVS(FRHICommandList* RHICmdList, const FRenderingCompositePassContext& Context)
 	{
 		const FVertexShaderRHIParamRef ShaderRHI = GetVertexShader();
 
-		FGlobalShader::SetParameters(ShaderRHI, Context.View);
+		FGlobalShader::SetParameters(RHICmdList, ShaderRHI, Context.View);
 
-		PostprocessParameter.SetVS(ShaderRHI, Context, TStaticSamplerState<SF_Bilinear,AM_Clamp,AM_Clamp,AM_Clamp>::GetRHI());
+		PostprocessParameter.SetVS(RHICmdList, ShaderRHI, Context, TStaticSamplerState<SF_Bilinear,AM_Clamp,AM_Clamp,AM_Clamp>::GetRHI());
 
 		if(EyeAdaptation.IsBound())
 		{
@@ -116,11 +116,11 @@ public:
 
 			if(EyeAdaptationRT)
 			{
-				SetTextureParameter(ShaderRHI, EyeAdaptation, EyeAdaptationRT->GetRenderTargetItem().TargetableTexture);
+				SetTextureParameter(RHICmdList, ShaderRHI, EyeAdaptation, EyeAdaptationRT->GetRenderTargetItem().TargetableTexture);
 			}
 			else
 			{
-				SetTextureParameter(ShaderRHI, EyeAdaptation, GWhiteTexture->TextureRHI);
+				SetTextureParameter(RHICmdList, ShaderRHI, EyeAdaptation, GWhiteTexture->TextureRHI);
 			}
 		}
 	}
@@ -162,6 +162,9 @@ void FRCPassPostProcessBloomSetup::Process(FRenderingCompositePassContext& Conte
 
 	const FSceneRenderTargetItem& DestRenderTarget = PassOutputs[0].RequestSurface(Context);
 
+	//@todo-rco: RHIPacketList
+	FRHICommandList* RHICmdList = nullptr;
+
 	// Set the view family's render target/viewport.
 	RHISetRenderTarget(DestRenderTarget.TargetableTexture, FTextureRHIRef());	
 
@@ -182,8 +185,8 @@ void FRCPassPostProcessBloomSetup::Process(FRenderingCompositePassContext& Conte
 
 	SetGlobalBoundShaderState(BoundShaderState, GFilterVertexDeclaration.VertexDeclarationRHI, *VertexShader, *PixelShader);
 
-	VertexShader->SetVS(Context);
-	PixelShader->SetPS(Context);
+	VertexShader->SetVS(RHICmdList, Context);
+	PixelShader->SetPS(RHICmdList, Context);
 
 	// Draw a quad mapping scene color to the view's render target
 	DrawRectangle(
diff --git a/Engine/Source/Runtime/Renderer/Private/PostProcess/PostProcessBlurGBuffer.h b/Engine/Source/Runtime/Renderer/Private/PostProcess/PostProcessBlurGBuffer.h
new file mode 100644
--- /dev/null
+++ b/Engine/Source/Runtime/Renderer/Private/PostProcess/PostProcessBlurGBuffer.h
@@ -0,0 +1,22 @@
+// Copyright 1998-2014 Epic Games, Inc. All Rights Reserved.
+
+/*=============================================================================
+	PostProcessBlurGBuffer.h: Post processing tone mapping implementation, can add bloom.
+=============================================================================*/
+
+#pragma once
+
+#include "RenderingCompositionGraph.h"
+
+// no input
+// derives from TRenderingCompositePassBase<InputCount, OutputCount>
+class FRCPassPostProcessBlurGBuffer: public TRenderingCompositePassBase<0, 1>
+{
+public:	
+	// interface FRenderingCompositePass ---------
+	virtual void Process(FRenderingCompositePassContext& Context);
+	virtual void Release() OVERRIDE { delete this; }
+	virtual FPooledRenderTargetDesc ComputeOutputDesc(EPassOutputId InPassOutputId) const;
+
+	static void SetShader(FRHICommandList* RHICmdList, const FRenderingCompositePassContext& Context, uint32 BlurType);
+};
diff --git a/Engine/Source/Runtime/Renderer/Private/PostProcess/PostProcessBokehDOF.cpp b/Engine/Source/Runtime/Renderer/Private/PostProcess/PostProcessBokehDOF.cpp
--- a/Engine/Source/Runtime/Renderer/Private/PostProcess/PostProcessBokehDOF.cpp
+++ b/Engine/Source/Runtime/Renderer/Private/PostProcess/PostProcessBokehDOF.cpp
@@ -78,7 +78,8 @@ class PostProcessVisualizeDOFPS : public FGlobalShader
 
 	static void ModifyCompilationEnvironment(EShaderPlatform Platform, FShaderCompilerEnvironment& OutEnvironment)
 	{
-		FGlobalShader::ModifyCompilationEnvironment(Platform, OutEnvironment);
+		FGlobalShader::ModifyCompilationEnvironment(Platform,OutEnvironment);
+		FDeferredPixelShaderParameters::ModifyCompilationEnvironment(Platform,OutEnvironment);
 	}
 
 	/** Default constructor. */
@@ -106,15 +107,15 @@ public:
 		return bShaderHasOutdatedParameters;
 	}
 
-	void SetParameters(const FRenderingCompositePassContext& Context)
+	void SetParameters(FRHICommandList* RHICmdList, const FRenderingCompositePassContext& Context)
 	{
 		const FPixelShaderRHIParamRef ShaderRHI = GetPixelShader();
 
-		FGlobalShader::SetParameters(ShaderRHI, Context.View);
+		FGlobalShader::SetParameters(RHICmdList, ShaderRHI, Context.View);
 
-		DeferredParameters.Set(ShaderRHI, Context.View);
+		DeferredParameters.Set(RHICmdList, ShaderRHI, Context.View);
 
-		PostprocessParameter.SetPS(ShaderRHI, Context, TStaticSamplerState<SF_Point,AM_Clamp,AM_Clamp,AM_Clamp>::GetRHI());
+		PostprocessParameter.SetPS(RHICmdList, ShaderRHI, Context, TStaticSamplerState<SF_Point,AM_Clamp,AM_Clamp,AM_Clamp>::GetRHI());
 
 		{
 			FVector4 DepthOfFieldParamValues[2];
@@ -124,7 +125,7 @@ public:
 
 			FRCPassPostProcessBokehDOF::ComputeDepthOfFieldParams(Context, DepthOfFieldParamValues);
 
-			SetShaderValueArray(ShaderRHI, DepthOfFieldParams, DepthOfFieldParamValues, 2);
+			SetShaderValueArray(RHICmdList, ShaderRHI, DepthOfFieldParams, DepthOfFieldParamValues, 2);
 		}
 	}
 
@@ -167,6 +168,9 @@ void FRCPassPostProcessVisualizeDOF::Process(FRenderingCompositePassContext& Con
 
 	const FSceneRenderTargetItem& DestRenderTarget = PassOutputs[0].RequestSurface(Context);
 
+	//@todo-rco: RHIPacketList
+	FRHICommandList* RHICmdList = nullptr;
+
 	// Set the view family's render target/viewport.
 	RHISetRenderTarget(DestRenderTarget.TargetableTexture, FTextureRHIRef());	
 
@@ -192,8 +196,8 @@ void FRCPassPostProcessVisualizeDOF::Process(FRenderingCompositePassContext& Con
 
 		SetGlobalBoundShaderState(BoundShaderState, GFilterVertexDeclaration.VertexDeclarationRHI, *VertexShader, *PixelShader);
 
-		VertexShader->SetParameters(Context);
-		PixelShader->SetParameters(Context);
+		VertexShader->SetParameters(RHICmdList, Context);
+		PixelShader->SetParameters(RHICmdList, Context);
 	}
 
 	// Draw a quad mapping scene color to the view's render target
@@ -237,7 +241,8 @@ class PostProcessBokehDOFSetupPS : public FGlobalShader
 
 	static void ModifyCompilationEnvironment(EShaderPlatform Platform, FShaderCompilerEnvironment& OutEnvironment)
 	{
-		FGlobalShader::ModifyCompilationEnvironment(Platform, OutEnvironment);
+		FGlobalShader::ModifyCompilationEnvironment(Platform,OutEnvironment);
+		FDeferredPixelShaderParameters::ModifyCompilationEnvironment(Platform,OutEnvironment);
 	}
 
 	/** Default constructor. */
@@ -265,14 +270,14 @@ public:
 		return bShaderHasOutdatedParameters;
 	}
 
-	void SetParameters(const FRenderingCompositePassContext& Context)
+	void SetParameters(FRHICommandList* RHICmdList, const FRenderingCompositePassContext& Context)
 	{
 		const FPixelShaderRHIParamRef ShaderRHI = GetPixelShader();
 
-		FGlobalShader::SetParameters(ShaderRHI, Context.View);
+		FGlobalShader::SetParameters(RHICmdList, ShaderRHI, Context.View);
 		
-		DeferredParameters.Set(ShaderRHI, Context.View);
-		PostprocessParameter.SetPS(ShaderRHI, Context, TStaticSamplerState<SF_Point,AM_Clamp,AM_Clamp,AM_Clamp>::GetRHI());
+		DeferredParameters.Set(RHICmdList, ShaderRHI, Context.View);
+		PostprocessParameter.SetPS(RHICmdList, ShaderRHI, Context, TStaticSamplerState<SF_Point,AM_Clamp,AM_Clamp,AM_Clamp>::GetRHI());
 
 		{
 			FVector4 DepthOfFieldParamValues[2];
@@ -282,7 +287,7 @@ public:
 
 			FRCPassPostProcessBokehDOF::ComputeDepthOfFieldParams(Context, DepthOfFieldParamValues);
 
-			SetShaderValueArray(ShaderRHI, DepthOfFieldParams, DepthOfFieldParamValues, 2);
+			SetShaderValueArray(RHICmdList, ShaderRHI, DepthOfFieldParams, DepthOfFieldParamValues, 2);
 		}
 	}
 
@@ -325,6 +330,9 @@ void FRCPassPostProcessBokehDOFSetup::Process(FRenderingCompositePassContext& Co
 
 	const FSceneRenderTargetItem& DestRenderTarget = PassOutputs[0].RequestSurface(Context);
 
+	//@todo-rco: RHIPacketList
+	FRHICommandList* RHICmdList = nullptr;
+
 	// Set the view family's render target/viewport.
 	RHISetRenderTarget(DestRenderTarget.TargetableTexture, FTextureRHIRef());	
 	
@@ -349,8 +357,8 @@ void FRCPassPostProcessBokehDOFSetup::Process(FRenderingCompositePassContext& Co
 
 		SetGlobalBoundShaderState(BoundShaderState, GFilterVertexDeclaration.VertexDeclarationRHI, *VertexShader, *PixelShader);
 
-		VertexShader->SetParameters(Context);
-		PixelShader->SetParameters(Context);
+		VertexShader->SetParameters(RHICmdList, Context);
+		PixelShader->SetParameters(RHICmdList, Context);
 	}
 
 	// Draw a quad mapping scene color to the view's render target
@@ -394,7 +402,8 @@ class FPostProcessBokehDOFVS : public FGlobalShader
 
 	static void ModifyCompilationEnvironment(EShaderPlatform Platform, FShaderCompilerEnvironment& OutEnvironment)
 	{
-		FGlobalShader::ModifyCompilationEnvironment(Platform, OutEnvironment);
+		FGlobalShader::ModifyCompilationEnvironment(Platform,OutEnvironment);
+		FDeferredPixelShaderParameters::ModifyCompilationEnvironment(Platform,OutEnvironment);
 		OutEnvironment.SetDefine(TEXT("DOF_METHOD"), DOFMethod);
 		OutEnvironment.SetDefine(TEXT("DOF_INDEX_STYLE"), DOFIndexStyle);
 	}
@@ -431,25 +440,25 @@ public:
 	}
 
 	/** to have a similar interface as all other shaders */
-	void SetParameters(const FRenderingCompositePassContext& Context, FIntPoint TileCountValue, uint32 TileSize, float PixelKernelSize, FIntPoint LeftTop)
+	void SetParameters(FRHICommandList* RHICmdList, const FRenderingCompositePassContext& Context, FIntPoint TileCountValue, uint32 TileSize, float PixelKernelSize, FIntPoint LeftTop)
 	{
 		const FVertexShaderRHIParamRef ShaderRHI = GetVertexShader();
 
-		FGlobalShader::SetParameters(ShaderRHI, Context.View);
-		DeferredParameters.Set(ShaderRHI, Context.View);
-		PostprocessParameter.SetVS(ShaderRHI, Context, TStaticSamplerState<SF_Bilinear,AM_Clamp,AM_Clamp,AM_Clamp>::GetRHI());
+		FGlobalShader::SetParameters(RHICmdList, ShaderRHI, Context.View);
+		DeferredParameters.Set(RHICmdList, ShaderRHI, Context.View);
+		PostprocessParameter.SetVS(RHICmdList, ShaderRHI, Context, TStaticSamplerState<SF_Bilinear,AM_Clamp,AM_Clamp,AM_Clamp>::GetRHI());
 //		PostprocessParameter.SetVS(ShaderRHI, Context, TStaticSamplerState<SF_Point,AM_Clamp,AM_Clamp,AM_Clamp>::GetRHI());
 
 		{
 			FIntRect TileCountAndSizeValue(TileCountValue, FIntPoint(TileSize, TileSize));
 
-			SetShaderValue(ShaderRHI, TileCountAndSize, TileCountAndSizeValue);
+			SetShaderValue(RHICmdList, ShaderRHI, TileCountAndSize, TileCountAndSizeValue);
 		}
 
 		{
 			FVector4 KernelSizeValue(PixelKernelSize, PixelKernelSize, LeftTop.X, LeftTop.Y);
 
-			SetShaderValue(ShaderRHI, KernelSize, KernelSizeValue);
+			SetShaderValue(RHICmdList, ShaderRHI, KernelSize, KernelSizeValue);
 		}
 
 		{
@@ -457,7 +466,7 @@ public:
 				Context.View.FinalPostProcessSettings.DepthOfFieldColorThreshold,
 				Context.View.FinalPostProcessSettings.DepthOfFieldSizeThreshold,0);
 
-			SetShaderValue(ShaderRHI, DepthOfFieldThresholds, Value);
+			SetShaderValue(RHICmdList, ShaderRHI, DepthOfFieldThresholds, Value);
 		}
 
 		{
@@ -465,7 +474,7 @@ public:
 
 			FRCPassPostProcessBokehDOF::ComputeDepthOfFieldParams(Context, DepthOfFieldParamValues);
 
-			SetShaderValueArray(ShaderRHI, DepthOfFieldParams, DepthOfFieldParamValues, 2);
+			SetShaderValueArray(RHICmdList, ShaderRHI, DepthOfFieldParams, DepthOfFieldParamValues, 2);
 		}
 	}
 
@@ -515,13 +524,13 @@ public:
 		return bShaderHasOutdatedParameters;
 	}
 
-	void SetParameters(const FRenderingCompositePassContext& Context, float PixelKernelSize)
+	void SetParameters(FRHICommandList* RHICmdList, const FRenderingCompositePassContext& Context, float PixelKernelSize)
 	{
 		const FPixelShaderRHIParamRef ShaderRHI = GetPixelShader();
 
-		FGlobalShader::SetParameters(ShaderRHI, Context.View);
+		FGlobalShader::SetParameters(RHICmdList, ShaderRHI, Context.View);
 
-		PostprocessParameter.SetPS(ShaderRHI, Context, TStaticSamplerState<SF_Bilinear,AM_Clamp,AM_Clamp,AM_Clamp>::GetRHI());
+		PostprocessParameter.SetPS(RHICmdList, ShaderRHI, Context, TStaticSamplerState<SF_Bilinear,AM_Clamp,AM_Clamp,AM_Clamp>::GetRHI());
 //		PostprocessParameter.SetPS(ShaderRHI, Context, TStaticSamplerState<SF_Point,AM_Clamp,AM_Clamp,AM_Clamp>::GetRHI());
 				
 		{
@@ -537,7 +546,7 @@ public:
 				TextureRHI = Context.View.FinalPostProcessSettings.DepthOfFieldBokehShape->Resource->TextureRHI;
 			}
 
-			SetTextureParameter(ShaderRHI, LensTexture, LensTextureSampler, TStaticSamplerState<SF_Trilinear,AM_Border,AM_Border,AM_Clamp>::GetRHI(), TextureRHI);
+			SetTextureParameter(RHICmdList, ShaderRHI, LensTexture, LensTextureSampler, TStaticSamplerState<SF_Trilinear,AM_Border,AM_Border,AM_Clamp>::GetRHI(), TextureRHI);
 		}
 	}
 };
@@ -553,7 +562,7 @@ VARIATION1(0,1)			VARIATION1(1,1)			VARIATION1(2,1)
 #undef VARIATION1
 
 template <uint32 DOFMethod, uint32 DOFIndexStyle>
-void FRCPassPostProcessBokehDOF::SetShaderTempl(const FRenderingCompositePassContext& Context, FIntPoint LeftTop, FIntPoint TileCount, uint32 TileSize, float PixelKernelSize)
+void FRCPassPostProcessBokehDOF::SetShaderTempl(FRHICommandList* RHICmdList, const FRenderingCompositePassContext& Context, FIntPoint LeftTop, FIntPoint TileCount, uint32 TileSize, float PixelKernelSize)
 {
 	TShaderMapRef<FPostProcessBokehDOFVS<DOFMethod, DOFIndexStyle> > VertexShader(GetGlobalShaderMap());
 	TShaderMapRef<FPostProcessBokehDOFPS> PixelShader(GetGlobalShaderMap());
@@ -562,8 +571,8 @@ void FRCPassPostProcessBokehDOF::SetShaderTempl(const FRenderingCompositePassCon
 
 	SetGlobalBoundShaderState(BoundShaderState, GFilterVertexDeclaration.VertexDeclarationRHI, *VertexShader, *PixelShader);
 
-	VertexShader->SetParameters(Context, TileCount, TileSize, PixelKernelSize, LeftTop);
-	PixelShader->SetParameters(Context, PixelKernelSize);
+	VertexShader->SetParameters(RHICmdList, Context, TileCount, TileSize, PixelKernelSize, LeftTop);
+	PixelShader->SetParameters(RHICmdList, Context, PixelKernelSize);
 }
 
 void FRCPassPostProcessBokehDOF::ComputeDepthOfFieldParams(const FRenderingCompositePassContext& Context, FVector4 Out[2])
@@ -622,6 +631,9 @@ void FRCPassPostProcessBokehDOF::Process(FRenderingCompositePassContext& Context
 
 	const FSceneRenderTargetItem& DestRenderTarget = PassOutputs[0].RequestSurface(Context);
 
+	//@todo-rco: RHIPacketList
+	FRHICommandList* RHICmdList = nullptr;
+
 	// Set the view family's render target/viewport.
 	RHISetRenderTarget(DestRenderTarget.TargetableTexture, FTextureRHIRef());	
 
@@ -667,11 +679,11 @@ void FRCPassPostProcessBokehDOF::Process(FRenderingCompositePassContext& Context
 			// high quality, visualize in red and green where we spend more performance
 			if(IndexStyle == BIS_Fast)
 			{
-				SetShaderTempl<2, 0>(Context, LeftTop, TileCount, TileSize, PixelKernelSize);
+				SetShaderTempl<2, 0>(RHICmdList, Context, LeftTop, TileCount, TileSize, PixelKernelSize);
 			}
 			else
 			{
-				SetShaderTempl<2, 1>(Context, LeftTop, TileCount, TileSize, PixelKernelSize);
+				SetShaderTempl<2, 1>(RHICmdList, Context, LeftTop, TileCount, TileSize, PixelKernelSize);
 			}
 		}
 		else
@@ -679,11 +691,11 @@ void FRCPassPostProcessBokehDOF::Process(FRenderingCompositePassContext& Context
 			// high quality
 			if(IndexStyle == BIS_Fast)
 			{
-				SetShaderTempl<1, 0>(Context, LeftTop, TileCount, TileSize, PixelKernelSize);
+				SetShaderTempl<1, 0>(RHICmdList, Context, LeftTop, TileCount, TileSize, PixelKernelSize);
 			}
 			else
 			{
-				SetShaderTempl<1, 1>(Context, LeftTop, TileCount, TileSize, PixelKernelSize);
+				SetShaderTempl<1, 1>(RHICmdList, Context, LeftTop, TileCount, TileSize, PixelKernelSize);
 			}
 		}
 	}
@@ -692,11 +704,11 @@ void FRCPassPostProcessBokehDOF::Process(FRenderingCompositePassContext& Context
 		// low quality
 		if(IndexStyle == BIS_Fast)
 		{
-			SetShaderTempl<0, 0>(Context, LeftTop, TileCount, TileSize, PixelKernelSize);
+			SetShaderTempl<0, 0>(RHICmdList, Context, LeftTop, TileCount, TileSize, PixelKernelSize);
 		}
 		else
 		{
-			SetShaderTempl<0, 1>(Context, LeftTop, TileCount, TileSize, PixelKernelSize);
+			SetShaderTempl<0, 1>(RHICmdList, Context, LeftTop, TileCount, TileSize, PixelKernelSize);
 		}
 	}
 
diff --git a/Engine/Source/Runtime/Renderer/Private/PostProcess/PostProcessBokehDOF.h b/Engine/Source/Runtime/Renderer/Private/PostProcess/PostProcessBokehDOF.h
--- a/Engine/Source/Runtime/Renderer/Private/PostProcess/PostProcessBokehDOF.h
+++ b/Engine/Source/Runtime/Renderer/Private/PostProcess/PostProcessBokehDOF.h
@@ -49,7 +49,7 @@ public:
 private:
 
 	template <uint32 HighQuality, uint32 IndexStyle>
-	void SetShaderTempl(const FRenderingCompositePassContext& Context, FIntPoint LeftTop, FIntPoint TileCount, uint32 TileSize, float PixelKernelSize);
+	void SetShaderTempl(FRHICommandList* RHICmdList, const FRenderingCompositePassContext& Context, FIntPoint LeftTop, FIntPoint TileCount, uint32 TileSize, float PixelKernelSize);
 
 	// border between front and back layer as we don't use viewports (only possible with GS)
 	const static uint32 SafetyBorder = 40;
diff --git a/Engine/Source/Runtime/Renderer/Private/PostProcess/PostProcessBokehDOFRecombine.cpp b/Engine/Source/Runtime/Renderer/Private/PostProcess/PostProcessBokehDOFRecombine.cpp
--- a/Engine/Source/Runtime/Renderer/Private/PostProcess/PostProcessBokehDOFRecombine.cpp
+++ b/Engine/Source/Runtime/Renderer/Private/PostProcess/PostProcessBokehDOFRecombine.cpp
@@ -30,6 +30,7 @@ class FPostProcessBokehDOFRecombinePS : public FGlobalShader
 	static void ModifyCompilationEnvironment(EShaderPlatform Platform, FShaderCompilerEnvironment& OutEnvironment)
 	{
 		FGlobalShader::ModifyCompilationEnvironment(Platform, OutEnvironment);
+		FDeferredPixelShaderParameters::ModifyCompilationEnvironment(Platform,OutEnvironment);
 		OutEnvironment.SetDefine(TEXT("RECOMBINE_METHOD"), Method);
 	}
 
@@ -58,20 +59,20 @@ public:
 		return bShaderHasOutdatedParameters;
 	}
 
-	void SetParameters(const FRenderingCompositePassContext& Context)
+	void SetParameters(FRHICommandList* RHICmdList, const FRenderingCompositePassContext& Context)
 	{
 		const FPixelShaderRHIParamRef ShaderRHI = GetPixelShader();
 
-		FGlobalShader::SetParameters(ShaderRHI, Context.View);
-		DeferredParameters.Set(ShaderRHI, Context.View);
-		PostprocessParameter.SetPS(ShaderRHI, Context, TStaticSamplerState<SF_Bilinear,AM_Clamp,AM_Clamp,AM_Clamp>::GetRHI());
+		FGlobalShader::SetParameters(RHICmdList, ShaderRHI, Context.View);
+		DeferredParameters.Set(RHICmdList, ShaderRHI, Context.View);
+		PostprocessParameter.SetPS(RHICmdList, ShaderRHI, Context, TStaticSamplerState<SF_Bilinear,AM_Clamp,AM_Clamp,AM_Clamp>::GetRHI());
 
 		{
 			FVector4 DepthOfFieldParamValues[2];
 
 			FRCPassPostProcessBokehDOF::ComputeDepthOfFieldParams(Context, DepthOfFieldParamValues);
 
-			SetShaderValueArray(ShaderRHI, DepthOfFieldParams, DepthOfFieldParamValues, 2);
+			SetShaderValueArray(RHICmdList, ShaderRHI, DepthOfFieldParams, DepthOfFieldParamValues, 2);
 		}
 	}
 
@@ -95,7 +96,7 @@ VARIATION1(1)			VARIATION1(2)			VARIATION1(3)
 
 
 template <uint32 Method>
-void FRCPassPostProcessBokehDOFRecombine::SetShader(const FRenderingCompositePassContext& Context)
+void FRCPassPostProcessBokehDOFRecombine::SetShader(FRHICommandList* RHICmdList, const FRenderingCompositePassContext& Context)
 {
 	TShaderMapRef<FPostProcessVS> VertexShader(GetGlobalShaderMap());
 	TShaderMapRef<FPostProcessBokehDOFRecombinePS<Method> > PixelShader(GetGlobalShaderMap());
@@ -104,8 +105,8 @@ void FRCPassPostProcessBokehDOFRecombine::SetShader(const FRenderingCompositePas
 
 	SetGlobalBoundShaderState(BoundShaderState, GFilterVertexDeclaration.VertexDeclarationRHI, *VertexShader, *PixelShader);
 
-	PixelShader->SetParameters(Context);
-	VertexShader->SetParameters(Context);
+	PixelShader->SetParameters(RHICmdList, Context);
+	VertexShader->SetParameters(RHICmdList, Context);
 }
 
 void FRCPassPostProcessBokehDOFRecombine::Process(FRenderingCompositePassContext& Context)
@@ -143,6 +144,8 @@ void FRCPassPostProcessBokehDOFRecombine::Process(FRenderingCompositePassContext
 	FIntRect HalfResViewRect = FIntRect::DivideAndRoundUp(View.ViewRect, ScaleToFullRes);
 
 	const FSceneRenderTargetItem& DestRenderTarget = PassOutputs[0].RequestSurface(Context);
+	//@todo-rco: RHIPacketList
+	FRHICommandList* RHICmdList = nullptr;
 
 	// Set the view family's render target/viewport.
 	RHISetRenderTarget(DestRenderTarget.TargetableTexture, FTextureRHIRef());	
@@ -158,10 +161,10 @@ void FRCPassPostProcessBokehDOFRecombine::Process(FRenderingCompositePassContext
 	RHISetDepthStencilState(TStaticDepthStencilState<false,CF_Always>::GetRHI());
 
 	switch(Method)
-	{
-		case 1: SetShader<1>(Context); break;
-		case 2: SetShader<2>(Context); break;
-		case 3: SetShader<3>(Context); break;
+		{
+		case 1: SetShader<1>(RHICmdList, Context); break;
+		case 2: SetShader<2>(RHICmdList, Context); break;
+		case 3: SetShader<3>(RHICmdList, Context); break;
 		default:
 			check(0);
 	}
diff --git a/Engine/Source/Runtime/Renderer/Private/PostProcess/PostProcessBokehDOFRecombine.h b/Engine/Source/Runtime/Renderer/Private/PostProcess/PostProcessBokehDOFRecombine.h
--- a/Engine/Source/Runtime/Renderer/Private/PostProcess/PostProcessBokehDOFRecombine.h
+++ b/Engine/Source/Runtime/Renderer/Private/PostProcess/PostProcessBokehDOFRecombine.h
@@ -23,5 +23,5 @@ public:
 private:
 
 	template <uint32 Method>
-	static void SetShader(const FRenderingCompositePassContext& Context);
+	static void SetShader(FRHICommandList* RHICmdList, const FRenderingCompositePassContext& Context);
 };
\ No newline at end of file
diff --git a/Engine/Source/Runtime/Renderer/Private/PostProcess/PostProcessBusyWait.cpp b/Engine/Source/Runtime/Renderer/Private/PostProcess/PostProcessBusyWait.cpp
--- a/Engine/Source/Runtime/Renderer/Private/PostProcess/PostProcessBusyWait.cpp
+++ b/Engine/Source/Runtime/Renderer/Private/PostProcess/PostProcessBusyWait.cpp
@@ -35,13 +35,13 @@ public:
 		GPUBusyWait.Bind(Initializer.ParameterMap,TEXT("GPUBusyWait"));
 	}
 
-	void SetPS(const FRenderingCompositePassContext& Context)
+	void SetPS(FRHICommandList* RHICmdList, const FRenderingCompositePassContext& Context)
 	{
 		const FPixelShaderRHIParamRef ShaderRHI = GetPixelShader();
 
-		FGlobalShader::SetParameters(ShaderRHI, Context.View);
+		FGlobalShader::SetParameters(RHICmdList, ShaderRHI, Context.View);
 
-		PostprocessParameter.SetPS(ShaderRHI, Context, TStaticSamplerState<SF_Bilinear,AM_Clamp,AM_Clamp,AM_Clamp>::GetRHI());
+		PostprocessParameter.SetPS(RHICmdList, ShaderRHI, Context, TStaticSamplerState<SF_Bilinear,AM_Clamp,AM_Clamp,AM_Clamp>::GetRHI());
 
 #if !(UE_BUILD_SHIPPING || UE_BUILD_TEST)
 		{
@@ -56,7 +56,7 @@ public:
 			// divide by viewport pixel count
 			uint32 Value = (uint32)(CVarValue * 1000000000.0 / 6.12 / PixelCount);
 
-			SetShaderValue(ShaderRHI, GPUBusyWait, Value);
+			SetShaderValue(RHICmdList, ShaderRHI, GPUBusyWait, Value);
 		}
 #endif
 	}
@@ -80,11 +80,14 @@ void FRCPassPostProcessBusyWait::Process(FRenderingCompositePassContext& Context
 	
 	FIntRect SrcRect = View.ViewRect;
 	FIntRect DestRect = View.UnscaledViewRect;
-
+	
 	GSceneRenderTargets.BeginRenderingLightAttenuation();
 	
 	const FSceneRenderTargetItem& DestRenderTarget = GSceneRenderTargets.GetLightAttenuation()->GetRenderTargetItem();
 
+	//@todo-rco: RHIPacketList
+	FRHICommandList* RHICmdList = nullptr;
+
 	// Set the view family's render target/viewport.
 	RHISetRenderTarget(DestRenderTarget.TargetableTexture, FTextureRHIRef());	
 	Context.SetViewportAndCallRHI(DestRect);
@@ -101,7 +104,7 @@ void FRCPassPostProcessBusyWait::Process(FRenderingCompositePassContext& Context
 
 	SetGlobalBoundShaderState(BoundShaderState, GFilterVertexDeclaration.VertexDeclarationRHI, *VertexShader, *PixelShader);
 
-	PixelShader->SetPS(Context);
+	PixelShader->SetPS(RHICmdList, Context);
 
 	// Draw a quad mapping scene color to the view's render target
 	DrawRectangle(
diff --git a/Engine/Source/Runtime/Renderer/Private/PostProcess/PostProcessCombineLUTs.cpp b/Engine/Source/Runtime/Renderer/Private/PostProcess/PostProcessCombineLUTs.cpp
--- a/Engine/Source/Runtime/Renderer/Private/PostProcess/PostProcessCombineLUTs.cpp
+++ b/Engine/Source/Runtime/Renderer/Private/PostProcess/PostProcessCombineLUTs.cpp
@@ -51,7 +51,7 @@ FColorRemapShaderParameters::FColorRemapShaderParameters(const FShaderParameterM
 	MappingPolynomial.Bind(ParameterMap, TEXT("MappingPolynomial"));
 }
 
-void FColorRemapShaderParameters::Set(const FPixelShaderRHIParamRef ShaderRHI)
+void FColorRemapShaderParameters::Set(FRHICommandList* RHICmdList, const FPixelShaderRHIParamRef ShaderRHI)
 {
 	FColorTransform ColorTransform;
 
@@ -75,7 +75,7 @@ void FColorRemapShaderParameters::Set(const FPixelShaderRHIParamRef ShaderRHI)
 		float b = 4 * ColorTransform.MidValue - 3 * ColorTransform.MinValue - ColorTransform.MaxValue;
 		float a = ColorTransform.MaxValue - ColorTransform.MinValue - b;
 
-		SetShaderValue(ShaderRHI, MappingPolynomial, FVector(a, b, c));
+		SetShaderValue(RHICmdList, ShaderRHI, MappingPolynomial, FVector(a, b, c));
 	}
 }
 
@@ -119,7 +119,7 @@ public:
 	}
 	FLUTBlenderPS() {}
 
-	void SetParameters(const FSceneView& View, FTexture* Texture[BlendCount], float Weights[BlendCount])
+	void SetParameters(FRHICommandList* RHICmdList, const FSceneView& View, FTexture* Texture[BlendCount], float Weights[BlendCount])
 	{
 		const FPixelShaderRHIParamRef ShaderRHI = GetPixelShader();
 
@@ -128,15 +128,15 @@ public:
 			// we don't need to set the neutral one
 			if(i != 0)
 			{
-				SetTextureParameter(ShaderRHI, TextureParameter[i], TextureParameterSampler[i], Texture[i]);
+				SetTextureParameter(RHICmdList, ShaderRHI, TextureParameter[i], TextureParameterSampler[i], Texture[i]);
 			}
 
-			SetShaderValue(ShaderRHI, WeightsParameter, Weights[i], i);
+			SetShaderValue(RHICmdList, ShaderRHI, WeightsParameter, Weights[i], i);
 		}
 
-		SetShaderValue(ShaderRHI, ColorScale, View.ColorScale);
-		SetShaderValue(ShaderRHI, OverlayColor, View.OverlayColor);
-		ColorRemapShaderParameters.Set(ShaderRHI);
+		SetShaderValue(RHICmdList, ShaderRHI, ColorScale, View.ColorScale);
+		SetShaderValue(RHICmdList, ShaderRHI, OverlayColor, View.OverlayColor);
+		ColorRemapShaderParameters.Set(RHICmdList, ShaderRHI);
 	}
 
 	static void ModifyCompilationEnvironment(EShaderPlatform Platform, FShaderCompilerEnvironment& OutEnvironment)
@@ -180,7 +180,7 @@ IMPLEMENT_SHADER_TYPE(template<>,FLUTBlenderPS<3>,TEXT("PostProcessCombineLUTs")
 IMPLEMENT_SHADER_TYPE(template<>,FLUTBlenderPS<4>,TEXT("PostProcessCombineLUTs"),TEXT("MainPS"),SF_Pixel);
 IMPLEMENT_SHADER_TYPE(template<>,FLUTBlenderPS<5>,TEXT("PostProcessCombineLUTs"),TEXT("MainPS"),SF_Pixel);
 
-void SetLUTBlenderShader(FRenderingCompositePassContext& Context, uint32 BlendCount, FTexture* Texture[], float Weights[], const FVolumeBounds& VolumeBounds)
+void SetLUTBlenderShader(FRHICommandList* RHICmdList, FRenderingCompositePassContext& Context, uint32 BlendCount, FTexture* Texture[], float Weights[], const FVolumeBounds& VolumeBounds)
 {
 	check(BlendCount > 0);
 
@@ -219,8 +219,8 @@ void SetLUTBlenderShader(FRenderingCompositePassContext& Context, uint32 BlendCo
 
 		SetGlobalBoundShaderState(*LocalBoundShaderState, GScreenVertexDeclaration.VertexDeclarationRHI, *VertexShader, LocalPixelShader, *GeometryShader);
 
-		VertexShader->SetParameters(VolumeBounds, VolumeBounds.MaxX - VolumeBounds.MinX);
-		GeometryShader->SetParameters(VolumeBounds);
+		VertexShader->SetParameters(RHICmdList, VolumeBounds, VolumeBounds.MaxX - VolumeBounds.MinX);
+		GeometryShader->SetParameters(RHICmdList, VolumeBounds);
 	}
 	else
 	{
@@ -228,13 +228,13 @@ void SetLUTBlenderShader(FRenderingCompositePassContext& Context, uint32 BlendCo
 
 		SetGlobalBoundShaderState(*LocalBoundShaderState, GFilterVertexDeclaration.VertexDeclarationRHI, *VertexShader, LocalPixelShader);
 
-		VertexShader->SetParameters(Context);
+		VertexShader->SetParameters(RHICmdList, Context);
 	}
 #define CASE_COUNT(BlendCount) \
 	case BlendCount: \
 	{ \
 	TShaderMapRef<FLUTBlenderPS<BlendCount> > PixelShader(GetGlobalShaderMap()); \
-	PixelShader->SetParameters(View, Texture, Weights); \
+	PixelShader->SetParameters(RHICmdList, View, Texture, Weights); \
 	}; \
 	break;
 
@@ -386,6 +386,9 @@ void FRCPassPostProcessCombineLUTs::Process(FRenderingCompositePassContext& Cont
 
 	const FSceneRenderTargetItem& DestRenderTarget = PassOutputs[0].RequestSurface(Context);
 
+	//@todo-rco: RHIPacketList
+	FRHICommandList* RHICmdList = nullptr;
+
 	// Set the view family's render target/viewport.
 	RHISetRenderTarget(DestRenderTarget.TargetableTexture, FTextureRHIRef());	
 	Context.SetViewportAndCallRHI(0, 0, 0.0f, DestSize.X, DestSize.Y, 1.0f );
@@ -397,7 +400,7 @@ void FRCPassPostProcessCombineLUTs::Process(FRenderingCompositePassContext& Cont
 
 	const FVolumeBounds VolumeBounds(16);
 
-	SetLUTBlenderShader(Context, LocalCount, LocalTextures, LocalWeights, VolumeBounds);
+	SetLUTBlenderShader(RHICmdList, Context, LocalCount, LocalTextures, LocalWeights, VolumeBounds);
 
 	if(UseVolumeTextureLUT(GRHIShaderPlatform))
 	{
diff --git a/Engine/Source/Runtime/Renderer/Private/PostProcess/PostProcessCombineLUTs.h b/Engine/Source/Runtime/Renderer/Private/PostProcess/PostProcessCombineLUTs.h
--- a/Engine/Source/Runtime/Renderer/Private/PostProcess/PostProcessCombineLUTs.h
+++ b/Engine/Source/Runtime/Renderer/Private/PostProcess/PostProcessCombineLUTs.h
@@ -38,7 +38,7 @@ public:
 
 	FColorRemapShaderParameters(const FShaderParameterMap& ParameterMap);
 
-	void Set(const FPixelShaderRHIParamRef ShaderRHI);
+	void Set(FRHICommandList* RHICmdList, const FPixelShaderRHIParamRef ShaderRHI);
 
 	friend FArchive& operator<<(FArchive& Ar,FColorRemapShaderParameters& P);
 
diff --git a/Engine/Source/Runtime/Renderer/Private/PostProcess/PostProcessCompositeEditorPrimitives.cpp b/Engine/Source/Runtime/Renderer/Private/PostProcess/PostProcessCompositeEditorPrimitives.cpp
--- a/Engine/Source/Runtime/Renderer/Private/PostProcess/PostProcessCompositeEditorPrimitives.cpp
+++ b/Engine/Source/Runtime/Renderer/Private/PostProcess/PostProcessCompositeEditorPrimitives.cpp
@@ -66,25 +66,25 @@ public:
 		FilteredSceneDepthTextureSampler.Bind(Initializer.ParameterMap,TEXT("FilteredSceneDepthTextureSampler"));
 	}
 
-	void SetParameters(const FRenderingCompositePassContext& Context)
+	void SetParameters(FRHICommandList* RHICmdList, const FRenderingCompositePassContext& Context)
 	{
 		const FPixelShaderRHIParamRef ShaderRHI = GetPixelShader();
 
-		FGlobalShader::SetParameters(ShaderRHI, Context.View);
+		FGlobalShader::SetParameters(RHICmdList, ShaderRHI, Context.View);
 
-		DeferredParameters.Set(ShaderRHI, Context.View);
+		DeferredParameters.Set(RHICmdList, ShaderRHI, Context.View);
 
 		FSamplerStateRHIRef SamplerStateRHIRef = TStaticSamplerState<SF_Point,AM_Clamp,AM_Clamp,AM_Clamp>::GetRHI();
-		PostProcessParameters.SetPS(ShaderRHI, Context, SamplerStateRHIRef);
+		PostProcessParameters.SetPS(RHICmdList, ShaderRHI, Context, SamplerStateRHIRef);
 		if(MSAASampleCount > 1)
 		{
-			SetTextureParameter( ShaderRHI, EditorPrimitivesColor, GSceneRenderTargets.EditorPrimitivesColor->GetRenderTargetItem().TargetableTexture );
-			SetTextureParameter( ShaderRHI, EditorPrimitivesDepth, GSceneRenderTargets.EditorPrimitivesDepth->GetRenderTargetItem().TargetableTexture );
+			SetTextureParameter(RHICmdList, ShaderRHI, EditorPrimitivesColor, GSceneRenderTargets.EditorPrimitivesColor->GetRenderTargetItem().TargetableTexture );
+			SetTextureParameter(RHICmdList, ShaderRHI, EditorPrimitivesDepth, GSceneRenderTargets.EditorPrimitivesDepth->GetRenderTargetItem().TargetableTexture );
 		}
 		else
 		{
-			SetTextureParameter( ShaderRHI, EditorPrimitivesColor, EditorPrimitivesColorSampler, SamplerStateRHIRef, GSceneRenderTargets.EditorPrimitivesColor->GetRenderTargetItem().ShaderResourceTexture );
-			SetTextureParameter( ShaderRHI, EditorPrimitivesDepth, GSceneRenderTargets.EditorPrimitivesDepth->GetRenderTargetItem().ShaderResourceTexture );
+			SetTextureParameter(RHICmdList, ShaderRHI, EditorPrimitivesColor, EditorPrimitivesColorSampler, SamplerStateRHIRef, GSceneRenderTargets.EditorPrimitivesColor->GetRenderTargetItem().ShaderResourceTexture );
+			SetTextureParameter(RHICmdList, ShaderRHI, EditorPrimitivesDepth, GSceneRenderTargets.EditorPrimitivesDepth->GetRenderTargetItem().ShaderResourceTexture );
 		}
 
 		{
@@ -104,13 +104,14 @@ public:
 				Value.G = 0;
 			}
 
-			SetShaderValue(ShaderRHI, EditorRenderParams, Value);
+			SetShaderValue(RHICmdList, ShaderRHI, EditorRenderParams, Value);
 		}
 
 		if(FilteredSceneDepthTexture.IsBound())
 		{
 			const FTexture2DRHIRef* DepthTexture = GSceneRenderTargets.GetActualDepthTexture();
 			SetTextureParameter(
+				RHICmdList,
 				ShaderRHI,
 				FilteredSceneDepthTexture,
 				FilteredSceneDepthTextureSampler,
@@ -161,7 +162,7 @@ private:
 
 
 template <uint32 MSAASampleCount>
-static void SetCompositePrimitivesShaderTempl(const FRenderingCompositePassContext& Context)
+static void SetCompositePrimitivesShaderTempl(FRHICommandList* RHICmdList, const FRenderingCompositePassContext& Context)
 {
 	TShaderMapRef<FPostProcessVS> VertexShader(GetGlobalShaderMap());
 	TShaderMapRef<FPostProcessCompositeEditorPrimitivesPS<MSAASampleCount> > PixelShader(GetGlobalShaderMap());
@@ -170,8 +171,8 @@ static void SetCompositePrimitivesShaderTempl(const FRenderingCompositePassConte
 
 	SetGlobalBoundShaderState(BoundShaderState, GFilterVertexDeclaration.VertexDeclarationRHI, *VertexShader, *PixelShader);
 
-	VertexShader->SetParameters(Context);
-	PixelShader->SetParameters(Context);
+	VertexShader->SetParameters(RHICmdList, Context);
+	PixelShader->SetParameters(RHICmdList, Context);
 }
 
 void FRCPassPostProcessCompositeEditorPrimitives::Process(FRenderingCompositePassContext& Context)
@@ -201,6 +202,9 @@ void FRCPassPostProcessCompositeEditorPrimitives::Process(FRenderingCompositePas
 
 	const uint32 MSAASampleCount = GSceneRenderTargets.EditorPrimitivesColor->GetDesc().NumSamples;
 
+	//@todo-rco: RHIPacketList
+	FRHICommandList* RHICmdList = nullptr;
+
 	{
 		RHISetRenderTarget( ColorTarget, DepthTarget );
 		Context.SetViewportAndCallRHI(DestRect);
@@ -239,19 +243,19 @@ void FRCPassPostProcessCompositeEditorPrimitives::Process(FRenderingCompositePas
 
 	if(MSAASampleCount == 1)
 	{
-		SetCompositePrimitivesShaderTempl<1>(Context);
+		SetCompositePrimitivesShaderTempl<1>(RHICmdList, Context);
 	}
 	else if(MSAASampleCount == 2)
 	{
-		SetCompositePrimitivesShaderTempl<2>(Context);
+		SetCompositePrimitivesShaderTempl<2>(RHICmdList, Context);
 	}
 	else if(MSAASampleCount == 4)
 	{
-		SetCompositePrimitivesShaderTempl<4>(Context);
+		SetCompositePrimitivesShaderTempl<4>(RHICmdList, Context);
 	}
 	else if(MSAASampleCount == 8)
 	{
-		SetCompositePrimitivesShaderTempl<8>(Context);
+		SetCompositePrimitivesShaderTempl<8>(RHICmdList, Context);
 	}
 	else
 	{
diff --git a/Engine/Source/Runtime/Renderer/Private/PostProcess/PostProcessDOF.cpp b/Engine/Source/Runtime/Renderer/Private/PostProcess/PostProcessDOF.cpp
--- a/Engine/Source/Runtime/Renderer/Private/PostProcess/PostProcessDOF.cpp
+++ b/Engine/Source/Runtime/Renderer/Private/PostProcess/PostProcessDOF.cpp
@@ -25,7 +25,8 @@ class FPostProcessDOFSetupPS : public FGlobalShader
 
 	static void ModifyCompilationEnvironment(EShaderPlatform Platform, FShaderCompilerEnvironment& OutEnvironment)
 	{
-		FGlobalShader::ModifyCompilationEnvironment(Platform, OutEnvironment);
+		FGlobalShader::ModifyCompilationEnvironment(Platform,OutEnvironment);
+		FDeferredPixelShaderParameters::ModifyCompilationEnvironment(Platform,OutEnvironment);
 		OutEnvironment.SetDefine(TEXT("ENABLE_NEAR_BLUR"), NearBlurEnable);
 	}
 
@@ -54,22 +55,22 @@ public:
 		return bShaderHasOutdatedParameters;
 	}
 
-	void SetParameters(const FRenderingCompositePassContext& Context)
+	void SetParameters(FRHICommandList* RHICmdList, const FRenderingCompositePassContext& Context)
 	{
 		const FPixelShaderRHIParamRef ShaderRHI = GetPixelShader();
 
-		FGlobalShader::SetParameters(ShaderRHI, Context.View);
+		FGlobalShader::SetParameters(RHICmdList, ShaderRHI, Context.View);
 
-		PostprocessParameter.SetPS(ShaderRHI, Context, TStaticSamplerState<SF_Point,AM_Border,AM_Border,AM_Clamp>::GetRHI());
+		PostprocessParameter.SetPS(RHICmdList, ShaderRHI, Context, TStaticSamplerState<SF_Point,AM_Border,AM_Border,AM_Clamp>::GetRHI());
 
-		DeferredParameters.Set(ShaderRHI, Context.View);
+		DeferredParameters.Set(RHICmdList, ShaderRHI, Context.View);
 
 		{
 			FVector4 DepthOfFieldParamValues[2];
 
 			FRCPassPostProcessBokehDOF::ComputeDepthOfFieldParams(Context, DepthOfFieldParamValues);
 
-			SetShaderValueArray(ShaderRHI, DepthOfFieldParams, DepthOfFieldParamValues, 2);
+			SetShaderValueArray(RHICmdList, ShaderRHI, DepthOfFieldParams, DepthOfFieldParamValues, 2);
 		}
 	}
 };
@@ -106,6 +107,9 @@ void FRCPassPostProcessDOFSetup::Process(FRenderingCompositePassContext& Context
 	const FSceneRenderTargetItem& DestRenderTarget0 = PassOutputs[0].RequestSurface(Context);
 	const FSceneRenderTargetItem& DestRenderTarget1 = bNearBlurEnabled ? PassOutputs[1].RequestSurface(Context) : FSceneRenderTargetItem();
 
+	//@todo-rco: RHIPacketList
+	FRHICommandList* RHICmdList = nullptr;
+
 	// Set the view family's render target/viewport.
 	FTextureRHIParamRef RenderTargets[2] =
 	{
@@ -137,7 +141,7 @@ void FRCPassPostProcessDOFSetup::Process(FRenderingCompositePassContext& Context
 		TShaderMapRef< FPostProcessDOFSetupPS<1> > PixelShader(GetGlobalShaderMap());
 		SetGlobalBoundShaderState(BoundShaderState, GFilterVertexDeclaration.VertexDeclarationRHI, *VertexShader, *PixelShader);
 		
-		PixelShader->SetParameters(Context);
+		PixelShader->SetParameters(RHICmdList, Context);
 	}
 	else
 	{
@@ -145,10 +149,10 @@ void FRCPassPostProcessDOFSetup::Process(FRenderingCompositePassContext& Context
 		TShaderMapRef< FPostProcessDOFSetupPS<0> > PixelShader(GetGlobalShaderMap());
 		SetGlobalBoundShaderState(BoundShaderState, GFilterVertexDeclaration.VertexDeclarationRHI, *VertexShader, *PixelShader);
 
-		PixelShader->SetParameters(Context);
+		PixelShader->SetParameters(RHICmdList, Context);
 	}
 
-	VertexShader->SetParameters(Context);
+	VertexShader->SetParameters(RHICmdList, Context);
 
 	// Draw a quad mapping scene color to the view's render target
 	DrawRectangle(
@@ -202,7 +206,8 @@ class FPostProcessDOFRecombinePS : public FGlobalShader
 
 	static void ModifyCompilationEnvironment(EShaderPlatform Platform, FShaderCompilerEnvironment& OutEnvironment)
 	{
-		FGlobalShader::ModifyCompilationEnvironment(Platform, OutEnvironment);
+		FGlobalShader::ModifyCompilationEnvironment(Platform,OutEnvironment);
+		FDeferredPixelShaderParameters::ModifyCompilationEnvironment(Platform,OutEnvironment);
 		OutEnvironment.SetDefine(TEXT("ENABLE_NEAR_BLUR"), NearBlurEnable);
 	}
 
@@ -231,14 +236,14 @@ public:
 		return bShaderHasOutdatedParameters;
 	}
 
-	void SetParameters(const FRenderingCompositePassContext& Context)
+	void SetParameters(FRHICommandList* RHICmdList, const FRenderingCompositePassContext& Context)
 	{
 		const FPixelShaderRHIParamRef ShaderRHI = GetPixelShader();
 
-		FGlobalShader::SetParameters(ShaderRHI, Context.View);
+		FGlobalShader::SetParameters(RHICmdList, ShaderRHI, Context.View);
 
-		DeferredParameters.Set(ShaderRHI, Context.View);
-		PostprocessParameter.SetPS(ShaderRHI, Context, TStaticSamplerState<SF_Bilinear,AM_Clamp,AM_Clamp,AM_Clamp>::GetRHI());
+		DeferredParameters.Set(RHICmdList, ShaderRHI, Context.View);
+		PostprocessParameter.SetPS(RHICmdList, ShaderRHI, Context, TStaticSamplerState<SF_Bilinear,AM_Clamp,AM_Clamp,AM_Clamp>::GetRHI());
 		
 		// Compute out of bounds UVs in the source texture.
 		FVector4 Bounds;
@@ -247,7 +252,7 @@ public:
 		Bounds.Z = (((float)(Context.View.ViewRect.Max.X & (~1))) - 3.0f) / ((float)(GSceneRenderTargets.GetBufferSizeXY().X));
 		Bounds.W = (((float)(Context.View.ViewRect.Max.Y & (~1))) - 3.0f) / ((float)(GSceneRenderTargets.GetBufferSizeXY().Y));
 
-		SetShaderValue(ShaderRHI, DepthOfFieldUVLimit, Bounds);
+		SetShaderValue(RHICmdList, ShaderRHI, DepthOfFieldUVLimit, Bounds);
 	}
 };
 
@@ -277,6 +282,9 @@ void FRCPassPostProcessDOFRecombine::Process(FRenderingCompositePassContext& Con
 
 	const FSceneRenderTargetItem& DestRenderTarget = PassOutputs[0].RequestSurface(Context);
 
+	//@todo-rco: RHIPacketList
+	FRHICommandList* RHICmdList = nullptr;
+
 	// Set the view family's render target/viewport.
 	RHISetRenderTarget(DestRenderTarget.TargetableTexture, FTextureRHIRef());	
 
@@ -297,17 +305,17 @@ void FRCPassPostProcessDOFRecombine::Process(FRenderingCompositePassContext& Con
 		static FGlobalBoundShaderState BoundShaderState;
 		TShaderMapRef< FPostProcessDOFRecombinePS<1> > PixelShader(GetGlobalShaderMap());
 		SetGlobalBoundShaderState(BoundShaderState, GFilterVertexDeclaration.VertexDeclarationRHI, *VertexShader, *PixelShader);
-		PixelShader->SetParameters(Context);
+		PixelShader->SetParameters(RHICmdList, Context);
 	}
 	else
 	{
 		static FGlobalBoundShaderState BoundShaderState;
 		TShaderMapRef< FPostProcessDOFRecombinePS<0> > PixelShader(GetGlobalShaderMap());
 		SetGlobalBoundShaderState(BoundShaderState, GFilterVertexDeclaration.VertexDeclarationRHI, *VertexShader, *PixelShader);
-		PixelShader->SetParameters(Context);
+		PixelShader->SetParameters(RHICmdList, Context);
 	}
 
-	VertexShader->SetParameters(Context);
+	VertexShader->SetParameters(RHICmdList, Context);
 
 	// Draw a quad mapping scene color to the view's render target
 	DrawRectangle(
diff --git a/Engine/Source/Runtime/Renderer/Private/PostProcess/PostProcessDownsample.cpp b/Engine/Source/Runtime/Renderer/Private/PostProcess/PostProcessDownsample.cpp
--- a/Engine/Source/Runtime/Renderer/Private/PostProcess/PostProcessDownsample.cpp
+++ b/Engine/Source/Runtime/Renderer/Private/PostProcess/PostProcessDownsample.cpp
@@ -24,7 +24,8 @@ class FPostProcessDownsamplePS : public FGlobalShader
 
 	static void ModifyCompilationEnvironment(EShaderPlatform Platform, FShaderCompilerEnvironment& OutEnvironment)
 	{
-		FGlobalShader::ModifyCompilationEnvironment(Platform, OutEnvironment);
+		FGlobalShader::ModifyCompilationEnvironment(Platform,OutEnvironment);
+		FDeferredPixelShaderParameters::ModifyCompilationEnvironment(Platform,OutEnvironment);
 		OutEnvironment.SetDefine(TEXT("METHOD"), Method);
 	}
 
@@ -51,13 +52,13 @@ public:
 		return bShaderHasOutdatedParameters;
 	}
 
-	void SetParameters(const FRenderingCompositePassContext& Context)
+	void SetParameters(FRHICommandList* RHICmdList, const FRenderingCompositePassContext& Context)
 	{
 		const FPixelShaderRHIParamRef ShaderRHI = GetPixelShader();
 
-		FGlobalShader::SetParameters(ShaderRHI, Context.View);
-		DeferredParameters.Set(ShaderRHI, Context.View);
-		PostprocessParameter.SetPS(ShaderRHI, Context, TStaticSamplerState<SF_Bilinear,AM_Clamp,AM_Clamp,AM_Clamp>::GetRHI());
+		FGlobalShader::SetParameters(RHICmdList, ShaderRHI, Context.View);
+		DeferredParameters.Set(RHICmdList, ShaderRHI, Context.View);
+		PostprocessParameter.SetPS(RHICmdList, ShaderRHI, Context, TStaticSamplerState<SF_Bilinear,AM_Clamp,AM_Clamp,AM_Clamp>::GetRHI());
 	}
 
 	static const TCHAR* GetSourceFilename()
@@ -107,11 +108,11 @@ public:
 		return bShaderHasOutdatedParameters;
 	}
 
-	void SetParameters(const FRenderingCompositePassContext& Context)
+	void SetParameters(FRHICommandList* RHICmdList, const FRenderingCompositePassContext& Context)
 	{
 		const FVertexShaderRHIParamRef ShaderRHI = GetVertexShader();
 
-		FGlobalShader::SetParameters(ShaderRHI, Context.View);
+		FGlobalShader::SetParameters(RHICmdList, ShaderRHI, Context.View);
 
 		const FPooledRenderTargetDesc* InputDesc = Context.Pass->GetInputDesc(ePId_Input0);
 
@@ -136,7 +137,7 @@ FRCPassPostProcessDownsample::FRCPassPostProcessDownsample(EPixelFormat InOverri
 
 
 template <uint32 Method>
-void FRCPassPostProcessDownsample::SetShader(const FRenderingCompositePassContext& Context)
+void FRCPassPostProcessDownsample::SetShader(FRHICommandList* RHICmdList, const FRenderingCompositePassContext& Context)
 {
 	TShaderMapRef<FPostProcessDownsampleVS> VertexShader(GetGlobalShaderMap());
 	TShaderMapRef<FPostProcessDownsamplePS<Method> > PixelShader(GetGlobalShaderMap());
@@ -145,8 +146,8 @@ void FRCPassPostProcessDownsample::SetShader(const FRenderingCompositePassContex
 
 	SetGlobalBoundShaderState(BoundShaderState, GFilterVertexDeclaration.VertexDeclarationRHI, *VertexShader, *PixelShader);
 
-	PixelShader->SetParameters(Context);
-	VertexShader->SetParameters(Context);
+	PixelShader->SetParameters(RHICmdList, Context);
+	VertexShader->SetParameters(RHICmdList, Context);
 }
 
 
@@ -173,6 +174,9 @@ void FRCPassPostProcessDownsample::Process(FRenderingCompositePassContext& Conte
 
 	const FSceneRenderTargetItem& DestRenderTarget = PassOutputs[0].RequestSurface(Context);
 
+	//@todo-rco: RHIPacketList
+	FRHICommandList* RHICmdList = nullptr;
+
 	// Set the view family's render target/viewport.
 	RHISetRenderTarget(DestRenderTarget.TargetableTexture, FTextureRHIRef());
 
@@ -190,18 +194,18 @@ void FRCPassPostProcessDownsample::Process(FRenderingCompositePassContext& Conte
 	{
 		// also put depth in alpha
 		InflateSize = 2;
-		SetShader<2>(Context);
+		SetShader<2>(RHICmdList, Context);
 	}
 	else
 	{
 		if (Quality == 0)
 		{
-			SetShader<0>(Context);
+			SetShader<0>(RHICmdList, Context);
 			InflateSize = 1;
 		}
 		else
 		{
-			SetShader<1>(Context);
+			SetShader<1>(RHICmdList, Context);
 			InflateSize = 2;
 		}
 	}
diff --git a/Engine/Source/Runtime/Renderer/Private/PostProcess/PostProcessDownsample.h b/Engine/Source/Runtime/Renderer/Private/PostProcess/PostProcessDownsample.h
--- a/Engine/Source/Runtime/Renderer/Private/PostProcess/PostProcessDownsample.h
+++ b/Engine/Source/Runtime/Renderer/Private/PostProcess/PostProcessDownsample.h
@@ -31,7 +31,7 @@ public:
 	virtual FPooledRenderTargetDesc ComputeOutputDesc(EPassOutputId InPassOutputId) const;
 
 private:
-	template <uint32 Method> static void SetShader(const FRenderingCompositePassContext& Context);
+	template <uint32 Method> static void SetShader(FRHICommandList* RHICmdList, const FRenderingCompositePassContext& Context);
 
 	EPixelFormat OverrideFormat;
 	// explained in constructor
diff --git a/Engine/Source/Runtime/Renderer/Private/PostProcess/PostProcessEyeAdaptation.cpp b/Engine/Source/Runtime/Renderer/Private/PostProcess/PostProcessEyeAdaptation.cpp
--- a/Engine/Source/Runtime/Renderer/Private/PostProcess/PostProcessEyeAdaptation.cpp
+++ b/Engine/Source/Runtime/Renderer/Private/PostProcess/PostProcessEyeAdaptation.cpp
@@ -41,19 +41,19 @@ public:
 		EyeAdaptationParams.Bind(Initializer.ParameterMap, TEXT("EyeAdaptationParams"));
 	}
 
-	void SetPS(const FRenderingCompositePassContext& Context)
+	void SetPS(FRHICommandList* RHICmdList, const FRenderingCompositePassContext& Context)
 	{
 		const FPixelShaderRHIParamRef ShaderRHI = GetPixelShader();
 
-		FGlobalShader::SetParameters(ShaderRHI, Context.View);
+		FGlobalShader::SetParameters(RHICmdList, ShaderRHI, Context.View);
 
-		PostprocessParameter.SetPS(ShaderRHI, Context, TStaticSamplerState<SF_Bilinear,AM_Clamp,AM_Clamp,AM_Clamp>::GetRHI());
+		PostprocessParameter.SetPS(RHICmdList, ShaderRHI, Context, TStaticSamplerState<SF_Bilinear,AM_Clamp,AM_Clamp,AM_Clamp>::GetRHI());
 
 		{
 			FVector4 Temp[3];
 
 			FRCPassPostProcessEyeAdaptation::ComputeEyeAdaptationParamsValue(Context.View, Temp);
-			SetShaderValueArray(ShaderRHI, EyeAdaptationParams, Temp, 3);
+			SetShaderValueArray(RHICmdList, ShaderRHI, EyeAdaptationParams, Temp, 3);
 		}
 	}
 
@@ -81,6 +81,8 @@ void FRCPassPostProcessEyeAdaptation::Process(FRenderingCompositePassContext& Co
 	FIntPoint DestSize = EyeAdaptation->GetDesc().Extent;
 
 	// we render to our own output render target, not the intermediate one created by the compositing system
+	//@todo-rco: RHIPacketList
+	FRHICommandList* RHICmdList = nullptr;
 
 	// Set the view family's render target/viewport.
 	RHISetRenderTarget(EyeAdaptation->GetRenderTargetItem().TargetableTexture, FTextureRHIRef());	
@@ -98,7 +100,7 @@ void FRCPassPostProcessEyeAdaptation::Process(FRenderingCompositePassContext& Co
 
 	SetGlobalBoundShaderState(BoundShaderState, GFilterVertexDeclaration.VertexDeclarationRHI, *VertexShader, *PixelShader);
 
-	PixelShader->SetPS(Context);
+	PixelShader->SetPS(RHICmdList, Context);
 
 	// Draw a quad mapping scene color to the view's render target
 	DrawRectangle(
diff --git a/Engine/Source/Runtime/Renderer/Private/PostProcess/PostProcessGBufferHints.cpp b/Engine/Source/Runtime/Renderer/Private/PostProcess/PostProcessGBufferHints.cpp
--- a/Engine/Source/Runtime/Renderer/Private/PostProcess/PostProcessGBufferHints.cpp
+++ b/Engine/Source/Runtime/Renderer/Private/PostProcess/PostProcessGBufferHints.cpp
@@ -22,6 +22,12 @@ class FPostProcessGBufferHintsPS : public FGlobalShader
 		return IsFeatureLevelSupported(Platform, ERHIFeatureLevel::SM5);
 	}
 
+	static void ModifyCompilationEnvironment(EShaderPlatform Platform, FShaderCompilerEnvironment& OutEnvironment)
+	{
+		FGlobalShader::ModifyCompilationEnvironment(Platform,OutEnvironment);
+		FDeferredPixelShaderParameters::ModifyCompilationEnvironment(Platform,OutEnvironment);
+	}
+
 	/** Default constructor. */
 	FPostProcessGBufferHintsPS() {}
 
@@ -41,23 +47,23 @@ public:
 		MiniFontTexture.Bind(Initializer.ParameterMap, TEXT("MiniFontTexture"));
 	}
 
-	void SetPS(const FRenderingCompositePassContext& Context)
+	void SetPS(FRHICommandList* RHICmdList, const FRenderingCompositePassContext& Context)
 	{
 		const FPixelShaderRHIParamRef ShaderRHI = GetPixelShader();
 
-		FGlobalShader::SetParameters(ShaderRHI, Context.View);
+		FGlobalShader::SetParameters(RHICmdList, ShaderRHI, Context.View);
 
-		PostprocessParameter.SetPS(ShaderRHI, Context, TStaticSamplerState<SF_Point,AM_Clamp,AM_Clamp,AM_Clamp>::GetRHI());
-		DeferredParameters.Set(ShaderRHI, Context.View);
+		PostprocessParameter.SetPS(RHICmdList, ShaderRHI, Context, TStaticSamplerState<SF_Point,AM_Clamp,AM_Clamp,AM_Clamp>::GetRHI());
+		DeferredParameters.Set(RHICmdList, ShaderRHI, Context.View);
 
 		{
 			FVector4 Temp[3];
 
 			FRCPassPostProcessEyeAdaptation::ComputeEyeAdaptationParamsValue(Context.View, Temp);
-			SetShaderValueArray(ShaderRHI, EyeAdaptationParams, Temp, 3);
+			SetShaderValueArray(RHICmdList, ShaderRHI, EyeAdaptationParams, Temp, 3);
 		}
 
-		SetTextureParameter(ShaderRHI, MiniFontTexture, GEngine->MiniFontTexture ? GEngine->MiniFontTexture->Resource->TextureRHI : GSystemTextures.WhiteDummy->GetRenderTargetItem().TargetableTexture);
+		SetTextureParameter(RHICmdList, ShaderRHI, MiniFontTexture, GEngine->MiniFontTexture ? GEngine->MiniFontTexture->Resource->TextureRHI : GSystemTextures.WhiteDummy->GetRenderTargetItem().TargetableTexture);
 	}
 	
 	// FShader interface.
@@ -91,6 +97,9 @@ void FRCPassPostProcessGBufferHints::Process(FRenderingCompositePassContext& Con
 
 	const FSceneRenderTargetItem& DestRenderTarget = PassOutputs[0].RequestSurface(Context);
 
+	//@todo-rco: RHIPacketList
+	FRHICommandList* RHICmdList = nullptr;
+
 	// Set the view family's render target/viewport.
 	RHISetRenderTarget(DestRenderTarget.TargetableTexture, FTextureRHIRef());	
 	Context.SetViewportAndCallRHI(DestRect);
@@ -107,7 +116,7 @@ void FRCPassPostProcessGBufferHints::Process(FRenderingCompositePassContext& Con
 
 	SetGlobalBoundShaderState(BoundShaderState, GFilterVertexDeclaration.VertexDeclarationRHI, *VertexShader, *PixelShader);
 
-	PixelShader->SetPS(Context);
+	PixelShader->SetPS(RHICmdList, Context);
 
 	// Draw a quad mapping scene color to the view's render target
 	DrawRectangle(
diff --git a/Engine/Source/Runtime/Renderer/Private/PostProcess/PostProcessHMD.cpp b/Engine/Source/Runtime/Renderer/Private/PostProcess/PostProcessHMD.cpp
--- a/Engine/Source/Runtime/Renderer/Private/PostProcess/PostProcessHMD.cpp
+++ b/Engine/Source/Runtime/Renderer/Private/PostProcess/PostProcessHMD.cpp
@@ -13,6 +13,7 @@
 #include "PostProcessHistogram.h"
 #include "PostProcessEyeAdaptation.h"
 #include "IHeadMountedDisplay.h"
+#include "RHICommandList.h"
 
 
 /** The filter vertex declaration resource type. */
@@ -55,7 +56,7 @@ class FPostProcessHMDVS : public FGlobalShader
 
 	static bool ShouldCache(EShaderPlatform Platform)
 	{
-		return true;
+		return !IsOpenGLPlatform(Platform);
 	}
 
 	/** Default constructor. */
@@ -71,18 +72,18 @@ public:
 		EyeToSrcUVOffset.Bind(Initializer.ParameterMap, TEXT("EyeToSrcUVOffset"));
 	}
 
-	void SetVS(const FRenderingCompositePassContext& Context, EStereoscopicPass StereoPass)
+	void SetVS(FRHICommandList* RHICmdList, const FRenderingCompositePassContext& Context, EStereoscopicPass StereoPass)
 	{
 		const FVertexShaderRHIParamRef ShaderRHI = GetVertexShader();
 
-		FGlobalShader::SetParameters(ShaderRHI, Context.View);
+		FGlobalShader::SetParameters(RHICmdList, ShaderRHI, Context.View);
 
 		check(GEngine->HMDDevice.IsValid());
 		FVector2D EyeToSrcUVScaleValue;
 		FVector2D EyeToSrcUVOffsetValue;
 		GEngine->HMDDevice->GetEyeRenderParams_RenderThread(StereoPass, EyeToSrcUVScaleValue, EyeToSrcUVOffsetValue);
-		SetShaderValue(ShaderRHI, EyeToSrcUVScale, EyeToSrcUVScaleValue);
-		SetShaderValue(ShaderRHI, EyeToSrcUVOffset, EyeToSrcUVOffsetValue);
+		SetShaderValue(RHICmdList, ShaderRHI, EyeToSrcUVScale, EyeToSrcUVScaleValue);
+		SetShaderValue(RHICmdList, ShaderRHI, EyeToSrcUVOffset, EyeToSrcUVOffsetValue);
 	}
 
 	// FShader interface.
@@ -101,7 +102,7 @@ class FPostProcessHMDPS : public FGlobalShader
 
 	static bool ShouldCache(EShaderPlatform Platform)
 	{
-		return true;
+		return !IsOpenGLPlatform(Platform);
 	}
 
 	/** Default constructor. */
@@ -122,14 +123,14 @@ public:
 
 	}
 
-	void SetPS(const FRenderingCompositePassContext& Context, FIntRect SrcRect, FIntPoint SrcBufferSize, EStereoscopicPass StereoPass, FMatrix& QuadTexTransform)
+	void SetPS(FRHICommandList* RHICmdList, const FRenderingCompositePassContext& Context, FIntRect SrcRect, FIntPoint SrcBufferSize, EStereoscopicPass StereoPass, FMatrix& QuadTexTransform)
 	{
 		const FPixelShaderRHIParamRef ShaderRHI = GetPixelShader();
 
-		FGlobalShader::SetParameters(ShaderRHI, Context.View);
+		FGlobalShader::SetParameters(RHICmdList, ShaderRHI, Context.View);
 
-		PostprocessParameter.SetPS(ShaderRHI, Context, TStaticSamplerState<SF_Bilinear, AM_Clamp, AM_Clamp, AM_Clamp>::GetRHI());
-		DeferredParameters.Set(ShaderRHI, Context.View);
+		PostprocessParameter.SetPS(RHICmdList, ShaderRHI, Context, TStaticSamplerState<SF_Bilinear, AM_Clamp, AM_Clamp, AM_Clamp>::GetRHI());
+		DeferredParameters.Set(RHICmdList, ShaderRHI, Context.View);
 	}
 
 	// FShader interface.
@@ -144,6 +145,7 @@ public:
 	{
 		FGlobalShader::ModifyCompilationEnvironment(Platform, OutEnvironment);
 		OutEnvironment.SetDefine(TEXT("USE_TIMEWARP"), uint32(0));
+		FDeferredPixelShaderParameters::ModifyCompilationEnvironment(Platform,OutEnvironment);
 	}
 };
 
@@ -172,6 +174,9 @@ void FRCPassPostProcessHMD::Process(FRenderingCompositePassContext& Context)
 
     const FSceneRenderTargetItem& DestRenderTarget = PassOutputs[0].RequestSurface(Context);
 
+	//@todo-rco: RHIPacketList
+	FRHICommandList* RHICmdList = nullptr;
+
 	// Set the view family's render target/viewport.
 	RHISetRenderTarget(DestRenderTarget.TargetableTexture, FTextureRHIRef());	
 
@@ -193,8 +198,8 @@ void FRCPassPostProcessHMD::Process(FRenderingCompositePassContext& Context)
 		TShaderMapRef<FPostProcessHMDPS> PixelShader(GetGlobalShaderMap());
 		static FGlobalBoundShaderState BoundShaderState;
 		SetGlobalBoundShaderState(BoundShaderState, GDistortionVertexDeclaration.VertexDeclarationRHI, *VertexShader, *PixelShader);
-		VertexShader->SetVS(Context, View.StereoPass);
-		PixelShader->SetPS(Context, SrcRect, SrcSize, View.StereoPass, QuadTexTransform);
+		VertexShader->SetVS(RHICmdList, Context, View.StereoPass);
+		PixelShader->SetPS(RHICmdList, Context, SrcRect, SrcSize, View.StereoPass, QuadTexTransform);
 	}
 
 	GEngine->HMDDevice->DrawDistortionMesh_RenderThread(Context, View, SrcSize);
diff --git a/Engine/Source/Runtime/Renderer/Private/PostProcess/PostProcessHistogram.cpp b/Engine/Source/Runtime/Renderer/Private/PostProcess/PostProcessHistogram.cpp
--- a/Engine/Source/Runtime/Renderer/Private/PostProcess/PostProcessHistogram.cpp
+++ b/Engine/Source/Runtime/Renderer/Private/PostProcess/PostProcessHistogram.cpp
@@ -55,25 +55,25 @@ public:
 		EyeAdaptationParams.Bind(Initializer.ParameterMap, TEXT("EyeAdaptationParams"));
 	}
 
-	void SetCS(const FRenderingCompositePassContext& Context, FIntPoint ThreadGroupCountValue, FIntPoint LeftTopOffsetValue, FIntPoint GatherExtent)
+	void SetCS(FRHICommandList* RHICmdList, const FRenderingCompositePassContext& Context, FIntPoint ThreadGroupCountValue, FIntPoint LeftTopOffsetValue, FIntPoint GatherExtent)
 	{
 		const FComputeShaderRHIParamRef ShaderRHI = GetComputeShader();
 
-		FGlobalShader::SetParameters(ShaderRHI, Context.View);
+		FGlobalShader::SetParameters(RHICmdList, ShaderRHI, Context.View);
 
-		PostprocessParameter.SetCS(ShaderRHI, Context, TStaticSamplerState<SF_Point,AM_Clamp,AM_Clamp,AM_Clamp>::GetRHI());
+		PostprocessParameter.SetCS(RHICmdList, ShaderRHI, Context, TStaticSamplerState<SF_Point,AM_Clamp,AM_Clamp,AM_Clamp>::GetRHI());
 
-		SetShaderValue(ShaderRHI, ThreadGroupCount, ThreadGroupCountValue);
-		SetShaderValue(ShaderRHI, LeftTopOffset, LeftTopOffsetValue);
+		SetShaderValue(RHICmdList, ShaderRHI, ThreadGroupCount, ThreadGroupCountValue);
+		SetShaderValue(RHICmdList, ShaderRHI, LeftTopOffset, LeftTopOffsetValue);
 
 		FVector4 HistogramParametersValue(GatherExtent.X, GatherExtent.Y, 0, 0);
-		SetShaderValue(ShaderRHI, HistogramParameters, HistogramParametersValue);
+		SetShaderValue(RHICmdList, ShaderRHI, HistogramParameters, HistogramParametersValue);
 
 		{
 			FVector4 Temp[3];
 
 			FRCPassPostProcessEyeAdaptation::ComputeEyeAdaptationParamsValue(Context.View, Temp);
-			SetShaderValueArray(ShaderRHI, EyeAdaptationParams, Temp, 3);
+			SetShaderValueArray(RHICmdList, ShaderRHI, EyeAdaptationParams, Temp, 3);
 		}
 	}
 	
@@ -123,9 +123,11 @@ void FRCPassPostProcessHistogram::Process(FRenderingCompositePassContext& Contex
 	FIntPoint ThreadGroupCountValue = ComputeThreadGroupCount(GatherExtent);
 
 	
-	ComputeShader->SetCS(Context, ThreadGroupCountValue, (DestRect.Min + FIntPoint(1, 1)) / 2, GatherExtent);
+	//@todo-rco: RHIPacketList
+	FRHICommandList* RHICmdList = nullptr;
+	ComputeShader->SetCS(RHICmdList, Context, ThreadGroupCountValue, (DestRect.Min + FIntPoint(1, 1)) / 2, GatherExtent);
 	
-	DispatchComputeShader(*ComputeShader, ThreadGroupCountValue.X, ThreadGroupCountValue.Y, 1);
+	DispatchComputeShader(RHICmdList, *ComputeShader, ThreadGroupCountValue.X, ThreadGroupCountValue.Y, 1);
 
 	// un-set destination
 	RHISetUAVParameter(ComputeShader->GetComputeShader(), ComputeShader->HistogramRWTexture.GetBaseIndex(), NULL);
diff --git a/Engine/Source/Runtime/Renderer/Private/PostProcess/PostProcessHistogramReduce.cpp b/Engine/Source/Runtime/Renderer/Private/PostProcess/PostProcessHistogramReduce.cpp
--- a/Engine/Source/Runtime/Renderer/Private/PostProcess/PostProcessHistogramReduce.cpp
+++ b/Engine/Source/Runtime/Renderer/Private/PostProcess/PostProcessHistogramReduce.cpp
@@ -45,15 +45,15 @@ public:
 		EyeAdapationTemporalParams.Bind(Initializer.ParameterMap, TEXT("EyeAdapationTemporalParams"));
 	}
 
-	void SetPS(const FRenderingCompositePassContext& Context, uint32 LoopSizeValue)
+	void SetPS(FRHICommandList* RHICmdList, const FRenderingCompositePassContext& Context, uint32 LoopSizeValue)
 	{
 		const FPixelShaderRHIParamRef ShaderRHI = GetPixelShader();
 
-		FGlobalShader::SetParameters(ShaderRHI, Context.View);
+		FGlobalShader::SetParameters(RHICmdList, ShaderRHI, Context.View);
 
-		PostprocessParameter.SetPS(ShaderRHI, Context, TStaticSamplerState<SF_Bilinear,AM_Clamp,AM_Clamp,AM_Clamp>::GetRHI());
+		PostprocessParameter.SetPS(RHICmdList, ShaderRHI, Context, TStaticSamplerState<SF_Bilinear,AM_Clamp,AM_Clamp,AM_Clamp>::GetRHI());
 
-		SetShaderValue(ShaderRHI, LoopSize, LoopSizeValue);
+		SetShaderValue(RHICmdList, ShaderRHI, LoopSize, LoopSizeValue);
 
 		if(EyeAdaptationTexture.IsBound())
 		{
@@ -61,18 +61,18 @@ public:
 
 			if(EyeAdaptationRT)
 			{
-				SetTextureParameter(ShaderRHI, EyeAdaptationTexture, EyeAdaptationRT->GetRenderTargetItem().TargetableTexture);
+				SetTextureParameter(RHICmdList, ShaderRHI, EyeAdaptationTexture, EyeAdaptationRT->GetRenderTargetItem().TargetableTexture);
 			}
 			else
 			{
 				// some views don't have a state, thumbnail rendering?
-				SetTextureParameter(ShaderRHI, EyeAdaptationTexture, GWhiteTexture->TextureRHI);
+				SetTextureParameter(RHICmdList, ShaderRHI, EyeAdaptationTexture, GWhiteTexture->TextureRHI);
 			}
 		}
 
 		// todo
 		FVector4 EyeAdapationTemporalParamsValue(0, 0, 0, 0);
-		SetShaderValue(ShaderRHI, EyeAdapationTemporalParams, EyeAdapationTemporalParamsValue);
+		SetShaderValue(RHICmdList, ShaderRHI, EyeAdapationTemporalParams, EyeAdapationTemporalParamsValue);
 	}
 	
 	// FShader interface.
@@ -125,7 +125,9 @@ void FRCPassPostProcessHistogramReduce::Process(FRenderingCompositePassContext&
 
 	uint32 LoopSizeValue = ComputeLoopSize(GatherExtent);
 
-	PixelShader->SetPS(Context, LoopSizeValue);
+	//@todo-rco: RHIPacketList
+	FRHICommandList* RHICmdList = nullptr;
+	PixelShader->SetPS(RHICmdList, Context, LoopSizeValue);
 
 	// Draw a quad mapping scene color to the view's render target
 	DrawRectangle(
diff --git a/Engine/Source/Runtime/Renderer/Private/PostProcess/PostProcessLensBlur.cpp b/Engine/Source/Runtime/Renderer/Private/PostProcess/PostProcessLensBlur.cpp
--- a/Engine/Source/Runtime/Renderer/Private/PostProcess/PostProcessLensBlur.cpp
+++ b/Engine/Source/Runtime/Renderer/Private/PostProcess/PostProcessLensBlur.cpp
@@ -49,31 +49,31 @@ public:
 	}
 
 	/** to have a similar interface as all other shaders */
-	void SetParameters(const FRenderingCompositePassContext& Context, FIntPoint TileCountValue, uint32 TileSize, float PixelKernelSize, float Threshold)
+	void SetParameters(FRHICommandList* RHICmdList, const FRenderingCompositePassContext& Context, FIntPoint TileCountValue, uint32 TileSize, float PixelKernelSize, float Threshold)
 	{
 		const FVertexShaderRHIParamRef ShaderRHI = GetVertexShader();
 
-		FGlobalShader::SetParameters(ShaderRHI, Context.View);
+		FGlobalShader::SetParameters(RHICmdList, ShaderRHI, Context.View);
 
-		PostprocessParameter.SetVS(ShaderRHI, Context, TStaticSamplerState<SF_Bilinear,AM_Clamp,AM_Clamp,AM_Clamp>::GetRHI());
+		PostprocessParameter.SetVS(RHICmdList, ShaderRHI, Context, TStaticSamplerState<SF_Bilinear,AM_Clamp,AM_Clamp,AM_Clamp>::GetRHI());
 
 		{
 			FIntRect TileCountAndSizeValue(TileCountValue, FIntPoint(TileSize, TileSize));
 
-			SetShaderValue(ShaderRHI, TileCountAndSize, TileCountAndSizeValue);
+			SetShaderValue(RHICmdList, ShaderRHI, TileCountAndSize, TileCountAndSizeValue);
 		}
 
 		{
 			// only approximate as the mip mapping doesn't produce accurate brightness scaling
 			FVector4 ColorScaleValue(1.0f / FMath::Max(1.0f, PixelKernelSize * PixelKernelSize), Threshold, 0, 0);
 
-			SetShaderValue(ShaderRHI, ColorScale, ColorScaleValue);
+			SetShaderValue(RHICmdList, ShaderRHI, ColorScale, ColorScaleValue);
 		}
 
 		{
 			FVector4 KernelSizeValue(PixelKernelSize, PixelKernelSize, 0, 0);
 
-			SetShaderValue(ShaderRHI, KernelSize, KernelSizeValue);
+			SetShaderValue(RHICmdList, ShaderRHI, KernelSize, KernelSizeValue);
 		}
 	}
 };
@@ -115,13 +115,13 @@ public:
 		return bShaderHasOutdatedParameters;
 	}
 
-	void SetParameters(const FRenderingCompositePassContext& Context, float PixelKernelSize)
+	void SetParameters(FRHICommandList* RHICmdList, const FRenderingCompositePassContext& Context, float PixelKernelSize)
 	{
 		const FPixelShaderRHIParamRef ShaderRHI = GetPixelShader();
 
-		FGlobalShader::SetParameters(ShaderRHI, Context.View);
+		FGlobalShader::SetParameters(RHICmdList, ShaderRHI, Context.View);
 
-		PostprocessParameter.SetPS(ShaderRHI, Context, TStaticSamplerState<SF_Bilinear,AM_Clamp,AM_Clamp,AM_Clamp>::GetRHI());
+		PostprocessParameter.SetPS(RHICmdList, ShaderRHI, Context, TStaticSamplerState<SF_Bilinear,AM_Clamp,AM_Clamp,AM_Clamp>::GetRHI());
 				
 		{
 			FTextureRHIParamRef TextureRHI = GWhiteTexture->TextureRHI;
@@ -146,7 +146,7 @@ public:
 				}
 			}
 
-			SetTextureParameter(ShaderRHI, LensTexture, LensTextureSampler, TStaticSamplerState<SF_Bilinear,AM_Border,AM_Border,AM_Clamp>::GetRHI(), TextureRHI);
+			SetTextureParameter(RHICmdList, ShaderRHI, LensTexture, LensTextureSampler, TStaticSamplerState<SF_Bilinear,AM_Border,AM_Border,AM_Clamp>::GetRHI(), TextureRHI);
 		}
 	}
 };
@@ -208,8 +208,10 @@ void FRCPassPostProcessLensBlur::Process(FRenderingCompositePassContext& Context
 
 	float PixelKernelSize = PercentKernelSize / 100.0f * ViewSize.X;
 
-	VertexShader->SetParameters(Context, TileCount, TileSize, PixelKernelSize, Threshold);
-	PixelShader->SetParameters(Context, PixelKernelSize);
+	//@todo-rco: RHIPacketList
+	FRHICommandList* RHICmdList = nullptr;
+	VertexShader->SetParameters(RHICmdList, Context, TileCount, TileSize, PixelKernelSize, Threshold);
+	PixelShader->SetParameters(RHICmdList, Context, PixelKernelSize);
 
 	RHISetStreamSource(0, NULL, 0, 0);
 
diff --git a/Engine/Source/Runtime/Renderer/Private/PostProcess/PostProcessLensFlares.cpp b/Engine/Source/Runtime/Renderer/Private/PostProcess/PostProcessLensFlares.cpp
--- a/Engine/Source/Runtime/Renderer/Private/PostProcess/PostProcessLensFlares.cpp
+++ b/Engine/Source/Runtime/Renderer/Private/PostProcess/PostProcessLensFlares.cpp
@@ -41,13 +41,13 @@ public:
 		return bShaderHasOutdatedParameters;
 	}
 
-	void SetParameters(const FRenderingCompositePassContext& Context)
+	void SetParameters(FRHICommandList* RHICmdList, const FRenderingCompositePassContext& Context)
 	{
 		const FPixelShaderRHIParamRef ShaderRHI = GetPixelShader();
 
-		FGlobalShader::SetParameters(ShaderRHI, Context.View);
+		FGlobalShader::SetParameters(RHICmdList, ShaderRHI, Context.View);
 
-		PostprocessParameter.SetPS(ShaderRHI, Context, TStaticSamplerState<SF_Bilinear,AM_Clamp,AM_Clamp,AM_Clamp>::GetRHI());
+		PostprocessParameter.SetPS(RHICmdList, ShaderRHI, Context, TStaticSamplerState<SF_Bilinear,AM_Clamp,AM_Clamp,AM_Clamp>::GetRHI());
 	}
 };
 
@@ -89,15 +89,15 @@ public:
 		return bShaderHasOutdatedParameters;
 	}
 
-	void SetParameters(const FRenderingCompositePassContext& Context, FVector2D TexScaleValue)
+	void SetParameters(FRHICommandList* RHICmdList, const FRenderingCompositePassContext& Context, FVector2D TexScaleValue)
 	{
 		const FPixelShaderRHIParamRef ShaderRHI = GetPixelShader();
 
-		FGlobalShader::SetParameters(ShaderRHI, Context.View);
+		FGlobalShader::SetParameters(RHICmdList, ShaderRHI, Context.View);
 
-		PostprocessParameter.SetPS(ShaderRHI, Context, TStaticSamplerState<SF_Bilinear,AM_Clamp,AM_Clamp,AM_Clamp>::GetRHI());
+		PostprocessParameter.SetPS(RHICmdList, ShaderRHI, Context, TStaticSamplerState<SF_Bilinear,AM_Clamp,AM_Clamp,AM_Clamp>::GetRHI());
 
-		SetShaderValue(ShaderRHI, TexScale, TexScaleValue);
+		SetShaderValue(RHICmdList, ShaderRHI, TexScale, TexScaleValue);
 	}
 };
 
@@ -154,7 +154,10 @@ void FRCPassPostProcessLensFlares::Process(FRenderingCompositePassContext& Conte
 	RHISetDepthStencilState(TStaticDepthStencilState<false,CF_Always>::GetRHI());
 
 	TShaderMapRef<FPostProcessVS> VertexShader(GetGlobalShaderMap());
-	
+
+	//@todo-rco: RHIPacketList
+	FRHICommandList* RHICmdList = nullptr;
+
 	// setup background (bloom), can be implemented to use additive blending to avoid the read here
 	{
 		TShaderMapRef<FPostProcessLensFlareBasePS> PixelShader(GetGlobalShaderMap());
@@ -163,8 +166,8 @@ void FRCPassPostProcessLensFlares::Process(FRenderingCompositePassContext& Conte
 
 		SetGlobalBoundShaderState(BoundShaderState, GFilterVertexDeclaration.VertexDeclarationRHI, *VertexShader, *PixelShader);
 
-		VertexShader->SetParameters(Context);
-		PixelShader->SetParameters(Context);
+		VertexShader->SetParameters(RHICmdList, Context);
+		PixelShader->SetParameters(RHICmdList, Context);
 
 		// Draw a quad mapping scene color to the view's render target
 		DrawRectangle(
@@ -191,8 +194,8 @@ void FRCPassPostProcessLensFlares::Process(FRenderingCompositePassContext& Conte
 
 		FVector2D TexScaleValue = FVector2D(TexSize2) / ViewSize2;
 
-		VertexShader->SetParameters(Context);
-		PixelShader->SetParameters(Context, TexScaleValue);
+		VertexShader->SetParameters(RHICmdList, Context);
+		PixelShader->SetParameters(RHICmdList, Context, TexScaleValue);
 
 		// todo: expose
 		const uint32 Count = 8;
@@ -217,7 +220,9 @@ void FRCPassPostProcessLensFlares::Process(FRenderingCompositePassContext& Conte
 			Alpha *= SizeScale;
 			
 			// set the individual flare color
-			SetShaderValue(PixelShader->GetPixelShader(), PixelShader->FlareColor, FlareColor * LensFlareHDRColor);
+			//@todo-rco: RHIPacketList
+			FRHICommandList* RHICmdList = nullptr;
+			SetShaderValue(RHICmdList, PixelShader->GetPixelShader(), PixelShader->FlareColor, FlareColor * LensFlareHDRColor);
 
 			// Draw a quad mapping scene color to the view's render target
 			DrawRectangle(
diff --git a/Engine/Source/Runtime/Renderer/Private/PostProcess/PostProcessMaterial.cpp b/Engine/Source/Runtime/Renderer/Private/PostProcess/PostProcessMaterial.cpp
--- a/Engine/Source/Runtime/Renderer/Private/PostProcess/PostProcessMaterial.cpp
+++ b/Engine/Source/Runtime/Renderer/Private/PostProcess/PostProcessMaterial.cpp
@@ -31,9 +31,9 @@ public:
 	{
 	}
 
-	void SetParameters( const FRenderingCompositePassContext& Context )
+	void SetParameters(FRHICommandList* RHICmdList, const FRenderingCompositePassContext& Context )
 	{
-		FMaterialShader::SetParameters(GetVertexShader(), Context.View);
+		FMaterialShader::SetParameters(RHICmdList, GetVertexShader(), Context.View);
 	}
 
 	// Begin FShader interface
@@ -70,12 +70,12 @@ public:
 		PostprocessParameter.Bind(Initializer.ParameterMap);
 	}
 
-	void SetParameters( const FRenderingCompositePassContext& Context, const FMaterialRenderProxy* Material )
+	void SetParameters(FRHICommandList* RHICmdList, const FRenderingCompositePassContext& Context, const FMaterialRenderProxy* Material )
 	{
 		const FPixelShaderRHIParamRef ShaderRHI = GetPixelShader();
 
-		FMaterialShader::SetParameters(ShaderRHI, Material, *Material->GetMaterial(Context.View.GetFeatureLevel()), Context.View, true, ESceneRenderTargetsMode::SetTextures);
-		PostprocessParameter.SetPS(ShaderRHI, Context, TStaticSamplerState<SF_Point,AM_Clamp,AM_Clamp,AM_Clamp>::GetRHI());
+		FMaterialShader::SetParameters(RHICmdList, ShaderRHI, Material, *Material->GetMaterial(Context.View.GetFeatureLevel()), Context.View, true, ESceneRenderTargetsMode::SetTextures);
+		PostprocessParameter.SetPS(RHICmdList, ShaderRHI, Context, TStaticSamplerState<SF_Point,AM_Clamp,AM_Clamp,AM_Clamp>::GetRHI());
 	}
 
 	virtual bool Serialize(FArchive& Ar)
@@ -154,8 +154,10 @@ void FRCPassPostProcessMaterial::Process(FRenderingCompositePassContext& Context
 	FBoundShaderStateRHIRef BoundShaderState = RHICreateBoundShaderState(GetVertexDeclarationFVector4(), VertexShader->GetVertexShader(), FHullShaderRHIRef(), FDomainShaderRHIRef(), PixelShader->GetPixelShader(), FGeometryShaderRHIRef());
 	RHISetBoundShaderState(BoundShaderState);
 
-	VertexShader->SetParameters(Context);
-	PixelShader->SetParameters(Context, MaterialInterface->GetRenderProxy(false));
+	//@todo-rco: RHIPacketList
+	FRHICommandList* RHICmdList = nullptr;
+	VertexShader->SetParameters(RHICmdList, Context);
+	PixelShader->SetParameters(RHICmdList, Context, MaterialInterface->GetRenderProxy(false));
 
 	DrawRectangle(
 		0, 0,
