UE4 - merge first pass at parallel rendering to main

diff --git a/Engine/Source/Runtime/Renderer/Private/PostProcess/PostProcessMobile.cpp b/Engine/Source/Runtime/Renderer/Private/PostProcess/PostProcessMobile.cpp
--- a/Engine/Source/Runtime/Renderer/Private/PostProcess/PostProcessMobile.cpp
+++ b/Engine/Source/Runtime/Renderer/Private/PostProcess/PostProcessMobile.cpp
@@ -39,13 +39,13 @@ public:
 		PostprocessParameter.Bind(Initializer.ParameterMap);
 	}
 
-	void SetVS(const FRenderingCompositePassContext& Context)
+	void SetVS(FRHICommandList* RHICmdList, const FRenderingCompositePassContext& Context)
 	{
 		const FVertexShaderRHIParamRef ShaderRHI = GetVertexShader();
 
-		FGlobalShader::SetParameters(ShaderRHI, Context.View);
+		FGlobalShader::SetParameters(RHICmdList, ShaderRHI, Context.View);
 
-		PostprocessParameter.SetVS(ShaderRHI, Context, TStaticSamplerState<SF_Bilinear,AM_Clamp,AM_Clamp,AM_Clamp>::GetRHI());
+		PostprocessParameter.SetVS(RHICmdList, ShaderRHI, Context, TStaticSamplerState<SF_Bilinear,AM_Clamp,AM_Clamp,AM_Clamp>::GetRHI());
 	}
 
 	// FShader interface.
@@ -92,20 +92,20 @@ public:
 		BloomThreshold.Bind(Initializer.ParameterMap, TEXT("BloomThreshold"));
 	}
 
-	void SetPS(const FRenderingCompositePassContext& Context)
+	void SetPS(FRHICommandList* RHICmdList, const FRenderingCompositePassContext& Context)
 	{
 		const FPixelShaderRHIParamRef ShaderRHI = GetPixelShader();
 
 		const FPostProcessSettings& Settings = Context.View.FinalPostProcessSettings;
 
-		FGlobalShader::SetParameters(ShaderRHI, Context.View);
+		FGlobalShader::SetParameters(RHICmdList, ShaderRHI, Context.View);
 
-		PostprocessParameter.SetPS(ShaderRHI, Context, TStaticSamplerState<SF_Bilinear,AM_Clamp,AM_Clamp,AM_Clamp>::GetRHI());
+		PostprocessParameter.SetPS(RHICmdList, ShaderRHI, Context, TStaticSamplerState<SF_Bilinear,AM_Clamp,AM_Clamp,AM_Clamp>::GetRHI());
 
 		float ExposureScale = FRCPassPostProcessEyeAdaptation::ComputeExposureScaleValue(Context.View);
 
 		FVector4 BloomThresholdValue(Settings.BloomThreshold, 0, 0, ExposureScale);
-		SetShaderValue(ShaderRHI, BloomThreshold, BloomThresholdValue);
+		SetShaderValue(RHICmdList, ShaderRHI, BloomThreshold, BloomThresholdValue);
 	}
 
 	// FShader interface.
@@ -129,7 +129,7 @@ IMPLEMENT_SHADER_TYPE(template<>,FPostProcessBloomSetupPS_ES2_2,TEXT("PostProces
 IMPLEMENT_SHADER_TYPE(template<>,FPostProcessBloomSetupPS_ES2_3,TEXT("PostProcessMobile"),TEXT("BloomPS_ES2"),SF_Pixel);
 
 template <uint32 UseSunDof>
-static void BloomSetup_SetShader(const FRenderingCompositePassContext& Context)
+static void BloomSetup_SetShader(FRHICommandList* RHICmdList, const FRenderingCompositePassContext& Context)
 {
 	TShaderMapRef<FPostProcessBloomSetupVS_ES2> VertexShader(GetGlobalShaderMap());
 	TShaderMapRef<FPostProcessBloomSetupPS_ES2<UseSunDof> > PixelShader(GetGlobalShaderMap());
@@ -138,11 +138,11 @@ static void BloomSetup_SetShader(const FRenderingCompositePassContext& Context)
 
 	SetGlobalBoundShaderState(BoundShaderState, GFilterVertexDeclaration.VertexDeclarationRHI, *VertexShader, *PixelShader);
 
-	VertexShader->SetVS(Context);
-	PixelShader->SetPS(Context);
+	VertexShader->SetVS(RHICmdList, Context);
+	PixelShader->SetPS(RHICmdList, Context);
 }
 
-void FRCPassPostProcessBloomSetupES2::SetShader(const FRenderingCompositePassContext& Context)
+void FRCPassPostProcessBloomSetupES2::SetShader(FRHICommandList* RHICmdList, const FRenderingCompositePassContext& Context)
 {
 	const FSceneView& View = Context.View;
 	uint32 UseSun = Context.View.bLightShaftUse ? 1 : 0;
@@ -151,10 +151,10 @@ void FRCPassPostProcessBloomSetupES2::SetShader(const FRenderingCompositePassCon
 
 	switch(UseSunDof)
 	{
-		case 0: BloomSetup_SetShader<0>(Context); break;
-		case 1: BloomSetup_SetShader<1>(Context); break;
-		case 2: BloomSetup_SetShader<2>(Context); break;
-		case 3: BloomSetup_SetShader<3>(Context); break;
+		case 0: BloomSetup_SetShader<0>(RHICmdList, Context); break;
+		case 1: BloomSetup_SetShader<1>(RHICmdList, Context); break;
+		case 2: BloomSetup_SetShader<2>(RHICmdList, Context); break;
+		case 3: BloomSetup_SetShader<3>(RHICmdList, Context); break;
 	}
 }
 
@@ -206,7 +206,9 @@ void FRCPassPostProcessBloomSetupES2::Process(FRenderingCompositePassContext& Co
 	RHISetRasterizerState(TStaticRasterizerState<>::GetRHI());
 	RHISetDepthStencilState(TStaticDepthStencilState<false,CF_Always>::GetRHI());
 
-	SetShader(Context);
+	//@todo-rco: RHIPacketList
+	FRHICommandList* RHICmdList = nullptr;
+	SetShader(RHICmdList, Context);
 
 	TShaderMapRef<FPostProcessBloomSetupVS_ES2> VertexShader(GetGlobalShaderMap());
 
@@ -269,13 +271,13 @@ public:
 		PostprocessParameter.Bind(Initializer.ParameterMap);
 	}
 
-	void SetVS(const FRenderingCompositePassContext& Context)
+	void SetVS(FRHICommandList* RHICmdList, const FRenderingCompositePassContext& Context)
 	{
 		const FVertexShaderRHIParamRef ShaderRHI = GetVertexShader();
 
-		FGlobalShader::SetParameters(ShaderRHI, Context.View);
+		FGlobalShader::SetParameters(RHICmdList, ShaderRHI, Context.View);
 
-		PostprocessParameter.SetVS(ShaderRHI, Context, TStaticSamplerState<SF_Bilinear,AM_Clamp,AM_Clamp,AM_Clamp>::GetRHI());
+		PostprocessParameter.SetVS(RHICmdList, ShaderRHI, Context, TStaticSamplerState<SF_Bilinear,AM_Clamp,AM_Clamp,AM_Clamp>::GetRHI());
 	}
 
 	// FShader interface.
@@ -315,20 +317,20 @@ public:
 		BloomThreshold.Bind(Initializer.ParameterMap, TEXT("BloomThreshold"));
 	}
 
-	void SetPS(const FRenderingCompositePassContext& Context)
+	void SetPS(FRHICommandList* RHICmdList, const FRenderingCompositePassContext& Context)
 	{
 		const FPixelShaderRHIParamRef ShaderRHI = GetPixelShader();
 
 		const FPostProcessSettings& Settings = Context.View.FinalPostProcessSettings;
 
-		FGlobalShader::SetParameters(ShaderRHI, Context.View);
+		FGlobalShader::SetParameters(RHICmdList, ShaderRHI, Context.View);
 
-		PostprocessParameter.SetPS(ShaderRHI, Context, TStaticSamplerState<SF_Bilinear,AM_Clamp,AM_Clamp,AM_Clamp>::GetRHI());
+		PostprocessParameter.SetPS(RHICmdList, ShaderRHI, Context, TStaticSamplerState<SF_Bilinear,AM_Clamp,AM_Clamp,AM_Clamp>::GetRHI());
 
 		float ExposureScale = FRCPassPostProcessEyeAdaptation::ComputeExposureScaleValue(Context.View);
 
 		FVector4 BloomThresholdValue(Settings.BloomThreshold, 0, 0, ExposureScale);
-		SetShaderValue(ShaderRHI, BloomThreshold, BloomThresholdValue);
+		SetShaderValue(RHICmdList, ShaderRHI, BloomThreshold, BloomThresholdValue);
 	}
 
 	// FShader interface.
@@ -345,7 +347,7 @@ IMPLEMENT_SHADER_TYPE(,FPostProcessBloomSetupSmallVS_ES2,TEXT("PostProcessMobile
 IMPLEMENT_SHADER_TYPE(,FPostProcessBloomSetupSmallPS_ES2,TEXT("PostProcessMobile"),TEXT("BloomSmallPS_ES2"),SF_Pixel);
 
 
-void FRCPassPostProcessBloomSetupSmallES2::SetShader(const FRenderingCompositePassContext& Context)
+void FRCPassPostProcessBloomSetupSmallES2::SetShader(FRHICommandList* RHICmdList, const FRenderingCompositePassContext& Context)
 {
 	TShaderMapRef<FPostProcessBloomSetupSmallVS_ES2> VertexShader(GetGlobalShaderMap());
 	TShaderMapRef<FPostProcessBloomSetupSmallPS_ES2> PixelShader(GetGlobalShaderMap());
@@ -354,8 +356,8 @@ void FRCPassPostProcessBloomSetupSmallES2::SetShader(const FRenderingCompositePa
 
 	SetGlobalBoundShaderState(BoundShaderState, GFilterVertexDeclaration.VertexDeclarationRHI, *VertexShader, *PixelShader);
 
-	VertexShader->SetVS(Context);
-	PixelShader->SetPS(Context);
+	VertexShader->SetVS(RHICmdList, Context);
+	PixelShader->SetPS(RHICmdList, Context);
 }
 
 void FRCPassPostProcessBloomSetupSmallES2::Process(FRenderingCompositePassContext& Context)
@@ -406,7 +408,9 @@ void FRCPassPostProcessBloomSetupSmallES2::Process(FRenderingCompositePassContex
 	RHISetRasterizerState(TStaticRasterizerState<>::GetRHI());
 	RHISetDepthStencilState(TStaticDepthStencilState<false,CF_Always>::GetRHI());
 
-	SetShader(Context);
+	//@todo-rco: RHIPacketList
+	FRHICommandList* RHICmdList = nullptr;
+	SetShader(RHICmdList, Context);
 
 	TShaderMapRef<FPostProcessBloomSetupSmallVS_ES2> VertexShader(GetGlobalShaderMap());
 	DrawRectangle(
@@ -473,12 +477,12 @@ public:
 		PostprocessParameter.Bind(Initializer.ParameterMap);
 	}
 
-	void SetPS(const FRenderingCompositePassContext& Context)
+	void SetPS(FRHICommandList* RHICmdList, const FRenderingCompositePassContext& Context)
 	{
 		const FPixelShaderRHIParamRef ShaderRHI = GetPixelShader();
 		const FPostProcessSettings& Settings = Context.View.FinalPostProcessSettings;
-		FGlobalShader::SetParameters(ShaderRHI, Context.View);
-		PostprocessParameter.SetPS(ShaderRHI, Context, TStaticSamplerState<SF_Bilinear,AM_Clamp,AM_Clamp,AM_Clamp>::GetRHI());
+		FGlobalShader::SetParameters(RHICmdList, ShaderRHI, Context.View);
+		PostprocessParameter.SetPS(RHICmdList, ShaderRHI, Context, TStaticSamplerState<SF_Bilinear,AM_Clamp,AM_Clamp,AM_Clamp>::GetRHI());
 	}
 
 	virtual bool Serialize(FArchive& Ar)
@@ -514,12 +518,12 @@ public:
 		BloomDownScale.Bind(Initializer.ParameterMap, TEXT("BloomDownScale"));
 	}
 
-	void SetVS(const FRenderingCompositePassContext& Context, float InScale)
+	void SetVS(FRHICommandList* RHICmdList, const FRenderingCompositePassContext& Context, float InScale)
 	{
 		const FVertexShaderRHIParamRef ShaderRHI = GetVertexShader();
-		FGlobalShader::SetParameters(ShaderRHI, Context.View);
-		PostprocessParameter.SetVS(ShaderRHI, Context, TStaticSamplerState<SF_Bilinear,AM_Clamp,AM_Clamp,AM_Clamp>::GetRHI());
-		SetShaderValue(ShaderRHI, BloomDownScale, InScale);
+		FGlobalShader::SetParameters(RHICmdList, ShaderRHI, Context.View);
+		PostprocessParameter.SetVS(RHICmdList, ShaderRHI, Context, TStaticSamplerState<SF_Bilinear,AM_Clamp,AM_Clamp,AM_Clamp>::GetRHI());
+		SetShaderValue(RHICmdList, ShaderRHI, BloomDownScale, InScale);
 	}
 
 	virtual bool Serialize(FArchive& Ar)
@@ -565,8 +569,10 @@ void FRCPassPostProcessBloomDownES2::Process(FRenderingCompositePassContext& Con
 
 	SetGlobalBoundShaderState(BoundShaderState, GFilterVertexDeclaration.VertexDeclarationRHI, *VertexShader, *PixelShader);
 
-	VertexShader->SetVS(Context, Scale);
-	PixelShader->SetPS(Context);
+	//@todo-rco: RHIPacketList
+	FRHICommandList* RHICmdList = nullptr;
+	VertexShader->SetVS(RHICmdList, Context, Scale);
+	PixelShader->SetPS(RHICmdList, Context);
 
 	FIntPoint SrcDstSize = PrePostSourceViewportSize/2;
 
@@ -637,14 +643,14 @@ public:
 		TintB.Bind(Initializer.ParameterMap, TEXT("BloomTintB"));
 	}
 
-	void SetPS(const FRenderingCompositePassContext& Context, FVector4& InTintA, FVector4& InTintB)
+	void SetPS(FRHICommandList* RHICmdList, const FRenderingCompositePassContext& Context, FVector4& InTintA, FVector4& InTintB)
 	{
 		const FPixelShaderRHIParamRef ShaderRHI = GetPixelShader();
 		const FPostProcessSettings& Settings = Context.View.FinalPostProcessSettings;
-		FGlobalShader::SetParameters(ShaderRHI, Context.View);
-		PostprocessParameter.SetPS(ShaderRHI, Context, TStaticSamplerState<SF_Bilinear,AM_Clamp,AM_Clamp,AM_Clamp>::GetRHI());
-		SetShaderValue(ShaderRHI, TintA, InTintA);
-		SetShaderValue(ShaderRHI, TintB, InTintB);
+		FGlobalShader::SetParameters(RHICmdList, ShaderRHI, Context.View);
+		PostprocessParameter.SetPS(RHICmdList, ShaderRHI, Context, TStaticSamplerState<SF_Bilinear,AM_Clamp,AM_Clamp,AM_Clamp>::GetRHI());
+		SetShaderValue(RHICmdList, ShaderRHI, TintA, InTintA);
+		SetShaderValue(RHICmdList, ShaderRHI, TintB, InTintB);
 	}
 
 	virtual bool Serialize(FArchive& Ar)
@@ -680,12 +686,12 @@ public:
 		BloomUpScales.Bind(Initializer.ParameterMap, TEXT("BloomUpScales"));
 	}
 
-	void SetVS(const FRenderingCompositePassContext& Context, FVector2D InScale)
+	void SetVS(FRHICommandList* RHICmdList, const FRenderingCompositePassContext& Context, FVector2D InScale)
 	{
 		const FVertexShaderRHIParamRef ShaderRHI = GetVertexShader();
-		FGlobalShader::SetParameters(ShaderRHI, Context.View);
-		PostprocessParameter.SetVS(ShaderRHI, Context, TStaticSamplerState<SF_Bilinear,AM_Clamp,AM_Clamp,AM_Clamp>::GetRHI());
-		SetShaderValue(ShaderRHI, BloomUpScales, InScale);
+		FGlobalShader::SetParameters(RHICmdList, ShaderRHI, Context.View);
+		PostprocessParameter.SetVS(RHICmdList, ShaderRHI, Context, TStaticSamplerState<SF_Bilinear,AM_Clamp,AM_Clamp,AM_Clamp>::GetRHI());
+		SetShaderValue(RHICmdList, ShaderRHI, BloomUpScales, InScale);
 	}
 
 	virtual bool Serialize(FArchive& Ar)
@@ -731,11 +737,13 @@ void FRCPassPostProcessBloomUpES2::Process(FRenderingCompositePassContext& Conte
 
 	SetGlobalBoundShaderState(BoundShaderState, GFilterVertexDeclaration.VertexDeclarationRHI, *VertexShader, *PixelShader);
 
+	//@todo-rco: RHIPacketList
+	FRHICommandList* RHICmdList = nullptr;
 	// The 1/8 factor is because bloom is using 8 taps in the filter.
-	VertexShader->SetVS(Context, FVector2D(ScaleAB.X, ScaleAB.Y));
+	VertexShader->SetVS(RHICmdList, Context, FVector2D(ScaleAB.X, ScaleAB.Y));
 	FVector4 TintAScaled = TintA * (1.0f/8.0f);
 	FVector4 TintBScaled = TintB * (1.0f/8.0f);
-	PixelShader->SetPS(Context, TintAScaled, TintBScaled);
+	PixelShader->SetPS(RHICmdList, Context, TintAScaled, TintBScaled);
 
 	FIntPoint SrcDstSize = PrePostSourceViewportSize;
 
@@ -808,19 +816,19 @@ public:
 		SunColorApertureDiv2Parameter.Bind(Initializer.ParameterMap, TEXT("SunColorApertureDiv2"));
 	}
 
-	void SetPS(const FRenderingCompositePassContext& Context)
+	void SetPS(FRHICommandList* RHICmdList, const FRenderingCompositePassContext& Context)
 	{
 		const FPixelShaderRHIParamRef ShaderRHI = GetPixelShader();
 		const FPostProcessSettings& Settings = Context.View.FinalPostProcessSettings;
-		FGlobalShader::SetParameters(ShaderRHI, Context.View);
-		PostprocessParameter.SetPS(ShaderRHI, Context, TStaticSamplerState<SF_Bilinear,AM_Clamp,AM_Clamp,AM_Clamp>::GetRHI());
+		FGlobalShader::SetParameters(RHICmdList, ShaderRHI, Context.View);
+		PostprocessParameter.SetPS(RHICmdList, ShaderRHI, Context, TStaticSamplerState<SF_Bilinear,AM_Clamp,AM_Clamp,AM_Clamp>::GetRHI());
 
 		FVector4 SunColorApertureDiv2;
 		SunColorApertureDiv2.X = Context.View.LightShaftColorMask.R;
 		SunColorApertureDiv2.Y = Context.View.LightShaftColorMask.G;
 		SunColorApertureDiv2.Z = Context.View.LightShaftColorMask.B;
 		SunColorApertureDiv2.W = Context.View.FinalPostProcessSettings.DepthOfFieldScale * 0.5f;
-		SetShaderValue(ShaderRHI, SunColorApertureDiv2Parameter, SunColorApertureDiv2);
+		SetShaderValue(RHICmdList, ShaderRHI, SunColorApertureDiv2Parameter, SunColorApertureDiv2);
 	}
 
 	virtual bool Serialize(FArchive& Ar)
@@ -869,11 +877,11 @@ public:
 		PostprocessParameter.Bind(Initializer.ParameterMap);
 	}
 
-	void SetVS(const FRenderingCompositePassContext& Context)
+	void SetVS(FRHICommandList* RHICmdList, const FRenderingCompositePassContext& Context)
 	{
 		const FVertexShaderRHIParamRef ShaderRHI = GetVertexShader();
-		FGlobalShader::SetParameters(ShaderRHI, Context.View);
-		PostprocessParameter.SetVS(ShaderRHI, Context, TStaticSamplerState<SF_Bilinear,AM_Clamp,AM_Clamp,AM_Clamp>::GetRHI());
+		FGlobalShader::SetParameters(RHICmdList, ShaderRHI, Context.View);
+		PostprocessParameter.SetVS(RHICmdList, ShaderRHI, Context, TStaticSamplerState<SF_Bilinear,AM_Clamp,AM_Clamp,AM_Clamp>::GetRHI());
 	}
 
 	virtual bool Serialize(FArchive& Ar)
@@ -887,7 +895,7 @@ public:
 IMPLEMENT_SHADER_TYPE(,FPostProcessSunMaskVS_ES2,TEXT("PostProcessMobile"),TEXT("SunMaskVS_ES2"),SF_Vertex);
 
 template <uint32 UseFetchSunDof>
-static void SunMask_SetShader(const FRenderingCompositePassContext& Context)
+static void SunMask_SetShader(FRHICommandList* RHICmdList, const FRenderingCompositePassContext& Context)
 {
 	TShaderMapRef<FPostProcessSunMaskVS_ES2> VertexShader(GetGlobalShaderMap());
 	TShaderMapRef<FPostProcessSunMaskPS_ES2<UseFetchSunDof> > PixelShader(GetGlobalShaderMap());
@@ -896,11 +904,11 @@ static void SunMask_SetShader(const FRenderingCompositePassContext& Context)
 
 	SetGlobalBoundShaderState(BoundShaderState, GFilterVertexDeclaration.VertexDeclarationRHI, *VertexShader, *PixelShader);
 
-	VertexShader->SetVS(Context);
-	PixelShader->SetPS(Context);
+	VertexShader->SetVS(RHICmdList, Context);
+	PixelShader->SetPS(RHICmdList, Context);
 }
 
-void FRCPassPostProcessSunMaskES2::SetShader(const FRenderingCompositePassContext& Context)
+void FRCPassPostProcessSunMaskES2::SetShader(FRHICommandList* RHICmdList, const FRenderingCompositePassContext& Context)
 {
 	const FSceneView& View = Context.View;
 	uint32 UseSun = Context.View.bLightShaftUse ? 1 : 0;
@@ -910,14 +918,14 @@ void FRCPassPostProcessSunMaskES2::SetShader(const FRenderingCompositePassContex
 
 	switch(UseFetchSunDof)
 	{
-		case 0: SunMask_SetShader<0>(Context); break;
-		case 1: SunMask_SetShader<1>(Context); break;
-		case 2: SunMask_SetShader<2>(Context); break;
-		case 3: SunMask_SetShader<3>(Context); break;
-		case 4: SunMask_SetShader<4>(Context); break;
-		case 5: SunMask_SetShader<5>(Context); break;
-		case 6: SunMask_SetShader<6>(Context); break;
-		case 7: SunMask_SetShader<7>(Context); break;
+		case 0: SunMask_SetShader<0>(RHICmdList, Context); break;
+		case 1: SunMask_SetShader<1>(RHICmdList, Context); break;
+		case 2: SunMask_SetShader<2>(RHICmdList, Context); break;
+		case 3: SunMask_SetShader<3>(RHICmdList, Context); break;
+		case 4: SunMask_SetShader<4>(RHICmdList, Context); break;
+		case 5: SunMask_SetShader<5>(RHICmdList, Context); break;
+		case 6: SunMask_SetShader<6>(RHICmdList, Context); break;
+		case 7: SunMask_SetShader<7>(RHICmdList, Context); break;
 	}
 }
 
@@ -942,8 +950,10 @@ void FRCPassPostProcessSunMaskES2::Process(FRenderingCompositePassContext& Conte
 	FIntRect SrcRect;
 	const FSceneView& View = Context.View;
 
-	TShaderMapRef<FPostProcessSunMaskVS_ES2> VertexShader(GetGlobalShaderMap());
+	//@todo-rco: RHIPacketList
+	FRHICommandList* RHICmdList = nullptr;
 
+	TShaderMapRef<FPostProcessSunMaskVS_ES2> VertexShader(GetGlobalShaderMap());
 	if(bOnChip)
 	{
 		SrcSize = DstSize;
@@ -955,7 +965,7 @@ void FRCPassPostProcessSunMaskES2::Process(FRenderingCompositePassContext& Conte
 		RHISetRasterizerState(TStaticRasterizerState<>::GetRHI());
 		RHISetDepthStencilState(TStaticDepthStencilState<false,CF_Always>::GetRHI());
 
-		SetShader(Context);
+		SetShader(RHICmdList, Context);
 
 		DrawRectangle(
 			0, 0,
@@ -987,7 +997,7 @@ void FRCPassPostProcessSunMaskES2::Process(FRenderingCompositePassContext& Conte
 		RHISetRasterizerState(TStaticRasterizerState<>::GetRHI());
 		RHISetDepthStencilState(TStaticDepthStencilState<false,CF_Always>::GetRHI());
 
-		SetShader(Context);
+		SetShader(RHICmdList, Context);
 
 		DrawRectangle(
 			0, 0,
@@ -1054,12 +1064,12 @@ public:
 		PostprocessParameter.Bind(Initializer.ParameterMap);
 	}
 
-	void SetPS(const FRenderingCompositePassContext& Context)
+	void SetPS(FRHICommandList* RHICmdList, const FRenderingCompositePassContext& Context)
 	{
 		const FPixelShaderRHIParamRef ShaderRHI = GetPixelShader();
 		const FPostProcessSettings& Settings = Context.View.FinalPostProcessSettings;
-		FGlobalShader::SetParameters(ShaderRHI, Context.View);
-		PostprocessParameter.SetPS(ShaderRHI, Context, TStaticSamplerState<SF_Bilinear,AM_Clamp,AM_Clamp,AM_Clamp>::GetRHI());
+		FGlobalShader::SetParameters(RHICmdList, ShaderRHI, Context.View);
+		PostprocessParameter.SetPS(RHICmdList, ShaderRHI, Context, TStaticSamplerState<SF_Bilinear,AM_Clamp,AM_Clamp,AM_Clamp>::GetRHI());
 	}
 
 	virtual bool Serialize(FArchive& Ar)
@@ -1097,13 +1107,13 @@ public:
 		LightShaftCenter.Bind(Initializer.ParameterMap, TEXT("LightShaftCenter"));
 	}
 
-	void SetVS(const FRenderingCompositePassContext& Context)
+	void SetVS(FRHICommandList* RHICmdList, const FRenderingCompositePassContext& Context)
 	{
 		const FVertexShaderRHIParamRef ShaderRHI = GetVertexShader();
-		FGlobalShader::SetParameters(ShaderRHI, Context.View);
-		PostprocessParameter.SetVS(ShaderRHI, Context, TStaticSamplerState<SF_Bilinear,AM_Clamp,AM_Clamp,AM_Clamp>::GetRHI());
+		FGlobalShader::SetParameters(RHICmdList, ShaderRHI, Context.View);
+		PostprocessParameter.SetVS(RHICmdList, ShaderRHI, Context, TStaticSamplerState<SF_Bilinear,AM_Clamp,AM_Clamp,AM_Clamp>::GetRHI());
 
-		SetShaderValue(ShaderRHI, LightShaftCenter, Context.View.LightShaftCenter);
+		SetShaderValue(RHICmdList, ShaderRHI, LightShaftCenter, Context.View.LightShaftCenter);
 	}
 
 	virtual bool Serialize(FArchive& Ar)
@@ -1117,7 +1127,7 @@ public:
 IMPLEMENT_SHADER_TYPE(,FPostProcessSunAlphaVS_ES2,TEXT("PostProcessMobile"),TEXT("SunAlphaVS_ES2"),SF_Vertex);
 
 template <uint32 UseDof>
-static void SunAlpha_SetShader(const FRenderingCompositePassContext& Context)
+static void SunAlpha_SetShader(FRHICommandList* RHICmdList, const FRenderingCompositePassContext& Context)
 {
 	TShaderMapRef<FPostProcessSunAlphaVS_ES2> VertexShader(GetGlobalShaderMap());
 	TShaderMapRef<FPostProcessSunAlphaPS_ES2<UseDof> > PixelShader(GetGlobalShaderMap());
@@ -1126,19 +1136,19 @@ static void SunAlpha_SetShader(const FRenderingCompositePassContext& Context)
 
 	SetGlobalBoundShaderState(BoundShaderState, GFilterVertexDeclaration.VertexDeclarationRHI, *VertexShader, *PixelShader);
 
-	VertexShader->SetVS(Context);
-	PixelShader->SetPS(Context);
+	VertexShader->SetVS(RHICmdList, Context);
+	PixelShader->SetPS(RHICmdList, Context);
 }
 
-void FRCPassPostProcessSunAlphaES2::SetShader(const FRenderingCompositePassContext& Context)
+void FRCPassPostProcessSunAlphaES2::SetShader(FRHICommandList* RHICmdList, const FRenderingCompositePassContext& Context)
 {
 	if(Context.View.FinalPostProcessSettings.DepthOfFieldScale > 0.0f)
 	{
-		SunAlpha_SetShader<1>(Context);
+		SunAlpha_SetShader<1>(RHICmdList, Context);
 	}
 	else
 	{
-		SunAlpha_SetShader<0>(Context);
+		SunAlpha_SetShader<0>(RHICmdList, Context);
 	}
 }
 
@@ -1167,7 +1177,9 @@ void FRCPassPostProcessSunAlphaES2::Process(FRenderingCompositePassContext& Cont
 	RHISetRasterizerState(TStaticRasterizerState<>::GetRHI());
 	RHISetDepthStencilState(TStaticDepthStencilState<false,CF_Always>::GetRHI());
 
-	SetShader(Context);
+	//@todo-rco: RHIPacketList
+	FRHICommandList* RHICmdList = nullptr;
+	SetShader(RHICmdList, Context);
 
 	FIntPoint SrcDstSize = PrePostSourceViewportSize / 4;
 	TShaderMapRef<FPostProcessSunAlphaVS_ES2> VertexShader(GetGlobalShaderMap());
@@ -1237,12 +1249,12 @@ public:
 		PostprocessParameter.Bind(Initializer.ParameterMap);
 	}
 
-	void SetPS(const FRenderingCompositePassContext& Context)
+	void SetPS(FRHICommandList* RHICmdList, const FRenderingCompositePassContext& Context)
 	{
 		const FPixelShaderRHIParamRef ShaderRHI = GetPixelShader();
 		const FPostProcessSettings& Settings = Context.View.FinalPostProcessSettings;
-		FGlobalShader::SetParameters(ShaderRHI, Context.View);
-		PostprocessParameter.SetPS(ShaderRHI, Context, TStaticSamplerState<SF_Bilinear,AM_Clamp,AM_Clamp,AM_Clamp>::GetRHI());
+		FGlobalShader::SetParameters(RHICmdList, ShaderRHI, Context.View);
+		PostprocessParameter.SetPS(RHICmdList, ShaderRHI, Context, TStaticSamplerState<SF_Bilinear,AM_Clamp,AM_Clamp,AM_Clamp>::GetRHI());
 	}
 
 	virtual bool Serialize(FArchive& Ar)
@@ -1278,13 +1290,13 @@ public:
 		LightShaftCenter.Bind(Initializer.ParameterMap, TEXT("LightShaftCenter"));
 	}
 
-	void SetVS(const FRenderingCompositePassContext& Context)
+	void SetVS(FRHICommandList* RHICmdList, const FRenderingCompositePassContext& Context)
 	{
 		const FVertexShaderRHIParamRef ShaderRHI = GetVertexShader();
-		FGlobalShader::SetParameters(ShaderRHI, Context.View);
-		PostprocessParameter.SetVS(ShaderRHI, Context, TStaticSamplerState<SF_Bilinear,AM_Clamp,AM_Clamp,AM_Clamp>::GetRHI());
+		FGlobalShader::SetParameters(RHICmdList, ShaderRHI, Context.View);
+		PostprocessParameter.SetVS(RHICmdList, ShaderRHI, Context, TStaticSamplerState<SF_Bilinear,AM_Clamp,AM_Clamp,AM_Clamp>::GetRHI());
 
-		SetShaderValue(ShaderRHI, LightShaftCenter, Context.View.LightShaftCenter);
+		SetShaderValue(RHICmdList, ShaderRHI, LightShaftCenter, Context.View.LightShaftCenter);
 	}
 
 	virtual bool Serialize(FArchive& Ar)
@@ -1330,8 +1342,10 @@ void FRCPassPostProcessSunBlurES2::Process(FRenderingCompositePassContext& Conte
 
 	SetGlobalBoundShaderState(BoundShaderState, GFilterVertexDeclaration.VertexDeclarationRHI, *VertexShader, *PixelShader);
 
-	VertexShader->SetVS(Context);
-	PixelShader->SetPS(Context);
+	//@todo-rco: RHIPacketList
+	FRHICommandList* RHICmdList = nullptr;
+	VertexShader->SetVS(RHICmdList, Context);
+	PixelShader->SetPS(RHICmdList, Context);
 
 	FIntPoint SrcDstSize = PrePostSourceViewportSize / 4;
 
@@ -1408,24 +1422,24 @@ public:
 		BloomColor.Bind(Initializer.ParameterMap, TEXT("BloomColor"));
 	}
 
-	void SetPS(const FRenderingCompositePassContext& Context)
+	void SetPS(FRHICommandList* RHICmdList, const FRenderingCompositePassContext& Context)
 	{
 		const FPixelShaderRHIParamRef ShaderRHI = GetPixelShader();
 		const FPostProcessSettings& Settings = Context.View.FinalPostProcessSettings;
-		FGlobalShader::SetParameters(ShaderRHI, Context.View);
-		PostprocessParameter.SetPS(ShaderRHI, Context, TStaticSamplerState<SF_Bilinear,AM_Clamp,AM_Clamp,AM_Clamp>::GetRHI());
+		FGlobalShader::SetParameters(RHICmdList, ShaderRHI, Context.View);
+		PostprocessParameter.SetPS(RHICmdList, ShaderRHI, Context, TStaticSamplerState<SF_Bilinear,AM_Clamp,AM_Clamp,AM_Clamp>::GetRHI());
 
 		FVector4 SunColorVignetteIntensityParam;
 		SunColorVignetteIntensityParam.X = Context.View.LightShaftColorApply.R;
 		SunColorVignetteIntensityParam.Y = Context.View.LightShaftColorApply.G;
 		SunColorVignetteIntensityParam.Z = Context.View.LightShaftColorApply.B;
 		SunColorVignetteIntensityParam.W = Settings.VignetteIntensity;
-		SetShaderValue(ShaderRHI, SunColorVignetteIntensity, SunColorVignetteIntensityParam);
+		SetShaderValue(RHICmdList, ShaderRHI, SunColorVignetteIntensity, SunColorVignetteIntensityParam);
 
-		SetShaderValue(ShaderRHI, VignetteColor, Context.View.FinalPostProcessSettings.VignetteColor);
+		SetShaderValue(RHICmdList, ShaderRHI, VignetteColor, Context.View.FinalPostProcessSettings.VignetteColor);
 
 		// Scaling Bloom1 by extra factor to match filter area difference between PC default and mobile.
-		SetShaderValue(ShaderRHI, BloomColor, Context.View.FinalPostProcessSettings.Bloom1Tint * Context.View.FinalPostProcessSettings.BloomIntensity * 0.5);
+		SetShaderValue(RHICmdList, ShaderRHI, BloomColor, Context.View.FinalPostProcessSettings.Bloom1Tint * Context.View.FinalPostProcessSettings.BloomIntensity * 0.5);
 	}
 
 	virtual bool Serialize(FArchive& Ar)
@@ -1468,13 +1482,13 @@ public:
 		LightShaftCenter.Bind(Initializer.ParameterMap, TEXT("LightShaftCenter"));
 	}
 
-	void SetVS(const FRenderingCompositePassContext& Context)
+	void SetVS(FRHICommandList* RHICmdList, const FRenderingCompositePassContext& Context)
 	{
 		const FVertexShaderRHIParamRef ShaderRHI = GetVertexShader();
-		FGlobalShader::SetParameters(ShaderRHI, Context.View);
-		PostprocessParameter.SetVS(ShaderRHI, Context, TStaticSamplerState<SF_Bilinear,AM_Clamp,AM_Clamp,AM_Clamp>::GetRHI());
+		FGlobalShader::SetParameters(RHICmdList, ShaderRHI, Context.View);
+		PostprocessParameter.SetVS(RHICmdList, ShaderRHI, Context, TStaticSamplerState<SF_Bilinear,AM_Clamp,AM_Clamp,AM_Clamp>::GetRHI());
 
-		SetShaderValue(ShaderRHI, LightShaftCenter, Context.View.LightShaftCenter);
+		SetShaderValue(RHICmdList, ShaderRHI, LightShaftCenter, Context.View.LightShaftCenter);
 	}
 
 	virtual bool Serialize(FArchive& Ar)
@@ -1490,7 +1504,7 @@ IMPLEMENT_SHADER_TYPE(,FPostProcessSunMergeVS_ES2,TEXT("PostProcessMobile"),TEXT
 
 
 template <uint32 UseSunBloom>
-static void SunMerge_SetShader(const FRenderingCompositePassContext& Context)
+static void SunMerge_SetShader(FRHICommandList* RHICmdList, const FRenderingCompositePassContext& Context)
 {
 	TShaderMapRef<FPostProcessSunMergeVS_ES2> VertexShader(GetGlobalShaderMap());
 	TShaderMapRef<FPostProcessSunMergePS_ES2<UseSunBloom> > PixelShader(GetGlobalShaderMap());
@@ -1499,11 +1513,11 @@ static void SunMerge_SetShader(const FRenderingCompositePassContext& Context)
 
 	SetGlobalBoundShaderState(BoundShaderState, GFilterVertexDeclaration.VertexDeclarationRHI, *VertexShader, *PixelShader);
 
-	VertexShader->SetVS(Context);
-	PixelShader->SetPS(Context);
+	VertexShader->SetVS(RHICmdList, Context);
+	PixelShader->SetPS(RHICmdList, Context);
 }
 
-void FRCPassPostProcessSunMergeES2::SetShader(const FRenderingCompositePassContext& Context)
+void FRCPassPostProcessSunMergeES2::SetShader(FRHICommandList* RHICmdList, const FRenderingCompositePassContext& Context)
 {
 	const FSceneView& View = Context.View;
 	uint32 UseBloom = (View.FinalPostProcessSettings.BloomIntensity > 0.0f) ? 1 : 0;
@@ -1512,10 +1526,10 @@ void FRCPassPostProcessSunMergeES2::SetShader(const FRenderingCompositePassConte
 
 	switch(UseSunBloom)
 	{
-	case 0: SunMerge_SetShader<0>(Context); break;
-	case 1: SunMerge_SetShader<1>(Context); break;
-	case 2: SunMerge_SetShader<2>(Context); break;
-	case 3: SunMerge_SetShader<3>(Context); break;
+	case 0: SunMerge_SetShader<0>(RHICmdList, Context); break;
+	case 1: SunMerge_SetShader<1>(RHICmdList, Context); break;
+	case 2: SunMerge_SetShader<2>(RHICmdList, Context); break;
+	case 3: SunMerge_SetShader<3>(RHICmdList, Context); break;
 	}
 }
 
@@ -1544,7 +1558,9 @@ void FRCPassPostProcessSunMergeES2::Process(FRenderingCompositePassContext& Cont
 	RHISetRasterizerState(TStaticRasterizerState<>::GetRHI());
 	RHISetDepthStencilState(TStaticDepthStencilState<false,CF_Always>::GetRHI());
 
-	SetShader(Context);
+	//@todo-rco: RHIPacketList
+	FRHICommandList* RHICmdList = nullptr;
+	SetShader(RHICmdList, Context);
 
 	FIntPoint SrcDstSize = PrePostSourceViewportSize / 4;
 	TShaderMapRef<FPostProcessSunMergeVS_ES2> VertexShader(GetGlobalShaderMap());
@@ -1631,25 +1647,25 @@ public:
 		BloomColor2.Bind(Initializer.ParameterMap, TEXT("BloomColor2"));
 	}
 
-	void SetPS(const FRenderingCompositePassContext& Context)
+	void SetPS(FRHICommandList* RHICmdList, const FRenderingCompositePassContext& Context)
 	{
 		const FPixelShaderRHIParamRef ShaderRHI = GetPixelShader();
 		const FPostProcessSettings& Settings = Context.View.FinalPostProcessSettings;
-		FGlobalShader::SetParameters(ShaderRHI, Context.View);
-		PostprocessParameter.SetPS(ShaderRHI, Context, TStaticSamplerState<SF_Bilinear,AM_Clamp,AM_Clamp,AM_Clamp>::GetRHI());
+		FGlobalShader::SetParameters(RHICmdList, ShaderRHI, Context.View);
+		PostprocessParameter.SetPS(RHICmdList, ShaderRHI, Context, TStaticSamplerState<SF_Bilinear,AM_Clamp,AM_Clamp,AM_Clamp>::GetRHI());
 
 		FVector4 SunColorVignetteIntensityParam;
 		SunColorVignetteIntensityParam.X = Context.View.LightShaftColorApply.R;
 		SunColorVignetteIntensityParam.Y = Context.View.LightShaftColorApply.G;
 		SunColorVignetteIntensityParam.Z = Context.View.LightShaftColorApply.B;
 		SunColorVignetteIntensityParam.W = Settings.VignetteIntensity;
-		SetShaderValue(ShaderRHI, SunColorVignetteIntensity, SunColorVignetteIntensityParam);
+		SetShaderValue(RHICmdList, ShaderRHI, SunColorVignetteIntensity, SunColorVignetteIntensityParam);
 
-		SetShaderValue(ShaderRHI, VignetteColor, Context.View.FinalPostProcessSettings.VignetteColor);
+		SetShaderValue(RHICmdList, ShaderRHI, VignetteColor, Context.View.FinalPostProcessSettings.VignetteColor);
 
 		// Scaling Bloom1 by extra factor to match filter area difference between PC default and mobile.
-		SetShaderValue(ShaderRHI, BloomColor, Context.View.FinalPostProcessSettings.Bloom1Tint * Context.View.FinalPostProcessSettings.BloomIntensity * 0.5);
-		SetShaderValue(ShaderRHI, BloomColor2, Context.View.FinalPostProcessSettings.Bloom2Tint * Context.View.FinalPostProcessSettings.BloomIntensity);
+		SetShaderValue(RHICmdList, ShaderRHI, BloomColor, Context.View.FinalPostProcessSettings.Bloom1Tint * Context.View.FinalPostProcessSettings.BloomIntensity * 0.5);
+		SetShaderValue(RHICmdList, ShaderRHI, BloomColor2, Context.View.FinalPostProcessSettings.Bloom2Tint * Context.View.FinalPostProcessSettings.BloomIntensity);
 	}
 
 	virtual bool Serialize(FArchive& Ar)
@@ -1682,11 +1698,11 @@ public:
 		PostprocessParameter.Bind(Initializer.ParameterMap);
 	}
 
-	void SetVS(const FRenderingCompositePassContext& Context)
+	void SetVS(FRHICommandList* RHICmdList, const FRenderingCompositePassContext& Context)
 	{
 		const FVertexShaderRHIParamRef ShaderRHI = GetVertexShader();
-		FGlobalShader::SetParameters(ShaderRHI, Context.View);
-		PostprocessParameter.SetVS(ShaderRHI, Context, TStaticSamplerState<SF_Bilinear,AM_Clamp,AM_Clamp,AM_Clamp>::GetRHI());
+		FGlobalShader::SetParameters(RHICmdList, ShaderRHI, Context.View);
+		PostprocessParameter.SetVS(RHICmdList, ShaderRHI, Context, TStaticSamplerState<SF_Bilinear,AM_Clamp,AM_Clamp,AM_Clamp>::GetRHI());
 	}
 
 	virtual bool Serialize(FArchive& Ar)
@@ -1699,7 +1715,7 @@ public:
 
 IMPLEMENT_SHADER_TYPE(,FPostProcessSunMergeSmallVS_ES2,TEXT("PostProcessMobile"),TEXT("SunMergeSmallVS_ES2"),SF_Vertex);
 
-void FRCPassPostProcessSunMergeSmallES2::SetShader(const FRenderingCompositePassContext& Context)
+void FRCPassPostProcessSunMergeSmallES2::SetShader(FRHICommandList* RHICmdList, const FRenderingCompositePassContext& Context)
 {
 	TShaderMapRef<FPostProcessSunMergeSmallVS_ES2> VertexShader(GetGlobalShaderMap());
 	TShaderMapRef<FPostProcessSunMergeSmallPS_ES2> PixelShader(GetGlobalShaderMap());
@@ -1708,8 +1724,8 @@ void FRCPassPostProcessSunMergeSmallES2::SetShader(const FRenderingCompositePass
 
 	SetGlobalBoundShaderState(BoundShaderState, GFilterVertexDeclaration.VertexDeclarationRHI, *VertexShader, *PixelShader);
 
-	VertexShader->SetVS(Context);
-	PixelShader->SetPS(Context);
+	VertexShader->SetVS(RHICmdList, Context);
+	PixelShader->SetPS(RHICmdList, Context);
 }
 
 void FRCPassPostProcessSunMergeSmallES2::Process(FRenderingCompositePassContext& Context)
@@ -1737,7 +1753,9 @@ void FRCPassPostProcessSunMergeSmallES2::Process(FRenderingCompositePassContext&
 	RHISetRasterizerState(TStaticRasterizerState<>::GetRHI());
 	RHISetDepthStencilState(TStaticDepthStencilState<false,CF_Always>::GetRHI());
 
-	SetShader(Context);
+	//@todo-rco: RHIPacketList
+	FRHICommandList* RHICmdList = nullptr;
+	SetShader(RHICmdList, Context);
 
 	FIntPoint SrcDstSize = PrePostSourceViewportSize / 4;
 	TShaderMapRef<FPostProcessSunMergeSmallVS_ES2> VertexShader(GetGlobalShaderMap());
@@ -1817,13 +1835,13 @@ public:
 		PostprocessParameter.Bind(Initializer.ParameterMap);
 	}
 
-	void SetVS(const FRenderingCompositePassContext& Context)
+	void SetVS(FRHICommandList* RHICmdList, const FRenderingCompositePassContext& Context)
 	{
 		const FVertexShaderRHIParamRef ShaderRHI = GetVertexShader();
 
-		FGlobalShader::SetParameters(ShaderRHI, Context.View);
+		FGlobalShader::SetParameters(RHICmdList, ShaderRHI, Context.View);
 
-		PostprocessParameter.SetVS(ShaderRHI, Context, TStaticSamplerState<SF_Point,AM_Clamp,AM_Clamp,AM_Clamp>::GetRHI());
+		PostprocessParameter.SetVS(RHICmdList, ShaderRHI, Context, TStaticSamplerState<SF_Point,AM_Clamp,AM_Clamp,AM_Clamp>::GetRHI());
 	}
 
 	// FShader interface.
@@ -1862,12 +1880,12 @@ public:
 		PostprocessParameter.Bind(Initializer.ParameterMap);
 	}
 
-	void SetPS(const FRenderingCompositePassContext& Context)
+	void SetPS(FRHICommandList* RHICmdList, const FRenderingCompositePassContext& Context)
 	{
 		const FPixelShaderRHIParamRef ShaderRHI = GetPixelShader();
 		const FPostProcessSettings& Settings = Context.View.FinalPostProcessSettings;
-		FGlobalShader::SetParameters(ShaderRHI, Context.View);
-		PostprocessParameter.SetPS(ShaderRHI, Context, TStaticSamplerState<SF_Point,AM_Clamp,AM_Clamp,AM_Clamp>::GetRHI());
+		FGlobalShader::SetParameters(RHICmdList, ShaderRHI, Context.View);
+		PostprocessParameter.SetPS(RHICmdList, ShaderRHI, Context, TStaticSamplerState<SF_Point,AM_Clamp,AM_Clamp,AM_Clamp>::GetRHI());
 	}
 
 	virtual bool Serialize(FArchive& Ar)
@@ -1886,7 +1904,7 @@ IMPLEMENT_SHADER_TYPE(template<>,FPostProcessDofDownPS_ES2_0,TEXT("PostProcessMo
 IMPLEMENT_SHADER_TYPE(template<>,FPostProcessDofDownPS_ES2_1,TEXT("PostProcessMobile"),TEXT("DofDownPS_ES2"),SF_Pixel);
 
 template <uint32 UseSun>
-static void DofDown_SetShader(const FRenderingCompositePassContext& Context)
+static void DofDown_SetShader(FRHICommandList* RHICmdList, const FRenderingCompositePassContext& Context)
 {
 	TShaderMapRef<FPostProcessDofDownVS_ES2> VertexShader(GetGlobalShaderMap());
 	TShaderMapRef<FPostProcessDofDownPS_ES2<UseSun> > PixelShader(GetGlobalShaderMap());
@@ -1895,20 +1913,20 @@ static void DofDown_SetShader(const FRenderingCompositePassContext& Context)
 
 	SetGlobalBoundShaderState(BoundShaderState, GFilterVertexDeclaration.VertexDeclarationRHI, *VertexShader, *PixelShader);
 
-	VertexShader->SetVS(Context);
-	PixelShader->SetPS(Context);
+	VertexShader->SetVS(RHICmdList, Context);
+	PixelShader->SetPS(RHICmdList, Context);
 }
 
-void FRCPassPostProcessDofDownES2::SetShader(const FRenderingCompositePassContext& Context)
+void FRCPassPostProcessDofDownES2::SetShader(FRHICommandList* RHICmdList, const FRenderingCompositePassContext& Context)
 {
 	const FSceneView& View = Context.View;
 	if(Context.View.bLightShaftUse)
 	{
-		DofDown_SetShader<1>(Context);
+		DofDown_SetShader<1>(RHICmdList, Context);
 	}
 	else
 	{
-		DofDown_SetShader<0>(Context);
+		DofDown_SetShader<0>(RHICmdList, Context);
 	}
 }
 
@@ -1960,7 +1978,9 @@ void FRCPassPostProcessDofDownES2::Process(FRenderingCompositePassContext& Conte
 	RHISetRasterizerState(TStaticRasterizerState<>::GetRHI());
 	RHISetDepthStencilState(TStaticDepthStencilState<false,CF_Always>::GetRHI());
 
-	SetShader(Context);
+	//@todo-rco: RHIPacketList
+	FRHICommandList* RHICmdList = nullptr;
+	SetShader(RHICmdList, Context);
 
 	TShaderMapRef<FPostProcessDofDownVS_ES2> VertexShader(GetGlobalShaderMap());
 
@@ -2023,13 +2043,13 @@ public:
 		PostprocessParameter.Bind(Initializer.ParameterMap);
 	}
 
-	void SetVS(const FRenderingCompositePassContext& Context)
+	void SetVS(FRHICommandList* RHICmdList, const FRenderingCompositePassContext& Context)
 	{
 		const FVertexShaderRHIParamRef ShaderRHI = GetVertexShader();
 
-		FGlobalShader::SetParameters(ShaderRHI, Context.View);
+		FGlobalShader::SetParameters(RHICmdList, ShaderRHI, Context.View);
 
-		PostprocessParameter.SetVS(ShaderRHI, Context, TStaticSamplerState<SF_Point,AM_Clamp,AM_Clamp,AM_Clamp>::GetRHI());
+		PostprocessParameter.SetVS(RHICmdList, ShaderRHI, Context, TStaticSamplerState<SF_Point,AM_Clamp,AM_Clamp,AM_Clamp>::GetRHI());
 	}
 
 	// FShader interface.
@@ -2068,12 +2088,12 @@ public:
 		PostprocessParameter.Bind(Initializer.ParameterMap);
 	}
 
-	void SetPS(const FRenderingCompositePassContext& Context)
+	void SetPS(FRHICommandList* RHICmdList, const FRenderingCompositePassContext& Context)
 	{
 		const FPixelShaderRHIParamRef ShaderRHI = GetPixelShader();
 		const FPostProcessSettings& Settings = Context.View.FinalPostProcessSettings;
-		FGlobalShader::SetParameters(ShaderRHI, Context.View);
-		PostprocessParameter.SetPS(ShaderRHI, Context, TStaticSamplerState<SF_Point,AM_Clamp,AM_Clamp,AM_Clamp>::GetRHI());
+		FGlobalShader::SetParameters(RHICmdList, ShaderRHI, Context.View);
+		PostprocessParameter.SetPS(RHICmdList, ShaderRHI, Context, TStaticSamplerState<SF_Point,AM_Clamp,AM_Clamp,AM_Clamp>::GetRHI());
 	}
 
 	virtual bool Serialize(FArchive& Ar)
@@ -2092,7 +2112,7 @@ IMPLEMENT_SHADER_TYPE(template<>,FPostProcessDofNearPS_ES2_0,TEXT("PostProcessMo
 IMPLEMENT_SHADER_TYPE(template<>,FPostProcessDofNearPS_ES2_1,TEXT("PostProcessMobile"),TEXT("DofNearPS_ES2"),SF_Pixel);
 
 template <uint32 UseSun>
-static void DofNear_SetShader(const FRenderingCompositePassContext& Context)
+static void DofNear_SetShader(FRHICommandList* RHICmdList, const FRenderingCompositePassContext& Context)
 {
 	TShaderMapRef<FPostProcessDofNearVS_ES2> VertexShader(GetGlobalShaderMap());
 	TShaderMapRef<FPostProcessDofNearPS_ES2<UseSun> > PixelShader(GetGlobalShaderMap());
@@ -2101,20 +2121,20 @@ static void DofNear_SetShader(const FRenderingCompositePassContext& Context)
 
 	SetGlobalBoundShaderState(BoundShaderState, GFilterVertexDeclaration.VertexDeclarationRHI, *VertexShader, *PixelShader);
 
-	VertexShader->SetVS(Context);
-	PixelShader->SetPS(Context);
+	VertexShader->SetVS(RHICmdList, Context);
+	PixelShader->SetPS(RHICmdList, Context);
 }
 
-void FRCPassPostProcessDofNearES2::SetShader(const FRenderingCompositePassContext& Context)
+void FRCPassPostProcessDofNearES2::SetShader(FRHICommandList* RHICmdList, const FRenderingCompositePassContext& Context)
 {
 	const FSceneView& View = Context.View;
 	if(Context.View.bLightShaftUse)
 	{
-		DofNear_SetShader<1>(Context);
+		DofNear_SetShader<1>(RHICmdList, Context);
 	}
 	else
 	{
-		DofNear_SetShader<0>(Context);
+		DofNear_SetShader<0>(RHICmdList, Context);
 	}
 }
 
@@ -2143,7 +2163,9 @@ void FRCPassPostProcessDofNearES2::Process(FRenderingCompositePassContext& Conte
 	RHISetRasterizerState(TStaticRasterizerState<>::GetRHI());
 	RHISetDepthStencilState(TStaticDepthStencilState<false,CF_Always>::GetRHI());
 
-	SetShader(Context);
+	//@todo-rco: RHIPacketList
+	FRHICommandList* RHICmdList = nullptr;
+	SetShader(RHICmdList, Context);
 
 	FIntPoint SrcDstSize = PrePostSourceViewportSize / 4;
 	TShaderMapRef<FPostProcessDofNearVS_ES2> VertexShader(GetGlobalShaderMap());
@@ -2210,12 +2232,12 @@ public:
 		PostprocessParameter.Bind(Initializer.ParameterMap);
 	}
 
-	void SetPS(const FRenderingCompositePassContext& Context)
+	void SetPS(FRHICommandList* RHICmdList, const FRenderingCompositePassContext& Context)
 	{
 		const FPixelShaderRHIParamRef ShaderRHI = GetPixelShader();
 		const FPostProcessSettings& Settings = Context.View.FinalPostProcessSettings;
-		FGlobalShader::SetParameters(ShaderRHI, Context.View);
-		PostprocessParameter.SetPS(ShaderRHI, Context, TStaticSamplerState<SF_Bilinear,AM_Clamp,AM_Clamp,AM_Clamp>::GetRHI());
+		FGlobalShader::SetParameters(RHICmdList, ShaderRHI, Context.View);
+		PostprocessParameter.SetPS(RHICmdList, ShaderRHI, Context, TStaticSamplerState<SF_Bilinear,AM_Clamp,AM_Clamp,AM_Clamp>::GetRHI());
 	}
 
 	virtual bool Serialize(FArchive& Ar)
@@ -2249,11 +2271,11 @@ public:
 		PostprocessParameter.Bind(Initializer.ParameterMap);
 	}
 
-	void SetVS(const FRenderingCompositePassContext& Context)
+	void SetVS(FRHICommandList* RHICmdList, const FRenderingCompositePassContext& Context)
 	{
 		const FVertexShaderRHIParamRef ShaderRHI = GetVertexShader();
-		FGlobalShader::SetParameters(ShaderRHI, Context.View);
-		PostprocessParameter.SetVS(ShaderRHI, Context, TStaticSamplerState<SF_Point,AM_Clamp,AM_Clamp,AM_Clamp>::GetRHI());
+		FGlobalShader::SetParameters(RHICmdList, ShaderRHI, Context.View);
+		PostprocessParameter.SetVS(RHICmdList, ShaderRHI, Context, TStaticSamplerState<SF_Point,AM_Clamp,AM_Clamp,AM_Clamp>::GetRHI());
 	}
 
 	virtual bool Serialize(FArchive& Ar)
@@ -2298,8 +2320,10 @@ void FRCPassPostProcessDofBlurES2::Process(FRenderingCompositePassContext& Conte
 	static FGlobalBoundShaderState BoundShaderState;
 	SetGlobalBoundShaderState(BoundShaderState, GFilterVertexDeclaration.VertexDeclarationRHI, *VertexShader, *PixelShader);
 
-	VertexShader->SetVS(Context);
-	PixelShader->SetPS(Context);
+	//@todo-rco: RHIPacketList
+	FRHICommandList* RHICmdList = nullptr;
+	VertexShader->SetVS(RHICmdList, Context);
+	PixelShader->SetPS(RHICmdList, Context);
 
 	FIntPoint SrcDstSize = PrePostSourceViewportSize / 2;
 
@@ -2367,12 +2391,12 @@ public:
 		PostprocessParameter.Bind(Initializer.ParameterMap);
 	}
 
-	void SetPS(const FRenderingCompositePassContext& Context)
+	void SetPS(FRHICommandList* RHICmdList, const FRenderingCompositePassContext& Context)
 	{
 		const FPixelShaderRHIParamRef ShaderRHI = GetPixelShader();
 		const FPostProcessSettings& Settings = Context.View.FinalPostProcessSettings;
-		FGlobalShader::SetParameters(ShaderRHI, Context.View);
-		PostprocessParameter.SetPS(ShaderRHI, Context, TStaticSamplerState<SF_Bilinear,AM_Clamp,AM_Clamp,AM_Clamp>::GetRHI());
+		FGlobalShader::SetParameters(RHICmdList, ShaderRHI, Context.View);
+		PostprocessParameter.SetPS(RHICmdList, ShaderRHI, Context, TStaticSamplerState<SF_Bilinear,AM_Clamp,AM_Clamp,AM_Clamp>::GetRHI());
 	}
 
 	virtual bool Serialize(FArchive& Ar)
@@ -2407,11 +2431,11 @@ public:
 		PostprocessParameter.Bind(Initializer.ParameterMap);
 	}
 
-	void SetVS(const FRenderingCompositePassContext& Context)
+	void SetVS(FRHICommandList* RHICmdList, const FRenderingCompositePassContext& Context)
 	{
 		const FVertexShaderRHIParamRef ShaderRHI = GetVertexShader();
-		FGlobalShader::SetParameters(ShaderRHI, Context.View);
-		PostprocessParameter.SetVS(ShaderRHI, Context, TStaticSamplerState<SF_Bilinear,AM_Clamp,AM_Clamp,AM_Clamp>::GetRHI());
+		FGlobalShader::SetParameters(RHICmdList, ShaderRHI, Context.View);
+		PostprocessParameter.SetVS(RHICmdList, ShaderRHI, Context, TStaticSamplerState<SF_Bilinear,AM_Clamp,AM_Clamp,AM_Clamp>::GetRHI());
 	}
 
 	virtual bool Serialize(FArchive& Ar)
@@ -2426,7 +2450,7 @@ IMPLEMENT_SHADER_TYPE(,FPostProcessSunAvgVS_ES2,TEXT("PostProcessMobile"),TEXT("
 
 
 
-static void SunAvg_SetShader(const FRenderingCompositePassContext& Context)
+static void SunAvg_SetShader(FRHICommandList* RHICmdList, const FRenderingCompositePassContext& Context)
 {
 	TShaderMapRef<FPostProcessSunAvgVS_ES2> VertexShader(GetGlobalShaderMap());
 	TShaderMapRef<FPostProcessSunAvgPS_ES2> PixelShader(GetGlobalShaderMap());
@@ -2435,13 +2459,13 @@ static void SunAvg_SetShader(const FRenderingCompositePassContext& Context)
 
 	SetGlobalBoundShaderState(BoundShaderState, GFilterVertexDeclaration.VertexDeclarationRHI, *VertexShader, *PixelShader);
 
-	VertexShader->SetVS(Context);
-	PixelShader->SetPS(Context);
+	VertexShader->SetVS(RHICmdList, Context);
+	PixelShader->SetPS(RHICmdList, Context);
 }
 
-void FRCPassPostProcessSunAvgES2::SetShader(const FRenderingCompositePassContext& Context)
+void FRCPassPostProcessSunAvgES2::SetShader(FRHICommandList* RHICmdList, const FRenderingCompositePassContext& Context)
 {
-	SunAvg_SetShader(Context);
+	SunAvg_SetShader(RHICmdList, Context);
 }
 
 void FRCPassPostProcessSunAvgES2::Process(FRenderingCompositePassContext& Context)
@@ -2469,7 +2493,9 @@ void FRCPassPostProcessSunAvgES2::Process(FRenderingCompositePassContext& Contex
 	RHISetRasterizerState(TStaticRasterizerState<>::GetRHI());
 	RHISetDepthStencilState(TStaticDepthStencilState<false,CF_Always>::GetRHI());
 
-	SetShader(Context);
+	//@todo-rco: RHIPacketList
+	FRHICommandList* RHICmdList = nullptr;
+	SetShader(RHICmdList, Context);
 
 	FIntPoint SrcDstSize = PrePostSourceViewportSize / 4;
 	TShaderMapRef<FPostProcessSunAvgVS_ES2> VertexShader(GetGlobalShaderMap());
@@ -2539,12 +2565,12 @@ public:
 		AaBlendAmount.Bind(Initializer.ParameterMap, TEXT("AaBlendAmount"));
 	}
 
-	void SetPS(const FRenderingCompositePassContext& Context)
+	void SetPS(FRHICommandList* RHICmdList, const FRenderingCompositePassContext& Context)
 	{
 		const FPixelShaderRHIParamRef ShaderRHI = GetPixelShader();
 		const FPostProcessSettings& Settings = Context.View.FinalPostProcessSettings;
-		FGlobalShader::SetParameters(ShaderRHI, Context.View);
-		PostprocessParameter.SetPS(ShaderRHI, Context, TStaticSamplerState<SF_Bilinear,AM_Clamp,AM_Clamp,AM_Clamp>::GetRHI());
+		FGlobalShader::SetParameters(RHICmdList, ShaderRHI, Context.View);
+		PostprocessParameter.SetPS(RHICmdList, ShaderRHI, Context, TStaticSamplerState<SF_Bilinear,AM_Clamp,AM_Clamp,AM_Clamp>::GetRHI());
 
 		// Compute the blend factor which decides the trade off between ghosting in motion and flicker when not moving.
 		// This works by computing the screen space motion vector of distant point at the center of the screen.
@@ -2607,12 +2633,12 @@ public:
 				BlendAmount = 0.0f;
 			}
 
-			SetShaderValue(ShaderRHI, AaBlendAmount, BlendAmount);
+			SetShaderValue(RHICmdList, ShaderRHI, AaBlendAmount, BlendAmount);
 		}
 		else
 		{
 			float BlendAmount = 0.0;
-			SetShaderValue(ShaderRHI, AaBlendAmount, BlendAmount);
+			SetShaderValue(RHICmdList, ShaderRHI, AaBlendAmount, BlendAmount);
 		}
 	}
 
@@ -2648,11 +2674,11 @@ public:
 		PostprocessParameter.Bind(Initializer.ParameterMap);
 	}
 
-	void SetVS(const FRenderingCompositePassContext& Context)
+	void SetVS(FRHICommandList* RHICmdList, const FRenderingCompositePassContext& Context)
 	{
 		const FVertexShaderRHIParamRef ShaderRHI = GetVertexShader();
-		FGlobalShader::SetParameters(ShaderRHI, Context.View);
-		PostprocessParameter.SetVS(ShaderRHI, Context, TStaticSamplerState<SF_Bilinear,AM_Clamp,AM_Clamp,AM_Clamp>::GetRHI());
+		FGlobalShader::SetParameters(RHICmdList, ShaderRHI, Context.View);
+		PostprocessParameter.SetVS(RHICmdList, ShaderRHI, Context, TStaticSamplerState<SF_Bilinear,AM_Clamp,AM_Clamp,AM_Clamp>::GetRHI());
 	}
 
 	virtual bool Serialize(FArchive& Ar)
@@ -2667,7 +2693,7 @@ IMPLEMENT_SHADER_TYPE(,FPostProcessAaVS_ES2,TEXT("PostProcessMobile"),TEXT("AaVS
 
 
 
-static void Aa_SetShader(const FRenderingCompositePassContext& Context)
+static void Aa_SetShader(FRHICommandList* RHICmdList, const FRenderingCompositePassContext& Context)
 {
 	TShaderMapRef<FPostProcessAaVS_ES2> VertexShader(GetGlobalShaderMap());
 	TShaderMapRef<FPostProcessAaPS_ES2> PixelShader(GetGlobalShaderMap());
@@ -2676,13 +2702,13 @@ static void Aa_SetShader(const FRenderingCompositePassContext& Context)
 
 	SetGlobalBoundShaderState(BoundShaderState, GFilterVertexDeclaration.VertexDeclarationRHI, *VertexShader, *PixelShader);
 
-	VertexShader->SetVS(Context);
-	PixelShader->SetPS(Context);
+	VertexShader->SetVS(RHICmdList, Context);
+	PixelShader->SetPS(RHICmdList, Context);
 }
 
-void FRCPassPostProcessAaES2::SetShader(const FRenderingCompositePassContext& Context)
+void FRCPassPostProcessAaES2::SetShader(FRHICommandList* RHICmdList, const FRenderingCompositePassContext& Context)
 {
-	Aa_SetShader(Context);
+	Aa_SetShader(RHICmdList, Context);
 }
 
 void FRCPassPostProcessAaES2::Process(FRenderingCompositePassContext& Context)
@@ -2710,7 +2736,9 @@ void FRCPassPostProcessAaES2::Process(FRenderingCompositePassContext& Context)
 	RHISetRasterizerState(TStaticRasterizerState<>::GetRHI());
 	RHISetDepthStencilState(TStaticDepthStencilState<false,CF_Always>::GetRHI());
 
-	SetShader(Context);
+	//@todo-rco: RHIPacketList
+	FRHICommandList* RHICmdList = nullptr;
+	SetShader(RHICmdList, Context);
 
 	FIntPoint SrcDstSize = PrePostSourceViewportSize;
 	TShaderMapRef<FPostProcessAaVS_ES2> VertexShader(GetGlobalShaderMap());
diff --git a/Engine/Source/Runtime/Renderer/Private/PostProcess/PostProcessMobile.h b/Engine/Source/Runtime/Renderer/Private/PostProcess/PostProcessMobile.h
--- a/Engine/Source/Runtime/Renderer/Private/PostProcess/PostProcessMobile.h
+++ b/Engine/Source/Runtime/Renderer/Private/PostProcess/PostProcessMobile.h
@@ -18,7 +18,7 @@ public:
 private:
 	FIntPoint PrePostSourceViewportSize;
 	bool bUsedFramebufferFetch;
-	void SetShader(const FRenderingCompositePassContext& Context);
+	void SetShader(FRHICommandList* RHICmdList, const FRenderingCompositePassContext& Context);
 };
 
 class FRCPassPostProcessBloomSetupSmallES2 : public TRenderingCompositePassBase<1, 1>
@@ -31,7 +31,7 @@ public:
 private:
 	FIntPoint PrePostSourceViewportSize;
 	bool bUsedFramebufferFetch;
-	void SetShader(const FRenderingCompositePassContext& Context);
+	void SetShader(FRHICommandList* RHICmdList, const FRenderingCompositePassContext& Context);
 };
 
 class FRCPassPostProcessDofNearES2 : public TRenderingCompositePassBase<1, 1>
@@ -43,7 +43,7 @@ public:
 	virtual void Release() { delete this; }
 private:
 	FIntPoint PrePostSourceViewportSize;
-	void SetShader(const FRenderingCompositePassContext& Context);
+	void SetShader(FRHICommandList* RHICmdList, const FRenderingCompositePassContext& Context);
 };
 
 class FRCPassPostProcessDofDownES2 : public TRenderingCompositePassBase<2, 1>
@@ -56,7 +56,7 @@ public:
 private:
 	FIntPoint PrePostSourceViewportSize;
 	bool bUsedFramebufferFetch;
-	void SetShader(const FRenderingCompositePassContext& Context);
+	void SetShader(FRHICommandList* RHICmdList, const FRenderingCompositePassContext& Context);
 };
 
 class FRCPassPostProcessDofBlurES2 : public TRenderingCompositePassBase<2, 1>
@@ -106,7 +106,7 @@ public:
 private:
 	FIntPoint PrePostSourceViewportSize;
 	bool bOnChip;
-	void SetShader(const FRenderingCompositePassContext& Context);
+	void SetShader(FRHICommandList* RHICmdList, const FRenderingCompositePassContext& Context);
 };
 
 class FRCPassPostProcessSunAlphaES2 : public TRenderingCompositePassBase<1, 1>
@@ -118,7 +118,7 @@ public:
 	virtual void Release() { delete this; }
 private:
 	FIntPoint PrePostSourceViewportSize;
-	void SetShader(const FRenderingCompositePassContext& Context);
+	void SetShader(FRHICommandList* RHICmdList, const FRenderingCompositePassContext& Context);
 };
 
 class FRCPassPostProcessSunBlurES2 : public TRenderingCompositePassBase<1, 1>
@@ -141,7 +141,7 @@ public:
 	virtual void Release() { delete this; }
 private:
 	FIntPoint PrePostSourceViewportSize;
-	void SetShader(const FRenderingCompositePassContext& Context);
+	void SetShader(FRHICommandList* RHICmdList, const FRenderingCompositePassContext& Context);
 };
 
 class FRCPassPostProcessSunMergeSmallES2 : public TRenderingCompositePassBase<2, 1>
@@ -153,7 +153,7 @@ public:
 	virtual void Release() { delete this; }
 private:
 	FIntPoint PrePostSourceViewportSize;
-	void SetShader(const FRenderingCompositePassContext& Context);
+	void SetShader(FRHICommandList* RHICmdList, const FRenderingCompositePassContext& Context);
 };
 
 class FRCPassPostProcessSunAvgES2 : public TRenderingCompositePassBase<2, 1>
@@ -165,7 +165,7 @@ public:
 	virtual void Release() { delete this; }
 private:
 	FIntPoint PrePostSourceViewportSize;
-	void SetShader(const FRenderingCompositePassContext& Context);
+	void SetShader(FRHICommandList* RHICmdList, const FRenderingCompositePassContext& Context);
 };
 
 class FRCPassPostProcessAaES2 : public TRenderingCompositePassBase<2, 1>
@@ -177,6 +177,6 @@ public:
 	virtual void Release() { delete this; }
 private:
 	FIntPoint PrePostSourceViewportSize;
-	void SetShader(const FRenderingCompositePassContext& Context);
+	void SetShader(FRHICommandList* RHICmdList, const FRenderingCompositePassContext& Context);
 };
 
diff --git a/Engine/Source/Runtime/Renderer/Private/PostProcess/PostProcessMorpheus.cpp b/Engine/Source/Runtime/Renderer/Private/PostProcess/PostProcessMorpheus.cpp
--- a/Engine/Source/Runtime/Renderer/Private/PostProcess/PostProcessMorpheus.cpp
+++ b/Engine/Source/Runtime/Renderer/Private/PostProcess/PostProcessMorpheus.cpp
@@ -55,7 +55,7 @@ public:
 	FDeferredPixelShaderParameters DeferredParameters;
 	
 #if NEW_MORPHEUS_DISTORTION
-	// Distortion parameter values	
+	// Distortion parameter values
 	FShaderParameter TextureScale;
 	FShaderParameter TextureOffset;
 	FShaderParameter TextureUVOffset;
@@ -68,7 +68,7 @@ public:
 	const FTexture*		DistortionTexture;
 #endif
 
-	FShaderResourceParameter DistortionTextureParam; 	
+	FShaderResourceParameter DistortionTextureParam; 
 	FShaderResourceParameter DistortionTextureSampler; 
 
 	/** Initialization constructor. */
@@ -99,7 +99,7 @@ public:
 		Scale.Bind(Initializer.ParameterMap, TEXT("Scale"));
 		HMDWarpParam.Bind(Initializer.ParameterMap, TEXT("HMDWarpParam"));
 		CAWarpParam.Bind(Initializer.ParameterMap, TEXT("CAWarpParam"));
-
+	
 		DistortionTextureParam.Bind(Initializer.ParameterMap, TEXT("DistortionTexture"));
 		check(DistortionTextureParam.IsBound());
 
@@ -111,14 +111,14 @@ public:
 	}
 
 
-	void SetPS(const FRenderingCompositePassContext& Context, FIntRect SrcRect, FIntPoint SrcBufferSize, EStereoscopicPass StereoPass, FMatrix& QuadTexTransform)
+	void SetPS(FRHICommandList* RHICmdList, const FRenderingCompositePassContext& Context, FIntRect SrcRect, FIntPoint SrcBufferSize, EStereoscopicPass StereoPass, FMatrix& QuadTexTransform)
 	{
 		const FPixelShaderRHIParamRef ShaderRHI = GetPixelShader();
 		
-		FGlobalShader::SetParameters(ShaderRHI, Context.View);
+		FGlobalShader::SetParameters(RHICmdList, ShaderRHI, Context.View);
 
-		PostprocessParameter.SetPS(ShaderRHI, Context, TStaticSamplerState<SF_Bilinear,AM_Clamp,AM_Clamp,AM_Clamp>::GetRHI());
-		DeferredParameters.Set(ShaderRHI, Context.View);
+		PostprocessParameter.SetPS(RHICmdList, ShaderRHI, Context, TStaticSamplerState<SF_Bilinear,AM_Clamp,AM_Clamp,AM_Clamp>::GetRHI());
+		DeferredParameters.Set(RHICmdList, ShaderRHI, Context.View);
 
 		{
 			check(GEngine->HMDDevice.IsValid());
@@ -146,13 +146,13 @@ public:
 			      
             QuadTexTransform = FMatrix::Identity;            
 #else
-			const FIntPoint SrcSize = SrcRect.Size();
-			const float BufferRatioX = float(SrcSize.X)/float(SrcBufferSize.X);
-			const float BufferRatioY = float(SrcSize.Y)/float(SrcBufferSize.Y);
-			const float w = float(SrcSize.X)/float(SrcBufferSize.X);
-			const float h = float(SrcSize.Y)/float(SrcBufferSize.Y);
-			const float x = float(SrcRect.Min.X)/float(SrcBufferSize.X);
-			const float y = float(SrcRect.Min.Y)/float(SrcBufferSize.Y);
+            const FIntPoint SrcSize = SrcRect.Size();
+            const float BufferRatioX = float(SrcSize.X)/float(SrcBufferSize.X);
+            const float BufferRatioY = float(SrcSize.Y)/float(SrcBufferSize.Y);
+            const float w = float(SrcSize.X)/float(SrcBufferSize.X);
+            const float h = float(SrcSize.Y)/float(SrcBufferSize.Y);
+            const float x = float(SrcRect.Min.X)/float(SrcBufferSize.X);
+            const float y = float(SrcRect.Min.Y)/float(SrcBufferSize.Y);
 
 			float DistortionCenterOffsetX = 0, DistortionCenterOffsetY = 0;
 			GEngine->HMDDevice->GetDistortionCenterOffset(DistortionCenterOffsetX, DistortionCenterOffsetY);
@@ -169,45 +169,45 @@ public:
 				DistortionTexture = GEngine->HMDDevice->GetDistortionTexture();
 				check(DistortionTexture != NULL);
 			}
-			SetTextureParameter(ShaderRHI, DistortionTextureParam, DistortionTextureSampler, TStaticSamplerState<SF_Bilinear,AM_Clamp,AM_Clamp,AM_Clamp>::GetRHI(), DistortionTexture->TextureRHI);
+			SetTextureParameter(RHICmdList, ShaderRHI, DistortionTextureParam, DistortionTextureSampler, TStaticSamplerState<SF_Bilinear,AM_Clamp,AM_Clamp,AM_Clamp>::GetRHI(), DistortionTexture->TextureRHI);
 
 			// Shifts texture coordinates to the center of the distortion function around the center of the lens
 			FVector2D ViewLensCenter;
 			ViewLensCenter.X = x + (0.5f + XCenterOffset) * w;
 			ViewLensCenter.Y = y + (0.5f + YCenterOffset) * h;
-			SetShaderValue(ShaderRHI, LensCenter, ViewLensCenter);
+			SetShaderValue(RHICmdList, ShaderRHI, LensCenter, ViewLensCenter);
 
 			// Texture coordinate for the center of the half scene texture, used to clamp sampling
 			FVector2D ViewScreenCenter;
 			ViewScreenCenter.X = x + w * 0.5f;
 			ViewScreenCenter.Y = y + h * 0.5f;
-			SetShaderValue(ShaderRHI, ScreenCenter, ViewScreenCenter);
-
+			SetShaderValue(RHICmdList, ShaderRHI, ScreenCenter, ViewScreenCenter);
+			
 			// Rescale output (sample) coordinates back to texture range and increase scale to support rendering outside the the screen
 			FVector2D ViewScale;
 			ViewScale.X = (w/2) * ScaleFactor;
 			ViewScale.Y = (h/2) * ScaleFactor * AspectRatio;
-			SetShaderValue(ShaderRHI, Scale, ViewScale);
+			SetShaderValue(RHICmdList, ShaderRHI, Scale, ViewScale);
 
-			// Distortion coefficients
-			FVector4 DistortionValues;
-			GEngine->HMDDevice->GetDistortionWarpValues(DistortionValues);
-			SetShaderValue(ShaderRHI, HMDWarpParam, DistortionValues);
+            // Distortion coefficients
+            FVector4 DistortionValues;
+            GEngine->HMDDevice->GetDistortionWarpValues(DistortionValues);
+			SetShaderValue(RHICmdList, ShaderRHI, HMDWarpParam, DistortionValues);
 
 			// CNN - Morpheus changes
 			// CA correction values
 			FVector4 CAValues;
 			GEngine->HMDDevice->GetChromaAbCorrectionValues(CAValues);
-			SetShaderValue(ShaderRHI, CAWarpParam, CAValues);
+			SetShaderValue(RHICmdList, ShaderRHI, CAWarpParam, CAValues);
 
 			// Rescale the texture coordinates to [-1,1] unit range and corrects aspect ratio
 			FVector2D ViewScaleIn;
-			ViewScaleIn.X = (2/w);
-			ViewScaleIn.Y = (2/h) / AspectRatio;
-
-			QuadTexTransform = FMatrix::Identity;
-			QuadTexTransform *= FTranslationMatrix(FVector(-ViewLensCenter.X * SrcBufferSize.X, -ViewLensCenter.Y * SrcBufferSize.Y, 0));
-			QuadTexTransform *= FMatrix(FPlane(ViewScaleIn.X,0,0,0), FPlane(0,ViewScaleIn.Y,0,0), FPlane(0,0,0,0), FPlane(0,0,0,1));
+            ViewScaleIn.X = (2/w);
+            ViewScaleIn.Y = (2/h) / AspectRatio;
+            
+            QuadTexTransform = FMatrix::Identity;
+            QuadTexTransform *= FTranslationMatrix(FVector(-ViewLensCenter.X * SrcBufferSize.X, -ViewLensCenter.Y * SrcBufferSize.Y, 0));
+            QuadTexTransform *= FMatrix(FPlane(ViewScaleIn.X,0,0,0), FPlane(0,ViewScaleIn.Y,0,0), FPlane(0,0,0,0), FPlane(0,0,0,1));
 #endif
 		}
 	}
@@ -254,14 +254,14 @@ class FPostProcessMorpheusVS : public FGlobalShader
 	FPostProcessMorpheusVS() {}
 
 	/** to have a similar interface as all other shaders */
-	void SetParameters(const FRenderingCompositePassContext& Context)
+	void SetParameters(FRHICommandList* RHICmdList, const FRenderingCompositePassContext& Context)
 	{
-		FGlobalShader::SetParameters(GetVertexShader(), Context.View);
+		FGlobalShader::SetParameters(RHICmdList, GetVertexShader(), Context.View);
 	}
 
-	void SetParameters(const FSceneView& View)
+	void SetParameters(FRHICommandList* RHICmdList, const FSceneView& View)
 	{
-		FGlobalShader::SetParameters(GetVertexShader(), View);
+		FGlobalShader::SetParameters(RHICmdList, GetVertexShader(), View);
 	}
 
 public:
@@ -285,7 +285,7 @@ void FRCPassPostProcessMorpheus::Process(FRenderingCompositePassContext& Context
 	{
 		// input is not hooked up correctly
 		return;
-	}	
+	}
 
 	const FSceneView& View = Context.View;
 	const FSceneViewFamily& ViewFamily = *(View.Family);
@@ -315,8 +315,9 @@ void FRCPassPostProcessMorpheus::Process(FRenderingCompositePassContext& Context
 
     FMatrix QuadTexTransform;
     FMatrix QuadPosTransform = FMatrix::Identity;
-
-	PixelShader->SetPS(Context, SrcRect, SrcSize, View.StereoPass, QuadTexTransform);
+	//@todo-rco: RHIPacketList
+	FRHICommandList* RHICmdList = nullptr;
+	PixelShader->SetPS(RHICmdList, Context, SrcRect, SrcSize, View.StereoPass, QuadTexTransform);
 
 	// Draw a quad mapping scene color to the view's render target
     DrawTransformedRectangle(
diff --git a/Engine/Source/Runtime/Renderer/Private/PostProcess/PostProcessMotionBlur.cpp b/Engine/Source/Runtime/Renderer/Private/PostProcess/PostProcessMotionBlur.cpp
--- a/Engine/Source/Runtime/Renderer/Private/PostProcess/PostProcessMotionBlur.cpp
+++ b/Engine/Source/Runtime/Renderer/Private/PostProcess/PostProcessMotionBlur.cpp
@@ -47,13 +47,13 @@ public:
 		return bShaderHasOutdatedParameters;
 	}
 
-	void SetParameters(const FRenderingCompositePassContext& Context)
+	void SetParameters(FRHICommandList* RHICmdList, const FRenderingCompositePassContext& Context)
 	{
 		const FPixelShaderRHIParamRef ShaderRHI = GetPixelShader();
 
-		FGlobalShader::SetParameters(ShaderRHI, Context.View);
+		FGlobalShader::SetParameters(RHICmdList, ShaderRHI, Context.View);
 
-		PostprocessParameter.SetPS(ShaderRHI, Context, TStaticSamplerState<SF_Point,AM_Clamp,AM_Clamp,AM_Clamp>::GetRHI());
+		PostprocessParameter.SetPS(RHICmdList, ShaderRHI, Context, TStaticSamplerState<SF_Point,AM_Clamp,AM_Clamp,AM_Clamp>::GetRHI());
 
 		{
 			const float SizeX = Context.View.ViewRect.Width();
@@ -71,7 +71,7 @@ public:
 			float ObjectScaleX = ViewMotionBlurScale * InvMaxVelocity;
 			float ObjectScaleY = ViewMotionBlurScale * InvMaxVelocity * InvAspectRatio;
 
-			SetShaderValue( ShaderRHI, VelocityScale, FVector4( ObjectScaleX, -ObjectScaleY, 0, 0 ) );
+			SetShaderValue(RHICmdList, ShaderRHI, VelocityScale, FVector4( ObjectScaleX, -ObjectScaleY, 0, 0 ) );
 		}
 	}
 };
@@ -127,6 +127,9 @@ void FRCPassPostProcessMotionBlurSetup::Process(FRenderingCompositePassContext&
 	RHISetRasterizerState(TStaticRasterizerState<>::GetRHI());
 	RHISetDepthStencilState(TStaticDepthStencilState<false,CF_Always>::GetRHI());
 
+	//@todo-rco: RHIPacketList
+	FRHICommandList* RHICmdList = nullptr;
+
 	TShaderMapRef<FPostProcessVS> VertexShader(GetGlobalShaderMap());
 
 	{		
@@ -135,8 +138,8 @@ void FRCPassPostProcessMotionBlurSetup::Process(FRenderingCompositePassContext&
 
 		SetGlobalBoundShaderState(BoundShaderState, GFilterVertexDeclaration.VertexDeclarationRHI, *VertexShader, *PixelShader);
 
-		PixelShader->SetParameters(Context);
-		VertexShader->SetParameters(Context);
+		PixelShader->SetParameters(RHICmdList, Context);
+		VertexShader->SetParameters(RHICmdList, Context);
 	}
 
 	// Draw a quad mapping scene color to the view's render target
@@ -217,7 +220,8 @@ class FPostProcessMotionBlurPS : public FGlobalShader
 
 	static void ModifyCompilationEnvironment(EShaderPlatform Platform, FShaderCompilerEnvironment& OutEnvironment)
 	{
-		FGlobalShader::ModifyCompilationEnvironment(Platform, OutEnvironment);
+		FGlobalShader::ModifyCompilationEnvironment(Platform,OutEnvironment);
+		FDeferredPixelShaderParameters::ModifyCompilationEnvironment(Platform,OutEnvironment);
 		OutEnvironment.SetDefine(TEXT("MOTION_BLUR_QUALITY"), Quality);
 	}
 
@@ -250,13 +254,13 @@ public:
 		return bShaderHasOutdatedParameters;
 	}
 
-	void SetParameters(const FRenderingCompositePassContext& Context)
+	void SetParameters(FRHICommandList* RHICmdList, const FRenderingCompositePassContext& Context)
 	{
 		const FPixelShaderRHIParamRef ShaderRHI = GetPixelShader();
 
-		FGlobalShader::SetParameters(ShaderRHI, Context.View);
+		FGlobalShader::SetParameters(RHICmdList, ShaderRHI, Context.View);
 
-		DeferredParameters.Set(ShaderRHI, Context.View);
+		DeferredParameters.Set(RHICmdList, ShaderRHI, Context.View);
 
 		{
 			bool bFiltered = false;
@@ -268,11 +272,11 @@ public:
 
 			if(bFiltered)
 			{
-				PostprocessParameter.SetPS(ShaderRHI, Context, TStaticSamplerState<SF_Bilinear,AM_Border,AM_Border,AM_Clamp>::GetRHI());
+				PostprocessParameter.SetPS(RHICmdList, ShaderRHI, Context, TStaticSamplerState<SF_Bilinear,AM_Border,AM_Border,AM_Clamp>::GetRHI());
 			}
 			else
 			{
-				PostprocessParameter.SetPS(ShaderRHI, Context, TStaticSamplerState<SF_Point,AM_Border,AM_Border,AM_Clamp>::GetRHI());
+				PostprocessParameter.SetPS(RHICmdList, ShaderRHI, Context, TStaticSamplerState<SF_Point,AM_Border,AM_Border,AM_Clamp>::GetRHI());
 			}
 		}
 		
@@ -282,11 +286,11 @@ public:
 			// then translate by the difference between last frame's camera position and this frame's camera position,
 			// then apply the rest of the transforms.  This effectively avoids precision issues near the extents of large levels whose world space position is very large.
 			FVector ViewOriginDelta = Context.View.ViewMatrices.ViewOrigin - Context.View.PrevViewMatrices.ViewOrigin;
-			SetShaderValue(ShaderRHI, PrevViewProjMatrix, FTranslationMatrix(ViewOriginDelta) * Context.View.PrevViewRotationProjMatrix);
+			SetShaderValue(RHICmdList, ShaderRHI, PrevViewProjMatrix, FTranslationMatrix(ViewOriginDelta) * Context.View.PrevViewRotationProjMatrix);
 		}
 		else
 		{
-			SetShaderValue( ShaderRHI, PrevViewProjMatrix, Context.View.ViewMatrices.GetViewRotationProjMatrix() );
+			SetShaderValue(RHICmdList, ShaderRHI, PrevViewProjMatrix, Context.View.ViewMatrices.GetViewRotationProjMatrix() );
 		}
 
 		TRefCountPtr<IPooledRenderTarget> InputPooledElement = Context.Pass->GetInput(ePId_Input0)->GetOutput()->RequestInput();
@@ -308,7 +312,7 @@ public:
 			FVector2D Mul(1.0f / SizeUV.X, 1.0f / SizeUV.Y);
 			FVector2D Add = - MinUV * Mul;
 			FVector4 TextureViewMadValue(Mul.X, Mul.Y, Add.X, Add.Y);
-			SetShaderValue(ShaderRHI, TextureViewMad, TextureViewMadValue);
+			SetShaderValue(RHICmdList, ShaderRHI, TextureViewMad, TextureViewMadValue);
 		}
 
 		{
@@ -337,7 +341,7 @@ public:
 				- ObjectMotionBlurScale * InvMaxVelocity * InvAspectRatio,
 				MaxVelocity * 2,
 				- MaxVelocity * 2 * AspectRatio);
-			SetShaderValue(ShaderRHI, MotionBlurParameters, MotionBlurParametersValue);
+			SetShaderValue(RHICmdList, ShaderRHI, MotionBlurParameters, MotionBlurParametersValue);
 		}
 	}
 
@@ -363,7 +367,7 @@ VARIATION1(0)			VARIATION1(1)			VARIATION1(2)			VARIATION1(3)			VARIATION1(4)
 
 // @param Quality 0: visualize, 1:low, 2:medium, 3:high, 4:very high
 template <uint32 Quality>
-static void SetMotionBlurShaderTempl(const FRenderingCompositePassContext& Context)
+static void SetMotionBlurShaderTempl(FRHICommandList* RHICmdList, const FRenderingCompositePassContext& Context)
 {
 	TShaderMapRef<FPostProcessVS> VertexShader(GetGlobalShaderMap());
 	TShaderMapRef<FPostProcessMotionBlurPS<Quality> > PixelShader(GetGlobalShaderMap());
@@ -372,8 +376,8 @@ static void SetMotionBlurShaderTempl(const FRenderingCompositePassContext& Conte
 
 	SetGlobalBoundShaderState(BoundShaderState, GFilterVertexDeclaration.VertexDeclarationRHI, *VertexShader, *PixelShader);
 
-	VertexShader->SetParameters(Context);
-	PixelShader->SetParameters(Context);
+	VertexShader->SetParameters(RHICmdList, Context);
+	PixelShader->SetParameters(RHICmdList, Context);
 }
 
 FRCPassPostProcessMotionBlur::FRCPassPostProcessMotionBlur(uint32 InQuality) 
@@ -425,22 +429,24 @@ void FRCPassPostProcessMotionBlur::Process(FRenderingCompositePassContext& Conte
 	RHISetRasterizerState(TStaticRasterizerState<>::GetRHI());
 	RHISetDepthStencilState(TStaticDepthStencilState<false,CF_Always>::GetRHI());
 
+	//@todo-rco: RHIPacketList
+	FRHICommandList* RHICmdList = nullptr;
 	if(Quality == 1)
 	{
-		SetMotionBlurShaderTempl<1>(Context);
+		SetMotionBlurShaderTempl<1>(RHICmdList, Context);
 	}
 	else if(Quality == 2)
 	{
-		SetMotionBlurShaderTempl<2>(Context);
+		SetMotionBlurShaderTempl<2>(RHICmdList, Context);
 	}
 	else if(Quality == 3)
 	{
-		SetMotionBlurShaderTempl<3>(Context);
+		SetMotionBlurShaderTempl<3>(RHICmdList, Context);
 	}
 	else
 	{
 		check(Quality == 4);
-		SetMotionBlurShaderTempl<4>(Context);
+		SetMotionBlurShaderTempl<4>(RHICmdList, Context);
 	}
 
 	TShaderMapRef<FPostProcessVS> VertexShader(GetGlobalShaderMap());
@@ -505,13 +511,13 @@ public:
 		return bShaderHasOutdatedParameters;
 	}
 
-	void SetParameters(const FRenderingCompositePassContext& Context)
+	void SetParameters(FRHICommandList* RHICmdList, const FRenderingCompositePassContext& Context)
 	{
 		const FPixelShaderRHIParamRef ShaderRHI = GetPixelShader();
 
-		FGlobalShader::SetParameters(ShaderRHI, Context.View);
+		FGlobalShader::SetParameters(RHICmdList, ShaderRHI, Context.View);
 
-		PostprocessParameter.SetPS(ShaderRHI, Context, TStaticSamplerState<SF_Bilinear,AM_Border,AM_Border,AM_Clamp>::GetRHI());
+		PostprocessParameter.SetPS(RHICmdList, ShaderRHI, Context, TStaticSamplerState<SF_Bilinear,AM_Border,AM_Border,AM_Clamp>::GetRHI());
 	}
 };
 
@@ -566,8 +572,10 @@ void FRCPassPostProcessMotionBlurRecombine::Process(FRenderingCompositePassConte
 
 	SetGlobalBoundShaderState(BoundShaderState, GFilterVertexDeclaration.VertexDeclarationRHI, *VertexShader, *PixelShader);
 
-	VertexShader->SetParameters(Context);
-	PixelShader->SetParameters(Context);
+	//@todo-rco: RHIPacketList
+	FRHICommandList* RHICmdList = nullptr;
+	VertexShader->SetParameters(RHICmdList, Context);
+	PixelShader->SetParameters(RHICmdList, Context);
 
 	// Draw a quad mapping scene color to the view's render target
 	DrawRectangle(
@@ -639,7 +647,9 @@ void FRCPassPostProcessVisualizeMotionBlur::Process(FRenderingCompositePassConte
 	RHISetDepthStencilState(TStaticDepthStencilState<false,CF_Always>::GetRHI());
 
 	// Quality 0: visualize
-	SetMotionBlurShaderTempl<0>(Context);
+	//@todo-rco: RHIPacketList
+	FRHICommandList* RHICmdList = nullptr;
+	SetMotionBlurShaderTempl<0>(RHICmdList, Context);
 
 	// Draw a quad mapping scene color to the view's render target
 	TShaderMapRef<FPostProcessVS> VertexShader(GetGlobalShaderMap());
diff --git a/Engine/Source/Runtime/Renderer/Private/PostProcess/PostProcessNoiseBlur.cpp b/Engine/Source/Runtime/Renderer/Private/PostProcess/PostProcessNoiseBlur.cpp
--- a/Engine/Source/Runtime/Renderer/Private/PostProcess/PostProcessNoiseBlur.cpp
+++ b/Engine/Source/Runtime/Renderer/Private/PostProcess/PostProcessNoiseBlur.cpp
@@ -23,7 +23,8 @@ class FPostProcessNoiseBlurPS : public FGlobalShader
 
 	static void ModifyCompilationEnvironment(EShaderPlatform Platform, FShaderCompilerEnvironment& OutEnvironment)
 	{
-		FGlobalShader::ModifyCompilationEnvironment(Platform, OutEnvironment);
+		FGlobalShader::ModifyCompilationEnvironment(Platform,OutEnvironment);
+		FDeferredPixelShaderParameters::ModifyCompilationEnvironment(Platform,OutEnvironment);
 		OutEnvironment.SetDefine(TEXT("METHOD"), Method);
 	}
 
@@ -52,17 +53,17 @@ public:
 		return bShaderHasOutdatedParameters;
 	}
 
-	void SetParameters(const FRenderingCompositePassContext& Context, float InRadius)
+	void SetParameters(FRHICommandList* RHICmdList, const FRenderingCompositePassContext& Context, float InRadius)
 	{
 		const FPixelShaderRHIParamRef ShaderRHI = GetPixelShader();
 
-		FGlobalShader::SetParameters(ShaderRHI, Context.View);
-		DeferredParameters.Set(ShaderRHI, Context.View);
-		PostprocessParameter.SetPS(ShaderRHI, Context, TStaticSamplerState<SF_Bilinear,AM_Border,AM_Border,AM_Border>::GetRHI());
+		FGlobalShader::SetParameters(RHICmdList, ShaderRHI, Context.View);
+		DeferredParameters.Set(RHICmdList, ShaderRHI, Context.View);
+		PostprocessParameter.SetPS(RHICmdList, ShaderRHI, Context, TStaticSamplerState<SF_Bilinear,AM_Border,AM_Border,AM_Border>::GetRHI());
 
 		{
 			FVector4 ColorScale(InRadius, 0, 0, 0);
-			SetShaderValue(ShaderRHI, NoiseParams, ColorScale);
+			SetShaderValue(RHICmdList, ShaderRHI, NoiseParams, ColorScale);
 		}
 	}
 
@@ -94,7 +95,7 @@ FRCPassPostProcessNoiseBlur::FRCPassPostProcessNoiseBlur(float InRadius, EPixelF
 
 
 template <uint32 Method>
-void SetNoiseBlurShader(const FRenderingCompositePassContext& Context, float InRadius)
+void SetNoiseBlurShader(FRHICommandList* RHICmdList, const FRenderingCompositePassContext& Context, float InRadius)
 {
 	TShaderMapRef<FPostProcessVS> VertexShader(GetGlobalShaderMap());
 	TShaderMapRef<FPostProcessNoiseBlurPS<Method> > PixelShader(GetGlobalShaderMap());
@@ -103,8 +104,8 @@ void SetNoiseBlurShader(const FRenderingCompositePassContext& Context, float InR
 
 	SetGlobalBoundShaderState(BoundShaderState, GFilterVertexDeclaration.VertexDeclarationRHI, *VertexShader, *PixelShader);
 
-	PixelShader->SetParameters(Context, InRadius);
-	VertexShader->SetParameters(Context);
+	PixelShader->SetParameters(RHICmdList, Context, InRadius);
+	VertexShader->SetParameters(RHICmdList, Context);
 }
 
 
@@ -147,17 +148,19 @@ void FRCPassPostProcessNoiseBlur::Process(FRenderingCompositePassContext& Contex
 	RHISetRasterizerState(TStaticRasterizerState<>::GetRHI());
 	RHISetDepthStencilState(TStaticDepthStencilState<false,CF_Always>::GetRHI());
 	
+	//@todo-rco: RHIPacketList
+	FRHICommandList* RHICmdList = nullptr;
 	if(Quality == 0)
 	{
-		SetNoiseBlurShader<0>(Context, Radius);
+		SetNoiseBlurShader<0>(RHICmdList, Context, Radius);
 	}
 	else if(Quality == 1)
 	{
-		SetNoiseBlurShader<1>(Context, Radius);
+		SetNoiseBlurShader<1>(RHICmdList, Context, Radius);
 	}
 	else
 	{
-		SetNoiseBlurShader<2>(Context, Radius);
+		SetNoiseBlurShader<2>(RHICmdList, Context, Radius);
 	}
 
 	// Draw a quad mapping scene color to the view's render target
diff --git a/Engine/Source/Runtime/Renderer/Private/PostProcess/PostProcessSceneColorFringe.cpp b/Engine/Source/Runtime/Renderer/Private/PostProcess/PostProcessSceneColorFringe.cpp
--- a/Engine/Source/Runtime/Renderer/Private/PostProcess/PostProcessSceneColorFringe.cpp
+++ b/Engine/Source/Runtime/Renderer/Private/PostProcess/PostProcessSceneColorFringe.cpp
@@ -49,15 +49,15 @@ public:
 		return bShaderHasOutdatedParameters;
 	}
 
-	void SetParameters(const FRenderingCompositePassContext& Context)
+	void SetParameters(FRHICommandList* RHICmdList, const FRenderingCompositePassContext& Context)
 	{
 		const FVertexShaderRHIParamRef ShaderRHI = GetVertexShader();
 
 		const FSceneView& View = Context.View;
 
-		FGlobalShader::SetParameters(ShaderRHI, Context.View);
+		FGlobalShader::SetParameters(RHICmdList, ShaderRHI, Context.View);
 		
-		PostprocessParameter.SetVS(ShaderRHI, Context, TStaticSamplerState<SF_Bilinear,AM_Clamp,AM_Clamp,AM_Clamp>::GetRHI());
+		PostprocessParameter.SetVS(RHICmdList, ShaderRHI, Context, TStaticSamplerState<SF_Bilinear,AM_Clamp,AM_Clamp,AM_Clamp>::GetRHI());
 
 		const FPooledRenderTargetDesc* InputDesc = Context.Pass->GetInputDesc(ePId_Input0);
 
@@ -73,7 +73,7 @@ public:
 			// we only get bigger to not leak in content from outside
 			FVector4 Value(1.0f, 1.0f - Offset * 0.5f, 1.0f - Offset, 0);
 
-			SetShaderValue(ShaderRHI, FringeUVParams, Value);
+			SetShaderValue(RHICmdList, ShaderRHI, FringeUVParams, Value);
 		}
 	}
 };
@@ -114,15 +114,15 @@ public:
 		return bShaderHasOutdatedParameters;
 	}
 
-	void SetParameters(const FRenderingCompositePassContext& Context)
+	void SetParameters(FRHICommandList* RHICmdList, const FRenderingCompositePassContext& Context)
 	{
 		const FPixelShaderRHIParamRef ShaderRHI = GetPixelShader();
 
 		const FSceneView& View = Context.View;
 
-		FGlobalShader::SetParameters(ShaderRHI, View);
+		FGlobalShader::SetParameters(RHICmdList, ShaderRHI, View);
 
-		PostprocessParameter.SetPS(ShaderRHI, Context, TStaticSamplerState<SF_Bilinear,AM_Clamp,AM_Clamp,AM_Clamp>::GetRHI());
+		PostprocessParameter.SetPS(RHICmdList, ShaderRHI, Context, TStaticSamplerState<SF_Bilinear,AM_Clamp,AM_Clamp,AM_Clamp>::GetRHI());
 
 		{
 			FVector4 Value[3];
@@ -133,7 +133,7 @@ public:
 			Value[0] = FMath::Lerp(FVector4(t, t, t, 0), FVector4(1, 0, 0, 0), Alpha);
 			Value[1] = FMath::Lerp(FVector4(t, t, t, 0), FVector4(0, 1, 0, 0), Alpha);
 			Value[2] = FMath::Lerp(FVector4(t, t, t, 0), FVector4(0, 0, 1, 0), Alpha);
-			SetShaderValueArray(ShaderRHI, FringeColorParams, Value, 3);
+			SetShaderValueArray(RHICmdList, ShaderRHI, FringeColorParams, Value, 3);
 		}
 	}
 };
@@ -173,8 +173,10 @@ void FRCPassPostProcessSceneColorFringe::Process(FRenderingCompositePassContext&
 
 	SetGlobalBoundShaderState(BoundShaderState, GFilterVertexDeclaration.VertexDeclarationRHI, *VertexShader, *PixelShader);
 
-	PixelShader->SetParameters(Context);
-	VertexShader->SetParameters(Context);
+	//@todo-rco: RHIPacketList
+	FRHICommandList* RHICmdList = nullptr;
+	PixelShader->SetParameters(RHICmdList, Context);
+	VertexShader->SetParameters(RHICmdList, Context);
 	
 	// Draw a quad mapping scene color to the view's render target
 	DrawRectangle(
diff --git a/Engine/Source/Runtime/Renderer/Private/PostProcess/PostProcessSelectionOutline.cpp b/Engine/Source/Runtime/Renderer/Private/PostProcess/PostProcessSelectionOutline.cpp
--- a/Engine/Source/Runtime/Renderer/Private/PostProcess/PostProcessSelectionOutline.cpp
+++ b/Engine/Source/Runtime/Renderer/Private/PostProcess/PostProcessSelectionOutline.cpp
@@ -203,19 +203,19 @@ public:
 		EditorPrimitivesStencil.Bind(Initializer.ParameterMap,TEXT("EditorPrimitivesStencil"));
 	}
 
-	void SetPS(const FRenderingCompositePassContext& Context)
+	void SetPS(FRHICommandList* RHICmdList, const FRenderingCompositePassContext& Context)
 	{
 		const FPixelShaderRHIParamRef ShaderRHI = GetPixelShader();
 
-		FGlobalShader::SetParameters(ShaderRHI, Context.View);
+		FGlobalShader::SetParameters(RHICmdList, ShaderRHI, Context.View);
 
-		DeferredParameters.Set(ShaderRHI, Context.View);
+		DeferredParameters.Set(RHICmdList, ShaderRHI, Context.View);
 
 		const FPostProcessSettings& Settings = Context.View.FinalPostProcessSettings;
 		const FSceneViewFamily& ViewFamily = *(Context.View.Family);
 		FSceneViewState* ViewState = (FSceneViewState*)Context.View.State;
 
-		PostprocessParameter.SetPS(ShaderRHI, Context, TStaticSamplerState<SF_Point,AM_Clamp,AM_Clamp,AM_Clamp>::GetRHI());
+		PostprocessParameter.SetPS(RHICmdList, ShaderRHI, Context, TStaticSamplerState<SF_Point,AM_Clamp,AM_Clamp,AM_Clamp>::GetRHI());
 
 		// PostprocessInput1MS and EditorPrimitivesStencil
 		{
@@ -233,11 +233,11 @@ public:
 
 			FTexture2DRHIRef& TargetableTexture = (FTexture2DRHIRef&)InputPooledElement->GetRenderTargetItem().TargetableTexture;
 
-			SetTextureParameter(ShaderRHI, PostprocessInput1MS, TargetableTexture);
+			SetTextureParameter(RHICmdList, ShaderRHI, PostprocessInput1MS, TargetableTexture);
 
 			if(EditorPrimitivesStencil.IsBound())
 			{
-				// cache the stencil SRV to avoid create calls each frame (the cache element is stored in the state)
+			// cache the stencil SRV to avoid create calls each frame (the cache element is stored in the state)
 				if(ViewState->SelectionOutlineCacheKey != TargetableTexture)
 				{
 					// release if not the right one (as the internally SRV stores a pointer to the texture we cannot get a false positive)
@@ -252,8 +252,8 @@ public:
 					ViewState->SelectionOutlineCacheValue = RHICreateShaderResourceView(TargetableTexture, 0, 1, PF_X24_G8);
 				}
 
-				SetSRVParameter(ShaderRHI, EditorPrimitivesStencil, ViewState->SelectionOutlineCacheValue);
-			}
+			SetSRVParameter(RHICmdList, ShaderRHI, EditorPrimitivesStencil, ViewState->SelectionOutlineCacheValue);
+		}
 		}
 
 #if WITH_EDITOR
@@ -262,8 +262,8 @@ public:
 
 			Value.A = GEngine->SelectionHighlightIntensity;
 
-			SetShaderValue(ShaderRHI, OutlineColor, Value );
-			SetShaderValue(ShaderRHI, BSPSelectionIntensity, GEngine->BSPSelectionHighlightIntensity);
+			SetShaderValue(RHICmdList, ShaderRHI, OutlineColor, Value );
+			SetShaderValue(RHICmdList, ShaderRHI, BSPSelectionIntensity, GEngine->BSPSelectionHighlightIntensity);
 		}
 #else
 		check(!"This shader is not used outside of the Editor.");
@@ -280,7 +280,7 @@ public:
 				Value.G = 0;
 			}
 
-			SetShaderValue(ShaderRHI, EditorRenderParams, Value );
+			SetShaderValue(RHICmdList, ShaderRHI, EditorRenderParams, Value );
 		}
 	}
 	
@@ -312,7 +312,7 @@ VARIATION1(1)			VARIATION1(2)			VARIATION1(4)			VARIATION1(8)
 
 
 template <uint32 MSAASampleCount>
-static void SetSelectionOutlineShaderTempl(const FRenderingCompositePassContext& Context)
+static void SetSelectionOutlineShaderTempl(FRHICommandList* RHICmdList, const FRenderingCompositePassContext& Context)
 {
 	TShaderMapRef<FPostProcessVS> VertexShader(GetGlobalShaderMap());
 	TShaderMapRef<FPostProcessSelectionOutlinePS<MSAASampleCount> > PixelShader(GetGlobalShaderMap());
@@ -320,7 +320,7 @@ static void SetSelectionOutlineShaderTempl(const FRenderingCompositePassContext&
 	static FGlobalBoundShaderState BoundShaderState;
 	SetGlobalBoundShaderState(BoundShaderState, GFilterVertexDeclaration.VertexDeclarationRHI, *VertexShader, *PixelShader);
 
-	PixelShader->SetPS(Context);
+	PixelShader->SetPS(RHICmdList, Context);
 }
 
 void FRCPassPostProcessSelectionOutline::Process(FRenderingCompositePassContext& Context)
@@ -334,6 +334,8 @@ void FRCPassPostProcessSelectionOutline::Process(FRenderingCompositePassContext&
 		// input is not hooked up correctly
 		return;
 	}
+	//@todo-rco: RHIPacketList
+	FRHICommandList* RHICmdList = nullptr;
 
 	const FSceneView& View = Context.View;
 	FIntRect ViewRect = View.ViewRect;
@@ -355,19 +357,19 @@ void FRCPassPostProcessSelectionOutline::Process(FRenderingCompositePassContext&
 
 	if(MSAASampleCount == 1)
 	{
-		SetSelectionOutlineShaderTempl<1>(Context);
+		SetSelectionOutlineShaderTempl<1>(RHICmdList, Context);
 	}
 	else if(MSAASampleCount == 2)
 	{
-		SetSelectionOutlineShaderTempl<2>(Context);
+		SetSelectionOutlineShaderTempl<2>(RHICmdList, Context);
 	}
 	else if(MSAASampleCount == 4)
 	{
-		SetSelectionOutlineShaderTempl<4>(Context);
+		SetSelectionOutlineShaderTempl<4>(RHICmdList, Context);
 	}
 	else if(MSAASampleCount == 8)
 	{
-		SetSelectionOutlineShaderTempl<8>(Context);
+		SetSelectionOutlineShaderTempl<8>(RHICmdList, Context);
 	}
 	else
 	{
diff --git a/Engine/Source/Runtime/Renderer/Private/PostProcess/PostProcessSubsurface.cpp b/Engine/Source/Runtime/Renderer/Private/PostProcess/PostProcessSubsurface.cpp
--- a/Engine/Source/Runtime/Renderer/Private/PostProcess/PostProcessSubsurface.cpp
+++ b/Engine/Source/Runtime/Renderer/Private/PostProcess/PostProcessSubsurface.cpp
@@ -44,14 +44,14 @@ public:
 		DeferredParameters.Bind(Initializer.ParameterMap);
 	}
 
-	void SetParameters(const FRenderingCompositePassContext& Context)
+	void SetParameters(FRHICommandList* RHICmdList, const FRenderingCompositePassContext& Context)
 	{
 		const FFinalPostProcessSettings& Settings = Context.View.FinalPostProcessSettings;
 		const FPixelShaderRHIParamRef ShaderRHI = GetPixelShader();
 
-		FGlobalShader::SetParameters(ShaderRHI, Context.View);
-		PostprocessParameter.SetPS(ShaderRHI, Context, TStaticSamplerState<SF_Point,AM_Clamp,AM_Clamp,AM_Clamp>::GetRHI());
-		DeferredParameters.Set(ShaderRHI, Context.View);
+		FGlobalShader::SetParameters(RHICmdList, ShaderRHI, Context.View);
+		PostprocessParameter.SetPS(RHICmdList, ShaderRHI, Context, TStaticSamplerState<SF_Point,AM_Clamp,AM_Clamp,AM_Clamp>::GetRHI());
+		DeferredParameters.Set(RHICmdList, ShaderRHI, Context.View);
 	}
 
 	// FShader interface.
@@ -83,7 +83,7 @@ public:
 
 
 template <uint32 SpecularCorrection>
-void SetSubsurfaceSetupShader(const FRenderingCompositePassContext& Context)
+void SetSubsurfaceSetupShader(FRHICommandList* RHICmdList, const FRenderingCompositePassContext& Context)
 {
 	TShaderMapRef<FPostProcessVS> VertexShader(GetGlobalShaderMap());
 	TShaderMapRef<FPostProcessSubsurfaceSetupPS<SpecularCorrection> > PixelShader(GetGlobalShaderMap());
@@ -92,8 +92,8 @@ void SetSubsurfaceSetupShader(const FRenderingCompositePassContext& Context)
 
 	SetGlobalBoundShaderState(BoundShaderState, GFilterVertexDeclaration.VertexDeclarationRHI, *VertexShader, *PixelShader);
 
-	PixelShader->SetParameters(Context);
-	VertexShader->SetParameters(Context);
+	PixelShader->SetParameters(RHICmdList, Context);
+	VertexShader->SetParameters(RHICmdList, Context);
 }
 
 static TAutoConsoleVariable<int> CVarSubsurfaceQuality(
@@ -146,6 +146,9 @@ void FRCPassPostProcessSubsurfaceSetup::Process(FRenderingCompositePassContext&
 
 	const FSceneRenderTargetItem& DestRenderTarget = PassOutputs[0].RequestSurface(Context);
 
+	//@todo-rco: RHIPacketList
+	FRHICommandList* RHICmdList = nullptr;
+
 	// Set the view family's render target/viewport.
 	RHISetRenderTarget(DestRenderTarget.TargetableTexture, FTextureRHIRef());	
 
@@ -162,12 +165,12 @@ void FRCPassPostProcessSubsurfaceSetup::Process(FRenderingCompositePassContext&
 	if(DoSpecularCorrection())
 	{
 		// with separate specular
-		SetSubsurfaceSetupShader<1>(Context);
+		SetSubsurfaceSetupShader<1>(RHICmdList, Context);
 	}
 	else
 	{
 		// no separate specular
-		SetSubsurfaceSetupShader<0>(Context);
+		SetSubsurfaceSetupShader<0>(RHICmdList, Context);
 	}
 
 	// Draw a quad mapping scene color to the view's render target
@@ -211,7 +214,8 @@ class FPostProcessSubsurfacePS : public FGlobalShader
 
 	static void ModifyCompilationEnvironment(EShaderPlatform Platform, FShaderCompilerEnvironment& OutEnvironment)
 	{
-		FGlobalShader::ModifyCompilationEnvironment(Platform, OutEnvironment);
+		FGlobalShader::ModifyCompilationEnvironment(Platform,OutEnvironment);
+		FDeferredPixelShaderParameters::ModifyCompilationEnvironment(Platform,OutEnvironment);
 		OutEnvironment.SetDefine(TEXT("METHOD"), Method);
 	}
 
@@ -240,17 +244,17 @@ public:
 		return bShaderHasOutdatedParameters;
 	}
 
-	void SetParameters(const FRenderingCompositePassContext& Context, float InRadius)
+	void SetParameters(FRHICommandList* RHICmdList, const FRenderingCompositePassContext& Context, float InRadius)
 	{
 		const FPixelShaderRHIParamRef ShaderRHI = GetPixelShader();
 
-		FGlobalShader::SetParameters(ShaderRHI, Context.View);
-		DeferredParameters.Set(ShaderRHI, Context.View);
-		PostprocessParameter.SetPS(ShaderRHI, Context, TStaticSamplerState<SF_Bilinear,AM_Border,AM_Border,AM_Border>::GetRHI());
+		FGlobalShader::SetParameters(RHICmdList, ShaderRHI, Context.View);
+		DeferredParameters.Set(RHICmdList, ShaderRHI, Context.View);
+		PostprocessParameter.SetPS(RHICmdList, ShaderRHI, Context, TStaticSamplerState<SF_Bilinear,AM_Border,AM_Border,AM_Border>::GetRHI());
 
 		{
 			FVector4 ColorScale(InRadius, 0, 0, 0);
-			SetShaderValue(ShaderRHI, SSSParams, ColorScale);
+			SetShaderValue(RHICmdList, ShaderRHI, SSSParams, ColorScale);
 		}
 	}
 
@@ -281,7 +285,7 @@ FRCPassPostProcessSubsurface::FRCPassPostProcessSubsurface(uint32 InPass, float
 
 
 template <uint32 Method>
-void SetSubsurfaceShader(const FRenderingCompositePassContext& Context, float InRadius)
+void SetSubsurfaceShader(FRHICommandList* RHICmdList, const FRenderingCompositePassContext& Context, float InRadius)
 {
 	TShaderMapRef<FPostProcessVS> VertexShader(GetGlobalShaderMap());
 	TShaderMapRef<FPostProcessSubsurfacePS<Method> > PixelShader(GetGlobalShaderMap());
@@ -290,8 +294,8 @@ void SetSubsurfaceShader(const FRenderingCompositePassContext& Context, float In
 
 	SetGlobalBoundShaderState(BoundShaderState, GFilterVertexDeclaration.VertexDeclarationRHI, *VertexShader, *PixelShader);
 
-	PixelShader->SetParameters(Context, InRadius);
-	VertexShader->SetParameters(Context);
+	PixelShader->SetParameters(RHICmdList, Context, InRadius);
+	VertexShader->SetParameters(RHICmdList, Context);
 }
 
 
@@ -331,6 +335,8 @@ void FRCPassPostProcessSubsurface::Process(FRenderingCompositePassContext& Conte
 	RHISetBlendState(TStaticBlendState<>::GetRHI());
 	RHISetRasterizerState(TStaticRasterizerState<>::GetRHI());
 	RHISetDepthStencilState(TStaticDepthStencilState<false,CF_Always>::GetRHI());
+	//@todo-rco: RHIPacketList
+	FRHICommandList* RHICmdList = nullptr;
 
 	TShaderMapRef<FPostProcessVS> VertexShader(GetGlobalShaderMap());
 
@@ -338,7 +344,7 @@ void FRCPassPostProcessSubsurface::Process(FRenderingCompositePassContext& Conte
 	{
 		SCOPED_DRAW_EVENT(SubsurfacePass0, DEC_SCENE_ITEMS);
 	
-		SetSubsurfaceShader<0>(Context, Radius);
+		SetSubsurfaceShader<0>(RHICmdList, Context, Radius);
 
 		// Draw a quad mapping scene color to the view's render target
 		DrawRectangle(
@@ -361,11 +367,11 @@ void FRCPassPostProcessSubsurface::Process(FRenderingCompositePassContext& Conte
 		if(DoSpecularCorrection())
 		{
 			// reconstruct specular and add it in final pass
-			SetSubsurfaceShader<2>(Context, Radius);
+			SetSubsurfaceShader<2>(RHICmdList, Context, Radius);
 		}
 		else
 		{
-			SetSubsurfaceShader<1>(Context, Radius);
+			SetSubsurfaceShader<1>(RHICmdList, Context, Radius);
 		}
 
 		// Draw a quad mapping scene color to the view's render target
diff --git a/Engine/Source/Runtime/Renderer/Private/PostProcess/PostProcessTemporalAA.cpp b/Engine/Source/Runtime/Renderer/Private/PostProcess/PostProcessTemporalAA.cpp
--- a/Engine/Source/Runtime/Renderer/Private/PostProcess/PostProcessTemporalAA.cpp
+++ b/Engine/Source/Runtime/Renderer/Private/PostProcess/PostProcessTemporalAA.cpp
@@ -52,6 +52,7 @@ class FPostProcessTemporalAAPS : public FGlobalShader
 	static void ModifyCompilationEnvironment(EShaderPlatform Platform, FShaderCompilerEnvironment& OutEnvironment)
 	{
 		FGlobalShader::ModifyCompilationEnvironment(Platform, OutEnvironment);
+		FDeferredPixelShaderParameters::ModifyCompilationEnvironment(Platform,OutEnvironment);
 
 		OutEnvironment.SetDefine( TEXT("RESPONSIVE"), Responsive );
 	}
@@ -87,11 +88,11 @@ public:
 		return bShaderHasOutdatedParameters;
 	}
 
-	void SetParameters(const FRenderingCompositePassContext& Context)
+	void SetParameters(FRHICommandList* RHICmdList, const FRenderingCompositePassContext& Context)
 	{
 		const FPixelShaderRHIParamRef ShaderRHI = GetPixelShader();
 
-		FGlobalShader::SetParameters(ShaderRHI, Context.View);
+		FGlobalShader::SetParameters(RHICmdList, ShaderRHI, Context.View);
 		
 		FSamplerStateRHIParamRef FilterTable[4];
 		FilterTable[0] = TStaticSamplerState<SF_Point,AM_Clamp,AM_Clamp,AM_Clamp>::GetRHI();
@@ -99,9 +100,9 @@ public:
 		FilterTable[2] = FilterTable[0];
 		FilterTable[3] = FilterTable[0];
 
-		PostprocessParameter.SetPS(ShaderRHI, Context, 0, false, FilterTable);
+		PostprocessParameter.SetPS(RHICmdList, ShaderRHI, Context, 0, false, FilterTable);
 
-		DeferredParameters.Set(ShaderRHI, Context.View);
+		DeferredParameters.Set(RHICmdList, ShaderRHI, Context.View);
 
 		FSceneViewState* ViewState = (FSceneViewState*)Context.View.State;
 
@@ -157,17 +158,17 @@ public:
 			
 			for( int32 i = 0; i < 9; i++ )
 			{
-				SetShaderValue( ShaderRHI, SampleWeights, Weights[i] / TotalWeight, i );
-				SetShaderValue( ShaderRHI, LowpassWeights, WeightsLow[i] / TotalWeightLow, i );
+				SetShaderValue(RHICmdList, ShaderRHI, SampleWeights, Weights[i] / TotalWeight, i );
+				SetShaderValue(RHICmdList, ShaderRHI, LowpassWeights, WeightsLow[i] / TotalWeightLow, i );
 			}
 			
 			FVector2D RandomOffsetValue;
 			TemporalRandom(&RandomOffsetValue, Context.View.FrameNumber);
-			SetShaderValue(ShaderRHI, RandomOffset, RandomOffsetValue);
+			SetShaderValue(RHICmdList, ShaderRHI, RandomOffset, RandomOffsetValue);
 
 		}
 
-		CameraMotionParams.Set(Context.View, ShaderRHI);
+		CameraMotionParams.Set(RHICmdList, Context.View, ShaderRHI);
 	}
 };
 
@@ -230,8 +231,10 @@ void FRCPassPostProcessSSRTemporalAA::Process(FRenderingCompositePassContext& Co
 
 	SetGlobalBoundShaderState(BoundShaderState, GFilterVertexDeclaration.VertexDeclarationRHI, *VertexShader, *PixelShader);
 
-	VertexShader->SetVS(Context);
-	PixelShader->SetParameters(Context);
+	//@todo-rco: RHIPacketList
+	FRHICommandList* RHICmdList = nullptr;
+	VertexShader->SetVS(RHICmdList, Context);
+	PixelShader->SetParameters(RHICmdList, Context);
 
 	// Draw a quad mapping scene color to the view's render target
 	DrawRectangle(
@@ -305,8 +308,10 @@ void FRCPassPostProcessDOFTemporalAA::Process(FRenderingCompositePassContext& Co
 
 	SetGlobalBoundShaderState(BoundShaderState, GFilterVertexDeclaration.VertexDeclarationRHI, *VertexShader, *PixelShader);
 
-	VertexShader->SetVS(Context);
-	PixelShader->SetParameters(Context);
+	//@todo-rco: RHIPacketList
+	FRHICommandList* RHICmdList = nullptr;
+	VertexShader->SetVS(RHICmdList, Context);
+	PixelShader->SetParameters(RHICmdList, Context);
 
 	// Draw a quad mapping scene color to the view's render target
 	DrawRectangle(
@@ -383,8 +388,10 @@ void FRCPassPostProcessLightShaftTemporalAA::Process(FRenderingCompositePassCont
 
 	SetGlobalBoundShaderState(BoundShaderState, GFilterVertexDeclaration.VertexDeclarationRHI, *VertexShader, *PixelShader);
 
-	VertexShader->SetVS(Context);
-	PixelShader->SetParameters(Context);
+	//@todo-rco: RHIPacketList
+	FRHICommandList* RHICmdList = nullptr;
+	VertexShader->SetVS(RHICmdList, Context);
+	PixelShader->SetParameters(RHICmdList, Context);
 
 	// Draw a quad mapping scene color to the view's render target
 	DrawRectangle(
@@ -440,6 +447,8 @@ void FRCPassPostProcessTemporalAA::Process(FRenderingCompositePassContext& Conte
 
 	const FSceneRenderTargetItem& DestRenderTarget = PassOutputs[0].RequestSurface(Context);
 
+	//@todo-rco: RHIPacketList
+	FRHICommandList* RHICmdList = nullptr;
 	//RHISetRenderTarget(DestRenderTarget.TargetableTexture, FTextureRHIRef());
 	RHISetRenderTarget(DestRenderTarget.TargetableTexture, GSceneRenderTargets.GetSceneDepthTexture());
 
@@ -468,8 +477,8 @@ void FRCPassPostProcessTemporalAA::Process(FRenderingCompositePassContext& Conte
 
 			SetGlobalBoundShaderState(BoundShaderState, GFilterVertexDeclaration.VertexDeclarationRHI, *VertexShader, *PixelShader);
 
-			VertexShader->SetVS(Context);
-			PixelShader->SetParameters(Context);
+			VertexShader->SetVS(RHICmdList, Context);
+			PixelShader->SetParameters(RHICmdList, Context);
 		}
 		else
 		{
@@ -479,8 +488,8 @@ void FRCPassPostProcessTemporalAA::Process(FRenderingCompositePassContext& Conte
 
 			SetGlobalBoundShaderState(BoundShaderState, GFilterVertexDeclaration.VertexDeclarationRHI, *VertexShader, *PixelShader);
 
-			VertexShader->SetVS(Context);
-			PixelShader->SetParameters(Context);
+			VertexShader->SetVS(RHICmdList, Context);
+			PixelShader->SetParameters(RHICmdList, Context);
 		}
 	
 		// Draw a quad mapping scene color to the view's render target
@@ -508,8 +517,8 @@ void FRCPassPostProcessTemporalAA::Process(FRenderingCompositePassContext& Conte
 
 			SetGlobalBoundShaderState(BoundShaderState, GFilterVertexDeclaration.VertexDeclarationRHI, *VertexShader, *PixelShader);
 
-			VertexShader->SetVS(Context);
-			PixelShader->SetParameters(Context);
+			VertexShader->SetVS(RHICmdList, Context);
+			PixelShader->SetParameters(RHICmdList, Context);
 		}
 		else
 		{
@@ -519,8 +528,8 @@ void FRCPassPostProcessTemporalAA::Process(FRenderingCompositePassContext& Conte
 
 			SetGlobalBoundShaderState(BoundShaderState, GFilterVertexDeclaration.VertexDeclarationRHI, *VertexShader, *PixelShader);
 
-			VertexShader->SetVS(Context);
-			PixelShader->SetParameters(Context);
+			VertexShader->SetVS(RHICmdList, Context);
+			PixelShader->SetParameters(RHICmdList, Context);
 		}
 
 		// Draw a quad mapping scene color to the view's render target
