UE4 - merge first pass at parallel rendering to main

diff --git a/Engine/Source/Runtime/Renderer/Private/PostProcess/PostProcessTestImage.cpp b/Engine/Source/Runtime/Renderer/Private/PostProcess/PostProcessTestImage.cpp
--- a/Engine/Source/Runtime/Renderer/Private/PostProcess/PostProcessTestImage.cpp
+++ b/Engine/Source/Runtime/Renderer/Private/PostProcess/PostProcessTestImage.cpp
@@ -23,6 +23,12 @@ class FPostProcessTestImagePS : public FGlobalShader
 		return IsFeatureLevelSupported(Platform, ERHIFeatureLevel::SM4);
 	}
 
+	static void ModifyCompilationEnvironment(EShaderPlatform Platform, FShaderCompilerEnvironment& OutEnvironment)
+	{
+		FGlobalShader::ModifyCompilationEnvironment(Platform,OutEnvironment);
+		FDeferredPixelShaderParameters::ModifyCompilationEnvironment(Platform,OutEnvironment);
+	}
+
 	/** Default constructor. */
 	FPostProcessTestImagePS() {}
 
@@ -44,26 +50,26 @@ public:
 		FrameTime.Bind(Initializer.ParameterMap,TEXT("FrameTime"));
 	}
 
-	void SetPS(const FRenderingCompositePassContext& Context)
+	void SetPS(FRHICommandList* RHICmdList, const FRenderingCompositePassContext& Context)
 	{
 		const FPixelShaderRHIParamRef ShaderRHI = GetPixelShader();
 		
-		FGlobalShader::SetParameters(ShaderRHI, Context.View);
+		FGlobalShader::SetParameters(RHICmdList, ShaderRHI, Context.View);
 
-		PostprocessParameter.SetPS(ShaderRHI, Context, TStaticSamplerState<SF_Point,AM_Clamp,AM_Clamp,AM_Clamp>::GetRHI());
-		DeferredParameters.Set(ShaderRHI, Context.View);
+		PostprocessParameter.SetPS(RHICmdList, ShaderRHI, Context, TStaticSamplerState<SF_Point,AM_Clamp,AM_Clamp,AM_Clamp>::GetRHI());
+		DeferredParameters.Set(RHICmdList, ShaderRHI, Context.View);
 
 		{
 			uint32 FrameNumberValue = Context.View.FrameNumber;
-			SetShaderValue(ShaderRHI, FrameNumber, FrameNumberValue);
+			SetShaderValue(RHICmdList, ShaderRHI, FrameNumber, FrameNumberValue);
 		}
 
 		{
 			float FrameTimeValue = Context.View.Family->CurrentRealTime;
-			SetShaderValue(ShaderRHI, FrameTime, FrameTimeValue);
+			SetShaderValue(RHICmdList, ShaderRHI, FrameTime, FrameTimeValue);
 		}
 
-		ColorRemapShaderParameters.Set(ShaderRHI);
+		ColorRemapShaderParameters.Set(RHICmdList, ShaderRHI);
 	}
 	
 	// FShader interface.
@@ -93,6 +99,9 @@ void FRCPassPostProcessTestImage::Process(FRenderingCompositePassContext& Contex
 
 	const FSceneRenderTargetItem& DestRenderTarget = PassOutputs[0].RequestSurface(Context);
 
+	//@todo-rco: RHIPacketList
+	FRHICommandList* RHICmdList = nullptr;
+
 	// Set the view family's render target/viewport.
 	RHISetRenderTarget(DestRenderTarget.TargetableTexture, FTextureRHIRef());	
 	Context.SetViewportAndCallRHI(DestRect);
@@ -109,7 +118,7 @@ void FRCPassPostProcessTestImage::Process(FRenderingCompositePassContext& Contex
 
 	SetGlobalBoundShaderState(BoundShaderState, GFilterVertexDeclaration.VertexDeclarationRHI, *VertexShader, *PixelShader);
 
-	PixelShader->SetPS(Context);
+	PixelShader->SetPS(RHICmdList, Context);
 
 	// Draw a quad mapping scene color to the view's render target
 	DrawRectangle(
diff --git a/Engine/Source/Runtime/Renderer/Private/PostProcess/PostProcessTonemap.cpp b/Engine/Source/Runtime/Renderer/Private/PostProcess/PostProcessTonemap.cpp
--- a/Engine/Source/Runtime/Renderer/Private/PostProcess/PostProcessTonemap.cpp
+++ b/Engine/Source/Runtime/Renderer/Private/PostProcess/PostProcessTonemap.cpp
@@ -716,7 +716,12 @@ void FilmPostSetConstants(FVector4* RESTRICT const Constants, const uint32 Confi
 	Constants[7] = FVector4(OutColorShadow_Tint2, 0.0f);
 }
 
-
+BEGIN_UNIFORM_BUFFER_STRUCT(FBloomDirtMaskParameters,)
+	DECLARE_UNIFORM_BUFFER_STRUCT_MEMBER(FVector4,Tint)
+	DECLARE_UNIFORM_BUFFER_STRUCT_MEMBER_TEXTURE(Texture2D,Mask)
+	DECLARE_UNIFORM_BUFFER_STRUCT_MEMBER_SAMPLER(SamplerState,MaskSampler)
+END_UNIFORM_BUFFER_STRUCT(FBloomDirtMaskParameters)
+IMPLEMENT_UNIFORM_BUFFER_STRUCT(FBloomDirtMaskParameters,TEXT("BloomDirtMask"));
 
 /**
  * Encapsulates the post processing tonemapper pixel shader.
@@ -769,9 +774,6 @@ public:
 	FShaderParameter TexScale;
 	FShaderParameter VignetteColorIntensity;
 	FShaderParameter GrainScaleBiasJitter;
-	FShaderParameter BloomDirtMaskTint;
-	FShaderResourceParameter BloomDirtMask;
-	FShaderResourceParameter BloomDirtMaskSampler;
 	FShaderResourceParameter ColorGradingLUT;
 	FShaderResourceParameter ColorGradingLUTSampler;
 	FShaderParameter InverseGamma;
@@ -797,9 +799,6 @@ public:
 		TexScale.Bind(Initializer.ParameterMap, TEXT("TexScale"));
 		VignetteColorIntensity.Bind(Initializer.ParameterMap, TEXT("VignetteColorIntensity"));
 		GrainScaleBiasJitter.Bind(Initializer.ParameterMap, TEXT("GrainScaleBiasJitter"));
-		BloomDirtMaskTint.Bind(Initializer.ParameterMap, TEXT("BloomDirtMaskTint"));
-		BloomDirtMask.Bind(Initializer.ParameterMap, TEXT("BloomDirtMask"));
-		BloomDirtMaskSampler.Bind(Initializer.ParameterMap, TEXT("BloomDirtMaskSampler"));
 		ColorGradingLUT.Bind(Initializer.ParameterMap, TEXT("ColorGradingLUT"));
 		ColorGradingLUTSampler.Bind(Initializer.ParameterMap, TEXT("ColorGradingLUTSampler"));
 		InverseGamma.Bind(Initializer.ParameterMap,TEXT("InverseGamma"));
@@ -819,40 +818,40 @@ public:
 	{
 		bool bShaderHasOutdatedParameters = FGlobalShader::Serialize(Ar);
 		Ar  << PostprocessParameter << ColorScale0 << ColorScale1 << InverseGamma << NoiseTexture << NoiseTextureSampler
-			<< TexScale << VignetteColorIntensity << GrainScaleBiasJitter << BloomDirtMaskTint << BloomDirtMask << BloomDirtMaskSampler 
+			<< TexScale << VignetteColorIntensity << GrainScaleBiasJitter
 			<< ColorGradingLUT << ColorGradingLUTSampler
 			<< ColorMatrixR_ColorCurveCd1 << ColorMatrixG_ColorCurveCd3Cm3 << ColorMatrixB_ColorCurveCm2 << ColorCurve_Cm0Cd0_Cd2_Ch0Cm1_Ch3 << ColorCurve_Ch1_Ch2 << ColorShadow_Luma << ColorShadow_Tint1 << ColorShadow_Tint2;
-
+		
 		return bShaderHasOutdatedParameters;
 	}
 
-	void SetPS(const FRenderingCompositePassContext& Context)
+	void SetPS(FRHICommandList* RHICmdList, const FRenderingCompositePassContext& Context)
 	{
 		const FPostProcessSettings& Settings = Context.View.FinalPostProcessSettings;
 		const FSceneViewFamily& ViewFamily = *(Context.View.Family);
 
 		const FPixelShaderRHIParamRef ShaderRHI = GetPixelShader();
 		
-		FGlobalShader::SetParameters(ShaderRHI, Context.View);
+		FGlobalShader::SetParameters(RHICmdList, ShaderRHI, Context.View);
 
-		PostprocessParameter.SetPS(ShaderRHI, Context, TStaticSamplerState<SF_Bilinear,AM_Clamp,AM_Clamp,AM_Clamp>::GetRHI());
+		PostprocessParameter.SetPS(RHICmdList, ShaderRHI, Context, TStaticSamplerState<SF_Bilinear,AM_Clamp,AM_Clamp,AM_Clamp>::GetRHI());
 			
 		{
 			FLinearColor Col = Settings.SceneColorTint;
 			FVector4 ColorScale(Col.R, Col.G, Col.B, 0);
-			SetShaderValue(ShaderRHI, ColorScale0, ColorScale);
+			SetShaderValue(RHICmdList, ShaderRHI, ColorScale0, ColorScale);
 		}
 		
 		{
 			FLinearColor Col = FLinearColor::White * Settings.BloomIntensity;
 			FVector4 ColorScale(Col.R, Col.G, Col.B, 0);
-			SetShaderValue(ShaderRHI, ColorScale1, ColorScale);
+			SetShaderValue(RHICmdList, ShaderRHI, ColorScale1, ColorScale);
 		}
 
 		{
 			UTexture2D* NoiseTextureValue = GEngine->HighFrequencyNoiseTexture;
 
-			SetTextureParameter(ShaderRHI, NoiseTexture, NoiseTextureSampler, TStaticSamplerState<SF_Point,AM_Wrap,AM_Wrap,AM_Wrap>::GetRHI(), NoiseTextureValue->Resource->TextureRHI);
+			SetTextureParameter(RHICmdList, ShaderRHI, NoiseTexture, NoiseTextureSampler, TStaticSamplerState<SF_Point,AM_Wrap,AM_Wrap,AM_Wrap>::GetRHI(), NoiseTextureValue->Resource->TextureRHI);
 		}
 
 		{
@@ -861,7 +860,7 @@ public:
 			// we assume the this pass runs in 1:1 pixel
 			FVector2D TexScaleValue = FVector2D(InputDesc->Extent) / FVector2D(Context.View.ViewRect.Size());
 
-			SetShaderValue(ShaderRHI, TexScale, TexScaleValue);
+			SetShaderValue(RHICmdList, ShaderRHI, TexScale, TexScaleValue);
 		}
 		
 		FVector4 VignetteColorIntensityValue;
@@ -869,27 +868,30 @@ public:
 		VignetteColorIntensityValue.Y = Settings.VignetteColor.G;
 		VignetteColorIntensityValue.Z = Settings.VignetteColor.B;
 		VignetteColorIntensityValue.W = Settings.VignetteIntensity;
-		SetShaderValue(ShaderRHI, VignetteColorIntensity, VignetteColorIntensityValue);
+		SetShaderValue(RHICmdList, ShaderRHI, VignetteColorIntensity, VignetteColorIntensityValue);
 
 		FVector GrainValue;
 		GrainPostSettings(&GrainValue, &Settings);
-		SetShaderValue(ShaderRHI, GrainScaleBiasJitter, GrainValue);
+		SetShaderValue(RHICmdList, ShaderRHI, GrainScaleBiasJitter, GrainValue);
 
+		const TShaderUniformBufferParameter<FBloomDirtMaskParameters>& BloomDirtMaskParam = GetUniformBufferParameter<FBloomDirtMaskParameters>();
+		if (BloomDirtMaskParam.IsBound())
 		{
-			float ExposureScale = FRCPassPostProcessEyeAdaptation::ComputeExposureScaleValue(Context.View);
+			FBloomDirtMaskParameters BloomDirtMaskParams;
 
+			float ExposureScale = FRCPassPostProcessEyeAdaptation::ComputeExposureScaleValue(Context.View);
 			FLinearColor Col = Settings.BloomDirtMaskTint * Settings.BloomDirtMaskIntensity;
-			FVector4 ColorScale(Col.R, Col.G, Col.B, ExposureScale);
-			SetShaderValue(ShaderRHI, BloomDirtMaskTint, ColorScale);
-		}
+			BloomDirtMaskParams.Tint = FVector4(Col.R, Col.G, Col.B, ExposureScale);
 
-		{
-			FTextureRHIParamRef BloomDirtMaskValue = GSystemTextures.BlackDummy->GetRenderTargetItem().TargetableTexture;
+			BloomDirtMaskParams.Mask = GSystemTextures.BlackDummy->GetRenderTargetItem().TargetableTexture;
 			if(Settings.BloomDirtMask && Settings.BloomDirtMask->Resource)
 			{
-				BloomDirtMaskValue = Settings.BloomDirtMask->Resource->TextureRHI;
+				BloomDirtMaskParams.Mask = Settings.BloomDirtMask->Resource->TextureRHI;
 			}
-			SetTextureParameter(ShaderRHI, BloomDirtMask, BloomDirtMaskSampler, TStaticSamplerState<SF_Bilinear,AM_Wrap,AM_Wrap,AM_Wrap>::GetRHI(), BloomDirtMaskValue);
+			BloomDirtMaskParams.MaskSampler = TStaticSamplerState<SF_Bilinear,AM_Wrap,AM_Wrap,AM_Wrap>::GetRHI();
+
+			FUniformBufferRHIRef BloomDirtMaskUB = TUniformBufferRef<FBloomDirtMaskParameters>::CreateUniformBufferImmediate(BloomDirtMaskParams, UniformBuffer_SingleDraw);
+			SetUniformBufferParameter(RHICmdList, ShaderRHI, BloomDirtMaskParam, BloomDirtMaskUB);
 		}
 		
 		// volume texture LUT
@@ -910,7 +912,7 @@ public:
 
 						const FTextureRHIRef& SrcTexture = InputPooledElement->GetRenderTargetItem().ShaderResourceTexture;
 					
-						SetTextureParameter(ShaderRHI, ColorGradingLUT, ColorGradingLUTSampler, TStaticSamplerState<SF_Bilinear,AM_Clamp,AM_Clamp,AM_Clamp>::GetRHI(), SrcTexture);
+						SetTextureParameter(RHICmdList, ShaderRHI, ColorGradingLUT, ColorGradingLUTSampler, TStaticSamplerState<SF_Bilinear,AM_Clamp,AM_Clamp,AM_Clamp>::GetRHI(), SrcTexture);
 					}
 				}
 			}
@@ -920,20 +922,20 @@ public:
 			FVector2D InvDisplayGammaValue;
 			InvDisplayGammaValue.X = 1.0f / ViewFamily.RenderTarget->GetDisplayGamma();
 			InvDisplayGammaValue.Y = 2.2f / ViewFamily.RenderTarget->GetDisplayGamma();
-			SetShaderValue(ShaderRHI, InverseGamma, InvDisplayGammaValue);
+			SetShaderValue(RHICmdList, ShaderRHI, InverseGamma, InvDisplayGammaValue);
 		}
 
 		{
 			FVector4 Constants[8];
 			FilmPostSetConstants(Constants, TonemapperConfBitmaskPC[ConfigIndex], &Context.View.FinalPostProcessSettings, false);
-			SetShaderValue(ShaderRHI, ColorMatrixR_ColorCurveCd1, Constants[0]);
-			SetShaderValue(ShaderRHI, ColorMatrixG_ColorCurveCd3Cm3, Constants[1]);
-			SetShaderValue(ShaderRHI, ColorMatrixB_ColorCurveCm2, Constants[2]); 
-			SetShaderValue(ShaderRHI, ColorCurve_Cm0Cd0_Cd2_Ch0Cm1_Ch3, Constants[3]); 
-			SetShaderValue(ShaderRHI, ColorCurve_Ch1_Ch2, Constants[4]);
-			SetShaderValue(ShaderRHI, ColorShadow_Luma, Constants[5]);
-			SetShaderValue(ShaderRHI, ColorShadow_Tint1, Constants[6]);
-			SetShaderValue(ShaderRHI, ColorShadow_Tint2, Constants[7]);
+			SetShaderValue(RHICmdList, ShaderRHI, ColorMatrixR_ColorCurveCd1, Constants[0]);
+			SetShaderValue(RHICmdList, ShaderRHI, ColorMatrixG_ColorCurveCd3Cm3, Constants[1]);
+			SetShaderValue(RHICmdList, ShaderRHI, ColorMatrixB_ColorCurveCm2, Constants[2]); 
+			SetShaderValue(RHICmdList, ShaderRHI, ColorCurve_Cm0Cd0_Cd2_Ch0Cm1_Ch3, Constants[3]); 
+			SetShaderValue(RHICmdList, ShaderRHI, ColorCurve_Ch1_Ch2, Constants[4]);
+			SetShaderValue(RHICmdList, ShaderRHI, ColorShadow_Luma, Constants[5]);
+			SetShaderValue(RHICmdList, ShaderRHI, ColorShadow_Tint1, Constants[6]);
+			SetShaderValue(RHICmdList, ShaderRHI, ColorShadow_Tint2, Constants[7]);
 		}
 	}
 	
@@ -967,7 +969,7 @@ FRCPassPostProcessTonemap::FRCPassPostProcessTonemap(bool bInDoGammaOnly)
 }
 
 template <uint32 ConfigIndex>
-static void SetShaderTempl(const FRenderingCompositePassContext& Context)
+static void SetShaderTempl(FRHICommandList* RHICmdList, const FRenderingCompositePassContext& Context)
 {
 	TShaderMapRef<FPostProcessTonemapVS> VertexShader(GetGlobalShaderMap());
 	TShaderMapRef<FPostProcessTonemapPS<ConfigIndex> > PixelShader(GetGlobalShaderMap());
@@ -976,8 +978,8 @@ static void SetShaderTempl(const FRenderingCompositePassContext& Context)
 
 	SetGlobalBoundShaderState(BoundShaderState, GFilterVertexDeclaration.VertexDeclarationRHI, *VertexShader, *PixelShader);
 
-	VertexShader->SetVS(Context);
-	PixelShader->SetPS(Context);
+	VertexShader->SetVS(RHICmdList, Context);
+	PixelShader->SetPS(RHICmdList, Context);
 }
 
 void FRCPassPostProcessTonemap::Process(FRenderingCompositePassContext& Context)
@@ -1004,6 +1006,9 @@ void FRCPassPostProcessTonemap::Process(FRenderingCompositePassContext& Context)
 
 	const FSceneRenderTargetItem& DestRenderTarget = PassOutputs[0].RequestSurface(Context);
 
+	//@todo-rco: RHIPacketList
+	FRHICommandList* RHICmdList = nullptr;
+
 	// Set the view family's render target/viewport.
 	RHISetRenderTarget(DestRenderTarget.TargetableTexture, FTextureRHIParamRef());	
 
@@ -1022,17 +1027,17 @@ void FRCPassPostProcessTonemap::Process(FRenderingCompositePassContext& Context)
 
 	switch(ConfigIndex)
 	{
-		case 0:	SetShaderTempl<0>(Context); break;
-		case 1:	SetShaderTempl<1>(Context);	break;
-		case 2: SetShaderTempl<2>(Context); break;
-		case 3: SetShaderTempl<3>(Context); break;
-		case 4: SetShaderTempl<4>(Context); break;
-		case 5: SetShaderTempl<5>(Context); break;
-		case 6: SetShaderTempl<6>(Context); break;
-		case 7: SetShaderTempl<7>(Context); break;
-		case 8: SetShaderTempl<8>(Context); break;
-		case 9: SetShaderTempl<9>(Context); break;
-		case 10: SetShaderTempl<10>(Context); break;
+		case 0:	SetShaderTempl<0>(RHICmdList, Context); break;
+		case 1:	SetShaderTempl<1>(RHICmdList, Context);	break;
+		case 2: SetShaderTempl<2>(RHICmdList, Context); break;
+		case 3: SetShaderTempl<3>(RHICmdList, Context); break;
+		case 4: SetShaderTempl<4>(RHICmdList, Context); break;
+		case 5: SetShaderTempl<5>(RHICmdList, Context); break;
+		case 6: SetShaderTempl<6>(RHICmdList, Context); break;
+		case 7: SetShaderTempl<7>(RHICmdList, Context); break;
+		case 8: SetShaderTempl<8>(RHICmdList, Context); break;
+		case 9: SetShaderTempl<9>(RHICmdList, Context); break;
+		case 10: SetShaderTempl<10>(RHICmdList, Context); break;
 		default:
 			check(0);
 	}
@@ -1177,29 +1182,29 @@ public:
 		return bShaderHasOutdatedParameters;
 	}
 
-	void SetPS(const FRenderingCompositePassContext& Context)
+	void SetPS(FRHICommandList* RHICmdList, const FRenderingCompositePassContext& Context)
 	{
 		const FPostProcessSettings& Settings = Context.View.FinalPostProcessSettings;
 		const FSceneViewFamily& ViewFamily = *(Context.View.Family);
 
 		const FPixelShaderRHIParamRef ShaderRHI = GetPixelShader();
 		
-		FGlobalShader::SetParameters(ShaderRHI, Context.View);
+		FGlobalShader::SetParameters(RHICmdList, ShaderRHI, Context.View);
 
-		PostprocessParameter.SetPS(ShaderRHI, Context, TStaticSamplerState<SF_Bilinear,AM_Clamp,AM_Clamp,AM_Clamp>::GetRHI());
+		PostprocessParameter.SetPS(RHICmdList, ShaderRHI, Context, TStaticSamplerState<SF_Bilinear,AM_Clamp,AM_Clamp,AM_Clamp>::GetRHI());
 			
-		SetShaderValue(ShaderRHI, OverlayColor, Context.View.OverlayColor);
+		SetShaderValue(RHICmdList, ShaderRHI, OverlayColor, Context.View.OverlayColor);
 
 		{
 			FLinearColor Col = Settings.SceneColorTint;
 			FVector4 ColorScale(Col.R, Col.G, Col.B, 0);
-			SetShaderValue(ShaderRHI, ColorScale0, ColorScale);
+			SetShaderValue(RHICmdList, ShaderRHI, ColorScale0, ColorScale);
 		}
 		
 		{
 			FLinearColor Col = FLinearColor::White * Settings.BloomIntensity;
 			FVector4 ColorScale(Col.R, Col.G, Col.B, 0);
-			SetShaderValue(ShaderRHI, ColorScale1, ColorScale);
+			SetShaderValue(RHICmdList, ShaderRHI, ColorScale1, ColorScale);
 		}
 
 		{
@@ -1208,7 +1213,7 @@ public:
 			// we assume the this pass runs in 1:1 pixel
 			FVector2D TexScaleValue = FVector2D(InputDesc->Extent) / FVector2D(Context.View.ViewRect.Size());
 
-			SetShaderValue(ShaderRHI, TexScale, TexScaleValue);
+			SetShaderValue(RHICmdList, ShaderRHI, TexScale, TexScaleValue);
 		}
 
 		FVector4 VignetteColorIntensityValue;
@@ -1216,30 +1221,30 @@ public:
 		VignetteColorIntensityValue.Y = Settings.VignetteColor.G;
 		VignetteColorIntensityValue.Z = Settings.VignetteColor.B;
 		VignetteColorIntensityValue.W = Settings.VignetteIntensity;
-		SetShaderValue(ShaderRHI, VignetteColorIntensity, VignetteColorIntensityValue);
+		SetShaderValue(RHICmdList, ShaderRHI, VignetteColorIntensity, VignetteColorIntensityValue);
 
 		FVector GrainValue;
 		GrainPostSettings(&GrainValue, &Settings);
-		SetShaderValue(ShaderRHI, GrainScaleBiasJitter, GrainValue);
+		SetShaderValue(RHICmdList, ShaderRHI, GrainScaleBiasJitter, GrainValue);
 
 		{
 			FVector2D InvDisplayGammaValue;
 			InvDisplayGammaValue.X = 1.0f / ViewFamily.RenderTarget->GetDisplayGamma();
 			InvDisplayGammaValue.Y = 2.2f / ViewFamily.RenderTarget->GetDisplayGamma();
-			SetShaderValue(ShaderRHI, InverseGamma, InvDisplayGammaValue);
+			SetShaderValue(RHICmdList, ShaderRHI, InverseGamma, InvDisplayGammaValue);
 		}
 
 		{
 			FVector4 Constants[8];
 			FilmPostSetConstants(Constants, TonemapperConfBitmaskMobile[ConfigIndex], &Context.View.FinalPostProcessSettings, true);
-			SetShaderValue(ShaderRHI, ColorMatrixR_ColorCurveCd1, Constants[0]);
-			SetShaderValue(ShaderRHI, ColorMatrixG_ColorCurveCd3Cm3, Constants[1]);
-			SetShaderValue(ShaderRHI, ColorMatrixB_ColorCurveCm2, Constants[2]); 
-			SetShaderValue(ShaderRHI, ColorCurve_Cm0Cd0_Cd2_Ch0Cm1_Ch3, Constants[3]); 
-			SetShaderValue(ShaderRHI, ColorCurve_Ch1_Ch2, Constants[4]);
-			SetShaderValue(ShaderRHI, ColorShadow_Luma, Constants[5]);
-			SetShaderValue(ShaderRHI, ColorShadow_Tint1, Constants[6]);
-			SetShaderValue(ShaderRHI, ColorShadow_Tint2, Constants[7]);
+			SetShaderValue(RHICmdList, ShaderRHI, ColorMatrixR_ColorCurveCd1, Constants[0]);
+			SetShaderValue(RHICmdList, ShaderRHI, ColorMatrixG_ColorCurveCd3Cm3, Constants[1]);
+			SetShaderValue(RHICmdList, ShaderRHI, ColorMatrixB_ColorCurveCm2, Constants[2]); 
+			SetShaderValue(RHICmdList, ShaderRHI, ColorCurve_Cm0Cd0_Cd2_Ch0Cm1_Ch3, Constants[3]); 
+			SetShaderValue(RHICmdList, ShaderRHI, ColorCurve_Ch1_Ch2, Constants[4]);
+			SetShaderValue(RHICmdList, ShaderRHI, ColorShadow_Luma, Constants[5]);
+			SetShaderValue(RHICmdList, ShaderRHI, ColorShadow_Tint1, Constants[6]);
+			SetShaderValue(RHICmdList, ShaderRHI, ColorShadow_Tint2, Constants[7]);
 		}
 	}
 	
@@ -1289,18 +1294,18 @@ public:
 		GrainRandomFull.Bind(Initializer.ParameterMap, TEXT("GrainRandomFull"));
 	}
 
-	void SetVS(const FRenderingCompositePassContext& Context)
+	void SetVS(FRHICommandList* RHICmdList, const FRenderingCompositePassContext& Context)
 	{
 		const FVertexShaderRHIParamRef ShaderRHI = GetVertexShader();
-		FGlobalShader::SetParameters(ShaderRHI, Context.View);
+		FGlobalShader::SetParameters(RHICmdList, ShaderRHI, Context.View);
 
-		PostprocessParameter.SetVS(ShaderRHI, Context, TStaticSamplerState<SF_Bilinear,AM_Clamp,AM_Clamp,AM_Clamp>::GetRHI());
+		PostprocessParameter.SetVS(RHICmdList, ShaderRHI, Context, TStaticSamplerState<SF_Bilinear,AM_Clamp,AM_Clamp,AM_Clamp>::GetRHI());
 
 		FVector GrainRandomFullValue;
 		GrainRandomFromFrame(&GrainRandomFullValue, Context.View.FrameNumber);
 		// TODO: Don't use full on mobile with framebuffer fetch.
 		GrainRandomFullValue.Z = bUsedFramebufferFetch ? 0.0f : 1.0f;
-		SetShaderValue(ShaderRHI, GrainRandomFull, GrainRandomFullValue);
+		SetShaderValue(RHICmdList, ShaderRHI, GrainRandomFull, GrainRandomFullValue);
 	}
 	
 	virtual bool Serialize(FArchive& Ar)
@@ -1315,7 +1320,7 @@ IMPLEMENT_SHADER_TYPE(,FPostProcessTonemapVS_ES2,TEXT("PostProcessTonemap"),TEXT
 
 
 template <uint32 ConfigIndex>
-static void SetShaderTemplES2(const FRenderingCompositePassContext& Context, bool bUsedFramebufferFetch)
+static void SetShaderTemplES2(FRHICommandList* RHICmdList, const FRenderingCompositePassContext& Context, bool bUsedFramebufferFetch)
 {
 	TShaderMapRef<FPostProcessTonemapVS_ES2> VertexShader(GetGlobalShaderMap());
 	TShaderMapRef<FPostProcessTonemapPS_ES2<ConfigIndex> > PixelShader(GetGlobalShaderMap());
@@ -1326,8 +1331,8 @@ static void SetShaderTemplES2(const FRenderingCompositePassContext& Context, boo
 
 	SetGlobalBoundShaderState(BoundShaderState, GFilterVertexDeclaration.VertexDeclarationRHI, *VertexShader, *PixelShader);
 
-	VertexShader->SetVS(Context);
-	PixelShader->SetPS(Context);
+	VertexShader->SetVS(RHICmdList, Context);
+	PixelShader->SetPS(RHICmdList, Context);
 }
 
 void FRCPassPostProcessTonemapES2::Process(FRenderingCompositePassContext& Context)
@@ -1354,6 +1359,9 @@ void FRCPassPostProcessTonemapES2::Process(FRenderingCompositePassContext& Conte
 
 	const FSceneRenderTargetItem& DestRenderTarget = PassOutputs[0].RequestSurface(Context);
 
+	//@todo-rco: RHIPacketList
+	FRHICommandList* RHICmdList = nullptr;
+
 	// Set the view family's render target/viewport.
 	RHISetRenderTarget(DestRenderTarget.TargetableTexture, FTextureRHIParamRef());	
 
@@ -1369,35 +1377,35 @@ void FRCPassPostProcessTonemapES2::Process(FRenderingCompositePassContext& Conte
 
 	switch(ConfigIndex)
 	{
-		case 0:	SetShaderTemplES2<0>(Context, bUsedFramebufferFetch); break;
-		case 1:	SetShaderTemplES2<1>(Context, bUsedFramebufferFetch); break;
-		case 2:	SetShaderTemplES2<2>(Context, bUsedFramebufferFetch); break;
-		case 3:	SetShaderTemplES2<3>(Context, bUsedFramebufferFetch); break;
-		case 4:	SetShaderTemplES2<4>(Context, bUsedFramebufferFetch); break;
-		case 5:	SetShaderTemplES2<5>(Context, bUsedFramebufferFetch); break;
-		case 6:	SetShaderTemplES2<6>(Context, bUsedFramebufferFetch); break;
-		case 7:	SetShaderTemplES2<7>(Context, bUsedFramebufferFetch); break;
-		case 8:	SetShaderTemplES2<8>(Context, bUsedFramebufferFetch); break;
-		case 9:	SetShaderTemplES2<9>(Context, bUsedFramebufferFetch); break;
-		case 10: SetShaderTemplES2<10>(Context, bUsedFramebufferFetch); break;
-		case 11: SetShaderTemplES2<11>(Context, bUsedFramebufferFetch); break;
-		case 12: SetShaderTemplES2<12>(Context, bUsedFramebufferFetch); break;
-		case 13: SetShaderTemplES2<13>(Context, bUsedFramebufferFetch); break;
-		case 14: SetShaderTemplES2<14>(Context, bUsedFramebufferFetch); break;
-		case 15: SetShaderTemplES2<15>(Context, bUsedFramebufferFetch); break;
-		case 16: SetShaderTemplES2<16>(Context, bUsedFramebufferFetch); break;
-		case 17: SetShaderTemplES2<17>(Context, bUsedFramebufferFetch); break;
-		case 18: SetShaderTemplES2<18>(Context, bUsedFramebufferFetch); break;
-		case 19: SetShaderTemplES2<19>(Context, bUsedFramebufferFetch); break;
-		case 20: SetShaderTemplES2<20>(Context, bUsedFramebufferFetch); break;
-		case 21: SetShaderTemplES2<21>(Context, bUsedFramebufferFetch); break;
-		case 22: SetShaderTemplES2<22>(Context, bUsedFramebufferFetch); break;
-		case 23: SetShaderTemplES2<23>(Context, bUsedFramebufferFetch); break;
-		case 24: SetShaderTemplES2<24>(Context, bUsedFramebufferFetch); break;
-		case 25: SetShaderTemplES2<25>(Context, bUsedFramebufferFetch); break;
-		case 26: SetShaderTemplES2<26>(Context, bUsedFramebufferFetch); break;
-		case 27: SetShaderTemplES2<27>(Context, bUsedFramebufferFetch); break;
-		case 28: SetShaderTemplES2<28>(Context, bUsedFramebufferFetch); break;
+		case 0:	SetShaderTemplES2<0>(RHICmdList, Context, bUsedFramebufferFetch); break;
+		case 1:	SetShaderTemplES2<1>(RHICmdList, Context, bUsedFramebufferFetch); break;
+		case 2:	SetShaderTemplES2<2>(RHICmdList, Context, bUsedFramebufferFetch); break;
+		case 3:	SetShaderTemplES2<3>(RHICmdList, Context, bUsedFramebufferFetch); break;
+		case 4:	SetShaderTemplES2<4>(RHICmdList, Context, bUsedFramebufferFetch); break;
+		case 5:	SetShaderTemplES2<5>(RHICmdList, Context, bUsedFramebufferFetch); break;
+		case 6:	SetShaderTemplES2<6>(RHICmdList, Context, bUsedFramebufferFetch); break;
+		case 7:	SetShaderTemplES2<7>(RHICmdList, Context, bUsedFramebufferFetch); break;
+		case 8:	SetShaderTemplES2<8>(RHICmdList, Context, bUsedFramebufferFetch); break;
+		case 9:	SetShaderTemplES2<9>(RHICmdList, Context, bUsedFramebufferFetch); break;
+		case 10: SetShaderTemplES2<10>(RHICmdList, Context, bUsedFramebufferFetch); break;
+		case 11: SetShaderTemplES2<11>(RHICmdList, Context, bUsedFramebufferFetch); break;
+		case 12: SetShaderTemplES2<12>(RHICmdList, Context, bUsedFramebufferFetch); break;
+		case 13: SetShaderTemplES2<13>(RHICmdList, Context, bUsedFramebufferFetch); break;
+		case 14: SetShaderTemplES2<14>(RHICmdList, Context, bUsedFramebufferFetch); break;
+		case 15: SetShaderTemplES2<15>(RHICmdList, Context, bUsedFramebufferFetch); break;
+		case 16: SetShaderTemplES2<16>(RHICmdList, Context, bUsedFramebufferFetch); break;
+		case 17: SetShaderTemplES2<17>(RHICmdList, Context, bUsedFramebufferFetch); break;
+		case 18: SetShaderTemplES2<18>(RHICmdList, Context, bUsedFramebufferFetch); break;
+		case 19: SetShaderTemplES2<19>(RHICmdList, Context, bUsedFramebufferFetch); break;
+		case 20: SetShaderTemplES2<20>(RHICmdList, Context, bUsedFramebufferFetch); break;
+		case 21: SetShaderTemplES2<21>(RHICmdList, Context, bUsedFramebufferFetch); break;
+		case 22: SetShaderTemplES2<22>(RHICmdList, Context, bUsedFramebufferFetch); break;
+		case 23: SetShaderTemplES2<23>(RHICmdList, Context, bUsedFramebufferFetch); break;
+		case 24: SetShaderTemplES2<24>(RHICmdList, Context, bUsedFramebufferFetch); break;
+		case 25: SetShaderTemplES2<25>(RHICmdList, Context, bUsedFramebufferFetch); break;
+		case 26: SetShaderTemplES2<26>(RHICmdList, Context, bUsedFramebufferFetch); break;
+		case 27: SetShaderTemplES2<27>(RHICmdList, Context, bUsedFramebufferFetch); break;
+		case 28: SetShaderTemplES2<28>(RHICmdList, Context, bUsedFramebufferFetch); break;
 		default:
 			check(0);
 	}
diff --git a/Engine/Source/Runtime/Renderer/Private/PostProcess/PostProcessTonemap.h b/Engine/Source/Runtime/Renderer/Private/PostProcess/PostProcessTonemap.h
--- a/Engine/Source/Runtime/Renderer/Private/PostProcess/PostProcessTonemap.h
+++ b/Engine/Source/Runtime/Renderer/Private/PostProcess/PostProcessTonemap.h
@@ -100,17 +100,17 @@ public:
 		GrainRandomFull.Bind(Initializer.ParameterMap, TEXT("GrainRandomFull"));
 	}
 
-	void SetVS(const FRenderingCompositePassContext& Context)
+	void SetVS(FRHICommandList* RHICmdList, const FRenderingCompositePassContext& Context)
 	{
 		const FVertexShaderRHIParamRef ShaderRHI = GetVertexShader();
 
-		FGlobalShader::SetParameters(ShaderRHI, Context.View);
+		FGlobalShader::SetParameters(RHICmdList, ShaderRHI, Context.View);
 
-		PostprocessParameter.SetVS(ShaderRHI, Context, TStaticSamplerState<SF_Bilinear,AM_Clamp,AM_Clamp,AM_Clamp>::GetRHI());
+		PostprocessParameter.SetVS(RHICmdList, ShaderRHI, Context, TStaticSamplerState<SF_Bilinear,AM_Clamp,AM_Clamp,AM_Clamp>::GetRHI());
 
 		FVector GrainRandomFullValue;
 		GrainRandomFromFrame(&GrainRandomFullValue, Context.View.FrameNumber);
-		SetShaderValue(ShaderRHI, GrainRandomFull, GrainRandomFullValue);
+		SetShaderValue(RHICmdList, ShaderRHI, GrainRandomFull, GrainRandomFullValue);
 
 		if(EyeAdaptation.IsBound())
 		{
@@ -118,12 +118,12 @@ public:
 
 			if(EyeAdaptationRT)
 			{
-				SetTextureParameter(ShaderRHI, EyeAdaptation, EyeAdaptationRT->GetRenderTargetItem().TargetableTexture);
+				SetTextureParameter(RHICmdList, ShaderRHI, EyeAdaptation, EyeAdaptationRT->GetRenderTargetItem().TargetableTexture);
 			}
 			else
 			{
 				// some views don't have a state, thumbnail rendering?
-				SetTextureParameter(ShaderRHI, EyeAdaptation, GWhiteTexture->TextureRHI);
+				SetTextureParameter(RHICmdList, ShaderRHI, EyeAdaptation, GWhiteTexture->TextureRHI);
 			}
 		}
 	}
diff --git a/Engine/Source/Runtime/Renderer/Private/PostProcess/PostProcessUIBlur.cpp b/Engine/Source/Runtime/Renderer/Private/PostProcess/PostProcessUIBlur.cpp
--- a/Engine/Source/Runtime/Renderer/Private/PostProcess/PostProcessUIBlur.cpp
+++ b/Engine/Source/Runtime/Renderer/Private/PostProcess/PostProcessUIBlur.cpp
@@ -34,13 +34,13 @@ public:
 		PostprocessParameter.Bind(Initializer.ParameterMap);
 	}
 
-	void SetPS(const FRenderingCompositePassContext& Context)
+	void SetPS(FRHICommandList* RHICmdList, const FRenderingCompositePassContext& Context)
 	{
 		const FPixelShaderRHIParamRef ShaderRHI = GetPixelShader();
 
-		FGlobalShader::SetParameters(ShaderRHI, Context.View);
+		FGlobalShader::SetParameters(RHICmdList, ShaderRHI, Context.View);
 
-		PostprocessParameter.SetPS(ShaderRHI, Context, TStaticSamplerState<SF_Bilinear,AM_Clamp,AM_Clamp,AM_Clamp>::GetRHI());
+		PostprocessParameter.SetPS(RHICmdList, ShaderRHI, Context, TStaticSamplerState<SF_Bilinear,AM_Clamp,AM_Clamp,AM_Clamp>::GetRHI());
 	}
 	
 	// FShader interface.
@@ -76,6 +76,9 @@ void FRCPassPostProcessUIBlur::Process(FRenderingCompositePassContext& Context)
 	TRefCountPtr<IPooledRenderTarget>& Dest = PassInputs[0].GetOutput()->PooledRenderTarget;
 	const FSceneRenderTargetItem& DestRenderTarget = Dest->GetRenderTargetItem();
 
+	//@todo-rco: RHIPacketList
+	FRHICommandList* RHICmdList = nullptr;
+
 	// Set the view family's render target/viewport.
 	RHISetRenderTarget(DestRenderTarget.TargetableTexture, FTextureRHIRef());	
 
@@ -91,7 +94,7 @@ void FRCPassPostProcessUIBlur::Process(FRenderingCompositePassContext& Context)
 
 	SetGlobalBoundShaderState(BoundShaderState, GFilterVertexDeclaration.VertexDeclarationRHI, *VertexShader, *PixelShader);
 
-	PixelShader->SetPS(Context);
+	PixelShader->SetPS(RHICmdList, Context);
 
 	// todo: we only expect few quads but ideally we have one one draw call for all quads
 	int32 RectCount = Context.View.UIBlurOverrideRectangles.Num();
diff --git a/Engine/Source/Runtime/Renderer/Private/PostProcess/PostProcessUpscale.cpp b/Engine/Source/Runtime/Renderer/Private/PostProcess/PostProcessUpscale.cpp
--- a/Engine/Source/Runtime/Renderer/Private/PostProcess/PostProcessUpscale.cpp
+++ b/Engine/Source/Runtime/Renderer/Private/PostProcess/PostProcessUpscale.cpp
@@ -25,7 +25,8 @@ class FPostProcessUpscalePS : public FGlobalShader
 
 	static void ModifyCompilationEnvironment(EShaderPlatform Platform, FShaderCompilerEnvironment& OutEnvironment)
 	{
-		FGlobalShader::ModifyCompilationEnvironment(Platform, OutEnvironment);
+		FGlobalShader::ModifyCompilationEnvironment(Platform,OutEnvironment);
+		FDeferredPixelShaderParameters::ModifyCompilationEnvironment(Platform,OutEnvironment);
 		OutEnvironment.SetDefine(TEXT("METHOD"), Method);
 	}
 
@@ -46,18 +47,18 @@ public:
 		UpscaleSoftness.Bind(Initializer.ParameterMap,TEXT("UpscaleSoftness"));
 	}
 
-	void SetPS(const FRenderingCompositePassContext& Context)
+	void SetPS(FRHICommandList* RHICmdList, const FRenderingCompositePassContext& Context)
 	{
 		const FPixelShaderRHIParamRef ShaderRHI = GetPixelShader();
 		
-		FGlobalShader::SetParameters(ShaderRHI, Context.View);
+		FGlobalShader::SetParameters(RHICmdList, ShaderRHI, Context.View);
 
 		FSamplerStateRHIParamRef FilterTable[2];
 		FilterTable[0] = TStaticSamplerState<SF_Bilinear,AM_Clamp,AM_Clamp,AM_Clamp>::GetRHI();
 		FilterTable[1] = TStaticSamplerState<SF_Point,AM_Clamp,AM_Clamp,AM_Clamp>::GetRHI();
 			
-		PostprocessParameter.SetPS(ShaderRHI, Context, 0, false, FilterTable);
-		DeferredParameters.Set(ShaderRHI, Context.View);
+		PostprocessParameter.SetPS(RHICmdList, ShaderRHI, Context, 0, false, FilterTable);
+		DeferredParameters.Set(RHICmdList, ShaderRHI, Context.View);
 
 		// If the method needs softness value
 		if(Method == 2)
@@ -66,7 +67,7 @@ public:
 
 			float UpscaleSoftnessValue = FMath::Clamp(CVar->GetValueOnRenderThread(), 0.0f, 1.0f);
 
-			SetShaderValue(ShaderRHI, UpscaleSoftness, UpscaleSoftnessValue);
+			SetShaderValue(RHICmdList, ShaderRHI, UpscaleSoftness, UpscaleSoftnessValue);
 		}
 	}
 	
@@ -107,7 +108,7 @@ FRCPassPostProcessUpscale::FRCPassPostProcessUpscale(uint32 InUpscaleMethod)
 }
 
 template <uint32 Method>
-void FRCPassPostProcessUpscale::SetShader(const FRenderingCompositePassContext& Context)
+void FRCPassPostProcessUpscale::SetShader(FRHICommandList* RHICmdList, const FRenderingCompositePassContext& Context)
 {
 	TShaderMapRef<FPostProcessVS> VertexShader(GetGlobalShaderMap());
 	TShaderMapRef<FPostProcessUpscalePS<Method> > PixelShader(GetGlobalShaderMap());
@@ -116,7 +117,7 @@ void FRCPassPostProcessUpscale::SetShader(const FRenderingCompositePassContext&
 
 	SetGlobalBoundShaderState(BoundShaderState, GFilterVertexDeclaration.VertexDeclarationRHI, *VertexShader, *PixelShader);
 
-	PixelShader->SetPS(Context);
+	PixelShader->SetPS(RHICmdList, Context);
 }
 
 void FRCPassPostProcessUpscale::Process(FRenderingCompositePassContext& Context)
@@ -139,6 +140,9 @@ void FRCPassPostProcessUpscale::Process(FRenderingCompositePassContext& Context)
 
 	const FSceneRenderTargetItem& DestRenderTarget = PassOutputs[0].RequestSurface(Context);
 
+	//@todo-rco: RHIPacketList
+	FRHICommandList* RHICmdList = nullptr;
+
 	// Set the view family's render target/viewport.
 	RHISetRenderTarget(DestRenderTarget.TargetableTexture, FTextureRHIRef());	
 	Context.SetViewportAndCallRHI(DestRect);
@@ -151,16 +155,16 @@ void FRCPassPostProcessUpscale::Process(FRenderingCompositePassContext& Context)
 	switch (UpscaleMethod)
 	{
 		case 0:
-			SetShader<0>(Context);
+			SetShader<0>(RHICmdList, Context);
 		break;
 		case 1:
-			SetShader<1>(Context);
+			SetShader<1>(RHICmdList, Context);
 		break;
 		case 2:
-			SetShader<2>(Context);
+			SetShader<2>(RHICmdList, Context);
 		break;
 		case 3:
-			SetShader<3>(Context);
+			SetShader<3>(RHICmdList, Context);
 		break;
 		default:
 			checkNoEntry();
diff --git a/Engine/Source/Runtime/Renderer/Private/PostProcess/PostProcessUpscale.h b/Engine/Source/Runtime/Renderer/Private/PostProcess/PostProcessUpscale.h
--- a/Engine/Source/Runtime/Renderer/Private/PostProcess/PostProcessUpscale.h
+++ b/Engine/Source/Runtime/Renderer/Private/PostProcess/PostProcessUpscale.h
@@ -27,7 +27,7 @@ public:
 	virtual void Release() OVERRIDE { delete this; }
 	FPooledRenderTargetDesc ComputeOutputDesc(EPassOutputId InPassOutputId) const;
 private:
-	template <uint32 Method> static void SetShader(const FRenderingCompositePassContext& Context);
+	template <uint32 Method> static void SetShader(FRHICommandList* RHICmdList, const FRenderingCompositePassContext& Context);
 	uint32 UpscaleMethod;
 };
 
diff --git a/Engine/Source/Runtime/Renderer/Private/PostProcess/PostProcessVisualizeBuffer.cpp b/Engine/Source/Runtime/Renderer/Private/PostProcess/PostProcessVisualizeBuffer.cpp
--- a/Engine/Source/Runtime/Renderer/Private/PostProcess/PostProcessVisualizeBuffer.cpp
+++ b/Engine/Source/Runtime/Renderer/Private/PostProcess/PostProcessVisualizeBuffer.cpp
@@ -26,6 +26,7 @@ class FPostProcessVisualizeBufferPS : public FGlobalShader
 	static void ModifyCompilationEnvironment(EShaderPlatform Platform, FShaderCompilerEnvironment& OutEnvironment)
 	{
 		FGlobalShader::ModifyCompilationEnvironment(Platform, OutEnvironment);
+		FDeferredPixelShaderParameters::ModifyCompilationEnvironment(Platform,OutEnvironment);
 		OutEnvironment.SetDefine(TEXT("DRAWING_TILE"), bDrawingTile ? TEXT("1") : TEXT("0"));
 	}
 
@@ -52,22 +53,23 @@ public:
 		}
 	}
 
-	void SetPS(const FRenderingCompositePassContext& Context)
+	void SetPS(FRHICommandList* RHICmdList, const FRenderingCompositePassContext& Context)
 	{
 		const FPixelShaderRHIParamRef ShaderRHI = GetPixelShader();
 
-		FGlobalShader::SetParameters(ShaderRHI, Context.View);
-		PostprocessParameter.SetPS(ShaderRHI, Context, TStaticSamplerState<SF_Point,AM_Clamp,AM_Clamp,AM_Clamp>::GetRHI());
-		DeferredParameters.Set(ShaderRHI, Context.View);
+		FGlobalShader::SetParameters(RHICmdList, ShaderRHI, Context.View);
+		PostprocessParameter.SetPS(RHICmdList, ShaderRHI, Context, TStaticSamplerState<SF_Point,AM_Clamp,AM_Clamp,AM_Clamp>::GetRHI());
+		DeferredParameters.Set(RHICmdList, ShaderRHI, Context.View);
 	}
 
-	void SetSourceTexture(FTextureRHIRef Texture)
+	void SetSourceTexture(FRHICommandList* RHICmdList, FTextureRHIRef Texture)
 	{
 		if (bDrawingTile && SourceTexture.IsBound())
 		{
 			const FPixelShaderRHIParamRef ShaderRHI = GetPixelShader();
 
 			SetTextureParameter(
+				RHICmdList, 
 				ShaderRHI,
 				SourceTexture,
 				SourceTextureSampler,
@@ -109,7 +111,7 @@ IMPLEMENT_SHADER_TYPE2(FPostProcessVisualizeBufferPS<true>, SF_Pixel);
 IMPLEMENT_SHADER_TYPE2(FPostProcessVisualizeBufferPS<false>, SF_Pixel);
 
 template <bool bDrawingTile>
-FShader* FRCPassPostProcessVisualizeBuffer::SetShaderTempl(const FRenderingCompositePassContext& Context)
+FShader* FRCPassPostProcessVisualizeBuffer::SetShaderTempl(FRHICommandList* RHICmdList, const FRenderingCompositePassContext& Context)
 {
 	TShaderMapRef<FPostProcessVS> VertexShader(GetGlobalShaderMap());
 	TShaderMapRef<FPostProcessVisualizeBufferPS<bDrawingTile> > PixelShader(GetGlobalShaderMap());
@@ -118,7 +120,7 @@ FShader* FRCPassPostProcessVisualizeBuffer::SetShaderTempl(const FRenderingCompo
 
 	SetGlobalBoundShaderState(BoundShaderState, GFilterVertexDeclaration.VertexDeclarationRHI, *VertexShader, *PixelShader);
 
-	PixelShader->SetPS(Context);
+	PixelShader->SetPS(RHICmdList, Context);
 
 	return *VertexShader;
 }
@@ -143,6 +145,9 @@ void FRCPassPostProcessVisualizeBuffer::Process(FRenderingCompositePassContext&
 
 	const FSceneRenderTargetItem& DestRenderTarget = PassOutputs[0].RequestSurface(Context);
 
+	//@todo-rco: RHIPacketList
+	FRHICommandList* RHICmdList = nullptr;
+
 	// Set the view family's render target/viewport.
 	RHISetRenderTarget(DestRenderTarget.TargetableTexture, FTextureRHIRef());	
 	Context.SetViewportAndCallRHI(DestRect);
@@ -153,18 +158,18 @@ void FRCPassPostProcessVisualizeBuffer::Process(FRenderingCompositePassContext&
 	RHISetDepthStencilState(TStaticDepthStencilState<false,CF_Always>::GetRHI());
 
 	{
-		FShader* VertexShader = SetShaderTempl<false>(Context);
-
-		// Draw a quad mapping scene color to the view's render target
-		DrawRectangle(
-			0, 0,
-			DestRect.Width(), DestRect.Height(),
-			SrcRect.Min.X, SrcRect.Min.Y,
-			SrcRect.Width(), SrcRect.Height(),
-			DestRect.Size(),
-			SrcSize,
+		FShader* VertexShader = SetShaderTempl<false>(RHICmdList, Context);
+
+	// Draw a quad mapping scene color to the view's render target
+	DrawRectangle(
+		0, 0,
+		DestRect.Width(), DestRect.Height(),
+		SrcRect.Min.X, SrcRect.Min.Y,
+		SrcRect.Width(), SrcRect.Height(),
+		DestRect.Size(),
+		SrcSize,
 			VertexShader,
-			EDRF_UseTriangleOptimization);
+		EDRF_UseTriangleOptimization);
 	}
 
 	RHISetBlendState(TStaticBlendState<CW_RGB, BO_Add, BF_SourceAlpha, BF_InverseSourceAlpha>::GetRHI());
@@ -175,7 +180,7 @@ void FRCPassPostProcessVisualizeBuffer::Process(FRenderingCompositePassContext&
 	static FGlobalBoundShaderState BoundShaderState;
 	SetGlobalBoundShaderState(BoundShaderState, GFilterVertexDeclaration.VertexDeclarationRHI, *VertexShader, *PixelShader);
 
-	PixelShader->SetPS(Context);
+	PixelShader->SetPS(RHICmdList, Context);
 
 	// Track the name and position of each tile we draw so we can write text labels over them
 	struct LabelRecord
@@ -204,7 +209,7 @@ void FRCPassPostProcessVisualizeBuffer::Process(FRenderingCompositePassContext&
 			int32 TileX = CurrentTileIndex % MaxTilesX;
 			int32 TileY = CurrentTileIndex / MaxTilesX;
 
-			PixelShader->SetSourceTexture(Texture);
+			PixelShader->SetSourceTexture(RHICmdList, Texture);
 
 			DrawRectangle(
 				TileX * TileWidth, TileY * TileHeight,
diff --git a/Engine/Source/Runtime/Renderer/Private/PostProcess/PostProcessVisualizeBuffer.h b/Engine/Source/Runtime/Renderer/Private/PostProcess/PostProcessVisualizeBuffer.h
--- a/Engine/Source/Runtime/Renderer/Private/PostProcess/PostProcessVisualizeBuffer.h
+++ b/Engine/Source/Runtime/Renderer/Private/PostProcess/PostProcessVisualizeBuffer.h
@@ -42,7 +42,7 @@ private:
 
 	// @return VertexShader
 	template <bool bDrawTileTemplate>
-	FShader* SetShaderTempl(const FRenderingCompositePassContext& Context);
+	FShader* SetShaderTempl(FRHICommandList* RHICmdList, const FRenderingCompositePassContext& Context);
 
 	TArray<TileData> Tiles;
 };
diff --git a/Engine/Source/Runtime/Renderer/Private/PostProcess/PostProcessVisualizeHDR.cpp b/Engine/Source/Runtime/Renderer/Private/PostProcess/PostProcessVisualizeHDR.cpp
--- a/Engine/Source/Runtime/Renderer/Private/PostProcess/PostProcessVisualizeHDR.cpp
+++ b/Engine/Source/Runtime/Renderer/Private/PostProcess/PostProcessVisualizeHDR.cpp
@@ -69,43 +69,43 @@ public:
 		ColorShadow_Tint2.Bind(Initializer.ParameterMap, TEXT("ColorShadow_Tint2"));
 	}
 
-	void SetPS(const FRenderingCompositePassContext& Context)
+	void SetPS(FRHICommandList* RHICmdList, const FRenderingCompositePassContext& Context)
 	{
 		const FPixelShaderRHIParamRef ShaderRHI = GetPixelShader();
 		
-		FGlobalShader::SetParameters(ShaderRHI, Context.View);
+		FGlobalShader::SetParameters(RHICmdList, ShaderRHI, Context.View);
 
 		const FPostProcessSettings& Settings = Context.View.FinalPostProcessSettings;
 		const FSceneViewFamily& ViewFamily = *(Context.View.Family);
 
-		PostprocessParameter.SetPS(ShaderRHI, Context, TStaticSamplerState<SF_Point,AM_Clamp,AM_Clamp,AM_Clamp>::GetRHI());
+		PostprocessParameter.SetPS(RHICmdList, ShaderRHI, Context, TStaticSamplerState<SF_Point,AM_Clamp,AM_Clamp,AM_Clamp>::GetRHI());
 
 		{
 			FVector4 Temp[3];
 
 			FRCPassPostProcessEyeAdaptation::ComputeEyeAdaptationParamsValue(Context.View, Temp);
-			SetShaderValueArray(ShaderRHI, EyeAdaptationParams, Temp, 3);
+			SetShaderValueArray(RHICmdList, ShaderRHI, EyeAdaptationParams, Temp, 3);
 		}
 
-		SetTextureParameter(ShaderRHI, MiniFontTexture, GEngine->MiniFontTexture ? GEngine->MiniFontTexture->Resource->TextureRHI : GSystemTextures.WhiteDummy->GetRenderTargetItem().TargetableTexture);
+		SetTextureParameter(RHICmdList, ShaderRHI, MiniFontTexture, GEngine->MiniFontTexture ? GEngine->MiniFontTexture->Resource->TextureRHI : GSystemTextures.WhiteDummy->GetRenderTargetItem().TargetableTexture);
 
 		{
 			float InvDisplayGammaValue = 1.0f / ViewFamily.RenderTarget->GetDisplayGamma();
 
-			SetShaderValue(ShaderRHI, InverseGamma, InvDisplayGammaValue);
+			SetShaderValue(RHICmdList, ShaderRHI, InverseGamma, InvDisplayGammaValue);
 		}
 
 		{
 			FVector4 Constants[8];
 			FilmPostSetConstants(Constants, ~0, &Context.View.FinalPostProcessSettings, false);
-			SetShaderValue(ShaderRHI, ColorMatrixR_ColorCurveCd1, Constants[0]);
-			SetShaderValue(ShaderRHI, ColorMatrixG_ColorCurveCd3Cm3, Constants[1]);
-			SetShaderValue(ShaderRHI, ColorMatrixB_ColorCurveCm2, Constants[2]); 
-			SetShaderValue(ShaderRHI, ColorCurve_Cm0Cd0_Cd2_Ch0Cm1_Ch3, Constants[3]); 
-			SetShaderValue(ShaderRHI, ColorCurve_Ch1_Ch2, Constants[4]);
-			SetShaderValue(ShaderRHI, ColorShadow_Luma, Constants[5]);
-			SetShaderValue(ShaderRHI, ColorShadow_Tint1, Constants[6]);
-			SetShaderValue(ShaderRHI, ColorShadow_Tint2, Constants[7]);
+			SetShaderValue(RHICmdList, ShaderRHI, ColorMatrixR_ColorCurveCd1, Constants[0]);
+			SetShaderValue(RHICmdList, ShaderRHI, ColorMatrixG_ColorCurveCd3Cm3, Constants[1]);
+			SetShaderValue(RHICmdList, ShaderRHI, ColorMatrixB_ColorCurveCm2, Constants[2]); 
+			SetShaderValue(RHICmdList, ShaderRHI, ColorCurve_Cm0Cd0_Cd2_Ch0Cm1_Ch3, Constants[3]); 
+			SetShaderValue(RHICmdList, ShaderRHI, ColorCurve_Ch1_Ch2, Constants[4]);
+			SetShaderValue(RHICmdList, ShaderRHI, ColorShadow_Luma, Constants[5]);
+			SetShaderValue(RHICmdList, ShaderRHI, ColorShadow_Tint1, Constants[6]);
+			SetShaderValue(RHICmdList, ShaderRHI, ColorShadow_Tint2, Constants[7]);
 		}
 	}
 	
@@ -153,6 +153,9 @@ void FRCPassPostProcessVisualizeHDR::Process(FRenderingCompositePassContext& Con
 
 	const FSceneRenderTargetItem& DestRenderTarget = PassOutputs[0].RequestSurface(Context);
 
+	//@todo-rco: RHIPacketList
+	FRHICommandList* RHICmdList = nullptr;
+
 	// Set the view family's render target/viewport.
 	RHISetRenderTarget(DestRenderTarget.TargetableTexture, FTextureRHIRef());	
 	Context.SetViewportAndCallRHI(DestRect);
@@ -169,7 +172,7 @@ void FRCPassPostProcessVisualizeHDR::Process(FRenderingCompositePassContext& Con
 
 	SetGlobalBoundShaderState(BoundShaderState, GFilterVertexDeclaration.VertexDeclarationRHI, *VertexShader, *PixelShader);
 
-	PixelShader->SetPS(Context);
+	PixelShader->SetPS(RHICmdList, Context);
 
 	// Draw a quad mapping scene color to the view's render target
 	DrawRectangle(
diff --git a/Engine/Source/Runtime/Renderer/Private/PostProcess/PostProcessWeightedSampleSum.cpp b/Engine/Source/Runtime/Renderer/Private/PostProcess/PostProcessWeightedSampleSum.cpp
--- a/Engine/Source/Runtime/Renderer/Private/PostProcess/PostProcessWeightedSampleSum.cpp
+++ b/Engine/Source/Runtime/Renderer/Private/PostProcess/PostProcessWeightedSampleSum.cpp
@@ -65,13 +65,13 @@ public:
 	}
 
 	/** Sets shader parameter values */
-	void SetParameters(FSamplerStateRHIParamRef SamplerStateRHI, FTextureRHIParamRef FilterTextureRHI, FTextureRHIParamRef AdditiveTextureRHI, const FLinearColor* SampleWeightValues)
+	void SetParameters(FRHICommandList* RHICmdList, FSamplerStateRHIParamRef SamplerStateRHI, FTextureRHIParamRef FilterTextureRHI, FTextureRHIParamRef AdditiveTextureRHI, const FLinearColor* SampleWeightValues)
 	{
 		const FPixelShaderRHIParamRef ShaderRHI = GetPixelShader();
 
-		SetTextureParameter(ShaderRHI, FilterTexture, FilterTextureSampler, SamplerStateRHI, FilterTextureRHI);
-		SetTextureParameter(ShaderRHI, AdditiveTexture, AdditiveTextureSampler, SamplerStateRHI, AdditiveTextureRHI);
-		SetShaderValueArray(ShaderRHI, SampleWeights, SampleWeightValues, NumSamples);
+		SetTextureParameter(RHICmdList, ShaderRHI, FilterTexture, FilterTextureSampler, SamplerStateRHI, FilterTextureRHI);
+		SetTextureParameter(RHICmdList, ShaderRHI, AdditiveTexture, AdditiveTextureSampler, SamplerStateRHI, AdditiveTextureRHI);
+		SetShaderValueArray(RHICmdList, ShaderRHI, SampleWeights, SampleWeightValues, NumSamples);
 	}
 
 	static const TCHAR* GetSourceFilename()
@@ -158,6 +158,7 @@ protected:
  * @param OutVertexShader - The vertex shader used for the filter
  */
 void SetFilterShaders(
+	FRHICommandList* RHICmdList, 
 	FSamplerStateRHIParamRef SamplerStateRHI,
 	FTextureRHIParamRef FilterTextureRHI,
 	FTextureRHIParamRef AdditiveTextureRHI,
@@ -169,6 +170,7 @@ void SetFilterShaders(
 	)
 {
 	check(CombineMethodInt <= 2);
+	check(!RHICmdList);
 
 	// A macro to handle setting the filter shader for a specific number of samples.
 #define SET_FILTER_SHADER_TYPE(NumSamples) \
@@ -183,7 +185,7 @@ void SetFilterShaders(
 				static FGlobalBoundShaderState BoundShaderState; \
 				SetGlobalBoundShaderState( BoundShaderState, GFilterVertexDeclaration.VertexDeclarationRHI, *VertexShader, *PixelShader); \
 			} \
-			PixelShader->SetParameters(SamplerStateRHI, FilterTextureRHI, AdditiveTextureRHI, SampleWeights); \
+			PixelShader->SetParameters(RHICmdList, SamplerStateRHI, FilterTextureRHI, AdditiveTextureRHI, SampleWeights); \
 		} \
 		else if(CombineMethodInt == 1) \
 		{ \
@@ -192,7 +194,7 @@ void SetFilterShaders(
 				static FGlobalBoundShaderState BoundShaderState; \
 				SetGlobalBoundShaderState( BoundShaderState, GFilterVertexDeclaration.VertexDeclarationRHI, *VertexShader, *PixelShader); \
 			} \
-			PixelShader->SetParameters(SamplerStateRHI, FilterTextureRHI, AdditiveTextureRHI, SampleWeights); \
+			PixelShader->SetParameters(RHICmdList, SamplerStateRHI, FilterTextureRHI, AdditiveTextureRHI, SampleWeights); \
 		} \
 		else\
 		{ \
@@ -201,9 +203,9 @@ void SetFilterShaders(
 				static FGlobalBoundShaderState BoundShaderState; \
 				SetGlobalBoundShaderState( BoundShaderState, GFilterVertexDeclaration.VertexDeclarationRHI, *VertexShader, *PixelShader); \
 			} \
-			PixelShader->SetParameters(SamplerStateRHI, FilterTextureRHI, AdditiveTextureRHI, SampleWeights); \
+			PixelShader->SetParameters(RHICmdList, SamplerStateRHI, FilterTextureRHI, AdditiveTextureRHI, SampleWeights); \
 		} \
-		VertexShader->SetParameters(SampleOffsets); \
+		VertexShader->SetParameters(RHICmdList, SampleOffsets); \
 		break; \
 	};
 
@@ -375,6 +377,9 @@ void FRCPassPostProcessWeightedSampleSum::Process(FRenderingCompositePassContext
 		BlurWeights[i] = TintValue * OffsetAndWeight[i].Y;
 	}
 
+	//@todo-rco: RHIPacketList
+	FRHICommandList* RHICmdList = nullptr;
+
 	RHISetRenderTarget(DestRenderTarget.TargetableTexture, FTextureRHIRef());
 
 	Context.SetViewportAndCallRHI(0, 0, 0.0f, DestSize.X, DestSize.Y, 1.0f);
@@ -429,6 +434,7 @@ void FRCPassPostProcessWeightedSampleSum::Process(FRenderingCompositePassContext
 
 	FShader* VertexShader = nullptr;
 	SetFilterShaders(
+		RHICmdList, 
 		TStaticSamplerState<SF_Bilinear,AM_Border,AM_Border,AM_Clamp>::GetRHI(),
 		FilterTexture,
 		AdditiveTexture,
diff --git a/Engine/Source/Runtime/Renderer/Private/PostProcess/PostProcessWeightedSampleSum.h b/Engine/Source/Runtime/Renderer/Private/PostProcess/PostProcessWeightedSampleSum.h
--- a/Engine/Source/Runtime/Renderer/Private/PostProcess/PostProcessWeightedSampleSum.h
+++ b/Engine/Source/Runtime/Renderer/Private/PostProcess/PostProcessWeightedSampleSum.h
@@ -62,7 +62,7 @@ public:
 	}
 
 	/** Sets shader parameter values */
-	void SetParameters(const FVector2D* SampleOffsetsValue)
+	void SetParameters(FRHICommandList* RHICmdList, const FVector2D* SampleOffsetsValue)
 	{
 		FVector4 PackedSampleOffsetsValue[NumSampleChunks];
 		for(int32 SampleIndex = 0;SampleIndex < NumSamples;SampleIndex += 2)
@@ -75,7 +75,7 @@ public:
 				PackedSampleOffsetsValue[SampleIndex / 2].Z = SampleOffsetsValue[SampleIndex + 1].Y;
 			}
 		}
-		SetShaderValueArray(GetVertexShader(),SampleOffsets,PackedSampleOffsetsValue,NumSampleChunks);
+		SetShaderValueArray(RHICmdList, GetVertexShader(),SampleOffsets,PackedSampleOffsetsValue,NumSampleChunks);
 	}
 
 private:
diff --git a/Engine/Source/Runtime/Renderer/Private/PostProcess/PostProcessing.h b/Engine/Source/Runtime/Renderer/Private/PostProcess/PostProcessing.h
--- a/Engine/Source/Runtime/Renderer/Private/PostProcess/PostProcessing.h
+++ b/Engine/Source/Runtime/Renderer/Private/PostProcess/PostProcessing.h
@@ -40,14 +40,14 @@ class FPostProcessVS : public FGlobalShader
 	FPostProcessVS() {}
 
 	/** to have a similar interface as all other shaders */
-	void SetParameters(const FRenderingCompositePassContext& Context)
+	void SetParameters(FRHICommandList* RHICmdList, const FRenderingCompositePassContext& Context)
 	{
-		FGlobalShader::SetParameters(GetVertexShader(), Context.View);
+		FGlobalShader::SetParameters(RHICmdList, GetVertexShader(), Context.View);
 	}
 
-	void SetParameters(const FSceneView& View)
+	void SetParameters(FRHICommandList* RHICmdList, const FSceneView& View)
 	{
-		FGlobalShader::SetParameters(GetVertexShader(), View);
+		FGlobalShader::SetParameters(RHICmdList, GetVertexShader(), View);
 	}
 
 public:
diff --git a/Engine/Source/Runtime/Renderer/Private/PostProcess/RenderingCompositionGraph.cpp b/Engine/Source/Runtime/Renderer/Private/PostProcess/RenderingCompositionGraph.cpp
--- a/Engine/Source/Runtime/Renderer/Private/PostProcess/RenderingCompositionGraph.cpp
+++ b/Engine/Source/Runtime/Renderer/Private/PostProcess/RenderingCompositionGraph.cpp
@@ -664,19 +664,19 @@ void FPostProcessPassParameters::Bind(const FShaderParameterMap& ParameterMap)
 	}
 }
 
-void FPostProcessPassParameters::SetPS(const FPixelShaderRHIParamRef& ShaderRHI, const FRenderingCompositePassContext& Context, FSamplerStateRHIParamRef Filter, bool bWhiteIfNoTexture, FSamplerStateRHIParamRef* FilterOverrideArray)
+void FPostProcessPassParameters::SetPS(FRHICommandList* RHICmdList, const FPixelShaderRHIParamRef& ShaderRHI, const FRenderingCompositePassContext& Context, FSamplerStateRHIParamRef Filter, bool bWhiteIfNoTexture, FSamplerStateRHIParamRef* FilterOverrideArray)
 {
-	Set(ShaderRHI, Context, Filter, bWhiteIfNoTexture, FilterOverrideArray);
+	Set(RHICmdList, ShaderRHI, Context, Filter, bWhiteIfNoTexture, FilterOverrideArray);
 }
 
-void FPostProcessPassParameters::SetCS(const FComputeShaderRHIParamRef& ShaderRHI, const FRenderingCompositePassContext& Context, FSamplerStateRHIParamRef Filter, bool bWhiteIfNoTexture, FSamplerStateRHIParamRef* FilterOverrideArray)
+void FPostProcessPassParameters::SetCS(FRHICommandList* RHICmdList, const FComputeShaderRHIParamRef& ShaderRHI, const FRenderingCompositePassContext& Context, FSamplerStateRHIParamRef Filter, bool bWhiteIfNoTexture, FSamplerStateRHIParamRef* FilterOverrideArray)
 {
-	Set(ShaderRHI, Context, Filter, bWhiteIfNoTexture, FilterOverrideArray);
+	Set(RHICmdList, ShaderRHI, Context, Filter, bWhiteIfNoTexture, FilterOverrideArray);
 }
 
-void FPostProcessPassParameters::SetVS(const FVertexShaderRHIParamRef& ShaderRHI, const FRenderingCompositePassContext& Context, FSamplerStateRHIParamRef Filter, bool bWhiteIfNoTexture, FSamplerStateRHIParamRef* FilterOverrideArray)
+void FPostProcessPassParameters::SetVS(FRHICommandList* RHICmdList, const FVertexShaderRHIParamRef& ShaderRHI, const FRenderingCompositePassContext& Context, FSamplerStateRHIParamRef Filter, bool bWhiteIfNoTexture, FSamplerStateRHIParamRef* FilterOverrideArray)
 {
-	Set(ShaderRHI, Context, Filter, bWhiteIfNoTexture, FilterOverrideArray);
+	Set(RHICmdList, ShaderRHI, Context, Filter, bWhiteIfNoTexture, FilterOverrideArray);
 }
 
 FArchive& operator<<(FArchive& Ar, FPostProcessPassParameters& P)
diff --git a/Engine/Source/Runtime/Renderer/Private/PostProcess/RenderingCompositionGraph.h b/Engine/Source/Runtime/Renderer/Private/PostProcess/RenderingCompositionGraph.h
--- a/Engine/Source/Runtime/Renderer/Private/PostProcess/RenderingCompositionGraph.h
+++ b/Engine/Source/Runtime/Renderer/Private/PostProcess/RenderingCompositionGraph.h
@@ -6,6 +6,8 @@
 
 #pragma once
 
+#include "RHIStaticStates.h"
+
 struct FRenderingCompositePass;
 struct FRenderingCompositeOutputRef;
 struct FRenderingCompositeOutput;
@@ -519,13 +521,13 @@ struct FPostProcessPassParameters
 	void Bind(const FShaderParameterMap& ParameterMap);
 
 	/** Set the pixel shader parameter values. */
-	void SetPS(const FPixelShaderRHIParamRef& ShaderRHI, const FRenderingCompositePassContext& Context, FSamplerStateRHIParamRef Filter = TStaticSamplerState<>::GetRHI(), bool bWhiteIfNoTexture = false, FSamplerStateRHIParamRef* FilterOverrideArray = 0);
+	void SetPS(FRHICommandList* RHICmdList, const FPixelShaderRHIParamRef& ShaderRHI, const FRenderingCompositePassContext& Context, FSamplerStateRHIParamRef Filter = TStaticSamplerState<>::GetRHI(), bool bWhiteIfNoTexture = false, FSamplerStateRHIParamRef* FilterOverrideArray = 0);
 
 	/** Set the compute shader parameter values. */
-	void SetCS(const FComputeShaderRHIParamRef& ShaderRHI, const FRenderingCompositePassContext& Context, FSamplerStateRHIParamRef Filter = TStaticSamplerState<>::GetRHI(), bool bWhiteIfNoTexture = false, FSamplerStateRHIParamRef* FilterOverrideArray = 0);
+	void SetCS(FRHICommandList* RHICmdList, const FComputeShaderRHIParamRef& ShaderRHI, const FRenderingCompositePassContext& Context, FSamplerStateRHIParamRef Filter = TStaticSamplerState<>::GetRHI(), bool bWhiteIfNoTexture = false, FSamplerStateRHIParamRef* FilterOverrideArray = 0);
 
 	/** Set the vertex shader parameter values. */
-	void SetVS(const FVertexShaderRHIParamRef& ShaderRHI, const FRenderingCompositePassContext& Context, FSamplerStateRHIParamRef Filter = TStaticSamplerState<>::GetRHI(), bool bWhiteIfNoTexture = false, FSamplerStateRHIParamRef* FilterOverrideArray = 0);
+	void SetVS(FRHICommandList* RHICmdList, const FVertexShaderRHIParamRef& ShaderRHI, const FRenderingCompositePassContext& Context, FSamplerStateRHIParamRef Filter = TStaticSamplerState<>::GetRHI(), bool bWhiteIfNoTexture = false, FSamplerStateRHIParamRef* FilterOverrideArray = 0);
 
 	/** Serializer. */
 	friend FArchive& operator<<(FArchive& Ar,FPostProcessPassParameters& P);
@@ -548,7 +550,7 @@ public:
 	// @param Filter can be 0 if FilterOverrideArray is used
 	// @param FilterOverrideArray can be 0 if Filter is used
 	template <class T>
-	void Set(const T& ShaderRHI, const FRenderingCompositePassContext& Context, FSamplerStateRHIParamRef Filter, bool bWhiteIfNoTexture = false, FSamplerStateRHIParamRef* FilterOverrideArray = 0)
+	void Set(FRHICommandList* RHICmdList, const T& ShaderRHI, const FRenderingCompositePassContext& Context, FSamplerStateRHIParamRef Filter, bool bWhiteIfNoTexture = false, FSamplerStateRHIParamRef* FilterOverrideArray = 0)
 	{
 		// assuming all outputs have the same size
 		FRenderingCompositeOutput* Output = Context.Pass->GetOutput(ePId_Output0);
@@ -561,6 +563,8 @@ public:
 		// but not both
 		check(!FilterOverrideArray || !Filter);
 
+		check(!RHICmdList);
+
 		if(BilinearTextureSampler0.IsBound())
 		{
 			RHISetShaderSampler(
@@ -589,13 +593,13 @@ public:
 			{
 				FVector4 Value(ViewportExtent.X, ViewportExtent.Y, 1.0f / ViewportExtent.X, 1.0f / ViewportExtent.Y);
 
-				SetShaderValue(ShaderRHI, ViewportSize, Value);
+				SetShaderValue(RHICmdList, ShaderRHI, ViewportSize, Value);
 			}
 
 			{
 				FVector4 Value(LocalViewport.Min.X, LocalViewport.Min.Y, LocalViewport.Max.X, LocalViewport.Max.Y);
 
-				SetShaderValue(ShaderRHI, ViewportRect, Value);
+				SetShaderValue(RHICmdList, ShaderRHI, ViewportRect, Value);
 			}
 
 			{
@@ -604,7 +608,7 @@ public:
 					-ViewportExtent.Y * 0.5f, 
 					ViewportExtent.X * 0.5f - 0.5f + ViewportOffset.X,
 					ViewportExtent.Y * 0.5f - 0.5f + ViewportOffset.Y);
-				SetShaderValue(ShaderRHI, ScreenPosToPixel, ScreenPosToPixelValue);
+				SetShaderValue(RHICmdList, ShaderRHI, ScreenPosToPixel, ScreenPosToPixelValue);
 			}
 		}
 
@@ -646,21 +650,21 @@ public:
 
 				const FTextureRHIRef& SrcTexture = InputPooledElement->GetRenderTargetItem().ShaderResourceTexture;
 
-				SetTextureParameter(ShaderRHI, PostprocessInputParameter[Id], PostprocessInputParameterSampler[Id], LocalFilter, SrcTexture);
-				SetTextureParameter(ShaderRHI, PostprocessInputNew[Id], InputPooledElement->GetRenderTargetItem().TargetableTexture);
+				SetTextureParameter(RHICmdList, ShaderRHI, PostprocessInputParameter[Id], PostprocessInputParameterSampler[Id], LocalFilter, SrcTexture);
+				SetTextureParameter(RHICmdList, ShaderRHI, PostprocessInputNew[Id], InputPooledElement->GetRenderTargetItem().TargetableTexture);
 
 				{
 					float Width = InputPooledElement->GetDesc().Extent.X;
 					float Height = InputPooledElement->GetDesc().Extent.Y;
 					FVector4 TextureSize(Width, Height, 1.0f / Width, 1.0f / Height);
-					SetShaderValue(ShaderRHI, PostprocessInputSizeParameter[Id], TextureSize);
+					SetShaderValue(RHICmdList, ShaderRHI, PostprocessInputSizeParameter[Id], TextureSize);
 					
 					//We could use the main scene min max here if it weren't that we need to pull the max in by a pixel on a per input basis.
 					FVector2D OnePPInputPixelUVSize = FVector2D(1.0f / Width, 1.0f / Height);
 					FVector4 PPInputMinMax = BaseSceneTexMinMax;
 					PPInputMinMax.Z -= OnePPInputPixelUVSize.X;
 					PPInputMinMax.W -= OnePPInputPixelUVSize.Y;
-					SetShaderValue(ShaderRHI, PostProcessInputMinMaxParameter[Id], PPInputMinMax);
+					SetShaderValue(RHICmdList, ShaderRHI, PostProcessInputMinMaxParameter[Id], PPInputMinMax);
 				}
 			}
 			else
@@ -669,12 +673,12 @@ public:
 
 				// if the input is not there but the shader request it we give it at least some data to avoid d3ddebug errors and shader permutations
 				// to make features optional we use default black for additive passes without shader permutations
-				SetTextureParameter(ShaderRHI, PostprocessInputParameter[Id], PostprocessInputParameterSampler[Id], LocalFilter, Texture->GetRenderTargetItem().TargetableTexture);
-				SetTextureParameter(ShaderRHI, PostprocessInputNew[Id], Texture->GetRenderTargetItem().TargetableTexture);
+				SetTextureParameter(RHICmdList, ShaderRHI, PostprocessInputParameter[Id], PostprocessInputParameterSampler[Id], LocalFilter, Texture->GetRenderTargetItem().TargetableTexture);
+				SetTextureParameter(RHICmdList, ShaderRHI, PostprocessInputNew[Id], Texture->GetRenderTargetItem().TargetableTexture);
 
 				FVector4 Dummy(1, 1, 1, 1);
-				SetShaderValue(ShaderRHI, PostprocessInputSizeParameter[Id], Dummy);
-				SetShaderValue(ShaderRHI, PostProcessInputMinMaxParameter[Id], Dummy);
+				SetShaderValue(RHICmdList, ShaderRHI, PostprocessInputSizeParameter[Id], Dummy);
+				SetShaderValue(RHICmdList, ShaderRHI, PostProcessInputMinMaxParameter[Id], Dummy);
 			}
 		}
 
diff --git a/Engine/Source/Runtime/Renderer/Private/PostProcess/SceneFilterRendering.cpp b/Engine/Source/Runtime/Renderer/Private/PostProcess/SceneFilterRendering.cpp
--- a/Engine/Source/Runtime/Renderer/Private/PostProcess/SceneFilterRendering.cpp
+++ b/Engine/Source/Runtime/Renderer/Private/PostProcess/SceneFilterRendering.cpp
@@ -136,7 +136,9 @@ void DrawRectangle(
 		1.0f / TargetSize.X, 1.0f / TargetSize.Y, 
 		1.0f / TextureSize.X, 1.0f / TextureSize.Y);
 
-	SetUniformBufferParameterImmediate(VertexShader->GetVertexShader(), VertexShader->GetUniformBufferParameter<FDrawRectangleParameters>(), Parameters);
+	//@todo-rco: RHIPacketList
+	FRHICommandList* RHICmdList = nullptr;
+	SetUniformBufferParameterImmediate(RHICmdList, VertexShader->GetVertexShader(), VertexShader->GetUniformBufferParameter<FDrawRectangleParameters>(), Parameters);
 	RHISetStreamSource(0, GScreenRectangleVertexBuffer.VertexBufferRHI, sizeof(FFilterVertex),0);
 
 
@@ -173,19 +175,19 @@ void DrawRectangle(
 }
 
 void DrawTransformedRectangle(
-    float X,
-    float Y,
-    float SizeX,
-    float SizeY,
-    const FMatrix& PosTransform,
-    float U,
-    float V,
-    float SizeU,
-    float SizeV,
-    const FMatrix& TexTransform,
-    FIntPoint TargetSize,
-    FIntPoint TextureSize
-    )
+	float X,
+	float Y,
+	float SizeX,
+	float SizeY,
+	const FMatrix& PosTransform,
+	float U,
+	float V,
+	float SizeU,
+	float SizeV,
+	const FMatrix& TexTransform,
+	FIntPoint TargetSize,
+	FIntPoint TextureSize
+	)
 {
 	float ClipSpaceQuadZ = 0.0f;
 
diff --git a/Engine/Source/Runtime/Renderer/Private/PostProcess/SceneFilterRendering.h b/Engine/Source/Runtime/Renderer/Private/PostProcess/SceneFilterRendering.h
--- a/Engine/Source/Runtime/Renderer/Private/PostProcess/SceneFilterRendering.h
+++ b/Engine/Source/Runtime/Renderer/Private/PostProcess/SceneFilterRendering.h
@@ -121,7 +121,7 @@ public:
 	}
 
 	/** Set the material shader parameter values. */
-	void Set(FShader* PixelShader, float DisplayGamma, FLinearColor const& ColorScale, FLinearColor const& ColorOverlay)
+	void Set(FRHICommandList* RHICmdList, FShader* PixelShader, float DisplayGamma, FLinearColor const& ColorScale, FLinearColor const& ColorOverlay)
 	{
 		// GammaColorScaleAndInverse
 
@@ -136,6 +136,7 @@ public:
 		ColorScaleAndInverse.W = InvDisplayGamma;
 
 		SetShaderValue(
+			RHICmdList, 
 			PixelShader->GetPixelShader(),
 			GammaColorScaleAndInverse,
 			ColorScaleAndInverse
@@ -151,6 +152,7 @@ public:
 		OverlayColor.W = 0.f; // Unused
 
 		SetShaderValue(
+			RHICmdList, 
 			PixelShader->GetPixelShader(),
 			GammaOverlayColor,
 			OverlayColor
@@ -165,6 +167,7 @@ public:
 		const FVector4 vRenderTargetExtent(BufferSizeX, BufferSizeY,  InvBufferSizeX, InvBufferSizeY);
 
 		SetShaderValue(
+			RHICmdList, 
 			PixelShader->GetPixelShader(),
 			RenderTargetExtent, 
 			vRenderTargetExtent);
diff --git a/Engine/Source/Runtime/Renderer/Private/PostProcess/SceneRenderTargets.cpp b/Engine/Source/Runtime/Renderer/Private/PostProcess/SceneRenderTargets.cpp
--- a/Engine/Source/Runtime/Renderer/Private/PostProcess/SceneRenderTargets.cpp
+++ b/Engine/Source/Runtime/Renderer/Private/PostProcess/SceneRenderTargets.cpp
@@ -13,6 +13,8 @@
 // for LightPropagationVolume feature, could be exposed
 const int ReflectiveShadowMapResolution = 256;
 
+IMPLEMENT_UNIFORM_BUFFER_STRUCT(FGBufferResourceStruct,TEXT("GBuffers"));
+
 /*-----------------------------------------------------------------------------
 FSceneRenderTargets
 -----------------------------------------------------------------------------*/
@@ -351,6 +353,7 @@ void FSceneRenderTargets::AllocLightAttenuation()
 
 void FSceneRenderTargets::FreeGBufferTargets()
 {
+	GBufferResourcesUniformBuffer.SafeRelease();
 	GBufferA.SafeRelease();
 	GBufferB.SafeRelease();
 	GBufferC.SafeRelease();
@@ -426,6 +429,46 @@ void FSceneRenderTargets::AllocGBufferTargets()
 	// otherwise we have a severe problem
 	check(GBufferA);
 
+	// Create the required render targets if running Highend.
+	if (GRHIFeatureLevel >= ERHIFeatureLevel::SM4)
+	{
+
+		// Allocate the Gbuffer resource uniform buffer.
+		const FSceneRenderTargetItem& GBufferA = this->GBufferA ? this->GBufferA->GetRenderTargetItem() : GSystemTextures.BlackDummy->GetRenderTargetItem();
+		const FSceneRenderTargetItem& GBufferB = this->GBufferB ? this->GBufferB->GetRenderTargetItem() : GSystemTextures.BlackDummy->GetRenderTargetItem();
+		const FSceneRenderTargetItem& GBufferC = this->GBufferC ? this->GBufferC->GetRenderTargetItem() : GSystemTextures.BlackDummy->GetRenderTargetItem();
+		const FSceneRenderTargetItem& GBufferD = this->GBufferD ? this->GBufferD->GetRenderTargetItem() : GSystemTextures.BlackDummy->GetRenderTargetItem();
+		const FSceneRenderTargetItem& GBufferE = this->GBufferE ? this->GBufferE->GetRenderTargetItem() : GSystemTextures.BlackDummy->GetRenderTargetItem();
+
+		FGBufferResourceStruct GBufferResourceStruct;
+
+		GBufferResourceStruct.GBufferATexture = GBufferA.ShaderResourceTexture;
+		GBufferResourceStruct.GBufferBTexture = GBufferB.ShaderResourceTexture;
+		GBufferResourceStruct.GBufferCTexture = GBufferC.ShaderResourceTexture;
+		GBufferResourceStruct.GBufferDTexture = GBufferD.ShaderResourceTexture;
+		GBufferResourceStruct.GBufferETexture = GBufferE.ShaderResourceTexture;
+
+		GBufferResourceStruct.GBufferATextureNonMS = GBufferA.ShaderResourceTexture;
+		GBufferResourceStruct.GBufferBTextureNonMS = GBufferB.ShaderResourceTexture;
+		GBufferResourceStruct.GBufferCTextureNonMS = GBufferC.ShaderResourceTexture;
+		GBufferResourceStruct.GBufferDTextureNonMS = GBufferD.ShaderResourceTexture;
+		GBufferResourceStruct.GBufferETextureNonMS = GBufferE.ShaderResourceTexture;
+
+		GBufferResourceStruct.GBufferATextureMS = GBufferA.TargetableTexture;
+		GBufferResourceStruct.GBufferBTextureMS = GBufferB.TargetableTexture;
+		GBufferResourceStruct.GBufferCTextureMS = GBufferC.TargetableTexture;
+		GBufferResourceStruct.GBufferDTextureMS = GBufferD.TargetableTexture;
+		GBufferResourceStruct.GBufferETextureMS = GBufferE.TargetableTexture;
+
+		GBufferResourceStruct.GBufferATextureSampler = TStaticSamplerState<>::GetRHI();
+		GBufferResourceStruct.GBufferBTextureSampler = TStaticSamplerState<>::GetRHI();
+		GBufferResourceStruct.GBufferCTextureSampler = TStaticSamplerState<>::GetRHI();
+		GBufferResourceStruct.GBufferDTextureSampler = TStaticSamplerState<>::GetRHI();
+		GBufferResourceStruct.GBufferETextureSampler = TStaticSamplerState<>::GetRHI();
+
+		GBufferResourcesUniformBuffer = FGBufferResourceStruct::CreateUniformBuffer(GBufferResourceStruct, UniformBuffer_SingleFrame);
+	}
+
 	// so that 
 	GBufferRefCount = 1;
 }
@@ -1174,22 +1217,22 @@ EPixelFormat FSceneRenderTargets::GetSceneColorFormat() const
 	}
 	else
 	{
-		switch(CurrentSceneColorFormat)
-		{
-			case 0:
-				SceneColorBufferFormat = PF_R8G8B8A8; break;
-			case 1:
-				SceneColorBufferFormat = PF_A2B10G10R10; break;
-			case 2:	
-				SceneColorBufferFormat = PF_FloatR11G11B10; break;
-			case 3:	
-				SceneColorBufferFormat = PF_FloatRGB; break;
-			case 4:
-				// default
-				break;
-			case 5:
-				SceneColorBufferFormat = PF_A32B32G32R32F; break;
-		}
+	switch(CurrentSceneColorFormat)
+	{
+		case 0:
+			SceneColorBufferFormat = PF_R8G8B8A8; break;
+		case 1:
+			SceneColorBufferFormat = PF_A2B10G10R10; break;
+		case 2:	
+			SceneColorBufferFormat = PF_FloatR11G11B10; break;
+		case 3:	
+			SceneColorBufferFormat = PF_FloatRGB; break;
+		case 4:
+			// default
+			break;
+		case 5:
+			SceneColorBufferFormat = PF_A32B32G32R32F; break;
+	}
 	}
 
 	return SceneColorBufferFormat;
@@ -1433,35 +1476,17 @@ FArchive& operator<<(FArchive& Ar,FSceneTextureShaderParameters& Parameters)
 void FDeferredPixelShaderParameters::Bind(const FShaderParameterMap& ParameterMap)
 {
 	SceneTextureParameters.Bind(ParameterMap);
-	GBufferATextureMS.Bind(ParameterMap,TEXT("GBufferATextureMS"));
-	GBufferBTextureMS.Bind(ParameterMap,TEXT("GBufferBTextureMS"));
-	GBufferCTextureMS.Bind(ParameterMap,TEXT("GBufferCTextureMS"));
-	GBufferDTextureMS.Bind(ParameterMap,TEXT("GBufferDTextureMS"));
-	GBufferETextureMS.Bind(ParameterMap,TEXT("GBufferETextureMS"));
+	
+	GBufferResources.Bind(ParameterMap,TEXT("GBuffers"));
 	DBufferATextureMS.Bind(ParameterMap,TEXT("DBufferATextureMS"));
 	DBufferBTextureMS.Bind(ParameterMap,TEXT("DBufferBTextureMS"));
 	DBufferCTextureMS.Bind(ParameterMap,TEXT("DBufferCTextureMS"));
 	ScreenSpaceAOTextureMS.Bind(ParameterMap,TEXT("ScreenSpaceAOTextureMS"));
-	GBufferATextureNonMS.Bind(ParameterMap,TEXT("GBufferATextureNonMS"));
-	GBufferBTextureNonMS.Bind(ParameterMap,TEXT("GBufferBTextureNonMS"));
-	GBufferCTextureNonMS.Bind(ParameterMap,TEXT("GBufferCTextureNonMS"));
-	GBufferDTextureNonMS.Bind(ParameterMap,TEXT("GBufferDTextureNonMS"));
-	GBufferETextureNonMS.Bind(ParameterMap,TEXT("GBufferETextureNonMS"));
 	DBufferATextureNonMS.Bind(ParameterMap,TEXT("DBufferATextureNonMS"));
 	DBufferBTextureNonMS.Bind(ParameterMap,TEXT("DBufferBTextureNonMS"));
 	DBufferCTextureNonMS.Bind(ParameterMap,TEXT("DBufferCTextureNonMS"));
 	ScreenSpaceAOTextureNonMS.Bind(ParameterMap,TEXT("ScreenSpaceAOTextureNonMS"));
 	CustomDepthTextureNonMS.Bind(ParameterMap,TEXT("CustomDepthTextureNonMS"));
-	GBufferATexture.Bind(ParameterMap,TEXT("GBufferATexture"));
-	GBufferATextureSampler.Bind(ParameterMap,TEXT("GBufferATextureSampler"));
-	GBufferBTexture.Bind(ParameterMap,TEXT("GBufferBTexture"));
-	GBufferBTextureSampler.Bind(ParameterMap,TEXT("GBufferBTextureSampler"));
-	GBufferCTexture.Bind(ParameterMap,TEXT("GBufferCTexture"));
-	GBufferCTextureSampler.Bind(ParameterMap,TEXT("GBufferCTextureSampler"));
-	GBufferDTexture.Bind(ParameterMap,TEXT("GBufferDTexture"));
-	GBufferDTextureSampler.Bind(ParameterMap,TEXT("GBufferDTextureSampler"));
-	GBufferETexture.Bind(ParameterMap,TEXT("GBufferETexture"));
-	GBufferETextureSampler.Bind(ParameterMap,TEXT("GBufferETextureSampler"));
 	DBufferATexture.Bind(ParameterMap,TEXT("DBufferATexture"));
 	DBufferATextureSampler.Bind(ParameterMap,TEXT("DBufferATextureSampler"));
 	DBufferBTexture.Bind(ParameterMap,TEXT("DBufferBTexture"));
@@ -1477,35 +1502,17 @@ void FDeferredPixelShaderParameters::Bind(const FShaderParameterMap& ParameterMa
 FArchive& operator<<(FArchive& Ar,FDeferredPixelShaderParameters& Parameters)
 {
 	Ar << Parameters.SceneTextureParameters;
-	Ar << Parameters.GBufferATextureMS;
-	Ar << Parameters.GBufferBTextureMS;
-	Ar << Parameters.GBufferCTextureMS;
-	Ar << Parameters.GBufferDTextureMS;
-	Ar << Parameters.GBufferETextureMS;
+
+	Ar << Parameters.GBufferResources;
 	Ar << Parameters.DBufferATextureMS;
 	Ar << Parameters.DBufferBTextureMS;
 	Ar << Parameters.DBufferCTextureMS;
 	Ar << Parameters.ScreenSpaceAOTextureMS;
-	Ar << Parameters.GBufferATextureNonMS;
-	Ar << Parameters.GBufferBTextureNonMS;
-	Ar << Parameters.GBufferCTextureNonMS;
-	Ar << Parameters.GBufferDTextureNonMS;
-	Ar << Parameters.GBufferETextureNonMS;
 	Ar << Parameters.DBufferATextureNonMS;
 	Ar << Parameters.DBufferBTextureNonMS;
 	Ar << Parameters.DBufferCTextureNonMS;
 	Ar << Parameters.ScreenSpaceAOTextureNonMS;
 	Ar << Parameters.CustomDepthTextureNonMS;
-	Ar << Parameters.GBufferATexture;
-	Ar << Parameters.GBufferATextureSampler;
-	Ar << Parameters.GBufferBTexture;
-	Ar << Parameters.GBufferBTextureSampler;
-	Ar << Parameters.GBufferCTexture;
-	Ar << Parameters.GBufferCTextureSampler;
-	Ar << Parameters.GBufferDTexture;
-	Ar << Parameters.GBufferDTextureSampler;
-	Ar << Parameters.GBufferETexture;
-	Ar << Parameters.GBufferETextureSampler;
 	Ar << Parameters.DBufferATexture;
 	Ar << Parameters.DBufferATextureSampler;
 	Ar << Parameters.DBufferBTexture;
diff --git a/Engine/Source/Runtime/Renderer/Private/PostProcess/SceneRenderTargets.h b/Engine/Source/Runtime/Renderer/Private/PostProcess/SceneRenderTargets.h
--- a/Engine/Source/Runtime/Renderer/Private/PostProcess/SceneRenderTargets.h
+++ b/Engine/Source/Runtime/Renderer/Private/PostProcess/SceneRenderTargets.h
@@ -76,6 +76,29 @@ static const int32 NumTranslucencyShadowSurfaces = 2;
 /** Downsample factor to use on the volume texture used for computing GI on translucency. */
 static const int32 TranslucentVolumeGISratchDownsampleFactor = 1;
 
+BEGIN_UNIFORM_BUFFER_STRUCT(FGBufferResourceStruct, )
+	DECLARE_UNIFORM_BUFFER_STRUCT_MEMBER_TEXTURE( Texture2D,			GBufferATexture )
+	DECLARE_UNIFORM_BUFFER_STRUCT_MEMBER_TEXTURE( Texture2D,			GBufferBTexture )
+	DECLARE_UNIFORM_BUFFER_STRUCT_MEMBER_TEXTURE( Texture2D,			GBufferCTexture )
+	DECLARE_UNIFORM_BUFFER_STRUCT_MEMBER_TEXTURE( Texture2D,			GBufferDTexture )
+	DECLARE_UNIFORM_BUFFER_STRUCT_MEMBER_TEXTURE( Texture2D,			GBufferETexture )
+	DECLARE_UNIFORM_BUFFER_STRUCT_MEMBER_TEXTURE( Texture2D<float4>,	GBufferATextureNonMS )
+	DECLARE_UNIFORM_BUFFER_STRUCT_MEMBER_TEXTURE( Texture2D<float4>,	GBufferBTextureNonMS )
+	DECLARE_UNIFORM_BUFFER_STRUCT_MEMBER_TEXTURE( Texture2D<float4>,	GBufferCTextureNonMS )
+	DECLARE_UNIFORM_BUFFER_STRUCT_MEMBER_TEXTURE( Texture2D<float4>,	GBufferDTextureNonMS )
+	DECLARE_UNIFORM_BUFFER_STRUCT_MEMBER_TEXTURE( Texture2D<float4>,	GBufferETextureNonMS )
+	DECLARE_UNIFORM_BUFFER_STRUCT_MEMBER_TEXTURE( Texture2DMS<float4>,	GBufferATextureMS )
+	DECLARE_UNIFORM_BUFFER_STRUCT_MEMBER_TEXTURE( Texture2DMS<float4>,	GBufferBTextureMS )
+	DECLARE_UNIFORM_BUFFER_STRUCT_MEMBER_TEXTURE( Texture2DMS<float4>,	GBufferCTextureMS )
+	DECLARE_UNIFORM_BUFFER_STRUCT_MEMBER_TEXTURE( Texture2DMS<float4>,	GBufferDTextureMS )
+	DECLARE_UNIFORM_BUFFER_STRUCT_MEMBER_TEXTURE( Texture2DMS<float4>,	GBufferETextureMS )
+	DECLARE_UNIFORM_BUFFER_STRUCT_MEMBER_SAMPLER( SamplerState,			GBufferATextureSampler )
+	DECLARE_UNIFORM_BUFFER_STRUCT_MEMBER_SAMPLER( SamplerState,			GBufferBTextureSampler )
+	DECLARE_UNIFORM_BUFFER_STRUCT_MEMBER_SAMPLER( SamplerState,			GBufferCTextureSampler )
+	DECLARE_UNIFORM_BUFFER_STRUCT_MEMBER_SAMPLER( SamplerState,			GBufferDTextureSampler )
+	DECLARE_UNIFORM_BUFFER_STRUCT_MEMBER_SAMPLER( SamplerState,			GBufferETextureSampler )
+END_UNIFORM_BUFFER_STRUCT( FGBufferResourceStruct )
+
 /**
  * Encapsulates the render targets used for scene rendering.
  */
@@ -320,7 +343,12 @@ public:
 	TRefCountPtr<IPooledRenderTarget> GetTranslucencyVolumeDirectional(ETranslucencyVolumeCascade Cascade) { return TranslucencyLightingVolumeDirectional[SelectTranslucencyVolumeTarget(Cascade)]; }
 
 	// ---
-
+	/** Get the uniform buffer containing GBuffer resources. */
+	FUniformBufferRHIParamRef GetGBufferResourcesUniformBuffer() const 
+	{ 
+		check(IsValidRef(GBufferResourcesUniformBuffer));
+		return GBufferResourcesUniformBuffer; 
+	}
 	/** */
 	FIntPoint GetBufferSizeXY() const { return BufferSize; }
 	/** */
@@ -490,7 +518,8 @@ private:
 	EPixelFormat GetSceneColorFormat() const;
 
 private:
-
+	/** Uniform buffer containing GBuffer resources. */
+	FUniformBufferRHIRef GBufferResourcesUniformBuffer;
 	/** size of the back buffer, in editor this has to be >= than the biggest view port */
 	FIntPoint BufferSize;
 	/** e.g. 2 */
@@ -542,7 +571,7 @@ public:
 
 	/** Sets the scene texture parameters for the given view. */
 	template<typename TParamRef>
-	void Set(const TParamRef& ShaderRHI, const FSceneView& View, ESceneRenderTargetsMode::Type TextureMode = ESceneRenderTargetsMode::SetTextures, ESamplerFilter ColorFilter = SF_Point) const
+	void Set(FRHICommandList* RHICmdList, const TParamRef& ShaderRHI, const FSceneView& View, ESceneRenderTargetsMode::Type TextureMode = ESceneRenderTargetsMode::SetTextures, ESamplerFilter ColorFilter = SF_Point) const
 	{
 		if (TextureMode == ESceneRenderTargetsMode::SetTextures)
 		{
@@ -571,6 +600,7 @@ public:
 				}
 
 				SetTextureParameter(
+					RHICmdList, 
 					ShaderRHI,
 					SceneColorTextureParameter,
 					SceneColorTextureParameterSampler,
@@ -584,6 +614,7 @@ public:
 				FSamplerStateRHIRef Filter;
 				Filter = TStaticSamplerState<SF_Point,AM_Clamp,AM_Clamp,AM_Clamp>::GetRHI();
 				SetTextureParameter(
+					RHICmdList, 
 					ShaderRHI,
 					SceneAlphaCopyTextureParameter,
 					SceneAlphaCopyTextureParameterSampler,
@@ -596,6 +627,7 @@ public:
 			{
 				const FTexture2DRHIRef* DepthTexture = GSceneRenderTargets.GetActualDepthTexture();
 				SetTextureParameter(
+					RHICmdList, 
 					ShaderRHI,
 					SceneDepthTextureParameter,
 					SceneDepthTextureParameterSampler,
@@ -606,30 +638,30 @@ public:
 
 			if (GRHIFeatureLevel >= ERHIFeatureLevel::SM5)
 			{
-				SetTextureParameter(ShaderRHI, SceneColorSurfaceParameter, GSceneRenderTargets.GetSceneColorSurface());
+				SetTextureParameter(RHICmdList, ShaderRHI, SceneColorSurfaceParameter, GSceneRenderTargets.GetSceneColorSurface());
 			}
-			if (GRHIFeatureLevel >= ERHIFeatureLevel::SM4)
+			if(GRHIFeatureLevel >= ERHIFeatureLevel::SM4)
 			{
 				if(GSupportsDepthFetchDuringDepthTest)
 				{
 					if(SceneDepthSurfaceParameter.IsBound())
 					{
-						SetTextureParameter(ShaderRHI, SceneDepthSurfaceParameter, GSceneRenderTargets.GetSceneDepthSurface());
+						SetTextureParameter(RHICmdList, ShaderRHI, SceneDepthSurfaceParameter, GSceneRenderTargets.GetSceneDepthSurface());
 					}
 					if(SceneDepthTextureNonMS.IsBound())
 					{
-						SetTextureParameter(ShaderRHI, SceneDepthTextureNonMS, GSceneRenderTargets.GetSceneDepthTexture());
+						SetTextureParameter(RHICmdList, ShaderRHI, SceneDepthTextureNonMS, GSceneRenderTargets.GetSceneDepthTexture());
 					}
 				}
 				else
 				{
 					if(SceneDepthSurfaceParameter.IsBound())
 					{
-						SetTextureParameter(ShaderRHI, SceneDepthSurfaceParameter, GSceneRenderTargets.GetAuxiliarySceneDepthSurface());
+						SetTextureParameter(RHICmdList, ShaderRHI, SceneDepthSurfaceParameter, GSceneRenderTargets.GetAuxiliarySceneDepthSurface());
 					}
 					if(SceneDepthTextureNonMS.IsBound())
 					{
-						SetTextureParameter(ShaderRHI, SceneDepthTextureNonMS, GSceneRenderTargets.GetAuxiliarySceneDepthSurface());
+						SetTextureParameter(RHICmdList, ShaderRHI, SceneDepthTextureNonMS, GSceneRenderTargets.GetAuxiliarySceneDepthSurface());
 					}
 				}
 			}
@@ -654,10 +686,11 @@ public:
 		}
 		else if( TextureMode == ESceneRenderTargetsMode::NonSceneAlignedPass )
 		{
-			SetTextureParameter(ShaderRHI, SceneColorTextureParameter, GBlackTexture->TextureRHI);
-			SetTextureParameter(ShaderRHI, SceneDepthTextureParameter, GBlackTexture->TextureRHI);
-			SetTextureParameter(ShaderRHI, SceneColorSurfaceParameter, GBlackTexture->TextureRHI);
-			SetTextureParameter(ShaderRHI, SceneDepthSurfaceParameter, GBlackTexture->TextureRHI);
+			FSamplerStateRHIParamRef DefaultSampler = TStaticSamplerState<SF_Point,AM_Clamp,AM_Clamp,AM_Clamp>::GetRHI();
+			SetTextureParameter(RHICmdList, ShaderRHI, SceneColorTextureParameter, SceneColorTextureParameterSampler, DefaultSampler, GBlackTexture->TextureRHI);
+			SetTextureParameter(RHICmdList, ShaderRHI, SceneDepthTextureParameter, SceneDepthTextureParameterSampler, DefaultSampler, GBlackTexture->TextureRHI);
+			SetTextureParameter(RHICmdList, ShaderRHI, SceneColorSurfaceParameter, GBlackTexture->TextureRHI);
+			SetTextureParameter(RHICmdList, ShaderRHI, SceneDepthSurfaceParameter, GBlackTexture->TextureRHI);
 		}
 	}
 
@@ -690,13 +723,18 @@ class FDeferredPixelShaderParameters
 {
 public:
 
+	// TODO: Remove me.
+	static void ModifyCompilationEnvironment(EShaderPlatform Platform, FShaderCompilerEnvironment& OutEnvironment)
+	{
+	}
+
 	void Bind(const FShaderParameterMap& ParameterMap);
 
 	template<typename ShaderRHIParamRef>
-	void Set(const ShaderRHIParamRef ShaderRHI, const FSceneView& View, ESceneRenderTargetsMode::Type TextureMode = ESceneRenderTargetsMode::SetTextures) const
+	void Set(FRHICommandList* RHICmdList, const ShaderRHIParamRef ShaderRHI, const FSceneView& View, ESceneRenderTargetsMode::Type TextureMode = ESceneRenderTargetsMode::SetTextures) const
 	{
 		// This is needed on PC ES2 for SceneAlphaCopy, probably should be refactored for performance.
-		SceneTextureParameters.Set(ShaderRHI, View, TextureMode, SF_Point);
+		SceneTextureParameters.Set(RHICmdList, ShaderRHI, View, TextureMode, SF_Point);
 
 		// if() is purely an optimization and could be removed
 		if(IsDBufferEnabled())
@@ -706,15 +744,15 @@ public:
 			IPooledRenderTarget* DBufferC = GSceneRenderTargets.DBufferC ? GSceneRenderTargets.DBufferC : GSystemTextures.BlackDummy;
 
 			// todo: optimize out when not needed
-			SetTextureParameter(ShaderRHI, DBufferATexture, DBufferATextureSampler, TStaticSamplerState<>::GetRHI(), DBufferA->GetRenderTargetItem().ShaderResourceTexture);
-			SetTextureParameter(ShaderRHI, DBufferBTexture, DBufferBTextureSampler, TStaticSamplerState<>::GetRHI(), DBufferB->GetRenderTargetItem().ShaderResourceTexture);
-			SetTextureParameter(ShaderRHI, DBufferCTexture, DBufferCTextureSampler, TStaticSamplerState<>::GetRHI(), DBufferC->GetRenderTargetItem().ShaderResourceTexture);
-			SetTextureParameter(ShaderRHI, DBufferATextureMS, DBufferA->GetRenderTargetItem().TargetableTexture);
-			SetTextureParameter(ShaderRHI, DBufferBTextureMS, DBufferB->GetRenderTargetItem().TargetableTexture);
-			SetTextureParameter(ShaderRHI, DBufferCTextureMS, DBufferC->GetRenderTargetItem().TargetableTexture);
-			SetTextureParameter(ShaderRHI, DBufferATextureNonMS, DBufferA->GetRenderTargetItem().ShaderResourceTexture);
-			SetTextureParameter(ShaderRHI, DBufferBTextureNonMS, DBufferB->GetRenderTargetItem().ShaderResourceTexture);
-			SetTextureParameter(ShaderRHI, DBufferCTextureNonMS, DBufferC->GetRenderTargetItem().ShaderResourceTexture);
+			SetTextureParameter(RHICmdList, ShaderRHI, DBufferATexture, DBufferATextureSampler, TStaticSamplerState<>::GetRHI(), DBufferA->GetRenderTargetItem().ShaderResourceTexture);
+			SetTextureParameter(RHICmdList, ShaderRHI, DBufferBTexture, DBufferBTextureSampler, TStaticSamplerState<>::GetRHI(), DBufferB->GetRenderTargetItem().ShaderResourceTexture);
+			SetTextureParameter(RHICmdList, ShaderRHI, DBufferCTexture, DBufferCTextureSampler, TStaticSamplerState<>::GetRHI(), DBufferC->GetRenderTargetItem().ShaderResourceTexture);
+			SetTextureParameter(RHICmdList, ShaderRHI, DBufferATextureMS, DBufferA->GetRenderTargetItem().TargetableTexture);
+			SetTextureParameter(RHICmdList, ShaderRHI, DBufferBTextureMS, DBufferB->GetRenderTargetItem().TargetableTexture);
+			SetTextureParameter(RHICmdList, ShaderRHI, DBufferCTextureMS, DBufferC->GetRenderTargetItem().TargetableTexture);
+			SetTextureParameter(RHICmdList, ShaderRHI, DBufferATextureNonMS, DBufferA->GetRenderTargetItem().ShaderResourceTexture);
+			SetTextureParameter(RHICmdList, ShaderRHI, DBufferBTextureNonMS, DBufferB->GetRenderTargetItem().ShaderResourceTexture);
+			SetTextureParameter(RHICmdList, ShaderRHI, DBufferCTextureNonMS, DBufferC->GetRenderTargetItem().ShaderResourceTexture);
 		}
 
 		if (TextureMode == ESceneRenderTargetsMode::SetTextures && GRHIFeatureLevel >= ERHIFeatureLevel::SM3)
@@ -735,52 +773,25 @@ public:
 
 			if (GRHIFeatureLevel >= ERHIFeatureLevel::SM4)
 			{
-				IPooledRenderTarget* GBufferA = GSceneRenderTargets.GBufferA ? GSceneRenderTargets.GBufferA : GSystemTextures.BlackDummy;
-				IPooledRenderTarget* GBufferB = GSceneRenderTargets.GBufferB ? GSceneRenderTargets.GBufferB : GSystemTextures.BlackDummy;
-				IPooledRenderTarget* GBufferC = GSceneRenderTargets.GBufferC ? GSceneRenderTargets.GBufferC : GSystemTextures.BlackDummy;
-				IPooledRenderTarget* GBufferD = GSceneRenderTargets.GBufferD ? GSceneRenderTargets.GBufferD : GSystemTextures.BlackDummy;
-
-				SetTextureParameter(ShaderRHI, GBufferATexture, GBufferATextureSampler, TStaticSamplerState<>::GetRHI(), GBufferA->GetRenderTargetItem().ShaderResourceTexture);
-				SetTextureParameter(ShaderRHI, GBufferBTexture, GBufferBTextureSampler, TStaticSamplerState<>::GetRHI(), GBufferB->GetRenderTargetItem().ShaderResourceTexture);
-				SetTextureParameter(ShaderRHI, GBufferCTexture, GBufferCTextureSampler, TStaticSamplerState<>::GetRHI(), GBufferC->GetRenderTargetItem().ShaderResourceTexture);
-				SetTextureParameter(ShaderRHI, GBufferDTexture, GBufferDTextureSampler, TStaticSamplerState<>::GetRHI(), GBufferD->GetRenderTargetItem().ShaderResourceTexture);
-				SetTextureParameter(ShaderRHI, GBufferATextureMS, GBufferA->GetRenderTargetItem().TargetableTexture);
-				SetTextureParameter(ShaderRHI, GBufferBTextureMS, GBufferB->GetRenderTargetItem().TargetableTexture);
-				SetTextureParameter(ShaderRHI, GBufferCTextureMS, GBufferC->GetRenderTargetItem().TargetableTexture);
-				SetTextureParameter(ShaderRHI, GBufferDTextureMS, GBufferD->GetRenderTargetItem().TargetableTexture);
-				SetTextureParameter(ShaderRHI, GBufferATextureNonMS, GBufferA->GetRenderTargetItem().ShaderResourceTexture);
-				SetTextureParameter(ShaderRHI, GBufferBTextureNonMS, GBufferB->GetRenderTargetItem().ShaderResourceTexture);
-				SetTextureParameter(ShaderRHI, GBufferCTextureNonMS, GBufferC->GetRenderTargetItem().ShaderResourceTexture);
-				SetTextureParameter(ShaderRHI, GBufferDTextureNonMS, GBufferD->GetRenderTargetItem().ShaderResourceTexture);
-
-				SetTextureParameter(ShaderRHI, ScreenSpaceAOTexture, ScreenSpaceAOTextureSampler, TStaticSamplerState<>::GetRHI(), ScreenSpaceAO->GetRenderTargetItem().ShaderResourceTexture);
-				SetTextureParameter(ShaderRHI, ScreenSpaceAOTextureMS, ScreenSpaceAO->GetRenderTargetItem().TargetableTexture);
-				SetTextureParameter(ShaderRHI, ScreenSpaceAOTextureNonMS, ScreenSpaceAO->GetRenderTargetItem().ShaderResourceTexture);
-
-				SetTextureParameter(ShaderRHI, CustomDepthTexture, CustomDepthTextureSampler, TStaticSamplerState<>::GetRHI(), CustomDepth->GetRenderTargetItem().ShaderResourceTexture);
-				SetTextureParameter(ShaderRHI, CustomDepthTextureNonMS, CustomDepth->GetRenderTargetItem().ShaderResourceTexture);
-
-				if (GSceneRenderTargets.IsStaticLightingAllowed())
+				//gilmerge
+				if (GBufferResources.IsBound())
 				{
-					IPooledRenderTarget* GBufferE = GSceneRenderTargets.GBufferE ? GSceneRenderTargets.GBufferE : GSystemTextures.BlackDummy;
-
-					SetTextureParameter(ShaderRHI, GBufferETexture, GBufferETextureSampler, TStaticSamplerState<>::GetRHI(), GBufferE->GetRenderTargetItem().ShaderResourceTexture);
-					SetTextureParameter(ShaderRHI, GBufferETextureMS, GBufferE->GetRenderTargetItem().TargetableTexture);
-					SetTextureParameter(ShaderRHI, GBufferETextureNonMS, GBufferE->GetRenderTargetItem().ShaderResourceTexture);
+					SetUniformBufferParameter(RHICmdList, ShaderRHI, GBufferResources, GSceneRenderTargets.GetGBufferResourcesUniformBuffer());
 				}
+
+				SetTextureParameter(RHICmdList, ShaderRHI, ScreenSpaceAOTexture, ScreenSpaceAOTextureSampler, TStaticSamplerState<>::GetRHI(), ScreenSpaceAO->GetRenderTargetItem().ShaderResourceTexture);
+				SetTextureParameter(RHICmdList, ShaderRHI, ScreenSpaceAOTextureMS, ScreenSpaceAO->GetRenderTargetItem().TargetableTexture);
+				SetTextureParameter(RHICmdList, ShaderRHI, ScreenSpaceAOTextureNonMS, ScreenSpaceAO->GetRenderTargetItem().ShaderResourceTexture);
+
+				SetTextureParameter(RHICmdList, ShaderRHI, CustomDepthTexture, CustomDepthTextureSampler, TStaticSamplerState<>::GetRHI(), CustomDepth->GetRenderTargetItem().ShaderResourceTexture);
+				SetTextureParameter(RHICmdList, ShaderRHI, CustomDepthTextureNonMS, CustomDepth->GetRenderTargetItem().ShaderResourceTexture);
 			}
 		}
 		else if (TextureMode == ESceneRenderTargetsMode::DontSet ||
 			TextureMode == ESceneRenderTargetsMode::DontSetIgnoreBoundByEditorCompositing)
 		{
 			// Verify that none of these are actually bound
-			checkSlow(!GBufferATextureSampler.IsBound()
-				&& !GBufferBTextureSampler.IsBound()
-				&& !GBufferCTextureSampler.IsBound()
-				&& !GBufferDTextureSampler.IsBound()
-				&& !ScreenSpaceAOTextureSampler.IsBound()
-				&& !CustomDepthTextureSampler.IsBound()
-				&& !GBufferETextureSampler.IsBound());
+			checkSlow(!GBufferResources.IsBound());
 		}
 	}
 
@@ -788,35 +799,16 @@ public:
 
 private:
 	FSceneTextureShaderParameters SceneTextureParameters;
-	FShaderResourceParameter GBufferATextureMS;
-	FShaderResourceParameter GBufferBTextureMS;
-	FShaderResourceParameter GBufferCTextureMS;
-	FShaderResourceParameter GBufferDTextureMS;
-	FShaderResourceParameter GBufferETextureMS;
+	FShaderUniformBufferParameter GBufferResources;
 	FShaderResourceParameter DBufferATextureMS;
 	FShaderResourceParameter DBufferBTextureMS;
 	FShaderResourceParameter DBufferCTextureMS;
 	FShaderResourceParameter ScreenSpaceAOTextureMS;
-	FShaderResourceParameter GBufferATextureNonMS;
-	FShaderResourceParameter GBufferBTextureNonMS;
-	FShaderResourceParameter GBufferCTextureNonMS;
-	FShaderResourceParameter GBufferDTextureNonMS;
-	FShaderResourceParameter GBufferETextureNonMS;
 	FShaderResourceParameter DBufferATextureNonMS;
 	FShaderResourceParameter DBufferBTextureNonMS;
 	FShaderResourceParameter DBufferCTextureNonMS;
 	FShaderResourceParameter ScreenSpaceAOTextureNonMS;
 	FShaderResourceParameter CustomDepthTextureNonMS;
-	FShaderResourceParameter GBufferATexture;
-	FShaderResourceParameter GBufferATextureSampler;
-	FShaderResourceParameter GBufferBTexture;
-	FShaderResourceParameter GBufferBTextureSampler;
-	FShaderResourceParameter GBufferCTexture;
-	FShaderResourceParameter GBufferCTextureSampler;
-	FShaderResourceParameter GBufferDTexture;
-	FShaderResourceParameter GBufferDTextureSampler;
-	FShaderResourceParameter GBufferETexture;
-	FShaderResourceParameter GBufferETextureSampler;
 	FShaderResourceParameter DBufferATexture;
 	FShaderResourceParameter DBufferATextureSampler;
 	FShaderResourceParameter DBufferBTexture;
diff --git a/Engine/Source/Runtime/Renderer/Private/PostProcess/ScreenSpaceReflections.cpp b/Engine/Source/Runtime/Renderer/Private/PostProcess/ScreenSpaceReflections.cpp
--- a/Engine/Source/Runtime/Renderer/Private/PostProcess/ScreenSpaceReflections.cpp
+++ b/Engine/Source/Runtime/Renderer/Private/PostProcess/ScreenSpaceReflections.cpp
@@ -74,6 +74,7 @@ class FPostProcessScreenSpaceReflectionsPS : public FGlobalShader
 	static void ModifyCompilationEnvironment(EShaderPlatform Platform, FShaderCompilerEnvironment& OutEnvironment)
 	{
 		FGlobalShader::ModifyCompilationEnvironment(Platform, OutEnvironment);
+		FDeferredPixelShaderParameters::ModifyCompilationEnvironment(Platform,OutEnvironment);
 		OutEnvironment.SetDefine( TEXT("PREV_FRAME_COLOR"), PrevFrame );
 		OutEnvironment.SetDefine( TEXT("SSR_QUALITY"), SSRQuality );
 	}
@@ -95,15 +96,15 @@ public:
 		SSRParams.Bind(Initializer.ParameterMap, TEXT("SSRParams"));
 	}
 
-	void SetParameters(const FRenderingCompositePassContext& Context)
+	void SetParameters(FRHICommandList* RHICmdList, const FRenderingCompositePassContext& Context)
 	{
 		const FFinalPostProcessSettings& Settings = Context.View.FinalPostProcessSettings;
 		const FPixelShaderRHIParamRef ShaderRHI = GetPixelShader();
 
-		FGlobalShader::SetParameters(ShaderRHI, Context.View);
+		FGlobalShader::SetParameters(RHICmdList, ShaderRHI, Context.View);
 
-		PostprocessParameter.SetPS(ShaderRHI, Context, TStaticSamplerState<SF_Point,AM_Clamp,AM_Clamp,AM_Clamp>::GetRHI());
-		DeferredParameters.Set(ShaderRHI, Context.View);
+		PostprocessParameter.SetPS(RHICmdList, ShaderRHI, Context, TStaticSamplerState<SF_Point,AM_Clamp,AM_Clamp,AM_Clamp>::GetRHI());
+		DeferredParameters.Set(RHICmdList, ShaderRHI, Context.View);
 
 		{
 			float MaxRoughness = FMath::Clamp(Context.View.FinalPostProcessSettings.ScreenSpaceReflectionMaxRoughness, 0.01f, 1.0f);
@@ -120,7 +121,7 @@ public:
 				0, 
 				0);
 
-			SetShaderValue(ShaderRHI, SSRParams, Value);
+			SetShaderValue(RHICmdList, ShaderRHI, SSRParams, Value);
 		}
 	}
 
@@ -174,6 +175,9 @@ void FRCPassPostProcessScreenSpaceReflections::Process(FRenderingCompositePassCo
 {
 	SCOPED_DRAW_EVENT(ScreenSpaceReflections, DEC_SCENE_ITEMS);
 
+	//@todo-rco: RHIPacketList
+	FRHICommandList* RHICmdList = nullptr;
+
 	const FSceneView& View = Context.View;
 
 	const FSceneRenderTargetItem& DestRenderTarget = PassOutputs[0].RequestSurface(Context);
@@ -199,8 +203,8 @@ void FRCPassPostProcessScreenSpaceReflections::Process(FRenderingCompositePassCo
 			TShaderMapRef< FPostProcessScreenSpaceReflectionsPS<A, B> > PixelShader(GetGlobalShaderMap()); \
 			static FGlobalBoundShaderState BoundShaderState; \
 			SetGlobalBoundShaderState(BoundShaderState, GFilterVertexDeclaration.VertexDeclarationRHI, *VertexShader, *PixelShader); \
-			VertexShader->SetParameters(Context); \
-			PixelShader->SetParameters(Context); \
+			VertexShader->SetParameters(RHICmdList, Context); \
+			PixelShader->SetParameters(RHICmdList, Context); \
 		}; \
 		break
 
diff --git a/Engine/Source/Runtime/Renderer/Private/PostProcess/VisualizeTexture.cpp b/Engine/Source/Runtime/Renderer/Private/PostProcess/VisualizeTexture.cpp
--- a/Engine/Source/Runtime/Renderer/Private/PostProcess/VisualizeTexture.cpp
+++ b/Engine/Source/Runtime/Renderer/Private/PostProcess/VisualizeTexture.cpp
@@ -73,7 +73,7 @@ public:
 		return bShaderHasOutdatedParameters;
 	}
 
-	void SetParameters(const FVisualizeTextureData& Data)
+	void SetParameters(FRHICommandList* RHICmdList, const FVisualizeTextureData& Data)
 	{
 		const FPixelShaderRHIParamRef ShaderRHI = GetPixelShader();
 		
@@ -93,22 +93,22 @@ public:
 			VisualizeParamValue[1] = FVector4(BlinkState, Data.bSaturateInsteadOfFrac ? 1.0f : 0.0f, Data.ArrayIndex, Data.CustomMip);
 			VisualizeParamValue[2] = FVector4(Data.InputValueMapping, 0, 0,0);
 
-			SetShaderValueArray(ShaderRHI, VisualizeParam, VisualizeParamValue, 3);
+			SetShaderValueArray(RHICmdList, ShaderRHI, VisualizeParam, VisualizeParamValue, 3);
 		}
 
 		{
 			FVector4 TextureExtentValue(Data.Desc.Extent.X, Data.Desc.Extent.Y, Data.Desc.Depth, 0);
 
-			SetShaderValue(ShaderRHI, TextureExtent, TextureExtentValue);
+			SetShaderValue(RHICmdList, ShaderRHI, TextureExtent, TextureExtentValue);
 		}
 		
 		
-		SetSRVParameter( ShaderRHI, VisualizeDepthStencilTexture, Data.StencilSRV );		
-		SetTextureParameter(ShaderRHI, VisualizeTexture2D, VisualizeTexture2DSampler, TStaticSamplerState<SF_Point,AM_Clamp,AM_Clamp,AM_Clamp>::GetRHI(), (FTextureRHIRef&)Data.RenderTargetItem.ShaderResourceTexture);
-		SetTextureParameter(ShaderRHI, VisualizeTexture2DMS, (FTextureRHIRef&)Data.RenderTargetItem.TargetableTexture);
-		SetTextureParameter(ShaderRHI, VisualizeTexture3D, VisualizeTexture3DSampler, TStaticSamplerState<SF_Point,AM_Clamp,AM_Clamp,AM_Clamp>::GetRHI(), (FTextureRHIRef&)Data.RenderTargetItem.ShaderResourceTexture);
-		SetTextureParameter(ShaderRHI, VisualizeTextureCube, VisualizeTextureCubeSampler, TStaticSamplerState<SF_Point,AM_Clamp,AM_Clamp,AM_Clamp>::GetRHI(), (FTextureRHIRef&)Data.RenderTargetItem.ShaderResourceTexture);
-		SetTextureParameter(ShaderRHI, VisualizeTextureCubeArray, VisualizeTextureCubeArraySampler, TStaticSamplerState<SF_Point,AM_Clamp,AM_Clamp,AM_Clamp>::GetRHI(), (FTextureRHIRef&)Data.RenderTargetItem.ShaderResourceTexture);
+		SetSRVParameter(RHICmdList, ShaderRHI, VisualizeDepthStencilTexture, Data.StencilSRV );		
+		SetTextureParameter(RHICmdList, ShaderRHI, VisualizeTexture2D, VisualizeTexture2DSampler, TStaticSamplerState<SF_Point,AM_Clamp,AM_Clamp,AM_Clamp>::GetRHI(), (FTextureRHIRef&)Data.RenderTargetItem.ShaderResourceTexture);
+		SetTextureParameter(RHICmdList, ShaderRHI, VisualizeTexture2DMS, (FTextureRHIRef&)Data.RenderTargetItem.TargetableTexture);
+		SetTextureParameter(RHICmdList, ShaderRHI, VisualizeTexture3D, VisualizeTexture3DSampler, TStaticSamplerState<SF_Point,AM_Clamp,AM_Clamp,AM_Clamp>::GetRHI(), (FTextureRHIRef&)Data.RenderTargetItem.ShaderResourceTexture);
+		SetTextureParameter(RHICmdList, ShaderRHI, VisualizeTextureCube, VisualizeTextureCubeSampler, TStaticSamplerState<SF_Point,AM_Clamp,AM_Clamp,AM_Clamp>::GetRHI(), (FTextureRHIRef&)Data.RenderTargetItem.ShaderResourceTexture);
+		SetTextureParameter(RHICmdList, ShaderRHI, VisualizeTextureCubeArray, VisualizeTextureCubeArraySampler, TStaticSamplerState<SF_Point,AM_Clamp,AM_Clamp,AM_Clamp>::GetRHI(), (FTextureRHIRef&)Data.RenderTargetItem.ShaderResourceTexture);
 	}
 
 	static const TCHAR* GetSourceFilename()
@@ -178,12 +178,12 @@ public:
 		return bShaderHasOutdatedParameters;
 	}
 
-	void SetParameters(const FSceneView& View, const IPooledRenderTarget& Src)
+	void SetParameters(FRHICommandList* RHICmdList, const FSceneView& View, const IPooledRenderTarget& Src)
 	{
 		const FPixelShaderRHIParamRef ShaderRHI = GetPixelShader();
 
-		FGlobalShader::SetParameters(ShaderRHI, View);
-		SetTextureParameter(ShaderRHI, VisualizeTexture2D, VisualizeTexture2DSampler, 
+		FGlobalShader::SetParameters(RHICmdList, ShaderRHI, View);
+		SetTextureParameter(RHICmdList, ShaderRHI, VisualizeTexture2D, VisualizeTexture2DSampler, 
 			TStaticSamplerState<SF_Point,AM_Clamp,AM_Clamp,AM_Clamp>::GetRHI(), (FTextureRHIRef&)Src.GetRenderTargetItem().ShaderResourceTexture);
 	}
 };
@@ -193,12 +193,15 @@ IMPLEMENT_SHADER_TYPE(,FVisualizeTexturePresentPS,TEXT("VisualizeTexture"),TEXT(
 
 template<uint32 TextureType> void VisualizeTextureForTextureType(const FVisualizeTextureData& Data)
 {
+	//@todo-rco: RHIPacketList
+	FRHICommandList* RHICmdList = nullptr;
+
 	TShaderMapRef<FScreenVS> VertexShader(GetGlobalShaderMap());
 	TShaderMapRef<VisualizeTexturePS<TextureType> > PixelShader(GetGlobalShaderMap());
 
 	static FGlobalBoundShaderState BoundShaderState;
 	SetGlobalBoundShaderState(BoundShaderState, GFilterVertexDeclaration.VertexDeclarationRHI, *VertexShader, *PixelShader);
-	PixelShader->SetParameters(Data);
+	PixelShader->SetParameters(RHICmdList, Data);
 
 	DrawRectangle(
 		// XY
@@ -489,6 +492,9 @@ void FVisualizeTexture::PresentContent(const FSceneView& View)
 		return;
 	}
 
+	//@todo-rco: RHIPacketList
+	FRHICommandList* RHICmdList = nullptr;
+
 	FPooledRenderTargetDesc Desc = VisualizeTextureDesc;
 
 	auto& RenderTarget = View.Family->RenderTarget->GetRenderTargetTexture();
@@ -506,8 +512,8 @@ void FVisualizeTexture::PresentContent(const FSceneView& View)
 
 	SetGlobalBoundShaderState(BoundShaderState, GFilterVertexDeclaration.VertexDeclarationRHI, *VertexShader, *PixelShader);
 
-	VertexShader->SetParameters(View);
-	PixelShader->SetParameters(View, *VisualizeTextureContent);
+	VertexShader->SetParameters(RHICmdList, View);
+	PixelShader->SetParameters(RHICmdList, View, *VisualizeTextureContent);
 
 	FIntRect DestRect = View.ViewRect;
 
