UE4 - merge first pass at parallel rendering to main

diff --git a/Engine/Source/Runtime/Renderer/Private/ReflectionEnvironment.cpp b/Engine/Source/Runtime/Renderer/Private/ReflectionEnvironment.cpp
--- a/Engine/Source/Runtime/Renderer/Private/ReflectionEnvironment.cpp
+++ b/Engine/Source/Runtime/Renderer/Private/ReflectionEnvironment.cpp
@@ -15,6 +15,7 @@
 #include "PostProcessTemporalAA.h"
 #include "PostProcessDownsample.h"
 #include "ReflectionEnvironment.h"
+#include "ShaderParameterUtils.h"
 #include "LightRendering.h"
 
 /** Tile size for the reflection environment compute shader, tweaked for 680 GTX. */
@@ -59,15 +60,15 @@ void FReflectionEnvironmentCubemapArray::InitDynamicRHI()
 		ReleaseCubeArray();
 
 		FPooledRenderTargetDesc Desc(FPooledRenderTargetDesc::CreateCubemapDesc(
-			GReflectionCaptureSize,
-			//@todo - get rid of the alpha channel (currently stores brightness which is a constant), could use PF_FloatRGB for half memory, would need to implement RHIReadSurface support
-			PF_FloatRGBA, 
-			TexCreate_None,
-			TexCreate_None,
-			false, 
-			// Cubemap array of 1 produces a regular cubemap, so guarantee it will be allocated as an array
-			FMath::Max<uint32>(MaxCubemaps, 2),
-			NumReflectionCaptureMips));
+		GReflectionCaptureSize,
+		//@todo - get rid of the alpha channel (currently stores brightness which is a constant), could use PF_FloatRGB for half memory, would need to implement RHIReadSurface support
+		PF_FloatRGBA, 
+		TexCreate_None,
+		TexCreate_None,
+		false, 
+		// Cubemap array of 1 produces a regular cubemap, so guarantee it will be allocated as an array
+		FMath::Max<uint32>(MaxCubemaps, 2),
+		NumReflectionCaptureMips));
 	
 		// Allocate TextureCubeArray for the scene's reflection captures
 		GRenderTargetPool.FindFreeElement(Desc, ReflectionEnvs, TEXT("ReflectionEnvs"));
@@ -148,6 +149,7 @@ public:
 	static void ModifyCompilationEnvironment(EShaderPlatform Platform, FShaderCompilerEnvironment& OutEnvironment)
 	{
 		FGlobalShader::ModifyCompilationEnvironment(Platform, OutEnvironment);
+		FDeferredPixelShaderParameters::ModifyCompilationEnvironment(Platform,OutEnvironment);
 		OutEnvironment.SetDefine(TEXT("THREADGROUP_SIZEX"), GReflectionEnvironmentTileSizeX);
 		OutEnvironment.SetDefine(TEXT("THREADGROUP_SIZEY"), GReflectionEnvironmentTileSizeY);
 		OutEnvironment.SetDefine(TEXT("MAX_CAPTURES"), GMaxNumReflectionCaptures);
@@ -175,12 +177,12 @@ public:
 	{
 	}
 
-	void SetParameters(const FSceneView& View, FTextureRHIParamRef SSRTexture, FUnorderedAccessViewRHIParamRef OutSceneColorUAV)
+	void SetParameters(FRHICommandList* RHICmdList, const FSceneView& View, FTextureRHIParamRef SSRTexture, FUnorderedAccessViewRHIParamRef OutSceneColorUAV)
 	{
 		const FComputeShaderRHIParamRef ShaderRHI = GetComputeShader();
 
-		FGlobalShader::SetParameters(ShaderRHI, View);
-		DeferredParameters.Set(ShaderRHI, View);
+		FGlobalShader::SetParameters(RHICmdList, ShaderRHI, View);
+		DeferredParameters.Set(RHICmdList, ShaderRHI, View);
 
 		FScene* Scene = (FScene*)View.Family->Scene;
 
@@ -190,18 +192,19 @@ public:
 		FSceneRenderTargetItem& CubemapArray = Scene->ReflectionSceneData.CubemapArray.GetRenderTarget();
 
 		SetTextureParameter(
+			RHICmdList, 
 			ShaderRHI, 
 			ReflectionEnvironmentColorTexture, 
 			ReflectionEnvironmentColorSampler, 
 			TStaticSamplerState<SF_Trilinear, AM_Clamp, AM_Clamp, AM_Clamp>::GetRHI(), 
 			CubemapArray.ShaderResourceTexture);
 
-		SetTextureParameter( ShaderRHI, ScreenSpaceReflections, SSRTexture );
+		SetTextureParameter(RHICmdList, ShaderRHI, ScreenSpaceReflections, SSRTexture );
 
-		SetTextureParameter( ShaderRHI, InSceneColor, GSceneRenderTargets.GetSceneColor()->GetRenderTargetItem().ShaderResourceTexture );
-		OutSceneColor.SetTexture(ShaderRHI, NULL, OutSceneColorUAV);
+		SetTextureParameter(RHICmdList, ShaderRHI, InSceneColor, GSceneRenderTargets.GetSceneColor()->GetRenderTargetItem().ShaderResourceTexture );
+		OutSceneColor.SetTexture(RHICmdList, ShaderRHI, NULL, OutSceneColorUAV);
 
-		SetShaderValue(ShaderRHI, ViewDimensionsParameter, View.ViewRect);
+		SetShaderValue(RHICmdList, ShaderRHI, ViewDimensionsParameter, View.ViewRect);
 
 		static TArray<FReflectionCaptureSortData> SortData;
 		SortData.Reset(Scene->ReflectionSceneData.RegisteredReflectionCaptures.Num());
@@ -260,19 +263,19 @@ public:
 			SamplePositionsBuffer.BoxScales[CaptureIndex] = SortData[CaptureIndex].BoxScales;
 		}
 
-		SetUniformBufferParameterImmediate(ShaderRHI, GetUniformBufferParameter<FReflectionCaptureData>(), SamplePositionsBuffer);
-		SetShaderValue(ShaderRHI, NumCaptures, SortData.Num());
+		SetUniformBufferParameterImmediate(RHICmdList, ShaderRHI, GetUniformBufferParameter<FReflectionCaptureData>(), SamplePositionsBuffer);
+		SetShaderValue(RHICmdList, ShaderRHI, NumCaptures, SortData.Num());
 
-		SetTextureParameter(ShaderRHI, PreIntegratedGF, PreIntegratedGFSampler, TStaticSamplerState<SF_Bilinear, AM_Clamp, AM_Clamp, AM_Clamp>::GetRHI(), GSystemTextures.PreintegratedGF->GetRenderTargetItem().ShaderResourceTexture);
+		SetTextureParameter(RHICmdList, ShaderRHI, PreIntegratedGF, PreIntegratedGFSampler, TStaticSamplerState<SF_Bilinear,AM_Clamp,AM_Clamp,AM_Clamp>::GetRHI(), GSystemTextures.PreintegratedGF->GetRenderTargetItem().ShaderResourceTexture);
 	
-		SkyLightParameters.SetParameters(ShaderRHI, Scene, View.Family->EngineShowFlags.SkyLighting);
+		SkyLightParameters.SetParameters(RHICmdList, ShaderRHI, Scene, View.Family->EngineShowFlags.SkyLighting);
 	}
 
 	void UnsetParameters()
 	{
 		const FComputeShaderRHIParamRef ShaderRHI = GetComputeShader();
-
-		OutSceneColor.UnsetUAV(ShaderRHI);
+		//@todo-rco: RHIPacketList
+		OutSceneColor.UnsetUAV(nullptr, ShaderRHI);
 	}
 
 	virtual bool Serialize(FArchive& Ar)
@@ -345,6 +348,7 @@ class FReflectionApplyPS : public FGlobalShader
 		OutEnvironment.SetDefine(TEXT("APPLY_SSR"), bSSR);
 		OutEnvironment.SetDefine(TEXT("APPLY_REFLECTION_ENV"), bReflectionEnv);
 		OutEnvironment.SetDefine(TEXT("APPLY_SKYLIGHT"), bSkylight);
+		FDeferredPixelShaderParameters::ModifyCompilationEnvironment(Platform,OutEnvironment);
 	}
 
 	/** Default constructor. */
@@ -374,17 +378,17 @@ public:
 		PreIntegratedGFSampler.Bind(Initializer.ParameterMap, TEXT("PreIntegratedGFSampler"));
 	}
 
-	void SetParameters( const FSceneView& View, FTextureRHIParamRef ReflectionEnv, FTextureRHIParamRef ScreenSpaceReflections )
+	void SetParameters(FRHICommandList* RHICmdList, const FSceneView& View, FTextureRHIParamRef ReflectionEnv, FTextureRHIParamRef ScreenSpaceReflections )
 	{
 		const FPixelShaderRHIParamRef ShaderRHI = GetPixelShader();
 
-		FGlobalShader::SetParameters(ShaderRHI, View);
-		DeferredParameters.Set(ShaderRHI, View);
-		SkyLightParameters.SetParameters(ShaderRHI, (FScene*)View.Family->Scene, true);
+		FGlobalShader::SetParameters(RHICmdList, ShaderRHI, View);
+		DeferredParameters.Set(RHICmdList, ShaderRHI, View);
+		SkyLightParameters.SetParameters(RHICmdList, ShaderRHI, (FScene*)View.Family->Scene, true);
 		
-		SetTextureParameter( ShaderRHI, ReflectionEnvTexture, ReflectionEnvSampler, TStaticSamplerState<SF_Point>::GetRHI(), ReflectionEnv );
-		SetTextureParameter( ShaderRHI, ScreenSpaceReflectionsTexture, ScreenSpaceReflectionsSampler, TStaticSamplerState<SF_Point>::GetRHI(), ScreenSpaceReflections );
-		SetTextureParameter( ShaderRHI, PreIntegratedGF, PreIntegratedGFSampler, TStaticSamplerState<SF_Bilinear, AM_Clamp, AM_Clamp, AM_Clamp>::GetRHI(), GSystemTextures.PreintegratedGF->GetRenderTargetItem().ShaderResourceTexture );
+		SetTextureParameter(RHICmdList, ShaderRHI, ReflectionEnvTexture, ReflectionEnvSampler, TStaticSamplerState<SF_Point>::GetRHI(), ReflectionEnv );
+		SetTextureParameter(RHICmdList, ShaderRHI, ScreenSpaceReflectionsTexture, ScreenSpaceReflectionsSampler, TStaticSamplerState<SF_Point>::GetRHI(), ScreenSpaceReflections );
+		SetTextureParameter(RHICmdList, ShaderRHI, PreIntegratedGF, PreIntegratedGFSampler, TStaticSamplerState<SF_Bilinear,AM_Clamp,AM_Clamp,AM_Clamp>::GetRHI(), GSystemTextures.PreintegratedGF->GetRenderTargetItem().ShaderResourceTexture );
 	}
 
 	// FShader interface.
@@ -427,6 +431,12 @@ class FReflectionCaptureSpecularBouncePS : public FGlobalShader
 		return IsFeatureLevelSupported(Platform, ERHIFeatureLevel::SM4);
 	}
 
+	static void ModifyCompilationEnvironment(EShaderPlatform Platform, FShaderCompilerEnvironment& OutEnvironment)
+	{
+		FGlobalShader::ModifyCompilationEnvironment(Platform,OutEnvironment);
+		FDeferredPixelShaderParameters::ModifyCompilationEnvironment(Platform,OutEnvironment);
+	}
+
 	/** Default constructor. */
 	FReflectionCaptureSpecularBouncePS() {}
 
@@ -440,13 +450,13 @@ public:
 		DeferredParameters.Bind(Initializer.ParameterMap);
 	}
 
-	void SetParameters(const FSceneView& View)
+	void SetParameters(FRHICommandList* RHICmdList, const FSceneView& View)
 	{
 		const FPixelShaderRHIParamRef ShaderRHI = GetPixelShader();
 
-		FGlobalShader::SetParameters(ShaderRHI, View);
+		FGlobalShader::SetParameters(RHICmdList, ShaderRHI, View);
 
-		DeferredParameters.Set(ShaderRHI, View);
+		DeferredParameters.Set(RHICmdList, ShaderRHI, View);
 	}
 
 	// FShader interface.
@@ -494,18 +504,18 @@ public:
 		DeferredParameters.Bind(Initializer.ParameterMap);
 	}
 
-	void SetParameters(const FSceneView& View, const FReflectionCaptureSortData& SortData)
+	void SetParameters(FRHICommandList* RHICmdList, const FSceneView& View, const FReflectionCaptureSortData& SortData)
 	{
 		const FPixelShaderRHIParamRef ShaderRHI = GetPixelShader();
-		FGlobalShader::SetParameters(ShaderRHI, View);
+		FGlobalShader::SetParameters(RHICmdList, ShaderRHI, View);
 
-		SetTextureParameter(ShaderRHI, ReflectionEnvironmentColorTexture, ReflectionEnvironmentColorSampler, TStaticSamplerState<SF_Trilinear,AM_Clamp,AM_Clamp,AM_Clamp>::GetRHI(), SortData.SM4FullHDRCubemap->TextureRHI);
+		SetTextureParameter(RHICmdList, ShaderRHI, ReflectionEnvironmentColorTexture, ReflectionEnvironmentColorSampler, TStaticSamplerState<SF_Trilinear,AM_Clamp,AM_Clamp,AM_Clamp>::GetRHI(), SortData.SM4FullHDRCubemap->TextureRHI);
 
-		DeferredParameters.Set(ShaderRHI, View);
-		SetShaderValue(ShaderRHI, CapturePositionAndRadius, SortData.PositionAndRadius);
-		SetShaderValue(ShaderRHI, CaptureProperties, SortData.CaptureProperties);
-		SetShaderValue(ShaderRHI, CaptureBoxTransform, SortData.BoxTransform);
-		SetShaderValue(ShaderRHI, CaptureBoxScales, SortData.BoxScales);
+		DeferredParameters.Set(RHICmdList, ShaderRHI, View);
+		SetShaderValue(RHICmdList, ShaderRHI, CapturePositionAndRadius, SortData.PositionAndRadius);
+		SetShaderValue(RHICmdList, ShaderRHI, CaptureProperties, SortData.CaptureProperties);
+		SetShaderValue(RHICmdList, ShaderRHI, CaptureBoxTransform, SortData.BoxTransform);
+		SetShaderValue(RHICmdList, ShaderRHI, CaptureBoxScales, SortData.BoxScales);
 	}
 
 	// FShader interface.
@@ -537,9 +547,9 @@ IMPLEMENT_SHADER_TYPE(template<>,TStandardDeferredReflectionPS<true>,TEXT("Refle
 IMPLEMENT_SHADER_TYPE(template<>,TStandardDeferredReflectionPS<false>,TEXT("ReflectionEnvironmentShaders"),TEXT("StandardDeferredReflectionPS"),SF_Pixel);
 
 void FDeferredShadingSceneRenderer::RenderReflectionCaptureSpecularBounceForAllViews()
-{		
-	// We're currently capturing a reflection capture, output SpecularColor * IndirectDiffuseGI for metals so they are not black in reflections,
-	// Since we don't have multiple bounce specular reflections
+{
+				// We're currently capturing a reflection capture, output SpecularColor * IndirectDiffuseGI for metals so they are not black in reflections,
+				// Since we don't have multiple bounce specular reflections
 	GSceneRenderTargets.BeginRenderingSceneColor();
 	RHISetRasterizerState( TStaticRasterizerState< FM_Solid, CM_None >::GetRHI() );
 	RHISetDepthStencilState( TStaticDepthStencilState< false, CF_Always >::GetRHI() );
@@ -557,12 +567,13 @@ void FDeferredShadingSceneRenderer::RenderReflectionCaptureSpecularBounceForAllV
 
 		RHISetViewport( View.ViewRect.Min.X, View.ViewRect.Min.Y, 0.0f, View.ViewRect.Max.X, View.ViewRect.Max.Y, 1.0f );
 
-		PixelShader->SetParameters( View );
+		//@todo-rco: RHIPacketList
+		PixelShader->SetParameters(nullptr, View );
 
 		DrawRectangle( 
-			0, 0, 
+			0, 0,
 			View.ViewRect.Width(), View.ViewRect.Height(),
-			0, 0, 
+			0, 0,
 			View.ViewRect.Width(), View.ViewRect.Height(),
 			FIntPoint(View.ViewRect.Width(), View.ViewRect.Height()),
 			GSceneRenderTargets.GetBufferSizeXY(),
@@ -580,7 +591,7 @@ bool FDeferredShadingSceneRenderer::ShouldDoReflectionEnvironment() const
 		&& ViewFamily.EngineShowFlags.ReflectionEnvironment;
 }
 
-void FDeferredShadingSceneRenderer::RenderImageBasedReflectionsSM5ForAllViews()
+void FDeferredShadingSceneRenderer::RenderImageBasedReflectionsSM5ForAllViews(FRHICommandList* RHICmdList)
 {
 	const uint32 bUseLightmaps = CVarDiffuseFromCaptures.GetValueOnRenderThread() == 0;
 
@@ -594,7 +605,7 @@ void FDeferredShadingSceneRenderer::RenderImageBasedReflectionsSM5ForAllViews()
 		GRenderTargetPool.FindFreeElement( Desc, NewSceneColor, TEXT("SceneColorEnv") );
 	}
 
-	// If we are in SM5, use the compute shader gather method
+				// If we are in SM5, use the compute shader gather method
 	for (int32 ViewIndex = 0, Num = Views.Num(); ViewIndex < Num; ViewIndex++)
 	{
 		const FViewInfo& View = Views[ViewIndex];
@@ -626,12 +637,12 @@ void FDeferredShadingSceneRenderer::RenderImageBasedReflectionsSM5ForAllViews()
 
 			RHISetComputeShader(ComputeShader->GetComputeShader());
 
-			ComputeShader->SetParameters(View, SSROutput->GetRenderTargetItem().ShaderResourceTexture, NewSceneColor->GetRenderTargetItem().UAV);
-			//ComputeShader->SetParameters(View, DestRenderTarget.UAV);
+			ComputeShader->SetParameters(RHICmdList, View, SSROutput->GetRenderTargetItem().ShaderResourceTexture, NewSceneColor->GetRenderTargetItem().UAV);
+			//ComputeShader->SetParameters(RHICmdList, View, DestRenderTarget.UAV);
 
 			uint32 GroupSizeX = (View.ViewRect.Size().X + GReflectionEnvironmentTileSizeX - 1) / GReflectionEnvironmentTileSizeX;
 			uint32 GroupSizeY = (View.ViewRect.Size().Y + GReflectionEnvironmentTileSizeY - 1) / GReflectionEnvironmentTileSizeY;
-			DispatchComputeShader(ComputeShader, GroupSizeX, GroupSizeY, 1);
+			DispatchComputeShader(RHICmdList, ComputeShader, GroupSizeX, GroupSizeY, 1);
 
 			ComputeShader->UnsetParameters();
 		}
@@ -641,7 +652,7 @@ void FDeferredShadingSceneRenderer::RenderImageBasedReflectionsSM5ForAllViews()
 	check(GSceneRenderTargets.GetSceneColor());
 }
 
-void FDeferredShadingSceneRenderer::RenderImageBasedReflectionsSM4ForAllViews(bool bReflectionEnv)
+void FDeferredShadingSceneRenderer::RenderImageBasedReflectionsSM4ForAllViews(FRHICommandList* RHICmdList, bool bReflectionEnv)
 {
 	const bool bSkyLight = Scene->SkyLight
 		&& Scene->SkyLight->ProcessedTexture
@@ -659,7 +670,7 @@ void FDeferredShadingSceneRenderer::RenderImageBasedReflectionsSM4ForAllViews(bo
 
 	static TArray<FReflectionCaptureSortData> SortData;
 
-	if(bReflectionEnv)
+	if (bReflectionEnv)
 	{
 		// shared for multiple views
 
@@ -680,9 +691,9 @@ void FDeferredShadingSceneRenderer::RenderImageBasedReflectionsSM4ForAllViews(bo
 			if (CurrentCapture->Shape == EReflectionCaptureShape::Plane)
 			{
 				NewSortEntry.BoxTransform = FMatrix(
-					FPlane(CurrentCapture->ReflectionPlane), 
-					FPlane(CurrentCapture->ReflectionXAxisAndYScale), 
-					FPlane(0, 0, 0, 0), 
+					FPlane(CurrentCapture->ReflectionPlane),
+					FPlane(CurrentCapture->ReflectionXAxisAndYScale),
+					FPlane(0, 0, 0, 0),
 					FPlane(0, 0, 0, 0));
 
 				NewSortEntry.BoxScales = FVector4(0);
@@ -710,14 +721,14 @@ void FDeferredShadingSceneRenderer::RenderImageBasedReflectionsSM4ForAllViews(bo
 		const bool bSSR = DoScreenSpaceReflections(View);
 
 		TRefCountPtr<IPooledRenderTarget> SSROutput = GSystemTextures.BlackDummy;
-		if( bSSR )
+		if (bSSR)
 		{
 			bRequiresApply = true;
 
-			ScreenSpaceReflections( View, SSROutput );
+			ScreenSpaceReflections(View, SSROutput);
 		}
 
-		if( bReflectionEnv )
+		if (bReflectionEnv)
 		{
 			bRequiresApply = true;
 
@@ -751,28 +762,28 @@ void FDeferredShadingSceneRenderer::RenderImageBasedReflectionsSM4ForAllViews(bo
 						static FGlobalBoundShaderState BoundShaderState;
 						SetGlobalBoundShaderState(BoundShaderState, GFilterVertexDeclaration.VertexDeclarationRHI, *VertexShader, *PixelShader);
 
-						PixelShader->SetParameters(View, ReflectionCapture);
+						PixelShader->SetParameters(RHICmdList, View, ReflectionCapture);
 					}
-					else 
+					else
 					{
 						TShaderMapRef<TStandardDeferredReflectionPS<false> > PixelShader(GetGlobalShaderMap());
 
 						static FGlobalBoundShaderState BoundShaderState;
 						SetGlobalBoundShaderState(BoundShaderState, GFilterVertexDeclaration.VertexDeclarationRHI, *VertexShader, *PixelShader);
 
-						PixelShader->SetParameters(View, ReflectionCapture);
+						PixelShader->SetParameters(RHICmdList, View, ReflectionCapture);
 					}
 
 					SetBoundingGeometryRasterizerAndDepthState(View, LightBounds);
-					VertexShader->SetSimpleLightParameters(View, LightBounds);
+					VertexShader->SetSimpleLightParameters(RHICmdList, View, LightBounds);
 					StencilingGeometry::DrawSphere();
 				}
 			}
 
-			GRenderTargetPool.VisualizeTexture.SetCheckPoint( LightAccumulation );
+			GRenderTargetPool.VisualizeTexture.SetCheckPoint(LightAccumulation);
 		}
 
-		if( bRequiresApply )
+		if (bRequiresApply)
 		{
 			// Apply reflections to screen
 			SCOPED_DRAW_EVENT(ReflectionApply, DEC_SCENE_ITEMS);
@@ -780,8 +791,8 @@ void FDeferredShadingSceneRenderer::RenderImageBasedReflectionsSM4ForAllViews(bo
 			GSceneRenderTargets.BeginRenderingSceneColor();
 
 			RHISetViewport(View.ViewRect.Min.X, View.ViewRect.Min.Y, 0.0f, View.ViewRect.Max.X, View.ViewRect.Max.Y, 1.0f);
-			RHISetRasterizerState(TStaticRasterizerState<FM_Solid,CM_None>::GetRHI());
-			RHISetDepthStencilState(TStaticDepthStencilState<false,CF_Always>::GetRHI());
+			RHISetRasterizerState(TStaticRasterizerState<FM_Solid, CM_None>::GetRHI());
+			RHISetDepthStencilState(TStaticDepthStencilState<false, CF_Always>::GetRHI());
 
 			if (GetReflectionEnvironmentCVar() == 2)
 			{
@@ -794,35 +805,35 @@ void FDeferredShadingSceneRenderer::RenderImageBasedReflectionsSM4ForAllViews(bo
 				RHISetBlendState(TStaticBlendState<CW_RGBA, BO_Add, BF_One, BF_One, BO_Add, BF_One, BF_One>::GetRHI());
 			}
 
-			TShaderMapRef< FPostProcessVS > VertexShader( GetGlobalShaderMap() );
+			TShaderMapRef< FPostProcessVS >		VertexShader(GetGlobalShaderMap());
 
 #define CASE(A,B,C) \
-			case ( (A << 2) | (B << 1) | C ): \
+			case ((A << 2) | (B << 1) | C) : \
 			{ \
-			TShaderMapRef< FReflectionApplyPS<A,B,C> > PixelShader(GetGlobalShaderMap()); \
+			TShaderMapRef< FReflectionApplyPS<A, B, C> > PixelShader(GetGlobalShaderMap()); \
 			static FGlobalBoundShaderState BoundShaderState; \
 			SetGlobalBoundShaderState(BoundShaderState, GFilterVertexDeclaration.VertexDeclarationRHI, *VertexShader, *PixelShader); \
-			PixelShader->SetParameters( View, LightAccumulation->GetRenderTargetItem().ShaderResourceTexture, SSROutput->GetRenderTargetItem().ShaderResourceTexture ); \
+			PixelShader->SetParameters(nullptr, View, LightAccumulation->GetRenderTargetItem().ShaderResourceTexture, SSROutput->GetRenderTargetItem().ShaderResourceTexture); \
 			}; \
 			break
 
-			switch( ((uint32)bSSR << 2) | ((uint32)bReflectionEnv << 1) | (uint32)bSkyLight )
+			switch (((uint32)bSSR << 2) | ((uint32)bReflectionEnv << 1) | (uint32)bSkyLight)
 			{
-				CASE(0,0,0);
-				CASE(0,0,1);
-				CASE(0,1,0);
-				CASE(0,1,1);
-				CASE(1,0,0);
-				CASE(1,0,1);
-				CASE(1,1,0);
-				CASE(1,1,1);
+				CASE(0, 0, 0);
+				CASE(0, 0, 1);
+				CASE(0, 1, 0);
+				CASE(0, 1, 1);
+				CASE(1, 0, 0);
+				CASE(1, 0, 1);
+				CASE(1, 1, 0);
+				CASE(1, 1, 1);
 			}
 #undef CASE
 
-			DrawRectangle( 
-				0, 0, 
+			DrawRectangle(
+				0, 0,
 				View.ViewRect.Width(), View.ViewRect.Height(),
-				View.ViewRect.Min.X, View.ViewRect.Min.Y, 
+				View.ViewRect.Min.X, View.ViewRect.Min.Y,
 				View.ViewRect.Width(), View.ViewRect.Height(),
 				FIntPoint(View.ViewRect.Width(), View.ViewRect.Height()),
 				GSceneRenderTargets.GetBufferSizeXY(),
@@ -831,7 +842,7 @@ void FDeferredShadingSceneRenderer::RenderImageBasedReflectionsSM4ForAllViews(bo
 	}
 }
 
-void FDeferredShadingSceneRenderer::RenderDeferredReflections()
+void FDeferredShadingSceneRenderer::RenderDeferredReflections(FRHICommandList* RHICmdList)
 {
 	if (IsSimpleDynamicLightingEnabled() || ViewFamily.EngineShowFlags.VisualizeLightCulling)
 	{
@@ -860,12 +871,12 @@ void FDeferredShadingSceneRenderer::RenderDeferredReflections()
 		if (bReflectionsWithCompute)
 		{
 			check(bReflectionEnv);
-			RenderImageBasedReflectionsSM5ForAllViews();
+			RenderImageBasedReflectionsSM5ForAllViews(RHICmdList);
 		}
 		else
 		{
 			// to test this code path run with -SM4
-			RenderImageBasedReflectionsSM4ForAllViews(bReflectionEnv);
+			RenderImageBasedReflectionsSM4ForAllViews(RHICmdList, bReflectionEnv);
 		}
 	}
 }
\ No newline at end of file
diff --git a/Engine/Source/Runtime/Renderer/Private/ReflectionEnvironmentCapture.cpp b/Engine/Source/Runtime/Renderer/Private/ReflectionEnvironmentCapture.cpp
--- a/Engine/Source/Runtime/Renderer/Private/ReflectionEnvironmentCapture.cpp
+++ b/Engine/Source/Runtime/Renderer/Private/ReflectionEnvironmentCapture.cpp
@@ -96,12 +96,13 @@ public:
 	}
 	FDownsamplePS() {}
 
-	void SetParameters(int32 CubeFaceValue, int32 SourceMipIndexValue, FSceneRenderTargetItem& SourceTextureValue)
+	void SetParameters(FRHICommandList* RHICmdList, int32 CubeFaceValue, int32 SourceMipIndexValue, FSceneRenderTargetItem& SourceTextureValue)
 	{
-		SetShaderValue(GetPixelShader(), CubeFace, CubeFaceValue);
-		SetShaderValue(GetPixelShader(), SourceMipIndex, SourceMipIndexValue);
+		SetShaderValue(RHICmdList, GetPixelShader(), CubeFace, CubeFaceValue);
+		SetShaderValue(RHICmdList, GetPixelShader(), SourceMipIndex, SourceMipIndexValue);
 
 		SetTextureParameter(
+			RHICmdList, 
 			GetPixelShader(), 
 			SourceTexture, 
 			SourceTextureSampler, 
@@ -157,13 +158,14 @@ public:
 	}
 	FCubeFilterPS() {}
 
-	void SetParameters(int32 CubeFaceValue, int32 SourceMipIndexValue, float RoughnessValue, FSceneRenderTargetItem& SourceTextureValue)
+	void SetParameters(FRHICommandList* RHICmdList, int32 CubeFaceValue, int32 SourceMipIndexValue, float RoughnessValue, FSceneRenderTargetItem& SourceTextureValue)
 	{
-		FDownsamplePS::SetParameters(CubeFaceValue, SourceMipIndexValue, SourceTextureValue);
+		FDownsamplePS::SetParameters(RHICmdList, CubeFaceValue, SourceMipIndexValue, SourceTextureValue);
 
-		SetShaderValue(GetPixelShader(), FilterRoughness, RoughnessValue);
+		SetShaderValue(RHICmdList, GetPixelShader(), FilterRoughness, RoughnessValue);
 
 		SetTextureParameter(
+			RHICmdList,
 			GetPixelShader(), 
 			AverageBrightnessTexture, 
 			AverageBrightnessSampler, 
@@ -240,8 +242,9 @@ public:
 	{
 	}
 
-	void SetParameters()
+	void SetParameters(FRHICommandList* RHICmdList)
 	{
+		check(!RHICmdList);
 		const int32 EffectiveTopMipSize = GReflectionCaptureSize;
 		const int32 NumMips = FMath::CeilLogTwo(EffectiveTopMipSize) + 1;
 		// Read from the smallest mip that was downsampled to
@@ -250,6 +253,7 @@ public:
 		if (Cubemap.IsValid())
 		{
 			SetTextureParameter(
+				RHICmdList,
 				GetPixelShader(), 
 				ReflectionEnvironmentColorTexture, 
 				ReflectionEnvironmentColorSampler, 
@@ -257,7 +261,7 @@ public:
 				Cubemap.ShaderResourceTexture);
 		}
 
-		SetShaderValue(GetPixelShader(), NumCaptureArrayMips, FMath::CeilLogTwo(GReflectionCaptureSize) + 1);
+		SetShaderValue(RHICmdList, GetPixelShader(), NumCaptureArrayMips, FMath::CeilLogTwo(GReflectionCaptureSize) + 1);
 	}
 
 	virtual bool Serialize(FArchive& Ar)
@@ -279,8 +283,10 @@ private:
 IMPLEMENT_SHADER_TYPE(,FComputeBrightnessPS,TEXT("ReflectionEnvironmentShaders"),TEXT("ComputeBrightnessMain"),SF_Pixel);
 
 /** Computes the average brightness of the given reflection capture and stores it in the scene. */
-void ComputeAverageBrightness()
+void ComputeAverageBrightness(FRHICommandList* RHICmdList)
 {
+	///@todo-rco:
+	check(!RHICmdList);
 	RHISetRenderTarget(GSceneRenderTargets.ReflectionBrightness->GetRenderTargetItem().TargetableTexture, NULL);
 	RHISetRasterizerState(TStaticRasterizerState<FM_Solid,CM_None>::GetRHI());
 	RHISetDepthStencilState(TStaticDepthStencilState<false,CF_Always>::GetRHI());
@@ -292,7 +298,7 @@ void ComputeAverageBrightness()
 	static FGlobalBoundShaderState BoundShaderState;
 	SetGlobalBoundShaderState(BoundShaderState, GFilterVertexDeclaration.VertexDeclarationRHI, *VertexShader, *PixelShader);
 
-	PixelShader->SetParameters();
+	PixelShader->SetParameters(RHICmdList);
 
 	DrawRectangle( 
 		0, 0, 
@@ -310,6 +316,9 @@ void FilterReflectionEnvironment(FSHVectorRGB3* OutIrradianceEnvironmentMap)
 	const int32 EffectiveTopMipSize = GReflectionCaptureSize;
 	const int32 NumMips = FMath::CeilLogTwo(EffectiveTopMipSize) + 1;
 
+	//@todo-rco: RHIPacketList
+	FRHICommandList* RHICmdList = nullptr;
+
 	int32 DiffuseConvolutionSourceMip = INDEX_NONE;
 	FSceneRenderTargetItem* DiffuseConvolutionSource = NULL;
 
@@ -331,40 +340,40 @@ void FilterReflectionEnvironment(FSHVectorRGB3* OutIrradianceEnvironmentMap)
 
 			if (GSupportsGSRenderTargetLayerSwitchingToMips)
 			{
-				RHISetRenderTarget(EffectiveRT.TargetableTexture, MipIndex, NULL);
+			RHISetRenderTarget(EffectiveRT.TargetableTexture, MipIndex, NULL);
 
-				const FIntRect ViewRect(0, 0, MipSize, MipSize);
-				RHISetViewport(0, 0, 0.0f, MipSize, MipSize, 1.0f);
-				RHISetRasterizerState(TStaticRasterizerState<FM_Solid,CM_None>::GetRHI());
-				RHISetDepthStencilState(TStaticDepthStencilState<false,CF_Always>::GetRHI());
-				RHISetBlendState(TStaticBlendState<>::GetRHI());
+			const FIntRect ViewRect(0, 0, MipSize, MipSize);
+			RHISetViewport(0, 0, 0.0f, MipSize, MipSize, 1.0f);
+			RHISetRasterizerState(TStaticRasterizerState<FM_Solid,CM_None>::GetRHI());
+			RHISetDepthStencilState(TStaticDepthStencilState<false,CF_Always>::GetRHI());
+			RHISetBlendState(TStaticBlendState<>::GetRHI());
 
-				TShaderMapRef<FScreenVSForGS> VertexShader(GetGlobalShaderMap());
-				TShaderMapRef<FDownsampleGS> GeometryShader(GetGlobalShaderMap());
+			TShaderMapRef<FScreenVSForGS> VertexShader(GetGlobalShaderMap());
+			TShaderMapRef<FDownsampleGS> GeometryShader(GetGlobalShaderMap());
 				TShaderMapRef<FDownsamplePS> PixelShader(GetGlobalShaderMap());
 
-				SetGlobalBoundShaderState(DownsampleBoundShaderState, GFilterVertexDeclaration.VertexDeclarationRHI, *VertexShader, *PixelShader, *GeometryShader);
+			SetGlobalBoundShaderState(DownsampleBoundShaderState, GFilterVertexDeclaration.VertexDeclarationRHI, *VertexShader, *PixelShader, *GeometryShader);
 
-				// Draw each face with a geometry shader that routes to the correct render target slice
-				for (int32 CubeFace = 0; CubeFace < CubeFace_MAX; CubeFace++)
-				{
-					SCOPED_DRAW_EVENT(DownsampleCubeFace, DEC_SCENE_ITEMS);
-					GeometryShader->SetParameters(CubeFace);
+			// Draw each face with a geometry shader that routes to the correct render target slice
+			for (int32 CubeFace = 0; CubeFace < CubeFace_MAX; CubeFace++)
+			{
+				SCOPED_DRAW_EVENT(DownsampleCubeFace, DEC_SCENE_ITEMS);
+				GeometryShader->SetParameters(RHICmdList, CubeFace);
 
-					PixelShader->SetParameters(CubeFace, SourceMipIndex, EffectiveSource);
+				PixelShader->SetParameters(RHICmdList, CubeFace, SourceMipIndex, EffectiveSource);
 
-					DrawRectangle( 
-						ViewRect.Min.X, ViewRect.Min.Y, 
-						ViewRect.Width(), ViewRect.Height(),
-						ViewRect.Min.X, ViewRect.Min.Y, 
-						ViewRect.Width(), ViewRect.Height(),
-						FIntPoint(ViewRect.Width(), ViewRect.Height()),
-						FIntPoint(MipSize, MipSize),
-						*VertexShader,
-						EDRF_UseTriangleOptimization);
+				DrawRectangle( 
+					ViewRect.Min.X, ViewRect.Min.Y, 
+					ViewRect.Width(), ViewRect.Height(),
+					ViewRect.Min.X, ViewRect.Min.Y, 
+					ViewRect.Width(), ViewRect.Height(),
+					FIntPoint(ViewRect.Width(), ViewRect.Height()),
+					FIntPoint(MipSize, MipSize),
+					*VertexShader,
+					EDRF_UseTriangleOptimization);
 
-					RHICopyToResolveTarget(EffectiveRT.TargetableTexture, EffectiveRT.ShaderResourceTexture, true, FResolveParams(FResolveRect(), (ECubeFace)CubeFace, MipIndex));
-				}
+				RHICopyToResolveTarget(EffectiveRT.TargetableTexture, EffectiveRT.ShaderResourceTexture, true, FResolveParams(FResolveRect(), (ECubeFace)CubeFace, MipIndex));
+			}
 			}
 			else
 			{
@@ -383,7 +392,7 @@ void FilterReflectionEnvironment(FSHVectorRGB3* OutIrradianceEnvironmentMap)
 
 					SetGlobalBoundShaderState(DownsampleBoundShaderState, GFilterVertexDeclaration.VertexDeclarationRHI, *VertexShader, *PixelShader);
 
-					PixelShader->SetParameters(CubeFace, SourceMipIndex, EffectiveSource);
+					PixelShader->SetParameters(RHICmdList, CubeFace, SourceMipIndex, EffectiveSource);
 
 					DrawRectangle( 
 						ViewRect.Min.X, ViewRect.Min.Y, 
@@ -413,7 +422,7 @@ void FilterReflectionEnvironment(FSHVectorRGB3* OutIrradianceEnvironmentMap)
 		ComputeDiffuseIrradiance(DiffuseConvolutionSource->ShaderResourceTexture, DiffuseConvolutionSourceMip, OutIrradianceEnvironmentMap);
 	}
 	
-	ComputeAverageBrightness();
+	ComputeAverageBrightness(RHICmdList);
 
 	{	
 		SCOPED_DRAW_EVENT(FilterCubeMap, DEC_SCENE_ITEMS);
@@ -431,17 +440,17 @@ void FilterReflectionEnvironment(FSHVectorRGB3* OutIrradianceEnvironmentMap)
 
 			if (GSupportsGSRenderTargetLayerSwitchingToMips)
 			{
-				RHISetRenderTarget(EffectiveRT.TargetableTexture, MipIndex, NULL);
+			RHISetRenderTarget(EffectiveRT.TargetableTexture, MipIndex, NULL);
 
-				const FIntRect ViewRect(0, 0, MipSize, MipSize);
-				RHISetViewport(0, 0, 0.0f, MipSize, MipSize, 1.0f);
-				RHISetRasterizerState(TStaticRasterizerState<FM_Solid,CM_None>::GetRHI());
-				RHISetDepthStencilState(TStaticDepthStencilState<false,CF_Always>::GetRHI());
-				RHISetBlendState(TStaticBlendState<>::GetRHI());
+			const FIntRect ViewRect(0, 0, MipSize, MipSize);
+			RHISetViewport(0, 0, 0.0f, MipSize, MipSize, 1.0f);
+			RHISetRasterizerState(TStaticRasterizerState<FM_Solid,CM_None>::GetRHI());
+			RHISetDepthStencilState(TStaticDepthStencilState<false,CF_Always>::GetRHI());
+			RHISetBlendState(TStaticBlendState<>::GetRHI());
+
+			TShaderMapRef<FScreenVSForGS> VertexShader(GetGlobalShaderMap());
+			TShaderMapRef<FDownsampleGS> GeometryShader(GetGlobalShaderMap());
 
-				TShaderMapRef<FScreenVSForGS> VertexShader(GetGlobalShaderMap());
-				TShaderMapRef<FDownsampleGS> GeometryShader(GetGlobalShaderMap());
-				
 				FCubeFilterPS* PixelShader;
 				if( bNormalize )
 				{
@@ -456,26 +465,26 @@ void FilterReflectionEnvironment(FSHVectorRGB3* OutIrradianceEnvironmentMap)
 					SetGlobalBoundShaderState(BoundShaderState, GFilterVertexDeclaration.VertexDeclarationRHI, *VertexShader, PixelShader, *GeometryShader);
 				}
 
-				for (int32 CubeFace = 0; CubeFace < CubeFace_MAX; CubeFace++)
-				{
-					SCOPED_DRAW_EVENT(FilterCubeFace, DEC_SCENE_ITEMS);
-					GeometryShader->SetParameters(CubeFace);
+			for (int32 CubeFace = 0; CubeFace < CubeFace_MAX; CubeFace++)
+			{
+				SCOPED_DRAW_EVENT(FilterCubeFace, DEC_SCENE_ITEMS);
+				GeometryShader->SetParameters(RHICmdList, CubeFace);
 
-					PixelShader->SetParameters(CubeFace, MipIndex, Roughness, EffectiveSource);
+				PixelShader->SetParameters(RHICmdList, CubeFace, MipIndex, Roughness, EffectiveSource);
 
-					DrawRectangle( 
-						ViewRect.Min.X, ViewRect.Min.Y, 
-						ViewRect.Width(), ViewRect.Height(),
-						ViewRect.Min.X, ViewRect.Min.Y, 
-						ViewRect.Width(), ViewRect.Height(),
-						FIntPoint(ViewRect.Width(), ViewRect.Height()),
-						FIntPoint(MipSize, MipSize),
-						*VertexShader,
-						EDRF_UseTriangleOptimization);
+				DrawRectangle( 
+					ViewRect.Min.X, ViewRect.Min.Y, 
+					ViewRect.Width(), ViewRect.Height(),
+					ViewRect.Min.X, ViewRect.Min.Y, 
+					ViewRect.Width(), ViewRect.Height(),
+					FIntPoint(ViewRect.Width(), ViewRect.Height()),
+					FIntPoint(MipSize, MipSize),
+					*VertexShader,
+					EDRF_UseTriangleOptimization);
 
-					RHICopyToResolveTarget(EffectiveRT.TargetableTexture, EffectiveRT.ShaderResourceTexture, true, FResolveParams(FResolveRect(), (ECubeFace)CubeFace, MipIndex));
-				}
+				RHICopyToResolveTarget(EffectiveRT.TargetableTexture, EffectiveRT.ShaderResourceTexture, true, FResolveParams(FResolveRect(), (ECubeFace)CubeFace, MipIndex));
 			}
+		}
 			else
 			{
 				for (int32 CubeFace = 0; CubeFace < CubeFace_MAX; CubeFace++)
@@ -506,7 +515,7 @@ void FilterReflectionEnvironment(FSHVectorRGB3* OutIrradianceEnvironmentMap)
 						SetGlobalBoundShaderState(BoundShaderState, GFilterVertexDeclaration.VertexDeclarationRHI, *VertexShader, PixelShader);
 					}
 
-					PixelShader->SetParameters(CubeFace, MipIndex, Roughness, EffectiveSource);
+					PixelShader->SetParameters(RHICmdList, CubeFace, MipIndex, Roughness, EffectiveSource);
 
 					DrawRectangle( 
 						ViewRect.Min.X, ViewRect.Min.Y, 
@@ -541,9 +550,9 @@ public:
 	{
 	}
 
-	void SetParameters(const FViewInfo& View)
+	void SetParameters(FRHICommandList* RHICmdList, const FViewInfo& View)
 	{
-		FGlobalShader::SetParameters(GetVertexShader(),View);
+		FGlobalShader::SetParameters(RHICmdList, GetVertexShader(),View);
 	}
 
 	virtual bool Serialize(FArchive& Ar)
@@ -573,9 +582,9 @@ public:
 	}
 	FRouteToCubeFaceGS() {}
 
-	void SetParameters(uint32 CubeFaceValue)
+	void SetParameters(FRHICommandList* RHICmdList, uint32 CubeFaceValue)
 	{
-		SetShaderValue(GetGeometryShader(), CubeFace, CubeFaceValue);
+		SetShaderValue(RHICmdList, GetGeometryShader(), CubeFace, CubeFaceValue);
 	}
 
 	virtual bool Serialize(FArchive& Ar)
@@ -602,6 +611,12 @@ public:
 		return IsFeatureLevelSupported(Platform, ERHIFeatureLevel::SM4);
 	}
 
+	static void ModifyCompilationEnvironment(EShaderPlatform Platform, FShaderCompilerEnvironment& OutEnvironment)
+	{
+		FGlobalShader::ModifyCompilationEnvironment(Platform,OutEnvironment);
+		FDeferredPixelShaderParameters::ModifyCompilationEnvironment(Platform,OutEnvironment);
+	}
+
 	FCopySceneColorToCubeFacePS(const ShaderMetaType::CompiledShaderInitializerType& Initializer):
 		FGlobalShader(Initializer)
 	{
@@ -612,14 +627,15 @@ public:
 	}
 	FCopySceneColorToCubeFacePS() {}
 
-	void SetParameters(const FViewInfo& View, bool bCapturingForSkyLight, bool bLowerHemisphereIsBlack)
+	void SetParameters(FRHICommandList* RHICmdList, const FViewInfo& View, bool bCapturingForSkyLight, bool bLowerHemisphereIsBlack)
 	{
 		const FPixelShaderRHIParamRef ShaderRHI = GetPixelShader();
 
-		FGlobalShader::SetParameters(ShaderRHI, View);
-		DeferredParameters.Set(ShaderRHI, View);
+		FGlobalShader::SetParameters(RHICmdList, ShaderRHI, View);
+		DeferredParameters.Set(RHICmdList, ShaderRHI, View);
 
 		SetTextureParameter(
+			RHICmdList, 
 			ShaderRHI, 
 			InTexture, 
 			InTextureSampler, 
@@ -646,7 +662,7 @@ public:
 			SkyLightParametersValue = FVector(2, 0, 0);
 		}
 
-		SetShaderValue(ShaderRHI, SkyLightCaptureParameters, SkyLightParametersValue);
+		SetShaderValue(RHICmdList, ShaderRHI, SkyLightCaptureParameters, SkyLightParametersValue);
 	}
 
 	virtual bool Serialize(FArchive& Ar)
@@ -691,19 +707,20 @@ public:
 	}
 	FCopyCubemapToCubeFacePS() {}
 
-	void SetParameters(const FTexture* SourceCubemap, uint32 CubeFaceValue, bool bIsSkyLight, bool bLowerHemisphereIsBlack)
+	void SetParameters(FRHICommandList* RHICmdList, const FTexture* SourceCubemap, uint32 CubeFaceValue, bool bIsSkyLight, bool bLowerHemisphereIsBlack)
 	{
 		const FPixelShaderRHIParamRef ShaderRHI = GetPixelShader();
 
-		SetShaderValue(ShaderRHI, CubeFace, CubeFaceValue);
+		SetShaderValue(RHICmdList, ShaderRHI, CubeFace, CubeFaceValue);
 
 		SetTextureParameter(
+			RHICmdList, 
 			ShaderRHI, 
 			SourceTexture, 
 			SourceTextureSampler, 
 			SourceCubemap);
 
-		SetShaderValue(ShaderRHI, SkyLightCaptureParameters, FVector(bIsSkyLight ? 1.0f : 0.0f, 0.0f, bLowerHemisphereIsBlack ? 1.0f : 0.0f));
+		SetShaderValue(RHICmdList, ShaderRHI, SkyLightCaptureParameters, FVector(bIsSkyLight ? 1.0f : 0.0f, 0.0f, bLowerHemisphereIsBlack ? 1.0f : 0.0f));
 	}
 
 	virtual bool Serialize(FArchive& Ar)
@@ -766,10 +783,11 @@ int32 FindOrAllocateCubemapIndex(FScene* Scene, const UReflectionCaptureComponen
 }
 
 /** Captures the scene for a reflection capture by rendering the scene multiple times and copying into a cubemap texture. */
-void CaptureSceneToScratchCubemap(FSceneRenderer* SceneRenderer, ECubeFace CubeFace, bool bCapturingForSkyLight, bool bLowerHemisphereIsBlack)
+void CaptureSceneToScratchCubemap(FRHICommandList* RHICmdList, FSceneRenderer* SceneRenderer, ECubeFace CubeFace, bool bCapturingForSkyLight, bool bLowerHemisphereIsBlack)
 {
 	FMemMark MemStackMark(FMemStack::Get());
-
+	//@todo-rco: RHIPacketList
+	check(!RHICmdList);
 	// update any resources that needed a deferred update
 	FDeferredUpdateResource::UpdateResources();
 	
@@ -789,41 +807,41 @@ void CaptureSceneToScratchCubemap(FSceneRenderer* SceneRenderer, ECubeFace CubeF
 		FSceneRenderTargetItem& EffectiveColorRT =  GSceneRenderTargets.ReflectionColorScratchCubemap[0]->GetRenderTargetItem();
 
 		{
+			check(!RHICmdList);
 			SCOPED_DRAW_EVENT(CubeMapCopyScene, DEC_SCENE_ITEMS);
-
 			if (GSupportsGSRenderTargetLayerSwitchingToMips)
 			{
-				// Copy the captured scene into the cubemap face
-				RHISetRenderTarget(EffectiveColorRT.TargetableTexture, NULL);
+			// Copy the captured scene into the cubemap face
+			RHISetRenderTarget(EffectiveColorRT.TargetableTexture, NULL);
 
-				const FIntRect ViewRect(0, 0, EffectiveSize, EffectiveSize);
-				RHISetViewport(0, 0, 0.0f, EffectiveSize, EffectiveSize, 1.0f);
-				RHISetRasterizerState(TStaticRasterizerState<FM_Solid,CM_None>::GetRHI());
-				RHISetDepthStencilState(TStaticDepthStencilState<false,CF_Always>::GetRHI());
-				RHISetBlendState(TStaticBlendState<>::GetRHI());
+			const FIntRect ViewRect(0, 0, EffectiveSize, EffectiveSize);
+			RHISetViewport(0, 0, 0.0f, EffectiveSize, EffectiveSize, 1.0f);
+			RHISetRasterizerState(TStaticRasterizerState<FM_Solid,CM_None>::GetRHI());
+			RHISetDepthStencilState(TStaticDepthStencilState<false,CF_Always>::GetRHI());
+			RHISetBlendState(TStaticBlendState<>::GetRHI());
 
-				TShaderMapRef<FCopyToCubeFaceVS> VertexShader(GetGlobalShaderMap());
-				TShaderMapRef<FRouteToCubeFaceGS> GeometryShader(GetGlobalShaderMap());
+			TShaderMapRef<FCopyToCubeFaceVS> VertexShader(GetGlobalShaderMap());
+			TShaderMapRef<FRouteToCubeFaceGS> GeometryShader(GetGlobalShaderMap());
 
-				TShaderMapRef<FCopySceneColorToCubeFacePS> PixelShader(GetGlobalShaderMap());
-				SetGlobalBoundShaderState(CopyColorCubemapBoundShaderState, GFilterVertexDeclaration.VertexDeclarationRHI, *VertexShader, *PixelShader, *GeometryShader);
+			TShaderMapRef<FCopySceneColorToCubeFacePS> PixelShader(GetGlobalShaderMap());
+			SetGlobalBoundShaderState(CopyColorCubemapBoundShaderState, GFilterVertexDeclaration.VertexDeclarationRHI, *VertexShader, *PixelShader, *GeometryShader);
 			
-				PixelShader->SetParameters(SceneRenderer->Views[0], bCapturingForSkyLight, bLowerHemisphereIsBlack);
-				VertexShader->SetParameters(SceneRenderer->Views[0]);
-				GeometryShader->SetParameters((uint32)CubeFace);
+				PixelShader->SetParameters(RHICmdList, SceneRenderer->Views[0], bCapturingForSkyLight, bLowerHemisphereIsBlack);
+				VertexShader->SetParameters(RHICmdList, SceneRenderer->Views[0]);
+				GeometryShader->SetParameters(RHICmdList, (uint32)CubeFace);
 
-				DrawRectangle( 
-					ViewRect.Min.X, ViewRect.Min.Y, 
-					ViewRect.Width(), ViewRect.Height(),
-					ViewRect.Min.X, ViewRect.Min.Y, 
-					ViewRect.Width() * GSupersampleCaptureFactor, ViewRect.Height() * GSupersampleCaptureFactor,
-					FIntPoint(ViewRect.Width(), ViewRect.Height()),
-					GSceneRenderTargets.GetBufferSizeXY(),
-					*VertexShader,
-					EDRF_UseTriangleOptimization);
+			DrawRectangle( 
+				ViewRect.Min.X, ViewRect.Min.Y, 
+				ViewRect.Width(), ViewRect.Height(),
+				ViewRect.Min.X, ViewRect.Min.Y, 
+				ViewRect.Width() * GSupersampleCaptureFactor, ViewRect.Height() * GSupersampleCaptureFactor,
+				FIntPoint(ViewRect.Width(), ViewRect.Height()),
+				GSceneRenderTargets.GetBufferSizeXY(),
+				*VertexShader,
+				EDRF_UseTriangleOptimization);
 
-				RHICopyToResolveTarget(EffectiveColorRT.TargetableTexture, EffectiveColorRT.ShaderResourceTexture, true, FResolveParams(FResolveRect()));
-			}
+			RHICopyToResolveTarget(EffectiveColorRT.TargetableTexture, EffectiveColorRT.ShaderResourceTexture, true, FResolveParams(FResolveRect()));
+		}
 			else
 			{
 				// Copy the captured scene into the cubemap face
@@ -839,8 +857,8 @@ void CaptureSceneToScratchCubemap(FSceneRenderer* SceneRenderer, ECubeFace CubeF
 				TShaderMapRef<FCopySceneColorToCubeFacePS> PixelShader(GetGlobalShaderMap());
 				SetGlobalBoundShaderState(CopyColorCubemapBoundShaderState, GFilterVertexDeclaration.VertexDeclarationRHI, *VertexShader, *PixelShader);
 
-				PixelShader->SetParameters(SceneRenderer->Views[0], bCapturingForSkyLight, bLowerHemisphereIsBlack);
-				VertexShader->SetParameters(SceneRenderer->Views[0]);
+				PixelShader->SetParameters(RHICmdList, SceneRenderer->Views[0], bCapturingForSkyLight, bLowerHemisphereIsBlack);
+				VertexShader->SetParameters(RHICmdList, SceneRenderer->Views[0]);
 
 				DrawRectangle( 
 					ViewRect.Min.X, ViewRect.Min.Y, 
@@ -859,7 +877,7 @@ void CaptureSceneToScratchCubemap(FSceneRenderer* SceneRenderer, ECubeFace CubeF
 	delete SceneRenderer;
 }
 
-void CopyCubemapToScratchCubemap(UTextureCube* SourceCubemap, bool bIsSkyLight, bool bLowerHemisphereIsBlack)
+void CopyCubemapToScratchCubemap(FRHICommandList* RHICmdList, UTextureCube* SourceCubemap, bool bIsSkyLight, bool bLowerHemisphereIsBlack)
 {
 	check(SourceCubemap);
 	const int32 EffectiveSize = GReflectionCaptureSize;
@@ -867,6 +885,8 @@ void CopyCubemapToScratchCubemap(UTextureCube* SourceCubemap, bool bIsSkyLight,
 
 	if (GSupportsGSRenderTargetLayerSwitchingToMips)
 	{
+		check(!RHICmdList);
+
 		// Copy the captured scene into the cubemap face
 		RHISetRenderTarget(EffectiveColorRT.TargetableTexture, NULL);
 
@@ -885,8 +905,8 @@ void CopyCubemapToScratchCubemap(UTextureCube* SourceCubemap, bool bIsSkyLight,
 		for (uint32 CubeFace = 0; CubeFace < CubeFace_MAX; CubeFace++)
 		{
 			SetGlobalBoundShaderState(CopyFromCubemapToCubemapBoundShaderState, GFilterVertexDeclaration.VertexDeclarationRHI, *VertexShader, *PixelShader, *GeometryShader);
-			GeometryShader->SetParameters(CubeFace);
-			PixelShader->SetParameters(SourceCubemapResource, CubeFace, bIsSkyLight, bLowerHemisphereIsBlack);
+			GeometryShader->SetParameters(RHICmdList, CubeFace);
+			PixelShader->SetParameters(RHICmdList, SourceCubemapResource, CubeFace, bIsSkyLight, bLowerHemisphereIsBlack);
 
 			DrawRectangle( 
 				ViewRect.Min.X, ViewRect.Min.Y, 
@@ -919,7 +939,7 @@ void CopyCubemapToScratchCubemap(UTextureCube* SourceCubemap, bool bIsSkyLight,
 			TShaderMapRef<FCopyCubemapToCubeFacePS> PixelShader(GetGlobalShaderMap());
 
 			SetGlobalBoundShaderState(CopyFromCubemapToCubemapBoundShaderState, GFilterVertexDeclaration.VertexDeclarationRHI, *VertexShader, *PixelShader);
-			PixelShader->SetParameters(SourceCubemapResource, CubeFace, bIsSkyLight, bLowerHemisphereIsBlack);
+			PixelShader->SetParameters(RHICmdList, SourceCubemapResource, CubeFace, bIsSkyLight, bLowerHemisphereIsBlack);
 
 			DrawRectangle( 
 				ViewRect.Min.X, ViewRect.Min.Y, 
@@ -1393,7 +1413,8 @@ void CaptureSceneIntoScratchCubemap(FScene* Scene, FVector CapturePosition, bool
 			bool, bCapturingForSkyLight, bCapturingForSkyLight,
 			bool, bLowerHemisphereIsBlack, bLowerHemisphereIsBlack,
 		{
-			CaptureSceneToScratchCubemap(SceneRenderer, CubeFace, bCapturingForSkyLight, bLowerHemisphereIsBlack);
+			//@todo-rco: RHIPacketList
+			CaptureSceneToScratchCubemap(nullptr, SceneRenderer, CubeFace, bCapturingForSkyLight, bLowerHemisphereIsBlack);
 
 			RHIEndFrame();
 		});
@@ -1550,7 +1571,8 @@ void FScene::UpdateSkyCaptureContents(const USkyLightComponent* CaptureComponent
 				UTextureCube*, SourceTexture, CaptureComponent->Cubemap,
 				bool, bLowerHemisphereIsBlack, CaptureComponent->bLowerHemisphereIsBlack,
 			{
-				CopyCubemapToScratchCubemap(SourceTexture, true, bLowerHemisphereIsBlack);
+				//@todo-rco: RHIPacketList
+				CopyCubemapToScratchCubemap(nullptr, SourceTexture, true, bLowerHemisphereIsBlack);
 			});
 		}
 		else
@@ -1572,13 +1594,13 @@ void FScene::UpdateSkyCaptureContents(const USkyLightComponent* CaptureComponent
 		// Optionally copy the filtered mip chain to the output texture
 		if (OutProcessedTexture)
 		{
-			ENQUEUE_UNIQUE_RENDER_COMMAND_TWOPARAMETER( 
-				CopyCommand,
-				FScene*, Scene, this,
+		ENQUEUE_UNIQUE_RENDER_COMMAND_TWOPARAMETER( 
+			CopyCommand,
+			FScene*, Scene, this,
 				FTexture*, ProcessedTexture, OutProcessedTexture,
-			{
+		{
 				CopyToSkyTexture(Scene, ProcessedTexture);
-			});
+		});
 		}
 	}
 }
diff --git a/Engine/Source/Runtime/Renderer/Private/ReflectionEnvironmentCapture.h b/Engine/Source/Runtime/Renderer/Private/ReflectionEnvironmentCapture.h
--- a/Engine/Source/Runtime/Renderer/Private/ReflectionEnvironmentCapture.h
+++ b/Engine/Source/Runtime/Renderer/Private/ReflectionEnvironmentCapture.h
@@ -27,9 +27,9 @@ public:
 	}
 	FDownsampleGS() {}
 
-	void SetParameters(int32 CubeFaceValue)
+	void SetParameters(FRHICommandList* RHICmdList, int32 CubeFaceValue)
 	{
-		SetShaderValue(GetGeometryShader(), CubeFace, CubeFaceValue);
+		SetShaderValue(RHICmdList, GetGeometryShader(), CubeFace, CubeFaceValue);
 	}
 
 	virtual bool Serialize(FArchive& Ar)
diff --git a/Engine/Source/Runtime/Renderer/Private/ReflectionEnvironmentDiffuseIrradiance.cpp b/Engine/Source/Runtime/Renderer/Private/ReflectionEnvironmentDiffuseIrradiance.cpp
--- a/Engine/Source/Runtime/Renderer/Private/ReflectionEnvironmentDiffuseIrradiance.cpp
+++ b/Engine/Source/Runtime/Renderer/Private/ReflectionEnvironmentDiffuseIrradiance.cpp
@@ -55,12 +55,13 @@ public:
 	}
 	FCopyDiffuseIrradiancePS() {}
 
-	void SetParameters(int32 CubeFaceValue, int32 SourceMipIndexValue, int32 CoefficientIndex, int32 FaceResolution, FTextureRHIRef& SourceTextureValue)
+	void SetParameters(FRHICommandList* RHICmdList, int32 CubeFaceValue, int32 SourceMipIndexValue, int32 CoefficientIndex, int32 FaceResolution, FTextureRHIRef& SourceTextureValue)
 	{
-		SetShaderValue(GetPixelShader(), CubeFace, CubeFaceValue);
-		SetShaderValue(GetPixelShader(), SourceMipIndex, SourceMipIndexValue);
+		SetShaderValue(RHICmdList, GetPixelShader(), CubeFace, CubeFaceValue);
+		SetShaderValue(RHICmdList, GetPixelShader(), SourceMipIndex, SourceMipIndexValue);
 
 		SetTextureParameter(
+			RHICmdList, 
 			GetPixelShader(), 
 			SourceTexture, 
 			SourceTextureSampler, 
@@ -70,11 +71,11 @@ public:
 		const FVector4 Mask0(CoefficientIndex == 0, CoefficientIndex == 1, CoefficientIndex == 2, CoefficientIndex == 3);
 		const FVector4 Mask1(CoefficientIndex == 4, CoefficientIndex == 5, CoefficientIndex == 6, CoefficientIndex == 7);
 		const float Mask2 = CoefficientIndex == 8;
-		SetShaderValue(GetPixelShader(), CoefficientMask0, Mask0);
-		SetShaderValue(GetPixelShader(), CoefficientMask1, Mask1);
-		SetShaderValue(GetPixelShader(), CoefficientMask2, Mask2);
+		SetShaderValue(RHICmdList, GetPixelShader(), CoefficientMask0, Mask0);
+		SetShaderValue(RHICmdList, GetPixelShader(), CoefficientMask1, Mask1);
+		SetShaderValue(RHICmdList, GetPixelShader(), CoefficientMask2, Mask2);
 
-		SetShaderValue(GetPixelShader(), NumSamples, FaceResolution * FaceResolution * 6);
+		SetShaderValue(RHICmdList, GetPixelShader(), NumSamples, FaceResolution * FaceResolution * 6);
 	}
 
 	virtual bool Serialize(FArchive& Ar)
@@ -130,12 +131,19 @@ public:
 	}
 	FAccumulateDiffuseIrradiancePS() {}
 
-	void SetParameters(int32 CubeFaceValue, int32 NumMips, int32 SourceMipIndexValue, int32 CoefficientIndex, FTextureRHIRef& SourceTextureValue)
+	static void ModifyCompilationEnvironment(EShaderPlatform Platform, FShaderCompilerEnvironment& OutEnvironment)
 	{
-		SetShaderValue(GetPixelShader(), CubeFace, CubeFaceValue);
-		SetShaderValue(GetPixelShader(), SourceMipIndex, SourceMipIndexValue);
+		FGlobalShader::ModifyCompilationEnvironment(Platform,OutEnvironment);
+		FDeferredPixelShaderParameters::ModifyCompilationEnvironment(Platform,OutEnvironment);
+	}
+
+	void SetParameters(FRHICommandList* RHICmdList, int32 CubeFaceValue, int32 NumMips, int32 SourceMipIndexValue, int32 CoefficientIndex, FTextureRHIRef& SourceTextureValue)
+	{
+		SetShaderValue(RHICmdList, GetPixelShader(), CubeFace, CubeFaceValue);
+		SetShaderValue(RHICmdList, GetPixelShader(), SourceMipIndex, SourceMipIndexValue);
 
 		SetTextureParameter(
+			RHICmdList, 
 			GetPixelShader(), 
 			SourceTexture, 
 			SourceTextureSampler, 
@@ -146,8 +154,8 @@ public:
 		const float HalfSourceTexelSize = .5f / MipSize;
 		const FVector4 Sample01Value(-HalfSourceTexelSize, -HalfSourceTexelSize, HalfSourceTexelSize, -HalfSourceTexelSize);
 		const FVector4 Sample23Value(-HalfSourceTexelSize, HalfSourceTexelSize, HalfSourceTexelSize, HalfSourceTexelSize);
-		SetShaderValue(GetPixelShader(), Sample01, Sample01Value);
-		SetShaderValue(GetPixelShader(), Sample23, Sample23Value);
+		SetShaderValue(RHICmdList, GetPixelShader(), Sample01, Sample01Value);
+		SetShaderValue(RHICmdList, GetPixelShader(), Sample23, Sample23Value);
 	}
 
 	virtual bool Serialize(FArchive& Ar)
@@ -196,11 +204,12 @@ public:
 	}
 	FAccumulateCubeFacesPS() {}
 
-	void SetParameters(int32 SourceMipIndexValue, FTextureRHIRef& SourceTextureValue)
+	void SetParameters(FRHICommandList* RHICmdList, int32 SourceMipIndexValue, FTextureRHIRef& SourceTextureValue)
 	{
-		SetShaderValue(GetPixelShader(), SourceMipIndex, SourceMipIndexValue);
+		SetShaderValue(RHICmdList, GetPixelShader(), SourceMipIndex, SourceMipIndexValue);
 
 		SetTextureParameter(
+			RHICmdList, 
 			GetPixelShader(), 
 			SourceTexture, 
 			SourceTextureSampler, 
@@ -229,6 +238,8 @@ FGlobalBoundShaderState AccumulateCubeFacesBoundShaderState;
 
 void ComputeDiffuseIrradiance(FTextureRHIRef LightingSource, int32 LightingSourceMipIndex, FSHVectorRGB3* OutIrradianceEnvironmentMap)
 {
+	//@todo-rco: RHIPacketList
+	FRHICommandList* RHICmdList = nullptr;
 	for (int32 CoefficientIndex = 0; CoefficientIndex < FSHVector3::MaxSHBasis; CoefficientIndex++)
 	{
 		// Copy the starting mip from the lighting texture, apply texel area weighting and appropriate SH coefficient
@@ -238,39 +249,39 @@ void ComputeDiffuseIrradiance(FTextureRHIRef LightingSource, int32 LightingSourc
 			FSceneRenderTargetItem& EffectiveRT = GetEffectiveDiffuseIrradianceRenderTarget(MipIndex);
 			if (GSupportsGSRenderTargetLayerSwitchingToMips)
 			{
-				RHISetRenderTarget(EffectiveRT.TargetableTexture, 0, NULL);
+			RHISetRenderTarget(EffectiveRT.TargetableTexture, 0, NULL);
 
-				const FIntRect ViewRect(0, 0, MipSize, MipSize);
-				RHISetViewport(0, 0, 0.0f, MipSize, MipSize, 1.0f);
-				RHISetRasterizerState(TStaticRasterizerState<FM_Solid,CM_None>::GetRHI());
-				RHISetDepthStencilState(TStaticDepthStencilState<false,CF_Always>::GetRHI());
-				RHISetBlendState(TStaticBlendState<>::GetRHI());
+			const FIntRect ViewRect(0, 0, MipSize, MipSize);
+			RHISetViewport(0, 0, 0.0f, MipSize, MipSize, 1.0f);
+			RHISetRasterizerState(TStaticRasterizerState<FM_Solid,CM_None>::GetRHI());
+			RHISetDepthStencilState(TStaticDepthStencilState<false,CF_Always>::GetRHI());
+			RHISetBlendState(TStaticBlendState<>::GetRHI());
 
-				TShaderMapRef<FScreenVSForGS> VertexShader(GetGlobalShaderMap());
-				TShaderMapRef<FDownsampleGS> GeometryShader(GetGlobalShaderMap());
-				TShaderMapRef<FCopyDiffuseIrradiancePS> PixelShader(GetGlobalShaderMap());
+			TShaderMapRef<FScreenVSForGS> VertexShader(GetGlobalShaderMap());
+			TShaderMapRef<FDownsampleGS> GeometryShader(GetGlobalShaderMap());
+			TShaderMapRef<FCopyDiffuseIrradiancePS> PixelShader(GetGlobalShaderMap());
 
-				SetGlobalBoundShaderState(CopyDiffuseIrradianceShaderState, GFilterVertexDeclaration.VertexDeclarationRHI, *VertexShader, *PixelShader, *GeometryShader);
+			SetGlobalBoundShaderState(CopyDiffuseIrradianceShaderState, GFilterVertexDeclaration.VertexDeclarationRHI, *VertexShader, *PixelShader, *GeometryShader);
 
-				// Draw each face with a geometry shader that routes to the correct render target slice
-				for (int32 CubeFace = 0; CubeFace < CubeFace_MAX; CubeFace++)
-				{
-					GeometryShader->SetParameters(CubeFace);
+			// Draw each face with a geometry shader that routes to the correct render target slice
+			for (int32 CubeFace = 0; CubeFace < CubeFace_MAX; CubeFace++)
+			{
+				GeometryShader->SetParameters(RHICmdList, CubeFace);
 
-					PixelShader->SetParameters(CubeFace, LightingSourceMipIndex, CoefficientIndex, MipSize, LightingSource);
+				PixelShader->SetParameters(RHICmdList, CubeFace, LightingSourceMipIndex, CoefficientIndex, MipSize, LightingSource);
 
-					DrawRectangle( 
-						ViewRect.Min.X, ViewRect.Min.Y, 
-						ViewRect.Width(), ViewRect.Height(),
-						ViewRect.Min.X, ViewRect.Min.Y, 
-						ViewRect.Width(), ViewRect.Height(),
-						FIntPoint(ViewRect.Width(), ViewRect.Height()),
-						FIntPoint(MipSize, MipSize),
-						*VertexShader);
+				DrawRectangle( 
+					ViewRect.Min.X, ViewRect.Min.Y, 
+					ViewRect.Width(), ViewRect.Height(),
+					ViewRect.Min.X, ViewRect.Min.Y, 
+					ViewRect.Width(), ViewRect.Height(),
+					FIntPoint(ViewRect.Width(), ViewRect.Height()),
+					FIntPoint(MipSize, MipSize),
+					*VertexShader);
 
-					RHICopyToResolveTarget(EffectiveRT.TargetableTexture, EffectiveRT.ShaderResourceTexture, true, FResolveParams(FResolveRect(), (ECubeFace)CubeFace, MipIndex));
-				}
+				RHICopyToResolveTarget(EffectiveRT.TargetableTexture, EffectiveRT.ShaderResourceTexture, true, FResolveParams(FResolveRect(), (ECubeFace)CubeFace, MipIndex));
 			}
+		}
 			else
 			{
 				for (int32 CubeFace = 0; CubeFace < CubeFace_MAX; CubeFace++)
@@ -288,7 +299,7 @@ void ComputeDiffuseIrradiance(FTextureRHIRef LightingSource, int32 LightingSourc
 					
 					SetGlobalBoundShaderState(CopyDiffuseIrradianceShaderState, GFilterVertexDeclaration.VertexDeclarationRHI, *VertexShader, *PixelShader);
 					
-					PixelShader->SetParameters(CubeFace, LightingSourceMipIndex, CoefficientIndex, MipSize, LightingSource);
+					PixelShader->SetParameters(RHICmdList, CubeFace, LightingSourceMipIndex, CoefficientIndex, MipSize, LightingSource);
 					
 					DrawRectangle(
 						ViewRect.Min.X, ViewRect.Min.Y,
@@ -316,42 +327,42 @@ void ComputeDiffuseIrradiance(FTextureRHIRef LightingSource, int32 LightingSourc
 				FSceneRenderTargetItem& EffectiveRT = GetEffectiveDiffuseIrradianceRenderTarget(MipIndex);
 				FSceneRenderTargetItem& EffectiveSource = GetEffectiveDiffuseIrradianceSourceTexture(MipIndex);
 				check(EffectiveRT.TargetableTexture != EffectiveSource.ShaderResourceTexture);
-				
+
 				if (GSupportsGSRenderTargetLayerSwitchingToMips)
 				{
-					RHISetRenderTarget(EffectiveRT.TargetableTexture, MipIndex, NULL);
-					
-					const FIntRect ViewRect(0, 0, MipSize, MipSize);
-					RHISetViewport(0, 0, 0.0f, MipSize, MipSize, 1.0f);
-					RHISetRasterizerState(TStaticRasterizerState<FM_Solid,CM_None>::GetRHI());
-					RHISetDepthStencilState(TStaticDepthStencilState<false,CF_Always>::GetRHI());
-					RHISetBlendState(TStaticBlendState<>::GetRHI());
-					
-					TShaderMapRef<FScreenVSForGS> VertexShader(GetGlobalShaderMap());
-					TShaderMapRef<FDownsampleGS> GeometryShader(GetGlobalShaderMap());
-					TShaderMapRef<FAccumulateDiffuseIrradiancePS> PixelShader(GetGlobalShaderMap());
-					
-					SetGlobalBoundShaderState(DiffuseIrradianceAccumulateShaderState, GFilterVertexDeclaration.VertexDeclarationRHI, *VertexShader, *PixelShader, *GeometryShader);
-					
-					// Draw each face with a geometry shader that routes to the correct render target slice
-					for (int32 CubeFace = 0; CubeFace < CubeFace_MAX; CubeFace++)
-					{
-						GeometryShader->SetParameters(CubeFace);
-						
-						PixelShader->SetParameters(CubeFace, NumMips, SourceMipIndex, CoefficientIndex, EffectiveSource.ShaderResourceTexture);
-						
-						DrawRectangle(
-							ViewRect.Min.X, ViewRect.Min.Y,
-							ViewRect.Width(), ViewRect.Height(),
-							ViewRect.Min.X, ViewRect.Min.Y,
-							ViewRect.Width(), ViewRect.Height(),
-							FIntPoint(ViewRect.Width(), ViewRect.Height()),
-							FIntPoint(MipSize, MipSize),
-							*VertexShader);
-						
-						RHICopyToResolveTarget(EffectiveRT.TargetableTexture, EffectiveRT.ShaderResourceTexture, true, FResolveParams(FResolveRect(), (ECubeFace)CubeFace, MipIndex));
-					}
+				RHISetRenderTarget(EffectiveRT.TargetableTexture, MipIndex, NULL);
+
+				const FIntRect ViewRect(0, 0, MipSize, MipSize);
+				RHISetViewport(0, 0, 0.0f, MipSize, MipSize, 1.0f);
+				RHISetRasterizerState(TStaticRasterizerState<FM_Solid,CM_None>::GetRHI());
+				RHISetDepthStencilState(TStaticDepthStencilState<false,CF_Always>::GetRHI());
+				RHISetBlendState(TStaticBlendState<>::GetRHI());
+
+				TShaderMapRef<FScreenVSForGS> VertexShader(GetGlobalShaderMap());
+				TShaderMapRef<FDownsampleGS> GeometryShader(GetGlobalShaderMap());
+				TShaderMapRef<FAccumulateDiffuseIrradiancePS> PixelShader(GetGlobalShaderMap());
+
+				SetGlobalBoundShaderState(DiffuseIrradianceAccumulateShaderState, GFilterVertexDeclaration.VertexDeclarationRHI, *VertexShader, *PixelShader, *GeometryShader);
+
+				// Draw each face with a geometry shader that routes to the correct render target slice
+				for (int32 CubeFace = 0; CubeFace < CubeFace_MAX; CubeFace++)
+				{
+					GeometryShader->SetParameters(RHICmdList, CubeFace);
+
+					PixelShader->SetParameters(RHICmdList, CubeFace, NumMips, SourceMipIndex, CoefficientIndex, EffectiveSource.ShaderResourceTexture);
+
+					DrawRectangle( 
+						ViewRect.Min.X, ViewRect.Min.Y, 
+						ViewRect.Width(), ViewRect.Height(),
+						ViewRect.Min.X, ViewRect.Min.Y, 
+						ViewRect.Width(), ViewRect.Height(),
+						FIntPoint(ViewRect.Width(), ViewRect.Height()),
+						FIntPoint(MipSize, MipSize),
+						*VertexShader);
+
+					RHICopyToResolveTarget(EffectiveRT.TargetableTexture, EffectiveRT.ShaderResourceTexture, true, FResolveParams(FResolveRect(), (ECubeFace)CubeFace, MipIndex));
 				}
+			}
 				else
 				{
 					for (int32 CubeFace = 0; CubeFace < CubeFace_MAX; CubeFace++)
@@ -369,7 +380,7 @@ void ComputeDiffuseIrradiance(FTextureRHIRef LightingSource, int32 LightingSourc
 						
 						SetGlobalBoundShaderState(DiffuseIrradianceAccumulateShaderState, GFilterVertexDeclaration.VertexDeclarationRHI, *VertexShader, *PixelShader);
 						
-						PixelShader->SetParameters(CubeFace, NumMips, SourceMipIndex, CoefficientIndex, EffectiveSource.ShaderResourceTexture);
+						PixelShader->SetParameters(RHICmdList, CubeFace, NumMips, SourceMipIndex, CoefficientIndex, EffectiveSource.ShaderResourceTexture);
 						
 						DrawRectangle(
 							ViewRect.Min.X, ViewRect.Min.Y,
@@ -404,7 +415,7 @@ void ComputeDiffuseIrradiance(FTextureRHIRef LightingSource, int32 LightingSourc
 			const int32 SourceMipIndex = NumMips - 1;
 			const int32 MipSize = 1;
 			FSceneRenderTargetItem& EffectiveSource = GetEffectiveDiffuseIrradianceRenderTarget(SourceMipIndex);
-			PixelShader->SetParameters(SourceMipIndex, EffectiveSource.ShaderResourceTexture);
+			PixelShader->SetParameters(RHICmdList, SourceMipIndex, EffectiveSource.ShaderResourceTexture);
 
 			DrawRectangle( 
 				ViewRect.Min.X, ViewRect.Min.Y, 
diff --git a/Engine/Source/Runtime/Renderer/Private/Renderer.cpp b/Engine/Source/Runtime/Renderer/Private/Renderer.cpp
--- a/Engine/Source/Runtime/Renderer/Private/Renderer.cpp
+++ b/Engine/Source/Runtime/Renderer/Private/Renderer.cpp
@@ -40,6 +40,8 @@ void FRendererModule::DrawTileMesh(const FSceneView& SceneView, const FMeshBatch
 	View.InitRHIResources();
 
 	const auto FeatureLevel = View.GetFeatureLevel();
+	//@todo-rco: RHIPacketList
+	FRHICommandList* RHICmdList = nullptr;
 
 	const FMaterial* Material = Mesh.MaterialRenderProxy->GetMaterial(FeatureLevel);
 
@@ -53,11 +55,11 @@ void FRendererModule::DrawTileMesh(const FSceneView& SceneView, const FMeshBatch
 		{
 			if (FeatureLevel >= ERHIFeatureLevel::SM3)
 			{
-				FTranslucencyDrawingPolicyFactory::DrawDynamicMesh(View, FTranslucencyDrawingPolicyFactory::ContextType(), Mesh, false, false, NULL, HitProxyId);
+				FTranslucencyDrawingPolicyFactory::DrawDynamicMesh(RHICmdList, View, FTranslucencyDrawingPolicyFactory::ContextType(), Mesh, false, false, NULL, HitProxyId);
 			}
 			else
 			{
-				FTranslucencyForwardShadingDrawingPolicyFactory::DrawDynamicMesh(View, FTranslucencyForwardShadingDrawingPolicyFactory::ContextType(), Mesh, false, false, NULL, HitProxyId);
+				FTranslucencyForwardShadingDrawingPolicyFactory::DrawDynamicMesh(RHICmdList, View, FTranslucencyForwardShadingDrawingPolicyFactory::ContextType(), Mesh, false, false, NULL, HitProxyId);
 			}
 		}
 		// handle opaque materials
@@ -69,17 +71,17 @@ void FRendererModule::DrawTileMesh(const FSceneView& SceneView, const FMeshBatch
 			// draw the mesh
 			if (bIsHitTesting)
 			{
-				FHitProxyDrawingPolicyFactory::DrawDynamicMesh(View, FHitProxyDrawingPolicyFactory::ContextType(), Mesh, false, false, NULL, HitProxyId);
+				FHitProxyDrawingPolicyFactory::DrawDynamicMesh(RHICmdList, View, FHitProxyDrawingPolicyFactory::ContextType(), Mesh, false, false, NULL, HitProxyId);
 			}
 			else
 			{
 				if (FeatureLevel >= ERHIFeatureLevel::SM3)
 				{
-					FBasePassOpaqueDrawingPolicyFactory::DrawDynamicMesh(View, FBasePassOpaqueDrawingPolicyFactory::ContextType(false, ESceneRenderTargetsMode::SetTextures), Mesh, false, false, NULL, HitProxyId);
+					FBasePassOpaqueDrawingPolicyFactory::DrawDynamicMesh(RHICmdList, View, FBasePassOpaqueDrawingPolicyFactory::ContextType(false, ESceneRenderTargetsMode::SetTextures), Mesh, false, false, NULL, HitProxyId);
 				}
 				else
 				{
-					FBasePassForwardOpaqueDrawingPolicyFactory::DrawDynamicMesh(View, FBasePassForwardOpaqueDrawingPolicyFactory::ContextType(ESceneRenderTargetsMode::SetTextures), Mesh, false, false, NULL, HitProxyId);
+					FBasePassForwardOpaqueDrawingPolicyFactory::DrawDynamicMesh(RHICmdList, View, FBasePassForwardOpaqueDrawingPolicyFactory::ContextType(ESceneRenderTargetsMode::SetTextures), Mesh, false, false, NULL, HitProxyId);
 				}
 			}
 		}	
diff --git a/Engine/Source/Runtime/Renderer/Private/RendererPrivate.h b/Engine/Source/Runtime/Renderer/Private/RendererPrivate.h
--- a/Engine/Source/Runtime/Renderer/Private/RendererPrivate.h
+++ b/Engine/Source/Runtime/Renderer/Private/RendererPrivate.h
@@ -10,6 +10,7 @@
 #include "Engine.h"
 #include "ShaderCore.h"
 #include "RHI.h"
+#include "RHIStaticStates.h"
 #include "ScenePrivate.h"
 
 DECLARE_LOG_CATEGORY_EXTERN(LogRenderer, Log, All);
diff --git a/Engine/Source/Runtime/Renderer/Private/SceneCaptureRendering.cpp b/Engine/Source/Runtime/Renderer/Private/SceneCaptureRendering.cpp
--- a/Engine/Source/Runtime/Renderer/Private/SceneCaptureRendering.cpp
+++ b/Engine/Source/Runtime/Renderer/Private/SceneCaptureRendering.cpp
@@ -52,8 +52,10 @@ void UpdateSceneCaptureContent_RenderThread(FSceneRenderer* SceneRenderer, FText
 
 			FRenderingCompositePassContext Context(SceneRenderer->Views[0]);
 
-			VertexShader->SetParameters(SceneRenderer->Views[0]);
-			PixelShader->SetParameters(TStaticSamplerState<SF_Point>::GetRHI(), GSceneRenderTargets.GetSceneColorTexture());
+			//@todo-rco: RHIPacketList
+			FRHICommandList* RHICmdList = nullptr;
+			VertexShader->SetParameters(RHICmdList, SceneRenderer->Views[0]);
+			PixelShader->SetParameters(RHICmdList, TStaticSamplerState<SF_Point>::GetRHI(), GSceneRenderTargets.GetSceneColorTexture());
 
 			FIntPoint TargetSize(UnconstrainedViewRect.Width(), UnconstrainedViewRect.Height());
 
diff --git a/Engine/Source/Runtime/Renderer/Private/SceneHitProxyRendering.cpp b/Engine/Source/Runtime/Renderer/Private/SceneHitProxyRendering.cpp
--- a/Engine/Source/Runtime/Renderer/Private/SceneHitProxyRendering.cpp
+++ b/Engine/Source/Runtime/Renderer/Private/SceneHitProxyRendering.cpp
@@ -17,14 +17,14 @@ class FHitProxyVS : public FMeshMaterialShader
 
 public:
 
-	void SetParameters(const FMaterialRenderProxy* MaterialRenderProxy,const FSceneView& View)
+	void SetParameters(FRHICommandList* RHICmdList, const FMaterialRenderProxy* MaterialRenderProxy,const FSceneView& View)
 	{
-		FMeshMaterialShader::SetParameters(GetVertexShader(), MaterialRenderProxy, *MaterialRenderProxy->GetMaterial(View.GetFeatureLevel()), View, ESceneRenderTargetsMode::SetTextures);
+		FMeshMaterialShader::SetParameters(RHICmdList, GetVertexShader(), MaterialRenderProxy, *MaterialRenderProxy->GetMaterial(View.GetFeatureLevel()), View, ESceneRenderTargetsMode::SetTextures);
 	}
 
-	void SetMesh(const FVertexFactory* VertexFactory,const FSceneView& View,const FPrimitiveSceneProxy* Proxy,const FMeshBatchElement& BatchElement)
+	void SetMesh(FRHICommandList* RHICmdList, const FVertexFactory* VertexFactory,const FSceneView& View,const FPrimitiveSceneProxy* Proxy,const FMeshBatchElement& BatchElement)
 	{
-		FMeshMaterialShader::SetMesh(GetVertexShader(),VertexFactory,View,Proxy,BatchElement);
+		FMeshMaterialShader::SetMesh(RHICmdList, GetVertexShader(),VertexFactory,View,Proxy,BatchElement);
 	}
 
 	virtual bool Serialize(FArchive& Ar)
@@ -123,19 +123,19 @@ public:
 
 	FHitProxyPS() {}
 
-	void SetParameters(const FMaterialRenderProxy* MaterialRenderProxy,const FSceneView& View)
+	void SetParameters(FRHICommandList* RHICmdList, const FMaterialRenderProxy* MaterialRenderProxy,const FSceneView& View)
 	{
-		FMeshMaterialShader::SetParameters(GetPixelShader(), MaterialRenderProxy, *MaterialRenderProxy->GetMaterial(View.GetFeatureLevel()), View, ESceneRenderTargetsMode::SetTextures);
+		FMeshMaterialShader::SetParameters(RHICmdList, GetPixelShader(), MaterialRenderProxy, *MaterialRenderProxy->GetMaterial(View.GetFeatureLevel()), View, ESceneRenderTargetsMode::SetTextures);
 	}
 
-	void SetMesh(const FVertexFactory* VertexFactory,const FSceneView& View,const FPrimitiveSceneProxy* Proxy,const FMeshBatchElement& BatchElement)
+	void SetMesh(FRHICommandList* RHICmdList, const FVertexFactory* VertexFactory,const FSceneView& View,const FPrimitiveSceneProxy* Proxy,const FMeshBatchElement& BatchElement)
 	{
-		FMeshMaterialShader::SetMesh(GetPixelShader(),VertexFactory,View,Proxy,BatchElement);
+		FMeshMaterialShader::SetMesh(RHICmdList, GetPixelShader(),VertexFactory,View,Proxy,BatchElement);
 	}
 
-	void SetHitProxyId(FHitProxyId HitProxyIdValue)
+	void SetHitProxyId(FRHICommandList* RHICmdList, FHitProxyId HitProxyIdValue)
 	{
-		SetShaderValue(GetPixelShader(), HitProxyId, HitProxyIdValue.GetColor().ReinterpretAsLinear());
+		SetShaderValue(RHICmdList, GetPixelShader(), HitProxyId, HitProxyIdValue.GetColor().ReinterpretAsLinear());
 	}
 
 	virtual bool Serialize(FArchive& Ar)
@@ -155,7 +155,7 @@ FHitProxyDrawingPolicy::FHitProxyDrawingPolicy(
 	const FVertexFactory* InVertexFactory,
 	const FMaterialRenderProxy* InMaterialRenderProxy,
 	ERHIFeatureLevel::Type InFeatureLevel
-	) :
+	):
 	FMeshDrawingPolicy(InVertexFactory, InMaterialRenderProxy, *InMaterialRenderProxy->GetMaterial(InFeatureLevel))
 {
 	HullShader = NULL;
@@ -173,23 +173,24 @@ FHitProxyDrawingPolicy::FHitProxyDrawingPolicy(
 	PixelShader = MaterialResource->GetShader<FHitProxyPS>(InVertexFactory->GetType());
 }
 
-void FHitProxyDrawingPolicy::DrawShared(const FSceneView* View, FBoundShaderStateRHIParamRef BoundShaderState) const
+void FHitProxyDrawingPolicy::DrawShared(FRHICommandList* RHICmdList, const FSceneView* View, FBoundShaderStateRHIParamRef BoundShaderState) const
 {
+	check(!RHICmdList);
 	// Set the actual shader & vertex declaration state
 	RHISetBoundShaderState( BoundShaderState);
 
 	// Set the depth-only shader parameters for the material.
-	VertexShader->SetParameters(MaterialRenderProxy,*View);
-	PixelShader->SetParameters(MaterialRenderProxy,*View);
+	VertexShader->SetParameters(RHICmdList, MaterialRenderProxy,*View);
+	PixelShader->SetParameters(RHICmdList, MaterialRenderProxy,*View);
 
 	if(HullShader && DomainShader)
 	{
-		HullShader->SetParameters(MaterialRenderProxy,*View);
-		DomainShader->SetParameters(MaterialRenderProxy,*View);
+		HullShader->SetParameters(RHICmdList, MaterialRenderProxy,*View);
+		DomainShader->SetParameters(RHICmdList, MaterialRenderProxy,*View);
 	}
 
 	// Set the shared mesh resources.
-	FMeshDrawingPolicy::DrawShared(View);
+	FMeshDrawingPolicy::DrawShared(RHICmdList, View);
 }
 
 /** 
@@ -213,6 +214,7 @@ FBoundShaderStateRHIRef FHitProxyDrawingPolicy::CreateBoundShaderState(ERHIFeatu
 }
 
 void FHitProxyDrawingPolicy::SetMeshRenderState(
+	FRHICommandList* RHICmdList, 
 	const FSceneView& View,
 	const FPrimitiveSceneProxy* PrimitiveSceneProxy,
 	const FMeshBatch& Mesh,
@@ -225,23 +227,23 @@ void FHitProxyDrawingPolicy::SetMeshRenderState(
 
 	const FMeshBatchElement& BatchElement = Mesh.Elements[BatchElementIndex];
 
-	VertexShader->SetMesh(VertexFactory, View, PrimitiveSceneProxy, BatchElement);
+	VertexShader->SetMesh(RHICmdList, VertexFactory, View, PrimitiveSceneProxy, BatchElement);
 
 	if(HullShader && DomainShader)
 	{
-		HullShader->SetMesh(VertexFactory, View, PrimitiveSceneProxy, BatchElement);
-		DomainShader->SetMesh(VertexFactory, View, PrimitiveSceneProxy, BatchElement);
+		HullShader->SetMesh(RHICmdList, VertexFactory, View, PrimitiveSceneProxy, BatchElement);
+		DomainShader->SetMesh(RHICmdList, VertexFactory, View, PrimitiveSceneProxy, BatchElement);
 	}
 
-	PixelShader->SetMesh(VertexFactory, View, PrimitiveSceneProxy, BatchElement);
+	PixelShader->SetMesh(RHICmdList, VertexFactory, View, PrimitiveSceneProxy, BatchElement);
 	// Per-instance hitproxies are supplied by the vertex factory
 	if( PrimitiveSceneProxy && PrimitiveSceneProxy->HasPerInstanceHitProxies() )
 	{
-		PixelShader->SetHitProxyId( FColor(0) );
+		PixelShader->SetHitProxyId(RHICmdList, FColor(0) );
 	}
 	else
 	{
-		PixelShader->SetHitProxyId( HitProxyId);	
+		PixelShader->SetHitProxyId(RHICmdList, HitProxyId);	
 	}
 
 	RHISetRasterizerState(GetStaticRasterizerState<false>(
@@ -258,7 +260,7 @@ void FHitProxyDrawingPolicyFactory::AddStaticMesh(FScene* Scene,FStaticMesh* Sta
 	// Add the static mesh to the DPG's hit proxy draw list.
 	const FMaterialRenderProxy* MaterialRenderProxy = StaticMesh->MaterialRenderProxy;
 	const FMaterial* Material = MaterialRenderProxy->GetMaterial(Scene->GetFeatureLevel());
-	if (!Material->IsMasked() && !Material->IsTwoSided() && !Material->MaterialModifiesMeshPosition())
+	if ( !Material->IsMasked() && !Material->IsTwoSided() && !Material->MaterialModifiesMeshPosition() )
 	{
 		// Default material doesn't handle masked, and doesn't have the correct bIsTwoSided setting.
 		MaterialRenderProxy = UMaterial::GetDefaultMaterial(MD_Surface)->GetRenderProxy(false);
@@ -289,6 +291,7 @@ void FHitProxyDrawingPolicyFactory::AddStaticMesh(FScene* Scene,FStaticMesh* Sta
 }
 
 bool FHitProxyDrawingPolicyFactory::DrawDynamicMesh(
+	FRHICommandList* RHICmdList, 
 	const FSceneView& View,
 	ContextType DrawingContext,
 	const FMeshBatch& Mesh,
@@ -308,17 +311,17 @@ bool FHitProxyDrawingPolicyFactory::DrawDynamicMesh(
 		if( View.bAllowTranslucentPrimitivesInHitProxy || !IsTranslucentBlendMode( Material->GetBlendMode() ) || ( HitProxy && HitProxy->AlwaysAllowsTranslucentPrimitives() ) )
 #endif
 		{
-			if (!Material->IsMasked() && !Material->IsTwoSided() && !Material->MaterialModifiesMeshPosition())
+			if ( !Material->IsMasked() && !Material->IsTwoSided() && !Material->MaterialModifiesMeshPosition() )
 			{
 				// Default material doesn't handle masked, and doesn't have the correct bIsTwoSided setting.
 				MaterialRenderProxy = UMaterial::GetDefaultMaterial(MD_Surface)->GetRenderProxy(false);
 			}
 			FHitProxyDrawingPolicy DrawingPolicy( Mesh.VertexFactory, MaterialRenderProxy, View.GetFeatureLevel() );
-			DrawingPolicy.DrawShared(&View, DrawingPolicy.CreateBoundShaderState(View.GetFeatureLevel()));
+			DrawingPolicy.DrawShared(RHICmdList, &View, DrawingPolicy.CreateBoundShaderState(View.GetFeatureLevel()));
 			for( int32 BatchElementIndex = 0; BatchElementIndex < Mesh.Elements.Num(); BatchElementIndex++ )
 			{
-				DrawingPolicy.SetMeshRenderState(View, PrimitiveSceneProxy, Mesh, BatchElementIndex, bBackFace, HitProxyId);
-				DrawingPolicy.DrawMesh(Mesh,BatchElementIndex);
+				DrawingPolicy.SetMeshRenderState(RHICmdList, View, PrimitiveSceneProxy, Mesh, BatchElementIndex, bBackFace, HitProxyId);
+				DrawingPolicy.DrawMesh(RHICmdList, Mesh,BatchElementIndex);
 			}
 			return true;
 		}
@@ -355,6 +358,9 @@ void FDeferredShadingSceneRenderer::RenderHitProxies()
 
 	RHISetRenderTarget(HitProxyRT->GetRenderTargetItem().TargetableTexture, GSceneRenderTargets.GetSceneDepthSurface());
 
+	//@todo-rco: RHIPacketList
+	FRHICommandList* RHICmdList = nullptr;
+
 	// Clear color for each view.
 	for(int32 ViewIndex = 0;ViewIndex < Views.Num();ViewIndex++)
 	{
@@ -393,12 +399,12 @@ void FDeferredShadingSceneRenderer::RenderHitProxies()
 		if( !View.bAllowTranslucentPrimitivesInHitProxy )
 		{
 			// Draw the scene's hit proxy draw lists. (opaque primitives only)
-			Scene->HitProxyDrawList_OpaqueOnly.DrawVisible(View,View.StaticMeshVisibilityMap,View.StaticMeshBatchVisibility);
+			Scene->HitProxyDrawList_OpaqueOnly.DrawVisible(RHICmdList, View,View.StaticMeshVisibilityMap,View.StaticMeshBatchVisibility);
 		}
 		else
 		{
 			// Draw the scene's hit proxy draw lists.
-			Scene->HitProxyDrawList.DrawVisible(View,View.StaticMeshVisibilityMap,View.StaticMeshBatchVisibility);
+			Scene->HitProxyDrawList.DrawVisible(RHICmdList, View,View.StaticMeshVisibilityMap,View.StaticMeshBatchVisibility);
 		}
 
 		const bool bPreFog = true;
@@ -509,5 +515,7 @@ void FDeferredShadingSceneRenderer::RenderHitProxies()
 				false,
 				1.0f
 				);
+
+	RHIEndScene();
 #endif
 }
diff --git a/Engine/Source/Runtime/Renderer/Private/SceneHitProxyRendering.h b/Engine/Source/Runtime/Renderer/Private/SceneHitProxyRendering.h
--- a/Engine/Source/Runtime/Renderer/Private/SceneHitProxyRendering.h
+++ b/Engine/Source/Runtime/Renderer/Private/SceneHitProxyRendering.h
@@ -30,7 +30,7 @@ public:
 			VertexShader == Other.VertexShader &&
 			PixelShader == Other.PixelShader;
 	}
-	void DrawShared(const FSceneView* View,FBoundShaderStateRHIParamRef BoundShaderState) const;
+	void DrawShared(FRHICommandList* RHICmdList, const FSceneView* View,FBoundShaderStateRHIParamRef BoundShaderState) const;
 
 	/** 
 	* Create bound shader state using the vertex decl from the mesh draw policy
@@ -40,6 +40,7 @@ public:
 	FBoundShaderStateRHIRef CreateBoundShaderState(ERHIFeatureLevel::Type InFeatureLevel);
 
 	void SetMeshRenderState(
+		FRHICommandList* RHICmdList, 
 		const FSceneView& View,
 		const FPrimitiveSceneProxy* PrimitiveSceneProxy,
 		const FMeshBatch& Mesh,
@@ -67,6 +68,7 @@ public:
 
 	static void AddStaticMesh(FScene* Scene,FStaticMesh* StaticMesh,ContextType DrawingContext = ContextType());
 	static bool DrawDynamicMesh(
+		FRHICommandList* RHICmdList, 
 		const FSceneView& View,
 		ContextType DrawingContext,
 		const FMeshBatch& Mesh,
diff --git a/Engine/Source/Runtime/Renderer/Private/SceneOcclusion.cpp b/Engine/Source/Runtime/Renderer/Private/SceneOcclusion.cpp
--- a/Engine/Source/Runtime/Renderer/Private/SceneOcclusion.cpp
+++ b/Engine/Source/Runtime/Renderer/Private/SceneOcclusion.cpp
@@ -430,7 +430,9 @@ static void IssueProjectedShadowOcclusionQuery(FViewInfo& View, const FProjected
 		// Allocate an occlusion query for the primitive from the occlusion query pool.
 		const FRenderQueryRHIRef ShadowOcclusionQuery = ViewState->OcclusionQueryPool.AllocateQuery();
 
-		VertexShader->SetParameters(View);
+		//@todo-rco: RHIPacketList
+		FRHICommandList* RHICmdList = nullptr;
+		VertexShader->SetParameters(RHICmdList, View);
 
 		// Draw the primitive's bounding box, using the occlusion query.
 		RHIBeginRenderQuery(ShadowOcclusionQuery);
@@ -640,19 +642,19 @@ public:
 		BoundsExtentSampler.Bind( Initializer.ParameterMap, TEXT("BoundsExtentSampler") );
 	}
 
-	void SetParameters( const FSceneView& View, const FHZB& HZB, FTextureRHIParamRef BoundsCenter, FTextureRHIParamRef BoundsExtent )
+	void SetParameters(FRHICommandList* RHICmdList, const FSceneView& View, const FHZB& HZB, FTextureRHIParamRef BoundsCenter, FTextureRHIParamRef BoundsExtent )
 	{
 		const FPixelShaderRHIParamRef ShaderRHI = GetPixelShader();
 
-		FGlobalShader::SetParameters( ShaderRHI, View );
+		FGlobalShader::SetParameters(RHICmdList, ShaderRHI, View );
 
 		const FVector2D InvSize( 1.0f / HZB.Size.X, 1.0f / HZB.Size.Y );
-		SetShaderValue( ShaderRHI, InvSizeParameter, InvSize );
+		SetShaderValue(RHICmdList, ShaderRHI, InvSizeParameter, InvSize );
 
-		SetTextureParameter( ShaderRHI, HZBTexture, HZBSampler, TStaticSamplerState<SF_Point,AM_Clamp,AM_Clamp,AM_Clamp>::GetRHI(), HZB.Texture->GetRenderTargetItem().ShaderResourceTexture );
+		SetTextureParameter(RHICmdList, ShaderRHI, HZBTexture, HZBSampler, TStaticSamplerState<SF_Point,AM_Clamp,AM_Clamp,AM_Clamp>::GetRHI(), HZB.Texture->GetRenderTargetItem().ShaderResourceTexture );
 
-		SetTextureParameter( ShaderRHI, BoundsCenterTexture, BoundsCenterSampler, TStaticSamplerState<SF_Point,AM_Clamp,AM_Clamp,AM_Clamp>::GetRHI(), BoundsCenter );
-		SetTextureParameter( ShaderRHI, BoundsExtentTexture, BoundsExtentSampler, TStaticSamplerState<SF_Point,AM_Clamp,AM_Clamp,AM_Clamp>::GetRHI(), BoundsExtent );
+		SetTextureParameter(RHICmdList, ShaderRHI, BoundsCenterTexture, BoundsCenterSampler, TStaticSamplerState<SF_Point,AM_Clamp,AM_Clamp,AM_Clamp>::GetRHI(), BoundsCenter );
+		SetTextureParameter(RHICmdList, ShaderRHI, BoundsExtentTexture, BoundsExtentSampler, TStaticSamplerState<SF_Point,AM_Clamp,AM_Clamp,AM_Clamp>::GetRHI(), BoundsExtent );
 	}
 
 	virtual bool Serialize(FArchive& Ar)
@@ -752,7 +754,9 @@ void FHZBOcclusionTester::Submit( const FViewInfo& View )
 		static FGlobalBoundShaderState BoundShaderState;
 		SetGlobalBoundShaderState( BoundShaderState, GFilterVertexDeclaration.VertexDeclarationRHI, *VertexShader, *PixelShader );
 
-		PixelShader->SetParameters( View,  ViewState->HZB, BoundsCenterTexture->GetRenderTargetItem().ShaderResourceTexture, BoundsExtentTexture->GetRenderTargetItem().ShaderResourceTexture );
+		//@todo-rco: RHIPacketList
+		FRHICommandList* RHICmdList = nullptr;
+		PixelShader->SetParameters(RHICmdList, View,  ViewState->HZB, BoundsCenterTexture->GetRenderTargetItem().ShaderResourceTexture, BoundsExtentTexture->GetRenderTargetItem().ShaderResourceTexture );
 
 		RHISetViewport( 0, 0, 0.0f, SizeX, SizeY, 1.0f );
 
@@ -808,31 +812,31 @@ public:
 		TextureParameterSampler.Bind( Initializer.ParameterMap, TEXT("TextureSampler") );
 	}
 
-	void SetParameters( const FSceneView& View, const FIntPoint& Size )
+	void SetParameters(FRHICommandList* RHICmdList, const FSceneView& View, const FIntPoint& Size )
 	{
 		const FPixelShaderRHIParamRef ShaderRHI = GetPixelShader();
 
-		FGlobalShader::SetParameters( ShaderRHI, View );
+		FGlobalShader::SetParameters(RHICmdList, ShaderRHI, View );
 
 		const FVector2D InvSize( 1.0f / Size.X, 1.0f / Size.Y );
-		SetShaderValue( ShaderRHI, InvSizeParameter, InvSize );
+		SetShaderValue(RHICmdList, ShaderRHI, InvSizeParameter, InvSize );
 		
-		SceneTextureParameters.Set( ShaderRHI, View );
+		SceneTextureParameters.Set(RHICmdList, ShaderRHI, View );
 	}
 
-	void SetParameters( const FSceneView& View, const FIntPoint& Size, FShaderResourceViewRHIParamRef ShaderResourceView )
+	void SetParameters(FRHICommandList* RHICmdList, const FSceneView& View, const FIntPoint& Size, FShaderResourceViewRHIParamRef ShaderResourceView )
 	{
 		const FPixelShaderRHIParamRef ShaderRHI = GetPixelShader();
 
-		FGlobalShader::SetParameters( ShaderRHI, View );
+		FGlobalShader::SetParameters(RHICmdList, ShaderRHI, View );
 
 		const FVector2D InvSize( 1.0f / Size.X, 1.0f / Size.Y );
-		SetShaderValue( ShaderRHI, InvSizeParameter, InvSize );
+		SetShaderValue(RHICmdList, ShaderRHI, InvSizeParameter, InvSize );
 
 		//SetTextureParameter( ShaderRHI, TextureParameter, TextureParameterSampler, TStaticSamplerState<SF_Point,AM_Clamp,AM_Clamp,AM_Clamp>::GetRHI(), Texture );
 
-		SetSRVParameter( ShaderRHI, TextureParameter, ShaderResourceView );
-		SetSamplerParameter( ShaderRHI, TextureParameterSampler, TStaticSamplerState<SF_Point,AM_Clamp,AM_Clamp,AM_Clamp>::GetRHI() );
+		SetSRVParameter(RHICmdList, ShaderRHI, TextureParameter, ShaderResourceView );
+		SetSamplerParameter(RHICmdList, ShaderRHI, TextureParameterSampler, TStaticSamplerState<SF_Point,AM_Clamp,AM_Clamp,AM_Clamp>::GetRHI() );
 	}
 
 	virtual bool Serialize(FArchive& Ar)
@@ -882,6 +886,9 @@ void BuildHZB( const FViewInfo& View )
 	RHISetRasterizerState( TStaticRasterizerState<>::GetRHI() );
 	RHISetDepthStencilState( TStaticDepthStencilState< false, CF_Always >::GetRHI() );
 	
+	//@todo-rco: RHIPacketList
+	FRHICommandList* RHICmdList = nullptr;
+
 	// Mip 0
 	{
 		SCOPED_DRAW_EVENTF(BuildHZB, DEC_SCENE_ITEMS, TEXT("HZB#%d"), 0);
@@ -895,7 +902,7 @@ void BuildHZB( const FViewInfo& View )
 		SetGlobalBoundShaderState( BoundShaderState, GFilterVertexDeclaration.VertexDeclarationRHI, *VertexShader, *PixelShader );
 
 		// Imperfect sampling, doesn't matter too much
-		PixelShader->SetParameters( View, GSceneRenderTargets.GetBufferSizeXY() );
+		PixelShader->SetParameters(RHICmdList, View, GSceneRenderTargets.GetBufferSizeXY() );
 
 		RHISetViewport( 0, 0, 0.0f, HZBSize.X, HZBSize.Y, 1.0f );
 
@@ -926,7 +933,7 @@ void BuildHZB( const FViewInfo& View )
 		static FGlobalBoundShaderState BoundShaderState;
 		SetGlobalBoundShaderState( BoundShaderState, GFilterVertexDeclaration.VertexDeclarationRHI, *VertexShader, *PixelShader );
 
-		PixelShader->SetParameters( View, SrcSize, ViewState->HZB.MipSRVs[ MipIndex - 1 ] );
+		PixelShader->SetParameters(RHICmdList, View, SrcSize, ViewState->HZB.MipSRVs[ MipIndex - 1 ] );
 
 		RHISetViewport( 0, 0, 0.0f, DstSize.X, DstSize.Y, 1.0f );
 
@@ -998,7 +1005,9 @@ void FDeferredShadingSceneRenderer::BeginOcclusionTests()
 			// Lookup the vertex shader.
 			TShaderMapRef<FOcclusionQueryVS> VertexShader(GetGlobalShaderMap());
 			SetGlobalBoundShaderState(OcclusionTestBoundShaderState,GetVertexDeclarationFVector3(),*VertexShader,NULL);
-			VertexShader->SetParameters(View);
+			//@todo-rco: RHIPacketList
+			FRHICommandList* RHICmdList = nullptr;
+			VertexShader->SetParameters(RHICmdList, View);
 
 			// Issue this frame's occlusion queries (occlusion queries from last frame may still be in flight)
 			TMap<FSceneViewState::FProjectedShadowKey, FRenderQueryRHIRef>& ShadowOcclusionQueryMap = ViewState->ShadowOcclusionQueryMap;
@@ -1056,7 +1065,7 @@ void FDeferredShadingSceneRenderer::BeginOcclusionTests()
 								ShadowOcclusionQueryMap.Add(Key, ShadowOcclusionQuery);
 
 								// Draw bounding sphere
-								VertexShader->SetParametersWithBoundingSphere(View, ProjectedShadowInfo.LightSceneInfo->Proxy->GetBoundingSphere());
+								VertexShader->SetParametersWithBoundingSphere(RHICmdList, View, ProjectedShadowInfo.LightSceneInfo->Proxy->GetBoundingSphere());
 								StencilingGeometry::DrawSphere();
 
 								RHIEndRenderQuery(ShadowOcclusionQuery);
@@ -1084,7 +1093,7 @@ void FDeferredShadingSceneRenderer::BeginOcclusionTests()
 		    if ( !ViewState->HasViewParent() && !ViewState->bIsFrozen )
     #endif
 		    {
-				VertexShader->SetParameters(View);
+				VertexShader->SetParameters(RHICmdList, View);
 
 				{
 					SCOPED_DRAW_EVENT(IndividualQueries, DEC_SCENE_ITEMS);
diff --git a/Engine/Source/Runtime/Renderer/Private/SceneOcclusion.h b/Engine/Source/Runtime/Renderer/Private/SceneOcclusion.h
--- a/Engine/Source/Runtime/Renderer/Private/SceneOcclusion.h
+++ b/Engine/Source/Runtime/Renderer/Private/SceneOcclusion.h
@@ -28,21 +28,21 @@ public:
 		FGlobalShader::ModifyCompilationEnvironment(Platform, OutEnvironment);
 	}
 
-	void SetParametersWithBoundingSphere(const FSceneView& View, const FSphere& BoundingSphere)
+	void SetParametersWithBoundingSphere(FRHICommandList* RHICmdList, const FSceneView& View, const FSphere& BoundingSphere)
 	{
-		FGlobalShader::SetParameters(GetVertexShader(), View);
+		FGlobalShader::SetParameters(RHICmdList, GetVertexShader(), View);
 
 		FVector4 StencilingSpherePosAndScale;
 		StencilingGeometry::GStencilSphereVertexBuffer.CalcTransform(StencilingSpherePosAndScale, BoundingSphere, View.ViewMatrices.PreViewTranslation);
-		StencilingGeometryParameters.Set(this, StencilingSpherePosAndScale);
+		StencilingGeometryParameters.Set(RHICmdList, this, StencilingSpherePosAndScale);
 	}
 
-	void SetParameters(const FSceneView& View)
+	void SetParameters(FRHICommandList* RHICmdList, const FSceneView& View)
 	{
-		FGlobalShader::SetParameters(GetVertexShader(),View);
+		FGlobalShader::SetParameters(RHICmdList, GetVertexShader(),View);
 
 		// Don't transform if rendering frustum
-		StencilingGeometryParameters.Set(this, FVector4(0,0,0,1));
+		StencilingGeometryParameters.Set(RHICmdList, this, FVector4(0,0,0,1));
 	}
 
 	// Begin FShader interface
diff --git a/Engine/Source/Runtime/Renderer/Private/ScenePostProcessing.cpp b/Engine/Source/Runtime/Renderer/Private/ScenePostProcessing.cpp
--- a/Engine/Source/Runtime/Renderer/Private/ScenePostProcessing.cpp
+++ b/Engine/Source/Runtime/Renderer/Private/ScenePostProcessing.cpp
@@ -134,17 +134,21 @@ void FSceneRenderer::GammaCorrectToViewportRenderTarget(const FViewInfo* View, f
 
 	const FPixelShaderRHIParamRef ShaderRHI = PixelShader->GetPixelShader();
 
+	//@todo-rco: RHIPacketList
+	FRHICommandList* RHICmdList = nullptr;
 	SetShaderValue(
+		RHICmdList, 
 		ShaderRHI,
 		PixelShader->InverseGamma,
 		InvDisplayGamma
 		);
-	SetShaderValue(ShaderRHI,PixelShader->ColorScale,View->ColorScale);
-	SetShaderValue(ShaderRHI,PixelShader->OverlayColor,View->OverlayColor);
+	SetShaderValue(RHICmdList, ShaderRHI,PixelShader->ColorScale,View->ColorScale);
+	SetShaderValue(RHICmdList, ShaderRHI,PixelShader->OverlayColor,View->OverlayColor);
 
 	const FTextureRHIRef DesiredSceneColorTexture = GSceneRenderTargets.GetSceneColorTexture();
 
 	SetTextureParameter(
+		RHICmdList, 
 		ShaderRHI,
 		PixelShader->SceneTexture,
 		PixelShader->SceneTextureSampler,
diff --git a/Engine/Source/Runtime/Renderer/Private/SceneRendering.cpp b/Engine/Source/Runtime/Renderer/Private/SceneRendering.cpp
--- a/Engine/Source/Runtime/Renderer/Private/SceneRendering.cpp
+++ b/Engine/Source/Runtime/Renderer/Private/SceneRendering.cpp
@@ -489,7 +489,7 @@ TUniformBufferRef<FViewUniformShaderParameters> FViewInfo::CreateUniformBuffer(
 	checkSlow(sizeof(ViewUniformShaderParameters.SkyIrradianceEnvironmentMap) == sizeof(FVector4) * 7);
 	SetupSkyIrradianceEnvironmentMapConstants((FVector4*)&ViewUniformShaderParameters.SkyIrradianceEnvironmentMap);
 
-	return TUniformBufferRef<FViewUniformShaderParameters>::CreateUniformBufferImmediate(ViewUniformShaderParameters, UniformBuffer_SingleUse);
+	return TUniformBufferRef<FViewUniformShaderParameters>::CreateUniformBufferImmediate(ViewUniformShaderParameters, UniformBuffer_SingleFrame);
 }
 
 void FViewInfo::InitRHIResources()
@@ -564,6 +564,7 @@ FSceneRenderer::FSceneRenderer(const FSceneViewFamily* InViewFamily,FHitProxyCon
 	check(Scene != NULL);
 
 	// Copy the individual views.
+	bool bAnyViewIsLocked = false;
 	Views.Empty(InViewFamily->Views.Num());
 	for(int32 ViewIndex = 0;ViewIndex < InViewFamily->Views.Num();ViewIndex++)
 	{
@@ -581,6 +582,7 @@ FSceneRenderer::FSceneRenderer(const FSceneViewFamily* InViewFamily,FHitProxyCon
 		FViewInfo* ViewInfo = new(Views) FViewInfo(InViewFamily->Views[ViewIndex]);
 		ViewFamily.Views[ViewIndex] = ViewInfo;
 		ViewInfo->Family = &ViewFamily;
+		bAnyViewIsLocked |= ViewInfo->bIsLocked;
 
 #if WITH_EDITOR
 		// Should we allow the user to select translucent primitives?
@@ -596,6 +598,14 @@ FSceneRenderer::FSceneRenderer(const FSceneViewFamily* InViewFamily,FHitProxyCon
 			ViewInfo->Drawer->Draw(ViewInfo,&ViewElementPDI);
 		}
 	}
+
+	// If any viewpoint has been locked, set time to zero to avoid time-based
+	// rendering differences in materials.
+	if (bAnyViewIsLocked)
+	{
+		ViewFamily.CurrentRealTime = 0.0f;
+		ViewFamily.CurrentWorldTime = 0.0f;
+	}
 	
 	if(HitProxyConsumer)
 	{
@@ -685,7 +695,8 @@ void FSceneRenderer::RenderFinish()
 			// display a message saying we're frozen
 			FSceneViewState* ViewState = (FSceneViewState*)View.State;
 			bool bViewParentOrFrozen = ViewState && (ViewState->HasViewParent() || ViewState->bIsFrozen);
-			if (bViewParentOrFrozen || bShowPrecomputedVisibilityWarning)
+			bool bLocked = View.bIsLocked;
+			if (bViewParentOrFrozen || bShowPrecomputedVisibilityWarning || bLocked)
 			{
 				// this is a helper class for FCanvas to be able to get screen size
 				class FRenderTargetTemp : public FRenderTarget
@@ -726,6 +737,12 @@ void FSceneRenderer::RenderFinish()
 					Canvas.DrawShadowedText( 10, Y, Message, GetStatsFont(), FLinearColor(1.0,0.05,0.05,1.0));
 					Y += 14;
 				}
+				if (bLocked)
+				{
+					const FText Message = NSLOCTEXT("Renderer", "ViewLocked", "VIEW LOCKED");
+					Canvas.DrawShadowedText( 10, Y, Message, GetStatsFont(), FLinearColor(0.8,1.0,0.2,1.0));
+					Y += 14;
+				}
 				Canvas.Flush();
 			}
 		}
@@ -794,6 +811,9 @@ void FSceneRenderer::RenderFinish()
 	}
 
 #endif
+
+	// Notify the RHI we are done rendering a scene.
+	RHIEndScene();
 }
 
 FSceneRenderer* FSceneRenderer::CreateSceneRenderer(const FSceneViewFamily* InViewFamily, FHitProxyConsumer* HitProxyConsumer)
@@ -839,6 +859,9 @@ void FSceneRenderer::RenderCustomDepthPass()
 	{
 		SCOPED_DRAW_EVENT(CustomDepth, DEC_SCENE_ITEMS);
 
+		//@todo-rco: RHIPacketList
+		FRHICommandList* RHICmdList = nullptr;
+
 		for(int32 ViewIndex = 0;ViewIndex < Views.Num();ViewIndex++)
 		{
 			SCOPED_CONDITIONAL_DRAW_EVENTF(EventView, Views.Num() > 1, DEC_SCENE_ITEMS, TEXT("View%d"), ViewIndex);
@@ -857,7 +880,7 @@ void FSceneRenderer::RenderCustomDepthPass()
 			RHISetDepthStencilState(TStaticDepthStencilState<true,CF_GreaterEqual>::GetRHI());
 			RHISetBlendState(TStaticBlendState<>::GetRHI());
 
-			View.CustomDepthSet.DrawPrims(&View, false);
+			View.CustomDepthSet.DrawPrims(RHICmdList, &View, false);
 		}
 
 		// resolve using the current ResolveParams 
diff --git a/Engine/Source/Runtime/Renderer/Private/SceneRendering.h b/Engine/Source/Runtime/Renderer/Private/SceneRendering.h
--- a/Engine/Source/Runtime/Renderer/Private/SceneRendering.h
+++ b/Engine/Source/Runtime/Renderer/Private/SceneRendering.h
@@ -561,7 +561,7 @@ protected:
 	void OnStartFrame();
 
 	/** Renders the scene's distortion */
-	void RenderDistortion();
+	void RenderDistortion(FRHICommandList* RHICmdList);
 };
 
 
@@ -583,10 +583,10 @@ protected:
 	void InitViews();
 
 	/** Renders the opaque base pass for forward shading. */
-	void RenderForwardShadingBasePass();
+	void RenderForwardShadingBasePass(FRHICommandList* RHICmdList);
 
 	/** Makes a copy of scene alpha so PC can emulate ES2 framebuffer fetch. */
-	void CopySceneAlpha(const FSceneView& View);
+	void CopySceneAlpha(FRHICommandList* RHICmdList, const FSceneView& View);
 
 	/** Renders the base pass for translucency. */
 	void RenderTranslucency();
diff --git a/Engine/Source/Runtime/Renderer/Private/SceneVisibility.cpp b/Engine/Source/Runtime/Renderer/Private/SceneVisibility.cpp
--- a/Engine/Source/Runtime/Renderer/Private/SceneVisibility.cpp
+++ b/Engine/Source/Runtime/Renderer/Private/SceneVisibility.cpp
@@ -127,7 +127,7 @@ static void UpdatePrimitiveFadingState(FPrimitiveFadingState& FadingState, FView
 
 				FDistanceCullFadeUniformShaderParameters Uniforms;
 				Uniforms.FadeTimeScaleBias = FadingState.FadeTimeScaleBias;
-				FadingState.UniformBuffer = FDistanceCullFadeUniformBufferRef::CreateUniformBufferImmediate( Uniforms, UniformBuffer_MultiUse );
+				FadingState.UniformBuffer = FDistanceCullFadeUniformBufferRef::CreateUniformBufferImmediate( Uniforms, UniformBuffer_MultiFrame );
 			}
 			else
 			{
@@ -151,7 +151,7 @@ static void UpdatePrimitiveFadingState(FPrimitiveFadingState& FadingState, FView
 
 				FDistanceCullFadeUniformShaderParameters Uniforms;
 				Uniforms.FadeTimeScaleBias = FadingState.FadeTimeScaleBias;
-				FadingState.UniformBuffer = FDistanceCullFadeUniformBufferRef::CreateUniformBufferImmediate( Uniforms, UniformBuffer_MultiUse );
+				FadingState.UniformBuffer = FDistanceCullFadeUniformBufferRef::CreateUniformBufferImmediate( Uniforms, UniformBuffer_MultiFrame );
 			}
 		}
 	}
@@ -910,6 +910,9 @@ float Halton( int32 Index, int32 Base )
 
 void FSceneRenderer::PreVisibilityFrameSetup()
 {
+	// Notify the RHI we are beginning to render a scene.
+	RHIBeginScene();
+
 	// Notify the FX system that the scene is about to perform visibility checks.
 	if (Scene->FXSystem)
 	{
diff --git a/Engine/Source/Runtime/Renderer/Private/ShaderBaseClasses.cpp b/Engine/Source/Runtime/Renderer/Private/ShaderBaseClasses.cpp
--- a/Engine/Source/Runtime/Renderer/Private/ShaderBaseClasses.cpp
+++ b/Engine/Source/Runtime/Renderer/Private/ShaderBaseClasses.cpp
@@ -34,51 +34,6 @@ FMaterialShader::FMaterialShader(const FMaterialShaderType::CompiledShaderInitia
 		ParameterCollectionUniformBuffers.Add(CollectionParameter);
 	}
 
-	// Bind uniform 2D texture parameters.
-	for(int32 ParameterIndex = 0;ParameterIndex < Initializer.UniformExpressionSet.Uniform2DTextureExpressions.Num();ParameterIndex++)
-	{
-		FShaderResourceParameter TextureShaderParameter;
-		FString TextureParameterName = FString::Printf(TEXT("MaterialTexture2D_%u"), ParameterIndex);
-		TextureShaderParameter.Bind(Initializer.ParameterMap,*TextureParameterName);
-
-		FShaderResourceParameter SamplerShaderParameter;
-		FString SamplerParameterName = FString::Printf(TEXT("MaterialTexture2D_%uSampler"), ParameterIndex);
-		SamplerShaderParameter.Bind(Initializer.ParameterMap,*SamplerParameterName);
-
-		if(TextureShaderParameter.IsBound() && SamplerShaderParameter.IsBound())
-		{
-			TUniformParameter<FShaderResourceParameter>* UniformParameter = new(Uniform2DTextureShaderResourceParameters) TUniformParameter<FShaderResourceParameter>();
-			UniformParameter->Index = ParameterIndex;
-			UniformParameter->ShaderParameter = TextureShaderParameter;
-
-			UniformParameter = new(Uniform2DTextureSamplerShaderResourceParameters) TUniformParameter<FShaderResourceParameter>();
-			UniformParameter->Index = ParameterIndex;
-			UniformParameter->ShaderParameter = SamplerShaderParameter;
-		}
-	}
-	// Bind uniform cube texture parameters.
-	for(int32 ParameterIndex = 0;ParameterIndex < Initializer.UniformExpressionSet.UniformCubeTextureExpressions.Num();ParameterIndex++)
-	{
-		FShaderResourceParameter TextureShaderParameter;
-		FString TextureParameterName = FString::Printf(TEXT("MaterialTextureCube_%u"), ParameterIndex);
-		TextureShaderParameter.Bind(Initializer.ParameterMap,*TextureParameterName);
-
-		FShaderResourceParameter SamplerShaderParameter;
-		FString SamplerParameterName = FString::Printf(TEXT("MaterialTextureCube_%uSampler"), ParameterIndex);
-		SamplerShaderParameter.Bind(Initializer.ParameterMap,*SamplerParameterName);
-
-		if(TextureShaderParameter.IsBound() && SamplerShaderParameter.IsBound())
-		{
-			TUniformParameter<FShaderResourceParameter>* UniformParameter = new(UniformCubeTextureShaderResourceParameters) TUniformParameter<FShaderResourceParameter>();
-			UniformParameter->Index = ParameterIndex;
-			UniformParameter->ShaderParameter = TextureShaderParameter;
-
-			UniformParameter = new(UniformCubeTextureSamplerShaderResourceParameters) TUniformParameter<FShaderResourceParameter>();
-			UniformParameter->Index = ParameterIndex;
-			UniformParameter->ShaderParameter = SamplerShaderParameter;
-		}
-	}
-
 	DeferredParameters.Bind(Initializer.ParameterMap);
 	LightAttenuation.Bind(Initializer.ParameterMap, TEXT("LightAttenuationTexture"));
 	LightAttenuationSampler.Bind(Initializer.ParameterMap, TEXT("LightAttenuationTextureSampler"));
@@ -104,10 +59,6 @@ bool FMaterialShader::Serialize(FArchive& Ar)
 	const bool bShaderHasOutdatedParameters = FShader::Serialize(Ar);
 	Ar << MaterialUniformBuffer;
 	Ar << ParameterCollectionUniformBuffers;
-	Ar << Uniform2DTextureShaderResourceParameters;
-	Ar << Uniform2DTextureSamplerShaderResourceParameters;
-	Ar << UniformCubeTextureShaderResourceParameters;
-	Ar << UniformCubeTextureSamplerShaderResourceParameters;
 	Ar << DeferredParameters;
 	Ar << LightAttenuation;
 	Ar << LightAttenuationSampler;
@@ -155,10 +106,6 @@ uint32 FMaterialShader::GetAllocatedSize() const
 {
 	return FShader::GetAllocatedSize()
 		+ ParameterCollectionUniformBuffers.GetAllocatedSize()
-		+ Uniform2DTextureShaderResourceParameters.GetAllocatedSize()
-		+ Uniform2DTextureSamplerShaderResourceParameters.GetAllocatedSize()
-		+ UniformCubeTextureShaderResourceParameters.GetAllocatedSize()
-		+ UniformCubeTextureSamplerShaderResourceParameters.GetAllocatedSize()
 		+ DebugDescription.GetAllocatedSize();
 }
 
