UE4 - merge first pass at parallel rendering to main

diff --git a/Engine/Source/Runtime/Renderer/Private/ShaderBaseClasses.h b/Engine/Source/Runtime/Renderer/Private/ShaderBaseClasses.h
--- a/Engine/Source/Runtime/Renderer/Private/ShaderBaseClasses.h
+++ b/Engine/Source/Runtime/Renderer/Private/ShaderBaseClasses.h
@@ -119,16 +119,22 @@ public:
 	FUniformBufferRHIParamRef GetParameterCollectionBuffer(const FGuid& Id, const FSceneInterface* SceneInterface) const;
 
 	template<typename ShaderRHIParamRef>
-	void SetParameters(const ShaderRHIParamRef ShaderRHI,const FSceneView& View)
+	void SetParameters(FRHICommandList* RHICmdList, const ShaderRHIParamRef ShaderRHI,const FSceneView& View)
 	{
 		check(GetUniformBufferParameter<FViewUniformShaderParameters>().IsInitialized());
 		CheckShaderIsValid();
-		SetUniformBufferParameter(ShaderRHI,GetUniformBufferParameter<FViewUniformShaderParameters>(),View.UniformBuffer);
+		SetUniformBufferParameter(RHICmdList, ShaderRHI,GetUniformBufferParameter<FViewUniformShaderParameters>(),View.UniformBuffer);
+	}
+
+	static void ModifyCompilationEnvironment(EShaderPlatform Platform, FShaderCompilerEnvironment& OutEnvironment)
+	{
+		FGlobalShader::ModifyCompilationEnvironment(Platform,OutEnvironment);
 	}
 
 	/** Sets pixel parameters that are material specific but not FMeshBatch specific. */
 	template<typename ShaderRHIParamRef>
 	void SetParameters(
+		FRHICommandList* RHICmdList,
 		const ShaderRHIParamRef ShaderRHI, 
 		const FMaterialRenderProxy* MaterialRenderProxy, 
 		const FMaterial& Material,
@@ -140,7 +146,7 @@ public:
 		FUniformExpressionCache TempUniformExpressionCache;
 		const FUniformExpressionCache* UniformExpressionCache = &MaterialRenderProxy->UniformExpressionCache[FeatureLevel];
 
-		SetParameters(ShaderRHI, View);
+		SetParameters(RHICmdList, ShaderRHI, View);
 
 		// If the material has cached uniform expressions for selection or hover
 		// and that is being overridden by show flags in the editor, recache
@@ -197,7 +203,7 @@ public:
 #endif
 
 		// Set the material uniform buffer.
-		SetUniformBufferParameter(ShaderRHI,MaterialUniformBuffer,UniformExpressionCache->UniformBuffer);
+		SetUniformBufferParameter(RHICmdList, ShaderRHI,MaterialUniformBuffer,UniformExpressionCache->UniformBuffer);
 
 		{
 			const TArray<FGuid>& ParameterCollections = UniformExpressionCache->ParameterCollections;
@@ -209,7 +215,7 @@ public:
 			for (int32 CollectionIndex = 0; CollectionIndex < ParameterCollectionsNum; CollectionIndex++)
 			{
 				FUniformBufferRHIParamRef UniformBuffer = GetParameterCollectionBuffer(ParameterCollections[CollectionIndex], View.Family->Scene);
-				SetUniformBufferParameter(ShaderRHI,ParameterCollectionUniformBuffers[CollectionIndex],UniformBuffer);
+				SetUniformBufferParameter(RHICmdList, ShaderRHI,ParameterCollectionUniformBuffers[CollectionIndex],UniformBuffer);
 			}
 		}
 
@@ -234,6 +240,7 @@ public:
 				Value->LastRenderTime = FApp::GetCurrentTime();
 
 				SetTextureParameter(
+					RHICmdList, 
 					ShaderRHI, 
 					TextureResourceParameter.ShaderParameter, 
 					SamplerResourceParameter.ShaderParameter, 
@@ -263,6 +270,7 @@ public:
 				Value->LastRenderTime = FApp::GetCurrentTime();
 
 				SetTextureParameter(
+					RHICmdList, 
 					ShaderRHI,
 					TextureResourceParameter.ShaderParameter,
 					SamplerResourceParameter.ShaderParameter,
@@ -271,34 +279,35 @@ public:
 			}
 		}
 
-		DeferredParameters.Set(ShaderRHI, View, TextureMode);
+		DeferredParameters.Set(RHICmdList, ShaderRHI, View, TextureMode);
 
-		AtmosphericFogTextureParameters.Set(ShaderRHI, View);
+		AtmosphericFogTextureParameters.Set(RHICmdList, ShaderRHI, View);
 
 		if (FeatureLevel >= ERHIFeatureLevel::SM3)
 		{
 			if(LightAttenuation.IsBound())
-			{
-				SetTextureParameter(
-					ShaderRHI,
-					LightAttenuation,
-					LightAttenuationSampler,
-					TStaticSamplerState<SF_Bilinear,AM_Clamp,AM_Clamp,AM_Clamp>::GetRHI(),
-					GSceneRenderTargets.GetLightAttenuationTexture());
-			}
+		{
+			SetTextureParameter(
+				RHICmdList,
+				ShaderRHI,
+				LightAttenuation,
+				LightAttenuationSampler,
+				TStaticSamplerState<SF_Bilinear,AM_Clamp,AM_Clamp,AM_Clamp>::GetRHI(),
+				GSceneRenderTargets.GetLightAttenuationTexture());
+		}
 		}
 
 		// if we are in a postprocessing pass
 		if(View.RenderingCompositePassContext)
 		{
-			PostprocessParameter.Set(ShaderRHI, *View.RenderingCompositePassContext, TStaticSamplerState<>::GetRHI());
+			PostprocessParameter.Set(RHICmdList, ShaderRHI, *View.RenderingCompositePassContext, TStaticSamplerState<>::GetRHI());
 		}
 
 		//Use of the eye adaptation texture here is experimental and potentially dangerous as it can introduce a feedback loop. May be removed.
 		if(EyeAdaptation.IsBound())
 		{
 			FTextureRHIRef& EyeAdaptationTex = GetEyeAdaptation(View);
-			SetTextureParameter(ShaderRHI, EyeAdaptation, EyeAdaptationTex);
+			SetTextureParameter(RHICmdList, ShaderRHI, EyeAdaptation, EyeAdaptationTex);
 		}
 
 		if (PerlinNoiseGradientTexture.IsBound() && IsValidRef(GSystemTextures.PerlinNoiseGradient))
@@ -306,6 +315,7 @@ public:
 			const FTexture2DRHIRef& Texture = (FTexture2DRHIRef&)GSystemTextures.PerlinNoiseGradient->GetRenderTargetItem().ShaderResourceTexture;
 			// Bind the PerlinNoiseGradientTexture as a texture
 			SetTextureParameter(
+				RHICmdList,
 				ShaderRHI,
 				PerlinNoiseGradientTexture,
 				PerlinNoiseGradientTextureSampler,
@@ -319,6 +329,7 @@ public:
 			const FTexture3DRHIRef& Texture = (FTexture3DRHIRef&)GSystemTextures.PerlinNoise3D->GetRenderTargetItem().ShaderResourceTexture;
 			// Bind the PerlinNoise3DTexture as a texture
 			SetTextureParameter(
+				RHICmdList,
 				ShaderRHI,
 				PerlinNoise3DTexture,
 				PerlinNoise3DTextureSampler,
@@ -381,36 +392,36 @@ public:
 	{}
 
 	template<typename ShaderRHIParamRef>
-	void SetParameters(const ShaderRHIParamRef ShaderRHI,const FMaterialRenderProxy* MaterialRenderProxy,const FMaterial& Material,const FSceneView& View,ESceneRenderTargetsMode::Type TextureMode)
+	void SetParameters(FRHICommandList* RHICmdList, const ShaderRHIParamRef ShaderRHI,const FMaterialRenderProxy* MaterialRenderProxy,const FMaterial& Material,const FSceneView& View,ESceneRenderTargetsMode::Type TextureMode)
 	{
-		FMaterialShader::SetParameters(ShaderRHI,MaterialRenderProxy,Material,View,false,TextureMode);
+		FMaterialShader::SetParameters(RHICmdList, ShaderRHI,MaterialRenderProxy,Material,View,false,TextureMode);
 	}
 
-	void SetVFParametersOnly(const FVertexFactory* VertexFactory,const FSceneView& View,const FMeshBatchElement& BatchElement)
+	void SetVFParametersOnly(FRHICommandList* RHICmdList, const FVertexFactory* VertexFactory,const FSceneView& View,const FMeshBatchElement& BatchElement)
 	{
-		VertexFactoryParameters.SetMesh(this,VertexFactory,View,BatchElement, 0);
+		VertexFactoryParameters.SetMesh(RHICmdList, this,VertexFactory,View,BatchElement, 0);
 	}
 
 	template<typename ShaderRHIParamRef>
-	void SetMesh(const ShaderRHIParamRef ShaderRHI,const FVertexFactory* VertexFactory,const FSceneView& View,const FPrimitiveSceneProxy* Proxy,const FMeshBatchElement& BatchElement,uint32 DataFlags = 0)
+	void SetMesh(FRHICommandList* RHICmdList, const ShaderRHIParamRef ShaderRHI,const FVertexFactory* VertexFactory,const FSceneView& View,const FPrimitiveSceneProxy* Proxy,const FMeshBatchElement& BatchElement,uint32 DataFlags = 0)
 	{
 		// Set the mesh for the vertex factory
-		VertexFactoryParameters.SetMesh(this,VertexFactory,View,BatchElement, DataFlags);
+		VertexFactoryParameters.SetMesh(RHICmdList, this,VertexFactory,View,BatchElement, DataFlags);
 		
 		if(IsValidRef(BatchElement.PrimitiveUniformBuffer))
 		{
-			SetUniformBufferParameter(ShaderRHI,GetUniformBufferParameter<FPrimitiveUniformShaderParameters>(),BatchElement.PrimitiveUniformBuffer);
+			SetUniformBufferParameter(RHICmdList, ShaderRHI,GetUniformBufferParameter<FPrimitiveUniformShaderParameters>(),BatchElement.PrimitiveUniformBuffer);
 		}
 		else
 		{
 			check(BatchElement.PrimitiveUniformBufferResource);
-			SetUniformBufferParameter(ShaderRHI,GetUniformBufferParameter<FPrimitiveUniformShaderParameters>(),*BatchElement.PrimitiveUniformBufferResource);
+			SetUniformBufferParameter(RHICmdList, ShaderRHI,GetUniformBufferParameter<FPrimitiveUniformShaderParameters>(),*BatchElement.PrimitiveUniformBufferResource);
 		}
 
 		TShaderUniformBufferParameter<FDistanceCullFadeUniformShaderParameters> LODParameter = GetUniformBufferParameter<FDistanceCullFadeUniformShaderParameters>();
 		if( LODParameter.IsBound() )
 		{
-			SetUniformBufferParameter( ShaderRHI,LODParameter,GetPrimitiveFadeUniformBufferParameter(View, Proxy));
+			SetUniformBufferParameter(RHICmdList, ShaderRHI,LODParameter,GetPrimitiveFadeUniformBufferParameter(View, Proxy));
 		}
 	}
 
@@ -465,14 +476,14 @@ public:
 
 	FBaseHS() {}
 
-	void SetParameters(const FMaterialRenderProxy* MaterialRenderProxy,const FSceneView& View)
+	void SetParameters(FRHICommandList* RHICmdList, const FMaterialRenderProxy* MaterialRenderProxy,const FSceneView& View)
 	{
-		FMeshMaterialShader::SetParameters(GetHullShader(), MaterialRenderProxy, *MaterialRenderProxy->GetMaterial(View.GetFeatureLevel()), View, ESceneRenderTargetsMode::SetTextures);
+		FMeshMaterialShader::SetParameters(RHICmdList, GetHullShader(), MaterialRenderProxy, *MaterialRenderProxy->GetMaterial(View.GetFeatureLevel()), View, ESceneRenderTargetsMode::SetTextures);
 	}
 
-	void SetMesh(const FVertexFactory* VertexFactory,const FSceneView& View,const FPrimitiveSceneProxy* Proxy,const FMeshBatchElement& BatchElement)
+	void SetMesh(FRHICommandList* RHICmdList, const FVertexFactory* VertexFactory,const FSceneView& View,const FPrimitiveSceneProxy* Proxy,const FMeshBatchElement& BatchElement)
 	{
-		FMeshMaterialShader::SetMesh(GetHullShader(),VertexFactory,View,Proxy,BatchElement);
+		FMeshMaterialShader::SetMesh(RHICmdList, GetHullShader(),VertexFactory,View,Proxy,BatchElement);
 	}
 };
 
@@ -511,14 +522,14 @@ public:
 
 	FBaseDS() {}
 
-	void SetParameters(const FMaterialRenderProxy* MaterialRenderProxy,const FSceneView& View)
+	void SetParameters(FRHICommandList* RHICmdList, const FMaterialRenderProxy* MaterialRenderProxy,const FSceneView& View)
 	{
-		FMeshMaterialShader::SetParameters(GetDomainShader(), MaterialRenderProxy, *MaterialRenderProxy->GetMaterial(View.GetFeatureLevel()), View, ESceneRenderTargetsMode::SetTextures);
+		FMeshMaterialShader::SetParameters(RHICmdList, GetDomainShader(), MaterialRenderProxy, *MaterialRenderProxy->GetMaterial(View.GetFeatureLevel()), View, ESceneRenderTargetsMode::SetTextures);
 	}
 
-	void SetMesh(const FVertexFactory* VertexFactory,const FSceneView& View,const FPrimitiveSceneProxy* Proxy,const FMeshBatchElement& BatchElement)
+	void SetMesh(FRHICommandList* RHICmdList, const FVertexFactory* VertexFactory,const FSceneView& View,const FPrimitiveSceneProxy* Proxy,const FMeshBatchElement& BatchElement)
 	{
-		FMeshMaterialShader::SetMesh(GetDomainShader(),VertexFactory,View,Proxy,BatchElement);
+		FMeshMaterialShader::SetMesh(RHICmdList, GetDomainShader(),VertexFactory,View,Proxy,BatchElement);
 	}
 };
 
diff --git a/Engine/Source/Runtime/Renderer/Private/ShaderComplexityRendering.cpp b/Engine/Source/Runtime/Renderer/Private/ShaderComplexityRendering.cpp
--- a/Engine/Source/Runtime/Renderer/Private/ShaderComplexityRendering.cpp
+++ b/Engine/Source/Runtime/Renderer/Private/ShaderComplexityRendering.cpp
@@ -18,6 +18,7 @@ float GetMaxShaderComplexityCount(ERHIFeatureLevel::Type InFeatureType)
 }
 
 void FShaderComplexityAccumulatePS::SetParameters(
+	FRHICommandList* RHICmdList, 
 	uint32 NumVertexInstructions, 
 	uint32 NumPixelInstructions,
 	ERHIFeatureLevel::Type InFeatureLevel)
@@ -25,7 +26,7 @@ void FShaderComplexityAccumulatePS::SetParameters(
 	//normalize the complexity so we can fit it in a low precision scene color which is necessary on some platforms
 	const float NormalizedComplexityValue = float(NumPixelInstructions) / GetMaxShaderComplexityCount(InFeatureLevel);
 
-	SetShaderValue( GetPixelShader(), NormalizedComplexity, NormalizedComplexityValue);
+	SetShaderValue(RHICmdList, GetPixelShader(), NormalizedComplexity, NormalizedComplexityValue);
 }
 
 IMPLEMENT_SHADER_TYPE(,FShaderComplexityAccumulatePS,TEXT("ShaderComplexityAccumulatePixelShader"),TEXT("Main"),SF_Pixel);
@@ -87,15 +88,16 @@ public:
 	}
 
 	virtual void SetParameters(
+		FRHICommandList* RHICmdList, 
 		const FRenderingCompositePassContext& Context,
 		const TArray<FLinearColor>& Colors
 		)
 	{
 		const FPixelShaderRHIParamRef ShaderRHI = GetPixelShader();
 
-		FGlobalShader::SetParameters(ShaderRHI, Context.View);
+		FGlobalShader::SetParameters(RHICmdList, ShaderRHI, Context.View);
 
-		PostprocessParameter.SetPS(ShaderRHI, Context, TStaticSamplerState<SF_Bilinear,AM_Clamp,AM_Clamp,AM_Clamp>::GetRHI());
+		PostprocessParameter.SetPS(RHICmdList, ShaderRHI, Context, TStaticSamplerState<SF_Bilinear,AM_Clamp,AM_Clamp,AM_Clamp>::GetRHI());
 
 		//Make sure there are at least NumShaderComplexityColors colors specified in the ini.
 		//If there are more than NumShaderComplexityColors they will be ignored.
@@ -104,7 +106,7 @@ public:
 		//pass the complexity -> color mapping into the pixel shader
 		for(int32 ColorIndex = 0; ColorIndex < NumShaderComplexityColors; ColorIndex ++)
 		{
-			SetShaderValue(ShaderRHI, ShaderComplexityColors, Colors[ColorIndex], ColorIndex);
+			SetShaderValue(RHICmdList, ShaderRHI, ShaderComplexityColors, Colors[ColorIndex], ColorIndex);
 		}
 	}
 
@@ -166,6 +168,9 @@ void FRCPassPostProcessVisualizeComplexity::Process(FRenderingCompositePassConte
 
 	const FSceneRenderTargetItem& DestRenderTarget = PassOutputs[0].RequestSurface(Context);
 
+	//@todo-rco: RHIPacketList
+	FRHICommandList* RHICmdList = nullptr;
+
 	// Set the view family's render target/viewport.
 	RHISetRenderTarget(DestRenderTarget.TargetableTexture, FTextureRHIRef());	
 	Context.SetViewportAndCallRHI(DestRect);
@@ -184,7 +189,7 @@ void FRCPassPostProcessVisualizeComplexity::Process(FRenderingCompositePassConte
 	static FGlobalBoundShaderState ShaderComplexityBoundShaderState;
 	SetGlobalBoundShaderState(ShaderComplexityBoundShaderState, GFilterVertexDeclaration.VertexDeclarationRHI, *VertexShader, *PixelShader);
 
-	PixelShader->SetParameters(Context, Colors);
+	PixelShader->SetParameters(RHICmdList, Context, Colors);
 	
 	DrawRectangle(
 		0, 0,
diff --git a/Engine/Source/Runtime/Renderer/Private/ShaderComplexityRendering.h b/Engine/Source/Runtime/Renderer/Private/ShaderComplexityRendering.h
--- a/Engine/Source/Runtime/Renderer/Private/ShaderComplexityRendering.h
+++ b/Engine/Source/Runtime/Renderer/Private/ShaderComplexityRendering.h
@@ -28,7 +28,8 @@ public:
 	FShaderComplexityAccumulatePS() {}
 
 	void SetParameters(
-		uint32 NumVertexInstructions,
+		FRHICommandList* RHICmdList, 
+		uint32 NumVertexInstructions, 
 		uint32 NumPixelInstructions,
 		ERHIFeatureLevel::Type InFeatureLevel);
 
diff --git a/Engine/Source/Runtime/Renderer/Private/ShadowRendering.cpp b/Engine/Source/Runtime/Renderer/Private/ShadowRendering.cpp
--- a/Engine/Source/Runtime/Renderer/Private/ShadowRendering.cpp
+++ b/Engine/Source/Runtime/Renderer/Private/ShadowRendering.cpp
@@ -103,19 +103,20 @@ public:
 	}
 
 	void SetParameters(
+		FRHICommandList* RHICmdList, 
 		const FMaterialRenderProxy* MaterialRenderProxy,
 		const FMaterial& Material,
 		const FSceneView& View,
 		const FProjectedShadowInfo* ShadowInfo
 		)
 	{
-		FMeshMaterialShader::SetParameters(GetVertexShader(),MaterialRenderProxy,Material,View,ESceneRenderTargetsMode::DontSet);
-		ShadowParameters.SetVertexShader(this, View, ShadowInfo, MaterialRenderProxy);
+		FMeshMaterialShader::SetParameters(RHICmdList, GetVertexShader(),MaterialRenderProxy,Material,View,ESceneRenderTargetsMode::DontSet);
+		ShadowParameters.SetVertexShader(RHICmdList, this, View, ShadowInfo, MaterialRenderProxy);
 	}
 
-	void SetMesh(const FVertexFactory* VertexFactory,const FSceneView& View,const FPrimitiveSceneProxy* Proxy,const FMeshBatchElement& BatchElement)
+	void SetMesh(FRHICommandList* RHICmdList, const FVertexFactory* VertexFactory,const FSceneView& View,const FPrimitiveSceneProxy* Proxy,const FMeshBatchElement& BatchElement)
 	{
-		FMeshMaterialShader::SetMesh(GetVertexShader(),VertexFactory,View,Proxy,BatchElement);
+		FMeshMaterialShader::SetMesh(RHICmdList, GetVertexShader(),VertexFactory,View,Proxy,BatchElement);
 	}
 
 private:
@@ -238,13 +239,14 @@ public:
 	}
 
 	void SetParameters(
+		FRHICommandList* RHICmdList, 
 		const FMaterialRenderProxy* MaterialRenderProxy,
 		const FSceneView& View,
 		const FProjectedShadowInfo* ShadowInfo
 		)
 	{
-		FBaseDS::SetParameters(MaterialRenderProxy, View);
-		ShadowParameters.SetDomainShader(this, View, ShadowInfo, MaterialRenderProxy);
+		FBaseDS::SetParameters(RHICmdList, MaterialRenderProxy, View);
+		ShadowParameters.SetDomainShader(RHICmdList, this, View, ShadowInfo, MaterialRenderProxy);
 	}
 
 private:
@@ -315,11 +317,12 @@ public:
 	}
 
 	void SetParameters(
+		FRHICommandList* RHICmdList, 
 		const FSceneView& View,
 		const FProjectedShadowInfo* ShadowInfo
 		)
 	{
-		FMaterialShader::SetParameters(GetGeometryShader(),View);
+		FMaterialShader::SetParameters(RHICmdList, GetGeometryShader(),View);
 
 		const FMatrix Translation = FTranslationMatrix(-View.ViewMatrices.PreViewTranslation);
 
@@ -332,6 +335,7 @@ public:
 
 		// Set the view projection matrices that will transform positions from world to cube map face space
 		SetShaderValueArray<FGeometryShaderRHIParamRef, FMatrix>(
+			RHICmdList, 
 			GetGeometryShader(),
 			ShadowViewProjectionMatrices,
 			TranslatedShadowViewProjectionMatrices,
@@ -339,7 +343,7 @@ public:
 			);
 	}
 
-	void SetMesh(const FPrimitiveSceneProxy* PrimitiveSceneProxy, const FProjectedShadowInfo* ShadowInfo, const FSceneView& View)
+	void SetMesh(FRHICommandList* RHICmdList, const FPrimitiveSceneProxy* PrimitiveSceneProxy, const FProjectedShadowInfo* ShadowInfo, const FSceneView& View)
 	{
 		if (MeshVisibleToFace.IsBound())
 		{
@@ -353,6 +357,7 @@ public:
 
 			// Set the view projection matrices that will transform positions from world to cube map face space
 			SetShaderValueArray<FGeometryShaderRHIParamRef, FVector4>(
+				RHICmdList, 
 				GetGeometryShader(),
 				MeshVisibleToFace,
 				MeshVisibleToFaceValue,
@@ -419,6 +424,7 @@ public:
 	TShadowDepthBasePS() {}
 
 	void SetParameters(
+		FRHICommandList* RHICmdList, 
 		const FMaterialRenderProxy* MaterialRenderProxy,
 		const FMaterial& Material,
 		const FSceneView& View,
@@ -427,16 +433,17 @@ public:
 	{
 		const FPixelShaderRHIParamRef ShaderRHI = GetPixelShader();
 
-		FMeshMaterialShader::SetParameters(ShaderRHI, MaterialRenderProxy, Material, View, ESceneRenderTargetsMode::DontSet);
+		FMeshMaterialShader::SetParameters(RHICmdList, ShaderRHI, MaterialRenderProxy, Material, View, ESceneRenderTargetsMode::DontSet);
 
-		SetShaderValue(ShaderRHI, ShadowParams, FVector2D(ShadowInfo->GetShaderDepthBias(), ShadowInfo->InvMaxSubjectDepth));
+		SetShaderValue(RHICmdList, ShaderRHI, ShadowParams, FVector2D(ShadowInfo->GetShaderDepthBias(), ShadowInfo->InvMaxSubjectDepth));
 
 		if(bRenderReflectiveShadowMap)
 		{
 			// LPV also propagates light transmission (for transmissive materials)
-			SetShaderValue(ShaderRHI,TransmissionStrength,View.FinalPostProcessSettings.LPVTransmissionIntensity);
-			SetShaderValue(ShaderRHI,ReflectiveShadowMapTextureResolution,FVector2D(GSceneRenderTargets.GetReflectiveShadowMapTextureResolution()));
+			SetShaderValue(RHICmdList, ShaderRHI,TransmissionStrength,View.FinalPostProcessSettings.LPVTransmissionIntensity);
+			SetShaderValue(RHICmdList, ShaderRHI,ReflectiveShadowMapTextureResolution,FVector2D(GSceneRenderTargets.GetReflectiveShadowMapTextureResolution()));
 			SetShaderValue(
+				RHICmdList, 
 				ShaderRHI,
 				ProjectionMatrixParameter,
 				FTranslationMatrix(ShadowInfo->PreShadowTranslation - View.ViewMatrices.PreViewTranslation) * ShadowInfo->SubjectAndReceiverMatrix
@@ -449,15 +456,15 @@ public:
 
 				if(Lpv)
 				{
-					SetUniformBufferParameter(ShaderRHI, GetUniformBufferParameter<FLpvWriteUniformBufferParameters>(), Lpv->GetWriteUniformBuffer());
+					SetUniformBufferParameter(RHICmdList, ShaderRHI, GetUniformBufferParameter<FLpvWriteUniformBufferParameters>(), Lpv->GetWriteUniformBuffer());
 				}
 			}
 		}
 	}
 
-	void SetMesh(const FVertexFactory* VertexFactory,const FSceneView& View,const FPrimitiveSceneProxy* Proxy,const FMeshBatchElement& BatchElement)
+	void SetMesh(FRHICommandList* RHICmdList, const FVertexFactory* VertexFactory,const FSceneView& View,const FPrimitiveSceneProxy* Proxy,const FMeshBatchElement& BatchElement)
 	{
-		FMeshMaterialShader::SetMesh(GetPixelShader(),VertexFactory,View,Proxy,BatchElement);
+		FMeshMaterialShader::SetMesh(RHICmdList, GetPixelShader(),VertexFactory,View,Proxy,BatchElement);
 	}
 
 	virtual bool Serialize(FArchive& Ar)
@@ -585,25 +592,25 @@ bool FShadowProjectionVS::ShouldCache(EShaderPlatform Platform)
 	return IsFeatureLevelSupported(Platform, ERHIFeatureLevel::SM3);
 }
 
-void FShadowProjectionVS::SetParameters(const FSceneView& View, const FProjectedShadowInfo* ShadowInfo, const int32 ShadowSplitIndex)
+void FShadowProjectionVS::SetParameters(FRHICommandList* RHICmdList, const FSceneView& View, const FProjectedShadowInfo* ShadowInfo, const int32 ShadowSplitIndex)
 {
-	FGlobalShader::SetParameters(GetVertexShader(),View);
+	FGlobalShader::SetParameters(RHICmdList, GetVertexShader(),View);
 	
 	if(ShadowInfo->IsWholeSceneDirectionalShadow())
 	{
 		// Calculate bounding geometry transform for whole scene directional shadow.
 		// Use a pair of pre-transformed planes for stenciling.
-		StencilingGeometryParameters.Set(this, FVector4(0,0,0,1));
+		StencilingGeometryParameters.Set(RHICmdList, this, FVector4(0,0,0,1));
 	}
 	else if(ShadowInfo->IsWholeScenePointLightShadow())
 	{
 		// Handle stenciling sphere for point light.
-		StencilingGeometryParameters.Set(this, View, ShadowInfo->LightSceneInfo);
+		StencilingGeometryParameters.Set(RHICmdList, this, View, ShadowInfo->LightSceneInfo);
 	}
 	else
 	{
 		// Other bounding geometry types are pre-transformed.
-		StencilingGeometryParameters.Set(this, FVector4(0,0,0,1));
+		StencilingGeometryParameters.Set(RHICmdList, this, FVector4(0,0,0,1));
 	}
 }
 
@@ -839,29 +846,30 @@ FShadowDepthDrawingPolicy<bRenderingReflectiveShadowMaps>::FShadowDepthDrawingPo
 }
 
 template <bool bRenderingReflectiveShadowMaps>
-void FShadowDepthDrawingPolicy<bRenderingReflectiveShadowMaps>::DrawShared(const FSceneView* View,FBoundShaderStateRHIParamRef BoundShaderState) const
+void FShadowDepthDrawingPolicy<bRenderingReflectiveShadowMaps>::DrawShared(FRHICommandList* RHICmdList, const FSceneView* View,FBoundShaderStateRHIParamRef BoundShaderState) const
 {
 	checkSlow(bDirectionalLight == PolicyShadowInfo->bDirectionalLight && bPreShadow == PolicyShadowInfo->bPreShadow);
 
+	check(!RHICmdList);
 	// Set the actual shader & vertex declaration state
 	RHISetBoundShaderState(BoundShaderState);
 
-	VertexShader->SetParameters(MaterialRenderProxy,*MaterialResource,*View,PolicyShadowInfo);
+	VertexShader->SetParameters(RHICmdList, MaterialRenderProxy,*MaterialResource,*View,PolicyShadowInfo);
 
 	if (GeometryShader)
 	{
-		GeometryShader->SetParameters(*View,PolicyShadowInfo);
+		GeometryShader->SetParameters(RHICmdList, *View,PolicyShadowInfo);
 	}
 
 	if(HullShader && DomainShader)
 	{
-		HullShader->SetParameters(MaterialRenderProxy,*View);
-		DomainShader->SetParameters(MaterialRenderProxy,*View,PolicyShadowInfo);
+		HullShader->SetParameters(RHICmdList, MaterialRenderProxy,*View);
+		DomainShader->SetParameters(RHICmdList, MaterialRenderProxy,*View,PolicyShadowInfo);
 	}
 
 	if (PixelShader)
 	{
-		PixelShader->SetParameters(MaterialRenderProxy,*MaterialResource,*View,PolicyShadowInfo);
+		PixelShader->SetParameters(RHICmdList, MaterialRenderProxy,*MaterialResource,*View,PolicyShadowInfo);
 	}
 
 	// Set the shared mesh resources.
@@ -871,7 +879,7 @@ void FShadowDepthDrawingPolicy<bRenderingReflectiveShadowMaps>::DrawShared(const
 	}
 	else
 	{
-		FMeshDrawingPolicy::DrawShared(View);
+		FMeshDrawingPolicy::DrawShared(RHICmdList, View);
 	}
 
 	// @TODO: only render directional light shadows as two sided, and only when blocking is enabled (required by geometry volume injection)
@@ -919,6 +927,7 @@ FBoundShaderStateRHIRef FShadowDepthDrawingPolicy<bRenderingReflectiveShadowMaps
 
 template <bool bRenderingReflectiveShadowMaps>
 void FShadowDepthDrawingPolicy<bRenderingReflectiveShadowMaps>::SetMeshRenderState(
+	FRHICommandList* RHICmdList, 
 	const FSceneView& View,
 	const FPrimitiveSceneProxy* PrimitiveSceneProxy,
 	const FMeshBatch& Mesh,
@@ -931,20 +940,20 @@ void FShadowDepthDrawingPolicy<bRenderingReflectiveShadowMaps>::SetMeshRenderSta
 
 	EmitMeshDrawEvents(PrimitiveSceneProxy, Mesh);
 
-	VertexShader->SetMesh(VertexFactory,View,PrimitiveSceneProxy,BatchElement);
+	VertexShader->SetMesh(RHICmdList, VertexFactory,View,PrimitiveSceneProxy,BatchElement);
 
 	if( HullShader && DomainShader )
 	{
-		HullShader->SetMesh(VertexFactory,View,PrimitiveSceneProxy,BatchElement);
-		DomainShader->SetMesh(VertexFactory,View,PrimitiveSceneProxy,BatchElement);
+		HullShader->SetMesh(RHICmdList, VertexFactory,View,PrimitiveSceneProxy,BatchElement);
+		DomainShader->SetMesh(RHICmdList, VertexFactory,View,PrimitiveSceneProxy,BatchElement);
 	}
 	if (GeometryShader)
 	{
-		GeometryShader->SetMesh(PrimitiveSceneProxy, PolicyShadowInfo, View);
+		GeometryShader->SetMesh(RHICmdList, PrimitiveSceneProxy, PolicyShadowInfo, View);
 	}
 	if (PixelShader)
 	{
-		PixelShader->SetMesh(VertexFactory,View,PrimitiveSceneProxy,BatchElement);
+		PixelShader->SetMesh(RHICmdList, VertexFactory,View,PrimitiveSceneProxy,BatchElement);
 	}
 	// Not calling FMeshDrawingPolicy::SetMeshRenderState as DrawShared sets the rasterizer state
 }
@@ -1026,6 +1035,7 @@ void FShadowDepthDrawingPolicyFactory::AddStaticMesh(FScene* Scene,FStaticMesh*
 }
 
 bool FShadowDepthDrawingPolicyFactory::DrawDynamicMesh(
+	FRHICommandList* RHICmdList, 
 	const FSceneView& View,
 	ContextType Context,
 	const FMeshBatch& Mesh,
@@ -1069,11 +1079,11 @@ bool FShadowDepthDrawingPolicyFactory::DrawDynamicMesh(
 					Mesh.ReverseCulling
 					);
 
-				DrawingPolicy.DrawShared(&View,DrawingPolicy.CreateBoundShaderState(View.GetFeatureLevel()));
+				DrawingPolicy.DrawShared(RHICmdList, &View,DrawingPolicy.CreateBoundShaderState(View.GetFeatureLevel()));
 				for( int32 BatchElementIndex = 0, Num = Mesh.Elements.Num(); BatchElementIndex < Num; BatchElementIndex++)
 				{
-					DrawingPolicy.SetMeshRenderState(View,PrimitiveSceneProxy,Mesh,BatchElementIndex,bBackFace,FMeshDrawingPolicy::ElementDataType());
-					DrawingPolicy.DrawMesh(Mesh,BatchElementIndex);
+					DrawingPolicy.SetMeshRenderState(RHICmdList, View,PrimitiveSceneProxy,Mesh,BatchElementIndex,bBackFace,FMeshDrawingPolicy::ElementDataType());
+					DrawingPolicy.DrawMesh(RHICmdList, Mesh,BatchElementIndex);
 				}
 			}
 			else
@@ -1089,11 +1099,11 @@ bool FShadowDepthDrawingPolicyFactory::DrawDynamicMesh(
 					Mesh.ReverseCulling
 					);
 
-				DrawingPolicy.DrawShared(&View, DrawingPolicy.CreateBoundShaderState(View.GetFeatureLevel()));
+				DrawingPolicy.DrawShared(RHICmdList, &View, DrawingPolicy.CreateBoundShaderState(View.GetFeatureLevel()));
 				for( int32 BatchElementIndex = 0; BatchElementIndex < Mesh.Elements.Num(); BatchElementIndex++)
 				{
-					DrawingPolicy.SetMeshRenderState(View,PrimitiveSceneProxy,Mesh,BatchElementIndex,bBackFace,FMeshDrawingPolicy::ElementDataType());
-					DrawingPolicy.DrawMesh(Mesh,BatchElementIndex);
+					DrawingPolicy.SetMeshRenderState(RHICmdList, View,PrimitiveSceneProxy,Mesh,BatchElementIndex,bBackFace,FMeshDrawingPolicy::ElementDataType());
+					DrawingPolicy.DrawMesh(RHICmdList, Mesh,BatchElementIndex);
 				}
 			}
 
@@ -1197,6 +1207,9 @@ void DrawMeshElements(FShadowDepthDrawingPolicy<bReflectiveShadowmap>& SharedDra
 	checkSlow(CompareDrawingPolicy(DebugPolicy, SharedDrawingPolicy) == 0);
 #endif
 
+	//@todo-rco: RHIPacketList
+	FRHICommandList* RHICmdList = nullptr;
+
 	// Render only those batch elements that match the current LOD
 	uint64 BatchElementMask = Mesh->Elements.Num() == 1 ? 1 : View.StaticMeshBatchVisibility[Mesh->Id];
 	int32 BatchElementIndex = 0;
@@ -1204,8 +1217,8 @@ void DrawMeshElements(FShadowDepthDrawingPolicy<bReflectiveShadowmap>& SharedDra
 	{
 		if(BatchElementMask & 1)
 		{
-			SharedDrawingPolicy.SetMeshRenderState(View, Mesh->PrimitiveSceneInfo->Proxy, *Mesh, BatchElementIndex, false, FMeshDrawingPolicy::ElementDataType());
-			SharedDrawingPolicy.DrawMesh(*Mesh, BatchElementIndex);
+			SharedDrawingPolicy.SetMeshRenderState(RHICmdList, View, Mesh->PrimitiveSceneInfo->Proxy, *Mesh, BatchElementIndex, false, FMeshDrawingPolicy::ElementDataType());
+			SharedDrawingPolicy.DrawMesh(RHICmdList, *Mesh, BatchElementIndex);
 			INC_DWORD_STAT(STAT_ShadowDynamicPathDrawCalls);
 		}
 
@@ -1215,7 +1228,7 @@ void DrawMeshElements(FShadowDepthDrawingPolicy<bReflectiveShadowmap>& SharedDra
 }
 
 template <bool bReflectiveShadowmap>
-void DrawShadowMeshElements(const FViewInfo& View, const FProjectedShadowInfo& ShadowInfo)
+void DrawShadowMeshElements(FRHICommandList* RHICmdList, const FViewInfo& View, const FProjectedShadowInfo& ShadowInfo)
 {
 	const FShadowStaticMeshElement& FirstShadowMesh = ShadowInfo.SubjectMeshElements[0];
 	const FMaterial* FirstMaterialResource = FirstShadowMesh.MaterialResource;
@@ -1247,14 +1260,14 @@ void DrawShadowMeshElements(const FViewInfo& View, const FProjectedShadowInfo& S
 			OldState = CurrentState;
 
 			SharedDrawingPolicy.UpdateElementState(CurrentState);
-			SharedDrawingPolicy.DrawShared(&View, SharedDrawingPolicy.CreateBoundShaderState(View.GetFeatureLevel()));
+			SharedDrawingPolicy.DrawShared(RHICmdList, &View, SharedDrawingPolicy.CreateBoundShaderState(View.GetFeatureLevel()));
 		}
 
 		DrawMeshElements(SharedDrawingPolicy, OldState, View, ShadowMesh.Mesh);
 	}
 }
 
-void FProjectedShadowInfo::RenderDepth(FDeferredShadingSceneRenderer* SceneRenderer)
+void FProjectedShadowInfo::RenderDepth(FRHICommandList* RHICmdList, FDeferredShadingSceneRenderer* SceneRenderer)
 {
 #if WANTS_DRAW_MESH_EVENTS
 	FString EventName;
@@ -1357,16 +1370,19 @@ void FProjectedShadowInfo::RenderDepth(FDeferredShadingSceneRenderer* SceneRende
 	// Draw the subject's static elements using static draw lists
 	if (IsWholeSceneDirectionalShadow())
 	{
+		//@todo-rco: RHIPacketList
+		FRHICommandList* RHICmdList = nullptr;
+
 		SCOPE_CYCLE_COUNTER(STAT_WholeSceneStaticDrawListShadowDepthsTime);
 
 		if (bReflectiveShadowmap)
 		{
-			SceneRenderer->Scene->WholeSceneReflectiveShadowMapDrawList.DrawVisible(*FoundView, StaticMeshWholeSceneShadowDepthMap, StaticMeshWholeSceneShadowBatchVisibility);
+			SceneRenderer->Scene->WholeSceneReflectiveShadowMapDrawList.DrawVisible(RHICmdList, *FoundView, StaticMeshWholeSceneShadowDepthMap, StaticMeshWholeSceneShadowBatchVisibility);
 		}
 		else
 		{
 			// Use the scene's shadow depth draw list with this shadow's visibility map
-			SceneRenderer->Scene->WholeSceneShadowDepthDrawList.DrawVisible(*FoundView, StaticMeshWholeSceneShadowDepthMap, StaticMeshWholeSceneShadowBatchVisibility);
+			SceneRenderer->Scene->WholeSceneShadowDepthDrawList.DrawVisible(RHICmdList, *FoundView, StaticMeshWholeSceneShadowDepthMap, StaticMeshWholeSceneShadowBatchVisibility);
 		}
 	}
 	// Draw the subject's static elements using manual state filtering
@@ -1377,12 +1393,12 @@ void FProjectedShadowInfo::RenderDepth(FDeferredShadingSceneRenderer* SceneRende
 		if(bReflectiveShadowmap && !bOnePassPointLightShadow)
 		{
 			// reflective shadow map
-			DrawShadowMeshElements<true>(*FoundView, *this);
+			DrawShadowMeshElements<true>(RHICmdList, *FoundView, *this);
 		}
 		else
 		{
 			// normal shadow map
-			DrawShadowMeshElements<false>(*FoundView, *this);
+			DrawShadowMeshElements<false>(RHICmdList, *FoundView, *this);
 		}
 	}
 
@@ -1452,7 +1468,7 @@ void StencilingGeometry::DrawCone()
 static FGlobalBoundShaderState MaskBoundShaderState[2];
 
 template <uint32 Quality>
-static void SetShadowProjectionShaderTemplNew(int32 ViewIndex, const FViewInfo& View, const FProjectedShadowInfo* ShadowInfo)
+static void SetShadowProjectionShaderTemplNew(FRHICommandList* RHICmdList, int32 ViewIndex, const FViewInfo& View, const FProjectedShadowInfo* ShadowInfo)
 {
 	if (ShadowInfo->bTranslucentShadow)
 	{
@@ -1467,8 +1483,8 @@ static void SetShadowProjectionShaderTemplNew(int32 ViewIndex, const FViewInfo&
 		SetGlobalBoundShaderState(BoundShaderState, GetVertexDeclarationFVector4(), ShadowProjVS, ShadowProjPS );
 
 		// Set shader parameters
-		ShadowProjVS->SetParameters(View, ShadowInfo);
-		ShadowProjPS->SetParameters(ViewIndex, View, ShadowInfo);
+		ShadowProjVS->SetParameters(RHICmdList, View, ShadowInfo);
+		ShadowProjPS->SetParameters(RHICmdList, ViewIndex, View, ShadowInfo);
 	}
 	else if (ShadowInfo->IsWholeSceneDirectionalShadow())
 	{
@@ -1484,7 +1500,7 @@ static void SetShadowProjectionShaderTemplNew(int32 ViewIndex, const FViewInfo&
 			static FGlobalBoundShaderState BoundShaderState;
 			SetGlobalBoundShaderState(BoundShaderState, GetVertexDeclarationFVector4(), ShadowProjVS, ShadowProjPS );
 
-			ShadowProjPS->SetParameters(ViewIndex, View, ShadowInfo);
+			ShadowProjPS->SetParameters(RHICmdList, ViewIndex, View, ShadowInfo);
 		}
 		else
 		{
@@ -1494,10 +1510,10 @@ static void SetShadowProjectionShaderTemplNew(int32 ViewIndex, const FViewInfo&
 			static FGlobalBoundShaderState BoundShaderState;
 			SetGlobalBoundShaderState(BoundShaderState, GetVertexDeclarationFVector4(), ShadowProjVS, ShadowProjPS );
 
-			ShadowProjPS->SetParameters(ViewIndex, View, ShadowInfo);
+			ShadowProjPS->SetParameters(RHICmdList, ViewIndex, View, ShadowInfo);
 		}
 
-		ShadowProjVS->SetParameters(View);
+		ShadowProjVS->SetParameters(RHICmdList, View);
 	}
 	else
 	{
@@ -1511,13 +1527,15 @@ static void SetShadowProjectionShaderTemplNew(int32 ViewIndex, const FViewInfo&
 		static FGlobalBoundShaderState BoundShaderState;
 		SetGlobalBoundShaderState(BoundShaderState, GetVertexDeclarationFVector4(), ShadowProjVS, ShadowProjPS );
 
-		ShadowProjVS->SetParameters(View, ShadowInfo);
-		ShadowProjPS->SetParameters(ViewIndex, View, ShadowInfo);
+		ShadowProjVS->SetParameters(RHICmdList, View, ShadowInfo);
+		ShadowProjPS->SetParameters(RHICmdList, ViewIndex, View, ShadowInfo);
 	}
 }
 
-void FProjectedShadowInfo::RenderProjection(int32 ViewIndex, const FViewInfo* View) const
+void FProjectedShadowInfo::RenderProjection(FRHICommandList* RHICmdList, int32 ViewIndex, const FViewInfo* View) const
 {
+	check(!RHICmdList);
+
 #if WANTS_DRAW_MESH_EVENTS
 	FString EventName;
 	GetShadowTypeNameForDrawEvent(EventName);
@@ -1656,6 +1674,7 @@ void FProjectedShadowInfo::RenderProjection(int32 ViewIndex, const FViewInfo* Vi
 							if (View->StaticMeshVisibilityMap[StaticMesh.Id])
 							{
 								FDepthDrawingPolicyFactory::DrawStaticMesh(
+									RHICmdList, 
 									*View,
 									FDepthDrawingPolicyFactory::ContextType(DDM_AllOccluders),
 									StaticMesh,
@@ -1695,54 +1714,54 @@ void FProjectedShadowInfo::RenderProjection(int32 ViewIndex, const FViewInfo* Vi
 		}
 		else
 		{
-			// Increment stencil on front-facing zfail, decrement on back-facing zfail.
-			// Note, this is a reversed Z depth surface, using CF_GreaterEqual.
-			RHISetDepthStencilState(TStaticDepthStencilState<
-				false,CF_GreaterEqual,
-				true,CF_Always,SO_Keep,SO_Increment,SO_Keep,
-				true,CF_Always,SO_Keep,SO_Decrement,SO_Keep,
-				0xff,0xff
-				>::GetRHI());
-
-			RHISetRasterizerState(TStaticRasterizerState<FM_Solid,CM_None>::GetRHI());
+		// Increment stencil on front-facing zfail, decrement on back-facing zfail.
+		// Note, this is a reversed Z depth surface, using CF_GreaterEqual.
+		RHISetDepthStencilState(TStaticDepthStencilState<
+			false,CF_GreaterEqual,
+			true,CF_Always,SO_Keep,SO_Increment,SO_Keep,
+			true,CF_Always,SO_Keep,SO_Decrement,SO_Keep,
+			0xff,0xff
+			>::GetRHI());
 
-			checkSlow(SplitIndex >= 0);
-			checkSlow(bDirectionalLight);
+		RHISetRasterizerState(TStaticRasterizerState<FM_Solid,CM_None>::GetRHI());
 
-			// Draw 2 fullscreen planes, front facing one at the near subfrustum plane, and back facing one at the far.
+		checkSlow(SplitIndex >= 0);
+		checkSlow(bDirectionalLight);
 
-			// Find the projection shaders.
-			TShaderMapRef<FShadowProjectionNoTransformVS> VertexShaderNoTransform(GetGlobalShaderMap());
-			VertexShaderNoTransform->SetParameters(*View);
-			SetGlobalBoundShaderState(MaskBoundShaderState[0], GetVertexDeclarationFVector4(), *VertexShaderNoTransform, nullptr);
+		// Draw 2 fullscreen planes, front facing one at the near subfrustum plane, and back facing one at the far.
 
-			FVector4 Near = View->ViewMatrices.ProjMatrix.TransformFVector4(FVector4(0, 0, CascadeSettings.SplitNear));
-			FVector4 Far = View->ViewMatrices.ProjMatrix.TransformFVector4(FVector4(0, 0, CascadeSettings.SplitFar));
-			float StencilNear = Near.Z / Near.W;
-			float StencilFar = Far.Z / Far.W;
+		// Find the projection shaders.
+		TShaderMapRef<FShadowProjectionNoTransformVS> VertexShaderNoTransform(GetGlobalShaderMap());
+		VertexShaderNoTransform->SetParameters(RHICmdList, *View);
+		SetGlobalBoundShaderState(MaskBoundShaderState[0], GetVertexDeclarationFVector4(), *VertexShaderNoTransform, nullptr);
+
+		FVector4 Near = View->ViewMatrices.ProjMatrix.TransformFVector4(FVector4(0, 0, CascadeSettings.SplitNear));
+		FVector4 Far = View->ViewMatrices.ProjMatrix.TransformFVector4(FVector4(0, 0, CascadeSettings.SplitFar));
+		float StencilNear = Near.Z / Near.W;
+		float StencilFar = Far.Z / Far.W;
+
+		FVector4 Verts[] =
+		{
+			// Far Plane
+			FVector4( 1,  1,  StencilFar),
+			FVector4(-1,  1,  StencilFar),
+			FVector4( 1, -1,  StencilFar),
+			FVector4( 1, -1,  StencilFar),
+			FVector4(-1,  1,  StencilFar),
+			FVector4(-1, -1,  StencilFar),
+
+			// Near Plane
+			FVector4(-1,  1, StencilNear),
+			FVector4( 1,  1, StencilNear),
+			FVector4(-1, -1, StencilNear),
+			FVector4(-1, -1, StencilNear),
+			FVector4( 1,  1, StencilNear),
+			FVector4( 1, -1, StencilNear),
+		};
 
-			FVector4 Verts[] =
-			{
-				// Far Plane
-				FVector4( 1,  1,  StencilFar),
-				FVector4(-1,  1,  StencilFar),
-				FVector4( 1, -1,  StencilFar),
-				FVector4( 1, -1,  StencilFar),
-				FVector4(-1,  1,  StencilFar),
-				FVector4(-1, -1,  StencilFar),
-
-				// Near Plane
-				FVector4(-1,  1, StencilNear),
-				FVector4( 1,  1, StencilNear),
-				FVector4(-1, -1, StencilNear),
-				FVector4(-1, -1, StencilNear),
-				FVector4( 1,  1, StencilNear),
-				FVector4( 1, -1, StencilNear),
-			};
-
-			// Only draw the near plane if this is not the nearest split
-			RHIDrawPrimitiveUP(PT_TriangleList, (SplitIndex > 0) ? 4 : 2, Verts, sizeof(FVector4));
-		}
+		// Only draw the near plane if this is not the nearest split
+		RHIDrawPrimitiveUP(PT_TriangleList, (SplitIndex > 0) ? 4 : 2, Verts, sizeof(FVector4));
+	}
 	}
 	// Not a preshadow, mask the projection to any pixels inside the frustum.
 	else
@@ -1784,7 +1803,7 @@ void FProjectedShadowInfo::RenderProjection(int32 ViewIndex, const FViewInfo* Vi
 		SetGlobalBoundShaderState(MaskBoundShaderState[1],GetVertexDeclarationFVector4(),*VertexShader,NULL);
 
 		// Set the projection vertex shader parameters
-		VertexShader->SetParameters(*View, this);
+		VertexShader->SetParameters(RHICmdList, *View, this);
 
 		// Draw the frustum using the stencil buffer to mask just the pixels which are inside the shadow frustum.
 		RHIDrawIndexedPrimitiveUP(PT_TriangleList, 0, 8, 12, GCubeIndices, sizeof(uint16), FrustumVertices, sizeof(FVector4));
@@ -1802,12 +1821,12 @@ void FProjectedShadowInfo::RenderProjection(int32 ViewIndex, const FViewInfo* Vi
 	else
 	{
 		// no depth test or writes, Test stencil for non-zero.
-		RHISetDepthStencilState( TStaticDepthStencilState<
-			false, CF_Always,
-			true, CF_NotEqual, SO_Keep, SO_Keep, SO_Keep,
-			false, CF_Always, SO_Keep, SO_Keep, SO_Keep,
-			0xff, 0xff
-		>::GetRHI() );
+	RHISetDepthStencilState(TStaticDepthStencilState<
+		false,CF_Always,
+		true,CF_NotEqual,SO_Keep,SO_Keep,SO_Keep,
+		false,CF_Always,SO_Keep,SO_Keep,SO_Keep,
+		0xff,0xff
+		>::GetRHI());
 	}
 
 	// Light Attenuation channel assignment:
@@ -1882,11 +1901,11 @@ void FProjectedShadowInfo::RenderProjection(int32 ViewIndex, const FViewInfo* Vi
 
 		switch(LocalQuality)
 		{
-			case 1: SetShadowProjectionShaderTemplNew<1>(ViewIndex, *View, this); break;
-			case 2: SetShadowProjectionShaderTemplNew<2>(ViewIndex, *View, this); break;
-			case 3: SetShadowProjectionShaderTemplNew<3>(ViewIndex, *View, this); break;
-			case 4: SetShadowProjectionShaderTemplNew<4>(ViewIndex, *View, this); break;
-			case 5: SetShadowProjectionShaderTemplNew<5>(ViewIndex, *View, this); break;
+			case 1: SetShadowProjectionShaderTemplNew<1>(RHICmdList, ViewIndex, *View, this); break;
+			case 2: SetShadowProjectionShaderTemplNew<2>(RHICmdList, ViewIndex, *View, this); break;
+			case 3: SetShadowProjectionShaderTemplNew<3>(RHICmdList, ViewIndex, *View, this); break;
+			case 4: SetShadowProjectionShaderTemplNew<4>(RHICmdList, ViewIndex, *View, this); break;
+			case 5: SetShadowProjectionShaderTemplNew<5>(RHICmdList, ViewIndex, *View, this); break;
 			default:
 				check(0);
 		}
@@ -1917,14 +1936,14 @@ void FProjectedShadowInfo::RenderProjection(int32 ViewIndex, const FViewInfo* Vi
 	}
 	else
 	{
-		// Clear the stencil buffer to 0.
-		RHIClear(false,FColor(0,0,0),false,0,true,0, FIntRect());
-	}
+	// Clear the stencil buffer to 0.
+	RHIClear(false,FColor(0,0,0),false,0,true,0, FIntRect());
+}
 }
 
 
 template <uint32 Quality>
-static void SetPointLightShaderTempl(int32 ViewIndex, const FViewInfo& View, const FProjectedShadowInfo* ShadowInfo)
+static void SetPointLightShaderTempl(FRHICommandList* RHICmdList, int32 ViewIndex, const FViewInfo& View, const FProjectedShadowInfo* ShadowInfo)
 {
 	TShaderMapRef<FShadowProjectionVS> VertexShader(GetGlobalShaderMap());
 	TShaderMapRef<TOnePassPointShadowProjectionPS<Quality> > PixelShader(GetGlobalShaderMap());
@@ -1932,15 +1951,17 @@ static void SetPointLightShaderTempl(int32 ViewIndex, const FViewInfo& View, con
 	static FGlobalBoundShaderState BoundShaderState;
 	SetGlobalBoundShaderState(BoundShaderState, GetVertexDeclarationFVector4(), *VertexShader, *PixelShader);
 
-	VertexShader->SetParameters(View, ShadowInfo);
-	PixelShader->SetParameters(ViewIndex, View, ShadowInfo);
+	VertexShader->SetParameters(RHICmdList, View, ShadowInfo);
+	PixelShader->SetParameters(RHICmdList, ViewIndex, View, ShadowInfo);
 }
 
 /** Render one pass point light shadow projections. */
-void FProjectedShadowInfo::RenderOnePassPointLightProjection(int32 ViewIndex, const FViewInfo& View) const
+void FProjectedShadowInfo::RenderOnePassPointLightProjection(FRHICommandList* RHICmdList, int32 ViewIndex, const FViewInfo& View) const
 {
 	SCOPE_CYCLE_COUNTER(STAT_RenderWholeSceneShadowProjectionsTime);
 
+	check(!RHICmdList);
+
 	checkSlow(bOnePassPointLightShadow);
 	
 	const FSphere LightBounds = LightSceneInfo->Proxy->GetBoundingSphere();
@@ -1987,11 +2008,11 @@ void FProjectedShadowInfo::RenderOnePassPointLightProjection(int32 ViewIndex, co
 
 		switch(LocalQuality)
 		{
-			case 1: SetPointLightShaderTempl<1>(ViewIndex, View, this); break;
-			case 2: SetPointLightShaderTempl<2>(ViewIndex, View, this); break;
-			case 3: SetPointLightShaderTempl<3>(ViewIndex, View, this); break;
-			case 4: SetPointLightShaderTempl<4>(ViewIndex, View, this); break;
-			case 5: SetPointLightShaderTempl<5>(ViewIndex, View, this); break;
+			case 1: SetPointLightShaderTempl<1>(RHICmdList, ViewIndex, View, this); break;
+			case 2: SetPointLightShaderTempl<2>(RHICmdList, ViewIndex, View, this); break;
+			case 3: SetPointLightShaderTempl<3>(RHICmdList, ViewIndex, View, this); break;
+			case 4: SetPointLightShaderTempl<4>(RHICmdList, ViewIndex, View, this); break;
+			case 5: SetPointLightShaderTempl<5>(RHICmdList, ViewIndex, View, this); break;
 			default:
 				check(0);
 		}
@@ -2315,7 +2336,7 @@ bool FDeferredShadingSceneRenderer::CheckForProjectedShadows( const FLightSceneI
 }
 
 /** Renders one pass point light shadows. */
-bool FDeferredShadingSceneRenderer::RenderOnePassPointLightShadows(const FLightSceneInfo* LightSceneInfo, bool bRenderedTranslucentObjectShadows, bool& bInjectedTranslucentVolume)
+bool FDeferredShadingSceneRenderer::RenderOnePassPointLightShadows(FRHICommandList* RHICmdList, const FLightSceneInfo* LightSceneInfo, bool bRenderedTranslucentObjectShadows, bool& bInjectedTranslucentVolume)
 {
 	bool bDirty = false;
 	FVisibleLightInfo& VisibleLightInfo = VisibleLightInfos[LightSceneInfo->Id];
@@ -2342,8 +2363,9 @@ bool FDeferredShadingSceneRenderer::RenderOnePassPointLightShadows(const FLightS
 				SCOPED_DRAW_EVENT(ShadowDepthsFromOpaque, DEC_SCENE_ITEMS);
 				GSceneRenderTargets.BeginRenderingCubeShadowDepth(ProjectedShadowInfo->ResolutionX);
 				ProjectedShadowInfo->bAllocated = true;
+				//gilmerge, if clear depth doesn't go into a cmdlist, then how can renderdepth? This occurs in 4 spots
 				ProjectedShadowInfo->ClearDepth(this);
-				ProjectedShadowInfo->RenderDepth(this);
+				ProjectedShadowInfo->RenderDepth(RHICmdList, this);
 				GSceneRenderTargets.FinishRenderingCubeShadowDepth(ProjectedShadowInfo->ResolutionX);
 			}
 
@@ -2361,7 +2383,7 @@ bool FDeferredShadingSceneRenderer::RenderOnePassPointLightShadows(const FLightS
 					// Set the device viewport for the view.
 					RHISetViewport(View.ViewRect.Min.X, View.ViewRect.Min.Y, 0.0f, View.ViewRect.Max.X, View.ViewRect.Max.Y, 1.0f);
 
-					ProjectedShadowInfo->RenderOnePassPointLightProjection(ViewIndex, View);
+					ProjectedShadowInfo->RenderOnePassPointLightProjection(RHICmdList, ViewIndex, View);
 				}
 			}
 
@@ -2384,10 +2406,13 @@ bool FDeferredShadingSceneRenderer::RenderOnePassPointLightShadows(const FLightS
 
 /** Renders the projections of the given Shadows to the appropriate color render target. */
 void FDeferredShadingSceneRenderer::RenderProjections(
+	FRHICommandList* RHICmdList, 
 	const FLightSceneInfo* LightSceneInfo, 
 	const TArray<FProjectedShadowInfo*,SceneRenderingAllocator>& Shadows
 	)
 {
+	check(!RHICmdList);
+
 	// Normal shadows render to light attenuation
 	GSceneRenderTargets.BeginRenderingLightAttenuation();
 
@@ -2412,7 +2437,7 @@ void FDeferredShadingSceneRenderer::RenderProjections(
 				// Only project the shadow if it's large enough in this particular view (split screen, etc... may have shadows that are large in one view but irrelevantly small in others)
 				if (ProjectedShadowInfo->FadeAlphas[ViewIndex] > 1.0f / 256.0f)
 				{
-					ProjectedShadowInfo->RenderProjection(ViewIndex, &View);
+					ProjectedShadowInfo->RenderProjection(RHICmdList, ViewIndex, &View);
 				}
 			}
 		}
@@ -2467,7 +2492,7 @@ struct FCompareFProjectedShadowInfoBySplitIndex
 	}
 };
 
-bool FDeferredShadingSceneRenderer::RenderTranslucentProjectedShadows( const FLightSceneInfo* LightSceneInfo )
+bool FDeferredShadingSceneRenderer::RenderTranslucentProjectedShadows(FRHICommandList* RHICmdList, const FLightSceneInfo* LightSceneInfo )
 {
 	SCOPE_CYCLE_COUNTER(STAT_ProjectedShadowDrawTime);
 
@@ -2587,7 +2612,7 @@ bool FDeferredShadingSceneRenderer::RenderTranslucentProjectedShadows( const FLi
 		{
 			{
 				SCOPED_DRAW_EVENT(ShadowProjectionOnOpaque, DEC_SCENE_ITEMS);
-				RenderProjections(LightSceneInfo, Shadows);
+				RenderProjections(RHICmdList, LightSceneInfo, Shadows);
 			}
 
 			for (int32 ShadowIndex = 0; ShadowIndex < Shadows.Num(); ShadowIndex++)
@@ -2626,7 +2651,7 @@ bool FDeferredShadingSceneRenderer::RenderTranslucentProjectedShadows( const FLi
 * @param LightSceneInfo Represents the current light
 * @return true if anything got rendered
 */
-bool FDeferredShadingSceneRenderer::RenderReflectiveShadowMaps( const FLightSceneInfo* LightSceneInfo )
+bool FDeferredShadingSceneRenderer::RenderReflectiveShadowMaps(FRHICommandList* RHICmdList, const FLightSceneInfo* LightSceneInfo )
 {
 	SCOPE_CYCLE_COUNTER(STAT_ReflectiveShadowMapDrawTime);
 
@@ -2743,7 +2768,7 @@ bool FDeferredShadingSceneRenderer::RenderReflectiveShadowMaps( const FLightScen
 				if (ProjectedShadowInfo->bAllocated && !ProjectedShadowInfo->bTranslucentShadow)
 				{
 					ProjectedShadowInfo->ClearDepth(this);
-					ProjectedShadowInfo->RenderDepth(this);
+					ProjectedShadowInfo->RenderDepth(RHICmdList, this);
 				}
 
 				GSceneRenderTargets.FinishRenderingReflectiveShadowMap();
@@ -2766,6 +2791,7 @@ bool FDeferredShadingSceneRenderer::RenderReflectiveShadowMaps( const FLightScen
 					if ( ProjectedShadowInfo->bWholeSceneShadow )
 					{
 						LightPropagationVolume->InjectDirectionalLightRSM( 
+							RHICmdList, 
 							GSceneRenderTargets.GetReflectiveShadowMapDiffuseTexture(), 
 							GSceneRenderTargets.GetReflectiveShadowMapNormalTexture(),
 							GSceneRenderTargets.GetReflectiveShadowMapDepthTexture(),
@@ -2797,7 +2823,7 @@ bool FDeferredShadingSceneRenderer::RenderReflectiveShadowMaps( const FLightScen
  * @param LightSceneInfo Represents the current light
  * @return true if anything got rendered
  */
-bool FDeferredShadingSceneRenderer::RenderProjectedShadows( const FLightSceneInfo* LightSceneInfo, bool bRenderedTranslucentObjectShadows, bool& bInjectedTranslucentVolume )
+bool FDeferredShadingSceneRenderer::RenderProjectedShadows(FRHICommandList* RHICmdList, const FLightSceneInfo* LightSceneInfo, bool bRenderedTranslucentObjectShadows, bool& bInjectedTranslucentVolume )
 {
 	SCOPE_CYCLE_COUNTER(STAT_ProjectedShadowDrawTime);
 
@@ -2915,7 +2941,7 @@ bool FDeferredShadingSceneRenderer::RenderProjectedShadows( const FLightSceneInf
 				FProjectedShadowInfo* ProjectedShadowInfo = Shadows[ShadowIndex];
 				if (ProjectedShadowInfo->bAllocated && !ProjectedShadowInfo->bTranslucentShadow)
 				{
-					ProjectedShadowInfo->RenderDepth(this);
+					ProjectedShadowInfo->RenderDepth(RHICmdList, this);
 				}
 			}
 
@@ -2926,7 +2952,7 @@ bool FDeferredShadingSceneRenderer::RenderProjectedShadows( const FLightSceneInf
 		{
 			{
 				SCOPED_DRAW_EVENT(ShadowProjectionOnOpaque, DEC_SCENE_ITEMS);
-				RenderProjections(LightSceneInfo, Shadows);
+				RenderProjections(RHICmdList, LightSceneInfo, Shadows);
 			}
 			
 			for (int32 ShadowIndex = 0; ShadowIndex < Shadows.Num(); ShadowIndex++)
@@ -2964,19 +2990,21 @@ bool FDeferredShadingSceneRenderer::RenderProjectedShadows( const FLightSceneInf
 		}
 	}
 
-	bAttenuationBufferDirty |= RenderCachedPreshadows(LightSceneInfo);
+	bAttenuationBufferDirty |= RenderCachedPreshadows(RHICmdList, LightSceneInfo);
 
 	if (Scene->GetFeatureLevel() >= ERHIFeatureLevel::SM4)
 	{
-		bAttenuationBufferDirty |= RenderOnePassPointLightShadows(LightSceneInfo, bRenderedTranslucentObjectShadows, bInjectedTranslucentVolume);
+		bAttenuationBufferDirty |= RenderOnePassPointLightShadows(RHICmdList, LightSceneInfo, bRenderedTranslucentObjectShadows, bInjectedTranslucentVolume);
 	}
 
 	return bAttenuationBufferDirty;
 }
 
 /** Renders preshadow depths for any preshadows whose depths aren't cached yet, and renders the projections of preshadows with opaque relevance. */
-bool FDeferredShadingSceneRenderer::RenderCachedPreshadows(const FLightSceneInfo* LightSceneInfo)
+bool FDeferredShadingSceneRenderer::RenderCachedPreshadows(FRHICommandList* RHICmdList, const FLightSceneInfo* LightSceneInfo)
 {
+	check(!RHICmdList);
+
 	bool bAttenuationBufferDirty = false;
 
 	FVisibleLightInfo& VisibleLightInfo = VisibleLightInfos[LightSceneInfo->Id];
@@ -3047,7 +3075,7 @@ bool FDeferredShadingSceneRenderer::RenderCachedPreshadows(const FLightSceneInfo
 				{
 					check(ProjectedShadowInfo->bAllocated);
 					ProjectedShadowInfo->ClearDepth(this);
-					ProjectedShadowInfo->RenderDepth(this);
+					ProjectedShadowInfo->RenderDepth(RHICmdList, this);
 					ProjectedShadowInfo->bDepthsCached = true;
 
 					const FResolveParams ResolveParams(FResolveRect(
@@ -3068,7 +3096,7 @@ bool FDeferredShadingSceneRenderer::RenderCachedPreshadows(const FLightSceneInfo
 		// Render the shadow projections.
 		{
 			SCOPED_DRAW_EVENTF(EventShadowProj, DEC_SCENE_ITEMS, TEXT("Cached Preshadow Projections"));
-			RenderProjections(LightSceneInfo, OpaqueCachedPreshadows);
+			RenderProjections(RHICmdList, LightSceneInfo, OpaqueCachedPreshadows);
 
 			// Mark the attenuation buffer as dirty.
 			bAttenuationBufferDirty = true;
diff --git a/Engine/Source/Runtime/Renderer/Private/ShadowRendering.h b/Engine/Source/Runtime/Renderer/Private/ShadowRendering.h
--- a/Engine/Source/Runtime/Renderer/Private/ShadowRendering.h
+++ b/Engine/Source/Runtime/Renderer/Private/ShadowRendering.h
@@ -7,6 +7,8 @@
 #ifndef __ShadowRendering_H__
 #define __ShadowRendering_H__
 
+#include "ShaderParameterUtils.h"
+
 // Forward declarations.
 class FProjectedShadowInfo;
 
@@ -331,7 +333,7 @@ public:
 			&& bUsePositionOnlyVS == Other.bUsePositionOnlyVS
 			&& bPreShadow == Other.bPreShadow;
 	}
-	void DrawShared(const FSceneView* View,FBoundShaderStateRHIParamRef BoundShaderState) const;
+	void DrawShared(FRHICommandList* RHICmdList, const FSceneView* View,FBoundShaderStateRHIParamRef BoundShaderState) const;
 
 	/** 
 	 * Create bound shader state using the vertex decl from the mesh draw policy
@@ -341,6 +343,7 @@ public:
 	FBoundShaderStateRHIRef CreateBoundShaderState(ERHIFeatureLevel::Type InFeatureLevel);
 
 	void SetMeshRenderState(
+		FRHICommandList* RHICmdList, 
 		const FSceneView& View,
 		const FPrimitiveSceneProxy* PrimitiveSceneProxy,
 		const FMeshBatch& Mesh,
@@ -410,6 +413,7 @@ public:
 	static void AddStaticMesh(FScene* Scene,FStaticMesh* StaticMesh);
 
 	static bool DrawDynamicMesh(
+		FRHICommandList* RHICmdList, 
 		const FSceneView& View,
 		ContextType Context,
 		const FMeshBatch& Mesh,
@@ -658,7 +662,7 @@ public:
 	/**
 	 * Renders the shadow subject depth.
 	 */
-	void RenderDepth(class FDeferredShadingSceneRenderer* SceneRenderer);
+	void RenderDepth(FRHICommandList* RHICmdList, class FDeferredShadingSceneRenderer* SceneRenderer);
 
 	void ClearDepth(class FDeferredShadingSceneRenderer* SceneRenderer);
 
@@ -668,10 +672,10 @@ public:
 	/**
 	 * Projects the shadow onto the scene for a particular view.
 	 */
-	void RenderProjection(int32 ViewIndex, const class FViewInfo* View) const;
+	void RenderProjection(FRHICommandList* RHICmdList, int32 ViewIndex, const class FViewInfo* View) const;
 
 	/** Render one pass point light shadow projections. */
-	void RenderOnePassPointLightProjection(int32 ViewIndex, const FViewInfo& View) const;
+	void RenderOnePassPointLightProjection(FRHICommandList* RHICmdList, int32 ViewIndex, const FViewInfo& View) const;
 
 	/**
 	 * Renders the projected shadow's frustum wireframe with the given FPrimitiveDrawInterface.
@@ -739,7 +743,7 @@ private:
 
 	void GetShadowTypeNameForDrawEvent(FString& TypeName) const;
 
-	template <bool bReflectiveShadowmap> friend void DrawShadowMeshElements(const FViewInfo& View, const FProjectedShadowInfo& ShadowInfo);
+	template <bool bReflectiveShadowmap> friend void DrawShadowMeshElements(FRHICommandList* RHICmdList, const FViewInfo& View, const FProjectedShadowInfo& ShadowInfo);
 };
 
 /** Shader parameters for rendering the depth of a mesh for shadowing. */
@@ -755,30 +759,31 @@ public:
 	}
 
 	template<typename ShaderRHIParamRef>
-	void Set(ShaderRHIParamRef ShaderRHI, const FSceneView& View, const FProjectedShadowInfo* ShadowInfo, const FMaterialRenderProxy* MaterialRenderProxy)
+	void Set(FRHICommandList* RHICmdList, ShaderRHIParamRef ShaderRHI, const FSceneView& View, const FProjectedShadowInfo* ShadowInfo, const FMaterialRenderProxy* MaterialRenderProxy)
 	{
 		SetShaderValue(
+			RHICmdList, 
 			ShaderRHI,
 			ProjectionMatrix,
 			FTranslationMatrix(ShadowInfo->PreShadowTranslation - View.ViewMatrices.PreViewTranslation) * ShadowInfo->SubjectAndReceiverMatrix
 			);
 
-		SetShaderValue(ShaderRHI, ShadowParams, FVector2D(ShadowInfo->GetShaderDepthBias(), ShadowInfo->InvMaxSubjectDepth));
+		SetShaderValue(RHICmdList, ShaderRHI, ShadowParams, FVector2D(ShadowInfo->GetShaderDepthBias(), ShadowInfo->InvMaxSubjectDepth));
 		// Only clamp vertices to the near plane when rendering whole scene directional light shadow depths or preshadows from directional lights
 		const bool bClampToNearPlaneValue = ShadowInfo->IsWholeSceneDirectionalShadow() || (ShadowInfo->bPreShadow && ShadowInfo->bDirectionalLight);
-		SetShaderValue(ShaderRHI,ClampToNearPlane,bClampToNearPlaneValue ? 1.0f : 0.0f);
+		SetShaderValue(RHICmdList, ShaderRHI,ClampToNearPlane,bClampToNearPlaneValue ? 1.0f : 0.0f);
 	}
 
 	/** Set the vertex shader parameter values. */
-	void SetVertexShader(FShader* VertexShader, const FSceneView& View, const FProjectedShadowInfo* ShadowInfo, const FMaterialRenderProxy* MaterialRenderProxy)
+	void SetVertexShader(FRHICommandList* RHICmdList, FShader* VertexShader, const FSceneView& View, const FProjectedShadowInfo* ShadowInfo, const FMaterialRenderProxy* MaterialRenderProxy)
 	{
-		Set(VertexShader->GetVertexShader(), View, ShadowInfo, MaterialRenderProxy);
+		Set(RHICmdList, VertexShader->GetVertexShader(), View, ShadowInfo, MaterialRenderProxy);
 	}
 
 	/** Set the domain shader parameter values. */
-	void SetDomainShader(FShader* DomainShader, const FSceneView& View, const FProjectedShadowInfo* ShadowInfo, const FMaterialRenderProxy* MaterialRenderProxy)
+	void SetDomainShader(FRHICommandList* RHICmdList, FShader* DomainShader, const FSceneView& View, const FProjectedShadowInfo* ShadowInfo, const FMaterialRenderProxy* MaterialRenderProxy)
 	{
-		Set(DomainShader->GetDomainShader(), View, ShadowInfo, MaterialRenderProxy);
+		Set(RHICmdList, DomainShader->GetDomainShader(), View, ShadowInfo, MaterialRenderProxy);
 	}
 
 	/** Serializer. */
@@ -810,25 +815,26 @@ public:
 		StencilPreViewTranslation.Bind(ParameterMap, TEXT("StencilingPreViewTranslation"));
 	}
 
-	void Set(FShader* Shader, const FVector4& InStencilingGeometryPosAndScale) const
+	void Set(FRHICommandList* RHICmdList, FShader* Shader, const FVector4& InStencilingGeometryPosAndScale) const
 	{
-		SetShaderValue(Shader->GetVertexShader(), StencilGeometryPosAndScale, InStencilingGeometryPosAndScale);
-		SetShaderValue(Shader->GetVertexShader(), StencilConeParameters, FVector4(0.0f, 0.0f, 0.0f, 0.0f));
+		SetShaderValue(RHICmdList, Shader->GetVertexShader(), StencilGeometryPosAndScale, InStencilingGeometryPosAndScale);
+		SetShaderValue(RHICmdList, Shader->GetVertexShader(), StencilConeParameters, FVector4(0.0f, 0.0f, 0.0f, 0.0f));
 	}
 
-	void Set(FShader* Shader, const FSceneView& View, const FLightSceneInfo* LightSceneInfo) const
+	void Set(FRHICommandList* RHICmdList, FShader* Shader, const FSceneView& View, const FLightSceneInfo* LightSceneInfo) const
 	{
 		FVector4 GeometryPosAndScale;
 		if(LightSceneInfo->Proxy->GetLightType() == LightType_Point)
 		{
 			StencilingGeometry::GStencilSphereVertexBuffer.CalcTransform(GeometryPosAndScale, LightSceneInfo->Proxy->GetBoundingSphere(), View.ViewMatrices.PreViewTranslation);
-			SetShaderValue(Shader->GetVertexShader(), StencilGeometryPosAndScale, GeometryPosAndScale);
-			SetShaderValue(Shader->GetVertexShader(), StencilConeParameters, FVector4(0.0f, 0.0f, 0.0f, 0.0f));
+			SetShaderValue(RHICmdList, Shader->GetVertexShader(), StencilGeometryPosAndScale, GeometryPosAndScale);
+			SetShaderValue(RHICmdList, Shader->GetVertexShader(), StencilConeParameters, FVector4(0.0f, 0.0f, 0.0f, 0.0f));
 		}
 		else if(LightSceneInfo->Proxy->GetLightType() == LightType_Spot)
 		{
-			SetShaderValue(Shader->GetVertexShader(), StencilConeTransform, LightSceneInfo->Proxy->GetLightToWorld());
+			SetShaderValue(RHICmdList, Shader->GetVertexShader(), StencilConeTransform, LightSceneInfo->Proxy->GetLightToWorld());
 			SetShaderValue(
+				RHICmdList, 
 				Shader->GetVertexShader(),
 				StencilConeParameters,
 				FVector4(
@@ -836,7 +842,7 @@ public:
 					StencilingGeometry::FStencilConeIndexBuffer::NumSlices,
 					LightSceneInfo->Proxy->GetOuterConeAngle(),
 					LightSceneInfo->Proxy->GetRadius()));
-			SetShaderValue(Shader->GetVertexShader(), StencilPreViewTranslation, View.ViewMatrices.PreViewTranslation);
+			SetShaderValue(RHICmdList, Shader->GetVertexShader(), StencilPreViewTranslation, View.ViewMatrices.PreViewTranslation);
 		}
 	}
 
@@ -879,7 +885,7 @@ public:
 		OutEnvironment.SetDefine(TEXT("USE_TRANSFORM"), (uint32)1);
 	}
 
-	void SetParameters(const FSceneView& View, const FProjectedShadowInfo* ShadowInfo, int32 ShadowSplitIndex = INDEX_NONE);
+	void SetParameters(FRHICommandList* RHICmdList, const FSceneView& View, const FProjectedShadowInfo* ShadowInfo, int32 ShadowSplitIndex = INDEX_NONE);
 
 	// Begin FShader interface
 	virtual bool Serialize(FArchive& Ar) OVERRIDE
@@ -919,9 +925,9 @@ public:
 		return true;
 	}
 
-	inline void SetParameters(const FSceneView& View)
+	inline void SetParameters(FRHICommandList* RHICmdList, const FSceneView& View)
 	{
-		FGlobalShader::SetParameters(GetVertexShader(),View); 
+		FGlobalShader::SetParameters(RHICmdList, GetVertexShader(),View); 
 	}
 };
 
@@ -952,12 +958,13 @@ public:
 	 * @param ShadowInfo - projected shadow info for a single light
 	 */
 	virtual void SetParameters(
+		FRHICommandList* RHICmdList, 
 		int32 ViewIndex,
 		const FSceneView& View,
 		const FProjectedShadowInfo* ShadowInfo
 		)
 	{ 
-		FGlobalShader::SetParameters(GetPixelShader(),View);
+		FGlobalShader::SetParameters(RHICmdList, GetPixelShader(),View);
 	}
 
 };
@@ -966,6 +973,9 @@ public:
 class FShadowProjectionShaderParameters
 {
 public:
+	static void ModifyCompilationEnvironment(EShaderPlatform Platform, FShaderCompilerEnvironment& OutEnvironment)
+	{
+	}
 
 	void Bind(const FShaderParameterMap& ParameterMap)
 	{
@@ -982,15 +992,15 @@ public:
 		InvFadePlaneLength.Bind(ParameterMap,TEXT("InvFadePlaneLength"));
 	}
 
-	void Set(FShader* Shader, const FSceneView& View, const FProjectedShadowInfo* ShadowInfo)
+	void Set(FRHICommandList* RHICmdList, FShader* Shader, const FSceneView& View, const FProjectedShadowInfo* ShadowInfo)
 	{
 		const FPixelShaderRHIParamRef ShaderRHI = Shader->GetPixelShader();
 
-		DeferredParameters.Set(ShaderRHI, View);
+		DeferredParameters.Set(RHICmdList, ShaderRHI, View);
 
 		// Set the transform from screen coordinates to shadow depth texture coordinates.
 		const FMatrix ScreenToShadow = ShadowInfo->GetScreenToShadowMatrix(View);
-		SetShaderValue(ShaderRHI, ScreenToShadowMatrix, ScreenToShadow);
+		SetShaderValue(RHICmdList, ShaderRHI, ScreenToShadowMatrix, ScreenToShadow);
 
 		const FIntPoint ShadowBufferResolution = ShadowInfo->GetShadowBufferResolution();
 
@@ -998,39 +1008,40 @@ public:
 		{
 			const float TransitionSize = ShadowInfo->ComputeTransitionSize();
 
-			SetShaderValue(ShaderRHI, SoftTransitionScale, FVector(0, 0, 1.0f / TransitionSize));
+			SetShaderValue(RHICmdList, ShaderRHI, SoftTransitionScale, FVector(0, 0, 1.0f / TransitionSize));
 		}
 
 		if (ShadowBufferSize.IsBound())
 		{
 			FVector2D ShadowBufferSizeValue(ShadowBufferResolution.X, ShadowBufferResolution.Y);
 
-			SetShaderValue(ShaderRHI, ShadowBufferSize,
+			SetShaderValue(RHICmdList, ShaderRHI, ShadowBufferSize,
 				FVector4(ShadowBufferSizeValue.X, ShadowBufferSizeValue.Y, 1.0f / ShadowBufferSizeValue.X, 1.0f / ShadowBufferSizeValue.Y));
 		}
 
 		FTexture2DRHIRef ShadowDepthTextureValue = GSceneRenderTargets.GetShadowDepthZTexture(ShadowInfo->bAllocatedInPreshadowCache);
 		FSamplerStateRHIParamRef DepthSamplerState = TStaticSamplerState<SF_Point,AM_Clamp,AM_Clamp,AM_Clamp>::GetRHI();
 
-		SetTextureParameter(ShaderRHI, ShadowDepthTexture, ShadowDepthTextureSampler, DepthSamplerState, ShadowDepthTextureValue);
-		SetTextureParameter(ShaderRHI, ShadowDepthTextureObject, ShadowDepthTextureValue);
+		SetTextureParameter(RHICmdList, ShaderRHI, ShadowDepthTexture, ShadowDepthTextureSampler, DepthSamplerState, ShadowDepthTextureValue);
+		SetTextureParameter(RHICmdList, ShaderRHI, ShadowDepthTextureObject, ShadowDepthTextureValue);
 
 		if (ShadowDepthSampler.IsBound())
 		{
-			RHISetShaderSampler(
+			FRHICommandList::SetShaderSampler(
+				RHICmdList,
 				ShaderRHI, 
 				ShadowDepthSampler.GetBaseIndex(), 
 				DepthSamplerState
 				);
 		}
 
-		SetShaderValue(ShaderRHI, ProjectionDepthBias, FVector2D(ShadowInfo->GetShaderDepthBias(), ShadowInfo->MaxSubjectZ - ShadowInfo->MinSubjectZ));
-		SetShaderValue(ShaderRHI, FadePlaneOffset, ShadowInfo->CascadeSettings.FadePlaneOffset);
+		SetShaderValue(RHICmdList, ShaderRHI, ProjectionDepthBias, FVector2D(ShadowInfo->GetShaderDepthBias(), ShadowInfo->MaxSubjectZ - ShadowInfo->MinSubjectZ));
+		SetShaderValue(RHICmdList, ShaderRHI, FadePlaneOffset, ShadowInfo->CascadeSettings.FadePlaneOffset);
 
 		if(InvFadePlaneLength.IsBound())
 		{
 			check(ShadowInfo->CascadeSettings.FadePlaneLength > 0);
-			SetShaderValue(ShaderRHI, InvFadePlaneLength, 1.0f / ShadowInfo->CascadeSettings.FadePlaneLength);
+			SetShaderValue(RHICmdList, ShaderRHI, InvFadePlaneLength, 1.0f / ShadowInfo->CascadeSettings.FadePlaneLength);
 		}
 	}
 
@@ -1104,7 +1115,8 @@ public:
 	 */
 	static void ModifyCompilationEnvironment(EShaderPlatform Platform, FShaderCompilerEnvironment& OutEnvironment)
 	{
-		FShadowProjectionPixelShaderInterface::ModifyCompilationEnvironment(Platform, OutEnvironment);
+		FShadowProjectionPixelShaderInterface::ModifyCompilationEnvironment(Platform,OutEnvironment);
+		FShadowProjectionShaderParameters::ModifyCompilationEnvironment(Platform,OutEnvironment);
 		OutEnvironment.SetDefine(TEXT("SHADOW_QUALITY"), Quality);
 		OutEnvironment.SetDefine(TEXT("USE_FADE_PLANE"), (uint32)(bUseFadePlane ? 1 : 0));
 	}
@@ -1115,17 +1127,18 @@ public:
 	 * @param ShadowInfo - projected shadow info for a single light
 	 */
 	virtual void SetParameters(
+		FRHICommandList* RHICmdList, 
 		int32 ViewIndex,
 		const FSceneView& View,
 		const FProjectedShadowInfo* ShadowInfo
-		)
+		) OVERRIDE
 	{
-		FShadowProjectionPixelShaderInterface::SetParameters(ViewIndex,View,ShadowInfo);
+		FShadowProjectionPixelShaderInterface::SetParameters(RHICmdList, ViewIndex,View,ShadowInfo);
 
-		ProjectionParameters.Set(this, View, ShadowInfo);
+		ProjectionParameters.Set(RHICmdList, this, View, ShadowInfo);
 
-		SetShaderValue( GetPixelShader(), ShadowFadeFraction, ShadowInfo->FadeAlphas[ViewIndex] );
-		SetShaderValue( GetPixelShader(), ShadowSharpen, ShadowInfo->LightSceneInfo->Proxy->GetShadowSharpen() * 7.0f + 1.0f );
+		SetShaderValue(RHICmdList, GetPixelShader(), ShadowFadeFraction, ShadowInfo->FadeAlphas[ViewIndex] );
+		SetShaderValue(RHICmdList, GetPixelShader(), ShadowSharpen, ShadowInfo->LightSceneInfo->Proxy->GetShadowSharpen() * 7.0f + 1.0f );
 	}
 
 	/**
@@ -1160,9 +1173,10 @@ public:
 		TranslucencyShadowTransmission1Sampler.Bind(ParameterMap,TEXT("TranslucencyShadowTransmission1Sampler"));
 	}
 
-	void Set(FShader* Shader) const
+	void Set(FRHICommandList* RHICmdList, FShader* Shader) const
 	{
 		SetTextureParameter(
+			RHICmdList, 
 			Shader->GetPixelShader(),
 			TranslucencyShadowTransmission0,
 			TranslucencyShadowTransmission0Sampler,
@@ -1171,6 +1185,7 @@ public:
 			);
 
 		SetTextureParameter(
+			RHICmdList, 
 			Shader->GetPixelShader(),
 			TranslucencyShadowTransmission1,
 			TranslucencyShadowTransmission1Sampler,
@@ -1224,13 +1239,14 @@ public:
 	}
 
 	virtual void SetParameters(
+		FRHICommandList* RHICmdList, 
 		int32 ViewIndex,
 		const FSceneView& View,
-		const FProjectedShadowInfo* ShadowInfo)
+		const FProjectedShadowInfo* ShadowInfo) OVERRIDE
 	{
-		TShadowProjectionPS<Quality>::SetParameters(ViewIndex, View, ShadowInfo);
+		TShadowProjectionPS<Quality>::SetParameters(RHICmdList, ViewIndex, View, ShadowInfo);
 
-		TranslucencyProjectionParameters.Set(this);
+		TranslucencyProjectionParameters.Set(RHICmdList, this);
 	}
 
 	/**
@@ -1263,9 +1279,10 @@ public:
 	}
 
 	template<typename ShaderRHIParamRef>
-	void Set(const ShaderRHIParamRef ShaderRHI, const FProjectedShadowInfo* ShadowInfo) const
+	void Set(FRHICommandList* RHICmdList, const ShaderRHIParamRef ShaderRHI, const FProjectedShadowInfo* ShadowInfo) const
 	{
 		SetTextureParameter(
+			RHICmdList, 
 			ShaderRHI, 
 			ShadowDepthTexture, 
 			GSceneRenderTargets.GetCubeShadowDepthZTexture(ShadowInfo->ResolutionX)
@@ -1282,13 +1299,14 @@ public:
 		}
 
 		SetShaderValueArray<ShaderRHIParamRef, FMatrix>(
+			RHICmdList, 
 			ShaderRHI,
 			ShadowViewProjectionMatrices,
 			ShadowInfo->OnePassShadowViewProjectionMatrices.GetData(),
 			ShadowInfo->OnePassShadowViewProjectionMatrices.Num()
 			);
 
-		SetShaderValue(ShaderRHI,InvShadowmapResolution,1.0f / ShadowInfo->ResolutionX);
+		SetShaderValue(RHICmdList, ShaderRHI,InvShadowmapResolution,1.0f / ShadowInfo->ResolutionX);
 	}
 
 	/** Serializer. */ 
@@ -1334,7 +1352,8 @@ public:
 
 	static void ModifyCompilationEnvironment(EShaderPlatform Platform, FShaderCompilerEnvironment& OutEnvironment)
 	{
-		FGlobalShader::ModifyCompilationEnvironment(Platform, OutEnvironment);
+		FGlobalShader::ModifyCompilationEnvironment(Platform,OutEnvironment);
+		FDeferredPixelShaderParameters::ModifyCompilationEnvironment(Platform,OutEnvironment);
 		OutEnvironment.SetDefine(TEXT("SHADOW_QUALITY"), Quality);
 	}
 
@@ -1344,6 +1363,7 @@ public:
 	}
 
 	void SetParameters(
+		FRHICommandList* RHICmdList, 
 		int32 ViewIndex,
 		const FSceneView& View,
 		const FProjectedShadowInfo* ShadowInfo
@@ -1351,16 +1371,16 @@ public:
 	{
 		const FPixelShaderRHIParamRef ShaderRHI = GetPixelShader();
 
-		FGlobalShader::SetParameters(ShaderRHI,View);
+		FGlobalShader::SetParameters(RHICmdList, ShaderRHI,View);
 
-		DeferredParameters.Set(ShaderRHI, View);
-		OnePassShadowParameters.Set(ShaderRHI, ShadowInfo);
+		DeferredParameters.Set(RHICmdList, ShaderRHI, View);
+		OnePassShadowParameters.Set(RHICmdList, ShaderRHI, ShadowInfo);
 
-		SetShaderValue(ShaderRHI, LightPosition, FVector4(ShadowInfo->LightSceneInfo->Proxy->GetPosition(), 1.0f / ShadowInfo->LightSceneInfo->Proxy->GetRadius()));
+		SetShaderValue(RHICmdList, ShaderRHI, LightPosition, FVector4(ShadowInfo->LightSceneInfo->Proxy->GetPosition(), 1.0f / ShadowInfo->LightSceneInfo->Proxy->GetRadius()));
 
-		SetShaderValue(ShaderRHI, ShadowFadeFraction, ShadowInfo->FadeAlphas[ViewIndex]);
-		SetShaderValue(ShaderRHI, ShadowSharpen, ShadowInfo->LightSceneInfo->Proxy->GetShadowSharpen() * 7.0f + 1.0f);
-		SetShaderValue(ShaderRHI, PointLightDepthBiasParameters, FVector2D(ShadowInfo->GetShaderDepthBias(), 0.0f));
+		SetShaderValue(RHICmdList, ShaderRHI, ShadowFadeFraction, ShadowInfo->FadeAlphas[ViewIndex]);
+		SetShaderValue(RHICmdList, ShaderRHI, ShadowSharpen, ShadowInfo->LightSceneInfo->Proxy->GetShadowSharpen() * 7.0f + 1.0f);
+		SetShaderValue(RHICmdList, ShaderRHI, PointLightDepthBiasParameters, FVector2D(ShadowInfo->GetShaderDepthBias(), 0.0f));
 	}
 
 	virtual bool Serialize(FArchive& Ar)
diff --git a/Engine/Source/Runtime/Renderer/Private/StaticMeshDrawList.h b/Engine/Source/Runtime/Renderer/Private/StaticMeshDrawList.h
--- a/Engine/Source/Runtime/Renderer/Private/StaticMeshDrawList.h
+++ b/Engine/Source/Runtime/Renderer/Private/StaticMeshDrawList.h
@@ -228,7 +228,7 @@ private:
 	* @param DrawingPolicyLink - the drawing policy link
 	* @param bDrawnShared - determines whether to draw shared 
 	*/
-	void DrawElement(const FViewInfo& View, const FElement& Element, uint64 BatchElementMask, FDrawingPolicyLink* DrawingPolicyLink, bool &bDrawnShared);
+	void DrawElement(FRHICommandList* RHICmdList, const FViewInfo& View, const FElement& Element, uint64 BatchElementMask, FDrawingPolicyLink* DrawingPolicyLink, bool &bDrawnShared);
 
 public:
 
@@ -261,7 +261,7 @@ public:
 	 * @param BatchVisibilityArray - An array of batch element visibility bitmasks.
 	 * @return True if any static meshes were drawn.
 	 */
-	bool DrawVisible(const FViewInfo& View, const TBitArray<SceneRenderingBitArrayAllocator>& StaticMeshVisibilityMap, const TArray<uint64,SceneRenderingAllocator>& BatchVisibilityArray);
+	bool DrawVisible(FRHICommandList* RHICmdList, const FViewInfo& View, const TBitArray<SceneRenderingBitArrayAllocator>& StaticMeshVisibilityMap, const TArray<uint64,SceneRenderingAllocator>& BatchVisibilityArray);
 
 	/**
 	 * Draws only the static meshes which are in the visibility map, sorted front-to-back.
@@ -271,7 +271,7 @@ public:
 	 * @param MaxToDraw - The maximum number of meshes to be drawn.
 	 * @return The number of static meshes drawn.
 	 */
-	int32 DrawVisibleFrontToBack(const FViewInfo& View, const TBitArray<SceneRenderingBitArrayAllocator>& StaticMeshVisibilityMap, const TArray<uint64,SceneRenderingAllocator>& BatchVisibilityArray, int32 MaxToDraw);
+	int32 DrawVisibleFrontToBack(FRHICommandList* RHICmdList, const FViewInfo& View, const TBitArray<SceneRenderingBitArrayAllocator>& StaticMeshVisibilityMap, const TArray<uint64,SceneRenderingAllocator>& BatchVisibilityArray, int32 MaxToDraw);
 
 	/** Sorts OrderedDrawingPolicies front to back. */
 	void SortFrontToBack(FVector ViewPosition);
diff --git a/Engine/Source/Runtime/Renderer/Private/StaticMeshDrawList.inl b/Engine/Source/Runtime/Renderer/Private/StaticMeshDrawList.inl
--- a/Engine/Source/Runtime/Renderer/Private/StaticMeshDrawList.inl
+++ b/Engine/Source/Runtime/Renderer/Private/StaticMeshDrawList.inl
@@ -7,6 +7,8 @@
 #ifndef __STATICMESHDRAWLIST_INL__
 #define __STATICMESHDRAWLIST_INL__
 
+#include "RHICommandList.h"
+
 // Expensive
 #define PER_MESH_DRAW_STATS 0
 
@@ -56,6 +58,7 @@ void TStaticMeshDrawList<DrawingPolicyType>::FElementHandle::Remove()
 
 template<typename DrawingPolicyType>
 void TStaticMeshDrawList<DrawingPolicyType>::DrawElement(
+	FRHICommandList* RHICmdList,
 	const FViewInfo& View,
 	const FElement& Element,
 	uint64 BatchElementMask,
@@ -73,7 +76,7 @@ void TStaticMeshDrawList<DrawingPolicyType>::DrawElement(
 		{
 			DrawingPolicyLink->CreateBoundShaderState();
 		}
-		DrawingPolicyLink->DrawingPolicy.DrawShared(&View,DrawingPolicyLink->BoundShaderState);
+		DrawingPolicyLink->DrawingPolicy.DrawShared(RHICmdList, &View,DrawingPolicyLink->BoundShaderState);
 		bDrawnShared = true;
 	}
 	
@@ -89,6 +92,7 @@ void TStaticMeshDrawList<DrawingPolicyType>::DrawElement(
 				INC_DWORD_STAT(STAT_StaticDrawListMeshDrawCalls);
 
 				DrawingPolicyLink->DrawingPolicy.SetMeshRenderState(
+					RHICmdList, 
 					View,
 					Element.Mesh->PrimitiveSceneInfo->Proxy,
 					*Element.Mesh,
@@ -96,7 +100,7 @@ void TStaticMeshDrawList<DrawingPolicyType>::DrawElement(
 					!!BackFace,
 					Element.PolicyData
 					);
-				DrawingPolicyLink->DrawingPolicy.DrawMesh(*Element.Mesh,BatchElementIndex);
+				DrawingPolicyLink->DrawingPolicy.DrawMesh(RHICmdList, *Element.Mesh,BatchElementIndex);
 			}
 		}
 
@@ -234,6 +238,7 @@ bool TStaticMeshDrawList<DrawingPolicyType>::DrawVisible(
 
 template<typename DrawingPolicyType>
 bool TStaticMeshDrawList<DrawingPolicyType>::DrawVisible(
+	FRHICommandList* RHICmdList,
 	const FViewInfo& View,
 	const TBitArray<SceneRenderingBitArrayAllocator>& StaticMeshVisibilityMap,
 	const TArray<uint64,SceneRenderingAllocator>& BatchVisibilityArray
@@ -256,7 +261,7 @@ bool TStaticMeshDrawList<DrawingPolicyType>::DrawVisible(
 				INC_DWORD_STAT_BY(STAT_StaticMeshTriangles,Element.Mesh->GetNumPrimitives());
 				// Avoid the cache miss looking up batch visibility if there is only one element.
 				uint64 BatchElementMask = Element.Mesh->Elements.Num() == 1 ? 1 : BatchVisibilityArray[Element.Mesh->Id];
-				DrawElement(View, Element, BatchElementMask, DrawingPolicyLink, bDrawnShared);
+				DrawElement(RHICmdList, View, Element, BatchElementMask, DrawingPolicyLink, bDrawnShared);
 				bDirty = true;
 			}
 		}
@@ -266,6 +271,7 @@ bool TStaticMeshDrawList<DrawingPolicyType>::DrawVisible(
 
 template<typename DrawingPolicyType>
 int32 TStaticMeshDrawList<DrawingPolicyType>::DrawVisibleFrontToBack(
+	FRHICommandList* RHICmdList,
 	const FViewInfo& View,
 	const TBitArray<SceneRenderingBitArrayAllocator>& StaticMeshVisibilityMap,
 	const TArray<uint64,SceneRenderingAllocator>& BatchVisibilityArray,
@@ -318,7 +324,7 @@ int32 TStaticMeshDrawList<DrawingPolicyType>::DrawVisibleFrontToBack(
 		INC_DWORD_STAT_BY(STAT_StaticMeshTriangles,Element.Mesh->GetNumPrimitives());
 		// Avoid the cache miss looking up batch visibility if there is only one element.
 		uint64 BatchElementMask = Element.Mesh->Elements.Num() == 1 ? 1 : BatchVisibilityArray[Element.Mesh->Id];
-		DrawElement(View, Element, BatchElementMask, DrawingPolicyLink, bDrawnShared);
+		DrawElement(RHICmdList, View, Element, BatchElementMask, DrawingPolicyLink, bDrawnShared);
 		NumDraws++;
 	}
 
diff --git a/Engine/Source/Runtime/Renderer/Private/TiledDeferredLightRendering.cpp b/Engine/Source/Runtime/Renderer/Private/TiledDeferredLightRendering.cpp
--- a/Engine/Source/Runtime/Renderer/Private/TiledDeferredLightRendering.cpp
+++ b/Engine/Source/Runtime/Renderer/Private/TiledDeferredLightRendering.cpp
@@ -80,6 +80,7 @@ public:
 	static void ModifyCompilationEnvironment(EShaderPlatform Platform, FShaderCompilerEnvironment& OutEnvironment)
 	{
 		FGlobalShader::ModifyCompilationEnvironment(Platform, OutEnvironment);
+		FDeferredPixelShaderParameters::ModifyCompilationEnvironment(Platform,OutEnvironment);
 		OutEnvironment.SetDefine(TEXT("THREADGROUP_SIZEX"), GDeferredLightTileSizeX);
 		OutEnvironment.SetDefine(TEXT("THREADGROUP_SIZEY"), GDeferredLightTileSizeY);
 		OutEnvironment.SetDefine(TEXT("MAX_LIGHTS"), GMaxNumTiledDeferredLights);
@@ -105,6 +106,7 @@ public:
 	}
 
 	void SetParameters(
+		FRHICommandList* RHICmdList, 
 		const FSceneView& View, 
 		int32 ViewIndex,
 		int32 NumViews,
@@ -118,14 +120,15 @@ public:
 	{
 		FComputeShaderRHIParamRef ShaderRHI = GetComputeShader();
 
-		FGlobalShader::SetParameters(ShaderRHI, View);
-		DeferredParameters.Set(ShaderRHI, View);
-		SetTextureParameter(ShaderRHI, InTexture, InTextureValue.GetRenderTargetItem().ShaderResourceTexture);
-		OutTexture.SetTexture(ShaderRHI, 0, OutTextureValue.GetRenderTargetItem().UAV);
-	
-		SetShaderValue(ShaderRHI, ViewDimensions, View.ViewRect);
+		FGlobalShader::SetParameters(RHICmdList, ShaderRHI, View);
+		DeferredParameters.Set(RHICmdList, ShaderRHI, View);
+		SetTextureParameter(RHICmdList, ShaderRHI, InTexture, InTextureValue.GetRenderTargetItem().ShaderResourceTexture);
+		OutTexture.SetTexture(RHICmdList, ShaderRHI, 0, OutTextureValue.GetRenderTargetItem().UAV);
+
+		SetShaderValue(RHICmdList, ShaderRHI, ViewDimensions, View.ViewRect);
 
 		SetTextureParameter(
+			RHICmdList, 
 			ShaderRHI,
 			PreIntegratedBRDF,
 			PreIntegratedBRDFSampler,
@@ -203,14 +206,14 @@ public:
 			}
 		}
 
-		SetUniformBufferParameterImmediate(ShaderRHI, GetUniformBufferParameter<FTiledDeferredLightData>(), LightData);
-		SetUniformBufferParameterImmediate(ShaderRHI, GetUniformBufferParameter<FTiledDeferredLightData2>(), LightData2);
-		SetShaderValue(ShaderRHI, NumLights, NumThisPass);
+		SetUniformBufferParameterImmediate(RHICmdList, ShaderRHI, GetUniformBufferParameter<FTiledDeferredLightData>(), LightData);
+		SetUniformBufferParameterImmediate(RHICmdList, ShaderRHI, GetUniformBufferParameter<FTiledDeferredLightData2>(), LightData2);
+		SetShaderValue(RHICmdList, ShaderRHI, NumLights, NumThisPass);
 	}
 
-	void UnsetParameters()
+	void UnsetParameters(FRHICommandList* RHICmdList)
 	{
-		OutTexture.UnsetUAV(GetComputeShader());
+		OutTexture.UnsetUAV(RHICmdList, GetComputeShader());
 	}
 
 	virtual bool Serialize(FArchive& Ar)
@@ -269,11 +272,12 @@ bool FDeferredShadingSceneRenderer::ShouldUseTiledDeferred(int32 NumUnshadowedLi
 
 template <bool bVisualizeLightCulling>
 static void SetShaderTemplTiledLighting(
-	const FSceneView& View,
+	FRHICommandList* RHICmdList, 
+	const FSceneView& View, 
 	int32 ViewIndex,
 	int32 NumViews,
-	const TArray<FSortedLightSceneInfo, SceneRenderingAllocator>& SortedLights,
-	int32 NumLightsToRenderInSortedLights,
+	const TArray<FSortedLightSceneInfo, SceneRenderingAllocator>& SortedLights, 
+	int32 NumLightsToRenderInSortedLights, 
 	const FSimpleLightArray& SimpleLights,
 	int32 StartIndex, 
 	int32 NumThisPass,
@@ -283,18 +287,20 @@ static void SetShaderTemplTiledLighting(
 	TShaderMapRef<FTiledDeferredLightingCS<bVisualizeLightCulling> > ComputeShader(GetGlobalShaderMap());
 	RHISetComputeShader(ComputeShader->GetComputeShader());
 
-	ComputeShader->SetParameters(View, ViewIndex, NumViews, SortedLights, NumLightsToRenderInSortedLights, SimpleLights, StartIndex, NumThisPass, InTexture, OutTexture);
+	ComputeShader->SetParameters(RHICmdList, View, ViewIndex, NumViews, SortedLights, NumLightsToRenderInSortedLights, SimpleLights, StartIndex, NumThisPass, InTexture, OutTexture);
 
 	uint32 GroupSizeX = (View.ViewRect.Size().X + GDeferredLightTileSizeX - 1) / GDeferredLightTileSizeX;
 	uint32 GroupSizeY = (View.ViewRect.Size().Y + GDeferredLightTileSizeY - 1) / GDeferredLightTileSizeY;
-	DispatchComputeShader(*ComputeShader, GroupSizeX, GroupSizeY, 1);
+	DispatchComputeShader(RHICmdList, *ComputeShader, GroupSizeX, GroupSizeY, 1);
 
-	ComputeShader->UnsetParameters();
+	ComputeShader->UnsetParameters(RHICmdList);
 }
 
 
-void FDeferredShadingSceneRenderer::RenderTiledDeferredLighting(const TArray<FSortedLightSceneInfo, SceneRenderingAllocator>& SortedLights, int32 NumUnshadowedLights, const FSimpleLightArray& SimpleLights)
+void FDeferredShadingSceneRenderer::RenderTiledDeferredLighting(FRHICommandList* RHICmdList, const TArray<FSortedLightSceneInfo, SceneRenderingAllocator>& SortedLights, int32 NumUnshadowedLights, const FSimpleLightArray& SimpleLights)
 {
+	check(!RHICmdList);
+
 	check(GUseTiledDeferredShading);
 	check(SortedLights.Num() >= NumUnshadowedLights);
 
@@ -339,11 +345,11 @@ void FDeferredShadingSceneRenderer::RenderTiledDeferredLighting(const TArray<FSo
 
 					if(View.Family->EngineShowFlags.VisualizeLightCulling)
 					{
-						SetShaderTemplTiledLighting<1>(View, ViewIndex, Views.Num(), SortedLights, NumLightsToRenderInSortedLights, SimpleLights, StartIndex, NumThisPass, InTexture, *OutTexture);
+						SetShaderTemplTiledLighting<1>(RHICmdList, View, ViewIndex, Views.Num(), SortedLights, NumLightsToRenderInSortedLights, SimpleLights, StartIndex, NumThisPass, InTexture, *OutTexture);
 					}
 					else
 					{
-						SetShaderTemplTiledLighting<0>(View, ViewIndex, Views.Num(), SortedLights, NumLightsToRenderInSortedLights, SimpleLights, StartIndex, NumThisPass, InTexture, *OutTexture);
+						SetShaderTemplTiledLighting<0>(RHICmdList, View, ViewIndex, Views.Num(), SortedLights, NumLightsToRenderInSortedLights, SimpleLights, StartIndex, NumThisPass, InTexture, *OutTexture);
 					}
 				}
 			}
