UE4 - merge first pass at parallel rendering to main

diff --git a/Engine/Source/Runtime/Renderer/Private/TranslucentLighting.cpp b/Engine/Source/Runtime/Renderer/Private/TranslucentLighting.cpp
--- a/Engine/Source/Runtime/Renderer/Private/TranslucentLighting.cpp
+++ b/Engine/Source/Runtime/Renderer/Private/TranslucentLighting.cpp
@@ -182,18 +182,19 @@ public:
 	}
 
 	void SetParameters(
+		FRHICommandList* RHICmdList, 
 		const FMaterialRenderProxy* MaterialRenderProxy,
 		const FSceneView& View,
 		const FProjectedShadowInfo* ShadowInfo
 		)
 	{
-		FMeshMaterialShader::SetParameters(GetVertexShader(), MaterialRenderProxy, *MaterialRenderProxy->GetMaterial(View.GetFeatureLevel()), View, ESceneRenderTargetsMode::DontSet);
-		ShadowParameters.SetVertexShader(this, View, ShadowInfo, MaterialRenderProxy);
+		FMeshMaterialShader::SetParameters(RHICmdList, GetVertexShader(), MaterialRenderProxy, *MaterialRenderProxy->GetMaterial(View.GetFeatureLevel()), View, ESceneRenderTargetsMode::DontSet);
+		ShadowParameters.SetVertexShader(RHICmdList, this, View, ShadowInfo, MaterialRenderProxy);
 	}
 
-	void SetMesh(const FVertexFactory* VertexFactory,const FSceneView& View,const FPrimitiveSceneProxy* Proxy,const FMeshBatchElement& BatchElement)
+	void SetMesh(FRHICommandList* RHICmdList, const FVertexFactory* VertexFactory,const FSceneView& View,const FPrimitiveSceneProxy* Proxy,const FMeshBatchElement& BatchElement)
 	{
-		FMeshMaterialShader::SetMesh(GetVertexShader(),VertexFactory,View,Proxy,BatchElement);
+		FMeshMaterialShader::SetMesh(RHICmdList, GetVertexShader(),VertexFactory,View,Proxy,BatchElement);
 	}
 
 private:
@@ -253,6 +254,7 @@ public:
 	FTranslucencyShadowDepthPS() {}
 
 	void SetParameters(
+		FRHICommandList* RHICmdList, 
 		const FMaterialRenderProxy* MaterialRenderProxy,
 		const FSceneView& View,
 		const FProjectedShadowInfo* ShadowInfo
@@ -264,19 +266,19 @@ public:
 		//@todo - scene depth can be bound by the material for use in depth fades
 		// This is incorrect when rendering a shadowmap as it's not from the camera's POV
 		// Set the scene depth texture to something safe when rendering shadow depths
-		FMeshMaterialShader::SetParameters(ShaderRHI, MaterialRenderProxy, *MaterialRenderProxy->GetMaterial(FeatureLevel), View, ESceneRenderTargetsMode::NonSceneAlignedPass);
+		FMeshMaterialShader::SetParameters(RHICmdList, ShaderRHI, MaterialRenderProxy, *MaterialRenderProxy->GetMaterial(FeatureLevel), View, ESceneRenderTargetsMode::NonSceneAlignedPass);
 
-		SetShaderValue(ShaderRHI, TranslInvMaxSubjectDepth, ShadowInfo->InvMaxSubjectDepth);
+		SetShaderValue(RHICmdList, ShaderRHI, TranslInvMaxSubjectDepth, ShadowInfo->InvMaxSubjectDepth);
 
 		const float LocalToWorldScale = ShadowInfo->ParentSceneInfo->Proxy->GetLocalToWorld().GetScaleVector().GetMax();
 		const float TranslucentShadowStartOffsetValue = MaterialRenderProxy->GetMaterial(FeatureLevel)->GetTranslucentShadowStartOffset() * LocalToWorldScale;
-		SetShaderValue(ShaderRHI,TranslucentShadowStartOffset, TranslucentShadowStartOffsetValue / (ShadowInfo->MaxSubjectZ - ShadowInfo->MinSubjectZ));
-		TranslucencyProjectionParameters.Set(this);
+		SetShaderValue(RHICmdList, ShaderRHI,TranslucentShadowStartOffset, TranslucentShadowStartOffsetValue / (ShadowInfo->MaxSubjectZ - ShadowInfo->MinSubjectZ));
+		TranslucencyProjectionParameters.Set(RHICmdList, this);
 	}
 
-	void SetMesh(const FVertexFactory* VertexFactory,const FSceneView& View,const FPrimitiveSceneProxy* Proxy,const FMeshBatchElement& BatchElement)
+	void SetMesh(FRHICommandList* RHICmdList, const FVertexFactory* VertexFactory,const FSceneView& View,const FPrimitiveSceneProxy* Proxy,const FMeshBatchElement& BatchElement)
 	{
-		FMeshMaterialShader::SetMesh(GetPixelShader(),VertexFactory,View,Proxy,BatchElement);
+		FMeshMaterialShader::SetMesh(RHICmdList, GetPixelShader(),VertexFactory,View,Proxy,BatchElement);
 	}
 
 	virtual bool Serialize(FArchive& Ar)
@@ -346,16 +348,18 @@ public:
 		}
 	}
 
-	void DrawShared(const FSceneView* View,FBoundShaderStateRHIParamRef BoundShaderState) const
+	void DrawShared(FRHICommandList* RHICmdList, const FSceneView* View,FBoundShaderStateRHIParamRef BoundShaderState) const
 	{
+		check(!RHICmdList);
+
 		// Set the shared mesh resources.
-		FMeshDrawingPolicy::DrawShared(View);
+		FMeshDrawingPolicy::DrawShared(RHICmdList, View);
 
 		// Set the actual shader & vertex declaration state
 		RHISetBoundShaderState(BoundShaderState);
 
-		VertexShader->SetParameters(MaterialRenderProxy, *View, PolicyShadowInfo);
-		PixelShader->SetParameters(MaterialRenderProxy, *View, PolicyShadowInfo);
+		VertexShader->SetParameters(RHICmdList, MaterialRenderProxy, *View, PolicyShadowInfo);
+		PixelShader->SetParameters(RHICmdList, MaterialRenderProxy, *View, PolicyShadowInfo);
 	}
 
 	/** 
@@ -375,6 +379,7 @@ public:
 	}
 
 	void SetMeshRenderState(
+		FRHICommandList* RHICmdList, 
 		const FSceneView& View,
 		const FPrimitiveSceneProxy* PrimitiveSceneProxy,
 		const FMeshBatch& Mesh,
@@ -384,10 +389,10 @@ public:
 		) const
 	{
 		const FMeshBatchElement& BatchElement = Mesh.Elements[BatchElementIndex];
-		VertexShader->SetMesh(VertexFactory,View,PrimitiveSceneProxy,BatchElement);
-		PixelShader->SetMesh(VertexFactory,View,PrimitiveSceneProxy,BatchElement);
+		VertexShader->SetMesh(RHICmdList, VertexFactory,View,PrimitiveSceneProxy,BatchElement);
+		PixelShader->SetMesh(RHICmdList, VertexFactory,View,PrimitiveSceneProxy,BatchElement);
 
-		FMeshDrawingPolicy::SetMeshRenderState(View,PrimitiveSceneProxy,Mesh,BatchElementIndex,bBackFace,ElementData);
+		FMeshDrawingPolicy::SetMeshRenderState(RHICmdList, View,PrimitiveSceneProxy,Mesh,BatchElementIndex,bBackFace,ElementData);
 	}
 
 	static const FProjectedShadowInfo* PolicyShadowInfo;
@@ -414,6 +419,7 @@ public:
 	};
 
 	static bool DrawDynamicMesh(
+		FRHICommandList* RHICmdList, 
 		const FSceneView& View,
 		ContextType DrawingContext,
 		const FMeshBatch& Mesh,
@@ -435,12 +441,12 @@ public:
 			if (IsTranslucentBlendMode(BlendMode))
 			{
 				FTranslucencyShadowDepthDrawingPolicy DrawingPolicy(Mesh.VertexFactory, MaterialRenderProxy, *MaterialRenderProxy->GetMaterial(FeatureLevel), DrawingContext.bDirectionalLight);
-				DrawingPolicy.DrawShared(&View, DrawingPolicy.CreateBoundShaderState(FeatureLevel));
+				DrawingPolicy.DrawShared(RHICmdList, &View, DrawingPolicy.CreateBoundShaderState(FeatureLevel));
 
 				for (int32 BatchElementIndex = 0; BatchElementIndex < Mesh.Elements.Num(); BatchElementIndex++)
 				{
-					DrawingPolicy.SetMeshRenderState(View,PrimitiveSceneProxy,Mesh,BatchElementIndex,bBackFace,FMeshDrawingPolicy::ElementDataType());
-					DrawingPolicy.DrawMesh(Mesh,BatchElementIndex);
+					DrawingPolicy.SetMeshRenderState(RHICmdList, View,PrimitiveSceneProxy,Mesh,BatchElementIndex,bBackFace,FMeshDrawingPolicy::ElementDataType());
+					DrawingPolicy.DrawMesh(RHICmdList, Mesh,BatchElementIndex);
 				}
 				bDirty = true;
 			}
@@ -450,6 +456,7 @@ public:
 	}
 
 	static bool DrawStaticMesh(
+		FRHICommandList* RHICmdList, 
 		const FViewInfo& View,
 		ContextType DrawingContext,
 		const FStaticMesh& StaticMesh,
@@ -461,6 +468,7 @@ public:
 		bool bDirty = false;
 
 		bDirty |= DrawDynamicMesh(
+			RHICmdList, 
 			View,
 			DrawingContext,
 			StaticMesh,
@@ -482,6 +490,9 @@ public:
 /** Renders shadow maps for translucent primitives. */
 void FProjectedShadowInfo::RenderTranslucencyDepths(FDeferredShadingSceneRenderer* SceneRenderer)
 {
+	//@todo-rco: RHIPacketList
+	FRHICommandList* RHICmdList = nullptr;
+
 	checkSlow(!bWholeSceneShadow);
 	SCOPE_CYCLE_COUNTER(STAT_RenderWholeSceneShadowDepthsTime);
 
@@ -587,6 +598,7 @@ void FProjectedShadowInfo::RenderTranslucencyDepths(FDeferredShadingSceneRendere
 					for (int32 MeshIndex = 0; MeshIndex < PrimitiveSceneInfo->StaticMeshes.Num(); MeshIndex++)
 					{
 						FTranslucencyShadowDepthDrawingPolicyFactory::DrawStaticMesh(
+							RHICmdList, 
 							*FoundView,
 							FTranslucencyShadowDepthDrawingPolicyFactory::ContextType(bDirectionalLight),
 							PrimitiveSceneInfo->StaticMeshes[MeshIndex],
@@ -632,14 +644,15 @@ public:
 	}
 	FFilterTranslucentVolumePS() {}
 
-	void SetParameters(const FViewInfo& View, int32 VolumeCascadeIndex)
+	void SetParameters(FRHICommandList* RHICmdList, const FViewInfo& View, int32 VolumeCascadeIndex)
 	{
 		const FPixelShaderRHIParamRef ShaderRHI = GetPixelShader();
 
-		FGlobalShader::SetParameters(ShaderRHI, View);
+		FGlobalShader::SetParameters(RHICmdList, ShaderRHI, View);
 
-		SetShaderValue(ShaderRHI, TexelSize, 1.0f / GTranslucencyLightingVolumeDim);
+		SetShaderValue(RHICmdList, ShaderRHI, TexelSize, 1.0f / GTranslucencyLightingVolumeDim);
 		SetTextureParameter(
+			RHICmdList,
 			ShaderRHI, 
 			TranslucencyLightingVolumeAmbient, 
 			TranslucencyLightingVolumeAmbientSampler, 
@@ -647,6 +660,7 @@ public:
 			GSceneRenderTargets.TranslucencyLightingVolumeAmbient[VolumeCascadeIndex]->GetRenderTargetItem().ShaderResourceTexture);
 
 		SetTextureParameter(
+			RHICmdList,
 			ShaderRHI, 
 			TranslucencyLightingVolumeDirectional, 
 			TranslucencyLightingVolumeDirectionalSampler, 
@@ -689,6 +703,7 @@ public:
 
 	template<typename ShaderRHIParamRef>
 	void Set(
+		FRHICommandList* RHICmdList, 
 		const ShaderRHIParamRef ShaderRHI, 
 		FShader* Shader, 
 		const FViewInfo& View, 
@@ -697,18 +712,18 @@ public:
 		int32 VolumeCascadeIndexValue, 
 		bool bShadowed) const
 	{
-		SetDeferredLightParameters(ShaderRHI, Shader->GetUniformBufferParameter<FDeferredLightUniformStruct>(), LightSceneInfo, View);
+		SetDeferredLightParameters(RHICmdList, ShaderRHI, Shader->GetUniformBufferParameter<FDeferredLightUniformStruct>(), LightSceneInfo, View);
 
 		if (bShadowed)
 		{
 			FVector4 ShadowmapMinMaxValue;
 			FMatrix WorldToShadowMatrixValue = ShadowMap->GetWorldToShadowMatrix(ShadowmapMinMaxValue);
 
-			SetShaderValue(ShaderRHI, WorldToShadowMatrix, WorldToShadowMatrixValue);
-			SetShaderValue(ShaderRHI, ShadowmapMinMax, ShadowmapMinMaxValue);
+			SetShaderValue(RHICmdList, ShaderRHI, WorldToShadowMatrix, WorldToShadowMatrixValue);
+			SetShaderValue(RHICmdList, ShaderRHI, ShadowmapMinMax, ShadowmapMinMaxValue);
 		}
 
-		SetShaderValue(ShaderRHI, VolumeCascadeIndex, VolumeCascadeIndexValue);
+		SetShaderValue(RHICmdList, ShaderRHI, VolumeCascadeIndex, VolumeCascadeIndexValue);
 	}
 
 	/** Serializer. */ 
@@ -752,11 +767,11 @@ public:
 	}
 	FTranslucentObjectShadowingPS() {}
 
-	void SetParameters(const FViewInfo& View, const FLightSceneInfo* LightSceneInfo, const FProjectedShadowInfo* ShadowMap, int32 VolumeCascadeIndexValue)
+	void SetParameters(FRHICommandList* RHICmdList, const FViewInfo& View, const FLightSceneInfo* LightSceneInfo, const FProjectedShadowInfo* ShadowMap, int32 VolumeCascadeIndexValue)
 	{
-		FGlobalShader::SetParameters(GetPixelShader(), View);
-		TranslucencyProjectionParameters.Set(this);
-		TranslucentInjectParameters.Set(GetPixelShader(), this, View, LightSceneInfo, ShadowMap, VolumeCascadeIndexValue, true);
+		FGlobalShader::SetParameters(RHICmdList, GetPixelShader(), View);
+		TranslucencyProjectionParameters.Set(RHICmdList, this);
+		TranslucentInjectParameters.Set(RHICmdList, GetPixelShader(), this, View, LightSceneInfo, ShadowMap, VolumeCascadeIndexValue, true);
 	}
 
 	virtual bool Serialize(FArchive& Ar)
@@ -822,6 +837,7 @@ public:
 	// @param InnerSplitIndex which CSM shadow map level, INDEX_NONE if no directional light
 	// @param VolumeCascadeIndexValue which volume we render to
 	void SetParameters(
+		FRHICommandList* RHICmdList, 
 		const FViewInfo& View, 
 		const FLightSceneInfo* LightSceneInfo, 
 		const FMaterialRenderProxy* MaterialProxy, 
@@ -833,10 +849,10 @@ public:
 		
 		const FPixelShaderRHIParamRef ShaderRHI = GetPixelShader();
 
-		FMaterialShader::SetParameters(ShaderRHI, MaterialProxy, *MaterialProxy->GetMaterial(View.GetFeatureLevel()), View, false, ESceneRenderTargetsMode::SetTextures);
+		FMaterialShader::SetParameters(RHICmdList, ShaderRHI, MaterialProxy, *MaterialProxy->GetMaterial(View.GetFeatureLevel()), View, false, ESceneRenderTargetsMode::SetTextures);
 		
 		FSphere ShadowBounds = ShadowMap ? ShadowMap->ShadowBounds : FSphere(FVector(0,0,0), HALF_WORLD_MAX);
-		SetShaderValue(ShaderRHI, CascadeBounds, ShadowBounds);
+		SetShaderValue(RHICmdList, ShaderRHI, CascadeBounds, ShadowBounds);
 
 		FSphere InnerBounds(0);
 		// default to ignore the plane
@@ -873,20 +889,21 @@ public:
 			ShadowInjectParamValue.W = -FadeParams.X * FadeParams.Y;
 		}
 
-		SetShaderValue(ShaderRHI, ShadowInjectParams, ShadowInjectParamValue);
-		SetShaderValue(ShaderRHI, InnerCascadeBounds, InnerBounds);
+		SetShaderValue(RHICmdList, ShaderRHI, ShadowInjectParams, ShadowInjectParamValue);
+		SetShaderValue(RHICmdList, ShaderRHI, InnerCascadeBounds, InnerBounds);
 
-		SetShaderValueArray(ShaderRHI, ClippingPlanes, Planes, ARRAY_COUNT(Planes));
+		SetShaderValueArray(RHICmdList, ShaderRHI, ClippingPlanes, Planes, ARRAY_COUNT(Planes));
 	
 
 		//@todo - needs to be a permutation to reduce shadow filtering work
-		SetShaderValue(ShaderRHI, SpotlightMask, (LightSceneInfo->Proxy->GetLightType() == LightType_Spot ? 1.0f : 0.0f));
+		SetShaderValue(RHICmdList, ShaderRHI, SpotlightMask, (LightSceneInfo->Proxy->GetLightType() == LightType_Spot ? 1.0f : 0.0f));
 
 		if (bShadowed)
 		{
-			SetShaderValue(ShaderRHI, DepthBiasParameters, FVector2D(ShadowMap->GetShaderDepthBias(), 1.0f / (ShadowMap->MaxSubjectZ - ShadowMap->MinSubjectZ)));
+			SetShaderValue(RHICmdList, ShaderRHI, DepthBiasParameters, FVector2D(ShadowMap->GetShaderDepthBias(), 1.0f / (ShadowMap->MaxSubjectZ - ShadowMap->MinSubjectZ)));
 
 			SetTextureParameter(
+				RHICmdList, 
 				ShaderRHI,
 				ShadowDepthTexture,
 				ShadowDepthTextureSampler,
@@ -897,11 +914,11 @@ public:
 
 		if (bShadowed && InjectionType == LightType_Point)
 		{
-			OnePassShadowParameters.Set(ShaderRHI, ShadowMap);
+			OnePassShadowParameters.Set(RHICmdList, ShaderRHI, ShadowMap);
 		}
 
-		LightFunctionParameters.Set(ShaderRHI, LightSceneInfo, 1);
-		TranslucentInjectParameters.Set(ShaderRHI, this, View, LightSceneInfo, ShadowMap, VolumeCascadeIndexValue, bShadowed);
+		LightFunctionParameters.Set(RHICmdList, ShaderRHI, LightSceneInfo, 1);
+		TranslucentInjectParameters.Set(RHICmdList, ShaderRHI, this, View, LightSceneInfo, ShadowMap, VolumeCascadeIndexValue, bShadowed);
 
 		if (LightFunctionWorldToLight.IsBound())
 		{
@@ -910,7 +927,7 @@ public:
 			const FVector InverseScale = FVector( 1.f / Scale.Z, 1.f / Scale.Y, 1.f / Scale.X );
 			const FMatrix WorldToLight = LightSceneInfo->Proxy->GetWorldToLight() * FScaleMatrix(FVector(InverseScale));	
 
-			SetShaderValue(ShaderRHI, LightFunctionWorldToLight, WorldToLight);
+			SetShaderValue(RHICmdList, ShaderRHI, LightFunctionWorldToLight, WorldToLight);
 		}
 	}
 
@@ -1011,8 +1028,10 @@ void RasterizeToVolumeTexture(FVolumeBounds VolumeBounds)
 
 /** Helper function that clears the given volume texture render targets. */
 template<int32 NumRenderTargets>
-void ClearVolumeTextures(const FTextureRHIParamRef* RenderTargets, const FLinearColor* ClearColors)
+void ClearVolumeTextures(FRHICommandList* RHICmdList, const FTextureRHIParamRef* RenderTargets, const FLinearColor* ClearColors)
 {
+	check(!RHICmdList);
+
 	RHISetRenderTargets(NumRenderTargets, RenderTargets, FTextureRHIRef(), 0, NULL);
 
 	static FGlobalBoundShaderState VolumeClearBoundShaderState;
@@ -1035,8 +1054,8 @@ void ClearVolumeTextures(const FTextureRHIParamRef* RenderTargets, const FLinear
 		
 		SetGlobalBoundShaderState(VolumeClearBoundShaderState, GScreenVertexDeclaration.VertexDeclarationRHI, *VertexShader, *PixelShader, *GeometryShader);
 
-		VertexShader->SetParameters(VolumeBounds, GTranslucencyLightingVolumeDim);
-		GeometryShader->SetParameters(VolumeBounds);
+		VertexShader->SetParameters(RHICmdList, VolumeBounds, GTranslucencyLightingVolumeDim);
+		GeometryShader->SetParameters(RHICmdList, VolumeBounds);
 
 		FLinearColor ShaderClearColors[MaxSimultaneousRenderTargets];
 		FMemory::MemZero(ShaderClearColors);
@@ -1046,13 +1065,13 @@ void ClearVolumeTextures(const FTextureRHIParamRef* RenderTargets, const FLinear
 			ShaderClearColors[i] = ClearColors[i];
 		}
 
-		SetShaderValueArray(PixelShader->GetPixelShader(), PixelShader->ColorParameter, ShaderClearColors, NumRenderTargets);
+		SetShaderValueArray(RHICmdList, PixelShader->GetPixelShader(), PixelShader->ColorParameter, ShaderClearColors, NumRenderTargets);
 
 		RasterizeToVolumeTexture(VolumeBounds);
 	}
 }
 
-void FDeferredShadingSceneRenderer::ClearTranslucentVolumeLighting()
+void FDeferredShadingSceneRenderer::ClearTranslucentVolumeLighting(FRHICommandList* RHICmdList)
 {
 	if (GUseTranslucentLightingVolumes && GSupportsVolumeTextureRendering)
 	{
@@ -1073,7 +1092,7 @@ void FDeferredShadingSceneRenderer::ClearTranslucentVolumeLighting()
 		ClearColors[2] = FLinearColor(0, 0, 0, 0);
 		ClearColors[3] = FLinearColor(0, 0, 0, 0);
 
-		ClearVolumeTextures<ARRAY_COUNT(RenderTargets)>(RenderTargets, ClearColors);
+		ClearVolumeTextures<ARRAY_COUNT(RenderTargets)>(RHICmdList, RenderTargets, ClearColors);
 	}
 }
 
@@ -1109,11 +1128,11 @@ public:
 		return bShaderHasOutdatedParameters;
 	}
 
-	void SetParameters(const FViewInfo& View, const FFinalPostProcessSettings::FCubemapEntry& CubemapEntry)
+	void SetParameters(FRHICommandList* RHICmdList, const FViewInfo& View, const FFinalPostProcessSettings::FCubemapEntry& CubemapEntry)
 	{
 		const FPixelShaderRHIParamRef ShaderRHI = GetPixelShader();
 
-		FGlobalShader::SetParameters(ShaderRHI, View);
+		FGlobalShader::SetParameters(RHICmdList, ShaderRHI, View);
 
 		CubemapShaderParameters.SetParameters(ShaderRHI, CubemapEntry);
 	}
@@ -1121,7 +1140,7 @@ public:
 
 IMPLEMENT_SHADER_TYPE(,FInjectAmbientCubemapPS,TEXT("TranslucentLightingShaders"),TEXT("InjectAmbientCubemapMainPS"),SF_Pixel);
 
-void FDeferredShadingSceneRenderer::InjectAmbientCubemapTranslucentVolumeLighting()
+void FDeferredShadingSceneRenderer::InjectAmbientCubemapTranslucentVolumeLighting(FRHICommandList* RHICmdList)
 {
 	//@todo - support multiple views
 	const FViewInfo& View = Views[0];
@@ -1149,15 +1168,15 @@ void FDeferredShadingSceneRenderer::InjectAmbientCubemapTranslucentVolumeLightin
 						
 			SetGlobalBoundShaderState(BoundShaderState, GScreenVertexDeclaration.VertexDeclarationRHI, *VertexShader, *PixelShader, *GeometryShader);
 
-			VertexShader->SetParameters(VolumeBounds, GTranslucencyLightingVolumeDim);
-			GeometryShader->SetParameters(VolumeBounds);
+			VertexShader->SetParameters(RHICmdList, VolumeBounds, GTranslucencyLightingVolumeDim);
+			GeometryShader->SetParameters(RHICmdList, VolumeBounds);
 
 			uint32 Count = View.FinalPostProcessSettings.ContributingCubemaps.Num();
 			for(uint32 i = 0; i < Count; ++i)
 			{
 				const FFinalPostProcessSettings::FCubemapEntry& CubemapEntry = View.FinalPostProcessSettings.ContributingCubemaps[i];
 
-				PixelShader->SetParameters(View, CubemapEntry);
+				PixelShader->SetParameters(RHICmdList, View, CubemapEntry);
 
 				RasterizeToVolumeTexture(VolumeBounds);
 			}
@@ -1195,20 +1214,20 @@ public:
 	}
 	FCompositeGIForTranslucencyPS() {}
 
-	void SetParameters(const FViewInfo& View, const FSceneViewState* ViewState, int32 VolumeCascadeIndexValue)
+	void SetParameters(FRHICommandList* RHICmdList, const FViewInfo& View, const FSceneViewState* ViewState, int32 VolumeCascadeIndexValue)
 	{
 		const FPixelShaderRHIParamRef ShaderRHI = GetPixelShader();
 
-		FGlobalShader::SetParameters(ShaderRHI, View);
+		FGlobalShader::SetParameters(RHICmdList, ShaderRHI, View);
 
-		SetShaderValue(ShaderRHI, VolumeCascadeIndex, VolumeCascadeIndexValue);
+		SetShaderValue(RHICmdList, ShaderRHI, VolumeCascadeIndex, VolumeCascadeIndexValue);
 
 		const FIndirectLightingCacheAllocation* LightingAllocation = ViewState->TranslucencyLightingCacheAllocations[VolumeCascadeIndexValue];
 
 		if (LightingAllocation && LightingAllocation->IsValid())
 		{
-			SetShaderValue(ShaderRHI, IndirectlightingCacheAdd, LightingAllocation->Add);
-			SetShaderValue(ShaderRHI, IndirectlightingCacheScale, LightingAllocation->Scale);
+			SetShaderValue(RHICmdList, ShaderRHI, IndirectlightingCacheAdd, LightingAllocation->Add);
+			SetShaderValue(RHICmdList, ShaderRHI, IndirectlightingCacheScale, LightingAllocation->Scale);
 
 			FScene* Scene = (FScene*)View.Family->Scene;
 			FIndirectLightingCache& LightingCache = Scene->IndirectLightingCache;
@@ -1216,6 +1235,7 @@ public:
 			if (LightingCache.IsInitialized())
 			{
 				SetTextureParameter(
+					RHICmdList, 
 					ShaderRHI,
 					IndirectLightingCacheTexture,
 					IndirectLightingCacheTextureSampler,
@@ -1224,6 +1244,7 @@ public:
 					);
 
 				SetTextureParameter(
+					RHICmdList, 
 					ShaderRHI,
 					IndirectLightingCacheTexture1,
 					IndirectLightingCacheTexture1Sampler,
@@ -1232,6 +1253,7 @@ public:
 					);
 
 				SetTextureParameter(
+					RHICmdList, 
 					ShaderRHI,
 					IndirectLightingCacheTexture2,
 					IndirectLightingCacheTexture2Sampler,
@@ -1271,7 +1293,7 @@ private:
 
 IMPLEMENT_SHADER_TYPE(,FCompositeGIForTranslucencyPS,TEXT("TranslucentLightingShaders"),TEXT("CompositeGIMainPS"),SF_Pixel);
 
-void FDeferredShadingSceneRenderer::CompositeIndirectTranslucentVolumeLighting()
+void FDeferredShadingSceneRenderer::CompositeIndirectTranslucentVolumeLighting(FRHICommandList* RHICmdList)
 {
 	bool bAnyViewAllowsIndirectLightingCache = false;
 
@@ -1316,9 +1338,9 @@ void FDeferredShadingSceneRenderer::CompositeIndirectTranslucentVolumeLighting()
 				static FGlobalBoundShaderState BoundShaderState;
 				SetGlobalBoundShaderState(BoundShaderState, GScreenVertexDeclaration.VertexDeclarationRHI, *VertexShader, *PixelShader, *GeometryShader);
 
-				PixelShader->SetParameters(View, ViewState, VolumeCascadeIndex);
-				VertexShader->SetParameters(VolumeBounds, GTranslucencyLightingVolumeDim);
-				GeometryShader->SetParameters(VolumeBounds);
+				PixelShader->SetParameters(RHICmdList, View, ViewState, VolumeCascadeIndex);
+				VertexShader->SetParameters(RHICmdList, VolumeBounds, GTranslucencyLightingVolumeDim);
+				GeometryShader->SetParameters(RHICmdList, VolumeBounds);
 
 				RasterizeToVolumeTexture(VolumeBounds);
 
@@ -1346,7 +1368,9 @@ void FDeferredShadingSceneRenderer::ClearTranslucentVolumePerObjectShadowing()
 		ClearColors[0] = FLinearColor(1, 1, 1, 1);
 		ClearColors[1] = FLinearColor(1, 1, 1, 1);
 
-		ClearVolumeTextures<ARRAY_COUNT(RenderTargets)>(RenderTargets, ClearColors);
+		//@todo-rco: RHIPacketList
+		FRHICommandList* RHICmdList = nullptr;
+		ClearVolumeTextures<ARRAY_COUNT(RenderTargets)>(RHICmdList, RenderTargets, ClearColors);
 	}
 }
 
@@ -1388,6 +1412,8 @@ void FDeferredShadingSceneRenderer::AccumulateTranslucentVolumeObjectShadowing(c
 		ClearTranslucentVolumePerObjectShadowing();
 	}
 
+	//@todo-rco: RHIPacketList
+	FRHICommandList* RHICmdList = nullptr;
 	if (GUseTranslucentLightingVolumes && GSupportsVolumeTextureRendering)
 	{
 		SCOPED_DRAW_EVENT(AccumulateTranslucentVolumeShadowing, DEC_SCENE_ITEMS);
@@ -1426,9 +1452,9 @@ void FDeferredShadingSceneRenderer::AccumulateTranslucentVolumeObjectShadowing(c
 				TShaderMapRef<FTranslucentObjectShadowingPS> PixelShader(GetGlobalShaderMap());
 
 				SetGlobalBoundShaderState(ObjectShadowingBoundShaderState, GScreenVertexDeclaration.VertexDeclarationRHI, *VertexShader, *PixelShader, *GeometryShader);
-				VertexShader->SetParameters(VolumeBounds, GTranslucencyLightingVolumeDim);
-				GeometryShader->SetParameters(VolumeBounds);
-				PixelShader->SetParameters(View, LightSceneInfo, InProjectedShadowInfo, VolumeCascadeIndex);
+				VertexShader->SetParameters(RHICmdList, VolumeBounds, GTranslucencyLightingVolumeDim);
+				GeometryShader->SetParameters(RHICmdList, VolumeBounds);
+				PixelShader->SetParameters(RHICmdList, View, LightSceneInfo, InProjectedShadowInfo, VolumeCascadeIndex);
 				
 				RasterizeToVolumeTexture(VolumeBounds);
 
@@ -1457,6 +1483,9 @@ void SetInjectionShader(
 	bool bApplyLightFunction,
 	bool bInverseSquared)
 {
+	//@todo-rco: RHIPacketList
+	FRHICommandList* RHICmdList = nullptr;
+
 	check(ShadowMap || !bShadowed);
 
 	const FMaterialShaderMap* MaterialShaderMap = MaterialProxy->GetMaterial(View.GetFeatureLevel())->GetRenderingThreadShaderMap();
@@ -1519,13 +1548,13 @@ void SetInjectionShader(
 		{
 			auto InjectionPixelShader = MaterialShaderMap->GetShader< TTranslucentLightingInjectPS<InjectionType, bShadowed, true, true && !Directional> >();
 			check(InjectionPixelShader);
-			InjectionPixelShader->SetParameters(View, LightSceneInfo, MaterialProxy, ShadowMap, InnerSplitIndex, VolumeCascadeIndexValue);
+			InjectionPixelShader->SetParameters(RHICmdList, View, LightSceneInfo, MaterialProxy, ShadowMap, InnerSplitIndex, VolumeCascadeIndexValue);
 		}
 		else
 		{
 			auto InjectionPixelShader = MaterialShaderMap->GetShader< TTranslucentLightingInjectPS<InjectionType, bShadowed, true, false> >();
 			check(InjectionPixelShader);
-			InjectionPixelShader->SetParameters(View, LightSceneInfo, MaterialProxy, ShadowMap, InnerSplitIndex, VolumeCascadeIndexValue);
+			InjectionPixelShader->SetParameters(RHICmdList, View, LightSceneInfo, MaterialProxy, ShadowMap, InnerSplitIndex, VolumeCascadeIndexValue);
 		}
 	}
 	else
@@ -1534,13 +1563,13 @@ void SetInjectionShader(
 		{
 			auto InjectionPixelShader = MaterialShaderMap->GetShader< TTranslucentLightingInjectPS<InjectionType, bShadowed, false, true && !Directional> >();
 			check(InjectionPixelShader);
-			InjectionPixelShader->SetParameters(View, LightSceneInfo, MaterialProxy, ShadowMap, InnerSplitIndex, VolumeCascadeIndexValue);
+			InjectionPixelShader->SetParameters(RHICmdList, View, LightSceneInfo, MaterialProxy, ShadowMap, InnerSplitIndex, VolumeCascadeIndexValue);
 		}
 		else
 		{
 			auto InjectionPixelShader = MaterialShaderMap->GetShader< TTranslucentLightingInjectPS<InjectionType, bShadowed, false, false> >();
 			check(InjectionPixelShader);
-			InjectionPixelShader->SetParameters(View, LightSceneInfo, MaterialProxy, ShadowMap, InnerSplitIndex, VolumeCascadeIndexValue);
+			InjectionPixelShader->SetParameters(RHICmdList, View, LightSceneInfo, MaterialProxy, ShadowMap, InnerSplitIndex, VolumeCascadeIndexValue);
 		}
 	}
 }
@@ -1603,6 +1632,9 @@ static void InjectTranslucentLightArray(const FViewInfo& View, const TArray<FTra
 {
 	INC_DWORD_STAT_BY(STAT_NumLightsInjectedIntoTranslucency, LightInjectionData.Num());
 
+	//@todo-rco: RHIPacketList
+	FRHICommandList* RHICmdList = nullptr;
+
 	// Inject into each volume cascade
 	// Operate on one cascade at a time to reduce render target switches
 	for (int32 VolumeCascadeIndex = 0; VolumeCascadeIndex < TVC_MAX; VolumeCascadeIndex++)
@@ -1671,8 +1703,8 @@ static void InjectTranslucentLightArray(const FViewInfo& View, const TArray<FTra
 					}
 				}
 
-				VertexShader->SetParameters(VolumeBounds, GTranslucencyLightingVolumeDim);
-				GeometryShader->SetParameters(VolumeBounds);
+				VertexShader->SetParameters(RHICmdList, VolumeBounds, GTranslucencyLightingVolumeDim);
+				GeometryShader->SetParameters(RHICmdList, VolumeBounds);
 				RasterizeToVolumeTexture(VolumeBounds);
 			}
 		}
@@ -1743,14 +1775,14 @@ public:
 	}
 	FSimpleLightTranslucentLightingInjectPS() {}
 
-	void SetParameters(const FViewInfo& View, const FSimpleLightEntry& SimpleLight, const FSimpleLightPerViewEntry& SimpleLightPerViewData, int32 VolumeCascadeIndexValue)
+	void SetParameters(FRHICommandList* RHICmdList, const FViewInfo& View, const FSimpleLightEntry& SimpleLight, const FSimpleLightPerViewEntry& SimpleLightPerViewData, int32 VolumeCascadeIndexValue)
 	{
-		FGlobalShader::SetParameters(GetPixelShader(), View);
+		FGlobalShader::SetParameters(RHICmdList, GetPixelShader(), View);
 
 		FVector4 PositionAndRadius(SimpleLightPerViewData.Position, SimpleLight.Radius);
-		SetShaderValue(GetPixelShader(), VolumeCascadeIndex, VolumeCascadeIndexValue);
-		SetShaderValue(GetPixelShader(), SimpleLightPositionAndRadius, PositionAndRadius);
-		SetShaderValue(GetPixelShader(), SimpleLightColorAndExponent, FVector4(SimpleLight.Color, SimpleLight.Exponent));
+		SetShaderValue(RHICmdList, GetPixelShader(), VolumeCascadeIndex, VolumeCascadeIndexValue);
+		SetShaderValue(RHICmdList, GetPixelShader(), SimpleLightPositionAndRadius, PositionAndRadius);
+		SetShaderValue(RHICmdList, GetPixelShader(), SimpleLightColorAndExponent, FVector4(SimpleLight.Color, SimpleLight.Exponent));
 	}
 
 	virtual bool Serialize(FArchive& Ar)
@@ -1774,6 +1806,9 @@ FGlobalBoundShaderState InjectSimpleLightBoundShaderState;
 
 void FDeferredShadingSceneRenderer::InjectSimpleTranslucentVolumeLightingArray(const FSimpleLightArray& SimpleLights)
 {
+	//@todo-rco: RHIPacketList
+	FRHICommandList* RHICmdList = nullptr;
+
 	SCOPE_CYCLE_COUNTER(STAT_TranslucentInjectTime);
 
 	int32 NumLightsToInject = 0;
@@ -1788,7 +1823,7 @@ void FDeferredShadingSceneRenderer::InjectSimpleTranslucentVolumeLightingArray(c
 
 	if (NumLightsToInject > 0)
 	{
-		//@todo - support multiple views 
+		//@todo - support multiple views
 		const FViewInfo& View = Views[0];
 		const int32 ViewIndex = 0;
 
@@ -1831,9 +1866,9 @@ void FDeferredShadingSceneRenderer::InjectSimpleTranslucentVolumeLightingArray(c
 						TShaderMapRef<FSimpleLightTranslucentLightingInjectPS> PixelShader(GetGlobalShaderMap());
 						SetGlobalBoundShaderState(InjectSimpleLightBoundShaderState, GScreenVertexDeclaration.VertexDeclarationRHI, *VertexShader, *PixelShader, *GeometryShader);
 
-						VertexShader->SetParameters(VolumeBounds, GTranslucencyLightingVolumeDim);
-						GeometryShader->SetParameters(VolumeBounds);
-						PixelShader->SetParameters(View, SimpleLight, SimpleLightPerViewData, VolumeCascadeIndex);
+						VertexShader->SetParameters(RHICmdList, VolumeBounds, GTranslucencyLightingVolumeDim);
+						GeometryShader->SetParameters(RHICmdList, VolumeBounds);
+						PixelShader->SetParameters(RHICmdList, View, SimpleLight, SimpleLightPerViewData, VolumeCascadeIndex);
 
 						// Accumulate the contribution of multiple lights
 						RHISetBlendState(TStaticBlendState<
@@ -1859,6 +1894,9 @@ void FDeferredShadingSceneRenderer::FilterTranslucentVolumeLighting()
 	{
 		if (GUseTranslucencyVolumeBlur)
 		{
+			//@todo-rco: RHIPacketList
+			FRHICommandList* RHICmdList = nullptr;
+
 			SCOPED_DRAW_EVENT(FilterTranslucentVolume, DEC_SCENE_ITEMS);
 
 			RHISetRasterizerState(TStaticRasterizerState<FM_Solid,CM_None>::GetRHI());
@@ -1889,9 +1927,9 @@ void FDeferredShadingSceneRenderer::FilterTranslucentVolumeLighting()
 				TShaderMapRef<FFilterTranslucentVolumePS> PixelShader(GetGlobalShaderMap());
 				SetGlobalBoundShaderState(FilterBoundShaderState, GScreenVertexDeclaration.VertexDeclarationRHI, *VertexShader, *PixelShader, *GeometryShader);
 
-				VertexShader->SetParameters(VolumeBounds, GTranslucencyLightingVolumeDim);
-				GeometryShader->SetParameters(VolumeBounds);
-				PixelShader->SetParameters(View, VolumeCascadeIndex);
+				VertexShader->SetParameters(RHICmdList, VolumeBounds, GTranslucencyLightingVolumeDim);
+				GeometryShader->SetParameters(RHICmdList, VolumeBounds);
+				PixelShader->SetParameters(RHICmdList, View, VolumeCascadeIndex);
 
 				RasterizeToVolumeTexture(VolumeBounds);
 
diff --git a/Engine/Source/Runtime/Renderer/Private/TranslucentRendering.cpp b/Engine/Source/Runtime/Renderer/Private/TranslucentRendering.cpp
--- a/Engine/Source/Runtime/Renderer/Private/TranslucentRendering.cpp
+++ b/Engine/Source/Runtime/Renderer/Private/TranslucentRendering.cpp
@@ -146,9 +146,9 @@ public:
 	}
 	FCopySceneColorPS() {}
 
-	void SetParameters(const FViewInfo& View)
+	void SetParameters(FRHICommandList* RHICmdList, const FViewInfo& View)
 	{
-		SceneTextureParameters.Set(GetPixelShader(), View);
+		SceneTextureParameters.Set(RHICmdList, GetPixelShader(), View);
 	}
 
 	virtual bool Serialize(FArchive& Ar)
@@ -182,7 +182,8 @@ void FDeferredShadingSceneRenderer::CopySceneColor(const FViewInfo& View, const
 	TShaderMapRef<FCopySceneColorPS> PixelShader(GetGlobalShaderMap());
 	SetGlobalBoundShaderState(CopySceneColorBoundShaderState, GFilterVertexDeclaration.VertexDeclarationRHI, *ScreenVertexShader, *PixelShader);
 
-	PixelShader->SetParameters(View);
+	/// ?
+	PixelShader->SetParameters(nullptr, View);
 
 	DrawRectangle( 
 		0, 0, 
@@ -247,6 +248,7 @@ public:
 	/** Draws the translucent mesh with a specific light-map type, and fog volume type */
 	template<typename LightMapPolicyType>
 	void Process(
+		FRHICommandList* RHICmdList, 
 		const FProcessBasePassMeshParameters& Parameters,
 		const LightMapPolicyType& LightMapPolicy,
 		const typename LightMapPolicyType::ElementDataType& LightMapElementData
@@ -270,6 +272,7 @@ public:
 			Parameters.bAllowFog
 			);
 		DrawingPolicy.DrawShared(
+			RHICmdList, 
 			&View,
 			DrawingPolicy.CreateBoundShaderState(View.GetFeatureLevel())
 			);
@@ -281,6 +284,7 @@ public:
 			if(BatchElementMask & 1)
 			{
 				DrawingPolicy.SetMeshRenderState(
+					RHICmdList, 
 					View,
 					Parameters.PrimitiveSceneProxy,
 					Parameters.Mesh,
@@ -288,7 +292,7 @@ public:
 					bBackFace,
 					typename TBasePassDrawingPolicy<LightMapPolicyType>::ElementDataType(LightMapElementData)
 					);
-				DrawingPolicy.DrawMesh(Parameters.Mesh,BatchElementIndex);
+				DrawingPolicy.DrawMesh(RHICmdList, Parameters.Mesh,BatchElementIndex);
 			}
 
 			BatchElementMask >>= 1;
@@ -320,6 +324,9 @@ bool FTranslucencyDrawingPolicyFactory::DrawMesh(
 	const EBlendMode BlendMode = Material->GetBlendMode();
 	const EMaterialShadingModel ShadingModel = Material->GetShadingModel();
 
+	//@todo-rco: RHIPacketList
+	FRHICommandList* RHICmdList = nullptr;
+
 	// Only render translucent materials.
 	if(IsTranslucentBlendMode(BlendMode))
 	{
@@ -346,6 +353,7 @@ bool FTranslucencyDrawingPolicyFactory::DrawMesh(
 		}
 
 		ProcessBasePassMesh(
+			RHICmdList, 
 			FProcessBasePassMeshParameters(
 				Mesh,
 				BatchElementMask,
@@ -383,6 +391,7 @@ bool FTranslucencyDrawingPolicyFactory::DrawMesh(
  * @return true if the mesh rendered
  */
 bool FTranslucencyDrawingPolicyFactory::DrawDynamicMesh(
+	FRHICommandList* RHICmdList, 
 	const FViewInfo& View,
 	ContextType DrawingContext,
 	const FMeshBatch& Mesh,
@@ -408,6 +417,7 @@ bool FTranslucencyDrawingPolicyFactory::DrawDynamicMesh(
  * @return true if the mesh rendered
  */
 bool FTranslucencyDrawingPolicyFactory::DrawStaticMesh(
+	FRHICommandList* RHICmdList, 
 	const FViewInfo& View,
 	ContextType DrawingContext,
 	const FStaticMesh& StaticMesh,
@@ -513,6 +523,9 @@ void FTranslucentPrimSet::RenderPrimitive(
 				);
 		}
 
+		//@todo-rco: RHIPacketList
+		FRHICommandList* RHICmdList = nullptr;
+
 		// Render static scene prim
 		if( ViewRelevance.bStaticRelevance )
 		{
@@ -526,6 +539,7 @@ void FTranslucentPrimSet::RenderPrimitive(
 					&& (StaticMesh.MaterialRenderProxy->GetMaterial(FeatureLevel)->IsSeparateTranslucencyEnabled() == bSeparateTranslucencyPass))
 				{
 					FTranslucencyDrawingPolicyFactory::DrawStaticMesh(
+						RHICmdList, 
 						View,
 						FTranslucencyDrawingPolicyFactory::ContextType(TranslucentSelfShadow),
 						StaticMesh,
diff --git a/Engine/Source/Runtime/Renderer/Private/TranslucentRendering.h b/Engine/Source/Runtime/Renderer/Private/TranslucentRendering.h
--- a/Engine/Source/Runtime/Renderer/Private/TranslucentRendering.h
+++ b/Engine/Source/Runtime/Renderer/Private/TranslucentRendering.h
@@ -28,6 +28,7 @@ public:
 	* @return true if the mesh rendered
 	*/
 	static bool DrawDynamicMesh(
+		FRHICommandList* RHICmdList, 
 		const FViewInfo& View,
 		ContextType DrawingContext,
 		const FMeshBatch& Mesh,
@@ -42,6 +43,7 @@ public:
 	* @return true if the mesh rendered
 	*/
 	static bool DrawStaticMesh(
+		FRHICommandList* RHICmdList, 
 		const FViewInfo& View,
 		ContextType DrawingContext,
 		const FStaticMesh& StaticMesh,
@@ -93,6 +95,7 @@ public:
 	* @return true if the mesh rendered
 	*/
 	static bool DrawDynamicMesh(
+		FRHICommandList* RHICmdList, 
 		const FViewInfo& View,
 		ContextType DrawingContext,
 		const FMeshBatch& Mesh,
@@ -107,6 +110,7 @@ public:
 	* @return true if the mesh rendered
 	*/
 	static bool DrawStaticMesh(
+		FRHICommandList* RHICmdList, 
 		const FViewInfo& View,
 		ContextType DrawingContext,
 		const FStaticMesh& StaticMesh,
@@ -171,10 +175,10 @@ public:
 
 	FWriteToSliceVS() {}
 
-	void SetParameters(const FVolumeBounds& VolumeBounds, uint32 VolumeResolution)
+	void SetParameters(FRHICommandList* RHICmdList, const FVolumeBounds& VolumeBounds, uint32 VolumeResolution)
 	{
 		const float InvVolumeResolution = 1.0f / VolumeResolution;
-		SetShaderValue(GetVertexShader(), UVScaleBias, FVector4(
+		SetShaderValue(RHICmdList, GetVertexShader(), UVScaleBias, FVector4(
 			(VolumeBounds.MaxX - VolumeBounds.MinX) * InvVolumeResolution, 
 			(VolumeBounds.MaxY - VolumeBounds.MinY) * InvVolumeResolution,
 			VolumeBounds.MinX * InvVolumeResolution,
@@ -210,9 +214,9 @@ public:
 	}
 	FWriteToSliceGS() {}
 
-	void SetParameters(const FVolumeBounds& VolumeBounds)
+	void SetParameters(FRHICommandList* RHICmdList, const FVolumeBounds& VolumeBounds)
 	{
-		SetShaderValue(GetGeometryShader(), MinZ, VolumeBounds.MinZ);
+		SetShaderValue(RHICmdList, GetGeometryShader(), MinZ, VolumeBounds.MinZ);
 	}
 
 	virtual bool Serialize(FArchive& Ar)
diff --git a/Engine/Source/Runtime/Renderer/Private/VelocityRendering.cpp b/Engine/Source/Runtime/Renderer/Private/VelocityRendering.cpp
--- a/Engine/Source/Runtime/Renderer/Private/VelocityRendering.cpp
+++ b/Engine/Source/Runtime/Renderer/Private/VelocityRendering.cpp
@@ -15,16 +15,16 @@ class FVelocityVS : public FMeshMaterialShader
 
 public:
 
-	void SetParameters(const FVertexFactory* VertexFactory,const FMaterialRenderProxy* MaterialRenderProxy,const FViewInfo& View)
+	void SetParameters(FRHICommandList* RHICmdList, const FVertexFactory* VertexFactory,const FMaterialRenderProxy* MaterialRenderProxy,const FViewInfo& View)
 	{
-		FMeshMaterialShader::SetParameters(GetVertexShader(), MaterialRenderProxy, *MaterialRenderProxy->GetMaterial(View.GetFeatureLevel()), View, ESceneRenderTargetsMode::DontSet);
+		FMeshMaterialShader::SetParameters(RHICmdList, GetVertexShader(), MaterialRenderProxy, *MaterialRenderProxy->GetMaterial(View.GetFeatureLevel()), View, ESceneRenderTargetsMode::DontSet);
 	}
 
-	void SetMesh(const FVertexFactory* VertexFactory, const FMeshBatch& Mesh, int32 BatchElementIndex, const FViewInfo& View, const FPrimitiveSceneProxy* Proxy, const FMatrix& InPreviousLocalToWorld)
+	void SetMesh(FRHICommandList* RHICmdList, const FVertexFactory* VertexFactory, const FMeshBatch& Mesh, int32 BatchElementIndex, const FViewInfo& View, const FPrimitiveSceneProxy* Proxy, const FMatrix& InPreviousLocalToWorld)
 	{
-		FMeshMaterialShader::SetMesh(GetVertexShader(), VertexFactory, View, Proxy, Mesh.Elements[BatchElementIndex]);
+		FMeshMaterialShader::SetMesh(RHICmdList, GetVertexShader(), VertexFactory, View, Proxy, Mesh.Elements[BatchElementIndex]);
 
-		SetShaderValue(GetVertexShader(), PreviousLocalToWorld, InPreviousLocalToWorld.ConcatTranslation(View.PrevViewMatrices.PreViewTranslation));
+		SetShaderValue(RHICmdList, GetVertexShader(), PreviousLocalToWorld, InPreviousLocalToWorld.ConcatTranslation(View.PrevViewMatrices.PreViewTranslation));
 	}
 
 	bool SupportsVelocity() const
@@ -91,9 +91,9 @@ class FVelocityDS : public FBaseDS
 
 public:
 
-	void SetParameters(const FMaterialRenderProxy* MaterialRenderProxy,const FViewInfo& View)
+	void SetParameters(FRHICommandList* RHICmdList, const FMaterialRenderProxy* MaterialRenderProxy,const FViewInfo& View)
 	{
-		FMeshMaterialShader::SetParameters(GetDomainShader(), MaterialRenderProxy, *MaterialRenderProxy->GetMaterial(View.GetFeatureLevel()), View, ESceneRenderTargetsMode::DontSet);
+		FMeshMaterialShader::SetParameters(RHICmdList, GetDomainShader(), MaterialRenderProxy, *MaterialRenderProxy->GetMaterial(View.GetFeatureLevel()), View, ESceneRenderTargetsMode::DontSet);
 	}
 
 protected:
@@ -144,16 +144,16 @@ public:
 	{
 	}
 
-	void SetParameters(const FVertexFactory* VertexFactory,const FMaterialRenderProxy* MaterialRenderProxy,const FViewInfo& View)
+	void SetParameters(FRHICommandList* RHICmdList, const FVertexFactory* VertexFactory,const FMaterialRenderProxy* MaterialRenderProxy,const FViewInfo& View)
 	{
-		FMeshMaterialShader::SetParameters(GetPixelShader(), MaterialRenderProxy, *MaterialRenderProxy->GetMaterial(View.GetFeatureLevel()), View, ESceneRenderTargetsMode::DontSet);
+		FMeshMaterialShader::SetParameters(RHICmdList, GetPixelShader(), MaterialRenderProxy, *MaterialRenderProxy->GetMaterial(View.GetFeatureLevel()), View, ESceneRenderTargetsMode::DontSet);
 	}
 
-	void SetMesh(const FVertexFactory* VertexFactory, const FMeshBatch& Mesh,int32 BatchElementIndex,const FViewInfo& View, const FPrimitiveSceneProxy* Proxy, bool bBackFace)
+	void SetMesh(FRHICommandList* RHICmdList, const FVertexFactory* VertexFactory, const FMeshBatch& Mesh,int32 BatchElementIndex,const FViewInfo& View, const FPrimitiveSceneProxy* Proxy, bool bBackFace)
 	{
 		const FPixelShaderRHIParamRef ShaderRHI = GetPixelShader();
 		
-		FMeshMaterialShader::SetMesh(ShaderRHI, VertexFactory, View, Proxy, Mesh.Elements[BatchElementIndex]);
+		FMeshMaterialShader::SetMesh(RHICmdList, ShaderRHI, VertexFactory, View, Proxy, Mesh.Elements[BatchElementIndex]);
 	}
 
 	virtual bool Serialize(FArchive& Ar)
@@ -206,27 +206,29 @@ bool FVelocityDrawingPolicy::SupportsVelocity() const
 	return (VertexShader && PixelShader) ? VertexShader->SupportsVelocity() : false;
 }
 
-void FVelocityDrawingPolicy::DrawShared(const FSceneView* SceneView, FBoundShaderStateRHIRef ShaderState) const
+void FVelocityDrawingPolicy::DrawShared(FRHICommandList* RHICmdList, const FSceneView* SceneView, FBoundShaderStateRHIRef ShaderState) const
 {
 	// NOTE: Assuming this cast is always safe!
 	FViewInfo* View = (FViewInfo*)SceneView;
 
 	// Set the depth-only shader parameters for the material.
-	RHISetBoundShaderState(ShaderState);
-	VertexShader->SetParameters(VertexFactory, MaterialRenderProxy, *View);
-	PixelShader->SetParameters(VertexFactory, MaterialRenderProxy, *View);
+	FRHICommandList::SetBoundShaderState(RHICmdList, ShaderState);
+
+	VertexShader->SetParameters(RHICmdList, VertexFactory, MaterialRenderProxy, *View);
+	PixelShader->SetParameters(RHICmdList, VertexFactory, MaterialRenderProxy, *View);
 
 	if(HullShader && DomainShader)
 	{
-		HullShader->SetParameters(MaterialRenderProxy, *View);
-		DomainShader->SetParameters(MaterialRenderProxy, *View);
+		HullShader->SetParameters(RHICmdList, MaterialRenderProxy, *View);
+		DomainShader->SetParameters(RHICmdList, MaterialRenderProxy, *View);
 	}
 
 	// Set the shared mesh resources.
-	FMeshDrawingPolicy::DrawShared( View );
+	FMeshDrawingPolicy::DrawShared(RHICmdList, View );
 }
 
 void FVelocityDrawingPolicy::SetMeshRenderState(
+	FRHICommandList* RHICmdList, 
 	const FSceneView& SceneView,
 	const FPrimitiveSceneProxy* PrimitiveSceneProxy,
 	const FMeshBatch& Mesh,
@@ -246,22 +248,22 @@ void FVelocityDrawingPolicy::SetMeshRenderState(
 	// previous transform can be stored in the scene for each primitive
 	if(Scene->MotionBlurInfoData.GetPrimitiveMotionBlurInfo(PrimitiveSceneProxy->GetPrimitiveSceneInfo(), PreviousLocalToWorld))
 	{
-		VertexShader->SetMesh(VertexFactory, Mesh, BatchElementIndex, View, PrimitiveSceneProxy, PreviousLocalToWorld);
+		VertexShader->SetMesh(RHICmdList, VertexFactory, Mesh, BatchElementIndex, View, PrimitiveSceneProxy, PreviousLocalToWorld);
 	}	
 	else
 	{
 		const FMatrix& LocalToWorld = PrimitiveSceneProxy->GetLocalToWorld();		// different implementation from UE4
-		VertexShader->SetMesh(VertexFactory, Mesh, BatchElementIndex, View, PrimitiveSceneProxy, LocalToWorld);
+		VertexShader->SetMesh(RHICmdList, VertexFactory, Mesh, BatchElementIndex, View, PrimitiveSceneProxy, LocalToWorld);
 	}
 
 	if (HullShader && DomainShader)
 	{
-		DomainShader->SetMesh(VertexFactory, View, PrimitiveSceneProxy, Mesh.Elements[BatchElementIndex]);
-		HullShader->SetMesh(VertexFactory, View, PrimitiveSceneProxy, Mesh.Elements[BatchElementIndex]);
+		DomainShader->SetMesh(RHICmdList, VertexFactory, View, PrimitiveSceneProxy, Mesh.Elements[BatchElementIndex]);
+		HullShader->SetMesh(RHICmdList, VertexFactory, View, PrimitiveSceneProxy, Mesh.Elements[BatchElementIndex]);
 	}
 
-	PixelShader->SetMesh(VertexFactory, Mesh, BatchElementIndex, View, PrimitiveSceneProxy, bBackFace);
-	FMeshDrawingPolicy::SetMeshRenderState(View, PrimitiveSceneProxy, Mesh, BatchElementIndex, bBackFace, ElementData);
+	PixelShader->SetMesh(RHICmdList, VertexFactory, Mesh, BatchElementIndex, View, PrimitiveSceneProxy, bBackFace);
+	FMeshDrawingPolicy::SetMeshRenderState(RHICmdList, View, PrimitiveSceneProxy, Mesh, BatchElementIndex, bBackFace, ElementData);
 }
 
 bool FVelocityDrawingPolicy::HasVelocity(const FViewInfo& View, const FPrimitiveSceneInfo* PrimitiveSceneInfo)
@@ -353,7 +355,7 @@ void FVelocityDrawingPolicyFactory::AddStaticMesh(FScene* Scene, FStaticMesh* St
 	    EBlendMode BlendMode = Material->GetBlendMode();
 	    if(BlendMode == BLEND_Opaque || BlendMode == BLEND_Masked)
 	    {
-			if (!Material->IsMasked() && !Material->IsTwoSided() && !Material->MaterialModifiesMeshPosition())
+		    if(!Material->IsMasked() && !Material->IsTwoSided() && !Material->MaterialModifiesMeshPosition())
 		    {
 			    // Default material doesn't handle masked or mesh-mod, and doesn't have the correct bIsTwoSided setting.
 			    MaterialRenderProxy = UMaterial::GetDefaultMaterial(MD_Surface)->GetRenderProxy(false);
@@ -370,6 +372,7 @@ void FVelocityDrawingPolicyFactory::AddStaticMesh(FScene* Scene, FStaticMesh* St
 }
 
 bool FVelocityDrawingPolicyFactory::DrawDynamicMesh(
+	FRHICommandList* RHICmdList, 
 	const FViewInfo& View,
 	ContextType DrawingContext,
 	const FMeshBatch& Mesh,
@@ -390,7 +393,7 @@ bool FVelocityDrawingPolicyFactory::DrawDynamicMesh(
 		// This should be enforced at a higher level
 		//@todo - figure out why this is failing and re-enable
 		//check(FVelocityDrawingPolicy::HasVelocity(View, PrimitiveSceneInfo));
-		if (!Material->IsMasked() && !Material->IsTwoSided() && !Material->MaterialModifiesMeshPosition())
+		if(!Material->IsMasked() && !Material->IsTwoSided() && !Material->MaterialModifiesMeshPosition())
 		{
 			// Default material doesn't handle masked, and doesn't have the correct bIsTwoSided setting.
 			MaterialRenderProxy = UMaterial::GetDefaultMaterial(MD_Surface)->GetRenderProxy(false);
@@ -398,11 +401,11 @@ bool FVelocityDrawingPolicyFactory::DrawDynamicMesh(
 		FVelocityDrawingPolicy DrawingPolicy(Mesh.VertexFactory, MaterialRenderProxy, *MaterialRenderProxy->GetMaterial(FeatureLevel));
 		if(DrawingPolicy.SupportsVelocity())
 		{			
-			DrawingPolicy.DrawShared(&View,DrawingPolicy.CreateBoundShaderState(FeatureLevel));
+			DrawingPolicy.DrawShared(RHICmdList, &View,DrawingPolicy.CreateBoundShaderState(FeatureLevel));
 			for(int32 BatchElementIndex = 0; BatchElementIndex < Mesh.Elements.Num(); ++BatchElementIndex)
 			{
-				DrawingPolicy.SetMeshRenderState(View, PrimitiveSceneProxy, Mesh, BatchElementIndex, bBackFace, FMeshDrawingPolicy::ElementDataType());
-				DrawingPolicy.DrawMesh(Mesh, BatchElementIndex);
+				DrawingPolicy.SetMeshRenderState(RHICmdList, View, PrimitiveSceneProxy, Mesh, BatchElementIndex, bBackFace, FMeshDrawingPolicy::ElementDataType());
+				DrawingPolicy.DrawMesh(RHICmdList, Mesh, BatchElementIndex);
 			}
 			return true;
 		}
@@ -469,6 +472,9 @@ void FDeferredShadingSceneRenderer::RenderVelocities(const FViewInfo& View, TRef
 
 	GPrevPerBoneMotionBlur.LockData();
 
+	//@todo-rco: RHIPacketList
+	FRHICommandList* RHICmdList = nullptr;
+
 	const uint32 MinX = View.ViewRect.Min.X * VelocityBufferSize.X / BufferSize.X;
 	const uint32 MinY = View.ViewRect.Min.Y * VelocityBufferSize.Y / BufferSize.Y;
 	const uint32 MaxX = View.ViewRect.Max.X * VelocityBufferSize.X / BufferSize.X;
@@ -491,7 +497,7 @@ void FDeferredShadingSceneRenderer::RenderVelocities(const FViewInfo& View, TRef
 	RHISetRasterizerState(GetStaticRasterizerState<true>(FM_Solid, CM_CW));
 
 	// Draw velocities for movable static meshes.
-	Scene->VelocityDrawList.DrawVisible(View,View.StaticMeshVelocityMap, View.StaticMeshBatchVisibility);
+	Scene->VelocityDrawList.DrawVisible(RHICmdList, View,View.StaticMeshVelocityMap, View.StaticMeshBatchVisibility);
 
 	// Draw velocities for movable dynamic meshes.
 	TDynamicPrimitiveDrawer<FVelocityDrawingPolicyFactory> Drawer(
diff --git a/Engine/Source/Runtime/Renderer/Private/VelocityRendering.h b/Engine/Source/Runtime/Renderer/Private/VelocityRendering.h
--- a/Engine/Source/Runtime/Renderer/Private/VelocityRendering.h
+++ b/Engine/Source/Runtime/Renderer/Private/VelocityRendering.h
@@ -27,8 +27,9 @@ public:
 			&& VertexShader == Other.VertexShader
 			&& PixelShader == Other.PixelShader;
 	}
-	void DrawShared( const FSceneView* View, FBoundShaderStateRHIRef ShaderState ) const;
+	void DrawShared(FRHICommandList* RHICmdList, const FSceneView* View, FBoundShaderStateRHIRef ShaderState) const;
 	void SetMeshRenderState(
+		FRHICommandList* RHICmdList, 
 		const FSceneView& View,
 		const FPrimitiveSceneProxy* PrimitiveSceneProxy,
 		const FMeshBatch& Mesh,
@@ -61,6 +62,7 @@ class FVelocityDrawingPolicyFactory : public FDepthDrawingPolicyFactory
 public:
 	static void AddStaticMesh(FScene* Scene, FStaticMesh* StaticMesh, ContextType = ContextType(DDM_AllOccluders));
 	static bool DrawDynamicMesh(	
+		FRHICommandList* RHICmdList, 
 		const FViewInfo& View,
 		ContextType DrawingContext,
 		const FMeshBatch& Mesh,
diff --git a/Engine/Source/Runtime/ShaderCore/Private/Shader.cpp b/Engine/Source/Runtime/ShaderCore/Private/Shader.cpp
--- a/Engine/Source/Runtime/ShaderCore/Private/Shader.cpp
+++ b/Engine/Source/Runtime/ShaderCore/Private/Shader.cpp
@@ -907,7 +907,8 @@ void FShader::FinishCleanup()
 void FShader::VerifyBoundUniformBufferParameters()
 {
 	// Support being called on a NULL pointer
-	if (this)
+// TODO: doesn't work with uniform buffer parameters on helper structs like FDeferredPixelShaderParameters
+	if (0&&this)
 	{
 		for (int32 StructIndex = 0; StructIndex < UniformBufferParameters.Num(); StructIndex++)
 		{
@@ -1032,21 +1033,27 @@ FShaderType* FindShaderTypeByName(const TCHAR* ShaderTypeName)
 
 
 void DispatchComputeShader(
+	FRHICommandList* RHICmdList, 
 	FShader* Shader,
 	uint32 ThreadGroupCountX,
 	uint32 ThreadGroupCountY,
 	uint32 ThreadGroupCountZ)
-{	
+{
+	check(!RHICmdList);
+
 	Shader->VerifyBoundUniformBufferParameters();
 	RHIDispatchComputeShader(ThreadGroupCountX, ThreadGroupCountY, ThreadGroupCountZ);
 }
 
 
 void DispatchIndirectComputeShader(
+	FRHICommandList* RHICmdList, 
 	FShader* Shader,
 	FVertexBufferRHIParamRef ArgumentBuffer,
 	uint32 ArgumentOffset)
-{	
+{
+	check(!RHICmdList);
+
 	Shader->VerifyBoundUniformBufferParameters();
 	RHIDispatchIndirectComputeShader(ArgumentBuffer, ArgumentOffset);
 }
diff --git a/Engine/Source/Runtime/ShaderCore/Private/ShaderCore.cpp b/Engine/Source/Runtime/ShaderCore/Private/ShaderCore.cpp
--- a/Engine/Source/Runtime/ShaderCore/Private/ShaderCore.cpp
+++ b/Engine/Source/Runtime/ShaderCore/Private/ShaderCore.cpp
@@ -150,6 +150,11 @@ void FShaderParameterMap::AddParameterAllocation(const TCHAR* ParameterName,uint
 	ParameterMap.Add(ParameterName,Allocation);
 }
 
+void FShaderParameterMap::RemoveParameterAllocation(const TCHAR* ParameterName)
+{
+	ParameterMap.Remove(ParameterName);
+}
+
 void FShaderCompilerOutput::GenerateOutputHash()
 {
 	FSHA1 HashState;
diff --git a/Engine/Source/Runtime/ShaderCore/Private/ShaderParameters.cpp b/Engine/Source/Runtime/ShaderCore/Private/ShaderParameters.cpp
--- a/Engine/Source/Runtime/ShaderCore/Private/ShaderParameters.cpp
+++ b/Engine/Source/Runtime/ShaderCore/Private/ShaderParameters.cpp
@@ -89,7 +89,9 @@ void FShaderUniformBufferParameter::ModifyCompilationEnvironment(const TCHAR* Pa
 		);
 
 	FString& GeneratedUniformBuffersInclude = OutEnvironment.IncludeFileNameToContentsMap.FindOrAdd("GeneratedUniformBuffers.usf");
-	GeneratedUniformBuffersInclude += FString::Printf(TEXT("#include \"UniformBuffers/%s.usf\"") LINE_TERMINATOR, ParameterName);
+	GeneratedUniformBuffersInclude += FString::Printf(TEXT("#include \"UniformBuffers/%s.usf\"")LINE_TERMINATOR, ParameterName);
+
+	Struct.AddResourceTableEntries(OutEnvironment.ResourceTableMap, OutEnvironment.ResourceTableLayoutHashes);
 }
 
 void FShaderUniformBufferParameter::Bind(const FShaderParameterMap& ParameterMap,const TCHAR* ParameterName,EShaderParameterFlags Flags)
@@ -125,17 +127,40 @@ void FShaderUniformBufferParameter::Bind(const FShaderParameterMap& ParameterMap
 	}
 }
 
+/** The individual bits of a uniform buffer declaration. */
+struct FUniformBufferDecl
+{
+	/** Members to place in the constant buffer. */
+	FString ConstantBufferMembers;
+	/** Members to place in the resource table. */
+	FString ResourceMembers;
+	/** Members in the struct HLSL shader code will access. */
+	FString StructMembers;
+	/** The HLSL initializer that will copy constants and resources in to the struct. */
+	FString Initializer;
+};
+
 /** Generates a HLSL struct declaration for a uniform buffer struct. */
-static FString CreateHLSLUniformBufferStructMembersDeclaration(const FUniformBufferStruct& UniformBufferStruct, bool bIgnorePadding)
+static void CreateHLSLUniformBufferStructMembersDeclaration(FUniformBufferDecl& Decl, const FUniformBufferStruct& UniformBufferStruct, const FString& NamePrefix, bool bExplicitPadding)
 {
-	FString Result;
 	uint32 HLSLBaseOffset = 0;
-
+	bool bFoundResourceMember = false;
+	int32 MemberIndex = 0;
 	const TArray<FUniformBufferStruct::FMember>& StructMembers = UniformBufferStruct.GetMembers();
-	for(int32 MemberIndex = 0;MemberIndex < StructMembers.Num();++MemberIndex)
+	
+	Decl.Initializer += TEXT("{");
+	int32 OpeningBraceLocPlusOne = Decl.Initializer.Len();
+
+	for(;MemberIndex < StructMembers.Num();++MemberIndex)
 	{
 		const FUniformBufferStruct::FMember& Member = StructMembers[MemberIndex];
 		
+		if (IsUniformBufferResourceType(Member.GetBaseType()))
+		{
+			bFoundResourceMember = true;
+			break;
+		}
+
 		FString ArrayDim;
 		if(Member.GetNumElements() > 0)
 		{
@@ -144,9 +169,10 @@ static FString CreateHLSLUniformBufferStructMembersDeclaration(const FUniformBuf
 
 		if(Member.GetBaseType() == UBMT_STRUCT)
 		{
-			Result += TEXT("struct {\r\n");
-			Result += CreateHLSLUniformBufferStructMembersDeclaration(*Member.GetStruct(), bIgnorePadding);
-			Result += FString::Printf(TEXT("} %s%s;\r\n"),Member.GetName(),*ArrayDim);
+			Decl.StructMembers += TEXT("struct {\r\n");
+			Decl.Initializer += TEXT(",");
+			CreateHLSLUniformBufferStructMembersDeclaration(Decl, *Member.GetStruct(), FString::Printf(TEXT("%s%s_"), *NamePrefix, Member.GetName()), bExplicitPadding);
+			Decl.StructMembers += FString::Printf(TEXT("} %s%s;\r\n"),Member.GetName(),*ArrayDim);
 			HLSLBaseOffset += Member.GetStruct()->GetSize() * Member.GetNumElements();
 		}
 		else
@@ -203,8 +229,10 @@ static FString CreateHLSLUniformBufferStructMembersDeclaration(const FUniformBuf
 				check(HLSLBaseOffset < Member.GetOffset());
 				while(HLSLBaseOffset < Member.GetOffset())
 				{
-					if(!bIgnorePadding)
-						Result += FString::Printf(TEXT("\tfloat1 _PrePadding%u;\r\n"),HLSLBaseOffset);
+					if(bExplicitPadding)
+					{
+						Decl.ConstantBufferMembers += FString::Printf(TEXT("\tfloat1 _%s_PrePadding%u;\r\n"), *NamePrefix, HLSLBaseOffset);
+					}
 					HLSLBaseOffset += 4;
 				};
 				check(HLSLBaseOffset == Member.GetOffset());
@@ -213,50 +241,86 @@ static FString CreateHLSLUniformBufferStructMembersDeclaration(const FUniformBuf
 			HLSLBaseOffset = Member.GetOffset() + HLSLMemberSize;
 
 			// Generate the member declaration.
-			Result += FString::Printf(TEXT("\t%s%s %s%s;\r\n"),*BaseTypeName,*TypeDim,Member.GetName(),*ArrayDim);
+			FString ParameterName = FString::Printf(TEXT("%s%s"),*NamePrefix,Member.GetName());
+			Decl.ConstantBufferMembers += FString::Printf(TEXT("\t%s%s %s%s;\r\n"),*BaseTypeName,*TypeDim,*ParameterName,*ArrayDim);
+			Decl.StructMembers += FString::Printf(TEXT("\t%s%s %s%s;\r\n"),*BaseTypeName,*TypeDim,Member.GetName(),*ArrayDim);
+			Decl.Initializer += FString::Printf(TEXT(",%s"),*ParameterName);
+		}
+	}
+
+	for(;MemberIndex < StructMembers.Num();++MemberIndex)
+	{
+		const FUniformBufferStruct::FMember& Member = StructMembers[MemberIndex];
+
+		if (!IsUniformBufferResourceType(Member.GetBaseType()))
+		{
+			check(!bFoundResourceMember);
+			continue;
 		}
+		bFoundResourceMember = true;
+
+		// TODO: handle arrays?
+		FString ParameterName = FString::Printf(TEXT("%s%s"),*NamePrefix,Member.GetName());
+		Decl.ResourceMembers += FString::Printf(TEXT("%s %s;\r\n"), Member.GetShaderType(), *ParameterName);
+		Decl.StructMembers += FString::Printf(TEXT("\t%s %s;\r\n"), Member.GetShaderType(), Member.GetName());
+		Decl.Initializer += FString::Printf(TEXT(",%s"), *ParameterName);
 	}
 
-	return Result;
+	Decl.Initializer += TEXT("}");
+	if (Decl.Initializer[OpeningBraceLocPlusOne] == TEXT(','))
+	{
+		Decl.Initializer[OpeningBraceLocPlusOne] = TEXT(' ');
+	}
 }
 
-/**
- * Creates a HLSL declaration of a uniform buffer with the given structure.
- * @param bIgnoreRowMajor on OpenGL this is the default so we suppress it (can be moved to the HLSL2GLSL converter)
- */
-static FString CreateHLSLSM5UniformBufferDeclaration(const TCHAR* Name,const FUniformBufferStruct& UniformBufferStruct, bool bIgnorePadding)
+/** Creates a HLSL declaration of a uniform buffer with the given structure. */
+static FString CreateHLSLUniformBufferDeclaration(const TCHAR* Name,const FUniformBufferStruct& UniformBufferStruct, bool bExplicitPadding)
 {
-	const FString StructMemberDeclaration = CreateHLSLUniformBufferStructMembersDeclaration(UniformBufferStruct, bIgnorePadding);
+	// If the uniform buffer has no members, we don't want to write out anything.  Shader compilers throw errors when faced with empty cbuffers and structs.
+	if (UniformBufferStruct.GetMembers().Num() > 0)
+{
+		FString NamePrefix(FString(Name) + FString(TEXT("_")));
+		FUniformBufferDecl Decl;
+		CreateHLSLUniformBufferStructMembersDeclaration(Decl, UniformBufferStruct, NamePrefix, bExplicitPadding);
 
 	return FString::Printf(
 		TEXT("#ifndef __UniformBuffer_%s_Definition__\r\n")
 		TEXT("#define __UniformBuffer_%s_Definition__\r\n")
 		TEXT("cbuffer %s\r\n")
 		TEXT("{\r\n")
-		TEXT("\tstruct\r\n")
-		TEXT("\t{\r\n")
 		TEXT("%s")
-		TEXT("\t} %s;\r\n")
 		TEXT("}\r\n")
+			TEXT("%s")
+			TEXT("static const struct\r\n")
+			TEXT("{\r\n")
+			TEXT("%s")
+			TEXT("} %s = %s;\r\n")
 		TEXT("#endif\r\n"),
 		Name,
 		Name,
 		Name,
-		*StructMemberDeclaration,
+			*Decl.ConstantBufferMembers,
+			*Decl.ResourceMembers,
+			*Decl.StructMembers,
+			Name,
+			*Decl.Initializer,
 		Name
 		);
 }
 
+	return FString(TEXT("\n"));
+}
+
 FString CreateUniformBufferShaderDeclaration(const TCHAR* Name,const FUniformBufferStruct& UniformBufferStruct,EShaderPlatform Platform)
 {
 	switch(Platform)
 	{
 	default:
 	case SP_PCD3D_SM5:
-		return CreateHLSLSM5UniformBufferDeclaration(Name,UniformBufferStruct, false);
+		return CreateHLSLUniformBufferDeclaration(Name,UniformBufferStruct, true);
 	case SP_OPENGL_SM4:
 	case SP_OPENGL_SM5:
-		return CreateHLSLSM5UniformBufferDeclaration(Name,UniformBufferStruct, true);
+		return CreateHLSLUniformBufferDeclaration(Name,UniformBufferStruct, false);
 	}
 }
 
@@ -289,14 +353,22 @@ void FShaderType::AddReferencedUniformBufferIncludes(FShaderCompilerEnvironment&
 
 	FString UniformBufferIncludes;
 
-	for (TMap<const TCHAR*,FCachedUniformBufferDeclaration>::TConstIterator StructIt(ReferencedUniformBufferStructsCache); StructIt; ++StructIt)
+	for (TMap<const TCHAR*,FCachedUniformBufferDeclaration>::TConstIterator It(ReferencedUniformBufferStructsCache); It; ++It)
 	{
-		check(StructIt.Value().Declaration[Platform].Len() > 0);
-		UniformBufferIncludes += FString::Printf(TEXT("#include \"UniformBuffers/%s.usf\"") LINE_TERMINATOR, StructIt.Key());
+		check(It.Value().Declaration[Platform].Len() > 0);
+		UniformBufferIncludes += FString::Printf(TEXT("#include \"UniformBuffers/%s.usf\"")LINE_TERMINATOR, It.Key());
 		OutEnvironment.IncludeFileNameToContentsMap.Add(
-			*FString::Printf(TEXT("UniformBuffers/%s.usf"),StructIt.Key()),
-			StructIt.Value().Declaration[Platform]
+			*FString::Printf(TEXT("UniformBuffers/%s.usf"),It.Key()),
+			It.Value().Declaration[Platform]
 			);
+
+		for (TLinkedList<FUniformBufferStruct*>::TIterator StructIt(FUniformBufferStruct::GetStructList()); StructIt; StructIt.Next())
+		{
+			if (It.Key() == StructIt->GetShaderVariableName())
+			{
+				StructIt->AddResourceTableEntries(OutEnvironment.ResourceTableMap, OutEnvironment.ResourceTableLayoutHashes);
+			}
+		}
 	}
 
 	FString& GeneratedUniformBuffersInclude = OutEnvironment.IncludeFileNameToContentsMap.FindOrAdd("GeneratedUniformBuffers.usf");
@@ -314,14 +386,22 @@ void FVertexFactoryType::AddReferencedUniformBufferIncludes(FShaderCompilerEnvir
 
 	FString UniformBufferIncludes;
 
-	for (TMap<const TCHAR*,FCachedUniformBufferDeclaration>::TConstIterator StructIt(ReferencedUniformBufferStructsCache); StructIt; ++StructIt)
+	for (TMap<const TCHAR*,FCachedUniformBufferDeclaration>::TConstIterator It(ReferencedUniformBufferStructsCache); It; ++It)
 	{
-		check(StructIt.Value().Declaration[Platform].Len() > 0);
-		UniformBufferIncludes += FString::Printf(TEXT("#include \"UniformBuffers/%s.usf\"") LINE_TERMINATOR, StructIt.Key());
+		check(It.Value().Declaration[Platform].Len() > 0);
+		UniformBufferIncludes += FString::Printf(TEXT("#include \"UniformBuffers/%s.usf\"")LINE_TERMINATOR, It.Key());
 		OutEnvironment.IncludeFileNameToContentsMap.Add(
-			*FString::Printf(TEXT("UniformBuffers/%s.usf"),StructIt.Key()),
-			StructIt.Value().Declaration[Platform]
+			*FString::Printf(TEXT("UniformBuffers/%s.usf"),It.Key()),
+			It.Value().Declaration[Platform]
 		);
+
+		for (TLinkedList<FUniformBufferStruct*>::TIterator StructIt(FUniformBufferStruct::GetStructList()); StructIt; StructIt.Next())
+		{
+			if (It.Key() == StructIt->GetShaderVariableName())
+			{
+				StructIt->AddResourceTableEntries(OutEnvironment.ResourceTableMap, OutEnvironment.ResourceTableLayoutHashes);
+			}
+		}
 	}
 
 	FString& GeneratedUniformBuffersInclude = OutEnvironment.IncludeFileNameToContentsMap.FindOrAdd("GeneratedUniformBuffers.usf");
diff --git a/Engine/Source/Runtime/ShaderCore/Private/VertexFactory.cpp b/Engine/Source/Runtime/ShaderCore/Private/VertexFactory.cpp
--- a/Engine/Source/Runtime/ShaderCore/Private/VertexFactory.cpp
+++ b/Engine/Source/Runtime/ShaderCore/Private/VertexFactory.cpp
@@ -7,6 +7,7 @@
 #include "ShaderCore.h"
 #include "Shader.h"
 #include "VertexFactory.h"
+#include "RHICommandList.h"
 
 uint32 FVertexFactoryType::NextHashIndex = 0;
 bool FVertexFactoryType::bInitializedSerializationHistory = false;
@@ -166,26 +167,29 @@ FVertexFactoryType* FindVertexFactoryType(FName TypeName)
 	return NULL;
 }
 
-void FVertexFactory::Set() const
+void FVertexFactory::Set(FRHICommandList* RHICmdList) const
 {
 	check(IsInitialized());
 	for(int32 StreamIndex = 0;StreamIndex < Streams.Num();StreamIndex++)
 	{
 		const FVertexStream& Stream = Streams[StreamIndex];
 		checkf(Stream.VertexBuffer->IsInitialized(), TEXT("Vertex buffer was not initialized! Stream %u, Stride %u, Name %s"), StreamIndex, Stream.Stride, *Stream.VertexBuffer->GetFriendlyName());
-		RHISetStreamSource(StreamIndex,Stream.VertexBuffer->VertexBufferRHI,Stream.Stride,Stream.Offset);
+		FRHICommandList::SetStreamSource(RHICmdList, StreamIndex, Stream.VertexBuffer->VertexBufferRHI, Stream.Stride, Stream.Offset);
 	}
 }
 
 void FVertexFactory::SetPositionStream() const
 {
+	//@todo-rco: RHIPacketList
+	FRHICommandList* RHICmdList = nullptr;
+
 	check(IsInitialized());
 	// Set the predefined vertex streams.
 	for(int32 StreamIndex = 0;StreamIndex < PositionStream.Num();StreamIndex++)
 	{
 		const FVertexStream& Stream = PositionStream[StreamIndex];
 		check(Stream.VertexBuffer->IsInitialized());
-		RHISetStreamSource(StreamIndex,Stream.VertexBuffer->VertexBufferRHI,Stream.Stride,Stream.Offset);
+		FRHICommandList::SetStreamSource(RHICmdList, StreamIndex, Stream.VertexBuffer->VertexBufferRHI, Stream.Stride, Stream.Offset);
 	}
 }
 
diff --git a/Engine/Source/Runtime/ShaderCore/Public/Shader.h b/Engine/Source/Runtime/ShaderCore/Public/Shader.h
--- a/Engine/Source/Runtime/ShaderCore/Public/Shader.h
+++ b/Engine/Source/Runtime/ShaderCore/Public/Shader.h
@@ -1122,11 +1122,11 @@ public:
 	TShaderMapRef(const TShaderMap<typename ShaderType::ShaderMetaType>* ShaderIndex):
 	 Shader(ShaderIndex->template GetShader<ShaderType>()) // gcc3 needs the template quantifier so it knows the < is not a less-than
 	{}
-	ShaderType* operator->() const
+	FORCEINLINE ShaderType* operator->() const
 	{
 		return Shader;
 	}
-	ShaderType* operator*() const
+	FORCEINLINE ShaderType* operator*() const
 	{
 		return Shader;
 	}
@@ -1320,6 +1320,7 @@ extern SHADERCORE_API FShaderType* FindShaderTypeByName(const TCHAR* ShaderTypeN
 
 /** Helper function to dispatch a compute shader while checking that parameters have been set correctly. */
 extern SHADERCORE_API void DispatchComputeShader(
+	FRHICommandList* RHICmdList,
 	FShader* Shader,
 	uint32 ThreadGroupCountX,
 	uint32 ThreadGroupCountY,
@@ -1327,6 +1328,7 @@ extern SHADERCORE_API void DispatchComputeShader(
 
 /** Helper function to dispatch a compute shader indirectly while checking that parameters have been set correctly. */
 extern SHADERCORE_API void DispatchIndirectComputeShader(
+	FRHICommandList* RHICmdList, 
 	FShader* Shader,
 	FVertexBufferRHIParamRef ArgumentBuffer,
 	uint32 ArgumentOffset);
diff --git a/Engine/Source/Runtime/ShaderCore/Public/ShaderCore.h b/Engine/Source/Runtime/ShaderCore/Public/ShaderCore.h
--- a/Engine/Source/Runtime/ShaderCore/Public/ShaderCore.h
+++ b/Engine/Source/Runtime/ShaderCore/Public/ShaderCore.h
@@ -7,8 +7,9 @@
 #pragma once
 
 #include "Core.h"
-#include "../../RHI/Public/RHI.h"
+#include "../../RHI/Public/RHIDefinitions.h"
 #include "SecureHash.h"
+#include "../../RenderCore/Public/UniformBuffer.h"
 
 /** 
  * Controls whether shader related logs are visible.  
@@ -46,19 +47,6 @@ DECLARE_MEMORY_STAT_EXTERN(TEXT("Shader Memory"),STAT_Shaders_ShaderMemory,STATG
 DECLARE_MEMORY_STAT_EXTERN(TEXT("Shader Resource Mem"),STAT_Shaders_ShaderResourceMemory,STATGROUP_Shaders, SHADERCORE_API);
 DECLARE_MEMORY_STAT_EXTERN(TEXT("Shader MapMemory"),STAT_Shaders_ShaderMapMemory,STATGROUP_Shaders, SHADERCORE_API);
 
-enum EShaderFrequency
-{
-	SF_Vertex			= 0,
-	SF_Hull				= 1,
-	SF_Domain			= 2,
-	SF_Pixel			= 3,
-	SF_Geometry			= 4,
-	SF_Compute			= 5,
-
-	SF_NumFrequencies	= 6,
-	SF_NumBits			= 3,
-};
-
 inline TStatId GetMemoryStatType(EShaderFrequency ShaderFrequency)
 {
 	checkAtCompileTime(6 == SF_NumFrequencies, Bad_EShaderFrequency);
@@ -132,6 +120,7 @@ public:
 	SHADERCORE_API bool FindParameterAllocation(const TCHAR* ParameterName,uint16& OutBufferIndex,uint16& OutBaseIndex,uint16& OutSize) const;
 	SHADERCORE_API bool ContainsParameterAllocation(const TCHAR* ParameterName) const;
 	SHADERCORE_API void AddParameterAllocation(const TCHAR* ParameterName,uint16 BufferIndex,uint16 BaseIndex,uint16 Size);
+	SHADERCORE_API void RemoveParameterAllocation(const TCHAR* ParameterName);
 	/** Checks that all parameters are bound and asserts if any aren't in a debug build
 	* @param InVertexFactoryType can be 0
 	*/
@@ -228,12 +217,105 @@ private:
 	TMap<FString,FString> Definitions;
 };
 
+struct FBaseShaderResourceTable
+{
+	/** Bits indicating which resource tables contain resources bound to this shader. */
+	uint32 ResourceTableBits;
+
+	/** Mapping of bound SRVs to their location in resource tables. */
+	TArray<uint32> ShaderResourceViewMap;
+
+	/** Mapping of bound sampler states to their location in resource tables. */
+	TArray<uint32> SamplerMap;
+
+	/** Mapping of bound UAVs to their location in resource tables. */
+	TArray<uint32> UnorderedAccessViewMap;
+
+	/** Hash of the layouts of resource tables at compile time, used for runtime validation. */
+	TArray<uint32> ResourceTableLayoutHashes;
+
+	FBaseShaderResourceTable() :
+		ResourceTableBits(0)
+	{
+	}
+
+	friend bool operator==(const FBaseShaderResourceTable &A, const FBaseShaderResourceTable& B)
+	{
+		bool bEqual = true;
+		bEqual &= (A.ResourceTableBits == B.ResourceTableBits);
+		bEqual &= (A.ShaderResourceViewMap.Num() == B.ShaderResourceViewMap.Num());
+		bEqual &= (A.SamplerMap.Num() == B.SamplerMap.Num());
+		bEqual &= (A.UnorderedAccessViewMap.Num() == B.UnorderedAccessViewMap.Num());
+		bEqual &= (A.ResourceTableLayoutHashes.Num() == B.ResourceTableLayoutHashes.Num());
+		if (!bEqual)
+		{
+			return false;
+		}
+		bEqual &= (FMemory::Memcmp(A.ShaderResourceViewMap.GetTypedData(), B.ShaderResourceViewMap.GetTypedData(), A.ShaderResourceViewMap.GetTypeSize()*A.ShaderResourceViewMap.Num()) == 0);
+		bEqual &= (FMemory::Memcmp(A.SamplerMap.GetTypedData(), B.SamplerMap.GetTypedData(), A.SamplerMap.GetTypeSize()*A.SamplerMap.Num()) == 0);
+		bEqual &= (FMemory::Memcmp(A.UnorderedAccessViewMap.GetTypedData(), B.UnorderedAccessViewMap.GetTypedData(), A.UnorderedAccessViewMap.GetTypeSize()*A.UnorderedAccessViewMap.Num()) == 0);
+		bEqual &= (FMemory::Memcmp(A.ResourceTableLayoutHashes.GetTypedData(), B.ResourceTableLayoutHashes.GetTypedData(), A.ResourceTableLayoutHashes.GetTypeSize()*A.ResourceTableLayoutHashes.Num()) == 0);
+		return bEqual;
+	}
+};
+
+inline FArchive& operator<<(FArchive& Ar, FBaseShaderResourceTable& SRT)
+{
+	Ar << SRT.ResourceTableBits;
+	Ar << SRT.ShaderResourceViewMap;
+	Ar << SRT.SamplerMap;
+	Ar << SRT.UnorderedAccessViewMap;
+	Ar << SRT.ResourceTableLayoutHashes;
+
+	return Ar;
+}
+
+struct FShaderResourceTable
+{
+	/** Bits indicating which resource tables contain resources bound to this shader. */
+	uint32 ResourceTableBits;
+
+	/** The max index of a uniform buffer from which resources are bound. */
+	uint32 MaxBoundResourceTable;
+
+	/** Mapping of bound Textures to their location in resource tables. */
+	TArray<uint32> TextureMap;
+
+	/** Mapping of bound SRVs to their location in resource tables. */
+	TArray<uint32> ShaderResourceViewMap;
+
+	/** Mapping of bound sampler states to their location in resource tables. */
+	TArray<uint32> SamplerMap;
+
+	/** Mapping of bound UAVs to their location in resource tables. */
+	TArray<uint32> UnorderedAccessViewMap;
+
+	/** Hash of the layouts of resource tables at compile time, used for runtime validation. */
+	TArray<uint32> ResourceTableLayoutHashes;
+
+	FShaderResourceTable()
+		: ResourceTableBits(0)
+		, MaxBoundResourceTable(0)
+	{
+	}
+};
+
+inline FArchive& operator<<(FArchive& Ar, FResourceTableEntry& Entry)
+{
+	Ar << Entry.UniformBufferName;
+	Ar << Entry.Type;
+	Ar << Entry.ResourceIndex;
+	return Ar;
+}
+
 /** The environment used to compile a shader. */
 struct FShaderCompilerEnvironment : public FRefCountedObject
 {
 	TMap<FString,FString> IncludeFileNameToContentsMap;
 	TArray<uint32> CompilerFlags;
 	TMap<uint32,uint8> RenderTargetOutputFormatsMap;
+	TMap<FString,FResourceTableEntry> ResourceTableMap;
+	TMap<FString,uint32> ResourceTableLayoutHashes;
 
 	/** Default constructor. */
 	FShaderCompilerEnvironment() 
@@ -279,7 +361,7 @@ struct FShaderCompilerEnvironment : public FRefCountedObject
 	{
 		return Definitions.GetDefinitionMap();
 	}
-
+	
 	void SetRenderTargetOutputFormat(uint32 RenderTargetIndex, EPixelFormat PixelFormat)
 	{
 		RenderTargetOutputFormatsMap.Add(RenderTargetIndex, PixelFormat);
@@ -288,7 +370,7 @@ struct FShaderCompilerEnvironment : public FRefCountedObject
 	friend FArchive& operator<<(FArchive& Ar,FShaderCompilerEnvironment& Environment)
 	{
 		// Note: this serialize is used to pass between UE4 and the shader compile worker, recompile both when modifying
-		return Ar << Environment.IncludeFileNameToContentsMap << Environment.Definitions << Environment.CompilerFlags << Environment.RenderTargetOutputFormatsMap;
+		return Ar << Environment.IncludeFileNameToContentsMap << Environment.Definitions << Environment.CompilerFlags << Environment.RenderTargetOutputFormatsMap << Environment.ResourceTableMap << Environment.ResourceTableLayoutHashes;
 	}
 
 	void Merge(const FShaderCompilerEnvironment& Other)
@@ -310,6 +392,8 @@ struct FShaderCompilerEnvironment : public FRefCountedObject
 		}
 
 		CompilerFlags.Append(Other.CompilerFlags);
+		ResourceTableMap.Append(Other.ResourceTableMap);
+		ResourceTableLayoutHashes.Append(Other.ResourceTableLayoutHashes);
 		Definitions.Merge(Other.Definitions);
 		RenderTargetOutputFormatsMap.Append(Other.RenderTargetOutputFormatsMap);
 	}
@@ -327,8 +411,9 @@ struct FShaderCompilerInput
 	FString SourceFilePrefix;
 	FString SourceFilename;
 	FString EntryPointName;
-	FString DumpDebugInfoPath;
-	FString ShaderDebugInfoDirectory;
+	//gilmerge, it isn't clear what the relationship is between these two and what the call sites are supposed to use. I might have botched the call sites.
+	FString DumpDebugInfoRootPath;	// Dump debug path (up to platform)
+	FString DumpDebugInfoPath;		// Dump debug path (platform/groupname)
 	FShaderCompilerEnvironment Environment;
 	TRefCountPtr<FShaderCompilerEnvironment> SharedEnvironment;
 
@@ -345,8 +430,8 @@ struct FShaderCompilerInput
 		Ar << Input.SourceFilePrefix;
 		Ar << Input.SourceFilename;
 		Ar << Input.EntryPointName;
+		Ar << Input.DumpDebugInfoRootPath;
 		Ar << Input.DumpDebugInfoPath;
-		Ar << Input.ShaderDebugInfoDirectory;
 		Ar << Input.Environment;
 
 		bool bHasSharedEnvironment = IsValidRef(Input.SharedEnvironment);
diff --git a/Engine/Source/Runtime/ShaderCore/Public/ShaderParameterUtils.h b/Engine/Source/Runtime/ShaderCore/Public/ShaderParameterUtils.h
new file mode 100644
--- /dev/null
+++ b/Engine/Source/Runtime/ShaderCore/Public/ShaderParameterUtils.h
@@ -0,0 +1,418 @@
+// Copyright 1998-2014 Epic Games, Inc. All Rights Reserved.
+
+/*=============================================================================
+	ShaderParameters.h: Shader parameter inline definitions.
+=============================================================================*/
+
+#pragma once
+
+#include "ShaderCore.h"
+#include "ShaderParameters.h"
+#include "RHICommandList.h"
+
+/**
+ * Sets the value of a  shader parameter.  Template'd on shader type
+ * A template parameter specified the type of the parameter value.
+ * NOTE: Shader should be the param ref type, NOT the param type, since Shader is passed by value. 
+ * Otherwise AddRef/ReleaseRef will be called many times.
+ */
+template<typename ShaderRHIParamRef, class ParameterType>
+void SetShaderValue(
+	FRHICommandList* RHICmdList, 
+	ShaderRHIParamRef Shader,
+	const FShaderParameter& Parameter,
+	const ParameterType& Value,
+	uint32 ElementIndex = 0
+	)
+{
+	checkAtCompileTime(!TIsPointerType<ParameterType>::Value, ErrorPassByValue);
+
+	const uint32 AlignedTypeSize = Align(sizeof(ParameterType),ShaderArrayElementAlignBytes);
+	const int32 NumBytesToSet = FMath::Min<int32>(sizeof(ParameterType),Parameter.GetNumBytes() - ElementIndex * AlignedTypeSize);
+
+	// This will trigger if the parameter was not serialized
+	checkSlow(Parameter.IsInitialized());
+
+	if(NumBytesToSet > 0)
+	{
+		FRHICommandList::SetShaderParameter(
+			RHICmdList,
+			Shader,
+			Parameter.GetBufferIndex(),
+			Parameter.GetBaseIndex() + ElementIndex * AlignedTypeSize,
+			(uint32)NumBytesToSet,
+			&Value,
+			true
+			);
+	}
+}
+
+/** Specialization of the above for C++ bool type. */
+template<typename ShaderRHIParamRef>
+void SetShaderValue(
+	FRHICommandList* RHICmdList, 
+	ShaderRHIParamRef Shader,
+	const FShaderParameter& Parameter,
+	const bool& Value,
+	uint32 ElementIndex = 0
+	)
+{
+	const uint32 BoolValue = Value;
+	SetShaderValue(RHICmdList, Shader, Parameter, BoolValue, ElementIndex);
+}
+
+/**
+ * Sets the value of a shader parameter array.  Template'd on shader type
+ * A template parameter specified the type of the parameter value.
+ * NOTE: Shader should be the param ref type, NOT the param type, since Shader is passed by value. 
+ * Otherwise AddRef/ReleaseRef will be called many times.
+ */
+template<typename ShaderRHIParamRef,class ParameterType>
+void SetShaderValueArray(
+	FRHICommandList* RHICmdList,
+	ShaderRHIParamRef Shader,
+	const FShaderParameter& Parameter,
+	const ParameterType* Values,
+	uint32 NumElements,
+	uint32 BaseElementIndex = 0
+	)
+{
+	const uint32 AlignedTypeSize = Align(sizeof(ParameterType),ShaderArrayElementAlignBytes);
+	const int32 NumBytesToSet = FMath::Min<int32>(NumElements * AlignedTypeSize,Parameter.GetNumBytes() - BaseElementIndex * AlignedTypeSize);
+
+	// This will trigger if the parameter was not serialized
+	checkSlow(Parameter.IsInitialized());
+
+	if(NumBytesToSet > 0)
+	{
+		FRHICommandList::SetShaderParameter(
+			RHICmdList,
+			Shader,
+			Parameter.GetBufferIndex(),
+			Parameter.GetBaseIndex() + BaseElementIndex * AlignedTypeSize,
+			(uint32)NumBytesToSet,
+			Values,
+			true
+			);
+	}
+}
+
+/** Specialization of the above for C++ bool type. */
+template<typename ShaderRHIParamRef>
+void SetShaderValueArray(
+	FRHICommandList* RHICmdList, 
+	ShaderRHIParamRef Shader,
+	const FShaderParameter& Parameter,
+	const bool* Values,
+	uint32 NumElements,
+	uint32 BaseElementIndex = 0
+	)
+{
+	UE_LOG(LogShaders, Fatal, TEXT("SetShaderValueArray does not support bool arrays."));
+}
+
+/**
+ * Sets the value of a pixel shader bool parameter.
+ */
+inline void SetPixelShaderBool(
+	FRHICommandList* RHICmdList, 
+	FPixelShaderRHIParamRef PixelShader,
+	const FShaderParameter& Parameter,
+	bool Value
+	)
+{
+	// This will trigger if the parameter was not serialized
+	checkSlow(Parameter.IsInitialized());
+
+	if (Parameter.GetNumBytes() > 0)
+	{
+		// Convert to uint32 before passing to RHI
+		uint32 BoolValue = Value;
+		FRHICommandList::SetShaderParameter(
+			RHICmdList,
+			PixelShader,
+			Parameter.GetBufferIndex(),
+			Parameter.GetBaseIndex(),
+			sizeof(BoolValue),
+			&BoolValue,
+			true
+			);
+	}
+}
+
+/**
+ * Sets the value of a shader texture parameter.  Template'd on shader type
+ */
+template<typename ShaderTypeRHIParamRef>
+FORCEINLINE void SetTextureParameter(
+	FRHICommandList* RHICmdList, 
+	ShaderTypeRHIParamRef Shader,
+	const FShaderResourceParameter& TextureParameter,
+	const FShaderResourceParameter& SamplerParameter,
+	const FTexture* Texture,
+	uint32 ElementIndex = 0
+	)
+{
+	// This will trigger if the parameter was not serialized
+	checkSlow(TextureParameter.IsInitialized());
+	checkSlow(SamplerParameter.IsInitialized());
+	if(TextureParameter.IsBound())
+	{
+		Texture->LastRenderTime = FApp::GetCurrentTime();
+
+		if (ElementIndex < TextureParameter.GetNumResources())
+		{
+			FRHICommandList::SetShaderTexture(RHICmdList, Shader, TextureParameter.GetBaseIndex() + ElementIndex, Texture->TextureRHI);
+		}
+	}
+	
+	// @todo ue4 samplerstate Should we maybe pass in two separate values? SamplerElement and TextureElement? Or never allow an array of samplers? Unsure best
+	// if there is a matching sampler for this texture array index (ElementIndex), then set it. This will help with this case:
+	//			Texture2D LightMapTextures[NUM_LIGHTMAP_COEFFICIENTS];
+	//			SamplerState LightMapTexturesSampler;
+	// In this case, we only set LightMapTexturesSampler when ElementIndex is 0, we don't set the sampler state for all 4 textures
+	// This assumes that the all textures want to use the same sampler state
+	if(SamplerParameter.IsBound())
+	{
+		if (ElementIndex < SamplerParameter.GetNumResources())
+		{
+			FRHICommandList::SetShaderSampler(RHICmdList, Shader, SamplerParameter.GetBaseIndex() + ElementIndex, Texture->SamplerStateRHI);
+		}
+	}
+}
+
+/**
+ * Sets the value of a shader texture parameter. Template'd on shader type.
+ */
+template<typename ShaderTypeRHIParamRef>
+FORCEINLINE void SetTextureParameter(
+	FRHICommandList* RHICmdList, 
+	ShaderTypeRHIParamRef Shader,
+	const FShaderResourceParameter& TextureParameter,
+	const FShaderResourceParameter& SamplerParameter,
+	FSamplerStateRHIParamRef SamplerStateRHI,
+	FTextureRHIParamRef TextureRHI,
+	uint32 ElementIndex = 0
+	)
+{
+	// This will trigger if the parameter was not serialized
+	checkSlow(TextureParameter.IsInitialized());
+	checkSlow(SamplerParameter.IsInitialized());
+	if(TextureParameter.IsBound())
+	{
+		if (ElementIndex < TextureParameter.GetNumResources())
+		{
+			FRHICommandList::SetShaderTexture(RHICmdList, Shader, TextureParameter.GetBaseIndex() + ElementIndex, TextureRHI);
+		}
+	}
+	// @todo ue4 samplerstate Should we maybe pass in two separate values? SamplerElement and TextureElement? Or never allow an array of samplers? Unsure best
+	// if there is a matching sampler for this texture array index (ElementIndex), then set it. This will help with this case:
+	//			Texture2D LightMapTextures[NUM_LIGHTMAP_COEFFICIENTS];
+	//			SamplerState LightMapTexturesSampler;
+	// In this case, we only set LightMapTexturesSampler when ElementIndex is 0, we don't set the sampler state for all 4 textures
+	// This assumes that the all textures want to use the same sampler state
+	if(SamplerParameter.IsBound())
+	{
+		if (ElementIndex < SamplerParameter.GetNumResources())
+		{
+			FRHICommandList::SetShaderSampler(RHICmdList, Shader, SamplerParameter.GetBaseIndex() + ElementIndex, SamplerStateRHI);
+		}
+	}
+}
+
+/**
+ * Sets the value of a shader surface parameter (e.g. to access MSAA samples).
+ * Template'd on shader type (e.g. pixel shader or compute shader).
+ */
+template<typename ShaderTypeRHIParamRef>
+FORCEINLINE void SetTextureParameter(
+	FRHICommandList* RHICmdList, 
+	ShaderTypeRHIParamRef Shader,
+	const FShaderResourceParameter& Parameter,
+	FTextureRHIParamRef NewTextureRHI
+	)
+{
+	if(Parameter.IsBound())
+	{
+		FRHICommandList::SetShaderTexture(
+			RHICmdList,
+			Shader,
+			Parameter.GetBaseIndex(),
+			NewTextureRHI
+			);
+	}
+}
+
+/**
+ * Sets the value of a shader sampler parameter. Template'd on shader type.
+ */
+template<typename ShaderTypeRHIParamRef>
+FORCEINLINE void SetSamplerParameter(
+	FRHICommandList* RHICmdList, 
+	ShaderTypeRHIParamRef Shader,
+	const FShaderResourceParameter& Parameter,
+	FSamplerStateRHIParamRef SamplerStateRHI
+	)
+{
+	if(Parameter.IsBound())
+	{
+		FRHICommandList::SetShaderSampler(
+			RHICmdList,
+			Shader,
+			Parameter.GetBaseIndex(),
+			SamplerStateRHI
+			);
+	}
+}
+
+/**
+ * Sets the value of a shader resource view parameter
+ * Template'd on shader type (e.g. pixel shader or compute shader).
+ */
+template<typename ShaderTypeRHIParamRef>
+FORCEINLINE void SetSRVParameter(
+	FRHICommandList* RHICmdList, 
+	ShaderTypeRHIParamRef Shader,
+	const FShaderResourceParameter& Parameter,
+	FShaderResourceViewRHIParamRef NewShaderResourceViewRHI
+	)
+{
+	if(Parameter.IsBound())
+	{
+		FRHICommandList::SetShaderResourceViewParameter(
+			RHICmdList,
+			Shader,
+			Parameter.GetBaseIndex(),
+			NewShaderResourceViewRHI
+			);
+	}
+}
+
+/**
+ * Sets the value of a unordered access view parameter
+ */
+FORCEINLINE void SetUAVParameter(
+	FRHICommandList* RHICmdList, 
+	FComputeShaderRHIParamRef ComputeShader,
+	const FShaderResourceParameter& Parameter,
+	FUnorderedAccessViewRHIParamRef NewUnorderedAccessViewRHI
+	)
+{
+	if(Parameter.IsBound())
+	{
+		FRHICommandList::SetUAVParameter(
+			RHICmdList,
+			ComputeShader,
+			Parameter.GetBaseIndex(),
+			NewUnorderedAccessViewRHI
+			);
+	}
+}
+
+template<typename TShaderRHIRef>
+inline void SetUAVParameterIfCS(FRHICommandList* RHICmdList, TShaderRHIRef,const FShaderResourceParameter& UAVParameter,FUnorderedAccessViewRHIParamRef UAV)
+{
+}
+
+template<> inline void SetUAVParameterIfCS<FComputeShaderRHIParamRef>(FRHICommandList* RHICmdList, FComputeShaderRHIParamRef ComputeShader,const FShaderResourceParameter& UAVParameter,FUnorderedAccessViewRHIParamRef UAV)
+{
+	SetUAVParameter(RHICmdList, ComputeShader,UAVParameter,UAV);
+}
+
+template<> inline void SetUAVParameterIfCS<FComputeShaderRHIRef>(FRHICommandList* RHICmdList, FComputeShaderRHIRef ComputeShader,const FShaderResourceParameter& UAVParameter,FUnorderedAccessViewRHIParamRef UAV)
+{
+	SetUAVParameter(RHICmdList, ComputeShader,UAVParameter,UAV);
+}
+
+template<typename TShaderRHIRef>
+inline void FRWShaderParameter::SetBuffer(FRHICommandList* RHICmdList, TShaderRHIRef Shader,const FRWBuffer& RWBuffer) const
+{
+	SetSRVParameter(RHICmdList, Shader,SRVParameter,RWBuffer.SRV);
+	SetUAVParameterIfCS(RHICmdList, Shader,UAVParameter,RWBuffer.UAV);
+}
+
+template<typename TShaderRHIRef>
+inline void FRWShaderParameter::SetTexture(FRHICommandList* RHICmdList, TShaderRHIRef Shader,const FTextureRHIParamRef Texture,FUnorderedAccessViewRHIParamRef UAV) const
+{
+	SetTextureParameter(RHICmdList, Shader,SRVParameter,Texture);
+	SetUAVParameterIfCS(RHICmdList, Shader,UAVParameter,UAV);
+}
+
+inline void FRWShaderParameter::UnsetUAV(FRHICommandList* RHICmdList, FComputeShaderRHIParamRef ComputeShader) const
+{
+	SetUAVParameter(RHICmdList, ComputeShader,UAVParameter,FUnorderedAccessViewRHIRef());
+}
+
+
+/** Sets the value of a shader uniform buffer parameter to a uniform buffer containing the struct. */
+template<typename TShaderRHIRef>
+inline void SetUniformBufferParameter(
+	FRHICommandList* RHICmdList,
+	TShaderRHIRef Shader,
+	const FShaderUniformBufferParameter& Parameter,
+	FUniformBufferRHIParamRef UniformBufferRHI
+	)
+{
+	// This will trigger if the parameter was not serialized
+	checkSlow(Parameter.IsInitialized());
+	if(Parameter.IsBound())
+	{
+		FRHICommandList::SetShaderUniformBuffer(RHICmdList, Shader, Parameter.GetBaseIndex(), UniformBufferRHI);
+	}
+}
+
+/** Sets the value of a shader uniform buffer parameter to a uniform buffer containing the struct. */
+template<typename TShaderRHIRef,typename TBufferStruct>
+inline void SetUniformBufferParameter(
+	FRHICommandList* RHICmdList,
+	TShaderRHIRef Shader,
+	const TShaderUniformBufferParameter<TBufferStruct>& Parameter,
+	const TUniformBufferRef<TBufferStruct>& UniformBufferRef
+	)
+{
+	// This will trigger if the parameter was not serialized
+	checkSlow(Parameter.IsInitialized());
+	if(Parameter.IsBound())
+	{
+		FRHICommandList::SetShaderUniformBuffer(RHICmdList, Shader, Parameter.GetBaseIndex(), UniformBufferRef);
+	}
+}
+
+/** Sets the value of a shader uniform buffer parameter to a uniform buffer containing the struct. */
+template<typename TShaderRHIRef,typename TBufferStruct>
+inline void SetUniformBufferParameter(
+	FRHICommandList* RHICmdList,
+	TShaderRHIRef Shader,
+	const TShaderUniformBufferParameter<TBufferStruct>& Parameter,
+	const TUniformBuffer<TBufferStruct>& UniformBuffer
+	)
+{
+	// This will trigger if the parameter was not serialized
+	checkSlow(Parameter.IsInitialized());
+	if (Parameter.IsBound())
+	{
+		FRHICommandList::SetShaderUniformBuffer(RHICmdList, Shader, Parameter.GetBaseIndex(), UniformBuffer.GetUniformBufferRHI());
+	}
+}
+
+/** Sets the value of a shader uniform buffer parameter to a value of the struct. */
+template<typename TShaderRHIRef,typename TBufferStruct>
+inline void SetUniformBufferParameterImmediate(
+	FRHICommandList* RHICmdList,
+	TShaderRHIRef Shader,
+	const TShaderUniformBufferParameter<TBufferStruct>& Parameter,
+	const TBufferStruct& UniformBufferValue
+	)
+{
+	// This will trigger if the parameter was not serialized
+	checkSlow(Parameter.IsInitialized());
+	if(Parameter.IsBound())
+	{
+		FRHICommandList::SetShaderUniformBuffer(
+			RHICmdList,
+			Shader,
+			Parameter.GetBaseIndex(),
+			RHICreateUniformBuffer(&UniformBufferValue,TBufferStruct::StaticStruct.GetLayout(),UniformBuffer_SingleDraw)
+			);
+	}
+}
