UE4 - Fix Mac compile error

diff --git a/Engine/Source/Runtime/Renderer/Private/PostProcess/SceneRenderTargets.h b/Engine/Source/Runtime/Renderer/Private/PostProcess/SceneRenderTargets.h
--- a/Engine/Source/Runtime/Renderer/Private/PostProcess/SceneRenderTargets.h
+++ b/Engine/Source/Runtime/Renderer/Private/PostProcess/SceneRenderTargets.h
@@ -9,6 +9,7 @@
 #include "ShaderParameters.h"
 #include "RenderTargetPool.h"
 #include "../SystemTextures.h"
+#include "RHIStaticStates.h"
 
 /** Number of cube map shadow depth surfaces that will be created and used for rendering one pass point light shadows. */
 static const int32 NumCubeShadowDepthSurfaces = 5;
@@ -158,10 +159,10 @@ public:
 	// only call if BeginRenderingCustomDepth() returned true
 	void FinishRenderingCustomDepth(const FResolveRect& ResolveRect = FResolveRect());
 
-    /**
-     * Resolve a previously rendered scene color surface.
-     */
-    void ResolveSceneColor(const FResolveRect& ResolveRect = FResolveRect());
+	/**
+	 * Resolve a previously rendered scene color surface.
+	 */
+	void ResolveSceneColor(const FResolveRect& ResolveRect = FResolveRect());
 
 	/** Resolves the GBuffer targets so that their resolved textures can be sampled. */
 	void ResolveGBufferSurfaces(const FResolveRect& ResolveRect = FResolveRect());
