Fix the r.TemporalAASharpness cvar to a fixed value (this is needed for the new temporal AA to get proper edge gradients).

diff --git a/Engine/Source/Runtime/Renderer/Private/PostProcess/PostProcessTemporalAA.cpp b/Engine/Source/Runtime/Renderer/Private/PostProcess/PostProcessTemporalAA.cpp
--- a/Engine/Source/Runtime/Renderer/Private/PostProcess/PostProcessTemporalAA.cpp
+++ b/Engine/Source/Runtime/Renderer/Private/PostProcess/PostProcessTemporalAA.cpp
@@ -129,6 +129,9 @@ public:
 			};
 		
 			float Sharpness = CVarTemporalAASharpness.GetValueOnRenderThread();
+			// With the new temporal AA, need to hardcode a value here.
+			// Going to remove this once the new temporal AA is validated.
+			Sharpness = -0.25f;
 			Sharpness = 1.0f + Sharpness * 0.5f;
 
 			float Weights[9];
