Fix for TTP#337266 & 337267 Particles render black when bounds are zero due to NaNs being generated for the indirect lighting coeficients

diff --git a/Engine/Source/Runtime/Renderer/Private/IndirectLightingCache.cpp b/Engine/Source/Runtime/Renderer/Private/IndirectLightingCache.cpp
--- a/Engine/Source/Runtime/Renderer/Private/IndirectLightingCache.cpp
+++ b/Engine/Source/Runtime/Renderer/Private/IndirectLightingCache.cpp
@@ -213,9 +213,9 @@ void FIndirectLightingCache::CalculateBlockPositionAndSize(const FBoxSphereBound
 	FVector RoundedBoundsSize;
 
 	// Find the exponent needed to represent the bounds size if BoundSizeRoundUpBase is the base
-	RoundedBoundsSize.X = FMath::LogX(BoundSizeRoundUpBase, Bounds.BoxExtent.X * 2);
-	RoundedBoundsSize.Y = FMath::LogX(BoundSizeRoundUpBase, Bounds.BoxExtent.Y * 2);
-	RoundedBoundsSize.Z = FMath::LogX(BoundSizeRoundUpBase, Bounds.BoxExtent.Z * 2);
+	RoundedBoundsSize.X = FMath::Max(1.f, FMath::LogX(BoundSizeRoundUpBase, Bounds.BoxExtent.X * 2));
+	RoundedBoundsSize.Y = FMath::Max(1.f, FMath::LogX(BoundSizeRoundUpBase, Bounds.BoxExtent.Y * 2));
+	RoundedBoundsSize.Z = FMath::Max(1.f, FMath::LogX(BoundSizeRoundUpBase, Bounds.BoxExtent.Z * 2));
 
 	// Round up to the next integer exponent to provide stability even when Bounds.BoxExtent is changing
 	RoundedBoundsSize.X = FMath::Pow(BoundSizeRoundUpBase, FMath::TruncToInt(RoundedBoundsSize.X) + 1);
