Added a permutation to the dynamic skylight diffuse shader to avoid a branch

diff --git a/Engine/Shaders/DistanceFieldLightingShared.usf b/Engine/Shaders/DistanceFieldLightingShared.usf
--- a/Engine/Shaders/DistanceFieldLightingShared.usf
+++ b/Engine/Shaders/DistanceFieldLightingShared.usf
@@ -112,7 +112,6 @@ SamplerState DistanceFieldNormalSampler;
 
 int CurrentLevelDownsampleFactor;
 float2 AOBufferSize;
-float4 AOScreenPositionScaleBias;
 
 int DownsampleFactorToBaseLevel;
 float2 BaseLevelTexelSize;
diff --git a/Engine/Shaders/DistanceFieldSurfaceCacheLighting.usf b/Engine/Shaders/DistanceFieldSurfaceCacheLighting.usf
--- a/Engine/Shaders/DistanceFieldSurfaceCacheLighting.usf
+++ b/Engine/Shaders/DistanceFieldSurfaceCacheLighting.usf
@@ -19,7 +19,7 @@ struct FObjectCullVertexOutput
 
 float ConservativeRadiusScale;
 
-/** */
+/** Used when culling objects into screenspace tile lists */
 void ObjectCullVS(
 	float4 InPosition : ATTRIBUTE0,
 	uint ObjectIndex : SV_InstanceID,
@@ -44,6 +44,7 @@ Buffer<float4> TileConeDepthRanges;
 
 float2 NumGroups;
 
+/** Intersects a single object with the tile and adds to the intersection list if needed. */
 void ObjectCullPS(
 	FObjectCullVertexOutput Input, 
 	in float4 SVPos : SV_POSITION,
@@ -104,6 +105,7 @@ void ObjectCullPS(
 
 float2 UVToTileHead;
 
+/** Computes the distance field normal, using a search through all the nearby objects to find the closest one, whose normal is used. */
 void ComputeDistanceFieldNormalPS(
 	in float4 UVAndScreenPos : TEXCOORD0, 
 	in float4 SVPos : SV_POSITION,
@@ -361,7 +363,7 @@ float InterpolationRadiusScale;
 float2 NormalizedOffsetToPixelCenter;
 float HackExpand;
 
-/** */
+/** Expands a screen-facing polygon to cover a surface cache record for splatting. */
 void IrradianceCacheSplatVS(
 	float2 InPosition : ATTRIBUTE0,
 	float2 InUV       : ATTRIBUTE1,
@@ -409,6 +411,7 @@ void IrradianceCacheSplatVS(
 float InterpolationAngleNormalization; 
 float InvMinCosPointBehindPlane;
 
+/** Computes surface cache weighting between the current pixel and the record being splatted. */
 void IrradianceCacheSplatPS(
 	FIrradianceCacheSplatVertexOutput Input, 
 	in float4 SVPos : SV_POSITION,
@@ -483,6 +486,7 @@ void IrradianceCacheSplatPS(
 Texture2D BentNormalAOTexture3;
 SamplerState BentNormalAOSampler3;
 
+/** Normalizes the splatted surface cache values, packs depth in alpha. */
 void AOCombinePS2(
 	in float4 UVAndScreenPos : TEXCOORD0, 
 	in float4 SVPos : SV_POSITION,
@@ -533,6 +537,7 @@ SamplerState BentNormalAOSampler;
 
 #define HALF_FILL_KERNEL_SIZE 2
 
+/** Fills in texels with no splatted weight from screenspace neighbors. */
 void FillGapsPS(
 	in float4 UVAndScreenPos : TEXCOORD0, 
 	in float4 SVPos : SV_POSITION,
@@ -572,6 +577,7 @@ Texture2D HistoryTexture;
 SamplerState HistorySampler;
 float HistoryWeight;
 
+/** Reprojection the occlusion history, only correct for pixels belonging to static objects (camera motion only). */
 void UpdateHistoryPS(
 	in float4 UVAndScreenPos : TEXCOORD0, 
 	out float4 OutColor : SV_Target0)
@@ -609,6 +615,7 @@ void UpdateHistoryPS(
 	OutColor.a = NewValue.a;
 }
 
+/** Upsamples the AO results to full resolution using a bilateral filter. */
 void AOUpsamplePS(in float4 UVAndScreenPos : TEXCOORD0, out float4 OutColor : SV_Target0)
 {
 #define BILATERAL_UPSAMPLE 1
diff --git a/Engine/Shaders/DistanceFieldSurfaceCacheLightingCompute.usf b/Engine/Shaders/DistanceFieldSurfaceCacheLightingCompute.usf
--- a/Engine/Shaders/DistanceFieldSurfaceCacheLightingCompute.usf
+++ b/Engine/Shaders/DistanceFieldSurfaceCacheLightingCompute.usf
@@ -30,6 +30,7 @@ RWBuffer<uint> RWTileHeadDataUnpacked;
 #define MAX_OBJECTS_PER_TILE 1
 #endif
 
+/** Builds tile depth ranges and bounding cones. */
 [numthreads(THREADGROUP_SIZEX, THREADGROUP_SIZEY, 1)]
 void BuildTileConesMain(
 	uint3 GroupId : SV_GroupID,
@@ -335,9 +336,6 @@ void DistanceFieldAOBuildTileListMain(
 
 /** View rect min in xy, max in zw. */
 float2 ThreadToCulledTile;
-float ConeHalfAngle;
-float BentNormalNormalizeFactor;
-float RecordRadiusScale;
 
 RWBuffer<float4> RWIrradianceCachePositionRadius;
 RWBuffer<float4> RWIrradianceCacheNormal;
@@ -360,6 +358,7 @@ groupshared uint2 CachedIrradianceCacheTileCoordinate[THREADGROUP_TOTALSIZE];
 groupshared uint NumQueuedIrradianceCacheRecords;
 groupshared uint BaseRecordIndex;
 
+/** Creates new surface cache records for sample points that don't have valid coverage from existing surface cache records. */
 [numthreads(THREADGROUP_SIZEX, THREADGROUP_SIZEY, 1)]
 void PopulateCacheCS(
 	uint3 GroupId : SV_GroupID,
@@ -405,8 +404,10 @@ void PopulateCacheCS(
 			//@todo - offset shading position along normal to avoid incorrect self-occlusion?
 			float3 WorldShadingPosition = OpaqueWorldPosition;
 
+			// For debugging
 			//if (all(DispatchThreadId.xy == uint2(2,1)))
 			{
+				// Allocate a new record and store off attributes of the created record
 				int NextSampleIndex;
 				InterlockedAdd(NumQueuedIrradianceCacheRecords, 1, NextSampleIndex);
 				CachedIrradianceCachePositionRadius[NextSampleIndex] = float4(WorldShadingPosition, SceneDepth * CurrentLevelDownsampleFactor * .005f);
@@ -446,6 +447,10 @@ void PopulateCacheCS(
 	}
 }
 
+float BentNormalNormalizeFactor;
+float TanConeHalfAngle;
+float RecordRadiusScale;
+
 Buffer<float4> IrradianceCachePositionRadius;
 Buffer<uint2> IrradianceCacheTileCoordinate;
 Buffer<float4> IrradianceCacheNormal;
@@ -462,6 +467,7 @@ RWBuffer<float4> RWIrradianceCacheBentNormal;
 groupshared uint SharedConeVisibility[NUM_SAMPLES];
 groupshared uint SharedMinOcclusionDistance;
 
+/** Computes ambient occlusion for a surface cache record by cone stepping through the nearby object distance fields. */
 [numthreads(FINAL_GATHER_THREADGROUP_SIZE, 1, 1)]
 void FinalGatherCS(
 	uint3 GroupId : SV_GroupID,
@@ -471,9 +477,6 @@ void FinalGatherCS(
 	uint ThreadIndex = GroupThreadId.x;
 	uint ObjectOffsetIndex = ThreadIndex / THREADS_PER_OBJECT;
 
-	float ShadingConeAngleScale = 1;
-	float ShadingConeTanTheta = tan(ShadingConeAngleScale * ConeHalfAngle);
-
 	if (ThreadIndex == 0)
 	{
 		for (uint ConeIndex = 0; ConeIndex < NUM_SAMPLES; ConeIndex++) 
@@ -583,7 +586,7 @@ void FinalGatherCS(
 						float DistanceToOccluder = Texture3DSampleLevel(DistanceFieldTexture, DistanceFieldSampler, StepVolumeUV, 0).x + DistanceToClamped;
 
 						float EffectiveDistanceToOccluder = DistanceToOccluder;
-						float SphereRadius = LocalStepOffset * ShadingConeTanTheta;
+						float SphereRadius = LocalStepOffset * TanConeHalfAngle;
 						// Allows use of negative distances in the small steps to create full penumbra range, reduces over-occlusion
 						float MinOcclusionFraction = 0;//-.6f * (1 - saturate(4 * StepIndex / (float)NUM_CONE_STEPS));
 						float Visibility = saturate((EffectiveDistanceToOccluder / SphereRadius - MinOcclusionFraction) / (1 - MinOcclusionFraction));
@@ -626,7 +629,13 @@ void FinalGatherCS(
 		Visibility = Visibility / (float)NUM_SAMPLES;
 
 		RWOccluderRadius[RecordIndex] = RecordRadiusScale * asfloat(SharedMinOcclusionDistance);
-		RWIrradianceCacheBentNormal[RecordIndex] = float4(EncodeNormalForAO(Visibility.xxx), 0);
+
+#define USE_DYNAMIC_OCCLUSION_BENT_NORMAL 0
+#if USE_DYNAMIC_OCCLUSION_BENT_NORMAL
+		RWIrradianceCacheBentNormal[RecordIndex] = float4(EncodeNormalForAO(UnoccludedDirection), 0);
+#else
+		RWIrradianceCacheBentNormal[RecordIndex] = float4(Visibility.xxx, 0);
+#endif
 	}
 }
 
diff --git a/Engine/Shaders/SkyLighting.usf b/Engine/Shaders/SkyLighting.usf
--- a/Engine/Shaders/SkyLighting.usf
+++ b/Engine/Shaders/SkyLighting.usf
@@ -9,7 +9,6 @@
 #include "ReflectionEnvironmentShared.usf"
 
 float BentNormalUVScale;
-float ApplyShadowing;
 Texture2D BentNormalAOTexture;
 SamplerState BentNormalAOSampler;
 
@@ -32,13 +31,12 @@ void SkyLightDiffusePS(
 		float DotProductFactor = 1;
 		float3 SkyLightingNormal = InGBufferData.WorldNormal;
 
-		BRANCH
-		if (ApplyShadowing > 0)
-		{
+#if APPLY_SHADOWING
 			float3 BentNormal = Texture2DSampleLevel(BentNormalAOTexture, BentNormalAOSampler, UV * BentNormalUVScale, 0).xyz;
-			SkyVisibility = BentNormal.x;
 
 			// Disabled for now as it needs tweaking
+			// Causes artifacts due to only having local occlusion, and distance field representation error (shading point inside distance field)
+			// Have to set matching define in DistanceFieldSurfaceCacheLightingCompute.usf
 			#define USE_DYNAMIC_OCCLUSION_BENT_NORMAL 0
 			#if USE_DYNAMIC_OCCLUSION_BENT_NORMAL
 				SkyVisibility = length(BentNormal);
@@ -51,15 +49,29 @@ void SkyLightDiffusePS(
 				// In the mostly unoccluded case, evaluate sky lighting with the material normal, because it has higher detail
 				// In the mostly occluded case, evaluate sky lighting with the bent normal, because it is a better representation of the incoming lighting
 				// Then treat the lighting evaluated along the bent normal as an area light, so we must apply the lambert term
-				SkyLightingNormal = lerp(NormalizedBentNormal, InGBufferData.WorldNormal, BentNormalWeightFactor);
-				DotProductFactor = lerp(dot(NormalizedBentNormal, InGBufferData.WorldNormal), 1, BentNormalWeightFactor);
+
+				FLATTEN
+				if (InGBufferData.HasDistanceFieldRepresentation != 0)
+				{
+					SkyLightingNormal = lerp(NormalizedBentNormal, InGBufferData.WorldNormal, BentNormalWeightFactor);
+					DotProductFactor = lerp(dot(NormalizedBentNormal, InGBufferData.WorldNormal), 1, BentNormalWeightFactor);
+				}
+				else
+				{
+					SkyVisibility = .5f;
+				}
+
+			#else
+				SkyVisibility = BentNormal.x;
+
+				FLATTEN
+				if (InGBufferData.HasDistanceFieldRepresentation == 0)
+				{
+					SkyVisibility = .5f;
+				}
 			#endif
 
-			if (InGBufferData.HasDistanceFieldRepresentation == 0)
-			{
-				SkyVisibility = .5f;
-			}
-		}
+#endif
 			
 		// Compute the preconvolved incoming lighting with the bent normal direction
 		float3 DiffuseLookup = GetSkySHDiffuse(SkyLightingNormal) * View.SkyLightColor.rgb;
diff --git a/Engine/Source/Runtime/Renderer/Private/DistanceFieldSurfaceCacheLighting.cpp b/Engine/Source/Runtime/Renderer/Private/DistanceFieldSurfaceCacheLighting.cpp
--- a/Engine/Source/Runtime/Renderer/Private/DistanceFieldSurfaceCacheLighting.cpp
+++ b/Engine/Source/Runtime/Renderer/Private/DistanceFieldSurfaceCacheLighting.cpp
@@ -729,7 +729,6 @@ public:
 	{
 		CurrentLevelDownsampleFactor.Bind(ParameterMap, TEXT("CurrentLevelDownsampleFactor"));
 		AOBufferSize.Bind(ParameterMap, TEXT("AOBufferSize"));
-		AOScreenPositionScaleBias.Bind(ParameterMap, TEXT("AOScreenPositionScaleBias"));
 		DownsampleFactorToBaseLevel.Bind(ParameterMap, TEXT("DownsampleFactorToBaseLevel"));
 		BaseLevelTexelSize.Bind(ParameterMap, TEXT("BaseLevelTexelSize"));
 	}
@@ -743,34 +742,6 @@ public:
 		const FVector2D AOBufferSizeValue = FIntPoint::DivideAndRoundUp(GSceneRenderTargets.GetBufferSizeXY(), CurrentLevelDownsampleFactorValue);
 		SetShaderValue(RHICmdList, ShaderRHI, AOBufferSize, AOBufferSizeValue);
 
-		const FVector2D InvAOBufferSize(1.0f / AOBufferSizeValue.X, 1.0f / AOBufferSizeValue.Y);
-
-		// Round down, to match how the view was set
-		const FIntPoint AOViewRectMin = FIntPoint(View.ViewRect.Min.X / CurrentLevelDownsampleFactorValue, View.ViewRect.Min.Y / CurrentLevelDownsampleFactorValue);
-		const FIntPoint AOViewRectSize = FIntPoint::DivideAndRoundUp(View.ViewRect.Size(), CurrentLevelDownsampleFactorValue);
-
-		/*
-		const FVector4 AOScreenPositionScaleBiasValue(
-			AOViewRectSize.X * InvAOBufferSize.X / +2.0f,
-			AOViewRectSize.Y * InvAOBufferSize.Y / (-2.0f * GProjectionSignY),
-			(AOViewRectSize.X / 2.0f + GPixelCenterOffset + AOViewRectMin.X) * InvAOBufferSize.X,
-			(AOViewRectSize.Y / 2.0f + GPixelCenterOffset + AOViewRectMin.Y) * InvAOBufferSize.Y
-			);*/
-
-		const FIntPoint RoundedDownViewDimensions = AOViewRectSize * CurrentLevelDownsampleFactorValue;
-
-		const FIntPoint BufferSize = GSceneRenderTargets.GetBufferSizeXY();
-		const float InvBufferSizeX = 1.0f / BufferSize.X;
-		const float InvBufferSizeY = 1.0f / BufferSize.Y;
-		const FVector4 ScreenPositionScaleBias(
-			RoundedDownViewDimensions.X * InvBufferSizeX / +2.0f,
-			RoundedDownViewDimensions.Y * InvBufferSizeY / (-2.0f * GProjectionSignY),
-			(RoundedDownViewDimensions.Y / 2.0f + GPixelCenterOffset + AOViewRectMin.Y) * InvBufferSizeY,
-			(RoundedDownViewDimensions.X / 2.0f + GPixelCenterOffset + AOViewRectMin.X) * InvBufferSizeX
-			);
-
-		SetShaderValue(RHICmdList, ShaderRHI, AOScreenPositionScaleBias, ScreenPositionScaleBias);
-
 		SetShaderValue(RHICmdList, ShaderRHI, DownsampleFactorToBaseLevel, CurrentLevelDownsampleFactorValue / GAODownsampleFactor);
 
 		const FIntPoint DownsampledBufferSize(GSceneRenderTargets.GetBufferSizeXY() / FIntPoint(GAODownsampleFactor, GAODownsampleFactor));
@@ -780,14 +751,13 @@ public:
 
 	friend FArchive& operator<<(FArchive& Ar,FAOLevelParameters& P)
 	{
-		Ar << P.CurrentLevelDownsampleFactor << P.AOBufferSize << P.AOScreenPositionScaleBias << P.DownsampleFactorToBaseLevel << P.BaseLevelTexelSize;
+		Ar << P.CurrentLevelDownsampleFactor << P.AOBufferSize << P.DownsampleFactorToBaseLevel << P.BaseLevelTexelSize;
 		return Ar;
 	}
 
 private:
 	FShaderParameter CurrentLevelDownsampleFactor;
 	FShaderParameter AOBufferSize;
-	FShaderParameter AOScreenPositionScaleBias;
 	FShaderParameter DownsampleFactorToBaseLevel;
 	FShaderParameter BaseLevelTexelSize;
 };
@@ -1301,15 +1271,11 @@ public:
 		ViewDimensionsParameter.Bind(Initializer.ParameterMap, TEXT("ViewDimensions"));
 		ThreadToCulledTile.Bind(Initializer.ParameterMap, TEXT("ThreadToCulledTile"));
 		TileListGroupSize.Bind(Initializer.ParameterMap, TEXT("TileListGroupSize"));
-		ConeHalfAngle.Bind(Initializer.ParameterMap, TEXT("ConeHalfAngle"));
-		BentNormalNormalizeFactor.Bind(Initializer.ParameterMap, TEXT("BentNormalNormalizeFactor"));
-		DistanceFieldAtlasTexelSize.Bind(Initializer.ParameterMap, TEXT("DistanceFieldAtlasTexelSize"));
 		DistanceFieldNormalTexture.Bind(Initializer.ParameterMap, TEXT("DistanceFieldNormalTexture"));
 		DistanceFieldNormalSampler.Bind(Initializer.ParameterMap, TEXT("DistanceFieldNormalSampler"));
 		IrradianceCacheSplatTexture.Bind(Initializer.ParameterMap, TEXT("IrradianceCacheSplatTexture"));
 		IrradianceCacheSplatSampler.Bind(Initializer.ParameterMap, TEXT("IrradianceCacheSplatSampler"));
 		AOLevelParameters.Bind(Initializer.ParameterMap);
-		RecordRadiusScale.Bind(Initializer.ParameterMap, TEXT("RecordRadiusScale"));
 	}
 
 	FPopulateCacheCS()
@@ -1351,24 +1317,6 @@ public:
 
 		SetShaderValue(RHICmdList, ShaderRHI, TileListGroupSize, TileListGroupSizeValue);
 
-		SetShaderValue(RHICmdList, ShaderRHI, ConeHalfAngle, GAOConeHalfAngle);
-
-		FVector UnoccludedVector(0);
-
-		for (int32 SampleIndex = 0; SampleIndex < NumConeSampleDirections; SampleIndex++)
-		{
-			UnoccludedVector += SampleDirections[SampleIndex];
-		}
-
-		float BentNormalNormalizeFactorValue = 1.0f / (UnoccludedVector / NumConeSampleDirections).Size();
-		SetShaderValue(RHICmdList, ShaderRHI, BentNormalNormalizeFactor, BentNormalNormalizeFactorValue);
-
-		const int32 NumTexelsOneDimX = GDistanceFieldVolumeTextureAtlas.GetSizeX();
-		const int32 NumTexelsOneDimY = GDistanceFieldVolumeTextureAtlas.GetSizeY();
-		const int32 NumTexelsOneDimZ = GDistanceFieldVolumeTextureAtlas.GetSizeZ();
-		const FVector InvTextureDim(1.0f / NumTexelsOneDimX, 1.0f / NumTexelsOneDimY, 1.0f / NumTexelsOneDimZ);
-		SetShaderValue(RHICmdList, ShaderRHI, DistanceFieldAtlasTexelSize, InvTextureDim);
-
 		SetTextureParameter(
 			RHICmdList,
 			ShaderRHI,
@@ -1386,8 +1334,6 @@ public:
 			TStaticSamplerState<SF_Point,AM_Clamp,AM_Clamp,AM_Clamp>::GetRHI(),
 			DistanceFieldAOIrradianceCacheSplat.ShaderResourceTexture
 			);
-
-		SetShaderValue(RHICmdList, ShaderRHI, RecordRadiusScale, GAORecordRadiusScale);
 	}
 
 	void UnsetParameters(FRHICommandList* RHICmdList)
@@ -1411,14 +1357,10 @@ public:
 		Ar << ViewDimensionsParameter;
 		Ar << ThreadToCulledTile;
 		Ar << TileListGroupSize;
-		Ar << ConeHalfAngle;
-		Ar << BentNormalNormalizeFactor;
-		Ar << DistanceFieldAtlasTexelSize;
 		Ar << DistanceFieldNormalTexture;
 		Ar << DistanceFieldNormalSampler;
 		Ar << IrradianceCacheSplatTexture;
 		Ar << IrradianceCacheSplatSampler;
-		Ar << RecordRadiusScale;
 		return bShaderHasOutdatedParameters;
 	}
 
@@ -1434,14 +1376,10 @@ private:
 	FShaderParameter ViewDimensionsParameter;
 	FShaderParameter ThreadToCulledTile;
 	FShaderParameter TileListGroupSize;
-	FShaderParameter ConeHalfAngle;
-	FShaderParameter BentNormalNormalizeFactor;
-	FShaderParameter DistanceFieldAtlasTexelSize;
 	FShaderResourceParameter DistanceFieldNormalTexture;
 	FShaderResourceParameter DistanceFieldNormalSampler;
 	FShaderResourceParameter IrradianceCacheSplatTexture;
 	FShaderResourceParameter IrradianceCacheSplatSampler;
-	FShaderParameter RecordRadiusScale;
 };
 
 IMPLEMENT_SHADER_TYPE(,FPopulateCacheCS,TEXT("DistanceFieldSurfaceCacheLightingCompute"),TEXT("PopulateCacheCS"),SF_Compute);
@@ -1553,7 +1491,7 @@ public:
 		TileHeadDataUnpacked.Bind(Initializer.ParameterMap, TEXT("TileHeadDataUnpacked"));
 		TileArrayData.Bind(Initializer.ParameterMap, TEXT("TileArrayData"));
 		TileListGroupSize.Bind(Initializer.ParameterMap, TEXT("TileListGroupSize"));
-		ConeHalfAngle.Bind(Initializer.ParameterMap, TEXT("ConeHalfAngle"));
+		TanConeHalfAngle.Bind(Initializer.ParameterMap, TEXT("TanConeHalfAngle"));
 		BentNormalNormalizeFactor.Bind(Initializer.ParameterMap, TEXT("BentNormalNormalizeFactor"));
 		DistanceFieldAtlasTexelSize.Bind(Initializer.ParameterMap, TEXT("DistanceFieldAtlasTexelSize"));
 		DistanceFieldNormalTexture.Bind(Initializer.ParameterMap, TEXT("DistanceFieldNormalTexture"));
@@ -1614,7 +1552,7 @@ public:
 
 		SetShaderValue(RHICmdList, ShaderRHI, TileListGroupSize, TileListGroupSizeValue);
 
-		SetShaderValue(RHICmdList, ShaderRHI, ConeHalfAngle, GAOConeHalfAngle);
+		SetShaderValue(RHICmdList, ShaderRHI, TanConeHalfAngle, FMath::Tan(GAOConeHalfAngle));
 
 		FVector UnoccludedVector(0);
 
@@ -1679,7 +1617,7 @@ public:
 		Ar << TileHeadData;
 		Ar << TileArrayData;
 		Ar << TileListGroupSize;
-		Ar << ConeHalfAngle;
+		Ar << TanConeHalfAngle;
 		Ar << BentNormalNormalizeFactor;
 		Ar << DistanceFieldAtlasTexelSize;
 		Ar << DistanceFieldNormalTexture;
@@ -1709,7 +1647,7 @@ private:
 	FShaderResourceParameter TileHeadData;
 	FShaderResourceParameter TileArrayData;
 	FShaderParameter TileListGroupSize;
-	FShaderParameter ConeHalfAngle;
+	FShaderParameter TanConeHalfAngle;
 	FShaderParameter BentNormalNormalizeFactor;
 	FShaderParameter DistanceFieldAtlasTexelSize;
 	FShaderResourceParameter DistanceFieldNormalTexture;
@@ -3100,9 +3038,10 @@ bool FDeferredShadingSceneRenderer::RenderDistanceFieldAOSurfaceCache(FSceneRend
 	return false;
 }
 
-class FDynamicSkyLightDiffusePS : public FGlobalShader
+template<bool bApplyShadowing>
+class TDynamicSkyLightDiffusePS : public FGlobalShader
 {
-	DECLARE_SHADER_TYPE(FDynamicSkyLightDiffusePS, Global);
+	DECLARE_SHADER_TYPE(TDynamicSkyLightDiffusePS, Global);
 public:
 
 	static bool ShouldCache(EShaderPlatform Platform)
@@ -3114,23 +3053,23 @@ public:
 	{
 		FGlobalShader::ModifyCompilationEnvironment(Platform,OutEnvironment);
 		OutEnvironment.SetDefine(TEXT("DOWNSAMPLE_FACTOR"), GAODownsampleFactor);
+		OutEnvironment.SetDefine(TEXT("APPLY_SHADOWING"), (uint32)(bApplyShadowing ? 1 : 0));
 	}
 
 	/** Default constructor. */
-	FDynamicSkyLightDiffusePS() {}
+	TDynamicSkyLightDiffusePS() {}
 
 	/** Initialization constructor. */
-	FDynamicSkyLightDiffusePS(const ShaderMetaType::CompiledShaderInitializerType& Initializer)
+	TDynamicSkyLightDiffusePS(const ShaderMetaType::CompiledShaderInitializerType& Initializer)
 		: FGlobalShader(Initializer)
 	{
 		DeferredParameters.Bind(Initializer.ParameterMap);
 		LightAccumulationTexture.Bind(Initializer.ParameterMap, TEXT("BentNormalAOTexture"));
 		LightAccumulationSampler.Bind(Initializer.ParameterMap, TEXT("BentNormalAOSampler"));
 		BentNormalUVScale.Bind(Initializer.ParameterMap, TEXT("BentNormalUVScale"));
-		ApplyShadowing.Bind(Initializer.ParameterMap, TEXT("ApplyShadowing"));
 	}
 
-	void SetParameters(FRHICommandList* RHICmdList, const FSceneView& View, FTextureRHIParamRef DynamicBentNormalAO, bool bApplyShadowing)
+	void SetParameters(FRHICommandList* RHICmdList, const FSceneView& View, FTextureRHIParamRef DynamicBentNormalAO)
 	{
 		const FPixelShaderRHIParamRef ShaderRHI = GetPixelShader();
 		FGlobalShader::SetParameters(RHICmdList, ShaderRHI, View);
@@ -3139,7 +3078,6 @@ public:
 		SetTextureParameter(RHICmdList, ShaderRHI, LightAccumulationTexture, LightAccumulationSampler, TStaticSamplerState<SF_Point>::GetRHI(), DynamicBentNormalAO);
 
 		SetShaderValue(RHICmdList, ShaderRHI, BentNormalUVScale, 1.0f);
-		SetShaderValue(RHICmdList, ShaderRHI, ApplyShadowing, bApplyShadowing ? 1.0f : 0.0f);
 	}
 
 	// FShader interface.
@@ -3150,7 +3088,6 @@ public:
 		Ar << LightAccumulationTexture;
 		Ar << LightAccumulationSampler;
 		Ar << BentNormalUVScale;
-		Ar << ApplyShadowing;
 		return bShaderHasOutdatedParameters;
 	}
 
@@ -3160,10 +3097,10 @@ private:
 	FShaderResourceParameter LightAccumulationTexture;
 	FShaderResourceParameter LightAccumulationSampler;
 	FShaderParameter BentNormalUVScale;
-	FShaderParameter ApplyShadowing;
 };
 
-IMPLEMENT_SHADER_TYPE(,FDynamicSkyLightDiffusePS,TEXT("SkyLighting"),TEXT("SkyLightDiffusePS"),SF_Pixel);
+IMPLEMENT_SHADER_TYPE(template<>,TDynamicSkyLightDiffusePS<true>,TEXT("SkyLighting"),TEXT("SkyLightDiffusePS"),SF_Pixel);
+IMPLEMENT_SHADER_TYPE(template<>,TDynamicSkyLightDiffusePS<false>,TEXT("SkyLighting"),TEXT("SkyLightDiffusePS"),SF_Pixel);
 
 void FDeferredShadingSceneRenderer::RenderDynamicSkyLighting()
 {
@@ -3208,13 +3145,26 @@ void FDeferredShadingSceneRenderer::RenderDynamicSkyLighting()
 			RHISetBlendState( TStaticBlendState< CW_RGB, BO_Add, BF_One, BF_One >::GetRHI() );
 
 			TShaderMapRef< FPostProcessVS > VertexShader( GetGlobalShaderMap() );
-			TShaderMapRef< FDynamicSkyLightDiffusePS > PixelShader( GetGlobalShaderMap() );
 
-			static FGlobalBoundShaderState BoundShaderState;
-			SetGlobalBoundShaderState( BoundShaderState, GFilterVertexDeclaration.VertexDeclarationRHI, *VertexShader, *PixelShader );
+			if (bApplyShadowing)
+			{
+				TShaderMapRef<TDynamicSkyLightDiffusePS<true> > PixelShader( GetGlobalShaderMap() );
 
-			PixelShader->SetParameters(RHICmdList, View, LightAccumulation->GetRenderTargetItem().ShaderResourceTexture, bApplyShadowing);
+				static FGlobalBoundShaderState BoundShaderState;
+				SetGlobalBoundShaderState(BoundShaderState, GFilterVertexDeclaration.VertexDeclarationRHI, *VertexShader, *PixelShader);
+
+				PixelShader->SetParameters(RHICmdList, View, LightAccumulation->GetRenderTargetItem().ShaderResourceTexture);
+			}
+			else
+			{
+				TShaderMapRef<TDynamicSkyLightDiffusePS<false> > PixelShader( GetGlobalShaderMap() );
 
+				static FGlobalBoundShaderState BoundShaderState;
+				SetGlobalBoundShaderState(BoundShaderState, GFilterVertexDeclaration.VertexDeclarationRHI, *VertexShader, *PixelShader);
+
+				PixelShader->SetParameters(RHICmdList, View, LightAccumulation->GetRenderTargetItem().ShaderResourceTexture);
+			}
+			
 			DrawRectangle( 
 				0, 0, 
 				View.ViewRect.Width(), View.ViewRect.Height(),
@@ -3225,4 +3175,4 @@ void FDeferredShadingSceneRenderer::RenderDynamicSkyLighting()
 				*VertexShader);
 		}
 	}
-}
\ No newline at end of file
+}
