Change cubemap clearing logic to stop race condition on initializing render targets

#ttp 337070 - Regression: EDITOR: CRASH: Closing the editor while an animation blueprint is open will cause it to crash when reopened

diff --git a/Engine/Source/Runtime/Renderer/Private/ReflectionEnvironmentCapture.cpp b/Engine/Source/Runtime/Renderer/Private/ReflectionEnvironmentCapture.cpp
--- a/Engine/Source/Runtime/Renderer/Private/ReflectionEnvironmentCapture.cpp
+++ b/Engine/Source/Runtime/Renderer/Private/ReflectionEnvironmentCapture.cpp
@@ -782,6 +782,51 @@
 	return CaptureIndex;
 }
 
+void ClearScratchCubemaps()
+{
+	// Clear scratch render targets to a consistent but noticeable value
+	// This makes debugging capture issues much easier, otherwise the random contents from previous captures is shown
+
+	const int32 NumMips = FMath::CeilLogTwo(GReflectionCaptureSize) + 1;
+	FSceneRenderTargetItem& RT0 = GSceneRenderTargets.ReflectionColorScratchCubemap[0]->GetRenderTargetItem();
+
+	for (int32 MipIndex = 0; MipIndex < NumMips; MipIndex++)
+	{
+		if (GSupportsGSRenderTargetLayerSwitchingToMips)
+		{
+			RHISetRenderTarget(RT0.TargetableTexture, MipIndex, NULL);
+			RHIClear(true, FLinearColor(0, 10000, 0, 0), false, 0, false, 0, FIntRect());
+		}
+		else
+		{
+			for (int32 CubeFace = 0; CubeFace < CubeFace_MAX; CubeFace++)
+			{
+				RHISetRenderTarget(RT0.TargetableTexture, MipIndex, CubeFace, NULL);
+				RHIClear(true, FLinearColor(0, 10000, 0, 0), false, 0, false, 0, FIntRect());
+			}
+		}
+	}
+
+	FSceneRenderTargetItem& RT1 = GSceneRenderTargets.ReflectionColorScratchCubemap[1]->GetRenderTargetItem();
+
+	for (int32 MipIndex = 0; MipIndex < NumMips; MipIndex++)
+	{
+		if (GSupportsGSRenderTargetLayerSwitchingToMips)
+		{
+			RHISetRenderTarget(RT1.TargetableTexture, MipIndex, NULL);
+			RHIClear(true, FLinearColor(0, 10000, 0, 0), false, 0, false, 0, FIntRect());
+		}
+		else
+		{
+			for (int32 CubeFace = 0; CubeFace < CubeFace_MAX; CubeFace++)
+			{
+				RHISetRenderTarget(RT1.TargetableTexture, MipIndex, CubeFace, NULL);
+				RHIClear(true, FLinearColor(0, 10000, 0, 0), false, 0, false, 0, FIntRect());
+			}
+		}
+	}
+}
+
 /** Captures the scene for a reflection capture by rendering the scene multiple times and copying into a cubemap texture. */
 void CaptureSceneToScratchCubemap(FRHICommandList* RHICmdList, FSceneRenderer* SceneRenderer, ECubeFace CubeFace, bool bCapturingForSkyLight, bool bLowerHemisphereIsBlack)
 {
@@ -797,6 +842,8 @@
 		// Render the scene normally for one face of the cubemap
 		SceneRenderer->Render();
 
+		ClearScratchCubemaps();
+
 #if PLATFORM_PS4 // @todo ps4 - this should be done a different way
 		// PS4 needs some code here to process the scene
 		extern void TEMP_PostReflectionCaptureRender();
@@ -880,6 +927,9 @@
 void CopyCubemapToScratchCubemap(FRHICommandList* RHICmdList, UTextureCube* SourceCubemap, bool bIsSkyLight, bool bLowerHemisphereIsBlack)
 {
 	check(SourceCubemap);
+	
+	ClearScratchCubemaps();
+
 	const int32 EffectiveSize = GReflectionCaptureSize;
 	FSceneRenderTargetItem& EffectiveColorRT =  GSceneRenderTargets.ReflectionColorScratchCubemap[0]->GetRenderTargetItem();
 
@@ -1213,51 +1263,6 @@
 	}
 }
 
-void ClearScratchCubemaps()
-{
-	// Clear scratch render targets to a consistent but noticeable value
-	// This makes debugging capture issues much easier, otherwise the random contents from previous captures is shown
-
-	const int32 NumMips = FMath::CeilLogTwo(GReflectionCaptureSize) + 1;
-	FSceneRenderTargetItem& RT0 = GSceneRenderTargets.ReflectionColorScratchCubemap[0]->GetRenderTargetItem();
-
-	for (int32 MipIndex = 0; MipIndex < NumMips; MipIndex++)
-	{
-		if (GSupportsGSRenderTargetLayerSwitchingToMips)
-		{
-			RHISetRenderTarget(RT0.TargetableTexture, MipIndex, NULL);
-			RHIClear(true, FLinearColor(0, 10000, 0, 0), false, 0, false, 0, FIntRect());
-		}
-		else
-		{
-			for (int32 CubeFace = 0; CubeFace < CubeFace_MAX; CubeFace++)
-			{
-				RHISetRenderTarget(RT0.TargetableTexture, MipIndex, CubeFace, NULL);
-				RHIClear(true, FLinearColor(0, 10000, 0, 0), false, 0, false, 0, FIntRect());
-			}
-		}
-	}
-
-	FSceneRenderTargetItem& RT1 = GSceneRenderTargets.ReflectionColorScratchCubemap[1]->GetRenderTargetItem();
-
-	for (int32 MipIndex = 0; MipIndex < NumMips; MipIndex++)
-	{
-		if (GSupportsGSRenderTargetLayerSwitchingToMips)
-		{
-			RHISetRenderTarget(RT1.TargetableTexture, MipIndex, NULL);
-			RHIClear(true, FLinearColor(0, 10000, 0, 0), false, 0, false, 0, FIntRect());
-		}
-		else
-		{
-			for (int32 CubeFace = 0; CubeFace < CubeFace_MAX; CubeFace++)
-			{
-				RHISetRenderTarget(RT1.TargetableTexture, MipIndex, CubeFace, NULL);
-				RHIClear(true, FLinearColor(0, 10000, 0, 0), false, 0, false, 0, FIntRect());
-			}
-		}
-	}
-}
-
 /** Creates a transformation for a cubemap face, following the D3D cubemap layout. */
 FMatrix CalcCubeFaceViewMatrix(ECubeFace Face, FVector WorldLocation)
 {
@@ -1489,12 +1494,6 @@
 		}
 		else
 		{
-			ENQUEUE_UNIQUE_RENDER_COMMAND( 
-				ClearCommand,
-			{
-				ClearScratchCubemaps();
-			});
-			
 			CaptureSceneIntoScratchCubemap(this, CaptureComponent->GetComponentLocation(), false, true, 0, false, false);
 
 			ENQUEUE_UNIQUE_RENDER_COMMAND( 
