Distortion rendering : removed duplicated console variable

diff --git a/Engine/Source/Runtime/Renderer/Private/DistortionRendering.cpp b/Engine/Source/Runtime/Renderer/Private/DistortionRendering.cpp
--- a/Engine/Source/Runtime/Renderer/Private/DistortionRendering.cpp
+++ b/Engine/Source/Runtime/Renderer/Private/DistortionRendering.cpp
@@ -766,24 +766,11 @@ void FDistortionPrimSet::AddScenePrimitive(FPrimitiveSceneProxy* PrimitiveSceneP
 	Prims.Add(PrimitiveSceneProxy);
 }
 
-static TAutoConsoleVariable<int32> CVarRefractionQuality(
-	TEXT("r.RefractionQuality"),
-	1,
-	TEXT("Defines the refraction quality.\n")
-	TEXT("  0: off\n")
-	TEXT("  1: on (default)\n"),
-	ECVF_Scalability | ECVF_RenderThreadSafe);
-
 /** 
  * Renders the scene's distortion 
  */
 void FSceneRenderer::RenderDistortion(FRHICommandList* RHICmdList)
 {
-	if (CVarRefractionQuality.GetValueOnRenderThread() == 0)
-	{
-		return; // Don't render distortion
-	}
-
 	SCOPED_DRAW_EVENT(Distortion, DEC_SCENE_ITEMS);
 
 	// do we need to render the distortion pass?
