UE4: Rearrange RHICmdList toward non-null cmd lists and the final "look"

diff --git a/Engine/Source/Runtime/Renderer/Private/AtmosphereRendering.cpp b/Engine/Source/Runtime/Renderer/Private/AtmosphereRendering.cpp
--- a/Engine/Source/Runtime/Renderer/Private/AtmosphereRendering.cpp
+++ b/Engine/Source/Runtime/Renderer/Private/AtmosphereRendering.cpp
@@ -62,7 +62,7 @@ public:
 	}
 
 	template<typename ShaderRHIParamRef>
-	void Set(FRHICommandList* RHICmdList, const ShaderRHIParamRef ShaderRHI, uint32 TextureIdx, const FTextureRHIRef Texture) const
+	void Set(FRHICommandList& RHICmdList, const ShaderRHIParamRef ShaderRHI, uint32 TextureIdx, const FTextureRHIRef Texture) const
 	{
 		if (TextureIdx >= 4)
 		{
@@ -75,7 +75,7 @@ public:
 	}
 
 	template<typename ShaderRHIParamRef>
-	void Set(FRHICommandList* RHICmdList, const ShaderRHIParamRef ShaderRHI, uint32 TextureIdx, TexType TextureType, const FAtmosphereTextures* AtmosphereTextures) const
+	void Set(FRHICommandList& RHICmdList, const ShaderRHIParamRef ShaderRHI, uint32 TextureIdx, TexType TextureType, const FAtmosphereTextures* AtmosphereTextures) const
 	{
 		if (TextureIdx >= 4 || TextureType >= Type_MAX || !AtmosphereTextures)
 		{
@@ -168,7 +168,7 @@ public:
 		OcclusionTextureSamplerParameter.Bind(Initializer.ParameterMap, TEXT("OcclusionTextureSampler"));
 	}
 
-	void SetParameters(FRHICommandList* RHICmdList, const FSceneView& View, TRefCountPtr<IPooledRenderTarget>& LightShaftOcclusion)
+	void SetParameters(FRHICommandList& RHICmdList, const FSceneView& View, TRefCountPtr<IPooledRenderTarget>& LightShaftOcclusion)
 	{
 		FGlobalShader::SetParameters(RHICmdList, GetPixelShader(), View);
 		SceneTextureParameters.Set(RHICmdList, GetPixelShader(), View);
@@ -292,7 +292,7 @@ public:
 	{
 	}
 
-	void SetParameters(FRHICommandList* RHICmdList, const FViewInfo& View)
+	void SetParameters(FRHICommandList& RHICmdList, const FViewInfo& View)
 	{
 		FGlobalShader::SetParameters(RHICmdList, GetVertexShader(),View);
 	}
@@ -351,7 +351,7 @@ void FSceneRenderer::InitAtmosphereConstants()
 
 FGlobalBoundShaderState AtmosphereBoundShaderState[EAtmosphereRenderFlag::E_RenderFlagMax];
 
-void SetAtmosphericFogShaders(FRHICommandList* RHICmdList, FScene* Scene, const FViewInfo& View, TRefCountPtr<IPooledRenderTarget>& LightShaftOcclusion)
+void SetAtmosphericFogShaders(FRHICommandList& RHICmdList, FScene* Scene, const FViewInfo& View, TRefCountPtr<IPooledRenderTarget>& LightShaftOcclusion)
 {
 	TShaderMapRef<FAtmosphericVS> VertexShader(GetGlobalShaderMap());
 	FAtmosphericFogPS* PixelShader = NULL;
@@ -410,7 +410,7 @@ void FDeferredShadingSceneRenderer::RenderAtmosphere(FLightShaftsOutput LightSha
 		};
 
 		//@todo-rco: RHIPacketList
-		FRHICommandList* RHICmdList = nullptr;
+		FRHICommandList& RHICmdList = FRHICommandList::GetNullRef();
 
 		GSceneRenderTargets.BeginRenderingSceneColor();
 		for(int32 ViewIndex = 0;ViewIndex < Views.Num();ViewIndex++)
@@ -474,7 +474,7 @@ public:
 		return bShaderHasOutdatedParameters;
 	}
 
-	void SetParameters(FRHICommandList* RHICmdList, const FViewInfo& View)
+	void SetParameters(FRHICommandList& RHICmdList, const FViewInfo& View)
 	{
 		FGlobalShader::SetParameters(RHICmdList, GetPixelShader(), View);
 	}
@@ -510,7 +510,7 @@ public:
 		return bShaderHasOutdatedParameters;
 	}
 
-	void SetParameters(FRHICommandList* RHICmdList, const FAtmosphereTextures* Textures)
+	void SetParameters(FRHICommandList& RHICmdList, const FAtmosphereTextures* Textures)
 	{
 		AtmosphereParameters.Set(RHICmdList, GetPixelShader(), 0, FAtmosphereShaderPrecomputeTextureParameters::Transmittance, Textures);
 	}
@@ -550,7 +550,7 @@ public:
 		return bShaderHasOutdatedParameters;
 	}
 
-	void SetParameters(FRHICommandList* RHICmdList, const FViewInfo& View, float FirstOrder, const FAtmosphereTextures* Textures)
+	void SetParameters(FRHICommandList& RHICmdList, const FViewInfo& View, float FirstOrder, const FAtmosphereTextures* Textures)
 	{
 		FGlobalShader::SetParameters(RHICmdList, GetPixelShader(), View);
 		AtmosphereParameters.Set(RHICmdList, GetPixelShader(), 0, FAtmosphereShaderPrecomputeTextureParameters::Transmittance, Textures);
@@ -590,7 +590,7 @@ public:
 		return bShaderHasOutdatedParameters;
 	}
 
-	void SetParameters(FRHICommandList* RHICmdList, const FAtmosphereTextures* Textures)
+	void SetParameters(FRHICommandList& RHICmdList, const FAtmosphereTextures* Textures)
 	{
 		AtmosphereParameters.Set(RHICmdList, GetPixelShader(), 0, FAtmosphereShaderPrecomputeTextureParameters::DeltaE, Textures);
 	}
@@ -616,7 +616,7 @@ public:
 		AtmosphereLayerParameter.Bind(Initializer.ParameterMap, TEXT("AtmosphereLayer"));
 	}
 
-	void SetParameters(FRHICommandList* RHICmdList, int AtmosphereLayer)
+	void SetParameters(FRHICommandList& RHICmdList, int AtmosphereLayer)
 	{
 		SetShaderValue(RHICmdList, GetGeometryShader(), AtmosphereLayerParameter, AtmosphereLayer);
 	}
@@ -663,7 +663,7 @@ public:
 		return bShaderHasOutdatedParameters;
 	}
 
-	void SetParameters(FRHICommandList* RHICmdList, const FViewInfo& View, float AtmosphereR, const FVector4& DhdH, const FAtmosphereTextures* Textures)
+	void SetParameters(FRHICommandList& RHICmdList, const FViewInfo& View, float AtmosphereR, const FVector4& DhdH, const FAtmosphereTextures* Textures)
 	{
 		FGlobalShader::SetParameters(RHICmdList, GetPixelShader(), View);
 		AtmosphereParameters.Set(RHICmdList, GetPixelShader(), 0, FAtmosphereShaderPrecomputeTextureParameters::Transmittance, Textures);
@@ -709,7 +709,7 @@ public:
 		return bShaderHasOutdatedParameters;
 	}
 
-	void SetParameters(FRHICommandList* RHICmdList, const FViewInfo& View, float AtmosphereR, const FVector4& DhdH, int AtmosphereLayer, const FAtmosphereTextures* Textures)
+	void SetParameters(FRHICommandList& RHICmdList, const FViewInfo& View, float AtmosphereR, const FVector4& DhdH, int AtmosphereLayer, const FAtmosphereTextures* Textures)
 	{
 		FGlobalShader::SetParameters(RHICmdList, GetPixelShader(), View);
 		AtmosphereParameters.Set(RHICmdList, GetPixelShader(), 0, FAtmosphereShaderPrecomputeTextureParameters::DeltaSR, Textures);
@@ -756,7 +756,7 @@ public:
 		return bShaderHasOutdatedParameters;
 	}
 
-	void SetParameters(FRHICommandList* RHICmdList, const FViewInfo& View, float AtmosphereR, const FVector4& DhdH, int AtmosphereLayer, const FAtmosphereTextures* Textures)
+	void SetParameters(FRHICommandList& RHICmdList, const FViewInfo& View, float AtmosphereR, const FVector4& DhdH, int AtmosphereLayer, const FAtmosphereTextures* Textures)
 	{
 		FGlobalShader::SetParameters(RHICmdList, GetPixelShader(), View);
 		AtmosphereParameters.Set(RHICmdList, GetPixelShader(), 0, FAtmosphereShaderPrecomputeTextureParameters::DeltaSR, Textures);
@@ -805,7 +805,7 @@ public:
 		return bShaderHasOutdatedParameters;
 	}
 
-	void SetParameters(FRHICommandList* RHICmdList, const FViewInfo& View, float AtmosphereR, const FVector4& DhdH, float FirstOrder, const FAtmosphereTextures* Textures)
+	void SetParameters(FRHICommandList& RHICmdList, const FViewInfo& View, float AtmosphereR, const FVector4& DhdH, float FirstOrder, const FAtmosphereTextures* Textures)
 	{
 		FGlobalShader::SetParameters(RHICmdList, GetPixelShader(), View);
 		AtmosphereParameters.Set(RHICmdList, GetPixelShader(), 0, FAtmosphereShaderPrecomputeTextureParameters::Transmittance, Textures);
@@ -855,7 +855,7 @@ public:
 		return bShaderHasOutdatedParameters;
 	}
 
-	void SetParameters(FRHICommandList* RHICmdList, const FViewInfo& View, float AtmosphereR, const FVector4& DhdH, float FirstOrder, const FAtmosphereTextures* Textures)
+	void SetParameters(FRHICommandList& RHICmdList, const FViewInfo& View, float AtmosphereR, const FVector4& DhdH, float FirstOrder, const FAtmosphereTextures* Textures)
 	{
 		FGlobalShader::SetParameters(RHICmdList, GetPixelShader(), View);
 		AtmosphereParameters.Set(RHICmdList, GetPixelShader(), 0, FAtmosphereShaderPrecomputeTextureParameters::Transmittance, Textures);
@@ -930,7 +930,7 @@ public:
 		return bShaderHasOutdatedParameters;
 	}
 
-	void SetParameters(FRHICommandList* RHICmdList, float AtmosphereR, const FVector4& DhdH, int AtmosphereLayer, const FAtmosphereTextures* Textures)
+	void SetParameters(FRHICommandList& RHICmdList, float AtmosphereR, const FVector4& DhdH, int AtmosphereLayer, const FAtmosphereTextures* Textures)
 	{
 		AtmosphereParameters.Set(RHICmdList, GetPixelShader(), 0, FAtmosphereShaderPrecomputeTextureParameters::Inscatter, Textures);
 		SetShaderValue(RHICmdList, GetPixelShader(), AtmosphereRParameter, AtmosphereR);
@@ -975,7 +975,7 @@ public:
 		return bShaderHasOutdatedParameters;
 	}
 
-	void SetParameters(FRHICommandList* RHICmdList, float AtmosphereR, const FVector4& DhdH, int AtmosphereLayer, const FAtmosphereTextures* Textures)
+	void SetParameters(FRHICommandList& RHICmdList, float AtmosphereR, const FVector4& DhdH, int AtmosphereLayer, const FAtmosphereTextures* Textures)
 	{
 		AtmosphereParameters.Set(RHICmdList, GetPixelShader(), 0, FAtmosphereShaderPrecomputeTextureParameters::DeltaSR, Textures);
 		SetShaderValue(RHICmdList, GetPixelShader(), AtmosphereRParameter, AtmosphereR);
@@ -1084,7 +1084,7 @@ void FAtmosphericFogSceneInfo::GetLayerValue(int Layer, float& AtmosphereR, FVec
 void FAtmosphericFogSceneInfo::RenderAtmosphereShaders(const FViewInfo& View, const FIntRect& ViewRect)
 {
 	//@todo-rco: RHIPacketList
-	FRHICommandList* RHICmdList = nullptr;
+	FRHICommandList& RHICmdList = FRHICommandList::GetNullRef();
 
 	check(Component != NULL);
 	switch(AtmospherePhase)
diff --git a/Engine/Source/Runtime/Renderer/Private/AtmosphereRendering.h b/Engine/Source/Runtime/Renderer/Private/AtmosphereRendering.h
--- a/Engine/Source/Runtime/Renderer/Private/AtmosphereRendering.h
+++ b/Engine/Source/Runtime/Renderer/Private/AtmosphereRendering.h
@@ -95,7 +95,7 @@ public:
 	}
 
 	template<typename ShaderRHIParamRef>
-	void Set(FRHICommandList* RHICmdList, const ShaderRHIParamRef ShaderRHI, const FSceneView& View) const
+	void Set(FRHICommandList& RHICmdList, const ShaderRHIParamRef ShaderRHI, const FSceneView& View) const
 	{
 		if (TransmittanceTexture.IsBound() || IrradianceTexture.IsBound() || InscatterTexture.IsBound())
 		{
diff --git a/Engine/Source/Runtime/Renderer/Private/BasePassRendering.cpp b/Engine/Source/Runtime/Renderer/Private/BasePassRendering.cpp
--- a/Engine/Source/Runtime/Renderer/Private/BasePassRendering.cpp
+++ b/Engine/Source/Runtime/Renderer/Private/BasePassRendering.cpp
@@ -70,7 +70,7 @@ void FSkyLightReflectionParameters::GetSkyParametersFromScene(const FScene* Scen
 	}
 }
 
-void FTranslucentLightingParameters::Set(FRHICommandList* RHICmdList, FShader* Shader, const FSceneView* View)
+void FTranslucentLightingParameters::Set(FRHICommandList& RHICmdList, FShader* Shader, const FSceneView* View)
 {
 	SetTextureParameter(
 		RHICmdList, 
@@ -107,7 +107,7 @@ void FTranslucentLightingParameters::Set(FRHICommandList* RHICmdList, FShader* S
 	SkyLightReflectionParameters.SetParameters(RHICmdList, Shader->GetPixelShader(), (const FScene*)(View->Family->Scene), true);
 }
 
-void FTranslucentLightingParameters::SetMesh(FRHICommandList* RHICmdList, FShader* Shader, const FPrimitiveSceneProxy* Proxy)
+void FTranslucentLightingParameters::SetMesh(FRHICommandList& RHICmdList, FShader* Shader, const FPrimitiveSceneProxy* Proxy)
 {
 	// Note: GBlackCubeArrayTexture has an alpha of 0, which is needed to represent invalid data so the sky cubemap can still be applied
 	FTextureRHIParamRef CubeArrayTexture = GBlackCubeArrayTexture->TextureRHI;
@@ -160,7 +160,7 @@ public:
 	/** Draws the translucent mesh with a specific light-map type, and fog volume type */
 	template<typename LightMapPolicyType>
 	void Process(
-		FRHICommandList* RHICmdList, 
+		FRHICommandList& RHICmdList, 
 		const FProcessBasePassMeshParameters& Parameters,
 		const LightMapPolicyType& LightMapPolicy,
 		const typename LightMapPolicyType::ElementDataType& LightMapElementData
@@ -209,7 +209,7 @@ void FBasePassOpaqueDrawingPolicyFactory::AddStaticMesh(FScene* Scene,FStaticMes
 	if( !IsTranslucentBlendMode(BlendMode) )
 	{
 		//@todo-rco: RHIPacketList
-		FRHICommandList* RHICmdList = nullptr;
+		FRHICommandList& RHICmdList = FRHICommandList::GetNullRef();
 		ProcessBasePassMesh(
 			RHICmdList, 
 			FProcessBasePassMeshParameters(
@@ -262,7 +262,7 @@ public:
 	/** Draws the translucent mesh with a specific light-map type, and shader complexity predicate. */
 	template<typename LightMapPolicyType>
 	void Process(
-		FRHICommandList* RHICmdList, 
+		FRHICommandList& RHICmdList, 
 		const FProcessBasePassMeshParameters& Parameters,
 		const LightMapPolicyType& LightMapPolicy,
 		const typename LightMapPolicyType::ElementDataType& LightMapElementData
@@ -328,7 +328,7 @@ public:
 };
 
 bool FBasePassOpaqueDrawingPolicyFactory::DrawDynamicMesh(
-	FRHICommandList* RHICmdList, 
+	FRHICommandList& RHICmdList, 
 	const FSceneView& View,
 	ContextType DrawingContext,
 	const FMeshBatch& Mesh,
@@ -371,7 +371,7 @@ bool FBasePassOpaqueDrawingPolicyFactory::DrawDynamicMesh(
 }
 
 void FCachedVolumeIndirectLightingPolicy::Set(
-	FRHICommandList* RHICmdList, 
+	FRHICommandList& RHICmdList, 
 	const VertexParametersType* VertexShaderParameters,
 	const PixelParametersType* PixelShaderParameters,
 	FShader* VertexShader,
@@ -434,7 +434,7 @@ void FCachedVolumeIndirectLightingPolicy::Set(
 }
 
 void FCachedVolumeIndirectLightingPolicy::SetMesh(
-	FRHICommandList* RHICmdList, 
+	FRHICommandList& RHICmdList, 
 	const FSceneView& View,
 	const FPrimitiveSceneProxy* PrimitiveSceneProxy,
 	const VertexParametersType* VertexShaderParameters,
@@ -476,7 +476,7 @@ void FCachedVolumeIndirectLightingPolicy::SetMesh(
 }
 
 void FCachedPointIndirectLightingPolicy::SetMesh(
-	FRHICommandList* RHICmdList, 
+	FRHICommandList& RHICmdList, 
 	const FSceneView& View,
 	const FPrimitiveSceneProxy* PrimitiveSceneProxy,
 	const VertexParametersType* VertexShaderParameters,
@@ -559,7 +559,7 @@ void FCachedPointIndirectLightingPolicy::SetMesh(
 }
 
 void FSelfShadowedCachedPointIndirectLightingPolicy::SetMesh(
-	FRHICommandList* RHICmdList, 
+	FRHICommandList& RHICmdList, 
 	const FSceneView& View,
 	const FPrimitiveSceneProxy* PrimitiveSceneProxy,
 	const VertexParametersType* VertexShaderParameters,
diff --git a/Engine/Source/Runtime/Renderer/Private/BasePassRendering.h b/Engine/Source/Runtime/Renderer/Private/BasePassRendering.h
--- a/Engine/Source/Runtime/Renderer/Private/BasePassRendering.h
+++ b/Engine/Source/Runtime/Renderer/Private/BasePassRendering.h
@@ -54,7 +54,7 @@ public:
 	}
 
 	void SetParameters(
-		FRHICommandList* RHICmdList, 
+		FRHICommandList& RHICmdList, 
 		const FMaterialRenderProxy* MaterialRenderProxy,
 		const FVertexFactory* VertexFactory,
 		const FMaterial& InMaterialResource,
@@ -72,7 +72,7 @@ public:
 		}
 	}
 
-	void SetMesh(FRHICommandList* RHICmdList, const FVertexFactory* VertexFactory,const FSceneView& View,const FPrimitiveSceneProxy* Proxy,const FMeshBatchElement& BatchElement)
+	void SetMesh(FRHICommandList& RHICmdList, const FVertexFactory* VertexFactory,const FSceneView& View,const FPrimitiveSceneProxy* Proxy,const FMeshBatchElement& BatchElement)
 	{
 		FMeshMaterialShader::SetMesh(RHICmdList, GetVertexShader(),VertexFactory,View,Proxy,BatchElement);
 	}
@@ -181,7 +181,7 @@ public:
 	}
 
 	void SetParameters(
-		FRHICommandList* RHICmdList, 
+		FRHICommandList& RHICmdList, 
 		const FMaterialRenderProxy* MaterialRenderProxy,
 		const FVertexFactory* VertexFactory,
 		const FSceneView& View
@@ -204,7 +204,7 @@ public:
 	}
 
 	template<typename TParamRef>
-	void SetParameters(FRHICommandList* RHICmdList, const TParamRef& ShaderRHI, const FScene* Scene, bool bApplySkyLight)
+	void SetParameters(FRHICommandList& RHICmdList, const TParamRef& ShaderRHI, const FScene* Scene, bool bApplySkyLight)
 	{
 		FTexture* SkyLightTextureResource = GBlackTextureCube;
 		float ApplySkyLightMask = 0;
@@ -253,9 +253,9 @@ public:
 		SkyLightReflectionParameters.Bind(ParameterMap);
 	}
 
-	void Set(FRHICommandList* RHICmdList, FShader* Shader, const FSceneView* View);
+	void Set(FRHICommandList& RHICmdList, FShader* Shader, const FSceneView* View);
 
-	void SetMesh(FRHICommandList* RHICmdList, FShader* Shader, const FPrimitiveSceneProxy* Proxy);
+	void SetMesh(FRHICommandList& RHICmdList, FShader* Shader, const FPrimitiveSceneProxy* Proxy);
 
 	/** Serializer. */
 	friend FArchive& operator<<(FArchive& Ar,FTranslucentLightingParameters& P)
@@ -325,7 +325,7 @@ public:
 	TBasePassPixelShaderBaseType() {}
 
 	void SetParameters(
-		FRHICommandList* RHICmdList, 
+		FRHICommandList& RHICmdList, 
 		const FMaterialRenderProxy* MaterialRenderProxy, 
 		const FMaterial& MaterialResource, 
 		const FSceneView* View, 
@@ -379,7 +379,7 @@ public:
 		}
 	}
 
-	void SetMesh(FRHICommandList* RHICmdList, const FVertexFactory* VertexFactory,const FSceneView& View,const FPrimitiveSceneProxy* Proxy,const FMeshBatchElement& BatchElement, EBlendMode BlendMode)
+	void SetMesh(FRHICommandList& RHICmdList, const FVertexFactory* VertexFactory,const FSceneView& View,const FPrimitiveSceneProxy* Proxy,const FMeshBatchElement& BatchElement, EBlendMode BlendMode)
 	{
 		if (View.GetFeatureLevel() >= ERHIFeatureLevel::SM4
 			&& IsTranslucentBlendMode(BlendMode))
@@ -549,9 +549,9 @@ public:
 			LightMapPolicy == Other.LightMapPolicy;
 	}
 
-	void DrawShared(FRHICommandList* RHICmdList, const FSceneView* View,FBoundShaderStateRHIParamRef BoundShaderState) const
+	void DrawShared(FRHICommandList& RHICmdList, const FSceneView* View,FBoundShaderStateRHIParamRef BoundShaderState) const
 	{
-		FRHICommandList::SetBoundShaderState(RHICmdList, BoundShaderState);
+		RHICmdList.SetBoundShaderState(BoundShaderState);
 
 		// Set the light-map policy.
 		LightMapPolicy.Set(RHICmdList, VertexShader,bOverrideWithShaderComplexity ? NULL : PixelShader,VertexShader,PixelShader,VertexFactory,MaterialRenderProxy,View);
@@ -574,7 +574,7 @@ public:
 			// If we are in the translucent pass then override the blend mode, otherwise maintain additive blending.
 			if (IsTranslucentBlendMode(BlendMode))
 			{
-				FRHICommandList::SetBlendState(RHICmdList, TStaticBlendState<CW_RGBA, BO_Add, BF_One, BF_One, BO_Add, BF_Zero, BF_One>::GetRHI());
+				RHICmdList.SetBlendState( TStaticBlendState<CW_RGBA, BO_Add, BF_One, BF_One, BO_Add, BF_Zero, BF_One>::GetRHI());
 			}
 
 			TShaderMapRef<FShaderComplexityAccumulatePS> ShaderComplexityPixelShader(GetGlobalShaderMap());
@@ -599,16 +599,16 @@ public:
 			case BLEND_Translucent:
 				// Alpha channel is only needed for SeparateTranslucency, before this was preserving the alpha channel but we no longer store depth in the alpha channel so it's no problem
 
-				FRHICommandList::SetBlendState(RHICmdList, TStaticBlendState<CW_RGBA, BO_Add, BF_SourceAlpha, BF_InverseSourceAlpha, BO_Add, BF_Zero, BF_InverseSourceAlpha>::GetRHI());
+				RHICmdList.SetBlendState( TStaticBlendState<CW_RGBA, BO_Add, BF_SourceAlpha, BF_InverseSourceAlpha, BO_Add, BF_Zero, BF_InverseSourceAlpha>::GetRHI());
 				break;
 			case BLEND_Additive:
 				// Add to the existing scene color
 				// Alpha channel is only needed for SeparateTranslucency, before this was preserving the alpha channel but we no longer store depth in the alpha channel so it's no problem
-				FRHICommandList::SetBlendState(RHICmdList, TStaticBlendState<CW_RGBA, BO_Add, BF_One, BF_One, BO_Add, BF_Zero, BF_InverseSourceAlpha>::GetRHI());
+				RHICmdList.SetBlendState( TStaticBlendState<CW_RGBA, BO_Add, BF_One, BF_One, BO_Add, BF_Zero, BF_InverseSourceAlpha>::GetRHI());
 				break;
 			case BLEND_Modulate:
 				// Modulate with the existing scene color, preserve destination alpha.
-				FRHICommandList::SetBlendState(RHICmdList, TStaticBlendState<CW_RGB, BO_Add, BF_DestColor, BF_Zero>::GetRHI());
+				RHICmdList.SetBlendState( TStaticBlendState<CW_RGB, BO_Add, BF_DestColor, BF_Zero>::GetRHI());
 				break;
 			};
 		}
@@ -645,7 +645,7 @@ public:
 	}
 
 	void SetMeshRenderState(
-		FRHICommandList* RHICmdList, 
+		FRHICommandList& RHICmdList, 
 		const FSceneView& View,
 		const FPrimitiveSceneProxy* PrimitiveSceneProxy,
 		const FMeshBatch& Mesh,
@@ -683,7 +683,7 @@ public:
 			if (BlendMode != BLEND_Opaque)
 			{
 				// Add complexity to existing, keep alpha
-				check(!RHICmdList);
+				RHICmdList.CheckIsNull();
 				RHISetBlendState(TStaticBlendState<CW_RGB,BO_Add,BF_One,BF_One>::GetRHI());
 			}
 
@@ -760,7 +760,7 @@ public:
 
 	static void AddStaticMesh(FScene* Scene,FStaticMesh* StaticMesh);
 	static bool DrawDynamicMesh(
-		FRHICommandList* RHICmdList, 
+		FRHICommandList& RHICmdList, 
 		const FSceneView& View,
 		ContextType DrawingContext,
 		const FMeshBatch& Mesh,
@@ -845,7 +845,7 @@ public:
 /** Processes a base pass mesh using an unknown light map policy, and unknown fog density policy. */
 template<typename ProcessActionType>
 void ProcessBasePassMesh(
-	FRHICommandList* RHICmdList, 
+	FRHICommandList& RHICmdList, 
 	const FProcessBasePassMeshParameters& Parameters,
 	const ProcessActionType& Action
 	)
diff --git a/Engine/Source/Runtime/Renderer/Private/ClearQuad.cpp b/Engine/Source/Runtime/Renderer/Private/ClearQuad.cpp
--- a/Engine/Source/Runtime/Renderer/Private/ClearQuad.cpp
+++ b/Engine/Source/Runtime/Renderer/Private/ClearQuad.cpp
@@ -25,7 +25,7 @@ void DrawClearQuadMRT(bool bClearColor,int32 NumClearColors,const FLinearColor*
 	}
 	
 	//@todo-rco: RHIPacketList
-	FRHICommandList* RHICmdList = nullptr;
+	FRHICommandList& RHICmdList = FRHICommandList::GetNullRef();
 
 	const FDepthStencilStateRHIParamRef DepthStencilStateRHI = 
 		(bClearDepth && bClearStencil)
diff --git a/Engine/Source/Runtime/Renderer/Private/CompositionLighting/PostProcessAmbient.cpp b/Engine/Source/Runtime/Renderer/Private/CompositionLighting/PostProcessAmbient.cpp
--- a/Engine/Source/Runtime/Renderer/Private/CompositionLighting/PostProcessAmbient.cpp
+++ b/Engine/Source/Runtime/Renderer/Private/CompositionLighting/PostProcessAmbient.cpp
@@ -47,7 +47,7 @@ public:
 		PreIntegratedGFSampler.Bind(Initializer.ParameterMap, TEXT("PreIntegratedGFSampler"));
 	}
 
-	void SetParameters(FRHICommandList* RHICmdList, const FRenderingCompositePassContext& Context, const FFinalPostProcessSettings::FCubemapEntry& Entry)
+	void SetParameters(FRHICommandList& RHICmdList, const FRenderingCompositePassContext& Context, const FFinalPostProcessSettings::FCubemapEntry& Entry)
 	{
 		const FPixelShaderRHIParamRef ShaderRHI = GetPixelShader();
 
@@ -114,8 +114,7 @@ void FRCPassPostProcessAmbient::Process(FRenderingCompositePassContext& Context)
 		}
 
 		//@todo-rco: RHIPacketList
-		FRHICommandList* RHICmdList = nullptr;
-		PixelShader->SetParameters(nullptr, Context, Context.View.FinalPostProcessSettings.ContributingCubemaps[i]);
+		PixelShader->SetParameters(FRHICommandList::GetNullRef(), Context, Context.View.FinalPostProcessSettings.ContributingCubemaps[i]);
 
 		// Draw a quad mapping scene color to the view's render target
 		DrawRectangle( 
@@ -168,7 +167,7 @@ template<typename TShaderRHIRef>
 void FCubemapShaderParameters::SetParametersTemplate(const TShaderRHIRef& ShaderRHI, const FFinalPostProcessSettings::FCubemapEntry& Entry) const
 {
 	//@todo-rco: RHIPacketList
-	FRHICommandList* RHICmdList = nullptr;
+	FRHICommandList& RHICmdList = FRHICommandList::GetNullRef();
 
 	// floats to render the cubemap
 	{
diff --git a/Engine/Source/Runtime/Renderer/Private/CompositionLighting/PostProcessAmbientOcclusion.cpp b/Engine/Source/Runtime/Renderer/Private/CompositionLighting/PostProcessAmbientOcclusion.cpp
--- a/Engine/Source/Runtime/Renderer/Private/CompositionLighting/PostProcessAmbientOcclusion.cpp
+++ b/Engine/Source/Runtime/Renderer/Private/CompositionLighting/PostProcessAmbientOcclusion.cpp
@@ -55,7 +55,7 @@ public:
 		AmbientOcclusionSetupParams.Bind(Initializer.ParameterMap, TEXT("AmbientOcclusionSetupParams"));
 	}
 
-	void SetParameters(FRHICommandList* RHICmdList, const FRenderingCompositePassContext& Context)
+	void SetParameters(FRHICommandList& RHICmdList, const FRenderingCompositePassContext& Context)
 	{
 		const FFinalPostProcessSettings& Settings = Context.View.FinalPostProcessSettings;
 		const FPixelShaderRHIParamRef ShaderRHI = GetPixelShader();
@@ -117,7 +117,7 @@ void FCameraMotionParameters::Bind(const FShaderParameterMap& ParameterMap)
 	CameraMotion.Bind(ParameterMap, TEXT("CameraMotion"));
 }
 
-void FCameraMotionParameters::Set(FRHICommandList* RHICmdList, const FSceneView& View, const FPixelShaderRHIParamRef ShaderRHI) const
+void FCameraMotionParameters::Set(FRHICommandList& RHICmdList, const FSceneView& View, const FPixelShaderRHIParamRef ShaderRHI) const
 {
 	FSceneViewState* ViewState = (FSceneViewState*)View.State;
 
@@ -235,7 +235,7 @@ public:
 		RandomNormalTextureSampler.Bind(Initializer.ParameterMap, TEXT("RandomNormalTextureSampler"));
 	}
 
-	void SetParameters(FRHICommandList* RHICmdList, const FRenderingCompositePassContext& Context, FIntPoint InputTextureSize)
+	void SetParameters(FRHICommandList& RHICmdList, const FRenderingCompositePassContext& Context, FIntPoint InputTextureSize)
 	{
 		const FFinalPostProcessSettings& Settings = Context.View.FinalPostProcessSettings;
 		const FPixelShaderRHIParamRef ShaderRHI = GetPixelShader();
@@ -321,7 +321,7 @@ public:
 		ScreenSpaceAOandSSRShaderParams.Bind(Initializer.ParameterMap);
 	}
 
-	void SetParameters(FRHICommandList* RHICmdList, const FRenderingCompositePassContext& Context, FIntPoint InputTextureSize)
+	void SetParameters(FRHICommandList& RHICmdList, const FRenderingCompositePassContext& Context, FIntPoint InputTextureSize)
 	{
 		const FFinalPostProcessSettings& Settings = Context.View.FinalPostProcessSettings;
 		const FPixelShaderRHIParamRef ShaderRHI = GetPixelShader();
@@ -347,7 +347,7 @@ IMPLEMENT_SHADER_TYPE(,FPostProcessBasePassAOPS,TEXT("PostProcessAmbientOcclusio
 
 
 template <uint32 bInitialSetup>
-void FRCPassPostProcessAmbientOcclusionSetup::SetShaderSetupTempl(FRHICommandList* RHICmdList, const FRenderingCompositePassContext& Context)
+void FRCPassPostProcessAmbientOcclusionSetup::SetShaderSetupTempl(FRHICommandList& RHICmdList, const FRenderingCompositePassContext& Context)
 {
 	TShaderMapRef<FPostProcessVS> VertexShader(GetGlobalShaderMap());
 	TShaderMapRef<FPostProcessAmbientOcclusionSetupPS<bInitialSetup> > PixelShader(GetGlobalShaderMap());
@@ -387,7 +387,7 @@ void FRCPassPostProcessAmbientOcclusionSetup::Process(FRenderingCompositePassCon
 	RHISetDepthStencilState(TStaticDepthStencilState<false,CF_Always>::GetRHI());
 
 	//@todo-rco: RHIPacketList
-	FRHICommandList* RHICmdList = nullptr;
+	FRHICommandList& RHICmdList = FRHICommandList::GetNullRef();
 	if(IsInitialPass())
 	{
 		SetShaderSetupTempl<1>(RHICmdList, Context);
@@ -470,7 +470,7 @@ void FRCPassPostProcessAmbientOcclusion::SetShaderTempl(const FRenderingComposit
 	const FPooledRenderTargetDesc* InputDesc0 = GetInputDesc(ePId_Input0);
 
 	//@todo-rco: RHIPacketList
-	FRHICommandList* RHICmdList = nullptr;
+	FRHICommandList& RHICmdList = FRHICommandList::GetNullRef();
 	VertexShader->SetParameters(RHICmdList, Context);
 	PixelShader->SetParameters(RHICmdList, Context, InputDesc0->Extent);
 }
@@ -616,7 +616,7 @@ void FRCPassPostProcessBasePassAO::Process(FRenderingCompositePassContext& Conte
 	SetGlobalBoundShaderState(BoundShaderState, GFilterVertexDeclaration.VertexDeclarationRHI, *VertexShader, *PixelShader);
 
 	//@todo-rco: RHIPacketList
-	FRHICommandList* RHICmdList = nullptr;
+	FRHICommandList& RHICmdList = FRHICommandList::GetNullRef();
 	VertexShader->SetParameters(RHICmdList, Context);
 	PixelShader->SetParameters(RHICmdList, Context, GSceneRenderTargets.GetBufferSizeXY());
 
diff --git a/Engine/Source/Runtime/Renderer/Private/CompositionLighting/PostProcessAmbientOcclusion.h b/Engine/Source/Runtime/Renderer/Private/CompositionLighting/PostProcessAmbientOcclusion.h
--- a/Engine/Source/Runtime/Renderer/Private/CompositionLighting/PostProcessAmbientOcclusion.h
+++ b/Engine/Source/Runtime/Renderer/Private/CompositionLighting/PostProcessAmbientOcclusion.h
@@ -26,7 +26,7 @@ private:
 	bool IsInitialPass() const;
 
 	template <uint32 bInitialSetup>
-	void SetShaderSetupTempl(FRHICommandList* RHICmdList, const FRenderingCompositePassContext& Context);
+	void SetShaderSetupTempl(FRHICommandList& RHICmdList, const FRenderingCompositePassContext& Context);
 };
 
 // ePId_Input0: defines the resolution we compute AO and provides the normal
@@ -78,7 +78,7 @@ public:
 	}
 
 	template<typename ShaderRHIParamRef>
-	void Set(FRHICommandList* RHICmdList, const FSceneView& View, const ShaderRHIParamRef ShaderRHI, FIntPoint InputTextureSize) const
+	void Set(FRHICommandList& RHICmdList, const FSceneView& View, const ShaderRHIParamRef ShaderRHI, FIntPoint InputTextureSize) const
 	{
 		const FFinalPostProcessSettings& Settings = View.FinalPostProcessSettings;
 
@@ -154,7 +154,7 @@ public:
 
 	void Bind(const FShaderParameterMap& ParameterMap);
 
-	void Set(FRHICommandList* RHICmdList, const FSceneView& View, const FPixelShaderRHIParamRef ShaderRHI) const;
+	void Set(FRHICommandList& RHICmdList, const FSceneView& View, const FPixelShaderRHIParamRef ShaderRHI) const;
 
 	friend FArchive& operator<<(FArchive& Ar, FCameraMotionParameters& This);
 
diff --git a/Engine/Source/Runtime/Renderer/Private/CompositionLighting/PostProcessDeferredDecals.cpp b/Engine/Source/Runtime/Renderer/Private/CompositionLighting/PostProcessDeferredDecals.cpp
--- a/Engine/Source/Runtime/Renderer/Private/CompositionLighting/PostProcessDeferredDecals.cpp
+++ b/Engine/Source/Runtime/Renderer/Private/CompositionLighting/PostProcessDeferredDecals.cpp
@@ -327,7 +327,7 @@ public:
 	}
 	FStencilDecalMaskPS() {}
 
-	void SetParameters(FRHICommandList* RHICmdList, const FSceneView& View)
+	void SetParameters(FRHICommandList& RHICmdList, const FSceneView& View)
 	{
 		const FPixelShaderRHIParamRef ShaderRHI = GetPixelShader();
 
@@ -368,7 +368,7 @@ void StencilDecalMask(const FSceneView& View)
 	SetGlobalBoundShaderState(StencilDecalMaskBoundShaderState, GFilterVertexDeclaration.VertexDeclarationRHI, *ScreenVertexShader, *PixelShader);
 
 	//@todo-rco: RHIPacketList
-	FRHICommandList* RHICmdList = nullptr;
+	FRHICommandList& RHICmdList = FRHICommandList::GetNullRef();
 	PixelShader->SetParameters(RHICmdList, View);
 
 	DrawRectangle( 
@@ -407,7 +407,7 @@ public:
 		FrustumComponentToClip.Bind(Initializer.ParameterMap, TEXT("FrustumComponentToClip"));
 	}
 
-	void SetParameters(FRHICommandList* RHICmdList, const FSceneView& View, const FMatrix& InFrustumComponentToClip)
+	void SetParameters(FRHICommandList& RHICmdList, const FSceneView& View, const FMatrix& InFrustumComponentToClip)
 	{
 		FMaterialShader::SetParameters(RHICmdList, GetVertexShader(), View);
 		SetShaderValue(RHICmdList, GetVertexShader(), FrustumComponentToClip, InFrustumComponentToClip);
@@ -473,7 +473,7 @@ public:
 		DecalToWorld.Bind(Initializer.ParameterMap,TEXT("DecalToWorld"));
 	}
 
-	void SetParameters(FRHICommandList* RHICmdList, const FSceneView& View, const FMaterialRenderProxy* MaterialProxy, const FDeferredDecalProxy& DecalProxy)
+	void SetParameters(FRHICommandList& RHICmdList, const FSceneView& View, const FMaterialRenderProxy* MaterialProxy, const FDeferredDecalProxy& DecalProxy)
 	{
 		const FPixelShaderRHIParamRef ShaderRHI = GetPixelShader();
 
@@ -604,14 +604,14 @@ private:
 static TGlobalResource<FUnitCubeVertexBuffer> GUnitCubeVertexBuffer;
 static TGlobalResource<FUnitCubeIndexBuffer> GUnitCubeIndexBuffer;
 
-void SetShader(FRHICommandList* RHICmdList, const FRenderingCompositePassContext& Context, const FTransientDecalRenderData& DecalData, FShader* VertexShader)
+void SetShader(FRHICommandList& RHICmdList, const FRenderingCompositePassContext& Context, const FTransientDecalRenderData& DecalData, FShader* VertexShader)
 {
 	const FSceneView& View = Context.View;
 	
 	const FMaterialShaderMap* MaterialShaderMap = DecalData.MaterialResource->GetRenderingThreadShaderMap();
 	FDeferredDecalPS* PixelShader = MaterialShaderMap->GetShader<FDeferredDecalPS>();
 
-	check(!RHICmdList);
+	RHICmdList.CheckIsNull();
 
 	// This was cached but when changing the material (e.g. editor) it wasn't updated.
 	// This will change with upcoming multi threaded rendering changes.
@@ -629,7 +629,7 @@ void SetShader(FRHICommandList* RHICmdList, const FRenderingCompositePassContext
 	PixelShader->SetParameters(RHICmdList, View, DecalData.MaterialProxy, *DecalData.DecalProxy);
 }
 
-bool RenderPreStencil(FRHICommandList* RHICmdList, FRenderingCompositePassContext& Context, const FMaterialShaderMap* MaterialShaderMap, const FMatrix& ComponentToWorldMatrix, const FMatrix& FrustumComponentToClip)
+bool RenderPreStencil(FRHICommandList& RHICmdList, FRenderingCompositePassContext& Context, const FMaterialShaderMap* MaterialShaderMap, const FMatrix& ComponentToWorldMatrix, const FMatrix& FrustumComponentToClip)
 {
 	SCOPED_DRAW_EVENT(RenderPreStencil, DEC_SCENE_ITEMS);
 
@@ -654,7 +654,7 @@ bool RenderPreStencil(FRHICommandList* RHICmdList, FRenderingCompositePassContex
 
 	FDeferredDecalVS* VertexShader = MaterialShaderMap->GetShader<FDeferredDecalVS>();
 	
-	check(!RHICmdList);
+	RHICmdList.CheckIsNull();
 
 	// This was cached but when changing the material (e.g. editor) it wasn't updated.
 	// This will change with upcoming multi threaded rendering changes.
@@ -709,7 +709,7 @@ FRCPassPostProcessDeferredDecals::FRCPassPostProcessDeferredDecals(uint32 InRend
 void FRCPassPostProcessDeferredDecals::Process(FRenderingCompositePassContext& Context)
 {
 	//@todo-rco: RHIPacketList
-	FRHICommandList* RHICmdList = nullptr;
+	FRHICommandList& RHICmdList = FRHICommandList::GetNullRef();
 
 	bool bDBuffer = IsDBufferEnabled();
 	float bDecalPreStencil = CVarStencilSizeThreshold.GetValueOnRenderThread() >= 0;
diff --git a/Engine/Source/Runtime/Renderer/Private/CompositionLighting/PostProcessLpvIndirect.cpp b/Engine/Source/Runtime/Renderer/Private/CompositionLighting/PostProcessLpvIndirect.cpp
--- a/Engine/Source/Runtime/Renderer/Private/CompositionLighting/PostProcessLpvIndirect.cpp
+++ b/Engine/Source/Runtime/Renderer/Private/CompositionLighting/PostProcessLpvIndirect.cpp
@@ -67,7 +67,7 @@ public:
 	}
 
 	void SetParameters(	
-		FRHICommandList* RHICmdList, 
+		FRHICommandList& RHICmdList, 
 #if LPV_VOLUME_TEXTURE
 		FTextureRHIParamRef* LpvBufferSRVsIn, 
 #else
@@ -79,7 +79,7 @@ public:
 	{
 		const FPixelShaderRHIParamRef ShaderRHI = GetPixelShader();
 
-		check(!RHICmdList);
+		RHICmdList.CheckIsNull();
 
 		SetUniformBufferParameter(RHICmdList, ShaderRHI, GetUniformBufferParameter<FLpvReadUniformBufferParameters>(), LpvUniformBuffer );
 
@@ -154,8 +154,8 @@ void FRCPassPostProcessLpvIndirect::Process(FRenderingCompositePassContext& Cont
 	};
 
 	//@todo-rco: RHIPacketList
-	FRHICommandList* RHICmdList = nullptr;
-	check(!RHICmdList);
+	FRHICommandList& RHICmdList = FRHICommandList::GetNullRef();
+	RHICmdList.CheckIsNull();
 
 	// Set the view family's render target/viewport.
 	RHISetRenderTargets(1, RenderTargets, GSceneRenderTargets.GetSceneDepthSurface(), 0, NULL);
diff --git a/Engine/Source/Runtime/Renderer/Private/CompositionLighting/PostProcessPassThrough.cpp b/Engine/Source/Runtime/Renderer/Private/CompositionLighting/PostProcessPassThrough.cpp
--- a/Engine/Source/Runtime/Renderer/Private/CompositionLighting/PostProcessPassThrough.cpp
+++ b/Engine/Source/Runtime/Renderer/Private/CompositionLighting/PostProcessPassThrough.cpp
@@ -46,7 +46,7 @@ public:
 		const FPixelShaderRHIParamRef ShaderRHI = GetPixelShader();
 
 		//@todo-rco: RHIPacketList
-		FRHICommandList* RHICmdList = nullptr;
+		FRHICommandList& RHICmdList = FRHICommandList::GetNullRef();
 		FGlobalShader::SetParameters(RHICmdList, ShaderRHI, Context.View);
 
 		PostprocessParameter.SetPS(RHICmdList, ShaderRHI, Context, TStaticSamplerState<SF_Point,AM_Clamp,AM_Clamp,AM_Clamp>::GetRHI());
@@ -124,7 +124,7 @@ void FRCPassPostProcessPassThrough::Process(FRenderingCompositePassContext& Cont
 	SetGlobalBoundShaderState(BoundShaderState, GFilterVertexDeclaration.VertexDeclarationRHI, *VertexShader, *PixelShader);
 
 	//@todo-rco: RHIPacketList
-	FRHICommandList* RHICmdList = nullptr;
+	FRHICommandList& RHICmdList = FRHICommandList::GetNullRef();
 	VertexShader->SetParameters(RHICmdList, Context);
 	PixelShader->SetParameters(Context);
 
diff --git a/Engine/Source/Runtime/Renderer/Private/CustomDepthRendering.cpp b/Engine/Source/Runtime/Renderer/Private/CustomDepthRendering.cpp
--- a/Engine/Source/Runtime/Renderer/Private/CustomDepthRendering.cpp
+++ b/Engine/Source/Runtime/Renderer/Private/CustomDepthRendering.cpp
@@ -16,7 +16,7 @@
 	FCustomDepthPrimSet
 -----------------------------------------------------------------------------*/
 
-bool FCustomDepthPrimSet::DrawPrims(FRHICommandList* RHICmdList, const FViewInfo* ViewInfo,bool bInitializeOffsets)
+bool FCustomDepthPrimSet::DrawPrims(FRHICommandList& RHICmdList, const FViewInfo* ViewInfo,bool bInitializeOffsets)
 {
 	bool bDirty=false;
 
diff --git a/Engine/Source/Runtime/Renderer/Private/CustomDepthRendering.h b/Engine/Source/Runtime/Renderer/Private/CustomDepthRendering.h
--- a/Engine/Source/Runtime/Renderer/Private/CustomDepthRendering.h
+++ b/Engine/Source/Runtime/Renderer/Private/CustomDepthRendering.h
@@ -18,7 +18,7 @@ public:
 	* @param ViewInfo - current view used to draw items
 	* @return true if anything was drawn
 	*/
-	bool DrawPrims(FRHICommandList* RHICmdList, const class FViewInfo* ViewInfo,bool bInitializeOffsets);
+	bool DrawPrims(FRHICommandList& RHICmdList, const class FViewInfo* ViewInfo,bool bInitializeOffsets);
 
 	/**
 	* Add a new primitive to the list of prims
diff --git a/Engine/Source/Runtime/Renderer/Private/DeferredShadingRenderer.cpp b/Engine/Source/Runtime/Renderer/Private/DeferredShadingRenderer.cpp
--- a/Engine/Source/Runtime/Renderer/Private/DeferredShadingRenderer.cpp
+++ b/Engine/Source/Runtime/Renderer/Private/DeferredShadingRenderer.cpp
@@ -155,7 +155,7 @@ void FDeferredShadingSceneRenderer::ClearGBufferAtMaxZ()
 	}
 
 	//@todo-rco: RHIPacketList
-	FRHICommandList* RHICmdList = nullptr;
+	FRHICommandList& RHICmdList = FRHICommandList::GetNullRef();
 
 	SetGlobalBoundShaderState(GClearMRTBoundShaderState[NumActiveRenderTargets - 1], GetVertexDeclarationFVector4(), *VertexShader, PixelShader);
 
@@ -183,7 +183,7 @@ void FDeferredShadingSceneRenderer::ClearGBufferAtMaxZ()
 	}
 }
 
-bool FDeferredShadingSceneRenderer::RenderBasePassStaticDataMasked(FRHICommandList* RHICmdList, FViewInfo& View)
+bool FDeferredShadingSceneRenderer::RenderBasePassStaticDataMasked(FRHICommandList& RHICmdList, FViewInfo& View)
 {
 	bool bDirty = false;
 
@@ -208,7 +208,7 @@ bool FDeferredShadingSceneRenderer::RenderBasePassStaticDataMasked(FRHICommandLi
 	return bDirty;
 }
 
-bool FDeferredShadingSceneRenderer::RenderBasePassStaticDataDefault(FRHICommandList* RHICmdList, FViewInfo& View)
+bool FDeferredShadingSceneRenderer::RenderBasePassStaticDataDefault(FRHICommandList& RHICmdList, FViewInfo& View)
 {
 	bool bDirty = false;
 
@@ -258,7 +258,7 @@ void FDeferredShadingSceneRenderer::SortBasePassStaticData(FVector ViewPosition)
 *
 * @return true if anything was rendered to scene color
 */
-bool FDeferredShadingSceneRenderer::RenderBasePassStaticData(FRHICommandList* RHICmdList, FViewInfo& View)
+bool FDeferredShadingSceneRenderer::RenderBasePassStaticData(FRHICommandList& RHICmdList, FViewInfo& View)
 {
 	bool bDirty = false;
 
@@ -330,7 +330,7 @@ bool FDeferredShadingSceneRenderer::RenderBasePass(FViewInfo& View)
 		SCOPE_CYCLE_COUNTER(STAT_RHICounterTEMP);
 		static IConsoleVariable* RHICmdListCVar = IConsoleManager::Get().FindConsoleVariable(TEXT("r.RHICmd"));
 		bool bUseRHICmdList = RHICmdListCVar->GetInt() != 0;
-		auto* RHICmdList = bUseRHICmdList ? GRHICommandList.CreateList() : nullptr;
+		FRHICommandList& RHICmdList = bUseRHICmdList ? GRHICommandList.CreateList() : FRHICommandList::GetNullRef();
 		bDirty |= RenderBasePassStaticData(RHICmdList, View);
 		GRHICommandList.ExecuteAndFreeList(RHICmdList);
 	}
@@ -516,7 +516,7 @@ void FDeferredShadingSceneRenderer::Render()
 	#endif
 
 	//@todo-rco: RHIPacketList
-	FRHICommandList* RHICmdList = nullptr;
+	FRHICommandList& RHICmdList = FRHICommandList::GetNullRef();
 
 	const bool bIsOcclusionTesting = DoOcclusionQueries() && (!bIsWireframe || bIsViewFrozen || bHasViewParent);
 
@@ -803,7 +803,7 @@ bool FDeferredShadingSceneRenderer::RenderPrePass()
 	GSceneRenderTargets.BeginRenderingPrePass();
 
 	//@todo-rco: RHIPacketList
-	FRHICommandList* RHICmdList = nullptr;
+	FRHICommandList& RHICmdList = FRHICommandList::GetNullRef();
 
 	// Clear the depth buffer.
 	// Note, this is a reversed Z depth surface, so 0.0f is the far plane.
@@ -899,7 +899,7 @@ bool FDeferredShadingSceneRenderer::RenderBasePass()
 	bool bDirty = false;
 
 	//@todo-rco: RHIPacketList
-	FRHICommandList* RHICmdList = nullptr;
+	FRHICommandList& RHICmdList = FRHICommandList::GetNullRef();
 
 	if(ViewFamily.EngineShowFlags.LightMapDensity && AllowDebugViewmodes())
 	{
@@ -940,7 +940,7 @@ bool FDeferredShadingSceneRenderer::RenderBasePass()
 	return bDirty;
 }
 
-void FDeferredShadingSceneRenderer::ClearLPVs(FRHICommandList* RHICmdList)
+void FDeferredShadingSceneRenderer::ClearLPVs(FRHICommandList& RHICmdList)
 {
 	// clear light propagation volumes
 
@@ -980,7 +980,7 @@ void FDeferredShadingSceneRenderer::PropagateLPVs()
 				SCOPED_DRAW_EVENT(UpdateLPVs, DEC_SCENE_ITEMS);
 				SCOPE_CYCLE_COUNTER(STAT_UpdateLPVs);
 				//@todo-rco: RHIPacketList
-				LightPropagationVolume->Propagate(nullptr);
+				LightPropagationVolume->Propagate(FRHICommandList::GetNullRef());
 			}
 		}
 	}
@@ -1007,7 +1007,7 @@ public:
 	}
 	FDownsampleSceneDepthPS() {}
 
-	void SetParameters(FRHICommandList* RHICmdList, const FSceneView& View)
+	void SetParameters(FRHICommandList& RHICmdList, const FSceneView& View)
 	{
 		FGlobalShader::SetParameters(RHICmdList, GetPixelShader(), View);
 
@@ -1054,7 +1054,7 @@ void FDeferredShadingSceneRenderer::UpdateDownsampledDepthSurface()
 	if (GSceneRenderTargets.UseDownsizedOcclusionQueries() && IsFeatureLevelSupported(GRHIShaderPlatform, ERHIFeatureLevel::SM3))
 	{
 		//@todo-rco: RHIPacketList
-		FRHICommandList* RHICmdList = nullptr;
+		FRHICommandList& RHICmdList = FRHICommandList::GetNullRef();
 
 		RHISetRenderTarget(NULL, GSceneRenderTargets.GetSmallDepthSurface());
 
diff --git a/Engine/Source/Runtime/Renderer/Private/DeferredShadingRenderer.h b/Engine/Source/Runtime/Renderer/Private/DeferredShadingRenderer.h
--- a/Engine/Source/Runtime/Renderer/Private/DeferredShadingRenderer.h
+++ b/Engine/Source/Runtime/Renderer/Private/DeferredShadingRenderer.h
@@ -45,15 +45,15 @@ public:
     void ClearGBufferAtMaxZ();
 
 	/** Clears LPVs for all views */
-	void ClearLPVs(FRHICommandList* RHICmdList);
+	void ClearLPVs(FRHICommandList& RHICmdList);
 
 	/** Propagates LPVs for all views */
 	void PropagateLPVs();
 
     /** Renders the basepass for the static data of a given View. */
-    bool RenderBasePassStaticData(FRHICommandList* RHICmdList, FViewInfo& View);
-	bool RenderBasePassStaticDataMasked(FRHICommandList* RHICmdList, FViewInfo& View);
-	bool RenderBasePassStaticDataDefault(FRHICommandList* RHICmdList, FViewInfo& View);
+    bool RenderBasePassStaticData(FRHICommandList& RHICmdList, FViewInfo& View);
+	bool RenderBasePassStaticDataMasked(FRHICommandList& RHICmdList, FViewInfo& View);
+	bool RenderBasePassStaticDataDefault(FRHICommandList& RHICmdList, FViewInfo& View);
 
 	/** Sorts base pass draw lists front to back for improved GPU culling. */
 	void SortBasePassStaticData(FVector ViewPosition);
@@ -174,7 +174,7 @@ private:
 	void RenderAtmosphere(FLightShaftsOutput LightShaftsOutput);
 
 	/** Renders reflections that can be done in a deferred pass. */
-	void RenderDeferredReflections(FRHICommandList* RHICmdList);
+	void RenderDeferredReflections(FRHICommandList& RHICmdList);
 
 	/** Render dynamic sky lighting from Movable sky lights. */
 	void RenderDynamicSkyLighting();
@@ -189,10 +189,10 @@ private:
 	bool ShouldUseTiledDeferred(int32 NumUnshadowedLights, int32 NumSimpleLights) const;
 
 	/** Renders the lights in SortedLights in the range [0, NumUnshadowedLights) using tiled deferred shading. */
-	void RenderTiledDeferredLighting(FRHICommandList* RHICmdList, const TArray<FSortedLightSceneInfo, SceneRenderingAllocator>& SortedLights, int32 NumUnshadowedLights, const FSimpleLightArray& SimpleLights);
+	void RenderTiledDeferredLighting(FRHICommandList& RHICmdList, const TArray<FSortedLightSceneInfo, SceneRenderingAllocator>& SortedLights, int32 NumUnshadowedLights, const FSimpleLightArray& SimpleLights);
 
 	/** Renders the scene's lighting. */
-	void RenderLights(FRHICommandList* RHICmdList);
+	void RenderLights(FRHICommandList& RHICmdList);
 
 	/** Renders an array of lights for the stationary light overlap viewmode. */
 	void RenderLightArrayForOverlapViewmode(const TSparseArray<FLightSceneInfoCompact>& LightArray);
@@ -240,20 +240,20 @@ private:
 	bool CheckForProjectedShadows( const FLightSceneInfo* LightSceneInfo ) const;
 
 	/** Renders one pass point light shadows. */
-	bool RenderOnePassPointLightShadows(FRHICommandList* RHICmdList, const FLightSceneInfo* LightSceneInfo, bool bRenderedTranslucentObjectShadows, bool& bInjectedTranslucentVolume);
+	bool RenderOnePassPointLightShadows(FRHICommandList& RHICmdList, const FLightSceneInfo* LightSceneInfo, bool bRenderedTranslucentObjectShadows, bool& bInjectedTranslucentVolume);
 
 	/** Renders the projections of the given Shadows to the appropriate color render target. */
 	void RenderProjections(
-		FRHICommandList* RHICmdList, 
+		FRHICommandList& RHICmdList, 
 		const FLightSceneInfo* LightSceneInfo, 
 		const TArray<FProjectedShadowInfo*,SceneRenderingAllocator>& Shadows
 		);
 
 	/** Renders the shadowmaps of translucent shadows and their projections onto opaque surfaces. */
-	bool RenderTranslucentProjectedShadows(FRHICommandList* RHICmdList, const FLightSceneInfo* LightSceneInfo );
+	bool RenderTranslucentProjectedShadows(FRHICommandList& RHICmdList, const FLightSceneInfo* LightSceneInfo );
 
 	/** Renders reflective shadowmaps for LPVs */
-	bool RenderReflectiveShadowMaps(FRHICommandList* RHICmdList, const FLightSceneInfo* LightSceneInfo );
+	bool RenderReflectiveShadowMaps(FRHICommandList& RHICmdList, const FLightSceneInfo* LightSceneInfo );
 
 	/**
 	  * Used by RenderLights to render projected shadows to the attenuation buffer.
@@ -261,10 +261,10 @@ private:
 	  * @param LightSceneInfo Represents the current light
 	  * @return true if anything got rendered
 	  */
-	bool RenderProjectedShadows(FRHICommandList* RHICmdList, const FLightSceneInfo* LightSceneInfo, bool bRenderedTranslucentObjectShadows, bool& bInjectedTranslucentVolume );
+	bool RenderProjectedShadows(FRHICommandList& RHICmdList, const FLightSceneInfo* LightSceneInfo, bool bRenderedTranslucentObjectShadows, bool& bInjectedTranslucentVolume );
 
 	/** Caches the depths of any preshadows that should be cached, and renders their projections. */
-	bool RenderCachedPreshadows(FRHICommandList* RHICmdList, const FLightSceneInfo* LightSceneInfo);
+	bool RenderCachedPreshadows(FRHICommandList& RHICmdList, const FLightSceneInfo* LightSceneInfo);
 
 	/**
 	  * Used by RenderLights to figure out if light functions need to be rendered to the attenuation buffer.
@@ -280,13 +280,13 @@ private:
 	  * @param LightSceneInfo Represents the current light
 	  * @param LightIndex The light's index into FScene::Lights
 	  */
-	bool RenderLightFunction(FRHICommandList* RHICmdList, const FLightSceneInfo* LightSceneInfo, bool bLightAttenuationCleared);
+	bool RenderLightFunction(FRHICommandList& RHICmdList, const FLightSceneInfo* LightSceneInfo, bool bLightAttenuationCleared);
 
 	/** Renders a light function indicating that whole scene shadowing being displayed is for previewing only, and will go away in game. */
-	bool RenderPreviewShadowsIndicator(FRHICommandList* RHICmdList, const FLightSceneInfo* LightSceneInfo, bool bLightAttenuationCleared);
+	bool RenderPreviewShadowsIndicator(FRHICommandList& RHICmdList, const FLightSceneInfo* LightSceneInfo, bool bLightAttenuationCleared);
 
 	/** Renders a light function with the given material. */
-	bool RenderLightFunctionForMaterial(FRHICommandList* RHICmdList, const FLightSceneInfo* LightSceneInfo, const FMaterialRenderProxy* MaterialProxy, bool bLightAttenuationCleared, bool bRenderingPreviewShadowsIndicator);
+	bool RenderLightFunctionForMaterial(FRHICommandList& RHICmdList, const FLightSceneInfo* LightSceneInfo, const FMaterialRenderProxy* MaterialProxy, bool bLightAttenuationCleared, bool bRenderingPreviewShadowsIndicator);
 
 	/**
 	  * Used by RenderLights to render a light to the scene color buffer.
@@ -301,13 +301,13 @@ private:
 	void RenderSimpleLightsStandardDeferred(const FSimpleLightArray& SimpleLights);
 
 	/** Clears the translucency lighting volumes before light accumulation. */
-	void ClearTranslucentVolumeLighting(FRHICommandList* RHICmdList);
+	void ClearTranslucentVolumeLighting(FRHICommandList& RHICmdList);
 
 	/** Add AmbientCubemap to the lighting volumes. */
-	void InjectAmbientCubemapTranslucentVolumeLighting(FRHICommandList* RHICmdList);
+	void InjectAmbientCubemapTranslucentVolumeLighting(FRHICommandList& RHICmdList);
 
 	/** Renders indirect lighting into the translucent lighting volumes. */
-	void CompositeIndirectTranslucentVolumeLighting(FRHICommandList* RHICmdList);
+	void CompositeIndirectTranslucentVolumeLighting(FRHICommandList& RHICmdList);
 
 	/** Clears the volume texture used to accumulate per object shadows for translucency. */
 	void ClearTranslucentVolumePerObjectShadowing();
@@ -331,10 +331,10 @@ private:
 	void RenderReflectionCaptureSpecularBounceForAllViews();
 
 	/** Render image based reflections (SSR, Env, SkyLight) without compute shaders */
-	void RenderImageBasedReflectionsSM4ForAllViews(FRHICommandList* RHICmdList, bool bReflectionEnv);
+	void RenderImageBasedReflectionsSM4ForAllViews(FRHICommandList& RHICmdList, bool bReflectionEnv);
 
 	/** Render image based reflections (SSR, Env, SkyLight) with compute shaders */
-	void RenderImageBasedReflectionsSM5ForAllViews(FRHICommandList* RHICmdList);
+	void RenderImageBasedReflectionsSM5ForAllViews(FRHICommandList& RHICmdList);
 
 	bool ShouldDoReflectionEnvironment() const;
 
diff --git a/Engine/Source/Runtime/Renderer/Private/DepthRendering.cpp b/Engine/Source/Runtime/Renderer/Private/DepthRendering.cpp
--- a/Engine/Source/Runtime/Renderer/Private/DepthRendering.cpp
+++ b/Engine/Source/Runtime/Renderer/Private/DepthRendering.cpp
@@ -37,12 +37,12 @@ public:
 			 && IsFeatureLevelSupported(Platform, ERHIFeatureLevel::SM3);
 	}
 
-	void SetParameters(FRHICommandList* RHICmdList, const FMaterialRenderProxy* MaterialRenderProxy,const FMaterial& MaterialResource,const FSceneView& View)
+	void SetParameters(FRHICommandList& RHICmdList, const FMaterialRenderProxy* MaterialRenderProxy,const FMaterial& MaterialResource,const FSceneView& View)
 	{
 		FMeshMaterialShader::SetParameters(RHICmdList, GetVertexShader(),MaterialRenderProxy,MaterialResource,View,ESceneRenderTargetsMode::DontSet);
 	}
 
-	void SetMesh(FRHICommandList* RHICmdList, const FVertexFactory* VertexFactory,const FSceneView& View,const FPrimitiveSceneProxy* Proxy,const FMeshBatchElement& BatchElement)
+	void SetMesh(FRHICommandList& RHICmdList, const FVertexFactory* VertexFactory,const FSceneView& View,const FPrimitiveSceneProxy* Proxy,const FMeshBatchElement& BatchElement)
 	{
 		FMeshMaterialShader::SetMesh(RHICmdList, GetVertexShader(),VertexFactory,View,Proxy,BatchElement);
 	}
@@ -115,12 +115,12 @@ public:
 
 	FDepthOnlyPS() {}
 
-	void SetParameters(FRHICommandList* RHICmdList, const FMaterialRenderProxy* MaterialRenderProxy,const FMaterial& MaterialResource,const FSceneView* View)
+	void SetParameters(FRHICommandList& RHICmdList, const FMaterialRenderProxy* MaterialRenderProxy,const FMaterial& MaterialResource,const FSceneView* View)
 	{
 		FMeshMaterialShader::SetParameters(RHICmdList, GetPixelShader(),MaterialRenderProxy,MaterialResource,*View,ESceneRenderTargetsMode::DontSet);
 	}
 
-	void SetMesh(FRHICommandList* RHICmdList, const FVertexFactory* VertexFactory,const FSceneView& View,const FPrimitiveSceneProxy* Proxy,const FMeshBatchElement& BatchElement)
+	void SetMesh(FRHICommandList& RHICmdList, const FVertexFactory* VertexFactory,const FSceneView& View,const FPrimitiveSceneProxy* Proxy,const FMeshBatchElement& BatchElement)
 	{
 		FMeshMaterialShader::SetMesh(RHICmdList, GetPixelShader(),VertexFactory,View,Proxy,BatchElement);
 	}
@@ -169,7 +169,7 @@ FDepthDrawingPolicy::FDepthDrawingPolicy(
 
 }
 
-void FDepthDrawingPolicy::DrawShared(FRHICommandList* RHICmdList, const FSceneView* View,FBoundShaderStateRHIParamRef BoundShaderState) const
+void FDepthDrawingPolicy::DrawShared(FRHICommandList& RHICmdList, const FSceneView* View,FBoundShaderStateRHIParamRef BoundShaderState) const
 {
 	// Set the actual shader & vertex declaration state
 	RHISetBoundShaderState(BoundShaderState);
@@ -208,7 +208,7 @@ FBoundShaderStateRHIRef FDepthDrawingPolicy::CreateBoundShaderState(ERHIFeatureL
 }
 
 void FDepthDrawingPolicy::SetMeshRenderState(
-	FRHICommandList* RHICmdList, 
+	FRHICommandList& RHICmdList, 
 	const FSceneView& View,
 	const FPrimitiveSceneProxy* PrimitiveSceneProxy,
 	const FMeshBatch& Mesh,
@@ -257,9 +257,9 @@ FPositionOnlyDepthDrawingPolicy::FPositionOnlyDepthDrawingPolicy(
 	VertexShader = InMaterialResource.GetShader<TDepthOnlyVS<true> >(InVertexFactory->GetType());
 }
 
-void FPositionOnlyDepthDrawingPolicy::DrawShared(FRHICommandList* RHICmdList, const FSceneView* View,FBoundShaderStateRHIParamRef BoundShaderState) const
+void FPositionOnlyDepthDrawingPolicy::DrawShared(FRHICommandList& RHICmdList, const FSceneView* View,FBoundShaderStateRHIParamRef BoundShaderState) const
 {
-	check(!RHICmdList);
+	RHICmdList.CheckIsNull();
 
 	// Set the actual shader & vertex declaration state
 	RHISetBoundShaderState( BoundShaderState);
@@ -286,7 +286,7 @@ FBoundShaderStateRHIRef FPositionOnlyDepthDrawingPolicy::CreateBoundShaderState(
 }
 
 void FPositionOnlyDepthDrawingPolicy::SetMeshRenderState(
-	FRHICommandList* RHICmdList, 
+	FRHICommandList& RHICmdList, 
 	const FSceneView& View,
 	const FPrimitiveSceneProxy* PrimitiveSceneProxy,
 	const FMeshBatch& Mesh,
@@ -374,7 +374,7 @@ void FDepthDrawingPolicyFactory::AddStaticMesh(FScene* Scene,FStaticMesh* Static
 }
 
 bool FDepthDrawingPolicyFactory::DrawMesh(
-	FRHICommandList* RHICmdList, 
+	FRHICommandList& RHICmdList, 
 	const FViewInfo& View,
 	ContextType DrawingContext,
 	const FMeshBatch& Mesh,
@@ -472,7 +472,7 @@ bool FDepthDrawingPolicyFactory::DrawMesh(
 }
 
 bool FDepthDrawingPolicyFactory::DrawDynamicMesh(
-	FRHICommandList* RHICmdList, 
+	FRHICommandList& RHICmdList, 
 	const FViewInfo& View,
 	ContextType DrawingContext,
 	const FMeshBatch& Mesh,
@@ -496,7 +496,7 @@ bool FDepthDrawingPolicyFactory::DrawDynamicMesh(
 
 
 bool FDepthDrawingPolicyFactory::DrawStaticMesh(
-	FRHICommandList* RHICmdList, 
+	FRHICommandList& RHICmdList, 
 	const FViewInfo& View,
 	ContextType DrawingContext,
 	const FStaticMesh& StaticMesh,
diff --git a/Engine/Source/Runtime/Renderer/Private/DepthRendering.h b/Engine/Source/Runtime/Renderer/Private/DepthRendering.h
--- a/Engine/Source/Runtime/Renderer/Private/DepthRendering.h
+++ b/Engine/Source/Runtime/Renderer/Private/DepthRendering.h
@@ -44,7 +44,7 @@ public:
 			&& PixelShader == Other.PixelShader;
 	}
 
-	void DrawShared(FRHICommandList* RHICmdList, const FSceneView* View,FBoundShaderStateRHIParamRef BoundShaderState) const;
+	void DrawShared(FRHICommandList& RHICmdList, const FSceneView* View,FBoundShaderStateRHIParamRef BoundShaderState) const;
 
 	/** 
 	* Create bound shader state using the vertex decl from the mesh draw policy
@@ -55,7 +55,7 @@ public:
 	FBoundShaderStateRHIRef CreateBoundShaderState(ERHIFeatureLevel::Type InFeatureLevel);
 
 	void SetMeshRenderState(
-		FRHICommandList* RHICmdList, 
+		FRHICommandList& RHICmdList, 
 		const FSceneView& View,
 		const FPrimitiveSceneProxy* PrimitiveSceneProxy,
 		const FMeshBatch& Mesh,
@@ -97,7 +97,7 @@ public:
 		return FMeshDrawingPolicy::Matches(Other) && VertexShader == Other.VertexShader;
 	}
 
-	void DrawShared(FRHICommandList* RHICmdList, const FSceneView* View,FBoundShaderStateRHIParamRef BoundShaderState) const;
+	void DrawShared(FRHICommandList& RHICmdList, const FSceneView* View,FBoundShaderStateRHIParamRef BoundShaderState) const;
 
 	/** 
 	* Create bound shader state using the vertex decl from the mesh draw policy
@@ -107,7 +107,7 @@ public:
 	FBoundShaderStateRHIRef CreateBoundShaderState(ERHIFeatureLevel::Type InFeatureLevel);
 
 	void SetMeshRenderState(
-		FRHICommandList* RHICmdList, 
+		FRHICommandList& RHICmdList, 
 		const FSceneView& View,
 		const FPrimitiveSceneProxy* PrimitiveSceneProxy,
 		const FMeshBatch& Mesh,
@@ -141,7 +141,7 @@ public:
 
 	static void AddStaticMesh(FScene* Scene,FStaticMesh* StaticMesh);
 	static bool DrawDynamicMesh(
-		FRHICommandList* RHICmdList, 
+		FRHICommandList& RHICmdList, 
 		const FViewInfo& View,
 		ContextType DrawingContext,
 		const FMeshBatch& Mesh,
@@ -152,7 +152,7 @@ public:
 		);
 
 	static bool DrawStaticMesh(
-		FRHICommandList* RHICmdList, 
+		FRHICommandList& RHICmdList, 
 		const FViewInfo& View,
 		ContextType DrawingContext,
 		const FStaticMesh& StaticMesh,
@@ -170,7 +170,7 @@ private:
 	* @return true if the mesh rendered
 	*/
 	static bool DrawMesh(
-		FRHICommandList* RHICmdList, 
+		FRHICommandList& RHICmdList, 
 		const FViewInfo& View,
 		ContextType DrawingContext,
 		const FMeshBatch& Mesh,
diff --git a/Engine/Source/Runtime/Renderer/Private/DistanceFieldLightingShared.h b/Engine/Source/Runtime/Renderer/Private/DistanceFieldLightingShared.h
--- a/Engine/Source/Runtime/Renderer/Private/DistanceFieldLightingShared.h
+++ b/Engine/Source/Runtime/Renderer/Private/DistanceFieldLightingShared.h
@@ -63,7 +63,7 @@ public:
 	}
 
 	template<typename ShaderRHIParamRef>
-	void Set(FRHICommandList* RHICmdList, const ShaderRHIParamRef ShaderRHI)
+	void Set(FRHICommandList& RHICmdList, const ShaderRHIParamRef ShaderRHI)
 	{
 		float AOMaxDistanceValue = GetAOMaxDistance();
 		SetShaderValue(RHICmdList, ShaderRHI, AOMaxDistance, AOMaxDistanceValue);
diff --git a/Engine/Source/Runtime/Renderer/Private/DistanceFieldSurfaceCacheLighting.cpp b/Engine/Source/Runtime/Renderer/Private/DistanceFieldSurfaceCacheLighting.cpp
--- a/Engine/Source/Runtime/Renderer/Private/DistanceFieldSurfaceCacheLighting.cpp
+++ b/Engine/Source/Runtime/Renderer/Private/DistanceFieldSurfaceCacheLighting.cpp
@@ -361,7 +361,7 @@ public:
 	}
 
 	template<typename TParamRef>
-	void Set(FRHICommandList* RHICmdList, const TParamRef& ShaderRHI, int32 NumObjectsValue)
+	void Set(FRHICommandList& RHICmdList, const TParamRef& ShaderRHI, int32 NumObjectsValue)
 	{
 		SetSRVParameter(RHICmdList, ShaderRHI, ObjectBounds, GAOObjectBuffers.ObjectData.BoundsSRV);
 		SetSRVParameter(RHICmdList, ShaderRHI, ObjectData, GAOObjectBuffers.ObjectData.DataSRV);
@@ -423,7 +423,7 @@ public:
 	FBuildTileConesCS()
 	{
 	}
-	void SetParameters(FRHICommandList* RHICmdList, const FSceneView& View, FScene* Scene, FVector2D NumGroupsValue)
+	void SetParameters(FRHICommandList& RHICmdList, const FSceneView& View, FScene* Scene, FVector2D NumGroupsValue)
 	{
 		FComputeShaderRHIParamRef ShaderRHI = GetComputeShader();
 
@@ -442,7 +442,7 @@ public:
 		SetShaderValue(RHICmdList, ShaderRHI, NumGroups, NumGroupsValue);
 	}
 
-	void UnsetParameters(FRHICommandList* RHICmdList)
+	void UnsetParameters(FRHICommandList& RHICmdList)
 	{
 		TileConeAxisAndCos.UnsetUAV(RHICmdList, GetComputeShader());
 		TileConeDepthRanges.UnsetUAV(RHICmdList, GetComputeShader());
@@ -502,7 +502,7 @@ public:
 
 	FObjectCullVS() {}
 
-	void SetParameters(FRHICommandList* RHICmdList, const FSceneView& View, int32 NumObjects)
+	void SetParameters(FRHICommandList& RHICmdList, const FSceneView& View, int32 NumObjects)
 	{
 		const FVertexShaderRHIParamRef ShaderRHI = GetVertexShader();
 		FGlobalShader::SetParameters(RHICmdList, ShaderRHI, View);
@@ -569,7 +569,7 @@ public:
 		NumGroups.Bind(Initializer.ParameterMap, TEXT("NumGroups"));
 	}
 
-	void SetParameters(FRHICommandList* RHICmdList, const FSceneView& View, FVector2D NumGroupsValue)
+	void SetParameters(FRHICommandList& RHICmdList, const FSceneView& View, FVector2D NumGroupsValue)
 	{
 		const FPixelShaderRHIParamRef ShaderRHI = GetPixelShader();
 		FGlobalShader::SetParameters(RHICmdList, ShaderRHI, View);
@@ -667,7 +667,7 @@ public:
 	FDistanceFieldBuildTileListCS()
 	{
 	}
-	void SetParameters(FRHICommandList* RHICmdList, const FSceneView& View, FScene* Scene, FVector2D NumGroupsValue, int32 NumObjects)
+	void SetParameters(FRHICommandList& RHICmdList, const FSceneView& View, FScene* Scene, FVector2D NumGroupsValue, int32 NumObjects)
 	{
 		FComputeShaderRHIParamRef ShaderRHI = GetComputeShader();
 
@@ -687,7 +687,7 @@ public:
 		SetShaderValue(RHICmdList, ShaderRHI, NumGroups, NumGroupsValue);
 	}
 
-	void UnsetParameters(FRHICommandList* RHICmdList)
+	void UnsetParameters(FRHICommandList& RHICmdList)
 	{
 		TileHeadData.UnsetUAV(RHICmdList, GetComputeShader());
 		TileArrayData.UnsetUAV(RHICmdList, GetComputeShader());
@@ -734,7 +734,7 @@ public:
 	}
 
 	template<typename TParamRef>
-	void Set(FRHICommandList* RHICmdList, const TParamRef& ShaderRHI, const FSceneView& View, int32 CurrentLevelDownsampleFactorValue)
+	void Set(FRHICommandList& RHICmdList, const TParamRef& ShaderRHI, const FSceneView& View, int32 CurrentLevelDownsampleFactorValue)
 	{
 		SetShaderValue(RHICmdList, ShaderRHI, CurrentLevelDownsampleFactor, CurrentLevelDownsampleFactorValue);
 
@@ -797,7 +797,7 @@ public:
 		AOParameters.Bind(Initializer.ParameterMap);
 	}
 
-	void SetParameters(FRHICommandList* RHICmdList, const FSceneView& View, int32 NumObjects)
+	void SetParameters(FRHICommandList& RHICmdList, const FSceneView& View, int32 NumObjects)
 	{
 		const FPixelShaderRHIParamRef ShaderRHI = GetPixelShader();
 
@@ -882,7 +882,7 @@ public:
 		AOParameters.Bind(Initializer.ParameterMap);
 	}
 
-	void SetParameters(FRHICommandList* RHICmdList, const FSceneView& View, FSceneRenderTargetItem& DistanceFieldNormalValue, int32 NumObjects)
+	void SetParameters(FRHICommandList& RHICmdList, const FSceneView& View, FSceneRenderTargetItem& DistanceFieldNormalValue, int32 NumObjects)
 	{
 		const FComputeShaderRHIParamRef ShaderRHI = GetComputeShader();
 
@@ -905,7 +905,7 @@ public:
 		SetShaderValue(RHICmdList, ShaderRHI, TileListGroupSize, TileIntersectionResources->TileDimensions);
 	}
 
-	void UnsetParameters(FRHICommandList* RHICmdList)
+	void UnsetParameters(FRHICommandList& RHICmdList)
 	{
 		DistanceFieldNormal.UnsetUAV(RHICmdList, GetComputeShader());
 	}
@@ -944,7 +944,7 @@ IMPLEMENT_SHADER_TYPE(,FComputeDistanceFieldNormalCS,TEXT("DistanceFieldSurfaceC
 void ComputeDistanceFieldNormal(const TArray<FViewInfo>& Views, FSceneRenderTargetItem& DistanceFieldNormal, int32 NumObjects)
 {
 	//@todo-rco: RHIPacketList
-	FRHICommandList* RHICmdList = nullptr;
+	FRHICommandList& RHICmdList = FRHICommandList::GetNullRef();
 
 	if (GAOComputeShaderNormalCalculation)
 	{
@@ -1032,7 +1032,7 @@ public:
 	{
 	}
 
-	void SetParameters(FRHICommandList* RHICmdList, const FSceneView& View, int32 DepthLevel)
+	void SetParameters(FRHICommandList& RHICmdList, const FSceneView& View, int32 DepthLevel)
 	{
 		FComputeShaderRHIParamRef ShaderRHI = GetComputeShader();
 		FGlobalShader::SetParameters(RHICmdList, ShaderRHI, View);
@@ -1047,7 +1047,7 @@ public:
 		SetShaderValue(RHICmdList, ShaderRHI, TrimFraction, GAOTrimOldRecordsFraction);
 	}
 
-	void UnsetParameters(FRHICommandList* RHICmdList)
+	void UnsetParameters(FRHICommandList& RHICmdList)
 	{
 		DispatchParameters.UnsetUAV(RHICmdList, GetComputeShader());
 	}
@@ -1109,7 +1109,7 @@ public:
 	{
 	}
 
-	void SetParameters(FRHICommandList* RHICmdList, const FSceneView& View, int32 DepthLevel)
+	void SetParameters(FRHICommandList& RHICmdList, const FSceneView& View, int32 DepthLevel)
 	{
 		FComputeShaderRHIParamRef ShaderRHI = GetComputeShader();
 		FGlobalShader::SetParameters(RHICmdList, ShaderRHI, View);
@@ -1134,7 +1134,7 @@ public:
 		SetShaderValue(RHICmdList, ShaderRHI, TrimFraction, GAOTrimOldRecordsFraction);
 	}
 
-	void UnsetParameters(FRHICommandList* RHICmdList)
+	void UnsetParameters(FRHICommandList& RHICmdList)
 	{
 		CopyIrradianceCachePositionRadius.UnsetUAV(RHICmdList, GetComputeShader());
 		CopyIrradianceCacheNormal.UnsetUAV(RHICmdList, GetComputeShader());
@@ -1204,7 +1204,7 @@ public:
 	{
 	}
 
-	void SetParameters(FRHICommandList* RHICmdList, const FSceneView& View, int32 DepthLevel)
+	void SetParameters(FRHICommandList& RHICmdList, const FSceneView& View, int32 DepthLevel)
 	{
 		FComputeShaderRHIParamRef ShaderRHI = GetComputeShader();
 		FGlobalShader::SetParameters(RHICmdList, ShaderRHI, View);
@@ -1216,7 +1216,7 @@ public:
 		SavedStartIndex.SetBuffer(RHICmdList, ShaderRHI, SurfaceCacheResources.Level[DepthLevel]->SavedStartIndex);
 	}
 
-	void UnsetParameters(FRHICommandList* RHICmdList)
+	void UnsetParameters(FRHICommandList& RHICmdList)
 	{
 		SavedStartIndex.UnsetUAV(RHICmdList, GetComputeShader());
 	}
@@ -1282,7 +1282,7 @@ public:
 	{
 	}
 
-	void SetParameters(FRHICommandList* RHICmdList, const FSceneView& View, FSceneRenderTargetItem& DistanceFieldAOIrradianceCacheSplat, FSceneRenderTargetItem& DistanceFieldNormal, int32 DownsampleFactorValue, int32 DepthLevel, FIntPoint TileListGroupSizeValue)
+	void SetParameters(FRHICommandList& RHICmdList, const FSceneView& View, FSceneRenderTargetItem& DistanceFieldAOIrradianceCacheSplat, FSceneRenderTargetItem& DistanceFieldNormal, int32 DownsampleFactorValue, int32 DepthLevel, FIntPoint TileListGroupSizeValue)
 	{
 		FComputeShaderRHIParamRef ShaderRHI = GetComputeShader();
 		FGlobalShader::SetParameters(RHICmdList, ShaderRHI, View);
@@ -1336,7 +1336,7 @@ public:
 			);
 	}
 
-	void UnsetParameters(FRHICommandList* RHICmdList)
+	void UnsetParameters(FRHICommandList& RHICmdList)
 	{
 		IrradianceCachePositionRadius.UnsetUAV(RHICmdList, GetComputeShader());
 		IrradianceCacheNormal.UnsetUAV(RHICmdList, GetComputeShader());
@@ -1414,7 +1414,7 @@ public:
 	{
 	}
 
-	void SetParameters(FRHICommandList* RHICmdList, const FSceneView& View, int32 DepthLevel)
+	void SetParameters(FRHICommandList& RHICmdList, const FSceneView& View, int32 DepthLevel)
 	{
 		FComputeShaderRHIParamRef ShaderRHI = GetComputeShader();
 		FGlobalShader::SetParameters(RHICmdList, ShaderRHI, View);
@@ -1428,7 +1428,7 @@ public:
 		DispatchParameters.SetBuffer(RHICmdList, ShaderRHI, SurfaceCacheResources.DispatchParameters);
 	}
 
-	void UnsetParameters(FRHICommandList* RHICmdList)
+	void UnsetParameters(FRHICommandList& RHICmdList)
 	{
 		DispatchParameters.UnsetUAV(RHICmdList, GetComputeShader());
 	}
@@ -1506,7 +1506,7 @@ public:
 	{
 	}
 
-	void SetParameters(FRHICommandList* RHICmdList, const FSceneView& View, FSceneRenderTargetItem& DistanceFieldAOIrradianceCacheSplat, FSceneRenderTargetItem& DistanceFieldNormal, int32 NumObjects, int32 DownsampleFactorValue, int32 DepthLevel, FIntPoint TileListGroupSizeValue)
+	void SetParameters(FRHICommandList& RHICmdList, const FSceneView& View, FSceneRenderTargetItem& DistanceFieldAOIrradianceCacheSplat, FSceneRenderTargetItem& DistanceFieldNormal, int32 NumObjects, int32 DownsampleFactorValue, int32 DepthLevel, FIntPoint TileListGroupSizeValue)
 	{
 		FComputeShaderRHIParamRef ShaderRHI = GetComputeShader();
 		FGlobalShader::SetParameters(RHICmdList, ShaderRHI, View);
@@ -1591,7 +1591,7 @@ public:
 		SetShaderValue(RHICmdList, ShaderRHI, RecordRadiusScale, GAORecordRadiusScale);
 	}
 
-	void UnsetParameters(FRHICommandList* RHICmdList)
+	void UnsetParameters(FRHICommandList& RHICmdList)
 	{
 		OccluderRadius.UnsetUAV(RHICmdList, GetComputeShader());
 		IrradianceCacheBentNormal.UnsetUAV(RHICmdList, GetComputeShader());
@@ -1692,7 +1692,7 @@ public:
 
 	TIrradianceCacheSplatVS() {}
 
-	void SetParameters(FRHICommandList* RHICmdList, const FSceneView& View, int32 DepthLevel, int32 CurrentLevelDownsampleFactorValue, FVector2D NormalizedOffsetToPixelCenterValue)
+	void SetParameters(FRHICommandList& RHICmdList, const FSceneView& View, int32 DepthLevel, int32 CurrentLevelDownsampleFactorValue, FVector2D NormalizedOffsetToPixelCenterValue)
 	{
 		const FVertexShaderRHIParamRef ShaderRHI = GetVertexShader();
 		FGlobalShader::SetParameters(RHICmdList, ShaderRHI, View);
@@ -1774,7 +1774,7 @@ public:
 		InvMinCosPointBehindPlane.Bind(Initializer.ParameterMap, TEXT("InvMinCosPointBehindPlane"));
 	}
 
-	void SetParameters(FRHICommandList* RHICmdList, const FSceneView& View, FSceneRenderTargetItem& DistanceFieldNormal, int32 DestLevelDownsampleFactor)
+	void SetParameters(FRHICommandList& RHICmdList, const FSceneView& View, FSceneRenderTargetItem& DistanceFieldNormal, int32 DestLevelDownsampleFactor)
 	{
 		const FPixelShaderRHIParamRef ShaderRHI = GetPixelShader();
 		FGlobalShader::SetParameters(RHICmdList, ShaderRHI, View);
@@ -1855,7 +1855,7 @@ public:
 		BentNormalAOSampler3.Bind(Initializer.ParameterMap,TEXT("BentNormalAOSampler3"));
 	}
 
-	void SetParameters(FRHICommandList* RHICmdList, const FSceneView& View, FSceneRenderTargetItem& InBentNormalTexture)
+	void SetParameters(FRHICommandList& RHICmdList, const FSceneView& View, FSceneRenderTargetItem& InBentNormalTexture)
 	{
 		const FPixelShaderRHIParamRef ShaderRHI = GetPixelShader();
 
@@ -1916,7 +1916,7 @@ public:
 		BentNormalAOTexelSize.Bind(Initializer.ParameterMap, TEXT("BentNormalAOTexelSize"));
 	}
 
-	void SetParameters(FRHICommandList* RHICmdList, const FSceneView& View, FSceneRenderTargetItem& DistanceFieldAOBentNormal2)
+	void SetParameters(FRHICommandList& RHICmdList, const FSceneView& View, FSceneRenderTargetItem& DistanceFieldAOBentNormal2)
 	{
 		const FPixelShaderRHIParamRef ShaderRHI = GetPixelShader();
 
@@ -1986,7 +1986,7 @@ public:
 		HistoryWeight.Bind(Initializer.ParameterMap, TEXT("HistoryWeight"));
 	}
 
-	void SetParameters(FRHICommandList* RHICmdList, const FSceneView& View, FSceneRenderTargetItem& HistoryTextureValue, FSceneRenderTargetItem& DistanceFieldAOBentNormal)
+	void SetParameters(FRHICommandList& RHICmdList, const FSceneView& View, FSceneRenderTargetItem& HistoryTextureValue, FSceneRenderTargetItem& DistanceFieldAOBentNormal)
 	{
 		const FPixelShaderRHIParamRef ShaderRHI = GetPixelShader();
 
@@ -2065,7 +2065,7 @@ void UpdateHistory(
 	TRefCountPtr<IPooledRenderTarget>& HistoryOutput)
 {
 	//@todo-rco: RHIPacketList
-	FRHICommandList* RHICmdList = nullptr;
+	FRHICommandList& RHICmdList = FRHICommandList::GetNullRef();
 
 	if (HistoryState)
 	{
@@ -2124,7 +2124,7 @@ void UpdateHistory(
 void PostProcessBentNormalAO(TArray<FViewInfo>& Views, FSceneRenderTargetItem& IrradianceCacheInterpolation, TRefCountPtr<IPooledRenderTarget>& AOOutput)
 {
 	//@todo-rco: RHIPacketList
-	FRHICommandList* RHICmdList = nullptr;
+	FRHICommandList& RHICmdList = FRHICommandList::GetNullRef();
 
 	TRefCountPtr<IPooledRenderTarget> DistanceFieldAOBentNormal;
 	AllocateOrReuseAORenderTarget(DistanceFieldAOBentNormal, TEXT("DistanceFieldBentNormalAO"));
@@ -2250,7 +2250,7 @@ public:
 		BentNormalAOSampler.Bind(Initializer.ParameterMap,TEXT("BentNormalAOSampler"));
 	}
 
-	void SetParameters(FRHICommandList* RHICmdList, const FSceneView& View, TRefCountPtr<IPooledRenderTarget>& DistanceFieldAOBentNormal)
+	void SetParameters(FRHICommandList& RHICmdList, const FSceneView& View, TRefCountPtr<IPooledRenderTarget>& DistanceFieldAOBentNormal)
 	{
 		const FPixelShaderRHIParamRef ShaderRHI = GetPixelShader();
 
@@ -2281,7 +2281,7 @@ IMPLEMENT_SHADER_TYPE(,FDistanceFieldAOUpsamplePS,TEXT("DistanceFieldSurfaceCach
 void UpsampleBentNormalAO(const TArray<FViewInfo>& Views, TRefCountPtr<IPooledRenderTarget>& DistanceFieldAOBentNormal)
 {
 	//@todo-rco: RHIPacketList
-	FRHICommandList* RHICmdList = nullptr;
+	FRHICommandList& RHICmdList = FRHICommandList::GetNullRef();
 
 	for (int32 ViewIndex = 0; ViewIndex < Views.Num(); ViewIndex++)
 	{
@@ -2473,7 +2473,7 @@ void GenerateBestSpacedVectors()
 FIntPoint BuildTileObjectLists(FScene* Scene, TArray<FViewInfo>& Views, int32 NumObjects)
 {
 	//@todo-rco: RHIPacketList
-	FRHICommandList* RHICmdList = nullptr;
+	FRHICommandList& RHICmdList = FRHICommandList::GetNullRef();
 
 	SCOPED_DRAW_EVENT(BuildTileList, DEC_SCENE_ITEMS);
 	RHISetRenderTarget(NULL, NULL);
@@ -2605,7 +2605,7 @@ FIntPoint BuildTileObjectLists(FScene* Scene, TArray<FViewInfo>& Views, int32 Nu
 void RenderIrradianceCacheInterpolation(TArray<FViewInfo>& Views, IPooledRenderTarget* InterpolationTarget, FSceneRenderTargetItem& DistanceFieldNormal, int32 DepthLevel, int32 DestLevelDownsampleFactor, bool bFinalInterpolation)
 {
 	//@todo-rco: RHIPacketList
-	FRHICommandList* RHICmdList = nullptr;
+	FRHICommandList& RHICmdList = FRHICommandList::GetNullRef();
 
 	check(!(bFinalInterpolation && DepthLevel != 0));
 
@@ -2820,7 +2820,7 @@ void UpdateVisibleObjectBuffers(const FScene* Scene, const FViewInfo& View, int3
 bool FDeferredShadingSceneRenderer::RenderDistanceFieldAOSurfaceCache(FSceneRenderTargetItem& OutBentNormalAO, bool bApplyToSceneColor)
 {
 	//@todo-rco: RHIPacketList
-	FRHICommandList* RHICmdList = nullptr;
+	FRHICommandList& RHICmdList = FRHICommandList::GetNullRef();
 
 	//@todo - support multiple views
 	const FViewInfo& View = Views[0];
@@ -3069,7 +3069,7 @@ public:
 		BentNormalUVScale.Bind(Initializer.ParameterMap, TEXT("BentNormalUVScale"));
 	}
 
-	void SetParameters(FRHICommandList* RHICmdList, const FSceneView& View, FTextureRHIParamRef DynamicBentNormalAO)
+	void SetParameters(FRHICommandList& RHICmdList, const FSceneView& View, FTextureRHIParamRef DynamicBentNormalAO)
 	{
 		const FPixelShaderRHIParamRef ShaderRHI = GetPixelShader();
 		FGlobalShader::SetParameters(RHICmdList, ShaderRHI, View);
@@ -3105,7 +3105,7 @@ IMPLEMENT_SHADER_TYPE(template<>,TDynamicSkyLightDiffusePS<false>,TEXT("SkyLight
 void FDeferredShadingSceneRenderer::RenderDynamicSkyLighting()
 {
 	//@todo-rco: RHIPacketList
-	FRHICommandList* RHICmdList = nullptr;
+	FRHICommandList& RHICmdList = FRHICommandList::GetNullRef();
 
 	if (Scene->SkyLight
 		&& Scene->SkyLight->ProcessedTexture
@@ -3159,12 +3159,12 @@ void FDeferredShadingSceneRenderer::RenderDynamicSkyLighting()
 			{
 				TShaderMapRef<TDynamicSkyLightDiffusePS<false> > PixelShader( GetGlobalShaderMap() );
 
-				static FGlobalBoundShaderState BoundShaderState;
+			static FGlobalBoundShaderState BoundShaderState;
 				SetGlobalBoundShaderState(BoundShaderState, GFilterVertexDeclaration.VertexDeclarationRHI, *VertexShader, *PixelShader);
 
 				PixelShader->SetParameters(RHICmdList, View, LightAccumulation->GetRenderTargetItem().ShaderResourceTexture);
 			}
-			
+
 			DrawRectangle( 
 				0, 0, 
 				View.ViewRect.Width(), View.ViewRect.Height(),
diff --git a/Engine/Source/Runtime/Renderer/Private/DistortionRendering.cpp b/Engine/Source/Runtime/Renderer/Private/DistortionRendering.cpp
--- a/Engine/Source/Runtime/Renderer/Private/DistortionRendering.cpp
+++ b/Engine/Source/Runtime/Renderer/Private/DistortionRendering.cpp
@@ -33,7 +33,7 @@ public:
 	}
 	FDistortionApplyScreenPS() {}
 
-	void SetParameters(FRHICommandList* RHICmdList, const FRenderingCompositePassContext& Context, IPooledRenderTarget& DistortionRT)
+	void SetParameters(FRHICommandList& RHICmdList, const FRenderingCompositePassContext& Context, IPooledRenderTarget& DistortionRT)
 	{
 		const FPixelShaderRHIParamRef ShaderRHI = GetPixelShader();
 
@@ -125,12 +125,12 @@ protected:
 
 public:
 	
-	void SetParameters(FRHICommandList* RHICmdList, const FVertexFactory* VertexFactory,const FMaterialRenderProxy* MaterialRenderProxy,const FSceneView* View)
+	void SetParameters(FRHICommandList& RHICmdList, const FVertexFactory* VertexFactory,const FMaterialRenderProxy* MaterialRenderProxy,const FSceneView* View)
 	{
 		FMeshMaterialShader::SetParameters(RHICmdList, GetVertexShader(), MaterialRenderProxy, *MaterialRenderProxy->GetMaterial(View->GetFeatureLevel()), *View, ESceneRenderTargetsMode::SetTextures);
 	}
 
-	void SetMesh(FRHICommandList* RHICmdList, const FVertexFactory* VertexFactory,const FSceneView& View,const FPrimitiveSceneProxy* Proxy,const FMeshBatchElement& BatchElement)
+	void SetMesh(FRHICommandList& RHICmdList, const FVertexFactory* VertexFactory,const FSceneView& View,const FPrimitiveSceneProxy* Proxy,const FMeshBatchElement& BatchElement)
 	{
 		FMeshMaterialShader::SetMesh(RHICmdList, GetVertexShader(),VertexFactory,View,Proxy,BatchElement);
 	}
@@ -215,7 +215,7 @@ public:
 	}
 
 	void SetParameters(
-		FRHICommandList* RHICmdList, 
+		FRHICommandList& RHICmdList, 
 		const FMaterialRenderProxy* MaterialRenderProxy,
 		const FSceneView& View
 		)
@@ -233,7 +233,7 @@ public:
 
 	}
 
-	void SetMesh(FRHICommandList* RHICmdList, const FVertexFactory* VertexFactory,const FSceneView& View,const FPrimitiveSceneProxy* Proxy,const FMeshBatchElement& BatchElement)
+	void SetMesh(FRHICommandList& RHICmdList, const FVertexFactory* VertexFactory,const FSceneView& View,const FPrimitiveSceneProxy* Proxy,const FMeshBatchElement& BatchElement)
 	{
 		FMeshMaterialShader::SetMesh(RHICmdList, GetPixelShader(),VertexFactory,View,Proxy,BatchElement);
 	}
@@ -295,7 +295,7 @@ public:
 	* @param CI - The command interface to execute the draw commands on.
 	* @param View - The view of the scene being drawn.
 	*/
-	void DrawShared(FRHICommandList* RHICmdList, const FSceneView* View,FBoundShaderStateRHIParamRef BoundShaderState) const;
+	void DrawShared(FRHICommandList& RHICmdList, const FSceneView* View,FBoundShaderStateRHIParamRef BoundShaderState) const;
 	
 	/** 
 	* Create bound shader state using the vertex decl from the mesh draw policy
@@ -311,7 +311,7 @@ public:
 	* @param ElementData - context specific data for mesh rendering
 	*/
 	void SetMeshRenderState(
-		FRHICommandList* RHICmdList, 
+		FRHICommandList& RHICmdList, 
 		const FSceneView& View,
 		const FPrimitiveSceneProxy* PrimitiveSceneProxy,
 		const FMeshBatch& Mesh,
@@ -404,12 +404,12 @@ bool TDistortionMeshDrawingPolicy<DistortMeshPolicy>::Matches(
 */
 template<class DistortMeshPolicy>
 void TDistortionMeshDrawingPolicy<DistortMeshPolicy>::DrawShared(
-	FRHICommandList* RHICmdList, 
+	FRHICommandList& RHICmdList, 
 	const FSceneView* View,
 	FBoundShaderStateRHIParamRef BoundShaderState
 	) const
 {
-	check(!RHICmdList);
+	RHICmdList.CheckIsNull();
 
 	// Set shared mesh resources
 	FMeshDrawingPolicy::DrawShared(RHICmdList, View);
@@ -493,7 +493,7 @@ FBoundShaderStateRHIRef TDistortionMeshDrawingPolicy<DistortMeshPolicy>::CreateB
 */
 template<class DistortMeshPolicy>
 void TDistortionMeshDrawingPolicy<DistortMeshPolicy>::SetMeshRenderState(
-	FRHICommandList* RHICmdList, 
+	FRHICommandList& RHICmdList, 
 	const FSceneView& View,
 	const FPrimitiveSceneProxy* PrimitiveSceneProxy,
 	const FMeshBatch& Mesh,
@@ -502,7 +502,7 @@ void TDistortionMeshDrawingPolicy<DistortMeshPolicy>::SetMeshRenderState(
 	const ElementDataType& ElementData
 	) const
 {
-	check(!RHICmdList);
+	RHICmdList.CheckIsNull();
 
 	const FMeshBatchElement& BatchElement = Mesh.Elements[BatchElementIndex];
 
@@ -555,7 +555,7 @@ public:
 	* @return true if the mesh rendered
 	*/
 	static bool DrawDynamicMesh(
-		FRHICommandList* RHICmdList, 
+		FRHICommandList& RHICmdList, 
 		const FSceneView& View,
 		ContextType DrawingContext,
 		const FMeshBatch& Mesh,
@@ -570,7 +570,7 @@ public:
 	* @return true if the mesh rendered
 	*/
 	static bool DrawStaticMesh(
-		FRHICommandList* RHICmdList, 
+		FRHICommandList& RHICmdList, 
 		const FSceneView* View,
 		ContextType DrawingContext,
 		const FStaticMesh& StaticMesh,
@@ -589,7 +589,7 @@ public:
 */
 template<class DistortMeshPolicy>
 bool TDistortionMeshDrawingPolicyFactory<DistortMeshPolicy>::DrawDynamicMesh(
-	FRHICommandList* RHICmdList, 
+	FRHICommandList& RHICmdList, 
 	const FSceneView& View,
 	ContextType bInitializeOffsets,
 	const FMeshBatch& Mesh,
@@ -633,7 +633,7 @@ bool TDistortionMeshDrawingPolicyFactory<DistortMeshPolicy>::DrawDynamicMesh(
 */
 template<class DistortMeshPolicy>
 bool TDistortionMeshDrawingPolicyFactory<DistortMeshPolicy>::DrawStaticMesh(
-	FRHICommandList* RHICmdList, 
+	FRHICommandList& RHICmdList, 
 	const FSceneView* View,
 	ContextType bInitializeOffsets,
 	const FStaticMesh& StaticMesh,
@@ -693,7 +693,7 @@ bool FDistortionPrimSet::DrawAccumulatedOffsets(const FViewInfo* ViewInfo,bool b
 	bool bDirty=false;
 
 	//@todo-rco: RHIPacketList
-	FRHICommandList* RHICmdList = nullptr;
+	FRHICommandList& RHICmdList = FRHICommandList::GetNullRef();
 
 	// Draw the view's elements with the translucent drawing policy.
 	bDirty |= DrawViewElements<TDistortionMeshDrawingPolicyFactory<FDistortMeshAccumulatePolicy> >(
@@ -769,7 +769,7 @@ void FDistortionPrimSet::AddScenePrimitive(FPrimitiveSceneProxy* PrimitiveSceneP
 /** 
  * Renders the scene's distortion 
  */
-void FSceneRenderer::RenderDistortion(FRHICommandList* RHICmdList)
+void FSceneRenderer::RenderDistortion(FRHICommandList& RHICmdList)
 {
 	SCOPED_DRAW_EVENT(Distortion, DEC_SCENE_ITEMS);
 
diff --git a/Engine/Source/Runtime/Renderer/Private/DrawingPolicy.cpp b/Engine/Source/Runtime/Renderer/Private/DrawingPolicy.cpp
--- a/Engine/Source/Runtime/Renderer/Private/DrawingPolicy.cpp
+++ b/Engine/Source/Runtime/Renderer/Private/DrawingPolicy.cpp
@@ -42,7 +42,7 @@ FMeshDrawingPolicy::FMeshDrawingPolicy(
 }
 
 void FMeshDrawingPolicy::SetMeshRenderState(
-	FRHICommandList* RHICmdList, 
+	FRHICommandList& RHICmdList, 
 	const FSceneView& View,
 	const FPrimitiveSceneProxy* PrimitiveSceneProxy,
 	const FMeshBatch& Mesh,
@@ -54,13 +54,13 @@ void FMeshDrawingPolicy::SetMeshRenderState(
 	EmitMeshDrawEvents(PrimitiveSceneProxy, Mesh);
 
 		// Use bitwise logic ops to avoid branches
-	FRHICommandList::SetRasterizerState(RHICmdList, GetStaticRasterizerState<true>(
+	RHICmdList.SetRasterizerState( GetStaticRasterizerState<true>(
 			( Mesh.bWireframe | IsWireframe() ) ? FM_Wireframe : FM_Solid, ( ( IsTwoSided() & !NeedsBackfacePass() ) | Mesh.bDisableBackfaceCulling ) ? CM_None :
 			( ( (View.bReverseCulling ^ bBackFace) ^ Mesh.ReverseCulling ) ? CM_CCW : CM_CW )
 			));
 }
 
-void FMeshDrawingPolicy::DrawMesh(FRHICommandList* RHICmdList, const FMeshBatch& Mesh, int32 BatchElementIndex) const
+void FMeshDrawingPolicy::DrawMesh(FRHICommandList& RHICmdList, const FMeshBatch& Mesh, int32 BatchElementIndex) const
 {
 	INC_DWORD_STAT(STAT_MeshDrawCalls);
 	SCOPED_CONDITIONAL_DRAW_EVENTF(MeshEvent, GEmitMeshDrawEvent != 0, DEC_SCENE_ITEMS, TEXT("Mesh Draw"));
@@ -73,7 +73,7 @@ void FMeshDrawingPolicy::DrawMesh(FRHICommandList* RHICmdList, const FMeshBatch&
 
 		if (BatchElement.DynamicIndexData)
 		{
-			check(!RHICmdList);
+			RHICmdList.CheckIsNull();
 
 			RHIDrawIndexedPrimitiveUP(
 				Mesh.Type,
@@ -88,7 +88,7 @@ void FMeshDrawingPolicy::DrawMesh(FRHICommandList* RHICmdList, const FMeshBatch&
 		}
 		else
 		{
-			check(!RHICmdList);
+			RHICmdList.CheckIsNull();
 
 			RHIDrawPrimitiveUP(
 				Mesh.Type,
@@ -103,8 +103,7 @@ void FMeshDrawingPolicy::DrawMesh(FRHICommandList* RHICmdList, const FMeshBatch&
 		if(BatchElement.IndexBuffer)
 		{
 			check(BatchElement.IndexBuffer->IsInitialized());
-			FRHICommandList::DrawIndexedPrimitive(
-					RHICmdList,
+			RHICmdList.DrawIndexedPrimitive(
 					BatchElement.IndexBuffer->IndexBufferRHI,
 					Mesh.Type,
 					0,
@@ -117,8 +116,7 @@ void FMeshDrawingPolicy::DrawMesh(FRHICommandList* RHICmdList, const FMeshBatch&
 		}
 		else
 		{
-			FRHICommandList::DrawPrimitive(
-					RHICmdList,
+			RHICmdList.DrawPrimitive(
 					Mesh.Type,
 					BatchElement.FirstIndex,
 					BatchElement.NumPrimitives,
@@ -128,7 +126,7 @@ void FMeshDrawingPolicy::DrawMesh(FRHICommandList* RHICmdList, const FMeshBatch&
 	}
 }
 
-void FMeshDrawingPolicy::DrawShared(FRHICommandList* RHICmdList, const FSceneView* View) const
+void FMeshDrawingPolicy::DrawShared(FRHICommandList& RHICmdList, const FSceneView* View) const
 {
 	check(VertexFactory && VertexFactory->IsInitialized());
 	VertexFactory->Set(RHICmdList);
diff --git a/Engine/Source/Runtime/Renderer/Private/DrawingPolicy.h b/Engine/Source/Runtime/Renderer/Private/DrawingPolicy.h
--- a/Engine/Source/Runtime/Renderer/Private/DrawingPolicy.h
+++ b/Engine/Source/Runtime/Renderer/Private/DrawingPolicy.h
@@ -67,7 +67,7 @@ public:
 	 * @param PrimitiveSceneProxy - The primitive drawing the dynamic mesh.  If this is a view element, this will be NULL.
 	 */
 	void SetMeshRenderState(
-		FRHICommandList* RHICmdList, 
+		FRHICommandList& RHICmdList, 
 		const FSceneView& View,
 		const FPrimitiveSceneProxy* PrimitiveSceneProxy,
 		const FMeshBatch& Mesh,
@@ -79,14 +79,14 @@ public:
 	/**
 	 * Executes the draw commands for a mesh.
 	 */
-	void DrawMesh(FRHICommandList* RHICmdList, const FMeshBatch& Mesh, int32 BatchElementIndex) const;
+	void DrawMesh(FRHICommandList& RHICmdList, const FMeshBatch& Mesh, int32 BatchElementIndex) const;
 
 	/**
 	 * Executes the draw commands which can be shared between any meshes using this drawer.
 	 * @param CI - The command interface to execute the draw commands on.
 	 * @param View - The view of the scene being drawn.
 	 */
-	void DrawShared(FRHICommandList* RHICmdList, const FSceneView* View) const;
+	void DrawShared(FRHICommandList& RHICmdList, const FSceneView* View) const;
 
 	/**
 	* Get the decl for this mesh policy type and vertexfactory
diff --git a/Engine/Source/Runtime/Renderer/Private/DynamicPrimitiveDrawing.inl b/Engine/Source/Runtime/Renderer/Private/DynamicPrimitiveDrawing.inl
--- a/Engine/Source/Runtime/Renderer/Private/DynamicPrimitiveDrawing.inl
+++ b/Engine/Source/Runtime/Renderer/Private/DynamicPrimitiveDrawing.inl
@@ -105,7 +105,7 @@ int32 TDynamicPrimitiveDrawer<DrawingPolicyFactoryType>::DrawMesh(const FMeshBat
 	Mesh.CheckUniformBuffers();
 #endif
 	//@todo-rco: RHIPacketList
-	FRHICommandList* RHICmdList = nullptr;
+	FRHICommandList& RHICmdList = FRHICommandList::GetNullRef();
 	INC_DWORD_STAT_BY(STAT_DynamicPathMeshDrawCalls,Mesh.Elements.Num());
 	const bool DrawDirty = DrawingPolicyFactoryType::DrawDynamicMesh(
 		RHICmdList,
@@ -221,7 +221,7 @@ bool DrawViewElements(
 	)
 {
 	//@todo-rco: RHIPacketList
-	FRHICommandList* RHICmdList = nullptr;
+	FRHICommandList& RHICmdList = FRHICommandList::GetNullRef();
 
 	// Get the correct element list based on dpg index
 	const TIndirectArray<FHitProxyMeshPair>& ViewMeshElementList = ( DPGIndex == SDPG_Foreground ? View.TopViewMeshElements : View.ViewMeshElements );
diff --git a/Engine/Source/Runtime/Renderer/Private/FogRendering.cpp b/Engine/Source/Runtime/Renderer/Private/FogRendering.cpp
--- a/Engine/Source/Runtime/Renderer/Private/FogRendering.cpp
+++ b/Engine/Source/Runtime/Renderer/Private/FogRendering.cpp
@@ -60,7 +60,7 @@ public:
 		FogStartZ.Bind(Initializer.ParameterMap,TEXT("FogStartZ"));
 	}
 
-	void SetParameters(FRHICommandList* RHICmdList, const FViewInfo& View)
+	void SetParameters(FRHICommandList& RHICmdList, const FViewInfo& View)
 	{
 		FGlobalShader::SetParameters(RHICmdList, GetVertexShader(),View);
 
@@ -124,7 +124,7 @@ public:
 		SceneTextureParameters.Bind(Initializer.ParameterMap);
 	}
 
-	void SetParameters(FRHICommandList* RHICmdList, const FViewInfo& View, FLightShaftsOutput LightShaftsOutput)
+	void SetParameters(FRHICommandList& RHICmdList, const FViewInfo& View, FLightShaftsOutput LightShaftsOutput)
 	{
 		FGlobalShader::SetParameters(RHICmdList, GetPixelShader(), View);
 		SceneTextureParameters.Set(RHICmdList, GetPixelShader(), View);
@@ -263,7 +263,7 @@ void FSceneRenderer::InitFogConstants()
 FGlobalBoundShaderState ExponentialBoundShaderState;
 
 /** Sets the bound shader state for either the per-pixel or per-sample fog pass. */
-void SetFogShaders(FRHICommandList* RHICmdList, FScene* Scene,const FViewInfo& View,FLightShaftsOutput LightShaftsOutput)
+void SetFogShaders(FRHICommandList& RHICmdList, FScene* Scene,const FViewInfo& View,FLightShaftsOutput LightShaftsOutput)
 {
 	if (Scene->ExponentialFogs.Num() > 0)
 	{
@@ -296,7 +296,7 @@ bool FDeferredShadingSceneRenderer::RenderFog(FLightShaftsOutput LightShaftsOutp
 		};
 
 		//@todo-rco: RHIPacketList
-		FRHICommandList* RHICmdList = nullptr;
+		FRHICommandList& RHICmdList = FRHICommandList::GetNullRef();
 
 		GSceneRenderTargets.BeginRenderingSceneColor();
 		for(int32 ViewIndex = 0;ViewIndex < Views.Num();ViewIndex++)
diff --git a/Engine/Source/Runtime/Renderer/Private/ForwardBasePassRendering.cpp b/Engine/Source/Runtime/Renderer/Private/ForwardBasePassRendering.cpp
--- a/Engine/Source/Runtime/Renderer/Private/ForwardBasePassRendering.cpp
+++ b/Engine/Source/Runtime/Renderer/Private/ForwardBasePassRendering.cpp
@@ -52,7 +52,7 @@ public:
 	/** Draws the translucent mesh with a specific light-map type, and fog volume type */
 	template<typename LightMapPolicyType>
 	void Process(
-		FRHICommandList* RHICmdList, 
+		FRHICommandList& RHICmdList, 
 		const FProcessBasePassMeshParameters& Parameters,
 		const LightMapPolicyType& LightMapPolicy,
 		const typename LightMapPolicyType::ElementDataType& LightMapElementData
@@ -144,13 +144,13 @@ public:
 	/** Draws the translucent mesh with a specific light-map type, and shader complexity predicate. */
 	template<typename LightMapPolicyType>
 	void Process(
-		FRHICommandList* RHICmdList, 
+		FRHICommandList& RHICmdList, 
 		const FProcessBasePassMeshParameters& Parameters,
 		const LightMapPolicyType& LightMapPolicy,
 		const typename LightMapPolicyType::ElementDataType& LightMapElementData
 		) const
 	{
-		check(!RHICmdList);
+		RHICmdList.CheckIsNull();
 
 #if !(UE_BUILD_SHIPPING || UE_BUILD_TEST)
 		// Treat masked materials as if they don't occlude in shader complexity, which is PVR behavior
@@ -204,7 +204,7 @@ public:
 };
 
 bool FBasePassForwardOpaqueDrawingPolicyFactory::DrawDynamicMesh(
-	FRHICommandList* RHICmdList, 
+	FRHICommandList& RHICmdList, 
 	const FSceneView& View,
 	ContextType DrawingContext,
 	const FMeshBatch& Mesh,
@@ -292,7 +292,7 @@ EBasePassSort::Type GetSortMode()
 	}
 }
 
-void FForwardShadingSceneRenderer::RenderForwardShadingBasePass(FRHICommandList* RHICmdList)
+void FForwardShadingSceneRenderer::RenderForwardShadingBasePass(FRHICommandList& RHICmdList)
 {
 	SCOPED_DRAW_EVENT(BasePass, DEC_SCENE_ITEMS);
 	SCOPE_CYCLE_COUNTER(STAT_BasePassDrawTime);
@@ -331,41 +331,39 @@ void FForwardShadingSceneRenderer::RenderForwardShadingBasePass(FRHICommandList*
 
 		{
 			SCOPE_CYCLE_COUNTER(STAT_RHICounterTEMP);
-			auto* PrevRHICmdList = RHICmdList;
-			if (!RHICmdList)
+			FRHICommandList* UseRHICmdList = &RHICmdList;
+			if (!&RHICmdList)
 			{
 				//@todo-rco: Very temp code!!!
 				static IConsoleVariable* RHICmdListCVar = IConsoleManager::Get().FindConsoleVariable(TEXT("r.RHICmd"));
 				bool bUseRHICmdList = RHICmdListCVar->GetInt() != 0;
-				RHICmdList = bUseRHICmdList ? GRHICommandList.CreateList() : nullptr;
+				UseRHICmdList = bUseRHICmdList ? &GRHICommandList.CreateList() : &FRHICommandList::GetNullRef();
 			}
 
 			// Render the base pass static data
 			if (SortMode == EBasePassSort::SortPerMesh)
 			{
 				SCOPE_CYCLE_COUNTER(STAT_StaticDrawListDrawTime);
-				MaxDraws -= Scene->BasePassForForwardShadingLowQualityLightMapDrawList[FScene::EBasePass_Default].DrawVisibleFrontToBack(RHICmdList, View,View.StaticMeshVisibilityMap,View.StaticMeshBatchVisibility,MaxDraws);
-				MaxDraws -= Scene->BasePassForForwardShadingDistanceFieldShadowMapLightMapDrawList[FScene::EBasePass_Default].DrawVisibleFrontToBack(RHICmdList, View,View.StaticMeshVisibilityMap,View.StaticMeshBatchVisibility,MaxDraws);
-				MaxDraws -= Scene->BasePassForForwardShadingDirectionalLightAndSHIndirectDrawList[FScene::EBasePass_Default].DrawVisibleFrontToBack(RHICmdList, View,View.StaticMeshVisibilityMap,View.StaticMeshBatchVisibility,MaxDraws);
-				MaxDraws -= Scene->BasePassForForwardShadingNoLightMapDrawList[FScene::EBasePass_Default].DrawVisibleFrontToBack(RHICmdList, View,View.StaticMeshVisibilityMap,View.StaticMeshBatchVisibility,MaxDraws);
+				MaxDraws -= Scene->BasePassForForwardShadingLowQualityLightMapDrawList[FScene::EBasePass_Default].DrawVisibleFrontToBack(*UseRHICmdList, View, View.StaticMeshVisibilityMap, View.StaticMeshBatchVisibility, MaxDraws);
+				MaxDraws -= Scene->BasePassForForwardShadingDistanceFieldShadowMapLightMapDrawList[FScene::EBasePass_Default].DrawVisibleFrontToBack(*UseRHICmdList, View, View.StaticMeshVisibilityMap, View.StaticMeshBatchVisibility, MaxDraws);
+				MaxDraws -= Scene->BasePassForForwardShadingDirectionalLightAndSHIndirectDrawList[FScene::EBasePass_Default].DrawVisibleFrontToBack(*UseRHICmdList, View, View.StaticMeshVisibilityMap, View.StaticMeshBatchVisibility, MaxDraws);
+				MaxDraws -= Scene->BasePassForForwardShadingNoLightMapDrawList[FScene::EBasePass_Default].DrawVisibleFrontToBack(*UseRHICmdList, View, View.StaticMeshVisibilityMap, View.StaticMeshBatchVisibility, MaxDraws);
 			}
 			else
 			{
 				SCOPE_CYCLE_COUNTER(STAT_StaticDrawListDrawTime);
-				Scene->BasePassForForwardShadingLowQualityLightMapDrawList[FScene::EBasePass_Default].DrawVisible(RHICmdList, View,View.StaticMeshVisibilityMap,View.StaticMeshBatchVisibility);
-				Scene->BasePassForForwardShadingDistanceFieldShadowMapLightMapDrawList[FScene::EBasePass_Default].DrawVisible(RHICmdList, View,View.StaticMeshVisibilityMap,View.StaticMeshBatchVisibility);
-				Scene->BasePassForForwardShadingDirectionalLightAndSHIndirectDrawList[FScene::EBasePass_Default].DrawVisible(RHICmdList, View,View.StaticMeshVisibilityMap,View.StaticMeshBatchVisibility);
-				Scene->BasePassForForwardShadingNoLightMapDrawList[FScene::EBasePass_Default].DrawVisible(RHICmdList, View,View.StaticMeshVisibilityMap,View.StaticMeshBatchVisibility);
+				Scene->BasePassForForwardShadingLowQualityLightMapDrawList[FScene::EBasePass_Default].DrawVisible(*UseRHICmdList, View, View.StaticMeshVisibilityMap, View.StaticMeshBatchVisibility);
+				Scene->BasePassForForwardShadingDistanceFieldShadowMapLightMapDrawList[FScene::EBasePass_Default].DrawVisible(*UseRHICmdList, View, View.StaticMeshVisibilityMap, View.StaticMeshBatchVisibility);
+				Scene->BasePassForForwardShadingDirectionalLightAndSHIndirectDrawList[FScene::EBasePass_Default].DrawVisible(*UseRHICmdList, View, View.StaticMeshVisibilityMap, View.StaticMeshBatchVisibility);
+				Scene->BasePassForForwardShadingNoLightMapDrawList[FScene::EBasePass_Default].DrawVisible(*UseRHICmdList, View, View.StaticMeshVisibilityMap, View.StaticMeshBatchVisibility);
 			}
 
-			if (PrevRHICmdList != RHICmdList)
+			if (UseRHICmdList != &RHICmdList)
 			{
-				if (RHICmdList)
+				if (UseRHICmdList)
 				{
-					GRHICommandList.ExecuteAndFreeList(RHICmdList);
+					GRHICommandList.ExecuteAndFreeList(*UseRHICmdList);
 				}
-
-				RHICmdList = PrevRHICmdList;
 			}
 		}
 
