UE4: Rearrange RHICmdList toward non-null cmd lists and the final "look"

diff --git a/Engine/Source/Runtime/Renderer/Private/ForwardBasePassRendering.h b/Engine/Source/Runtime/Renderer/Private/ForwardBasePassRendering.h
--- a/Engine/Source/Runtime/Renderer/Private/ForwardBasePassRendering.h
+++ b/Engine/Source/Runtime/Renderer/Private/ForwardBasePassRendering.h
@@ -56,7 +56,7 @@ public:
 	}
 
 	void SetParameters(
-		FRHICommandList* RHICmdList, 
+		FRHICommandList& RHICmdList, 
 		const FMaterialRenderProxy* MaterialRenderProxy,
 		const FVertexFactory* VertexFactory,
 		const FMaterial& InMaterialResource,
@@ -68,7 +68,7 @@ public:
 		FMeshMaterialShader::SetParameters(RHICmdList, GetVertexShader(),MaterialRenderProxy,InMaterialResource,View,TextureMode);
 	}
 
-	void SetMesh(FRHICommandList* RHICmdList, const FVertexFactory* VertexFactory,const FSceneView& View,const FPrimitiveSceneProxy* Proxy,const FMeshBatchElement& BatchElement)
+	void SetMesh(FRHICommandList& RHICmdList, const FVertexFactory* VertexFactory,const FSceneView& View,const FPrimitiveSceneProxy* Proxy,const FMeshBatchElement& BatchElement)
 	{
 		FMeshMaterialShader::SetMesh(RHICmdList, GetVertexShader(),VertexFactory,View,Proxy,BatchElement);
 	}
@@ -130,12 +130,12 @@ public:
 	}
 	TBasePassForForwardShadingPSBaseType() {}
 
-	void SetParameters(FRHICommandList* RHICmdList, const FMaterialRenderProxy* MaterialRenderProxy,const FMaterial& MaterialResource,const FSceneView* View,ESceneRenderTargetsMode::Type TextureMode)
+	void SetParameters(FRHICommandList& RHICmdList, const FMaterialRenderProxy* MaterialRenderProxy,const FMaterial& MaterialResource,const FSceneView* View,ESceneRenderTargetsMode::Type TextureMode)
 	{
 		FMeshMaterialShader::SetParameters(RHICmdList, GetPixelShader(),MaterialRenderProxy,MaterialResource,*View,TextureMode);
 	}
 
-	void SetMesh(FRHICommandList* RHICmdList, const FVertexFactory* VertexFactory,const FSceneView& View,const FPrimitiveSceneProxy* Proxy,const FMeshBatchElement& BatchElement)
+	void SetMesh(FRHICommandList& RHICmdList, const FVertexFactory* VertexFactory,const FSceneView& View,const FPrimitiveSceneProxy* Proxy,const FMeshBatchElement& BatchElement)
 	{
 		if (ReflectionCubemap.IsBound())
 		{
@@ -264,10 +264,10 @@ public:
 			SceneTextureMode == Other.SceneTextureMode;
 	}
 
-	void DrawShared(FRHICommandList* RHICmdList, const FSceneView* View,FBoundShaderStateRHIParamRef BoundShaderState) const
+	void DrawShared(FRHICommandList& RHICmdList, const FSceneView* View,FBoundShaderStateRHIParamRef BoundShaderState) const
 	{
 		// Set the actual shader & vertex declaration state
-		FRHICommandList::SetBoundShaderState(RHICmdList, BoundShaderState);
+		RHICmdList.SetBoundShaderState(BoundShaderState);
 
 		VertexShader->SetParameters(RHICmdList, MaterialRenderProxy, VertexFactory, *MaterialResource, *View, SceneTextureMode);
 
@@ -276,12 +276,12 @@ public:
 		{
 			if (BlendMode == BLEND_Opaque)
 			{
-				FRHICommandList::SetBlendState(RHICmdList, TStaticBlendStateWriteMask<CW_RGBA>::GetRHI());
+				RHICmdList.SetBlendState( TStaticBlendStateWriteMask<CW_RGBA>::GetRHI());
 			}
 			else
 			{
 				// Add complexity to existing
-				FRHICommandList::SetBlendState(RHICmdList, TStaticBlendState<CW_RGBA,BO_Add,BF_One,BF_One,BO_Add,BF_Zero,BF_One>::GetRHI());
+				RHICmdList.SetBlendState( TStaticBlendState<CW_RGBA,BO_Add,BF_One,BF_One,BO_Add,BF_Zero,BF_One>::GetRHI());
 			}
 
 			TShaderMapRef<FShaderComplexityAccumulatePS> ShaderComplexityPixelShader(GetGlobalShaderMap());
@@ -304,15 +304,15 @@ public:
 				// Masked materials are rendered together in the base pass, where the blend state is set at a higher level
 				break;
 			case BLEND_Translucent:
-				FRHICommandList::SetBlendState(RHICmdList, TStaticBlendState<CW_RGB, BO_Add,BF_SourceAlpha,BF_InverseSourceAlpha,BO_Add,BF_Zero,BF_InverseSourceAlpha>::GetRHI());
+				RHICmdList.SetBlendState( TStaticBlendState<CW_RGB, BO_Add,BF_SourceAlpha,BF_InverseSourceAlpha,BO_Add,BF_Zero,BF_InverseSourceAlpha>::GetRHI());
 				break;
 			case BLEND_Additive:
 				// Add to the existing scene color
-				FRHICommandList::SetBlendState(RHICmdList, TStaticBlendState<CW_RGB, BO_Add,BF_One,BF_One,BO_Add,BF_Zero,BF_InverseSourceAlpha>::GetRHI());
+				RHICmdList.SetBlendState( TStaticBlendState<CW_RGB, BO_Add,BF_One,BF_One,BO_Add,BF_Zero,BF_InverseSourceAlpha>::GetRHI());
 				break;
 			case BLEND_Modulate:
 				// Modulate with the existing scene color
-				FRHICommandList::SetBlendState(RHICmdList, TStaticBlendState<CW_RGB,BO_Add,BF_DestColor,BF_Zero>::GetRHI());
+				RHICmdList.SetBlendState( TStaticBlendState<CW_RGB,BO_Add,BF_DestColor,BF_Zero>::GetRHI());
 				break;
 			};
 		}
@@ -350,7 +350,7 @@ public:
 	}
 
 	void SetMeshRenderState(
-		FRHICommandList* RHICmdList, 
+		FRHICommandList& RHICmdList, 
 		const FSceneView& View,
 		const FPrimitiveSceneProxy* PrimitiveSceneProxy,
 		const FMeshBatch& Mesh,
@@ -438,7 +438,7 @@ public:
 
 	static void AddStaticMesh(FScene* Scene,FStaticMesh* StaticMesh);
 	static bool DrawDynamicMesh(
-		FRHICommandList* RHICmdList, 
+		FRHICommandList& RHICmdList, 
 		const FSceneView& View,
 		ContextType DrawingContext,
 		const FMeshBatch& Mesh,
@@ -472,7 +472,7 @@ void ProcessBasePassMeshForForwardShading(
 	check(!AllowHighQualityLightmaps());
 
 	//@todo-rco: RHIPacketList
-	FRHICommandList* RHICmdList = nullptr;
+	FRHICommandList& RHICmdList = FRHICommandList::GetNullRef();
 
 	if (LightMapInteraction.GetType() == LMIT_Texture)
 	{
diff --git a/Engine/Source/Runtime/Renderer/Private/ForwardShadingRenderer.cpp b/Engine/Source/Runtime/Renderer/Private/ForwardShadingRenderer.cpp
--- a/Engine/Source/Runtime/Renderer/Private/ForwardShadingRenderer.cpp
+++ b/Engine/Source/Runtime/Renderer/Private/ForwardShadingRenderer.cpp
@@ -84,7 +84,7 @@ void FForwardShadingSceneRenderer::Render()
 	RHIClear(true, FLinearColor::Black, true, 0.0f, true, 0, FIntRect());
 
 	//@todo-rco: RHIPacketList
-	FRHICommandList* RHICmdList = nullptr;
+	FRHICommandList& RHICmdList = FRHICommandList::GetNullRef();
 	RenderForwardShadingBasePass(RHICmdList);
 
 	// Notify the FX system that opaque primitives have been rendered.
diff --git a/Engine/Source/Runtime/Renderer/Private/ForwardTranslucentRendering.cpp b/Engine/Source/Runtime/Renderer/Private/ForwardTranslucentRendering.cpp
--- a/Engine/Source/Runtime/Renderer/Private/ForwardTranslucentRendering.cpp
+++ b/Engine/Source/Runtime/Renderer/Private/ForwardTranslucentRendering.cpp
@@ -25,7 +25,7 @@ public:
 	}
 	FForwardCopySceneAlphaPS() {}
 
-	void SetParameters(FRHICommandList* RHICmdList, const FSceneView& View)
+	void SetParameters(FRHICommandList& RHICmdList, const FSceneView& View)
 	{
 		SceneTextureParameters.Set(RHICmdList, GetPixelShader(), View);
 	}
@@ -45,7 +45,7 @@ IMPLEMENT_SHADER_TYPE(,FForwardCopySceneAlphaPS,TEXT("TranslucentLightingShaders
 
 FGlobalBoundShaderState ForwardCopySceneAlphaBoundShaderState;
 
-void FForwardShadingSceneRenderer::CopySceneAlpha(FRHICommandList* RHICmdList, const FSceneView& View)
+void FForwardShadingSceneRenderer::CopySceneAlpha(FRHICommandList& RHICmdList, const FSceneView& View)
 {
 	SCOPED_DRAW_EVENTF(EventCopy, DEC_SCENE_ITEMS, TEXT("CopySceneAlpha"));
 	RHISetRasterizerState(TStaticRasterizerState<FM_Solid,CM_None>::GetRHI());
@@ -120,7 +120,7 @@ public:
 	/** Draws the translucent mesh with a specific light-map type, and fog volume type */
 	template<typename LightMapPolicyType>
 	void Process(
-		FRHICommandList* RHICmdList, 
+		FRHICommandList& RHICmdList, 
 		const FProcessBasePassMeshParameters& Parameters,
 		const LightMapPolicyType& LightMapPolicy,
 		const typename LightMapPolicyType::ElementDataType& LightMapElementData
@@ -165,7 +165,7 @@ public:
  * @return true if the mesh rendered
  */
 bool FTranslucencyForwardShadingDrawingPolicyFactory::DrawDynamicMesh(
-	FRHICommandList* RHICmdList, 
+	FRHICommandList& RHICmdList, 
 	const FViewInfo& View,
 	ContextType DrawingContext,
 	const FMeshBatch& Mesh,
@@ -212,7 +212,7 @@ bool FTranslucencyForwardShadingDrawingPolicyFactory::DrawDynamicMesh(
  * @return true if the mesh rendered
  */
 bool FTranslucencyForwardShadingDrawingPolicyFactory::DrawStaticMesh(
-	FRHICommandList* RHICmdList, 
+	FRHICommandList& RHICmdList, 
 	const FViewInfo& View,
 	ContextType DrawingContext,
 	const FStaticMesh& StaticMesh,
@@ -280,7 +280,7 @@ void FTranslucentPrimSet::RenderPrimitiveForForwardShading(
 		}
 
 		//@todo-rco: RHIPacketList
-		FRHICommandList* RHICmdList = nullptr;
+		FRHICommandList& RHICmdList = FRHICommandList::GetNullRef();
 
 		// Render static scene prim
 		if( ViewRelevance.bStaticRelevance )
@@ -313,7 +313,7 @@ void FForwardShadingSceneRenderer::RenderTranslucency()
 	if (ShouldRenderTranslucency())
 	{
 		//@todo-rco: RHIPacketList
-		FRHICommandList* RHICmdList = nullptr;
+		FRHICommandList& RHICmdList = FRHICommandList::GetNullRef();
 
 		const bool bGammaSpace = !IsMobileHDR();
 		const bool bLinearHDR64 = !bGammaSpace && !IsMobileHDR32bpp();
diff --git a/Engine/Source/Runtime/Renderer/Private/GPUBenchmark.cpp b/Engine/Source/Runtime/Renderer/Private/GPUBenchmark.cpp
--- a/Engine/Source/Runtime/Renderer/Private/GPUBenchmark.cpp
+++ b/Engine/Source/Runtime/Renderer/Private/GPUBenchmark.cpp
@@ -57,7 +57,7 @@ public:
 		return bShaderHasOutdatedParameters;
 	}
 
-	void SetParameters(FRHICommandList* RHICmdList, const FSceneView& View, TRefCountPtr<IPooledRenderTarget>& Src)
+	void SetParameters(FRHICommandList& RHICmdList, const FSceneView& View, TRefCountPtr<IPooledRenderTarget>& Src)
 	{
 		const FPixelShaderRHIParamRef ShaderRHI = GetPixelShader();
 
@@ -113,7 +113,7 @@ public:
 		return bShaderHasOutdatedParameters;
 	}
 
-	void SetParameters(FRHICommandList* RHICmdList, const FSceneView& View)
+	void SetParameters(FRHICommandList& RHICmdList, const FSceneView& View)
 	{
 		const FVertexShaderRHIParamRef ShaderRHI = GetVertexShader();
 		
@@ -125,7 +125,7 @@ IMPLEMENT_SHADER_TYPE(,FPostProcessBenchmarkVS,TEXT("GPUBenchmark"),TEXT("MainBe
 
 
 template <uint32 Method>
-void RunBenchmarkShader(FRHICommandList* RHICmdList, const FSceneView& View, TRefCountPtr<IPooledRenderTarget>& Src, uint32 Count)
+void RunBenchmarkShader(FRHICommandList& RHICmdList, const FSceneView& View, TRefCountPtr<IPooledRenderTarget>& Src, uint32 Count)
 {
 	TShaderMapRef<FPostProcessBenchmarkVS> VertexShader(GetGlobalShaderMap());
 	TShaderMapRef<FPostProcessBenchmarkPS<Method> > PixelShader(GetGlobalShaderMap());
@@ -151,7 +151,7 @@ void RunBenchmarkShader(FRHICommandList* RHICmdList, const FSceneView& View, TRe
 	}
 }
 
-void RunBenchmarkShader(FRHICommandList* RHICmdList, const FSceneView& View, uint32 MethodId, TRefCountPtr<IPooledRenderTarget>& Src, uint32 Count)
+void RunBenchmarkShader(FRHICommandList& RHICmdList, const FSceneView& View, uint32 MethodId, TRefCountPtr<IPooledRenderTarget>& Src, uint32 Count)
 {
 	SCOPED_DRAW_EVENTF(Benchmark, DEC_SCENE_ITEMS, TEXT("Benchmark Method:%d"), MethodId);
 
@@ -347,7 +347,7 @@ private:
 void RendererGPUBenchmark(FSynthBenchmarkResults& InOut, const FSceneView& View, uint32 WorkScale, bool bDebugOut)
 {
 	//@todo-rco: RHIPacketList
-	FRHICommandList* RHICmdList = nullptr;
+	FRHICommandList& RHICmdList = FRHICommandList::GetNullRef();
 
 	check(IsInRenderingThread());
 	
diff --git a/Engine/Source/Runtime/Renderer/Private/LightFunctionRendering.cpp b/Engine/Source/Runtime/Renderer/Private/LightFunctionRendering.cpp
--- a/Engine/Source/Runtime/Renderer/Private/LightFunctionRendering.cpp
+++ b/Engine/Source/Runtime/Renderer/Private/LightFunctionRendering.cpp
@@ -33,7 +33,7 @@ public:
 		StencilingGeometryParameters.Bind(Initializer.ParameterMap);
 	}
 
-	void SetParameters(FRHICommandList* RHICmdList, const FSceneView* View, const FLightSceneInfo* LightSceneInfo )
+	void SetParameters(FRHICommandList& RHICmdList, const FSceneView* View, const FLightSceneInfo* LightSceneInfo )
 	{
 		FMaterialShader::SetParameters(RHICmdList, GetVertexShader(), *View);
 		
@@ -93,7 +93,7 @@ public:
 		DeferredParameters.Bind(Initializer.ParameterMap);
 	}
 
-	void SetParameters(FRHICommandList* RHICmdList, const FSceneView* View, const FLightSceneInfo* LightSceneInfo, const FMaterialRenderProxy* MaterialProxy, bool bRenderingPreviewShadowIndicator, float ShadowFadeFraction )
+	void SetParameters(FRHICommandList& RHICmdList, const FSceneView* View, const FLightSceneInfo* LightSceneInfo, const FMaterialRenderProxy* MaterialProxy, bool bRenderingPreviewShadowIndicator, float ShadowFadeFraction )
 	{
 		const FPixelShaderRHIParamRef ShaderRHI = GetPixelShader();
 
@@ -212,7 +212,7 @@ bool FDeferredShadingSceneRenderer::CheckForLightFunction( const FLightSceneInfo
  *
  * @param LightSceneInfo Represents the current light
  */
-bool FDeferredShadingSceneRenderer::RenderLightFunction(FRHICommandList* RHICmdList, const FLightSceneInfo* LightSceneInfo, bool bLightAttenuationCleared)
+bool FDeferredShadingSceneRenderer::RenderLightFunction(FRHICommandList& RHICmdList, const FLightSceneInfo* LightSceneInfo, bool bLightAttenuationCleared)
 {
 	if (ViewFamily.EngineShowFlags.LightFunctions)
 	{
@@ -222,7 +222,7 @@ bool FDeferredShadingSceneRenderer::RenderLightFunction(FRHICommandList* RHICmdL
 	return false;
 }
 
-bool FDeferredShadingSceneRenderer::RenderPreviewShadowsIndicator(FRHICommandList* RHICmdList, const FLightSceneInfo* LightSceneInfo, bool bLightAttenuationCleared)
+bool FDeferredShadingSceneRenderer::RenderPreviewShadowsIndicator(FRHICommandList& RHICmdList, const FLightSceneInfo* LightSceneInfo, bool bLightAttenuationCleared)
 {
 	if (GEngine->PreviewShadowsIndicatorMaterial)
 	{
@@ -232,7 +232,7 @@ bool FDeferredShadingSceneRenderer::RenderPreviewShadowsIndicator(FRHICommandLis
 	return false;
 }
 
-bool FDeferredShadingSceneRenderer::RenderLightFunctionForMaterial(FRHICommandList* RHICmdList, const FLightSceneInfo* LightSceneInfo, const FMaterialRenderProxy* MaterialProxy, bool bLightAttenuationCleared, bool bRenderingPreviewShadowsIndicator)
+bool FDeferredShadingSceneRenderer::RenderLightFunctionForMaterial(FRHICommandList& RHICmdList, const FLightSceneInfo* LightSceneInfo, const FMaterialRenderProxy* MaterialProxy, bool bLightAttenuationCleared, bool bRenderingPreviewShadowsIndicator)
 {
 	bool bRenderedLightFunction = false;
 
@@ -248,7 +248,7 @@ bool FDeferredShadingSceneRenderer::RenderLightFunctionForMaterial(FRHICommandLi
 		const FMaterialShaderMap* MaterialShaderMap = Material->GetRenderingThreadShaderMap();
 		FLightFunctionVS* VertexShader = MaterialShaderMap->GetShader<FLightFunctionVS>();
 		FLightFunctionPS* PixelShader = MaterialShaderMap->GetShader<FLightFunctionPS>();
-		check(!RHICmdList);
+		RHICmdList.CheckIsNull();
 
 		// This was cached but when changing the material (e.g. editor) it wasn't updated.
 		// This will change with upcoming multi threaded rendering changes.
diff --git a/Engine/Source/Runtime/Renderer/Private/LightMapDensityRendering.cpp b/Engine/Source/Runtime/Renderer/Private/LightMapDensityRendering.cpp
--- a/Engine/Source/Runtime/Renderer/Private/LightMapDensityRendering.cpp
+++ b/Engine/Source/Runtime/Renderer/Private/LightMapDensityRendering.cpp
@@ -92,7 +92,7 @@ bool FDeferredShadingSceneRenderer::RenderLightMapDensities()
 }
 
 bool FLightMapDensityDrawingPolicyFactory::DrawDynamicMesh(
-	FRHICommandList* RHICmdList, 
+	FRHICommandList& RHICmdList, 
 	const FViewInfo& View,
 	ContextType DrawingContext,
 	const FMeshBatch& Mesh,
diff --git a/Engine/Source/Runtime/Renderer/Private/LightMapDensityRendering.h b/Engine/Source/Runtime/Renderer/Private/LightMapDensityRendering.h
--- a/Engine/Source/Runtime/Renderer/Private/LightMapDensityRendering.h
+++ b/Engine/Source/Runtime/Renderer/Private/LightMapDensityRendering.h
@@ -45,7 +45,7 @@ public:
 	}
 
 	void SetParameters(		
-		FRHICommandList* RHICmdList, 
+		FRHICommandList& RHICmdList, 
 		const FMaterialRenderProxy* MaterialRenderProxy,
 		const FSceneView& View
 		)
@@ -53,7 +53,7 @@ public:
 		FMeshMaterialShader::SetParameters(RHICmdList, GetVertexShader(), MaterialRenderProxy, *MaterialRenderProxy->GetMaterial(View.GetFeatureLevel()), View, ESceneRenderTargetsMode::SetTextures);
 	}
 
-	void SetMesh(FRHICommandList* RHICmdList, const FVertexFactory* VertexFactory,const FSceneView& View,const FPrimitiveSceneProxy* Proxy,const FMeshBatchElement& BatchElement)
+	void SetMesh(FRHICommandList& RHICmdList, const FVertexFactory* VertexFactory,const FSceneView& View,const FPrimitiveSceneProxy* Proxy,const FMeshBatchElement& BatchElement)
 	{
 		FMeshMaterialShader::SetMesh(RHICmdList, GetVertexShader(),VertexFactory,View,Proxy,BatchElement);
 	}
@@ -159,7 +159,7 @@ public:
 	}
 	TLightMapDensityPS() {}
 
-	void SetParameters(FRHICommandList* RHICmdList, const FMaterialRenderProxy* MaterialRenderProxy,const FSceneView* View)
+	void SetParameters(FRHICommandList& RHICmdList, const FMaterialRenderProxy* MaterialRenderProxy,const FSceneView* View)
 	{
 		FMeshMaterialShader::SetParameters(RHICmdList, GetPixelShader(), MaterialRenderProxy, *MaterialRenderProxy->GetMaterial(View->GetFeatureLevel()), *View, ESceneRenderTargetsMode::SetTextures);
 
@@ -175,7 +175,7 @@ public:
 		}
 	}
 
-	void SetMesh(FRHICommandList* RHICmdList, const FVertexFactory* VertexFactory,const FPrimitiveSceneProxy* PrimitiveSceneProxy,const FMeshBatchElement& BatchElement,const FSceneView& View,bool bBackFace,FVector& InBuiltLightingAndSelectedFlags,FVector2D& InLightMapResolutionScale, bool bTextureMapped)
+	void SetMesh(FRHICommandList& RHICmdList, const FVertexFactory* VertexFactory,const FPrimitiveSceneProxy* PrimitiveSceneProxy,const FMeshBatchElement& BatchElement,const FSceneView& View,bool bBackFace,FVector& InBuiltLightingAndSelectedFlags,FVector2D& InLightMapResolutionScale, bool bTextureMapped)
 	{
 		FMeshMaterialShader::SetMesh(RHICmdList, GetPixelShader(),VertexFactory,View,PrimitiveSceneProxy,BatchElement);
 
@@ -296,9 +296,9 @@ public:
 			LightMapPolicy == Other.LightMapPolicy;
 	}
 
-	void DrawShared(FRHICommandList* RHICmdList, const FSceneView* View, FBoundShaderStateRHIRef ShaderState ) const
+	void DrawShared(FRHICommandList& RHICmdList, const FSceneView* View, FBoundShaderStateRHIRef ShaderState ) const
 	{
-		check(!RHICmdList);
+		RHICmdList.CheckIsNull();
 
 		// Set the actual shader & vertex declaration state
 		RHISetBoundShaderState(ShaderState);
@@ -320,7 +320,7 @@ public:
 	}
 
 	void SetMeshRenderState(
-		FRHICommandList* RHICmdList, 
+		FRHICommandList& RHICmdList, 
 		const FSceneView& View,
 		const FPrimitiveSceneProxy* PrimitiveSceneProxy,
 		const FMeshBatch& Mesh,
@@ -447,7 +447,7 @@ public:
 	struct ContextType {};
 
 	static bool DrawDynamicMesh(
-		FRHICommandList* RHICmdList, 
+		FRHICommandList& RHICmdList, 
 		const FViewInfo& View,
 		ContextType DrawingContext,
 		const FMeshBatch& Mesh,
diff --git a/Engine/Source/Runtime/Renderer/Private/LightMapRendering.h b/Engine/Source/Runtime/Renderer/Private/LightMapRendering.h
--- a/Engine/Source/Runtime/Renderer/Private/LightMapRendering.h
+++ b/Engine/Source/Runtime/Renderer/Private/LightMapRendering.h
@@ -42,7 +42,7 @@ public:
 	{}
 
 	void Set(
-		FRHICommandList* RHICmdList, 
+		FRHICommandList& RHICmdList, 
 		const VertexParametersType* VertexShaderParameters,
 		const PixelParametersType* PixelShaderParameters,
 		FShader* VertexShader,
@@ -57,7 +57,7 @@ public:
 	}
 
 	void SetMesh(
-		FRHICommandList* RHICmdList,
+		FRHICommandList& RHICmdList,
 		const FSceneView& View,
 		const FPrimitiveSceneProxy* PrimitiveSceneProxy,
 		const VertexParametersType* VertexShaderParameters,
@@ -105,7 +105,7 @@ public:
 			LightMapCoordinateScaleBiasParameter.Bind(ParameterMap,TEXT("LightMapCoordinateScaleBias"));
 		}
 
-		void SetLightMapScale(FRHICommandList* RHICmdList, FShader* VertexShader,const FLightMapInteraction& LightMapInteraction) const
+		void SetLightMapScale(FRHICommandList& RHICmdList, FShader* VertexShader,const FLightMapInteraction& LightMapInteraction) const
 		{
 			const FVector2D LightmapCoordinateScale = LightMapInteraction.GetCoordinateScale();
 			const FVector2D LightmapCoordinateBias = LightMapInteraction.GetCoordinateBias();
@@ -139,7 +139,7 @@ public:
 			LightMapAddParameter.Bind(ParameterMap,TEXT("LightMapAdd"));
 		}
 
-		void SetLightMapTexture(FRHICommandList* RHICmdList, FShader* PixelShader, const UTexture2D* LightMapTexture) const
+		void SetLightMapTexture(FRHICommandList& RHICmdList, FShader* PixelShader, const UTexture2D* LightMapTexture) const
 		{
 			FTexture* TextureResource = GBlackTexture;
 
@@ -157,7 +157,7 @@ public:
 				);
 		}
 
-		void SetSkyOcclusionTexture(FRHICommandList* RHICmdList, FShader* PixelShader, const UTexture2D* SkyOcclusionTextureValue) const
+		void SetSkyOcclusionTexture(FRHICommandList& RHICmdList, FShader* PixelShader, const UTexture2D* SkyOcclusionTextureValue) const
 		{
 			FTexture* TextureResource = GWhiteTexture;
 
@@ -175,7 +175,7 @@ public:
 				);
 		}
 
-		void SetLightMapScale(FRHICommandList* RHICmdList, FShader* PixelShader,const FLightMapInteraction& LightMapInteraction) const
+		void SetLightMapScale(FRHICommandList& RHICmdList, FShader* PixelShader,const FLightMapInteraction& LightMapInteraction) const
 		{
 			SetShaderValueArray(RHICmdList, PixelShader->GetPixelShader(),LightMapScaleParameter,LightMapInteraction.GetScaleArray(),LightMapInteraction.GetNumLightmapCoefficients());
 			SetShaderValueArray(RHICmdList, PixelShader->GetPixelShader(),LightMapAddParameter,LightMapInteraction.GetAddArray(),LightMapInteraction.GetNumLightmapCoefficients());
@@ -232,7 +232,7 @@ public:
 	}
 
 	void Set(
-		FRHICommandList* RHICmdList, 
+		FRHICommandList& RHICmdList, 
 		const VertexParametersType* VertexShaderParameters,
 		const PixelParametersType* PixelShaderParameters,
 		FShader* VertexShader,
@@ -247,7 +247,7 @@ public:
 	}
 
 	void SetMesh(
-		FRHICommandList* RHICmdList, 
+		FRHICommandList& RHICmdList, 
 		const FSceneView& View,
 		const FPrimitiveSceneProxy* PrimitiveSceneProxy,
 		const VertexParametersType* VertexShaderParameters,
@@ -324,7 +324,7 @@ public:
 			Super::VertexParametersType::Bind(ParameterMap);
 		}
 
-		void SetMesh(FRHICommandList* RHICmdList, FShader* VertexShader, const FShadowMapInteraction& ShadowMapInteraction) const
+		void SetMesh(FRHICommandList& RHICmdList, FShader* VertexShader, const FShadowMapInteraction& ShadowMapInteraction) const
 		{
 			SetShaderValue(RHICmdList, VertexShader->GetVertexShader(), ShadowMapCoordinateScaleBias, FVector4(ShadowMapInteraction.GetCoordinateScale(), ShadowMapInteraction.GetCoordinateBias()));
 		}
@@ -353,7 +353,7 @@ public:
 			Super::PixelParametersType::Bind(ParameterMap);
 		}
 
-		void SetMesh(FRHICommandList* RHICmdList, FShader* PixelShader, const FShadowMapInteraction& ShadowMapInteraction, const FVector4& DistanceFieldValues) const
+		void SetMesh(FRHICommandList& RHICmdList, FShader* PixelShader, const FShadowMapInteraction& ShadowMapInteraction, const FVector4& DistanceFieldValues) const
 		{
 			SetShaderValue(RHICmdList, PixelShader->GetPixelShader(), DistanceFieldParameters, DistanceFieldValues);
 			
@@ -399,7 +399,7 @@ public:
 	TDistanceFieldShadowsAndLightMapPolicy() {}
 
 	void SetMesh(
-		FRHICommandList* RHICmdList, 
+		FRHICommandList& RHICmdList, 
 		const FSceneView& View,
 		const FPrimitiveSceneProxy* PrimitiveSceneProxy,
 		const VertexParametersType* VertexShaderParameters,
@@ -438,7 +438,7 @@ public:
 	}
 
 	void Set(
-		FRHICommandList* RHICmdList, 
+		FRHICommandList& RHICmdList, 
 		const VertexParametersType* VertexShaderParameters,
 		const PixelParametersType* PixelShaderParameters,
 		FShader* VertexShader,
@@ -453,7 +453,7 @@ public:
 	}
 
 	void SetMesh(
-		FRHICommandList* RHICmdList, 
+		FRHICommandList& RHICmdList, 
 		const FSceneView& View,
 		const FPrimitiveSceneProxy* PrimitiveSceneProxy,
 		const VertexParametersType* VertexShaderParameters,
@@ -551,7 +551,7 @@ public:
 	FSelfShadowedTranslucencyPolicy() {}
 
 	void SetMesh(
-		FRHICommandList* RHICmdList,
+		FRHICommandList& RHICmdList,
 		const FSceneView& View,
 		const FPrimitiveSceneProxy* PrimitiveSceneProxy,
 		const VertexParametersType* VertexShaderParameters,
@@ -665,7 +665,7 @@ public:
 	FCachedVolumeIndirectLightingPolicy() {}
 
 	void Set(
-		FRHICommandList* RHICmdList, 
+		FRHICommandList& RHICmdList, 
 		const VertexParametersType* VertexShaderParameters,
 		const PixelParametersType* PixelShaderParameters,
 		FShader* VertexShader,
@@ -676,7 +676,7 @@ public:
 		) const;
 
 	void SetMesh(
-		FRHICommandList* RHICmdList, 
+		FRHICommandList& RHICmdList, 
 		const FSceneView& View,
 		const FPrimitiveSceneProxy* PrimitiveSceneProxy,
 		const VertexParametersType* VertexShaderParameters,
@@ -749,7 +749,7 @@ public:
 	FCachedPointIndirectLightingPolicy() {}
 
 	void SetMesh(
-		FRHICommandList* RHICmdList, 
+		FRHICommandList& RHICmdList, 
 		const FSceneView& View,
 		const FPrimitiveSceneProxy* PrimitiveSceneProxy,
 		const VertexParametersType* VertexShaderParameters,
@@ -811,7 +811,7 @@ public:
 	FSelfShadowedCachedPointIndirectLightingPolicy() {}
 
 	void SetMesh(
-		FRHICommandList* RHICmdList, 
+		FRHICommandList& RHICmdList, 
 		const FSceneView& View,
 		const FPrimitiveSceneProxy* PrimitiveSceneProxy,
 		const VertexParametersType* VertexShaderParameters,
@@ -911,7 +911,7 @@ public:
 	}
 
 	void Set(
-		FRHICommandList* RHICmdList, 
+		FRHICommandList& RHICmdList, 
 		const VertexParametersType* VertexShaderParameters,
 		const PixelParametersType* PixelShaderParameters,
 		FShader* VertexShader,
@@ -926,7 +926,7 @@ public:
 	}
 
 	void SetMesh(
-		FRHICommandList* RHICmdList, 
+		FRHICommandList& RHICmdList, 
 		const FSceneView& View,
 		const FPrimitiveSceneProxy* PrimitiveSceneProxy,
 		const VertexParametersType* VertexShaderParameters,
diff --git a/Engine/Source/Runtime/Renderer/Private/LightPropagationVolume.cpp b/Engine/Source/Runtime/Renderer/Private/LightPropagationVolume.cpp
--- a/Engine/Source/Runtime/Renderer/Private/LightPropagationVolume.cpp
+++ b/Engine/Source/Runtime/Renderer/Private/LightPropagationVolume.cpp
@@ -177,11 +177,11 @@ public:
 		return bShaderHasOutdatedParameters;
 	}
 
-	void SetParameters(FRHICommandList* RHICmdList, const FLpvBaseWriteShaderParams& Params )
+	void SetParameters(FRHICommandList& RHICmdList, const FLpvBaseWriteShaderParams& Params )
 	{
 		FComputeShaderRHIParamRef ShaderRHI = GetComputeShader();
 		SetUniformBufferParameter(RHICmdList, ShaderRHI, GetUniformBufferParameter<FLpvWriteUniformBufferParameters>(), Params.UniformBuffer );
-		check(!RHICmdList);
+		RHICmdList.CheckIsNull();
 #if LPV_VOLUME_TEXTURE
 		for(int i  =0; i < 7; i++)
 		{
@@ -265,9 +265,9 @@ public:
 	}
 
 	// Unbinds any buffers that have been bound.
-	void UnbindBuffers(FRHICommandList* RHICmdList)
+	void UnbindBuffers(FRHICommandList& RHICmdList)
 	{
-		check(!RHICmdList);
+		RHICmdList.CheckIsNull();
 		FComputeShaderRHIParamRef ShaderRHI = GetComputeShader();
 #if LPV_VOLUME_TEXTURE
 		for ( int i = 0; i < 7; i++ )
@@ -480,7 +480,7 @@ public:
 	}
 
 	void SetParameters(
-		FRHICommandList* RHICmdList, 
+		FRHICommandList& RHICmdList, 
 		FLpvBaseWriteShaderParams& BaseParams,
 		FTextureRHIParamRef RsmDiffuseTextureRHI,
 		FTextureRHIParamRef RsmNormalTextureRHI,
@@ -499,7 +499,7 @@ public:
 	}
 
 	// Unbinds any buffers that have been bound.
-	void UnbindBuffers(FRHICommandList* RHICmdList)
+	void UnbindBuffers(FRHICommandList& RHICmdList)
 	{
 		FLpvWriteShaderCSBase::UnbindBuffers(RHICmdList);
 	}
@@ -658,7 +658,7 @@ class FLpvInjectShader_Base : public FLpvWriteShaderCSBase
 {
 public:
 	void SetParameters(
-		FRHICommandList* RHICmdList, 
+		FRHICommandList& RHICmdList, 
 		FLpvBaseWriteShaderParams& BaseParams,
 		FDirectLightInjectBufferRef& InjectUniformBuffer )
 	{
@@ -938,7 +938,7 @@ void FLightPropagationVolume::InitSettings( const FSceneView& View )
 /**
 * Clears the LPV
 */
-void FLightPropagationVolume::Clear(FRHICommandList* RHICmdList)
+void FLightPropagationVolume::Clear(FRHICommandList& RHICmdList)
 {
 	if ( !bEnabled )
 	{
@@ -994,7 +994,7 @@ void FLightPropagationVolume::Clear(FRHICommandList* RHICmdList)
 	//RHIBatchComputeShaderEnd();
 
 	// Clear the list UAV counters
-	check(!RHICmdList);
+	RHICmdList.CheckIsNull();
 
 	RHISetUAVParameter( FComputeShaderRHIRef(), 7, mVplListBuffer->UAV, 0 );
 	RHISetUAVParameter( FComputeShaderRHIRef(), 7, GvListBuffer->UAV, 0 );
@@ -1051,7 +1051,7 @@ void FLightPropagationVolume::SetVplInjectionConstants(
 * Injects a Directional light into the LPV
 */
 void FLightPropagationVolume::InjectDirectionalLightRSM(
-	FRHICommandList* RHICmdList, 
+	FRHICommandList& RHICmdList, 
 	const FTexture2DRHIRef&		RsmDiffuseTex, 
 	const FTexture2DRHIRef&		RsmNormalTex, 
 	const FTexture2DRHIRef&		RsmDepthTex, 
@@ -1109,7 +1109,7 @@ void FLightPropagationVolume::InjectDirectionalLightRSM(
 /**
 * Propagates light in the LPV 
 */
-void FLightPropagationVolume::Propagate(FRHICommandList* RHICmdList)
+void FLightPropagationVolume::Propagate(FRHICommandList& RHICmdList)
 {
 	if ( !bEnabled )
 	{
@@ -1209,7 +1209,7 @@ void FLightPropagationVolume::GetShaderParams( FLpvBaseWriteShaderParams& OutPar
 	OutParams.UniformBuffer = LpvWriteUniformBuffer;
 }
 
-void FLightPropagationVolume::InjectLightDirect(FRHICommandList* RHICmdList, const FLightSceneProxy& Light )
+void FLightPropagationVolume::InjectLightDirect(FRHICommandList& RHICmdList, const FLightSceneProxy& Light )
 {
 	if(!bEnabled)
 	{
diff --git a/Engine/Source/Runtime/Renderer/Private/LightPropagationVolume.h b/Engine/Source/Runtime/Renderer/Private/LightPropagationVolume.h
--- a/Engine/Source/Runtime/Renderer/Private/LightPropagationVolume.h
+++ b/Engine/Source/Runtime/Renderer/Private/LightPropagationVolume.h
@@ -126,23 +126,23 @@ public:
 
 	void InitSettings( const FSceneView& View );
 
-	void Clear(FRHICommandList* RHICmdList);
+	void Clear(FRHICommandList& RHICmdList);
 
 	void SetVplInjectionConstants(
 		const FProjectedShadowInfo&	ProjectedShadowInfo,
 		const FLightSceneProxy* LightProxy );
 
 	void InjectDirectionalLightRSM(
-		FRHICommandList* RHICmdList, 
+		FRHICommandList& RHICmdList, 
 		const FTexture2DRHIRef&		RsmDiffuseTex, 
 		const FTexture2DRHIRef&		RsmNormalTex, 
 		const FTexture2DRHIRef&		RsmDepthTex, 
 		const FProjectedShadowInfo&	ProjectedShadowInfo,
 		const FLinearColor&			LightColour );
 
-	void InjectLightDirect(FRHICommandList* RHICmdList, const FLightSceneProxy& Light );
+	void InjectLightDirect(FRHICommandList& RHICmdList, const FLightSceneProxy& Light );
 
-	void Propagate(FRHICommandList* RHICmdList);
+	void Propagate(FRHICommandList& RHICmdList);
 
 	void Visualise( const FSceneView& View ) const;
 
diff --git a/Engine/Source/Runtime/Renderer/Private/LightPropagationVolumeVisualisation.cpp b/Engine/Source/Runtime/Renderer/Private/LightPropagationVolumeVisualisation.cpp
--- a/Engine/Source/Runtime/Renderer/Private/LightPropagationVolumeVisualisation.cpp
+++ b/Engine/Source/Runtime/Renderer/Private/LightPropagationVolumeVisualisation.cpp
@@ -56,7 +56,7 @@ public:
 	static void ModifyCompilationEnvironment( EShaderPlatform Platform, FShaderCompilerEnvironment& OutEnvironment )				{ FLpvVisualiseBase::ModifyCompilationEnvironment( Platform, OutEnvironment ); }
 
 	void SetParameters(
-		FRHICommandList* RHICmdList, 
+		FRHICommandList& RHICmdList, 
 		const FSceneView& View )
 	{
 		FGeometryShaderRHIParamRef ShaderRHI = GetGeometryShader();
@@ -78,7 +78,7 @@ public:
 	static void ModifyCompilationEnvironment( EShaderPlatform Platform, FShaderCompilerEnvironment& OutEnvironment )				{ FLpvVisualiseBase::ModifyCompilationEnvironment( Platform, OutEnvironment ); }
 
 	void SetParameters(
-		FRHICommandList* RHICmdList, 
+		FRHICommandList& RHICmdList, 
 		const FSceneView& View )
 	{
 		FVertexShaderRHIParamRef ShaderRHI = GetVertexShader();
@@ -122,13 +122,13 @@ public:
 	static void ModifyCompilationEnvironment( EShaderPlatform Platform, FShaderCompilerEnvironment& OutEnvironment )				{ FLpvVisualiseBase::ModifyCompilationEnvironment( Platform, OutEnvironment ); }
 
 	void SetParameters(
-		FRHICommandList* RHICmdList, 
+		FRHICommandList& RHICmdList, 
 		const FLightPropagationVolume* LPV,
 		const FSceneView& View )
 	{
 		FPixelShaderRHIParamRef ShaderRHI = GetPixelShader();
 		FGlobalShader::SetParameters(RHICmdList, ShaderRHI, View); 
-		check(!RHICmdList);
+		RHICmdList.CheckIsNull();
 
 #if LPV_VOLUME_TEXTURE
 		for ( int i=0; i<7; i++ )
@@ -262,11 +262,14 @@ void FLightPropagationVolume::Visualise( const FSceneView& View ) const
 	RHISetBlendState(TStaticBlendState<CW_RGB,BO_Add,BF_One,BF_One>::GetRHI());
 
 	//@todo-rco: RHIPacketList
+
+	FRHICommandList& RHICmdList = FRHICommandList::GetNullRef();
+
 	SetGlobalBoundShaderState(LpvVisBoundShaderState, GSimpleElementVertexDeclaration.VertexDeclarationRHI, *VertexShader, *PixelShader, *GeometryShader);
 
-	VertexShader->SetParameters(nullptr, View );
-	GeometryShader->SetParameters(nullptr, View );
-	PixelShader->SetParameters(nullptr, this, View );
+	VertexShader->SetParameters(RHICmdList, View);
+	GeometryShader->SetParameters(RHICmdList, View);
+	PixelShader->SetParameters(RHICmdList, this, View);
 
 	RHISetStreamSource(0, NULL, 0, 0);
 	RHIDrawPrimitive( PT_PointList, 0, 1, 32*3 );
diff --git a/Engine/Source/Runtime/Renderer/Private/LightRendering.cpp b/Engine/Source/Runtime/Renderer/Private/LightRendering.cpp
--- a/Engine/Source/Runtime/Renderer/Private/LightRendering.cpp
+++ b/Engine/Source/Runtime/Renderer/Private/LightRendering.cpp
@@ -59,14 +59,14 @@ public:
 	{
 	}
 
-	void SetParameters(FRHICommandList* RHICmdList, const FSceneView& View, const FLightSceneInfo* LightSceneInfo)
+	void SetParameters(FRHICommandList& RHICmdList, const FSceneView& View, const FLightSceneInfo* LightSceneInfo)
 	{
 		const FPixelShaderRHIParamRef ShaderRHI = GetPixelShader();
 		SetParametersBase(RHICmdList, ShaderRHI, View, LightSceneInfo->Proxy->GetIESTextureResource());
 		SetDeferredLightParameters(RHICmdList, ShaderRHI, GetUniformBufferParameter<FDeferredLightUniformStruct>(), LightSceneInfo, View);
 	}
 
-	void SetParametersSimpleLight(FRHICommandList* RHICmdList, const FSceneView& View, const FSimpleLightEntry& SimpleLight, const FSimpleLightPerViewEntry& SimpleLightPerViewData)
+	void SetParametersSimpleLight(FRHICommandList& RHICmdList, const FSceneView& View, const FSimpleLightEntry& SimpleLight, const FSimpleLightPerViewEntry& SimpleLightPerViewData)
 	{
 		const FPixelShaderRHIParamRef ShaderRHI = GetPixelShader();
 		SetParametersBase(RHICmdList, ShaderRHI, View, NULL);
@@ -95,7 +95,7 @@ public:
 
 private:
 
-	void SetParametersBase(FRHICommandList* RHICmdList, const FPixelShaderRHIParamRef ShaderRHI, const FSceneView& View, FTexture* IESTextureResource)
+	void SetParametersBase(FRHICommandList& RHICmdList, const FPixelShaderRHIParamRef ShaderRHI, const FSceneView& View, FTexture* IESTextureResource)
 	{
 		FGlobalShader::SetParameters(RHICmdList, ShaderRHI,View);
 		DeferredParameters.Set(RHICmdList, ShaderRHI, View);
@@ -190,7 +190,7 @@ public:
 	{
 	}
 
-	void SetParameters(FRHICommandList* RHICmdList, const FSceneView& View, const FLightSceneInfo* LightSceneInfo)
+	void SetParameters(FRHICommandList& RHICmdList, const FSceneView& View, const FLightSceneInfo* LightSceneInfo)
 	{
 		const FPixelShaderRHIParamRef ShaderRHI = GetPixelShader();
 		FGlobalShader::SetParameters(RHICmdList, ShaderRHI,View);
@@ -278,9 +278,9 @@ void GetLightNameForDrawEvent(const FLightSceneProxy* LightProxy, FString& Light
 uint32 GetShadowQuality();
 
 /** Renders the scene's lighting. */
-void FDeferredShadingSceneRenderer::RenderLights(FRHICommandList* RHICmdList)
+void FDeferredShadingSceneRenderer::RenderLights(FRHICommandList& RHICmdList)
 {
-	check(!RHICmdList);
+	RHICmdList.CheckIsNull();
 
 	SCOPED_DRAW_EVENT(Lights, DEC_SCENE_ITEMS);
 
@@ -569,7 +569,7 @@ void FDeferredShadingSceneRenderer::RenderLights(FRHICommandList* RHICmdList)
 							if ( Lpv && LightSceneInfo->Proxy )
 							{
 								//@todo-rco: RHIPacketList
-								Lpv->InjectLightDirect(nullptr, *LightSceneInfo->Proxy );
+								Lpv->InjectLightDirect(FRHICommandList::GetNullRef(), *LightSceneInfo->Proxy);
 							}
 						}
 					}					
@@ -659,7 +659,7 @@ static FVertexDeclarationRHIParamRef GetDeferredLightingVertexDeclaration()
 
 template<bool bUseIESProfile, bool bRadialAttenuation, bool bInverseSquaredFalloff>
 static void SetShaderTemplLighting(
-	FRHICommandList* RHICmdList, 
+	FRHICommandList& RHICmdList, 
 	const FSceneView& View, 
 	FShader* VertexShader,
 	const FLightSceneInfo* LightSceneInfo)
@@ -680,7 +680,7 @@ static void SetShaderTemplLighting(
 
 template<bool bUseIESProfile, bool bRadialAttenuation, bool bInverseSquaredFalloff>
 static void SetShaderTemplLightingSimple(
-	FRHICommandList* RHICmdList, 
+	FRHICommandList& RHICmdList, 
 	const FSceneView& View, 
 	FShader* VertexShader,
 	const FSimpleLightEntry& SimpleLight,
@@ -731,7 +731,7 @@ void FDeferredShadingSceneRenderer::RenderLight(const FLightSceneInfo* LightScen
 		}
 
 		//@todo-rco: RHIPacketList
-		FRHICommandList* RHICmdList = nullptr;
+		FRHICommandList& RHICmdList = FRHICommandList::GetNullRef();
 
 		// Set the device viewport for the view.
 		RHISetViewport(View.ViewRect.Min.X, View.ViewRect.Min.Y, 0.0f, View.ViewRect.Max.X, View.ViewRect.Max.Y, 1.0f);
@@ -841,7 +841,7 @@ void FDeferredShadingSceneRenderer::RenderSimpleLightsStandardDeferred(const FSi
 	SCOPED_DRAW_EVENT(StandardDeferredSimpleLights, DEC_SCENE_ITEMS);
 	
 	//@todo-rco: RHIPacketList
-	FRHICommandList* RHICmdList = nullptr;
+	FRHICommandList& RHICmdList = FRHICommandList::GetNullRef();
 
 	// Use additive blending for color
 	RHISetBlendState(TStaticBlendState<CW_RGBA, BO_Add, BF_One, BF_One, BO_Add, BF_One, BF_One>::GetRHI());
diff --git a/Engine/Source/Runtime/Renderer/Private/LightRendering.h b/Engine/Source/Runtime/Renderer/Private/LightRendering.h
--- a/Engine/Source/Runtime/Renderer/Private/LightRendering.h
+++ b/Engine/Source/Runtime/Renderer/Private/LightRendering.h
@@ -28,7 +28,7 @@ extern float GMinScreenRadiusForLights;
 
 template<typename ShaderRHIParamRef>
 void SetDeferredLightParameters(
-	FRHICommandList* RHICmdList, 
+	FRHICommandList& RHICmdList, 
 	const ShaderRHIParamRef ShaderRHI, 
 	const TShaderUniformBufferParameter<FDeferredLightUniformStruct>& DeferredLightUniformBufferParameter, 
 	const FLightSceneInfo* LightSceneInfo,
@@ -96,7 +96,7 @@ void SetDeferredLightParameters(
 
 template<typename ShaderRHIParamRef>
 void SetSimpleDeferredLightParameters(
-	FRHICommandList* RHICmdList, 
+	FRHICommandList& RHICmdList, 
 	const ShaderRHIParamRef ShaderRHI, 
 	const TShaderUniformBufferParameter<FDeferredLightUniformStruct>& DeferredLightUniformBufferParameter, 
 	const FSimpleLightEntry& SimpleLight,
@@ -130,7 +130,7 @@ public:
 	}
 
 	template<typename ShaderRHIParamRef>
-	void Set(FRHICommandList* RHICmdList, const ShaderRHIParamRef ShaderRHI, const FLightSceneInfo* LightSceneInfo, float ShadowFadeFraction) const
+	void Set(FRHICommandList& RHICmdList, const ShaderRHIParamRef ShaderRHI, const FLightSceneInfo* LightSceneInfo, float ShadowFadeFraction) const
 	{
 		const bool bIsSpotLight = LightSceneInfo->Proxy->GetLightType() == LightType_Spot;
 		const bool bIsPointLight = LightSceneInfo->Proxy->GetLightType() == LightType_Point;
@@ -174,13 +174,13 @@ public:
 		StencilingGeometryParameters.Bind(Initializer.ParameterMap);
 	}
 
-	void SetParameters(FRHICommandList* RHICmdList, const FViewInfo& View, const FLightSceneInfo* LightSceneInfo)
+	void SetParameters(FRHICommandList& RHICmdList, const FViewInfo& View, const FLightSceneInfo* LightSceneInfo)
 	{
 		FGlobalShader::SetParameters(RHICmdList, GetVertexShader(),View);
 		StencilingGeometryParameters.Set(RHICmdList, this, View, LightSceneInfo);
 	}
 
-	void SetSimpleLightParameters(FRHICommandList* RHICmdList, const FViewInfo& View, const FSphere& LightBounds)
+	void SetSimpleLightParameters(FRHICommandList& RHICmdList, const FViewInfo& View, const FSphere& LightBounds)
 	{
 		FGlobalShader::SetParameters(RHICmdList, GetVertexShader(),View);
 
diff --git a/Engine/Source/Runtime/Renderer/Private/LightShaftRendering.cpp b/Engine/Source/Runtime/Renderer/Private/LightShaftRendering.cpp
--- a/Engine/Source/Runtime/Renderer/Private/LightShaftRendering.cpp
+++ b/Engine/Source/Runtime/Renderer/Private/LightShaftRendering.cpp
@@ -105,7 +105,7 @@ public:
 	}
 
 	template<typename ShaderRHIParamRef>
-	void SetParameters(FRHICommandList* RHICmdList, const ShaderRHIParamRef Shader, const FLightSceneInfo* LightSceneInfo, const FSceneView& View, TRefCountPtr<IPooledRenderTarget>& PassSource)
+	void SetParameters(FRHICommandList& RHICmdList, const ShaderRHIParamRef Shader, const FLightSceneInfo* LightSceneInfo, const FSceneView& View, TRefCountPtr<IPooledRenderTarget>& PassSource)
 	{
 		const uint32 DownsampleFactor = GetLightShaftDownsampleFactor();
 		FIntPoint DownSampledViewSize(FMath::FloorToInt(View.ViewRect.Width() / DownsampleFactor), FMath::FloorToInt(View.ViewRect.Height() / DownsampleFactor));
@@ -238,7 +238,7 @@ public:
 	}
 
 	/** Sets shader parameter values */
-	void SetParameters(FRHICommandList* RHICmdList, const FViewInfo& View)
+	void SetParameters(FRHICommandList& RHICmdList, const FViewInfo& View)
 	{
 		FGlobalShader::SetParameters(RHICmdList, GetVertexShader(), View);
 	}
@@ -292,7 +292,7 @@ public:
 	}
 
 	/** Sets shader parameter values */
-	void SetParameters(FRHICommandList* RHICmdList, const FLightSceneInfo* LightSceneInfo, const FViewInfo& View, TRefCountPtr<IPooledRenderTarget>& PassSource)
+	void SetParameters(FRHICommandList& RHICmdList, const FLightSceneInfo* LightSceneInfo, const FViewInfo& View, TRefCountPtr<IPooledRenderTarget>& PassSource)
 	{
 		FGlobalShader::SetParameters(RHICmdList, GetPixelShader(), View);
 		LightShaftParameters.SetParameters(RHICmdList, GetPixelShader(), LightSceneInfo, View, PassSource);
@@ -361,7 +361,7 @@ public:
 	}
 
 	/** Sets shader parameter values */
-	void SetParameters(FRHICommandList* RHICmdList, const FLightSceneInfo* LightSceneInfo, const FViewInfo& View, int32 PassIndex, TRefCountPtr<IPooledRenderTarget>& PassSource)
+	void SetParameters(FRHICommandList& RHICmdList, const FLightSceneInfo* LightSceneInfo, const FViewInfo& View, int32 PassIndex, TRefCountPtr<IPooledRenderTarget>& PassSource)
 	{
 		FGlobalShader::SetParameters(RHICmdList, GetPixelShader(), View);
 		LightShaftParameters.SetParameters(RHICmdList, GetPixelShader(), LightSceneInfo, View, PassSource);
@@ -411,7 +411,7 @@ public:
 	}
 
 	/** Sets shader parameter values */
-	void SetParameters(FRHICommandList* RHICmdList, const FLightSceneInfo* LightSceneInfo, const FViewInfo& View, TRefCountPtr<IPooledRenderTarget>& PassSource)
+	void SetParameters(FRHICommandList& RHICmdList, const FLightSceneInfo* LightSceneInfo, const FViewInfo& View, TRefCountPtr<IPooledRenderTarget>& PassSource)
 	{
 		FGlobalShader::SetParameters(RHICmdList, GetPixelShader(), View);
 		LightShaftParameters.SetParameters(RHICmdList, GetPixelShader(), LightSceneInfo, View, PassSource);
@@ -460,7 +460,7 @@ void DownsamplePass(const FViewInfo& View, const FLightSceneInfo* LightSceneInfo
 	const uint32 DownsampledSizeY = View.ViewRect.Height() / DownsampleFactor;
 
 	//@todo-rco: RHIPacketList
-	FRHICommandList* RHICmdList = nullptr;
+	FRHICommandList& RHICmdList = FRHICommandList::GetNullRef();
 
 	RHISetRenderTarget(LightShaftsDest->GetRenderTargetItem().TargetableTexture, FTextureRHIRef());
 	RHISetViewport(DownSampledXY.X, DownSampledXY.Y, 0.0f, DownSampledXY.X + DownsampledSizeX, DownSampledXY.Y + DownsampledSizeY, 1.0f);
@@ -620,7 +620,7 @@ void ApplyRadialBlurPasses(
 
 		TRefCountPtr<IPooledRenderTarget>& EffectiveSource = PassIndex == 0 ? FirstPassSource : LightShaftsSource;
 		/// ?
-		BlurLightShaftsPixelShader->SetParameters(nullptr, LightSceneInfo, View, PassIndex, EffectiveSource);
+		BlurLightShaftsPixelShader->SetParameters(FRHICommandList::GetNullRef(), LightSceneInfo, View, PassIndex, EffectiveSource);
 
 		{
 			SCOPED_DRAW_EVENT(RadialBlur, DEC_SCENE_ITEMS);
@@ -663,7 +663,7 @@ void FinishOcclusionTerm(const FViewInfo& View, const FLightSceneInfo* const Lig
 	TShaderMapRef<FFinishOcclusionPixelShader> MaskOcclusionTermPixelShader(GetGlobalShaderMap());
 	SetGlobalBoundShaderState(AccumulateTermBoundShaderState, GFilterVertexDeclaration.VertexDeclarationRHI, *ScreenVertexShader, *MaskOcclusionTermPixelShader);
 	/// ?
-	MaskOcclusionTermPixelShader->SetParameters(nullptr, LightSceneInfo, View, LightShaftsSource);
+	MaskOcclusionTermPixelShader->SetParameters(FRHICommandList::GetNullRef(), LightSceneInfo, View, LightShaftsSource);
 
 	{
 		SCOPED_DRAW_EVENT(FinishOcclusion, DEC_SCENE_ITEMS);
@@ -811,7 +811,7 @@ public:
 	}
 
 	/** Sets shader parameter values */
-	void SetParameters(FRHICommandList* RHICmdList, const FViewInfo& View, TRefCountPtr<IPooledRenderTarget>& LightShaftOcclusion)
+	void SetParameters(FRHICommandList& RHICmdList, const FViewInfo& View, TRefCountPtr<IPooledRenderTarget>& LightShaftOcclusion)
 	{
 		FGlobalShader::SetParameters(RHICmdList, GetPixelShader(), View);
 
@@ -848,7 +848,7 @@ void ApplyLightShaftBloom(const FViewInfo& View, const FLightSceneInfo* const Li
 	SetGlobalBoundShaderState(ApplyLightShaftsBoundShaderState, GFilterVertexDeclaration.VertexDeclarationRHI, *ScreenVertexShader, *ApplyLightShaftsPixelShader);
 
 	/// ?
-	ApplyLightShaftsPixelShader->SetParameters(nullptr, View, LightShaftsSource);
+	ApplyLightShaftsPixelShader->SetParameters(FRHICommandList::GetNullRef(), View, LightShaftsSource);
 
 	const FIntPoint BufferSize = GSceneRenderTargets.GetBufferSizeXY();
 	const uint32 DownsampleFactor	= GetLightShaftDownsampleFactor();
diff --git a/Engine/Source/Runtime/Renderer/Private/PostProcess/PostProcessAA.cpp b/Engine/Source/Runtime/Renderer/Private/PostProcess/PostProcessAA.cpp
--- a/Engine/Source/Runtime/Renderer/Private/PostProcess/PostProcessAA.cpp
+++ b/Engine/Source/Runtime/Renderer/Private/PostProcess/PostProcessAA.cpp
@@ -70,7 +70,7 @@ public:
 		return bShaderHasOutdatedParameters;
 	}
 
-	void SetParameters(FRHICommandList* RHICmdList, const FRenderingCompositePassContext& Context)
+	void SetParameters(FRHICommandList& RHICmdList, const FRenderingCompositePassContext& Context)
 	{
 		const FPixelShaderRHIParamRef ShaderRHI = GetPixelShader();
 
@@ -186,7 +186,7 @@ public:
 		return bShaderHasOutdatedParameters;
 	}
 
-	void SetParameters(FRHICommandList* RHICmdList, const FRenderingCompositePassContext& Context)
+	void SetParameters(FRHICommandList& RHICmdList, const FRenderingCompositePassContext& Context)
 	{
 		const FVertexShaderRHIParamRef ShaderRHI = GetVertexShader();
 
@@ -213,7 +213,7 @@ IMPLEMENT_SHADER_TYPE(,FFXAAVS,TEXT("FXAAShader"),TEXT("FxaaVS"),SF_Vertex);
 
 // @param Quality 1..6
 template <uint32 Quality>
-static void SetShaderTemplAA(FRHICommandList* RHICmdList, const FRenderingCompositePassContext& Context)
+static void SetShaderTemplAA(FRHICommandList& RHICmdList, const FRenderingCompositePassContext& Context)
 {
 	TShaderMapRef<FFXAAVS> VertexShader(GetGlobalShaderMap());
 	TShaderMapRef<FPostProcessAntiAliasingPS<Quality> > PixelShader(GetGlobalShaderMap());
@@ -249,8 +249,8 @@ void FRCPassPostProcessAA::Process(FRenderingCompositePassContext& Context)
 	const FSceneRenderTargetItem& DestRenderTarget = PassOutputs[0].RequestSurface(Context);
 
 	//@todo-rco: RHIPacketList
-	FRHICommandList* RHICmdList = nullptr;
-	check(!RHICmdList);
+	FRHICommandList& RHICmdList = FRHICommandList::GetNullRef();
+	RHICmdList.CheckIsNull();
 
 	// Set the view family's render target/viewport.
 	RHISetRenderTarget(DestRenderTarget.TargetableTexture, FTextureRHIRef());	
diff --git a/Engine/Source/Runtime/Renderer/Private/PostProcess/PostProcessBloomSetup.cpp b/Engine/Source/Runtime/Renderer/Private/PostProcess/PostProcessBloomSetup.cpp
--- a/Engine/Source/Runtime/Renderer/Private/PostProcess/PostProcessBloomSetup.cpp
+++ b/Engine/Source/Runtime/Renderer/Private/PostProcess/PostProcessBloomSetup.cpp
@@ -48,7 +48,7 @@ public:
 		BloomThreshold.Bind(Initializer.ParameterMap, TEXT("BloomThreshold"));
 	}
 
-	void SetPS(FRHICommandList* RHICmdList, const FRenderingCompositePassContext& Context)
+	void SetPS(FRHICommandList& RHICmdList, const FRenderingCompositePassContext& Context)
 	{
 		const FPixelShaderRHIParamRef ShaderRHI = GetPixelShader();
 
@@ -102,7 +102,7 @@ public:
 		EyeAdaptation.Bind(Initializer.ParameterMap, TEXT("EyeAdaptation"));
 	}
 	
-	void SetVS(FRHICommandList* RHICmdList, const FRenderingCompositePassContext& Context)
+	void SetVS(FRHICommandList& RHICmdList, const FRenderingCompositePassContext& Context)
 	{
 		const FVertexShaderRHIParamRef ShaderRHI = GetVertexShader();
 
@@ -163,7 +163,7 @@ void FRCPassPostProcessBloomSetup::Process(FRenderingCompositePassContext& Conte
 	const FSceneRenderTargetItem& DestRenderTarget = PassOutputs[0].RequestSurface(Context);
 
 	//@todo-rco: RHIPacketList
-	FRHICommandList* RHICmdList = nullptr;
+	FRHICommandList& RHICmdList = FRHICommandList::GetNullRef();
 
 	// Set the view family's render target/viewport.
 	RHISetRenderTarget(DestRenderTarget.TargetableTexture, FTextureRHIRef());	
diff --git a/Engine/Source/Runtime/Renderer/Private/PostProcess/PostProcessBlurGBuffer.h b/Engine/Source/Runtime/Renderer/Private/PostProcess/PostProcessBlurGBuffer.h
--- a/Engine/Source/Runtime/Renderer/Private/PostProcess/PostProcessBlurGBuffer.h
+++ b/Engine/Source/Runtime/Renderer/Private/PostProcess/PostProcessBlurGBuffer.h
@@ -18,5 +18,5 @@ public:
 	virtual void Release() OVERRIDE { delete this; }
 	virtual FPooledRenderTargetDesc ComputeOutputDesc(EPassOutputId InPassOutputId) const;
 
-	static void SetShader(FRHICommandList* RHICmdList, const FRenderingCompositePassContext& Context, uint32 BlurType);
+	static void SetShader(FRHICommandList& RHICmdList, const FRenderingCompositePassContext& Context, uint32 BlurType);
 };
diff --git a/Engine/Source/Runtime/Renderer/Private/PostProcess/PostProcessBokehDOF.cpp b/Engine/Source/Runtime/Renderer/Private/PostProcess/PostProcessBokehDOF.cpp
--- a/Engine/Source/Runtime/Renderer/Private/PostProcess/PostProcessBokehDOF.cpp
+++ b/Engine/Source/Runtime/Renderer/Private/PostProcess/PostProcessBokehDOF.cpp
@@ -107,7 +107,7 @@ public:
 		return bShaderHasOutdatedParameters;
 	}
 
-	void SetParameters(FRHICommandList* RHICmdList, const FRenderingCompositePassContext& Context)
+	void SetParameters(FRHICommandList& RHICmdList, const FRenderingCompositePassContext& Context)
 	{
 		const FPixelShaderRHIParamRef ShaderRHI = GetPixelShader();
 
@@ -169,7 +169,7 @@ void FRCPassPostProcessVisualizeDOF::Process(FRenderingCompositePassContext& Con
 	const FSceneRenderTargetItem& DestRenderTarget = PassOutputs[0].RequestSurface(Context);
 
 	//@todo-rco: RHIPacketList
-	FRHICommandList* RHICmdList = nullptr;
+	FRHICommandList& RHICmdList = FRHICommandList::GetNullRef();
 
 	// Set the view family's render target/viewport.
 	RHISetRenderTarget(DestRenderTarget.TargetableTexture, FTextureRHIRef());	
@@ -270,7 +270,7 @@ public:
 		return bShaderHasOutdatedParameters;
 	}
 
-	void SetParameters(FRHICommandList* RHICmdList, const FRenderingCompositePassContext& Context)
+	void SetParameters(FRHICommandList& RHICmdList, const FRenderingCompositePassContext& Context)
 	{
 		const FPixelShaderRHIParamRef ShaderRHI = GetPixelShader();
 
@@ -331,7 +331,7 @@ void FRCPassPostProcessBokehDOFSetup::Process(FRenderingCompositePassContext& Co
 	const FSceneRenderTargetItem& DestRenderTarget = PassOutputs[0].RequestSurface(Context);
 
 	//@todo-rco: RHIPacketList
-	FRHICommandList* RHICmdList = nullptr;
+	FRHICommandList& RHICmdList = FRHICommandList::GetNullRef();
 
 	// Set the view family's render target/viewport.
 	RHISetRenderTarget(DestRenderTarget.TargetableTexture, FTextureRHIRef());	
@@ -440,7 +440,7 @@ public:
 	}
 
 	/** to have a similar interface as all other shaders */
-	void SetParameters(FRHICommandList* RHICmdList, const FRenderingCompositePassContext& Context, FIntPoint TileCountValue, uint32 TileSize, float PixelKernelSize, FIntPoint LeftTop)
+	void SetParameters(FRHICommandList& RHICmdList, const FRenderingCompositePassContext& Context, FIntPoint TileCountValue, uint32 TileSize, float PixelKernelSize, FIntPoint LeftTop)
 	{
 		const FVertexShaderRHIParamRef ShaderRHI = GetVertexShader();
 
@@ -524,7 +524,7 @@ public:
 		return bShaderHasOutdatedParameters;
 	}
 
-	void SetParameters(FRHICommandList* RHICmdList, const FRenderingCompositePassContext& Context, float PixelKernelSize)
+	void SetParameters(FRHICommandList& RHICmdList, const FRenderingCompositePassContext& Context, float PixelKernelSize)
 	{
 		const FPixelShaderRHIParamRef ShaderRHI = GetPixelShader();
 
@@ -562,7 +562,7 @@ VARIATION1(0,1)			VARIATION1(1,1)			VARIATION1(2,1)
 #undef VARIATION1
 
 template <uint32 DOFMethod, uint32 DOFIndexStyle>
-void FRCPassPostProcessBokehDOF::SetShaderTempl(FRHICommandList* RHICmdList, const FRenderingCompositePassContext& Context, FIntPoint LeftTop, FIntPoint TileCount, uint32 TileSize, float PixelKernelSize)
+void FRCPassPostProcessBokehDOF::SetShaderTempl(FRHICommandList& RHICmdList, const FRenderingCompositePassContext& Context, FIntPoint LeftTop, FIntPoint TileCount, uint32 TileSize, float PixelKernelSize)
 {
 	TShaderMapRef<FPostProcessBokehDOFVS<DOFMethod, DOFIndexStyle> > VertexShader(GetGlobalShaderMap());
 	TShaderMapRef<FPostProcessBokehDOFPS> PixelShader(GetGlobalShaderMap());
@@ -632,7 +632,7 @@ void FRCPassPostProcessBokehDOF::Process(FRenderingCompositePassContext& Context
 	const FSceneRenderTargetItem& DestRenderTarget = PassOutputs[0].RequestSurface(Context);
 
 	//@todo-rco: RHIPacketList
-	FRHICommandList* RHICmdList = nullptr;
+	FRHICommandList& RHICmdList = FRHICommandList::GetNullRef();
 
 	// Set the view family's render target/viewport.
 	RHISetRenderTarget(DestRenderTarget.TargetableTexture, FTextureRHIRef());	
diff --git a/Engine/Source/Runtime/Renderer/Private/PostProcess/PostProcessBokehDOF.h b/Engine/Source/Runtime/Renderer/Private/PostProcess/PostProcessBokehDOF.h
--- a/Engine/Source/Runtime/Renderer/Private/PostProcess/PostProcessBokehDOF.h
+++ b/Engine/Source/Runtime/Renderer/Private/PostProcess/PostProcessBokehDOF.h
@@ -49,7 +49,7 @@ public:
 private:
 
 	template <uint32 HighQuality, uint32 IndexStyle>
-	void SetShaderTempl(FRHICommandList* RHICmdList, const FRenderingCompositePassContext& Context, FIntPoint LeftTop, FIntPoint TileCount, uint32 TileSize, float PixelKernelSize);
+	void SetShaderTempl(FRHICommandList& RHICmdList, const FRenderingCompositePassContext& Context, FIntPoint LeftTop, FIntPoint TileCount, uint32 TileSize, float PixelKernelSize);
 
 	// border between front and back layer as we don't use viewports (only possible with GS)
 	const static uint32 SafetyBorder = 40;
diff --git a/Engine/Source/Runtime/Renderer/Private/PostProcess/PostProcessBokehDOFRecombine.cpp b/Engine/Source/Runtime/Renderer/Private/PostProcess/PostProcessBokehDOFRecombine.cpp
--- a/Engine/Source/Runtime/Renderer/Private/PostProcess/PostProcessBokehDOFRecombine.cpp
+++ b/Engine/Source/Runtime/Renderer/Private/PostProcess/PostProcessBokehDOFRecombine.cpp
@@ -59,7 +59,7 @@ public:
 		return bShaderHasOutdatedParameters;
 	}
 
-	void SetParameters(FRHICommandList* RHICmdList, const FRenderingCompositePassContext& Context)
+	void SetParameters(FRHICommandList& RHICmdList, const FRenderingCompositePassContext& Context)
 	{
 		const FPixelShaderRHIParamRef ShaderRHI = GetPixelShader();
 
@@ -96,7 +96,7 @@ VARIATION1(1)			VARIATION1(2)			VARIATION1(3)
 
 
 template <uint32 Method>
-void FRCPassPostProcessBokehDOFRecombine::SetShader(FRHICommandList* RHICmdList, const FRenderingCompositePassContext& Context)
+void FRCPassPostProcessBokehDOFRecombine::SetShader(FRHICommandList& RHICmdList, const FRenderingCompositePassContext& Context)
 {
 	TShaderMapRef<FPostProcessVS> VertexShader(GetGlobalShaderMap());
 	TShaderMapRef<FPostProcessBokehDOFRecombinePS<Method> > PixelShader(GetGlobalShaderMap());
@@ -145,7 +145,7 @@ void FRCPassPostProcessBokehDOFRecombine::Process(FRenderingCompositePassContext
 
 	const FSceneRenderTargetItem& DestRenderTarget = PassOutputs[0].RequestSurface(Context);
 	//@todo-rco: RHIPacketList
-	FRHICommandList* RHICmdList = nullptr;
+	FRHICommandList& RHICmdList = FRHICommandList::GetNullRef();
 
 	// Set the view family's render target/viewport.
 	RHISetRenderTarget(DestRenderTarget.TargetableTexture, FTextureRHIRef());	
diff --git a/Engine/Source/Runtime/Renderer/Private/PostProcess/PostProcessBokehDOFRecombine.h b/Engine/Source/Runtime/Renderer/Private/PostProcess/PostProcessBokehDOFRecombine.h
--- a/Engine/Source/Runtime/Renderer/Private/PostProcess/PostProcessBokehDOFRecombine.h
+++ b/Engine/Source/Runtime/Renderer/Private/PostProcess/PostProcessBokehDOFRecombine.h
@@ -23,5 +23,5 @@ public:
 private:
 
 	template <uint32 Method>
-	static void SetShader(FRHICommandList* RHICmdList, const FRenderingCompositePassContext& Context);
+	static void SetShader(FRHICommandList& RHICmdList, const FRenderingCompositePassContext& Context);
 };
\ No newline at end of file
diff --git a/Engine/Source/Runtime/Renderer/Private/PostProcess/PostProcessBusyWait.cpp b/Engine/Source/Runtime/Renderer/Private/PostProcess/PostProcessBusyWait.cpp
--- a/Engine/Source/Runtime/Renderer/Private/PostProcess/PostProcessBusyWait.cpp
+++ b/Engine/Source/Runtime/Renderer/Private/PostProcess/PostProcessBusyWait.cpp
@@ -35,7 +35,7 @@ public:
 		GPUBusyWait.Bind(Initializer.ParameterMap,TEXT("GPUBusyWait"));
 	}
 
-	void SetPS(FRHICommandList* RHICmdList, const FRenderingCompositePassContext& Context)
+	void SetPS(FRHICommandList& RHICmdList, const FRenderingCompositePassContext& Context)
 	{
 		const FPixelShaderRHIParamRef ShaderRHI = GetPixelShader();
 
@@ -86,7 +86,7 @@ void FRCPassPostProcessBusyWait::Process(FRenderingCompositePassContext& Context
 	const FSceneRenderTargetItem& DestRenderTarget = GSceneRenderTargets.GetLightAttenuation()->GetRenderTargetItem();
 
 	//@todo-rco: RHIPacketList
-	FRHICommandList* RHICmdList = nullptr;
+	FRHICommandList& RHICmdList = FRHICommandList::GetNullRef();
 
 	// Set the view family's render target/viewport.
 	RHISetRenderTarget(DestRenderTarget.TargetableTexture, FTextureRHIRef());	
diff --git a/Engine/Source/Runtime/Renderer/Private/PostProcess/PostProcessCombineLUTs.cpp b/Engine/Source/Runtime/Renderer/Private/PostProcess/PostProcessCombineLUTs.cpp
--- a/Engine/Source/Runtime/Renderer/Private/PostProcess/PostProcessCombineLUTs.cpp
+++ b/Engine/Source/Runtime/Renderer/Private/PostProcess/PostProcessCombineLUTs.cpp
@@ -51,7 +51,7 @@ FColorRemapShaderParameters::FColorRemapShaderParameters(const FShaderParameterM
 	MappingPolynomial.Bind(ParameterMap, TEXT("MappingPolynomial"));
 }
 
-void FColorRemapShaderParameters::Set(FRHICommandList* RHICmdList, const FPixelShaderRHIParamRef ShaderRHI)
+void FColorRemapShaderParameters::Set(FRHICommandList& RHICmdList, const FPixelShaderRHIParamRef ShaderRHI)
 {
 	FColorTransform ColorTransform;
 
@@ -119,7 +119,7 @@ public:
 	}
 	FLUTBlenderPS() {}
 
-	void SetParameters(FRHICommandList* RHICmdList, const FSceneView& View, FTexture* Texture[BlendCount], float Weights[BlendCount])
+	void SetParameters(FRHICommandList& RHICmdList, const FSceneView& View, FTexture* Texture[BlendCount], float Weights[BlendCount])
 	{
 		const FPixelShaderRHIParamRef ShaderRHI = GetPixelShader();
 
@@ -180,7 +180,7 @@ IMPLEMENT_SHADER_TYPE(template<>,FLUTBlenderPS<3>,TEXT("PostProcessCombineLUTs")
 IMPLEMENT_SHADER_TYPE(template<>,FLUTBlenderPS<4>,TEXT("PostProcessCombineLUTs"),TEXT("MainPS"),SF_Pixel);
 IMPLEMENT_SHADER_TYPE(template<>,FLUTBlenderPS<5>,TEXT("PostProcessCombineLUTs"),TEXT("MainPS"),SF_Pixel);
 
-void SetLUTBlenderShader(FRHICommandList* RHICmdList, FRenderingCompositePassContext& Context, uint32 BlendCount, FTexture* Texture[], float Weights[], const FVolumeBounds& VolumeBounds)
+void SetLUTBlenderShader(FRHICommandList& RHICmdList, FRenderingCompositePassContext& Context, uint32 BlendCount, FTexture* Texture[], float Weights[], const FVolumeBounds& VolumeBounds)
 {
 	check(BlendCount > 0);
 
@@ -387,7 +387,7 @@ void FRCPassPostProcessCombineLUTs::Process(FRenderingCompositePassContext& Cont
 	const FSceneRenderTargetItem& DestRenderTarget = PassOutputs[0].RequestSurface(Context);
 
 	//@todo-rco: RHIPacketList
-	FRHICommandList* RHICmdList = nullptr;
+	FRHICommandList& RHICmdList = FRHICommandList::GetNullRef();
 
 	// Set the view family's render target/viewport.
 	RHISetRenderTarget(DestRenderTarget.TargetableTexture, FTextureRHIRef());	
diff --git a/Engine/Source/Runtime/Renderer/Private/PostProcess/PostProcessCombineLUTs.h b/Engine/Source/Runtime/Renderer/Private/PostProcess/PostProcessCombineLUTs.h
--- a/Engine/Source/Runtime/Renderer/Private/PostProcess/PostProcessCombineLUTs.h
+++ b/Engine/Source/Runtime/Renderer/Private/PostProcess/PostProcessCombineLUTs.h
@@ -38,7 +38,7 @@ public:
 
 	FColorRemapShaderParameters(const FShaderParameterMap& ParameterMap);
 
-	void Set(FRHICommandList* RHICmdList, const FPixelShaderRHIParamRef ShaderRHI);
+	void Set(FRHICommandList& RHICmdList, const FPixelShaderRHIParamRef ShaderRHI);
 
 	friend FArchive& operator<<(FArchive& Ar,FColorRemapShaderParameters& P);
 
diff --git a/Engine/Source/Runtime/Renderer/Private/PostProcess/PostProcessCompositeEditorPrimitives.cpp b/Engine/Source/Runtime/Renderer/Private/PostProcess/PostProcessCompositeEditorPrimitives.cpp
--- a/Engine/Source/Runtime/Renderer/Private/PostProcess/PostProcessCompositeEditorPrimitives.cpp
+++ b/Engine/Source/Runtime/Renderer/Private/PostProcess/PostProcessCompositeEditorPrimitives.cpp
@@ -66,7 +66,7 @@ public:
 		FilteredSceneDepthTextureSampler.Bind(Initializer.ParameterMap,TEXT("FilteredSceneDepthTextureSampler"));
 	}
 
-	void SetParameters(FRHICommandList* RHICmdList, const FRenderingCompositePassContext& Context)
+	void SetParameters(FRHICommandList& RHICmdList, const FRenderingCompositePassContext& Context)
 	{
 		const FPixelShaderRHIParamRef ShaderRHI = GetPixelShader();
 
@@ -162,7 +162,7 @@ private:
 
 
 template <uint32 MSAASampleCount>
-static void SetCompositePrimitivesShaderTempl(FRHICommandList* RHICmdList, const FRenderingCompositePassContext& Context)
+static void SetCompositePrimitivesShaderTempl(FRHICommandList& RHICmdList, const FRenderingCompositePassContext& Context)
 {
 	TShaderMapRef<FPostProcessVS> VertexShader(GetGlobalShaderMap());
 	TShaderMapRef<FPostProcessCompositeEditorPrimitivesPS<MSAASampleCount> > PixelShader(GetGlobalShaderMap());
@@ -203,7 +203,7 @@ void FRCPassPostProcessCompositeEditorPrimitives::Process(FRenderingCompositePas
 	const uint32 MSAASampleCount = GSceneRenderTargets.EditorPrimitivesColor->GetDesc().NumSamples;
 
 	//@todo-rco: RHIPacketList
-	FRHICommandList* RHICmdList = nullptr;
+	FRHICommandList& RHICmdList = FRHICommandList::GetNullRef();
 
 	{
 		RHISetRenderTarget( ColorTarget, DepthTarget );
diff --git a/Engine/Source/Runtime/Renderer/Private/PostProcess/PostProcessDOF.cpp b/Engine/Source/Runtime/Renderer/Private/PostProcess/PostProcessDOF.cpp
--- a/Engine/Source/Runtime/Renderer/Private/PostProcess/PostProcessDOF.cpp
+++ b/Engine/Source/Runtime/Renderer/Private/PostProcess/PostProcessDOF.cpp
@@ -55,7 +55,7 @@ public:
 		return bShaderHasOutdatedParameters;
 	}
 
-	void SetParameters(FRHICommandList* RHICmdList, const FRenderingCompositePassContext& Context)
+	void SetParameters(FRHICommandList& RHICmdList, const FRenderingCompositePassContext& Context)
 	{
 		const FPixelShaderRHIParamRef ShaderRHI = GetPixelShader();
 
@@ -108,7 +108,7 @@ void FRCPassPostProcessDOFSetup::Process(FRenderingCompositePassContext& Context
 	const FSceneRenderTargetItem& DestRenderTarget1 = bNearBlurEnabled ? PassOutputs[1].RequestSurface(Context) : FSceneRenderTargetItem();
 
 	//@todo-rco: RHIPacketList
-	FRHICommandList* RHICmdList = nullptr;
+	FRHICommandList& RHICmdList = FRHICommandList::GetNullRef();
 
 	// Set the view family's render target/viewport.
 	FTextureRHIParamRef RenderTargets[2] =
@@ -236,7 +236,7 @@ public:
 		return bShaderHasOutdatedParameters;
 	}
 
-	void SetParameters(FRHICommandList* RHICmdList, const FRenderingCompositePassContext& Context)
+	void SetParameters(FRHICommandList& RHICmdList, const FRenderingCompositePassContext& Context)
 	{
 		const FPixelShaderRHIParamRef ShaderRHI = GetPixelShader();
 
@@ -283,7 +283,7 @@ void FRCPassPostProcessDOFRecombine::Process(FRenderingCompositePassContext& Con
 	const FSceneRenderTargetItem& DestRenderTarget = PassOutputs[0].RequestSurface(Context);
 
 	//@todo-rco: RHIPacketList
-	FRHICommandList* RHICmdList = nullptr;
+	FRHICommandList& RHICmdList = FRHICommandList::GetNullRef();
 
 	// Set the view family's render target/viewport.
 	RHISetRenderTarget(DestRenderTarget.TargetableTexture, FTextureRHIRef());	
diff --git a/Engine/Source/Runtime/Renderer/Private/PostProcess/PostProcessDownsample.cpp b/Engine/Source/Runtime/Renderer/Private/PostProcess/PostProcessDownsample.cpp
--- a/Engine/Source/Runtime/Renderer/Private/PostProcess/PostProcessDownsample.cpp
+++ b/Engine/Source/Runtime/Renderer/Private/PostProcess/PostProcessDownsample.cpp
@@ -52,7 +52,7 @@ public:
 		return bShaderHasOutdatedParameters;
 	}
 
-	void SetParameters(FRHICommandList* RHICmdList, const FRenderingCompositePassContext& Context)
+	void SetParameters(FRHICommandList& RHICmdList, const FRenderingCompositePassContext& Context)
 	{
 		const FPixelShaderRHIParamRef ShaderRHI = GetPixelShader();
 
@@ -108,7 +108,7 @@ public:
 		return bShaderHasOutdatedParameters;
 	}
 
-	void SetParameters(FRHICommandList* RHICmdList, const FRenderingCompositePassContext& Context)
+	void SetParameters(FRHICommandList& RHICmdList, const FRenderingCompositePassContext& Context)
 	{
 		const FVertexShaderRHIParamRef ShaderRHI = GetVertexShader();
 
@@ -137,7 +137,7 @@ FRCPassPostProcessDownsample::FRCPassPostProcessDownsample(EPixelFormat InOverri
 
 
 template <uint32 Method>
-void FRCPassPostProcessDownsample::SetShader(FRHICommandList* RHICmdList, const FRenderingCompositePassContext& Context)
+void FRCPassPostProcessDownsample::SetShader(FRHICommandList& RHICmdList, const FRenderingCompositePassContext& Context)
 {
 	TShaderMapRef<FPostProcessDownsampleVS> VertexShader(GetGlobalShaderMap());
 	TShaderMapRef<FPostProcessDownsamplePS<Method> > PixelShader(GetGlobalShaderMap());
@@ -175,7 +175,7 @@ void FRCPassPostProcessDownsample::Process(FRenderingCompositePassContext& Conte
 	const FSceneRenderTargetItem& DestRenderTarget = PassOutputs[0].RequestSurface(Context);
 
 	//@todo-rco: RHIPacketList
-	FRHICommandList* RHICmdList = nullptr;
+	FRHICommandList& RHICmdList = FRHICommandList::GetNullRef();
 
 	// Set the view family's render target/viewport.
 	RHISetRenderTarget(DestRenderTarget.TargetableTexture, FTextureRHIRef());
diff --git a/Engine/Source/Runtime/Renderer/Private/PostProcess/PostProcessDownsample.h b/Engine/Source/Runtime/Renderer/Private/PostProcess/PostProcessDownsample.h
--- a/Engine/Source/Runtime/Renderer/Private/PostProcess/PostProcessDownsample.h
+++ b/Engine/Source/Runtime/Renderer/Private/PostProcess/PostProcessDownsample.h
@@ -31,7 +31,7 @@ public:
 	virtual FPooledRenderTargetDesc ComputeOutputDesc(EPassOutputId InPassOutputId) const;
 
 private:
-	template <uint32 Method> static void SetShader(FRHICommandList* RHICmdList, const FRenderingCompositePassContext& Context);
+	template <uint32 Method> static void SetShader(FRHICommandList& RHICmdList, const FRenderingCompositePassContext& Context);
 
 	EPixelFormat OverrideFormat;
 	// explained in constructor
diff --git a/Engine/Source/Runtime/Renderer/Private/PostProcess/PostProcessEyeAdaptation.cpp b/Engine/Source/Runtime/Renderer/Private/PostProcess/PostProcessEyeAdaptation.cpp
--- a/Engine/Source/Runtime/Renderer/Private/PostProcess/PostProcessEyeAdaptation.cpp
+++ b/Engine/Source/Runtime/Renderer/Private/PostProcess/PostProcessEyeAdaptation.cpp
@@ -41,7 +41,7 @@ public:
 		EyeAdaptationParams.Bind(Initializer.ParameterMap, TEXT("EyeAdaptationParams"));
 	}
 
-	void SetPS(FRHICommandList* RHICmdList, const FRenderingCompositePassContext& Context)
+	void SetPS(FRHICommandList& RHICmdList, const FRenderingCompositePassContext& Context)
 	{
 		const FPixelShaderRHIParamRef ShaderRHI = GetPixelShader();
 
@@ -82,7 +82,7 @@ void FRCPassPostProcessEyeAdaptation::Process(FRenderingCompositePassContext& Co
 
 	// we render to our own output render target, not the intermediate one created by the compositing system
 	//@todo-rco: RHIPacketList
-	FRHICommandList* RHICmdList = nullptr;
+	FRHICommandList& RHICmdList = FRHICommandList::GetNullRef();
 
 	// Set the view family's render target/viewport.
 	RHISetRenderTarget(EyeAdaptation->GetRenderTargetItem().TargetableTexture, FTextureRHIRef());	
diff --git a/Engine/Source/Runtime/Renderer/Private/PostProcess/PostProcessHMD.cpp b/Engine/Source/Runtime/Renderer/Private/PostProcess/PostProcessHMD.cpp
--- a/Engine/Source/Runtime/Renderer/Private/PostProcess/PostProcessHMD.cpp
+++ b/Engine/Source/Runtime/Renderer/Private/PostProcess/PostProcessHMD.cpp
@@ -72,7 +72,7 @@ public:
 		EyeToSrcUVOffset.Bind(Initializer.ParameterMap, TEXT("EyeToSrcUVOffset"));
 	}
 
-	void SetVS(FRHICommandList* RHICmdList, const FRenderingCompositePassContext& Context, EStereoscopicPass StereoPass)
+	void SetVS(FRHICommandList& RHICmdList, const FRenderingCompositePassContext& Context, EStereoscopicPass StereoPass)
 	{
 		const FVertexShaderRHIParamRef ShaderRHI = GetVertexShader();
 
@@ -123,7 +123,7 @@ public:
 
 	}
 
-	void SetPS(FRHICommandList* RHICmdList, const FRenderingCompositePassContext& Context, FIntRect SrcRect, FIntPoint SrcBufferSize, EStereoscopicPass StereoPass, FMatrix& QuadTexTransform)
+	void SetPS(FRHICommandList& RHICmdList, const FRenderingCompositePassContext& Context, FIntRect SrcRect, FIntPoint SrcBufferSize, EStereoscopicPass StereoPass, FMatrix& QuadTexTransform)
 	{
 		const FPixelShaderRHIParamRef ShaderRHI = GetPixelShader();
 
@@ -175,7 +175,7 @@ void FRCPassPostProcessHMD::Process(FRenderingCompositePassContext& Context)
     const FSceneRenderTargetItem& DestRenderTarget = PassOutputs[0].RequestSurface(Context);
 
 	//@todo-rco: RHIPacketList
-	FRHICommandList* RHICmdList = nullptr;
+	FRHICommandList& RHICmdList = FRHICommandList::GetNullRef();
 
 	// Set the view family's render target/viewport.
 	RHISetRenderTarget(DestRenderTarget.TargetableTexture, FTextureRHIRef());	
diff --git a/Engine/Source/Runtime/Renderer/Private/PostProcess/PostProcessHistogram.cpp b/Engine/Source/Runtime/Renderer/Private/PostProcess/PostProcessHistogram.cpp
--- a/Engine/Source/Runtime/Renderer/Private/PostProcess/PostProcessHistogram.cpp
+++ b/Engine/Source/Runtime/Renderer/Private/PostProcess/PostProcessHistogram.cpp
@@ -55,7 +55,7 @@ public:
 		EyeAdaptationParams.Bind(Initializer.ParameterMap, TEXT("EyeAdaptationParams"));
 	}
 
-	void SetCS(FRHICommandList* RHICmdList, const FRenderingCompositePassContext& Context, FIntPoint ThreadGroupCountValue, FIntPoint LeftTopOffsetValue, FIntPoint GatherExtent)
+	void SetCS(FRHICommandList& RHICmdList, const FRenderingCompositePassContext& Context, FIntPoint ThreadGroupCountValue, FIntPoint LeftTopOffsetValue, FIntPoint GatherExtent)
 	{
 		const FComputeShaderRHIParamRef ShaderRHI = GetComputeShader();
 
@@ -124,7 +124,7 @@ void FRCPassPostProcessHistogram::Process(FRenderingCompositePassContext& Contex
 
 	
 	//@todo-rco: RHIPacketList
-	FRHICommandList* RHICmdList = nullptr;
+	FRHICommandList& RHICmdList = FRHICommandList::GetNullRef();
 	ComputeShader->SetCS(RHICmdList, Context, ThreadGroupCountValue, (DestRect.Min + FIntPoint(1, 1)) / 2, GatherExtent);
 	
 	DispatchComputeShader(RHICmdList, *ComputeShader, ThreadGroupCountValue.X, ThreadGroupCountValue.Y, 1);
diff --git a/Engine/Source/Runtime/Renderer/Private/PostProcess/PostProcessHistogramReduce.cpp b/Engine/Source/Runtime/Renderer/Private/PostProcess/PostProcessHistogramReduce.cpp
--- a/Engine/Source/Runtime/Renderer/Private/PostProcess/PostProcessHistogramReduce.cpp
+++ b/Engine/Source/Runtime/Renderer/Private/PostProcess/PostProcessHistogramReduce.cpp
@@ -45,7 +45,7 @@ public:
 		EyeAdapationTemporalParams.Bind(Initializer.ParameterMap, TEXT("EyeAdapationTemporalParams"));
 	}
 
-	void SetPS(FRHICommandList* RHICmdList, const FRenderingCompositePassContext& Context, uint32 LoopSizeValue)
+	void SetPS(FRHICommandList& RHICmdList, const FRenderingCompositePassContext& Context, uint32 LoopSizeValue)
 	{
 		const FPixelShaderRHIParamRef ShaderRHI = GetPixelShader();
 
@@ -126,7 +126,7 @@ void FRCPassPostProcessHistogramReduce::Process(FRenderingCompositePassContext&
 	uint32 LoopSizeValue = ComputeLoopSize(GatherExtent);
 
 	//@todo-rco: RHIPacketList
-	FRHICommandList* RHICmdList = nullptr;
+	FRHICommandList& RHICmdList = FRHICommandList::GetNullRef();
 	PixelShader->SetPS(RHICmdList, Context, LoopSizeValue);
 
 	// Draw a quad mapping scene color to the view's render target
diff --git a/Engine/Source/Runtime/Renderer/Private/PostProcess/PostProcessLensBlur.cpp b/Engine/Source/Runtime/Renderer/Private/PostProcess/PostProcessLensBlur.cpp
--- a/Engine/Source/Runtime/Renderer/Private/PostProcess/PostProcessLensBlur.cpp
+++ b/Engine/Source/Runtime/Renderer/Private/PostProcess/PostProcessLensBlur.cpp
@@ -49,7 +49,7 @@ public:
 	}
 
 	/** to have a similar interface as all other shaders */
-	void SetParameters(FRHICommandList* RHICmdList, const FRenderingCompositePassContext& Context, FIntPoint TileCountValue, uint32 TileSize, float PixelKernelSize, float Threshold)
+	void SetParameters(FRHICommandList& RHICmdList, const FRenderingCompositePassContext& Context, FIntPoint TileCountValue, uint32 TileSize, float PixelKernelSize, float Threshold)
 	{
 		const FVertexShaderRHIParamRef ShaderRHI = GetVertexShader();
 
@@ -115,7 +115,7 @@ public:
 		return bShaderHasOutdatedParameters;
 	}
 
-	void SetParameters(FRHICommandList* RHICmdList, const FRenderingCompositePassContext& Context, float PixelKernelSize)
+	void SetParameters(FRHICommandList& RHICmdList, const FRenderingCompositePassContext& Context, float PixelKernelSize)
 	{
 		const FPixelShaderRHIParamRef ShaderRHI = GetPixelShader();
 
@@ -209,7 +209,7 @@ void FRCPassPostProcessLensBlur::Process(FRenderingCompositePassContext& Context
 	float PixelKernelSize = PercentKernelSize / 100.0f * ViewSize.X;
 
 	//@todo-rco: RHIPacketList
-	FRHICommandList* RHICmdList = nullptr;
+	FRHICommandList& RHICmdList = FRHICommandList::GetNullRef();
 	VertexShader->SetParameters(RHICmdList, Context, TileCount, TileSize, PixelKernelSize, Threshold);
 	PixelShader->SetParameters(RHICmdList, Context, PixelKernelSize);
 
diff --git a/Engine/Source/Runtime/Renderer/Private/PostProcess/PostProcessLensFlares.cpp b/Engine/Source/Runtime/Renderer/Private/PostProcess/PostProcessLensFlares.cpp
--- a/Engine/Source/Runtime/Renderer/Private/PostProcess/PostProcessLensFlares.cpp
+++ b/Engine/Source/Runtime/Renderer/Private/PostProcess/PostProcessLensFlares.cpp
@@ -41,7 +41,7 @@ public:
 		return bShaderHasOutdatedParameters;
 	}
 
-	void SetParameters(FRHICommandList* RHICmdList, const FRenderingCompositePassContext& Context)
+	void SetParameters(FRHICommandList& RHICmdList, const FRenderingCompositePassContext& Context)
 	{
 		const FPixelShaderRHIParamRef ShaderRHI = GetPixelShader();
 
@@ -89,7 +89,7 @@ public:
 		return bShaderHasOutdatedParameters;
 	}
 
-	void SetParameters(FRHICommandList* RHICmdList, const FRenderingCompositePassContext& Context, FVector2D TexScaleValue)
+	void SetParameters(FRHICommandList& RHICmdList, const FRenderingCompositePassContext& Context, FVector2D TexScaleValue)
 	{
 		const FPixelShaderRHIParamRef ShaderRHI = GetPixelShader();
 
@@ -156,7 +156,7 @@ void FRCPassPostProcessLensFlares::Process(FRenderingCompositePassContext& Conte
 	TShaderMapRef<FPostProcessVS> VertexShader(GetGlobalShaderMap());
 
 	//@todo-rco: RHIPacketList
-	FRHICommandList* RHICmdList = nullptr;
+	FRHICommandList& RHICmdList = FRHICommandList::GetNullRef();
 
 	// setup background (bloom), can be implemented to use additive blending to avoid the read here
 	{
@@ -221,7 +221,7 @@ void FRCPassPostProcessLensFlares::Process(FRenderingCompositePassContext& Conte
 			
 			// set the individual flare color
 			//@todo-rco: RHIPacketList
-			FRHICommandList* RHICmdList = nullptr;
+			FRHICommandList& RHICmdList = FRHICommandList::GetNullRef();
 			SetShaderValue(RHICmdList, PixelShader->GetPixelShader(), PixelShader->FlareColor, FlareColor * LensFlareHDRColor);
 
 			// Draw a quad mapping scene color to the view's render target
diff --git a/Engine/Source/Runtime/Renderer/Private/PostProcess/PostProcessMaterial.cpp b/Engine/Source/Runtime/Renderer/Private/PostProcess/PostProcessMaterial.cpp
--- a/Engine/Source/Runtime/Renderer/Private/PostProcess/PostProcessMaterial.cpp
+++ b/Engine/Source/Runtime/Renderer/Private/PostProcess/PostProcessMaterial.cpp
@@ -31,7 +31,7 @@ public:
 	{
 	}
 
-	void SetParameters(FRHICommandList* RHICmdList, const FRenderingCompositePassContext& Context )
+	void SetParameters(FRHICommandList& RHICmdList, const FRenderingCompositePassContext& Context )
 	{
 		FMaterialShader::SetParameters(RHICmdList, GetVertexShader(), Context.View);
 	}
@@ -70,7 +70,7 @@ public:
 		PostprocessParameter.Bind(Initializer.ParameterMap);
 	}
 
-	void SetParameters(FRHICommandList* RHICmdList, const FRenderingCompositePassContext& Context, const FMaterialRenderProxy* Material )
+	void SetParameters(FRHICommandList& RHICmdList, const FRenderingCompositePassContext& Context, const FMaterialRenderProxy* Material )
 	{
 		const FPixelShaderRHIParamRef ShaderRHI = GetPixelShader();
 
@@ -155,7 +155,7 @@ void FRCPassPostProcessMaterial::Process(FRenderingCompositePassContext& Context
 	RHISetBoundShaderState(BoundShaderState);
 
 	//@todo-rco: RHIPacketList
-	FRHICommandList* RHICmdList = nullptr;
+	FRHICommandList& RHICmdList = FRHICommandList::GetNullRef();
 	VertexShader->SetParameters(RHICmdList, Context);
 	PixelShader->SetParameters(RHICmdList, Context, MaterialInterface->GetRenderProxy(false));
 
