UE4: Rearrange RHICmdList toward non-null cmd lists and the final "look"

diff --git a/Engine/Source/Runtime/Renderer/Private/PostProcess/PostProcessMobile.cpp b/Engine/Source/Runtime/Renderer/Private/PostProcess/PostProcessMobile.cpp
--- a/Engine/Source/Runtime/Renderer/Private/PostProcess/PostProcessMobile.cpp
+++ b/Engine/Source/Runtime/Renderer/Private/PostProcess/PostProcessMobile.cpp
@@ -39,7 +39,7 @@ public:
 		PostprocessParameter.Bind(Initializer.ParameterMap);
 	}
 
-	void SetVS(FRHICommandList* RHICmdList, const FRenderingCompositePassContext& Context)
+	void SetVS(FRHICommandList& RHICmdList, const FRenderingCompositePassContext& Context)
 	{
 		const FVertexShaderRHIParamRef ShaderRHI = GetVertexShader();
 
@@ -92,7 +92,7 @@ public:
 		BloomThreshold.Bind(Initializer.ParameterMap, TEXT("BloomThreshold"));
 	}
 
-	void SetPS(FRHICommandList* RHICmdList, const FRenderingCompositePassContext& Context)
+	void SetPS(FRHICommandList& RHICmdList, const FRenderingCompositePassContext& Context)
 	{
 		const FPixelShaderRHIParamRef ShaderRHI = GetPixelShader();
 
@@ -129,7 +129,7 @@ IMPLEMENT_SHADER_TYPE(template<>,FPostProcessBloomSetupPS_ES2_2,TEXT("PostProces
 IMPLEMENT_SHADER_TYPE(template<>,FPostProcessBloomSetupPS_ES2_3,TEXT("PostProcessMobile"),TEXT("BloomPS_ES2"),SF_Pixel);
 
 template <uint32 UseSunDof>
-static void BloomSetup_SetShader(FRHICommandList* RHICmdList, const FRenderingCompositePassContext& Context)
+static void BloomSetup_SetShader(FRHICommandList& RHICmdList, const FRenderingCompositePassContext& Context)
 {
 	TShaderMapRef<FPostProcessBloomSetupVS_ES2> VertexShader(GetGlobalShaderMap());
 	TShaderMapRef<FPostProcessBloomSetupPS_ES2<UseSunDof> > PixelShader(GetGlobalShaderMap());
@@ -142,7 +142,7 @@ static void BloomSetup_SetShader(FRHICommandList* RHICmdList, const FRenderingCo
 	PixelShader->SetPS(RHICmdList, Context);
 }
 
-void FRCPassPostProcessBloomSetupES2::SetShader(FRHICommandList* RHICmdList, const FRenderingCompositePassContext& Context)
+void FRCPassPostProcessBloomSetupES2::SetShader(FRHICommandList& RHICmdList, const FRenderingCompositePassContext& Context)
 {
 	const FSceneView& View = Context.View;
 	uint32 UseSun = Context.View.bLightShaftUse ? 1 : 0;
@@ -207,7 +207,7 @@ void FRCPassPostProcessBloomSetupES2::Process(FRenderingCompositePassContext& Co
 	RHISetDepthStencilState(TStaticDepthStencilState<false,CF_Always>::GetRHI());
 
 	//@todo-rco: RHIPacketList
-	FRHICommandList* RHICmdList = nullptr;
+	FRHICommandList& RHICmdList = FRHICommandList::GetNullRef();
 	SetShader(RHICmdList, Context);
 
 	TShaderMapRef<FPostProcessBloomSetupVS_ES2> VertexShader(GetGlobalShaderMap());
@@ -271,7 +271,7 @@ public:
 		PostprocessParameter.Bind(Initializer.ParameterMap);
 	}
 
-	void SetVS(FRHICommandList* RHICmdList, const FRenderingCompositePassContext& Context)
+	void SetVS(FRHICommandList& RHICmdList, const FRenderingCompositePassContext& Context)
 	{
 		const FVertexShaderRHIParamRef ShaderRHI = GetVertexShader();
 
@@ -317,7 +317,7 @@ public:
 		BloomThreshold.Bind(Initializer.ParameterMap, TEXT("BloomThreshold"));
 	}
 
-	void SetPS(FRHICommandList* RHICmdList, const FRenderingCompositePassContext& Context)
+	void SetPS(FRHICommandList& RHICmdList, const FRenderingCompositePassContext& Context)
 	{
 		const FPixelShaderRHIParamRef ShaderRHI = GetPixelShader();
 
@@ -347,7 +347,7 @@ IMPLEMENT_SHADER_TYPE(,FPostProcessBloomSetupSmallVS_ES2,TEXT("PostProcessMobile
 IMPLEMENT_SHADER_TYPE(,FPostProcessBloomSetupSmallPS_ES2,TEXT("PostProcessMobile"),TEXT("BloomSmallPS_ES2"),SF_Pixel);
 
 
-void FRCPassPostProcessBloomSetupSmallES2::SetShader(FRHICommandList* RHICmdList, const FRenderingCompositePassContext& Context)
+void FRCPassPostProcessBloomSetupSmallES2::SetShader(FRHICommandList& RHICmdList, const FRenderingCompositePassContext& Context)
 {
 	TShaderMapRef<FPostProcessBloomSetupSmallVS_ES2> VertexShader(GetGlobalShaderMap());
 	TShaderMapRef<FPostProcessBloomSetupSmallPS_ES2> PixelShader(GetGlobalShaderMap());
@@ -409,7 +409,7 @@ void FRCPassPostProcessBloomSetupSmallES2::Process(FRenderingCompositePassContex
 	RHISetDepthStencilState(TStaticDepthStencilState<false,CF_Always>::GetRHI());
 
 	//@todo-rco: RHIPacketList
-	FRHICommandList* RHICmdList = nullptr;
+	FRHICommandList& RHICmdList = FRHICommandList::GetNullRef();
 	SetShader(RHICmdList, Context);
 
 	TShaderMapRef<FPostProcessBloomSetupSmallVS_ES2> VertexShader(GetGlobalShaderMap());
@@ -477,7 +477,7 @@ public:
 		PostprocessParameter.Bind(Initializer.ParameterMap);
 	}
 
-	void SetPS(FRHICommandList* RHICmdList, const FRenderingCompositePassContext& Context)
+	void SetPS(FRHICommandList& RHICmdList, const FRenderingCompositePassContext& Context)
 	{
 		const FPixelShaderRHIParamRef ShaderRHI = GetPixelShader();
 		const FPostProcessSettings& Settings = Context.View.FinalPostProcessSettings;
@@ -518,7 +518,7 @@ public:
 		BloomDownScale.Bind(Initializer.ParameterMap, TEXT("BloomDownScale"));
 	}
 
-	void SetVS(FRHICommandList* RHICmdList, const FRenderingCompositePassContext& Context, float InScale)
+	void SetVS(FRHICommandList& RHICmdList, const FRenderingCompositePassContext& Context, float InScale)
 	{
 		const FVertexShaderRHIParamRef ShaderRHI = GetVertexShader();
 		FGlobalShader::SetParameters(RHICmdList, ShaderRHI, Context.View);
@@ -570,7 +570,7 @@ void FRCPassPostProcessBloomDownES2::Process(FRenderingCompositePassContext& Con
 	SetGlobalBoundShaderState(BoundShaderState, GFilterVertexDeclaration.VertexDeclarationRHI, *VertexShader, *PixelShader);
 
 	//@todo-rco: RHIPacketList
-	FRHICommandList* RHICmdList = nullptr;
+	FRHICommandList& RHICmdList = FRHICommandList::GetNullRef();
 	VertexShader->SetVS(RHICmdList, Context, Scale);
 	PixelShader->SetPS(RHICmdList, Context);
 
@@ -643,7 +643,7 @@ public:
 		TintB.Bind(Initializer.ParameterMap, TEXT("BloomTintB"));
 	}
 
-	void SetPS(FRHICommandList* RHICmdList, const FRenderingCompositePassContext& Context, FVector4& InTintA, FVector4& InTintB)
+	void SetPS(FRHICommandList& RHICmdList, const FRenderingCompositePassContext& Context, FVector4& InTintA, FVector4& InTintB)
 	{
 		const FPixelShaderRHIParamRef ShaderRHI = GetPixelShader();
 		const FPostProcessSettings& Settings = Context.View.FinalPostProcessSettings;
@@ -686,7 +686,7 @@ public:
 		BloomUpScales.Bind(Initializer.ParameterMap, TEXT("BloomUpScales"));
 	}
 
-	void SetVS(FRHICommandList* RHICmdList, const FRenderingCompositePassContext& Context, FVector2D InScale)
+	void SetVS(FRHICommandList& RHICmdList, const FRenderingCompositePassContext& Context, FVector2D InScale)
 	{
 		const FVertexShaderRHIParamRef ShaderRHI = GetVertexShader();
 		FGlobalShader::SetParameters(RHICmdList, ShaderRHI, Context.View);
@@ -738,7 +738,7 @@ void FRCPassPostProcessBloomUpES2::Process(FRenderingCompositePassContext& Conte
 	SetGlobalBoundShaderState(BoundShaderState, GFilterVertexDeclaration.VertexDeclarationRHI, *VertexShader, *PixelShader);
 
 	//@todo-rco: RHIPacketList
-	FRHICommandList* RHICmdList = nullptr;
+	FRHICommandList& RHICmdList = FRHICommandList::GetNullRef();
 	// The 1/8 factor is because bloom is using 8 taps in the filter.
 	VertexShader->SetVS(RHICmdList, Context, FVector2D(ScaleAB.X, ScaleAB.Y));
 	FVector4 TintAScaled = TintA * (1.0f/8.0f);
@@ -816,7 +816,7 @@ public:
 		SunColorApertureDiv2Parameter.Bind(Initializer.ParameterMap, TEXT("SunColorApertureDiv2"));
 	}
 
-	void SetPS(FRHICommandList* RHICmdList, const FRenderingCompositePassContext& Context)
+	void SetPS(FRHICommandList& RHICmdList, const FRenderingCompositePassContext& Context)
 	{
 		const FPixelShaderRHIParamRef ShaderRHI = GetPixelShader();
 		const FPostProcessSettings& Settings = Context.View.FinalPostProcessSettings;
@@ -877,7 +877,7 @@ public:
 		PostprocessParameter.Bind(Initializer.ParameterMap);
 	}
 
-	void SetVS(FRHICommandList* RHICmdList, const FRenderingCompositePassContext& Context)
+	void SetVS(FRHICommandList& RHICmdList, const FRenderingCompositePassContext& Context)
 	{
 		const FVertexShaderRHIParamRef ShaderRHI = GetVertexShader();
 		FGlobalShader::SetParameters(RHICmdList, ShaderRHI, Context.View);
@@ -895,7 +895,7 @@ public:
 IMPLEMENT_SHADER_TYPE(,FPostProcessSunMaskVS_ES2,TEXT("PostProcessMobile"),TEXT("SunMaskVS_ES2"),SF_Vertex);
 
 template <uint32 UseFetchSunDof>
-static void SunMask_SetShader(FRHICommandList* RHICmdList, const FRenderingCompositePassContext& Context)
+static void SunMask_SetShader(FRHICommandList& RHICmdList, const FRenderingCompositePassContext& Context)
 {
 	TShaderMapRef<FPostProcessSunMaskVS_ES2> VertexShader(GetGlobalShaderMap());
 	TShaderMapRef<FPostProcessSunMaskPS_ES2<UseFetchSunDof> > PixelShader(GetGlobalShaderMap());
@@ -908,7 +908,7 @@ static void SunMask_SetShader(FRHICommandList* RHICmdList, const FRenderingCompo
 	PixelShader->SetPS(RHICmdList, Context);
 }
 
-void FRCPassPostProcessSunMaskES2::SetShader(FRHICommandList* RHICmdList, const FRenderingCompositePassContext& Context)
+void FRCPassPostProcessSunMaskES2::SetShader(FRHICommandList& RHICmdList, const FRenderingCompositePassContext& Context)
 {
 	const FSceneView& View = Context.View;
 	uint32 UseSun = Context.View.bLightShaftUse ? 1 : 0;
@@ -951,7 +951,7 @@ void FRCPassPostProcessSunMaskES2::Process(FRenderingCompositePassContext& Conte
 	const FSceneView& View = Context.View;
 
 	//@todo-rco: RHIPacketList
-	FRHICommandList* RHICmdList = nullptr;
+	FRHICommandList& RHICmdList = FRHICommandList::GetNullRef();
 
 	TShaderMapRef<FPostProcessSunMaskVS_ES2> VertexShader(GetGlobalShaderMap());
 	if(bOnChip)
@@ -1064,7 +1064,7 @@ public:
 		PostprocessParameter.Bind(Initializer.ParameterMap);
 	}
 
-	void SetPS(FRHICommandList* RHICmdList, const FRenderingCompositePassContext& Context)
+	void SetPS(FRHICommandList& RHICmdList, const FRenderingCompositePassContext& Context)
 	{
 		const FPixelShaderRHIParamRef ShaderRHI = GetPixelShader();
 		const FPostProcessSettings& Settings = Context.View.FinalPostProcessSettings;
@@ -1107,7 +1107,7 @@ public:
 		LightShaftCenter.Bind(Initializer.ParameterMap, TEXT("LightShaftCenter"));
 	}
 
-	void SetVS(FRHICommandList* RHICmdList, const FRenderingCompositePassContext& Context)
+	void SetVS(FRHICommandList& RHICmdList, const FRenderingCompositePassContext& Context)
 	{
 		const FVertexShaderRHIParamRef ShaderRHI = GetVertexShader();
 		FGlobalShader::SetParameters(RHICmdList, ShaderRHI, Context.View);
@@ -1127,7 +1127,7 @@ public:
 IMPLEMENT_SHADER_TYPE(,FPostProcessSunAlphaVS_ES2,TEXT("PostProcessMobile"),TEXT("SunAlphaVS_ES2"),SF_Vertex);
 
 template <uint32 UseDof>
-static void SunAlpha_SetShader(FRHICommandList* RHICmdList, const FRenderingCompositePassContext& Context)
+static void SunAlpha_SetShader(FRHICommandList& RHICmdList, const FRenderingCompositePassContext& Context)
 {
 	TShaderMapRef<FPostProcessSunAlphaVS_ES2> VertexShader(GetGlobalShaderMap());
 	TShaderMapRef<FPostProcessSunAlphaPS_ES2<UseDof> > PixelShader(GetGlobalShaderMap());
@@ -1140,7 +1140,7 @@ static void SunAlpha_SetShader(FRHICommandList* RHICmdList, const FRenderingComp
 	PixelShader->SetPS(RHICmdList, Context);
 }
 
-void FRCPassPostProcessSunAlphaES2::SetShader(FRHICommandList* RHICmdList, const FRenderingCompositePassContext& Context)
+void FRCPassPostProcessSunAlphaES2::SetShader(FRHICommandList& RHICmdList, const FRenderingCompositePassContext& Context)
 {
 	if(Context.View.FinalPostProcessSettings.DepthOfFieldScale > 0.0f)
 	{
@@ -1178,7 +1178,7 @@ void FRCPassPostProcessSunAlphaES2::Process(FRenderingCompositePassContext& Cont
 	RHISetDepthStencilState(TStaticDepthStencilState<false,CF_Always>::GetRHI());
 
 	//@todo-rco: RHIPacketList
-	FRHICommandList* RHICmdList = nullptr;
+	FRHICommandList& RHICmdList = FRHICommandList::GetNullRef();
 	SetShader(RHICmdList, Context);
 
 	FIntPoint SrcDstSize = PrePostSourceViewportSize / 4;
@@ -1249,7 +1249,7 @@ public:
 		PostprocessParameter.Bind(Initializer.ParameterMap);
 	}
 
-	void SetPS(FRHICommandList* RHICmdList, const FRenderingCompositePassContext& Context)
+	void SetPS(FRHICommandList& RHICmdList, const FRenderingCompositePassContext& Context)
 	{
 		const FPixelShaderRHIParamRef ShaderRHI = GetPixelShader();
 		const FPostProcessSettings& Settings = Context.View.FinalPostProcessSettings;
@@ -1290,7 +1290,7 @@ public:
 		LightShaftCenter.Bind(Initializer.ParameterMap, TEXT("LightShaftCenter"));
 	}
 
-	void SetVS(FRHICommandList* RHICmdList, const FRenderingCompositePassContext& Context)
+	void SetVS(FRHICommandList& RHICmdList, const FRenderingCompositePassContext& Context)
 	{
 		const FVertexShaderRHIParamRef ShaderRHI = GetVertexShader();
 		FGlobalShader::SetParameters(RHICmdList, ShaderRHI, Context.View);
@@ -1343,7 +1343,7 @@ void FRCPassPostProcessSunBlurES2::Process(FRenderingCompositePassContext& Conte
 	SetGlobalBoundShaderState(BoundShaderState, GFilterVertexDeclaration.VertexDeclarationRHI, *VertexShader, *PixelShader);
 
 	//@todo-rco: RHIPacketList
-	FRHICommandList* RHICmdList = nullptr;
+	FRHICommandList& RHICmdList = FRHICommandList::GetNullRef();
 	VertexShader->SetVS(RHICmdList, Context);
 	PixelShader->SetPS(RHICmdList, Context);
 
@@ -1422,7 +1422,7 @@ public:
 		BloomColor.Bind(Initializer.ParameterMap, TEXT("BloomColor"));
 	}
 
-	void SetPS(FRHICommandList* RHICmdList, const FRenderingCompositePassContext& Context)
+	void SetPS(FRHICommandList& RHICmdList, const FRenderingCompositePassContext& Context)
 	{
 		const FPixelShaderRHIParamRef ShaderRHI = GetPixelShader();
 		const FPostProcessSettings& Settings = Context.View.FinalPostProcessSettings;
@@ -1482,7 +1482,7 @@ public:
 		LightShaftCenter.Bind(Initializer.ParameterMap, TEXT("LightShaftCenter"));
 	}
 
-	void SetVS(FRHICommandList* RHICmdList, const FRenderingCompositePassContext& Context)
+	void SetVS(FRHICommandList& RHICmdList, const FRenderingCompositePassContext& Context)
 	{
 		const FVertexShaderRHIParamRef ShaderRHI = GetVertexShader();
 		FGlobalShader::SetParameters(RHICmdList, ShaderRHI, Context.View);
@@ -1504,7 +1504,7 @@ IMPLEMENT_SHADER_TYPE(,FPostProcessSunMergeVS_ES2,TEXT("PostProcessMobile"),TEXT
 
 
 template <uint32 UseSunBloom>
-static void SunMerge_SetShader(FRHICommandList* RHICmdList, const FRenderingCompositePassContext& Context)
+static void SunMerge_SetShader(FRHICommandList& RHICmdList, const FRenderingCompositePassContext& Context)
 {
 	TShaderMapRef<FPostProcessSunMergeVS_ES2> VertexShader(GetGlobalShaderMap());
 	TShaderMapRef<FPostProcessSunMergePS_ES2<UseSunBloom> > PixelShader(GetGlobalShaderMap());
@@ -1517,7 +1517,7 @@ static void SunMerge_SetShader(FRHICommandList* RHICmdList, const FRenderingComp
 	PixelShader->SetPS(RHICmdList, Context);
 }
 
-void FRCPassPostProcessSunMergeES2::SetShader(FRHICommandList* RHICmdList, const FRenderingCompositePassContext& Context)
+void FRCPassPostProcessSunMergeES2::SetShader(FRHICommandList& RHICmdList, const FRenderingCompositePassContext& Context)
 {
 	const FSceneView& View = Context.View;
 	uint32 UseBloom = (View.FinalPostProcessSettings.BloomIntensity > 0.0f) ? 1 : 0;
@@ -1559,7 +1559,7 @@ void FRCPassPostProcessSunMergeES2::Process(FRenderingCompositePassContext& Cont
 	RHISetDepthStencilState(TStaticDepthStencilState<false,CF_Always>::GetRHI());
 
 	//@todo-rco: RHIPacketList
-	FRHICommandList* RHICmdList = nullptr;
+	FRHICommandList& RHICmdList = FRHICommandList::GetNullRef();
 	SetShader(RHICmdList, Context);
 
 	FIntPoint SrcDstSize = PrePostSourceViewportSize / 4;
@@ -1647,7 +1647,7 @@ public:
 		BloomColor2.Bind(Initializer.ParameterMap, TEXT("BloomColor2"));
 	}
 
-	void SetPS(FRHICommandList* RHICmdList, const FRenderingCompositePassContext& Context)
+	void SetPS(FRHICommandList& RHICmdList, const FRenderingCompositePassContext& Context)
 	{
 		const FPixelShaderRHIParamRef ShaderRHI = GetPixelShader();
 		const FPostProcessSettings& Settings = Context.View.FinalPostProcessSettings;
@@ -1698,7 +1698,7 @@ public:
 		PostprocessParameter.Bind(Initializer.ParameterMap);
 	}
 
-	void SetVS(FRHICommandList* RHICmdList, const FRenderingCompositePassContext& Context)
+	void SetVS(FRHICommandList& RHICmdList, const FRenderingCompositePassContext& Context)
 	{
 		const FVertexShaderRHIParamRef ShaderRHI = GetVertexShader();
 		FGlobalShader::SetParameters(RHICmdList, ShaderRHI, Context.View);
@@ -1715,7 +1715,7 @@ public:
 
 IMPLEMENT_SHADER_TYPE(,FPostProcessSunMergeSmallVS_ES2,TEXT("PostProcessMobile"),TEXT("SunMergeSmallVS_ES2"),SF_Vertex);
 
-void FRCPassPostProcessSunMergeSmallES2::SetShader(FRHICommandList* RHICmdList, const FRenderingCompositePassContext& Context)
+void FRCPassPostProcessSunMergeSmallES2::SetShader(FRHICommandList& RHICmdList, const FRenderingCompositePassContext& Context)
 {
 	TShaderMapRef<FPostProcessSunMergeSmallVS_ES2> VertexShader(GetGlobalShaderMap());
 	TShaderMapRef<FPostProcessSunMergeSmallPS_ES2> PixelShader(GetGlobalShaderMap());
@@ -1754,7 +1754,7 @@ void FRCPassPostProcessSunMergeSmallES2::Process(FRenderingCompositePassContext&
 	RHISetDepthStencilState(TStaticDepthStencilState<false,CF_Always>::GetRHI());
 
 	//@todo-rco: RHIPacketList
-	FRHICommandList* RHICmdList = nullptr;
+	FRHICommandList& RHICmdList = FRHICommandList::GetNullRef();
 	SetShader(RHICmdList, Context);
 
 	FIntPoint SrcDstSize = PrePostSourceViewportSize / 4;
@@ -1835,7 +1835,7 @@ public:
 		PostprocessParameter.Bind(Initializer.ParameterMap);
 	}
 
-	void SetVS(FRHICommandList* RHICmdList, const FRenderingCompositePassContext& Context)
+	void SetVS(FRHICommandList& RHICmdList, const FRenderingCompositePassContext& Context)
 	{
 		const FVertexShaderRHIParamRef ShaderRHI = GetVertexShader();
 
@@ -1880,7 +1880,7 @@ public:
 		PostprocessParameter.Bind(Initializer.ParameterMap);
 	}
 
-	void SetPS(FRHICommandList* RHICmdList, const FRenderingCompositePassContext& Context)
+	void SetPS(FRHICommandList& RHICmdList, const FRenderingCompositePassContext& Context)
 	{
 		const FPixelShaderRHIParamRef ShaderRHI = GetPixelShader();
 		const FPostProcessSettings& Settings = Context.View.FinalPostProcessSettings;
@@ -1904,7 +1904,7 @@ IMPLEMENT_SHADER_TYPE(template<>,FPostProcessDofDownPS_ES2_0,TEXT("PostProcessMo
 IMPLEMENT_SHADER_TYPE(template<>,FPostProcessDofDownPS_ES2_1,TEXT("PostProcessMobile"),TEXT("DofDownPS_ES2"),SF_Pixel);
 
 template <uint32 UseSun>
-static void DofDown_SetShader(FRHICommandList* RHICmdList, const FRenderingCompositePassContext& Context)
+static void DofDown_SetShader(FRHICommandList& RHICmdList, const FRenderingCompositePassContext& Context)
 {
 	TShaderMapRef<FPostProcessDofDownVS_ES2> VertexShader(GetGlobalShaderMap());
 	TShaderMapRef<FPostProcessDofDownPS_ES2<UseSun> > PixelShader(GetGlobalShaderMap());
@@ -1917,7 +1917,7 @@ static void DofDown_SetShader(FRHICommandList* RHICmdList, const FRenderingCompo
 	PixelShader->SetPS(RHICmdList, Context);
 }
 
-void FRCPassPostProcessDofDownES2::SetShader(FRHICommandList* RHICmdList, const FRenderingCompositePassContext& Context)
+void FRCPassPostProcessDofDownES2::SetShader(FRHICommandList& RHICmdList, const FRenderingCompositePassContext& Context)
 {
 	const FSceneView& View = Context.View;
 	if(Context.View.bLightShaftUse)
@@ -1979,7 +1979,7 @@ void FRCPassPostProcessDofDownES2::Process(FRenderingCompositePassContext& Conte
 	RHISetDepthStencilState(TStaticDepthStencilState<false,CF_Always>::GetRHI());
 
 	//@todo-rco: RHIPacketList
-	FRHICommandList* RHICmdList = nullptr;
+	FRHICommandList& RHICmdList = FRHICommandList::GetNullRef();
 	SetShader(RHICmdList, Context);
 
 	TShaderMapRef<FPostProcessDofDownVS_ES2> VertexShader(GetGlobalShaderMap());
@@ -2043,7 +2043,7 @@ public:
 		PostprocessParameter.Bind(Initializer.ParameterMap);
 	}
 
-	void SetVS(FRHICommandList* RHICmdList, const FRenderingCompositePassContext& Context)
+	void SetVS(FRHICommandList& RHICmdList, const FRenderingCompositePassContext& Context)
 	{
 		const FVertexShaderRHIParamRef ShaderRHI = GetVertexShader();
 
@@ -2088,7 +2088,7 @@ public:
 		PostprocessParameter.Bind(Initializer.ParameterMap);
 	}
 
-	void SetPS(FRHICommandList* RHICmdList, const FRenderingCompositePassContext& Context)
+	void SetPS(FRHICommandList& RHICmdList, const FRenderingCompositePassContext& Context)
 	{
 		const FPixelShaderRHIParamRef ShaderRHI = GetPixelShader();
 		const FPostProcessSettings& Settings = Context.View.FinalPostProcessSettings;
@@ -2112,7 +2112,7 @@ IMPLEMENT_SHADER_TYPE(template<>,FPostProcessDofNearPS_ES2_0,TEXT("PostProcessMo
 IMPLEMENT_SHADER_TYPE(template<>,FPostProcessDofNearPS_ES2_1,TEXT("PostProcessMobile"),TEXT("DofNearPS_ES2"),SF_Pixel);
 
 template <uint32 UseSun>
-static void DofNear_SetShader(FRHICommandList* RHICmdList, const FRenderingCompositePassContext& Context)
+static void DofNear_SetShader(FRHICommandList& RHICmdList, const FRenderingCompositePassContext& Context)
 {
 	TShaderMapRef<FPostProcessDofNearVS_ES2> VertexShader(GetGlobalShaderMap());
 	TShaderMapRef<FPostProcessDofNearPS_ES2<UseSun> > PixelShader(GetGlobalShaderMap());
@@ -2125,7 +2125,7 @@ static void DofNear_SetShader(FRHICommandList* RHICmdList, const FRenderingCompo
 	PixelShader->SetPS(RHICmdList, Context);
 }
 
-void FRCPassPostProcessDofNearES2::SetShader(FRHICommandList* RHICmdList, const FRenderingCompositePassContext& Context)
+void FRCPassPostProcessDofNearES2::SetShader(FRHICommandList& RHICmdList, const FRenderingCompositePassContext& Context)
 {
 	const FSceneView& View = Context.View;
 	if(Context.View.bLightShaftUse)
@@ -2164,7 +2164,7 @@ void FRCPassPostProcessDofNearES2::Process(FRenderingCompositePassContext& Conte
 	RHISetDepthStencilState(TStaticDepthStencilState<false,CF_Always>::GetRHI());
 
 	//@todo-rco: RHIPacketList
-	FRHICommandList* RHICmdList = nullptr;
+	FRHICommandList& RHICmdList = FRHICommandList::GetNullRef();
 	SetShader(RHICmdList, Context);
 
 	FIntPoint SrcDstSize = PrePostSourceViewportSize / 4;
@@ -2232,7 +2232,7 @@ public:
 		PostprocessParameter.Bind(Initializer.ParameterMap);
 	}
 
-	void SetPS(FRHICommandList* RHICmdList, const FRenderingCompositePassContext& Context)
+	void SetPS(FRHICommandList& RHICmdList, const FRenderingCompositePassContext& Context)
 	{
 		const FPixelShaderRHIParamRef ShaderRHI = GetPixelShader();
 		const FPostProcessSettings& Settings = Context.View.FinalPostProcessSettings;
@@ -2271,7 +2271,7 @@ public:
 		PostprocessParameter.Bind(Initializer.ParameterMap);
 	}
 
-	void SetVS(FRHICommandList* RHICmdList, const FRenderingCompositePassContext& Context)
+	void SetVS(FRHICommandList& RHICmdList, const FRenderingCompositePassContext& Context)
 	{
 		const FVertexShaderRHIParamRef ShaderRHI = GetVertexShader();
 		FGlobalShader::SetParameters(RHICmdList, ShaderRHI, Context.View);
@@ -2321,7 +2321,7 @@ void FRCPassPostProcessDofBlurES2::Process(FRenderingCompositePassContext& Conte
 	SetGlobalBoundShaderState(BoundShaderState, GFilterVertexDeclaration.VertexDeclarationRHI, *VertexShader, *PixelShader);
 
 	//@todo-rco: RHIPacketList
-	FRHICommandList* RHICmdList = nullptr;
+	FRHICommandList& RHICmdList = FRHICommandList::GetNullRef();
 	VertexShader->SetVS(RHICmdList, Context);
 	PixelShader->SetPS(RHICmdList, Context);
 
@@ -2391,7 +2391,7 @@ public:
 		PostprocessParameter.Bind(Initializer.ParameterMap);
 	}
 
-	void SetPS(FRHICommandList* RHICmdList, const FRenderingCompositePassContext& Context)
+	void SetPS(FRHICommandList& RHICmdList, const FRenderingCompositePassContext& Context)
 	{
 		const FPixelShaderRHIParamRef ShaderRHI = GetPixelShader();
 		const FPostProcessSettings& Settings = Context.View.FinalPostProcessSettings;
@@ -2431,7 +2431,7 @@ public:
 		PostprocessParameter.Bind(Initializer.ParameterMap);
 	}
 
-	void SetVS(FRHICommandList* RHICmdList, const FRenderingCompositePassContext& Context)
+	void SetVS(FRHICommandList& RHICmdList, const FRenderingCompositePassContext& Context)
 	{
 		const FVertexShaderRHIParamRef ShaderRHI = GetVertexShader();
 		FGlobalShader::SetParameters(RHICmdList, ShaderRHI, Context.View);
@@ -2450,7 +2450,7 @@ IMPLEMENT_SHADER_TYPE(,FPostProcessSunAvgVS_ES2,TEXT("PostProcessMobile"),TEXT("
 
 
 
-static void SunAvg_SetShader(FRHICommandList* RHICmdList, const FRenderingCompositePassContext& Context)
+static void SunAvg_SetShader(FRHICommandList& RHICmdList, const FRenderingCompositePassContext& Context)
 {
 	TShaderMapRef<FPostProcessSunAvgVS_ES2> VertexShader(GetGlobalShaderMap());
 	TShaderMapRef<FPostProcessSunAvgPS_ES2> PixelShader(GetGlobalShaderMap());
@@ -2463,7 +2463,7 @@ static void SunAvg_SetShader(FRHICommandList* RHICmdList, const FRenderingCompos
 	PixelShader->SetPS(RHICmdList, Context);
 }
 
-void FRCPassPostProcessSunAvgES2::SetShader(FRHICommandList* RHICmdList, const FRenderingCompositePassContext& Context)
+void FRCPassPostProcessSunAvgES2::SetShader(FRHICommandList& RHICmdList, const FRenderingCompositePassContext& Context)
 {
 	SunAvg_SetShader(RHICmdList, Context);
 }
@@ -2494,7 +2494,7 @@ void FRCPassPostProcessSunAvgES2::Process(FRenderingCompositePassContext& Contex
 	RHISetDepthStencilState(TStaticDepthStencilState<false,CF_Always>::GetRHI());
 
 	//@todo-rco: RHIPacketList
-	FRHICommandList* RHICmdList = nullptr;
+	FRHICommandList& RHICmdList = FRHICommandList::GetNullRef();
 	SetShader(RHICmdList, Context);
 
 	FIntPoint SrcDstSize = PrePostSourceViewportSize / 4;
@@ -2565,7 +2565,7 @@ public:
 		AaBlendAmount.Bind(Initializer.ParameterMap, TEXT("AaBlendAmount"));
 	}
 
-	void SetPS(FRHICommandList* RHICmdList, const FRenderingCompositePassContext& Context)
+	void SetPS(FRHICommandList& RHICmdList, const FRenderingCompositePassContext& Context)
 	{
 		const FPixelShaderRHIParamRef ShaderRHI = GetPixelShader();
 		const FPostProcessSettings& Settings = Context.View.FinalPostProcessSettings;
@@ -2674,7 +2674,7 @@ public:
 		PostprocessParameter.Bind(Initializer.ParameterMap);
 	}
 
-	void SetVS(FRHICommandList* RHICmdList, const FRenderingCompositePassContext& Context)
+	void SetVS(FRHICommandList& RHICmdList, const FRenderingCompositePassContext& Context)
 	{
 		const FVertexShaderRHIParamRef ShaderRHI = GetVertexShader();
 		FGlobalShader::SetParameters(RHICmdList, ShaderRHI, Context.View);
@@ -2693,7 +2693,7 @@ IMPLEMENT_SHADER_TYPE(,FPostProcessAaVS_ES2,TEXT("PostProcessMobile"),TEXT("AaVS
 
 
 
-static void Aa_SetShader(FRHICommandList* RHICmdList, const FRenderingCompositePassContext& Context)
+static void Aa_SetShader(FRHICommandList& RHICmdList, const FRenderingCompositePassContext& Context)
 {
 	TShaderMapRef<FPostProcessAaVS_ES2> VertexShader(GetGlobalShaderMap());
 	TShaderMapRef<FPostProcessAaPS_ES2> PixelShader(GetGlobalShaderMap());
@@ -2706,7 +2706,7 @@ static void Aa_SetShader(FRHICommandList* RHICmdList, const FRenderingCompositeP
 	PixelShader->SetPS(RHICmdList, Context);
 }
 
-void FRCPassPostProcessAaES2::SetShader(FRHICommandList* RHICmdList, const FRenderingCompositePassContext& Context)
+void FRCPassPostProcessAaES2::SetShader(FRHICommandList& RHICmdList, const FRenderingCompositePassContext& Context)
 {
 	Aa_SetShader(RHICmdList, Context);
 }
@@ -2737,7 +2737,7 @@ void FRCPassPostProcessAaES2::Process(FRenderingCompositePassContext& Context)
 	RHISetDepthStencilState(TStaticDepthStencilState<false,CF_Always>::GetRHI());
 
 	//@todo-rco: RHIPacketList
-	FRHICommandList* RHICmdList = nullptr;
+	FRHICommandList& RHICmdList = FRHICommandList::GetNullRef();
 	SetShader(RHICmdList, Context);
 
 	FIntPoint SrcDstSize = PrePostSourceViewportSize;
diff --git a/Engine/Source/Runtime/Renderer/Private/PostProcess/PostProcessMobile.h b/Engine/Source/Runtime/Renderer/Private/PostProcess/PostProcessMobile.h
--- a/Engine/Source/Runtime/Renderer/Private/PostProcess/PostProcessMobile.h
+++ b/Engine/Source/Runtime/Renderer/Private/PostProcess/PostProcessMobile.h
@@ -18,7 +18,7 @@ public:
 private:
 	FIntPoint PrePostSourceViewportSize;
 	bool bUsedFramebufferFetch;
-	void SetShader(FRHICommandList* RHICmdList, const FRenderingCompositePassContext& Context);
+	void SetShader(FRHICommandList& RHICmdList, const FRenderingCompositePassContext& Context);
 };
 
 class FRCPassPostProcessBloomSetupSmallES2 : public TRenderingCompositePassBase<1, 1>
@@ -31,7 +31,7 @@ public:
 private:
 	FIntPoint PrePostSourceViewportSize;
 	bool bUsedFramebufferFetch;
-	void SetShader(FRHICommandList* RHICmdList, const FRenderingCompositePassContext& Context);
+	void SetShader(FRHICommandList& RHICmdList, const FRenderingCompositePassContext& Context);
 };
 
 class FRCPassPostProcessDofNearES2 : public TRenderingCompositePassBase<1, 1>
@@ -43,7 +43,7 @@ public:
 	virtual void Release() { delete this; }
 private:
 	FIntPoint PrePostSourceViewportSize;
-	void SetShader(FRHICommandList* RHICmdList, const FRenderingCompositePassContext& Context);
+	void SetShader(FRHICommandList& RHICmdList, const FRenderingCompositePassContext& Context);
 };
 
 class FRCPassPostProcessDofDownES2 : public TRenderingCompositePassBase<2, 1>
@@ -56,7 +56,7 @@ public:
 private:
 	FIntPoint PrePostSourceViewportSize;
 	bool bUsedFramebufferFetch;
-	void SetShader(FRHICommandList* RHICmdList, const FRenderingCompositePassContext& Context);
+	void SetShader(FRHICommandList& RHICmdList, const FRenderingCompositePassContext& Context);
 };
 
 class FRCPassPostProcessDofBlurES2 : public TRenderingCompositePassBase<2, 1>
@@ -106,7 +106,7 @@ public:
 private:
 	FIntPoint PrePostSourceViewportSize;
 	bool bOnChip;
-	void SetShader(FRHICommandList* RHICmdList, const FRenderingCompositePassContext& Context);
+	void SetShader(FRHICommandList& RHICmdList, const FRenderingCompositePassContext& Context);
 };
 
 class FRCPassPostProcessSunAlphaES2 : public TRenderingCompositePassBase<1, 1>
@@ -118,7 +118,7 @@ public:
 	virtual void Release() { delete this; }
 private:
 	FIntPoint PrePostSourceViewportSize;
-	void SetShader(FRHICommandList* RHICmdList, const FRenderingCompositePassContext& Context);
+	void SetShader(FRHICommandList& RHICmdList, const FRenderingCompositePassContext& Context);
 };
 
 class FRCPassPostProcessSunBlurES2 : public TRenderingCompositePassBase<1, 1>
@@ -141,7 +141,7 @@ public:
 	virtual void Release() { delete this; }
 private:
 	FIntPoint PrePostSourceViewportSize;
-	void SetShader(FRHICommandList* RHICmdList, const FRenderingCompositePassContext& Context);
+	void SetShader(FRHICommandList& RHICmdList, const FRenderingCompositePassContext& Context);
 };
 
 class FRCPassPostProcessSunMergeSmallES2 : public TRenderingCompositePassBase<2, 1>
@@ -153,7 +153,7 @@ public:
 	virtual void Release() { delete this; }
 private:
 	FIntPoint PrePostSourceViewportSize;
-	void SetShader(FRHICommandList* RHICmdList, const FRenderingCompositePassContext& Context);
+	void SetShader(FRHICommandList& RHICmdList, const FRenderingCompositePassContext& Context);
 };
 
 class FRCPassPostProcessSunAvgES2 : public TRenderingCompositePassBase<2, 1>
@@ -165,7 +165,7 @@ public:
 	virtual void Release() { delete this; }
 private:
 	FIntPoint PrePostSourceViewportSize;
-	void SetShader(FRHICommandList* RHICmdList, const FRenderingCompositePassContext& Context);
+	void SetShader(FRHICommandList& RHICmdList, const FRenderingCompositePassContext& Context);
 };
 
 class FRCPassPostProcessAaES2 : public TRenderingCompositePassBase<2, 1>
@@ -177,6 +177,6 @@ public:
 	virtual void Release() { delete this; }
 private:
 	FIntPoint PrePostSourceViewportSize;
-	void SetShader(FRHICommandList* RHICmdList, const FRenderingCompositePassContext& Context);
+	void SetShader(FRHICommandList& RHICmdList, const FRenderingCompositePassContext& Context);
 };
 
diff --git a/Engine/Source/Runtime/Renderer/Private/PostProcess/PostProcessMorpheus.cpp b/Engine/Source/Runtime/Renderer/Private/PostProcess/PostProcessMorpheus.cpp
--- a/Engine/Source/Runtime/Renderer/Private/PostProcess/PostProcessMorpheus.cpp
+++ b/Engine/Source/Runtime/Renderer/Private/PostProcess/PostProcessMorpheus.cpp
@@ -111,7 +111,7 @@ public:
 	}
 
 
-	void SetPS(FRHICommandList* RHICmdList, const FRenderingCompositePassContext& Context, FIntRect SrcRect, FIntPoint SrcBufferSize, EStereoscopicPass StereoPass, FMatrix& QuadTexTransform)
+	void SetPS(FRHICommandList& RHICmdList, const FRenderingCompositePassContext& Context, FIntRect SrcRect, FIntPoint SrcBufferSize, EStereoscopicPass StereoPass, FMatrix& QuadTexTransform)
 	{
 		const FPixelShaderRHIParamRef ShaderRHI = GetPixelShader();
 		
@@ -254,12 +254,12 @@ class FPostProcessMorpheusVS : public FGlobalShader
 	FPostProcessMorpheusVS() {}
 
 	/** to have a similar interface as all other shaders */
-	void SetParameters(FRHICommandList* RHICmdList, const FRenderingCompositePassContext& Context)
+	void SetParameters(FRHICommandList& RHICmdList, const FRenderingCompositePassContext& Context)
 	{
 		FGlobalShader::SetParameters(RHICmdList, GetVertexShader(), Context.View);
 	}
 
-	void SetParameters(FRHICommandList* RHICmdList, const FSceneView& View)
+	void SetParameters(FRHICommandList& RHICmdList, const FSceneView& View)
 	{
 		FGlobalShader::SetParameters(RHICmdList, GetVertexShader(), View);
 	}
@@ -316,7 +316,7 @@ void FRCPassPostProcessMorpheus::Process(FRenderingCompositePassContext& Context
     FMatrix QuadTexTransform;
     FMatrix QuadPosTransform = FMatrix::Identity;
 	//@todo-rco: RHIPacketList
-	FRHICommandList* RHICmdList = nullptr;
+	FRHICommandList& RHICmdList = FRHICommandList::GetNullRef();
 	PixelShader->SetPS(RHICmdList, Context, SrcRect, SrcSize, View.StereoPass, QuadTexTransform);
 
 	// Draw a quad mapping scene color to the view's render target
diff --git a/Engine/Source/Runtime/Renderer/Private/PostProcess/PostProcessMotionBlur.cpp b/Engine/Source/Runtime/Renderer/Private/PostProcess/PostProcessMotionBlur.cpp
--- a/Engine/Source/Runtime/Renderer/Private/PostProcess/PostProcessMotionBlur.cpp
+++ b/Engine/Source/Runtime/Renderer/Private/PostProcess/PostProcessMotionBlur.cpp
@@ -47,7 +47,7 @@ public:
 		return bShaderHasOutdatedParameters;
 	}
 
-	void SetParameters(FRHICommandList* RHICmdList, const FRenderingCompositePassContext& Context)
+	void SetParameters(FRHICommandList& RHICmdList, const FRenderingCompositePassContext& Context)
 	{
 		const FPixelShaderRHIParamRef ShaderRHI = GetPixelShader();
 
@@ -128,7 +128,7 @@ void FRCPassPostProcessMotionBlurSetup::Process(FRenderingCompositePassContext&
 	RHISetDepthStencilState(TStaticDepthStencilState<false,CF_Always>::GetRHI());
 
 	//@todo-rco: RHIPacketList
-	FRHICommandList* RHICmdList = nullptr;
+	FRHICommandList& RHICmdList = FRHICommandList::GetNullRef();
 
 	TShaderMapRef<FPostProcessVS> VertexShader(GetGlobalShaderMap());
 
@@ -254,7 +254,7 @@ public:
 		return bShaderHasOutdatedParameters;
 	}
 
-	void SetParameters(FRHICommandList* RHICmdList, const FRenderingCompositePassContext& Context)
+	void SetParameters(FRHICommandList& RHICmdList, const FRenderingCompositePassContext& Context)
 	{
 		const FPixelShaderRHIParamRef ShaderRHI = GetPixelShader();
 
@@ -367,7 +367,7 @@ VARIATION1(0)			VARIATION1(1)			VARIATION1(2)			VARIATION1(3)			VARIATION1(4)
 
 // @param Quality 0: visualize, 1:low, 2:medium, 3:high, 4:very high
 template <uint32 Quality>
-static void SetMotionBlurShaderTempl(FRHICommandList* RHICmdList, const FRenderingCompositePassContext& Context)
+static void SetMotionBlurShaderTempl(FRHICommandList& RHICmdList, const FRenderingCompositePassContext& Context)
 {
 	TShaderMapRef<FPostProcessVS> VertexShader(GetGlobalShaderMap());
 	TShaderMapRef<FPostProcessMotionBlurPS<Quality> > PixelShader(GetGlobalShaderMap());
@@ -430,7 +430,7 @@ void FRCPassPostProcessMotionBlur::Process(FRenderingCompositePassContext& Conte
 	RHISetDepthStencilState(TStaticDepthStencilState<false,CF_Always>::GetRHI());
 
 	//@todo-rco: RHIPacketList
-	FRHICommandList* RHICmdList = nullptr;
+	FRHICommandList& RHICmdList = FRHICommandList::GetNullRef();
 	if(Quality == 1)
 	{
 		SetMotionBlurShaderTempl<1>(RHICmdList, Context);
@@ -511,7 +511,7 @@ public:
 		return bShaderHasOutdatedParameters;
 	}
 
-	void SetParameters(FRHICommandList* RHICmdList, const FRenderingCompositePassContext& Context)
+	void SetParameters(FRHICommandList& RHICmdList, const FRenderingCompositePassContext& Context)
 	{
 		const FPixelShaderRHIParamRef ShaderRHI = GetPixelShader();
 
@@ -573,7 +573,7 @@ void FRCPassPostProcessMotionBlurRecombine::Process(FRenderingCompositePassConte
 	SetGlobalBoundShaderState(BoundShaderState, GFilterVertexDeclaration.VertexDeclarationRHI, *VertexShader, *PixelShader);
 
 	//@todo-rco: RHIPacketList
-	FRHICommandList* RHICmdList = nullptr;
+	FRHICommandList& RHICmdList = FRHICommandList::GetNullRef();
 	VertexShader->SetParameters(RHICmdList, Context);
 	PixelShader->SetParameters(RHICmdList, Context);
 
@@ -648,7 +648,7 @@ void FRCPassPostProcessVisualizeMotionBlur::Process(FRenderingCompositePassConte
 
 	// Quality 0: visualize
 	//@todo-rco: RHIPacketList
-	FRHICommandList* RHICmdList = nullptr;
+	FRHICommandList& RHICmdList = FRHICommandList::GetNullRef();
 	SetMotionBlurShaderTempl<0>(RHICmdList, Context);
 
 	// Draw a quad mapping scene color to the view's render target
diff --git a/Engine/Source/Runtime/Renderer/Private/PostProcess/PostProcessNoiseBlur.cpp b/Engine/Source/Runtime/Renderer/Private/PostProcess/PostProcessNoiseBlur.cpp
--- a/Engine/Source/Runtime/Renderer/Private/PostProcess/PostProcessNoiseBlur.cpp
+++ b/Engine/Source/Runtime/Renderer/Private/PostProcess/PostProcessNoiseBlur.cpp
@@ -53,7 +53,7 @@ public:
 		return bShaderHasOutdatedParameters;
 	}
 
-	void SetParameters(FRHICommandList* RHICmdList, const FRenderingCompositePassContext& Context, float InRadius)
+	void SetParameters(FRHICommandList& RHICmdList, const FRenderingCompositePassContext& Context, float InRadius)
 	{
 		const FPixelShaderRHIParamRef ShaderRHI = GetPixelShader();
 
@@ -95,7 +95,7 @@ FRCPassPostProcessNoiseBlur::FRCPassPostProcessNoiseBlur(float InRadius, EPixelF
 
 
 template <uint32 Method>
-void SetNoiseBlurShader(FRHICommandList* RHICmdList, const FRenderingCompositePassContext& Context, float InRadius)
+void SetNoiseBlurShader(FRHICommandList& RHICmdList, const FRenderingCompositePassContext& Context, float InRadius)
 {
 	TShaderMapRef<FPostProcessVS> VertexShader(GetGlobalShaderMap());
 	TShaderMapRef<FPostProcessNoiseBlurPS<Method> > PixelShader(GetGlobalShaderMap());
@@ -149,7 +149,7 @@ void FRCPassPostProcessNoiseBlur::Process(FRenderingCompositePassContext& Contex
 	RHISetDepthStencilState(TStaticDepthStencilState<false,CF_Always>::GetRHI());
 	
 	//@todo-rco: RHIPacketList
-	FRHICommandList* RHICmdList = nullptr;
+	FRHICommandList& RHICmdList = FRHICommandList::GetNullRef();
 	if(Quality == 0)
 	{
 		SetNoiseBlurShader<0>(RHICmdList, Context, Radius);
diff --git a/Engine/Source/Runtime/Renderer/Private/PostProcess/PostProcessSceneColorFringe.cpp b/Engine/Source/Runtime/Renderer/Private/PostProcess/PostProcessSceneColorFringe.cpp
--- a/Engine/Source/Runtime/Renderer/Private/PostProcess/PostProcessSceneColorFringe.cpp
+++ b/Engine/Source/Runtime/Renderer/Private/PostProcess/PostProcessSceneColorFringe.cpp
@@ -49,7 +49,7 @@ public:
 		return bShaderHasOutdatedParameters;
 	}
 
-	void SetParameters(FRHICommandList* RHICmdList, const FRenderingCompositePassContext& Context)
+	void SetParameters(FRHICommandList& RHICmdList, const FRenderingCompositePassContext& Context)
 	{
 		const FVertexShaderRHIParamRef ShaderRHI = GetVertexShader();
 
@@ -114,7 +114,7 @@ public:
 		return bShaderHasOutdatedParameters;
 	}
 
-	void SetParameters(FRHICommandList* RHICmdList, const FRenderingCompositePassContext& Context)
+	void SetParameters(FRHICommandList& RHICmdList, const FRenderingCompositePassContext& Context)
 	{
 		const FPixelShaderRHIParamRef ShaderRHI = GetPixelShader();
 
@@ -174,7 +174,7 @@ void FRCPassPostProcessSceneColorFringe::Process(FRenderingCompositePassContext&
 	SetGlobalBoundShaderState(BoundShaderState, GFilterVertexDeclaration.VertexDeclarationRHI, *VertexShader, *PixelShader);
 
 	//@todo-rco: RHIPacketList
-	FRHICommandList* RHICmdList = nullptr;
+	FRHICommandList& RHICmdList = FRHICommandList::GetNullRef();
 	PixelShader->SetParameters(RHICmdList, Context);
 	VertexShader->SetParameters(RHICmdList, Context);
 	
diff --git a/Engine/Source/Runtime/Renderer/Private/PostProcess/PostProcessSelectionOutline.cpp b/Engine/Source/Runtime/Renderer/Private/PostProcess/PostProcessSelectionOutline.cpp
--- a/Engine/Source/Runtime/Renderer/Private/PostProcess/PostProcessSelectionOutline.cpp
+++ b/Engine/Source/Runtime/Renderer/Private/PostProcess/PostProcessSelectionOutline.cpp
@@ -203,7 +203,7 @@ public:
 		EditorPrimitivesStencil.Bind(Initializer.ParameterMap,TEXT("EditorPrimitivesStencil"));
 	}
 
-	void SetPS(FRHICommandList* RHICmdList, const FRenderingCompositePassContext& Context)
+	void SetPS(FRHICommandList& RHICmdList, const FRenderingCompositePassContext& Context)
 	{
 		const FPixelShaderRHIParamRef ShaderRHI = GetPixelShader();
 
@@ -312,7 +312,7 @@ VARIATION1(1)			VARIATION1(2)			VARIATION1(4)			VARIATION1(8)
 
 
 template <uint32 MSAASampleCount>
-static void SetSelectionOutlineShaderTempl(FRHICommandList* RHICmdList, const FRenderingCompositePassContext& Context)
+static void SetSelectionOutlineShaderTempl(FRHICommandList& RHICmdList, const FRenderingCompositePassContext& Context)
 {
 	TShaderMapRef<FPostProcessVS> VertexShader(GetGlobalShaderMap());
 	TShaderMapRef<FPostProcessSelectionOutlinePS<MSAASampleCount> > PixelShader(GetGlobalShaderMap());
@@ -335,7 +335,7 @@ void FRCPassPostProcessSelectionOutline::Process(FRenderingCompositePassContext&
 		return;
 	}
 	//@todo-rco: RHIPacketList
-	FRHICommandList* RHICmdList = nullptr;
+	FRHICommandList& RHICmdList = FRHICommandList::GetNullRef();
 
 	const FSceneView& View = Context.View;
 	FIntRect ViewRect = View.ViewRect;
diff --git a/Engine/Source/Runtime/Renderer/Private/PostProcess/PostProcessSubsurface.cpp b/Engine/Source/Runtime/Renderer/Private/PostProcess/PostProcessSubsurface.cpp
--- a/Engine/Source/Runtime/Renderer/Private/PostProcess/PostProcessSubsurface.cpp
+++ b/Engine/Source/Runtime/Renderer/Private/PostProcess/PostProcessSubsurface.cpp
@@ -44,7 +44,7 @@ public:
 		DeferredParameters.Bind(Initializer.ParameterMap);
 	}
 
-	void SetParameters(FRHICommandList* RHICmdList, const FRenderingCompositePassContext& Context)
+	void SetParameters(FRHICommandList& RHICmdList, const FRenderingCompositePassContext& Context)
 	{
 		const FFinalPostProcessSettings& Settings = Context.View.FinalPostProcessSettings;
 		const FPixelShaderRHIParamRef ShaderRHI = GetPixelShader();
@@ -83,7 +83,7 @@ public:
 
 
 template <uint32 SpecularCorrection>
-void SetSubsurfaceSetupShader(FRHICommandList* RHICmdList, const FRenderingCompositePassContext& Context)
+void SetSubsurfaceSetupShader(FRHICommandList& RHICmdList, const FRenderingCompositePassContext& Context)
 {
 	TShaderMapRef<FPostProcessVS> VertexShader(GetGlobalShaderMap());
 	TShaderMapRef<FPostProcessSubsurfaceSetupPS<SpecularCorrection> > PixelShader(GetGlobalShaderMap());
@@ -147,7 +147,7 @@ void FRCPassPostProcessSubsurfaceSetup::Process(FRenderingCompositePassContext&
 	const FSceneRenderTargetItem& DestRenderTarget = PassOutputs[0].RequestSurface(Context);
 
 	//@todo-rco: RHIPacketList
-	FRHICommandList* RHICmdList = nullptr;
+	FRHICommandList& RHICmdList = FRHICommandList::GetNullRef();
 
 	// Set the view family's render target/viewport.
 	RHISetRenderTarget(DestRenderTarget.TargetableTexture, FTextureRHIRef());	
@@ -244,7 +244,7 @@ public:
 		return bShaderHasOutdatedParameters;
 	}
 
-	void SetParameters(FRHICommandList* RHICmdList, const FRenderingCompositePassContext& Context, float InRadius)
+	void SetParameters(FRHICommandList& RHICmdList, const FRenderingCompositePassContext& Context, float InRadius)
 	{
 		const FPixelShaderRHIParamRef ShaderRHI = GetPixelShader();
 
@@ -285,7 +285,7 @@ FRCPassPostProcessSubsurface::FRCPassPostProcessSubsurface(uint32 InPass, float
 
 
 template <uint32 Method>
-void SetSubsurfaceShader(FRHICommandList* RHICmdList, const FRenderingCompositePassContext& Context, float InRadius)
+void SetSubsurfaceShader(FRHICommandList& RHICmdList, const FRenderingCompositePassContext& Context, float InRadius)
 {
 	TShaderMapRef<FPostProcessVS> VertexShader(GetGlobalShaderMap());
 	TShaderMapRef<FPostProcessSubsurfacePS<Method> > PixelShader(GetGlobalShaderMap());
@@ -336,7 +336,7 @@ void FRCPassPostProcessSubsurface::Process(FRenderingCompositePassContext& Conte
 	RHISetRasterizerState(TStaticRasterizerState<>::GetRHI());
 	RHISetDepthStencilState(TStaticDepthStencilState<false,CF_Always>::GetRHI());
 	//@todo-rco: RHIPacketList
-	FRHICommandList* RHICmdList = nullptr;
+	FRHICommandList& RHICmdList = FRHICommandList::GetNullRef();
 
 	TShaderMapRef<FPostProcessVS> VertexShader(GetGlobalShaderMap());
 
diff --git a/Engine/Source/Runtime/Renderer/Private/PostProcess/PostProcessTemporalAA.cpp b/Engine/Source/Runtime/Renderer/Private/PostProcess/PostProcessTemporalAA.cpp
--- a/Engine/Source/Runtime/Renderer/Private/PostProcess/PostProcessTemporalAA.cpp
+++ b/Engine/Source/Runtime/Renderer/Private/PostProcess/PostProcessTemporalAA.cpp
@@ -88,7 +88,7 @@ public:
 		return bShaderHasOutdatedParameters;
 	}
 
-	void SetParameters(FRHICommandList* RHICmdList, const FRenderingCompositePassContext& Context)
+	void SetParameters(FRHICommandList& RHICmdList, const FRenderingCompositePassContext& Context)
 	{
 		const FPixelShaderRHIParamRef ShaderRHI = GetPixelShader();
 
@@ -235,7 +235,7 @@ void FRCPassPostProcessSSRTemporalAA::Process(FRenderingCompositePassContext& Co
 	SetGlobalBoundShaderState(BoundShaderState, GFilterVertexDeclaration.VertexDeclarationRHI, *VertexShader, *PixelShader);
 
 	//@todo-rco: RHIPacketList
-	FRHICommandList* RHICmdList = nullptr;
+	FRHICommandList& RHICmdList = FRHICommandList::GetNullRef();
 	VertexShader->SetVS(RHICmdList, Context);
 	PixelShader->SetParameters(RHICmdList, Context);
 
@@ -312,7 +312,7 @@ void FRCPassPostProcessDOFTemporalAA::Process(FRenderingCompositePassContext& Co
 	SetGlobalBoundShaderState(BoundShaderState, GFilterVertexDeclaration.VertexDeclarationRHI, *VertexShader, *PixelShader);
 
 	//@todo-rco: RHIPacketList
-	FRHICommandList* RHICmdList = nullptr;
+	FRHICommandList& RHICmdList = FRHICommandList::GetNullRef();
 	VertexShader->SetVS(RHICmdList, Context);
 	PixelShader->SetParameters(RHICmdList, Context);
 
@@ -392,7 +392,7 @@ void FRCPassPostProcessLightShaftTemporalAA::Process(FRenderingCompositePassCont
 	SetGlobalBoundShaderState(BoundShaderState, GFilterVertexDeclaration.VertexDeclarationRHI, *VertexShader, *PixelShader);
 
 	//@todo-rco: RHIPacketList
-	FRHICommandList* RHICmdList = nullptr;
+	FRHICommandList& RHICmdList = FRHICommandList::GetNullRef();
 	VertexShader->SetVS(RHICmdList, Context);
 	PixelShader->SetParameters(RHICmdList, Context);
 
@@ -451,7 +451,7 @@ void FRCPassPostProcessTemporalAA::Process(FRenderingCompositePassContext& Conte
 	const FSceneRenderTargetItem& DestRenderTarget = PassOutputs[0].RequestSurface(Context);
 
 	//@todo-rco: RHIPacketList
-	FRHICommandList* RHICmdList = nullptr;
+	FRHICommandList& RHICmdList = FRHICommandList::GetNullRef();
 	//RHISetRenderTarget(DestRenderTarget.TargetableTexture, FTextureRHIRef());
 	RHISetRenderTarget(DestRenderTarget.TargetableTexture, GSceneRenderTargets.GetSceneDepthTexture());
 
diff --git a/Engine/Source/Runtime/Renderer/Private/PostProcess/PostProcessTonemap.cpp b/Engine/Source/Runtime/Renderer/Private/PostProcess/PostProcessTonemap.cpp
--- a/Engine/Source/Runtime/Renderer/Private/PostProcess/PostProcessTonemap.cpp
+++ b/Engine/Source/Runtime/Renderer/Private/PostProcess/PostProcessTonemap.cpp
@@ -825,7 +825,7 @@ public:
 		return bShaderHasOutdatedParameters;
 	}
 
-	void SetPS(FRHICommandList* RHICmdList, const FRenderingCompositePassContext& Context)
+	void SetPS(FRHICommandList& RHICmdList, const FRenderingCompositePassContext& Context)
 	{
 		const FPostProcessSettings& Settings = Context.View.FinalPostProcessSettings;
 		const FSceneViewFamily& ViewFamily = *(Context.View.Family);
@@ -969,7 +969,7 @@ FRCPassPostProcessTonemap::FRCPassPostProcessTonemap(bool bInDoGammaOnly)
 }
 
 template <uint32 ConfigIndex>
-static void SetShaderTempl(FRHICommandList* RHICmdList, const FRenderingCompositePassContext& Context)
+static void SetShaderTempl(FRHICommandList& RHICmdList, const FRenderingCompositePassContext& Context)
 {
 	TShaderMapRef<FPostProcessTonemapVS> VertexShader(GetGlobalShaderMap());
 	TShaderMapRef<FPostProcessTonemapPS<ConfigIndex> > PixelShader(GetGlobalShaderMap());
@@ -1007,7 +1007,7 @@ void FRCPassPostProcessTonemap::Process(FRenderingCompositePassContext& Context)
 	const FSceneRenderTargetItem& DestRenderTarget = PassOutputs[0].RequestSurface(Context);
 
 	//@todo-rco: RHIPacketList
-	FRHICommandList* RHICmdList = nullptr;
+	FRHICommandList& RHICmdList = FRHICommandList::GetNullRef();
 
 	// Set the view family's render target/viewport.
 	RHISetRenderTarget(DestRenderTarget.TargetableTexture, FTextureRHIParamRef());	
@@ -1182,7 +1182,7 @@ public:
 		return bShaderHasOutdatedParameters;
 	}
 
-	void SetPS(FRHICommandList* RHICmdList, const FRenderingCompositePassContext& Context)
+	void SetPS(FRHICommandList& RHICmdList, const FRenderingCompositePassContext& Context)
 	{
 		const FPostProcessSettings& Settings = Context.View.FinalPostProcessSettings;
 		const FSceneViewFamily& ViewFamily = *(Context.View.Family);
@@ -1294,7 +1294,7 @@ public:
 		GrainRandomFull.Bind(Initializer.ParameterMap, TEXT("GrainRandomFull"));
 	}
 
-	void SetVS(FRHICommandList* RHICmdList, const FRenderingCompositePassContext& Context)
+	void SetVS(FRHICommandList& RHICmdList, const FRenderingCompositePassContext& Context)
 	{
 		const FVertexShaderRHIParamRef ShaderRHI = GetVertexShader();
 		FGlobalShader::SetParameters(RHICmdList, ShaderRHI, Context.View);
@@ -1320,7 +1320,7 @@ IMPLEMENT_SHADER_TYPE(,FPostProcessTonemapVS_ES2,TEXT("PostProcessTonemap"),TEXT
 
 
 template <uint32 ConfigIndex>
-static void SetShaderTemplES2(FRHICommandList* RHICmdList, const FRenderingCompositePassContext& Context, bool bUsedFramebufferFetch)
+static void SetShaderTemplES2(FRHICommandList& RHICmdList, const FRenderingCompositePassContext& Context, bool bUsedFramebufferFetch)
 {
 	TShaderMapRef<FPostProcessTonemapVS_ES2> VertexShader(GetGlobalShaderMap());
 	TShaderMapRef<FPostProcessTonemapPS_ES2<ConfigIndex> > PixelShader(GetGlobalShaderMap());
@@ -1360,7 +1360,7 @@ void FRCPassPostProcessTonemapES2::Process(FRenderingCompositePassContext& Conte
 	const FSceneRenderTargetItem& DestRenderTarget = PassOutputs[0].RequestSurface(Context);
 
 	//@todo-rco: RHIPacketList
-	FRHICommandList* RHICmdList = nullptr;
+	FRHICommandList& RHICmdList = FRHICommandList::GetNullRef();
 
 	// Set the view family's render target/viewport.
 	RHISetRenderTarget(DestRenderTarget.TargetableTexture, FTextureRHIParamRef());	
diff --git a/Engine/Source/Runtime/Renderer/Private/PostProcess/PostProcessTonemap.h b/Engine/Source/Runtime/Renderer/Private/PostProcess/PostProcessTonemap.h
--- a/Engine/Source/Runtime/Renderer/Private/PostProcess/PostProcessTonemap.h
+++ b/Engine/Source/Runtime/Renderer/Private/PostProcess/PostProcessTonemap.h
@@ -100,7 +100,7 @@ public:
 		GrainRandomFull.Bind(Initializer.ParameterMap, TEXT("GrainRandomFull"));
 	}
 
-	void SetVS(FRHICommandList* RHICmdList, const FRenderingCompositePassContext& Context)
+	void SetVS(FRHICommandList& RHICmdList, const FRenderingCompositePassContext& Context)
 	{
 		const FVertexShaderRHIParamRef ShaderRHI = GetVertexShader();
 
diff --git a/Engine/Source/Runtime/Renderer/Private/PostProcess/PostProcessUpscale.cpp b/Engine/Source/Runtime/Renderer/Private/PostProcess/PostProcessUpscale.cpp
--- a/Engine/Source/Runtime/Renderer/Private/PostProcess/PostProcessUpscale.cpp
+++ b/Engine/Source/Runtime/Renderer/Private/PostProcess/PostProcessUpscale.cpp
@@ -47,7 +47,7 @@ public:
 		UpscaleSoftness.Bind(Initializer.ParameterMap,TEXT("UpscaleSoftness"));
 	}
 
-	void SetPS(FRHICommandList* RHICmdList, const FRenderingCompositePassContext& Context)
+	void SetPS(FRHICommandList& RHICmdList, const FRenderingCompositePassContext& Context)
 	{
 		const FPixelShaderRHIParamRef ShaderRHI = GetPixelShader();
 		
@@ -108,7 +108,7 @@ FRCPassPostProcessUpscale::FRCPassPostProcessUpscale(uint32 InUpscaleMethod)
 }
 
 template <uint32 Method>
-void FRCPassPostProcessUpscale::SetShader(FRHICommandList* RHICmdList, const FRenderingCompositePassContext& Context)
+void FRCPassPostProcessUpscale::SetShader(FRHICommandList& RHICmdList, const FRenderingCompositePassContext& Context)
 {
 	TShaderMapRef<FPostProcessVS> VertexShader(GetGlobalShaderMap());
 	TShaderMapRef<FPostProcessUpscalePS<Method> > PixelShader(GetGlobalShaderMap());
@@ -141,7 +141,7 @@ void FRCPassPostProcessUpscale::Process(FRenderingCompositePassContext& Context)
 	const FSceneRenderTargetItem& DestRenderTarget = PassOutputs[0].RequestSurface(Context);
 
 	//@todo-rco: RHIPacketList
-	FRHICommandList* RHICmdList = nullptr;
+	FRHICommandList& RHICmdList = FRHICommandList::GetNullRef();
 
 	// Set the view family's render target/viewport.
 	RHISetRenderTarget(DestRenderTarget.TargetableTexture, FTextureRHIRef());	
diff --git a/Engine/Source/Runtime/Renderer/Private/PostProcess/PostProcessUpscale.h b/Engine/Source/Runtime/Renderer/Private/PostProcess/PostProcessUpscale.h
--- a/Engine/Source/Runtime/Renderer/Private/PostProcess/PostProcessUpscale.h
+++ b/Engine/Source/Runtime/Renderer/Private/PostProcess/PostProcessUpscale.h
@@ -27,7 +27,7 @@ public:
 	virtual void Release() OVERRIDE { delete this; }
 	FPooledRenderTargetDesc ComputeOutputDesc(EPassOutputId InPassOutputId) const;
 private:
-	template <uint32 Method> static void SetShader(FRHICommandList* RHICmdList, const FRenderingCompositePassContext& Context);
+	template <uint32 Method> static void SetShader(FRHICommandList& RHICmdList, const FRenderingCompositePassContext& Context);
 	uint32 UpscaleMethod;
 };
 
diff --git a/Engine/Source/Runtime/Renderer/Private/PostProcess/PostProcessVisualizeBuffer.cpp b/Engine/Source/Runtime/Renderer/Private/PostProcess/PostProcessVisualizeBuffer.cpp
--- a/Engine/Source/Runtime/Renderer/Private/PostProcess/PostProcessVisualizeBuffer.cpp
+++ b/Engine/Source/Runtime/Renderer/Private/PostProcess/PostProcessVisualizeBuffer.cpp
@@ -53,7 +53,7 @@ public:
 		}
 	}
 
-	void SetPS(FRHICommandList* RHICmdList, const FRenderingCompositePassContext& Context)
+	void SetPS(FRHICommandList& RHICmdList, const FRenderingCompositePassContext& Context)
 	{
 		const FPixelShaderRHIParamRef ShaderRHI = GetPixelShader();
 
@@ -62,7 +62,7 @@ public:
 		DeferredParameters.Set(RHICmdList, ShaderRHI, Context.View);
 	}
 
-	void SetSourceTexture(FRHICommandList* RHICmdList, FTextureRHIRef Texture)
+	void SetSourceTexture(FRHICommandList& RHICmdList, FTextureRHIRef Texture)
 	{
 		if (bDrawingTile && SourceTexture.IsBound())
 		{
@@ -111,7 +111,7 @@ IMPLEMENT_SHADER_TYPE2(FPostProcessVisualizeBufferPS<true>, SF_Pixel);
 IMPLEMENT_SHADER_TYPE2(FPostProcessVisualizeBufferPS<false>, SF_Pixel);
 
 template <bool bDrawingTile>
-FShader* FRCPassPostProcessVisualizeBuffer::SetShaderTempl(FRHICommandList* RHICmdList, const FRenderingCompositePassContext& Context)
+FShader* FRCPassPostProcessVisualizeBuffer::SetShaderTempl(FRHICommandList& RHICmdList, const FRenderingCompositePassContext& Context)
 {
 	TShaderMapRef<FPostProcessVS> VertexShader(GetGlobalShaderMap());
 	TShaderMapRef<FPostProcessVisualizeBufferPS<bDrawingTile> > PixelShader(GetGlobalShaderMap());
@@ -146,7 +146,7 @@ void FRCPassPostProcessVisualizeBuffer::Process(FRenderingCompositePassContext&
 	const FSceneRenderTargetItem& DestRenderTarget = PassOutputs[0].RequestSurface(Context);
 
 	//@todo-rco: RHIPacketList
-	FRHICommandList* RHICmdList = nullptr;
+	FRHICommandList& RHICmdList = FRHICommandList::GetNullRef();
 
 	// Set the view family's render target/viewport.
 	RHISetRenderTarget(DestRenderTarget.TargetableTexture, FTextureRHIRef());	
diff --git a/Engine/Source/Runtime/Renderer/Private/PostProcess/PostProcessVisualizeBuffer.h b/Engine/Source/Runtime/Renderer/Private/PostProcess/PostProcessVisualizeBuffer.h
--- a/Engine/Source/Runtime/Renderer/Private/PostProcess/PostProcessVisualizeBuffer.h
+++ b/Engine/Source/Runtime/Renderer/Private/PostProcess/PostProcessVisualizeBuffer.h
@@ -42,7 +42,7 @@ private:
 
 	// @return VertexShader
 	template <bool bDrawTileTemplate>
-	FShader* SetShaderTempl(FRHICommandList* RHICmdList, const FRenderingCompositePassContext& Context);
+	FShader* SetShaderTempl(FRHICommandList& RHICmdList, const FRenderingCompositePassContext& Context);
 
 	TArray<TileData> Tiles;
 };
diff --git a/Engine/Source/Runtime/Renderer/Private/PostProcess/PostProcessWeightedSampleSum.cpp b/Engine/Source/Runtime/Renderer/Private/PostProcess/PostProcessWeightedSampleSum.cpp
--- a/Engine/Source/Runtime/Renderer/Private/PostProcess/PostProcessWeightedSampleSum.cpp
+++ b/Engine/Source/Runtime/Renderer/Private/PostProcess/PostProcessWeightedSampleSum.cpp
@@ -65,7 +65,7 @@ public:
 	}
 
 	/** Sets shader parameter values */
-	void SetParameters(FRHICommandList* RHICmdList, FSamplerStateRHIParamRef SamplerStateRHI, FTextureRHIParamRef FilterTextureRHI, FTextureRHIParamRef AdditiveTextureRHI, const FLinearColor* SampleWeightValues)
+	void SetParameters(FRHICommandList& RHICmdList, FSamplerStateRHIParamRef SamplerStateRHI, FTextureRHIParamRef FilterTextureRHI, FTextureRHIParamRef AdditiveTextureRHI, const FLinearColor* SampleWeightValues)
 	{
 		const FPixelShaderRHIParamRef ShaderRHI = GetPixelShader();
 
@@ -158,7 +158,7 @@ protected:
  * @param OutVertexShader - The vertex shader used for the filter
  */
 void SetFilterShaders(
-	FRHICommandList* RHICmdList, 
+	FRHICommandList& RHICmdList, 
 	FSamplerStateRHIParamRef SamplerStateRHI,
 	FTextureRHIParamRef FilterTextureRHI,
 	FTextureRHIParamRef AdditiveTextureRHI,
@@ -170,7 +170,7 @@ void SetFilterShaders(
 	)
 {
 	check(CombineMethodInt <= 2);
-	check(!RHICmdList);
+	RHICmdList.CheckIsNull();
 
 	// A macro to handle setting the filter shader for a specific number of samples.
 #define SET_FILTER_SHADER_TYPE(NumSamples) \
@@ -378,7 +378,7 @@ void FRCPassPostProcessWeightedSampleSum::Process(FRenderingCompositePassContext
 	}
 
 	//@todo-rco: RHIPacketList
-	FRHICommandList* RHICmdList = nullptr;
+	FRHICommandList& RHICmdList = FRHICommandList::GetNullRef();
 
 	RHISetRenderTarget(DestRenderTarget.TargetableTexture, FTextureRHIRef());
 
diff --git a/Engine/Source/Runtime/Renderer/Private/PostProcess/PostProcessWeightedSampleSum.h b/Engine/Source/Runtime/Renderer/Private/PostProcess/PostProcessWeightedSampleSum.h
--- a/Engine/Source/Runtime/Renderer/Private/PostProcess/PostProcessWeightedSampleSum.h
+++ b/Engine/Source/Runtime/Renderer/Private/PostProcess/PostProcessWeightedSampleSum.h
@@ -62,7 +62,7 @@ public:
 	}
 
 	/** Sets shader parameter values */
-	void SetParameters(FRHICommandList* RHICmdList, const FVector2D* SampleOffsetsValue)
+	void SetParameters(FRHICommandList& RHICmdList, const FVector2D* SampleOffsetsValue)
 	{
 		FVector4 PackedSampleOffsetsValue[NumSampleChunks];
 		for(int32 SampleIndex = 0;SampleIndex < NumSamples;SampleIndex += 2)
diff --git a/Engine/Source/Runtime/Renderer/Private/PostProcess/PostProcessing.h b/Engine/Source/Runtime/Renderer/Private/PostProcess/PostProcessing.h
--- a/Engine/Source/Runtime/Renderer/Private/PostProcess/PostProcessing.h
+++ b/Engine/Source/Runtime/Renderer/Private/PostProcess/PostProcessing.h
@@ -40,12 +40,12 @@ class FPostProcessVS : public FGlobalShader
 	FPostProcessVS() {}
 
 	/** to have a similar interface as all other shaders */
-	void SetParameters(FRHICommandList* RHICmdList, const FRenderingCompositePassContext& Context)
+	void SetParameters(FRHICommandList& RHICmdList, const FRenderingCompositePassContext& Context)
 	{
 		FGlobalShader::SetParameters(RHICmdList, GetVertexShader(), Context.View);
 	}
 
-	void SetParameters(FRHICommandList* RHICmdList, const FSceneView& View)
+	void SetParameters(FRHICommandList& RHICmdList, const FSceneView& View)
 	{
 		FGlobalShader::SetParameters(RHICmdList, GetVertexShader(), View);
 	}
diff --git a/Engine/Source/Runtime/Renderer/Private/PostProcess/RenderingCompositionGraph.cpp b/Engine/Source/Runtime/Renderer/Private/PostProcess/RenderingCompositionGraph.cpp
--- a/Engine/Source/Runtime/Renderer/Private/PostProcess/RenderingCompositionGraph.cpp
+++ b/Engine/Source/Runtime/Renderer/Private/PostProcess/RenderingCompositionGraph.cpp
@@ -664,17 +664,17 @@ void FPostProcessPassParameters::Bind(const FShaderParameterMap& ParameterMap)
 	}
 }
 
-void FPostProcessPassParameters::SetPS(FRHICommandList* RHICmdList, const FPixelShaderRHIParamRef& ShaderRHI, const FRenderingCompositePassContext& Context, FSamplerStateRHIParamRef Filter, bool bWhiteIfNoTexture, FSamplerStateRHIParamRef* FilterOverrideArray)
+void FPostProcessPassParameters::SetPS(FRHICommandList& RHICmdList, const FPixelShaderRHIParamRef& ShaderRHI, const FRenderingCompositePassContext& Context, FSamplerStateRHIParamRef Filter, bool bWhiteIfNoTexture, FSamplerStateRHIParamRef* FilterOverrideArray)
 {
 	Set(RHICmdList, ShaderRHI, Context, Filter, bWhiteIfNoTexture, FilterOverrideArray);
 }
 
-void FPostProcessPassParameters::SetCS(FRHICommandList* RHICmdList, const FComputeShaderRHIParamRef& ShaderRHI, const FRenderingCompositePassContext& Context, FSamplerStateRHIParamRef Filter, bool bWhiteIfNoTexture, FSamplerStateRHIParamRef* FilterOverrideArray)
+void FPostProcessPassParameters::SetCS(FRHICommandList& RHICmdList, const FComputeShaderRHIParamRef& ShaderRHI, const FRenderingCompositePassContext& Context, FSamplerStateRHIParamRef Filter, bool bWhiteIfNoTexture, FSamplerStateRHIParamRef* FilterOverrideArray)
 {
 	Set(RHICmdList, ShaderRHI, Context, Filter, bWhiteIfNoTexture, FilterOverrideArray);
 }
 
-void FPostProcessPassParameters::SetVS(FRHICommandList* RHICmdList, const FVertexShaderRHIParamRef& ShaderRHI, const FRenderingCompositePassContext& Context, FSamplerStateRHIParamRef Filter, bool bWhiteIfNoTexture, FSamplerStateRHIParamRef* FilterOverrideArray)
+void FPostProcessPassParameters::SetVS(FRHICommandList& RHICmdList, const FVertexShaderRHIParamRef& ShaderRHI, const FRenderingCompositePassContext& Context, FSamplerStateRHIParamRef Filter, bool bWhiteIfNoTexture, FSamplerStateRHIParamRef* FilterOverrideArray)
 {
 	Set(RHICmdList, ShaderRHI, Context, Filter, bWhiteIfNoTexture, FilterOverrideArray);
 }
diff --git a/Engine/Source/Runtime/Renderer/Private/PostProcess/RenderingCompositionGraph.h b/Engine/Source/Runtime/Renderer/Private/PostProcess/RenderingCompositionGraph.h
--- a/Engine/Source/Runtime/Renderer/Private/PostProcess/RenderingCompositionGraph.h
+++ b/Engine/Source/Runtime/Renderer/Private/PostProcess/RenderingCompositionGraph.h
@@ -521,13 +521,13 @@ struct FPostProcessPassParameters
 	void Bind(const FShaderParameterMap& ParameterMap);
 
 	/** Set the pixel shader parameter values. */
-	void SetPS(FRHICommandList* RHICmdList, const FPixelShaderRHIParamRef& ShaderRHI, const FRenderingCompositePassContext& Context, FSamplerStateRHIParamRef Filter = TStaticSamplerState<>::GetRHI(), bool bWhiteIfNoTexture = false, FSamplerStateRHIParamRef* FilterOverrideArray = 0);
+	void SetPS(FRHICommandList& RHICmdList, const FPixelShaderRHIParamRef& ShaderRHI, const FRenderingCompositePassContext& Context, FSamplerStateRHIParamRef Filter = TStaticSamplerState<>::GetRHI(), bool bWhiteIfNoTexture = false, FSamplerStateRHIParamRef* FilterOverrideArray = 0);
 
 	/** Set the compute shader parameter values. */
-	void SetCS(FRHICommandList* RHICmdList, const FComputeShaderRHIParamRef& ShaderRHI, const FRenderingCompositePassContext& Context, FSamplerStateRHIParamRef Filter = TStaticSamplerState<>::GetRHI(), bool bWhiteIfNoTexture = false, FSamplerStateRHIParamRef* FilterOverrideArray = 0);
+	void SetCS(FRHICommandList& RHICmdList, const FComputeShaderRHIParamRef& ShaderRHI, const FRenderingCompositePassContext& Context, FSamplerStateRHIParamRef Filter = TStaticSamplerState<>::GetRHI(), bool bWhiteIfNoTexture = false, FSamplerStateRHIParamRef* FilterOverrideArray = 0);
 
 	/** Set the vertex shader parameter values. */
-	void SetVS(FRHICommandList* RHICmdList, const FVertexShaderRHIParamRef& ShaderRHI, const FRenderingCompositePassContext& Context, FSamplerStateRHIParamRef Filter = TStaticSamplerState<>::GetRHI(), bool bWhiteIfNoTexture = false, FSamplerStateRHIParamRef* FilterOverrideArray = 0);
+	void SetVS(FRHICommandList& RHICmdList, const FVertexShaderRHIParamRef& ShaderRHI, const FRenderingCompositePassContext& Context, FSamplerStateRHIParamRef Filter = TStaticSamplerState<>::GetRHI(), bool bWhiteIfNoTexture = false, FSamplerStateRHIParamRef* FilterOverrideArray = 0);
 
 	/** Serializer. */
 	friend FArchive& operator<<(FArchive& Ar,FPostProcessPassParameters& P);
@@ -550,7 +550,7 @@ public:
 	// @param Filter can be 0 if FilterOverrideArray is used
 	// @param FilterOverrideArray can be 0 if Filter is used
 	template <class T>
-	void Set(FRHICommandList* RHICmdList, const T& ShaderRHI, const FRenderingCompositePassContext& Context, FSamplerStateRHIParamRef Filter, bool bWhiteIfNoTexture = false, FSamplerStateRHIParamRef* FilterOverrideArray = 0)
+	void Set(FRHICommandList& RHICmdList, const T& ShaderRHI, const FRenderingCompositePassContext& Context, FSamplerStateRHIParamRef Filter, bool bWhiteIfNoTexture = false, FSamplerStateRHIParamRef* FilterOverrideArray = 0)
 	{
 		// assuming all outputs have the same size
 		FRenderingCompositeOutput* Output = Context.Pass->GetOutput(ePId_Output0);
@@ -563,7 +563,7 @@ public:
 		// but not both
 		check(!FilterOverrideArray || !Filter);
 
-		check(!RHICmdList);
+		RHICmdList.CheckIsNull();
 
 		if(BilinearTextureSampler0.IsBound())
 		{
diff --git a/Engine/Source/Runtime/Renderer/Private/PostProcess/SceneFilterRendering.h b/Engine/Source/Runtime/Renderer/Private/PostProcess/SceneFilterRendering.h
--- a/Engine/Source/Runtime/Renderer/Private/PostProcess/SceneFilterRendering.h
+++ b/Engine/Source/Runtime/Renderer/Private/PostProcess/SceneFilterRendering.h
@@ -121,7 +121,7 @@ public:
 	}
 
 	/** Set the material shader parameter values. */
-	void Set(FRHICommandList* RHICmdList, FShader* PixelShader, float DisplayGamma, FLinearColor const& ColorScale, FLinearColor const& ColorOverlay)
+	void Set(FRHICommandList& RHICmdList, FShader* PixelShader, float DisplayGamma, FLinearColor const& ColorScale, FLinearColor const& ColorOverlay)
 	{
 		// GammaColorScaleAndInverse
 
diff --git a/Engine/Source/Runtime/Renderer/Private/PostProcess/SceneRenderTargets.h b/Engine/Source/Runtime/Renderer/Private/PostProcess/SceneRenderTargets.h
--- a/Engine/Source/Runtime/Renderer/Private/PostProcess/SceneRenderTargets.h
+++ b/Engine/Source/Runtime/Renderer/Private/PostProcess/SceneRenderTargets.h
@@ -572,7 +572,7 @@ public:
 
 	/** Sets the scene texture parameters for the given view. */
 	template<typename TParamRef>
-	void Set(FRHICommandList* RHICmdList, const TParamRef& ShaderRHI, const FSceneView& View, ESceneRenderTargetsMode::Type TextureMode = ESceneRenderTargetsMode::SetTextures, ESamplerFilter ColorFilter = SF_Point) const
+	void Set(FRHICommandList& RHICmdList, const TParamRef& ShaderRHI, const FSceneView& View, ESceneRenderTargetsMode::Type TextureMode = ESceneRenderTargetsMode::SetTextures, ESamplerFilter ColorFilter = SF_Point) const
 	{
 		if (TextureMode == ESceneRenderTargetsMode::SetTextures)
 		{
@@ -732,7 +732,7 @@ public:
 	void Bind(const FShaderParameterMap& ParameterMap);
 
 	template<typename ShaderRHIParamRef>
-	void Set(FRHICommandList* RHICmdList, const ShaderRHIParamRef ShaderRHI, const FSceneView& View, ESceneRenderTargetsMode::Type TextureMode = ESceneRenderTargetsMode::SetTextures) const
+	void Set(FRHICommandList& RHICmdList, const ShaderRHIParamRef ShaderRHI, const FSceneView& View, ESceneRenderTargetsMode::Type TextureMode = ESceneRenderTargetsMode::SetTextures) const
 	{
 		// This is needed on PC ES2 for SceneAlphaCopy, probably should be refactored for performance.
 		SceneTextureParameters.Set(RHICmdList, ShaderRHI, View, TextureMode, SF_Point);
diff --git a/Engine/Source/Runtime/Renderer/Private/PostProcess/ScreenSpaceReflections.cpp b/Engine/Source/Runtime/Renderer/Private/PostProcess/ScreenSpaceReflections.cpp
--- a/Engine/Source/Runtime/Renderer/Private/PostProcess/ScreenSpaceReflections.cpp
+++ b/Engine/Source/Runtime/Renderer/Private/PostProcess/ScreenSpaceReflections.cpp
@@ -96,7 +96,7 @@ public:
 		SSRParams.Bind(Initializer.ParameterMap, TEXT("SSRParams"));
 	}
 
-	void SetParameters(FRHICommandList* RHICmdList, const FRenderingCompositePassContext& Context)
+	void SetParameters(FRHICommandList& RHICmdList, const FRenderingCompositePassContext& Context)
 	{
 		const FFinalPostProcessSettings& Settings = Context.View.FinalPostProcessSettings;
 		const FPixelShaderRHIParamRef ShaderRHI = GetPixelShader();
@@ -176,7 +176,7 @@ void FRCPassPostProcessScreenSpaceReflections::Process(FRenderingCompositePassCo
 	SCOPED_DRAW_EVENT(ScreenSpaceReflections, DEC_SCENE_ITEMS);
 
 	//@todo-rco: RHIPacketList
-	FRHICommandList* RHICmdList = nullptr;
+	FRHICommandList& RHICmdList = FRHICommandList::GetNullRef();
 
 	const FSceneView& View = Context.View;
 
diff --git a/Engine/Source/Runtime/Renderer/Private/PostProcess/VisualizeTexture.cpp b/Engine/Source/Runtime/Renderer/Private/PostProcess/VisualizeTexture.cpp
--- a/Engine/Source/Runtime/Renderer/Private/PostProcess/VisualizeTexture.cpp
+++ b/Engine/Source/Runtime/Renderer/Private/PostProcess/VisualizeTexture.cpp
@@ -73,7 +73,7 @@ public:
 		return bShaderHasOutdatedParameters;
 	}
 
-	void SetParameters(FRHICommandList* RHICmdList, const FVisualizeTextureData& Data)
+	void SetParameters(FRHICommandList& RHICmdList, const FVisualizeTextureData& Data)
 	{
 		const FPixelShaderRHIParamRef ShaderRHI = GetPixelShader();
 		
@@ -178,7 +178,7 @@ public:
 		return bShaderHasOutdatedParameters;
 	}
 
-	void SetParameters(FRHICommandList* RHICmdList, const FSceneView& View, const IPooledRenderTarget& Src)
+	void SetParameters(FRHICommandList& RHICmdList, const FSceneView& View, const IPooledRenderTarget& Src)
 	{
 		const FPixelShaderRHIParamRef ShaderRHI = GetPixelShader();
 
@@ -194,7 +194,7 @@ IMPLEMENT_SHADER_TYPE(,FVisualizeTexturePresentPS,TEXT("VisualizeTexture"),TEXT(
 template<uint32 TextureType> void VisualizeTextureForTextureType(const FVisualizeTextureData& Data)
 {
 	//@todo-rco: RHIPacketList
-	FRHICommandList* RHICmdList = nullptr;
+	FRHICommandList& RHICmdList = FRHICommandList::GetNullRef();
 
 	TShaderMapRef<FScreenVS> VertexShader(GetGlobalShaderMap());
 	TShaderMapRef<VisualizeTexturePS<TextureType> > PixelShader(GetGlobalShaderMap());
@@ -493,7 +493,7 @@ void FVisualizeTexture::PresentContent(const FSceneView& View)
 	}
 
 	//@todo-rco: RHIPacketList
-	FRHICommandList* RHICmdList = nullptr;
+	FRHICommandList& RHICmdList = FRHICommandList::GetNullRef();
 
 	FPooledRenderTargetDesc Desc = VisualizeTextureDesc;
 
diff --git a/Engine/Source/Runtime/Renderer/Private/ReflectionEnvironment.cpp b/Engine/Source/Runtime/Renderer/Private/ReflectionEnvironment.cpp
--- a/Engine/Source/Runtime/Renderer/Private/ReflectionEnvironment.cpp
+++ b/Engine/Source/Runtime/Renderer/Private/ReflectionEnvironment.cpp
@@ -177,7 +177,7 @@ public:
 	{
 	}
 
-	void SetParameters(FRHICommandList* RHICmdList, const FSceneView& View, FTextureRHIParamRef SSRTexture, FUnorderedAccessViewRHIParamRef OutSceneColorUAV)
+	void SetParameters(FRHICommandList& RHICmdList, const FSceneView& View, FTextureRHIParamRef SSRTexture, FUnorderedAccessViewRHIParamRef OutSceneColorUAV)
 	{
 		const FComputeShaderRHIParamRef ShaderRHI = GetComputeShader();
 
@@ -275,7 +275,7 @@ public:
 	{
 		const FComputeShaderRHIParamRef ShaderRHI = GetComputeShader();
 		//@todo-rco: RHIPacketList
-		OutSceneColor.UnsetUAV(nullptr, ShaderRHI);
+		OutSceneColor.UnsetUAV(FRHICommandList::GetNullRef(), ShaderRHI);
 	}
 
 	virtual bool Serialize(FArchive& Ar)
@@ -378,7 +378,7 @@ public:
 		PreIntegratedGFSampler.Bind(Initializer.ParameterMap, TEXT("PreIntegratedGFSampler"));
 	}
 
-	void SetParameters(FRHICommandList* RHICmdList, const FSceneView& View, FTextureRHIParamRef ReflectionEnv, FTextureRHIParamRef ScreenSpaceReflections )
+	void SetParameters(FRHICommandList& RHICmdList, const FSceneView& View, FTextureRHIParamRef ReflectionEnv, FTextureRHIParamRef ScreenSpaceReflections )
 	{
 		const FPixelShaderRHIParamRef ShaderRHI = GetPixelShader();
 
@@ -450,7 +450,7 @@ public:
 		DeferredParameters.Bind(Initializer.ParameterMap);
 	}
 
-	void SetParameters(FRHICommandList* RHICmdList, const FSceneView& View)
+	void SetParameters(FRHICommandList& RHICmdList, const FSceneView& View)
 	{
 		const FPixelShaderRHIParamRef ShaderRHI = GetPixelShader();
 
@@ -504,7 +504,7 @@ public:
 		DeferredParameters.Bind(Initializer.ParameterMap);
 	}
 
-	void SetParameters(FRHICommandList* RHICmdList, const FSceneView& View, const FReflectionCaptureSortData& SortData)
+	void SetParameters(FRHICommandList& RHICmdList, const FSceneView& View, const FReflectionCaptureSortData& SortData)
 	{
 		const FPixelShaderRHIParamRef ShaderRHI = GetPixelShader();
 		FGlobalShader::SetParameters(RHICmdList, ShaderRHI, View);
@@ -568,7 +568,7 @@ void FDeferredShadingSceneRenderer::RenderReflectionCaptureSpecularBounceForAllV
 		RHISetViewport( View.ViewRect.Min.X, View.ViewRect.Min.Y, 0.0f, View.ViewRect.Max.X, View.ViewRect.Max.Y, 1.0f );
 
 		//@todo-rco: RHIPacketList
-		PixelShader->SetParameters(nullptr, View );
+		PixelShader->SetParameters(FRHICommandList::GetNullRef(), View);
 
 		DrawRectangle( 
 			0, 0,
@@ -591,7 +591,7 @@ bool FDeferredShadingSceneRenderer::ShouldDoReflectionEnvironment() const
 		&& ViewFamily.EngineShowFlags.ReflectionEnvironment;
 }
 
-void FDeferredShadingSceneRenderer::RenderImageBasedReflectionsSM5ForAllViews(FRHICommandList* RHICmdList)
+void FDeferredShadingSceneRenderer::RenderImageBasedReflectionsSM5ForAllViews(FRHICommandList& RHICmdList)
 {
 	const uint32 bUseLightmaps = CVarDiffuseFromCaptures.GetValueOnRenderThread() == 0;
 
@@ -652,7 +652,7 @@ void FDeferredShadingSceneRenderer::RenderImageBasedReflectionsSM5ForAllViews(FR
 	check(GSceneRenderTargets.GetSceneColor());
 }
 
-void FDeferredShadingSceneRenderer::RenderImageBasedReflectionsSM4ForAllViews(FRHICommandList* RHICmdList, bool bReflectionEnv)
+void FDeferredShadingSceneRenderer::RenderImageBasedReflectionsSM4ForAllViews(FRHICommandList& RHICmdList, bool bReflectionEnv)
 {
 	const bool bSkyLight = Scene->SkyLight
 		&& Scene->SkyLight->ProcessedTexture
@@ -813,7 +813,7 @@ void FDeferredShadingSceneRenderer::RenderImageBasedReflectionsSM4ForAllViews(FR
 			TShaderMapRef< FReflectionApplyPS<A, B, C> > PixelShader(GetGlobalShaderMap()); \
 			static FGlobalBoundShaderState BoundShaderState; \
 			SetGlobalBoundShaderState(BoundShaderState, GFilterVertexDeclaration.VertexDeclarationRHI, *VertexShader, *PixelShader); \
-			PixelShader->SetParameters(nullptr, View, LightAccumulation->GetRenderTargetItem().ShaderResourceTexture, SSROutput->GetRenderTargetItem().ShaderResourceTexture); \
+			PixelShader->SetParameters(FRHICommandList::GetNullRef(), View, LightAccumulation->GetRenderTargetItem().ShaderResourceTexture, SSROutput->GetRenderTargetItem().ShaderResourceTexture); \
 			}; \
 			break
 
@@ -842,7 +842,7 @@ void FDeferredShadingSceneRenderer::RenderImageBasedReflectionsSM4ForAllViews(FR
 	}
 }
 
-void FDeferredShadingSceneRenderer::RenderDeferredReflections(FRHICommandList* RHICmdList)
+void FDeferredShadingSceneRenderer::RenderDeferredReflections(FRHICommandList& RHICmdList)
 {
 	if (IsSimpleDynamicLightingEnabled() || ViewFamily.EngineShowFlags.VisualizeLightCulling)
 	{
diff --git a/Engine/Source/Runtime/Renderer/Private/ReflectionEnvironmentCapture.cpp b/Engine/Source/Runtime/Renderer/Private/ReflectionEnvironmentCapture.cpp
--- a/Engine/Source/Runtime/Renderer/Private/ReflectionEnvironmentCapture.cpp
+++ b/Engine/Source/Runtime/Renderer/Private/ReflectionEnvironmentCapture.cpp
@@ -96,7 +96,7 @@ public:
 	}
 	FDownsamplePS() {}
 
-	void SetParameters(FRHICommandList* RHICmdList, int32 CubeFaceValue, int32 SourceMipIndexValue, FSceneRenderTargetItem& SourceTextureValue)
+	void SetParameters(FRHICommandList& RHICmdList, int32 CubeFaceValue, int32 SourceMipIndexValue, FSceneRenderTargetItem& SourceTextureValue)
 	{
 		SetShaderValue(RHICmdList, GetPixelShader(), CubeFace, CubeFaceValue);
 		SetShaderValue(RHICmdList, GetPixelShader(), SourceMipIndex, SourceMipIndexValue);
@@ -158,7 +158,7 @@ public:
 	}
 	FCubeFilterPS() {}
 
-	void SetParameters(FRHICommandList* RHICmdList, int32 CubeFaceValue, int32 SourceMipIndexValue, float RoughnessValue, FSceneRenderTargetItem& SourceTextureValue)
+	void SetParameters(FRHICommandList& RHICmdList, int32 CubeFaceValue, int32 SourceMipIndexValue, float RoughnessValue, FSceneRenderTargetItem& SourceTextureValue)
 	{
 		FDownsamplePS::SetParameters(RHICmdList, CubeFaceValue, SourceMipIndexValue, SourceTextureValue);
 
@@ -242,9 +242,9 @@ public:
 	{
 	}
 
-	void SetParameters(FRHICommandList* RHICmdList)
+	void SetParameters(FRHICommandList& RHICmdList)
 	{
-		check(!RHICmdList);
+		RHICmdList.CheckIsNull();
 		const int32 EffectiveTopMipSize = GReflectionCaptureSize;
 		const int32 NumMips = FMath::CeilLogTwo(EffectiveTopMipSize) + 1;
 		// Read from the smallest mip that was downsampled to
@@ -283,10 +283,10 @@ private:
 IMPLEMENT_SHADER_TYPE(,FComputeBrightnessPS,TEXT("ReflectionEnvironmentShaders"),TEXT("ComputeBrightnessMain"),SF_Pixel);
 
 /** Computes the average brightness of the given reflection capture and stores it in the scene. */
-void ComputeAverageBrightness(FRHICommandList* RHICmdList)
+void ComputeAverageBrightness(FRHICommandList& RHICmdList)
 {
 	///@todo-rco:
-	check(!RHICmdList);
+	RHICmdList.CheckIsNull();
 	RHISetRenderTarget(GSceneRenderTargets.ReflectionBrightness->GetRenderTargetItem().TargetableTexture, NULL);
 	RHISetRasterizerState(TStaticRasterizerState<FM_Solid,CM_None>::GetRHI());
 	RHISetDepthStencilState(TStaticDepthStencilState<false,CF_Always>::GetRHI());
@@ -317,7 +317,7 @@ void FilterReflectionEnvironment(FSHVectorRGB3* OutIrradianceEnvironmentMap)
 	const int32 NumMips = FMath::CeilLogTwo(EffectiveTopMipSize) + 1;
 
 	//@todo-rco: RHIPacketList
-	FRHICommandList* RHICmdList = nullptr;
+	FRHICommandList& RHICmdList = FRHICommandList::GetNullRef();
 
 	int32 DiffuseConvolutionSourceMip = INDEX_NONE;
 	FSceneRenderTargetItem* DiffuseConvolutionSource = NULL;
@@ -550,7 +550,7 @@ public:
 	{
 	}
 
-	void SetParameters(FRHICommandList* RHICmdList, const FViewInfo& View)
+	void SetParameters(FRHICommandList& RHICmdList, const FViewInfo& View)
 	{
 		FGlobalShader::SetParameters(RHICmdList, GetVertexShader(),View);
 	}
@@ -582,7 +582,7 @@ public:
 	}
 	FRouteToCubeFaceGS() {}
 
-	void SetParameters(FRHICommandList* RHICmdList, uint32 CubeFaceValue)
+	void SetParameters(FRHICommandList& RHICmdList, uint32 CubeFaceValue)
 	{
 		SetShaderValue(RHICmdList, GetGeometryShader(), CubeFace, CubeFaceValue);
 	}
@@ -627,7 +627,7 @@ public:
 	}
 	FCopySceneColorToCubeFacePS() {}
 
-	void SetParameters(FRHICommandList* RHICmdList, const FViewInfo& View, bool bCapturingForSkyLight, bool bLowerHemisphereIsBlack)
+	void SetParameters(FRHICommandList& RHICmdList, const FViewInfo& View, bool bCapturingForSkyLight, bool bLowerHemisphereIsBlack)
 	{
 		const FPixelShaderRHIParamRef ShaderRHI = GetPixelShader();
 
@@ -707,7 +707,7 @@ public:
 	}
 	FCopyCubemapToCubeFacePS() {}
 
-	void SetParameters(FRHICommandList* RHICmdList, const FTexture* SourceCubemap, uint32 CubeFaceValue, bool bIsSkyLight, bool bLowerHemisphereIsBlack)
+	void SetParameters(FRHICommandList& RHICmdList, const FTexture* SourceCubemap, uint32 CubeFaceValue, bool bIsSkyLight, bool bLowerHemisphereIsBlack)
 	{
 		const FPixelShaderRHIParamRef ShaderRHI = GetPixelShader();
 
@@ -828,11 +828,11 @@ void ClearScratchCubemaps()
 }
 
 /** Captures the scene for a reflection capture by rendering the scene multiple times and copying into a cubemap texture. */
-void CaptureSceneToScratchCubemap(FRHICommandList* RHICmdList, FSceneRenderer* SceneRenderer, ECubeFace CubeFace, bool bCapturingForSkyLight, bool bLowerHemisphereIsBlack)
+void CaptureSceneToScratchCubemap(FRHICommandList& RHICmdList, FSceneRenderer* SceneRenderer, ECubeFace CubeFace, bool bCapturingForSkyLight, bool bLowerHemisphereIsBlack)
 {
 	FMemMark MemStackMark(FMemStack::Get());
 	//@todo-rco: RHIPacketList
-	check(!RHICmdList);
+	RHICmdList.CheckIsNull();
 	// update any resources that needed a deferred update
 	FDeferredUpdateResource::UpdateResources();
 	
@@ -854,7 +854,7 @@ void CaptureSceneToScratchCubemap(FRHICommandList* RHICmdList, FSceneRenderer* S
 		FSceneRenderTargetItem& EffectiveColorRT =  GSceneRenderTargets.ReflectionColorScratchCubemap[0]->GetRenderTargetItem();
 
 		{
-			check(!RHICmdList);
+			RHICmdList.CheckIsNull();
 			SCOPED_DRAW_EVENT(CubeMapCopyScene, DEC_SCENE_ITEMS);
 			if (GSupportsGSRenderTargetLayerSwitchingToMips)
 			{
@@ -924,7 +924,7 @@ void CaptureSceneToScratchCubemap(FRHICommandList* RHICmdList, FSceneRenderer* S
 	delete SceneRenderer;
 }
 
-void CopyCubemapToScratchCubemap(FRHICommandList* RHICmdList, UTextureCube* SourceCubemap, bool bIsSkyLight, bool bLowerHemisphereIsBlack)
+void CopyCubemapToScratchCubemap(FRHICommandList& RHICmdList, UTextureCube* SourceCubemap, bool bIsSkyLight, bool bLowerHemisphereIsBlack)
 {
 	check(SourceCubemap);
 	
@@ -935,7 +935,7 @@ void CopyCubemapToScratchCubemap(FRHICommandList* RHICmdList, UTextureCube* Sour
 
 	if (GSupportsGSRenderTargetLayerSwitchingToMips)
 	{
-		check(!RHICmdList);
+		RHICmdList.CheckIsNull();
 
 		// Copy the captured scene into the cubemap face
 		RHISetRenderTarget(EffectiveColorRT.TargetableTexture, NULL);
@@ -1419,7 +1419,7 @@ void CaptureSceneIntoScratchCubemap(FScene* Scene, FVector CapturePosition, bool
 			bool, bLowerHemisphereIsBlack, bLowerHemisphereIsBlack,
 		{
 			//@todo-rco: RHIPacketList
-			CaptureSceneToScratchCubemap(nullptr, SceneRenderer, CubeFace, bCapturingForSkyLight, bLowerHemisphereIsBlack);
+			CaptureSceneToScratchCubemap(FRHICommandList::GetNullRef(), SceneRenderer, CubeFace, bCapturingForSkyLight, bLowerHemisphereIsBlack);
 
 			RHIEndFrame();
 		});
@@ -1571,7 +1571,7 @@ void FScene::UpdateSkyCaptureContents(const USkyLightComponent* CaptureComponent
 				bool, bLowerHemisphereIsBlack, CaptureComponent->bLowerHemisphereIsBlack,
 			{
 				//@todo-rco: RHIPacketList
-				CopyCubemapToScratchCubemap(nullptr, SourceTexture, true, bLowerHemisphereIsBlack);
+				CopyCubemapToScratchCubemap(FRHICommandList::GetNullRef(), SourceTexture, true, bLowerHemisphereIsBlack);
 			});
 		}
 		else
diff --git a/Engine/Source/Runtime/Renderer/Private/ReflectionEnvironmentCapture.h b/Engine/Source/Runtime/Renderer/Private/ReflectionEnvironmentCapture.h
--- a/Engine/Source/Runtime/Renderer/Private/ReflectionEnvironmentCapture.h
+++ b/Engine/Source/Runtime/Renderer/Private/ReflectionEnvironmentCapture.h
@@ -27,7 +27,7 @@ public:
 	}
 	FDownsampleGS() {}
 
-	void SetParameters(FRHICommandList* RHICmdList, int32 CubeFaceValue)
+	void SetParameters(FRHICommandList& RHICmdList, int32 CubeFaceValue)
 	{
 		SetShaderValue(RHICmdList, GetGeometryShader(), CubeFace, CubeFaceValue);
 	}
diff --git a/Engine/Source/Runtime/Renderer/Private/ReflectionEnvironmentDiffuseIrradiance.cpp b/Engine/Source/Runtime/Renderer/Private/ReflectionEnvironmentDiffuseIrradiance.cpp
--- a/Engine/Source/Runtime/Renderer/Private/ReflectionEnvironmentDiffuseIrradiance.cpp
+++ b/Engine/Source/Runtime/Renderer/Private/ReflectionEnvironmentDiffuseIrradiance.cpp
@@ -55,7 +55,7 @@ public:
 	}
 	FCopyDiffuseIrradiancePS() {}
 
-	void SetParameters(FRHICommandList* RHICmdList, int32 CubeFaceValue, int32 SourceMipIndexValue, int32 CoefficientIndex, int32 FaceResolution, FTextureRHIRef& SourceTextureValue)
+	void SetParameters(FRHICommandList& RHICmdList, int32 CubeFaceValue, int32 SourceMipIndexValue, int32 CoefficientIndex, int32 FaceResolution, FTextureRHIRef& SourceTextureValue)
 	{
 		SetShaderValue(RHICmdList, GetPixelShader(), CubeFace, CubeFaceValue);
 		SetShaderValue(RHICmdList, GetPixelShader(), SourceMipIndex, SourceMipIndexValue);
@@ -137,7 +137,7 @@ public:
 		FDeferredPixelShaderParameters::ModifyCompilationEnvironment(Platform,OutEnvironment);
 	}
 
-	void SetParameters(FRHICommandList* RHICmdList, int32 CubeFaceValue, int32 NumMips, int32 SourceMipIndexValue, int32 CoefficientIndex, FTextureRHIRef& SourceTextureValue)
+	void SetParameters(FRHICommandList& RHICmdList, int32 CubeFaceValue, int32 NumMips, int32 SourceMipIndexValue, int32 CoefficientIndex, FTextureRHIRef& SourceTextureValue)
 	{
 		SetShaderValue(RHICmdList, GetPixelShader(), CubeFace, CubeFaceValue);
 		SetShaderValue(RHICmdList, GetPixelShader(), SourceMipIndex, SourceMipIndexValue);
@@ -204,7 +204,7 @@ public:
 	}
 	FAccumulateCubeFacesPS() {}
 
-	void SetParameters(FRHICommandList* RHICmdList, int32 SourceMipIndexValue, FTextureRHIRef& SourceTextureValue)
+	void SetParameters(FRHICommandList& RHICmdList, int32 SourceMipIndexValue, FTextureRHIRef& SourceTextureValue)
 	{
 		SetShaderValue(RHICmdList, GetPixelShader(), SourceMipIndex, SourceMipIndexValue);
 
@@ -239,7 +239,7 @@ FGlobalBoundShaderState AccumulateCubeFacesBoundShaderState;
 void ComputeDiffuseIrradiance(FTextureRHIRef LightingSource, int32 LightingSourceMipIndex, FSHVectorRGB3* OutIrradianceEnvironmentMap)
 {
 	//@todo-rco: RHIPacketList
-	FRHICommandList* RHICmdList = nullptr;
+	FRHICommandList& RHICmdList = FRHICommandList::GetNullRef();
 	for (int32 CoefficientIndex = 0; CoefficientIndex < FSHVector3::MaxSHBasis; CoefficientIndex++)
 	{
 		// Copy the starting mip from the lighting texture, apply texel area weighting and appropriate SH coefficient
diff --git a/Engine/Source/Runtime/Renderer/Private/SceneHitProxyRendering.cpp b/Engine/Source/Runtime/Renderer/Private/SceneHitProxyRendering.cpp
--- a/Engine/Source/Runtime/Renderer/Private/SceneHitProxyRendering.cpp
+++ b/Engine/Source/Runtime/Renderer/Private/SceneHitProxyRendering.cpp
@@ -17,12 +17,12 @@ class FHitProxyVS : public FMeshMaterialShader
 
 public:
 
-	void SetParameters(FRHICommandList* RHICmdList, const FMaterialRenderProxy* MaterialRenderProxy,const FSceneView& View)
+	void SetParameters(FRHICommandList& RHICmdList, const FMaterialRenderProxy* MaterialRenderProxy,const FSceneView& View)
 	{
 		FMeshMaterialShader::SetParameters(RHICmdList, GetVertexShader(), MaterialRenderProxy, *MaterialRenderProxy->GetMaterial(View.GetFeatureLevel()), View, ESceneRenderTargetsMode::SetTextures);
 	}
 
-	void SetMesh(FRHICommandList* RHICmdList, const FVertexFactory* VertexFactory,const FSceneView& View,const FPrimitiveSceneProxy* Proxy,const FMeshBatchElement& BatchElement)
+	void SetMesh(FRHICommandList& RHICmdList, const FVertexFactory* VertexFactory,const FSceneView& View,const FPrimitiveSceneProxy* Proxy,const FMeshBatchElement& BatchElement)
 	{
 		FMeshMaterialShader::SetMesh(RHICmdList, GetVertexShader(),VertexFactory,View,Proxy,BatchElement);
 	}
@@ -123,17 +123,17 @@ public:
 
 	FHitProxyPS() {}
 
-	void SetParameters(FRHICommandList* RHICmdList, const FMaterialRenderProxy* MaterialRenderProxy,const FSceneView& View)
+	void SetParameters(FRHICommandList& RHICmdList, const FMaterialRenderProxy* MaterialRenderProxy,const FSceneView& View)
 	{
 		FMeshMaterialShader::SetParameters(RHICmdList, GetPixelShader(), MaterialRenderProxy, *MaterialRenderProxy->GetMaterial(View.GetFeatureLevel()), View, ESceneRenderTargetsMode::SetTextures);
 	}
 
-	void SetMesh(FRHICommandList* RHICmdList, const FVertexFactory* VertexFactory,const FSceneView& View,const FPrimitiveSceneProxy* Proxy,const FMeshBatchElement& BatchElement)
+	void SetMesh(FRHICommandList& RHICmdList, const FVertexFactory* VertexFactory,const FSceneView& View,const FPrimitiveSceneProxy* Proxy,const FMeshBatchElement& BatchElement)
 	{
 		FMeshMaterialShader::SetMesh(RHICmdList, GetPixelShader(),VertexFactory,View,Proxy,BatchElement);
 	}
 
-	void SetHitProxyId(FRHICommandList* RHICmdList, FHitProxyId HitProxyIdValue)
+	void SetHitProxyId(FRHICommandList& RHICmdList, FHitProxyId HitProxyIdValue)
 	{
 		SetShaderValue(RHICmdList, GetPixelShader(), HitProxyId, HitProxyIdValue.GetColor().ReinterpretAsLinear());
 	}
@@ -173,9 +173,9 @@ FHitProxyDrawingPolicy::FHitProxyDrawingPolicy(
 	PixelShader = MaterialResource->GetShader<FHitProxyPS>(InVertexFactory->GetType());
 }
 
-void FHitProxyDrawingPolicy::DrawShared(FRHICommandList* RHICmdList, const FSceneView* View, FBoundShaderStateRHIParamRef BoundShaderState) const
+void FHitProxyDrawingPolicy::DrawShared(FRHICommandList& RHICmdList, const FSceneView* View, FBoundShaderStateRHIParamRef BoundShaderState) const
 {
-	check(!RHICmdList);
+	RHICmdList.CheckIsNull();
 	// Set the actual shader & vertex declaration state
 	RHISetBoundShaderState( BoundShaderState);
 
@@ -214,7 +214,7 @@ FBoundShaderStateRHIRef FHitProxyDrawingPolicy::CreateBoundShaderState(ERHIFeatu
 }
 
 void FHitProxyDrawingPolicy::SetMeshRenderState(
-	FRHICommandList* RHICmdList, 
+	FRHICommandList& RHICmdList, 
 	const FSceneView& View,
 	const FPrimitiveSceneProxy* PrimitiveSceneProxy,
 	const FMeshBatch& Mesh,
@@ -291,7 +291,7 @@ void FHitProxyDrawingPolicyFactory::AddStaticMesh(FScene* Scene,FStaticMesh* Sta
 }
 
 bool FHitProxyDrawingPolicyFactory::DrawDynamicMesh(
-	FRHICommandList* RHICmdList, 
+	FRHICommandList& RHICmdList, 
 	const FSceneView& View,
 	ContextType DrawingContext,
 	const FMeshBatch& Mesh,
@@ -359,7 +359,7 @@ void FDeferredShadingSceneRenderer::RenderHitProxies()
 	RHISetRenderTarget(HitProxyRT->GetRenderTargetItem().TargetableTexture, GSceneRenderTargets.GetSceneDepthSurface());
 
 	//@todo-rco: RHIPacketList
-	FRHICommandList* RHICmdList = nullptr;
+	FRHICommandList& RHICmdList = FRHICommandList::GetNullRef();
 
 	// Clear color for each view.
 	for(int32 ViewIndex = 0;ViewIndex < Views.Num();ViewIndex++)
diff --git a/Engine/Source/Runtime/Renderer/Private/SceneHitProxyRendering.h b/Engine/Source/Runtime/Renderer/Private/SceneHitProxyRendering.h
--- a/Engine/Source/Runtime/Renderer/Private/SceneHitProxyRendering.h
+++ b/Engine/Source/Runtime/Renderer/Private/SceneHitProxyRendering.h
@@ -30,7 +30,7 @@ public:
 			VertexShader == Other.VertexShader &&
 			PixelShader == Other.PixelShader;
 	}
-	void DrawShared(FRHICommandList* RHICmdList, const FSceneView* View,FBoundShaderStateRHIParamRef BoundShaderState) const;
+	void DrawShared(FRHICommandList& RHICmdList, const FSceneView* View,FBoundShaderStateRHIParamRef BoundShaderState) const;
 
 	/** 
 	* Create bound shader state using the vertex decl from the mesh draw policy
@@ -40,7 +40,7 @@ public:
 	FBoundShaderStateRHIRef CreateBoundShaderState(ERHIFeatureLevel::Type InFeatureLevel);
 
 	void SetMeshRenderState(
-		FRHICommandList* RHICmdList, 
+		FRHICommandList& RHICmdList, 
 		const FSceneView& View,
 		const FPrimitiveSceneProxy* PrimitiveSceneProxy,
 		const FMeshBatch& Mesh,
@@ -68,7 +68,7 @@ public:
 
 	static void AddStaticMesh(FScene* Scene,FStaticMesh* StaticMesh,ContextType DrawingContext = ContextType());
 	static bool DrawDynamicMesh(
-		FRHICommandList* RHICmdList, 
+		FRHICommandList& RHICmdList, 
 		const FSceneView& View,
 		ContextType DrawingContext,
 		const FMeshBatch& Mesh,
diff --git a/Engine/Source/Runtime/Renderer/Private/SceneOcclusion.cpp b/Engine/Source/Runtime/Renderer/Private/SceneOcclusion.cpp
--- a/Engine/Source/Runtime/Renderer/Private/SceneOcclusion.cpp
+++ b/Engine/Source/Runtime/Renderer/Private/SceneOcclusion.cpp
@@ -431,7 +431,7 @@ static void IssueProjectedShadowOcclusionQuery(FViewInfo& View, const FProjected
 		const FRenderQueryRHIRef ShadowOcclusionQuery = ViewState->OcclusionQueryPool.AllocateQuery();
 
 		//@todo-rco: RHIPacketList
-		FRHICommandList* RHICmdList = nullptr;
+		FRHICommandList& RHICmdList = FRHICommandList::GetNullRef();
 		VertexShader->SetParameters(RHICmdList, View);
 
 		// Draw the primitive's bounding box, using the occlusion query.
@@ -642,7 +642,7 @@ public:
 		BoundsExtentSampler.Bind( Initializer.ParameterMap, TEXT("BoundsExtentSampler") );
 	}
 
-	void SetParameters(FRHICommandList* RHICmdList, const FSceneView& View, const FHZB& HZB, FTextureRHIParamRef BoundsCenter, FTextureRHIParamRef BoundsExtent )
+	void SetParameters(FRHICommandList& RHICmdList, const FSceneView& View, const FHZB& HZB, FTextureRHIParamRef BoundsCenter, FTextureRHIParamRef BoundsExtent )
 	{
 		const FPixelShaderRHIParamRef ShaderRHI = GetPixelShader();
 
@@ -755,7 +755,7 @@ void FHZBOcclusionTester::Submit( const FViewInfo& View )
 		SetGlobalBoundShaderState( BoundShaderState, GFilterVertexDeclaration.VertexDeclarationRHI, *VertexShader, *PixelShader );
 
 		//@todo-rco: RHIPacketList
-		FRHICommandList* RHICmdList = nullptr;
+		FRHICommandList& RHICmdList = FRHICommandList::GetNullRef();
 		PixelShader->SetParameters(RHICmdList, View,  ViewState->HZB, BoundsCenterTexture->GetRenderTargetItem().ShaderResourceTexture, BoundsExtentTexture->GetRenderTargetItem().ShaderResourceTexture );
 
 		RHISetViewport( 0, 0, 0.0f, SizeX, SizeY, 1.0f );
@@ -812,7 +812,7 @@ public:
 		TextureParameterSampler.Bind( Initializer.ParameterMap, TEXT("TextureSampler") );
 	}
 
-	void SetParameters(FRHICommandList* RHICmdList, const FSceneView& View, const FIntPoint& Size )
+	void SetParameters(FRHICommandList& RHICmdList, const FSceneView& View, const FIntPoint& Size )
 	{
 		const FPixelShaderRHIParamRef ShaderRHI = GetPixelShader();
 
@@ -824,7 +824,7 @@ public:
 		SceneTextureParameters.Set(RHICmdList, ShaderRHI, View );
 	}
 
-	void SetParameters(FRHICommandList* RHICmdList, const FSceneView& View, const FIntPoint& Size, FShaderResourceViewRHIParamRef ShaderResourceView )
+	void SetParameters(FRHICommandList& RHICmdList, const FSceneView& View, const FIntPoint& Size, FShaderResourceViewRHIParamRef ShaderResourceView )
 	{
 		const FPixelShaderRHIParamRef ShaderRHI = GetPixelShader();
 
@@ -887,7 +887,7 @@ void BuildHZB( const FViewInfo& View )
 	RHISetDepthStencilState( TStaticDepthStencilState< false, CF_Always >::GetRHI() );
 	
 	//@todo-rco: RHIPacketList
-	FRHICommandList* RHICmdList = nullptr;
+	FRHICommandList& RHICmdList = FRHICommandList::GetNullRef();
 
 	// Mip 0
 	{
@@ -1006,7 +1006,7 @@ void FDeferredShadingSceneRenderer::BeginOcclusionTests()
 			TShaderMapRef<FOcclusionQueryVS> VertexShader(GetGlobalShaderMap());
 			SetGlobalBoundShaderState(OcclusionTestBoundShaderState,GetVertexDeclarationFVector3(),*VertexShader,NULL);
 			//@todo-rco: RHIPacketList
-			FRHICommandList* RHICmdList = nullptr;
+			FRHICommandList& RHICmdList = FRHICommandList::GetNullRef();
 			VertexShader->SetParameters(RHICmdList, View);
 
 			// Issue this frame's occlusion queries (occlusion queries from last frame may still be in flight)
diff --git a/Engine/Source/Runtime/Renderer/Private/SceneOcclusion.h b/Engine/Source/Runtime/Renderer/Private/SceneOcclusion.h
--- a/Engine/Source/Runtime/Renderer/Private/SceneOcclusion.h
+++ b/Engine/Source/Runtime/Renderer/Private/SceneOcclusion.h
@@ -28,7 +28,7 @@ public:
 		FGlobalShader::ModifyCompilationEnvironment(Platform, OutEnvironment);
 	}
 
-	void SetParametersWithBoundingSphere(FRHICommandList* RHICmdList, const FSceneView& View, const FSphere& BoundingSphere)
+	void SetParametersWithBoundingSphere(FRHICommandList& RHICmdList, const FSceneView& View, const FSphere& BoundingSphere)
 	{
 		FGlobalShader::SetParameters(RHICmdList, GetVertexShader(), View);
 
@@ -37,7 +37,7 @@ public:
 		StencilingGeometryParameters.Set(RHICmdList, this, StencilingSpherePosAndScale);
 	}
 
-	void SetParameters(FRHICommandList* RHICmdList, const FSceneView& View)
+	void SetParameters(FRHICommandList& RHICmdList, const FSceneView& View)
 	{
 		FGlobalShader::SetParameters(RHICmdList, GetVertexShader(),View);
 
diff --git a/Engine/Source/Runtime/Renderer/Private/SceneRendering.h b/Engine/Source/Runtime/Renderer/Private/SceneRendering.h
--- a/Engine/Source/Runtime/Renderer/Private/SceneRendering.h
+++ b/Engine/Source/Runtime/Renderer/Private/SceneRendering.h
@@ -561,7 +561,7 @@ protected:
 	void OnStartFrame();
 
 	/** Renders the scene's distortion */
-	void RenderDistortion(FRHICommandList* RHICmdList);
+	void RenderDistortion(FRHICommandList& RHICmdList);
 };
 
 
@@ -583,10 +583,10 @@ protected:
 	void InitViews();
 
 	/** Renders the opaque base pass for forward shading. */
-	void RenderForwardShadingBasePass(FRHICommandList* RHICmdList);
+	void RenderForwardShadingBasePass(FRHICommandList& RHICmdList);
 
 	/** Makes a copy of scene alpha so PC can emulate ES2 framebuffer fetch. */
-	void CopySceneAlpha(FRHICommandList* RHICmdList, const FSceneView& View);
+	void CopySceneAlpha(FRHICommandList& RHICmdList, const FSceneView& View);
 
 	/** Renders the base pass for translucency. */
 	void RenderTranslucency();
