UE4 - fix clang platforms

diff --git a/Engine/Source/Runtime/Renderer/Private/PostProcess/RenderingCompositionGraph.h b/Engine/Source/Runtime/Renderer/Private/PostProcess/RenderingCompositionGraph.h
--- a/Engine/Source/Runtime/Renderer/Private/PostProcess/RenderingCompositionGraph.h
+++ b/Engine/Source/Runtime/Renderer/Private/PostProcess/RenderingCompositionGraph.h
@@ -7,6 +7,7 @@
 #pragma once
 
 #include "RHIStaticStates.h"
+#include "RHICommandList.h"
 
 struct FRenderingCompositePass;
 struct FRenderingCompositeOutputRef;
