Better fix for eye adaptation overbrightness. Don't calculate the histogram in the min==max case, instead test for isfinite before dividing in the shader.

diff --git a/Engine/Shaders/PostProcessEyeAdaptation.usf b/Engine/Shaders/PostProcessEyeAdaptation.usf
--- a/Engine/Shaders/PostProcessEyeAdaptation.usf
+++ b/Engine/Shaders/PostProcessEyeAdaptation.usf
@@ -32,7 +32,7 @@ void MainPS(float4 UVAndScreenPos : TEXCOORD0, out float OutColor : SV_Target0)
 	
 	float TargetExposure = ComputeEyeAdaptationExposure(InputNew0);
 	float OldExposureScale = InputNew0.Load(int3(0, 1, 0)).x;
-	float OldExposure = ExposureOffsetMultipler / OldExposureScale;
+	float OldExposure = ExposureOffsetMultipler / (isfinite( OldExposureScale ) ? OldExposureScale : 1.0f);
 	float FrameTime = EyeAdaptationParams[1].y;
 
 	// eye adaptation changes over time
diff --git a/Engine/Source/Runtime/Renderer/Private/PostProcess/PostProcessing.cpp b/Engine/Source/Runtime/Renderer/Private/PostProcess/PostProcessing.cpp
--- a/Engine/Source/Runtime/Renderer/Private/PostProcess/PostProcessing.cpp
+++ b/Engine/Source/Runtime/Renderer/Private/PostProcess/PostProcessing.cpp
@@ -964,6 +964,7 @@ void FPostProcessing::Process(const FViewInfo& View, TRefCountPtr<IPooledRenderT
 				bool bHistogramNeeded = View.Family->EngineShowFlags.VisualizeHDR;
 
 				if(View.Family->EngineShowFlags.EyeAdaptation
+					&& View.FinalPostProcessSettings.AutoExposureMinBrightness < View.FinalPostProcessSettings.AutoExposureMaxBrightness
 					&& !View.bIsSceneCapture) // Eye adaption is not available for scene captures.
 				{
 					bHistogramNeeded = true;
