Remove spurious ClearScratchCubemaps() from CaptureSceneToScratchCubemap()

CaptureSceneToScratchCubemap() is called once per cube face, so we were clearing the whole cubemap before drawing each face, meaning only the last face was actually valid. The rest were all luminous green!

diff --git a/Engine/Source/Runtime/Renderer/Private/ReflectionEnvironmentCapture.cpp b/Engine/Source/Runtime/Renderer/Private/ReflectionEnvironmentCapture.cpp
--- a/Engine/Source/Runtime/Renderer/Private/ReflectionEnvironmentCapture.cpp
+++ b/Engine/Source/Runtime/Renderer/Private/ReflectionEnvironmentCapture.cpp
@@ -842,8 +842,6 @@ void CaptureSceneToScratchCubemap(FRHICommandList& RHICmdList, FSceneRenderer* S
 		// Render the scene normally for one face of the cubemap
 		SceneRenderer->Render();
 
-		ClearScratchCubemaps();
-
 #if PLATFORM_PS4 // @todo ps4 - this should be done a different way
 		// PS4 needs some code here to process the scene
 		extern void TEMP_PostReflectionCaptureRender();
