UE4 - more work preparing for parallel rendering. Merged the cmd list into the post process context, among other things.

diff --git a/Engine/Plugins/OculusRift/Source/OculusRift/Private/OculusRiftRender.cpp b/Engine/Plugins/OculusRift/Source/OculusRift/Private/OculusRiftRender.cpp
--- a/Engine/Plugins/OculusRift/Source/OculusRift/Private/OculusRiftRender.cpp
+++ b/Engine/Plugins/OculusRift/Source/OculusRift/Private/OculusRiftRender.cpp
@@ -106,7 +106,7 @@ void FOculusRiftHMD::DrawDistortionMesh_RenderThread(FRenderingCompositePassCont
 
 	Ptr<FDistortionMesh> mesh = RenderParams_RenderThread.pDistortionMesh[(View.StereoPass == eSSP_LEFT_EYE) ? 0 : 1];
 
-	RHIDrawIndexedPrimitiveUP(PT_TriangleList, 0, mesh->NumVertices, mesh->NumTriangles, mesh->pIndices,
+	DrawIndexedPrimitiveUP(Context.RHICmdList, PT_TriangleList, 0, mesh->NumVertices, mesh->NumTriangles, mesh->pIndices,
 		sizeof(mesh->pIndices[0]), mesh->pVertices, sizeof(mesh->pVertices[0]));
 #else
 	check(0);
diff --git a/Engine/Source/Editor/UnrealEd/Private/CubemapUnwrapUtils.cpp b/Engine/Source/Editor/UnrealEd/Private/CubemapUnwrapUtils.cpp
--- a/Engine/Source/Editor/UnrealEd/Private/CubemapUnwrapUtils.cpp
+++ b/Engine/Source/Editor/UnrealEd/Private/CubemapUnwrapUtils.cpp
@@ -141,10 +141,10 @@ void FMipLevelBatchedElementParameters::BindShaders_RenderThread( FRHICommandLis
 {
 	TShaderMapRef<FCubemapTexturePropertiesVS> VertexShader(GetGlobalShaderMap());
 	TShaderMapRef<TPixelShader> PixelShader(GetGlobalShaderMap());
-	RHICmdList.CheckIsNull();
+	RHICmdList.CheckIsNull(); // need new approach for "static FGlobalBoundShaderState" for parallel rendering
 
 	static FGlobalBoundShaderState BoundShaderState;
-	SetGlobalBoundShaderState(BoundShaderState, GSimpleElementVertexDeclaration.VertexDeclarationRHI, *VertexShader, *PixelShader);
+	SetGlobalBoundShaderState(RHICmdList, BoundShaderState, GSimpleElementVertexDeclaration.VertexDeclarationRHI, *VertexShader, *PixelShader);
 
 	VertexShader->SetParameters(RHICmdList, InTransform);
 
@@ -166,9 +166,9 @@ void FIESLightProfileBatchedElementParameters::BindShaders_RenderThread( FRHICom
 	TShaderMapRef<FSimpleElementVS> VertexShader(GetGlobalShaderMap());
 	TShaderMapRef<FIESLightProfilePS> PixelShader(GetGlobalShaderMap());
 
-	RHICmdList.CheckIsNull();
+	RHICmdList.CheckIsNull(); // need new approach for "static FGlobalBoundShaderState" for parallel rendering
 	static FGlobalBoundShaderState BoundShaderState;
-	SetGlobalBoundShaderState(BoundShaderState, GSimpleElementVertexDeclaration.VertexDeclarationRHI, *VertexShader, *PixelShader);
+	SetGlobalBoundShaderState(RHICmdList, BoundShaderState, GSimpleElementVertexDeclaration.VertexDeclarationRHI, *VertexShader, *PixelShader);
 	RHISetBlendState(TStaticBlendState<>::GetRHI());
 
 	VertexShader->SetParameters(RHICmdList, InTransform);
diff --git a/Engine/Source/Editor/UnrealEd/Private/MeshPaintRendering.cpp b/Engine/Source/Editor/UnrealEd/Private/MeshPaintRendering.cpp
--- a/Engine/Source/Editor/UnrealEd/Private/MeshPaintRendering.cpp
+++ b/Engine/Source/Editor/UnrealEd/Private/MeshPaintRendering.cpp
@@ -355,8 +355,9 @@ namespace MeshPaintRendering
 		TShaderMapRef< TMeshPaintVertexShader > VertexShader( GetGlobalShaderMap() );
 		TShaderMapRef< TMeshPaintPixelShader > PixelShader( GetGlobalShaderMap() );
 
+		RHICmdList.CheckIsNull(); // need new approach for "static FGlobalBoundShaderState" for parallel rendering
 		static FGlobalBoundShaderState BoundShaderState;
-		SetGlobalBoundShaderState( BoundShaderState, GMeshPaintVertexDeclaration.VertexDeclarationRHI, *VertexShader, *PixelShader );
+		SetGlobalBoundShaderState(RHICmdList, BoundShaderState, GMeshPaintVertexDeclaration.VertexDeclarationRHI, *VertexShader, *PixelShader);
 
 		// Set vertex shader parameters
 		VertexShader->SetParameters(RHICmdList, InTransform );
@@ -375,8 +376,9 @@ namespace MeshPaintRendering
 		TShaderMapRef< TMeshPaintDilateVertexShader > VertexShader( GetGlobalShaderMap() );
 		TShaderMapRef< TMeshPaintDilatePixelShader > PixelShader( GetGlobalShaderMap() );
 
+		RHICmdList.CheckIsNull(); // need new approach for "static FGlobalBoundShaderState" for parallel rendering
 		static FGlobalBoundShaderState BoundShaderState;
-		SetGlobalBoundShaderState( BoundShaderState, GMeshPaintDilateVertexDeclaration.VertexDeclarationRHI, *VertexShader, *PixelShader );
+		SetGlobalBoundShaderState(RHICmdList, BoundShaderState, GMeshPaintDilateVertexDeclaration.VertexDeclarationRHI, *VertexShader, *PixelShader);
 
 		// Set vertex shader parameters
 		VertexShader->SetParameters(RHICmdList, InTransform );
diff --git a/Engine/Source/Editor/UnrealEd/Private/NormalMapPreview.cpp b/Engine/Source/Editor/UnrealEd/Private/NormalMapPreview.cpp
--- a/Engine/Source/Editor/UnrealEd/Private/NormalMapPreview.cpp
+++ b/Engine/Source/Editor/UnrealEd/Private/NormalMapPreview.cpp
@@ -84,10 +84,9 @@ void FNormalMapBatchedElementParameters::BindShaders_RenderThread(
 	TShaderMapRef<FSimpleElementVS> VertexShader(GetGlobalShaderMap());
 	TShaderMapRef<FSimpleElementNormalMapPS> PixelShader(GetGlobalShaderMap());
 
-	RHICmdList.CheckIsNull();
-
+	RHICmdList.CheckIsNull(); // need new approach for "static FGlobalBoundShaderState" for parallel rendering
 	static FGlobalBoundShaderState BoundShaderState;
-	SetGlobalBoundShaderState(BoundShaderState, GSimpleElementVertexDeclaration.VertexDeclarationRHI, *VertexShader, *PixelShader);
+	SetGlobalBoundShaderState(RHICmdList, BoundShaderState, GSimpleElementVertexDeclaration.VertexDeclarationRHI, *VertexShader, *PixelShader);
 
 	VertexShader->SetParameters(RHICmdList, InTransform);
 	PixelShader->SetParameters(RHICmdList, Texture);
diff --git a/Engine/Source/Editor/UnrealEd/Private/Texture2DPreview.cpp b/Engine/Source/Editor/UnrealEd/Private/Texture2DPreview.cpp
--- a/Engine/Source/Editor/UnrealEd/Private/Texture2DPreview.cpp
+++ b/Engine/Source/Editor/UnrealEd/Private/Texture2DPreview.cpp
@@ -85,10 +85,10 @@ void FBatchedElementTexture2DPreviewParameters::BindShaders_RenderThread(
 {
 	TShaderMapRef<FSimpleElementVS> VertexShader(GetGlobalShaderMap());
 	TShaderMapRef<FSimpleElementTexture2DPreviewPS> PixelShader(GetGlobalShaderMap());
-	RHICmdList.CheckIsNull();
 
+	RHICmdList.CheckIsNull(); // need new approach for "static FGlobalBoundShaderState" for parallel rendering
 	static FGlobalBoundShaderState BoundShaderState;
-	SetGlobalBoundShaderState(BoundShaderState, GSimpleElementVertexDeclaration.VertexDeclarationRHI, *VertexShader, *PixelShader);
+	SetGlobalBoundShaderState(RHICmdList, BoundShaderState, GSimpleElementVertexDeclaration.VertexDeclarationRHI, *VertexShader, *PixelShader);
 
 	VertexShader->SetParameters(RHICmdList, InTransform);
 	PixelShader->SetParameters(RHICmdList, Texture, ColorWeights, InGamma, MipLevel, bIsNormalMap);
diff --git a/Engine/Source/Runtime/Engine/Private/BatchedElements.cpp b/Engine/Source/Runtime/Engine/Private/BatchedElements.cpp
--- a/Engine/Source/Runtime/Engine/Private/BatchedElements.cpp
+++ b/Engine/Source/Runtime/Engine/Private/BatchedElements.cpp
@@ -275,7 +275,7 @@ void FBatchedElements::AddSprite(
 }
 
 /** Translates a ESimpleElementBlendMode into a RHI state change for rendering a mesh with the blend mode normally. */
-static void SetBlendState(ESimpleElementBlendMode BlendMode)
+static void SetBlendState(FRHICommandList& RHICmdList, ESimpleElementBlendMode BlendMode)
 {
 	switch(BlendMode)
 	{
@@ -283,49 +283,49 @@ static void SetBlendState(ESimpleElementBlendMode BlendMode)
 	case SE_BLEND_Masked:
 	case SE_BLEND_MaskedDistanceField:
 	case SE_BLEND_MaskedDistanceFieldShadowed:
-		RHISetBlendState(TStaticBlendState<>::GetRHI());
+		RHICmdList.SetBlendState(TStaticBlendState<>::GetRHI());
 		break;
 	case SE_BLEND_Translucent:
 	case SE_BLEND_TranslucentDistanceField:
 	case SE_BLEND_TranslucentDistanceFieldShadowed:
-		RHISetBlendState(TStaticBlendState<CW_RGB,BO_Add,BF_SourceAlpha,BF_InverseSourceAlpha,BO_Add,BF_Zero,BF_One>::GetRHI());
+		RHICmdList.SetBlendState(TStaticBlendState<CW_RGB, BO_Add, BF_SourceAlpha, BF_InverseSourceAlpha, BO_Add, BF_Zero, BF_One>::GetRHI());
 		break;
 	case SE_BLEND_Additive:
-		RHISetBlendState(TStaticBlendState<CW_RGB,BO_Add,BF_One,BF_One>::GetRHI());
+		RHICmdList.SetBlendState(TStaticBlendState<CW_RGB, BO_Add, BF_One, BF_One>::GetRHI());
 		break;
 	case SE_BLEND_Modulate:
-		RHISetBlendState(TStaticBlendState<CW_RGB,BO_Add,BF_DestColor,BF_Zero>::GetRHI());
+		RHICmdList.SetBlendState(TStaticBlendState<CW_RGB, BO_Add, BF_DestColor, BF_Zero>::GetRHI());
 		break;
 	case SE_BLEND_AlphaComposite:
-		RHISetBlendState(TStaticBlendState<CW_RGBA,BO_Add,BF_One,BF_InverseSourceAlpha,BO_Add,BF_One,BF_InverseSourceAlpha>::GetRHI());
+		RHICmdList.SetBlendState(TStaticBlendState<CW_RGBA, BO_Add, BF_One, BF_InverseSourceAlpha, BO_Add, BF_One, BF_InverseSourceAlpha>::GetRHI());
 		break;
 	}
 }
 
 /** Translates a ESimpleElementBlendMode into a RHI state change for rendering a mesh with the blend mode for hit testing. */
-static void SetHitTestingBlendState(ESimpleElementBlendMode BlendMode)
+static void SetHitTestingBlendState(FRHICommandList& RHICmdList, ESimpleElementBlendMode BlendMode)
 {
 	switch(BlendMode)
 	{
 	case SE_BLEND_Opaque:
-		RHISetBlendState(TStaticBlendState<>::GetRHI());
+		RHICmdList.SetBlendState(TStaticBlendState<>::GetRHI());
 		break;
 	case SE_BLEND_Masked:
 	case SE_BLEND_MaskedDistanceField:
 	case SE_BLEND_MaskedDistanceFieldShadowed:
-		RHISetBlendState(TStaticBlendState<CW_RGBA,BO_Add,BF_One,BF_Zero,BO_Add,BF_One,BF_Zero>::GetRHI());
+		RHICmdList.SetBlendState(TStaticBlendState<CW_RGBA, BO_Add, BF_One, BF_Zero, BO_Add, BF_One, BF_Zero>::GetRHI());
 		break;
 	case SE_BLEND_AlphaComposite:
 	case SE_BLEND_Translucent:
 	case SE_BLEND_TranslucentDistanceField:
 	case SE_BLEND_TranslucentDistanceFieldShadowed:
-		RHISetBlendState(TStaticBlendState<CW_RGBA,BO_Add,BF_One,BF_Zero,BO_Add,BF_One,BF_Zero>::GetRHI());
+		RHICmdList.SetBlendState(TStaticBlendState<CW_RGBA, BO_Add, BF_One, BF_Zero, BO_Add, BF_One, BF_Zero>::GetRHI());
 		break;
 	case SE_BLEND_Additive:
-		RHISetBlendState(TStaticBlendState<CW_RGBA,BO_Add,BF_One,BF_Zero,BO_Add,BF_One,BF_Zero>::GetRHI());
+		RHICmdList.SetBlendState(TStaticBlendState<CW_RGBA, BO_Add, BF_One, BF_Zero, BO_Add, BF_One, BF_Zero>::GetRHI());
 		break;
 	case SE_BLEND_Modulate:
-		RHISetBlendState(TStaticBlendState<CW_RGBA,BO_Add,BF_One,BF_Zero,BO_Add,BF_One,BF_Zero>::GetRHI());
+		RHICmdList.SetBlendState(TStaticBlendState<CW_RGBA, BO_Add, BF_One, BF_Zero, BO_Add, BF_One, BF_Zero>::GetRHI());
 		break;
 	}
 }
@@ -346,6 +346,7 @@ float GBatchedElementSmoothWidth = 12;
  * Sets the appropriate vertex and pixel shader.
  */
 void FBatchedElements::PrepareShaders(
+	FRHICommandList& RHICmdList,
 	ESimpleElementBlendMode BlendMode,
 	const FMatrix& Transform,
 	bool bSwitchVerticalAxis,
@@ -358,9 +359,6 @@ void FBatchedElements::PrepareShaders(
 	FTexture2DRHIRef DepthTexture
 	) const
 {
-	//@todo-rco: RHIPacketList
-	FRHICommandList& RHICmdList = FRHICommandList::GetNullRef();
-
 	// used to mask individual channels and desaturate
 	FMatrix ColorWeights( FPlane(1, 0, 0, 0), FPlane(0, 1, 0, 0), FPlane(0, 0, 1, 0), FPlane(0, 0, 0, 0) );
 	
@@ -392,13 +390,13 @@ void FBatchedElements::PrepareShaders(
 		// If we are only to draw the alpha channel, make the Blend state opaque, to allow easy identification of the alpha values
 		if( bAlphaOnly )
 		{
-			SetBlendState( SE_BLEND_Opaque );
+			SetBlendState(RHICmdList, SE_BLEND_Opaque);
 
 			R.W = G.W = B.W = 1.0f;
 		}
 		else
 		{
-			SetBlendState( !bAlphaChannel ? SE_BLEND_Opaque : SE_BLEND_Translucent ); // If alpha channel is disabled, do not allow alpha blending
+			SetBlendState(RHICmdList, !bAlphaChannel ? SE_BLEND_Opaque : SE_BLEND_Translucent); // If alpha channel is disabled, do not allow alpha blending
 
 			// Determine the red, green, blue and alpha components of their respective weights to enable that colours prominence
 			R.X = bRedChannel ? 1.0f : 0.0f;
@@ -437,11 +435,11 @@ void FBatchedElements::PrepareShaders(
 	    if (bHitTesting)
 	    {
 		    TShaderMapRef<FSimpleElementHitProxyPS> HitTestingPixelShader(GetGlobalShaderMap());
-			SetGlobalBoundShaderState(HitTestingBoundShaderState, GSimpleElementVertexDeclaration.VertexDeclarationRHI, 
+			SetGlobalBoundShaderState(RHICmdList, HitTestingBoundShaderState, GSimpleElementVertexDeclaration.VertexDeclarationRHI,
 				*VertexShader, *HitTestingPixelShader);
 
 			HitTestingPixelShader->SetParameters(RHICmdList, Texture);
-		    SetHitTestingBlendState( BlendMode);
+			SetHitTestingBlendState(RHICmdList, BlendMode);
 		    
 	    }
 	    else
@@ -450,10 +448,10 @@ void FBatchedElements::PrepareShaders(
 		    {
 			    // use clip() in the shader instead of alpha testing as cards that don't support floating point blending
 			    // also don't support alpha testing to floating point render targets
-				RHISetBlendState(TStaticBlendState<>::GetRHI());
+				RHICmdList.SetBlendState(TStaticBlendState<>::GetRHI());
 
 			    TShaderMapRef<FSimpleElementMaskedGammaPS> MaskedPixelShader(GetGlobalShaderMap());
-				SetGlobalBoundShaderState( MaskedBoundShaderState, GSimpleElementVertexDeclaration.VertexDeclarationRHI, 
+				SetGlobalBoundShaderState(RHICmdList, MaskedBoundShaderState, GSimpleElementVertexDeclaration.VertexDeclarationRHI,
 					*VertexShader, *MaskedPixelShader);
 
 				MaskedPixelShader->SetEditorCompositingParameters(RHICmdList, View, DepthTexture );
@@ -471,17 +469,17 @@ void FBatchedElements::PrepareShaders(
 				    BlendMode == SE_BLEND_TranslucentDistanceFieldShadowed)
 			    {
 				    // enable alpha blending and disable clip ref value for translucent rendering
-				    RHISetBlendState(TStaticBlendState<CW_RGB,BO_Add,BF_SourceAlpha,BF_InverseSourceAlpha>::GetRHI());
+					RHICmdList.SetBlendState(TStaticBlendState<CW_RGB, BO_Add, BF_SourceAlpha, BF_InverseSourceAlpha>::GetRHI());
 				    AlphaRefVal = 0.0f;
 			    }
 			    else
 			    {
 				    // clip is done in shader so just render opaque
-				    RHISetBlendState(TStaticBlendState<>::GetRHI());
+					RHICmdList.SetBlendState(TStaticBlendState<>::GetRHI());
 			    }
 			    
 			    TShaderMapRef<FSimpleElementDistanceFieldGammaPS> DistanceFieldPixelShader(GetGlobalShaderMap());			
-				SetGlobalBoundShaderState( DistanceFieldBoundShaderState, GSimpleElementVertexDeclaration.VertexDeclarationRHI, 
+				SetGlobalBoundShaderState(RHICmdList, DistanceFieldBoundShaderState, GSimpleElementVertexDeclaration.VertexDeclarationRHI,
 					*VertexShader, *DistanceFieldPixelShader );			
 
 			    // @todo - expose these as options for batch rendering
@@ -511,19 +509,19 @@ void FBatchedElements::PrepareShaders(
 			else if(BlendMode >= SE_BLEND_RGBA_MASK_START && BlendMode <= SE_BLEND_RGBA_MASK_END)
 			{
 				TShaderMapRef<FSimpleElementColorChannelMaskPS> ColorChannelMaskPixelShader(GetGlobalShaderMap());
-				SetGlobalBoundShaderState( ColorChannelMaskShaderState, GSimpleElementVertexDeclaration.VertexDeclarationRHI, 
+				SetGlobalBoundShaderState(RHICmdList, ColorChannelMaskShaderState, GSimpleElementVertexDeclaration.VertexDeclarationRHI,
 					*VertexShader, *ColorChannelMaskPixelShader );
 			
 				ColorChannelMaskPixelShader->SetParameters(RHICmdList, Texture, ColorWeights, GammaToUse );
 			}
 			else
 		    {
-			    SetBlendState(BlendMode);
+				SetBlendState(RHICmdList, BlendMode);
     
 			    if (FMath::Abs(Gamma - 1.0f) < KINDA_SMALL_NUMBER)
 			    {
 				    TShaderMapRef<FSimpleElementPS> RegularPixelShader(GetGlobalShaderMap());
-					SetGlobalBoundShaderState( SimpleBoundShaderState, GSimpleElementVertexDeclaration.VertexDeclarationRHI, 
+					SetGlobalBoundShaderState(RHICmdList, SimpleBoundShaderState, GSimpleElementVertexDeclaration.VertexDeclarationRHI,
 						*VertexShader, *RegularPixelShader );
 
 					RegularPixelShader->SetEditorCompositingParameters(RHICmdList, View, DepthTexture);
@@ -532,7 +530,7 @@ void FBatchedElements::PrepareShaders(
 			    else
 			    {
 				    TShaderMapRef<FSimpleElementGammaPS> RegularPixelShader(GetGlobalShaderMap());
-					SetGlobalBoundShaderState( RegularBoundShaderState, GSimpleElementVertexDeclaration.VertexDeclarationRHI, 
+					SetGlobalBoundShaderState(RHICmdList, RegularBoundShaderState, GSimpleElementVertexDeclaration.VertexDeclarationRHI,
 						*VertexShader, *RegularPixelShader );
 
 				    RegularPixelShader->SetParameters(RHICmdList, Texture,Gamma,BlendMode);
@@ -552,7 +550,7 @@ void FBatchedElements::PrepareShaders(
 PRAGMA_DISABLE_OPTIMIZATION
 #endif
 
-void FBatchedElements::DrawPointElements( const FMatrix& Transform, const uint32 ViewportSizeX, const uint32 ViewportSizeY, const FVector& CameraX, const FVector& CameraY ) const
+void FBatchedElements::DrawPointElements(FRHICommandList& RHICmdList, const FMatrix& Transform, const uint32 ViewportSizeX, const uint32 ViewportSizeY, const FVector& CameraX, const FVector& CameraY) const
 {
 	// Draw the point elements.
 	if( Points.Num() > 0 )
@@ -563,7 +561,7 @@ void FBatchedElements::DrawPointElements( const FMatrix& Transform, const uint32
 		const int32 NumPoints = Points.Num();
 		const int32 NumTris = NumPoints * 2;
 		const int32 NumVertices = NumTris * 3;
-		RHIBeginDrawPrimitiveUP(PT_TriangleList, NumTris, NumVertices, sizeof(FSimpleElementVertex), VerticesPtr);
+		RHICmdList.BeginDrawPrimitiveUP(PT_TriangleList, NumTris, NumVertices, sizeof(FSimpleElementVertex), VerticesPtr);
 		FSimpleElementVertex* PointVertices = (FSimpleElementVertex*)VerticesPtr;
 
 		int32 VertIdx = 0;
@@ -589,12 +587,12 @@ void FBatchedElements::DrawPointElements( const FMatrix& Transform, const uint32
 		}
 
 		// Draw the sprite.
-		RHIEndDrawPrimitiveUP();
+		RHICmdList.EndDrawPrimitiveUP();
 	}
 }
 
 
-bool FBatchedElements::Draw(bool bNeedToSwitchVerticalAxis, const FMatrix& Transform,uint32 ViewportSizeX,uint32 ViewportSizeY,bool bHitTesting,float Gamma, const FSceneView* View, FTexture2DRHIRef DepthTexture ) const
+bool FBatchedElements::Draw(FRHICommandList& RHICmdList, bool bNeedToSwitchVerticalAxis, const FMatrix& Transform, uint32 ViewportSizeX, uint32 ViewportSizeY, bool bHitTesting, float Gamma, const FSceneView* View, FTexture2DRHIRef DepthTexture) const
 {
 	if( HasPrimsToDraw() )
 	{
@@ -603,7 +601,7 @@ bool FBatchedElements::Draw(bool bNeedToSwitchVerticalAxis, const FMatrix& Trans
 		FVector CameraY = InvTransform.TransformVector(FVector(0,1,0)).SafeNormal();
 		FVector CameraZ = InvTransform.TransformVector(FVector(0,0,1)).SafeNormal();
 
-		RHISetRasterizerState(TStaticRasterizerState<FM_Solid,CM_None>::GetRHI());
+		RHICmdList.SetRasterizerState(TStaticRasterizerState<FM_Solid, CM_None>::GetRHI());
 
 		if( LineVertices.Num() > 0 || Points.Num() > 0 || ThickLines.Num() > 0 || WireTris.Num() > 0 )
 		{
@@ -611,7 +609,7 @@ bool FBatchedElements::Draw(bool bNeedToSwitchVerticalAxis, const FMatrix& Trans
 			FBatchedElementParameters* BatchedElementParameters = NULL;
 
 			// Set the appropriate pixel shader parameters & shader state for the non-textured elements.
-			PrepareShaders( SE_BLEND_Opaque, Transform, bNeedToSwitchVerticalAxis, BatchedElementParameters, GWhiteTexture, bHitTesting, Gamma, NULL, View, DepthTexture);
+			PrepareShaders(RHICmdList, SE_BLEND_Opaque, Transform, bNeedToSwitchVerticalAxis, BatchedElementParameters, GWhiteTexture, bHitTesting, Gamma, NULL, View, DepthTexture);
 
 			// Draw the line elements.
 			if( LineVertices.Num() > 0 )
@@ -627,13 +625,13 @@ bool FBatchedElements::Draw(bool bNeedToSwitchVerticalAxis, const FMatrix& Trans
 				while( MinVertex < TotalVerts )
 				{
 					int32 NumLinePrims = FMath::Min( MaxVerticesAllowed, TotalVerts - MinVertex ) / 2;
-					RHIDrawPrimitiveUP(PT_LineList,NumLinePrims,&LineVertices[MinVertex],sizeof(FSimpleElementVertex));
+					DrawPrimitiveUP(RHICmdList, PT_LineList, NumLinePrims, &LineVertices[MinVertex], sizeof(FSimpleElementVertex));
 					MinVertex += NumLinePrims * 2;
 				}
 			}
 
 			// Draw points
-			DrawPointElements( Transform, ViewportSizeX, ViewportSizeY, CameraX, CameraY );
+			DrawPointElements(RHICmdList, Transform, ViewportSizeX, ViewportSizeY, CameraX, CameraY);
 
 			if ( ThickLines.Num() > 0 )
 			{
@@ -667,10 +665,10 @@ bool FBatchedElements::Draw(bool bNeedToSwitchVerticalAxis, const FMatrix& Trans
 					const bool bEnableMSAA = true;
 					const bool bEnableLineAA = false;
 					FRasterizerStateInitializerRHI Initializer = { FM_Solid, CM_None, 0, DepthBiasThisBatch, bEnableMSAA, bEnableLineAA };
-					RHISetRasterizerState(RHICreateRasterizerState(Initializer).GetReference());
+					RHICmdList.SetRasterizerState(RHICreateRasterizerState(Initializer).GetReference());
 
 					void* ThickVertexData = NULL;
-					RHIBeginDrawPrimitiveUP(PT_TriangleList,8*NumLinesThisBatch,8*3*NumLinesThisBatch,sizeof(FSimpleElementVertex),ThickVertexData);
+					RHICmdList.BeginDrawPrimitiveUP(PT_TriangleList, 8 * NumLinesThisBatch, 8 * 3 * NumLinesThisBatch, sizeof(FSimpleElementVertex), ThickVertexData);
 					FSimpleElementVertex* ThickVertices = (FSimpleElementVertex*)ThickVertexData;
 					check(ThickVertices);
 
@@ -741,10 +739,10 @@ bool FBatchedElements::Draw(bool bNeedToSwitchVerticalAxis, const FMatrix& Trans
 
 						ThickVertices += 24;
 					}
-					RHIEndDrawPrimitiveUP();	
+					RHICmdList.EndDrawPrimitiveUP();
 				}
 
-				RHISetRasterizerState(TStaticRasterizerState<FM_Solid,CM_None>::GetRHI());
+				RHICmdList.SetRasterizerState(TStaticRasterizerState<FM_Solid, CM_None>::GetRHI());
 			}
 			// Draw the wireframe triangles.
 			if (WireTris.Num() > 0)
@@ -779,14 +777,14 @@ bool FBatchedElements::Draw(bool bNeedToSwitchVerticalAxis, const FMatrix& Trans
 					}
 
 					Initializer.DepthBias = DepthBias;
-					RHISetRasterizerState(RHICreateRasterizerState(Initializer).GetReference());
+					RHICmdList.SetRasterizerState(RHICreateRasterizerState(Initializer).GetReference());
 
 					int32 NumTris = MaxTri - MinTri;
-					RHIDrawPrimitiveUP(PT_TriangleList, NumTris, &WireTriVerts[MinTri*3], sizeof(FSimpleElementVertex));
+					DrawPrimitiveUP(RHICmdList, PT_TriangleList, NumTris, &WireTriVerts[MinTri * 3], sizeof(FSimpleElementVertex));
 					MinTri = MaxTri;
 				}
 
-				RHISetRasterizerState(TStaticRasterizerState<FM_Solid,CM_None>::GetRHI());
+				RHICmdList.SetRasterizerState(TStaticRasterizerState<FM_Solid, CM_None>::GetRHI());
 			}
 		}
 
@@ -822,8 +820,8 @@ bool FBatchedElements::Draw(bool bNeedToSwitchVerticalAxis, const FMatrix& Trans
 					//New batch, draw previous and clear
 					const int32 VertexCount = SpriteList.Num();
 					const int32 PrimCount = VertexCount / 3;
-					PrepareShaders( CurrentBlendMode, Transform, bNeedToSwitchVerticalAxis, BatchedElementParameters, CurrentTexture, bHitTesting, Gamma, NULL, View, DepthTexture);
-					RHIDrawPrimitiveUP(PT_TriangleList, PrimCount, SpriteList.GetTypedData(), sizeof(FSimpleElementVertex));
+					PrepareShaders(RHICmdList, CurrentBlendMode, Transform, bNeedToSwitchVerticalAxis, BatchedElementParameters, CurrentTexture, bHitTesting, Gamma, NULL, View, DepthTexture);
+					DrawPrimitiveUP(RHICmdList, PT_TriangleList, PrimCount, SpriteList.GetTypedData(), sizeof(FSimpleElementVertex));
 
 					SpriteList.Empty(6);
 					CurrentTexture = Sprite.Texture;
@@ -856,8 +854,8 @@ bool FBatchedElements::Draw(bool bNeedToSwitchVerticalAxis, const FMatrix& Trans
 				//Draw last batch
 				const int32 VertexCount = SpriteList.Num();
 				const int32 PrimCount = VertexCount / 3;
-				PrepareShaders( CurrentBlendMode, Transform, bNeedToSwitchVerticalAxis, BatchedElementParameters, CurrentTexture, bHitTesting, Gamma, NULL, View, DepthTexture );
-				RHIDrawPrimitiveUP(PT_TriangleList, PrimCount, SpriteList.GetTypedData(), sizeof(FSimpleElementVertex));
+				PrepareShaders(RHICmdList, CurrentBlendMode, Transform, bNeedToSwitchVerticalAxis, BatchedElementParameters, CurrentTexture, bHitTesting, Gamma, NULL, View, DepthTexture);
+				DrawPrimitiveUP(RHICmdList, PT_TriangleList, PrimCount, SpriteList.GetTypedData(), sizeof(FSimpleElementVertex));
 			}
 		}
 
@@ -869,10 +867,11 @@ bool FBatchedElements::Draw(bool bNeedToSwitchVerticalAxis, const FMatrix& Trans
 				const FBatchedMeshElement& MeshElement = MeshElements[MeshIndex];
 
 				// Set the appropriate pixel shader for the mesh.
-				PrepareShaders(MeshElement.BlendMode, Transform, bNeedToSwitchVerticalAxis, MeshElement.BatchedElementParameters, MeshElement.Texture, bHitTesting, Gamma, &MeshElement.GlowInfo);
+				PrepareShaders(RHICmdList, MeshElement.BlendMode, Transform, bNeedToSwitchVerticalAxis, MeshElement.BatchedElementParameters, MeshElement.Texture, bHitTesting, Gamma, &MeshElement.GlowInfo);
 
 				// Draw the mesh.
-				RHIDrawIndexedPrimitiveUP(
+				DrawIndexedPrimitiveUP(
+					RHICmdList,
 					PT_TriangleList,
 					0,
 					MeshElement.MaxVertex - MeshElement.MinVertex + 1,
@@ -893,10 +892,11 @@ bool FBatchedElements::Draw(bool bNeedToSwitchVerticalAxis, const FMatrix& Trans
 				FBatchedElementParameters* BatchedElementParameters = NULL;
 
 				// Set the appropriate pixel shader for the mesh.
-				PrepareShaders(MeshElement.BlendMode,Transform, bNeedToSwitchVerticalAxis,BatchedElementParameters,MeshElement.Texture,bHitTesting,Gamma);
+				PrepareShaders(RHICmdList, MeshElement.BlendMode, Transform, bNeedToSwitchVerticalAxis, BatchedElementParameters, MeshElement.Texture, bHitTesting, Gamma);
 
 				// Draw the mesh.
-				RHIDrawPrimitiveUP(
+				DrawPrimitiveUP(
+					RHICmdList,
 					PT_QuadList,
 					MeshElement.Vertices.Num() / 4,
 					MeshElement.Vertices.GetTypedData(),
diff --git a/Engine/Source/Runtime/Engine/Private/CanvasRenderTarget2D.cpp b/Engine/Source/Runtime/Engine/Private/CanvasRenderTarget2D.cpp
--- a/Engine/Source/Runtime/Engine/Private/CanvasRenderTarget2D.cpp
+++ b/Engine/Source/Runtime/Engine/Private/CanvasRenderTarget2D.cpp
@@ -44,8 +44,11 @@ void UCanvasRenderTarget2D::UpdateResource()
 		TextureRenderTarget,
 		(FTextureRenderTarget2DResource*)GameThread_GetRenderTargetResource(),
 		{
-			RHISetRenderTarget(TextureRenderTarget->GetRenderTargetTexture(), FTexture2DRHIRef());
-			RHISetViewport(0, 0, 0.0f, TextureRenderTarget->GetSizeXY().X, TextureRenderTarget->GetSizeXY().Y, 1.0f);
+			//@todo-rco: RHIPacketList
+			FRHICommandList& RHICmdList = FRHICommandList::GetNullRef();
+
+			SetRenderTarget(RHICmdList, TextureRenderTarget->GetRenderTargetTexture(), FTexture2DRHIRef());
+			RHICmdList.SetViewport(0, 0, 0.0f, TextureRenderTarget->GetSizeXY().X, TextureRenderTarget->GetSizeXY().Y, 1.0f);
 		}
 	);
 
diff --git a/Engine/Source/Runtime/Engine/Private/NavMeshRenderingHelpers.h b/Engine/Source/Runtime/Engine/Private/NavMeshRenderingHelpers.h
--- a/Engine/Source/Runtime/Engine/Private/NavMeshRenderingHelpers.h
+++ b/Engine/Source/Runtime/Engine/Private/NavMeshRenderingHelpers.h
@@ -310,7 +310,11 @@ public:
 
 			const bool bNeedToSwitchVerticalAxis = IsES2Platform(GRHIShaderPlatform) && !IsPCPlatform(GRHIShaderPlatform);
 			const FTexture2DRHIRef DepthTexture;
+			//@todo-rco: RHIPacketList
+			FRHICommandList& RHICmdList = FRHICommandList::GetNullRef();
+
 			BatchedElements.Draw(
+				RHICmdList,
 				bNeedToSwitchVerticalAxis,
 				View->ViewProjectionMatrix,
 				View->ViewRect.Width(),
@@ -322,6 +326,7 @@ public:
 				);
 
 			ProxyData->BatchedElements.Draw(
+				RHICmdList,
 				bNeedToSwitchVerticalAxis,
 				View->ViewProjectionMatrix,
 				View->ViewRect.Width(),
diff --git a/Engine/Source/Runtime/Engine/Private/Particles/ParticleCurveTexture.cpp b/Engine/Source/Runtime/Engine/Private/Particles/ParticleCurveTexture.cpp
--- a/Engine/Source/Runtime/Engine/Private/Particles/ParticleCurveTexture.cpp
+++ b/Engine/Source/Runtime/Engine/Private/Particles/ParticleCurveTexture.cpp
@@ -189,16 +189,16 @@ static void InjectCurves(
 
 	RHICmdList.CheckIsNull();
 
-	RHISetRenderTarget( CurveTextureTargetRHI, FTextureRHIParamRef() );
-	RHISetScissorRect( false, 0, 0, 0, 0 );
-	RHISetViewport( 0, 0, 0.0f, GParticleCurveTextureSizeX, GParticleCurveTextureSizeY, 1.0f );
-	RHISetDepthStencilState( TStaticDepthStencilState<false,CF_Always>::GetRHI() );
-	RHISetRasterizerState( TStaticRasterizerState<FM_Solid,CM_None>::GetRHI() );
-	RHISetBlendState( TStaticBlendState<>::GetRHI() );
+	SetRenderTarget(RHICmdList, CurveTextureTargetRHI, FTextureRHIParamRef());
+	RHICmdList.SetScissorRect(false, 0, 0, 0, 0);
+	RHICmdList.SetViewport(0, 0, 0.0f, GParticleCurveTextureSizeX, GParticleCurveTextureSizeY, 1.0f);
+	RHICmdList.SetDepthStencilState(TStaticDepthStencilState<false, CF_Always>::GetRHI());
+	RHICmdList.SetRasterizerState(TStaticRasterizerState<FM_Solid, CM_None>::GetRHI());
+	RHICmdList.SetBlendState(TStaticBlendState<>::GetRHI());
 
 	if (bFirstCall)
 	{
-		RHIClear(true,FLinearColor::Blue,false,0.0f,false,0,FIntRect());
+		RHICmdList.Clear(true, FLinearColor::Blue, false, 0.0f, false, 0, FIntRect());
 		bFirstCall = false;
 	}
 
@@ -227,8 +227,10 @@ static void InjectCurves(
 		TShaderMapRef<FParticleCurveInjectionPS> PixelShader( GetGlobalShaderMap() );
 
 		// Bound shader state.
+		RHICmdList.CheckIsNull(); // need new approach for "static FGlobalBoundShaderState" for parallel rendering
 		static FGlobalBoundShaderState BoundShaderState;
 		SetGlobalBoundShaderState(
+			RHICmdList,
 			BoundShaderState,
 			GParticleCurveInjectionVertexDeclaration.VertexDeclarationRHI,
 			*VertexShader,
@@ -239,7 +241,7 @@ static void InjectCurves(
 		VertexShader->SetParameters(RHICmdList, CurveOffset );
 
 		// Stream 0: New particles.
-		RHISetStreamSource(
+		RHICmdList.SetStreamSource(
 			0,
 			ScratchVertexBufferRHI,
 			/*Stride=*/ sizeof(FColor),
@@ -247,7 +249,7 @@ static void InjectCurves(
 			);
 
 		// Stream 1: TexCoord.
-		RHISetStreamSource(
+		RHICmdList.SetStreamSource(
 			1,
 			GParticleTexCoordVertexBuffer.VertexBufferRHI,
 			/*Stride=*/ sizeof(FVector2D),
@@ -255,7 +257,7 @@ static void InjectCurves(
 			);
 
 		// Inject particles.
-		RHIDrawIndexedPrimitive(
+		RHICmdList.DrawIndexedPrimitive(
 			GParticleIndexBuffer.IndexBufferRHI,
 			PT_TriangleList,
 			/*BaseVertexIndex=*/ 0,
@@ -267,7 +269,7 @@ static void InjectCurves(
 			);
 	}
 
-	RHICopyToResolveTarget( CurveTextureTargetRHI, CurveTextureRHI, /*bKeepOriginalSurface=*/ false, FResolveParams() );
+	RHICmdList.CopyToResolveTarget(CurveTextureTargetRHI, CurveTextureRHI, /*bKeepOriginalSurface=*/ false, FResolveParams());
 }
 
 /*------------------------------------------------------------------------------
diff --git a/Engine/Source/Runtime/Engine/Private/Particles/ParticleGpuSimulation.cpp b/Engine/Source/Runtime/Engine/Private/Particles/ParticleGpuSimulation.cpp
--- a/Engine/Source/Runtime/Engine/Private/Particles/ParticleGpuSimulation.cpp
+++ b/Engine/Source/Runtime/Engine/Private/Particles/ParticleGpuSimulation.cpp
@@ -1227,8 +1227,10 @@ void ExecuteSimulationCommands(
 	TShaderMapRef<TParticleSimulationPS<bUseDepthBufferCollision> > PixelShader(GetGlobalShaderMap());
 
 	// Bound shader state.
+	RHICmdList.CheckIsNull(); // need new approach for "static FGlobalBoundShaderState" for parallel rendering
 	static FGlobalBoundShaderState BoundShaderState;
 	SetGlobalBoundShaderState(
+		RHICmdList,
 		BoundShaderState,
 		GParticleTileVertexDeclaration.VertexDeclarationRHI,
 		*VertexShader,
@@ -1280,9 +1282,14 @@ void ClearTiles(const TArray<uint32>& Tiles)
 	TShaderMapRef<FParticleTileVS> VertexShader( GetGlobalShaderMap() );
 	TShaderMapRef<FParticleSimulationClearPS> PixelShader( GetGlobalShaderMap() );
 
+	//@todo-rco: RHIPacketList
+	FRHICommandList& RHICmdList = FRHICommandList::GetNullRef();
+
 	// Bound shader state.
+	RHICmdList.CheckIsNull(); // need new approach for "static FGlobalBoundShaderState" for parallel rendering
 	static FGlobalBoundShaderState BoundShaderState;
 	SetGlobalBoundShaderState(
+		RHICmdList,
 		BoundShaderState,
 		GParticleTileVertexDeclaration.VertexDeclarationRHI,
 		*VertexShader,
@@ -1297,7 +1304,7 @@ void ClearTiles(const TArray<uint32>& Tiles)
 		const uint32* TilesPtr = Tiles.GetTypedData() + FirstTile;
 		BuildTileVertexBuffer( BufferParam, TilesPtr, TilesThisDrawCall );
 		//@todo-rco: RHIPacketList
-		VertexShader->SetParameters(FRHICommandList::GetNullRef(), ShaderParam);
+		VertexShader->SetParameters(RHICmdList, ShaderParam);
 		DrawAlignedParticleTiles( ComputeAlignedTileCount(TilesThisDrawCall) );
 		TileCount -= TilesThisDrawCall;
 		FirstTile += TilesThisDrawCall;
@@ -1506,6 +1513,9 @@ void InjectNewParticles( const TArray<FNewParticle>& NewParticles )
 	int32 ParticleCount = NewParticles.Num();
 	int32 FirstParticle = 0;
 
+	//@todo-rco: RHIPacketList
+	FRHICommandList& RHICmdList = FRHICommandList::GetNullRef();
+
 	while ( ParticleCount > 0 )
 	{
 		// Copy new particles in to the vertex buffer.
@@ -1522,8 +1532,10 @@ void InjectNewParticles( const TArray<FNewParticle>& NewParticles )
 		TShaderMapRef<FParticleInjectionPS> PixelShader( GetGlobalShaderMap() );
 
 		// Bound shader state.
+		RHICmdList.CheckIsNull(); // need new approach for "static FGlobalBoundShaderState" for parallel rendering
 		static FGlobalBoundShaderState BoundShaderState;
 		SetGlobalBoundShaderState(
+			RHICmdList,
 			BoundShaderState,
 			GParticleInjectionVertexDeclaration.VertexDeclarationRHI,
 			*VertexShader,
@@ -1734,30 +1746,32 @@ static void VisualizeGPUSimulation(
 	const float DisplayOffsetX = 60.0f;
 	const float DisplayOffsetY = 60.0f;
 
+	//@todo-rco: RHIPacketList
+	FRHICommandList& RHICmdList = FRHICommandList::GetNullRef();
+
 	// Setup render states.
 	FIntPoint TargetSize = RenderTarget->GetSizeXY();
-	RHISetRenderTarget( RenderTarget->GetRenderTargetTexture(), FTextureRHIParamRef() );
-	RHISetViewport( 0, 0, 0.0f, TargetSize.X, TargetSize.Y, 1.0f );
-	RHISetDepthStencilState( TStaticDepthStencilState<false,CF_Always>::GetRHI() );
-	RHISetRasterizerState( TStaticRasterizerState<FM_Solid,CM_None>::GetRHI() );
-	RHISetBlendState( TStaticBlendState<>::GetRHI() );
+	SetRenderTarget(RHICmdList, RenderTarget->GetRenderTargetTexture(), FTextureRHIParamRef());
+	RHICmdList.SetViewport(0, 0, 0.0f, TargetSize.X, TargetSize.Y, 1.0f);
+	RHICmdList.SetDepthStencilState(TStaticDepthStencilState<false, CF_Always>::GetRHI());
+	RHICmdList.SetRasterizerState(TStaticRasterizerState<FM_Solid, CM_None>::GetRHI());
+	RHICmdList.SetBlendState(TStaticBlendState<>::GetRHI());
 
 	// Grab shaders.
 	TShaderMapRef<FParticleSimVisualizeVS> VertexShader( GetGlobalShaderMap() );
 	TShaderMapRef<FParticleSimVisualizePS> PixelShader( GetGlobalShaderMap() );
 
 	// Bound shader state.
+	RHICmdList.CheckIsNull(); // need new approach for "static FGlobalBoundShaderState" for parallel rendering
 	static FGlobalBoundShaderState BoundShaderState;
 	SetGlobalBoundShaderState(
+		RHICmdList,
 		BoundShaderState,
 		GParticleSimVisualizeVertexDeclaration.VertexDeclarationRHI,
 		*VertexShader,
 		*PixelShader
 		);
 
-	//@todo-rco: RHIPacketList
-	FRHICommandList& RHICmdList = FRHICommandList::GetNullRef();
-
 	// Parameters for the visualization.
 	FParticleSimVisualizeParameters Parameters;
 	Parameters.ScaleBias.X = 2.0f * DisplaySizeX / (float)TargetSize.X;
@@ -1771,7 +1785,7 @@ static void VisualizeGPUSimulation(
 	const int32 VertexStride = sizeof(FVector2D);
 	
 	// Bind vertex stream.
-	RHISetStreamSource(
+	RHICmdList.SetStreamSource(
 		0,
 		GParticleTexCoordVertexBuffer.VertexBufferRHI,
 		VertexStride,
@@ -1779,7 +1793,7 @@ static void VisualizeGPUSimulation(
 		);
 
 	// Draw.
-	RHIDrawIndexedPrimitive(
+	RHICmdList.DrawIndexedPrimitive(
 		GParticleIndexBuffer.IndexBufferRHI,
 		PT_TriangleList,
 		/*BaseVertexIndex=*/ 0,
@@ -4005,23 +4019,26 @@ void FFXSystem::SimulateGPUParticles(
 	FParticleStateTextures& CurrentStateTextures = ParticleSimulationResources->GetCurrentStateTextures();
 	FParticleStateTextures& PrevStateTextures = ParticleSimulationResources->GetPreviousStateTextures();
 
+	//@todo-rco: RHIPacketList
+	FRHICommandList& RHICmdList = FRHICommandList::GetNullRef();
+
 	// On some platforms, the textures are filled with garbage after creation, so we need to clear them to black the first time we use them
 	if ( !CurrentStateTextures.bTexturesCleared )
 	{
-		RHISetRenderTarget( CurrentStateTextures.PositionTextureTargetRHI, FTextureRHIRef() );
-		RHIClear( true, FLinearColor::Black, false, 1.0f, false, 0, FIntRect() );
-		RHISetRenderTarget( CurrentStateTextures.VelocityTextureTargetRHI, FTextureRHIRef() );
-		RHIClear( true, FLinearColor::Black, false, 1.0f, false, 0, FIntRect() );
+		SetRenderTarget(RHICmdList, CurrentStateTextures.PositionTextureTargetRHI, FTextureRHIRef());
+		RHICmdList.Clear(true, FLinearColor::Black, false, 1.0f, false, 0, FIntRect());
+		SetRenderTarget(RHICmdList, CurrentStateTextures.VelocityTextureTargetRHI, FTextureRHIRef());
+		RHICmdList.Clear(true, FLinearColor::Black, false, 1.0f, false, 0, FIntRect());
 
 		CurrentStateTextures.bTexturesCleared = true;
 	}
 
 	if ( !PrevStateTextures.bTexturesCleared )
 	{
-		RHISetRenderTarget( PrevStateTextures.PositionTextureTargetRHI, FTextureRHIRef() );
-		RHIClear( true, FLinearColor::Black, false, 1.0f, false, 0, FIntRect() );
-		RHISetRenderTarget( PrevStateTextures.VelocityTextureTargetRHI, FTextureRHIRef() );
-		RHIClear( true, FLinearColor::Black, false, 1.0f, false, 0, FIntRect() );
+		SetRenderTarget(RHICmdList, PrevStateTextures.PositionTextureTargetRHI, FTextureRHIRef());
+		RHICmdList.Clear(true, FLinearColor::Black, false, 1.0f, false, 0, FIntRect());
+		SetRenderTarget(RHICmdList, PrevStateTextures.VelocityTextureTargetRHI, FTextureRHIRef());
+		RHICmdList.Clear(true, FLinearColor::Black, false, 1.0f, false, 0, FIntRect());
 
 		PrevStateTextures.bTexturesCleared = true;
 	}
@@ -4033,11 +4050,11 @@ void FFXSystem::SimulateGPUParticles(
 		CurrentStateTextures.VelocityTextureTargetRHI,
 		CurrentStateTextures.PositionZWTextureTargetRHI
 	};
-	RHISetRenderTargets(2, RenderTargets, FTextureRHIParamRef(), 0, NULL);
-	RHISetViewport(0, 0, 0.0f, GParticleSimulationTextureSizeX, GParticleSimulationTextureSizeY, 1.0f);
-	RHISetDepthStencilState(TStaticDepthStencilState<false,CF_Always>::GetRHI());
-	RHISetRasterizerState(TStaticRasterizerState<FM_Solid,CM_None>::GetRHI());
-	RHISetBlendState(TStaticBlendState<>::GetRHI());
+	SetRenderTargets(RHICmdList, 2, RenderTargets, FTextureRHIParamRef(), 0, NULL);
+	RHICmdList.SetViewport(0, 0, 0.0f, GParticleSimulationTextureSizeX, GParticleSimulationTextureSizeY, 1.0f);
+	RHICmdList.SetDepthStencilState(TStaticDepthStencilState<false, CF_Always>::GetRHI());
+	RHICmdList.SetRasterizerState(TStaticRasterizerState<FM_Solid, CM_None>::GetRHI());
+	RHICmdList.SetBlendState(TStaticBlendState<>::GetRHI());
 
 	// Simulations that don't use vector fields can share some state.
 	FVectorFieldUniformBufferRef EmptyVectorFieldUniformBuffer;
@@ -4205,20 +4222,20 @@ void FFXSystem::SimulateGPUParticles(
 			ParticleSimulationResources->SimulationAttributesTexture.TextureTargetRHI,
 			CurrentStateTextures.PositionZWTextureTargetRHI
 		};
-		RHISetRenderTargets(4, InjectRenderTargets, FTextureRHIParamRef(), 0, NULL);
-		RHISetViewport(0, 0, 0.0f, GParticleSimulationTextureSizeX, GParticleSimulationTextureSizeY, 1.0f);
+		SetRenderTargets(RHICmdList, 4, InjectRenderTargets, FTextureRHIParamRef(), 0, NULL);
+		RHICmdList.SetViewport(0, 0, 0.0f, GParticleSimulationTextureSizeX, GParticleSimulationTextureSizeY, 1.0f);
 
 		// Inject particles.
 		InjectNewParticles(NewParticles);
 
 		// Resolve attributes textures. State textures are resolved later.
-		RHICopyToResolveTarget( 
+		RHICmdList.CopyToResolveTarget(
 			ParticleSimulationResources->RenderAttributesTexture.TextureTargetRHI,
 			ParticleSimulationResources->RenderAttributesTexture.TextureRHI,
 			/*bKeepOriginalSurface=*/ false,
 			FResolveParams()
 			);
-		RHICopyToResolveTarget( 
+		RHICmdList.CopyToResolveTarget(
 			ParticleSimulationResources->SimulationAttributesTexture.TextureTargetRHI,
 			ParticleSimulationResources->SimulationAttributesTexture.TextureRHI,
 			/*bKeepOriginalSurface=*/ false,
@@ -4231,11 +4248,11 @@ void FFXSystem::SimulateGPUParticles(
 	NewParticles.Reset();
 
 	// Resolve all textures.
-	RHICopyToResolveTarget(CurrentStateTextures.PositionTextureTargetRHI, CurrentStateTextures.PositionTextureRHI, /*bKeepOriginalSurface=*/ false, FResolveParams());
-	RHICopyToResolveTarget(CurrentStateTextures.VelocityTextureTargetRHI, CurrentStateTextures.VelocityTextureRHI, /*bKeepOriginalSurface=*/ false, FResolveParams());
+	RHICmdList.CopyToResolveTarget(CurrentStateTextures.PositionTextureTargetRHI, CurrentStateTextures.PositionTextureRHI, /*bKeepOriginalSurface=*/ false, FResolveParams());
+	RHICmdList.CopyToResolveTarget(CurrentStateTextures.VelocityTextureTargetRHI, CurrentStateTextures.VelocityTextureRHI, /*bKeepOriginalSurface=*/ false, FResolveParams());
 
 	// Clear render targets so we can safely read from them.
-	RHISetRenderTarget(FTextureRHIParamRef(),FTextureRHIParamRef());
+	SetRenderTarget(RHICmdList, FTextureRHIParamRef(), FTextureRHIParamRef());
 
 	// Stats.
 	if (Phase == EParticleSimulatePhase::Last)
diff --git a/Engine/Source/Runtime/Engine/Private/ScreenRendering.cpp b/Engine/Source/Runtime/Engine/Private/ScreenRendering.cpp
--- a/Engine/Source/Runtime/Engine/Private/ScreenRendering.cpp
+++ b/Engine/Source/Runtime/Engine/Private/ScreenRendering.cpp
@@ -21,15 +21,16 @@ FGlobalBoundShaderState ScreenBoundShaderState;
 /**
  * Draws a texture rectangle on the screen using normalized (-1 to 1) device coordinates.
  */
-void DrawScreenQuad( float X0, float Y0, float U0, float V0, float X1, float Y1, float U1, float V1, const FTexture* Texture )
+void DrawScreenQuad(FRHICommandList& RHICmdList, float X0, float Y0, float U0, float V0, float X1, float Y1, float U1, float V1, const FTexture* Texture)
 { 
 	// Set the screen shaders
 	TShaderMapRef<FScreenVS> ScreenVertexShader(GetGlobalShaderMap());
 	TShaderMapRef<FScreenPS> ScreenPixelShader(GetGlobalShaderMap());
-	SetGlobalBoundShaderState(ScreenBoundShaderState, GScreenVertexDeclaration.VertexDeclarationRHI, *ScreenVertexShader, *ScreenPixelShader);
 
-	//@todo-rco: RHIPacketList
-	ScreenPixelShader->SetParameters(FRHICommandList::GetNullRef(), Texture);
+	RHICmdList.CheckIsNull(); // need new approach for "static FGlobalBoundShaderState" for parallel rendering
+	SetGlobalBoundShaderState(RHICmdList, ScreenBoundShaderState, GScreenVertexDeclaration.VertexDeclarationRHI, *ScreenVertexShader, *ScreenPixelShader);
+
+	ScreenPixelShader->SetParameters(RHICmdList, Texture);
 
 	// Generate the vertices used
 	FScreenVertex Vertices[4];
@@ -54,18 +55,19 @@ void DrawScreenQuad( float X0, float Y0, float U0, float V0, float X1, float Y1,
 	Vertices[3].UV.X       = U0;
 	Vertices[3].UV.Y       = V1;
 
-	RHIDrawPrimitiveUP(PT_TriangleStrip,2,Vertices,sizeof(Vertices[0]));
+	DrawPrimitiveUP(RHICmdList, PT_TriangleStrip, 2, Vertices, sizeof(Vertices[0]));
 }
 
-void DrawNormalizedScreenQuad( float X, float Y, float U0, float V0, float SizeX, float SizeY, float U1, float V1, FIntPoint TargetSize, FTexture2DRHIRef TextureRHI )
+void DrawNormalizedScreenQuad(FRHICommandList& RHICmdList, float X, float Y, float U0, float V0, float SizeX, float SizeY, float U1, float V1, FIntPoint TargetSize, FTexture2DRHIRef TextureRHI)
 {
 	// Set the screen shaders 
 	TShaderMapRef<FScreenVS> ScreenVertexShader(GetGlobalShaderMap());
 	TShaderMapRef<FScreenPS> ScreenPixelShader(GetGlobalShaderMap());
-	SetGlobalBoundShaderState(ScreenBoundShaderState, GScreenVertexDeclaration.VertexDeclarationRHI, *ScreenVertexShader, *ScreenPixelShader);
+	RHICmdList.CheckIsNull(); // need new approach for "static FGlobalBoundShaderState" for parallel rendering
+	SetGlobalBoundShaderState(RHICmdList, ScreenBoundShaderState, GScreenVertexDeclaration.VertexDeclarationRHI, *ScreenVertexShader, *ScreenPixelShader);
 
 	/// ?
-	ScreenPixelShader->SetParameters(FRHICommandList::GetNullRef(), TStaticSamplerState<SF_Bilinear>::GetRHI(), TextureRHI);
+	ScreenPixelShader->SetParameters(RHICmdList, TStaticSamplerState<SF_Bilinear>::GetRHI(), TextureRHI);
 
 	// Generate the vertices used
 	FScreenVertex Vertices[4];
@@ -96,5 +98,5 @@ void DrawNormalizedScreenQuad( float X, float Y, float U0, float V0, float SizeX
 		Vertices[VertexIndex].Position.Y = (+1.0f - 2.0f * (Vertices[VertexIndex].Position.Y) / (float)TargetSize.Y);
 	}
 
-	RHIDrawPrimitiveUP(PT_TriangleStrip,2,Vertices,sizeof(Vertices[0]));
+	DrawPrimitiveUP(RHICmdList, PT_TriangleStrip, 2, Vertices, sizeof(Vertices[0]));
 }
diff --git a/Engine/Source/Runtime/Engine/Private/StaticBoundShaderState.cpp b/Engine/Source/Runtime/Engine/Private/StaticBoundShaderState.cpp
--- a/Engine/Source/Runtime/Engine/Private/StaticBoundShaderState.cpp
+++ b/Engine/Source/Runtime/Engine/Private/StaticBoundShaderState.cpp
@@ -92,16 +92,8 @@ void FGlobalBoundShaderStateResource::ReleaseRHI()
 	BoundShaderState.SafeRelease();
 }
 
-/**
- * SetGlobalBoundShaderState - sets the global bound shader state, also creates and caches it if necessary
- *
- * @param BoundShaderState - current bound shader state, will be updated if it wasn't a valid ref
- * @param VertexDeclaration - the vertex declaration to use in creating the new bound shader state
- * @param VertexShader - the vertex shader to use in creating the new bound shader state
- * @param PixelShader - the pixel shader to use in creating the new bound shader state
- * @param GeometryShader - the geometry shader to use in creating the new bound shader state (0 if not used)
- */
 void SetGlobalBoundShaderState(
+	FRHICommandList& RHICmdList,
 	FGlobalBoundShaderState& GlobalBoundShaderState,
 	FVertexDeclarationRHIParamRef VertexDeclaration,
 	FShader* VertexShader,
@@ -115,7 +107,7 @@ void SetGlobalBoundShaderState(
 	PixelShader->VerifyBoundUniformBufferParameters();
 	GeometryShader->VerifyBoundUniformBufferParameters();
 
-	RHISetBoundShaderState(
+	RHICmdList.SetBoundShaderState(
 		GlobalBoundShaderState.GetInitializedRHI(
 			VertexDeclaration,
 			GETSAFERHISHADER_VERTEX(VertexShader),
diff --git a/Engine/Source/Runtime/Engine/Private/UserInterface/Canvas.cpp b/Engine/Source/Runtime/Engine/Private/UserInterface/Canvas.cpp
--- a/Engine/Source/Runtime/Engine/Private/UserInterface/Canvas.cpp
+++ b/Engine/Source/Runtime/Engine/Private/UserInterface/Canvas.cpp
@@ -175,8 +175,12 @@ bool FCanvasBatchedElementRenderItem::Render( const FCanvas* Canvas )
 		{
 			const bool bNeedsToSwitchVerticalAxis = IsES2Platform(GRHIShaderPlatform) && !IsPCPlatform(GRHIShaderPlatform);
 
+			//@todo-rco: RHIPacketList
+			FRHICommandList& RHICmdList = FRHICommandList::GetNullRef();
+
 			// draw batched items
 			Data->BatchedElements.Draw(
+				RHICmdList,
 				bNeedsToSwitchVerticalAxis,
 				Data->Transform.GetMatrix(),
 				CanvasRenderTarget->GetSizeXY().X,
@@ -218,9 +222,12 @@ bool FCanvasBatchedElementRenderItem::Render( const FCanvas* Canvas )
 				FBatchedDrawParameters,Parameters,DrawParameters,
 			{
 				const bool bNeedsToSwitchVerticalAxis = IsES2Platform(GRHIShaderPlatform) && !IsPCPlatform(GRHIShaderPlatform);
+				//@todo-rco: RHIPacketList
+				FRHICommandList& RHICmdList = FRHICommandList::GetNullRef();
 
 				// draw batched items
 				Parameters.RenderData->BatchedElements.Draw(
+					RHICmdList,
 					bNeedsToSwitchVerticalAxis,
 					Parameters.RenderData->Transform.GetMatrix(),
 					Parameters.ViewportSizeX,
diff --git a/Engine/Source/Runtime/Engine/Public/BatchedElements.h b/Engine/Source/Runtime/Engine/Public/BatchedElements.h
--- a/Engine/Source/Runtime/Engine/Public/BatchedElements.h
+++ b/Engine/Source/Runtime/Engine/Public/BatchedElements.h
@@ -181,7 +181,7 @@ public:
 	 * @param View			Optional FSceneView for shaders that need access to view constants
 	 * @param DepthTexture	DepthTexture for manual depth testing with editor compositing in the pixel shader
 	 */
-	bool Draw(bool bNeedToSwitchVerticalAxis, const FMatrix& Transform,uint32 ViewportSizeX,uint32 ViewportSizeY,bool bHitTesting,float Gamma = 1.0f, const FSceneView* View = NULL, FTexture2DRHIRef DepthTexture = FTexture2DRHIRef()) const;
+	bool Draw(FRHICommandList& RHICmdList, bool bNeedToSwitchVerticalAxis, const FMatrix& Transform, uint32 ViewportSizeX, uint32 ViewportSizeY, bool bHitTesting, float Gamma = 1.0f, const FSceneView* View = NULL, FTexture2DRHIRef DepthTexture = FTexture2DRHIRef()) const;
 	
 	FORCEINLINE bool HasPrimsToDraw() const
 	{
@@ -215,7 +215,7 @@ private:
 	 * @param	CameraX		Local space normalized view direction X vector
 	 * @param	CameraY		Local space normalized view direction Y vector
 	 */
-	void DrawPointElements( const FMatrix& Transform, const uint32 ViewportSizeX, const uint32 ViewportSizeY, const FVector& CameraX, const FVector& CameraY ) const;
+	void DrawPointElements(FRHICommandList& RHICmdList, const FMatrix& Transform, const uint32 ViewportSizeX, const uint32 ViewportSizeY, const FVector& CameraX, const FVector& CameraY) const;
 
 	TArray<FSimpleElementVertex> LineVertices;
 
@@ -316,6 +316,7 @@ private:
 	 * Sets the appropriate vertex and pixel shader.
 	 */
 	void PrepareShaders(
+		FRHICommandList& RHICmdList,
 		ESimpleElementBlendMode BlendMode,
 		const FMatrix& Transform,
 		bool bSwitchVerticalAxis,
diff --git a/Engine/Source/Runtime/Engine/Public/ScreenRendering.h b/Engine/Source/Runtime/Engine/Public/ScreenRendering.h
--- a/Engine/Source/Runtime/Engine/Public/ScreenRendering.h
+++ b/Engine/Source/Runtime/Engine/Public/ScreenRendering.h
@@ -129,6 +129,6 @@ public:
 /**
  * Draws a texture rectangle on the screen using normalized (-1 to 1) device coordinates.
  */
-extern void DrawScreenQuad(  float X0, float Y0, float U0, float V0, float X1, float Y1, float U1, float V1, const FTexture* Texture );
+extern void DrawScreenQuad(FRHICommandList& RHICmdList, float X0, float Y0, float U0, float V0, float X1, float Y1, float U1, float V1, const FTexture* Texture);
 
-extern void ENGINE_API DrawNormalizedScreenQuad(  float X0, float Y0, float U0, float V0, float X1, float Y1, float U1, float V1, FIntPoint TargetSize, FTexture2DRHIRef TextureRHI );
+extern void ENGINE_API DrawNormalizedScreenQuad(FRHICommandList& RHICmdList, float X0, float Y0, float U0, float V0, float X1, float Y1, float U1, float V1, FIntPoint TargetSize, FTexture2DRHIRef TextureRHI);
diff --git a/Engine/Source/Runtime/Engine/Public/StaticBoundShaderState.h b/Engine/Source/Runtime/Engine/Public/StaticBoundShaderState.h
--- a/Engine/Source/Runtime/Engine/Public/StaticBoundShaderState.h
+++ b/Engine/Source/Runtime/Engine/Public/StaticBoundShaderState.h
@@ -53,6 +53,7 @@ typedef TGlobalResource<FGlobalBoundShaderStateResource> FGlobalBoundShaderState
 /**
  * SetGlobalBoundShaderState - sets the global bound shader state, also creates and caches it if necessary
  *
+ * @param RHICmdList				- cmdlist to add to
  * @param BoundShaderState			- current bound shader state, will be updated if it wasn't a valid ref
  * @param VertexDeclaration			- the vertex declaration to use in creating the new bound shader state
  * @param VertexShader				- the vertex shader to use in creating the new bound shader state
@@ -60,6 +61,7 @@ typedef TGlobalResource<FGlobalBoundShaderStateResource> FGlobalBoundShaderState
  * @param GeometryShader			- the geometry shader to use in creating the new bound shader state (0 if not used)
  */
 extern ENGINE_API void SetGlobalBoundShaderState(
+	FRHICommandList& RHICmdList,
 	FGlobalBoundShaderState& BoundShaderState,
 	FVertexDeclarationRHIParamRef VertexDeclaration,
 	FShader* VertexShader,
diff --git a/Engine/Source/Runtime/OpenGLDrv/Private/OpenGLDrv.cpp b/Engine/Source/Runtime/OpenGLDrv/Private/OpenGLDrv.cpp
--- a/Engine/Source/Runtime/OpenGLDrv/Private/OpenGLDrv.cpp
+++ b/Engine/Source/Runtime/OpenGLDrv/Private/OpenGLDrv.cpp
@@ -386,7 +386,7 @@ static FGlobalBoundShaderState LongGPUTaskBoundShaderState;
 
 void FOpenGLDynamicRHI::IssueLongGPUTask()
 {
-
+	FRHICommandList& RHICmdList = FRHICommandList::GetNullRef();
 	int32 LargestViewportIndex = INDEX_NONE;
 	int32 LargestViewportPixels = 0;
 
@@ -405,15 +405,15 @@ void FOpenGLDynamicRHI::IssueLongGPUTask()
 	{
 		FOpenGLViewport* Viewport = Viewports[LargestViewportIndex];
 
-		RHISetRenderTarget(Viewport->GetBackBuffer(), FTextureRHIRef());
-		RHISetBlendState(TStaticBlendState<CW_RGBA, BO_Add, BF_One, BF_One>::GetRHI(), FLinearColor::Black);
-		RHISetDepthStencilState(TStaticDepthStencilState<false,CF_Always>::GetRHI(), 0);
-		RHISetRasterizerState(TStaticRasterizerState<FM_Solid,CM_None>::GetRHI());
+		SetRenderTarget(RHICmdList, Viewport->GetBackBuffer(), FTextureRHIRef());
+		RHICmdList.SetBlendState(TStaticBlendState<CW_RGBA, BO_Add, BF_One, BF_One>::GetRHI(), FLinearColor::Black);
+		RHICmdList.SetDepthStencilState(TStaticDepthStencilState<false, CF_Always>::GetRHI(), 0);
+		RHICmdList.SetRasterizerState(TStaticRasterizerState<FM_Solid, CM_None>::GetRHI());
 
 		TShaderMapRef<TOneColorVS<true> > VertexShader(GetGlobalShaderMap());
 		TShaderMapRef<FOpenGLRHILongGPUTaskPS> PixelShader(GetGlobalShaderMap());
 
-		SetGlobalBoundShaderState(LongGPUTaskBoundShaderState, GOpenGLVector4VertexDeclaration.VertexDeclarationRHI, *VertexShader, *PixelShader, 0);
+		SetGlobalBoundShaderState(RHICmdList, LongGPUTaskBoundShaderState, GOpenGLVector4VertexDeclaration.VertexDeclarationRHI, *VertexShader, *PixelShader, 0);
 
 		// Draw a fullscreen quad
 		FVector4 Vertices[4];
@@ -421,7 +421,7 @@ void FOpenGLDynamicRHI::IssueLongGPUTask()
 		Vertices[1].Set(  1.0f,  1.0f, 0, 1.0f );
 		Vertices[2].Set( -1.0f, -1.0f, 0, 1.0f );
 		Vertices[3].Set(  1.0f, -1.0f, 0, 1.0f );
-		RHIDrawPrimitiveUP(PT_TriangleStrip, 2, Vertices, sizeof(Vertices[0]) );
+		DrawPrimitiveUP(RHICmdList, PT_TriangleStrip, 2, Vertices, sizeof(Vertices[0]));
 	}
 }
 
diff --git a/Engine/Source/Runtime/OpenGLDrv/Private/OpenGLViewport.cpp b/Engine/Source/Runtime/OpenGLDrv/Private/OpenGLViewport.cpp
--- a/Engine/Source/Runtime/OpenGLDrv/Private/OpenGLViewport.cpp
+++ b/Engine/Source/Runtime/OpenGLDrv/Private/OpenGLViewport.cpp
@@ -92,15 +92,17 @@ void FOpenGLDynamicRHI::RHIBeginDrawingViewport(FViewportRHIParamRef ViewportRHI
 	{
 		GPUProfilingData.FrameTiming.InitResource();
 	}
+	//@todo-rco: RHIPacketList
+	FRHICommandList& RHICmdList = FRHICommandList::GetNullRef();
 
 	// Set the render target and viewport.
 	if( RenderTarget )
 	{
-		RHISetRenderTarget(RenderTarget, FTextureRHIRef());
+		SetRenderTarget(RHICmdList, RenderTarget, FTextureRHIRef());
 	}
 	else
 	{
-		RHISetRenderTarget(DrawingViewport->GetBackBuffer(), FTextureRHIRef());
+		SetRenderTarget(RHICmdList, DrawingViewport->GetBackBuffer(), FTextureRHIRef());
 	}
 }
 
diff --git a/Engine/Source/Runtime/RHI/Public/RHICommandList.h b/Engine/Source/Runtime/RHI/Public/RHICommandList.h
--- a/Engine/Source/Runtime/RHI/Public/RHICommandList.h
+++ b/Engine/Source/Runtime/RHI/Public/RHICommandList.h
@@ -81,16 +81,106 @@ public:
 
 	inline void SetRasterizerState(FRasterizerStateRHIParamRef State);
 
-	inline void SetBlendState(FBlendStateRHIParamRef State, const FLinearColor& BlendFactor);
-	inline void SetBlendState(FBlendStateRHIParamRef State)
-	{
-		SetBlendState(State, FLinearColor::White);
-	}
+	inline void SetBlendState(FBlendStateRHIParamRef State, const FLinearColor& BlendFactor = FLinearColor::White);
 
 	inline void DrawPrimitive(uint32 PrimitiveType, uint32 BaseVertexIndex, uint32 NumPrimitives, uint32 NumInstances);
 	inline void DrawIndexedPrimitive(FIndexBufferRHIParamRef IndexBuffer, uint32 PrimitiveType, int32 BaseVertexIndex, uint32 MinIndex, uint32 NumVertices, uint32 StartIndex, uint32 NumPrimitives, uint32 NumInstances);
 	inline void SetStreamSource(uint32 StreamIndex, FVertexBufferRHIParamRef VertexBuffer, uint32 Stride, uint32 Offset);
 
+	// need cmd list implementations
+
+	FORCEINLINE void SetViewport(uint32 MinX, uint32 MinY, float MinZ, uint32 MaxX, uint32 MaxY, float MaxZ)
+	{
+		CheckIsNull();
+		RHISetViewport(MinX, MinY, MinZ, MaxX, MaxY, MaxZ);
+	}
+
+	FORCEINLINE void SetDepthStencilState(FDepthStencilStateRHIParamRef NewStateRHI, uint32 StencilRef = 0)
+	{
+		CheckIsNull();
+		RHISetDepthStencilState(NewStateRHI, StencilRef);
+	}
+
+	FORCEINLINE void BeginDrawPrimitiveUP(uint32 PrimitiveType, uint32 NumPrimitives, uint32 NumVertices, uint32 VertexDataStride, void*& OutVertexData)
+	{
+		CheckIsNull();
+		RHIBeginDrawPrimitiveUP(PrimitiveType, NumPrimitives, NumVertices, VertexDataStride, OutVertexData);
+	}
+	
+	FORCEINLINE void EndDrawPrimitiveUP()
+	{
+		CheckIsNull();
+		RHIEndDrawPrimitiveUP();
+	}
+
+	FORCEINLINE void CopyToResolveTarget(FTextureRHIParamRef SourceTextureRHI, FTextureRHIParamRef DestTextureRHI, bool bKeepOriginalSurface, const FResolveParams& ResolveParams)
+	{
+		CheckIsNull();
+		RHICopyToResolveTarget(SourceTextureRHI, DestTextureRHI, bKeepOriginalSurface, ResolveParams);
+	}
+
+	FORCEINLINE void BeginDrawIndexedPrimitiveUP(uint32 PrimitiveType, uint32 NumPrimitives, uint32 NumVertices, uint32 VertexDataStride, void*& OutVertexData, uint32 MinVertexIndex, uint32 NumIndices, uint32 IndexDataStride, void*& OutIndexData)
+	{
+		CheckIsNull();
+		RHIBeginDrawIndexedPrimitiveUP(PrimitiveType, NumPrimitives, NumVertices, VertexDataStride, OutVertexData, MinVertexIndex, NumIndices, IndexDataStride, OutIndexData);
+	}
+
+	FORCEINLINE void EndDrawIndexedPrimitiveUP()
+	{
+		CheckIsNull();
+		RHIEndDrawIndexedPrimitiveUP();
+	}
+
+	FORCEINLINE void Clear(bool bClearColor, const FLinearColor& Color, bool bClearDepth, float Depth, bool bClearStencil, uint32 Stencil, FIntRect ExcludeRect)
+	{
+		CheckIsNull();
+		RHIClear(bClearColor, Color, bClearDepth, Depth, bClearStencil, Stencil, ExcludeRect);
+	}
+
+	FORCEINLINE void ClearMRT(bool bClearColor, int32 NumClearColors, const FLinearColor* ClearColorArray, bool bClearDepth, float Depth, bool bClearStencil, uint32 Stencil, FIntRect ExcludeRect)
+	{
+		CheckIsNull();
+		RHIClearMRT(bClearColor, NumClearColors, ClearColorArray, bClearDepth, Depth, bClearStencil, Stencil, ExcludeRect);
+	}
+
+	FORCEINLINE void SetRenderTargets(uint32 NewNumSimultaneousRenderTargets, const FRHIRenderTargetView* NewRenderTargetsRHI, FTextureRHIParamRef NewDepthStencilTargetRHI, uint32 NewNumUAVs, const FUnorderedAccessViewRHIParamRef* UAVs)
+	{
+		CheckIsNull();
+		RHISetRenderTargets(NewNumSimultaneousRenderTargets, NewRenderTargetsRHI, NewDepthStencilTargetRHI, NewNumUAVs, UAVs);
+	}
+
+	FORCEINLINE void EnableDepthBoundsTest(bool bEnable, float MinDepth, float MaxDepth)
+	{
+		CheckIsNull();
+		RHIEnableDepthBoundsTest(bEnable, MinDepth, MaxDepth);
+	}
+
+	FORCEINLINE void SetScissorRect(bool bEnable, uint32 MinX, uint32 MinY, uint32 MaxX, uint32 MaxY)
+	{
+		CheckIsNull();
+		RHISetScissorRect(bEnable, MinX, MinY, MaxX, MaxY);
+	}
+
+	FORCEINLINE void ClearUAV(FUnorderedAccessViewRHIParamRef UnorderedAccessViewRHI, const uint32* Values)
+	{
+		CheckIsNull();
+		RHIClearUAV(UnorderedAccessViewRHI, Values);
+	}
+
+	FORCEINLINE void SetComputeShader(FComputeShaderRHIParamRef ComputeShaderRHI)
+	{
+		CheckIsNull();
+		RHISetComputeShader(ComputeShaderRHI);
+	}
+
+	FORCEINLINE void DrawPrimitiveIndirect(uint32 PrimitiveType, FVertexBufferRHIParamRef ArgumentBufferRHI, uint32 ArgumentOffset)
+	{
+		CheckIsNull();
+		RHIDrawPrimitiveIndirect(PrimitiveType, ArgumentBufferRHI, ArgumentOffset);
+	}
+
+
+
 	FORCEINLINE bool IsNull() const
 	{
 		return !Memory;
diff --git a/Engine/Source/Runtime/RHI/Public/RHIUtilities.h b/Engine/Source/Runtime/RHI/Public/RHIUtilities.h
--- a/Engine/Source/Runtime/RHI/Public/RHIUtilities.h
+++ b/Engine/Source/Runtime/RHI/Public/RHIUtilities.h
@@ -4,6 +4,7 @@
 #pragma once
 
 #include "RHI.h"
+#include "RHICommandList.h"
 
 /** Encapsulates a GPU read/write buffer with its UAV and SRV. */
 struct FRWBuffer
@@ -119,6 +120,49 @@ struct FRWBufferByteAddress
 	}
 };
 
+/** Helper for the common case of using a single color and depth render target. */
+inline void SetRenderTarget(FRHICommandList& RHICmdList, FTextureRHIParamRef NewRenderTarget, FTextureRHIParamRef NewDepthStencilTarget)
+{
+	FRHIRenderTargetView RTV(NewRenderTarget);
+	RHICmdList.SetRenderTargets(1, &RTV, NewDepthStencilTarget, 0, NULL);
+}
+
+/** Helper for the common case of using a single color and depth render target, with a mip index for the color target. */
+inline void SetRenderTarget(FRHICommandList& RHICmdList, FTextureRHIParamRef NewRenderTarget, int32 MipIndex, FTextureRHIParamRef NewDepthStencilTarget)
+{
+	FRHIRenderTargetView RTV(NewRenderTarget, MipIndex, -1);
+	RHICmdList.SetRenderTargets(1, &RTV, NewDepthStencilTarget, 0, NULL);
+}
+
+/** Helper for the common case of using a single color and depth render target, with a mip index for the color target. */
+inline void SetRenderTarget(FRHICommandList& RHICmdList, FTextureRHIParamRef NewRenderTarget, int32 MipIndex, int32 ArraySliceIndex, FTextureRHIParamRef NewDepthStencilTarget)
+{
+	FRHIRenderTargetView RTV(NewRenderTarget, MipIndex, ArraySliceIndex);
+	RHICmdList.SetRenderTargets(1, &RTV, NewDepthStencilTarget, 0, NULL);
+}
+
+/** Helper that converts FTextureRHIParamRef's into FRHIRenderTargetView's. */
+inline void SetRenderTargets(
+	FRHICommandList& RHICmdList,
+	uint32 NewNumSimultaneousRenderTargets, 
+	const FTextureRHIParamRef* NewRenderTargetsRHI,
+	FTextureRHIParamRef NewDepthStencilTargetRHI,
+	uint32 NewNumUAVs,
+	const FUnorderedAccessViewRHIParamRef* UAVs
+	)
+{
+	FRHIRenderTargetView RTVs[MaxSimultaneousRenderTargets];
+
+	for (uint32 Index = 0; Index < NewNumSimultaneousRenderTargets; Index++)
+	{
+		RTVs[Index] = FRHIRenderTargetView(NewRenderTargetsRHI[Index]);
+	}
+
+	RHICmdList.SetRenderTargets(NewNumSimultaneousRenderTargets, RTVs, NewDepthStencilTargetRHI, NewNumUAVs, UAVs);
+}
+
+
+
 /** Helper for the common case of using a single color and depth render target. */
 inline void RHISetRenderTarget(FTextureRHIParamRef NewRenderTarget, FTextureRHIParamRef NewDepthStencilTarget)
 {
@@ -142,7 +186,7 @@ inline void RHISetRenderTarget(FTextureRHIParamRef NewRenderTarget, int32 MipInd
 
 /** Helper that converts FTextureRHIParamRef's into FRHIRenderTargetView's. */
 inline void RHISetRenderTargets(
-	uint32 NewNumSimultaneousRenderTargets, 
+	uint32 NewNumSimultaneousRenderTargets,
 	const FTextureRHIParamRef* NewRenderTargetsRHI,
 	FTextureRHIParamRef NewDepthStencilTargetRHI,
 	uint32 NewNumUAVs,
@@ -158,7 +202,6 @@ inline void RHISetRenderTargets(
 
 	RHISetRenderTargets(NewNumSimultaneousRenderTargets, RTVs, NewDepthStencilTargetRHI, NewNumUAVs, UAVs);
 }
-
 /** Sets a depth-stencil state with a default stencil ref value of 0. */
 inline void RHISetDepthStencilState(FDepthStencilStateRHIParamRef NewState)
 {
@@ -378,13 +421,13 @@ inline uint32 RHIGetVertexCountForPrimitiveCount(uint32 NumPrimitives, uint32 Pr
  * @param VertexData A reference to memory preallocate in RHIBeginDrawPrimitiveUP
  * @param VertexDataStride Size of each vertex
  */
-inline void RHIDrawPrimitiveUP( uint32 PrimitiveType, uint32 NumPrimitives, const void* VertexData, uint32 VertexDataStride )
+inline void DrawPrimitiveUP(FRHICommandList& RHICmdList, uint32 PrimitiveType, uint32 NumPrimitives, const void* VertexData, uint32 VertexDataStride)
 {
 	void* Buffer = NULL;
 	const uint32 VertexCount = RHIGetVertexCountForPrimitiveCount( NumPrimitives, PrimitiveType );
-	RHIBeginDrawPrimitiveUP( PrimitiveType, NumPrimitives, VertexCount, VertexDataStride, Buffer );
+	RHICmdList.BeginDrawPrimitiveUP(PrimitiveType, NumPrimitives, VertexCount, VertexDataStride, Buffer);
 	FMemory::Memcpy( Buffer, VertexData, VertexCount * VertexDataStride );
-	RHIEndDrawPrimitiveUP();
+	RHICmdList.EndDrawPrimitiveUP();
 }
 
 /**
@@ -398,7 +441,8 @@ inline void RHIDrawPrimitiveUP( uint32 PrimitiveType, uint32 NumPrimitives, cons
  * @param VertexData The memory preallocate in RHIBeginDrawIndexedPrimitiveUP
  * @param VertexDataStride The size of one vertex
  */
-inline void RHIDrawIndexedPrimitiveUP(
+inline void DrawIndexedPrimitiveUP(
+	FRHICommandList& RHICmdList,
 	uint32 PrimitiveType,
 	uint32 MinVertexIndex,
 	uint32 NumVertices,
@@ -411,7 +455,7 @@ inline void RHIDrawIndexedPrimitiveUP(
 	void* VertexBuffer = NULL;
 	void* IndexBuffer = NULL;
 	const uint32 NumIndices = RHIGetVertexCountForPrimitiveCount( NumPrimitives, PrimitiveType );
-	RHIBeginDrawIndexedPrimitiveUP(
+	RHICmdList.BeginDrawIndexedPrimitiveUP(
 		PrimitiveType,
 		NumPrimitives,
 		NumVertices,
@@ -423,5 +467,5 @@ inline void RHIDrawIndexedPrimitiveUP(
 		IndexBuffer );
 	FMemory::Memcpy( VertexBuffer, VertexData, NumVertices * VertexDataStride );
 	FMemory::Memcpy( IndexBuffer, IndexData, NumIndices * IndexDataStride );
-	RHIEndDrawIndexedPrimitiveUP();
+	RHICmdList.EndDrawIndexedPrimitiveUP();
 }
diff --git a/Engine/Source/Runtime/Renderer/Private/AtmosphereRendering.cpp b/Engine/Source/Runtime/Renderer/Private/AtmosphereRendering.cpp
--- a/Engine/Source/Runtime/Renderer/Private/AtmosphereRendering.cpp
+++ b/Engine/Source/Runtime/Renderer/Private/AtmosphereRendering.cpp
@@ -384,7 +384,7 @@ void SetAtmosphericFogShaders(FRHICommandList& RHICmdList, FScene* Scene, const
 		break;
 	}
 	
-	SetGlobalBoundShaderState(AtmosphereBoundShaderState[Scene->AtmosphericFog->RenderFlag], GAtmophereVertexDeclaration.VertexDeclarationRHI, *VertexShader, PixelShader);
+	SetGlobalBoundShaderState(RHICmdList, AtmosphereBoundShaderState[Scene->AtmosphericFog->RenderFlag], GAtmophereVertexDeclaration.VertexDeclarationRHI, *VertexShader, PixelShader);
 	VertexShader->SetParameters(RHICmdList, View);
 	PixelShader->SetParameters(RHICmdList, View, LightShaftOcclusion);
 }
@@ -430,7 +430,8 @@ void FDeferredShadingSceneRenderer::RenderAtmosphere(FLightShaftsOutput LightSha
 			SetAtmosphericFogShaders(RHICmdList, Scene, View, LightShaftsOutput.LightShaftOcclusion);
 
 			// Draw a quad covering the view.
-			RHIDrawIndexedPrimitiveUP(
+			DrawIndexedPrimitiveUP(
+				RHICmdList,
 				PT_TriangleList,
 				0,
 				ARRAY_COUNT(Vertices),
@@ -1054,10 +1055,11 @@ FIntPoint FAtmosphericFogSceneInfo::GetTextureSize()
 	return AtmosphereTextures->AtmosphereTransmittance->GetDesc().Extent;
 }
 
-void FAtmosphericFogSceneInfo::DrawQuad(const FIntRect& ViewRect, FShader* VertexShader)
+void FAtmosphericFogSceneInfo::DrawQuad(FRHICommandList& RHICmdList, const FIntRect& ViewRect, FShader* VertexShader)
 {
 	// Draw a quad mapping scene color to the view's render target
 	DrawRectangle( 
+		RHICmdList,
 		ViewRect.Min.X, ViewRect.Min.Y, 
 		ViewRect.Width(), ViewRect.Height(),
 		ViewRect.Min.X, ViewRect.Min.Y, 
@@ -1081,11 +1083,8 @@ void FAtmosphericFogSceneInfo::GetLayerValue(int Layer, float& AtmosphereR, FVec
 	DhdH = FVector4(DMin, DMax, DMinP, DMaxP);
 }
 
-void FAtmosphericFogSceneInfo::RenderAtmosphereShaders(const FViewInfo& View, const FIntRect& ViewRect)
+void FAtmosphericFogSceneInfo::RenderAtmosphereShaders(FRHICommandList& RHICmdList, const FViewInfo& View, const FIntRect& ViewRect)
 {
-	//@todo-rco: RHIPacketList
-	FRHICommandList& RHICmdList = FRHICommandList::GetNullRef();
-
 	check(Component != NULL);
 	switch(AtmospherePhase)
 	{
@@ -1098,29 +1097,29 @@ void FAtmosphericFogSceneInfo::RenderAtmosphereShaders(const FViewInfo& View, co
 			TShaderMapRef<FAtmosphereTransmittancePS> PixelShader(GetGlobalShaderMap());
 
 			static FGlobalBoundShaderState BoundShaderState;
-			SetGlobalBoundShaderState( BoundShaderState, GFilterVertexDeclaration.VertexDeclarationRHI, *VertexShader, *PixelShader);
+			SetGlobalBoundShaderState(RHICmdList, BoundShaderState, GFilterVertexDeclaration.VertexDeclarationRHI, *VertexShader, *PixelShader);
 			PixelShader->SetParameters(RHICmdList, View);
-			DrawQuad(ViewRect, *VertexShader);
+			DrawQuad(RHICmdList, ViewRect, *VertexShader);
 
-			RHICopyToResolveTarget(DestRenderTarget.TargetableTexture, DestRenderTarget.ShaderResourceTexture, true, FResolveParams() );
+			RHICmdList.CopyToResolveTarget(DestRenderTarget.TargetableTexture, DestRenderTarget.ShaderResourceTexture, true, FResolveParams());
 		}
 		break;
 	case AP_Irradiance1:
 		{
 			const FSceneRenderTargetItem& DestRenderTarget = AtmosphereTextures->AtmosphereDeltaE->GetRenderTargetItem();
-			RHISetRenderTarget(DestRenderTarget.TargetableTexture, FTextureRHIRef());
+			SetRenderTarget(RHICmdList, DestRenderTarget.TargetableTexture, FTextureRHIRef());
 
 			TShaderMapRef<FAtmospherePrecomputeVS> VertexShader(GetGlobalShaderMap());
 			TShaderMapRef<FAtmosphereIrradiance1PS> PixelShader(GetGlobalShaderMap());
 
 			static FGlobalBoundShaderState BoundShaderState;
-			SetGlobalBoundShaderState( BoundShaderState, GFilterVertexDeclaration.VertexDeclarationRHI, *VertexShader, *PixelShader);
+			SetGlobalBoundShaderState(RHICmdList, BoundShaderState, GFilterVertexDeclaration.VertexDeclarationRHI, *VertexShader, *PixelShader);
 
 			PixelShader->SetParameters(RHICmdList, AtmosphereTextures);
 
-			DrawQuad(ViewRect, *VertexShader);
+			DrawQuad(RHICmdList, ViewRect, *VertexShader);
 
-			RHICopyToResolveTarget(DestRenderTarget.TargetableTexture, DestRenderTarget.ShaderResourceTexture, true, FResolveParams() );
+			RHICmdList.CopyToResolveTarget(DestRenderTarget.TargetableTexture, DestRenderTarget.ShaderResourceTexture, true, FResolveParams());
 		}
 		break;
 	case AP_Inscatter1:
@@ -1130,14 +1129,15 @@ void FAtmosphericFogSceneInfo::RenderAtmosphereShaders(const FViewInfo& View, co
 				FTextureRHIParamRef RenderTargets[2];
 				RenderTargets[0] = AtmosphereTextures->AtmosphereDeltaSR->GetRenderTargetItem().TargetableTexture;
 				RenderTargets[1] = AtmosphereTextures->AtmosphereDeltaSM->GetRenderTargetItem().TargetableTexture;
-				RHISetRenderTargets(ARRAY_COUNT(RenderTargets), RenderTargets, FTextureRHIRef(), 0, NULL);
+				SetRenderTargets(RHICmdList, ARRAY_COUNT(RenderTargets), RenderTargets, FTextureRHIRef(), 0, NULL);
 
 				TShaderMapRef<FAtmospherePrecomputeVS> VertexShader(GetGlobalShaderMap());
 				TShaderMapRef<FAtmosphereGS> GeometryShader(GetGlobalShaderMap());
 				TShaderMapRef<FAtmosphereInscatter1PS> PixelShader(GetGlobalShaderMap());
 
 				static FGlobalBoundShaderState BoundShaderState;
-				SetGlobalBoundShaderState( BoundShaderState, GFilterVertexDeclaration.VertexDeclarationRHI, *VertexShader, *PixelShader, *GeometryShader);
+				RHICmdList.CheckIsNull(); // need new approach for "static FGlobalBoundShaderState" for parallel rendering
+				SetGlobalBoundShaderState(RHICmdList, BoundShaderState, GFilterVertexDeclaration.VertexDeclarationRHI, *VertexShader, *PixelShader, *GeometryShader);
 
 				//
 				float r;
@@ -1145,13 +1145,13 @@ void FAtmosphericFogSceneInfo::RenderAtmosphereShaders(const FViewInfo& View, co
 				GetLayerValue(Layer, r, DhdH);
 				GeometryShader->SetParameters(RHICmdList, Layer);
 				PixelShader->SetParameters(RHICmdList, View, r, DhdH, AtmosphereTextures);
-				DrawQuad(ViewRect, *VertexShader);
+				DrawQuad(RHICmdList, ViewRect, *VertexShader);
 
 				if (Atmosphere3DTextureIndex == Component->PrecomputeParams.InscatterAltitudeSampleNum - 1)
 				{
-					RHICopyToResolveTarget(AtmosphereTextures->AtmosphereDeltaSR->GetRenderTargetItem().TargetableTexture,
+					RHICmdList.CopyToResolveTarget(AtmosphereTextures->AtmosphereDeltaSR->GetRenderTargetItem().TargetableTexture,
 						AtmosphereTextures->AtmosphereDeltaSR->GetRenderTargetItem().ShaderResourceTexture, true, FResolveParams());
-					RHICopyToResolveTarget(AtmosphereTextures->AtmosphereDeltaSM->GetRenderTargetItem().TargetableTexture,
+					RHICmdList.CopyToResolveTarget(AtmosphereTextures->AtmosphereDeltaSM->GetRenderTargetItem().TargetableTexture,
 						AtmosphereTextures->AtmosphereDeltaSM->GetRenderTargetItem().ShaderResourceTexture, true, FResolveParams());
 				}
 			}
@@ -1160,10 +1160,10 @@ void FAtmosphericFogSceneInfo::RenderAtmosphereShaders(const FViewInfo& View, co
 	case AP_ClearIrradiance:
 		{
 			const FSceneRenderTargetItem& DestRenderTarget = AtmosphereTextures->AtmosphereIrradiance->GetRenderTargetItem();
-			RHISetRenderTarget(DestRenderTarget.TargetableTexture, FTextureRHIRef());
+			SetRenderTarget(RHICmdList, DestRenderTarget.TargetableTexture, FTextureRHIRef());
 
-			RHIClear(  true, FLinearColor::Black, false, 0.0f, false, 0, FIntRect());
-			RHICopyToResolveTarget(DestRenderTarget.TargetableTexture, DestRenderTarget.ShaderResourceTexture, true, FResolveParams() );
+			RHICmdList.Clear(true, FLinearColor::Black, false, 0.0f, false, 0, FIntRect());
+			RHICmdList.CopyToResolveTarget(DestRenderTarget.TargetableTexture, DestRenderTarget.ShaderResourceTexture, true, FResolveParams());
 		}
 		break;
 	case AP_CopyInscatter1:
@@ -1171,25 +1171,26 @@ void FAtmosphericFogSceneInfo::RenderAtmosphereShaders(const FViewInfo& View, co
 			int32 Layer = Atmosphere3DTextureIndex;
 			{
 				const FSceneRenderTargetItem& DestRenderTarget = AtmosphereTextures->AtmosphereInscatter->GetRenderTargetItem();
-				RHISetRenderTarget(DestRenderTarget.TargetableTexture, FTextureRHIRef());
+				SetRenderTarget(RHICmdList, DestRenderTarget.TargetableTexture, FTextureRHIRef());
 
 				TShaderMapRef<FAtmospherePrecomputeVS> VertexShader(GetGlobalShaderMap());
 				TShaderMapRef<FAtmosphereGS> GeometryShader(GetGlobalShaderMap());
 				TShaderMapRef<FAtmosphereCopyInscatter1PS> PixelShader(GetGlobalShaderMap());
 
 				static FGlobalBoundShaderState BoundShaderState;
-				SetGlobalBoundShaderState( BoundShaderState, GFilterVertexDeclaration.VertexDeclarationRHI, *VertexShader, *PixelShader, *GeometryShader);
+				RHICmdList.CheckIsNull(); // need new approach for "static FGlobalBoundShaderState" for parallel rendering
+				SetGlobalBoundShaderState(RHICmdList, BoundShaderState, GFilterVertexDeclaration.VertexDeclarationRHI, *VertexShader, *PixelShader, *GeometryShader);
 
 				float r;
 				FVector4 DhdH;
 				GetLayerValue(Layer, r, DhdH);
 				GeometryShader->SetParameters(RHICmdList, Layer);
 				PixelShader->SetParameters(RHICmdList, View, r, DhdH, Layer, AtmosphereTextures);
-				DrawQuad(ViewRect, *VertexShader);
+				DrawQuad(RHICmdList, ViewRect, *VertexShader);
 
 				if (Atmosphere3DTextureIndex == Component->PrecomputeParams.InscatterAltitudeSampleNum - 1)
 				{
-					RHICopyToResolveTarget(DestRenderTarget.TargetableTexture, DestRenderTarget.ShaderResourceTexture, true, FResolveParams() );
+					RHICmdList.CopyToResolveTarget(DestRenderTarget.TargetableTexture, DestRenderTarget.ShaderResourceTexture, true, FResolveParams());
 				}
 			}
 		}
@@ -1199,25 +1200,26 @@ void FAtmosphericFogSceneInfo::RenderAtmosphereShaders(const FViewInfo& View, co
 			int32 Layer = Atmosphere3DTextureIndex;
 			{
 				const FSceneRenderTargetItem& DestRenderTarget = AtmosphereTextures->AtmosphereDeltaJ->GetRenderTargetItem();
-				RHISetRenderTarget(DestRenderTarget.TargetableTexture, FTextureRHIRef());
+				SetRenderTarget(RHICmdList, DestRenderTarget.TargetableTexture, FTextureRHIRef());
 
 				TShaderMapRef<FAtmospherePrecomputeVS> VertexShader(GetGlobalShaderMap());
 				TShaderMapRef<FAtmosphereGS> GeometryShader(GetGlobalShaderMap());
 				TShaderMapRef<FAtmosphereInscatterSPS> PixelShader(GetGlobalShaderMap());
 
 				static FGlobalBoundShaderState BoundShaderState;
-				SetGlobalBoundShaderState( BoundShaderState, GFilterVertexDeclaration.VertexDeclarationRHI, *VertexShader, *PixelShader, *GeometryShader);
+				RHICmdList.CheckIsNull(); // need new approach for "static FGlobalBoundShaderState" for parallel rendering
+				SetGlobalBoundShaderState(RHICmdList, BoundShaderState, GFilterVertexDeclaration.VertexDeclarationRHI, *VertexShader, *PixelShader, *GeometryShader);
 
 				float r;
 				FVector4 DhdH;
 				GetLayerValue(Layer, r, DhdH);
 				GeometryShader->SetParameters(RHICmdList, Layer);
 				PixelShader->SetParameters(RHICmdList, View, r, DhdH, AtmoshpereOrder == 2 ? 1.f : 0.f, AtmosphereTextures);
-				DrawQuad(ViewRect, *VertexShader);
+				DrawQuad(RHICmdList, ViewRect, *VertexShader);
 
 				if (Atmosphere3DTextureIndex == Component->PrecomputeParams.InscatterAltitudeSampleNum - 1)
 				{
-					RHICopyToResolveTarget(DestRenderTarget.TargetableTexture, DestRenderTarget.ShaderResourceTexture, true, FResolveParams() );
+					RHICmdList.CopyToResolveTarget(DestRenderTarget.TargetableTexture, DestRenderTarget.ShaderResourceTexture, true, FResolveParams());
 				}
 			}
 		}
@@ -1225,19 +1227,20 @@ void FAtmosphericFogSceneInfo::RenderAtmosphereShaders(const FViewInfo& View, co
 	case AP_IrradianceN:
 		{
 			const FSceneRenderTargetItem& DestRenderTarget = AtmosphereTextures->AtmosphereDeltaE->GetRenderTargetItem();
-			RHISetRenderTarget(DestRenderTarget.TargetableTexture, FTextureRHIRef());
+			SetRenderTarget(RHICmdList, DestRenderTarget.TargetableTexture, FTextureRHIRef());
 
 			TShaderMapRef<FAtmospherePrecomputeVS> VertexShader(GetGlobalShaderMap());
 			TShaderMapRef<FAtmosphereIrradianceNPS> PixelShader(GetGlobalShaderMap());
 
 			static FGlobalBoundShaderState BoundShaderState;
-			SetGlobalBoundShaderState( BoundShaderState, GFilterVertexDeclaration.VertexDeclarationRHI, *VertexShader, *PixelShader);
+			RHICmdList.CheckIsNull(); // need new approach for "static FGlobalBoundShaderState" for parallel rendering
+			SetGlobalBoundShaderState(RHICmdList, BoundShaderState, GFilterVertexDeclaration.VertexDeclarationRHI, *VertexShader, *PixelShader);
 
 			PixelShader->SetParameters(RHICmdList, View, AtmoshpereOrder == 2 ? 1.f : 0.f, AtmosphereTextures);
 
-			DrawQuad(ViewRect, *VertexShader);
+			DrawQuad(RHICmdList, ViewRect, *VertexShader);
 
-			RHICopyToResolveTarget(DestRenderTarget.TargetableTexture, DestRenderTarget.ShaderResourceTexture, true, FResolveParams() );
+			RHICmdList.CopyToResolveTarget(DestRenderTarget.TargetableTexture, DestRenderTarget.ShaderResourceTexture, true, FResolveParams());
 		}
 		break;
 
@@ -1246,82 +1249,85 @@ void FAtmosphericFogSceneInfo::RenderAtmosphereShaders(const FViewInfo& View, co
 			int32 Layer = Atmosphere3DTextureIndex;
 			{
 				const FSceneRenderTargetItem& DestRenderTarget = AtmosphereTextures->AtmosphereDeltaSR->GetRenderTargetItem();
-				RHISetRenderTarget(DestRenderTarget.TargetableTexture, FTextureRHIRef());
+				SetRenderTarget(RHICmdList, DestRenderTarget.TargetableTexture, FTextureRHIRef());
 
 				TShaderMapRef<FAtmospherePrecomputeVS> VertexShader(GetGlobalShaderMap());
 				TShaderMapRef<FAtmosphereGS> GeometryShader(GetGlobalShaderMap());
 				TShaderMapRef<FAtmosphereInscatterNPS> PixelShader(GetGlobalShaderMap());
 
 				static FGlobalBoundShaderState BoundShaderState;
-				SetGlobalBoundShaderState( BoundShaderState, GFilterVertexDeclaration.VertexDeclarationRHI, *VertexShader, *PixelShader, *GeometryShader);
+				RHICmdList.CheckIsNull(); // need new approach for "static FGlobalBoundShaderState" for parallel rendering
+				SetGlobalBoundShaderState(RHICmdList, BoundShaderState, GFilterVertexDeclaration.VertexDeclarationRHI, *VertexShader, *PixelShader, *GeometryShader);
 
 				float r;
 				FVector4 DhdH;
 				GetLayerValue(Layer, r, DhdH);
 				GeometryShader->SetParameters(RHICmdList, Layer);
 				PixelShader->SetParameters(RHICmdList, View, r, DhdH, AtmoshpereOrder == 2 ? 1.f : 0.f, AtmosphereTextures);
-				DrawQuad(ViewRect, *VertexShader);
+				DrawQuad(RHICmdList, ViewRect, *VertexShader);
 
 				if (Atmosphere3DTextureIndex == Component->PrecomputeParams.InscatterAltitudeSampleNum - 1)
 				{
-					RHICopyToResolveTarget(DestRenderTarget.TargetableTexture, DestRenderTarget.ShaderResourceTexture, true, FResolveParams() );
+					RHICmdList.CopyToResolveTarget(DestRenderTarget.TargetableTexture, DestRenderTarget.ShaderResourceTexture, true, FResolveParams());
 				}
 			}
 		}
 		break;
 	case AP_CopyIrradiance:
 		{
-			RHISetBlendState(TStaticBlendState<CW_RGBA,BO_Add,BF_One,BF_One,BO_Add,BF_One,BF_One>::GetRHI());
+			RHICmdList.SetBlendState(TStaticBlendState<CW_RGBA, BO_Add, BF_One, BF_One, BO_Add, BF_One, BF_One>::GetRHI());
 
 			const FSceneRenderTargetItem& DestRenderTarget = AtmosphereTextures->AtmosphereIrradiance->GetRenderTargetItem();
-			RHISetRenderTarget(DestRenderTarget.TargetableTexture, FTextureRHIRef());
+			SetRenderTarget(RHICmdList, DestRenderTarget.TargetableTexture, FTextureRHIRef());
 
 			TShaderMapRef<FAtmospherePrecomputeVS> VertexShader(GetGlobalShaderMap());
 			TShaderMapRef<FAtmosphereCopyIrradiancePS> PixelShader(GetGlobalShaderMap());
 
 			static FGlobalBoundShaderState BoundShaderState;
-			SetGlobalBoundShaderState( BoundShaderState, GFilterVertexDeclaration.VertexDeclarationRHI, *VertexShader, *PixelShader);
+			RHICmdList.CheckIsNull(); // need new approach for "static FGlobalBoundShaderState" for parallel rendering
+			SetGlobalBoundShaderState(RHICmdList, BoundShaderState, GFilterVertexDeclaration.VertexDeclarationRHI, *VertexShader, *PixelShader);
 
 			PixelShader->SetParameters(RHICmdList, AtmosphereTextures);
 
-			DrawQuad(ViewRect, *VertexShader);
+			DrawQuad(RHICmdList, ViewRect, *VertexShader);
 
-			RHICopyToResolveTarget(DestRenderTarget.TargetableTexture, DestRenderTarget.ShaderResourceTexture, true, FResolveParams() );
+			RHICmdList.CopyToResolveTarget(DestRenderTarget.TargetableTexture, DestRenderTarget.ShaderResourceTexture, true, FResolveParams());
 
-			RHISetBlendState(TStaticBlendState<>::GetRHI());
+			RHICmdList.SetBlendState(TStaticBlendState<>::GetRHI());
 		}
 		break;
 
 	case AP_CopyInscatterN:
 		{
-			RHISetBlendState(TStaticBlendState<CW_RGBA,BO_Add,BF_One,BF_One,BO_Add,BF_One,BF_One>::GetRHI());
+			RHICmdList.SetBlendState(TStaticBlendState<CW_RGBA, BO_Add, BF_One, BF_One, BO_Add, BF_One, BF_One>::GetRHI());
 
 			int32 Layer = Atmosphere3DTextureIndex;
 			{
 				const FSceneRenderTargetItem& DestRenderTarget = AtmosphereTextures->AtmosphereInscatter->GetRenderTargetItem();
-				RHISetRenderTarget(DestRenderTarget.TargetableTexture, FTextureRHIRef());
+				SetRenderTarget(RHICmdList, DestRenderTarget.TargetableTexture, FTextureRHIRef());
 
 				TShaderMapRef<FAtmospherePrecomputeVS> VertexShader(GetGlobalShaderMap());
 				TShaderMapRef<FAtmosphereGS> GeometryShader(GetGlobalShaderMap());
 				TShaderMapRef<FAtmosphereCopyInscatterNPS> PixelShader(GetGlobalShaderMap());
 
 				static FGlobalBoundShaderState BoundShaderState;
-				SetGlobalBoundShaderState( BoundShaderState, GFilterVertexDeclaration.VertexDeclarationRHI, *VertexShader, *PixelShader, *GeometryShader);
+				RHICmdList.CheckIsNull(); // need new approach for "static FGlobalBoundShaderState" for parallel rendering
+				SetGlobalBoundShaderState(RHICmdList, BoundShaderState, GFilterVertexDeclaration.VertexDeclarationRHI, *VertexShader, *PixelShader, *GeometryShader);
 
 				float r;
 				FVector4 DhdH;
 				GetLayerValue(Layer, r, DhdH);
 				GeometryShader->SetParameters(RHICmdList, Layer);
 				PixelShader->SetParameters(RHICmdList, View, r, DhdH, Layer, AtmosphereTextures);
-				DrawQuad(ViewRect, *VertexShader);
+				DrawQuad(RHICmdList, ViewRect, *VertexShader);
 
 				if (Atmosphere3DTextureIndex == Component->PrecomputeParams.InscatterAltitudeSampleNum - 1)
 				{
-					RHICopyToResolveTarget(DestRenderTarget.TargetableTexture, DestRenderTarget.ShaderResourceTexture, true, FResolveParams() );
+					RHICmdList.CopyToResolveTarget(DestRenderTarget.TargetableTexture, DestRenderTarget.ShaderResourceTexture, true, FResolveParams());
 				}
 			}
 
-			RHISetBlendState(TStaticBlendState<>::GetRHI());
+			RHICmdList.SetBlendState(TStaticBlendState<>::GetRHI());
 		}
 		break;
 
@@ -1330,25 +1336,26 @@ void FAtmosphericFogSceneInfo::RenderAtmosphereShaders(const FViewInfo& View, co
 			int32 Layer = Atmosphere3DTextureIndex;
 			{
 				const FSceneRenderTargetItem& DestRenderTarget = AtmosphereTextures->AtmosphereDeltaSR->GetRenderTargetItem();
-				RHISetRenderTarget(DestRenderTarget.TargetableTexture, FTextureRHIRef());
+				SetRenderTarget(RHICmdList, DestRenderTarget.TargetableTexture, FTextureRHIRef());
 
 				TShaderMapRef<FAtmospherePrecomputeVS> VertexShader(GetGlobalShaderMap());
 				TShaderMapRef<FAtmosphereGS> GeometryShader(GetGlobalShaderMap());
 				TShaderMapRef<FAtmosphereCopyInscatterFPS> PixelShader(GetGlobalShaderMap());
 
 				static FGlobalBoundShaderState BoundShaderState;
-				SetGlobalBoundShaderState( BoundShaderState, GFilterVertexDeclaration.VertexDeclarationRHI, *VertexShader, *PixelShader, *GeometryShader);
+				RHICmdList.CheckIsNull(); // need new approach for "static FGlobalBoundShaderState" for parallel rendering
+				SetGlobalBoundShaderState(RHICmdList, BoundShaderState, GFilterVertexDeclaration.VertexDeclarationRHI, *VertexShader, *PixelShader, *GeometryShader);
 
 				float r;
 				FVector4 DhdH;
 				GetLayerValue(Layer, r, DhdH);
 				GeometryShader->SetParameters(RHICmdList, Layer);
 				PixelShader->SetParameters(RHICmdList, r, DhdH, Layer, AtmosphereTextures);
-				DrawQuad(ViewRect, *VertexShader);
+				DrawQuad(RHICmdList, ViewRect, *VertexShader);
 
 				if (Atmosphere3DTextureIndex == Component->PrecomputeParams.InscatterAltitudeSampleNum - 1)
 				{
-					RHICopyToResolveTarget(DestRenderTarget.TargetableTexture, DestRenderTarget.ShaderResourceTexture, true, FResolveParams() );
+					RHICmdList.CopyToResolveTarget(DestRenderTarget.TargetableTexture, DestRenderTarget.ShaderResourceTexture, true, FResolveParams());
 				}
 			}
 		}
@@ -1358,25 +1365,26 @@ void FAtmosphericFogSceneInfo::RenderAtmosphereShaders(const FViewInfo& View, co
 			int32 Layer = Atmosphere3DTextureIndex;
 			{
 				const FSceneRenderTargetItem& DestRenderTarget = AtmosphereTextures->AtmosphereInscatter->GetRenderTargetItem();
-				RHISetRenderTarget(DestRenderTarget.TargetableTexture, FTextureRHIRef());
+				SetRenderTarget(RHICmdList, DestRenderTarget.TargetableTexture, FTextureRHIRef());
 
 				TShaderMapRef<FAtmospherePrecomputeVS> VertexShader(GetGlobalShaderMap());
 				TShaderMapRef<FAtmosphereGS> GeometryShader(GetGlobalShaderMap());
 				TShaderMapRef<FAtmosphereCopyInscatterFBackPS> PixelShader(GetGlobalShaderMap());
 
 				static FGlobalBoundShaderState BoundShaderState;
-				SetGlobalBoundShaderState( BoundShaderState, GFilterVertexDeclaration.VertexDeclarationRHI, *VertexShader, *PixelShader, *GeometryShader);
+				RHICmdList.CheckIsNull(); // need new approach for "static FGlobalBoundShaderState" for parallel rendering
+				SetGlobalBoundShaderState(RHICmdList, BoundShaderState, GFilterVertexDeclaration.VertexDeclarationRHI, *VertexShader, *PixelShader, *GeometryShader);
 
 				float r;
 				FVector4 DhdH;
 				GetLayerValue(Layer, r, DhdH);
 				GeometryShader->SetParameters(RHICmdList, Layer);
 				PixelShader->SetParameters(RHICmdList, r, DhdH, Layer, AtmosphereTextures);
-				DrawQuad(ViewRect, *VertexShader);
+				DrawQuad(RHICmdList, ViewRect, *VertexShader);
 
 				if (Atmosphere3DTextureIndex == Component->PrecomputeParams.InscatterAltitudeSampleNum - 1)
 				{
-					RHICopyToResolveTarget(DestRenderTarget.TargetableTexture, DestRenderTarget.ShaderResourceTexture, true, FResolveParams() );
+					RHICmdList.CopyToResolveTarget(DestRenderTarget.TargetableTexture, DestRenderTarget.ShaderResourceTexture, true, FResolveParams());
 				}
 			}
 		}
@@ -1385,21 +1393,21 @@ void FAtmosphericFogSceneInfo::RenderAtmosphereShaders(const FViewInfo& View, co
 	}
 }
 
-void FAtmosphericFogSceneInfo::PrecomputeAtmosphereData(const FViewInfo* View, FSceneViewFamily& ViewFamily)
+void FAtmosphericFogSceneInfo::PrecomputeAtmosphereData(FRHICommandList& RHICmdList, const FViewInfo* View, FSceneViewFamily& ViewFamily)
 {
 	// Set the view family's render target/viewport.
 	FIntPoint TexSize = GetTextureSize();
 	FIntRect ViewRect(0, 0, TexSize.X, TexSize.Y);
 
 	// turn off culling and blending
-	RHISetRasterizerState(TStaticRasterizerState<FM_Solid,CM_None>::GetRHI());
-	RHISetBlendState(TStaticBlendState<>::GetRHI());
+	RHICmdList.SetRasterizerState(TStaticRasterizerState<FM_Solid, CM_None>::GetRHI());
+	RHICmdList.SetBlendState(TStaticBlendState<>::GetRHI());
 	// turn off depth reads/writes
-	RHISetDepthStencilState(TStaticDepthStencilState<false,CF_Always>::GetRHI());
+	RHICmdList.SetDepthStencilState(TStaticDepthStencilState<false, CF_Always>::GetRHI());
 
-	RHISetViewport(0, 0, 0.0f, TexSize.X, TexSize.Y, 0.0f);
+	RHICmdList.SetViewport(0, 0, 0.0f, TexSize.X, TexSize.Y, 0.0f);
 
-	RenderAtmosphereShaders(*View, ViewRect);
+	RenderAtmosphereShaders(RHICmdList, *View, ViewRect);
 }
 
 void FAtmosphericFogSceneInfo::ReadPixelsPtr(TRefCountPtr<IPooledRenderTarget> RenderTarget, FColor* OutData, FIntRect InRect)
@@ -1440,7 +1448,7 @@ void FAtmosphericFogSceneInfo::Read3DPixelsPtr(TRefCountPtr<IPooledRenderTarget>
 	FMemory::Memcpy( OutData, Data.GetTypedData(), Data.Num() * sizeof(FFloat16Color) );
 }
 
-void FAtmosphericFogSceneInfo::PrecomputeTextures(const FViewInfo* View, FSceneViewFamily* ViewFamily)
+void FAtmosphericFogSceneInfo::PrecomputeTextures(FRHICommandList& RHICmdList, const FViewInfo* View, FSceneViewFamily* ViewFamily)
 {
 	check(Component != NULL);
 	if (AtmosphereTextures == NULL)
@@ -1456,7 +1464,7 @@ void FAtmosphericFogSceneInfo::PrecomputeTextures(const FViewInfo* View, FSceneV
 	// Atmosphere 
 	if (PrecomputeCounter.GetValue() < UAtmosphericFogComponent::EFinishedComputation)
 	{
-		PrecomputeAtmosphereData(View, *ViewFamily);
+		PrecomputeAtmosphereData(RHICmdList, View, *ViewFamily);
 
 		switch(AtmospherePhase)
 		{
@@ -1500,6 +1508,7 @@ void FAtmosphericFogSceneInfo::PrecomputeTextures(const FViewInfo* View, FSceneV
 				AtmoshpereOrder = 2;
 			}
 		}
+		RHICmdList.CheckIsNull(); // this is all synchronous readback
 
 		if (AtmospherePhase >= AP_MAX)
 		{
diff --git a/Engine/Source/Runtime/Renderer/Private/AtmosphereRendering.h b/Engine/Source/Runtime/Renderer/Private/AtmosphereRendering.h
--- a/Engine/Source/Runtime/Renderer/Private/AtmosphereRendering.h
+++ b/Engine/Source/Runtime/Renderer/Private/AtmosphereRendering.h
@@ -64,16 +64,16 @@ public:
 	~FAtmosphericFogSceneInfo();
 
 #if WITH_EDITOR
-	void PrecomputeTextures(const FViewInfo* View, FSceneViewFamily* ViewFamily);
+	void PrecomputeTextures(FRHICommandList& RHICmdList, const FViewInfo* View, FSceneViewFamily* ViewFamily);
 	void StartPrecompute();
 
 private:
 	/** Atmosphere pre-computation related functions */
 	FIntPoint GetTextureSize();
-	inline void DrawQuad(const FIntRect& ViewRect, FShader* VertexShader);
+	inline void DrawQuad(FRHICommandList& RHICmdList, const FIntRect& ViewRect, FShader* VertexShader);
 	void GetLayerValue(int Layer, float& AtmosphereR, FVector4& DhdH);
-	void RenderAtmosphereShaders(const FViewInfo& View, const FIntRect& ViewRect);
-	void PrecomputeAtmosphereData(const FViewInfo* View, FSceneViewFamily& ViewFamily);
+	void RenderAtmosphereShaders(FRHICommandList& RHICmdList, const FViewInfo& View, const FIntRect& ViewRect);
+	void PrecomputeAtmosphereData(FRHICommandList& RHICmdList, const FViewInfo* View, FSceneViewFamily& ViewFamily);
 
 	void ReadPixelsPtr(TRefCountPtr<IPooledRenderTarget> RenderTarget, FColor* OutData, FIntRect InRect);
 	void Read3DPixelsPtr(TRefCountPtr<IPooledRenderTarget> RenderTarget, FFloat16Color* OutData, FIntRect InRect, FIntPoint InZMinMax);
diff --git a/Engine/Source/Runtime/Renderer/Private/BasePassRendering.cpp b/Engine/Source/Runtime/Renderer/Private/BasePassRendering.cpp
--- a/Engine/Source/Runtime/Renderer/Private/BasePassRendering.cpp
+++ b/Engine/Source/Runtime/Renderer/Private/BasePassRendering.cpp
@@ -196,7 +196,7 @@ public:
 	}
 };
 
-void FBasePassOpaqueDrawingPolicyFactory::AddStaticMesh(FScene* Scene,FStaticMesh* StaticMesh)
+void FBasePassOpaqueDrawingPolicyFactory::AddStaticMesh(FRHICommandList& RHICmdList, FScene* Scene, FStaticMesh* StaticMesh)
 {
 	// Determine the mesh's material and blend mode.
 	const FMaterial* Material = StaticMesh->MaterialRenderProxy->GetMaterial(Scene->GetFeatureLevel());
@@ -208,8 +208,6 @@ void FBasePassOpaqueDrawingPolicyFactory::AddStaticMesh(FScene* Scene,FStaticMes
 	// Only draw opaque materials.
 	if( !IsTranslucentBlendMode(BlendMode) )
 	{
-		//@todo-rco: RHIPacketList
-		FRHICommandList& RHICmdList = FRHICommandList::GetNullRef();
 		ProcessBasePassMesh(
 			RHICmdList, 
 			FProcessBasePassMeshParameters(
diff --git a/Engine/Source/Runtime/Renderer/Private/BasePassRendering.h b/Engine/Source/Runtime/Renderer/Private/BasePassRendering.h
--- a/Engine/Source/Runtime/Renderer/Private/BasePassRendering.h
+++ b/Engine/Source/Runtime/Renderer/Private/BasePassRendering.h
@@ -758,7 +758,7 @@ public:
 		{}
 	};
 
-	static void AddStaticMesh(FScene* Scene,FStaticMesh* StaticMesh);
+	static void AddStaticMesh(FRHICommandList& RHICmdList, FScene* Scene, FStaticMesh* StaticMesh);
 	static bool DrawDynamicMesh(
 		FRHICommandList& RHICmdList, 
 		const FSceneView& View,
diff --git a/Engine/Source/Runtime/Renderer/Private/ClearQuad.cpp b/Engine/Source/Runtime/Renderer/Private/ClearQuad.cpp
--- a/Engine/Source/Runtime/Renderer/Private/ClearQuad.cpp
+++ b/Engine/Source/Runtime/Renderer/Private/ClearQuad.cpp
@@ -6,7 +6,7 @@
 FGlobalBoundShaderState GClearMRTBoundShaderState[8];
 
 // TODO support ExcludeRect
-void DrawClearQuadMRT(bool bClearColor,int32 NumClearColors,const FLinearColor* ClearColorArray,bool bClearDepth,float Depth,bool bClearStencil,uint32 Stencil)
+void DrawClearQuadMRT(FRHICommandList& RHICmdList, bool bClearColor, int32 NumClearColors, const FLinearColor* ClearColorArray, bool bClearDepth, float Depth, bool bClearStencil, uint32 Stencil)
 {
 	// Set new states
 	FBlendStateRHIParamRef BlendStateRHI;
@@ -24,9 +24,6 @@ void DrawClearQuadMRT(bool bClearColor,int32 NumClearColors,const FLinearColor*
 			: TStaticBlendStateWriteMask<CW_NONE,CW_NONE,CW_NONE,CW_NONE,CW_NONE,CW_NONE,CW_NONE,CW_NONE>::GetRHI();
 	}
 	
-	//@todo-rco: RHIPacketList
-	FRHICommandList& RHICmdList = FRHICommandList::GetNullRef();
-
 	const FDepthStencilStateRHIParamRef DepthStencilStateRHI = 
 		(bClearDepth && bClearStencil)
 			? TStaticDepthStencilState<
@@ -46,9 +43,9 @@ void DrawClearQuadMRT(bool bClearColor,int32 NumClearColors,const FLinearColor*
 						>::GetRHI()
 					: TStaticDepthStencilState<false, CF_Always>::GetRHI();
 
-	RHISetRasterizerState( TStaticRasterizerState<FM_Solid,CM_None>::GetRHI() );
-	RHISetBlendState( BlendStateRHI );
-	RHISetDepthStencilState( DepthStencilStateRHI );
+	RHICmdList.SetRasterizerState(TStaticRasterizerState<FM_Solid, CM_None>::GetRHI());
+	RHICmdList.SetBlendState(BlendStateRHI);
+	RHICmdList.SetDepthStencilState(DepthStencilStateRHI);
 
 	// Set the new shaders
 	TShaderMapRef<TOneColorVS<true> > VertexShader(GetGlobalShaderMap());
@@ -98,7 +95,7 @@ void DrawClearQuadMRT(bool bClearColor,int32 NumClearColors,const FLinearColor*
 		PixelShader = *MRTPixelShader;
 	}
 
-	SetGlobalBoundShaderState(GClearMRTBoundShaderState[FMath::Max(NumClearColors - 1, 0)], GetVertexDeclarationFVector4(), *VertexShader, PixelShader);
+	SetGlobalBoundShaderState(RHICmdList, GClearMRTBoundShaderState[FMath::Max(NumClearColors - 1, 0)], GetVertexDeclarationFVector4(), *VertexShader, PixelShader);
 	FLinearColor ShaderClearColors[MaxSimultaneousRenderTargets];
 	FMemory::MemZero(ShaderClearColors);
 
@@ -142,7 +139,7 @@ void DrawClearQuadMRT(bool bClearColor,int32 NumClearColors,const FLinearColor*
 			Vertices[8] = OuterVertices[0];
 			Vertices[9] = InnerVertices[0];
 
-			RHIDrawPrimitiveUP(PT_TriangleStrip, 8, Vertices, sizeof(Vertices[0]) );
+			DrawPrimitiveUP(RHICmdList, PT_TriangleStrip, 8, Vertices, sizeof(Vertices[0]) );
 		}
 		else*/
 		{
@@ -152,7 +149,7 @@ void DrawClearQuadMRT(bool bClearColor,int32 NumClearColors,const FLinearColor*
 			Vertices[1].Set(  1.0f,  1.0f, Depth, 1.0f );
 			Vertices[2].Set( -1.0f, -1.0f, Depth, 1.0f );
 			Vertices[3].Set(  1.0f, -1.0f, Depth, 1.0f );
-			RHIDrawPrimitiveUP(PT_TriangleStrip, 2, Vertices, sizeof(Vertices[0]) );
+			DrawPrimitiveUP(RHICmdList, PT_TriangleStrip, 2, Vertices, sizeof(Vertices[0]));
 		}
 	}
 }
diff --git a/Engine/Source/Runtime/Renderer/Private/ClearQuad.h b/Engine/Source/Runtime/Renderer/Private/ClearQuad.h
--- a/Engine/Source/Runtime/Renderer/Private/ClearQuad.h
+++ b/Engine/Source/Runtime/Renderer/Private/ClearQuad.h
@@ -3,9 +3,9 @@
 #pragma once
 
 // TODO support ExcludeRect
-void DrawClearQuadMRT(bool bClearColor,int32 NumClearColors,const FLinearColor* ClearColorArray,bool bClearDepth,float Depth,bool bClearStencil,uint32 Stencil);
+void DrawClearQuadMRT(FRHICommandList& RHICmdList, bool bClearColor, int32 NumClearColors, const FLinearColor* ClearColorArray, bool bClearDepth, float Depth, bool bClearStencil, uint32 Stencil);
 
-inline void DrawClearQuad(bool bClearColor,const FLinearColor& Color,bool bClearDepth,float Depth,bool bClearStencil,uint32 Stencil)
+inline void DrawClearQuad(FRHICommandList& RHICmdList, bool bClearColor, const FLinearColor& Color, bool bClearDepth, float Depth, bool bClearStencil, uint32 Stencil)
 {
-	DrawClearQuadMRT(bClearColor, 1, &Color, bClearDepth, Depth, bClearStencil, Stencil);
+	DrawClearQuadMRT(RHICmdList, bClearColor, 1, &Color, bClearDepth, Depth, bClearStencil, Stencil);
 }
\ No newline at end of file
