UE4 - more work preparing for parallel rendering. Merged the cmd list into the post process context, among other things.

diff --git a/Engine/Source/Runtime/Renderer/Private/CompositionLighting/CompositionLighting.cpp b/Engine/Source/Runtime/Renderer/Private/CompositionLighting/CompositionLighting.cpp
--- a/Engine/Source/Runtime/Renderer/Private/CompositionLighting/CompositionLighting.cpp
+++ b/Engine/Source/Runtime/Renderer/Private/CompositionLighting/CompositionLighting.cpp
@@ -247,18 +247,20 @@ void FCompositionLighting::ProcessBeforeBasePass(const FViewInfo& View)
 void FCompositionLighting::ProcessAfterBasePass(const FViewInfo& View)
 {
 	check(IsInRenderingThread());
+	//@todo-rco: RHIPacketList
+	FRHICommandList& RHICmdList = FRHICommandList::GetNullRef();
 
 	// might get renamed to refracted or ...WithAO
 	GSceneRenderTargets.GetSceneColor()->SetDebugName(TEXT("SceneColor"));
 	// to be able to observe results with VisualizeTexture
 
-	GRenderTargetPool.VisualizeTexture.SetCheckPoint(GSceneRenderTargets.GetSceneColor());
-	GRenderTargetPool.VisualizeTexture.SetCheckPoint(GSceneRenderTargets.GBufferA);
-	GRenderTargetPool.VisualizeTexture.SetCheckPoint(GSceneRenderTargets.GBufferB);
-	GRenderTargetPool.VisualizeTexture.SetCheckPoint(GSceneRenderTargets.GBufferC);
-	GRenderTargetPool.VisualizeTexture.SetCheckPoint(GSceneRenderTargets.GBufferD);
-	GRenderTargetPool.VisualizeTexture.SetCheckPoint(GSceneRenderTargets.GBufferE);
-	GRenderTargetPool.VisualizeTexture.SetCheckPoint(GSceneRenderTargets.ScreenSpaceAO);
+	GRenderTargetPool.VisualizeTexture.SetCheckPoint(RHICmdList, GSceneRenderTargets.GetSceneColor());
+	GRenderTargetPool.VisualizeTexture.SetCheckPoint(RHICmdList, GSceneRenderTargets.GBufferA);
+	GRenderTargetPool.VisualizeTexture.SetCheckPoint(RHICmdList, GSceneRenderTargets.GBufferB);
+	GRenderTargetPool.VisualizeTexture.SetCheckPoint(RHICmdList, GSceneRenderTargets.GBufferC);
+	GRenderTargetPool.VisualizeTexture.SetCheckPoint(RHICmdList, GSceneRenderTargets.GBufferD);
+	GRenderTargetPool.VisualizeTexture.SetCheckPoint(RHICmdList, GSceneRenderTargets.GBufferE);
+	GRenderTargetPool.VisualizeTexture.SetCheckPoint(RHICmdList, GSceneRenderTargets.ScreenSpaceAO);
 	
 	// so that the passes can register themselves to the graph
 	{
@@ -308,10 +310,12 @@ void FCompositionLighting::ProcessAfterBasePass(const FViewInfo& View)
 void FCompositionLighting::ProcessLighting(const FViewInfo& View)
 {
 	check(IsInRenderingThread());
+	//@todo-rco: RHIPacketList
+	FRHICommandList& RHICmdList = FRHICommandList::GetNullRef();
 
-	GRenderTargetPool.VisualizeTexture.SetCheckPoint(GSceneRenderTargets.ReflectiveShadowMapDiffuse);
-	GRenderTargetPool.VisualizeTexture.SetCheckPoint(GSceneRenderTargets.ReflectiveShadowMapNormal);
-	GRenderTargetPool.VisualizeTexture.SetCheckPoint(GSceneRenderTargets.ReflectiveShadowMapDepth);
+	GRenderTargetPool.VisualizeTexture.SetCheckPoint(RHICmdList, GSceneRenderTargets.ReflectiveShadowMapDiffuse);
+	GRenderTargetPool.VisualizeTexture.SetCheckPoint(RHICmdList, GSceneRenderTargets.ReflectiveShadowMapNormal);
+	GRenderTargetPool.VisualizeTexture.SetCheckPoint(RHICmdList, GSceneRenderTargets.ReflectiveShadowMapDepth);
 
 	{
 		FMemMark Mark(FMemStack::Get());
diff --git a/Engine/Source/Runtime/Renderer/Private/CompositionLighting/PostProcessAmbient.cpp b/Engine/Source/Runtime/Renderer/Private/CompositionLighting/PostProcessAmbient.cpp
--- a/Engine/Source/Runtime/Renderer/Private/CompositionLighting/PostProcessAmbient.cpp
+++ b/Engine/Source/Runtime/Renderer/Private/CompositionLighting/PostProcessAmbient.cpp
@@ -52,7 +52,7 @@ public:
 		const FPixelShaderRHIParamRef ShaderRHI = GetPixelShader();
 
 		FGlobalShader::SetParameters(RHICmdList, ShaderRHI, Context.View);
-		PostprocessParameter.SetPS(RHICmdList, ShaderRHI, Context, TStaticSamplerState<SF_Bilinear,AM_Clamp,AM_Clamp,AM_Clamp>::GetRHI());
+		PostprocessParameter.SetPS(ShaderRHI, Context, TStaticSamplerState<SF_Bilinear,AM_Clamp,AM_Clamp,AM_Clamp>::GetRHI());
 		DeferredParameters.Set(RHICmdList, ShaderRHI, Context.View);
 		CubemapShaderParameters.SetParameters(ShaderRHI, Entry);
 		SetTextureParameter(RHICmdList, ShaderRHI, PreIntegratedGF, PreIntegratedGFSampler, TStaticSamplerState<SF_Bilinear,AM_Clamp,AM_Clamp,AM_Clamp>::GetRHI(), GSystemTextures.PreintegratedGF->GetRenderTargetItem().ShaderResourceTexture );
@@ -84,13 +84,13 @@ void FRCPassPostProcessAmbient::Process(FRenderingCompositePassContext& Context)
 	const FSceneRenderTargetItem& DestRenderTarget = GSceneRenderTargets.GetSceneColor()->GetRenderTargetItem();
 
 	// Set the view family's render target/viewport.
-	RHISetRenderTarget(DestRenderTarget.TargetableTexture, FTextureRHIRef());	
+	SetRenderTarget(Context.RHICmdList, DestRenderTarget.TargetableTexture, FTextureRHIRef());
 	Context.SetViewportAndCallRHI(View.ViewRect);
 
 	// set the state
-	RHISetBlendState(TStaticBlendState<CW_RGB,BO_Add,BF_One,BF_One,BO_Add,BF_One,BF_One>::GetRHI());
-	RHISetRasterizerState(TStaticRasterizerState<>::GetRHI());
-	RHISetDepthStencilState(TStaticDepthStencilState<false,CF_Always>::GetRHI());
+	Context.RHICmdList.SetBlendState(TStaticBlendState<CW_RGB, BO_Add, BF_One, BF_One, BO_Add, BF_One, BF_One>::GetRHI());
+	Context.RHICmdList.SetRasterizerState(TStaticRasterizerState<>::GetRHI());
+	Context.RHICmdList.SetDepthStencilState(TStaticDepthStencilState<false, CF_Always>::GetRHI());
 
 	TShaderMapRef<FPostProcessVS> VertexShader(GetGlobalShaderMap());
 
@@ -103,6 +103,7 @@ void FRCPassPostProcessAmbient::Process(FRenderingCompositePassContext& Context)
 	TShaderMapRef<FPostProcessAmbientPS> PixelShader(GetGlobalShaderMap());
 
 	static FGlobalBoundShaderState BoundShaderState;
+	Context.RHICmdList.CheckIsNull(); // need new approach for "static FGlobalBoundShaderState" for parallel rendering
 
 	uint32 Count = Context.View.FinalPostProcessSettings.ContributingCubemaps.Num();
 	for(uint32 i = 0; i < Count; ++i)
@@ -110,14 +111,15 @@ void FRCPassPostProcessAmbient::Process(FRenderingCompositePassContext& Context)
 		if(i == 0)
 		{
 			// call it once after setting up the shader data to avoid the warnings in the function
-			SetGlobalBoundShaderState(BoundShaderState, GFilterVertexDeclaration.VertexDeclarationRHI, *VertexShader, *PixelShader);
+			SetGlobalBoundShaderState(Context.RHICmdList, BoundShaderState, GFilterVertexDeclaration.VertexDeclarationRHI, *VertexShader, *PixelShader);
 		}
 
 		//@todo-rco: RHIPacketList
-		PixelShader->SetParameters(FRHICommandList::GetNullRef(), Context, Context.View.FinalPostProcessSettings.ContributingCubemaps[i]);
+		PixelShader->SetParameters(Context.RHICmdList, Context, Context.View.FinalPostProcessSettings.ContributingCubemaps[i]);
 
 		// Draw a quad mapping scene color to the view's render target
-		DrawRectangle( 
+		DrawRectangle(
+			Context.RHICmdList,
 			0, 0,
 			View.ViewRect.Width(), View.ViewRect.Height(),
 			View.ViewRect.Min.X, View.ViewRect.Min.Y, 
@@ -128,7 +130,7 @@ void FRCPassPostProcessAmbient::Process(FRenderingCompositePassContext& Context)
 			EDRF_UseTriangleOptimization);
 	}
 
-	RHICopyToResolveTarget(DestRenderTarget.TargetableTexture, DestRenderTarget.ShaderResourceTexture, false, FResolveParams());
+	Context.RHICmdList.CopyToResolveTarget(DestRenderTarget.TargetableTexture, DestRenderTarget.ShaderResourceTexture, false, FResolveParams());
 }
 
 FPooledRenderTargetDesc FRCPassPostProcessAmbient::ComputeOutputDesc(EPassOutputId InPassOutputId) const
diff --git a/Engine/Source/Runtime/Renderer/Private/CompositionLighting/PostProcessAmbientOcclusion.cpp b/Engine/Source/Runtime/Renderer/Private/CompositionLighting/PostProcessAmbientOcclusion.cpp
--- a/Engine/Source/Runtime/Renderer/Private/CompositionLighting/PostProcessAmbientOcclusion.cpp
+++ b/Engine/Source/Runtime/Renderer/Private/CompositionLighting/PostProcessAmbientOcclusion.cpp
@@ -55,22 +55,22 @@ public:
 		AmbientOcclusionSetupParams.Bind(Initializer.ParameterMap, TEXT("AmbientOcclusionSetupParams"));
 	}
 
-	void SetParameters(FRHICommandList& RHICmdList, const FRenderingCompositePassContext& Context)
+	void SetParameters(const FRenderingCompositePassContext& Context)
 	{
 		const FFinalPostProcessSettings& Settings = Context.View.FinalPostProcessSettings;
 		const FPixelShaderRHIParamRef ShaderRHI = GetPixelShader();
 
-		FGlobalShader::SetParameters(RHICmdList, ShaderRHI, Context.View);
+		FGlobalShader::SetParameters(Context.RHICmdList, ShaderRHI, Context.View);
 
-		PostprocessParameter.SetPS(RHICmdList, ShaderRHI, Context, TStaticSamplerState<SF_Bilinear,AM_Clamp,AM_Clamp,AM_Clamp>::GetRHI());
-		DeferredParameters.Set(RHICmdList, ShaderRHI, Context.View);
+		PostprocessParameter.SetPS(ShaderRHI, Context, TStaticSamplerState<SF_Bilinear, AM_Clamp, AM_Clamp, AM_Clamp>::GetRHI());
+		DeferredParameters.Set(Context.RHICmdList, ShaderRHI, Context.View);
 
 		// e.g. 4 means the input texture is 4x smaller than the buffer size
 		uint32 ScaleToFullRes = GSceneRenderTargets.GetBufferSizeXY().X / Context.Pass->GetOutput(ePId_Output0)->RenderTargetDesc.Extent.X;
 
 		// /1000 to be able to define the value in that distance
 		FVector4 AmbientOcclusionSetupParamsValue = FVector4(ScaleToFullRes, Settings.AmbientOcclusionMipThreshold / ScaleToFullRes, 0, 0);
-		SetShaderValue(RHICmdList, ShaderRHI, AmbientOcclusionSetupParams, AmbientOcclusionSetupParamsValue);
+		SetShaderValue(Context.RHICmdList, ShaderRHI, AmbientOcclusionSetupParams, AmbientOcclusionSetupParamsValue);
 	}
 
 	// FShader interface.
@@ -235,22 +235,22 @@ public:
 		RandomNormalTextureSampler.Bind(Initializer.ParameterMap, TEXT("RandomNormalTextureSampler"));
 	}
 
-	void SetParameters(FRHICommandList& RHICmdList, const FRenderingCompositePassContext& Context, FIntPoint InputTextureSize)
+	void SetParameters(const FRenderingCompositePassContext& Context, FIntPoint InputTextureSize)
 	{
 		const FFinalPostProcessSettings& Settings = Context.View.FinalPostProcessSettings;
 		const FPixelShaderRHIParamRef ShaderRHI = GetPixelShader();
 
-		FGlobalShader::SetParameters(RHICmdList, ShaderRHI, Context.View);
+		FGlobalShader::SetParameters(Context.RHICmdList, ShaderRHI, Context.View);
 
 		// SF_Point is better than bilinear to avoid halos around objects
-		PostprocessParameter.SetPS(RHICmdList, ShaderRHI, Context, TStaticSamplerState<SF_Point,AM_Clamp,AM_Clamp,AM_Clamp>::GetRHI());
-		DeferredParameters.Set(RHICmdList, ShaderRHI, Context.View);
+		PostprocessParameter.SetPS(ShaderRHI, Context, TStaticSamplerState<SF_Point, AM_Clamp, AM_Clamp, AM_Clamp>::GetRHI());
+		DeferredParameters.Set(Context.RHICmdList, ShaderRHI, Context.View);
 
 		const FSceneRenderTargetItem& SSAORandomization = GSystemTextures.SSAORandomization->GetRenderTargetItem();
 
-		SetTextureParameter(RHICmdList, ShaderRHI, RandomNormalTexture, RandomNormalTextureSampler, TStaticSamplerState<SF_Point,AM_Wrap,AM_Wrap,AM_Wrap>::GetRHI(), SSAORandomization.ShaderResourceTexture);
+		SetTextureParameter(Context.RHICmdList, ShaderRHI, RandomNormalTexture, RandomNormalTextureSampler, TStaticSamplerState<SF_Point, AM_Wrap, AM_Wrap, AM_Wrap>::GetRHI(), SSAORandomization.ShaderResourceTexture);
 
-		ScreenSpaceAOandSSRShaderParams.Set(RHICmdList, Context.View, ShaderRHI, InputTextureSize);
+		ScreenSpaceAOandSSRShaderParams.Set(Context.RHICmdList, Context.View, ShaderRHI, InputTextureSize);
 	}
 	
 	// FShader interface.
@@ -321,15 +321,15 @@ public:
 		ScreenSpaceAOandSSRShaderParams.Bind(Initializer.ParameterMap);
 	}
 
-	void SetParameters(FRHICommandList& RHICmdList, const FRenderingCompositePassContext& Context, FIntPoint InputTextureSize)
+	void SetParameters(const FRenderingCompositePassContext& Context, FIntPoint InputTextureSize)
 	{
 		const FFinalPostProcessSettings& Settings = Context.View.FinalPostProcessSettings;
 		const FPixelShaderRHIParamRef ShaderRHI = GetPixelShader();
 
-		FGlobalShader::SetParameters(RHICmdList, ShaderRHI, Context.View);
-		PostprocessParameter.SetPS(RHICmdList, ShaderRHI, Context, TStaticSamplerState<SF_Point,AM_Clamp,AM_Clamp,AM_Clamp>::GetRHI());
-		DeferredParameters.Set(RHICmdList, ShaderRHI, Context.View);
-		ScreenSpaceAOandSSRShaderParams.Set(RHICmdList, Context.View, ShaderRHI, InputTextureSize);
+		FGlobalShader::SetParameters(Context.RHICmdList, ShaderRHI, Context.View);
+		PostprocessParameter.SetPS(ShaderRHI, Context, TStaticSamplerState<SF_Point,AM_Clamp,AM_Clamp,AM_Clamp>::GetRHI());
+		DeferredParameters.Set(Context.RHICmdList, ShaderRHI, Context.View);
+		ScreenSpaceAOandSSRShaderParams.Set(Context.RHICmdList, Context.View, ShaderRHI, InputTextureSize);
 	}
 
 	// FShader interface.
@@ -347,17 +347,18 @@ IMPLEMENT_SHADER_TYPE(,FPostProcessBasePassAOPS,TEXT("PostProcessAmbientOcclusio
 
 
 template <uint32 bInitialSetup>
-void FRCPassPostProcessAmbientOcclusionSetup::SetShaderSetupTempl(FRHICommandList& RHICmdList, const FRenderingCompositePassContext& Context)
+void FRCPassPostProcessAmbientOcclusionSetup::SetShaderSetupTempl(const FRenderingCompositePassContext& Context)
 {
 	TShaderMapRef<FPostProcessVS> VertexShader(GetGlobalShaderMap());
 	TShaderMapRef<FPostProcessAmbientOcclusionSetupPS<bInitialSetup> > PixelShader(GetGlobalShaderMap());
 
 	static FGlobalBoundShaderState BoundShaderState;
+	Context.RHICmdList.CheckIsNull(); // need new approach for "static FGlobalBoundShaderState" for parallel rendering
 
-	SetGlobalBoundShaderState(BoundShaderState, GFilterVertexDeclaration.VertexDeclarationRHI, *VertexShader, *PixelShader);
+	SetGlobalBoundShaderState(Context.RHICmdList, BoundShaderState, GFilterVertexDeclaration.VertexDeclarationRHI, *VertexShader, *PixelShader);
 
-	VertexShader->SetParameters(RHICmdList, Context);
-	PixelShader->SetParameters(RHICmdList, Context);
+	VertexShader->SetParameters(Context);
+	PixelShader->SetParameters(Context);
 }
 
 void FRCPassPostProcessAmbientOcclusionSetup::Process(FRenderingCompositePassContext& Context)
@@ -377,30 +378,29 @@ void FRCPassPostProcessAmbientOcclusionSetup::Process(FRenderingCompositePassCon
 	FIntRect DestRect = SrcRect  / ScaleFactor;
 
 	// Set the view family's render target/viewport.
-	RHISetRenderTarget(DestRenderTarget.TargetableTexture, FTextureRHIParamRef());
+	SetRenderTarget(Context.RHICmdList, DestRenderTarget.TargetableTexture, FTextureRHIParamRef());
 
 	Context.SetViewportAndCallRHI(DestRect);
 
 	// set the state
-	RHISetBlendState(TStaticBlendState<>::GetRHI());
-	RHISetRasterizerState(TStaticRasterizerState<>::GetRHI());
-	RHISetDepthStencilState(TStaticDepthStencilState<false,CF_Always>::GetRHI());
+	Context.RHICmdList.SetBlendState(TStaticBlendState<>::GetRHI());
+	Context.RHICmdList.SetRasterizerState(TStaticRasterizerState<>::GetRHI());
+	Context.RHICmdList.SetDepthStencilState(TStaticDepthStencilState<false, CF_Always>::GetRHI());
 
-	//@todo-rco: RHIPacketList
-	FRHICommandList& RHICmdList = FRHICommandList::GetNullRef();
 	if(IsInitialPass())
 	{
-		SetShaderSetupTempl<1>(RHICmdList, Context);
+		SetShaderSetupTempl<1>(Context);
 	}
 	else
 	{
-		SetShaderSetupTempl<0>(RHICmdList, Context);
+		SetShaderSetupTempl<0>(Context);
 	}
 
 	TShaderMapRef<FPostProcessVS> VertexShader(GetGlobalShaderMap());
 
 	// Draw a quad mapping scene color to the view's render target
 	DrawRectangle( 
+		Context.RHICmdList,
 		0, 0,
 		DestRect.Width(), DestRect.Height(),
 		SrcRect.Min.X, SrcRect.Min.Y, 
@@ -410,7 +410,7 @@ void FRCPassPostProcessAmbientOcclusionSetup::Process(FRenderingCompositePassCon
 		*VertexShader,
 		EDRF_UseTriangleOptimization);
 
-	RHICopyToResolveTarget(DestRenderTarget.TargetableTexture, DestRenderTarget.ShaderResourceTexture, false, FResolveParams());
+	Context.RHICmdList.CopyToResolveTarget(DestRenderTarget.TargetableTexture, DestRenderTarget.ShaderResourceTexture, false, FResolveParams());
 }
 
 FPooledRenderTargetDesc FRCPassPostProcessAmbientOcclusionSetup::ComputeOutputDesc(EPassOutputId InPassOutputId) const
@@ -464,15 +464,14 @@ void FRCPassPostProcessAmbientOcclusion::SetShaderTempl(const FRenderingComposit
 	TShaderMapRef<FPostProcessAmbientOcclusionPS<bTAOSetupAsInput, bDoUpsample, SampleSetQuality> > PixelShader(GetGlobalShaderMap());
 
 	static FGlobalBoundShaderState BoundShaderState;
+	Context.RHICmdList.CheckIsNull(); // need new approach for "static FGlobalBoundShaderState" for parallel rendering
 
-	SetGlobalBoundShaderState(BoundShaderState, GFilterVertexDeclaration.VertexDeclarationRHI, *VertexShader, *PixelShader);
+	SetGlobalBoundShaderState(Context.RHICmdList, BoundShaderState, GFilterVertexDeclaration.VertexDeclarationRHI, *VertexShader, *PixelShader);
 
 	const FPooledRenderTargetDesc* InputDesc0 = GetInputDesc(ePId_Input0);
 
-	//@todo-rco: RHIPacketList
-	FRHICommandList& RHICmdList = FRHICommandList::GetNullRef();
-	VertexShader->SetParameters(RHICmdList, Context);
-	PixelShader->SetParameters(RHICmdList, Context, InputDesc0->Extent);
+	VertexShader->SetParameters(Context);
+	PixelShader->SetParameters(Context, InputDesc0->Extent);
 }
 
 void FRCPassPostProcessAmbientOcclusion::Process(FRenderingCompositePassContext& Context)
@@ -505,13 +504,13 @@ void FRCPassPostProcessAmbientOcclusion::Process(FRenderingCompositePassContext&
 	FIntRect ViewRect = FIntRect::DivideAndRoundUp(View.ViewRect, ScaleToFullRes);
 
 	// Set the view family's render target/viewport.
-	RHISetRenderTarget(DestRenderTarget->TargetableTexture, FTextureRHIRef());	
+	SetRenderTarget(Context.RHICmdList, DestRenderTarget->TargetableTexture, FTextureRHIRef());
 	Context.SetViewportAndCallRHI(ViewRect);
 
 	// set the state
-	RHISetBlendState(TStaticBlendState<>::GetRHI());
-	RHISetRasterizerState(TStaticRasterizerState<>::GetRHI());
-	RHISetDepthStencilState(TStaticDepthStencilState<false,CF_Always>::GetRHI());
+	Context.RHICmdList.SetBlendState(TStaticBlendState<>::GetRHI());
+	Context.RHICmdList.SetRasterizerState(TStaticRasterizerState<>::GetRHI());
+	Context.RHICmdList.SetDepthStencilState(TStaticDepthStencilState<false, CF_Always>::GetRHI());
 
 	int32 SampleSetQuality = CVarAmbientOcclusionSampleSetQuality.GetValueOnRenderThread();
 
@@ -555,6 +554,7 @@ void FRCPassPostProcessAmbientOcclusion::Process(FRenderingCompositePassContext&
 
 	// Draw a quad mapping scene color to the view's render target
 	DrawRectangle( 
+		Context.RHICmdList,
 		0, 0,
 		ViewRect.Width(), ViewRect.Height(),
 		ViewRect.Min.X, ViewRect.Min.Y, 
@@ -565,7 +565,7 @@ void FRCPassPostProcessAmbientOcclusion::Process(FRenderingCompositePassContext&
 		EDRF_UseTriangleOptimization);
 
 
-	RHICopyToResolveTarget(DestRenderTarget->TargetableTexture, DestRenderTarget->ShaderResourceTexture, false, FResolveParams());
+	Context.RHICmdList.CopyToResolveTarget(DestRenderTarget->TargetableTexture, DestRenderTarget->ShaderResourceTexture, false, FResolveParams());
 }
 
 FPooledRenderTargetDesc FRCPassPostProcessAmbientOcclusion::ComputeOutputDesc(EPassOutputId InPassOutputId) const
@@ -600,28 +600,28 @@ void FRCPassPostProcessBasePassAO::Process(FRenderingCompositePassContext& Conte
 	const FSceneRenderTargetItem& DestRenderTarget = GSceneRenderTargets.GetSceneColor()->GetRenderTargetItem();
 
 	// Set the view family's render target/viewport.
-	RHISetRenderTarget(DestRenderTarget.TargetableTexture, FTextureRHIRef());	
+	SetRenderTarget(Context.RHICmdList, DestRenderTarget.TargetableTexture, FTextureRHIRef());
 	Context.SetViewportAndCallRHI(View.ViewRect);
 
 	// set the state
-	RHISetBlendState(TStaticBlendState<CW_RGB,BO_Add,BF_DestColor,BF_Zero>::GetRHI());
-	RHISetRasterizerState(TStaticRasterizerState<>::GetRHI());
-	RHISetDepthStencilState(TStaticDepthStencilState<false,CF_Always>::GetRHI());
+	Context.RHICmdList.SetBlendState(TStaticBlendState<CW_RGB, BO_Add, BF_DestColor, BF_Zero>::GetRHI());
+	Context.RHICmdList.SetRasterizerState(TStaticRasterizerState<>::GetRHI());
+	Context.RHICmdList.SetDepthStencilState(TStaticDepthStencilState<false, CF_Always>::GetRHI());
 
 	TShaderMapRef<FPostProcessVS> VertexShader(GetGlobalShaderMap());
 	TShaderMapRef<FPostProcessBasePassAOPS> PixelShader(GetGlobalShaderMap());
 
 	static FGlobalBoundShaderState BoundShaderState;
+	Context.RHICmdList.CheckIsNull(); // need new approach for "static FGlobalBoundShaderState" for parallel rendering
 
-	SetGlobalBoundShaderState(BoundShaderState, GFilterVertexDeclaration.VertexDeclarationRHI, *VertexShader, *PixelShader);
+	SetGlobalBoundShaderState(Context.RHICmdList, BoundShaderState, GFilterVertexDeclaration.VertexDeclarationRHI, *VertexShader, *PixelShader);
 
-	//@todo-rco: RHIPacketList
-	FRHICommandList& RHICmdList = FRHICommandList::GetNullRef();
-	VertexShader->SetParameters(RHICmdList, Context);
-	PixelShader->SetParameters(RHICmdList, Context, GSceneRenderTargets.GetBufferSizeXY());
+	VertexShader->SetParameters(Context);
+	PixelShader->SetParameters(Context, GSceneRenderTargets.GetBufferSizeXY());
 
 	// Draw a quad mapping scene color to the view's render target
 	DrawRectangle( 
+		Context.RHICmdList,
 		0, 0,
 		View.ViewRect.Width(), View.ViewRect.Height(),
 		View.ViewRect.Min.X, View.ViewRect.Min.Y, 
@@ -631,7 +631,7 @@ void FRCPassPostProcessBasePassAO::Process(FRenderingCompositePassContext& Conte
 		*VertexShader,
 		EDRF_UseTriangleOptimization);
 
-	RHICopyToResolveTarget(DestRenderTarget.TargetableTexture, DestRenderTarget.ShaderResourceTexture, false, FResolveParams());
+	Context.RHICmdList.CopyToResolveTarget(DestRenderTarget.TargetableTexture, DestRenderTarget.ShaderResourceTexture, false, FResolveParams());
 }
 
 FPooledRenderTargetDesc FRCPassPostProcessBasePassAO::ComputeOutputDesc(EPassOutputId InPassOutputId) const
diff --git a/Engine/Source/Runtime/Renderer/Private/CompositionLighting/PostProcessAmbientOcclusion.h b/Engine/Source/Runtime/Renderer/Private/CompositionLighting/PostProcessAmbientOcclusion.h
--- a/Engine/Source/Runtime/Renderer/Private/CompositionLighting/PostProcessAmbientOcclusion.h
+++ b/Engine/Source/Runtime/Renderer/Private/CompositionLighting/PostProcessAmbientOcclusion.h
@@ -26,7 +26,7 @@ private:
 	bool IsInitialPass() const;
 
 	template <uint32 bInitialSetup>
-	void SetShaderSetupTempl(FRHICommandList& RHICmdList, const FRenderingCompositePassContext& Context);
+	void SetShaderSetupTempl(const FRenderingCompositePassContext& Context);
 };
 
 // ePId_Input0: defines the resolution we compute AO and provides the normal
diff --git a/Engine/Source/Runtime/Renderer/Private/CompositionLighting/PostProcessDeferredDecals.cpp b/Engine/Source/Runtime/Renderer/Private/CompositionLighting/PostProcessDeferredDecals.cpp
--- a/Engine/Source/Runtime/Renderer/Private/CompositionLighting/PostProcessDeferredDecals.cpp
+++ b/Engine/Source/Runtime/Renderer/Private/CompositionLighting/PostProcessDeferredDecals.cpp
@@ -351,27 +351,27 @@ IMPLEMENT_SHADER_TYPE(,FStencilDecalMaskPS,TEXT("DeferredDecal"),TEXT("StencilDe
 FGlobalBoundShaderState StencilDecalMaskBoundShaderState;
 
 /** Draws a full view quad that sets stencil to 1 anywhere that decals should not be projected. */
-void StencilDecalMask(const FSceneView& View)
+void StencilDecalMask(FRHICommandList& RHICmdList, const FSceneView& View)
 {
 	SCOPED_DRAW_EVENT(StencilDecalMask, DEC_SCENE_ITEMS);
-	RHISetRasterizerState(TStaticRasterizerState<FM_Solid,CM_None>::GetRHI());
-	RHISetBlendState(TStaticBlendState<CW_NONE>::GetRHI());
-	RHISetRenderTarget(NULL, GSceneRenderTargets.GetSceneDepthSurface());
-	RHISetViewport(View.ViewRect.Min.X, View.ViewRect.Min.Y, 0.0f, View.ViewRect.Max.X, View.ViewRect.Max.Y, 1.0f);
+	RHICmdList.SetRasterizerState(TStaticRasterizerState<FM_Solid, CM_None>::GetRHI());
+	RHICmdList.SetBlendState(TStaticBlendState<CW_NONE>::GetRHI());
+	SetRenderTarget(RHICmdList, NULL, GSceneRenderTargets.GetSceneDepthSurface());
+	RHICmdList.SetViewport(View.ViewRect.Min.X, View.ViewRect.Min.Y, 0.0f, View.ViewRect.Max.X, View.ViewRect.Max.Y, 1.0f);
 
 	// Write 1 to highest bit of stencil to areas that should not receive decals
-	RHISetDepthStencilState(TStaticDepthStencilState<false,CF_Always,true,CF_Always,SO_Replace,SO_Replace,SO_Replace>::GetRHI(), 0x80);
+	RHICmdList.SetDepthStencilState(TStaticDepthStencilState<false, CF_Always, true, CF_Always, SO_Replace, SO_Replace, SO_Replace>::GetRHI(), 0x80);
 
 	TShaderMapRef<FScreenVS> ScreenVertexShader(GetGlobalShaderMap());
 	TShaderMapRef<FStencilDecalMaskPS> PixelShader(GetGlobalShaderMap());
+	RHICmdList.CheckIsNull(); // need new approach for "static FGlobalBoundShaderState" for parallel rendering
 	
-	SetGlobalBoundShaderState(StencilDecalMaskBoundShaderState, GFilterVertexDeclaration.VertexDeclarationRHI, *ScreenVertexShader, *PixelShader);
+	SetGlobalBoundShaderState(RHICmdList, StencilDecalMaskBoundShaderState, GFilterVertexDeclaration.VertexDeclarationRHI, *ScreenVertexShader, *PixelShader);
 
-	//@todo-rco: RHIPacketList
-	FRHICommandList& RHICmdList = FRHICommandList::GetNullRef();
 	PixelShader->SetParameters(RHICmdList, View);
 
 	DrawRectangle( 
+		RHICmdList,
 		0, 0, 
 		View.ViewRect.Width(), View.ViewRect.Height(),
 		View.ViewRect.Min.X, View.ViewRect.Min.Y, 
@@ -757,14 +757,14 @@ void FRCPassPostProcessDeferredDecals::Process(FRenderingCompositePassContext& C
 		SCOPED_DRAW_EVENT(DBufferClear, DEC_SCENE_ITEMS);
 		{
 			// could be optimized
-			RHISetRenderTarget(GSceneRenderTargets.DBufferA->GetRenderTargetItem().TargetableTexture, FTextureRHIParamRef());
-			RHIClear(true, FLinearColor(0, 0, 0, 1), false, 0, false, 0, FIntRect());
-			RHISetRenderTarget(GSceneRenderTargets.DBufferB->GetRenderTargetItem().TargetableTexture, FTextureRHIParamRef());
+			SetRenderTarget(RHICmdList, GSceneRenderTargets.DBufferA->GetRenderTargetItem().TargetableTexture, FTextureRHIParamRef());
+			RHICmdList.Clear(true, FLinearColor(0, 0, 0, 1), false, 0, false, 0, FIntRect());
+			SetRenderTarget(RHICmdList, GSceneRenderTargets.DBufferB->GetRenderTargetItem().TargetableTexture, FTextureRHIParamRef());
 			// todo: some hardware would like to have 0 or 1 for faster clear, we chose 128/255 to represent 0 (8 bit cannot represent 0.5f)
-			RHIClear(true, FLinearColor(128.0f/255.0f, 128.0f/255.0f, 128.0f/255.0f, 1), false, 0, false, 0, FIntRect());
-			RHISetRenderTarget(GSceneRenderTargets.DBufferC->GetRenderTargetItem().TargetableTexture, FTextureRHIParamRef());
+			RHICmdList.Clear(true, FLinearColor(128.0f / 255.0f, 128.0f / 255.0f, 128.0f / 255.0f, 1), false, 0, false, 0, FIntRect());
+			SetRenderTarget(RHICmdList, GSceneRenderTargets.DBufferC->GetRenderTargetItem().TargetableTexture, FTextureRHIParamRef());
 			// R:roughness, G:roughness opacity
-			RHIClear(true, FLinearColor(0, 1, 0, 1), false, 0, false, 0, FIntRect());
+			RHICmdList.Clear(true, FLinearColor(0, 1, 0, 1), false, 0, false, 0, FIntRect());
 		}
 	}
 
@@ -847,7 +847,7 @@ void FRCPassPostProcessDeferredDecals::Process(FRenderingCompositePassContext& C
 		// Setup a stencil mask to prevent certain pixels from receiving deferred decals
 		if(bStencilDecals)
 		{
-			StencilDecalMask(View);
+			StencilDecalMask(RHICmdList, View);
 		}
 
 		// Sort by sort order to allow control over composited result
@@ -943,16 +943,16 @@ void FRCPassPostProcessDeferredDecals::Process(FRenderingCompositePassContext& C
 							RenderTargets[2] = GSceneRenderTargets.GBufferB->GetRenderTargetItem().TargetableTexture;
 							RenderTargets[3] = GSceneRenderTargets.GBufferC->GetRenderTargetItem().TargetableTexture;
 							RenderTargets[4] = GSceneRenderTargets.GBufferD->GetRenderTargetItem().TargetableTexture;
-							RHISetRenderTargets( ARRAY_COUNT(RenderTargets), RenderTargets, GSceneRenderTargets.GetSceneDepthSurface(), 0, NULL);
+							SetRenderTargets(RHICmdList, ARRAY_COUNT(RenderTargets), RenderTargets, GSceneRenderTargets.GetSceneDepthSurface(), 0, NULL);
 						}
 						break;
 
 					case RTM_GBufferNormal:
-						RHISetRenderTarget(GSceneRenderTargets.GBufferA->GetRenderTargetItem().TargetableTexture, GSceneRenderTargets.GetSceneDepthSurface());	
+						SetRenderTarget(RHICmdList, GSceneRenderTargets.GBufferA->GetRenderTargetItem().TargetableTexture, GSceneRenderTargets.GetSceneDepthSurface());
 						break;
 					
 					case RTM_SceneColor:
-						RHISetRenderTarget(PassOutputs[0].PooledRenderTarget->GetRenderTargetItem().TargetableTexture, GSceneRenderTargets.GetSceneDepthSurface());	
+						SetRenderTarget(RHICmdList, PassOutputs[0].PooledRenderTarget->GetRenderTargetItem().TargetableTexture, GSceneRenderTargets.GetSceneDepthSurface());
 						break;
 
 					case RTM_DBuffer:
@@ -961,7 +961,7 @@ void FRCPassPostProcessDeferredDecals::Process(FRenderingCompositePassContext& C
 							RenderTargets[0] = GSceneRenderTargets.DBufferA->GetRenderTargetItem().TargetableTexture;
 							RenderTargets[1] = GSceneRenderTargets.DBufferB->GetRenderTargetItem().TargetableTexture;
 							RenderTargets[2] = GSceneRenderTargets.DBufferC->GetRenderTargetItem().TargetableTexture;
-							RHISetRenderTargets( ARRAY_COUNT(RenderTargets), RenderTargets, GSceneRenderTargets.GetSceneDepthSurface(), 0, NULL);
+							SetRenderTargets(RHICmdList, ARRAY_COUNT(RenderTargets), RenderTargets, GSceneRenderTargets.GetSceneDepthSurface(), 0, NULL);
 						}
 						break;
 
@@ -1079,9 +1079,9 @@ void FRCPassPostProcessDeferredDecals::Process(FRenderingCompositePassContext& C
 	if(RenderStage == 0)
 	{
 		// before BasePass
-		GRenderTargetPool.VisualizeTexture.SetCheckPoint(GSceneRenderTargets.DBufferA);
-		GRenderTargetPool.VisualizeTexture.SetCheckPoint(GSceneRenderTargets.DBufferB);
-		GRenderTargetPool.VisualizeTexture.SetCheckPoint(GSceneRenderTargets.DBufferC);
+		GRenderTargetPool.VisualizeTexture.SetCheckPoint(RHICmdList, GSceneRenderTargets.DBufferA);
+		GRenderTargetPool.VisualizeTexture.SetCheckPoint(RHICmdList, GSceneRenderTargets.DBufferB);
+		GRenderTargetPool.VisualizeTexture.SetCheckPoint(RHICmdList, GSceneRenderTargets.DBufferC);
 	}
 }
 
diff --git a/Engine/Source/Runtime/Renderer/Private/CompositionLighting/PostProcessLpvIndirect.cpp b/Engine/Source/Runtime/Renderer/Private/CompositionLighting/PostProcessLpvIndirect.cpp
--- a/Engine/Source/Runtime/Renderer/Private/CompositionLighting/PostProcessLpvIndirect.cpp
+++ b/Engine/Source/Runtime/Renderer/Private/CompositionLighting/PostProcessLpvIndirect.cpp
@@ -67,7 +67,6 @@ public:
 	}
 
 	void SetParameters(	
-		FRHICommandList& RHICmdList, 
 #if LPV_VOLUME_TEXTURE
 		FTextureRHIParamRef* LpvBufferSRVsIn, 
 #else
@@ -79,29 +78,29 @@ public:
 	{
 		const FPixelShaderRHIParamRef ShaderRHI = GetPixelShader();
 
-		RHICmdList.CheckIsNull();
+		Context.RHICmdList.CheckIsNull();
 
-		SetUniformBufferParameter(RHICmdList, ShaderRHI, GetUniformBufferParameter<FLpvReadUniformBufferParameters>(), LpvUniformBuffer );
+		SetUniformBufferParameter(Context.RHICmdList, ShaderRHI, GetUniformBufferParameter<FLpvReadUniformBufferParameters>(), LpvUniformBuffer);
 
 #if LPV_VOLUME_TEXTURE
 		for ( int i=0; i<7; i++ )
 		{
 			if ( LpvBufferSRVParameters[i].IsBound() )
 			{
-				RHISetShaderTexture( ShaderRHI, LpvBufferSRVParameters[i].GetBaseIndex(), LpvBufferSRVsIn[i] );
-				SetTextureParameter(RHICmdList, ShaderRHI, LpvBufferSRVParameters[i], LpvVolumeTextureSampler, TStaticSamplerState<SF_Bilinear,AM_Clamp,AM_Clamp,AM_Clamp>::GetRHI(), LpvBufferSRVsIn[i] );
+				Context.RHICmdList.SetShaderTexture(ShaderRHI, LpvBufferSRVParameters[i].GetBaseIndex(), LpvBufferSRVsIn[i]);
+				SetTextureParameter(Context.RHICmdList, ShaderRHI, LpvBufferSRVParameters[i], LpvVolumeTextureSampler, TStaticSamplerState<SF_Bilinear, AM_Clamp, AM_Clamp, AM_Clamp>::GetRHI(), LpvBufferSRVsIn[i]);
 			}
 		}
 #else
 		if ( LpvBufferSRV.IsBound() )
 		{
-			RHISetShaderResourceViewParameter( ShaderRHI, LpvBufferSRV.GetBaseIndex(), LpvBufferSRVIn );
+			Context.RHICmdList.SetShaderResourceViewParameter( ShaderRHI, LpvBufferSRV.GetBaseIndex(), LpvBufferSRVIn );
 		}
 #endif
-		FGlobalShader::SetParameters(RHICmdList, ShaderRHI, Context.View);
-		PostprocessParameter.SetPS(RHICmdList, ShaderRHI, Context, TStaticSamplerState<SF_Bilinear,AM_Clamp,AM_Clamp,AM_Clamp>::GetRHI());
-		DeferredParameters.Set(RHICmdList, ShaderRHI, Context.View);
-		SetTextureParameter(RHICmdList, ShaderRHI, PreIntegratedGF, PreIntegratedGFSampler, TStaticSamplerState<SF_Bilinear,AM_Clamp,AM_Clamp,AM_Clamp>::GetRHI(), GSystemTextures.PreintegratedGF->GetRenderTargetItem().ShaderResourceTexture );
+		FGlobalShader::SetParameters(Context.RHICmdList, ShaderRHI, Context.View);
+		PostprocessParameter.SetPS(ShaderRHI, Context, TStaticSamplerState<SF_Bilinear, AM_Clamp, AM_Clamp, AM_Clamp>::GetRHI());
+		DeferredParameters.Set(Context.RHICmdList, ShaderRHI, Context.View);
+		SetTextureParameter(Context.RHICmdList, ShaderRHI, PreIntegratedGF, PreIntegratedGFSampler, TStaticSamplerState<SF_Bilinear, AM_Clamp, AM_Clamp, AM_Clamp>::GetRHI(), GSystemTextures.PreintegratedGF->GetRenderTargetItem().ShaderResourceTexture);
 	}
 	
 	// FShader interface.
@@ -153,27 +152,23 @@ void FRCPassPostProcessLpvIndirect::Process(FRenderingCompositePassContext& Cont
 		DestColorRenderTarget.TargetableTexture,
 	};
 
-	//@todo-rco: RHIPacketList
-	FRHICommandList& RHICmdList = FRHICommandList::GetNullRef();
-	RHICmdList.CheckIsNull();
-
 	// Set the view family's render target/viewport.
-	RHISetRenderTargets(1, RenderTargets, GSceneRenderTargets.GetSceneDepthSurface(), 0, NULL);
+	SetRenderTargets(Context.RHICmdList, 1, RenderTargets, GSceneRenderTargets.GetSceneDepthSurface(), 0, NULL);
 
 	Context.SetViewportAndCallRHI(View.ViewRect);
 
 	// set the state
 	if ( ViewFamily.EngineShowFlags.LpvLightingOnly )
 	{
-		RHISetBlendState( TStaticBlendState<>::GetRHI() );
+		Context.RHICmdList.SetBlendState(TStaticBlendState<>::GetRHI());
 	}
 	else
 	{
 		// additive blending
-		RHISetBlendState(TStaticBlendState<CW_RGB,BO_Add,BF_One,BF_One>::GetRHI());
+		Context.RHICmdList.SetBlendState(TStaticBlendState<CW_RGB, BO_Add, BF_One, BF_One>::GetRHI());
 	}
-	RHISetRasterizerState(TStaticRasterizerState<>::GetRHI());
-	RHISetDepthStencilState(TStaticDepthStencilState<false,CF_Always>::GetRHI());
+	Context.RHICmdList.SetRasterizerState(TStaticRasterizerState<>::GetRHI());
+	Context.RHICmdList.SetDepthStencilState(TStaticDepthStencilState<false, CF_Always>::GetRHI());
 
 	TShaderMapRef<FPostProcessVS> VertexShader(GetGlobalShaderMap());
 
@@ -195,9 +190,10 @@ void FRCPassPostProcessLpvIndirect::Process(FRenderingCompositePassContext& Cont
 	TShaderMapRef<FPostProcessLpvIndirectPS> PixelShader(GetGlobalShaderMap());
 
 	static FGlobalBoundShaderState BoundShaderState;
+	Context.RHICmdList.CheckIsNull(); // need new approach for "static FGlobalBoundShaderState" for parallel rendering
 
 	// call it once after setting up the shader data to avoid the warnings in the function
-	SetGlobalBoundShaderState(BoundShaderState, GFilterVertexDeclaration.VertexDeclarationRHI, *VertexShader, *PixelShader);
+	SetGlobalBoundShaderState(Context.RHICmdList, BoundShaderState, GFilterVertexDeclaration.VertexDeclarationRHI, *VertexShader, *PixelShader);
 
 	FLpvReadUniformBufferParameters	LpvReadUniformBufferParams;
 	FLpvReadUniformBufferRef LpvReadUniformBuffer;
@@ -211,14 +207,15 @@ void FRCPassPostProcessLpvIndirect::Process(FRenderingCompositePassContext& Cont
 	{
 		LpvBufferSrvs[i] = Lpv->GetLpvBufferSrv(i);
 	}
-	PixelShader->SetParameters(RHICmdList, LpvBufferSrvs, LpvReadUniformBuffer, Context );
+	PixelShader->SetParameters(LpvBufferSrvs, LpvReadUniformBuffer, Context);
 #else
 	FShaderResourceViewRHIParamRef LpvBufferSrv = Lpv->GetLpvBufferSrv();
-	PixelShader->SetParameters(RHICmdList, LpvBufferSrv, LpvReadUniformBuffer, Context );
+	PixelShader->SetParameters(LpvBufferSrv, LpvReadUniformBuffer, Context);
 #endif
 
 	// Draw a quad mapping scene color to the view's render target
 	DrawRectangle( 
+		Context.RHICmdList,
 		0, 0,
 		View.ViewRect.Width(), View.ViewRect.Height(),
 		View.ViewRect.Min.X, View.ViewRect.Min.Y, 
@@ -227,7 +224,7 @@ void FRCPassPostProcessLpvIndirect::Process(FRenderingCompositePassContext& Cont
 		GSceneRenderTargets.GetBufferSizeXY(),
 		*VertexShader);
 
-	RHICopyToResolveTarget(DestColorRenderTarget.TargetableTexture, DestColorRenderTarget.ShaderResourceTexture, false, FResolveParams());
+	Context.RHICmdList.CopyToResolveTarget(DestColorRenderTarget.TargetableTexture, DestColorRenderTarget.ShaderResourceTexture, false, FResolveParams());
 }
 
 FPooledRenderTargetDesc FRCPassPostProcessLpvIndirect::ComputeOutputDesc(EPassOutputId InPassOutputId) const
diff --git a/Engine/Source/Runtime/Renderer/Private/CompositionLighting/PostProcessPassThrough.cpp b/Engine/Source/Runtime/Renderer/Private/CompositionLighting/PostProcessPassThrough.cpp
--- a/Engine/Source/Runtime/Renderer/Private/CompositionLighting/PostProcessPassThrough.cpp
+++ b/Engine/Source/Runtime/Renderer/Private/CompositionLighting/PostProcessPassThrough.cpp
@@ -49,7 +49,7 @@ public:
 		FRHICommandList& RHICmdList = FRHICommandList::GetNullRef();
 		FGlobalShader::SetParameters(RHICmdList, ShaderRHI, Context.View);
 
-		PostprocessParameter.SetPS(RHICmdList, ShaderRHI, Context, TStaticSamplerState<SF_Point,AM_Clamp,AM_Clamp,AM_Clamp>::GetRHI());
+		PostprocessParameter.SetPS(ShaderRHI, Context, TStaticSamplerState<SF_Point,AM_Clamp,AM_Clamp,AM_Clamp>::GetRHI());
 	}
 };
 
@@ -100,36 +100,36 @@ void FRCPassPostProcessPassThrough::Process(FRenderingCompositePassContext& Cont
 	const FSceneRenderTargetItem& DestRenderTarget = Dest ? Dest->GetRenderTargetItem() : PassOutputs[0].RequestSurface(Context);
 
 	// Set the view family's render target/viewport.
-	RHISetRenderTarget(DestRenderTarget.TargetableTexture, FTextureRHIRef());	
+	SetRenderTarget(Context.RHICmdList, DestRenderTarget.TargetableTexture, FTextureRHIRef());
 	Context.SetViewportAndCallRHI(0, 0, 0.0f, DestSize.X, DestSize.Y, 1.0f);
 
 	// set the state
 	if(bAdditiveBlend)
 	{
-		RHISetBlendState(TStaticBlendState<CW_RGB,BO_Add,BF_One,BF_One,BO_Add,BF_One,BF_One>::GetRHI());
+		Context.RHICmdList.SetBlendState(TStaticBlendState<CW_RGB, BO_Add, BF_One, BF_One, BO_Add, BF_One, BF_One>::GetRHI());
 	}
 	else
 	{
-		RHISetBlendState(TStaticBlendState<>::GetRHI());
+		Context.RHICmdList.SetBlendState(TStaticBlendState<>::GetRHI());
 	}
 
-	RHISetRasterizerState(TStaticRasterizerState<>::GetRHI());
-	RHISetDepthStencilState(TStaticDepthStencilState<false,CF_Always>::GetRHI());
+	Context.RHICmdList.SetRasterizerState(TStaticRasterizerState<>::GetRHI());
+	Context.RHICmdList.SetDepthStencilState(TStaticDepthStencilState<false, CF_Always>::GetRHI());
 
 	TShaderMapRef<FPostProcessVS> VertexShader(GetGlobalShaderMap());
 	TShaderMapRef<FPostProcessPassThroughPS> PixelShader(GetGlobalShaderMap());
 
 	static FGlobalBoundShaderState BoundShaderState;
+	Context.RHICmdList.CheckIsNull(); // need new approach for "static FGlobalBoundShaderState" for parallel rendering
 
-	SetGlobalBoundShaderState(BoundShaderState, GFilterVertexDeclaration.VertexDeclarationRHI, *VertexShader, *PixelShader);
+	SetGlobalBoundShaderState(Context.RHICmdList, BoundShaderState, GFilterVertexDeclaration.VertexDeclarationRHI, *VertexShader, *PixelShader);
 
-	//@todo-rco: RHIPacketList
-	FRHICommandList& RHICmdList = FRHICommandList::GetNullRef();
-	VertexShader->SetParameters(RHICmdList, Context);
+	VertexShader->SetParameters(Context);
 	PixelShader->SetParameters(Context);
 
 	// Draw a quad mapping scene color to the view's render target
 	DrawRectangle(
+		Context.RHICmdList,
 		DestRect.Min.X, DestRect.Min.Y,
 		DestRect.Width(), DestRect.Height(),
 		SrcRect.Min.X, SrcRect.Min.Y,
@@ -140,7 +140,7 @@ void FRCPassPostProcessPassThrough::Process(FRenderingCompositePassContext& Cont
 		EDRF_UseTriangleOptimization);
 
 
-	RHICopyToResolveTarget(DestRenderTarget.TargetableTexture, DestRenderTarget.ShaderResourceTexture, false, FResolveParams());
+	Context.RHICmdList.CopyToResolveTarget(DestRenderTarget.TargetableTexture, DestRenderTarget.ShaderResourceTexture, false, FResolveParams());
 }
 
 
diff --git a/Engine/Source/Runtime/Renderer/Private/DeferredShadingRenderer.cpp b/Engine/Source/Runtime/Renderer/Private/DeferredShadingRenderer.cpp
--- a/Engine/Source/Runtime/Renderer/Private/DeferredShadingRenderer.cpp
+++ b/Engine/Source/Runtime/Renderer/Private/DeferredShadingRenderer.cpp
@@ -114,7 +114,7 @@ extern FGlobalBoundShaderState GClearMRTBoundShaderState[8];
 /** 
 * Clears view where Z is still at the maximum value (ie no geometry rendered)
 */
-void FDeferredShadingSceneRenderer::ClearGBufferAtMaxZ()
+void FDeferredShadingSceneRenderer::ClearGBufferAtMaxZ(FRHICommandList& RHICmdList)
 {
 	// Assumes BeginRenderingSceneColor() has been called before this function
 	SCOPED_DRAW_EVENT(ClearGBufferAtMaxZ, DEC_SCENE_ITEMS);
@@ -154,16 +154,13 @@ void FDeferredShadingSceneRenderer::ClearGBufferAtMaxZ()
 		break;
 	}
 
-	//@todo-rco: RHIPacketList
-	FRHICommandList& RHICmdList = FRHICommandList::GetNullRef();
-
-	SetGlobalBoundShaderState(GClearMRTBoundShaderState[NumActiveRenderTargets - 1], GetVertexDeclarationFVector4(), *VertexShader, PixelShader);
+	SetGlobalBoundShaderState(RHICmdList, GClearMRTBoundShaderState[NumActiveRenderTargets - 1], GetVertexDeclarationFVector4(), *VertexShader, PixelShader);
 
 	// Opaque rendering, depth test but no depth writes
-	RHISetRasterizerState( TStaticRasterizerState<FM_Solid,CM_None>::GetRHI() );
-	RHISetBlendState(TStaticBlendStateWriteMask<>::GetRHI());
+	RHICmdList.SetRasterizerState(TStaticRasterizerState<FM_Solid, CM_None>::GetRHI());
+	RHICmdList.SetBlendState(TStaticBlendStateWriteMask<>::GetRHI());
 	// Note, this is a reversed Z depth surface, using CF_GreaterEqual.
-	RHISetDepthStencilState(TStaticDepthStencilState<false,CF_GreaterEqual>::GetRHI());
+	RHICmdList.SetDepthStencilState(TStaticDepthStencilState<false, CF_GreaterEqual>::GetRHI());
 
 	// Clear each viewport by drawing background color at MaxZ depth
 	for(int32 ViewIndex = 0;ViewIndex < Views.Num();ViewIndex++)
@@ -173,13 +170,13 @@ void FDeferredShadingSceneRenderer::ClearGBufferAtMaxZ()
 		FViewInfo& View = Views[ViewIndex];
 
 		// Set viewport for this view
-		RHISetViewport(View.ViewRect.Min.X, View.ViewRect.Min.Y, 0, View.ViewRect.Max.X, View.ViewRect.Max.Y, 1);
+		RHICmdList.SetViewport(View.ViewRect.Min.X, View.ViewRect.Min.Y, 0, View.ViewRect.Max.X, View.ViewRect.Max.Y, 1);
 
 		// Setup PS
 		SetShaderValueArray(RHICmdList, PixelShader->GetPixelShader(),PixelShader->ColorParameter, ClearColors, NumActiveRenderTargets);
 
 		// Render quad
-		RHIDrawPrimitiveUP(PT_TriangleStrip, 2, ClearQuadVertices, sizeof(ClearQuadVertices[0]) );
+		DrawPrimitiveUP(RHICmdList, PT_TriangleStrip, 2, ClearQuadVertices, sizeof(ClearQuadVertices[0]));
 	}
 }
 
@@ -297,19 +294,22 @@ bool FDeferredShadingSceneRenderer::RenderBasePassDynamicData(FViewInfo& View)
 
 	if( !View.Family->EngineShowFlags.CompositeEditorPrimitives )
 	{
+		//@todo-rco: RHIPacketList
+		FRHICommandList& RHICmdList = FRHICommandList::GetNullRef();
+
 		const bool bNeedToSwitchVerticalAxis = IsES2Platform(GRHIShaderPlatform) && !IsPCPlatform(GRHIShaderPlatform);
 
 		// Draw the base pass for the view's batched mesh elements.
 		bDirty = DrawViewElements<FBasePassOpaqueDrawingPolicyFactory>(View,FBasePassOpaqueDrawingPolicyFactory::ContextType(false, ESceneRenderTargetsMode::DontSet), SDPG_World, true) || bDirty;
 
 		// Draw the view's batched simple elements(lines, sprites, etc).
-		bDirty = View.BatchedViewElements.Draw(bNeedToSwitchVerticalAxis, View.ViewProjectionMatrix, View.ViewRect.Width(), View.ViewRect.Height(), false) || bDirty;
+		bDirty = View.BatchedViewElements.Draw(RHICmdList, bNeedToSwitchVerticalAxis, View.ViewProjectionMatrix, View.ViewRect.Width(), View.ViewRect.Height(), false) || bDirty;
 		
 		// Draw foreground objects last
 		bDirty = DrawViewElements<FBasePassOpaqueDrawingPolicyFactory>(View,FBasePassOpaqueDrawingPolicyFactory::ContextType(false, ESceneRenderTargetsMode::DontSet), SDPG_Foreground, true) || bDirty;
 
 		// Draw the view's batched simple elements(lines, sprites, etc).
-		bDirty = View.TopBatchedViewElements.Draw(bNeedToSwitchVerticalAxis, View.ViewProjectionMatrix, View.ViewRect.Width(), View.ViewRect.Height(), false) || bDirty;
+		bDirty = View.TopBatchedViewElements.Draw(RHICmdList, bNeedToSwitchVerticalAxis, View.ViewProjectionMatrix, View.ViewRect.Width(), View.ViewRect.Height(), false) || bDirty;
 
 	}
 
@@ -606,7 +606,7 @@ void FDeferredShadingSceneRenderer::Render()
 	{
 		// Clears view by drawing quad at maximum Z
 		// TODO: if all the platforms have fast color clears, we can replace this with an RHIClear.
-		ClearGBufferAtMaxZ();
+		ClearGBufferAtMaxZ(RHICmdList);
 
 		bRequiresFarZQuadClear = false;
 	}
@@ -629,7 +629,7 @@ void FDeferredShadingSceneRenderer::Render()
 
 	// Update the quarter-sized depth buffer with the current contents of the scene depth texture.
 	// This needs to happen before occlusion tests, which makes use of the small depth buffer.
-	UpdateDownsampledDepthSurface();
+	UpdateDownsampledDepthSurface(RHICmdList);
 
 	// Issue occlusion queries
 	// This is done after the downsampled depth buffer is created so that it can be used for issuing queries
@@ -654,7 +654,7 @@ void FDeferredShadingSceneRenderer::Render()
 			GCompositionLighting.ProcessAfterBasePass(Views[ViewIndex]);
 		}
 
-		RenderDynamicSkyLighting();
+		RenderDynamicSkyLighting(RHICmdList);
 
 		// Clear the translucent lighting volumes before we accumulate
 		ClearTranslucentVolumeLighting(RHICmdList);
@@ -666,7 +666,7 @@ void FDeferredShadingSceneRenderer::Render()
 		CompositeIndirectTranslucentVolumeLighting(RHICmdList);
 
 		// Filter the translucency lighting volume now that it is complete
-		FilterTranslucentVolumeLighting();
+		FilterTranslucentVolumeLighting(RHICmdList);
 
 		// Clear LPVs for all views
 		if ( IsFeatureLevelSupported(GRHIShaderPlatform, ERHIFeatureLevel::SM5) )
@@ -689,7 +689,7 @@ void FDeferredShadingSceneRenderer::Render()
 	if( ViewFamily.EngineShowFlags.StationaryLightOverlap &&
 		FeatureLevel >= ERHIFeatureLevel::SM4)
 	{
-		RenderStationaryLightOverlap();
+		RenderStationaryLightOverlap(RHICmdList);
 	}
 
 	FLightShaftsOutput LightShaftOutput;
@@ -697,7 +697,7 @@ void FDeferredShadingSceneRenderer::Render()
 	// Draw Lightshafts
 	if (ViewFamily.EngineShowFlags.LightShafts)
 	{
-		LightShaftOutput = RenderLightShaftOcclusion();
+		LightShaftOutput = RenderLightShaftOcclusion(RHICmdList);
 	}
 
 	// Draw atmosphere
@@ -718,7 +718,7 @@ void FDeferredShadingSceneRenderer::Render()
 			if (Scene->bIsEditorScene)
 			{
 				// Precompute Atmospheric Textures
-				Scene->AtmosphericFog->PrecomputeTextures(Views.GetTypedData(), &ViewFamily);
+				Scene->AtmosphericFog->PrecomputeTextures(RHICmdList, Views.GetTypedData(), &ViewFamily);
 			}
 #endif
 			RenderAtmosphere(LightShaftOutput);
@@ -730,7 +730,7 @@ void FDeferredShadingSceneRenderer::Render()
 	// Draw fog.
 	if(ShouldRenderFog(ViewFamily))
 	{
-		RenderFog(LightShaftOutput);
+		RenderFog(RHICmdList, LightShaftOutput);
 	}
 
 	// No longer needed, release
@@ -753,13 +753,13 @@ void FDeferredShadingSceneRenderer::Render()
 
 	if (ViewFamily.EngineShowFlags.LightShafts)
 	{
-		RenderLightShaftBloom();
+		RenderLightShaftBloom(RHICmdList);
 	}
 
 	if (ViewFamily.EngineShowFlags.VisualizeDistanceFieldAO)
 	{
 		FSceneRenderTargetItem DummyOutput(NULL, NULL, NULL);
-		RenderDistanceFieldAOSurfaceCache(DummyOutput, true);
+		RenderDistanceFieldAOSurfaceCache(RHICmdList, DummyOutput, true);
 	}
 
 	// Resolve the scene color for post processing.
@@ -1049,14 +1049,11 @@ IMPLEMENT_SHADER_TYPE(,FDownsampleSceneDepthPS,TEXT("DownsampleDepthPixelShader"
 FGlobalBoundShaderState DownsampleDepthBoundShaderState;
 
 /** Updates the downsized depth buffer with the current full resolution depth buffer. */
-void FDeferredShadingSceneRenderer::UpdateDownsampledDepthSurface()
+void FDeferredShadingSceneRenderer::UpdateDownsampledDepthSurface(FRHICommandList& RHICmdList)
 {
 	if (GSceneRenderTargets.UseDownsizedOcclusionQueries() && IsFeatureLevelSupported(GRHIShaderPlatform, ERHIFeatureLevel::SM3))
 	{
-		//@todo-rco: RHIPacketList
-		FRHICommandList& RHICmdList = FRHICommandList::GetNullRef();
-
-		RHISetRenderTarget(NULL, GSceneRenderTargets.GetSmallDepthSurface());
+		SetRenderTarget(RHICmdList, NULL, GSceneRenderTargets.GetSmallDepthSurface());
 
 		SCOPED_DRAW_EVENT(DownsampleDepth, DEC_SCENE_ITEMS);
 
@@ -1069,11 +1066,11 @@ void FDeferredShadingSceneRenderer::UpdateDownsampledDepthSurface()
 
 			extern TGlobalResource<FFilterVertexDeclaration> GFilterVertexDeclaration;
 
-			SetGlobalBoundShaderState(DownsampleDepthBoundShaderState, GFilterVertexDeclaration.VertexDeclarationRHI, *ScreenVertexShader, *PixelShader);
+			SetGlobalBoundShaderState(RHICmdList, DownsampleDepthBoundShaderState, GFilterVertexDeclaration.VertexDeclarationRHI, *ScreenVertexShader, *PixelShader);
 
-			RHISetBlendState(TStaticBlendState<CW_NONE>::GetRHI());
-			RHISetRasterizerState(TStaticRasterizerState<FM_Solid,CM_None>::GetRHI());
-			RHISetDepthStencilState(TStaticDepthStencilState<true,CF_Always>::GetRHI());
+			RHICmdList.SetBlendState(TStaticBlendState<CW_NONE>::GetRHI());
+			RHICmdList.SetRasterizerState(TStaticRasterizerState<FM_Solid, CM_None>::GetRHI());
+			RHICmdList.SetDepthStencilState(TStaticDepthStencilState<true, CF_Always>::GetRHI());
 
 			PixelShader->SetParameters(RHICmdList, View);
 
@@ -1082,9 +1079,10 @@ void FDeferredShadingSceneRenderer::UpdateDownsampledDepthSurface()
 			const uint32 DownsampledSizeX = FMath::TruncToInt(View.ViewRect.Width() / GSceneRenderTargets.GetSmallColorDepthDownsampleFactor());
 			const uint32 DownsampledSizeY = FMath::TruncToInt(View.ViewRect.Height() / GSceneRenderTargets.GetSmallColorDepthDownsampleFactor());
 
-			RHISetViewport(DownsampledX, DownsampledY, 0.0f, DownsampledX + DownsampledSizeX, DownsampledY + DownsampledSizeY, 1.0f);
+			RHICmdList.SetViewport(DownsampledX, DownsampledY, 0.0f, DownsampledX + DownsampledSizeX, DownsampledY + DownsampledSizeY, 1.0f);
 
 			DrawRectangle(
+				RHICmdList,
 				0, 0,
 				DownsampledSizeX, DownsampledSizeY,
 				View.ViewRect.Min.X, View.ViewRect.Min.Y,
diff --git a/Engine/Source/Runtime/Renderer/Private/DeferredShadingRenderer.h b/Engine/Source/Runtime/Renderer/Private/DeferredShadingRenderer.h
--- a/Engine/Source/Runtime/Renderer/Private/DeferredShadingRenderer.h
+++ b/Engine/Source/Runtime/Renderer/Private/DeferredShadingRenderer.h
@@ -42,7 +42,7 @@ public:
 	void ClearView();
 
 	/** Clears gbuffer where Z is still at the maximum value (ie no geometry rendered) */
-    void ClearGBufferAtMaxZ();
+	void ClearGBufferAtMaxZ(FRHICommandList& RHICmdList);
 
 	/** Clears LPVs for all views */
 	void ClearLPVs(FRHICommandList& RHICmdList);
@@ -168,7 +168,7 @@ private:
 	void BeginOcclusionTests();
 
 	/** Renders the scene's fogging. */
-	bool RenderFog(FLightShaftsOutput LightShaftsOutput);
+	bool RenderFog(FRHICommandList& RHICmdList, FLightShaftsOutput LightShaftsOutput);
 
 	/** Renders the scene's atmosphere. */
 	void RenderAtmosphere(FLightShaftsOutput LightShaftsOutput);
@@ -177,10 +177,10 @@ private:
 	void RenderDeferredReflections(FRHICommandList& RHICmdList);
 
 	/** Render dynamic sky lighting from Movable sky lights. */
-	void RenderDynamicSkyLighting();
+	void RenderDynamicSkyLighting(FRHICommandList& RHICmdList);
 
 	/** Render Ambient Occlusion using mesh distance fields and the surface cache, which supports dynamic rigid meshes. */
-	bool RenderDistanceFieldAOSurfaceCache(FSceneRenderTargetItem& OutBentNormalAO, bool bApplyToSceneColor);
+	bool RenderDistanceFieldAOSurfaceCache(FRHICommandList& RHICmdList, FSceneRenderTargetItem& OutBentNormalAO, bool bApplyToSceneColor);
 
 	/** Whether tiled deferred is supported and can be used at all. */
 	bool CanUseTiledDeferred() const;
@@ -195,10 +195,10 @@ private:
 	void RenderLights(FRHICommandList& RHICmdList);
 
 	/** Renders an array of lights for the stationary light overlap viewmode. */
-	void RenderLightArrayForOverlapViewmode(const TSparseArray<FLightSceneInfoCompact>& LightArray);
+	void RenderLightArrayForOverlapViewmode(FRHICommandList& RHICmdList, const TSparseArray<FLightSceneInfoCompact>& LightArray);
 
 	/** Render stationary light overlap as complexity to scene color. */
-	void RenderStationaryLightOverlap();
+	void RenderStationaryLightOverlap(FRHICommandList& RHICmdList);
 	
 	/** 
 	 * Renders the scene's translucency.
@@ -206,9 +206,9 @@ private:
 	void RenderTranslucency();
 
 	/** Renders the scene's light shafts */
-	FLightShaftsOutput RenderLightShaftOcclusion();
+	FLightShaftsOutput RenderLightShaftOcclusion(FRHICommandList& RHICmdList);
 
-	void RenderLightShaftBloom();
+	void RenderLightShaftBloom(FRHICommandList& RHICmdList);
 
 	/** Makes a copy of scene color so that a material with a scene color node can read the up to date scene color. */
 	void CopySceneColor(const FViewInfo& View, const FPrimitiveSceneInfo* PrimitiveSceneInfo);
@@ -223,7 +223,7 @@ private:
 	bool RenderLightMapDensities();
 
 	/** Updates the downsized depth buffer with the current full resolution depth buffer. */
-	void UpdateDownsampledDepthSurface();
+	void UpdateDownsampledDepthSurface(FRHICommandList& RHICmdList);
 
 	/**
 	 * Finish rendering a view, writing the contents to ViewFamily.RenderTarget.
@@ -295,10 +295,10 @@ private:
 	  * @param LightIndex The light's index into FScene::Lights
 	  * @return true if anything got rendered
 	  */
-	void RenderLight(const FLightSceneInfo* LightSceneInfo, bool bRenderOverlap, bool bIssueDrawEvent);
+	void RenderLight(FRHICommandList& RHICmdList, const FLightSceneInfo* LightSceneInfo, bool bRenderOverlap, bool bIssueDrawEvent);
 
 	/** Renders an array of simple lights using standard deferred shading. */
-	void RenderSimpleLightsStandardDeferred(const FSimpleLightArray& SimpleLights);
+	void RenderSimpleLightsStandardDeferred(FRHICommandList& RHICmdList, const FSimpleLightArray& SimpleLights);
 
 	/** Clears the translucency lighting volumes before light accumulation. */
 	void ClearTranslucentVolumeLighting(FRHICommandList& RHICmdList);
@@ -310,25 +310,25 @@ private:
 	void CompositeIndirectTranslucentVolumeLighting(FRHICommandList& RHICmdList);
 
 	/** Clears the volume texture used to accumulate per object shadows for translucency. */
-	void ClearTranslucentVolumePerObjectShadowing();
+	void ClearTranslucentVolumePerObjectShadowing(FRHICommandList& RHICmdList);
 
 	/** Accumulates the per object shadow's contribution for translucency. */
-	void AccumulateTranslucentVolumeObjectShadowing(const FProjectedShadowInfo* InProjectedShadowInfo, bool bClearVolume);
+	void AccumulateTranslucentVolumeObjectShadowing(FRHICommandList& RHICmdList, const FProjectedShadowInfo* InProjectedShadowInfo, bool bClearVolume);
 
 	/** Accumulates direct lighting for the given light.  InProjectedShadowInfo can be NULL in which case the light will be unshadowed. */
-	void InjectTranslucentVolumeLighting(const FLightSceneInfo& LightSceneInfo, const FProjectedShadowInfo* InProjectedShadowInfo);
+	void InjectTranslucentVolumeLighting(FRHICommandList& RHICmdList, const FLightSceneInfo& LightSceneInfo, const FProjectedShadowInfo* InProjectedShadowInfo);
 
 	/** Accumulates direct lighting for an array of unshadowed lights. */
-	void InjectTranslucentVolumeLightingArray(const TArray<FSortedLightSceneInfo, SceneRenderingAllocator>& SortedLights, int32 NumLights);
+	void InjectTranslucentVolumeLightingArray(FRHICommandList& RHICmdList, const TArray<FSortedLightSceneInfo, SceneRenderingAllocator>& SortedLights, int32 NumLights);
 
 	/** Accumulates direct lighting for simple lights. */
-	void InjectSimpleTranslucentVolumeLightingArray(const FSimpleLightArray& SimpleLights);
+	void InjectSimpleTranslucentVolumeLightingArray(FRHICommandList& RHICmdList, const FSimpleLightArray& SimpleLights);
 
 	/** Filters the translucency lighting volumes to reduce aliasing. */
-	void FilterTranslucentVolumeLighting();
+	void FilterTranslucentVolumeLighting(FRHICommandList& RHICmdList);
 
 	/** Output SpecularColor * IndirectDiffuseGI for metals so they are not black in reflections */
-	void RenderReflectionCaptureSpecularBounceForAllViews();
+	void RenderReflectionCaptureSpecularBounceForAllViews(FRHICommandList& RHICmdList);
 
 	/** Render image based reflections (SSR, Env, SkyLight) without compute shaders */
 	void RenderImageBasedReflectionsSM4ForAllViews(FRHICommandList& RHICmdList, bool bReflectionEnv);
diff --git a/Engine/Source/Runtime/Renderer/Private/DistanceFieldSurfaceCacheLighting.cpp b/Engine/Source/Runtime/Renderer/Private/DistanceFieldSurfaceCacheLighting.cpp
--- a/Engine/Source/Runtime/Renderer/Private/DistanceFieldSurfaceCacheLighting.cpp
+++ b/Engine/Source/Runtime/Renderer/Private/DistanceFieldSurfaceCacheLighting.cpp
@@ -941,14 +941,11 @@ private:
 
 IMPLEMENT_SHADER_TYPE(,FComputeDistanceFieldNormalCS,TEXT("DistanceFieldSurfaceCacheLighting"),TEXT("ComputeDistanceFieldNormalCS"),SF_Compute);
 
-void ComputeDistanceFieldNormal(const TArray<FViewInfo>& Views, FSceneRenderTargetItem& DistanceFieldNormal, int32 NumObjects)
+void ComputeDistanceFieldNormal(FRHICommandList& RHICmdList, const TArray<FViewInfo>& Views, FSceneRenderTargetItem& DistanceFieldNormal, int32 NumObjects)
 {
-	//@todo-rco: RHIPacketList
-	FRHICommandList& RHICmdList = FRHICommandList::GetNullRef();
-
 	if (GAOComputeShaderNormalCalculation)
 	{
-		RHISetRenderTarget(NULL, NULL);
+		SetRenderTarget(RHICmdList, NULL, NULL);
 
 		for (int32 ViewIndex = 0; ViewIndex < Views.Num(); ViewIndex++)
 		{
@@ -961,7 +958,7 @@ void ComputeDistanceFieldNormal(const TArray<FViewInfo>& Views, FSceneRenderTarg
 				SCOPED_DRAW_EVENT(ComputeNormalCS, DEC_SCENE_ITEMS);
 				TShaderMapRef<FComputeDistanceFieldNormalCS> ComputeShader(GetGlobalShaderMap());
 
-				RHISetComputeShader(ComputeShader->GetComputeShader());	
+				RHICmdList.SetComputeShader(ComputeShader->GetComputeShader());
 				ComputeShader->SetParameters(RHICmdList, View, DistanceFieldNormal, NumObjects);
 				DispatchComputeShader(RHICmdList, *ComputeShader, GroupSizeX, GroupSizeY, 1);
 
@@ -971,7 +968,7 @@ void ComputeDistanceFieldNormal(const TArray<FViewInfo>& Views, FSceneRenderTarg
 	}
 	else
 	{
-		RHISetRenderTarget(DistanceFieldNormal.TargetableTexture, NULL);
+		SetRenderTarget(RHICmdList, DistanceFieldNormal.TargetableTexture, NULL);
 
 		for (int32 ViewIndex = 0; ViewIndex < Views.Num(); ViewIndex++)
 		{
@@ -979,20 +976,22 @@ void ComputeDistanceFieldNormal(const TArray<FViewInfo>& Views, FSceneRenderTarg
 
 			SCOPED_DRAW_EVENT(ComputeNormal, DEC_SCENE_ITEMS);
 
-			RHISetViewport(View.ViewRect.Min.X / GAODownsampleFactor, View.ViewRect.Min.Y / GAODownsampleFactor, 0.0f, View.ViewRect.Max.X / GAODownsampleFactor, View.ViewRect.Max.Y / GAODownsampleFactor, 1.0f);
-			RHISetRasterizerState(TStaticRasterizerState<FM_Solid,CM_None>::GetRHI());
-			RHISetDepthStencilState(TStaticDepthStencilState<false,CF_Always>::GetRHI());
-			RHISetBlendState(TStaticBlendState<>::GetRHI());
+			RHICmdList.SetViewport(View.ViewRect.Min.X / GAODownsampleFactor, View.ViewRect.Min.Y / GAODownsampleFactor, 0.0f, View.ViewRect.Max.X / GAODownsampleFactor, View.ViewRect.Max.Y / GAODownsampleFactor, 1.0f);
+			RHICmdList.SetRasterizerState(TStaticRasterizerState<FM_Solid, CM_None>::GetRHI());
+			RHICmdList.SetDepthStencilState(TStaticDepthStencilState<false, CF_Always>::GetRHI());
+			RHICmdList.SetBlendState(TStaticBlendState<>::GetRHI());
 
 			TShaderMapRef<FPostProcessVS> VertexShader( GetGlobalShaderMap() );
 			TShaderMapRef<FComputeDistanceFieldNormalPS> PixelShader( GetGlobalShaderMap() );
 
 			static FGlobalBoundShaderState BoundShaderState;
-			SetGlobalBoundShaderState(BoundShaderState, GFilterVertexDeclaration.VertexDeclarationRHI, *VertexShader, *PixelShader);
+			RHICmdList.CheckIsNull(); // need new approach for "static FGlobalBoundShaderState" for parallel rendering
+			SetGlobalBoundShaderState(RHICmdList, BoundShaderState, GFilterVertexDeclaration.VertexDeclarationRHI, *VertexShader, *PixelShader);
 
 			PixelShader->SetParameters(RHICmdList, View, NumObjects);
 
 			DrawRectangle( 
+				RHICmdList,
 				0, 0, 
 				View.ViewRect.Width() / GAODownsampleFactor, View.ViewRect.Height() / GAODownsampleFactor,
 				View.ViewRect.Min.X, View.ViewRect.Min.Y, 
@@ -2055,6 +2054,7 @@ void AllocateOrReuseAORenderTarget(TRefCountPtr<IPooledRenderTarget>& Target, co
 }
 
 void UpdateHistory(
+	FRHICommandList& RHICmdList,
 	FViewInfo& View, 
 	const TCHAR* HistoryRTName,
 	/** Contains last frame's history, if non-NULL.  This will be updated with the new frame's history. */
@@ -2064,35 +2064,35 @@ void UpdateHistory(
 	/** Output of Temporal Reprojection for the next step in the pipeline. */
 	TRefCountPtr<IPooledRenderTarget>& HistoryOutput)
 {
-	//@todo-rco: RHIPacketList
-	FRHICommandList& RHICmdList = FRHICommandList::GetNullRef();
-
 	if (HistoryState)
 	{
 		if (*HistoryState && !View.bCameraCut && !View.bPrevTransformsReset)
 		{
+			RHICmdList.CheckIsNull(); // direct allocation of resources won't work in parallel
 			// Reuse a render target from the pool with a consistent name, for vis purposes
 			TRefCountPtr<IPooledRenderTarget> NewHistory;
 			AllocateOrReuseAORenderTarget(NewHistory, HistoryRTName);
 
 			{
 				SCOPED_DRAW_EVENT(UpdateHistory, DEC_SCENE_ITEMS);
-				RHISetRenderTarget(NewHistory->GetRenderTargetItem().TargetableTexture, NULL);
+				SetRenderTarget(RHICmdList, NewHistory->GetRenderTargetItem().TargetableTexture, NULL);
 
-				RHISetViewport(View.ViewRect.Min.X / GAODownsampleFactor, View.ViewRect.Min.Y / GAODownsampleFactor, 0.0f, View.ViewRect.Max.X / GAODownsampleFactor, View.ViewRect.Max.Y / GAODownsampleFactor, 1.0f);
-				RHISetRasterizerState(TStaticRasterizerState<FM_Solid,CM_None>::GetRHI());
-				RHISetDepthStencilState(TStaticDepthStencilState<false,CF_Always>::GetRHI());
-				RHISetBlendState(TStaticBlendState<>::GetRHI());
+				RHICmdList.SetViewport(View.ViewRect.Min.X / GAODownsampleFactor, View.ViewRect.Min.Y / GAODownsampleFactor, 0.0f, View.ViewRect.Max.X / GAODownsampleFactor, View.ViewRect.Max.Y / GAODownsampleFactor, 1.0f);
+				RHICmdList.SetRasterizerState(TStaticRasterizerState<FM_Solid, CM_None>::GetRHI());
+				RHICmdList.SetDepthStencilState(TStaticDepthStencilState<false, CF_Always>::GetRHI());
+				RHICmdList.SetBlendState(TStaticBlendState<>::GetRHI());
 
 				TShaderMapRef<FPostProcessVS> VertexShader( GetGlobalShaderMap() );
 				TShaderMapRef<FUpdateHistoryPS> PixelShader( GetGlobalShaderMap() );
 
 				static FGlobalBoundShaderState BoundShaderState;
-				SetGlobalBoundShaderState(BoundShaderState, GFilterVertexDeclaration.VertexDeclarationRHI, *VertexShader, *PixelShader);
+				RHICmdList.CheckIsNull(); // need new approach for "static FGlobalBoundShaderState" for parallel rendering
+				SetGlobalBoundShaderState(RHICmdList, BoundShaderState, GFilterVertexDeclaration.VertexDeclarationRHI, *VertexShader, *PixelShader);
 
 				PixelShader->SetParameters(RHICmdList, View, (*HistoryState)->GetRenderTargetItem(), AOSource->GetRenderTargetItem());
 
 				DrawRectangle( 
+					RHICmdList,
 					0, 0, 
 					View.ViewRect.Width() / GAODownsampleFactor, View.ViewRect.Height() / GAODownsampleFactor,
 					View.ViewRect.Min.X / GAODownsampleFactor, View.ViewRect.Min.Y / GAODownsampleFactor, 
@@ -2121,11 +2121,9 @@ void UpdateHistory(
 	}
 }
 
-void PostProcessBentNormalAO(TArray<FViewInfo>& Views, FSceneRenderTargetItem& IrradianceCacheInterpolation, TRefCountPtr<IPooledRenderTarget>& AOOutput)
+void PostProcessBentNormalAO(FRHICommandList& RHICmdList, TArray<FViewInfo>& Views, FSceneRenderTargetItem& IrradianceCacheInterpolation, TRefCountPtr<IPooledRenderTarget>& AOOutput)
 {
-	//@todo-rco: RHIPacketList
-	FRHICommandList& RHICmdList = FRHICommandList::GetNullRef();
-
+	RHICmdList.CheckIsNull(); // direct allocation of resources won't work in parallel
 	TRefCountPtr<IPooledRenderTarget> DistanceFieldAOBentNormal;
 	AllocateOrReuseAORenderTarget(DistanceFieldAOBentNormal, TEXT("DistanceFieldBentNormalAO"));
 
@@ -2139,7 +2137,7 @@ void PostProcessBentNormalAO(TArray<FViewInfo>& Views, FSceneRenderTargetItem& I
 	{
 		SCOPED_DRAW_EVENT(AOCombine, DEC_SCENE_ITEMS);
 
-		RHISetRenderTarget(GAOFillGaps 
+		SetRenderTarget(RHICmdList, GAOFillGaps
 			? DistanceFieldAOBentNormal2->GetRenderTargetItem().TargetableTexture
 			: DistanceFieldAOBentNormal->GetRenderTargetItem().TargetableTexture, 
 			NULL);
@@ -2147,21 +2145,23 @@ void PostProcessBentNormalAO(TArray<FViewInfo>& Views, FSceneRenderTargetItem& I
 		for (int32 ViewIndex = 0; ViewIndex < Views.Num(); ViewIndex++)
 		{
 			FViewInfo& View = Views[ViewIndex];
-			RHISetViewport(View.ViewRect.Min.X / GAODownsampleFactor, View.ViewRect.Min.Y / GAODownsampleFactor, 0.0f, View.ViewRect.Max.X / GAODownsampleFactor, View.ViewRect.Max.Y / GAODownsampleFactor, 1.0f);
-			RHISetRasterizerState(TStaticRasterizerState<FM_Solid,CM_None>::GetRHI());
-			RHISetDepthStencilState(TStaticDepthStencilState<false,CF_Always>::GetRHI());
-			RHISetBlendState(TStaticBlendState<>::GetRHI());
+			RHICmdList.SetViewport(View.ViewRect.Min.X / GAODownsampleFactor, View.ViewRect.Min.Y / GAODownsampleFactor, 0.0f, View.ViewRect.Max.X / GAODownsampleFactor, View.ViewRect.Max.Y / GAODownsampleFactor, 1.0f);
+			RHICmdList.SetRasterizerState(TStaticRasterizerState<FM_Solid, CM_None>::GetRHI());
+			RHICmdList.SetDepthStencilState(TStaticDepthStencilState<false, CF_Always>::GetRHI());
+			RHICmdList.SetBlendState(TStaticBlendState<>::GetRHI());
 
 			TShaderMapRef<FPostProcessVS> VertexShader(GetGlobalShaderMap());
 			TShaderMapRef<FDistanceFieldAOCombinePS2> PixelShader(GetGlobalShaderMap());
 
 			static FGlobalBoundShaderState BoundShaderState;
-			SetGlobalBoundShaderState(BoundShaderState, GFilterVertexDeclaration.VertexDeclarationRHI, *VertexShader, *PixelShader);
+			RHICmdList.CheckIsNull(); // need new approach for "static FGlobalBoundShaderState" for parallel rendering
+			SetGlobalBoundShaderState(RHICmdList, BoundShaderState, GFilterVertexDeclaration.VertexDeclarationRHI, *VertexShader, *PixelShader);
 
 			PixelShader->SetParameters(RHICmdList, View, IrradianceCacheInterpolation);
 
 			//@todo - get the round up correct
 			DrawRectangle( 
+				RHICmdList,
 				0, 0, 
 				View.ViewRect.Width() / GAODownsampleFactor, View.ViewRect.Height() / GAODownsampleFactor,
 				View.ViewRect.Min.X / GAODownsampleFactor, View.ViewRect.Min.Y / GAODownsampleFactor, 
@@ -2175,27 +2175,29 @@ void PostProcessBentNormalAO(TArray<FViewInfo>& Views, FSceneRenderTargetItem& I
 	if (GAOFillGaps)
 	{
 		SCOPED_DRAW_EVENT(FillGaps, DEC_SCENE_ITEMS);
-		RHISetRenderTarget(DistanceFieldAOBentNormal->GetRenderTargetItem().TargetableTexture, NULL);
+		SetRenderTarget(RHICmdList, DistanceFieldAOBentNormal->GetRenderTargetItem().TargetableTexture, NULL);
 
 		for (int32 ViewIndex = 0; ViewIndex < Views.Num(); ViewIndex++)
 		{
 			const FViewInfo& View = Views[ViewIndex];
 
-			RHISetViewport(View.ViewRect.Min.X / GAODownsampleFactor, View.ViewRect.Min.Y / GAODownsampleFactor, 0.0f, View.ViewRect.Max.X / GAODownsampleFactor, View.ViewRect.Max.Y / GAODownsampleFactor, 1.0f);
-			RHISetRasterizerState(TStaticRasterizerState<FM_Solid,CM_None>::GetRHI());
-			RHISetDepthStencilState(TStaticDepthStencilState<false,CF_Always>::GetRHI());
-			RHISetBlendState(TStaticBlendState<>::GetRHI());
+			RHICmdList.SetViewport(View.ViewRect.Min.X / GAODownsampleFactor, View.ViewRect.Min.Y / GAODownsampleFactor, 0.0f, View.ViewRect.Max.X / GAODownsampleFactor, View.ViewRect.Max.Y / GAODownsampleFactor, 1.0f);
+			RHICmdList.SetRasterizerState(TStaticRasterizerState<FM_Solid, CM_None>::GetRHI());
+			RHICmdList.SetDepthStencilState(TStaticDepthStencilState<false, CF_Always>::GetRHI());
+			RHICmdList.SetBlendState(TStaticBlendState<>::GetRHI());
 
 			TShaderMapRef<FPostProcessVS> VertexShader( GetGlobalShaderMap() );
 			TShaderMapRef<FFillGapsPS> PixelShader( GetGlobalShaderMap() );
 
 			static FGlobalBoundShaderState BoundShaderState;
-			SetGlobalBoundShaderState(BoundShaderState, GFilterVertexDeclaration.VertexDeclarationRHI, *VertexShader, *PixelShader);
+			RHICmdList.CheckIsNull(); // need new approach for "static FGlobalBoundShaderState" for parallel rendering
+			SetGlobalBoundShaderState(RHICmdList, BoundShaderState, GFilterVertexDeclaration.VertexDeclarationRHI, *VertexShader, *PixelShader);
 
 			VertexShader->SetParameters(RHICmdList, View);
 			PixelShader->SetParameters(RHICmdList, View, DistanceFieldAOBentNormal2->GetRenderTargetItem());
 
 			DrawRectangle( 
+				RHICmdList,
 				0, 0, 
 				View.ViewRect.Width() / GAODownsampleFactor, View.ViewRect.Height() / GAODownsampleFactor,
 				View.ViewRect.Min.X / GAODownsampleFactor, View.ViewRect.Min.Y / GAODownsampleFactor, 
@@ -2215,6 +2217,7 @@ void PostProcessBentNormalAO(TArray<FViewInfo>& Views, FSceneRenderTargetItem& I
 		FViewInfo& View = Views[0];
 
 		UpdateHistory(
+			RHICmdList,
 			View, 
 			TEXT("DistanceFieldAOHistory"),
 			HistoryState,
@@ -2278,31 +2281,30 @@ private:
 
 IMPLEMENT_SHADER_TYPE(,FDistanceFieldAOUpsamplePS,TEXT("DistanceFieldSurfaceCacheLighting"),TEXT("AOUpsamplePS"),SF_Pixel);
 
-void UpsampleBentNormalAO(const TArray<FViewInfo>& Views, TRefCountPtr<IPooledRenderTarget>& DistanceFieldAOBentNormal)
+void UpsampleBentNormalAO(FRHICommandList& RHICmdList, const TArray<FViewInfo>& Views, TRefCountPtr<IPooledRenderTarget>& DistanceFieldAOBentNormal)
 {
-	//@todo-rco: RHIPacketList
-	FRHICommandList& RHICmdList = FRHICommandList::GetNullRef();
-
 	for (int32 ViewIndex = 0; ViewIndex < Views.Num(); ViewIndex++)
 	{
 		const FViewInfo& View = Views[ViewIndex];
 
 		SCOPED_DRAW_EVENT(UpsampleAO, DEC_SCENE_ITEMS);
 
-		RHISetViewport(View.ViewRect.Min.X, View.ViewRect.Min.Y, 0.0f, View.ViewRect.Max.X, View.ViewRect.Max.Y, 1.0f);
-		RHISetRasterizerState(TStaticRasterizerState<FM_Solid,CM_None>::GetRHI());
-		RHISetDepthStencilState(TStaticDepthStencilState<false,CF_Always>::GetRHI());
-		RHISetBlendState(TStaticBlendState<>::GetRHI());
+		RHICmdList.SetViewport(View.ViewRect.Min.X, View.ViewRect.Min.Y, 0.0f, View.ViewRect.Max.X, View.ViewRect.Max.Y, 1.0f);
+		RHICmdList.SetRasterizerState(TStaticRasterizerState<FM_Solid, CM_None>::GetRHI());
+		RHICmdList.SetDepthStencilState(TStaticDepthStencilState<false, CF_Always>::GetRHI());
+		RHICmdList.SetBlendState(TStaticBlendState<>::GetRHI());
 
 		TShaderMapRef<FPostProcessVS> VertexShader( GetGlobalShaderMap() );
 		TShaderMapRef<FDistanceFieldAOUpsamplePS> PixelShader( GetGlobalShaderMap() );
 
 		static FGlobalBoundShaderState BoundShaderState;
-		SetGlobalBoundShaderState(BoundShaderState, GFilterVertexDeclaration.VertexDeclarationRHI, *VertexShader, *PixelShader);
+		RHICmdList.CheckIsNull(); // need new approach for "static FGlobalBoundShaderState" for parallel rendering
+		SetGlobalBoundShaderState(RHICmdList, BoundShaderState, GFilterVertexDeclaration.VertexDeclarationRHI, *VertexShader, *PixelShader);
 
 		PixelShader->SetParameters(RHICmdList, View, DistanceFieldAOBentNormal);
 
 		DrawRectangle( 
+			RHICmdList,
 			0, 0, 
 			View.ViewRect.Width(), View.ViewRect.Height(),
 			View.ViewRect.Min.X / GAODownsampleFactor, View.ViewRect.Min.Y / GAODownsampleFactor, 
@@ -2470,18 +2472,17 @@ void GenerateBestSpacedVectors()
 	}
 }
 
-FIntPoint BuildTileObjectLists(FScene* Scene, TArray<FViewInfo>& Views, int32 NumObjects)
+FIntPoint BuildTileObjectLists(FRHICommandList& RHICmdList, FScene* Scene, TArray<FViewInfo>& Views, int32 NumObjects)
 {
-	//@todo-rco: RHIPacketList
-	FRHICommandList& RHICmdList = FRHICommandList::GetNullRef();
-
 	SCOPED_DRAW_EVENT(BuildTileList, DEC_SCENE_ITEMS);
-	RHISetRenderTarget(NULL, NULL);
+	SetRenderTarget(RHICmdList, NULL, NULL);
 
 	FIntPoint TileListGroupSize;
 
 	if (GAOScatterTileCulling)
 	{
+		RHICmdList.CheckIsNull(); // direct allocation of resources won't work in parallel
+
 		for (int32 ViewIndex = 0; ViewIndex < Views.Num(); ViewIndex++)
 		{
 			const FViewInfo& View = Views[ViewIndex];
@@ -2512,7 +2513,7 @@ FIntPoint BuildTileObjectLists(FScene* Scene, TArray<FViewInfo>& Views, int32 Nu
 				SCOPED_DRAW_EVENT(BuildTileCones, DEC_SCENE_ITEMS);
 				TShaderMapRef<FBuildTileConesCS> ComputeShader(GetGlobalShaderMap());
 
-				RHISetComputeShader(ComputeShader->GetComputeShader());	
+				RHICmdList.SetComputeShader(ComputeShader->GetComputeShader());
 				ComputeShader->SetParameters(RHICmdList, View, Scene, FVector2D(GroupSizeX, GroupSizeY));
 				DispatchComputeShader(RHICmdList, *ComputeShader, GroupSizeX, GroupSizeY, 1);
 
@@ -2527,24 +2528,25 @@ FIntPoint BuildTileObjectLists(FScene* Scene, TArray<FViewInfo>& Views, int32 Nu
 
 				TArray<FUnorderedAccessViewRHIParamRef> UAVs;
 				PixelShader->GetUAVs(Views[0], UAVs);
-				RHISetRenderTargets(0, (const FRHIRenderTargetView *)NULL, NULL, UAVs.Num(), UAVs.GetData());
+				RHICmdList.SetRenderTargets(0, (const FRHIRenderTargetView *)NULL, NULL, UAVs.Num(), UAVs.GetData());
 
-				RHISetViewport(0, 0, 0.0f, GroupSizeX, GroupSizeY, 1.0f);
+				RHICmdList.SetViewport(0, 0, 0.0f, GroupSizeX, GroupSizeY, 1.0f);
 
 				// Render backfaces since camera may intersect
-				RHISetRasterizerState(View.bReverseCulling ? TStaticRasterizerState<FM_Solid,CM_CW>::GetRHI() : TStaticRasterizerState<FM_Solid,CM_CCW>::GetRHI());
-				RHISetDepthStencilState(TStaticDepthStencilState<false,CF_Always>::GetRHI());
-				RHISetBlendState(TStaticBlendState<>::GetRHI());
+				RHICmdList.SetRasterizerState(View.bReverseCulling ? TStaticRasterizerState<FM_Solid, CM_CW>::GetRHI() : TStaticRasterizerState<FM_Solid, CM_CCW>::GetRHI());
+				RHICmdList.SetDepthStencilState(TStaticDepthStencilState<false, CF_Always>::GetRHI());
+				RHICmdList.SetBlendState(TStaticBlendState<>::GetRHI());
 
 				static FGlobalBoundShaderState BoundShaderState;
-				SetGlobalBoundShaderState(BoundShaderState, GetVertexDeclarationFVector4(), *VertexShader, *PixelShader);
+				RHICmdList.CheckIsNull(); // need new approach for "static FGlobalBoundShaderState" for parallel rendering
+				SetGlobalBoundShaderState(RHICmdList, BoundShaderState, GetVertexDeclarationFVector4(), *VertexShader, *PixelShader);
 
 				VertexShader->SetParameters(RHICmdList, View, NumObjects);
 				PixelShader->SetParameters(RHICmdList, View, FVector2D(GroupSizeX, GroupSizeY));
 
-				RHISetStreamSource(0, StencilingGeometry::GLowPolyStencilSphereVertexBuffer.VertexBufferRHI, sizeof(FVector4), 0);
+				RHICmdList.SetStreamSource(0, StencilingGeometry::GLowPolyStencilSphereVertexBuffer.VertexBufferRHI, sizeof(FVector4), 0);
 
-				RHIDrawIndexedPrimitive(
+				RHICmdList.DrawIndexedPrimitive(
 					StencilingGeometry::GLowPolyStencilSphereIndexBuffer.IndexBufferRHI, 
 					PT_TriangleList, 
 					0,
@@ -2586,12 +2588,12 @@ FIntPoint BuildTileObjectLists(FScene* Scene, TArray<FViewInfo>& Views, int32 Nu
 
 			// Indicates the clear value for each channel of the UAV format
 			uint32 ClearValues[4] = {0};
-			RHIClearUAV(TileIntersectionResources->TileArrayNextAllocation.UAV, ClearValues);
-			RHIClearUAV(TileIntersectionResources->TileHeadData.UAV, ClearValues);
+			RHICmdList.ClearUAV(TileIntersectionResources->TileArrayNextAllocation.UAV, ClearValues);
+			RHICmdList.ClearUAV(TileIntersectionResources->TileHeadData.UAV, ClearValues);
 
 			TShaderMapRef<FDistanceFieldBuildTileListCS > ComputeShader(GetGlobalShaderMap());
 
-			RHISetComputeShader(ComputeShader->GetComputeShader());	
+			RHICmdList.SetComputeShader(ComputeShader->GetComputeShader());
 			ComputeShader->SetParameters(RHICmdList, View, Scene, FVector2D(GroupSizeX, GroupSizeY), NumObjects);
 			DispatchComputeShader(RHICmdList, *ComputeShader, GroupSizeX, GroupSizeY, 1);
 
@@ -2602,19 +2604,16 @@ FIntPoint BuildTileObjectLists(FScene* Scene, TArray<FViewInfo>& Views, int32 Nu
 	return TileListGroupSize;
 }
 
-void RenderIrradianceCacheInterpolation(TArray<FViewInfo>& Views, IPooledRenderTarget* InterpolationTarget, FSceneRenderTargetItem& DistanceFieldNormal, int32 DepthLevel, int32 DestLevelDownsampleFactor, bool bFinalInterpolation)
+void RenderIrradianceCacheInterpolation(FRHICommandList& RHICmdList, TArray<FViewInfo>& Views, IPooledRenderTarget* InterpolationTarget, FSceneRenderTargetItem& DistanceFieldNormal, int32 DepthLevel, int32 DestLevelDownsampleFactor, bool bFinalInterpolation)
 {
-	//@todo-rco: RHIPacketList
-	FRHICommandList& RHICmdList = FRHICommandList::GetNullRef();
-
 	check(!(bFinalInterpolation && DepthLevel != 0));
 
 	{
 		SCOPED_DRAW_EVENT(IrradianceCacheSplat, DEC_SCENE_ITEMS);
 		const int32 MaxAllowedLevel = GAOMaxSupportedLevel;
 
-		RHISetRenderTarget(InterpolationTarget->GetRenderTargetItem().TargetableTexture, NULL);
-		RHIClear(true, FLinearColor(0, 0, 0, 0), false, 0, false, 0, FIntRect());
+		SetRenderTarget(RHICmdList, InterpolationTarget->GetRenderTargetItem().TargetableTexture, NULL);
+		RHICmdList.Clear(true, FLinearColor(0, 0, 0, 0), false, 0, false, 0, FIntRect());
 
 		for (int32 ViewIndex = 0; ViewIndex < Views.Num(); ViewIndex++)
 		{
@@ -2623,13 +2622,12 @@ void RenderIrradianceCacheInterpolation(TArray<FViewInfo>& Views, IPooledRenderT
 			const FScene* Scene = (const FScene*)View.Family->Scene;
 			FSurfaceCacheResources& SurfaceCacheResources = *Scene->SurfaceCacheResources;
 
-			RHISetViewport(View.ViewRect.Min.X / DestLevelDownsampleFactor, View.ViewRect.Min.Y / DestLevelDownsampleFactor, 0.0f, FMath::DivideAndRoundUp(View.ViewRect.Max.X, DestLevelDownsampleFactor), FMath::DivideAndRoundUp(View.ViewRect.Max.Y, DestLevelDownsampleFactor), 1.0f);
-			//RHISetViewport(View.ViewRect.Min.X / DestLevelDownsampleFactor, View.ViewRect.Min.Y / DestLevelDownsampleFactor, 0.0f, View.ViewRect.Max.X / DestLevelDownsampleFactor, View.ViewRect.Max.Y / DestLevelDownsampleFactor, 1.0f);
+			RHICmdList.SetViewport(View.ViewRect.Min.X / DestLevelDownsampleFactor, View.ViewRect.Min.Y / DestLevelDownsampleFactor, 0.0f, FMath::DivideAndRoundUp(View.ViewRect.Max.X, DestLevelDownsampleFactor), FMath::DivideAndRoundUp(View.ViewRect.Max.Y, DestLevelDownsampleFactor), 1.0f);
 
-			RHISetRasterizerState(TStaticRasterizerState<FM_Solid,CM_None>::GetRHI());
+			RHICmdList.SetRasterizerState(TStaticRasterizerState<FM_Solid, CM_None>::GetRHI());
 			//@todo - render front faces with depth testing, requires depth buffer
-			RHISetDepthStencilState(TStaticDepthStencilState<false,CF_Always>::GetRHI());
-			RHISetBlendState(TStaticBlendState<CW_RGBA, BO_Add, BF_One, BF_One, BO_Add, BF_One, BF_One>::GetRHI());
+			RHICmdList.SetDepthStencilState(TStaticDepthStencilState<false, CF_Always>::GetRHI());
+			RHICmdList.SetBlendState(TStaticBlendState<CW_RGBA, BO_Add, BF_One, BF_One, BO_Add, BF_One, BF_One>::GetRHI());
 
 			//FVector2D NormalizedOffsetToPixelCenter(-(-DestLevelDownsampleFactor / 2 + .5f) / (float)View.ViewRect.Width(), -(-DestLevelDownsampleFactor / 2 + .5f) / (float)View.ViewRect.Height());
 			uint32 DownsampledViewSizeX = FMath::DivideAndRoundUp(View.ViewRect.Width(), DestLevelDownsampleFactor);
@@ -2655,7 +2653,7 @@ void RenderIrradianceCacheInterpolation(TArray<FViewInfo>& Views, IPooledRenderT
 			FVector2D OffsetToTopResPixel = (FVector2D(.5f, .5f) / BufferSize - ScreenPositionBias) / ScreenPositionScale;
 			FVector2D NormalizedOffsetToPixelCenter = OffsetToLowResCorner - OffsetToTopResPixel;
 
-			RHISetStreamSource(0, GCircleVertexBuffer.VertexBufferRHI, sizeof(FScreenVertex), 0);
+			RHICmdList.SetStreamSource(0, GCircleVertexBuffer.VertexBufferRHI, sizeof(FScreenVertex), 0);
 
 			if (bFinalInterpolation)
 			{
@@ -2665,12 +2663,13 @@ void RenderIrradianceCacheInterpolation(TArray<FViewInfo>& Views, IPooledRenderT
 				for (int32 SplatSourceDepthLevel = MaxAllowedLevel; SplatSourceDepthLevel >= DepthLevel; SplatSourceDepthLevel--)
 				{
 					static FGlobalBoundShaderState BoundShaderState;
-					SetGlobalBoundShaderState(BoundShaderState, GScreenVertexDeclaration.VertexDeclarationRHI, *VertexShader, *PixelShader);
+					RHICmdList.CheckIsNull(); // need new approach for "static FGlobalBoundShaderState" for parallel rendering
+					SetGlobalBoundShaderState(RHICmdList, BoundShaderState, GScreenVertexDeclaration.VertexDeclarationRHI, *VertexShader, *PixelShader);
 
 					VertexShader->SetParameters(RHICmdList, View, SplatSourceDepthLevel, DestLevelDownsampleFactor, NormalizedOffsetToPixelCenter);
 					PixelShader->SetParameters(RHICmdList, View, DistanceFieldNormal, DestLevelDownsampleFactor);
 
-					RHIDrawPrimitiveIndirect(PT_TriangleList, SurfaceCacheResources.Level[SplatSourceDepthLevel]->ScatterDrawParameters.Buffer, 0);
+					RHICmdList.DrawPrimitiveIndirect(PT_TriangleList, SurfaceCacheResources.Level[SplatSourceDepthLevel]->ScatterDrawParameters.Buffer, 0);
 				}
 			}
 			else
@@ -2681,12 +2680,13 @@ void RenderIrradianceCacheInterpolation(TArray<FViewInfo>& Views, IPooledRenderT
 				for (int32 SplatSourceDepthLevel = MaxAllowedLevel; SplatSourceDepthLevel >= DepthLevel; SplatSourceDepthLevel--)
 				{
 					static FGlobalBoundShaderState BoundShaderState;
-					SetGlobalBoundShaderState(BoundShaderState, GScreenVertexDeclaration.VertexDeclarationRHI, *VertexShader, *PixelShader);
+					RHICmdList.CheckIsNull(); // need new approach for "static FGlobalBoundShaderState" for parallel rendering
+					SetGlobalBoundShaderState(RHICmdList, BoundShaderState, GScreenVertexDeclaration.VertexDeclarationRHI, *VertexShader, *PixelShader);
 
 					VertexShader->SetParameters(RHICmdList, View, SplatSourceDepthLevel, DestLevelDownsampleFactor, NormalizedOffsetToPixelCenter);
 					PixelShader->SetParameters(RHICmdList, View, DistanceFieldNormal, DestLevelDownsampleFactor);
 
-					RHIDrawPrimitiveIndirect(PT_TriangleList, SurfaceCacheResources.Level[SplatSourceDepthLevel]->ScatterDrawParameters.Buffer, 0);
+					RHICmdList.DrawPrimitiveIndirect(PT_TriangleList, SurfaceCacheResources.Level[SplatSourceDepthLevel]->ScatterDrawParameters.Buffer, 0);
 				}
 			}
 		}
@@ -2817,11 +2817,8 @@ void UpdateVisibleObjectBuffers(const FScene* Scene, const FViewInfo& View, int3
 	}
 }
 
-bool FDeferredShadingSceneRenderer::RenderDistanceFieldAOSurfaceCache(FSceneRenderTargetItem& OutBentNormalAO, bool bApplyToSceneColor)
+bool FDeferredShadingSceneRenderer::RenderDistanceFieldAOSurfaceCache(FRHICommandList& RHICmdList, FSceneRenderTargetItem& OutBentNormalAO, bool bApplyToSceneColor)
 {
-	//@todo-rco: RHIPacketList
-	FRHICommandList& RHICmdList = FRHICommandList::GetNullRef();
-
 	//@todo - support multiple views
 	const FViewInfo& View = Views[0];
 
@@ -2835,6 +2832,8 @@ bool FDeferredShadingSceneRenderer::RenderDistanceFieldAOSurfaceCache(FSceneRend
 		QUICK_SCOPE_CYCLE_COUNTER(STAT_RenderDistanceFieldAOSurfaceCache);
 		SCOPED_DRAW_EVENT(DistanceFieldAO, DEC_SCENE_ITEMS);
 
+		RHICmdList.CheckIsNull(); // direct allocation of things will not work, probably
+
 		// Update the global distance field atlas
 		GDistanceFieldVolumeTextureAtlas.UpdateAllocations();
 
@@ -2858,7 +2857,7 @@ bool FDeferredShadingSceneRenderer::RenderDistanceFieldAOSurfaceCache(FSceneRend
 				// Reset the number of active cache records to 0
 				for (int32 DepthLevel = MaxAllowedLevel; DepthLevel >= 0; DepthLevel--)
 				{
-					RHIClearUAV(SurfaceCacheResources.Level[DepthLevel]->ScatterDrawParameters.UAV, ClearValues);
+					RHICmdList.ClearUAV(SurfaceCacheResources.Level[DepthLevel]->ScatterDrawParameters.UAV, ClearValues);
 				}
 
 				SurfaceCacheResources.bClearedResources = true;
@@ -2874,7 +2873,7 @@ bool FDeferredShadingSceneRenderer::RenderDistanceFieldAOSurfaceCache(FSceneRend
 				QUICK_SCOPE_CYCLE_COUNTER(STAT_AOIssueGPUWork);
 
 				// Intersect objects with screen tiles, build lists
-				FIntPoint TileListGroupSize = BuildTileObjectLists(Scene, Views, NumObjects);
+				FIntPoint TileListGroupSize = BuildTileObjectLists(RHICmdList, Scene, Views, NumObjects);
 
 				TRefCountPtr<IPooledRenderTarget> DistanceFieldNormal;
 
@@ -2885,9 +2884,9 @@ bool FDeferredShadingSceneRenderer::RenderDistanceFieldAOSurfaceCache(FSceneRend
 				}
 
 				// Compute the distance field normal, this is used for surface caching instead of the GBuffer normal because it matches the occluding geometry
-				ComputeDistanceFieldNormal(Views, DistanceFieldNormal->GetRenderTargetItem(), NumObjects);
+				ComputeDistanceFieldNormal(RHICmdList, Views, DistanceFieldNormal->GetRenderTargetItem(), NumObjects);
 
-				GRenderTargetPool.VisualizeTexture.SetCheckPoint(DistanceFieldNormal);
+				GRenderTargetPool.VisualizeTexture.SetCheckPoint(RHICmdList, DistanceFieldNormal);
 
 				if (GAOReuseAcrossFrames && GAOTrimOldRecordsFraction > 0)
 				{
@@ -2899,7 +2898,7 @@ bool FDeferredShadingSceneRenderer::RenderDistanceFieldAOSurfaceCache(FSceneRend
 						{	
 							TShaderMapRef<FSetupCopyIndirectArgumentsCS> ComputeShader(GetGlobalShaderMap());
 
-							RHISetComputeShader(ComputeShader->GetComputeShader());	
+							RHICmdList.SetComputeShader(ComputeShader->GetComputeShader());
 							ComputeShader->SetParameters(RHICmdList, View, DepthLevel);
 							DispatchComputeShader(RHICmdList, *ComputeShader, 1, 1, 1);
 
@@ -2909,7 +2908,7 @@ bool FDeferredShadingSceneRenderer::RenderDistanceFieldAOSurfaceCache(FSceneRend
 						{
 							TShaderMapRef<FCopyIrradianceCacheSamplesCS> ComputeShader(GetGlobalShaderMap());
 
-							RHISetComputeShader(ComputeShader->GetComputeShader());	
+							RHICmdList.SetComputeShader(ComputeShader->GetComputeShader());
 							ComputeShader->SetParameters(RHICmdList, View, DepthLevel);
 							DispatchIndirectComputeShader(RHICmdList, *ComputeShader, SurfaceCacheResources.DispatchParameters.Buffer, 0);
 
@@ -2938,17 +2937,17 @@ bool FDeferredShadingSceneRenderer::RenderDistanceFieldAOSurfaceCache(FSceneRend
 					}
 
 					// Splat / interpolate the surface cache records onto the buffer sized for the current depth level
-					RenderIrradianceCacheInterpolation(Views, DistanceFieldAOIrradianceCacheSplat, DistanceFieldNormal->GetRenderTargetItem(), DepthLevel, DestLevelDownsampleFactor, false);
+					RenderIrradianceCacheInterpolation(RHICmdList, Views, DistanceFieldAOIrradianceCacheSplat, DistanceFieldNormal->GetRenderTargetItem(), DepthLevel, DestLevelDownsampleFactor, false);
 
 					{
 						SCOPED_DRAW_EVENT(PopulateIrradianceCache, DEC_SCENE_ITEMS);
-						RHISetRenderTarget(NULL, NULL);
+						SetRenderTarget(RHICmdList, NULL, NULL);
 
 						// Save off the current record count before adding any more
 						{	
 							TShaderMapRef<FSaveStartIndexCS> ComputeShader(GetGlobalShaderMap());
 
-							RHISetComputeShader(ComputeShader->GetComputeShader());	
+							RHICmdList.SetComputeShader(ComputeShader->GetComputeShader());
 							ComputeShader->SetParameters(RHICmdList, View, DepthLevel);
 							DispatchComputeShader(RHICmdList, *ComputeShader, 1, 1, 1);
 
@@ -2966,7 +2965,7 @@ bool FDeferredShadingSceneRenderer::RenderDistanceFieldAOSurfaceCache(FSceneRend
 
 							TShaderMapRef<FPopulateCacheCS> ComputeShader(GetGlobalShaderMap());
 
-							RHISetComputeShader(ComputeShader->GetComputeShader());	
+							RHICmdList.SetComputeShader(ComputeShader->GetComputeShader());
 							ComputeShader->SetParameters(RHICmdList, View, DistanceFieldAOIrradianceCacheSplat->GetRenderTargetItem(), DistanceFieldNormal->GetRenderTargetItem(), DestLevelDownsampleFactor, DepthLevel, TileListGroupSize);
 							DispatchComputeShader(RHICmdList, *ComputeShader, GroupSizeX, GroupSizeY, 1);
 
@@ -2980,7 +2979,7 @@ bool FDeferredShadingSceneRenderer::RenderDistanceFieldAOSurfaceCache(FSceneRend
 						{	
 							TShaderMapRef<FSetupFinalGatherIndirectArgumentsCS> ComputeShader(GetGlobalShaderMap());
 
-							RHISetComputeShader(ComputeShader->GetComputeShader());	
+							RHICmdList.SetComputeShader(ComputeShader->GetComputeShader());
 							ComputeShader->SetParameters(RHICmdList, View, DepthLevel);
 							DispatchComputeShader(RHICmdList, *ComputeShader, 1, 1, 1);
 
@@ -2991,7 +2990,7 @@ bool FDeferredShadingSceneRenderer::RenderDistanceFieldAOSurfaceCache(FSceneRend
 						{
 							TShaderMapRef<FFinalGatherCS> ComputeShader(GetGlobalShaderMap());
 
-							RHISetComputeShader(ComputeShader->GetComputeShader());	
+							RHICmdList.SetComputeShader(ComputeShader->GetComputeShader());
 							ComputeShader->SetParameters(RHICmdList, View, DistanceFieldAOIrradianceCacheSplat->GetRenderTargetItem(), DistanceFieldNormal->GetRenderTargetItem(), NumObjects, DestLevelDownsampleFactor, DepthLevel, TileListGroupSize);
 							DispatchIndirectComputeShader(RHICmdList, *ComputeShader, SurfaceCacheResources.DispatchParameters.Buffer, 0);
 
@@ -3011,12 +3010,12 @@ bool FDeferredShadingSceneRenderer::RenderDistanceFieldAOSurfaceCache(FSceneRend
 				// Splat the surface cache records onto the opaque pixels, using less strict weighting so the lighting is smoothed in world space
 				{
 					SCOPED_DRAW_EVENT(FinalIrradianceCacheSplat, DEC_SCENE_ITEMS);
-					RenderIrradianceCacheInterpolation(Views, IrradianceCacheAccumulation, DistanceFieldNormal->GetRenderTargetItem(), 0, GAODownsampleFactor, true);
+					RenderIrradianceCacheInterpolation(RHICmdList, Views, IrradianceCacheAccumulation, DistanceFieldNormal->GetRenderTargetItem(), 0, GAODownsampleFactor, true);
 				}
 
 				// Post process the AO to cover over artifacts
 				TRefCountPtr<IPooledRenderTarget> AOOutput;
-				PostProcessBentNormalAO(Views, IrradianceCacheAccumulation->GetRenderTargetItem(), AOOutput);
+				PostProcessBentNormalAO(RHICmdList, Views, IrradianceCacheAccumulation->GetRenderTargetItem(), AOOutput);
 
 				if (bApplyToSceneColor)
 				{
@@ -3028,7 +3027,7 @@ bool FDeferredShadingSceneRenderer::RenderDistanceFieldAOSurfaceCache(FSceneRend
 				}
 
 				// Upsample to full resolution, write to output
-				UpsampleBentNormalAO(Views, AOOutput);
+				UpsampleBentNormalAO(RHICmdList, Views, AOOutput);
 
 				return true;
 			}
@@ -3102,11 +3101,8 @@ private:
 IMPLEMENT_SHADER_TYPE(template<>,TDynamicSkyLightDiffusePS<true>,TEXT("SkyLighting"),TEXT("SkyLightDiffusePS"),SF_Pixel);
 IMPLEMENT_SHADER_TYPE(template<>,TDynamicSkyLightDiffusePS<false>,TEXT("SkyLighting"),TEXT("SkyLightDiffusePS"),SF_Pixel);
 
-void FDeferredShadingSceneRenderer::RenderDynamicSkyLighting()
+void FDeferredShadingSceneRenderer::RenderDynamicSkyLighting(FRHICommandList& RHICmdList)
 {
-	//@todo-rco: RHIPacketList
-	FRHICommandList& RHICmdList = FRHICommandList::GetNullRef();
-
 	if (Scene->SkyLight
 		&& Scene->SkyLight->ProcessedTexture
 		&& !Scene->SkyLight->bWantsStaticShadowing
@@ -3130,19 +3126,20 @@ void FDeferredShadingSceneRenderer::RenderDynamicSkyLighting()
 			&& ViewFamily.EngineShowFlags.DistanceFieldAO 
 			&& !ViewFamily.EngineShowFlags.VisualizeDistanceFieldAO)
 		{
-			bApplyShadowing = RenderDistanceFieldAOSurfaceCache(LightAccumulation->GetRenderTargetItem(), false);
+			bApplyShadowing = RenderDistanceFieldAOSurfaceCache(RHICmdList, LightAccumulation->GetRenderTargetItem(), false);
 		}
 
+		RHICmdList.CheckIsNull(); // BRSC is not supported in parallel
 		GSceneRenderTargets.BeginRenderingSceneColor();
 
 		for( int32 ViewIndex = 0; ViewIndex < Views.Num(); ViewIndex++ )
 		{
 			const FViewInfo& View = Views[ViewIndex];
 
-			RHISetViewport( View.ViewRect.Min.X, View.ViewRect.Min.Y, 0.0f, View.ViewRect.Max.X, View.ViewRect.Max.Y, 1.0f );
-			RHISetRasterizerState( TStaticRasterizerState< FM_Solid, CM_None >::GetRHI() );
-			RHISetDepthStencilState( TStaticDepthStencilState< false, CF_Always >::GetRHI() );
-			RHISetBlendState( TStaticBlendState< CW_RGB, BO_Add, BF_One, BF_One >::GetRHI() );
+			RHICmdList.SetViewport(View.ViewRect.Min.X, View.ViewRect.Min.Y, 0.0f, View.ViewRect.Max.X, View.ViewRect.Max.Y, 1.0f);
+			RHICmdList.SetRasterizerState(TStaticRasterizerState< FM_Solid, CM_None >::GetRHI());
+			RHICmdList.SetDepthStencilState(TStaticDepthStencilState< false, CF_Always >::GetRHI());
+			RHICmdList.SetBlendState(TStaticBlendState< CW_RGB, BO_Add, BF_One, BF_One >::GetRHI());
 
 			TShaderMapRef< FPostProcessVS > VertexShader( GetGlobalShaderMap() );
 
@@ -3151,7 +3148,8 @@ void FDeferredShadingSceneRenderer::RenderDynamicSkyLighting()
 				TShaderMapRef<TDynamicSkyLightDiffusePS<true> > PixelShader( GetGlobalShaderMap() );
 
 				static FGlobalBoundShaderState BoundShaderState;
-				SetGlobalBoundShaderState(BoundShaderState, GFilterVertexDeclaration.VertexDeclarationRHI, *VertexShader, *PixelShader);
+				RHICmdList.CheckIsNull(); // need new approach for "static FGlobalBoundShaderState" for parallel rendering
+				SetGlobalBoundShaderState(RHICmdList, BoundShaderState, GFilterVertexDeclaration.VertexDeclarationRHI, *VertexShader, *PixelShader);
 
 				PixelShader->SetParameters(RHICmdList, View, LightAccumulation->GetRenderTargetItem().ShaderResourceTexture);
 			}
@@ -3159,13 +3157,15 @@ void FDeferredShadingSceneRenderer::RenderDynamicSkyLighting()
 			{
 				TShaderMapRef<TDynamicSkyLightDiffusePS<false> > PixelShader( GetGlobalShaderMap() );
 
-			static FGlobalBoundShaderState BoundShaderState;
-				SetGlobalBoundShaderState(BoundShaderState, GFilterVertexDeclaration.VertexDeclarationRHI, *VertexShader, *PixelShader);
+				static FGlobalBoundShaderState BoundShaderState;
+				RHICmdList.CheckIsNull(); // need new approach for "static FGlobalBoundShaderState" for parallel rendering
+				SetGlobalBoundShaderState(RHICmdList, BoundShaderState, GFilterVertexDeclaration.VertexDeclarationRHI, *VertexShader, *PixelShader);
 
 				PixelShader->SetParameters(RHICmdList, View, LightAccumulation->GetRenderTargetItem().ShaderResourceTexture);
 			}
 
 			DrawRectangle( 
+				RHICmdList,
 				0, 0, 
 				View.ViewRect.Width(), View.ViewRect.Height(),
 				View.ViewRect.Min.X, View.ViewRect.Min.Y, 
