UE4 - more work preparing for parallel rendering. Merged the cmd list into the post process context, among other things.

diff --git a/Engine/Source/Runtime/Renderer/Private/DistortionRendering.cpp b/Engine/Source/Runtime/Renderer/Private/DistortionRendering.cpp
--- a/Engine/Source/Runtime/Renderer/Private/DistortionRendering.cpp
+++ b/Engine/Source/Runtime/Renderer/Private/DistortionRendering.cpp
@@ -33,7 +33,7 @@ public:
 	}
 	FDistortionApplyScreenPS() {}
 
-	void SetParameters(FRHICommandList& RHICmdList, const FRenderingCompositePassContext& Context, IPooledRenderTarget& DistortionRT)
+	void SetParameters(const FRenderingCompositePassContext& Context, IPooledRenderTarget& DistortionRT)
 	{
 		const FPixelShaderRHIParamRef ShaderRHI = GetPixelShader();
 
@@ -42,7 +42,7 @@ public:
 
 		// Here we use SF_Point as in fullscreen the pixels are 1:1 mapped.
 		SetTextureParameter(
-			RHICmdList, 
+			Context.RHICmdList,
 			ShaderRHI,
 			DistortionTexture,
 			DistortionTextureSampler,
@@ -51,7 +51,7 @@ public:
 			);
 
 		SetTextureParameter(
-			RHICmdList, 
+			Context.RHICmdList,
 			ShaderRHI,
 			SceneColorTexture,
 			SceneColorTextureSampler,
@@ -65,7 +65,7 @@ public:
 												(float)ViewportRect.Min.Y/SceneBufferSize.Y,
 												(float)ViewportRect.Max.X/SceneBufferSize.X,
 												(float)ViewportRect.Max.Y/SceneBufferSize.Y);
-		SetShaderValue(RHICmdList, ShaderRHI, SceneColorRect, SceneColorRectValue);
+		SetShaderValue(Context.RHICmdList, ShaderRHI, SceneColorRect, SceneColorRectValue);
 	}
 
 	virtual bool Serialize(FArchive& Ar)
@@ -409,13 +409,11 @@ void TDistortionMeshDrawingPolicy<DistortMeshPolicy>::DrawShared(
 	FBoundShaderStateRHIParamRef BoundShaderState
 	) const
 {
-	RHICmdList.CheckIsNull();
-
 	// Set shared mesh resources
 	FMeshDrawingPolicy::DrawShared(RHICmdList, View);
 
 	// Set the actual shader & vertex declaration state
-	RHISetBoundShaderState(BoundShaderState);
+	RHICmdList.SetBoundShaderState(BoundShaderState);
 
 	// Set the translucent shader parameters for the material instance
 	VertexShader->SetParameters(RHICmdList, VertexFactory,MaterialRenderProxy,View);
@@ -530,7 +528,7 @@ void TDistortionMeshDrawingPolicy<DistortMeshPolicy>::SetMeshRenderState(
 	}
 	
 	// Set rasterizer state.
-	RHISetRasterizerState(GetStaticRasterizerState<true>(
+	RHICmdList.SetRasterizerState(GetStaticRasterizerState<true>(
 		(Mesh.bWireframe || IsWireframe()) ? FM_Wireframe : FM_Solid,
 		IsTwoSided() ? CM_None : (XOR( XOR(View.bReverseCulling,bBackFace), Mesh.ReverseCulling) ? CM_CCW : CM_CW)));
 }
@@ -810,7 +808,7 @@ void FSceneRenderer::RenderDistortion(FRHICommandList& RHICmdList)
 		// DistortionRT==0 should never happen but better we don't crash
 		if(DistortionRT)
 		{
-			RHISetRenderTarget(DistortionRT->GetRenderTargetItem().TargetableTexture,GSceneRenderTargets.GetSceneDepthSurface());
+			SetRenderTarget(RHICmdList, DistortionRT->GetRenderTargetItem().TargetableTexture, GSceneRenderTargets.GetSceneDepthSurface());
 
 			for(int32 ViewIndex = 0;ViewIndex < Views.Num();ViewIndex++)
 			{
@@ -818,17 +816,19 @@ void FSceneRenderer::RenderDistortion(FRHICommandList& RHICmdList)
 
 				FViewInfo& View = Views[ViewIndex];
 				// viewport to match view size
-				RHISetViewport(View.ViewRect.Min.X,View.ViewRect.Min.Y,0.0f,View.ViewRect.Max.X,View.ViewRect.Max.Y,1.0f);
+				RHICmdList.SetViewport(View.ViewRect.Min.X, View.ViewRect.Min.Y, 0.0f, View.ViewRect.Max.X, View.ViewRect.Max.Y, 1.0f);
 
 				// clear offsets to 0, stencil to 0
-				RHIClear(true,FLinearColor(0,0,0,0),false,0,true,0,FIntRect());
+				RHICmdList.Clear(true, FLinearColor(0, 0, 0, 0), false, 0, true, 0, FIntRect());
 
 				// enable depth test but disable depth writes
 				// Note, this is a reversed Z depth surface, using CF_GreaterEqual.
-				RHISetDepthStencilState(TStaticDepthStencilState<false,CF_GreaterEqual>::GetRHI());
+				RHICmdList.SetDepthStencilState(TStaticDepthStencilState<false, CF_GreaterEqual>::GetRHI());
 
 				// additive blending of offsets (or complexity if the shader complexity viewmode is enabled)
-				RHISetBlendState(TStaticBlendState<CW_RGBA,BO_Add,BF_One,BF_One,BO_Add,BF_One,BF_One>::GetRHI());
+				RHICmdList.SetBlendState(TStaticBlendState<CW_RGBA, BO_Add, BF_One, BF_One, BO_Add, BF_One, BF_One>::GetRHI());
+
+				RHICmdList.CheckIsNull(); // DrawAccumulatedOffsets
 
 				// draw only distortion meshes to accumulate their offsets
 				bDirty |= View.DistortionPrimSet.DrawAccumulatedOffsets(&View,false);
@@ -838,12 +838,12 @@ void FSceneRenderer::RenderDistortion(FRHICommandList& RHICmdList)
 			{
 				// restore default stencil state
 				// Note, this is a reversed Z depth surface, using CF_GreaterEqual.
-				RHISetDepthStencilState(TStaticDepthStencilState<true,CF_GreaterEqual>::GetRHI());
+				RHICmdList.SetDepthStencilState(TStaticDepthStencilState<true, CF_GreaterEqual>::GetRHI());
 
 				// resolve using the current ResolveParams 
-				RHICopyToResolveTarget(DistortionRT->GetRenderTargetItem().TargetableTexture, DistortionRT->GetRenderTargetItem().ShaderResourceTexture, false, FResolveParams());
+				RHICmdList.CopyToResolveTarget(DistortionRT->GetRenderTargetItem().TargetableTexture, DistortionRT->GetRenderTargetItem().ShaderResourceTexture, false, FResolveParams());
 				// to be able to observe results with VisualizeTexture
-				GRenderTargetPool.VisualizeTexture.SetCheckPoint(DistortionRT);
+				GRenderTargetPool.VisualizeTexture.SetCheckPoint(RHICmdList, DistortionRT);
 			}
 		}
 	}
@@ -852,6 +852,8 @@ void FSceneRenderer::RenderDistortion(FRHICommandList& RHICmdList)
 	{
 		SCOPED_DRAW_EVENT(DistortionApply, DEC_SCENE_ITEMS);
 
+		RHICmdList.CheckIsNull(); // GSceneRenderTargets.ResolveSceneColor
+
 		GSceneRenderTargets.ResolveSceneColor(FResolveRect(0, 0, ViewFamily.FamilySizeX, ViewFamily.FamilySizeY));
 
 // OCULUS BEGIN: select ONE render target for all views (eyes)
@@ -870,9 +872,9 @@ void FSceneRenderer::RenderDistortion(FRHICommandList& RHICmdList)
 //			const FSceneView& View = Context.View;
 
 			// set the state
-			RHISetBlendState(TStaticBlendState<>::GetRHI());
-			RHISetRasterizerState(TStaticRasterizerState<>::GetRHI());
-			RHISetDepthStencilState(TStaticDepthStencilState<false,CF_Always>::GetRHI());
+			RHICmdList.SetBlendState(TStaticBlendState<>::GetRHI());
+			RHICmdList.SetRasterizerState(TStaticRasterizerState<>::GetRHI());
+			RHICmdList.SetDepthStencilState(TStaticDepthStencilState<false, CF_Always>::GetRHI());
 
 			{
 // OCULUS BEGIN: Weird, different eyes may have different render targets!
@@ -881,7 +883,7 @@ void FSceneRenderer::RenderDistortion(FRHICommandList& RHICmdList)
 // 				const FSceneRenderTargetItem& DestRenderTarget = NewSceneColor->GetRenderTargetItem(); 
 // OCULUS END
 
-				RHISetRenderTarget( DestRenderTarget.TargetableTexture, FTextureRHIParamRef());
+				SetRenderTarget(RHICmdList, DestRenderTarget.TargetableTexture, FTextureRHIParamRef());
 
 				// useful when we move this into the compositing graph
 				FRenderingCompositePassContext Context(View);
@@ -893,13 +895,15 @@ void FSceneRenderer::RenderDistortion(FRHICommandList& RHICmdList)
 				TShaderMapRef<FDistortionApplyScreenPS> PixelShader(GetGlobalShaderMap());
 
 				static FGlobalBoundShaderState BoundShaderState;
-				SetGlobalBoundShaderState(BoundShaderState, GFilterVertexDeclaration.VertexDeclarationRHI, *VertexShader, *PixelShader);
+				RHICmdList.CheckIsNull(); // need new approach for "static FGlobalBoundShaderState" for parallel rendering
+				SetGlobalBoundShaderState(RHICmdList, BoundShaderState, GFilterVertexDeclaration.VertexDeclarationRHI, *VertexShader, *PixelShader);
 
-				VertexShader->SetParameters(RHICmdList, Context);
-				PixelShader->SetParameters(RHICmdList, Context, *DistortionRT);
+				VertexShader->SetParameters(Context);
+				PixelShader->SetParameters(Context, *DistortionRT);
 
 				// Draw a quad mapping scene color to the view's render target
 				DrawRectangle(
+					RHICmdList,
 					0, 0,
 					View.ViewRect.Width(), View.ViewRect.Height(),
 					View.ViewRect.Min.X, View.ViewRect.Min.Y, 
@@ -909,7 +913,7 @@ void FSceneRenderer::RenderDistortion(FRHICommandList& RHICmdList)
 					*VertexShader,
 					EDRF_UseTriangleOptimization);
 
-				RHICopyToResolveTarget(DestRenderTarget.TargetableTexture, DestRenderTarget.ShaderResourceTexture, false, FResolveParams());
+				RHICmdList.CopyToResolveTarget(DestRenderTarget.TargetableTexture, DestRenderTarget.ShaderResourceTexture, false, FResolveParams());
 
 // OCULUS BEGIN
 // 				GSceneRenderTargets.SetSceneColor(NewSceneColor);
@@ -920,7 +924,8 @@ void FSceneRenderer::RenderDistortion(FRHICommandList& RHICmdList)
 		
 		}
 
-// OCULUS BEGIN
+		RHICmdList.CheckIsNull(); // GSceneRenderTargets ops
+		// OCULUS BEGIN
 		GSceneRenderTargets.SetSceneColor(NewSceneColor);
 		check(GSceneRenderTargets.GetSceneColor());
 // OCULUS END				
diff --git a/Engine/Source/Runtime/Renderer/Private/DrawingPolicy.cpp b/Engine/Source/Runtime/Renderer/Private/DrawingPolicy.cpp
--- a/Engine/Source/Runtime/Renderer/Private/DrawingPolicy.cpp
+++ b/Engine/Source/Runtime/Renderer/Private/DrawingPolicy.cpp
@@ -73,9 +73,8 @@ void FMeshDrawingPolicy::DrawMesh(FRHICommandList& RHICmdList, const FMeshBatch&
 
 		if (BatchElement.DynamicIndexData)
 		{
-			RHICmdList.CheckIsNull();
-
-			RHIDrawIndexedPrimitiveUP(
+			DrawIndexedPrimitiveUP(
+				RHICmdList,
 				Mesh.Type,
 				BatchElement.MinVertexIndex,
 				BatchElement.MaxVertexIndex - BatchElement.MinVertexIndex + 1,
@@ -90,7 +89,8 @@ void FMeshDrawingPolicy::DrawMesh(FRHICommandList& RHICmdList, const FMeshBatch&
 		{
 			RHICmdList.CheckIsNull();
 
-			RHIDrawPrimitiveUP(
+			DrawPrimitiveUP(
+				RHICmdList,
 				Mesh.Type,
 				BatchElement.NumPrimitives,
 				Mesh.DynamicVertexData,
diff --git a/Engine/Source/Runtime/Renderer/Private/DynamicPrimitiveDrawing.inl b/Engine/Source/Runtime/Renderer/Private/DynamicPrimitiveDrawing.inl
--- a/Engine/Source/Runtime/Renderer/Private/DynamicPrimitiveDrawing.inl
+++ b/Engine/Source/Runtime/Renderer/Private/DynamicPrimitiveDrawing.inl
@@ -22,9 +22,12 @@ TDynamicPrimitiveDrawer<DrawingPolicyFactoryType>::~TDynamicPrimitiveDrawer()
 		{
 			DepthTexture = GSceneRenderTargets.GetSceneDepthTexture();
 		}
+		//@todo-rco: RHIPacketList
+		FRHICommandList& RHICmdList = FRHICommandList::GetNullRef();
 
 		// Draw the batched elements.
 		BatchedElements.Draw(
+			RHICmdList,
 			bNeedToSwitchVerticalAxis,
 			View->ViewProjectionMatrix,
 			View->ViewRect.Width(),
diff --git a/Engine/Source/Runtime/Renderer/Private/FogRendering.cpp b/Engine/Source/Runtime/Renderer/Private/FogRendering.cpp
--- a/Engine/Source/Runtime/Renderer/Private/FogRendering.cpp
+++ b/Engine/Source/Runtime/Renderer/Private/FogRendering.cpp
@@ -270,13 +270,13 @@ void SetFogShaders(FRHICommandList& RHICmdList, FScene* Scene,const FViewInfo& V
 		TShaderMapRef<FHeightFogVS> VertexShader(GetGlobalShaderMap());
 		TShaderMapRef<FExponentialHeightFogPS> ExponentialHeightFogPixelShader(GetGlobalShaderMap());
 
-		SetGlobalBoundShaderState(ExponentialBoundShaderState, GFogVertexDeclaration.VertexDeclarationRHI, *VertexShader, *ExponentialHeightFogPixelShader);
+		SetGlobalBoundShaderState(RHICmdList, ExponentialBoundShaderState, GFogVertexDeclaration.VertexDeclarationRHI, *VertexShader, *ExponentialHeightFogPixelShader);
 		VertexShader->SetParameters(RHICmdList, View);
 		ExponentialHeightFogPixelShader->SetParameters(RHICmdList, View, LightShaftsOutput);
 	}
 }
 
-bool FDeferredShadingSceneRenderer::RenderFog(FLightShaftsOutput LightShaftsOutput)
+bool FDeferredShadingSceneRenderer::RenderFog(FRHICommandList& RHICmdList, FLightShaftsOutput LightShaftsOutput)
 {
 	if (Scene->ExponentialFogs.Num() > 0)
 	{
@@ -295,8 +295,7 @@ bool FDeferredShadingSceneRenderer::RenderFog(FLightShaftsOutput LightShaftsOutp
 			0, 2, 3
 		};
 
-		//@todo-rco: RHIPacketList
-		FRHICommandList& RHICmdList = FRHICommandList::GetNullRef();
+		RHICmdList.CheckIsNull(); // GSceneRenderTargets ops
 
 		GSceneRenderTargets.BeginRenderingSceneColor();
 		for(int32 ViewIndex = 0;ViewIndex < Views.Num();ViewIndex++)
@@ -309,19 +308,20 @@ bool FDeferredShadingSceneRenderer::RenderFog(FLightShaftsOutput LightShaftsOutp
 			}
 
 			// Set the device viewport for the view.
-			RHISetViewport(View.ViewRect.Min.X, View.ViewRect.Min.Y, 0.0f, View.ViewRect.Max.X, View.ViewRect.Max.Y, 1.0f);
+			RHICmdList.SetViewport(View.ViewRect.Min.X, View.ViewRect.Min.Y, 0.0f, View.ViewRect.Max.X, View.ViewRect.Max.Y, 1.0f);
 			
-			RHISetRasterizerState(TStaticRasterizerState<FM_Solid,CM_None>::GetRHI());
+			RHICmdList.SetRasterizerState(TStaticRasterizerState<FM_Solid, CM_None>::GetRHI());
 			
 			// disable alpha writes in order to preserve scene depth values on PC
-			RHISetBlendState(TStaticBlendState<CW_RGB,BO_Add,BF_One,BF_SourceAlpha>::GetRHI());
+			RHICmdList.SetBlendState(TStaticBlendState<CW_RGB, BO_Add, BF_One, BF_SourceAlpha>::GetRHI());
 
-			RHISetDepthStencilState(TStaticDepthStencilState<false,CF_Always>::GetRHI());
+			RHICmdList.SetDepthStencilState(TStaticDepthStencilState<false, CF_Always>::GetRHI());
 
 			SetFogShaders(RHICmdList, Scene,View,LightShaftsOutput);
 
 			// Draw a quad covering the view.
-			RHIDrawIndexedPrimitiveUP(
+			DrawIndexedPrimitiveUP(
+				RHICmdList,
 				PT_TriangleList,
 				0,
 				ARRAY_COUNT(Vertices),
diff --git a/Engine/Source/Runtime/Renderer/Private/ForwardBasePassRendering.cpp b/Engine/Source/Runtime/Renderer/Private/ForwardBasePassRendering.cpp
--- a/Engine/Source/Runtime/Renderer/Private/ForwardBasePassRendering.cpp
+++ b/Engine/Source/Runtime/Renderer/Private/ForwardBasePassRendering.cpp
@@ -400,15 +400,17 @@ void FForwardShadingSceneRenderer::RenderForwardShadingBasePass(FRHICommandList&
 
 			// Draw the base pass for the view's batched mesh elements.
 			DrawViewElements<FBasePassForwardOpaqueDrawingPolicyFactory>(View,FBasePassForwardOpaqueDrawingPolicyFactory::ContextType(ESceneRenderTargetsMode::DontSet), SDPG_World, true);
+			//@todo-rco: RHIPacketList
+			FRHICommandList& RHICmdList = FRHICommandList::GetNullRef();
 
 			// Draw the view's batched simple elements(lines, sprites, etc).
-			View.BatchedViewElements.Draw(bNeedToSwitchVerticalAxis, View.ViewProjectionMatrix, View.ViewRect.Width(), View.ViewRect.Height(), false);
+			View.BatchedViewElements.Draw(RHICmdList, bNeedToSwitchVerticalAxis, View.ViewProjectionMatrix, View.ViewRect.Width(), View.ViewRect.Height(), false);
 
 			// Draw foreground objects last
 			DrawViewElements<FBasePassForwardOpaqueDrawingPolicyFactory>(View,FBasePassForwardOpaqueDrawingPolicyFactory::ContextType(ESceneRenderTargetsMode::DontSet), SDPG_Foreground, true);
 
 			// Draw the view's batched simple elements(lines, sprites, etc).
-			View.TopBatchedViewElements.Draw(bNeedToSwitchVerticalAxis, View.ViewProjectionMatrix, View.ViewRect.Width(), View.ViewRect.Height(), false);
+			View.TopBatchedViewElements.Draw(RHICmdList, bNeedToSwitchVerticalAxis, View.ViewProjectionMatrix, View.ViewRect.Width(), View.ViewRect.Height(), false);
 		}
 
 		// Issue static draw list masked draw calls last, as PVR wants it
diff --git a/Engine/Source/Runtime/Renderer/Private/ForwardShadingRenderer.cpp b/Engine/Source/Runtime/Renderer/Private/ForwardShadingRenderer.cpp
--- a/Engine/Source/Runtime/Renderer/Private/ForwardShadingRenderer.cpp
+++ b/Engine/Source/Runtime/Renderer/Private/ForwardShadingRenderer.cpp
@@ -55,7 +55,10 @@ void FForwardShadingSceneRenderer::Render()
 
 	// Find the visible primitives.
 	InitViews();
-	
+
+	//@todo-rco: RHIPacketList
+	FRHICommandList& RHICmdList = FRHICommandList::GetNullRef();
+
 	// Notify the FX system that the scene is about to be rendered.
 	if (Scene->FXSystem)
 	{
@@ -71,7 +74,7 @@ void FForwardShadingSceneRenderer::Render()
 	const bool bGammaSpace = !IsMobileHDR();
 	if( bGammaSpace )
 	{
-		RHISetRenderTarget( ViewFamily.RenderTarget->GetRenderTargetTexture(), GSceneRenderTargets.GetSceneDepthTexture() );
+		SetRenderTarget(RHICmdList, ViewFamily.RenderTarget->GetRenderTargetTexture(), GSceneRenderTargets.GetSceneDepthTexture());
 	}
 	else
 	{
@@ -81,10 +84,8 @@ void FForwardShadingSceneRenderer::Render()
 
 	// Clear color and depth buffer
 	// Note, this is a reversed Z depth surface, so 0.0f is the far plane.
-	RHIClear(true, FLinearColor::Black, true, 0.0f, true, 0, FIntRect());
+	RHICmdList.Clear(true, FLinearColor::Black, true, 0.0f, true, 0, FIntRect());
 
-	//@todo-rco: RHIPacketList
-	FRHICommandList& RHICmdList = FRHICommandList::GetNullRef();
 	RenderForwardShadingBasePass(RHICmdList);
 
 	// Notify the FX system that opaque primitives have been rendered.
diff --git a/Engine/Source/Runtime/Renderer/Private/ForwardTranslucentRendering.cpp b/Engine/Source/Runtime/Renderer/Private/ForwardTranslucentRendering.cpp
--- a/Engine/Source/Runtime/Renderer/Private/ForwardTranslucentRendering.cpp
+++ b/Engine/Source/Runtime/Renderer/Private/ForwardTranslucentRendering.cpp
@@ -52,6 +52,7 @@ void FForwardShadingSceneRenderer::CopySceneAlpha(FRHICommandList& RHICmdList, c
 	RHISetDepthStencilState(TStaticDepthStencilState<false,CF_Always>::GetRHI());
 	RHISetBlendState(TStaticBlendState<>::GetRHI());
 
+	RHICmdList.CheckIsNull(); // GSceneRenderTargets ops
 	GSceneRenderTargets.ResolveSceneColor();
 
 	GSceneRenderTargets.BeginRenderingSceneAlphaCopy();
@@ -59,15 +60,16 @@ void FForwardShadingSceneRenderer::CopySceneAlpha(FRHICommandList& RHICmdList, c
 	int X = GSceneRenderTargets.GetBufferSizeXY().X;
 	int Y = GSceneRenderTargets.GetBufferSizeXY().Y;
 
-	RHISetViewport(0, 0, 0.0f, X, Y, 1.0f );
+	RHICmdList.SetViewport(0, 0, 0.0f, X, Y, 1.0f);
 
 	TShaderMapRef<FScreenVS> ScreenVertexShader(GetGlobalShaderMap());
 	TShaderMapRef<FForwardCopySceneAlphaPS> PixelShader(GetGlobalShaderMap());
-	SetGlobalBoundShaderState(ForwardCopySceneAlphaBoundShaderState, GFilterVertexDeclaration.VertexDeclarationRHI, *ScreenVertexShader, *PixelShader);
+	SetGlobalBoundShaderState(RHICmdList, ForwardCopySceneAlphaBoundShaderState, GFilterVertexDeclaration.VertexDeclarationRHI, *ScreenVertexShader, *PixelShader);
 
 	PixelShader->SetParameters(RHICmdList, View);
 
 	DrawRectangle( 
+		RHICmdList,
 		0, 0, 
 		X, Y, 
 		0, 0, 
diff --git a/Engine/Source/Runtime/Renderer/Private/GPUBenchmark.cpp b/Engine/Source/Runtime/Renderer/Private/GPUBenchmark.cpp
--- a/Engine/Source/Runtime/Renderer/Private/GPUBenchmark.cpp
+++ b/Engine/Source/Runtime/Renderer/Private/GPUBenchmark.cpp
@@ -131,8 +131,9 @@ void RunBenchmarkShader(FRHICommandList& RHICmdList, const FSceneView& View, TRe
 	TShaderMapRef<FPostProcessBenchmarkPS<Method> > PixelShader(GetGlobalShaderMap());
 
 	static FGlobalBoundShaderState BoundShaderState;
+	RHICmdList.CheckIsNull(); // need new approach for "static FGlobalBoundShaderState" for parallel rendering
 
-	SetGlobalBoundShaderState(BoundShaderState, GFilterVertexDeclaration.VertexDeclarationRHI, *VertexShader, *PixelShader);
+	SetGlobalBoundShaderState(RHICmdList, BoundShaderState, GFilterVertexDeclaration.VertexDeclarationRHI, *VertexShader, *PixelShader);
 
 	PixelShader->SetParameters(RHICmdList, View, Src);
 	VertexShader->SetParameters(RHICmdList, View);
@@ -140,6 +141,7 @@ void RunBenchmarkShader(FRHICommandList& RHICmdList, const FSceneView& View, TRe
 	for(uint32 i = 0; i < Count; ++i)
 	{
 		DrawRectangle(
+			RHICmdList,
 			0, 0,
 			GBenchmarkResolution, GBenchmarkResolution,
 			0, 0,
@@ -427,7 +429,7 @@ void RendererGPUBenchmark(FSynthBenchmarkResults& InOut, const FSceneView& View,
 				// 0 / 1
 				const uint32 SrcRTIndex = 1 - DestRTIndex;
 
-				GRenderTargetPool.VisualizeTexture.SetCheckPoint(RTItems[DestRTIndex]);
+				GRenderTargetPool.VisualizeTexture.SetCheckPoint(RHICmdList, RTItems[DestRTIndex]);
 
 				RHISetRenderTarget(RTItems[DestRTIndex]->GetRenderTargetItem().TargetableTexture, FTextureRHIRef());	
 
diff --git a/Engine/Source/Runtime/Renderer/Private/LightFunctionRendering.cpp b/Engine/Source/Runtime/Renderer/Private/LightFunctionRendering.cpp
--- a/Engine/Source/Runtime/Renderer/Private/LightFunctionRendering.cpp
+++ b/Engine/Source/Runtime/Renderer/Private/LightFunctionRendering.cpp
@@ -334,7 +334,7 @@ bool FDeferredShadingSceneRenderer::RenderLightFunctionForMaterial(FRHICommandLi
 
 					// Project the light function using a sphere around the light
 					//@todo - could use a cone for spotlights
-					StencilingGeometry::DrawSphere();
+					StencilingGeometry::DrawSphere(RHICmdList);
 				}
 			}
 		}
diff --git a/Engine/Source/Runtime/Renderer/Private/LightPropagationVolume.h b/Engine/Source/Runtime/Renderer/Private/LightPropagationVolume.h
--- a/Engine/Source/Runtime/Renderer/Private/LightPropagationVolume.h
+++ b/Engine/Source/Runtime/Renderer/Private/LightPropagationVolume.h
@@ -144,7 +144,7 @@ public:
 
 	void Propagate(FRHICommandList& RHICmdList);
 
-	void Visualise( const FSceneView& View ) const;
+	void Visualise(FRHICommandList& RHICmdList, const FSceneView& View) const;
 
 	const FIntVector& GetGridOffset() const { return mGridOffset; }
 
diff --git a/Engine/Source/Runtime/Renderer/Private/LightPropagationVolumeVisualisation.cpp b/Engine/Source/Runtime/Renderer/Private/LightPropagationVolumeVisualisation.cpp
--- a/Engine/Source/Runtime/Renderer/Private/LightPropagationVolumeVisualisation.cpp
+++ b/Engine/Source/Runtime/Renderer/Private/LightPropagationVolumeVisualisation.cpp
@@ -211,7 +211,7 @@ public:
 	FShaderResourceParameter InGvBuffer;
 #endif
 
-	void UnbindBuffers()
+	void UnbindBuffers(FRHICommandList& RHICmdList)
 	{
 		// TODO: Is this necessary here?
 		FPixelShaderRHIParamRef ShaderRHI = GetPixelShader();
@@ -220,11 +220,11 @@ public:
 		{
 			if ( LpvBufferSRVParameters[i].IsBound() )
 			{
-				RHISetShaderTexture( ShaderRHI, LpvBufferSRVParameters[i].GetBaseIndex(), FTextureRHIParamRef() );
+				RHICmdList.SetShaderTexture(ShaderRHI, LpvBufferSRVParameters[i].GetBaseIndex(), FTextureRHIParamRef());
 			}
 		}
 #else
-		if ( InLpvBuffer.IsBound() ) RHISetShaderResourceViewParameter( ShaderRHI, InLpvBuffer.GetBaseIndex(), FShaderResourceViewRHIParamRef() );
+		if ( InLpvBuffer.IsBound() ) RHICmdList.SetShaderResourceViewParameter( ShaderRHI, InLpvBuffer.GetBaseIndex(), FShaderResourceViewRHIParamRef() );
 #endif
 
 #if LPV_GV_VOLUME_TEXTURE
@@ -232,11 +232,11 @@ public:
 		{
 			if ( GvBufferSRVParameters[i].IsBound() )
 			{
-				RHISetShaderTexture( ShaderRHI, GvBufferSRVParameters[i].GetBaseIndex(), FTextureRHIParamRef() );
+				RHICmdList.SetShaderTexture(ShaderRHI, GvBufferSRVParameters[i].GetBaseIndex(), FTextureRHIParamRef());
 			}
 		}
 #else
-		if ( InGvBuffer.IsBound() ) RHISetShaderResourceViewParameter( ShaderRHI, InGvBuffer.GetBaseIndex(), FShaderResourceViewRHIParamRef() );
+		if ( InGvBuffer.IsBound() ) RHICmdList.SetShaderResourceViewParameter( ShaderRHI, InGvBuffer.GetBaseIndex(), FShaderResourceViewRHIParamRef() );
 #endif
 	}
 
@@ -248,7 +248,7 @@ IMPLEMENT_SHADER_TYPE(,FLpvVisualiseVS,TEXT("LPVVisualise"),TEXT("VShader"),SF_V
 IMPLEMENT_SHADER_TYPE(,FLpvVisualisePS,TEXT("LPVVisualise"),TEXT("PShader"),SF_Pixel);
 
 
-void FLightPropagationVolume::Visualise( const FSceneView& View ) const
+void FLightPropagationVolume::Visualise(FRHICommandList& RHICmdList, const FSceneView& View) const
 {
 	SCOPED_DRAW_EVENT(LpvVisualise, DEC_LIGHT);
 	check(View.GetFeatureLevel() == ERHIFeatureLevel::SM5);
@@ -257,22 +257,18 @@ void FLightPropagationVolume::Visualise( const FSceneView& View ) const
 	TShaderMapRef<FLpvVisualiseGS> GeometryShader(GetGlobalShaderMap());
 	TShaderMapRef<FLpvVisualisePS> PixelShader(GetGlobalShaderMap());
 
-	RHISetDepthStencilState( TStaticDepthStencilState<false,CF_Always>::GetRHI() );
-	RHISetRasterizerState( TStaticRasterizerState<FM_Solid,CM_None>::GetRHI() );
-	RHISetBlendState(TStaticBlendState<CW_RGB,BO_Add,BF_One,BF_One>::GetRHI());
+	RHICmdList.SetDepthStencilState(TStaticDepthStencilState<false, CF_Always>::GetRHI());
+	RHICmdList.SetRasterizerState(TStaticRasterizerState<FM_Solid, CM_None>::GetRHI());
+	RHICmdList.SetBlendState(TStaticBlendState<CW_RGB, BO_Add, BF_One, BF_One>::GetRHI());
 
-	//@todo-rco: RHIPacketList
-
-	FRHICommandList& RHICmdList = FRHICommandList::GetNullRef();
-
-	SetGlobalBoundShaderState(LpvVisBoundShaderState, GSimpleElementVertexDeclaration.VertexDeclarationRHI, *VertexShader, *PixelShader, *GeometryShader);
+	SetGlobalBoundShaderState(RHICmdList, LpvVisBoundShaderState, GSimpleElementVertexDeclaration.VertexDeclarationRHI, *VertexShader, *PixelShader, *GeometryShader);
 
 	VertexShader->SetParameters(RHICmdList, View);
 	GeometryShader->SetParameters(RHICmdList, View);
 	PixelShader->SetParameters(RHICmdList, this, View);
 
-	RHISetStreamSource(0, NULL, 0, 0);
-	RHIDrawPrimitive( PT_PointList, 0, 1, 32*3 );
+	RHICmdList.SetStreamSource(0, NULL, 0, 0);
+	RHICmdList.DrawPrimitive(PT_PointList, 0, 1, 32 * 3);
 
-	PixelShader->UnbindBuffers();
+	PixelShader->UnbindBuffers(RHICmdList);
 }
diff --git a/Engine/Source/Runtime/Renderer/Private/LightRendering.cpp b/Engine/Source/Runtime/Renderer/Private/LightRendering.cpp
--- a/Engine/Source/Runtime/Renderer/Private/LightRendering.cpp
+++ b/Engine/Source/Runtime/Renderer/Private/LightRendering.cpp
@@ -421,7 +421,7 @@ void FDeferredShadingSceneRenderer::RenderLights(FRHICommandList& RHICmdList)
 			else if (SimpleLights.InstanceData.Num() > 0)
 			{
 				GSceneRenderTargets.BeginRenderingSceneColor();
-				RenderSimpleLightsStandardDeferred(SimpleLights);
+				RenderSimpleLightsStandardDeferred(RHICmdList, SimpleLights);
 			}
 
 			{
@@ -437,7 +437,7 @@ void FDeferredShadingSceneRenderer::RenderLights(FRHICommandList& RHICmdList)
 					const FLightSceneInfo* const LightSceneInfo = LightSceneInfoCompact.LightSceneInfo;
 
 					// Render the light to the scene color buffer, using a 1x1 white texture as input 
-					RenderLight( LightSceneInfo, false, false );
+					RenderLight(RHICmdList, LightSceneInfo, false, false);
 				}
 			}
 
@@ -447,13 +447,13 @@ void FDeferredShadingSceneRenderer::RenderLights(FRHICommandList& RHICmdList)
 				{
 					// Inject non-shadowed, non-light function lights in to the volume.
 					SCOPED_DRAW_EVENT(InjectNonShadowedTranslucentLighting, DEC_SCENE_ITEMS);
-					InjectTranslucentVolumeLightingArray(SortedLights, AttenuationLightStart);
+					InjectTranslucentVolumeLightingArray(RHICmdList, SortedLights, AttenuationLightStart);
 				}
 				
 				if (SimpleLights.InstanceData.Num() > 0)
 				{
 					SCOPED_DRAW_EVENT(InjectSimpleLightsTranslucentLighting, DEC_SCENE_ITEMS);
-					InjectSimpleTranslucentVolumeLightingArray(SimpleLights);
+					InjectSimpleTranslucentVolumeLightingArray(RHICmdList, SimpleLights);
 				}
 			}
 		}
@@ -527,7 +527,7 @@ void FDeferredShadingSceneRenderer::RenderLights(FRHICommandList& RHICmdList)
 			{
 				SCOPED_DRAW_EVENT(InjectTranslucentVolume, DEC_SCENE_ITEMS);
 				// Accumulate this light's unshadowed contribution to the translucency lighting volume
-				InjectTranslucentVolumeLighting(LightSceneInfo, NULL);
+				InjectTranslucentVolumeLighting(RHICmdList, LightSceneInfo, NULL);
 			}
 
 			GSceneRenderTargets.SetLightAttenuationMode(true);
@@ -536,7 +536,7 @@ void FDeferredShadingSceneRenderer::RenderLights(FRHICommandList& RHICmdList)
 			// Render the light to the scene color buffer, conditionally using the attenuation buffer or a 1x1 white texture as input 
 			if(bDirectLighting)
 			{
-				RenderLight( &LightSceneInfo, false, true );
+				RenderLight(RHICmdList, &LightSceneInfo, false, true);
 			}
 		}
 
@@ -579,7 +579,7 @@ void FDeferredShadingSceneRenderer::RenderLights(FRHICommandList& RHICmdList)
 	}
 }
 
-void FDeferredShadingSceneRenderer::RenderLightArrayForOverlapViewmode(const TSparseArray<FLightSceneInfoCompact>& LightArray)
+void FDeferredShadingSceneRenderer::RenderLightArrayForOverlapViewmode(FRHICommandList& RHICmdList, const TSparseArray<FLightSceneInfoCompact>& LightArray)
 {
 	for (TSparseArray<FLightSceneInfoCompact>::TConstIterator LightIt(LightArray); LightIt; ++LightIt)
 	{
@@ -606,12 +606,12 @@ void FDeferredShadingSceneRenderer::RenderLightArrayForOverlapViewmode(const TSp
 			&& !LightSceneInfo->Proxy->HasStaticLighting()
 			&& LightSceneInfo->Proxy->CastsStaticShadow())
 		{
-			RenderLight(LightSceneInfo, true, false);
+			RenderLight(RHICmdList, LightSceneInfo, true, false);
 		}
 	}
 }
 
-void FDeferredShadingSceneRenderer::RenderStationaryLightOverlap()
+void FDeferredShadingSceneRenderer::RenderStationaryLightOverlap(FRHICommandList& RHICmdList)
 {
 	if (Scene->bIsEditorScene)
 	{
@@ -620,31 +620,31 @@ void FDeferredShadingSceneRenderer::RenderStationaryLightOverlap()
 		// Clear to discard base pass values in scene color since we didn't skip that, to have valid scene depths
 		RHIClear(true, FLinearColor::Black, false, 0, false, 0, FIntRect());
 
-		RenderLightArrayForOverlapViewmode(Scene->Lights);
+		RenderLightArrayForOverlapViewmode(RHICmdList, Scene->Lights);
 
 		//Note: making use of FScene::InvisibleLights, which contains lights that haven't been added to the scene in the same way as visible lights
 		// So code called by RenderLightArrayForOverlapViewmode must be careful what it accesses
-		RenderLightArrayForOverlapViewmode(Scene->InvisibleLights);
+		RenderLightArrayForOverlapViewmode(RHICmdList, Scene->InvisibleLights);
 	}
 }
 
 /** Sets up rasterizer and depth state for rendering bounding geometry in a deferred pass. */
-void SetBoundingGeometryRasterizerAndDepthState(const FViewInfo& View, const FSphere& LightBounds)
+void SetBoundingGeometryRasterizerAndDepthState(FRHICommandList& RHICmdList, const FViewInfo& View, const FSphere& LightBounds)
 {
 	const bool bCameraInsideLightGeometry = ((FVector)View.ViewMatrices.ViewOrigin - LightBounds.Center).SizeSquared() < FMath::Square(LightBounds.W * 1.05f + View.NearClippingDistance * 2.0f);
 	if (bCameraInsideLightGeometry)
 	{
 		// Render backfaces with depth tests disabled since the camera is inside (or close to inside) the light geometry
-		RHISetRasterizerState(View.bReverseCulling ? TStaticRasterizerState<FM_Solid,CM_CW>::GetRHI() : TStaticRasterizerState<FM_Solid,CM_CCW>::GetRHI());
+		RHICmdList.SetRasterizerState(View.bReverseCulling ? TStaticRasterizerState<FM_Solid, CM_CW>::GetRHI() : TStaticRasterizerState<FM_Solid, CM_CCW>::GetRHI());
 	}
 	else
 	{
 		// Render frontfaces with depth tests on to get the speedup from HiZ since the camera is outside the light geometry
-		RHISetRasterizerState(View.bReverseCulling ? TStaticRasterizerState<FM_Solid,CM_CCW>::GetRHI() : TStaticRasterizerState<FM_Solid,CM_CW>::GetRHI());
+		RHICmdList.SetRasterizerState(View.bReverseCulling ? TStaticRasterizerState<FM_Solid, CM_CCW>::GetRHI() : TStaticRasterizerState<FM_Solid, CM_CW>::GetRHI());
 	}
 
 	// Note, this is a reversed Z depth surface, using CF_GreaterEqual.
-	RHISetDepthStencilState(
+	RHICmdList.SetDepthStencilState(
 		bCameraInsideLightGeometry
 		? TStaticDepthStencilState<false,CF_Always>::GetRHI()
 		: TStaticDepthStencilState<false,CF_GreaterEqual>::GetRHI()
@@ -667,13 +667,13 @@ static void SetShaderTemplLighting(
 	if(View.Family->EngineShowFlags.VisualizeLightCulling)
 	{
 		TShaderMapRef<TDeferredLightPS<false, bRadialAttenuation, false, true> > PixelShader(GetGlobalShaderMap());
-		SetGlobalBoundShaderState(PixelShader->GetBoundShaderState(), GetDeferredLightingVertexDeclaration<bRadialAttenuation>(), VertexShader, *PixelShader);
+		SetGlobalBoundShaderState(RHICmdList, PixelShader->GetBoundShaderState(), GetDeferredLightingVertexDeclaration<bRadialAttenuation>(), VertexShader, *PixelShader);
 		PixelShader->SetParameters(RHICmdList, View, LightSceneInfo);
 	}
 	else
 	{
 		TShaderMapRef<TDeferredLightPS<bUseIESProfile, bRadialAttenuation, bInverseSquaredFalloff, false> > PixelShader(GetGlobalShaderMap());
-		SetGlobalBoundShaderState(PixelShader->GetBoundShaderState(), GetDeferredLightingVertexDeclaration<bRadialAttenuation>(), VertexShader, *PixelShader);
+		SetGlobalBoundShaderState(RHICmdList, PixelShader->GetBoundShaderState(), GetDeferredLightingVertexDeclaration<bRadialAttenuation>(), VertexShader, *PixelShader);
 		PixelShader->SetParameters(RHICmdList, View, LightSceneInfo);
 	}
 }
@@ -689,13 +689,13 @@ static void SetShaderTemplLightingSimple(
 	if(View.Family->EngineShowFlags.VisualizeLightCulling)
 	{
 		TShaderMapRef<TDeferredLightPS<false, bRadialAttenuation, false, true> > PixelShader(GetGlobalShaderMap());
-		SetGlobalBoundShaderState(PixelShader->GetBoundShaderState(), GetDeferredLightingVertexDeclaration<bRadialAttenuation>(), VertexShader, *PixelShader);
+		SetGlobalBoundShaderState(RHICmdList, PixelShader->GetBoundShaderState(), GetDeferredLightingVertexDeclaration<bRadialAttenuation>(), VertexShader, *PixelShader);
 		PixelShader->SetParametersSimpleLight(RHICmdList, View, SimpleLight, SimpleLightPerViewData);
 	}
 	else
 	{
 		TShaderMapRef<TDeferredLightPS<bUseIESProfile, bRadialAttenuation, bInverseSquaredFalloff, false> > PixelShader(GetGlobalShaderMap());
-		SetGlobalBoundShaderState(PixelShader->GetBoundShaderState(), GetDeferredLightingVertexDeclaration<bRadialAttenuation>(), VertexShader, *PixelShader);
+		SetGlobalBoundShaderState(RHICmdList, PixelShader->GetBoundShaderState(), GetDeferredLightingVertexDeclaration<bRadialAttenuation>(), VertexShader, *PixelShader);
 		PixelShader->SetParametersSimpleLight(RHICmdList, View, SimpleLight, SimpleLightPerViewData);
 	}
 }
@@ -707,14 +707,14 @@ static void SetShaderTemplLightingSimple(
  * @param LightIndex The light's index into FScene::Lights
  * @return true if anything got rendered
  */
-void FDeferredShadingSceneRenderer::RenderLight(const FLightSceneInfo* LightSceneInfo, bool bRenderOverlap, bool bIssueDrawEvent)
+void FDeferredShadingSceneRenderer::RenderLight(FRHICommandList& RHICmdList, const FLightSceneInfo* LightSceneInfo, bool bRenderOverlap, bool bIssueDrawEvent)
 {
 	SCOPE_CYCLE_COUNTER(STAT_DirectLightRenderingTime);
 	INC_DWORD_STAT(STAT_NumLightsUsingStandardDeferred);
 	SCOPED_CONDITIONAL_DRAW_EVENT(StandardDeferredLighting, bIssueDrawEvent, DEC_SCENE_ITEMS);
 
 	// Use additive blending for color
-	RHISetBlendState(TStaticBlendState<CW_RGBA, BO_Add, BF_One, BF_One, BO_Add, BF_One, BF_One>::GetRHI());
+	RHICmdList.SetBlendState(TStaticBlendState<CW_RGBA, BO_Add, BF_One, BF_One, BO_Add, BF_One, BF_One>::GetRHI());
 
 	bool bStencilDirty = false;
 	const FSphere LightBounds = LightSceneInfo->Proxy->GetBoundingSphere();
@@ -730,23 +730,20 @@ void FDeferredShadingSceneRenderer::RenderLight(const FLightSceneInfo* LightScen
 			bUseIESTexture = (LightSceneInfo->Proxy->GetIESTextureResource() != 0);
 		}
 
-		//@todo-rco: RHIPacketList
-		FRHICommandList& RHICmdList = FRHICommandList::GetNullRef();
-
 		// Set the device viewport for the view.
-		RHISetViewport(View.ViewRect.Min.X, View.ViewRect.Min.Y, 0.0f, View.ViewRect.Max.X, View.ViewRect.Max.Y, 1.0f);
+		RHICmdList.SetViewport(View.ViewRect.Min.X, View.ViewRect.Min.Y, 0.0f, View.ViewRect.Max.X, View.ViewRect.Max.Y, 1.0f);
 
 		if (LightSceneInfo->Proxy->GetLightType() == LightType_Directional)
 		{
 			TShaderMapRef<TDeferredLightVS<false> > VertexShader(GetGlobalShaderMap());
 
-			RHISetRasterizerState(TStaticRasterizerState<FM_Solid,CM_None>::GetRHI());
-			RHISetDepthStencilState(TStaticDepthStencilState<false,CF_Always>::GetRHI());
+			RHICmdList.SetRasterizerState(TStaticRasterizerState<FM_Solid, CM_None>::GetRHI());
+			RHICmdList.SetDepthStencilState(TStaticDepthStencilState<false, CF_Always>::GetRHI());
 
 			if (bRenderOverlap)
 			{
 				TShaderMapRef<TDeferredLightOverlapPS<false> > PixelShader(GetGlobalShaderMap());
-				SetGlobalBoundShaderState(PixelShader->GetBoundShaderState(), GetDeferredLightingVertexDeclaration<false>(), *VertexShader, *PixelShader);
+				SetGlobalBoundShaderState(RHICmdList, PixelShader->GetBoundShaderState(), GetDeferredLightingVertexDeclaration<false>(), *VertexShader, *PixelShader);
 				PixelShader->SetParameters(RHICmdList, View, LightSceneInfo);
 			}
 			else
@@ -765,6 +762,7 @@ void FDeferredShadingSceneRenderer::RenderLight(const FLightSceneInfo* LightScen
 
 			// Apply the directional light as a full screen quad
 			DrawRectangle( 
+				RHICmdList,
 				0, 0,
 				View.ViewRect.Width(), View.ViewRect.Height(),
 				View.ViewRect.Min.X, View.ViewRect.Min.Y, 
@@ -778,12 +776,12 @@ void FDeferredShadingSceneRenderer::RenderLight(const FLightSceneInfo* LightScen
 		{
 			TShaderMapRef<TDeferredLightVS<true> > VertexShader(GetGlobalShaderMap());
 
-			SetBoundingGeometryRasterizerAndDepthState(View, LightBounds);
+			SetBoundingGeometryRasterizerAndDepthState(RHICmdList, View, LightBounds);
 
 			if (bRenderOverlap)
 			{
 				TShaderMapRef<TDeferredLightOverlapPS<true> > PixelShader(GetGlobalShaderMap());
-				SetGlobalBoundShaderState(PixelShader->GetBoundShaderState(), GetDeferredLightingVertexDeclaration<true>(), *VertexShader, *PixelShader);
+				SetGlobalBoundShaderState(RHICmdList, PixelShader->GetBoundShaderState(), GetDeferredLightingVertexDeclaration<true>(), *VertexShader, *PixelShader);
 				PixelShader->SetParameters(RHICmdList, View, LightSceneInfo);
 			}
 			else
@@ -818,11 +816,11 @@ void FDeferredShadingSceneRenderer::RenderLight(const FLightSceneInfo* LightScen
 			{
 				// Apply the point or spot light with some approximately bounding geometry, 
 				// So we can get speedups from depth testing and not processing pixels outside of the light's influence.
-				StencilingGeometry::DrawSphere();
+				StencilingGeometry::DrawSphere(RHICmdList);
 			}
 			else if (LightSceneInfo->Proxy->GetLightType() == LightType_Spot)
 			{
-				StencilingGeometry::DrawCone();
+				StencilingGeometry::DrawCone(RHICmdList);
 			}
 		}
 	}
@@ -830,21 +828,18 @@ void FDeferredShadingSceneRenderer::RenderLight(const FLightSceneInfo* LightScen
 	if (bStencilDirty)
 	{
 		// Clear the stencil buffer to 0.
-		RHIClear(false,FColor(0,0,0),false,0,true,0, FIntRect());
+		RHICmdList.Clear(false, FColor(0, 0, 0), false, 0, true, 0, FIntRect());
 	}
 }
 
-void FDeferredShadingSceneRenderer::RenderSimpleLightsStandardDeferred(const FSimpleLightArray& SimpleLights)
+void FDeferredShadingSceneRenderer::RenderSimpleLightsStandardDeferred(FRHICommandList& RHICmdList, const FSimpleLightArray& SimpleLights)
 {
 	SCOPE_CYCLE_COUNTER(STAT_DirectLightRenderingTime);
 	INC_DWORD_STAT_BY(STAT_NumLightsUsingStandardDeferred, SimpleLights.InstanceData.Num());
 	SCOPED_DRAW_EVENT(StandardDeferredSimpleLights, DEC_SCENE_ITEMS);
 	
-	//@todo-rco: RHIPacketList
-	FRHICommandList& RHICmdList = FRHICommandList::GetNullRef();
-
 	// Use additive blending for color
-	RHISetBlendState(TStaticBlendState<CW_RGBA, BO_Add, BF_One, BF_One, BO_Add, BF_One, BF_One>::GetRHI());
+	RHICmdList.SetBlendState(TStaticBlendState<CW_RGBA, BO_Add, BF_One, BF_One, BO_Add, BF_One, BF_One>::GetRHI());
 
 	const int32 NumViews = Views.Num();
 	for (int32 LightIndex = 0; LightIndex < SimpleLights.InstanceData.Num(); LightIndex++)
@@ -858,11 +853,11 @@ void FDeferredShadingSceneRenderer::RenderSimpleLightsStandardDeferred(const FSi
 
 			FViewInfo& View = Views[ViewIndex];
 			// Set the device viewport for the view.
-			RHISetViewport(View.ViewRect.Min.X, View.ViewRect.Min.Y, 0.0f, View.ViewRect.Max.X, View.ViewRect.Max.Y, 1.0f);
+			RHICmdList.SetViewport(View.ViewRect.Min.X, View.ViewRect.Min.Y, 0.0f, View.ViewRect.Max.X, View.ViewRect.Max.Y, 1.0f);
 
 			TShaderMapRef<TDeferredLightVS<true> > VertexShader(GetGlobalShaderMap());
 
-			SetBoundingGeometryRasterizerAndDepthState(View, LightBounds);
+			SetBoundingGeometryRasterizerAndDepthState(RHICmdList, View, LightBounds);
 
 			if (SimpleLight.Exponent == 0)
 			{
@@ -879,7 +874,7 @@ void FDeferredShadingSceneRenderer::RenderSimpleLightsStandardDeferred(const FSi
 
 			// Apply the point or spot light with some approximately bounding geometry, 
 			// So we can get speedups from depth testing and not processing pixels outside of the light's influence.
-			StencilingGeometry::DrawSphere();
+			StencilingGeometry::DrawSphere(RHICmdList);
 		}
 	}
 }
diff --git a/Engine/Source/Runtime/Renderer/Private/LightRendering.h b/Engine/Source/Runtime/Renderer/Private/LightRendering.h
--- a/Engine/Source/Runtime/Renderer/Private/LightRendering.h
+++ b/Engine/Source/Runtime/Renderer/Private/LightRendering.h
@@ -201,6 +201,6 @@ private:
 	FStencilingGeometryShaderParameters StencilingGeometryParameters;
 };
 
-extern void SetBoundingGeometryRasterizerAndDepthState(const FViewInfo& View, const FSphere& LightBounds);
+extern void SetBoundingGeometryRasterizerAndDepthState(FRHICommandList& RHICmdList, const FViewInfo& View, const FSphere& LightBounds);
 
 #endif
\ No newline at end of file
diff --git a/Engine/Source/Runtime/Renderer/Private/LightShaftRendering.cpp b/Engine/Source/Runtime/Renderer/Private/LightShaftRendering.cpp
--- a/Engine/Source/Runtime/Renderer/Private/LightShaftRendering.cpp
+++ b/Engine/Source/Runtime/Renderer/Private/LightShaftRendering.cpp
@@ -440,15 +440,18 @@ void AllocateOrReuseLightShaftRenderTarget(TRefCountPtr<IPooledRenderTarget>& Ta
 		FIntPoint LightShaftSize(FMath::Max<uint32>(BufferSize.X / GetLightShaftDownsampleFactor(), 1), FMath::Max<uint32>(BufferSize.Y / GetLightShaftDownsampleFactor(), 1));
 		FPooledRenderTargetDesc Desc(FPooledRenderTargetDesc::Create2DDesc(LightShaftSize, LightShaftFilterBufferFormat, TexCreate_None, TexCreate_RenderTargetable, false));
 		GRenderTargetPool.FindFreeElement(Desc, Target, Name);
-		
-		RHISetRenderTarget(Target->GetRenderTargetItem().TargetableTexture, FTextureRHIRef());
-		RHIClear(true, FLinearColor(0, 0, 0, 0), false, 1.0f, false, 0, FIntRect());
+
+		//@todo-rco: RHIPacketList
+		FRHICommandList& RHICmdList = FRHICommandList::GetNullRef();
+
+		SetRenderTarget(RHICmdList, Target->GetRenderTargetItem().TargetableTexture, FTextureRHIRef());
+		RHICmdList.Clear(true, FLinearColor(0, 0, 0, 0), false, 1.0f, false, 0, FIntRect());
 	}
 }
 
 /** Generates the downsampled light shaft mask for either occlusion or bloom.  This swaps input and output before returning. */
 template<bool bDownsampleOcclusion>
-void DownsamplePass(const FViewInfo& View, const FLightSceneInfo* LightSceneInfo, TRefCountPtr<IPooledRenderTarget>& LightShaftsSource, TRefCountPtr<IPooledRenderTarget>& LightShaftsDest)
+void DownsamplePass(FRHICommandList& RHICmdList, const FViewInfo& View, const FLightSceneInfo* LightSceneInfo, TRefCountPtr<IPooledRenderTarget>& LightShaftsSource, TRefCountPtr<IPooledRenderTarget>& LightShaftsDest)
 {
 	SCOPED_DRAW_EVENT(Downsample, DEC_SCENE_ITEMS);
 
@@ -459,11 +462,8 @@ void DownsamplePass(const FViewInfo& View, const FLightSceneInfo* LightSceneInfo
 	const uint32 DownsampledSizeX = View.ViewRect.Width() / DownsampleFactor;
 	const uint32 DownsampledSizeY = View.ViewRect.Height() / DownsampleFactor;
 
-	//@todo-rco: RHIPacketList
-	FRHICommandList& RHICmdList = FRHICommandList::GetNullRef();
-
-	RHISetRenderTarget(LightShaftsDest->GetRenderTargetItem().TargetableTexture, FTextureRHIRef());
-	RHISetViewport(DownSampledXY.X, DownSampledXY.Y, 0.0f, DownSampledXY.X + DownsampledSizeX, DownSampledXY.Y + DownsampledSizeY, 1.0f);
+	SetRenderTarget(RHICmdList, LightShaftsDest->GetRenderTargetItem().TargetableTexture, FTextureRHIRef());
+	RHICmdList.SetViewport(DownSampledXY.X, DownSampledXY.Y, 0.0f, DownSampledXY.X + DownsampledSizeX, DownSampledXY.Y + DownsampledSizeY, 1.0f);
 
 	// Set shaders and texture
 	TShaderMapRef<FDownsampleLightShaftsVertexShader> DownsampleLightShaftsVertexShader(GetGlobalShaderMap());
@@ -474,14 +474,14 @@ void DownsamplePass(const FViewInfo& View, const FLightSceneInfo* LightSceneInfo
 	case LightType_Directional:
 		{
 			TShaderMapRef<TDownsampleLightShaftsPixelShader<LightType_Directional, bDownsampleOcclusion> > DownsampleLightShaftsPixelShader(GetGlobalShaderMap());
-			SetGlobalBoundShaderState(DownsampleDirectionalLightShaftsBoundShaderState[bDownsampleOcclusion], GFilterVertexDeclaration.VertexDeclarationRHI, *DownsampleLightShaftsVertexShader, *DownsampleLightShaftsPixelShader);
+			SetGlobalBoundShaderState(RHICmdList, DownsampleDirectionalLightShaftsBoundShaderState[bDownsampleOcclusion], GFilterVertexDeclaration.VertexDeclarationRHI, *DownsampleLightShaftsVertexShader, *DownsampleLightShaftsPixelShader);
 			DownsampleLightShaftsPixelShader->SetParameters(RHICmdList, LightSceneInfo, View, UnusedRT);
 		}
 		break;
 	case LightType_Spot:
 		{
 			TShaderMapRef<TDownsampleLightShaftsPixelShader<LightType_Spot, bDownsampleOcclusion> > DownsampleLightShaftsPixelShader(GetGlobalShaderMap());
-			SetGlobalBoundShaderState(DownsampleSpotLightShaftsBoundShaderState[bDownsampleOcclusion], GFilterVertexDeclaration.VertexDeclarationRHI, *DownsampleLightShaftsVertexShader, *DownsampleLightShaftsPixelShader);
+			SetGlobalBoundShaderState(RHICmdList, DownsampleSpotLightShaftsBoundShaderState[bDownsampleOcclusion], GFilterVertexDeclaration.VertexDeclarationRHI, *DownsampleLightShaftsVertexShader, *DownsampleLightShaftsPixelShader);
 			DownsampleLightShaftsPixelShader->SetParameters(RHICmdList, LightSceneInfo, View, UnusedRT);
 		}
 		break;
@@ -489,7 +489,7 @@ void DownsamplePass(const FViewInfo& View, const FLightSceneInfo* LightSceneInfo
 	case LightType_Point:
 		{
 			TShaderMapRef<TDownsampleLightShaftsPixelShader<LightType_Point, bDownsampleOcclusion> > DownsampleLightShaftsPixelShader(GetGlobalShaderMap());
-			SetGlobalBoundShaderState(DownsamplePointLightShaftsBoundShaderState[bDownsampleOcclusion], GFilterVertexDeclaration.VertexDeclarationRHI, *DownsampleLightShaftsVertexShader, *DownsampleLightShaftsPixelShader);
+			SetGlobalBoundShaderState(RHICmdList, DownsamplePointLightShaftsBoundShaderState[bDownsampleOcclusion], GFilterVertexDeclaration.VertexDeclarationRHI, *DownsampleLightShaftsVertexShader, *DownsampleLightShaftsPixelShader);
 			DownsampleLightShaftsPixelShader->SetParameters(RHICmdList, LightSceneInfo, View, UnusedRT);
 		}
 		break;
@@ -498,12 +498,13 @@ void DownsamplePass(const FViewInfo& View, const FLightSceneInfo* LightSceneInfo
 	DownsampleLightShaftsVertexShader->SetParameters(RHICmdList, View);
 
 	// No depth tests, no backface culling.
-	RHISetBlendState(TStaticBlendState<>::GetRHI());
-	RHISetRasterizerState(TStaticRasterizerState<FM_Solid,CM_None>::GetRHI());
-	RHISetDepthStencilState(TStaticDepthStencilState<false,CF_Always>::GetRHI());
+	RHICmdList.SetBlendState(TStaticBlendState<>::GetRHI());
+	RHICmdList.SetRasterizerState(TStaticRasterizerState<FM_Solid, CM_None>::GetRHI());
+	RHICmdList.SetDepthStencilState(TStaticDepthStencilState<false, CF_Always>::GetRHI());
 
 	// Downsample scene color and depth, and convert them into a bloom term and an occlusion masking term
 	DrawRectangle( 
+		RHICmdList,
 		0, 0, 
 		DownsampledSizeX, DownsampledSizeY,
 		View.ViewRect.Min.X, View.ViewRect.Min.Y, 
@@ -513,7 +514,7 @@ void DownsamplePass(const FViewInfo& View, const FLightSceneInfo* LightSceneInfo
 		*DownsampleLightShaftsVertexShader,
 		EDRF_UseTriangleOptimization);
 
-	RHICopyToResolveTarget(LightShaftsDest->GetRenderTargetItem().TargetableTexture, LightShaftsDest->GetRenderTargetItem().ShaderResourceTexture, false, FResolveParams() );
+	RHICmdList.CopyToResolveTarget(LightShaftsDest->GetRenderTargetItem().TargetableTexture, LightShaftsDest->GetRenderTargetItem().ShaderResourceTexture, false, FResolveParams());
 
 	Swap(LightShaftsSource, LightShaftsDest);
 }
@@ -587,6 +588,7 @@ void ApplyTemporalAA(
 
 /** Applies screen space radial blur passes. */
 void ApplyRadialBlurPasses(
+	FRHICommandList& RHICmdList,
 	const FViewInfo& View, 
 	const FLightSceneInfo* const LightSceneInfo, 
 	/** First pass source - this will not be overwritten. */
@@ -608,24 +610,25 @@ void ApplyRadialBlurPasses(
 
 	for (uint32 PassIndex = 0; PassIndex < NumPasses; PassIndex++)
 	{
-		RHISetRenderTarget(LightShaftsDest->GetRenderTargetItem().TargetableTexture, FTextureRHIRef());
-		RHISetViewport(0, 0, 0.0f, FilterBufferSize.X, FilterBufferSize.Y, 1.0f);
+		SetRenderTarget(RHICmdList, LightShaftsDest->GetRenderTargetItem().TargetableTexture, FTextureRHIRef());
+		RHICmdList.SetViewport(0, 0, 0.0f, FilterBufferSize.X, FilterBufferSize.Y, 1.0f);
 
-		RHISetBlendState(TStaticBlendState<>::GetRHI());
-		RHISetRasterizerState(TStaticRasterizerState<>::GetRHI());
-		RHISetDepthStencilState(TStaticDepthStencilState<false,CF_Always>::GetRHI());
+		RHICmdList.SetBlendState(TStaticBlendState<>::GetRHI());
+		RHICmdList.SetRasterizerState(TStaticRasterizerState<>::GetRHI());
+		RHICmdList.SetDepthStencilState(TStaticDepthStencilState<false, CF_Always>::GetRHI());
 
 		TShaderMapRef<FBlurLightShaftsPixelShader> BlurLightShaftsPixelShader(GetGlobalShaderMap());
-		SetGlobalBoundShaderState(BlurLightShaftsBoundShaderState, GFilterVertexDeclaration.VertexDeclarationRHI, *ScreenVertexShader, *BlurLightShaftsPixelShader);
+		SetGlobalBoundShaderState(RHICmdList, BlurLightShaftsBoundShaderState, GFilterVertexDeclaration.VertexDeclarationRHI, *ScreenVertexShader, *BlurLightShaftsPixelShader);
 
 		TRefCountPtr<IPooledRenderTarget>& EffectiveSource = PassIndex == 0 ? FirstPassSource : LightShaftsSource;
 		/// ?
-		BlurLightShaftsPixelShader->SetParameters(FRHICommandList::GetNullRef(), LightSceneInfo, View, PassIndex, EffectiveSource);
+		BlurLightShaftsPixelShader->SetParameters(RHICmdList, LightSceneInfo, View, PassIndex, EffectiveSource);
 
 		{
 			SCOPED_DRAW_EVENT(RadialBlur, DEC_SCENE_ITEMS);
 			// Apply a radial blur to the bloom and occlusion mask
 			DrawRectangle( 
+				RHICmdList,
 				DownSampledXY.X, DownSampledXY.Y, 
 				DownsampledSizeX, DownsampledSizeY,
 				DownSampledXY.X, DownSampledXY.Y, 
@@ -635,14 +638,14 @@ void ApplyRadialBlurPasses(
 				EDRF_UseTriangleOptimization);
 		}
 
-		RHICopyToResolveTarget(LightShaftsDest->GetRenderTargetItem().TargetableTexture, LightShaftsDest->GetRenderTargetItem().ShaderResourceTexture, false, FResolveParams());
+		RHICmdList.CopyToResolveTarget(LightShaftsDest->GetRenderTargetItem().TargetableTexture, LightShaftsDest->GetRenderTargetItem().ShaderResourceTexture, false, FResolveParams());
 
 		// Swap input and output for the next pass
 		Swap(LightShaftsSource, LightShaftsDest);
 	}
 }
 
-void FinishOcclusionTerm(const FViewInfo& View, const FLightSceneInfo* const LightSceneInfo, TRefCountPtr<IPooledRenderTarget>& LightShaftsSource, TRefCountPtr<IPooledRenderTarget>& LightShaftsDest)
+void FinishOcclusionTerm(FRHICommandList& RHICmdList, const FViewInfo& View, const FLightSceneInfo* const LightSceneInfo, TRefCountPtr<IPooledRenderTarget>& LightShaftsSource, TRefCountPtr<IPooledRenderTarget>& LightShaftsDest)
 {
 	TShaderMapRef<FScreenVS> ScreenVertexShader(GetGlobalShaderMap());
 
@@ -653,22 +656,23 @@ void FinishOcclusionTerm(const FViewInfo& View, const FLightSceneInfo* const Lig
 	const uint32 DownsampledSizeX = View.ViewRect.Width() / DownsampleFactor;
 	const uint32 DownsampledSizeY = View.ViewRect.Height() / DownsampleFactor;
 
-	RHISetRenderTarget(LightShaftsDest->GetRenderTargetItem().TargetableTexture, FTextureRHIRef());
-	RHISetViewport(0, 0, 0.0f, FilterBufferSize.X, FilterBufferSize.Y, 1.0f);
+	SetRenderTarget(RHICmdList, LightShaftsDest->GetRenderTargetItem().TargetableTexture, FTextureRHIRef());
+	RHICmdList.SetViewport(0, 0, 0.0f, FilterBufferSize.X, FilterBufferSize.Y, 1.0f);
 
-	RHISetBlendState(TStaticBlendState<>::GetRHI());
-	RHISetRasterizerState(TStaticRasterizerState<>::GetRHI());
-	RHISetDepthStencilState(TStaticDepthStencilState<false,CF_Always>::GetRHI());
+	RHICmdList.SetBlendState(TStaticBlendState<>::GetRHI());
+	RHICmdList.SetRasterizerState(TStaticRasterizerState<>::GetRHI());
+	RHICmdList.SetDepthStencilState(TStaticDepthStencilState<false, CF_Always>::GetRHI());
 
 	TShaderMapRef<FFinishOcclusionPixelShader> MaskOcclusionTermPixelShader(GetGlobalShaderMap());
-	SetGlobalBoundShaderState(AccumulateTermBoundShaderState, GFilterVertexDeclaration.VertexDeclarationRHI, *ScreenVertexShader, *MaskOcclusionTermPixelShader);
+	SetGlobalBoundShaderState(RHICmdList, AccumulateTermBoundShaderState, GFilterVertexDeclaration.VertexDeclarationRHI, *ScreenVertexShader, *MaskOcclusionTermPixelShader);
 	/// ?
-	MaskOcclusionTermPixelShader->SetParameters(FRHICommandList::GetNullRef(), LightSceneInfo, View, LightShaftsSource);
+	MaskOcclusionTermPixelShader->SetParameters(RHICmdList, LightSceneInfo, View, LightShaftsSource);
 
 	{
 		SCOPED_DRAW_EVENT(FinishOcclusion, DEC_SCENE_ITEMS);
 		// Apply a radial blur to the bloom and occlusion mask
 		DrawRectangle( 
+			RHICmdList,
 			DownSampledXY.X, DownSampledXY.Y, 
 			DownsampledSizeX, DownsampledSizeY,
 			DownSampledXY.X, DownSampledXY.Y, 
@@ -678,7 +682,7 @@ void FinishOcclusionTerm(const FViewInfo& View, const FLightSceneInfo* const Lig
 			EDRF_UseTriangleOptimization);
 	}
 
-	RHICopyToResolveTarget(LightShaftsDest->GetRenderTargetItem().TargetableTexture, LightShaftsDest->GetRenderTargetItem().ShaderResourceTexture, false, FResolveParams());
+	RHICmdList.CopyToResolveTarget(LightShaftsDest->GetRenderTargetItem().TargetableTexture, LightShaftsDest->GetRenderTargetItem().ShaderResourceTexture, false, FResolveParams());
 }
 
 bool DoesViewFamilyAllowLightShafts(const FSceneViewFamily& ViewFamily)
@@ -711,7 +715,7 @@ bool ShouldRenderLightShaftsForLight(const FViewInfo& View, const FLightSceneInf
 }
 
 /** Renders light shafts. */
-FLightShaftsOutput FDeferredShadingSceneRenderer::RenderLightShaftOcclusion()
+FLightShaftsOutput FDeferredShadingSceneRenderer::RenderLightShaftOcclusion(FRHICommandList& RHICmdList)
 {
 	FLightShaftsOutput Output;
 
@@ -732,6 +736,8 @@ FLightShaftsOutput FDeferredShadingSceneRenderer::RenderLightShaftOcclusion()
 			{
 				SCOPED_DRAW_EVENT(RenderLightShaftOcclusion, DEC_SCENE_ITEMS);
 
+				RHICmdList.CheckIsNull(); // resource allocation
+
 				// Allocate light shaft render targets on demand, using the pool
 				// Need two targets to ping pong between
 				AllocateOrReuseLightShaftRenderTarget(LightShafts0, TEXT("LightShafts0"));
@@ -746,7 +752,7 @@ FLightShaftsOutput FDeferredShadingSceneRenderer::RenderLightShaftOcclusion()
 						INC_DWORD_STAT(STAT_LightShaftsLights);
 
 						// Create a downsampled occlusion mask from scene depth, result will be in LightShafts0
-						DownsamplePass<true>(View, LightSceneInfo, LightShafts0, LightShafts1);
+						DownsamplePass<true>(RHICmdList, View, LightSceneInfo, LightShafts0, LightShafts1);
 
 						FSceneViewState* ViewState = (FSceneViewState*)View.State;
 						// Find the previous frame's occlusion mask
@@ -759,10 +765,10 @@ FLightShaftsOutput FDeferredShadingSceneRenderer::RenderLightShaftOcclusion()
 
 						// Apply radial blur passes
 						// Send HistoryOutput in as the first pass input only, so it will not be overwritten by any subsequent passes, since it is needed for next frame
-						ApplyRadialBlurPasses(View, LightSceneInfo, HistoryOutput, LightShafts0, LightShafts1);
+						ApplyRadialBlurPasses(RHICmdList, View, LightSceneInfo, HistoryOutput, LightShafts0, LightShafts1);
 
 						// Apply post-blur masking
-						FinishOcclusionTerm(View, LightSceneInfo, LightShafts0, LightShafts1);
+						FinishOcclusionTerm(RHICmdList, View, LightSceneInfo, LightShafts0, LightShafts1);
 
 						//@todo - different views could have different result render targets
 						Output.LightShaftOcclusion = LightShafts1;
@@ -831,24 +837,25 @@ private:
 
 IMPLEMENT_SHADER_TYPE(,FApplyLightShaftsPixelShader,TEXT("LightShaftShader"),TEXT("ApplyLightShaftsPixelMain"),SF_Pixel);
 
-void ApplyLightShaftBloom(const FViewInfo& View, const FLightSceneInfo* const LightSceneInfo, TRefCountPtr<IPooledRenderTarget>& LightShaftsSource)
+void ApplyLightShaftBloom(FRHICommandList& RHICmdList, const FViewInfo& View, const FLightSceneInfo* const LightSceneInfo, TRefCountPtr<IPooledRenderTarget>& LightShaftsSource)
 {
 	SCOPED_DRAW_EVENT(Apply, DEC_SCENE_ITEMS);
+	RHICmdList.CheckIsNull(); // GSceneRenderTargets ops
 
 	GSceneRenderTargets.BeginRenderingSceneColor();
 
-	RHISetViewport(View.ViewRect.Min.X, View.ViewRect.Min.Y, 0.0f, View.ViewRect.Max.X, View.ViewRect.Max.Y, 1.0f);
-	RHISetBlendState(TStaticBlendState<CW_RGB, BO_Add, BF_One, BF_One>::GetRHI());
-	RHISetRasterizerState(TStaticRasterizerState<>::GetRHI());
-	RHISetDepthStencilState(TStaticDepthStencilState<false,CF_Always>::GetRHI());
+	RHICmdList.SetViewport(View.ViewRect.Min.X, View.ViewRect.Min.Y, 0.0f, View.ViewRect.Max.X, View.ViewRect.Max.Y, 1.0f);
+	RHICmdList.SetBlendState(TStaticBlendState<CW_RGB, BO_Add, BF_One, BF_One>::GetRHI());
+	RHICmdList.SetRasterizerState(TStaticRasterizerState<>::GetRHI());
+	RHICmdList.SetDepthStencilState(TStaticDepthStencilState<false, CF_Always>::GetRHI());
 
 	TShaderMapRef<FScreenVS> ScreenVertexShader(GetGlobalShaderMap());
 	TShaderMapRef<FApplyLightShaftsPixelShader> ApplyLightShaftsPixelShader(GetGlobalShaderMap());
 
-	SetGlobalBoundShaderState(ApplyLightShaftsBoundShaderState, GFilterVertexDeclaration.VertexDeclarationRHI, *ScreenVertexShader, *ApplyLightShaftsPixelShader);
+	SetGlobalBoundShaderState(RHICmdList, ApplyLightShaftsBoundShaderState, GFilterVertexDeclaration.VertexDeclarationRHI, *ScreenVertexShader, *ApplyLightShaftsPixelShader);
 
 	/// ?
-	ApplyLightShaftsPixelShader->SetParameters(FRHICommandList::GetNullRef(), View, LightShaftsSource);
+	ApplyLightShaftsPixelShader->SetParameters(RHICmdList, View, LightShaftsSource);
 
 	const FIntPoint BufferSize = GSceneRenderTargets.GetBufferSizeXY();
 	const uint32 DownsampleFactor	= GetLightShaftDownsampleFactor();
@@ -858,6 +865,7 @@ void ApplyLightShaftBloom(const FViewInfo& View, const FLightSceneInfo* const Li
 	const uint32 DownsampledSizeY = View.ViewRect.Height() / DownsampleFactor;
 
 	DrawRectangle( 
+		RHICmdList,
 		0, 0, 
 		View.ViewRect.Width(), View.ViewRect.Height(),
 		DownSampledXY.X, DownSampledXY.Y, 
@@ -869,8 +877,10 @@ void ApplyLightShaftBloom(const FViewInfo& View, const FLightSceneInfo* const Li
 	GSceneRenderTargets.FinishRenderingSceneColor(false);
 }
 
-void FDeferredShadingSceneRenderer::RenderLightShaftBloom()
+void FDeferredShadingSceneRenderer::RenderLightShaftBloom(FRHICommandList& RHICmdList)
 {
+	RHICmdList.CheckIsNull(); // resource allocation
+
 	if (DoesViewFamilyAllowLightShafts(ViewFamily))
 	{
 		TRefCountPtr<IPooledRenderTarget> LightShafts0;
@@ -897,7 +907,7 @@ void FDeferredShadingSceneRenderer::RenderLightShaftBloom()
 						INC_DWORD_STAT(STAT_LightShaftsLights);
 
 						// Generate the bloom source from scene color, masked by depth and downsampled
-						DownsamplePass<false>(View, LightSceneInfo, LightShafts0, LightShafts1);
+						DownsamplePass<false>(RHICmdList, View, LightSceneInfo, LightShafts0, LightShafts1);
 
 						FSceneViewState* ViewState = (FSceneViewState*)View.State;
 						TRefCountPtr<IPooledRenderTarget>* HistoryState = NULL;
@@ -916,10 +926,10 @@ void FDeferredShadingSceneRenderer::RenderLightShaftBloom()
 
 						// Apply radial blur passes
 						// Send HistoryOutput in as the first pass input only, so it will not be overwritten by any subsequent passes, since it is needed for next frame
-						ApplyRadialBlurPasses(View, LightSceneInfo, HistoryOutput, LightShafts0, LightShafts1);
+						ApplyRadialBlurPasses(RHICmdList, View, LightSceneInfo, HistoryOutput, LightShafts0, LightShafts1);
 						
 						// Add light shaft bloom to scene color in full res
-						ApplyLightShaftBloom(View, LightSceneInfo, LightShafts0);
+						ApplyLightShaftBloom(RHICmdList, View, LightSceneInfo, LightShafts0);
 					}
 				}
 			}
diff --git a/Engine/Source/Runtime/Renderer/Private/PostProcess/PostProcessAA.cpp b/Engine/Source/Runtime/Renderer/Private/PostProcess/PostProcessAA.cpp
--- a/Engine/Source/Runtime/Renderer/Private/PostProcess/PostProcessAA.cpp
+++ b/Engine/Source/Runtime/Renderer/Private/PostProcess/PostProcessAA.cpp
@@ -70,15 +70,15 @@ public:
 		return bShaderHasOutdatedParameters;
 	}
 
-	void SetParameters(FRHICommandList& RHICmdList, const FRenderingCompositePassContext& Context)
+	void SetParameters(const FRenderingCompositePassContext& Context)
 	{
 		const FPixelShaderRHIParamRef ShaderRHI = GetPixelShader();
 
 		const FSceneView& View = Context.View;
 
-		FGlobalShader::SetParameters(RHICmdList, ShaderRHI, View);
+		FGlobalShader::SetParameters(Context.RHICmdList, ShaderRHI, View);
 		
-		PostprocessParameter.SetPS(RHICmdList, ShaderRHI, Context, TStaticSamplerState<SF_Bilinear,AM_Clamp,AM_Clamp,AM_Clamp>::GetRHI());
+		PostprocessParameter.SetPS(ShaderRHI, Context, TStaticSamplerState<SF_Bilinear,AM_Clamp,AM_Clamp,AM_Clamp>::GetRHI());
 
 		const FPooledRenderTargetDesc* InputDesc = Context.Pass->GetInputDesc(ePId_Input0);
 
@@ -90,48 +90,48 @@ public:
 
 		FVector2D InvExtent = FVector2D(1.0f / InputDesc->Extent.X, 1.0f / InputDesc->Extent.Y);
 
-		SetShaderValue(RHICmdList, ShaderRHI, fxaaQualityRcpFrame, InvExtent);
+		SetShaderValue(Context.RHICmdList, ShaderRHI, fxaaQualityRcpFrame, InvExtent);
 
 		{
 			float N = 0.5f;
 			FVector4 Value(-N * InvExtent.X, -N * InvExtent.Y, N * InvExtent.X, N * InvExtent.Y);
-			SetShaderValue(RHICmdList, ShaderRHI, fxaaConsoleRcpFrameOpt, Value);
+			SetShaderValue(Context.RHICmdList, ShaderRHI, fxaaConsoleRcpFrameOpt, Value);
 		}
 
 		{
 			float N = 2.0f;
 			FVector4 Value(-N * InvExtent.X, -N * InvExtent.Y, N * InvExtent.X, N * InvExtent.Y);
-			SetShaderValue(RHICmdList, ShaderRHI, fxaaConsoleRcpFrameOpt2, Value);
+			SetShaderValue(Context.RHICmdList, ShaderRHI, fxaaConsoleRcpFrameOpt2, Value);
 		}
 
 		{
 			float Value = 0.75f;
-			SetShaderValue(RHICmdList, ShaderRHI, fxaaQualitySubpix, Value);
+			SetShaderValue(Context.RHICmdList, ShaderRHI, fxaaQualitySubpix, Value);
 		}
 
 		{
 			float Value = 0.166f;
-			SetShaderValue(RHICmdList, ShaderRHI, fxaaQualityEdgeThreshold, Value);
+			SetShaderValue(Context.RHICmdList, ShaderRHI, fxaaQualityEdgeThreshold, Value);
 		}
 
 		{
 			float Value = 0.0833f;
-			SetShaderValue(RHICmdList, ShaderRHI, fxaaQualityEdgeThresholdMin, Value);
+			SetShaderValue(Context.RHICmdList, ShaderRHI, fxaaQualityEdgeThresholdMin, Value);
 		}
 
 		{
 			float Value = 8.0f;
-			SetShaderValue(RHICmdList, ShaderRHI, fxaaConsoleEdgeSharpness, Value);
+			SetShaderValue(Context.RHICmdList, ShaderRHI, fxaaConsoleEdgeSharpness, Value);
 		}
 
 		{
 			float Value = 0.125f;
-			SetShaderValue(RHICmdList, ShaderRHI, fxaaConsoleEdgeThreshold, Value);
+			SetShaderValue(Context.RHICmdList, ShaderRHI, fxaaConsoleEdgeThreshold, Value);
 		}
 
 		{
 			float Value = 0.05f;
-			SetShaderValue(RHICmdList, ShaderRHI, fxaaConsoleEdgeThresholdMin, Value);
+			SetShaderValue(Context.RHICmdList, ShaderRHI, fxaaConsoleEdgeThresholdMin, Value);
 		}
 	}
 	
@@ -186,11 +186,11 @@ public:
 		return bShaderHasOutdatedParameters;
 	}
 
-	void SetParameters(FRHICommandList& RHICmdList, const FRenderingCompositePassContext& Context)
+	void SetParameters(const FRenderingCompositePassContext& Context)
 	{
 		const FVertexShaderRHIParamRef ShaderRHI = GetVertexShader();
 
-		FGlobalShader::SetParameters(RHICmdList, ShaderRHI, Context.View);
+		FGlobalShader::SetParameters(Context.RHICmdList, ShaderRHI, Context.View);
 
 		const FPooledRenderTargetDesc* InputDesc = Context.Pass->GetInputDesc(ePId_Input0);
 
@@ -202,7 +202,7 @@ public:
 
 		FVector2D InvExtent = FVector2D(1.0f / InputDesc->Extent.X, 1.0f / InputDesc->Extent.Y);
 
-		SetShaderValue(RHICmdList, ShaderRHI, fxaaQualityRcpFrame, InvExtent);
+		SetShaderValue(Context.RHICmdList, ShaderRHI, fxaaQualityRcpFrame, InvExtent);
 	}
 
 	FShaderParameter fxaaQualityRcpFrame;
@@ -213,17 +213,18 @@ IMPLEMENT_SHADER_TYPE(,FFXAAVS,TEXT("FXAAShader"),TEXT("FxaaVS"),SF_Vertex);
 
 // @param Quality 1..6
 template <uint32 Quality>
-static void SetShaderTemplAA(FRHICommandList& RHICmdList, const FRenderingCompositePassContext& Context)
+static void SetShaderTemplAA(const FRenderingCompositePassContext& Context)
 {
 	TShaderMapRef<FFXAAVS> VertexShader(GetGlobalShaderMap());
 	TShaderMapRef<FPostProcessAntiAliasingPS<Quality> > PixelShader(GetGlobalShaderMap());
 
 	static FGlobalBoundShaderState BoundShaderState;
+	Context.RHICmdList.CheckIsNull(); // need new approach for "static FGlobalBoundShaderState" for parallel rendering
 
-	SetGlobalBoundShaderState(BoundShaderState, GFilterVertexDeclaration.VertexDeclarationRHI, *VertexShader, *PixelShader);
+	SetGlobalBoundShaderState(Context.RHICmdList, BoundShaderState, GFilterVertexDeclaration.VertexDeclarationRHI, *VertexShader, *PixelShader);
 
-	PixelShader->SetParameters(RHICmdList, Context);
-	VertexShader->SetParameters(RHICmdList, Context);
+	PixelShader->SetParameters(Context);
+	VertexShader->SetParameters(Context);
 }
 
 void FRCPassPostProcessAA::Process(FRenderingCompositePassContext& Context)
@@ -248,33 +249,32 @@ void FRCPassPostProcessAA::Process(FRenderingCompositePassContext& Context)
 
 	const FSceneRenderTargetItem& DestRenderTarget = PassOutputs[0].RequestSurface(Context);
 
-	//@todo-rco: RHIPacketList
-	FRHICommandList& RHICmdList = FRHICommandList::GetNullRef();
-	RHICmdList.CheckIsNull();
+	Context.RHICmdList.CheckIsNull(); // need new approach for "static FGlobalBoundShaderState" for parallel rendering
 
 	// Set the view family's render target/viewport.
-	RHISetRenderTarget(DestRenderTarget.TargetableTexture, FTextureRHIRef());	
+	SetRenderTarget(Context.RHICmdList, DestRenderTarget.TargetableTexture, FTextureRHIRef());
 	Context.SetViewportAndCallRHI(0, 0, 0.0f, DestSize.X, DestSize.Y, 1.0f );
 
 	// set the state
-	RHISetBlendState(TStaticBlendState<>::GetRHI());
-	RHISetRasterizerState(TStaticRasterizerState<>::GetRHI());
-	RHISetDepthStencilState(TStaticDepthStencilState<false,CF_Always>::GetRHI());
+	Context.RHICmdList.SetBlendState(TStaticBlendState<>::GetRHI());
+	Context.RHICmdList.SetRasterizerState(TStaticRasterizerState<>::GetRHI());
+	Context.RHICmdList.SetDepthStencilState(TStaticDepthStencilState<false, CF_Always>::GetRHI());
 
 	switch(Quality)
 	{
-		case 1:  SetShaderTemplAA<1>(RHICmdList, Context); break;
-		case 2:  SetShaderTemplAA<2>(RHICmdList, Context); break;
-		case 3:  SetShaderTemplAA<3>(RHICmdList, Context); break;
-		case 4:  SetShaderTemplAA<4>(RHICmdList, Context); break;
-		case 5:  SetShaderTemplAA<5>(RHICmdList, Context); break;
-		default: SetShaderTemplAA<6>(RHICmdList, Context); break;
+		case 1:  SetShaderTemplAA<1>(Context); break;
+		case 2:  SetShaderTemplAA<2>(Context); break;
+		case 3:  SetShaderTemplAA<3>(Context); break;
+		case 4:  SetShaderTemplAA<4>(Context); break;
+		case 5:  SetShaderTemplAA<5>(Context); break;
+		default: SetShaderTemplAA<6>(Context); break;
 	}
 	
 	// Draw a quad mapping scene color to the view's render target
 	TShaderMapRef<FFXAAVS> VertexShader(GetGlobalShaderMap());
 
 	DrawRectangle(
+		Context.RHICmdList,
 		DestRect.Min.X, DestRect.Min.Y,
 		DestRect.Width(), DestRect.Height(),
 		SrcRect.Min.X, SrcRect.Min.Y,
@@ -284,7 +284,7 @@ void FRCPassPostProcessAA::Process(FRenderingCompositePassContext& Context)
 		*VertexShader,
 		EDRF_Default);
 
-	RHICopyToResolveTarget(DestRenderTarget.TargetableTexture, DestRenderTarget.ShaderResourceTexture, false, FResolveParams());
+	Context.RHICmdList.CopyToResolveTarget(DestRenderTarget.TargetableTexture, DestRenderTarget.ShaderResourceTexture, false, FResolveParams());
 }
 
 FPooledRenderTargetDesc FRCPassPostProcessAA::ComputeOutputDesc(EPassOutputId InPassOutputId) const
diff --git a/Engine/Source/Runtime/Renderer/Private/PostProcess/PostProcessBloomSetup.cpp b/Engine/Source/Runtime/Renderer/Private/PostProcess/PostProcessBloomSetup.cpp
--- a/Engine/Source/Runtime/Renderer/Private/PostProcess/PostProcessBloomSetup.cpp
+++ b/Engine/Source/Runtime/Renderer/Private/PostProcess/PostProcessBloomSetup.cpp
@@ -48,20 +48,20 @@ public:
 		BloomThreshold.Bind(Initializer.ParameterMap, TEXT("BloomThreshold"));
 	}
 
-	void SetPS(FRHICommandList& RHICmdList, const FRenderingCompositePassContext& Context)
+	void SetPS(const FRenderingCompositePassContext& Context)
 	{
 		const FPixelShaderRHIParamRef ShaderRHI = GetPixelShader();
 
 		const FPostProcessSettings& Settings = Context.View.FinalPostProcessSettings;
 
-		FGlobalShader::SetParameters(RHICmdList, ShaderRHI, Context.View);
+		FGlobalShader::SetParameters(Context.RHICmdList, ShaderRHI, Context.View);
 
-		PostprocessParameter.SetPS(RHICmdList, ShaderRHI, Context, TStaticSamplerState<SF_Bilinear,AM_Clamp,AM_Clamp,AM_Clamp>::GetRHI());
+		PostprocessParameter.SetPS(ShaderRHI, Context, TStaticSamplerState<SF_Bilinear,AM_Clamp,AM_Clamp,AM_Clamp>::GetRHI());
 
 		float ExposureScale = FRCPassPostProcessEyeAdaptation::ComputeExposureScaleValue(Context.View);
 
 		FVector4 BloomThresholdValue(Settings.BloomThreshold, 0, 0, ExposureScale);
-		SetShaderValue(RHICmdList, ShaderRHI, BloomThreshold, BloomThresholdValue);
+		SetShaderValue(Context.RHICmdList, ShaderRHI, BloomThreshold, BloomThresholdValue);
 	}
 	
 	// FShader interface.
@@ -102,13 +102,13 @@ public:
 		EyeAdaptation.Bind(Initializer.ParameterMap, TEXT("EyeAdaptation"));
 	}
 	
-	void SetVS(FRHICommandList& RHICmdList, const FRenderingCompositePassContext& Context)
+	void SetVS(const FRenderingCompositePassContext& Context)
 	{
 		const FVertexShaderRHIParamRef ShaderRHI = GetVertexShader();
 
-		FGlobalShader::SetParameters(RHICmdList, ShaderRHI, Context.View);
+		FGlobalShader::SetParameters(Context.RHICmdList, ShaderRHI, Context.View);
 
-		PostprocessParameter.SetVS(RHICmdList, ShaderRHI, Context, TStaticSamplerState<SF_Bilinear,AM_Clamp,AM_Clamp,AM_Clamp>::GetRHI());
+		PostprocessParameter.SetVS(ShaderRHI, Context, TStaticSamplerState<SF_Bilinear,AM_Clamp,AM_Clamp,AM_Clamp>::GetRHI());
 
 		if(EyeAdaptation.IsBound())
 		{
@@ -116,11 +116,11 @@ public:
 
 			if(EyeAdaptationRT)
 			{
-				SetTextureParameter(RHICmdList, ShaderRHI, EyeAdaptation, EyeAdaptationRT->GetRenderTargetItem().TargetableTexture);
+				SetTextureParameter(Context.RHICmdList, ShaderRHI, EyeAdaptation, EyeAdaptationRT->GetRenderTargetItem().TargetableTexture);
 			}
 			else
 			{
-				SetTextureParameter(RHICmdList, ShaderRHI, EyeAdaptation, GWhiteTexture->TextureRHI);
+				SetTextureParameter(Context.RHICmdList, ShaderRHI, EyeAdaptation, GWhiteTexture->TextureRHI);
 			}
 		}
 	}
@@ -162,34 +162,33 @@ void FRCPassPostProcessBloomSetup::Process(FRenderingCompositePassContext& Conte
 
 	const FSceneRenderTargetItem& DestRenderTarget = PassOutputs[0].RequestSurface(Context);
 
-	//@todo-rco: RHIPacketList
-	FRHICommandList& RHICmdList = FRHICommandList::GetNullRef();
-
 	// Set the view family's render target/viewport.
-	RHISetRenderTarget(DestRenderTarget.TargetableTexture, FTextureRHIRef());	
+	SetRenderTarget(Context.RHICmdList, DestRenderTarget.TargetableTexture, FTextureRHIRef());
 
 	// is optimized away if possible (RT size=view size, )
-	RHIClear(true, FLinearColor::Black, false, 1.0f, false, 0, DestRect);
+	Context.RHICmdList.Clear(true, FLinearColor::Black, false, 1.0f, false, 0, DestRect);
 
 	Context.SetViewportAndCallRHI(0, 0, 0.0f, DestSize.X, DestSize.Y, 1.0f );
 
 	// set the state
-	RHISetBlendState(TStaticBlendState<>::GetRHI());
-	RHISetRasterizerState(TStaticRasterizerState<>::GetRHI());
-	RHISetDepthStencilState(TStaticDepthStencilState<false,CF_Always>::GetRHI());
+	Context.RHICmdList.SetBlendState(TStaticBlendState<>::GetRHI());
+	Context.RHICmdList.SetRasterizerState(TStaticRasterizerState<>::GetRHI());
+	Context.RHICmdList.SetDepthStencilState(TStaticDepthStencilState<false, CF_Always>::GetRHI());
 
 	TShaderMapRef<FPostProcessBloomSetupVS> VertexShader(GetGlobalShaderMap());
 	TShaderMapRef<FPostProcessBloomSetupPS> PixelShader(GetGlobalShaderMap());
 
 	static FGlobalBoundShaderState BoundShaderState;
+	Context.RHICmdList.CheckIsNull(); // need new approach for "static FGlobalBoundShaderState" for parallel rendering
 
-	SetGlobalBoundShaderState(BoundShaderState, GFilterVertexDeclaration.VertexDeclarationRHI, *VertexShader, *PixelShader);
+	SetGlobalBoundShaderState(Context.RHICmdList, BoundShaderState, GFilterVertexDeclaration.VertexDeclarationRHI, *VertexShader, *PixelShader);
 
-	VertexShader->SetVS(RHICmdList, Context);
-	PixelShader->SetPS(RHICmdList, Context);
+	VertexShader->SetVS(Context);
+	PixelShader->SetPS(Context);
 
 	// Draw a quad mapping scene color to the view's render target
 	DrawRectangle(
+		Context.RHICmdList,
 		DestRect.Min.X, DestRect.Min.Y,
 		DestRect.Width(), DestRect.Height(),
 		SrcRect.Min.X, SrcRect.Min.Y,
@@ -199,7 +198,7 @@ void FRCPassPostProcessBloomSetup::Process(FRenderingCompositePassContext& Conte
 		*VertexShader,
 		EDRF_UseTriangleOptimization);
 
-	RHICopyToResolveTarget(DestRenderTarget.TargetableTexture, DestRenderTarget.ShaderResourceTexture, false, FResolveParams());
+	Context.RHICmdList.CopyToResolveTarget(DestRenderTarget.TargetableTexture, DestRenderTarget.ShaderResourceTexture, false, FResolveParams());
 }
 
 FPooledRenderTargetDesc FRCPassPostProcessBloomSetup::ComputeOutputDesc(EPassOutputId InPassOutputId) const
diff --git a/Engine/Source/Runtime/Renderer/Private/PostProcess/PostProcessBokehDOF.cpp b/Engine/Source/Runtime/Renderer/Private/PostProcess/PostProcessBokehDOF.cpp
--- a/Engine/Source/Runtime/Renderer/Private/PostProcess/PostProcessBokehDOF.cpp
+++ b/Engine/Source/Runtime/Renderer/Private/PostProcess/PostProcessBokehDOF.cpp
@@ -107,15 +107,15 @@ public:
 		return bShaderHasOutdatedParameters;
 	}
 
-	void SetParameters(FRHICommandList& RHICmdList, const FRenderingCompositePassContext& Context)
+	void SetParameters(const FRenderingCompositePassContext& Context)
 	{
 		const FPixelShaderRHIParamRef ShaderRHI = GetPixelShader();
 
-		FGlobalShader::SetParameters(RHICmdList, ShaderRHI, Context.View);
+		FGlobalShader::SetParameters(Context.RHICmdList, ShaderRHI, Context.View);
 
-		DeferredParameters.Set(RHICmdList, ShaderRHI, Context.View);
+		DeferredParameters.Set(Context.RHICmdList, ShaderRHI, Context.View);
 
-		PostprocessParameter.SetPS(RHICmdList, ShaderRHI, Context, TStaticSamplerState<SF_Point,AM_Clamp,AM_Clamp,AM_Clamp>::GetRHI());
+		PostprocessParameter.SetPS(ShaderRHI, Context, TStaticSamplerState<SF_Point,AM_Clamp,AM_Clamp,AM_Clamp>::GetRHI());
 
 		{
 			FVector4 DepthOfFieldParamValues[2];
@@ -125,7 +125,7 @@ public:
 
 			FRCPassPostProcessBokehDOF::ComputeDepthOfFieldParams(Context, DepthOfFieldParamValues);
 
-			SetShaderValueArray(RHICmdList, ShaderRHI, DepthOfFieldParams, DepthOfFieldParamValues, 2);
+			SetShaderValueArray(Context.RHICmdList, ShaderRHI, DepthOfFieldParams, DepthOfFieldParamValues, 2);
 		}
 	}
 
@@ -168,23 +168,20 @@ void FRCPassPostProcessVisualizeDOF::Process(FRenderingCompositePassContext& Con
 
 	const FSceneRenderTargetItem& DestRenderTarget = PassOutputs[0].RequestSurface(Context);
 
-	//@todo-rco: RHIPacketList
-	FRHICommandList& RHICmdList = FRHICommandList::GetNullRef();
-
 	// Set the view family's render target/viewport.
-	RHISetRenderTarget(DestRenderTarget.TargetableTexture, FTextureRHIRef());	
+	SetRenderTarget(Context.RHICmdList, DestRenderTarget.TargetableTexture, FTextureRHIRef());
 
 	// can be optimized (don't clear areas we overwrite, don't clear when full screen),
 	// needed when a camera (matinee) has black borders or with multiple viewports
 	// focal distance depth is stored in the alpha channel to avoid DOF artifacts
-	RHIClear(true, FLinearColor(0, 0, 0, View.FinalPostProcessSettings.DepthOfFieldFocalDistance), false, 0.0f, false, 0, DestRect);
+	Context.RHICmdList.Clear(true, FLinearColor(0, 0, 0, View.FinalPostProcessSettings.DepthOfFieldFocalDistance), false, 0.0f, false, 0, DestRect);
 
 	Context.SetViewportAndCallRHI(0, 0, 0.0f, DestSize.X, DestSize.Y, 1.0f );
 
 	// set the state
-	RHISetBlendState(TStaticBlendState<>::GetRHI());
-	RHISetRasterizerState(TStaticRasterizerState<>::GetRHI());
-	RHISetDepthStencilState(TStaticDepthStencilState<false,CF_Always>::GetRHI());
+	Context.RHICmdList.SetBlendState(TStaticBlendState<>::GetRHI());
+	Context.RHICmdList.SetRasterizerState(TStaticRasterizerState<>::GetRHI());
+	Context.RHICmdList.SetDepthStencilState(TStaticDepthStencilState<false, CF_Always>::GetRHI());
 
 	// setup shader
 	TShaderMapRef<FPostProcessVS> VertexShader(GetGlobalShaderMap());
@@ -194,14 +191,15 @@ void FRCPassPostProcessVisualizeDOF::Process(FRenderingCompositePassContext& Con
 
 		static FGlobalBoundShaderState BoundShaderState;
 
-		SetGlobalBoundShaderState(BoundShaderState, GFilterVertexDeclaration.VertexDeclarationRHI, *VertexShader, *PixelShader);
+		SetGlobalBoundShaderState(Context.RHICmdList, BoundShaderState, GFilterVertexDeclaration.VertexDeclarationRHI, *VertexShader, *PixelShader);
 
-		VertexShader->SetParameters(RHICmdList, Context);
-		PixelShader->SetParameters(RHICmdList, Context);
+		VertexShader->SetParameters(Context);
+		PixelShader->SetParameters(Context);
 	}
 
 	// Draw a quad mapping scene color to the view's render target
 	DrawRectangle(
+		Context.RHICmdList,
 		DestRect.Min.X, DestRect.Min.Y,
 		DestRect.Width(), DestRect.Height(),
 		SrcRect.Min.X, SrcRect.Min.Y,
@@ -211,7 +209,7 @@ void FRCPassPostProcessVisualizeDOF::Process(FRenderingCompositePassContext& Con
 		*VertexShader,
 		EDRF_UseTriangleOptimization);
 
-	RHICopyToResolveTarget(DestRenderTarget.TargetableTexture, DestRenderTarget.ShaderResourceTexture, false, FResolveParams());
+	Context.RHICmdList.CopyToResolveTarget(DestRenderTarget.TargetableTexture, DestRenderTarget.ShaderResourceTexture, false, FResolveParams());
 }
 
 FPooledRenderTargetDesc FRCPassPostProcessVisualizeDOF::ComputeOutputDesc(EPassOutputId InPassOutputId) const
@@ -270,14 +268,14 @@ public:
 		return bShaderHasOutdatedParameters;
 	}
 
-	void SetParameters(FRHICommandList& RHICmdList, const FRenderingCompositePassContext& Context)
+	void SetParameters(const FRenderingCompositePassContext& Context)
 	{
 		const FPixelShaderRHIParamRef ShaderRHI = GetPixelShader();
 
-		FGlobalShader::SetParameters(RHICmdList, ShaderRHI, Context.View);
+		FGlobalShader::SetParameters(Context.RHICmdList, ShaderRHI, Context.View);
 		
-		DeferredParameters.Set(RHICmdList, ShaderRHI, Context.View);
-		PostprocessParameter.SetPS(RHICmdList, ShaderRHI, Context, TStaticSamplerState<SF_Point,AM_Clamp,AM_Clamp,AM_Clamp>::GetRHI());
+		DeferredParameters.Set(Context.RHICmdList, ShaderRHI, Context.View);
+		PostprocessParameter.SetPS(ShaderRHI, Context, TStaticSamplerState<SF_Point,AM_Clamp,AM_Clamp,AM_Clamp>::GetRHI());
 
 		{
 			FVector4 DepthOfFieldParamValues[2];
@@ -287,7 +285,7 @@ public:
 
 			FRCPassPostProcessBokehDOF::ComputeDepthOfFieldParams(Context, DepthOfFieldParamValues);
 
-			SetShaderValueArray(RHICmdList, ShaderRHI, DepthOfFieldParams, DepthOfFieldParamValues, 2);
+			SetShaderValueArray(Context.RHICmdList, ShaderRHI, DepthOfFieldParams, DepthOfFieldParamValues, 2);
 		}
 	}
 
@@ -330,23 +328,20 @@ void FRCPassPostProcessBokehDOFSetup::Process(FRenderingCompositePassContext& Co
 
 	const FSceneRenderTargetItem& DestRenderTarget = PassOutputs[0].RequestSurface(Context);
 
-	//@todo-rco: RHIPacketList
-	FRHICommandList& RHICmdList = FRHICommandList::GetNullRef();
-
 	// Set the view family's render target/viewport.
-	RHISetRenderTarget(DestRenderTarget.TargetableTexture, FTextureRHIRef());	
+	SetRenderTarget(Context.RHICmdList, DestRenderTarget.TargetableTexture, FTextureRHIRef());
 	
 	// can be optimized (don't clear areas we overwrite, don't clear when full screen),
 	// needed when a camera (matinee) has black borders or with multiple viewports
 	// focal distance depth is stored in the alpha channel to avoid DOF artifacts
-	RHIClear(true, FLinearColor(0, 0, 0, View.FinalPostProcessSettings.DepthOfFieldFocalDistance), false, 0.0f, false, 0, DestRect);
+	Context.RHICmdList.Clear(true, FLinearColor(0, 0, 0, View.FinalPostProcessSettings.DepthOfFieldFocalDistance), false, 0.0f, false, 0, DestRect);
 
 	Context.SetViewportAndCallRHI(0, 0, 0.0f, DestSize.X, DestSize.Y, 1.0f );
 
 	// set the state
-	RHISetBlendState(TStaticBlendState<>::GetRHI());
-	RHISetRasterizerState(TStaticRasterizerState<>::GetRHI());
-	RHISetDepthStencilState(TStaticDepthStencilState<false,CF_Always>::GetRHI());
+	Context.RHICmdList.SetBlendState(TStaticBlendState<>::GetRHI());
+	Context.RHICmdList.SetRasterizerState(TStaticRasterizerState<>::GetRHI());
+	Context.RHICmdList.SetDepthStencilState(TStaticDepthStencilState<false, CF_Always>::GetRHI());
 
 	// setup shader
 	TShaderMapRef<FPostProcessVS> VertexShader(GetGlobalShaderMap());
@@ -355,14 +350,15 @@ void FRCPassPostProcessBokehDOFSetup::Process(FRenderingCompositePassContext& Co
 
 		static FGlobalBoundShaderState BoundShaderState;
 
-		SetGlobalBoundShaderState(BoundShaderState, GFilterVertexDeclaration.VertexDeclarationRHI, *VertexShader, *PixelShader);
+		SetGlobalBoundShaderState(Context.RHICmdList, BoundShaderState, GFilterVertexDeclaration.VertexDeclarationRHI, *VertexShader, *PixelShader);
 
-		VertexShader->SetParameters(RHICmdList, Context);
-		PixelShader->SetParameters(RHICmdList, Context);
+		VertexShader->SetParameters(Context);
+		PixelShader->SetParameters(Context);
 	}
 
 	// Draw a quad mapping scene color to the view's render target
 	DrawRectangle(
+		Context.RHICmdList,
 		DestRect.Min.X, DestRect.Min.Y,
 		DestRect.Width(), DestRect.Height(),
 		SrcRect.Min.X, SrcRect.Min.Y,
@@ -372,7 +368,7 @@ void FRCPassPostProcessBokehDOFSetup::Process(FRenderingCompositePassContext& Co
 		*VertexShader,
 		EDRF_UseTriangleOptimization);
 
-	RHICopyToResolveTarget(DestRenderTarget.TargetableTexture, DestRenderTarget.ShaderResourceTexture, false, FResolveParams());
+	Context.RHICmdList.CopyToResolveTarget(DestRenderTarget.TargetableTexture, DestRenderTarget.ShaderResourceTexture, false, FResolveParams());
 }
 
 FPooledRenderTargetDesc FRCPassPostProcessBokehDOFSetup::ComputeOutputDesc(EPassOutputId InPassOutputId) const
@@ -440,25 +436,25 @@ public:
 	}
 
 	/** to have a similar interface as all other shaders */
-	void SetParameters(FRHICommandList& RHICmdList, const FRenderingCompositePassContext& Context, FIntPoint TileCountValue, uint32 TileSize, float PixelKernelSize, FIntPoint LeftTop)
+	void SetParameters(const FRenderingCompositePassContext& Context, FIntPoint TileCountValue, uint32 TileSize, float PixelKernelSize, FIntPoint LeftTop)
 	{
 		const FVertexShaderRHIParamRef ShaderRHI = GetVertexShader();
 
-		FGlobalShader::SetParameters(RHICmdList, ShaderRHI, Context.View);
-		DeferredParameters.Set(RHICmdList, ShaderRHI, Context.View);
-		PostprocessParameter.SetVS(RHICmdList, ShaderRHI, Context, TStaticSamplerState<SF_Bilinear,AM_Clamp,AM_Clamp,AM_Clamp>::GetRHI());
+		FGlobalShader::SetParameters(Context.RHICmdList, ShaderRHI, Context.View);
+		DeferredParameters.Set(Context.RHICmdList, ShaderRHI, Context.View);
+		PostprocessParameter.SetVS(ShaderRHI, Context, TStaticSamplerState<SF_Bilinear,AM_Clamp,AM_Clamp,AM_Clamp>::GetRHI());
 //		PostprocessParameter.SetVS(ShaderRHI, Context, TStaticSamplerState<SF_Point,AM_Clamp,AM_Clamp,AM_Clamp>::GetRHI());
 
 		{
 			FIntRect TileCountAndSizeValue(TileCountValue, FIntPoint(TileSize, TileSize));
 
-			SetShaderValue(RHICmdList, ShaderRHI, TileCountAndSize, TileCountAndSizeValue);
+			SetShaderValue(Context.RHICmdList, ShaderRHI, TileCountAndSize, TileCountAndSizeValue);
 		}
 
 		{
 			FVector4 KernelSizeValue(PixelKernelSize, PixelKernelSize, LeftTop.X, LeftTop.Y);
 
-			SetShaderValue(RHICmdList, ShaderRHI, KernelSize, KernelSizeValue);
+			SetShaderValue(Context.RHICmdList, ShaderRHI, KernelSize, KernelSizeValue);
 		}
 
 		{
@@ -466,7 +462,7 @@ public:
 				Context.View.FinalPostProcessSettings.DepthOfFieldColorThreshold,
 				Context.View.FinalPostProcessSettings.DepthOfFieldSizeThreshold,0);
 
-			SetShaderValue(RHICmdList, ShaderRHI, DepthOfFieldThresholds, Value);
+			SetShaderValue(Context.RHICmdList, ShaderRHI, DepthOfFieldThresholds, Value);
 		}
 
 		{
@@ -474,7 +470,7 @@ public:
 
 			FRCPassPostProcessBokehDOF::ComputeDepthOfFieldParams(Context, DepthOfFieldParamValues);
 
-			SetShaderValueArray(RHICmdList, ShaderRHI, DepthOfFieldParams, DepthOfFieldParamValues, 2);
+			SetShaderValueArray(Context.RHICmdList, ShaderRHI, DepthOfFieldParams, DepthOfFieldParamValues, 2);
 		}
 	}
 
@@ -524,13 +520,13 @@ public:
 		return bShaderHasOutdatedParameters;
 	}
 
-	void SetParameters(FRHICommandList& RHICmdList, const FRenderingCompositePassContext& Context, float PixelKernelSize)
+	void SetParameters(const FRenderingCompositePassContext& Context, float PixelKernelSize)
 	{
 		const FPixelShaderRHIParamRef ShaderRHI = GetPixelShader();
 
-		FGlobalShader::SetParameters(RHICmdList, ShaderRHI, Context.View);
+		FGlobalShader::SetParameters(Context.RHICmdList, ShaderRHI, Context.View);
 
-		PostprocessParameter.SetPS(RHICmdList, ShaderRHI, Context, TStaticSamplerState<SF_Bilinear,AM_Clamp,AM_Clamp,AM_Clamp>::GetRHI());
+		PostprocessParameter.SetPS(ShaderRHI, Context, TStaticSamplerState<SF_Bilinear,AM_Clamp,AM_Clamp,AM_Clamp>::GetRHI());
 //		PostprocessParameter.SetPS(ShaderRHI, Context, TStaticSamplerState<SF_Point,AM_Clamp,AM_Clamp,AM_Clamp>::GetRHI());
 				
 		{
@@ -546,7 +542,7 @@ public:
 				TextureRHI = Context.View.FinalPostProcessSettings.DepthOfFieldBokehShape->Resource->TextureRHI;
 			}
 
-			SetTextureParameter(RHICmdList, ShaderRHI, LensTexture, LensTextureSampler, TStaticSamplerState<SF_Trilinear,AM_Border,AM_Border,AM_Clamp>::GetRHI(), TextureRHI);
+			SetTextureParameter(Context.RHICmdList, ShaderRHI, LensTexture, LensTextureSampler, TStaticSamplerState<SF_Trilinear, AM_Border, AM_Border, AM_Clamp>::GetRHI(), TextureRHI);
 		}
 	}
 };
@@ -562,17 +558,18 @@ VARIATION1(0,1)			VARIATION1(1,1)			VARIATION1(2,1)
 #undef VARIATION1
 
 template <uint32 DOFMethod, uint32 DOFIndexStyle>
-void FRCPassPostProcessBokehDOF::SetShaderTempl(FRHICommandList& RHICmdList, const FRenderingCompositePassContext& Context, FIntPoint LeftTop, FIntPoint TileCount, uint32 TileSize, float PixelKernelSize)
+void FRCPassPostProcessBokehDOF::SetShaderTempl(const FRenderingCompositePassContext& Context, FIntPoint LeftTop, FIntPoint TileCount, uint32 TileSize, float PixelKernelSize)
 {
 	TShaderMapRef<FPostProcessBokehDOFVS<DOFMethod, DOFIndexStyle> > VertexShader(GetGlobalShaderMap());
 	TShaderMapRef<FPostProcessBokehDOFPS> PixelShader(GetGlobalShaderMap());
 
 	static FGlobalBoundShaderState BoundShaderState;
+	Context.RHICmdList.CheckIsNull(); // need new approach for "static FGlobalBoundShaderState" for parallel rendering
 
-	SetGlobalBoundShaderState(BoundShaderState, GFilterVertexDeclaration.VertexDeclarationRHI, *VertexShader, *PixelShader);
+	SetGlobalBoundShaderState(Context.RHICmdList, BoundShaderState, GFilterVertexDeclaration.VertexDeclarationRHI, *VertexShader, *PixelShader);
 
-	VertexShader->SetParameters(RHICmdList, Context, TileCount, TileSize, PixelKernelSize, LeftTop);
-	PixelShader->SetParameters(RHICmdList, Context, PixelKernelSize);
+	VertexShader->SetParameters(Context, TileCount, TileSize, PixelKernelSize, LeftTop);
+	PixelShader->SetParameters(Context, PixelKernelSize);
 }
 
 void FRCPassPostProcessBokehDOF::ComputeDepthOfFieldParams(const FRenderingCompositePassContext& Context, FVector4 Out[2])
@@ -631,22 +628,19 @@ void FRCPassPostProcessBokehDOF::Process(FRenderingCompositePassContext& Context
 
 	const FSceneRenderTargetItem& DestRenderTarget = PassOutputs[0].RequestSurface(Context);
 
-	//@todo-rco: RHIPacketList
-	FRHICommandList& RHICmdList = FRHICommandList::GetNullRef();
-
 	// Set the view family's render target/viewport.
-	RHISetRenderTarget(DestRenderTarget.TargetableTexture, FTextureRHIRef());	
+	SetRenderTarget(Context.RHICmdList, DestRenderTarget.TargetableTexture, FTextureRHIRef());
 
 	// This clean is required to make the accumulation working
-	RHIClear(true, FLinearColor(0, 0, 0, 0), false, 1.0f, false, 0, FIntRect());
+	Context.RHICmdList.Clear(true, FLinearColor(0, 0, 0, 0), false, 1.0f, false, 0, FIntRect());
 
 	// we need to output to the whole rendertarget
 	Context.SetViewportAndCallRHI(0, 0, 0.0f, PassOutputs[0].RenderTargetDesc.Extent.X, PassOutputs[0].RenderTargetDesc.Extent.Y, 1.0f);
 
 	// set the state (additive blending)
-	RHISetBlendState(TStaticBlendState<CW_RGBA,BO_Add,BF_One,BF_One,BO_Add,BF_One,BF_One>::GetRHI());
-	RHISetRasterizerState(TStaticRasterizerState<>::GetRHI());
-	RHISetDepthStencilState(TStaticDepthStencilState<false,CF_Always>::GetRHI());
+	Context.RHICmdList.SetBlendState(TStaticBlendState<CW_RGBA, BO_Add, BF_One, BF_One, BO_Add, BF_One, BF_One>::GetRHI());
+	Context.RHICmdList.SetRasterizerState(TStaticRasterizerState<>::GetRHI());
+	Context.RHICmdList.SetDepthStencilState(TStaticDepthStencilState<false, CF_Always>::GetRHI());
 	
 	static const auto CVar = IConsoleManager::Get().FindTConsoleVariableDataInt(TEXT("r.DepthOfFieldQuality"));
 	check(CVar);
@@ -679,11 +673,11 @@ void FRCPassPostProcessBokehDOF::Process(FRenderingCompositePassContext& Context
 			// high quality, visualize in red and green where we spend more performance
 			if(IndexStyle == BIS_Fast)
 			{
-				SetShaderTempl<2, 0>(RHICmdList, Context, LeftTop, TileCount, TileSize, PixelKernelSize);
+				SetShaderTempl<2, 0>(Context, LeftTop, TileCount, TileSize, PixelKernelSize);
 			}
 			else
 			{
-				SetShaderTempl<2, 1>(RHICmdList, Context, LeftTop, TileCount, TileSize, PixelKernelSize);
+				SetShaderTempl<2, 1>(Context, LeftTop, TileCount, TileSize, PixelKernelSize);
 			}
 		}
 		else
@@ -691,11 +685,11 @@ void FRCPassPostProcessBokehDOF::Process(FRenderingCompositePassContext& Context
 			// high quality
 			if(IndexStyle == BIS_Fast)
 			{
-				SetShaderTempl<1, 0>(RHICmdList, Context, LeftTop, TileCount, TileSize, PixelKernelSize);
+				SetShaderTempl<1, 0>(Context, LeftTop, TileCount, TileSize, PixelKernelSize);
 			}
 			else
 			{
-				SetShaderTempl<1, 1>(RHICmdList, Context, LeftTop, TileCount, TileSize, PixelKernelSize);
+				SetShaderTempl<1, 1>(Context, LeftTop, TileCount, TileSize, PixelKernelSize);
 			}
 		}
 	}
@@ -704,28 +698,28 @@ void FRCPassPostProcessBokehDOF::Process(FRenderingCompositePassContext& Context
 		// low quality
 		if(IndexStyle == BIS_Fast)
 		{
-			SetShaderTempl<0, 0>(RHICmdList, Context, LeftTop, TileCount, TileSize, PixelKernelSize);
+			SetShaderTempl<0, 0>(Context, LeftTop, TileCount, TileSize, PixelKernelSize);
 		}
 		else
 		{
-			SetShaderTempl<0, 1>(RHICmdList, Context, LeftTop, TileCount, TileSize, PixelKernelSize);
+			SetShaderTempl<0, 1>(Context, LeftTop, TileCount, TileSize, PixelKernelSize);
 		}
 	}
 
 	// needs to be the same on shader side (faster on NVIDIA and AMD)
 	int32 QuadsPerInstance = 8;
 
-	RHISetStreamSource(0, NULL, 0, 0);
+	Context.RHICmdList.SetStreamSource(0, NULL, 0, 0);
 	if(IndexStyle == BIS_Fast)
 	{
-		RHIDrawIndexedPrimitive( GBokehIndexBuffer.IndexBufferRHI, PT_TriangleList, 0, 0, 32, 0, 2 * QuadsPerInstance, FMath::DivideAndRoundUp(TileCount.X * TileCount.Y, QuadsPerInstance));
+		Context.RHICmdList.DrawIndexedPrimitive(GBokehIndexBuffer.IndexBufferRHI, PT_TriangleList, 0, 0, 32, 0, 2 * QuadsPerInstance, FMath::DivideAndRoundUp(TileCount.X * TileCount.Y, QuadsPerInstance));
 	}
 	else
 	{
-		RHIDrawIndexedPrimitive( GBokehSlowIndexBuffer.IndexBufferRHI, PT_TriangleList, 0, 0, 32, 0, 2 * QuadsPerInstance, FMath::DivideAndRoundUp(TileCount.X * TileCount.Y, QuadsPerInstance));
+		Context.RHICmdList.DrawIndexedPrimitive(GBokehSlowIndexBuffer.IndexBufferRHI, PT_TriangleList, 0, 0, 32, 0, 2 * QuadsPerInstance, FMath::DivideAndRoundUp(TileCount.X * TileCount.Y, QuadsPerInstance));
 	}
 
-	RHICopyToResolveTarget(DestRenderTarget.TargetableTexture, DestRenderTarget.ShaderResourceTexture, false, FResolveParams());
+	Context.RHICmdList.CopyToResolveTarget(DestRenderTarget.TargetableTexture, DestRenderTarget.ShaderResourceTexture, false, FResolveParams());
 }
 
 FPooledRenderTargetDesc FRCPassPostProcessBokehDOF::ComputeOutputDesc(EPassOutputId InPassOutputId) const
diff --git a/Engine/Source/Runtime/Renderer/Private/PostProcess/PostProcessBokehDOF.h b/Engine/Source/Runtime/Renderer/Private/PostProcess/PostProcessBokehDOF.h
--- a/Engine/Source/Runtime/Renderer/Private/PostProcess/PostProcessBokehDOF.h
+++ b/Engine/Source/Runtime/Renderer/Private/PostProcess/PostProcessBokehDOF.h
@@ -49,7 +49,7 @@ public:
 private:
 
 	template <uint32 HighQuality, uint32 IndexStyle>
-	void SetShaderTempl(FRHICommandList& RHICmdList, const FRenderingCompositePassContext& Context, FIntPoint LeftTop, FIntPoint TileCount, uint32 TileSize, float PixelKernelSize);
+	void SetShaderTempl(const FRenderingCompositePassContext& Context, FIntPoint LeftTop, FIntPoint TileCount, uint32 TileSize, float PixelKernelSize);
 
 	// border between front and back layer as we don't use viewports (only possible with GS)
 	const static uint32 SafetyBorder = 40;
diff --git a/Engine/Source/Runtime/Renderer/Private/PostProcess/PostProcessBokehDOFRecombine.cpp b/Engine/Source/Runtime/Renderer/Private/PostProcess/PostProcessBokehDOFRecombine.cpp
--- a/Engine/Source/Runtime/Renderer/Private/PostProcess/PostProcessBokehDOFRecombine.cpp
+++ b/Engine/Source/Runtime/Renderer/Private/PostProcess/PostProcessBokehDOFRecombine.cpp
@@ -59,20 +59,20 @@ public:
 		return bShaderHasOutdatedParameters;
 	}
 
-	void SetParameters(FRHICommandList& RHICmdList, const FRenderingCompositePassContext& Context)
+	void SetParameters(const FRenderingCompositePassContext& Context)
 	{
 		const FPixelShaderRHIParamRef ShaderRHI = GetPixelShader();
 
-		FGlobalShader::SetParameters(RHICmdList, ShaderRHI, Context.View);
-		DeferredParameters.Set(RHICmdList, ShaderRHI, Context.View);
-		PostprocessParameter.SetPS(RHICmdList, ShaderRHI, Context, TStaticSamplerState<SF_Bilinear,AM_Clamp,AM_Clamp,AM_Clamp>::GetRHI());
+		FGlobalShader::SetParameters(Context.RHICmdList, ShaderRHI, Context.View);
+		DeferredParameters.Set(Context.RHICmdList, ShaderRHI, Context.View);
+		PostprocessParameter.SetPS(ShaderRHI, Context, TStaticSamplerState<SF_Bilinear,AM_Clamp,AM_Clamp,AM_Clamp>::GetRHI());
 
 		{
 			FVector4 DepthOfFieldParamValues[2];
 
 			FRCPassPostProcessBokehDOF::ComputeDepthOfFieldParams(Context, DepthOfFieldParamValues);
 
-			SetShaderValueArray(RHICmdList, ShaderRHI, DepthOfFieldParams, DepthOfFieldParamValues, 2);
+			SetShaderValueArray(Context.RHICmdList, ShaderRHI, DepthOfFieldParams, DepthOfFieldParamValues, 2);
 		}
 	}
 
@@ -96,17 +96,18 @@ VARIATION1(1)			VARIATION1(2)			VARIATION1(3)
 
 
 template <uint32 Method>
-void FRCPassPostProcessBokehDOFRecombine::SetShader(FRHICommandList& RHICmdList, const FRenderingCompositePassContext& Context)
+void FRCPassPostProcessBokehDOFRecombine::SetShader(const FRenderingCompositePassContext& Context)
 {
 	TShaderMapRef<FPostProcessVS> VertexShader(GetGlobalShaderMap());
 	TShaderMapRef<FPostProcessBokehDOFRecombinePS<Method> > PixelShader(GetGlobalShaderMap());
 
 	static FGlobalBoundShaderState BoundShaderState;
+	Context.RHICmdList.CheckIsNull(); // need new approach for "static FGlobalBoundShaderState" for parallel rendering
 
-	SetGlobalBoundShaderState(BoundShaderState, GFilterVertexDeclaration.VertexDeclarationRHI, *VertexShader, *PixelShader);
+	SetGlobalBoundShaderState(Context.RHICmdList, BoundShaderState, GFilterVertexDeclaration.VertexDeclarationRHI, *VertexShader, *PixelShader);
 
-	PixelShader->SetParameters(RHICmdList, Context);
-	VertexShader->SetParameters(RHICmdList, Context);
+	PixelShader->SetParameters(Context);
+	VertexShader->SetParameters(Context);
 }
 
 void FRCPassPostProcessBokehDOFRecombine::Process(FRenderingCompositePassContext& Context)
@@ -144,27 +145,25 @@ void FRCPassPostProcessBokehDOFRecombine::Process(FRenderingCompositePassContext
 	FIntRect HalfResViewRect = FIntRect::DivideAndRoundUp(View.ViewRect, ScaleToFullRes);
 
 	const FSceneRenderTargetItem& DestRenderTarget = PassOutputs[0].RequestSurface(Context);
-	//@todo-rco: RHIPacketList
-	FRHICommandList& RHICmdList = FRHICommandList::GetNullRef();
 
 	// Set the view family's render target/viewport.
-	RHISetRenderTarget(DestRenderTarget.TargetableTexture, FTextureRHIRef());	
+	SetRenderTarget(Context.RHICmdList, DestRenderTarget.TargetableTexture, FTextureRHIRef());
 
 	// is optimized away if possible (RT size=view size, )
-	RHIClear(true, FLinearColor::Black, false, 1.0f, false, 0, View.ViewRect);
+	Context.RHICmdList.Clear(true, FLinearColor::Black, false, 1.0f, false, 0, View.ViewRect);
 
 	Context.SetViewportAndCallRHI(View.ViewRect);
 
 	// set the state
-	RHISetBlendState(TStaticBlendState<>::GetRHI());
-	RHISetRasterizerState(TStaticRasterizerState<>::GetRHI());
-	RHISetDepthStencilState(TStaticDepthStencilState<false,CF_Always>::GetRHI());
+	Context.RHICmdList.SetBlendState(TStaticBlendState<>::GetRHI());
+	Context.RHICmdList.SetRasterizerState(TStaticRasterizerState<>::GetRHI());
+	Context.RHICmdList.SetDepthStencilState(TStaticDepthStencilState<false, CF_Always>::GetRHI());
 
 	switch(Method)
 		{
-		case 1: SetShader<1>(RHICmdList, Context); break;
-		case 2: SetShader<2>(RHICmdList, Context); break;
-		case 3: SetShader<3>(RHICmdList, Context); break;
+		case 1: SetShader<1>(Context); break;
+		case 2: SetShader<2>(Context); break;
+		case 3: SetShader<3>(Context); break;
 		default:
 			check(0);
 	}
@@ -173,6 +172,7 @@ void FRCPassPostProcessBokehDOFRecombine::Process(FRenderingCompositePassContext
 
 	// Draw a quad mapping scene color to the view's render target
 	DrawRectangle(
+		Context.RHICmdList,
 		0, 0,
 		View.ViewRect.Width(), View.ViewRect.Height(),
 		HalfResViewRect.Min.X, HalfResViewRect.Min.Y, 
@@ -182,7 +182,7 @@ void FRCPassPostProcessBokehDOFRecombine::Process(FRenderingCompositePassContext
 		*VertexShader,
 		EDRF_UseTriangleOptimization);
 
-	RHICopyToResolveTarget(DestRenderTarget.TargetableTexture, DestRenderTarget.ShaderResourceTexture, false, FResolveParams());
+	Context.RHICmdList.CopyToResolveTarget(DestRenderTarget.TargetableTexture, DestRenderTarget.ShaderResourceTexture, false, FResolveParams());
 }
 
 FPooledRenderTargetDesc FRCPassPostProcessBokehDOFRecombine::ComputeOutputDesc(EPassOutputId InPassOutputId) const
diff --git a/Engine/Source/Runtime/Renderer/Private/PostProcess/PostProcessBokehDOFRecombine.h b/Engine/Source/Runtime/Renderer/Private/PostProcess/PostProcessBokehDOFRecombine.h
--- a/Engine/Source/Runtime/Renderer/Private/PostProcess/PostProcessBokehDOFRecombine.h
+++ b/Engine/Source/Runtime/Renderer/Private/PostProcess/PostProcessBokehDOFRecombine.h
@@ -23,5 +23,5 @@ public:
 private:
 
 	template <uint32 Method>
-	static void SetShader(FRHICommandList& RHICmdList, const FRenderingCompositePassContext& Context);
+	static void SetShader(const FRenderingCompositePassContext& Context);
 };
\ No newline at end of file
