UE4 - more work preparing for parallel rendering. Merged the cmd list into the post process context, among other things.

diff --git a/Engine/Source/Runtime/Renderer/Private/PostProcess/PostProcessBusyWait.cpp b/Engine/Source/Runtime/Renderer/Private/PostProcess/PostProcessBusyWait.cpp
--- a/Engine/Source/Runtime/Renderer/Private/PostProcess/PostProcessBusyWait.cpp
+++ b/Engine/Source/Runtime/Renderer/Private/PostProcess/PostProcessBusyWait.cpp
@@ -35,13 +35,13 @@ public:
 		GPUBusyWait.Bind(Initializer.ParameterMap,TEXT("GPUBusyWait"));
 	}
 
-	void SetPS(FRHICommandList& RHICmdList, const FRenderingCompositePassContext& Context)
+	void SetPS(const FRenderingCompositePassContext& Context)
 	{
 		const FPixelShaderRHIParamRef ShaderRHI = GetPixelShader();
 
-		FGlobalShader::SetParameters(RHICmdList, ShaderRHI, Context.View);
+		FGlobalShader::SetParameters(Context.RHICmdList, ShaderRHI, Context.View);
 
-		PostprocessParameter.SetPS(RHICmdList, ShaderRHI, Context, TStaticSamplerState<SF_Bilinear,AM_Clamp,AM_Clamp,AM_Clamp>::GetRHI());
+		PostprocessParameter.SetPS(ShaderRHI, Context, TStaticSamplerState<SF_Bilinear,AM_Clamp,AM_Clamp,AM_Clamp>::GetRHI());
 
 #if !(UE_BUILD_SHIPPING || UE_BUILD_TEST)
 		{
@@ -56,7 +56,7 @@ public:
 			// divide by viewport pixel count
 			uint32 Value = (uint32)(CVarValue * 1000000000.0 / 6.12 / PixelCount);
 
-			SetShaderValue(RHICmdList, ShaderRHI, GPUBusyWait, Value);
+			SetShaderValue(Context.RHICmdList, ShaderRHI, GPUBusyWait, Value);
 		}
 #endif
 	}
@@ -85,29 +85,27 @@ void FRCPassPostProcessBusyWait::Process(FRenderingCompositePassContext& Context
 	
 	const FSceneRenderTargetItem& DestRenderTarget = GSceneRenderTargets.GetLightAttenuation()->GetRenderTargetItem();
 
-	//@todo-rco: RHIPacketList
-	FRHICommandList& RHICmdList = FRHICommandList::GetNullRef();
-
 	// Set the view family's render target/viewport.
-	RHISetRenderTarget(DestRenderTarget.TargetableTexture, FTextureRHIRef());	
+	SetRenderTarget(Context.RHICmdList, DestRenderTarget.TargetableTexture, FTextureRHIRef());
 	Context.SetViewportAndCallRHI(DestRect);
 
 	// set the state
-	RHISetBlendState(TStaticBlendState<>::GetRHI());
-	RHISetRasterizerState(TStaticRasterizerState<>::GetRHI());
-	RHISetDepthStencilState(TStaticDepthStencilState<false,CF_Always>::GetRHI());
+	Context.RHICmdList.SetBlendState(TStaticBlendState<>::GetRHI());
+	Context.RHICmdList.SetRasterizerState(TStaticRasterizerState<>::GetRHI());
+	Context.RHICmdList.SetDepthStencilState(TStaticDepthStencilState<false, CF_Always>::GetRHI());
 
 	TShaderMapRef<FPostProcessVS> VertexShader(GetGlobalShaderMap());
 	TShaderMapRef<FPostProcessBusyWaitPS> PixelShader(GetGlobalShaderMap());
 
 	static FGlobalBoundShaderState BoundShaderState;
 
-	SetGlobalBoundShaderState(BoundShaderState, GFilterVertexDeclaration.VertexDeclarationRHI, *VertexShader, *PixelShader);
+	SetGlobalBoundShaderState(Context.RHICmdList, BoundShaderState, GFilterVertexDeclaration.VertexDeclarationRHI, *VertexShader, *PixelShader);
 
-	PixelShader->SetPS(RHICmdList, Context);
+	PixelShader->SetPS(Context);
 
 	// Draw a quad mapping scene color to the view's render target
 	DrawRectangle(
+		Context.RHICmdList,
 		0, 0,
 		DestRect.Width(), DestRect.Height(),
 		SrcRect.Min.X, SrcRect.Min.Y,
@@ -117,8 +115,9 @@ void FRCPassPostProcessBusyWait::Process(FRenderingCompositePassContext& Context
 		*VertexShader,
 		EDRF_UseTriangleOptimization);
 
-	RHICopyToResolveTarget(DestRenderTarget.TargetableTexture, DestRenderTarget.ShaderResourceTexture, false, FResolveParams());
+	Context.RHICmdList.CopyToResolveTarget(DestRenderTarget.TargetableTexture, DestRenderTarget.ShaderResourceTexture, false, FResolveParams());
 
+	Context.RHICmdList.CheckIsNull(); // this won't work for parallel rendering
 	GSceneRenderTargets.SetLightAttenuation(0);
 }
 
diff --git a/Engine/Source/Runtime/Renderer/Private/PostProcess/PostProcessCombineLUTs.cpp b/Engine/Source/Runtime/Renderer/Private/PostProcess/PostProcessCombineLUTs.cpp
--- a/Engine/Source/Runtime/Renderer/Private/PostProcess/PostProcessCombineLUTs.cpp
+++ b/Engine/Source/Runtime/Renderer/Private/PostProcess/PostProcessCombineLUTs.cpp
@@ -180,7 +180,7 @@ IMPLEMENT_SHADER_TYPE(template<>,FLUTBlenderPS<3>,TEXT("PostProcessCombineLUTs")
 IMPLEMENT_SHADER_TYPE(template<>,FLUTBlenderPS<4>,TEXT("PostProcessCombineLUTs"),TEXT("MainPS"),SF_Pixel);
 IMPLEMENT_SHADER_TYPE(template<>,FLUTBlenderPS<5>,TEXT("PostProcessCombineLUTs"),TEXT("MainPS"),SF_Pixel);
 
-void SetLUTBlenderShader(FRHICommandList& RHICmdList, FRenderingCompositePassContext& Context, uint32 BlendCount, FTexture* Texture[], float Weights[], const FVolumeBounds& VolumeBounds)
+void SetLUTBlenderShader(FRenderingCompositePassContext& Context, uint32 BlendCount, FTexture* Texture[], float Weights[], const FVolumeBounds& VolumeBounds)
 {
 	check(BlendCount > 0);
 
@@ -217,24 +217,24 @@ void SetLUTBlenderShader(FRHICommandList& RHICmdList, FRenderingCompositePassCon
 		TShaderMapRef<FWriteToSliceVS> VertexShader(GetGlobalShaderMap());
 		TShaderMapRef<FWriteToSliceGS> GeometryShader(GetGlobalShaderMap());
 
-		SetGlobalBoundShaderState(*LocalBoundShaderState, GScreenVertexDeclaration.VertexDeclarationRHI, *VertexShader, LocalPixelShader, *GeometryShader);
+		SetGlobalBoundShaderState(Context.RHICmdList, *LocalBoundShaderState, GScreenVertexDeclaration.VertexDeclarationRHI, *VertexShader, LocalPixelShader, *GeometryShader);
 
-		VertexShader->SetParameters(RHICmdList, VolumeBounds, VolumeBounds.MaxX - VolumeBounds.MinX);
-		GeometryShader->SetParameters(RHICmdList, VolumeBounds);
+		VertexShader->SetParameters(Context.RHICmdList, VolumeBounds, VolumeBounds.MaxX - VolumeBounds.MinX);
+		GeometryShader->SetParameters(Context.RHICmdList, VolumeBounds);
 	}
 	else
 	{
 		TShaderMapRef<FPostProcessVS> VertexShader(GetGlobalShaderMap());
 
-		SetGlobalBoundShaderState(*LocalBoundShaderState, GFilterVertexDeclaration.VertexDeclarationRHI, *VertexShader, LocalPixelShader);
+		SetGlobalBoundShaderState(Context.RHICmdList, *LocalBoundShaderState, GFilterVertexDeclaration.VertexDeclarationRHI, *VertexShader, LocalPixelShader);
 
-		VertexShader->SetParameters(RHICmdList, Context);
+		VertexShader->SetParameters(Context);
 	}
 #define CASE_COUNT(BlendCount) \
 	case BlendCount: \
 	{ \
 	TShaderMapRef<FLUTBlenderPS<BlendCount> > PixelShader(GetGlobalShaderMap()); \
-	PixelShader->SetParameters(RHICmdList, View, Texture, Weights); \
+	PixelShader->SetParameters(Context.RHICmdList, View, Texture, Weights); \
 	}; \
 	break;
 
@@ -386,26 +386,23 @@ void FRCPassPostProcessCombineLUTs::Process(FRenderingCompositePassContext& Cont
 
 	const FSceneRenderTargetItem& DestRenderTarget = PassOutputs[0].RequestSurface(Context);
 
-	//@todo-rco: RHIPacketList
-	FRHICommandList& RHICmdList = FRHICommandList::GetNullRef();
-
 	// Set the view family's render target/viewport.
-	RHISetRenderTarget(DestRenderTarget.TargetableTexture, FTextureRHIRef());	
+	SetRenderTarget(Context.RHICmdList, DestRenderTarget.TargetableTexture, FTextureRHIRef());
 	Context.SetViewportAndCallRHI(0, 0, 0.0f, DestSize.X, DestSize.Y, 1.0f );
 
 	// set the state
-	RHISetBlendState(TStaticBlendState<>::GetRHI());
-	RHISetRasterizerState(TStaticRasterizerState<>::GetRHI());
-	RHISetDepthStencilState(TStaticDepthStencilState<false,CF_Always>::GetRHI());
+	Context.RHICmdList.SetBlendState(TStaticBlendState<>::GetRHI());
+	Context.RHICmdList.SetRasterizerState(TStaticRasterizerState<>::GetRHI());
+	Context.RHICmdList.SetDepthStencilState(TStaticDepthStencilState<false, CF_Always>::GetRHI());
 
 	const FVolumeBounds VolumeBounds(16);
 
-	SetLUTBlenderShader(RHICmdList, Context, LocalCount, LocalTextures, LocalWeights, VolumeBounds);
+	SetLUTBlenderShader(Context, LocalCount, LocalTextures, LocalWeights, VolumeBounds);
 
 	if(UseVolumeTextureLUT(GRHIShaderPlatform))
 	{
 		// use volume texture 16x16x16
-		RasterizeToVolumeTexture(VolumeBounds);
+		RasterizeToVolumeTexture(Context.RHICmdList, VolumeBounds);
 	}
 	else
 	{
@@ -413,6 +410,7 @@ void FRCPassPostProcessCombineLUTs::Process(FRenderingCompositePassContext& Cont
 		TShaderMapRef<FPostProcessVS> VertexShader(GetGlobalShaderMap());
 
 		DrawRectangle( 
+			Context.RHICmdList,
 			0, 0,						// XY
 			16 * 16, 16,				// SizeXY
 			0, 0,						// UV
@@ -423,7 +421,7 @@ void FRCPassPostProcessCombineLUTs::Process(FRenderingCompositePassContext& Cont
 			EDRF_UseTriangleOptimization);
 	}
 
-	RHICopyToResolveTarget(DestRenderTarget.TargetableTexture, DestRenderTarget.ShaderResourceTexture, false, FResolveParams());
+	Context.RHICmdList.CopyToResolveTarget(DestRenderTarget.TargetableTexture, DestRenderTarget.ShaderResourceTexture, false, FResolveParams());
 }
 
 FPooledRenderTargetDesc FRCPassPostProcessCombineLUTs::ComputeOutputDesc(EPassOutputId InPassOutputId) const
diff --git a/Engine/Source/Runtime/Renderer/Private/PostProcess/PostProcessCompositeEditorPrimitives.cpp b/Engine/Source/Runtime/Renderer/Private/PostProcess/PostProcessCompositeEditorPrimitives.cpp
--- a/Engine/Source/Runtime/Renderer/Private/PostProcess/PostProcessCompositeEditorPrimitives.cpp
+++ b/Engine/Source/Runtime/Renderer/Private/PostProcess/PostProcessCompositeEditorPrimitives.cpp
@@ -66,25 +66,25 @@ public:
 		FilteredSceneDepthTextureSampler.Bind(Initializer.ParameterMap,TEXT("FilteredSceneDepthTextureSampler"));
 	}
 
-	void SetParameters(FRHICommandList& RHICmdList, const FRenderingCompositePassContext& Context)
+	void SetParameters(const FRenderingCompositePassContext& Context)
 	{
 		const FPixelShaderRHIParamRef ShaderRHI = GetPixelShader();
 
-		FGlobalShader::SetParameters(RHICmdList, ShaderRHI, Context.View);
+		FGlobalShader::SetParameters(Context.RHICmdList, ShaderRHI, Context.View);
 
-		DeferredParameters.Set(RHICmdList, ShaderRHI, Context.View);
+		DeferredParameters.Set(Context.RHICmdList, ShaderRHI, Context.View);
 
 		FSamplerStateRHIRef SamplerStateRHIRef = TStaticSamplerState<SF_Point,AM_Clamp,AM_Clamp,AM_Clamp>::GetRHI();
-		PostProcessParameters.SetPS(RHICmdList, ShaderRHI, Context, SamplerStateRHIRef);
+		PostProcessParameters.SetPS(ShaderRHI, Context, SamplerStateRHIRef);
 		if(MSAASampleCount > 1)
 		{
-			SetTextureParameter(RHICmdList, ShaderRHI, EditorPrimitivesColor, GSceneRenderTargets.EditorPrimitivesColor->GetRenderTargetItem().TargetableTexture );
-			SetTextureParameter(RHICmdList, ShaderRHI, EditorPrimitivesDepth, GSceneRenderTargets.EditorPrimitivesDepth->GetRenderTargetItem().TargetableTexture );
+			SetTextureParameter(Context.RHICmdList, ShaderRHI, EditorPrimitivesColor, GSceneRenderTargets.EditorPrimitivesColor->GetRenderTargetItem().TargetableTexture);
+			SetTextureParameter(Context.RHICmdList, ShaderRHI, EditorPrimitivesDepth, GSceneRenderTargets.EditorPrimitivesDepth->GetRenderTargetItem().TargetableTexture);
 		}
 		else
 		{
-			SetTextureParameter(RHICmdList, ShaderRHI, EditorPrimitivesColor, EditorPrimitivesColorSampler, SamplerStateRHIRef, GSceneRenderTargets.EditorPrimitivesColor->GetRenderTargetItem().ShaderResourceTexture );
-			SetTextureParameter(RHICmdList, ShaderRHI, EditorPrimitivesDepth, GSceneRenderTargets.EditorPrimitivesDepth->GetRenderTargetItem().ShaderResourceTexture );
+			SetTextureParameter(Context.RHICmdList, ShaderRHI, EditorPrimitivesColor, EditorPrimitivesColorSampler, SamplerStateRHIRef, GSceneRenderTargets.EditorPrimitivesColor->GetRenderTargetItem().ShaderResourceTexture);
+			SetTextureParameter(Context.RHICmdList, ShaderRHI, EditorPrimitivesDepth, GSceneRenderTargets.EditorPrimitivesDepth->GetRenderTargetItem().ShaderResourceTexture);
 		}
 
 		{
@@ -104,14 +104,14 @@ public:
 				Value.G = 0;
 			}
 
-			SetShaderValue(RHICmdList, ShaderRHI, EditorRenderParams, Value);
+			SetShaderValue(Context.RHICmdList, ShaderRHI, EditorRenderParams, Value);
 		}
 
 		if(FilteredSceneDepthTexture.IsBound())
 		{
 			const FTexture2DRHIRef* DepthTexture = GSceneRenderTargets.GetActualDepthTexture();
 			SetTextureParameter(
-				RHICmdList,
+				Context.RHICmdList,
 				ShaderRHI,
 				FilteredSceneDepthTexture,
 				FilteredSceneDepthTextureSampler,
@@ -162,17 +162,18 @@ private:
 
 
 template <uint32 MSAASampleCount>
-static void SetCompositePrimitivesShaderTempl(FRHICommandList& RHICmdList, const FRenderingCompositePassContext& Context)
+static void SetCompositePrimitivesShaderTempl(const FRenderingCompositePassContext& Context)
 {
 	TShaderMapRef<FPostProcessVS> VertexShader(GetGlobalShaderMap());
 	TShaderMapRef<FPostProcessCompositeEditorPrimitivesPS<MSAASampleCount> > PixelShader(GetGlobalShaderMap());
 
 	static FGlobalBoundShaderState BoundShaderState;
+	Context.RHICmdList.CheckIsNull(); // need new approach for "static FGlobalBoundShaderState" for parallel rendering
 
-	SetGlobalBoundShaderState(BoundShaderState, GFilterVertexDeclaration.VertexDeclarationRHI, *VertexShader, *PixelShader);
+	SetGlobalBoundShaderState(Context.RHICmdList, BoundShaderState, GFilterVertexDeclaration.VertexDeclarationRHI, *VertexShader, *PixelShader);
 
-	VertexShader->SetParameters(RHICmdList, Context);
-	PixelShader->SetParameters(RHICmdList, Context);
+	VertexShader->SetParameters(Context);
+	PixelShader->SetParameters(Context);
 }
 
 void FRCPassPostProcessCompositeEditorPrimitives::Process(FRenderingCompositePassContext& Context)
@@ -202,11 +203,8 @@ void FRCPassPostProcessCompositeEditorPrimitives::Process(FRenderingCompositePas
 
 	const uint32 MSAASampleCount = GSceneRenderTargets.EditorPrimitivesColor->GetDesc().NumSamples;
 
-	//@todo-rco: RHIPacketList
-	FRHICommandList& RHICmdList = FRHICommandList::GetNullRef();
-
 	{
-		RHISetRenderTarget( ColorTarget, DepthTarget );
+		SetRenderTarget(Context.RHICmdList, ColorTarget, DepthTarget);
 		Context.SetViewportAndCallRHI(DestRect);
 
 		if(bClearIsNeeded)
@@ -214,48 +212,49 @@ void FRCPassPostProcessCompositeEditorPrimitives::Process(FRenderingCompositePas
 			SCOPED_DRAW_EVENT(ClearViewEditorPrimitives, DEC_SCENE_ITEMS);
 			// Clear color and depth
 			// Note, this is a reversed Z depth surface, so 0.0f is the far plane.
-			RHIClear(true, FLinearColor(0, 0, 0, 0), true, 0.0f, false, 0, FIntRect());
+			Context.RHICmdList.Clear(true, FLinearColor(0, 0, 0, 0), true, 0.0f, false, 0, FIntRect());
 		}
 
 		SCOPED_DRAW_EVENT(RenderEditorPrimitives, DEC_SCENE_ITEMS);
 
-		RHISetRasterizerState( View.bReverseCulling?TStaticRasterizerState<FM_Solid,CM_CW>::GetRHI() : TStaticRasterizerState<FM_Solid,CM_CCW>::GetRHI() );
+		Context.RHICmdList.SetRasterizerState(View.bReverseCulling ? TStaticRasterizerState<FM_Solid, CM_CW>::GetRHI() : TStaticRasterizerState<FM_Solid, CM_CCW>::GetRHI());
 
-		RenderPrimitivesToComposite(View);
+		RenderPrimitivesToComposite(Context.RHICmdList, View);
 
-		GRenderTargetPool.VisualizeTexture.SetCheckPoint(GSceneRenderTargets.EditorPrimitivesColor);
+		Context.RHICmdList.CheckIsNull(); // this probably won't work with parallel rendering
+		GRenderTargetPool.VisualizeTexture.SetCheckPoint(Context.RHICmdList, GSceneRenderTargets.EditorPrimitivesColor);
 	}
 
-//	RHICopyToResolveTarget(GSceneRenderTargets.EditorPrimitivesColor->GetRenderTargetItem().TargetableTexture, GSceneRenderTargets.EditorPrimitivesColor->GetRenderTargetItem().ShaderResourceTexture, false, FResolveParams());
+	//	Context.RHICmdList.CopyToResolveTarget(GSceneRenderTargets.EditorPrimitivesColor->GetRenderTargetItem().TargetableTexture, GSceneRenderTargets.EditorPrimitivesColor->GetRenderTargetItem().ShaderResourceTexture, false, FResolveParams());
 
 	const FSceneRenderTargetItem& DestRenderTarget = PassOutputs[0].RequestSurface(Context);
 	const FTexture2DRHIRef& DestRenderTargetSurface = (const FTexture2DRHIRef&)DestRenderTarget.TargetableTexture;
 
 	// Set the view family's render target/viewport.
-	RHISetRenderTarget(DestRenderTargetSurface, FTextureRHIRef());	
+	SetRenderTarget(Context.RHICmdList, DestRenderTargetSurface, FTextureRHIRef());
 
 	Context.SetViewportAndCallRHI(DestRect);
 
 	// set the state
-	RHISetBlendState(TStaticBlendState<>::GetRHI());
-	RHISetRasterizerState(TStaticRasterizerState<>::GetRHI());
-	RHISetDepthStencilState(TStaticDepthStencilState<false,CF_Always>::GetRHI());
+	Context.RHICmdList.SetBlendState(TStaticBlendState<>::GetRHI());
+	Context.RHICmdList.SetRasterizerState(TStaticRasterizerState<>::GetRHI());
+	Context.RHICmdList.SetDepthStencilState(TStaticDepthStencilState<false, CF_Always>::GetRHI());
 
 	if(MSAASampleCount == 1)
 	{
-		SetCompositePrimitivesShaderTempl<1>(RHICmdList, Context);
+		SetCompositePrimitivesShaderTempl<1>(Context);
 	}
 	else if(MSAASampleCount == 2)
 	{
-		SetCompositePrimitivesShaderTempl<2>(RHICmdList, Context);
+		SetCompositePrimitivesShaderTempl<2>(Context);
 	}
 	else if(MSAASampleCount == 4)
 	{
-		SetCompositePrimitivesShaderTempl<4>(RHICmdList, Context);
+		SetCompositePrimitivesShaderTempl<4>(Context);
 	}
 	else if(MSAASampleCount == 8)
 	{
-		SetCompositePrimitivesShaderTempl<8>(RHICmdList, Context);
+		SetCompositePrimitivesShaderTempl<8>(Context);
 	}
 	else
 	{
@@ -267,6 +266,7 @@ void FRCPassPostProcessCompositeEditorPrimitives::Process(FRenderingCompositePas
 
 	// Draw a quad mapping our render targets to the view's render target
 	DrawRectangle(
+		Context.RHICmdList,
 		0, 0,
 		DestRect.Width(), DestRect.Height(),
 		SrcRect.Min.X, SrcRect.Min.Y,
@@ -276,9 +276,10 @@ void FRCPassPostProcessCompositeEditorPrimitives::Process(FRenderingCompositePas
 		*VertexShader,
 		EDRF_UseTriangleOptimization);
 
-	RHICopyToResolveTarget(DestRenderTargetSurface, DestRenderTarget.ShaderResourceTexture, false, FResolveParams());
+	Context.RHICmdList.CopyToResolveTarget(DestRenderTargetSurface, DestRenderTarget.ShaderResourceTexture, false, FResolveParams());
 
 	// Clean up targets
+	Context.RHICmdList.CheckIsNull(); // this probably won't work with parallel rendering
 	GSceneRenderTargets.CleanUpEditorPrimitiveTargets();
 }
 
@@ -292,12 +293,14 @@ FPooledRenderTargetDesc FRCPassPostProcessCompositeEditorPrimitives::ComputeOutp
 	return Ret;
 }
 
-void FRCPassPostProcessCompositeEditorPrimitives::RenderPrimitivesToComposite(const FViewInfo& View)
+void FRCPassPostProcessCompositeEditorPrimitives::RenderPrimitivesToComposite(FRHICommandList& RHICmdList, const FViewInfo& View)
 {
+	RHICmdList.CheckIsNull(); // This uses DrawDynamicElements
+
 	// Always depth test against other editor primitives
 	// Note, this is a reversed Z depth surface, using CF_GreaterEqual.
-	RHISetDepthStencilState(TStaticDepthStencilState<true,CF_GreaterEqual>::GetRHI() );
-	RHISetBlendState(TStaticBlendStateWriteMask<CW_RGBA>::GetRHI());
+	RHICmdList.SetDepthStencilState(TStaticDepthStencilState<true, CF_GreaterEqual>::GetRHI());
+	RHICmdList.SetBlendState(TStaticBlendStateWriteMask<CW_RGBA>::GetRHI());
 
 	// most objects should be occluded by the existing scene so we do a manual depth test in the shader
 	bool bDepthTest = true;
@@ -353,24 +356,24 @@ void FRCPassPostProcessCompositeEditorPrimitives::RenderPrimitivesToComposite(co
 	DrawViewElements<FBasePassOpaqueDrawingPolicyFactory>(View,FBasePassOpaqueDrawingPolicyFactory::ContextType(bDepthTest, ESceneRenderTargetsMode::SetTextures), SDPG_World, false);
 
 	// Draw the view's batched simple elements(lines, sprites, etc).
-	View.BatchedViewElements.Draw(bNeedToSwitchVerticalAxis, View.ViewProjectionMatrix, View.ViewRect.Width(), View.ViewRect.Height(), false, 1.0f, &View, SceneDepth );
+	View.BatchedViewElements.Draw(RHICmdList, bNeedToSwitchVerticalAxis, View.ViewProjectionMatrix, View.ViewRect.Width(), View.ViewRect.Height(), false, 1.0f, &View, SceneDepth);
 
 	// Draw foreground objects. Draw twice, once without depth testing to bring them into the foreground and again to depth test against themselves
 	{
 		// Do not test against non-composited objects
 		bDepthTest = false;
 
-		RHISetDepthStencilState(TStaticDepthStencilState<false,CF_Always>::GetRHI() );
+		RHICmdList.SetDepthStencilState(TStaticDepthStencilState<false, CF_Always>::GetRHI());
 
 		DrawViewElements<FBasePassOpaqueDrawingPolicyFactory>(View,FBasePassOpaqueDrawingPolicyFactory::ContextType(bDepthTest, ESceneRenderTargetsMode::SetTextures), SDPG_Foreground, false);
 
-		View.TopBatchedViewElements.Draw(bNeedToSwitchVerticalAxis, View.ViewProjectionMatrix, View.ViewRect.Width(), View.ViewRect.Height(), false);
+		View.TopBatchedViewElements.Draw(RHICmdList, bNeedToSwitchVerticalAxis, View.ViewProjectionMatrix, View.ViewRect.Width(), View.ViewRect.Height(), false);
 
 		// Note, this is a reversed Z depth surface, using CF_GreaterEqual.
-		RHISetDepthStencilState(TStaticDepthStencilState<true,CF_GreaterEqual>::GetRHI() );
+		RHICmdList.SetDepthStencilState(TStaticDepthStencilState<true, CF_GreaterEqual>::GetRHI());
 
 		DrawViewElements<FBasePassOpaqueDrawingPolicyFactory>(View,FBasePassOpaqueDrawingPolicyFactory::ContextType(bDepthTest, ESceneRenderTargetsMode::SetTextures), SDPG_Foreground, false);
 
-		View.TopBatchedViewElements.Draw(bNeedToSwitchVerticalAxis, View.ViewProjectionMatrix, View.ViewRect.Width(), View.ViewRect.Height(), false);
+		View.TopBatchedViewElements.Draw(RHICmdList, bNeedToSwitchVerticalAxis, View.ViewProjectionMatrix, View.ViewRect.Width(), View.ViewRect.Height(), false);
 	}
 }
diff --git a/Engine/Source/Runtime/Renderer/Private/PostProcess/PostProcessCompositeEditorPrimitives.h b/Engine/Source/Runtime/Renderer/Private/PostProcess/PostProcessCompositeEditorPrimitives.h
--- a/Engine/Source/Runtime/Renderer/Private/PostProcess/PostProcessCompositeEditorPrimitives.h
+++ b/Engine/Source/Runtime/Renderer/Private/PostProcess/PostProcessCompositeEditorPrimitives.h
@@ -21,6 +21,6 @@ private:
 	 *
 	 * @param View	The view to draw in
 	 */
-	void RenderPrimitivesToComposite(const FViewInfo& View);
+	void RenderPrimitivesToComposite(FRHICommandList& RHICmdList, const FViewInfo& View);
 };
 
diff --git a/Engine/Source/Runtime/Renderer/Private/PostProcess/PostProcessDOF.cpp b/Engine/Source/Runtime/Renderer/Private/PostProcess/PostProcessDOF.cpp
--- a/Engine/Source/Runtime/Renderer/Private/PostProcess/PostProcessDOF.cpp
+++ b/Engine/Source/Runtime/Renderer/Private/PostProcess/PostProcessDOF.cpp
@@ -55,22 +55,22 @@ public:
 		return bShaderHasOutdatedParameters;
 	}
 
-	void SetParameters(FRHICommandList& RHICmdList, const FRenderingCompositePassContext& Context)
+	void SetParameters(const FRenderingCompositePassContext& Context)
 	{
 		const FPixelShaderRHIParamRef ShaderRHI = GetPixelShader();
 
-		FGlobalShader::SetParameters(RHICmdList, ShaderRHI, Context.View);
+		FGlobalShader::SetParameters(Context.RHICmdList, ShaderRHI, Context.View);
 
-		PostprocessParameter.SetPS(RHICmdList, ShaderRHI, Context, TStaticSamplerState<SF_Point,AM_Border,AM_Border,AM_Clamp>::GetRHI());
+		PostprocessParameter.SetPS(ShaderRHI, Context, TStaticSamplerState<SF_Point,AM_Border,AM_Border,AM_Clamp>::GetRHI());
 
-		DeferredParameters.Set(RHICmdList, ShaderRHI, Context.View);
+		DeferredParameters.Set(Context.RHICmdList, ShaderRHI, Context.View);
 
 		{
 			FVector4 DepthOfFieldParamValues[2];
 
 			FRCPassPostProcessBokehDOF::ComputeDepthOfFieldParams(Context, DepthOfFieldParamValues);
 
-			SetShaderValueArray(RHICmdList, ShaderRHI, DepthOfFieldParams, DepthOfFieldParamValues, 2);
+			SetShaderValueArray(Context.RHICmdList, ShaderRHI, DepthOfFieldParams, DepthOfFieldParamValues, 2);
 		}
 	}
 };
@@ -107,16 +107,13 @@ void FRCPassPostProcessDOFSetup::Process(FRenderingCompositePassContext& Context
 	const FSceneRenderTargetItem& DestRenderTarget0 = PassOutputs[0].RequestSurface(Context);
 	const FSceneRenderTargetItem& DestRenderTarget1 = bNearBlurEnabled ? PassOutputs[1].RequestSurface(Context) : FSceneRenderTargetItem();
 
-	//@todo-rco: RHIPacketList
-	FRHICommandList& RHICmdList = FRHICommandList::GetNullRef();
-
 	// Set the view family's render target/viewport.
 	FTextureRHIParamRef RenderTargets[2] =
 	{
 		DestRenderTarget0.TargetableTexture,
 		DestRenderTarget1.TargetableTexture
 	};
-	RHISetRenderTargets(NumRenderTargets, RenderTargets, FTextureRHIParamRef(), 0, NULL);
+	SetRenderTargets(Context.RHICmdList, NumRenderTargets, RenderTargets, FTextureRHIParamRef(), 0, NULL);
 	
 	FLinearColor ClearColors[2] = 
 	{
@@ -124,38 +121,42 @@ void FRCPassPostProcessDOFSetup::Process(FRenderingCompositePassContext& Context
 		FLinearColor(0, 0, 0, 0)
 	};
 	// is optimized away if possible (RT size=view size, )
-	RHIClearMRT(true, NumRenderTargets, ClearColors, false, 1.0f, false, 0, DestRect);
+	Context.RHICmdList.ClearMRT(true, NumRenderTargets, ClearColors, false, 1.0f, false, 0, DestRect);
 
 	Context.SetViewportAndCallRHI(0, 0, 0.0f, DestSize.X, DestSize.Y, 1.0f );
 
 	// set the state
-	RHISetBlendState(TStaticBlendState<>::GetRHI());
-	RHISetRasterizerState(TStaticRasterizerState<>::GetRHI());
-	RHISetDepthStencilState(TStaticDepthStencilState<false,CF_Always>::GetRHI());
+	Context.RHICmdList.SetBlendState(TStaticBlendState<>::GetRHI());
+	Context.RHICmdList.SetRasterizerState(TStaticRasterizerState<>::GetRHI());
+	Context.RHICmdList.SetDepthStencilState(TStaticDepthStencilState<false, CF_Always>::GetRHI());
 	
 	TShaderMapRef<FPostProcessVS> VertexShader(GetGlobalShaderMap());
 
 	if (bNearBlurEnabled)
 	{
 		static FGlobalBoundShaderState BoundShaderState;
+		Context.RHICmdList.CheckIsNull(); // need new approach for "static FGlobalBoundShaderState" for parallel rendering
+
 		TShaderMapRef< FPostProcessDOFSetupPS<1> > PixelShader(GetGlobalShaderMap());
-		SetGlobalBoundShaderState(BoundShaderState, GFilterVertexDeclaration.VertexDeclarationRHI, *VertexShader, *PixelShader);
+		SetGlobalBoundShaderState(Context.RHICmdList, BoundShaderState, GFilterVertexDeclaration.VertexDeclarationRHI, *VertexShader, *PixelShader);
 		
-		PixelShader->SetParameters(RHICmdList, Context);
+		PixelShader->SetParameters(Context);
 	}
 	else
 	{
 		static FGlobalBoundShaderState BoundShaderState;
+		Context.RHICmdList.CheckIsNull(); // need new approach for "static FGlobalBoundShaderState" for parallel rendering
 		TShaderMapRef< FPostProcessDOFSetupPS<0> > PixelShader(GetGlobalShaderMap());
-		SetGlobalBoundShaderState(BoundShaderState, GFilterVertexDeclaration.VertexDeclarationRHI, *VertexShader, *PixelShader);
+		SetGlobalBoundShaderState(Context.RHICmdList, BoundShaderState, GFilterVertexDeclaration.VertexDeclarationRHI, *VertexShader, *PixelShader);
 
-		PixelShader->SetParameters(RHICmdList, Context);
+		PixelShader->SetParameters(Context);
 	}
 
-	VertexShader->SetParameters(RHICmdList, Context);
+	VertexShader->SetParameters(Context);
 
 	// Draw a quad mapping scene color to the view's render target
 	DrawRectangle(
+		Context.RHICmdList,
 		DestRect.Min.X, DestRect.Min.Y,
 		DestRect.Width() + 1, DestRect.Height() + 1,
 		SrcRect.Min.X, SrcRect.Min.Y,
@@ -165,8 +166,8 @@ void FRCPassPostProcessDOFSetup::Process(FRenderingCompositePassContext& Context
 		*VertexShader,
 		EDRF_UseTriangleOptimization);
 
-	RHICopyToResolveTarget(DestRenderTarget0.TargetableTexture, DestRenderTarget0.ShaderResourceTexture, false, FResolveParams());
-	RHICopyToResolveTarget(DestRenderTarget1.TargetableTexture, DestRenderTarget1.ShaderResourceTexture, false, FResolveParams());
+	Context.RHICmdList.CopyToResolveTarget(DestRenderTarget0.TargetableTexture, DestRenderTarget0.ShaderResourceTexture, false, FResolveParams());
+	Context.RHICmdList.CopyToResolveTarget(DestRenderTarget1.TargetableTexture, DestRenderTarget1.ShaderResourceTexture, false, FResolveParams());
 }
 
 FPooledRenderTargetDesc FRCPassPostProcessDOFSetup::ComputeOutputDesc(EPassOutputId InPassOutputId) const
@@ -236,14 +237,14 @@ public:
 		return bShaderHasOutdatedParameters;
 	}
 
-	void SetParameters(FRHICommandList& RHICmdList, const FRenderingCompositePassContext& Context)
+	void SetParameters(const FRenderingCompositePassContext& Context)
 	{
 		const FPixelShaderRHIParamRef ShaderRHI = GetPixelShader();
 
-		FGlobalShader::SetParameters(RHICmdList, ShaderRHI, Context.View);
+		FGlobalShader::SetParameters(Context.RHICmdList, ShaderRHI, Context.View);
 
-		DeferredParameters.Set(RHICmdList, ShaderRHI, Context.View);
-		PostprocessParameter.SetPS(RHICmdList, ShaderRHI, Context, TStaticSamplerState<SF_Bilinear,AM_Clamp,AM_Clamp,AM_Clamp>::GetRHI());
+		DeferredParameters.Set(Context.RHICmdList, ShaderRHI, Context.View);
+		PostprocessParameter.SetPS(ShaderRHI, Context, TStaticSamplerState<SF_Bilinear,AM_Clamp,AM_Clamp,AM_Clamp>::GetRHI());
 		
 		// Compute out of bounds UVs in the source texture.
 		FVector4 Bounds;
@@ -252,7 +253,7 @@ public:
 		Bounds.Z = (((float)(Context.View.ViewRect.Max.X & (~1))) - 3.0f) / ((float)(GSceneRenderTargets.GetBufferSizeXY().X));
 		Bounds.W = (((float)(Context.View.ViewRect.Max.Y & (~1))) - 3.0f) / ((float)(GSceneRenderTargets.GetBufferSizeXY().Y));
 
-		SetShaderValue(RHICmdList, ShaderRHI, DepthOfFieldUVLimit, Bounds);
+		SetShaderValue(Context.RHICmdList, ShaderRHI, DepthOfFieldUVLimit, Bounds);
 	}
 };
 
@@ -282,43 +283,43 @@ void FRCPassPostProcessDOFRecombine::Process(FRenderingCompositePassContext& Con
 
 	const FSceneRenderTargetItem& DestRenderTarget = PassOutputs[0].RequestSurface(Context);
 
-	//@todo-rco: RHIPacketList
-	FRHICommandList& RHICmdList = FRHICommandList::GetNullRef();
-
 	// Set the view family's render target/viewport.
-	RHISetRenderTarget(DestRenderTarget.TargetableTexture, FTextureRHIRef());	
+	SetRenderTarget(Context.RHICmdList, DestRenderTarget.TargetableTexture, FTextureRHIRef());
 
 	// is optimized away if possible (RT size=view size, )
-	RHIClear(true, FLinearColor::Black, false, 1.0f, false, 0, View.ViewRect);
+	Context.RHICmdList.Clear(true, FLinearColor::Black, false, 1.0f, false, 0, View.ViewRect);
 
 	Context.SetViewportAndCallRHI(View.ViewRect);
 
 	// set the state
-	RHISetBlendState(TStaticBlendState<>::GetRHI());
-	RHISetRasterizerState(TStaticRasterizerState<>::GetRHI());
-	RHISetDepthStencilState(TStaticDepthStencilState<false,CF_Always>::GetRHI());
+	Context.RHICmdList.SetBlendState(TStaticBlendState<>::GetRHI());
+	Context.RHICmdList.SetRasterizerState(TStaticRasterizerState<>::GetRHI());
+	Context.RHICmdList.SetDepthStencilState(TStaticDepthStencilState<false, CF_Always>::GetRHI());
 
 	TShaderMapRef<FPostProcessVS> VertexShader(GetGlobalShaderMap());
 
 	if (bNearBlurEnabled)
 	{
 		static FGlobalBoundShaderState BoundShaderState;
+		Context.RHICmdList.CheckIsNull(); // need new approach for "static FGlobalBoundShaderState" for parallel rendering
 		TShaderMapRef< FPostProcessDOFRecombinePS<1> > PixelShader(GetGlobalShaderMap());
-		SetGlobalBoundShaderState(BoundShaderState, GFilterVertexDeclaration.VertexDeclarationRHI, *VertexShader, *PixelShader);
-		PixelShader->SetParameters(RHICmdList, Context);
+		SetGlobalBoundShaderState(Context.RHICmdList, BoundShaderState, GFilterVertexDeclaration.VertexDeclarationRHI, *VertexShader, *PixelShader);
+		PixelShader->SetParameters(Context);
 	}
 	else
 	{
 		static FGlobalBoundShaderState BoundShaderState;
+		Context.RHICmdList.CheckIsNull(); // need new approach for "static FGlobalBoundShaderState" for parallel rendering
 		TShaderMapRef< FPostProcessDOFRecombinePS<0> > PixelShader(GetGlobalShaderMap());
-		SetGlobalBoundShaderState(BoundShaderState, GFilterVertexDeclaration.VertexDeclarationRHI, *VertexShader, *PixelShader);
-		PixelShader->SetParameters(RHICmdList, Context);
+		SetGlobalBoundShaderState(Context.RHICmdList, BoundShaderState, GFilterVertexDeclaration.VertexDeclarationRHI, *VertexShader, *PixelShader);
+		PixelShader->SetParameters(Context);
 	}
 
-	VertexShader->SetParameters(RHICmdList, Context);
+	VertexShader->SetParameters(Context);
 
 	// Draw a quad mapping scene color to the view's render target
 	DrawRectangle(
+		Context.RHICmdList,
 		0, 0,
 		View.ViewRect.Width(), View.ViewRect.Height(),
 		HalfResViewRect.Min.X, HalfResViewRect.Min.Y, 
@@ -328,7 +329,7 @@ void FRCPassPostProcessDOFRecombine::Process(FRenderingCompositePassContext& Con
 		*VertexShader,
 		EDRF_UseTriangleOptimization);
 
-	RHICopyToResolveTarget(DestRenderTarget.TargetableTexture, DestRenderTarget.ShaderResourceTexture, false, FResolveParams());
+	Context.RHICmdList.CopyToResolveTarget(DestRenderTarget.TargetableTexture, DestRenderTarget.ShaderResourceTexture, false, FResolveParams());
 }
 
 FPooledRenderTargetDesc FRCPassPostProcessDOFRecombine::ComputeOutputDesc(EPassOutputId InPassOutputId) const
diff --git a/Engine/Source/Runtime/Renderer/Private/PostProcess/PostProcessDownsample.cpp b/Engine/Source/Runtime/Renderer/Private/PostProcess/PostProcessDownsample.cpp
--- a/Engine/Source/Runtime/Renderer/Private/PostProcess/PostProcessDownsample.cpp
+++ b/Engine/Source/Runtime/Renderer/Private/PostProcess/PostProcessDownsample.cpp
@@ -52,13 +52,13 @@ public:
 		return bShaderHasOutdatedParameters;
 	}
 
-	void SetParameters(FRHICommandList& RHICmdList, const FRenderingCompositePassContext& Context)
+	void SetParameters(const FRenderingCompositePassContext& Context)
 	{
 		const FPixelShaderRHIParamRef ShaderRHI = GetPixelShader();
 
-		FGlobalShader::SetParameters(RHICmdList, ShaderRHI, Context.View);
-		DeferredParameters.Set(RHICmdList, ShaderRHI, Context.View);
-		PostprocessParameter.SetPS(RHICmdList, ShaderRHI, Context, TStaticSamplerState<SF_Bilinear,AM_Clamp,AM_Clamp,AM_Clamp>::GetRHI());
+		FGlobalShader::SetParameters(Context.RHICmdList, ShaderRHI, Context.View);
+		DeferredParameters.Set(Context.RHICmdList, ShaderRHI, Context.View);
+		PostprocessParameter.SetPS(ShaderRHI, Context, TStaticSamplerState<SF_Bilinear,AM_Clamp,AM_Clamp,AM_Clamp>::GetRHI());
 	}
 
 	static const TCHAR* GetSourceFilename()
@@ -108,11 +108,11 @@ public:
 		return bShaderHasOutdatedParameters;
 	}
 
-	void SetParameters(FRHICommandList& RHICmdList, const FRenderingCompositePassContext& Context)
+	void SetParameters(const FRenderingCompositePassContext& Context)
 	{
 		const FVertexShaderRHIParamRef ShaderRHI = GetVertexShader();
 
-		FGlobalShader::SetParameters(RHICmdList, ShaderRHI, Context.View);
+		FGlobalShader::SetParameters(Context.RHICmdList, ShaderRHI, Context.View);
 
 		const FPooledRenderTargetDesc* InputDesc = Context.Pass->GetInputDesc(ePId_Input0);
 
@@ -137,17 +137,18 @@ FRCPassPostProcessDownsample::FRCPassPostProcessDownsample(EPixelFormat InOverri
 
 
 template <uint32 Method>
-void FRCPassPostProcessDownsample::SetShader(FRHICommandList& RHICmdList, const FRenderingCompositePassContext& Context)
+void FRCPassPostProcessDownsample::SetShader(const FRenderingCompositePassContext& Context)
 {
 	TShaderMapRef<FPostProcessDownsampleVS> VertexShader(GetGlobalShaderMap());
 	TShaderMapRef<FPostProcessDownsamplePS<Method> > PixelShader(GetGlobalShaderMap());
 
 	static FGlobalBoundShaderState BoundShaderState;
+	Context.RHICmdList.CheckIsNull(); // need new approach for "static FGlobalBoundShaderState" for parallel rendering
 
-	SetGlobalBoundShaderState(BoundShaderState, GFilterVertexDeclaration.VertexDeclarationRHI, *VertexShader, *PixelShader);
+	SetGlobalBoundShaderState(Context.RHICmdList, BoundShaderState, GFilterVertexDeclaration.VertexDeclarationRHI, *VertexShader, *PixelShader);
 
-	PixelShader->SetParameters(RHICmdList, Context);
-	VertexShader->SetParameters(RHICmdList, Context);
+	PixelShader->SetParameters(Context);
+	VertexShader->SetParameters(Context);
 }
 
 
@@ -174,18 +175,15 @@ void FRCPassPostProcessDownsample::Process(FRenderingCompositePassContext& Conte
 
 	const FSceneRenderTargetItem& DestRenderTarget = PassOutputs[0].RequestSurface(Context);
 
-	//@todo-rco: RHIPacketList
-	FRHICommandList& RHICmdList = FRHICommandList::GetNullRef();
-
 	// Set the view family's render target/viewport.
-	RHISetRenderTarget(DestRenderTarget.TargetableTexture, FTextureRHIRef());
+	SetRenderTarget(Context.RHICmdList, DestRenderTarget.TargetableTexture, FTextureRHIRef());
 
 	Context.SetViewportAndCallRHI(0, 0, 0.0f, DestSize.X, DestSize.Y, 1.0f );
 
 	// set the state
-	RHISetBlendState(TStaticBlendState<>::GetRHI());
-	RHISetRasterizerState(TStaticRasterizerState<>::GetRHI());
-	RHISetDepthStencilState(TStaticDepthStencilState<false,CF_Always>::GetRHI());
+	Context.RHICmdList.SetBlendState(TStaticBlendState<>::GetRHI());
+	Context.RHICmdList.SetRasterizerState(TStaticRasterizerState<>::GetRHI());
+	Context.RHICmdList.SetDepthStencilState(TStaticDepthStencilState<false, CF_Always>::GetRHI());
 
 	// InflateSize increases the size of the source/dest rectangle to compensate for bilinear reads and UIBlur pass requirements.
 	int32 InflateSize;
@@ -194,18 +192,18 @@ void FRCPassPostProcessDownsample::Process(FRenderingCompositePassContext& Conte
 	{
 		// also put depth in alpha
 		InflateSize = 2;
-		SetShader<2>(RHICmdList, Context);
+		SetShader<2>(Context);
 	}
 	else
 	{
 		if (Quality == 0)
 		{
-			SetShader<0>(RHICmdList, Context);
+			SetShader<0>(Context);
 			InflateSize = 1;
 		}
 		else
 		{
-			SetShader<1>(RHICmdList, Context);
+			SetShader<1>(Context);
 			InflateSize = 2;
 		}
 	}
@@ -218,7 +216,7 @@ void FRCPassPostProcessDownsample::Process(FRenderingCompositePassContext& Conte
 	// Otherwise perform the clear when the dest rectangle has been computed.
 	if (bHasMultipleQuads || Context.View.GetFeatureLevel() == ERHIFeatureLevel::ES2)
 	{
-		RHIClear(true, FLinearColor(0, 0, 0, 0), false, 1.0f, false, 0, FIntRect());
+		Context.RHICmdList.Clear(true, FLinearColor(0, 0, 0, 0), false, 1.0f, false, 0, FIntRect());
 		bHasCleared = true;
 	}
 
@@ -234,11 +232,12 @@ void FRCPassPostProcessDownsample::Process(FRenderingCompositePassContext& Conte
 
 			if (bHasCleared == false)
 			{
-				RHIClear(true, FLinearColor(0, 0, 0, 0), false, 1.0f, false, 0, DestRect);
+				Context.RHICmdList.Clear(true, FLinearColor(0, 0, 0, 0), false, 1.0f, false, 0, DestRect);
 			}
 
 			// Draw a quad mapping scene color to the view's render target
 			DrawRectangle(
+				Context.RHICmdList,
 				DestRect.Min.X, DestRect.Min.Y,
 				DestRect.Width(), DestRect.Height(),
 				SrcRect.Min.X, SrcRect.Min.Y,
@@ -276,11 +275,12 @@ void FRCPassPostProcessDownsample::Process(FRenderingCompositePassContext& Conte
 
 				if (bHasCleared == false)
 				{
-					RHIClear(true, FLinearColor(0, 0, 0, 0), false, 1.0f, false, 0, DestRect);
+					Context.RHICmdList.Clear(true, FLinearColor(0, 0, 0, 0), false, 1.0f, false, 0, DestRect);
 				}
 
 				// Draw a quad mapping scene color to the view's render target
 				DrawRectangle(
+					Context.RHICmdList,
 					DestRect.Min.X, DestRect.Min.Y,
 					DestRect.Width(), DestRect.Height(),
 					SrcRect.Min.X, SrcRect.Min.Y,
@@ -297,7 +297,7 @@ void FRCPassPostProcessDownsample::Process(FRenderingCompositePassContext& Conte
 		break;
 	}
 
-	RHICopyToResolveTarget(DestRenderTarget.TargetableTexture, DestRenderTarget.ShaderResourceTexture, false, FResolveParams());
+	Context.RHICmdList.CopyToResolveTarget(DestRenderTarget.TargetableTexture, DestRenderTarget.ShaderResourceTexture, false, FResolveParams());
 }
 
 bool FRCPassPostProcessDownsample::IsDepthInputAvailable() const
diff --git a/Engine/Source/Runtime/Renderer/Private/PostProcess/PostProcessDownsample.h b/Engine/Source/Runtime/Renderer/Private/PostProcess/PostProcessDownsample.h
--- a/Engine/Source/Runtime/Renderer/Private/PostProcess/PostProcessDownsample.h
+++ b/Engine/Source/Runtime/Renderer/Private/PostProcess/PostProcessDownsample.h
@@ -31,7 +31,7 @@ public:
 	virtual FPooledRenderTargetDesc ComputeOutputDesc(EPassOutputId InPassOutputId) const;
 
 private:
-	template <uint32 Method> static void SetShader(FRHICommandList& RHICmdList, const FRenderingCompositePassContext& Context);
+	template <uint32 Method> static void SetShader(const FRenderingCompositePassContext& Context);
 
 	EPixelFormat OverrideFormat;
 	// explained in constructor
diff --git a/Engine/Source/Runtime/Renderer/Private/PostProcess/PostProcessEyeAdaptation.cpp b/Engine/Source/Runtime/Renderer/Private/PostProcess/PostProcessEyeAdaptation.cpp
--- a/Engine/Source/Runtime/Renderer/Private/PostProcess/PostProcessEyeAdaptation.cpp
+++ b/Engine/Source/Runtime/Renderer/Private/PostProcess/PostProcessEyeAdaptation.cpp
@@ -41,19 +41,19 @@ public:
 		EyeAdaptationParams.Bind(Initializer.ParameterMap, TEXT("EyeAdaptationParams"));
 	}
 
-	void SetPS(FRHICommandList& RHICmdList, const FRenderingCompositePassContext& Context)
+	void SetPS(const FRenderingCompositePassContext& Context)
 	{
 		const FPixelShaderRHIParamRef ShaderRHI = GetPixelShader();
 
-		FGlobalShader::SetParameters(RHICmdList, ShaderRHI, Context.View);
+		FGlobalShader::SetParameters(Context.RHICmdList, ShaderRHI, Context.View);
 
-		PostprocessParameter.SetPS(RHICmdList, ShaderRHI, Context, TStaticSamplerState<SF_Bilinear,AM_Clamp,AM_Clamp,AM_Clamp>::GetRHI());
+		PostprocessParameter.SetPS(ShaderRHI, Context, TStaticSamplerState<SF_Bilinear,AM_Clamp,AM_Clamp,AM_Clamp>::GetRHI());
 
 		{
 			FVector4 Temp[3];
 
 			FRCPassPostProcessEyeAdaptation::ComputeEyeAdaptationParamsValue(Context.View, Temp);
-			SetShaderValueArray(RHICmdList, ShaderRHI, EyeAdaptationParams, Temp, 3);
+			SetShaderValueArray(Context.RHICmdList, ShaderRHI, EyeAdaptationParams, Temp, 3);
 		}
 	}
 
@@ -81,29 +81,28 @@ void FRCPassPostProcessEyeAdaptation::Process(FRenderingCompositePassContext& Co
 	FIntPoint DestSize = EyeAdaptation->GetDesc().Extent;
 
 	// we render to our own output render target, not the intermediate one created by the compositing system
-	//@todo-rco: RHIPacketList
-	FRHICommandList& RHICmdList = FRHICommandList::GetNullRef();
-
 	// Set the view family's render target/viewport.
-	RHISetRenderTarget(EyeAdaptation->GetRenderTargetItem().TargetableTexture, FTextureRHIRef());	
+	SetRenderTarget(Context.RHICmdList, EyeAdaptation->GetRenderTargetItem().TargetableTexture, FTextureRHIRef());
 	Context.SetViewportAndCallRHI(0, 0, 0.0f, DestSize.X, DestSize.Y, 1.0f );
 
 	// set the state
-	RHISetBlendState(TStaticBlendState<>::GetRHI());
-	RHISetRasterizerState(TStaticRasterizerState<>::GetRHI());
-	RHISetDepthStencilState(TStaticDepthStencilState<false,CF_Always>::GetRHI());
+	Context.RHICmdList.SetBlendState(TStaticBlendState<>::GetRHI());
+	Context.RHICmdList.SetRasterizerState(TStaticRasterizerState<>::GetRHI());
+	Context.RHICmdList.SetDepthStencilState(TStaticDepthStencilState<false, CF_Always>::GetRHI());
 
 	TShaderMapRef<FPostProcessVS> VertexShader(GetGlobalShaderMap());
 	TShaderMapRef<FPostProcessEyeAdaptationPS> PixelShader(GetGlobalShaderMap());
 
 	static FGlobalBoundShaderState BoundShaderState;
+	Context.RHICmdList.CheckIsNull(); // need new approach for "static FGlobalBoundShaderState" for parallel rendering
 
-	SetGlobalBoundShaderState(BoundShaderState, GFilterVertexDeclaration.VertexDeclarationRHI, *VertexShader, *PixelShader);
+	SetGlobalBoundShaderState(Context.RHICmdList, BoundShaderState, GFilterVertexDeclaration.VertexDeclarationRHI, *VertexShader, *PixelShader);
 
-	PixelShader->SetPS(RHICmdList, Context);
+	PixelShader->SetPS(Context);
 
 	// Draw a quad mapping scene color to the view's render target
 	DrawRectangle(
+		Context.RHICmdList,
 		0, 0,
 		DestSize.X, DestSize.Y,
 		0, 0,
@@ -113,7 +112,7 @@ void FRCPassPostProcessEyeAdaptation::Process(FRenderingCompositePassContext& Co
 		*VertexShader,
 		EDRF_UseTriangleOptimization);
 
-	RHICopyToResolveTarget(EyeAdaptation->GetRenderTargetItem().TargetableTexture, EyeAdaptation->GetRenderTargetItem().ShaderResourceTexture, false, FResolveParams());
+	Context.RHICmdList.CopyToResolveTarget(EyeAdaptation->GetRenderTargetItem().TargetableTexture, EyeAdaptation->GetRenderTargetItem().ShaderResourceTexture, false, FResolveParams());
 }
 
 void FRCPassPostProcessEyeAdaptation::ComputeEyeAdaptationParamsValue(const FViewInfo& View, FVector4 Out[3])
diff --git a/Engine/Source/Runtime/Renderer/Private/PostProcess/PostProcessGBufferHints.cpp b/Engine/Source/Runtime/Renderer/Private/PostProcess/PostProcessGBufferHints.cpp
--- a/Engine/Source/Runtime/Renderer/Private/PostProcess/PostProcessGBufferHints.cpp
+++ b/Engine/Source/Runtime/Renderer/Private/PostProcess/PostProcessGBufferHints.cpp
@@ -47,23 +47,23 @@ public:
 		MiniFontTexture.Bind(Initializer.ParameterMap, TEXT("MiniFontTexture"));
 	}
 
-	void SetPS(FRHICommandList& RHICmdList, const FRenderingCompositePassContext& Context)
+	void SetPS(const FRenderingCompositePassContext& Context)
 	{
 		const FPixelShaderRHIParamRef ShaderRHI = GetPixelShader();
 
-		FGlobalShader::SetParameters(RHICmdList, ShaderRHI, Context.View);
+		FGlobalShader::SetParameters(Context.RHICmdList, ShaderRHI, Context.View);
 
-		PostprocessParameter.SetPS(RHICmdList, ShaderRHI, Context, TStaticSamplerState<SF_Point,AM_Clamp,AM_Clamp,AM_Clamp>::GetRHI());
-		DeferredParameters.Set(RHICmdList, ShaderRHI, Context.View);
+		PostprocessParameter.SetPS(ShaderRHI, Context, TStaticSamplerState<SF_Point,AM_Clamp,AM_Clamp,AM_Clamp>::GetRHI());
+		DeferredParameters.Set(Context.RHICmdList, ShaderRHI, Context.View);
 
 		{
 			FVector4 Temp[3];
 
 			FRCPassPostProcessEyeAdaptation::ComputeEyeAdaptationParamsValue(Context.View, Temp);
-			SetShaderValueArray(RHICmdList, ShaderRHI, EyeAdaptationParams, Temp, 3);
+			SetShaderValueArray(Context.RHICmdList, ShaderRHI, EyeAdaptationParams, Temp, 3);
 		}
 
-		SetTextureParameter(RHICmdList, ShaderRHI, MiniFontTexture, GEngine->MiniFontTexture ? GEngine->MiniFontTexture->Resource->TextureRHI : GSystemTextures.WhiteDummy->GetRenderTargetItem().TargetableTexture);
+		SetTextureParameter(Context.RHICmdList, ShaderRHI, MiniFontTexture, GEngine->MiniFontTexture ? GEngine->MiniFontTexture->Resource->TextureRHI : GSystemTextures.WhiteDummy->GetRenderTargetItem().TargetableTexture);
 	}
 	
 	// FShader interface.
@@ -97,29 +97,28 @@ void FRCPassPostProcessGBufferHints::Process(FRenderingCompositePassContext& Con
 
 	const FSceneRenderTargetItem& DestRenderTarget = PassOutputs[0].RequestSurface(Context);
 
-	//@todo-rco: RHIPacketList
-	FRHICommandList& RHICmdList = FRHICommandList::GetNullRef();
-
 	// Set the view family's render target/viewport.
-	RHISetRenderTarget(DestRenderTarget.TargetableTexture, FTextureRHIRef());	
+	SetRenderTarget(Context.RHICmdList, DestRenderTarget.TargetableTexture, FTextureRHIRef());
 	Context.SetViewportAndCallRHI(DestRect);
 
 	// set the state
-	RHISetBlendState(TStaticBlendState<>::GetRHI());
-	RHISetRasterizerState(TStaticRasterizerState<>::GetRHI());
-	RHISetDepthStencilState(TStaticDepthStencilState<false,CF_Always>::GetRHI());
+	Context.RHICmdList.SetBlendState(TStaticBlendState<>::GetRHI());
+	Context.RHICmdList.SetRasterizerState(TStaticRasterizerState<>::GetRHI());
+	Context.RHICmdList.SetDepthStencilState(TStaticDepthStencilState<false, CF_Always>::GetRHI());
 
 	TShaderMapRef<FPostProcessVS> VertexShader(GetGlobalShaderMap());
 	TShaderMapRef<FPostProcessGBufferHintsPS> PixelShader(GetGlobalShaderMap());
 
 	static FGlobalBoundShaderState BoundShaderState;
+	Context.RHICmdList.CheckIsNull(); // need new approach for "static FGlobalBoundShaderState" for parallel rendering
 
-	SetGlobalBoundShaderState(BoundShaderState, GFilterVertexDeclaration.VertexDeclarationRHI, *VertexShader, *PixelShader);
+	SetGlobalBoundShaderState(Context.RHICmdList, BoundShaderState, GFilterVertexDeclaration.VertexDeclarationRHI, *VertexShader, *PixelShader);
 
-	PixelShader->SetPS(RHICmdList, Context);
+	PixelShader->SetPS(Context);
 
 	// Draw a quad mapping scene color to the view's render target
 	DrawRectangle(
+		Context.RHICmdList,
 		0, 0,
 		DestRect.Width(), DestRect.Height(),
 		SrcRect.Min.X, SrcRect.Min.Y,
@@ -150,6 +149,7 @@ void FRCPassPostProcessGBufferHints::Process(FRenderingCompositePassContext& Con
 			return Texture;
 		}
 	} TempRenderTarget(View, (const FTexture2DRHIRef&)DestRenderTarget.TargetableTexture);
+	Context.RHICmdList.CheckIsNull(); // need canvas to have a cmd list
 
 	FCanvas Canvas(&TempRenderTarget, NULL, ViewFamily.CurrentRealTime, ViewFamily.CurrentWorldTime, ViewFamily.DeltaWorldTime);
 
@@ -173,7 +173,7 @@ void FRCPassPostProcessGBufferHints::Process(FRenderingCompositePassContext& Con
 
 	Canvas.Flush();
 
-	RHICopyToResolveTarget(DestRenderTarget.TargetableTexture, DestRenderTarget.ShaderResourceTexture, false, FResolveParams());
+	Context.RHICmdList.CopyToResolveTarget(DestRenderTarget.TargetableTexture, DestRenderTarget.ShaderResourceTexture, false, FResolveParams());
 }
 
 FPooledRenderTargetDesc FRCPassPostProcessGBufferHints::ComputeOutputDesc(EPassOutputId InPassOutputId) const
diff --git a/Engine/Source/Runtime/Renderer/Private/PostProcess/PostProcessHMD.cpp b/Engine/Source/Runtime/Renderer/Private/PostProcess/PostProcessHMD.cpp
--- a/Engine/Source/Runtime/Renderer/Private/PostProcess/PostProcessHMD.cpp
+++ b/Engine/Source/Runtime/Renderer/Private/PostProcess/PostProcessHMD.cpp
@@ -72,18 +72,18 @@ public:
 		EyeToSrcUVOffset.Bind(Initializer.ParameterMap, TEXT("EyeToSrcUVOffset"));
 	}
 
-	void SetVS(FRHICommandList& RHICmdList, const FRenderingCompositePassContext& Context, EStereoscopicPass StereoPass)
+	void SetVS(const FRenderingCompositePassContext& Context, EStereoscopicPass StereoPass)
 	{
 		const FVertexShaderRHIParamRef ShaderRHI = GetVertexShader();
 
-		FGlobalShader::SetParameters(RHICmdList, ShaderRHI, Context.View);
+		FGlobalShader::SetParameters(Context.RHICmdList, ShaderRHI, Context.View);
 
 		check(GEngine->HMDDevice.IsValid());
 		FVector2D EyeToSrcUVScaleValue;
 		FVector2D EyeToSrcUVOffsetValue;
 		GEngine->HMDDevice->GetEyeRenderParams_RenderThread(StereoPass, EyeToSrcUVScaleValue, EyeToSrcUVOffsetValue);
-		SetShaderValue(RHICmdList, ShaderRHI, EyeToSrcUVScale, EyeToSrcUVScaleValue);
-		SetShaderValue(RHICmdList, ShaderRHI, EyeToSrcUVOffset, EyeToSrcUVOffsetValue);
+		SetShaderValue(Context.RHICmdList, ShaderRHI, EyeToSrcUVScale, EyeToSrcUVScaleValue);
+		SetShaderValue(Context.RHICmdList, ShaderRHI, EyeToSrcUVOffset, EyeToSrcUVOffsetValue);
 	}
 
 	// FShader interface.
@@ -123,14 +123,14 @@ public:
 
 	}
 
-	void SetPS(FRHICommandList& RHICmdList, const FRenderingCompositePassContext& Context, FIntRect SrcRect, FIntPoint SrcBufferSize, EStereoscopicPass StereoPass, FMatrix& QuadTexTransform)
+	void SetPS(const FRenderingCompositePassContext& Context, FIntRect SrcRect, FIntPoint SrcBufferSize, EStereoscopicPass StereoPass, FMatrix& QuadTexTransform)
 	{
 		const FPixelShaderRHIParamRef ShaderRHI = GetPixelShader();
 
-		FGlobalShader::SetParameters(RHICmdList, ShaderRHI, Context.View);
+		FGlobalShader::SetParameters(Context.RHICmdList, ShaderRHI, Context.View);
 
-		PostprocessParameter.SetPS(RHICmdList, ShaderRHI, Context, TStaticSamplerState<SF_Bilinear, AM_Clamp, AM_Clamp, AM_Clamp>::GetRHI());
-		DeferredParameters.Set(RHICmdList, ShaderRHI, Context.View);
+		PostprocessParameter.SetPS(ShaderRHI, Context, TStaticSamplerState<SF_Bilinear, AM_Clamp, AM_Clamp, AM_Clamp>::GetRHI());
+		DeferredParameters.Set(Context.RHICmdList, ShaderRHI, Context.View);
 	}
 
 	// FShader interface.
@@ -174,19 +174,16 @@ void FRCPassPostProcessHMD::Process(FRenderingCompositePassContext& Context)
 
     const FSceneRenderTargetItem& DestRenderTarget = PassOutputs[0].RequestSurface(Context);
 
-	//@todo-rco: RHIPacketList
-	FRHICommandList& RHICmdList = FRHICommandList::GetNullRef();
-
 	// Set the view family's render target/viewport.
-	RHISetRenderTarget(DestRenderTarget.TargetableTexture, FTextureRHIRef());	
+	SetRenderTarget(Context.RHICmdList, DestRenderTarget.TargetableTexture, FTextureRHIRef());
 
 	Context.SetViewportAndCallRHI(DestRect);
-	RHIClear(true, FLinearColor::Black, false, 1.0f, false, 0, FIntRect());
+	Context.RHICmdList.Clear(true, FLinearColor::Black, false, 1.0f, false, 0, FIntRect());
 
 	// set the state
-	RHISetBlendState(TStaticBlendState<>::GetRHI());
-	RHISetRasterizerState(TStaticRasterizerState<>::GetRHI());
-	RHISetDepthStencilState(TStaticDepthStencilState<false,CF_Always>::GetRHI());
+	Context.RHICmdList.SetBlendState(TStaticBlendState<>::GetRHI());
+	Context.RHICmdList.SetRasterizerState(TStaticRasterizerState<>::GetRHI());
+	Context.RHICmdList.SetDepthStencilState(TStaticDepthStencilState<false, CF_Always>::GetRHI());
 
 	FMatrix QuadTexTransform = FMatrix::Identity;
 	FMatrix QuadPosTransform = FMatrix::Identity;
@@ -197,14 +194,15 @@ void FRCPassPostProcessHMD::Process(FRenderingCompositePassContext& Context)
 		TShaderMapRef<FPostProcessHMDVS> VertexShader(GetGlobalShaderMap());
 		TShaderMapRef<FPostProcessHMDPS> PixelShader(GetGlobalShaderMap());
 		static FGlobalBoundShaderState BoundShaderState;
-		SetGlobalBoundShaderState(BoundShaderState, GDistortionVertexDeclaration.VertexDeclarationRHI, *VertexShader, *PixelShader);
-		VertexShader->SetVS(RHICmdList, Context, View.StereoPass);
-		PixelShader->SetPS(RHICmdList, Context, SrcRect, SrcSize, View.StereoPass, QuadTexTransform);
+		Context.RHICmdList.CheckIsNull(); // need new approach for "static FGlobalBoundShaderState" for parallel rendering
+		SetGlobalBoundShaderState(Context.RHICmdList, BoundShaderState, GDistortionVertexDeclaration.VertexDeclarationRHI, *VertexShader, *PixelShader);
+		VertexShader->SetVS(Context, View.StereoPass);
+		PixelShader->SetPS(Context, SrcRect, SrcSize, View.StereoPass, QuadTexTransform);
 	}
-
+	Context.RHICmdList.CheckIsNull(); // this call may not properly use the cmd list in the context
 	GEngine->HMDDevice->DrawDistortionMesh_RenderThread(Context, View, SrcSize);
 
-	RHICopyToResolveTarget(DestRenderTarget.TargetableTexture, DestRenderTarget.ShaderResourceTexture, false, FResolveParams());
+	Context.RHICmdList.CopyToResolveTarget(DestRenderTarget.TargetableTexture, DestRenderTarget.ShaderResourceTexture, false, FResolveParams());
 }
 
 FPooledRenderTargetDesc FRCPassPostProcessHMD::ComputeOutputDesc(EPassOutputId InPassOutputId) const
diff --git a/Engine/Source/Runtime/Renderer/Private/PostProcess/PostProcessHistogram.cpp b/Engine/Source/Runtime/Renderer/Private/PostProcess/PostProcessHistogram.cpp
--- a/Engine/Source/Runtime/Renderer/Private/PostProcess/PostProcessHistogram.cpp
+++ b/Engine/Source/Runtime/Renderer/Private/PostProcess/PostProcessHistogram.cpp
@@ -61,7 +61,7 @@ public:
 
 		FGlobalShader::SetParameters(RHICmdList, ShaderRHI, Context.View);
 
-		PostprocessParameter.SetCS(RHICmdList, ShaderRHI, Context, TStaticSamplerState<SF_Point,AM_Clamp,AM_Clamp,AM_Clamp>::GetRHI());
+		PostprocessParameter.SetCS(ShaderRHI, Context, TStaticSamplerState<SF_Point,AM_Clamp,AM_Clamp,AM_Clamp>::GetRHI());
 
 		SetShaderValue(RHICmdList, ShaderRHI, ThreadGroupCount, ThreadGroupCountValue);
 		SetShaderValue(RHICmdList, ShaderRHI, LeftTopOffset, LeftTopOffsetValue);
diff --git a/Engine/Source/Runtime/Renderer/Private/PostProcess/PostProcessHistogramReduce.cpp b/Engine/Source/Runtime/Renderer/Private/PostProcess/PostProcessHistogramReduce.cpp
--- a/Engine/Source/Runtime/Renderer/Private/PostProcess/PostProcessHistogramReduce.cpp
+++ b/Engine/Source/Runtime/Renderer/Private/PostProcess/PostProcessHistogramReduce.cpp
@@ -45,15 +45,15 @@ public:
 		EyeAdapationTemporalParams.Bind(Initializer.ParameterMap, TEXT("EyeAdapationTemporalParams"));
 	}
 
-	void SetPS(FRHICommandList& RHICmdList, const FRenderingCompositePassContext& Context, uint32 LoopSizeValue)
+	void SetPS(const FRenderingCompositePassContext& Context, uint32 LoopSizeValue)
 	{
 		const FPixelShaderRHIParamRef ShaderRHI = GetPixelShader();
 
-		FGlobalShader::SetParameters(RHICmdList, ShaderRHI, Context.View);
+		FGlobalShader::SetParameters(Context.RHICmdList, ShaderRHI, Context.View);
 
-		PostprocessParameter.SetPS(RHICmdList, ShaderRHI, Context, TStaticSamplerState<SF_Bilinear,AM_Clamp,AM_Clamp,AM_Clamp>::GetRHI());
+		PostprocessParameter.SetPS(ShaderRHI, Context, TStaticSamplerState<SF_Bilinear,AM_Clamp,AM_Clamp,AM_Clamp>::GetRHI());
 
-		SetShaderValue(RHICmdList, ShaderRHI, LoopSize, LoopSizeValue);
+		SetShaderValue(Context.RHICmdList, ShaderRHI, LoopSize, LoopSizeValue);
 
 		if(EyeAdaptationTexture.IsBound())
 		{
@@ -61,18 +61,18 @@ public:
 
 			if(EyeAdaptationRT)
 			{
-				SetTextureParameter(RHICmdList, ShaderRHI, EyeAdaptationTexture, EyeAdaptationRT->GetRenderTargetItem().TargetableTexture);
+				SetTextureParameter(Context.RHICmdList, ShaderRHI, EyeAdaptationTexture, EyeAdaptationRT->GetRenderTargetItem().TargetableTexture);
 			}
 			else
 			{
 				// some views don't have a state, thumbnail rendering?
-				SetTextureParameter(RHICmdList, ShaderRHI, EyeAdaptationTexture, GWhiteTexture->TextureRHI);
+				SetTextureParameter(Context.RHICmdList, ShaderRHI, EyeAdaptationTexture, GWhiteTexture->TextureRHI);
 			}
 		}
 
 		// todo
 		FVector4 EyeAdapationTemporalParamsValue(0, 0, 0, 0);
-		SetShaderValue(RHICmdList, ShaderRHI, EyeAdapationTemporalParams, EyeAdapationTemporalParamsValue);
+		SetShaderValue(Context.RHICmdList, ShaderRHI, EyeAdapationTemporalParams, EyeAdapationTemporalParamsValue);
 	}
 	
 	// FShader interface.
@@ -106,31 +106,31 @@ void FRCPassPostProcessHistogramReduce::Process(FRenderingCompositePassContext&
 	const FSceneRenderTargetItem& DestRenderTarget = PassOutputs[0].RequestSurface(Context);
 
 	// Set the view family's render target/viewport.
-	RHISetRenderTarget(DestRenderTarget.TargetableTexture, FTextureRHIRef());	
+	SetRenderTarget(Context.RHICmdList, DestRenderTarget.TargetableTexture, FTextureRHIRef());
 
 	// set the state
-	RHISetBlendState(TStaticBlendState<>::GetRHI());
-	RHISetRasterizerState(TStaticRasterizerState<>::GetRHI());
-	RHISetDepthStencilState(TStaticDepthStencilState<false,CF_Always>::GetRHI());
+	Context.RHICmdList.SetBlendState(TStaticBlendState<>::GetRHI());
+	Context.RHICmdList.SetRasterizerState(TStaticRasterizerState<>::GetRHI());
+	Context.RHICmdList.SetDepthStencilState(TStaticDepthStencilState<false, CF_Always>::GetRHI());
 
 	TShaderMapRef<FPostProcessVS> VertexShader(GetGlobalShaderMap());
 	TShaderMapRef<FPostProcessHistogramReducePS> PixelShader(GetGlobalShaderMap());
 
 	static FGlobalBoundShaderState BoundShaderState;
+	Context.RHICmdList.CheckIsNull(); // need new approach for "static FGlobalBoundShaderState" for parallel rendering
 
-	SetGlobalBoundShaderState(BoundShaderState, GFilterVertexDeclaration.VertexDeclarationRHI, *VertexShader, *PixelShader);
+	SetGlobalBoundShaderState(Context.RHICmdList, BoundShaderState, GFilterVertexDeclaration.VertexDeclarationRHI, *VertexShader, *PixelShader);
 
 	// we currently assume the input is half res, one full res pixel less to avoid getting bilinear filtered input
 	FIntPoint GatherExtent = (View.ViewRect.Size() - FIntPoint(1, 1)) / 2;
 
 	uint32 LoopSizeValue = ComputeLoopSize(GatherExtent);
 
-	//@todo-rco: RHIPacketList
-	FRHICommandList& RHICmdList = FRHICommandList::GetNullRef();
-	PixelShader->SetPS(RHICmdList, Context, LoopSizeValue);
+	PixelShader->SetPS(Context, LoopSizeValue);
 
 	// Draw a quad mapping scene color to the view's render target
 	DrawRectangle(
+		Context.RHICmdList,
 		0, 0,
 		DestSize.X, DestSize.Y,
 		0, 0,
@@ -140,7 +140,7 @@ void FRCPassPostProcessHistogramReduce::Process(FRenderingCompositePassContext&
 		*VertexShader,
 		EDRF_UseTriangleOptimization);
 
-	RHICopyToResolveTarget(DestRenderTarget.TargetableTexture, DestRenderTarget.ShaderResourceTexture, false, FResolveParams());
+	Context.RHICmdList.CopyToResolveTarget(DestRenderTarget.TargetableTexture, DestRenderTarget.ShaderResourceTexture, false, FResolveParams());
 }
 
 uint32 FRCPassPostProcessHistogramReduce::ComputeLoopSize(FIntPoint PixelExtent)
diff --git a/Engine/Source/Runtime/Renderer/Private/PostProcess/PostProcessLensBlur.cpp b/Engine/Source/Runtime/Renderer/Private/PostProcess/PostProcessLensBlur.cpp
--- a/Engine/Source/Runtime/Renderer/Private/PostProcess/PostProcessLensBlur.cpp
+++ b/Engine/Source/Runtime/Renderer/Private/PostProcess/PostProcessLensBlur.cpp
@@ -49,31 +49,31 @@ public:
 	}
 
 	/** to have a similar interface as all other shaders */
-	void SetParameters(FRHICommandList& RHICmdList, const FRenderingCompositePassContext& Context, FIntPoint TileCountValue, uint32 TileSize, float PixelKernelSize, float Threshold)
+	void SetParameters(const FRenderingCompositePassContext& Context, FIntPoint TileCountValue, uint32 TileSize, float PixelKernelSize, float Threshold)
 	{
 		const FVertexShaderRHIParamRef ShaderRHI = GetVertexShader();
 
-		FGlobalShader::SetParameters(RHICmdList, ShaderRHI, Context.View);
+		FGlobalShader::SetParameters(Context.RHICmdList, ShaderRHI, Context.View);
 
-		PostprocessParameter.SetVS(RHICmdList, ShaderRHI, Context, TStaticSamplerState<SF_Bilinear,AM_Clamp,AM_Clamp,AM_Clamp>::GetRHI());
+		PostprocessParameter.SetVS(ShaderRHI, Context, TStaticSamplerState<SF_Bilinear, AM_Clamp, AM_Clamp, AM_Clamp>::GetRHI());
 
 		{
 			FIntRect TileCountAndSizeValue(TileCountValue, FIntPoint(TileSize, TileSize));
 
-			SetShaderValue(RHICmdList, ShaderRHI, TileCountAndSize, TileCountAndSizeValue);
+			SetShaderValue(Context.RHICmdList, ShaderRHI, TileCountAndSize, TileCountAndSizeValue);
 		}
 
 		{
 			// only approximate as the mip mapping doesn't produce accurate brightness scaling
 			FVector4 ColorScaleValue(1.0f / FMath::Max(1.0f, PixelKernelSize * PixelKernelSize), Threshold, 0, 0);
 
-			SetShaderValue(RHICmdList, ShaderRHI, ColorScale, ColorScaleValue);
+			SetShaderValue(Context.RHICmdList, ShaderRHI, ColorScale, ColorScaleValue);
 		}
 
 		{
 			FVector4 KernelSizeValue(PixelKernelSize, PixelKernelSize, 0, 0);
 
-			SetShaderValue(RHICmdList, ShaderRHI, KernelSize, KernelSizeValue);
+			SetShaderValue(Context.RHICmdList, ShaderRHI, KernelSize, KernelSizeValue);
 		}
 	}
 };
@@ -115,13 +115,13 @@ public:
 		return bShaderHasOutdatedParameters;
 	}
 
-	void SetParameters(FRHICommandList& RHICmdList, const FRenderingCompositePassContext& Context, float PixelKernelSize)
+	void SetParameters(const FRenderingCompositePassContext& Context, float PixelKernelSize)
 	{
 		const FPixelShaderRHIParamRef ShaderRHI = GetPixelShader();
 
-		FGlobalShader::SetParameters(RHICmdList, ShaderRHI, Context.View);
+		FGlobalShader::SetParameters(Context.RHICmdList, ShaderRHI, Context.View);
 
-		PostprocessParameter.SetPS(RHICmdList, ShaderRHI, Context, TStaticSamplerState<SF_Bilinear,AM_Clamp,AM_Clamp,AM_Clamp>::GetRHI());
+		PostprocessParameter.SetPS(ShaderRHI, Context, TStaticSamplerState<SF_Bilinear, AM_Clamp, AM_Clamp, AM_Clamp>::GetRHI());
 				
 		{
 			FTextureRHIParamRef TextureRHI = GWhiteTexture->TextureRHI;
@@ -146,7 +146,7 @@ public:
 				}
 			}
 
-			SetTextureParameter(RHICmdList, ShaderRHI, LensTexture, LensTextureSampler, TStaticSamplerState<SF_Bilinear,AM_Border,AM_Border,AM_Clamp>::GetRHI(), TextureRHI);
+			SetTextureParameter(Context.RHICmdList, ShaderRHI, LensTexture, LensTextureSampler, TStaticSamplerState<SF_Bilinear, AM_Border, AM_Border, AM_Clamp>::GetRHI(), TextureRHI);
 		}
 	}
 };
@@ -184,23 +184,23 @@ void FRCPassPostProcessLensBlur::Process(FRenderingCompositePassContext& Context
 	const FSceneRenderTargetItem& DestRenderTarget = PassOutputs[0].RequestSurface(Context);
 
 	// Set the view family's render target/viewport.
-	RHISetRenderTarget(DestRenderTarget.TargetableTexture, FTextureRHIRef());	
+	SetRenderTarget(Context.RHICmdList, DestRenderTarget.TargetableTexture, FTextureRHIRef());
 
-	RHIClear(true, FLinearColor::Black, false, 1.0f, false, 0, FIntRect() );
+	Context.RHICmdList.Clear(true, FLinearColor::Black, false, 1.0f, false, 0, FIntRect());
 
 	Context.SetViewportAndCallRHI(ViewRect);
 
 	// set the state (additive blending)
-	RHISetBlendState(TStaticBlendState<CW_RGB,BO_Add,BF_One,BF_One>::GetRHI());
-	RHISetRasterizerState(TStaticRasterizerState<>::GetRHI());
-	RHISetDepthStencilState(TStaticDepthStencilState<false,CF_Always>::GetRHI());
+	Context.RHICmdList.SetBlendState(TStaticBlendState<CW_RGB, BO_Add, BF_One, BF_One>::GetRHI());
+	Context.RHICmdList.SetRasterizerState(TStaticRasterizerState<>::GetRHI());
+	Context.RHICmdList.SetDepthStencilState(TStaticDepthStencilState<false, CF_Always>::GetRHI());
 
 	TShaderMapRef<FPostProcessLensBlurVS> VertexShader(GetGlobalShaderMap());
 	TShaderMapRef<FPostProcessLensBlurPS> PixelShader(GetGlobalShaderMap());
 
 	static FGlobalBoundShaderState BoundShaderState;
-
-	SetGlobalBoundShaderState(BoundShaderState, GFilterVertexDeclaration.VertexDeclarationRHI, *VertexShader, *PixelShader);
+	Context.RHICmdList.CheckIsNull(); // need new approach for "static FGlobalBoundShaderState" for parallel rendering
+	SetGlobalBoundShaderState(Context.RHICmdList, BoundShaderState, GFilterVertexDeclaration.VertexDeclarationRHI, *VertexShader, *PixelShader);
 
 	uint32 TileSize = 1;
 
@@ -208,19 +208,17 @@ void FRCPassPostProcessLensBlur::Process(FRenderingCompositePassContext& Context
 
 	float PixelKernelSize = PercentKernelSize / 100.0f * ViewSize.X;
 
-	//@todo-rco: RHIPacketList
-	FRHICommandList& RHICmdList = FRHICommandList::GetNullRef();
-	VertexShader->SetParameters(RHICmdList, Context, TileCount, TileSize, PixelKernelSize, Threshold);
-	PixelShader->SetParameters(RHICmdList, Context, PixelKernelSize);
+	VertexShader->SetParameters(Context, TileCount, TileSize, PixelKernelSize, Threshold);
+	PixelShader->SetParameters(Context, PixelKernelSize);
 
-	RHISetStreamSource(0, NULL, 0, 0);
+	Context.RHICmdList.SetStreamSource(0, NULL, 0, 0);
 
 	// needs to be the same on shader side (faster on NVIDIA and AMD)
 	int32 QuadsPerInstance = 4;
 
-	RHIDrawPrimitive(PT_TriangleList, 0, 2, FMath::DivideAndRoundUp(TileCount.X * TileCount.Y, QuadsPerInstance));
+	Context.RHICmdList.DrawPrimitive(PT_TriangleList, 0, 2, FMath::DivideAndRoundUp(TileCount.X * TileCount.Y, QuadsPerInstance));
 
-	RHICopyToResolveTarget(DestRenderTarget.TargetableTexture, DestRenderTarget.ShaderResourceTexture, false, FResolveParams());
+	Context.RHICmdList.CopyToResolveTarget(DestRenderTarget.TargetableTexture, DestRenderTarget.ShaderResourceTexture, false, FResolveParams());
 }
 
 
diff --git a/Engine/Source/Runtime/Renderer/Private/PostProcess/PostProcessLensFlares.cpp b/Engine/Source/Runtime/Renderer/Private/PostProcess/PostProcessLensFlares.cpp
--- a/Engine/Source/Runtime/Renderer/Private/PostProcess/PostProcessLensFlares.cpp
+++ b/Engine/Source/Runtime/Renderer/Private/PostProcess/PostProcessLensFlares.cpp
@@ -41,13 +41,13 @@ public:
 		return bShaderHasOutdatedParameters;
 	}
 
-	void SetParameters(FRHICommandList& RHICmdList, const FRenderingCompositePassContext& Context)
+	void SetParameters(const FRenderingCompositePassContext& Context)
 	{
 		const FPixelShaderRHIParamRef ShaderRHI = GetPixelShader();
 
-		FGlobalShader::SetParameters(RHICmdList, ShaderRHI, Context.View);
+		FGlobalShader::SetParameters(Context.RHICmdList, ShaderRHI, Context.View);
 
-		PostprocessParameter.SetPS(RHICmdList, ShaderRHI, Context, TStaticSamplerState<SF_Bilinear,AM_Clamp,AM_Clamp,AM_Clamp>::GetRHI());
+		PostprocessParameter.SetPS(ShaderRHI, Context, TStaticSamplerState<SF_Bilinear, AM_Clamp, AM_Clamp, AM_Clamp>::GetRHI());
 	}
 };
 
@@ -89,15 +89,15 @@ public:
 		return bShaderHasOutdatedParameters;
 	}
 
-	void SetParameters(FRHICommandList& RHICmdList, const FRenderingCompositePassContext& Context, FVector2D TexScaleValue)
+	void SetParameters(const FRenderingCompositePassContext& Context, FVector2D TexScaleValue)
 	{
 		const FPixelShaderRHIParamRef ShaderRHI = GetPixelShader();
 
-		FGlobalShader::SetParameters(RHICmdList, ShaderRHI, Context.View);
+		FGlobalShader::SetParameters(Context.RHICmdList, ShaderRHI, Context.View);
 
-		PostprocessParameter.SetPS(RHICmdList, ShaderRHI, Context, TStaticSamplerState<SF_Bilinear,AM_Clamp,AM_Clamp,AM_Clamp>::GetRHI());
+		PostprocessParameter.SetPS(ShaderRHI, Context, TStaticSamplerState<SF_Bilinear, AM_Clamp, AM_Clamp, AM_Clamp>::GetRHI());
 
-		SetShaderValue(RHICmdList, ShaderRHI, TexScale, TexScaleValue);
+		SetShaderValue(Context.RHICmdList, ShaderRHI, TexScale, TexScaleValue);
 	}
 };
 
@@ -141,17 +141,17 @@ void FRCPassPostProcessLensFlares::Process(FRenderingCompositePassContext& Conte
 	const FSceneRenderTargetItem& DestRenderTarget = PassOutputs[0].RequestSurface(Context);
 
 	// Set the view family's render target/viewport.
-	RHISetRenderTarget(DestRenderTarget.TargetableTexture, FTextureRHIRef());	
+	SetRenderTarget(Context.RHICmdList, DestRenderTarget.TargetableTexture, FTextureRHIRef());
 		
 	// is optimized away if possible (RT size=view size, )
-	RHIClear(true, FLinearColor::Black, false, 1.0f, false, 0, ViewRect1);
+	Context.RHICmdList.Clear(true, FLinearColor::Black, false, 1.0f, false, 0, ViewRect1);
 
 	Context.SetViewportAndCallRHI(ViewRect1);
 
 	// set the state
-	RHISetBlendState(TStaticBlendState<>::GetRHI());
-	RHISetRasterizerState(TStaticRasterizerState<>::GetRHI());
-	RHISetDepthStencilState(TStaticDepthStencilState<false,CF_Always>::GetRHI());
+	Context.RHICmdList.SetBlendState(TStaticBlendState<>::GetRHI());
+	Context.RHICmdList.SetRasterizerState(TStaticRasterizerState<>::GetRHI());
+	Context.RHICmdList.SetDepthStencilState(TStaticDepthStencilState<false, CF_Always>::GetRHI());
 
 	TShaderMapRef<FPostProcessVS> VertexShader(GetGlobalShaderMap());
 
@@ -163,14 +163,15 @@ void FRCPassPostProcessLensFlares::Process(FRenderingCompositePassContext& Conte
 		TShaderMapRef<FPostProcessLensFlareBasePS> PixelShader(GetGlobalShaderMap());
 
 		static FGlobalBoundShaderState BoundShaderState;
+		Context.RHICmdList.CheckIsNull(); // need new approach for "static FGlobalBoundShaderState" for parallel rendering
+		SetGlobalBoundShaderState(Context.RHICmdList, BoundShaderState, GFilterVertexDeclaration.VertexDeclarationRHI, *VertexShader, *PixelShader);
 
-		SetGlobalBoundShaderState(BoundShaderState, GFilterVertexDeclaration.VertexDeclarationRHI, *VertexShader, *PixelShader);
-
-		VertexShader->SetParameters(RHICmdList, Context);
-		PixelShader->SetParameters(RHICmdList, Context);
+		VertexShader->SetParameters(Context);
+		PixelShader->SetParameters(Context);
 
 		// Draw a quad mapping scene color to the view's render target
 		DrawRectangle(
+			Context.RHICmdList,
 			0, 0,
 			ViewSize1.X, ViewSize1.Y,
 			ViewRect1.Min.X, ViewRect1.Min.Y,
@@ -182,20 +183,20 @@ void FRCPassPostProcessLensFlares::Process(FRenderingCompositePassContext& Conte
 	}
 
 	// additive blend
-	RHISetBlendState(TStaticBlendState<CW_RGB,BO_Add,BF_One,BF_One>::GetRHI());
+	Context.RHICmdList.SetBlendState(TStaticBlendState<CW_RGB, BO_Add, BF_One, BF_One>::GetRHI());
 
 	// add lens flares on top of that
 	{
 		TShaderMapRef<FPostProcessLensFlaresPS> PixelShader(GetGlobalShaderMap());
 
 		static FGlobalBoundShaderState BoundShaderState;
-
-		SetGlobalBoundShaderState(BoundShaderState, GFilterVertexDeclaration.VertexDeclarationRHI, *VertexShader, *PixelShader);
+		Context.RHICmdList.CheckIsNull(); // need new approach for "static FGlobalBoundShaderState" for parallel rendering
+		SetGlobalBoundShaderState(Context.RHICmdList, BoundShaderState, GFilterVertexDeclaration.VertexDeclarationRHI, *VertexShader, *PixelShader);
 
 		FVector2D TexScaleValue = FVector2D(TexSize2) / ViewSize2;
 
-		VertexShader->SetParameters(RHICmdList, Context);
-		PixelShader->SetParameters(RHICmdList, Context, TexScaleValue);
+		VertexShader->SetParameters(Context);
+		PixelShader->SetParameters(Context, TexScaleValue);
 
 		// todo: expose
 		const uint32 Count = 8;
@@ -220,12 +221,11 @@ void FRCPassPostProcessLensFlares::Process(FRenderingCompositePassContext& Conte
 			Alpha *= SizeScale;
 			
 			// set the individual flare color
-			//@todo-rco: RHIPacketList
-			FRHICommandList& RHICmdList = FRHICommandList::GetNullRef();
-			SetShaderValue(RHICmdList, PixelShader->GetPixelShader(), PixelShader->FlareColor, FlareColor * LensFlareHDRColor);
+			SetShaderValue(Context.RHICmdList, PixelShader->GetPixelShader(), PixelShader->FlareColor, FlareColor * LensFlareHDRColor);
 
 			// Draw a quad mapping scene color to the view's render target
 			DrawRectangle(
+				Context.RHICmdList,
 				Center.X - 0.5f * ViewSize1.X * Alpha, Center.Y - 0.5f * ViewSize1.Y * Alpha,
 				ViewSize1.X * Alpha, ViewSize1.Y * Alpha,
 				ViewRect2.Min.X, ViewRect2.Min.Y,
@@ -237,7 +237,7 @@ void FRCPassPostProcessLensFlares::Process(FRenderingCompositePassContext& Conte
 		}
 	}
 
-	RHICopyToResolveTarget(DestRenderTarget.TargetableTexture, DestRenderTarget.ShaderResourceTexture, false, FResolveParams());
+	Context.RHICmdList.CopyToResolveTarget(DestRenderTarget.TargetableTexture, DestRenderTarget.ShaderResourceTexture, false, FResolveParams());
 }
 
 FPooledRenderTargetDesc FRCPassPostProcessLensFlares::ComputeOutputDesc(EPassOutputId InPassOutputId) const
diff --git a/Engine/Source/Runtime/Renderer/Private/PostProcess/PostProcessMaterial.cpp b/Engine/Source/Runtime/Renderer/Private/PostProcess/PostProcessMaterial.cpp
--- a/Engine/Source/Runtime/Renderer/Private/PostProcess/PostProcessMaterial.cpp
+++ b/Engine/Source/Runtime/Renderer/Private/PostProcess/PostProcessMaterial.cpp
@@ -75,7 +75,7 @@ public:
 		const FPixelShaderRHIParamRef ShaderRHI = GetPixelShader();
 
 		FMaterialShader::SetParameters(RHICmdList, ShaderRHI, Material, *Material->GetMaterial(Context.View.GetFeatureLevel()), Context.View, true, ESceneRenderTargetsMode::SetTextures);
-		PostprocessParameter.SetPS(RHICmdList, ShaderRHI, Context, TStaticSamplerState<SF_Point,AM_Clamp,AM_Clamp,AM_Clamp>::GetRHI());
+		PostprocessParameter.SetPS(ShaderRHI, Context, TStaticSamplerState<SF_Point,AM_Clamp,AM_Clamp,AM_Clamp>::GetRHI());
 	}
 
 	virtual bool Serialize(FArchive& Ar)
@@ -130,19 +130,19 @@ void FRCPassPostProcessMaterial::Process(FRenderingCompositePassContext& Context
 	const FSceneRenderTargetItem& DestRenderTarget = PassOutputs[0].RequestSurface(Context);
 
 	// Set the view family's render target/viewport.
-	RHISetRenderTarget(DestRenderTarget.TargetableTexture, FTextureRHIParamRef());	
+	SetRenderTarget(Context.RHICmdList, DestRenderTarget.TargetableTexture, FTextureRHIParamRef());
 
 	if( ViewFamily.RenderTarget->GetRenderTargetTexture() != DestRenderTarget.TargetableTexture )
 	{
-		RHIClear(true, FLinearColor::Black, false, 1.0f, false, 0, View.ViewRect);
+		Context.RHICmdList.Clear(true, FLinearColor::Black, false, 1.0f, false, 0, View.ViewRect);
 	}
 
 	Context.SetViewportAndCallRHI(View.ViewRect);
 
 	// set the state
-	RHISetBlendState(TStaticBlendState<>::GetRHI());
-	RHISetRasterizerState(TStaticRasterizerState<>::GetRHI());
-	RHISetDepthStencilState(TStaticDepthStencilState<false,CF_Always>::GetRHI());
+	Context.RHICmdList.SetBlendState(TStaticBlendState<>::GetRHI());
+	Context.RHICmdList.SetRasterizerState(TStaticRasterizerState<>::GetRHI());
+	Context.RHICmdList.SetDepthStencilState(TStaticDepthStencilState<false, CF_Always>::GetRHI());
 
 	float ScaleX = 1.0f / InputDesc->Extent.X;
 	float ScaleY = 1.0f / InputDesc->Extent.Y;
@@ -152,14 +152,13 @@ void FRCPassPostProcessMaterial::Process(FRenderingCompositePassContext& Context
 	FPostProcessMaterialVS* VertexShader = MaterialShaderMap->GetShader<FPostProcessMaterialVS>();
 
 	FBoundShaderStateRHIRef BoundShaderState = RHICreateBoundShaderState(GetVertexDeclarationFVector4(), VertexShader->GetVertexShader(), FHullShaderRHIRef(), FDomainShaderRHIRef(), PixelShader->GetPixelShader(), FGeometryShaderRHIRef());
-	RHISetBoundShaderState(BoundShaderState);
+	Context.RHICmdList.SetBoundShaderState(BoundShaderState);
 
-	//@todo-rco: RHIPacketList
-	FRHICommandList& RHICmdList = FRHICommandList::GetNullRef();
-	VertexShader->SetParameters(RHICmdList, Context);
-	PixelShader->SetParameters(RHICmdList, Context, MaterialInterface->GetRenderProxy(false));
+	VertexShader->SetParameters(Context.RHICmdList, Context);
+	PixelShader->SetParameters(Context.RHICmdList, Context, MaterialInterface->GetRenderProxy(false));
 
 	DrawRectangle(
+		Context.RHICmdList,
 		0, 0,
 		DestRect.Width(), DestRect.Height(),
 		SrcRect.Min.X, SrcRect.Min.Y,
@@ -169,7 +168,7 @@ void FRCPassPostProcessMaterial::Process(FRenderingCompositePassContext& Context
 		VertexShader,
 		EDRF_UseTriangleOptimization);
 
-	RHICopyToResolveTarget(DestRenderTarget.TargetableTexture, DestRenderTarget.ShaderResourceTexture, false, FResolveParams());
+	Context.RHICmdList.CopyToResolveTarget(DestRenderTarget.TargetableTexture, DestRenderTarget.ShaderResourceTexture, false, FResolveParams());
 
 	if(Material->NeedsGBuffer())
 	{
