UE4 - more work preparing for parallel rendering. Merged the cmd list into the post process context, among other things.

diff --git a/Engine/Source/Runtime/Renderer/Private/PostProcess/PostProcessMobile.cpp b/Engine/Source/Runtime/Renderer/Private/PostProcess/PostProcessMobile.cpp
--- a/Engine/Source/Runtime/Renderer/Private/PostProcess/PostProcessMobile.cpp
+++ b/Engine/Source/Runtime/Renderer/Private/PostProcess/PostProcessMobile.cpp
@@ -39,13 +39,13 @@ public:
 		PostprocessParameter.Bind(Initializer.ParameterMap);
 	}
 
-	void SetVS(FRHICommandList& RHICmdList, const FRenderingCompositePassContext& Context)
+	void SetVS(const FRenderingCompositePassContext& Context)
 	{
 		const FVertexShaderRHIParamRef ShaderRHI = GetVertexShader();
 
-		FGlobalShader::SetParameters(RHICmdList, ShaderRHI, Context.View);
+		FGlobalShader::SetParameters(Context.RHICmdList, ShaderRHI, Context.View);
 
-		PostprocessParameter.SetVS(RHICmdList, ShaderRHI, Context, TStaticSamplerState<SF_Bilinear,AM_Clamp,AM_Clamp,AM_Clamp>::GetRHI());
+		PostprocessParameter.SetVS(ShaderRHI, Context, TStaticSamplerState<SF_Bilinear,AM_Clamp,AM_Clamp,AM_Clamp>::GetRHI());
 	}
 
 	// FShader interface.
@@ -92,20 +92,20 @@ public:
 		BloomThreshold.Bind(Initializer.ParameterMap, TEXT("BloomThreshold"));
 	}
 
-	void SetPS(FRHICommandList& RHICmdList, const FRenderingCompositePassContext& Context)
+	void SetPS(const FRenderingCompositePassContext& Context)
 	{
 		const FPixelShaderRHIParamRef ShaderRHI = GetPixelShader();
 
 		const FPostProcessSettings& Settings = Context.View.FinalPostProcessSettings;
 
-		FGlobalShader::SetParameters(RHICmdList, ShaderRHI, Context.View);
+		FGlobalShader::SetParameters(Context.RHICmdList, ShaderRHI, Context.View);
 
-		PostprocessParameter.SetPS(RHICmdList, ShaderRHI, Context, TStaticSamplerState<SF_Bilinear,AM_Clamp,AM_Clamp,AM_Clamp>::GetRHI());
+		PostprocessParameter.SetPS(ShaderRHI, Context, TStaticSamplerState<SF_Bilinear, AM_Clamp, AM_Clamp, AM_Clamp>::GetRHI());
 
 		float ExposureScale = FRCPassPostProcessEyeAdaptation::ComputeExposureScaleValue(Context.View);
 
 		FVector4 BloomThresholdValue(Settings.BloomThreshold, 0, 0, ExposureScale);
-		SetShaderValue(RHICmdList, ShaderRHI, BloomThreshold, BloomThresholdValue);
+		SetShaderValue(Context.RHICmdList, ShaderRHI, BloomThreshold, BloomThresholdValue);
 	}
 
 	// FShader interface.
@@ -129,20 +129,20 @@ IMPLEMENT_SHADER_TYPE(template<>,FPostProcessBloomSetupPS_ES2_2,TEXT("PostProces
 IMPLEMENT_SHADER_TYPE(template<>,FPostProcessBloomSetupPS_ES2_3,TEXT("PostProcessMobile"),TEXT("BloomPS_ES2"),SF_Pixel);
 
 template <uint32 UseSunDof>
-static void BloomSetup_SetShader(FRHICommandList& RHICmdList, const FRenderingCompositePassContext& Context)
+static void BloomSetup_SetShader(const FRenderingCompositePassContext& Context)
 {
 	TShaderMapRef<FPostProcessBloomSetupVS_ES2> VertexShader(GetGlobalShaderMap());
 	TShaderMapRef<FPostProcessBloomSetupPS_ES2<UseSunDof> > PixelShader(GetGlobalShaderMap());
 
 	static FGlobalBoundShaderState BoundShaderState;
+	Context.RHICmdList.CheckIsNull(); // need new approach for "static FGlobalBoundShaderState" for parallel rendering
+	SetGlobalBoundShaderState(Context.RHICmdList, BoundShaderState, GFilterVertexDeclaration.VertexDeclarationRHI, *VertexShader, *PixelShader);
 
-	SetGlobalBoundShaderState(BoundShaderState, GFilterVertexDeclaration.VertexDeclarationRHI, *VertexShader, *PixelShader);
-
-	VertexShader->SetVS(RHICmdList, Context);
-	PixelShader->SetPS(RHICmdList, Context);
+	VertexShader->SetVS(Context);
+	PixelShader->SetPS(Context);
 }
 
-void FRCPassPostProcessBloomSetupES2::SetShader(FRHICommandList& RHICmdList, const FRenderingCompositePassContext& Context)
+void FRCPassPostProcessBloomSetupES2::SetShader(const FRenderingCompositePassContext& Context)
 {
 	const FSceneView& View = Context.View;
 	uint32 UseSun = Context.View.bLightShaftUse ? 1 : 0;
@@ -151,10 +151,10 @@ void FRCPassPostProcessBloomSetupES2::SetShader(FRHICommandList& RHICmdList, con
 
 	switch(UseSunDof)
 	{
-		case 0: BloomSetup_SetShader<0>(RHICmdList, Context); break;
-		case 1: BloomSetup_SetShader<1>(RHICmdList, Context); break;
-		case 2: BloomSetup_SetShader<2>(RHICmdList, Context); break;
-		case 3: BloomSetup_SetShader<3>(RHICmdList, Context); break;
+		case 0: BloomSetup_SetShader<0>(Context); break;
+		case 1: BloomSetup_SetShader<1>(Context); break;
+		case 2: BloomSetup_SetShader<2>(Context); break;
+		case 3: BloomSetup_SetShader<3>(Context); break;
 	}
 }
 
@@ -195,24 +195,23 @@ void FRCPassPostProcessBloomSetupES2::Process(FRenderingCompositePassContext& Co
 	}
 
 	const FSceneRenderTargetItem& DestRenderTarget = PassOutputs[0].RequestSurface(Context);
-	RHISetRenderTarget(DestRenderTarget.TargetableTexture, FTextureRHIRef());	
+	SetRenderTarget(Context.RHICmdList, DestRenderTarget.TargetableTexture, FTextureRHIRef());
 
 	// is optimized away if possible (RT size=view size, )
-	RHIClear(true, FLinearColor::Black, false, 1.0f, false, 0, FIntRect());
+	Context.RHICmdList.Clear(true, FLinearColor::Black, false, 1.0f, false, 0, FIntRect());
 
 	Context.SetViewportAndCallRHI(0, 0, 0.0f, DstX, DstY, 1.0f );
 
-	RHISetBlendState(TStaticBlendState<>::GetRHI());
-	RHISetRasterizerState(TStaticRasterizerState<>::GetRHI());
-	RHISetDepthStencilState(TStaticDepthStencilState<false,CF_Always>::GetRHI());
+	Context.RHICmdList.SetBlendState(TStaticBlendState<>::GetRHI());
+	Context.RHICmdList.SetRasterizerState(TStaticRasterizerState<>::GetRHI());
+	Context.RHICmdList.SetDepthStencilState(TStaticDepthStencilState<false, CF_Always>::GetRHI());
 
-	//@todo-rco: RHIPacketList
-	FRHICommandList& RHICmdList = FRHICommandList::GetNullRef();
-	SetShader(RHICmdList, Context);
+	SetShader(Context);
 
 	TShaderMapRef<FPostProcessBloomSetupVS_ES2> VertexShader(GetGlobalShaderMap());
 
 	DrawRectangle(
+		Context.RHICmdList,
 		0, 0,
 		DstX, DstY,
 		SrcRect.Min.X, SrcRect.Min.Y,
@@ -222,7 +221,7 @@ void FRCPassPostProcessBloomSetupES2::Process(FRenderingCompositePassContext& Co
 		*VertexShader,
 		EDRF_UseTriangleOptimization);
 
-	RHICopyToResolveTarget(DestRenderTarget.TargetableTexture, DestRenderTarget.ShaderResourceTexture, false, FResolveParams());
+	Context.RHICmdList.CopyToResolveTarget(DestRenderTarget.TargetableTexture, DestRenderTarget.ShaderResourceTexture, false, FResolveParams());
 }
 
 FPooledRenderTargetDesc FRCPassPostProcessBloomSetupES2::ComputeOutputDesc(EPassOutputId InPassOutputId) const
@@ -271,13 +270,13 @@ public:
 		PostprocessParameter.Bind(Initializer.ParameterMap);
 	}
 
-	void SetVS(FRHICommandList& RHICmdList, const FRenderingCompositePassContext& Context)
+	void SetVS(const FRenderingCompositePassContext& Context)
 	{
 		const FVertexShaderRHIParamRef ShaderRHI = GetVertexShader();
 
-		FGlobalShader::SetParameters(RHICmdList, ShaderRHI, Context.View);
+		FGlobalShader::SetParameters(Context.RHICmdList, ShaderRHI, Context.View);
 
-		PostprocessParameter.SetVS(RHICmdList, ShaderRHI, Context, TStaticSamplerState<SF_Bilinear,AM_Clamp,AM_Clamp,AM_Clamp>::GetRHI());
+		PostprocessParameter.SetVS(ShaderRHI, Context, TStaticSamplerState<SF_Bilinear,AM_Clamp,AM_Clamp,AM_Clamp>::GetRHI());
 	}
 
 	// FShader interface.
@@ -317,20 +316,20 @@ public:
 		BloomThreshold.Bind(Initializer.ParameterMap, TEXT("BloomThreshold"));
 	}
 
-	void SetPS(FRHICommandList& RHICmdList, const FRenderingCompositePassContext& Context)
+	void SetPS(const FRenderingCompositePassContext& Context)
 	{
 		const FPixelShaderRHIParamRef ShaderRHI = GetPixelShader();
 
 		const FPostProcessSettings& Settings = Context.View.FinalPostProcessSettings;
 
-		FGlobalShader::SetParameters(RHICmdList, ShaderRHI, Context.View);
+		FGlobalShader::SetParameters(Context.RHICmdList, ShaderRHI, Context.View);
 
-		PostprocessParameter.SetPS(RHICmdList, ShaderRHI, Context, TStaticSamplerState<SF_Bilinear,AM_Clamp,AM_Clamp,AM_Clamp>::GetRHI());
+		PostprocessParameter.SetPS(ShaderRHI, Context, TStaticSamplerState<SF_Bilinear,AM_Clamp,AM_Clamp,AM_Clamp>::GetRHI());
 
 		float ExposureScale = FRCPassPostProcessEyeAdaptation::ComputeExposureScaleValue(Context.View);
 
 		FVector4 BloomThresholdValue(Settings.BloomThreshold, 0, 0, ExposureScale);
-		SetShaderValue(RHICmdList, ShaderRHI, BloomThreshold, BloomThresholdValue);
+		SetShaderValue(Context.RHICmdList, ShaderRHI, BloomThreshold, BloomThresholdValue);
 	}
 
 	// FShader interface.
@@ -347,17 +346,17 @@ IMPLEMENT_SHADER_TYPE(,FPostProcessBloomSetupSmallVS_ES2,TEXT("PostProcessMobile
 IMPLEMENT_SHADER_TYPE(,FPostProcessBloomSetupSmallPS_ES2,TEXT("PostProcessMobile"),TEXT("BloomSmallPS_ES2"),SF_Pixel);
 
 
-void FRCPassPostProcessBloomSetupSmallES2::SetShader(FRHICommandList& RHICmdList, const FRenderingCompositePassContext& Context)
+void FRCPassPostProcessBloomSetupSmallES2::SetShader(const FRenderingCompositePassContext& Context)
 {
 	TShaderMapRef<FPostProcessBloomSetupSmallVS_ES2> VertexShader(GetGlobalShaderMap());
 	TShaderMapRef<FPostProcessBloomSetupSmallPS_ES2> PixelShader(GetGlobalShaderMap());
 
 	static FGlobalBoundShaderState BoundShaderState;
+	Context.RHICmdList.CheckIsNull(); // need new approach for "static FGlobalBoundShaderState" for parallel rendering
+	SetGlobalBoundShaderState(Context.RHICmdList, BoundShaderState, GFilterVertexDeclaration.VertexDeclarationRHI, *VertexShader, *PixelShader);
 
-	SetGlobalBoundShaderState(BoundShaderState, GFilterVertexDeclaration.VertexDeclarationRHI, *VertexShader, *PixelShader);
-
-	VertexShader->SetVS(RHICmdList, Context);
-	PixelShader->SetPS(RHICmdList, Context);
+	VertexShader->SetVS(Context);
+	PixelShader->SetPS(Context);
 }
 
 void FRCPassPostProcessBloomSetupSmallES2::Process(FRenderingCompositePassContext& Context)
@@ -397,23 +396,22 @@ void FRCPassPostProcessBloomSetupSmallES2::Process(FRenderingCompositePassContex
 	}
 
 	const FSceneRenderTargetItem& DestRenderTarget = PassOutputs[0].RequestSurface(Context);
-	RHISetRenderTarget(DestRenderTarget.TargetableTexture, FTextureRHIRef());	
+	SetRenderTarget(Context.RHICmdList, DestRenderTarget.TargetableTexture, FTextureRHIRef());
 
 	// is optimized away if possible (RT size=view size, )
-	RHIClear(true, FLinearColor::Black, false, 1.0f, false, 0, FIntRect());
+	Context.RHICmdList.Clear(true, FLinearColor::Black, false, 1.0f, false, 0, FIntRect());
 
 	Context.SetViewportAndCallRHI(0, 0, 0.0f, DstX, DstY, 1.0f );
 
-	RHISetBlendState(TStaticBlendState<>::GetRHI());
-	RHISetRasterizerState(TStaticRasterizerState<>::GetRHI());
-	RHISetDepthStencilState(TStaticDepthStencilState<false,CF_Always>::GetRHI());
+	Context.RHICmdList.SetBlendState(TStaticBlendState<>::GetRHI());
+	Context.RHICmdList.SetRasterizerState(TStaticRasterizerState<>::GetRHI());
+	Context.RHICmdList.SetDepthStencilState(TStaticDepthStencilState<false, CF_Always>::GetRHI());
 
-	//@todo-rco: RHIPacketList
-	FRHICommandList& RHICmdList = FRHICommandList::GetNullRef();
-	SetShader(RHICmdList, Context);
+	SetShader(Context);
 
 	TShaderMapRef<FPostProcessBloomSetupSmallVS_ES2> VertexShader(GetGlobalShaderMap());
 	DrawRectangle(
+		Context.RHICmdList,
 		0, 0,
 		DstX, DstY,
 		SrcRect.Min.X, SrcRect.Min.Y,
@@ -423,7 +421,7 @@ void FRCPassPostProcessBloomSetupSmallES2::Process(FRenderingCompositePassContex
 		*VertexShader,
 		EDRF_UseTriangleOptimization);
 
-	RHICopyToResolveTarget(DestRenderTarget.TargetableTexture, DestRenderTarget.ShaderResourceTexture, false, FResolveParams());
+	Context.RHICmdList.CopyToResolveTarget(DestRenderTarget.TargetableTexture, DestRenderTarget.ShaderResourceTexture, false, FResolveParams());
 }
 
 FPooledRenderTargetDesc FRCPassPostProcessBloomSetupSmallES2::ComputeOutputDesc(EPassOutputId InPassOutputId) const
@@ -477,12 +475,12 @@ public:
 		PostprocessParameter.Bind(Initializer.ParameterMap);
 	}
 
-	void SetPS(FRHICommandList& RHICmdList, const FRenderingCompositePassContext& Context)
+	void SetPS(const FRenderingCompositePassContext& Context)
 	{
 		const FPixelShaderRHIParamRef ShaderRHI = GetPixelShader();
 		const FPostProcessSettings& Settings = Context.View.FinalPostProcessSettings;
-		FGlobalShader::SetParameters(RHICmdList, ShaderRHI, Context.View);
-		PostprocessParameter.SetPS(RHICmdList, ShaderRHI, Context, TStaticSamplerState<SF_Bilinear,AM_Clamp,AM_Clamp,AM_Clamp>::GetRHI());
+		FGlobalShader::SetParameters(Context.RHICmdList, ShaderRHI, Context.View);
+		PostprocessParameter.SetPS(ShaderRHI, Context, TStaticSamplerState<SF_Bilinear,AM_Clamp,AM_Clamp,AM_Clamp>::GetRHI());
 	}
 
 	virtual bool Serialize(FArchive& Ar)
@@ -518,12 +516,12 @@ public:
 		BloomDownScale.Bind(Initializer.ParameterMap, TEXT("BloomDownScale"));
 	}
 
-	void SetVS(FRHICommandList& RHICmdList, const FRenderingCompositePassContext& Context, float InScale)
+	void SetVS(const FRenderingCompositePassContext& Context, float InScale)
 	{
 		const FVertexShaderRHIParamRef ShaderRHI = GetVertexShader();
-		FGlobalShader::SetParameters(RHICmdList, ShaderRHI, Context.View);
-		PostprocessParameter.SetVS(RHICmdList, ShaderRHI, Context, TStaticSamplerState<SF_Bilinear,AM_Clamp,AM_Clamp,AM_Clamp>::GetRHI());
-		SetShaderValue(RHICmdList, ShaderRHI, BloomDownScale, InScale);
+		FGlobalShader::SetParameters(Context.RHICmdList, ShaderRHI, Context.View);
+		PostprocessParameter.SetVS(ShaderRHI, Context, TStaticSamplerState<SF_Bilinear,AM_Clamp,AM_Clamp,AM_Clamp>::GetRHI());
+		SetShaderValue(Context.RHICmdList, ShaderRHI, BloomDownScale, InScale);
 	}
 
 	virtual bool Serialize(FArchive& Ar)
@@ -551,32 +549,32 @@ void FRCPassPostProcessBloomDownES2::Process(FRenderingCompositePassContext& Con
 	DstRect.Max.Y = DstY;
 
 	const FSceneRenderTargetItem& DestRenderTarget = PassOutputs[0].RequestSurface(Context);
-	RHISetRenderTarget(DestRenderTarget.TargetableTexture, FTextureRHIRef());	
+	SetRenderTarget(Context.RHICmdList, DestRenderTarget.TargetableTexture, FTextureRHIRef());
 
 	// is optimized away if possible (RT size=view size, )
-	RHIClear(true, FLinearColor::Black, false, 1.0f, false, 0, FIntRect());
+	Context.RHICmdList.Clear(true, FLinearColor::Black, false, 1.0f, false, 0, FIntRect());
 
 	Context.SetViewportAndCallRHI(0, 0, 0.0f, DstX, DstY, 1.0f );
 
-	RHISetBlendState(TStaticBlendState<>::GetRHI());
-	RHISetRasterizerState(TStaticRasterizerState<>::GetRHI());
-	RHISetDepthStencilState(TStaticDepthStencilState<false,CF_Always>::GetRHI());
+	Context.RHICmdList.SetBlendState(TStaticBlendState<>::GetRHI());
+	Context.RHICmdList.SetRasterizerState(TStaticRasterizerState<>::GetRHI());
+	Context.RHICmdList.SetDepthStencilState(TStaticDepthStencilState<false, CF_Always>::GetRHI());
 
 	TShaderMapRef<FPostProcessBloomDownVS_ES2> VertexShader(GetGlobalShaderMap());
 	TShaderMapRef<FPostProcessBloomDownPS_ES2> PixelShader(GetGlobalShaderMap());
 
 	static FGlobalBoundShaderState BoundShaderState;
+	Context.RHICmdList.CheckIsNull(); // need new approach for "static FGlobalBoundShaderState" for parallel rendering
 
-	SetGlobalBoundShaderState(BoundShaderState, GFilterVertexDeclaration.VertexDeclarationRHI, *VertexShader, *PixelShader);
+	SetGlobalBoundShaderState(Context.RHICmdList, BoundShaderState, GFilterVertexDeclaration.VertexDeclarationRHI, *VertexShader, *PixelShader);
 
-	//@todo-rco: RHIPacketList
-	FRHICommandList& RHICmdList = FRHICommandList::GetNullRef();
-	VertexShader->SetVS(RHICmdList, Context, Scale);
-	PixelShader->SetPS(RHICmdList, Context);
+	VertexShader->SetVS(Context, Scale);
+	PixelShader->SetPS(Context);
 
 	FIntPoint SrcDstSize = PrePostSourceViewportSize/2;
 
 	DrawRectangle(
+		Context.RHICmdList,
 		0, 0,
 		DstX, DstY,
 		0, 0,
@@ -586,7 +584,7 @@ void FRCPassPostProcessBloomDownES2::Process(FRenderingCompositePassContext& Con
 		*VertexShader,
 		EDRF_UseTriangleOptimization);
 
-	RHICopyToResolveTarget(DestRenderTarget.TargetableTexture, DestRenderTarget.ShaderResourceTexture, false, FResolveParams());
+	Context.RHICmdList.CopyToResolveTarget(DestRenderTarget.TargetableTexture, DestRenderTarget.ShaderResourceTexture, false, FResolveParams());
 }
 
 FPooledRenderTargetDesc FRCPassPostProcessBloomDownES2::ComputeOutputDesc(EPassOutputId InPassOutputId) const
@@ -643,14 +641,14 @@ public:
 		TintB.Bind(Initializer.ParameterMap, TEXT("BloomTintB"));
 	}
 
-	void SetPS(FRHICommandList& RHICmdList, const FRenderingCompositePassContext& Context, FVector4& InTintA, FVector4& InTintB)
+	void SetPS(const FRenderingCompositePassContext& Context, FVector4& InTintA, FVector4& InTintB)
 	{
 		const FPixelShaderRHIParamRef ShaderRHI = GetPixelShader();
 		const FPostProcessSettings& Settings = Context.View.FinalPostProcessSettings;
-		FGlobalShader::SetParameters(RHICmdList, ShaderRHI, Context.View);
-		PostprocessParameter.SetPS(RHICmdList, ShaderRHI, Context, TStaticSamplerState<SF_Bilinear,AM_Clamp,AM_Clamp,AM_Clamp>::GetRHI());
-		SetShaderValue(RHICmdList, ShaderRHI, TintA, InTintA);
-		SetShaderValue(RHICmdList, ShaderRHI, TintB, InTintB);
+		FGlobalShader::SetParameters(Context.RHICmdList, ShaderRHI, Context.View);
+		PostprocessParameter.SetPS(ShaderRHI, Context, TStaticSamplerState<SF_Bilinear,AM_Clamp,AM_Clamp,AM_Clamp>::GetRHI());
+		SetShaderValue(Context.RHICmdList, ShaderRHI, TintA, InTintA);
+		SetShaderValue(Context.RHICmdList, ShaderRHI, TintB, InTintB);
 	}
 
 	virtual bool Serialize(FArchive& Ar)
@@ -686,12 +684,12 @@ public:
 		BloomUpScales.Bind(Initializer.ParameterMap, TEXT("BloomUpScales"));
 	}
 
-	void SetVS(FRHICommandList& RHICmdList, const FRenderingCompositePassContext& Context, FVector2D InScale)
+	void SetVS(const FRenderingCompositePassContext& Context, FVector2D InScale)
 	{
 		const FVertexShaderRHIParamRef ShaderRHI = GetVertexShader();
-		FGlobalShader::SetParameters(RHICmdList, ShaderRHI, Context.View);
-		PostprocessParameter.SetVS(RHICmdList, ShaderRHI, Context, TStaticSamplerState<SF_Bilinear,AM_Clamp,AM_Clamp,AM_Clamp>::GetRHI());
-		SetShaderValue(RHICmdList, ShaderRHI, BloomUpScales, InScale);
+		FGlobalShader::SetParameters(Context.RHICmdList, ShaderRHI, Context.View);
+		PostprocessParameter.SetVS(ShaderRHI, Context, TStaticSamplerState<SF_Bilinear,AM_Clamp,AM_Clamp,AM_Clamp>::GetRHI());
+		SetShaderValue(Context.RHICmdList, ShaderRHI, BloomUpScales, InScale);
 	}
 
 	virtual bool Serialize(FArchive& Ar)
@@ -719,35 +717,35 @@ void FRCPassPostProcessBloomUpES2::Process(FRenderingCompositePassContext& Conte
 	DstRect.Max.Y = DstY;
 
 	const FSceneRenderTargetItem& DestRenderTarget = PassOutputs[0].RequestSurface(Context);
-	RHISetRenderTarget(DestRenderTarget.TargetableTexture, FTextureRHIRef());	
+	SetRenderTarget(Context.RHICmdList, DestRenderTarget.TargetableTexture, FTextureRHIRef());
 
 	// is optimized away if possible (RT size=view size, )
-	RHIClear(true, FLinearColor::Black, false, 1.0f, false, 0, FIntRect());
+	Context.RHICmdList.Clear(true, FLinearColor::Black, false, 1.0f, false, 0, FIntRect());
 
 	Context.SetViewportAndCallRHI(0, 0, 0.0f, DstX, DstY, 1.0f );
 
-	RHISetBlendState(TStaticBlendState<>::GetRHI());
-	RHISetRasterizerState(TStaticRasterizerState<>::GetRHI());
-	RHISetDepthStencilState(TStaticDepthStencilState<false,CF_Always>::GetRHI());
+	Context.RHICmdList.SetBlendState(TStaticBlendState<>::GetRHI());
+	Context.RHICmdList.SetRasterizerState(TStaticRasterizerState<>::GetRHI());
+	Context.RHICmdList.SetDepthStencilState(TStaticDepthStencilState<false, CF_Always>::GetRHI());
 
 	TShaderMapRef<FPostProcessBloomUpVS_ES2> VertexShader(GetGlobalShaderMap());
 	TShaderMapRef<FPostProcessBloomUpPS_ES2> PixelShader(GetGlobalShaderMap());
 
 	static FGlobalBoundShaderState BoundShaderState;
+	Context.RHICmdList.CheckIsNull(); // need new approach for "static FGlobalBoundShaderState" for parallel rendering
 
-	SetGlobalBoundShaderState(BoundShaderState, GFilterVertexDeclaration.VertexDeclarationRHI, *VertexShader, *PixelShader);
+	SetGlobalBoundShaderState(Context.RHICmdList, BoundShaderState, GFilterVertexDeclaration.VertexDeclarationRHI, *VertexShader, *PixelShader);
 
-	//@todo-rco: RHIPacketList
-	FRHICommandList& RHICmdList = FRHICommandList::GetNullRef();
 	// The 1/8 factor is because bloom is using 8 taps in the filter.
-	VertexShader->SetVS(RHICmdList, Context, FVector2D(ScaleAB.X, ScaleAB.Y));
+	VertexShader->SetVS(Context, FVector2D(ScaleAB.X, ScaleAB.Y));
 	FVector4 TintAScaled = TintA * (1.0f/8.0f);
 	FVector4 TintBScaled = TintB * (1.0f/8.0f);
-	PixelShader->SetPS(RHICmdList, Context, TintAScaled, TintBScaled);
+	PixelShader->SetPS(Context, TintAScaled, TintBScaled);
 
 	FIntPoint SrcDstSize = PrePostSourceViewportSize;
 
 	DrawRectangle(
+		Context.RHICmdList,
 		0, 0,
 		DstX, DstY,
 		0, 0,
@@ -757,7 +755,7 @@ void FRCPassPostProcessBloomUpES2::Process(FRenderingCompositePassContext& Conte
 		*VertexShader,
 		EDRF_UseTriangleOptimization);
 
-	RHICopyToResolveTarget(DestRenderTarget.TargetableTexture, DestRenderTarget.ShaderResourceTexture, false, FResolveParams());
+	Context.RHICmdList.CopyToResolveTarget(DestRenderTarget.TargetableTexture, DestRenderTarget.ShaderResourceTexture, false, FResolveParams());
 }
 
 FPooledRenderTargetDesc FRCPassPostProcessBloomUpES2::ComputeOutputDesc(EPassOutputId InPassOutputId) const
@@ -816,19 +814,19 @@ public:
 		SunColorApertureDiv2Parameter.Bind(Initializer.ParameterMap, TEXT("SunColorApertureDiv2"));
 	}
 
-	void SetPS(FRHICommandList& RHICmdList, const FRenderingCompositePassContext& Context)
+	void SetPS(const FRenderingCompositePassContext& Context)
 	{
 		const FPixelShaderRHIParamRef ShaderRHI = GetPixelShader();
 		const FPostProcessSettings& Settings = Context.View.FinalPostProcessSettings;
-		FGlobalShader::SetParameters(RHICmdList, ShaderRHI, Context.View);
-		PostprocessParameter.SetPS(RHICmdList, ShaderRHI, Context, TStaticSamplerState<SF_Bilinear,AM_Clamp,AM_Clamp,AM_Clamp>::GetRHI());
+		FGlobalShader::SetParameters(Context.RHICmdList, ShaderRHI, Context.View);
+		PostprocessParameter.SetPS(ShaderRHI, Context, TStaticSamplerState<SF_Bilinear,AM_Clamp,AM_Clamp,AM_Clamp>::GetRHI());
 
 		FVector4 SunColorApertureDiv2;
 		SunColorApertureDiv2.X = Context.View.LightShaftColorMask.R;
 		SunColorApertureDiv2.Y = Context.View.LightShaftColorMask.G;
 		SunColorApertureDiv2.Z = Context.View.LightShaftColorMask.B;
 		SunColorApertureDiv2.W = Context.View.FinalPostProcessSettings.DepthOfFieldScale * 0.5f;
-		SetShaderValue(RHICmdList, ShaderRHI, SunColorApertureDiv2Parameter, SunColorApertureDiv2);
+		SetShaderValue(Context.RHICmdList, ShaderRHI, SunColorApertureDiv2Parameter, SunColorApertureDiv2);
 	}
 
 	virtual bool Serialize(FArchive& Ar)
@@ -877,11 +875,11 @@ public:
 		PostprocessParameter.Bind(Initializer.ParameterMap);
 	}
 
-	void SetVS(FRHICommandList& RHICmdList, const FRenderingCompositePassContext& Context)
+	void SetVS(const FRenderingCompositePassContext& Context)
 	{
 		const FVertexShaderRHIParamRef ShaderRHI = GetVertexShader();
-		FGlobalShader::SetParameters(RHICmdList, ShaderRHI, Context.View);
-		PostprocessParameter.SetVS(RHICmdList, ShaderRHI, Context, TStaticSamplerState<SF_Bilinear,AM_Clamp,AM_Clamp,AM_Clamp>::GetRHI());
+		FGlobalShader::SetParameters(Context.RHICmdList, ShaderRHI, Context.View);
+		PostprocessParameter.SetVS(ShaderRHI, Context, TStaticSamplerState<SF_Bilinear,AM_Clamp,AM_Clamp,AM_Clamp>::GetRHI());
 	}
 
 	virtual bool Serialize(FArchive& Ar)
@@ -895,20 +893,20 @@ public:
 IMPLEMENT_SHADER_TYPE(,FPostProcessSunMaskVS_ES2,TEXT("PostProcessMobile"),TEXT("SunMaskVS_ES2"),SF_Vertex);
 
 template <uint32 UseFetchSunDof>
-static void SunMask_SetShader(FRHICommandList& RHICmdList, const FRenderingCompositePassContext& Context)
+static void SunMask_SetShader(const FRenderingCompositePassContext& Context)
 {
 	TShaderMapRef<FPostProcessSunMaskVS_ES2> VertexShader(GetGlobalShaderMap());
 	TShaderMapRef<FPostProcessSunMaskPS_ES2<UseFetchSunDof> > PixelShader(GetGlobalShaderMap());
 
 	static FGlobalBoundShaderState BoundShaderState;
+	Context.RHICmdList.CheckIsNull(); // need new approach for "static FGlobalBoundShaderState" for parallel rendering
+	SetGlobalBoundShaderState(Context.RHICmdList, BoundShaderState, GFilterVertexDeclaration.VertexDeclarationRHI, *VertexShader, *PixelShader);
 
-	SetGlobalBoundShaderState(BoundShaderState, GFilterVertexDeclaration.VertexDeclarationRHI, *VertexShader, *PixelShader);
-
-	VertexShader->SetVS(RHICmdList, Context);
-	PixelShader->SetPS(RHICmdList, Context);
+	VertexShader->SetVS(Context);
+	PixelShader->SetPS(Context);
 }
 
-void FRCPassPostProcessSunMaskES2::SetShader(FRHICommandList& RHICmdList, const FRenderingCompositePassContext& Context)
+void FRCPassPostProcessSunMaskES2::SetShader(const FRenderingCompositePassContext& Context)
 {
 	const FSceneView& View = Context.View;
 	uint32 UseSun = Context.View.bLightShaftUse ? 1 : 0;
@@ -918,14 +916,14 @@ void FRCPassPostProcessSunMaskES2::SetShader(FRHICommandList& RHICmdList, const
 
 	switch(UseFetchSunDof)
 	{
-		case 0: SunMask_SetShader<0>(RHICmdList, Context); break;
-		case 1: SunMask_SetShader<1>(RHICmdList, Context); break;
-		case 2: SunMask_SetShader<2>(RHICmdList, Context); break;
-		case 3: SunMask_SetShader<3>(RHICmdList, Context); break;
-		case 4: SunMask_SetShader<4>(RHICmdList, Context); break;
-		case 5: SunMask_SetShader<5>(RHICmdList, Context); break;
-		case 6: SunMask_SetShader<6>(RHICmdList, Context); break;
-		case 7: SunMask_SetShader<7>(RHICmdList, Context); break;
+		case 0: SunMask_SetShader<0>(Context); break;
+		case 1: SunMask_SetShader<1>(Context); break;
+		case 2: SunMask_SetShader<2>(Context); break;
+		case 3: SunMask_SetShader<3>(Context); break;
+		case 4: SunMask_SetShader<4>(Context); break;
+		case 5: SunMask_SetShader<5>(Context); break;
+		case 6: SunMask_SetShader<6>(Context); break;
+		case 7: SunMask_SetShader<7>(Context); break;
 	}
 }
 
@@ -950,9 +948,6 @@ void FRCPassPostProcessSunMaskES2::Process(FRenderingCompositePassContext& Conte
 	FIntRect SrcRect;
 	const FSceneView& View = Context.View;
 
-	//@todo-rco: RHIPacketList
-	FRHICommandList& RHICmdList = FRHICommandList::GetNullRef();
-
 	TShaderMapRef<FPostProcessSunMaskVS_ES2> VertexShader(GetGlobalShaderMap());
 	if(bOnChip)
 	{
@@ -961,13 +956,14 @@ void FRCPassPostProcessSunMaskES2::Process(FRenderingCompositePassContext& Conte
 
 		Context.SetViewportAndCallRHI(0, 0, 0.0f, DstX, DstY, 1.0f );
 
-		RHISetBlendState(TStaticBlendState<>::GetRHI());
-		RHISetRasterizerState(TStaticRasterizerState<>::GetRHI());
-		RHISetDepthStencilState(TStaticDepthStencilState<false,CF_Always>::GetRHI());
+		Context.RHICmdList.SetBlendState(TStaticBlendState<>::GetRHI());
+		Context.RHICmdList.SetRasterizerState(TStaticRasterizerState<>::GetRHI());
+		Context.RHICmdList.SetDepthStencilState(TStaticDepthStencilState<false,CF_Always>::GetRHI());
 
-		SetShader(RHICmdList, Context);
+		SetShader(Context);
 
 		DrawRectangle(
+			Context.RHICmdList,
 			0, 0,
 			DstX, DstY,
 			SrcRect.Min.X, SrcRect.Min.Y,
@@ -986,20 +982,21 @@ void FRCPassPostProcessSunMaskES2::Process(FRenderingCompositePassContext& Conte
 		SrcRect = View.ViewRect;
 
 		const FSceneRenderTargetItem& DestRenderTarget = PassOutputs[0].RequestSurface(Context);
-		RHISetRenderTarget(DestRenderTarget.TargetableTexture, FTextureRHIRef());	
+		SetRenderTarget(Context.RHICmdList, DestRenderTarget.TargetableTexture, FTextureRHIRef());	
 
 		// is optimized away if possible (RT size=view size, )
-		RHIClear(true, FLinearColor::Black, false, 1.0f, false, 0, FIntRect());
+		Context.RHICmdList.Clear(true, FLinearColor::Black, false, 1.0f, false, 0, FIntRect());
 
 		Context.SetViewportAndCallRHI(0, 0, 0.0f, DstX, DstY, 1.0f );
 
-		RHISetBlendState(TStaticBlendState<>::GetRHI());
-		RHISetRasterizerState(TStaticRasterizerState<>::GetRHI());
-		RHISetDepthStencilState(TStaticDepthStencilState<false,CF_Always>::GetRHI());
+		Context.RHICmdList.SetBlendState(TStaticBlendState<>::GetRHI());
+		Context.RHICmdList.SetRasterizerState(TStaticRasterizerState<>::GetRHI());
+		Context.RHICmdList.SetDepthStencilState(TStaticDepthStencilState<false,CF_Always>::GetRHI());
 
-		SetShader(RHICmdList, Context);
+		SetShader(Context);
 
 		DrawRectangle(
+			Context.RHICmdList,
 			0, 0,
 			DstX, DstY,
 			SrcRect.Min.X, SrcRect.Min.Y,
@@ -1009,7 +1006,7 @@ void FRCPassPostProcessSunMaskES2::Process(FRenderingCompositePassContext& Conte
 			*VertexShader,
 			EDRF_UseTriangleOptimization);
 
-		RHICopyToResolveTarget(DestRenderTarget.TargetableTexture, DestRenderTarget.ShaderResourceTexture, false, FResolveParams());
+		Context.RHICmdList.CopyToResolveTarget(DestRenderTarget.TargetableTexture, DestRenderTarget.ShaderResourceTexture, false, FResolveParams());
 	}
 }
 
@@ -1064,12 +1061,12 @@ public:
 		PostprocessParameter.Bind(Initializer.ParameterMap);
 	}
 
-	void SetPS(FRHICommandList& RHICmdList, const FRenderingCompositePassContext& Context)
+	void SetPS(const FRenderingCompositePassContext& Context)
 	{
 		const FPixelShaderRHIParamRef ShaderRHI = GetPixelShader();
 		const FPostProcessSettings& Settings = Context.View.FinalPostProcessSettings;
-		FGlobalShader::SetParameters(RHICmdList, ShaderRHI, Context.View);
-		PostprocessParameter.SetPS(RHICmdList, ShaderRHI, Context, TStaticSamplerState<SF_Bilinear,AM_Clamp,AM_Clamp,AM_Clamp>::GetRHI());
+		FGlobalShader::SetParameters(Context.RHICmdList, ShaderRHI, Context.View);
+		PostprocessParameter.SetPS(ShaderRHI, Context, TStaticSamplerState<SF_Bilinear,AM_Clamp,AM_Clamp,AM_Clamp>::GetRHI());
 	}
 
 	virtual bool Serialize(FArchive& Ar)
@@ -1107,13 +1104,13 @@ public:
 		LightShaftCenter.Bind(Initializer.ParameterMap, TEXT("LightShaftCenter"));
 	}
 
-	void SetVS(FRHICommandList& RHICmdList, const FRenderingCompositePassContext& Context)
+	void SetVS(const FRenderingCompositePassContext& Context)
 	{
 		const FVertexShaderRHIParamRef ShaderRHI = GetVertexShader();
-		FGlobalShader::SetParameters(RHICmdList, ShaderRHI, Context.View);
-		PostprocessParameter.SetVS(RHICmdList, ShaderRHI, Context, TStaticSamplerState<SF_Bilinear,AM_Clamp,AM_Clamp,AM_Clamp>::GetRHI());
+		FGlobalShader::SetParameters(Context.RHICmdList, ShaderRHI, Context.View);
+		PostprocessParameter.SetVS(ShaderRHI, Context, TStaticSamplerState<SF_Bilinear,AM_Clamp,AM_Clamp,AM_Clamp>::GetRHI());
 
-		SetShaderValue(RHICmdList, ShaderRHI, LightShaftCenter, Context.View.LightShaftCenter);
+		SetShaderValue(Context.RHICmdList, ShaderRHI, LightShaftCenter, Context.View.LightShaftCenter);
 	}
 
 	virtual bool Serialize(FArchive& Ar)
@@ -1127,28 +1124,28 @@ public:
 IMPLEMENT_SHADER_TYPE(,FPostProcessSunAlphaVS_ES2,TEXT("PostProcessMobile"),TEXT("SunAlphaVS_ES2"),SF_Vertex);
 
 template <uint32 UseDof>
-static void SunAlpha_SetShader(FRHICommandList& RHICmdList, const FRenderingCompositePassContext& Context)
+static void SunAlpha_SetShader(const FRenderingCompositePassContext& Context)
 {
 	TShaderMapRef<FPostProcessSunAlphaVS_ES2> VertexShader(GetGlobalShaderMap());
 	TShaderMapRef<FPostProcessSunAlphaPS_ES2<UseDof> > PixelShader(GetGlobalShaderMap());
 
 	static FGlobalBoundShaderState BoundShaderState;
 
-	SetGlobalBoundShaderState(BoundShaderState, GFilterVertexDeclaration.VertexDeclarationRHI, *VertexShader, *PixelShader);
+	SetGlobalBoundShaderState(Context.RHICmdList, BoundShaderState, GFilterVertexDeclaration.VertexDeclarationRHI, *VertexShader, *PixelShader);
 
-	VertexShader->SetVS(RHICmdList, Context);
-	PixelShader->SetPS(RHICmdList, Context);
+	VertexShader->SetVS(Context);
+	PixelShader->SetPS(Context);
 }
 
-void FRCPassPostProcessSunAlphaES2::SetShader(FRHICommandList& RHICmdList, const FRenderingCompositePassContext& Context)
+void FRCPassPostProcessSunAlphaES2::SetShader(const FRenderingCompositePassContext& Context)
 {
 	if(Context.View.FinalPostProcessSettings.DepthOfFieldScale > 0.0f)
 	{
-		SunAlpha_SetShader<1>(RHICmdList, Context);
+		SunAlpha_SetShader<1>(Context);
 	}
 	else
 	{
-		SunAlpha_SetShader<0>(RHICmdList, Context);
+		SunAlpha_SetShader<0>(Context);
 	}
 }
 
@@ -1166,25 +1163,24 @@ void FRCPassPostProcessSunAlphaES2::Process(FRenderingCompositePassContext& Cont
 	DstRect.Max.Y = DstY;
 
 	const FSceneRenderTargetItem& DestRenderTarget = PassOutputs[0].RequestSurface(Context);
-	RHISetRenderTarget(DestRenderTarget.TargetableTexture, FTextureRHIRef());	
+	SetRenderTarget(Context.RHICmdList, DestRenderTarget.TargetableTexture, FTextureRHIRef());	
 
 	// is optimized away if possible (RT size=view size, )
-	RHIClear(true, FLinearColor::Black, false, 1.0f, false, 0, FIntRect());
+	Context.RHICmdList.Clear(true, FLinearColor::Black, false, 1.0f, false, 0, FIntRect());
 
 	Context.SetViewportAndCallRHI(0, 0, 0.0f, DstX, DstY, 1.0f );
 
-	RHISetBlendState(TStaticBlendState<>::GetRHI());
-	RHISetRasterizerState(TStaticRasterizerState<>::GetRHI());
-	RHISetDepthStencilState(TStaticDepthStencilState<false,CF_Always>::GetRHI());
+	Context.RHICmdList.SetBlendState(TStaticBlendState<>::GetRHI());
+	Context.RHICmdList.SetRasterizerState(TStaticRasterizerState<>::GetRHI());
+	Context.RHICmdList.SetDepthStencilState(TStaticDepthStencilState<false,CF_Always>::GetRHI());
 
-	//@todo-rco: RHIPacketList
-	FRHICommandList& RHICmdList = FRHICommandList::GetNullRef();
-	SetShader(RHICmdList, Context);
+	SetShader(Context);
 
 	FIntPoint SrcDstSize = PrePostSourceViewportSize / 4;
 	TShaderMapRef<FPostProcessSunAlphaVS_ES2> VertexShader(GetGlobalShaderMap());
 
 	DrawRectangle(
+		Context.RHICmdList,
 		0, 0,
 		DstX, DstY,
 		0, 0,
@@ -1194,7 +1190,7 @@ void FRCPassPostProcessSunAlphaES2::Process(FRenderingCompositePassContext& Cont
 		*VertexShader,
 		EDRF_UseTriangleOptimization);
 
-	RHICopyToResolveTarget(DestRenderTarget.TargetableTexture, DestRenderTarget.ShaderResourceTexture, false, FResolveParams());
+	Context.RHICmdList.CopyToResolveTarget(DestRenderTarget.TargetableTexture, DestRenderTarget.ShaderResourceTexture, false, FResolveParams());
 }
 
 FPooledRenderTargetDesc FRCPassPostProcessSunAlphaES2::ComputeOutputDesc(EPassOutputId InPassOutputId) const
@@ -1249,12 +1245,12 @@ public:
 		PostprocessParameter.Bind(Initializer.ParameterMap);
 	}
 
-	void SetPS(FRHICommandList& RHICmdList, const FRenderingCompositePassContext& Context)
+	void SetPS(const FRenderingCompositePassContext& Context)
 	{
 		const FPixelShaderRHIParamRef ShaderRHI = GetPixelShader();
 		const FPostProcessSettings& Settings = Context.View.FinalPostProcessSettings;
-		FGlobalShader::SetParameters(RHICmdList, ShaderRHI, Context.View);
-		PostprocessParameter.SetPS(RHICmdList, ShaderRHI, Context, TStaticSamplerState<SF_Bilinear,AM_Clamp,AM_Clamp,AM_Clamp>::GetRHI());
+		FGlobalShader::SetParameters(Context.RHICmdList, ShaderRHI, Context.View);
+		PostprocessParameter.SetPS(ShaderRHI, Context, TStaticSamplerState<SF_Bilinear,AM_Clamp,AM_Clamp,AM_Clamp>::GetRHI());
 	}
 
 	virtual bool Serialize(FArchive& Ar)
@@ -1290,13 +1286,13 @@ public:
 		LightShaftCenter.Bind(Initializer.ParameterMap, TEXT("LightShaftCenter"));
 	}
 
-	void SetVS(FRHICommandList& RHICmdList, const FRenderingCompositePassContext& Context)
+	void SetVS(const FRenderingCompositePassContext& Context)
 	{
 		const FVertexShaderRHIParamRef ShaderRHI = GetVertexShader();
-		FGlobalShader::SetParameters(RHICmdList, ShaderRHI, Context.View);
-		PostprocessParameter.SetVS(RHICmdList, ShaderRHI, Context, TStaticSamplerState<SF_Bilinear,AM_Clamp,AM_Clamp,AM_Clamp>::GetRHI());
+		FGlobalShader::SetParameters(Context.RHICmdList, ShaderRHI, Context.View);
+		PostprocessParameter.SetVS(ShaderRHI, Context, TStaticSamplerState<SF_Bilinear,AM_Clamp,AM_Clamp,AM_Clamp>::GetRHI());
 
-		SetShaderValue(RHICmdList, ShaderRHI, LightShaftCenter, Context.View.LightShaftCenter);
+		SetShaderValue(Context.RHICmdList, ShaderRHI, LightShaftCenter, Context.View.LightShaftCenter);
 	}
 
 	virtual bool Serialize(FArchive& Ar)
@@ -1324,32 +1320,32 @@ void FRCPassPostProcessSunBlurES2::Process(FRenderingCompositePassContext& Conte
 	DstRect.Max.Y = DstY;
 
 	const FSceneRenderTargetItem& DestRenderTarget = PassOutputs[0].RequestSurface(Context);
-	RHISetRenderTarget(DestRenderTarget.TargetableTexture, FTextureRHIRef());	
+	SetRenderTarget(Context.RHICmdList, DestRenderTarget.TargetableTexture, FTextureRHIRef());	
 
 	// is optimized away if possible (RT size=view size, )
-	RHIClear(true, FLinearColor::Black, false, 1.0f, false, 0, FIntRect());
+	Context.RHICmdList.Clear(true, FLinearColor::Black, false, 1.0f, false, 0, FIntRect());
 
 	Context.SetViewportAndCallRHI(0, 0, 0.0f, DstX, DstY, 1.0f );
 
-	RHISetBlendState(TStaticBlendState<>::GetRHI());
-	RHISetRasterizerState(TStaticRasterizerState<>::GetRHI());
-	RHISetDepthStencilState(TStaticDepthStencilState<false,CF_Always>::GetRHI());
+	Context.RHICmdList.SetBlendState(TStaticBlendState<>::GetRHI());
+	Context.RHICmdList.SetRasterizerState(TStaticRasterizerState<>::GetRHI());
+	Context.RHICmdList.SetDepthStencilState(TStaticDepthStencilState<false,CF_Always>::GetRHI());
 
 	TShaderMapRef<FPostProcessSunBlurVS_ES2> VertexShader(GetGlobalShaderMap());
 	TShaderMapRef<FPostProcessSunBlurPS_ES2> PixelShader(GetGlobalShaderMap());
 
 	static FGlobalBoundShaderState BoundShaderState;
+	Context.RHICmdList.CheckIsNull(); // need new approach for "static FGlobalBoundShaderState" for parallel rendering
 
-	SetGlobalBoundShaderState(BoundShaderState, GFilterVertexDeclaration.VertexDeclarationRHI, *VertexShader, *PixelShader);
+	SetGlobalBoundShaderState(Context.RHICmdList, BoundShaderState, GFilterVertexDeclaration.VertexDeclarationRHI, *VertexShader, *PixelShader);
 
-	//@todo-rco: RHIPacketList
-	FRHICommandList& RHICmdList = FRHICommandList::GetNullRef();
-	VertexShader->SetVS(RHICmdList, Context);
-	PixelShader->SetPS(RHICmdList, Context);
+	VertexShader->SetVS(Context);
+	PixelShader->SetPS(Context);
 
 	FIntPoint SrcDstSize = PrePostSourceViewportSize / 4;
 
 	DrawRectangle(
+		Context.RHICmdList,
 		0, 0,
 		DstX, DstY,
 		0, 0,
@@ -1359,7 +1355,7 @@ void FRCPassPostProcessSunBlurES2::Process(FRenderingCompositePassContext& Conte
 		*VertexShader,
 		EDRF_UseTriangleOptimization);
 
-	RHICopyToResolveTarget(DestRenderTarget.TargetableTexture, DestRenderTarget.ShaderResourceTexture, false, FResolveParams());
+	Context.RHICmdList.CopyToResolveTarget(DestRenderTarget.TargetableTexture, DestRenderTarget.ShaderResourceTexture, false, FResolveParams());
 }
 
 FPooledRenderTargetDesc FRCPassPostProcessSunBlurES2::ComputeOutputDesc(EPassOutputId InPassOutputId) const
@@ -1422,24 +1418,24 @@ public:
 		BloomColor.Bind(Initializer.ParameterMap, TEXT("BloomColor"));
 	}
 
-	void SetPS(FRHICommandList& RHICmdList, const FRenderingCompositePassContext& Context)
+	void SetPS(const FRenderingCompositePassContext& Context)
 	{
 		const FPixelShaderRHIParamRef ShaderRHI = GetPixelShader();
 		const FPostProcessSettings& Settings = Context.View.FinalPostProcessSettings;
-		FGlobalShader::SetParameters(RHICmdList, ShaderRHI, Context.View);
-		PostprocessParameter.SetPS(RHICmdList, ShaderRHI, Context, TStaticSamplerState<SF_Bilinear,AM_Clamp,AM_Clamp,AM_Clamp>::GetRHI());
+		FGlobalShader::SetParameters(Context.RHICmdList, ShaderRHI, Context.View);
+		PostprocessParameter.SetPS(ShaderRHI, Context, TStaticSamplerState<SF_Bilinear,AM_Clamp,AM_Clamp,AM_Clamp>::GetRHI());
 
 		FVector4 SunColorVignetteIntensityParam;
 		SunColorVignetteIntensityParam.X = Context.View.LightShaftColorApply.R;
 		SunColorVignetteIntensityParam.Y = Context.View.LightShaftColorApply.G;
 		SunColorVignetteIntensityParam.Z = Context.View.LightShaftColorApply.B;
 		SunColorVignetteIntensityParam.W = Settings.VignetteIntensity;
-		SetShaderValue(RHICmdList, ShaderRHI, SunColorVignetteIntensity, SunColorVignetteIntensityParam);
+		SetShaderValue(Context.RHICmdList, ShaderRHI, SunColorVignetteIntensity, SunColorVignetteIntensityParam);
 
-		SetShaderValue(RHICmdList, ShaderRHI, VignetteColor, Context.View.FinalPostProcessSettings.VignetteColor);
+		SetShaderValue(Context.RHICmdList, ShaderRHI, VignetteColor, Context.View.FinalPostProcessSettings.VignetteColor);
 
 		// Scaling Bloom1 by extra factor to match filter area difference between PC default and mobile.
-		SetShaderValue(RHICmdList, ShaderRHI, BloomColor, Context.View.FinalPostProcessSettings.Bloom1Tint * Context.View.FinalPostProcessSettings.BloomIntensity * 0.5);
+		SetShaderValue(Context.RHICmdList, ShaderRHI, BloomColor, Context.View.FinalPostProcessSettings.Bloom1Tint * Context.View.FinalPostProcessSettings.BloomIntensity * 0.5);
 	}
 
 	virtual bool Serialize(FArchive& Ar)
@@ -1482,13 +1478,13 @@ public:
 		LightShaftCenter.Bind(Initializer.ParameterMap, TEXT("LightShaftCenter"));
 	}
 
-	void SetVS(FRHICommandList& RHICmdList, const FRenderingCompositePassContext& Context)
+	void SetVS(const FRenderingCompositePassContext& Context)
 	{
 		const FVertexShaderRHIParamRef ShaderRHI = GetVertexShader();
-		FGlobalShader::SetParameters(RHICmdList, ShaderRHI, Context.View);
-		PostprocessParameter.SetVS(RHICmdList, ShaderRHI, Context, TStaticSamplerState<SF_Bilinear,AM_Clamp,AM_Clamp,AM_Clamp>::GetRHI());
+		FGlobalShader::SetParameters(Context.RHICmdList, ShaderRHI, Context.View);
+		PostprocessParameter.SetVS(ShaderRHI, Context, TStaticSamplerState<SF_Bilinear,AM_Clamp,AM_Clamp,AM_Clamp>::GetRHI());
 
-		SetShaderValue(RHICmdList, ShaderRHI, LightShaftCenter, Context.View.LightShaftCenter);
+		SetShaderValue(Context.RHICmdList, ShaderRHI, LightShaftCenter, Context.View.LightShaftCenter);
 	}
 
 	virtual bool Serialize(FArchive& Ar)
@@ -1504,20 +1500,21 @@ IMPLEMENT_SHADER_TYPE(,FPostProcessSunMergeVS_ES2,TEXT("PostProcessMobile"),TEXT
 
 
 template <uint32 UseSunBloom>
-static void SunMerge_SetShader(FRHICommandList& RHICmdList, const FRenderingCompositePassContext& Context)
+static void SunMerge_SetShader(const FRenderingCompositePassContext& Context)
 {
 	TShaderMapRef<FPostProcessSunMergeVS_ES2> VertexShader(GetGlobalShaderMap());
 	TShaderMapRef<FPostProcessSunMergePS_ES2<UseSunBloom> > PixelShader(GetGlobalShaderMap());
 
 	static FGlobalBoundShaderState BoundShaderState;
+	Context.RHICmdList.CheckIsNull(); // need new approach for "static FGlobalBoundShaderState" for parallel rendering
 
-	SetGlobalBoundShaderState(BoundShaderState, GFilterVertexDeclaration.VertexDeclarationRHI, *VertexShader, *PixelShader);
+	SetGlobalBoundShaderState(Context.RHICmdList, BoundShaderState, GFilterVertexDeclaration.VertexDeclarationRHI, *VertexShader, *PixelShader);
 
-	VertexShader->SetVS(RHICmdList, Context);
-	PixelShader->SetPS(RHICmdList, Context);
+	VertexShader->SetVS(Context);
+	PixelShader->SetPS(Context);
 }
 
-void FRCPassPostProcessSunMergeES2::SetShader(FRHICommandList& RHICmdList, const FRenderingCompositePassContext& Context)
+void FRCPassPostProcessSunMergeES2::SetShader(const FRenderingCompositePassContext& Context)
 {
 	const FSceneView& View = Context.View;
 	uint32 UseBloom = (View.FinalPostProcessSettings.BloomIntensity > 0.0f) ? 1 : 0;
@@ -1526,10 +1523,10 @@ void FRCPassPostProcessSunMergeES2::SetShader(FRHICommandList& RHICmdList, const
 
 	switch(UseSunBloom)
 	{
-	case 0: SunMerge_SetShader<0>(RHICmdList, Context); break;
-	case 1: SunMerge_SetShader<1>(RHICmdList, Context); break;
-	case 2: SunMerge_SetShader<2>(RHICmdList, Context); break;
-	case 3: SunMerge_SetShader<3>(RHICmdList, Context); break;
+	case 0: SunMerge_SetShader<0>(Context); break;
+	case 1: SunMerge_SetShader<1>(Context); break;
+	case 2: SunMerge_SetShader<2>(Context); break;
+	case 3: SunMerge_SetShader<3>(Context); break;
 	}
 }
 
@@ -1547,25 +1544,24 @@ void FRCPassPostProcessSunMergeES2::Process(FRenderingCompositePassContext& Cont
 	DstRect.Max.Y = DstY;
 
 	const FSceneRenderTargetItem& DestRenderTarget = PassOutputs[0].RequestSurface(Context);
-	RHISetRenderTarget(DestRenderTarget.TargetableTexture, FTextureRHIRef());	
+	SetRenderTarget(Context.RHICmdList, DestRenderTarget.TargetableTexture, FTextureRHIRef());	
 
 	// is optimized away if possible (RT size=view size, )
-	RHIClear(true, FLinearColor::Black, false, 1.0f, false, 0, FIntRect());
+	Context.RHICmdList.Clear(true, FLinearColor::Black, false, 1.0f, false, 0, FIntRect());
 
 	Context.SetViewportAndCallRHI(0, 0, 0.0f, DstX, DstY, 1.0f );
 
-	RHISetBlendState(TStaticBlendState<>::GetRHI());
-	RHISetRasterizerState(TStaticRasterizerState<>::GetRHI());
-	RHISetDepthStencilState(TStaticDepthStencilState<false,CF_Always>::GetRHI());
+	Context.RHICmdList.SetBlendState(TStaticBlendState<>::GetRHI());
+	Context.RHICmdList.SetRasterizerState(TStaticRasterizerState<>::GetRHI());
+	Context.RHICmdList.SetDepthStencilState(TStaticDepthStencilState<false,CF_Always>::GetRHI());
 
-	//@todo-rco: RHIPacketList
-	FRHICommandList& RHICmdList = FRHICommandList::GetNullRef();
-	SetShader(RHICmdList, Context);
+	SetShader(Context);
 
 	FIntPoint SrcDstSize = PrePostSourceViewportSize / 4;
 	TShaderMapRef<FPostProcessSunMergeVS_ES2> VertexShader(GetGlobalShaderMap());
 
 	DrawRectangle(
+		Context.RHICmdList,
 		0, 0,
 		DstX, DstY,
 		0, 0,
@@ -1575,9 +1571,11 @@ void FRCPassPostProcessSunMergeES2::Process(FRenderingCompositePassContext& Cont
 		*VertexShader,
 		EDRF_UseTriangleOptimization);
 
-	RHICopyToResolveTarget(DestRenderTarget.TargetableTexture, DestRenderTarget.ShaderResourceTexture, false, FResolveParams());
+	Context.RHICmdList.CopyToResolveTarget(DestRenderTarget.TargetableTexture, DestRenderTarget.ShaderResourceTexture, false, FResolveParams());
 
 	// Double buffer sun+bloom+vignette composite.
+	Context.RHICmdList.CheckIsNull(); // this modifies the view state, so probably won't work parallel
+
 	if(Context.View.FinalPostProcessSettings.AntiAliasingMethod == AAM_TemporalAA)
 	{
 		FSceneViewState* ViewState = (FSceneViewState*)Context.View.State;
@@ -1647,25 +1645,25 @@ public:
 		BloomColor2.Bind(Initializer.ParameterMap, TEXT("BloomColor2"));
 	}
 
-	void SetPS(FRHICommandList& RHICmdList, const FRenderingCompositePassContext& Context)
+	void SetPS(const FRenderingCompositePassContext& Context)
 	{
 		const FPixelShaderRHIParamRef ShaderRHI = GetPixelShader();
 		const FPostProcessSettings& Settings = Context.View.FinalPostProcessSettings;
-		FGlobalShader::SetParameters(RHICmdList, ShaderRHI, Context.View);
-		PostprocessParameter.SetPS(RHICmdList, ShaderRHI, Context, TStaticSamplerState<SF_Bilinear,AM_Clamp,AM_Clamp,AM_Clamp>::GetRHI());
+		FGlobalShader::SetParameters(Context.RHICmdList, ShaderRHI, Context.View);
+		PostprocessParameter.SetPS(ShaderRHI, Context, TStaticSamplerState<SF_Bilinear,AM_Clamp,AM_Clamp,AM_Clamp>::GetRHI());
 
 		FVector4 SunColorVignetteIntensityParam;
 		SunColorVignetteIntensityParam.X = Context.View.LightShaftColorApply.R;
 		SunColorVignetteIntensityParam.Y = Context.View.LightShaftColorApply.G;
 		SunColorVignetteIntensityParam.Z = Context.View.LightShaftColorApply.B;
 		SunColorVignetteIntensityParam.W = Settings.VignetteIntensity;
-		SetShaderValue(RHICmdList, ShaderRHI, SunColorVignetteIntensity, SunColorVignetteIntensityParam);
+		SetShaderValue(Context.RHICmdList, ShaderRHI, SunColorVignetteIntensity, SunColorVignetteIntensityParam);
 
-		SetShaderValue(RHICmdList, ShaderRHI, VignetteColor, Context.View.FinalPostProcessSettings.VignetteColor);
+		SetShaderValue(Context.RHICmdList, ShaderRHI, VignetteColor, Context.View.FinalPostProcessSettings.VignetteColor);
 
 		// Scaling Bloom1 by extra factor to match filter area difference between PC default and mobile.
-		SetShaderValue(RHICmdList, ShaderRHI, BloomColor, Context.View.FinalPostProcessSettings.Bloom1Tint * Context.View.FinalPostProcessSettings.BloomIntensity * 0.5);
-		SetShaderValue(RHICmdList, ShaderRHI, BloomColor2, Context.View.FinalPostProcessSettings.Bloom2Tint * Context.View.FinalPostProcessSettings.BloomIntensity);
+		SetShaderValue(Context.RHICmdList, ShaderRHI, BloomColor, Context.View.FinalPostProcessSettings.Bloom1Tint * Context.View.FinalPostProcessSettings.BloomIntensity * 0.5);
+		SetShaderValue(Context.RHICmdList, ShaderRHI, BloomColor2, Context.View.FinalPostProcessSettings.Bloom2Tint * Context.View.FinalPostProcessSettings.BloomIntensity);
 	}
 
 	virtual bool Serialize(FArchive& Ar)
@@ -1698,11 +1696,11 @@ public:
 		PostprocessParameter.Bind(Initializer.ParameterMap);
 	}
 
-	void SetVS(FRHICommandList& RHICmdList, const FRenderingCompositePassContext& Context)
+	void SetVS(const FRenderingCompositePassContext& Context)
 	{
 		const FVertexShaderRHIParamRef ShaderRHI = GetVertexShader();
-		FGlobalShader::SetParameters(RHICmdList, ShaderRHI, Context.View);
-		PostprocessParameter.SetVS(RHICmdList, ShaderRHI, Context, TStaticSamplerState<SF_Bilinear,AM_Clamp,AM_Clamp,AM_Clamp>::GetRHI());
+		FGlobalShader::SetParameters(Context.RHICmdList, ShaderRHI, Context.View);
+		PostprocessParameter.SetVS(ShaderRHI, Context, TStaticSamplerState<SF_Bilinear,AM_Clamp,AM_Clamp,AM_Clamp>::GetRHI());
 	}
 
 	virtual bool Serialize(FArchive& Ar)
@@ -1715,17 +1713,18 @@ public:
 
 IMPLEMENT_SHADER_TYPE(,FPostProcessSunMergeSmallVS_ES2,TEXT("PostProcessMobile"),TEXT("SunMergeSmallVS_ES2"),SF_Vertex);
 
-void FRCPassPostProcessSunMergeSmallES2::SetShader(FRHICommandList& RHICmdList, const FRenderingCompositePassContext& Context)
+void FRCPassPostProcessSunMergeSmallES2::SetShader(const FRenderingCompositePassContext& Context)
 {
 	TShaderMapRef<FPostProcessSunMergeSmallVS_ES2> VertexShader(GetGlobalShaderMap());
 	TShaderMapRef<FPostProcessSunMergeSmallPS_ES2> PixelShader(GetGlobalShaderMap());
 
 	static FGlobalBoundShaderState BoundShaderState;
+	Context.RHICmdList.CheckIsNull(); // need new approach for "static FGlobalBoundShaderState" for parallel rendering
 
-	SetGlobalBoundShaderState(BoundShaderState, GFilterVertexDeclaration.VertexDeclarationRHI, *VertexShader, *PixelShader);
+	SetGlobalBoundShaderState(Context.RHICmdList, BoundShaderState, GFilterVertexDeclaration.VertexDeclarationRHI, *VertexShader, *PixelShader);
 
-	VertexShader->SetVS(RHICmdList, Context);
-	PixelShader->SetPS(RHICmdList, Context);
+	VertexShader->SetVS(Context);
+	PixelShader->SetPS(Context);
 }
 
 void FRCPassPostProcessSunMergeSmallES2::Process(FRenderingCompositePassContext& Context)
@@ -1742,25 +1741,24 @@ void FRCPassPostProcessSunMergeSmallES2::Process(FRenderingCompositePassContext&
 	DstRect.Max.Y = DstY;
 
 	const FSceneRenderTargetItem& DestRenderTarget = PassOutputs[0].RequestSurface(Context);
-	RHISetRenderTarget(DestRenderTarget.TargetableTexture, FTextureRHIRef());	
+	SetRenderTarget(Context.RHICmdList, DestRenderTarget.TargetableTexture, FTextureRHIRef());	
 
 	// is optimized away if possible (RT size=view size, )
-	RHIClear(true, FLinearColor::Black, false, 1.0f, false, 0, FIntRect());
+	Context.RHICmdList.Clear(true, FLinearColor::Black, false, 1.0f, false, 0, FIntRect());
 
 	Context.SetViewportAndCallRHI(0, 0, 0.0f, DstX, DstY, 1.0f );
 
-	RHISetBlendState(TStaticBlendState<>::GetRHI());
-	RHISetRasterizerState(TStaticRasterizerState<>::GetRHI());
-	RHISetDepthStencilState(TStaticDepthStencilState<false,CF_Always>::GetRHI());
+	Context.RHICmdList.SetBlendState(TStaticBlendState<>::GetRHI());
+	Context.RHICmdList.SetRasterizerState(TStaticRasterizerState<>::GetRHI());
+	Context.RHICmdList.SetDepthStencilState(TStaticDepthStencilState<false,CF_Always>::GetRHI());
 
-	//@todo-rco: RHIPacketList
-	FRHICommandList& RHICmdList = FRHICommandList::GetNullRef();
-	SetShader(RHICmdList, Context);
+	SetShader(Context);
 
 	FIntPoint SrcDstSize = PrePostSourceViewportSize / 4;
 	TShaderMapRef<FPostProcessSunMergeSmallVS_ES2> VertexShader(GetGlobalShaderMap());
 
 	DrawRectangle(
+		Context.RHICmdList, 
 		0, 0,
 		DstX, DstY,
 		0, 0,
@@ -1770,10 +1768,11 @@ void FRCPassPostProcessSunMergeSmallES2::Process(FRenderingCompositePassContext&
 		*VertexShader,
 		EDRF_UseTriangleOptimization);
 
-	RHICopyToResolveTarget(DestRenderTarget.TargetableTexture, DestRenderTarget.ShaderResourceTexture, false, FResolveParams());
+	Context.RHICmdList.CopyToResolveTarget(DestRenderTarget.TargetableTexture, DestRenderTarget.ShaderResourceTexture, false, FResolveParams());
 
 	// Double buffer sun+bloom+vignette composite.
-	if(Context.View.FinalPostProcessSettings.AntiAliasingMethod == AAM_TemporalAA)
+	Context.RHICmdList.CheckIsNull(); // this modifies the view state, so probably won't work parallel
+	if (Context.View.FinalPostProcessSettings.AntiAliasingMethod == AAM_TemporalAA)
 	{
 		FSceneViewState* ViewState = (FSceneViewState*)Context.View.State;
 		if(ViewState) 
@@ -1835,13 +1834,13 @@ public:
 		PostprocessParameter.Bind(Initializer.ParameterMap);
 	}
 
-	void SetVS(FRHICommandList& RHICmdList, const FRenderingCompositePassContext& Context)
+	void SetVS(const FRenderingCompositePassContext& Context)
 	{
 		const FVertexShaderRHIParamRef ShaderRHI = GetVertexShader();
 
-		FGlobalShader::SetParameters(RHICmdList, ShaderRHI, Context.View);
+		FGlobalShader::SetParameters(Context.RHICmdList, ShaderRHI, Context.View);
 
-		PostprocessParameter.SetVS(RHICmdList, ShaderRHI, Context, TStaticSamplerState<SF_Point,AM_Clamp,AM_Clamp,AM_Clamp>::GetRHI());
+		PostprocessParameter.SetVS(ShaderRHI, Context, TStaticSamplerState<SF_Point,AM_Clamp,AM_Clamp,AM_Clamp>::GetRHI());
 	}
 
 	// FShader interface.
@@ -1880,12 +1879,12 @@ public:
 		PostprocessParameter.Bind(Initializer.ParameterMap);
 	}
 
-	void SetPS(FRHICommandList& RHICmdList, const FRenderingCompositePassContext& Context)
+	void SetPS(const FRenderingCompositePassContext& Context)
 	{
 		const FPixelShaderRHIParamRef ShaderRHI = GetPixelShader();
 		const FPostProcessSettings& Settings = Context.View.FinalPostProcessSettings;
-		FGlobalShader::SetParameters(RHICmdList, ShaderRHI, Context.View);
-		PostprocessParameter.SetPS(RHICmdList, ShaderRHI, Context, TStaticSamplerState<SF_Point,AM_Clamp,AM_Clamp,AM_Clamp>::GetRHI());
+		FGlobalShader::SetParameters(Context.RHICmdList, ShaderRHI, Context.View);
+		PostprocessParameter.SetPS(ShaderRHI, Context, TStaticSamplerState<SF_Point,AM_Clamp,AM_Clamp,AM_Clamp>::GetRHI());
 	}
 
 	virtual bool Serialize(FArchive& Ar)
@@ -1904,29 +1903,30 @@ IMPLEMENT_SHADER_TYPE(template<>,FPostProcessDofDownPS_ES2_0,TEXT("PostProcessMo
 IMPLEMENT_SHADER_TYPE(template<>,FPostProcessDofDownPS_ES2_1,TEXT("PostProcessMobile"),TEXT("DofDownPS_ES2"),SF_Pixel);
 
 template <uint32 UseSun>
-static void DofDown_SetShader(FRHICommandList& RHICmdList, const FRenderingCompositePassContext& Context)
+static void DofDown_SetShader(const FRenderingCompositePassContext& Context)
 {
 	TShaderMapRef<FPostProcessDofDownVS_ES2> VertexShader(GetGlobalShaderMap());
 	TShaderMapRef<FPostProcessDofDownPS_ES2<UseSun> > PixelShader(GetGlobalShaderMap());
 
 	static FGlobalBoundShaderState BoundShaderState;
+	Context.RHICmdList.CheckIsNull(); // need new approach for "static FGlobalBoundShaderState" for parallel rendering
 
-	SetGlobalBoundShaderState(BoundShaderState, GFilterVertexDeclaration.VertexDeclarationRHI, *VertexShader, *PixelShader);
+	SetGlobalBoundShaderState(Context.RHICmdList, BoundShaderState, GFilterVertexDeclaration.VertexDeclarationRHI, *VertexShader, *PixelShader);
 
-	VertexShader->SetVS(RHICmdList, Context);
-	PixelShader->SetPS(RHICmdList, Context);
+	VertexShader->SetVS(Context);
+	PixelShader->SetPS(Context);
 }
 
-void FRCPassPostProcessDofDownES2::SetShader(FRHICommandList& RHICmdList, const FRenderingCompositePassContext& Context)
+void FRCPassPostProcessDofDownES2::SetShader(const FRenderingCompositePassContext& Context)
 {
 	const FSceneView& View = Context.View;
 	if(Context.View.bLightShaftUse)
 	{
-		DofDown_SetShader<1>(RHICmdList, Context);
+		DofDown_SetShader<1>(Context);
 	}
 	else
 	{
-		DofDown_SetShader<0>(RHICmdList, Context);
+		DofDown_SetShader<0>(Context);
 	}
 }
 
@@ -1967,24 +1967,23 @@ void FRCPassPostProcessDofDownES2::Process(FRenderingCompositePassContext& Conte
 	}
 
 	const FSceneRenderTargetItem& DestRenderTarget = PassOutputs[0].RequestSurface(Context);
-	RHISetRenderTarget(DestRenderTarget.TargetableTexture, FTextureRHIRef());	
+	SetRenderTarget(Context.RHICmdList, DestRenderTarget.TargetableTexture, FTextureRHIRef());	
 
 	// is optimized away if possible (RT size=view size, )
-	RHIClear(true, FLinearColor::Black, false, 1.0f, false, 0, FIntRect());
+	Context.RHICmdList.Clear(true, FLinearColor::Black, false, 1.0f, false, 0, FIntRect());
 
 	Context.SetViewportAndCallRHI(0, 0, 0.0f, DstX, DstY, 1.0f );
 
-	RHISetBlendState(TStaticBlendState<>::GetRHI());
-	RHISetRasterizerState(TStaticRasterizerState<>::GetRHI());
-	RHISetDepthStencilState(TStaticDepthStencilState<false,CF_Always>::GetRHI());
+	Context.RHICmdList.SetBlendState(TStaticBlendState<>::GetRHI());
+	Context.RHICmdList.SetRasterizerState(TStaticRasterizerState<>::GetRHI());
+	Context.RHICmdList.SetDepthStencilState(TStaticDepthStencilState<false,CF_Always>::GetRHI());
 
-	//@todo-rco: RHIPacketList
-	FRHICommandList& RHICmdList = FRHICommandList::GetNullRef();
-	SetShader(RHICmdList, Context);
+	SetShader(Context);
 
 	TShaderMapRef<FPostProcessDofDownVS_ES2> VertexShader(GetGlobalShaderMap());
 
 	DrawRectangle(
+		Context.RHICmdList,
 		0, 0,
 		DstX, DstY,
 		SrcRect.Min.X, SrcRect.Min.Y,
@@ -1994,7 +1993,7 @@ void FRCPassPostProcessDofDownES2::Process(FRenderingCompositePassContext& Conte
 		*VertexShader,
 		EDRF_UseTriangleOptimization);
 
-	RHICopyToResolveTarget(DestRenderTarget.TargetableTexture, DestRenderTarget.ShaderResourceTexture, false, FResolveParams());
+	Context.RHICmdList.CopyToResolveTarget(DestRenderTarget.TargetableTexture, DestRenderTarget.ShaderResourceTexture, false, FResolveParams());
 }
 
 FPooledRenderTargetDesc FRCPassPostProcessDofDownES2::ComputeOutputDesc(EPassOutputId InPassOutputId) const
@@ -2043,13 +2042,13 @@ public:
 		PostprocessParameter.Bind(Initializer.ParameterMap);
 	}
 
-	void SetVS(FRHICommandList& RHICmdList, const FRenderingCompositePassContext& Context)
+	void SetVS(const FRenderingCompositePassContext& Context)
 	{
 		const FVertexShaderRHIParamRef ShaderRHI = GetVertexShader();
 
-		FGlobalShader::SetParameters(RHICmdList, ShaderRHI, Context.View);
+		FGlobalShader::SetParameters(Context.RHICmdList, ShaderRHI, Context.View);
 
-		PostprocessParameter.SetVS(RHICmdList, ShaderRHI, Context, TStaticSamplerState<SF_Point,AM_Clamp,AM_Clamp,AM_Clamp>::GetRHI());
+		PostprocessParameter.SetVS(ShaderRHI, Context, TStaticSamplerState<SF_Point,AM_Clamp,AM_Clamp,AM_Clamp>::GetRHI());
 	}
 
 	// FShader interface.
@@ -2088,12 +2087,12 @@ public:
 		PostprocessParameter.Bind(Initializer.ParameterMap);
 	}
 
-	void SetPS(FRHICommandList& RHICmdList, const FRenderingCompositePassContext& Context)
+	void SetPS(const FRenderingCompositePassContext& Context)
 	{
 		const FPixelShaderRHIParamRef ShaderRHI = GetPixelShader();
 		const FPostProcessSettings& Settings = Context.View.FinalPostProcessSettings;
-		FGlobalShader::SetParameters(RHICmdList, ShaderRHI, Context.View);
-		PostprocessParameter.SetPS(RHICmdList, ShaderRHI, Context, TStaticSamplerState<SF_Point,AM_Clamp,AM_Clamp,AM_Clamp>::GetRHI());
+		FGlobalShader::SetParameters(Context.RHICmdList, ShaderRHI, Context.View);
+		PostprocessParameter.SetPS(ShaderRHI, Context, TStaticSamplerState<SF_Point,AM_Clamp,AM_Clamp,AM_Clamp>::GetRHI());
 	}
 
 	virtual bool Serialize(FArchive& Ar)
@@ -2112,29 +2111,30 @@ IMPLEMENT_SHADER_TYPE(template<>,FPostProcessDofNearPS_ES2_0,TEXT("PostProcessMo
 IMPLEMENT_SHADER_TYPE(template<>,FPostProcessDofNearPS_ES2_1,TEXT("PostProcessMobile"),TEXT("DofNearPS_ES2"),SF_Pixel);
 
 template <uint32 UseSun>
-static void DofNear_SetShader(FRHICommandList& RHICmdList, const FRenderingCompositePassContext& Context)
+static void DofNear_SetShader(const FRenderingCompositePassContext& Context)
 {
 	TShaderMapRef<FPostProcessDofNearVS_ES2> VertexShader(GetGlobalShaderMap());
 	TShaderMapRef<FPostProcessDofNearPS_ES2<UseSun> > PixelShader(GetGlobalShaderMap());
 
 	static FGlobalBoundShaderState BoundShaderState;
+	Context.RHICmdList.CheckIsNull(); // need new approach for "static FGlobalBoundShaderState" for parallel rendering
 
-	SetGlobalBoundShaderState(BoundShaderState, GFilterVertexDeclaration.VertexDeclarationRHI, *VertexShader, *PixelShader);
+	SetGlobalBoundShaderState(Context.RHICmdList, BoundShaderState, GFilterVertexDeclaration.VertexDeclarationRHI, *VertexShader, *PixelShader);
 
-	VertexShader->SetVS(RHICmdList, Context);
-	PixelShader->SetPS(RHICmdList, Context);
+	VertexShader->SetVS(Context);
+	PixelShader->SetPS(Context);
 }
 
-void FRCPassPostProcessDofNearES2::SetShader(FRHICommandList& RHICmdList, const FRenderingCompositePassContext& Context)
+void FRCPassPostProcessDofNearES2::SetShader(const FRenderingCompositePassContext& Context)
 {
 	const FSceneView& View = Context.View;
 	if(Context.View.bLightShaftUse)
 	{
-		DofNear_SetShader<1>(RHICmdList, Context);
+		DofNear_SetShader<1>(Context);
 	}
 	else
 	{
-		DofNear_SetShader<0>(RHICmdList, Context);
+		DofNear_SetShader<0>(Context);
 	}
 }
 
@@ -2152,25 +2152,24 @@ void FRCPassPostProcessDofNearES2::Process(FRenderingCompositePassContext& Conte
 	DstRect.Max.Y = DstY;
 
 	const FSceneRenderTargetItem& DestRenderTarget = PassOutputs[0].RequestSurface(Context);
-	RHISetRenderTarget(DestRenderTarget.TargetableTexture, FTextureRHIRef());	
+	SetRenderTarget(Context.RHICmdList, DestRenderTarget.TargetableTexture, FTextureRHIRef());	
 
 	// is optimized away if possible (RT size=view size, )
-	RHIClear(true, FLinearColor::Black, false, 1.0f, false, 0, FIntRect());
+	Context.RHICmdList.Clear(true, FLinearColor::Black, false, 1.0f, false, 0, FIntRect());
 
 	Context.SetViewportAndCallRHI(0, 0, 0.0f, DstX, DstY, 1.0f );
 
-	RHISetBlendState(TStaticBlendState<>::GetRHI());
-	RHISetRasterizerState(TStaticRasterizerState<>::GetRHI());
-	RHISetDepthStencilState(TStaticDepthStencilState<false,CF_Always>::GetRHI());
+	Context.RHICmdList.SetBlendState(TStaticBlendState<>::GetRHI());
+	Context.RHICmdList.SetRasterizerState(TStaticRasterizerState<>::GetRHI());
+	Context.RHICmdList.SetDepthStencilState(TStaticDepthStencilState<false,CF_Always>::GetRHI());
 
-	//@todo-rco: RHIPacketList
-	FRHICommandList& RHICmdList = FRHICommandList::GetNullRef();
-	SetShader(RHICmdList, Context);
+	SetShader(Context);
 
 	FIntPoint SrcDstSize = PrePostSourceViewportSize / 4;
 	TShaderMapRef<FPostProcessDofNearVS_ES2> VertexShader(GetGlobalShaderMap());
 
 	DrawRectangle(
+		Context.RHICmdList,
 		0, 0,
 		DstX, DstY,
 		0, 0,
@@ -2180,7 +2179,7 @@ void FRCPassPostProcessDofNearES2::Process(FRenderingCompositePassContext& Conte
 		*VertexShader,
 		EDRF_UseTriangleOptimization);
 
-	RHICopyToResolveTarget(DestRenderTarget.TargetableTexture, DestRenderTarget.ShaderResourceTexture, false, FResolveParams());
+	Context.RHICmdList.CopyToResolveTarget(DestRenderTarget.TargetableTexture, DestRenderTarget.ShaderResourceTexture, false, FResolveParams());
 }
 
 FPooledRenderTargetDesc FRCPassPostProcessDofNearES2::ComputeOutputDesc(EPassOutputId InPassOutputId) const
@@ -2232,12 +2231,12 @@ public:
 		PostprocessParameter.Bind(Initializer.ParameterMap);
 	}
 
-	void SetPS(FRHICommandList& RHICmdList, const FRenderingCompositePassContext& Context)
+	void SetPS(const FRenderingCompositePassContext& Context)
 	{
 		const FPixelShaderRHIParamRef ShaderRHI = GetPixelShader();
 		const FPostProcessSettings& Settings = Context.View.FinalPostProcessSettings;
-		FGlobalShader::SetParameters(RHICmdList, ShaderRHI, Context.View);
-		PostprocessParameter.SetPS(RHICmdList, ShaderRHI, Context, TStaticSamplerState<SF_Bilinear,AM_Clamp,AM_Clamp,AM_Clamp>::GetRHI());
+		FGlobalShader::SetParameters(Context.RHICmdList, ShaderRHI, Context.View);
+		PostprocessParameter.SetPS(ShaderRHI, Context, TStaticSamplerState<SF_Bilinear,AM_Clamp,AM_Clamp,AM_Clamp>::GetRHI());
 	}
 
 	virtual bool Serialize(FArchive& Ar)
@@ -2271,11 +2270,11 @@ public:
 		PostprocessParameter.Bind(Initializer.ParameterMap);
 	}
 
-	void SetVS(FRHICommandList& RHICmdList, const FRenderingCompositePassContext& Context)
+	void SetVS(const FRenderingCompositePassContext& Context)
 	{
 		const FVertexShaderRHIParamRef ShaderRHI = GetVertexShader();
-		FGlobalShader::SetParameters(RHICmdList, ShaderRHI, Context.View);
-		PostprocessParameter.SetVS(RHICmdList, ShaderRHI, Context, TStaticSamplerState<SF_Point,AM_Clamp,AM_Clamp,AM_Clamp>::GetRHI());
+		FGlobalShader::SetParameters(Context.RHICmdList, ShaderRHI, Context.View);
+		PostprocessParameter.SetVS(ShaderRHI, Context, TStaticSamplerState<SF_Point,AM_Clamp,AM_Clamp,AM_Clamp>::GetRHI());
 	}
 
 	virtual bool Serialize(FArchive& Ar)
@@ -2303,31 +2302,31 @@ void FRCPassPostProcessDofBlurES2::Process(FRenderingCompositePassContext& Conte
 	DstRect.Max.Y = DstY;
 
 	const FSceneRenderTargetItem& DestRenderTarget = PassOutputs[0].RequestSurface(Context);
-	RHISetRenderTarget(DestRenderTarget.TargetableTexture, FTextureRHIRef());	
+	SetRenderTarget(Context.RHICmdList, DestRenderTarget.TargetableTexture, FTextureRHIRef());	
 
 	// is optimized away if possible (RT size=view size, )
-	RHIClear(true, FLinearColor::Black, false, 1.0f, false, 0, FIntRect());
+	Context.RHICmdList.Clear(true, FLinearColor::Black, false, 1.0f, false, 0, FIntRect());
 
 	Context.SetViewportAndCallRHI(0, 0, 0.0f, DstX, DstY, 1.0f );
 
-	RHISetBlendState(TStaticBlendState<>::GetRHI());
-	RHISetRasterizerState(TStaticRasterizerState<>::GetRHI());
-	RHISetDepthStencilState(TStaticDepthStencilState<false,CF_Always>::GetRHI());
+	Context.RHICmdList.SetBlendState(TStaticBlendState<>::GetRHI());
+	Context.RHICmdList.SetRasterizerState(TStaticRasterizerState<>::GetRHI());
+	Context.RHICmdList.SetDepthStencilState(TStaticDepthStencilState<false,CF_Always>::GetRHI());
 
 	TShaderMapRef<FPostProcessDofBlurVS_ES2> VertexShader(GetGlobalShaderMap());
 	TShaderMapRef<FPostProcessDofBlurPS_ES2> PixelShader(GetGlobalShaderMap());
 
 	static FGlobalBoundShaderState BoundShaderState;
-	SetGlobalBoundShaderState(BoundShaderState, GFilterVertexDeclaration.VertexDeclarationRHI, *VertexShader, *PixelShader);
+	Context.RHICmdList.CheckIsNull(); // need new approach for "static FGlobalBoundShaderState" for parallel rendering
+	SetGlobalBoundShaderState(Context.RHICmdList, BoundShaderState, GFilterVertexDeclaration.VertexDeclarationRHI, *VertexShader, *PixelShader);
 
-	//@todo-rco: RHIPacketList
-	FRHICommandList& RHICmdList = FRHICommandList::GetNullRef();
-	VertexShader->SetVS(RHICmdList, Context);
-	PixelShader->SetPS(RHICmdList, Context);
+	VertexShader->SetVS(Context);
+	PixelShader->SetPS(Context);
 
 	FIntPoint SrcDstSize = PrePostSourceViewportSize / 2;
 
 	DrawRectangle(
+		Context.RHICmdList,
 		0, 0,
 		DstX, DstY,
 		0, 0,
@@ -2337,7 +2336,7 @@ void FRCPassPostProcessDofBlurES2::Process(FRenderingCompositePassContext& Conte
 		*VertexShader,
 		EDRF_UseTriangleOptimization);
 
-	RHICopyToResolveTarget(DestRenderTarget.TargetableTexture, DestRenderTarget.ShaderResourceTexture, false, FResolveParams());
+	Context.RHICmdList.CopyToResolveTarget(DestRenderTarget.TargetableTexture, DestRenderTarget.ShaderResourceTexture, false, FResolveParams());
 }
 
 FPooledRenderTargetDesc FRCPassPostProcessDofBlurES2::ComputeOutputDesc(EPassOutputId InPassOutputId) const
@@ -2391,12 +2390,12 @@ public:
 		PostprocessParameter.Bind(Initializer.ParameterMap);
 	}
 
-	void SetPS(FRHICommandList& RHICmdList, const FRenderingCompositePassContext& Context)
+	void SetPS(const FRenderingCompositePassContext& Context)
 	{
 		const FPixelShaderRHIParamRef ShaderRHI = GetPixelShader();
 		const FPostProcessSettings& Settings = Context.View.FinalPostProcessSettings;
-		FGlobalShader::SetParameters(RHICmdList, ShaderRHI, Context.View);
-		PostprocessParameter.SetPS(RHICmdList, ShaderRHI, Context, TStaticSamplerState<SF_Bilinear,AM_Clamp,AM_Clamp,AM_Clamp>::GetRHI());
+		FGlobalShader::SetParameters(Context.RHICmdList, ShaderRHI, Context.View);
+		PostprocessParameter.SetPS(ShaderRHI, Context, TStaticSamplerState<SF_Bilinear,AM_Clamp,AM_Clamp,AM_Clamp>::GetRHI());
 	}
 
 	virtual bool Serialize(FArchive& Ar)
@@ -2431,11 +2430,11 @@ public:
 		PostprocessParameter.Bind(Initializer.ParameterMap);
 	}
 
-	void SetVS(FRHICommandList& RHICmdList, const FRenderingCompositePassContext& Context)
+	void SetVS(const FRenderingCompositePassContext& Context)
 	{
 		const FVertexShaderRHIParamRef ShaderRHI = GetVertexShader();
-		FGlobalShader::SetParameters(RHICmdList, ShaderRHI, Context.View);
-		PostprocessParameter.SetVS(RHICmdList, ShaderRHI, Context, TStaticSamplerState<SF_Bilinear,AM_Clamp,AM_Clamp,AM_Clamp>::GetRHI());
+		FGlobalShader::SetParameters(Context.RHICmdList, ShaderRHI, Context.View);
+		PostprocessParameter.SetVS(ShaderRHI, Context, TStaticSamplerState<SF_Bilinear,AM_Clamp,AM_Clamp,AM_Clamp>::GetRHI());
 	}
 
 	virtual bool Serialize(FArchive& Ar)
@@ -2450,22 +2449,23 @@ IMPLEMENT_SHADER_TYPE(,FPostProcessSunAvgVS_ES2,TEXT("PostProcessMobile"),TEXT("
 
 
 
-static void SunAvg_SetShader(FRHICommandList& RHICmdList, const FRenderingCompositePassContext& Context)
+static void SunAvg_SetShader(const FRenderingCompositePassContext& Context)
 {
 	TShaderMapRef<FPostProcessSunAvgVS_ES2> VertexShader(GetGlobalShaderMap());
 	TShaderMapRef<FPostProcessSunAvgPS_ES2> PixelShader(GetGlobalShaderMap());
 
 	static FGlobalBoundShaderState BoundShaderState;
+	Context.RHICmdList.CheckIsNull(); // need new approach for "static FGlobalBoundShaderState" for parallel rendering
 
-	SetGlobalBoundShaderState(BoundShaderState, GFilterVertexDeclaration.VertexDeclarationRHI, *VertexShader, *PixelShader);
+	SetGlobalBoundShaderState(Context.RHICmdList, BoundShaderState, GFilterVertexDeclaration.VertexDeclarationRHI, *VertexShader, *PixelShader);
 
-	VertexShader->SetVS(RHICmdList, Context);
-	PixelShader->SetPS(RHICmdList, Context);
+	VertexShader->SetVS(Context);
+	PixelShader->SetPS(Context);
 }
 
-void FRCPassPostProcessSunAvgES2::SetShader(FRHICommandList& RHICmdList, const FRenderingCompositePassContext& Context)
+void FRCPassPostProcessSunAvgES2::SetShader(const FRenderingCompositePassContext& Context)
 {
-	SunAvg_SetShader(RHICmdList, Context);
+	SunAvg_SetShader(Context);
 }
 
 void FRCPassPostProcessSunAvgES2::Process(FRenderingCompositePassContext& Context)
@@ -2482,25 +2482,24 @@ void FRCPassPostProcessSunAvgES2::Process(FRenderingCompositePassContext& Contex
 	DstRect.Max.Y = DstY;
 
 	const FSceneRenderTargetItem& DestRenderTarget = PassOutputs[0].RequestSurface(Context);
-	RHISetRenderTarget(DestRenderTarget.TargetableTexture, FTextureRHIRef());	
+	SetRenderTarget(Context.RHICmdList, DestRenderTarget.TargetableTexture, FTextureRHIRef());	
 
 	// is optimized away if possible (RT size=view size, )
-	RHIClear(true, FLinearColor::Black, false, 1.0f, false, 0, FIntRect());
+	Context.RHICmdList.Clear(true, FLinearColor::Black, false, 1.0f, false, 0, FIntRect());
 
 	Context.SetViewportAndCallRHI(0, 0, 0.0f, DstX, DstY, 1.0f );
 
-	RHISetBlendState(TStaticBlendState<>::GetRHI());
-	RHISetRasterizerState(TStaticRasterizerState<>::GetRHI());
-	RHISetDepthStencilState(TStaticDepthStencilState<false,CF_Always>::GetRHI());
+	Context.RHICmdList.SetBlendState(TStaticBlendState<>::GetRHI());
+	Context.RHICmdList.SetRasterizerState(TStaticRasterizerState<>::GetRHI());
+	Context.RHICmdList.SetDepthStencilState(TStaticDepthStencilState<false,CF_Always>::GetRHI());
 
-	//@todo-rco: RHIPacketList
-	FRHICommandList& RHICmdList = FRHICommandList::GetNullRef();
-	SetShader(RHICmdList, Context);
+	SetShader(Context);
 
 	FIntPoint SrcDstSize = PrePostSourceViewportSize / 4;
 	TShaderMapRef<FPostProcessSunAvgVS_ES2> VertexShader(GetGlobalShaderMap());
 
 	DrawRectangle(
+		Context.RHICmdList,
 		0, 0,
 		DstX, DstY,
 		0, 0,
@@ -2510,7 +2509,7 @@ void FRCPassPostProcessSunAvgES2::Process(FRenderingCompositePassContext& Contex
 		*VertexShader,
 		EDRF_UseTriangleOptimization);
 
-	RHICopyToResolveTarget(DestRenderTarget.TargetableTexture, DestRenderTarget.ShaderResourceTexture, false, FResolveParams());
+	Context.RHICmdList.CopyToResolveTarget(DestRenderTarget.TargetableTexture, DestRenderTarget.ShaderResourceTexture, false, FResolveParams());
 }
 
 FPooledRenderTargetDesc FRCPassPostProcessSunAvgES2::ComputeOutputDesc(EPassOutputId InPassOutputId) const
@@ -2565,12 +2564,12 @@ public:
 		AaBlendAmount.Bind(Initializer.ParameterMap, TEXT("AaBlendAmount"));
 	}
 
-	void SetPS(FRHICommandList& RHICmdList, const FRenderingCompositePassContext& Context)
+	void SetPS(const FRenderingCompositePassContext& Context)
 	{
 		const FPixelShaderRHIParamRef ShaderRHI = GetPixelShader();
 		const FPostProcessSettings& Settings = Context.View.FinalPostProcessSettings;
-		FGlobalShader::SetParameters(RHICmdList, ShaderRHI, Context.View);
-		PostprocessParameter.SetPS(RHICmdList, ShaderRHI, Context, TStaticSamplerState<SF_Bilinear,AM_Clamp,AM_Clamp,AM_Clamp>::GetRHI());
+		FGlobalShader::SetParameters(Context.RHICmdList, ShaderRHI, Context.View);
+		PostprocessParameter.SetPS(ShaderRHI, Context, TStaticSamplerState<SF_Bilinear,AM_Clamp,AM_Clamp,AM_Clamp>::GetRHI());
 
 		// Compute the blend factor which decides the trade off between ghosting in motion and flicker when not moving.
 		// This works by computing the screen space motion vector of distant point at the center of the screen.
@@ -2633,12 +2632,12 @@ public:
 				BlendAmount = 0.0f;
 			}
 
-			SetShaderValue(RHICmdList, ShaderRHI, AaBlendAmount, BlendAmount);
+			SetShaderValue(Context.RHICmdList, ShaderRHI, AaBlendAmount, BlendAmount);
 		}
 		else
 		{
 			float BlendAmount = 0.0;
-			SetShaderValue(RHICmdList, ShaderRHI, AaBlendAmount, BlendAmount);
+			SetShaderValue(Context.RHICmdList, ShaderRHI, AaBlendAmount, BlendAmount);
 		}
 	}
 
@@ -2674,11 +2673,11 @@ public:
 		PostprocessParameter.Bind(Initializer.ParameterMap);
 	}
 
-	void SetVS(FRHICommandList& RHICmdList, const FRenderingCompositePassContext& Context)
+	void SetVS(const FRenderingCompositePassContext& Context)
 	{
 		const FVertexShaderRHIParamRef ShaderRHI = GetVertexShader();
-		FGlobalShader::SetParameters(RHICmdList, ShaderRHI, Context.View);
-		PostprocessParameter.SetVS(RHICmdList, ShaderRHI, Context, TStaticSamplerState<SF_Bilinear,AM_Clamp,AM_Clamp,AM_Clamp>::GetRHI());
+		FGlobalShader::SetParameters(Context.RHICmdList, ShaderRHI, Context.View);
+		PostprocessParameter.SetVS(ShaderRHI, Context, TStaticSamplerState<SF_Bilinear,AM_Clamp,AM_Clamp,AM_Clamp>::GetRHI());
 	}
 
 	virtual bool Serialize(FArchive& Ar)
@@ -2693,22 +2692,23 @@ IMPLEMENT_SHADER_TYPE(,FPostProcessAaVS_ES2,TEXT("PostProcessMobile"),TEXT("AaVS
 
 
 
-static void Aa_SetShader(FRHICommandList& RHICmdList, const FRenderingCompositePassContext& Context)
+static void Aa_SetShader(const FRenderingCompositePassContext& Context)
 {
 	TShaderMapRef<FPostProcessAaVS_ES2> VertexShader(GetGlobalShaderMap());
 	TShaderMapRef<FPostProcessAaPS_ES2> PixelShader(GetGlobalShaderMap());
 
 	static FGlobalBoundShaderState BoundShaderState;
+	Context.RHICmdList.CheckIsNull(); // need new approach for "static FGlobalBoundShaderState" for parallel rendering
 
-	SetGlobalBoundShaderState(BoundShaderState, GFilterVertexDeclaration.VertexDeclarationRHI, *VertexShader, *PixelShader);
+	SetGlobalBoundShaderState(Context.RHICmdList, BoundShaderState, GFilterVertexDeclaration.VertexDeclarationRHI, *VertexShader, *PixelShader);
 
-	VertexShader->SetVS(RHICmdList, Context);
-	PixelShader->SetPS(RHICmdList, Context);
+	VertexShader->SetVS(Context);
+	PixelShader->SetPS(Context);
 }
 
-void FRCPassPostProcessAaES2::SetShader(FRHICommandList& RHICmdList, const FRenderingCompositePassContext& Context)
+void FRCPassPostProcessAaES2::SetShader(const FRenderingCompositePassContext& Context)
 {
-	Aa_SetShader(RHICmdList, Context);
+	Aa_SetShader(Context);
 }
 
 void FRCPassPostProcessAaES2::Process(FRenderingCompositePassContext& Context)
@@ -2725,25 +2725,24 @@ void FRCPassPostProcessAaES2::Process(FRenderingCompositePassContext& Context)
 	DstRect.Max.Y = DstY;
 
 	const FSceneRenderTargetItem& DestRenderTarget = PassOutputs[0].RequestSurface(Context);
-	RHISetRenderTarget(DestRenderTarget.TargetableTexture, FTextureRHIRef());	
+	SetRenderTarget(Context.RHICmdList, DestRenderTarget.TargetableTexture, FTextureRHIRef());	
 
 	// is optimized away if possible (RT size=view size, )
-	RHIClear(true, FLinearColor::Black, false, 1.0f, false, 0, FIntRect());
+	Context.RHICmdList.Clear(true, FLinearColor::Black, false, 1.0f, false, 0, FIntRect());
 
 	Context.SetViewportAndCallRHI(0, 0, 0.0f, DstX, DstY, 1.0f );
 
-	RHISetBlendState(TStaticBlendState<>::GetRHI());
-	RHISetRasterizerState(TStaticRasterizerState<>::GetRHI());
-	RHISetDepthStencilState(TStaticDepthStencilState<false,CF_Always>::GetRHI());
+	Context.RHICmdList.SetBlendState(TStaticBlendState<>::GetRHI());
+	Context.RHICmdList.SetRasterizerState(TStaticRasterizerState<>::GetRHI());
+	Context.RHICmdList.SetDepthStencilState(TStaticDepthStencilState<false,CF_Always>::GetRHI());
 
-	//@todo-rco: RHIPacketList
-	FRHICommandList& RHICmdList = FRHICommandList::GetNullRef();
-	SetShader(RHICmdList, Context);
+	SetShader(Context);
 
 	FIntPoint SrcDstSize = PrePostSourceViewportSize;
 	TShaderMapRef<FPostProcessAaVS_ES2> VertexShader(GetGlobalShaderMap());
 
 	DrawRectangle(
+		Context.RHICmdList,
 		0, 0,
 		DstX, DstY,
 		0, 0,
@@ -2753,7 +2752,7 @@ void FRCPassPostProcessAaES2::Process(FRenderingCompositePassContext& Context)
 		*VertexShader,
 		EDRF_UseTriangleOptimization);
 
-	RHICopyToResolveTarget(DestRenderTarget.TargetableTexture, DestRenderTarget.ShaderResourceTexture, false, FResolveParams());
+	Context.RHICmdList.CopyToResolveTarget(DestRenderTarget.TargetableTexture, DestRenderTarget.ShaderResourceTexture, false, FResolveParams());
 }
 
 FPooledRenderTargetDesc FRCPassPostProcessAaES2::ComputeOutputDesc(EPassOutputId InPassOutputId) const
diff --git a/Engine/Source/Runtime/Renderer/Private/PostProcess/PostProcessMobile.h b/Engine/Source/Runtime/Renderer/Private/PostProcess/PostProcessMobile.h
--- a/Engine/Source/Runtime/Renderer/Private/PostProcess/PostProcessMobile.h
+++ b/Engine/Source/Runtime/Renderer/Private/PostProcess/PostProcessMobile.h
@@ -18,7 +18,7 @@ public:
 private:
 	FIntPoint PrePostSourceViewportSize;
 	bool bUsedFramebufferFetch;
-	void SetShader(FRHICommandList& RHICmdList, const FRenderingCompositePassContext& Context);
+	void SetShader(const FRenderingCompositePassContext& Context);
 };
 
 class FRCPassPostProcessBloomSetupSmallES2 : public TRenderingCompositePassBase<1, 1>
@@ -31,7 +31,7 @@ public:
 private:
 	FIntPoint PrePostSourceViewportSize;
 	bool bUsedFramebufferFetch;
-	void SetShader(FRHICommandList& RHICmdList, const FRenderingCompositePassContext& Context);
+	void SetShader(const FRenderingCompositePassContext& Context);
 };
 
 class FRCPassPostProcessDofNearES2 : public TRenderingCompositePassBase<1, 1>
@@ -43,7 +43,7 @@ public:
 	virtual void Release() { delete this; }
 private:
 	FIntPoint PrePostSourceViewportSize;
-	void SetShader(FRHICommandList& RHICmdList, const FRenderingCompositePassContext& Context);
+	void SetShader(const FRenderingCompositePassContext& Context);
 };
 
 class FRCPassPostProcessDofDownES2 : public TRenderingCompositePassBase<2, 1>
@@ -56,7 +56,7 @@ public:
 private:
 	FIntPoint PrePostSourceViewportSize;
 	bool bUsedFramebufferFetch;
-	void SetShader(FRHICommandList& RHICmdList, const FRenderingCompositePassContext& Context);
+	void SetShader(const FRenderingCompositePassContext& Context);
 };
 
 class FRCPassPostProcessDofBlurES2 : public TRenderingCompositePassBase<2, 1>
@@ -106,7 +106,7 @@ public:
 private:
 	FIntPoint PrePostSourceViewportSize;
 	bool bOnChip;
-	void SetShader(FRHICommandList& RHICmdList, const FRenderingCompositePassContext& Context);
+	void SetShader(const FRenderingCompositePassContext& Context);
 };
 
 class FRCPassPostProcessSunAlphaES2 : public TRenderingCompositePassBase<1, 1>
@@ -118,7 +118,7 @@ public:
 	virtual void Release() { delete this; }
 private:
 	FIntPoint PrePostSourceViewportSize;
-	void SetShader(FRHICommandList& RHICmdList, const FRenderingCompositePassContext& Context);
+	void SetShader(const FRenderingCompositePassContext& Context);
 };
 
 class FRCPassPostProcessSunBlurES2 : public TRenderingCompositePassBase<1, 1>
@@ -141,7 +141,7 @@ public:
 	virtual void Release() { delete this; }
 private:
 	FIntPoint PrePostSourceViewportSize;
-	void SetShader(FRHICommandList& RHICmdList, const FRenderingCompositePassContext& Context);
+	void SetShader(const FRenderingCompositePassContext& Context);
 };
 
 class FRCPassPostProcessSunMergeSmallES2 : public TRenderingCompositePassBase<2, 1>
@@ -153,7 +153,7 @@ public:
 	virtual void Release() { delete this; }
 private:
 	FIntPoint PrePostSourceViewportSize;
-	void SetShader(FRHICommandList& RHICmdList, const FRenderingCompositePassContext& Context);
+	void SetShader(const FRenderingCompositePassContext& Context);
 };
 
 class FRCPassPostProcessSunAvgES2 : public TRenderingCompositePassBase<2, 1>
@@ -165,7 +165,7 @@ public:
 	virtual void Release() { delete this; }
 private:
 	FIntPoint PrePostSourceViewportSize;
-	void SetShader(FRHICommandList& RHICmdList, const FRenderingCompositePassContext& Context);
+	void SetShader(const FRenderingCompositePassContext& Context);
 };
 
 class FRCPassPostProcessAaES2 : public TRenderingCompositePassBase<2, 1>
@@ -177,6 +177,6 @@ public:
 	virtual void Release() { delete this; }
 private:
 	FIntPoint PrePostSourceViewportSize;
-	void SetShader(FRHICommandList& RHICmdList, const FRenderingCompositePassContext& Context);
+	void SetShader(const FRenderingCompositePassContext& Context);
 };
 
diff --git a/Engine/Source/Runtime/Renderer/Private/PostProcess/PostProcessMorpheus.cpp b/Engine/Source/Runtime/Renderer/Private/PostProcess/PostProcessMorpheus.cpp
--- a/Engine/Source/Runtime/Renderer/Private/PostProcess/PostProcessMorpheus.cpp
+++ b/Engine/Source/Runtime/Renderer/Private/PostProcess/PostProcessMorpheus.cpp
@@ -111,14 +111,14 @@ public:
 	}
 
 
-	void SetPS(FRHICommandList& RHICmdList, const FRenderingCompositePassContext& Context, FIntRect SrcRect, FIntPoint SrcBufferSize, EStereoscopicPass StereoPass, FMatrix& QuadTexTransform)
+	void SetPS(const FRenderingCompositePassContext& Context, FIntRect SrcRect, FIntPoint SrcBufferSize, EStereoscopicPass StereoPass, FMatrix& QuadTexTransform)
 	{
 		const FPixelShaderRHIParamRef ShaderRHI = GetPixelShader();
 		
-		FGlobalShader::SetParameters(RHICmdList, ShaderRHI, Context.View);
+		FGlobalShader::SetParameters(Context.RHICmdList, ShaderRHI, Context.View);
 
-		PostprocessParameter.SetPS(RHICmdList, ShaderRHI, Context, TStaticSamplerState<SF_Bilinear,AM_Clamp,AM_Clamp,AM_Clamp>::GetRHI());
-		DeferredParameters.Set(RHICmdList, ShaderRHI, Context.View);
+		PostprocessParameter.SetPS(ShaderRHI, Context, TStaticSamplerState<SF_Bilinear,AM_Clamp,AM_Clamp,AM_Clamp>::GetRHI());
+		DeferredParameters.Set(Context.RHICmdList, ShaderRHI, Context.View);
 
 		{
 			check(GEngine->HMDDevice.IsValid());
@@ -130,18 +130,18 @@ public:
 			if (StereoPass == eSSP_LEFT_EYE)
 			{
 				FTexture* TextureLeft = HMDDevice->GetDistortionTextureLeft();
-				SetTextureParameter(ShaderRHI, DistortionTextureParam, DistortionTextureSampler, TextureLeft->SamplerStateRHI, TextureLeft->TextureRHI);
-				SetShaderValue(ShaderRHI, TextureScale, HMDDevice->GetTextureScaleLeft());
-				SetShaderValue(ShaderRHI, TextureOffset, HMDDevice->GetTextureOffsetLeft());
-				SetShaderValue(ShaderRHI, TextureUVOffset, 0.0f);
+				SetTextureParameter(Context.RHICmdList, ShaderRHI, DistortionTextureParam, DistortionTextureSampler, TextureLeft->SamplerStateRHI, TextureLeft->TextureRHI);
+				SetShaderValue(Context.RHICmdList, ShaderRHI, TextureScale, HMDDevice->GetTextureScaleLeft());
+				SetShaderValue(Context.RHICmdList, ShaderRHI, TextureOffset, HMDDevice->GetTextureOffsetLeft());
+				SetShaderValue(Context.RHICmdList, ShaderRHI, TextureUVOffset, 0.0f);
 			}
 			else
 			{
 				FTexture* TextureRight = HMDDevice->GetDistortionTextureRight();
-				SetTextureParameter(ShaderRHI, DistortionTextureParam, DistortionTextureSampler, TextureRight->SamplerStateRHI, TextureRight->TextureRHI);
-				SetShaderValue(ShaderRHI, TextureScale, HMDDevice->GetTextureScaleRight());
-				SetShaderValue(ShaderRHI, TextureOffset, HMDDevice->GetTextureOffsetRight());
-				SetShaderValue(ShaderRHI, TextureUVOffset, -0.5f);
+				SetTextureParameter(Context.RHICmdList, ShaderRHI, DistortionTextureParam, DistortionTextureSampler, TextureRight->SamplerStateRHI, TextureRight->TextureRHI);
+				SetShaderValue(Context.RHICmdList, ShaderRHI, TextureScale, HMDDevice->GetTextureScaleRight());
+				SetShaderValue(Context.RHICmdList, ShaderRHI, TextureOffset, HMDDevice->GetTextureOffsetRight());
+				SetShaderValue(Context.RHICmdList, ShaderRHI, TextureUVOffset, -0.5f);
 			}				
 			      
             QuadTexTransform = FMatrix::Identity;            
@@ -169,36 +169,36 @@ public:
 				DistortionTexture = GEngine->HMDDevice->GetDistortionTexture();
 				check(DistortionTexture != NULL);
 			}
-			SetTextureParameter(RHICmdList, ShaderRHI, DistortionTextureParam, DistortionTextureSampler, TStaticSamplerState<SF_Bilinear,AM_Clamp,AM_Clamp,AM_Clamp>::GetRHI(), DistortionTexture->TextureRHI);
+			SetTextureParameter(Context.RHICmdList, ShaderRHI, DistortionTextureParam, DistortionTextureSampler, TStaticSamplerState<SF_Bilinear, AM_Clamp, AM_Clamp, AM_Clamp>::GetRHI(), DistortionTexture->TextureRHI);
 
 			// Shifts texture coordinates to the center of the distortion function around the center of the lens
 			FVector2D ViewLensCenter;
 			ViewLensCenter.X = x + (0.5f + XCenterOffset) * w;
 			ViewLensCenter.Y = y + (0.5f + YCenterOffset) * h;
-			SetShaderValue(RHICmdList, ShaderRHI, LensCenter, ViewLensCenter);
+			SetShaderValue(Context.RHICmdList, ShaderRHI, LensCenter, ViewLensCenter);
 
 			// Texture coordinate for the center of the half scene texture, used to clamp sampling
 			FVector2D ViewScreenCenter;
 			ViewScreenCenter.X = x + w * 0.5f;
 			ViewScreenCenter.Y = y + h * 0.5f;
-			SetShaderValue(RHICmdList, ShaderRHI, ScreenCenter, ViewScreenCenter);
+			SetShaderValue(Context.RHICmdList, ShaderRHI, ScreenCenter, ViewScreenCenter);
 			
 			// Rescale output (sample) coordinates back to texture range and increase scale to support rendering outside the the screen
 			FVector2D ViewScale;
 			ViewScale.X = (w/2) * ScaleFactor;
 			ViewScale.Y = (h/2) * ScaleFactor * AspectRatio;
-			SetShaderValue(RHICmdList, ShaderRHI, Scale, ViewScale);
+			SetShaderValue(Context.RHICmdList, ShaderRHI, Scale, ViewScale);
 
             // Distortion coefficients
             FVector4 DistortionValues;
             GEngine->HMDDevice->GetDistortionWarpValues(DistortionValues);
-			SetShaderValue(RHICmdList, ShaderRHI, HMDWarpParam, DistortionValues);
+			SetShaderValue(Context.RHICmdList, ShaderRHI, HMDWarpParam, DistortionValues);
 
 			// CNN - Morpheus changes
 			// CA correction values
 			FVector4 CAValues;
 			GEngine->HMDDevice->GetChromaAbCorrectionValues(CAValues);
-			SetShaderValue(RHICmdList, ShaderRHI, CAWarpParam, CAValues);
+			SetShaderValue(Context.RHICmdList, ShaderRHI, CAWarpParam, CAValues);
 
 			// Rescale the texture coordinates to [-1,1] unit range and corrects aspect ratio
 			FVector2D ViewScaleIn;
@@ -254,9 +254,9 @@ class FPostProcessMorpheusVS : public FGlobalShader
 	FPostProcessMorpheusVS() {}
 
 	/** to have a similar interface as all other shaders */
-	void SetParameters(FRHICommandList& RHICmdList, const FRenderingCompositePassContext& Context)
+	void SetParameters(const FRenderingCompositePassContext& Context)
 	{
-		FGlobalShader::SetParameters(RHICmdList, GetVertexShader(), Context.View);
+		FGlobalShader::SetParameters(Context.RHICmdList, GetVertexShader(), Context.View);
 	}
 
 	void SetParameters(FRHICommandList& RHICmdList, const FSceneView& View)
@@ -297,30 +297,31 @@ void FRCPassPostProcessMorpheus::Process(FRenderingCompositePassContext& Context
     const FSceneRenderTargetItem& DestRenderTarget = PassOutputs[0].RequestSurface(Context);
 
 	// Set the view family's render target/viewport.
-	RHISetRenderTarget(DestRenderTarget.TargetableTexture, FTextureRHIRef());	
+	SetRenderTarget(Context.RHICmdList, DestRenderTarget.TargetableTexture, FTextureRHIRef());
 
 	Context.SetViewportAndCallRHI(DestRect);
 
 	// set the state
-	RHISetBlendState(TStaticBlendState<>::GetRHI());
-	RHISetRasterizerState(TStaticRasterizerState<>::GetRHI());
-	RHISetDepthStencilState(TStaticDepthStencilState<false,CF_Always>::GetRHI());
+	Context.RHICmdList.SetBlendState(TStaticBlendState<>::GetRHI());
+	Context.RHICmdList.SetRasterizerState(TStaticRasterizerState<>::GetRHI());
+	Context.RHICmdList.SetDepthStencilState(TStaticDepthStencilState<false, CF_Always>::GetRHI());
 
 	TShaderMapRef<FPostProcessMorpheusVS> VertexShader(GetGlobalShaderMap());
 	TShaderMapRef<FPostProcessMorpheusPS> PixelShader(GetGlobalShaderMap());
 
 	static FGlobalBoundShaderState BoundShaderState;
+	Context.RHICmdList.CheckIsNull(); // need new approach for "static FGlobalBoundShaderState" for parallel rendering
 
-	SetGlobalBoundShaderState(BoundShaderState, GFilterVertexDeclaration.VertexDeclarationRHI, *VertexShader, *PixelShader);
+	SetGlobalBoundShaderState(Context.RHICmdList, BoundShaderState, GFilterVertexDeclaration.VertexDeclarationRHI, *VertexShader, *PixelShader);
 
     FMatrix QuadTexTransform;
     FMatrix QuadPosTransform = FMatrix::Identity;
-	//@todo-rco: RHIPacketList
-	FRHICommandList& RHICmdList = FRHICommandList::GetNullRef();
-	PixelShader->SetPS(RHICmdList, Context, SrcRect, SrcSize, View.StereoPass, QuadTexTransform);
+
+	PixelShader->SetPS(Context, SrcRect, SrcSize, View.StereoPass, QuadTexTransform);
 
 	// Draw a quad mapping scene color to the view's render target
     DrawTransformedRectangle(
+		Context.RHICmdList,
         0, 0,
         DestRect.Width(), DestRect.Height(),
         QuadPosTransform,
@@ -331,7 +332,7 @@ void FRCPassPostProcessMorpheus::Process(FRenderingCompositePassContext& Context
         SrcSize
         );
 
-	RHICopyToResolveTarget(DestRenderTarget.TargetableTexture, DestRenderTarget.ShaderResourceTexture, false, FResolveParams());
+	Context.RHICmdList.CopyToResolveTarget(DestRenderTarget.TargetableTexture, DestRenderTarget.ShaderResourceTexture, false, FResolveParams());
 }
 
 FPooledRenderTargetDesc FRCPassPostProcessMorpheus::ComputeOutputDesc(EPassOutputId InPassOutputId) const
