UE4 - more work preparing for parallel rendering. Merged the cmd list into the post process context, among other things.

diff --git a/Engine/Source/Runtime/Renderer/Private/PostProcess/PostProcessMotionBlur.cpp b/Engine/Source/Runtime/Renderer/Private/PostProcess/PostProcessMotionBlur.cpp
--- a/Engine/Source/Runtime/Renderer/Private/PostProcess/PostProcessMotionBlur.cpp
+++ b/Engine/Source/Runtime/Renderer/Private/PostProcess/PostProcessMotionBlur.cpp
@@ -47,13 +47,13 @@ public:
 		return bShaderHasOutdatedParameters;
 	}
 
-	void SetParameters(FRHICommandList& RHICmdList, const FRenderingCompositePassContext& Context)
+	void SetParameters(const FRenderingCompositePassContext& Context)
 	{
 		const FPixelShaderRHIParamRef ShaderRHI = GetPixelShader();
 
-		FGlobalShader::SetParameters(RHICmdList, ShaderRHI, Context.View);
+		FGlobalShader::SetParameters(Context.RHICmdList, ShaderRHI, Context.View);
 
-		PostprocessParameter.SetPS(RHICmdList, ShaderRHI, Context, TStaticSamplerState<SF_Point,AM_Clamp,AM_Clamp,AM_Clamp>::GetRHI());
+		PostprocessParameter.SetPS(ShaderRHI, Context, TStaticSamplerState<SF_Point, AM_Clamp, AM_Clamp, AM_Clamp>::GetRHI());
 
 		{
 			const float SizeX = Context.View.ViewRect.Width();
@@ -71,7 +71,7 @@ public:
 			float ObjectScaleX = ViewMotionBlurScale * InvMaxVelocity;
 			float ObjectScaleY = ViewMotionBlurScale * InvMaxVelocity * InvAspectRatio;
 
-			SetShaderValue(RHICmdList, ShaderRHI, VelocityScale, FVector4( ObjectScaleX, -ObjectScaleY, 0, 0 ) );
+			SetShaderValue(Context.RHICmdList, ShaderRHI, VelocityScale, FVector4(ObjectScaleX, -ObjectScaleY, 0, 0));
 		}
 	}
 };
@@ -114,36 +114,35 @@ void FRCPassPostProcessMotionBlurSetup::Process(FRenderingCompositePassContext&
 		DestRenderTarget0.TargetableTexture,
 		DestRenderTarget1.TargetableTexture
 	};
-	RHISetRenderTargets( ARRAY_COUNT(RenderTargets), RenderTargets, FTextureRHIParamRef(), 0, NULL );
+	SetRenderTargets(Context.RHICmdList, ARRAY_COUNT(RenderTargets), RenderTargets, FTextureRHIParamRef(), 0, NULL);
 
 	// is optimized away if possible (RT size=view size, )
 	FLinearColor ClearColors[2] = {FLinearColor(0,0,0,0), FLinearColor(0,0,0,0)};
-	RHIClearMRT(true, 2, ClearColors, false, 1.0f, false, 0, DestRect);
+	Context.RHICmdList.ClearMRT(true, 2, ClearColors, false, 1.0f, false, 0, DestRect);
 
 	Context.SetViewportAndCallRHI(0, 0, 0.0f, DestSize.X, DestSize.Y, 1.0f );
 
 	// set the state
-	RHISetBlendState(TStaticBlendState<>::GetRHI());
-	RHISetRasterizerState(TStaticRasterizerState<>::GetRHI());
-	RHISetDepthStencilState(TStaticDepthStencilState<false,CF_Always>::GetRHI());
-
-	//@todo-rco: RHIPacketList
-	FRHICommandList& RHICmdList = FRHICommandList::GetNullRef();
+	Context.RHICmdList.SetBlendState(TStaticBlendState<>::GetRHI());
+	Context.RHICmdList.SetRasterizerState(TStaticRasterizerState<>::GetRHI());
+	Context.RHICmdList.SetDepthStencilState(TStaticDepthStencilState<false, CF_Always>::GetRHI());
 
 	TShaderMapRef<FPostProcessVS> VertexShader(GetGlobalShaderMap());
 
 	{		
 		TShaderMapRef<FPostProcessMotionBlurSetupPS > PixelShader(GetGlobalShaderMap());
 		static FGlobalBoundShaderState BoundShaderState;
+		Context.RHICmdList.CheckIsNull(); // need new approach for "static FGlobalBoundShaderState" for parallel rendering
 
-		SetGlobalBoundShaderState(BoundShaderState, GFilterVertexDeclaration.VertexDeclarationRHI, *VertexShader, *PixelShader);
+		SetGlobalBoundShaderState(Context.RHICmdList, BoundShaderState, GFilterVertexDeclaration.VertexDeclarationRHI, *VertexShader, *PixelShader);
 
-		PixelShader->SetParameters(RHICmdList, Context);
-		VertexShader->SetParameters(RHICmdList, Context);
+		PixelShader->SetParameters(Context);
+		VertexShader->SetParameters(Context);
 	}
 
 	// Draw a quad mapping scene color to the view's render target
 	DrawRectangle(
+		Context.RHICmdList,
 		DestRect.Min.X, DestRect.Min.Y,
 		DestRect.Width(), DestRect.Height(),
 		SrcRect.Min.X, SrcRect.Min.Y,
@@ -153,8 +152,8 @@ void FRCPassPostProcessMotionBlurSetup::Process(FRenderingCompositePassContext&
 		*VertexShader,
 		EDRF_UseTriangleOptimization);
 
-	RHICopyToResolveTarget(DestRenderTarget0.TargetableTexture, DestRenderTarget0.ShaderResourceTexture, false, FResolveParams());
-	RHICopyToResolveTarget(DestRenderTarget1.TargetableTexture, DestRenderTarget1.ShaderResourceTexture, false, FResolveParams());
+	Context.RHICmdList.CopyToResolveTarget(DestRenderTarget0.TargetableTexture, DestRenderTarget0.ShaderResourceTexture, false, FResolveParams());
+	Context.RHICmdList.CopyToResolveTarget(DestRenderTarget1.TargetableTexture, DestRenderTarget1.ShaderResourceTexture, false, FResolveParams());
 }
 
 FPooledRenderTargetDesc FRCPassPostProcessMotionBlurSetup::ComputeOutputDesc(EPassOutputId InPassOutputId) const
@@ -254,13 +253,13 @@ public:
 		return bShaderHasOutdatedParameters;
 	}
 
-	void SetParameters(FRHICommandList& RHICmdList, const FRenderingCompositePassContext& Context)
+	void SetParameters(const FRenderingCompositePassContext& Context)
 	{
 		const FPixelShaderRHIParamRef ShaderRHI = GetPixelShader();
 
-		FGlobalShader::SetParameters(RHICmdList, ShaderRHI, Context.View);
+		FGlobalShader::SetParameters(Context.RHICmdList, ShaderRHI, Context.View);
 
-		DeferredParameters.Set(RHICmdList, ShaderRHI, Context.View);
+		DeferredParameters.Set(Context.RHICmdList, ShaderRHI, Context.View);
 
 		{
 			bool bFiltered = false;
@@ -272,11 +271,11 @@ public:
 
 			if(bFiltered)
 			{
-				PostprocessParameter.SetPS(RHICmdList, ShaderRHI, Context, TStaticSamplerState<SF_Bilinear,AM_Border,AM_Border,AM_Clamp>::GetRHI());
+				PostprocessParameter.SetPS(ShaderRHI, Context, TStaticSamplerState<SF_Bilinear,AM_Border,AM_Border,AM_Clamp>::GetRHI());
 			}
 			else
 			{
-				PostprocessParameter.SetPS(RHICmdList, ShaderRHI, Context, TStaticSamplerState<SF_Point,AM_Border,AM_Border,AM_Clamp>::GetRHI());
+				PostprocessParameter.SetPS(ShaderRHI, Context, TStaticSamplerState<SF_Point,AM_Border,AM_Border,AM_Clamp>::GetRHI());
 			}
 		}
 		
@@ -286,11 +285,11 @@ public:
 			// then translate by the difference between last frame's camera position and this frame's camera position,
 			// then apply the rest of the transforms.  This effectively avoids precision issues near the extents of large levels whose world space position is very large.
 			FVector ViewOriginDelta = Context.View.ViewMatrices.ViewOrigin - Context.View.PrevViewMatrices.ViewOrigin;
-			SetShaderValue(RHICmdList, ShaderRHI, PrevViewProjMatrix, FTranslationMatrix(ViewOriginDelta) * Context.View.PrevViewRotationProjMatrix);
+			SetShaderValue(Context.RHICmdList, ShaderRHI, PrevViewProjMatrix, FTranslationMatrix(ViewOriginDelta) * Context.View.PrevViewRotationProjMatrix);
 		}
 		else
 		{
-			SetShaderValue(RHICmdList, ShaderRHI, PrevViewProjMatrix, Context.View.ViewMatrices.GetViewRotationProjMatrix() );
+			SetShaderValue(Context.RHICmdList, ShaderRHI, PrevViewProjMatrix, Context.View.ViewMatrices.GetViewRotationProjMatrix());
 		}
 
 		TRefCountPtr<IPooledRenderTarget> InputPooledElement = Context.Pass->GetInput(ePId_Input0)->GetOutput()->RequestInput();
@@ -312,7 +311,7 @@ public:
 			FVector2D Mul(1.0f / SizeUV.X, 1.0f / SizeUV.Y);
 			FVector2D Add = - MinUV * Mul;
 			FVector4 TextureViewMadValue(Mul.X, Mul.Y, Add.X, Add.Y);
-			SetShaderValue(RHICmdList, ShaderRHI, TextureViewMad, TextureViewMadValue);
+			SetShaderValue(Context.RHICmdList, ShaderRHI, TextureViewMad, TextureViewMadValue);
 		}
 
 		{
@@ -341,7 +340,7 @@ public:
 				- ObjectMotionBlurScale * InvMaxVelocity * InvAspectRatio,
 				MaxVelocity * 2,
 				- MaxVelocity * 2 * AspectRatio);
-			SetShaderValue(RHICmdList, ShaderRHI, MotionBlurParameters, MotionBlurParametersValue);
+			SetShaderValue(Context.RHICmdList, ShaderRHI, MotionBlurParameters, MotionBlurParametersValue);
 		}
 	}
 
@@ -367,17 +366,18 @@ VARIATION1(0)			VARIATION1(1)			VARIATION1(2)			VARIATION1(3)			VARIATION1(4)
 
 // @param Quality 0: visualize, 1:low, 2:medium, 3:high, 4:very high
 template <uint32 Quality>
-static void SetMotionBlurShaderTempl(FRHICommandList& RHICmdList, const FRenderingCompositePassContext& Context)
+static void SetMotionBlurShaderTempl(const FRenderingCompositePassContext& Context)
 {
 	TShaderMapRef<FPostProcessVS> VertexShader(GetGlobalShaderMap());
 	TShaderMapRef<FPostProcessMotionBlurPS<Quality> > PixelShader(GetGlobalShaderMap());
 
 	static FGlobalBoundShaderState BoundShaderState;
+	Context.RHICmdList.CheckIsNull(); // need new approach for "static FGlobalBoundShaderState" for parallel rendering
 
-	SetGlobalBoundShaderState(BoundShaderState, GFilterVertexDeclaration.VertexDeclarationRHI, *VertexShader, *PixelShader);
+	SetGlobalBoundShaderState(Context.RHICmdList, BoundShaderState, GFilterVertexDeclaration.VertexDeclarationRHI, *VertexShader, *PixelShader);
 
-	VertexShader->SetParameters(RHICmdList, Context);
-	PixelShader->SetParameters(RHICmdList, Context);
+	VertexShader->SetParameters(Context);
+	PixelShader->SetParameters(Context);
 }
 
 FRCPassPostProcessMotionBlur::FRCPassPostProcessMotionBlur(uint32 InQuality) 
@@ -417,42 +417,43 @@ void FRCPassPostProcessMotionBlur::Process(FRenderingCompositePassContext& Conte
 	const FSceneRenderTargetItem& DestRenderTarget = PassOutputs[0].RequestSurface(Context);
 
 	// Set the view family's render target/viewport.
-	RHISetRenderTarget(DestRenderTarget.TargetableTexture, FTextureRHIRef());	
+	SetRenderTarget(Context.RHICmdList, DestRenderTarget.TargetableTexture, FTextureRHIRef());
 
 	// is optimized away if possible (RT size=view size, )
-	RHIClear(true, FLinearColor::Black, false, 1.0f, false, 0, SrcRect);
+	Context.RHICmdList.Clear(true, FLinearColor::Black, false, 1.0f, false, 0, SrcRect);
 
 	Context.SetViewportAndCallRHI(SrcRect);
 
 	// set the state
-	RHISetBlendState(TStaticBlendState<>::GetRHI());
-	RHISetRasterizerState(TStaticRasterizerState<>::GetRHI());
-	RHISetDepthStencilState(TStaticDepthStencilState<false,CF_Always>::GetRHI());
+	Context.RHICmdList.SetBlendState(TStaticBlendState<>::GetRHI());
+	Context.RHICmdList.SetRasterizerState(TStaticRasterizerState<>::GetRHI());
+	Context.RHICmdList.SetDepthStencilState(TStaticDepthStencilState<false, CF_Always>::GetRHI());
 
 	//@todo-rco: RHIPacketList
 	FRHICommandList& RHICmdList = FRHICommandList::GetNullRef();
 	if(Quality == 1)
 	{
-		SetMotionBlurShaderTempl<1>(RHICmdList, Context);
+		SetMotionBlurShaderTempl<1>(Context);
 	}
 	else if(Quality == 2)
 	{
-		SetMotionBlurShaderTempl<2>(RHICmdList, Context);
+		SetMotionBlurShaderTempl<2>(Context);
 	}
 	else if(Quality == 3)
 	{
-		SetMotionBlurShaderTempl<3>(RHICmdList, Context);
+		SetMotionBlurShaderTempl<3>(Context);
 	}
 	else
 	{
 		check(Quality == 4);
-		SetMotionBlurShaderTempl<4>(RHICmdList, Context);
+		SetMotionBlurShaderTempl<4>(Context);
 	}
 
 	TShaderMapRef<FPostProcessVS> VertexShader(GetGlobalShaderMap());
 
 	// Draw a quad mapping scene color to the view's render target
 	DrawRectangle(
+		Context.RHICmdList,
 		0, 0,
 		SrcRect.Width(), SrcRect.Height(),
 		SrcRect.Min.X, SrcRect.Min.Y, 
@@ -462,7 +463,7 @@ void FRCPassPostProcessMotionBlur::Process(FRenderingCompositePassContext& Conte
 		*VertexShader,
 		EDRF_UseTriangleOptimization);
 
-	RHICopyToResolveTarget(DestRenderTarget.TargetableTexture, DestRenderTarget.ShaderResourceTexture, false, FResolveParams());
+	Context.RHICmdList.CopyToResolveTarget(DestRenderTarget.TargetableTexture, DestRenderTarget.ShaderResourceTexture, false, FResolveParams());
 }
 
 FPooledRenderTargetDesc FRCPassPostProcessMotionBlur::ComputeOutputDesc(EPassOutputId InPassOutputId) const
@@ -511,13 +512,13 @@ public:
 		return bShaderHasOutdatedParameters;
 	}
 
-	void SetParameters(FRHICommandList& RHICmdList, const FRenderingCompositePassContext& Context)
+	void SetParameters(const FRenderingCompositePassContext& Context)
 	{
 		const FPixelShaderRHIParamRef ShaderRHI = GetPixelShader();
 
-		FGlobalShader::SetParameters(RHICmdList, ShaderRHI, Context.View);
+		FGlobalShader::SetParameters(Context.RHICmdList, ShaderRHI, Context.View);
 
-		PostprocessParameter.SetPS(RHICmdList, ShaderRHI, Context, TStaticSamplerState<SF_Bilinear,AM_Border,AM_Border,AM_Clamp>::GetRHI());
+		PostprocessParameter.SetPS(ShaderRHI, Context, TStaticSamplerState<SF_Bilinear,AM_Border,AM_Border,AM_Clamp>::GetRHI());
 	}
 };
 
@@ -553,32 +554,32 @@ void FRCPassPostProcessMotionBlurRecombine::Process(FRenderingCompositePassConte
 	const FSceneRenderTargetItem& DestRenderTarget = PassOutputs[0].RequestSurface(Context);
 
 	// Set the view family's render target/viewport.
-	RHISetRenderTarget(DestRenderTarget.TargetableTexture, FTextureRHIRef());	
+	SetRenderTarget(Context.RHICmdList, DestRenderTarget.TargetableTexture, FTextureRHIRef());
 
 	// is optimized away if possible (RT size=view size, )
-	RHIClear(true, FLinearColor::Black, false, 1.0f, false, 0, SrcRect);
+	Context.RHICmdList.Clear(true, FLinearColor::Black, false, 1.0f, false, 0, SrcRect);
 
 	Context.SetViewportAndCallRHI(SrcRect);
 
 	// set the state
-	RHISetBlendState(TStaticBlendState<>::GetRHI());
-	RHISetRasterizerState(TStaticRasterizerState<>::GetRHI());
-	RHISetDepthStencilState(TStaticDepthStencilState<false,CF_Always>::GetRHI());
+	Context.RHICmdList.SetBlendState(TStaticBlendState<>::GetRHI());
+	Context.RHICmdList.SetRasterizerState(TStaticRasterizerState<>::GetRHI());
+	Context.RHICmdList.SetDepthStencilState(TStaticDepthStencilState<false, CF_Always>::GetRHI());
 
 	TShaderMapRef<FPostProcessVS> VertexShader(GetGlobalShaderMap());
 	TShaderMapRef<FPostProcessMotionBlurRecombinePS> PixelShader(GetGlobalShaderMap());
 
 	static FGlobalBoundShaderState BoundShaderState;
+	Context.RHICmdList.CheckIsNull(); // need new approach for "static FGlobalBoundShaderState" for parallel rendering
 
-	SetGlobalBoundShaderState(BoundShaderState, GFilterVertexDeclaration.VertexDeclarationRHI, *VertexShader, *PixelShader);
+	SetGlobalBoundShaderState(Context.RHICmdList, BoundShaderState, GFilterVertexDeclaration.VertexDeclarationRHI, *VertexShader, *PixelShader);
 
-	//@todo-rco: RHIPacketList
-	FRHICommandList& RHICmdList = FRHICommandList::GetNullRef();
-	VertexShader->SetParameters(RHICmdList, Context);
-	PixelShader->SetParameters(RHICmdList, Context);
+	VertexShader->SetParameters(Context);
+	PixelShader->SetParameters(Context);
 
 	// Draw a quad mapping scene color to the view's render target
 	DrawRectangle(
+		Context.RHICmdList,
 		0, 0,
 		SrcRect.Width(), SrcRect.Height(),
 		SrcRect.Min.X, SrcRect.Min.Y, 
@@ -588,7 +589,7 @@ void FRCPassPostProcessMotionBlurRecombine::Process(FRenderingCompositePassConte
 		*VertexShader,
 		EDRF_UseTriangleOptimization);
 
-	RHICopyToResolveTarget(DestRenderTarget.TargetableTexture, DestRenderTarget.ShaderResourceTexture, false, FResolveParams());
+	Context.RHICmdList.CopyToResolveTarget(DestRenderTarget.TargetableTexture, DestRenderTarget.ShaderResourceTexture, false, FResolveParams());
 }
 
 FPooledRenderTargetDesc FRCPassPostProcessMotionBlurRecombine::ComputeOutputDesc(EPassOutputId InPassOutputId) const
@@ -634,27 +635,26 @@ void FRCPassPostProcessVisualizeMotionBlur::Process(FRenderingCompositePassConte
 	const FSceneRenderTargetItem& DestRenderTarget = PassOutputs[0].RequestSurface(Context);
 
 	// Set the view family's render target/viewport.
-	RHISetRenderTarget(DestRenderTarget.TargetableTexture, FTextureRHIRef());
+	SetRenderTarget(Context.RHICmdList, DestRenderTarget.TargetableTexture, FTextureRHIRef());
 	
 	// is optimized away if possible (RT size=view size, )
-	RHIClear(true, FLinearColor::Black, false, 1.0f, false, 0, SrcRect);
+	Context.RHICmdList.Clear(true, FLinearColor::Black, false, 1.0f, false, 0, SrcRect);
 
 	Context.SetViewportAndCallRHI(SrcRect);
 
 	// set the state
-	RHISetBlendState(TStaticBlendState<>::GetRHI());
-	RHISetRasterizerState(TStaticRasterizerState<>::GetRHI());
-	RHISetDepthStencilState(TStaticDepthStencilState<false,CF_Always>::GetRHI());
+	Context.RHICmdList.SetBlendState(TStaticBlendState<>::GetRHI());
+	Context.RHICmdList.SetRasterizerState(TStaticRasterizerState<>::GetRHI());
+	Context.RHICmdList.SetDepthStencilState(TStaticDepthStencilState<false, CF_Always>::GetRHI());
 
 	// Quality 0: visualize
-	//@todo-rco: RHIPacketList
-	FRHICommandList& RHICmdList = FRHICommandList::GetNullRef();
-	SetMotionBlurShaderTempl<0>(RHICmdList, Context);
+	SetMotionBlurShaderTempl<0>(Context);
 
 	// Draw a quad mapping scene color to the view's render target
 	TShaderMapRef<FPostProcessVS> VertexShader(GetGlobalShaderMap());
 
 	DrawRectangle(
+		Context.RHICmdList,
 		0, 0,
 		SrcRect.Width(), SrcRect.Height(),
 		SrcRect.Min.X, SrcRect.Min.Y, 
@@ -664,7 +664,7 @@ void FRCPassPostProcessVisualizeMotionBlur::Process(FRenderingCompositePassConte
 		*VertexShader,
 		EDRF_UseTriangleOptimization);
 
-	RHICopyToResolveTarget(DestRenderTarget.TargetableTexture, DestRenderTarget.ShaderResourceTexture, false, FResolveParams());
+	Context.RHICmdList.CopyToResolveTarget(DestRenderTarget.TargetableTexture, DestRenderTarget.ShaderResourceTexture, false, FResolveParams());
 }
 
 FPooledRenderTargetDesc FRCPassPostProcessVisualizeMotionBlur::ComputeOutputDesc(EPassOutputId InPassOutputId) const
diff --git a/Engine/Source/Runtime/Renderer/Private/PostProcess/PostProcessNoiseBlur.cpp b/Engine/Source/Runtime/Renderer/Private/PostProcess/PostProcessNoiseBlur.cpp
--- a/Engine/Source/Runtime/Renderer/Private/PostProcess/PostProcessNoiseBlur.cpp
+++ b/Engine/Source/Runtime/Renderer/Private/PostProcess/PostProcessNoiseBlur.cpp
@@ -53,17 +53,17 @@ public:
 		return bShaderHasOutdatedParameters;
 	}
 
-	void SetParameters(FRHICommandList& RHICmdList, const FRenderingCompositePassContext& Context, float InRadius)
+	void SetParameters(const FRenderingCompositePassContext& Context, float InRadius)
 	{
 		const FPixelShaderRHIParamRef ShaderRHI = GetPixelShader();
 
-		FGlobalShader::SetParameters(RHICmdList, ShaderRHI, Context.View);
-		DeferredParameters.Set(RHICmdList, ShaderRHI, Context.View);
-		PostprocessParameter.SetPS(RHICmdList, ShaderRHI, Context, TStaticSamplerState<SF_Bilinear,AM_Border,AM_Border,AM_Border>::GetRHI());
+		FGlobalShader::SetParameters(Context.RHICmdList, ShaderRHI, Context.View);
+		DeferredParameters.Set(Context.RHICmdList, ShaderRHI, Context.View);
+		PostprocessParameter.SetPS(ShaderRHI, Context, TStaticSamplerState<SF_Bilinear,AM_Border,AM_Border,AM_Border>::GetRHI());
 
 		{
 			FVector4 ColorScale(InRadius, 0, 0, 0);
-			SetShaderValue(RHICmdList, ShaderRHI, NoiseParams, ColorScale);
+			SetShaderValue(Context.RHICmdList, ShaderRHI, NoiseParams, ColorScale);
 		}
 	}
 
@@ -95,17 +95,18 @@ FRCPassPostProcessNoiseBlur::FRCPassPostProcessNoiseBlur(float InRadius, EPixelF
 
 
 template <uint32 Method>
-void SetNoiseBlurShader(FRHICommandList& RHICmdList, const FRenderingCompositePassContext& Context, float InRadius)
+void SetNoiseBlurShader(const FRenderingCompositePassContext& Context, float InRadius)
 {
 	TShaderMapRef<FPostProcessVS> VertexShader(GetGlobalShaderMap());
 	TShaderMapRef<FPostProcessNoiseBlurPS<Method> > PixelShader(GetGlobalShaderMap());
 
 	static FGlobalBoundShaderState BoundShaderState;
+	Context.RHICmdList.CheckIsNull(); // need new approach for "static FGlobalBoundShaderState" for parallel rendering
 
-	SetGlobalBoundShaderState(BoundShaderState, GFilterVertexDeclaration.VertexDeclarationRHI, *VertexShader, *PixelShader);
+	SetGlobalBoundShaderState(Context.RHICmdList, BoundShaderState, GFilterVertexDeclaration.VertexDeclarationRHI, *VertexShader, *PixelShader);
 
-	PixelShader->SetParameters(RHICmdList, Context, InRadius);
-	VertexShader->SetParameters(RHICmdList, Context);
+	PixelShader->SetParameters(Context, InRadius);
+	VertexShader->SetParameters(Context);
 }
 
 
@@ -136,36 +137,35 @@ void FRCPassPostProcessNoiseBlur::Process(FRenderingCompositePassContext& Contex
 	const FSceneRenderTargetItem& DestRenderTarget = PassOutputs[0].RequestSurface(Context);
 
 	// Set the view family's render target/viewport.
-	RHISetRenderTarget(DestRenderTarget.TargetableTexture, FTextureRHIRef());	
+	SetRenderTarget(Context.RHICmdList, DestRenderTarget.TargetableTexture, FTextureRHIRef());
 
 	// is optimized away if possible (RT size=view size, )
-	RHIClear(true, FLinearColor(0, 0, 0, 0), false, 1.0f, false, 0, DestRect);
+	Context.RHICmdList.Clear(true, FLinearColor(0, 0, 0, 0), false, 1.0f, false, 0, DestRect);
 
 	Context.SetViewportAndCallRHI(0, 0, 0.0f, DestSize.X, DestSize.Y, 1.0f );
 
 	// set the state
-	RHISetBlendState(TStaticBlendState<>::GetRHI());
-	RHISetRasterizerState(TStaticRasterizerState<>::GetRHI());
-	RHISetDepthStencilState(TStaticDepthStencilState<false,CF_Always>::GetRHI());
+	Context.RHICmdList.SetBlendState(TStaticBlendState<>::GetRHI());
+	Context.RHICmdList.SetRasterizerState(TStaticRasterizerState<>::GetRHI());
+	Context.RHICmdList.SetDepthStencilState(TStaticDepthStencilState<false, CF_Always>::GetRHI());
 	
-	//@todo-rco: RHIPacketList
-	FRHICommandList& RHICmdList = FRHICommandList::GetNullRef();
 	if(Quality == 0)
 	{
-		SetNoiseBlurShader<0>(RHICmdList, Context, Radius);
+		SetNoiseBlurShader<0>(Context, Radius);
 	}
 	else if(Quality == 1)
 	{
-		SetNoiseBlurShader<1>(RHICmdList, Context, Radius);
+		SetNoiseBlurShader<1>(Context, Radius);
 	}
 	else
 	{
-		SetNoiseBlurShader<2>(RHICmdList, Context, Radius);
+		SetNoiseBlurShader<2>(Context, Radius);
 	}
 
 	// Draw a quad mapping scene color to the view's render target
 	TShaderMapRef<FPostProcessVS> VertexShader(GetGlobalShaderMap());
 	DrawRectangle(
+		Context.RHICmdList,
 		DestRect.Min.X, DestRect.Min.Y,
 		DestRect.Width(), DestRect.Height(),
 		SrcRect.Min.X, SrcRect.Min.Y,
@@ -175,7 +175,7 @@ void FRCPassPostProcessNoiseBlur::Process(FRenderingCompositePassContext& Contex
 		*VertexShader,
 		EDRF_UseTriangleOptimization);
 
-	RHICopyToResolveTarget(DestRenderTarget.TargetableTexture, DestRenderTarget.ShaderResourceTexture, false, FResolveParams());
+	Context.RHICmdList.CopyToResolveTarget(DestRenderTarget.TargetableTexture, DestRenderTarget.ShaderResourceTexture, false, FResolveParams());
 }
 
 
diff --git a/Engine/Source/Runtime/Renderer/Private/PostProcess/PostProcessSceneColorFringe.cpp b/Engine/Source/Runtime/Renderer/Private/PostProcess/PostProcessSceneColorFringe.cpp
--- a/Engine/Source/Runtime/Renderer/Private/PostProcess/PostProcessSceneColorFringe.cpp
+++ b/Engine/Source/Runtime/Renderer/Private/PostProcess/PostProcessSceneColorFringe.cpp
@@ -49,15 +49,15 @@ public:
 		return bShaderHasOutdatedParameters;
 	}
 
-	void SetParameters(FRHICommandList& RHICmdList, const FRenderingCompositePassContext& Context)
+	void SetParameters(const FRenderingCompositePassContext& Context)
 	{
 		const FVertexShaderRHIParamRef ShaderRHI = GetVertexShader();
 
 		const FSceneView& View = Context.View;
 
-		FGlobalShader::SetParameters(RHICmdList, ShaderRHI, Context.View);
+		FGlobalShader::SetParameters(Context.RHICmdList, ShaderRHI, Context.View);
 		
-		PostprocessParameter.SetVS(RHICmdList, ShaderRHI, Context, TStaticSamplerState<SF_Bilinear,AM_Clamp,AM_Clamp,AM_Clamp>::GetRHI());
+		PostprocessParameter.SetVS(ShaderRHI, Context, TStaticSamplerState<SF_Bilinear,AM_Clamp,AM_Clamp,AM_Clamp>::GetRHI());
 
 		const FPooledRenderTargetDesc* InputDesc = Context.Pass->GetInputDesc(ePId_Input0);
 
@@ -73,7 +73,7 @@ public:
 			// we only get bigger to not leak in content from outside
 			FVector4 Value(1.0f, 1.0f - Offset * 0.5f, 1.0f - Offset, 0);
 
-			SetShaderValue(RHICmdList, ShaderRHI, FringeUVParams, Value);
+			SetShaderValue(Context.RHICmdList, ShaderRHI, FringeUVParams, Value);
 		}
 	}
 };
@@ -114,15 +114,15 @@ public:
 		return bShaderHasOutdatedParameters;
 	}
 
-	void SetParameters(FRHICommandList& RHICmdList, const FRenderingCompositePassContext& Context)
+	void SetParameters(const FRenderingCompositePassContext& Context)
 	{
 		const FPixelShaderRHIParamRef ShaderRHI = GetPixelShader();
 
 		const FSceneView& View = Context.View;
 
-		FGlobalShader::SetParameters(RHICmdList, ShaderRHI, View);
+		FGlobalShader::SetParameters(Context.RHICmdList, ShaderRHI, View);
 
-		PostprocessParameter.SetPS(RHICmdList, ShaderRHI, Context, TStaticSamplerState<SF_Bilinear,AM_Clamp,AM_Clamp,AM_Clamp>::GetRHI());
+		PostprocessParameter.SetPS(ShaderRHI, Context, TStaticSamplerState<SF_Bilinear,AM_Clamp,AM_Clamp,AM_Clamp>::GetRHI());
 
 		{
 			FVector4 Value[3];
@@ -133,7 +133,7 @@ public:
 			Value[0] = FMath::Lerp(FVector4(t, t, t, 0), FVector4(1, 0, 0, 0), Alpha);
 			Value[1] = FMath::Lerp(FVector4(t, t, t, 0), FVector4(0, 1, 0, 0), Alpha);
 			Value[2] = FMath::Lerp(FVector4(t, t, t, 0), FVector4(0, 0, 1, 0), Alpha);
-			SetShaderValueArray(RHICmdList, ShaderRHI, FringeColorParams, Value, 3);
+			SetShaderValueArray(Context.RHICmdList, ShaderRHI, FringeColorParams, Value, 3);
 		}
 	}
 };
@@ -158,28 +158,28 @@ void FRCPassPostProcessSceneColorFringe::Process(FRenderingCompositePassContext&
 	const FSceneRenderTargetItem& DestRenderTarget = PassOutputs[0].RequestSurface(Context);
 
 	// Set the view family's render target/viewport.
-	RHISetRenderTarget(DestRenderTarget.TargetableTexture, FTextureRHIRef());	
+	SetRenderTarget(Context.RHICmdList, DestRenderTarget.TargetableTexture, FTextureRHIRef());
 	Context.SetViewportAndCallRHI(View.ViewRect);
 
 	// set the state
-	RHISetBlendState(TStaticBlendState<>::GetRHI());
-	RHISetRasterizerState(TStaticRasterizerState<>::GetRHI());
-	RHISetDepthStencilState(TStaticDepthStencilState<false,CF_Always>::GetRHI());
+	Context.RHICmdList.SetBlendState(TStaticBlendState<>::GetRHI());
+	Context.RHICmdList.SetRasterizerState(TStaticRasterizerState<>::GetRHI());
+	Context.RHICmdList.SetDepthStencilState(TStaticDepthStencilState<false, CF_Always>::GetRHI());
 	
 	TShaderMapRef<FPostProcessSceneColorFringeVS> VertexShader(GetGlobalShaderMap());
 	TShaderMapRef<FPostProcessSceneColorFringePS> PixelShader(GetGlobalShaderMap());
 
 	static FGlobalBoundShaderState BoundShaderState;
+	Context.RHICmdList.CheckIsNull(); // need new approach for "static FGlobalBoundShaderState" for parallel rendering
 
-	SetGlobalBoundShaderState(BoundShaderState, GFilterVertexDeclaration.VertexDeclarationRHI, *VertexShader, *PixelShader);
+	SetGlobalBoundShaderState(Context.RHICmdList, BoundShaderState, GFilterVertexDeclaration.VertexDeclarationRHI, *VertexShader, *PixelShader);
 
-	//@todo-rco: RHIPacketList
-	FRHICommandList& RHICmdList = FRHICommandList::GetNullRef();
-	PixelShader->SetParameters(RHICmdList, Context);
-	VertexShader->SetParameters(RHICmdList, Context);
+	PixelShader->SetParameters(Context);
+	VertexShader->SetParameters(Context);
 	
 	// Draw a quad mapping scene color to the view's render target
 	DrawRectangle(
+		Context.RHICmdList,
 		0, 0,
 		View.ViewRect.Width(), View.ViewRect.Height(),
 		View.ViewRect.Min.X, View.ViewRect.Min.Y, 
@@ -189,7 +189,7 @@ void FRCPassPostProcessSceneColorFringe::Process(FRenderingCompositePassContext&
 		*VertexShader,
 		EDRF_UseTriangleOptimization);
 
-	RHICopyToResolveTarget(DestRenderTarget.TargetableTexture, DestRenderTarget.ShaderResourceTexture, false, FResolveParams());
+	Context.RHICmdList.CopyToResolveTarget(DestRenderTarget.TargetableTexture, DestRenderTarget.ShaderResourceTexture, false, FResolveParams());
 }
 
 
diff --git a/Engine/Source/Runtime/Renderer/Private/PostProcess/PostProcessSelectionOutline.cpp b/Engine/Source/Runtime/Renderer/Private/PostProcess/PostProcessSelectionOutline.cpp
--- a/Engine/Source/Runtime/Renderer/Private/PostProcess/PostProcessSelectionOutline.cpp
+++ b/Engine/Source/Runtime/Renderer/Private/PostProcess/PostProcessSelectionOutline.cpp
@@ -34,9 +34,9 @@ void FRCPassPostProcessSelectionOutlineColor::Process(FRenderingCompositePassCon
 	const FSceneRenderTargetItem& DestRenderTarget = PassOutputs[0].RequestSurface(Context);
 	
 	// Set the render target/viewport.
-	RHISetRenderTarget(FTextureRHIParamRef(), DestRenderTarget.TargetableTexture);	
+	SetRenderTarget(Context.RHICmdList, FTextureRHIParamRef(), DestRenderTarget.TargetableTexture);
 	// This is a reversed Z depth surface, so 0.0f is the far plane.
-	RHIClear(false, FLinearColor(0, 0, 0, 0), true, 0.0f, true, 0, FIntRect());
+	Context.RHICmdList.Clear(false, FLinearColor(0, 0, 0, 0), true, 0.0f, true, 0, FIntRect());
 	Context.SetViewportAndCallRHI(ViewRect);
 
 	if(View.VisibleDynamicPrimitives.Num() && View.Family->EngineShowFlags.Selection)
@@ -57,14 +57,16 @@ void FRCPassPostProcessSelectionOutlineColor::Process(FRenderingCompositePassCon
 
 		if (PrimitivesByActor.Num())
 		{
+			Context.RHICmdList.CheckIsNull(); // this uses draw dynamic elements
+
 			// 0 means no object, 1 means BSP so we start with 1
 			uint32 StencilValue = 2;
 
 			TShaderMapRef< TOneColorVS<false> > VertexShader(GetGlobalShaderMap());
 			TShaderMapRef<FOneColorPS> PixelShader(GetGlobalShaderMap());
 
-			RHISetRasterizerState(TStaticRasterizerState<>::GetRHI());
-			RHISetBlendState(TStaticBlendStateWriteMask<CW_NONE, CW_NONE, CW_NONE, CW_NONE>::GetRHI());
+			Context.RHICmdList.SetRasterizerState(TStaticRasterizerState<>::GetRHI());
+			Context.RHICmdList.SetBlendState(TStaticBlendStateWriteMask<CW_NONE, CW_NONE, CW_NONE, CW_NONE>::GetRHI());
 
 			// Sort by actor
 			TArray<FName> Actors;
@@ -75,12 +77,12 @@ void FRCPassPostProcessSelectionOutlineColor::Process(FRenderingCompositePassCon
 				if (bBSP)
 				{
 					// This is a reversed Z depth surface, using CF_GreaterEqual.
-					RHISetDepthStencilState(TStaticDepthStencilState<true,CF_GreaterEqual,true,CF_Always,SO_Keep,SO_Keep,SO_Replace>::GetRHI(),1);
+					Context.RHICmdList.SetDepthStencilState(TStaticDepthStencilState<true, CF_GreaterEqual, true, CF_Always, SO_Keep, SO_Keep, SO_Replace>::GetRHI(), 1);
 				}
 				else
 				{
 					// This is a reversed Z depth surface, using CF_GreaterEqual.
-					RHISetDepthStencilState(TStaticDepthStencilState<true,CF_GreaterEqual,true,CF_Always,SO_Keep,SO_Keep,SO_Replace>::GetRHI(),StencilValue);
+					Context.RHICmdList.SetDepthStencilState(TStaticDepthStencilState<true, CF_GreaterEqual, true, CF_Always, SO_Keep, SO_Keep, SO_Replace>::GetRHI(), StencilValue);
 
 					// we want to use 1..255 for all objects, not correct silhouettes after that is acceptable
 					StencilValue = (StencilValue == 255) ? 2 : (StencilValue + 1);
@@ -111,24 +113,24 @@ void FRCPassPostProcessSelectionOutlineColor::Process(FRenderingCompositePassCon
 			// so we to 4 clears - one for each border.
 
 			// top
-			RHISetScissorRect(true, ViewRect.Min.X, ViewRect.Min.Y, ViewRect.Max.X, InnerRect.Min.Y);
-			RHIClear(false, FLinearColor(0, 0, 0, 0), true, 0.0f, true, 0, FIntRect());
+			Context.RHICmdList.SetScissorRect(true, ViewRect.Min.X, ViewRect.Min.Y, ViewRect.Max.X, InnerRect.Min.Y);
+			Context.RHICmdList.Clear(false, FLinearColor(0, 0, 0, 0), true, 0.0f, true, 0, FIntRect());
 			// bottom
-			RHISetScissorRect(true, ViewRect.Min.X, InnerRect.Max.Y, ViewRect.Max.X, ViewRect.Max.Y);
-			RHIClear(false, FLinearColor(0, 0, 0, 0), true, 0.0f, true, 0, FIntRect());
+			Context.RHICmdList.SetScissorRect(true, ViewRect.Min.X, InnerRect.Max.Y, ViewRect.Max.X, ViewRect.Max.Y);
+			Context.RHICmdList.Clear(false, FLinearColor(0, 0, 0, 0), true, 0.0f, true, 0, FIntRect());
 			// left
-			RHISetScissorRect(true, ViewRect.Min.X, ViewRect.Min.Y, InnerRect.Min.X, ViewRect.Max.Y);
-			RHIClear(false, FLinearColor(0, 0, 0, 0), true, 0.0f, true, 0, FIntRect());
+			Context.RHICmdList.SetScissorRect(true, ViewRect.Min.X, ViewRect.Min.Y, InnerRect.Min.X, ViewRect.Max.Y);
+			Context.RHICmdList.Clear(false, FLinearColor(0, 0, 0, 0), true, 0.0f, true, 0, FIntRect());
 			// right
-			RHISetScissorRect(true, InnerRect.Max.X, ViewRect.Min.Y, ViewRect.Max.X, ViewRect.Max.Y);
-			RHIClear(false, FLinearColor(0, 0, 0, 0), true, 0.0f, true, 0, FIntRect());
+			Context.RHICmdList.SetScissorRect(true, InnerRect.Max.X, ViewRect.Min.Y, ViewRect.Max.X, ViewRect.Max.Y);
+			Context.RHICmdList.Clear(false, FLinearColor(0, 0, 0, 0), true, 0.0f, true, 0, FIntRect());
 
-			RHISetScissorRect(false, 0, 0, 0, 0);
+			Context.RHICmdList.SetScissorRect(false, 0, 0, 0, 0);
 		}
 	}
 	
 	// Resolve to the output
-	RHICopyToResolveTarget(DestRenderTarget.TargetableTexture, DestRenderTarget.ShaderResourceTexture, false, FResolveParams());
+	Context.RHICmdList.CopyToResolveTarget(DestRenderTarget.TargetableTexture, DestRenderTarget.ShaderResourceTexture, false, FResolveParams());
 }
 
 FPooledRenderTargetDesc FRCPassPostProcessSelectionOutlineColor::ComputeOutputDesc(EPassOutputId InPassOutputId) const
@@ -203,19 +205,19 @@ public:
 		EditorPrimitivesStencil.Bind(Initializer.ParameterMap,TEXT("EditorPrimitivesStencil"));
 	}
 
-	void SetPS(FRHICommandList& RHICmdList, const FRenderingCompositePassContext& Context)
+	void SetPS(const FRenderingCompositePassContext& Context)
 	{
 		const FPixelShaderRHIParamRef ShaderRHI = GetPixelShader();
 
-		FGlobalShader::SetParameters(RHICmdList, ShaderRHI, Context.View);
+		FGlobalShader::SetParameters(Context.RHICmdList, ShaderRHI, Context.View);
 
-		DeferredParameters.Set(RHICmdList, ShaderRHI, Context.View);
+		DeferredParameters.Set(Context.RHICmdList, ShaderRHI, Context.View);
 
 		const FPostProcessSettings& Settings = Context.View.FinalPostProcessSettings;
 		const FSceneViewFamily& ViewFamily = *(Context.View.Family);
 		FSceneViewState* ViewState = (FSceneViewState*)Context.View.State;
 
-		PostprocessParameter.SetPS(RHICmdList, ShaderRHI, Context, TStaticSamplerState<SF_Point,AM_Clamp,AM_Clamp,AM_Clamp>::GetRHI());
+		PostprocessParameter.SetPS(ShaderRHI, Context, TStaticSamplerState<SF_Point,AM_Clamp,AM_Clamp,AM_Clamp>::GetRHI());
 
 		// PostprocessInput1MS and EditorPrimitivesStencil
 		{
@@ -233,7 +235,7 @@ public:
 
 			FTexture2DRHIRef& TargetableTexture = (FTexture2DRHIRef&)InputPooledElement->GetRenderTargetItem().TargetableTexture;
 
-			SetTextureParameter(RHICmdList, ShaderRHI, PostprocessInput1MS, TargetableTexture);
+			SetTextureParameter(Context.RHICmdList, ShaderRHI, PostprocessInput1MS, TargetableTexture);
 
 			if(EditorPrimitivesStencil.IsBound())
 			{
@@ -252,7 +254,7 @@ public:
 					ViewState->SelectionOutlineCacheValue = RHICreateShaderResourceView(TargetableTexture, 0, 1, PF_X24_G8);
 				}
 
-			SetSRVParameter(RHICmdList, ShaderRHI, EditorPrimitivesStencil, ViewState->SelectionOutlineCacheValue);
+				SetSRVParameter(Context.RHICmdList, ShaderRHI, EditorPrimitivesStencil, ViewState->SelectionOutlineCacheValue);
 		}
 		}
 
@@ -262,8 +264,8 @@ public:
 
 			Value.A = GEngine->SelectionHighlightIntensity;
 
-			SetShaderValue(RHICmdList, ShaderRHI, OutlineColor, Value );
-			SetShaderValue(RHICmdList, ShaderRHI, BSPSelectionIntensity, GEngine->BSPSelectionHighlightIntensity);
+			SetShaderValue(Context.RHICmdList, ShaderRHI, OutlineColor, Value);
+			SetShaderValue(Context.RHICmdList, ShaderRHI, BSPSelectionIntensity, GEngine->BSPSelectionHighlightIntensity);
 		}
 #else
 		check(!"This shader is not used outside of the Editor.");
@@ -280,7 +282,7 @@ public:
 				Value.G = 0;
 			}
 
-			SetShaderValue(RHICmdList, ShaderRHI, EditorRenderParams, Value );
+			SetShaderValue(Context.RHICmdList, ShaderRHI, EditorRenderParams, Value);
 		}
 	}
 	
@@ -312,15 +314,17 @@ VARIATION1(1)			VARIATION1(2)			VARIATION1(4)			VARIATION1(8)
 
 
 template <uint32 MSAASampleCount>
-static void SetSelectionOutlineShaderTempl(FRHICommandList& RHICmdList, const FRenderingCompositePassContext& Context)
+static void SetSelectionOutlineShaderTempl(const FRenderingCompositePassContext& Context)
 {
 	TShaderMapRef<FPostProcessVS> VertexShader(GetGlobalShaderMap());
 	TShaderMapRef<FPostProcessSelectionOutlinePS<MSAASampleCount> > PixelShader(GetGlobalShaderMap());
 
 	static FGlobalBoundShaderState BoundShaderState;
-	SetGlobalBoundShaderState(BoundShaderState, GFilterVertexDeclaration.VertexDeclarationRHI, *VertexShader, *PixelShader);
+	Context.RHICmdList.CheckIsNull(); // need new approach for "static FGlobalBoundShaderState" for parallel rendering
+
+	SetGlobalBoundShaderState(Context.RHICmdList, BoundShaderState, GFilterVertexDeclaration.VertexDeclarationRHI, *VertexShader, *PixelShader);
 
-	PixelShader->SetPS(RHICmdList, Context);
+	PixelShader->SetPS(Context);
 }
 
 void FRCPassPostProcessSelectionOutline::Process(FRenderingCompositePassContext& Context)
@@ -334,8 +338,6 @@ void FRCPassPostProcessSelectionOutline::Process(FRenderingCompositePassContext&
 		// input is not hooked up correctly
 		return;
 	}
-	//@todo-rco: RHIPacketList
-	FRHICommandList& RHICmdList = FRHICommandList::GetNullRef();
 
 	const FSceneView& View = Context.View;
 	FIntRect ViewRect = View.ViewRect;
@@ -344,32 +346,32 @@ void FRCPassPostProcessSelectionOutline::Process(FRenderingCompositePassContext&
 	const FSceneRenderTargetItem& DestRenderTarget = PassOutputs[0].RequestSurface(Context);
 
 	// Set the view family's render target/viewport.
-	RHISetRenderTarget(DestRenderTarget.TargetableTexture, FTextureRHIRef());	
-	RHIClear(true, FLinearColor::Black, false, 1.0f, false, 0, ViewRect);
+	SetRenderTarget(Context.RHICmdList, DestRenderTarget.TargetableTexture, FTextureRHIRef());
+	Context.RHICmdList.Clear(true, FLinearColor::Black, false, 1.0f, false, 0, ViewRect);
 	Context.SetViewportAndCallRHI(ViewRect);
 
 	// set the state
-	RHISetBlendState(TStaticBlendState<>::GetRHI());
-	RHISetRasterizerState(TStaticRasterizerState<>::GetRHI());
-	RHISetDepthStencilState(TStaticDepthStencilState<false,CF_Always>::GetRHI());
+	Context.RHICmdList.SetBlendState(TStaticBlendState<>::GetRHI());
+	Context.RHICmdList.SetRasterizerState(TStaticRasterizerState<>::GetRHI());
+	Context.RHICmdList.SetDepthStencilState(TStaticDepthStencilState<false, CF_Always>::GetRHI());
 
 	const uint32 MSAASampleCount = GSceneRenderTargets.GetEditorMSAACompositingSampleCount();
 
 	if(MSAASampleCount == 1)
 	{
-		SetSelectionOutlineShaderTempl<1>(RHICmdList, Context);
+		SetSelectionOutlineShaderTempl<1>(Context);
 	}
 	else if(MSAASampleCount == 2)
 	{
-		SetSelectionOutlineShaderTempl<2>(RHICmdList, Context);
+		SetSelectionOutlineShaderTempl<2>(Context);
 	}
 	else if(MSAASampleCount == 4)
 	{
-		SetSelectionOutlineShaderTempl<4>(RHICmdList, Context);
+		SetSelectionOutlineShaderTempl<4>(Context);
 	}
 	else if(MSAASampleCount == 8)
 	{
-		SetSelectionOutlineShaderTempl<8>(RHICmdList, Context);
+		SetSelectionOutlineShaderTempl<8>(Context);
 	}
 	else
 	{
@@ -380,6 +382,7 @@ void FRCPassPostProcessSelectionOutline::Process(FRenderingCompositePassContext&
 	// Draw a quad mapping scene color to the view's render target
 	TShaderMapRef<FPostProcessVS> VertexShader(GetGlobalShaderMap());
 	DrawRectangle(
+		Context.RHICmdList,
 		0, 0,
 		ViewRect.Width(), ViewRect.Height(),
 		ViewRect.Min.X, ViewRect.Min.Y,
@@ -389,7 +392,7 @@ void FRCPassPostProcessSelectionOutline::Process(FRenderingCompositePassContext&
 		*VertexShader,
 		EDRF_UseTriangleOptimization);
 
-	RHICopyToResolveTarget(DestRenderTarget.TargetableTexture, DestRenderTarget.ShaderResourceTexture, false, FResolveParams());
+	Context.RHICmdList.CopyToResolveTarget(DestRenderTarget.TargetableTexture, DestRenderTarget.ShaderResourceTexture, false, FResolveParams());
 }
 
 FPooledRenderTargetDesc FRCPassPostProcessSelectionOutline::ComputeOutputDesc(EPassOutputId InPassOutputId) const
diff --git a/Engine/Source/Runtime/Renderer/Private/PostProcess/PostProcessSubsurface.cpp b/Engine/Source/Runtime/Renderer/Private/PostProcess/PostProcessSubsurface.cpp
--- a/Engine/Source/Runtime/Renderer/Private/PostProcess/PostProcessSubsurface.cpp
+++ b/Engine/Source/Runtime/Renderer/Private/PostProcess/PostProcessSubsurface.cpp
@@ -44,14 +44,14 @@ public:
 		DeferredParameters.Bind(Initializer.ParameterMap);
 	}
 
-	void SetParameters(FRHICommandList& RHICmdList, const FRenderingCompositePassContext& Context)
+	void SetParameters(const FRenderingCompositePassContext& Context)
 	{
 		const FFinalPostProcessSettings& Settings = Context.View.FinalPostProcessSettings;
 		const FPixelShaderRHIParamRef ShaderRHI = GetPixelShader();
 
-		FGlobalShader::SetParameters(RHICmdList, ShaderRHI, Context.View);
-		PostprocessParameter.SetPS(RHICmdList, ShaderRHI, Context, TStaticSamplerState<SF_Point,AM_Clamp,AM_Clamp,AM_Clamp>::GetRHI());
-		DeferredParameters.Set(RHICmdList, ShaderRHI, Context.View);
+		FGlobalShader::SetParameters(Context.RHICmdList, ShaderRHI, Context.View);
+		PostprocessParameter.SetPS(ShaderRHI, Context, TStaticSamplerState<SF_Point,AM_Clamp,AM_Clamp,AM_Clamp>::GetRHI());
+		DeferredParameters.Set(Context.RHICmdList, ShaderRHI, Context.View);
 	}
 
 	// FShader interface.
@@ -83,17 +83,18 @@ public:
 
 
 template <uint32 SpecularCorrection>
-void SetSubsurfaceSetupShader(FRHICommandList& RHICmdList, const FRenderingCompositePassContext& Context)
+void SetSubsurfaceSetupShader(const FRenderingCompositePassContext& Context)
 {
 	TShaderMapRef<FPostProcessVS> VertexShader(GetGlobalShaderMap());
 	TShaderMapRef<FPostProcessSubsurfaceSetupPS<SpecularCorrection> > PixelShader(GetGlobalShaderMap());
 
 	static FGlobalBoundShaderState BoundShaderState;
+	Context.RHICmdList.CheckIsNull(); // need new approach for "static FGlobalBoundShaderState" for parallel rendering
 
-	SetGlobalBoundShaderState(BoundShaderState, GFilterVertexDeclaration.VertexDeclarationRHI, *VertexShader, *PixelShader);
+	SetGlobalBoundShaderState(Context.RHICmdList, BoundShaderState, GFilterVertexDeclaration.VertexDeclarationRHI, *VertexShader, *PixelShader);
 
-	PixelShader->SetParameters(RHICmdList, Context);
-	VertexShader->SetParameters(RHICmdList, Context);
+	PixelShader->SetParameters(Context);
+	VertexShader->SetParameters(Context);
 }
 
 static TAutoConsoleVariable<int> CVarSubsurfaceQuality(
@@ -146,36 +147,34 @@ void FRCPassPostProcessSubsurfaceSetup::Process(FRenderingCompositePassContext&
 
 	const FSceneRenderTargetItem& DestRenderTarget = PassOutputs[0].RequestSurface(Context);
 
-	//@todo-rco: RHIPacketList
-	FRHICommandList& RHICmdList = FRHICommandList::GetNullRef();
-
 	// Set the view family's render target/viewport.
-	RHISetRenderTarget(DestRenderTarget.TargetableTexture, FTextureRHIRef());	
+	SetRenderTarget(Context.RHICmdList, DestRenderTarget.TargetableTexture, FTextureRHIRef());
 
 	// is optimized away if possible (RT size=view size, )
-	RHIClear(true, FLinearColor::Black, false, 1.0f, false, 0, DestRect);
+	Context.RHICmdList.Clear(true, FLinearColor::Black, false, 1.0f, false, 0, DestRect);
 
 	Context.SetViewportAndCallRHI(0, 0, 0.0f, DestSize.X, DestSize.Y, 1.0f );
 
 	// set the state
-	RHISetBlendState(TStaticBlendState<>::GetRHI());
-	RHISetRasterizerState(TStaticRasterizerState<>::GetRHI());
-	RHISetDepthStencilState(TStaticDepthStencilState<false,CF_Always>::GetRHI());
+	Context.RHICmdList.SetBlendState(TStaticBlendState<>::GetRHI());
+	Context.RHICmdList.SetRasterizerState(TStaticRasterizerState<>::GetRHI());
+	Context.RHICmdList.SetDepthStencilState(TStaticDepthStencilState<false, CF_Always>::GetRHI());
 
 	if(DoSpecularCorrection())
 	{
 		// with separate specular
-		SetSubsurfaceSetupShader<1>(RHICmdList, Context);
+		SetSubsurfaceSetupShader<1>(Context);
 	}
 	else
 	{
 		// no separate specular
-		SetSubsurfaceSetupShader<0>(RHICmdList, Context);
+		SetSubsurfaceSetupShader<0>(Context);
 	}
 
 	// Draw a quad mapping scene color to the view's render target
 	TShaderMapRef<FPostProcessVS> VertexShader(GetGlobalShaderMap());
 	DrawRectangle(
+		Context.RHICmdList,
 		DestRect.Min.X, DestRect.Min.Y,
 		DestRect.Width(), DestRect.Height(),
 		SrcRect.Min.X, SrcRect.Min.Y,
@@ -185,7 +184,7 @@ void FRCPassPostProcessSubsurfaceSetup::Process(FRenderingCompositePassContext&
 		*VertexShader,
 		EDRF_UseTriangleOptimization);
 
-	RHICopyToResolveTarget(DestRenderTarget.TargetableTexture, DestRenderTarget.ShaderResourceTexture, false, FResolveParams());
+	Context.RHICmdList.CopyToResolveTarget(DestRenderTarget.TargetableTexture, DestRenderTarget.ShaderResourceTexture, false, FResolveParams());
 }
 
 FPooledRenderTargetDesc FRCPassPostProcessSubsurfaceSetup::ComputeOutputDesc(EPassOutputId InPassOutputId) const
@@ -244,17 +243,17 @@ public:
 		return bShaderHasOutdatedParameters;
 	}
 
-	void SetParameters(FRHICommandList& RHICmdList, const FRenderingCompositePassContext& Context, float InRadius)
+	void SetParameters(const FRenderingCompositePassContext& Context, float InRadius)
 	{
 		const FPixelShaderRHIParamRef ShaderRHI = GetPixelShader();
 
-		FGlobalShader::SetParameters(RHICmdList, ShaderRHI, Context.View);
-		DeferredParameters.Set(RHICmdList, ShaderRHI, Context.View);
-		PostprocessParameter.SetPS(RHICmdList, ShaderRHI, Context, TStaticSamplerState<SF_Bilinear,AM_Border,AM_Border,AM_Border>::GetRHI());
+		FGlobalShader::SetParameters(Context.RHICmdList, ShaderRHI, Context.View);
+		DeferredParameters.Set(Context.RHICmdList, ShaderRHI, Context.View);
+		PostprocessParameter.SetPS(ShaderRHI, Context, TStaticSamplerState<SF_Bilinear,AM_Border,AM_Border,AM_Border>::GetRHI());
 
 		{
 			FVector4 ColorScale(InRadius, 0, 0, 0);
-			SetShaderValue(RHICmdList, ShaderRHI, SSSParams, ColorScale);
+			SetShaderValue(Context.RHICmdList, ShaderRHI, SSSParams, ColorScale);
 		}
 	}
 
@@ -285,17 +284,18 @@ FRCPassPostProcessSubsurface::FRCPassPostProcessSubsurface(uint32 InPass, float
 
 
 template <uint32 Method>
-void SetSubsurfaceShader(FRHICommandList& RHICmdList, const FRenderingCompositePassContext& Context, float InRadius)
+void SetSubsurfaceShader(const FRenderingCompositePassContext& Context, float InRadius)
 {
 	TShaderMapRef<FPostProcessVS> VertexShader(GetGlobalShaderMap());
 	TShaderMapRef<FPostProcessSubsurfacePS<Method> > PixelShader(GetGlobalShaderMap());
 
 	static FGlobalBoundShaderState BoundShaderState;
+	Context.RHICmdList.CheckIsNull(); // need new approach for "static FGlobalBoundShaderState" for parallel rendering
 
-	SetGlobalBoundShaderState(BoundShaderState, GFilterVertexDeclaration.VertexDeclarationRHI, *VertexShader, *PixelShader);
+	SetGlobalBoundShaderState(Context.RHICmdList, BoundShaderState, GFilterVertexDeclaration.VertexDeclarationRHI, *VertexShader, *PixelShader);
 
-	PixelShader->SetParameters(RHICmdList, Context, InRadius);
-	VertexShader->SetParameters(RHICmdList, Context);
+	PixelShader->SetParameters(Context, InRadius);
+	VertexShader->SetParameters(Context);
 }
 
 
@@ -324,19 +324,17 @@ void FRCPassPostProcessSubsurface::Process(FRenderingCompositePassContext& Conte
 	const FSceneRenderTargetItem& DestRenderTarget = PassOutputs[0].RequestSurface(Context);
 
 	// Set the view family's render target/viewport.
-	RHISetRenderTarget(DestRenderTarget.TargetableTexture, FTextureRHIRef());	
+	SetRenderTarget(Context.RHICmdList, DestRenderTarget.TargetableTexture, FTextureRHIRef());
 
 	// is optimized away if possible (RT size=view size, )
-	RHIClear(true, FLinearColor(0, 0, 0, 0), false, 1.0f, false, 0, DestRect);
+	Context.RHICmdList.Clear(true, FLinearColor(0, 0, 0, 0), false, 1.0f, false, 0, DestRect);
 
 	Context.SetViewportAndCallRHI(0, 0, 0.0f, DestSize.X, DestSize.Y, 1.0f );
 
 	// set the state
-	RHISetBlendState(TStaticBlendState<>::GetRHI());
-	RHISetRasterizerState(TStaticRasterizerState<>::GetRHI());
-	RHISetDepthStencilState(TStaticDepthStencilState<false,CF_Always>::GetRHI());
-	//@todo-rco: RHIPacketList
-	FRHICommandList& RHICmdList = FRHICommandList::GetNullRef();
+	Context.RHICmdList.SetBlendState(TStaticBlendState<>::GetRHI());
+	Context.RHICmdList.SetRasterizerState(TStaticRasterizerState<>::GetRHI());
+	Context.RHICmdList.SetDepthStencilState(TStaticDepthStencilState<false, CF_Always>::GetRHI());
 
 	TShaderMapRef<FPostProcessVS> VertexShader(GetGlobalShaderMap());
 
@@ -344,10 +342,11 @@ void FRCPassPostProcessSubsurface::Process(FRenderingCompositePassContext& Conte
 	{
 		SCOPED_DRAW_EVENT(SubsurfacePass0, DEC_SCENE_ITEMS);
 	
-		SetSubsurfaceShader<0>(RHICmdList, Context, Radius);
+		SetSubsurfaceShader<0>(Context, Radius);
 
 		// Draw a quad mapping scene color to the view's render target
 		DrawRectangle(
+			Context.RHICmdList,
 			DestRect.Min.X, DestRect.Min.Y,
 			DestRect.Width(), DestRect.Height(),
 			SrcRect.Min.X, SrcRect.Min.Y,
@@ -357,7 +356,7 @@ void FRCPassPostProcessSubsurface::Process(FRenderingCompositePassContext& Conte
 			*VertexShader,
 			EDRF_UseTriangleOptimization);
 
-		RHICopyToResolveTarget(DestRenderTarget.TargetableTexture, DestRenderTarget.ShaderResourceTexture, false, FResolveParams());
+		Context.RHICmdList.CopyToResolveTarget(DestRenderTarget.TargetableTexture, DestRenderTarget.ShaderResourceTexture, false, FResolveParams());
 	}
 	
 	if(Pass == 1)
@@ -367,15 +366,16 @@ void FRCPassPostProcessSubsurface::Process(FRenderingCompositePassContext& Conte
 		if(DoSpecularCorrection())
 		{
 			// reconstruct specular and add it in final pass
-			SetSubsurfaceShader<2>(RHICmdList, Context, Radius);
+			SetSubsurfaceShader<2>(Context, Radius);
 		}
 		else
 		{
-			SetSubsurfaceShader<1>(RHICmdList, Context, Radius);
+			SetSubsurfaceShader<1>(Context, Radius);
 		}
 
 		// Draw a quad mapping scene color to the view's render target
 		DrawRectangle(
+			Context.RHICmdList,
 			DestRect.Min.X, DestRect.Min.Y,
 			DestRect.Width(), DestRect.Height(),
 			SrcRect.Min.X, SrcRect.Min.Y,
@@ -385,7 +385,7 @@ void FRCPassPostProcessSubsurface::Process(FRenderingCompositePassContext& Conte
 			*VertexShader,
 			EDRF_UseTriangleOptimization);
 
-		RHICopyToResolveTarget(DestRenderTarget.TargetableTexture, DestRenderTarget.ShaderResourceTexture, false, FResolveParams());
+		Context.RHICmdList.CopyToResolveTarget(DestRenderTarget.TargetableTexture, DestRenderTarget.ShaderResourceTexture, false, FResolveParams());
 	}
 }
 
diff --git a/Engine/Source/Runtime/Renderer/Private/PostProcess/PostProcessTemporalAA.cpp b/Engine/Source/Runtime/Renderer/Private/PostProcess/PostProcessTemporalAA.cpp
--- a/Engine/Source/Runtime/Renderer/Private/PostProcess/PostProcessTemporalAA.cpp
+++ b/Engine/Source/Runtime/Renderer/Private/PostProcess/PostProcessTemporalAA.cpp
@@ -88,11 +88,11 @@ public:
 		return bShaderHasOutdatedParameters;
 	}
 
-	void SetParameters(FRHICommandList& RHICmdList, const FRenderingCompositePassContext& Context)
+	void SetParameters(const FRenderingCompositePassContext& Context)
 	{
 		const FPixelShaderRHIParamRef ShaderRHI = GetPixelShader();
 
-		FGlobalShader::SetParameters(RHICmdList, ShaderRHI, Context.View);
+		FGlobalShader::SetParameters(Context.RHICmdList, ShaderRHI, Context.View);
 		
 		FSamplerStateRHIParamRef FilterTable[4];
 		FilterTable[0] = TStaticSamplerState<SF_Point,AM_Clamp,AM_Clamp,AM_Clamp>::GetRHI();
@@ -100,9 +100,9 @@ public:
 		FilterTable[2] = FilterTable[0];
 		FilterTable[3] = FilterTable[0];
 
-		PostprocessParameter.SetPS(RHICmdList, ShaderRHI, Context, 0, false, FilterTable);
+		PostprocessParameter.SetPS(ShaderRHI, Context, 0, false, FilterTable);
 
-		DeferredParameters.Set(RHICmdList, ShaderRHI, Context.View);
+		DeferredParameters.Set(Context.RHICmdList, ShaderRHI, Context.View);
 
 		FSceneViewState* ViewState = (FSceneViewState*)Context.View.State;
 
@@ -161,17 +161,17 @@ public:
 			
 			for( int32 i = 0; i < 9; i++ )
 			{
-				SetShaderValue(RHICmdList, ShaderRHI, SampleWeights, Weights[i] / TotalWeight, i );
-				SetShaderValue(RHICmdList, ShaderRHI, LowpassWeights, WeightsLow[i] / TotalWeightLow, i );
+				SetShaderValue(Context.RHICmdList, ShaderRHI, SampleWeights, Weights[i] / TotalWeight, i );
+				SetShaderValue(Context.RHICmdList, ShaderRHI, LowpassWeights, WeightsLow[i] / TotalWeightLow, i );
 			}
 			
 			FVector2D RandomOffsetValue;
 			TemporalRandom(&RandomOffsetValue, Context.View.FrameNumber);
-			SetShaderValue(RHICmdList, ShaderRHI, RandomOffset, RandomOffsetValue);
+			SetShaderValue(Context.RHICmdList, ShaderRHI, RandomOffset, RandomOffsetValue);
 
 		}
 
-		CameraMotionParams.Set(RHICmdList, Context.View, ShaderRHI);
+		CameraMotionParams.Set(Context.RHICmdList, Context.View, ShaderRHI);
 	}
 };
 
@@ -215,32 +215,34 @@ void FRCPassPostProcessSSRTemporalAA::Process(FRenderingCompositePassContext& Co
 
 	const FSceneRenderTargetItem& DestRenderTarget = PassOutputs[0].RequestSurface(Context);
 
-	RHISetRenderTarget(DestRenderTarget.TargetableTexture, FTextureRHIRef());
+	SetRenderTarget(Context.RHICmdList, DestRenderTarget.TargetableTexture, FTextureRHIRef());
 
 	// is optimized away if possible (RT size=view size, )
-	RHIClear(true, FLinearColor::Black, false, 1.0f, false, 0, SrcRect);
+	Context.RHICmdList.Clear(true, FLinearColor::Black, false, 1.0f, false, 0, SrcRect);
 
 	Context.SetViewportAndCallRHI(SrcRect);
 
 	// set the state
-	RHISetBlendState(TStaticBlendState<>::GetRHI());
-	RHISetRasterizerState(TStaticRasterizerState<>::GetRHI());
-	RHISetDepthStencilState(TStaticDepthStencilState<false,CF_Always>::GetRHI());
+	Context.RHICmdList.SetBlendState(TStaticBlendState<>::GetRHI());
+	Context.RHICmdList.SetRasterizerState(TStaticRasterizerState<>::GetRHI());
+	Context.RHICmdList.SetDepthStencilState(TStaticDepthStencilState<false,CF_Always>::GetRHI());
 
 	TShaderMapRef< FPostProcessTonemapVS >			VertexShader( GetGlobalShaderMap() );
 	TShaderMapRef< FPostProcessTemporalAAPS<2,0> >	PixelShader( GetGlobalShaderMap() );
 
 	static FGlobalBoundShaderState BoundShaderState;
+	Context.RHICmdList.CheckIsNull(); // need new approach for "static FGlobalBoundShaderState" for parallel rendering
 
-	SetGlobalBoundShaderState(BoundShaderState, GFilterVertexDeclaration.VertexDeclarationRHI, *VertexShader, *PixelShader);
+	SetGlobalBoundShaderState(Context.RHICmdList, BoundShaderState, GFilterVertexDeclaration.VertexDeclarationRHI, *VertexShader, *PixelShader);
 
 	//@todo-rco: RHIPacketList
 	FRHICommandList& RHICmdList = FRHICommandList::GetNullRef();
-	VertexShader->SetVS(RHICmdList, Context);
-	PixelShader->SetParameters(RHICmdList, Context);
+	VertexShader->SetVS(Context);
+	PixelShader->SetParameters(Context);
 
 	// Draw a quad mapping scene color to the view's render target
 	DrawRectangle(
+		Context.RHICmdList,
 		0, 0,
 		SrcRect.Width(), SrcRect.Height(),
 		SrcRect.Min.X, SrcRect.Min.Y, 
@@ -250,7 +252,7 @@ void FRCPassPostProcessSSRTemporalAA::Process(FRenderingCompositePassContext& Co
 		*VertexShader,
 		EDRF_UseTriangleOptimization);
 
-	RHICopyToResolveTarget(DestRenderTarget.TargetableTexture, DestRenderTarget.ShaderResourceTexture, false, FResolveParams());
+	Context.RHICmdList.CopyToResolveTarget(DestRenderTarget.TargetableTexture, DestRenderTarget.ShaderResourceTexture, false, FResolveParams());
 }
 
 FPooledRenderTargetDesc FRCPassPostProcessSSRTemporalAA::ComputeOutputDesc(EPassOutputId InPassOutputId) const
@@ -292,32 +294,32 @@ void FRCPassPostProcessDOFTemporalAA::Process(FRenderingCompositePassContext& Co
 
 	const FSceneRenderTargetItem& DestRenderTarget = PassOutputs[0].RequestSurface(Context);
 
-	RHISetRenderTarget(DestRenderTarget.TargetableTexture, FTextureRHIRef());
+	SetRenderTarget(Context.RHICmdList, DestRenderTarget.TargetableTexture, FTextureRHIRef());
 
 	// is optimized away if possible (RT size=view size, )
-	RHIClear(true, FLinearColor::Black, false, 1.0f, false, 0, SrcRect);
+	Context.RHICmdList.Clear(true, FLinearColor::Black, false, 1.0f, false, 0, SrcRect);
 
 	Context.SetViewportAndCallRHI(SrcRect);
 
 	// set the state
-	RHISetBlendState(TStaticBlendState<>::GetRHI());
-	RHISetRasterizerState(TStaticRasterizerState<>::GetRHI());
-	RHISetDepthStencilState(TStaticDepthStencilState<false,CF_Always>::GetRHI());
+	Context.RHICmdList.SetBlendState(TStaticBlendState<>::GetRHI());
+	Context.RHICmdList.SetRasterizerState(TStaticRasterizerState<>::GetRHI());
+	Context.RHICmdList.SetDepthStencilState(TStaticDepthStencilState<false,CF_Always>::GetRHI());
 
 	TShaderMapRef< FPostProcessTonemapVS >			VertexShader( GetGlobalShaderMap() );
 	TShaderMapRef< FPostProcessTemporalAAPS<0,0> >	PixelShader( GetGlobalShaderMap() );
 
 	static FGlobalBoundShaderState BoundShaderState;
+	Context.RHICmdList.CheckIsNull(); // need new approach for "static FGlobalBoundShaderState" for parallel rendering
 
-	SetGlobalBoundShaderState(BoundShaderState, GFilterVertexDeclaration.VertexDeclarationRHI, *VertexShader, *PixelShader);
+	SetGlobalBoundShaderState(Context.RHICmdList, BoundShaderState, GFilterVertexDeclaration.VertexDeclarationRHI, *VertexShader, *PixelShader);
 
-	//@todo-rco: RHIPacketList
-	FRHICommandList& RHICmdList = FRHICommandList::GetNullRef();
-	VertexShader->SetVS(RHICmdList, Context);
-	PixelShader->SetParameters(RHICmdList, Context);
+	VertexShader->SetVS(Context);
+	PixelShader->SetParameters(Context);
 
 	// Draw a quad mapping scene color to the view's render target
 	DrawRectangle(
+		Context.RHICmdList,
 		0, 0,
 		SrcRect.Width(), SrcRect.Height(),
 		SrcRect.Min.X, SrcRect.Min.Y, 
@@ -327,7 +329,9 @@ void FRCPassPostProcessDOFTemporalAA::Process(FRenderingCompositePassContext& Co
 		*VertexShader,
 		EDRF_UseTriangleOptimization);
 
-	RHICopyToResolveTarget(DestRenderTarget.TargetableTexture, DestRenderTarget.ShaderResourceTexture, false, FResolveParams());
+	Context.RHICmdList.CopyToResolveTarget(DestRenderTarget.TargetableTexture, DestRenderTarget.ShaderResourceTexture, false, FResolveParams());
+	Context.RHICmdList.CheckIsNull(); // this modifies the view state, which is probably not safe for parallel rendering
+
 	ViewState->DOFHistoryRT = PassOutputs[0].PooledRenderTarget;
 	check( ViewState->DOFHistoryRT );
 }
@@ -372,32 +376,32 @@ void FRCPassPostProcessLightShaftTemporalAA::Process(FRenderingCompositePassCont
 
 	const FSceneRenderTargetItem& DestRenderTarget = PassOutputs[0].RequestSurface(Context);
 
-	RHISetRenderTarget(DestRenderTarget.TargetableTexture, FTextureRHIRef());
+	SetRenderTarget(Context.RHICmdList, DestRenderTarget.TargetableTexture, FTextureRHIRef());
 
 	// is optimized away if possible (RT size=view size, )
-	RHIClear(true, FLinearColor::Black, false, 1.0f, false, 0, SrcRect);
+	Context.RHICmdList.Clear(true, FLinearColor::Black, false, 1.0f, false, 0, SrcRect);
 
 	Context.SetViewportAndCallRHI(SrcRect);
 
 	// set the state
-	RHISetBlendState(TStaticBlendState<>::GetRHI());
-	RHISetRasterizerState(TStaticRasterizerState<>::GetRHI());
-	RHISetDepthStencilState(TStaticDepthStencilState<false,CF_Always>::GetRHI());
+	Context.RHICmdList.SetBlendState(TStaticBlendState<>::GetRHI());
+	Context.RHICmdList.SetRasterizerState(TStaticRasterizerState<>::GetRHI());
+	Context.RHICmdList.SetDepthStencilState(TStaticDepthStencilState<false,CF_Always>::GetRHI());
 
 	TShaderMapRef< FPostProcessTonemapVS >			VertexShader( GetGlobalShaderMap() );
 	TShaderMapRef< FPostProcessTemporalAAPS<3,0> >	PixelShader( GetGlobalShaderMap() );
 
 	static FGlobalBoundShaderState BoundShaderState;
+	Context.RHICmdList.CheckIsNull(); // need new approach for "static FGlobalBoundShaderState" for parallel rendering
 
-	SetGlobalBoundShaderState(BoundShaderState, GFilterVertexDeclaration.VertexDeclarationRHI, *VertexShader, *PixelShader);
+	SetGlobalBoundShaderState(Context.RHICmdList, BoundShaderState, GFilterVertexDeclaration.VertexDeclarationRHI, *VertexShader, *PixelShader);
 
-	//@todo-rco: RHIPacketList
-	FRHICommandList& RHICmdList = FRHICommandList::GetNullRef();
-	VertexShader->SetVS(RHICmdList, Context);
-	PixelShader->SetParameters(RHICmdList, Context);
+	VertexShader->SetVS(Context);
+	PixelShader->SetParameters(Context);
 
 	// Draw a quad mapping scene color to the view's render target
 	DrawRectangle(
+		Context.RHICmdList,
 		0, 0,
 		SrcRect.Width(), SrcRect.Height(),
 		SrcRect.Min.X, SrcRect.Min.Y, 
@@ -407,7 +411,7 @@ void FRCPassPostProcessLightShaftTemporalAA::Process(FRenderingCompositePassCont
 		*VertexShader,
 		EDRF_UseTriangleOptimization);
 
-	RHICopyToResolveTarget(DestRenderTarget.TargetableTexture, DestRenderTarget.ShaderResourceTexture, false, FResolveParams());
+	Context.RHICmdList.CopyToResolveTarget(DestRenderTarget.TargetableTexture, DestRenderTarget.ShaderResourceTexture, false, FResolveParams());
 }
 
 FPooledRenderTargetDesc FRCPassPostProcessLightShaftTemporalAA::ComputeOutputDesc(EPassOutputId InPassOutputId) const
@@ -450,13 +454,11 @@ void FRCPassPostProcessTemporalAA::Process(FRenderingCompositePassContext& Conte
 
 	const FSceneRenderTargetItem& DestRenderTarget = PassOutputs[0].RequestSurface(Context);
 
-	//@todo-rco: RHIPacketList
-	FRHICommandList& RHICmdList = FRHICommandList::GetNullRef();
-	//RHISetRenderTarget(DestRenderTarget.TargetableTexture, FTextureRHIRef());
-	RHISetRenderTarget(DestRenderTarget.TargetableTexture, GSceneRenderTargets.GetSceneDepthTexture());
+	//Context.RHICmdList.SetRenderTarget(DestRenderTarget.TargetableTexture, FTextureRHIRef());
+	SetRenderTarget(Context.RHICmdList, DestRenderTarget.TargetableTexture, GSceneRenderTargets.GetSceneDepthTexture());
 
 	// is optimized away if possible (RT size=view size, )
-	RHIClear(true, FLinearColor::Black, false, 1.0f, false, 0, SrcRect);
+	Context.RHICmdList.Clear(true, FLinearColor::Black, false, 1.0f, false, 0, SrcRect);
 
 	Context.SetViewportAndCallRHI(SrcRect);
 
@@ -469,7 +471,7 @@ void FRCPassPostProcessTemporalAA::Process(FRenderingCompositePassContext& Conte
 	{	
 		// Normal temporal feedback
 		// Draw to pixels where stencil == 0
-		RHISetDepthStencilState(TStaticDepthStencilState<false,CF_Always,true,CF_Equal,SO_Keep,SO_Keep,SO_Keep>::GetRHI(), 0);
+		Context.RHICmdList.SetDepthStencilState(TStaticDepthStencilState<false,CF_Always,true,CF_Equal,SO_Keep,SO_Keep,SO_Keep>::GetRHI(), 0);
 
 		TShaderMapRef< FPostProcessTonemapVS >			VertexShader(GetGlobalShaderMap());
 		if (bUseFast)
@@ -477,26 +479,29 @@ void FRCPassPostProcessTemporalAA::Process(FRenderingCompositePassContext& Conte
 			TShaderMapRef< FPostProcessTemporalAAPS<4,0> >	PixelShader( GetGlobalShaderMap() );
 
 			static FGlobalBoundShaderState BoundShaderState;
+			Context.RHICmdList.CheckIsNull(); // need new approach for "static FGlobalBoundShaderState" for parallel rendering
 
-			SetGlobalBoundShaderState(BoundShaderState, GFilterVertexDeclaration.VertexDeclarationRHI, *VertexShader, *PixelShader);
+			SetGlobalBoundShaderState(Context.RHICmdList, BoundShaderState, GFilterVertexDeclaration.VertexDeclarationRHI, *VertexShader, *PixelShader);
 
-			VertexShader->SetVS(RHICmdList, Context);
-			PixelShader->SetParameters(RHICmdList, Context);
+			VertexShader->SetVS(Context);
+			PixelShader->SetParameters(Context);
 		}
 		else
 		{
 			TShaderMapRef< FPostProcessTemporalAAPS<1,0> >	PixelShader( GetGlobalShaderMap() );
 
 			static FGlobalBoundShaderState BoundShaderState;
+			Context.RHICmdList.CheckIsNull(); // need new approach for "static FGlobalBoundShaderState" for parallel rendering
 
-			SetGlobalBoundShaderState(BoundShaderState, GFilterVertexDeclaration.VertexDeclarationRHI, *VertexShader, *PixelShader);
+			SetGlobalBoundShaderState(Context.RHICmdList, BoundShaderState, GFilterVertexDeclaration.VertexDeclarationRHI, *VertexShader, *PixelShader);
 
-			VertexShader->SetVS(RHICmdList, Context);
-			PixelShader->SetParameters(RHICmdList, Context);
+			VertexShader->SetVS(Context);
+			PixelShader->SetParameters(Context);
 		}
 	
 		// Draw a quad mapping scene color to the view's render target
 		DrawRectangle(
+			Context.RHICmdList,
 			0, 0,
 			SrcRect.Width(), SrcRect.Height(),
 			SrcRect.Min.X, SrcRect.Min.Y, 
@@ -509,7 +514,7 @@ void FRCPassPostProcessTemporalAA::Process(FRenderingCompositePassContext& Conte
 
 	{	// Responsive feedback for tagged pixels
 		// Draw to pixels where stencil != 0
-		RHISetDepthStencilState(TStaticDepthStencilState<false,CF_Always,true,CF_NotEqual,SO_Keep,SO_Keep,SO_Keep>::GetRHI(), 0);
+		Context.RHICmdList.SetDepthStencilState(TStaticDepthStencilState<false,CF_Always,true,CF_NotEqual,SO_Keep,SO_Keep,SO_Keep>::GetRHI(), 0);
 		
 		TShaderMapRef< FPostProcessTonemapVS >			VertexShader( GetGlobalShaderMap() );
 		if(bUseFast)
@@ -517,26 +522,29 @@ void FRCPassPostProcessTemporalAA::Process(FRenderingCompositePassContext& Conte
 			TShaderMapRef< FPostProcessTemporalAAPS<4,1> >	PixelShader( GetGlobalShaderMap() );
 
 			static FGlobalBoundShaderState BoundShaderState;
+			Context.RHICmdList.CheckIsNull(); // need new approach for "static FGlobalBoundShaderState" for parallel rendering
 
-			SetGlobalBoundShaderState(BoundShaderState, GFilterVertexDeclaration.VertexDeclarationRHI, *VertexShader, *PixelShader);
+			SetGlobalBoundShaderState(Context.RHICmdList, BoundShaderState, GFilterVertexDeclaration.VertexDeclarationRHI, *VertexShader, *PixelShader);
 
-			VertexShader->SetVS(RHICmdList, Context);
-			PixelShader->SetParameters(RHICmdList, Context);
+			VertexShader->SetVS(Context);
+			PixelShader->SetParameters(Context);
 		}
 		else
 		{
 			TShaderMapRef< FPostProcessTemporalAAPS<1,1> >	PixelShader( GetGlobalShaderMap() );
 
 			static FGlobalBoundShaderState BoundShaderState;
+			Context.RHICmdList.CheckIsNull(); // need new approach for "static FGlobalBoundShaderState" for parallel rendering
 
-			SetGlobalBoundShaderState(BoundShaderState, GFilterVertexDeclaration.VertexDeclarationRHI, *VertexShader, *PixelShader);
+			SetGlobalBoundShaderState(Context.RHICmdList, BoundShaderState, GFilterVertexDeclaration.VertexDeclarationRHI, *VertexShader, *PixelShader);
 
-			VertexShader->SetVS(RHICmdList, Context);
-			PixelShader->SetParameters(RHICmdList, Context);
+			VertexShader->SetVS(Context);
+			PixelShader->SetParameters(Context);
 		}
 
 		// Draw a quad mapping scene color to the view's render target
 		DrawRectangle(
+			Context.RHICmdList,
 			0, 0,
 			SrcRect.Width(), SrcRect.Height(),
 			SrcRect.Min.X, SrcRect.Min.Y, 
@@ -547,7 +555,9 @@ void FRCPassPostProcessTemporalAA::Process(FRenderingCompositePassContext& Conte
 			EDRF_UseTriangleOptimization);
 	}
 
-	RHICopyToResolveTarget(DestRenderTarget.TargetableTexture, DestRenderTarget.ShaderResourceTexture, false, FResolveParams());
+	Context.RHICmdList.CopyToResolveTarget(DestRenderTarget.TargetableTexture, DestRenderTarget.ShaderResourceTexture, false, FResolveParams());
+	Context.RHICmdList.CheckIsNull(); // this modifes the viewstate which is probably not appropriate for parallel rendering
+
 	ViewState->TemporalAAHistoryRT = PassOutputs[0].PooledRenderTarget;
 	check( ViewState->TemporalAAHistoryRT );
 
diff --git a/Engine/Source/Runtime/Renderer/Private/PostProcess/PostProcessTestImage.cpp b/Engine/Source/Runtime/Renderer/Private/PostProcess/PostProcessTestImage.cpp
--- a/Engine/Source/Runtime/Renderer/Private/PostProcess/PostProcessTestImage.cpp
+++ b/Engine/Source/Runtime/Renderer/Private/PostProcess/PostProcessTestImage.cpp
@@ -50,26 +50,26 @@ public:
 		FrameTime.Bind(Initializer.ParameterMap,TEXT("FrameTime"));
 	}
 
-	void SetPS(FRHICommandList& RHICmdList, const FRenderingCompositePassContext& Context)
+	void SetPS(const FRenderingCompositePassContext& Context)
 	{
 		const FPixelShaderRHIParamRef ShaderRHI = GetPixelShader();
 		
-		FGlobalShader::SetParameters(RHICmdList, ShaderRHI, Context.View);
+		FGlobalShader::SetParameters(Context.RHICmdList, ShaderRHI, Context.View);
 
-		PostprocessParameter.SetPS(RHICmdList, ShaderRHI, Context, TStaticSamplerState<SF_Point,AM_Clamp,AM_Clamp,AM_Clamp>::GetRHI());
-		DeferredParameters.Set(RHICmdList, ShaderRHI, Context.View);
+		PostprocessParameter.SetPS(ShaderRHI, Context, TStaticSamplerState<SF_Point,AM_Clamp,AM_Clamp,AM_Clamp>::GetRHI());
+		DeferredParameters.Set(Context.RHICmdList, ShaderRHI, Context.View);
 
 		{
 			uint32 FrameNumberValue = Context.View.FrameNumber;
-			SetShaderValue(RHICmdList, ShaderRHI, FrameNumber, FrameNumberValue);
+			SetShaderValue(Context.RHICmdList, ShaderRHI, FrameNumber, FrameNumberValue);
 		}
 
 		{
 			float FrameTimeValue = Context.View.Family->CurrentRealTime;
-			SetShaderValue(RHICmdList, ShaderRHI, FrameTime, FrameTimeValue);
+			SetShaderValue(Context.RHICmdList, ShaderRHI, FrameTime, FrameTimeValue);
 		}
 
-		ColorRemapShaderParameters.Set(RHICmdList, ShaderRHI);
+		ColorRemapShaderParameters.Set(Context.RHICmdList, ShaderRHI);
 	}
 	
 	// FShader interface.
@@ -99,29 +99,28 @@ void FRCPassPostProcessTestImage::Process(FRenderingCompositePassContext& Contex
 
 	const FSceneRenderTargetItem& DestRenderTarget = PassOutputs[0].RequestSurface(Context);
 
-	//@todo-rco: RHIPacketList
-	FRHICommandList& RHICmdList = FRHICommandList::GetNullRef();
-
 	// Set the view family's render target/viewport.
-	RHISetRenderTarget(DestRenderTarget.TargetableTexture, FTextureRHIRef());	
+	SetRenderTarget(Context.RHICmdList, DestRenderTarget.TargetableTexture, FTextureRHIRef());
 	Context.SetViewportAndCallRHI(DestRect);
 
 	// set the state
-	RHISetBlendState(TStaticBlendState<>::GetRHI());
-	RHISetRasterizerState(TStaticRasterizerState<>::GetRHI());
-	RHISetDepthStencilState(TStaticDepthStencilState<false,CF_Always>::GetRHI());
+	Context.RHICmdList.SetBlendState(TStaticBlendState<>::GetRHI());
+	Context.RHICmdList.SetRasterizerState(TStaticRasterizerState<>::GetRHI());
+	Context.RHICmdList.SetDepthStencilState(TStaticDepthStencilState<false, CF_Always>::GetRHI());
 
 	TShaderMapRef<FPostProcessVS> VertexShader(GetGlobalShaderMap());
 	TShaderMapRef<FPostProcessTestImagePS> PixelShader(GetGlobalShaderMap());
 
 	static FGlobalBoundShaderState BoundShaderState;
+	Context.RHICmdList.CheckIsNull(); // need new approach for "static FGlobalBoundShaderState" for parallel rendering
 
-	SetGlobalBoundShaderState(BoundShaderState, GFilterVertexDeclaration.VertexDeclarationRHI, *VertexShader, *PixelShader);
+	SetGlobalBoundShaderState(Context.RHICmdList, BoundShaderState, GFilterVertexDeclaration.VertexDeclarationRHI, *VertexShader, *PixelShader);
 
-	PixelShader->SetPS(RHICmdList, Context);
+	PixelShader->SetPS(Context);
 
 	// Draw a quad mapping scene color to the view's render target
 	DrawRectangle(
+		Context.RHICmdList,
 		0, 0,
 		DestRect.Width(), DestRect.Height(),
 		SrcRect.Min.X, SrcRect.Min.Y,
@@ -152,6 +151,7 @@ void FRCPassPostProcessTestImage::Process(FRenderingCompositePassContext& Contex
 				return Texture;
 			}
 		} TempRenderTarget(View, (const FTexture2DRHIRef&)DestRenderTarget.TargetableTexture);
+		Context.RHICmdList.CheckIsNull(); // canvas doesn't use the cmd list
 
 		FCanvas Canvas(&TempRenderTarget, NULL, ViewFamily.CurrentRealTime, ViewFamily.CurrentWorldTime, ViewFamily.DeltaWorldTime);
 
@@ -199,7 +199,7 @@ void FRCPassPostProcessTestImage::Process(FRenderingCompositePassContext& Contex
 	}
 
 
-	RHICopyToResolveTarget(DestRenderTarget.TargetableTexture, DestRenderTarget.ShaderResourceTexture, false, FResolveParams());
+	Context.RHICmdList.CopyToResolveTarget(DestRenderTarget.TargetableTexture, DestRenderTarget.ShaderResourceTexture, false, FResolveParams());
 }
 
 FPooledRenderTargetDesc FRCPassPostProcessTestImage::ComputeOutputDesc(EPassOutputId InPassOutputId) const
diff --git a/Engine/Source/Runtime/Renderer/Private/PostProcess/PostProcessTonemap.cpp b/Engine/Source/Runtime/Renderer/Private/PostProcess/PostProcessTonemap.cpp
--- a/Engine/Source/Runtime/Renderer/Private/PostProcess/PostProcessTonemap.cpp
+++ b/Engine/Source/Runtime/Renderer/Private/PostProcess/PostProcessTonemap.cpp
@@ -825,33 +825,33 @@ public:
 		return bShaderHasOutdatedParameters;
 	}
 
-	void SetPS(FRHICommandList& RHICmdList, const FRenderingCompositePassContext& Context)
+	void SetPS(const FRenderingCompositePassContext& Context)
 	{
 		const FPostProcessSettings& Settings = Context.View.FinalPostProcessSettings;
 		const FSceneViewFamily& ViewFamily = *(Context.View.Family);
 
 		const FPixelShaderRHIParamRef ShaderRHI = GetPixelShader();
 		
-		FGlobalShader::SetParameters(RHICmdList, ShaderRHI, Context.View);
+		FGlobalShader::SetParameters(Context.RHICmdList, ShaderRHI, Context.View);
 
-		PostprocessParameter.SetPS(RHICmdList, ShaderRHI, Context, TStaticSamplerState<SF_Bilinear,AM_Clamp,AM_Clamp,AM_Clamp>::GetRHI());
+		PostprocessParameter.SetPS(ShaderRHI, Context, TStaticSamplerState<SF_Bilinear,AM_Clamp,AM_Clamp,AM_Clamp>::GetRHI());
 			
 		{
 			FLinearColor Col = Settings.SceneColorTint;
 			FVector4 ColorScale(Col.R, Col.G, Col.B, 0);
-			SetShaderValue(RHICmdList, ShaderRHI, ColorScale0, ColorScale);
+			SetShaderValue(Context.RHICmdList, ShaderRHI, ColorScale0, ColorScale);
 		}
 		
 		{
 			FLinearColor Col = FLinearColor::White * Settings.BloomIntensity;
 			FVector4 ColorScale(Col.R, Col.G, Col.B, 0);
-			SetShaderValue(RHICmdList, ShaderRHI, ColorScale1, ColorScale);
+			SetShaderValue(Context.RHICmdList, ShaderRHI, ColorScale1, ColorScale);
 		}
 
 		{
 			UTexture2D* NoiseTextureValue = GEngine->HighFrequencyNoiseTexture;
 
-			SetTextureParameter(RHICmdList, ShaderRHI, NoiseTexture, NoiseTextureSampler, TStaticSamplerState<SF_Point,AM_Wrap,AM_Wrap,AM_Wrap>::GetRHI(), NoiseTextureValue->Resource->TextureRHI);
+			SetTextureParameter(Context.RHICmdList, ShaderRHI, NoiseTexture, NoiseTextureSampler, TStaticSamplerState<SF_Point, AM_Wrap, AM_Wrap, AM_Wrap>::GetRHI(), NoiseTextureValue->Resource->TextureRHI);
 		}
 
 		{
@@ -860,7 +860,7 @@ public:
 			// we assume the this pass runs in 1:1 pixel
 			FVector2D TexScaleValue = FVector2D(InputDesc->Extent) / FVector2D(Context.View.ViewRect.Size());
 
-			SetShaderValue(RHICmdList, ShaderRHI, TexScale, TexScaleValue);
+			SetShaderValue(Context.RHICmdList, ShaderRHI, TexScale, TexScaleValue);
 		}
 		
 		FVector4 VignetteColorIntensityValue;
@@ -868,11 +868,11 @@ public:
 		VignetteColorIntensityValue.Y = Settings.VignetteColor.G;
 		VignetteColorIntensityValue.Z = Settings.VignetteColor.B;
 		VignetteColorIntensityValue.W = Settings.VignetteIntensity;
-		SetShaderValue(RHICmdList, ShaderRHI, VignetteColorIntensity, VignetteColorIntensityValue);
+		SetShaderValue(Context.RHICmdList, ShaderRHI, VignetteColorIntensity, VignetteColorIntensityValue);
 
 		FVector GrainValue;
 		GrainPostSettings(&GrainValue, &Settings);
-		SetShaderValue(RHICmdList, ShaderRHI, GrainScaleBiasJitter, GrainValue);
+		SetShaderValue(Context.RHICmdList, ShaderRHI, GrainScaleBiasJitter, GrainValue);
 
 		const TShaderUniformBufferParameter<FBloomDirtMaskParameters>& BloomDirtMaskParam = GetUniformBufferParameter<FBloomDirtMaskParameters>();
 		if (BloomDirtMaskParam.IsBound())
@@ -891,7 +891,7 @@ public:
 			BloomDirtMaskParams.MaskSampler = TStaticSamplerState<SF_Bilinear,AM_Wrap,AM_Wrap,AM_Wrap>::GetRHI();
 
 			FUniformBufferRHIRef BloomDirtMaskUB = TUniformBufferRef<FBloomDirtMaskParameters>::CreateUniformBufferImmediate(BloomDirtMaskParams, UniformBuffer_SingleDraw);
-			SetUniformBufferParameter(RHICmdList, ShaderRHI, BloomDirtMaskParam, BloomDirtMaskUB);
+			SetUniformBufferParameter(Context.RHICmdList, ShaderRHI, BloomDirtMaskParam, BloomDirtMaskUB);
 		}
 		
 		// volume texture LUT
@@ -912,7 +912,7 @@ public:
 
 						const FTextureRHIRef& SrcTexture = InputPooledElement->GetRenderTargetItem().ShaderResourceTexture;
 					
-						SetTextureParameter(RHICmdList, ShaderRHI, ColorGradingLUT, ColorGradingLUTSampler, TStaticSamplerState<SF_Bilinear,AM_Clamp,AM_Clamp,AM_Clamp>::GetRHI(), SrcTexture);
+						SetTextureParameter(Context.RHICmdList, ShaderRHI, ColorGradingLUT, ColorGradingLUTSampler, TStaticSamplerState<SF_Bilinear, AM_Clamp, AM_Clamp, AM_Clamp>::GetRHI(), SrcTexture);
 					}
 				}
 			}
@@ -922,20 +922,20 @@ public:
 			FVector2D InvDisplayGammaValue;
 			InvDisplayGammaValue.X = 1.0f / ViewFamily.RenderTarget->GetDisplayGamma();
 			InvDisplayGammaValue.Y = 2.2f / ViewFamily.RenderTarget->GetDisplayGamma();
-			SetShaderValue(RHICmdList, ShaderRHI, InverseGamma, InvDisplayGammaValue);
+			SetShaderValue(Context.RHICmdList, ShaderRHI, InverseGamma, InvDisplayGammaValue);
 		}
 
 		{
 			FVector4 Constants[8];
 			FilmPostSetConstants(Constants, TonemapperConfBitmaskPC[ConfigIndex], &Context.View.FinalPostProcessSettings, false);
-			SetShaderValue(RHICmdList, ShaderRHI, ColorMatrixR_ColorCurveCd1, Constants[0]);
-			SetShaderValue(RHICmdList, ShaderRHI, ColorMatrixG_ColorCurveCd3Cm3, Constants[1]);
-			SetShaderValue(RHICmdList, ShaderRHI, ColorMatrixB_ColorCurveCm2, Constants[2]); 
-			SetShaderValue(RHICmdList, ShaderRHI, ColorCurve_Cm0Cd0_Cd2_Ch0Cm1_Ch3, Constants[3]); 
-			SetShaderValue(RHICmdList, ShaderRHI, ColorCurve_Ch1_Ch2, Constants[4]);
-			SetShaderValue(RHICmdList, ShaderRHI, ColorShadow_Luma, Constants[5]);
-			SetShaderValue(RHICmdList, ShaderRHI, ColorShadow_Tint1, Constants[6]);
-			SetShaderValue(RHICmdList, ShaderRHI, ColorShadow_Tint2, Constants[7]);
+			SetShaderValue(Context.RHICmdList, ShaderRHI, ColorMatrixR_ColorCurveCd1, Constants[0]);
+			SetShaderValue(Context.RHICmdList, ShaderRHI, ColorMatrixG_ColorCurveCd3Cm3, Constants[1]);
+			SetShaderValue(Context.RHICmdList, ShaderRHI, ColorMatrixB_ColorCurveCm2, Constants[2]); 
+			SetShaderValue(Context.RHICmdList, ShaderRHI, ColorCurve_Cm0Cd0_Cd2_Ch0Cm1_Ch3, Constants[3]); 
+			SetShaderValue(Context.RHICmdList, ShaderRHI, ColorCurve_Ch1_Ch2, Constants[4]);
+			SetShaderValue(Context.RHICmdList, ShaderRHI, ColorShadow_Luma, Constants[5]);
+			SetShaderValue(Context.RHICmdList, ShaderRHI, ColorShadow_Tint1, Constants[6]);
+			SetShaderValue(Context.RHICmdList, ShaderRHI, ColorShadow_Tint2, Constants[7]);
 		}
 	}
 	
@@ -969,17 +969,18 @@ FRCPassPostProcessTonemap::FRCPassPostProcessTonemap(bool bInDoGammaOnly)
 }
 
 template <uint32 ConfigIndex>
-static void SetShaderTempl(FRHICommandList& RHICmdList, const FRenderingCompositePassContext& Context)
+static void SetShaderTempl(const FRenderingCompositePassContext& Context)
 {
 	TShaderMapRef<FPostProcessTonemapVS> VertexShader(GetGlobalShaderMap());
 	TShaderMapRef<FPostProcessTonemapPS<ConfigIndex> > PixelShader(GetGlobalShaderMap());
 
 	static FGlobalBoundShaderState BoundShaderState;
+	Context.RHICmdList.CheckIsNull(); // need new approach for "static FGlobalBoundShaderState" for parallel rendering
 
-	SetGlobalBoundShaderState(BoundShaderState, GFilterVertexDeclaration.VertexDeclarationRHI, *VertexShader, *PixelShader);
+	SetGlobalBoundShaderState(Context.RHICmdList, BoundShaderState, GFilterVertexDeclaration.VertexDeclarationRHI, *VertexShader, *PixelShader);
 
-	VertexShader->SetVS(RHICmdList, Context);
-	PixelShader->SetPS(RHICmdList, Context);
+	VertexShader->SetVS(Context);
+	PixelShader->SetPS(Context);
 }
 
 void FRCPassPostProcessTonemap::Process(FRenderingCompositePassContext& Context)
@@ -1010,34 +1011,34 @@ void FRCPassPostProcessTonemap::Process(FRenderingCompositePassContext& Context)
 	FRHICommandList& RHICmdList = FRHICommandList::GetNullRef();
 
 	// Set the view family's render target/viewport.
-	RHISetRenderTarget(DestRenderTarget.TargetableTexture, FTextureRHIParamRef());	
+	SetRenderTarget(Context.RHICmdList, DestRenderTarget.TargetableTexture, FTextureRHIParamRef());
 
 	if( ViewFamily.RenderTarget->GetRenderTargetTexture() != DestRenderTarget.TargetableTexture )
 	{
 		// needed to not have PostProcessAA leaking in content (e.g. Matinee black borders), is optimized away if possible (RT size=view size, )
-		RHIClear(true, FLinearColor::Black, false, 1.0f, false, 0, View.ViewRect);
+		Context.RHICmdList.Clear(true, FLinearColor::Black, false, 1.0f, false, 0, View.ViewRect);
 	}
 
 	Context.SetViewportAndCallRHI(View.ViewRect);
 
 	// set the state
-	RHISetBlendState(TStaticBlendState<>::GetRHI());
-	RHISetRasterizerState(TStaticRasterizerState<>::GetRHI());
-	RHISetDepthStencilState(TStaticDepthStencilState<false,CF_Always>::GetRHI());
+	Context.RHICmdList.SetBlendState(TStaticBlendState<>::GetRHI());
+	Context.RHICmdList.SetRasterizerState(TStaticRasterizerState<>::GetRHI());
+	Context.RHICmdList.SetDepthStencilState(TStaticDepthStencilState<false, CF_Always>::GetRHI());
 
 	switch(ConfigIndex)
 	{
-		case 0:	SetShaderTempl<0>(RHICmdList, Context); break;
-		case 1:	SetShaderTempl<1>(RHICmdList, Context);	break;
-		case 2: SetShaderTempl<2>(RHICmdList, Context); break;
-		case 3: SetShaderTempl<3>(RHICmdList, Context); break;
-		case 4: SetShaderTempl<4>(RHICmdList, Context); break;
-		case 5: SetShaderTempl<5>(RHICmdList, Context); break;
-		case 6: SetShaderTempl<6>(RHICmdList, Context); break;
-		case 7: SetShaderTempl<7>(RHICmdList, Context); break;
-		case 8: SetShaderTempl<8>(RHICmdList, Context); break;
-		case 9: SetShaderTempl<9>(RHICmdList, Context); break;
-		case 10: SetShaderTempl<10>(RHICmdList, Context); break;
+		case 0:	SetShaderTempl<0>(Context); break;
+		case 1:	SetShaderTempl<1>(Context);	break;
+		case 2: SetShaderTempl<2>(Context); break;
+		case 3: SetShaderTempl<3>(Context); break;
+		case 4: SetShaderTempl<4>(Context); break;
+		case 5: SetShaderTempl<5>(Context); break;
+		case 6: SetShaderTempl<6>(Context); break;
+		case 7: SetShaderTempl<7>(Context); break;
+		case 8: SetShaderTempl<8>(Context); break;
+		case 9: SetShaderTempl<9>(Context); break;
+		case 10: SetShaderTempl<10>(Context); break;
 		default:
 			check(0);
 	}
@@ -1046,6 +1047,7 @@ void FRCPassPostProcessTonemap::Process(FRenderingCompositePassContext& Context)
 	TShaderMapRef<FPostProcessTonemapVS> VertexShader(GetGlobalShaderMap());
 
 	DrawRectangle(
+		Context.RHICmdList,
 		0, 0,
 		View.ViewRect.Width(), View.ViewRect.Height(),
 		View.ViewRect.Min.X, View.ViewRect.Min.Y, 
@@ -1055,11 +1057,12 @@ void FRCPassPostProcessTonemap::Process(FRenderingCompositePassContext& Context)
 		*VertexShader,
 		EDRF_UseTriangleOptimization);
 
-	RHICopyToResolveTarget(DestRenderTarget.TargetableTexture, DestRenderTarget.ShaderResourceTexture, false, FResolveParams());
+	Context.RHICmdList.CopyToResolveTarget(DestRenderTarget.TargetableTexture, DestRenderTarget.ShaderResourceTexture, false, FResolveParams());
 
 	// We only release the SceneColor after the last view was processed (SplitScreen)
 	if(Context.View.Family->Views[Context.View.Family->Views.Num() - 1] == &Context.View)
 	{
+		Context.RHICmdList.CheckIsNull(); // probably can't do this parallel
 		// The RT should be released as early as possible to allow sharing of that memory for other purposes.
 		// This becomes even more important with some limited VRam (XBoxOne).
 		GSceneRenderTargets.SetSceneColor(0);
@@ -1182,29 +1185,29 @@ public:
 		return bShaderHasOutdatedParameters;
 	}
 
-	void SetPS(FRHICommandList& RHICmdList, const FRenderingCompositePassContext& Context)
+	void SetPS(const FRenderingCompositePassContext& Context)
 	{
 		const FPostProcessSettings& Settings = Context.View.FinalPostProcessSettings;
 		const FSceneViewFamily& ViewFamily = *(Context.View.Family);
 
 		const FPixelShaderRHIParamRef ShaderRHI = GetPixelShader();
 		
-		FGlobalShader::SetParameters(RHICmdList, ShaderRHI, Context.View);
+		FGlobalShader::SetParameters(Context.RHICmdList, ShaderRHI, Context.View);
 
-		PostprocessParameter.SetPS(RHICmdList, ShaderRHI, Context, TStaticSamplerState<SF_Bilinear,AM_Clamp,AM_Clamp,AM_Clamp>::GetRHI());
+		PostprocessParameter.SetPS(ShaderRHI, Context, TStaticSamplerState<SF_Bilinear,AM_Clamp,AM_Clamp,AM_Clamp>::GetRHI());
 			
-		SetShaderValue(RHICmdList, ShaderRHI, OverlayColor, Context.View.OverlayColor);
+		SetShaderValue(Context.RHICmdList, ShaderRHI, OverlayColor, Context.View.OverlayColor);
 
 		{
 			FLinearColor Col = Settings.SceneColorTint;
 			FVector4 ColorScale(Col.R, Col.G, Col.B, 0);
-			SetShaderValue(RHICmdList, ShaderRHI, ColorScale0, ColorScale);
+			SetShaderValue(Context.RHICmdList, ShaderRHI, ColorScale0, ColorScale);
 		}
 		
 		{
 			FLinearColor Col = FLinearColor::White * Settings.BloomIntensity;
 			FVector4 ColorScale(Col.R, Col.G, Col.B, 0);
-			SetShaderValue(RHICmdList, ShaderRHI, ColorScale1, ColorScale);
+			SetShaderValue(Context.RHICmdList, ShaderRHI, ColorScale1, ColorScale);
 		}
 
 		{
@@ -1213,7 +1216,7 @@ public:
 			// we assume the this pass runs in 1:1 pixel
 			FVector2D TexScaleValue = FVector2D(InputDesc->Extent) / FVector2D(Context.View.ViewRect.Size());
 
-			SetShaderValue(RHICmdList, ShaderRHI, TexScale, TexScaleValue);
+			SetShaderValue(Context.RHICmdList, ShaderRHI, TexScale, TexScaleValue);
 		}
 
 		FVector4 VignetteColorIntensityValue;
@@ -1221,30 +1224,30 @@ public:
 		VignetteColorIntensityValue.Y = Settings.VignetteColor.G;
 		VignetteColorIntensityValue.Z = Settings.VignetteColor.B;
 		VignetteColorIntensityValue.W = Settings.VignetteIntensity;
-		SetShaderValue(RHICmdList, ShaderRHI, VignetteColorIntensity, VignetteColorIntensityValue);
+		SetShaderValue(Context.RHICmdList, ShaderRHI, VignetteColorIntensity, VignetteColorIntensityValue);
 
 		FVector GrainValue;
 		GrainPostSettings(&GrainValue, &Settings);
-		SetShaderValue(RHICmdList, ShaderRHI, GrainScaleBiasJitter, GrainValue);
+		SetShaderValue(Context.RHICmdList, ShaderRHI, GrainScaleBiasJitter, GrainValue);
 
 		{
 			FVector2D InvDisplayGammaValue;
 			InvDisplayGammaValue.X = 1.0f / ViewFamily.RenderTarget->GetDisplayGamma();
 			InvDisplayGammaValue.Y = 2.2f / ViewFamily.RenderTarget->GetDisplayGamma();
-			SetShaderValue(RHICmdList, ShaderRHI, InverseGamma, InvDisplayGammaValue);
+			SetShaderValue(Context.RHICmdList, ShaderRHI, InverseGamma, InvDisplayGammaValue);
 		}
 
 		{
 			FVector4 Constants[8];
 			FilmPostSetConstants(Constants, TonemapperConfBitmaskMobile[ConfigIndex], &Context.View.FinalPostProcessSettings, true);
-			SetShaderValue(RHICmdList, ShaderRHI, ColorMatrixR_ColorCurveCd1, Constants[0]);
-			SetShaderValue(RHICmdList, ShaderRHI, ColorMatrixG_ColorCurveCd3Cm3, Constants[1]);
-			SetShaderValue(RHICmdList, ShaderRHI, ColorMatrixB_ColorCurveCm2, Constants[2]); 
-			SetShaderValue(RHICmdList, ShaderRHI, ColorCurve_Cm0Cd0_Cd2_Ch0Cm1_Ch3, Constants[3]); 
-			SetShaderValue(RHICmdList, ShaderRHI, ColorCurve_Ch1_Ch2, Constants[4]);
-			SetShaderValue(RHICmdList, ShaderRHI, ColorShadow_Luma, Constants[5]);
-			SetShaderValue(RHICmdList, ShaderRHI, ColorShadow_Tint1, Constants[6]);
-			SetShaderValue(RHICmdList, ShaderRHI, ColorShadow_Tint2, Constants[7]);
+			SetShaderValue(Context.RHICmdList, ShaderRHI, ColorMatrixR_ColorCurveCd1, Constants[0]);
+			SetShaderValue(Context.RHICmdList, ShaderRHI, ColorMatrixG_ColorCurveCd3Cm3, Constants[1]);
+			SetShaderValue(Context.RHICmdList, ShaderRHI, ColorMatrixB_ColorCurveCm2, Constants[2]); 
+			SetShaderValue(Context.RHICmdList, ShaderRHI, ColorCurve_Cm0Cd0_Cd2_Ch0Cm1_Ch3, Constants[3]); 
+			SetShaderValue(Context.RHICmdList, ShaderRHI, ColorCurve_Ch1_Ch2, Constants[4]);
+			SetShaderValue(Context.RHICmdList, ShaderRHI, ColorShadow_Luma, Constants[5]);
+			SetShaderValue(Context.RHICmdList, ShaderRHI, ColorShadow_Tint1, Constants[6]);
+			SetShaderValue(Context.RHICmdList, ShaderRHI, ColorShadow_Tint2, Constants[7]);
 		}
 	}
 	
@@ -1294,18 +1297,18 @@ public:
 		GrainRandomFull.Bind(Initializer.ParameterMap, TEXT("GrainRandomFull"));
 	}
 
-	void SetVS(FRHICommandList& RHICmdList, const FRenderingCompositePassContext& Context)
+	void SetVS(const FRenderingCompositePassContext& Context)
 	{
 		const FVertexShaderRHIParamRef ShaderRHI = GetVertexShader();
-		FGlobalShader::SetParameters(RHICmdList, ShaderRHI, Context.View);
+		FGlobalShader::SetParameters(Context.RHICmdList, ShaderRHI, Context.View);
 
-		PostprocessParameter.SetVS(RHICmdList, ShaderRHI, Context, TStaticSamplerState<SF_Bilinear,AM_Clamp,AM_Clamp,AM_Clamp>::GetRHI());
+		PostprocessParameter.SetVS(ShaderRHI, Context, TStaticSamplerState<SF_Bilinear,AM_Clamp,AM_Clamp,AM_Clamp>::GetRHI());
 
 		FVector GrainRandomFullValue;
 		GrainRandomFromFrame(&GrainRandomFullValue, Context.View.FrameNumber);
 		// TODO: Don't use full on mobile with framebuffer fetch.
 		GrainRandomFullValue.Z = bUsedFramebufferFetch ? 0.0f : 1.0f;
-		SetShaderValue(RHICmdList, ShaderRHI, GrainRandomFull, GrainRandomFullValue);
+		SetShaderValue(Context.RHICmdList, ShaderRHI, GrainRandomFull, GrainRandomFullValue);
 	}
 	
 	virtual bool Serialize(FArchive& Ar)
@@ -1320,7 +1323,7 @@ IMPLEMENT_SHADER_TYPE(,FPostProcessTonemapVS_ES2,TEXT("PostProcessTonemap"),TEXT
 
 
 template <uint32 ConfigIndex>
-static void SetShaderTemplES2(FRHICommandList& RHICmdList, const FRenderingCompositePassContext& Context, bool bUsedFramebufferFetch)
+static void SetShaderTemplES2(const FRenderingCompositePassContext& Context, bool bUsedFramebufferFetch)
 {
 	TShaderMapRef<FPostProcessTonemapVS_ES2> VertexShader(GetGlobalShaderMap());
 	TShaderMapRef<FPostProcessTonemapPS_ES2<ConfigIndex> > PixelShader(GetGlobalShaderMap());
@@ -1328,11 +1331,12 @@ static void SetShaderTemplES2(FRHICommandList& RHICmdList, const FRenderingCompo
 	VertexShader->bUsedFramebufferFetch = bUsedFramebufferFetch;
 
 	static FGlobalBoundShaderState BoundShaderState;
+	Context.RHICmdList.CheckIsNull(); // need new approach for "static FGlobalBoundShaderState" for parallel rendering
 
-	SetGlobalBoundShaderState(BoundShaderState, GFilterVertexDeclaration.VertexDeclarationRHI, *VertexShader, *PixelShader);
+	SetGlobalBoundShaderState(Context.RHICmdList, BoundShaderState, GFilterVertexDeclaration.VertexDeclarationRHI, *VertexShader, *PixelShader);
 
-	VertexShader->SetVS(RHICmdList, Context);
-	PixelShader->SetPS(RHICmdList, Context);
+	VertexShader->SetVS(Context);
+	PixelShader->SetPS(Context);
 }
 
 void FRCPassPostProcessTonemapES2::Process(FRenderingCompositePassContext& Context)
@@ -1359,53 +1363,50 @@ void FRCPassPostProcessTonemapES2::Process(FRenderingCompositePassContext& Conte
 
 	const FSceneRenderTargetItem& DestRenderTarget = PassOutputs[0].RequestSurface(Context);
 
-	//@todo-rco: RHIPacketList
-	FRHICommandList& RHICmdList = FRHICommandList::GetNullRef();
-
 	// Set the view family's render target/viewport.
-	RHISetRenderTarget(DestRenderTarget.TargetableTexture, FTextureRHIParamRef());	
+	SetRenderTarget(Context.RHICmdList, DestRenderTarget.TargetableTexture, FTextureRHIParamRef());
 
 	// Full clear to avoid restore
-	RHIClear(true, FLinearColor::Black, false, 1.0f, false, 0, FIntRect());
+	Context.RHICmdList.Clear(true, FLinearColor::Black, false, 1.0f, false, 0, FIntRect());
 
 	Context.SetViewportAndCallRHI(View.ViewRect);
 
 	// set the state
-	RHISetBlendState(TStaticBlendState<>::GetRHI());
-	RHISetRasterizerState(TStaticRasterizerState<>::GetRHI());
-	RHISetDepthStencilState(TStaticDepthStencilState<false,CF_Always>::GetRHI());
+	Context.RHICmdList.SetBlendState(TStaticBlendState<>::GetRHI());
+	Context.RHICmdList.SetRasterizerState(TStaticRasterizerState<>::GetRHI());
+	Context.RHICmdList.SetDepthStencilState(TStaticDepthStencilState<false, CF_Always>::GetRHI());
 
 	switch(ConfigIndex)
 	{
-		case 0:	SetShaderTemplES2<0>(RHICmdList, Context, bUsedFramebufferFetch); break;
-		case 1:	SetShaderTemplES2<1>(RHICmdList, Context, bUsedFramebufferFetch); break;
-		case 2:	SetShaderTemplES2<2>(RHICmdList, Context, bUsedFramebufferFetch); break;
-		case 3:	SetShaderTemplES2<3>(RHICmdList, Context, bUsedFramebufferFetch); break;
-		case 4:	SetShaderTemplES2<4>(RHICmdList, Context, bUsedFramebufferFetch); break;
-		case 5:	SetShaderTemplES2<5>(RHICmdList, Context, bUsedFramebufferFetch); break;
-		case 6:	SetShaderTemplES2<6>(RHICmdList, Context, bUsedFramebufferFetch); break;
-		case 7:	SetShaderTemplES2<7>(RHICmdList, Context, bUsedFramebufferFetch); break;
-		case 8:	SetShaderTemplES2<8>(RHICmdList, Context, bUsedFramebufferFetch); break;
-		case 9:	SetShaderTemplES2<9>(RHICmdList, Context, bUsedFramebufferFetch); break;
-		case 10: SetShaderTemplES2<10>(RHICmdList, Context, bUsedFramebufferFetch); break;
-		case 11: SetShaderTemplES2<11>(RHICmdList, Context, bUsedFramebufferFetch); break;
-		case 12: SetShaderTemplES2<12>(RHICmdList, Context, bUsedFramebufferFetch); break;
-		case 13: SetShaderTemplES2<13>(RHICmdList, Context, bUsedFramebufferFetch); break;
-		case 14: SetShaderTemplES2<14>(RHICmdList, Context, bUsedFramebufferFetch); break;
-		case 15: SetShaderTemplES2<15>(RHICmdList, Context, bUsedFramebufferFetch); break;
-		case 16: SetShaderTemplES2<16>(RHICmdList, Context, bUsedFramebufferFetch); break;
-		case 17: SetShaderTemplES2<17>(RHICmdList, Context, bUsedFramebufferFetch); break;
-		case 18: SetShaderTemplES2<18>(RHICmdList, Context, bUsedFramebufferFetch); break;
-		case 19: SetShaderTemplES2<19>(RHICmdList, Context, bUsedFramebufferFetch); break;
-		case 20: SetShaderTemplES2<20>(RHICmdList, Context, bUsedFramebufferFetch); break;
-		case 21: SetShaderTemplES2<21>(RHICmdList, Context, bUsedFramebufferFetch); break;
-		case 22: SetShaderTemplES2<22>(RHICmdList, Context, bUsedFramebufferFetch); break;
-		case 23: SetShaderTemplES2<23>(RHICmdList, Context, bUsedFramebufferFetch); break;
-		case 24: SetShaderTemplES2<24>(RHICmdList, Context, bUsedFramebufferFetch); break;
-		case 25: SetShaderTemplES2<25>(RHICmdList, Context, bUsedFramebufferFetch); break;
-		case 26: SetShaderTemplES2<26>(RHICmdList, Context, bUsedFramebufferFetch); break;
-		case 27: SetShaderTemplES2<27>(RHICmdList, Context, bUsedFramebufferFetch); break;
-		case 28: SetShaderTemplES2<28>(RHICmdList, Context, bUsedFramebufferFetch); break;
+		case 0:	SetShaderTemplES2<0>(Context, bUsedFramebufferFetch); break;
+		case 1:	SetShaderTemplES2<1>(Context, bUsedFramebufferFetch); break;
+		case 2:	SetShaderTemplES2<2>(Context, bUsedFramebufferFetch); break;
+		case 3:	SetShaderTemplES2<3>(Context, bUsedFramebufferFetch); break;
+		case 4:	SetShaderTemplES2<4>(Context, bUsedFramebufferFetch); break;
+		case 5:	SetShaderTemplES2<5>(Context, bUsedFramebufferFetch); break;
+		case 6:	SetShaderTemplES2<6>(Context, bUsedFramebufferFetch); break;
+		case 7:	SetShaderTemplES2<7>(Context, bUsedFramebufferFetch); break;
+		case 8:	SetShaderTemplES2<8>(Context, bUsedFramebufferFetch); break;
+		case 9:	SetShaderTemplES2<9>(Context, bUsedFramebufferFetch); break;
+		case 10: SetShaderTemplES2<10>(Context, bUsedFramebufferFetch); break;
+		case 11: SetShaderTemplES2<11>(Context, bUsedFramebufferFetch); break;
+		case 12: SetShaderTemplES2<12>(Context, bUsedFramebufferFetch); break;
+		case 13: SetShaderTemplES2<13>(Context, bUsedFramebufferFetch); break;
+		case 14: SetShaderTemplES2<14>(Context, bUsedFramebufferFetch); break;
+		case 15: SetShaderTemplES2<15>(Context, bUsedFramebufferFetch); break;
+		case 16: SetShaderTemplES2<16>(Context, bUsedFramebufferFetch); break;
+		case 17: SetShaderTemplES2<17>(Context, bUsedFramebufferFetch); break;
+		case 18: SetShaderTemplES2<18>(Context, bUsedFramebufferFetch); break;
+		case 19: SetShaderTemplES2<19>(Context, bUsedFramebufferFetch); break;
+		case 20: SetShaderTemplES2<20>(Context, bUsedFramebufferFetch); break;
+		case 21: SetShaderTemplES2<21>(Context, bUsedFramebufferFetch); break;
+		case 22: SetShaderTemplES2<22>(Context, bUsedFramebufferFetch); break;
+		case 23: SetShaderTemplES2<23>(Context, bUsedFramebufferFetch); break;
+		case 24: SetShaderTemplES2<24>(Context, bUsedFramebufferFetch); break;
+		case 25: SetShaderTemplES2<25>(Context, bUsedFramebufferFetch); break;
+		case 26: SetShaderTemplES2<26>(Context, bUsedFramebufferFetch); break;
+		case 27: SetShaderTemplES2<27>(Context, bUsedFramebufferFetch); break;
+		case 28: SetShaderTemplES2<28>(Context, bUsedFramebufferFetch); break;
 		default:
 			check(0);
 	}
@@ -1414,6 +1415,7 @@ void FRCPassPostProcessTonemapES2::Process(FRenderingCompositePassContext& Conte
 	TShaderMapRef<FPostProcessTonemapVS_ES2> VertexShader(GetGlobalShaderMap());
 
 	DrawRectangle(
+		Context.RHICmdList,
 		0, 0,
 		View.ViewRect.Width(), View.ViewRect.Height(),
 		View.ViewRect.Min.X, View.ViewRect.Min.Y, 
@@ -1423,11 +1425,13 @@ void FRCPassPostProcessTonemapES2::Process(FRenderingCompositePassContext& Conte
 		*VertexShader,
 		EDRF_UseTriangleOptimization);
 
-	RHICopyToResolveTarget(DestRenderTarget.TargetableTexture, DestRenderTarget.ShaderResourceTexture, false, FResolveParams());
+	Context.RHICmdList.CopyToResolveTarget(DestRenderTarget.TargetableTexture, DestRenderTarget.ShaderResourceTexture, false, FResolveParams());
 
 	// Double buffer tonemapper output for temporal AA.
 	if(Context.View.FinalPostProcessSettings.AntiAliasingMethod == AAM_TemporalAA)
 	{
+		Context.RHICmdList.CheckIsNull(); // modifies the viewstate, probably not appropriate for parallel work
+
 		FSceneViewState* ViewState = (FSceneViewState*)Context.View.State;
 		if(ViewState) 
 		{
diff --git a/Engine/Source/Runtime/Renderer/Private/PostProcess/PostProcessTonemap.h b/Engine/Source/Runtime/Renderer/Private/PostProcess/PostProcessTonemap.h
--- a/Engine/Source/Runtime/Renderer/Private/PostProcess/PostProcessTonemap.h
+++ b/Engine/Source/Runtime/Renderer/Private/PostProcess/PostProcessTonemap.h
@@ -100,17 +100,17 @@ public:
 		GrainRandomFull.Bind(Initializer.ParameterMap, TEXT("GrainRandomFull"));
 	}
 
-	void SetVS(FRHICommandList& RHICmdList, const FRenderingCompositePassContext& Context)
+	void SetVS(const FRenderingCompositePassContext& Context)
 	{
 		const FVertexShaderRHIParamRef ShaderRHI = GetVertexShader();
 
-		FGlobalShader::SetParameters(RHICmdList, ShaderRHI, Context.View);
+		FGlobalShader::SetParameters(Context.RHICmdList, ShaderRHI, Context.View);
 
-		PostprocessParameter.SetVS(RHICmdList, ShaderRHI, Context, TStaticSamplerState<SF_Bilinear,AM_Clamp,AM_Clamp,AM_Clamp>::GetRHI());
+		PostprocessParameter.SetVS(ShaderRHI, Context, TStaticSamplerState<SF_Bilinear,AM_Clamp,AM_Clamp,AM_Clamp>::GetRHI());
 
 		FVector GrainRandomFullValue;
 		GrainRandomFromFrame(&GrainRandomFullValue, Context.View.FrameNumber);
-		SetShaderValue(RHICmdList, ShaderRHI, GrainRandomFull, GrainRandomFullValue);
+		SetShaderValue(Context.RHICmdList, ShaderRHI, GrainRandomFull, GrainRandomFullValue);
 
 		if(EyeAdaptation.IsBound())
 		{
@@ -118,12 +118,12 @@ public:
 
 			if(EyeAdaptationRT)
 			{
-				SetTextureParameter(RHICmdList, ShaderRHI, EyeAdaptation, EyeAdaptationRT->GetRenderTargetItem().TargetableTexture);
+				SetTextureParameter(Context.RHICmdList, ShaderRHI, EyeAdaptation, EyeAdaptationRT->GetRenderTargetItem().TargetableTexture);
 			}
 			else
 			{
 				// some views don't have a state, thumbnail rendering?
-				SetTextureParameter(RHICmdList, ShaderRHI, EyeAdaptation, GWhiteTexture->TextureRHI);
+				SetTextureParameter(Context.RHICmdList, ShaderRHI, EyeAdaptation, GWhiteTexture->TextureRHI);
 			}
 		}
 	}
