UE4 - more work preparing for parallel rendering. Merged the cmd list into the post process context, among other things.

diff --git a/Engine/Source/Runtime/Renderer/Private/PostProcess/PostProcessUIBlur.cpp b/Engine/Source/Runtime/Renderer/Private/PostProcess/PostProcessUIBlur.cpp
--- a/Engine/Source/Runtime/Renderer/Private/PostProcess/PostProcessUIBlur.cpp
+++ b/Engine/Source/Runtime/Renderer/Private/PostProcess/PostProcessUIBlur.cpp
@@ -34,13 +34,13 @@ public:
 		PostprocessParameter.Bind(Initializer.ParameterMap);
 	}
 
-	void SetPS(FRHICommandList& RHICmdList, const FRenderingCompositePassContext& Context)
+	void SetPS(const FRenderingCompositePassContext& Context)
 	{
 		const FPixelShaderRHIParamRef ShaderRHI = GetPixelShader();
 
-		FGlobalShader::SetParameters(RHICmdList, ShaderRHI, Context.View);
+		FGlobalShader::SetParameters(Context.RHICmdList, ShaderRHI, Context.View);
 
-		PostprocessParameter.SetPS(RHICmdList, ShaderRHI, Context, TStaticSamplerState<SF_Bilinear,AM_Clamp,AM_Clamp,AM_Clamp>::GetRHI());
+		PostprocessParameter.SetPS(ShaderRHI, Context, TStaticSamplerState<SF_Bilinear,AM_Clamp,AM_Clamp,AM_Clamp>::GetRHI());
 	}
 	
 	// FShader interface.
@@ -76,25 +76,24 @@ void FRCPassPostProcessUIBlur::Process(FRenderingCompositePassContext& Context)
 	TRefCountPtr<IPooledRenderTarget>& Dest = PassInputs[0].GetOutput()->PooledRenderTarget;
 	const FSceneRenderTargetItem& DestRenderTarget = Dest->GetRenderTargetItem();
 
-	//@todo-rco: RHIPacketList
-	FRHICommandList& RHICmdList = FRHICommandList::GetNullRef();
 
 	// Set the view family's render target/viewport.
-	RHISetRenderTarget(DestRenderTarget.TargetableTexture, FTextureRHIRef());	
+	SetRenderTarget(Context.RHICmdList, DestRenderTarget.TargetableTexture, FTextureRHIRef());
 
 	// set the state
-	RHISetBlendState(TStaticBlendState<>::GetRHI());
-	RHISetRasterizerState(TStaticRasterizerState<>::GetRHI());
-	RHISetDepthStencilState(TStaticDepthStencilState<false,CF_Always>::GetRHI());
+	Context.RHICmdList.SetBlendState(TStaticBlendState<>::GetRHI());
+	Context.RHICmdList.SetRasterizerState(TStaticRasterizerState<>::GetRHI());
+	Context.RHICmdList.SetDepthStencilState(TStaticDepthStencilState<false, CF_Always>::GetRHI());
 
 	TShaderMapRef<FPostProcessVS> VertexShader(GetGlobalShaderMap());
 	TShaderMapRef<FPostProcessUIBlurPS> PixelShader(GetGlobalShaderMap());
 
 	static FGlobalBoundShaderState BoundShaderState;
+	Context.RHICmdList.CheckIsNull(); // need new approach for "static FGlobalBoundShaderState" for parallel rendering
 
-	SetGlobalBoundShaderState(BoundShaderState, GFilterVertexDeclaration.VertexDeclarationRHI, *VertexShader, *PixelShader);
+	SetGlobalBoundShaderState(Context.RHICmdList, BoundShaderState, GFilterVertexDeclaration.VertexDeclarationRHI, *VertexShader, *PixelShader);
 
-	PixelShader->SetPS(RHICmdList, Context);
+	PixelShader->SetPS(Context);
 
 	// todo: we only expect few quads but ideally we have one one draw call for all quads
 	int32 RectCount = Context.View.UIBlurOverrideRectangles.Num();
@@ -111,6 +110,7 @@ void FRCPassPostProcessUIBlur::Process(FRenderingCompositePassContext& Context)
 		DestRect = DestRect.Scale(UpsampleScale);
 
 		DrawRectangle(
+			Context.RHICmdList,
 			DestRect.Min.X, DestRect.Min.Y,
 			DestRect.Width(), DestRect.Height(),
 			SrcRect.Min.X, SrcRect.Min.Y,
@@ -121,7 +121,7 @@ void FRCPassPostProcessUIBlur::Process(FRenderingCompositePassContext& Context)
 			EDRF_Default);
 	}
 
-	RHICopyToResolveTarget(DestRenderTarget.TargetableTexture, DestRenderTarget.ShaderResourceTexture, false, FResolveParams());
+	Context.RHICmdList.CopyToResolveTarget(DestRenderTarget.TargetableTexture, DestRenderTarget.ShaderResourceTexture, false, FResolveParams());
 	PassOutputs[0].PooledRenderTarget = BlurOutputRef.PooledRenderTarget;
 	PassOutputs[0].RenderTargetDesc = BlurOutputRef.RenderTargetDesc;
 }
diff --git a/Engine/Source/Runtime/Renderer/Private/PostProcess/PostProcessUpscale.cpp b/Engine/Source/Runtime/Renderer/Private/PostProcess/PostProcessUpscale.cpp
--- a/Engine/Source/Runtime/Renderer/Private/PostProcess/PostProcessUpscale.cpp
+++ b/Engine/Source/Runtime/Renderer/Private/PostProcess/PostProcessUpscale.cpp
@@ -47,18 +47,18 @@ public:
 		UpscaleSoftness.Bind(Initializer.ParameterMap,TEXT("UpscaleSoftness"));
 	}
 
-	void SetPS(FRHICommandList& RHICmdList, const FRenderingCompositePassContext& Context)
+	void SetPS(const FRenderingCompositePassContext& Context)
 	{
 		const FPixelShaderRHIParamRef ShaderRHI = GetPixelShader();
 		
-		FGlobalShader::SetParameters(RHICmdList, ShaderRHI, Context.View);
+		FGlobalShader::SetParameters(Context.RHICmdList, ShaderRHI, Context.View);
 
 		FSamplerStateRHIParamRef FilterTable[2];
 		FilterTable[0] = TStaticSamplerState<SF_Bilinear,AM_Clamp,AM_Clamp,AM_Clamp>::GetRHI();
 		FilterTable[1] = TStaticSamplerState<SF_Point,AM_Clamp,AM_Clamp,AM_Clamp>::GetRHI();
 			
-		PostprocessParameter.SetPS(RHICmdList, ShaderRHI, Context, 0, false, FilterTable);
-		DeferredParameters.Set(RHICmdList, ShaderRHI, Context.View);
+		PostprocessParameter.SetPS(ShaderRHI, Context, 0, false, FilterTable);
+		DeferredParameters.Set(Context.RHICmdList, ShaderRHI, Context.View);
 
 		// If the method needs softness value
 		if(Method == 2)
@@ -67,7 +67,7 @@ public:
 
 			float UpscaleSoftnessValue = FMath::Clamp(CVar->GetValueOnRenderThread(), 0.0f, 1.0f);
 
-			SetShaderValue(RHICmdList, ShaderRHI, UpscaleSoftness, UpscaleSoftnessValue);
+			SetShaderValue(Context.RHICmdList, ShaderRHI, UpscaleSoftness, UpscaleSoftnessValue);
 		}
 	}
 	
@@ -108,16 +108,17 @@ FRCPassPostProcessUpscale::FRCPassPostProcessUpscale(uint32 InUpscaleMethod)
 }
 
 template <uint32 Method>
-void FRCPassPostProcessUpscale::SetShader(FRHICommandList& RHICmdList, const FRenderingCompositePassContext& Context)
+void FRCPassPostProcessUpscale::SetShader(const FRenderingCompositePassContext& Context)
 {
 	TShaderMapRef<FPostProcessVS> VertexShader(GetGlobalShaderMap());
 	TShaderMapRef<FPostProcessUpscalePS<Method> > PixelShader(GetGlobalShaderMap());
 
 	static FGlobalBoundShaderState BoundShaderState;
+	Context.RHICmdList.CheckIsNull(); // need new approach for "static FGlobalBoundShaderState" for parallel rendering
 
-	SetGlobalBoundShaderState(BoundShaderState, GFilterVertexDeclaration.VertexDeclarationRHI, *VertexShader, *PixelShader);
+	SetGlobalBoundShaderState(Context.RHICmdList, BoundShaderState, GFilterVertexDeclaration.VertexDeclarationRHI, *VertexShader, *PixelShader);
 
-	PixelShader->SetPS(RHICmdList, Context);
+	PixelShader->SetPS(Context);
 }
 
 void FRCPassPostProcessUpscale::Process(FRenderingCompositePassContext& Context)
@@ -140,31 +141,28 @@ void FRCPassPostProcessUpscale::Process(FRenderingCompositePassContext& Context)
 
 	const FSceneRenderTargetItem& DestRenderTarget = PassOutputs[0].RequestSurface(Context);
 
-	//@todo-rco: RHIPacketList
-	FRHICommandList& RHICmdList = FRHICommandList::GetNullRef();
-
 	// Set the view family's render target/viewport.
-	RHISetRenderTarget(DestRenderTarget.TargetableTexture, FTextureRHIRef());	
+	SetRenderTarget(Context.RHICmdList, DestRenderTarget.TargetableTexture, FTextureRHIRef());
 	Context.SetViewportAndCallRHI(DestRect);
 
 	// set the state
-	RHISetBlendState(TStaticBlendState<>::GetRHI());
-	RHISetRasterizerState(TStaticRasterizerState<>::GetRHI());
-	RHISetDepthStencilState(TStaticDepthStencilState<false,CF_Always>::GetRHI());
+	Context.RHICmdList.SetBlendState(TStaticBlendState<>::GetRHI());
+	Context.RHICmdList.SetRasterizerState(TStaticRasterizerState<>::GetRHI());
+	Context.RHICmdList.SetDepthStencilState(TStaticDepthStencilState<false, CF_Always>::GetRHI());
 
 	switch (UpscaleMethod)
 	{
 		case 0:
-			SetShader<0>(RHICmdList, Context);
+			SetShader<0>(Context);
 		break;
 		case 1:
-			SetShader<1>(RHICmdList, Context);
+			SetShader<1>(Context);
 		break;
 		case 2:
-			SetShader<2>(RHICmdList, Context);
+			SetShader<2>(Context);
 		break;
 		case 3:
-			SetShader<3>(RHICmdList, Context);
+			SetShader<3>(Context);
 		break;
 		default:
 			checkNoEntry();
@@ -174,6 +172,7 @@ void FRCPassPostProcessUpscale::Process(FRenderingCompositePassContext& Context)
 	// Draw a quad mapping scene color to the view's render target
 	TShaderMapRef<FPostProcessVS> VertexShader(GetGlobalShaderMap());
 	DrawRectangle(
+		Context.RHICmdList,
 		0, 0,
 		DestRect.Width(), DestRect.Height(),
 		SrcRect.Min.X, SrcRect.Min.Y,
@@ -183,7 +182,7 @@ void FRCPassPostProcessUpscale::Process(FRenderingCompositePassContext& Context)
 		*VertexShader,
 		EDRF_UseTriangleOptimization);
 
-	RHICopyToResolveTarget(DestRenderTarget.TargetableTexture, DestRenderTarget.ShaderResourceTexture, false, FResolveParams());
+	Context.RHICmdList.CopyToResolveTarget(DestRenderTarget.TargetableTexture, DestRenderTarget.ShaderResourceTexture, false, FResolveParams());
 }
 
 FPooledRenderTargetDesc FRCPassPostProcessUpscale::ComputeOutputDesc(EPassOutputId InPassOutputId) const
diff --git a/Engine/Source/Runtime/Renderer/Private/PostProcess/PostProcessUpscale.h b/Engine/Source/Runtime/Renderer/Private/PostProcess/PostProcessUpscale.h
--- a/Engine/Source/Runtime/Renderer/Private/PostProcess/PostProcessUpscale.h
+++ b/Engine/Source/Runtime/Renderer/Private/PostProcess/PostProcessUpscale.h
@@ -27,7 +27,7 @@ public:
 	virtual void Release() OVERRIDE { delete this; }
 	FPooledRenderTargetDesc ComputeOutputDesc(EPassOutputId InPassOutputId) const;
 private:
-	template <uint32 Method> static void SetShader(FRHICommandList& RHICmdList, const FRenderingCompositePassContext& Context);
+	template <uint32 Method> static void SetShader(const FRenderingCompositePassContext& Context);
 	uint32 UpscaleMethod;
 };
 
diff --git a/Engine/Source/Runtime/Renderer/Private/PostProcess/PostProcessVisualizeBuffer.cpp b/Engine/Source/Runtime/Renderer/Private/PostProcess/PostProcessVisualizeBuffer.cpp
--- a/Engine/Source/Runtime/Renderer/Private/PostProcess/PostProcessVisualizeBuffer.cpp
+++ b/Engine/Source/Runtime/Renderer/Private/PostProcess/PostProcessVisualizeBuffer.cpp
@@ -53,13 +53,13 @@ public:
 		}
 	}
 
-	void SetPS(FRHICommandList& RHICmdList, const FRenderingCompositePassContext& Context)
+	void SetPS(const FRenderingCompositePassContext& Context)
 	{
 		const FPixelShaderRHIParamRef ShaderRHI = GetPixelShader();
 
-		FGlobalShader::SetParameters(RHICmdList, ShaderRHI, Context.View);
-		PostprocessParameter.SetPS(RHICmdList, ShaderRHI, Context, TStaticSamplerState<SF_Point,AM_Clamp,AM_Clamp,AM_Clamp>::GetRHI());
-		DeferredParameters.Set(RHICmdList, ShaderRHI, Context.View);
+		FGlobalShader::SetParameters(Context.RHICmdList, ShaderRHI, Context.View);
+		PostprocessParameter.SetPS(ShaderRHI, Context, TStaticSamplerState<SF_Point,AM_Clamp,AM_Clamp,AM_Clamp>::GetRHI());
+		DeferredParameters.Set(Context.RHICmdList, ShaderRHI, Context.View);
 	}
 
 	void SetSourceTexture(FRHICommandList& RHICmdList, FTextureRHIRef Texture)
@@ -111,16 +111,17 @@ IMPLEMENT_SHADER_TYPE2(FPostProcessVisualizeBufferPS<true>, SF_Pixel);
 IMPLEMENT_SHADER_TYPE2(FPostProcessVisualizeBufferPS<false>, SF_Pixel);
 
 template <bool bDrawingTile>
-FShader* FRCPassPostProcessVisualizeBuffer::SetShaderTempl(FRHICommandList& RHICmdList, const FRenderingCompositePassContext& Context)
+FShader* FRCPassPostProcessVisualizeBuffer::SetShaderTempl(const FRenderingCompositePassContext& Context)
 {
 	TShaderMapRef<FPostProcessVS> VertexShader(GetGlobalShaderMap());
 	TShaderMapRef<FPostProcessVisualizeBufferPS<bDrawingTile> > PixelShader(GetGlobalShaderMap());
 
 	static FGlobalBoundShaderState BoundShaderState;
+	Context.RHICmdList.CheckIsNull(); // need new approach for "static FGlobalBoundShaderState" for parallel rendering
 
-	SetGlobalBoundShaderState(BoundShaderState, GFilterVertexDeclaration.VertexDeclarationRHI, *VertexShader, *PixelShader);
+	SetGlobalBoundShaderState(Context.RHICmdList, BoundShaderState, GFilterVertexDeclaration.VertexDeclarationRHI, *VertexShader, *PixelShader);
 
-	PixelShader->SetPS(RHICmdList, Context);
+	PixelShader->SetPS(Context);
 
 	return *VertexShader;
 }
@@ -145,23 +146,21 @@ void FRCPassPostProcessVisualizeBuffer::Process(FRenderingCompositePassContext&
 
 	const FSceneRenderTargetItem& DestRenderTarget = PassOutputs[0].RequestSurface(Context);
 
-	//@todo-rco: RHIPacketList
-	FRHICommandList& RHICmdList = FRHICommandList::GetNullRef();
-
 	// Set the view family's render target/viewport.
-	RHISetRenderTarget(DestRenderTarget.TargetableTexture, FTextureRHIRef());	
+	SetRenderTarget(Context.RHICmdList, DestRenderTarget.TargetableTexture, FTextureRHIRef());
 	Context.SetViewportAndCallRHI(DestRect);
 
 	// set the state
-	RHISetBlendState(TStaticBlendState<>::GetRHI());
-	RHISetRasterizerState(TStaticRasterizerState<>::GetRHI());
-	RHISetDepthStencilState(TStaticDepthStencilState<false,CF_Always>::GetRHI());
+	Context.RHICmdList.SetBlendState(TStaticBlendState<>::GetRHI());
+	Context.RHICmdList.SetRasterizerState(TStaticRasterizerState<>::GetRHI());
+	Context.RHICmdList.SetDepthStencilState(TStaticDepthStencilState<false, CF_Always>::GetRHI());
 
 	{
-		FShader* VertexShader = SetShaderTempl<false>(RHICmdList, Context);
+		FShader* VertexShader = SetShaderTempl<false>(Context);
 
 	// Draw a quad mapping scene color to the view's render target
 	DrawRectangle(
+		Context.RHICmdList,
 		0, 0,
 		DestRect.Width(), DestRect.Height(),
 		SrcRect.Min.X, SrcRect.Min.Y,
@@ -172,15 +171,16 @@ void FRCPassPostProcessVisualizeBuffer::Process(FRenderingCompositePassContext&
 		EDRF_UseTriangleOptimization);
 	}
 
-	RHISetBlendState(TStaticBlendState<CW_RGB, BO_Add, BF_SourceAlpha, BF_InverseSourceAlpha>::GetRHI());
+	Context.RHICmdList.SetBlendState(TStaticBlendState<CW_RGB, BO_Add, BF_SourceAlpha, BF_InverseSourceAlpha>::GetRHI());
 
 	TShaderMapRef<FPostProcessVS> VertexShader(GetGlobalShaderMap());
 	TShaderMapRef<FPostProcessVisualizeBufferPS<true> > PixelShader(GetGlobalShaderMap());
 
 	static FGlobalBoundShaderState BoundShaderState;
-	SetGlobalBoundShaderState(BoundShaderState, GFilterVertexDeclaration.VertexDeclarationRHI, *VertexShader, *PixelShader);
+	Context.RHICmdList.CheckIsNull(); // need new approach for "static FGlobalBoundShaderState" for parallel rendering
+	SetGlobalBoundShaderState(Context.RHICmdList, BoundShaderState, GFilterVertexDeclaration.VertexDeclarationRHI, *VertexShader, *PixelShader);
 
-	PixelShader->SetPS(RHICmdList, Context);
+	PixelShader->SetPS(Context);
 
 	// Track the name and position of each tile we draw so we can write text labels over them
 	struct LabelRecord
@@ -209,9 +209,10 @@ void FRCPassPostProcessVisualizeBuffer::Process(FRenderingCompositePassContext&
 			int32 TileX = CurrentTileIndex % MaxTilesX;
 			int32 TileY = CurrentTileIndex / MaxTilesX;
 
-			PixelShader->SetSourceTexture(RHICmdList, Texture);
+			PixelShader->SetSourceTexture(Context.RHICmdList, Texture);
 
 			DrawRectangle(
+				Context.RHICmdList,
 				TileX * TileWidth, TileY * TileHeight,
 				TileWidth, TileHeight,
 				SrcRect.Min.X, SrcRect.Min.Y,
@@ -251,6 +252,7 @@ void FRCPassPostProcessVisualizeBuffer::Process(FRenderingCompositePassContext&
 			return Texture;
 		}
 	} TempRenderTarget(View, (const FTexture2DRHIRef&)DestRenderTarget.TargetableTexture);
+	Context.RHICmdList.CheckIsNull(); // canvas doesn't support cmd lists
 
 	FCanvas Canvas(&TempRenderTarget, NULL, ViewFamily.CurrentRealTime, ViewFamily.CurrentWorldTime, ViewFamily.DeltaWorldTime);
 	FLinearColor LabelColor(1, 1, 0);
@@ -261,7 +263,7 @@ void FRCPassPostProcessVisualizeBuffer::Process(FRenderingCompositePassContext&
 	Canvas.Flush();
 
 
-	RHICopyToResolveTarget(DestRenderTarget.TargetableTexture, DestRenderTarget.ShaderResourceTexture, false, FResolveParams());
+	Context.RHICmdList.CopyToResolveTarget(DestRenderTarget.TargetableTexture, DestRenderTarget.ShaderResourceTexture, false, FResolveParams());
 }
 
 FPooledRenderTargetDesc FRCPassPostProcessVisualizeBuffer::ComputeOutputDesc(EPassOutputId InPassOutputId) const
diff --git a/Engine/Source/Runtime/Renderer/Private/PostProcess/PostProcessVisualizeBuffer.h b/Engine/Source/Runtime/Renderer/Private/PostProcess/PostProcessVisualizeBuffer.h
--- a/Engine/Source/Runtime/Renderer/Private/PostProcess/PostProcessVisualizeBuffer.h
+++ b/Engine/Source/Runtime/Renderer/Private/PostProcess/PostProcessVisualizeBuffer.h
@@ -42,7 +42,7 @@ private:
 
 	// @return VertexShader
 	template <bool bDrawTileTemplate>
-	FShader* SetShaderTempl(FRHICommandList& RHICmdList, const FRenderingCompositePassContext& Context);
+	FShader* SetShaderTempl(const FRenderingCompositePassContext& Context);
 
 	TArray<TileData> Tiles;
 };
diff --git a/Engine/Source/Runtime/Renderer/Private/PostProcess/PostProcessVisualizeHDR.cpp b/Engine/Source/Runtime/Renderer/Private/PostProcess/PostProcessVisualizeHDR.cpp
--- a/Engine/Source/Runtime/Renderer/Private/PostProcess/PostProcessVisualizeHDR.cpp
+++ b/Engine/Source/Runtime/Renderer/Private/PostProcess/PostProcessVisualizeHDR.cpp
@@ -69,43 +69,43 @@ public:
 		ColorShadow_Tint2.Bind(Initializer.ParameterMap, TEXT("ColorShadow_Tint2"));
 	}
 
-	void SetPS(FRHICommandList& RHICmdList, const FRenderingCompositePassContext& Context)
+	void SetPS(const FRenderingCompositePassContext& Context)
 	{
 		const FPixelShaderRHIParamRef ShaderRHI = GetPixelShader();
 		
-		FGlobalShader::SetParameters(RHICmdList, ShaderRHI, Context.View);
+		FGlobalShader::SetParameters(Context.RHICmdList, ShaderRHI, Context.View);
 
 		const FPostProcessSettings& Settings = Context.View.FinalPostProcessSettings;
 		const FSceneViewFamily& ViewFamily = *(Context.View.Family);
 
-		PostprocessParameter.SetPS(RHICmdList, ShaderRHI, Context, TStaticSamplerState<SF_Point,AM_Clamp,AM_Clamp,AM_Clamp>::GetRHI());
+		PostprocessParameter.SetPS(ShaderRHI, Context, TStaticSamplerState<SF_Point,AM_Clamp,AM_Clamp,AM_Clamp>::GetRHI());
 
 		{
 			FVector4 Temp[3];
 
 			FRCPassPostProcessEyeAdaptation::ComputeEyeAdaptationParamsValue(Context.View, Temp);
-			SetShaderValueArray(RHICmdList, ShaderRHI, EyeAdaptationParams, Temp, 3);
+			SetShaderValueArray(Context.RHICmdList, ShaderRHI, EyeAdaptationParams, Temp, 3);
 		}
 
-		SetTextureParameter(RHICmdList, ShaderRHI, MiniFontTexture, GEngine->MiniFontTexture ? GEngine->MiniFontTexture->Resource->TextureRHI : GSystemTextures.WhiteDummy->GetRenderTargetItem().TargetableTexture);
+		SetTextureParameter(Context.RHICmdList, ShaderRHI, MiniFontTexture, GEngine->MiniFontTexture ? GEngine->MiniFontTexture->Resource->TextureRHI : GSystemTextures.WhiteDummy->GetRenderTargetItem().TargetableTexture);
 
 		{
 			float InvDisplayGammaValue = 1.0f / ViewFamily.RenderTarget->GetDisplayGamma();
 
-			SetShaderValue(RHICmdList, ShaderRHI, InverseGamma, InvDisplayGammaValue);
+			SetShaderValue(Context.RHICmdList, ShaderRHI, InverseGamma, InvDisplayGammaValue);
 		}
 
 		{
 			FVector4 Constants[8];
 			FilmPostSetConstants(Constants, ~0, &Context.View.FinalPostProcessSettings, false);
-			SetShaderValue(RHICmdList, ShaderRHI, ColorMatrixR_ColorCurveCd1, Constants[0]);
-			SetShaderValue(RHICmdList, ShaderRHI, ColorMatrixG_ColorCurveCd3Cm3, Constants[1]);
-			SetShaderValue(RHICmdList, ShaderRHI, ColorMatrixB_ColorCurveCm2, Constants[2]); 
-			SetShaderValue(RHICmdList, ShaderRHI, ColorCurve_Cm0Cd0_Cd2_Ch0Cm1_Ch3, Constants[3]); 
-			SetShaderValue(RHICmdList, ShaderRHI, ColorCurve_Ch1_Ch2, Constants[4]);
-			SetShaderValue(RHICmdList, ShaderRHI, ColorShadow_Luma, Constants[5]);
-			SetShaderValue(RHICmdList, ShaderRHI, ColorShadow_Tint1, Constants[6]);
-			SetShaderValue(RHICmdList, ShaderRHI, ColorShadow_Tint2, Constants[7]);
+			SetShaderValue(Context.RHICmdList, ShaderRHI, ColorMatrixR_ColorCurveCd1, Constants[0]);
+			SetShaderValue(Context.RHICmdList, ShaderRHI, ColorMatrixG_ColorCurveCd3Cm3, Constants[1]);
+			SetShaderValue(Context.RHICmdList, ShaderRHI, ColorMatrixB_ColorCurveCm2, Constants[2]);
+			SetShaderValue(Context.RHICmdList, ShaderRHI, ColorCurve_Cm0Cd0_Cd2_Ch0Cm1_Ch3, Constants[3]);
+			SetShaderValue(Context.RHICmdList, ShaderRHI, ColorCurve_Ch1_Ch2, Constants[4]);
+			SetShaderValue(Context.RHICmdList, ShaderRHI, ColorShadow_Luma, Constants[5]);
+			SetShaderValue(Context.RHICmdList, ShaderRHI, ColorShadow_Tint1, Constants[6]);
+			SetShaderValue(Context.RHICmdList, ShaderRHI, ColorShadow_Tint2, Constants[7]);
 		}
 	}
 	
@@ -153,29 +153,28 @@ void FRCPassPostProcessVisualizeHDR::Process(FRenderingCompositePassContext& Con
 
 	const FSceneRenderTargetItem& DestRenderTarget = PassOutputs[0].RequestSurface(Context);
 
-	//@todo-rco: RHIPacketList
-	FRHICommandList& RHICmdList = FRHICommandList::GetNullRef();
-
 	// Set the view family's render target/viewport.
-	RHISetRenderTarget(DestRenderTarget.TargetableTexture, FTextureRHIRef());	
+	SetRenderTarget(Context.RHICmdList, DestRenderTarget.TargetableTexture, FTextureRHIRef());
 	Context.SetViewportAndCallRHI(DestRect);
 
 	// set the state
-	RHISetBlendState(TStaticBlendState<>::GetRHI());
-	RHISetRasterizerState(TStaticRasterizerState<>::GetRHI());
-	RHISetDepthStencilState(TStaticDepthStencilState<false,CF_Always>::GetRHI());
+	Context.RHICmdList.SetBlendState(TStaticBlendState<>::GetRHI());
+	Context.RHICmdList.SetRasterizerState(TStaticRasterizerState<>::GetRHI());
+	Context.RHICmdList.SetDepthStencilState(TStaticDepthStencilState<false, CF_Always>::GetRHI());
 
 	TShaderMapRef<FPostProcessVS> VertexShader(GetGlobalShaderMap());
 	TShaderMapRef<FPostProcessVisualizeHDRPS> PixelShader(GetGlobalShaderMap());
 
 	static FGlobalBoundShaderState BoundShaderState;
+	Context.RHICmdList.CheckIsNull(); // need new approach for "static FGlobalBoundShaderState" for parallel rendering
 
-	SetGlobalBoundShaderState(BoundShaderState, GFilterVertexDeclaration.VertexDeclarationRHI, *VertexShader, *PixelShader);
+	SetGlobalBoundShaderState(Context.RHICmdList, BoundShaderState, GFilterVertexDeclaration.VertexDeclarationRHI, *VertexShader, *PixelShader);
 
-	PixelShader->SetPS(RHICmdList, Context);
+	PixelShader->SetPS(Context);
 
 	// Draw a quad mapping scene color to the view's render target
 	DrawRectangle(
+		Context.RHICmdList,
 		DestRect.Min.X, DestRect.Min.Y,
 		DestRect.Width(), DestRect.Height(),
 		SrcRect.Min.X, SrcRect.Min.Y,
@@ -206,6 +205,7 @@ void FRCPassPostProcessVisualizeHDR::Process(FRenderingCompositePassContext& Con
 			return Texture;
 		}
 	} TempRenderTarget(View, (const FTexture2DRHIRef&)DestRenderTarget.TargetableTexture);
+	Context.RHICmdList.CheckIsNull(); // canvas doesn't use the command list
 
 	FCanvas Canvas(&TempRenderTarget, NULL, ViewFamily.CurrentRealTime, ViewFamily.CurrentWorldTime, ViewFamily.DeltaWorldTime);
 
@@ -264,7 +264,7 @@ void FRCPassPostProcessVisualizeHDR::Process(FRenderingCompositePassContext& Con
 
 	Canvas.Flush();
 
-	RHICopyToResolveTarget(DestRenderTarget.TargetableTexture, DestRenderTarget.ShaderResourceTexture, false, FResolveParams());
+	Context.RHICmdList.CopyToResolveTarget(DestRenderTarget.TargetableTexture, DestRenderTarget.ShaderResourceTexture, false, FResolveParams());
 }
 
 FPooledRenderTargetDesc FRCPassPostProcessVisualizeHDR::ComputeOutputDesc(EPassOutputId InPassOutputId) const
diff --git a/Engine/Source/Runtime/Renderer/Private/PostProcess/PostProcessWeightedSampleSum.cpp b/Engine/Source/Runtime/Renderer/Private/PostProcess/PostProcessWeightedSampleSum.cpp
--- a/Engine/Source/Runtime/Renderer/Private/PostProcess/PostProcessWeightedSampleSum.cpp
+++ b/Engine/Source/Runtime/Renderer/Private/PostProcess/PostProcessWeightedSampleSum.cpp
@@ -170,7 +170,7 @@ void SetFilterShaders(
 	)
 {
 	check(CombineMethodInt <= 2);
-	RHICmdList.CheckIsNull();
+	RHICmdList.CheckIsNull(); // need new approach for "static FGlobalBoundShaderState" for parallel rendering
 
 	// A macro to handle setting the filter shader for a specific number of samples.
 #define SET_FILTER_SHADER_TYPE(NumSamples) \
@@ -183,7 +183,7 @@ void SetFilterShaders(
 			TShaderMapRef<TFilterPS<NumSamples, 0> > PixelShader(GetGlobalShaderMap()); \
 			{ \
 				static FGlobalBoundShaderState BoundShaderState; \
-				SetGlobalBoundShaderState( BoundShaderState, GFilterVertexDeclaration.VertexDeclarationRHI, *VertexShader, *PixelShader); \
+				SetGlobalBoundShaderState(RHICmdList, BoundShaderState, GFilterVertexDeclaration.VertexDeclarationRHI, *VertexShader, *PixelShader); \
 			} \
 			PixelShader->SetParameters(RHICmdList, SamplerStateRHI, FilterTextureRHI, AdditiveTextureRHI, SampleWeights); \
 		} \
@@ -192,7 +192,7 @@ void SetFilterShaders(
 			TShaderMapRef<TFilterPS<NumSamples, 1> > PixelShader(GetGlobalShaderMap()); \
 			{ \
 				static FGlobalBoundShaderState BoundShaderState; \
-				SetGlobalBoundShaderState( BoundShaderState, GFilterVertexDeclaration.VertexDeclarationRHI, *VertexShader, *PixelShader); \
+				SetGlobalBoundShaderState(RHICmdList, BoundShaderState, GFilterVertexDeclaration.VertexDeclarationRHI, *VertexShader, *PixelShader); \
 			} \
 			PixelShader->SetParameters(RHICmdList, SamplerStateRHI, FilterTextureRHI, AdditiveTextureRHI, SampleWeights); \
 		} \
@@ -201,7 +201,7 @@ void SetFilterShaders(
 			TShaderMapRef<TFilterPS<NumSamples, 2> > PixelShader(GetGlobalShaderMap()); \
 			{ \
 				static FGlobalBoundShaderState BoundShaderState; \
-				SetGlobalBoundShaderState( BoundShaderState, GFilterVertexDeclaration.VertexDeclarationRHI, *VertexShader, *PixelShader); \
+				SetGlobalBoundShaderState(RHICmdList, BoundShaderState, GFilterVertexDeclaration.VertexDeclarationRHI, *VertexShader, *PixelShader); \
 			} \
 			PixelShader->SetParameters(RHICmdList, SamplerStateRHI, FilterTextureRHI, AdditiveTextureRHI, SampleWeights); \
 		} \
@@ -377,17 +377,14 @@ void FRCPassPostProcessWeightedSampleSum::Process(FRenderingCompositePassContext
 		BlurWeights[i] = TintValue * OffsetAndWeight[i].Y;
 	}
 
-	//@todo-rco: RHIPacketList
-	FRHICommandList& RHICmdList = FRHICommandList::GetNullRef();
-
-	RHISetRenderTarget(DestRenderTarget.TargetableTexture, FTextureRHIRef());
+	SetRenderTarget(Context.RHICmdList, DestRenderTarget.TargetableTexture, FTextureRHIRef());
 
 	Context.SetViewportAndCallRHI(0, 0, 0.0f, DestSize.X, DestSize.Y, 1.0f);
 
 	// set the state
-	RHISetBlendState(TStaticBlendState<>::GetRHI());
-	RHISetRasterizerState(TStaticRasterizerState<>::GetRHI());
-	RHISetDepthStencilState(TStaticDepthStencilState<false,CF_Always>::GetRHI());
+	Context.RHICmdList.SetBlendState(TStaticBlendState<>::GetRHI());
+	Context.RHICmdList.SetRasterizerState(TStaticRasterizerState<>::GetRHI());
+	Context.RHICmdList.SetDepthStencilState(TStaticDepthStencilState<false, CF_Always>::GetRHI());
 
 	const FTextureRHIRef& FilterTexture = InputPooledElement->GetRenderTargetItem().ShaderResourceTexture;
 
@@ -434,7 +431,7 @@ void FRCPassPostProcessWeightedSampleSum::Process(FRenderingCompositePassContext
 
 	FShader* VertexShader = nullptr;
 	SetFilterShaders(
-		RHICmdList, 
+		Context.RHICmdList,
 		TStaticSamplerState<SF_Bilinear,AM_Border,AM_Border,AM_Clamp>::GetRHI(),
 		FilterTexture,
 		AdditiveTexture,
@@ -452,7 +449,7 @@ void FRCPassPostProcessWeightedSampleSum::Process(FRenderingCompositePassContext
 	// Otherwise perform the clear when the dest rectangle has been computed.
 	if (bHasMultipleQuads || Context.View.GetFeatureLevel() == ERHIFeatureLevel::ES2)
 	{
-		RHIClear(true, FLinearColor(0, 0, 0, 0), false, 1.0f, false, 0, FIntRect());
+		Context.RHICmdList.Clear(true, FLinearColor(0, 0, 0, 0), false, 1.0f, false, 0, FIntRect());
 		bRequiresClear = false;
 	}
 
@@ -463,7 +460,7 @@ void FRCPassPostProcessWeightedSampleSum::Process(FRenderingCompositePassContext
 			FIntRect SrcRect =  View.ViewRect / SrcScaleFactor;
 			FIntRect DestRect = View.ViewRect / DstScaleFactor;
 
-			DrawQuad(bDoFastBlur, SrcRect, DestRect, bRequiresClear, DestSize, SrcSize, VertexShader);
+			DrawQuad(Context.RHICmdList, bDoFastBlur, SrcRect, DestRect, bRequiresClear, DestSize, SrcSize, VertexShader);
 		}
 		break;
 		case EPostProcessRectSource::GBS_UIBlurRects:
@@ -483,7 +480,7 @@ void FRCPassPostProcessWeightedSampleSum::Process(FRenderingCompositePassContext
 				FIntRect SrcRect = ScaledQuad / SrcScaleFactor;
 				FIntRect DestRect = ScaledQuad / DstScaleFactor;
 
-				DrawQuad(bDoFastBlur, SrcRect, DestRect, bRequiresClear, DestSize, SrcSize, VertexShader);
+				DrawQuad(Context.RHICmdList, bDoFastBlur, SrcRect, DestRect, bRequiresClear, DestSize, SrcSize, VertexShader);
 			}
 		}
 		break;
@@ -491,7 +488,7 @@ void FRCPassPostProcessWeightedSampleSum::Process(FRenderingCompositePassContext
 			checkNoEntry();
 		break;
 	}
-	RHICopyToResolveTarget(DestRenderTarget.TargetableTexture, DestRenderTarget.ShaderResourceTexture, false, FResolveParams());
+	Context.RHICmdList.CopyToResolveTarget(DestRenderTarget.TargetableTexture, DestRenderTarget.ShaderResourceTexture, false, FResolveParams());
 }
 
 FPooledRenderTargetDesc FRCPassPostProcessWeightedSampleSum::ComputeOutputDesc(EPassOutputId InPassOutputId) const
@@ -567,7 +564,7 @@ bool FRCPassPostProcessWeightedSampleSum::DoFastBlur() const
 	return bRet;
 }
 
-void FRCPassPostProcessWeightedSampleSum::DrawQuad( bool bDoFastBlur, FIntRect SrcRect, FIntRect DestRect, bool bRequiresClear, FIntPoint DestSize, FIntPoint SrcSize, FShader* VertexShader ) const
+void FRCPassPostProcessWeightedSampleSum::DrawQuad(FRHICommandList& RHICmdList, bool bDoFastBlur, FIntRect SrcRect, FIntRect DestRect, bool bRequiresClear, FIntPoint DestSize, FIntPoint SrcSize, FShader* VertexShader) const
 {
 	if (bDoFastBlur)
 	{
@@ -585,11 +582,12 @@ void FRCPassPostProcessWeightedSampleSum::DrawQuad( bool bDoFastBlur, FIntRect S
 
 	if (bRequiresClear)
 	{
-		RHIClear(true, FLinearColor(0, 0, 0, 0), false, 1.0f, false, 0, DestRect);
+		RHICmdList.Clear(true, FLinearColor(0, 0, 0, 0), false, 1.0f, false, 0, DestRect);
 	}
 
 	// Draw a quad mapping scene color to the view's render target
 	DrawRectangle(
+		RHICmdList,
 		DestRect.Min.X, DestRect.Min.Y,
 		DestRect.Width(), DestRect.Height(),
 		SrcRect.Min.X, SrcRect.Min.Y,
diff --git a/Engine/Source/Runtime/Renderer/Private/PostProcess/PostProcessWeightedSampleSum.h b/Engine/Source/Runtime/Renderer/Private/PostProcess/PostProcessWeightedSampleSum.h
--- a/Engine/Source/Runtime/Renderer/Private/PostProcess/PostProcessWeightedSampleSum.h
+++ b/Engine/Source/Runtime/Renderer/Private/PostProcess/PostProcessWeightedSampleSum.h
@@ -124,7 +124,7 @@ public:
 	static int GetIntegerKernelRadius(ERHIFeatureLevel::Type InFeatureLevel, float KernelRadius);
 
 private:
-	void DrawQuad(bool bDoFastBlur, FIntRect SrcRect, FIntRect DestRect, bool bRequiresClear, FIntPoint DestSize, FIntPoint SrcSize, FShader* VertexShader) const;
+	void DrawQuad(FRHICommandList& RHICmdList, bool bDoFastBlur, FIntRect SrcRect, FIntRect DestRect, bool bRequiresClear, FIntPoint DestSize, FIntPoint SrcSize, FShader* VertexShader) const;
 	static uint32 GetMaxNumSamples(ERHIFeatureLevel::Type InFeatureLevel);
 
 	// e.g. EFS_Horiz or EFS_Vert
diff --git a/Engine/Source/Runtime/Renderer/Private/PostProcess/PostProcessing.h b/Engine/Source/Runtime/Renderer/Private/PostProcess/PostProcessing.h
--- a/Engine/Source/Runtime/Renderer/Private/PostProcess/PostProcessing.h
+++ b/Engine/Source/Runtime/Renderer/Private/PostProcess/PostProcessing.h
@@ -40,9 +40,9 @@ class FPostProcessVS : public FGlobalShader
 	FPostProcessVS() {}
 
 	/** to have a similar interface as all other shaders */
-	void SetParameters(FRHICommandList& RHICmdList, const FRenderingCompositePassContext& Context)
+	void SetParameters(const FRenderingCompositePassContext& Context)
 	{
-		FGlobalShader::SetParameters(RHICmdList, GetVertexShader(), Context.View);
+		FGlobalShader::SetParameters(Context.RHICmdList, GetVertexShader(), Context.View);
 	}
 
 	void SetParameters(FRHICommandList& RHICmdList, const FSceneView& View)
diff --git a/Engine/Source/Runtime/Renderer/Private/PostProcess/RenderTargetPool.cpp b/Engine/Source/Runtime/Renderer/Private/PostProcess/RenderTargetPool.cpp
--- a/Engine/Source/Runtime/Renderer/Private/PostProcess/RenderTargetPool.cpp
+++ b/Engine/Source/Runtime/Renderer/Private/PostProcess/RenderTargetPool.cpp
@@ -307,23 +307,26 @@ bool FRenderTargetPool::FindFreeElement(const FPooledRenderTargetDesc& Desc, TRe
 	{
 		static const auto ICVar = IConsoleManager::Get().FindTConsoleVariableDataInt(TEXT("r.RenderTargetPoolTest"));
 		
+		//@todo-rco: RHIPacketList
+		FRHICommandList& RHICmdList = FRHICommandList::GetNullRef();
+
 		if(ICVar->GetValueOnRenderThread())
 		{
 			if(Found->GetDesc().TargetableFlags & TexCreate_RenderTargetable)
 			{
-				RHISetRenderTarget(Found->RenderTargetItem.TargetableTexture, FTextureRHIRef());	
-				RHIClear(true, FLinearColor(1000, 1000, 1000, 1000), false, 1.0f, false, 0, FIntRect());
+				SetRenderTarget(RHICmdList, Found->RenderTargetItem.TargetableTexture, FTextureRHIRef());
+				RHICmdList.Clear(true, FLinearColor(1000, 1000, 1000, 1000), false, 1.0f, false, 0, FIntRect());
 			}
 			else if(Found->GetDesc().TargetableFlags & TexCreate_UAV)
 			{
 				const uint32 ZeroClearValue[4] = { 1000, 1000, 1000, 1000 };
-				RHIClearUAV(Found->RenderTargetItem.UAV, ZeroClearValue);
+				RHICmdList.ClearUAV(Found->RenderTargetItem.UAV, ZeroClearValue);
 			}
 
 			if(Desc.TargetableFlags & TexCreate_DepthStencilTargetable)
 			{
-				RHISetRenderTarget(FTextureRHIRef(), Found->RenderTargetItem.TargetableTexture);
-				RHIClear(false, FLinearColor(0, 0, 0, 0), true, 0.0f, false, 0, FIntRect());
+				SetRenderTarget(RHICmdList, FTextureRHIRef(), Found->RenderTargetItem.TargetableTexture);
+				RHICmdList.Clear(false, FLinearColor(0, 0, 0, 0), true, 0.0f, false, 0, FIntRect());
 			}
 		}
 	}
@@ -676,7 +679,9 @@ inline void DrawBorder(FCanvas& Canvas, const FIntRect Rect, FLinearColor Color)
 
 void FRenderTargetPool::PresentContent(const FSceneView& View)
 {
-	if(RenderTargetPoolEvents.Num())
+	//@todo-rco: RHIPacketList
+	FRHICommandList& RHICmdList = FRHICommandList::GetNullRef();
+	if (RenderTargetPoolEvents.Num())
 	{
 		AddPhaseEvent(TEXT("FrameEnd"));
 
@@ -689,12 +694,12 @@ void FRenderTargetPool::PresentContent(const FSceneView& View)
 		{
 			SMemoryStats MemoryStats = ComputeView();
 
-			RHISetRenderTarget(View.Family->RenderTarget->GetRenderTargetTexture(),FTextureRHIRef());
-			RHISetViewport(0, 0, 0.0f, GSceneRenderTargets.GetBufferSizeXY().X, GSceneRenderTargets.GetBufferSizeXY().Y, 1.0f );
+			SetRenderTarget(RHICmdList, View.Family->RenderTarget->GetRenderTargetTexture(), FTextureRHIRef());
+			RHICmdList.SetViewport(0, 0, 0.0f, GSceneRenderTargets.GetBufferSizeXY().X, GSceneRenderTargets.GetBufferSizeXY().Y, 1.0f);
 
-			RHISetBlendState(TStaticBlendState<>::GetRHI());
-			RHISetRasterizerState(TStaticRasterizerState<>::GetRHI());
-			RHISetDepthStencilState(TStaticDepthStencilState<false,CF_Always>::GetRHI());
+			RHICmdList.SetBlendState(TStaticBlendState<>::GetRHI());
+			RHICmdList.SetRasterizerState(TStaticRasterizerState<>::GetRHI());
+			RHICmdList.SetDepthStencilState(TStaticDepthStencilState<false, CF_Always>::GetRHI());
 
 			// this is a helper class for FCanvas to be able to get screen size
 			class FRenderTargetTemp : public FRenderTarget
@@ -848,7 +853,7 @@ void FRenderTargetPool::PresentContent(const FSceneView& View)
 		}
 	}
 
-	VisualizeTexture.PresentContent(View);
+	VisualizeTexture.PresentContent(RHICmdList, View);
 }
 
 void FRenderTargetPool::AddDeallocEvents()
diff --git a/Engine/Source/Runtime/Renderer/Private/PostProcess/RenderingCompositionGraph.cpp b/Engine/Source/Runtime/Renderer/Private/PostProcess/RenderingCompositionGraph.cpp
--- a/Engine/Source/Runtime/Renderer/Private/PostProcess/RenderingCompositionGraph.cpp
+++ b/Engine/Source/Runtime/Renderer/Private/PostProcess/RenderingCompositionGraph.cpp
@@ -107,6 +107,7 @@ FRenderingCompositePassContext::FRenderingCompositePassContext(const FViewInfo&
 	: View(InView)
 	//		, CompositingOutputRTItem(InRenderTargetItem)
 	, Pass(0)
+	, RHICmdList(FRHICommandList::GetNullRef())
 	, ViewPortRect(0, 0, 0 ,0)
 {
 	check(!IsViewportValid());
@@ -526,7 +527,7 @@ void FRenderingCompositionGraph::RecursivelyProcess(const FRenderingCompositeOut
 			// use intermediate texture unless it's the last one where we render to the final output
 			if(PassOutput->PooledRenderTarget)
 			{
-				GRenderTargetPool.VisualizeTexture.SetCheckPoint(PassOutput->PooledRenderTarget);
+				GRenderTargetPool.VisualizeTexture.SetCheckPoint(Context.RHICmdList, PassOutput->PooledRenderTarget);
 
 				// If this buffer was given a dump filename, write it out
 				const FString& Filename = Pass->GetOutputDumpFilename((EPassOutputId)OutputId);
@@ -664,19 +665,19 @@ void FPostProcessPassParameters::Bind(const FShaderParameterMap& ParameterMap)
 	}
 }
 
-void FPostProcessPassParameters::SetPS(FRHICommandList& RHICmdList, const FPixelShaderRHIParamRef& ShaderRHI, const FRenderingCompositePassContext& Context, FSamplerStateRHIParamRef Filter, bool bWhiteIfNoTexture, FSamplerStateRHIParamRef* FilterOverrideArray)
+void FPostProcessPassParameters::SetPS(const FPixelShaderRHIParamRef& ShaderRHI, const FRenderingCompositePassContext& Context, FSamplerStateRHIParamRef Filter, bool bWhiteIfNoTexture, FSamplerStateRHIParamRef* FilterOverrideArray)
 {
-	Set(RHICmdList, ShaderRHI, Context, Filter, bWhiteIfNoTexture, FilterOverrideArray);
+	Set(Context.RHICmdList, ShaderRHI, Context, Filter, bWhiteIfNoTexture, FilterOverrideArray);
 }
 
-void FPostProcessPassParameters::SetCS(FRHICommandList& RHICmdList, const FComputeShaderRHIParamRef& ShaderRHI, const FRenderingCompositePassContext& Context, FSamplerStateRHIParamRef Filter, bool bWhiteIfNoTexture, FSamplerStateRHIParamRef* FilterOverrideArray)
+void FPostProcessPassParameters::SetCS(const FComputeShaderRHIParamRef& ShaderRHI, const FRenderingCompositePassContext& Context, FSamplerStateRHIParamRef Filter, bool bWhiteIfNoTexture, FSamplerStateRHIParamRef* FilterOverrideArray)
 {
-	Set(RHICmdList, ShaderRHI, Context, Filter, bWhiteIfNoTexture, FilterOverrideArray);
+	Set(Context.RHICmdList, ShaderRHI, Context, Filter, bWhiteIfNoTexture, FilterOverrideArray);
 }
 
-void FPostProcessPassParameters::SetVS(FRHICommandList& RHICmdList, const FVertexShaderRHIParamRef& ShaderRHI, const FRenderingCompositePassContext& Context, FSamplerStateRHIParamRef Filter, bool bWhiteIfNoTexture, FSamplerStateRHIParamRef* FilterOverrideArray)
+void FPostProcessPassParameters::SetVS(const FVertexShaderRHIParamRef& ShaderRHI, const FRenderingCompositePassContext& Context, FSamplerStateRHIParamRef Filter, bool bWhiteIfNoTexture, FSamplerStateRHIParamRef* FilterOverrideArray)
 {
-	Set(RHICmdList, ShaderRHI, Context, Filter, bWhiteIfNoTexture, FilterOverrideArray);
+	Set(Context.RHICmdList, ShaderRHI, Context, Filter, bWhiteIfNoTexture, FilterOverrideArray);
 }
 
 FArchive& operator<<(FArchive& Ar, FPostProcessPassParameters& P)
diff --git a/Engine/Source/Runtime/Renderer/Private/PostProcess/RenderingCompositionGraph.h b/Engine/Source/Runtime/Renderer/Private/PostProcess/RenderingCompositionGraph.h
--- a/Engine/Source/Runtime/Renderer/Private/PostProcess/RenderingCompositionGraph.h
+++ b/Engine/Source/Runtime/Renderer/Private/PostProcess/RenderingCompositionGraph.h
@@ -139,7 +139,7 @@ struct FRenderingCompositePassContext
 	{
 		ViewPortRect = InViewPortRect;
 
-		RHISetViewport(ViewPortRect.Min.X, ViewPortRect.Min.Y, InMinZ, ViewPortRect.Max.X, ViewPortRect.Max.Y, InMaxZ);
+		RHICmdList.SetViewport(ViewPortRect.Min.X, ViewPortRect.Min.Y, InMinZ, ViewPortRect.Max.X, ViewPortRect.Max.Y, InMaxZ);
 	}
 
 	// call this method instead of RHISetViewport() so we can cache the values and use them to map beteen ScreenPos and pixels
@@ -179,6 +179,8 @@ struct FRenderingCompositePassContext
 
 	FRenderingCompositionGraph Graph;
 
+	mutable FRHICommandList& RHICmdList;
+
 private:
 	// cached state to map between ScreenPos and pixels
 	FIntRect ViewPortRect;
@@ -522,13 +524,13 @@ struct FPostProcessPassParameters
 	void Bind(const FShaderParameterMap& ParameterMap);
 
 	/** Set the pixel shader parameter values. */
-	void SetPS(FRHICommandList& RHICmdList, const FPixelShaderRHIParamRef& ShaderRHI, const FRenderingCompositePassContext& Context, FSamplerStateRHIParamRef Filter = TStaticSamplerState<>::GetRHI(), bool bWhiteIfNoTexture = false, FSamplerStateRHIParamRef* FilterOverrideArray = 0);
+	void SetPS(const FPixelShaderRHIParamRef& ShaderRHI, const FRenderingCompositePassContext& Context, FSamplerStateRHIParamRef Filter = TStaticSamplerState<>::GetRHI(), bool bWhiteIfNoTexture = false, FSamplerStateRHIParamRef* FilterOverrideArray = 0);
 
 	/** Set the compute shader parameter values. */
-	void SetCS(FRHICommandList& RHICmdList, const FComputeShaderRHIParamRef& ShaderRHI, const FRenderingCompositePassContext& Context, FSamplerStateRHIParamRef Filter = TStaticSamplerState<>::GetRHI(), bool bWhiteIfNoTexture = false, FSamplerStateRHIParamRef* FilterOverrideArray = 0);
+	void SetCS(const FComputeShaderRHIParamRef& ShaderRHI, const FRenderingCompositePassContext& Context, FSamplerStateRHIParamRef Filter = TStaticSamplerState<>::GetRHI(), bool bWhiteIfNoTexture = false, FSamplerStateRHIParamRef* FilterOverrideArray = 0);
 
 	/** Set the vertex shader parameter values. */
-	void SetVS(FRHICommandList& RHICmdList, const FVertexShaderRHIParamRef& ShaderRHI, const FRenderingCompositePassContext& Context, FSamplerStateRHIParamRef Filter = TStaticSamplerState<>::GetRHI(), bool bWhiteIfNoTexture = false, FSamplerStateRHIParamRef* FilterOverrideArray = 0);
+	void SetVS(const FVertexShaderRHIParamRef& ShaderRHI, const FRenderingCompositePassContext& Context, FSamplerStateRHIParamRef Filter = TStaticSamplerState<>::GetRHI(), bool bWhiteIfNoTexture = false, FSamplerStateRHIParamRef* FilterOverrideArray = 0);
 
 	/** Serializer. */
 	friend FArchive& operator<<(FArchive& Ar,FPostProcessPassParameters& P);
diff --git a/Engine/Source/Runtime/Renderer/Private/PostProcess/SceneFilterRendering.cpp b/Engine/Source/Runtime/Renderer/Private/PostProcess/SceneFilterRendering.cpp
--- a/Engine/Source/Runtime/Renderer/Private/PostProcess/SceneFilterRendering.cpp
+++ b/Engine/Source/Runtime/Renderer/Private/PostProcess/SceneFilterRendering.cpp
@@ -100,6 +100,7 @@ static void DoDrawRectangleFlagOverride(EDrawRectangleFlags& Flags)
 }
 
 void DrawRectangle(
+	FRHICommandList& RHICmdList,
 	float X,
 	float Y,
 	float SizeX,
@@ -136,10 +137,8 @@ void DrawRectangle(
 		1.0f / TargetSize.X, 1.0f / TargetSize.Y, 
 		1.0f / TextureSize.X, 1.0f / TextureSize.Y);
 
-	//@todo-rco: RHIPacketList
-	FRHICommandList& RHICmdList = FRHICommandList::GetNullRef();
 	SetUniformBufferParameterImmediate(RHICmdList, VertexShader->GetVertexShader(), VertexShader->GetUniformBufferParameter<FDrawRectangleParameters>(), Parameters);
-	RHISetStreamSource(0, GScreenRectangleVertexBuffer.VertexBufferRHI, sizeof(FFilterVertex),0);
+	RHICmdList.SetStreamSource(0, GScreenRectangleVertexBuffer.VertexBufferRHI, sizeof(FFilterVertex), 0);
 
 
 	if (Flags == EDRF_UseTriangleOptimization)
@@ -148,7 +147,7 @@ void DrawRectangle(
 		// as we do not have a diagonal edge (through the center) for the rasterizer/span-dispatch. Although the actual benefit of this technique is dependent upon hardware.
 
 		// We offset into the index buffer when using the triangle optimization to access the correct vertices.
-		RHIDrawIndexedPrimitive(
+		RHICmdList.DrawIndexedPrimitive(
 			GScreenRectangleIndexBuffer.IndexBufferRHI,
 			PT_TriangleList,
 			/*BaseVertexIndex=*/ 0,
@@ -161,7 +160,7 @@ void DrawRectangle(
 	}
 	else
 	{
-		RHIDrawIndexedPrimitive(
+		RHICmdList.DrawIndexedPrimitive(
 			GScreenRectangleIndexBuffer.IndexBufferRHI,
 			PT_TriangleList,
 			/*BaseVertexIndex=*/ 0,
@@ -175,6 +174,7 @@ void DrawRectangle(
 }
 
 void DrawTransformedRectangle(
+	FRHICommandList& RHICmdList,
 	float X,
 	float Y,
 	float SizeX,
@@ -216,5 +216,5 @@ void DrawTransformedRectangle(
 
 	static const uint16 Indices[] =	{ 0, 1, 3, 0, 3, 2 };
 
-	RHIDrawIndexedPrimitiveUP(PT_TriangleList, 0, 4, 2, Indices, sizeof(Indices[0]), Vertices, sizeof(Vertices[0]));
+	DrawIndexedPrimitiveUP(RHICmdList, PT_TriangleList, 0, 4, 2, Indices, sizeof(Indices[0]), Vertices, sizeof(Vertices[0]));
 }
diff --git a/Engine/Source/Runtime/Renderer/Private/PostProcess/SceneRenderTargets.cpp b/Engine/Source/Runtime/Renderer/Private/PostProcess/SceneRenderTargets.cpp
--- a/Engine/Source/Runtime/Renderer/Private/PostProcess/SceneRenderTargets.cpp
+++ b/Engine/Source/Runtime/Renderer/Private/PostProcess/SceneRenderTargets.cpp
@@ -248,6 +248,8 @@ void FSceneRenderTargets::BeginRenderingSceneColor( bool bGBufferPass )
 	}
 	
 	AllocSceneColor();
+	//@todo-rco: RHIPacketList
+	FRHICommandList& RHICmdList = FRHICommandList::GetNullRef();
 
 	// Set the scene color surface as the render target, and the scene depth surface as the depth-stencil target.
 	if (bGBufferPass && GRHIFeatureLevel >= ERHIFeatureLevel::SM4)
@@ -270,11 +272,11 @@ void FSceneRenderTargets::BeginRenderingSceneColor( bool bGBufferPass )
 			MRTCount--;
 		}
 		
-		RHISetRenderTargets(MRTCount, RenderTargets, GetSceneDepthSurface(), 0, NULL);
+		SetRenderTargets(RHICmdList, MRTCount, RenderTargets, GetSceneDepthSurface(), 0, NULL);
 	}
 	else
 	{
-		RHISetRenderTarget(GetSceneColorSurface(), GetSceneDepthSurface());
+		SetRenderTarget(RHICmdList, GetSceneColorSurface(), GetSceneDepthSurface());
 	}
 } 
 
@@ -594,8 +596,10 @@ bool FSceneRenderTargets::BeginRenderingCustomDepth(bool bPrimitives)
 	if(CustomDepth)
 	{
 		SCOPED_DRAW_EVENT(BeginRenderingCustomDepth, DEC_SCENE_ITEMS);
+		//@todo-rco: RHIPacketList
+		FRHICommandList& RHICmdList = FRHICommandList::GetNullRef();
 
-		RHISetRenderTarget(FTextureRHIRef(), CustomDepth->GetRenderTargetItem().ShaderResourceTexture);
+		SetRenderTarget(RHICmdList, FTextureRHIRef(), CustomDepth->GetRenderTargetItem().ShaderResourceTexture);
 
 		return true;
 	}
@@ -647,27 +651,36 @@ void FSceneRenderTargets::BeginRenderingPrePass()
 
 	// Set the scene depth surface and a DUMMY buffer as color buffer
 	// (as long as it's the same dimension as the depth buffer),
-	RHISetRenderTarget( FTextureRHIRef(), GetSceneDepthSurface());
+	//@todo-rco: RHIPacketList
+	FRHICommandList& RHICmdList = FRHICommandList::GetNullRef();
+
+	SetRenderTarget(RHICmdList, FTextureRHIRef(), GetSceneDepthSurface());
 }
 
 void FSceneRenderTargets::FinishRenderingPrePass()
 {
 	SCOPED_DRAW_EVENT(FinishRenderingPrePass, DEC_SCENE_ITEMS);
 
-	GRenderTargetPool.VisualizeTexture.SetCheckPoint(SceneDepthZ);
+	//@todo-rco: RHIPacketList
+	FRHICommandList& RHICmdList = FRHICommandList::GetNullRef();
+	GRenderTargetPool.VisualizeTexture.SetCheckPoint(RHICmdList, SceneDepthZ);
 }
 
 void FSceneRenderTargets::BeginRenderingShadowDepth()
 {
-	GRenderTargetPool.VisualizeTexture.SetCheckPoint(ShadowDepthZ);
+	//@todo-rco: RHIPacketList
+	FRHICommandList& RHICmdList = FRHICommandList::GetNullRef();
+	GRenderTargetPool.VisualizeTexture.SetCheckPoint(RHICmdList, ShadowDepthZ);
 	
-	RHISetRenderTarget(FTextureRHIRef(), GetShadowDepthZSurface());
+	SetRenderTarget(RHICmdList, FTextureRHIRef(), GetShadowDepthZSurface());
 }
 
 void FSceneRenderTargets::BeginRenderingCubeShadowDepth(int32 ShadowResolution)
 {
 	SCOPED_DRAW_EVENT(BeginRenderingCubeShadowDepth, DEC_SCENE_ITEMS);
-	RHISetRenderTarget(FTextureRHIRef(), GetCubeShadowDepthZSurface(ShadowResolution));   
+	//@todo-rco: RHIPacketList
+	FRHICommandList& RHICmdList = FRHICommandList::GetNullRef();
+	SetRenderTarget(RHICmdList, FTextureRHIRef(), GetCubeShadowDepthZSurface(ShadowResolution));
 }
 
 void FSceneRenderTargets::FinishRenderingShadowDepth(const FResolveRect& ResolveRect)
@@ -688,11 +701,16 @@ void FSceneRenderTargets::BeginRenderingReflectiveShadowMap( FLightPropagationVo
 	Uavs[1] = Lpv->GetGvListHeadBufferUav();
 	Uavs[2] = Lpv->GetVplListBufferUav();
 	Uavs[3] = Lpv->GetVplListHeadBufferUav();
-	RHISetRenderTargets( ARRAY_COUNT(RenderTargets), RenderTargets, GetReflectiveShadowMapDepthSurface(), 4, Uavs);
+	//@todo-rco: RHIPacketList
+	FRHICommandList& RHICmdList = FRHICommandList::GetNullRef();
+	SetRenderTargets(RHICmdList, ARRAY_COUNT(RenderTargets), RenderTargets, GetReflectiveShadowMapDepthSurface(), 4, Uavs);
 }
 
 void FSceneRenderTargets::FinishRenderingReflectiveShadowMap(const FResolveRect& ResolveRect)
 {
+	//@todo-rco: RHIPacketList
+	FRHICommandList& RHICmdList = FRHICommandList::GetNullRef();
+
 	// Resolve the shadow depth z surface.
 	RHICopyToResolveTarget(GetReflectiveShadowMapDepthSurface(), GetReflectiveShadowMapDepthTexture(), false, FResolveParams(ResolveRect));
 	RHICopyToResolveTarget(GetReflectiveShadowMapDiffuseSurface(), GetReflectiveShadowMapDiffuseTexture(), false, FResolveParams(ResolveRect));
@@ -701,7 +719,7 @@ void FSceneRenderTargets::FinishRenderingReflectiveShadowMap(const FResolveRect&
 	// Unset render targets
 	FTextureRHIParamRef RenderTargets[2] = {NULL};
 	FUnorderedAccessViewRHIParamRef Uavs[2] = {NULL};
-	RHISetRenderTargets( ARRAY_COUNT(RenderTargets), RenderTargets, FTextureRHIParamRef(), 2, Uavs);
+	SetRenderTargets(RHICmdList, ARRAY_COUNT(RenderTargets), RenderTargets, FTextureRHIParamRef(), 2, Uavs);
 }
 
 void FSceneRenderTargets::FinishRenderingCubeShadowDepth(int32 ShadowResolution, const FResolveParams& ResolveParams)
@@ -713,15 +731,19 @@ void FSceneRenderTargets::FinishRenderingCubeShadowDepth(int32 ShadowResolution,
 void FSceneRenderTargets::BeginRenderingSceneAlphaCopy()
 {
 	SCOPED_DRAW_EVENT(BeginRenderingSceneAlphaCopy, DEC_SCENE_ITEMS);
-	GRenderTargetPool.VisualizeTexture.SetCheckPoint(GSceneRenderTargets.SceneAlphaCopy);
-	RHISetRenderTarget(GetSceneAlphaCopySurface(), 0);
+	//@todo-rco: RHIPacketList
+	FRHICommandList& RHICmdList = FRHICommandList::GetNullRef();
+	GRenderTargetPool.VisualizeTexture.SetCheckPoint(RHICmdList, GSceneRenderTargets.SceneAlphaCopy);
+	SetRenderTarget(RHICmdList, GetSceneAlphaCopySurface(), 0);
 }
 
 void FSceneRenderTargets::FinishRenderingSceneAlphaCopy()
 {
 	SCOPED_DRAW_EVENT(FinishRenderingSceneAlphaCopy, DEC_SCENE_ITEMS);
+	//@todo-rco: RHIPacketList
+	FRHICommandList& RHICmdList = FRHICommandList::GetNullRef();
 	RHICopyToResolveTarget(GetSceneAlphaCopySurface(), SceneAlphaCopy->GetRenderTargetItem().ShaderResourceTexture, false, FResolveParams(FResolveRect()));
-	GRenderTargetPool.VisualizeTexture.SetCheckPoint(GSceneRenderTargets.SceneAlphaCopy);
+	GRenderTargetPool.VisualizeTexture.SetCheckPoint(RHICmdList, GSceneRenderTargets.SceneAlphaCopy);
 }
 
 
@@ -731,7 +753,9 @@ void FSceneRenderTargets::BeginRenderingLightAttenuation()
 
 	AllocLightAttenuation();
 
-	GRenderTargetPool.VisualizeTexture.SetCheckPoint(GSceneRenderTargets.GetLightAttenuation());
+	//@todo-rco: RHIPacketList
+	FRHICommandList& RHICmdList = FRHICommandList::GetNullRef();
+	GRenderTargetPool.VisualizeTexture.SetCheckPoint(RHICmdList, GSceneRenderTargets.GetLightAttenuation());
 
 	// Set the light attenuation surface as the render target, and the scene depth buffer as the depth-stencil surface.
 	RHISetRenderTarget(GetLightAttenuationSurface(),GetSceneDepthSurface());
@@ -744,7 +768,9 @@ void FSceneRenderTargets::FinishRenderingLightAttenuation()
 	// Resolve the light attenuation surface.
 	RHICopyToResolveTarget(GetLightAttenuationSurface(), LightAttenuation->GetRenderTargetItem().ShaderResourceTexture, false, FResolveParams(FResolveRect()));
 	
-	GRenderTargetPool.VisualizeTexture.SetCheckPoint(GSceneRenderTargets.GetLightAttenuation());
+	//@todo-rco: RHIPacketList
+	FRHICommandList& RHICmdList = FRHICommandList::GetNullRef();
+	GRenderTargetPool.VisualizeTexture.SetCheckPoint(RHICmdList, GSceneRenderTargets.GetLightAttenuation());
 }
 
 void FSceneRenderTargets::BeginRenderingTranslucency(const FViewInfo& View)
diff --git a/Engine/Source/Runtime/Renderer/Private/PostProcess/ScreenSpaceReflections.cpp b/Engine/Source/Runtime/Renderer/Private/PostProcess/ScreenSpaceReflections.cpp
--- a/Engine/Source/Runtime/Renderer/Private/PostProcess/ScreenSpaceReflections.cpp
+++ b/Engine/Source/Runtime/Renderer/Private/PostProcess/ScreenSpaceReflections.cpp
@@ -96,15 +96,15 @@ public:
 		SSRParams.Bind(Initializer.ParameterMap, TEXT("SSRParams"));
 	}
 
-	void SetParameters(FRHICommandList& RHICmdList, const FRenderingCompositePassContext& Context)
+	void SetParameters(const FRenderingCompositePassContext& Context)
 	{
 		const FFinalPostProcessSettings& Settings = Context.View.FinalPostProcessSettings;
 		const FPixelShaderRHIParamRef ShaderRHI = GetPixelShader();
 
-		FGlobalShader::SetParameters(RHICmdList, ShaderRHI, Context.View);
+		FGlobalShader::SetParameters(Context.RHICmdList, ShaderRHI, Context.View);
 
-		PostprocessParameter.SetPS(RHICmdList, ShaderRHI, Context, TStaticSamplerState<SF_Point,AM_Clamp,AM_Clamp,AM_Clamp>::GetRHI());
-		DeferredParameters.Set(RHICmdList, ShaderRHI, Context.View);
+		PostprocessParameter.SetPS(ShaderRHI, Context, TStaticSamplerState<SF_Point, AM_Clamp, AM_Clamp, AM_Clamp>::GetRHI());
+		DeferredParameters.Set(Context.RHICmdList, ShaderRHI, Context.View);
 
 		{
 			float MaxRoughness = FMath::Clamp(Context.View.FinalPostProcessSettings.ScreenSpaceReflectionMaxRoughness, 0.01f, 1.0f);
@@ -121,7 +121,7 @@ public:
 				0, 
 				0);
 
-			SetShaderValue(RHICmdList, ShaderRHI, SSRParams, Value);
+			SetShaderValue(Context.RHICmdList, ShaderRHI, SSRParams, Value);
 		}
 	}
 
@@ -175,26 +175,24 @@ void FRCPassPostProcessScreenSpaceReflections::Process(FRenderingCompositePassCo
 {
 	SCOPED_DRAW_EVENT(ScreenSpaceReflections, DEC_SCENE_ITEMS);
 
-	//@todo-rco: RHIPacketList
-	FRHICommandList& RHICmdList = FRHICommandList::GetNullRef();
-
 	const FSceneView& View = Context.View;
 
 	const FSceneRenderTargetItem& DestRenderTarget = PassOutputs[0].RequestSurface(Context);
 
 	// Set the view family's render target/viewport.
-	RHISetRenderTarget(DestRenderTarget.TargetableTexture, FTextureRHIRef());
-	RHIClear(true, FLinearColor(0, 0, 0, 0), false, 1.0f, false, 0, FIntRect());
+	SetRenderTarget(Context.RHICmdList, DestRenderTarget.TargetableTexture, FTextureRHIRef());
+	Context.RHICmdList.Clear(true, FLinearColor(0, 0, 0, 0), false, 1.0f, false, 0, FIntRect());
 	Context.SetViewportAndCallRHI(View.ViewRect);
 
 	// set the state
-	RHISetBlendState(TStaticBlendState<>::GetRHI());
-	RHISetRasterizerState(TStaticRasterizerState<>::GetRHI());
-	RHISetDepthStencilState(TStaticDepthStencilState<false,CF_Always>::GetRHI());
+	Context.RHICmdList.SetBlendState(TStaticBlendState<>::GetRHI());
+	Context.RHICmdList.SetRasterizerState(TStaticRasterizerState<>::GetRHI());
+	Context.RHICmdList.SetDepthStencilState(TStaticDepthStencilState<false, CF_Always>::GetRHI());
 
 	int SSRQuality = ComputeSSRQuality(View.FinalPostProcessSettings.ScreenSpaceReflectionQuality);
 
 	SSRQuality = FMath::Clamp(SSRQuality, 1, 4);
+	Context.RHICmdList.CheckIsNull(); // need new approach for "static FGlobalBoundShaderState" for parallel rendering
 
 	#define CASE(A, B) \
 		case (A + 2 * (B + 3 * 0 )): \
@@ -202,9 +200,9 @@ void FRCPassPostProcessScreenSpaceReflections::Process(FRenderingCompositePassCo
 			TShaderMapRef< FPostProcessVS > VertexShader(GetGlobalShaderMap()); \
 			TShaderMapRef< FPostProcessScreenSpaceReflectionsPS<A, B> > PixelShader(GetGlobalShaderMap()); \
 			static FGlobalBoundShaderState BoundShaderState; \
-			SetGlobalBoundShaderState(BoundShaderState, GFilterVertexDeclaration.VertexDeclarationRHI, *VertexShader, *PixelShader); \
-			VertexShader->SetParameters(RHICmdList, Context); \
-			PixelShader->SetParameters(RHICmdList, Context); \
+			SetGlobalBoundShaderState(Context.RHICmdList, BoundShaderState, GFilterVertexDeclaration.VertexDeclarationRHI, *VertexShader, *PixelShader); \
+			VertexShader->SetParameters(Context); \
+			PixelShader->SetParameters(Context); \
 		}; \
 		break
 
@@ -221,6 +219,7 @@ void FRCPassPostProcessScreenSpaceReflections::Process(FRenderingCompositePassCo
 	// Draw a quad mapping scene color to the view's render target
 	TShaderMapRef< FPostProcessVS > VertexShader(GetGlobalShaderMap());
 	DrawRectangle( 
+		Context.RHICmdList,
 		0, 0,
 		View.ViewRect.Width(), View.ViewRect.Height(),
 		View.ViewRect.Min.X, View.ViewRect.Min.Y, 
@@ -230,7 +229,7 @@ void FRCPassPostProcessScreenSpaceReflections::Process(FRenderingCompositePassCo
 		*VertexShader,
 		EDRF_UseTriangleOptimization);
 
-	RHICopyToResolveTarget(DestRenderTarget.TargetableTexture, DestRenderTarget.ShaderResourceTexture, false, FResolveParams());
+	Context.RHICmdList.CopyToResolveTarget(DestRenderTarget.TargetableTexture, DestRenderTarget.ShaderResourceTexture, false, FResolveParams());
 }
 
 FPooledRenderTargetDesc FRCPassPostProcessScreenSpaceReflections::ComputeOutputDesc(EPassOutputId InPassOutputId) const
diff --git a/Engine/Source/Runtime/Renderer/Private/PostProcess/VisualizeTexture.cpp b/Engine/Source/Runtime/Renderer/Private/PostProcess/VisualizeTexture.cpp
--- a/Engine/Source/Runtime/Renderer/Private/PostProcess/VisualizeTexture.cpp
+++ b/Engine/Source/Runtime/Renderer/Private/PostProcess/VisualizeTexture.cpp
@@ -191,19 +191,19 @@ public:
 IMPLEMENT_SHADER_TYPE(,FVisualizeTexturePresentPS,TEXT("VisualizeTexture"),TEXT("PresentPS"),SF_Pixel);
 
 
-template<uint32 TextureType> void VisualizeTextureForTextureType(const FVisualizeTextureData& Data)
+template<uint32 TextureType> void VisualizeTextureForTextureType(FRHICommandList& RHICmdList, const FVisualizeTextureData& Data)
 {
-	//@todo-rco: RHIPacketList
-	FRHICommandList& RHICmdList = FRHICommandList::GetNullRef();
-
 	TShaderMapRef<FScreenVS> VertexShader(GetGlobalShaderMap());
 	TShaderMapRef<VisualizeTexturePS<TextureType> > PixelShader(GetGlobalShaderMap());
 
 	static FGlobalBoundShaderState BoundShaderState;
-	SetGlobalBoundShaderState(BoundShaderState, GFilterVertexDeclaration.VertexDeclarationRHI, *VertexShader, *PixelShader);
+	RHICmdList.CheckIsNull(); // need new approach for "static FGlobalBoundShaderState" for parallel rendering
+
+	SetGlobalBoundShaderState(RHICmdList, BoundShaderState, GFilterVertexDeclaration.VertexDeclarationRHI, *VertexShader, *PixelShader);
 	PixelShader->SetParameters(RHICmdList, Data);
 
 	DrawRectangle(
+		RHICmdList,
 		// XY
 		0, 0,
 		// SizeXY
@@ -220,7 +220,7 @@ template<uint32 TextureType> void VisualizeTextureForTextureType(const FVisualiz
 		EDRF_UseTriangleOptimization);
 }
 
-void RenderVisualizeTexture(const FVisualizeTextureData& Data)
+void RenderVisualizeTexture(FRHICommandList& RHICmdList, const FVisualizeTextureData& Data)
 {
 	if(Data.Desc.Is2DTexture())
 	{
@@ -228,30 +228,30 @@ void RenderVisualizeTexture(const FVisualizeTextureData& Data)
 		if(Data.Desc.NumSamples > 1)
 		{
 			// MSAA
-			VisualizeTextureForTextureType<5>(Data);
+			VisualizeTextureForTextureType<5>(RHICmdList, Data);
 		}
 		else
 		{
 			// non MSAA
-			VisualizeTextureForTextureType<2>(Data);
+			VisualizeTextureForTextureType<2>(RHICmdList, Data);
 		}
 	}
 	else if(Data.Desc.Is3DTexture())
 	{
 		// Volume
-		VisualizeTextureForTextureType<3>(Data);
+		VisualizeTextureForTextureType<3>(RHICmdList, Data);
 	}
 	else if(Data.Desc.IsCubemap())
 	{
 		if(Data.Desc.IsArray())
 		{
 			// Cube[]
-			VisualizeTextureForTextureType<4>(Data);
+			VisualizeTextureForTextureType<4>(RHICmdList, Data);
 		}
 		else
 		{
 			// Cube
-			VisualizeTextureForTextureType<0>(Data);
+			VisualizeTextureForTextureType<0>(RHICmdList, Data);
 		}
 	}
 }
@@ -311,11 +311,12 @@ FIntRect FVisualizeTexture::ComputeVisualizeTextureRect(FIntPoint InputTextureSi
 	return ret;
 }
 
-void FVisualizeTexture::GenerateContent(const FSceneRenderTargetItem& RenderTargetItem, const FPooledRenderTargetDesc& Desc)
+void FVisualizeTexture::GenerateContent(FRHICommandList& RHICmdList, const FSceneRenderTargetItem& RenderTargetItem, const FPooledRenderTargetDesc& Desc)
 {
 	// otherwise StartFrame() wasn't called
 	check(ViewRect != FIntRect(0, 0, 0, 0))
 
+
 	FTexture2DRHIRef VisTexture = (FTexture2DRHIRef&)RenderTargetItem.ShaderResourceTexture;
 
 	if(!IsValidRef(VisTexture) || !Desc.IsValid())
@@ -343,11 +344,11 @@ void FVisualizeTexture::GenerateContent(const FSceneRenderTargetItem& RenderTarg
 
 	const FSceneRenderTargetItem& DestRenderTarget = VisualizeTextureContent->GetRenderTargetItem();
 
-	RHISetRenderTarget(DestRenderTarget.TargetableTexture, FTextureRHIRef());	
-	RHIClear(true, FLinearColor(1,1,0,1), false, 0.0f, false, 0, FIntRect());
-	RHISetBlendState(TStaticBlendState<>::GetRHI());
-	RHISetRasterizerState(TStaticRasterizerState<>::GetRHI());
-	RHISetDepthStencilState(TStaticDepthStencilState<false,CF_Always>::GetRHI());
+	SetRenderTarget(RHICmdList, DestRenderTarget.TargetableTexture, FTextureRHIRef());
+	RHICmdList.Clear(true, FLinearColor(1, 1, 0, 1), false, 0.0f, false, 0, FIntRect());
+	RHICmdList.SetBlendState(TStaticBlendState<>::GetRHI());
+	RHICmdList.SetRasterizerState(TStaticRasterizerState<>::GetRHI());
+	RHICmdList.SetDepthStencilState(TStaticDepthStencilState<false, CF_Always>::GetRHI());
 
 	FIntPoint RTExtent = GSceneRenderTargets.GetBufferSizeXY();
 
@@ -373,6 +374,7 @@ void FVisualizeTexture::GenerateContent(const FSceneRenderTargetItem& RenderTarg
 		default:
 			break;
 	}
+	RHICmdList.CheckIsNull(); // this stuff is manging resources directly, so that won't work in parallel
 
 	bool bIsDefault = StencilSRVSrc == GBlackTexture->TextureRHI;
 	bool bDepthStencil = Desc.Is2DTexture() && Desc.Format == PF_DepthStencil;
@@ -417,9 +419,9 @@ void FVisualizeTexture::GenerateContent(const FSceneRenderTargetItem& RenderTarg
 	{	
 		SCOPED_DRAW_EVENT(VisualizeTexture, DEC_SCENE_ITEMS);
 		// continue rendering to HDR if necessary
-		RenderVisualizeTexture(VisualizeTextureData);
+		RenderVisualizeTexture(RHICmdList, VisualizeTextureData);
 	}
-	RHICopyToResolveTarget(DestRenderTarget.TargetableTexture, DestRenderTarget.ShaderResourceTexture, false, FResolveParams());
+	RHICmdList.CopyToResolveTarget(DestRenderTarget.TargetableTexture, DestRenderTarget.ShaderResourceTexture, false, FResolveParams());
 
 	VisualizeTextureDesc = Desc;
 
@@ -445,6 +447,7 @@ void FVisualizeTexture::GenerateContent(const FSceneRenderTargetItem& RenderTarg
 		
 		
 		
+		RHICmdList.CheckIsNull(); // synchronous readback, won't work with parallel
 		
 		RHIReadSurfaceData(Texture, FIntRect(0, 0, Extent.X, Extent.Y), Bitmap, ReadDataFlags);
 
@@ -470,7 +473,7 @@ void FVisualizeTexture::GenerateContent(const FSceneRenderTargetItem& RenderTarg
 	}
 }
 
-void FVisualizeTexture::PresentContent(const FSceneView& View)
+void FVisualizeTexture::PresentContent(FRHICommandList& RHICmdList, const FSceneView& View)
 {
 	if(Mode != 0)
 	{
@@ -478,7 +481,7 @@ void FVisualizeTexture::PresentContent(const FSceneView& View)
 		FPooledRenderTarget* Element = GRenderTargetPool.GetElementById(Mode - 1);
 		if(Element)
 		{
-			GenerateContent(Element->GetRenderTargetItem(), Element->GetDesc());
+			GenerateContent(RHICmdList, Element->GetRenderTargetItem(), Element->GetDesc());
 		}
 	}
 
@@ -492,25 +495,23 @@ void FVisualizeTexture::PresentContent(const FSceneView& View)
 		return;
 	}
 
-	//@todo-rco: RHIPacketList
-	FRHICommandList& RHICmdList = FRHICommandList::GetNullRef();
-
 	FPooledRenderTargetDesc Desc = VisualizeTextureDesc;
 
 	auto& RenderTarget = View.Family->RenderTarget->GetRenderTargetTexture();
-	RHISetRenderTarget(RenderTarget, FTextureRHIRef());
-	RHISetViewport(0, 0, 0.0f, RenderTarget->GetSizeX(), RenderTarget->GetSizeY(), 1.0f);
+	SetRenderTarget(RHICmdList, RenderTarget, FTextureRHIRef());
+	RHICmdList.SetViewport(0, 0, 0.0f, RenderTarget->GetSizeX(), RenderTarget->GetSizeY(), 1.0f);
 
-	RHISetBlendState(TStaticBlendState<>::GetRHI());
-	RHISetRasterizerState(TStaticRasterizerState<>::GetRHI());
-	RHISetDepthStencilState(TStaticDepthStencilState<false,CF_Always>::GetRHI());
+	RHICmdList.SetBlendState(TStaticBlendState<>::GetRHI());
+	RHICmdList.SetRasterizerState(TStaticRasterizerState<>::GetRHI());
+	RHICmdList.SetDepthStencilState(TStaticDepthStencilState<false, CF_Always>::GetRHI());
 
 	TShaderMapRef<FPostProcessVS> VertexShader(GetGlobalShaderMap());
 	TShaderMapRef<FVisualizeTexturePresentPS> PixelShader(GetGlobalShaderMap());
 
 	static FGlobalBoundShaderState BoundShaderState;
+	RHICmdList.CheckIsNull(); // need new approach for "static FGlobalBoundShaderState" for parallel rendering
 
-	SetGlobalBoundShaderState(BoundShaderState, GFilterVertexDeclaration.VertexDeclarationRHI, *VertexShader, *PixelShader);
+	SetGlobalBoundShaderState(RHICmdList, BoundShaderState, GFilterVertexDeclaration.VertexDeclarationRHI, *VertexShader, *PixelShader);
 
 	VertexShader->SetParameters(RHICmdList, View);
 	PixelShader->SetParameters(RHICmdList, View, *VisualizeTextureContent);
@@ -521,6 +522,7 @@ void FVisualizeTexture::PresentContent(const FSceneView& View)
 
 	// Draw a quad mapping scene color to the view's render target
 	DrawRectangle(
+		RHICmdList,
 		VisualizeTextureRect.Min.X, VisualizeTextureRect.Min.Y,
 		VisualizeTextureRect.Width(), VisualizeTextureRect.Height(),
 		0, 0,
@@ -548,6 +550,7 @@ void FVisualizeTexture::PresentContent(const FSceneView& View)
 			return View.Family->RenderTarget->GetRenderTargetTexture();
 		}
 	} TempRenderTarget(View);
+	RHICmdList.CheckIsNull(); // canvas doesn't work with command lists
 
 	FCanvas Canvas(&TempRenderTarget, NULL, View.Family->CurrentRealTime, View.Family->CurrentWorldTime, View.Family->DeltaWorldTime);
 
@@ -621,7 +624,7 @@ void FVisualizeTexture::SetObserveTarget(const FString& InObservedDebugName, uin
 	ObservedDebugNameReusedGoal = InObservedDebugNameReusedGoal;
 }
 
-void FVisualizeTexture::SetCheckPoint(const IPooledRenderTarget* PooledRenderTarget)
+void FVisualizeTexture::SetCheckPoint(FRHICommandList& RHICmdList, const IPooledRenderTarget* PooledRenderTarget)
 {
 #if !(UE_BUILD_SHIPPING || UE_BUILD_TEST)
 	if (!PooledRenderTarget || !bEnabled)
@@ -648,7 +651,7 @@ void FVisualizeTexture::SetCheckPoint(const IPooledRenderTarget* PooledRenderTar
 		// if multiple times reused during the frame, is that the one we want to look at?
 		if(*UsageCountPtr == ObservedDebugNameReusedGoal || ObservedDebugNameReusedGoal == 0xffffffff)
 		{
-			GenerateContent(RenderTargetItem, Desc);
+			GenerateContent(RHICmdList, RenderTargetItem, Desc);
 		}
 	}
 	// only needed for VisualizeTexture (todo: optimize out when possible)
diff --git a/Engine/Source/Runtime/Renderer/Private/PostProcess/VisualizeTexture.h b/Engine/Source/Runtime/Renderer/Private/PostProcess/VisualizeTexture.h
--- a/Engine/Source/Runtime/Renderer/Private/PostProcess/VisualizeTexture.h
+++ b/Engine/Source/Runtime/Renderer/Private/PostProcess/VisualizeTexture.h
@@ -12,7 +12,7 @@ public:
 	FVisualizeTexture();
 
 	/** renders the VisualizeTextureContent to the current render target */
-	void PresentContent(const FSceneView& View);
+	void PresentContent(FRHICommandList& RHICmdList, const FSceneView& View);
 	/** */
 	void OnStartFrame(const FSceneView& View);
 	/** */
@@ -23,7 +23,7 @@ public:
 	 * calling this allows to grab the state of the texture at this point to be queried by visualizetexture e.g. "vis LightAttenuation@2"
 	 * @param PooledRenderTarget 0 is silently ignored
 	 */
-	void SetCheckPoint(const IPooledRenderTarget* PooledRenderTarget);
+	void SetCheckPoint(FRHICommandList& RHICmdList, const IPooledRenderTarget* PooledRenderTarget);
 
 	// @param bExtended true: with more convenience - not needed for crashes but useful from the console
 	void DebugLog(bool bExtended);
@@ -88,7 +88,7 @@ private:
 
 	// is called by FPooledRenderTarget
 
-	void GenerateContent(const FSceneRenderTargetItem& RenderTargetItem, const FPooledRenderTargetDesc& Desc);
+	void GenerateContent(FRHICommandList& RHICmdList, const FSceneRenderTargetItem& RenderTargetItem, const FPooledRenderTargetDesc& Desc);
 
 	FIntRect ComputeVisualizeTextureRect(FIntPoint InputTextureSize) const;
 };
diff --git a/Engine/Source/Runtime/Renderer/Private/ReflectionEnvironment.cpp b/Engine/Source/Runtime/Renderer/Private/ReflectionEnvironment.cpp
--- a/Engine/Source/Runtime/Renderer/Private/ReflectionEnvironment.cpp
+++ b/Engine/Source/Runtime/Renderer/Private/ReflectionEnvironment.cpp
@@ -271,11 +271,10 @@ public:
 		SkyLightParameters.SetParameters(RHICmdList, ShaderRHI, Scene, View.Family->EngineShowFlags.SkyLighting);
 	}
 
-	void UnsetParameters()
+	void UnsetParameters(FRHICommandList& RHICmdList)
 	{
 		const FComputeShaderRHIParamRef ShaderRHI = GetComputeShader();
-		//@todo-rco: RHIPacketList
-		OutSceneColor.UnsetUAV(FRHICommandList::GetNullRef(), ShaderRHI);
+		OutSceneColor.UnsetUAV(RHICmdList, ShaderRHI);
 	}
 
 	virtual bool Serialize(FArchive& Ar)
@@ -546,31 +545,35 @@ private:
 IMPLEMENT_SHADER_TYPE(template<>,TStandardDeferredReflectionPS<true>,TEXT("ReflectionEnvironmentShaders"),TEXT("StandardDeferredReflectionPS"),SF_Pixel);
 IMPLEMENT_SHADER_TYPE(template<>,TStandardDeferredReflectionPS<false>,TEXT("ReflectionEnvironmentShaders"),TEXT("StandardDeferredReflectionPS"),SF_Pixel);
 
-void FDeferredShadingSceneRenderer::RenderReflectionCaptureSpecularBounceForAllViews()
+void FDeferredShadingSceneRenderer::RenderReflectionCaptureSpecularBounceForAllViews(FRHICommandList& RHICmdList)
 {
+	RHICmdList.CheckIsNull(); // GSceneRenderTargets ops
+
 				// We're currently capturing a reflection capture, output SpecularColor * IndirectDiffuseGI for metals so they are not black in reflections,
 				// Since we don't have multiple bounce specular reflections
 	GSceneRenderTargets.BeginRenderingSceneColor();
-	RHISetRasterizerState( TStaticRasterizerState< FM_Solid, CM_None >::GetRHI() );
-	RHISetDepthStencilState( TStaticDepthStencilState< false, CF_Always >::GetRHI() );
-	RHISetBlendState( TStaticBlendState< CW_RGB, BO_Add, BF_One, BF_One >::GetRHI() );
+	RHICmdList.SetRasterizerState(TStaticRasterizerState< FM_Solid, CM_None >::GetRHI());
+	RHICmdList.SetDepthStencilState(TStaticDepthStencilState< false, CF_Always >::GetRHI());
+	RHICmdList.SetBlendState(TStaticBlendState< CW_RGB, BO_Add, BF_One, BF_One >::GetRHI());
 
 	TShaderMapRef< FPostProcessVS > VertexShader( GetGlobalShaderMap() );
 	TShaderMapRef< FReflectionCaptureSpecularBouncePS > PixelShader( GetGlobalShaderMap() );
 
 	static FGlobalBoundShaderState BoundShaderState;
-	SetGlobalBoundShaderState( BoundShaderState, GFilterVertexDeclaration.VertexDeclarationRHI, *VertexShader, *PixelShader );
+	RHICmdList.CheckIsNull(); // need new approach for "static FGlobalBoundShaderState" for parallel rendering
+	SetGlobalBoundShaderState(RHICmdList, BoundShaderState, GFilterVertexDeclaration.VertexDeclarationRHI, *VertexShader, *PixelShader);
 
 	for (int32 ViewIndex = 0, Num = Views.Num(); ViewIndex < Num; ViewIndex++)
 	{
 		const FViewInfo& View = Views[ViewIndex];
 
-		RHISetViewport( View.ViewRect.Min.X, View.ViewRect.Min.Y, 0.0f, View.ViewRect.Max.X, View.ViewRect.Max.Y, 1.0f );
+		RHICmdList.SetViewport(View.ViewRect.Min.X, View.ViewRect.Min.Y, 0.0f, View.ViewRect.Max.X, View.ViewRect.Max.Y, 1.0f);
 
 		//@todo-rco: RHIPacketList
-		PixelShader->SetParameters(FRHICommandList::GetNullRef(), View);
+		PixelShader->SetParameters(RHICmdList, View);
 
 		DrawRectangle( 
+			RHICmdList,
 			0, 0,
 			View.ViewRect.Width(), View.ViewRect.Height(),
 			0, 0,
@@ -595,6 +598,8 @@ void FDeferredShadingSceneRenderer::RenderImageBasedReflectionsSM5ForAllViews(FR
 {
 	const uint32 bUseLightmaps = CVarDiffuseFromCaptures.GetValueOnRenderThread() == 0;
 
+	RHICmdList.CheckIsNull(); // GSceneRenderTargets ops
+
 	TRefCountPtr<IPooledRenderTarget> NewSceneColor;
 	{
 		GSceneRenderTargets.ResolveSceneColor(FResolveRect(0, 0, ViewFamily.FamilySizeX, ViewFamily.FamilySizeY));
@@ -623,7 +628,7 @@ void FDeferredShadingSceneRenderer::RenderImageBasedReflectionsSM5ForAllViews(FR
 			// Render the reflection environment with tiled deferred culling
 			SCOPED_DRAW_EVENT(ReflectionEnvironmentGather, DEC_SCENE_ITEMS);
 
-			RHISetRenderTarget(NULL, NULL);
+			SetRenderTarget(RHICmdList, NULL, NULL);
 
 			FReflectionEnvironmentTiledDeferredCS* ComputeShader = NULL;
 			if( bUseLightmaps )
@@ -635,7 +640,7 @@ void FDeferredShadingSceneRenderer::RenderImageBasedReflectionsSM5ForAllViews(FR
 				ComputeShader = *TShaderMapRef< TReflectionEnvironmentTiledDeferredCS<0> >( GetGlobalShaderMap() );
 			}
 
-			RHISetComputeShader(ComputeShader->GetComputeShader());
+			RHICmdList.SetComputeShader(ComputeShader->GetComputeShader());
 
 			ComputeShader->SetParameters(RHICmdList, View, SSROutput->GetRenderTargetItem().ShaderResourceTexture, NewSceneColor->GetRenderTargetItem().UAV);
 			//ComputeShader->SetParameters(RHICmdList, View, DestRenderTarget.UAV);
@@ -644,7 +649,7 @@ void FDeferredShadingSceneRenderer::RenderImageBasedReflectionsSM5ForAllViews(FR
 			uint32 GroupSizeY = (View.ViewRect.Size().Y + GReflectionEnvironmentTileSizeY - 1) / GReflectionEnvironmentTileSizeY;
 			DispatchComputeShader(RHICmdList, ComputeShader, GroupSizeX, GroupSizeY, 1);
 
-			ComputeShader->UnsetParameters();
+			ComputeShader->UnsetParameters(RHICmdList);
 		}
 	}
 
@@ -734,15 +739,15 @@ void FDeferredShadingSceneRenderer::RenderImageBasedReflectionsSM4ForAllViews(FR
 
 			SCOPED_DRAW_EVENT(StandardDeferredReflectionEnvironment, DEC_SCENE_ITEMS);
 
-			RHISetRenderTarget(LightAccumulation->GetRenderTargetItem().TargetableTexture, NULL);
+			SetRenderTarget(RHICmdList, LightAccumulation->GetRenderTargetItem().TargetableTexture, NULL);
 
 			// Clear to no reflection contribution, alpha of 1 indicates full background contribution
-			RHIClear(true, FLinearColor(0, 0, 0, 1), false, 0, false, 0, FIntRect());
+			RHICmdList.Clear(true, FLinearColor(0, 0, 0, 1), false, 0, false, 0, FIntRect());
 
-			RHISetViewport(View.ViewRect.Min.X, View.ViewRect.Min.Y, 0.0f, View.ViewRect.Max.X, View.ViewRect.Max.Y, 1.0f);
+			RHICmdList.SetViewport(View.ViewRect.Min.X, View.ViewRect.Min.Y, 0.0f, View.ViewRect.Max.X, View.ViewRect.Max.Y, 1.0f);
 
 			// rgb accumulates reflection contribution front to back, alpha accumulates (1 - alpha0) * (1 - alpha 1)...
-			RHISetBlendState(TStaticBlendState<CW_RGBA, BO_Add, BF_DestAlpha, BF_One, BO_Add, BF_Zero, BF_InverseSourceAlpha>::GetRHI());
+			RHICmdList.SetBlendState(TStaticBlendState<CW_RGBA, BO_Add, BF_DestAlpha, BF_One, BO_Add, BF_Zero, BF_InverseSourceAlpha>::GetRHI());
 
 			for (int32 ReflectionCaptureIndex = 0; ReflectionCaptureIndex < SortData.Num(); ReflectionCaptureIndex++)
 			{
@@ -760,7 +765,8 @@ void FDeferredShadingSceneRenderer::RenderImageBasedReflectionsSM4ForAllViews(FR
 						TShaderMapRef<TStandardDeferredReflectionPS<true> > PixelShader(GetGlobalShaderMap());
 
 						static FGlobalBoundShaderState BoundShaderState;
-						SetGlobalBoundShaderState(BoundShaderState, GFilterVertexDeclaration.VertexDeclarationRHI, *VertexShader, *PixelShader);
+						RHICmdList.CheckIsNull(); // need new approach for "static FGlobalBoundShaderState" for parallel rendering
+						SetGlobalBoundShaderState(RHICmdList, BoundShaderState, GFilterVertexDeclaration.VertexDeclarationRHI, *VertexShader, *PixelShader);
 
 						PixelShader->SetParameters(RHICmdList, View, ReflectionCapture);
 					}
@@ -769,50 +775,54 @@ void FDeferredShadingSceneRenderer::RenderImageBasedReflectionsSM4ForAllViews(FR
 						TShaderMapRef<TStandardDeferredReflectionPS<false> > PixelShader(GetGlobalShaderMap());
 
 						static FGlobalBoundShaderState BoundShaderState;
-						SetGlobalBoundShaderState(BoundShaderState, GFilterVertexDeclaration.VertexDeclarationRHI, *VertexShader, *PixelShader);
+						RHICmdList.CheckIsNull(); // need new approach for "static FGlobalBoundShaderState" for parallel rendering
+						SetGlobalBoundShaderState(RHICmdList, BoundShaderState, GFilterVertexDeclaration.VertexDeclarationRHI, *VertexShader, *PixelShader);
 
 						PixelShader->SetParameters(RHICmdList, View, ReflectionCapture);
 					}
 
-					SetBoundingGeometryRasterizerAndDepthState(View, LightBounds);
+					SetBoundingGeometryRasterizerAndDepthState(RHICmdList, View, LightBounds);
 					VertexShader->SetSimpleLightParameters(RHICmdList, View, LightBounds);
-					StencilingGeometry::DrawSphere();
+					StencilingGeometry::DrawSphere(RHICmdList);
 				}
 			}
 
-			GRenderTargetPool.VisualizeTexture.SetCheckPoint(LightAccumulation);
+			GRenderTargetPool.VisualizeTexture.SetCheckPoint(RHICmdList, LightAccumulation);
 		}
 
 		if (bRequiresApply)
 		{
 			// Apply reflections to screen
 			SCOPED_DRAW_EVENT(ReflectionApply, DEC_SCENE_ITEMS);
+			RHICmdList.CheckIsNull(); // GSceneRenderTargets ops
 
 			GSceneRenderTargets.BeginRenderingSceneColor();
 
-			RHISetViewport(View.ViewRect.Min.X, View.ViewRect.Min.Y, 0.0f, View.ViewRect.Max.X, View.ViewRect.Max.Y, 1.0f);
-			RHISetRasterizerState(TStaticRasterizerState<FM_Solid, CM_None>::GetRHI());
-			RHISetDepthStencilState(TStaticDepthStencilState<false, CF_Always>::GetRHI());
+			RHICmdList.SetViewport(View.ViewRect.Min.X, View.ViewRect.Min.Y, 0.0f, View.ViewRect.Max.X, View.ViewRect.Max.Y, 1.0f);
+			RHICmdList.SetRasterizerState(TStaticRasterizerState<FM_Solid, CM_None>::GetRHI());
+			RHICmdList.SetDepthStencilState(TStaticDepthStencilState<false, CF_Always>::GetRHI());
 
 			if (GetReflectionEnvironmentCVar() == 2)
 			{
 				// override scene color for debugging
-				RHISetBlendState(TStaticBlendState<>::GetRHI());
+				RHICmdList.SetBlendState(TStaticBlendState<>::GetRHI());
 			}
 			else
 			{
 				// additive to scene color
-				RHISetBlendState(TStaticBlendState<CW_RGBA, BO_Add, BF_One, BF_One, BO_Add, BF_One, BF_One>::GetRHI());
+				RHICmdList.SetBlendState(TStaticBlendState<CW_RGBA, BO_Add, BF_One, BF_One, BO_Add, BF_One, BF_One>::GetRHI());
 			}
 
 			TShaderMapRef< FPostProcessVS >		VertexShader(GetGlobalShaderMap());
 
+			RHICmdList.CheckIsNull(); // need new approach for "static FGlobalBoundShaderState" for parallel rendering
+
 #define CASE(A,B,C) \
 			case ((A << 2) | (B << 1) | C) : \
 			{ \
 			TShaderMapRef< FReflectionApplyPS<A, B, C> > PixelShader(GetGlobalShaderMap()); \
 			static FGlobalBoundShaderState BoundShaderState; \
-			SetGlobalBoundShaderState(BoundShaderState, GFilterVertexDeclaration.VertexDeclarationRHI, *VertexShader, *PixelShader); \
+			SetGlobalBoundShaderState(RHICmdList, BoundShaderState, GFilterVertexDeclaration.VertexDeclarationRHI, *VertexShader, *PixelShader); \
 			PixelShader->SetParameters(FRHICommandList::GetNullRef(), View, LightAccumulation->GetRenderTargetItem().ShaderResourceTexture, SSROutput->GetRenderTargetItem().ShaderResourceTexture); \
 			}; \
 			break
@@ -831,6 +841,7 @@ void FDeferredShadingSceneRenderer::RenderImageBasedReflectionsSM4ForAllViews(FR
 #undef CASE
 
 			DrawRectangle(
+				RHICmdList,
 				0, 0,
 				View.ViewRect.Width(), View.ViewRect.Height(),
 				View.ViewRect.Min.X, View.ViewRect.Min.Y,
@@ -858,7 +869,7 @@ void FDeferredShadingSceneRenderer::RenderDeferredReflections(FRHICommandList& R
 
 	if (bAnyViewIsReflectionCapture)
 	{
-		RenderReflectionCaptureSpecularBounceForAllViews();
+		RenderReflectionCaptureSpecularBounceForAllViews(RHICmdList);
 	}
 	else
 	{
