UE4 - more work preparing for parallel rendering. Merged the cmd list into the post process context, among other things.

diff --git a/Engine/Source/Runtime/Renderer/Private/ReflectionEnvironmentCapture.cpp b/Engine/Source/Runtime/Renderer/Private/ReflectionEnvironmentCapture.cpp
--- a/Engine/Source/Runtime/Renderer/Private/ReflectionEnvironmentCapture.cpp
+++ b/Engine/Source/Runtime/Renderer/Private/ReflectionEnvironmentCapture.cpp
@@ -285,22 +285,22 @@ IMPLEMENT_SHADER_TYPE(,FComputeBrightnessPS,TEXT("ReflectionEnvironmentShaders")
 /** Computes the average brightness of the given reflection capture and stores it in the scene. */
 void ComputeAverageBrightness(FRHICommandList& RHICmdList)
 {
-	///@todo-rco:
-	RHICmdList.CheckIsNull();
-	RHISetRenderTarget(GSceneRenderTargets.ReflectionBrightness->GetRenderTargetItem().TargetableTexture, NULL);
-	RHISetRasterizerState(TStaticRasterizerState<FM_Solid,CM_None>::GetRHI());
-	RHISetDepthStencilState(TStaticDepthStencilState<false,CF_Always>::GetRHI());
-	RHISetBlendState(TStaticBlendState<>::GetRHI());
+	SetRenderTarget(RHICmdList, GSceneRenderTargets.ReflectionBrightness->GetRenderTargetItem().TargetableTexture, NULL);
+	RHICmdList.SetRasterizerState(TStaticRasterizerState<FM_Solid, CM_None>::GetRHI());
+	RHICmdList.SetDepthStencilState(TStaticDepthStencilState<false, CF_Always>::GetRHI());
+	RHICmdList.SetBlendState(TStaticBlendState<>::GetRHI());
 
 	TShaderMapRef<FPostProcessVS> VertexShader(GetGlobalShaderMap());
 	TShaderMapRef<FComputeBrightnessPS> PixelShader(GetGlobalShaderMap());
 
 	static FGlobalBoundShaderState BoundShaderState;
-	SetGlobalBoundShaderState(BoundShaderState, GFilterVertexDeclaration.VertexDeclarationRHI, *VertexShader, *PixelShader);
+	RHICmdList.CheckIsNull(); // need new approach for "static FGlobalBoundShaderState" for parallel rendering
+	SetGlobalBoundShaderState(RHICmdList, BoundShaderState, GFilterVertexDeclaration.VertexDeclarationRHI, *VertexShader, *PixelShader);
 
 	PixelShader->SetParameters(RHICmdList);
 
 	DrawRectangle( 
+		RHICmdList,
 		0, 0, 
 		1, 1,
 		0, 0, 
@@ -311,14 +311,11 @@ void ComputeAverageBrightness(FRHICommandList& RHICmdList)
 }
 
 /** Generates mips for glossiness and filters the cubemap for a given reflection. */
-void FilterReflectionEnvironment(FSHVectorRGB3* OutIrradianceEnvironmentMap)
+void FilterReflectionEnvironment(FRHICommandList& RHICmdList, FSHVectorRGB3* OutIrradianceEnvironmentMap)
 {
 	const int32 EffectiveTopMipSize = GReflectionCaptureSize;
 	const int32 NumMips = FMath::CeilLogTwo(EffectiveTopMipSize) + 1;
 
-	//@todo-rco: RHIPacketList
-	FRHICommandList& RHICmdList = FRHICommandList::GetNullRef();
-
 	int32 DiffuseConvolutionSourceMip = INDEX_NONE;
 	FSceneRenderTargetItem* DiffuseConvolutionSource = NULL;
 
@@ -340,19 +337,19 @@ void FilterReflectionEnvironment(FSHVectorRGB3* OutIrradianceEnvironmentMap)
 
 			if (GSupportsGSRenderTargetLayerSwitchingToMips)
 			{
-			RHISetRenderTarget(EffectiveRT.TargetableTexture, MipIndex, NULL);
+				SetRenderTarget(RHICmdList, EffectiveRT.TargetableTexture, MipIndex, NULL);
 
 			const FIntRect ViewRect(0, 0, MipSize, MipSize);
-			RHISetViewport(0, 0, 0.0f, MipSize, MipSize, 1.0f);
-			RHISetRasterizerState(TStaticRasterizerState<FM_Solid,CM_None>::GetRHI());
-			RHISetDepthStencilState(TStaticDepthStencilState<false,CF_Always>::GetRHI());
-			RHISetBlendState(TStaticBlendState<>::GetRHI());
+				RHICmdList.SetViewport(0, 0, 0.0f, MipSize, MipSize, 1.0f);
+				RHICmdList.SetRasterizerState(TStaticRasterizerState<FM_Solid, CM_None>::GetRHI());
+				RHICmdList.SetDepthStencilState(TStaticDepthStencilState<false, CF_Always>::GetRHI());
+				RHICmdList.SetBlendState(TStaticBlendState<>::GetRHI());
 
 			TShaderMapRef<FScreenVSForGS> VertexShader(GetGlobalShaderMap());
 			TShaderMapRef<FDownsampleGS> GeometryShader(GetGlobalShaderMap());
 				TShaderMapRef<FDownsamplePS> PixelShader(GetGlobalShaderMap());
 
-			SetGlobalBoundShaderState(DownsampleBoundShaderState, GFilterVertexDeclaration.VertexDeclarationRHI, *VertexShader, *PixelShader, *GeometryShader);
+				SetGlobalBoundShaderState(RHICmdList, DownsampleBoundShaderState, GFilterVertexDeclaration.VertexDeclarationRHI, *VertexShader, *PixelShader, *GeometryShader);
 
 			// Draw each face with a geometry shader that routes to the correct render target slice
 			for (int32 CubeFace = 0; CubeFace < CubeFace_MAX; CubeFace++)
@@ -363,6 +360,7 @@ void FilterReflectionEnvironment(FSHVectorRGB3* OutIrradianceEnvironmentMap)
 				PixelShader->SetParameters(RHICmdList, CubeFace, SourceMipIndex, EffectiveSource);
 
 				DrawRectangle( 
+						RHICmdList,
 					ViewRect.Min.X, ViewRect.Min.Y, 
 					ViewRect.Width(), ViewRect.Height(),
 					ViewRect.Min.X, ViewRect.Min.Y, 
@@ -372,29 +370,30 @@ void FilterReflectionEnvironment(FSHVectorRGB3* OutIrradianceEnvironmentMap)
 					*VertexShader,
 					EDRF_UseTriangleOptimization);
 
-				RHICopyToResolveTarget(EffectiveRT.TargetableTexture, EffectiveRT.ShaderResourceTexture, true, FResolveParams(FResolveRect(), (ECubeFace)CubeFace, MipIndex));
+					RHICmdList.CopyToResolveTarget(EffectiveRT.TargetableTexture, EffectiveRT.ShaderResourceTexture, true, FResolveParams(FResolveRect(), (ECubeFace)CubeFace, MipIndex));
 			}
 			}
 			else
 			{
 				for (int32 CubeFace = 0; CubeFace < CubeFace_MAX; CubeFace++)
 				{
-					RHISetRenderTarget(EffectiveRT.TargetableTexture, MipIndex, CubeFace, NULL);
+					SetRenderTarget(RHICmdList, EffectiveRT.TargetableTexture, MipIndex, CubeFace, NULL);
 
 					const FIntRect ViewRect(0, 0, MipSize, MipSize);
-					RHISetViewport(0, 0, 0.0f, MipSize, MipSize, 1.0f);
-					RHISetRasterizerState(TStaticRasterizerState<FM_Solid,CM_None>::GetRHI());
-					RHISetDepthStencilState(TStaticDepthStencilState<false,CF_Always>::GetRHI());
-					RHISetBlendState(TStaticBlendState<>::GetRHI());
+					RHICmdList.SetViewport(0, 0, 0.0f, MipSize, MipSize, 1.0f);
+					RHICmdList.SetRasterizerState(TStaticRasterizerState<FM_Solid, CM_None>::GetRHI());
+					RHICmdList.SetDepthStencilState(TStaticDepthStencilState<false, CF_Always>::GetRHI());
+					RHICmdList.SetBlendState(TStaticBlendState<>::GetRHI());
 
 					TShaderMapRef<FScreenVSForGS> VertexShader(GetGlobalShaderMap());
 					TShaderMapRef<FDownsamplePS> PixelShader(GetGlobalShaderMap());
 
-					SetGlobalBoundShaderState(DownsampleBoundShaderState, GFilterVertexDeclaration.VertexDeclarationRHI, *VertexShader, *PixelShader);
+					SetGlobalBoundShaderState(RHICmdList, DownsampleBoundShaderState, GFilterVertexDeclaration.VertexDeclarationRHI, *VertexShader, *PixelShader);
 
 					PixelShader->SetParameters(RHICmdList, CubeFace, SourceMipIndex, EffectiveSource);
 
 					DrawRectangle( 
+						RHICmdList,
 						ViewRect.Min.X, ViewRect.Min.Y, 
 						ViewRect.Width(), ViewRect.Height(),
 						ViewRect.Min.X, ViewRect.Min.Y, 
@@ -403,7 +402,7 @@ void FilterReflectionEnvironment(FSHVectorRGB3* OutIrradianceEnvironmentMap)
 						FIntPoint(MipSize, MipSize),
 						*VertexShader);
 
-					RHICopyToResolveTarget(EffectiveRT.TargetableTexture, EffectiveRT.ShaderResourceTexture, true, FResolveParams(FResolveRect(), (ECubeFace)CubeFace, MipIndex));
+					RHICmdList.CopyToResolveTarget(EffectiveRT.TargetableTexture, EffectiveRT.ShaderResourceTexture, true, FResolveParams(FResolveRect(), (ECubeFace)CubeFace, MipIndex));
 				}
 			}
 
@@ -419,7 +418,7 @@ void FilterReflectionEnvironment(FSHVectorRGB3* OutIrradianceEnvironmentMap)
 	{
 		SCOPED_DRAW_EVENT(ComputeDiffuseIrradiance, DEC_SCENE_ITEMS);
 		check(DiffuseConvolutionSource != NULL);
-		ComputeDiffuseIrradiance(DiffuseConvolutionSource->ShaderResourceTexture, DiffuseConvolutionSourceMip, OutIrradianceEnvironmentMap);
+		ComputeDiffuseIrradiance(RHICmdList, DiffuseConvolutionSource->ShaderResourceTexture, DiffuseConvolutionSourceMip, OutIrradianceEnvironmentMap);
 	}
 	
 	ComputeAverageBrightness(RHICmdList);
@@ -440,13 +439,13 @@ void FilterReflectionEnvironment(FSHVectorRGB3* OutIrradianceEnvironmentMap)
 
 			if (GSupportsGSRenderTargetLayerSwitchingToMips)
 			{
-			RHISetRenderTarget(EffectiveRT.TargetableTexture, MipIndex, NULL);
+				SetRenderTarget(RHICmdList, EffectiveRT.TargetableTexture, MipIndex, NULL);
 
 			const FIntRect ViewRect(0, 0, MipSize, MipSize);
-			RHISetViewport(0, 0, 0.0f, MipSize, MipSize, 1.0f);
-			RHISetRasterizerState(TStaticRasterizerState<FM_Solid,CM_None>::GetRHI());
-			RHISetDepthStencilState(TStaticDepthStencilState<false,CF_Always>::GetRHI());
-			RHISetBlendState(TStaticBlendState<>::GetRHI());
+				RHICmdList.SetViewport(0, 0, 0.0f, MipSize, MipSize, 1.0f);
+				RHICmdList.SetRasterizerState(TStaticRasterizerState<FM_Solid, CM_None>::GetRHI());
+				RHICmdList.SetDepthStencilState(TStaticDepthStencilState<false, CF_Always>::GetRHI());
+				RHICmdList.SetBlendState(TStaticBlendState<>::GetRHI());
 
 			TShaderMapRef<FScreenVSForGS> VertexShader(GetGlobalShaderMap());
 			TShaderMapRef<FDownsampleGS> GeometryShader(GetGlobalShaderMap());
@@ -456,13 +455,15 @@ void FilterReflectionEnvironment(FSHVectorRGB3* OutIrradianceEnvironmentMap)
 				{
 					PixelShader = *TShaderMapRef< TCubeFilterPS<1> >( GetGlobalShaderMap() );
 					static FGlobalBoundShaderState BoundShaderState;
-					SetGlobalBoundShaderState(BoundShaderState, GFilterVertexDeclaration.VertexDeclarationRHI, *VertexShader, PixelShader, *GeometryShader);
+					RHICmdList.CheckIsNull(); // need new approach for "static FGlobalBoundShaderState" for parallel rendering
+					SetGlobalBoundShaderState(RHICmdList, BoundShaderState, GFilterVertexDeclaration.VertexDeclarationRHI, *VertexShader, PixelShader, *GeometryShader);
 				}
 				else
 				{
 					PixelShader = *TShaderMapRef< TCubeFilterPS<0> >( GetGlobalShaderMap() );
 					static FGlobalBoundShaderState BoundShaderState;
-					SetGlobalBoundShaderState(BoundShaderState, GFilterVertexDeclaration.VertexDeclarationRHI, *VertexShader, PixelShader, *GeometryShader);
+					RHICmdList.CheckIsNull(); // need new approach for "static FGlobalBoundShaderState" for parallel rendering
+					SetGlobalBoundShaderState(RHICmdList, BoundShaderState, GFilterVertexDeclaration.VertexDeclarationRHI, *VertexShader, PixelShader, *GeometryShader);
 				}
 
 			for (int32 CubeFace = 0; CubeFace < CubeFace_MAX; CubeFace++)
@@ -473,6 +474,7 @@ void FilterReflectionEnvironment(FSHVectorRGB3* OutIrradianceEnvironmentMap)
 				PixelShader->SetParameters(RHICmdList, CubeFace, MipIndex, Roughness, EffectiveSource);
 
 				DrawRectangle( 
+					RHICmdList,
 					ViewRect.Min.X, ViewRect.Min.Y, 
 					ViewRect.Width(), ViewRect.Height(),
 					ViewRect.Min.X, ViewRect.Min.Y, 
@@ -482,20 +484,20 @@ void FilterReflectionEnvironment(FSHVectorRGB3* OutIrradianceEnvironmentMap)
 					*VertexShader,
 					EDRF_UseTriangleOptimization);
 
-				RHICopyToResolveTarget(EffectiveRT.TargetableTexture, EffectiveRT.ShaderResourceTexture, true, FResolveParams(FResolveRect(), (ECubeFace)CubeFace, MipIndex));
+				RHICmdList.CopyToResolveTarget(EffectiveRT.TargetableTexture, EffectiveRT.ShaderResourceTexture, true, FResolveParams(FResolveRect(), (ECubeFace)CubeFace, MipIndex));
 			}
 		}
 			else
 			{
 				for (int32 CubeFace = 0; CubeFace < CubeFace_MAX; CubeFace++)
 				{
-					RHISetRenderTarget(EffectiveRT.TargetableTexture, MipIndex, CubeFace, NULL);
+					SetRenderTarget(RHICmdList, EffectiveRT.TargetableTexture, MipIndex, CubeFace, NULL);
 
 					const FIntRect ViewRect(0, 0, MipSize, MipSize);
-					RHISetViewport(0, 0, 0.0f, MipSize, MipSize, 1.0f);
-					RHISetRasterizerState(TStaticRasterizerState<FM_Solid,CM_None>::GetRHI());
-					RHISetDepthStencilState(TStaticDepthStencilState<false,CF_Always>::GetRHI());
-					RHISetBlendState(TStaticBlendState<>::GetRHI());
+					RHICmdList.SetViewport(0, 0, 0.0f, MipSize, MipSize, 1.0f);
+					RHICmdList.SetRasterizerState(TStaticRasterizerState<FM_Solid, CM_None>::GetRHI());
+					RHICmdList.SetDepthStencilState(TStaticDepthStencilState<false, CF_Always>::GetRHI());
+					RHICmdList.SetBlendState(TStaticBlendState<>::GetRHI());
 
 					TShaderMapRef<FScreenVSForGS> VertexShader(GetGlobalShaderMap());
 					TShaderMapRef< TCubeFilterPS<1> > CaptureCubemapArrayPixelShader(GetGlobalShaderMap());
@@ -506,18 +508,21 @@ void FilterReflectionEnvironment(FSHVectorRGB3* OutIrradianceEnvironmentMap)
 					{
 						PixelShader = *TShaderMapRef< TCubeFilterPS<1> >( GetGlobalShaderMap() );
 						static FGlobalBoundShaderState BoundShaderState;
-						SetGlobalBoundShaderState(BoundShaderState, GFilterVertexDeclaration.VertexDeclarationRHI, *VertexShader, PixelShader);
+						RHICmdList.CheckIsNull(); // need new approach for "static FGlobalBoundShaderState" for parallel rendering
+						SetGlobalBoundShaderState(RHICmdList, BoundShaderState, GFilterVertexDeclaration.VertexDeclarationRHI, *VertexShader, PixelShader);
 					}
 					else
 					{
 						PixelShader = *TShaderMapRef< TCubeFilterPS<0> >( GetGlobalShaderMap() );
 						static FGlobalBoundShaderState BoundShaderState;
-						SetGlobalBoundShaderState(BoundShaderState, GFilterVertexDeclaration.VertexDeclarationRHI, *VertexShader, PixelShader);
+						RHICmdList.CheckIsNull(); // need new approach for "static FGlobalBoundShaderState" for parallel rendering
+						SetGlobalBoundShaderState(RHICmdList, BoundShaderState, GFilterVertexDeclaration.VertexDeclarationRHI, *VertexShader, PixelShader);
 					}
 
 					PixelShader->SetParameters(RHICmdList, CubeFace, MipIndex, Roughness, EffectiveSource);
 
 					DrawRectangle( 
+						RHICmdList,
 						ViewRect.Min.X, ViewRect.Min.Y, 
 						ViewRect.Width(), ViewRect.Height(),
 						ViewRect.Min.X, ViewRect.Min.Y, 
@@ -526,7 +531,7 @@ void FilterReflectionEnvironment(FSHVectorRGB3* OutIrradianceEnvironmentMap)
 						FIntPoint(MipSize, MipSize),
 						*VertexShader);
 
-					RHICopyToResolveTarget(EffectiveRT.TargetableTexture, EffectiveRT.ShaderResourceTexture, true, FResolveParams(FResolveRect(), (ECubeFace)CubeFace, MipIndex));
+					RHICmdList.CopyToResolveTarget(EffectiveRT.TargetableTexture, EffectiveRT.ShaderResourceTexture, true, FResolveParams(FResolveRect(), (ECubeFace)CubeFace, MipIndex));
 				}
 			}
 		}
@@ -782,7 +787,7 @@ int32 FindOrAllocateCubemapIndex(FScene* Scene, const UReflectionCaptureComponen
 	return CaptureIndex;
 }
 
-void ClearScratchCubemaps()
+void ClearScratchCubemaps(FRHICommandList& RHICmdList)
 {
 	// Clear scratch render targets to a consistent but noticeable value
 	// This makes debugging capture issues much easier, otherwise the random contents from previous captures is shown
@@ -794,15 +799,15 @@ void ClearScratchCubemaps()
 	{
 		if (GSupportsGSRenderTargetLayerSwitchingToMips)
 		{
-			RHISetRenderTarget(RT0.TargetableTexture, MipIndex, NULL);
-			RHIClear(true, FLinearColor(0, 10000, 0, 0), false, 0, false, 0, FIntRect());
+			SetRenderTarget(RHICmdList, RT0.TargetableTexture, MipIndex, NULL);
+			RHICmdList.Clear(true, FLinearColor(0, 10000, 0, 0), false, 0, false, 0, FIntRect());
 		}
 		else
 		{
 			for (int32 CubeFace = 0; CubeFace < CubeFace_MAX; CubeFace++)
 			{
-				RHISetRenderTarget(RT0.TargetableTexture, MipIndex, CubeFace, NULL);
-				RHIClear(true, FLinearColor(0, 10000, 0, 0), false, 0, false, 0, FIntRect());
+				SetRenderTarget(RHICmdList, RT0.TargetableTexture, MipIndex, CubeFace, NULL);
+				RHICmdList.Clear(true, FLinearColor(0, 10000, 0, 0), false, 0, false, 0, FIntRect());
 			}
 		}
 	}
@@ -813,15 +818,15 @@ void ClearScratchCubemaps()
 	{
 		if (GSupportsGSRenderTargetLayerSwitchingToMips)
 		{
-			RHISetRenderTarget(RT1.TargetableTexture, MipIndex, NULL);
-			RHIClear(true, FLinearColor(0, 10000, 0, 0), false, 0, false, 0, FIntRect());
+			SetRenderTarget(RHICmdList, RT1.TargetableTexture, MipIndex, NULL);
+			RHICmdList.Clear(true, FLinearColor(0, 10000, 0, 0), false, 0, false, 0, FIntRect());
 		}
 		else
 		{
 			for (int32 CubeFace = 0; CubeFace < CubeFace_MAX; CubeFace++)
 			{
-				RHISetRenderTarget(RT1.TargetableTexture, MipIndex, CubeFace, NULL);
-				RHIClear(true, FLinearColor(0, 10000, 0, 0), false, 0, false, 0, FIntRect());
+				SetRenderTarget(RHICmdList, RT1.TargetableTexture, MipIndex, CubeFace, NULL);
+				RHICmdList.Clear(true, FLinearColor(0, 10000, 0, 0), false, 0, false, 0, FIntRect());
 			}
 		}
 	}
@@ -831,8 +836,7 @@ void ClearScratchCubemaps()
 void CaptureSceneToScratchCubemap(FRHICommandList& RHICmdList, FSceneRenderer* SceneRenderer, ECubeFace CubeFace, bool bCapturingForSkyLight, bool bLowerHemisphereIsBlack)
 {
 	FMemMark MemStackMark(FMemStack::Get());
-	//@todo-rco: RHIPacketList
-	RHICmdList.CheckIsNull();
+	RHICmdList.CheckIsNull(); //FDeferredUpdateResource::UpdateResources()
 	// update any resources that needed a deferred update
 	FDeferredUpdateResource::UpdateResources();
 	
@@ -845,6 +849,7 @@ void CaptureSceneToScratchCubemap(FRHICommandList& RHICmdList, FSceneRenderer* S
 #if PLATFORM_PS4 // @todo ps4 - this should be done a different way
 		// PS4 needs some code here to process the scene
 		extern void TEMP_PostReflectionCaptureRender();
+		RHICmdList.CheckIsNull(); //FDeferredUpdateResource::UpdateResources()
 		TEMP_PostReflectionCaptureRender();
 #endif
 
@@ -852,30 +857,30 @@ void CaptureSceneToScratchCubemap(FRHICommandList& RHICmdList, FSceneRenderer* S
 		FSceneRenderTargetItem& EffectiveColorRT =  GSceneRenderTargets.ReflectionColorScratchCubemap[0]->GetRenderTargetItem();
 
 		{
-			RHICmdList.CheckIsNull();
 			SCOPED_DRAW_EVENT(CubeMapCopyScene, DEC_SCENE_ITEMS);
 			if (GSupportsGSRenderTargetLayerSwitchingToMips)
 			{
 			// Copy the captured scene into the cubemap face
-			RHISetRenderTarget(EffectiveColorRT.TargetableTexture, NULL);
+				SetRenderTarget(RHICmdList, EffectiveColorRT.TargetableTexture, NULL);
 
 			const FIntRect ViewRect(0, 0, EffectiveSize, EffectiveSize);
-			RHISetViewport(0, 0, 0.0f, EffectiveSize, EffectiveSize, 1.0f);
-			RHISetRasterizerState(TStaticRasterizerState<FM_Solid,CM_None>::GetRHI());
-			RHISetDepthStencilState(TStaticDepthStencilState<false,CF_Always>::GetRHI());
-			RHISetBlendState(TStaticBlendState<>::GetRHI());
+				RHICmdList.SetViewport(0, 0, 0.0f, EffectiveSize, EffectiveSize, 1.0f);
+				RHICmdList.SetRasterizerState(TStaticRasterizerState<FM_Solid, CM_None>::GetRHI());
+				RHICmdList.SetDepthStencilState(TStaticDepthStencilState<false, CF_Always>::GetRHI());
+				RHICmdList.SetBlendState(TStaticBlendState<>::GetRHI());
 
 			TShaderMapRef<FCopyToCubeFaceVS> VertexShader(GetGlobalShaderMap());
 			TShaderMapRef<FRouteToCubeFaceGS> GeometryShader(GetGlobalShaderMap());
 
 			TShaderMapRef<FCopySceneColorToCubeFacePS> PixelShader(GetGlobalShaderMap());
-			SetGlobalBoundShaderState(CopyColorCubemapBoundShaderState, GFilterVertexDeclaration.VertexDeclarationRHI, *VertexShader, *PixelShader, *GeometryShader);
+				SetGlobalBoundShaderState(RHICmdList, CopyColorCubemapBoundShaderState, GFilterVertexDeclaration.VertexDeclarationRHI, *VertexShader, *PixelShader, *GeometryShader);
 			
 				PixelShader->SetParameters(RHICmdList, SceneRenderer->Views[0], bCapturingForSkyLight, bLowerHemisphereIsBlack);
 				VertexShader->SetParameters(RHICmdList, SceneRenderer->Views[0]);
 				GeometryShader->SetParameters(RHICmdList, (uint32)CubeFace);
 
 			DrawRectangle( 
+					RHICmdList,
 				ViewRect.Min.X, ViewRect.Min.Y, 
 				ViewRect.Width(), ViewRect.Height(),
 				ViewRect.Min.X, ViewRect.Min.Y, 
@@ -885,27 +890,28 @@ void CaptureSceneToScratchCubemap(FRHICommandList& RHICmdList, FSceneRenderer* S
 				*VertexShader,
 				EDRF_UseTriangleOptimization);
 
-			RHICopyToResolveTarget(EffectiveColorRT.TargetableTexture, EffectiveColorRT.ShaderResourceTexture, true, FResolveParams(FResolveRect()));
+				RHICmdList.CopyToResolveTarget(EffectiveColorRT.TargetableTexture, EffectiveColorRT.ShaderResourceTexture, true, FResolveParams(FResolveRect()));
 		}
 			else
 			{
 				// Copy the captured scene into the cubemap face
-				RHISetRenderTarget(EffectiveColorRT.TargetableTexture, 0, CubeFace, NULL);
+				SetRenderTarget(RHICmdList, EffectiveColorRT.TargetableTexture, 0, CubeFace, NULL);
 
 				const FIntRect ViewRect(0, 0, EffectiveSize, EffectiveSize);
-				RHISetViewport(0, 0, 0.0f, EffectiveSize, EffectiveSize, 1.0f);
-				RHISetRasterizerState(TStaticRasterizerState<FM_Solid,CM_None>::GetRHI());
-				RHISetDepthStencilState(TStaticDepthStencilState<false,CF_Always>::GetRHI());
-				RHISetBlendState(TStaticBlendState<>::GetRHI());
+				RHICmdList.SetViewport(0, 0, 0.0f, EffectiveSize, EffectiveSize, 1.0f);
+				RHICmdList.SetRasterizerState(TStaticRasterizerState<FM_Solid, CM_None>::GetRHI());
+				RHICmdList.SetDepthStencilState(TStaticDepthStencilState<false, CF_Always>::GetRHI());
+				RHICmdList.SetBlendState(TStaticBlendState<>::GetRHI());
 
 				TShaderMapRef<FCopyToCubeFaceVS> VertexShader(GetGlobalShaderMap());
 				TShaderMapRef<FCopySceneColorToCubeFacePS> PixelShader(GetGlobalShaderMap());
-				SetGlobalBoundShaderState(CopyColorCubemapBoundShaderState, GFilterVertexDeclaration.VertexDeclarationRHI, *VertexShader, *PixelShader);
+				SetGlobalBoundShaderState(RHICmdList, CopyColorCubemapBoundShaderState, GFilterVertexDeclaration.VertexDeclarationRHI, *VertexShader, *PixelShader);
 
 				PixelShader->SetParameters(RHICmdList, SceneRenderer->Views[0], bCapturingForSkyLight, bLowerHemisphereIsBlack);
 				VertexShader->SetParameters(RHICmdList, SceneRenderer->Views[0]);
 
 				DrawRectangle( 
+					RHICmdList,
 					ViewRect.Min.X, ViewRect.Min.Y, 
 					ViewRect.Width(), ViewRect.Height(),
 					ViewRect.Min.X, ViewRect.Min.Y, 
@@ -914,7 +920,7 @@ void CaptureSceneToScratchCubemap(FRHICommandList& RHICmdList, FSceneRenderer* S
 					GSceneRenderTargets.GetBufferSizeXY(),
 					*VertexShader);
 
-				RHICopyToResolveTarget(EffectiveColorRT.TargetableTexture, EffectiveColorRT.ShaderResourceTexture, true, FResolveParams(FResolveRect(), CubeFace));
+				RHICmdList.CopyToResolveTarget(EffectiveColorRT.TargetableTexture, EffectiveColorRT.ShaderResourceTexture, true, FResolveParams(FResolveRect(), CubeFace));
 			}
 		}
 	}
@@ -926,25 +932,23 @@ void CopyCubemapToScratchCubemap(FRHICommandList& RHICmdList, UTextureCube* Sour
 {
 	check(SourceCubemap);
 	
-	ClearScratchCubemaps();
+	ClearScratchCubemaps(RHICmdList);
 
 	const int32 EffectiveSize = GReflectionCaptureSize;
 	FSceneRenderTargetItem& EffectiveColorRT =  GSceneRenderTargets.ReflectionColorScratchCubemap[0]->GetRenderTargetItem();
 
 	if (GSupportsGSRenderTargetLayerSwitchingToMips)
 	{
-		RHICmdList.CheckIsNull();
-
 		// Copy the captured scene into the cubemap face
-		RHISetRenderTarget(EffectiveColorRT.TargetableTexture, NULL);
+		SetRenderTarget(RHICmdList, EffectiveColorRT.TargetableTexture, NULL);
 
 		const FTexture* SourceCubemapResource = SourceCubemap->Resource;
 		const FIntPoint SourceDimensions(SourceCubemapResource->GetSizeX(), SourceCubemapResource->GetSizeY());
 		const FIntRect ViewRect(0, 0, EffectiveSize, EffectiveSize);
-		RHISetViewport(0, 0, 0.0f, EffectiveSize, EffectiveSize, 1.0f);
-		RHISetRasterizerState(TStaticRasterizerState<FM_Solid,CM_None>::GetRHI());
-		RHISetDepthStencilState(TStaticDepthStencilState<false,CF_Always>::GetRHI());
-		RHISetBlendState(TStaticBlendState<>::GetRHI());
+		RHICmdList.SetViewport(0, 0, 0.0f, EffectiveSize, EffectiveSize, 1.0f);
+		RHICmdList.SetRasterizerState(TStaticRasterizerState<FM_Solid, CM_None>::GetRHI());
+		RHICmdList.SetDepthStencilState(TStaticDepthStencilState<false, CF_Always>::GetRHI());
+		RHICmdList.SetBlendState(TStaticBlendState<>::GetRHI());
 
 		TShaderMapRef<FScreenVSForGS> VertexShader(GetGlobalShaderMap());
 		TShaderMapRef<FDownsampleGS> GeometryShader(GetGlobalShaderMap());
@@ -952,11 +956,12 @@ void CopyCubemapToScratchCubemap(FRHICommandList& RHICmdList, UTextureCube* Sour
 
 		for (uint32 CubeFace = 0; CubeFace < CubeFace_MAX; CubeFace++)
 		{
-			SetGlobalBoundShaderState(CopyFromCubemapToCubemapBoundShaderState, GFilterVertexDeclaration.VertexDeclarationRHI, *VertexShader, *PixelShader, *GeometryShader);
+			SetGlobalBoundShaderState(RHICmdList, CopyFromCubemapToCubemapBoundShaderState, GFilterVertexDeclaration.VertexDeclarationRHI, *VertexShader, *PixelShader, *GeometryShader);
 			GeometryShader->SetParameters(RHICmdList, CubeFace);
 			PixelShader->SetParameters(RHICmdList, SourceCubemapResource, CubeFace, bIsSkyLight, bLowerHemisphereIsBlack);
 
 			DrawRectangle( 
+				RHICmdList,
 				ViewRect.Min.X, ViewRect.Min.Y, 
 				ViewRect.Width(), ViewRect.Height(),
 				0, 0, 
@@ -966,30 +971,31 @@ void CopyCubemapToScratchCubemap(FRHICommandList& RHICmdList, UTextureCube* Sour
 				*VertexShader);
 		}
 
-		RHICopyToResolveTarget(EffectiveColorRT.TargetableTexture, EffectiveColorRT.ShaderResourceTexture, true, FResolveParams(FResolveRect()));
+		RHICmdList.CopyToResolveTarget(EffectiveColorRT.TargetableTexture, EffectiveColorRT.ShaderResourceTexture, true, FResolveParams(FResolveRect()));
 	}
 	else
 	{
 		for (uint32 CubeFace = 0; CubeFace < CubeFace_MAX; CubeFace++)
 		{
 			// Copy the captured scene into the cubemap face
-			RHISetRenderTarget(EffectiveColorRT.TargetableTexture, 0, CubeFace, NULL);
+			SetRenderTarget(RHICmdList, EffectiveColorRT.TargetableTexture, 0, CubeFace, NULL);
 
 			const FTexture* SourceCubemapResource = SourceCubemap->Resource;
 			const FIntPoint SourceDimensions(SourceCubemapResource->GetSizeX(), SourceCubemapResource->GetSizeY());
 			const FIntRect ViewRect(0, 0, EffectiveSize, EffectiveSize);
-			RHISetViewport(0, 0, 0.0f, EffectiveSize, EffectiveSize, 1.0f);
-			RHISetRasterizerState(TStaticRasterizerState<FM_Solid,CM_None>::GetRHI());
-			RHISetDepthStencilState(TStaticDepthStencilState<false,CF_Always>::GetRHI());
-			RHISetBlendState(TStaticBlendState<>::GetRHI());
+			RHICmdList.SetViewport(0, 0, 0.0f, EffectiveSize, EffectiveSize, 1.0f);
+			RHICmdList.SetRasterizerState(TStaticRasterizerState<FM_Solid, CM_None>::GetRHI());
+			RHICmdList.SetDepthStencilState(TStaticDepthStencilState<false, CF_Always>::GetRHI());
+			RHICmdList.SetBlendState(TStaticBlendState<>::GetRHI());
 
 			TShaderMapRef<FScreenVSForGS> VertexShader(GetGlobalShaderMap());
 			TShaderMapRef<FCopyCubemapToCubeFacePS> PixelShader(GetGlobalShaderMap());
 
-			SetGlobalBoundShaderState(CopyFromCubemapToCubemapBoundShaderState, GFilterVertexDeclaration.VertexDeclarationRHI, *VertexShader, *PixelShader);
+			SetGlobalBoundShaderState(RHICmdList, CopyFromCubemapToCubemapBoundShaderState, GFilterVertexDeclaration.VertexDeclarationRHI, *VertexShader, *PixelShader);
 			PixelShader->SetParameters(RHICmdList, SourceCubemapResource, CubeFace, bIsSkyLight, bLowerHemisphereIsBlack);
 
 			DrawRectangle( 
+				RHICmdList,
 				ViewRect.Min.X, ViewRect.Min.Y, 
 				ViewRect.Width(), ViewRect.Height(),
 				0, 0, 
@@ -998,7 +1004,7 @@ void CopyCubemapToScratchCubemap(FRHICommandList& RHICmdList, UTextureCube* Sour
 				SourceDimensions,
 				*VertexShader);
 
-			RHICopyToResolveTarget(EffectiveColorRT.TargetableTexture, EffectiveColorRT.ShaderResourceTexture, true, FResolveParams(FResolveRect(), (ECubeFace)CubeFace));
+			RHICmdList.CopyToResolveTarget(EffectiveColorRT.TargetableTexture, EffectiveColorRT.ShaderResourceTexture, true, FResolveParams(FResolveRect(), (ECubeFace)CubeFace));
 		}
 	}
 }
@@ -1497,7 +1503,10 @@ void FScene::UpdateReflectionCaptureContents(UReflectionCaptureComponent* Captur
 			ENQUEUE_UNIQUE_RENDER_COMMAND( 
 				FilterCommand,
 			{
-				FilterReflectionEnvironment(NULL);
+				//@todo-rco: RHIPacketList
+				FRHICommandList& RHICmdList = FRHICommandList::GetNullRef();
+
+				FilterReflectionEnvironment(RHICmdList, NULL);
 			});
 
 			// Create a proxy to represent the reflection capture to the rendering thread
@@ -1553,7 +1562,10 @@ void FScene::UpdateSkyCaptureContents(const USkyLightComponent* CaptureComponent
 			ClearCommand,
 		{
 			//@todo realtime skylight updates - skip this
-			ClearScratchCubemaps();
+			//@todo-rco: RHIPacketList
+			FRHICommandList& RHICmdList = FRHICommandList::GetNullRef();
+
+			ClearScratchCubemaps(RHICmdList);
 		});
 
 		if (CaptureComponent->SourceType == SLS_CapturedScene)
@@ -1581,7 +1593,10 @@ void FScene::UpdateSkyCaptureContents(const USkyLightComponent* CaptureComponent
 			FilterCommand,
 			FSHVectorRGB3*, IrradianceEnvironmentMap, &OutIrradianceEnvironmentMap,
 		{
-			FilterReflectionEnvironment(IrradianceEnvironmentMap);
+			//@todo-rco: RHIPacketList
+			FRHICommandList& RHICmdList = FRHICommandList::GetNullRef();
+
+			FilterReflectionEnvironment(RHICmdList, IrradianceEnvironmentMap);
 		});
 
 		// Wait until the SH coefficients have been written out by the RT before returning
diff --git a/Engine/Source/Runtime/Renderer/Private/ReflectionEnvironmentCapture.h b/Engine/Source/Runtime/Renderer/Private/ReflectionEnvironmentCapture.h
--- a/Engine/Source/Runtime/Renderer/Private/ReflectionEnvironmentCapture.h
+++ b/Engine/Source/Runtime/Renderer/Private/ReflectionEnvironmentCapture.h
@@ -8,7 +8,7 @@
 
 extern ENGINE_API int32 GReflectionCaptureSize;
 
-extern void ComputeDiffuseIrradiance(FTextureRHIRef LightingSource, int32 LightingSourceMipIndex, FSHVectorRGB3* OutIrradianceEnvironmentMap);
+extern void ComputeDiffuseIrradiance(FRHICommandList& RHICmdList, FTextureRHIRef LightingSource, int32 LightingSourceMipIndex, FSHVectorRGB3* OutIrradianceEnvironmentMap);
 
 class FDownsampleGS : public FGlobalShader
 {
diff --git a/Engine/Source/Runtime/Renderer/Private/ReflectionEnvironmentDiffuseIrradiance.cpp b/Engine/Source/Runtime/Renderer/Private/ReflectionEnvironmentDiffuseIrradiance.cpp
--- a/Engine/Source/Runtime/Renderer/Private/ReflectionEnvironmentDiffuseIrradiance.cpp
+++ b/Engine/Source/Runtime/Renderer/Private/ReflectionEnvironmentDiffuseIrradiance.cpp
@@ -236,10 +236,8 @@ IMPLEMENT_SHADER_TYPE(,FAccumulateCubeFacesPS,TEXT("ReflectionEnvironmentShaders
 
 FGlobalBoundShaderState AccumulateCubeFacesBoundShaderState;
 
-void ComputeDiffuseIrradiance(FTextureRHIRef LightingSource, int32 LightingSourceMipIndex, FSHVectorRGB3* OutIrradianceEnvironmentMap)
+void ComputeDiffuseIrradiance(FRHICommandList& RHICmdList, FTextureRHIRef LightingSource, int32 LightingSourceMipIndex, FSHVectorRGB3* OutIrradianceEnvironmentMap)
 {
-	//@todo-rco: RHIPacketList
-	FRHICommandList& RHICmdList = FRHICommandList::GetNullRef();
 	for (int32 CoefficientIndex = 0; CoefficientIndex < FSHVector3::MaxSHBasis; CoefficientIndex++)
 	{
 		// Copy the starting mip from the lighting texture, apply texel area weighting and appropriate SH coefficient
@@ -249,59 +247,61 @@ void ComputeDiffuseIrradiance(FTextureRHIRef LightingSource, int32 LightingSourc
 			FSceneRenderTargetItem& EffectiveRT = GetEffectiveDiffuseIrradianceRenderTarget(MipIndex);
 			if (GSupportsGSRenderTargetLayerSwitchingToMips)
 			{
-			RHISetRenderTarget(EffectiveRT.TargetableTexture, 0, NULL);
+				SetRenderTarget(RHICmdList, EffectiveRT.TargetableTexture, 0, NULL);
 
-			const FIntRect ViewRect(0, 0, MipSize, MipSize);
-			RHISetViewport(0, 0, 0.0f, MipSize, MipSize, 1.0f);
-			RHISetRasterizerState(TStaticRasterizerState<FM_Solid,CM_None>::GetRHI());
-			RHISetDepthStencilState(TStaticDepthStencilState<false,CF_Always>::GetRHI());
-			RHISetBlendState(TStaticBlendState<>::GetRHI());
+				const FIntRect ViewRect(0, 0, MipSize, MipSize);
+				RHICmdList.SetViewport(0, 0, 0.0f, MipSize, MipSize, 1.0f);
+				RHICmdList.SetRasterizerState(TStaticRasterizerState<FM_Solid, CM_None>::GetRHI());
+				RHICmdList.SetDepthStencilState(TStaticDepthStencilState<false, CF_Always>::GetRHI());
+				RHICmdList.SetBlendState(TStaticBlendState<>::GetRHI());
 
-			TShaderMapRef<FScreenVSForGS> VertexShader(GetGlobalShaderMap());
-			TShaderMapRef<FDownsampleGS> GeometryShader(GetGlobalShaderMap());
-			TShaderMapRef<FCopyDiffuseIrradiancePS> PixelShader(GetGlobalShaderMap());
+				TShaderMapRef<FScreenVSForGS> VertexShader(GetGlobalShaderMap());
+				TShaderMapRef<FDownsampleGS> GeometryShader(GetGlobalShaderMap());
+				TShaderMapRef<FCopyDiffuseIrradiancePS> PixelShader(GetGlobalShaderMap());
 
-			SetGlobalBoundShaderState(CopyDiffuseIrradianceShaderState, GFilterVertexDeclaration.VertexDeclarationRHI, *VertexShader, *PixelShader, *GeometryShader);
+				SetGlobalBoundShaderState(RHICmdList, CopyDiffuseIrradianceShaderState, GFilterVertexDeclaration.VertexDeclarationRHI, *VertexShader, *PixelShader, *GeometryShader);
 
-			// Draw each face with a geometry shader that routes to the correct render target slice
-			for (int32 CubeFace = 0; CubeFace < CubeFace_MAX; CubeFace++)
-			{
-				GeometryShader->SetParameters(RHICmdList, CubeFace);
+				// Draw each face with a geometry shader that routes to the correct render target slice
+				for (int32 CubeFace = 0; CubeFace < CubeFace_MAX; CubeFace++)
+				{
+					GeometryShader->SetParameters(RHICmdList, CubeFace);
 
-				PixelShader->SetParameters(RHICmdList, CubeFace, LightingSourceMipIndex, CoefficientIndex, MipSize, LightingSource);
+					PixelShader->SetParameters(RHICmdList, CubeFace, LightingSourceMipIndex, CoefficientIndex, MipSize, LightingSource);
 
-				DrawRectangle( 
-					ViewRect.Min.X, ViewRect.Min.Y, 
-					ViewRect.Width(), ViewRect.Height(),
-					ViewRect.Min.X, ViewRect.Min.Y, 
-					ViewRect.Width(), ViewRect.Height(),
-					FIntPoint(ViewRect.Width(), ViewRect.Height()),
-					FIntPoint(MipSize, MipSize),
-					*VertexShader);
+					DrawRectangle( 
+						RHICmdList,
+						ViewRect.Min.X, ViewRect.Min.Y, 
+						ViewRect.Width(), ViewRect.Height(),
+						ViewRect.Min.X, ViewRect.Min.Y, 
+						ViewRect.Width(), ViewRect.Height(),
+						FIntPoint(ViewRect.Width(), ViewRect.Height()),
+						FIntPoint(MipSize, MipSize),
+						*VertexShader);
 
-				RHICopyToResolveTarget(EffectiveRT.TargetableTexture, EffectiveRT.ShaderResourceTexture, true, FResolveParams(FResolveRect(), (ECubeFace)CubeFace, MipIndex));
+					RHICmdList.CopyToResolveTarget(EffectiveRT.TargetableTexture, EffectiveRT.ShaderResourceTexture, true, FResolveParams(FResolveRect(), (ECubeFace)CubeFace, MipIndex));
+				}
 			}
-		}
 			else
 			{
 				for (int32 CubeFace = 0; CubeFace < CubeFace_MAX; CubeFace++)
 				{
-					RHISetRenderTarget(EffectiveRT.TargetableTexture, 0, CubeFace, NULL);
+					SetRenderTarget(RHICmdList, EffectiveRT.TargetableTexture, 0, CubeFace, NULL);
 					
 					const FIntRect ViewRect(0, 0, MipSize, MipSize);
-					RHISetViewport(0, 0, 0.0f, MipSize, MipSize, 1.0f);
-					RHISetRasterizerState(TStaticRasterizerState<FM_Solid,CM_None>::GetRHI());
-					RHISetDepthStencilState(TStaticDepthStencilState<false,CF_Always>::GetRHI());
-					RHISetBlendState(TStaticBlendState<>::GetRHI());
+					RHICmdList.SetViewport(0, 0, 0.0f, MipSize, MipSize, 1.0f);
+					RHICmdList.SetRasterizerState(TStaticRasterizerState<FM_Solid, CM_None>::GetRHI());
+					RHICmdList.SetDepthStencilState(TStaticDepthStencilState<false, CF_Always>::GetRHI());
+					RHICmdList.SetBlendState(TStaticBlendState<>::GetRHI());
 					
 					TShaderMapRef<FScreenVSForGS> VertexShader(GetGlobalShaderMap());
 					TShaderMapRef<FCopyDiffuseIrradiancePS> PixelShader(GetGlobalShaderMap());
 					
-					SetGlobalBoundShaderState(CopyDiffuseIrradianceShaderState, GFilterVertexDeclaration.VertexDeclarationRHI, *VertexShader, *PixelShader);
+					SetGlobalBoundShaderState(RHICmdList, CopyDiffuseIrradianceShaderState, GFilterVertexDeclaration.VertexDeclarationRHI, *VertexShader, *PixelShader);
 					
 					PixelShader->SetParameters(RHICmdList, CubeFace, LightingSourceMipIndex, CoefficientIndex, MipSize, LightingSource);
 					
 					DrawRectangle(
+						RHICmdList,
 						ViewRect.Min.X, ViewRect.Min.Y,
 						ViewRect.Width(), ViewRect.Height(),
 						ViewRect.Min.X, ViewRect.Min.Y,
@@ -310,7 +310,7 @@ void ComputeDiffuseIrradiance(FTextureRHIRef LightingSource, int32 LightingSourc
 						FIntPoint(MipSize, MipSize),
 						*VertexShader);
 					
-					RHICopyToResolveTarget(EffectiveRT.TargetableTexture, EffectiveRT.ShaderResourceTexture, true, FResolveParams(FResolveRect(), (ECubeFace)CubeFace, MipIndex));
+					RHICmdList.CopyToResolveTarget(EffectiveRT.TargetableTexture, EffectiveRT.ShaderResourceTexture, true, FResolveParams(FResolveRect(), (ECubeFace)CubeFace, MipIndex));
 				}
 			}
 		}
@@ -330,59 +330,61 @@ void ComputeDiffuseIrradiance(FTextureRHIRef LightingSource, int32 LightingSourc
 
 				if (GSupportsGSRenderTargetLayerSwitchingToMips)
 				{
-				RHISetRenderTarget(EffectiveRT.TargetableTexture, MipIndex, NULL);
+					SetRenderTarget(RHICmdList, EffectiveRT.TargetableTexture, MipIndex, NULL);
 
-				const FIntRect ViewRect(0, 0, MipSize, MipSize);
-				RHISetViewport(0, 0, 0.0f, MipSize, MipSize, 1.0f);
-				RHISetRasterizerState(TStaticRasterizerState<FM_Solid,CM_None>::GetRHI());
-				RHISetDepthStencilState(TStaticDepthStencilState<false,CF_Always>::GetRHI());
-				RHISetBlendState(TStaticBlendState<>::GetRHI());
+					const FIntRect ViewRect(0, 0, MipSize, MipSize);
+					RHICmdList.SetViewport(0, 0, 0.0f, MipSize, MipSize, 1.0f);
+					RHICmdList.SetRasterizerState(TStaticRasterizerState<FM_Solid, CM_None>::GetRHI());
+					RHICmdList.SetDepthStencilState(TStaticDepthStencilState<false, CF_Always>::GetRHI());
+					RHICmdList.SetBlendState(TStaticBlendState<>::GetRHI());
 
-				TShaderMapRef<FScreenVSForGS> VertexShader(GetGlobalShaderMap());
-				TShaderMapRef<FDownsampleGS> GeometryShader(GetGlobalShaderMap());
-				TShaderMapRef<FAccumulateDiffuseIrradiancePS> PixelShader(GetGlobalShaderMap());
+					TShaderMapRef<FScreenVSForGS> VertexShader(GetGlobalShaderMap());
+					TShaderMapRef<FDownsampleGS> GeometryShader(GetGlobalShaderMap());
+					TShaderMapRef<FAccumulateDiffuseIrradiancePS> PixelShader(GetGlobalShaderMap());
 
-				SetGlobalBoundShaderState(DiffuseIrradianceAccumulateShaderState, GFilterVertexDeclaration.VertexDeclarationRHI, *VertexShader, *PixelShader, *GeometryShader);
+					SetGlobalBoundShaderState(RHICmdList, DiffuseIrradianceAccumulateShaderState, GFilterVertexDeclaration.VertexDeclarationRHI, *VertexShader, *PixelShader, *GeometryShader);
 
-				// Draw each face with a geometry shader that routes to the correct render target slice
-				for (int32 CubeFace = 0; CubeFace < CubeFace_MAX; CubeFace++)
-				{
-					GeometryShader->SetParameters(RHICmdList, CubeFace);
+					// Draw each face with a geometry shader that routes to the correct render target slice
+					for (int32 CubeFace = 0; CubeFace < CubeFace_MAX; CubeFace++)
+					{
+						GeometryShader->SetParameters(RHICmdList, CubeFace);
 
-					PixelShader->SetParameters(RHICmdList, CubeFace, NumMips, SourceMipIndex, CoefficientIndex, EffectiveSource.ShaderResourceTexture);
+						PixelShader->SetParameters(RHICmdList, CubeFace, NumMips, SourceMipIndex, CoefficientIndex, EffectiveSource.ShaderResourceTexture);
 
-					DrawRectangle( 
-						ViewRect.Min.X, ViewRect.Min.Y, 
-						ViewRect.Width(), ViewRect.Height(),
-						ViewRect.Min.X, ViewRect.Min.Y, 
-						ViewRect.Width(), ViewRect.Height(),
-						FIntPoint(ViewRect.Width(), ViewRect.Height()),
-						FIntPoint(MipSize, MipSize),
-						*VertexShader);
+						DrawRectangle( 
+							RHICmdList,
+							ViewRect.Min.X, ViewRect.Min.Y, 
+							ViewRect.Width(), ViewRect.Height(),
+							ViewRect.Min.X, ViewRect.Min.Y, 
+							ViewRect.Width(), ViewRect.Height(),
+							FIntPoint(ViewRect.Width(), ViewRect.Height()),
+							FIntPoint(MipSize, MipSize),
+							*VertexShader);
 
-					RHICopyToResolveTarget(EffectiveRT.TargetableTexture, EffectiveRT.ShaderResourceTexture, true, FResolveParams(FResolveRect(), (ECubeFace)CubeFace, MipIndex));
+						RHICmdList.CopyToResolveTarget(EffectiveRT.TargetableTexture, EffectiveRT.ShaderResourceTexture, true, FResolveParams(FResolveRect(), (ECubeFace)CubeFace, MipIndex));
+					}
 				}
-			}
 				else
 				{
 					for (int32 CubeFace = 0; CubeFace < CubeFace_MAX; CubeFace++)
 					{
-						RHISetRenderTarget(EffectiveRT.TargetableTexture, MipIndex, CubeFace, NULL);
+						SetRenderTarget(RHICmdList, EffectiveRT.TargetableTexture, MipIndex, CubeFace, NULL);
 						
 						const FIntRect ViewRect(0, 0, MipSize, MipSize);
-						RHISetViewport(0, 0, 0.0f, MipSize, MipSize, 1.0f);
-						RHISetRasterizerState(TStaticRasterizerState<FM_Solid,CM_None>::GetRHI());
-						RHISetDepthStencilState(TStaticDepthStencilState<false,CF_Always>::GetRHI());
-						RHISetBlendState(TStaticBlendState<>::GetRHI());
+						RHICmdList.SetViewport(0, 0, 0.0f, MipSize, MipSize, 1.0f);
+						RHICmdList.SetRasterizerState(TStaticRasterizerState<FM_Solid, CM_None>::GetRHI());
+						RHICmdList.SetDepthStencilState(TStaticDepthStencilState<false, CF_Always>::GetRHI());
+						RHICmdList.SetBlendState(TStaticBlendState<>::GetRHI());
 						
 						TShaderMapRef<FScreenVSForGS> VertexShader(GetGlobalShaderMap());
 						TShaderMapRef<FAccumulateDiffuseIrradiancePS> PixelShader(GetGlobalShaderMap());
 						
-						SetGlobalBoundShaderState(DiffuseIrradianceAccumulateShaderState, GFilterVertexDeclaration.VertexDeclarationRHI, *VertexShader, *PixelShader);
+						SetGlobalBoundShaderState(RHICmdList, DiffuseIrradianceAccumulateShaderState, GFilterVertexDeclaration.VertexDeclarationRHI, *VertexShader, *PixelShader);
 						
 						PixelShader->SetParameters(RHICmdList, CubeFace, NumMips, SourceMipIndex, CoefficientIndex, EffectiveSource.ShaderResourceTexture);
 						
 						DrawRectangle(
+							RHICmdList,
 							ViewRect.Min.X, ViewRect.Min.Y,
 							ViewRect.Width(), ViewRect.Height(),
 							ViewRect.Min.X, ViewRect.Min.Y,
@@ -391,7 +393,7 @@ void ComputeDiffuseIrradiance(FTextureRHIRef LightingSource, int32 LightingSourc
 							FIntPoint(MipSize, MipSize),
 							*VertexShader);
 						
-						RHICopyToResolveTarget(EffectiveRT.TargetableTexture, EffectiveRT.ShaderResourceTexture, true, FResolveParams(FResolveRect(), (ECubeFace)CubeFace, MipIndex));
+						RHICmdList.CopyToResolveTarget(EffectiveRT.TargetableTexture, EffectiveRT.ShaderResourceTexture, true, FResolveParams(FResolveRect(), (ECubeFace)CubeFace, MipIndex));
 					}
 				}
 			}
@@ -400,17 +402,17 @@ void ComputeDiffuseIrradiance(FTextureRHIRef LightingSource, int32 LightingSourc
 		{
 			// Gather the cubemap face results and normalize, copy this coefficient to GSceneRenderTargets.SkySHIrradianceMap
 			FSceneRenderTargetItem& EffectiveRT = GSceneRenderTargets.SkySHIrradianceMap->GetRenderTargetItem();
-			RHISetRenderTarget(EffectiveRT.TargetableTexture, 0, NULL);
+			SetRenderTarget(RHICmdList, EffectiveRT.TargetableTexture, 0, NULL);
 
 			const FIntRect ViewRect(CoefficientIndex, 0, CoefficientIndex + 1, 1);
-			RHISetViewport(0, 0, 0.0f, FSHVector3::MaxSHBasis, 1, 1.0f);
-			RHISetRasterizerState(TStaticRasterizerState<FM_Solid,CM_None>::GetRHI());
-			RHISetDepthStencilState(TStaticDepthStencilState<false,CF_Always>::GetRHI());
-			RHISetBlendState(TStaticBlendState<>::GetRHI());
+			RHICmdList.SetViewport(0, 0, 0.0f, FSHVector3::MaxSHBasis, 1, 1.0f);
+			RHICmdList.SetRasterizerState(TStaticRasterizerState<FM_Solid, CM_None>::GetRHI());
+			RHICmdList.SetDepthStencilState(TStaticDepthStencilState<false, CF_Always>::GetRHI());
+			RHICmdList.SetBlendState(TStaticBlendState<>::GetRHI());
 
 			TShaderMapRef<FScreenVS> VertexShader(GetGlobalShaderMap());
 			TShaderMapRef<FAccumulateCubeFacesPS> PixelShader(GetGlobalShaderMap());
-			SetGlobalBoundShaderState(AccumulateCubeFacesBoundShaderState, GFilterVertexDeclaration.VertexDeclarationRHI, *VertexShader, *PixelShader);
+			SetGlobalBoundShaderState(RHICmdList, AccumulateCubeFacesBoundShaderState, GFilterVertexDeclaration.VertexDeclarationRHI, *VertexShader, *PixelShader);
 
 			const int32 SourceMipIndex = NumMips - 1;
 			const int32 MipSize = 1;
@@ -418,6 +420,7 @@ void ComputeDiffuseIrradiance(FTextureRHIRef LightingSource, int32 LightingSourc
 			PixelShader->SetParameters(RHICmdList, SourceMipIndex, EffectiveSource.ShaderResourceTexture);
 
 			DrawRectangle( 
+				RHICmdList,
 				ViewRect.Min.X, ViewRect.Min.Y, 
 				ViewRect.Width(), ViewRect.Height(),
 				0, 0, 
@@ -426,7 +429,7 @@ void ComputeDiffuseIrradiance(FTextureRHIRef LightingSource, int32 LightingSourc
 				FIntPoint(MipSize, MipSize),
 				*VertexShader);
 
-			RHICopyToResolveTarget(EffectiveRT.TargetableTexture, EffectiveRT.ShaderResourceTexture, true, FResolveParams());
+			RHICmdList.CopyToResolveTarget(EffectiveRT.TargetableTexture, EffectiveRT.ShaderResourceTexture, true, FResolveParams());
 		}
 	}
 
@@ -434,6 +437,7 @@ void ComputeDiffuseIrradiance(FTextureRHIRef LightingSource, int32 LightingSourc
 		// Read back the completed SH environment map
 		FSceneRenderTargetItem& EffectiveRT = GSceneRenderTargets.SkySHIrradianceMap->GetRenderTargetItem();
 		check(EffectiveRT.ShaderResourceTexture->GetFormat() == PF_FloatRGBA);
+		RHICmdList.CheckIsNull(); // synchronous readback 
 
 		TArray<FFloat16Color> SurfaceData;
 		RHIReadSurfaceFloatData(EffectiveRT.ShaderResourceTexture, FIntRect(0, 0, FSHVector3::MaxSHBasis, 1), SurfaceData, CubeFace_PosX, 0, 0);
diff --git a/Engine/Source/Runtime/Renderer/Private/SceneCaptureRendering.cpp b/Engine/Source/Runtime/Renderer/Private/SceneCaptureRendering.cpp
--- a/Engine/Source/Runtime/Renderer/Private/SceneCaptureRendering.cpp
+++ b/Engine/Source/Runtime/Renderer/Private/SceneCaptureRendering.cpp
@@ -15,6 +15,7 @@
 
 void UpdateSceneCaptureContent_RenderThread(FSceneRenderer* SceneRenderer, FTextureRenderTargetResource* TextureRenderTarget, const FName OwnerName, const FResolveParams& ResolveParams, bool bUseSceneColorTexture)
 {
+	FRHICommandList& RHICmdList = FRHICommandList::GetNullRef();
 	FMemMark MemStackMark(FMemStack::Get());
 
 	// update any resources that needed a deferred update
@@ -31,35 +32,35 @@ void UpdateSceneCaptureContent_RenderThread(FSceneRenderer* SceneRenderer, FText
 		const FRenderTarget* Target = SceneRenderer->ViewFamily.RenderTarget;
 		FIntRect ViewRect = SceneRenderer->Views[0].ViewRect;
 		FIntRect UnconstrainedViewRect = SceneRenderer->Views[0].UnconstrainedViewRect;
-		RHISetRenderTarget(Target->GetRenderTargetTexture(), NULL);
-		RHIClear(true, FLinearColor::Black, false, 1.0f, false, 0, ViewRect);
+		SetRenderTarget(RHICmdList, Target->GetRenderTargetTexture(), NULL);
+		RHICmdList.Clear(true, FLinearColor::Black, false, 1.0f, false, 0, ViewRect);
 		SceneRenderer->Render();
 
 		// Copy the captured scene into the destination texture
 		if (bUseSceneColorTexture)
 		{
 			// Copy the captured scene into the destination texture
-			RHISetRenderTarget(Target->GetRenderTargetTexture(), NULL);
+			SetRenderTarget(RHICmdList, Target->GetRenderTargetTexture(), NULL);
 
-			RHISetRasterizerState(TStaticRasterizerState<FM_Solid, CM_None>::GetRHI());
-			RHISetDepthStencilState(TStaticDepthStencilState<false, CF_Always>::GetRHI());
-			RHISetBlendState(TStaticBlendState<>::GetRHI());
+			RHICmdList.SetRasterizerState(TStaticRasterizerState<FM_Solid, CM_None>::GetRHI());
+			RHICmdList.SetDepthStencilState(TStaticDepthStencilState<false, CF_Always>::GetRHI());
+			RHICmdList.SetBlendState(TStaticBlendState<>::GetRHI());
 
 			TShaderMapRef<FScreenVS> VertexShader(GetGlobalShaderMap());
 			TShaderMapRef<FScreenPS> PixelShader(GetGlobalShaderMap());
 			static FGlobalBoundShaderState BoundShaderState;
-			SetGlobalBoundShaderState(BoundShaderState, GFilterVertexDeclaration.VertexDeclarationRHI, *VertexShader, *PixelShader);
+			RHICmdList.CheckIsNull(); // need new approach for "static FGlobalBoundShaderState" for parallel rendering
+			SetGlobalBoundShaderState(RHICmdList, BoundShaderState, GFilterVertexDeclaration.VertexDeclarationRHI, *VertexShader, *PixelShader);
 
 			FRenderingCompositePassContext Context(SceneRenderer->Views[0]);
 
-			//@todo-rco: RHIPacketList
-			FRHICommandList& RHICmdList = FRHICommandList::GetNullRef();
 			VertexShader->SetParameters(RHICmdList, SceneRenderer->Views[0]);
 			PixelShader->SetParameters(RHICmdList, TStaticSamplerState<SF_Point>::GetRHI(), GSceneRenderTargets.GetSceneColorTexture());
 
 			FIntPoint TargetSize(UnconstrainedViewRect.Width(), UnconstrainedViewRect.Height());
 
 			DrawRectangle(
+				RHICmdList,
 				ViewRect.Min.X, ViewRect.Min.Y,
 				ViewRect.Width(), ViewRect.Height(),
 				ViewRect.Min.X, ViewRect.Min.Y,
@@ -70,7 +71,7 @@ void UpdateSceneCaptureContent_RenderThread(FSceneRenderer* SceneRenderer, FText
 				EDRF_UseTriangleOptimization);
 		}
 
-		RHICopyToResolveTarget(TextureRenderTarget->GetRenderTargetTexture(), TextureRenderTarget->TextureRHI, false, ResolveParams);
+		RHICmdList.CopyToResolveTarget(TextureRenderTarget->GetRenderTargetTexture(), TextureRenderTarget->TextureRHI, false, ResolveParams);
 	}
 
 	delete SceneRenderer;
diff --git a/Engine/Source/Runtime/Renderer/Private/SceneCore.cpp b/Engine/Source/Runtime/Renderer/Private/SceneCore.cpp
--- a/Engine/Source/Runtime/Renderer/Private/SceneCore.cpp
+++ b/Engine/Source/Runtime/Renderer/Private/SceneCore.cpp
@@ -310,7 +310,7 @@ void FStaticMesh::AddToDrawLists(FScene* Scene)
 				}
 
 				// Add the static mesh to the DPG's base pass draw list.
-				FBasePassOpaqueDrawingPolicyFactory::AddStaticMesh(Scene,this);
+				FBasePassOpaqueDrawingPolicyFactory::AddStaticMesh(FRHICommandList::GetNullRef(), Scene, this);
 
 				FVelocityDrawingPolicyFactory::AddStaticMesh(Scene, this);
 			}
diff --git a/Engine/Source/Runtime/Renderer/Private/SceneHitProxyRendering.cpp b/Engine/Source/Runtime/Renderer/Private/SceneHitProxyRendering.cpp
--- a/Engine/Source/Runtime/Renderer/Private/SceneHitProxyRendering.cpp
+++ b/Engine/Source/Runtime/Renderer/Private/SceneHitProxyRendering.cpp
@@ -356,11 +356,11 @@ void FDeferredShadingSceneRenderer::RenderHitProxies()
 		return;
 	}
 
-	RHISetRenderTarget(HitProxyRT->GetRenderTargetItem().TargetableTexture, GSceneRenderTargets.GetSceneDepthSurface());
-
 	//@todo-rco: RHIPacketList
 	FRHICommandList& RHICmdList = FRHICommandList::GetNullRef();
 
+	SetRenderTarget(RHICmdList, HitProxyRT->GetRenderTargetItem().TargetableTexture, GSceneRenderTargets.GetSceneDepthSurface());
+
 	// Clear color for each view.
 	for(int32 ViewIndex = 0;ViewIndex < Views.Num();ViewIndex++)
 	{
@@ -433,7 +433,7 @@ void FDeferredShadingSceneRenderer::RenderHitProxies()
 		DrawViewElements<FHitProxyDrawingPolicyFactory>(View,FHitProxyDrawingPolicyFactory::ContextType(),SDPG_World,bPreFog);
 
 		// Draw the view's batched simple elements(lines, sprites, etc).
-		View.BatchedViewElements.Draw(bNeedToSwitchVerticalAxis, View.ViewProjectionMatrix, View.ViewRect.Width(), View.ViewRect.Height(), true);
+		View.BatchedViewElements.Draw(RHICmdList, bNeedToSwitchVerticalAxis, View.ViewProjectionMatrix, View.ViewRect.Width(), View.ViewRect.Height(), true);
 
 		// Some elements should never be occluded (e.g. gizmos).
 		// So we render those twice, first to overwrite potentially nearer objects,
@@ -442,20 +442,20 @@ void FDeferredShadingSceneRenderer::RenderHitProxies()
 		// Draw the view's foreground elements last.
 		DrawViewElements<FHitProxyDrawingPolicyFactory>(View,FHitProxyDrawingPolicyFactory::ContextType(),SDPG_Foreground,bPreFog);
 
-		View.TopBatchedViewElements.Draw(bNeedToSwitchVerticalAxis, View.ViewProjectionMatrix, View.ViewRect.Width(), View.ViewRect.Height(), true);
+		View.TopBatchedViewElements.Draw(RHICmdList, bNeedToSwitchVerticalAxis, View.ViewProjectionMatrix, View.ViewRect.Width(), View.ViewRect.Height(), true);
 
 		// Note, this is a reversed Z depth surface, using CF_GreaterEqual.
 		RHISetDepthStencilState(TStaticDepthStencilState<true,CF_GreaterEqual>::GetRHI());
 		// Draw the view's foreground elements last.
 		DrawViewElements<FHitProxyDrawingPolicyFactory>(View,FHitProxyDrawingPolicyFactory::ContextType(),SDPG_Foreground,bPreFog);
 
-		View.TopBatchedViewElements.Draw(bNeedToSwitchVerticalAxis, View.ViewProjectionMatrix, View.ViewRect.Width(), View.ViewRect.Height(), true);
+		View.TopBatchedViewElements.Draw(RHICmdList, bNeedToSwitchVerticalAxis, View.ViewProjectionMatrix, View.ViewRect.Width(), View.ViewRect.Height(), true);
 	}
 
 	// Finish drawing to the hit proxy render target.
 	RHICopyToResolveTarget(HitProxyRT->GetRenderTargetItem().TargetableTexture, HitProxyRT->GetRenderTargetItem().ShaderResourceTexture, false, FResolveParams());
 	// to be able to observe results with VisualizeTexture
-	GRenderTargetPool.VisualizeTexture.SetCheckPoint(HitProxyRT);
+	GRenderTargetPool.VisualizeTexture.SetCheckPoint(RHICmdList, HitProxyRT);
 
 	// After scene rendering, disable the depth buffer.
 	RHISetDepthStencilState(TStaticDepthStencilState<false,CF_Always>::GetRHI());
@@ -506,8 +506,9 @@ void FDeferredShadingSceneRenderer::RenderHitProxies()
 				);
 
 	// Draw the triangles to the view family's render target.
-	RHISetRenderTarget(ViewFamily.RenderTarget->GetRenderTargetTexture(),FTextureRHIRef());
+	SetRenderTarget(RHICmdList, ViewFamily.RenderTarget->GetRenderTargetTexture(), FTextureRHIRef());
 	BatchedElements.Draw(
+				RHICmdList,
 				bNeedToSwitchVerticalAxis,
 				PixelToView,
 				RenderTargetSize.X,
diff --git a/Engine/Source/Runtime/Renderer/Private/SceneOcclusion.cpp b/Engine/Source/Runtime/Renderer/Private/SceneOcclusion.cpp
--- a/Engine/Source/Runtime/Renderer/Private/SceneOcclusion.cpp
+++ b/Engine/Source/Runtime/Renderer/Private/SceneOcclusion.cpp
@@ -676,6 +676,7 @@ IMPLEMENT_SHADER_TYPE(,FHZBTestPS,TEXT("HZBOcclusion"),TEXT("HZBTestPS"),SF_Pixe
 void FHZBOcclusionTester::Submit( const FViewInfo& View )
 {
 	SCOPED_DRAW_EVENT(SubmitHZB, DEC_SCENE_ITEMS);
+	FRHICommandList& RHICmdList = FRHICommandList::GetNullRef(); // many things below which do not support parallel
 
 	FSceneViewState* ViewState = (FSceneViewState*)View.State;
 	if( !ViewState )
@@ -746,22 +747,22 @@ void FHZBOcclusionTester::Submit( const FViewInfo& View )
 	{
 		SCOPED_DRAW_EVENT(TestHZB, DEC_SCENE_ITEMS);
 
-		RHISetRenderTarget( ResultsTextureGPU->GetRenderTargetItem().TargetableTexture, NULL );
+		SetRenderTarget(RHICmdList, ResultsTextureGPU->GetRenderTargetItem().TargetableTexture, NULL);
 
 		TShaderMapRef< FScreenVS >	VertexShader( GetGlobalShaderMap() );
 		TShaderMapRef< FHZBTestPS >	PixelShader( GetGlobalShaderMap() );
 
 		static FGlobalBoundShaderState BoundShaderState;
-		SetGlobalBoundShaderState( BoundShaderState, GFilterVertexDeclaration.VertexDeclarationRHI, *VertexShader, *PixelShader );
+		RHICmdList.CheckIsNull(); // need new approach for "static FGlobalBoundShaderState" for parallel rendering
+		SetGlobalBoundShaderState(RHICmdList, BoundShaderState, GFilterVertexDeclaration.VertexDeclarationRHI, *VertexShader, *PixelShader);
 
-		//@todo-rco: RHIPacketList
-		FRHICommandList& RHICmdList = FRHICommandList::GetNullRef();
 		PixelShader->SetParameters(RHICmdList, View,  ViewState->HZB, BoundsCenterTexture->GetRenderTargetItem().ShaderResourceTexture, BoundsExtentTexture->GetRenderTargetItem().ShaderResourceTexture );
 
-		RHISetViewport( 0, 0, 0.0f, SizeX, SizeY, 1.0f );
+		RHICmdList.SetViewport(0, 0, 0.0f, SizeX, SizeY, 1.0f);
 
 		// TODO draw quads covering blocks added above
 		DrawRectangle(
+			RHICmdList,
 			0, 0,
 			SizeX, SizeY,
 			0, 0,
@@ -772,10 +773,10 @@ void FHZBOcclusionTester::Submit( const FViewInfo& View )
 			EDRF_UseTriangleOptimization);
 	}
 
-	GRenderTargetPool.VisualizeTexture.SetCheckPoint( ResultsTextureGPU );
+	GRenderTargetPool.VisualizeTexture.SetCheckPoint(RHICmdList, ResultsTextureGPU);
 
 	// Transfer memory GPU -> CPU
-	RHICopyToResolveTarget( ResultsTextureGPU->GetRenderTargetItem().TargetableTexture, ResultsTextureCPU->GetRenderTargetItem().ShaderResourceTexture, false, FResolveParams() );
+	RHICmdList.CopyToResolveTarget(ResultsTextureGPU->GetRenderTargetItem().TargetableTexture, ResultsTextureCPU->GetRenderTargetItem().ShaderResourceTexture, false, FResolveParams());
 }
 
 template< uint32 Stage >
@@ -882,31 +883,33 @@ void BuildHZB( const FViewInfo& View )
 	
 	FSceneRenderTargetItem& HZBRenderTarget = ViewState->HZB.Texture->GetRenderTargetItem();
 	
-	RHISetBlendState( TStaticBlendState<>::GetRHI() );
-	RHISetRasterizerState( TStaticRasterizerState<>::GetRHI() );
-	RHISetDepthStencilState( TStaticDepthStencilState< false, CF_Always >::GetRHI() );
-	
 	//@todo-rco: RHIPacketList
 	FRHICommandList& RHICmdList = FRHICommandList::GetNullRef();
 
+	RHICmdList.SetBlendState(TStaticBlendState<>::GetRHI());
+	RHICmdList.SetRasterizerState(TStaticRasterizerState<>::GetRHI());
+	RHICmdList.SetDepthStencilState(TStaticDepthStencilState< false, CF_Always >::GetRHI());
+	
 	// Mip 0
 	{
 		SCOPED_DRAW_EVENTF(BuildHZB, DEC_SCENE_ITEMS, TEXT("HZB#%d"), 0);
 
-		RHISetRenderTarget( HZBRenderTarget.TargetableTexture, 0, NULL );
+		SetRenderTarget(RHICmdList, HZBRenderTarget.TargetableTexture, 0, NULL);
 
 		TShaderMapRef< FScreenVS >		VertexShader( GetGlobalShaderMap() );
 		TShaderMapRef< THZBBuildPS<0> >	PixelShader( GetGlobalShaderMap() );
 
 		static FGlobalBoundShaderState BoundShaderState;
-		SetGlobalBoundShaderState( BoundShaderState, GFilterVertexDeclaration.VertexDeclarationRHI, *VertexShader, *PixelShader );
+		RHICmdList.CheckIsNull(); // need new approach for "static FGlobalBoundShaderState" for parallel rendering
+		SetGlobalBoundShaderState(RHICmdList, BoundShaderState, GFilterVertexDeclaration.VertexDeclarationRHI, *VertexShader, *PixelShader);
 
 		// Imperfect sampling, doesn't matter too much
 		PixelShader->SetParameters(RHICmdList, View, GSceneRenderTargets.GetBufferSizeXY() );
 
-		RHISetViewport( 0, 0, 0.0f, HZBSize.X, HZBSize.Y, 1.0f );
+		RHICmdList.SetViewport(0, 0, 0.0f, HZBSize.X, HZBSize.Y, 1.0f);
 
 		DrawRectangle(
+			RHICmdList,
 			0, 0,
 			HZBSize.X, HZBSize.Y,
 			View.ViewRect.Min.X, View.ViewRect.Min.Y,
@@ -925,19 +928,21 @@ void BuildHZB( const FViewInfo& View )
 	{
 		SCOPED_DRAW_EVENTF(BuildHZB, DEC_SCENE_ITEMS, TEXT("HZB#%d"), MipIndex);
 
-		RHISetRenderTarget( HZBRenderTarget.TargetableTexture, MipIndex, NULL );
+		SetRenderTarget(RHICmdList, HZBRenderTarget.TargetableTexture, MipIndex, NULL);
 
 		TShaderMapRef< FScreenVS >		VertexShader( GetGlobalShaderMap() );
 		TShaderMapRef< THZBBuildPS<1> >	PixelShader( GetGlobalShaderMap() );
 
 		static FGlobalBoundShaderState BoundShaderState;
-		SetGlobalBoundShaderState( BoundShaderState, GFilterVertexDeclaration.VertexDeclarationRHI, *VertexShader, *PixelShader );
+		RHICmdList.CheckIsNull(); // need new approach for "static FGlobalBoundShaderState" for parallel rendering
+		SetGlobalBoundShaderState(RHICmdList, BoundShaderState, GFilterVertexDeclaration.VertexDeclarationRHI, *VertexShader, *PixelShader);
 
 		PixelShader->SetParameters(RHICmdList, View, SrcSize, ViewState->HZB.MipSRVs[ MipIndex - 1 ] );
 
-		RHISetViewport( 0, 0, 0.0f, DstSize.X, DstSize.Y, 1.0f );
+		RHICmdList.SetViewport(0, 0, 0.0f, DstSize.X, DstSize.Y, 1.0f);
 
 		DrawRectangle(
+			RHICmdList,
 			0, 0,
 			DstSize.X, DstSize.Y,
 			0, 0,
@@ -947,27 +952,28 @@ void BuildHZB( const FViewInfo& View )
 			*VertexShader,
 			EDRF_UseTriangleOptimization);
 
-		RHICopyToResolveTarget( HZBRenderTarget.TargetableTexture, HZBRenderTarget.ShaderResourceTexture, false, FResolveParams(FResolveRect(), CubeFace_PosX, MipIndex) );
+		RHICmdList.CopyToResolveTarget(HZBRenderTarget.TargetableTexture, HZBRenderTarget.ShaderResourceTexture, false, FResolveParams(FResolveRect(), CubeFace_PosX, MipIndex));
 
 		SrcSize /= 2;
 		DstSize /= 2;
 	}
 
-	GRenderTargetPool.VisualizeTexture.SetCheckPoint( ViewState->HZB.Texture );
+	GRenderTargetPool.VisualizeTexture.SetCheckPoint(RHICmdList, ViewState->HZB.Texture);
 }
 
 void FDeferredShadingSceneRenderer::BeginOcclusionTests()
 {
 	SCOPE_CYCLE_COUNTER(STAT_BeginOcclusionTestsTime);
 	const bool bUseDownsampledDepth = IsValidRef(GSceneRenderTargets.GetSmallDepthSurface()) && GSceneRenderTargets.UseDownsizedOcclusionQueries();
+	FRHICommandList& RHICmdList = FRHICommandList::GetNullRef();
 
 	if (bUseDownsampledDepth)
 	{
-		RHISetRenderTarget(NULL, GSceneRenderTargets.GetSmallDepthSurface());
+		SetRenderTarget(RHICmdList, NULL, GSceneRenderTargets.GetSmallDepthSurface());
 	}
 	else
 	{
-		RHISetRenderTarget(NULL, GSceneRenderTargets.GetSceneDepthSurface());
+		SetRenderTarget(RHICmdList, NULL, GSceneRenderTargets.GetSceneDepthSurface());
 	}
 
 	// Perform occlusion queries for each view
@@ -984,11 +990,11 @@ void FDeferredShadingSceneRenderer::BeginOcclusionTests()
 			const uint32 DownsampledSizeY = FMath::TruncToInt(View.ViewRect.Height() / GSceneRenderTargets.GetSmallColorDepthDownsampleFactor());
 
 			// Setup the viewport for rendering to the downsampled depth buffer
-			RHISetViewport(DownsampledX, DownsampledY, 0.0f, DownsampledX + DownsampledSizeX, DownsampledY + DownsampledSizeY, 1.0f);
+			RHICmdList.SetViewport(DownsampledX, DownsampledY, 0.0f, DownsampledX + DownsampledSizeX, DownsampledY + DownsampledSizeY, 1.0f);
 		}
 		else
 		{
-			RHISetViewport(View.ViewRect.Min.X, View.ViewRect.Min.Y, 0.0f, View.ViewRect.Max.X, View.ViewRect.Max.Y, 1.0f);
+			RHICmdList.SetViewport(View.ViewRect.Min.X, View.ViewRect.Min.Y, 0.0f, View.ViewRect.Max.X, View.ViewRect.Max.Y, 1.0f);
 		}
     
 	    FSceneViewState* ViewState = (FSceneViewState*)View.State;
@@ -997,16 +1003,14 @@ void FDeferredShadingSceneRenderer::BeginOcclusionTests()
 	    {
 			// Depth tests, no depth writes, no color writes, opaque, solid rasterization wo/ backface culling.
 			// Note, this is a reversed Z depth surface, using CF_GreaterEqual.
-			RHISetDepthStencilState(TStaticDepthStencilState<false,CF_GreaterEqual>::GetRHI());
+			RHICmdList.SetDepthStencilState(TStaticDepthStencilState<false, CF_GreaterEqual>::GetRHI());
 			// We only need to render the front-faces of the culling geometry (this halves the amount of pixels we touch)
-			RHISetRasterizerState(View.bReverseCulling ? TStaticRasterizerState<FM_Solid,CM_CCW>::GetRHI() : TStaticRasterizerState<FM_Solid,CM_CW>::GetRHI()); 
-			RHISetBlendState(TStaticBlendState<CW_NONE>::GetRHI());
+			RHICmdList.SetRasterizerState(View.bReverseCulling ? TStaticRasterizerState<FM_Solid, CM_CCW>::GetRHI() : TStaticRasterizerState<FM_Solid, CM_CW>::GetRHI());
+			RHICmdList.SetBlendState(TStaticBlendState<CW_NONE>::GetRHI());
 
 			// Lookup the vertex shader.
 			TShaderMapRef<FOcclusionQueryVS> VertexShader(GetGlobalShaderMap());
-			SetGlobalBoundShaderState(OcclusionTestBoundShaderState,GetVertexDeclarationFVector3(),*VertexShader,NULL);
-			//@todo-rco: RHIPacketList
-			FRHICommandList& RHICmdList = FRHICommandList::GetNullRef();
+			SetGlobalBoundShaderState(RHICmdList, OcclusionTestBoundShaderState, GetVertexDeclarationFVector3(), *VertexShader, NULL);
 			VertexShader->SetParameters(RHICmdList, View);
 
 			// Issue this frame's occlusion queries (occlusion queries from last frame may still be in flight)
@@ -1051,6 +1055,7 @@ void FDeferredShadingSceneRenderer::BeginOcclusionTests()
 							if (!bCameraInsideLightGeometry)
 							{
 								const FRenderQueryRHIRef ShadowOcclusionQuery = ViewState->OcclusionQueryPool.AllocateQuery();
+								RHICmdList.CheckIsNull(); // RHIBeginRenderQuery
 								RHIBeginRenderQuery(ShadowOcclusionQuery);
 
 								FSceneViewState::FProjectedShadowKey Key(
@@ -1066,7 +1071,7 @@ void FDeferredShadingSceneRenderer::BeginOcclusionTests()
 
 								// Draw bounding sphere
 								VertexShader->SetParametersWithBoundingSphere(RHICmdList, View, ProjectedShadowInfo.LightSceneInfo->Proxy->GetBoundingSphere());
-								StencilingGeometry::DrawSphere();
+								StencilingGeometry::DrawSphere(RHICmdList);
 
 								RHIEndRenderQuery(ShadowOcclusionQuery);
 							}
diff --git a/Engine/Source/Runtime/Renderer/Private/ScenePostProcessing.cpp b/Engine/Source/Runtime/Renderer/Private/ScenePostProcessing.cpp
--- a/Engine/Source/Runtime/Renderer/Private/ScenePostProcessing.cpp
+++ b/Engine/Source/Runtime/Renderer/Private/ScenePostProcessing.cpp
@@ -98,32 +98,33 @@ void FDeferredShadingSceneRenderer::FinishRenderViewTarget(const FViewInfo* View
 }
 
 // TODO: REMOVE if no longer needed:
-void FSceneRenderer::GammaCorrectToViewportRenderTarget(const FViewInfo* View, float OverrideGamma)
+void FSceneRenderer::GammaCorrectToViewportRenderTarget(FRHICommandList& RHICmdList, const FViewInfo* View, float OverrideGamma)
 {
 	// Set the view family's render target/viewport.
-	RHISetRenderTarget(ViewFamily.RenderTarget->GetRenderTargetTexture(),FTextureRHIRef());	
+	SetRenderTarget(RHICmdList, ViewFamily.RenderTarget->GetRenderTargetTexture(), FTextureRHIRef());
 
 	// Deferred the clear until here so the garbage left in the non rendered regions by the post process effects do not show up
 	if( ViewFamily.bDeferClear )
 	{
-		RHIClear(true, FLinearColor::Black, false, 0.0f, false, 0, FIntRect());
+		RHICmdList.Clear(true, FLinearColor::Black, false, 0.0f, false, 0, FIntRect());
 		ViewFamily.bDeferClear = false;
 	}
 
 	SCOPED_DRAW_EVENT(GammaCorrection, DEC_SCENE_ITEMS);
 
 	// turn off culling and blending
-	RHISetRasterizerState(TStaticRasterizerState<FM_Solid,CM_None>::GetRHI());
-	RHISetBlendState(TStaticBlendState<>::GetRHI());
+	RHICmdList.SetRasterizerState(TStaticRasterizerState<FM_Solid, CM_None>::GetRHI());
+	RHICmdList.SetBlendState(TStaticBlendState<>::GetRHI());
 
 	// turn off depth reads/writes
-	RHISetDepthStencilState(TStaticDepthStencilState<false,CF_Always>::GetRHI());
+	RHICmdList.SetDepthStencilState(TStaticDepthStencilState<false, CF_Always>::GetRHI());
 
 	TShaderMapRef<FGammaCorrectionVS> VertexShader(GetGlobalShaderMap());
 	TShaderMapRef<FGammaCorrectionPS> PixelShader(GetGlobalShaderMap());
 
 	static FGlobalBoundShaderState PostProcessBoundShaderState;
-	SetGlobalBoundShaderState( PostProcessBoundShaderState, GFilterVertexDeclaration.VertexDeclarationRHI, *VertexShader, *PixelShader);
+	RHICmdList.CheckIsNull(); // need new approach for "static FGlobalBoundShaderState" for parallel rendering
+	SetGlobalBoundShaderState(RHICmdList, PostProcessBoundShaderState, GFilterVertexDeclaration.VertexDeclarationRHI, *VertexShader, *PixelShader);
 
 	float InvDisplayGamma = 1.0f / ViewFamily.RenderTarget->GetDisplayGamma();
 
@@ -134,8 +135,6 @@ void FSceneRenderer::GammaCorrectToViewportRenderTarget(const FViewInfo* View, f
 
 	const FPixelShaderRHIParamRef ShaderRHI = PixelShader->GetPixelShader();
 
-	//@todo-rco: RHIPacketList
-	FRHICommandList& RHICmdList = FRHICommandList::GetNullRef();
 	SetShaderValue(
 		RHICmdList, 
 		ShaderRHI,
@@ -158,6 +157,7 @@ void FSceneRenderer::GammaCorrectToViewportRenderTarget(const FViewInfo* View, f
 
 	// Draw a quad mapping scene color to the view's render target
 	DrawRectangle(
+		RHICmdList,
 		View->UnscaledViewRect.Min.X,View->UnscaledViewRect.Min.Y,
 		View->UnscaledViewRect.Width(),View->UnscaledViewRect.Height(),
 		View->ViewRect.Min.X,View->ViewRect.Min.Y,
diff --git a/Engine/Source/Runtime/Renderer/Private/SceneRendering.h b/Engine/Source/Runtime/Renderer/Private/SceneRendering.h
--- a/Engine/Source/Runtime/Renderer/Private/SceneRendering.h
+++ b/Engine/Source/Runtime/Renderer/Private/SceneRendering.h
@@ -551,7 +551,7 @@ protected:
 	bool ShouldRenderTranslucency() const;
 
 	/** TODO: REMOVE if no longer needed: Copies scene color to the viewport's render target after applying gamma correction. */
-	void GammaCorrectToViewportRenderTarget(const FViewInfo* View, float OverrideGamma);
+	void GammaCorrectToViewportRenderTarget(FRHICommandList& RHICmdList, const FViewInfo* View, float OverrideGamma);
 
 	/** Updates state for the end of the frame. */
 	void RenderFinish();
diff --git a/Engine/Source/Runtime/Renderer/Private/ShaderComplexityRendering.cpp b/Engine/Source/Runtime/Renderer/Private/ShaderComplexityRendering.cpp
--- a/Engine/Source/Runtime/Renderer/Private/ShaderComplexityRendering.cpp
+++ b/Engine/Source/Runtime/Renderer/Private/ShaderComplexityRendering.cpp
@@ -88,16 +88,15 @@ public:
 	}
 
 	virtual void SetParameters(
-		FRHICommandList& RHICmdList, 
 		const FRenderingCompositePassContext& Context,
 		const TArray<FLinearColor>& Colors
 		)
 	{
 		const FPixelShaderRHIParamRef ShaderRHI = GetPixelShader();
 
-		FGlobalShader::SetParameters(RHICmdList, ShaderRHI, Context.View);
+		FGlobalShader::SetParameters(Context.RHICmdList, ShaderRHI, Context.View);
 
-		PostprocessParameter.SetPS(RHICmdList, ShaderRHI, Context, TStaticSamplerState<SF_Bilinear,AM_Clamp,AM_Clamp,AM_Clamp>::GetRHI());
+		PostprocessParameter.SetPS(ShaderRHI, Context, TStaticSamplerState<SF_Bilinear, AM_Clamp, AM_Clamp, AM_Clamp>::GetRHI());
 
 		//Make sure there are at least NumShaderComplexityColors colors specified in the ini.
 		//If there are more than NumShaderComplexityColors they will be ignored.
@@ -106,7 +105,7 @@ public:
 		//pass the complexity -> color mapping into the pixel shader
 		for(int32 ColorIndex = 0; ColorIndex < NumShaderComplexityColors; ColorIndex ++)
 		{
-			SetShaderValue(RHICmdList, ShaderRHI, ShaderComplexityColors, Colors[ColorIndex], ColorIndex);
+			SetShaderValue(Context.RHICmdList, ShaderRHI, ShaderComplexityColors, Colors[ColorIndex], ColorIndex);
 		}
 	}
 
@@ -168,30 +167,29 @@ void FRCPassPostProcessVisualizeComplexity::Process(FRenderingCompositePassConte
 
 	const FSceneRenderTargetItem& DestRenderTarget = PassOutputs[0].RequestSurface(Context);
 
-	//@todo-rco: RHIPacketList
-	FRHICommandList& RHICmdList = FRHICommandList::GetNullRef();
-
 	// Set the view family's render target/viewport.
-	RHISetRenderTarget(DestRenderTarget.TargetableTexture, FTextureRHIRef());	
+	SetRenderTarget(Context.RHICmdList, DestRenderTarget.TargetableTexture, FTextureRHIRef());
 	Context.SetViewportAndCallRHI(DestRect);
 
 	// turn off culling and blending
-	RHISetRasterizerState(TStaticRasterizerState<FM_Solid,CM_None>::GetRHI());
-	RHISetBlendState(TStaticBlendState<>::GetRHI());
+	Context.RHICmdList.SetRasterizerState(TStaticRasterizerState<FM_Solid, CM_None>::GetRHI());
+	Context.RHICmdList.SetBlendState(TStaticBlendState<>::GetRHI());
 
 	// turn off depth reads/writes
-	RHISetDepthStencilState(TStaticDepthStencilState<false,CF_Always>::GetRHI());
+	Context.RHICmdList.SetDepthStencilState(TStaticDepthStencilState<false, CF_Always>::GetRHI());
 
 	//reuse this generic vertex shader
 	TShaderMapRef<FShaderComplexityApplyVS> VertexShader(GetGlobalShaderMap());
 	TShaderMapRef<FShaderComplexityApplyPS> PixelShader(GetGlobalShaderMap());
 
 	static FGlobalBoundShaderState ShaderComplexityBoundShaderState;
-	SetGlobalBoundShaderState(ShaderComplexityBoundShaderState, GFilterVertexDeclaration.VertexDeclarationRHI, *VertexShader, *PixelShader);
+	Context.RHICmdList.CheckIsNull(); // need new approach for "static FGlobalBoundShaderState" for parallel rendering
+	SetGlobalBoundShaderState(Context.RHICmdList, ShaderComplexityBoundShaderState, GFilterVertexDeclaration.VertexDeclarationRHI, *VertexShader, *PixelShader);
 
-	PixelShader->SetParameters(RHICmdList, Context, Colors);
+	PixelShader->SetParameters(Context, Colors);
 	
 	DrawRectangle(
+		Context.RHICmdList,
 		0, 0,
 		DestRect.Width(), DestRect.Height(),
 		SrcRect.Min.X, SrcRect.Min.Y,
@@ -201,7 +199,7 @@ void FRCPassPostProcessVisualizeComplexity::Process(FRenderingCompositePassConte
 		*VertexShader,
 		EDRF_UseTriangleOptimization);
 
-	RHICopyToResolveTarget(DestRenderTarget.TargetableTexture, DestRenderTarget.ShaderResourceTexture, false, FResolveParams());
+	Context.RHICmdList.CopyToResolveTarget(DestRenderTarget.TargetableTexture, DestRenderTarget.ShaderResourceTexture, false, FResolveParams());
 }
 
 FPooledRenderTargetDesc FRCPassPostProcessVisualizeComplexity::ComputeOutputDesc(EPassOutputId InPassOutputId) const
diff --git a/Engine/Source/Runtime/Renderer/Private/ShadowRendering.cpp b/Engine/Source/Runtime/Renderer/Private/ShadowRendering.cpp
--- a/Engine/Source/Runtime/Renderer/Private/ShadowRendering.cpp
+++ b/Engine/Source/Runtime/Renderer/Private/ShadowRendering.cpp
@@ -1184,7 +1184,7 @@ void FProjectedShadowInfo::ClearDepth(FDeferredShadingSceneRenderer* SceneRender
 
 
 template <bool bReflectiveShadowmap>
-void DrawMeshElements(FShadowDepthDrawingPolicy<bReflectiveShadowmap>& SharedDrawingPolicy, const FShadowStaticMeshElement& State, const FViewInfo& View, const FStaticMesh* Mesh)
+void DrawMeshElements(FRHICommandList& RHICmdList, FShadowDepthDrawingPolicy<bReflectiveShadowmap>& SharedDrawingPolicy, const FShadowStaticMeshElement& State, const FViewInfo& View, const FStaticMesh* Mesh)
 {
 #if UE_BUILD_DEBUG
 	// During shadow setup we should have already overridden materials with default material where needed.
@@ -1207,9 +1207,6 @@ void DrawMeshElements(FShadowDepthDrawingPolicy<bReflectiveShadowmap>& SharedDra
 	checkSlow(CompareDrawingPolicy(DebugPolicy, SharedDrawingPolicy) == 0);
 #endif
 
-	//@todo-rco: RHIPacketList
-	FRHICommandList& RHICmdList = FRHICommandList::GetNullRef();
-
 	// Render only those batch elements that match the current LOD
 	uint64 BatchElementMask = Mesh->Elements.Num() == 1 ? 1 : View.StaticMeshBatchVisibility[Mesh->Id];
 	int32 BatchElementIndex = 0;
@@ -1263,7 +1260,7 @@ void DrawShadowMeshElements(FRHICommandList& RHICmdList, const FViewInfo& View,
 			SharedDrawingPolicy.DrawShared(RHICmdList, &View, SharedDrawingPolicy.CreateBoundShaderState(View.GetFeatureLevel()));
 		}
 
-		DrawMeshElements(SharedDrawingPolicy, OldState, View, ShadowMesh.Mesh);
+		DrawMeshElements(RHICmdList, SharedDrawingPolicy, OldState, View, ShadowMesh.Mesh);
 	}
 }
 
@@ -1283,7 +1280,7 @@ void FProjectedShadowInfo::RenderDepth(FRHICommandList& RHICmdList, FDeferredSha
 	if (bOnePassPointLightShadow)
 	{
 		// Set the viewport to the whole render target since it's a cube map, don't leave any border space
-		RHISetViewport(
+		RHICmdList.SetViewport(
 			0,
 			0,
 			0.0f,
@@ -1297,7 +1294,7 @@ void FProjectedShadowInfo::RenderDepth(FRHICommandList& RHICmdList, FDeferredSha
 		// Do not use a shadow border for RSMs.
 		if(bReflectiveShadowmap)
 		{
-			RHISetViewport(
+			RHICmdList.SetViewport(
 				X,
 				Y,
 				0.0f,
@@ -1308,7 +1305,7 @@ void FProjectedShadowInfo::RenderDepth(FRHICommandList& RHICmdList, FDeferredSha
 		}
 		else
 		{
-			RHISetViewport(
+			RHICmdList.SetViewport(
 				X + SHADOW_BORDER,
 				Y + SHADOW_BORDER,
 				0.0f,
@@ -1322,15 +1319,15 @@ void FProjectedShadowInfo::RenderDepth(FRHICommandList& RHICmdList, FDeferredSha
 	if (bReflectiveShadowmap && !bOnePassPointLightShadow)
 	{
 		// Enable color writes to the reflective shadow map targets with opaque blending
-		RHISetBlendState(TStaticBlendStateWriteMask<CW_RGBA, CW_RGBA>::GetRHI());
+		RHICmdList.SetBlendState(TStaticBlendStateWriteMask<CW_RGBA, CW_RGBA>::GetRHI());
 	}
 	else
 	{
 		// Disable color writes
-		RHISetBlendState(TStaticBlendState<CW_NONE>::GetRHI());
+		RHICmdList.SetBlendState(TStaticBlendState<CW_NONE>::GetRHI());
 	}
 	
-	RHISetDepthStencilState(TStaticDepthStencilState<true,CF_LessEqual>::GetRHI());
+	RHICmdList.SetDepthStencilState(TStaticDepthStencilState<true, CF_LessEqual>::GetRHI());
 
 	// Choose an arbitrary view where this shadow's subject is relevant.
 	FViewInfo* FoundView = NULL;
@@ -1347,6 +1344,7 @@ void FProjectedShadowInfo::RenderDepth(FRHICommandList& RHICmdList, FDeferredSha
 	}
 	check(FoundView);
 
+	RHICmdList.CheckIsNull(); // this messes with the view, won't work with parallel rendering. Need to copy and modify the view instead
 	// Backup properties of the view that we will override
 	TUniformBufferRef<FViewUniformShaderParameters> OriginalUniformBuffer = FoundView->UniformBuffer;
 	FMatrix OriginalViewMatrix = FoundView->ViewMatrices.ViewMatrix;
@@ -1370,9 +1368,6 @@ void FProjectedShadowInfo::RenderDepth(FRHICommandList& RHICmdList, FDeferredSha
 	// Draw the subject's static elements using static draw lists
 	if (IsWholeSceneDirectionalShadow())
 	{
-		//@todo-rco: RHIPacketList
-		FRHICommandList& RHICmdList = FRHICommandList::GetNullRef();
-
 		SCOPE_CYCLE_COUNTER(STAT_WholeSceneStaticDrawListShadowDepthsTime);
 
 		if (bReflectiveShadowmap)
@@ -1447,20 +1442,20 @@ void FProjectedShadowInfo::RenderDepth(FRHICommandList& RHICmdList, FDeferredSha
 	FoundView->ViewMatrices.ViewMatrix = OriginalViewMatrix;
 }
 
-void StencilingGeometry::DrawSphere()
+void StencilingGeometry::DrawSphere(FRHICommandList& RHICmdList)
 {
-	RHISetStreamSource(0, StencilingGeometry::GStencilSphereVertexBuffer.VertexBufferRHI, sizeof(FVector4), 0);
-	RHIDrawIndexedPrimitive(StencilingGeometry::GStencilSphereIndexBuffer.IndexBufferRHI, PT_TriangleList, 0,0, 
+	RHICmdList.SetStreamSource(0, StencilingGeometry::GStencilSphereVertexBuffer.VertexBufferRHI, sizeof(FVector4), 0);
+	RHICmdList.DrawIndexedPrimitive(StencilingGeometry::GStencilSphereIndexBuffer.IndexBufferRHI, PT_TriangleList, 0, 0,
 		StencilingGeometry::GStencilSphereVertexBuffer.GetVertexCount(), 0, 
 		StencilingGeometry::GStencilSphereIndexBuffer.GetIndexCount() / 3, 0);
 }
 
-void StencilingGeometry::DrawCone()
+void StencilingGeometry::DrawCone(FRHICommandList& RHICmdList)
 {
 	// No Stream Source needed since it will generate vertices on the fly
-	RHISetStreamSource(0, StencilingGeometry::GStencilConeVertexBuffer.VertexBufferRHI, sizeof(FVector4), 0);
+	RHICmdList.SetStreamSource(0, StencilingGeometry::GStencilConeVertexBuffer.VertexBufferRHI, sizeof(FVector4), 0);
 
-	RHIDrawIndexedPrimitive(StencilingGeometry::GStencilConeIndexBuffer.IndexBufferRHI, PT_TriangleList, 0, 0, 
+	RHICmdList.DrawIndexedPrimitive(StencilingGeometry::GStencilConeIndexBuffer.IndexBufferRHI, PT_TriangleList, 0, 0,
 		FStencilConeIndexBuffer::NumVerts, 0, StencilingGeometry::GStencilConeIndexBuffer.GetIndexCount() / 3, 0);
 }
 
@@ -1480,7 +1475,8 @@ static void SetShadowProjectionShaderTemplNew(FRHICommandList& RHICmdList, int32
 
 		// Bind shader
 		static FGlobalBoundShaderState BoundShaderState;
-		SetGlobalBoundShaderState(BoundShaderState, GetVertexDeclarationFVector4(), ShadowProjVS, ShadowProjPS );
+		RHICmdList.CheckIsNull(); // need new approach for "static FGlobalBoundShaderState" for parallel rendering
+		SetGlobalBoundShaderState(RHICmdList, BoundShaderState, GetVertexDeclarationFVector4(), ShadowProjVS, ShadowProjPS);
 
 		// Set shader parameters
 		ShadowProjVS->SetParameters(RHICmdList, View, ShadowInfo);
@@ -1498,7 +1494,8 @@ static void SetShadowProjectionShaderTemplNew(FRHICommandList& RHICmdList, int32
 			FShadowProjectionPixelShaderInterface* ShadowProjPS = GetGlobalShaderMap()->GetShader<TShadowProjectionPS<Quality, true> >();
 
 			static FGlobalBoundShaderState BoundShaderState;
-			SetGlobalBoundShaderState(BoundShaderState, GetVertexDeclarationFVector4(), ShadowProjVS, ShadowProjPS );
+			RHICmdList.CheckIsNull(); // need new approach for "static FGlobalBoundShaderState" for parallel rendering
+			SetGlobalBoundShaderState(RHICmdList, BoundShaderState, GetVertexDeclarationFVector4(), ShadowProjVS, ShadowProjPS);
 
 			ShadowProjPS->SetParameters(RHICmdList, ViewIndex, View, ShadowInfo);
 		}
@@ -1508,7 +1505,8 @@ static void SetShadowProjectionShaderTemplNew(FRHICommandList& RHICmdList, int32
 			FShadowProjectionPixelShaderInterface* ShadowProjPS = GetGlobalShaderMap()->GetShader<TShadowProjectionPS<Quality, false> >();
 
 			static FGlobalBoundShaderState BoundShaderState;
-			SetGlobalBoundShaderState(BoundShaderState, GetVertexDeclarationFVector4(), ShadowProjVS, ShadowProjPS );
+			RHICmdList.CheckIsNull(); // need new approach for "static FGlobalBoundShaderState" for parallel rendering
+			SetGlobalBoundShaderState(RHICmdList, BoundShaderState, GetVertexDeclarationFVector4(), ShadowProjVS, ShadowProjPS);
 
 			ShadowProjPS->SetParameters(RHICmdList, ViewIndex, View, ShadowInfo);
 		}
@@ -1525,7 +1523,8 @@ static void SetShadowProjectionShaderTemplNew(FRHICommandList& RHICmdList, int32
 		FShadowProjectionPixelShaderInterface* ShadowProjPS = GetGlobalShaderMap()->GetShader<TShadowProjectionPS<Quality, false> >();
 
 		static FGlobalBoundShaderState BoundShaderState;
-		SetGlobalBoundShaderState(BoundShaderState, GetVertexDeclarationFVector4(), ShadowProjVS, ShadowProjPS );
+		RHICmdList.CheckIsNull(); // need new approach for "static FGlobalBoundShaderState" for parallel rendering
+		SetGlobalBoundShaderState(RHICmdList, BoundShaderState, GetVertexDeclarationFVector4(), ShadowProjVS, ShadowProjPS);
 
 		ShadowProjVS->SetParameters(RHICmdList, View, ShadowInfo);
 		ShadowProjPS->SetParameters(RHICmdList, ViewIndex, View, ShadowInfo);
@@ -1534,8 +1533,6 @@ static void SetShadowProjectionShaderTemplNew(FRHICommandList& RHICmdList, int32
 
 void FProjectedShadowInfo::RenderProjection(FRHICommandList& RHICmdList, int32 ViewIndex, const FViewInfo* View) const
 {
-	RHICmdList.CheckIsNull();
-
 #if WANTS_DRAW_MESH_EVENTS
 	FString EventName;
 	GetShadowTypeNameForDrawEvent(EventName);
@@ -1628,8 +1625,8 @@ void FProjectedShadowInfo::RenderProjection(FRHICommandList& RHICmdList, int32 V
 
 	// Depth test wo/ writes, no color writing.
 	// Note, this is a reversed Z depth surface, using CF_GreaterEqual.
-	RHISetDepthStencilState(TStaticDepthStencilState<false,CF_GreaterEqual>::GetRHI());
-	RHISetBlendState(TStaticBlendState<CW_NONE>::GetRHI());
+	RHICmdList.SetDepthStencilState(TStaticDepthStencilState<false, CF_GreaterEqual>::GetRHI());
+	RHICmdList.SetBlendState(TStaticBlendState<CW_NONE>::GetRHI());
 	
 	bool bDepthBoundsTestEnabled = false;
 
@@ -1640,7 +1637,7 @@ void FProjectedShadowInfo::RenderProjection(FRHICommandList& RHICmdList, int32 V
 
 		// Set stencil to one.
 		// Note, this is a reversed Z depth surface, using CF_GreaterEqual.
-		RHISetDepthStencilState(TStaticDepthStencilState<
+		RHICmdList.SetDepthStencilState(TStaticDepthStencilState<
 			false,CF_GreaterEqual,
 			true,CF_Always,SO_Keep,SO_Keep,SO_Replace,
 			false,CF_Always,SO_Keep,SO_Keep,SO_Keep,
@@ -1709,65 +1706,65 @@ void FProjectedShadowInfo::RenderProjection(FRHICommandList& RHICmdList, int32 V
 			}
 
 			// Note, using a reversed z depth surface
-			RHIEnableDepthBoundsTest( true, DepthFar, DepthNear );
+			RHICmdList.EnableDepthBoundsTest(true, DepthFar, DepthNear);
 			bDepthBoundsTestEnabled = true;
 		}
 		else
 		{
 		// Increment stencil on front-facing zfail, decrement on back-facing zfail.
 		// Note, this is a reversed Z depth surface, using CF_GreaterEqual.
-		RHISetDepthStencilState(TStaticDepthStencilState<
+			RHICmdList.SetDepthStencilState(TStaticDepthStencilState<
 			false,CF_GreaterEqual,
 			true,CF_Always,SO_Keep,SO_Increment,SO_Keep,
 			true,CF_Always,SO_Keep,SO_Decrement,SO_Keep,
 			0xff,0xff
 			>::GetRHI());
 
-		RHISetRasterizerState(TStaticRasterizerState<FM_Solid,CM_None>::GetRHI());
+			RHICmdList.SetRasterizerState(TStaticRasterizerState<FM_Solid, CM_None>::GetRHI());
 
-		checkSlow(SplitIndex >= 0);
-		checkSlow(bDirectionalLight);
+			checkSlow(SplitIndex >= 0);
+			checkSlow(bDirectionalLight);
 
-		// Draw 2 fullscreen planes, front facing one at the near subfrustum plane, and back facing one at the far.
+			// Draw 2 fullscreen planes, front facing one at the near subfrustum plane, and back facing one at the far.
 
-		// Find the projection shaders.
-		TShaderMapRef<FShadowProjectionNoTransformVS> VertexShaderNoTransform(GetGlobalShaderMap());
-		VertexShaderNoTransform->SetParameters(RHICmdList, *View);
-		SetGlobalBoundShaderState(MaskBoundShaderState[0], GetVertexDeclarationFVector4(), *VertexShaderNoTransform, nullptr);
-
-		FVector4 Near = View->ViewMatrices.ProjMatrix.TransformFVector4(FVector4(0, 0, CascadeSettings.SplitNear));
-		FVector4 Far = View->ViewMatrices.ProjMatrix.TransformFVector4(FVector4(0, 0, CascadeSettings.SplitFar));
-		float StencilNear = Near.Z / Near.W;
-		float StencilFar = Far.Z / Far.W;
-
-		FVector4 Verts[] =
-		{
-			// Far Plane
-			FVector4( 1,  1,  StencilFar),
-			FVector4(-1,  1,  StencilFar),
-			FVector4( 1, -1,  StencilFar),
-			FVector4( 1, -1,  StencilFar),
-			FVector4(-1,  1,  StencilFar),
-			FVector4(-1, -1,  StencilFar),
-
-			// Near Plane
-			FVector4(-1,  1, StencilNear),
-			FVector4( 1,  1, StencilNear),
-			FVector4(-1, -1, StencilNear),
-			FVector4(-1, -1, StencilNear),
-			FVector4( 1,  1, StencilNear),
-			FVector4( 1, -1, StencilNear),
-		};
+			// Find the projection shaders.
+			TShaderMapRef<FShadowProjectionNoTransformVS> VertexShaderNoTransform(GetGlobalShaderMap());
+			VertexShaderNoTransform->SetParameters(RHICmdList, *View);
+			SetGlobalBoundShaderState(RHICmdList, MaskBoundShaderState[0], GetVertexDeclarationFVector4(), *VertexShaderNoTransform, nullptr);
 
-		// Only draw the near plane if this is not the nearest split
-		RHIDrawPrimitiveUP(PT_TriangleList, (SplitIndex > 0) ? 4 : 2, Verts, sizeof(FVector4));
-	}
+			FVector4 Near = View->ViewMatrices.ProjMatrix.TransformFVector4(FVector4(0, 0, CascadeSettings.SplitNear));
+			FVector4 Far = View->ViewMatrices.ProjMatrix.TransformFVector4(FVector4(0, 0, CascadeSettings.SplitFar));
+			float StencilNear = Near.Z / Near.W;
+			float StencilFar = Far.Z / Far.W;
+
+			FVector4 Verts[] =
+			{
+				// Far Plane
+				FVector4( 1,  1,  StencilFar),
+				FVector4(-1,  1,  StencilFar),
+				FVector4( 1, -1,  StencilFar),
+				FVector4( 1, -1,  StencilFar),
+				FVector4(-1,  1,  StencilFar),
+				FVector4(-1, -1,  StencilFar),
+
+				// Near Plane
+				FVector4(-1,  1, StencilNear),
+				FVector4( 1,  1, StencilNear),
+				FVector4(-1, -1, StencilNear),
+				FVector4(-1, -1, StencilNear),
+				FVector4( 1,  1, StencilNear),
+				FVector4( 1, -1, StencilNear),
+			};
+
+			// Only draw the near plane if this is not the nearest split
+			DrawPrimitiveUP(RHICmdList, PT_TriangleList, (SplitIndex > 0) ? 4 : 2, Verts, sizeof(FVector4));
+		}
 	}
 	// Not a preshadow, mask the projection to any pixels inside the frustum.
 	else
 	{
 		// Solid rasterization wo/ backface culling.
-		RHISetRasterizerState(TStaticRasterizerState<FM_Solid,CM_None>::GetRHI());
+		RHICmdList.SetRasterizerState(TStaticRasterizerState<FM_Solid, CM_None>::GetRHI());
 
 		if (bCameraInsideShadowFrustum)
 		{
@@ -1776,7 +1773,7 @@ void FProjectedShadowInfo::RenderProjection(FRHICommandList& RHICmdList, int32 V
 			// zfail stenciling is somewhat slower than zpass because on xbox 360 HiZ will be disabled when setting up stencil.
 			// Increment stencil on front-facing zfail, decrement on back-facing zfail.
 			// Note, this is a reversed Z depth surface, using CF_GreaterEqual.
-			RHISetDepthStencilState(TStaticDepthStencilState<
+			RHICmdList.SetDepthStencilState(TStaticDepthStencilState<
 				false,CF_GreaterEqual,
 				true,CF_Always,SO_Keep,SO_Increment,SO_Keep,
 				true,CF_Always,SO_Keep,SO_Decrement,SO_Keep,
@@ -1788,7 +1785,7 @@ void FProjectedShadowInfo::RenderProjection(FRHICommandList& RHICmdList, int32 V
 			// Increment stencil on front-facing zpass, decrement on back-facing zpass.
 			// HiZ will be enabled on xbox 360 which will save a little GPU time.
 			// Note, this is a reversed Z depth surface, using CF_GreaterEqual.
-			RHISetDepthStencilState(TStaticDepthStencilState<
+			RHICmdList.SetDepthStencilState(TStaticDepthStencilState<
 				false,CF_GreaterEqual,
 				true,CF_Always,SO_Keep,SO_Keep,SO_Increment,
 				true,CF_Always,SO_Keep,SO_Keep,SO_Decrement,
@@ -1800,33 +1797,33 @@ void FProjectedShadowInfo::RenderProjection(FRHICommandList& RHICmdList, int32 V
 		TShaderMapRef<FShadowProjectionVS> VertexShader(GetGlobalShaderMap());
 
 		// Cache the bound shader state
-		SetGlobalBoundShaderState(MaskBoundShaderState[1],GetVertexDeclarationFVector4(),*VertexShader,NULL);
+		SetGlobalBoundShaderState(RHICmdList, MaskBoundShaderState[1], GetVertexDeclarationFVector4(), *VertexShader, NULL);
 
 		// Set the projection vertex shader parameters
 		VertexShader->SetParameters(RHICmdList, *View, this);
 
 		// Draw the frustum using the stencil buffer to mask just the pixels which are inside the shadow frustum.
-		RHIDrawIndexedPrimitiveUP(PT_TriangleList, 0, 8, 12, GCubeIndices, sizeof(uint16), FrustumVertices, sizeof(FVector4));
+		DrawIndexedPrimitiveUP(RHICmdList, PT_TriangleList, 0, 8, 12, GCubeIndices, sizeof(uint16), FrustumVertices, sizeof(FVector4));
 	}
 
 	// solid rasterization w/ back-face culling.
-	RHISetRasterizerState(
+	RHICmdList.SetRasterizerState(
 		XOR(View->bReverseCulling, IsWholeSceneDirectionalShadow()) ? TStaticRasterizerState<FM_Solid,CM_CCW>::GetRHI() : TStaticRasterizerState<FM_Solid,CM_CW>::GetRHI());
 
 	if( bDepthBoundsTestEnabled )
 	{
 		// no depth test or writes
-		RHISetDepthStencilState(TStaticDepthStencilState<false, CF_Always>::GetRHI());
+		RHICmdList.SetDepthStencilState(TStaticDepthStencilState<false, CF_Always>::GetRHI());
 	}
 	else
 	{
 		// no depth test or writes, Test stencil for non-zero.
-	RHISetDepthStencilState(TStaticDepthStencilState<
-		false,CF_Always,
-		true,CF_NotEqual,SO_Keep,SO_Keep,SO_Keep,
-		false,CF_Always,SO_Keep,SO_Keep,SO_Keep,
-		0xff,0xff
-		>::GetRHI());
+		RHICmdList.SetDepthStencilState(TStaticDepthStencilState<
+			false,CF_Always,
+			true,CF_NotEqual,SO_Keep,SO_Keep,SO_Keep,
+			false,CF_Always,SO_Keep,SO_Keep,SO_Keep,
+			0xff,0xff
+			>::GetRHI());
 	}
 
 	// Light Attenuation channel assignment:
@@ -1849,12 +1846,12 @@ void FProjectedShadowInfo::RenderProjection(FRHICommandList& RHICmdList, int32 V
 		if(CascadeSettings.FadePlaneLength > 0)
 		{
 			// alpha is used to fade between cascades, we don't don't need to do BO_Min as we leave B and A untouched which has translucency shadow
-			RHISetBlendState(TStaticBlendState<CW_RG, BO_Add,BF_SourceAlpha,BF_InverseSourceAlpha>::GetRHI());
+			RHICmdList.SetBlendState(TStaticBlendState<CW_RG, BO_Add, BF_SourceAlpha, BF_InverseSourceAlpha>::GetRHI());
 		}
 		else
 		{
 			// first cascade rendered or old method doesn't require fading (CO_Min is needed to combine multiple shadow passes)
-			RHISetBlendState(TStaticBlendState<CW_RG, BO_Min,BF_One,BF_One>::GetRHI());
+			RHICmdList.SetBlendState(TStaticBlendState<CW_RG, BO_Min, BF_One, BF_One>::GetRHI());
 		}
 	}
 	else
@@ -1862,7 +1859,7 @@ void FProjectedShadowInfo::RenderProjection(FRHICommandList& RHICmdList, int32 V
 		// use B and A in Light Attenuation (see above)
 
 		// (CO_Min is needed to combine multiple shadow passes)
-		RHISetBlendState(TStaticBlendState<CW_BA, BO_Min,BF_One,BF_One, BO_Min,BF_One,BF_One>::GetRHI());
+		RHICmdList.SetBlendState(TStaticBlendState<CW_BA, BO_Min, BF_One, BF_One, BO_Min, BF_One, BF_One>::GetRHI());
 	}
 	
 	{
@@ -1921,24 +1918,24 @@ void FProjectedShadowInfo::RenderProjection(FRHICommandList& RHICmdList, int32 V
 			FVector4(-1.0f, -1.0f, 0.0f),
 			FVector4(1.0f, -1.0f, 0.0f),
 		};
-		RHIDrawPrimitiveUP(PT_TriangleStrip, 2, Verts, sizeof(FVector4));
+		DrawPrimitiveUP(RHICmdList, PT_TriangleStrip, 2, Verts, sizeof(FVector4));
 	}
 	else
 	{
 		// Draw the frustum using the projection shader..
-		RHIDrawIndexedPrimitiveUP( PT_TriangleList, 0, 8, 12, GCubeIndices, sizeof(uint16), FrustumVertices, sizeof(FVector4));
+		DrawIndexedPrimitiveUP(RHICmdList, PT_TriangleList, 0, 8, 12, GCubeIndices, sizeof(uint16), FrustumVertices, sizeof(FVector4));
 	}
 
 	if( bDepthBoundsTestEnabled )
 	{
 		// Disable depth bounds testing
-		RHIEnableDepthBoundsTest(false, 0, 1);
+		RHICmdList.EnableDepthBoundsTest(false, 0, 1);
 	}
 	else
 	{
-	// Clear the stencil buffer to 0.
-	RHIClear(false,FColor(0,0,0),false,0,true,0, FIntRect());
-}
+		// Clear the stencil buffer to 0.
+		RHICmdList.Clear(false, FColor(0, 0, 0), false, 0, true, 0, FIntRect());
+	}
 }
 
 
@@ -1949,7 +1946,8 @@ static void SetPointLightShaderTempl(FRHICommandList& RHICmdList, int32 ViewInde
 	TShaderMapRef<TOnePassPointShadowProjectionPS<Quality> > PixelShader(GetGlobalShaderMap());
 
 	static FGlobalBoundShaderState BoundShaderState;
-	SetGlobalBoundShaderState(BoundShaderState, GetVertexDeclarationFVector4(), *VertexShader, *PixelShader);
+	RHICmdList.CheckIsNull(); // need new approach for "static FGlobalBoundShaderState" for parallel rendering
+	SetGlobalBoundShaderState(RHICmdList, BoundShaderState, GetVertexDeclarationFVector4(), *VertexShader, *PixelShader);
 
 	VertexShader->SetParameters(RHICmdList, View, ShadowInfo);
 	PixelShader->SetParameters(RHICmdList, ViewIndex, View, ShadowInfo);
@@ -1960,29 +1958,27 @@ void FProjectedShadowInfo::RenderOnePassPointLightProjection(FRHICommandList& RH
 {
 	SCOPE_CYCLE_COUNTER(STAT_RenderWholeSceneShadowProjectionsTime);
 
-	RHICmdList.CheckIsNull();
-
 	checkSlow(bOnePassPointLightShadow);
 	
 	const FSphere LightBounds = LightSceneInfo->Proxy->GetBoundingSphere();
 
 	// Use min blending to RGB since shadows may overlap, preserve Alpha
-	RHISetBlendState(TStaticBlendState<CW_RGBA,BO_Min,BF_One,BF_One,BO_Add,BF_One,BF_Zero>::GetRHI());
+	RHICmdList.SetBlendState(TStaticBlendState<CW_RGBA, BO_Min, BF_One, BF_One, BO_Add, BF_One, BF_Zero>::GetRHI());
 
 	const bool bCameraInsideLightGeometry = ((FVector)View.ViewMatrices.ViewOrigin - LightBounds.Center).SizeSquared() < FMath::Square(LightBounds.W * 1.05f + View.NearClippingDistance * 2.0f);
 
 	if (bCameraInsideLightGeometry)
 	{
-		RHISetDepthStencilState(TStaticDepthStencilState<false, CF_Always>::GetRHI());
+		RHICmdList.SetDepthStencilState(TStaticDepthStencilState<false, CF_Always>::GetRHI());
 		// Render backfaces with depth tests disabled since the camera is inside (or close to inside) the light geometry
-		RHISetRasterizerState(View.bReverseCulling ? TStaticRasterizerState<FM_Solid,CM_CW>::GetRHI() : TStaticRasterizerState<FM_Solid,CM_CCW>::GetRHI());
+		RHICmdList.SetRasterizerState(View.bReverseCulling ? TStaticRasterizerState<FM_Solid, CM_CW>::GetRHI() : TStaticRasterizerState<FM_Solid, CM_CCW>::GetRHI());
 	}
 	else
 	{
 		// Render frontfaces with depth tests on to get the speedup from HiZ since the camera is outside the light geometry
 		// Note, this is a reversed Z depth surface, using CF_GreaterEqual.
-		RHISetDepthStencilState(TStaticDepthStencilState<false, CF_GreaterEqual>::GetRHI());
-		RHISetRasterizerState(View.bReverseCulling ? TStaticRasterizerState<FM_Solid,CM_CCW>::GetRHI() : TStaticRasterizerState<FM_Solid,CM_CW>::GetRHI());
+		RHICmdList.SetDepthStencilState(TStaticDepthStencilState<false, CF_GreaterEqual>::GetRHI());
+		RHICmdList.SetRasterizerState(View.bReverseCulling ? TStaticRasterizerState<FM_Solid, CM_CCW>::GetRHI() : TStaticRasterizerState<FM_Solid, CM_CW>::GetRHI());
 	}
 
 	{
@@ -2020,7 +2016,7 @@ void FProjectedShadowInfo::RenderOnePassPointLightProjection(FRHICommandList& RH
 
 	// Project the point light shadow with some approximately bounding geometry, 
 	// So we can get speedups from depth testing and not processing pixels outside of the light's influence.
-	StencilingGeometry::DrawSphere();
+	StencilingGeometry::DrawSphere(RHICmdList);
 }
 
 void FProjectedShadowInfo::RenderFrustumWireframe(FPrimitiveDrawInterface* PDI) const
@@ -2394,7 +2390,7 @@ bool FDeferredShadingSceneRenderer::RenderOnePassPointLightShadows(FRHICommandLi
 				bInjectedTranslucentVolume = true;
 				SCOPED_DRAW_EVENT(InjectTranslucentVolume, DEC_SCENE_ITEMS);
 				// Inject the shadowed light into the translucency lighting volumes
-				InjectTranslucentVolumeLighting(*LightSceneInfo, ProjectedShadowInfo);
+				InjectTranslucentVolumeLighting(RHICmdList, *LightSceneInfo, ProjectedShadowInfo);
 			}
 
 			bDirty = true;
@@ -2623,7 +2619,7 @@ bool FDeferredShadingSceneRenderer::RenderTranslucentProjectedShadows(FRHIComman
 					&& !ProjectedShadowInfo->bDirectionalLight)
 				{
 					// Accumulate this object shadow's contribution on the translucency lighting volume
-					AccumulateTranslucentVolumeObjectShadowing(ProjectedShadowInfo, !bRenderedToVolumeObjectShadowing);
+					AccumulateTranslucentVolumeObjectShadowing(RHICmdList, ProjectedShadowInfo, !bRenderedToVolumeObjectShadowing);
 					bRenderedToVolumeObjectShadowing = true;
 				}
 			}
@@ -2969,7 +2965,7 @@ bool FDeferredShadingSceneRenderer::RenderProjectedShadows(FRHICommandList& RHIC
 					bInjectedTranslucentVolume = true;
 					SCOPED_DRAW_EVENT(InjectTranslucentVolume, DEC_SCENE_ITEMS);
 					// Inject the shadowed light into the translucency lighting volumes
-					InjectTranslucentVolumeLighting(*LightSceneInfo, ProjectedShadowInfo);
+					InjectTranslucentVolumeLighting(RHICmdList, *LightSceneInfo, ProjectedShadowInfo);
 				}
 			}
 
@@ -3064,7 +3060,7 @@ bool FDeferredShadingSceneRenderer::RenderCachedPreshadows(FRHICommandList& RHIC
 		{
 			SCOPED_DRAW_EVENTF(EventShadowDepths, DEC_SCENE_ITEMS, TEXT("Preshadow Cache Depths"));
 
-			RHISetRenderTarget(FTextureRHIRef(), GSceneRenderTargets.PreShadowCacheDepthZ->GetRenderTargetItem().TargetableTexture);
+			SetRenderTarget(RHICmdList, FTextureRHIRef(), GSceneRenderTargets.PreShadowCacheDepthZ->GetRenderTargetItem().TargetableTexture);
 
 			// Render depth for each shadow
 			for (int32 ShadowIndex = 0; ShadowIndex < CachedPreshadows.Num(); ShadowIndex++)
diff --git a/Engine/Source/Runtime/Renderer/Private/ShadowRendering.h b/Engine/Source/Runtime/Renderer/Private/ShadowRendering.h
--- a/Engine/Source/Runtime/Renderer/Private/ShadowRendering.h
+++ b/Engine/Source/Runtime/Renderer/Private/ShadowRendering.h
@@ -19,9 +19,9 @@ namespace StencilingGeometry
 	* Draws a sphere using RHIDrawIndexedPrimitive, useful as approximate bounding geometry for deferred passes.
 	* Note: The sphere will be of unit size unless transformed by the shader. 
 	*/
-	extern void DrawSphere();
+	extern void DrawSphere(FRHICommandList& RHICmdList);
 	/** Renders a cone with a spherical cap, used for rendering spot lights in deferred passes. */
-	extern void DrawCone();
+	extern void DrawCone(FRHICommandList& RHICmdList);
 
 	/** 
 	* Vertex buffer for a sphere of unit size. Used for drawing a sphere as approximate bounding geometry for deferred passes.
