UE4 - more work preparing for parallel rendering. Merged the cmd list into the post process context, among other things.

diff --git a/Engine/Source/Runtime/Renderer/Private/TranslucentLighting.cpp b/Engine/Source/Runtime/Renderer/Private/TranslucentLighting.cpp
--- a/Engine/Source/Runtime/Renderer/Private/TranslucentLighting.cpp
+++ b/Engine/Source/Runtime/Renderer/Private/TranslucentLighting.cpp
@@ -554,7 +554,7 @@ void FProjectedShadowInfo::RenderTranslucencyDepths(FDeferredShadingSceneRendere
 			);
 
 		FLinearColor ClearColors[2] = {FLinearColor(0,0,0,0), FLinearColor(0,0,0,0)};
-		DrawClearQuadMRT(true, ARRAY_COUNT(ClearColors), ClearColors, false, 1.0f, false, 0);
+		DrawClearQuadMRT(RHICmdList, true, ARRAY_COUNT(ClearColors), ClearColors, false, 1.0f, false, 0);
 
 		// Set the viewport for the shadow.
 		RHISetViewport(
@@ -1016,43 +1016,42 @@ public:
 TGlobalResource<FVolumeRasterizeVertexBuffer> GVolumeRasterizeVertexBuffer;
 
 /** Draws a quad per volume texture slice to the subregion of the volume texture specified. */
-void RasterizeToVolumeTexture(FVolumeBounds VolumeBounds)
+void RasterizeToVolumeTexture(FRHICommandList& RHICmdList, FVolumeBounds VolumeBounds)
 {
 	// Setup the viewport to only render to the given bounds
-	RHISetViewport(VolumeBounds.MinX, VolumeBounds.MinY, 0, VolumeBounds.MaxX, VolumeBounds.MaxY, 0);
-	RHISetStreamSource(0, GVolumeRasterizeVertexBuffer.VertexBufferRHI, sizeof(FScreenVertex), 0);
+	RHICmdList.SetViewport(VolumeBounds.MinX, VolumeBounds.MinY, 0, VolumeBounds.MaxX, VolumeBounds.MaxY, 0);
+	RHICmdList.SetStreamSource(0, GVolumeRasterizeVertexBuffer.VertexBufferRHI, sizeof(FScreenVertex), 0);
 	const int32 NumInstances = VolumeBounds.MaxZ - VolumeBounds.MinZ;
 	// Render a quad per slice affected by the given bounds
-	RHIDrawPrimitive(PT_TriangleStrip, 0, 2, NumInstances);
+	RHICmdList.DrawPrimitive(PT_TriangleStrip, 0, 2, NumInstances);
 }
 
 /** Helper function that clears the given volume texture render targets. */
 template<int32 NumRenderTargets>
 void ClearVolumeTextures(FRHICommandList& RHICmdList, const FTextureRHIParamRef* RenderTargets, const FLinearColor* ClearColors)
 {
-	RHICmdList.CheckIsNull();
-
-	RHISetRenderTargets(NumRenderTargets, RenderTargets, FTextureRHIRef(), 0, NULL);
+	SetRenderTargets(RHICmdList, NumRenderTargets, RenderTargets, FTextureRHIRef(), 0, NULL);
 
 	static FGlobalBoundShaderState VolumeClearBoundShaderState;
+	RHICmdList.CheckIsNull(); // need new approach for "static FGlobalBoundShaderState" for parallel rendering
 
 	// Currently using a manual clear, which is ~10x faster than a hardware clear of the volume textures on AMD PC GPU's
 	if (false)
 	{
-		RHIClearMRT(true, NumRenderTargets, ClearColors, false, 0, false, 0, FIntRect());
+		RHICmdList.ClearMRT(true, NumRenderTargets, ClearColors, false, 0, false, 0, FIntRect());
 	}
 	else
 	{
-		RHISetRasterizerState(TStaticRasterizerState<FM_Solid,CM_None>::GetRHI());
-		RHISetDepthStencilState(TStaticDepthStencilState<false,CF_Always>::GetRHI());
-		RHISetBlendState(TStaticBlendState<>::GetRHI());
+		RHICmdList.SetRasterizerState(TStaticRasterizerState<FM_Solid, CM_None>::GetRHI());
+		RHICmdList.SetDepthStencilState(TStaticDepthStencilState<false, CF_Always>::GetRHI());
+		RHICmdList.SetBlendState(TStaticBlendState<>::GetRHI());
 
 		const FVolumeBounds VolumeBounds(GTranslucencyLightingVolumeDim);
 		TShaderMapRef<FWriteToSliceVS> VertexShader(GetGlobalShaderMap());
 		TShaderMapRef<FWriteToSliceGS> GeometryShader(GetGlobalShaderMap());
 		TShaderMapRef<TOneColorPixelShaderMRT<NumRenderTargets> > PixelShader(GetGlobalShaderMap());
 		
-		SetGlobalBoundShaderState(VolumeClearBoundShaderState, GScreenVertexDeclaration.VertexDeclarationRHI, *VertexShader, *PixelShader, *GeometryShader);
+		SetGlobalBoundShaderState(RHICmdList, VolumeClearBoundShaderState, GScreenVertexDeclaration.VertexDeclarationRHI, *VertexShader, *PixelShader, *GeometryShader);
 
 		VertexShader->SetParameters(RHICmdList, VolumeBounds, GTranslucencyLightingVolumeDim);
 		GeometryShader->SetParameters(RHICmdList, VolumeBounds);
@@ -1067,7 +1066,7 @@ void ClearVolumeTextures(FRHICommandList& RHICmdList, const FTextureRHIParamRef*
 
 		SetShaderValueArray(RHICmdList, PixelShader->GetPixelShader(), PixelShader->ColorParameter, ShaderClearColors, NumRenderTargets);
 
-		RasterizeToVolumeTexture(VolumeBounds);
+		RasterizeToVolumeTexture(RHICmdList, VolumeBounds);
 	}
 }
 
@@ -1149,24 +1148,24 @@ void FDeferredShadingSceneRenderer::InjectAmbientCubemapTranslucentVolumeLightin
 	{
 		SCOPED_DRAW_EVENT(InjectAmbientCubemapTranslucentVolumeLighting, DEC_SCENE_ITEMS);
 
-		RHISetRasterizerState(TStaticRasterizerState<FM_Solid,CM_None>::GetRHI());
-		RHISetDepthStencilState(TStaticDepthStencilState<false,CF_Always>::GetRHI());
-		RHISetBlendState(TStaticBlendState<CW_RGBA, BO_Add, BF_One, BF_One, BO_Add, BF_One, BF_One>::GetRHI());
+		RHICmdList.SetRasterizerState(TStaticRasterizerState<FM_Solid, CM_None>::GetRHI());
+		RHICmdList.SetDepthStencilState(TStaticDepthStencilState<false, CF_Always>::GetRHI());
+		RHICmdList.SetBlendState(TStaticBlendState<CW_RGBA, BO_Add, BF_One, BF_One, BO_Add, BF_One, BF_One>::GetRHI());
 
 		const FVolumeBounds VolumeBounds(GTranslucencyLightingVolumeDim);
 		
 		for (int32 VolumeCascadeIndex = 0; VolumeCascadeIndex < TVC_MAX; VolumeCascadeIndex++)
 		{
 			// we don't update the directional volume (could be a HQ option)
-			RHISetRenderTarget(GSceneRenderTargets.TranslucencyLightingVolumeAmbient[VolumeCascadeIndex]->GetRenderTargetItem().TargetableTexture, FTextureRHIRef());
+			SetRenderTarget(RHICmdList, GSceneRenderTargets.TranslucencyLightingVolumeAmbient[VolumeCascadeIndex]->GetRenderTargetItem().TargetableTexture, FTextureRHIRef());
 
 			TShaderMapRef<FWriteToSliceVS> VertexShader(GetGlobalShaderMap());
 			TShaderMapRef<FWriteToSliceGS> GeometryShader(GetGlobalShaderMap());
 			TShaderMapRef<FInjectAmbientCubemapPS> PixelShader(GetGlobalShaderMap());
 
 			static FGlobalBoundShaderState BoundShaderState;
-						
-			SetGlobalBoundShaderState(BoundShaderState, GScreenVertexDeclaration.VertexDeclarationRHI, *VertexShader, *PixelShader, *GeometryShader);
+			RHICmdList.CheckIsNull(); // need new approach for "static FGlobalBoundShaderState" for parallel rendering						
+			SetGlobalBoundShaderState(RHICmdList, BoundShaderState, GScreenVertexDeclaration.VertexDeclarationRHI, *VertexShader, *PixelShader, *GeometryShader);
 
 			VertexShader->SetParameters(RHICmdList, VolumeBounds, GTranslucencyLightingVolumeDim);
 			GeometryShader->SetParameters(RHICmdList, VolumeBounds);
@@ -1178,10 +1177,10 @@ void FDeferredShadingSceneRenderer::InjectAmbientCubemapTranslucentVolumeLightin
 
 				PixelShader->SetParameters(RHICmdList, View, CubemapEntry);
 
-				RasterizeToVolumeTexture(VolumeBounds);
+				RasterizeToVolumeTexture(RHICmdList, VolumeBounds);
 			}
 			
-			RHICopyToResolveTarget(GSceneRenderTargets.TranslucencyLightingVolumeAmbient[VolumeCascadeIndex]->GetRenderTargetItem().TargetableTexture,
+			RHICmdList.CopyToResolveTarget(GSceneRenderTargets.TranslucencyLightingVolumeAmbient[VolumeCascadeIndex]->GetRenderTargetItem().TargetableTexture,
 				GSceneRenderTargets.TranslucencyLightingVolumeAmbient[VolumeCascadeIndex]->GetRenderTargetItem().ShaderResourceTexture, true, FResolveParams());
 		}
 	}
@@ -1308,9 +1307,9 @@ void FDeferredShadingSceneRenderer::CompositeIndirectTranslucentVolumeLighting(F
 	{
 		SCOPED_DRAW_EVENT(CompositeIndirectTranslucentVolumeLighting, DEC_SCENE_ITEMS);
 
-		RHISetRasterizerState(TStaticRasterizerState<FM_Solid,CM_None>::GetRHI());
-		RHISetDepthStencilState(TStaticDepthStencilState<false,CF_Always>::GetRHI());
-		RHISetBlendState(TStaticBlendState<CW_RGBA, BO_Add, BF_One, BF_One, BO_Add, BF_One, BF_One>::GetRHI());
+		RHICmdList.SetRasterizerState(TStaticRasterizerState<FM_Solid, CM_None>::GetRHI());
+		RHICmdList.SetDepthStencilState(TStaticDepthStencilState<false, CF_Always>::GetRHI());
+		RHICmdList.SetBlendState(TStaticBlendState<CW_RGBA, BO_Add, BF_One, BF_One, BO_Add, BF_One, BF_One>::GetRHI());
 
 		const FVolumeBounds VolumeBounds(GTranslucencyLightingVolumeDim);
 
@@ -1320,7 +1319,7 @@ void FDeferredShadingSceneRenderer::CompositeIndirectTranslucentVolumeLighting(F
 			RenderTargets[0] = GSceneRenderTargets.TranslucencyLightingVolumeAmbient[VolumeCascadeIndex]->GetRenderTargetItem().TargetableTexture;
 			RenderTargets[1] = GSceneRenderTargets.TranslucencyLightingVolumeDirectional[VolumeCascadeIndex]->GetRenderTargetItem().TargetableTexture;
 
-			RHISetRenderTargets(ARRAY_COUNT(RenderTargets), RenderTargets, FTextureRHIRef(), 0, NULL);
+			SetRenderTargets(RHICmdList, ARRAY_COUNT(RenderTargets), RenderTargets, FTextureRHIRef(), 0, NULL);
 
 
 			//@todo - support multiple views
@@ -1336,24 +1335,25 @@ void FDeferredShadingSceneRenderer::CompositeIndirectTranslucentVolumeLighting(F
 				TShaderMapRef<FCompositeGIForTranslucencyPS> PixelShader(GetGlobalShaderMap());
 
 				static FGlobalBoundShaderState BoundShaderState;
-				SetGlobalBoundShaderState(BoundShaderState, GScreenVertexDeclaration.VertexDeclarationRHI, *VertexShader, *PixelShader, *GeometryShader);
+				RHICmdList.CheckIsNull(); // need new approach for "static FGlobalBoundShaderState" for parallel rendering
+				SetGlobalBoundShaderState(RHICmdList, BoundShaderState, GScreenVertexDeclaration.VertexDeclarationRHI, *VertexShader, *PixelShader, *GeometryShader);
 
 				PixelShader->SetParameters(RHICmdList, View, ViewState, VolumeCascadeIndex);
 				VertexShader->SetParameters(RHICmdList, VolumeBounds, GTranslucencyLightingVolumeDim);
 				GeometryShader->SetParameters(RHICmdList, VolumeBounds);
 
-				RasterizeToVolumeTexture(VolumeBounds);
+				RasterizeToVolumeTexture(RHICmdList, VolumeBounds);
 
-				RHICopyToResolveTarget(GSceneRenderTargets.TranslucencyLightingVolumeAmbient[VolumeCascadeIndex]->GetRenderTargetItem().TargetableTexture,
+				RHICmdList.CopyToResolveTarget(GSceneRenderTargets.TranslucencyLightingVolumeAmbient[VolumeCascadeIndex]->GetRenderTargetItem().TargetableTexture,
 					GSceneRenderTargets.TranslucencyLightingVolumeAmbient[VolumeCascadeIndex]->GetRenderTargetItem().ShaderResourceTexture, true, FResolveParams());
-				RHICopyToResolveTarget(GSceneRenderTargets.TranslucencyLightingVolumeDirectional[VolumeCascadeIndex]->GetRenderTargetItem().TargetableTexture,
+				RHICmdList.CopyToResolveTarget(GSceneRenderTargets.TranslucencyLightingVolumeDirectional[VolumeCascadeIndex]->GetRenderTargetItem().TargetableTexture,
 					GSceneRenderTargets.TranslucencyLightingVolumeDirectional[VolumeCascadeIndex]->GetRenderTargetItem().ShaderResourceTexture, true, FResolveParams());
 			}
 		}
 	}
 }
 
-void FDeferredShadingSceneRenderer::ClearTranslucentVolumePerObjectShadowing()
+void FDeferredShadingSceneRenderer::ClearTranslucentVolumePerObjectShadowing(FRHICommandList& RHICmdList)
 {
 	if (GUseTranslucentLightingVolumes && GSupportsVolumeTextureRendering)
 	{
@@ -1368,8 +1368,6 @@ void FDeferredShadingSceneRenderer::ClearTranslucentVolumePerObjectShadowing()
 		ClearColors[0] = FLinearColor(1, 1, 1, 1);
 		ClearColors[1] = FLinearColor(1, 1, 1, 1);
 
-		//@todo-rco: RHIPacketList
-		FRHICommandList& RHICmdList = FRHICommandList::GetNullRef();
 		ClearVolumeTextures<ARRAY_COUNT(RenderTargets)>(RHICmdList, RenderTargets, ClearColors);
 	}
 }
@@ -1403,26 +1401,24 @@ FVolumeBounds CalculateLightVolumeBounds(const FSphere& LightBounds, const FView
 
 FGlobalBoundShaderState ObjectShadowingBoundShaderState;
 
-void FDeferredShadingSceneRenderer::AccumulateTranslucentVolumeObjectShadowing(const FProjectedShadowInfo* InProjectedShadowInfo, bool bClearVolume)
+void FDeferredShadingSceneRenderer::AccumulateTranslucentVolumeObjectShadowing(FRHICommandList& RHICmdList, const FProjectedShadowInfo* InProjectedShadowInfo, bool bClearVolume)
 {
 	const FLightSceneInfo* LightSceneInfo = InProjectedShadowInfo->LightSceneInfo;
 
 	if (bClearVolume)
 	{
-		ClearTranslucentVolumePerObjectShadowing();
+		ClearTranslucentVolumePerObjectShadowing(RHICmdList);
 	}
 
-	//@todo-rco: RHIPacketList
-	FRHICommandList& RHICmdList = FRHICommandList::GetNullRef();
 	if (GUseTranslucentLightingVolumes && GSupportsVolumeTextureRendering)
 	{
 		SCOPED_DRAW_EVENT(AccumulateTranslucentVolumeShadowing, DEC_SCENE_ITEMS);
 
-		RHISetRasterizerState(TStaticRasterizerState<FM_Solid,CM_None>::GetRHI());
-		RHISetDepthStencilState(TStaticDepthStencilState<false,CF_Always>::GetRHI());
+		RHICmdList.SetRasterizerState(TStaticRasterizerState<FM_Solid, CM_None>::GetRHI());
+		RHICmdList.SetDepthStencilState(TStaticDepthStencilState<false, CF_Always>::GetRHI());
 
 		// Modulate the contribution of multiple object shadows in rgb
-		RHISetBlendState(TStaticBlendState<CW_RGB, BO_Add, BF_DestColor, BF_Zero>::GetRHI());
+		RHICmdList.SetBlendState(TStaticBlendState<CW_RGB, BO_Add, BF_DestColor, BF_Zero>::GetRHI());
 
 		// Inject into each volume cascade
 		for (int32 VolumeCascadeIndex = 0; VolumeCascadeIndex < TVC_MAX; VolumeCascadeIndex++)
@@ -1445,20 +1441,20 @@ void FDeferredShadingSceneRenderer::AccumulateTranslucentVolumeObjectShadowing(c
 					RenderTarget = GSceneRenderTargets.GetTranslucencyVolumeDirectional(TVC_Inner)->GetRenderTargetItem().TargetableTexture;
 				}
 
-				RHISetRenderTarget(RenderTarget, FTextureRHIRef());
+				SetRenderTarget(RHICmdList, RenderTarget, FTextureRHIRef());
 
 				TShaderMapRef<FWriteToSliceVS> VertexShader(GetGlobalShaderMap());
 				TShaderMapRef<FWriteToSliceGS> GeometryShader(GetGlobalShaderMap());
 				TShaderMapRef<FTranslucentObjectShadowingPS> PixelShader(GetGlobalShaderMap());
 
-				SetGlobalBoundShaderState(ObjectShadowingBoundShaderState, GScreenVertexDeclaration.VertexDeclarationRHI, *VertexShader, *PixelShader, *GeometryShader);
+				SetGlobalBoundShaderState(RHICmdList, ObjectShadowingBoundShaderState, GScreenVertexDeclaration.VertexDeclarationRHI, *VertexShader, *PixelShader, *GeometryShader);
 				VertexShader->SetParameters(RHICmdList, VolumeBounds, GTranslucencyLightingVolumeDim);
 				GeometryShader->SetParameters(RHICmdList, VolumeBounds);
 				PixelShader->SetParameters(RHICmdList, View, LightSceneInfo, InProjectedShadowInfo, VolumeCascadeIndex);
 				
-				RasterizeToVolumeTexture(VolumeBounds);
+				RasterizeToVolumeTexture(RHICmdList, VolumeBounds);
 
-				RHICopyToResolveTarget(GSceneRenderTargets.GetTranslucencyVolumeAmbient((ETranslucencyVolumeCascade)VolumeCascadeIndex)->GetRenderTargetItem().TargetableTexture,
+				RHICmdList.CopyToResolveTarget(GSceneRenderTargets.GetTranslucencyVolumeAmbient((ETranslucencyVolumeCascade)VolumeCascadeIndex)->GetRenderTargetItem().TargetableTexture,
 					GSceneRenderTargets.GetTranslucencyVolumeAmbient((ETranslucencyVolumeCascade)VolumeCascadeIndex)->GetRenderTargetItem().ShaderResourceTexture, true, FResolveParams());
 			}
 		}
@@ -1472,6 +1468,7 @@ void FDeferredShadingSceneRenderer::AccumulateTranslucentVolumeObjectShadowing(c
  */
 template<ELightComponentType InjectionType, bool bShadowed>
 void SetInjectionShader(
+	FRHICommandList& RHICmdList,
 	const FViewInfo& View, 
 	const FMaterialRenderProxy* MaterialProxy,
 	const FLightSceneInfo* LightSceneInfo, 
@@ -1483,9 +1480,6 @@ void SetInjectionShader(
 	bool bApplyLightFunction,
 	bool bInverseSquared)
 {
-	//@todo-rco: RHIPacketList
-	FRHICommandList& RHICmdList = FRHICommandList::GetNullRef();
-
 	check(ShadowMap || !bShadowed);
 
 	const FMaterialShaderMap* MaterialShaderMap = MaterialProxy->GetMaterial(View.GetFeatureLevel())->GetRenderingThreadShaderMap();
@@ -1533,13 +1527,15 @@ void SetInjectionShader(
 
 	// Recreate the bound shader state if the shader map has changed since the last cache
 	// This can happen due to async shader compiling
+	RHICmdList.CheckIsNull(); // RHICreateBoundShaderState
+
 	if (!IsValidRef(BoundShaderState) || CachedShaderMap != MaterialShaderMap)
 	{
 		CachedShaderMap = MaterialShaderMap;
 		BoundShaderState = RHICreateBoundShaderState(GScreenVertexDeclaration.VertexDeclarationRHI, VertexShader->GetVertexShader(), FHullShaderRHIRef(), FDomainShaderRHIRef(), PixelShader->GetPixelShader(), GeometryShader->GetGeometryShader());
 	}
 
-	RHISetBoundShaderState(BoundShaderState);
+	RHICmdList.SetBoundShaderState(BoundShaderState);
 
 	// Now shader is set, bind parameters
 	if (bApplyLightFunction)
@@ -1628,13 +1624,10 @@ static void AddLightForInjection(
 }
 
 /** Injects all the lights in LightInjectionData into the translucent lighting volume textures. */
-static void InjectTranslucentLightArray(const FViewInfo& View, const TArray<FTranslucentLightInjectionData, SceneRenderingAllocator>& LightInjectionData)
+static void InjectTranslucentLightArray(FRHICommandList& RHICmdList, const FViewInfo& View, const TArray<FTranslucentLightInjectionData, SceneRenderingAllocator>& LightInjectionData)
 {
 	INC_DWORD_STAT_BY(STAT_NumLightsInjectedIntoTranslucency, LightInjectionData.Num());
 
-	//@todo-rco: RHIPacketList
-	FRHICommandList& RHICmdList = FRHICommandList::GetNullRef();
-
 	// Inject into each volume cascade
 	// Operate on one cascade at a time to reduce render target switches
 	for (int32 VolumeCascadeIndex = 0; VolumeCascadeIndex < TVC_MAX; VolumeCascadeIndex++)
@@ -1642,17 +1635,17 @@ static void InjectTranslucentLightArray(const FViewInfo& View, const TArray<FTra
 		const IPooledRenderTarget* RT0 = GSceneRenderTargets.TranslucencyLightingVolumeAmbient[VolumeCascadeIndex];
 		const IPooledRenderTarget* RT1 = GSceneRenderTargets.TranslucencyLightingVolumeDirectional[VolumeCascadeIndex];
 
-		GRenderTargetPool.VisualizeTexture.SetCheckPoint(RT0);
-		GRenderTargetPool.VisualizeTexture.SetCheckPoint(RT1);
+		GRenderTargetPool.VisualizeTexture.SetCheckPoint(RHICmdList, RT0);
+		GRenderTargetPool.VisualizeTexture.SetCheckPoint(RHICmdList, RT1);
 
 		FTextureRHIParamRef RenderTargets[2];
 		RenderTargets[0] = RT0->GetRenderTargetItem().TargetableTexture;
 		RenderTargets[1] = RT1->GetRenderTargetItem().TargetableTexture;
 
-		RHISetRenderTargets(ARRAY_COUNT(RenderTargets), RenderTargets, FTextureRHIRef(), 0, NULL);
+		SetRenderTargets(RHICmdList, ARRAY_COUNT(RenderTargets), RenderTargets, FTextureRHIRef(), 0, NULL);
 
-		RHISetRasterizerState(TStaticRasterizerState<FM_Solid,CM_None>::GetRHI());
-		RHISetDepthStencilState(TStaticDepthStencilState<false,CF_Always>::GetRHI());
+		RHICmdList.SetRasterizerState(TStaticRasterizerState<FM_Solid, CM_None>::GetRHI());
+		RHICmdList.SetDepthStencilState(TStaticDepthStencilState<false, CF_Always>::GetRHI());
 
 		for (int32 LightIndex = 0; LightIndex < LightInjectionData.Num(); LightIndex++)
 		{
@@ -1671,50 +1664,50 @@ static void InjectTranslucentLightArray(const FViewInfo& View, const TArray<FTra
 				{
 					// Accumulate the contribution of multiple lights
 					// Directional lights write their shadowing into alpha of the ambient texture
-					RHISetBlendState(TStaticBlendState<
+					RHICmdList.SetBlendState(TStaticBlendState<
 						CW_RGBA, BO_Add, BF_One, BF_One, BO_Add, BF_One, BF_One,
 						CW_RGB, BO_Add, BF_One, BF_One, BO_Add, BF_One, BF_One>::GetRHI());
 					
 					if (InjectionData.ProjectedShadowInfo)
 					{
 						// shadows, restricting light contribution to the cascade bounds (except last cascade far to get light functions and no shadows there)
-						SetInjectionShader<LightType_Directional,true>(View, InjectionData.LightFunctionMaterialProxy, LightSceneInfo, InjectionData.ProjectedShadowInfo, InjectionData.ProjectedShadowInfo->SplitIndex, VolumeCascadeIndex, *VertexShader, *GeometryShader, InjectionData.bApplyLightFunction, false);
+						SetInjectionShader<LightType_Directional, true>(RHICmdList, View, InjectionData.LightFunctionMaterialProxy, LightSceneInfo, InjectionData.ProjectedShadowInfo, InjectionData.ProjectedShadowInfo->SplitIndex, VolumeCascadeIndex, *VertexShader, *GeometryShader, InjectionData.bApplyLightFunction, false);
 					}
 					else
 					{
 						// no shadows
-						SetInjectionShader<LightType_Directional,false>(View, InjectionData.LightFunctionMaterialProxy, LightSceneInfo, InjectionData.ProjectedShadowInfo, -1, VolumeCascadeIndex, *VertexShader, *GeometryShader, InjectionData.bApplyLightFunction, false);
+						SetInjectionShader<LightType_Directional, false>(RHICmdList, View, InjectionData.LightFunctionMaterialProxy, LightSceneInfo, InjectionData.ProjectedShadowInfo, -1, VolumeCascadeIndex, *VertexShader, *GeometryShader, InjectionData.bApplyLightFunction, false);
 					}
 				}
 				else
 				{
 					// Accumulate the contribution of multiple lights
-					RHISetBlendState(TStaticBlendState<
+					RHICmdList.SetBlendState(TStaticBlendState<
 						CW_RGB, BO_Add, BF_One, BF_One, BO_Add, BF_Zero, BF_One,
 						CW_RGB, BO_Add, BF_One, BF_One, BO_Add, BF_Zero, BF_One>::GetRHI());
 
 					if (InjectionData.ProjectedShadowInfo)
 					{
-						SetInjectionShader<LightType_Point,true>(View, InjectionData.LightFunctionMaterialProxy, LightSceneInfo, InjectionData.ProjectedShadowInfo, -1, VolumeCascadeIndex, *VertexShader, *GeometryShader, InjectionData.bApplyLightFunction, bInverseSquared);
+						SetInjectionShader<LightType_Point, true>(RHICmdList, View, InjectionData.LightFunctionMaterialProxy, LightSceneInfo, InjectionData.ProjectedShadowInfo, -1, VolumeCascadeIndex, *VertexShader, *GeometryShader, InjectionData.bApplyLightFunction, bInverseSquared);
 					}
 					else
 					{
-						SetInjectionShader<LightType_Point,false>(View, InjectionData.LightFunctionMaterialProxy, LightSceneInfo, InjectionData.ProjectedShadowInfo, -1, VolumeCascadeIndex, *VertexShader, *GeometryShader, InjectionData.bApplyLightFunction, bInverseSquared);
+						SetInjectionShader<LightType_Point, false>(RHICmdList, View, InjectionData.LightFunctionMaterialProxy, LightSceneInfo, InjectionData.ProjectedShadowInfo, -1, VolumeCascadeIndex, *VertexShader, *GeometryShader, InjectionData.bApplyLightFunction, bInverseSquared);
 					}
 				}
 
 				VertexShader->SetParameters(RHICmdList, VolumeBounds, GTranslucencyLightingVolumeDim);
 				GeometryShader->SetParameters(RHICmdList, VolumeBounds);
-				RasterizeToVolumeTexture(VolumeBounds);
+				RasterizeToVolumeTexture(RHICmdList, VolumeBounds);
 			}
 		}
 
-		RHICopyToResolveTarget(RT0->GetRenderTargetItem().TargetableTexture, RT0->GetRenderTargetItem().ShaderResourceTexture, true, FResolveParams());
-		RHICopyToResolveTarget(RT1->GetRenderTargetItem().TargetableTexture, RT1->GetRenderTargetItem().ShaderResourceTexture, true, FResolveParams());
+		RHICmdList.CopyToResolveTarget(RT0->GetRenderTargetItem().TargetableTexture, RT0->GetRenderTargetItem().ShaderResourceTexture, true, FResolveParams());
+		RHICmdList.CopyToResolveTarget(RT1->GetRenderTargetItem().TargetableTexture, RT1->GetRenderTargetItem().ShaderResourceTexture, true, FResolveParams());
 	}
 }
 
-void FDeferredShadingSceneRenderer::InjectTranslucentVolumeLighting(const FLightSceneInfo& LightSceneInfo, const FProjectedShadowInfo* InProjectedShadowInfo)
+void FDeferredShadingSceneRenderer::InjectTranslucentVolumeLighting(FRHICommandList& RHICmdList, const FLightSceneInfo& LightSceneInfo, const FProjectedShadowInfo* InProjectedShadowInfo)
 {
 	if (GUseTranslucentLightingVolumes && GSupportsVolumeTextureRendering)
 	{
@@ -1728,11 +1721,11 @@ void FDeferredShadingSceneRenderer::InjectTranslucentVolumeLighting(const FLight
 		AddLightForInjection(*this, LightSceneInfo, InProjectedShadowInfo, LightInjectionData);
 
 		// shadowed or unshadowed (InProjectedShadowInfo==0)
-		InjectTranslucentLightArray(View, LightInjectionData);
+		InjectTranslucentLightArray(RHICmdList, View, LightInjectionData);
 	}
 }
 
-void FDeferredShadingSceneRenderer::InjectTranslucentVolumeLightingArray(const TArray<FSortedLightSceneInfo, SceneRenderingAllocator>& SortedLights, int32 NumLights)
+void FDeferredShadingSceneRenderer::InjectTranslucentVolumeLightingArray(FRHICommandList& RHICmdList, const TArray<FSortedLightSceneInfo, SceneRenderingAllocator>& SortedLights, int32 NumLights)
 {
 	SCOPE_CYCLE_COUNTER(STAT_TranslucentInjectTime);
 
@@ -1752,7 +1745,7 @@ void FDeferredShadingSceneRenderer::InjectTranslucentVolumeLightingArray(const T
 	}
 
 	// non-shadowed, non-light function lights
-	InjectTranslucentLightArray(View, LightInjectionData);
+	InjectTranslucentLightArray(RHICmdList, View, LightInjectionData);
 }
 
 /** Pixel shader used to inject simple lights into the translucent lighting volume */
@@ -1804,11 +1797,8 @@ IMPLEMENT_SHADER_TYPE(,FSimpleLightTranslucentLightingInjectPS,TEXT("Translucent
 
 FGlobalBoundShaderState InjectSimpleLightBoundShaderState;
 
-void FDeferredShadingSceneRenderer::InjectSimpleTranslucentVolumeLightingArray(const FSimpleLightArray& SimpleLights)
+void FDeferredShadingSceneRenderer::InjectSimpleTranslucentVolumeLightingArray(FRHICommandList& RHICmdList, const FSimpleLightArray& SimpleLights)
 {
-	//@todo-rco: RHIPacketList
-	FRHICommandList& RHICmdList = FRHICommandList::GetNullRef();
-
 	SCOPE_CYCLE_COUNTER(STAT_TranslucentInjectTime);
 
 	int32 NumLightsToInject = 0;
@@ -1836,17 +1826,17 @@ void FDeferredShadingSceneRenderer::InjectSimpleTranslucentVolumeLightingArray(c
 			const IPooledRenderTarget* RT0 = GSceneRenderTargets.TranslucencyLightingVolumeAmbient[VolumeCascadeIndex];
 			const IPooledRenderTarget* RT1 = GSceneRenderTargets.TranslucencyLightingVolumeDirectional[VolumeCascadeIndex];
 
-			GRenderTargetPool.VisualizeTexture.SetCheckPoint(RT0);
-			GRenderTargetPool.VisualizeTexture.SetCheckPoint(RT1);
+			GRenderTargetPool.VisualizeTexture.SetCheckPoint(RHICmdList, RT0);
+			GRenderTargetPool.VisualizeTexture.SetCheckPoint(RHICmdList, RT1);
 
 			FTextureRHIParamRef RenderTargets[2];
 			RenderTargets[0] = RT0->GetRenderTargetItem().TargetableTexture;
 			RenderTargets[1] = RT1->GetRenderTargetItem().TargetableTexture;
 
-			RHISetRenderTargets(ARRAY_COUNT(RenderTargets), RenderTargets, FTextureRHIRef(), 0, NULL);
+			SetRenderTargets(RHICmdList, ARRAY_COUNT(RenderTargets), RenderTargets, FTextureRHIRef(), 0, NULL);
 
-			RHISetRasterizerState(TStaticRasterizerState<FM_Solid,CM_None>::GetRHI());
-			RHISetDepthStencilState(TStaticDepthStencilState<false,CF_Always>::GetRHI());
+			RHICmdList.SetRasterizerState(TStaticRasterizerState<FM_Solid, CM_None>::GetRHI());
+			RHICmdList.SetDepthStencilState(TStaticDepthStencilState<false, CF_Always>::GetRHI());
 
 		
 			for (int32 LightIndex = 0; LightIndex < SimpleLights.InstanceData.Num(); LightIndex++)
@@ -1864,44 +1854,41 @@ void FDeferredShadingSceneRenderer::InjectSimpleTranslucentVolumeLightingArray(c
 						TShaderMapRef<FWriteToSliceVS> VertexShader(GetGlobalShaderMap());
 						TShaderMapRef<FWriteToSliceGS> GeometryShader(GetGlobalShaderMap());
 						TShaderMapRef<FSimpleLightTranslucentLightingInjectPS> PixelShader(GetGlobalShaderMap());
-						SetGlobalBoundShaderState(InjectSimpleLightBoundShaderState, GScreenVertexDeclaration.VertexDeclarationRHI, *VertexShader, *PixelShader, *GeometryShader);
+						SetGlobalBoundShaderState(RHICmdList, InjectSimpleLightBoundShaderState, GScreenVertexDeclaration.VertexDeclarationRHI, *VertexShader, *PixelShader, *GeometryShader);
 
 						VertexShader->SetParameters(RHICmdList, VolumeBounds, GTranslucencyLightingVolumeDim);
 						GeometryShader->SetParameters(RHICmdList, VolumeBounds);
 						PixelShader->SetParameters(RHICmdList, View, SimpleLight, SimpleLightPerViewData, VolumeCascadeIndex);
 
 						// Accumulate the contribution of multiple lights
-						RHISetBlendState(TStaticBlendState<
+						RHICmdList.SetBlendState(TStaticBlendState<
 							CW_RGB, BO_Add, BF_One, BF_One, BO_Add, BF_Zero, BF_One,
 							CW_RGB, BO_Add, BF_One, BF_One, BO_Add, BF_Zero, BF_One>::GetRHI());
 
-						RasterizeToVolumeTexture(VolumeBounds);
+						RasterizeToVolumeTexture(RHICmdList, VolumeBounds);
 					}
 				}
 			}
 
-			RHICopyToResolveTarget(RT0->GetRenderTargetItem().TargetableTexture, RT0->GetRenderTargetItem().ShaderResourceTexture, true, FResolveParams());
-			RHICopyToResolveTarget(RT1->GetRenderTargetItem().TargetableTexture, RT1->GetRenderTargetItem().ShaderResourceTexture, true, FResolveParams());
+			RHICmdList.CopyToResolveTarget(RT0->GetRenderTargetItem().TargetableTexture, RT0->GetRenderTargetItem().ShaderResourceTexture, true, FResolveParams());
+			RHICmdList.CopyToResolveTarget(RT1->GetRenderTargetItem().TargetableTexture, RT1->GetRenderTargetItem().ShaderResourceTexture, true, FResolveParams());
 		}
 	}
 }
 
 FGlobalBoundShaderState FilterBoundShaderState;
 
-void FDeferredShadingSceneRenderer::FilterTranslucentVolumeLighting()
+void FDeferredShadingSceneRenderer::FilterTranslucentVolumeLighting(FRHICommandList& RHICmdList)
 {
 	if (GUseTranslucentLightingVolumes && GSupportsVolumeTextureRendering)
 	{
 		if (GUseTranslucencyVolumeBlur)
 		{
-			//@todo-rco: RHIPacketList
-			FRHICommandList& RHICmdList = FRHICommandList::GetNullRef();
-
 			SCOPED_DRAW_EVENT(FilterTranslucentVolume, DEC_SCENE_ITEMS);
 
-			RHISetRasterizerState(TStaticRasterizerState<FM_Solid,CM_None>::GetRHI());
-			RHISetDepthStencilState(TStaticDepthStencilState<false,CF_Always>::GetRHI());
-			RHISetBlendState(TStaticBlendState<>::GetRHI());
+			RHICmdList.SetRasterizerState(TStaticRasterizerState<FM_Solid, CM_None>::GetRHI());
+			RHICmdList.SetDepthStencilState(TStaticDepthStencilState<false, CF_Always>::GetRHI());
+			RHICmdList.SetBlendState(TStaticBlendState<>::GetRHI());
 
 			// Filter each cascade
 			for (int32 VolumeCascadeIndex = 0; VolumeCascadeIndex < TVC_MAX; VolumeCascadeIndex++)
@@ -1909,14 +1896,14 @@ void FDeferredShadingSceneRenderer::FilterTranslucentVolumeLighting()
 				const IPooledRenderTarget* RT0 = GSceneRenderTargets.GetTranslucencyVolumeAmbient((ETranslucencyVolumeCascade)VolumeCascadeIndex);
 				const IPooledRenderTarget* RT1 = GSceneRenderTargets.GetTranslucencyVolumeDirectional((ETranslucencyVolumeCascade)VolumeCascadeIndex);
 
-				GRenderTargetPool.VisualizeTexture.SetCheckPoint(RT0);
-				GRenderTargetPool.VisualizeTexture.SetCheckPoint(RT1);
+				GRenderTargetPool.VisualizeTexture.SetCheckPoint(RHICmdList, RT0);
+				GRenderTargetPool.VisualizeTexture.SetCheckPoint(RHICmdList, RT1);
 
 				FTextureRHIParamRef RenderTargets[2];
 				RenderTargets[0] = RT0->GetRenderTargetItem().TargetableTexture;
 				RenderTargets[1] = RT1->GetRenderTargetItem().TargetableTexture;
 
-				RHISetRenderTargets(ARRAY_COUNT(RenderTargets), RenderTargets, FTextureRHIRef(), 0, NULL);
+				SetRenderTargets(RHICmdList, ARRAY_COUNT(RenderTargets), RenderTargets, FTextureRHIRef(), 0, NULL);
 
 				//@todo - support multiple views
 				const FViewInfo& View = Views[0];
@@ -1925,16 +1912,16 @@ void FDeferredShadingSceneRenderer::FilterTranslucentVolumeLighting()
 				TShaderMapRef<FWriteToSliceVS> VertexShader(GetGlobalShaderMap());
 				TShaderMapRef<FWriteToSliceGS> GeometryShader(GetGlobalShaderMap());
 				TShaderMapRef<FFilterTranslucentVolumePS> PixelShader(GetGlobalShaderMap());
-				SetGlobalBoundShaderState(FilterBoundShaderState, GScreenVertexDeclaration.VertexDeclarationRHI, *VertexShader, *PixelShader, *GeometryShader);
+				SetGlobalBoundShaderState(RHICmdList, FilterBoundShaderState, GScreenVertexDeclaration.VertexDeclarationRHI, *VertexShader, *PixelShader, *GeometryShader);
 
 				VertexShader->SetParameters(RHICmdList, VolumeBounds, GTranslucencyLightingVolumeDim);
 				GeometryShader->SetParameters(RHICmdList, VolumeBounds);
 				PixelShader->SetParameters(RHICmdList, View, VolumeCascadeIndex);
 
-				RasterizeToVolumeTexture(VolumeBounds);
+				RasterizeToVolumeTexture(RHICmdList, VolumeBounds);
 
-				RHICopyToResolveTarget(RT0->GetRenderTargetItem().TargetableTexture, RT0->GetRenderTargetItem().ShaderResourceTexture, true, FResolveParams());
-				RHICopyToResolveTarget(RT1->GetRenderTargetItem().TargetableTexture, RT1->GetRenderTargetItem().ShaderResourceTexture, true, FResolveParams());
+				RHICmdList.CopyToResolveTarget(RT0->GetRenderTargetItem().TargetableTexture, RT0->GetRenderTargetItem().ShaderResourceTexture, true, FResolveParams());
+				RHICmdList.CopyToResolveTarget(RT1->GetRenderTargetItem().TargetableTexture, RT1->GetRenderTargetItem().ShaderResourceTexture, true, FResolveParams());
 			}
 		}
 	}
diff --git a/Engine/Source/Runtime/Renderer/Private/TranslucentRendering.cpp b/Engine/Source/Runtime/Renderer/Private/TranslucentRendering.cpp
--- a/Engine/Source/Runtime/Renderer/Private/TranslucentRendering.cpp
+++ b/Engine/Source/Runtime/Renderer/Private/TranslucentRendering.cpp
@@ -168,6 +168,9 @@ FGlobalBoundShaderState CopySceneColorBoundShaderState;
 
 void FDeferredShadingSceneRenderer::CopySceneColor(const FViewInfo& View, const FPrimitiveSceneInfo* PrimitiveSceneInfo)
 {
+	//@todo-rco: RHIPacketList
+	FRHICommandList& RHICmdList = FRHICommandList::GetNullRef();
+
 	SCOPED_DRAW_EVENTF(EventCopy, DEC_SCENE_ITEMS, TEXT("CopySceneColor for %s %s"), *PrimitiveSceneInfo->Proxy->GetOwnerName().ToString(), *PrimitiveSceneInfo->Proxy->GetResourceName().ToString());
 	RHISetRasterizerState(TStaticRasterizerState<FM_Solid,CM_None>::GetRHI());
 	RHISetDepthStencilState(TStaticDepthStencilState<false,CF_Always>::GetRHI());
@@ -180,12 +183,13 @@ void FDeferredShadingSceneRenderer::CopySceneColor(const FViewInfo& View, const
 
 	TShaderMapRef<FScreenVS> ScreenVertexShader(GetGlobalShaderMap());
 	TShaderMapRef<FCopySceneColorPS> PixelShader(GetGlobalShaderMap());
-	SetGlobalBoundShaderState(CopySceneColorBoundShaderState, GFilterVertexDeclaration.VertexDeclarationRHI, *ScreenVertexShader, *PixelShader);
+	SetGlobalBoundShaderState(RHICmdList, CopySceneColorBoundShaderState, GFilterVertexDeclaration.VertexDeclarationRHI, *ScreenVertexShader, *PixelShader);
 
 	/// ?
-	PixelShader->SetParameters(FRHICommandList::GetNullRef(), View);
+	PixelShader->SetParameters(RHICmdList, View);
 
-	DrawRectangle( 
+	DrawRectangle(
+		RHICmdList,
 		0, 0, 
 		View.ViewRect.Width(), View.ViewRect.Height(),
 		View.ViewRect.Min.X, View.ViewRect.Min.Y, 
@@ -616,6 +620,9 @@ void FDeferredShadingSceneRenderer::RenderTranslucency()
 {
 	if (ShouldRenderTranslucency())
 	{
+		//@todo-rco: RHIPacketList
+		FRHICommandList& RHICmdList = FRHICommandList::GetNullRef();
+
 		SCOPED_DRAW_EVENT(Translucency, DEC_SCENE_ITEMS);
 
 		for (int32 ViewIndex = 0; ViewIndex < Views.Num(); ViewIndex++)
@@ -641,7 +648,7 @@ void FDeferredShadingSceneRenderer::RenderTranslucency()
 
 				if (LightPropagationVolume)
 				{
-					LightPropagationVolume->Visualise(View);
+					LightPropagationVolume->Visualise(RHICmdList, View);
 				}
 			}
 			
diff --git a/Engine/Source/Runtime/Renderer/Private/TranslucentRendering.h b/Engine/Source/Runtime/Renderer/Private/TranslucentRendering.h
--- a/Engine/Source/Runtime/Renderer/Private/TranslucentRendering.h
+++ b/Engine/Source/Runtime/Renderer/Private/TranslucentRendering.h
@@ -230,4 +230,4 @@ private:
 	FShaderParameter MinZ;
 };
 
-extern void RasterizeToVolumeTexture(FVolumeBounds VolumeBounds);
\ No newline at end of file
+extern void RasterizeToVolumeTexture(FRHICommandList& RHICmdList, FVolumeBounds VolumeBounds);
\ No newline at end of file
diff --git a/Engine/Source/Runtime/Renderer/Private/VelocityRendering.cpp b/Engine/Source/Runtime/Renderer/Private/VelocityRendering.cpp
--- a/Engine/Source/Runtime/Renderer/Private/VelocityRendering.cpp
+++ b/Engine/Source/Runtime/Renderer/Private/VelocityRendering.cpp
@@ -527,5 +527,5 @@ void FDeferredShadingSceneRenderer::RenderVelocities(const FViewInfo& View, TRef
 	GPrevPerBoneMotionBlur.UnlockData(bLastFrame);
 
 	// to be able to observe results with VisualizeTexture
-	GRenderTargetPool.VisualizeTexture.SetCheckPoint(VelocityRT);
+	GRenderTargetPool.VisualizeTexture.SetCheckPoint(RHICmdList, VelocityRT);
 }
diff --git a/Engine/Source/Runtime/SlateRHIRenderer/Private/SlateRHIRenderer.cpp b/Engine/Source/Runtime/SlateRHIRenderer/Private/SlateRHIRenderer.cpp
--- a/Engine/Source/Runtime/SlateRHIRenderer/Private/SlateRHIRenderer.cpp
+++ b/Engine/Source/Runtime/SlateRHIRenderer/Private/SlateRHIRenderer.cpp
@@ -756,7 +756,11 @@ void FSlateRHIRenderer::CopyWindowsToDrawBuffer(const TArray<FString>& KeypressB
 				DrawWindowToBuffer,
 				FDrawWindowToBufferContext,Context,DrawWindowToBufferContext,
 			{
+				//@todo-rco: RHIPacketList
+				FRHICommandList& RHICmdList = FRHICommandList::GetNullRef();
+
 				DrawNormalizedScreenQuad(
+					RHICmdList,
 					Context.WindowRect.Min.X, Context.WindowRect.Min.Y,
 					0, 0,
 					Context.WindowRect.Width(), Context.WindowRect.Height(),
diff --git a/Engine/Source/Runtime/SlateRHIRenderer/Private/SlateRHIRenderingPolicy.cpp b/Engine/Source/Runtime/SlateRHIRenderer/Private/SlateRHIRenderingPolicy.cpp
--- a/Engine/Source/Runtime/SlateRHIRenderer/Private/SlateRHIRenderingPolicy.cpp
+++ b/Engine/Source/Runtime/SlateRHIRenderer/Private/SlateRHIRenderingPolicy.cpp
@@ -299,7 +299,7 @@ void FSlateRHIRenderingPolicy::DrawElements( const FIntPoint& InViewportSize, FS
 	FSlateElementIndexBuffer& IndexBuffer = IndexBuffers[CurrentBufferIndex];
 
 	// Set the vertex buffer stream
-	RHISetStreamSource(0, VertexBuffer.VertexBufferRHI, sizeof(FSlateVertex), 0);
+	RHICmdList.SetStreamSource(0, VertexBuffer.VertexBufferRHI, sizeof(FSlateVertex), 0);
 
 	TShaderMapRef<FSlateElementVS> VertexShader(GetGlobalShaderMap());
 
@@ -316,7 +316,7 @@ void FSlateRHIRenderingPolicy::DrawElements( const FIntPoint& InViewportSize, FS
 		if( bSetupBatchRenderingData )
 		{
 			// Set the vertex buffer stream
-			RHISetStreamSource(0, VertexBuffer.VertexBufferRHI, sizeof(FSlateVertex), 0);
+			RHICmdList.SetStreamSource(0, VertexBuffer.VertexBufferRHI, sizeof(FSlateVertex), 0);
 
 			bSetupBatchRenderingData = false;
 		}
@@ -338,32 +338,32 @@ void FSlateRHIRenderingPolicy::DrawElements( const FIntPoint& InViewportSize, FS
 
 			if( bDrawDisabled )
 			{
-				SetGlobalBoundShaderState( DisabledShaderStates[ShaderType][bUseTextureAlpha], GSlateVertexDeclaration.VertexDeclarationRHI, *VertexShader, PixelShader );
+				SetGlobalBoundShaderState(RHICmdList, DisabledShaderStates[ShaderType][bUseTextureAlpha], GSlateVertexDeclaration.VertexDeclarationRHI, *VertexShader, PixelShader);
 			}
 			else
 			{
-				SetGlobalBoundShaderState( NormalShaderStates[ShaderType][bUseTextureAlpha], GSlateVertexDeclaration.VertexDeclarationRHI, *VertexShader, PixelShader );
+				SetGlobalBoundShaderState(RHICmdList, NormalShaderStates[ShaderType][bUseTextureAlpha], GSlateVertexDeclaration.VertexDeclarationRHI, *VertexShader, PixelShader);
 			}
 
 			VertexShader->SetViewProjection(RHICmdList, ViewProjectionMatrix );
 			VertexShader->SetShaderParameters(RHICmdList, ShaderParams.VertexParams );
 
-			RHISetBlendState(
+			RHICmdList.SetBlendState(
 				(RenderBatch.DrawFlags & ESlateBatchDrawFlag::NoBlending)
 				? TStaticBlendState<>::GetRHI()
 				: TStaticBlendState<CW_RGBA,BO_Add,BF_SourceAlpha,BF_InverseSourceAlpha,BO_Add,BF_Zero,BF_One>::GetRHI()
 				);
 
 			// Disable stencil testing by default
-			RHISetDepthStencilState( DSOff );
+			RHICmdList.SetDepthStencilState(DSOff);
 
 			if( DrawFlags & ESlateBatchDrawFlag::Wireframe )
 			{
-				RHISetRasterizerState(TStaticRasterizerState<FM_Wireframe,CM_None,true>::GetRHI());
+				RHICmdList.SetRasterizerState(TStaticRasterizerState<FM_Wireframe, CM_None, true>::GetRHI());
 			}
 			else
 			{
-				RHISetRasterizerState(TStaticRasterizerState<FM_Solid,CM_None,true>::GetRHI());
+				RHICmdList.SetRasterizerState(TStaticRasterizerState<FM_Solid, CM_None, true>::GetRHI());
 			}
 
 
@@ -386,7 +386,7 @@ void FSlateRHIRenderingPolicy::DrawElements( const FIntPoint& InViewportSize, FS
 					, SO_SaturatedIncrement
 					, SO_SaturatedIncrement>::GetRHI();
 
-				RHISetDepthStencilState( DSLineSegment, 0x01 );
+				RHICmdList.SetDepthStencilState( DSLineSegment, 0x01 );
 #endif
 
 				// The lookup table used for splines should clamp when sampling otherwise the results are wrong
@@ -428,7 +428,7 @@ void FSlateRHIRenderingPolicy::DrawElements( const FIntPoint& InViewportSize, FS
 
 			uint32 PrimitiveCount = RenderBatch.DrawPrimitiveType == ESlateDrawPrimitive::LineList ? RenderBatch.NumIndices / 2 : RenderBatch.NumIndices / 3; 
 
-			RHIDrawIndexedPrimitive( IndexBuffer.IndexBufferRHI, GetRHIPrimitiveType( RenderBatch.DrawPrimitiveType ), RenderBatch.VertexOffset, 0, RenderBatch.NumVertices, RenderBatch.IndexOffset, PrimitiveCount, 1 ); 
+			RHICmdList.DrawIndexedPrimitive(IndexBuffer.IndexBufferRHI, GetRHIPrimitiveType(RenderBatch.DrawPrimitiveType), RenderBatch.VertexOffset, 0, RenderBatch.NumVertices, RenderBatch.IndexOffset, PrimitiveCount, 1);
 		}
 		else if( ShaderResource && ShaderResource->GetType() == ESlateShaderResource::Material )
 		{
diff --git a/Engine/Source/Runtime/Windows/D3D11RHI/Private/D3D11Commands.cpp b/Engine/Source/Runtime/Windows/D3D11RHI/Private/D3D11Commands.cpp
--- a/Engine/Source/Runtime/Windows/D3D11RHI/Private/D3D11Commands.cpp
+++ b/Engine/Source/Runtime/Windows/D3D11RHI/Private/D3D11Commands.cpp
@@ -1977,8 +1977,10 @@ void FD3D11DynamicRHI::RHIClearMRT(bool bClearColor,int32 NumClearColors,const F
 			TShaderMapRef<TOneColorPixelShaderMRT<8> > MRTPixelShader(GetGlobalShaderMap());
 			PixelShader = *MRTPixelShader;
 		}
+		//@todo-rco: RHIPacketList
+		FRHICommandList& RHICmdList = FRHICommandList::GetNullRef();
 
-		SetGlobalBoundShaderState(GD3D11ClearMRTBoundShaderState[FMath::Max(BoundRenderTargets.GetNumActiveTargets() - 1, 0)], GD3D11Vector4VertexDeclaration.VertexDeclarationRHI, *VertexShader, PixelShader);
+		SetGlobalBoundShaderState(RHICmdList, GD3D11ClearMRTBoundShaderState[FMath::Max(BoundRenderTargets.GetNumActiveTargets() - 1, 0)], GD3D11Vector4VertexDeclaration.VertexDeclarationRHI, *VertexShader, PixelShader);
 		FLinearColor ShaderClearColors[MaxSimultaneousRenderTargets];
 		FMemory::MemZero(ShaderClearColors);
 
@@ -1987,8 +1989,6 @@ void FD3D11DynamicRHI::RHIClearMRT(bool bClearColor,int32 NumClearColors,const F
 			ShaderClearColors[i] = ClearColorArray[i];
 		}
 
-		//@todo-rco: RHIPacketList
-		FRHICommandList& RHICmdList = FRHICommandList::GetNullRef();
 		SetShaderValueArray(RHICmdList, PixelShader->GetPixelShader(),PixelShader->ColorParameter,ShaderClearColors,NumClearColors);
 		
 		{
@@ -2024,7 +2024,7 @@ void FD3D11DynamicRHI::RHIClearMRT(bool bClearColor,int32 NumClearColors,const F
 				Vertices[8] = OuterVertices[0];
 				Vertices[9] = InnerVertices[0];
 
-				RHIDrawPrimitiveUP(PT_TriangleStrip, 8, Vertices, sizeof(Vertices[0]) );
+				DrawPrimitiveUP(RHICmdList, PT_TriangleStrip, 8, Vertices, sizeof(Vertices[0]));
 			}
 			else
 			{
@@ -2034,7 +2034,7 @@ void FD3D11DynamicRHI::RHIClearMRT(bool bClearColor,int32 NumClearColors,const F
 				Vertices[1].Set(  1.0f,  1.0f, Depth, 1.0f );
 				Vertices[2].Set( -1.0f, -1.0f, Depth, 1.0f );
 				Vertices[3].Set(  1.0f, -1.0f, Depth, 1.0f );
-				RHIDrawPrimitiveUP(PT_TriangleStrip, 2, Vertices, sizeof(Vertices[0]) );
+				DrawPrimitiveUP(RHICmdList, PT_TriangleStrip, 2, Vertices, sizeof(Vertices[0]));
 			}
 		}
 
diff --git a/Engine/Source/Runtime/Windows/D3D11RHI/Private/D3D11RHI.cpp b/Engine/Source/Runtime/Windows/D3D11RHI/Private/D3D11RHI.cpp
--- a/Engine/Source/Runtime/Windows/D3D11RHI/Private/D3D11RHI.cpp
+++ b/Engine/Source/Runtime/Windows/D3D11RHI/Private/D3D11RHI.cpp
@@ -182,6 +182,9 @@ FGlobalBoundShaderState LongGPUTaskBoundShaderState;
 
 void FD3D11DynamicRHI::IssueLongGPUTask()
 {
+	//@todo-rco: RHIPacketList
+	FRHICommandList& RHICmdList = FRHICommandList::GetNullRef();
+
 	if (GRHIFeatureLevel >= ERHIFeatureLevel::SM3)
 	{
 		int32 LargestViewportIndex = INDEX_NONE;
@@ -202,15 +205,15 @@ void FD3D11DynamicRHI::IssueLongGPUTask()
 		{
 			FD3D11Viewport* Viewport = Viewports[LargestViewportIndex];
 
-			RHISetRenderTarget(Viewport->GetBackBuffer(), FTextureRHIRef());
-			RHISetBlendState(TStaticBlendState<CW_RGBA, BO_Add, BF_One, BF_One>::GetRHI(), FLinearColor::Black);
-			RHISetDepthStencilState(TStaticDepthStencilState<false,CF_Always>::GetRHI(), 0);
-			RHISetRasterizerState(TStaticRasterizerState<FM_Solid,CM_None>::GetRHI());
+			SetRenderTarget(RHICmdList, Viewport->GetBackBuffer(), FTextureRHIRef());
+			RHICmdList.SetBlendState(TStaticBlendState<CW_RGBA, BO_Add, BF_One, BF_One>::GetRHI(), FLinearColor::Black);
+			RHICmdList.SetDepthStencilState(TStaticDepthStencilState<false, CF_Always>::GetRHI(), 0);
+			RHICmdList.SetRasterizerState(TStaticRasterizerState<FM_Solid, CM_None>::GetRHI());
 
 			TShaderMapRef<TOneColorVS<true> > VertexShader(GetGlobalShaderMap());
 			TShaderMapRef<FLongGPUTaskPS> PixelShader(GetGlobalShaderMap());
 
-			SetGlobalBoundShaderState(LongGPUTaskBoundShaderState, GD3D11Vector4VertexDeclaration.VertexDeclarationRHI, *VertexShader, *PixelShader);
+			SetGlobalBoundShaderState(RHICmdList, LongGPUTaskBoundShaderState, GD3D11Vector4VertexDeclaration.VertexDeclarationRHI, *VertexShader, *PixelShader);
 
 			// Draw a fullscreen quad
 			FVector4 Vertices[4];
@@ -218,7 +221,7 @@ void FD3D11DynamicRHI::IssueLongGPUTask()
 			Vertices[1].Set(  1.0f,  1.0f, 0, 1.0f );
 			Vertices[2].Set( -1.0f, -1.0f, 0, 1.0f );
 			Vertices[3].Set(  1.0f, -1.0f, 0, 1.0f );
-			RHIDrawPrimitiveUP(PT_TriangleStrip, 2, Vertices, sizeof(Vertices[0]) );
+			DrawPrimitiveUP(RHICmdList, PT_TriangleStrip, 2, Vertices, sizeof(Vertices[0]));
 		}
 	}
 }
diff --git a/Engine/Source/Runtime/Windows/D3D11RHI/Private/D3D11RenderTarget.cpp b/Engine/Source/Runtime/Windows/D3D11RHI/Private/D3D11RenderTarget.cpp
--- a/Engine/Source/Runtime/Windows/D3D11RHI/Private/D3D11RenderTarget.cpp
+++ b/Engine/Source/Runtime/Windows/D3D11RHI/Private/D3D11RenderTarget.cpp
@@ -250,7 +250,7 @@ void FD3D11DynamicRHI::ResolveTextureUsingShader(
 	// Set the vertex and pixel shader
 	TShaderMapRef<FResolveVS> ResolveVertexShader(GetGlobalShaderMap());
 	TShaderMapRef<TPixelShader> ResolvePixelShader(GetGlobalShaderMap());
-	SetGlobalBoundShaderState(ResolveBoundShaderState, GScreenVertexDeclaration.VertexDeclarationRHI, *ResolveVertexShader, *ResolvePixelShader);
+	SetGlobalBoundShaderState(RHICmdList, ResolveBoundShaderState, GScreenVertexDeclaration.VertexDeclarationRHI, *ResolveVertexShader, *ResolvePixelShader);
 
 	ResolvePixelShader->SetParameters(RHICmdList, Direct3DDeviceContext,PixelShaderParameter);
 
@@ -281,7 +281,7 @@ void FD3D11DynamicRHI::ResolveTextureUsingShader(
 	Vertices[3].UV.X       = MinU;
 	Vertices[3].UV.Y       = MaxV;
 
-	RHIDrawPrimitiveUP(PT_TriangleStrip,2,Vertices,sizeof(Vertices[0]));
+	DrawPrimitiveUP(RHICmdList, PT_TriangleStrip, 2, Vertices, sizeof(Vertices[0]));
 
 	if (SourceTexture)
 	{
