UE4 - don't need mutable on a reference and clang does not like it

diff --git a/Engine/Source/Runtime/Renderer/Private/PostProcess/RenderingCompositionGraph.h b/Engine/Source/Runtime/Renderer/Private/PostProcess/RenderingCompositionGraph.h
--- a/Engine/Source/Runtime/Renderer/Private/PostProcess/RenderingCompositionGraph.h
+++ b/Engine/Source/Runtime/Renderer/Private/PostProcess/RenderingCompositionGraph.h
@@ -179,7 +179,7 @@ struct FRenderingCompositePassContext
 
 	FRenderingCompositionGraph Graph;
 
-	mutable FRHICommandList& RHICmdList;
+	FRHICommandList& RHICmdList;
 
 private:
 	// cached state to map between ScreenPos and pixels
