fix for invisible cards in Black Jack on iPad 2
TTP337518
#ue4
#ios

diff --git a/Engine/Shaders/BasePassForForwardShadingPixelShader.usf b/Engine/Shaders/BasePassForForwardShadingPixelShader.usf
--- a/Engine/Shaders/BasePassForForwardShadingPixelShader.usf
+++ b/Engine/Shaders/BasePassForForwardShadingPixelShader.usf
@@ -227,7 +227,7 @@ void Main(
 	#if !MATERIALBLENDING_MODULATE
 		// The exposure scale is just a scalar but needs to be a float4 to workaround a driver bug on IOS.
 		// After 4.2 we can put the workaround in the cross compiler.
-		OutColor.rgb *= View.ExposureScale.xyz;
+		OutColor.rgba *= View.ExposureScale.xyzw;
 	#endif
 
 	#if OUTPUT_GAMMA_SPACE
diff --git a/Engine/Source/Runtime/Renderer/Private/SceneRendering.cpp b/Engine/Source/Runtime/Renderer/Private/SceneRendering.cpp
--- a/Engine/Source/Runtime/Renderer/Private/SceneRendering.cpp
+++ b/Engine/Source/Runtime/Renderer/Private/SceneRendering.cpp
@@ -415,7 +415,7 @@ TUniformBufferRef<FViewUniformShaderParameters> FViewInfo::CreateUniformBuffer(
 	// The exposure scale is just a scalar but needs to be a float4 to workaround a driver bug on IOS.
 	// After 4.2 we can put the workaround in the cross compiler.
 	float ExposureScale = FRCPassPostProcessEyeAdaptation::ComputeExposureScaleValue( *this );
-	ViewUniformShaderParameters.ExposureScale = FVector4(ExposureScale,ExposureScale,ExposureScale,ExposureScale);
+	ViewUniformShaderParameters.ExposureScale = FVector4(ExposureScale,ExposureScale,ExposureScale,1.0f);
 	ViewUniformShaderParameters.DepthOfFieldFocalDistance = FinalPostProcessSettings.DepthOfFieldFocalDistance;
 	ViewUniformShaderParameters.DepthOfFieldFocalRegion = FinalPostProcessSettings.DepthOfFieldFocalRegion;
 	ViewUniformShaderParameters.DepthOfFieldNearTransitionRegion = FinalPostProcessSettings.DepthOfFieldNearTransitionRegion;
