* fixed crash when using affectsindirectlight checkbox and not having LPV enabled

diff --git a/Engine/Source/Runtime/Renderer/Private/LightRendering.cpp b/Engine/Source/Runtime/Renderer/Private/LightRendering.cpp
--- a/Engine/Source/Runtime/Renderer/Private/LightRendering.cpp
+++ b/Engine/Source/Runtime/Renderer/Private/LightRendering.cpp
@@ -561,15 +561,18 @@ void FDeferredShadingSceneRenderer::RenderLights(FRHICommandList& RHICmdList)
 					// Inject the light directly into all relevant LPVs
 					for (int32 ViewIndex = 0; ViewIndex < Views.Num(); ViewIndex++)
 					{
-						if (LightSceneInfo->ShouldRenderLight(Views[ViewIndex]))
-						{
-							FSceneViewState* ViewState = (FSceneViewState*)Views[ViewIndex].State;
+						FSceneViewState* ViewState = (FSceneViewState*)Views[ViewIndex].State;
 
-							FLightPropagationVolume* Lpv = ViewState->GetLightPropagationVolume();
-							if ( Lpv && LightSceneInfo->Proxy )
+						if (ViewState)
+						{
+							if (LightSceneInfo->ShouldRenderLight(Views[ViewIndex]))
 							{
-								//@todo-rco: RHIPacketList
-								Lpv->InjectLightDirect(FRHICommandList::GetNullRef(), *LightSceneInfo->Proxy);
+								FLightPropagationVolume* Lpv = ViewState->GetLightPropagationVolume();
+								if (Lpv && LightSceneInfo->Proxy)
+								{
+									//@todo-rco: RHIPacketList
+									Lpv->InjectLightDirect(FRHICommandList::GetNullRef(), *LightSceneInfo->Proxy);
+								}
 							}
 						}
 					}					
