Moved color fringe into the tonemapper removes a round trip to memory just for this pass.

diff --git a/Engine/Shaders/PostProcessTonemap.usf b/Engine/Shaders/PostProcessTonemap.usf
--- a/Engine/Shaders/PostProcessTonemap.usf
+++ b/Engine/Shaders/PostProcessTonemap.usf
@@ -99,6 +99,21 @@ float PseudoRandomTex(float2 xy)
 
 float SwitchVerticalAxis;
 
+
+
+// can be optimized
+float2 ScreenPosToUV(float2 ScreenPos)
+{
+	return (ScreenPos * ScreenPosToPixel.xy + ScreenPosToPixel.zw) * PostprocessInput0Size.zw;
+}
+
+float2 UVToScreenPos(float2 UV)
+{
+	return (UV * PostprocessInput0Size.xy - ScreenPosToPixel.zw) / ScreenPosToPixel.xy;
+}
+
+float4 FringeUVParams;
+
 // vertex shader entry point
 void MainVS(
 	in float4 InPosition : ATTRIBUTE0,
@@ -106,6 +121,7 @@ void MainVS(
 	out float4 OutTexCoord : TEXCOORD0,
 	out float3 OutExposureScaleVignette : TEXCOORD1,
 	out float4 OutGrainUV : TEXCOORD2,
+	out float4 OutFringe : TEXCOORD3,
 	out float4 OutPosition : SV_POSITION
 	)
 {
@@ -126,6 +142,11 @@ void MainVS(
 	// Grain
 	OutGrainUV.xy = OutTexCoord.xy + PostprocessInput0Size.zw * float2(-0.5,0.5);
 	OutGrainUV.zw = OutTexCoord.xy + GrainRandomFull.xy;
+
+	// Fringe
+	float2 ScreenPos = UVToScreenPos(OutTexCoord.xy);
+	OutFringe.xy = ScreenPosToUV(ScreenPos * FringeUVParams.r);
+	OutFringe.zw = ScreenPosToUV(ScreenPos * FringeUVParams.g);
 }
 
 // Function graphing
@@ -175,13 +196,35 @@ void MainPS(
 	in float4 UVAndScreenPos : TEXCOORD0,
 	in float3 InExposureScaleVignette : TEXCOORD1,
 	in float4 GrainUV : TEXCOORD2,
+	in float4 FringeUV : TEXCOORD3,
 	out float4 OutColor : SV_Target0
 	)
 {
 	float2 UV = UVAndScreenPos.xy;
 	float2 ScreenSpacePos = UVAndScreenPos.zw;
 
-	half4 SceneColor = Texture2DSample(PostprocessInput0, PostprocessInput0Sampler, UV);
+	#if USE_GRAIN_JITTER || USE_GRAIN_INTENSITY || USE_GRAIN_QUANTIZATION
+		half Grain = GrainFromUV(GrainUV.zw);
+	#endif
+
+	float2 SceneUV = UV.xy;
+	#if USE_GRAIN_JITTER
+		SceneUV = lerp(UV.xy, GrainUV.xy, (1.0 - Grain*Grain) * GrainScaleBiasJitter.z);
+	#endif
+
+	#if USE_COLOR_FRINGE
+		float2 SceneUVJitter = float2(0.0, 0.0);
+		#if USE_GRAIN_JITTER
+			SceneUVJitter = SceneUV.xy - UV.xy;
+		#endif
+		half4 SceneColor = Texture2DSample(PostprocessInput0, PostprocessInput0Sampler, SceneUV);
+		half SceneColorG = Texture2DSample(PostprocessInput0, PostprocessInput0Sampler, FringeUV.xy + SceneUVJitter.xy).g;
+		half SceneColorB = Texture2DSample(PostprocessInput0, PostprocessInput0Sampler, FringeUV.zw + SceneUVJitter.xy).b;
+		SceneColor.g = SceneColorG;
+		SceneColor.b = SceneColorB;
+	#else
+		half4 SceneColor = Texture2DSample(PostprocessInput0, PostprocessInput0Sampler, SceneUV);
+	#endif
 
 	#if USE_GAMMA_ONLY
 
@@ -189,15 +232,6 @@ void MainPS(
 
 	#else 
 
-		#if USE_GRAIN_JITTER || USE_GRAIN_INTENSITY || USE_GRAIN_QUANTIZATION
-			half Grain = GrainFromUV(GrainUV.zw);
-		#endif
-
-		#if USE_GRAIN_JITTER
-			half3 GrainColor = Texture2DSample(PostprocessInput0, PostprocessInput0Sampler, GrainUV.xy).rgb;
-			SceneColor.rgb = lerp(SceneColor.rgb, GrainColor.rgb, (1.0 - Grain*Grain) * GrainScaleBiasJitter.z);
-		#endif
-
 		half3 LinearColor = SceneColor.rgb * ColorScale0.rgb;
 
 		#if USE_BLOOM
diff --git a/Engine/Source/Runtime/Renderer/Private/PostProcess/PostProcessTonemap.cpp b/Engine/Source/Runtime/Renderer/Private/PostProcess/PostProcessTonemap.cpp
--- a/Engine/Source/Runtime/Renderer/Private/PostProcess/PostProcessTonemap.cpp
+++ b/Engine/Source/Runtime/Renderer/Private/PostProcess/PostProcessTonemap.cpp
@@ -34,6 +34,8 @@ typedef enum {
 	TonemapperVignetteColor     = (1<<10),
 	TonemapperLightShafts       = (1<<11),
 	TonemapperMosaic            = (1<<12),
+	TonemapperColorFringe       = (1<<13),
+
 } TonemapperOption;
 
 // Tonemapper option cost (0 = no cost, 255 = max cost).
@@ -53,6 +55,7 @@ static uint8 TonemapperCostTab[] = {
 	1, //TonemapperVignetteColor
 	1, //TonemapperLightShafts
 	1, //TonemapperMosaic
+	1, //TonemapperColorFringe
 };
 
 // Edit the following to add and remove configurations.
@@ -71,6 +74,7 @@ static uint32 TonemapperConfBitmaskPC[11] = {
 	TonemapperGrainQuantization +
 	TonemapperVignette +
 	TonemapperVignetteColor +
+	TonemapperColorFringe +
 	0,
 
 	TonemapperBloom +
@@ -79,6 +83,7 @@ static uint32 TonemapperConfBitmaskPC[11] = {
 	TonemapperShadowTint +
 	TonemapperVignette +
 	TonemapperGrainQuantization +
+	TonemapperColorFringe +
 	0,
 
 	TonemapperBloom + 
@@ -109,6 +114,7 @@ static uint32 TonemapperConfBitmaskPC[11] = {
 	TonemapperGrainIntensity +
 	TonemapperVignette +
 	TonemapperVignetteColor +
+	TonemapperColorFringe +
 	0,
 
 	TonemapperBloom +
@@ -116,6 +122,7 @@ static uint32 TonemapperConfBitmaskPC[11] = {
 	TonemapperColorMatrix +
 	TonemapperShadowTint +
 	TonemapperVignette +
+	TonemapperColorFringe +
 	0,
 
 	TonemapperBloom + 
@@ -463,6 +470,7 @@ static uint32 TonemapperGenerateBitmask(const FRenderingCompositePassContext* RE
 	Bitmask += (Settings->GrainIntensity > 0.0f)       ? TonemapperGrainIntensity : 0;
 	Bitmask += (Settings->VignetteIntensity > 0.0f)    ? TonemapperVignette       : 0;
 	Bitmask += (VignetteColor.GetAbsMax() != 0.0f)     ? TonemapperVignetteColor  : 0;
+
 	return Bitmask;
 }
 
@@ -501,6 +509,18 @@ static uint32 TonemapperGenerateBitmaskPC(const FRenderingCompositePassContext*
 		}
 	}
 
+	static const auto CVar = IConsoleManager::Get().FindTConsoleVariableDataInt(TEXT("r.SceneColorFringeQuality")); 
+
+	int32 FringeQuality = CVar->GetValueOnRenderThread();
+	if(Context->View.Family->EngineShowFlags.SceneColorFringe
+		// Removed this from the camera imperfections toggle because this no longer takes an extra pass.
+		// && Context->View.Family->EngineShowFlags.CameraImperfections
+		&& Context->View.FinalPostProcessSettings.SceneFringeIntensity > 0.01f
+		&& FringeQuality > 0)
+	{
+		Bitmask |= TonemapperColorFringe;
+	}
+
 	return Bitmask + TonemapperGenerateBitmaskPost(Context);
 }
 
@@ -752,6 +772,7 @@ class FPostProcessTonemapPS : public FGlobalShader
 		OutEnvironment.SetDefine(TEXT("USE_GRAIN_QUANTIZATION"), TonemapperIsDefined(ConfigBitmask, TonemapperGrainQuantization));
 		OutEnvironment.SetDefine(TEXT("USE_VIGNETTE"),           TonemapperIsDefined(ConfigBitmask, TonemapperVignette));
 		OutEnvironment.SetDefine(TEXT("USE_VIGNETTE_COLOR"),     TonemapperIsDefined(ConfigBitmask, TonemapperVignetteColor));
+		OutEnvironment.SetDefine(TEXT("USE_COLOR_FRINGE"),		 TonemapperIsDefined(ConfigBitmask, TonemapperColorFringe));
 		// @todo Mac OS X: in order to share precompiled shaders between GL 3.3 & GL 4.1 devices we mustn't use volume-texture rendering as it isn't universally supported.
 		OutEnvironment.SetDefine(TEXT("USE_VOLUME_LUT"),		 (IsFeatureLevelSupported(Platform,ERHIFeatureLevel::SM4) && GSupportsVolumeTextureRendering && !PLATFORM_MAC));
 
diff --git a/Engine/Source/Runtime/Renderer/Private/PostProcess/PostProcessTonemap.h b/Engine/Source/Runtime/Renderer/Private/PostProcess/PostProcessTonemap.h
--- a/Engine/Source/Runtime/Renderer/Private/PostProcess/PostProcessTonemap.h
+++ b/Engine/Source/Runtime/Renderer/Private/PostProcess/PostProcessTonemap.h
@@ -90,6 +90,7 @@ public:
 	FPostProcessPassParameters PostprocessParameter;
 	FShaderResourceParameter EyeAdaptation;
 	FShaderParameter GrainRandomFull;
+	FShaderParameter FringeUVParams;
 
 	/** Initialization constructor. */
 	FPostProcessTonemapVS(const ShaderMetaType::CompiledShaderInitializerType& Initializer)
@@ -98,6 +99,7 @@ public:
 		PostprocessParameter.Bind(Initializer.ParameterMap);
 		EyeAdaptation.Bind(Initializer.ParameterMap, TEXT("EyeAdaptation"));
 		GrainRandomFull.Bind(Initializer.ParameterMap, TEXT("GrainRandomFull"));
+		FringeUVParams.Bind(Initializer.ParameterMap, TEXT("FringeUVParams"));
 	}
 
 	void SetVS(const FRenderingCompositePassContext& Context)
@@ -126,13 +128,22 @@ public:
 				SetTextureParameter(Context.RHICmdList, ShaderRHI, EyeAdaptation, GWhiteTexture->TextureRHI);
 			}
 		}
+
+		{
+			// for scene color fringe
+			// from percent to fraction
+			float Offset = Context.View.FinalPostProcessSettings.SceneFringeIntensity * 0.01f;
+			// we only get bigger to not leak in content from outside
+			FVector4 Value(1.0f - Offset * 0.5f, 1.0f - Offset, 0.0f, 0.0f);
+			SetShaderValue(Context.RHICmdList, ShaderRHI, FringeUVParams, Value);
+		}
 	}
 	
 	// FShader interface.
 	virtual bool Serialize(FArchive& Ar)
 	{
 		bool bShaderHasOutdatedParameters = FGlobalShader::Serialize(Ar);
-		Ar << PostprocessParameter << GrainRandomFull << EyeAdaptation;
+		Ar << PostprocessParameter << GrainRandomFull << EyeAdaptation << FringeUVParams;
 		return bShaderHasOutdatedParameters;
 	}
 };
diff --git a/Engine/Source/Runtime/Renderer/Private/PostProcess/PostProcessing.cpp b/Engine/Source/Runtime/Renderer/Private/PostProcess/PostProcessing.cpp
--- a/Engine/Source/Runtime/Renderer/Private/PostProcess/PostProcessing.cpp
+++ b/Engine/Source/Runtime/Renderer/Private/PostProcess/PostProcessing.cpp
@@ -40,7 +40,6 @@
 #include "PostProcessAmbientOcclusion.h"
 #include "ScreenSpaceReflections.h"
 #include "PostProcessTestImage.h"
-#include "PostProcessSceneColorFringe.h"
 #include "PostProcessSubsurface.h"
 #include "PostProcessUIBlur.h"
 #include "HighResScreenshot.h"
@@ -186,15 +185,6 @@ static void AddPostProcessAA(FPostprocessContext& Context)
 	Context.FinalOutput = FRenderingCompositeOutputRef(Node);
 }
 
-static void AddSceneColorFringe(FPostprocessContext& Context)
-{
-	FRenderingCompositePass* Node = Context.Graph.RegisterPass(new(FMemStack::Get()) FRCPassPostProcessSceneColorFringe());
-
-	Node->SetInput(ePId_Input0, FRenderingCompositeOutputRef(Context.FinalOutput));
-
-	Context.FinalOutput = FRenderingCompositeOutputRef(Node);
-}
-
 static FRenderingCompositeOutputRef AddPostProcessEyeAdaptation(FPostprocessContext& Context, FRenderingCompositeOutputRef& Histogram)
 {
 	FRenderingCompositePass* Node = Context.Graph.RegisterPass(new(FMemStack::Get()) FRCPassPostProcessEyeAdaptation());
@@ -1014,20 +1004,6 @@ void FPostProcessing::Process(const FViewInfo& View, TRefCountPtr<IPooledRenderT
 				AddPostProcessAA(Context);
 			}
 			
-			{
-				static const auto CVar = IConsoleManager::Get().FindTConsoleVariableDataInt(TEXT("r.SceneColorFringeQuality")); 
-
-				int32 FringeQuality = CVar->GetValueOnRenderThread();
-
-				if(View.Family->EngineShowFlags.SceneColorFringe
-					&& View.Family->EngineShowFlags.CameraImperfections
-					&& View.FinalPostProcessSettings.SceneFringeIntensity > 0.01f
-					&& FringeQuality > 0)
-				{
-					AddSceneColorFringe(Context);
-				}
-			}
-
 			if(bDepthOfField && Context.View.Family->EngineShowFlags.VisualizeDOF)
 			{
 				FRenderingCompositePass* VisualizeNode = Context.Graph.RegisterPass(new(FMemStack::Get()) FRCPassPostProcessVisualizeDOF());
