Turn on responsive AA everywhere during a camera cut.

diff --git a/Engine/Source/Runtime/Renderer/Private/PostProcess/PostProcessTemporalAA.cpp b/Engine/Source/Runtime/Renderer/Private/PostProcess/PostProcessTemporalAA.cpp
--- a/Engine/Source/Runtime/Renderer/Private/PostProcess/PostProcessTemporalAA.cpp
+++ b/Engine/Source/Runtime/Renderer/Private/PostProcess/PostProcessTemporalAA.cpp
@@ -467,7 +467,53 @@ void FRCPassPostProcessTemporalAA::Process(FRenderingCompositePassContext& Conte
 	uint32 Quality = FMath::Clamp(CVar->GetValueOnRenderThread(), 1, 6);
 	bool bUseFast = Quality == 3;
 
+	if(Context.View.bCameraCut)
+	{
+		// On camera cut this turns on responsive everywhere.
+		
+		// Normal temporal feedback
+		// Draw to pixels where stencil == 0
+		Context.RHICmdList.SetDepthStencilState(TStaticDepthStencilState<false,CF_Always,true,CF_Equal,SO_Keep,SO_Keep,SO_Keep>::GetRHI(), 0);
+
+		TShaderMapRef< FPostProcessTonemapVS >			VertexShader(GetGlobalShaderMap());
+		if (bUseFast)
+		{
+			TShaderMapRef< FPostProcessTemporalAAPS<4,1> >	PixelShader( GetGlobalShaderMap() );
+
+			static FGlobalBoundShaderState BoundShaderState;
+			Context.RHICmdList.CheckIsNull(); // need new approach for "static FGlobalBoundShaderState" for parallel rendering
+
+			SetGlobalBoundShaderState(Context.RHICmdList, BoundShaderState, GFilterVertexDeclaration.VertexDeclarationRHI, *VertexShader, *PixelShader);
+
+			VertexShader->SetVS(Context);
+			PixelShader->SetParameters(Context);
+		}
+		else
+		{
+			TShaderMapRef< FPostProcessTemporalAAPS<1,1> >	PixelShader( GetGlobalShaderMap() );
 
+			static FGlobalBoundShaderState BoundShaderState;
+			Context.RHICmdList.CheckIsNull(); // need new approach for "static FGlobalBoundShaderState" for parallel rendering
+
+			SetGlobalBoundShaderState(Context.RHICmdList, BoundShaderState, GFilterVertexDeclaration.VertexDeclarationRHI, *VertexShader, *PixelShader);
+
+			VertexShader->SetVS(Context);
+			PixelShader->SetParameters(Context);
+		}
+	
+		// Draw a quad mapping scene color to the view's render target
+		DrawRectangle(
+			Context.RHICmdList,
+			0, 0,
+			SrcRect.Width(), SrcRect.Height(),
+			SrcRect.Min.X, SrcRect.Min.Y, 
+			SrcRect.Width(), SrcRect.Height(),
+			SrcRect.Size(),
+			SrcSize,
+			*VertexShader,
+			EDRF_UseTriangleOptimization);
+	}
+	else
 	{	
 		// Normal temporal feedback
 		// Draw to pixels where stencil == 0
