* CPU optimization, CVar was queried many times per frame

diff --git a/Engine/Source/Runtime/Renderer/Private/PostProcess/SceneRenderTargets.cpp b/Engine/Source/Runtime/Renderer/Private/PostProcess/SceneRenderTargets.cpp
--- a/Engine/Source/Runtime/Renderer/Private/PostProcess/SceneRenderTargets.cpp
+++ b/Engine/Source/Runtime/Renderer/Private/PostProcess/SceneRenderTargets.cpp
@@ -1365,7 +1365,7 @@ FIntPoint FSceneRenderTargets::GetPreShadowCacheTextureResolution() const
 	// Higher numbers increase cache hit rate but also memory usage
 	const int32 ExpandFactor = 2;
 
-	auto CVarPreShadowResolutionFactor = IConsoleManager::Get().FindTConsoleVariableDataFloat(TEXT("r.Shadow.PreShadowResolutionFactor"));
+	static auto CVarPreShadowResolutionFactor = IConsoleManager::Get().FindTConsoleVariableDataFloat(TEXT("r.Shadow.PreShadowResolutionFactor"));
 	return FIntPoint(FMath::TruncToInt(ShadowDepthResolution.X * CVarPreShadowResolutionFactor->GetValueOnRenderThread()), FMath::TruncToInt(ShadowDepthResolution.Y * CVarPreShadowResolutionFactor->GetValueOnRenderThread())) * ExpandFactor;
 }
 
diff --git a/Engine/Source/Runtime/Renderer/Private/ShadowRendering.cpp b/Engine/Source/Runtime/Renderer/Private/ShadowRendering.cpp
--- a/Engine/Source/Runtime/Renderer/Private/ShadowRendering.cpp
+++ b/Engine/Source/Runtime/Renderer/Private/ShadowRendering.cpp
@@ -1876,7 +1876,7 @@ void FProjectedShadowInfo::RenderProjection(FRHICommandList& RHICmdList, int32 V
 			}
 			else if (LocalQuality > 2 && !bWholeSceneShadow)
 			{
-				auto CVarPreShadowResolutionFactor = IConsoleManager::Get().FindTConsoleVariableDataFloat(TEXT("r.Shadow.PreShadowResolutionFactor"));
+				static auto CVarPreShadowResolutionFactor = IConsoleManager::Get().FindTConsoleVariableDataFloat(TEXT("r.Shadow.PreShadowResolutionFactor"));
 				const int32 TargetResolution = bPreShadow ? FMath::TruncToInt(512 * CVarPreShadowResolutionFactor->GetValueOnRenderThread()) : 512;
 
 				int32 Reduce = 0;
