Merging using Main_to_ST7

diff --git a/Engine/Shaders/MaterialTemplate.usf b/Engine/Shaders/MaterialTemplate.usf
--- a/Engine/Shaders/MaterialTemplate.usf
+++ b/Engine/Shaders/MaterialTemplate.usf
@@ -769,7 +769,7 @@ float3 GetSpeedTreeVertexOffset(FMaterialVertexParameters Parameters, int Geomet
 			if (WindType > SPEEDTREE_WIND_TYPE_FASTEST && WindType < SPEEDTREE_WIND_TYPE_PALM) 
 			{
 				float4 WindExtra = float4(Parameters.TexCoords[6].x, Parameters.TexCoords[6].y, Parameters.TexCoords[7].x, Parameters.TexCoords[7].y);
-				float LeafWindTrigOffset = FinalPosition.x + FinalPosition.y;
+				float LeafWindTrigOffset = Anchor.x + Anchor.y;
 				FinalPosition = LeafWind(WindExtra.w > 0.0, FinalPosition, Normal, WindExtra.x, float3(0,0,0), WindExtra.y, WindExtra.z, LeafWindTrigOffset, WindType);
 			}
 				
diff --git a/Engine/Source/Runtime/Renderer/Private/RendererScene.cpp b/Engine/Source/Runtime/Renderer/Private/RendererScene.cpp
--- a/Engine/Source/Runtime/Renderer/Private/RendererScene.cpp
+++ b/Engine/Source/Runtime/Renderer/Private/RendererScene.cpp
@@ -1335,8 +1335,6 @@ void FScene::RemoveSpeedTreeWind(class FVertexFactory* VertexFactory, class USta
 			UStaticMesh*, StaticMesh, StaticMesh,
 			FVertexFactory*,VertexFactory,VertexFactory,
 			{
-				Scene->SpeedTreeVertexFactoryMap.Remove(VertexFactory);
-				
 				FSpeedTreeWindComputation** WindComputationRef = Scene->SpeedTreeWindComputationMap.Find(StaticMesh);
 				if (WindComputationRef != NULL)
 				{
@@ -1345,6 +1343,14 @@ void FScene::RemoveSpeedTreeWind(class FVertexFactory* VertexFactory, class USta
 					WindComputation->ReferenceCount--;
 					if (WindComputation->ReferenceCount < 1)
 					{
+						for (auto Iter = Scene->SpeedTreeVertexFactoryMap.CreateIterator(); Iter; ++Iter )
+						{
+							if (Iter.Value() == StaticMesh)
+							{
+								Iter.RemoveCurrent();
+							}
+						}
+
 						Scene->SpeedTreeWindComputationMap.Remove(StaticMesh);
 						WindComputation->UniformBuffer.ReleaseResource();
 						delete WindComputation;
