Correct temporal AA, not this sets all required render state, fixes scene changes. Uses responsive AA for scene changes.

diff --git a/Engine/Source/Runtime/Renderer/Private/PostProcess/PostProcessTemporalAA.cpp b/Engine/Source/Runtime/Renderer/Private/PostProcess/PostProcessTemporalAA.cpp
--- a/Engine/Source/Runtime/Renderer/Private/PostProcess/PostProcessTemporalAA.cpp
+++ b/Engine/Source/Runtime/Renderer/Private/PostProcess/PostProcessTemporalAA.cpp
@@ -467,13 +467,15 @@ void FRCPassPostProcessTemporalAA::Process(FRenderingCompositePassContext& Conte
 	uint32 Quality = FMath::Clamp(CVar->GetValueOnRenderThread(), 1, 6);
 	bool bUseFast = Quality == 3;
 
+	Context.RHICmdList.SetBlendState(TStaticBlendState<>::GetRHI());
+	Context.RHICmdList.SetRasterizerState(TStaticRasterizerState<>::GetRHI());
+
 	if(Context.View.bCameraCut)
 	{
 		// On camera cut this turns on responsive everywhere.
 		
 		// Normal temporal feedback
-		// Draw to pixels where stencil == 0
-		Context.RHICmdList.SetDepthStencilState(TStaticDepthStencilState<false,CF_Always,true,CF_Equal,SO_Keep,SO_Keep,SO_Keep>::GetRHI(), 0);
+		Context.RHICmdList.SetDepthStencilState(TStaticDepthStencilState<false,CF_Always>::GetRHI());
 
 		TShaderMapRef< FPostProcessTonemapVS >			VertexShader(GetGlobalShaderMap());
 		if (bUseFast)
@@ -514,91 +516,93 @@ void FRCPassPostProcessTemporalAA::Process(FRenderingCompositePassContext& Conte
 			EDRF_UseTriangleOptimization);
 	}
 	else
-	{	
-		// Normal temporal feedback
-		// Draw to pixels where stencil == 0
-		Context.RHICmdList.SetDepthStencilState(TStaticDepthStencilState<false,CF_Always,true,CF_Equal,SO_Keep,SO_Keep,SO_Keep>::GetRHI(), 0);
-
-		TShaderMapRef< FPostProcessTonemapVS >			VertexShader(GetGlobalShaderMap());
-		if (bUseFast)
-		{
-			TShaderMapRef< FPostProcessTemporalAAPS<4,0> >	PixelShader( GetGlobalShaderMap() );
-
-			static FGlobalBoundShaderState BoundShaderState;
-			Context.RHICmdList.CheckIsNull(); // need new approach for "static FGlobalBoundShaderState" for parallel rendering
-
-			SetGlobalBoundShaderState(Context.RHICmdList, BoundShaderState, GFilterVertexDeclaration.VertexDeclarationRHI, *VertexShader, *PixelShader);
-
-			VertexShader->SetVS(Context);
-			PixelShader->SetParameters(Context);
-		}
-		else
-		{
-			TShaderMapRef< FPostProcessTemporalAAPS<1,0> >	PixelShader( GetGlobalShaderMap() );
-
-			static FGlobalBoundShaderState BoundShaderState;
-			Context.RHICmdList.CheckIsNull(); // need new approach for "static FGlobalBoundShaderState" for parallel rendering
-
-			SetGlobalBoundShaderState(Context.RHICmdList, BoundShaderState, GFilterVertexDeclaration.VertexDeclarationRHI, *VertexShader, *PixelShader);
-
-			VertexShader->SetVS(Context);
-			PixelShader->SetParameters(Context);
-		}
+	{
+		{	
+			// Normal temporal feedback
+			// Draw to pixels where stencil == 0
+			Context.RHICmdList.SetDepthStencilState(TStaticDepthStencilState<false,CF_Always,true,CF_Equal,SO_Keep,SO_Keep,SO_Keep>::GetRHI(), 0);
 	
-		// Draw a quad mapping scene color to the view's render target
-		DrawRectangle(
-			Context.RHICmdList,
-			0, 0,
-			SrcRect.Width(), SrcRect.Height(),
-			SrcRect.Min.X, SrcRect.Min.Y, 
-			SrcRect.Width(), SrcRect.Height(),
-			SrcRect.Size(),
-			SrcSize,
-			*VertexShader,
-			EDRF_UseTriangleOptimization);
-	}
-
-	{	// Responsive feedback for tagged pixels
-		// Draw to pixels where stencil != 0
-		Context.RHICmdList.SetDepthStencilState(TStaticDepthStencilState<false,CF_Always,true,CF_NotEqual,SO_Keep,SO_Keep,SO_Keep>::GetRHI(), 0);
+			TShaderMapRef< FPostProcessTonemapVS >			VertexShader(GetGlobalShaderMap());
+			if (bUseFast)
+			{
+				TShaderMapRef< FPostProcessTemporalAAPS<4,0> >	PixelShader( GetGlobalShaderMap() );
+	
+				static FGlobalBoundShaderState BoundShaderState;
+				Context.RHICmdList.CheckIsNull(); // need new approach for "static FGlobalBoundShaderState" for parallel rendering
+	
+				SetGlobalBoundShaderState(Context.RHICmdList, BoundShaderState, GFilterVertexDeclaration.VertexDeclarationRHI, *VertexShader, *PixelShader);
+	
+				VertexShader->SetVS(Context);
+				PixelShader->SetParameters(Context);
+			}
+			else
+			{
+				TShaderMapRef< FPostProcessTemporalAAPS<1,0> >	PixelShader( GetGlobalShaderMap() );
+	
+				static FGlobalBoundShaderState BoundShaderState;
+				Context.RHICmdList.CheckIsNull(); // need new approach for "static FGlobalBoundShaderState" for parallel rendering
+	
+				SetGlobalBoundShaderState(Context.RHICmdList, BoundShaderState, GFilterVertexDeclaration.VertexDeclarationRHI, *VertexShader, *PixelShader);
+	
+				VertexShader->SetVS(Context);
+				PixelShader->SetParameters(Context);
+			}
 		
-		TShaderMapRef< FPostProcessTonemapVS >			VertexShader( GetGlobalShaderMap() );
-		if(bUseFast)
-		{
-			TShaderMapRef< FPostProcessTemporalAAPS<4,1> >	PixelShader( GetGlobalShaderMap() );
-
-			static FGlobalBoundShaderState BoundShaderState;
-			Context.RHICmdList.CheckIsNull(); // need new approach for "static FGlobalBoundShaderState" for parallel rendering
-
-			SetGlobalBoundShaderState(Context.RHICmdList, BoundShaderState, GFilterVertexDeclaration.VertexDeclarationRHI, *VertexShader, *PixelShader);
-
-			VertexShader->SetVS(Context);
-			PixelShader->SetParameters(Context);
+			// Draw a quad mapping scene color to the view's render target
+			DrawRectangle(
+				Context.RHICmdList,
+				0, 0,
+				SrcRect.Width(), SrcRect.Height(),
+				SrcRect.Min.X, SrcRect.Min.Y, 
+				SrcRect.Width(), SrcRect.Height(),
+				SrcRect.Size(),
+				SrcSize,
+				*VertexShader,
+				EDRF_UseTriangleOptimization);
 		}
-		else
-		{
-			TShaderMapRef< FPostProcessTemporalAAPS<1,1> >	PixelShader( GetGlobalShaderMap() );
-
-			static FGlobalBoundShaderState BoundShaderState;
-			Context.RHICmdList.CheckIsNull(); // need new approach for "static FGlobalBoundShaderState" for parallel rendering
-
-			SetGlobalBoundShaderState(Context.RHICmdList, BoundShaderState, GFilterVertexDeclaration.VertexDeclarationRHI, *VertexShader, *PixelShader);
-
-			VertexShader->SetVS(Context);
-			PixelShader->SetParameters(Context);
+	
+		{	// Responsive feedback for tagged pixels
+			// Draw to pixels where stencil != 0
+			Context.RHICmdList.SetDepthStencilState(TStaticDepthStencilState<false,CF_Always,true,CF_NotEqual,SO_Keep,SO_Keep,SO_Keep>::GetRHI(), 0);
+			
+			TShaderMapRef< FPostProcessTonemapVS >			VertexShader( GetGlobalShaderMap() );
+			if(bUseFast)
+			{
+				TShaderMapRef< FPostProcessTemporalAAPS<4,1> >	PixelShader( GetGlobalShaderMap() );
+	
+				static FGlobalBoundShaderState BoundShaderState;
+				Context.RHICmdList.CheckIsNull(); // need new approach for "static FGlobalBoundShaderState" for parallel rendering
+	
+				SetGlobalBoundShaderState(Context.RHICmdList, BoundShaderState, GFilterVertexDeclaration.VertexDeclarationRHI, *VertexShader, *PixelShader);
+	
+				VertexShader->SetVS(Context);
+				PixelShader->SetParameters(Context);
+			}
+			else
+			{
+				TShaderMapRef< FPostProcessTemporalAAPS<1,1> >	PixelShader( GetGlobalShaderMap() );
+	
+				static FGlobalBoundShaderState BoundShaderState;
+				Context.RHICmdList.CheckIsNull(); // need new approach for "static FGlobalBoundShaderState" for parallel rendering
+	
+				SetGlobalBoundShaderState(Context.RHICmdList, BoundShaderState, GFilterVertexDeclaration.VertexDeclarationRHI, *VertexShader, *PixelShader);
+	
+				VertexShader->SetVS(Context);
+				PixelShader->SetParameters(Context);
+			}
+	
+			// Draw a quad mapping scene color to the view's render target
+			DrawRectangle(
+				Context.RHICmdList,
+				0, 0,
+				SrcRect.Width(), SrcRect.Height(),
+				SrcRect.Min.X, SrcRect.Min.Y, 
+				SrcRect.Width(), SrcRect.Height(),
+				SrcRect.Size(),
+				SrcSize,
+				*VertexShader,
+				EDRF_UseTriangleOptimization);
 		}
-
-		// Draw a quad mapping scene color to the view's render target
-		DrawRectangle(
-			Context.RHICmdList,
-			0, 0,
-			SrcRect.Width(), SrcRect.Height(),
-			SrcRect.Min.X, SrcRect.Min.Y, 
-			SrcRect.Width(), SrcRect.Height(),
-			SrcRect.Size(),
-			SrcSize,
-			*VertexShader,
-			EDRF_UseTriangleOptimization);
 	}
 
 	Context.RHICmdList.CopyToResolveTarget(DestRenderTarget.TargetableTexture, DestRenderTarget.ShaderResourceTexture, false, FResolveParams());
diff --git a/Engine/Source/Runtime/Renderer/Private/PostProcess/PostProcessing.cpp b/Engine/Source/Runtime/Renderer/Private/PostProcess/PostProcessing.cpp
--- a/Engine/Source/Runtime/Renderer/Private/PostProcess/PostProcessing.cpp
+++ b/Engine/Source/Runtime/Renderer/Private/PostProcess/PostProcessing.cpp
@@ -999,7 +999,7 @@ void FPostProcessing::Process(const FViewInfo& View, TRefCountPtr<IPooledRenderT
 				AddTonemapper(Context, BloomOutputCombined, EyeAdaptation, CombinedLUT);
 			}
 
-			if(AntiAliasingMethod == AAM_FXAA || (AntiAliasingMethod == AAM_TemporalAA && View.bCameraCut))
+			if(AntiAliasingMethod == AAM_FXAA)
 			{
 				AddPostProcessAA(Context);
 			}
