UE4 - Whole scene static shadows using RHICmdList (r.RHICmd > 1); static meshes pass using r.RHICmd > 0

diff --git a/Engine/Source/Runtime/Engine/Public/SceneUtils.h b/Engine/Source/Runtime/Engine/Public/SceneUtils.h
--- a/Engine/Source/Runtime/Engine/Public/SceneUtils.h
+++ b/Engine/Source/Runtime/Engine/Public/SceneUtils.h
@@ -43,8 +43,8 @@ inline void appSetCounterValue(const TCHAR* CounterName, float Value) {}
 		bool bDrawEventHasBeenEmitted;
 
 		/** Default constructor, initializing all member variables. */
-		FDrawEvent()
-		:	bDrawEventHasBeenEmitted( false )
+		FORCEINLINE FDrawEvent() :
+			bDrawEventHasBeenEmitted( false )
 		{}
 
 		/**
diff --git a/Engine/Source/Runtime/Renderer/Private/DeferredShadingRenderer.cpp b/Engine/Source/Runtime/Renderer/Private/DeferredShadingRenderer.cpp
--- a/Engine/Source/Runtime/Renderer/Private/DeferredShadingRenderer.cpp
+++ b/Engine/Source/Runtime/Renderer/Private/DeferredShadingRenderer.cpp
@@ -329,7 +329,7 @@ bool FDeferredShadingSceneRenderer::RenderBasePass(FViewInfo& View)
 		//@todo-rco: Very temp code!!!
 		SCOPE_CYCLE_COUNTER(STAT_RHICounterTEMP);
 		static IConsoleVariable* RHICmdListCVar = IConsoleManager::Get().FindConsoleVariable(TEXT("r.RHICmd"));
-		bool bUseRHICmdList = RHICmdListCVar->GetInt() != 0;
+		bool bUseRHICmdList = (RHICmdListCVar->GetInt() >= 1);
 		FRHICommandList& RHICmdList = bUseRHICmdList ? GRHICommandList.CreateList() : FRHICommandList::GetNullRef();
 		bDirty |= RenderBasePassStaticData(RHICmdList, View);
 		GRHICommandList.ExecuteAndFreeList(RHICmdList);
diff --git a/Engine/Source/Runtime/Renderer/Private/DepthRendering.cpp b/Engine/Source/Runtime/Renderer/Private/DepthRendering.cpp
--- a/Engine/Source/Runtime/Renderer/Private/DepthRendering.cpp
+++ b/Engine/Source/Runtime/Renderer/Private/DepthRendering.cpp
@@ -172,7 +172,7 @@ FDepthDrawingPolicy::FDepthDrawingPolicy(
 void FDepthDrawingPolicy::DrawShared(FRHICommandList& RHICmdList, const FSceneView* View,FBoundShaderStateRHIParamRef BoundShaderState) const
 {
 	// Set the actual shader & vertex declaration state
-	RHISetBoundShaderState(BoundShaderState);
+	RHICmdList.SetBoundShaderState(BoundShaderState);
 
 	// Set the depth-only shader parameters for the material.
 	VertexShader->SetParameters(RHICmdList, MaterialRenderProxy,*MaterialResource,*View);
@@ -262,13 +262,13 @@ void FPositionOnlyDepthDrawingPolicy::DrawShared(FRHICommandList& RHICmdList, co
 	RHICmdList.CheckIsNull();
 
 	// Set the actual shader & vertex declaration state
-	RHISetBoundShaderState( BoundShaderState);
+	RHICmdList.SetBoundShaderState( BoundShaderState);
 
 	// Set the depth-only shader parameters for the material.
 	VertexShader->SetParameters(RHICmdList, MaterialRenderProxy,*MaterialResource,*View);
 
 	// Set the shared mesh resources.
-	VertexFactory->SetPositionStream();
+	VertexFactory->SetPositionStream(RHICmdList);
 }
 
 /** 
diff --git a/Engine/Source/Runtime/Renderer/Private/ShadowRendering.cpp b/Engine/Source/Runtime/Renderer/Private/ShadowRendering.cpp
--- a/Engine/Source/Runtime/Renderer/Private/ShadowRendering.cpp
+++ b/Engine/Source/Runtime/Renderer/Private/ShadowRendering.cpp
@@ -850,9 +850,8 @@ void FShadowDepthDrawingPolicy<bRenderingReflectiveShadowMaps>::DrawShared(FRHIC
 {
 	checkSlow(bDirectionalLight == PolicyShadowInfo->bDirectionalLight && bPreShadow == PolicyShadowInfo->bPreShadow);
 
-	RHICmdList.CheckIsNull();
 	// Set the actual shader & vertex declaration state
-	RHISetBoundShaderState(BoundShaderState);
+	RHICmdList.SetBoundShaderState(BoundShaderState);
 
 	VertexShader->SetParameters(RHICmdList, MaterialRenderProxy,*MaterialResource,*View,PolicyShadowInfo);
 
@@ -875,7 +874,7 @@ void FShadowDepthDrawingPolicy<bRenderingReflectiveShadowMaps>::DrawShared(FRHIC
 	// Set the shared mesh resources.
 	if (bUsePositionOnlyVS)
 	{
-		VertexFactory->SetPositionStream();
+		VertexFactory->SetPositionStream(RHICmdList);
 	}
 	else
 	{
@@ -890,7 +889,7 @@ void FShadowDepthDrawingPolicy<bRenderingReflectiveShadowMaps>::DrawShared(FRHIC
 	}
 
 	// Set the rasterizer state only once per draw list bucket, instead of once per mesh in SetMeshRenderState as an optimization
-	RHISetRasterizerState(GetStaticRasterizerState<true>(
+	RHICmdList.SetRasterizerState(GetStaticRasterizerState<true>(
 		FM_Solid,
 		bIsTwoSided ? CM_None :
 			// Have to flip culling when doing one pass point light shadows for some reason
@@ -1370,6 +1369,12 @@ void FProjectedShadowInfo::RenderDepth(FRHICommandList& RHICmdList, FDeferredSha
 	{
 		SCOPE_CYCLE_COUNTER(STAT_WholeSceneStaticDrawListShadowDepthsTime);
 
+//@todo-rco: Very temp code!!!
+SCOPE_CYCLE_COUNTER(STAT_RHICounterTEMP);
+static IConsoleVariable* RHICmdListCVar = IConsoleManager::Get().FindConsoleVariable(TEXT("r.RHICmd"));
+bool bUseRHICmdList = (RHICmdListCVar->GetInt() >= 2);
+FRHICommandList& RHICmdList = bUseRHICmdList ? GRHICommandList.CreateList() : FRHICommandList::GetNullRef();
+
 		if (bReflectiveShadowmap)
 		{
 			SceneRenderer->Scene->WholeSceneReflectiveShadowMapDrawList.DrawVisible(RHICmdList, *FoundView, StaticMeshWholeSceneShadowDepthMap, StaticMeshWholeSceneShadowBatchVisibility);
@@ -1379,6 +1384,8 @@ void FProjectedShadowInfo::RenderDepth(FRHICommandList& RHICmdList, FDeferredSha
 			// Use the scene's shadow depth draw list with this shadow's visibility map
 			SceneRenderer->Scene->WholeSceneShadowDepthDrawList.DrawVisible(RHICmdList, *FoundView, StaticMeshWholeSceneShadowDepthMap, StaticMeshWholeSceneShadowBatchVisibility);
 		}
+
+GRHICommandList.ExecuteAndFreeList(RHICmdList);
 	}
 	// Draw the subject's static elements using manual state filtering
 	else if (SubjectMeshElements.Num() > 0)
diff --git a/Engine/Source/Runtime/ShaderCore/Private/VertexFactory.cpp b/Engine/Source/Runtime/ShaderCore/Private/VertexFactory.cpp
--- a/Engine/Source/Runtime/ShaderCore/Private/VertexFactory.cpp
+++ b/Engine/Source/Runtime/ShaderCore/Private/VertexFactory.cpp
@@ -178,11 +178,8 @@ void FVertexFactory::Set(FRHICommandList& RHICmdList) const
 	}
 }
 
-void FVertexFactory::SetPositionStream() const
+void FVertexFactory::SetPositionStream(FRHICommandList& RHICmdList) const
 {
-	//@todo-rco: RHIPacketList
-	FRHICommandList& RHICmdList = FRHICommandList::GetNullRef();
-
 	check(IsInitialized());
 	// Set the predefined vertex streams.
 	for(int32 StreamIndex = 0;StreamIndex < PositionStream.Num();StreamIndex++)
diff --git a/Engine/Source/Runtime/ShaderCore/Public/VertexFactory.h b/Engine/Source/Runtime/ShaderCore/Public/VertexFactory.h
--- a/Engine/Source/Runtime/ShaderCore/Public/VertexFactory.h
+++ b/Engine/Source/Runtime/ShaderCore/Public/VertexFactory.h
@@ -344,7 +344,7 @@ public:
 	/**
 	* Sets the position stream as the current stream source.
 	*/
-	void SetPositionStream() const;
+	void SetPositionStream(FRHICommandList& RHICmdList) const;
 
 	/**
 	* Can be overridden by FVertexFactory subclasses to modify their compile environment just before compilation occurs.
