Fixed SSR noise when temporal AA is off.

diff --git a/Engine/Source/Runtime/Renderer/Private/PostProcess/ScreenSpaceReflections.cpp b/Engine/Source/Runtime/Renderer/Private/PostProcess/ScreenSpaceReflections.cpp
--- a/Engine/Source/Runtime/Renderer/Private/PostProcess/ScreenSpaceReflections.cpp
+++ b/Engine/Source/Runtime/Renderer/Private/PostProcess/ScreenSpaceReflections.cpp
@@ -269,7 +269,7 @@ void ScreenSpaceReflections( const FViewInfo& View, TRefCountPtr<IPooledRenderTa
 		Context.FinalOutput = FRenderingCompositeOutputRef( TracePass );
 	}
 
-	const bool bTemporalFilter = CVarSSRTemporal.GetValueOnRenderThread() != 0;
+	const bool bTemporalFilter = View.FinalPostProcessSettings.AntiAliasingMethod != AAM_TemporalAA || CVarSSRTemporal.GetValueOnRenderThread() != 0;
 
 	if( ViewState && bTemporalFilter )
 	{
