UE4 - RHICmdList on Translucent rendering

diff --git a/Engine/Source/Runtime/Renderer/Private/TranslucentRendering.cpp b/Engine/Source/Runtime/Renderer/Private/TranslucentRendering.cpp
--- a/Engine/Source/Runtime/Renderer/Private/TranslucentRendering.cpp
+++ b/Engine/Source/Runtime/Renderer/Private/TranslucentRendering.cpp
@@ -310,6 +310,7 @@ public:
 * @return true if the mesh rendered
 */
 bool FTranslucencyDrawingPolicyFactory::DrawMesh(
+	FRHICommandList& RHICmdList,
 	const FViewInfo& View,
 	ContextType DrawingContext,
 	const FMeshBatch& Mesh,
@@ -328,9 +329,6 @@ bool FTranslucencyDrawingPolicyFactory::DrawMesh(
 	const EBlendMode BlendMode = Material->GetBlendMode();
 	const EMaterialShadingModel ShadingModel = Material->GetShadingModel();
 
-	//@todo-rco: RHIPacketList
-	FRHICommandList& RHICmdList = FRHICommandList::GetNullRef();
-
 	// Only render translucent materials.
 	if(IsTranslucentBlendMode(BlendMode))
 	{
@@ -343,17 +341,17 @@ bool FTranslucencyDrawingPolicyFactory::DrawMesh(
 		{
 			if( bDisableDepthTest )
 			{
-				RHISetDepthStencilState(TStaticDepthStencilState<false,CF_Always,true,CF_Always,SO_Keep,SO_Keep,SO_Replace>::GetRHI(), 1);
+				RHICmdList.SetDepthStencilState(TStaticDepthStencilState<false,CF_Always,true,CF_Always,SO_Keep,SO_Keep,SO_Replace>::GetRHI(), 1);
 			}
 			else
 			{
 				// Note, this is a reversed Z depth surface, using CF_GreaterEqual.	
-				RHISetDepthStencilState(TStaticDepthStencilState<false,CF_GreaterEqual,true,CF_Always,SO_Keep,SO_Keep,SO_Replace>::GetRHI(), 1);
+				RHICmdList.SetDepthStencilState(TStaticDepthStencilState<false,CF_GreaterEqual,true,CF_Always,SO_Keep,SO_Keep,SO_Replace>::GetRHI(), 1);
 			}
 		}
 		else if( bDisableDepthTest )
 		{
-			RHISetDepthStencilState(TStaticDepthStencilState<false,CF_Always>::GetRHI());
+			RHICmdList.SetDepthStencilState(TStaticDepthStencilState<false,CF_Always>::GetRHI());
 		}
 
 		ProcessBasePassMesh(
@@ -381,7 +379,7 @@ bool FTranslucencyDrawingPolicyFactory::DrawMesh(
 		{
 			// Restore default depth state
 			// Note, this is a reversed Z depth surface, using CF_GreaterEqual.	
-			RHISetDepthStencilState(TStaticDepthStencilState<false,CF_GreaterEqual>::GetRHI());
+			RHICmdList.SetDepthStencilState(TStaticDepthStencilState<false,CF_GreaterEqual>::GetRHI());
 		}
 
 		bDirty = true;
@@ -406,6 +404,7 @@ bool FTranslucencyDrawingPolicyFactory::DrawDynamicMesh(
 	)
 {
 	return DrawMesh(
+		RHICmdList,
 		View,
 		DrawingContext,
 		Mesh,
@@ -432,6 +431,7 @@ bool FTranslucencyDrawingPolicyFactory::DrawStaticMesh(
 	)
 {
 	return DrawMesh(
+		RHICmdList,
 		View,
 		DrawingContext,
 		StaticMesh,
@@ -527,12 +527,15 @@ void FTranslucentPrimSet::RenderPrimitive(
 				);
 		}
 
-		//@todo-rco: RHIPacketList
-		FRHICommandList& RHICmdList = FRHICommandList::GetNullRef();
-
 		// Render static scene prim
 		if( ViewRelevance.bStaticRelevance )
 		{
+			//@todo-rco: Very temp code!!!
+			SCOPE_CYCLE_COUNTER(STAT_RHICounterTEMP);
+			static IConsoleVariable* RHICmdListCVar = IConsoleManager::Get().FindConsoleVariable(TEXT("r.RHICmd"));
+			bool bUseRHICmdList = (RHICmdListCVar->GetInt() >= 3);
+			FRHICommandList& RHICmdList = bUseRHICmdList ? GRHICommandList.CreateList() : FRHICommandList::GetNullRef();
+
 			// Render static meshes from static scene prim
 			for( int32 StaticMeshIdx=0; StaticMeshIdx < PrimitiveSceneInfo->StaticMeshes.Num(); StaticMeshIdx++ )
 			{
@@ -554,6 +557,8 @@ void FTranslucentPrimSet::RenderPrimitive(
 						);
 				}
 			}
+
+			GRHICommandList.ExecuteAndFreeList(RHICmdList);
 		}
 	}
 }
diff --git a/Engine/Source/Runtime/Renderer/Private/TranslucentRendering.h b/Engine/Source/Runtime/Renderer/Private/TranslucentRendering.h
--- a/Engine/Source/Runtime/Renderer/Private/TranslucentRendering.h
+++ b/Engine/Source/Runtime/Renderer/Private/TranslucentRendering.h
@@ -66,6 +66,7 @@ private:
 	* @return true if the mesh rendered
 	*/
 	static bool DrawMesh(
+		FRHICommandList& RHICmdList,
 		const FViewInfo& View,
 		ContextType DrawingContext,
 		const FMeshBatch& Mesh,
