Merging fix for shadow using invalid transform

diff --git a/Engine/Source/Runtime/Renderer/Private/ShadowRendering.cpp b/Engine/Source/Runtime/Renderer/Private/ShadowRendering.cpp
--- a/Engine/Source/Runtime/Renderer/Private/ShadowRendering.cpp
+++ b/Engine/Source/Runtime/Renderer/Private/ShadowRendering.cpp
@@ -2435,7 +2435,7 @@ void FDeferredShadingSceneRenderer::RenderProjections(
 		for (int32 ShadowIndex = 0; ShadowIndex < Shadows.Num(); ShadowIndex++)
 		{
 			FProjectedShadowInfo* ProjectedShadowInfo = Shadows[ShadowIndex];
-			if (ProjectedShadowInfo->bAllocated)
+			if (ProjectedShadowInfo->bAllocated && ProjectedShadowInfo->bValidTransform)
 			{
 				// Only project the shadow if it's large enough in this particular view (split screen, etc... may have shadows that are large in one view but irrelevantly small in others)
 				if (ProjectedShadowInfo->FadeAlphas[ViewIndex] > 1.0f / 256.0f)
