Moved FMaterialShaderto renderer public so slate can access it.

diff --git a/Engine/Source/Runtime/Engine/Private/Materials/Material.cpp b/Engine/Source/Runtime/Engine/Private/Materials/Material.cpp
--- a/Engine/Source/Runtime/Engine/Private/Materials/Material.cpp
+++ b/Engine/Source/Runtime/Engine/Private/Materials/Material.cpp
@@ -24,7 +24,7 @@
 #include "Materials/MaterialExpressionTextureSampleParameter.h"
 #include "Materials/MaterialExpressionVectorParameter.h"
 #include "EditorSupportDelegates.h"
-#include "MaterialShader.h"
+#include "MaterialShaderType.h"
 #include "MaterialInstanceSupport.h"
 #include "UObjectAnnotation.h"
 #include "MaterialCompiler.h"
diff --git a/Engine/Source/Runtime/Engine/Private/Materials/MaterialInstance.cpp b/Engine/Source/Runtime/Engine/Private/Materials/MaterialInstance.cpp
--- a/Engine/Source/Runtime/Engine/Private/Materials/MaterialInstance.cpp
+++ b/Engine/Source/Runtime/Engine/Private/Materials/MaterialInstance.cpp
@@ -10,7 +10,7 @@
 #include "Materials/MaterialExpressionLandscapeLayerWeight.h"
 #include "Materials/MaterialInstanceBasePropertyOverrides.h"
 #include "MaterialInstanceSupport.h"
-#include "MaterialShader.h"
+#include "MaterialShaderType.h"
 #include "TargetPlatform.h"
 #include "Engine/Font.h"
 
diff --git a/Engine/Source/Runtime/Engine/Private/Materials/MaterialShader.cpp b/Engine/Source/Runtime/Engine/Private/Materials/MaterialShader.cpp
--- a/Engine/Source/Runtime/Engine/Private/Materials/MaterialShader.cpp
+++ b/Engine/Source/Runtime/Engine/Private/Materials/MaterialShader.cpp
@@ -1,15 +1,11 @@
 // Copyright 1998-2014 Epic Games, Inc. All Rights Reserved.
 
-/*=============================================================================
-	MaterialShader.h: Material shader definitions.
-=============================================================================*/
-
 #include "EnginePrivate.h"
 #include "DiagnosticTable.h"
-#include "MaterialShader.h"
+#include "MaterialShaderType.h"
+#include "MeshMaterialShaderType.h"
 #include "ShaderCompiler.h"
 #include "DerivedDataCacheInterface.h"
-#include "MeshMaterialShader.h"
 #include "TargetPlatform.h"
 #include "../ShaderDerivedDataVersion.h"
 
diff --git a/Engine/Source/Runtime/Engine/Private/Materials/MaterialShared.cpp b/Engine/Source/Runtime/Engine/Private/Materials/MaterialShared.cpp
--- a/Engine/Source/Runtime/Engine/Private/Materials/MaterialShared.cpp
+++ b/Engine/Source/Runtime/Engine/Private/Materials/MaterialShared.cpp
@@ -10,8 +10,8 @@
 #include "PixelFormat.h"
 #include "ShaderCompiler.h"
 #include "MaterialCompiler.h"
-#include "MaterialShader.h"
-#include "MeshMaterialShader.h"
+#include "MaterialShaderType.h"
+#include "MeshMaterialShaderType.h"
 #include "HLSLMaterialTranslator.h"
 #include "MaterialUniformExpressions.h"
 #include "Developer/TargetPlatform/Public/TargetPlatform.h"
diff --git a/Engine/Source/Runtime/Engine/Private/Materials/MeshMaterialShader.cpp b/Engine/Source/Runtime/Engine/Private/Materials/MeshMaterialShader.cpp
--- a/Engine/Source/Runtime/Engine/Private/Materials/MeshMaterialShader.cpp
+++ b/Engine/Source/Runtime/Engine/Private/Materials/MeshMaterialShader.cpp
@@ -6,8 +6,7 @@
 
 #include "EnginePrivate.h"
 #include "ShaderCompiler.h"
-#include "MeshMaterialShader.h"
-#include "MaterialShader.h"
+#include "MeshMaterialShaderType.h"
 
 /**
 * Finds a FMeshMaterialShaderType by name.
diff --git a/Engine/Source/Runtime/Engine/Public/MaterialShader.h b/Engine/Source/Runtime/Engine/Public/MaterialShader.h
deleted file mode 100644
diff --git a/Engine/Source/Runtime/Engine/Public/MaterialShaderType.h b/Engine/Source/Runtime/Engine/Public/MaterialShaderType.h
new file mode 100644
--- /dev/null
+++ b/Engine/Source/Runtime/Engine/Public/MaterialShaderType.h
@@ -0,0 +1,146 @@
+// Copyright 1998-2014 Epic Games, Inc. All Rights Reserved.
+
+/*=============================================================================
+	MaterialShader.h: Material shader definitions.
+=============================================================================*/
+
+#pragma once
+
+#include "Shader.h"
+#include "ShaderParameters.h"
+#include "MaterialShared.h"
+#include "GlobalShader.h"
+
+/** A macro to implement material shaders. */
+#define IMPLEMENT_MATERIAL_SHADER_TYPE(TemplatePrefix,ShaderClass,SourceFilename,FunctionName,Frequency) \
+	IMPLEMENT_SHADER_TYPE( \
+		TemplatePrefix, \
+		ShaderClass, \
+		SourceFilename, \
+		FunctionName, \
+		Frequency \
+		);
+
+/** Converts an EMaterialShadingModel to a string description. */
+extern FString GetShadingModelString(EMaterialShadingModel ShadingModel);
+
+/** Converts an EBlendMode to a string description. */
+extern FString GetBlendModeString(EBlendMode BlendMode);
+
+/** Called for every material shader to update the appropriate stats. */
+extern void UpdateMaterialShaderCompilingStats(const FMaterial* Material);
+
+/**
+ * Dump material stats for a given platform.
+ * 
+ * @param	Platform	Platform to dump stats for.
+ */
+extern ENGINE_API void DumpMaterialStats( EShaderPlatform Platform );
+
+/**
+ * A shader meta type for material-linked shaders.
+ */
+class FMaterialShaderType : public FShaderType
+{
+public:
+
+	/**
+	 * Finds a FMaterialShaderType by name.
+	 */
+	static FMaterialShaderType* GetTypeByName(const FString& TypeName);
+
+	struct CompiledShaderInitializerType : FGlobalShaderType::CompiledShaderInitializerType
+	{
+		const FUniformExpressionSet& UniformExpressionSet;
+		const FString DebugDescription;
+
+		CompiledShaderInitializerType(
+			FShaderType* InType,
+			const FShaderCompilerOutput& CompilerOutput,
+			FShaderResource* InResource,
+			const FUniformExpressionSet& InUniformExpressionSet,
+			const FSHAHash& InMaterialShaderMapHash,
+			FVertexFactoryType* InVertexFactoryType,
+			const FString& InDebugDescription
+			)
+		: FGlobalShaderType::CompiledShaderInitializerType(InType,CompilerOutput,InResource,InMaterialShaderMapHash,InVertexFactoryType)
+		, UniformExpressionSet(InUniformExpressionSet)
+		, DebugDescription(InDebugDescription)
+		{}
+	};
+	typedef FShader* (*ConstructCompiledType)(const CompiledShaderInitializerType&);
+	typedef bool (*ShouldCacheType)(EShaderPlatform,const FMaterial*);
+	typedef void (*ModifyCompilationEnvironmentType)(EShaderPlatform, const FMaterial*, FShaderCompilerEnvironment&);
+
+	FMaterialShaderType(
+		const TCHAR* InName,
+		const TCHAR* InSourceFilename,
+		const TCHAR* InFunctionName,
+		uint32 InFrequency,
+		ConstructSerializedType InConstructSerializedRef,
+		ConstructCompiledType InConstructCompiledRef,
+		ModifyCompilationEnvironmentType InModifyCompilationEnvironmentRef,
+		ShouldCacheType InShouldCacheRef
+		):
+		FShaderType(InName,InSourceFilename,InFunctionName,InFrequency,InConstructSerializedRef),
+		ConstructCompiledRef(InConstructCompiledRef),
+		ShouldCacheRef(InShouldCacheRef),
+		ModifyCompilationEnvironmentRef(InModifyCompilationEnvironmentRef)
+	{}
+
+	/**
+	 * Enqueues a compilation for a new shader of this type.
+	 * @param Material - The material to link the shader with.
+	 */
+	void BeginCompileShader(
+		uint32 ShaderMapId,
+		const FMaterial* Material,
+		FShaderCompilerEnvironment* MaterialEnvironment,
+		EShaderPlatform Platform,
+		TArray<FShaderCompileJob*>& NewJobs
+		);
+
+	/**
+	 * Either creates a new instance of this type or returns an equivalent existing shader.
+	 * @param Material - The material to link the shader with.
+	 * @param CurrentJob - Compile job that was enqueued by BeginCompileShader.
+	 */
+	FShader* FinishCompileShader(
+		const FUniformExpressionSet& UniformExpressionSet,
+		const FSHAHash& MaterialShaderMapHash,
+		const FShaderCompileJob& CurrentJob,
+		const FString& InDebugDescription
+		);
+
+	/**
+	 * Checks if the shader type should be cached for a particular platform and material.
+	 * @param Platform - The platform to check.
+	 * @param Material - The material to check.
+	 * @return True if this shader type should be cached.
+	 */
+	bool ShouldCache(EShaderPlatform Platform,const FMaterial* Material) const
+	{
+		return (*ShouldCacheRef)(Platform,Material);
+	}
+
+	// Dynamic casting.
+	virtual FMaterialShaderType* GetMaterialShaderType() { return this; }
+
+protected:
+
+	/**
+	 * Sets up the environment used to compile an instance of this shader type.
+	 * @param Platform - Platform to compile for.
+	 * @param Environment - The shader compile environment that the function modifies.
+	 */
+	void SetupCompileEnvironment(EShaderPlatform Platform, const FMaterial* Material, FShaderCompilerEnvironment& Environment)
+	{
+		// Allow the shader type to modify its compile environment.
+		(*ModifyCompilationEnvironmentRef)(Platform, Material, Environment);
+	}
+
+private:
+	ConstructCompiledType ConstructCompiledRef;
+	ShouldCacheType ShouldCacheRef;
+	ModifyCompilationEnvironmentType ModifyCompilationEnvironmentRef;
+};
\ No newline at end of file
diff --git a/Engine/Source/Runtime/Engine/Public/MeshMaterialShader.h b/Engine/Source/Runtime/Engine/Public/MeshMaterialShader.h
deleted file mode 100644
diff --git a/Engine/Source/Runtime/Engine/Public/MeshMaterialShaderType.h b/Engine/Source/Runtime/Engine/Public/MeshMaterialShaderType.h
new file mode 100644
--- /dev/null
+++ b/Engine/Source/Runtime/Engine/Public/MeshMaterialShaderType.h
@@ -0,0 +1,120 @@
+// Copyright 1998-2014 Epic Games, Inc. All Rights Reserved.
+
+/*=============================================================================
+	MeshMaterialShader.h: Mesh material shader definitions.
+=============================================================================*/
+
+#pragma once
+
+#include "Shader.h"
+#include "MaterialShader.h"
+
+/**
+ * A shader meta type for material-linked shaders which use a vertex factory.
+ */
+class FMeshMaterialShaderType : public FShaderType
+{
+public:
+
+	/**
+	* Finds a FMeshMaterialShaderType by name.
+	*/
+	static FMeshMaterialShaderType* GetTypeByName(const FString& TypeName);
+
+	struct CompiledShaderInitializerType : FMaterialShaderType::CompiledShaderInitializerType
+	{
+		FVertexFactoryType* VertexFactoryType;
+		CompiledShaderInitializerType(
+			FShaderType* InType,
+			const FShaderCompilerOutput& CompilerOutput,
+			FShaderResource* InResource,
+			const FUniformExpressionSet& InUniformExpressionSet,
+			const FSHAHash& InMaterialShaderMapHash,
+			const FString& InDebugDescription,
+			FVertexFactoryType* InVertexFactoryType
+			):
+			FMaterialShaderType::CompiledShaderInitializerType(InType,CompilerOutput,InResource,InUniformExpressionSet,InMaterialShaderMapHash,InVertexFactoryType,InDebugDescription),
+			VertexFactoryType(InVertexFactoryType)
+		{}
+	};
+	typedef FShader* (*ConstructCompiledType)(const CompiledShaderInitializerType&);
+	typedef bool (*ShouldCacheType)(EShaderPlatform,const FMaterial*,const FVertexFactoryType* VertexFactoryType);
+	typedef void (*ModifyCompilationEnvironmentType)(EShaderPlatform, const FMaterial*, FShaderCompilerEnvironment&);
+
+	FMeshMaterialShaderType(
+		const TCHAR* InName,
+		const TCHAR* InSourceFilename,
+		const TCHAR* InFunctionName,
+		uint32 InFrequency,
+		ConstructSerializedType InConstructSerializedRef,
+		ConstructCompiledType InConstructCompiledRef,
+		ModifyCompilationEnvironmentType InModifyCompilationEnvironmentRef,
+		ShouldCacheType InShouldCacheRef
+		):
+		FShaderType(InName,InSourceFilename,InFunctionName,InFrequency,InConstructSerializedRef),
+		ConstructCompiledRef(InConstructCompiledRef),
+		ShouldCacheRef(InShouldCacheRef),
+		ModifyCompilationEnvironmentRef(InModifyCompilationEnvironmentRef)
+	{}
+
+	/**
+	 * Enqueues a compilation for a new shader of this type.
+	 * @param Platform - The platform to compile for.
+	 * @param Material - The material to link the shader with.
+	 * @param VertexFactoryType - The vertex factory to compile with.
+	 */
+	void BeginCompileShader(
+		uint32 ShaderMapId,
+		EShaderPlatform Platform,
+		const FMaterial* Material,
+		FShaderCompilerEnvironment* MaterialEnvironment,
+		FVertexFactoryType* VertexFactoryType,
+		TArray<FShaderCompileJob*>& NewJobs
+		);
+
+	/**
+	 * Either creates a new instance of this type or returns an equivalent existing shader.
+	 * @param Material - The material to link the shader with.
+	 * @param CurrentJob - Compile job that was enqueued by BeginCompileShader.
+	 */
+	FShader* FinishCompileShader(
+		const FUniformExpressionSet& UniformExpressionSet, 
+		const FSHAHash& MaterialShaderMapHash,
+		const FShaderCompileJob& CurrentJob,
+		const FString& InDebugDescription
+		);
+
+	/**
+	 * Checks if the shader type should be cached for a particular platform, material, and vertex factory type.
+	 * @param Platform - The platform to check.
+	 * @param Material - The material to check.
+	 * @param VertexFactoryType - The vertex factory type to check.
+	 * @return True if this shader type should be cached.
+	 */
+	bool ShouldCache(EShaderPlatform Platform,const FMaterial* Material,const FVertexFactoryType* VertexFactoryType) const
+	{
+		return (*ShouldCacheRef)(Platform,Material,VertexFactoryType);
+	}
+
+	// Dynamic casting.
+	virtual FMeshMaterialShaderType* GetMeshMaterialShaderType() { return this; }
+
+protected:
+
+	/**
+	 * Sets up the environment used to compile an instance of this shader type.
+	 * @param Platform - Platform to compile for.
+	 * @param Environment - The shader compile environment that the function modifies.
+	 */
+	void SetupCompileEnvironment(EShaderPlatform Platform, const FMaterial* Material, FShaderCompilerEnvironment& Environment)
+	{
+		// Allow the shader type to modify its compile environment.
+		(*ModifyCompilationEnvironmentRef)(Platform, Material, Environment);
+	}
+
+private:
+	ConstructCompiledType ConstructCompiledRef;
+	ShouldCacheType ShouldCacheRef;
+	ModifyCompilationEnvironmentType ModifyCompilationEnvironmentRef;
+};
+
diff --git a/Engine/Source/Runtime/Renderer/Private/AtmosphereRendering.h b/Engine/Source/Runtime/Renderer/Private/AtmosphereRendering.h
--- a/Engine/Source/Runtime/Renderer/Private/AtmosphereRendering.h
+++ b/Engine/Source/Runtime/Renderer/Private/AtmosphereRendering.h
@@ -80,43 +80,4 @@ private:
 #endif
 };
 
-/** Shader parameters needed for atmosphere passes. */
-class FAtmosphereShaderTextureParameters
-{
-public:
-	void Bind(const FShaderParameterMap& ParameterMap)
-	{
-		TransmittanceTexture.Bind(ParameterMap,TEXT("AtmosphereTransmittanceTexture"));
-		TransmittanceTextureSampler.Bind(ParameterMap,TEXT("AtmosphereTransmittanceTextureSampler"));
-		IrradianceTexture.Bind(ParameterMap,TEXT("AtmosphereIrradianceTexture"));
-		IrradianceTextureSampler.Bind(ParameterMap,TEXT("AtmosphereIrradianceTextureSampler"));
-		InscatterTexture.Bind(ParameterMap,TEXT("AtmosphereInscatterTexture"));
-		InscatterTextureSampler.Bind(ParameterMap,TEXT("AtmosphereInscatterTextureSampler"));
-	}
-
-	template<typename ShaderRHIParamRef>
-	void Set(FRHICommandList& RHICmdList, const ShaderRHIParamRef ShaderRHI, const FSceneView& View) const
-	{
-		if (TransmittanceTexture.IsBound() || IrradianceTexture.IsBound() || InscatterTexture.IsBound())
-		{
-			SetTextureParameter(RHICmdList, ShaderRHI, TransmittanceTexture, TransmittanceTextureSampler, 
-				TStaticSamplerState<SF_Bilinear>::GetRHI(), View.AtmosphereTransmittanceTexture);
-			SetTextureParameter(RHICmdList, ShaderRHI, IrradianceTexture, IrradianceTextureSampler, 
-				TStaticSamplerState<SF_Bilinear>::GetRHI(), View.AtmosphereIrradianceTexture);
-			SetTextureParameter(RHICmdList, ShaderRHI, InscatterTexture, InscatterTextureSampler, 
-				TStaticSamplerState<SF_Bilinear>::GetRHI(), View.AtmosphereInscatterTexture);
-		}
-	}
-
-	friend FArchive& operator<<(FArchive& Ar,FAtmosphereShaderTextureParameters& P);
-
-private:
-	FShaderResourceParameter TransmittanceTexture;
-	FShaderResourceParameter TransmittanceTextureSampler;
-	FShaderResourceParameter IrradianceTexture;
-	FShaderResourceParameter IrradianceTextureSampler;
-	FShaderResourceParameter InscatterTexture;
-	FShaderResourceParameter InscatterTextureSampler;
-};
-
 extern bool ShouldRenderAtmosphere(const FSceneViewFamily& Family);
\ No newline at end of file
diff --git a/Engine/Source/Runtime/Renderer/Private/AtmosphereTextures.cpp b/Engine/Source/Runtime/Renderer/Private/AtmosphereTextures.cpp
--- a/Engine/Source/Runtime/Renderer/Private/AtmosphereTextures.cpp
+++ b/Engine/Source/Runtime/Renderer/Private/AtmosphereTextures.cpp
@@ -65,6 +65,41 @@ void FAtmosphereTextures::ReleaseDynamicRHI()
 	GRenderTargetPool.FreeUnusedResources();
 }
 
+
+void FAtmosphereShaderTextureParameters::Bind(const FShaderParameterMap& ParameterMap)
+{
+	TransmittanceTexture.Bind(ParameterMap,TEXT("AtmosphereTransmittanceTexture"));
+	TransmittanceTextureSampler.Bind(ParameterMap,TEXT("AtmosphereTransmittanceTextureSampler"));
+	IrradianceTexture.Bind(ParameterMap,TEXT("AtmosphereIrradianceTexture"));
+	IrradianceTextureSampler.Bind(ParameterMap,TEXT("AtmosphereIrradianceTextureSampler"));
+	InscatterTexture.Bind(ParameterMap,TEXT("AtmosphereInscatterTexture"));
+	InscatterTextureSampler.Bind(ParameterMap,TEXT("AtmosphereInscatterTextureSampler"));
+}
+
+template< typename ShaderRHIParamRef >
+void FAtmosphereShaderTextureParameters::Set( FRHICommandList& RHICmdList, const ShaderRHIParamRef ShaderRHI, const FSceneView& View ) const
+{
+	if (TransmittanceTexture.IsBound() || IrradianceTexture.IsBound() || InscatterTexture.IsBound())
+	{
+		SetTextureParameter(RHICmdList, ShaderRHI, TransmittanceTexture, TransmittanceTextureSampler, 
+			TStaticSamplerState<SF_Bilinear>::GetRHI(), View.AtmosphereTransmittanceTexture);
+		SetTextureParameter(RHICmdList, ShaderRHI, IrradianceTexture, IrradianceTextureSampler, 
+			TStaticSamplerState<SF_Bilinear>::GetRHI(), View.AtmosphereIrradianceTexture);
+		SetTextureParameter(RHICmdList, ShaderRHI, InscatterTexture, InscatterTextureSampler, 
+			TStaticSamplerState<SF_Bilinear>::GetRHI(), View.AtmosphereInscatterTexture);
+	}
+}
+
+#define IMPLEMENT_ATMOSPHERE_TEXTURE_PARAM_SET( ShaderRHIParamRef ) \
+	template void FAtmosphereShaderTextureParameters::Set< ShaderRHIParamRef >( FRHICommandList& RHICmdList, const ShaderRHIParamRef ShaderRHI, const FSceneView& View ) const;
+
+IMPLEMENT_ATMOSPHERE_TEXTURE_PARAM_SET( FVertexShaderRHIParamRef );
+IMPLEMENT_ATMOSPHERE_TEXTURE_PARAM_SET( FHullShaderRHIParamRef );
+IMPLEMENT_ATMOSPHERE_TEXTURE_PARAM_SET( FDomainShaderRHIParamRef );
+IMPLEMENT_ATMOSPHERE_TEXTURE_PARAM_SET( FGeometryShaderRHIParamRef );
+IMPLEMENT_ATMOSPHERE_TEXTURE_PARAM_SET( FPixelShaderRHIParamRef );
+IMPLEMENT_ATMOSPHERE_TEXTURE_PARAM_SET( FComputeShaderRHIParamRef );
+
 FArchive& operator<<(FArchive& Ar,FAtmosphereShaderTextureParameters& Parameters)
 {
 	Ar << Parameters.TransmittanceTexture;
diff --git a/Engine/Source/Runtime/Renderer/Private/CompositionLighting/PostProcessAmbient.cpp b/Engine/Source/Runtime/Renderer/Private/CompositionLighting/PostProcessAmbient.cpp
--- a/Engine/Source/Runtime/Renderer/Private/CompositionLighting/PostProcessAmbient.cpp
+++ b/Engine/Source/Runtime/Renderer/Private/CompositionLighting/PostProcessAmbient.cpp
@@ -23,7 +23,6 @@ class FPostProcessAmbientPS : public FGlobalShader
 	static void ModifyCompilationEnvironment(EShaderPlatform Platform, FShaderCompilerEnvironment& OutEnvironment)
 	{
 		FGlobalShader::ModifyCompilationEnvironment(Platform,OutEnvironment);
-		FDeferredPixelShaderParameters::ModifyCompilationEnvironment(Platform,OutEnvironment);
 	}
 
 	/** Default constructor. */
diff --git a/Engine/Source/Runtime/Renderer/Private/CompositionLighting/PostProcessAmbientOcclusion.cpp b/Engine/Source/Runtime/Renderer/Private/CompositionLighting/PostProcessAmbientOcclusion.cpp
--- a/Engine/Source/Runtime/Renderer/Private/CompositionLighting/PostProcessAmbientOcclusion.cpp
+++ b/Engine/Source/Runtime/Renderer/Private/CompositionLighting/PostProcessAmbientOcclusion.cpp
@@ -34,7 +34,6 @@ class FPostProcessAmbientOcclusionSetupPS : public FGlobalShader
 	static void ModifyCompilationEnvironment(EShaderPlatform Platform, FShaderCompilerEnvironment& OutEnvironment)
 	{
 		FGlobalShader::ModifyCompilationEnvironment(Platform,OutEnvironment);
-		FDeferredPixelShaderParameters::ModifyCompilationEnvironment(Platform,OutEnvironment);
 		OutEnvironment.SetDefine(TEXT("INITIAL_PASS"), bInitialPass);
 	}
 
@@ -207,7 +206,6 @@ class FPostProcessAmbientOcclusionPS : public FGlobalShader
 	static void ModifyCompilationEnvironment(EShaderPlatform Platform, FShaderCompilerEnvironment& OutEnvironment)
 	{
 		FGlobalShader::ModifyCompilationEnvironment(Platform,OutEnvironment);
-		FDeferredPixelShaderParameters::ModifyCompilationEnvironment(Platform,OutEnvironment);
 
 		OutEnvironment.SetDefine(TEXT("USE_UPSAMPLE"), bDoUpsample);
 		OutEnvironment.SetDefine(TEXT("USE_AO_SETUP_AS_INPUT"), bTAOSetupAsInput);
@@ -301,7 +299,6 @@ class FPostProcessBasePassAOPS : public FGlobalShader
 	static void ModifyCompilationEnvironment(EShaderPlatform Platform, FShaderCompilerEnvironment& OutEnvironment)
 	{
 		FGlobalShader::ModifyCompilationEnvironment(Platform,OutEnvironment);
-		FDeferredPixelShaderParameters::ModifyCompilationEnvironment(Platform,OutEnvironment);
 	}
 
 	/** Default constructor. */
diff --git a/Engine/Source/Runtime/Renderer/Private/CompositionLighting/PostProcessDeferredDecals.cpp b/Engine/Source/Runtime/Renderer/Private/CompositionLighting/PostProcessDeferredDecals.cpp
--- a/Engine/Source/Runtime/Renderer/Private/CompositionLighting/PostProcessDeferredDecals.cpp
+++ b/Engine/Source/Runtime/Renderer/Private/CompositionLighting/PostProcessDeferredDecals.cpp
@@ -317,7 +317,6 @@ public:
 	static void ModifyCompilationEnvironment(EShaderPlatform Platform, FShaderCompilerEnvironment& OutEnvironment)
 	{
 		FGlobalShader::ModifyCompilationEnvironment(Platform,OutEnvironment);
-		FDeferredPixelShaderParameters::ModifyCompilationEnvironment(Platform,OutEnvironment);
 	}
 
 	FStencilDecalMaskPS(const ShaderMetaType::CompiledShaderInitializerType& Initializer):
diff --git a/Engine/Source/Runtime/Renderer/Private/LightRendering.cpp b/Engine/Source/Runtime/Renderer/Private/LightRendering.cpp
--- a/Engine/Source/Runtime/Renderer/Private/LightRendering.cpp
+++ b/Engine/Source/Runtime/Renderer/Private/LightRendering.cpp
@@ -35,7 +35,6 @@ public:
 	static void ModifyCompilationEnvironment(EShaderPlatform Platform, FShaderCompilerEnvironment& OutEnvironment)
 	{
 		FGlobalShader::ModifyCompilationEnvironment(Platform, OutEnvironment);
-		FDeferredPixelShaderParameters::ModifyCompilationEnvironment(Platform,OutEnvironment);
 		OutEnvironment.SetDefine(TEXT("USE_IES_PROFILE"), (uint32)bUseIESProfile);
 		OutEnvironment.SetDefine(TEXT("RADIAL_ATTENUATION"), (uint32)bRadialAttenuation);
 		OutEnvironment.SetDefine(TEXT("INVERSE_SQUARED_FALLOFF"), (uint32)bInverseSquaredFalloff);
@@ -175,7 +174,6 @@ public:
 	static void ModifyCompilationEnvironment(EShaderPlatform Platform, FShaderCompilerEnvironment& OutEnvironment)
 	{
 		FGlobalShader::ModifyCompilationEnvironment(Platform, OutEnvironment);
-		FDeferredPixelShaderParameters::ModifyCompilationEnvironment(Platform,OutEnvironment);
 		OutEnvironment.SetDefine(TEXT("RADIAL_ATTENUATION"), (uint32)bRadialAttenuation);
 	}
 
diff --git a/Engine/Source/Runtime/Renderer/Private/PostProcess/PostProcessBokehDOF.cpp b/Engine/Source/Runtime/Renderer/Private/PostProcess/PostProcessBokehDOF.cpp
--- a/Engine/Source/Runtime/Renderer/Private/PostProcess/PostProcessBokehDOF.cpp
+++ b/Engine/Source/Runtime/Renderer/Private/PostProcess/PostProcessBokehDOF.cpp
@@ -79,7 +79,6 @@ class PostProcessVisualizeDOFPS : public FGlobalShader
 	static void ModifyCompilationEnvironment(EShaderPlatform Platform, FShaderCompilerEnvironment& OutEnvironment)
 	{
 		FGlobalShader::ModifyCompilationEnvironment(Platform,OutEnvironment);
-		FDeferredPixelShaderParameters::ModifyCompilationEnvironment(Platform,OutEnvironment);
 	}
 
 	/** Default constructor. */
@@ -240,7 +239,6 @@ class PostProcessBokehDOFSetupPS : public FGlobalShader
 	static void ModifyCompilationEnvironment(EShaderPlatform Platform, FShaderCompilerEnvironment& OutEnvironment)
 	{
 		FGlobalShader::ModifyCompilationEnvironment(Platform,OutEnvironment);
-		FDeferredPixelShaderParameters::ModifyCompilationEnvironment(Platform,OutEnvironment);
 	}
 
 	/** Default constructor. */
@@ -399,7 +397,6 @@ class FPostProcessBokehDOFVS : public FGlobalShader
 	static void ModifyCompilationEnvironment(EShaderPlatform Platform, FShaderCompilerEnvironment& OutEnvironment)
 	{
 		FGlobalShader::ModifyCompilationEnvironment(Platform,OutEnvironment);
-		FDeferredPixelShaderParameters::ModifyCompilationEnvironment(Platform,OutEnvironment);
 		OutEnvironment.SetDefine(TEXT("DOF_METHOD"), DOFMethod);
 		OutEnvironment.SetDefine(TEXT("DOF_INDEX_STYLE"), DOFIndexStyle);
 	}
diff --git a/Engine/Source/Runtime/Renderer/Private/PostProcess/PostProcessBokehDOFRecombine.cpp b/Engine/Source/Runtime/Renderer/Private/PostProcess/PostProcessBokehDOFRecombine.cpp
--- a/Engine/Source/Runtime/Renderer/Private/PostProcess/PostProcessBokehDOFRecombine.cpp
+++ b/Engine/Source/Runtime/Renderer/Private/PostProcess/PostProcessBokehDOFRecombine.cpp
@@ -30,7 +30,6 @@ class FPostProcessBokehDOFRecombinePS : public FGlobalShader
 	static void ModifyCompilationEnvironment(EShaderPlatform Platform, FShaderCompilerEnvironment& OutEnvironment)
 	{
 		FGlobalShader::ModifyCompilationEnvironment(Platform, OutEnvironment);
-		FDeferredPixelShaderParameters::ModifyCompilationEnvironment(Platform,OutEnvironment);
 		OutEnvironment.SetDefine(TEXT("RECOMBINE_METHOD"), Method);
 	}
 
diff --git a/Engine/Source/Runtime/Renderer/Private/PostProcess/PostProcessDOF.cpp b/Engine/Source/Runtime/Renderer/Private/PostProcess/PostProcessDOF.cpp
--- a/Engine/Source/Runtime/Renderer/Private/PostProcess/PostProcessDOF.cpp
+++ b/Engine/Source/Runtime/Renderer/Private/PostProcess/PostProcessDOF.cpp
@@ -26,7 +26,6 @@ class FPostProcessDOFSetupPS : public FGlobalShader
 	static void ModifyCompilationEnvironment(EShaderPlatform Platform, FShaderCompilerEnvironment& OutEnvironment)
 	{
 		FGlobalShader::ModifyCompilationEnvironment(Platform,OutEnvironment);
-		FDeferredPixelShaderParameters::ModifyCompilationEnvironment(Platform,OutEnvironment);
 		OutEnvironment.SetDefine(TEXT("ENABLE_NEAR_BLUR"), NearBlurEnable);
 	}
 
@@ -208,7 +207,6 @@ class FPostProcessDOFRecombinePS : public FGlobalShader
 	static void ModifyCompilationEnvironment(EShaderPlatform Platform, FShaderCompilerEnvironment& OutEnvironment)
 	{
 		FGlobalShader::ModifyCompilationEnvironment(Platform,OutEnvironment);
-		FDeferredPixelShaderParameters::ModifyCompilationEnvironment(Platform,OutEnvironment);
 		OutEnvironment.SetDefine(TEXT("ENABLE_NEAR_BLUR"), NearBlurEnable);
 	}
 
diff --git a/Engine/Source/Runtime/Renderer/Private/PostProcess/PostProcessDownsample.cpp b/Engine/Source/Runtime/Renderer/Private/PostProcess/PostProcessDownsample.cpp
--- a/Engine/Source/Runtime/Renderer/Private/PostProcess/PostProcessDownsample.cpp
+++ b/Engine/Source/Runtime/Renderer/Private/PostProcess/PostProcessDownsample.cpp
@@ -25,7 +25,6 @@ class FPostProcessDownsamplePS : public FGlobalShader
 	static void ModifyCompilationEnvironment(EShaderPlatform Platform, FShaderCompilerEnvironment& OutEnvironment)
 	{
 		FGlobalShader::ModifyCompilationEnvironment(Platform,OutEnvironment);
-		FDeferredPixelShaderParameters::ModifyCompilationEnvironment(Platform,OutEnvironment);
 		OutEnvironment.SetDefine(TEXT("METHOD"), Method);
 	}
 
diff --git a/Engine/Source/Runtime/Renderer/Private/PostProcess/PostProcessGBufferHints.cpp b/Engine/Source/Runtime/Renderer/Private/PostProcess/PostProcessGBufferHints.cpp
--- a/Engine/Source/Runtime/Renderer/Private/PostProcess/PostProcessGBufferHints.cpp
+++ b/Engine/Source/Runtime/Renderer/Private/PostProcess/PostProcessGBufferHints.cpp
@@ -25,7 +25,6 @@ class FPostProcessGBufferHintsPS : public FGlobalShader
 	static void ModifyCompilationEnvironment(EShaderPlatform Platform, FShaderCompilerEnvironment& OutEnvironment)
 	{
 		FGlobalShader::ModifyCompilationEnvironment(Platform,OutEnvironment);
-		FDeferredPixelShaderParameters::ModifyCompilationEnvironment(Platform,OutEnvironment);
 	}
 
 	/** Default constructor. */
diff --git a/Engine/Source/Runtime/Renderer/Private/PostProcess/PostProcessHMD.cpp b/Engine/Source/Runtime/Renderer/Private/PostProcess/PostProcessHMD.cpp
--- a/Engine/Source/Runtime/Renderer/Private/PostProcess/PostProcessHMD.cpp
+++ b/Engine/Source/Runtime/Renderer/Private/PostProcess/PostProcessHMD.cpp
@@ -146,7 +146,6 @@ public:
 	{
 		FGlobalShader::ModifyCompilationEnvironment(Platform, OutEnvironment);
 		OutEnvironment.SetDefine(TEXT("USE_TIMEWARP"), uint32(0));
-		FDeferredPixelShaderParameters::ModifyCompilationEnvironment(Platform,OutEnvironment);
 	}
 };
 
diff --git a/Engine/Source/Runtime/Renderer/Private/PostProcess/PostProcessMotionBlur.cpp b/Engine/Source/Runtime/Renderer/Private/PostProcess/PostProcessMotionBlur.cpp
--- a/Engine/Source/Runtime/Renderer/Private/PostProcess/PostProcessMotionBlur.cpp
+++ b/Engine/Source/Runtime/Renderer/Private/PostProcess/PostProcessMotionBlur.cpp
@@ -220,7 +220,6 @@ class FPostProcessMotionBlurPS : public FGlobalShader
 	static void ModifyCompilationEnvironment(EShaderPlatform Platform, FShaderCompilerEnvironment& OutEnvironment)
 	{
 		FGlobalShader::ModifyCompilationEnvironment(Platform,OutEnvironment);
-		FDeferredPixelShaderParameters::ModifyCompilationEnvironment(Platform,OutEnvironment);
 		OutEnvironment.SetDefine(TEXT("MOTION_BLUR_QUALITY"), Quality);
 	}
 
diff --git a/Engine/Source/Runtime/Renderer/Private/PostProcess/PostProcessNoiseBlur.cpp b/Engine/Source/Runtime/Renderer/Private/PostProcess/PostProcessNoiseBlur.cpp
--- a/Engine/Source/Runtime/Renderer/Private/PostProcess/PostProcessNoiseBlur.cpp
+++ b/Engine/Source/Runtime/Renderer/Private/PostProcess/PostProcessNoiseBlur.cpp
@@ -24,7 +24,6 @@ class FPostProcessNoiseBlurPS : public FGlobalShader
 	static void ModifyCompilationEnvironment(EShaderPlatform Platform, FShaderCompilerEnvironment& OutEnvironment)
 	{
 		FGlobalShader::ModifyCompilationEnvironment(Platform,OutEnvironment);
-		FDeferredPixelShaderParameters::ModifyCompilationEnvironment(Platform,OutEnvironment);
 		OutEnvironment.SetDefine(TEXT("METHOD"), Method);
 	}
 
diff --git a/Engine/Source/Runtime/Renderer/Private/PostProcess/PostProcessSubsurface.cpp b/Engine/Source/Runtime/Renderer/Private/PostProcess/PostProcessSubsurface.cpp
--- a/Engine/Source/Runtime/Renderer/Private/PostProcess/PostProcessSubsurface.cpp
+++ b/Engine/Source/Runtime/Renderer/Private/PostProcess/PostProcessSubsurface.cpp
@@ -214,7 +214,6 @@ class FPostProcessSubsurfacePS : public FGlobalShader
 	static void ModifyCompilationEnvironment(EShaderPlatform Platform, FShaderCompilerEnvironment& OutEnvironment)
 	{
 		FGlobalShader::ModifyCompilationEnvironment(Platform,OutEnvironment);
-		FDeferredPixelShaderParameters::ModifyCompilationEnvironment(Platform,OutEnvironment);
 		OutEnvironment.SetDefine(TEXT("METHOD"), Method);
 	}
 
diff --git a/Engine/Source/Runtime/Renderer/Private/PostProcess/PostProcessTemporalAA.cpp b/Engine/Source/Runtime/Renderer/Private/PostProcess/PostProcessTemporalAA.cpp
--- a/Engine/Source/Runtime/Renderer/Private/PostProcess/PostProcessTemporalAA.cpp
+++ b/Engine/Source/Runtime/Renderer/Private/PostProcess/PostProcessTemporalAA.cpp
@@ -52,8 +52,6 @@ class FPostProcessTemporalAAPS : public FGlobalShader
 	static void ModifyCompilationEnvironment(EShaderPlatform Platform, FShaderCompilerEnvironment& OutEnvironment)
 	{
 		FGlobalShader::ModifyCompilationEnvironment(Platform, OutEnvironment);
-		FDeferredPixelShaderParameters::ModifyCompilationEnvironment(Platform,OutEnvironment);
-
 		OutEnvironment.SetDefine( TEXT("RESPONSIVE"), Responsive );
 	}
 
diff --git a/Engine/Source/Runtime/Renderer/Private/PostProcess/PostProcessTestImage.cpp b/Engine/Source/Runtime/Renderer/Private/PostProcess/PostProcessTestImage.cpp
--- a/Engine/Source/Runtime/Renderer/Private/PostProcess/PostProcessTestImage.cpp
+++ b/Engine/Source/Runtime/Renderer/Private/PostProcess/PostProcessTestImage.cpp
@@ -26,7 +26,6 @@ class FPostProcessTestImagePS : public FGlobalShader
 	static void ModifyCompilationEnvironment(EShaderPlatform Platform, FShaderCompilerEnvironment& OutEnvironment)
 	{
 		FGlobalShader::ModifyCompilationEnvironment(Platform,OutEnvironment);
-		FDeferredPixelShaderParameters::ModifyCompilationEnvironment(Platform,OutEnvironment);
 	}
 
 	/** Default constructor. */
diff --git a/Engine/Source/Runtime/Renderer/Private/PostProcess/PostProcessUpscale.cpp b/Engine/Source/Runtime/Renderer/Private/PostProcess/PostProcessUpscale.cpp
--- a/Engine/Source/Runtime/Renderer/Private/PostProcess/PostProcessUpscale.cpp
+++ b/Engine/Source/Runtime/Renderer/Private/PostProcess/PostProcessUpscale.cpp
@@ -26,7 +26,6 @@ class FPostProcessUpscalePS : public FGlobalShader
 	static void ModifyCompilationEnvironment(EShaderPlatform Platform, FShaderCompilerEnvironment& OutEnvironment)
 	{
 		FGlobalShader::ModifyCompilationEnvironment(Platform,OutEnvironment);
-		FDeferredPixelShaderParameters::ModifyCompilationEnvironment(Platform,OutEnvironment);
 		OutEnvironment.SetDefine(TEXT("METHOD"), Method);
 	}
 
diff --git a/Engine/Source/Runtime/Renderer/Private/PostProcess/PostProcessVisualizeBuffer.cpp b/Engine/Source/Runtime/Renderer/Private/PostProcess/PostProcessVisualizeBuffer.cpp
--- a/Engine/Source/Runtime/Renderer/Private/PostProcess/PostProcessVisualizeBuffer.cpp
+++ b/Engine/Source/Runtime/Renderer/Private/PostProcess/PostProcessVisualizeBuffer.cpp
@@ -26,7 +26,6 @@ class FPostProcessVisualizeBufferPS : public FGlobalShader
 	static void ModifyCompilationEnvironment(EShaderPlatform Platform, FShaderCompilerEnvironment& OutEnvironment)
 	{
 		FGlobalShader::ModifyCompilationEnvironment(Platform, OutEnvironment);
-		FDeferredPixelShaderParameters::ModifyCompilationEnvironment(Platform,OutEnvironment);
 		OutEnvironment.SetDefine(TEXT("DRAWING_TILE"), bDrawingTile ? TEXT("1") : TEXT("0"));
 	}
 
diff --git a/Engine/Source/Runtime/Renderer/Private/PostProcess/RenderingCompositionGraph.cpp b/Engine/Source/Runtime/Renderer/Private/PostProcess/RenderingCompositionGraph.cpp
--- a/Engine/Source/Runtime/Renderer/Private/PostProcess/RenderingCompositionGraph.cpp
+++ b/Engine/Source/Runtime/Renderer/Private/PostProcess/RenderingCompositionGraph.cpp
@@ -680,6 +680,165 @@ void FPostProcessPassParameters::SetVS(const FVertexShaderRHIParamRef& ShaderRHI
 	Set(Context.RHICmdList, ShaderRHI, Context, Filter, bWhiteIfNoTexture, FilterOverrideArray);
 }
 
+template< typename ShaderRHIParamRef >
+void FPostProcessPassParameters::Set(
+	FRHICommandList& RHICmdList,
+	const ShaderRHIParamRef& ShaderRHI,
+	const FRenderingCompositePassContext& Context,
+	FSamplerStateRHIParamRef Filter,
+	bool bWhiteIfNoTexture,
+	FSamplerStateRHIParamRef* FilterOverrideArray)
+{
+	// assuming all outputs have the same size
+	FRenderingCompositeOutput* Output = Context.Pass->GetOutput(ePId_Output0);
+
+	// Output0 should always exist
+	check(Output);
+
+	// one should be on
+	check(FilterOverrideArray || Filter);
+	// but not both
+	check(!FilterOverrideArray || !Filter);
+
+	RHICmdList.CheckIsNull();
+
+	if(BilinearTextureSampler0.IsBound())
+	{
+		RHISetShaderSampler(
+			ShaderRHI, 
+			BilinearTextureSampler0.GetBaseIndex(), 
+			TStaticSamplerState<SF_Bilinear>::GetRHI()
+			);
+	}
+
+	if(BilinearTextureSampler1.IsBound())
+	{
+		RHISetShaderSampler(
+			ShaderRHI, 
+			BilinearTextureSampler1.GetBaseIndex(), 
+			TStaticSamplerState<SF_Bilinear>::GetRHI()
+			);
+	}
+
+	if(ViewportSize.IsBound() || ScreenPosToPixel.IsBound() || ViewportRect.IsBound())
+	{
+		FIntRect LocalViewport = Context.GetViewport();
+
+		FIntPoint ViewportOffset = LocalViewport.Min;
+		FIntPoint ViewportExtent = LocalViewport.Size();
+
+		{
+			FVector4 Value(ViewportExtent.X, ViewportExtent.Y, 1.0f / ViewportExtent.X, 1.0f / ViewportExtent.Y);
+
+			SetShaderValue(RHICmdList, ShaderRHI, ViewportSize, Value);
+		}
+
+		{
+			FVector4 Value(LocalViewport.Min.X, LocalViewport.Min.Y, LocalViewport.Max.X, LocalViewport.Max.Y);
+
+			SetShaderValue(RHICmdList, ShaderRHI, ViewportRect, Value);
+		}
+
+		{
+			FVector4 ScreenPosToPixelValue(
+				ViewportExtent.X * 0.5f,
+				-ViewportExtent.Y * 0.5f, 
+				ViewportExtent.X * 0.5f - 0.5f + ViewportOffset.X,
+				ViewportExtent.Y * 0.5f - 0.5f + ViewportOffset.Y);
+			SetShaderValue(RHICmdList, ShaderRHI, ScreenPosToPixel, ScreenPosToPixelValue);
+		}
+	}
+
+	//Calculate a base scene texture min max which will be pulled in by a pixel for each PP input.
+	FIntRect ContextViewportRect = Context.IsViewportValid() ? Context.GetViewport() : FIntRect(0,0,0,0);
+	const FIntPoint SceneRTSize = GSceneRenderTargets.GetBufferSizeXY();
+	FVector4 BaseSceneTexMinMax(	((float)ContextViewportRect.Min.X/SceneRTSize.X), 
+									((float)ContextViewportRect.Min.Y/SceneRTSize.Y), 
+									((float)ContextViewportRect.Max.X/SceneRTSize.X), 
+									((float)ContextViewportRect.Max.Y/SceneRTSize.Y) );
+
+	// ePId_Input0, ePId_Input1, ...
+	for(uint32 Id = 0; Id < (uint32)ePId_Input_MAX; ++Id)
+	{
+		FRenderingCompositeOutputRef* OutputRef = Context.Pass->GetInput((EPassInputId)Id);
+
+		if(!OutputRef)
+		{
+			// Pass doesn't have more inputs
+			break;
+		}
+
+		const auto FeatureLevel = Context.GetFeatureLevel();
+
+		FRenderingCompositeOutput* Input = OutputRef->GetOutput();
+
+		TRefCountPtr<IPooledRenderTarget> InputPooledElement;
+
+		if(Input)
+		{
+			InputPooledElement = Input->RequestInput();
+		}
+
+		FSamplerStateRHIParamRef LocalFilter = FilterOverrideArray ? FilterOverrideArray[Id] : Filter;
+
+		if(InputPooledElement)
+		{
+			check(!InputPooledElement->IsFree());
+
+			const FTextureRHIRef& SrcTexture = InputPooledElement->GetRenderTargetItem().ShaderResourceTexture;
+
+			SetTextureParameter(RHICmdList, ShaderRHI, PostprocessInputParameter[Id], PostprocessInputParameterSampler[Id], LocalFilter, SrcTexture);
+			SetTextureParameter(RHICmdList, ShaderRHI, PostprocessInputNew[Id], InputPooledElement->GetRenderTargetItem().TargetableTexture);
+
+			{
+				float Width = InputPooledElement->GetDesc().Extent.X;
+				float Height = InputPooledElement->GetDesc().Extent.Y;
+				FVector4 TextureSize(Width, Height, 1.0f / Width, 1.0f / Height);
+				SetShaderValue(RHICmdList, ShaderRHI, PostprocessInputSizeParameter[Id], TextureSize);
+					
+				//We could use the main scene min max here if it weren't that we need to pull the max in by a pixel on a per input basis.
+				FVector2D OnePPInputPixelUVSize = FVector2D(1.0f / Width, 1.0f / Height);
+				FVector4 PPInputMinMax = BaseSceneTexMinMax;
+				PPInputMinMax.Z -= OnePPInputPixelUVSize.X;
+				PPInputMinMax.W -= OnePPInputPixelUVSize.Y;
+				SetShaderValue(RHICmdList, ShaderRHI, PostProcessInputMinMaxParameter[Id], PPInputMinMax);
+			}
+		}
+		else
+		{
+			IPooledRenderTarget* Texture = bWhiteIfNoTexture ? GSystemTextures.WhiteDummy : GSystemTextures.BlackDummy;
+
+			// if the input is not there but the shader request it we give it at least some data to avoid d3ddebug errors and shader permutations
+			// to make features optional we use default black for additive passes without shader permutations
+			SetTextureParameter(RHICmdList, ShaderRHI, PostprocessInputParameter[Id], PostprocessInputParameterSampler[Id], LocalFilter, Texture->GetRenderTargetItem().TargetableTexture);
+			SetTextureParameter(RHICmdList, ShaderRHI, PostprocessInputNew[Id], Texture->GetRenderTargetItem().TargetableTexture);
+
+			FVector4 Dummy(1, 1, 1, 1);
+			SetShaderValue(RHICmdList, ShaderRHI, PostprocessInputSizeParameter[Id], Dummy);
+			SetShaderValue(RHICmdList, ShaderRHI, PostProcessInputMinMaxParameter[Id], Dummy);
+		}
+	}
+
+	// todo warning if Input[] or InputSize[] is bound but not available, maybe set a specific input texture (blinking?)
+}
+
+#define IMPLEMENT_POST_PROCESS_PARAM_SET( ShaderRHIParamRef ) \
+	template void FPostProcessPassParameters::Set< ShaderRHIParamRef >( \
+		FRHICommandList& RHICmdList,					\
+		const ShaderRHIParamRef& ShaderRHI,				\
+		const FRenderingCompositePassContext& Context,	\
+		FSamplerStateRHIParamRef Filter,				\
+		bool bWhiteIfNoTexture,							\
+		FSamplerStateRHIParamRef* FilterOverrideArray	\
+	);
+
+IMPLEMENT_POST_PROCESS_PARAM_SET( FVertexShaderRHIParamRef );
+IMPLEMENT_POST_PROCESS_PARAM_SET( FHullShaderRHIParamRef );
+IMPLEMENT_POST_PROCESS_PARAM_SET( FDomainShaderRHIParamRef );
+IMPLEMENT_POST_PROCESS_PARAM_SET( FGeometryShaderRHIParamRef );
+IMPLEMENT_POST_PROCESS_PARAM_SET( FPixelShaderRHIParamRef );
+IMPLEMENT_POST_PROCESS_PARAM_SET( FComputeShaderRHIParamRef );
+
 FArchive& operator<<(FArchive& Ar, FPostProcessPassParameters& P)
 {
 	Ar << P.BilinearTextureSampler0 << P.BilinearTextureSampler1 << P.ViewportSize << P.ScreenPosToPixel << P.ViewportRect;
diff --git a/Engine/Source/Runtime/Renderer/Private/PostProcess/RenderingCompositionGraph.h b/Engine/Source/Runtime/Renderer/Private/PostProcess/RenderingCompositionGraph.h
--- a/Engine/Source/Runtime/Renderer/Private/PostProcess/RenderingCompositionGraph.h
+++ b/Engine/Source/Runtime/Renderer/Private/PostProcess/RenderingCompositionGraph.h
@@ -16,49 +16,6 @@ class FRenderingCompositionGraph;
 class FViewInfo;
 class FShaderParameterMap;
 
-// This is the index for the texture input of this pass.
-// More that that should not be needed.
-// Could be an uint32 but for better readability and type safety it's an enum.
-// Counting starts from 0 in consecutive order.
-enum EPassInputId
-{
-	ePId_Input0,
-	ePId_Input1,
-	ePId_Input2,
-	ePId_Input3,
-	ePId_Input4,
-	ePId_Input5,
-	ePId_Input6,
-/*	ePId_Input7,
-	ePId_Input8,
-	ePId_Input9,
-	ePId_Input10,
-	ePId_Input11,
-	ePId_Input12,
-	ePId_Input13,
-	ePId_Input14,
-	ePId_Input15*/
-
-	// to get the total count of inputs
-	ePId_Input_MAX
-};
-
-// Usually the same as the MRT number but it doesn't have to be implemented as MRT.
-// More that that should not be needed.
-// Could be an uint32 but for better readability and type safety it's an enum.
-// Counting starts from 0 in consecutive order.
-enum EPassOutputId
-{
-	ePId_Output0,
-	ePId_Output1,
-	ePId_Output2,
-	ePId_Output3,
-	ePId_Output4,
-	ePId_Output5,
-	ePId_Output6,
-	ePId_Output7
-};
-
 struct FRenderingCompositeOutputRef;
 struct FRenderingCompositePassContext;
 
@@ -516,177 +473,4 @@ public:
 	FPooledRenderTargetDesc ComputeOutputDesc(EPassOutputId InPassOutputId) const { FPooledRenderTargetDesc Desc; Desc.DebugName = TEXT("Root"); return Desc; }
 };
 
-// currently hard coded to 4 input textures
-// convenient but not the most optimized solution
-struct FPostProcessPassParameters 
-{
-	/** Initialization constructor. */
-	void Bind(const FShaderParameterMap& ParameterMap);
-
-	/** Set the pixel shader parameter values. */
-	void SetPS(const FPixelShaderRHIParamRef& ShaderRHI, const FRenderingCompositePassContext& Context, FSamplerStateRHIParamRef Filter = TStaticSamplerState<>::GetRHI(), bool bWhiteIfNoTexture = false, FSamplerStateRHIParamRef* FilterOverrideArray = 0);
-
-	/** Set the compute shader parameter values. */
-	void SetCS(const FComputeShaderRHIParamRef& ShaderRHI, const FRenderingCompositePassContext& Context, FSamplerStateRHIParamRef Filter = TStaticSamplerState<>::GetRHI(), bool bWhiteIfNoTexture = false, FSamplerStateRHIParamRef* FilterOverrideArray = 0);
-
-	/** Set the vertex shader parameter values. */
-	void SetVS(const FVertexShaderRHIParamRef& ShaderRHI, const FRenderingCompositePassContext& Context, FSamplerStateRHIParamRef Filter = TStaticSamplerState<>::GetRHI(), bool bWhiteIfNoTexture = false, FSamplerStateRHIParamRef* FilterOverrideArray = 0);
-
-	/** Serializer. */
-	friend FArchive& operator<<(FArchive& Ar,FPostProcessPassParameters& P);
-
-private:
-
-	FShaderParameter ViewportSize;
-	FShaderParameter ViewportRect;
-	FShaderResourceParameter PostprocessInputParameter[ePId_Input_MAX];
-	FShaderResourceParameter PostprocessInputParameterSampler[ePId_Input_MAX];
-	FShaderResourceParameter PostprocessInputNew[ePId_Input_MAX];
-	FShaderParameter PostprocessInputSizeParameter[ePId_Input_MAX];
-	FShaderParameter PostProcessInputMinMaxParameter[ePId_Input_MAX];
-	FShaderParameter ScreenPosToPixel;
-	FShaderResourceParameter BilinearTextureSampler0;
-	FShaderResourceParameter BilinearTextureSampler1;
-
-public:
-	// using this template method in many places the executable size might go up
-	// @param Filter can be 0 if FilterOverrideArray is used
-	// @param FilterOverrideArray can be 0 if Filter is used
-	template <class T>
-	void Set(FRHICommandList& RHICmdList, const T& ShaderRHI, const FRenderingCompositePassContext& Context, FSamplerStateRHIParamRef Filter, bool bWhiteIfNoTexture = false, FSamplerStateRHIParamRef* FilterOverrideArray = 0)
-	{
-		// assuming all outputs have the same size
-		FRenderingCompositeOutput* Output = Context.Pass->GetOutput(ePId_Output0);
-
-		// Output0 should always exist
-		check(Output);
-
-		// one should be on
-		check(FilterOverrideArray || Filter);
-		// but not both
-		check(!FilterOverrideArray || !Filter);
-
-		RHICmdList.CheckIsNull();
-
-		if(BilinearTextureSampler0.IsBound())
-		{
-			RHISetShaderSampler(
-				ShaderRHI, 
-				BilinearTextureSampler0.GetBaseIndex(), 
-				TStaticSamplerState<SF_Bilinear>::GetRHI()
-				);
-		}
-
-		if(BilinearTextureSampler1.IsBound())
-		{
-			RHISetShaderSampler(
-				ShaderRHI, 
-				BilinearTextureSampler1.GetBaseIndex(), 
-				TStaticSamplerState<SF_Bilinear>::GetRHI()
-				);
-		}
-
-		if(ViewportSize.IsBound() || ScreenPosToPixel.IsBound() || ViewportRect.IsBound())
-		{
-			FIntRect LocalViewport = Context.GetViewport();
-
-			FIntPoint ViewportOffset = LocalViewport.Min;
-			FIntPoint ViewportExtent = LocalViewport.Size();
-
-			{
-				FVector4 Value(ViewportExtent.X, ViewportExtent.Y, 1.0f / ViewportExtent.X, 1.0f / ViewportExtent.Y);
-
-				SetShaderValue(RHICmdList, ShaderRHI, ViewportSize, Value);
-			}
-
-			{
-				FVector4 Value(LocalViewport.Min.X, LocalViewport.Min.Y, LocalViewport.Max.X, LocalViewport.Max.Y);
-
-				SetShaderValue(RHICmdList, ShaderRHI, ViewportRect, Value);
-			}
-
-			{
-				FVector4 ScreenPosToPixelValue(
-					ViewportExtent.X * 0.5f,
-					-ViewportExtent.Y * 0.5f, 
-					ViewportExtent.X * 0.5f - 0.5f + ViewportOffset.X,
-					ViewportExtent.Y * 0.5f - 0.5f + ViewportOffset.Y);
-				SetShaderValue(RHICmdList, ShaderRHI, ScreenPosToPixel, ScreenPosToPixelValue);
-			}
-		}
-
-		//Calculate a base scene texture min max which will be pulled in by a pixel for each PP input.
-		FIntRect ContextViewportRect = Context.IsViewportValid() ? Context.GetViewport() : FIntRect(0,0,0,0);
-		const FIntPoint SceneRTSize = GSceneRenderTargets.GetBufferSizeXY();
-		FVector4 BaseSceneTexMinMax(	((float)ContextViewportRect.Min.X/SceneRTSize.X), 
-										((float)ContextViewportRect.Min.Y/SceneRTSize.Y), 
-										((float)ContextViewportRect.Max.X/SceneRTSize.X), 
-										((float)ContextViewportRect.Max.Y/SceneRTSize.Y) );
-
-		// ePId_Input0, ePId_Input1, ...
-		for(uint32 Id = 0; Id < (uint32)ePId_Input_MAX; ++Id)
-		{
-			FRenderingCompositeOutputRef* OutputRef = Context.Pass->GetInput((EPassInputId)Id);
-
-			if(!OutputRef)
-			{
-				// Pass doesn't have more inputs
-				break;
-			}
-
-			const auto FeatureLevel = Context.GetFeatureLevel();
-
-			FRenderingCompositeOutput* Input = OutputRef->GetOutput();
-
-			TRefCountPtr<IPooledRenderTarget> InputPooledElement;
-
-			if(Input)
-			{
-				InputPooledElement = Input->RequestInput();
-			}
-
-			FSamplerStateRHIParamRef LocalFilter = FilterOverrideArray ? FilterOverrideArray[Id] : Filter;
-
-			if(InputPooledElement)
-			{
-				check(!InputPooledElement->IsFree());
-
-				const FTextureRHIRef& SrcTexture = InputPooledElement->GetRenderTargetItem().ShaderResourceTexture;
-
-				SetTextureParameter(RHICmdList, ShaderRHI, PostprocessInputParameter[Id], PostprocessInputParameterSampler[Id], LocalFilter, SrcTexture);
-				SetTextureParameter(RHICmdList, ShaderRHI, PostprocessInputNew[Id], InputPooledElement->GetRenderTargetItem().TargetableTexture);
-
-				{
-					float Width = InputPooledElement->GetDesc().Extent.X;
-					float Height = InputPooledElement->GetDesc().Extent.Y;
-					FVector4 TextureSize(Width, Height, 1.0f / Width, 1.0f / Height);
-					SetShaderValue(RHICmdList, ShaderRHI, PostprocessInputSizeParameter[Id], TextureSize);
-					
-					//We could use the main scene min max here if it weren't that we need to pull the max in by a pixel on a per input basis.
-					FVector2D OnePPInputPixelUVSize = FVector2D(1.0f / Width, 1.0f / Height);
-					FVector4 PPInputMinMax = BaseSceneTexMinMax;
-					PPInputMinMax.Z -= OnePPInputPixelUVSize.X;
-					PPInputMinMax.W -= OnePPInputPixelUVSize.Y;
-					SetShaderValue(RHICmdList, ShaderRHI, PostProcessInputMinMaxParameter[Id], PPInputMinMax);
-				}
-			}
-			else
-			{
-				IPooledRenderTarget* Texture = bWhiteIfNoTexture ? GSystemTextures.WhiteDummy : GSystemTextures.BlackDummy;
-
-				// if the input is not there but the shader request it we give it at least some data to avoid d3ddebug errors and shader permutations
-				// to make features optional we use default black for additive passes without shader permutations
-				SetTextureParameter(RHICmdList, ShaderRHI, PostprocessInputParameter[Id], PostprocessInputParameterSampler[Id], LocalFilter, Texture->GetRenderTargetItem().TargetableTexture);
-				SetTextureParameter(RHICmdList, ShaderRHI, PostprocessInputNew[Id], Texture->GetRenderTargetItem().TargetableTexture);
-
-				FVector4 Dummy(1, 1, 1, 1);
-				SetShaderValue(RHICmdList, ShaderRHI, PostprocessInputSizeParameter[Id], Dummy);
-				SetShaderValue(RHICmdList, ShaderRHI, PostProcessInputMinMaxParameter[Id], Dummy);
-			}
-		}
-
-		// todo warning if Input[] or InputSize[] is bound but not available, maybe set a specific input texture (blinking?)
-	}
-};
-
 void CompositionGraph_OnStartFrame();
diff --git a/Engine/Source/Runtime/Renderer/Private/PostProcess/SceneRenderTargets.cpp b/Engine/Source/Runtime/Renderer/Private/PostProcess/SceneRenderTargets.cpp
--- a/Engine/Source/Runtime/Renderer/Private/PostProcess/SceneRenderTargets.cpp
+++ b/Engine/Source/Runtime/Renderer/Private/PostProcess/SceneRenderTargets.cpp
@@ -1489,6 +1489,151 @@ void FSceneTextureShaderParameters::Bind(const FShaderParameterMap& ParameterMap
 	SceneDepthSurfaceParameter.Bind(ParameterMap,TEXT("SceneDepthSurface"));
 }
 
+template< typename ShaderRHIParamRef >
+void FSceneTextureShaderParameters::Set(
+	FRHICommandList& RHICmdList,
+	const ShaderRHIParamRef& ShaderRHI,
+	const FSceneView& View,
+	ESceneRenderTargetsMode::Type TextureMode,
+	ESamplerFilter ColorFilter ) const
+{
+	if (TextureMode == ESceneRenderTargetsMode::SetTextures)
+	{
+		// optimization possible: TShaderRHIParamRef is no param Ref
+		if (SceneColorTextureParameter.IsBound())
+		{
+			FSamplerStateRHIRef Filter;
+			switch ( ColorFilter )
+			{
+			case SF_Bilinear:
+				Filter = TStaticSamplerState<SF_Bilinear,AM_Clamp,AM_Clamp,AM_Clamp>::GetRHI();
+				break;
+			case SF_Trilinear:
+				Filter = TStaticSamplerState<SF_Trilinear,AM_Clamp,AM_Clamp,AM_Clamp>::GetRHI();
+				break;
+			case SF_AnisotropicPoint:
+				Filter = TStaticSamplerState<SF_AnisotropicPoint,AM_Clamp,AM_Clamp,AM_Clamp>::GetRHI();
+				break;
+			case SF_AnisotropicLinear:
+				Filter = TStaticSamplerState<SF_AnisotropicLinear,AM_Clamp,AM_Clamp,AM_Clamp>::GetRHI();
+				break;
+			case SF_Point:
+			default:
+				Filter = TStaticSamplerState<SF_Point,AM_Clamp,AM_Clamp,AM_Clamp>::GetRHI();
+				break;
+			}
+
+			SetTextureParameter(
+				RHICmdList, 
+				ShaderRHI,
+				SceneColorTextureParameter,
+				SceneColorTextureParameterSampler,
+				Filter,
+				GSceneRenderTargets.GetSceneColorTexture()
+				);
+		}
+
+		if (SceneAlphaCopyTextureParameter.IsBound() && GSceneRenderTargets.HasSceneAlphaCopyTexture())
+		{
+			FSamplerStateRHIRef Filter;
+			Filter = TStaticSamplerState<SF_Point,AM_Clamp,AM_Clamp,AM_Clamp>::GetRHI();
+			SetTextureParameter(
+				RHICmdList, 
+				ShaderRHI,
+				SceneAlphaCopyTextureParameter,
+				SceneAlphaCopyTextureParameterSampler,
+				Filter,
+				GSceneRenderTargets.GetSceneAlphaCopyTexture()
+				);
+		}
+
+		if(SceneDepthTextureParameter.IsBound() || SceneDepthTextureParameterSampler.IsBound())
+		{
+			const FTexture2DRHIRef* DepthTexture = GSceneRenderTargets.GetActualDepthTexture();
+			SetTextureParameter(
+				RHICmdList, 
+				ShaderRHI,
+				SceneDepthTextureParameter,
+				SceneDepthTextureParameterSampler,
+				TStaticSamplerState<SF_Point,AM_Clamp,AM_Clamp,AM_Clamp>::GetRHI(),
+				*DepthTexture
+				);
+		}
+
+		if (GRHIFeatureLevel >= ERHIFeatureLevel::SM5)
+		{
+			SetTextureParameter(RHICmdList, ShaderRHI, SceneColorSurfaceParameter, GSceneRenderTargets.GetSceneColorSurface());
+		}
+		if(GRHIFeatureLevel >= ERHIFeatureLevel::SM4)
+		{
+			if(GSupportsDepthFetchDuringDepthTest)
+			{
+				if(SceneDepthSurfaceParameter.IsBound())
+				{
+					SetTextureParameter(RHICmdList, ShaderRHI, SceneDepthSurfaceParameter, GSceneRenderTargets.GetSceneDepthSurface());
+				}
+				if(SceneDepthTextureNonMS.IsBound())
+				{
+					SetTextureParameter(RHICmdList, ShaderRHI, SceneDepthTextureNonMS, GSceneRenderTargets.GetSceneDepthTexture());
+				}
+			}
+			else
+			{
+				if(SceneDepthSurfaceParameter.IsBound())
+				{
+					SetTextureParameter(RHICmdList, ShaderRHI, SceneDepthSurfaceParameter, GSceneRenderTargets.GetAuxiliarySceneDepthSurface());
+				}
+				if(SceneDepthTextureNonMS.IsBound())
+				{
+					SetTextureParameter(RHICmdList, ShaderRHI, SceneDepthTextureNonMS, GSceneRenderTargets.GetAuxiliarySceneDepthSurface());
+				}
+			}
+		}
+	}
+	else if (TextureMode == ESceneRenderTargetsMode::DontSet)
+	{
+		// Verify that none of these were bound if we were told not to set them
+		checkSlow(!SceneColorTextureParameter.IsBound()
+			&& !SceneDepthTextureParameter.IsBound()
+			&& !SceneColorSurfaceParameter.IsBound()
+			&& !SceneDepthSurfaceParameter.IsBound()
+			&& !SceneDepthTextureNonMS.IsBound());
+	}
+	else if (TextureMode == ESceneRenderTargetsMode::DontSetIgnoreBoundByEditorCompositing)
+	{
+		// Verify that none of these were bound if we were told not to set them
+		// ignore SceneDepthTextureNonMS
+		checkSlow(!SceneColorTextureParameter.IsBound()
+			&& !SceneDepthTextureParameter.IsBound()
+			&& !SceneColorSurfaceParameter.IsBound()
+			&& !SceneDepthSurfaceParameter.IsBound());
+	}
+	else if( TextureMode == ESceneRenderTargetsMode::NonSceneAlignedPass )
+	{
+		FSamplerStateRHIParamRef DefaultSampler = TStaticSamplerState<SF_Point,AM_Clamp,AM_Clamp,AM_Clamp>::GetRHI();
+		SetTextureParameter(RHICmdList, ShaderRHI, SceneColorTextureParameter, SceneColorTextureParameterSampler, DefaultSampler, GBlackTexture->TextureRHI);
+		SetTextureParameter(RHICmdList, ShaderRHI, SceneDepthTextureParameter, SceneDepthTextureParameterSampler, DefaultSampler, GBlackTexture->TextureRHI);
+		SetTextureParameter(RHICmdList, ShaderRHI, SceneColorSurfaceParameter, GBlackTexture->TextureRHI);
+		SetTextureParameter(RHICmdList, ShaderRHI, SceneDepthSurfaceParameter, GBlackTexture->TextureRHI);
+	}
+}
+
+#define IMPLEMENT_SCENE_TEXTURE_PARAM_SET( ShaderRHIParamRef ) \
+	template void FSceneTextureShaderParameters::Set< ShaderRHIParamRef >( \
+		FRHICommandList& RHICmdList,				\
+		const ShaderRHIParamRef& ShaderRHI,			\
+		const FSceneView& View,						\
+		ESceneRenderTargetsMode::Type TextureMode,	\
+		ESamplerFilter ColorFilter					\
+	) const;
+
+IMPLEMENT_SCENE_TEXTURE_PARAM_SET( FVertexShaderRHIParamRef );
+IMPLEMENT_SCENE_TEXTURE_PARAM_SET( FHullShaderRHIParamRef );
+IMPLEMENT_SCENE_TEXTURE_PARAM_SET( FDomainShaderRHIParamRef );
+IMPLEMENT_SCENE_TEXTURE_PARAM_SET( FGeometryShaderRHIParamRef );
+IMPLEMENT_SCENE_TEXTURE_PARAM_SET( FPixelShaderRHIParamRef );
+IMPLEMENT_SCENE_TEXTURE_PARAM_SET( FComputeShaderRHIParamRef );
+
 FArchive& operator<<(FArchive& Ar,FSceneTextureShaderParameters& Parameters)
 {
 	Ar << Parameters.SceneColorTextureParameter;
@@ -1530,6 +1675,88 @@ void FDeferredPixelShaderParameters::Bind(const FShaderParameterMap& ParameterMa
 	CustomDepthTextureSampler.Bind(ParameterMap,TEXT("CustomDepthTextureSampler"));
 }
 
+bool IsDBufferEnabled();
+
+template< typename ShaderRHIParamRef >
+void FDeferredPixelShaderParameters::Set(FRHICommandList& RHICmdList, const ShaderRHIParamRef ShaderRHI, const FSceneView& View, ESceneRenderTargetsMode::Type TextureMode) const
+{
+	// This is needed on PC ES2 for SceneAlphaCopy, probably should be refactored for performance.
+	SceneTextureParameters.Set(RHICmdList, ShaderRHI, View, TextureMode, SF_Point);
+
+	// if() is purely an optimization and could be removed
+	if(IsDBufferEnabled())
+	{
+		IPooledRenderTarget* DBufferA = GSceneRenderTargets.DBufferA ? GSceneRenderTargets.DBufferA : GSystemTextures.BlackDummy;
+		IPooledRenderTarget* DBufferB = GSceneRenderTargets.DBufferB ? GSceneRenderTargets.DBufferB : GSystemTextures.BlackDummy;
+		IPooledRenderTarget* DBufferC = GSceneRenderTargets.DBufferC ? GSceneRenderTargets.DBufferC : GSystemTextures.BlackDummy;
+
+		// todo: optimize out when not needed
+		SetTextureParameter(RHICmdList, ShaderRHI, DBufferATexture, DBufferATextureSampler, TStaticSamplerState<>::GetRHI(), DBufferA->GetRenderTargetItem().ShaderResourceTexture);
+		SetTextureParameter(RHICmdList, ShaderRHI, DBufferBTexture, DBufferBTextureSampler, TStaticSamplerState<>::GetRHI(), DBufferB->GetRenderTargetItem().ShaderResourceTexture);
+		SetTextureParameter(RHICmdList, ShaderRHI, DBufferCTexture, DBufferCTextureSampler, TStaticSamplerState<>::GetRHI(), DBufferC->GetRenderTargetItem().ShaderResourceTexture);
+		SetTextureParameter(RHICmdList, ShaderRHI, DBufferATextureMS, DBufferA->GetRenderTargetItem().TargetableTexture);
+		SetTextureParameter(RHICmdList, ShaderRHI, DBufferBTextureMS, DBufferB->GetRenderTargetItem().TargetableTexture);
+		SetTextureParameter(RHICmdList, ShaderRHI, DBufferCTextureMS, DBufferC->GetRenderTargetItem().TargetableTexture);
+		SetTextureParameter(RHICmdList, ShaderRHI, DBufferATextureNonMS, DBufferA->GetRenderTargetItem().ShaderResourceTexture);
+		SetTextureParameter(RHICmdList, ShaderRHI, DBufferBTextureNonMS, DBufferB->GetRenderTargetItem().ShaderResourceTexture);
+		SetTextureParameter(RHICmdList, ShaderRHI, DBufferCTextureNonMS, DBufferC->GetRenderTargetItem().ShaderResourceTexture);
+	}
+
+	if (TextureMode == ESceneRenderTargetsMode::SetTextures && GRHIFeatureLevel >= ERHIFeatureLevel::SM3)
+	{
+		// if there is no ambient occlusion it's better to have white there
+		IPooledRenderTarget* ScreenSpaceAO = GSceneRenderTargets.ScreenSpaceAO;
+		if(!GSceneRenderTargets.bScreenSpaceAOIsValid)
+		{
+			ScreenSpaceAO = GSystemTextures.WhiteDummy;
+		}
+
+		// if there is no custom depth it's better to have the far distance there
+		IPooledRenderTarget* CustomDepth = GSceneRenderTargets.bCustomDepthIsValid ? GSceneRenderTargets.CustomDepth : 0;
+		if(!CustomDepth)
+		{
+			CustomDepth = GSystemTextures.BlackDummy;
+		}
+
+		if (GRHIFeatureLevel >= ERHIFeatureLevel::SM4)
+		{
+			//gilmerge
+			if (GBufferResources.IsBound())
+			{
+				SetUniformBufferParameter(RHICmdList, ShaderRHI, GBufferResources, GSceneRenderTargets.GetGBufferResourcesUniformBuffer());
+			}
+
+			SetTextureParameter(RHICmdList, ShaderRHI, ScreenSpaceAOTexture, ScreenSpaceAOTextureSampler, TStaticSamplerState<>::GetRHI(), ScreenSpaceAO->GetRenderTargetItem().ShaderResourceTexture);
+			SetTextureParameter(RHICmdList, ShaderRHI, ScreenSpaceAOTextureMS, ScreenSpaceAO->GetRenderTargetItem().TargetableTexture);
+			SetTextureParameter(RHICmdList, ShaderRHI, ScreenSpaceAOTextureNonMS, ScreenSpaceAO->GetRenderTargetItem().ShaderResourceTexture);
+
+			SetTextureParameter(RHICmdList, ShaderRHI, CustomDepthTexture, CustomDepthTextureSampler, TStaticSamplerState<>::GetRHI(), CustomDepth->GetRenderTargetItem().ShaderResourceTexture);
+			SetTextureParameter(RHICmdList, ShaderRHI, CustomDepthTextureNonMS, CustomDepth->GetRenderTargetItem().ShaderResourceTexture);
+		}
+	}
+	else if (TextureMode == ESceneRenderTargetsMode::DontSet ||
+		TextureMode == ESceneRenderTargetsMode::DontSetIgnoreBoundByEditorCompositing)
+	{
+		// Verify that none of these are actually bound
+		checkSlow(!GBufferResources.IsBound());
+	}
+}
+
+#define IMPLEMENT_DEFERRED_PARAMETERS_SET( ShaderRHIParamRef ) \
+	template void FDeferredPixelShaderParameters::Set< ShaderRHIParamRef >( \
+		FRHICommandList& RHICmdList,				\
+		const ShaderRHIParamRef ShaderRHI,			\
+		const FSceneView& View,						\
+		ESceneRenderTargetsMode::Type TextureMode	\
+	) const;
+
+IMPLEMENT_DEFERRED_PARAMETERS_SET( FVertexShaderRHIParamRef );
+IMPLEMENT_DEFERRED_PARAMETERS_SET( FHullShaderRHIParamRef );
+IMPLEMENT_DEFERRED_PARAMETERS_SET( FDomainShaderRHIParamRef );
+IMPLEMENT_DEFERRED_PARAMETERS_SET( FGeometryShaderRHIParamRef );
+IMPLEMENT_DEFERRED_PARAMETERS_SET( FPixelShaderRHIParamRef );
+IMPLEMENT_DEFERRED_PARAMETERS_SET( FComputeShaderRHIParamRef );
+
 FArchive& operator<<(FArchive& Ar,FDeferredPixelShaderParameters& Parameters)
 {
 	Ar << Parameters.SceneTextureParameters;
diff --git a/Engine/Source/Runtime/Renderer/Private/PostProcess/SceneRenderTargets.h b/Engine/Source/Runtime/Renderer/Private/PostProcess/SceneRenderTargets.h
--- a/Engine/Source/Runtime/Renderer/Private/PostProcess/SceneRenderTargets.h
+++ b/Engine/Source/Runtime/Renderer/Private/PostProcess/SceneRenderTargets.h
@@ -551,273 +551,3 @@ private:
 
 /** The global render targets used for scene rendering. */
 extern TGlobalResource<FSceneRenderTargets> GSceneRenderTargets;
-
-namespace ESceneRenderTargetsMode
-{
-	enum Type
-	{
-		SetTextures,
-		DontSet,
-		DontSetIgnoreBoundByEditorCompositing,
-		NonSceneAlignedPass,
-	};
-}
-
-/** Encapsulates scene texture shader parameter bindings. */
-class FSceneTextureShaderParameters
-{
-public:
-	/** Binds the parameters using a compiled shader's parameter map. */
-	void Bind(const FShaderParameterMap& ParameterMap);
-
-	/** Sets the scene texture parameters for the given view. */
-	template<typename TParamRef>
-	void Set(FRHICommandList& RHICmdList, const TParamRef& ShaderRHI, const FSceneView& View, ESceneRenderTargetsMode::Type TextureMode = ESceneRenderTargetsMode::SetTextures, ESamplerFilter ColorFilter = SF_Point) const
-	{
-		if (TextureMode == ESceneRenderTargetsMode::SetTextures)
-		{
-			// optimization possible: TShaderRHIParamRef is no param Ref
-			if (SceneColorTextureParameter.IsBound())
-			{
-				FSamplerStateRHIRef Filter;
-				switch ( ColorFilter )
-				{
-				case SF_Bilinear:
-					Filter = TStaticSamplerState<SF_Bilinear,AM_Clamp,AM_Clamp,AM_Clamp>::GetRHI();
-					break;
-				case SF_Trilinear:
-					Filter = TStaticSamplerState<SF_Trilinear,AM_Clamp,AM_Clamp,AM_Clamp>::GetRHI();
-					break;
-				case SF_AnisotropicPoint:
-					Filter = TStaticSamplerState<SF_AnisotropicPoint,AM_Clamp,AM_Clamp,AM_Clamp>::GetRHI();
-					break;
-				case SF_AnisotropicLinear:
-					Filter = TStaticSamplerState<SF_AnisotropicLinear,AM_Clamp,AM_Clamp,AM_Clamp>::GetRHI();
-					break;
-				case SF_Point:
-				default:
-					Filter = TStaticSamplerState<SF_Point,AM_Clamp,AM_Clamp,AM_Clamp>::GetRHI();
-					break;
-				}
-
-				SetTextureParameter(
-					RHICmdList, 
-					ShaderRHI,
-					SceneColorTextureParameter,
-					SceneColorTextureParameterSampler,
-					Filter,
-					GSceneRenderTargets.GetSceneColorTexture()
-					);
-			}
-
-			if (SceneAlphaCopyTextureParameter.IsBound() && GSceneRenderTargets.HasSceneAlphaCopyTexture())
-			{
-				FSamplerStateRHIRef Filter;
-				Filter = TStaticSamplerState<SF_Point,AM_Clamp,AM_Clamp,AM_Clamp>::GetRHI();
-				SetTextureParameter(
-					RHICmdList, 
-					ShaderRHI,
-					SceneAlphaCopyTextureParameter,
-					SceneAlphaCopyTextureParameterSampler,
-					Filter,
-					GSceneRenderTargets.GetSceneAlphaCopyTexture()
-					);
-			}
-
-			if(SceneDepthTextureParameter.IsBound() || SceneDepthTextureParameterSampler.IsBound())
-			{
-				const FTexture2DRHIRef* DepthTexture = GSceneRenderTargets.GetActualDepthTexture();
-				SetTextureParameter(
-					RHICmdList, 
-					ShaderRHI,
-					SceneDepthTextureParameter,
-					SceneDepthTextureParameterSampler,
-					TStaticSamplerState<SF_Point,AM_Clamp,AM_Clamp,AM_Clamp>::GetRHI(),
-					*DepthTexture
-					);
-			}
-
-			if (GRHIFeatureLevel >= ERHIFeatureLevel::SM5)
-			{
-				SetTextureParameter(RHICmdList, ShaderRHI, SceneColorSurfaceParameter, GSceneRenderTargets.GetSceneColorSurface());
-			}
-			if(GRHIFeatureLevel >= ERHIFeatureLevel::SM4)
-			{
-				if(GSupportsDepthFetchDuringDepthTest)
-				{
-					if(SceneDepthSurfaceParameter.IsBound())
-					{
-						SetTextureParameter(RHICmdList, ShaderRHI, SceneDepthSurfaceParameter, GSceneRenderTargets.GetSceneDepthSurface());
-					}
-					if(SceneDepthTextureNonMS.IsBound())
-					{
-						SetTextureParameter(RHICmdList, ShaderRHI, SceneDepthTextureNonMS, GSceneRenderTargets.GetSceneDepthTexture());
-					}
-				}
-				else
-				{
-					if(SceneDepthSurfaceParameter.IsBound())
-					{
-						SetTextureParameter(RHICmdList, ShaderRHI, SceneDepthSurfaceParameter, GSceneRenderTargets.GetAuxiliarySceneDepthSurface());
-					}
-					if(SceneDepthTextureNonMS.IsBound())
-					{
-						SetTextureParameter(RHICmdList, ShaderRHI, SceneDepthTextureNonMS, GSceneRenderTargets.GetAuxiliarySceneDepthSurface());
-					}
-				}
-			}
-		}
-		else if (TextureMode == ESceneRenderTargetsMode::DontSet)
-		{
-			// Verify that none of these were bound if we were told not to set them
-			checkSlow(!SceneColorTextureParameter.IsBound()
-				&& !SceneDepthTextureParameter.IsBound()
-				&& !SceneColorSurfaceParameter.IsBound()
-				&& !SceneDepthSurfaceParameter.IsBound()
-				&& !SceneDepthTextureNonMS.IsBound());
-		}
-		else if (TextureMode == ESceneRenderTargetsMode::DontSetIgnoreBoundByEditorCompositing)
-		{
-			// Verify that none of these were bound if we were told not to set them
-			// ignore SceneDepthTextureNonMS
-			checkSlow(!SceneColorTextureParameter.IsBound()
-				&& !SceneDepthTextureParameter.IsBound()
-				&& !SceneColorSurfaceParameter.IsBound()
-				&& !SceneDepthSurfaceParameter.IsBound());
-		}
-		else if( TextureMode == ESceneRenderTargetsMode::NonSceneAlignedPass )
-		{
-			FSamplerStateRHIParamRef DefaultSampler = TStaticSamplerState<SF_Point,AM_Clamp,AM_Clamp,AM_Clamp>::GetRHI();
-			SetTextureParameter(RHICmdList, ShaderRHI, SceneColorTextureParameter, SceneColorTextureParameterSampler, DefaultSampler, GBlackTexture->TextureRHI);
-			SetTextureParameter(RHICmdList, ShaderRHI, SceneDepthTextureParameter, SceneDepthTextureParameterSampler, DefaultSampler, GBlackTexture->TextureRHI);
-			SetTextureParameter(RHICmdList, ShaderRHI, SceneColorSurfaceParameter, GBlackTexture->TextureRHI);
-			SetTextureParameter(RHICmdList, ShaderRHI, SceneDepthSurfaceParameter, GBlackTexture->TextureRHI);
-		}
-	}
-
-	/** Serializer. */
-	friend FArchive& operator<<(FArchive& Ar,FSceneTextureShaderParameters& P);
-
-private:
-	/** The SceneColorTexture parameter for materials that use SceneColor */
-	FShaderResourceParameter SceneColorTextureParameter;
-	FShaderResourceParameter SceneColorTextureParameterSampler;
-	/** The SceneDepthTexture parameter for materials that use SceneDepth */
-	FShaderResourceParameter SceneDepthTextureParameter;
-	FShaderResourceParameter SceneDepthTextureParameterSampler;
-	/** The SceneAlphaCopyTexture parameter for materials that use SceneAlphaCopy */
-	FShaderResourceParameter SceneAlphaCopyTextureParameter;
-	FShaderResourceParameter SceneAlphaCopyTextureParameterSampler;
-
-	/** for MSAA access to the scene color */
-	FShaderResourceParameter SceneColorSurfaceParameter;
-	/** The SceneColorTextureMSAA parameter for materials that use SceneColorTextureMSAA */
-	FShaderResourceParameter SceneDepthSurfaceParameter;
-	/**  */
-	FShaderResourceParameter SceneDepthTextureNonMS;
-};
-
-bool IsDBufferEnabled();
-
-/** Pixel shader parameters needed for deferred passes. */
-class FDeferredPixelShaderParameters
-{
-public:
-
-	// TODO: Remove me.
-	static void ModifyCompilationEnvironment(EShaderPlatform Platform, FShaderCompilerEnvironment& OutEnvironment)
-	{
-	}
-
-	void Bind(const FShaderParameterMap& ParameterMap);
-
-	template<typename ShaderRHIParamRef>
-	void Set(FRHICommandList& RHICmdList, const ShaderRHIParamRef ShaderRHI, const FSceneView& View, ESceneRenderTargetsMode::Type TextureMode = ESceneRenderTargetsMode::SetTextures) const
-	{
-		// This is needed on PC ES2 for SceneAlphaCopy, probably should be refactored for performance.
-		SceneTextureParameters.Set(RHICmdList, ShaderRHI, View, TextureMode, SF_Point);
-
-		// if() is purely an optimization and could be removed
-		if(IsDBufferEnabled())
-		{
-			IPooledRenderTarget* DBufferA = GSceneRenderTargets.DBufferA ? GSceneRenderTargets.DBufferA : GSystemTextures.BlackDummy;
-			IPooledRenderTarget* DBufferB = GSceneRenderTargets.DBufferB ? GSceneRenderTargets.DBufferB : GSystemTextures.BlackDummy;
-			IPooledRenderTarget* DBufferC = GSceneRenderTargets.DBufferC ? GSceneRenderTargets.DBufferC : GSystemTextures.BlackDummy;
-
-			// todo: optimize out when not needed
-			SetTextureParameter(RHICmdList, ShaderRHI, DBufferATexture, DBufferATextureSampler, TStaticSamplerState<>::GetRHI(), DBufferA->GetRenderTargetItem().ShaderResourceTexture);
-			SetTextureParameter(RHICmdList, ShaderRHI, DBufferBTexture, DBufferBTextureSampler, TStaticSamplerState<>::GetRHI(), DBufferB->GetRenderTargetItem().ShaderResourceTexture);
-			SetTextureParameter(RHICmdList, ShaderRHI, DBufferCTexture, DBufferCTextureSampler, TStaticSamplerState<>::GetRHI(), DBufferC->GetRenderTargetItem().ShaderResourceTexture);
-			SetTextureParameter(RHICmdList, ShaderRHI, DBufferATextureMS, DBufferA->GetRenderTargetItem().TargetableTexture);
-			SetTextureParameter(RHICmdList, ShaderRHI, DBufferBTextureMS, DBufferB->GetRenderTargetItem().TargetableTexture);
-			SetTextureParameter(RHICmdList, ShaderRHI, DBufferCTextureMS, DBufferC->GetRenderTargetItem().TargetableTexture);
-			SetTextureParameter(RHICmdList, ShaderRHI, DBufferATextureNonMS, DBufferA->GetRenderTargetItem().ShaderResourceTexture);
-			SetTextureParameter(RHICmdList, ShaderRHI, DBufferBTextureNonMS, DBufferB->GetRenderTargetItem().ShaderResourceTexture);
-			SetTextureParameter(RHICmdList, ShaderRHI, DBufferCTextureNonMS, DBufferC->GetRenderTargetItem().ShaderResourceTexture);
-		}
-
-		if (TextureMode == ESceneRenderTargetsMode::SetTextures && GRHIFeatureLevel >= ERHIFeatureLevel::SM3)
-		{
-			// if there is no ambient occlusion it's better to have white there
-			IPooledRenderTarget* ScreenSpaceAO = GSceneRenderTargets.ScreenSpaceAO;
-			if(!GSceneRenderTargets.bScreenSpaceAOIsValid)
-			{
-				ScreenSpaceAO = GSystemTextures.WhiteDummy;
-			}
-
-			// if there is no custom depth it's better to have the far distance there
-			IPooledRenderTarget* CustomDepth = GSceneRenderTargets.bCustomDepthIsValid ? GSceneRenderTargets.CustomDepth : 0;
-			if(!CustomDepth)
-			{
-				CustomDepth = GSystemTextures.BlackDummy;
-			}
-
-			if (GRHIFeatureLevel >= ERHIFeatureLevel::SM4)
-			{
-				//gilmerge
-				if (GBufferResources.IsBound())
-				{
-					SetUniformBufferParameter(RHICmdList, ShaderRHI, GBufferResources, GSceneRenderTargets.GetGBufferResourcesUniformBuffer());
-				}
-
-				SetTextureParameter(RHICmdList, ShaderRHI, ScreenSpaceAOTexture, ScreenSpaceAOTextureSampler, TStaticSamplerState<>::GetRHI(), ScreenSpaceAO->GetRenderTargetItem().ShaderResourceTexture);
-				SetTextureParameter(RHICmdList, ShaderRHI, ScreenSpaceAOTextureMS, ScreenSpaceAO->GetRenderTargetItem().TargetableTexture);
-				SetTextureParameter(RHICmdList, ShaderRHI, ScreenSpaceAOTextureNonMS, ScreenSpaceAO->GetRenderTargetItem().ShaderResourceTexture);
-
-				SetTextureParameter(RHICmdList, ShaderRHI, CustomDepthTexture, CustomDepthTextureSampler, TStaticSamplerState<>::GetRHI(), CustomDepth->GetRenderTargetItem().ShaderResourceTexture);
-				SetTextureParameter(RHICmdList, ShaderRHI, CustomDepthTextureNonMS, CustomDepth->GetRenderTargetItem().ShaderResourceTexture);
-			}
-		}
-		else if (TextureMode == ESceneRenderTargetsMode::DontSet ||
-			TextureMode == ESceneRenderTargetsMode::DontSetIgnoreBoundByEditorCompositing)
-		{
-			// Verify that none of these are actually bound
-			checkSlow(!GBufferResources.IsBound());
-		}
-	}
-
-	friend FArchive& operator<<(FArchive& Ar,FDeferredPixelShaderParameters& P);
-
-private:
-	FSceneTextureShaderParameters SceneTextureParameters;
-	FShaderUniformBufferParameter GBufferResources;
-	FShaderResourceParameter DBufferATextureMS;
-	FShaderResourceParameter DBufferBTextureMS;
-	FShaderResourceParameter DBufferCTextureMS;
-	FShaderResourceParameter ScreenSpaceAOTextureMS;
-	FShaderResourceParameter DBufferATextureNonMS;
-	FShaderResourceParameter DBufferBTextureNonMS;
-	FShaderResourceParameter DBufferCTextureNonMS;
-	FShaderResourceParameter ScreenSpaceAOTextureNonMS;
-	FShaderResourceParameter CustomDepthTextureNonMS;
-	FShaderResourceParameter DBufferATexture;
-	FShaderResourceParameter DBufferATextureSampler;
-	FShaderResourceParameter DBufferBTexture;
-	FShaderResourceParameter DBufferBTextureSampler;
-	FShaderResourceParameter DBufferCTexture;
-	FShaderResourceParameter DBufferCTextureSampler;
-	FShaderResourceParameter ScreenSpaceAOTexture;
-	FShaderResourceParameter ScreenSpaceAOTextureSampler;
-	FShaderResourceParameter CustomDepthTexture;
-	FShaderResourceParameter CustomDepthTextureSampler;
-};
diff --git a/Engine/Source/Runtime/Renderer/Private/PostProcess/ScreenSpaceReflections.cpp b/Engine/Source/Runtime/Renderer/Private/PostProcess/ScreenSpaceReflections.cpp
--- a/Engine/Source/Runtime/Renderer/Private/PostProcess/ScreenSpaceReflections.cpp
+++ b/Engine/Source/Runtime/Renderer/Private/PostProcess/ScreenSpaceReflections.cpp
@@ -74,7 +74,6 @@ class FPostProcessScreenSpaceReflectionsPS : public FGlobalShader
 	static void ModifyCompilationEnvironment(EShaderPlatform Platform, FShaderCompilerEnvironment& OutEnvironment)
 	{
 		FGlobalShader::ModifyCompilationEnvironment(Platform, OutEnvironment);
-		FDeferredPixelShaderParameters::ModifyCompilationEnvironment(Platform,OutEnvironment);
 		OutEnvironment.SetDefine( TEXT("PREV_FRAME_COLOR"), PrevFrame );
 		OutEnvironment.SetDefine( TEXT("SSR_QUALITY"), SSRQuality );
 	}
diff --git a/Engine/Source/Runtime/Renderer/Private/ReflectionEnvironment.cpp b/Engine/Source/Runtime/Renderer/Private/ReflectionEnvironment.cpp
--- a/Engine/Source/Runtime/Renderer/Private/ReflectionEnvironment.cpp
+++ b/Engine/Source/Runtime/Renderer/Private/ReflectionEnvironment.cpp
@@ -149,7 +149,6 @@ public:
 	static void ModifyCompilationEnvironment(EShaderPlatform Platform, FShaderCompilerEnvironment& OutEnvironment)
 	{
 		FGlobalShader::ModifyCompilationEnvironment(Platform, OutEnvironment);
-		FDeferredPixelShaderParameters::ModifyCompilationEnvironment(Platform,OutEnvironment);
 		OutEnvironment.SetDefine(TEXT("THREADGROUP_SIZEX"), GReflectionEnvironmentTileSizeX);
 		OutEnvironment.SetDefine(TEXT("THREADGROUP_SIZEY"), GReflectionEnvironmentTileSizeY);
 		OutEnvironment.SetDefine(TEXT("MAX_CAPTURES"), GMaxNumReflectionCaptures);
@@ -347,7 +346,6 @@ class FReflectionApplyPS : public FGlobalShader
 		OutEnvironment.SetDefine(TEXT("APPLY_SSR"), bSSR);
 		OutEnvironment.SetDefine(TEXT("APPLY_REFLECTION_ENV"), bReflectionEnv);
 		OutEnvironment.SetDefine(TEXT("APPLY_SKYLIGHT"), bSkylight);
-		FDeferredPixelShaderParameters::ModifyCompilationEnvironment(Platform,OutEnvironment);
 	}
 
 	/** Default constructor. */
@@ -433,7 +431,6 @@ class FReflectionCaptureSpecularBouncePS : public FGlobalShader
 	static void ModifyCompilationEnvironment(EShaderPlatform Platform, FShaderCompilerEnvironment& OutEnvironment)
 	{
 		FGlobalShader::ModifyCompilationEnvironment(Platform,OutEnvironment);
-		FDeferredPixelShaderParameters::ModifyCompilationEnvironment(Platform,OutEnvironment);
 	}
 
 	/** Default constructor. */
diff --git a/Engine/Source/Runtime/Renderer/Private/ReflectionEnvironmentCapture.cpp b/Engine/Source/Runtime/Renderer/Private/ReflectionEnvironmentCapture.cpp
--- a/Engine/Source/Runtime/Renderer/Private/ReflectionEnvironmentCapture.cpp
+++ b/Engine/Source/Runtime/Renderer/Private/ReflectionEnvironmentCapture.cpp
@@ -619,7 +619,6 @@ public:
 	static void ModifyCompilationEnvironment(EShaderPlatform Platform, FShaderCompilerEnvironment& OutEnvironment)
 	{
 		FGlobalShader::ModifyCompilationEnvironment(Platform,OutEnvironment);
-		FDeferredPixelShaderParameters::ModifyCompilationEnvironment(Platform,OutEnvironment);
 	}
 
 	FCopySceneColorToCubeFacePS(const ShaderMetaType::CompiledShaderInitializerType& Initializer):
diff --git a/Engine/Source/Runtime/Renderer/Private/ReflectionEnvironmentDiffuseIrradiance.cpp b/Engine/Source/Runtime/Renderer/Private/ReflectionEnvironmentDiffuseIrradiance.cpp
--- a/Engine/Source/Runtime/Renderer/Private/ReflectionEnvironmentDiffuseIrradiance.cpp
+++ b/Engine/Source/Runtime/Renderer/Private/ReflectionEnvironmentDiffuseIrradiance.cpp
@@ -134,7 +134,6 @@ public:
 	static void ModifyCompilationEnvironment(EShaderPlatform Platform, FShaderCompilerEnvironment& OutEnvironment)
 	{
 		FGlobalShader::ModifyCompilationEnvironment(Platform,OutEnvironment);
-		FDeferredPixelShaderParameters::ModifyCompilationEnvironment(Platform,OutEnvironment);
 	}
 
 	void SetParameters(FRHICommandList& RHICmdList, int32 CubeFaceValue, int32 NumMips, int32 SourceMipIndexValue, int32 CoefficientIndex, FTextureRHIRef& SourceTextureValue)
diff --git a/Engine/Source/Runtime/Renderer/Private/RendererScene.cpp b/Engine/Source/Runtime/Renderer/Private/RendererScene.cpp
--- a/Engine/Source/Runtime/Renderer/Private/RendererScene.cpp
+++ b/Engine/Source/Runtime/Renderer/Private/RendererScene.cpp
@@ -8,6 +8,7 @@
 #include "ScenePrivate.h"
 #include "ShaderCompiler.h"
 #include "StaticMeshResources.h"
+#include "ParameterCollection.h"
 #include "DistanceFieldSurfaceCacheLighting.h"
 #include "EngineModule.h"
 
diff --git a/Engine/Source/Runtime/Renderer/Private/ScenePrivate.h b/Engine/Source/Runtime/Renderer/Private/ScenePrivate.h
--- a/Engine/Source/Runtime/Renderer/Private/ScenePrivate.h
+++ b/Engine/Source/Runtime/Renderer/Private/ScenePrivate.h
@@ -65,6 +65,7 @@ bool IsMobileHDR32bpp();
 #include "ScopedPointer.h"
 #include "ClearQuad.h"
 #include "SpeedTreeWind.h"
+#include "AtmosphereRendering.h"
 
 /** Factor by which to grow occlusion tests **/
 #define OCCLUSION_SLOP (1.0f)
diff --git a/Engine/Source/Runtime/Renderer/Private/ShaderBaseClasses.cpp b/Engine/Source/Runtime/Renderer/Private/ShaderBaseClasses.cpp
--- a/Engine/Source/Runtime/Renderer/Private/ShaderBaseClasses.cpp
+++ b/Engine/Source/Runtime/Renderer/Private/ShaderBaseClasses.cpp
@@ -54,6 +54,232 @@ FUniformBufferRHIParamRef FMaterialShader::GetParameterCollectionBuffer(const FG
 	return Scene ? Scene->GetParameterCollectionBuffer(Id) : FUniformBufferRHIParamRef();
 }
 
+template<typename ShaderRHIParamRef>
+void FMaterialShader::SetParameters(
+	FRHICommandList& RHICmdList,
+	const ShaderRHIParamRef ShaderRHI, 
+	const FMaterialRenderProxy* MaterialRenderProxy, 
+	const FMaterial& Material,
+	const FSceneView& View, 
+	bool bDeferredPass, 
+	ESceneRenderTargetsMode::Type TextureMode)
+{
+	ERHIFeatureLevel::Type FeatureLevel = View.GetFeatureLevel();
+	FUniformExpressionCache TempUniformExpressionCache;
+	const FUniformExpressionCache* UniformExpressionCache = &MaterialRenderProxy->UniformExpressionCache[FeatureLevel];
+
+	SetParameters(RHICmdList, ShaderRHI, View);
+
+	// If the material has cached uniform expressions for selection or hover
+	// and that is being overridden by show flags in the editor, recache
+	// expressions for this draw call.
+	const bool bOverrideSelection =
+		GIsEditor &&
+		!View.Family->EngineShowFlags.Selection &&
+		(MaterialRenderProxy->IsSelected() || MaterialRenderProxy->IsHovered());
+
+	if (!bAllowCachedUniformExpressions || !UniformExpressionCache->bUpToDate || bOverrideSelection)
+	{
+		FMaterialRenderContext MaterialRenderContext(MaterialRenderProxy, Material, &View);
+		MaterialRenderProxy->EvaluateUniformExpressions(TempUniformExpressionCache, MaterialRenderContext);
+		UniformExpressionCache = &TempUniformExpressionCache;
+	}
+
+	check(Material.GetRenderingThreadShaderMap());
+	check(Material.GetRenderingThreadShaderMap()->IsValidForRendering());
+	check(Material.GetFeatureLevel() == FeatureLevel);
+
+#if NO_LOGGING == 0
+	// Validate that the shader is being used for a material that matches the uniform expression set the shader was compiled for.
+	const FUniformExpressionSet& MaterialUniformExpressionSet = Material.GetRenderingThreadShaderMap()->GetUniformExpressionSet();
+	const bool bUniformExpressionSetMismatch = !DebugUniformExpressionSet.Matches(MaterialUniformExpressionSet)
+		|| UniformExpressionCache->CachedUniformExpressionShaderMap != Material.GetRenderingThreadShaderMap();
+
+	if (bUniformExpressionSetMismatch)
+	{
+		UE_LOG(
+			LogShaders,
+			Fatal,
+			TEXT("%s shader uniform expression set mismatch for material %s/%s.\n")
+			TEXT("Shader compilation info:                %s\n")
+			TEXT("Material render proxy compilation info: %s\n")
+			TEXT("Shader uniform expression set:   %u vectors, %u scalars, %u 2D textures, %u cube textures, shader map %p\n")
+			TEXT("Material uniform expression set: %u vectors, %u scalars, %u 2D textures, %u cube textures, shader map %p\n"),
+			GetType()->GetName(),
+			*MaterialRenderProxy->GetFriendlyName(),
+			*Material.GetFriendlyName(),
+			*DebugDescription,
+			*Material.GetRenderingThreadShaderMap()->GetDebugDescription(),
+			DebugUniformExpressionSet.NumVectorExpressions,
+			DebugUniformExpressionSet.NumScalarExpressions,
+			DebugUniformExpressionSet.Num2DTextureExpressions,
+			DebugUniformExpressionSet.NumCubeTextureExpressions,
+			UniformExpressionCache->CachedUniformExpressionShaderMap,
+			MaterialUniformExpressionSet.UniformVectorExpressions.Num(),
+			MaterialUniformExpressionSet.UniformScalarExpressions.Num(),
+			MaterialUniformExpressionSet.Uniform2DTextureExpressions.Num(),
+			MaterialUniformExpressionSet.UniformCubeTextureExpressions.Num(),
+			Material.GetRenderingThreadShaderMap()
+			);
+	}
+#endif
+
+	// Set the material uniform buffer.
+	SetUniformBufferParameter(RHICmdList, ShaderRHI,MaterialUniformBuffer,UniformExpressionCache->UniformBuffer);
+
+	{
+		const TArray<FGuid>& ParameterCollections = UniformExpressionCache->ParameterCollections;
+		const int32 ParameterCollectionsNum = ParameterCollections.Num();
+
+		check(ParameterCollectionUniformBuffers.Num() >= ParameterCollectionsNum);
+
+		// Find each referenced parameter collection's uniform buffer in the scene and set the parameter
+		for (int32 CollectionIndex = 0; CollectionIndex < ParameterCollectionsNum; CollectionIndex++)
+		{
+			FUniformBufferRHIParamRef UniformBuffer = GetParameterCollectionBuffer(ParameterCollections[CollectionIndex], View.Family->Scene);
+			SetUniformBufferParameter(RHICmdList, ShaderRHI,ParameterCollectionUniformBuffers[CollectionIndex],UniformBuffer);
+		}
+	}
+
+	// Set 2D texture uniform expressions.
+	{
+		int32 Count = Uniform2DTextureShaderResourceParameters.Num();
+		for(int32 ParameterIndex = 0;ParameterIndex < Count;ParameterIndex++)
+		{
+			const TUniformParameter<FShaderResourceParameter>& TextureResourceParameter = Uniform2DTextureShaderResourceParameters[ParameterIndex];
+			check(UniformExpressionCache->Textures.IsValidIndex(TextureResourceParameter.Index));
+			const TUniformParameter<FShaderResourceParameter>& SamplerResourceParameter = Uniform2DTextureSamplerShaderResourceParameters[ParameterIndex];
+
+			const UTexture* Texture2D = UniformExpressionCache->Textures[TextureResourceParameter.Index];
+			const FTexture* Value = Texture2D ? Texture2D->Resource : NULL;
+
+			if( !Value )
+			{
+				Value = GWhiteTexture;
+			}
+
+			checkSlow(Value);
+			Value->LastRenderTime = FApp::GetCurrentTime();
+
+			SetTextureParameter(
+				RHICmdList, 
+				ShaderRHI, 
+				TextureResourceParameter.ShaderParameter, 
+				SamplerResourceParameter.ShaderParameter, 
+				bDeferredPass && Value->DeferredPassSamplerStateRHI ? Value->DeferredPassSamplerStateRHI : Value->SamplerStateRHI, 
+				Value->TextureRHI);
+		}
+	}
+
+	// Set cube texture uniform expressions.
+	{
+		int32 Count = UniformCubeTextureShaderResourceParameters.Num(); 
+		for(int32 ParameterIndex = 0;ParameterIndex < Count;ParameterIndex++)
+		{
+			const TUniformParameter<FShaderResourceParameter>& TextureResourceParameter = UniformCubeTextureShaderResourceParameters[ParameterIndex];
+			check(UniformExpressionCache->CubeTextures.IsValidIndex(TextureResourceParameter.Index));
+			const TUniformParameter<FShaderResourceParameter>& SamplerResourceParameter = UniformCubeTextureSamplerShaderResourceParameters[ParameterIndex];
+
+			const UTexture* CubeTexture = UniformExpressionCache->CubeTextures[TextureResourceParameter.Index];
+			const FTexture* Value = CubeTexture ? CubeTexture->Resource : NULL;
+
+			if( !Value )
+			{
+				Value = GWhiteTextureCube;
+			}
+
+			checkSlow(Value);
+			Value->LastRenderTime = FApp::GetCurrentTime();
+
+			SetTextureParameter(
+				RHICmdList, 
+				ShaderRHI,
+				TextureResourceParameter.ShaderParameter,
+				SamplerResourceParameter.ShaderParameter,
+				bDeferredPass && Value->DeferredPassSamplerStateRHI ? Value->DeferredPassSamplerStateRHI : Value->SamplerStateRHI, 
+				Value->TextureRHI);
+		}
+	}
+
+	DeferredParameters.Set(RHICmdList, ShaderRHI, View, TextureMode);
+
+	AtmosphericFogTextureParameters.Set(RHICmdList, ShaderRHI, View);
+
+	if (FeatureLevel >= ERHIFeatureLevel::SM3)
+	{
+		// for copied scene color
+		if(LightAttenuation.IsBound())
+		{
+			SetTextureParameter(
+				RHICmdList,
+				ShaderRHI,
+				LightAttenuation,
+				LightAttenuationSampler,
+				TStaticSamplerState<SF_Bilinear,AM_Clamp,AM_Clamp,AM_Clamp>::GetRHI(),
+				GSceneRenderTargets.GetLightAttenuationTexture());
+		}
+	}
+
+	// if we are in a postprocessing pass
+	if(View.RenderingCompositePassContext)
+	{
+		PostprocessParameter.Set(RHICmdList, ShaderRHI, *View.RenderingCompositePassContext, TStaticSamplerState<>::GetRHI());
+	}
+
+	//Use of the eye adaptation texture here is experimental and potentially dangerous as it can introduce a feedback loop. May be removed.
+	if(EyeAdaptation.IsBound())
+	{
+		FTextureRHIRef& EyeAdaptationTex = GetEyeAdaptation(View);
+		SetTextureParameter(RHICmdList, ShaderRHI, EyeAdaptation, EyeAdaptationTex);
+	}
+
+	if (PerlinNoiseGradientTexture.IsBound() && IsValidRef(GSystemTextures.PerlinNoiseGradient))
+	{
+		const FTexture2DRHIRef& Texture = (FTexture2DRHIRef&)GSystemTextures.PerlinNoiseGradient->GetRenderTargetItem().ShaderResourceTexture;
+		// Bind the PerlinNoiseGradientTexture as a texture
+		SetTextureParameter(
+			RHICmdList,
+			ShaderRHI,
+			PerlinNoiseGradientTexture,
+			PerlinNoiseGradientTextureSampler,
+			TStaticSamplerState<SF_Point,AM_Wrap,AM_Wrap,AM_Wrap>::GetRHI(),
+			Texture
+			);
+	}
+
+	if (PerlinNoise3DTexture.IsBound() && IsValidRef(GSystemTextures.PerlinNoise3D))
+	{
+		const FTexture3DRHIRef& Texture = (FTexture3DRHIRef&)GSystemTextures.PerlinNoise3D->GetRenderTargetItem().ShaderResourceTexture;
+		// Bind the PerlinNoise3DTexture as a texture
+		SetTextureParameter(
+			RHICmdList,
+			ShaderRHI,
+			PerlinNoise3DTexture,
+			PerlinNoise3DTextureSampler,
+			TStaticSamplerState<SF_Bilinear,AM_Wrap,AM_Wrap,AM_Wrap>::GetRHI(),
+			Texture
+			);
+	}
+}
+
+#define IMPLEMENT_MATERIAL_SHADER_SetParameters( ShaderRHIParamRef ) \
+	template void FMaterialShader::SetParameters< ShaderRHIParamRef >( \
+		FRHICommandList& RHICmdList,					\
+		const ShaderRHIParamRef ShaderRHI,				\
+		const FMaterialRenderProxy* MaterialRenderProxy,\
+		const FMaterial& Material,						\
+		const FSceneView& View,							\
+		bool bDeferredPass,								\
+		ESceneRenderTargetsMode::Type TextureMode		\
+	);
+
+IMPLEMENT_MATERIAL_SHADER_SetParameters( FVertexShaderRHIParamRef );
+IMPLEMENT_MATERIAL_SHADER_SetParameters( FHullShaderRHIParamRef );
+IMPLEMENT_MATERIAL_SHADER_SetParameters( FDomainShaderRHIParamRef );
+IMPLEMENT_MATERIAL_SHADER_SetParameters( FGeometryShaderRHIParamRef );
+IMPLEMENT_MATERIAL_SHADER_SetParameters( FPixelShaderRHIParamRef );
+IMPLEMENT_MATERIAL_SHADER_SetParameters( FComputeShaderRHIParamRef );
+
 bool FMaterialShader::Serialize(FArchive& Ar)
 {
 	const bool bShaderHasOutdatedParameters = FShader::Serialize(Ar);
diff --git a/Engine/Source/Runtime/Renderer/Private/ShaderBaseClasses.h b/Engine/Source/Runtime/Renderer/Private/ShaderBaseClasses.h
--- a/Engine/Source/Runtime/Renderer/Private/ShaderBaseClasses.h
+++ b/Engine/Source/Runtime/Renderer/Private/ShaderBaseClasses.h
@@ -6,22 +6,7 @@
 
 #pragma once
 
-#include "MaterialShader.h"
 #include "MeshMaterialShader.h"
-#include "AtmosphereRendering.h"
-#include "PostProcess/RenderingCompositionGraph.h"
-#include "ParameterCollection.h"
-
-template<typename ParameterType> 
-struct TUniformParameter
-{
-	int32 Index;
-	ParameterType ShaderParameter;
-	friend FArchive& operator<<(FArchive& Ar,TUniformParameter<ParameterType>& P)
-	{
-		return Ar << P.Index << P.ShaderParameter;
-	}
-};
 
 /** The uniform shader parameters associated with a LOD fade. */
 // This was moved out of ScenePrivate.h to workaround MSVC vs clang template issue (it's used in this header file, so needs to be declared earlier)
@@ -31,416 +16,6 @@ END_UNIFORM_BUFFER_STRUCT(FDistanceCullFadeUniformShaderParameters)
 
 typedef TUniformBufferRef< FDistanceCullFadeUniformShaderParameters > FDistanceCullFadeUniformBufferRef;
 
-/**
- * Debug information related to uniform expression sets.
- */
-class FDebugUniformExpressionSet
-{
-public:
-	/** The number of each type of expression contained in the set. */
-	int32 NumVectorExpressions;
-	int32 NumScalarExpressions;
-	int32 Num2DTextureExpressions;
-	int32 NumCubeTextureExpressions;
-
-	FDebugUniformExpressionSet()
-		: NumVectorExpressions(0)
-		, NumScalarExpressions(0)
-		, Num2DTextureExpressions(0)
-		, NumCubeTextureExpressions(0)
-	{
-	}
-
-	explicit FDebugUniformExpressionSet(const FUniformExpressionSet& InUniformExpressionSet)
-	{
-		InitFromExpressionSet(InUniformExpressionSet);
-	}
-
-	/** Initialize from a uniform expression set. */
-	void InitFromExpressionSet(const FUniformExpressionSet& InUniformExpressionSet)
-	{
-		NumVectorExpressions = InUniformExpressionSet.UniformVectorExpressions.Num();
-		NumScalarExpressions = InUniformExpressionSet.UniformScalarExpressions.Num();
-		Num2DTextureExpressions = InUniformExpressionSet.Uniform2DTextureExpressions.Num();
-		NumCubeTextureExpressions = InUniformExpressionSet.UniformCubeTextureExpressions.Num();
-	}
-
-	/** Returns true if the number of uniform expressions matches those with which the debug set was initialized. */
-	bool Matches(const FUniformExpressionSet& InUniformExpressionSet) const
-	{
-		return NumVectorExpressions == InUniformExpressionSet.UniformVectorExpressions.Num()
-			&& NumScalarExpressions == InUniformExpressionSet.UniformScalarExpressions.Num()
-			&& Num2DTextureExpressions == InUniformExpressionSet.Uniform2DTextureExpressions.Num()
-			&& NumCubeTextureExpressions == InUniformExpressionSet.UniformCubeTextureExpressions.Num();
-	}
-};
-
-/** Serialization for debug uniform expression sets. */
-inline FArchive& operator<<(FArchive& Ar, FDebugUniformExpressionSet& DebugExpressionSet)
-{
-	Ar << DebugExpressionSet.NumVectorExpressions;
-	Ar << DebugExpressionSet.NumScalarExpressions;
-	Ar << DebugExpressionSet.Num2DTextureExpressions;
-	Ar << DebugExpressionSet.NumCubeTextureExpressions;
-	return Ar;
-}
-
-
-/*
-class FGlobalShader : public FShader
-{
-	DECLARE_SHADER_TYPE(FGlobalShader,Global);
-public:
-
-	ENGINE_API FGlobalShader(const ShaderMetaType::CompiledShaderInitializerType& Initializer);
-
-
-	typedef void (*ModifyCompilationEnvironmentType)(EShaderPlatform, FShaderCompilerEnvironment&);
-
-	// FShader interface.
-	ENGINE_API static void ModifyCompilationEnvironment(EShaderPlatform Platform, FShaderCompilerEnvironment& OutEnvironment);
-
-*/
-
-/** Base class of all shaders that need material parameters. */
-class FMaterialShader : public FShader
-{
-public:
-	FMaterialShader() {}
-
-	FMaterialShader(const FMaterialShaderType::CompiledShaderInitializerType& Initializer);
-
-	typedef void (*ModifyCompilationEnvironmentType)(EShaderPlatform, const FMaterial*, FShaderCompilerEnvironment&);
-
-	ENGINE_API static void ModifyCompilationEnvironment(EShaderPlatform Platform, const FMaterial* Material, FShaderCompilerEnvironment& OutEnvironment)
-	{
-	}
-
-	FUniformBufferRHIParamRef GetParameterCollectionBuffer(const FGuid& Id, const FSceneInterface* SceneInterface) const;
-
-	template<typename ShaderRHIParamRef>
-	void SetParameters(FRHICommandList& RHICmdList, const ShaderRHIParamRef ShaderRHI,const FSceneView& View)
-	{
-		check(GetUniformBufferParameter<FViewUniformShaderParameters>().IsInitialized());
-		CheckShaderIsValid();
-		SetUniformBufferParameter(RHICmdList, ShaderRHI,GetUniformBufferParameter<FViewUniformShaderParameters>(),View.UniformBuffer);
-	}
-
-	static void ModifyCompilationEnvironment(EShaderPlatform Platform, FShaderCompilerEnvironment& OutEnvironment)
-	{
-		FGlobalShader::ModifyCompilationEnvironment(Platform,OutEnvironment);
-	}
-
-	/** Sets pixel parameters that are material specific but not FMeshBatch specific. */
-	template<typename ShaderRHIParamRef>
-	void SetParameters(
-		FRHICommandList& RHICmdList,
-		const ShaderRHIParamRef ShaderRHI, 
-		const FMaterialRenderProxy* MaterialRenderProxy, 
-		const FMaterial& Material,
-		const FSceneView& View, 
-		bool bDeferredPass, 
-		ESceneRenderTargetsMode::Type TextureMode)
-	{
-		ERHIFeatureLevel::Type FeatureLevel = View.GetFeatureLevel();
-		FUniformExpressionCache TempUniformExpressionCache;
-		const FUniformExpressionCache* UniformExpressionCache = &MaterialRenderProxy->UniformExpressionCache[FeatureLevel];
-
-		SetParameters(RHICmdList, ShaderRHI, View);
-
-		// If the material has cached uniform expressions for selection or hover
-		// and that is being overridden by show flags in the editor, recache
-		// expressions for this draw call.
-		const bool bOverrideSelection =
-			GIsEditor &&
-			!View.Family->EngineShowFlags.Selection &&
-			(MaterialRenderProxy->IsSelected() || MaterialRenderProxy->IsHovered());
-
-		if (!bAllowCachedUniformExpressions || !UniformExpressionCache->bUpToDate || bOverrideSelection)
-		{
-			FMaterialRenderContext MaterialRenderContext(MaterialRenderProxy, Material, &View);
-			MaterialRenderProxy->EvaluateUniformExpressions(TempUniformExpressionCache, MaterialRenderContext);
-			UniformExpressionCache = &TempUniformExpressionCache;
-		}
-
-		check(Material.GetRenderingThreadShaderMap());
-		check(Material.GetRenderingThreadShaderMap()->IsValidForRendering());
-		check(Material.GetFeatureLevel() == FeatureLevel);
-
-#if NO_LOGGING == 0
-		// Validate that the shader is being used for a material that matches the uniform expression set the shader was compiled for.
-		const FUniformExpressionSet& MaterialUniformExpressionSet = Material.GetRenderingThreadShaderMap()->GetUniformExpressionSet();
-		const bool bUniformExpressionSetMismatch = !DebugUniformExpressionSet.Matches(MaterialUniformExpressionSet)
-			|| UniformExpressionCache->CachedUniformExpressionShaderMap != Material.GetRenderingThreadShaderMap();
-
-		if (bUniformExpressionSetMismatch)
-		{
-			UE_LOG(
-				LogShaders,
-				Fatal,
-				TEXT("%s shader uniform expression set mismatch for material %s/%s.\n")
-				TEXT("Shader compilation info:                %s\n")
-				TEXT("Material render proxy compilation info: %s\n")
-				TEXT("Shader uniform expression set:   %u vectors, %u scalars, %u 2D textures, %u cube textures, shader map %p\n")
-				TEXT("Material uniform expression set: %u vectors, %u scalars, %u 2D textures, %u cube textures, shader map %p\n"),
-				GetType()->GetName(),
-				*MaterialRenderProxy->GetFriendlyName(),
-				*Material.GetFriendlyName(),
-				*DebugDescription,
-				*Material.GetRenderingThreadShaderMap()->GetDebugDescription(),
-				DebugUniformExpressionSet.NumVectorExpressions,
-				DebugUniformExpressionSet.NumScalarExpressions,
-				DebugUniformExpressionSet.Num2DTextureExpressions,
-				DebugUniformExpressionSet.NumCubeTextureExpressions,
-				UniformExpressionCache->CachedUniformExpressionShaderMap,
-				MaterialUniformExpressionSet.UniformVectorExpressions.Num(),
-				MaterialUniformExpressionSet.UniformScalarExpressions.Num(),
-				MaterialUniformExpressionSet.Uniform2DTextureExpressions.Num(),
-				MaterialUniformExpressionSet.UniformCubeTextureExpressions.Num(),
-				Material.GetRenderingThreadShaderMap()
-				);
-		}
-#endif
-
-		// Set the material uniform buffer.
-		SetUniformBufferParameter(RHICmdList, ShaderRHI,MaterialUniformBuffer,UniformExpressionCache->UniformBuffer);
-
-		{
-			const TArray<FGuid>& ParameterCollections = UniformExpressionCache->ParameterCollections;
-			const int32 ParameterCollectionsNum = ParameterCollections.Num();
-
-			check(ParameterCollectionUniformBuffers.Num() >= ParameterCollectionsNum);
-
-			// Find each referenced parameter collection's uniform buffer in the scene and set the parameter
-			for (int32 CollectionIndex = 0; CollectionIndex < ParameterCollectionsNum; CollectionIndex++)
-			{
-				FUniformBufferRHIParamRef UniformBuffer = GetParameterCollectionBuffer(ParameterCollections[CollectionIndex], View.Family->Scene);
-				SetUniformBufferParameter(RHICmdList, ShaderRHI,ParameterCollectionUniformBuffers[CollectionIndex],UniformBuffer);
-			}
-		}
-
-		// Set 2D texture uniform expressions.
-		{
-			int32 Count = Uniform2DTextureShaderResourceParameters.Num();
-			for(int32 ParameterIndex = 0;ParameterIndex < Count;ParameterIndex++)
-			{
-				const TUniformParameter<FShaderResourceParameter>& TextureResourceParameter = Uniform2DTextureShaderResourceParameters[ParameterIndex];
-				check(UniformExpressionCache->Textures.IsValidIndex(TextureResourceParameter.Index));
-				const TUniformParameter<FShaderResourceParameter>& SamplerResourceParameter = Uniform2DTextureSamplerShaderResourceParameters[ParameterIndex];
-
-				const UTexture* Texture2D = UniformExpressionCache->Textures[TextureResourceParameter.Index];
-				const FTexture* Value = Texture2D ? Texture2D->Resource : NULL;
-
-				if( !Value )
-				{
-					Value = GWhiteTexture;
-				}
-
-				checkSlow(Value);
-				Value->LastRenderTime = FApp::GetCurrentTime();
-
-				SetTextureParameter(
-					RHICmdList, 
-					ShaderRHI, 
-					TextureResourceParameter.ShaderParameter, 
-					SamplerResourceParameter.ShaderParameter, 
-					bDeferredPass && Value->DeferredPassSamplerStateRHI ? Value->DeferredPassSamplerStateRHI : Value->SamplerStateRHI, 
-					Value->TextureRHI);
-			}
-		}
-
-		// Set cube texture uniform expressions.
-		{
-			int32 Count = UniformCubeTextureShaderResourceParameters.Num(); 
-			for(int32 ParameterIndex = 0;ParameterIndex < Count;ParameterIndex++)
-			{
-				const TUniformParameter<FShaderResourceParameter>& TextureResourceParameter = UniformCubeTextureShaderResourceParameters[ParameterIndex];
-				check(UniformExpressionCache->CubeTextures.IsValidIndex(TextureResourceParameter.Index));
-				const TUniformParameter<FShaderResourceParameter>& SamplerResourceParameter = UniformCubeTextureSamplerShaderResourceParameters[ParameterIndex];
-
-				const UTexture* CubeTexture = UniformExpressionCache->CubeTextures[TextureResourceParameter.Index];
-				const FTexture* Value = CubeTexture ? CubeTexture->Resource : NULL;
-
-				if( !Value )
-				{
-					Value = GWhiteTextureCube;
-				}
-
-				checkSlow(Value);
-				Value->LastRenderTime = FApp::GetCurrentTime();
-
-				SetTextureParameter(
-					RHICmdList, 
-					ShaderRHI,
-					TextureResourceParameter.ShaderParameter,
-					SamplerResourceParameter.ShaderParameter,
-					bDeferredPass && Value->DeferredPassSamplerStateRHI ? Value->DeferredPassSamplerStateRHI : Value->SamplerStateRHI, 
-					Value->TextureRHI);
-			}
-		}
-
-		DeferredParameters.Set(RHICmdList, ShaderRHI, View, TextureMode);
-
-		AtmosphericFogTextureParameters.Set(RHICmdList, ShaderRHI, View);
-
-		if (FeatureLevel >= ERHIFeatureLevel::SM3)
-		{
-			if(LightAttenuation.IsBound())
-		{
-			SetTextureParameter(
-				RHICmdList,
-				ShaderRHI,
-				LightAttenuation,
-				LightAttenuationSampler,
-				TStaticSamplerState<SF_Bilinear,AM_Clamp,AM_Clamp,AM_Clamp>::GetRHI(),
-				GSceneRenderTargets.GetLightAttenuationTexture());
-		}
-		}
-
-		// if we are in a postprocessing pass
-		if(View.RenderingCompositePassContext)
-		{
-			PostprocessParameter.Set(RHICmdList, ShaderRHI, *View.RenderingCompositePassContext, TStaticSamplerState<>::GetRHI());
-		}
-
-		//Use of the eye adaptation texture here is experimental and potentially dangerous as it can introduce a feedback loop. May be removed.
-		if(EyeAdaptation.IsBound())
-		{
-			FTextureRHIRef& EyeAdaptationTex = GetEyeAdaptation(View);
-			SetTextureParameter(RHICmdList, ShaderRHI, EyeAdaptation, EyeAdaptationTex);
-		}
-
-		if (PerlinNoiseGradientTexture.IsBound() && IsValidRef(GSystemTextures.PerlinNoiseGradient))
-		{
-			const FTexture2DRHIRef& Texture = (FTexture2DRHIRef&)GSystemTextures.PerlinNoiseGradient->GetRenderTargetItem().ShaderResourceTexture;
-			// Bind the PerlinNoiseGradientTexture as a texture
-			SetTextureParameter(
-				RHICmdList,
-				ShaderRHI,
-				PerlinNoiseGradientTexture,
-				PerlinNoiseGradientTextureSampler,
-				TStaticSamplerState<SF_Point,AM_Wrap,AM_Wrap,AM_Wrap>::GetRHI(),
-				Texture
-				);
-		}
-
-		if (PerlinNoise3DTexture.IsBound() && IsValidRef(GSystemTextures.PerlinNoise3D))
-		{
-			const FTexture3DRHIRef& Texture = (FTexture3DRHIRef&)GSystemTextures.PerlinNoise3D->GetRenderTargetItem().ShaderResourceTexture;
-			// Bind the PerlinNoise3DTexture as a texture
-			SetTextureParameter(
-				RHICmdList,
-				ShaderRHI,
-				PerlinNoise3DTexture,
-				PerlinNoise3DTextureSampler,
-				TStaticSamplerState<SF_Bilinear,AM_Wrap,AM_Wrap,AM_Wrap>::GetRHI(),
-				Texture
-				);
-		}
-	}
-
-	FTextureRHIRef& GetEyeAdaptation(const FSceneView& View);
-
-	// FShader interface.
-	virtual bool Serialize(FArchive& Ar);
-	virtual uint32 GetAllocatedSize() const override;
-
-private:
-
-	FShaderUniformBufferParameter MaterialUniformBuffer;
-	TArray<FShaderUniformBufferParameter> ParameterCollectionUniformBuffers;
-	TArray<TUniformParameter<FShaderResourceParameter> > Uniform2DTextureShaderResourceParameters;
-	TArray<TUniformParameter<FShaderResourceParameter> > Uniform2DTextureSamplerShaderResourceParameters;
-	TArray<TUniformParameter<FShaderResourceParameter> > UniformCubeTextureShaderResourceParameters;
-	TArray<TUniformParameter<FShaderResourceParameter> > UniformCubeTextureSamplerShaderResourceParameters;
-	FDeferredPixelShaderParameters DeferredParameters;
-	FShaderResourceParameter LightAttenuation;
-	FShaderResourceParameter LightAttenuationSampler;
-
-	// For materials using atmospheric fog color 
-	FAtmosphereShaderTextureParameters AtmosphericFogTextureParameters;
-	FPostProcessPassParameters PostprocessParameter;
-
-//Use of the eye adaptation texture here is experimental and potentially dangerous as it can introduce a feedback loop. May be removed.
-	FShaderResourceParameter EyeAdaptation;
-
-	/** The PerlinNoiseGradientTexture parameter for materials that use GradientNoise */
-	FShaderResourceParameter PerlinNoiseGradientTexture;
-	FShaderResourceParameter PerlinNoiseGradientTextureSampler;
-	/** The PerlinNoise3DTexture parameter for materials that use GradientNoise */
-	FShaderResourceParameter PerlinNoise3DTexture;
-	FShaderResourceParameter PerlinNoise3DTextureSampler;
-
-	FDebugUniformExpressionSet DebugUniformExpressionSet;
-	FString DebugDescription;
-
-	/** If true, cached uniform expressions are allowed. */
-	static int32 bAllowCachedUniformExpressions;
-	/** Console variable ref to toggle cached uniform expressions. */
-	static FAutoConsoleVariableRef CVarAllowCachedUniformExpressions;
-};
-
-/** Base class of all shaders that need material and vertex factory parameters. */
-class FMeshMaterialShader : public FMaterialShader
-{
-public:
-	FMeshMaterialShader() {}
-
-	FMeshMaterialShader(const FMeshMaterialShaderType::CompiledShaderInitializerType& Initializer)
-		:	FMaterialShader(Initializer)
-		,	VertexFactoryParameters(Initializer.VertexFactoryType,Initializer.ParameterMap,(EShaderFrequency)Initializer.Target.Frequency)
-	{}
-
-	template<typename ShaderRHIParamRef>
-	void SetParameters(FRHICommandList& RHICmdList, const ShaderRHIParamRef ShaderRHI,const FMaterialRenderProxy* MaterialRenderProxy,const FMaterial& Material,const FSceneView& View,ESceneRenderTargetsMode::Type TextureMode)
-	{
-		FMaterialShader::SetParameters(RHICmdList, ShaderRHI,MaterialRenderProxy,Material,View,false,TextureMode);
-	}
-
-	void SetVFParametersOnly(FRHICommandList& RHICmdList, const FVertexFactory* VertexFactory,const FSceneView& View,const FMeshBatchElement& BatchElement)
-	{
-		VertexFactoryParameters.SetMesh(RHICmdList, this,VertexFactory,View,BatchElement, 0);
-	}
-
-	template<typename ShaderRHIParamRef>
-	void SetMesh(FRHICommandList& RHICmdList, const ShaderRHIParamRef ShaderRHI,const FVertexFactory* VertexFactory,const FSceneView& View,const FPrimitiveSceneProxy* Proxy,const FMeshBatchElement& BatchElement,uint32 DataFlags = 0)
-	{
-		// Set the mesh for the vertex factory
-		VertexFactoryParameters.SetMesh(RHICmdList, this,VertexFactory,View,BatchElement, DataFlags);
-		
-		if(IsValidRef(BatchElement.PrimitiveUniformBuffer))
-		{
-			SetUniformBufferParameter(RHICmdList, ShaderRHI,GetUniformBufferParameter<FPrimitiveUniformShaderParameters>(),BatchElement.PrimitiveUniformBuffer);
-		}
-		else
-		{
-			check(BatchElement.PrimitiveUniformBufferResource);
-			SetUniformBufferParameter(RHICmdList, ShaderRHI,GetUniformBufferParameter<FPrimitiveUniformShaderParameters>(),*BatchElement.PrimitiveUniformBufferResource);
-		}
-
-		TShaderUniformBufferParameter<FDistanceCullFadeUniformShaderParameters> LODParameter = GetUniformBufferParameter<FDistanceCullFadeUniformShaderParameters>();
-		if( LODParameter.IsBound() )
-		{
-			SetUniformBufferParameter(RHICmdList, ShaderRHI,LODParameter,GetPrimitiveFadeUniformBufferParameter(View, Proxy));
-		}
-	}
-
-	/**
-	 * Retrieves the fade uniform buffer parameter from a FSceneViewState for the primitive
-	 * This code was moved from SetMesh() to work around the template first-use vs first-seen differences between MSVC and others
-	 */
-	FUniformBufferRHIParamRef GetPrimitiveFadeUniformBufferParameter(const FSceneView& View, const FPrimitiveSceneProxy* Proxy);
-
-
-	// FShader interface.
-	virtual const FVertexFactoryParameterRef* GetVertexFactoryParameterRef() const { return &VertexFactoryParameters; }
-	virtual bool Serialize(FArchive& Ar);
-	virtual uint32 GetAllocatedSize() const override;
-
-private:
-	FVertexFactoryParameterRef VertexFactoryParameters;
-};
-
 /** Base Hull shader for drawing policy rendering */
 class FBaseHS : public FMeshMaterialShader
 {
@@ -478,12 +53,12 @@ public:
 
 	void SetParameters(FRHICommandList& RHICmdList, const FMaterialRenderProxy* MaterialRenderProxy,const FSceneView& View)
 	{
-		FMeshMaterialShader::SetParameters(RHICmdList, GetHullShader(), MaterialRenderProxy, *MaterialRenderProxy->GetMaterial(View.GetFeatureLevel()), View, ESceneRenderTargetsMode::SetTextures);
+		FMeshMaterialShader::SetParameters(RHICmdList, (FHullShaderRHIParamRef)GetHullShader(), MaterialRenderProxy, *MaterialRenderProxy->GetMaterial(View.GetFeatureLevel()), View, ESceneRenderTargetsMode::SetTextures);
 	}
 
 	void SetMesh(FRHICommandList& RHICmdList, const FVertexFactory* VertexFactory,const FSceneView& View,const FPrimitiveSceneProxy* Proxy,const FMeshBatchElement& BatchElement)
 	{
-		FMeshMaterialShader::SetMesh(RHICmdList, GetHullShader(),VertexFactory,View,Proxy,BatchElement);
+		FMeshMaterialShader::SetMesh(RHICmdList, (FHullShaderRHIParamRef)GetHullShader(),VertexFactory,View,Proxy,BatchElement);
 	}
 };
 
@@ -524,12 +99,12 @@ public:
 
 	void SetParameters(FRHICommandList& RHICmdList, const FMaterialRenderProxy* MaterialRenderProxy,const FSceneView& View)
 	{
-		FMeshMaterialShader::SetParameters(RHICmdList, GetDomainShader(), MaterialRenderProxy, *MaterialRenderProxy->GetMaterial(View.GetFeatureLevel()), View, ESceneRenderTargetsMode::SetTextures);
+		FMeshMaterialShader::SetParameters(RHICmdList, (FDomainShaderRHIParamRef)GetDomainShader(), MaterialRenderProxy, *MaterialRenderProxy->GetMaterial(View.GetFeatureLevel()), View, ESceneRenderTargetsMode::SetTextures);
 	}
 
 	void SetMesh(FRHICommandList& RHICmdList, const FVertexFactory* VertexFactory,const FSceneView& View,const FPrimitiveSceneProxy* Proxy,const FMeshBatchElement& BatchElement)
 	{
-		FMeshMaterialShader::SetMesh(RHICmdList, GetDomainShader(),VertexFactory,View,Proxy,BatchElement);
+		FMeshMaterialShader::SetMesh(RHICmdList, (FDomainShaderRHIParamRef)GetDomainShader(),VertexFactory,View,Proxy,BatchElement);
 	}
 };
 
diff --git a/Engine/Source/Runtime/Renderer/Private/ShadowRendering.h b/Engine/Source/Runtime/Renderer/Private/ShadowRendering.h
--- a/Engine/Source/Runtime/Renderer/Private/ShadowRendering.h
+++ b/Engine/Source/Runtime/Renderer/Private/ShadowRendering.h
@@ -1352,7 +1352,6 @@ public:
 	static void ModifyCompilationEnvironment(EShaderPlatform Platform, FShaderCompilerEnvironment& OutEnvironment)
 	{
 		FGlobalShader::ModifyCompilationEnvironment(Platform,OutEnvironment);
-		FDeferredPixelShaderParameters::ModifyCompilationEnvironment(Platform,OutEnvironment);
 		OutEnvironment.SetDefine(TEXT("SHADOW_QUALITY"), Quality);
 	}
 
diff --git a/Engine/Source/Runtime/Renderer/Private/TiledDeferredLightRendering.cpp b/Engine/Source/Runtime/Renderer/Private/TiledDeferredLightRendering.cpp
--- a/Engine/Source/Runtime/Renderer/Private/TiledDeferredLightRendering.cpp
+++ b/Engine/Source/Runtime/Renderer/Private/TiledDeferredLightRendering.cpp
@@ -80,7 +80,6 @@ public:
 	static void ModifyCompilationEnvironment(EShaderPlatform Platform, FShaderCompilerEnvironment& OutEnvironment)
 	{
 		FGlobalShader::ModifyCompilationEnvironment(Platform, OutEnvironment);
-		FDeferredPixelShaderParameters::ModifyCompilationEnvironment(Platform,OutEnvironment);
 		OutEnvironment.SetDefine(TEXT("THREADGROUP_SIZEX"), GDeferredLightTileSizeX);
 		OutEnvironment.SetDefine(TEXT("THREADGROUP_SIZEY"), GDeferredLightTileSizeY);
 		OutEnvironment.SetDefine(TEXT("MAX_LIGHTS"), GMaxNumTiledDeferredLights);
diff --git a/Engine/Source/Runtime/Renderer/Private/VelocityRendering.cpp b/Engine/Source/Runtime/Renderer/Private/VelocityRendering.cpp
--- a/Engine/Source/Runtime/Renderer/Private/VelocityRendering.cpp
+++ b/Engine/Source/Runtime/Renderer/Private/VelocityRendering.cpp
@@ -93,7 +93,7 @@ public:
 
 	void SetParameters(FRHICommandList& RHICmdList, const FMaterialRenderProxy* MaterialRenderProxy,const FViewInfo& View)
 	{
-		FMeshMaterialShader::SetParameters(RHICmdList, GetDomainShader(), MaterialRenderProxy, *MaterialRenderProxy->GetMaterial(View.GetFeatureLevel()), View, ESceneRenderTargetsMode::DontSet);
+		FMeshMaterialShader::SetParameters(RHICmdList, (FDomainShaderRHIParamRef)GetDomainShader(), MaterialRenderProxy, *MaterialRenderProxy->GetMaterial(View.GetFeatureLevel()), View, ESceneRenderTargetsMode::DontSet);
 	}
 
 protected:
diff --git a/Engine/Source/Runtime/Renderer/Public/AtmosphereTextureParameters.h b/Engine/Source/Runtime/Renderer/Public/AtmosphereTextureParameters.h
new file mode 100644
--- /dev/null
+++ b/Engine/Source/Runtime/Renderer/Public/AtmosphereTextureParameters.h
@@ -0,0 +1,27 @@
+// Copyright 1998-2014 Epic Games, Inc. All Rights Reserved.
+
+/*=============================================================================
+	AtmosphereTextureParameters.h: Shader base classes
+=============================================================================*/
+
+#pragma once
+
+/** Shader parameters needed for atmosphere passes. */
+class FAtmosphereShaderTextureParameters
+{
+public:
+	void Bind(const FShaderParameterMap& ParameterMap);
+
+	template< typename ShaderRHIParamRef >
+	void Set( FRHICommandList& RHICmdList, const ShaderRHIParamRef ShaderRHI, const FSceneView& View ) const;
+
+	friend FArchive& operator<<(FArchive& Ar,FAtmosphereShaderTextureParameters& P);
+
+private:
+	FShaderResourceParameter TransmittanceTexture;
+	FShaderResourceParameter TransmittanceTextureSampler;
+	FShaderResourceParameter IrradianceTexture;
+	FShaderResourceParameter IrradianceTextureSampler;
+	FShaderResourceParameter InscatterTexture;
+	FShaderResourceParameter InscatterTextureSampler;
+};
\ No newline at end of file
diff --git a/Engine/Source/Runtime/Renderer/Public/MaterialShader.h b/Engine/Source/Runtime/Renderer/Public/MaterialShader.h
new file mode 100644
--- /dev/null
+++ b/Engine/Source/Runtime/Renderer/Public/MaterialShader.h
@@ -0,0 +1,160 @@
+// Copyright 1998-2014 Epic Games, Inc. All Rights Reserved.
+
+/*=============================================================================
+	MaterialShader.h: Shader base classes
+=============================================================================*/
+
+#pragma once
+
+#include "MaterialShaderType.h"
+#include "AtmosphereTextureParameters.h"
+#include "SceneRenderTargetParameters.h"
+#include "PostProcessParameters.h"
+
+template<typename ParameterType> 
+struct TUniformParameter
+{
+	int32 Index;
+	ParameterType ShaderParameter;
+	friend FArchive& operator<<(FArchive& Ar,TUniformParameter<ParameterType>& P)
+	{
+		return Ar << P.Index << P.ShaderParameter;
+	}
+};
+
+
+/**
+ * Debug information related to uniform expression sets.
+ */
+class FDebugUniformExpressionSet
+{
+public:
+	/** The number of each type of expression contained in the set. */
+	int32 NumVectorExpressions;
+	int32 NumScalarExpressions;
+	int32 Num2DTextureExpressions;
+	int32 NumCubeTextureExpressions;
+
+	FDebugUniformExpressionSet()
+		: NumVectorExpressions(0)
+		, NumScalarExpressions(0)
+		, Num2DTextureExpressions(0)
+		, NumCubeTextureExpressions(0)
+	{
+	}
+
+	explicit FDebugUniformExpressionSet(const FUniformExpressionSet& InUniformExpressionSet)
+	{
+		InitFromExpressionSet(InUniformExpressionSet);
+	}
+
+	/** Initialize from a uniform expression set. */
+	void InitFromExpressionSet(const FUniformExpressionSet& InUniformExpressionSet)
+	{
+		NumVectorExpressions = InUniformExpressionSet.UniformVectorExpressions.Num();
+		NumScalarExpressions = InUniformExpressionSet.UniformScalarExpressions.Num();
+		Num2DTextureExpressions = InUniformExpressionSet.Uniform2DTextureExpressions.Num();
+		NumCubeTextureExpressions = InUniformExpressionSet.UniformCubeTextureExpressions.Num();
+	}
+
+	/** Returns true if the number of uniform expressions matches those with which the debug set was initialized. */
+	bool Matches(const FUniformExpressionSet& InUniformExpressionSet) const
+	{
+		return NumVectorExpressions == InUniformExpressionSet.UniformVectorExpressions.Num()
+			&& NumScalarExpressions == InUniformExpressionSet.UniformScalarExpressions.Num()
+			&& Num2DTextureExpressions == InUniformExpressionSet.Uniform2DTextureExpressions.Num()
+			&& NumCubeTextureExpressions == InUniformExpressionSet.UniformCubeTextureExpressions.Num();
+	}
+};
+
+/** Serialization for debug uniform expression sets. */
+inline FArchive& operator<<(FArchive& Ar, FDebugUniformExpressionSet& DebugExpressionSet)
+{
+	Ar << DebugExpressionSet.NumVectorExpressions;
+	Ar << DebugExpressionSet.NumScalarExpressions;
+	Ar << DebugExpressionSet.Num2DTextureExpressions;
+	Ar << DebugExpressionSet.NumCubeTextureExpressions;
+	return Ar;
+}
+
+
+/** Base class of all shaders that need material parameters. */
+class FMaterialShader : public FShader
+{
+public:
+	FMaterialShader() {}
+
+	FMaterialShader(const FMaterialShaderType::CompiledShaderInitializerType& Initializer);
+
+	typedef void (*ModifyCompilationEnvironmentType)(EShaderPlatform, const FMaterial*, FShaderCompilerEnvironment&);
+
+	ENGINE_API static void ModifyCompilationEnvironment(EShaderPlatform Platform, const FMaterial* Material, FShaderCompilerEnvironment& OutEnvironment)
+	{
+	}
+
+	FUniformBufferRHIParamRef GetParameterCollectionBuffer(const FGuid& Id, const FSceneInterface* SceneInterface) const;
+
+	template<typename ShaderRHIParamRef>
+	void SetParameters(FRHICommandList& RHICmdList, const ShaderRHIParamRef ShaderRHI,const FSceneView& View)
+	{
+		check(GetUniformBufferParameter<FViewUniformShaderParameters>().IsInitialized());
+		CheckShaderIsValid();
+		SetUniformBufferParameter(RHICmdList, ShaderRHI,GetUniformBufferParameter<FViewUniformShaderParameters>(),View.UniformBuffer);
+	}
+
+	static void ModifyCompilationEnvironment(EShaderPlatform Platform, FShaderCompilerEnvironment& OutEnvironment)
+	{
+		FGlobalShader::ModifyCompilationEnvironment(Platform,OutEnvironment);
+	}
+
+	/** Sets pixel parameters that are material specific but not FMeshBatch specific. */
+	template< typename ShaderRHIParamRef >
+	void SetParameters(
+		FRHICommandList& RHICmdList,
+		const ShaderRHIParamRef ShaderRHI, 
+		const FMaterialRenderProxy* MaterialRenderProxy, 
+		const FMaterial& Material,
+		const FSceneView& View, 
+		bool bDeferredPass, 
+		ESceneRenderTargetsMode::Type TextureMode );
+
+	FTextureRHIRef& GetEyeAdaptation(const FSceneView& View);
+
+	// FShader interface.
+	virtual bool Serialize(FArchive& Ar);
+	virtual uint32 GetAllocatedSize() const override;
+
+private:
+
+	FShaderUniformBufferParameter MaterialUniformBuffer;
+	TArray<FShaderUniformBufferParameter> ParameterCollectionUniformBuffers;
+	TArray<TUniformParameter<FShaderResourceParameter> > Uniform2DTextureShaderResourceParameters;
+	TArray<TUniformParameter<FShaderResourceParameter> > Uniform2DTextureSamplerShaderResourceParameters;
+	TArray<TUniformParameter<FShaderResourceParameter> > UniformCubeTextureShaderResourceParameters;
+	TArray<TUniformParameter<FShaderResourceParameter> > UniformCubeTextureSamplerShaderResourceParameters;
+	FDeferredPixelShaderParameters DeferredParameters;
+	FShaderResourceParameter LightAttenuation;
+	FShaderResourceParameter LightAttenuationSampler;
+
+	// For materials using atmospheric fog color 
+	FAtmosphereShaderTextureParameters AtmosphericFogTextureParameters;
+	FPostProcessPassParameters PostprocessParameter;
+
+	//Use of the eye adaptation texture here is experimental and potentially dangerous as it can introduce a feedback loop. May be removed.
+	FShaderResourceParameter EyeAdaptation;
+
+	/** The PerlinNoiseGradientTexture parameter for materials that use GradientNoise */
+	FShaderResourceParameter PerlinNoiseGradientTexture;
+	FShaderResourceParameter PerlinNoiseGradientTextureSampler;
+	/** The PerlinNoise3DTexture parameter for materials that use GradientNoise */
+	FShaderResourceParameter PerlinNoise3DTexture;
+	FShaderResourceParameter PerlinNoise3DTextureSampler;
+
+	FDebugUniformExpressionSet DebugUniformExpressionSet;
+	FString DebugDescription;
+
+	/** If true, cached uniform expressions are allowed. */
+	static int32 bAllowCachedUniformExpressions;
+	/** Console variable ref to toggle cached uniform expressions. */
+	static FAutoConsoleVariableRef CVarAllowCachedUniformExpressions;
+};
\ No newline at end of file
diff --git a/Engine/Source/Runtime/Renderer/Public/MeshMaterialShader.h b/Engine/Source/Runtime/Renderer/Public/MeshMaterialShader.h
new file mode 100644
--- /dev/null
+++ b/Engine/Source/Runtime/Renderer/Public/MeshMaterialShader.h
@@ -0,0 +1,76 @@
+// Copyright 1998-2014 Epic Games, Inc. All Rights Reserved.
+
+/*=============================================================================
+	MeshMaterialShader.h: Shader base classes
+=============================================================================*/
+
+#pragma once
+
+#include "MaterialShader.h"
+#include "MeshMaterialShaderType.h"
+
+/** Base class of all shaders that need material and vertex factory parameters. */
+class FMeshMaterialShader : public FMaterialShader
+{
+public:
+	FMeshMaterialShader() {}
+
+	FMeshMaterialShader(const FMeshMaterialShaderType::CompiledShaderInitializerType& Initializer)
+		:	FMaterialShader(Initializer)
+		,	VertexFactoryParameters(Initializer.VertexFactoryType,Initializer.ParameterMap,(EShaderFrequency)Initializer.Target.Frequency)
+	{}
+
+	template< typename ShaderRHIParamRef >
+	void SetParameters(
+		FRHICommandList& RHICmdList,
+		const ShaderRHIParamRef ShaderRHI,
+		const FMaterialRenderProxy* MaterialRenderProxy,
+		const FMaterial& Material,
+		const FSceneView& View,
+		ESceneRenderTargetsMode::Type TextureMode )
+	{
+		FMaterialShader::SetParameters(RHICmdList, ShaderRHI,MaterialRenderProxy,Material,View,false,TextureMode);
+	}
+
+	void SetVFParametersOnly(FRHICommandList& RHICmdList, const FVertexFactory* VertexFactory,const FSceneView& View,const FMeshBatchElement& BatchElement)
+	{
+		VertexFactoryParameters.SetMesh(RHICmdList, this,VertexFactory,View,BatchElement, 0);
+	}
+
+	template<typename ShaderRHIParamRef>
+	void SetMesh(FRHICommandList& RHICmdList, const ShaderRHIParamRef ShaderRHI,const FVertexFactory* VertexFactory,const FSceneView& View,const FPrimitiveSceneProxy* Proxy,const FMeshBatchElement& BatchElement,uint32 DataFlags = 0)
+	{
+		// Set the mesh for the vertex factory
+		VertexFactoryParameters.SetMesh(RHICmdList, this,VertexFactory,View,BatchElement, DataFlags);
+		
+		if(IsValidRef(BatchElement.PrimitiveUniformBuffer))
+		{
+			SetUniformBufferParameter(RHICmdList, ShaderRHI,GetUniformBufferParameter<FPrimitiveUniformShaderParameters>(),BatchElement.PrimitiveUniformBuffer);
+		}
+		else
+		{
+			check(BatchElement.PrimitiveUniformBufferResource);
+			SetUniformBufferParameter(RHICmdList, ShaderRHI,GetUniformBufferParameter<FPrimitiveUniformShaderParameters>(),*BatchElement.PrimitiveUniformBufferResource);
+		}
+
+		TShaderUniformBufferParameter<FDistanceCullFadeUniformShaderParameters> LODParameter = GetUniformBufferParameter<FDistanceCullFadeUniformShaderParameters>();
+		if( LODParameter.IsBound() )
+		{
+			SetUniformBufferParameter(RHICmdList, ShaderRHI,LODParameter,GetPrimitiveFadeUniformBufferParameter(View, Proxy));
+		}
+	}
+
+	/**
+	 * Retrieves the fade uniform buffer parameter from a FSceneViewState for the primitive
+	 * This code was moved from SetMesh() to work around the template first-use vs first-seen differences between MSVC and others
+	 */
+	FUniformBufferRHIParamRef GetPrimitiveFadeUniformBufferParameter(const FSceneView& View, const FPrimitiveSceneProxy* Proxy);
+
+	// FShader interface.
+	virtual const FVertexFactoryParameterRef* GetVertexFactoryParameterRef() const { return &VertexFactoryParameters; }
+	virtual bool Serialize(FArchive& Ar);
+	virtual uint32 GetAllocatedSize() const override;
+
+private:
+	FVertexFactoryParameterRef VertexFactoryParameters;
+};
\ No newline at end of file
diff --git a/Engine/Source/Runtime/Renderer/Public/PostProcessParameters.h b/Engine/Source/Runtime/Renderer/Public/PostProcessParameters.h
new file mode 100644
--- /dev/null
+++ b/Engine/Source/Runtime/Renderer/Public/PostProcessParameters.h
@@ -0,0 +1,100 @@
+// Copyright 1998-2014 Epic Games, Inc. All Rights Reserved.
+
+/*=============================================================================
+	PostProcessParameters.h: Shader base classes
+=============================================================================*/
+
+#pragma once
+
+#include "RHIStaticStates.h"
+
+struct FRenderingCompositePassContext;
+
+// This is the index for the texture input of this pass.
+// More that that should not be needed.
+// Could be an uint32 but for better readability and type safety it's an enum.
+// Counting starts from 0 in consecutive order.
+enum EPassInputId
+{
+	ePId_Input0,
+	ePId_Input1,
+	ePId_Input2,
+	ePId_Input3,
+	ePId_Input4,
+	ePId_Input5,
+	ePId_Input6,
+/*	ePId_Input7,
+	ePId_Input8,
+	ePId_Input9,
+	ePId_Input10,
+	ePId_Input11,
+	ePId_Input12,
+	ePId_Input13,
+	ePId_Input14,
+	ePId_Input15*/
+
+	// to get the total count of inputs
+	ePId_Input_MAX
+};
+
+// Usually the same as the MRT number but it doesn't have to be implemented as MRT.
+// More that that should not be needed.
+// Could be an uint32 but for better readability and type safety it's an enum.
+// Counting starts from 0 in consecutive order.
+enum EPassOutputId
+{
+	ePId_Output0,
+	ePId_Output1,
+	ePId_Output2,
+	ePId_Output3,
+	ePId_Output4,
+	ePId_Output5,
+	ePId_Output6,
+	ePId_Output7
+};
+
+// currently hard coded to 4 input textures
+// convenient but not the most optimized solution
+struct FPostProcessPassParameters 
+{
+	/** Initialization constructor. */
+	void Bind(const FShaderParameterMap& ParameterMap);
+
+	/** Set the pixel shader parameter values. */
+	void SetPS(const FPixelShaderRHIParamRef& ShaderRHI, const FRenderingCompositePassContext& Context, FSamplerStateRHIParamRef Filter = TStaticSamplerState<>::GetRHI(), bool bWhiteIfNoTexture = false, FSamplerStateRHIParamRef* FilterOverrideArray = 0);
+
+	/** Set the compute shader parameter values. */
+	void SetCS(const FComputeShaderRHIParamRef& ShaderRHI, const FRenderingCompositePassContext& Context, FSamplerStateRHIParamRef Filter = TStaticSamplerState<>::GetRHI(), bool bWhiteIfNoTexture = false, FSamplerStateRHIParamRef* FilterOverrideArray = 0);
+
+	/** Set the vertex shader parameter values. */
+	void SetVS(const FVertexShaderRHIParamRef& ShaderRHI, const FRenderingCompositePassContext& Context, FSamplerStateRHIParamRef Filter = TStaticSamplerState<>::GetRHI(), bool bWhiteIfNoTexture = false, FSamplerStateRHIParamRef* FilterOverrideArray = 0);
+
+	/** Serializer. */
+	friend FArchive& operator<<(FArchive& Ar,FPostProcessPassParameters& P);
+
+private:
+
+	FShaderParameter ViewportSize;
+	FShaderParameter ViewportRect;
+	FShaderResourceParameter PostprocessInputParameter[ePId_Input_MAX];
+	FShaderResourceParameter PostprocessInputParameterSampler[ePId_Input_MAX];
+	FShaderResourceParameter PostprocessInputNew[ePId_Input_MAX];
+	FShaderParameter PostprocessInputSizeParameter[ePId_Input_MAX];
+	FShaderParameter PostProcessInputMinMaxParameter[ePId_Input_MAX];
+	FShaderParameter ScreenPosToPixel;
+	FShaderResourceParameter BilinearTextureSampler0;
+	FShaderResourceParameter BilinearTextureSampler1;
+
+public:
+	// @param Filter can be 0 if FilterOverrideArray is used
+	// @param FilterOverrideArray can be 0 if Filter is used
+	template< typename ShaderRHIParamRef >
+	void Set(
+		FRHICommandList& RHICmdList,
+		const ShaderRHIParamRef& ShaderRHI,
+		const FRenderingCompositePassContext& Context,
+		FSamplerStateRHIParamRef Filter,
+		bool bWhiteIfNoTexture = false,
+		FSamplerStateRHIParamRef* FilterOverrideArray = 0
+	);
+};
\ No newline at end of file
diff --git a/Engine/Source/Runtime/Renderer/Public/SceneRenderTargetParameters.h b/Engine/Source/Runtime/Renderer/Public/SceneRenderTargetParameters.h
new file mode 100644
--- /dev/null
+++ b/Engine/Source/Runtime/Renderer/Public/SceneRenderTargetParameters.h
@@ -0,0 +1,97 @@
+// Copyright 1998-2014 Epic Games, Inc. All Rights Reserved.
+
+/*=============================================================================
+	SceneRenderTargetParameters.h: Shader base classes
+=============================================================================*/
+
+#pragma once
+
+namespace ESceneRenderTargetsMode
+{
+	enum Type
+	{
+		SetTextures,
+		DontSet,
+		DontSetIgnoreBoundByEditorCompositing,
+		NonSceneAlignedPass,
+	};
+}
+
+/** Encapsulates scene texture shader parameter bindings. */
+class FSceneTextureShaderParameters
+{
+public:
+	/** Binds the parameters using a compiled shader's parameter map. */
+	void Bind(const FShaderParameterMap& ParameterMap);
+
+	/** Sets the scene texture parameters for the given view. */
+	template< typename ShaderRHIParamRef >
+	void Set(
+		FRHICommandList& RHICmdList,
+		const ShaderRHIParamRef& ShaderRHI,
+		const FSceneView& View,
+		ESceneRenderTargetsMode::Type TextureMode = ESceneRenderTargetsMode::SetTextures,
+		ESamplerFilter ColorFilter = SF_Point
+	) const;
+
+	/** Serializer. */
+	friend FArchive& operator<<(FArchive& Ar,FSceneTextureShaderParameters& P);
+
+private:
+	/** The SceneColorTexture parameter for materials that use SceneColor */
+	FShaderResourceParameter SceneColorTextureParameter;
+	FShaderResourceParameter SceneColorTextureParameterSampler;
+	/** The SceneDepthTexture parameter for materials that use SceneDepth */
+	FShaderResourceParameter SceneDepthTextureParameter;
+	FShaderResourceParameter SceneDepthTextureParameterSampler;
+	/** The SceneAlphaCopyTexture parameter for materials that use SceneAlphaCopy */
+	FShaderResourceParameter SceneAlphaCopyTextureParameter;
+	FShaderResourceParameter SceneAlphaCopyTextureParameterSampler;
+
+	/** for MSAA access to the scene color */
+	FShaderResourceParameter SceneColorSurfaceParameter;
+	/** The SceneColorTextureMSAA parameter for materials that use SceneColorTextureMSAA */
+	FShaderResourceParameter SceneDepthSurfaceParameter;
+	/**  */
+	FShaderResourceParameter SceneDepthTextureNonMS;
+};
+
+/** Pixel shader parameters needed for deferred passes. */
+class FDeferredPixelShaderParameters
+{
+public:
+	void Bind(const FShaderParameterMap& ParameterMap);
+
+	template< typename ShaderRHIParamRef >
+	void Set(
+		FRHICommandList& RHICmdList,
+		const ShaderRHIParamRef ShaderRHI,
+		const FSceneView& View,
+		ESceneRenderTargetsMode::Type TextureMode = ESceneRenderTargetsMode::SetTextures
+	) const;
+
+	friend FArchive& operator<<(FArchive& Ar,FDeferredPixelShaderParameters& P);
+
+private:
+	FSceneTextureShaderParameters SceneTextureParameters;
+	FShaderUniformBufferParameter GBufferResources;
+	FShaderResourceParameter DBufferATextureMS;
+	FShaderResourceParameter DBufferBTextureMS;
+	FShaderResourceParameter DBufferCTextureMS;
+	FShaderResourceParameter ScreenSpaceAOTextureMS;
+	FShaderResourceParameter DBufferATextureNonMS;
+	FShaderResourceParameter DBufferBTextureNonMS;
+	FShaderResourceParameter DBufferCTextureNonMS;
+	FShaderResourceParameter ScreenSpaceAOTextureNonMS;
+	FShaderResourceParameter CustomDepthTextureNonMS;
+	FShaderResourceParameter DBufferATexture;
+	FShaderResourceParameter DBufferATextureSampler;
+	FShaderResourceParameter DBufferBTexture;
+	FShaderResourceParameter DBufferBTextureSampler;
+	FShaderResourceParameter DBufferCTexture;
+	FShaderResourceParameter DBufferCTextureSampler;
+	FShaderResourceParameter ScreenSpaceAOTexture;
+	FShaderResourceParameter ScreenSpaceAOTextureSampler;
+	FShaderResourceParameter CustomDepthTexture;
+	FShaderResourceParameter CustomDepthTextureSampler;
+};
\ No newline at end of file
