Fix for Mac. Not sure why others compile correctly.

diff --git a/Engine/Source/Runtime/Renderer/Private/ShaderBaseClasses.cpp b/Engine/Source/Runtime/Renderer/Private/ShaderBaseClasses.cpp
--- a/Engine/Source/Runtime/Renderer/Private/ShaderBaseClasses.cpp
+++ b/Engine/Source/Runtime/Renderer/Private/ShaderBaseClasses.cpp
@@ -335,6 +335,55 @@ uint32 FMaterialShader::GetAllocatedSize() const
 		+ DebugDescription.GetAllocatedSize();
 }
 
+
+template< typename ShaderRHIParamRef >
+void FMeshMaterialShader::SetMesh(
+	FRHICommandList& RHICmdList,
+	const ShaderRHIParamRef ShaderRHI,
+	const FVertexFactory* VertexFactory,
+	const FSceneView& View,
+	const FPrimitiveSceneProxy* Proxy,
+	const FMeshBatchElement& BatchElement,
+	uint32 DataFlags )
+{
+	// Set the mesh for the vertex factory
+	VertexFactoryParameters.SetMesh(RHICmdList, this,VertexFactory,View,BatchElement, DataFlags);
+		
+	if(IsValidRef(BatchElement.PrimitiveUniformBuffer))
+	{
+		SetUniformBufferParameter(RHICmdList, ShaderRHI,GetUniformBufferParameter<FPrimitiveUniformShaderParameters>(),BatchElement.PrimitiveUniformBuffer);
+	}
+	else
+	{
+		check(BatchElement.PrimitiveUniformBufferResource);
+		SetUniformBufferParameter(RHICmdList, ShaderRHI,GetUniformBufferParameter<FPrimitiveUniformShaderParameters>(),*BatchElement.PrimitiveUniformBufferResource);
+	}
+
+	TShaderUniformBufferParameter<FDistanceCullFadeUniformShaderParameters> LODParameter = GetUniformBufferParameter<FDistanceCullFadeUniformShaderParameters>();
+	if( LODParameter.IsBound() )
+	{
+		SetUniformBufferParameter(RHICmdList, ShaderRHI,LODParameter,GetPrimitiveFadeUniformBufferParameter(View, Proxy));
+	}
+}
+
+#define IMPLEMENT_MESH_MATERIAL_SHADER_SetMesh( ShaderRHIParamRef ) \
+	template void FMeshMaterialShader::SetMesh< ShaderRHIParamRef >( \
+		FRHICommandList& RHICmdList,			\
+		const ShaderRHIParamRef ShaderRHI,		\
+		const FVertexFactory* VertexFactory,	\
+		const FSceneView& View,					\
+		const FPrimitiveSceneProxy* Proxy,		\
+		const FMeshBatchElement& BatchElement,	\
+		uint32 DataFlags						\
+	);
+
+IMPLEMENT_MESH_MATERIAL_SHADER_SetMesh( FVertexShaderRHIParamRef );
+IMPLEMENT_MESH_MATERIAL_SHADER_SetMesh( FHullShaderRHIParamRef );
+IMPLEMENT_MESH_MATERIAL_SHADER_SetMesh( FDomainShaderRHIParamRef );
+IMPLEMENT_MESH_MATERIAL_SHADER_SetMesh( FGeometryShaderRHIParamRef );
+IMPLEMENT_MESH_MATERIAL_SHADER_SetMesh( FPixelShaderRHIParamRef );
+IMPLEMENT_MESH_MATERIAL_SHADER_SetMesh( FComputeShaderRHIParamRef );
+
 bool FMeshMaterialShader::Serialize(FArchive& Ar)
 {
 	bool bShaderHasOutdatedParameters = FMaterialShader::Serialize(Ar);
diff --git a/Engine/Source/Runtime/Renderer/Public/MeshMaterialShader.h b/Engine/Source/Runtime/Renderer/Public/MeshMaterialShader.h
--- a/Engine/Source/Runtime/Renderer/Public/MeshMaterialShader.h
+++ b/Engine/Source/Runtime/Renderer/Public/MeshMaterialShader.h
@@ -37,28 +37,16 @@ public:
 		VertexFactoryParameters.SetMesh(RHICmdList, this,VertexFactory,View,BatchElement, 0);
 	}
 
-	template<typename ShaderRHIParamRef>
-	void SetMesh(FRHICommandList& RHICmdList, const ShaderRHIParamRef ShaderRHI,const FVertexFactory* VertexFactory,const FSceneView& View,const FPrimitiveSceneProxy* Proxy,const FMeshBatchElement& BatchElement,uint32 DataFlags = 0)
-	{
-		// Set the mesh for the vertex factory
-		VertexFactoryParameters.SetMesh(RHICmdList, this,VertexFactory,View,BatchElement, DataFlags);
-		
-		if(IsValidRef(BatchElement.PrimitiveUniformBuffer))
-		{
-			SetUniformBufferParameter(RHICmdList, ShaderRHI,GetUniformBufferParameter<FPrimitiveUniformShaderParameters>(),BatchElement.PrimitiveUniformBuffer);
-		}
-		else
-		{
-			check(BatchElement.PrimitiveUniformBufferResource);
-			SetUniformBufferParameter(RHICmdList, ShaderRHI,GetUniformBufferParameter<FPrimitiveUniformShaderParameters>(),*BatchElement.PrimitiveUniformBufferResource);
-		}
-
-		TShaderUniformBufferParameter<FDistanceCullFadeUniformShaderParameters> LODParameter = GetUniformBufferParameter<FDistanceCullFadeUniformShaderParameters>();
-		if( LODParameter.IsBound() )
-		{
-			SetUniformBufferParameter(RHICmdList, ShaderRHI,LODParameter,GetPrimitiveFadeUniformBufferParameter(View, Proxy));
-		}
-	}
+	template< typename ShaderRHIParamRef >
+	void SetMesh(
+		FRHICommandList& RHICmdList,
+		const ShaderRHIParamRef ShaderRHI,
+		const FVertexFactory* VertexFactory,
+		const FSceneView& View,
+		const FPrimitiveSceneProxy* Proxy,
+		const FMeshBatchElement& BatchElement,
+		uint32 DataFlags = 0
+	);
 
 	/**
 	 * Retrieves the fade uniform buffer parameter from a FSceneViewState for the primitive
