UE4 - More RHICmdList feature complete; added SetDepthStencilState RHI Cmd

diff --git a/Engine/Source/Runtime/RHI/Private/RHICommandList.cpp b/Engine/Source/Runtime/RHI/Private/RHICommandList.cpp
--- a/Engine/Source/Runtime/RHI/Private/RHICommandList.cpp
+++ b/Engine/Source/Runtime/RHI/Private/RHICommandList.cpp
@@ -219,6 +219,16 @@ void FRHICommandListExecutor::ExecuteList(FRHICommandList& CmdList)
 				RHICmd->VertexBuffer->Release();
 			}
 			break;
+		case ERCT_SetDepthStencilState:
+			{
+				auto* RHICmd = Iter.NextCommand<FRHICommandSetDepthStencilState>();
+				if (!bOnlyTestMemAccess)
+				{
+					RHISetDepthStencilState(RHICmd->State, RHICmd->StencilRef);
+				}
+				RHICmd->State->Release();
+			}
+			break;
 		default:
 			checkf(0, TEXT("Unknown RHI Command %d!"), Cmd->Type);
 		}
diff --git a/Engine/Source/Runtime/RHI/Public/RHICommandList.h b/Engine/Source/Runtime/RHI/Public/RHICommandList.h
--- a/Engine/Source/Runtime/RHI/Public/RHICommandList.h
+++ b/Engine/Source/Runtime/RHI/Public/RHICommandList.h
@@ -25,6 +25,7 @@ enum ERHICommandType
 	ERCT_SetBoundShaderState,
 	ERCT_SetBlendState,
 	ERCT_SetStreamSource,
+	ERCT_SetDepthStencilState,
 
 	ERCT_,
 };
@@ -223,11 +224,7 @@ public:
 		RHISetViewport(MinX, MinY, MinZ, MaxX, MaxY, MaxZ);
 	}
 
-	FORCEINLINE void SetDepthStencilState(FDepthStencilStateRHIParamRef NewStateRHI, uint32 StencilRef = 0)
-	{
-		CheckIsNull();
-		RHISetDepthStencilState(NewStateRHI, StencilRef);
-	}
+	inline void SetDepthStencilState(FDepthStencilStateRHIParamRef NewStateRHI, uint32 StencilRef = 0);
 
 	FORCEINLINE void BeginDrawPrimitiveUP(uint32 PrimitiveType, uint32 NumPrimitives, uint32 NumVertices, uint32 VertexDataStride, void*& OutVertexData)
 	{
diff --git a/Engine/Source/Runtime/RHI/Public/RHICommandList.inl b/Engine/Source/Runtime/RHI/Public/RHICommandList.inl
--- a/Engine/Source/Runtime/RHI/Public/RHICommandList.inl
+++ b/Engine/Source/Runtime/RHI/Public/RHICommandList.inl
@@ -62,6 +62,18 @@ struct FRHICommandSetRasterizerState : public FRHICommand
 	}
 };
 
+struct FRHICommandSetDepthStencilState : public FRHICommand
+{
+	FDepthStencilStateRHIParamRef State;
+	uint32 StencilRef;
+	FORCEINLINE void Set(FDepthStencilStateRHIParamRef InState, uint32 InStencilRef)
+	{
+		Type = ERCT_SetDepthStencilState;
+		State = InState;
+		StencilRef = InStencilRef;
+	}
+};
+
 struct FRHICommandSetShaderParameter : public FRHICommandPerShader
 {
 	const void* NewValue;
@@ -312,6 +324,19 @@ FORCEINLINE void FRHICommandList::SetRasterizerState(FRasterizerStateRHIParamRef
 	State->AddRef();
 }
 
+FORCEINLINE void FRHICommandList::SetDepthStencilState(FDepthStencilStateRHIParamRef State, uint32 StencilRef)
+{
+	if (Bypass())
+	{
+		RHISetDepthStencilState(State, StencilRef);
+		return;
+	}
+
+	auto* Cmd = AddCommand<FRHICommandSetDepthStencilState>();
+	Cmd->Set(State, StencilRef);
+	State->AddRef();
+}
+
 FORCEINLINE void FRHICommandList::DrawPrimitive(uint32 PrimitiveType, uint32 BaseVertexIndex, uint32 NumPrimitives, uint32 NumInstances)
 {
 	if (Bypass())
diff --git a/Engine/Source/Runtime/Renderer/Private/ShadowRendering.cpp b/Engine/Source/Runtime/Renderer/Private/ShadowRendering.cpp
--- a/Engine/Source/Runtime/Renderer/Private/ShadowRendering.cpp
+++ b/Engine/Source/Runtime/Renderer/Private/ShadowRendering.cpp
@@ -1364,44 +1364,46 @@ void FProjectedShadowInfo::RenderDepth(FRHICommandList& RHICmdList, FDeferredSha
 	FShadowDepthDrawingPolicy<false>::PolicyShadowInfo = this;
 	FShadowDepthDrawingPolicy<true>::PolicyShadowInfo = this;
 
-	// Draw the subject's static elements using static draw lists
-	if (IsWholeSceneDirectionalShadow())
 	{
-		SCOPE_CYCLE_COUNTER(STAT_WholeSceneStaticDrawListShadowDepthsTime);
-
-//@todo-rco: Very temp code!!!
-SCOPE_CYCLE_COUNTER(STAT_RHICounterTEMP);
-static IConsoleVariable* RHICmdListCVar = IConsoleManager::Get().FindConsoleVariable(TEXT("r.RHICmd"));
-bool bUseRHICmdList = (RHICmdListCVar->GetInt() >= 2);
-FRHICommandList& RHICmdList = bUseRHICmdList ? GRHICommandList.CreateList() : FRHICommandList::GetNullRef();
+		//@todo-rco: Very temp code!!!
+		SCOPE_CYCLE_COUNTER(STAT_RHICounterTEMP);
+		static IConsoleVariable* RHICmdListCVar = IConsoleManager::Get().FindConsoleVariable(TEXT("r.RHICmd"));
+		bool bUseRHICmdList = (RHICmdListCVar->GetInt() >= 2);
+		FRHICommandList& RHICmdList = bUseRHICmdList ? GRHICommandList.CreateList() : FRHICommandList::GetNullRef();
 
-		if (bReflectiveShadowmap)
-		{
-			SceneRenderer->Scene->WholeSceneReflectiveShadowMapDrawList.DrawVisible(RHICmdList, *FoundView, StaticMeshWholeSceneShadowDepthMap, StaticMeshWholeSceneShadowBatchVisibility);
-		}
-		else
+		// Draw the subject's static elements using static draw lists
+		if (IsWholeSceneDirectionalShadow())
 		{
-			// Use the scene's shadow depth draw list with this shadow's visibility map
-			SceneRenderer->Scene->WholeSceneShadowDepthDrawList.DrawVisible(RHICmdList, *FoundView, StaticMeshWholeSceneShadowDepthMap, StaticMeshWholeSceneShadowBatchVisibility);
-		}
-
-GRHICommandList.ExecuteAndFreeList(RHICmdList);
-	}
-	// Draw the subject's static elements using manual state filtering
-	else if (SubjectMeshElements.Num() > 0)
-	{
-		SCOPE_CYCLE_COUNTER(STAT_WholeSceneStaticShadowDepthsTime);
+			SCOPE_CYCLE_COUNTER(STAT_WholeSceneStaticDrawListShadowDepthsTime);
 
-		if(bReflectiveShadowmap && !bOnePassPointLightShadow)
-		{
-			// reflective shadow map
-			DrawShadowMeshElements<true>(RHICmdList, *FoundView, *this);
+			if (bReflectiveShadowmap)
+			{
+				SceneRenderer->Scene->WholeSceneReflectiveShadowMapDrawList.DrawVisible(RHICmdList, *FoundView, StaticMeshWholeSceneShadowDepthMap, StaticMeshWholeSceneShadowBatchVisibility);
+			}
+			else
+			{
+				// Use the scene's shadow depth draw list with this shadow's visibility map
+				SceneRenderer->Scene->WholeSceneShadowDepthDrawList.DrawVisible(RHICmdList, *FoundView, StaticMeshWholeSceneShadowDepthMap, StaticMeshWholeSceneShadowBatchVisibility);
+			}
 		}
-		else
+		// Draw the subject's static elements using manual state filtering
+		else if (SubjectMeshElements.Num() > 0)
 		{
-			// normal shadow map
-			DrawShadowMeshElements<false>(RHICmdList, *FoundView, *this);
+			SCOPE_CYCLE_COUNTER(STAT_WholeSceneStaticShadowDepthsTime);
+
+			if(bReflectiveShadowmap && !bOnePassPointLightShadow)
+			{
+				// reflective shadow map
+				DrawShadowMeshElements<true>(RHICmdList, *FoundView, *this);
+			}
+			else
+			{
+				// normal shadow map
+				DrawShadowMeshElements<false>(RHICmdList, *FoundView, *this);
+			}
 		}
+
+		GRHICommandList.ExecuteAndFreeList(RHICmdList);
 	}
 
 	// Draw the subject's dynamic elements.
diff --git a/Engine/Source/Runtime/Windows/D3D11RHI/Private/D3D11Commands.cpp b/Engine/Source/Runtime/Windows/D3D11RHI/Private/D3D11Commands.cpp
--- a/Engine/Source/Runtime/Windows/D3D11RHI/Private/D3D11Commands.cpp
+++ b/Engine/Source/Runtime/Windows/D3D11RHI/Private/D3D11Commands.cpp
@@ -2458,6 +2458,28 @@ void FD3D11DynamicRHI::RHIExecuteCommandList(FRHICommandList* CmdList)
 				RHICmd->VertexBuffer->Release();
 			}
 			break;
+		case ERCT_SetDepthStencilState:
+			{
+				auto* RHICmd = Iter.NextCommand<FRHICommandSetDepthStencilState>();
+				{
+					auto* State = (FD3D11DepthStencilState*)RHICmd->State;
+
+					if (CurrentDepthTexture && State->AccessType != CurrentDSVAccessType)
+					{
+						CurrentDSVAccessType = State->AccessType;
+						CurrentDepthStencilTarget = CurrentDepthTexture->GetDepthStencilView(CurrentDSVAccessType);
+
+						// Unbind any shader views of the depth stencil target that are bound.
+						ConditionalClearShaderResource(CurrentDepthTexture);
+
+						CommitRenderTargetsAndUAVs();
+					}
+
+					StateCache.SetDepthStencilState(State->Resource, RHICmd->StencilRef);
+				}
+				RHICmd->State->Release();
+			}
+			break;
 		default:
 			checkf(0, TEXT("Unimplemented RHICmd %d!"), Cmd->Type);
 		}
