Blueprints can now request that certain components or actors be hidden in a scene capture.

GitHub pull request 201 from bigzer.

diff --git a/Engine/Source/Runtime/Engine/Classes/Components/SceneCaptureComponent.h b/Engine/Source/Runtime/Engine/Classes/Components/SceneCaptureComponent.h
--- a/Engine/Source/Runtime/Engine/Classes/Components/SceneCaptureComponent.h
+++ b/Engine/Source/Runtime/Engine/Classes/Components/SceneCaptureComponent.h
@@ -10,6 +10,10 @@ class USceneCaptureComponent : public USceneComponent
 {
 	GENERATED_UCLASS_BODY()
 
+	/** The components won't rendered by current component.*/
+ 	UPROPERTY()
+ 	TArray<TWeakObjectPtr<UPrimitiveComponent> > HiddenComponents;
+
 	/** Whether to update the capture's contents every frame.  If disabled, the component will render once on load and then only when moved. */
 	UPROPERTY(EditAnywhere, BlueprintReadWrite, Category=SceneCapture)
 	bool bCaptureEveryFrame;
@@ -20,5 +24,12 @@ class USceneCaptureComponent : public USceneComponent
 
 public:
 
+	/** Adds the component to our list of hidden components. */
+	UFUNCTION(BlueprintCallable, Category = "Rendering|SceneCapture")
+	ENGINE_API void HideComponent(UPrimitiveComponent* InComponent);
+
+	/** Adds all primitive components in the actor to our list of hidden components. */
+	UFUNCTION(BlueprintCallable, Category = "Rendering|SceneCapture")
+	ENGINE_API void HideActorComponents(AActor* InActor);
 };
 
diff --git a/Engine/Source/Runtime/Engine/Private/SceneCaptureComponent.cpp b/Engine/Source/Runtime/Engine/Private/SceneCaptureComponent.cpp
--- a/Engine/Source/Runtime/Engine/Private/SceneCaptureComponent.cpp
+++ b/Engine/Source/Runtime/Engine/Private/SceneCaptureComponent.cpp
@@ -151,6 +151,27 @@ USceneCaptureComponent::USceneCaptureComponent(const class FPostConstructInitial
 }
 
 
+void USceneCaptureComponent::HideComponent(UPrimitiveComponent* InComponent)
+{
+	if (InComponent)
+	{
+		HiddenComponents.AddUnique(InComponent);
+	}
+}
+
+void USceneCaptureComponent::HideActorComponents(AActor* InActor)
+{
+	if (InActor)
+	{
+		TArray<UPrimitiveComponent*> PrimitiveComponents;
+		InActor->GetComponents(PrimitiveComponents);
+		for (int32 ComponentIndex = 0, NumComponents = PrimitiveComponents.Num(); ComponentIndex < NumComponents; ++ComponentIndex)
+		{
+			HiddenComponents.AddUnique(PrimitiveComponents[ComponentIndex]);
+		}
+	}
+}
+
 // -----------------------------------------------
 
 
diff --git a/Engine/Source/Runtime/Renderer/Private/SceneCaptureRendering.cpp b/Engine/Source/Runtime/Renderer/Private/SceneCaptureRendering.cpp
--- a/Engine/Source/Runtime/Renderer/Private/SceneCaptureRendering.cpp
+++ b/Engine/Source/Runtime/Renderer/Private/SceneCaptureRendering.cpp
@@ -78,7 +78,7 @@ void UpdateSceneCaptureContent_RenderThread(FSceneRenderer* SceneRenderer, FText
 }
 
 
-FSceneRenderer* FScene::CreateSceneRenderer( UTextureRenderTarget* TextureTarget, const FMatrix& ViewMatrix, const FVector& ViewLocation, float FOV, float MaxViewDistance, bool bCaptureSceneColour, FPostProcessSettings* PostProcessSettings, float PostProcessBlendWeight )
+FSceneRenderer* FScene::CreateSceneRenderer( USceneCaptureComponent* SceneCaptureComponent, UTextureRenderTarget* TextureTarget, const FMatrix& ViewMatrix, const FVector& ViewLocation, float FOV, float MaxViewDistance, bool bCaptureSceneColour, FPostProcessSettings* PostProcessSettings, float PostProcessBlendWeight )
 {
 	FIntPoint CaptureSize(TextureTarget->GetSurfaceWidth(), TextureTarget->GetSurfaceHeight());
 
@@ -138,6 +138,17 @@ FSceneRenderer* FScene::CreateSceneRenderer( UTextureRenderTarget* TextureTarget
 
 	View->bIsSceneCapture = true;
 
+	check(SceneCaptureComponent);
+	for (auto It = SceneCaptureComponent->HiddenComponents.CreateConstIterator(); It; ++It)
+	{
+		// If the primitive component was destroyed, the weak pointer will return NULL.
+		UPrimitiveComponent* PrimitiveComponent = It->Get();
+		if (PrimitiveComponent)
+		{
+			View->HiddenPrimitives.Add(PrimitiveComponent->ComponentId);
+		}
+	}
+
 	ViewFamily.Views.Add(View);
 
 	View->StartFinalPostprocessSettings(ViewLocation);
@@ -167,7 +178,7 @@ void FScene::UpdateSceneCaptureContents(USceneCaptureComponent2D* CaptureCompone
 			FPlane(0,	0,	0,	1));
 		const float FOV = CaptureComponent->FOVAngle * (float)PI / 360.0f;
 		const bool bUseSceneColorTexture = CaptureComponent->CaptureSource == SCS_SceneColorHDR;
-		FSceneRenderer* SceneRenderer = CreateSceneRenderer(CaptureComponent->TextureTarget, ViewMatrix , ViewLocation, FOV, CaptureComponent->MaxViewDistanceOverride, bUseSceneColorTexture, &CaptureComponent->PostProcessSettings, CaptureComponent->PostProcessBlendWeight);
+		FSceneRenderer* SceneRenderer = CreateSceneRenderer(CaptureComponent, CaptureComponent->TextureTarget, ViewMatrix , ViewLocation, FOV, CaptureComponent->MaxViewDistanceOverride, bUseSceneColorTexture, &CaptureComponent->PostProcessSettings, CaptureComponent->PostProcessBlendWeight);
 
 		FTextureRenderTargetResource* TextureRenderTarget = CaptureComponent->TextureTarget->GameThread_GetRenderTargetResource();
 		const FName OwnerName = CaptureComponent->GetOwner() ? CaptureComponent->GetOwner()->GetFName() : NAME_None;
@@ -238,7 +249,7 @@ void FScene::UpdateSceneCaptureContents(USceneCaptureComponentCube* CaptureCompo
 			const ECubeFace TargetFace = (ECubeFace)faceidx;
 			const FVector Location = CaptureComponent->GetComponentToWorld().GetTranslation();
 			const FMatrix ViewMatrix = FLocal::CalcCubeFaceTransform(TargetFace, Location);
-			FSceneRenderer* SceneRenderer = CreateSceneRenderer(CaptureComponent->TextureTarget, ViewMatrix, Location, FOV, CaptureComponent->MaxViewDistanceOverride);
+			FSceneRenderer* SceneRenderer = CreateSceneRenderer(CaptureComponent, CaptureComponent->TextureTarget, ViewMatrix, Location, FOV, CaptureComponent->MaxViewDistanceOverride);
 
 			FTextureRenderTargetCubeResource* TextureRenderTarget = static_cast<FTextureRenderTargetCubeResource*>(CaptureComponent->TextureTarget->GameThread_GetRenderTargetResource());
 			const FName OwnerName = CaptureComponent->GetOwner() ? CaptureComponent->GetOwner()->GetFName() : NAME_None;
diff --git a/Engine/Source/Runtime/Renderer/Private/ScenePrivate.h b/Engine/Source/Runtime/Renderer/Private/ScenePrivate.h
--- a/Engine/Source/Runtime/Renderer/Private/ScenePrivate.h
+++ b/Engine/Source/Runtime/Renderer/Private/ScenePrivate.h
@@ -1491,6 +1491,7 @@ private:
 	/**
 	 * Builds a FSceneRenderer instance for a given view and render target. Helper function for Scene capture code.
 	 *
+	 * @param	SceneCaptureComponent - The scene capture component for which to create a scene renderer.
 	 * @param	TextureTarget - render target to draw to.
 	 * @param	ViewMatrix - Camera view matrix.
 	 * @param	ViewLocation - Camera location.
@@ -1501,7 +1502,7 @@ private:
 	 * @param	PostProcessBlendWeight - Blendweight for PostProcessSettings.
 	 * @return	pointer to a configured SceneRenderer instance.
 	 */
-	FSceneRenderer* CreateSceneRenderer(UTextureRenderTarget* TextureTarget, const FMatrix& ViewMatrix, const FVector& ViewLocation, float FOV, float MaxViewDistance, bool bCaptureSceneColour = true, FPostProcessSettings* PostProcessSettings = NULL, float PostProcessBlendWeight = 0);
+	FSceneRenderer* CreateSceneRenderer(USceneCaptureComponent* SceneCaptureComponent, UTextureRenderTarget* TextureTarget, const FMatrix& ViewMatrix, const FVector& ViewLocation, float FOV, float MaxViewDistance, bool bCaptureSceneColour = true, FPostProcessSettings* PostProcessSettings = NULL, float PostProcessBlendWeight = 0);
 
 private:
 	/** 
