Dynamic sky light now compiles for OpenGL

diff --git a/Engine/Source/Runtime/Renderer/Private/DistanceFieldSurfaceCacheLighting.cpp b/Engine/Source/Runtime/Renderer/Private/DistanceFieldSurfaceCacheLighting.cpp
--- a/Engine/Source/Runtime/Renderer/Private/DistanceFieldSurfaceCacheLighting.cpp
+++ b/Engine/Source/Runtime/Renderer/Private/DistanceFieldSurfaceCacheLighting.cpp
@@ -392,7 +392,7 @@ public:
 
 	static bool ShouldCache(EShaderPlatform Platform)
 	{
-		return IsFeatureLevelSupported(Platform, ERHIFeatureLevel::SM5) && DoesPlatformSupportDistanceFieldAO(Platform) && !IsOpenGLPlatform(Platform);
+		return IsFeatureLevelSupported(Platform, ERHIFeatureLevel::SM5) && DoesPlatformSupportDistanceFieldAO(Platform);
 	}
 
 	static void ModifyCompilationEnvironment(EShaderPlatform Platform, FShaderCompilerEnvironment& OutEnvironment)
@@ -484,7 +484,7 @@ public:
 
 	static bool ShouldCache(EShaderPlatform Platform) 
 	{ 
-		return IsFeatureLevelSupported(Platform, ERHIFeatureLevel::SM5) && DoesPlatformSupportDistanceFieldAO(Platform) && !IsOpenGLPlatform(Platform); 
+		return IsFeatureLevelSupported(Platform, ERHIFeatureLevel::SM5) && DoesPlatformSupportDistanceFieldAO(Platform); 
 	}
 
 	static void ModifyCompilationEnvironment(EShaderPlatform Platform, FShaderCompilerEnvironment& OutEnvironment)
@@ -543,7 +543,7 @@ public:
 
 	static bool ShouldCache(EShaderPlatform Platform)
 	{
-		return IsFeatureLevelSupported(Platform, ERHIFeatureLevel::SM5) && DoesPlatformSupportDistanceFieldAO(Platform) && !IsOpenGLPlatform(Platform);
+		return IsFeatureLevelSupported(Platform, ERHIFeatureLevel::SM5) && DoesPlatformSupportDistanceFieldAO(Platform);
 	}
 
 	static void ModifyCompilationEnvironment(EShaderPlatform Platform, FShaderCompilerEnvironment& OutEnvironment)
@@ -636,7 +636,7 @@ public:
 
 	static bool ShouldCache(EShaderPlatform Platform)
 	{
-		return IsFeatureLevelSupported(Platform, ERHIFeatureLevel::SM5) && DoesPlatformSupportDistanceFieldAO(Platform) && !IsOpenGLPlatform(Platform);
+		return IsFeatureLevelSupported(Platform, ERHIFeatureLevel::SM5) && DoesPlatformSupportDistanceFieldAO(Platform);
 	}
 
 	static void ModifyCompilationEnvironment(EShaderPlatform Platform, FShaderCompilerEnvironment& OutEnvironment)
@@ -769,7 +769,7 @@ public:
 
 	static bool ShouldCache(EShaderPlatform Platform)
 	{
-		return IsFeatureLevelSupported(Platform, ERHIFeatureLevel::SM5) && DoesPlatformSupportDistanceFieldAO(Platform) && !IsOpenGLPlatform(Platform);
+		return IsFeatureLevelSupported(Platform, ERHIFeatureLevel::SM5) && DoesPlatformSupportDistanceFieldAO(Platform);
 	}
 
 	static void ModifyCompilationEnvironment(EShaderPlatform Platform, FShaderCompilerEnvironment& OutEnvironment)
@@ -853,7 +853,7 @@ public:
 
 	static bool ShouldCache(EShaderPlatform Platform)
 	{
-		return IsFeatureLevelSupported(Platform, ERHIFeatureLevel::SM5) && DoesPlatformSupportDistanceFieldAO(Platform) && !IsOpenGLPlatform(Platform);
+		return IsFeatureLevelSupported(Platform, ERHIFeatureLevel::SM5) && DoesPlatformSupportDistanceFieldAO(Platform);
 	}
 
 	static void ModifyCompilationEnvironment(EShaderPlatform Platform, FShaderCompilerEnvironment& OutEnvironment)
@@ -1010,7 +1010,7 @@ public:
 
 	static bool ShouldCache(EShaderPlatform Platform)
 	{
-		return IsFeatureLevelSupported(Platform, ERHIFeatureLevel::SM5) && DoesPlatformSupportDistanceFieldAO(Platform) && !IsOpenGLPlatform(Platform);
+		return IsFeatureLevelSupported(Platform, ERHIFeatureLevel::SM5) && DoesPlatformSupportDistanceFieldAO(Platform);
 	}
 
 	static void ModifyCompilationEnvironment(EShaderPlatform Platform, FShaderCompilerEnvironment& OutEnvironment)
@@ -1077,7 +1077,7 @@ public:
 
 	static bool ShouldCache(EShaderPlatform Platform)
 	{
-		return IsFeatureLevelSupported(Platform, ERHIFeatureLevel::SM5) && DoesPlatformSupportDistanceFieldAO(Platform) && !IsOpenGLPlatform(Platform);
+		return IsFeatureLevelSupported(Platform, ERHIFeatureLevel::SM5) && DoesPlatformSupportDistanceFieldAO(Platform);
 	}
 
 	static void ModifyCompilationEnvironment(EShaderPlatform Platform, FShaderCompilerEnvironment& OutEnvironment)
@@ -1189,7 +1189,7 @@ public:
 
 	static bool ShouldCache(EShaderPlatform Platform)
 	{
-		return IsFeatureLevelSupported(Platform, ERHIFeatureLevel::SM5) && DoesPlatformSupportDistanceFieldAO(Platform) && !IsOpenGLPlatform(Platform);
+		return IsFeatureLevelSupported(Platform, ERHIFeatureLevel::SM5) && DoesPlatformSupportDistanceFieldAO(Platform);
 	}
 
 	FSaveStartIndexCS(const ShaderMetaType::CompiledShaderInitializerType& Initializer)
@@ -1243,7 +1243,7 @@ public:
 
 	static bool ShouldCache(EShaderPlatform Platform)
 	{
-		return IsFeatureLevelSupported(Platform, ERHIFeatureLevel::SM5) && DoesPlatformSupportDistanceFieldAO(Platform) && !IsOpenGLPlatform(Platform);
+		return IsFeatureLevelSupported(Platform, ERHIFeatureLevel::SM5) && DoesPlatformSupportDistanceFieldAO(Platform);
 	}
 
 	static void ModifyCompilationEnvironment(EShaderPlatform Platform, FShaderCompilerEnvironment& OutEnvironment)
@@ -1390,7 +1390,7 @@ public:
 
 	static bool ShouldCache(EShaderPlatform Platform)
 	{
-		return IsFeatureLevelSupported(Platform, ERHIFeatureLevel::SM5) && DoesPlatformSupportDistanceFieldAO(Platform) && !IsOpenGLPlatform(Platform);
+		return IsFeatureLevelSupported(Platform, ERHIFeatureLevel::SM5) && DoesPlatformSupportDistanceFieldAO(Platform);
 	}
 
 	static void ModifyCompilationEnvironment(EShaderPlatform Platform, FShaderCompilerEnvironment& OutEnvironment)
@@ -1457,7 +1457,7 @@ public:
 
 	static bool ShouldCache(EShaderPlatform Platform)
 	{
-		return IsFeatureLevelSupported(Platform, ERHIFeatureLevel::SM5) && DoesPlatformSupportDistanceFieldAO(Platform) && !IsOpenGLPlatform(Platform);
+		return IsFeatureLevelSupported(Platform, ERHIFeatureLevel::SM5) && DoesPlatformSupportDistanceFieldAO(Platform);
 	}
 
 	static void ModifyCompilationEnvironment(EShaderPlatform Platform, FShaderCompilerEnvironment& OutEnvironment)
@@ -1667,7 +1667,7 @@ public:
 
 	static bool ShouldCache(EShaderPlatform Platform) 
 	{ 
-		return IsFeatureLevelSupported(Platform, ERHIFeatureLevel::SM5) && DoesPlatformSupportDistanceFieldAO(Platform) && !IsOpenGLPlatform(Platform);
+		return IsFeatureLevelSupported(Platform, ERHIFeatureLevel::SM5) && DoesPlatformSupportDistanceFieldAO(Platform);
 	}
 
 	static void ModifyCompilationEnvironment(EShaderPlatform Platform, FShaderCompilerEnvironment& OutEnvironment)
@@ -1746,7 +1746,7 @@ public:
 
 	static bool ShouldCache(EShaderPlatform Platform)
 	{
-		return IsFeatureLevelSupported(Platform, ERHIFeatureLevel::SM5) && DoesPlatformSupportDistanceFieldAO(Platform) && !IsOpenGLPlatform(Platform);
+		return IsFeatureLevelSupported(Platform, ERHIFeatureLevel::SM5) && DoesPlatformSupportDistanceFieldAO(Platform);
 	}
 
 	static void ModifyCompilationEnvironment(EShaderPlatform Platform, FShaderCompilerEnvironment& OutEnvironment)
@@ -1832,7 +1832,7 @@ public:
 
 	static bool ShouldCache(EShaderPlatform Platform)
 	{
-		return IsFeatureLevelSupported(Platform, ERHIFeatureLevel::SM5) && DoesPlatformSupportDistanceFieldAO(Platform) && !IsOpenGLPlatform(Platform);
+		return IsFeatureLevelSupported(Platform, ERHIFeatureLevel::SM5) && DoesPlatformSupportDistanceFieldAO(Platform);
 	}
 
 	static void ModifyCompilationEnvironment(EShaderPlatform Platform, FShaderCompilerEnvironment& OutEnvironment)
@@ -1895,7 +1895,7 @@ public:
 
 	static bool ShouldCache(EShaderPlatform Platform)
 	{
-		return IsFeatureLevelSupported(Platform, ERHIFeatureLevel::SM5) && DoesPlatformSupportDistanceFieldAO(Platform) && !IsOpenGLPlatform(Platform);
+		return IsFeatureLevelSupported(Platform, ERHIFeatureLevel::SM5) && DoesPlatformSupportDistanceFieldAO(Platform);
 	}
 
 	static void ModifyCompilationEnvironment(EShaderPlatform Platform, FShaderCompilerEnvironment& OutEnvironment)
@@ -1961,7 +1961,7 @@ public:
 
 	static bool ShouldCache(EShaderPlatform Platform)
 	{
-		return IsFeatureLevelSupported(Platform, ERHIFeatureLevel::SM5) && DoesPlatformSupportDistanceFieldAO(Platform) && !IsOpenGLPlatform(Platform);
+		return IsFeatureLevelSupported(Platform, ERHIFeatureLevel::SM5) && DoesPlatformSupportDistanceFieldAO(Platform);
 	}
 
 	static void ModifyCompilationEnvironment(EShaderPlatform Platform, FShaderCompilerEnvironment& OutEnvironment)
@@ -2233,7 +2233,7 @@ public:
 
 	static bool ShouldCache(EShaderPlatform Platform)
 	{
-		return IsFeatureLevelSupported(Platform, ERHIFeatureLevel::SM5) && DoesPlatformSupportDistanceFieldAO(Platform) && !IsOpenGLPlatform(Platform);
+		return IsFeatureLevelSupported(Platform, ERHIFeatureLevel::SM5) && DoesPlatformSupportDistanceFieldAO(Platform);
 	}
 
 	static void ModifyCompilationEnvironment(EShaderPlatform Platform, FShaderCompilerEnvironment& OutEnvironment)
@@ -3045,7 +3045,7 @@ public:
 
 	static bool ShouldCache(EShaderPlatform Platform)
 	{
-		return IsFeatureLevelSupported(Platform, ERHIFeatureLevel::SM4) && !IsOpenGLPlatform(Platform);
+		return IsFeatureLevelSupported(Platform, ERHIFeatureLevel::SM4);
 	}
 
 	static void ModifyCompilationEnvironment(EShaderPlatform Platform, FShaderCompilerEnvironment& OutEnvironment)
