Fix for crash rendering lit translucency in OpenGL

diff --git a/Engine/Source/Runtime/Renderer/Private/TranslucentLighting.cpp b/Engine/Source/Runtime/Renderer/Private/TranslucentLighting.cpp
--- a/Engine/Source/Runtime/Renderer/Private/TranslucentLighting.cpp
+++ b/Engine/Source/Runtime/Renderer/Private/TranslucentLighting.cpp
@@ -936,8 +936,8 @@ public:
 
 		const FStaticShadowDepthMap* StaticShadowDepthMap = LightSceneInfo->Proxy->GetStaticShadowDepthMap();
 		const uint32 bStaticallyShadowedValue = LightSceneInfo->bPrecomputedLightingIsValid && StaticShadowDepthMap && StaticShadowDepthMap->TextureRHI ? 1 : 0;
-		FTextureRHIParamRef StaticShadowDepthMapTexture = StaticShadowDepthMap ? StaticShadowDepthMap->TextureRHI : GWhiteTexture->TextureRHI;
-		const FMatrix WorldToStaticShadow = StaticShadowDepthMap ? StaticShadowDepthMap->WorldToLight : FMatrix::Identity;
+		FTextureRHIParamRef StaticShadowDepthMapTexture = bStaticallyShadowedValue ? StaticShadowDepthMap->TextureRHI : GWhiteTexture->TextureRHI;
+		const FMatrix WorldToStaticShadow = bStaticallyShadowedValue ? StaticShadowDepthMap->WorldToLight : FMatrix::Identity;
 
 		SetShaderValue(RHICmdList, ShaderRHI, bStaticallyShadowed, bStaticallyShadowedValue);
 
