Fix linker errors attempting to use FMaterialShader outside of the renderer module

diff --git a/Engine/Source/Runtime/Renderer/Private/ShaderBaseClasses.cpp b/Engine/Source/Runtime/Renderer/Private/ShaderBaseClasses.cpp
--- a/Engine/Source/Runtime/Renderer/Private/ShaderBaseClasses.cpp
+++ b/Engine/Source/Runtime/Renderer/Private/ShaderBaseClasses.cpp
@@ -263,7 +263,7 @@ void FMaterialShader::SetParameters(
 }
 
 #define IMPLEMENT_MATERIAL_SHADER_SetParameters( ShaderRHIParamRef ) \
-	template void FMaterialShader::SetParameters< ShaderRHIParamRef >( \
+	template RENDERER_API void FMaterialShader::SetParameters< ShaderRHIParamRef >( \
 		FRHICommandList& RHICmdList,					\
 		const ShaderRHIParamRef ShaderRHI,				\
 		const FMaterialRenderProxy* MaterialRenderProxy,\
diff --git a/Engine/Source/Runtime/Renderer/Public/MaterialShader.h b/Engine/Source/Runtime/Renderer/Public/MaterialShader.h
--- a/Engine/Source/Runtime/Renderer/Public/MaterialShader.h
+++ b/Engine/Source/Runtime/Renderer/Public/MaterialShader.h
@@ -79,7 +79,7 @@ inline FArchive& operator<<(FArchive& Ar, FDebugUniformExpressionSet& DebugExpre
 
 
 /** Base class of all shaders that need material parameters. */
-class FMaterialShader : public FShader
+class RENDERER_API FMaterialShader : public FShader
 {
 public:
 	FMaterialShader() {}
@@ -88,7 +88,7 @@ public:
 
 	typedef void (*ModifyCompilationEnvironmentType)(EShaderPlatform, const FMaterial*, FShaderCompilerEnvironment&);
 
-	ENGINE_API static void ModifyCompilationEnvironment(EShaderPlatform Platform, const FMaterial* Material, FShaderCompilerEnvironment& OutEnvironment)
+	static void ModifyCompilationEnvironment(EShaderPlatform Platform, const FMaterial* Material, FShaderCompilerEnvironment& OutEnvironment)
 	{
 	}
 
@@ -157,4 +157,4 @@ private:
 	static int32 bAllowCachedUniformExpressions;
 	/** Console variable ref to toggle cached uniform expressions. */
 	static FAutoConsoleVariableRef CVarAllowCachedUniformExpressions;
-};
\ No newline at end of file
+};
