* added VisualzieSSR showflag (from NVIDIA branch)

diff --git a/Engine/Shaders/ScreenSpaceReflections.usf b/Engine/Shaders/ScreenSpaceReflections.usf
--- a/Engine/Shaders/ScreenSpaceReflections.usf
+++ b/Engine/Shaders/ScreenSpaceReflections.usf
@@ -309,6 +309,15 @@ void ScreenSpaceReflectionsPS(in float4 UVAndScreenPos : TEXCOORD0, out float4 O
 	}
 #endif
 
+#if SSR_QUALITY == 0
+	// visualize SSR
+
+	float PatternMask = ((PixelPos.x / 2 + PixelPos.y / 2) % 2) * 0.7f;
+
+	OutColor = lerp(float4(1, 0, 0, 1), float4(1, 1 ,0, 1), PatternMask) * 0.3f;
+	return;
+#endif
+
 	float Roughness = ScreenSpaceData.GBuffer.Roughness;
 
 	float3 N = ScreenSpaceData.GBuffer.WorldNormal;
diff --git a/Engine/Source/Runtime/Engine/Public/ShowFlags.h b/Engine/Source/Runtime/Engine/Public/ShowFlags.h
--- a/Engine/Source/Runtime/Engine/Public/ShowFlags.h
+++ b/Engine/Source/Runtime/Engine/Public/ShowFlags.h
@@ -209,6 +209,7 @@ struct FEngineShowFlags
 		// we enable it manually on the editor view ports
 		SetSnap(false);
 		SetVisualizeDistanceFieldAO(false);
+		SetVisualizeSSR(false);
 	}
 
 	static bool CanBeToggledInEditor(const TCHAR* Name)
diff --git a/Engine/Source/Runtime/Engine/Public/ShowFlagsValues.inl b/Engine/Source/Runtime/Engine/Public/ShowFlagsValues.inl
--- a/Engine/Source/Runtime/Engine/Public/ShowFlagsValues.inl
+++ b/Engine/Source/Runtime/Engine/Public/ShowFlagsValues.inl
@@ -266,6 +266,8 @@ SHOWFLAG_ALWAYS_ACCESSIBLE(Paper2DSprites, SFG_Advanced, LOCTEXT("Paper2DSprites
 SHOWFLAG_ALWAYS_ACCESSIBLE(VisualizeDistanceFieldAO, SFG_Visualize, LOCTEXT("VisualizeDistanceFieldAOSF", "Visualize Distance Field AO"))
 /** Distance field AO */
 SHOWFLAG_ALWAYS_ACCESSIBLE(DistanceFieldAO, SFG_LightingFeatures, LOCTEXT("DistanceFieldAOSF", "Distance Field AO"))
+/** Visualize screen space reflections, for developer (by default off): */
+SHOWFLAG_FIXED_IN_SHIPPING(VisualizeSSR, 0, SFG_Visualize, LOCTEXT("VisualizeSSR", "Visualize ScreenSpaceReflections"))
 
 #undef SHOWFLAG_ALWAYS_ACCESSIBLE
 #undef SHOWFLAG_FIXED_IN_SHIPPING
diff --git a/Engine/Source/Runtime/Renderer/Private/PostProcess/ScreenSpaceReflections.cpp b/Engine/Source/Runtime/Renderer/Private/PostProcess/ScreenSpaceReflections.cpp
--- a/Engine/Source/Runtime/Renderer/Private/PostProcess/ScreenSpaceReflections.cpp
+++ b/Engine/Source/Runtime/Renderer/Private/PostProcess/ScreenSpaceReflections.cpp
@@ -22,7 +22,7 @@ static TAutoConsoleVariable<int32> CVarSSRQuality(
 	TEXT(" 1: low (no glossy)\n")
 	TEXT(" 2: medium (no glossy)\n")
 	TEXT(" 3: high (glossy/using roughness, few samples)\n")
-	TEXT(" 4: very high (likely too slow for real-time)\n"),
+	TEXT(" 4: very high (likely too slow for real-time)"),
 	ECVF_Scalability | ECVF_RenderThreadSafe);
 
 static TAutoConsoleVariable<int32> CVarSSRTemporal(
@@ -60,7 +60,10 @@ bool DoScreenSpaceReflections(const FViewInfo& View)
 	return true;
 }
 
-/** Encapsulates the post processing screen space reflections pixel shader. */
+/**
+ * Encapsulates the post processing screen space reflections pixel shader.
+ * @param SSRQuality 0:Visualize Mask
+ */
 template<uint32 PrevFrame, uint32 SSRQuality >
 class FPostProcessScreenSpaceReflectionsPS : public FGlobalShader
 {
@@ -138,6 +141,7 @@ public:
 	typedef FPostProcessScreenSpaceReflectionsPS<A,B> FPostProcessScreenSpaceReflectionsPS##A##B; \
 	IMPLEMENT_SHADER_TYPE(template<>,FPostProcessScreenSpaceReflectionsPS##A##B,TEXT("ScreenSpaceReflections"),TEXT("ScreenSpaceReflectionsPS"),SF_Pixel)
 
+IMPLEMENT_REFLECTION_PIXELSHADER_TYPE(0,0);
 IMPLEMENT_REFLECTION_PIXELSHADER_TYPE(0,1);
 IMPLEMENT_REFLECTION_PIXELSHADER_TYPE(1,1);
 IMPLEMENT_REFLECTION_PIXELSHADER_TYPE(0,2);
@@ -193,10 +197,19 @@ void FRCPassPostProcessScreenSpaceReflections::Process(FRenderingCompositePassCo
 	SSRQuality = FMath::Clamp(SSRQuality, 1, 4);
 	Context.RHICmdList.CheckIsNull(); // need new approach for "static FGlobalBoundShaderState" for parallel rendering
 
+	uint32 iPreFrame = bPrevFrame ? 1 : 0;
+
+	if (View.Family->EngineShowFlags.VisualizeSSR)
+	{
+		iPreFrame = 0;
+		SSRQuality = 0;
+	}
+
+	TShaderMapRef< FPostProcessVS > VertexShader(GetGlobalShaderMap());
+
 	#define CASE(A, B) \
 		case (A + 2 * (B + 3 * 0 )): \
 		{ \
-			TShaderMapRef< FPostProcessVS > VertexShader(GetGlobalShaderMap()); \
 			TShaderMapRef< FPostProcessScreenSpaceReflectionsPS<A, B> > PixelShader(GetGlobalShaderMap()); \
 			static FGlobalBoundShaderState BoundShaderState; \
 			SetGlobalBoundShaderState(Context.RHICmdList, BoundShaderState, GFilterVertexDeclaration.VertexDeclarationRHI, *VertexShader, *PixelShader); \
@@ -205,18 +218,20 @@ void FRCPassPostProcessScreenSpaceReflections::Process(FRenderingCompositePassCo
 		}; \
 		break
 
-	switch( bPrevFrame + 2 * (SSRQuality + 3 * 0 ) )
+	switch (iPreFrame + 2 * (SSRQuality + 3 * 0))
 	{
+		CASE(0,0);
 		CASE(0,1);	CASE(1,1);
 		CASE(0,2);	CASE(1,2);
 		CASE(0,3);	CASE(1,3);
 		CASE(0,4);	CASE(1,4);
+		default:
+			check(!"Missing case in FRCPassPostProcessScreenSpaceReflections");
 	}
 	#undef CASE
 
 
 	// Draw a quad mapping scene color to the view's render target
-	TShaderMapRef< FPostProcessVS > VertexShader(GetGlobalShaderMap());
 	DrawRectangle( 
 		Context.RHICmdList,
 		0, 0,
