Fixed issue where transparent objects using scene color could read from a placeholder white texture.

#ttp 335552 - RENDERING: INFILTRATOR: Maglev very bright first run.
removed scene color related states from FMaterialRelevance and FPrimitiveViewRelevance structures.
scene color requirement is tested at render time instead.

diff --git a/Engine/Shaders/BasePassCommon.usf b/Engine/Shaders/BasePassCommon.usf
--- a/Engine/Shaders/BasePassCommon.usf
+++ b/Engine/Shaders/BasePassCommon.usf
@@ -10,7 +10,7 @@
 // Translucent materials need to compute fogging in the forward shading pass
 // Materials that read from scene color skip getting fogged, because the contents of the scene color lookup have already been fogged
 // This is not foolproof, as any additional color the material adds will then not be fogged correctly
-#define NEEDS_BASEPASS_FOGGING			(ENABLE_TRANSLUCENCY_VERTEX_FOG && ((MATERIALBLENDING_TRANSLUCENT || MATERIALBLENDING_ADDITIVE || MATERIALBLENDING_MODULATE) && !MATERIAL_USES_SCENE_COLOR))
+#define NEEDS_BASEPASS_FOGGING			(ENABLE_TRANSLUCENCY_VERTEX_FOG && ((MATERIALBLENDING_TRANSLUCENT || MATERIALBLENDING_ADDITIVE || MATERIALBLENDING_MODULATE) && !MATERIAL_USES_SCENE_COLOR_COPY))
 #define NEEDS_LIGHTMAP					(NEEDS_LIGHTMAP_COORDINATE)
 
 struct FSharedBasePassInterpolants
diff --git a/Engine/Source/Runtime/Engine/Classes/Materials/MaterialInterface.h b/Engine/Source/Runtime/Engine/Classes/Materials/MaterialInterface.h
--- a/Engine/Source/Runtime/Engine/Classes/Materials/MaterialInterface.h
+++ b/Engine/Source/Runtime/Engine/Classes/Materials/MaterialInterface.h
@@ -42,9 +42,6 @@ struct ENGINE_API FMaterialRelevance
 	UPROPERTY()
 		uint32 bDistortion : 1;
 
-	UPROPERTY()
-		uint32 bUsesSceneColor : 1;
-
 	UPROPERTY()
 		uint32 bSeparateTranslucency : 1;
 
@@ -59,7 +56,6 @@ struct ENGINE_API FMaterialRelevance
 		: bOpaque(false)
 		, bMasked(false)
 		, bDistortion(false)
-		, bUsesSceneColor(false)
 		, bSeparateTranslucency(false)
 		, bNormalTranslucency(false)
 		, bDisableDepthTest(false)
@@ -71,7 +67,6 @@ struct ENGINE_API FMaterialRelevance
 		bOpaque |= B.bOpaque;
 		bMasked |= B.bMasked;
 		bDistortion |= B.bDistortion;
-		bUsesSceneColor |= B.bUsesSceneColor;
 		bSeparateTranslucency |= B.bSeparateTranslucency;
 		bNormalTranslucency |= B.bNormalTranslucency;
 		bDisableDepthTest |= B.bDisableDepthTest;
diff --git a/Engine/Source/Runtime/Engine/Private/Materials/HLSLMaterialTranslator.h b/Engine/Source/Runtime/Engine/Private/Materials/HLSLMaterialTranslator.h
--- a/Engine/Source/Runtime/Engine/Private/Materials/HLSLMaterialTranslator.h
+++ b/Engine/Source/Runtime/Engine/Private/Materials/HLSLMaterialTranslator.h
@@ -300,7 +300,7 @@ public:
 			Material->CompileErrors.Empty();
 			Material->ErrorExpressions.Empty();
 
-			MaterialCompilationOutput.bUsesSceneColor = false;
+			MaterialCompilationOutput.bRequiresSceneColorCopy = false;
 			check(!MaterialCompilationOutput.bNeedsSceneTextures);
 			MaterialCompilationOutput.bUsesEyeAdaptation = false;
 
@@ -408,17 +408,15 @@ public:
 				Errorf(TEXT("Only transparent or postprocess materials can read from scene depth."));
 			}
 
-			if (MaterialCompilationOutput.bUsesSceneColor && Material->GetMaterialDomain() != MD_PostProcess)
+			if (MaterialCompilationOutput.bRequiresSceneColorCopy)
 			{
-				if (!IsTranslucentBlendMode(Material->GetBlendMode()))
+				if (Material->GetMaterialDomain() != MD_Surface)
 				{
-					Errorf(TEXT("Only transparent or postprocess materials can read from scene color."));
+					Errorf(TEXT("Only 'surface' material domain can use the scene color node."));
 				}
-				else if (!Material->IsSeparateTranslucencyEnabled())
+				else if (!IsTranslucentBlendMode(Material->GetBlendMode()))
 				{
-					// rendering without separate translucency could mean the scenecolor surface
-					// may be set simultaneously as both source texture and destination target. (error)
-					Errorf(TEXT("Cannot read from scene color whilst separate translucency is disabled."));
+					Errorf(TEXT("Only translucent materials can use the scene color node."));
 				}
 			}
 
@@ -448,7 +446,7 @@ public:
 					}
 					else if (bNeedsSceneTexturePostProcessInputs)
 					{
-						Errorf(TEXT("SceneTexture expressions cannot use post process inputs in non post process domain materials"));
+						Errorf(TEXT("SceneTexture expressions cannot use post process inputs or scene color in non post process domain materials"));
 					}
 				}
 			}
@@ -2328,9 +2326,6 @@ protected:
 				|| SceneTextureId == PPI_SceneColor);
 		}
 
-		MaterialCompilationOutput.bUsesSceneColor = MaterialCompilationOutput.bUsesSceneColor
-			|| SceneTextureId == PPI_SceneColor;
-
 		MaterialCompilationOutput.bNeedsGBuffer = MaterialCompilationOutput.bNeedsGBuffer
 			|| SceneTextureId == PPI_DiffuseColor 
 			|| SceneTextureId == PPI_SpecularColor
@@ -2366,7 +2361,7 @@ protected:
 			return INDEX_NONE;
 		}
 
-		MaterialCompilationOutput.bUsesSceneColor = true;
+		MaterialCompilationOutput.bRequiresSceneColorCopy = true;
 
 		int32 ScreenUVCode = GetScreenAlignedUV(Offset, UV, bUseOffset);
 		return AddCodeChunk(
diff --git a/Engine/Source/Runtime/Engine/Private/Materials/Material.cpp b/Engine/Source/Runtime/Engine/Private/Materials/Material.cpp
--- a/Engine/Source/Runtime/Engine/Private/Materials/Material.cpp
+++ b/Engine/Source/Runtime/Engine/Private/Materials/Material.cpp
@@ -3375,7 +3375,7 @@ static void ListSceneColorMaterials()
 	{
 		UMaterialInterface* Mat = *It;
 		const FMaterial* MatRes = Mat->GetRenderProxy(false)->GetMaterial(GRHIFeatureLevel);
-		if (MatRes && MatRes->UsesSceneColor())
+		if (MatRes && MatRes->RequiresSceneColorCopy())
 		{
 			UMaterial* BaseMat = Mat->GetMaterial();
 			UE_LOG(LogConsoleResponse,Display,TEXT("[SepTrans=%d] %s"),
diff --git a/Engine/Source/Runtime/Engine/Private/Materials/MaterialInterface.cpp b/Engine/Source/Runtime/Engine/Private/Materials/MaterialInterface.cpp
--- a/Engine/Source/Runtime/Engine/Private/Materials/MaterialInterface.cpp
+++ b/Engine/Source/Runtime/Engine/Private/Materials/MaterialInterface.cpp
@@ -14,7 +14,6 @@ void FMaterialRelevance::SetPrimitiveViewRelevance(FPrimitiveViewRelevance& OutV
 	OutViewRelevance.bOpaqueRelevance = bOpaque;
 	OutViewRelevance.bMaskedRelevance = bMasked;
 	OutViewRelevance.bDistortionRelevance = bDistortion;
-	OutViewRelevance.bSceneColorRelevance = bUsesSceneColor;
 	OutViewRelevance.bSeparateTranslucencyRelevance = bSeparateTranslucency;
 	OutViewRelevance.bNormalTranslucencyRelevance = bNormalTranslucency;
 }
@@ -49,7 +48,6 @@ FMaterialRelevance UMaterialInterface::GetRelevance_Internal(const UMaterial* Ma
 		MaterialRelevance.bOpaque = !bIsTranslucent;
 		MaterialRelevance.bMasked = Material->bIsMasked;
 		MaterialRelevance.bDistortion = Material->bUsesDistortion && bIsTranslucent;
-		MaterialRelevance.bUsesSceneColor = MaterialResource->UsesSceneColor();
 		MaterialRelevance.bSeparateTranslucency = bIsTranslucent && Material->bEnableSeparateTranslucency;
 		MaterialRelevance.bNormalTranslucency = bIsTranslucent && !Material->bEnableSeparateTranslucency;
 		MaterialRelevance.bDisableDepthTest = bIsTranslucent && Material->bDisableDepthTest;
diff --git a/Engine/Source/Runtime/Engine/Private/Materials/MaterialShared.cpp b/Engine/Source/Runtime/Engine/Private/Materials/MaterialShared.cpp
--- a/Engine/Source/Runtime/Engine/Private/Materials/MaterialShared.cpp
+++ b/Engine/Source/Runtime/Engine/Private/Materials/MaterialShared.cpp
@@ -318,7 +318,7 @@ void FMaterialCompilationOutput::Serialize(FArchive& Ar)
 
 	if (Ar.UE4Ver() >= VER_UE4_PURGED_FMATERIAL_COMPILE_OUTPUTS)
 	{
-		Ar << bUsesSceneColor;
+		Ar << bRequiresSceneColorCopy;
 	}
 
 	Ar << bNeedsSceneTextures;
@@ -434,15 +434,15 @@ const TArray<TRefCountPtr<FMaterialUniformExpressionTexture> >& FMaterial::GetUn
 	return EmptyExpressions;
 }
 
-bool FMaterial::UsesSceneColor() const 
+bool FMaterial::RequiresSceneColorCopy() const 
 {
 	if (IsInGameThread() && GameThreadShaderMap)
 	{
-		return GameThreadShaderMap->UsesSceneColor();
+		return GameThreadShaderMap->RequiresSceneColorCopy();
 	}
 	else if (IsInRenderingThread() && RenderingThreadShaderMap)
 	{
-		return RenderingThreadShaderMap->UsesSceneColor();
+		return RenderingThreadShaderMap->RequiresSceneColorCopy();
 	}
 
 	return false;
@@ -1175,7 +1175,7 @@ void FMaterial::SetupMaterialEnvironment(
 	OutEnvironment.SetDefine(TEXT("MATERIAL_TWOSIDED"), IsTwoSided() ? TEXT("1") : TEXT("0"));
 	OutEnvironment.SetDefine(TEXT("MATERIAL_TANGENTSPACENORMAL"), IsTangentSpaceNormal() ? TEXT("1") : TEXT("0"));
 	OutEnvironment.SetDefine(TEXT("GENERATE_SPHERICAL_PARTICLE_NORMALS"),ShouldGenerateSphericalParticleNormals() ? TEXT("1") : TEXT("0"));
-	OutEnvironment.SetDefine(TEXT("MATERIAL_USES_SCENE_COLOR"), UsesSceneColor() ? TEXT("1") : TEXT("0"));
+	OutEnvironment.SetDefine(TEXT("MATERIAL_USES_SCENE_COLOR_COPY"), RequiresSceneColorCopy() ? TEXT("1") : TEXT("0"));
 	OutEnvironment.SetDefine(TEXT("MATERIAL_FULLY_ROUGH"), IsFullyRough() ? TEXT("1") : TEXT("0"));
 	OutEnvironment.SetDefine(TEXT("MATERIAL_NONMETAL"), IsNonmetal() ? TEXT("1") : TEXT("0"));
 	OutEnvironment.SetDefine(TEXT("MATERIAL_USE_LM_DIRECTIONALITY"), UseLmDirectionality() ? TEXT("1") : TEXT("0"));
diff --git a/Engine/Source/Runtime/Engine/Public/MaterialShared.h b/Engine/Source/Runtime/Engine/Public/MaterialShared.h
--- a/Engine/Source/Runtime/Engine/Public/MaterialShared.h
+++ b/Engine/Source/Runtime/Engine/Public/MaterialShared.h
@@ -252,7 +252,7 @@ class FMaterialCompilationOutput
 {
 public:
 	FMaterialCompilationOutput() :
-		bUsesSceneColor(false),
+		bRequiresSceneColorCopy(false),
 		bNeedsSceneTextures(false),
 		bUsesEyeAdaptation(false),
 		bModifiesMeshPosition(false),
@@ -266,7 +266,7 @@ public:
 	/** 
 	 * Indicates whether the material uses scene color. 
 	 */
-	bool bUsesSceneColor;
+	bool bRequiresSceneColorCopy;
 
 	/** true if the material needs the scenetexture lookups. */
 	bool bNeedsSceneTextures;
@@ -636,7 +636,7 @@ public:
 	bool IsCompilationFinalized() const { return bCompilationFinalized; }
 	bool CompiledSuccessfully() const { return bCompiledSuccessfully; }
 	const FString& GetDebugDescription() const { return DebugDescription; }
-	bool UsesSceneColor() const { return MaterialCompilationOutput.bUsesSceneColor; }
+	bool RequiresSceneColorCopy() const { return MaterialCompilationOutput.bRequiresSceneColorCopy; }
 	bool NeedsSceneTextures() const { return MaterialCompilationOutput.bNeedsSceneTextures; }
 	bool NeedsGBuffer() const { return MaterialCompilationOutput.bNeedsGBuffer; }
 	bool UsesEyeAdaptation() const { return MaterialCompilationOutput.bUsesEyeAdaptation; }
@@ -972,7 +972,7 @@ public:
 		//@todo - remove non-dynamic parameter particle VF and always support dynamic parameter
 		return true; 
 	}
-	ENGINE_API bool UsesSceneColor() const;
+	ENGINE_API bool RequiresSceneColorCopy() const;
 	ENGINE_API bool NeedsSceneTextures() const;
 	ENGINE_API bool NeedsGBuffer() const;
 	ENGINE_API bool UsesEyeAdaptation() const;	
diff --git a/Engine/Source/Runtime/Engine/Public/PrimitiveViewRelevance.h b/Engine/Source/Runtime/Engine/Public/PrimitiveViewRelevance.h
--- a/Engine/Source/Runtime/Engine/Public/PrimitiveViewRelevance.h
+++ b/Engine/Source/Runtime/Engine/Public/PrimitiveViewRelevance.h
@@ -35,8 +35,6 @@ struct FPrimitiveViewRelevance
 	uint32 bMaskedRelevance : 1;
 	/** The primitive has one or more distortion elements. */
 	uint32 bDistortionRelevance : 1;
-	/** The primitive reads from scene color. */
-	uint32 bSceneColorRelevance : 1;
 	/** The primitive has one or more elements that have SeparateTranslucency. */
 	uint32 bSeparateTranslucencyRelevance : 1;
 	/** The primitive has one or more elements that have normal translucency. */
@@ -69,7 +67,6 @@ struct FPrimitiveViewRelevance
 		bOpaqueRelevance(true),
 		bMaskedRelevance(false),
 		bDistortionRelevance(false),
-		bSceneColorRelevance(false),
 		bSeparateTranslucencyRelevance(false),
 		bNormalTranslucencyRelevance(false),
 		bInitializedThisFrame(false)
@@ -89,7 +86,6 @@ struct FPrimitiveViewRelevance
 		bRenderInMainPass |= B.bRenderInMainPass !=0;
 		bEditorPrimitiveRelevance |= B.bEditorPrimitiveRelevance !=0;
 		bEditorNoDepthTestPrimitiveRelevance |= B.bEditorNoDepthTestPrimitiveRelevance !=0;
-		bSceneColorRelevance |= B.bSceneColorRelevance != 0;
 		bNeedsPreRenderView |= B.bNeedsPreRenderView != 0;
 		bHasSimpleLights |= B.bHasSimpleLights != 0;
 		bSeparateTranslucencyRelevance |= B.bSeparateTranslucencyRelevance != 0;
diff --git a/Engine/Source/Runtime/Renderer/Private/DeferredShadingRenderer.h b/Engine/Source/Runtime/Renderer/Private/DeferredShadingRenderer.h
--- a/Engine/Source/Runtime/Renderer/Private/DeferredShadingRenderer.h
+++ b/Engine/Source/Runtime/Renderer/Private/DeferredShadingRenderer.h
@@ -210,9 +210,6 @@ private:
 
 	void RenderLightShaftBloom(FRHICommandList& RHICmdList);
 
-	/** Makes a copy of scene color so that a material with a scene color node can read the up to date scene color. */
-	void CopySceneColor(const FViewInfo& View, const FPrimitiveSceneInfo* PrimitiveSceneInfo);
-
 	/** Reuses an existing translucent shadow map if possible or re-renders one if necessary. */
 	const FProjectedShadowInfo* PrepareTranslucentShadowMap(const FViewInfo& View, FPrimitiveSceneInfo* PrimitiveSceneInfo, bool bSeparateTranslucencyPass);
 
diff --git a/Engine/Source/Runtime/Renderer/Private/TranslucentRendering.cpp b/Engine/Source/Runtime/Renderer/Private/TranslucentRendering.cpp
--- a/Engine/Source/Runtime/Renderer/Private/TranslucentRendering.cpp
+++ b/Engine/Source/Runtime/Renderer/Private/TranslucentRendering.cpp
@@ -166,12 +166,12 @@ IMPLEMENT_SHADER_TYPE(,FCopySceneColorPS,TEXT("TranslucentLightingShaders"),TEXT
 
 FGlobalBoundShaderState CopySceneColorBoundShaderState;
 
-void FDeferredShadingSceneRenderer::CopySceneColor(const FViewInfo& View, const FPrimitiveSceneInfo* PrimitiveSceneInfo)
+void FTranslucencyDrawingPolicyFactory::CopySceneColor(const FViewInfo& View, const FPrimitiveSceneProxy* PrimitiveSceneProxy)
 {
 	//@todo-rco: RHIPacketList
 	FRHICommandList& RHICmdList = FRHICommandList::GetNullRef();
 
-	SCOPED_DRAW_EVENTF(EventCopy, DEC_SCENE_ITEMS, TEXT("CopySceneColor for %s %s"), *PrimitiveSceneInfo->Proxy->GetOwnerName().ToString(), *PrimitiveSceneInfo->Proxy->GetResourceName().ToString());
+	SCOPED_DRAW_EVENTF(EventCopy, DEC_SCENE_ITEMS, TEXT("CopySceneColor for %s %s"), *PrimitiveSceneProxy->GetOwnerName().ToString(), *PrimitiveSceneProxy->GetResourceName().ToString());
 	RHISetRasterizerState(TStaticRasterizerState<FM_Solid,CM_None>::GetRHI());
 	RHISetDepthStencilState(TStaticDepthStencilState<false,CF_Always>::GetRHI());
 	RHISetBlendState(TStaticBlendState<>::GetRHI());
@@ -192,7 +192,7 @@ void FDeferredShadingSceneRenderer::CopySceneColor(const FViewInfo& View, const
 		RHICmdList,
 		0, 0, 
 		View.ViewRect.Width(), View.ViewRect.Height(),
-		View.ViewRect.Min.X, View.ViewRect.Min.Y, 
+		View.ViewRect.Min.X, View.ViewRect.Min.Y,
 		View.ViewRect.Width(), View.ViewRect.Height(),
 		FIntPoint(View.ViewRect.Width(), View.ViewRect.Height()),
 		GSceneRenderTargets.GetBufferSizeXY(),
@@ -332,6 +332,21 @@ bool FTranslucencyDrawingPolicyFactory::DrawMesh(
 	// Only render translucent materials.
 	if(IsTranslucentBlendMode(BlendMode))
 	{
+		if (Material->RequiresSceneColorCopy())
+		{
+			if (DrawingContext.bSceneColorCopyIsUpToDate == false)
+			{
+				CopySceneColor(View, PrimitiveSceneProxy);
+				// Restore state
+				SetTranslucentRenderTargetAndState(View, false);
+
+				// separate translucency is not updating scene color so we don't need to copy it multiple times.
+				// optimization:
+				// we should consider the same for non separate translucency (could cause artifacts but will be much faster)
+				DrawingContext.bSceneColorCopyIsUpToDate = DrawingContext.bSeparateTranslucencyPass;
+			}
+		}
+
 		const bool bDisableDepthTest = Material->ShouldDisableDepthTest();
 		const bool bEnableResponsiveAA = Material->ShouldEnableResponsiveAA();
 		// editor compositing not supported on translucent materials currently
@@ -465,10 +480,6 @@ void FTranslucentPrimSet::DrawPrimitives(
 	const TArray<FSortedPrim,SceneRenderingAllocator>& PhaseSortedPrimitives =
 		bSeparateTranslucencyPass ? SortedSeparateTranslucencyPrims : SortedPrims;
 
-	// optimization - separate translucency is not updating scene color so we don't need to copy it multiple times
-	// we should consider the same for non separate translucency (can cause artifacts there but will be much faster)
-	bool bSceneColorWasAlreadyCopied = false;
-
 	if( PhaseSortedPrimitives.Num() )
 	{
 		// Draw sorted scene prims
@@ -482,18 +493,6 @@ void FTranslucentPrimSet::DrawPrimitives(
 			
 			const FProjectedShadowInfo* TranslucentSelfShadow = Renderer.PrepareTranslucentShadowMap(View, PrimitiveSceneInfo, bSeparateTranslucencyPass);
 
-			if (ViewRelevance.bSceneColorRelevance)
-			{
-				if (!bSeparateTranslucencyPass || !bSceneColorWasAlreadyCopied)
-				{
-					Renderer.CopySceneColor(View, PrimitiveSceneInfo);
-					bSceneColorWasAlreadyCopied = true;
-
-					// Restore state
-					SetTranslucentRenderTargetAndState(View, bSeparateTranslucencyPass);
-				}
-			}
-
 			RenderPrimitive(View, PrimitiveSceneInfo, ViewRelevance, TranslucentSelfShadow, bSeparateTranslucencyPass);
 		}
 	}
@@ -516,7 +515,7 @@ void FTranslucentPrimSet::RenderPrimitive(
 		{
 			TDynamicPrimitiveDrawer<FTranslucencyDrawingPolicyFactory> TranslucencyDrawer(
 				&View,
-				FTranslucencyDrawingPolicyFactory::ContextType(TranslucentSelfShadow),
+				FTranslucencyDrawingPolicyFactory::ContextType(TranslucentSelfShadow, bSeparateTranslucencyPass),
 				false
 				);
 
@@ -548,7 +547,7 @@ void FTranslucentPrimSet::RenderPrimitive(
 					FTranslucencyDrawingPolicyFactory::DrawStaticMesh(
 						RHICmdList, 
 						View,
-						FTranslucencyDrawingPolicyFactory::ContextType(TranslucentSelfShadow),
+						FTranslucencyDrawingPolicyFactory::ContextType( TranslucentSelfShadow, bSeparateTranslucencyPass),
 						StaticMesh,
 						StaticMesh.Elements.Num() == 1 ? 1 : View.StaticMeshBatchVisibility[StaticMesh.Id],
 						false,
diff --git a/Engine/Source/Runtime/Renderer/Private/TranslucentRendering.h b/Engine/Source/Runtime/Renderer/Private/TranslucentRendering.h
--- a/Engine/Source/Runtime/Renderer/Private/TranslucentRendering.h
+++ b/Engine/Source/Runtime/Renderer/Private/TranslucentRendering.h
@@ -17,9 +17,13 @@ public:
 	struct ContextType 
 	{
 		const FProjectedShadowInfo* TranslucentSelfShadow;
+		bool bSeparateTranslucencyPass;
+		bool bSceneColorCopyIsUpToDate;
 
-		ContextType(const FProjectedShadowInfo* InTranslucentSelfShadow = NULL) 
-			: TranslucentSelfShadow(InTranslucentSelfShadow) 
+		ContextType(const FProjectedShadowInfo* InTranslucentSelfShadow = NULL, bool bSeparateTranslucencyPassIn = false)
+			: TranslucentSelfShadow(InTranslucentSelfShadow),
+			bSeparateTranslucencyPass(bSeparateTranslucencyPassIn),
+			bSceneColorCopyIsUpToDate(false)
 		{}
 	};
 
@@ -76,6 +80,11 @@ private:
 		const FPrimitiveSceneProxy* PrimitiveSceneProxy,
 		FHitProxyId HitProxyId
 		);
+
+	/**
+	* Resolves the scene color target and copies it for use as a source texture.
+	*/
+	static void CopySceneColor(const FViewInfo& View, const FPrimitiveSceneProxy* PrimitiveSceneProxy);
 };
 
 
