UE4 - rework RHI user-facing APIs to be safer and incorporate the desired sematics

diff --git a/Engine/Source/Runtime/RHI/Public/RHIUtilities.h b/Engine/Source/Runtime/RHI/Public/RHIUtilities.h
--- a/Engine/Source/Runtime/RHI/Public/RHIUtilities.h
+++ b/Engine/Source/Runtime/RHI/Public/RHIUtilities.h
@@ -16,14 +16,14 @@ struct FRWBuffer
 
 	FRWBuffer(): NumBytes(0) {}
 
-	void Initialize(uint32 BytesPerElement,uint32 NumElements,EPixelFormat Format,uint32 AdditionalUsage = 0)
+	void Initialize(uint32 BytesPerElement, uint32 NumElements, EPixelFormat Format, uint32 AdditionalUsage = 0)
 	{
 		check(GRHIFeatureLevel == ERHIFeatureLevel::SM5);
 		NumBytes = BytesPerElement * NumElements;
 		FRHIResourceCreateInfo CreateInfo;
-		Buffer = RHICreateVertexBuffer(NumBytes,BUF_UnorderedAccess | BUF_ShaderResource | AdditionalUsage, CreateInfo);
-		UAV = RHICreateUnorderedAccessView(Buffer,Format);
-		SRV = RHICreateShaderResourceView(Buffer,BytesPerElement,Format);
+		Buffer = RHICreateVertexBuffer(NumBytes, BUF_UnorderedAccess | BUF_ShaderResource | AdditionalUsage, CreateInfo);
+		UAV = RHICreateUnorderedAccessView(Buffer, Format);
+		SRV = RHICreateShaderResourceView(Buffer, BytesPerElement, Format);
 	}
 
 	void Release()
@@ -44,13 +44,13 @@ struct FReadBuffer
 
 	FReadBuffer(): NumBytes(0) {}
 
-	void Initialize(uint32 BytesPerElement,uint32 NumElements,EPixelFormat Format,uint32 AdditionalUsage = 0)
+	void Initialize(uint32 BytesPerElement, uint32 NumElements, EPixelFormat Format, uint32 AdditionalUsage = 0)
 	{
 		check(GRHIFeatureLevel >= ERHIFeatureLevel::SM4);
 		NumBytes = BytesPerElement * NumElements;
 		FRHIResourceCreateInfo CreateInfo;
-		Buffer = RHICreateVertexBuffer(NumBytes,BUF_ShaderResource | AdditionalUsage,CreateInfo);
-		SRV = RHICreateShaderResourceView(Buffer,BytesPerElement,Format);
+		Buffer = RHICreateVertexBuffer(NumBytes, BUF_ShaderResource | AdditionalUsage, CreateInfo);
+		SRV = RHICreateShaderResourceView(Buffer, BytesPerElement, Format);
 	}
 
 	void Release()
@@ -71,13 +71,13 @@ struct FRWBufferStructured
 
 	FRWBufferStructured(): NumBytes(0) {}
 
-	void Initialize( uint32 BytesPerElement, uint32 NumElements, uint32 AdditionalUsage = 0, bool bUseUavCounter = false, bool bAppendBuffer = false )
+	void Initialize(uint32 BytesPerElement, uint32 NumElements, uint32 AdditionalUsage = 0, bool bUseUavCounter = false, bool bAppendBuffer = false)
 	{
 		check(GRHIFeatureLevel == ERHIFeatureLevel::SM5);
 		NumBytes = BytesPerElement * NumElements;
 		FRHIResourceCreateInfo CreateInfo;
 		Buffer = RHICreateStructuredBuffer(BytesPerElement, NumBytes, BUF_UnorderedAccess | BUF_ShaderResource | AdditionalUsage, CreateInfo);
-		UAV = RHICreateUnorderedAccessView(Buffer,bUseUavCounter, bAppendBuffer);
+		UAV = RHICreateUnorderedAccessView(Buffer, bUseUavCounter, bAppendBuffer);
 		SRV = RHICreateShaderResourceView(Buffer);
 	}
 
@@ -100,14 +100,14 @@ struct FRWBufferByteAddress
 
 	FRWBufferByteAddress(): NumBytes(0) {}
 
-	void Initialize( uint32 InNumBytes, uint32 AdditionalUsage = 0 )
+	void Initialize(uint32 InNumBytes, uint32 AdditionalUsage = 0)
 	{
 		NumBytes = InNumBytes;
 		check(GRHIFeatureLevel == ERHIFeatureLevel::SM5);
 		check( NumBytes % 4 == 0 );
 		FRHIResourceCreateInfo CreateInfo;
 		Buffer = RHICreateStructuredBuffer(4, NumBytes, BUF_UnorderedAccess | BUF_ShaderResource | BUF_ByteAddressBuffer | AdditionalUsage, CreateInfo);
-		UAV = RHICreateUnorderedAccessView(Buffer, false, false );
+		UAV = RHICreateUnorderedAccessView(Buffer, false, false);
 		SRV = RHICreateShaderResourceView(Buffer);
 	}
 
@@ -121,21 +121,21 @@ struct FRWBufferByteAddress
 };
 
 /** Helper for the common case of using a single color and depth render target. */
-inline void SetRenderTarget(FRHICommandList& RHICmdList, FTextureRHIParamRef NewRenderTarget, FTextureRHIParamRef NewDepthStencilTarget)
+inline void SetRenderTarget(FRHICommandListImmediate& RHICmdList, FTextureRHIParamRef NewRenderTarget, FTextureRHIParamRef NewDepthStencilTarget)
 {
 	FRHIRenderTargetView RTV(NewRenderTarget);
 	RHICmdList.SetRenderTargets(1, &RTV, NewDepthStencilTarget, 0, NULL);
 }
 
 /** Helper for the common case of using a single color and depth render target, with a mip index for the color target. */
-inline void SetRenderTarget(FRHICommandList& RHICmdList, FTextureRHIParamRef NewRenderTarget, int32 MipIndex, FTextureRHIParamRef NewDepthStencilTarget)
+inline void SetRenderTarget(FRHICommandListImmediate& RHICmdList, FTextureRHIParamRef NewRenderTarget, int32 MipIndex, FTextureRHIParamRef NewDepthStencilTarget)
 {
 	FRHIRenderTargetView RTV(NewRenderTarget, MipIndex, -1);
 	RHICmdList.SetRenderTargets(1, &RTV, NewDepthStencilTarget, 0, NULL);
 }
 
 /** Helper for the common case of using a single color and depth render target, with a mip index for the color target. */
-inline void SetRenderTarget(FRHICommandList& RHICmdList, FTextureRHIParamRef NewRenderTarget, int32 MipIndex, int32 ArraySliceIndex, FTextureRHIParamRef NewDepthStencilTarget)
+inline void SetRenderTarget(FRHICommandListImmediate& RHICmdList, FTextureRHIParamRef NewRenderTarget, int32 MipIndex, int32 ArraySliceIndex, FTextureRHIParamRef NewDepthStencilTarget)
 {
 	FRHIRenderTargetView RTV(NewRenderTarget, MipIndex, ArraySliceIndex);
 	RHICmdList.SetRenderTargets(1, &RTV, NewDepthStencilTarget, 0, NULL);
@@ -143,7 +143,7 @@ inline void SetRenderTarget(FRHICommandList& RHICmdList, FTextureRHIParamRef New
 
 /** Helper that converts FTextureRHIParamRef's into FRHIRenderTargetView's. */
 inline void SetRenderTargets(
-	FRHICommandList& RHICmdList,
+	FRHICommandListImmediate& RHICmdList,
 	uint32 NewNumSimultaneousRenderTargets, 
 	const FTextureRHIParamRef* NewRenderTargetsRHI,
 	FTextureRHIParamRef NewDepthStencilTargetRHI,
@@ -161,59 +161,6 @@ inline void SetRenderTargets(
 	RHICmdList.SetRenderTargets(NewNumSimultaneousRenderTargets, RTVs, NewDepthStencilTargetRHI, NewNumUAVs, UAVs);
 }
 
-
-
-/** Helper for the common case of using a single color and depth render target. */
-inline void RHISetRenderTarget(FTextureRHIParamRef NewRenderTarget, FTextureRHIParamRef NewDepthStencilTarget)
-{
-	FRHIRenderTargetView RTV(NewRenderTarget);
-	RHISetRenderTargets(1, &RTV, NewDepthStencilTarget, 0, NULL);
-}
-
-/** Helper for the common case of using a single color and depth render target, with a mip index for the color target. */
-inline void RHISetRenderTarget(FTextureRHIParamRef NewRenderTarget, int32 MipIndex, FTextureRHIParamRef NewDepthStencilTarget)
-{
-	FRHIRenderTargetView RTV(NewRenderTarget, MipIndex, -1);
-	RHISetRenderTargets(1, &RTV, NewDepthStencilTarget, 0, NULL);
-}
-
-/** Helper for the common case of using a single color and depth render target, with a mip index for the color target. */
-inline void RHISetRenderTarget(FTextureRHIParamRef NewRenderTarget, int32 MipIndex, int32 ArraySliceIndex, FTextureRHIParamRef NewDepthStencilTarget)
-{
-	FRHIRenderTargetView RTV(NewRenderTarget, MipIndex, ArraySliceIndex);
-	RHISetRenderTargets(1, &RTV, NewDepthStencilTarget, 0, NULL);
-}
-
-/** Helper that converts FTextureRHIParamRef's into FRHIRenderTargetView's. */
-inline void RHISetRenderTargets(
-	uint32 NewNumSimultaneousRenderTargets,
-	const FTextureRHIParamRef* NewRenderTargetsRHI,
-	FTextureRHIParamRef NewDepthStencilTargetRHI,
-	uint32 NewNumUAVs,
-	const FUnorderedAccessViewRHIParamRef* UAVs
-	)
-{
-	FRHIRenderTargetView RTVs[MaxSimultaneousRenderTargets];
-
-	for (uint32 Index = 0; Index < NewNumSimultaneousRenderTargets; Index++)
-	{
-		RTVs[Index] = FRHIRenderTargetView(NewRenderTargetsRHI[Index]);
-	}
-
-	RHISetRenderTargets(NewNumSimultaneousRenderTargets, RTVs, NewDepthStencilTargetRHI, NewNumUAVs, UAVs);
-}
-/** Sets a depth-stencil state with a default stencil ref value of 0. */
-inline void RHISetDepthStencilState(FDepthStencilStateRHIParamRef NewState)
-{
-	RHISetDepthStencilState(NewState,0);
-}
-
-/** Sets a blend state with a default blend factor of white. */
-inline void RHISetBlendState(FBlendStateRHIParamRef NewState)
-{
-	RHISetBlendState(NewState,FLinearColor::White);
-}
-
 /**
  * Creates 1 or 2 textures with the same dimensions/format.
  * If the RHI supports textures that can be used as both shader resources and render targets,
@@ -258,13 +205,13 @@ inline void RHICreateTargetableShaderResource2D(
 	if (!bForceSeparateTargetAndShaderResource/* && GSupportsRenderDepthTargetableShaderResources*/)
 	{
 		// Create a single texture that has both TargetableTextureFlags and TexCreate_ShaderResource set.
-		OutTargetableTexture = OutShaderResourceTexture = RHICreateTexture2D(SizeX,SizeY,Format,NumMips,NumSamples,Flags | TargetableTextureFlags | TexCreate_ShaderResource,CreateInfo);
+		OutTargetableTexture = OutShaderResourceTexture = RHICreateTexture2D(SizeX, SizeY, Format, NumMips, NumSamples, Flags | TargetableTextureFlags | TexCreate_ShaderResource, CreateInfo);
 	}
 	else
 	{
 		// Create a texture that has TargetableTextureFlags set, and a second texture that has TexCreate_ResolveTargetable and TexCreate_ShaderResource set.
-		OutTargetableTexture = RHICreateTexture2D(SizeX,SizeY,Format,NumMips,NumSamples,Flags | TargetableTextureFlags,CreateInfo);
-		OutShaderResourceTexture = RHICreateTexture2D(SizeX,SizeY,Format,NumMips,1,Flags | TexCreate_ResolveTargetable | TexCreate_ShaderResource,CreateInfo);
+		OutTargetableTexture = RHICreateTexture2D(SizeX, SizeY, Format, NumMips, NumSamples, Flags | TargetableTextureFlags, CreateInfo);
+		OutShaderResourceTexture = RHICreateTexture2D(SizeX, SizeY, Format, NumMips, 1, Flags | TexCreate_ResolveTargetable | TexCreate_ShaderResource, CreateInfo);
 	}
 }
 
@@ -305,13 +252,13 @@ inline void RHICreateTargetableShaderResourceCube(
 	if(!bForceSeparateTargetAndShaderResource/* && GSupportsRenderDepthTargetableShaderResources*/)
 	{
 		// Create a single texture that has both TargetableTextureFlags and TexCreate_ShaderResource set.
-		OutTargetableTexture = OutShaderResourceTexture = RHICreateTextureCube(LinearSize,Format,NumMips,Flags | TargetableTextureFlags | TexCreate_ShaderResource,CreateInfo);
+		OutTargetableTexture = OutShaderResourceTexture = RHICreateTextureCube(LinearSize, Format, NumMips, Flags | TargetableTextureFlags | TexCreate_ShaderResource, CreateInfo);
 	}
 	else
 	{
 		// Create a texture that has TargetableTextureFlags set, and a second texture that has TexCreate_ResolveTargetable and TexCreate_ShaderResource set.
-		OutTargetableTexture = RHICreateTextureCube(LinearSize,Format,NumMips,Flags | TargetableTextureFlags,CreateInfo);
-		OutShaderResourceTexture = RHICreateTextureCube(LinearSize,Format,NumMips,Flags | TexCreate_ResolveTargetable | TexCreate_ShaderResource,CreateInfo);
+		OutTargetableTexture = RHICreateTextureCube(LinearSize, Format, NumMips, Flags | TargetableTextureFlags, CreateInfo);
+		OutShaderResourceTexture = RHICreateTextureCube(LinearSize, Format, NumMips, Flags | TexCreate_ResolveTargetable | TexCreate_ShaderResource, CreateInfo);
 	}
 }
 
@@ -350,13 +297,13 @@ inline void RHICreateTargetableShaderResourceCubeArray(
 	if(!bForceSeparateTargetAndShaderResource/* && GSupportsRenderDepthTargetableShaderResources*/)
 	{
 		// Create a single texture that has both TargetableTextureFlags and TexCreate_ShaderResource set.
-		OutTargetableTexture = OutShaderResourceTexture = RHICreateTextureCubeArray(LinearSize,ArraySize,Format,NumMips,Flags | TargetableTextureFlags | TexCreate_ShaderResource,CreateInfo);
+		OutTargetableTexture = OutShaderResourceTexture = RHICreateTextureCubeArray(LinearSize, ArraySize, Format, NumMips, Flags | TargetableTextureFlags | TexCreate_ShaderResource, CreateInfo);
 	}
 	else
 	{
 		// Create a texture that has TargetableTextureFlags set, and a second texture that has TexCreate_ResolveTargetable and TexCreate_ShaderResource set.
-		OutTargetableTexture = RHICreateTextureCubeArray(LinearSize,ArraySize,Format,NumMips,Flags | TargetableTextureFlags,CreateInfo);
-		OutShaderResourceTexture = RHICreateTextureCubeArray(LinearSize,ArraySize,Format,NumMips,Flags | TexCreate_ResolveTargetable | TexCreate_ShaderResource,CreateInfo);
+		OutTargetableTexture = RHICreateTextureCubeArray(LinearSize, ArraySize, Format, NumMips, Flags | TargetableTextureFlags, CreateInfo);
+		OutShaderResourceTexture = RHICreateTextureCubeArray(LinearSize, ArraySize, Format, NumMips, Flags | TexCreate_ResolveTargetable | TexCreate_ShaderResource, CreateInfo);
 	}
 }
 /**
@@ -365,7 +312,7 @@ inline void RHICreateTargetableShaderResourceCubeArray(
  * @param PrimitiveType The type of primitives.
  * @returns The number of vertices.
  */
-inline uint32 RHIGetVertexCountForPrimitiveCount(uint32 NumPrimitives, uint32 PrimitiveType)
+inline uint32 GetVertexCountForPrimitiveCount(uint32 NumPrimitives, uint32 PrimitiveType)
 {
 	uint32 VertexCount = 0;
 	switch(PrimitiveType)
@@ -421,10 +368,10 @@ inline uint32 RHIGetVertexCountForPrimitiveCount(uint32 NumPrimitives, uint32 Pr
  * @param VertexData A reference to memory preallocate in RHIBeginDrawPrimitiveUP
  * @param VertexDataStride Size of each vertex
  */
-inline void DrawPrimitiveUP(FRHICommandList& RHICmdList, uint32 PrimitiveType, uint32 NumPrimitives, const void* VertexData, uint32 VertexDataStride)
+inline void DrawPrimitiveUP(FRHICommandListImmediate& RHICmdList, uint32 PrimitiveType, uint32 NumPrimitives, const void* VertexData, uint32 VertexDataStride)
 {
 	void* Buffer = NULL;
-	const uint32 VertexCount = RHIGetVertexCountForPrimitiveCount( NumPrimitives, PrimitiveType );
+	const uint32 VertexCount = GetVertexCountForPrimitiveCount( NumPrimitives, PrimitiveType );
 	RHICmdList.BeginDrawPrimitiveUP(PrimitiveType, NumPrimitives, VertexCount, VertexDataStride, Buffer);
 	FMemory::Memcpy( Buffer, VertexData, VertexCount * VertexDataStride );
 	RHICmdList.EndDrawPrimitiveUP();
@@ -442,7 +389,7 @@ inline void DrawPrimitiveUP(FRHICommandList& RHICmdList, uint32 PrimitiveType, u
  * @param VertexDataStride The size of one vertex
  */
 inline void DrawIndexedPrimitiveUP(
-	FRHICommandList& RHICmdList,
+	FRHICommandListImmediate& RHICmdList,
 	uint32 PrimitiveType,
 	uint32 MinVertexIndex,
 	uint32 NumVertices,
@@ -454,7 +401,7 @@ inline void DrawIndexedPrimitiveUP(
 {
 	void* VertexBuffer = NULL;
 	void* IndexBuffer = NULL;
-	const uint32 NumIndices = RHIGetVertexCountForPrimitiveCount( NumPrimitives, PrimitiveType );
+	const uint32 NumIndices = GetVertexCountForPrimitiveCount( NumPrimitives, PrimitiveType );
 	RHICmdList.BeginDrawIndexedPrimitiveUP(
 		PrimitiveType,
 		NumPrimitives,
diff --git a/Engine/Source/Runtime/RenderCore/Private/RenderUtils.cpp b/Engine/Source/Runtime/RenderCore/Private/RenderUtils.cpp
--- a/Engine/Source/Runtime/RenderCore/Private/RenderUtils.cpp
+++ b/Engine/Source/Runtime/RenderCore/Private/RenderUtils.cpp
@@ -148,18 +148,18 @@ class FColoredTexture : public FTexture
 {
 public:
 	// FResource interface.
-	virtual void InitRHI()
+	virtual void InitRHI() override
 	{
 		// Create the texture RHI.  		
 		FRHIResourceCreateInfo CreateInfo;
-		FTexture2DRHIRef Texture2D = RHICreateTexture2D(1,1,PF_B8G8R8A8,1,1,TexCreate_ShaderResource,CreateInfo);
+		FTexture2DRHIRef Texture2D = RHICreateTexture2D(1, 1, PF_B8G8R8A8, 1, 1, TexCreate_ShaderResource, CreateInfo);
 		TextureRHI = Texture2D;
 
 		// Write the contents of the texture.
 		uint32 DestStride;
-		FColor* DestBuffer = (FColor*)RHILockTexture2D(Texture2D,0,RLM_WriteOnly,DestStride,false);
+		FColor* DestBuffer = (FColor*)RHILockTexture2D(Texture2D, 0, RLM_WriteOnly, DestStride, false);
 		*DestBuffer = FColor(R, G, B, A);
-		RHIUnlockTexture2D(Texture2D,0,false);
+		RHIUnlockTexture2D(Texture2D, 0, false);
 
 		// Create the sampler state RHI resource.
 		FSamplerStateInitializerRHI SamplerStateInitializer(SF_Point,AM_Wrap,AM_Wrap,AM_Wrap);
@@ -235,7 +235,7 @@ public:
 	/**
 	 * Initialize RHI resources.
 	 */
-	virtual void InitRHI()
+	virtual void InitRHI() override
 	{
 		if (IsFeatureLevelSupported(GRHIShaderPlatform, ERHIFeatureLevel::SM4))
 		{
@@ -275,7 +275,7 @@ class FBlackArrayTexture : public FTexture
 {
 public:
 	// FResource interface.
-	virtual void InitRHI()
+	virtual void InitRHI() override
 	{
 		if (IsFeatureLevelSupported(GRHIShaderPlatform, ERHIFeatureLevel::SM4))
 		{
@@ -327,7 +327,7 @@ public:
 	static const FColor MipColors[NumMips];
 
 	// FResource interface.
-	virtual void InitRHI()
+	virtual void InitRHI() override
 	{
 		// Create the texture RHI.
 		int32 TextureSize = 1 << (NumMips - 1);
@@ -340,7 +340,7 @@ public:
 		int32 Size = TextureSize;
 		for ( int32 MipIndex=0; MipIndex < NumMips; ++MipIndex )
 		{
-			FColor* DestBuffer = (FColor*) RHILockTexture2D(Texture2D,MipIndex,RLM_WriteOnly,DestStride,false);
+			FColor* DestBuffer = (FColor*)RHILockTexture2D(Texture2D, MipIndex, RLM_WriteOnly, DestStride, false);
 			for ( int32 Y=0; Y < Size; ++Y )
 			{
 				for ( int32 X=0; X < Size; ++X )
@@ -349,7 +349,7 @@ public:
 				}
 				DestBuffer += DestStride / sizeof(FColor);
 			}
-			RHIUnlockTexture2D(Texture2D,MipIndex,false);
+			RHIUnlockTexture2D(Texture2D, MipIndex, false);
 			Size >>= 1;
 		}
 
@@ -409,7 +409,7 @@ public:
 	{}
 
 	// FRenderResource interface.
-	virtual void InitRHI()
+	virtual void InitRHI() override
 	{
 		// Create the texture RHI.
 		FRHIResourceCreateInfo CreateInfo;
@@ -420,9 +420,9 @@ public:
 		for(uint32 FaceIndex = 0;FaceIndex < 6;FaceIndex++)
 		{
 			uint32 DestStride;
-			FColor* DestBuffer = (FColor*)RHILockTextureCubeFace(TextureCube,FaceIndex,0,0,RLM_WriteOnly,DestStride,false);
+			FColor* DestBuffer = (FColor*)RHILockTextureCubeFace(TextureCube, FaceIndex, 0, 0, RLM_WriteOnly, DestStride, false);
 			*DestBuffer = Color;
-			RHIUnlockTextureCubeFace(TextureCube,FaceIndex,0,0,false);
+			RHIUnlockTextureCubeFace(TextureCube, FaceIndex, 0, 0, false);
 		}
 
 		// Create the sampler state RHI resource.
@@ -467,7 +467,7 @@ class FBlackCubeArrayTexture : public FTexture
 {
 public:
 	// FResource interface.
-	virtual void InitRHI()
+	virtual void InitRHI() override
 	{
 		if (IsFeatureLevelSupported(GRHIShaderPlatform, ERHIFeatureLevel::SM5))
 		{
@@ -479,10 +479,10 @@ public:
 			for(uint32 FaceIndex = 0;FaceIndex < 6;FaceIndex++)
 			{
 				uint32 DestStride;
-				FColor* DestBuffer = (FColor*)RHILockTextureCubeFace(TextureCubeArray,FaceIndex,0,0,RLM_WriteOnly,DestStride,false);
+				FColor* DestBuffer = (FColor*)RHILockTextureCubeFace(TextureCubeArray, FaceIndex, 0, 0, RLM_WriteOnly, DestStride, false);
 				// Note: alpha is used by reflection environment to say how much of the foreground texture is visible, so 0 says it is completely invisible
 				*DestBuffer = FColor(0, 0, 0, 0);
-				RHIUnlockTextureCubeFace(TextureCubeArray,FaceIndex,0,0,false);
+				RHIUnlockTextureCubeFace(TextureCubeArray, FaceIndex, 0, 0, false);
 			}
 
 			// Create the sampler state RHI resource.
@@ -724,13 +724,13 @@ class FVector4VertexDeclaration : public FRenderResource
 {
 public:
 	FVertexDeclarationRHIRef VertexDeclarationRHI;
-	virtual void InitRHI()
+	virtual void InitRHI() override
 	{
 		FVertexDeclarationElementList Elements;
 		Elements.Add(FVertexElement(0, 0, VET_Float4, 0, sizeof(FVector4)));
 		VertexDeclarationRHI = RHICreateVertexDeclaration(Elements);
 	}
-	virtual void ReleaseRHI()
+	virtual void ReleaseRHI() override
 	{
 		VertexDeclarationRHI.SafeRelease();
 	}
@@ -747,13 +747,13 @@ class FVector3VertexDeclaration : public FRenderResource
 {
 public:
 	FVertexDeclarationRHIRef VertexDeclarationRHI;
-	virtual void InitRHI()
+	virtual void InitRHI() override
 	{
 		FVertexDeclarationElementList Elements;
 		Elements.Add(FVertexElement(0, 0, VET_Float3, 0, sizeof(FVector)));
 		VertexDeclarationRHI = RHICreateVertexDeclaration(Elements);
 	}
-	virtual void ReleaseRHI()
+	virtual void ReleaseRHI() override
 	{
 		VertexDeclarationRHI.SafeRelease();
 	}
diff --git a/Engine/Source/Runtime/RenderCore/Private/RenderingThread.cpp b/Engine/Source/Runtime/RenderCore/Private/RenderingThread.cpp
--- a/Engine/Source/Runtime/RenderCore/Private/RenderingThread.cpp
+++ b/Engine/Source/Runtime/RenderCore/Private/RenderingThread.cpp
@@ -36,22 +36,6 @@ FString GRenderingThreadError;
  */
 volatile bool GIsRenderingThreadHealthy = true;
 
-/** Whether the rendering thread is suspended (not even processing the tickables) */
-volatile int32 GIsRenderingThreadSuspended = 0;
-
-/** Thread used for rendering */
-RENDERCORE_API FRunnableThread* GRenderingThread = NULL;
-
-// Unlike IsInRenderingThread, this will always return false if we are running single threaded. It only returns true if this is actually a separate rendering thread. Mostly useful for checks
-RENDERCORE_API bool IsInActualRenderingThread()
-{
-	return GRenderingThread && FPlatformTLS::GetCurrentThreadId() == GRenderingThread->GetThreadID();
-}
-
-RENDERCORE_API bool IsInRenderingThread()
-{
-	return !GRenderingThread || GIsRenderingThreadSuspended || (FPlatformTLS::GetCurrentThreadId() == GRenderingThread->GetThreadID());
-}
 
 /**
  * Maximum rate the rendering thread will tick tickables when idle (in Hz)
@@ -778,6 +762,10 @@ void FlushRenderingCommands()
 	delete PendingCleanupObjects;
 }
 
+FRHICommandListImmediate& GetImmediateCommandList_ForRenderCommand()
+{
+	return FRHICommandListExecutor::GetImmediateCommandList();
+}
 
 /** The set of deferred cleanup objects which are pending cleanup. */
 static TLockFreePointerList<FDeferredCleanupInterface>	PendingCleanupObjectsList;
diff --git a/Engine/Source/Runtime/RenderCore/Public/RenderCore.h b/Engine/Source/Runtime/RenderCore/Public/RenderCore.h
--- a/Engine/Source/Runtime/RenderCore/Public/RenderCore.h
+++ b/Engine/Source/Runtime/RenderCore/Public/RenderCore.h
@@ -4,9 +4,7 @@
 	RenderCore.h: Render core module definitions.
 =============================================================================*/
 
-#ifndef __RenderCore_H__
-#define __RenderCore_H__
-
+#pragma once
 #include "Core.h"
 
 DECLARE_LOG_CATEGORY_EXTERN(LogRendererCore, Log, All);
@@ -271,4 +269,3 @@ extern RENDERCORE_API uint32 GSwapBufferTime;
 // shared by renderer and engine, compiles down to a constant in final release
 RENDERCORE_API int32 GetCVarForceLOD();
 
-#endif
diff --git a/Engine/Source/Runtime/RenderCore/Public/RenderResource.h b/Engine/Source/Runtime/RenderCore/Public/RenderResource.h
--- a/Engine/Source/Runtime/RenderCore/Public/RenderResource.h
+++ b/Engine/Source/Runtime/RenderCore/Public/RenderResource.h
@@ -324,7 +324,7 @@ public:
 	}
 
 	// FRenderResource interface.
-	virtual void ReleaseRHI()
+	virtual void ReleaseRHI() override
 	{
 		TextureRHI.SafeRelease();
 		SamplerStateRHI.SafeRelease();
@@ -343,7 +343,7 @@ public:
 	virtual ~FVertexBuffer() {}
 
 	// FRenderResource interface.
-	virtual void ReleaseRHI()
+	virtual void ReleaseRHI() override
 	{
 		VertexBufferRHI.SafeRelease();
 	}
@@ -360,11 +360,11 @@ public:
 	/** 
 	* Initialize the RHI for this rendering resource 
 	*/
-	virtual void InitRHI()
+	virtual void InitRHI() override
 	{
 		// create a static vertex buffer
 		FRHIResourceCreateInfo CreateInfo;
-		VertexBufferRHI = RHICreateVertexBuffer(sizeof(uint32), BUF_Static|BUF_ZeroStride,CreateInfo);
+		VertexBufferRHI = RHICreateVertexBuffer(sizeof(uint32), BUF_Static | BUF_ZeroStride, CreateInfo);
 		uint32* Vertices = (uint32*)RHILockVertexBuffer(VertexBufferRHI, 0, sizeof(uint32), RLM_WriteOnly);
 		Vertices[0] = FColor(255, 255, 255, 255).DWColor();
 		RHIUnlockVertexBuffer(VertexBufferRHI);
@@ -384,7 +384,7 @@ public:
 	virtual ~FIndexBuffer() {}
 
 	// FRenderResource interface.
-	virtual void ReleaseRHI()
+	virtual void ReleaseRHI() override
 	{
 		IndexBufferRHI.SafeRelease();
 	}
diff --git a/Engine/Source/Runtime/RenderCore/Public/RenderingThread.h b/Engine/Source/Runtime/RenderCore/Public/RenderingThread.h
--- a/Engine/Source/Runtime/RenderCore/Public/RenderingThread.h
+++ b/Engine/Source/Runtime/RenderCore/Public/RenderingThread.h
@@ -15,16 +15,6 @@
 // Render thread API
 ////////////////////////////////////
 
-/** @return True if called from the rendering thread, or if called from ANY thread during single threaded rendering */
-extern RENDERCORE_API bool IsInRenderingThread();
-
-/** @return True if called from the rendering thread. */
-// Unlike IsInRenderingThread, this will always return false if we are running single threaded. It only returns true if this is actually a separate rendering thread. Mostly useful for checks
-extern RENDERCORE_API bool IsInActualRenderingThread();
-
-/** Thread used for rendering */
-extern RENDERCORE_API FRunnableThread* GRenderingThread;
-
 /**
  * Whether the renderer is currently running in a separate thread.
  * If this is false, then all rendering commands will be executed immediately instead of being enqueued in the rendering command buffer.
@@ -160,12 +150,16 @@ public:
 //
 // Macros for using render commands.
 //
+// ideally this would be inline, however that changes the module dependency situation
+extern RENDERCORE_API class FRHICommandListImmediate& GetImmediateCommandList_ForRenderCommand();
+
 
 DECLARE_STATS_GROUP(TEXT("Render Thread Commands"), STATGROUP_RenderThreadCommands, STATCAT_Advanced);
 
 #define TASK_FUNCTION(Code) \
 		void DoTask(ENamedThreads::Type CurrentThread, const FGraphEventRef& MyCompletionGraphEvent) \
 		{ \
+			FRHICommandListImmediate& RHICmdList = GetImmediateCommandList_ForRenderCommand(); \
 			Code; \
 		} 
 
diff --git a/Engine/Source/Runtime/RenderCore/Public/UniformBuffer.h b/Engine/Source/Runtime/RenderCore/Public/UniformBuffer.h
--- a/Engine/Source/Runtime/RenderCore/Public/UniformBuffer.h
+++ b/Engine/Source/Runtime/RenderCore/Public/UniformBuffer.h
@@ -49,7 +49,7 @@ public:
 	}
 	
 	// FRenderResource interface.
-	virtual void InitDynamicRHI()
+	virtual void InitDynamicRHI() override
 	{
 		UniformBufferRHI = RHICreateUniformBuffer(Contents,TBufferStruct::StaticStruct.GetLayout(),BufferUsage);
 	}
diff --git a/Engine/Source/Runtime/Renderer/Private/AmbientCubemapParameters.h b/Engine/Source/Runtime/Renderer/Private/AmbientCubemapParameters.h
--- a/Engine/Source/Runtime/Renderer/Private/AmbientCubemapParameters.h
+++ b/Engine/Source/Runtime/Renderer/Private/AmbientCubemapParameters.h
@@ -13,8 +13,8 @@ public:
 
 	void Bind(const FShaderParameterMap& ParameterMap);
 
-	void SetParameters(const FPixelShaderRHIParamRef ShaderRHI, const FFinalPostProcessSettings::FCubemapEntry& Entry) const;
-	void SetParameters(const FComputeShaderRHIParamRef ShaderRHI, const FFinalPostProcessSettings::FCubemapEntry& Entry) const;
+	void SetParameters(FRHICommandList& RHICmdList, const FPixelShaderRHIParamRef ShaderRHI, const FFinalPostProcessSettings::FCubemapEntry& Entry) const;
+	void SetParameters(FRHICommandList& RHICmdList, const FComputeShaderRHIParamRef ShaderRHI, const FFinalPostProcessSettings::FCubemapEntry& Entry) const;
 
 	friend FArchive& operator<<(FArchive& Ar, FCubemapShaderParameters& P);
 
@@ -25,6 +25,6 @@ private:
 	FShaderResourceParameter AmbientCubemapSampler;
 
 	template<typename TShaderRHIRef>
-	void SetParametersTemplate(const TShaderRHIRef& ShaderRHI, const FFinalPostProcessSettings::FCubemapEntry& Entry) const;
+	void SetParametersTemplate(FRHICommandList& RHICmdList, const TShaderRHIRef& ShaderRHI, const FFinalPostProcessSettings::FCubemapEntry& Entry) const;
 };
 
diff --git a/Engine/Source/Runtime/Renderer/Private/AtmosphereRendering.cpp b/Engine/Source/Runtime/Renderer/Private/AtmosphereRendering.cpp
--- a/Engine/Source/Runtime/Renderer/Private/AtmosphereRendering.cpp
+++ b/Engine/Source/Runtime/Renderer/Private/AtmosphereRendering.cpp
@@ -262,14 +262,14 @@ public:
 	// Destructor
 	virtual ~FAtmopshereVertexDeclaration() {}
 
-	virtual void InitRHI()
+	virtual void InitRHI() override
 	{
 		FVertexDeclarationElementList Elements;
 		Elements.Add(FVertexElement(0, 0, VET_Float2, 0, sizeof(FVector2D)));
 		VertexDeclarationRHI = RHICreateVertexDeclaration(Elements);
 	}
 
-	virtual void ReleaseRHI()
+	virtual void ReleaseRHI() override
 	{
 		VertexDeclarationRHI.SafeRelease();
 	}
@@ -351,7 +351,7 @@ void FSceneRenderer::InitAtmosphereConstants()
 
 FGlobalBoundShaderState AtmosphereBoundShaderState[EAtmosphereRenderFlag::E_RenderFlagMax];
 
-void SetAtmosphericFogShaders(FRHICommandList& RHICmdList, FScene* Scene, const FViewInfo& View, TRefCountPtr<IPooledRenderTarget>& LightShaftOcclusion)
+void SetAtmosphericFogShaders(FRHICommandListImmediate& RHICmdList, FScene* Scene, const FViewInfo& View, TRefCountPtr<IPooledRenderTarget>& LightShaftOcclusion)
 {
 	TShaderMapRef<FAtmosphericVS> VertexShader(GetGlobalShaderMap());
 	FAtmosphericFogPS* PixelShader = NULL;
@@ -389,7 +389,7 @@ void SetAtmosphericFogShaders(FRHICommandList& RHICmdList, FScene* Scene, const
 	PixelShader->SetParameters(RHICmdList, View, LightShaftOcclusion);
 }
 
-void FDeferredShadingSceneRenderer::RenderAtmosphere(FLightShaftsOutput LightShaftsOutput)
+void FDeferredShadingSceneRenderer::RenderAtmosphere(FRHICommandListImmediate& RHICmdList, FLightShaftsOutput LightShaftsOutput)
 {
 	// Atmospheric fog?
 	if (Scene->GetFeatureLevel() >= ERHIFeatureLevel::SM4 && Scene->HasAtmosphericFog())
@@ -409,23 +409,21 @@ void FDeferredShadingSceneRenderer::RenderAtmosphere(FLightShaftsOutput LightSha
 			0, 2, 3
 		};
 
-		//@todo-rco: RHIPacketList
-		FRHICommandList& RHICmdList = FRHICommandList::GetNullRef();
-
-		GSceneRenderTargets.BeginRenderingSceneColor();
+		
+		GSceneRenderTargets.BeginRenderingSceneColor(RHICmdList);
 		for(int32 ViewIndex = 0;ViewIndex < Views.Num();ViewIndex++)
 		{
 			const FViewInfo& View = Views[ViewIndex];
 
 			// Set the device viewport for the view.
-			RHISetViewport(View.ViewRect.Min.X, View.ViewRect.Min.Y, 0.0f, View.ViewRect.Max.X, View.ViewRect.Max.Y, 1.0f);
+			RHICmdList.SetViewport(View.ViewRect.Min.X, View.ViewRect.Min.Y, 0.0f, View.ViewRect.Max.X, View.ViewRect.Max.Y, 1.0f);
 			
-			RHISetRasterizerState(TStaticRasterizerState<FM_Solid,CM_None>::GetRHI());
+			RHICmdList.SetRasterizerState(TStaticRasterizerState<FM_Solid, CM_None>::GetRHI());
 			
 			// disable alpha writes in order to preserve scene depth values on PC
-			RHISetBlendState(TStaticBlendState<CW_RGB,BO_Add,BF_One,BF_SourceAlpha>::GetRHI());
+			RHICmdList.SetBlendState(TStaticBlendState<CW_RGB, BO_Add, BF_One, BF_SourceAlpha>::GetRHI());
 
-			RHISetDepthStencilState(TStaticDepthStencilState<false,CF_Always>::GetRHI());
+			RHICmdList.SetDepthStencilState(TStaticDepthStencilState<false,CF_Always>::GetRHI());
 
 			SetAtmosphericFogShaders(RHICmdList, Scene, View, LightShaftsOutput.LightShaftOcclusion);
 
@@ -444,7 +442,7 @@ void FDeferredShadingSceneRenderer::RenderAtmosphere(FLightShaftsOutput LightSha
 		}
 
 		//no need to resolve since we used alpha blending
-		GSceneRenderTargets.FinishRenderingSceneColor(false);
+		GSceneRenderTargets.FinishRenderingSceneColor(RHICmdList, false);
 	}
 }
 
@@ -1083,7 +1081,7 @@ void FAtmosphericFogSceneInfo::GetLayerValue(int Layer, float& AtmosphereR, FVec
 	DhdH = FVector4(DMin, DMax, DMinP, DMaxP);
 }
 
-void FAtmosphericFogSceneInfo::RenderAtmosphereShaders(FRHICommandList& RHICmdList, const FViewInfo& View, const FIntRect& ViewRect)
+void FAtmosphericFogSceneInfo::RenderAtmosphereShaders(FRHICommandListImmediate& RHICmdList, const FViewInfo& View, const FIntRect& ViewRect)
 {
 	check(Component != NULL);
 	switch(AtmospherePhase)
@@ -1091,7 +1089,7 @@ void FAtmosphericFogSceneInfo::RenderAtmosphereShaders(FRHICommandList& RHICmdLi
 	case AP_Transmittance:
 		{
 			const FSceneRenderTargetItem& DestRenderTarget = AtmosphereTextures->AtmosphereTransmittance->GetRenderTargetItem();
-			RHISetRenderTarget(DestRenderTarget.TargetableTexture, FTextureRHIRef());
+			SetRenderTarget(RHICmdList, DestRenderTarget.TargetableTexture, FTextureRHIRef());
 
 			TShaderMapRef<FAtmospherePrecomputeVS> VertexShader(GetGlobalShaderMap());
 			TShaderMapRef<FAtmosphereTransmittancePS> PixelShader(GetGlobalShaderMap());
@@ -1136,7 +1134,7 @@ void FAtmosphericFogSceneInfo::RenderAtmosphereShaders(FRHICommandList& RHICmdLi
 				TShaderMapRef<FAtmosphereInscatter1PS> PixelShader(GetGlobalShaderMap());
 
 				static FGlobalBoundShaderState BoundShaderState;
-				RHICmdList.CheckIsNull(); // need new approach for "static FGlobalBoundShaderState" for parallel rendering
+				
 				SetGlobalBoundShaderState(RHICmdList, BoundShaderState, GFilterVertexDeclaration.VertexDeclarationRHI, *VertexShader, *PixelShader, *GeometryShader);
 
 				//
@@ -1178,7 +1176,7 @@ void FAtmosphericFogSceneInfo::RenderAtmosphereShaders(FRHICommandList& RHICmdLi
 				TShaderMapRef<FAtmosphereCopyInscatter1PS> PixelShader(GetGlobalShaderMap());
 
 				static FGlobalBoundShaderState BoundShaderState;
-				RHICmdList.CheckIsNull(); // need new approach for "static FGlobalBoundShaderState" for parallel rendering
+				
 				SetGlobalBoundShaderState(RHICmdList, BoundShaderState, GFilterVertexDeclaration.VertexDeclarationRHI, *VertexShader, *PixelShader, *GeometryShader);
 
 				float r;
@@ -1207,7 +1205,7 @@ void FAtmosphericFogSceneInfo::RenderAtmosphereShaders(FRHICommandList& RHICmdLi
 				TShaderMapRef<FAtmosphereInscatterSPS> PixelShader(GetGlobalShaderMap());
 
 				static FGlobalBoundShaderState BoundShaderState;
-				RHICmdList.CheckIsNull(); // need new approach for "static FGlobalBoundShaderState" for parallel rendering
+				
 				SetGlobalBoundShaderState(RHICmdList, BoundShaderState, GFilterVertexDeclaration.VertexDeclarationRHI, *VertexShader, *PixelShader, *GeometryShader);
 
 				float r;
@@ -1233,7 +1231,7 @@ void FAtmosphericFogSceneInfo::RenderAtmosphereShaders(FRHICommandList& RHICmdLi
 			TShaderMapRef<FAtmosphereIrradianceNPS> PixelShader(GetGlobalShaderMap());
 
 			static FGlobalBoundShaderState BoundShaderState;
-			RHICmdList.CheckIsNull(); // need new approach for "static FGlobalBoundShaderState" for parallel rendering
+			
 			SetGlobalBoundShaderState(RHICmdList, BoundShaderState, GFilterVertexDeclaration.VertexDeclarationRHI, *VertexShader, *PixelShader);
 
 			PixelShader->SetParameters(RHICmdList, View, AtmoshpereOrder == 2 ? 1.f : 0.f, AtmosphereTextures);
@@ -1256,7 +1254,7 @@ void FAtmosphericFogSceneInfo::RenderAtmosphereShaders(FRHICommandList& RHICmdLi
 				TShaderMapRef<FAtmosphereInscatterNPS> PixelShader(GetGlobalShaderMap());
 
 				static FGlobalBoundShaderState BoundShaderState;
-				RHICmdList.CheckIsNull(); // need new approach for "static FGlobalBoundShaderState" for parallel rendering
+				
 				SetGlobalBoundShaderState(RHICmdList, BoundShaderState, GFilterVertexDeclaration.VertexDeclarationRHI, *VertexShader, *PixelShader, *GeometryShader);
 
 				float r;
@@ -1284,7 +1282,7 @@ void FAtmosphericFogSceneInfo::RenderAtmosphereShaders(FRHICommandList& RHICmdLi
 			TShaderMapRef<FAtmosphereCopyIrradiancePS> PixelShader(GetGlobalShaderMap());
 
 			static FGlobalBoundShaderState BoundShaderState;
-			RHICmdList.CheckIsNull(); // need new approach for "static FGlobalBoundShaderState" for parallel rendering
+			
 			SetGlobalBoundShaderState(RHICmdList, BoundShaderState, GFilterVertexDeclaration.VertexDeclarationRHI, *VertexShader, *PixelShader);
 
 			PixelShader->SetParameters(RHICmdList, AtmosphereTextures);
@@ -1311,7 +1309,7 @@ void FAtmosphericFogSceneInfo::RenderAtmosphereShaders(FRHICommandList& RHICmdLi
 				TShaderMapRef<FAtmosphereCopyInscatterNPS> PixelShader(GetGlobalShaderMap());
 
 				static FGlobalBoundShaderState BoundShaderState;
-				RHICmdList.CheckIsNull(); // need new approach for "static FGlobalBoundShaderState" for parallel rendering
+				
 				SetGlobalBoundShaderState(RHICmdList, BoundShaderState, GFilterVertexDeclaration.VertexDeclarationRHI, *VertexShader, *PixelShader, *GeometryShader);
 
 				float r;
@@ -1343,7 +1341,7 @@ void FAtmosphericFogSceneInfo::RenderAtmosphereShaders(FRHICommandList& RHICmdLi
 				TShaderMapRef<FAtmosphereCopyInscatterFPS> PixelShader(GetGlobalShaderMap());
 
 				static FGlobalBoundShaderState BoundShaderState;
-				RHICmdList.CheckIsNull(); // need new approach for "static FGlobalBoundShaderState" for parallel rendering
+				
 				SetGlobalBoundShaderState(RHICmdList, BoundShaderState, GFilterVertexDeclaration.VertexDeclarationRHI, *VertexShader, *PixelShader, *GeometryShader);
 
 				float r;
@@ -1372,7 +1370,7 @@ void FAtmosphericFogSceneInfo::RenderAtmosphereShaders(FRHICommandList& RHICmdLi
 				TShaderMapRef<FAtmosphereCopyInscatterFBackPS> PixelShader(GetGlobalShaderMap());
 
 				static FGlobalBoundShaderState BoundShaderState;
-				RHICmdList.CheckIsNull(); // need new approach for "static FGlobalBoundShaderState" for parallel rendering
+				
 				SetGlobalBoundShaderState(RHICmdList, BoundShaderState, GFilterVertexDeclaration.VertexDeclarationRHI, *VertexShader, *PixelShader, *GeometryShader);
 
 				float r;
@@ -1393,7 +1391,7 @@ void FAtmosphericFogSceneInfo::RenderAtmosphereShaders(FRHICommandList& RHICmdLi
 	}
 }
 
-void FAtmosphericFogSceneInfo::PrecomputeAtmosphereData(FRHICommandList& RHICmdList, const FViewInfo* View, FSceneViewFamily& ViewFamily)
+void FAtmosphericFogSceneInfo::PrecomputeAtmosphereData(FRHICommandListImmediate& RHICmdList, const FViewInfo* View, FSceneViewFamily& ViewFamily)
 {
 	// Set the view family's render target/viewport.
 	FIntPoint TexSize = GetTextureSize();
@@ -1410,11 +1408,11 @@ void FAtmosphericFogSceneInfo::PrecomputeAtmosphereData(FRHICommandList& RHICmdL
 	RenderAtmosphereShaders(RHICmdList, *View, ViewRect);
 }
 
-void FAtmosphericFogSceneInfo::ReadPixelsPtr(TRefCountPtr<IPooledRenderTarget> RenderTarget, FColor* OutData, FIntRect InRect)
+void FAtmosphericFogSceneInfo::ReadPixelsPtr(FRHICommandListImmediate& RHICmdList, TRefCountPtr<IPooledRenderTarget> RenderTarget, FColor* OutData, FIntRect InRect)
 {
 	TArray<FFloat16Color> Data;
 
-	RHIReadSurfaceFloatData(
+	RHICmdList.ReadSurfaceFloatData(
 		RenderTarget->GetRenderTargetItem().ShaderResourceTexture,
 		InRect,
 		Data,
@@ -1434,11 +1432,11 @@ void FAtmosphericFogSceneInfo::ReadPixelsPtr(TRefCountPtr<IPooledRenderTarget> R
 	}
 }
 
-void FAtmosphericFogSceneInfo::Read3DPixelsPtr(TRefCountPtr<IPooledRenderTarget> RenderTarget, FFloat16Color* OutData, FIntRect InRect, FIntPoint InZMinMax)
+void FAtmosphericFogSceneInfo::Read3DPixelsPtr(FRHICommandListImmediate& RHICmdList, TRefCountPtr<IPooledRenderTarget> RenderTarget, FFloat16Color* OutData, FIntRect InRect, FIntPoint InZMinMax)
 {
 	TArray<FFloat16Color> Data;
 
-	RHIRead3DSurfaceFloatData(
+	RHICmdList.Read3DSurfaceFloatData(
 		RenderTarget->GetRenderTargetItem().ShaderResourceTexture,
 		InRect,
 		InZMinMax,
@@ -1448,7 +1446,7 @@ void FAtmosphericFogSceneInfo::Read3DPixelsPtr(TRefCountPtr<IPooledRenderTarget>
 	FMemory::Memcpy( OutData, Data.GetTypedData(), Data.Num() * sizeof(FFloat16Color) );
 }
 
-void FAtmosphericFogSceneInfo::PrecomputeTextures(FRHICommandList& RHICmdList, const FViewInfo* View, FSceneViewFamily* ViewFamily)
+void FAtmosphericFogSceneInfo::PrecomputeTextures(FRHICommandListImmediate& RHICmdList, const FViewInfo* View, FSceneViewFamily* ViewFamily)
 {
 	check(Component != NULL);
 	if (AtmosphereTextures == NULL)
@@ -1508,7 +1506,6 @@ void FAtmosphericFogSceneInfo::PrecomputeTextures(FRHICommandList& RHICmdList, c
 				AtmoshpereOrder = 2;
 			}
 		}
-		RHICmdList.CheckIsNull(); // this is all synchronous readback
 
 		if (AtmospherePhase >= AP_MAX)
 		{
@@ -1522,7 +1519,7 @@ void FAtmosphericFogSceneInfo::PrecomputeTextures(FRHICommandList& RHICmdList, c
 				int32 TotalByte = sizeof(FColor) * Extent.X * Extent.Y;
 				PrecomputeTransmittance.Lock(LOCK_READ_WRITE);
 				FColor* TransmittanceData = (FColor*)PrecomputeTransmittance.Realloc(TotalByte);
-				ReadPixelsPtr(AtmosphereTextures->AtmosphereTransmittance, TransmittanceData, FIntRect(0, 0, Extent.X, Extent.Y));
+				ReadPixelsPtr(RHICmdList, AtmosphereTextures->AtmosphereTransmittance, TransmittanceData, FIntRect(0, 0, Extent.X, Extent.Y));
 				PrecomputeTransmittance.Unlock();
 			}
 
@@ -1531,7 +1528,7 @@ void FAtmosphericFogSceneInfo::PrecomputeTextures(FRHICommandList& RHICmdList, c
 				int32 TotalByte = sizeof(FColor) * Extent.X * Extent.Y;
 				PrecomputeIrradiance.Lock(LOCK_READ_WRITE);
 				FColor* IrradianceData = (FColor*)PrecomputeIrradiance.Realloc(TotalByte);
-				ReadPixelsPtr(AtmosphereTextures->AtmosphereIrradiance, IrradianceData, FIntRect(0, 0, Extent.X, Extent.Y));
+				ReadPixelsPtr(RHICmdList, AtmosphereTextures->AtmosphereIrradiance, IrradianceData, FIntRect(0, 0, Extent.X, Extent.Y));
 				PrecomputeIrradiance.Unlock();
 			}
 
@@ -1542,7 +1539,7 @@ void FAtmosphericFogSceneInfo::PrecomputeTextures(FRHICommandList& RHICmdList, c
 				int32 TotalByte = sizeof(FFloat16Color) * SizeX * SizeY * SizeZ;
 				PrecomputeInscatter.Lock(LOCK_READ_WRITE);
 				FFloat16Color* InscatterData = (FFloat16Color*)PrecomputeInscatter.Realloc(TotalByte);
-				Read3DPixelsPtr(AtmosphereTextures->AtmosphereInscatter, InscatterData, FIntRect(0, 0, SizeX, SizeY), FIntPoint(0, SizeZ));
+				Read3DPixelsPtr(RHICmdList, AtmosphereTextures->AtmosphereInscatter, InscatterData, FIntRect(0, 0, SizeX, SizeY), FIntPoint(0, SizeZ));
 				PrecomputeInscatter.Unlock();
 			}
 
diff --git a/Engine/Source/Runtime/Renderer/Private/AtmosphereRendering.h b/Engine/Source/Runtime/Renderer/Private/AtmosphereRendering.h
--- a/Engine/Source/Runtime/Renderer/Private/AtmosphereRendering.h
+++ b/Engine/Source/Runtime/Renderer/Private/AtmosphereRendering.h
@@ -64,7 +64,7 @@ public:
 	~FAtmosphericFogSceneInfo();
 
 #if WITH_EDITOR
-	void PrecomputeTextures(FRHICommandList& RHICmdList, const FViewInfo* View, FSceneViewFamily* ViewFamily);
+	void PrecomputeTextures(FRHICommandListImmediate& RHICmdList, const FViewInfo* View, FSceneViewFamily* ViewFamily);
 	void StartPrecompute();
 
 private:
@@ -72,11 +72,11 @@ private:
 	FIntPoint GetTextureSize();
 	inline void DrawQuad(FRHICommandList& RHICmdList, const FIntRect& ViewRect, FShader* VertexShader);
 	void GetLayerValue(int Layer, float& AtmosphereR, FVector4& DhdH);
-	void RenderAtmosphereShaders(FRHICommandList& RHICmdList, const FViewInfo& View, const FIntRect& ViewRect);
-	void PrecomputeAtmosphereData(FRHICommandList& RHICmdList, const FViewInfo* View, FSceneViewFamily& ViewFamily);
+	void RenderAtmosphereShaders(FRHICommandListImmediate& RHICmdList, const FViewInfo& View, const FIntRect& ViewRect);
+	void PrecomputeAtmosphereData(FRHICommandListImmediate& RHICmdList, const FViewInfo* View, FSceneViewFamily& ViewFamily);
 
-	void ReadPixelsPtr(TRefCountPtr<IPooledRenderTarget> RenderTarget, FColor* OutData, FIntRect InRect);
-	void Read3DPixelsPtr(TRefCountPtr<IPooledRenderTarget> RenderTarget, FFloat16Color* OutData, FIntRect InRect, FIntPoint InZMinMax);
+	void ReadPixelsPtr(FRHICommandListImmediate& RHICmdList, TRefCountPtr<IPooledRenderTarget> RenderTarget, FColor* OutData, FIntRect InRect);
+	void Read3DPixelsPtr(FRHICommandListImmediate& RHICmdList, TRefCountPtr<IPooledRenderTarget> RenderTarget, FFloat16Color* OutData, FIntRect InRect, FIntPoint InZMinMax);
 #endif
 };
 
diff --git a/Engine/Source/Runtime/Renderer/Private/AtmosphereTextures.cpp b/Engine/Source/Runtime/Renderer/Private/AtmosphereTextures.cpp
--- a/Engine/Source/Runtime/Renderer/Private/AtmosphereTextures.cpp
+++ b/Engine/Source/Runtime/Renderer/Private/AtmosphereTextures.cpp
@@ -19,18 +19,20 @@ void FAtmosphereTextures::InitDynamicRHI()
 		FPooledRenderTargetDesc TransmittanceDesc(FPooledRenderTargetDesc::Create2DDesc(GTransmittanceTexSize, PF_FloatRGBA, TexCreate_None, TexCreate_RenderTargetable, false));
 		GRenderTargetPool.FindFreeElement(TransmittanceDesc, AtmosphereTransmittance, TEXT("AtmosphereTransmittance"));
 
-		RHISetRenderTarget(AtmosphereTransmittance->GetRenderTargetItem().TargetableTexture, FTextureRHIRef());
-		RHIClear(true, FLinearColor(0, 0, 0, 0), false, 0, false, 0, FIntRect());
-		RHICopyToResolveTarget(AtmosphereTransmittance->GetRenderTargetItem().TargetableTexture, AtmosphereTransmittance->GetRenderTargetItem().ShaderResourceTexture, true, FResolveParams());
+		FRHICommandListImmediate& RHICmdList = FRHICommandListExecutor::GetImmediateCommandList();
+
+		SetRenderTarget(RHICmdList, AtmosphereTransmittance->GetRenderTargetItem().TargetableTexture, FTextureRHIRef());
+		RHICmdList.Clear(true, FLinearColor(0, 0, 0, 0), false, 0, false, 0, FIntRect());
+		RHICmdList.CopyToResolveTarget(AtmosphereTransmittance->GetRenderTargetItem().TargetableTexture, AtmosphereTransmittance->GetRenderTargetItem().ShaderResourceTexture, true, FResolveParams());
 
 		// Irradiance
 		FIntPoint GIrradianceTexSize(PrecomputeParams->IrradianceTexWidth, PrecomputeParams->IrradianceTexHeight);
 		FPooledRenderTargetDesc IrradianceDesc(FPooledRenderTargetDesc::Create2DDesc(GIrradianceTexSize, PF_FloatRGBA, TexCreate_None, TexCreate_RenderTargetable, false));
 		GRenderTargetPool.FindFreeElement(IrradianceDesc, AtmosphereIrradiance, TEXT("AtmosphereIrradiance"));
 
-		RHISetRenderTarget(AtmosphereIrradiance->GetRenderTargetItem().TargetableTexture, FTextureRHIRef());
-		RHIClear(true, FLinearColor(0, 0, 0, 0), false, 0, false, 0, FIntRect());
-		RHICopyToResolveTarget(AtmosphereIrradiance->GetRenderTargetItem().TargetableTexture, AtmosphereIrradiance->GetRenderTargetItem().ShaderResourceTexture, true, FResolveParams());
+		SetRenderTarget(RHICmdList, AtmosphereIrradiance->GetRenderTargetItem().TargetableTexture, FTextureRHIRef());
+		RHICmdList.Clear(true, FLinearColor(0, 0, 0, 0), false, 0, false, 0, FIntRect());
+		RHICmdList.CopyToResolveTarget(AtmosphereIrradiance->GetRenderTargetItem().TargetableTexture, AtmosphereIrradiance->GetRenderTargetItem().ShaderResourceTexture, true, FResolveParams());
 
 		// DeltaE
 		GRenderTargetPool.FindFreeElement(IrradianceDesc, AtmosphereDeltaE, TEXT("AtmosphereDeltaE"));
diff --git a/Engine/Source/Runtime/Renderer/Private/AtmosphereTextures.h b/Engine/Source/Runtime/Renderer/Private/AtmosphereTextures.h
--- a/Engine/Source/Runtime/Renderer/Private/AtmosphereTextures.h
+++ b/Engine/Source/Runtime/Renderer/Private/AtmosphereTextures.h
@@ -29,8 +29,8 @@ public:
 	const struct FAtmospherePrecomputeParameters* PrecomputeParams;
 
 	// FRenderResource interface.
-	virtual void InitDynamicRHI();
-	virtual void ReleaseDynamicRHI();
+	virtual void InitDynamicRHI() override;
+	virtual void ReleaseDynamicRHI() override;
 
 	// Final result
 	TRefCountPtr<IPooledRenderTarget> AtmosphereTransmittance;
diff --git a/Engine/Source/Runtime/Renderer/Private/BasePassRendering.cpp b/Engine/Source/Runtime/Renderer/Private/BasePassRendering.cpp
--- a/Engine/Source/Runtime/Renderer/Private/BasePassRendering.cpp
+++ b/Engine/Source/Runtime/Renderer/Private/BasePassRendering.cpp
@@ -160,7 +160,7 @@ public:
 	/** Draws the translucent mesh with a specific light-map type, and fog volume type */
 	template<typename LightMapPolicyType>
 	void Process(
-		FRHICommandList& RHICmdList, 
+		FRHICommandList& RHICmdList,
 		const FProcessBasePassMeshParameters& Parameters,
 		const LightMapPolicyType& LightMapPolicy,
 		const typename LightMapPolicyType::ElementDataType& LightMapElementData
@@ -277,7 +277,7 @@ public:
 		if(View.Family->EngineShowFlags.ShaderComplexity)
 		{
 			// Note, this is a reversed Z depth surface, using CF_GreaterEqual.
-			RHISetDepthStencilState(TStaticDepthStencilState<false,CF_GreaterEqual>::GetRHI());
+			RHICmdList.SetDepthStencilState(TStaticDepthStencilState<false,CF_GreaterEqual>::GetRHI());
 		}
 #endif
 		const FScene* Scene = Parameters.PrimitiveSceneProxy ? Parameters.PrimitiveSceneProxy->GetPrimitiveSceneInfo()->Scene : NULL;
@@ -319,7 +319,7 @@ public:
 		if(View.Family->EngineShowFlags.ShaderComplexity)
 		{
 			// Note, this is a reversed Z depth surface, using CF_GreaterEqual.
-			RHISetDepthStencilState(TStaticDepthStencilState<true,CF_GreaterEqual>::GetRHI());
+			RHICmdList.SetDepthStencilState(TStaticDepthStencilState<true,CF_GreaterEqual>::GetRHI());
 		}
 #endif
 	}
diff --git a/Engine/Source/Runtime/Renderer/Private/BasePassRendering.h b/Engine/Source/Runtime/Renderer/Private/BasePassRendering.h
--- a/Engine/Source/Runtime/Renderer/Private/BasePassRendering.h
+++ b/Engine/Source/Runtime/Renderer/Private/BasePassRendering.h
@@ -683,8 +683,7 @@ public:
 			if (BlendMode != BLEND_Opaque)
 			{
 				// Add complexity to existing, keep alpha
-				RHICmdList.CheckIsNull();
-				RHISetBlendState(TStaticBlendState<CW_RGB,BO_Add,BF_One,BF_One>::GetRHI());
+				RHICmdList.SetBlendState(TStaticBlendState<CW_RGB,BO_Add,BF_One,BF_One>::GetRHI());
 			}
 
 			TShaderMapRef<FShaderComplexityAccumulatePS> ShaderComplexityPixelShader(GetGlobalShaderMap());
@@ -845,7 +844,7 @@ public:
 /** Processes a base pass mesh using an unknown light map policy, and unknown fog density policy. */
 template<typename ProcessActionType>
 void ProcessBasePassMesh(
-	FRHICommandList& RHICmdList, 
+	FRHICommandList& RHICmdList,
 	const FProcessBasePassMeshParameters& Parameters,
 	const ProcessActionType& Action
 	)
@@ -897,19 +896,19 @@ void ProcessBasePassMesh(
 						if (ShadowMapInteraction.GetType() == SMIT_Texture)
 						{
 							Action.template Process< TDistanceFieldShadowsAndLightMapPolicy<HQ_LIGHTMAP> >(
-								RHICmdList, 
+								RHICmdList,
 								Parameters,
 								TDistanceFieldShadowsAndLightMapPolicy<HQ_LIGHTMAP>(),
 								TDistanceFieldShadowsAndLightMapPolicy<HQ_LIGHTMAP>::ElementDataType(ShadowMapInteraction, LightMapInteraction));
 						}
 						else
 						{
-							Action.template Process< TLightMapPolicy<HQ_LIGHTMAP> >(RHICmdList, Parameters, TLightMapPolicy<HQ_LIGHTMAP>(), LightMapInteraction );
+							Action.template Process< TLightMapPolicy<HQ_LIGHTMAP> >(RHICmdList, Parameters, TLightMapPolicy<HQ_LIGHTMAP>(), LightMapInteraction);
 						}
 					} 
 					else 
 					{ 
-						Action.template Process< TLightMapPolicy<LQ_LIGHTMAP> >(RHICmdList, Parameters, TLightMapPolicy<LQ_LIGHTMAP>(), LightMapInteraction );
+						Action.template Process< TLightMapPolicy<LQ_LIGHTMAP> >(RHICmdList, Parameters, TLightMapPolicy<LQ_LIGHTMAP>(), LightMapInteraction);
 					} 
 					break;
 				default:
@@ -934,22 +933,22 @@ void ProcessBasePassMesh(
 								if (CanIndirectLightingCacheUseVolumeTexture(Parameters.FeatureLevel) && Action.AllowIndirectLightingCacheVolumeTexture())
 								{
 									// Use a lightmap policy that supports reading indirect lighting from a volume texture for dynamic objects
-									Action.template Process<FCachedVolumeIndirectLightingPolicy>(RHICmdList, Parameters,FCachedVolumeIndirectLightingPolicy(),FCachedVolumeIndirectLightingPolicy::ElementDataType());
+									Action.template Process<FCachedVolumeIndirectLightingPolicy>(RHICmdList, Parameters, FCachedVolumeIndirectLightingPolicy(), FCachedVolumeIndirectLightingPolicy::ElementDataType());
 								}
 								else
 								{
 									// Use a lightmap policy that supports reading indirect lighting from a single SH sample
-									Action.template Process<FCachedPointIndirectLightingPolicy>(RHICmdList, Parameters,FCachedPointIndirectLightingPolicy(),FCachedPointIndirectLightingPolicy::ElementDataType(false));
+									Action.template Process<FCachedPointIndirectLightingPolicy>(RHICmdList, Parameters, FCachedPointIndirectLightingPolicy(), FCachedPointIndirectLightingPolicy::ElementDataType(false));
 								}
 							}
 							else
 							{
-								Action.template Process<FNoLightMapPolicy>(RHICmdList, Parameters,FNoLightMapPolicy(),FNoLightMapPolicy::ElementDataType());
+								Action.template Process<FNoLightMapPolicy>(RHICmdList, Parameters, FNoLightMapPolicy(), FNoLightMapPolicy::ElementDataType());
 							}
 						}
 						else
 						{
-							Action.template Process<FNoLightMapPolicy>(RHICmdList, Parameters,FNoLightMapPolicy(),FNoLightMapPolicy::ElementDataType());
+							Action.template Process<FNoLightMapPolicy>(RHICmdList, Parameters, FNoLightMapPolicy(), FNoLightMapPolicy::ElementDataType());
 						}
 					}
 					break;
diff --git a/Engine/Source/Runtime/Renderer/Private/ClearQuad.cpp b/Engine/Source/Runtime/Renderer/Private/ClearQuad.cpp
--- a/Engine/Source/Runtime/Renderer/Private/ClearQuad.cpp
+++ b/Engine/Source/Runtime/Renderer/Private/ClearQuad.cpp
@@ -6,7 +6,7 @@
 FGlobalBoundShaderState GClearMRTBoundShaderState[8];
 
 // TODO support ExcludeRect
-void DrawClearQuadMRT(FRHICommandList& RHICmdList, bool bClearColor, int32 NumClearColors, const FLinearColor* ClearColorArray, bool bClearDepth, float Depth, bool bClearStencil, uint32 Stencil)
+void DrawClearQuadMRT(FRHICommandListImmediate& RHICmdList, bool bClearColor, int32 NumClearColors, const FLinearColor* ClearColorArray, bool bClearDepth, float Depth, bool bClearStencil, uint32 Stencil)
 {
 	// Set new states
 	FBlendStateRHIParamRef BlendStateRHI;
diff --git a/Engine/Source/Runtime/Renderer/Private/ClearQuad.h b/Engine/Source/Runtime/Renderer/Private/ClearQuad.h
--- a/Engine/Source/Runtime/Renderer/Private/ClearQuad.h
+++ b/Engine/Source/Runtime/Renderer/Private/ClearQuad.h
@@ -3,9 +3,9 @@
 #pragma once
 
 // TODO support ExcludeRect
-void DrawClearQuadMRT(FRHICommandList& RHICmdList, bool bClearColor, int32 NumClearColors, const FLinearColor* ClearColorArray, bool bClearDepth, float Depth, bool bClearStencil, uint32 Stencil);
+void DrawClearQuadMRT(FRHICommandListImmediate& RHICmdList, bool bClearColor, int32 NumClearColors, const FLinearColor* ClearColorArray, bool bClearDepth, float Depth, bool bClearStencil, uint32 Stencil);
 
-inline void DrawClearQuad(FRHICommandList& RHICmdList, bool bClearColor, const FLinearColor& Color, bool bClearDepth, float Depth, bool bClearStencil, uint32 Stencil)
+inline void DrawClearQuad(FRHICommandListImmediate& RHICmdList, bool bClearColor, const FLinearColor& Color, bool bClearDepth, float Depth, bool bClearStencil, uint32 Stencil)
 {
 	DrawClearQuadMRT(RHICmdList, bClearColor, 1, &Color, bClearDepth, Depth, bClearStencil, Stencil);
 }
\ No newline at end of file
diff --git a/Engine/Source/Runtime/Renderer/Private/CompositionLighting/CompositionLighting.cpp b/Engine/Source/Runtime/Renderer/Private/CompositionLighting/CompositionLighting.cpp
--- a/Engine/Source/Runtime/Renderer/Private/CompositionLighting/CompositionLighting.cpp
+++ b/Engine/Source/Runtime/Renderer/Private/CompositionLighting/CompositionLighting.cpp
@@ -209,14 +209,14 @@ static void AddDeferredDecalsBeforeLighting(FPostprocessContext& Context)
 	Context.FinalOutput = FRenderingCompositeOutputRef(Pass);
 }
 
-void FCompositionLighting::ProcessBeforeBasePass(const FViewInfo& View)
+void FCompositionLighting::ProcessBeforeBasePass(FRHICommandListImmediate& RHICmdList, const FViewInfo& View)
 {
 	check(IsInRenderingThread());
 
 	// so that the passes can register themselves to the graph
 	{
 		FMemMark Mark(FMemStack::Get());
-		FRenderingCompositePassContext CompositeContext(View);
+		FRenderingCompositePassContext CompositeContext(RHICmdList, View);
 
 		FPostprocessContext Context(CompositeContext.Graph, View);
 
@@ -244,12 +244,10 @@ void FCompositionLighting::ProcessBeforeBasePass(const FViewInfo& View)
 	}
 }
 
-void FCompositionLighting::ProcessAfterBasePass(const FViewInfo& View)
+void FCompositionLighting::ProcessAfterBasePass(FRHICommandListImmediate& RHICmdList, const FViewInfo& View)
 {
 	check(IsInRenderingThread());
-	//@todo-rco: RHIPacketList
-	FRHICommandList& RHICmdList = FRHICommandList::GetNullRef();
-
+	
 	// might get renamed to refracted or ...WithAO
 	GSceneRenderTargets.GetSceneColor()->SetDebugName(TEXT("SceneColor"));
 	// to be able to observe results with VisualizeTexture
@@ -265,7 +263,7 @@ void FCompositionLighting::ProcessAfterBasePass(const FViewInfo& View)
 	// so that the passes can register themselves to the graph
 	{
 		FMemMark Mark(FMemStack::Get());
-		FRenderingCompositePassContext CompositeContext(View);
+		FRenderingCompositePassContext CompositeContext(RHICmdList, View);
 
 		FPostprocessContext Context(CompositeContext.Graph, View);
 
@@ -307,19 +305,17 @@ void FCompositionLighting::ProcessAfterBasePass(const FViewInfo& View)
 }
 
 
-void FCompositionLighting::ProcessLighting(const FViewInfo& View)
+void FCompositionLighting::ProcessLighting(FRHICommandListImmediate& RHICmdList, const FViewInfo& View)
 {
 	check(IsInRenderingThread());
-	//@todo-rco: RHIPacketList
-	FRHICommandList& RHICmdList = FRHICommandList::GetNullRef();
-
+	
 	GRenderTargetPool.VisualizeTexture.SetCheckPoint(RHICmdList, GSceneRenderTargets.ReflectiveShadowMapDiffuse);
 	GRenderTargetPool.VisualizeTexture.SetCheckPoint(RHICmdList, GSceneRenderTargets.ReflectiveShadowMapNormal);
 	GRenderTargetPool.VisualizeTexture.SetCheckPoint(RHICmdList, GSceneRenderTargets.ReflectiveShadowMapDepth);
 
 	{
 		FMemMark Mark(FMemStack::Get());
-		FRenderingCompositePassContext CompositeContext(View);
+		FRenderingCompositePassContext CompositeContext(RHICmdList, View);
 
 		FPostprocessContext Context(CompositeContext.Graph, View);
 
diff --git a/Engine/Source/Runtime/Renderer/Private/CompositionLighting/CompositionLighting.h b/Engine/Source/Runtime/Renderer/Private/CompositionLighting/CompositionLighting.h
--- a/Engine/Source/Runtime/Renderer/Private/CompositionLighting/CompositionLighting.h
+++ b/Engine/Source/Runtime/Renderer/Private/CompositionLighting/CompositionLighting.h
@@ -14,11 +14,11 @@
 class FCompositionLighting
 {
 public:
-	void ProcessBeforeBasePass(const FViewInfo& View);
+	void ProcessBeforeBasePass(FRHICommandListImmediate& RHICmdList, const FViewInfo& View);
 
-	void ProcessAfterBasePass(const FViewInfo& View);
+	void ProcessAfterBasePass(FRHICommandListImmediate& RHICmdList, const FViewInfo& View);
 
-	void ProcessLighting(const FViewInfo& View);
+	void ProcessLighting(FRHICommandListImmediate& RHICmdList, const FViewInfo& View);
 };
 
 /** The global used for deferred lighting. */
diff --git a/Engine/Source/Runtime/Renderer/Private/CompositionLighting/PostProcessAmbient.cpp b/Engine/Source/Runtime/Renderer/Private/CompositionLighting/PostProcessAmbient.cpp
--- a/Engine/Source/Runtime/Renderer/Private/CompositionLighting/PostProcessAmbient.cpp
+++ b/Engine/Source/Runtime/Renderer/Private/CompositionLighting/PostProcessAmbient.cpp
@@ -53,7 +53,7 @@ public:
 		FGlobalShader::SetParameters(RHICmdList, ShaderRHI, Context.View);
 		PostprocessParameter.SetPS(ShaderRHI, Context, TStaticSamplerState<SF_Bilinear,AM_Clamp,AM_Clamp,AM_Clamp>::GetRHI());
 		DeferredParameters.Set(RHICmdList, ShaderRHI, Context.View);
-		CubemapShaderParameters.SetParameters(ShaderRHI, Entry);
+		CubemapShaderParameters.SetParameters(RHICmdList, ShaderRHI, Entry);
 		SetTextureParameter(RHICmdList, ShaderRHI, PreIntegratedGF, PreIntegratedGFSampler, TStaticSamplerState<SF_Bilinear,AM_Clamp,AM_Clamp,AM_Clamp>::GetRHI(), GSystemTextures.PreintegratedGF->GetRenderTargetItem().ShaderResourceTexture );
 	}
 	
@@ -102,7 +102,7 @@ void FRCPassPostProcessAmbient::Process(FRenderingCompositePassContext& Context)
 	TShaderMapRef<FPostProcessAmbientPS> PixelShader(GetGlobalShaderMap());
 
 	static FGlobalBoundShaderState BoundShaderState;
-	Context.RHICmdList.CheckIsNull(); // need new approach for "static FGlobalBoundShaderState" for parallel rendering
+	
 
 	uint32 Count = Context.View.FinalPostProcessSettings.ContributingCubemaps.Num();
 	for(uint32 i = 0; i < Count; ++i)
@@ -113,7 +113,7 @@ void FRCPassPostProcessAmbient::Process(FRenderingCompositePassContext& Context)
 			SetGlobalBoundShaderState(Context.RHICmdList, BoundShaderState, GFilterVertexDeclaration.VertexDeclarationRHI, *VertexShader, *PixelShader);
 		}
 
-		//@todo-rco: RHIPacketList
+		
 		PixelShader->SetParameters(Context.RHICmdList, Context, Context.View.FinalPostProcessSettings.ContributingCubemaps[i]);
 
 		// Draw a quad mapping scene color to the view's render target
@@ -154,22 +154,19 @@ void FCubemapShaderParameters::Bind(const FShaderParameterMap& ParameterMap)
 	AmbientCubemapSampler.Bind(ParameterMap, TEXT("AmbientCubemapSampler"));
 }
 
-void FCubemapShaderParameters::SetParameters(const FPixelShaderRHIParamRef ShaderRHI, const FFinalPostProcessSettings::FCubemapEntry& Entry) const
+void FCubemapShaderParameters::SetParameters(FRHICommandList& RHICmdList, const FPixelShaderRHIParamRef ShaderRHI, const FFinalPostProcessSettings::FCubemapEntry& Entry) const
 {
-	SetParametersTemplate(ShaderRHI,Entry);
+	SetParametersTemplate(RHICmdList, ShaderRHI, Entry);
 }
 
-void FCubemapShaderParameters::SetParameters(const FComputeShaderRHIParamRef ShaderRHI, const FFinalPostProcessSettings::FCubemapEntry& Entry) const
+void FCubemapShaderParameters::SetParameters(FRHICommandList& RHICmdList, const FComputeShaderRHIParamRef ShaderRHI, const FFinalPostProcessSettings::FCubemapEntry& Entry) const
 {
-	SetParametersTemplate(ShaderRHI,Entry);
+	SetParametersTemplate(RHICmdList, ShaderRHI, Entry);
 }
 
 template<typename TShaderRHIRef>
-void FCubemapShaderParameters::SetParametersTemplate(const TShaderRHIRef& ShaderRHI, const FFinalPostProcessSettings::FCubemapEntry& Entry) const
+void FCubemapShaderParameters::SetParametersTemplate(FRHICommandList& RHICmdList, const TShaderRHIRef& ShaderRHI, const FFinalPostProcessSettings::FCubemapEntry& Entry) const
 {
-	//@todo-rco: RHIPacketList
-	FRHICommandList& RHICmdList = FRHICommandList::GetNullRef();
-
 	// floats to render the cubemap
 	{
 		float MipCount = 0;
diff --git a/Engine/Source/Runtime/Renderer/Private/CompositionLighting/PostProcessAmbientOcclusion.cpp b/Engine/Source/Runtime/Renderer/Private/CompositionLighting/PostProcessAmbientOcclusion.cpp
--- a/Engine/Source/Runtime/Renderer/Private/CompositionLighting/PostProcessAmbientOcclusion.cpp
+++ b/Engine/Source/Runtime/Renderer/Private/CompositionLighting/PostProcessAmbientOcclusion.cpp
@@ -350,7 +350,7 @@ void FRCPassPostProcessAmbientOcclusionSetup::SetShaderSetupTempl(const FRenderi
 	TShaderMapRef<FPostProcessAmbientOcclusionSetupPS<bInitialSetup> > PixelShader(GetGlobalShaderMap());
 
 	static FGlobalBoundShaderState BoundShaderState;
-	Context.RHICmdList.CheckIsNull(); // need new approach for "static FGlobalBoundShaderState" for parallel rendering
+	
 
 	SetGlobalBoundShaderState(Context.RHICmdList, BoundShaderState, GFilterVertexDeclaration.VertexDeclarationRHI, *VertexShader, *PixelShader);
 
@@ -461,7 +461,7 @@ void FRCPassPostProcessAmbientOcclusion::SetShaderTempl(const FRenderingComposit
 	TShaderMapRef<FPostProcessAmbientOcclusionPS<bTAOSetupAsInput, bDoUpsample, SampleSetQuality> > PixelShader(GetGlobalShaderMap());
 
 	static FGlobalBoundShaderState BoundShaderState;
-	Context.RHICmdList.CheckIsNull(); // need new approach for "static FGlobalBoundShaderState" for parallel rendering
+	
 
 	SetGlobalBoundShaderState(Context.RHICmdList, BoundShaderState, GFilterVertexDeclaration.VertexDeclarationRHI, *VertexShader, *PixelShader);
 
@@ -609,7 +609,7 @@ void FRCPassPostProcessBasePassAO::Process(FRenderingCompositePassContext& Conte
 	TShaderMapRef<FPostProcessBasePassAOPS> PixelShader(GetGlobalShaderMap());
 
 	static FGlobalBoundShaderState BoundShaderState;
-	Context.RHICmdList.CheckIsNull(); // need new approach for "static FGlobalBoundShaderState" for parallel rendering
+	
 
 	SetGlobalBoundShaderState(Context.RHICmdList, BoundShaderState, GFilterVertexDeclaration.VertexDeclarationRHI, *VertexShader, *PixelShader);
 
diff --git a/Engine/Source/Runtime/Renderer/Private/CompositionLighting/PostProcessDeferredDecals.cpp b/Engine/Source/Runtime/Renderer/Private/CompositionLighting/PostProcessDeferredDecals.cpp
--- a/Engine/Source/Runtime/Renderer/Private/CompositionLighting/PostProcessDeferredDecals.cpp
+++ b/Engine/Source/Runtime/Renderer/Private/CompositionLighting/PostProcessDeferredDecals.cpp
@@ -136,7 +136,7 @@ struct FTransientDecalRenderData
 	}
 };
 // @param RenderState 0:before BasePass, 1:before lighting, (later we could add "after lighting" and multiply)
-void SetDecalBlendState(const ERHIFeatureLevel::Type SMFeatureLevel, uint32 RenderStage, EDecalBlendMode DecalBlendMode, bool bHasNormal)
+void SetDecalBlendState(FRHICommandList& RHICmdList, const ERHIFeatureLevel::Type SMFeatureLevel, uint32 RenderStage, EDecalBlendMode DecalBlendMode, bool bHasNormal)
 {
 	if(RenderStage == 0)
 	{
@@ -144,7 +144,7 @@ void SetDecalBlendState(const ERHIFeatureLevel::Type SMFeatureLevel, uint32 Rend
 		switch(DecalBlendMode)
 		{
 		case DBM_DBuffer_ColorNormalRoughness:
-			RHISetBlendState( TStaticBlendState< 
+			RHICmdList.SetBlendState( TStaticBlendState< 
 				CW_RGBA, BO_Add, BF_SourceAlpha, BF_InverseSourceAlpha,			BO_Add,BF_Zero,BF_InverseSourceAlpha,
 				CW_RGBA, BO_Add, BF_SourceAlpha, BF_InverseSourceAlpha,			BO_Add,BF_Zero,BF_InverseSourceAlpha,
 				CW_RGBA, BO_Add, BF_SourceAlpha, BF_InverseSourceAlpha,			BO_Add,BF_Zero,BF_InverseSourceAlpha >::GetRHI() );		
@@ -152,7 +152,7 @@ void SetDecalBlendState(const ERHIFeatureLevel::Type SMFeatureLevel, uint32 Rend
 
 		case DBM_DBuffer_Color:
 			// we can optimize using less MRT later
-			RHISetBlendState( TStaticBlendState< 
+			RHICmdList.SetBlendState( TStaticBlendState< 
 				CW_RGBA, BO_Add, BF_SourceAlpha, BF_InverseSourceAlpha,			BO_Add,BF_Zero,BF_InverseSourceAlpha,
 				CW_RGBA, BO_Add, BF_Zero, BF_One,								BO_Add,BF_Zero,BF_One,
 				CW_RGBA, BO_Add, BF_Zero, BF_One,								BO_Add,BF_Zero,BF_One>::GetRHI() );		
@@ -160,7 +160,7 @@ void SetDecalBlendState(const ERHIFeatureLevel::Type SMFeatureLevel, uint32 Rend
 
 		case DBM_DBuffer_ColorNormal:
 			// we can optimize using less MRT later
-			RHISetBlendState( TStaticBlendState< 
+			RHICmdList.SetBlendState( TStaticBlendState< 
 				CW_RGBA, BO_Add, BF_SourceAlpha, BF_InverseSourceAlpha,			BO_Add,BF_Zero,BF_InverseSourceAlpha,
 				CW_RGBA, BO_Add, BF_SourceAlpha, BF_InverseSourceAlpha,			BO_Add,BF_Zero,BF_InverseSourceAlpha,
 				CW_RGBA, BO_Add, BF_Zero, BF_One,								BO_Add,BF_Zero,BF_One >::GetRHI() );		
@@ -168,7 +168,7 @@ void SetDecalBlendState(const ERHIFeatureLevel::Type SMFeatureLevel, uint32 Rend
 
 		case DBM_DBuffer_ColorRoughness:
 			// we can optimize using less MRT later
-			RHISetBlendState( TStaticBlendState< 
+			RHICmdList.SetBlendState( TStaticBlendState< 
 				CW_RGBA, BO_Add, BF_SourceAlpha, BF_InverseSourceAlpha,			BO_Add,BF_Zero,BF_InverseSourceAlpha,
 				CW_RGBA, BO_Add, BF_Zero, BF_One,								BO_Add,BF_Zero,BF_One,
 				CW_RGBA, BO_Add, BF_SourceAlpha, BF_InverseSourceAlpha,			BO_Add,BF_Zero,BF_InverseSourceAlpha >::GetRHI() );		
@@ -176,7 +176,7 @@ void SetDecalBlendState(const ERHIFeatureLevel::Type SMFeatureLevel, uint32 Rend
 
 		case DBM_DBuffer_Normal:
 			// we can optimize using less MRT later
-			RHISetBlendState( TStaticBlendState< 
+			RHICmdList.SetBlendState( TStaticBlendState< 
 				CW_RGBA, BO_Add, BF_Zero, BF_One,								BO_Add,BF_Zero,BF_One,
 				CW_RGBA, BO_Add, BF_SourceAlpha, BF_InverseSourceAlpha,			BO_Add,BF_Zero,BF_InverseSourceAlpha,
 				CW_RGBA, BO_Add, BF_Zero, BF_One,								BO_Add,BF_Zero,BF_One>::GetRHI() );		
@@ -184,7 +184,7 @@ void SetDecalBlendState(const ERHIFeatureLevel::Type SMFeatureLevel, uint32 Rend
 
 		case DBM_DBuffer_NormalRoughness:
 			// we can optimize using less MRT later
-			RHISetBlendState( TStaticBlendState< 
+			RHICmdList.SetBlendState( TStaticBlendState< 
 				CW_RGBA, BO_Add, BF_Zero, BF_One,								BO_Add,BF_Zero,BF_One,
 				CW_RGBA, BO_Add, BF_SourceAlpha, BF_InverseSourceAlpha,			BO_Add,BF_Zero,BF_InverseSourceAlpha,
 				CW_RGBA, BO_Add, BF_SourceAlpha, BF_InverseSourceAlpha,			BO_Add,BF_Zero,BF_InverseSourceAlpha >::GetRHI() );		
@@ -192,7 +192,7 @@ void SetDecalBlendState(const ERHIFeatureLevel::Type SMFeatureLevel, uint32 Rend
 
 		case DBM_DBuffer_Roughness:
 			// we can optimize using less MRT later
-			RHISetBlendState( TStaticBlendState< 
+			RHICmdList.SetBlendState( TStaticBlendState< 
 				CW_RGBA, BO_Add, BF_Zero, BF_One,								BO_Add,BF_Zero,BF_One,
 				CW_RGBA, BO_Add, BF_Zero, BF_One,								BO_Add,BF_Zero,BF_One,
 				CW_RGBA, BO_Add, BF_SourceAlpha, BF_InverseSourceAlpha,			BO_Add,BF_Zero,BF_InverseSourceAlpha >::GetRHI() );		
@@ -217,7 +217,7 @@ void SetDecalBlendState(const ERHIFeatureLevel::Type SMFeatureLevel, uint32 Rend
 			{
 				if(bHasNormal)
 				{
-					RHISetBlendState( TStaticBlendState<
+					RHICmdList.SetBlendState( TStaticBlendState<
 						CW_RGB, BO_Add, BF_SourceAlpha, BF_One,						BO_Add, BF_Zero, BF_One,	// Emissive
 						CW_RGB, BO_Add, BF_SourceAlpha, BF_InverseSourceAlpha,		BO_Add, BF_Zero, BF_One,	// Normal
 						CW_RGB, BO_Add, BF_SourceAlpha, BF_InverseSourceAlpha,		BO_Add, BF_Zero, BF_One,	// Metallic, Specular
@@ -227,7 +227,7 @@ void SetDecalBlendState(const ERHIFeatureLevel::Type SMFeatureLevel, uint32 Rend
 				}
 				else
 				{
-					RHISetBlendState( TStaticBlendState<
+					RHICmdList.SetBlendState( TStaticBlendState<
 						CW_RGB, BO_Add, BF_SourceAlpha, BF_One,						BO_Add, BF_Zero, BF_One,	// Emissive
 						CW_RGB, BO_Add, BF_Zero,		BF_One,						BO_Add, BF_Zero, BF_One,	// Normal
 						CW_RGB, BO_Add, BF_SourceAlpha, BF_InverseSourceAlpha,		BO_Add, BF_Zero, BF_One,	// Metallic, Specular
@@ -238,7 +238,7 @@ void SetDecalBlendState(const ERHIFeatureLevel::Type SMFeatureLevel, uint32 Rend
 			}
 			else if(SMFeatureLevel == ERHIFeatureLevel::SM4)
 			{
-				RHISetBlendState( TStaticBlendState<
+				RHICmdList.SetBlendState( TStaticBlendState<
 					CW_RGB, BO_Add, BF_SourceAlpha, BF_One,							BO_Add, BF_Zero, BF_One,	// Emissive
 					CW_RGB, BO_Add, BF_SourceAlpha, BF_One,							BO_Add, BF_Zero, BF_One,	// Normal
 					CW_RGB, BO_Add, BF_SourceAlpha, BF_One,							BO_Add, BF_Zero, BF_One,	// Metallic, Specular
@@ -253,7 +253,7 @@ void SetDecalBlendState(const ERHIFeatureLevel::Type SMFeatureLevel, uint32 Rend
 			{
 				if(bHasNormal)
 				{
-					RHISetBlendState( TStaticBlendState<
+					RHICmdList.SetBlendState( TStaticBlendState<
 						CW_RGB, BO_Add, BF_SourceAlpha, BF_One,						BO_Add, BF_Zero, BF_One,	// Emissive
 						CW_RGB, BO_Add, BF_SourceAlpha, BF_InverseSourceAlpha,		BO_Add, BF_Zero, BF_One,	// Normal
 						CW_RGB, BO_Add, BF_SourceAlpha, BF_InverseSourceAlpha,		BO_Add, BF_Zero, BF_One,	// Metallic, Specular
@@ -263,7 +263,7 @@ void SetDecalBlendState(const ERHIFeatureLevel::Type SMFeatureLevel, uint32 Rend
 				}
 				else
 				{
-					RHISetBlendState( TStaticBlendState<
+					RHICmdList.SetBlendState( TStaticBlendState<
 						CW_RGB, BO_Add, BF_SourceAlpha, BF_One,						BO_Add, BF_Zero, BF_One,	// Emissive
 						CW_RGB, BO_Add, BF_Zero,		BF_One,						BO_Add, BF_Zero, BF_One,	// Normal
 						CW_RGB, BO_Add, BF_SourceAlpha, BF_InverseSourceAlpha,		BO_Add, BF_Zero, BF_One,	// Metallic, Specular
@@ -274,7 +274,7 @@ void SetDecalBlendState(const ERHIFeatureLevel::Type SMFeatureLevel, uint32 Rend
 			}
 			else if(SMFeatureLevel == ERHIFeatureLevel::SM4)
 			{
-				RHISetBlendState( TStaticBlendState<
+				RHICmdList.SetBlendState( TStaticBlendState<
 					CW_RGB, BO_Add, BF_SourceAlpha, BF_InverseSourceAlpha,			BO_Add, BF_Zero, BF_One,	// Emissive
 					CW_RGB, BO_Add, BF_SourceAlpha, BF_InverseSourceAlpha,			BO_Add, BF_Zero, BF_One,	// Normal
 					CW_RGB, BO_Add, BF_SourceAlpha, BF_InverseSourceAlpha,			BO_Add, BF_Zero, BF_One,	// Metallic, Specular
@@ -285,11 +285,11 @@ void SetDecalBlendState(const ERHIFeatureLevel::Type SMFeatureLevel, uint32 Rend
 			break;
 
 		case DBM_Normal:
-			RHISetBlendState( TStaticBlendState< CW_RGB, BO_Add, BF_SourceAlpha, BF_InverseSourceAlpha >::GetRHI() );
+			RHICmdList.SetBlendState( TStaticBlendState< CW_RGB, BO_Add, BF_SourceAlpha, BF_InverseSourceAlpha >::GetRHI() );
 			break;
 
 		case DBM_Emissive:
-			RHISetBlendState( TStaticBlendState< CW_RGB, BO_Add, BF_SourceAlpha, BF_One >::GetRHI() );
+			RHICmdList.SetBlendState( TStaticBlendState< CW_RGB, BO_Add, BF_SourceAlpha, BF_One >::GetRHI() );
 			break;
 
 		default:
@@ -350,7 +350,7 @@ IMPLEMENT_SHADER_TYPE(,FStencilDecalMaskPS,TEXT("DeferredDecal"),TEXT("StencilDe
 FGlobalBoundShaderState StencilDecalMaskBoundShaderState;
 
 /** Draws a full view quad that sets stencil to 1 anywhere that decals should not be projected. */
-void StencilDecalMask(FRHICommandList& RHICmdList, const FSceneView& View)
+void StencilDecalMask(FRHICommandListImmediate& RHICmdList, const FSceneView& View)
 {
 	SCOPED_DRAW_EVENT(StencilDecalMask, DEC_SCENE_ITEMS);
 	RHICmdList.SetRasterizerState(TStaticRasterizerState<FM_Solid, CM_None>::GetRHI());
@@ -363,7 +363,6 @@ void StencilDecalMask(FRHICommandList& RHICmdList, const FSceneView& View)
 
 	TShaderMapRef<FScreenVS> ScreenVertexShader(GetGlobalShaderMap());
 	TShaderMapRef<FStencilDecalMaskPS> PixelShader(GetGlobalShaderMap());
-	RHICmdList.CheckIsNull(); // need new approach for "static FGlobalBoundShaderState" for parallel rendering
 	
 	SetGlobalBoundShaderState(RHICmdList, StencilDecalMaskBoundShaderState, GFilterVertexDeclaration.VertexDeclarationRHI, *ScreenVertexShader, *PixelShader);
 
@@ -603,15 +602,13 @@ private:
 static TGlobalResource<FUnitCubeVertexBuffer> GUnitCubeVertexBuffer;
 static TGlobalResource<FUnitCubeIndexBuffer> GUnitCubeIndexBuffer;
 
-void SetShader(FRHICommandList& RHICmdList, const FRenderingCompositePassContext& Context, const FTransientDecalRenderData& DecalData, FShader* VertexShader)
+void SetShader(const FRenderingCompositePassContext& Context, const FTransientDecalRenderData& DecalData, FShader* VertexShader)
 {
 	const FSceneView& View = Context.View;
 	
 	const FMaterialShaderMap* MaterialShaderMap = DecalData.MaterialResource->GetRenderingThreadShaderMap();
 	FDeferredDecalPS* PixelShader = MaterialShaderMap->GetShader<FDeferredDecalPS>();
 
-	RHICmdList.CheckIsNull();
-
 	// This was cached but when changing the material (e.g. editor) it wasn't updated.
 	// This will change with upcoming multi threaded rendering changes.
 	FBoundShaderStateRHIRef BoundShaderState;
@@ -620,15 +617,15 @@ void SetShader(FRHICommandList& RHICmdList, const FRenderingCompositePassContext
 		FMemory::Memzero(Strides, sizeof(Strides));
 		Strides[0] = sizeof(FVector);
 
-		BoundShaderState = RHICreateBoundShaderState(GetVertexDeclarationFVector3(), VertexShader->GetVertexShader(), FHullShaderRHIRef(), FDomainShaderRHIRef(), PixelShader->GetPixelShader(), FGeometryShaderRHIRef());
+		BoundShaderState = Context.RHICmdList.CreateBoundShaderState(GetVertexDeclarationFVector3(), VertexShader->GetVertexShader(), FHullShaderRHIRef(), FDomainShaderRHIRef(), PixelShader->GetPixelShader(), FGeometryShaderRHIRef());
 	}
 
-	RHISetBoundShaderState(BoundShaderState);
+	Context.RHICmdList.SetBoundShaderState(BoundShaderState);
 
-	PixelShader->SetParameters(RHICmdList, View, DecalData.MaterialProxy, *DecalData.DecalProxy);
+	PixelShader->SetParameters(Context.RHICmdList, View, DecalData.MaterialProxy, *DecalData.DecalProxy);
 }
 
-bool RenderPreStencil(FRHICommandList& RHICmdList, FRenderingCompositePassContext& Context, const FMaterialShaderMap* MaterialShaderMap, const FMatrix& ComponentToWorldMatrix, const FMatrix& FrustumComponentToClip)
+bool RenderPreStencil(FRenderingCompositePassContext& Context, const FMaterialShaderMap* MaterialShaderMap, const FMatrix& ComponentToWorldMatrix, const FMatrix& FrustumComponentToClip)
 {
 	SCOPED_DRAW_EVENT(RenderPreStencil, DEC_SCENE_ITEMS);
 
@@ -653,24 +650,22 @@ bool RenderPreStencil(FRHICommandList& RHICmdList, FRenderingCompositePassContex
 
 	FDeferredDecalVS* VertexShader = MaterialShaderMap->GetShader<FDeferredDecalVS>();
 	
-	RHICmdList.CheckIsNull();
-
 	// This was cached but when changing the material (e.g. editor) it wasn't updated.
 	// This will change with upcoming multi threaded rendering changes.
 	FBoundShaderStateRHIRef BoundShaderState;
 	{
-		BoundShaderState = RHICreateBoundShaderState(GetVertexDeclarationFVector3(), VertexShader->GetVertexShader(), FHullShaderRHIRef(), FDomainShaderRHIRef(), NULL, FGeometryShaderRHIRef());
+		BoundShaderState = Context.RHICmdList.CreateBoundShaderState(GetVertexDeclarationFVector3(), VertexShader->GetVertexShader(), FHullShaderRHIRef(), FDomainShaderRHIRef(), NULL, FGeometryShaderRHIRef());
 	}
 
-	RHISetBoundShaderState(BoundShaderState);
+	Context.RHICmdList.SetBoundShaderState(BoundShaderState);
 
-	VertexShader->SetParameters(RHICmdList, View, FrustumComponentToClip);
+	VertexShader->SetParameters(Context.RHICmdList, View, FrustumComponentToClip);
 
 	// Set states, the state cache helps us avoiding redundant sets
-	RHISetRasterizerState(TStaticRasterizerState<FM_Solid,CM_None>::GetRHI());
+	Context.RHICmdList.SetRasterizerState(TStaticRasterizerState<FM_Solid, CM_None>::GetRHI());
 
 	// all the same to have DX10 working
-	RHISetBlendState( TStaticBlendState<
+	Context.RHICmdList.SetBlendState(TStaticBlendState<
 		CW_NONE, BO_Add, BF_SourceAlpha, BF_InverseSourceAlpha,		BO_Add, BF_Zero, BF_One,	// Emissive
 		CW_NONE, BO_Add, BF_SourceAlpha, BF_InverseSourceAlpha,		BO_Add, BF_Zero, BF_One,	// Normal
 		CW_NONE, BO_Add, BF_SourceAlpha, BF_InverseSourceAlpha,		BO_Add, BF_Zero, BF_One,	// Metallic, Specular
@@ -679,7 +674,7 @@ bool RenderPreStencil(FRHICommandList& RHICmdList, FRenderingCompositePassContex
 	>::GetRHI() );
 
 	// Carmack's reverse on the bounds
-	RHISetDepthStencilState(TStaticDepthStencilState<
+	Context.RHICmdList.SetDepthStencilState(TStaticDepthStencilState<
 		false,CF_LessEqual,
 		true,CF_Equal,SO_Keep,SO_Keep,SO_Increment,
 		true,CF_Equal,SO_Keep,SO_Keep,SO_Decrement,
@@ -687,7 +682,7 @@ bool RenderPreStencil(FRHICommandList& RHICmdList, FRenderingCompositePassContex
 	>::GetRHI());
 
 	// Render decal mask
-	RHIDrawIndexedPrimitive(GUnitCubeIndexBuffer.IndexBufferRHI, PT_TriangleList, 0, 0, 8, 0, GUnitCubeIndexBuffer.GetIndexCount() / 3, 0);
+	Context.RHICmdList.DrawIndexedPrimitive(GUnitCubeIndexBuffer.IndexBufferRHI, PT_TriangleList, 0, 0, 8, 0, GUnitCubeIndexBuffer.GetIndexCount() / 3, 0);
 
 	return true;
 }
@@ -707,8 +702,8 @@ FRCPassPostProcessDeferredDecals::FRCPassPostProcessDeferredDecals(uint32 InRend
 
 void FRCPassPostProcessDeferredDecals::Process(FRenderingCompositePassContext& Context)
 {
-	//@todo-rco: RHIPacketList
-	FRHICommandList& RHICmdList = FRHICommandList::GetNullRef();
+	
+	FRHICommandListImmediate& RHICmdList = Context.RHICmdList;
 
 	bool bDBuffer = IsDBufferEnabled();
 	float bDecalPreStencil = CVarStencilSizeThreshold.GetValueOnRenderThread() >= 0;
@@ -890,7 +885,7 @@ void FRCPassPostProcessDeferredDecals::Process(FRenderingCompositePassContext& C
 		SCOPED_DRAW_EVENT(Decals, DEC_SCENE_ITEMS);
 		INC_DWORD_STAT_BY(STAT_Decals, SortedDecals.Num());
 		
-		RHISetStreamSource(0, GUnitCubeVertexBuffer.VertexBufferRHI, sizeof(FVector4), 0);
+		RHICmdList.SetStreamSource(0, GUnitCubeVertexBuffer.VertexBufferRHI, sizeof(FVector4), 0);
 
 		for (int32 DecalIndex = 0; DecalIndex < SortedDecals.Num(); DecalIndex++)
 		{
@@ -911,7 +906,7 @@ void FRCPassPostProcessDeferredDecals::Process(FRenderingCompositePassContext& C
 			{
 				if(bDecalPreStencil)
 				{
-					bThisDecalUsesStencil = RenderPreStencil(RHICmdList, Context, MaterialShaderMap, ComponentToWorldMatrix, FrustumComponentToClip);
+					bThisDecalUsesStencil = RenderPreStencil(Context, MaterialShaderMap, ComponentToWorldMatrix, FrustumComponentToClip);
 					WasInsideDecal = -1;
 					LastDecalBlendMode = -1;
 				}
@@ -982,12 +977,12 @@ void FRCPassPostProcessDeferredDecals::Process(FRenderingCompositePassContext& C
 				LastDecalBlendMode = DecalData.DecalBlendMode;
 				LastDecalHasNormal = (int32)DecalData.bHasNormal;
 
-				SetDecalBlendState(SMFeatureLevel, RenderStage, (EDecalBlendMode)LastDecalBlendMode, DecalData.bHasNormal);
+				SetDecalBlendState(RHICmdList, SMFeatureLevel, RenderStage, (EDecalBlendMode)LastDecalBlendMode, DecalData.bHasNormal);
 			}
 
 			{
 				FDeferredDecalVS* VertexShader = MaterialShaderMap->GetShader<FDeferredDecalVS>();
-				SetShader(RHICmdList, Context, DecalData, VertexShader);
+				SetShader(Context, DecalData, VertexShader);
 
 				VertexShader->SetParameters(RHICmdList, View, FrustumComponentToClip);
 
@@ -998,13 +993,13 @@ void FRCPassPostProcessDeferredDecals::Process(FRenderingCompositePassContext& C
 					if ( !(IsInsideDecal & 1) )
 					{
 						// Render backfaces with depth tests disabled since the camera is inside (or close to inside) the light function geometry
-						RHISetRasterizerState(View.bReverseCulling ? TStaticRasterizerState<FM_Solid,CM_CCW>::GetRHI() : TStaticRasterizerState<FM_Solid,CM_CW>::GetRHI());
+						RHICmdList.SetRasterizerState(View.bReverseCulling ? TStaticRasterizerState<FM_Solid, CM_CCW>::GetRHI() : TStaticRasterizerState<FM_Solid, CM_CW>::GetRHI());
 						if(bStencilDecals)
 						{
 							// Enable stencil testing, only write to pixels with stencil of 0
 							if ( bThisDecalUsesStencil )
 							{
-								RHISetDepthStencilState(TStaticDepthStencilState<
+								RHICmdList.SetDepthStencilState(TStaticDepthStencilState<
 									false,CF_Always,
 									true,CF_Equal,SO_Zero,SO_Zero,SO_Zero,
 									true,CF_Equal,SO_Zero,SO_Zero,SO_Zero,
@@ -1013,7 +1008,7 @@ void FRCPassPostProcessDeferredDecals::Process(FRenderingCompositePassContext& C
 							}
 							else
 							{
-								RHISetDepthStencilState(TStaticDepthStencilState<
+								RHICmdList.SetDepthStencilState(TStaticDepthStencilState<
 									false,CF_Always,
 									true,CF_Equal,SO_Keep,SO_Keep,SO_Keep,
 									false,CF_Always,SO_Keep,SO_Keep,SO_Keep,
@@ -1022,7 +1017,7 @@ void FRCPassPostProcessDeferredDecals::Process(FRenderingCompositePassContext& C
 						}
 						else
 						{
-							RHISetDepthStencilState(TStaticDepthStencilState<false,CF_Always,true>::GetRHI(), 0);
+							RHICmdList.SetDepthStencilState(TStaticDepthStencilState<false, CF_Always, true>::GetRHI(), 0);
 						}
 					}
 					else
@@ -1036,7 +1031,7 @@ void FRCPassPostProcessDeferredDecals::Process(FRenderingCompositePassContext& C
 							// Note, this is a reversed Z depth surface, using CF_GreaterEqual.
 							if ( bThisDecalUsesStencil )
 							{
-								RHISetDepthStencilState(TStaticDepthStencilState<
+								RHICmdList.SetDepthStencilState(TStaticDepthStencilState<
 									false,CF_GreaterEqual,
 									true,CF_Equal,SO_Zero,SO_Zero,SO_Zero,
 									true,CF_Equal,SO_Zero,SO_Zero,SO_Zero,
@@ -1045,25 +1040,25 @@ void FRCPassPostProcessDeferredDecals::Process(FRenderingCompositePassContext& C
 							}
 							else
 							{
-								RHISetDepthStencilState(TStaticDepthStencilState<
+								RHICmdList.SetDepthStencilState(TStaticDepthStencilState<
 									false,CF_GreaterEqual,
 									true,CF_Equal,SO_Keep,SO_Keep,SO_Keep,
 									false,CF_Always,SO_Keep,SO_Keep,SO_Keep,
 									0x80,0x00>::GetRHI(), 0);
 							}
-							RHISetRasterizerState(View.bReverseCulling ? TStaticRasterizerState<FM_Solid,CM_CW>::GetRHI() : TStaticRasterizerState<FM_Solid,CM_CCW>::GetRHI());
+							RHICmdList.SetRasterizerState(View.bReverseCulling ? TStaticRasterizerState<FM_Solid, CM_CW>::GetRHI() : TStaticRasterizerState<FM_Solid, CM_CCW>::GetRHI());
 						}
 						else
 						{
-							RHISetDepthStencilState(TStaticDepthStencilState<false,CF_GreaterEqual>::GetRHI(), 0);
+							RHICmdList.SetDepthStencilState(TStaticDepthStencilState<false, CF_GreaterEqual>::GetRHI(), 0);
 						}
-						RHISetRasterizerState(View.bReverseCulling ? TStaticRasterizerState<FM_Solid,CM_CW>::GetRHI() : TStaticRasterizerState<FM_Solid,CM_CCW>::GetRHI());
+						RHICmdList.SetRasterizerState(View.bReverseCulling ? TStaticRasterizerState<FM_Solid, CM_CW>::GetRHI() : TStaticRasterizerState<FM_Solid, CM_CCW>::GetRHI());
 					}
 				}
 
-				SetShader(RHICmdList, Context, DecalData, VertexShader);
+				SetShader(Context, DecalData, VertexShader);
 
-				RHIDrawIndexedPrimitive(GUnitCubeIndexBuffer.IndexBufferRHI, PT_TriangleList, 0, 0, 8, 0, GUnitCubeIndexBuffer.GetIndexCount() / 3, 0);
+				RHICmdList.DrawIndexedPrimitive(GUnitCubeIndexBuffer.IndexBufferRHI, PT_TriangleList, 0, 0, 8, 0, GUnitCubeIndexBuffer.GetIndexCount() / 3, 0);
 			}
 		}
 
@@ -1072,7 +1067,7 @@ void FRCPassPostProcessDeferredDecals::Process(FRenderingCompositePassContext& C
 		// we don't modify stencil but if out input was having stencil for us (after base pass - we need to clear)
 
 		// Clear stencil to 0, which is the assumed default by other passes
-		RHIClear(false, FLinearColor::White, false, 0, true, 0, FIntRect());
+		RHICmdList.Clear(false, FLinearColor::White, false, 0, true, 0, FIntRect());
 	}
 
 	if(RenderStage == 0)
diff --git a/Engine/Source/Runtime/Renderer/Private/CompositionLighting/PostProcessLpvIndirect.cpp b/Engine/Source/Runtime/Renderer/Private/CompositionLighting/PostProcessLpvIndirect.cpp
--- a/Engine/Source/Runtime/Renderer/Private/CompositionLighting/PostProcessLpvIndirect.cpp
+++ b/Engine/Source/Runtime/Renderer/Private/CompositionLighting/PostProcessLpvIndirect.cpp
@@ -78,8 +78,6 @@ public:
 	{
 		const FPixelShaderRHIParamRef ShaderRHI = GetPixelShader();
 
-		Context.RHICmdList.CheckIsNull();
-
 		SetUniformBufferParameter(Context.RHICmdList, ShaderRHI, GetUniformBufferParameter<FLpvReadUniformBufferParameters>(), LpvUniformBuffer);
 
 #if LPV_VOLUME_TEXTURE
@@ -190,7 +188,7 @@ void FRCPassPostProcessLpvIndirect::Process(FRenderingCompositePassContext& Cont
 	TShaderMapRef<FPostProcessLpvIndirectPS> PixelShader(GetGlobalShaderMap());
 
 	static FGlobalBoundShaderState BoundShaderState;
-	Context.RHICmdList.CheckIsNull(); // need new approach for "static FGlobalBoundShaderState" for parallel rendering
+	
 
 	// call it once after setting up the shader data to avoid the warnings in the function
 	SetGlobalBoundShaderState(Context.RHICmdList, BoundShaderState, GFilterVertexDeclaration.VertexDeclarationRHI, *VertexShader, *PixelShader);
diff --git a/Engine/Source/Runtime/Renderer/Private/CompositionLighting/PostProcessPassThrough.cpp b/Engine/Source/Runtime/Renderer/Private/CompositionLighting/PostProcessPassThrough.cpp
--- a/Engine/Source/Runtime/Renderer/Private/CompositionLighting/PostProcessPassThrough.cpp
+++ b/Engine/Source/Runtime/Renderer/Private/CompositionLighting/PostProcessPassThrough.cpp
@@ -44,10 +44,8 @@ public:
 	void SetParameters(const FRenderingCompositePassContext& Context)
 	{
 		const FPixelShaderRHIParamRef ShaderRHI = GetPixelShader();
-
-		//@todo-rco: RHIPacketList
-		FRHICommandList& RHICmdList = FRHICommandList::GetNullRef();
-		FGlobalShader::SetParameters(RHICmdList, ShaderRHI, Context.View);
+		
+		FGlobalShader::SetParameters(Context.RHICmdList, ShaderRHI, Context.View);
 
 		PostprocessParameter.SetPS(ShaderRHI, Context, TStaticSamplerState<SF_Point,AM_Clamp,AM_Clamp,AM_Clamp>::GetRHI());
 	}
@@ -120,7 +118,7 @@ void FRCPassPostProcessPassThrough::Process(FRenderingCompositePassContext& Cont
 	TShaderMapRef<FPostProcessPassThroughPS> PixelShader(GetGlobalShaderMap());
 
 	static FGlobalBoundShaderState BoundShaderState;
-	Context.RHICmdList.CheckIsNull(); // need new approach for "static FGlobalBoundShaderState" for parallel rendering
+	
 
 	SetGlobalBoundShaderState(Context.RHICmdList, BoundShaderState, GFilterVertexDeclaration.VertexDeclarationRHI, *VertexShader, *PixelShader);
 
