UE4 - rework RHI user-facing APIs to be safer and incorporate the desired sematics

diff --git a/Engine/Source/Runtime/Renderer/Private/DeferredShadingRenderer.cpp b/Engine/Source/Runtime/Renderer/Private/DeferredShadingRenderer.cpp
--- a/Engine/Source/Runtime/Renderer/Private/DeferredShadingRenderer.cpp
+++ b/Engine/Source/Runtime/Renderer/Private/DeferredShadingRenderer.cpp
@@ -100,11 +100,11 @@ FDeferredShadingSceneRenderer::FDeferredShadingSceneRenderer(const FSceneViewFam
 /** 
 * Clears a view. 
 */
-void FDeferredShadingSceneRenderer::ClearView()
+void FDeferredShadingSceneRenderer::ClearView(FRHICommandListImmediate& RHICmdList)
 {
 	// Clear the G Buffer render targets
 	const bool bClearBlack = Views[0].Family->EngineShowFlags.ShaderComplexity || Views[0].Family->EngineShowFlags.StationaryLightOverlap;
-	GSceneRenderTargets.ClearGBufferTargets( bClearBlack ? FLinearColor(0,0,0,0) : Views[0].BackgroundColor );
+	GSceneRenderTargets.ClearGBufferTargets(RHICmdList, bClearBlack ? FLinearColor(0, 0, 0, 0) : Views[0].BackgroundColor);
 }
 
 namespace
@@ -123,7 +123,7 @@ extern FGlobalBoundShaderState GClearMRTBoundShaderState[8];
 /** 
 * Clears view where Z is still at the maximum value (ie no geometry rendered)
 */
-void FDeferredShadingSceneRenderer::ClearGBufferAtMaxZ(FRHICommandList& RHICmdList)
+void FDeferredShadingSceneRenderer::ClearGBufferAtMaxZ(FRHICommandListImmediate& RHICmdList)
 {
 	// Assumes BeginRenderingSceneColor() has been called before this function
 	SCOPED_DRAW_EVENT(ClearGBufferAtMaxZ, DEC_SCENE_ITEMS);
@@ -297,15 +297,13 @@ bool FDeferredShadingSceneRenderer::RenderBasePassStaticData(FRHICommandList& RH
 * @return true if anything was rendered to scene color
 */
 
-bool FDeferredShadingSceneRenderer::RenderBasePassDynamicData(FViewInfo& View)
+bool FDeferredShadingSceneRenderer::RenderBasePassDynamicData(FRHICommandListImmediate& RHICmdList, FViewInfo& View)
 {
 	bool bDirty=0;
 
 	if( !View.Family->EngineShowFlags.CompositeEditorPrimitives )
 	{
-		//@todo-rco: RHIPacketList
-		FRHICommandList& RHICmdList = FRHICommandList::GetNullRef();
-
+		
 		const bool bNeedToSwitchVerticalAxis = IsES2Platform(GRHIShaderPlatform) && !IsPCPlatform(GRHIShaderPlatform);
 
 		// Draw the base pass for the view's batched mesh elements.
@@ -329,7 +327,7 @@ bool FDeferredShadingSceneRenderer::RenderBasePassDynamicData(FViewInfo& View)
 * Renders the basepass for a given DPG and View.
 * @return true if anything was rendered to scene color
 */
-bool FDeferredShadingSceneRenderer::RenderBasePass(FViewInfo& View)
+bool FDeferredShadingSceneRenderer::RenderBasePass(FRHICommandListImmediate& RHICmdList, FViewInfo& View)
 {
 	bool bDirty=0;
 
@@ -339,9 +337,10 @@ bool FDeferredShadingSceneRenderer::RenderBasePass(FViewInfo& View)
 		SCOPE_CYCLE_COUNTER(STAT_RHICounterTEMP);
 		static IConsoleVariable* RHICmdListCVar = IConsoleManager::Get().FindConsoleVariable(TEXT("r.RHICmd"));
 		bool bUseRHICmdList = (RHICmdListCVar->GetInt() >= 1);
-		FRHICommandList& RHICmdList = bUseRHICmdList ? GRHICommandList.CreateList() : FRHICommandList::GetNullRef();
-		bDirty |= RenderBasePassStaticData(RHICmdList, View);
-		GRHICommandList.ExecuteAndFreeList(RHICmdList);
+
+		FRHICommandList LocalRHICmdList;
+		FRHICommandList& UseRHICmdList = bUseRHICmdList ? LocalRHICmdList : RHICmdList;
+		bDirty |= RenderBasePassStaticData(UseRHICmdList, View);
 	}
 
 	{
@@ -379,7 +378,7 @@ bool FDeferredShadingSceneRenderer::RenderBasePass(FViewInfo& View)
 			bDirty |= Drawer.IsDirty(); 
 		}
 
-		bDirty |= RenderBasePassDynamicData(View);
+		bDirty |= RenderBasePassDynamicData(RHICmdList, View);
 	}
 
 	return bDirty;
@@ -387,7 +386,7 @@ bool FDeferredShadingSceneRenderer::RenderBasePass(FViewInfo& View)
 
 /** Render the TexturePool texture */
 #if !(UE_BUILD_SHIPPING || UE_BUILD_TEST)
-void FDeferredShadingSceneRenderer::RenderVisualizeTexturePool()
+void FDeferredShadingSceneRenderer::RenderVisualizeTexturePool(FRHICommandListImmediate& RHICmdList)
 {
 	TRefCountPtr<IPooledRenderTarget> VisualizeTexturePool;
 
@@ -402,15 +401,15 @@ void FDeferredShadingSceneRenderer::RenderVisualizeTexturePool()
 	GRenderTargetPool.FindFreeElement(Desc, VisualizeTexturePool, TEXT("VisualizeTexturePool"));
 	
 	uint32 Pitch;
-	FColor* TextureData = (FColor*)RHILockTexture2D((FTexture2DRHIRef&)VisualizeTexturePool->GetRenderTargetItem().ShaderResourceTexture, 0, RLM_WriteOnly, Pitch, false );
+	FColor* TextureData = (FColor*)RHICmdList.LockTexture2D((FTexture2DRHIRef&)VisualizeTexturePool->GetRenderTargetItem().ShaderResourceTexture, 0, RLM_WriteOnly, Pitch, false);
 	if(TextureData)
 	{
 		// clear with grey to get reliable background color
 		FMemory::Memset(TextureData, 0x88, TexturePoolVisualizerSizeX * TexturePoolVisualizerSizeY * 4);
-		RHIGetTextureMemoryVisualizeData(TextureData, TexturePoolVisualizerSizeX, TexturePoolVisualizerSizeY, Pitch, 4096);
+		RHICmdList.GetTextureMemoryVisualizeData(TextureData, TexturePoolVisualizerSizeX, TexturePoolVisualizerSizeY, Pitch, 4096);
 	}
 
-	RHIUnlockTexture2D((FTexture2DRHIRef&)VisualizeTexturePool->GetRenderTargetItem().ShaderResourceTexture, 0, false);
+	RHICmdList.UnlockTexture2D((FTexture2DRHIRef&)VisualizeTexturePool->GetRenderTargetItem().ShaderResourceTexture, 0, false);
 
 	FIntPoint RTExtent = GSceneRenderTargets.GetBufferSizeXY();
 
@@ -425,19 +424,19 @@ void FDeferredShadingSceneRenderer::RenderVisualizeTexturePool()
 /** 
 * Finishes the view family rendering.
 */
-void FDeferredShadingSceneRenderer::RenderFinish()
+void FDeferredShadingSceneRenderer::RenderFinish(FRHICommandListImmediate& RHICmdList)
 {
 #if !(UE_BUILD_SHIPPING || UE_BUILD_TEST)
 	{
 		if(CVarVisualizeTexturePool.GetValueOnRenderThread())
 		{
-			RenderVisualizeTexturePool();
+			RenderVisualizeTexturePool(RHICmdList);
 		}
 	}
 
 #endif //!(UE_BUILD_SHIPPING || UE_BUILD_TEST)
 
-	FSceneRenderer::RenderFinish();
+	FSceneRenderer::RenderFinish(RHICmdList);
 
 	//grab the new transform out of the proxies for next frame
 	if(ViewFamily.EngineShowFlags.MotionBlur || ViewFamily.EngineShowFlags.TemporalAA)
@@ -454,7 +453,7 @@ void FDeferredShadingSceneRenderer::RenderFinish()
 /** 
 * Renders the view family. 
 */
-void FDeferredShadingSceneRenderer::Render()
+void FDeferredShadingSceneRenderer::Render(FRHICommandListImmediate& RHICmdList)
 {
 	if(!ViewFamily.EngineShowFlags.Rendering)
 	{
@@ -466,13 +465,13 @@ void FDeferredShadingSceneRenderer::Render()
 	const auto FeatureLevel = Scene->GetFeatureLevel();
 
 	// Initialize global system textures (pass-through if already initialized).
-	GSystemTextures.InitializeTextures(FeatureLevel);
+	GSystemTextures.InitializeTextures(RHICmdList, FeatureLevel);
 
 	// Allocate the maximum scene render target space for the current view family.
 	GSceneRenderTargets.Allocate(ViewFamily);
 
 	// Find the visible primitives.
-	InitViews();
+	InitViews(RHICmdList);
 
 	const bool bIsWireframe = ViewFamily.EngineShowFlags.Wireframe;
 	static const auto ClearMethodCVar = IConsoleManager::Get().FindTConsoleVariableDataInt(TEXT("r.ClearSceneMethod"));
@@ -490,7 +489,7 @@ void FDeferredShadingSceneRenderer::Render()
 				break;
 			}
 		
-		case 1: // RHIClear
+		case 1: // RHICmdList.Clear
 			{
 				bRequiresRHIClear = true;
 				bRequiresFarZQuadClear = false;
@@ -522,9 +521,7 @@ void FDeferredShadingSceneRenderer::Render()
 		const bool bHasViewParent = Views[0].State && ((FSceneViewState*)Views[0].State)->HasViewParent();
 	#endif
 
-	//@todo-rco: RHIPacketList
-	FRHICommandList& RHICmdList = FRHICommandList::GetNullRef();
-
+	
 	const bool bIsOcclusionTesting = DoOcclusionQueries() && (!bIsWireframe || bIsViewFrozen || bHasViewParent);
 
 	// Dynamic vertex and index buffers need to be committed before rendering.
@@ -534,20 +531,20 @@ void FDeferredShadingSceneRenderer::Render()
 	// Notify the FX system that the scene is about to be rendered.
 	if (Scene->FXSystem)
 	{
-		Scene->FXSystem->PreRender();
+		Scene->FXSystem->PreRender(RHICmdList);
 	}
 
 	GRenderTargetPool.AddPhaseEvent(TEXT("EarlyZPass"));
 
 	// Draw the scene pre-pass / early z pass, populating the scene depth buffer and HiZ
-	RenderPrePass();
+	RenderPrePass(RHICmdList);
 	
 	GSceneRenderTargets.AllocGBufferTargets();
 	
 	// Clear scene color buffer if necessary.
 	if ( bRequiresRHIClear )
 	{
-		ClearView();
+		ClearView(RHICmdList);
 
 		// Only clear once.
 		bRequiresRHIClear = false;
@@ -566,38 +563,38 @@ void FDeferredShadingSceneRenderer::Render()
 
 	if(IsDBufferEnabled())
 	{
-		GSceneRenderTargets.ResolveSceneDepthTexture();
-		GSceneRenderTargets.ResolveSceneDepthToAuxiliaryTexture();
+		GSceneRenderTargets.ResolveSceneDepthTexture(RHICmdList);
+		GSceneRenderTargets.ResolveSceneDepthToAuxiliaryTexture(RHICmdList);
 
 		// e.g. ambient cubemaps, ambient occlusion, deferred decals
 		for(int32 ViewIndex = 0;ViewIndex < Views.Num();ViewIndex++)
 		{	
 			SCOPED_CONDITIONAL_DRAW_EVENTF(EventView,Views.Num() > 1, DEC_SCENE_ITEMS, TEXT("View%d"), ViewIndex);
-			GCompositionLighting.ProcessBeforeBasePass(Views[ViewIndex]);
+			GCompositionLighting.ProcessBeforeBasePass(RHICmdList, Views[ViewIndex]);
 		}
 	}
 
 	if(bIsWireframe && FDeferredShadingSceneRenderer::ShouldCompositeEditorPrimitives(Views[0]))
 	{
 		// In Editor we want wire frame view modes to be MSAA for better quality. Resolve will be done with EditorPrimitives
-		RHISetRenderTarget(GSceneRenderTargets.GetEditorPrimitivesColor(), GSceneRenderTargets.GetEditorPrimitivesDepth());
-		RHIClear(true, FLinearColor(0, 0, 0, 0), true, 0.0f, false, 0, FIntRect());
+		SetRenderTarget(RHICmdList, GSceneRenderTargets.GetEditorPrimitivesColor(), GSceneRenderTargets.GetEditorPrimitivesDepth());
+		RHICmdList.Clear(true, FLinearColor(0, 0, 0, 0), true, 0.0f, false, 0, FIntRect());
 	}
 	else
 	{
 		// Begin rendering to scene color
-		GSceneRenderTargets.BeginRenderingSceneColor(true);
+		GSceneRenderTargets.BeginRenderingSceneColor(RHICmdList, true);
 	}
 
 	GRenderTargetPool.AddPhaseEvent(TEXT("BasePass"));
 
-	RenderBasePass();
+	RenderBasePass(RHICmdList);
 
 	if(ViewFamily.EngineShowFlags.VisualizeLightCulling)
 	{
 		// clear out emissive and baked lighting (not too efficient but simple and only needed for this debug view)
-		GSceneRenderTargets.BeginRenderingSceneColor(false);
-		RHIClear(true, FLinearColor(0, 0, 0, 0), false, 0, false, 0, FIntRect());
+		GSceneRenderTargets.BeginRenderingSceneColor(RHICmdList, false);
+		RHICmdList.Clear(true, FLinearColor(0, 0, 0, 0), false, 0, false, 0, FIntRect());
 	}
 
 	GSceneRenderTargets.DBufferA.SafeRelease();
@@ -612,22 +609,23 @@ void FDeferredShadingSceneRenderer::Render()
 	if (bRequiresFarZQuadClear)
 	{
 		// Clears view by drawing quad at maximum Z
-		// TODO: if all the platforms have fast color clears, we can replace this with an RHIClear.
+		// TODO: if all the platforms have fast color clears, we can replace this with an RHICmdList.Clear.
 		ClearGBufferAtMaxZ(RHICmdList);
 
 		bRequiresFarZQuadClear = false;
 	}
 	
-	GSceneRenderTargets.ResolveSceneColor(FResolveRect(0, 0, ViewFamily.FamilySizeX, ViewFamily.FamilySizeY));
-	GSceneRenderTargets.ResolveSceneDepthTexture();
-	GSceneRenderTargets.ResolveSceneDepthToAuxiliaryTexture();
+	GSceneRenderTargets.ResolveSceneColor(RHICmdList, FResolveRect(0, 0, ViewFamily.FamilySizeX, ViewFamily.FamilySizeY));
+	GSceneRenderTargets.ResolveSceneDepthTexture(RHICmdList);
+	GSceneRenderTargets.ResolveSceneDepthToAuxiliaryTexture(RHICmdList);
 	
-	RenderCustomDepthPass();
+	RenderCustomDepthPass(RHICmdList);
 
 	// Notify the FX system that opaque primitives have been rendered and we now have a valid depth buffer.
 	if (Scene->FXSystem && Views.IsValidIndex(0))
 	{
 		Scene->FXSystem->PostRenderOpaque(
+			RHICmdList,
 			Views.GetTypedData(),
 			GSceneRenderTargets.GetSceneDepthTexture(),
 			GSceneRenderTargets.GetGBufferATexture()
@@ -642,7 +640,7 @@ void FDeferredShadingSceneRenderer::Render()
 	// This is done after the downsampled depth buffer is created so that it can be used for issuing queries
 	if ( bIsOcclusionTesting )
 	{
-		BeginOcclusionTests();
+		BeginOcclusionTests(RHICmdList);
 	}
 	
 	// Render lighting.
@@ -658,7 +656,7 @@ void FDeferredShadingSceneRenderer::Render()
 		for(int32 ViewIndex = 0;ViewIndex < Views.Num();ViewIndex++)
 		{	
 			SCOPED_CONDITIONAL_DRAW_EVENTF(EventView,Views.Num() > 1, DEC_SCENE_ITEMS, TEXT("View%d"), ViewIndex);
-			GCompositionLighting.ProcessAfterBasePass(Views[ViewIndex]);
+			GCompositionLighting.ProcessAfterBasePass(RHICmdList, Views[ViewIndex]);
 		}
 
 		RenderDynamicSkyLighting(RHICmdList);
@@ -678,7 +676,7 @@ void FDeferredShadingSceneRenderer::Render()
 		// Clear LPVs for all views
 		if ( IsFeatureLevelSupported(GRHIShaderPlatform, ERHIFeatureLevel::SM5) )
 		{
-			PropagateLPVs();
+			PropagateLPVs(RHICmdList);
 		}
 
 		// Render reflections that only operate on opaque pixels
@@ -689,7 +687,7 @@ void FDeferredShadingSceneRenderer::Render()
 		for(int32 ViewIndex = 0; ViewIndex < Views.Num(); ++ViewIndex)
 		{	
 			SCOPED_CONDITIONAL_DRAW_EVENTF(EventView,Views.Num() > 1, DEC_SCENE_ITEMS, TEXT("View%d"), ViewIndex);
-			GCompositionLighting.ProcessLighting(Views[ViewIndex]);
+			GCompositionLighting.ProcessLighting(RHICmdList, Views[ViewIndex]);
 		}
 	}
 
@@ -728,7 +726,7 @@ void FDeferredShadingSceneRenderer::Render()
 				Scene->AtmosphericFog->PrecomputeTextures(RHICmdList, Views.GetTypedData(), &ViewFamily);
 			}
 #endif
-			RenderAtmosphere(LightShaftOutput);
+			RenderAtmosphere(RHICmdList, LightShaftOutput);
 		}
 	}
 
@@ -750,7 +748,7 @@ void FDeferredShadingSceneRenderer::Render()
 	{
 		SCOPE_CYCLE_COUNTER(STAT_TranslucencyDrawTime);
 
-		RenderTranslucency();
+		RenderTranslucency(RHICmdList);
 
 		if (ViewFamily.EngineShowFlags.Refraction)
 		{
@@ -773,7 +771,7 @@ void FDeferredShadingSceneRenderer::Render()
 	}
 
 	// Resolve the scene color for post processing.
-	GSceneRenderTargets.ResolveSceneColor(FResolveRect(0, 0, ViewFamily.FamilySizeX, ViewFamily.FamilySizeY));
+	GSceneRenderTargets.ResolveSceneColor(RHICmdList, FResolveRect(0, 0, ViewFamily.FamilySizeX, ViewFamily.FamilySizeY));
 
 #if !(UE_BUILD_SHIPPING || UE_BUILD_TEST)
 	if(CVarTestUIBlur.GetValueOnRenderThread() > 0)
@@ -790,7 +788,7 @@ void FDeferredShadingSceneRenderer::Render()
 		for(int32 ViewIndex = 0;ViewIndex < Views.Num();ViewIndex++)
 		{	
 			SCOPED_CONDITIONAL_DRAW_EVENTF(EventView, Views.Num() > 1, DEC_SCENE_ITEMS, TEXT("View%d"), ViewIndex);
-			FinishRenderViewTarget(&Views[ViewIndex], ViewIndex == (Views.Num() - 1));
+			FinishRenderViewTarget(RHICmdList, &Views[ViewIndex], ViewIndex == (Views.Num() - 1));
 		}
 	}
 	else
@@ -799,25 +797,23 @@ void FDeferredShadingSceneRenderer::Render()
 		GSceneRenderTargets.AdjustGBufferRefCount(-1);
 	}
 
-	RenderFinish();
+	RenderFinish(RHICmdList);
 }
 
 /** Renders the scene's prepass and occlusion queries */
-bool FDeferredShadingSceneRenderer::RenderPrePass()
+bool FDeferredShadingSceneRenderer::RenderPrePass(FRHICommandListImmediate& RHICmdList)
 {
 	SCOPED_DRAW_EVENT(PrePass, DEC_SCENE_ITEMS);
 	SCOPE_CYCLE_COUNTER(STAT_DepthDrawTime);
 
 	bool bDirty = false;
 
-	GSceneRenderTargets.BeginRenderingPrePass();
-
-	//@todo-rco: RHIPacketList
-	FRHICommandList& RHICmdList = FRHICommandList::GetNullRef();
+	GSceneRenderTargets.BeginRenderingPrePass(RHICmdList);
 
+	
 	// Clear the depth buffer.
 	// Note, this is a reversed Z depth surface, so 0.0f is the far plane.
-	RHIClear(false,FLinearColor::Black,true,0.0f,true,0, FIntRect());
+	RHICmdList.Clear(false,FLinearColor::Black,true,0.0f,true,0, FIntRect());
 
 	// Draw a depth pass to avoid overdraw in the other passes.
 	if(EarlyZPassMode != DDM_None)
@@ -829,10 +825,10 @@ bool FDeferredShadingSceneRenderer::RenderPrePass()
 			const FViewInfo& View = Views[ViewIndex];
 
 			// Disable color writes, enable depth tests and writes.
-			RHISetBlendState(TStaticBlendState<CW_NONE>::GetRHI());
+			RHICmdList.SetBlendState(TStaticBlendState<CW_NONE>::GetRHI());
 			// Note, this is a reversed Z depth surface, using CF_GreaterEqual.
-			RHISetDepthStencilState(TStaticDepthStencilState<true,CF_GreaterEqual>::GetRHI());
-			RHISetViewport(View.ViewRect.Min.X, View.ViewRect.Min.Y, 0, View.ViewRect.Max.X, View.ViewRect.Max.Y, 1);
+			RHICmdList.SetDepthStencilState(TStaticDepthStencilState<true,CF_GreaterEqual>::GetRHI());
+			RHICmdList.SetViewport(View.ViewRect.Min.X, View.ViewRect.Min.Y, 0, View.ViewRect.Max.X, View.ViewRect.Max.Y, 1);
 
 			// Draw the static occluder primitives using a depth drawing policy.
 			{
@@ -895,7 +891,7 @@ bool FDeferredShadingSceneRenderer::RenderPrePass()
 		}
 	}
 
-	GSceneRenderTargets.FinishRenderingPrePass();
+	GSceneRenderTargets.FinishRenderingPrePass(RHICmdList);
 
 	return bDirty;
 }
@@ -904,17 +900,15 @@ bool FDeferredShadingSceneRenderer::RenderPrePass()
  * Renders the scene's base pass 
  * @return true if anything was rendered
  */
-bool FDeferredShadingSceneRenderer::RenderBasePass()
+bool FDeferredShadingSceneRenderer::RenderBasePass(FRHICommandListImmediate& RHICmdList)
 {
 	bool bDirty = false;
 
-	//@todo-rco: RHIPacketList
-	FRHICommandList& RHICmdList = FRHICommandList::GetNullRef();
 
 	if(ViewFamily.EngineShowFlags.LightMapDensity && AllowDebugViewmodes())
 	{
 		// Override the base pass with the lightmap density pass if the viewmode is enabled.
-		bDirty = RenderLightMapDensities();
+		bDirty = RenderLightMapDensities(RHICmdList);
 	}
 	else
 	{
@@ -930,27 +924,27 @@ bool FDeferredShadingSceneRenderer::RenderBasePass()
 			if (ViewFamily.EngineShowFlags.ShaderComplexity)
 			{
 				// Additive blending when shader complexity viewmode is enabled.
-				RHISetBlendState(TStaticBlendState<CW_RGBA,BO_Add,BF_One,BF_One,BO_Add,BF_Zero,BF_One>::GetRHI());
+				RHICmdList.SetBlendState(TStaticBlendState<CW_RGBA,BO_Add,BF_One,BF_One,BO_Add,BF_Zero,BF_One>::GetRHI());
 				// Disable depth writes as we have a full depth prepass.
-				RHISetDepthStencilState(TStaticDepthStencilState<false,CF_GreaterEqual>::GetRHI());
+				RHICmdList.SetDepthStencilState(TStaticDepthStencilState<false,CF_GreaterEqual>::GetRHI());
 			}
 			else
 			{
 				// Opaque blending for all G buffer targets, depth tests and writes.
-				RHISetBlendState(TStaticBlendStateWriteMask<CW_RGBA, CW_RGBA, CW_RGBA, CW_RGBA>::GetRHI());
+				RHICmdList.SetBlendState(TStaticBlendStateWriteMask<CW_RGBA, CW_RGBA, CW_RGBA, CW_RGBA>::GetRHI());
 				// Note, this is a reversed Z depth surface, using CF_GreaterEqual.
-				RHISetDepthStencilState(TStaticDepthStencilState<true,CF_GreaterEqual>::GetRHI());
+				RHICmdList.SetDepthStencilState(TStaticDepthStencilState<true,CF_GreaterEqual>::GetRHI());
 			}
-			RHISetViewport(View.ViewRect.Min.X, View.ViewRect.Min.Y, 0, View.ViewRect.Max.X, View.ViewRect.Max.Y, 1);
+			RHICmdList.SetViewport(View.ViewRect.Min.X, View.ViewRect.Min.Y, 0, View.ViewRect.Max.X, View.ViewRect.Max.Y, 1);
 
-			bDirty |= RenderBasePass(View);
+			bDirty |= RenderBasePass(RHICmdList, View);
 		}
 	}
 
 	return bDirty;
 }
 
-void FDeferredShadingSceneRenderer::ClearLPVs(FRHICommandList& RHICmdList)
+void FDeferredShadingSceneRenderer::ClearLPVs(FRHICommandListImmediate& RHICmdList)
 {
 	// clear light propagation volumes
 
@@ -974,7 +968,7 @@ void FDeferredShadingSceneRenderer::ClearLPVs(FRHICommandList& RHICmdList)
 	}
 }
 
-void FDeferredShadingSceneRenderer::PropagateLPVs()
+void FDeferredShadingSceneRenderer::PropagateLPVs(FRHICommandListImmediate& RHICmdList)
 {
 	for(int32 ViewIndex = 0; ViewIndex < Views.Num(); ViewIndex++)
 	{
@@ -989,8 +983,8 @@ void FDeferredShadingSceneRenderer::PropagateLPVs()
 			{
 				SCOPED_DRAW_EVENT(UpdateLPVs, DEC_SCENE_ITEMS);
 				SCOPE_CYCLE_COUNTER(STAT_UpdateLPVs);
-				//@todo-rco: RHIPacketList
-				LightPropagationVolume->Propagate(FRHICommandList::GetNullRef());
+				
+				LightPropagationVolume->Propagate(RHICmdList);
 			}
 		}
 	}
@@ -1059,7 +1053,7 @@ IMPLEMENT_SHADER_TYPE(,FDownsampleSceneDepthPS,TEXT("DownsampleDepthPixelShader"
 FGlobalBoundShaderState DownsampleDepthBoundShaderState;
 
 /** Updates the downsized depth buffer with the current full resolution depth buffer. */
-void FDeferredShadingSceneRenderer::UpdateDownsampledDepthSurface(FRHICommandList& RHICmdList)
+void FDeferredShadingSceneRenderer::UpdateDownsampledDepthSurface(FRHICommandListImmediate& RHICmdList)
 {
 	if (GSceneRenderTargets.UseDownsizedOcclusionQueries() && IsFeatureLevelSupported(GRHIShaderPlatform, ERHIFeatureLevel::SM3))
 	{
diff --git a/Engine/Source/Runtime/Renderer/Private/DeferredShadingRenderer.h b/Engine/Source/Runtime/Renderer/Private/DeferredShadingRenderer.h
--- a/Engine/Source/Runtime/Renderer/Private/DeferredShadingRenderer.h
+++ b/Engine/Source/Runtime/Renderer/Private/DeferredShadingRenderer.h
@@ -39,16 +39,16 @@ public:
 	FDeferredShadingSceneRenderer(const FSceneViewFamily* InViewFamily,FHitProxyConsumer* HitProxyConsumer);
 
 	/** Clears a view */
-	void ClearView();
+	void ClearView(FRHICommandListImmediate& RHICmdList);
 
 	/** Clears gbuffer where Z is still at the maximum value (ie no geometry rendered) */
-	void ClearGBufferAtMaxZ(FRHICommandList& RHICmdList);
+	void ClearGBufferAtMaxZ(FRHICommandListImmediate& RHICmdList);
 
 	/** Clears LPVs for all views */
-	void ClearLPVs(FRHICommandList& RHICmdList);
+	void ClearLPVs(FRHICommandListImmediate& RHICmdList);
 
 	/** Propagates LPVs for all views */
-	void PropagateLPVs();
+	void PropagateLPVs(FRHICommandListImmediate& RHICmdList);
 
     /** Renders the basepass for the static data of a given View. */
     bool RenderBasePassStaticData(FRHICommandList& RHICmdList, FViewInfo& View);
@@ -59,28 +59,28 @@ public:
 	void SortBasePassStaticData(FVector ViewPosition);
 
     /** Renders the basepass for the dynamic data of a given View. */
-    bool RenderBasePassDynamicData(FViewInfo& View);
+	bool RenderBasePassDynamicData(FRHICommandListImmediate& RHICmdList, FViewInfo& View);
 
     /** Renders the basepass for a given View. */
-    bool RenderBasePass(FViewInfo& View);
+	bool RenderBasePass(FRHICommandListImmediate& RHICmdList, FViewInfo& View);
 	
 	/** 
 	* Renders the scene's base pass 
 	* @return true if anything was rendered
 	*/
-	bool RenderBasePass();
+	bool RenderBasePass(FRHICommandListImmediate& RHICmdList);
 
 	/** Finishes the view family rendering. */
-    void RenderFinish();
+	void RenderFinish(FRHICommandListImmediate& RHICmdList);
 
 	/** Renders the view family. */
-	virtual void Render() override;
+	virtual void Render(FRHICommandListImmediate& RHICmdList) override;
 
 	/** Render the view family's hit proxies. */
-	virtual void RenderHitProxies() override;
+	virtual void RenderHitProxies(FRHICommandListImmediate& RHICmdList) override;
 
 #if !(UE_BUILD_SHIPPING || UE_BUILD_TEST)
-	void RenderVisualizeTexturePool();
+	void RenderVisualizeTexturePool(FRHICommandListImmediate& RHICmdList);
 #endif
 
 	/**
@@ -99,6 +99,7 @@ private:
 
 	/** Creates a per object projected shadow for the given interaction. */
 	void CreatePerObjectProjectedShadow(
+		FRHICommandListImmediate& RHICmdList,
 		FLightPrimitiveInteraction* Interaction, 
 		bool bCreateTranslucentObjectShadow, 
 		bool bCreateInsetObjectShadow,
@@ -139,13 +140,14 @@ private:
 		uint32 InResolutionX);
 
 	/** Calculates projected shadow visibility. */
-	void InitProjectedShadowVisibility();
+	void InitProjectedShadowVisibility(FRHICommandListImmediate& RHICmdList);
 
 	/** Returns whether a per object shadow should be created due to the light being a stationary light. */
 	bool ShouldCreateObjectShadowForStationaryLight(const FLightSceneInfo* LightSceneInfo, const FPrimitiveSceneProxy* PrimitiveSceneProxy, bool bInteractionShadowMapped) const;
 
 	/** Creates shadows for the given interaction. */
 	void SetupInteractionShadows(
+		FRHICommandListImmediate& RHICmdList,
 		FLightPrimitiveInteraction* Interaction, 
 		FVisibleLightInfo& VisibleLightInfo, 
 		bool bStaticSceneOnly,
@@ -153,34 +155,34 @@ private:
 		TArray<FProjectedShadowInfo*,SceneRenderingAllocator>& PreShadows);
 
 	/** Finds the visible dynamic shadows for each view. */
-	void InitDynamicShadows();
+	void InitDynamicShadows(FRHICommandListImmediate& RHICmdList);
 
 	/** Determines which primitives are visible for each view. */
-	void InitViews(); 
+	void InitViews(FRHICommandListImmediate& RHICmdList);
 
 	/**
 	 * Renders the scene's prepass and occlusion queries.
 	 * @return true if anything was rendered
 	 */
-	bool RenderPrePass();
+	bool RenderPrePass(FRHICommandListImmediate& RHICmdList);
 
 	/** Issues occlusion queries. */
-	void BeginOcclusionTests();
+	void BeginOcclusionTests(FRHICommandListImmediate& RHICmdList);
 
 	/** Renders the scene's fogging. */
-	bool RenderFog(FRHICommandList& RHICmdList, FLightShaftsOutput LightShaftsOutput);
+	bool RenderFog(FRHICommandListImmediate& RHICmdList, FLightShaftsOutput LightShaftsOutput);
 
 	/** Renders the scene's atmosphere. */
-	void RenderAtmosphere(FLightShaftsOutput LightShaftsOutput);
+	void RenderAtmosphere(FRHICommandListImmediate& RHICmdList, FLightShaftsOutput LightShaftsOutput);
 
 	/** Renders reflections that can be done in a deferred pass. */
-	void RenderDeferredReflections(FRHICommandList& RHICmdList);
+	void RenderDeferredReflections(FRHICommandListImmediate& RHICmdList);
 
 	/** Render dynamic sky lighting from Movable sky lights. */
-	void RenderDynamicSkyLighting(FRHICommandList& RHICmdList);
+	void RenderDynamicSkyLighting(FRHICommandListImmediate& RHICmdList);
 
 	/** Render Ambient Occlusion using mesh distance fields and the surface cache, which supports dynamic rigid meshes. */
-	bool RenderDistanceFieldAOSurfaceCache(FRHICommandList& RHICmdList, FSceneRenderTargetItem& OutBentNormalAO, bool bApplyToSceneColor);
+	bool RenderDistanceFieldAOSurfaceCache(FRHICommandListImmediate& RHICmdList, FSceneRenderTargetItem& OutBentNormalAO, bool bApplyToSceneColor);
 
 	/** Whether tiled deferred is supported and can be used at all. */
 	bool CanUseTiledDeferred() const;
@@ -189,44 +191,44 @@ private:
 	bool ShouldUseTiledDeferred(int32 NumUnshadowedLights, int32 NumSimpleLights) const;
 
 	/** Renders the lights in SortedLights in the range [0, NumUnshadowedLights) using tiled deferred shading. */
-	void RenderTiledDeferredLighting(FRHICommandList& RHICmdList, const TArray<FSortedLightSceneInfo, SceneRenderingAllocator>& SortedLights, int32 NumUnshadowedLights, const FSimpleLightArray& SimpleLights);
+	void RenderTiledDeferredLighting(FRHICommandListImmediate& RHICmdList, const TArray<FSortedLightSceneInfo, SceneRenderingAllocator>& SortedLights, int32 NumUnshadowedLights, const FSimpleLightArray& SimpleLights);
 
 	/** Renders the scene's lighting. */
-	void RenderLights(FRHICommandList& RHICmdList);
+	void RenderLights(FRHICommandListImmediate& RHICmdList);
 
 	/** Renders an array of lights for the stationary light overlap viewmode. */
-	void RenderLightArrayForOverlapViewmode(FRHICommandList& RHICmdList, const TSparseArray<FLightSceneInfoCompact>& LightArray);
+	void RenderLightArrayForOverlapViewmode(FRHICommandListImmediate& RHICmdList, const TSparseArray<FLightSceneInfoCompact>& LightArray);
 
 	/** Render stationary light overlap as complexity to scene color. */
-	void RenderStationaryLightOverlap(FRHICommandList& RHICmdList);
+	void RenderStationaryLightOverlap(FRHICommandListImmediate& RHICmdList);
 	
 	/** 
 	 * Renders the scene's translucency.
 	 */
-	void RenderTranslucency();
+	void RenderTranslucency(FRHICommandListImmediate& RHICmdList);
 
 	/** Renders the scene's light shafts */
-	FLightShaftsOutput RenderLightShaftOcclusion(FRHICommandList& RHICmdList);
+	FLightShaftsOutput RenderLightShaftOcclusion(FRHICommandListImmediate& RHICmdList);
 
-	void RenderLightShaftBloom(FRHICommandList& RHICmdList);
+	void RenderLightShaftBloom(FRHICommandListImmediate& RHICmdList);
 
 	/** Reuses an existing translucent shadow map if possible or re-renders one if necessary. */
-	const FProjectedShadowInfo* PrepareTranslucentShadowMap(const FViewInfo& View, FPrimitiveSceneInfo* PrimitiveSceneInfo, bool bSeparateTranslucencyPass);
+	const FProjectedShadowInfo* PrepareTranslucentShadowMap(FRHICommandListImmediate& RHICmdList, const FViewInfo& View, FPrimitiveSceneInfo* PrimitiveSceneInfo, bool bSeparateTranslucencyPass);
 
 	/** Renders the velocities of movable objects for the motion blur effect. */
-	void RenderVelocities(const FViewInfo& View, TRefCountPtr<IPooledRenderTarget>& VelocityRT, bool bLastFrame);
+	void RenderVelocities(FRHICommandListImmediate& RHICmdList, const FViewInfo& View, TRefCountPtr<IPooledRenderTarget>& VelocityRT, bool bLastFrame);
 
 	/** Renders world-space lightmap density instead of the normal color. */
-	bool RenderLightMapDensities();
+	bool RenderLightMapDensities(FRHICommandListImmediate& RHICmdList);
 
 	/** Updates the downsized depth buffer with the current full resolution depth buffer. */
-	void UpdateDownsampledDepthSurface(FRHICommandList& RHICmdList);
+	void UpdateDownsampledDepthSurface(FRHICommandListImmediate& RHICmdList);
 
 	/**
 	 * Finish rendering a view, writing the contents to ViewFamily.RenderTarget.
 	 * @param View - The view to process.
 	*/
-	void FinishRenderViewTarget(const FViewInfo* View, bool bLastView);
+	void FinishRenderViewTarget(FRHICommandListImmediate& RHICmdList, const FViewInfo* View, bool bLastView);
 
 	/**
 	  * Used by RenderLights to figure out if projected shadows need to be rendered to the attenuation buffer.
@@ -237,20 +239,20 @@ private:
 	bool CheckForProjectedShadows( const FLightSceneInfo* LightSceneInfo ) const;
 
 	/** Renders one pass point light shadows. */
-	bool RenderOnePassPointLightShadows(FRHICommandList& RHICmdList, const FLightSceneInfo* LightSceneInfo, bool bRenderedTranslucentObjectShadows, bool& bInjectedTranslucentVolume);
+	bool RenderOnePassPointLightShadows(FRHICommandListImmediate& RHICmdList, const FLightSceneInfo* LightSceneInfo, bool bRenderedTranslucentObjectShadows, bool& bInjectedTranslucentVolume);
 
 	/** Renders the projections of the given Shadows to the appropriate color render target. */
 	void RenderProjections(
-		FRHICommandList& RHICmdList, 
+		FRHICommandListImmediate& RHICmdList,
 		const FLightSceneInfo* LightSceneInfo, 
 		const TArray<FProjectedShadowInfo*,SceneRenderingAllocator>& Shadows
 		);
 
 	/** Renders the shadowmaps of translucent shadows and their projections onto opaque surfaces. */
-	bool RenderTranslucentProjectedShadows(FRHICommandList& RHICmdList, const FLightSceneInfo* LightSceneInfo );
+	bool RenderTranslucentProjectedShadows(FRHICommandListImmediate& RHICmdList, const FLightSceneInfo* LightSceneInfo);
 
 	/** Renders reflective shadowmaps for LPVs */
-	bool RenderReflectiveShadowMaps(FRHICommandList& RHICmdList, const FLightSceneInfo* LightSceneInfo );
+	bool RenderReflectiveShadowMaps(FRHICommandListImmediate& RHICmdList, const FLightSceneInfo* LightSceneInfo);
 
 	/**
 	  * Used by RenderLights to render projected shadows to the attenuation buffer.
@@ -258,10 +260,10 @@ private:
 	  * @param LightSceneInfo Represents the current light
 	  * @return true if anything got rendered
 	  */
-	bool RenderProjectedShadows(FRHICommandList& RHICmdList, const FLightSceneInfo* LightSceneInfo, bool bRenderedTranslucentObjectShadows, bool& bInjectedTranslucentVolume );
+	bool RenderProjectedShadows(FRHICommandListImmediate& RHICmdList, const FLightSceneInfo* LightSceneInfo, bool bRenderedTranslucentObjectShadows, bool& bInjectedTranslucentVolume);
 
 	/** Caches the depths of any preshadows that should be cached, and renders their projections. */
-	bool RenderCachedPreshadows(FRHICommandList& RHICmdList, const FLightSceneInfo* LightSceneInfo);
+	bool RenderCachedPreshadows(FRHICommandListImmediate& RHICmdList, const FLightSceneInfo* LightSceneInfo);
 
 	/**
 	  * Used by RenderLights to figure out if light functions need to be rendered to the attenuation buffer.
@@ -277,13 +279,13 @@ private:
 	  * @param LightSceneInfo Represents the current light
 	  * @param LightIndex The light's index into FScene::Lights
 	  */
-	bool RenderLightFunction(FRHICommandList& RHICmdList, const FLightSceneInfo* LightSceneInfo, bool bLightAttenuationCleared);
+	bool RenderLightFunction(FRHICommandListImmediate& RHICmdList, const FLightSceneInfo* LightSceneInfo, bool bLightAttenuationCleared);
 
 	/** Renders a light function indicating that whole scene shadowing being displayed is for previewing only, and will go away in game. */
-	bool RenderPreviewShadowsIndicator(FRHICommandList& RHICmdList, const FLightSceneInfo* LightSceneInfo, bool bLightAttenuationCleared);
+	bool RenderPreviewShadowsIndicator(FRHICommandListImmediate& RHICmdList, const FLightSceneInfo* LightSceneInfo, bool bLightAttenuationCleared);
 
 	/** Renders a light function with the given material. */
-	bool RenderLightFunctionForMaterial(FRHICommandList& RHICmdList, const FLightSceneInfo* LightSceneInfo, const FMaterialRenderProxy* MaterialProxy, bool bLightAttenuationCleared, bool bRenderingPreviewShadowsIndicator);
+	bool RenderLightFunctionForMaterial(FRHICommandListImmediate& RHICmdList, const FLightSceneInfo* LightSceneInfo, const FMaterialRenderProxy* MaterialProxy, bool bLightAttenuationCleared, bool bRenderingPreviewShadowsIndicator);
 
 	/**
 	  * Used by RenderLights to render a light to the scene color buffer.
@@ -292,46 +294,46 @@ private:
 	  * @param LightIndex The light's index into FScene::Lights
 	  * @return true if anything got rendered
 	  */
-	void RenderLight(FRHICommandList& RHICmdList, const FLightSceneInfo* LightSceneInfo, bool bRenderOverlap, bool bIssueDrawEvent);
+	void RenderLight(FRHICommandListImmediate& RHICmdList, const FLightSceneInfo* LightSceneInfo, bool bRenderOverlap, bool bIssueDrawEvent);
 
 	/** Renders an array of simple lights using standard deferred shading. */
-	void RenderSimpleLightsStandardDeferred(FRHICommandList& RHICmdList, const FSimpleLightArray& SimpleLights);
+	void RenderSimpleLightsStandardDeferred(FRHICommandListImmediate& RHICmdList, const FSimpleLightArray& SimpleLights);
 
 	/** Clears the translucency lighting volumes before light accumulation. */
-	void ClearTranslucentVolumeLighting(FRHICommandList& RHICmdList);
+	void ClearTranslucentVolumeLighting(FRHICommandListImmediate& RHICmdList);
 
 	/** Add AmbientCubemap to the lighting volumes. */
-	void InjectAmbientCubemapTranslucentVolumeLighting(FRHICommandList& RHICmdList);
+	void InjectAmbientCubemapTranslucentVolumeLighting(FRHICommandListImmediate& RHICmdList);
 
 	/** Renders indirect lighting into the translucent lighting volumes. */
-	void CompositeIndirectTranslucentVolumeLighting(FRHICommandList& RHICmdList);
+	void CompositeIndirectTranslucentVolumeLighting(FRHICommandListImmediate& RHICmdList);
 
 	/** Clears the volume texture used to accumulate per object shadows for translucency. */
-	void ClearTranslucentVolumePerObjectShadowing(FRHICommandList& RHICmdList);
+	void ClearTranslucentVolumePerObjectShadowing(FRHICommandListImmediate& RHICmdList);
 
 	/** Accumulates the per object shadow's contribution for translucency. */
-	void AccumulateTranslucentVolumeObjectShadowing(FRHICommandList& RHICmdList, const FProjectedShadowInfo* InProjectedShadowInfo, bool bClearVolume);
+	void AccumulateTranslucentVolumeObjectShadowing(FRHICommandListImmediate& RHICmdList, const FProjectedShadowInfo* InProjectedShadowInfo, bool bClearVolume);
 
 	/** Accumulates direct lighting for the given light.  InProjectedShadowInfo can be NULL in which case the light will be unshadowed. */
-	void InjectTranslucentVolumeLighting(FRHICommandList& RHICmdList, const FLightSceneInfo& LightSceneInfo, const FProjectedShadowInfo* InProjectedShadowInfo);
+	void InjectTranslucentVolumeLighting(FRHICommandListImmediate& RHICmdList, const FLightSceneInfo& LightSceneInfo, const FProjectedShadowInfo* InProjectedShadowInfo);
 
 	/** Accumulates direct lighting for an array of unshadowed lights. */
-	void InjectTranslucentVolumeLightingArray(FRHICommandList& RHICmdList, const TArray<FSortedLightSceneInfo, SceneRenderingAllocator>& SortedLights, int32 NumLights);
+	void InjectTranslucentVolumeLightingArray(FRHICommandListImmediate& RHICmdList, const TArray<FSortedLightSceneInfo, SceneRenderingAllocator>& SortedLights, int32 NumLights);
 
 	/** Accumulates direct lighting for simple lights. */
-	void InjectSimpleTranslucentVolumeLightingArray(FRHICommandList& RHICmdList, const FSimpleLightArray& SimpleLights);
+	void InjectSimpleTranslucentVolumeLightingArray(FRHICommandListImmediate& RHICmdList, const FSimpleLightArray& SimpleLights);
 
 	/** Filters the translucency lighting volumes to reduce aliasing. */
-	void FilterTranslucentVolumeLighting(FRHICommandList& RHICmdList);
+	void FilterTranslucentVolumeLighting(FRHICommandListImmediate& RHICmdList);
 
 	/** Output SpecularColor * IndirectDiffuseGI for metals so they are not black in reflections */
-	void RenderReflectionCaptureSpecularBounceForAllViews(FRHICommandList& RHICmdList);
+	void RenderReflectionCaptureSpecularBounceForAllViews(FRHICommandListImmediate& RHICmdList);
 
 	/** Render image based reflections (SSR, Env, SkyLight) without compute shaders */
-	void RenderImageBasedReflectionsSM4ForAllViews(FRHICommandList& RHICmdList, bool bReflectionEnv);
+	void RenderImageBasedReflectionsSM4ForAllViews(FRHICommandListImmediate& RHICmdList, bool bReflectionEnv);
 
 	/** Render image based reflections (SSR, Env, SkyLight) with compute shaders */
-	void RenderImageBasedReflectionsSM5ForAllViews(FRHICommandList& RHICmdList);
+	void RenderImageBasedReflectionsSM5ForAllViews(FRHICommandListImmediate& RHICmdList);
 
 	bool ShouldDoReflectionEnvironment() const;
 
diff --git a/Engine/Source/Runtime/Renderer/Private/DepthRendering.cpp b/Engine/Source/Runtime/Renderer/Private/DepthRendering.cpp
--- a/Engine/Source/Runtime/Renderer/Private/DepthRendering.cpp
+++ b/Engine/Source/Runtime/Renderer/Private/DepthRendering.cpp
@@ -259,8 +259,6 @@ FPositionOnlyDepthDrawingPolicy::FPositionOnlyDepthDrawingPolicy(
 
 void FPositionOnlyDepthDrawingPolicy::DrawShared(FRHICommandList& RHICmdList, const FSceneView* View,FBoundShaderStateRHIParamRef BoundShaderState) const
 {
-	RHICmdList.CheckIsNull();
-
 	// Set the actual shader & vertex declaration state
 	RHICmdList.SetBoundShaderState( BoundShaderState);
 
diff --git a/Engine/Source/Runtime/Renderer/Private/DistanceFieldSurfaceCacheLighting.cpp b/Engine/Source/Runtime/Renderer/Private/DistanceFieldSurfaceCacheLighting.cpp
--- a/Engine/Source/Runtime/Renderer/Private/DistanceFieldSurfaceCacheLighting.cpp
+++ b/Engine/Source/Runtime/Renderer/Private/DistanceFieldSurfaceCacheLighting.cpp
@@ -282,7 +282,7 @@ public:
 		MaxObjects = 0;
 	}
 
-	virtual void InitDynamicRHI() 
+	virtual void InitDynamicRHI()  override
 	{
 		if (MaxObjects > 0)
 		{
@@ -297,7 +297,7 @@ public:
 		}
 	}
 
-	virtual void ReleaseDynamicRHI()
+	virtual void ReleaseDynamicRHI() override
 	{
 		ObjectData.Bounds.SafeRelease();
 		ObjectData.Data.SafeRelease();
@@ -941,7 +941,7 @@ private:
 
 IMPLEMENT_SHADER_TYPE(,FComputeDistanceFieldNormalCS,TEXT("DistanceFieldSurfaceCacheLighting"),TEXT("ComputeDistanceFieldNormalCS"),SF_Compute);
 
-void ComputeDistanceFieldNormal(FRHICommandList& RHICmdList, const TArray<FViewInfo>& Views, FSceneRenderTargetItem& DistanceFieldNormal, int32 NumObjects)
+void ComputeDistanceFieldNormal(FRHICommandListImmediate& RHICmdList, const TArray<FViewInfo>& Views, FSceneRenderTargetItem& DistanceFieldNormal, int32 NumObjects)
 {
 	if (GAOComputeShaderNormalCalculation)
 	{
@@ -985,7 +985,7 @@ void ComputeDistanceFieldNormal(FRHICommandList& RHICmdList, const TArray<FViewI
 			TShaderMapRef<FComputeDistanceFieldNormalPS> PixelShader( GetGlobalShaderMap() );
 
 			static FGlobalBoundShaderState BoundShaderState;
-			RHICmdList.CheckIsNull(); // need new approach for "static FGlobalBoundShaderState" for parallel rendering
+			
 			SetGlobalBoundShaderState(RHICmdList, BoundShaderState, GFilterVertexDeclaration.VertexDeclarationRHI, *VertexShader, *PixelShader);
 
 			PixelShader->SetParameters(RHICmdList, View, NumObjects);
@@ -2054,7 +2054,7 @@ void AllocateOrReuseAORenderTarget(TRefCountPtr<IPooledRenderTarget>& Target, co
 }
 
 void UpdateHistory(
-	FRHICommandList& RHICmdList,
+	FRHICommandListImmediate& RHICmdList,
 	FViewInfo& View, 
 	const TCHAR* HistoryRTName,
 	/** Contains last frame's history, if non-NULL.  This will be updated with the new frame's history. */
@@ -2068,7 +2068,6 @@ void UpdateHistory(
 	{
 		if (*HistoryState && !View.bCameraCut && !View.bPrevTransformsReset)
 		{
-			RHICmdList.CheckIsNull(); // direct allocation of resources won't work in parallel
 			// Reuse a render target from the pool with a consistent name, for vis purposes
 			TRefCountPtr<IPooledRenderTarget> NewHistory;
 			AllocateOrReuseAORenderTarget(NewHistory, HistoryRTName);
@@ -2086,7 +2085,7 @@ void UpdateHistory(
 				TShaderMapRef<FUpdateHistoryPS> PixelShader( GetGlobalShaderMap() );
 
 				static FGlobalBoundShaderState BoundShaderState;
-				RHICmdList.CheckIsNull(); // need new approach for "static FGlobalBoundShaderState" for parallel rendering
+				
 				SetGlobalBoundShaderState(RHICmdList, BoundShaderState, GFilterVertexDeclaration.VertexDeclarationRHI, *VertexShader, *PixelShader);
 
 				PixelShader->SetParameters(RHICmdList, View, (*HistoryState)->GetRenderTargetItem(), AOSource->GetRenderTargetItem());
@@ -2121,9 +2120,8 @@ void UpdateHistory(
 	}
 }
 
-void PostProcessBentNormalAO(FRHICommandList& RHICmdList, TArray<FViewInfo>& Views, FSceneRenderTargetItem& IrradianceCacheInterpolation, TRefCountPtr<IPooledRenderTarget>& AOOutput)
+void PostProcessBentNormalAO(FRHICommandListImmediate& RHICmdList, TArray<FViewInfo>& Views, FSceneRenderTargetItem& IrradianceCacheInterpolation, TRefCountPtr<IPooledRenderTarget>& AOOutput)
 {
-	RHICmdList.CheckIsNull(); // direct allocation of resources won't work in parallel
 	TRefCountPtr<IPooledRenderTarget> DistanceFieldAOBentNormal;
 	AllocateOrReuseAORenderTarget(DistanceFieldAOBentNormal, TEXT("DistanceFieldBentNormalAO"));
 
@@ -2154,7 +2152,7 @@ void PostProcessBentNormalAO(FRHICommandList& RHICmdList, TArray<FViewInfo>& Vie
 			TShaderMapRef<FDistanceFieldAOCombinePS2> PixelShader(GetGlobalShaderMap());
 
 			static FGlobalBoundShaderState BoundShaderState;
-			RHICmdList.CheckIsNull(); // need new approach for "static FGlobalBoundShaderState" for parallel rendering
+			
 			SetGlobalBoundShaderState(RHICmdList, BoundShaderState, GFilterVertexDeclaration.VertexDeclarationRHI, *VertexShader, *PixelShader);
 
 			PixelShader->SetParameters(RHICmdList, View, IrradianceCacheInterpolation);
@@ -2190,7 +2188,7 @@ void PostProcessBentNormalAO(FRHICommandList& RHICmdList, TArray<FViewInfo>& Vie
 			TShaderMapRef<FFillGapsPS> PixelShader( GetGlobalShaderMap() );
 
 			static FGlobalBoundShaderState BoundShaderState;
-			RHICmdList.CheckIsNull(); // need new approach for "static FGlobalBoundShaderState" for parallel rendering
+			
 			SetGlobalBoundShaderState(RHICmdList, BoundShaderState, GFilterVertexDeclaration.VertexDeclarationRHI, *VertexShader, *PixelShader);
 
 			VertexShader->SetParameters(RHICmdList, View);
@@ -2281,7 +2279,7 @@ private:
 
 IMPLEMENT_SHADER_TYPE(,FDistanceFieldAOUpsamplePS,TEXT("DistanceFieldSurfaceCacheLighting"),TEXT("AOUpsamplePS"),SF_Pixel);
 
-void UpsampleBentNormalAO(FRHICommandList& RHICmdList, const TArray<FViewInfo>& Views, TRefCountPtr<IPooledRenderTarget>& DistanceFieldAOBentNormal)
+void UpsampleBentNormalAO(FRHICommandListImmediate& RHICmdList, const TArray<FViewInfo>& Views, TRefCountPtr<IPooledRenderTarget>& DistanceFieldAOBentNormal)
 {
 	for (int32 ViewIndex = 0; ViewIndex < Views.Num(); ViewIndex++)
 	{
@@ -2298,7 +2296,7 @@ void UpsampleBentNormalAO(FRHICommandList& RHICmdList, const TArray<FViewInfo>&
 		TShaderMapRef<FDistanceFieldAOUpsamplePS> PixelShader( GetGlobalShaderMap() );
 
 		static FGlobalBoundShaderState BoundShaderState;
-		RHICmdList.CheckIsNull(); // need new approach for "static FGlobalBoundShaderState" for parallel rendering
+		
 		SetGlobalBoundShaderState(RHICmdList, BoundShaderState, GFilterVertexDeclaration.VertexDeclarationRHI, *VertexShader, *PixelShader);
 
 		PixelShader->SetParameters(RHICmdList, View, DistanceFieldAOBentNormal);
@@ -2472,7 +2470,7 @@ void GenerateBestSpacedVectors()
 	}
 }
 
-FIntPoint BuildTileObjectLists(FRHICommandList& RHICmdList, FScene* Scene, TArray<FViewInfo>& Views, int32 NumObjects)
+FIntPoint BuildTileObjectLists(FRHICommandListImmediate& RHICmdList, FScene* Scene, TArray<FViewInfo>& Views, int32 NumObjects)
 {
 	SCOPED_DRAW_EVENT(BuildTileList, DEC_SCENE_ITEMS);
 	SetRenderTarget(RHICmdList, NULL, NULL);
@@ -2481,7 +2479,6 @@ FIntPoint BuildTileObjectLists(FRHICommandList& RHICmdList, FScene* Scene, TArra
 
 	if (GAOScatterTileCulling)
 	{
-		RHICmdList.CheckIsNull(); // direct allocation of resources won't work in parallel
 
 		for (int32 ViewIndex = 0; ViewIndex < Views.Num(); ViewIndex++)
 		{
@@ -2538,7 +2535,7 @@ FIntPoint BuildTileObjectLists(FRHICommandList& RHICmdList, FScene* Scene, TArra
 				RHICmdList.SetBlendState(TStaticBlendState<>::GetRHI());
 
 				static FGlobalBoundShaderState BoundShaderState;
-				RHICmdList.CheckIsNull(); // need new approach for "static FGlobalBoundShaderState" for parallel rendering
+				
 				SetGlobalBoundShaderState(RHICmdList, BoundShaderState, GetVertexDeclarationFVector4(), *VertexShader, *PixelShader);
 
 				VertexShader->SetParameters(RHICmdList, View, NumObjects);
@@ -2604,7 +2601,7 @@ FIntPoint BuildTileObjectLists(FRHICommandList& RHICmdList, FScene* Scene, TArra
 	return TileListGroupSize;
 }
 
-void RenderIrradianceCacheInterpolation(FRHICommandList& RHICmdList, TArray<FViewInfo>& Views, IPooledRenderTarget* InterpolationTarget, FSceneRenderTargetItem& DistanceFieldNormal, int32 DepthLevel, int32 DestLevelDownsampleFactor, bool bFinalInterpolation)
+void RenderIrradianceCacheInterpolation(FRHICommandListImmediate& RHICmdList, TArray<FViewInfo>& Views, IPooledRenderTarget* InterpolationTarget, FSceneRenderTargetItem& DistanceFieldNormal, int32 DepthLevel, int32 DestLevelDownsampleFactor, bool bFinalInterpolation)
 {
 	check(!(bFinalInterpolation && DepthLevel != 0));
 
@@ -2663,7 +2660,7 @@ void RenderIrradianceCacheInterpolation(FRHICommandList& RHICmdList, TArray<FVie
 				for (int32 SplatSourceDepthLevel = MaxAllowedLevel; SplatSourceDepthLevel >= DepthLevel; SplatSourceDepthLevel--)
 				{
 					static FGlobalBoundShaderState BoundShaderState;
-					RHICmdList.CheckIsNull(); // need new approach for "static FGlobalBoundShaderState" for parallel rendering
+					
 					SetGlobalBoundShaderState(RHICmdList, BoundShaderState, GScreenVertexDeclaration.VertexDeclarationRHI, *VertexShader, *PixelShader);
 
 					VertexShader->SetParameters(RHICmdList, View, SplatSourceDepthLevel, DestLevelDownsampleFactor, NormalizedOffsetToPixelCenter);
@@ -2680,7 +2677,7 @@ void RenderIrradianceCacheInterpolation(FRHICommandList& RHICmdList, TArray<FVie
 				for (int32 SplatSourceDepthLevel = MaxAllowedLevel; SplatSourceDepthLevel >= DepthLevel; SplatSourceDepthLevel--)
 				{
 					static FGlobalBoundShaderState BoundShaderState;
-					RHICmdList.CheckIsNull(); // need new approach for "static FGlobalBoundShaderState" for parallel rendering
+					
 					SetGlobalBoundShaderState(RHICmdList, BoundShaderState, GScreenVertexDeclaration.VertexDeclarationRHI, *VertexShader, *PixelShader);
 
 					VertexShader->SetParameters(RHICmdList, View, SplatSourceDepthLevel, DestLevelDownsampleFactor, NormalizedOffsetToPixelCenter);
@@ -2817,7 +2814,7 @@ void UpdateVisibleObjectBuffers(const FScene* Scene, const FViewInfo& View, int3
 	}
 }
 
-bool FDeferredShadingSceneRenderer::RenderDistanceFieldAOSurfaceCache(FRHICommandList& RHICmdList, FSceneRenderTargetItem& OutBentNormalAO, bool bApplyToSceneColor)
+bool FDeferredShadingSceneRenderer::RenderDistanceFieldAOSurfaceCache(FRHICommandListImmediate& RHICmdList, FSceneRenderTargetItem& OutBentNormalAO, bool bApplyToSceneColor)
 {
 	//@todo - support multiple views
 	const FViewInfo& View = Views[0];
@@ -2832,8 +2829,6 @@ bool FDeferredShadingSceneRenderer::RenderDistanceFieldAOSurfaceCache(FRHIComman
 		QUICK_SCOPE_CYCLE_COUNTER(STAT_RenderDistanceFieldAOSurfaceCache);
 		SCOPED_DRAW_EVENT(DistanceFieldAO, DEC_SCENE_ITEMS);
 
-		RHICmdList.CheckIsNull(); // direct allocation of things will not work, probably
-
 		// Update the global distance field atlas
 		GDistanceFieldVolumeTextureAtlas.UpdateAllocations();
 
@@ -3019,11 +3014,11 @@ bool FDeferredShadingSceneRenderer::RenderDistanceFieldAOSurfaceCache(FRHIComman
 
 				if (bApplyToSceneColor)
 				{
-					GSceneRenderTargets.BeginRenderingSceneColor(false);
+					GSceneRenderTargets.BeginRenderingSceneColor(RHICmdList, false);
 				}
 				else
 				{
-					RHISetRenderTarget(OutBentNormalAO.TargetableTexture, NULL);
+					SetRenderTarget(RHICmdList, OutBentNormalAO.TargetableTexture, NULL);
 				}
 
 				// Upsample to full resolution, write to output
@@ -3101,7 +3096,7 @@ private:
 IMPLEMENT_SHADER_TYPE(template<>,TDynamicSkyLightDiffusePS<true>,TEXT("SkyLighting"),TEXT("SkyLightDiffusePS"),SF_Pixel);
 IMPLEMENT_SHADER_TYPE(template<>,TDynamicSkyLightDiffusePS<false>,TEXT("SkyLighting"),TEXT("SkyLightDiffusePS"),SF_Pixel);
 
-void FDeferredShadingSceneRenderer::RenderDynamicSkyLighting(FRHICommandList& RHICmdList)
+void FDeferredShadingSceneRenderer::RenderDynamicSkyLighting(FRHICommandListImmediate& RHICmdList)
 {
 	if (Scene->SkyLight
 		&& Scene->SkyLight->ProcessedTexture
@@ -3129,8 +3124,7 @@ void FDeferredShadingSceneRenderer::RenderDynamicSkyLighting(FRHICommandList& RH
 			bApplyShadowing = RenderDistanceFieldAOSurfaceCache(RHICmdList, LightAccumulation->GetRenderTargetItem(), false);
 		}
 
-		RHICmdList.CheckIsNull(); // BRSC is not supported in parallel
-		GSceneRenderTargets.BeginRenderingSceneColor();
+		GSceneRenderTargets.BeginRenderingSceneColor(RHICmdList);
 
 		for( int32 ViewIndex = 0; ViewIndex < Views.Num(); ViewIndex++ )
 		{
@@ -3148,7 +3142,7 @@ void FDeferredShadingSceneRenderer::RenderDynamicSkyLighting(FRHICommandList& RH
 				TShaderMapRef<TDynamicSkyLightDiffusePS<true> > PixelShader( GetGlobalShaderMap() );
 
 				static FGlobalBoundShaderState BoundShaderState;
-				RHICmdList.CheckIsNull(); // need new approach for "static FGlobalBoundShaderState" for parallel rendering
+				
 				SetGlobalBoundShaderState(RHICmdList, BoundShaderState, GFilterVertexDeclaration.VertexDeclarationRHI, *VertexShader, *PixelShader);
 
 				PixelShader->SetParameters(RHICmdList, View, LightAccumulation->GetRenderTargetItem().ShaderResourceTexture);
@@ -3158,7 +3152,7 @@ void FDeferredShadingSceneRenderer::RenderDynamicSkyLighting(FRHICommandList& RH
 				TShaderMapRef<TDynamicSkyLightDiffusePS<false> > PixelShader( GetGlobalShaderMap() );
 
 				static FGlobalBoundShaderState BoundShaderState;
-				RHICmdList.CheckIsNull(); // need new approach for "static FGlobalBoundShaderState" for parallel rendering
+				
 				SetGlobalBoundShaderState(RHICmdList, BoundShaderState, GFilterVertexDeclaration.VertexDeclarationRHI, *VertexShader, *PixelShader);
 
 				PixelShader->SetParameters(RHICmdList, View, LightAccumulation->GetRenderTargetItem().ShaderResourceTexture);
diff --git a/Engine/Source/Runtime/Renderer/Private/DistanceFieldSurfaceCacheLighting.h b/Engine/Source/Runtime/Renderer/Private/DistanceFieldSurfaceCacheLighting.h
--- a/Engine/Source/Runtime/Renderer/Private/DistanceFieldSurfaceCacheLighting.h
+++ b/Engine/Source/Runtime/Renderer/Private/DistanceFieldSurfaceCacheLighting.h
@@ -14,7 +14,7 @@ static const int32 GMaxNumObjectsPerTile = 512;
 class FRefinementLevelResources
 {
 public:
-	void InitDynamicRHI() 
+	void InitDynamicRHI()
 	{
 		//@todo - handle max exceeded
 		PositionAndRadius.Initialize(sizeof(float) * 4, GMaxIrradianceCacheSamples, PF_A32B32G32R32F, BUF_Static);
@@ -23,7 +23,7 @@ public:
 		BentNormal.Initialize(sizeof(FFloat16Color), GMaxIrradianceCacheSamples, PF_FloatRGBA, BUF_Static);
 		ScatterDrawParameters.Initialize(sizeof(uint32), 4, PF_R32_UINT, BUF_Static | BUF_DrawIndirect);
 		SavedStartIndex.Initialize(sizeof(uint32), 1, PF_R32_UINT, BUF_Static);
-		TileCoordinate.Initialize(sizeof(uint16) * 2, GMaxIrradianceCacheSamples, PF_R16G16_UINT, BUF_Static);
+		TileCoordinate.Initialize(sizeof(uint16)* 2, GMaxIrradianceCacheSamples, PF_R16G16_UINT, BUF_Static);
 	}
 
 	void ReleaseDynamicRHI()
@@ -71,7 +71,7 @@ public:
 		delete TempResources;
 	}
 
-	virtual void InitDynamicRHI() 
+	virtual void InitDynamicRHI() override
 	{
 		DispatchParameters.Initialize(sizeof(uint32), 3, PF_R32_UINT, BUF_Static | BUF_DrawIndirect);
 
@@ -85,7 +85,7 @@ public:
 		bClearedResources = false;
 	}
 
-	virtual void ReleaseDynamicRHI()
+	virtual void ReleaseDynamicRHI() override
 	{
 		DispatchParameters.Release();
 
@@ -111,20 +111,20 @@ public:
 class FTileIntersectionResources : public FRenderResource
 {
 public:
-	virtual void InitDynamicRHI() 
+	virtual void InitDynamicRHI() override
 	{
-		TileConeAxisAndCos.Initialize(sizeof(float) * 4, TileDimensions.X * TileDimensions.Y, PF_A32B32G32R32F, BUF_Static);
-		TileConeDepthRanges.Initialize(sizeof(float) * 4, TileDimensions.X * TileDimensions.Y, PF_A32B32G32R32F, BUF_Static);
+		TileConeAxisAndCos.Initialize(sizeof(float)* 4, TileDimensions.X * TileDimensions.Y, PF_A32B32G32R32F, BUF_Static);
+		TileConeDepthRanges.Initialize(sizeof(float)* 4, TileDimensions.X * TileDimensions.Y, PF_A32B32G32R32F, BUF_Static);
 
 		TileHeadDataUnpacked.Initialize(sizeof(uint32), TileDimensions.X * TileDimensions.Y * 4, PF_R32_UINT, BUF_Static);
-		TileHeadData.Initialize(sizeof(uint32) * 4, TileDimensions.X * TileDimensions.Y, PF_R32G32B32A32_UINT, BUF_Static);
+		TileHeadData.Initialize(sizeof(uint32)* 4, TileDimensions.X * TileDimensions.Y, PF_R32G32B32A32_UINT, BUF_Static);
 
 		//@todo - handle max exceeded
 		TileArrayData.Initialize(sizeof(uint16), GMaxNumObjectsPerTile * TileDimensions.X * TileDimensions.Y * 3, PF_R16_UINT, BUF_Static);
 		TileArrayNextAllocation.Initialize(sizeof(uint32), 1, PF_R32_UINT, BUF_Static);
 	}
 
-	virtual void ReleaseDynamicRHI()
+	virtual void ReleaseDynamicRHI() override
 	{
 		TileConeAxisAndCos.Release();
 		TileConeDepthRanges.Release();
diff --git a/Engine/Source/Runtime/Renderer/Private/DistortionRendering.cpp b/Engine/Source/Runtime/Renderer/Private/DistortionRendering.cpp
--- a/Engine/Source/Runtime/Renderer/Private/DistortionRendering.cpp
+++ b/Engine/Source/Runtime/Renderer/Private/DistortionRendering.cpp
@@ -500,7 +500,6 @@ void TDistortionMeshDrawingPolicy<DistortMeshPolicy>::SetMeshRenderState(
 	const ElementDataType& ElementData
 	) const
 {
-	RHICmdList.CheckIsNull();
 
 	const FMeshBatchElement& BatchElement = Mesh.Elements[BatchElementIndex];
 
@@ -686,13 +685,10 @@ bool TDistortionMeshDrawingPolicyFactory<DistortMeshPolicy>::IsMaterialIgnored(c
 	FDistortionPrimSet
 -----------------------------------------------------------------------------*/
 
-bool FDistortionPrimSet::DrawAccumulatedOffsets(const FViewInfo* ViewInfo,bool bInitializeOffsets)
+bool FDistortionPrimSet::DrawAccumulatedOffsets(FRHICommandListImmediate& RHICmdList, const FViewInfo* ViewInfo, bool bInitializeOffsets)
 {
 	bool bDirty=false;
 
-	//@todo-rco: RHIPacketList
-	FRHICommandList& RHICmdList = FRHICommandList::GetNullRef();
-
 	// Draw the view's elements with the translucent drawing policy.
 	bDirty |= DrawViewElements<TDistortionMeshDrawingPolicyFactory<FDistortMeshAccumulatePolicy> >(
 		*ViewInfo,
@@ -767,7 +763,7 @@ void FDistortionPrimSet::AddScenePrimitive(FPrimitiveSceneProxy* PrimitiveSceneP
 /** 
  * Renders the scene's distortion 
  */
-void FSceneRenderer::RenderDistortion(FRHICommandList& RHICmdList)
+void FSceneRenderer::RenderDistortion(FRHICommandListImmediate& RHICmdList)
 {
 	SCOPED_DRAW_EVENT(Distortion, DEC_SCENE_ITEMS);
 
@@ -828,10 +824,8 @@ void FSceneRenderer::RenderDistortion(FRHICommandList& RHICmdList)
 				// additive blending of offsets (or complexity if the shader complexity viewmode is enabled)
 				RHICmdList.SetBlendState(TStaticBlendState<CW_RGBA, BO_Add, BF_One, BF_One, BO_Add, BF_One, BF_One>::GetRHI());
 
-				RHICmdList.CheckIsNull(); // DrawAccumulatedOffsets
-
 				// draw only distortion meshes to accumulate their offsets
-				bDirty |= View.DistortionPrimSet.DrawAccumulatedOffsets(&View,false);
+				bDirty |= View.DistortionPrimSet.DrawAccumulatedOffsets(RHICmdList, &View, false);
 			}
 
 			if (bDirty)
@@ -852,9 +846,7 @@ void FSceneRenderer::RenderDistortion(FRHICommandList& RHICmdList)
 	{
 		SCOPED_DRAW_EVENT(DistortionApply, DEC_SCENE_ITEMS);
 
-		RHICmdList.CheckIsNull(); // GSceneRenderTargets.ResolveSceneColor
-
-		GSceneRenderTargets.ResolveSceneColor(FResolveRect(0, 0, ViewFamily.FamilySizeX, ViewFamily.FamilySizeY));
+		GSceneRenderTargets.ResolveSceneColor(RHICmdList, FResolveRect(0, 0, ViewFamily.FamilySizeX, ViewFamily.FamilySizeY));
 
 // OCULUS BEGIN: select ONE render target for all views (eyes)
 		TRefCountPtr<IPooledRenderTarget> NewSceneColor;
@@ -886,7 +878,7 @@ void FSceneRenderer::RenderDistortion(FRHICommandList& RHICmdList)
 				SetRenderTarget(RHICmdList, DestRenderTarget.TargetableTexture, FTextureRHIParamRef());
 
 				// useful when we move this into the compositing graph
-				FRenderingCompositePassContext Context(View);
+				FRenderingCompositePassContext Context(RHICmdList, View);
 
 				// Set the view family's render target/viewport.
 				Context.SetViewportAndCallRHI(View.ViewRect);
@@ -895,7 +887,7 @@ void FSceneRenderer::RenderDistortion(FRHICommandList& RHICmdList)
 				TShaderMapRef<FDistortionApplyScreenPS> PixelShader(GetGlobalShaderMap());
 
 				static FGlobalBoundShaderState BoundShaderState;
-				RHICmdList.CheckIsNull(); // need new approach for "static FGlobalBoundShaderState" for parallel rendering
+				
 				SetGlobalBoundShaderState(RHICmdList, BoundShaderState, GFilterVertexDeclaration.VertexDeclarationRHI, *VertexShader, *PixelShader);
 
 				VertexShader->SetParameters(Context);
@@ -924,13 +916,12 @@ void FSceneRenderer::RenderDistortion(FRHICommandList& RHICmdList)
 		
 		}
 
-		RHICmdList.CheckIsNull(); // GSceneRenderTargets ops
 		// OCULUS BEGIN
 		GSceneRenderTargets.SetSceneColor(NewSceneColor);
 		check(GSceneRenderTargets.GetSceneColor());
 // OCULUS END				
 		// Distortions RT is no longer needed, buffer can be reused by the pool, see BeginRenderingDistortionAccumulation() call above
-		GSceneRenderTargets.FinishRenderingSceneColor(false);
+		GSceneRenderTargets.FinishRenderingSceneColor(RHICmdList, false);
 	}
 }
 
diff --git a/Engine/Source/Runtime/Renderer/Private/DistortionRendering.h b/Engine/Source/Runtime/Renderer/Private/DistortionRendering.h
--- a/Engine/Source/Runtime/Renderer/Private/DistortionRendering.h
+++ b/Engine/Source/Runtime/Renderer/Private/DistortionRendering.h
@@ -19,7 +19,7 @@ public:
 	* @param DPGIndex - current DPG used to draw items
 	* @return true if anything was drawn
 	*/
-	bool DrawAccumulatedOffsets(const class FViewInfo* ViewInfo,bool bInitializeOffsets);
+	bool DrawAccumulatedOffsets(FRHICommandListImmediate& RHICmdList, const class FViewInfo* ViewInfo, bool bInitializeOffsets);
 
 	/**
 	* Add a new primitive to the list of distortion prims
diff --git a/Engine/Source/Runtime/Renderer/Private/DrawingPolicy.cpp b/Engine/Source/Runtime/Renderer/Private/DrawingPolicy.cpp
--- a/Engine/Source/Runtime/Renderer/Private/DrawingPolicy.cpp
+++ b/Engine/Source/Runtime/Renderer/Private/DrawingPolicy.cpp
@@ -71,10 +71,14 @@ void FMeshDrawingPolicy::DrawMesh(FRHICommandList& RHICmdList, const FMeshBatch&
 	{
 		check(Mesh.DynamicVertexData);
 
+		//hack here to get the code checked in, but this needs to be fixed. It was not sensible to wind this all the way out of the callstack
+		FRHICommandListImmediate& RHICmdListIm = FRHICommandListExecutor::GetImmediateCommandList();
+		check(&RHICmdList == &RHICmdListIm);
+
 		if (BatchElement.DynamicIndexData)
 		{
 			DrawIndexedPrimitiveUP(
-				RHICmdList,
+				RHICmdListIm,
 				Mesh.Type,
 				BatchElement.MinVertexIndex,
 				BatchElement.MaxVertexIndex - BatchElement.MinVertexIndex + 1,
@@ -87,10 +91,8 @@ void FMeshDrawingPolicy::DrawMesh(FRHICommandList& RHICmdList, const FMeshBatch&
 		}
 		else
 		{
-			RHICmdList.CheckIsNull();
-
 			DrawPrimitiveUP(
-				RHICmdList,
+				RHICmdListIm,
 				Mesh.Type,
 				BatchElement.NumPrimitives,
 				Mesh.DynamicVertexData,
diff --git a/Engine/Source/Runtime/Renderer/Private/DynamicPrimitiveDrawing.inl b/Engine/Source/Runtime/Renderer/Private/DynamicPrimitiveDrawing.inl
--- a/Engine/Source/Runtime/Renderer/Private/DynamicPrimitiveDrawing.inl
+++ b/Engine/Source/Runtime/Renderer/Private/DynamicPrimitiveDrawing.inl
@@ -22,8 +22,8 @@ TDynamicPrimitiveDrawer<DrawingPolicyFactoryType>::~TDynamicPrimitiveDrawer()
 		{
 			DepthTexture = GSceneRenderTargets.GetSceneDepthTexture();
 		}
-		//@todo-rco: RHIPacketList
-		FRHICommandList& RHICmdList = FRHICommandList::GetNullRef();
+		
+		FRHICommandListImmediate& RHICmdList = FRHICommandListExecutor::GetImmediateCommandList();
 
 		// Draw the batched elements.
 		BatchedElements.Draw(
@@ -107,8 +107,8 @@ int32 TDynamicPrimitiveDrawer<DrawingPolicyFactoryType>::DrawMesh(const FMeshBat
 #if DO_CHECK
 	Mesh.CheckUniformBuffers();
 #endif
-	//@todo-rco: RHIPacketList
-	FRHICommandList& RHICmdList = FRHICommandList::GetNullRef();
+	
+	FRHICommandListImmediate& RHICmdList = FRHICommandListExecutor::GetImmediateCommandList();
 	INC_DWORD_STAT_BY(STAT_DynamicPathMeshDrawCalls,Mesh.Elements.Num());
 	const bool DrawDirty = DrawingPolicyFactoryType::DrawDynamicMesh(
 		RHICmdList,
@@ -223,8 +223,8 @@ bool DrawViewElements(
 	bool bPreFog
 	)
 {
-	//@todo-rco: RHIPacketList
-	FRHICommandList& RHICmdList = FRHICommandList::GetNullRef();
+	
+	FRHICommandListImmediate& RHICmdList = FRHICommandListExecutor::GetImmediateCommandList();
 
 	// Get the correct element list based on dpg index
 	const TIndirectArray<FHitProxyMeshPair>& ViewMeshElementList = ( DPGIndex == SDPG_Foreground ? View.TopViewMeshElements : View.ViewMeshElements );
diff --git a/Engine/Source/Runtime/Renderer/Private/FogRendering.cpp b/Engine/Source/Runtime/Renderer/Private/FogRendering.cpp
--- a/Engine/Source/Runtime/Renderer/Private/FogRendering.cpp
+++ b/Engine/Source/Runtime/Renderer/Private/FogRendering.cpp
@@ -200,14 +200,14 @@ public:
 	// Destructor
 	virtual ~FFogVertexDeclaration() {}
 
-	virtual void InitRHI()
+	virtual void InitRHI() override
 	{
 		FVertexDeclarationElementList Elements;
 		Elements.Add(FVertexElement(0, 0, VET_Float2, 0, sizeof(FVector2D)));
 		VertexDeclarationRHI = RHICreateVertexDeclaration(Elements);
 	}
 
-	virtual void ReleaseRHI()
+	virtual void ReleaseRHI() override
 	{
 		VertexDeclarationRHI.SafeRelease();
 	}
@@ -275,7 +275,7 @@ void FSceneRenderer::InitFogConstants()
 FGlobalBoundShaderState ExponentialBoundShaderState;
 
 /** Sets the bound shader state for either the per-pixel or per-sample fog pass. */
-void SetFogShaders(FRHICommandList& RHICmdList, FScene* Scene,const FViewInfo& View,FLightShaftsOutput LightShaftsOutput)
+void SetFogShaders(FRHICommandListImmediate& RHICmdList, FScene* Scene, const FViewInfo& View, FLightShaftsOutput LightShaftsOutput)
 {
 	if (Scene->ExponentialFogs.Num() > 0)
 	{
@@ -288,7 +288,7 @@ void SetFogShaders(FRHICommandList& RHICmdList, FScene* Scene,const FViewInfo& V
 	}
 }
 
-bool FDeferredShadingSceneRenderer::RenderFog(FRHICommandList& RHICmdList, FLightShaftsOutput LightShaftsOutput)
+bool FDeferredShadingSceneRenderer::RenderFog(FRHICommandListImmediate& RHICmdList, FLightShaftsOutput LightShaftsOutput)
 {
 	if (Scene->ExponentialFogs.Num() > 0)
 	{
@@ -307,9 +307,7 @@ bool FDeferredShadingSceneRenderer::RenderFog(FRHICommandList& RHICmdList, FLigh
 			0, 2, 3
 		};
 
-		RHICmdList.CheckIsNull(); // GSceneRenderTargets ops
-
-		GSceneRenderTargets.BeginRenderingSceneColor();
+		GSceneRenderTargets.BeginRenderingSceneColor(RHICmdList);
 		for(int32 ViewIndex = 0;ViewIndex < Views.Num();ViewIndex++)
 		{
 			const FViewInfo& View = Views[ViewIndex];
@@ -346,7 +344,7 @@ bool FDeferredShadingSceneRenderer::RenderFog(FRHICommandList& RHICmdList, FLigh
 		}
 
 		//no need to resolve since we used alpha blending
-		GSceneRenderTargets.FinishRenderingSceneColor(false);
+		GSceneRenderTargets.FinishRenderingSceneColor(RHICmdList, false);
 		return true;
 	}
 
diff --git a/Engine/Source/Runtime/Renderer/Private/ForwardBasePassRendering.cpp b/Engine/Source/Runtime/Renderer/Private/ForwardBasePassRendering.cpp
--- a/Engine/Source/Runtime/Renderer/Private/ForwardBasePassRendering.cpp
+++ b/Engine/Source/Runtime/Renderer/Private/ForwardBasePassRendering.cpp
@@ -150,14 +150,12 @@ public:
 		const typename LightMapPolicyType::ElementDataType& LightMapElementData
 		) const
 	{
-		RHICmdList.CheckIsNull();
-
 #if !(UE_BUILD_SHIPPING || UE_BUILD_TEST)
 		// Treat masked materials as if they don't occlude in shader complexity, which is PVR behavior
 		if(Parameters.BlendMode == BLEND_Masked && View.Family->EngineShowFlags.ShaderComplexity)
 		{
 			// Note, this is a reversed Z depth surface, using CF_GreaterEqual.
-			RHISetDepthStencilState(TStaticDepthStencilState<false,CF_GreaterEqual>::GetRHI());
+			RHICmdList.SetDepthStencilState(TStaticDepthStencilState<false,CF_GreaterEqual>::GetRHI());
 		}
 #endif
 
@@ -197,7 +195,7 @@ public:
 		if(Parameters.BlendMode == BLEND_Masked && View.Family->EngineShowFlags.ShaderComplexity)
 		{
 			// Note, this is a reversed Z depth surface, using CF_GreaterEqual.
-			RHISetDepthStencilState(TStaticDepthStencilState<true,CF_GreaterEqual>::GetRHI());
+			RHICmdList.SetDepthStencilState(TStaticDepthStencilState<true,CF_GreaterEqual>::GetRHI());
 		}
 #endif
 	}
@@ -292,7 +290,7 @@ EBasePassSort::Type GetSortMode()
 	}
 }
 
-void FForwardShadingSceneRenderer::RenderForwardShadingBasePass(FRHICommandList& RHICmdList)
+void FForwardShadingSceneRenderer::RenderForwardShadingBasePass(FRHICommandListImmediate& RHICmdList)
 {
 	SCOPED_DRAW_EVENT(BasePass, DEC_SCENE_ITEMS);
 	SCOPE_CYCLE_COUNTER(STAT_BasePassDrawTime);
@@ -324,46 +322,35 @@ void FForwardShadingSceneRenderer::RenderForwardShadingBasePass(FRHICommandList&
 		FViewInfo& View = Views[ViewIndex];
 
 		// Opaque blending
-		RHISetBlendState(TStaticBlendStateWriteMask<CW_RGBA>::GetRHI());
+		RHICmdList.SetBlendState(TStaticBlendStateWriteMask<CW_RGBA>::GetRHI());
 		// Note, this is a reversed Z depth surface, using CF_GreaterEqual.
-		RHISetDepthStencilState(TStaticDepthStencilState<true, CF_GreaterEqual>::GetRHI());
-		RHISetViewport(View.ViewRect.Min.X, View.ViewRect.Min.Y, 0, View.ViewRect.Max.X, View.ViewRect.Max.Y, 1);
+		RHICmdList.SetDepthStencilState(TStaticDepthStencilState<true, CF_GreaterEqual>::GetRHI());
+		RHICmdList.SetViewport(View.ViewRect.Min.X, View.ViewRect.Min.Y, 0, View.ViewRect.Max.X, View.ViewRect.Max.Y, 1);
 
 		{
 			SCOPE_CYCLE_COUNTER(STAT_RHICounterTEMP);
-			FRHICommandList* UseRHICmdList = &RHICmdList;
-			if (!&RHICmdList)
-			{
-				//@todo-rco: Very temp code!!!
-				static IConsoleVariable* RHICmdListCVar = IConsoleManager::Get().FindConsoleVariable(TEXT("r.RHICmd"));
-				bool bUseRHICmdList = RHICmdListCVar->GetInt() != 0;
-				UseRHICmdList = bUseRHICmdList ? &GRHICommandList.CreateList() : &FRHICommandList::GetNullRef();
-			}
+			//@todo-rco: Very temp code!!!
+			static IConsoleVariable* RHICmdListCVar = IConsoleManager::Get().FindConsoleVariable(TEXT("r.RHICmd"));
+			bool bUseRHICmdList = RHICmdListCVar->GetInt() != 0;
+			FRHICommandList LocalRHICmdList;
+			FRHICommandList& UseRHICmdList = bUseRHICmdList ? LocalRHICmdList : RHICmdList;
 
 			// Render the base pass static data
 			if (SortMode == EBasePassSort::SortPerMesh)
 			{
 				SCOPE_CYCLE_COUNTER(STAT_StaticDrawListDrawTime);
-				MaxDraws -= Scene->BasePassForForwardShadingLowQualityLightMapDrawList[FScene::EBasePass_Default].DrawVisibleFrontToBack(*UseRHICmdList, View, View.StaticMeshVisibilityMap, View.StaticMeshBatchVisibility, MaxDraws);
-				MaxDraws -= Scene->BasePassForForwardShadingDistanceFieldShadowMapLightMapDrawList[FScene::EBasePass_Default].DrawVisibleFrontToBack(*UseRHICmdList, View, View.StaticMeshVisibilityMap, View.StaticMeshBatchVisibility, MaxDraws);
-				MaxDraws -= Scene->BasePassForForwardShadingDirectionalLightAndSHIndirectDrawList[FScene::EBasePass_Default].DrawVisibleFrontToBack(*UseRHICmdList, View, View.StaticMeshVisibilityMap, View.StaticMeshBatchVisibility, MaxDraws);
-				MaxDraws -= Scene->BasePassForForwardShadingNoLightMapDrawList[FScene::EBasePass_Default].DrawVisibleFrontToBack(*UseRHICmdList, View, View.StaticMeshVisibilityMap, View.StaticMeshBatchVisibility, MaxDraws);
+				MaxDraws -= Scene->BasePassForForwardShadingLowQualityLightMapDrawList[FScene::EBasePass_Default].DrawVisibleFrontToBack(UseRHICmdList, View, View.StaticMeshVisibilityMap, View.StaticMeshBatchVisibility, MaxDraws);
+				MaxDraws -= Scene->BasePassForForwardShadingDistanceFieldShadowMapLightMapDrawList[FScene::EBasePass_Default].DrawVisibleFrontToBack(UseRHICmdList, View, View.StaticMeshVisibilityMap, View.StaticMeshBatchVisibility, MaxDraws);
+				MaxDraws -= Scene->BasePassForForwardShadingDirectionalLightAndSHIndirectDrawList[FScene::EBasePass_Default].DrawVisibleFrontToBack(UseRHICmdList, View, View.StaticMeshVisibilityMap, View.StaticMeshBatchVisibility, MaxDraws);
+				MaxDraws -= Scene->BasePassForForwardShadingNoLightMapDrawList[FScene::EBasePass_Default].DrawVisibleFrontToBack(UseRHICmdList, View, View.StaticMeshVisibilityMap, View.StaticMeshBatchVisibility, MaxDraws);
 			}
 			else
 			{
 				SCOPE_CYCLE_COUNTER(STAT_StaticDrawListDrawTime);
-				Scene->BasePassForForwardShadingLowQualityLightMapDrawList[FScene::EBasePass_Default].DrawVisible(*UseRHICmdList, View, View.StaticMeshVisibilityMap, View.StaticMeshBatchVisibility);
-				Scene->BasePassForForwardShadingDistanceFieldShadowMapLightMapDrawList[FScene::EBasePass_Default].DrawVisible(*UseRHICmdList, View, View.StaticMeshVisibilityMap, View.StaticMeshBatchVisibility);
-				Scene->BasePassForForwardShadingDirectionalLightAndSHIndirectDrawList[FScene::EBasePass_Default].DrawVisible(*UseRHICmdList, View, View.StaticMeshVisibilityMap, View.StaticMeshBatchVisibility);
-				Scene->BasePassForForwardShadingNoLightMapDrawList[FScene::EBasePass_Default].DrawVisible(*UseRHICmdList, View, View.StaticMeshVisibilityMap, View.StaticMeshBatchVisibility);
-			}
-
-			if (UseRHICmdList != &RHICmdList)
-			{
-				if (UseRHICmdList)
-				{
-					GRHICommandList.ExecuteAndFreeList(*UseRHICmdList);
-				}
+				Scene->BasePassForForwardShadingLowQualityLightMapDrawList[FScene::EBasePass_Default].DrawVisible(UseRHICmdList, View, View.StaticMeshVisibilityMap, View.StaticMeshBatchVisibility);
+				Scene->BasePassForForwardShadingDistanceFieldShadowMapLightMapDrawList[FScene::EBasePass_Default].DrawVisible(UseRHICmdList, View, View.StaticMeshVisibilityMap, View.StaticMeshBatchVisibility);
+				Scene->BasePassForForwardShadingDirectionalLightAndSHIndirectDrawList[FScene::EBasePass_Default].DrawVisible(UseRHICmdList, View, View.StaticMeshVisibilityMap, View.StaticMeshBatchVisibility);
+				Scene->BasePassForForwardShadingNoLightMapDrawList[FScene::EBasePass_Default].DrawVisible(UseRHICmdList, View, View.StaticMeshVisibilityMap, View.StaticMeshBatchVisibility);
 			}
 		}
 
@@ -400,9 +387,7 @@ void FForwardShadingSceneRenderer::RenderForwardShadingBasePass(FRHICommandList&
 
 			// Draw the base pass for the view's batched mesh elements.
 			DrawViewElements<FBasePassForwardOpaqueDrawingPolicyFactory>(View,FBasePassForwardOpaqueDrawingPolicyFactory::ContextType(ESceneRenderTargetsMode::DontSet), SDPG_World, true);
-			//@todo-rco: RHIPacketList
-			FRHICommandList& RHICmdList = FRHICommandList::GetNullRef();
-
+			
 			// Draw the view's batched simple elements(lines, sprites, etc).
 			View.BatchedViewElements.Draw(RHICmdList, bNeedToSwitchVerticalAxis, View.ViewProjectionMatrix, View.ViewRect.Width(), View.ViewRect.Height(), false);
 
diff --git a/Engine/Source/Runtime/Renderer/Private/ForwardBasePassRendering.h b/Engine/Source/Runtime/Renderer/Private/ForwardBasePassRendering.h
--- a/Engine/Source/Runtime/Renderer/Private/ForwardBasePassRendering.h
+++ b/Engine/Source/Runtime/Renderer/Private/ForwardBasePassRendering.h
@@ -382,7 +382,7 @@ public:
 			if (BlendMode != BLEND_Opaque)
 			{
 				// Add complexity to existing, keep alpha
-				RHISetBlendState(TStaticBlendState<CW_RGB,BO_Add,BF_One,BF_One>::GetRHI());
+				RHICmdList.SetBlendState(TStaticBlendState<CW_RGB,BO_Add,BF_One,BF_One>::GetRHI());
 			}
 
 			TShaderMapRef<FShaderComplexityAccumulatePS> ShaderComplexityPixelShader(GetGlobalShaderMap());
@@ -471,8 +471,8 @@ void ProcessBasePassMeshForForwardShading(
 
 	check(!AllowHighQualityLightmaps());
 
-	//@todo-rco: RHIPacketList
-	FRHICommandList& RHICmdList = FRHICommandList::GetNullRef();
+	
+	FRHICommandListImmediate& RHICmdList = FRHICommandListExecutor::GetImmediateCommandList();
 
 	if (LightMapInteraction.GetType() == LMIT_Texture)
 	{
diff --git a/Engine/Source/Runtime/Renderer/Private/ForwardShadingRenderer.cpp b/Engine/Source/Runtime/Renderer/Private/ForwardShadingRenderer.cpp
--- a/Engine/Source/Runtime/Renderer/Private/ForwardShadingRenderer.cpp
+++ b/Engine/Source/Runtime/Renderer/Private/ForwardShadingRenderer.cpp
@@ -22,14 +22,14 @@ FForwardShadingSceneRenderer::FForwardShadingSceneRenderer(const FSceneViewFamil
  * Initialize scene's views.
  * Check visibility, sort translucent items, etc.
  */
-void FForwardShadingSceneRenderer::InitViews()
+void FForwardShadingSceneRenderer::InitViews(FRHICommandListImmediate& RHICmdList)
 {
 	SCOPED_DRAW_EVENT(InitViews, DEC_SCENE_ITEMS);
 
 	SCOPE_CYCLE_COUNTER(STAT_InitViewsTime);
 
-	PreVisibilityFrameSetup();
-	ComputeViewVisibility();
+	PreVisibilityFrameSetup(RHICmdList);
+	ComputeViewVisibility(RHICmdList);
 	PostVisibilityFrameSetup();
 
 	OnStartFrame();
@@ -38,7 +38,7 @@ void FForwardShadingSceneRenderer::InitViews()
 /** 
 * Renders the view family. 
 */
-void FForwardShadingSceneRenderer::Render()
+void FForwardShadingSceneRenderer::Render(FRHICommandListImmediate& RHICmdList)
 {
 	if(!ViewFamily.EngineShowFlags.Rendering)
 	{
@@ -48,21 +48,18 @@ void FForwardShadingSceneRenderer::Render()
 	auto FeatureLevel = ViewFamily.Scene->GetFeatureLevel();
 
 	// Initialize global system textures (pass-through if already initialized).
-	GSystemTextures.InitializeTextures(FeatureLevel);
+	GSystemTextures.InitializeTextures(RHICmdList, FeatureLevel);
 
 	// Allocate the maximum scene render target space for the current view family.
 	GSceneRenderTargets.Allocate(ViewFamily);
 
 	// Find the visible primitives.
-	InitViews();
-
-	//@todo-rco: RHIPacketList
-	FRHICommandList& RHICmdList = FRHICommandList::GetNullRef();
-
+	InitViews(RHICmdList);
+	
 	// Notify the FX system that the scene is about to be rendered.
 	if (Scene->FXSystem)
 	{
-		Scene->FXSystem->PreRender();
+		Scene->FXSystem->PreRender(RHICmdList);
 	}
 
 	GRenderTargetPool.VisualizeTexture.OnStartFrame(Views[0]);
@@ -79,7 +76,7 @@ void FForwardShadingSceneRenderer::Render()
 	else
 	{
 		// Begin rendering to scene color
-		GSceneRenderTargets.BeginRenderingSceneColor(false);
+		GSceneRenderTargets.BeginRenderingSceneColor(RHICmdList, false);
 	}
 
 	// Clear color and depth buffer
@@ -89,12 +86,12 @@ void FForwardShadingSceneRenderer::Render()
 	RenderForwardShadingBasePass(RHICmdList);
 
 	// Make a copy of the scene depth if the current hardware doesn't support reading and writing to the same depth buffer
-	GSceneRenderTargets.ResolveSceneDepthToAuxiliaryTexture();
+	GSceneRenderTargets.ResolveSceneDepthToAuxiliaryTexture(RHICmdList);
 
 	// Notify the FX system that opaque primitives have been rendered.
 	if (Scene->FXSystem)
 	{
-		Scene->FXSystem->PostRenderOpaque();
+		Scene->FXSystem->PostRenderOpaque(RHICmdList);
 	}
 
 	// Draw translucency.
@@ -107,7 +104,7 @@ void FForwardShadingSceneRenderer::Render()
 			// to apply refraction effect by distorting the scene color
 			RenderDistortion(RHICmdList);
 		}
-		RenderTranslucency();
+		RenderTranslucency(RHICmdList);
 	}
 
 	if( !bGammaSpace )
@@ -130,7 +127,7 @@ void FForwardShadingSceneRenderer::Render()
 			FIntPoint PrePostSourceViewportSize = GSceneRenderTargets.GetBufferSizeXY();
 
 			FMemMark Mark(FMemStack::Get());
-			FRenderingCompositePassContext CompositeContext(View);
+			FRenderingCompositePassContext CompositeContext(RHICmdList, View);
 
 			FRenderingCompositePass* PostProcessSunMask = CompositeContext.Graph.RegisterPass(new(FMemStack::Get()) FRCPassPostProcessSunMaskES2(PrePostSourceViewportSize, true));
 			CompositeContext.Root->AddDependency(FRenderingCompositeOutputRef(PostProcessSunMask));
@@ -138,10 +135,10 @@ void FForwardShadingSceneRenderer::Render()
 		}
 
 		// Resolve the scene color for post processing.
-		GSceneRenderTargets.ResolveSceneColor(FResolveRect(0, 0, ViewFamily.FamilySizeX, ViewFamily.FamilySizeY));
+		GSceneRenderTargets.ResolveSceneColor(RHICmdList, FResolveRect(0, 0, ViewFamily.FamilySizeX, ViewFamily.FamilySizeY));
 
 		// Drop depth and stencil before post processing to avoid export.
-		RHIDiscardRenderTargets(true, true, 0);
+		RHICmdList.DiscardRenderTargets(true, true, 0);
 
 		// Finish rendering for each view.
 		{
@@ -150,10 +147,10 @@ void FForwardShadingSceneRenderer::Render()
 			for(int32 ViewIndex = 0;ViewIndex < Views.Num();ViewIndex++)
 			{	
 				SCOPED_CONDITIONAL_DRAW_EVENTF(EventView, Views.Num() > 1, DEC_SCENE_ITEMS, TEXT("View%d"), ViewIndex);
-				GPostProcessing.ProcessES2( Views[ViewIndex], bOnChipSunMask );
+				GPostProcessing.ProcessES2(RHICmdList, Views[ViewIndex], bOnChipSunMask);
 			}
 		}
 	}
 
-	RenderFinish();
+	RenderFinish(RHICmdList);
 }
diff --git a/Engine/Source/Runtime/Renderer/Private/ForwardTranslucentRendering.cpp b/Engine/Source/Runtime/Renderer/Private/ForwardTranslucentRendering.cpp
--- a/Engine/Source/Runtime/Renderer/Private/ForwardTranslucentRendering.cpp
+++ b/Engine/Source/Runtime/Renderer/Private/ForwardTranslucentRendering.cpp
@@ -45,17 +45,16 @@ IMPLEMENT_SHADER_TYPE(,FForwardCopySceneAlphaPS,TEXT("TranslucentLightingShaders
 
 FGlobalBoundShaderState ForwardCopySceneAlphaBoundShaderState;
 
-void FForwardShadingSceneRenderer::CopySceneAlpha(FRHICommandList& RHICmdList, const FSceneView& View)
+void FForwardShadingSceneRenderer::CopySceneAlpha(FRHICommandListImmediate& RHICmdList, const FSceneView& View)
 {
 	SCOPED_DRAW_EVENTF(EventCopy, DEC_SCENE_ITEMS, TEXT("CopySceneAlpha"));
-	RHISetRasterizerState(TStaticRasterizerState<FM_Solid,CM_None>::GetRHI());
-	RHISetDepthStencilState(TStaticDepthStencilState<false,CF_Always>::GetRHI());
-	RHISetBlendState(TStaticBlendState<>::GetRHI());
+	RHICmdList.SetRasterizerState(TStaticRasterizerState<FM_Solid, CM_None>::GetRHI());
+	RHICmdList.SetDepthStencilState(TStaticDepthStencilState<false,CF_Always>::GetRHI());
+	RHICmdList.SetBlendState(TStaticBlendState<>::GetRHI());
 
-	RHICmdList.CheckIsNull(); // GSceneRenderTargets ops
-	GSceneRenderTargets.ResolveSceneColor();
+	GSceneRenderTargets.ResolveSceneColor(RHICmdList);
 
-	GSceneRenderTargets.BeginRenderingSceneAlphaCopy();
+	GSceneRenderTargets.BeginRenderingSceneAlphaCopy(RHICmdList);
 
 	int X = GSceneRenderTargets.GetBufferSizeXY().X;
 	int Y = GSceneRenderTargets.GetBufferSizeXY().Y;
@@ -79,7 +78,7 @@ void FForwardShadingSceneRenderer::CopySceneAlpha(FRHICommandList& RHICmdList, c
 		*ScreenVertexShader,
 		EDRF_UseTriangleOptimization);
 
-	GSceneRenderTargets.FinishRenderingSceneAlphaCopy();
+	GSceneRenderTargets.FinishRenderingSceneAlphaCopy(RHICmdList);
 }
 
 
@@ -243,7 +242,7 @@ bool FTranslucencyForwardShadingDrawingPolicyFactory::DrawStaticMesh(
 FTranslucentPrimSet
 -----------------------------------------------------------------------------*/
 
-void FTranslucentPrimSet::DrawPrimitivesForForwardShading(const FViewInfo& View, FSceneRenderer& Renderer) const
+void FTranslucentPrimSet::DrawPrimitivesForForwardShading(FRHICommandListImmediate& RHICmdList, const FViewInfo& View, FSceneRenderer& Renderer) const
 {
 	// Draw sorted scene prims
 	for (int32 PrimIdx = 0; PrimIdx < SortedPrims.Num(); PrimIdx++)
@@ -252,11 +251,12 @@ void FTranslucentPrimSet::DrawPrimitivesForForwardShading(const FViewInfo& View,
 		int32 PrimitiveId = PrimitiveSceneInfo->GetIndex();
 		const FPrimitiveViewRelevance& ViewRelevance = View.PrimitiveViewRelevanceMap[PrimitiveId];
 
-		RenderPrimitiveForForwardShading(View, PrimitiveSceneInfo, ViewRelevance);
+		RenderPrimitiveForForwardShading(RHICmdList, View, PrimitiveSceneInfo, ViewRelevance);
 	}
 }
 
 void FTranslucentPrimSet::RenderPrimitiveForForwardShading(
+	FRHICommandListImmediate& RHICmdList,
 	const FViewInfo& View, 
 	FPrimitiveSceneInfo* PrimitiveSceneInfo, 
 	const FPrimitiveViewRelevance& ViewRelevance) const
@@ -281,9 +281,6 @@ void FTranslucentPrimSet::RenderPrimitiveForForwardShading(
 				);
 		}
 
-		//@todo-rco: RHIPacketList
-		FRHICommandList& RHICmdList = FRHICommandList::GetNullRef();
-
 		// Render static scene prim
 		if( ViewRelevance.bStaticRelevance )
 		{
@@ -310,13 +307,10 @@ void FTranslucentPrimSet::RenderPrimitiveForForwardShading(
 	}
 }
 
-void FForwardShadingSceneRenderer::RenderTranslucency()
+void FForwardShadingSceneRenderer::RenderTranslucency(FRHICommandListImmediate& RHICmdList)
 {
 	if (ShouldRenderTranslucency())
 	{
-		//@todo-rco: RHIPacketList
-		FRHICommandList& RHICmdList = FRHICommandList::GetNullRef();
-
 		const bool bGammaSpace = !IsMobileHDR();
 		const bool bLinearHDR64 = !bGammaSpace && !IsMobileHDR32bpp();
 
@@ -339,15 +333,15 @@ void FForwardShadingSceneRenderer::RenderTranslucency()
 
 			if (!bGammaSpace)
 			{
-				GSceneRenderTargets.BeginRenderingTranslucency(View);
+				GSceneRenderTargets.BeginRenderingTranslucency(RHICmdList, View);
 			}
 
 			// Enable depth test, disable depth writes.
 			// Note, this is a reversed Z depth surface, using CF_GreaterEqual.
-			RHISetDepthStencilState(TStaticDepthStencilState<false,CF_GreaterEqual>::GetRHI());
+			RHICmdList.SetDepthStencilState(TStaticDepthStencilState<false,CF_GreaterEqual>::GetRHI());
 
 			// Draw only translucent prims that don't read from scene color
-			View.TranslucentPrimSet.DrawPrimitivesForForwardShading(View, *this);
+			View.TranslucentPrimSet.DrawPrimitivesForForwardShading(RHICmdList, View, *this);
 			// Draw the view's mesh elements with the translucent drawing policy.
 			DrawViewElements<FTranslucencyForwardShadingDrawingPolicyFactory>(View,FTranslucencyForwardShadingDrawingPolicyFactory::ContextType(),SDPG_World,false);
 			// Draw the view's mesh elements with the translucent drawing policy.
diff --git a/Engine/Source/Runtime/Renderer/Private/GPUBenchmark.cpp b/Engine/Source/Runtime/Renderer/Private/GPUBenchmark.cpp
--- a/Engine/Source/Runtime/Renderer/Private/GPUBenchmark.cpp
+++ b/Engine/Source/Runtime/Renderer/Private/GPUBenchmark.cpp
@@ -125,13 +125,13 @@ IMPLEMENT_SHADER_TYPE(,FPostProcessBenchmarkVS,TEXT("GPUBenchmark"),TEXT("MainBe
 
 
 template <uint32 Method>
-void RunBenchmarkShader(FRHICommandList& RHICmdList, const FSceneView& View, TRefCountPtr<IPooledRenderTarget>& Src, uint32 Count)
+void RunBenchmarkShader(FRHICommandListImmediate& RHICmdList, const FSceneView& View, TRefCountPtr<IPooledRenderTarget>& Src, uint32 Count)
 {
 	TShaderMapRef<FPostProcessBenchmarkVS> VertexShader(GetGlobalShaderMap());
 	TShaderMapRef<FPostProcessBenchmarkPS<Method> > PixelShader(GetGlobalShaderMap());
 
 	static FGlobalBoundShaderState BoundShaderState;
-	RHICmdList.CheckIsNull(); // need new approach for "static FGlobalBoundShaderState" for parallel rendering
+	
 
 	SetGlobalBoundShaderState(RHICmdList, BoundShaderState, GFilterVertexDeclaration.VertexDeclarationRHI, *VertexShader, *PixelShader);
 
@@ -153,7 +153,7 @@ void RunBenchmarkShader(FRHICommandList& RHICmdList, const FSceneView& View, TRe
 	}
 }
 
-void RunBenchmarkShader(FRHICommandList& RHICmdList, const FSceneView& View, uint32 MethodId, TRefCountPtr<IPooledRenderTarget>& Src, uint32 Count)
+void RunBenchmarkShader(FRHICommandListImmediate& RHICmdList, const FSceneView& View, uint32 MethodId, TRefCountPtr<IPooledRenderTarget>& Src, uint32 Count)
 {
 	SCOPED_DRAW_EVENTF(Benchmark, DEC_SCENE_ITEMS, TEXT("Benchmark Method:%d"), MethodId);
 
@@ -346,11 +346,8 @@ private:
 #endif
 
 
-void RendererGPUBenchmark(FSynthBenchmarkResults& InOut, const FSceneView& View, uint32 WorkScale, bool bDebugOut)
+void RendererGPUBenchmark(FRHICommandListImmediate& RHICmdList, FSynthBenchmarkResults& InOut, const FSceneView& View, uint32 WorkScale, bool bDebugOut)
 {
-	//@todo-rco: RHIPacketList
-	FRHICommandList& RHICmdList = FRHICommandList::GetNullRef();
-
 	check(IsInRenderingThread());
 	
 	// two RT to ping pong so we force the GPU to flush it's pipeline
@@ -368,9 +365,9 @@ void RendererGPUBenchmark(FSynthBenchmarkResults& InOut, const FSceneView& View,
 	}
 
 	// set the state
-	RHISetBlendState(TStaticBlendState<>::GetRHI());
-	RHISetRasterizerState(TStaticRasterizerState<>::GetRHI());
-	RHISetDepthStencilState(TStaticDepthStencilState<false,CF_Always>::GetRHI());
+	RHICmdList.SetBlendState(TStaticBlendState<>::GetRHI());
+	RHICmdList.SetRasterizerState(TStaticRasterizerState<>::GetRHI());
+	RHICmdList.SetDepthStencilState(TStaticDepthStencilState<false,CF_Always>::GetRHI());
 
 	{
 		// larger number means more accuracy but slower, some slower GPUs might timeout with a number to large
@@ -413,7 +410,7 @@ void RendererGPUBenchmark(FSynthBenchmarkResults& InOut, const FSceneView& View,
 		// e.g. on NV670: Method3 (mostly fill rate )-> 26GP/s (seems realistic)
 		// reference: http://en.wikipedia.org/wiki/Comparison_of_Nvidia_graphics_processing_units theoretical: 29.3G/s
 
-		RHIEndRenderQuery(TimerQueries[0]);
+		RHICmdList.EndRenderQuery(TimerQueries[0]);
 
 		// multiple iterations to see how trust able the values are
 		for(uint32 Iteration = 0; Iteration < IterationCount; ++Iteration)
@@ -431,14 +428,14 @@ void RendererGPUBenchmark(FSynthBenchmarkResults& InOut, const FSceneView& View,
 
 				GRenderTargetPool.VisualizeTexture.SetCheckPoint(RHICmdList, RTItems[DestRTIndex]);
 
-				RHISetRenderTarget(RTItems[DestRTIndex]->GetRenderTargetItem().TargetableTexture, FTextureRHIRef());	
+				SetRenderTarget(RHICmdList, RTItems[DestRTIndex]->GetRenderTargetItem().TargetableTexture, FTextureRHIRef());	
 
 				// decide how much work we do in this pass
 				PassCount[Iteration] = (Iteration / 10 + 1) * WorkScale;
 
 				RunBenchmarkShader(RHICmdList, View, MethodId, RTItems[SrcRTIndex], PassCount[Iteration]);
 
-				RHICopyToResolveTarget(RTItems[DestRTIndex]->GetRenderTargetItem().TargetableTexture, RTItems[DestRTIndex]->GetRenderTargetItem().ShaderResourceTexture, false, FResolveParams());
+				RHICmdList.CopyToResolveTarget(RTItems[DestRTIndex]->GetRenderTargetItem().TargetableTexture, RTItems[DestRTIndex]->GetRenderTargetItem().ShaderResourceTexture, false, FResolveParams());
 
 				/*if(bGPUCPUSync)
 				{
@@ -447,7 +444,7 @@ void RendererGPUBenchmark(FSynthBenchmarkResults& InOut, const FSceneView& View,
 					FResolveParams Param;
 
 					Param.Rect = FResolveRect(0, 0, 1, 1);
-					RHICopyToResolveTarget(
+					RHICmdList.CopyToResolveTarget(
 						RTItems[DestRTIndex]->GetRenderTargetItem().TargetableTexture,
 						RTItems[2]->GetRenderTargetItem().ShaderResourceTexture,
 						false,
@@ -461,7 +458,7 @@ void RendererGPUBenchmark(FSynthBenchmarkResults& InOut, const FSceneView& View,
 					RHIUnmapStagingSurface(RTItems[2]->GetRenderTargetItem().ShaderResourceTexture);
 				}*/
 
-				RHIEndRenderQuery(TimerQueries[QueryIndex]);
+				RHICmdList.EndRenderQuery(TimerQueries[QueryIndex]);
 
 				// ping pong
 				DestRTIndex = 1 - DestRTIndex;
@@ -470,8 +467,8 @@ void RendererGPUBenchmark(FSynthBenchmarkResults& InOut, const FSceneView& View,
 
 		{
 			uint64 OldAbsTime = 0;
-			RHIGetRenderQueryResult(TimerQueries[0], OldAbsTime, true);
-			GTimerQueryPool.ReleaseQuery(TimerQueries[0]);
+			RHICmdList.GetRenderQueryResult(TimerQueries[0], OldAbsTime, true);
+			GTimerQueryPool.ReleaseQuery(RHICmdList, TimerQueries[0]);
 
 #if !UE_BUILD_SHIPPING
 			FBenchmarkGraph BenchmarkGraph(IterationCount, IterationCount, *(FPaths::ScreenShotDir() + TEXT("GPUSynthBenchmarkGraph.bmp")));
@@ -486,8 +483,8 @@ void RendererGPUBenchmark(FSynthBenchmarkResults& InOut, const FSceneView& View,
 					uint32 QueryIndex = 1 + Iteration * MethodCount + MethodId;
 
 					uint64 AbsTime;
-					RHIGetRenderQueryResult(TimerQueries[QueryIndex], AbsTime, true);
-					GTimerQueryPool.ReleaseQuery(TimerQueries[QueryIndex]);
+					RHICmdList.GetRenderQueryResult(TimerQueries[QueryIndex], AbsTime, true);
+					GTimerQueryPool.ReleaseQuery(RHICmdList, TimerQueries[QueryIndex]);
 
 					Results[MethodId] = AbsTime - OldAbsTime;
 					OldAbsTime = AbsTime;
diff --git a/Engine/Source/Runtime/Renderer/Private/GPUBenchmark.h b/Engine/Source/Runtime/Renderer/Private/GPUBenchmark.h
--- a/Engine/Source/Runtime/Renderer/Private/GPUBenchmark.h
+++ b/Engine/Source/Runtime/Renderer/Private/GPUBenchmark.h
@@ -8,4 +8,4 @@
 
 // @param >0, WorkScale 10 for normal precision and runtime of less than a second
 // @param bDebugOut has no effect in shipping
-void RendererGPUBenchmark(FSynthBenchmarkResults& InOut, const FSceneView& View, uint32 WorkScale, bool bDebugOut = false);
+void RendererGPUBenchmark(FRHICommandListImmediate& RHICmdList, FSynthBenchmarkResults& InOut, const FSceneView& View, uint32 WorkScale, bool bDebugOut = false);
diff --git a/Engine/Source/Runtime/Renderer/Private/LightFunctionRendering.cpp b/Engine/Source/Runtime/Renderer/Private/LightFunctionRendering.cpp
--- a/Engine/Source/Runtime/Renderer/Private/LightFunctionRendering.cpp
+++ b/Engine/Source/Runtime/Renderer/Private/LightFunctionRendering.cpp
@@ -212,7 +212,7 @@ bool FDeferredShadingSceneRenderer::CheckForLightFunction( const FLightSceneInfo
  *
  * @param LightSceneInfo Represents the current light
  */
-bool FDeferredShadingSceneRenderer::RenderLightFunction(FRHICommandList& RHICmdList, const FLightSceneInfo* LightSceneInfo, bool bLightAttenuationCleared)
+bool FDeferredShadingSceneRenderer::RenderLightFunction(FRHICommandListImmediate& RHICmdList, const FLightSceneInfo* LightSceneInfo, bool bLightAttenuationCleared)
 {
 	if (ViewFamily.EngineShowFlags.LightFunctions)
 	{
@@ -222,7 +222,7 @@ bool FDeferredShadingSceneRenderer::RenderLightFunction(FRHICommandList& RHICmdL
 	return false;
 }
 
-bool FDeferredShadingSceneRenderer::RenderPreviewShadowsIndicator(FRHICommandList& RHICmdList, const FLightSceneInfo* LightSceneInfo, bool bLightAttenuationCleared)
+bool FDeferredShadingSceneRenderer::RenderPreviewShadowsIndicator(FRHICommandListImmediate& RHICmdList, const FLightSceneInfo* LightSceneInfo, bool bLightAttenuationCleared)
 {
 	if (GEngine->PreviewShadowsIndicatorMaterial)
 	{
@@ -232,13 +232,13 @@ bool FDeferredShadingSceneRenderer::RenderPreviewShadowsIndicator(FRHICommandLis
 	return false;
 }
 
-bool FDeferredShadingSceneRenderer::RenderLightFunctionForMaterial(FRHICommandList& RHICmdList, const FLightSceneInfo* LightSceneInfo, const FMaterialRenderProxy* MaterialProxy, bool bLightAttenuationCleared, bool bRenderingPreviewShadowsIndicator)
+bool FDeferredShadingSceneRenderer::RenderLightFunctionForMaterial(FRHICommandListImmediate& RHICmdList, const FLightSceneInfo* LightSceneInfo, const FMaterialRenderProxy* MaterialProxy, bool bLightAttenuationCleared, bool bRenderingPreviewShadowsIndicator)
 {
 	bool bRenderedLightFunction = false;
 
 	if (MaterialProxy && MaterialProxy->GetMaterial(Scene->GetFeatureLevel())->IsLightFunction())
 	{
-		GSceneRenderTargets.BeginRenderingLightAttenuation();
+		GSceneRenderTargets.BeginRenderingLightAttenuation(RHICmdList);
 
 		bRenderedLightFunction = true;
 
@@ -248,7 +248,6 @@ bool FDeferredShadingSceneRenderer::RenderLightFunctionForMaterial(FRHICommandLi
 		const FMaterialShaderMap* MaterialShaderMap = Material->GetRenderingThreadShaderMap();
 		FLightFunctionVS* VertexShader = MaterialShaderMap->GetShader<FLightFunctionVS>();
 		FLightFunctionPS* PixelShader = MaterialShaderMap->GetShader<FLightFunctionPS>();
-		RHICmdList.CheckIsNull();
 
 		// This was cached but when changing the material (e.g. editor) it wasn't updated.
 		// This will change with upcoming multi threaded rendering changes.
@@ -279,23 +278,23 @@ bool FDeferredShadingSceneRenderer::RenderLightFunctionForMaterial(FRHICommandLi
 				{
 					if( !bLightAttenuationCleared )
 					{
-						LightSceneInfo->Proxy->SetScissorRect(View); 
-						RHIClear(true, FLinearColor::White, false, 0, false, 0, FIntRect());
+						LightSceneInfo->Proxy->SetScissorRect(RHICmdList, View);
+						RHICmdList.Clear(true, FLinearColor::White, false, 0, false, 0, FIntRect());
 					}
 				}
 				else
 				{
 					// Set the device viewport for the view.
-					RHISetViewport(View.ViewRect.Min.X, View.ViewRect.Min.Y, 0.0f, View.ViewRect.Max.X, View.ViewRect.Max.Y, 1.0f);
+					RHICmdList.SetViewport(View.ViewRect.Min.X, View.ViewRect.Min.Y, 0.0f, View.ViewRect.Max.X, View.ViewRect.Max.Y, 1.0f);
 
 					// Set the states to modulate the light function with the render target.
-					RHISetDepthStencilState(TStaticDepthStencilState<false,CF_Always>::GetRHI());
+					RHICmdList.SetDepthStencilState(TStaticDepthStencilState<false,CF_Always>::GetRHI());
 
 					if( bLightAttenuationCleared )
 					{
 						if (bRenderingPreviewShadowsIndicator)
 						{
-							RHISetBlendState(TStaticBlendState<CW_RGBA,BO_Max,BF_One,BF_One,BO_Max,BF_One,BF_One>::GetRHI());
+							RHICmdList.SetBlendState(TStaticBlendState<CW_RGBA,BO_Max,BF_One,BF_One,BO_Max,BF_One,BF_One>::GetRHI());
 						}
 						else
 						{
@@ -303,32 +302,32 @@ bool FDeferredShadingSceneRenderer::RenderLightFunctionForMaterial(FRHICommandLi
 							// Light function shadows now write to the blue channel.
 
 							// Use modulated blending to BA since light functions are combined in the same buffer as normal shadows
-							RHISetBlendState(TStaticBlendState<CW_BA,BO_Add,BF_DestColor,BF_Zero,BO_Add,BF_Zero,BF_One>::GetRHI());
+							RHICmdList.SetBlendState(TStaticBlendState<CW_BA,BO_Add,BF_DestColor,BF_Zero,BO_Add,BF_Zero,BF_One>::GetRHI());
 						}
 					}
 					else
 					{
-						RHISetBlendState(TStaticBlendState<CW_RGBA>::GetRHI());
+						RHICmdList.SetBlendState(TStaticBlendState<CW_RGBA>::GetRHI());
 					}
 
 					if (((FVector)View.ViewMatrices.ViewOrigin - LightBounds.Center).SizeSquared() < FMath::Square(LightBounds.W * 1.05f + View.NearClippingDistance * 2.0f))
 					{
 						// Render backfaces with depth tests disabled since the camera is inside (or close to inside) the light function geometry
-						RHISetRasterizerState(View.bReverseCulling ? TStaticRasterizerState<FM_Solid,CM_CW>::GetRHI() : TStaticRasterizerState<FM_Solid,CM_CCW>::GetRHI());
+						RHICmdList.SetRasterizerState(View.bReverseCulling ? TStaticRasterizerState<FM_Solid, CM_CW>::GetRHI() : TStaticRasterizerState<FM_Solid, CM_CCW>::GetRHI());
 					}
 					else
 					{
 						// Render frontfaces with depth tests on to get the speedup from HiZ since the camera is outside the light function geometry
 						// Note, this is a reversed Z depth surface, using CF_GreaterEqual.
-						RHISetDepthStencilState(TStaticDepthStencilState<false,CF_GreaterEqual>::GetRHI());
-						RHISetRasterizerState(View.bReverseCulling ? TStaticRasterizerState<FM_Solid,CM_CCW>::GetRHI() : TStaticRasterizerState<FM_Solid,CM_CW>::GetRHI());
+						RHICmdList.SetDepthStencilState(TStaticDepthStencilState<false,CF_GreaterEqual>::GetRHI());
+						RHICmdList.SetRasterizerState(View.bReverseCulling ? TStaticRasterizerState<FM_Solid, CM_CCW>::GetRHI() : TStaticRasterizerState<FM_Solid, CM_CW>::GetRHI());
 					}
 
 					// Set the light's scissor rectangle.
-					LightSceneInfo->Proxy->SetScissorRect(View);
+					LightSceneInfo->Proxy->SetScissorRect(RHICmdList, View);
 
 					// Render a bounding light sphere.
-					RHISetBoundShaderState(LightFunctionBoundShaderState);
+					RHICmdList.SetBoundShaderState(LightFunctionBoundShaderState);
 					VertexShader->SetParameters(RHICmdList, &View, LightSceneInfo);
 					PixelShader->SetParameters(RHICmdList, &View, LightSceneInfo, MaterialProxy, bRenderingPreviewShadowsIndicator, FadeAlpha);
 
@@ -340,12 +339,12 @@ bool FDeferredShadingSceneRenderer::RenderLightFunctionForMaterial(FRHICommandLi
 		}
 
 		// Restore states.
-		RHISetScissorRect(false,0,0,0,0);
+		RHICmdList.SetScissorRect(false, 0, 0, 0, 0);
 
 		if (bRenderedLightFunction)
 		{
 			// Restore stencil buffer to all 0's which is the assumed default state
-			RHIClear(false,FColor(0,0,0),false,0,true,0, FIntRect());
+			RHICmdList.Clear(false,FColor(0,0,0),false,0,true,0, FIntRect());
 		}
 	}
 	return bRenderedLightFunction;
diff --git a/Engine/Source/Runtime/Renderer/Private/LightMapDensityRendering.cpp b/Engine/Source/Runtime/Renderer/Private/LightMapDensityRendering.cpp
--- a/Engine/Source/Runtime/Renderer/Private/LightMapDensityRendering.cpp
+++ b/Engine/Source/Runtime/Renderer/Private/LightMapDensityRendering.cpp
@@ -34,10 +34,9 @@ IMPLEMENT_DENSITY_LIGHTMAPPED_SHADER_TYPE( TLightMapPolicy<LQ_LIGHTMAP>, TLightM
 IMPLEMENT_DENSITY_LIGHTMAPPED_SHADER_TYPE( TLightMapPolicy<HQ_LIGHTMAP>, TLightMapPolicyHQ );
 #endif
 
-bool FDeferredShadingSceneRenderer::RenderLightMapDensities()
+bool FDeferredShadingSceneRenderer::RenderLightMapDensities(FRHICommandListImmediate& RHICmdList)
 {
 	bool bDirty=0;
-
 	if (Scene->GetFeatureLevel() >= ERHIFeatureLevel::SM3)
 	{
 		SCOPED_DRAW_EVENT(LightMapDensity, DEC_SCENE_ITEMS);
@@ -50,10 +49,10 @@ bool FDeferredShadingSceneRenderer::RenderLightMapDensities()
 			FViewInfo& View = Views[ViewIndex];
 
 			// Opaque blending, depth tests and writes.
-			RHISetBlendState(TStaticBlendState<>::GetRHI());
+			RHICmdList.SetBlendState(TStaticBlendState<>::GetRHI());
 			// Note, this is a reversed Z depth surface, using CF_GreaterEqual.
-			RHISetDepthStencilState(TStaticDepthStencilState<true,CF_GreaterEqual>::GetRHI());
-			RHISetViewport(View.ViewRect.Min.X, View.ViewRect.Min.Y, 0, View.ViewRect.Max.X, View.ViewRect.Max.Y, 1);
+			RHICmdList.SetDepthStencilState(TStaticDepthStencilState<true,CF_GreaterEqual>::GetRHI());
+			RHICmdList.SetViewport(View.ViewRect.Min.X, View.ViewRect.Min.Y, 0, View.ViewRect.Max.X, View.ViewRect.Max.Y, 1);
 
 			{
 				SCOPED_DRAW_EVENT(Dynamic, DEC_SCENE_ITEMS);
diff --git a/Engine/Source/Runtime/Renderer/Private/LightMapDensityRendering.h b/Engine/Source/Runtime/Renderer/Private/LightMapDensityRendering.h
--- a/Engine/Source/Runtime/Renderer/Private/LightMapDensityRendering.h
+++ b/Engine/Source/Runtime/Renderer/Private/LightMapDensityRendering.h
@@ -298,10 +298,8 @@ public:
 
 	void DrawShared(FRHICommandList& RHICmdList, const FSceneView* View, FBoundShaderStateRHIRef ShaderState ) const
 	{
-		RHICmdList.CheckIsNull();
-
 		// Set the actual shader & vertex declaration state
-		RHISetBoundShaderState(ShaderState);
+		RHICmdList.SetBoundShaderState(ShaderState);
 
 		// Set the base pass shader parameters for the material.
 		VertexShader->SetParameters(RHICmdList, MaterialRenderProxy,*View);
@@ -313,7 +311,7 @@ public:
 			DomainShader->SetParameters(RHICmdList, MaterialRenderProxy,*View);
 		}
 
-		RHISetBlendState(TStaticBlendState<>::GetRHI());
+		RHICmdList.SetBlendState(TStaticBlendState<>::GetRHI());
 
 		// Set the light-map policy.
 		LightMapPolicy.Set(RHICmdList, VertexShader,PixelShader,VertexShader,PixelShader,VertexFactory,MaterialRenderProxy,View);
