UE4 - rework RHI user-facing APIs to be safer and incorporate the desired sematics

diff --git a/Engine/Source/Runtime/Renderer/Private/LightPropagationVolume.cpp b/Engine/Source/Runtime/Renderer/Private/LightPropagationVolume.cpp
--- a/Engine/Source/Runtime/Renderer/Private/LightPropagationVolume.cpp
+++ b/Engine/Source/Runtime/Renderer/Private/LightPropagationVolume.cpp
@@ -181,17 +181,16 @@ public:
 	{
 		FComputeShaderRHIParamRef ShaderRHI = GetComputeShader();
 		SetUniformBufferParameter(RHICmdList, ShaderRHI, GetUniformBufferParameter<FLpvWriteUniformBufferParameters>(), Params.UniformBuffer );
-		RHICmdList.CheckIsNull();
 #if LPV_VOLUME_TEXTURE
 		for(int i  =0; i < 7; i++)
 		{
 			if ( LpvBufferSRVParameters[i].IsBound() )
 			{
-				RHISetShaderTexture( ShaderRHI, LpvBufferSRVParameters[i].GetBaseIndex(), Params.LpvBufferSRVs[i] );
+				RHICmdList.SetShaderTexture(ShaderRHI, LpvBufferSRVParameters[i].GetBaseIndex(), Params.LpvBufferSRVs[i]);
 			}
 			if ( LpvBufferUAVs[i].IsBound() )
 			{
-				RHISetUAVParameter( ShaderRHI, LpvBufferUAVs[i].GetBaseIndex(), Params.LpvBufferUAVs[i] );
+				RHICmdList.SetUAVParameter( ShaderRHI, LpvBufferUAVs[i].GetBaseIndex(), Params.LpvBufferUAVs[i] );
 			}
 
 			SetTextureParameter(RHICmdList, ShaderRHI, LpvBufferSRVParameters[i], LpvVolumeTextureSampler, TStaticSamplerState<SF_Bilinear,AM_Clamp,AM_Clamp,AM_Clamp>::GetRHI(), Params.LpvBufferSRVs[i] );
@@ -199,39 +198,39 @@ public:
 #else
 		if ( LpvBufferSRV.IsBound() )
 		{
-			RHISetShaderResourceViewParameter( ShaderRHI, LpvBufferSRV.GetBaseIndex(), Params.LpvBufferSRV );
+			RHICmdList.SetShaderResourceViewParameter( ShaderRHI, LpvBufferSRV.GetBaseIndex(), Params.LpvBufferSRV );
 		}
 		if ( LpvBufferUAV.IsBound() )
 		{
-			RHISetUAVParameter( ShaderRHI, LpvBufferUAV.GetBaseIndex(), Params.LpvBufferUAV );
+			RHICmdList.SetUAVParameter( ShaderRHI, LpvBufferUAV.GetBaseIndex(), Params.LpvBufferUAV );
 		}
 #endif
 		if ( VplListHeadBufferSRV.IsBound() )
 		{
-			RHISetShaderResourceViewParameter( ShaderRHI, VplListHeadBufferSRV.GetBaseIndex(), Params.VplListHeadBufferSRV );
+			RHICmdList.SetShaderResourceViewParameter( ShaderRHI, VplListHeadBufferSRV.GetBaseIndex(), Params.VplListHeadBufferSRV );
 		}
 		if ( VplListHeadBufferUAV.IsBound() )
 		{
-			RHISetUAVParameter( ShaderRHI, VplListHeadBufferUAV.GetBaseIndex(), Params.VplListHeadBufferUAV );
+			RHICmdList.SetUAVParameter( ShaderRHI, VplListHeadBufferUAV.GetBaseIndex(), Params.VplListHeadBufferUAV );
 		}
 		if ( VplListBufferSRV.IsBound() )
 		{
-			RHISetShaderResourceViewParameter( ShaderRHI, VplListBufferSRV.GetBaseIndex(), Params.VplListBufferSRV );
+			RHICmdList.SetShaderResourceViewParameter( ShaderRHI, VplListBufferSRV.GetBaseIndex(), Params.VplListBufferSRV );
 		}
 		if ( VplListBufferUAV.IsBound() )
 		{
-			RHISetUAVParameter( ShaderRHI, VplListBufferUAV.GetBaseIndex(), Params.VplListBufferUAV ); 
+			RHICmdList.SetUAVParameter( ShaderRHI, VplListBufferUAV.GetBaseIndex(), Params.VplListBufferUAV ); 
 		}
 #if LPV_GV_VOLUME_TEXTURE
 		for ( int i=0;i<3; i++ )
 		{
 			if ( GvBufferSRVParameters[i].IsBound() )
 			{
-				RHISetShaderTexture( ShaderRHI, GvBufferSRVParameters[i].GetBaseIndex(), Params.GvBufferSRVs[i] );
+				RHICmdList.SetShaderTexture(ShaderRHI, GvBufferSRVParameters[i].GetBaseIndex(), Params.GvBufferSRVs[i]);
 			}
 			if ( GvBufferUAVs[i].IsBound() )
 			{
-				RHISetUAVParameter( ShaderRHI, GvBufferUAVs[i].GetBaseIndex(), Params.GvBufferUAVs[i] );
+				RHICmdList.SetUAVParameter( ShaderRHI, GvBufferUAVs[i].GetBaseIndex(), Params.GvBufferUAVs[i] );
 			}
 
 			SetTextureParameter(RHICmdList, ShaderRHI, GvBufferSRVParameters[i], LpvVolumeTextureSampler, TStaticSamplerState<SF_Bilinear,AM_Clamp,AM_Clamp,AM_Clamp>::GetRHI(), Params.GvBufferSRVs[i] );
@@ -239,103 +238,102 @@ public:
 #else
 		if ( GvBufferSRV.IsBound() )
 		{
-			RHISetShaderResourceViewParameter( ShaderRHI, GvBufferSRV.GetBaseIndex(), Params.GvBufferSRV );
+			RHICmdList.SetShaderResourceViewParameter( ShaderRHI, GvBufferSRV.GetBaseIndex(), Params.GvBufferSRV );
 		}
 		if ( GvBufferUAV.IsBound() )
 		{
-			RHISetUAVParameter( ShaderRHI, GvBufferUAV.GetBaseIndex(), Params.GvBufferUAV );
+			RHICmdList.SetUAVParameter( ShaderRHI, GvBufferUAV.GetBaseIndex(), Params.GvBufferUAV );
 		}
 #endif
 		if(GvListBufferUAV.IsBound())
 		{
-			RHISetUAVParameter(ShaderRHI, GvListBufferUAV.GetBaseIndex(), Params.GvListBufferUAV); 
+			RHICmdList.SetUAVParameter(ShaderRHI, GvListBufferUAV.GetBaseIndex(), Params.GvListBufferUAV); 
 		}
 		if(GvListHeadBufferUAV.IsBound())
 		{
-			RHISetUAVParameter(ShaderRHI, GvListHeadBufferUAV.GetBaseIndex(), Params.GvListHeadBufferUAV); 
+			RHICmdList.SetUAVParameter(ShaderRHI, GvListHeadBufferUAV.GetBaseIndex(), Params.GvListHeadBufferUAV); 
 		}
 		if(GvListBufferSRV.IsBound())
 		{
-			RHISetShaderResourceViewParameter(ShaderRHI, GvListBufferSRV.GetBaseIndex(), Params.GvListBufferSRV);
+			RHICmdList.SetShaderResourceViewParameter(ShaderRHI, GvListBufferSRV.GetBaseIndex(), Params.GvListBufferSRV);
 		}
 		if(GvListHeadBufferSRV.IsBound())
 		{
-			RHISetShaderResourceViewParameter(ShaderRHI, GvListHeadBufferSRV.GetBaseIndex(), Params.GvListHeadBufferSRV);
+			RHICmdList.SetShaderResourceViewParameter(ShaderRHI, GvListHeadBufferSRV.GetBaseIndex(), Params.GvListHeadBufferSRV);
 		}
 	}
 
 	// Unbinds any buffers that have been bound.
 	void UnbindBuffers(FRHICommandList& RHICmdList)
 	{
-		RHICmdList.CheckIsNull();
 		FComputeShaderRHIParamRef ShaderRHI = GetComputeShader();
 #if LPV_VOLUME_TEXTURE
 		for ( int i = 0; i < 7; i++ )
 		{
 			if ( LpvBufferSRVParameters[i].IsBound() )
 		{
-				RHISetShaderTexture( ShaderRHI, LpvBufferSRVParameters[i].GetBaseIndex(), FTextureRHIParamRef()  );
+				RHICmdList.SetShaderTexture(ShaderRHI, LpvBufferSRVParameters[i].GetBaseIndex(), FTextureRHIParamRef());
 		}
 			if ( LpvBufferUAVs[i].IsBound() )
 		{
-				RHISetUAVParameter( ShaderRHI, LpvBufferUAVs[i].GetBaseIndex(), FUnorderedAccessViewRHIParamRef() );
+				RHICmdList.SetUAVParameter( ShaderRHI, LpvBufferUAVs[i].GetBaseIndex(), FUnorderedAccessViewRHIParamRef() );
 		}
 	}
 #else
 		if ( LpvBufferSRV.IsBound() )
 		{
-			RHISetShaderResourceViewParameter( ShaderRHI, LpvBufferSRV.GetBaseIndex(), FShaderResourceViewRHIParamRef() );
+			RHICmdList.SetShaderResourceViewParameter( ShaderRHI, LpvBufferSRV.GetBaseIndex(), FShaderResourceViewRHIParamRef() );
 		}
 		if ( LpvBufferUAV.IsBound() )
 		{
-			RHISetUAVParameter( ShaderRHI, LpvBufferUAV.GetBaseIndex(), FUnorderedAccessViewRHIParamRef() );
+			RHICmdList.SetUAVParameter( ShaderRHI, LpvBufferUAV.GetBaseIndex(), FUnorderedAccessViewRHIParamRef() );
 		}
 #endif
 		if ( VplListHeadBufferSRV.IsBound() )
 		{
-			RHISetShaderResourceViewParameter( ShaderRHI, VplListHeadBufferSRV.GetBaseIndex(), FShaderResourceViewRHIParamRef() );
+			RHICmdList.SetShaderResourceViewParameter( ShaderRHI, VplListHeadBufferSRV.GetBaseIndex(), FShaderResourceViewRHIParamRef() );
 		}
 		if ( VplListHeadBufferUAV.IsBound() )
 		{
-			RHISetUAVParameter( ShaderRHI, VplListHeadBufferUAV.GetBaseIndex(), FUnorderedAccessViewRHIParamRef() );
+			RHICmdList.SetUAVParameter( ShaderRHI, VplListHeadBufferUAV.GetBaseIndex(), FUnorderedAccessViewRHIParamRef() );
 		}
 		if ( VplListBufferSRV.IsBound() )
 		{
-			RHISetShaderResourceViewParameter( ShaderRHI, VplListBufferSRV.GetBaseIndex(), FShaderResourceViewRHIParamRef() );
+			RHICmdList.SetShaderResourceViewParameter( ShaderRHI, VplListBufferSRV.GetBaseIndex(), FShaderResourceViewRHIParamRef() );
 		}
 		if ( VplListBufferUAV.IsBound() )
 		{
-			RHISetUAVParameter( ShaderRHI, VplListBufferUAV.GetBaseIndex(), FUnorderedAccessViewRHIParamRef() );
+			RHICmdList.SetUAVParameter( ShaderRHI, VplListBufferUAV.GetBaseIndex(), FUnorderedAccessViewRHIParamRef() );
 		}
 #if LPV_GV_VOLUME_TEXTURE
 		for ( int i = 0; i < 3; i++ )
 		{
 			if ( GvBufferSRVParameters[i].IsBound() )
 		{
-				RHISetShaderTexture( ShaderRHI, GvBufferSRVParameters[i].GetBaseIndex(), FTextureRHIParamRef() );
+				RHICmdList.SetShaderTexture(ShaderRHI, GvBufferSRVParameters[i].GetBaseIndex(), FTextureRHIParamRef());
 		}
 			if ( GvBufferUAVs[i].IsBound() )
 		{
-				RHISetUAVParameter( ShaderRHI, GvBufferUAVs[i].GetBaseIndex(), FUnorderedAccessViewRHIParamRef() );
+				RHICmdList.SetUAVParameter( ShaderRHI, GvBufferUAVs[i].GetBaseIndex(), FUnorderedAccessViewRHIParamRef() );
 		}
 		}
 #else
 		if ( GvBufferSRV.IsBound() )
 		{
-			RHISetShaderResourceViewParameter( ShaderRHI, GvBufferSRV.GetBaseIndex(), FShaderResourceViewRHIParamRef() );
+			RHICmdList.SetShaderResourceViewParameter( ShaderRHI, GvBufferSRV.GetBaseIndex(), FShaderResourceViewRHIParamRef() );
 		}
 		if ( GvBufferUAV.IsBound() )
 		{
-			RHISetUAVParameter( ShaderRHI, GvBufferUAV.GetBaseIndex(), FUnorderedAccessViewRHIParamRef() );
+			RHICmdList.SetUAVParameter( ShaderRHI, GvBufferUAV.GetBaseIndex(), FUnorderedAccessViewRHIParamRef() );
 		}
 #endif
 		if(GvListBufferUAV.IsBound())
 		{
-			RHISetUAVParameter(ShaderRHI, GvListBufferUAV.GetBaseIndex(), FUnorderedAccessViewRHIParamRef());
+			RHICmdList.SetUAVParameter(ShaderRHI, GvListBufferUAV.GetBaseIndex(), FUnorderedAccessViewRHIParamRef());
 		}
 		if(GvListHeadBufferUAV.IsBound())
 		{
-			RHISetUAVParameter(ShaderRHI, GvListHeadBufferUAV.GetBaseIndex(), FUnorderedAccessViewRHIParamRef());
+			RHICmdList.SetUAVParameter(ShaderRHI, GvListHeadBufferUAV.GetBaseIndex(), FUnorderedAccessViewRHIParamRef());
 		}
 	}
 
@@ -724,12 +722,11 @@ FLightPropagationVolume::FLightPropagationVolume() :
 		mLpvBuffers[i]->Initialize( sizeof(LPVBufferElement), LPV_GRIDRES*LPV_GRIDRES*LPV_GRIDRES, BUF_FastVRAM, false, false );
 	}
 #endif
-
 	// VPL List buffers
 	mVplListBuffer = new FRWBufferStructured();
-	mVplListBuffer->Initialize( sizeof( VplListEntry ), 256*256*LPV_GRIDRES, 0, true, false );
+	mVplListBuffer->Initialize(sizeof(VplListEntry), 256 * 256 * LPV_GRIDRES, 0, true, false);
 	mVplListHeadBuffer = new FRWBufferByteAddress();
-	mVplListHeadBuffer->Initialize( LPV_GRIDRES*LPV_GRIDRES*LPV_GRIDRES*4, BUF_ByteAddressBuffer );
+	mVplListHeadBuffer->Initialize(LPV_GRIDRES*LPV_GRIDRES*LPV_GRIDRES * 4, BUF_ByteAddressBuffer);
 
 	// Geometry volume buffers
 #if !LPV_GV_VOLUME_TEXTURE
@@ -744,9 +741,9 @@ FLightPropagationVolume::FLightPropagationVolume() :
 #endif // !LPV_GV_VOLUME_TEXTURE
 
 	GvListBuffer = new FRWBufferStructured();
-	GvListBuffer->Initialize( sizeof( VplListEntry ), 256*256*LPV_GRIDRES, 0, true, false );
+	GvListBuffer->Initialize(sizeof(VplListEntry), 256 * 256 * LPV_GRIDRES, 0, true, false);
 	GvListHeadBuffer = new FRWBufferByteAddress();
-	GvListHeadBuffer->Initialize( LPV_GRIDRES*LPV_GRIDRES*LPV_GRIDRES*4, BUF_ByteAddressBuffer );
+	GvListHeadBuffer->Initialize(LPV_GRIDRES*LPV_GRIDRES*LPV_GRIDRES * 4, BUF_ByteAddressBuffer);
 
 	LpvWriteUniformBufferParams = new FLpvWriteUniformBufferParameters;
 	FMemory::Memzero( LpvWriteUniformBufferParams, sizeof(FLpvWriteUniformBufferParameters) );
@@ -938,7 +935,7 @@ void FLightPropagationVolume::InitSettings( const FSceneView& View )
 /**
 * Clears the LPV
 */
-void FLightPropagationVolume::Clear(FRHICommandList& RHICmdList)
+void FLightPropagationVolume::Clear(FRHICommandListImmediate& RHICmdList)
 {
 	if ( !bEnabled )
 	{
@@ -958,7 +955,7 @@ void FLightPropagationVolume::Clear(FRHICommandList& RHICmdList)
 	// Clear the list buffers
 	{
 		TShaderMapRef<FLpvClearListsCS> Shader(GetGlobalShaderMap());
-		RHISetComputeShader(Shader->GetComputeShader());
+		RHICmdList.SetComputeShader(Shader->GetComputeShader());
 
 		FLpvBaseWriteShaderParams ShaderParams;
 		GetShaderParams( ShaderParams );
@@ -970,7 +967,7 @@ void FLightPropagationVolume::Clear(FRHICommandList& RHICmdList)
 	// Clear the LPV (or fade, if REFINE_OVER_TIME is enabled)
 	{
 		TShaderMapRef<FLpvClearCS> Shader(GetGlobalShaderMap());
-		RHISetComputeShader(Shader->GetComputeShader());
+		RHICmdList.SetComputeShader(Shader->GetComputeShader());
 
 		FLpvBaseWriteShaderParams ShaderParams;
 		GetShaderParams( ShaderParams );
@@ -983,7 +980,7 @@ void FLightPropagationVolume::Clear(FRHICommandList& RHICmdList)
 	if ( SecondaryOcclusionStrength > 0.001f )
 	{
 		TShaderMapRef<FLpvClearGeometryVolumeCS> Shader(GetGlobalShaderMap());
-		RHISetComputeShader(Shader->GetComputeShader());
+		RHICmdList.SetComputeShader(Shader->GetComputeShader());
 
 		FLpvBaseWriteShaderParams ShaderParams;
 		GetShaderParams( ShaderParams );
@@ -993,12 +990,9 @@ void FLightPropagationVolume::Clear(FRHICommandList& RHICmdList)
 	}
 	//RHIBatchComputeShaderEnd();
 
-	// Clear the list UAV counters
-	RHICmdList.CheckIsNull();
-
-	RHISetUAVParameter( FComputeShaderRHIRef(), 7, mVplListBuffer->UAV, 0 );
-	RHISetUAVParameter( FComputeShaderRHIRef(), 7, GvListBuffer->UAV, 0 );
-	RHISetUAVParameter( FComputeShaderRHIRef(), 7, FUnorderedAccessViewRHIParamRef(), 0 );
+	RHICmdList.SetUAVParameter( FComputeShaderRHIRef(), 7, mVplListBuffer->UAV, 0 );
+	RHICmdList.SetUAVParameter( FComputeShaderRHIRef(), 7, GvListBuffer->UAV, 0 );
+	RHICmdList.SetUAVParameter( FComputeShaderRHIRef(), 7, FUnorderedAccessViewRHIParamRef(), 0 );
 }
 
 /**
@@ -1051,7 +1045,7 @@ void FLightPropagationVolume::SetVplInjectionConstants(
 * Injects a Directional light into the LPV
 */
 void FLightPropagationVolume::InjectDirectionalLightRSM(
-	FRHICommandList& RHICmdList, 
+	FRHICommandListImmediate& RHICmdList,
 	const FTexture2DRHIRef&		RsmDiffuseTex, 
 	const FTexture2DRHIRef&		RsmNormalTex, 
 	const FTexture2DRHIRef&		RsmDepthTex, 
@@ -1065,7 +1059,7 @@ void FLightPropagationVolume::InjectDirectionalLightRSM(
 		SetVplInjectionConstants(ProjectedShadowInfo, LightProxy );
 
 		TShaderMapRef<FLpvInject_GenerateVplListsCS> Shader(GetGlobalShaderMap());
-		RHISetComputeShader(Shader->GetComputeShader());
+		RHICmdList.SetComputeShader(Shader->GetComputeShader());
 
 		// Clear the list counter the first time this function is called in a frame
 		FLpvBaseWriteShaderParams ShaderParams;
@@ -1090,7 +1084,7 @@ void FLightPropagationVolume::InjectDirectionalLightRSM(
 		LpvWriteUniformBufferParams->GeometryVolumeCaptureLightDirection = LightDirection;
 
 		TShaderMapRef<FLpvBuildGeometryVolumeCS> Shader(GetGlobalShaderMap());
-		RHISetComputeShader(Shader->GetComputeShader());
+		RHICmdList.SetComputeShader(Shader->GetComputeShader());
 
 		LpvWriteUniformBuffer.SetContents( *LpvWriteUniformBufferParams ); // FIXME: is this causing a stall? Double-buffer?
 
@@ -1109,7 +1103,7 @@ void FLightPropagationVolume::InjectDirectionalLightRSM(
 /**
 * Propagates light in the LPV 
 */
-void FLightPropagationVolume::Propagate(FRHICommandList& RHICmdList)
+void FLightPropagationVolume::Propagate(FRHICommandListImmediate& RHICmdList)
 {
 	if ( !bEnabled )
 	{
@@ -1123,7 +1117,7 @@ void FLightPropagationVolume::Propagate(FRHICommandList& RHICmdList)
 		mWriteBufferIndex = 1-mWriteBufferIndex; // Swap buffers with each iteration
 
 		TShaderMapRef<FLpvInject_AccumulateVplListsCS> Shader(GetGlobalShaderMap());
-		RHISetComputeShader(Shader->GetComputeShader());
+		RHICmdList.SetComputeShader(Shader->GetComputeShader());
 
 		LpvWriteUniformBuffer.SetContents( *LpvWriteUniformBufferParams );
 
@@ -1149,7 +1143,7 @@ void FLightPropagationVolume::Propagate(FRHICommandList& RHICmdList)
 			bool bSecondaryBounces		= (SecondaryBounceStrength > 0.001f); 
 			mWriteBufferIndex = 1-mWriteBufferIndex; // Swap buffers with each iteration
 			FLpvWriteShaderCSBase* Shader = GetPropagateShader( bSecondaryOcclusion, bSecondaryBounces );
-			RHISetComputeShader(Shader->GetComputeShader());
+			RHICmdList.SetComputeShader(Shader->GetComputeShader());
 
 			LpvWriteUniformBufferParams->PropagationIndex = PropagationStep;
 			PropagationStep++;
@@ -1209,7 +1203,7 @@ void FLightPropagationVolume::GetShaderParams( FLpvBaseWriteShaderParams& OutPar
 	OutParams.UniformBuffer = LpvWriteUniformBuffer;
 }
 
-void FLightPropagationVolume::InjectLightDirect(FRHICommandList& RHICmdList, const FLightSceneProxy& Light )
+void FLightPropagationVolume::InjectLightDirect(FRHICommandListImmediate& RHICmdList, const FLightSceneProxy& Light)
 {
 	if(!bEnabled)
 	{
@@ -1250,7 +1244,7 @@ void FLightPropagationVolume::InjectLightDirect(FRHICommandList& RHICmdList, con
 				Shader = (FLpvInjectShader_Base*)*TShaderMapRef<TLpvInject_PointLightCS<0> >( GetGlobalShaderMap() );
 			}
 		}
-		RHISetComputeShader(Shader->GetComputeShader());
+		RHICmdList.SetComputeShader(Shader->GetComputeShader());
 
 		FDirectLightInjectBufferRef InjectUniformBuffer = FDirectLightInjectBufferRef::CreateUniformBufferImmediate(InjectUniformBufferParams, UniformBuffer_SingleDraw);
 
diff --git a/Engine/Source/Runtime/Renderer/Private/LightPropagationVolume.h b/Engine/Source/Runtime/Renderer/Private/LightPropagationVolume.h
--- a/Engine/Source/Runtime/Renderer/Private/LightPropagationVolume.h
+++ b/Engine/Source/Runtime/Renderer/Private/LightPropagationVolume.h
@@ -126,25 +126,25 @@ public:
 
 	void InitSettings( const FSceneView& View );
 
-	void Clear(FRHICommandList& RHICmdList);
+	void Clear(FRHICommandListImmediate& RHICmdList);
 
 	void SetVplInjectionConstants(
 		const FProjectedShadowInfo&	ProjectedShadowInfo,
 		const FLightSceneProxy* LightProxy );
 
 	void InjectDirectionalLightRSM(
-		FRHICommandList& RHICmdList, 
+		FRHICommandListImmediate& RHICmdList,
 		const FTexture2DRHIRef&		RsmDiffuseTex, 
 		const FTexture2DRHIRef&		RsmNormalTex, 
 		const FTexture2DRHIRef&		RsmDepthTex, 
 		const FProjectedShadowInfo&	ProjectedShadowInfo,
 		const FLinearColor&			LightColour );
 
-	void InjectLightDirect(FRHICommandList& RHICmdList, const FLightSceneProxy& Light );
+	void InjectLightDirect(FRHICommandListImmediate& RHICmdList, const FLightSceneProxy& Light);
 
-	void Propagate(FRHICommandList& RHICmdList);
+	void Propagate(FRHICommandListImmediate& RHICmdList);
 
-	void Visualise(FRHICommandList& RHICmdList, const FSceneView& View) const;
+	void Visualise(FRHICommandListImmediate& RHICmdList, const FSceneView& View) const;
 
 	const FIntVector& GetGridOffset() const { return mGridOffset; }
 
diff --git a/Engine/Source/Runtime/Renderer/Private/LightPropagationVolumeVisualisation.cpp b/Engine/Source/Runtime/Renderer/Private/LightPropagationVolumeVisualisation.cpp
--- a/Engine/Source/Runtime/Renderer/Private/LightPropagationVolumeVisualisation.cpp
+++ b/Engine/Source/Runtime/Renderer/Private/LightPropagationVolumeVisualisation.cpp
@@ -128,7 +128,6 @@ public:
 	{
 		FPixelShaderRHIParamRef ShaderRHI = GetPixelShader();
 		FGlobalShader::SetParameters(RHICmdList, ShaderRHI, View); 
-		RHICmdList.CheckIsNull();
 
 #if LPV_VOLUME_TEXTURE
 		for ( int i=0; i<7; i++ )
@@ -136,14 +135,14 @@ public:
 			FTextureRHIParamRef LpvBufferSrv = LPV->LpvVolumeTextures[ 1-LPV->mWriteBufferIndex ][i]->GetRenderTargetItem().ShaderResourceTexture;
 			if ( LpvBufferSRVParameters[i].IsBound() )
 			{
-				RHISetShaderTexture( ShaderRHI, LpvBufferSRVParameters[i].GetBaseIndex(), LpvBufferSrv );
+				RHICmdList.SetShaderTexture(ShaderRHI, LpvBufferSRVParameters[i].GetBaseIndex(), LpvBufferSrv);
 			}
 			SetTextureParameter(RHICmdList, ShaderRHI, LpvBufferSRVParameters[i], LpvVolumeTextureSampler, TStaticSamplerState<SF_Bilinear,AM_Clamp,AM_Clamp,AM_Clamp>::GetRHI(), LpvBufferSrv );
 		}
 #else
 		if ( InLpvBuffer.IsBound() ) 
 		{
-			RHISetShaderResourceViewParameter( ShaderRHI, InLpvBuffer.GetBaseIndex(), LPV->mLpvBuffers[ LPV->mWriteBufferIndex ]->SRV ); 
+			RHICmdList.SetShaderResourceViewParameter( ShaderRHI, InLpvBuffer.GetBaseIndex(), LPV->mLpvBuffers[ LPV->mWriteBufferIndex ]->SRV ); 
 		}
 #endif
 
@@ -153,7 +152,7 @@ public:
 			FTextureRHIParamRef GvBufferSrv = LPV->GvVolumeTextures[i]->GetRenderTargetItem().ShaderResourceTexture;
 			if ( GvBufferSRVParameters[i].IsBound() )
 			{
-				RHISetShaderTexture( ShaderRHI, GvBufferSRVParameters[i].GetBaseIndex(), GvBufferSrv );
+				RHICmdList.SetShaderTexture(ShaderRHI, GvBufferSRVParameters[i].GetBaseIndex(), GvBufferSrv);
 			}
 			SetTextureParameter(RHICmdList, ShaderRHI, GvBufferSRVParameters[i], LpvVolumeTextureSampler, TStaticSamplerState<SF_Bilinear,AM_Clamp,AM_Clamp,AM_Clamp>::GetRHI(), GvBufferSrv );
 		}
@@ -161,7 +160,7 @@ public:
 #else
 		if ( InGvBuffer.IsBound() ) 
 		{
-			RHISetShaderResourceViewParameter( ShaderRHI, InGvBuffer.GetBaseIndex(), LPV->GvBuffer->SRV ); 
+			RHICmdList.SetShaderResourceViewParameter( ShaderRHI, InGvBuffer.GetBaseIndex(), LPV->GvBuffer->SRV ); 
 		}
 #endif
 	}
@@ -248,7 +247,7 @@ IMPLEMENT_SHADER_TYPE(,FLpvVisualiseVS,TEXT("LPVVisualise"),TEXT("VShader"),SF_V
 IMPLEMENT_SHADER_TYPE(,FLpvVisualisePS,TEXT("LPVVisualise"),TEXT("PShader"),SF_Pixel);
 
 
-void FLightPropagationVolume::Visualise(FRHICommandList& RHICmdList, const FSceneView& View) const
+void FLightPropagationVolume::Visualise(FRHICommandListImmediate& RHICmdList, const FSceneView& View) const
 {
 	SCOPED_DRAW_EVENT(LpvVisualise, DEC_LIGHT);
 	check(View.GetFeatureLevel() == ERHIFeatureLevel::SM5);
diff --git a/Engine/Source/Runtime/Renderer/Private/LightRendering.cpp b/Engine/Source/Runtime/Renderer/Private/LightRendering.cpp
--- a/Engine/Source/Runtime/Renderer/Private/LightRendering.cpp
+++ b/Engine/Source/Runtime/Renderer/Private/LightRendering.cpp
@@ -276,10 +276,8 @@ void GetLightNameForDrawEvent(const FLightSceneProxy* LightProxy, FString& Light
 uint32 GetShadowQuality();
 
 /** Renders the scene's lighting. */
-void FDeferredShadingSceneRenderer::RenderLights(FRHICommandList& RHICmdList)
+void FDeferredShadingSceneRenderer::RenderLights(FRHICommandListImmediate& RHICmdList)
 {
-	RHICmdList.CheckIsNull();
-
 	SCOPED_DRAW_EVENT(Lights, DEC_SCENE_ITEMS);
 
 	if(IsSimpleDynamicLightingEnabled())
@@ -418,14 +416,14 @@ void FDeferredShadingSceneRenderer::RenderLights(FRHICommandList& RHICmdList)
 			}
 			else if (SimpleLights.InstanceData.Num() > 0)
 			{
-				GSceneRenderTargets.BeginRenderingSceneColor();
+				GSceneRenderTargets.BeginRenderingSceneColor(RHICmdList);
 				RenderSimpleLightsStandardDeferred(RHICmdList, SimpleLights);
 			}
 
 			{
 				SCOPED_DRAW_EVENT(StandardDeferredLighting, DEC_SCENE_ITEMS);
 
-				GSceneRenderTargets.BeginRenderingSceneColor();
+				GSceneRenderTargets.BeginRenderingSceneColor(RHICmdList);
 
 				// Draw non-shadowed non-light function lights without changing render targets between them
 				for (int32 LightIndex = StandardDeferredStart; LightIndex < AttenuationLightStart; LightIndex++)
@@ -474,15 +472,15 @@ void FDeferredShadingSceneRenderer::RenderLights(FRHICommandList& RHICmdList)
 			SCOPED_DRAW_EVENTF(EventLightPass, DEC_SCENE_ITEMS, *LightNameWithLevel);
 
 			// Do not resolve to scene color texture, this is done lazily
-			GSceneRenderTargets.FinishRenderingSceneColor(false);
+			GSceneRenderTargets.FinishRenderingSceneColor(RHICmdList, false);
 
 			if (bDrawShadows)
 			{
 				INC_DWORD_STAT(STAT_NumShadowedLights);
 
 				// All shadows render with min blending
-				GSceneRenderTargets.BeginRenderingLightAttenuation();
-				RHIClear(true, FLinearColor::White, false, 0, false, 0, FIntRect());
+				GSceneRenderTargets.BeginRenderingLightAttenuation(RHICmdList);
+				RHICmdList.Clear(true, FLinearColor::White, false, 0, false, 0, FIntRect());
 
 				bool bRenderedTranslucentObjectShadows = RenderTranslucentProjectedShadows(RHICmdList, &LightSceneInfo );
 				// Render non-modulated projected shadows to the attenuation buffer.
@@ -519,7 +517,7 @@ void FDeferredShadingSceneRenderer::RenderLights(FRHICommandList& RHICmdList)
 			}
 				
 			// Resolve light attenuation buffer
-			GSceneRenderTargets.FinishRenderingLightAttenuation();
+			GSceneRenderTargets.FinishRenderingLightAttenuation(RHICmdList);
 			
 			if(bDirectLighting && !bInjectedTranslucentVolume)
 			{
@@ -529,7 +527,7 @@ void FDeferredShadingSceneRenderer::RenderLights(FRHICommandList& RHICmdList)
 			}
 
 			GSceneRenderTargets.SetLightAttenuationMode(true);
-			GSceneRenderTargets.BeginRenderingSceneColor();
+			GSceneRenderTargets.BeginRenderingSceneColor(RHICmdList);
 
 			// Render the light to the scene color buffer, conditionally using the attenuation buffer or a 1x1 white texture as input 
 			if(bDirectLighting)
@@ -539,7 +537,7 @@ void FDeferredShadingSceneRenderer::RenderLights(FRHICommandList& RHICmdList)
 		}
 
 		// Do not resolve to scene color texture, this is done lazily
-		GSceneRenderTargets.FinishRenderingSceneColor(false);
+		GSceneRenderTargets.FinishRenderingSceneColor(RHICmdList, false);
 
 		// Restore the default mode
 		GSceneRenderTargets.SetLightAttenuationMode(true);
@@ -568,8 +566,8 @@ void FDeferredShadingSceneRenderer::RenderLights(FRHICommandList& RHICmdList)
 								FLightPropagationVolume* Lpv = ViewState->GetLightPropagationVolume();
 								if (Lpv && LightSceneInfo->Proxy)
 								{
-									//@todo-rco: RHIPacketList
-									Lpv->InjectLightDirect(FRHICommandList::GetNullRef(), *LightSceneInfo->Proxy);
+									
+									Lpv->InjectLightDirect(RHICmdList, *LightSceneInfo->Proxy);
 								}
 							}
 						}
@@ -580,7 +578,7 @@ void FDeferredShadingSceneRenderer::RenderLights(FRHICommandList& RHICmdList)
 	}
 }
 
-void FDeferredShadingSceneRenderer::RenderLightArrayForOverlapViewmode(FRHICommandList& RHICmdList, const TSparseArray<FLightSceneInfoCompact>& LightArray)
+void FDeferredShadingSceneRenderer::RenderLightArrayForOverlapViewmode(FRHICommandListImmediate& RHICmdList, const TSparseArray<FLightSceneInfoCompact>& LightArray)
 {
 	for (TSparseArray<FLightSceneInfoCompact>::TConstIterator LightIt(LightArray); LightIt; ++LightIt)
 	{
@@ -612,14 +610,14 @@ void FDeferredShadingSceneRenderer::RenderLightArrayForOverlapViewmode(FRHIComma
 	}
 }
 
-void FDeferredShadingSceneRenderer::RenderStationaryLightOverlap(FRHICommandList& RHICmdList)
+void FDeferredShadingSceneRenderer::RenderStationaryLightOverlap(FRHICommandListImmediate& RHICmdList)
 {
 	if (Scene->bIsEditorScene)
 	{
-		GSceneRenderTargets.BeginRenderingSceneColor();
+		GSceneRenderTargets.BeginRenderingSceneColor(RHICmdList);
 
 		// Clear to discard base pass values in scene color since we didn't skip that, to have valid scene depths
-		RHIClear(true, FLinearColor::Black, false, 0, false, 0, FIntRect());
+		RHICmdList.Clear(true, FLinearColor::Black, false, 0, false, 0, FIntRect());
 
 		RenderLightArrayForOverlapViewmode(RHICmdList, Scene->Lights);
 
@@ -660,7 +658,7 @@ static FVertexDeclarationRHIParamRef GetDeferredLightingVertexDeclaration()
 
 template<bool bUseIESProfile, bool bRadialAttenuation, bool bInverseSquaredFalloff>
 static void SetShaderTemplLighting(
-	FRHICommandList& RHICmdList, 
+	FRHICommandListImmediate& RHICmdList,
 	const FSceneView& View, 
 	FShader* VertexShader,
 	const FLightSceneInfo* LightSceneInfo)
@@ -681,7 +679,7 @@ static void SetShaderTemplLighting(
 
 template<bool bUseIESProfile, bool bRadialAttenuation, bool bInverseSquaredFalloff>
 static void SetShaderTemplLightingSimple(
-	FRHICommandList& RHICmdList, 
+	FRHICommandListImmediate& RHICmdList,
 	const FSceneView& View, 
 	FShader* VertexShader,
 	const FSimpleLightEntry& SimpleLight,
@@ -708,7 +706,7 @@ static void SetShaderTemplLightingSimple(
  * @param LightIndex The light's index into FScene::Lights
  * @return true if anything got rendered
  */
-void FDeferredShadingSceneRenderer::RenderLight(FRHICommandList& RHICmdList, const FLightSceneInfo* LightSceneInfo, bool bRenderOverlap, bool bIssueDrawEvent)
+void FDeferredShadingSceneRenderer::RenderLight(FRHICommandListImmediate& RHICmdList, const FLightSceneInfo* LightSceneInfo, bool bRenderOverlap, bool bIssueDrawEvent)
 {
 	SCOPE_CYCLE_COUNTER(STAT_DirectLightRenderingTime);
 	INC_DWORD_STAT(STAT_NumLightsUsingStandardDeferred);
@@ -833,7 +831,7 @@ void FDeferredShadingSceneRenderer::RenderLight(FRHICommandList& RHICmdList, con
 	}
 }
 
-void FDeferredShadingSceneRenderer::RenderSimpleLightsStandardDeferred(FRHICommandList& RHICmdList, const FSimpleLightArray& SimpleLights)
+void FDeferredShadingSceneRenderer::RenderSimpleLightsStandardDeferred(FRHICommandListImmediate& RHICmdList, const FSimpleLightArray& SimpleLights)
 {
 	SCOPE_CYCLE_COUNTER(STAT_DirectLightRenderingTime);
 	INC_DWORD_STAT_BY(STAT_NumLightsUsingStandardDeferred, SimpleLights.InstanceData.Num());
diff --git a/Engine/Source/Runtime/Renderer/Private/LightShaftRendering.cpp b/Engine/Source/Runtime/Renderer/Private/LightShaftRendering.cpp
--- a/Engine/Source/Runtime/Renderer/Private/LightShaftRendering.cpp
+++ b/Engine/Source/Runtime/Renderer/Private/LightShaftRendering.cpp
@@ -431,7 +431,7 @@ FGlobalBoundShaderState BlurLightShaftsBoundShaderState;
 FGlobalBoundShaderState AccumulateTermBoundShaderState;
 FGlobalBoundShaderState ApplyLightShaftsBoundShaderState;
 
-void AllocateOrReuseLightShaftRenderTarget(TRefCountPtr<IPooledRenderTarget>& Target, const TCHAR* Name)
+void AllocateOrReuseLightShaftRenderTarget(FRHICommandListImmediate& RHICmdList, TRefCountPtr<IPooledRenderTarget>& Target, const TCHAR* Name)
 {
 	if (!Target)
 	{
@@ -441,9 +441,6 @@ void AllocateOrReuseLightShaftRenderTarget(TRefCountPtr<IPooledRenderTarget>& Ta
 		FPooledRenderTargetDesc Desc(FPooledRenderTargetDesc::Create2DDesc(LightShaftSize, LightShaftFilterBufferFormat, TexCreate_None, TexCreate_RenderTargetable, false));
 		GRenderTargetPool.FindFreeElement(Desc, Target, Name);
 
-		//@todo-rco: RHIPacketList
-		FRHICommandList& RHICmdList = FRHICommandList::GetNullRef();
-
 		SetRenderTarget(RHICmdList, Target->GetRenderTargetItem().TargetableTexture, FTextureRHIRef());
 		RHICmdList.Clear(true, FLinearColor(0, 0, 0, 0), false, 1.0f, false, 0, FIntRect());
 	}
@@ -451,7 +448,7 @@ void AllocateOrReuseLightShaftRenderTarget(TRefCountPtr<IPooledRenderTarget>& Ta
 
 /** Generates the downsampled light shaft mask for either occlusion or bloom.  This swaps input and output before returning. */
 template<bool bDownsampleOcclusion>
-void DownsamplePass(FRHICommandList& RHICmdList, const FViewInfo& View, const FLightSceneInfo* LightSceneInfo, TRefCountPtr<IPooledRenderTarget>& LightShaftsSource, TRefCountPtr<IPooledRenderTarget>& LightShaftsDest)
+void DownsamplePass(FRHICommandListImmediate& RHICmdList, const FViewInfo& View, const FLightSceneInfo* LightSceneInfo, TRefCountPtr<IPooledRenderTarget>& LightShaftsSource, TRefCountPtr<IPooledRenderTarget>& LightShaftsDest)
 {
 	SCOPED_DRAW_EVENT(Downsample, DEC_SCENE_ITEMS);
 
@@ -521,6 +518,7 @@ void DownsamplePass(FRHICommandList& RHICmdList, const FViewInfo& View, const FL
 
 /** Applies Temporal AA to the light shaft source. */
 void ApplyTemporalAA(
+	FRHICommandListImmediate& RHICmdList,
 	FViewInfo& View, 
 	const TCHAR* HistoryRTName,
 	/** Contains last frame's history, if non-NULL.  This will be updated with the new frame's history. */
@@ -536,7 +534,7 @@ void ApplyTemporalAA(
 		if (*HistoryState && !View.bCameraCut)
 		{
 			FMemMark Mark(FMemStack::Get());
-			FRenderingCompositePassContext CompositeContext(View);
+			FRenderingCompositePassContext CompositeContext(RHICmdList, View);
 			FPostprocessContext Context(CompositeContext.Graph, View);
 
 			// Nodes for input render targets
@@ -553,7 +551,7 @@ void ApplyTemporalAA(
 
 			// Reuse a render target from the pool with a consistent name, for vis purposes
 			TRefCountPtr<IPooledRenderTarget> NewHistory;
-			AllocateOrReuseLightShaftRenderTarget(NewHistory, HistoryRTName);
+			AllocateOrReuseLightShaftRenderTarget(RHICmdList, NewHistory, HistoryRTName);
 
 			// Setup the output to write to the new history render target
 			Context.FinalOutput = FRenderingCompositeOutputRef(NodeTemporalAA);
@@ -576,7 +574,7 @@ void ApplyTemporalAA(
 			HistoryOutput = LightShaftsSource;
 			LightShaftsSource = NULL;
 
-			AllocateOrReuseLightShaftRenderTarget(LightShaftsSource, HistoryRTName);
+			AllocateOrReuseLightShaftRenderTarget(RHICmdList, LightShaftsSource, HistoryRTName);
 		}
 	}
 	else
@@ -588,7 +586,7 @@ void ApplyTemporalAA(
 
 /** Applies screen space radial blur passes. */
 void ApplyRadialBlurPasses(
-	FRHICommandList& RHICmdList,
+	FRHICommandListImmediate& RHICmdList,
 	const FViewInfo& View, 
 	const FLightSceneInfo* const LightSceneInfo, 
 	/** First pass source - this will not be overwritten. */
@@ -645,7 +643,7 @@ void ApplyRadialBlurPasses(
 	}
 }
 
-void FinishOcclusionTerm(FRHICommandList& RHICmdList, const FViewInfo& View, const FLightSceneInfo* const LightSceneInfo, TRefCountPtr<IPooledRenderTarget>& LightShaftsSource, TRefCountPtr<IPooledRenderTarget>& LightShaftsDest)
+void FinishOcclusionTerm(FRHICommandListImmediate& RHICmdList, const FViewInfo& View, const FLightSceneInfo* const LightSceneInfo, TRefCountPtr<IPooledRenderTarget>& LightShaftsSource, TRefCountPtr<IPooledRenderTarget>& LightShaftsDest)
 {
 	TShaderMapRef<FScreenVS> ScreenVertexShader(GetGlobalShaderMap());
 
@@ -715,7 +713,7 @@ bool ShouldRenderLightShaftsForLight(const FViewInfo& View, const FLightSceneInf
 }
 
 /** Renders light shafts. */
-FLightShaftsOutput FDeferredShadingSceneRenderer::RenderLightShaftOcclusion(FRHICommandList& RHICmdList)
+FLightShaftsOutput FDeferredShadingSceneRenderer::RenderLightShaftOcclusion(FRHICommandListImmediate& RHICmdList)
 {
 	FLightShaftsOutput Output;
 
@@ -736,12 +734,10 @@ FLightShaftsOutput FDeferredShadingSceneRenderer::RenderLightShaftOcclusion(FRHI
 			{
 				SCOPED_DRAW_EVENT(RenderLightShaftOcclusion, DEC_SCENE_ITEMS);
 
-				RHICmdList.CheckIsNull(); // resource allocation
-
 				// Allocate light shaft render targets on demand, using the pool
 				// Need two targets to ping pong between
-				AllocateOrReuseLightShaftRenderTarget(LightShafts0, TEXT("LightShafts0"));
-				AllocateOrReuseLightShaftRenderTarget(LightShafts1, TEXT("LightShafts1"));
+				AllocateOrReuseLightShaftRenderTarget(RHICmdList, LightShafts0, TEXT("LightShafts0"));
+				AllocateOrReuseLightShaftRenderTarget(RHICmdList, LightShafts1, TEXT("LightShafts1"));
 
 				for (int ViewIndex = 0;ViewIndex < Views.Num();ViewIndex++)
 				{
@@ -761,7 +757,7 @@ FLightShaftsOutput FDeferredShadingSceneRenderer::RenderLightShaftOcclusion(FRHI
 
 						// Apply temporal AA to the occlusion mask
 						// Result will be in HistoryOutput
-						ApplyTemporalAA(View, TEXT("LSOcclusionHistory"), HistoryState, LightShafts0, HistoryOutput);
+						ApplyTemporalAA(RHICmdList, View, TEXT("LSOcclusionHistory"), HistoryState, LightShafts0, HistoryOutput);
 
 						// Apply radial blur passes
 						// Send HistoryOutput in as the first pass input only, so it will not be overwritten by any subsequent passes, since it is needed for next frame
@@ -837,12 +833,11 @@ private:
 
 IMPLEMENT_SHADER_TYPE(,FApplyLightShaftsPixelShader,TEXT("LightShaftShader"),TEXT("ApplyLightShaftsPixelMain"),SF_Pixel);
 
-void ApplyLightShaftBloom(FRHICommandList& RHICmdList, const FViewInfo& View, const FLightSceneInfo* const LightSceneInfo, TRefCountPtr<IPooledRenderTarget>& LightShaftsSource)
+void ApplyLightShaftBloom(FRHICommandListImmediate& RHICmdList, const FViewInfo& View, const FLightSceneInfo* const LightSceneInfo, TRefCountPtr<IPooledRenderTarget>& LightShaftsSource)
 {
 	SCOPED_DRAW_EVENT(Apply, DEC_SCENE_ITEMS);
-	RHICmdList.CheckIsNull(); // GSceneRenderTargets ops
 
-	GSceneRenderTargets.BeginRenderingSceneColor();
+	GSceneRenderTargets.BeginRenderingSceneColor(RHICmdList);
 
 	RHICmdList.SetViewport(View.ViewRect.Min.X, View.ViewRect.Min.Y, 0.0f, View.ViewRect.Max.X, View.ViewRect.Max.Y, 1.0f);
 	RHICmdList.SetBlendState(TStaticBlendState<CW_RGB, BO_Add, BF_One, BF_One>::GetRHI());
@@ -874,13 +869,11 @@ void ApplyLightShaftBloom(FRHICommandList& RHICmdList, const FViewInfo& View, co
 		*ScreenVertexShader,
 		EDRF_UseTriangleOptimization);
 
-	GSceneRenderTargets.FinishRenderingSceneColor(false);
+	GSceneRenderTargets.FinishRenderingSceneColor(RHICmdList, false);
 }
 
-void FDeferredShadingSceneRenderer::RenderLightShaftBloom(FRHICommandList& RHICmdList)
+void FDeferredShadingSceneRenderer::RenderLightShaftBloom(FRHICommandListImmediate& RHICmdList)
 {
-	RHICmdList.CheckIsNull(); // resource allocation
-
 	if (DoesViewFamilyAllowLightShafts(ViewFamily))
 	{
 		TRefCountPtr<IPooledRenderTarget> LightShafts0;
@@ -895,8 +888,8 @@ void FDeferredShadingSceneRenderer::RenderLightShaftBloom(FRHICommandList& RHICm
 				SCOPED_DRAW_EVENT(RenderLightShaftBloom, DEC_SCENE_ITEMS);
 
 				// Allocate light shaft render targets on demand, using the pool
-				AllocateOrReuseLightShaftRenderTarget(LightShafts0, TEXT("LightShafts0"));
-				AllocateOrReuseLightShaftRenderTarget(LightShafts1, TEXT("LightShafts1"));
+				AllocateOrReuseLightShaftRenderTarget(RHICmdList, LightShafts0, TEXT("LightShafts0"));
+				AllocateOrReuseLightShaftRenderTarget(RHICmdList, LightShafts1, TEXT("LightShafts1"));
 
 				for (int ViewIndex = 0;ViewIndex < Views.Num();ViewIndex++)
 				{
@@ -922,7 +915,7 @@ void FDeferredShadingSceneRenderer::RenderLightShaftBloom(FRHICommandList& RHICm
 
 						// Apply temporal AA to the occlusion mask
 						// Result will be in HistoryOutput
-						ApplyTemporalAA(View, TEXT("LSBloomHistory"), HistoryState, LightShafts0, HistoryOutput);
+						ApplyTemporalAA(RHICmdList, View, TEXT("LSBloomHistory"), HistoryState, LightShafts0, HistoryOutput);
 
 						// Apply radial blur passes
 						// Send HistoryOutput in as the first pass input only, so it will not be overwritten by any subsequent passes, since it is needed for next frame
diff --git a/Engine/Source/Runtime/Renderer/Private/PostProcess/PostProcessAA.cpp b/Engine/Source/Runtime/Renderer/Private/PostProcess/PostProcessAA.cpp
--- a/Engine/Source/Runtime/Renderer/Private/PostProcess/PostProcessAA.cpp
+++ b/Engine/Source/Runtime/Renderer/Private/PostProcess/PostProcessAA.cpp
@@ -219,7 +219,7 @@ static void SetShaderTemplAA(const FRenderingCompositePassContext& Context)
 	TShaderMapRef<FPostProcessAntiAliasingPS<Quality> > PixelShader(GetGlobalShaderMap());
 
 	static FGlobalBoundShaderState BoundShaderState;
-	Context.RHICmdList.CheckIsNull(); // need new approach for "static FGlobalBoundShaderState" for parallel rendering
+	
 
 	SetGlobalBoundShaderState(Context.RHICmdList, BoundShaderState, GFilterVertexDeclaration.VertexDeclarationRHI, *VertexShader, *PixelShader);
 
@@ -249,7 +249,7 @@ void FRCPassPostProcessAA::Process(FRenderingCompositePassContext& Context)
 
 	const FSceneRenderTargetItem& DestRenderTarget = PassOutputs[0].RequestSurface(Context);
 
-	Context.RHICmdList.CheckIsNull(); // need new approach for "static FGlobalBoundShaderState" for parallel rendering
+	
 
 	// Set the view family's render target/viewport.
 	SetRenderTarget(Context.RHICmdList, DestRenderTarget.TargetableTexture, FTextureRHIRef());
diff --git a/Engine/Source/Runtime/Renderer/Private/PostProcess/PostProcessBloomSetup.cpp b/Engine/Source/Runtime/Renderer/Private/PostProcess/PostProcessBloomSetup.cpp
--- a/Engine/Source/Runtime/Renderer/Private/PostProcess/PostProcessBloomSetup.cpp
+++ b/Engine/Source/Runtime/Renderer/Private/PostProcess/PostProcessBloomSetup.cpp
@@ -179,7 +179,7 @@ void FRCPassPostProcessBloomSetup::Process(FRenderingCompositePassContext& Conte
 	TShaderMapRef<FPostProcessBloomSetupPS> PixelShader(GetGlobalShaderMap());
 
 	static FGlobalBoundShaderState BoundShaderState;
-	Context.RHICmdList.CheckIsNull(); // need new approach for "static FGlobalBoundShaderState" for parallel rendering
+	
 
 	SetGlobalBoundShaderState(Context.RHICmdList, BoundShaderState, GFilterVertexDeclaration.VertexDeclarationRHI, *VertexShader, *PixelShader);
 
diff --git a/Engine/Source/Runtime/Renderer/Private/PostProcess/PostProcessBokehDOF.cpp b/Engine/Source/Runtime/Renderer/Private/PostProcess/PostProcessBokehDOF.cpp
--- a/Engine/Source/Runtime/Renderer/Private/PostProcess/PostProcessBokehDOF.cpp
+++ b/Engine/Source/Runtime/Renderer/Private/PostProcess/PostProcessBokehDOF.cpp
@@ -563,7 +563,7 @@ void FRCPassPostProcessBokehDOF::SetShaderTempl(const FRenderingCompositePassCon
 	TShaderMapRef<FPostProcessBokehDOFPS> PixelShader(GetGlobalShaderMap());
 
 	static FGlobalBoundShaderState BoundShaderState;
-	Context.RHICmdList.CheckIsNull(); // need new approach for "static FGlobalBoundShaderState" for parallel rendering
+	
 
 	SetGlobalBoundShaderState(Context.RHICmdList, BoundShaderState, GFilterVertexDeclaration.VertexDeclarationRHI, *VertexShader, *PixelShader);
 
diff --git a/Engine/Source/Runtime/Renderer/Private/PostProcess/PostProcessBusyWait.cpp b/Engine/Source/Runtime/Renderer/Private/PostProcess/PostProcessBusyWait.cpp
--- a/Engine/Source/Runtime/Renderer/Private/PostProcess/PostProcessBusyWait.cpp
+++ b/Engine/Source/Runtime/Renderer/Private/PostProcess/PostProcessBusyWait.cpp
@@ -89,7 +89,7 @@ void FRCPassPostProcessBusyWait::Process(FRenderingCompositePassContext& Context
 	FIntRect SrcRect = View.ViewRect;
 	FIntRect DestRect = View.UnscaledViewRect;
 	
-	GSceneRenderTargets.BeginRenderingLightAttenuation();
+	GSceneRenderTargets.BeginRenderingLightAttenuation(Context.RHICmdList);
 	
 	const FSceneRenderTargetItem& DestRenderTarget = GSceneRenderTargets.GetLightAttenuation()->GetRenderTargetItem();
 
@@ -125,7 +125,6 @@ void FRCPassPostProcessBusyWait::Process(FRenderingCompositePassContext& Context
 
 	Context.RHICmdList.CopyToResolveTarget(DestRenderTarget.TargetableTexture, DestRenderTarget.ShaderResourceTexture, false, FResolveParams());
 
-	Context.RHICmdList.CheckIsNull(); // this won't work for parallel rendering
 	GSceneRenderTargets.SetLightAttenuation(0);
 }
 
diff --git a/Engine/Source/Runtime/Renderer/Private/PostProcess/PostProcessCompositeEditorPrimitives.cpp b/Engine/Source/Runtime/Renderer/Private/PostProcess/PostProcessCompositeEditorPrimitives.cpp
--- a/Engine/Source/Runtime/Renderer/Private/PostProcess/PostProcessCompositeEditorPrimitives.cpp
+++ b/Engine/Source/Runtime/Renderer/Private/PostProcess/PostProcessCompositeEditorPrimitives.cpp
@@ -168,7 +168,7 @@ static void SetCompositePrimitivesShaderTempl(const FRenderingCompositePassConte
 	TShaderMapRef<FPostProcessCompositeEditorPrimitivesPS<MSAASampleCount> > PixelShader(GetGlobalShaderMap());
 
 	static FGlobalBoundShaderState BoundShaderState;
-	Context.RHICmdList.CheckIsNull(); // need new approach for "static FGlobalBoundShaderState" for parallel rendering
+	
 
 	SetGlobalBoundShaderState(Context.RHICmdList, BoundShaderState, GFilterVertexDeclaration.VertexDeclarationRHI, *VertexShader, *PixelShader);
 
@@ -221,7 +221,6 @@ void FRCPassPostProcessCompositeEditorPrimitives::Process(FRenderingCompositePas
 
 		RenderPrimitivesToComposite(Context.RHICmdList, View);
 
-		Context.RHICmdList.CheckIsNull(); // this probably won't work with parallel rendering
 		GRenderTargetPool.VisualizeTexture.SetCheckPoint(Context.RHICmdList, GSceneRenderTargets.EditorPrimitivesColor);
 	}
 
@@ -279,7 +278,6 @@ void FRCPassPostProcessCompositeEditorPrimitives::Process(FRenderingCompositePas
 	Context.RHICmdList.CopyToResolveTarget(DestRenderTargetSurface, DestRenderTarget.ShaderResourceTexture, false, FResolveParams());
 
 	// Clean up targets
-	Context.RHICmdList.CheckIsNull(); // this probably won't work with parallel rendering
 	GSceneRenderTargets.CleanUpEditorPrimitiveTargets();
 }
 
@@ -293,9 +291,8 @@ FPooledRenderTargetDesc FRCPassPostProcessCompositeEditorPrimitives::ComputeOutp
 	return Ret;
 }
 
-void FRCPassPostProcessCompositeEditorPrimitives::RenderPrimitivesToComposite(FRHICommandList& RHICmdList, const FViewInfo& View)
+void FRCPassPostProcessCompositeEditorPrimitives::RenderPrimitivesToComposite(FRHICommandListImmediate& RHICmdList, const FViewInfo& View)
 {
-	RHICmdList.CheckIsNull(); // This uses DrawDynamicElements
 
 	// Always depth test against other editor primitives
 	// Note, this is a reversed Z depth surface, using CF_GreaterEqual.
@@ -372,7 +369,7 @@ void FRCPassPostProcessCompositeEditorPrimitives::RenderPrimitivesToComposite(FR
 		// Note, this is a reversed Z depth surface, using CF_GreaterEqual.
 		RHICmdList.SetDepthStencilState(TStaticDepthStencilState<true, CF_GreaterEqual>::GetRHI());
 
-		DrawViewElements<FBasePassOpaqueDrawingPolicyFactory>(View,FBasePassOpaqueDrawingPolicyFactory::ContextType(bDepthTest, ESceneRenderTargetsMode::SetTextures), SDPG_Foreground, false);
+		DrawViewElements<FBasePassOpaqueDrawingPolicyFactory>(View, FBasePassOpaqueDrawingPolicyFactory::ContextType(bDepthTest, ESceneRenderTargetsMode::SetTextures), SDPG_Foreground, false);
 
 		View.TopBatchedViewElements.Draw(RHICmdList, bNeedToSwitchVerticalAxis, View.ViewProjectionMatrix, View.ViewRect.Width(), View.ViewRect.Height(), false);
 	}
diff --git a/Engine/Source/Runtime/Renderer/Private/PostProcess/PostProcessCompositeEditorPrimitives.h b/Engine/Source/Runtime/Renderer/Private/PostProcess/PostProcessCompositeEditorPrimitives.h
--- a/Engine/Source/Runtime/Renderer/Private/PostProcess/PostProcessCompositeEditorPrimitives.h
+++ b/Engine/Source/Runtime/Renderer/Private/PostProcess/PostProcessCompositeEditorPrimitives.h
@@ -21,6 +21,6 @@ private:
 	 *
 	 * @param View	The view to draw in
 	 */
-	void RenderPrimitivesToComposite(FRHICommandList& RHICmdList, const FViewInfo& View);
+	void RenderPrimitivesToComposite(FRHICommandListImmediate& RHICmdList, const FViewInfo& View);
 };
 
diff --git a/Engine/Source/Runtime/Renderer/Private/PostProcess/PostProcessDOF.cpp b/Engine/Source/Runtime/Renderer/Private/PostProcess/PostProcessDOF.cpp
--- a/Engine/Source/Runtime/Renderer/Private/PostProcess/PostProcessDOF.cpp
+++ b/Engine/Source/Runtime/Renderer/Private/PostProcess/PostProcessDOF.cpp
@@ -134,7 +134,7 @@ void FRCPassPostProcessDOFSetup::Process(FRenderingCompositePassContext& Context
 	if (bNearBlurEnabled)
 	{
 		static FGlobalBoundShaderState BoundShaderState;
-		Context.RHICmdList.CheckIsNull(); // need new approach for "static FGlobalBoundShaderState" for parallel rendering
+		
 
 		TShaderMapRef< FPostProcessDOFSetupPS<1> > PixelShader(GetGlobalShaderMap());
 		SetGlobalBoundShaderState(Context.RHICmdList, BoundShaderState, GFilterVertexDeclaration.VertexDeclarationRHI, *VertexShader, *PixelShader);
@@ -144,7 +144,7 @@ void FRCPassPostProcessDOFSetup::Process(FRenderingCompositePassContext& Context
 	else
 	{
 		static FGlobalBoundShaderState BoundShaderState;
-		Context.RHICmdList.CheckIsNull(); // need new approach for "static FGlobalBoundShaderState" for parallel rendering
+		
 		TShaderMapRef< FPostProcessDOFSetupPS<0> > PixelShader(GetGlobalShaderMap());
 		SetGlobalBoundShaderState(Context.RHICmdList, BoundShaderState, GFilterVertexDeclaration.VertexDeclarationRHI, *VertexShader, *PixelShader);
 
@@ -299,7 +299,7 @@ void FRCPassPostProcessDOFRecombine::Process(FRenderingCompositePassContext& Con
 	if (bNearBlurEnabled)
 	{
 		static FGlobalBoundShaderState BoundShaderState;
-		Context.RHICmdList.CheckIsNull(); // need new approach for "static FGlobalBoundShaderState" for parallel rendering
+		
 		TShaderMapRef< FPostProcessDOFRecombinePS<1> > PixelShader(GetGlobalShaderMap());
 		SetGlobalBoundShaderState(Context.RHICmdList, BoundShaderState, GFilterVertexDeclaration.VertexDeclarationRHI, *VertexShader, *PixelShader);
 		PixelShader->SetParameters(Context);
@@ -307,7 +307,7 @@ void FRCPassPostProcessDOFRecombine::Process(FRenderingCompositePassContext& Con
 	else
 	{
 		static FGlobalBoundShaderState BoundShaderState;
-		Context.RHICmdList.CheckIsNull(); // need new approach for "static FGlobalBoundShaderState" for parallel rendering
+		
 		TShaderMapRef< FPostProcessDOFRecombinePS<0> > PixelShader(GetGlobalShaderMap());
 		SetGlobalBoundShaderState(Context.RHICmdList, BoundShaderState, GFilterVertexDeclaration.VertexDeclarationRHI, *VertexShader, *PixelShader);
 		PixelShader->SetParameters(Context);
diff --git a/Engine/Source/Runtime/Renderer/Private/PostProcess/PostProcessGBufferHints.cpp b/Engine/Source/Runtime/Renderer/Private/PostProcess/PostProcessGBufferHints.cpp
--- a/Engine/Source/Runtime/Renderer/Private/PostProcess/PostProcessGBufferHints.cpp
+++ b/Engine/Source/Runtime/Renderer/Private/PostProcess/PostProcessGBufferHints.cpp
@@ -109,7 +109,7 @@ void FRCPassPostProcessGBufferHints::Process(FRenderingCompositePassContext& Con
 	TShaderMapRef<FPostProcessGBufferHintsPS> PixelShader(GetGlobalShaderMap());
 
 	static FGlobalBoundShaderState BoundShaderState;
-	Context.RHICmdList.CheckIsNull(); // need new approach for "static FGlobalBoundShaderState" for parallel rendering
+	
 
 	SetGlobalBoundShaderState(Context.RHICmdList, BoundShaderState, GFilterVertexDeclaration.VertexDeclarationRHI, *VertexShader, *PixelShader);
 
@@ -148,7 +148,6 @@ void FRCPassPostProcessGBufferHints::Process(FRenderingCompositePassContext& Con
 			return Texture;
 		}
 	} TempRenderTarget(View, (const FTexture2DRHIRef&)DestRenderTarget.TargetableTexture);
-	Context.RHICmdList.CheckIsNull(); // need canvas to have a cmd list
 
 	FCanvas Canvas(&TempRenderTarget, NULL, ViewFamily.CurrentRealTime, ViewFamily.CurrentWorldTime, ViewFamily.DeltaWorldTime);
 
diff --git a/Engine/Source/Runtime/Renderer/Private/PostProcess/PostProcessHMD.cpp b/Engine/Source/Runtime/Renderer/Private/PostProcess/PostProcessHMD.cpp
--- a/Engine/Source/Runtime/Renderer/Private/PostProcess/PostProcessHMD.cpp
+++ b/Engine/Source/Runtime/Renderer/Private/PostProcess/PostProcessHMD.cpp
@@ -25,7 +25,7 @@ public:
 	/** Destructor. */
 	virtual ~FDistortionVertexDeclaration() {}
 
-	virtual void InitRHI()
+	virtual void InitRHI() override
 	{
 		uint16 Stride = sizeof(FDistortionVertex);
 		FVertexDeclarationElementList Elements;
@@ -37,7 +37,7 @@ public:
 		VertexDeclarationRHI = RHICreateVertexDeclaration(Elements);
 	}
 
-	virtual void ReleaseRHI()
+	virtual void ReleaseRHI() override
 	{
 		VertexDeclarationRHI.SafeRelease();
 	}
@@ -194,12 +194,11 @@ void FRCPassPostProcessHMD::Process(FRenderingCompositePassContext& Context)
 		TShaderMapRef<FPostProcessHMDVS> VertexShader(GetGlobalShaderMap());
 		TShaderMapRef<FPostProcessHMDPS> PixelShader(GetGlobalShaderMap());
 		static FGlobalBoundShaderState BoundShaderState;
-		Context.RHICmdList.CheckIsNull(); // need new approach for "static FGlobalBoundShaderState" for parallel rendering
+		
 		SetGlobalBoundShaderState(Context.RHICmdList, BoundShaderState, GDistortionVertexDeclaration.VertexDeclarationRHI, *VertexShader, *PixelShader);
 		VertexShader->SetVS(Context, View.StereoPass);
 		PixelShader->SetPS(Context, SrcRect, SrcSize, View.StereoPass, QuadTexTransform);
 	}
-	Context.RHICmdList.CheckIsNull(); // this call may not properly use the cmd list in the context
 	GEngine->HMDDevice->DrawDistortionMesh_RenderThread(Context, View, SrcSize);
 
 	Context.RHICmdList.CopyToResolveTarget(DestRenderTarget.TargetableTexture, DestRenderTarget.ShaderResourceTexture, false, FResolveParams());
diff --git a/Engine/Source/Runtime/Renderer/Private/PostProcess/PostProcessHistogram.cpp b/Engine/Source/Runtime/Renderer/Private/PostProcess/PostProcessHistogram.cpp
--- a/Engine/Source/Runtime/Renderer/Private/PostProcess/PostProcessHistogram.cpp
+++ b/Engine/Source/Runtime/Renderer/Private/PostProcess/PostProcessHistogram.cpp
@@ -109,30 +109,27 @@ void FRCPassPostProcessHistogram::Process(FRenderingCompositePassContext& Contex
 
 	TShaderMapRef<FPostProcessHistogramCS> ComputeShader(GetGlobalShaderMap());
 
-	RHISetComputeShader(ComputeShader->GetComputeShader());
-	RHISetRenderTarget(FTextureRHIRef(), FTextureRHIRef());	
+	Context.RHICmdList.SetComputeShader(ComputeShader->GetComputeShader());
+	SetRenderTarget(Context.RHICmdList, FTextureRHIRef(), FTextureRHIRef());
 	
 
 	// set destination
 	check(DestRenderTarget.UAV);
-	RHISetUAVParameter(ComputeShader->GetComputeShader(), ComputeShader->HistogramRWTexture.GetBaseIndex(), DestRenderTarget.UAV);
+	Context.RHICmdList.SetUAVParameter(ComputeShader->GetComputeShader(), ComputeShader->HistogramRWTexture.GetBaseIndex(), DestRenderTarget.UAV);
 
 	// we currently assume the input is half res, one full res pixel less to avoid getting bilinear filtered input
 	FIntPoint GatherExtent = (DestRect.Size() - FIntPoint(1, 1)) / 2;
 
 	FIntPoint ThreadGroupCountValue = ComputeThreadGroupCount(GatherExtent);
 
+	ComputeShader->SetCS(Context.RHICmdList, Context, ThreadGroupCountValue, (DestRect.Min + FIntPoint(1, 1)) / 2, GatherExtent);
 	
-	//@todo-rco: RHIPacketList
-	FRHICommandList& RHICmdList = FRHICommandList::GetNullRef();
-	ComputeShader->SetCS(RHICmdList, Context, ThreadGroupCountValue, (DestRect.Min + FIntPoint(1, 1)) / 2, GatherExtent);
-	
-	DispatchComputeShader(RHICmdList, *ComputeShader, ThreadGroupCountValue.X, ThreadGroupCountValue.Y, 1);
+	DispatchComputeShader(Context.RHICmdList, *ComputeShader, ThreadGroupCountValue.X, ThreadGroupCountValue.Y, 1);
 
 	// un-set destination
-	RHISetUAVParameter(ComputeShader->GetComputeShader(), ComputeShader->HistogramRWTexture.GetBaseIndex(), NULL);
+	Context.RHICmdList.SetUAVParameter(ComputeShader->GetComputeShader(), ComputeShader->HistogramRWTexture.GetBaseIndex(), NULL);
 
-	RHICopyToResolveTarget(DestRenderTarget.TargetableTexture, DestRenderTarget.ShaderResourceTexture, false, FResolveParams());
+	Context.RHICmdList.CopyToResolveTarget(DestRenderTarget.TargetableTexture, DestRenderTarget.ShaderResourceTexture, false, FResolveParams());
 }
 
 FIntPoint FRCPassPostProcessHistogram::ComputeThreadGroupCount(FIntPoint PixelExtent)
diff --git a/Engine/Source/Runtime/Renderer/Private/PostProcess/PostProcessLensFlares.cpp b/Engine/Source/Runtime/Renderer/Private/PostProcess/PostProcessLensFlares.cpp
--- a/Engine/Source/Runtime/Renderer/Private/PostProcess/PostProcessLensFlares.cpp
+++ b/Engine/Source/Runtime/Renderer/Private/PostProcess/PostProcessLensFlares.cpp
@@ -155,15 +155,13 @@ void FRCPassPostProcessLensFlares::Process(FRenderingCompositePassContext& Conte
 
 	TShaderMapRef<FPostProcessVS> VertexShader(GetGlobalShaderMap());
 
-	//@todo-rco: RHIPacketList
-	FRHICommandList& RHICmdList = FRHICommandList::GetNullRef();
-
+	
 	// setup background (bloom), can be implemented to use additive blending to avoid the read here
 	{
 		TShaderMapRef<FPostProcessLensFlareBasePS> PixelShader(GetGlobalShaderMap());
 
 		static FGlobalBoundShaderState BoundShaderState;
-		Context.RHICmdList.CheckIsNull(); // need new approach for "static FGlobalBoundShaderState" for parallel rendering
+		
 		SetGlobalBoundShaderState(Context.RHICmdList, BoundShaderState, GFilterVertexDeclaration.VertexDeclarationRHI, *VertexShader, *PixelShader);
 
 		VertexShader->SetParameters(Context);
@@ -190,7 +188,7 @@ void FRCPassPostProcessLensFlares::Process(FRenderingCompositePassContext& Conte
 		TShaderMapRef<FPostProcessLensFlaresPS> PixelShader(GetGlobalShaderMap());
 
 		static FGlobalBoundShaderState BoundShaderState;
-		Context.RHICmdList.CheckIsNull(); // need new approach for "static FGlobalBoundShaderState" for parallel rendering
+		
 		SetGlobalBoundShaderState(Context.RHICmdList, BoundShaderState, GFilterVertexDeclaration.VertexDeclarationRHI, *VertexShader, *PixelShader);
 
 		FVector2D TexScaleValue = FVector2D(TexSize2) / ViewSize2;
diff --git a/Engine/Source/Runtime/Renderer/Private/PostProcess/PostProcessMaterial.cpp b/Engine/Source/Runtime/Renderer/Private/PostProcess/PostProcessMaterial.cpp
--- a/Engine/Source/Runtime/Renderer/Private/PostProcess/PostProcessMaterial.cpp
+++ b/Engine/Source/Runtime/Renderer/Private/PostProcess/PostProcessMaterial.cpp
@@ -120,7 +120,7 @@ void FRCPassPostProcessMaterial::Process(FRenderingCompositePassContext& Context
 	const FSceneViewFamily& ViewFamily = *(View.Family);
 
 	// hacky cast
-	FRenderingCompositePassContext RenderingCompositePassContext((FViewInfo&)View);
+	FRenderingCompositePassContext RenderingCompositePassContext(Context.RHICmdList, (FViewInfo&)View);
 	RenderingCompositePassContext.Pass = this;
 
 	FIntRect SrcRect = View.ViewRect;
diff --git a/Engine/Source/Runtime/Renderer/Private/PostProcess/PostProcessMobile.cpp b/Engine/Source/Runtime/Renderer/Private/PostProcess/PostProcessMobile.cpp
--- a/Engine/Source/Runtime/Renderer/Private/PostProcess/PostProcessMobile.cpp
+++ b/Engine/Source/Runtime/Renderer/Private/PostProcess/PostProcessMobile.cpp
@@ -135,7 +135,7 @@ static void BloomSetup_SetShader(const FRenderingCompositePassContext& Context)
 	TShaderMapRef<FPostProcessBloomSetupPS_ES2<UseSunDof> > PixelShader(GetGlobalShaderMap());
 
 	static FGlobalBoundShaderState BoundShaderState;
-	Context.RHICmdList.CheckIsNull(); // need new approach for "static FGlobalBoundShaderState" for parallel rendering
+	
 	SetGlobalBoundShaderState(Context.RHICmdList, BoundShaderState, GFilterVertexDeclaration.VertexDeclarationRHI, *VertexShader, *PixelShader);
 
 	VertexShader->SetVS(Context);
@@ -352,7 +352,7 @@ void FRCPassPostProcessBloomSetupSmallES2::SetShader(const FRenderingCompositePa
 	TShaderMapRef<FPostProcessBloomSetupSmallPS_ES2> PixelShader(GetGlobalShaderMap());
 
 	static FGlobalBoundShaderState BoundShaderState;
-	Context.RHICmdList.CheckIsNull(); // need new approach for "static FGlobalBoundShaderState" for parallel rendering
+	
 	SetGlobalBoundShaderState(Context.RHICmdList, BoundShaderState, GFilterVertexDeclaration.VertexDeclarationRHI, *VertexShader, *PixelShader);
 
 	VertexShader->SetVS(Context);
@@ -564,7 +564,7 @@ void FRCPassPostProcessBloomDownES2::Process(FRenderingCompositePassContext& Con
 	TShaderMapRef<FPostProcessBloomDownPS_ES2> PixelShader(GetGlobalShaderMap());
 
 	static FGlobalBoundShaderState BoundShaderState;
-	Context.RHICmdList.CheckIsNull(); // need new approach for "static FGlobalBoundShaderState" for parallel rendering
+	
 
 	SetGlobalBoundShaderState(Context.RHICmdList, BoundShaderState, GFilterVertexDeclaration.VertexDeclarationRHI, *VertexShader, *PixelShader);
 
@@ -732,7 +732,7 @@ void FRCPassPostProcessBloomUpES2::Process(FRenderingCompositePassContext& Conte
 	TShaderMapRef<FPostProcessBloomUpPS_ES2> PixelShader(GetGlobalShaderMap());
 
 	static FGlobalBoundShaderState BoundShaderState;
-	Context.RHICmdList.CheckIsNull(); // need new approach for "static FGlobalBoundShaderState" for parallel rendering
+	
 
 	SetGlobalBoundShaderState(Context.RHICmdList, BoundShaderState, GFilterVertexDeclaration.VertexDeclarationRHI, *VertexShader, *PixelShader);
 
@@ -899,7 +899,7 @@ static void SunMask_SetShader(const FRenderingCompositePassContext& Context)
 	TShaderMapRef<FPostProcessSunMaskPS_ES2<UseFetchSunDof> > PixelShader(GetGlobalShaderMap());
 
 	static FGlobalBoundShaderState BoundShaderState;
-	Context.RHICmdList.CheckIsNull(); // need new approach for "static FGlobalBoundShaderState" for parallel rendering
+	
 	SetGlobalBoundShaderState(Context.RHICmdList, BoundShaderState, GFilterVertexDeclaration.VertexDeclarationRHI, *VertexShader, *PixelShader);
 
 	VertexShader->SetVS(Context);
@@ -1335,7 +1335,7 @@ void FRCPassPostProcessSunBlurES2::Process(FRenderingCompositePassContext& Conte
 	TShaderMapRef<FPostProcessSunBlurPS_ES2> PixelShader(GetGlobalShaderMap());
 
 	static FGlobalBoundShaderState BoundShaderState;
-	Context.RHICmdList.CheckIsNull(); // need new approach for "static FGlobalBoundShaderState" for parallel rendering
+	
 
 	SetGlobalBoundShaderState(Context.RHICmdList, BoundShaderState, GFilterVertexDeclaration.VertexDeclarationRHI, *VertexShader, *PixelShader);
 
@@ -1506,7 +1506,7 @@ static void SunMerge_SetShader(const FRenderingCompositePassContext& Context)
 	TShaderMapRef<FPostProcessSunMergePS_ES2<UseSunBloom> > PixelShader(GetGlobalShaderMap());
 
 	static FGlobalBoundShaderState BoundShaderState;
-	Context.RHICmdList.CheckIsNull(); // need new approach for "static FGlobalBoundShaderState" for parallel rendering
+	
 
 	SetGlobalBoundShaderState(Context.RHICmdList, BoundShaderState, GFilterVertexDeclaration.VertexDeclarationRHI, *VertexShader, *PixelShader);
 
@@ -1574,8 +1574,6 @@ void FRCPassPostProcessSunMergeES2::Process(FRenderingCompositePassContext& Cont
 	Context.RHICmdList.CopyToResolveTarget(DestRenderTarget.TargetableTexture, DestRenderTarget.ShaderResourceTexture, false, FResolveParams());
 
 	// Double buffer sun+bloom+vignette composite.
-	Context.RHICmdList.CheckIsNull(); // this modifies the view state, so probably won't work parallel
-
 	if(Context.View.FinalPostProcessSettings.AntiAliasingMethod == AAM_TemporalAA)
 	{
 		FSceneViewState* ViewState = (FSceneViewState*)Context.View.State;
@@ -1719,7 +1717,7 @@ void FRCPassPostProcessSunMergeSmallES2::SetShader(const FRenderingCompositePass
 	TShaderMapRef<FPostProcessSunMergeSmallPS_ES2> PixelShader(GetGlobalShaderMap());
 
 	static FGlobalBoundShaderState BoundShaderState;
-	Context.RHICmdList.CheckIsNull(); // need new approach for "static FGlobalBoundShaderState" for parallel rendering
+	
 
 	SetGlobalBoundShaderState(Context.RHICmdList, BoundShaderState, GFilterVertexDeclaration.VertexDeclarationRHI, *VertexShader, *PixelShader);
 
@@ -1771,7 +1769,7 @@ void FRCPassPostProcessSunMergeSmallES2::Process(FRenderingCompositePassContext&
 	Context.RHICmdList.CopyToResolveTarget(DestRenderTarget.TargetableTexture, DestRenderTarget.ShaderResourceTexture, false, FResolveParams());
 
 	// Double buffer sun+bloom+vignette composite.
-	Context.RHICmdList.CheckIsNull(); // this modifies the view state, so probably won't work parallel
+
 	if (Context.View.FinalPostProcessSettings.AntiAliasingMethod == AAM_TemporalAA)
 	{
 		FSceneViewState* ViewState = (FSceneViewState*)Context.View.State;
@@ -1909,7 +1907,7 @@ static void DofDown_SetShader(const FRenderingCompositePassContext& Context)
 	TShaderMapRef<FPostProcessDofDownPS_ES2<UseSun> > PixelShader(GetGlobalShaderMap());
 
 	static FGlobalBoundShaderState BoundShaderState;
-	Context.RHICmdList.CheckIsNull(); // need new approach for "static FGlobalBoundShaderState" for parallel rendering
+	
 
 	SetGlobalBoundShaderState(Context.RHICmdList, BoundShaderState, GFilterVertexDeclaration.VertexDeclarationRHI, *VertexShader, *PixelShader);
 
@@ -2117,7 +2115,7 @@ static void DofNear_SetShader(const FRenderingCompositePassContext& Context)
 	TShaderMapRef<FPostProcessDofNearPS_ES2<UseSun> > PixelShader(GetGlobalShaderMap());
 
 	static FGlobalBoundShaderState BoundShaderState;
-	Context.RHICmdList.CheckIsNull(); // need new approach for "static FGlobalBoundShaderState" for parallel rendering
+	
 
 	SetGlobalBoundShaderState(Context.RHICmdList, BoundShaderState, GFilterVertexDeclaration.VertexDeclarationRHI, *VertexShader, *PixelShader);
 
@@ -2317,7 +2315,7 @@ void FRCPassPostProcessDofBlurES2::Process(FRenderingCompositePassContext& Conte
 	TShaderMapRef<FPostProcessDofBlurPS_ES2> PixelShader(GetGlobalShaderMap());
 
 	static FGlobalBoundShaderState BoundShaderState;
-	Context.RHICmdList.CheckIsNull(); // need new approach for "static FGlobalBoundShaderState" for parallel rendering
+	
 	SetGlobalBoundShaderState(Context.RHICmdList, BoundShaderState, GFilterVertexDeclaration.VertexDeclarationRHI, *VertexShader, *PixelShader);
 
 	VertexShader->SetVS(Context);
@@ -2455,7 +2453,7 @@ static void SunAvg_SetShader(const FRenderingCompositePassContext& Context)
 	TShaderMapRef<FPostProcessSunAvgPS_ES2> PixelShader(GetGlobalShaderMap());
 
 	static FGlobalBoundShaderState BoundShaderState;
-	Context.RHICmdList.CheckIsNull(); // need new approach for "static FGlobalBoundShaderState" for parallel rendering
+	
 
 	SetGlobalBoundShaderState(Context.RHICmdList, BoundShaderState, GFilterVertexDeclaration.VertexDeclarationRHI, *VertexShader, *PixelShader);
 
@@ -2698,7 +2696,7 @@ static void Aa_SetShader(const FRenderingCompositePassContext& Context)
 	TShaderMapRef<FPostProcessAaPS_ES2> PixelShader(GetGlobalShaderMap());
 
 	static FGlobalBoundShaderState BoundShaderState;
-	Context.RHICmdList.CheckIsNull(); // need new approach for "static FGlobalBoundShaderState" for parallel rendering
+	
 
 	SetGlobalBoundShaderState(Context.RHICmdList, BoundShaderState, GFilterVertexDeclaration.VertexDeclarationRHI, *VertexShader, *PixelShader);
 
diff --git a/Engine/Source/Runtime/Renderer/Private/PostProcess/PostProcessMotionBlur.cpp b/Engine/Source/Runtime/Renderer/Private/PostProcess/PostProcessMotionBlur.cpp
--- a/Engine/Source/Runtime/Renderer/Private/PostProcess/PostProcessMotionBlur.cpp
+++ b/Engine/Source/Runtime/Renderer/Private/PostProcess/PostProcessMotionBlur.cpp
@@ -139,7 +139,7 @@ void FRCPassPostProcessMotionBlurSetup::Process(FRenderingCompositePassContext&
 	{		
 		TShaderMapRef<FPostProcessMotionBlurSetupPS > PixelShader(GetGlobalShaderMap());
 		static FGlobalBoundShaderState BoundShaderState;
-		Context.RHICmdList.CheckIsNull(); // need new approach for "static FGlobalBoundShaderState" for parallel rendering
+		
 
 		SetGlobalBoundShaderState(Context.RHICmdList, BoundShaderState, GFilterVertexDeclaration.VertexDeclarationRHI, *VertexShader, *PixelShader);
 
@@ -377,7 +377,7 @@ static void SetMotionBlurShaderTempl(const FRenderingCompositePassContext& Conte
 	TShaderMapRef<FPostProcessMotionBlurPS<Quality> > PixelShader(GetGlobalShaderMap());
 
 	static FGlobalBoundShaderState BoundShaderState;
-	Context.RHICmdList.CheckIsNull(); // need new approach for "static FGlobalBoundShaderState" for parallel rendering
+	
 
 	SetGlobalBoundShaderState(Context.RHICmdList, BoundShaderState, GFilterVertexDeclaration.VertexDeclarationRHI, *VertexShader, *PixelShader);
 
@@ -434,8 +434,6 @@ void FRCPassPostProcessMotionBlur::Process(FRenderingCompositePassContext& Conte
 	Context.RHICmdList.SetRasterizerState(TStaticRasterizerState<>::GetRHI());
 	Context.RHICmdList.SetDepthStencilState(TStaticDepthStencilState<false, CF_Always>::GetRHI());
 
-	//@todo-rco: RHIPacketList
-	FRHICommandList& RHICmdList = FRHICommandList::GetNullRef();
 	if(Quality == 1)
 	{
 		SetMotionBlurShaderTempl<1>(Context);
@@ -575,7 +573,7 @@ void FRCPassPostProcessMotionBlurRecombine::Process(FRenderingCompositePassConte
 	TShaderMapRef<FPostProcessMotionBlurRecombinePS> PixelShader(GetGlobalShaderMap());
 
 	static FGlobalBoundShaderState BoundShaderState;
-	Context.RHICmdList.CheckIsNull(); // need new approach for "static FGlobalBoundShaderState" for parallel rendering
+	
 
 	SetGlobalBoundShaderState(Context.RHICmdList, BoundShaderState, GFilterVertexDeclaration.VertexDeclarationRHI, *VertexShader, *PixelShader);
 
diff --git a/Engine/Source/Runtime/Renderer/Private/PostProcess/PostProcessSelectionOutline.cpp b/Engine/Source/Runtime/Renderer/Private/PostProcess/PostProcessSelectionOutline.cpp
--- a/Engine/Source/Runtime/Renderer/Private/PostProcess/PostProcessSelectionOutline.cpp
+++ b/Engine/Source/Runtime/Renderer/Private/PostProcess/PostProcessSelectionOutline.cpp
@@ -57,8 +57,6 @@ void FRCPassPostProcessSelectionOutlineColor::Process(FRenderingCompositePassCon
 
 		if (PrimitivesByActor.Num())
 		{
-			Context.RHICmdList.CheckIsNull(); // this uses draw dynamic elements
-
 			// 0 means no object, 1 means BSP so we start with 1
 			uint32 StencilValue = 2;
 
@@ -109,7 +107,7 @@ void FRCPassPostProcessSelectionOutlineColor::Process(FRenderingCompositePassCon
 			InnerRect.InflateRect(-1);
 
 			// We could use Clear with InnerRect but this is just an optimization - on some hardware it might do a full clear (and we cannot disable yet)
-			//			RHIClear(false, FLinearColor(0, 0, 0, 0), true, 0.0f, true, 0, InnerRect);
+			//			RHICmdList.Clear(false, FLinearColor(0, 0, 0, 0), true, 0.0f, true, 0, InnerRect);
 			// so we to 4 clears - one for each border.
 
 			// top
@@ -320,7 +318,7 @@ static void SetSelectionOutlineShaderTempl(const FRenderingCompositePassContext&
 	TShaderMapRef<FPostProcessSelectionOutlinePS<MSAASampleCount> > PixelShader(GetGlobalShaderMap());
 
 	static FGlobalBoundShaderState BoundShaderState;
-	Context.RHICmdList.CheckIsNull(); // need new approach for "static FGlobalBoundShaderState" for parallel rendering
+	
 
 	SetGlobalBoundShaderState(Context.RHICmdList, BoundShaderState, GFilterVertexDeclaration.VertexDeclarationRHI, *VertexShader, *PixelShader);
 
diff --git a/Engine/Source/Runtime/Renderer/Private/PostProcess/PostProcessSubsurface.cpp b/Engine/Source/Runtime/Renderer/Private/PostProcess/PostProcessSubsurface.cpp
--- a/Engine/Source/Runtime/Renderer/Private/PostProcess/PostProcessSubsurface.cpp
+++ b/Engine/Source/Runtime/Renderer/Private/PostProcess/PostProcessSubsurface.cpp
@@ -89,7 +89,7 @@ void SetSubsurfaceSetupShader(const FRenderingCompositePassContext& Context)
 	TShaderMapRef<FPostProcessSubsurfaceSetupPS<SpecularCorrection> > PixelShader(GetGlobalShaderMap());
 
 	static FGlobalBoundShaderState BoundShaderState;
-	Context.RHICmdList.CheckIsNull(); // need new approach for "static FGlobalBoundShaderState" for parallel rendering
+	
 
 	SetGlobalBoundShaderState(Context.RHICmdList, BoundShaderState, GFilterVertexDeclaration.VertexDeclarationRHI, *VertexShader, *PixelShader);
 
@@ -289,7 +289,7 @@ void SetSubsurfaceShader(const FRenderingCompositePassContext& Context, float In
 	TShaderMapRef<FPostProcessSubsurfacePS<Method> > PixelShader(GetGlobalShaderMap());
 
 	static FGlobalBoundShaderState BoundShaderState;
-	Context.RHICmdList.CheckIsNull(); // need new approach for "static FGlobalBoundShaderState" for parallel rendering
+	
 
 	SetGlobalBoundShaderState(Context.RHICmdList, BoundShaderState, GFilterVertexDeclaration.VertexDeclarationRHI, *VertexShader, *PixelShader);
 
diff --git a/Engine/Source/Runtime/Renderer/Private/PostProcess/PostProcessTemporalAA.cpp b/Engine/Source/Runtime/Renderer/Private/PostProcess/PostProcessTemporalAA.cpp
--- a/Engine/Source/Runtime/Renderer/Private/PostProcess/PostProcessTemporalAA.cpp
+++ b/Engine/Source/Runtime/Renderer/Private/PostProcess/PostProcessTemporalAA.cpp
@@ -229,12 +229,10 @@ void FRCPassPostProcessSSRTemporalAA::Process(FRenderingCompositePassContext& Co
 	TShaderMapRef< FPostProcessTemporalAAPS<2,0> >	PixelShader( GetGlobalShaderMap() );
 
 	static FGlobalBoundShaderState BoundShaderState;
-	Context.RHICmdList.CheckIsNull(); // need new approach for "static FGlobalBoundShaderState" for parallel rendering
+	
 
 	SetGlobalBoundShaderState(Context.RHICmdList, BoundShaderState, GFilterVertexDeclaration.VertexDeclarationRHI, *VertexShader, *PixelShader);
 
-	//@todo-rco: RHIPacketList
-	FRHICommandList& RHICmdList = FRHICommandList::GetNullRef();
 	VertexShader->SetVS(Context);
 	PixelShader->SetParameters(Context);
 
@@ -308,7 +306,7 @@ void FRCPassPostProcessDOFTemporalAA::Process(FRenderingCompositePassContext& Co
 	TShaderMapRef< FPostProcessTemporalAAPS<0,0> >	PixelShader( GetGlobalShaderMap() );
 
 	static FGlobalBoundShaderState BoundShaderState;
-	Context.RHICmdList.CheckIsNull(); // need new approach for "static FGlobalBoundShaderState" for parallel rendering
+	
 
 	SetGlobalBoundShaderState(Context.RHICmdList, BoundShaderState, GFilterVertexDeclaration.VertexDeclarationRHI, *VertexShader, *PixelShader);
 
@@ -328,7 +326,6 @@ void FRCPassPostProcessDOFTemporalAA::Process(FRenderingCompositePassContext& Co
 		EDRF_UseTriangleOptimization);
 
 	Context.RHICmdList.CopyToResolveTarget(DestRenderTarget.TargetableTexture, DestRenderTarget.ShaderResourceTexture, false, FResolveParams());
-	Context.RHICmdList.CheckIsNull(); // this modifies the view state, which is probably not safe for parallel rendering
 
 	ViewState->DOFHistoryRT = PassOutputs[0].PooledRenderTarget;
 	check( ViewState->DOFHistoryRT );
@@ -390,7 +387,7 @@ void FRCPassPostProcessLightShaftTemporalAA::Process(FRenderingCompositePassCont
 	TShaderMapRef< FPostProcessTemporalAAPS<3,0> >	PixelShader( GetGlobalShaderMap() );
 
 	static FGlobalBoundShaderState BoundShaderState;
-	Context.RHICmdList.CheckIsNull(); // need new approach for "static FGlobalBoundShaderState" for parallel rendering
+	
 
 	SetGlobalBoundShaderState(Context.RHICmdList, BoundShaderState, GFilterVertexDeclaration.VertexDeclarationRHI, *VertexShader, *PixelShader);
 
@@ -452,7 +449,7 @@ void FRCPassPostProcessTemporalAA::Process(FRenderingCompositePassContext& Conte
 
 	const FSceneRenderTargetItem& DestRenderTarget = PassOutputs[0].RequestSurface(Context);
 
-	//Context.RHICmdList.SetRenderTarget(DestRenderTarget.TargetableTexture, FTextureRHIRef());
+	//Context.SetRenderTarget(RHICmdList, DestRenderTarget.TargetableTexture, FTextureRHIRef());
 	SetRenderTarget(Context.RHICmdList, DestRenderTarget.TargetableTexture, GSceneRenderTargets.GetSceneDepthTexture());
 
 	// is optimized away if possible (RT size=view size, )
@@ -481,7 +478,7 @@ void FRCPassPostProcessTemporalAA::Process(FRenderingCompositePassContext& Conte
 			TShaderMapRef< FPostProcessTemporalAAPS<4,1> >	PixelShader( GetGlobalShaderMap() );
 
 			static FGlobalBoundShaderState BoundShaderState;
-			Context.RHICmdList.CheckIsNull(); // need new approach for "static FGlobalBoundShaderState" for parallel rendering
+			
 
 			SetGlobalBoundShaderState(Context.RHICmdList, BoundShaderState, GFilterVertexDeclaration.VertexDeclarationRHI, *VertexShader, *PixelShader);
 
@@ -493,7 +490,7 @@ void FRCPassPostProcessTemporalAA::Process(FRenderingCompositePassContext& Conte
 			TShaderMapRef< FPostProcessTemporalAAPS<1,1> >	PixelShader( GetGlobalShaderMap() );
 
 			static FGlobalBoundShaderState BoundShaderState;
-			Context.RHICmdList.CheckIsNull(); // need new approach for "static FGlobalBoundShaderState" for parallel rendering
+			
 
 			SetGlobalBoundShaderState(Context.RHICmdList, BoundShaderState, GFilterVertexDeclaration.VertexDeclarationRHI, *VertexShader, *PixelShader);
 
@@ -526,7 +523,7 @@ void FRCPassPostProcessTemporalAA::Process(FRenderingCompositePassContext& Conte
 				TShaderMapRef< FPostProcessTemporalAAPS<4,0> >	PixelShader( GetGlobalShaderMap() );
 	
 				static FGlobalBoundShaderState BoundShaderState;
-				Context.RHICmdList.CheckIsNull(); // need new approach for "static FGlobalBoundShaderState" for parallel rendering
+				
 	
 				SetGlobalBoundShaderState(Context.RHICmdList, BoundShaderState, GFilterVertexDeclaration.VertexDeclarationRHI, *VertexShader, *PixelShader);
 	
@@ -538,7 +535,7 @@ void FRCPassPostProcessTemporalAA::Process(FRenderingCompositePassContext& Conte
 				TShaderMapRef< FPostProcessTemporalAAPS<1,0> >	PixelShader( GetGlobalShaderMap() );
 	
 				static FGlobalBoundShaderState BoundShaderState;
-				Context.RHICmdList.CheckIsNull(); // need new approach for "static FGlobalBoundShaderState" for parallel rendering
+				
 	
 				SetGlobalBoundShaderState(Context.RHICmdList, BoundShaderState, GFilterVertexDeclaration.VertexDeclarationRHI, *VertexShader, *PixelShader);
 	
@@ -569,7 +566,7 @@ void FRCPassPostProcessTemporalAA::Process(FRenderingCompositePassContext& Conte
 				TShaderMapRef< FPostProcessTemporalAAPS<4,1> >	PixelShader( GetGlobalShaderMap() );
 	
 				static FGlobalBoundShaderState BoundShaderState;
-				Context.RHICmdList.CheckIsNull(); // need new approach for "static FGlobalBoundShaderState" for parallel rendering
+				
 	
 				SetGlobalBoundShaderState(Context.RHICmdList, BoundShaderState, GFilterVertexDeclaration.VertexDeclarationRHI, *VertexShader, *PixelShader);
 	
@@ -581,7 +578,7 @@ void FRCPassPostProcessTemporalAA::Process(FRenderingCompositePassContext& Conte
 				TShaderMapRef< FPostProcessTemporalAAPS<1,1> >	PixelShader( GetGlobalShaderMap() );
 	
 				static FGlobalBoundShaderState BoundShaderState;
-				Context.RHICmdList.CheckIsNull(); // need new approach for "static FGlobalBoundShaderState" for parallel rendering
+				
 	
 				SetGlobalBoundShaderState(Context.RHICmdList, BoundShaderState, GFilterVertexDeclaration.VertexDeclarationRHI, *VertexShader, *PixelShader);
 	
@@ -604,7 +601,6 @@ void FRCPassPostProcessTemporalAA::Process(FRenderingCompositePassContext& Conte
 	}
 
 	Context.RHICmdList.CopyToResolveTarget(DestRenderTarget.TargetableTexture, DestRenderTarget.ShaderResourceTexture, false, FResolveParams());
-	Context.RHICmdList.CheckIsNull(); // this modifes the viewstate which is probably not appropriate for parallel rendering
 
 	ViewState->TemporalAAHistoryRT = PassOutputs[0].PooledRenderTarget;
 	check( ViewState->TemporalAAHistoryRT );
diff --git a/Engine/Source/Runtime/Renderer/Private/PostProcess/PostProcessTestImage.cpp b/Engine/Source/Runtime/Renderer/Private/PostProcess/PostProcessTestImage.cpp
--- a/Engine/Source/Runtime/Renderer/Private/PostProcess/PostProcessTestImage.cpp
+++ b/Engine/Source/Runtime/Renderer/Private/PostProcess/PostProcessTestImage.cpp
@@ -111,7 +111,7 @@ void FRCPassPostProcessTestImage::Process(FRenderingCompositePassContext& Contex
 	TShaderMapRef<FPostProcessTestImagePS> PixelShader(GetGlobalShaderMap());
 
 	static FGlobalBoundShaderState BoundShaderState;
-	Context.RHICmdList.CheckIsNull(); // need new approach for "static FGlobalBoundShaderState" for parallel rendering
+	
 
 	SetGlobalBoundShaderState(Context.RHICmdList, BoundShaderState, GFilterVertexDeclaration.VertexDeclarationRHI, *VertexShader, *PixelShader);
 
@@ -150,7 +150,6 @@ void FRCPassPostProcessTestImage::Process(FRenderingCompositePassContext& Contex
 				return Texture;
 			}
 		} TempRenderTarget(View, (const FTexture2DRHIRef&)DestRenderTarget.TargetableTexture);
-		Context.RHICmdList.CheckIsNull(); // canvas doesn't use the cmd list
 
 		FCanvas Canvas(&TempRenderTarget, NULL, ViewFamily.CurrentRealTime, ViewFamily.CurrentWorldTime, ViewFamily.DeltaWorldTime);
 
diff --git a/Engine/Source/Runtime/Renderer/Private/PostProcess/PostProcessTonemap.cpp b/Engine/Source/Runtime/Renderer/Private/PostProcess/PostProcessTonemap.cpp
--- a/Engine/Source/Runtime/Renderer/Private/PostProcess/PostProcessTonemap.cpp
+++ b/Engine/Source/Runtime/Renderer/Private/PostProcess/PostProcessTonemap.cpp
@@ -996,7 +996,7 @@ static void SetShaderTempl(const FRenderingCompositePassContext& Context)
 	TShaderMapRef<FPostProcessTonemapPS<ConfigIndex> > PixelShader(GetGlobalShaderMap());
 
 	static FGlobalBoundShaderState BoundShaderState;
-	Context.RHICmdList.CheckIsNull(); // need new approach for "static FGlobalBoundShaderState" for parallel rendering
+	
 
 	SetGlobalBoundShaderState(Context.RHICmdList, BoundShaderState, GFilterVertexDeclaration.VertexDeclarationRHI, *VertexShader, *PixelShader);
 
@@ -1028,9 +1028,7 @@ void FRCPassPostProcessTonemap::Process(FRenderingCompositePassContext& Context)
 
 	const FSceneRenderTargetItem& DestRenderTarget = PassOutputs[0].RequestSurface(Context);
 
-	//@todo-rco: RHIPacketList
-	FRHICommandList& RHICmdList = FRHICommandList::GetNullRef();
-
+	
 	// Set the view family's render target/viewport.
 	SetRenderTarget(Context.RHICmdList, DestRenderTarget.TargetableTexture, FTextureRHIParamRef());
 
@@ -1083,7 +1081,6 @@ void FRCPassPostProcessTonemap::Process(FRenderingCompositePassContext& Context)
 	// We only release the SceneColor after the last view was processed (SplitScreen)
 	if(Context.View.Family->Views[Context.View.Family->Views.Num() - 1] == &Context.View)
 	{
-		Context.RHICmdList.CheckIsNull(); // probably can't do this parallel
 		// The RT should be released as early as possible to allow sharing of that memory for other purposes.
 		// This becomes even more important with some limited VRam (XBoxOne).
 		GSceneRenderTargets.SetSceneColor(0);
@@ -1352,7 +1349,7 @@ static void SetShaderTemplES2(const FRenderingCompositePassContext& Context, boo
 	VertexShader->bUsedFramebufferFetch = bUsedFramebufferFetch;
 
 	static FGlobalBoundShaderState BoundShaderState;
-	Context.RHICmdList.CheckIsNull(); // need new approach for "static FGlobalBoundShaderState" for parallel rendering
+	
 
 	SetGlobalBoundShaderState(Context.RHICmdList, BoundShaderState, GFilterVertexDeclaration.VertexDeclarationRHI, *VertexShader, *PixelShader);
 
@@ -1451,8 +1448,6 @@ void FRCPassPostProcessTonemapES2::Process(FRenderingCompositePassContext& Conte
 	// Double buffer tonemapper output for temporal AA.
 	if(Context.View.FinalPostProcessSettings.AntiAliasingMethod == AAM_TemporalAA)
 	{
-		Context.RHICmdList.CheckIsNull(); // modifies the viewstate, probably not appropriate for parallel work
-
 		FSceneViewState* ViewState = (FSceneViewState*)Context.View.State;
 		if(ViewState) 
 		{
diff --git a/Engine/Source/Runtime/Renderer/Private/PostProcess/PostProcessVisualizeBuffer.cpp b/Engine/Source/Runtime/Renderer/Private/PostProcess/PostProcessVisualizeBuffer.cpp
--- a/Engine/Source/Runtime/Renderer/Private/PostProcess/PostProcessVisualizeBuffer.cpp
+++ b/Engine/Source/Runtime/Renderer/Private/PostProcess/PostProcessVisualizeBuffer.cpp
@@ -116,7 +116,7 @@ FShader* FRCPassPostProcessVisualizeBuffer::SetShaderTempl(const FRenderingCompo
 	TShaderMapRef<FPostProcessVisualizeBufferPS<bDrawingTile> > PixelShader(GetGlobalShaderMap());
 
 	static FGlobalBoundShaderState BoundShaderState;
-	Context.RHICmdList.CheckIsNull(); // need new approach for "static FGlobalBoundShaderState" for parallel rendering
+	
 
 	SetGlobalBoundShaderState(Context.RHICmdList, BoundShaderState, GFilterVertexDeclaration.VertexDeclarationRHI, *VertexShader, *PixelShader);
 
@@ -176,7 +176,7 @@ void FRCPassPostProcessVisualizeBuffer::Process(FRenderingCompositePassContext&
 	TShaderMapRef<FPostProcessVisualizeBufferPS<true> > PixelShader(GetGlobalShaderMap());
 
 	static FGlobalBoundShaderState BoundShaderState;
-	Context.RHICmdList.CheckIsNull(); // need new approach for "static FGlobalBoundShaderState" for parallel rendering
+	
 	SetGlobalBoundShaderState(Context.RHICmdList, BoundShaderState, GFilterVertexDeclaration.VertexDeclarationRHI, *VertexShader, *PixelShader);
 
 	PixelShader->SetPS(Context);
@@ -251,7 +251,6 @@ void FRCPassPostProcessVisualizeBuffer::Process(FRenderingCompositePassContext&
 			return Texture;
 		}
 	} TempRenderTarget(View, (const FTexture2DRHIRef&)DestRenderTarget.TargetableTexture);
-	Context.RHICmdList.CheckIsNull(); // canvas doesn't support cmd lists
 
 	FCanvas Canvas(&TempRenderTarget, NULL, ViewFamily.CurrentRealTime, ViewFamily.CurrentWorldTime, ViewFamily.DeltaWorldTime);
 	FLinearColor LabelColor(1, 1, 0);
diff --git a/Engine/Source/Runtime/Renderer/Private/PostProcess/PostProcessVisualizeHDR.cpp b/Engine/Source/Runtime/Renderer/Private/PostProcess/PostProcessVisualizeHDR.cpp
--- a/Engine/Source/Runtime/Renderer/Private/PostProcess/PostProcessVisualizeHDR.cpp
+++ b/Engine/Source/Runtime/Renderer/Private/PostProcess/PostProcessVisualizeHDR.cpp
@@ -166,7 +166,7 @@ void FRCPassPostProcessVisualizeHDR::Process(FRenderingCompositePassContext& Con
 	TShaderMapRef<FPostProcessVisualizeHDRPS> PixelShader(GetGlobalShaderMap());
 
 	static FGlobalBoundShaderState BoundShaderState;
-	Context.RHICmdList.CheckIsNull(); // need new approach for "static FGlobalBoundShaderState" for parallel rendering
+	
 
 	SetGlobalBoundShaderState(Context.RHICmdList, BoundShaderState, GFilterVertexDeclaration.VertexDeclarationRHI, *VertexShader, *PixelShader);
 
@@ -205,7 +205,6 @@ void FRCPassPostProcessVisualizeHDR::Process(FRenderingCompositePassContext& Con
 			return Texture;
 		}
 	} TempRenderTarget(View, (const FTexture2DRHIRef&)DestRenderTarget.TargetableTexture);
-	Context.RHICmdList.CheckIsNull(); // canvas doesn't use the command list
 
 	FCanvas Canvas(&TempRenderTarget, NULL, ViewFamily.CurrentRealTime, ViewFamily.CurrentWorldTime, ViewFamily.DeltaWorldTime);
 
diff --git a/Engine/Source/Runtime/Renderer/Private/PostProcess/PostProcessWeightedSampleSum.cpp b/Engine/Source/Runtime/Renderer/Private/PostProcess/PostProcessWeightedSampleSum.cpp
--- a/Engine/Source/Runtime/Renderer/Private/PostProcess/PostProcessWeightedSampleSum.cpp
+++ b/Engine/Source/Runtime/Renderer/Private/PostProcess/PostProcessWeightedSampleSum.cpp
@@ -158,7 +158,7 @@ protected:
  * @param OutVertexShader - The vertex shader used for the filter
  */
 void SetFilterShaders(
-	FRHICommandList& RHICmdList, 
+	FRHICommandListImmediate& RHICmdList,
 	FSamplerStateRHIParamRef SamplerStateRHI,
 	FTextureRHIParamRef FilterTextureRHI,
 	FTextureRHIParamRef AdditiveTextureRHI,
@@ -170,7 +170,7 @@ void SetFilterShaders(
 	)
 {
 	check(CombineMethodInt <= 2);
-	RHICmdList.CheckIsNull(); // need new approach for "static FGlobalBoundShaderState" for parallel rendering
+	
 
 	// A macro to handle setting the filter shader for a specific number of samples.
 #define SET_FILTER_SHADER_TYPE(NumSamples) \
@@ -573,7 +573,7 @@ bool FRCPassPostProcessWeightedSampleSum::DoFastBlur() const
 	return bRet;
 }
 
-void FRCPassPostProcessWeightedSampleSum::DrawQuad(FRHICommandList& RHICmdList, bool bDoFastBlur, FIntRect SrcRect, FIntRect DestRect, bool bRequiresClear, FIntPoint DestSize, FIntPoint SrcSize, FShader* VertexShader) const
+void FRCPassPostProcessWeightedSampleSum::DrawQuad(FRHICommandListImmediate& RHICmdList, bool bDoFastBlur, FIntRect SrcRect, FIntRect DestRect, bool bRequiresClear, FIntPoint DestSize, FIntPoint SrcSize, FShader* VertexShader) const
 {
 	if (bDoFastBlur)
 	{
diff --git a/Engine/Source/Runtime/Renderer/Private/PostProcess/PostProcessWeightedSampleSum.h b/Engine/Source/Runtime/Renderer/Private/PostProcess/PostProcessWeightedSampleSum.h
--- a/Engine/Source/Runtime/Renderer/Private/PostProcess/PostProcessWeightedSampleSum.h
+++ b/Engine/Source/Runtime/Renderer/Private/PostProcess/PostProcessWeightedSampleSum.h
@@ -124,7 +124,7 @@ public:
 	static int GetIntegerKernelRadius(ERHIFeatureLevel::Type InFeatureLevel, float KernelRadius);
 
 private:
-	void DrawQuad(FRHICommandList& RHICmdList, bool bDoFastBlur, FIntRect SrcRect, FIntRect DestRect, bool bRequiresClear, FIntPoint DestSize, FIntPoint SrcSize, FShader* VertexShader) const;
+	void DrawQuad(FRHICommandListImmediate& RHICmdList, bool bDoFastBlur, FIntRect SrcRect, FIntRect DestRect, bool bRequiresClear, FIntPoint DestSize, FIntPoint SrcSize, FShader* VertexShader) const;
 	static uint32 GetMaxNumSamples(ERHIFeatureLevel::Type InFeatureLevel);
 
 	// e.g. EFS_Horiz or EFS_Vert
diff --git a/Engine/Source/Runtime/Renderer/Private/PostProcess/PostProcessing.cpp b/Engine/Source/Runtime/Renderer/Private/PostProcess/PostProcessing.cpp
--- a/Engine/Source/Runtime/Renderer/Private/PostProcess/PostProcessing.cpp
+++ b/Engine/Source/Runtime/Renderer/Private/PostProcess/PostProcessing.cpp
@@ -748,7 +748,7 @@ void OverrideRenderTarget(FRenderingCompositeOutputRef It, TRefCountPtr<IPooledR
 	}
 }
 
-void FPostProcessing::Process(const FViewInfo& View, TRefCountPtr<IPooledRenderTarget>& VelocityRT)
+void FPostProcessing::Process(FRHICommandListImmediate& RHICmdList, const FViewInfo& View, TRefCountPtr<IPooledRenderTarget>& VelocityRT)
 {
 	QUICK_SCOPE_CYCLE_COUNTER( STAT_PostProcessing_Process );
 
@@ -771,7 +771,7 @@ void FPostProcessing::Process(const FViewInfo& View, TRefCountPtr<IPooledRenderT
 	// so that the passes can register themselves to the graph
 	{
 		FMemMark Mark(FMemStack::Get());
-		FRenderingCompositePassContext CompositeContext(View);
+		FRenderingCompositePassContext CompositeContext(RHICmdList, View);
 
 		FPostprocessContext Context(CompositeContext.Graph, View);
 		
@@ -1190,7 +1190,7 @@ void FPostProcessing::Process(const FViewInfo& View, TRefCountPtr<IPooledRenderT
 }
 
 
-void FPostProcessing::ProcessES2( const FViewInfo& View, bool bUsedFramebufferFetch )
+void FPostProcessing::ProcessES2(FRHICommandListImmediate& RHICmdList, const FViewInfo& View, bool bUsedFramebufferFetch)
 {
 	check(IsInRenderingThread());
 
@@ -1205,7 +1205,7 @@ void FPostProcessing::ProcessES2( const FViewInfo& View, bool bUsedFramebufferFe
 	// so that the passes can register themselves to the graph
 	{
 		FMemMark Mark(FMemStack::Get());
-		FRenderingCompositePassContext CompositeContext(View);
+		FRenderingCompositePassContext CompositeContext(RHICmdList, View);
 
 		FPostprocessContext Context(CompositeContext.Graph, View);
 		FRenderingCompositeOutputRef BloomOutput;
diff --git a/Engine/Source/Runtime/Renderer/Private/PostProcess/PostProcessing.h b/Engine/Source/Runtime/Renderer/Private/PostProcess/PostProcessing.h
--- a/Engine/Source/Runtime/Renderer/Private/PostProcess/PostProcessing.h
+++ b/Engine/Source/Runtime/Renderer/Private/PostProcess/PostProcessing.h
@@ -67,9 +67,9 @@ class FPostProcessing
 {
 public:
 	// @param VelocityRT only valid if motion blur is supported
-	void Process(const FViewInfo& View, TRefCountPtr<IPooledRenderTarget>& VelocityRT);
+	void Process(FRHICommandListImmediate& RHICmdList, const FViewInfo& View, TRefCountPtr<IPooledRenderTarget>& VelocityRT);
 
-	void ProcessES2( const FViewInfo& View, bool bUsedFramebufferFetch );
+	void ProcessES2(FRHICommandListImmediate& RHICmdList, const FViewInfo& View, bool bUsedFramebufferFetch);
 };
 
 /** The global used for post processing. */
diff --git a/Engine/Source/Runtime/Renderer/Private/PostProcess/RenderTargetPool.cpp b/Engine/Source/Runtime/Renderer/Private/PostProcess/RenderTargetPool.cpp
--- a/Engine/Source/Runtime/Renderer/Private/PostProcess/RenderTargetPool.cpp
+++ b/Engine/Source/Runtime/Renderer/Private/PostProcess/RenderTargetPool.cpp
@@ -317,8 +317,8 @@ bool FRenderTargetPool::FindFreeElement(const FPooledRenderTargetDesc& Desc, TRe
 
 #if !(UE_BUILD_SHIPPING || UE_BUILD_TEST)
 	{
-		//@todo-rco: RHIPacketList
-		FRHICommandList& RHICmdList = FRHICommandList::GetNullRef();
+		
+		FRHICommandListImmediate& RHICmdList = FRHICommandListExecutor::GetImmediateCommandList();
 
 		if(CVarRenderTargetPoolTest.GetValueOnRenderThread())
 		{
@@ -687,10 +687,8 @@ inline void DrawBorder(FCanvas& Canvas, const FIntRect Rect, FLinearColor Color)
 	Canvas.DrawTile(Rect.Max.X - 1, Rect.Min.Y + 1, 1, Rect.Max.Y - Rect.Min.Y - 2, 0, 0, 1, 1, Color);
 }
 
-void FRenderTargetPool::PresentContent(const FSceneView& View)
+void FRenderTargetPool::PresentContent(FRHICommandListImmediate& RHICmdList, const FSceneView& View)
 {
-	//@todo-rco: RHIPacketList
-	FRHICommandList& RHICmdList = FRHICommandList::GetNullRef();
 	if (RenderTargetPoolEvents.Num())
 	{
 		AddPhaseEvent(TEXT("FrameEnd"));
diff --git a/Engine/Source/Runtime/Renderer/Private/PostProcess/RenderTargetPool.h b/Engine/Source/Runtime/Renderer/Private/PostProcess/RenderTargetPool.h
--- a/Engine/Source/Runtime/Renderer/Private/PostProcess/RenderTargetPool.h
+++ b/Engine/Source/Runtime/Renderer/Private/PostProcess/RenderTargetPool.h
@@ -220,7 +220,7 @@ public:
 	void AddPhaseEvent(const TCHAR* InPhaseName);
 
 	/** renders the VisualizeTextureContent to the current render target */
-	void PresentContent(const FSceneView& View);
+	void PresentContent(FRHICommandListImmediate& RHICmdList, const FSceneView& View);
 
 	FVisualizeTexture VisualizeTexture;
 
diff --git a/Engine/Source/Runtime/Renderer/Private/PostProcess/RenderingCompositionGraph.cpp b/Engine/Source/Runtime/Renderer/Private/PostProcess/RenderingCompositionGraph.cpp
--- a/Engine/Source/Runtime/Renderer/Private/PostProcess/RenderingCompositionGraph.cpp
+++ b/Engine/Source/Runtime/Renderer/Private/PostProcess/RenderingCompositionGraph.cpp
@@ -103,11 +103,11 @@ FAutoConsoleCommand CmdCompositionGraphDebug(
 	);
 #endif
 
-FRenderingCompositePassContext::FRenderingCompositePassContext(const FViewInfo& InView/*, const FSceneRenderTargetItem& InRenderTargetItem*/)
+FRenderingCompositePassContext::FRenderingCompositePassContext(FRHICommandListImmediate& InRHICmdList, const FViewInfo& InView/*, const FSceneRenderTargetItem& InRenderTargetItem*/)
 	: View(InView)
 	//		, CompositingOutputRTItem(InRenderTargetItem)
 	, Pass(0)
-	, RHICmdList(FRHICommandList::GetNullRef())
+	, RHICmdList(InRHICmdList)
 	, ViewPortRect(0, 0, 0 ,0)
 {
 	check(!IsViewportValid());
@@ -251,7 +251,7 @@ void FRenderingCompositionGraph::DumpOutputToFile(FRenderingCompositePassContext
 	TArray<FColor> Bitmap;
 
 	FIntPoint Extent = Context.View.ViewRect.Size();
-	RHIReadSurfaceData(Texture, FIntRect(0, 0, Extent.X, Extent.Y), Bitmap, ReadDataFlags);
+	Context.RHICmdList.ReadSurfaceData(Texture, FIntRect(0, 0, Extent.X, Extent.Y), Bitmap, ReadDataFlags);
 
 	static TCHAR File[MAX_SPRINTF];
 	FCString::Sprintf( File, TEXT("%s.bmp"), *Filename);
@@ -540,7 +540,7 @@ void FRenderingCompositionGraph::RecursivelyProcess(const FRenderingCompositeOut
 				TArray<FColor>* OutputColorArray = Pass->GetOutputColorArray((EPassOutputId)OutputId);
 				if (OutputColorArray)
 				{
-					RHIReadSurfaceData(
+					Context.RHICmdList.ReadSurfaceData(
 						PassOutput->PooledRenderTarget->GetRenderTargetItem().TargetableTexture,
 						Context.View.ViewRect,
 						*OutputColorArray,
@@ -667,22 +667,21 @@ void FPostProcessPassParameters::Bind(const FShaderParameterMap& ParameterMap)
 
 void FPostProcessPassParameters::SetPS(const FPixelShaderRHIParamRef& ShaderRHI, const FRenderingCompositePassContext& Context, FSamplerStateRHIParamRef Filter, bool bWhiteIfNoTexture, FSamplerStateRHIParamRef* FilterOverrideArray)
 {
-	Set(Context.RHICmdList, ShaderRHI, Context, Filter, bWhiteIfNoTexture, FilterOverrideArray);
+	Set(ShaderRHI, Context, Filter, bWhiteIfNoTexture, FilterOverrideArray);
 }
 
 void FPostProcessPassParameters::SetCS(const FComputeShaderRHIParamRef& ShaderRHI, const FRenderingCompositePassContext& Context, FSamplerStateRHIParamRef Filter, bool bWhiteIfNoTexture, FSamplerStateRHIParamRef* FilterOverrideArray)
 {
-	Set(Context.RHICmdList, ShaderRHI, Context, Filter, bWhiteIfNoTexture, FilterOverrideArray);
+	Set(ShaderRHI, Context, Filter, bWhiteIfNoTexture, FilterOverrideArray);
 }
 
 void FPostProcessPassParameters::SetVS(const FVertexShaderRHIParamRef& ShaderRHI, const FRenderingCompositePassContext& Context, FSamplerStateRHIParamRef Filter, bool bWhiteIfNoTexture, FSamplerStateRHIParamRef* FilterOverrideArray)
 {
-	Set(Context.RHICmdList, ShaderRHI, Context, Filter, bWhiteIfNoTexture, FilterOverrideArray);
+	Set(ShaderRHI, Context, Filter, bWhiteIfNoTexture, FilterOverrideArray);
 }
 
 template< typename ShaderRHIParamRef >
 void FPostProcessPassParameters::Set(
-	FRHICommandList& RHICmdList,
 	const ShaderRHIParamRef& ShaderRHI,
 	const FRenderingCompositePassContext& Context,
 	FSamplerStateRHIParamRef Filter,
@@ -700,11 +699,11 @@ void FPostProcessPassParameters::Set(
 	// but not both
 	check(!FilterOverrideArray || !Filter);
 
-	RHICmdList.CheckIsNull();
+	FRHICommandListImmediate& RHICmdList = Context.RHICmdList;
 
 	if(BilinearTextureSampler0.IsBound())
 	{
-		RHISetShaderSampler(
+		RHICmdList.SetShaderSampler(
 			ShaderRHI, 
 			BilinearTextureSampler0.GetBaseIndex(), 
 			TStaticSamplerState<SF_Bilinear>::GetRHI()
@@ -713,7 +712,7 @@ void FPostProcessPassParameters::Set(
 
 	if(BilinearTextureSampler1.IsBound())
 	{
-		RHISetShaderSampler(
+		RHICmdList.SetShaderSampler(
 			ShaderRHI, 
 			BilinearTextureSampler1.GetBaseIndex(), 
 			TStaticSamplerState<SF_Bilinear>::GetRHI()
@@ -824,7 +823,6 @@ void FPostProcessPassParameters::Set(
 
 #define IMPLEMENT_POST_PROCESS_PARAM_SET( ShaderRHIParamRef ) \
 	template void FPostProcessPassParameters::Set< ShaderRHIParamRef >( \
-		FRHICommandList& RHICmdList,					\
 		const ShaderRHIParamRef& ShaderRHI,				\
 		const FRenderingCompositePassContext& Context,	\
 		FSamplerStateRHIParamRef Filter,				\
diff --git a/Engine/Source/Runtime/Renderer/Private/PostProcess/RenderingCompositionGraph.h b/Engine/Source/Runtime/Renderer/Private/PostProcess/RenderingCompositionGraph.h
--- a/Engine/Source/Runtime/Renderer/Private/PostProcess/RenderingCompositionGraph.h
+++ b/Engine/Source/Runtime/Renderer/Private/PostProcess/RenderingCompositionGraph.h
@@ -79,7 +79,7 @@ private:
 
 struct FRenderingCompositePassContext
 {
-	FRenderingCompositePassContext(const FViewInfo& InView);
+	FRenderingCompositePassContext(FRHICommandListImmediate& RHICmdList, const FViewInfo& InView);
 
 	~FRenderingCompositePassContext();
 
@@ -136,7 +136,7 @@ struct FRenderingCompositePassContext
 
 	FRenderingCompositionGraph Graph;
 
-	FRHICommandList& RHICmdList;
+	FRHICommandListImmediate& RHICmdList;
 
 private:
 	// cached state to map between ScreenPos and pixels
diff --git a/Engine/Source/Runtime/Renderer/Private/PostProcess/SceneFilterRendering.cpp b/Engine/Source/Runtime/Renderer/Private/PostProcess/SceneFilterRendering.cpp
--- a/Engine/Source/Runtime/Renderer/Private/PostProcess/SceneFilterRendering.cpp
+++ b/Engine/Source/Runtime/Renderer/Private/PostProcess/SceneFilterRendering.cpp
@@ -183,7 +183,7 @@ void DrawRectangle(
 }
 
 void DrawTransformedRectangle(
-	FRHICommandList& RHICmdList,
+	FRHICommandListImmediate& RHICmdList,
 	float X,
 	float Y,
 	float SizeX,
