UE4 - rework RHI user-facing APIs to be safer and incorporate the desired sematics

diff --git a/Engine/Source/Runtime/Renderer/Private/PostProcess/SceneRenderTargets.cpp b/Engine/Source/Runtime/Renderer/Private/PostProcess/SceneRenderTargets.cpp
--- a/Engine/Source/Runtime/Renderer/Private/PostProcess/SceneRenderTargets.cpp
+++ b/Engine/Source/Runtime/Renderer/Private/PostProcess/SceneRenderTargets.cpp
@@ -223,13 +223,13 @@ void FSceneRenderTargets::Allocate(const FSceneViewFamily& ViewFamily)
 }
 
 /** Clears the GBuffer render targets to default values. */
-void FSceneRenderTargets::ClearGBufferTargets(const FLinearColor& ClearColor)
+void FSceneRenderTargets::ClearGBufferTargets(FRHICommandListImmediate& RHICmdList, const FLinearColor& ClearColor)
 {
 	SCOPED_DRAW_EVENT(ClearGBufferTargets, DEC_SCENE_ITEMS);
 
 	// Clear GBufferA, GBufferB, GBufferC, GBufferD, GBufferE
 	{
-		GSceneRenderTargets.BeginRenderingSceneColor(true);
+		GSceneRenderTargets.BeginRenderingSceneColor(RHICmdList, true);
 
 		int32 NumToClear = GSceneRenderTargets.GetNumGBufferTargets();
 		if (NumToClear > 1)
@@ -237,16 +237,16 @@ void FSceneRenderTargets::ClearGBufferTargets(const FLinearColor& ClearColor)
 			// Using 0 and 1 ensures we go through the fast path on Intel integrated GPUs.
 			// Normal could be 0.5,0.5,0.5 but then it would not use the fast path.
 			FLinearColor ClearColors[6] = {ClearColor, FLinearColor(0, 0, 0, 0), FLinearColor(0,0,0,0), FLinearColor(0,0,0,0), FLinearColor(0,1,1,1), FLinearColor(1,1,1,1)};
-			RHIClearMRT(true, NumToClear, ClearColors, false, 0, false, 0, FIntRect());
+			RHICmdList.ClearMRT(true, NumToClear, ClearColors, false, 0, false, 0, FIntRect());
 		}
 		else
 		{
-			RHIClear(true, ClearColor, false, 0, false, 0, FIntRect());
+			RHICmdList.Clear(true, ClearColor, false, 0, false, 0, FIntRect());
 		}
 	}
 }
 
-void FSceneRenderTargets::BeginRenderingSceneColor( bool bGBufferPass )
+void FSceneRenderTargets::BeginRenderingSceneColor(FRHICommandListImmediate& RHICmdList, bool bGBufferPass)
 {
 	SCOPED_DRAW_EVENT(BeginRenderingSceneColor, DEC_SCENE_ITEMS);
 
@@ -256,9 +256,7 @@ void FSceneRenderTargets::BeginRenderingSceneColor( bool bGBufferPass )
 	}
 	
 	AllocSceneColor();
-	//@todo-rco: RHIPacketList
-	FRHICommandList& RHICmdList = FRHICommandList::GetNullRef();
-
+	
 	// Set the scene color surface as the render target, and the scene depth surface as the depth-stencil target.
 	if (bGBufferPass && GRHIFeatureLevel >= ERHIFeatureLevel::SM4)
 	{
@@ -589,26 +587,24 @@ void FSceneRenderTargets::AdjustGBufferRefCount(int Delta)
 	}
 }
 
-void FSceneRenderTargets::FinishRenderingSceneColor(bool bKeepChanges, const FResolveRect& ResolveRect)
+void FSceneRenderTargets::FinishRenderingSceneColor(FRHICommandListImmediate& RHICmdList, bool bKeepChanges, const FResolveRect& ResolveRect)
 {
 	SCOPED_DRAW_EVENT(FinishRenderingSceneColor, DEC_SCENE_ITEMS);
 
 	if(bKeepChanges)
 	{
-		ResolveSceneColor();
+		ResolveSceneColor(RHICmdList);
 	}
 }
 
-bool FSceneRenderTargets::BeginRenderingCustomDepth(bool bPrimitives)
+bool FSceneRenderTargets::BeginRenderingCustomDepth(FRHICommandListImmediate& RHICmdList, bool bPrimitives)
 {
 	IPooledRenderTarget* CustomDepth = RequestCustomDepth(bPrimitives);
 
 	if(CustomDepth)
 	{
 		SCOPED_DRAW_EVENT(BeginRenderingCustomDepth, DEC_SCENE_ITEMS);
-		//@todo-rco: RHIPacketList
-		FRHICommandList& RHICmdList = FRHICommandList::GetNullRef();
-
+		
 		SetRenderTarget(RHICmdList, FTextureRHIRef(), CustomDepth->GetRenderTargetItem().ShaderResourceTexture);
 
 		return true;
@@ -617,11 +613,11 @@ bool FSceneRenderTargets::BeginRenderingCustomDepth(bool bPrimitives)
 	return false;
 }
 
-void FSceneRenderTargets::FinishRenderingCustomDepth(const FResolveRect& ResolveRect)
+void FSceneRenderTargets::FinishRenderingCustomDepth(FRHICommandListImmediate& RHICmdList, const FResolveRect& ResolveRect)
 {
 	SCOPED_DRAW_EVENT(FinishRenderingCustomDepth, DEC_SCENE_ITEMS);
 
-	RHICopyToResolveTarget(SceneColor->GetRenderTargetItem().TargetableTexture, SceneColor->GetRenderTargetItem().ShaderResourceTexture, true, FResolveParams(ResolveRect));
+	RHICmdList.CopyToResolveTarget(SceneColor->GetRenderTargetItem().TargetableTexture, SceneColor->GetRenderTargetItem().ShaderResourceTexture, true, FResolveParams(ResolveRect));
 
 	bCustomDepthIsValid = true;
 }
@@ -629,77 +625,67 @@ void FSceneRenderTargets::FinishRenderingCustomDepth(const FResolveRect& Resolve
 /**
 * Saves a previously rendered scene color target
 */
-void FSceneRenderTargets::ResolveSceneColor(const FResolveRect& ResolveRect)
+void FSceneRenderTargets::ResolveSceneColor(FRHICommandListImmediate& RHICmdList, const FResolveRect& ResolveRect)
 {
     SCOPED_DRAW_EVENT(ResolveSceneColor, DEC_SCENE_ITEMS);
 
-	RHICopyToResolveTarget(GetSceneColorSurface(), GetSceneColorTexture(), true, FResolveParams(ResolveRect));
+	RHICmdList.CopyToResolveTarget(GetSceneColorSurface(), GetSceneColorTexture(), true, FResolveParams(ResolveRect));
 }
 
 /** Resolves the GBuffer targets so that their resolved textures can be sampled. */
-void FSceneRenderTargets::ResolveGBufferSurfaces(const FResolveRect& ResolveRect)
+void FSceneRenderTargets::ResolveGBufferSurfaces(FRHICommandListImmediate& RHICmdList, const FResolveRect& ResolveRect)
 {
 	if (GRHIFeatureLevel >= ERHIFeatureLevel::SM4)
 	{
 		SCOPED_DRAW_EVENT(ResolveGBufferSurfaces, DEC_SCENE_ITEMS);
 
-		RHICopyToResolveTarget(GSceneRenderTargets.GBufferA->GetRenderTargetItem().TargetableTexture, GSceneRenderTargets.GBufferA->GetRenderTargetItem().ShaderResourceTexture, true, FResolveParams(ResolveRect));
-		RHICopyToResolveTarget(GSceneRenderTargets.GBufferB->GetRenderTargetItem().TargetableTexture, GSceneRenderTargets.GBufferB->GetRenderTargetItem().ShaderResourceTexture, true, FResolveParams(ResolveRect));
-		RHICopyToResolveTarget(GSceneRenderTargets.GBufferC->GetRenderTargetItem().TargetableTexture, GSceneRenderTargets.GBufferC->GetRenderTargetItem().ShaderResourceTexture, true, FResolveParams(ResolveRect));
-		RHICopyToResolveTarget(GSceneRenderTargets.GBufferD->GetRenderTargetItem().TargetableTexture, GSceneRenderTargets.GBufferD->GetRenderTargetItem().ShaderResourceTexture, true, FResolveParams(ResolveRect));
+		RHICmdList.CopyToResolveTarget(GSceneRenderTargets.GBufferA->GetRenderTargetItem().TargetableTexture, GSceneRenderTargets.GBufferA->GetRenderTargetItem().ShaderResourceTexture, true, FResolveParams(ResolveRect));
+		RHICmdList.CopyToResolveTarget(GSceneRenderTargets.GBufferB->GetRenderTargetItem().TargetableTexture, GSceneRenderTargets.GBufferB->GetRenderTargetItem().ShaderResourceTexture, true, FResolveParams(ResolveRect));
+		RHICmdList.CopyToResolveTarget(GSceneRenderTargets.GBufferC->GetRenderTargetItem().TargetableTexture, GSceneRenderTargets.GBufferC->GetRenderTargetItem().ShaderResourceTexture, true, FResolveParams(ResolveRect));
+		RHICmdList.CopyToResolveTarget(GSceneRenderTargets.GBufferD->GetRenderTargetItem().TargetableTexture, GSceneRenderTargets.GBufferD->GetRenderTargetItem().ShaderResourceTexture, true, FResolveParams(ResolveRect));
 		
 		if (bAllowStaticLighting)
 		{
-			RHICopyToResolveTarget(GSceneRenderTargets.GBufferE->GetRenderTargetItem().TargetableTexture, GSceneRenderTargets.GBufferE->GetRenderTargetItem().ShaderResourceTexture, true, FResolveParams(ResolveRect));
+			RHICmdList.CopyToResolveTarget(GSceneRenderTargets.GBufferE->GetRenderTargetItem().TargetableTexture, GSceneRenderTargets.GBufferE->GetRenderTargetItem().ShaderResourceTexture, true, FResolveParams(ResolveRect));
 		}
 	}
 }
 
-void FSceneRenderTargets::BeginRenderingPrePass()
+void FSceneRenderTargets::BeginRenderingPrePass(FRHICommandListImmediate& RHICmdList)
 {
 	SCOPED_DRAW_EVENT(BeginRenderingPrePass, DEC_SCENE_ITEMS);
 
 	// Set the scene depth surface and a DUMMY buffer as color buffer
 	// (as long as it's the same dimension as the depth buffer),
-	//@todo-rco: RHIPacketList
-	FRHICommandList& RHICmdList = FRHICommandList::GetNullRef();
-
+	
 	SetRenderTarget(RHICmdList, FTextureRHIRef(), GetSceneDepthSurface());
 }
 
-void FSceneRenderTargets::FinishRenderingPrePass()
+void FSceneRenderTargets::FinishRenderingPrePass(FRHICommandListImmediate& RHICmdList)
 {
 	SCOPED_DRAW_EVENT(FinishRenderingPrePass, DEC_SCENE_ITEMS);
-
-	//@todo-rco: RHIPacketList
-	FRHICommandList& RHICmdList = FRHICommandList::GetNullRef();
 	GRenderTargetPool.VisualizeTexture.SetCheckPoint(RHICmdList, SceneDepthZ);
 }
 
-void FSceneRenderTargets::BeginRenderingShadowDepth()
+void FSceneRenderTargets::BeginRenderingShadowDepth(FRHICommandListImmediate& RHICmdList)
 {
-	//@todo-rco: RHIPacketList
-	FRHICommandList& RHICmdList = FRHICommandList::GetNullRef();
 	GRenderTargetPool.VisualizeTexture.SetCheckPoint(RHICmdList, ShadowDepthZ);
-	
 	SetRenderTarget(RHICmdList, FTextureRHIRef(), GetShadowDepthZSurface());
 }
 
-void FSceneRenderTargets::BeginRenderingCubeShadowDepth(int32 ShadowResolution)
+void FSceneRenderTargets::BeginRenderingCubeShadowDepth(FRHICommandListImmediate& RHICmdList, int32 ShadowResolution)
 {
 	SCOPED_DRAW_EVENT(BeginRenderingCubeShadowDepth, DEC_SCENE_ITEMS);
-	//@todo-rco: RHIPacketList
-	FRHICommandList& RHICmdList = FRHICommandList::GetNullRef();
 	SetRenderTarget(RHICmdList, FTextureRHIRef(), GetCubeShadowDepthZSurface(ShadowResolution));
 }
 
-void FSceneRenderTargets::FinishRenderingShadowDepth(const FResolveRect& ResolveRect)
+void FSceneRenderTargets::FinishRenderingShadowDepth(FRHICommandListImmediate& RHICmdList, const FResolveRect& ResolveRect)
 {
 	// Resolve the shadow depth z surface.
-	RHICopyToResolveTarget(GetShadowDepthZSurface(), GetShadowDepthZTexture(), false, FResolveParams(ResolveRect));
+	RHICmdList.CopyToResolveTarget(GetShadowDepthZSurface(), GetShadowDepthZTexture(), false, FResolveParams(ResolveRect));
 }
 
-void FSceneRenderTargets::BeginRenderingReflectiveShadowMap( FLightPropagationVolume* Lpv )
+void FSceneRenderTargets::BeginRenderingReflectiveShadowMap(FRHICommandListImmediate& RHICmdList, FLightPropagationVolume* Lpv)
 {
 	FTextureRHIParamRef RenderTargets[2];
 	RenderTargets[0] = GetReflectiveShadowMapNormalSurface();
@@ -711,20 +697,16 @@ void FSceneRenderTargets::BeginRenderingReflectiveShadowMap( FLightPropagationVo
 	Uavs[1] = Lpv->GetGvListHeadBufferUav();
 	Uavs[2] = Lpv->GetVplListBufferUav();
 	Uavs[3] = Lpv->GetVplListHeadBufferUav();
-	//@todo-rco: RHIPacketList
-	FRHICommandList& RHICmdList = FRHICommandList::GetNullRef();
+	
 	SetRenderTargets(RHICmdList, ARRAY_COUNT(RenderTargets), RenderTargets, GetReflectiveShadowMapDepthSurface(), 4, Uavs);
 }
 
-void FSceneRenderTargets::FinishRenderingReflectiveShadowMap(const FResolveRect& ResolveRect)
+void FSceneRenderTargets::FinishRenderingReflectiveShadowMap(FRHICommandListImmediate& RHICmdList, const FResolveRect& ResolveRect)
 {
-	//@todo-rco: RHIPacketList
-	FRHICommandList& RHICmdList = FRHICommandList::GetNullRef();
-
 	// Resolve the shadow depth z surface.
-	RHICopyToResolveTarget(GetReflectiveShadowMapDepthSurface(), GetReflectiveShadowMapDepthTexture(), false, FResolveParams(ResolveRect));
-	RHICopyToResolveTarget(GetReflectiveShadowMapDiffuseSurface(), GetReflectiveShadowMapDiffuseTexture(), false, FResolveParams(ResolveRect));
-	RHICopyToResolveTarget(GetReflectiveShadowMapNormalSurface(), GetReflectiveShadowMapNormalTexture(), false, FResolveParams(ResolveRect));
+	RHICmdList.CopyToResolveTarget(GetReflectiveShadowMapDepthSurface(), GetReflectiveShadowMapDepthTexture(), false, FResolveParams(ResolveRect));
+	RHICmdList.CopyToResolveTarget(GetReflectiveShadowMapDiffuseSurface(), GetReflectiveShadowMapDiffuseTexture(), false, FResolveParams(ResolveRect));
+	RHICmdList.CopyToResolveTarget(GetReflectiveShadowMapNormalSurface(), GetReflectiveShadowMapNormalTexture(), false, FResolveParams(ResolveRect));
 
 	// Unset render targets
 	FTextureRHIParamRef RenderTargets[2] = {NULL};
@@ -732,68 +714,60 @@ void FSceneRenderTargets::FinishRenderingReflectiveShadowMap(const FResolveRect&
 	SetRenderTargets(RHICmdList, ARRAY_COUNT(RenderTargets), RenderTargets, FTextureRHIParamRef(), 2, Uavs);
 }
 
-void FSceneRenderTargets::FinishRenderingCubeShadowDepth(int32 ShadowResolution, const FResolveParams& ResolveParams)
+void FSceneRenderTargets::FinishRenderingCubeShadowDepth(FRHICommandListImmediate& RHICmdList, int32 ShadowResolution, const FResolveParams& ResolveParams)
 {
 	SCOPED_DRAW_EVENT(FinishRenderingCubeShadowDepth, DEC_SCENE_ITEMS);
-	RHICopyToResolveTarget(GetCubeShadowDepthZSurface(ShadowResolution), GetCubeShadowDepthZTexture(ShadowResolution), false, ResolveParams);
+	RHICmdList.CopyToResolveTarget(GetCubeShadowDepthZSurface(ShadowResolution), GetCubeShadowDepthZTexture(ShadowResolution), false, ResolveParams);
 }
 
-void FSceneRenderTargets::BeginRenderingSceneAlphaCopy()
+void FSceneRenderTargets::BeginRenderingSceneAlphaCopy(FRHICommandListImmediate& RHICmdList)
 {
 	SCOPED_DRAW_EVENT(BeginRenderingSceneAlphaCopy, DEC_SCENE_ITEMS);
-	//@todo-rco: RHIPacketList
-	FRHICommandList& RHICmdList = FRHICommandList::GetNullRef();
 	GRenderTargetPool.VisualizeTexture.SetCheckPoint(RHICmdList, GSceneRenderTargets.SceneAlphaCopy);
 	SetRenderTarget(RHICmdList, GetSceneAlphaCopySurface(), 0);
 }
 
-void FSceneRenderTargets::FinishRenderingSceneAlphaCopy()
+void FSceneRenderTargets::FinishRenderingSceneAlphaCopy(FRHICommandListImmediate& RHICmdList)
 {
 	SCOPED_DRAW_EVENT(FinishRenderingSceneAlphaCopy, DEC_SCENE_ITEMS);
-	//@todo-rco: RHIPacketList
-	FRHICommandList& RHICmdList = FRHICommandList::GetNullRef();
-	RHICopyToResolveTarget(GetSceneAlphaCopySurface(), SceneAlphaCopy->GetRenderTargetItem().ShaderResourceTexture, false, FResolveParams(FResolveRect()));
+	RHICmdList.CopyToResolveTarget(GetSceneAlphaCopySurface(), SceneAlphaCopy->GetRenderTargetItem().ShaderResourceTexture, false, FResolveParams(FResolveRect()));
 	GRenderTargetPool.VisualizeTexture.SetCheckPoint(RHICmdList, GSceneRenderTargets.SceneAlphaCopy);
 }
 
 
-void FSceneRenderTargets::BeginRenderingLightAttenuation()
+void FSceneRenderTargets::BeginRenderingLightAttenuation(FRHICommandListImmediate& RHICmdList)
 {
 	SCOPED_DRAW_EVENT(BeginRenderingLightAttenuation, DEC_SCENE_ITEMS);
 
 	AllocLightAttenuation();
 
-	//@todo-rco: RHIPacketList
-	FRHICommandList& RHICmdList = FRHICommandList::GetNullRef();
 	GRenderTargetPool.VisualizeTexture.SetCheckPoint(RHICmdList, GSceneRenderTargets.GetLightAttenuation());
 
 	// Set the light attenuation surface as the render target, and the scene depth buffer as the depth-stencil surface.
-	RHISetRenderTarget(GetLightAttenuationSurface(),GetSceneDepthSurface());
+	SetRenderTarget(RHICmdList, GetLightAttenuationSurface(),GetSceneDepthSurface());
 }
 
-void FSceneRenderTargets::FinishRenderingLightAttenuation()
+void FSceneRenderTargets::FinishRenderingLightAttenuation(FRHICommandListImmediate& RHICmdList)
 {
 	SCOPED_DRAW_EVENT(FinishRenderingLightAttenuation, DEC_SCENE_ITEMS);
 
 	// Resolve the light attenuation surface.
-	RHICopyToResolveTarget(GetLightAttenuationSurface(), LightAttenuation->GetRenderTargetItem().ShaderResourceTexture, false, FResolveParams(FResolveRect()));
+	RHICmdList.CopyToResolveTarget(GetLightAttenuationSurface(), LightAttenuation->GetRenderTargetItem().ShaderResourceTexture, false, FResolveParams(FResolveRect()));
 	
-	//@todo-rco: RHIPacketList
-	FRHICommandList& RHICmdList = FRHICommandList::GetNullRef();
 	GRenderTargetPool.VisualizeTexture.SetCheckPoint(RHICmdList, GSceneRenderTargets.GetLightAttenuation());
 }
 
-void FSceneRenderTargets::BeginRenderingTranslucency(const FViewInfo& View)
+void FSceneRenderTargets::BeginRenderingTranslucency(FRHICommandListImmediate& RHICmdList, const FViewInfo& View)
 {
 	// Use the scene color buffer.
-	GSceneRenderTargets.BeginRenderingSceneColor();
+	GSceneRenderTargets.BeginRenderingSceneColor(RHICmdList);
 	
 	// viewport to match view size
-	RHISetViewport(View.ViewRect.Min.X, View.ViewRect.Min.Y, 0.0f, View.ViewRect.Max.X, View.ViewRect.Max.Y, 1.0f);
+	RHICmdList.SetViewport(View.ViewRect.Min.X, View.ViewRect.Min.Y, 0.0f, View.ViewRect.Max.X, View.ViewRect.Max.Y, 1.0f);
 }
 
 
-bool FSceneRenderTargets::BeginRenderingSeparateTranslucency(const FViewInfo& View, bool bFirstTimeThisFrame)
+bool FSceneRenderTargets::BeginRenderingSeparateTranslucency(FRHICommandListImmediate& RHICmdList, const FViewInfo& View, bool bFirstTimeThisFrame)
 {
 	if(IsSeparateTranslucencyActive(View))
 	{
@@ -810,14 +784,14 @@ bool FSceneRenderTargets::BeginRenderingSeparateTranslucency(const FViewInfo& Vi
 		TRefCountPtr<IPooledRenderTarget>& SeparateTranslucency = ViewState->GetSeparateTranslucency(View);
 
 		// Use a separate render target for translucency
-		RHISetRenderTarget(SeparateTranslucency->GetRenderTargetItem().TargetableTexture, GetSceneDepthSurface());
+		SetRenderTarget(RHICmdList, SeparateTranslucency->GetRenderTargetItem().TargetableTexture, GetSceneDepthSurface());
 		
 		if(bFirstTimeThisFrame)
 		{
-			RHIClear(true, FLinearColor(0, 0, 0, 1), false, 0, false, 0, FIntRect());
+			RHICmdList.Clear(true, FLinearColor(0, 0, 0, 1), false, 0, false, 0, FIntRect());
 		}
 
-		RHISetViewport(View.ViewRect.Min.X, View.ViewRect.Min.Y, 0.0f, View.ViewRect.Max.X, View.ViewRect.Max.Y, 1.0f);
+		RHICmdList.SetViewport(View.ViewRect.Min.X, View.ViewRect.Min.Y, 0.0f, View.ViewRect.Max.X, View.ViewRect.Max.Y, 1.0f);
 
 		return true;
 	}
@@ -825,25 +799,25 @@ bool FSceneRenderTargets::BeginRenderingSeparateTranslucency(const FViewInfo& Vi
 	return false;
 }
 
-void FSceneRenderTargets::FinishRenderingSeparateTranslucency(const FViewInfo& View)
+void FSceneRenderTargets::FinishRenderingSeparateTranslucency(FRHICommandListImmediate& RHICmdList, const FViewInfo& View)
 {
 	if(IsSeparateTranslucencyActive(View))
 	{
 		FSceneViewState* ViewState = (FSceneViewState*)View.State;
 		TRefCountPtr<IPooledRenderTarget>& SeparateTranslucency = ViewState->GetSeparateTranslucency(View);
 
-		RHICopyToResolveTarget(SeparateTranslucency->GetRenderTargetItem().TargetableTexture, SeparateTranslucency->GetRenderTargetItem().ShaderResourceTexture, true, FResolveParams());
+		RHICmdList.CopyToResolveTarget(SeparateTranslucency->GetRenderTargetItem().TargetableTexture, SeparateTranslucency->GetRenderTargetItem().ShaderResourceTexture, true, FResolveParams());
 	}
 }
 
-void FSceneRenderTargets::ResolveSceneDepthTexture()
+void FSceneRenderTargets::ResolveSceneDepthTexture(FRHICommandListImmediate& RHICmdList)
 {
 	SCOPED_DRAW_EVENT(ResolveSceneDepthTexture, DEC_SCENE_ITEMS);
 
-	RHICopyToResolveTarget(GetSceneDepthSurface(), GetSceneDepthTexture(), true, FResolveParams());
+	RHICmdList.CopyToResolveTarget(GetSceneDepthSurface(), GetSceneDepthTexture(), true, FResolveParams());
 }
 
-void FSceneRenderTargets::ResolveSceneDepthToAuxiliaryTexture()
+void FSceneRenderTargets::ResolveSceneDepthToAuxiliaryTexture(FRHICommandListImmediate& RHICmdList)
 {
 	// Resolve the scene depth to an auxiliary texture when SM3/SM4 is in use. This needs to happen so the auxiliary texture can be bound as a shader parameter
 	// while the primary scene depth texture can be bound as the target. Simultaneously binding a single DepthStencil resource as a parameter and target
@@ -852,7 +826,7 @@ void FSceneRenderTargets::ResolveSceneDepthToAuxiliaryTexture()
 	{
 		SCOPED_DRAW_EVENT(ResolveSceneDepthToAuxiliaryTexture, DEC_SCENE_ITEMS);
 
-		RHICopyToResolveTarget(GetSceneDepthSurface(), GetAuxiliarySceneDepthTexture(), true, FResolveParams());
+		RHICmdList.CopyToResolveTarget(GetSceneDepthSurface(), GetAuxiliarySceneDepthTexture(), true, FResolveParams());
 	}
 }
 
@@ -1295,7 +1269,8 @@ void FSceneRenderTargets::InitDynamicRHI()
 	{
 		// start with a defined state for the scissor rect (D3D11 was returning (0,0,0,0) which caused a clear to not execute correctly)
 		// todo: move this to an earlier place (for dx9 is has to be after device creation which is after window creation)
-		RHISetScissorRect(false, 0, 0, 0, 0);
+		FRHICommandListImmediate& RHICmdList = FRHICommandListExecutor::GetImmediateCommandList();
+		RHICmdList.SetScissorRect(false, 0, 0, 0, 0);
 
 		if (GRHIFeatureLevel == ERHIFeatureLevel::ES2)
 		{
diff --git a/Engine/Source/Runtime/Renderer/Private/PostProcess/SceneRenderTargets.h b/Engine/Source/Runtime/Renderer/Private/PostProcess/SceneRenderTargets.h
--- a/Engine/Source/Runtime/Renderer/Private/PostProcess/SceneRenderTargets.h
+++ b/Engine/Source/Runtime/Renderer/Private/PostProcess/SceneRenderTargets.h
@@ -142,65 +142,65 @@ public:
 	/**
 	 * Clears the GBuffer render targets to default values.
 	 */
-	void ClearGBufferTargets(const FLinearColor& ClearColor);
+	void ClearGBufferTargets(FRHICommandListImmediate& RHICmdList, const FLinearColor& ClearColor);
 	/**
 	 * Sets the scene color target and restores its contents if necessary
 	 * @param bGBufferPass - Whether the pass about to be rendered is the GBuffer population pass
 	 */
-	void BeginRenderingSceneColor(bool bGBufferPass = false);
+	void BeginRenderingSceneColor(FRHICommandListImmediate& RHICmdList, bool bGBufferPass = false);
 	/**
 	 * Called when finished rendering to the scene color surface
 	 * @param bKeepChanges - if true then the SceneColorSurface is resolved to the SceneColorTexture
 	 */
-	void FinishRenderingSceneColor(bool bKeepChanges = true, const FResolveRect& ResolveRect = FResolveRect());
+	void FinishRenderingSceneColor(FRHICommandListImmediate& RHICmdList, bool bKeepChanges = true, const FResolveRect& ResolveRect = FResolveRect());
 
 	// @return true: call FinishRenderingCustomDepth after rendering, false: don't render it, feature is disabled
-	bool BeginRenderingCustomDepth(bool bPrimitives);
+	bool BeginRenderingCustomDepth(FRHICommandListImmediate& RHICmdList, bool bPrimitives);
 	// only call if BeginRenderingCustomDepth() returned true
-	void FinishRenderingCustomDepth(const FResolveRect& ResolveRect = FResolveRect());
+	void FinishRenderingCustomDepth(FRHICommandListImmediate& RHICmdList, const FResolveRect& ResolveRect = FResolveRect());
 
 	/**
 	 * Resolve a previously rendered scene color surface.
 	 */
-	void ResolveSceneColor(const FResolveRect& ResolveRect = FResolveRect());
+	void ResolveSceneColor(FRHICommandListImmediate& RHICmdList, const FResolveRect& ResolveRect = FResolveRect());
 
 	/** Resolves the GBuffer targets so that their resolved textures can be sampled. */
-	void ResolveGBufferSurfaces(const FResolveRect& ResolveRect = FResolveRect());
+	void ResolveGBufferSurfaces(FRHICommandListImmediate& RHICmdList, const FResolveRect& ResolveRect = FResolveRect());
 
-	void BeginRenderingShadowDepth();
+	void BeginRenderingShadowDepth(FRHICommandListImmediate& RHICmdList);
 
 	/** Binds the appropriate shadow depth cube map for rendering. */
-	void BeginRenderingCubeShadowDepth(int32 ShadowResolution);
+	void BeginRenderingCubeShadowDepth(FRHICommandListImmediate& RHICmdList, int32 ShadowResolution);
 
 	/**
 	 * Called when finished rendering to the subject shadow depths so the surface can be copied to texture
 	 * @param ResolveParams - optional resolve params
 	 */
-	void FinishRenderingShadowDepth(const FResolveRect& ResolveRect = FResolveRect());
+	void FinishRenderingShadowDepth(FRHICommandListImmediate& RHICmdList, const FResolveRect& ResolveRect = FResolveRect());
 
-	void BeginRenderingReflectiveShadowMap( class FLightPropagationVolume* Lpv );
-	void FinishRenderingReflectiveShadowMap( const FResolveRect& ResolveRect = FResolveRect() );
+	void BeginRenderingReflectiveShadowMap(FRHICommandListImmediate& RHICmdList, class FLightPropagationVolume* Lpv);
+	void FinishRenderingReflectiveShadowMap(FRHICommandListImmediate& RHICmdList, const FResolveRect& ResolveRect = FResolveRect());
 
 	/** Resolves the appropriate shadow depth cube map and restores default state. */
-	void FinishRenderingCubeShadowDepth(int32 ShadowResolution, const FResolveParams& ResolveParams = FResolveParams());
+	void FinishRenderingCubeShadowDepth(FRHICommandListImmediate& RHICmdList, int32 ShadowResolution, const FResolveParams& ResolveParams = FResolveParams());
 	
-	void BeginRenderingTranslucency(const class FViewInfo& View);
+	void BeginRenderingTranslucency(FRHICommandListImmediate& RHICmdList, const class FViewInfo& View);
 
-	bool BeginRenderingSeparateTranslucency(const FViewInfo& View, bool bFirstTimeThisFrame);
-	void FinishRenderingSeparateTranslucency(const FViewInfo& View);
+	bool BeginRenderingSeparateTranslucency(FRHICommandListImmediate& RHICmdList, const FViewInfo& View, bool bFirstTimeThisFrame);
+	void FinishRenderingSeparateTranslucency(FRHICommandListImmediate& RHICmdList, const FViewInfo& View);
 	void FreeSeparateTranslucency();
 
-	void ResolveSceneDepthTexture();
-	void ResolveSceneDepthToAuxiliaryTexture();
+	void ResolveSceneDepthTexture(FRHICommandListImmediate& RHICmdList);
+	void ResolveSceneDepthToAuxiliaryTexture(FRHICommandListImmediate& RHICmdList);
 
-	void BeginRenderingPrePass();
-	void FinishRenderingPrePass();
+	void BeginRenderingPrePass(FRHICommandListImmediate& RHICmdList);
+	void FinishRenderingPrePass(FRHICommandListImmediate& RHICmdList);
 
-	void BeginRenderingSceneAlphaCopy();
-	void FinishRenderingSceneAlphaCopy();
+	void BeginRenderingSceneAlphaCopy(FRHICommandListImmediate& RHICmdList);
+	void FinishRenderingSceneAlphaCopy(FRHICommandListImmediate& RHICmdList);
 
-	void BeginRenderingLightAttenuation();
-	void FinishRenderingLightAttenuation();
+	void BeginRenderingLightAttenuation(FRHICommandListImmediate& RHICmdList);
+	void FinishRenderingLightAttenuation(FRHICommandListImmediate& RHICmdList);
 
 	/**
 	 * Cleans up editor primitive targets that we no longer need
@@ -229,8 +229,8 @@ public:
 
 
 	// FRenderResource interface.
-	virtual void InitDynamicRHI();
-	virtual void ReleaseDynamicRHI();
+	virtual void InitDynamicRHI() override;
+	virtual void ReleaseDynamicRHI() override;
 
 	// Texture Accessors -----------
 
diff --git a/Engine/Source/Runtime/Renderer/Private/PostProcess/ScreenSpaceReflections.cpp b/Engine/Source/Runtime/Renderer/Private/PostProcess/ScreenSpaceReflections.cpp
--- a/Engine/Source/Runtime/Renderer/Private/PostProcess/ScreenSpaceReflections.cpp
+++ b/Engine/Source/Runtime/Renderer/Private/PostProcess/ScreenSpaceReflections.cpp
@@ -195,7 +195,7 @@ void FRCPassPostProcessScreenSpaceReflections::Process(FRenderingCompositePassCo
 	int SSRQuality = ComputeSSRQuality(View.FinalPostProcessSettings.ScreenSpaceReflectionQuality);
 
 	SSRQuality = FMath::Clamp(SSRQuality, 1, 4);
-	Context.RHICmdList.CheckIsNull(); // need new approach for "static FGlobalBoundShaderState" for parallel rendering
+	
 
 	uint32 iPreFrame = bPrevFrame ? 1 : 0;
 
@@ -254,13 +254,13 @@ FPooledRenderTargetDesc FRCPassPostProcessScreenSpaceReflections::ComputeOutputD
 
 	return Ret;
 }
-void BuildHZB( const FViewInfo& View );
+void BuildHZB(FRHICommandListImmediate& RHICmdList, const FViewInfo& View);
 
-void ScreenSpaceReflections( const FViewInfo& View, TRefCountPtr<IPooledRenderTarget>& SSROutput )
+void ScreenSpaceReflections(FRHICommandListImmediate& RHICmdList, const FViewInfo& View, TRefCountPtr<IPooledRenderTarget>& SSROutput)
 {
-	BuildHZB( View );
+	BuildHZB(RHICmdList, View);
 
-	FRenderingCompositePassContext CompositeContext( View );
+	FRenderingCompositePassContext CompositeContext(RHICmdList, View);
 	FPostprocessContext Context( CompositeContext.Graph, View );
 
 	FSceneViewState* ViewState = (FSceneViewState*)Context.View.State;
diff --git a/Engine/Source/Runtime/Renderer/Private/PostProcess/ScreenSpaceReflections.h b/Engine/Source/Runtime/Renderer/Private/PostProcess/ScreenSpaceReflections.h
--- a/Engine/Source/Runtime/Renderer/Private/PostProcess/ScreenSpaceReflections.h
+++ b/Engine/Source/Runtime/Renderer/Private/PostProcess/ScreenSpaceReflections.h
@@ -49,6 +49,6 @@ public:
 	virtual FPooledRenderTargetDesc ComputeOutputDesc(EPassOutputId InPassOutputId) const;
 };
 
-void ScreenSpaceReflections( const FViewInfo& View, TRefCountPtr<IPooledRenderTarget>& SSROutput );
+void ScreenSpaceReflections(FRHICommandListImmediate& RHICmdList, const FViewInfo& View, TRefCountPtr<IPooledRenderTarget>& SSROutput);
 
 bool DoScreenSpaceReflections(const FViewInfo& View);
diff --git a/Engine/Source/Runtime/Renderer/Private/PostProcess/VisualizeTexture.cpp b/Engine/Source/Runtime/Renderer/Private/PostProcess/VisualizeTexture.cpp
--- a/Engine/Source/Runtime/Renderer/Private/PostProcess/VisualizeTexture.cpp
+++ b/Engine/Source/Runtime/Renderer/Private/PostProcess/VisualizeTexture.cpp
@@ -191,13 +191,13 @@ public:
 IMPLEMENT_SHADER_TYPE(,FVisualizeTexturePresentPS,TEXT("VisualizeTexture"),TEXT("PresentPS"),SF_Pixel);
 
 
-template<uint32 TextureType> void VisualizeTextureForTextureType(FRHICommandList& RHICmdList, const FVisualizeTextureData& Data)
+template<uint32 TextureType> void VisualizeTextureForTextureType(FRHICommandListImmediate& RHICmdList, const FVisualizeTextureData& Data)
 {
 	TShaderMapRef<FScreenVS> VertexShader(GetGlobalShaderMap());
 	TShaderMapRef<VisualizeTexturePS<TextureType> > PixelShader(GetGlobalShaderMap());
 
 	static FGlobalBoundShaderState BoundShaderState;
-	RHICmdList.CheckIsNull(); // need new approach for "static FGlobalBoundShaderState" for parallel rendering
+	
 
 	SetGlobalBoundShaderState(RHICmdList, BoundShaderState, GFilterVertexDeclaration.VertexDeclarationRHI, *VertexShader, *PixelShader);
 	PixelShader->SetParameters(RHICmdList, Data);
@@ -220,7 +220,7 @@ template<uint32 TextureType> void VisualizeTextureForTextureType(FRHICommandList
 		EDRF_UseTriangleOptimization);
 }
 
-void RenderVisualizeTexture(FRHICommandList& RHICmdList, const FVisualizeTextureData& Data)
+void RenderVisualizeTexture(FRHICommandListImmediate& RHICmdList, const FVisualizeTextureData& Data)
 {
 	if(Data.Desc.Is2DTexture())
 	{
@@ -311,7 +311,7 @@ FIntRect FVisualizeTexture::ComputeVisualizeTextureRect(FIntPoint InputTextureSi
 	return ret;
 }
 
-void FVisualizeTexture::GenerateContent(FRHICommandList& RHICmdList, const FSceneRenderTargetItem& RenderTargetItem, const FPooledRenderTargetDesc& Desc)
+void FVisualizeTexture::GenerateContent(FRHICommandListImmediate& RHICmdList, const FSceneRenderTargetItem& RenderTargetItem, const FPooledRenderTargetDesc& Desc)
 {
 	// otherwise StartFrame() wasn't called
 	check(ViewRect != FIntRect(0, 0, 0, 0))
@@ -374,7 +374,6 @@ void FVisualizeTexture::GenerateContent(FRHICommandList& RHICmdList, const FScen
 		default:
 			break;
 	}
-	RHICmdList.CheckIsNull(); // this stuff is manging resources directly, so that won't work in parallel
 
 	bool bIsDefault = StencilSRVSrc == GBlackTexture->TextureRHI;
 	bool bDepthStencil = Desc.Is2DTexture() && Desc.Format == PF_DepthStencil;
@@ -447,9 +446,7 @@ void FVisualizeTexture::GenerateContent(FRHICommandList& RHICmdList, const FScen
 		
 		
 		
-		RHICmdList.CheckIsNull(); // synchronous readback, won't work with parallel
-		
-		RHIReadSurfaceData(Texture, FIntRect(0, 0, Extent.X, Extent.Y), Bitmap, ReadDataFlags);
+		RHICmdList.ReadSurfaceData(Texture, FIntRect(0, 0, Extent.X, Extent.Y), Bitmap, ReadDataFlags);
 
 		// if the format and texture type is supported
 		if(Bitmap.Num())
@@ -473,7 +470,7 @@ void FVisualizeTexture::GenerateContent(FRHICommandList& RHICmdList, const FScen
 	}
 }
 
-void FVisualizeTexture::PresentContent(FRHICommandList& RHICmdList, const FSceneView& View)
+void FVisualizeTexture::PresentContent(FRHICommandListImmediate& RHICmdList, const FSceneView& View)
 {
 	if(Mode != 0)
 	{
@@ -509,7 +506,7 @@ void FVisualizeTexture::PresentContent(FRHICommandList& RHICmdList, const FScene
 	TShaderMapRef<FVisualizeTexturePresentPS> PixelShader(GetGlobalShaderMap());
 
 	static FGlobalBoundShaderState BoundShaderState;
-	RHICmdList.CheckIsNull(); // need new approach for "static FGlobalBoundShaderState" for parallel rendering
+	
 
 	SetGlobalBoundShaderState(RHICmdList, BoundShaderState, GFilterVertexDeclaration.VertexDeclarationRHI, *VertexShader, *PixelShader);
 
@@ -550,7 +547,6 @@ void FVisualizeTexture::PresentContent(FRHICommandList& RHICmdList, const FScene
 			return View.Family->RenderTarget->GetRenderTargetTexture();
 		}
 	} TempRenderTarget(View);
-	RHICmdList.CheckIsNull(); // canvas doesn't work with command lists
 
 	FCanvas Canvas(&TempRenderTarget, NULL, View.Family->CurrentRealTime, View.Family->CurrentWorldTime, View.Family->DeltaWorldTime);
 
@@ -624,8 +620,9 @@ void FVisualizeTexture::SetObserveTarget(const FString& InObservedDebugName, uin
 	ObservedDebugNameReusedGoal = InObservedDebugNameReusedGoal;
 }
 
-void FVisualizeTexture::SetCheckPoint(FRHICommandList& RHICmdList, const IPooledRenderTarget* PooledRenderTarget)
+void FVisualizeTexture::SetCheckPoint(FRHICommandListImmediate& RHICmdList, const IPooledRenderTarget* PooledRenderTarget)
 {
+
 #if !(UE_BUILD_SHIPPING || UE_BUILD_TEST)
 	if (!PooledRenderTarget || !bEnabled)
 	{
diff --git a/Engine/Source/Runtime/Renderer/Private/PostProcess/VisualizeTexture.h b/Engine/Source/Runtime/Renderer/Private/PostProcess/VisualizeTexture.h
--- a/Engine/Source/Runtime/Renderer/Private/PostProcess/VisualizeTexture.h
+++ b/Engine/Source/Runtime/Renderer/Private/PostProcess/VisualizeTexture.h
@@ -12,7 +12,7 @@ public:
 	FVisualizeTexture();
 
 	/** renders the VisualizeTextureContent to the current render target */
-	void PresentContent(FRHICommandList& RHICmdList, const FSceneView& View);
+	void PresentContent(FRHICommandListImmediate& RHICmdList, const FSceneView& View);
 	/** */
 	void OnStartFrame(const FSceneView& View);
 	/** */
@@ -23,7 +23,7 @@ public:
 	 * calling this allows to grab the state of the texture at this point to be queried by visualizetexture e.g. "vis LightAttenuation@2"
 	 * @param PooledRenderTarget 0 is silently ignored
 	 */
-	void SetCheckPoint(FRHICommandList& RHICmdList, const IPooledRenderTarget* PooledRenderTarget);
+	void SetCheckPoint(FRHICommandListImmediate& RHICmdList, const IPooledRenderTarget* PooledRenderTarget);
 
 	// @param bExtended true: with more convenience - not needed for crashes but useful from the console
 	void DebugLog(bool bExtended);
@@ -88,7 +88,7 @@ private:
 
 	// is called by FPooledRenderTarget
 
-	void GenerateContent(FRHICommandList& RHICmdList, const FSceneRenderTargetItem& RenderTargetItem, const FPooledRenderTargetDesc& Desc);
+	void GenerateContent(FRHICommandListImmediate& RHICmdList, const FSceneRenderTargetItem& RenderTargetItem, const FPooledRenderTargetDesc& Desc);
 
 	FIntRect ComputeVisualizeTextureRect(FIntPoint InputTextureSize) const;
 };
diff --git a/Engine/Source/Runtime/Renderer/Private/PrimitiveSceneInfo.cpp b/Engine/Source/Runtime/Renderer/Private/PrimitiveSceneInfo.cpp
--- a/Engine/Source/Runtime/Renderer/Private/PrimitiveSceneInfo.cpp
+++ b/Engine/Source/Runtime/Renderer/Private/PrimitiveSceneInfo.cpp
@@ -126,7 +126,7 @@ FPrimitiveSceneInfo::~FPrimitiveSceneInfo()
 	check(!OctreeId.IsValidId());
 }
 
-void FPrimitiveSceneInfo::AddStaticMeshes()
+void FPrimitiveSceneInfo::AddStaticMeshes(FRHICommandListImmediate& RHICmdList)
 {
 	// Cache the primitive's static mesh elements.
 	FBatchingSPDI BatchingSPDI(this);
@@ -145,11 +145,11 @@ void FPrimitiveSceneInfo::AddStaticMeshes()
 
 		// By this point, the index buffer render resource must be initialized
 		// Add the static mesh to the appropriate draw lists.
-		Mesh.AddToDrawLists(Scene);
+		Mesh.AddToDrawLists(RHICmdList, Scene);
 	}
 }
 
-void FPrimitiveSceneInfo::AddToScene(bool bUpdateStaticDrawLists)
+void FPrimitiveSceneInfo::AddToScene(FRHICommandListImmediate& RHICmdList, bool bUpdateStaticDrawLists)
 {
 	check(IsInRenderingThread());
 	
@@ -176,7 +176,7 @@ void FPrimitiveSceneInfo::AddToScene(bool bUpdateStaticDrawLists)
 
 	if (bUpdateStaticDrawLists)
 	{
-		AddStaticMeshes();
+		AddStaticMeshes(RHICmdList);
 	}
 
 	// create potential storage for our compact info
@@ -293,7 +293,7 @@ void FPrimitiveSceneInfo::RemoveFromScene(bool bUpdateStaticDrawLists)
 	}
 }
 
-void FPrimitiveSceneInfo::ConditionalUpdateStaticMeshes()
+void FPrimitiveSceneInfo::ConditionalUpdateStaticMeshes(FRHICommandListImmediate& RHICmdList)
 {
 	if (bNeedsStaticMeshUpdate)
 	{
@@ -303,7 +303,7 @@ void FPrimitiveSceneInfo::ConditionalUpdateStaticMeshes()
 		for (int32 MeshIndex = 0; MeshIndex < StaticMeshes.Num(); MeshIndex++)
 		{
 			StaticMeshes[MeshIndex].RemoveFromDrawLists();
-			StaticMeshes[MeshIndex].AddToDrawLists(Scene);
+			StaticMeshes[MeshIndex].AddToDrawLists(RHICmdList, Scene);
 		}
 	}
 }
diff --git a/Engine/Source/Runtime/Renderer/Private/PrimitiveSceneInfo.h b/Engine/Source/Runtime/Renderer/Private/PrimitiveSceneInfo.h
--- a/Engine/Source/Runtime/Renderer/Private/PrimitiveSceneInfo.h
+++ b/Engine/Source/Runtime/Renderer/Private/PrimitiveSceneInfo.h
@@ -229,19 +229,19 @@ public:
 	~FPrimitiveSceneInfo();
 
 	/** Adds the primitive to the scene. */
-	void AddToScene(bool bUpdateStaticDrawLists);
+	void AddToScene(FRHICommandListImmediate& RHICmdList, bool bUpdateStaticDrawLists);
 
 	/** Removes the primitive from the scene. */
 	void RemoveFromScene(bool bUpdateStaticDrawLists);
 
 	/** Updates the primitive's static meshes in the scene. */
-	void ConditionalUpdateStaticMeshes();
+	void ConditionalUpdateStaticMeshes(FRHICommandListImmediate& RHICmdList);
 
 	/** Sets a flag to update the primitive's static meshes before it is next rendered. */
 	void BeginDeferredUpdateStaticMeshes();
 
 	/** Adds the primitive's static meshes to the scene. */
-	void AddStaticMeshes();
+	void AddStaticMeshes(FRHICommandListImmediate& RHICmdList);
 
 	/** Removes the primitive's static meshes from the scene. */
 	void RemoveStaticMeshes();
diff --git a/Engine/Source/Runtime/Renderer/Private/ReflectionEnvironment.cpp b/Engine/Source/Runtime/Renderer/Private/ReflectionEnvironment.cpp
--- a/Engine/Source/Runtime/Renderer/Private/ReflectionEnvironment.cpp
+++ b/Engine/Source/Runtime/Renderer/Private/ReflectionEnvironment.cpp
@@ -549,13 +549,11 @@ private:
 IMPLEMENT_SHADER_TYPE(template<>,TStandardDeferredReflectionPS<true>,TEXT("ReflectionEnvironmentShaders"),TEXT("StandardDeferredReflectionPS"),SF_Pixel);
 IMPLEMENT_SHADER_TYPE(template<>,TStandardDeferredReflectionPS<false>,TEXT("ReflectionEnvironmentShaders"),TEXT("StandardDeferredReflectionPS"),SF_Pixel);
 
-void FDeferredShadingSceneRenderer::RenderReflectionCaptureSpecularBounceForAllViews(FRHICommandList& RHICmdList)
+void FDeferredShadingSceneRenderer::RenderReflectionCaptureSpecularBounceForAllViews(FRHICommandListImmediate& RHICmdList)
 {
-	RHICmdList.CheckIsNull(); // GSceneRenderTargets ops
-
 				// We're currently capturing a reflection capture, output SpecularColor * IndirectDiffuseGI for metals so they are not black in reflections,
 				// Since we don't have multiple bounce specular reflections
-	GSceneRenderTargets.BeginRenderingSceneColor();
+	GSceneRenderTargets.BeginRenderingSceneColor(RHICmdList);
 	RHICmdList.SetRasterizerState(TStaticRasterizerState< FM_Solid, CM_None >::GetRHI());
 	RHICmdList.SetDepthStencilState(TStaticDepthStencilState< false, CF_Always >::GetRHI());
 	RHICmdList.SetBlendState(TStaticBlendState< CW_RGB, BO_Add, BF_One, BF_One >::GetRHI());
@@ -564,7 +562,7 @@ void FDeferredShadingSceneRenderer::RenderReflectionCaptureSpecularBounceForAllV
 	TShaderMapRef< FReflectionCaptureSpecularBouncePS > PixelShader( GetGlobalShaderMap() );
 
 	static FGlobalBoundShaderState BoundShaderState;
-	RHICmdList.CheckIsNull(); // need new approach for "static FGlobalBoundShaderState" for parallel rendering
+	
 	SetGlobalBoundShaderState(RHICmdList, BoundShaderState, GFilterVertexDeclaration.VertexDeclarationRHI, *VertexShader, *PixelShader);
 
 	for (int32 ViewIndex = 0, Num = Views.Num(); ViewIndex < Num; ViewIndex++)
@@ -573,7 +571,7 @@ void FDeferredShadingSceneRenderer::RenderReflectionCaptureSpecularBounceForAllV
 
 		RHICmdList.SetViewport(View.ViewRect.Min.X, View.ViewRect.Min.Y, 0.0f, View.ViewRect.Max.X, View.ViewRect.Max.Y, 1.0f);
 
-		//@todo-rco: RHIPacketList
+		
 		PixelShader->SetParameters(RHICmdList, View);
 
 		DrawRectangle( 
@@ -598,15 +596,13 @@ bool FDeferredShadingSceneRenderer::ShouldDoReflectionEnvironment() const
 		&& ViewFamily.EngineShowFlags.ReflectionEnvironment;
 }
 
-void FDeferredShadingSceneRenderer::RenderImageBasedReflectionsSM5ForAllViews(FRHICommandList& RHICmdList)
+void FDeferredShadingSceneRenderer::RenderImageBasedReflectionsSM5ForAllViews(FRHICommandListImmediate& RHICmdList)
 {
 	const uint32 bUseLightmaps = CVarDiffuseFromCaptures.GetValueOnRenderThread() == 0;
 
-	RHICmdList.CheckIsNull(); // GSceneRenderTargets ops
-
 	TRefCountPtr<IPooledRenderTarget> NewSceneColor;
 	{
-		GSceneRenderTargets.ResolveSceneColor(FResolveRect(0, 0, ViewFamily.FamilySizeX, ViewFamily.FamilySizeY));
+		GSceneRenderTargets.ResolveSceneColor(RHICmdList, FResolveRect(0, 0, ViewFamily.FamilySizeX, ViewFamily.FamilySizeY));
 
 		FPooledRenderTargetDesc Desc = GSceneRenderTargets.GetSceneColor()->GetDesc();
 		Desc.TargetableFlags |= TexCreate_UAV;
@@ -624,7 +620,7 @@ void FDeferredShadingSceneRenderer::RenderImageBasedReflectionsSM5ForAllViews(FR
 		TRefCountPtr<IPooledRenderTarget> SSROutput = GSystemTextures.BlackDummy;
 		if( bSSR )
 		{
-			ScreenSpaceReflections( View, SSROutput );
+			ScreenSpaceReflections(RHICmdList, View, SSROutput);
 		}
 
 		// ReflectionEnv is assumed to be on when going into this method
@@ -661,7 +657,7 @@ void FDeferredShadingSceneRenderer::RenderImageBasedReflectionsSM5ForAllViews(FR
 	check(GSceneRenderTargets.GetSceneColor());
 }
 
-void FDeferredShadingSceneRenderer::RenderImageBasedReflectionsSM4ForAllViews(FRHICommandList& RHICmdList, bool bReflectionEnv)
+void FDeferredShadingSceneRenderer::RenderImageBasedReflectionsSM4ForAllViews(FRHICommandListImmediate& RHICmdList, bool bReflectionEnv)
 {
 	const bool bSkyLight = Scene->SkyLight
 		&& Scene->SkyLight->ProcessedTexture
@@ -734,7 +730,7 @@ void FDeferredShadingSceneRenderer::RenderImageBasedReflectionsSM4ForAllViews(FR
 		{
 			bRequiresApply = true;
 
-			ScreenSpaceReflections(View, SSROutput);
+			ScreenSpaceReflections(RHICmdList, View, SSROutput);
 		}
 
 		if (bReflectionEnv)
@@ -769,7 +765,7 @@ void FDeferredShadingSceneRenderer::RenderImageBasedReflectionsSM4ForAllViews(FR
 						TShaderMapRef<TStandardDeferredReflectionPS<true> > PixelShader(GetGlobalShaderMap());
 
 						static FGlobalBoundShaderState BoundShaderState;
-						RHICmdList.CheckIsNull(); // need new approach for "static FGlobalBoundShaderState" for parallel rendering
+						
 						SetGlobalBoundShaderState(RHICmdList, BoundShaderState, GFilterVertexDeclaration.VertexDeclarationRHI, *VertexShader, *PixelShader);
 
 						PixelShader->SetParameters(RHICmdList, View, ReflectionCapture);
@@ -779,7 +775,7 @@ void FDeferredShadingSceneRenderer::RenderImageBasedReflectionsSM4ForAllViews(FR
 						TShaderMapRef<TStandardDeferredReflectionPS<false> > PixelShader(GetGlobalShaderMap());
 
 						static FGlobalBoundShaderState BoundShaderState;
-						RHICmdList.CheckIsNull(); // need new approach for "static FGlobalBoundShaderState" for parallel rendering
+						
 						SetGlobalBoundShaderState(RHICmdList, BoundShaderState, GFilterVertexDeclaration.VertexDeclarationRHI, *VertexShader, *PixelShader);
 
 						PixelShader->SetParameters(RHICmdList, View, ReflectionCapture);
@@ -798,9 +794,8 @@ void FDeferredShadingSceneRenderer::RenderImageBasedReflectionsSM4ForAllViews(FR
 		{
 			// Apply reflections to screen
 			SCOPED_DRAW_EVENT(ReflectionApply, DEC_SCENE_ITEMS);
-			RHICmdList.CheckIsNull(); // GSceneRenderTargets ops
 
-			GSceneRenderTargets.BeginRenderingSceneColor();
+			GSceneRenderTargets.BeginRenderingSceneColor(RHICmdList);
 
 			RHICmdList.SetViewport(View.ViewRect.Min.X, View.ViewRect.Min.Y, 0.0f, View.ViewRect.Max.X, View.ViewRect.Max.Y, 1.0f);
 			RHICmdList.SetRasterizerState(TStaticRasterizerState<FM_Solid, CM_None>::GetRHI());
@@ -819,7 +814,7 @@ void FDeferredShadingSceneRenderer::RenderImageBasedReflectionsSM4ForAllViews(FR
 
 			TShaderMapRef< FPostProcessVS >		VertexShader(GetGlobalShaderMap());
 
-			RHICmdList.CheckIsNull(); // need new approach for "static FGlobalBoundShaderState" for parallel rendering
+			
 
 #define CASE(A,B,C) \
 			case ((A << 2) | (B << 1) | C) : \
@@ -827,7 +822,7 @@ void FDeferredShadingSceneRenderer::RenderImageBasedReflectionsSM4ForAllViews(FR
 			TShaderMapRef< FReflectionApplyPS<A, B, C> > PixelShader(GetGlobalShaderMap()); \
 			static FGlobalBoundShaderState BoundShaderState; \
 			SetGlobalBoundShaderState(RHICmdList, BoundShaderState, GFilterVertexDeclaration.VertexDeclarationRHI, *VertexShader, *PixelShader); \
-			PixelShader->SetParameters(FRHICommandList::GetNullRef(), View, LightAccumulation->GetRenderTargetItem().ShaderResourceTexture, SSROutput->GetRenderTargetItem().ShaderResourceTexture); \
+			PixelShader->SetParameters(RHICmdList, View, LightAccumulation->GetRenderTargetItem().ShaderResourceTexture, SSROutput->GetRenderTargetItem().ShaderResourceTexture); \
 			}; \
 			break
 
@@ -857,7 +852,7 @@ void FDeferredShadingSceneRenderer::RenderImageBasedReflectionsSM4ForAllViews(FR
 	}
 }
 
-void FDeferredShadingSceneRenderer::RenderDeferredReflections(FRHICommandList& RHICmdList)
+void FDeferredShadingSceneRenderer::RenderDeferredReflections(FRHICommandListImmediate& RHICmdList)
 {
 	if (IsSimpleDynamicLightingEnabled() || ViewFamily.EngineShowFlags.VisualizeLightCulling)
 	{
diff --git a/Engine/Source/Runtime/Renderer/Private/ReflectionEnvironmentCapture.cpp b/Engine/Source/Runtime/Renderer/Private/ReflectionEnvironmentCapture.cpp
--- a/Engine/Source/Runtime/Renderer/Private/ReflectionEnvironmentCapture.cpp
+++ b/Engine/Source/Runtime/Renderer/Private/ReflectionEnvironmentCapture.cpp
@@ -236,7 +236,6 @@ public:
 
 	void SetParameters(FRHICommandList& RHICmdList)
 	{
-		RHICmdList.CheckIsNull();
 		const int32 EffectiveTopMipSize = GReflectionCaptureSize;
 		const int32 NumMips = FMath::CeilLogTwo(EffectiveTopMipSize) + 1;
 		// Read from the smallest mip that was downsampled to
@@ -275,7 +274,7 @@ private:
 IMPLEMENT_SHADER_TYPE(,FComputeBrightnessPS,TEXT("ReflectionEnvironmentShaders"),TEXT("ComputeBrightnessMain"),SF_Pixel);
 
 /** Computes the average brightness of the given reflection capture and stores it in the scene. */
-void ComputeAverageBrightness(FRHICommandList& RHICmdList)
+void ComputeAverageBrightness(FRHICommandListImmediate& RHICmdList)
 {
 	SetRenderTarget(RHICmdList, GSceneRenderTargets.ReflectionBrightness->GetRenderTargetItem().TargetableTexture, NULL);
 	RHICmdList.SetRasterizerState(TStaticRasterizerState<FM_Solid, CM_None>::GetRHI());
@@ -286,7 +285,7 @@ void ComputeAverageBrightness(FRHICommandList& RHICmdList)
 	TShaderMapRef<FComputeBrightnessPS> PixelShader(GetGlobalShaderMap());
 
 	static FGlobalBoundShaderState BoundShaderState;
-	RHICmdList.CheckIsNull(); // need new approach for "static FGlobalBoundShaderState" for parallel rendering
+	
 	SetGlobalBoundShaderState(RHICmdList, BoundShaderState, GFilterVertexDeclaration.VertexDeclarationRHI, *VertexShader, *PixelShader);
 
 	PixelShader->SetParameters(RHICmdList);
@@ -303,7 +302,7 @@ void ComputeAverageBrightness(FRHICommandList& RHICmdList)
 }
 
 /** Generates mips for glossiness and filters the cubemap for a given reflection. */
-void FilterReflectionEnvironment(FRHICommandList& RHICmdList, FSHVectorRGB3* OutIrradianceEnvironmentMap)
+void FilterReflectionEnvironment(FRHICommandListImmediate& RHICmdList, FSHVectorRGB3* OutIrradianceEnvironmentMap)
 {
 	const int32 EffectiveTopMipSize = GReflectionCaptureSize;
 	const int32 NumMips = FMath::CeilLogTwo(EffectiveTopMipSize) + 1;
@@ -444,14 +443,14 @@ void FilterReflectionEnvironment(FRHICommandList& RHICmdList, FSHVectorRGB3* Out
 				{
 					PixelShader = *TShaderMapRef< TCubeFilterPS<1> >( GetGlobalShaderMap() );
 					static FGlobalBoundShaderState BoundShaderState;
-					RHICmdList.CheckIsNull(); // need new approach for "static FGlobalBoundShaderState" for parallel rendering
+					
 					SetGlobalBoundShaderState(RHICmdList, BoundShaderState, GFilterVertexDeclaration.VertexDeclarationRHI, *VertexShader, PixelShader, *GeometryShader);
 				}
 				else
 				{
 					PixelShader = *TShaderMapRef< TCubeFilterPS<0> >( GetGlobalShaderMap() );
 					static FGlobalBoundShaderState BoundShaderState;
-					RHICmdList.CheckIsNull(); // need new approach for "static FGlobalBoundShaderState" for parallel rendering
+					
 					SetGlobalBoundShaderState(RHICmdList, BoundShaderState, GFilterVertexDeclaration.VertexDeclarationRHI, *VertexShader, PixelShader, *GeometryShader);
 				}
 
@@ -497,14 +496,14 @@ void FilterReflectionEnvironment(FRHICommandList& RHICmdList, FSHVectorRGB3* Out
 					{
 						PixelShader = *TShaderMapRef< TCubeFilterPS<1> >( GetGlobalShaderMap() );
 						static FGlobalBoundShaderState BoundShaderState;
-						RHICmdList.CheckIsNull(); // need new approach for "static FGlobalBoundShaderState" for parallel rendering
+						
 						SetGlobalBoundShaderState(RHICmdList, BoundShaderState, GFilterVertexDeclaration.VertexDeclarationRHI, *VertexShader, PixelShader);
 					}
 					else
 					{
 						PixelShader = *TShaderMapRef< TCubeFilterPS<0> >( GetGlobalShaderMap() );
 						static FGlobalBoundShaderState BoundShaderState;
-						RHICmdList.CheckIsNull(); // need new approach for "static FGlobalBoundShaderState" for parallel rendering
+						
 						SetGlobalBoundShaderState(RHICmdList, BoundShaderState, GFilterVertexDeclaration.VertexDeclarationRHI, *VertexShader, PixelShader);
 					}
 
@@ -775,7 +774,7 @@ int32 FindOrAllocateCubemapIndex(FScene* Scene, const UReflectionCaptureComponen
 	return CaptureIndex;
 }
 
-void ClearScratchCubemaps(FRHICommandList& RHICmdList)
+void ClearScratchCubemaps(FRHICommandListImmediate& RHICmdList)
 {
 	// Clear scratch render targets to a consistent but noticeable value
 	// This makes debugging capture issues much easier, otherwise the random contents from previous captures is shown
@@ -821,10 +820,9 @@ void ClearScratchCubemaps(FRHICommandList& RHICmdList)
 }
 
 /** Captures the scene for a reflection capture by rendering the scene multiple times and copying into a cubemap texture. */
-void CaptureSceneToScratchCubemap(FRHICommandList& RHICmdList, FSceneRenderer* SceneRenderer, ECubeFace CubeFace, bool bCapturingForSkyLight, bool bLowerHemisphereIsBlack)
+void CaptureSceneToScratchCubemap(FRHICommandListImmediate& RHICmdList, FSceneRenderer* SceneRenderer, ECubeFace CubeFace, bool bCapturingForSkyLight, bool bLowerHemisphereIsBlack)
 {
 	FMemMark MemStackMark(FMemStack::Get());
-	RHICmdList.CheckIsNull(); //FDeferredUpdateResource::UpdateResources()
 	// update any resources that needed a deferred update
 	FDeferredUpdateResource::UpdateResources();
 	
@@ -832,7 +830,7 @@ void CaptureSceneToScratchCubemap(FRHICommandList& RHICmdList, FSceneRenderer* S
 		SCOPED_DRAW_EVENT(CubeMapCapture, DEC_SCENE_ITEMS);
 
 		// Render the scene normally for one face of the cubemap
-		SceneRenderer->Render();
+		SceneRenderer->Render(RHICmdList);
 
 #if PLATFORM_PS4 // @todo ps4 - this should be done a different way
 		// PS4 needs some code here to process the scene
@@ -916,7 +914,7 @@ void CaptureSceneToScratchCubemap(FRHICommandList& RHICmdList, FSceneRenderer* S
 	delete SceneRenderer;
 }
 
-void CopyCubemapToScratchCubemap(FRHICommandList& RHICmdList, UTextureCube* SourceCubemap, bool bIsSkyLight, bool bLowerHemisphereIsBlack)
+void CopyCubemapToScratchCubemap(FRHICommandListImmediate& RHICmdList, UTextureCube* SourceCubemap, bool bIsSkyLight, bool bLowerHemisphereIsBlack)
 {
 	check(SourceCubemap);
 	
@@ -1145,7 +1143,7 @@ void FScene::UpdateAllReflectionCaptures()
 	}
 }
 
-void GetReflectionCaptureData_RenderingThread(FScene* Scene, const UReflectionCaptureComponent* Component, FReflectionCaptureFullHDRDerivedData* OutDerivedData)
+void GetReflectionCaptureData_RenderingThread(FRHICommandListImmediate& RHICmdList, FScene* Scene, const UReflectionCaptureComponent* Component, FReflectionCaptureFullHDRDerivedData* OutDerivedData)
 {
 	const FCaptureComponentSceneState* ComponentStatePtr = Scene->ReflectionSceneData.AllocatedReflectionCaptureState.Find(Component);
 
@@ -1185,7 +1183,7 @@ void GetReflectionCaptureData_RenderingThread(FScene* Scene, const UReflectionCa
 				// Read each mip face
 				//@todo - do this without blocking the GPU so many times
 				//@todo - pool the temporary textures in RHIReadSurfaceFloatData instead of always creating new ones
-				RHIReadSurfaceFloatData(EffectiveDest.ShaderResourceTexture, FIntRect(0, 0, MipSize, MipSize), SurfaceData, (ECubeFace)CubeFace, CaptureIndex, MipIndex);
+				RHICmdList.ReadSurfaceFloatData(EffectiveDest.ShaderResourceTexture, FIntRect(0, 0, MipSize, MipSize), SurfaceData, (ECubeFace)CubeFace, CaptureIndex, MipIndex);
 				const int32 DestIndex = MipBaseIndex + CubeFace * CubeFaceBytes;
 				uint8* FaceData = &CaptureData[DestIndex];
 				check(SurfaceData.Num() * SurfaceData.GetTypeSize() == CubeFaceBytes);
@@ -1209,7 +1207,7 @@ void FScene::GetReflectionCaptureData(UReflectionCaptureComponent* Component, FR
 		const UReflectionCaptureComponent*,Component,Component,
 		FReflectionCaptureFullHDRDerivedData*,OutDerivedData,&OutDerivedData,
 	{
-		GetReflectionCaptureData_RenderingThread(Scene, Component, OutDerivedData);
+		GetReflectionCaptureData_RenderingThread(RHICmdList, Scene, Component, OutDerivedData);
 	});
 
 	// Necessary since the RT is writing to OutDerivedData directly
@@ -1338,7 +1336,7 @@ void CaptureSceneIntoScratchCubemap(FScene* Scene, FVector CapturePosition, bool
 			BeginFrame,
 		{
 			GFrameNumberRenderThread++;
-			RHIBeginFrame();
+			RHICmdList.BeginFrame();
 		})
 
 		FSceneViewFamilyContext ViewFamily( FSceneViewFamily::ConstructionValues(
@@ -1410,15 +1408,15 @@ void CaptureSceneIntoScratchCubemap(FScene* Scene, FVector CapturePosition, bool
 			bool, bCapturingForSkyLight, bCapturingForSkyLight,
 			bool, bLowerHemisphereIsBlack, bLowerHemisphereIsBlack,
 		{
-			//@todo-rco: RHIPacketList
-			CaptureSceneToScratchCubemap(FRHICommandList::GetNullRef(), SceneRenderer, CubeFace, bCapturingForSkyLight, bLowerHemisphereIsBlack);
+			
+			CaptureSceneToScratchCubemap(RHICmdList, SceneRenderer, CubeFace, bCapturingForSkyLight, bLowerHemisphereIsBlack);
 
-			RHIEndFrame();
+			RHICmdList.EndFrame();
 		});
 	}
 }
 
-void CopyToSceneArray(FScene* Scene, FReflectionCaptureProxy* ReflectionProxy)
+void CopyToSceneArray(FRHICommandListImmediate& RHICmdList, FScene* Scene, FReflectionCaptureProxy* ReflectionProxy)
 {
 	const int32 EffectiveTopMipSize = GReflectionCaptureSize;
 	const int32 NumMips = FMath::CeilLogTwo(EffectiveTopMipSize) + 1;
@@ -1434,12 +1432,12 @@ void CopyToSceneArray(FScene* Scene, FReflectionCaptureProxy* ReflectionProxy)
 
 		for (int32 CubeFace = 0; CubeFace < CubeFace_MAX; CubeFace++)
 		{
-			RHICopyToResolveTarget(EffectiveSource.ShaderResourceTexture, EffectiveDest.ShaderResourceTexture, true, FResolveParams(FResolveRect(), (ECubeFace)CubeFace, MipIndex, 0, CaptureIndex));
+			RHICmdList.CopyToResolveTarget(EffectiveSource.ShaderResourceTexture, EffectiveDest.ShaderResourceTexture, true, FResolveParams(FResolveRect(), (ECubeFace)CubeFace, MipIndex, 0, CaptureIndex));
 		}
 	}
 }
 
-void CopyToComponentTexture(FScene* Scene, FReflectionCaptureProxy* ReflectionProxy)
+void CopyToComponentTexture(FRHICommandListImmediate& RHICmdList, FScene* Scene, FReflectionCaptureProxy* ReflectionProxy)
 {
 	check(ReflectionProxy->SM4FullHDRCubemap);
 	const int32 EffectiveTopMipSize = GReflectionCaptureSize;
@@ -1453,7 +1451,7 @@ void CopyToComponentTexture(FScene* Scene, FReflectionCaptureProxy* ReflectionPr
 
 		for (int32 CubeFace = 0; CubeFace < CubeFace_MAX; CubeFace++)
 		{
-			RHICopyToResolveTarget(EffectiveSource.ShaderResourceTexture, ReflectionProxy->SM4FullHDRCubemap->TextureRHI, true, FResolveParams(FResolveRect(), (ECubeFace)CubeFace, MipIndex, 0, 0));
+			RHICmdList.CopyToResolveTarget(EffectiveSource.ShaderResourceTexture, ReflectionProxy->SM4FullHDRCubemap->TextureRHI, true, FResolveParams(FResolveRect(), (ECubeFace)CubeFace, MipIndex, 0, 0));
 		}
 	}
 }
@@ -1491,9 +1489,6 @@ void FScene::UpdateReflectionCaptureContents(UReflectionCaptureComponent* Captur
 			ENQUEUE_UNIQUE_RENDER_COMMAND( 
 				FilterCommand,
 			{
-				//@todo-rco: RHIPacketList
-				FRHICommandList& RHICmdList = FRHICommandList::GetNullRef();
-
 				FilterReflectionEnvironment(RHICmdList, NULL);
 			});
 
@@ -1510,11 +1505,11 @@ void FScene::UpdateReflectionCaptureContents(UReflectionCaptureComponent* Captur
 			{
 				if (FeatureLevel == ERHIFeatureLevel::SM5)
 				{
-					CopyToSceneArray(Scene, ReflectionProxy);
+					CopyToSceneArray(RHICmdList, Scene, ReflectionProxy);
 				}
 				else if (FeatureLevel == ERHIFeatureLevel::SM4)
 				{
-					CopyToComponentTexture(Scene, ReflectionProxy);
+					CopyToComponentTexture(RHICmdList, Scene, ReflectionProxy);
 				}
 
 				// Clean up the proxy now that the rendering thread is done with it
@@ -1524,7 +1519,7 @@ void FScene::UpdateReflectionCaptureContents(UReflectionCaptureComponent* Captur
 	}
 }
 
-void CopyToSkyTexture(FScene* Scene, FTexture* ProcessedTexture)
+void CopyToSkyTexture(FRHICommandListImmediate& RHICmdList, FScene* Scene, FTexture* ProcessedTexture)
 {
 	const int32 EffectiveTopMipSize = ProcessedTexture->GetSizeX();
 	const int32 NumMips = FMath::CeilLogTwo(EffectiveTopMipSize) + 1;
@@ -1537,7 +1532,7 @@ void CopyToSkyTexture(FScene* Scene, FTexture* ProcessedTexture)
 
 		for (int32 CubeFace = 0; CubeFace < CubeFace_MAX; CubeFace++)
 		{
-			RHICopyToResolveTarget(EffectiveSource.ShaderResourceTexture, ProcessedTexture->TextureRHI, true, FResolveParams(FResolveRect(), (ECubeFace)CubeFace, MipIndex, 0, 0));
+			RHICmdList.CopyToResolveTarget(EffectiveSource.ShaderResourceTexture, ProcessedTexture->TextureRHI, true, FResolveParams(FResolveRect(), (ECubeFace)CubeFace, MipIndex, 0, 0));
 		}
 	}
 }
@@ -1550,9 +1545,6 @@ void FScene::UpdateSkyCaptureContents(const USkyLightComponent* CaptureComponent
 			ClearCommand,
 		{
 			//@todo realtime skylight updates - skip this
-			//@todo-rco: RHIPacketList
-			FRHICommandList& RHICmdList = FRHICommandList::GetNullRef();
-
 			ClearScratchCubemaps(RHICmdList);
 		});
 
@@ -1568,8 +1560,8 @@ void FScene::UpdateSkyCaptureContents(const USkyLightComponent* CaptureComponent
 				UTextureCube*, SourceTexture, CaptureComponent->Cubemap,
 				bool, bLowerHemisphereIsBlack, CaptureComponent->bLowerHemisphereIsBlack,
 			{
-				//@todo-rco: RHIPacketList
-				CopyCubemapToScratchCubemap(FRHICommandList::GetNullRef(), SourceTexture, true, bLowerHemisphereIsBlack);
+				
+				CopyCubemapToScratchCubemap(RHICmdList, SourceTexture, true, bLowerHemisphereIsBlack);
 			});
 		}
 		else
@@ -1581,9 +1573,6 @@ void FScene::UpdateSkyCaptureContents(const USkyLightComponent* CaptureComponent
 			FilterCommand,
 			FSHVectorRGB3*, IrradianceEnvironmentMap, &OutIrradianceEnvironmentMap,
 		{
-			//@todo-rco: RHIPacketList
-			FRHICommandList& RHICmdList = FRHICommandList::GetNullRef();
-
 			FilterReflectionEnvironment(RHICmdList, IrradianceEnvironmentMap);
 		});
 
@@ -1599,7 +1588,7 @@ void FScene::UpdateSkyCaptureContents(const USkyLightComponent* CaptureComponent
 			FScene*, Scene, this,
 				FTexture*, ProcessedTexture, OutProcessedTexture,
 		{
-				CopyToSkyTexture(Scene, ProcessedTexture);
+			CopyToSkyTexture(RHICmdList, Scene, ProcessedTexture);
 		});
 		}
 	}
diff --git a/Engine/Source/Runtime/Renderer/Private/ReflectionEnvironmentCapture.h b/Engine/Source/Runtime/Renderer/Private/ReflectionEnvironmentCapture.h
--- a/Engine/Source/Runtime/Renderer/Private/ReflectionEnvironmentCapture.h
+++ b/Engine/Source/Runtime/Renderer/Private/ReflectionEnvironmentCapture.h
@@ -8,7 +8,7 @@
 
 extern ENGINE_API int32 GReflectionCaptureSize;
 
-extern void ComputeDiffuseIrradiance(FRHICommandList& RHICmdList, FTextureRHIRef LightingSource, int32 LightingSourceMipIndex, FSHVectorRGB3* OutIrradianceEnvironmentMap);
+extern void ComputeDiffuseIrradiance(FRHICommandListImmediate& RHICmdList, FTextureRHIRef LightingSource, int32 LightingSourceMipIndex, FSHVectorRGB3* OutIrradianceEnvironmentMap);
 
 class FDownsampleGS : public FGlobalShader
 {
diff --git a/Engine/Source/Runtime/Renderer/Private/ReflectionEnvironmentDiffuseIrradiance.cpp b/Engine/Source/Runtime/Renderer/Private/ReflectionEnvironmentDiffuseIrradiance.cpp
--- a/Engine/Source/Runtime/Renderer/Private/ReflectionEnvironmentDiffuseIrradiance.cpp
+++ b/Engine/Source/Runtime/Renderer/Private/ReflectionEnvironmentDiffuseIrradiance.cpp
@@ -235,7 +235,7 @@ IMPLEMENT_SHADER_TYPE(,FAccumulateCubeFacesPS,TEXT("ReflectionEnvironmentShaders
 
 FGlobalBoundShaderState AccumulateCubeFacesBoundShaderState;
 
-void ComputeDiffuseIrradiance(FRHICommandList& RHICmdList, FTextureRHIRef LightingSource, int32 LightingSourceMipIndex, FSHVectorRGB3* OutIrradianceEnvironmentMap)
+void ComputeDiffuseIrradiance(FRHICommandListImmediate& RHICmdList, FTextureRHIRef LightingSource, int32 LightingSourceMipIndex, FSHVectorRGB3* OutIrradianceEnvironmentMap)
 {
 	for (int32 CoefficientIndex = 0; CoefficientIndex < FSHVector3::MaxSHBasis; CoefficientIndex++)
 	{
@@ -436,10 +436,9 @@ void ComputeDiffuseIrradiance(FRHICommandList& RHICmdList, FTextureRHIRef Lighti
 		// Read back the completed SH environment map
 		FSceneRenderTargetItem& EffectiveRT = GSceneRenderTargets.SkySHIrradianceMap->GetRenderTargetItem();
 		check(EffectiveRT.ShaderResourceTexture->GetFormat() == PF_FloatRGBA);
-		RHICmdList.CheckIsNull(); // synchronous readback 
 
 		TArray<FFloat16Color> SurfaceData;
-		RHIReadSurfaceFloatData(EffectiveRT.ShaderResourceTexture, FIntRect(0, 0, FSHVector3::MaxSHBasis, 1), SurfaceData, CubeFace_PosX, 0, 0);
+		RHICmdList.ReadSurfaceFloatData(EffectiveRT.ShaderResourceTexture, FIntRect(0, 0, FSHVector3::MaxSHBasis, 1), SurfaceData, CubeFace_PosX, 0, 0);
 		check(SurfaceData.Num() == FSHVector3::MaxSHBasis);
 
 		for (int32 CoefficientIndex = 0; CoefficientIndex < FSHVector3::MaxSHBasis; CoefficientIndex++)
diff --git a/Engine/Source/Runtime/Renderer/Private/Renderer.cpp b/Engine/Source/Runtime/Renderer/Private/Renderer.cpp
--- a/Engine/Source/Runtime/Renderer/Private/Renderer.cpp
+++ b/Engine/Source/Runtime/Renderer/Private/Renderer.cpp
@@ -32,7 +32,7 @@ void FRendererModule::SceneRenderTargetsSetBufferSize(uint32 SizeX, uint32 SizeY
 	GSceneRenderTargets.UpdateRHI();
 }
 
-void FRendererModule::DrawTileMesh(const FSceneView& SceneView, const FMeshBatch& Mesh, bool bIsHitTesting, const FHitProxyId& HitProxyId)
+void FRendererModule::DrawTileMesh(FRHICommandListImmediate& RHICmdList, const FSceneView& SceneView, const FMeshBatch& Mesh, bool bIsHitTesting, const FHitProxyId& HitProxyId)
 {
 	// Create an FViewInfo so we can initialize its RHI resources
 	//@todo - reuse this view for multiple tiles, this is going to be slow for each tile
@@ -40,9 +40,7 @@ void FRendererModule::DrawTileMesh(const FSceneView& SceneView, const FMeshBatch
 	View.InitRHIResources();
 
 	const auto FeatureLevel = View.GetFeatureLevel();
-	//@todo-rco: RHIPacketList
-	FRHICommandList& RHICmdList = FRHICommandList::GetNullRef();
-
+	
 	const FMaterial* Material = Mesh.MaterialRenderProxy->GetMaterial(FeatureLevel);
 
 	//get the blend mode of the material
@@ -66,7 +64,7 @@ void FRendererModule::DrawTileMesh(const FSceneView& SceneView, const FMeshBatch
 		else
 		{
 			// make sure we are doing opaque drawing
-			RHISetBlendState(TStaticBlendState<>::GetRHI());
+			RHICmdList.SetBlendState(TStaticBlendState<>::GetRHI());
 
 			// draw the mesh
 			if (bIsHitTesting)
@@ -160,7 +158,7 @@ void FRendererModule::GPUBenchmark(FSynthBenchmarkResults& InOut, uint32 WorkSca
 	  bool, bDebugOut, bDebugOut,
 	  FSynthBenchmarkResults&, InOut, InOut,
 	{
-		RendererGPUBenchmark(InOut, DummyView, WorkScale, bDebugOut);
+		RendererGPUBenchmark(RHICmdList, InOut, DummyView, WorkScale, bDebugOut);
 	});
 	FlushRenderingCommands();
 }
diff --git a/Engine/Source/Runtime/Renderer/Private/RendererPrivate.h b/Engine/Source/Runtime/Renderer/Private/RendererPrivate.h
--- a/Engine/Source/Runtime/Renderer/Private/RendererPrivate.h
+++ b/Engine/Source/Runtime/Renderer/Private/RendererPrivate.h
@@ -30,7 +30,7 @@ public:
 	virtual uint32 GetNumDynamicLightsAffectingPrimitive(const FPrimitiveSceneInfo* PrimitiveSceneInfo,const FLightCacheInterface* LCI) override;
 	virtual void ReallocateSceneRenderTargets() override;
 	virtual void SceneRenderTargetsSetBufferSize(uint32 SizeX, uint32 SizeY) override;
-	virtual void DrawTileMesh(const FSceneView& View, const FMeshBatch& Mesh, bool bIsHitTesting, const FHitProxyId& HitProxyId) override;
+	virtual void DrawTileMesh(FRHICommandListImmediate& RHICmdList, const FSceneView& View, const FMeshBatch& Mesh, bool bIsHitTesting, const FHitProxyId& HitProxyId) override;
 	virtual void RenderTargetPoolFindFreeElement(const FPooledRenderTargetDesc& Desc, TRefCountPtr<IPooledRenderTarget> &Out, const TCHAR* InDebugName) override;
 	virtual void TickRenderTargetPool() override;
 	virtual void DebugLogOnCrash() override;
diff --git a/Engine/Source/Runtime/Renderer/Private/RendererScene.cpp b/Engine/Source/Runtime/Renderer/Private/RendererScene.cpp
--- a/Engine/Source/Runtime/Renderer/Private/RendererScene.cpp
+++ b/Engine/Source/Runtime/Renderer/Private/RendererScene.cpp
@@ -139,7 +139,7 @@ void FScene::CheckPrimitiveArrays()
 	check(Primitives.Num() == PrimitiveOcclusionBounds.Num());
 }
 
-void FScene::AddPrimitiveSceneInfo_RenderThread(FPrimitiveSceneInfo* PrimitiveSceneInfo)
+void FScene::AddPrimitiveSceneInfo_RenderThread(FRHICommandListImmediate& RHICmdList, FPrimitiveSceneInfo* PrimitiveSceneInfo)
 {
 	SCOPE_CYCLE_COUNTER(STAT_AddScenePrimitiveRenderThreadTime);
 	
@@ -161,7 +161,7 @@ void FScene::AddPrimitiveSceneInfo_RenderThread(FPrimitiveSceneInfo* PrimitiveSc
 	PrimitiveSceneInfo->LinkAttachmentGroup();
 
 	// Add the primitive to the scene.
-	PrimitiveSceneInfo->AddToScene(true);
+	PrimitiveSceneInfo->AddToScene(RHICmdList, true);
 }
 
 /**
@@ -336,12 +336,12 @@ void FScene::AddPrimitive(UPrimitiveComponent* Primitive)
 		FPrimitiveSceneInfo*,PrimitiveSceneInfo,PrimitiveSceneInfo,
 		{
 			FScopeCycleCounter Context(PrimitiveSceneInfo->Proxy->GetStatId());
-			Scene->AddPrimitiveSceneInfo_RenderThread(PrimitiveSceneInfo);
+			Scene->AddPrimitiveSceneInfo_RenderThread(RHICmdList, PrimitiveSceneInfo);
 		});
 
 }
 
-void FScene::UpdatePrimitiveTransform_RenderThread(FPrimitiveSceneProxy* PrimitiveSceneProxy, const FBoxSphereBounds& WorldBounds, const FBoxSphereBounds& LocalBounds, const FMatrix& LocalToWorld, const FVector& OwnerPosition)
+void FScene::UpdatePrimitiveTransform_RenderThread(FRHICommandListImmediate& RHICmdList, FPrimitiveSceneProxy* PrimitiveSceneProxy, const FBoxSphereBounds& WorldBounds, const FBoxSphereBounds& LocalBounds, const FMatrix& LocalToWorld, const FVector& OwnerPosition)
 {
 	SCOPE_CYCLE_COUNTER(STAT_UpdatePrimitiveTransformRenderThreadTime);
 
@@ -364,7 +364,7 @@ void FScene::UpdatePrimitiveTransform_RenderThread(FPrimitiveSceneProxy* Primiti
 	check(!(bUpdateStaticDrawLists && PrimitiveSceneProxy->GetPrimitiveSceneInfo()->StaticMeshes.Num()));
 
 	// Re-add the primitive to the scene with the new transform.
-	PrimitiveSceneProxy->GetPrimitiveSceneInfo()->AddToScene(bUpdateStaticDrawLists);
+	PrimitiveSceneProxy->GetPrimitiveSceneInfo()->AddToScene(RHICmdList, bUpdateStaticDrawLists);
 }
 
 void FScene::UpdatePrimitiveTransform(UPrimitiveComponent* Primitive)
@@ -450,7 +450,7 @@ void FScene::UpdatePrimitiveTransform(UPrimitiveComponent* Primitive)
 				FPrimitiveUpdateParams,UpdateParams,UpdateParams,
 				{
 					FScopeCycleCounter Context(UpdateParams.PrimitiveSceneProxy->GetStatId());
-					UpdateParams.Scene->UpdatePrimitiveTransform_RenderThread(UpdateParams.PrimitiveSceneProxy, UpdateParams.WorldBounds, UpdateParams.LocalBounds, UpdateParams.LocalToWorld, UpdateParams.OwnerPosition);
+					UpdateParams.Scene->UpdatePrimitiveTransform_RenderThread(RHICmdList, UpdateParams.PrimitiveSceneProxy, UpdateParams.WorldBounds, UpdateParams.LocalBounds, UpdateParams.LocalToWorld, UpdateParams.OwnerPosition);
 				});
 		}
 	}
@@ -1510,7 +1510,7 @@ void FScene::SetPrecomputedVisibility(const FPrecomputedVisibilityHandler* Preco
 	});
 }
 
-void FScene::SetShaderMapsOnMaterialResources_RenderThread(const FMaterialsToUpdateMap& MaterialsToUpdate)
+void FScene::SetShaderMapsOnMaterialResources_RenderThread(FRHICommandListImmediate& RHICmdList, const FMaterialsToUpdateMap& MaterialsToUpdate)
 {
 	SCOPE_CYCLE_COUNTER(STAT_Scene_SetShaderMapsOnMaterialResources_RT);
 
@@ -1554,7 +1554,7 @@ void FScene::SetShaderMapsOnMaterialResources_RenderThread(const FMaterialsToUpd
 	// Update static draw lists, which cache shader references from materials, but the shader map has now changed
 	if (bFoundAnyInitializedMaterials)
 	{
-		UpdateStaticDrawListsForMaterials_RenderThread(MaterialArray);
+		UpdateStaticDrawListsForMaterials_RenderThread(RHICmdList, MaterialArray);
 	}
 }
 
@@ -1571,11 +1571,11 @@ void FScene::SetShaderMapsOnMaterialResources(const TMap<FMaterial*, class FMate
 		FScene*,Scene,this,
 		FMaterialsToUpdateMap,MaterialsToUpdate,MaterialsToUpdate,
 	{
-		Scene->SetShaderMapsOnMaterialResources_RenderThread(MaterialsToUpdate);
+		Scene->SetShaderMapsOnMaterialResources_RenderThread(RHICmdList, MaterialsToUpdate);
 	});
 }
 
-void FScene::UpdateStaticDrawListsForMaterials_RenderThread(const TArray<const FMaterial*>& Materials)
+void FScene::UpdateStaticDrawListsForMaterials_RenderThread(FRHICommandListImmediate& RHICmdList, const TArray<const FMaterial*>& Materials)
 {
 	SCOPE_CYCLE_COUNTER(STAT_Scene_UpdateStaticDrawListsForMaterials_RT);
 
@@ -1614,7 +1614,7 @@ void FScene::UpdateStaticDrawListsForMaterials_RenderThread(const TArray<const F
 		FPrimitiveSceneInfo* Primitive = PrimitivesToUpdate[PrimitiveIndex];
 
 		Primitive->RemoveStaticMeshes();
-		Primitive->AddStaticMeshes();
+		Primitive->AddStaticMeshes(RHICmdList);
 	}
 }
 
@@ -1625,7 +1625,7 @@ void FScene::UpdateStaticDrawListsForMaterials(const TArray<const FMaterial*>& M
 		FScene*,Scene,this,
 		TArray<const FMaterial*>,Materials,Materials,
 	{
-		Scene->UpdateStaticDrawListsForMaterials_RenderThread(Materials);
+		Scene->UpdateStaticDrawListsForMaterials_RenderThread(RHICmdList, Materials);
 	});
 }
 
diff --git a/Engine/Source/Runtime/Renderer/Private/SceneCaptureRendering.cpp b/Engine/Source/Runtime/Renderer/Private/SceneCaptureRendering.cpp
--- a/Engine/Source/Runtime/Renderer/Private/SceneCaptureRendering.cpp
+++ b/Engine/Source/Runtime/Renderer/Private/SceneCaptureRendering.cpp
@@ -13,9 +13,8 @@
 #include "PostProcessAmbient.h"
 #include "PostProcessing.h"
 
-void UpdateSceneCaptureContent_RenderThread(FSceneRenderer* SceneRenderer, FTextureRenderTargetResource* TextureRenderTarget, const FName OwnerName, const FResolveParams& ResolveParams, bool bUseSceneColorTexture)
+void UpdateSceneCaptureContent_RenderThread(FRHICommandListImmediate& RHICmdList, FSceneRenderer* SceneRenderer, FTextureRenderTargetResource* TextureRenderTarget, const FName OwnerName, const FResolveParams& ResolveParams, bool bUseSceneColorTexture)
 {
-	FRHICommandList& RHICmdList = FRHICommandList::GetNullRef();
 	FMemMark MemStackMark(FMemStack::Get());
 
 	// update any resources that needed a deferred update
@@ -34,7 +33,7 @@ void UpdateSceneCaptureContent_RenderThread(FSceneRenderer* SceneRenderer, FText
 		FIntRect UnconstrainedViewRect = SceneRenderer->Views[0].UnconstrainedViewRect;
 		SetRenderTarget(RHICmdList, Target->GetRenderTargetTexture(), NULL);
 		RHICmdList.Clear(true, FLinearColor::Black, false, 1.0f, false, 0, ViewRect);
-		SceneRenderer->Render();
+		SceneRenderer->Render(RHICmdList);
 
 		// Copy the captured scene into the destination texture
 		if (bUseSceneColorTexture)
@@ -49,10 +48,9 @@ void UpdateSceneCaptureContent_RenderThread(FSceneRenderer* SceneRenderer, FText
 			TShaderMapRef<FScreenVS> VertexShader(GetGlobalShaderMap());
 			TShaderMapRef<FScreenPS> PixelShader(GetGlobalShaderMap());
 			static FGlobalBoundShaderState BoundShaderState;
-			RHICmdList.CheckIsNull(); // need new approach for "static FGlobalBoundShaderState" for parallel rendering
 			SetGlobalBoundShaderState(RHICmdList, BoundShaderState, GFilterVertexDeclaration.VertexDeclarationRHI, *VertexShader, *PixelShader);
 
-			FRenderingCompositePassContext Context(SceneRenderer->Views[0]);
+			FRenderingCompositePassContext Context(RHICmdList, SceneRenderer->Views[0]);
 
 			VertexShader->SetParameters(RHICmdList, SceneRenderer->Views[0]);
 			PixelShader->SetParameters(RHICmdList, TStaticSamplerState<SF_Point>::GetRHI(), GSceneRenderTargets.GetSceneColorTexture());
@@ -190,7 +188,7 @@ void FScene::UpdateSceneCaptureContents(USceneCaptureComponent2D* CaptureCompone
 			FName, OwnerName, OwnerName,
 			bool, bUseSceneColorTexture, bUseSceneColorTexture,
 		{
-			UpdateSceneCaptureContent_RenderThread(SceneRenderer, TextureRenderTarget, OwnerName, FResolveParams(), bUseSceneColorTexture);
+			UpdateSceneCaptureContent_RenderThread(RHICmdList, SceneRenderer, TextureRenderTarget, OwnerName, FResolveParams(), bUseSceneColorTexture);
 		});
 	}
 }
@@ -261,7 +259,7 @@ void FScene::UpdateSceneCaptureContents(USceneCaptureComponentCube* CaptureCompo
 				FName, OwnerName, OwnerName,
 				ECubeFace, TargetFace, TargetFace,
 			{
-				UpdateSceneCaptureContent_RenderThread(SceneRenderer, TextureRenderTarget, OwnerName, FResolveParams(FResolveRect(), TargetFace), true);
+				UpdateSceneCaptureContent_RenderThread(RHICmdList, SceneRenderer, TextureRenderTarget, OwnerName, FResolveParams(FResolveRect(), TargetFace), true);
 			});
 		}
 	}
diff --git a/Engine/Source/Runtime/Renderer/Private/SceneCore.cpp b/Engine/Source/Runtime/Renderer/Private/SceneCore.cpp
--- a/Engine/Source/Runtime/Renderer/Private/SceneCore.cpp
+++ b/Engine/Source/Runtime/Renderer/Private/SceneCore.cpp
@@ -274,7 +274,7 @@ void FStaticMesh::UnlinkDrawList(FStaticMesh::FDrawListElementLink* Link)
 	verify(DrawListLinks.RemoveSingleSwap(Link) == 1);
 }
 
-void FStaticMesh::AddToDrawLists(FScene* Scene)
+void FStaticMesh::AddToDrawLists(FRHICommandListImmediate& RHICmdList, FScene* Scene)
 {
 	const auto FeatureLevel = Scene->GetFeatureLevel();
 	if (FeatureLevel >= ERHIFeatureLevel::SM3)
@@ -310,7 +310,7 @@ void FStaticMesh::AddToDrawLists(FScene* Scene)
 				}
 
 				// Add the static mesh to the DPG's base pass draw list.
-				FBasePassOpaqueDrawingPolicyFactory::AddStaticMesh(FRHICommandList::GetNullRef(), Scene, this);
+				FBasePassOpaqueDrawingPolicyFactory::AddStaticMesh(RHICmdList, Scene, this);
 
 				FVelocityDrawingPolicyFactory::AddStaticMesh(Scene, this);
 			}
diff --git a/Engine/Source/Runtime/Renderer/Private/SceneCore.h b/Engine/Source/Runtime/Renderer/Private/SceneCore.h
--- a/Engine/Source/Runtime/Renderer/Private/SceneCore.h
+++ b/Engine/Source/Runtime/Renderer/Private/SceneCore.h
@@ -162,7 +162,7 @@ public:
 	void UnlinkDrawList(FDrawListElementLink* Link);
 
 	/** Adds the static mesh to the appropriate draw lists in a scene. */
-	void AddToDrawLists(FScene* Scene);
+	void AddToDrawLists(FRHICommandListImmediate& RHICmdList, FScene* Scene);
 
 	/** Removes the static mesh from all draw lists. */
 	void RemoveFromDrawLists();
diff --git a/Engine/Source/Runtime/Renderer/Private/SceneHitProxyRendering.cpp b/Engine/Source/Runtime/Renderer/Private/SceneHitProxyRendering.cpp
--- a/Engine/Source/Runtime/Renderer/Private/SceneHitProxyRendering.cpp
+++ b/Engine/Source/Runtime/Renderer/Private/SceneHitProxyRendering.cpp
@@ -175,9 +175,8 @@ FHitProxyDrawingPolicy::FHitProxyDrawingPolicy(
 
 void FHitProxyDrawingPolicy::DrawShared(FRHICommandList& RHICmdList, const FSceneView* View, FBoundShaderStateRHIParamRef BoundShaderState) const
 {
-	RHICmdList.CheckIsNull();
 	// Set the actual shader & vertex declaration state
-	RHISetBoundShaderState( BoundShaderState);
+	RHICmdList.SetBoundShaderState(BoundShaderState);
 
 	// Set the depth-only shader parameters for the material.
 	VertexShader->SetParameters(RHICmdList, MaterialRenderProxy,*View);
@@ -246,7 +245,7 @@ void FHitProxyDrawingPolicy::SetMeshRenderState(
 		PixelShader->SetHitProxyId(RHICmdList, HitProxyId);	
 	}
 
-	RHISetRasterizerState(GetStaticRasterizerState<false>(
+	RHICmdList.SetRasterizerState(GetStaticRasterizerState<false>(
 		(Mesh.bWireframe || IsWireframe()) ? FM_Wireframe : FM_Solid,
 		((IsTwoSided() && !NeedsBackfacePass()) || Mesh.bDisableBackfaceCulling) ? CM_None :
 			(XOR(XOR(View.bReverseCulling,bBackFace), Mesh.ReverseCulling) ? CM_CCW : CM_CW)
@@ -329,14 +328,14 @@ bool FHitProxyDrawingPolicyFactory::DrawDynamicMesh(
 	return false;
 }
 
-void FDeferredShadingSceneRenderer::RenderHitProxies()
+void FDeferredShadingSceneRenderer::RenderHitProxies(FRHICommandListImmediate& RHICmdList)
 {
 #if WITH_EDITOR
 
 	auto FeatureLevel = ViewFamily.Scene->GetFeatureLevel();
 
 	// Initialize global system textures (pass-through if already initialized).
-	GSystemTextures.InitializeTextures(FeatureLevel);
+	GSystemTextures.InitializeTextures(RHICmdList, FeatureLevel);
 
 	// Allocate the maximum scene render target space for the current view family.
 	GSceneRenderTargets.Allocate(ViewFamily);
@@ -355,22 +354,19 @@ void FDeferredShadingSceneRenderer::RenderHitProxies()
 		// HitProxyRT==0 should never happen but better we don't crash
 		return;
 	}
-
-	//@todo-rco: RHIPacketList
-	FRHICommandList& RHICmdList = FRHICommandList::GetNullRef();
-
+	
 	SetRenderTarget(RHICmdList, HitProxyRT->GetRenderTargetItem().TargetableTexture, GSceneRenderTargets.GetSceneDepthSurface());
 
 	// Clear color for each view.
 	for(int32 ViewIndex = 0;ViewIndex < Views.Num();ViewIndex++)
 	{
 		const FViewInfo& View = Views[ViewIndex];
-		RHISetViewport(View.ViewRect.Min.X, View.ViewRect.Min.Y, 0.0f, View.ViewRect.Max.X, View.ViewRect.Max.Y, 1.0f);
-		RHIClear(true,FLinearColor::White,false,1,false,0, FIntRect());
+		RHICmdList.SetViewport(View.ViewRect.Min.X, View.ViewRect.Min.Y, 0.0f, View.ViewRect.Max.X, View.ViewRect.Max.Y, 1.0f);
+		RHICmdList.Clear(true, FLinearColor::White, false, 1, false, 0, FIntRect());
 	}
 
 	// Find the visible primitives.
-	InitViews();
+	InitViews(RHICmdList);
 
 	// Dynamic vertex and index buffers need to be committed before rendering.
 	FGlobalDynamicVertexBuffer::Get().Commit();
@@ -378,8 +374,8 @@ void FDeferredShadingSceneRenderer::RenderHitProxies()
 
 	// Depth tests + writes, no alpha blending.
 	// Note, this is a reversed Z depth surface, using CF_GreaterEqual.
-	RHISetDepthStencilState(TStaticDepthStencilState<true,CF_GreaterEqual>::GetRHI());
-	RHISetBlendState(TStaticBlendState<>::GetRHI());
+	RHICmdList.SetDepthStencilState(TStaticDepthStencilState<true, CF_GreaterEqual>::GetRHI());
+	RHICmdList.SetBlendState(TStaticBlendState<>::GetRHI());
 
 	const bool bNeedToSwitchVerticalAxis = IsES2Platform(GRHIShaderPlatform) && !IsPCPlatform(GRHIShaderPlatform);
 
@@ -388,13 +384,12 @@ void FDeferredShadingSceneRenderer::RenderHitProxies()
 		const FViewInfo& View = Views[ViewIndex];
 
 		// Set the device viewport for the view.
-		RHISetViewport(View.ViewRect.Min.X, View.ViewRect.Min.Y, 0.0f, View.ViewRect.Max.X, View.ViewRect.Max.Y, 1.0f);
+		RHICmdList.SetViewport(View.ViewRect.Min.X, View.ViewRect.Min.Y, 0.0f, View.ViewRect.Max.X, View.ViewRect.Max.Y, 1.0f);
 
 		// Clear the depth buffer for each DPG.
 		// Note, this is a reversed Z depth surface, so 0.0f is the far plane.
-		RHIClear(false,FLinearColor::Black,true,0.0f,true,0, FIntRect());
+		RHICmdList.Clear(false, FLinearColor::Black, true, 0.0f, true, 0, FIntRect());
 
-			
 		// Adjust the visibility map for this view
 		if( !View.bAllowTranslucentPrimitivesInHitProxy )
 		{
@@ -438,14 +433,14 @@ void FDeferredShadingSceneRenderer::RenderHitProxies()
 		// Some elements should never be occluded (e.g. gizmos).
 		// So we render those twice, first to overwrite potentially nearer objects,
 		// then again to allows proper occlusion within those elements.
-		RHISetDepthStencilState(TStaticDepthStencilState<false,CF_Always>::GetRHI());
+		RHICmdList.SetDepthStencilState(TStaticDepthStencilState<false, CF_Always>::GetRHI());
 		// Draw the view's foreground elements last.
 		DrawViewElements<FHitProxyDrawingPolicyFactory>(View,FHitProxyDrawingPolicyFactory::ContextType(),SDPG_Foreground,bPreFog);
 
 		View.TopBatchedViewElements.Draw(RHICmdList, bNeedToSwitchVerticalAxis, View.ViewProjectionMatrix, View.ViewRect.Width(), View.ViewRect.Height(), true);
 
 		// Note, this is a reversed Z depth surface, using CF_GreaterEqual.
-		RHISetDepthStencilState(TStaticDepthStencilState<true,CF_GreaterEqual>::GetRHI());
+		RHICmdList.SetDepthStencilState(TStaticDepthStencilState<true, CF_GreaterEqual>::GetRHI());
 		// Draw the view's foreground elements last.
 		DrawViewElements<FHitProxyDrawingPolicyFactory>(View,FHitProxyDrawingPolicyFactory::ContextType(),SDPG_Foreground,bPreFog);
 
@@ -453,12 +448,12 @@ void FDeferredShadingSceneRenderer::RenderHitProxies()
 	}
 
 	// Finish drawing to the hit proxy render target.
-	RHICopyToResolveTarget(HitProxyRT->GetRenderTargetItem().TargetableTexture, HitProxyRT->GetRenderTargetItem().ShaderResourceTexture, false, FResolveParams());
+	RHICmdList.CopyToResolveTarget(HitProxyRT->GetRenderTargetItem().TargetableTexture, HitProxyRT->GetRenderTargetItem().ShaderResourceTexture, false, FResolveParams());
 	// to be able to observe results with VisualizeTexture
 	GRenderTargetPool.VisualizeTexture.SetCheckPoint(RHICmdList, HitProxyRT);
 
 	// After scene rendering, disable the depth buffer.
-	RHISetDepthStencilState(TStaticDepthStencilState<false,CF_Always>::GetRHI());
+	RHICmdList.SetDepthStencilState(TStaticDepthStencilState<false, CF_Always>::GetRHI());
 
 	//
 	// Copy the hit proxy buffer into the view family's render target.
@@ -517,6 +512,6 @@ void FDeferredShadingSceneRenderer::RenderHitProxies()
 				1.0f
 				);
 
-	RHIEndScene();
+	RHICmdList.EndScene();
 #endif
 }
diff --git a/Engine/Source/Runtime/Renderer/Private/SceneOcclusion.cpp b/Engine/Source/Runtime/Renderer/Private/SceneOcclusion.cpp
--- a/Engine/Source/Runtime/Renderer/Private/SceneOcclusion.cpp
+++ b/Engine/Source/Runtime/Renderer/Private/SceneOcclusion.cpp
@@ -69,7 +69,7 @@ FRenderQueryRHIRef FRenderQueryPool::AllocateQuery()
 	return Queries.Pop();
 }
 
-void FRenderQueryPool::ReleaseQuery( FRenderQueryRHIRef &Query )
+void FRenderQueryPool::ReleaseQuery(FRHICommandListImmediate& RHICmdList, FRenderQueryRHIRef &Query)
 {
 	if ( IsValidRef(Query) )
 	{
@@ -80,7 +80,7 @@ void FRenderQueryPool::ReleaseQuery( FRenderQueryRHIRef &Query )
 			Queries.Add( Query );
 
 			// Tell RHI we don't need the result anymore.
-			RHIResetRenderQuery( Query );
+			RHICmdList.ResetRenderQuery(Query);
 		}
 
 		// De-ref without deleting.
@@ -199,7 +199,7 @@ const uint8* FSceneViewState::GetPrecomputedVisibilityData(FViewInfo& View, cons
 	return PrecomputedVisibilityData;
 }
 
-void FSceneViewState::TrimOcclusionHistory(float MinHistoryTime,float MinQueryTime,int32 FrameNumber)
+void FSceneViewState::TrimOcclusionHistory(FRHICommandListImmediate& RHICmdList, float MinHistoryTime, float MinQueryTime, int32 FrameNumber)
 {
 	// Only trim every few frames, since stale entries won't cause problems
 	if (FrameNumber % 6 == 0)
@@ -212,7 +212,7 @@ void FSceneViewState::TrimOcclusionHistory(float MinHistoryTime,float MinQueryTi
 			// If the primitive has an old pending occlusion query, release it.
 			if(PrimitiveIt->LastConsideredTime < MinQueryTime)
 			{
-				PrimitiveIt->ReleaseQueries(OcclusionQueryPool);
+				PrimitiveIt->ReleaseQueries(RHICmdList, OcclusionQueryPool);
 			}
 
 			// If the primitive hasn't been considered for visibility recently, remove its history from the set.
@@ -224,7 +224,7 @@ void FSceneViewState::TrimOcclusionHistory(float MinHistoryTime,float MinQueryTi
 	}
 }
 
-bool FSceneViewState::IsShadowOccluded(FPrimitiveComponentId PrimitiveId, const ULightComponent* Light, int32 SplitIndex, bool bTranslucentShadow) const
+bool FSceneViewState::IsShadowOccluded(FRHICommandListImmediate& RHICmdList, FPrimitiveComponentId PrimitiveId, const ULightComponent* Light, int32 SplitIndex, bool bTranslucentShadow) const
 {
 	// Find the shadow's occlusion query from the previous frame.
 	const FSceneViewState::FProjectedShadowKey Key(PrimitiveId, Light, SplitIndex, bTranslucentShadow);
@@ -235,7 +235,7 @@ bool FSceneViewState::IsShadowOccluded(FPrimitiveComponentId PrimitiveId, const
 	// Only block on the query if not running SLI
 	const bool bWaitOnQuery = GNumActiveGPUsForRendering == 1;
 
-	if(Query && RHIGetRenderQueryResult(*Query, NumSamples, bWaitOnQuery))
+	if (Query && RHICmdList.GetRenderQueryResult(*Query, NumSamples, bWaitOnQuery))
 	{
 		// If the shadow's occlusion query didn't have any pixels visible the previous frame, it's occluded.
 		return NumSamples == 0;
@@ -303,7 +303,7 @@ void FSceneViewState::ApplyWorldOffset(FVector InOffset)
 class FOcclusionQueryIndexBuffer : public FIndexBuffer
 {
 public:
-	virtual void InitRHI()
+	virtual void InitRHI() override
 	{
 		const uint32 MaxBatchedPrimitives = FOcclusionQueryBatcher::OccludedPrimitiveQueryBatchSize;
 		const uint32 Stride = sizeof(uint16);
@@ -336,7 +336,7 @@ FOcclusionQueryBatcher::~FOcclusionQueryBatcher()
 	check(!BatchOcclusionQueries.Num());
 }
 
-void FOcclusionQueryBatcher::Flush()
+void FOcclusionQueryBatcher::Flush(FRHICommandListImmediate& RHICmdList)
 {
 	if(BatchOcclusionQueries.Num())
 	{
@@ -354,9 +354,9 @@ void FOcclusionQueryBatcher::Flush()
 			uint32 VertexBufferOffset = Batch.VertexAllocation.VertexOffset;
 			const int32 NumPrimitivesThisBatch = (BatchIndex != (NumBatches-1)) ? MaxBatchedPrimitives : NumBatchedPrimitives;
 				
-			RHIBeginRenderQuery(BatchOcclusionQuery);
-			RHISetStreamSource(0,VertexBufferRHI,sizeof(FVector),VertexBufferOffset);
-			RHIDrawIndexedPrimitive(
+			RHICmdList.BeginRenderQuery(BatchOcclusionQuery);
+			RHICmdList.SetStreamSource(0, VertexBufferRHI, sizeof(FVector), VertexBufferOffset);
+			RHICmdList.DrawIndexedPrimitive(
 				IndexBufferRHI,
 				PT_TriangleList,
 				/*BaseVertexIndex=*/ 0,
@@ -366,7 +366,7 @@ void FOcclusionQueryBatcher::Flush()
 				/*NumPrimitives=*/ 12 * NumPrimitivesThisBatch,
 				/*NumInstances=*/ 1
 				);
-			RHIEndRenderQuery(BatchOcclusionQuery);
+			RHICmdList.EndRenderQuery(BatchOcclusionQuery);
 		}
 		INC_DWORD_STAT_BY(STAT_OcclusionQueries,BatchOcclusionQueries.Num());
 
@@ -410,7 +410,7 @@ FRenderQueryRHIParamRef FOcclusionQueryBatcher::BatchPrimitive(const FVector& Bo
 	return CurrentBatchOcclusionQuery->Query;
 }
 
-static void IssueProjectedShadowOcclusionQuery(FViewInfo& View, const FProjectedShadowInfo& ProjectedShadowInfo, FOcclusionQueryVS* VertexShader)
+static void IssueProjectedShadowOcclusionQuery(FRHICommandListImmediate& RHICmdList, FViewInfo& View, const FProjectedShadowInfo& ProjectedShadowInfo, FOcclusionQueryVS* VertexShader)
 {
 	FSceneViewState* ViewState = (FSceneViewState*)View.State;
 	TMap<FSceneViewState::FProjectedShadowKey, FRenderQueryRHIRef>& ShadowOcclusionQueryMap = ViewState->ShadowOcclusionQueryMap;
@@ -427,20 +427,19 @@ static void IssueProjectedShadowOcclusionQuery(FViewInfo& View, const FProjected
 		);
 	if( !bIntersectsNearClippingPlane )
 	{
+
 		// Allocate an occlusion query for the primitive from the occlusion query pool.
 		const FRenderQueryRHIRef ShadowOcclusionQuery = ViewState->OcclusionQueryPool.AllocateQuery();
-
-		//@todo-rco: RHIPacketList
-		FRHICommandList& RHICmdList = FRHICommandList::GetNullRef();
+		
 		VertexShader->SetParameters(RHICmdList, View);
 
 		// Draw the primitive's bounding box, using the occlusion query.
-		RHIBeginRenderQuery(ShadowOcclusionQuery);
+		RHICmdList.BeginRenderQuery(ShadowOcclusionQuery);
 
 		void* VerticesPtr;
 		void* IndicesPtr;
 		// preallocate memory to fill out with vertices and indices
-		RHIBeginDrawIndexedPrimitiveUP( PT_TriangleList, 12, 8, sizeof(FVector), VerticesPtr, 0, NUM_CUBE_VERTICES, sizeof(uint16), IndicesPtr);
+		RHICmdList.BeginDrawIndexedPrimitiveUP(PT_TriangleList, 12, 8, sizeof(FVector), VerticesPtr, 0, NUM_CUBE_VERTICES, sizeof(uint16), IndicesPtr);
 		FVector* Vertices = (FVector*)VerticesPtr;
 		uint16* Indices = (uint16*)IndicesPtr;
 
@@ -479,8 +478,8 @@ static void IssueProjectedShadowOcclusionQuery(FViewInfo& View, const FProjected
 		checkSlow(ShadowOcclusionQueryMap.Find(Key) == NULL);
 		ShadowOcclusionQueryMap.Add(Key, ShadowOcclusionQuery);
 
-		RHIEndDrawIndexedPrimitiveUP();
-		RHIEndRenderQuery(ShadowOcclusionQuery);
+		RHICmdList.EndDrawIndexedPrimitiveUP();
+		RHICmdList.EndRenderQuery(ShadowOcclusionQuery);
 	}
 }
 
@@ -548,7 +547,7 @@ uint32 FHZBOcclusionTester::AddBounds( const FVector& BoundsCenter, const FVecto
 	return Index;
 }
 
-void FHZBOcclusionTester::MapResults()
+void FHZBOcclusionTester::MapResults(FRHICommandListImmediate& RHICmdList)
 {
 	check( !ResultsBuffer );
 
@@ -567,7 +566,7 @@ void FHZBOcclusionTester::MapResults()
 		int32 Width = 0;
 		int32 Height = 0;
 
-		RHIMapStagingSurface( ResultsTextureCPU->GetRenderTargetItem().ShaderResourceTexture, (void*&)ResultsBuffer, Width, Height );
+		RHICmdList.MapStagingSurface(ResultsTextureCPU->GetRenderTargetItem().ShaderResourceTexture, (void*&)ResultsBuffer, Width, Height);
 
 		// Can happen because of device removed, we might crash later but this occlusion culling system can behave gracefully.
 		if(!ResultsBuffer)
@@ -583,12 +582,12 @@ void FHZBOcclusionTester::MapResults()
 	check( ResultsBuffer );
 }
 
-void FHZBOcclusionTester::UnmapResults()
+void FHZBOcclusionTester::UnmapResults(FRHICommandListImmediate& RHICmdList)
 {
 	check( ResultsBuffer );
 	if(!IsInvalidFrame())
 	{
-		RHIUnmapStagingSurface( ResultsTextureCPU->GetRenderTargetItem().ShaderResourceTexture );
+		RHICmdList.UnmapStagingSurface(ResultsTextureCPU->GetRenderTargetItem().ShaderResourceTexture);
 	}
 	ResultsBuffer = NULL;
 }
@@ -673,10 +672,9 @@ public:
 
 IMPLEMENT_SHADER_TYPE(,FHZBTestPS,TEXT("HZBOcclusion"),TEXT("HZBTestPS"),SF_Pixel);
 
-void FHZBOcclusionTester::Submit( const FViewInfo& View )
+void FHZBOcclusionTester::Submit(FRHICommandListImmediate& RHICmdList, const FViewInfo& View)
 {
 	SCOPED_DRAW_EVENT(SubmitHZB, DEC_SCENE_ITEMS);
-	FRHICommandList& RHICmdList = FRHICommandList::GetNullRef(); // many things below which do not support parallel
 
 	FSceneViewState* ViewState = (FSceneViewState*)View.State;
 	if( !ViewState )
@@ -753,7 +751,6 @@ void FHZBOcclusionTester::Submit( const FViewInfo& View )
 		TShaderMapRef< FHZBTestPS >	PixelShader( GetGlobalShaderMap() );
 
 		static FGlobalBoundShaderState BoundShaderState;
-		RHICmdList.CheckIsNull(); // need new approach for "static FGlobalBoundShaderState" for parallel rendering
 		SetGlobalBoundShaderState(RHICmdList, BoundShaderState, GFilterVertexDeclaration.VertexDeclarationRHI, *VertexShader, *PixelShader);
 
 		PixelShader->SetParameters(RHICmdList, View,  ViewState->HZB, BoundsCenterTexture->GetRenderTargetItem().ShaderResourceTexture, BoundsExtentTexture->GetRenderTargetItem().ShaderResourceTexture );
@@ -854,7 +851,7 @@ public:
 IMPLEMENT_SHADER_TYPE(template<>,THZBBuildPS<0>,TEXT("HZBOcclusion"),TEXT("HZBBuildPS"),SF_Pixel);
 IMPLEMENT_SHADER_TYPE(template<>,THZBBuildPS<1>,TEXT("HZBOcclusion"),TEXT("HZBBuildPS"),SF_Pixel);
 
-void BuildHZB( const FViewInfo& View )
+void BuildHZB(FRHICommandListImmediate& RHICmdList, const FViewInfo& View)
 {
 	QUICK_SCOPE_CYCLE_COUNTER(STAT_BuildHZB);
 	SCOPED_DRAW_EVENT(BuildHZB, DEC_SCENE_ITEMS);
@@ -883,9 +880,7 @@ void BuildHZB( const FViewInfo& View )
 	
 	FSceneRenderTargetItem& HZBRenderTarget = ViewState->HZB.Texture->GetRenderTargetItem();
 	
-	//@todo-rco: RHIPacketList
-	FRHICommandList& RHICmdList = FRHICommandList::GetNullRef();
-
+	
 	RHICmdList.SetBlendState(TStaticBlendState<>::GetRHI());
 	RHICmdList.SetRasterizerState(TStaticRasterizerState<>::GetRHI());
 	RHICmdList.SetDepthStencilState(TStaticDepthStencilState< false, CF_Always >::GetRHI());
@@ -900,7 +895,7 @@ void BuildHZB( const FViewInfo& View )
 		TShaderMapRef< THZBBuildPS<0> >	PixelShader( GetGlobalShaderMap() );
 
 		static FGlobalBoundShaderState BoundShaderState;
-		RHICmdList.CheckIsNull(); // need new approach for "static FGlobalBoundShaderState" for parallel rendering
+		
 		SetGlobalBoundShaderState(RHICmdList, BoundShaderState, GFilterVertexDeclaration.VertexDeclarationRHI, *VertexShader, *PixelShader);
 
 		// Imperfect sampling, doesn't matter too much
@@ -934,7 +929,7 @@ void BuildHZB( const FViewInfo& View )
 		TShaderMapRef< THZBBuildPS<1> >	PixelShader( GetGlobalShaderMap() );
 
 		static FGlobalBoundShaderState BoundShaderState;
-		RHICmdList.CheckIsNull(); // need new approach for "static FGlobalBoundShaderState" for parallel rendering
+		
 		SetGlobalBoundShaderState(RHICmdList, BoundShaderState, GFilterVertexDeclaration.VertexDeclarationRHI, *VertexShader, *PixelShader);
 
 		PixelShader->SetParameters(RHICmdList, View, SrcSize, ViewState->HZB.MipSRVs[ MipIndex - 1 ] );
@@ -961,11 +956,10 @@ void BuildHZB( const FViewInfo& View )
 	GRenderTargetPool.VisualizeTexture.SetCheckPoint(RHICmdList, ViewState->HZB.Texture);
 }
 
-void FDeferredShadingSceneRenderer::BeginOcclusionTests()
+void FDeferredShadingSceneRenderer::BeginOcclusionTests(FRHICommandListImmediate& RHICmdList)
 {
 	SCOPE_CYCLE_COUNTER(STAT_BeginOcclusionTestsTime);
 	const bool bUseDownsampledDepth = IsValidRef(GSceneRenderTargets.GetSmallDepthSurface()) && GSceneRenderTargets.UseDownsizedOcclusionQueries();
-	FRHICommandList& RHICmdList = FRHICommandList::GetNullRef();
 
 	if (bUseDownsampledDepth)
 	{
@@ -1017,14 +1011,14 @@ void FDeferredShadingSceneRenderer::BeginOcclusionTests()
 			TMap<FSceneViewState::FProjectedShadowKey, FRenderQueryRHIRef>& ShadowOcclusionQueryMap = ViewState->ShadowOcclusionQueryMap;
 
 			// Clear primitives which haven't been visible recently out of the occlusion history, and reset old pending occlusion queries.
-			ViewState->TrimOcclusionHistory(ViewFamily.CurrentRealTime - GEngine->PrimitiveProbablyVisibleTime,ViewFamily.CurrentRealTime,FrameNumber);
+			ViewState->TrimOcclusionHistory(RHICmdList, ViewFamily.CurrentRealTime - GEngine->PrimitiveProbablyVisibleTime, ViewFamily.CurrentRealTime, FrameNumber);
 
 			// Give back all these occlusion queries to the pool.
 			for ( TMap<FSceneViewState::FProjectedShadowKey, FRenderQueryRHIRef>::TIterator QueryIt(ShadowOcclusionQueryMap); QueryIt; ++QueryIt )
 			{
 				//FRenderQueryRHIParamRef Query = QueryIt.Value();
 				//check( Query.GetRefCount() == 1 );
-				ViewState->OcclusionQueryPool.ReleaseQuery( QueryIt.Value() );
+				ViewState->OcclusionQueryPool.ReleaseQuery(RHICmdList, QueryIt.Value());
 			}
 			ShadowOcclusionQueryMap.Reset();
 
@@ -1055,8 +1049,7 @@ void FDeferredShadingSceneRenderer::BeginOcclusionTests()
 							if (!bCameraInsideLightGeometry)
 							{
 								const FRenderQueryRHIRef ShadowOcclusionQuery = ViewState->OcclusionQueryPool.AllocateQuery();
-								RHICmdList.CheckIsNull(); // RHIBeginRenderQuery
-								RHIBeginRenderQuery(ShadowOcclusionQuery);
+								RHICmdList.BeginRenderQuery(ShadowOcclusionQuery);
 
 								FSceneViewState::FProjectedShadowKey Key(
 									ProjectedShadowInfo.ParentSceneInfo ? 
@@ -1073,14 +1066,14 @@ void FDeferredShadingSceneRenderer::BeginOcclusionTests()
 								VertexShader->SetParametersWithBoundingSphere(RHICmdList, View, ProjectedShadowInfo.LightSceneInfo->Proxy->GetBoundingSphere());
 								StencilingGeometry::DrawSphere(RHICmdList);
 
-								RHIEndRenderQuery(ShadowOcclusionQuery);
+								RHICmdList.EndRenderQuery(ShadowOcclusionQuery);
 							}
 						}
 						// Don't query preshadows, since they are culled if their subject is occluded.
 						// Don't query if any subjects are visible because the shadow frustum will be definitely unoccluded
 						else if (!ProjectedShadowInfo.IsWholeSceneDirectionalShadow() && !ProjectedShadowInfo.bPreShadow && !ProjectedShadowInfo.SubjectsVisible(View))
 						{
-							IssueProjectedShadowOcclusionQuery(View, ProjectedShadowInfo, *VertexShader);
+							IssueProjectedShadowOcclusionQuery(RHICmdList, View, ProjectedShadowInfo, *VertexShader);
 						}
 					}
 
@@ -1088,7 +1081,7 @@ void FDeferredShadingSceneRenderer::BeginOcclusionTests()
 					for(int32 ShadowIndex = 0;ShadowIndex < VisibleLightInfo.OccludedPerObjectShadows.Num();ShadowIndex++)
 					{
 						const FProjectedShadowInfo& ProjectedShadowInfo = *VisibleLightInfo.OccludedPerObjectShadows[ShadowIndex];
-						IssueProjectedShadowOcclusionQuery(View, ProjectedShadowInfo, *VertexShader);
+						IssueProjectedShadowOcclusionQuery(RHICmdList, View, ProjectedShadowInfo, *VertexShader);
 					}
 				}
 			}
@@ -1102,11 +1095,11 @@ void FDeferredShadingSceneRenderer::BeginOcclusionTests()
 
 				{
 					SCOPED_DRAW_EVENT(IndividualQueries, DEC_SCENE_ITEMS);
-					View.IndividualOcclusionQueries.Flush();
+					View.IndividualOcclusionQueries.Flush(RHICmdList);
 				}
 				{
 					SCOPED_DRAW_EVENT(GroupedQueries, DEC_SCENE_ITEMS);
-					View.GroupedOcclusionQueries.Flush();
+					View.GroupedOcclusionQueries.Flush(RHICmdList);
 				}
 		    }
 	    }
@@ -1123,14 +1116,14 @@ void FDeferredShadingSceneRenderer::BeginOcclusionTests()
 
 			SCOPED_DRAW_EVENT(HZB, DEC_SCENE_ITEMS);
 
-			BuildHZB( View );
-			ViewState->HZBOcclusionTests.Submit( View );
+			BuildHZB(RHICmdList, View);
+			ViewState->HZBOcclusionTests.Submit(RHICmdList, View);
 		}
 	}
 
 	if (bUseDownsampledDepth)
 	{
 		// Restore default render target
-		GSceneRenderTargets.BeginRenderingSceneColor();
+		GSceneRenderTargets.BeginRenderingSceneColor(RHICmdList);
 	}
 }
