UE4 - rework RHI user-facing APIs to be safer and incorporate the desired sematics

diff --git a/Engine/Source/Runtime/Renderer/Private/ScenePostProcessing.cpp b/Engine/Source/Runtime/Renderer/Private/ScenePostProcessing.cpp
--- a/Engine/Source/Runtime/Renderer/Private/ScenePostProcessing.cpp
+++ b/Engine/Source/Runtime/Renderer/Private/ScenePostProcessing.cpp
@@ -80,14 +80,14 @@ IMPLEMENT_SHADER_TYPE(,FGammaCorrectionVS,TEXT("GammaCorrection"),TEXT("MainVS")
 * Finish rendering a view, writing the contents to ViewFamily.RenderTarget.
 * @param View - The view to process.
 */
-void FDeferredShadingSceneRenderer::FinishRenderViewTarget(const FViewInfo* View, bool bLastView)
+void FDeferredShadingSceneRenderer::FinishRenderViewTarget(FRHICommandListImmediate& RHICmdList, const FViewInfo* View, bool bLastView)
 {
 	TRefCountPtr<IPooledRenderTarget> VelocityRT;
 
 	// Render the velocities of movable objects for the motion blur effect (currently we only support one view)
-	RenderVelocities(*View, VelocityRT, bLastView);
+	RenderVelocities(RHICmdList, *View, VelocityRT, bLastView);
 
-	GPostProcessing.Process(*View, VelocityRT);
+	GPostProcessing.Process(RHICmdList, *View, VelocityRT);
 
 	// we rendered to it during the frame, seems we haven't made use of it, because there is should be released
 	FSceneViewState* ViewState = (FSceneViewState*)View->State;
@@ -98,7 +98,7 @@ void FDeferredShadingSceneRenderer::FinishRenderViewTarget(const FViewInfo* View
 }
 
 // TODO: REMOVE if no longer needed:
-void FSceneRenderer::GammaCorrectToViewportRenderTarget(FRHICommandList& RHICmdList, const FViewInfo* View, float OverrideGamma)
+void FSceneRenderer::GammaCorrectToViewportRenderTarget(FRHICommandListImmediate& RHICmdList, const FViewInfo* View, float OverrideGamma)
 {
 	// Set the view family's render target/viewport.
 	SetRenderTarget(RHICmdList, ViewFamily.RenderTarget->GetRenderTargetTexture(), FTextureRHIRef());
@@ -123,7 +123,7 @@ void FSceneRenderer::GammaCorrectToViewportRenderTarget(FRHICommandList& RHICmdL
 	TShaderMapRef<FGammaCorrectionPS> PixelShader(GetGlobalShaderMap());
 
 	static FGlobalBoundShaderState PostProcessBoundShaderState;
-	RHICmdList.CheckIsNull(); // need new approach for "static FGlobalBoundShaderState" for parallel rendering
+	
 	SetGlobalBoundShaderState(RHICmdList, PostProcessBoundShaderState, GFilterVertexDeclaration.VertexDeclarationRHI, *VertexShader, *PixelShader);
 
 	float InvDisplayGamma = 1.0f / ViewFamily.RenderTarget->GetDisplayGamma();
diff --git a/Engine/Source/Runtime/Renderer/Private/ScenePrivate.h b/Engine/Source/Runtime/Renderer/Private/ScenePrivate.h
--- a/Engine/Source/Runtime/Renderer/Private/ScenePrivate.h
+++ b/Engine/Source/Runtime/Renderer/Private/ScenePrivate.h
@@ -134,15 +134,15 @@ public:
 	}
 
 	template<class TOcclusionQueryPool> // here we use a template just to allow this to be inlined without sorting out the header order
-	FORCEINLINE void ReleaseQueries(TOcclusionQueryPool& Pool)
+	FORCEINLINE void ReleaseQueries(FRHICommandListImmediate& RHICmdList, TOcclusionQueryPool& Pool)
 	{
 #if WITH_SLI
 		for (int32 QueryIndex = 0; QueryIndex < NumBufferedFrames; QueryIndex++)
 		{
-			Pool.ReleaseQuery(PendingOcclusionQuery[QueryIndex]);
+			Pool.ReleaseQuery(RHICmdList, PendingOcclusionQuery[QueryIndex]);
 		}
 #else
-		Pool.ReleaseQuery(PendingOcclusionQuery);
+		Pool.ReleaseQuery(RHICmdList, PendingOcclusionQuery);
 #endif
 	}
 
@@ -205,7 +205,7 @@ public:
 	FRenderQueryRHIRef AllocateQuery();
 
 	/** De-reference an render query, returning it to the pool instead of deleting it when the refcount reaches 0. */
-	void ReleaseQuery( FRenderQueryRHIRef &Query );
+	void ReleaseQuery(FRHICommandListImmediate& RHICmdList, FRenderQueryRHIRef &Query);
 
 private:
 	/** Container for available render queries. */
@@ -578,14 +578,14 @@ public:
 	 *							visible and unoccluded since this time will be discarded.
 	 * @param MinQueryTime - The pending occlusion queries older than this will be discarded.
 	 */
-	void TrimOcclusionHistory(float MinHistoryTime,float MinQueryTime,int32 FrameNumber);
+	void TrimOcclusionHistory(FRHICommandListImmediate& RHICmdList, float MinHistoryTime, float MinQueryTime, int32 FrameNumber);
 
 	/**
 	 * Checks whether a shadow is occluded this frame.
 	 * @param Primitive - The shadow subject.
 	 * @param Light - The shadow source.
 	 */
-	bool IsShadowOccluded(FPrimitiveComponentId PrimitiveId, const ULightComponent* Light, int32 SplitIndex, bool bTranslucentShadow) const;
+	bool IsShadowOccluded(FRHICommandListImmediate& RHICmdList, FPrimitiveComponentId PrimitiveId, const ULightComponent* Light, int32 SplitIndex, bool bTranslucentShadow) const;
 
 	TRefCountPtr<IPooledRenderTarget>& GetEyeAdaptation()
 	{
@@ -606,12 +606,12 @@ public:
 	}
 
     // FRenderResource interface.
-	virtual void InitDynamicRHI()
+	virtual void InitDynamicRHI() override
     {
 		HZBOcclusionTests.InitDynamicRHI();
 	}
 
-    virtual void ReleaseDynamicRHI()
+	virtual void ReleaseDynamicRHI() override
     {
         ShadowOcclusionQueryMap.Reset();
 		PrimitiveOcclusionHistorySet.Empty();
@@ -732,8 +732,8 @@ public:
 		MaxCubemaps(0)
 	{}
 
-	virtual void InitDynamicRHI();
-	virtual void ReleaseDynamicRHI();
+	virtual void InitDynamicRHI() override;
+	virtual void ReleaseDynamicRHI() override;
 
 	/** 
 	 * Updates the maximum number of cubemaps that this array is allocated for.
@@ -1424,7 +1424,7 @@ private:
 	 * Adds a primitive to the scene.  Called in the rendering thread by AddPrimitive.
 	 * @param PrimitiveSceneInfo - The primitive being added.
 	 */
-	void AddPrimitiveSceneInfo_RenderThread(FPrimitiveSceneInfo* PrimitiveSceneInfo);
+	void AddPrimitiveSceneInfo_RenderThread(FRHICommandListImmediate& RHICmdList, FPrimitiveSceneInfo* PrimitiveSceneInfo);
 
 	/**
 	 * Removes a primitive from the scene.  Called in the rendering thread by RemovePrimitive.
@@ -1433,7 +1433,7 @@ private:
 	void RemovePrimitiveSceneInfo_RenderThread(FPrimitiveSceneInfo* PrimitiveSceneInfo);
 
 	/** Updates a primitive's transform, called on the rendering thread. */
-	void UpdatePrimitiveTransform_RenderThread(FPrimitiveSceneProxy* PrimitiveSceneProxy, const FBoxSphereBounds& WorldBounds, const FBoxSphereBounds& LocalBounds, const FMatrix& LocalToWorld, const FVector& OwnerPosition);
+	void UpdatePrimitiveTransform_RenderThread(FRHICommandListImmediate& RHICmdList, FPrimitiveSceneProxy* PrimitiveSceneProxy, const FBoxSphereBounds& WorldBounds, const FBoxSphereBounds& LocalBounds, const FMatrix& LocalToWorld, const FVector& OwnerPosition);
 
 	/** Updates a single primitive's lighting attachment root. */
 	void UpdatePrimitiveLightingAttachmentRoot(UPrimitiveComponent* Primitive);
@@ -1468,10 +1468,10 @@ private:
 	void UpdateAllReflectionCaptures();
 
 	/** Sets shader maps on the specified materials without blocking. */
-	void SetShaderMapsOnMaterialResources_RenderThread(const FMaterialsToUpdateMap& MaterialsToUpdate);
+	void SetShaderMapsOnMaterialResources_RenderThread(FRHICommandListImmediate& RHICmdList, const FMaterialsToUpdateMap& MaterialsToUpdate);
 
 	/** Updates static draw lists for the given materials. */
-	void UpdateStaticDrawListsForMaterials_RenderThread(const TArray<const FMaterial*>& Materials);
+	void UpdateStaticDrawListsForMaterials_RenderThread(FRHICommandListImmediate& RHICmdList, const TArray<const FMaterial*>& Materials);
 
 	/**
 	 * Shifts scene data by provided delta
diff --git a/Engine/Source/Runtime/Renderer/Private/SceneRendering.cpp b/Engine/Source/Runtime/Renderer/Private/SceneRendering.cpp
--- a/Engine/Source/Runtime/Renderer/Private/SceneRendering.cpp
+++ b/Engine/Source/Runtime/Renderer/Private/SceneRendering.cpp
@@ -654,7 +654,7 @@ FSceneRenderer::~FSceneRenderer()
 /** 
 * Finishes the view family rendering.
 */
-void FSceneRenderer::RenderFinish()
+void FSceneRenderer::RenderFinish(FRHICommandListImmediate& RHICmdList)
 {
 	SCOPED_DRAW_EVENT(RenderFinish, DEC_SCENE_ITEMS);
 
@@ -664,7 +664,7 @@ void FSceneRenderer::RenderFinish()
 		const FViewInfo& View = Views[0];
 
 		FMemMark Mark(FMemStack::Get());
-		FRenderingCompositePassContext CompositeContext(View);
+		FRenderingCompositePassContext CompositeContext(RHICmdList, View);
 
 		FRenderingCompositeOutputRef BusyWait;
 		{
@@ -809,14 +809,14 @@ void FSceneRenderer::RenderFinish()
 				continue;
 			}
 
-			GRenderTargetPool.PresentContent(View);
+			GRenderTargetPool.PresentContent(RHICmdList, View);
 		}
 	}
 
 #endif
 
 	// Notify the RHI we are done rendering a scene.
-	RHIEndScene();
+	RHICmdList.EndScene();
 }
 
 FSceneRenderer* FSceneRenderer::CreateSceneRenderer(const FSceneViewFamily* InViewFamily, FHitProxyConsumer* HitProxyConsumer)
@@ -833,7 +833,7 @@ FSceneRenderer* FSceneRenderer::CreateSceneRenderer(const FSceneViewFamily* InVi
 	}
 }
 
-void FSceneRenderer::RenderCustomDepthPass()
+void FSceneRenderer::RenderCustomDepthPass(FRHICommandListImmediate& RHICmdList)
 {
 	if(!IsFeatureLevelSupported(GRHIShaderPlatform, ERHIFeatureLevel::SM3))
 	{
@@ -858,36 +858,33 @@ void FSceneRenderer::RenderCustomDepthPass()
 	}
 
 	// Render CustomDepth
-	if(GSceneRenderTargets.BeginRenderingCustomDepth(bPrimitives))
+	if (GSceneRenderTargets.BeginRenderingCustomDepth(RHICmdList, bPrimitives))
 	{
 		SCOPED_DRAW_EVENT(CustomDepth, DEC_SCENE_ITEMS);
 
-		//@todo-rco: RHIPacketList
-		FRHICommandList& RHICmdList = FRHICommandList::GetNullRef();
-
 		for(int32 ViewIndex = 0;ViewIndex < Views.Num();ViewIndex++)
 		{
 			SCOPED_CONDITIONAL_DRAW_EVENTF(EventView, Views.Num() > 1, DEC_SCENE_ITEMS, TEXT("View%d"), ViewIndex);
 
 			FViewInfo& View = Views[ViewIndex];
 
-			RHISetViewport(View.ViewRect.Min.X, View.ViewRect.Min.Y, 0.0f, View.ViewRect.Max.X, View.ViewRect.Max.Y, 1.0f);
+			RHICmdList.SetViewport(View.ViewRect.Min.X, View.ViewRect.Min.Y, 0.0f, View.ViewRect.Max.X, View.ViewRect.Max.Y, 1.0f);
 
 			// Note, this is a reversed Z depth surface, so 0.0f is the far plane.
-			RHIClear(false, FLinearColor(0,0,0,0), true, 0.0f, false, 0, FIntRect());
+			RHICmdList.Clear(false, FLinearColor(0, 0, 0, 0), true, 0.0f, false, 0, FIntRect());
 			
 			// seems this is set each draw call anyway
-			RHISetRasterizerState(TStaticRasterizerState<>::GetRHI());
+			RHICmdList.SetRasterizerState(TStaticRasterizerState<>::GetRHI());
 
 			// Note, this is a reversed Z depth surface, so CF_GreaterEqual.
-			RHISetDepthStencilState(TStaticDepthStencilState<true,CF_GreaterEqual>::GetRHI());
-			RHISetBlendState(TStaticBlendState<>::GetRHI());
+			RHICmdList.SetDepthStencilState(TStaticDepthStencilState<true, CF_GreaterEqual>::GetRHI());
+			RHICmdList.SetBlendState(TStaticBlendState<>::GetRHI());
 
 			View.CustomDepthSet.DrawPrims(RHICmdList, &View, false);
 		}
 
 		// resolve using the current ResolveParams 
-		GSceneRenderTargets.FinishRenderingCustomDepth();
+		GSceneRenderTargets.FinishRenderingCustomDepth(RHICmdList);
 	}
 }
 
@@ -919,7 +916,7 @@ void FSceneRenderer::OnStartFrame()
  *
  * @param SceneRenderer	Scene renderer to use for rendering.
  */
-static void RenderViewFamily_RenderThread( FSceneRenderer* SceneRenderer )
+static void RenderViewFamily_RenderThread(FRHICommandListImmediate& RHICmdList, FSceneRenderer* SceneRenderer)
 {
     FMemMark MemStackMark(FMemStack::Get());
 
@@ -946,12 +943,12 @@ static void RenderViewFamily_RenderThread( FSceneRenderer* SceneRenderer )
 		if(SceneRenderer->ViewFamily.EngineShowFlags.HitProxies)
 		{
 			// Render the scene's hit proxies.
-			SceneRenderer->RenderHitProxies();
+			SceneRenderer->RenderHitProxies(RHICmdList);
 		}
 		else
 		{
 			// Render the scene.
-			SceneRenderer->Render();
+			SceneRenderer->Render(RHICmdList);
 		}
 
 #if STATS
@@ -1016,7 +1013,7 @@ void FRendererModule::BeginRenderingViewFamily(FCanvas* Canvas,const FSceneViewF
 			FDrawSceneCommand,
 			FSceneRenderer*,SceneRenderer,SceneRenderer,
 		{
-			RenderViewFamily_RenderThread(SceneRenderer);
+			RenderViewFamily_RenderThread(RHICmdList, SceneRenderer);
 		});
 	}
 }
diff --git a/Engine/Source/Runtime/Renderer/Private/SceneRendering.h b/Engine/Source/Runtime/Renderer/Private/SceneRendering.h
--- a/Engine/Source/Runtime/Renderer/Private/SceneRendering.h
+++ b/Engine/Source/Runtime/Renderer/Private/SceneRendering.h
@@ -84,10 +84,10 @@ public:
 	* @param PhaseSortedPrimitives - array with the primitives we want to draw
 	* @param bSeparateTranslucencyPass
 	*/
-	void DrawPrimitives(const class FViewInfo& View, class FDeferredShadingSceneRenderer& Renderer, bool bSeparateTranslucencyPass) const;
+	void DrawPrimitives(FRHICommandListImmediate& RHICmdList, const class FViewInfo& View, class FDeferredShadingSceneRenderer& Renderer, bool bSeparateTranslucencyPass) const;
 
 	/** Draw all the primitives in this set for the forward shading pipeline. */
-	void DrawPrimitivesForForwardShading(const class FViewInfo& View, class FSceneRenderer& Renderer) const;
+	void DrawPrimitivesForForwardShading(FRHICommandListImmediate& RHICmdList, const class FViewInfo& View, class FSceneRenderer& Renderer) const;
 
 	/**
 	* Add a new primitive to the list of sorted prims
@@ -184,10 +184,10 @@ private:
 	TArray<FSortedPrim,SceneRenderingAllocator> SortedSeparateTranslucencyPrims;
 
 	/** Renders a single primitive for the deferred shading pipeline. */
-	void RenderPrimitive(const FViewInfo& View, FPrimitiveSceneInfo* PrimitiveSceneInfo, const FPrimitiveViewRelevance& ViewRelevance, const FProjectedShadowInfo* TranslucentSelfShadow, bool bSeparateTranslucencyPass) const;
+	void RenderPrimitive(FRHICommandListImmediate& RHICmdList, const FViewInfo& View, FPrimitiveSceneInfo* PrimitiveSceneInfo, const FPrimitiveViewRelevance& ViewRelevance, const FProjectedShadowInfo* TranslucentSelfShadow, bool bSeparateTranslucencyPass) const;
 
 	/** Renders a single primitive for the forward shading pipeline. */
-	void RenderPrimitiveForForwardShading(const FViewInfo& View, FPrimitiveSceneInfo* PrimitiveSceneInfo, const FPrimitiveViewRelevance& ViewRelevance) const;
+	void RenderPrimitiveForForwardShading(FRHICommandListImmediate& RHICmdList, const FViewInfo& View, FPrimitiveSceneInfo* PrimitiveSceneInfo, const FPrimitiveViewRelevance& ViewRelevance) const;
 };
 
 /** A batched occlusion primitive. */
@@ -214,7 +214,7 @@ public:
 	~FOcclusionQueryBatcher();
 
 	/** Renders the current batch and resets the batch state. */
-	void Flush();
+	void Flush(FRHICommandListImmediate& RHICmdList);
 
 	/**
 	 * Batches a primitive's occlusion query for rendering.
@@ -253,16 +253,16 @@ public:
 					~FHZBOcclusionTester() {}
 
 	// FRenderResource interface
-	virtual void	InitDynamicRHI();
-	virtual void	ReleaseDynamicRHI();
+					virtual void	InitDynamicRHI() override;
+					virtual void	ReleaseDynamicRHI() override;
 	
 	uint32			GetNum() const { return Primitives.Num(); }
 
 	uint32			AddBounds( const FVector& BoundsOrigin, const FVector& BoundsExtent );
-	void			Submit( const FViewInfo& View );
+	void			Submit(FRHICommandListImmediate& RHICmdList, const FViewInfo& View);
 
-	void			MapResults();
-	void			UnmapResults();
+	void			MapResults(FRHICommandListImmediate& RHICmdList);
+	void			UnmapResults(FRHICommandListImmediate& RHICmdList);
 	bool			IsVisible( uint32 Index ) const;
 
 	bool IsValidFrame(uint32 FrameNumber) const;
@@ -520,8 +520,8 @@ public:
 
 	// FSceneRenderer interface
 
-	virtual void Render() = 0;
-	virtual void RenderHitProxies() {}
+	virtual void Render(FRHICommandListImmediate& RHICmdList) = 0;
+	virtual void RenderHitProxies(FRHICommandListImmediate& RHICmdList) {}
 
 	/** Creates a scene renderer based on the current feature level. */
 	static FSceneRenderer* CreateSceneRenderer(const FSceneViewFamily* InViewFamily, FHitProxyConsumer* HitProxyConsumer);
@@ -533,10 +533,10 @@ protected:
 	// Shared functionality between all scene renderers
 
 	/** Performs once per frame setup prior to visibility determination. */
-	void PreVisibilityFrameSetup();
+	void PreVisibilityFrameSetup(FRHICommandListImmediate& RHICmdList);
 
 	/** Computes which primitives are visible and relevant for each view. */
-	void ComputeViewVisibility();
+	void ComputeViewVisibility(FRHICommandListImmediate& RHICmdList);
 
 	/** Performs once per frame setup after to visibility determination. */
 	void PostVisibilityFrameSetup();
@@ -551,17 +551,17 @@ protected:
 	bool ShouldRenderTranslucency() const;
 
 	/** TODO: REMOVE if no longer needed: Copies scene color to the viewport's render target after applying gamma correction. */
-	void GammaCorrectToViewportRenderTarget(FRHICommandList& RHICmdList, const FViewInfo* View, float OverrideGamma);
+	void GammaCorrectToViewportRenderTarget(FRHICommandListImmediate& RHICmdList, const FViewInfo* View, float OverrideGamma);
 
 	/** Updates state for the end of the frame. */
-	void RenderFinish();
+	void RenderFinish(FRHICommandListImmediate& RHICmdList);
 
-	void RenderCustomDepthPass();
+	void RenderCustomDepthPass(FRHICommandListImmediate& RHICmdList);
 
 	void OnStartFrame();
 
 	/** Renders the scene's distortion */
-	void RenderDistortion(FRHICommandList& RHICmdList);
+	void RenderDistortion(FRHICommandListImmediate& RHICmdList);
 };
 
 
@@ -576,18 +576,18 @@ public:
 
 	// FSceneRenderer interface
 
-	virtual void Render() override;
+	virtual void Render(FRHICommandListImmediate& RHICmdList) override;
 
 protected:
 
-	void InitViews();
+	void InitViews(FRHICommandListImmediate& RHICmdList);
 
 	/** Renders the opaque base pass for forward shading. */
-	void RenderForwardShadingBasePass(FRHICommandList& RHICmdList);
+	void RenderForwardShadingBasePass(FRHICommandListImmediate& RHICmdList);
 
 	/** Makes a copy of scene alpha so PC can emulate ES2 framebuffer fetch. */
-	void CopySceneAlpha(FRHICommandList& RHICmdList, const FSceneView& View);
+	void CopySceneAlpha(FRHICommandListImmediate& RHICmdList, const FSceneView& View);
 
 	/** Renders the base pass for translucency. */
-	void RenderTranslucency();
+	void RenderTranslucency(FRHICommandListImmediate& RHICmdList);
 };
diff --git a/Engine/Source/Runtime/Renderer/Private/SceneVisibility.cpp b/Engine/Source/Runtime/Renderer/Private/SceneVisibility.cpp
--- a/Engine/Source/Runtime/Renderer/Private/SceneVisibility.cpp
+++ b/Engine/Source/Runtime/Renderer/Private/SceneVisibility.cpp
@@ -321,7 +321,7 @@ static void UpdatePrimitiveFading(const FScene* Scene, FViewInfo& View)
 /**
  * Cull occluded primitives in the view.
  */
-static int32 OcclusionCull(const FScene* Scene, FViewInfo& View)
+static int32 OcclusionCull(FRHICommandListImmediate& RHICmdList, const FScene* Scene, FViewInfo& View)
 {
 	SCOPE_CYCLE_COUNTER(STAT_OcclusionCull);
 
@@ -372,7 +372,7 @@ static int32 OcclusionCull(const FScene* Scene, FViewInfo& View)
 			{
 				QUICK_SCOPE_CYCLE_COUNTER(STAT_MapHZBResults);
 				check(!ViewState->HZBOcclusionTests.IsValidFrame(View.FrameNumber));
-				ViewState->HZBOcclusionTests.MapResults();
+				ViewState->HZBOcclusionTests.MapResults(RHICmdList);
 			}
 
 			FViewElementPDI OcclusionPDI(&View, NULL);
@@ -439,7 +439,7 @@ static int32 OcclusionCull(const FScene* Scene, FViewInfo& View)
 							if (IsValidRef(PastQuery))
 							{
 								// NOTE: RHIGetOcclusionQueryResult should never fail when using a blocking call, rendering artifacts may show up.
-								if ( RHIGetRenderQueryResult(PastQuery,NumSamples,true) )
+								if (RHICmdList.GetRenderQueryResult(PastQuery, NumSamples, true))
 								{
 									// we render occlusion without MSAA
 									uint32 NumPixels = (uint32)NumSamples;
@@ -499,7 +499,7 @@ static int32 OcclusionCull(const FScene* Scene, FViewInfo& View)
 
 					if (bClearQueries)
 					{
-						ViewState->OcclusionQueryPool.ReleaseQuery( PrimitiveOcclusionHistory->GetPastQuery(View.FrameNumber) );
+						ViewState->OcclusionQueryPool.ReleaseQuery(RHICmdList, PrimitiveOcclusionHistory->GetPastQuery(View.FrameNumber));
 					}
 				}
 
@@ -596,7 +596,7 @@ static int32 OcclusionCull(const FScene* Scene, FViewInfo& View)
 			{
 				QUICK_SCOPE_CYCLE_COUNTER(STAT_HZBUnmapResults);
 
-				ViewState->HZBOcclusionTests.UnmapResults();
+				ViewState->HZBOcclusionTests.UnmapResults(RHICmdList);
 
 				if( bSubmitQueries )
 				{
@@ -638,6 +638,7 @@ typedef TArray<uint8, SceneRenderingAllocator> FPrimitivePreRenderViewMasks;
  *                                callback for this view will have ViewBit set.
  */
 static void ComputeRelevanceForView(
+	FRHICommandListImmediate& RHICmdList,
 	const FScene* Scene,
 	FViewInfo& View,
 	uint8 ViewBit,
@@ -749,7 +750,7 @@ static void ComputeRelevanceForView(
 			PrimitiveSceneInfo->bNeedsCachedReflectionCaptureUpdate = false;
 		}
 
-		PrimitiveSceneInfo->ConditionalUpdateStaticMeshes();
+		PrimitiveSceneInfo->ConditionalUpdateStaticMeshes(RHICmdList);
 	}
 }
 
@@ -918,10 +919,10 @@ float Halton( int32 Index, int32 Base )
 	return Result;
 }
 
-void FSceneRenderer::PreVisibilityFrameSetup()
+void FSceneRenderer::PreVisibilityFrameSetup(FRHICommandListImmediate& RHICmdList)
 {
 	// Notify the RHI we are beginning to render a scene.
-	RHIBeginScene();
+	RHICmdList.BeginScene();
 
 	// Notify the FX system that the scene is about to perform visibility checks.
 	if (Scene->FXSystem)
@@ -1192,7 +1193,7 @@ void FSceneRenderer::PreVisibilityFrameSetup()
 	}
 }
 
-void FSceneRenderer::ComputeViewVisibility()
+void FSceneRenderer::ComputeViewVisibility(FRHICommandListImmediate& RHICmdList)
 {
 	SCOPE_CYCLE_COUNTER(STAT_ViewVisibilityTime);
 
@@ -1363,7 +1364,7 @@ void FSceneRenderer::ComputeViewVisibility()
 		// Occlusion cull for all primitives in the view frustum, but not in wireframe.
 		if (!View.Family->EngineShowFlags.Wireframe)
 		{
-			int32 NumOccludedPrimitivesInView = OcclusionCull(Scene, View);
+			int32 NumOccludedPrimitivesInView = OcclusionCull(RHICmdList, Scene, View);
 			STAT(NumOccludedPrimitives += NumOccludedPrimitivesInView);
 		}
 
@@ -1372,7 +1373,7 @@ void FSceneRenderer::ComputeViewVisibility()
 
 		// Compute view relevance for all visible primitives.
 		RelevantStaticPrimitives.Reset();
-		ComputeRelevanceForView(Scene, View, ViewBit, RelevantStaticPrimitives, PreRenderViewMasks);
+		ComputeRelevanceForView(RHICmdList, Scene, View, ViewBit, RelevantStaticPrimitives, PreRenderViewMasks);
 		MarkRelevantStaticMeshesForView(Scene, View, RelevantStaticPrimitives);
 
 #if !(UE_BUILD_SHIPPING || UE_BUILD_TEST)
@@ -1510,14 +1511,14 @@ uint32 GetShadowQuality();
  * Initialize scene's views.
  * Check visibility, sort translucent items, etc.
  */
-void FDeferredShadingSceneRenderer::InitViews()
+void FDeferredShadingSceneRenderer::InitViews(FRHICommandListImmediate& RHICmdList)
 {
 	SCOPED_DRAW_EVENT(InitViews, DEC_SCENE_ITEMS);
 
 	SCOPE_CYCLE_COUNTER(STAT_InitViewsTime);
 
-	PreVisibilityFrameSetup();
-	ComputeViewVisibility();
+	PreVisibilityFrameSetup(RHICmdList);
+	ComputeViewVisibility(RHICmdList);
 	PostVisibilityFrameSetup();
 
 	FVector AverageViewPosition(0);
@@ -1535,7 +1536,7 @@ void FDeferredShadingSceneRenderer::InitViews()
 	if (bDynamicShadows && !IsSimpleDynamicLightingEnabled())
 	{
 		// Setup dynamic shadows.
-		InitDynamicShadows();
+		InitDynamicShadows(RHICmdList);
 	}
 
 	if(ViewFamily.EngineShowFlags.MotionBlur && GRenderingRealtimeClock.GetGamePaused())
diff --git a/Engine/Source/Runtime/Renderer/Private/ShaderBaseClasses.cpp b/Engine/Source/Runtime/Renderer/Private/ShaderBaseClasses.cpp
--- a/Engine/Source/Runtime/Renderer/Private/ShaderBaseClasses.cpp
+++ b/Engine/Source/Runtime/Renderer/Private/ShaderBaseClasses.cpp
@@ -223,7 +223,7 @@ void FMaterialShader::SetParameters(
 	// if we are in a postprocessing pass
 	if(View.RenderingCompositePassContext)
 	{
-		PostprocessParameter.Set(RHICmdList, ShaderRHI, *View.RenderingCompositePassContext, TStaticSamplerState<>::GetRHI());
+		PostprocessParameter.Set(ShaderRHI, *View.RenderingCompositePassContext, TStaticSamplerState<>::GetRHI());
 	}
 
 	//Use of the eye adaptation texture here is experimental and potentially dangerous as it can introduce a feedback loop. May be removed.
diff --git a/Engine/Source/Runtime/Renderer/Private/ShadowRendering.cpp b/Engine/Source/Runtime/Renderer/Private/ShadowRendering.cpp
--- a/Engine/Source/Runtime/Renderer/Private/ShadowRendering.cpp
+++ b/Engine/Source/Runtime/Renderer/Private/ShadowRendering.cpp
@@ -1127,12 +1127,12 @@ static void CheckShadowDepthMaterials(const FMaterialRenderProxy* InRenderProxy,
 	check(Material == InMaterial);
 }
 
-void FProjectedShadowInfo::ClearDepth(FDeferredShadingSceneRenderer* SceneRenderer)
+void FProjectedShadowInfo::ClearDepth(FRHICommandListImmediate& RHICmdList, FDeferredShadingSceneRenderer* SceneRenderer)
 {
 	if (bOnePassPointLightShadow)
 	{
 		// Set the viewport to the whole render target since it's a cube map, don't leave any border space
-		RHISetViewport(
+		RHICmdList.SetViewport(
 			0,
 			0,
 			0.0f,
@@ -1142,14 +1142,14 @@ void FProjectedShadowInfo::ClearDepth(FDeferredShadingSceneRenderer* SceneRender
 			);
 
 		// Clear depth only.
-		RHIClear(false,FColor(255,255,255),true,1.0f,false,0, FIntRect());
+		RHICmdList.Clear(false, FColor(255, 255, 255), true, 1.0f, false, 0, FIntRect());
 	}
 	else
 	{
 		if (bReflectiveShadowmap)
 		{
 			// Set the viewport for the shadow.
-			RHISetViewport(
+			RHICmdList.SetViewport(
 				X,
 				Y,
 				0.0f,
@@ -1160,12 +1160,12 @@ void FProjectedShadowInfo::ClearDepth(FDeferredShadingSceneRenderer* SceneRender
 
 			// Clear color and depth targets
 			FLinearColor ClearColors[2] = {FLinearColor(0,0,1,0), FLinearColor(0,0,0,0)};
-			RHIClearMRT(true, 2, ClearColors, true, 1.0f, false, 0, FIntRect());
+			RHICmdList.ClearMRT(true, 2, ClearColors, true, 1.0f, false, 0, FIntRect());
 		}
 		else
 		{
 			// Set the viewport for the shadow.
-			RHISetViewport(
+			RHICmdList.SetViewport(
 				X,
 				Y,
 				0.0f,
@@ -1175,7 +1175,7 @@ void FProjectedShadowInfo::ClearDepth(FDeferredShadingSceneRenderer* SceneRender
 				);
 
 			// Clear depth only.
-			RHIClear(false,FColor(255,255,255),true,1.0f,false,0, FIntRect());
+			RHICmdList.Clear(false, FColor(255, 255, 255), true, 1.0f, false, 0, FIntRect());
 		}
 	}
 }
@@ -1263,7 +1263,7 @@ void DrawShadowMeshElements(FRHICommandList& RHICmdList, const FViewInfo& View,
 	}
 }
 
-void FProjectedShadowInfo::RenderDepth(FRHICommandList& RHICmdList, FDeferredShadingSceneRenderer* SceneRenderer)
+void FProjectedShadowInfo::RenderDepth(FRHICommandListImmediate& RHICmdList, FDeferredShadingSceneRenderer* SceneRenderer)
 {
 #if WANTS_DRAW_MESH_EVENTS
 	FString EventName;
@@ -1343,7 +1343,6 @@ void FProjectedShadowInfo::RenderDepth(FRHICommandList& RHICmdList, FDeferredSha
 	}
 	check(FoundView);
 
-	RHICmdList.CheckIsNull(); // this messes with the view, won't work with parallel rendering. Need to copy and modify the view instead
 	// Backup properties of the view that we will override
 	TUniformBufferRef<FViewUniformShaderParameters> OriginalUniformBuffer = FoundView->UniformBuffer;
 	FMatrix OriginalViewMatrix = FoundView->ViewMatrices.ViewMatrix;
@@ -1365,11 +1364,12 @@ void FProjectedShadowInfo::RenderDepth(FRHICommandList& RHICmdList, FDeferredSha
 	FShadowDepthDrawingPolicy<true>::PolicyShadowInfo = this;
 
 	{
-		//@todo-rco: Very temp code!!!
+	//@todo-rco: Very temp code!!!
 		SCOPE_CYCLE_COUNTER(STAT_RHICounterTEMP);
 		static IConsoleVariable* RHICmdListCVar = IConsoleManager::Get().FindConsoleVariable(TEXT("r.RHICmd"));
 		bool bUseRHICmdList = (RHICmdListCVar->GetInt() >= 2);
-		FRHICommandList& RHICmdList = bUseRHICmdList ? GRHICommandList.CreateList() : FRHICommandList::GetNullRef();
+		FRHICommandList LocalRHICmdList;
+		FRHICommandList& UseRHICmdList = bUseRHICmdList ? LocalRHICmdList : RHICmdList;
 
 		// Draw the subject's static elements using static draw lists
 		if (IsWholeSceneDirectionalShadow())
@@ -1378,12 +1378,12 @@ void FProjectedShadowInfo::RenderDepth(FRHICommandList& RHICmdList, FDeferredSha
 
 			if (bReflectiveShadowmap)
 			{
-				SceneRenderer->Scene->WholeSceneReflectiveShadowMapDrawList.DrawVisible(RHICmdList, *FoundView, StaticMeshWholeSceneShadowDepthMap, StaticMeshWholeSceneShadowBatchVisibility);
+				SceneRenderer->Scene->WholeSceneReflectiveShadowMapDrawList.DrawVisible(UseRHICmdList, *FoundView, StaticMeshWholeSceneShadowDepthMap, StaticMeshWholeSceneShadowBatchVisibility);
 			}
 			else
 			{
 				// Use the scene's shadow depth draw list with this shadow's visibility map
-				SceneRenderer->Scene->WholeSceneShadowDepthDrawList.DrawVisible(RHICmdList, *FoundView, StaticMeshWholeSceneShadowDepthMap, StaticMeshWholeSceneShadowBatchVisibility);
+				SceneRenderer->Scene->WholeSceneShadowDepthDrawList.DrawVisible(UseRHICmdList, *FoundView, StaticMeshWholeSceneShadowDepthMap, StaticMeshWholeSceneShadowBatchVisibility);
 			}
 		}
 		// Draw the subject's static elements using manual state filtering
@@ -1394,16 +1394,15 @@ void FProjectedShadowInfo::RenderDepth(FRHICommandList& RHICmdList, FDeferredSha
 			if(bReflectiveShadowmap && !bOnePassPointLightShadow)
 			{
 				// reflective shadow map
-				DrawShadowMeshElements<true>(RHICmdList, *FoundView, *this);
+				DrawShadowMeshElements<true>(UseRHICmdList, *FoundView, *this);
 			}
 			else
 			{
 				// normal shadow map
-				DrawShadowMeshElements<false>(RHICmdList, *FoundView, *this);
+				DrawShadowMeshElements<false>(UseRHICmdList, *FoundView, *this);
 			}
 		}
 
-		GRHICommandList.ExecuteAndFreeList(RHICmdList);
 	}
 
 	// Draw the subject's dynamic elements.
@@ -1472,7 +1471,7 @@ void StencilingGeometry::DrawCone(FRHICommandList& RHICmdList)
 static FGlobalBoundShaderState MaskBoundShaderState[2];
 
 template <uint32 Quality>
-static void SetShadowProjectionShaderTemplNew(FRHICommandList& RHICmdList, int32 ViewIndex, const FViewInfo& View, const FProjectedShadowInfo* ShadowInfo)
+static void SetShadowProjectionShaderTemplNew(FRHICommandListImmediate& RHICmdList, int32 ViewIndex, const FViewInfo& View, const FProjectedShadowInfo* ShadowInfo)
 {
 	if (ShadowInfo->bTranslucentShadow)
 	{
@@ -1484,7 +1483,7 @@ static void SetShadowProjectionShaderTemplNew(FRHICommandList& RHICmdList, int32
 
 		// Bind shader
 		static FGlobalBoundShaderState BoundShaderState;
-		RHICmdList.CheckIsNull(); // need new approach for "static FGlobalBoundShaderState" for parallel rendering
+		
 		SetGlobalBoundShaderState(RHICmdList, BoundShaderState, GetVertexDeclarationFVector4(), ShadowProjVS, ShadowProjPS);
 
 		// Set shader parameters
@@ -1503,7 +1502,7 @@ static void SetShadowProjectionShaderTemplNew(FRHICommandList& RHICmdList, int32
 			FShadowProjectionPixelShaderInterface* ShadowProjPS = GetGlobalShaderMap()->GetShader<TShadowProjectionPS<Quality, true> >();
 
 			static FGlobalBoundShaderState BoundShaderState;
-			RHICmdList.CheckIsNull(); // need new approach for "static FGlobalBoundShaderState" for parallel rendering
+			
 			SetGlobalBoundShaderState(RHICmdList, BoundShaderState, GetVertexDeclarationFVector4(), ShadowProjVS, ShadowProjPS);
 
 			ShadowProjPS->SetParameters(RHICmdList, ViewIndex, View, ShadowInfo);
@@ -1514,7 +1513,7 @@ static void SetShadowProjectionShaderTemplNew(FRHICommandList& RHICmdList, int32
 			FShadowProjectionPixelShaderInterface* ShadowProjPS = GetGlobalShaderMap()->GetShader<TShadowProjectionPS<Quality, false> >();
 
 			static FGlobalBoundShaderState BoundShaderState;
-			RHICmdList.CheckIsNull(); // need new approach for "static FGlobalBoundShaderState" for parallel rendering
+			
 			SetGlobalBoundShaderState(RHICmdList, BoundShaderState, GetVertexDeclarationFVector4(), ShadowProjVS, ShadowProjPS);
 
 			ShadowProjPS->SetParameters(RHICmdList, ViewIndex, View, ShadowInfo);
@@ -1532,7 +1531,7 @@ static void SetShadowProjectionShaderTemplNew(FRHICommandList& RHICmdList, int32
 		FShadowProjectionPixelShaderInterface* ShadowProjPS = GetGlobalShaderMap()->GetShader<TShadowProjectionPS<Quality, false> >();
 
 		static FGlobalBoundShaderState BoundShaderState;
-		RHICmdList.CheckIsNull(); // need new approach for "static FGlobalBoundShaderState" for parallel rendering
+		
 		SetGlobalBoundShaderState(RHICmdList, BoundShaderState, GetVertexDeclarationFVector4(), ShadowProjVS, ShadowProjPS);
 
 		ShadowProjVS->SetParameters(RHICmdList, View, ShadowInfo);
@@ -1540,7 +1539,7 @@ static void SetShadowProjectionShaderTemplNew(FRHICommandList& RHICmdList, int32
 	}
 }
 
-void FProjectedShadowInfo::RenderProjection(FRHICommandList& RHICmdList, int32 ViewIndex, const FViewInfo* View) const
+void FProjectedShadowInfo::RenderProjection(FRHICommandListImmediate& RHICmdList, int32 ViewIndex, const FViewInfo* View) const
 {
 #if WANTS_DRAW_MESH_EVENTS
 	FString EventName;
@@ -1949,13 +1948,13 @@ void FProjectedShadowInfo::RenderProjection(FRHICommandList& RHICmdList, int32 V
 
 
 template <uint32 Quality>
-static void SetPointLightShaderTempl(FRHICommandList& RHICmdList, int32 ViewIndex, const FViewInfo& View, const FProjectedShadowInfo* ShadowInfo)
+static void SetPointLightShaderTempl(FRHICommandListImmediate& RHICmdList, int32 ViewIndex, const FViewInfo& View, const FProjectedShadowInfo* ShadowInfo)
 {
 	TShaderMapRef<FShadowProjectionVS> VertexShader(GetGlobalShaderMap());
 	TShaderMapRef<TOnePassPointShadowProjectionPS<Quality> > PixelShader(GetGlobalShaderMap());
 
 	static FGlobalBoundShaderState BoundShaderState;
-	RHICmdList.CheckIsNull(); // need new approach for "static FGlobalBoundShaderState" for parallel rendering
+	
 	SetGlobalBoundShaderState(RHICmdList, BoundShaderState, GetVertexDeclarationFVector4(), *VertexShader, *PixelShader);
 
 	VertexShader->SetParameters(RHICmdList, View, ShadowInfo);
@@ -1963,7 +1962,7 @@ static void SetPointLightShaderTempl(FRHICommandList& RHICmdList, int32 ViewInde
 }
 
 /** Render one pass point light shadow projections. */
-void FProjectedShadowInfo::RenderOnePassPointLightProjection(FRHICommandList& RHICmdList, int32 ViewIndex, const FViewInfo& View) const
+void FProjectedShadowInfo::RenderOnePassPointLightProjection(FRHICommandListImmediate& RHICmdList, int32 ViewIndex, const FViewInfo& View) const
 {
 	SCOPE_CYCLE_COUNTER(STAT_RenderWholeSceneShadowProjectionsTime);
 
@@ -2341,7 +2340,7 @@ bool FDeferredShadingSceneRenderer::CheckForProjectedShadows( const FLightSceneI
 }
 
 /** Renders one pass point light shadows. */
-bool FDeferredShadingSceneRenderer::RenderOnePassPointLightShadows(FRHICommandList& RHICmdList, const FLightSceneInfo* LightSceneInfo, bool bRenderedTranslucentObjectShadows, bool& bInjectedTranslucentVolume)
+bool FDeferredShadingSceneRenderer::RenderOnePassPointLightShadows(FRHICommandListImmediate& RHICmdList, const FLightSceneInfo* LightSceneInfo, bool bRenderedTranslucentObjectShadows, bool& bInjectedTranslucentVolume)
 {
 	bool bDirty = false;
 	FVisibleLightInfo& VisibleLightInfo = VisibleLightInfos[LightSceneInfo->Id];
@@ -2366,16 +2365,16 @@ bool FDeferredShadingSceneRenderer::RenderOnePassPointLightShadows(FRHICommandLi
 
 			{
 				SCOPED_DRAW_EVENT(ShadowDepthsFromOpaque, DEC_SCENE_ITEMS);
-				GSceneRenderTargets.BeginRenderingCubeShadowDepth(ProjectedShadowInfo->ResolutionX);
+				GSceneRenderTargets.BeginRenderingCubeShadowDepth(RHICmdList, ProjectedShadowInfo->ResolutionX);
 				ProjectedShadowInfo->bAllocated = true;
 				//gilmerge, if clear depth doesn't go into a cmdlist, then how can renderdepth? This occurs in 4 spots
-				ProjectedShadowInfo->ClearDepth(this);
+				ProjectedShadowInfo->ClearDepth(RHICmdList, this);
 				ProjectedShadowInfo->RenderDepth(RHICmdList, this);
-				GSceneRenderTargets.FinishRenderingCubeShadowDepth(ProjectedShadowInfo->ResolutionX);
+				GSceneRenderTargets.FinishRenderingCubeShadowDepth(RHICmdList, ProjectedShadowInfo->ResolutionX);
 			}
 
 			{
-				GSceneRenderTargets.BeginRenderingLightAttenuation();
+				GSceneRenderTargets.BeginRenderingLightAttenuation(RHICmdList);
 
 				SCOPED_DRAW_EVENT(ShadowProjectionOnOpaque, DEC_SCENE_ITEMS);
 
@@ -2386,7 +2385,7 @@ bool FDeferredShadingSceneRenderer::RenderOnePassPointLightShadows(FRHICommandLi
 					const FViewInfo& View = Views[ViewIndex];
 
 					// Set the device viewport for the view.
-					RHISetViewport(View.ViewRect.Min.X, View.ViewRect.Min.Y, 0.0f, View.ViewRect.Max.X, View.ViewRect.Max.Y, 1.0f);
+					RHICmdList.SetViewport(View.ViewRect.Min.X, View.ViewRect.Min.Y, 0.0f, View.ViewRect.Max.X, View.ViewRect.Max.Y, 1.0f);
 
 					ProjectedShadowInfo->RenderOnePassPointLightProjection(RHICmdList, ViewIndex, View);
 				}
@@ -2411,15 +2410,13 @@ bool FDeferredShadingSceneRenderer::RenderOnePassPointLightShadows(FRHICommandLi
 
 /** Renders the projections of the given Shadows to the appropriate color render target. */
 void FDeferredShadingSceneRenderer::RenderProjections(
-	FRHICommandList& RHICmdList, 
+	FRHICommandListImmediate& RHICmdList,
 	const FLightSceneInfo* LightSceneInfo, 
 	const TArray<FProjectedShadowInfo*,SceneRenderingAllocator>& Shadows
 	)
 {
-	RHICmdList.CheckIsNull();
-
 	// Normal shadows render to light attenuation
-	GSceneRenderTargets.BeginRenderingLightAttenuation();
+	GSceneRenderTargets.BeginRenderingLightAttenuation(RHICmdList);
 
 	for (int32 ViewIndex = 0; ViewIndex < Views.Num(); ViewIndex++)
 	{
@@ -2428,10 +2425,10 @@ void FDeferredShadingSceneRenderer::RenderProjections(
 		const FViewInfo& View = Views[ViewIndex];
 
 		// Set the device viewport for the view.
-		RHISetViewport(View.ViewRect.Min.X, View.ViewRect.Min.Y, 0.0f, View.ViewRect.Max.X, View.ViewRect.Max.Y, 1.0f);
+		RHICmdList.SetViewport(View.ViewRect.Min.X, View.ViewRect.Min.Y, 0.0f, View.ViewRect.Max.X, View.ViewRect.Max.Y, 1.0f);
 
 		// Set the light's scissor rectangle.
-		LightSceneInfo->Proxy->SetScissorRect(View);
+		LightSceneInfo->Proxy->SetScissorRect(RHICmdList, View);
 
 		// Project the shadow depth buffers onto the scene.
 		for (int32 ShadowIndex = 0; ShadowIndex < Shadows.Num(); ShadowIndex++)
@@ -2448,7 +2445,7 @@ void FDeferredShadingSceneRenderer::RenderProjections(
 		}
 
 		// Reset the scissor rectangle.
-		RHISetScissorRect(false, 0, 0, 0, 0);
+		RHICmdList.SetScissorRect(false, 0, 0, 0, 0);
 	}
 }
 
@@ -2497,7 +2494,7 @@ struct FCompareFProjectedShadowInfoBySplitIndex
 	}
 };
 
-bool FDeferredShadingSceneRenderer::RenderTranslucentProjectedShadows(FRHICommandList& RHICmdList, const FLightSceneInfo* LightSceneInfo )
+bool FDeferredShadingSceneRenderer::RenderTranslucentProjectedShadows(FRHICommandListImmediate& RHICmdList, const FLightSceneInfo* LightSceneInfo)
 {
 	SCOPE_CYCLE_COUNTER(STAT_ProjectedShadowDrawTime);
 
@@ -2608,7 +2605,7 @@ bool FDeferredShadingSceneRenderer::RenderTranslucentProjectedShadows(FRHIComman
 
 				if (ProjectedShadowInfo->bAllocated)
 				{
-					ProjectedShadowInfo->RenderTranslucencyDepths(this);
+					ProjectedShadowInfo->RenderTranslucencyDepths(RHICmdList, this);
 				}
 			}
 		}
@@ -2656,7 +2653,7 @@ bool FDeferredShadingSceneRenderer::RenderTranslucentProjectedShadows(FRHIComman
 * @param LightSceneInfo Represents the current light
 * @return true if anything got rendered
 */
-bool FDeferredShadingSceneRenderer::RenderReflectiveShadowMaps(FRHICommandList& RHICmdList, const FLightSceneInfo* LightSceneInfo )
+bool FDeferredShadingSceneRenderer::RenderReflectiveShadowMaps(FRHICommandListImmediate& RHICmdList, const FLightSceneInfo* LightSceneInfo)
 {
 	SCOPE_CYCLE_COUNTER(STAT_ReflectiveShadowMapDrawTime);
 
@@ -2766,17 +2763,17 @@ bool FDeferredShadingSceneRenderer::RenderReflectiveShadowMaps(FRHICommandList&
 				FLightPropagationVolume* LightPropagationVolume = ViewState->GetLightPropagationVolume();
 
 				check(LightPropagationVolume);
-				GSceneRenderTargets.BeginRenderingReflectiveShadowMap(LightPropagationVolume);	
+				GSceneRenderTargets.BeginRenderingReflectiveShadowMap(RHICmdList, LightPropagationVolume);
 
 				LightPropagationVolume->SetVplInjectionConstants( *ProjectedShadowInfo, LightSceneInfo->Proxy );
 
 				if (ProjectedShadowInfo->bAllocated && !ProjectedShadowInfo->bTranslucentShadow)
 				{
-					ProjectedShadowInfo->ClearDepth(this);
+					ProjectedShadowInfo->ClearDepth(RHICmdList, this);
 					ProjectedShadowInfo->RenderDepth(RHICmdList, this);
 				}
 
-				GSceneRenderTargets.FinishRenderingReflectiveShadowMap();
+				GSceneRenderTargets.FinishRenderingReflectiveShadowMap(RHICmdList);
 			}
 		}
 
@@ -2828,7 +2825,7 @@ bool FDeferredShadingSceneRenderer::RenderReflectiveShadowMaps(FRHICommandList&
  * @param LightSceneInfo Represents the current light
  * @return true if anything got rendered
  */
-bool FDeferredShadingSceneRenderer::RenderProjectedShadows(FRHICommandList& RHICmdList, const FLightSceneInfo* LightSceneInfo, bool bRenderedTranslucentObjectShadows, bool& bInjectedTranslucentVolume )
+bool FDeferredShadingSceneRenderer::RenderProjectedShadows(FRHICommandListImmediate& RHICmdList, const FLightSceneInfo* LightSceneInfo, bool bRenderedTranslucentObjectShadows, bool& bInjectedTranslucentVolume)
 {
 	SCOPE_CYCLE_COUNTER(STAT_ProjectedShadowDrawTime);
 
@@ -2936,9 +2933,9 @@ bool FDeferredShadingSceneRenderer::RenderProjectedShadows(FRHICommandList& RHIC
 			// Render the shadow depths.
 			SCOPED_DRAW_EVENT(ShadowDepthsFromOpaque, DEC_SCENE_ITEMS);
 
-			GSceneRenderTargets.BeginRenderingShadowDepth();	
+			GSceneRenderTargets.BeginRenderingShadowDepth(RHICmdList);
 
-			RHIClear(false,FColor(255,255,255),true,1.0f,false,0, FIntRect());
+			RHICmdList.Clear(false, FColor(255, 255, 255), true, 1.0f, false, 0, FIntRect());
 
 			// render depth for each shadow
 			for (int32 ShadowIndex = 0; ShadowIndex < Shadows.Num(); ShadowIndex++)
@@ -2950,7 +2947,7 @@ bool FDeferredShadingSceneRenderer::RenderProjectedShadows(FRHICommandList& RHIC
 				}
 			}
 
-			GSceneRenderTargets.FinishRenderingShadowDepth();
+			GSceneRenderTargets.FinishRenderingShadowDepth(RHICmdList);
 		}
 
 		// Render the shadow projections.
@@ -3006,10 +3003,8 @@ bool FDeferredShadingSceneRenderer::RenderProjectedShadows(FRHICommandList& RHIC
 }
 
 /** Renders preshadow depths for any preshadows whose depths aren't cached yet, and renders the projections of preshadows with opaque relevance. */
-bool FDeferredShadingSceneRenderer::RenderCachedPreshadows(FRHICommandList& RHICmdList, const FLightSceneInfo* LightSceneInfo)
+bool FDeferredShadingSceneRenderer::RenderCachedPreshadows(FRHICommandListImmediate& RHICmdList, const FLightSceneInfo* LightSceneInfo)
 {
-	RHICmdList.CheckIsNull();
-
 	bool bAttenuationBufferDirty = false;
 
 	FVisibleLightInfo& VisibleLightInfo = VisibleLightInfos[LightSceneInfo->Id];
@@ -3079,7 +3074,7 @@ bool FDeferredShadingSceneRenderer::RenderCachedPreshadows(FRHICommandList& RHIC
 				if (!ProjectedShadowInfo->bDepthsCached)
 				{
 					check(ProjectedShadowInfo->bAllocated);
-					ProjectedShadowInfo->ClearDepth(this);
+					ProjectedShadowInfo->ClearDepth(RHICmdList, this);
 					ProjectedShadowInfo->RenderDepth(RHICmdList, this);
 					ProjectedShadowInfo->bDepthsCached = true;
 
@@ -3089,7 +3084,7 @@ bool FDeferredShadingSceneRenderer::RenderCachedPreshadows(FRHICommandList& RHIC
 						ProjectedShadowInfo->X + ProjectedShadowInfo->ResolutionX + SHADOW_BORDER * 2,
 						ProjectedShadowInfo->Y + ProjectedShadowInfo->ResolutionY + SHADOW_BORDER * 2));
 
-					RHICopyToResolveTarget(
+					RHICmdList.CopyToResolveTarget(
 						GSceneRenderTargets.PreShadowCacheDepthZ->GetRenderTargetItem().TargetableTexture, 
 						GSceneRenderTargets.PreShadowCacheDepthZ->GetRenderTargetItem().ShaderResourceTexture,
 						false, 
diff --git a/Engine/Source/Runtime/Renderer/Private/ShadowRendering.h b/Engine/Source/Runtime/Renderer/Private/ShadowRendering.h
--- a/Engine/Source/Runtime/Renderer/Private/ShadowRendering.h
+++ b/Engine/Source/Runtime/Renderer/Private/ShadowRendering.h
@@ -662,20 +662,20 @@ public:
 	/**
 	 * Renders the shadow subject depth.
 	 */
-	void RenderDepth(FRHICommandList& RHICmdList, class FDeferredShadingSceneRenderer* SceneRenderer);
+	void RenderDepth(FRHICommandListImmediate& RHICmdList, class FDeferredShadingSceneRenderer* SceneRenderer);
 
-	void ClearDepth(class FDeferredShadingSceneRenderer* SceneRenderer);
+	void ClearDepth(FRHICommandListImmediate& RHICmdList, class FDeferredShadingSceneRenderer* SceneRenderer);
 
 	/** Renders shadow maps for translucent primitives. */
-	void RenderTranslucencyDepths(class FDeferredShadingSceneRenderer* SceneRenderer);
+	void RenderTranslucencyDepths(FRHICommandListImmediate& RHICmdList, class FDeferredShadingSceneRenderer* SceneRenderer);
 
 	/**
 	 * Projects the shadow onto the scene for a particular view.
 	 */
-	void RenderProjection(FRHICommandList& RHICmdList, int32 ViewIndex, const class FViewInfo* View) const;
+	void RenderProjection(FRHICommandListImmediate& RHICmdList, int32 ViewIndex, const class FViewInfo* View) const;
 
 	/** Render one pass point light shadow projections. */
-	void RenderOnePassPointLightProjection(FRHICommandList& RHICmdList, int32 ViewIndex, const FViewInfo& View) const;
+	void RenderOnePassPointLightProjection(FRHICommandListImmediate& RHICmdList, int32 ViewIndex, const FViewInfo& View) const;
 
 	/**
 	 * Renders the projected shadow's frustum wireframe with the given FPrimitiveDrawInterface.
@@ -1289,7 +1289,7 @@ public:
 
 		if (ShadowDepthCubeComparisonSampler.IsBound())
 		{
-			RHISetShaderSampler(
+			RHICmdList.SetShaderSampler(
 				ShaderRHI, 
 				ShadowDepthCubeComparisonSampler.GetBaseIndex(), 
 				// Use a comparison sampler to do hardware PCF
diff --git a/Engine/Source/Runtime/Renderer/Private/ShadowSetup.cpp b/Engine/Source/Runtime/Renderer/Private/ShadowSetup.cpp
--- a/Engine/Source/Runtime/Renderer/Private/ShadowSetup.cpp
+++ b/Engine/Source/Runtime/Renderer/Private/ShadowSetup.cpp
@@ -1024,6 +1024,7 @@ bool FDeferredShadingSceneRenderer::ShouldCreateObjectShadowForStationaryLight(c
 }
 
 void FDeferredShadingSceneRenderer::SetupInteractionShadows(
+	FRHICommandListImmediate& RHICmdList,
 	FLightPrimitiveInteraction* Interaction, 
 	FVisibleLightInfo& VisibleLightInfo, 
 	bool bStaticSceneOnly,
@@ -1055,12 +1056,13 @@ void FDeferredShadingSceneRenderer::SetupInteractionShadows(
 			&& (bCreateTranslucentObjectShadow || bCreateInsetObjectShadow || bCreateObjectShadowForStationaryLight))
 		{
 			// Create projected shadow infos
-			CreatePerObjectProjectedShadow(Interaction, bCreateTranslucentObjectShadow, bCreateInsetObjectShadow || bCreateObjectShadowForStationaryLight, ViewDependentWholeSceneShadows, PreShadows);
+			CreatePerObjectProjectedShadow(RHICmdList, Interaction, bCreateTranslucentObjectShadow, bCreateInsetObjectShadow || bCreateObjectShadowForStationaryLight, ViewDependentWholeSceneShadows, PreShadows);
 		}
 	}
 }
 
 void FDeferredShadingSceneRenderer::CreatePerObjectProjectedShadow(
+	FRHICommandListImmediate& RHICmdList,
 	FLightPrimitiveInteraction* Interaction, 
 	bool bCreateTranslucentObjectShadow, 
 	bool bCreateOpaqueObjectShadow,
@@ -1103,14 +1105,14 @@ void FDeferredShadingSceneRenderer::CreatePerObjectProjectedShadow(
 			!bCreateOpaqueObjectShadow ||
 			(
 				!View.bIgnoreExistingQueries &&	View.State &&
-				((FSceneViewState*)View.State)->IsShadowOccluded(PrimitiveSceneInfo->PrimitiveComponentId, LightSceneInfo->Proxy->GetLightComponent(), INDEX_NONE, false)
+				((FSceneViewState*)View.State)->IsShadowOccluded(RHICmdList, PrimitiveSceneInfo->PrimitiveComponentId, LightSceneInfo->Proxy->GetLightComponent(), INDEX_NONE, false)
 			);
 
 		const bool bTranslucentShadowIsOccluded = 
 			!bCreateTranslucentObjectShadow ||
 			(
 				!View.bIgnoreExistingQueries && View.State &&
-				((FSceneViewState*)View.State)->IsShadowOccluded(PrimitiveSceneInfo->PrimitiveComponentId, LightSceneInfo->Proxy->GetLightComponent(), INDEX_NONE, true)
+				((FSceneViewState*)View.State)->IsShadowOccluded(RHICmdList, PrimitiveSceneInfo->PrimitiveComponentId, LightSceneInfo->Proxy->GetLightComponent(), INDEX_NONE, true)
 			);
 
 		const bool bSubjectIsVisibleInThisView = View.PrimitiveVisibilityMap[PrimitiveSceneInfo->GetIndex()];
@@ -1549,7 +1551,7 @@ void FDeferredShadingSceneRenderer::CreateWholeSceneProjectedShadow(FLightSceneI
 	}
 }
 
-void FDeferredShadingSceneRenderer::InitProjectedShadowVisibility()
+void FDeferredShadingSceneRenderer::InitProjectedShadowVisibility(FRHICommandListImmediate& RHICmdList)
 {
 	// Initialize the views' ProjectedShadowVisibilityMaps and remove shadows without subjects.
 	for(TSparseArray<FLightSceneInfoCompact>::TConstIterator LightIt(Scene->Lights);LightIt;++LightIt)
@@ -1614,6 +1616,7 @@ void FDeferredShadingSceneRenderer::InitProjectedShadowVisibility()
 						!View.bIgnoreExistingQueries &&
 						View.State &&
 						((FSceneViewState*)View.State)->IsShadowOccluded(
+							RHICmdList,
 							ProjectedShadowInfo.ParentSceneInfo ? 
 								ProjectedShadowInfo.ParentSceneInfo->PrimitiveComponentId :
 								FPrimitiveComponentId(),
@@ -1892,7 +1895,7 @@ void FDeferredShadingSceneRenderer::GatherShadowPrimitives(
 	}
 }
 
-void FDeferredShadingSceneRenderer::InitDynamicShadows()
+void FDeferredShadingSceneRenderer::InitDynamicShadows(FRHICommandListImmediate& RHICmdList)
 {
 	SCOPE_CYCLE_COUNTER(STAT_DynamicShadowSetupTime);
 
@@ -2053,7 +2056,7 @@ void FDeferredShadingSceneRenderer::InitDynamicShadows()
 							Interaction = Interaction->GetNextPrimitive()
 							)
 						{
-							SetupInteractionShadows(Interaction, VisibleLightInfo, bStaticSceneOnly, ViewDependentWholeSceneShadows, PreShadows);
+							SetupInteractionShadows(RHICmdList, Interaction, VisibleLightInfo, bStaticSceneOnly, ViewDependentWholeSceneShadows, PreShadows);
 						}
 					}
 				}
@@ -2061,7 +2064,7 @@ void FDeferredShadingSceneRenderer::InitDynamicShadows()
 		}
 
 		// Calculate visibility of the projected shadows.
-		InitProjectedShadowVisibility();
+		InitProjectedShadowVisibility(RHICmdList);
 	}
 
 	// Clear old preshadows and attempt to add new ones to the cache
diff --git a/Engine/Source/Runtime/Renderer/Private/SystemTextures.cpp b/Engine/Source/Runtime/Renderer/Private/SystemTextures.cpp
--- a/Engine/Source/Runtime/Renderer/Private/SystemTextures.cpp
+++ b/Engine/Source/Runtime/Renderer/Private/SystemTextures.cpp
@@ -16,7 +16,7 @@ TGlobalResource<FSystemTextures> GSystemTextures;
 
 uint8 Quantize8SignedByte(float x);
 
-void FSystemTextures::InitializeTextures(ERHIFeatureLevel::Type InFeatureLevel)
+void FSystemTextures::InitializeTextures(FRHICommandListImmediate& RHICmdList, ERHIFeatureLevel::Type InFeatureLevel)
 {
 	if (bTexturesInitialized && FeatureLevelInitializedTo >= InFeatureLevel)
 	{
@@ -26,7 +26,7 @@ void FSystemTextures::InitializeTextures(ERHIFeatureLevel::Type InFeatureLevel)
 
 	// start with a defined state for the scissor rect (D3D11 was returning (0,0,0,0) which caused a clear to not execute correctly)
 	// todo: move this to an earlier place (for dx9 is has to be after device creation which is after window creation)
-	RHISetScissorRect(false, 0, 0, 0, 0);
+	RHICmdList.SetScissorRect(false, 0, 0, 0, 0);
 
 	// First initialize textures that are common to all feature levels. This is always done the first time we come into this function, as doesn't care about the
 	// requested feature level
@@ -37,9 +37,9 @@ void FSystemTextures::InitializeTextures(ERHIFeatureLevel::Type InFeatureLevel)
 			FPooledRenderTargetDesc Desc(FPooledRenderTargetDesc::Create2DDesc(FIntPoint(1, 1), PF_B8G8R8A8, TexCreate_HideInVisualizeTexture, TexCreate_RenderTargetable, false));
 			GRenderTargetPool.FindFreeElement(Desc, WhiteDummy, TEXT("WhiteDummy"));
 
-			RHISetRenderTarget(WhiteDummy->GetRenderTargetItem().TargetableTexture, FTextureRHIRef());
-			RHIClear(true, FLinearColor(1, 1, 1, 1), false, 0, false, 0, FIntRect());
-			RHICopyToResolveTarget(WhiteDummy->GetRenderTargetItem().TargetableTexture, WhiteDummy->GetRenderTargetItem().ShaderResourceTexture, true, FResolveParams());
+			SetRenderTarget(RHICmdList, WhiteDummy->GetRenderTargetItem().TargetableTexture, FTextureRHIRef());
+			RHICmdList.Clear(true, FLinearColor(1, 1, 1, 1), false, 0, false, 0, FIntRect());
+			RHICmdList.CopyToResolveTarget(WhiteDummy->GetRenderTargetItem().TargetableTexture, WhiteDummy->GetRenderTargetItem().ShaderResourceTexture, true, FResolveParams());
 		}
 
 		// Create a dummy black texture.
@@ -47,9 +47,9 @@ void FSystemTextures::InitializeTextures(ERHIFeatureLevel::Type InFeatureLevel)
 			FPooledRenderTargetDesc Desc(FPooledRenderTargetDesc::Create2DDesc(FIntPoint(1, 1), PF_B8G8R8A8, TexCreate_HideInVisualizeTexture, TexCreate_RenderTargetable, false));
 			GRenderTargetPool.FindFreeElement(Desc, BlackDummy, TEXT("BlackDummy"));
 
-			RHISetRenderTarget(BlackDummy->GetRenderTargetItem().TargetableTexture, FTextureRHIRef());
-			RHIClear(true, FLinearColor(0, 0, 0, 0), false, 0, false, 0, FIntRect());
-			RHICopyToResolveTarget(BlackDummy->GetRenderTargetItem().TargetableTexture, BlackDummy->GetRenderTargetItem().ShaderResourceTexture, true, FResolveParams());
+			SetRenderTarget(RHICmdList, BlackDummy->GetRenderTargetItem().TargetableTexture, FTextureRHIRef());
+			RHICmdList.Clear(true, FLinearColor(0, 0, 0, 0), false, 0, false, 0, FIntRect());
+			RHICmdList.CopyToResolveTarget(BlackDummy->GetRenderTargetItem().TargetableTexture, BlackDummy->GetRenderTargetItem().ShaderResourceTexture, true, FResolveParams());
 		}
 
 		// Create the PerlinNoiseGradient texture
@@ -58,7 +58,7 @@ void FSystemTextures::InitializeTextures(ERHIFeatureLevel::Type InFeatureLevel)
 			GRenderTargetPool.FindFreeElement(Desc, PerlinNoiseGradient, TEXT("PerlinNoiseGradient"));
 			// Write the contents of the texture.
 			uint32 DestStride;
-			uint8* DestBuffer = (uint8*)RHILockTexture2D((FTexture2DRHIRef&)PerlinNoiseGradient->GetRenderTargetItem().ShaderResourceTexture, 0, RLM_WriteOnly, DestStride, false);
+			uint8* DestBuffer = (uint8*)RHICmdList.LockTexture2D((FTexture2DRHIRef&)PerlinNoiseGradient->GetRenderTargetItem().ShaderResourceTexture, 0, RLM_WriteOnly, DestStride, false);
 			// seed the pseudo random stream with a good value
 			FRandomStream RandomStream(12345);
 			// Values represent float3 values in the -1..1 range.
@@ -80,7 +80,7 @@ void FSystemTextures::InitializeTextures(ERHIFeatureLevel::Type InFeatureLevel)
 					*Dest = gradtable[(uint32)(RandomStream.GetFraction() * 11.9999999f)];
 				}
 			}
-			RHIUnlockTexture2D((FTexture2DRHIRef&)PerlinNoiseGradient->GetRenderTargetItem().ShaderResourceTexture, 0, false);
+			RHICmdList.UnlockTexture2D((FTexture2DRHIRef&)PerlinNoiseGradient->GetRenderTargetItem().ShaderResourceTexture, 0, false);
 		}
 
 		if (!GSupportsShaderFramebufferFetch && GPixelFormats[PF_FloatRGBA].Supported)
@@ -88,9 +88,9 @@ void FSystemTextures::InitializeTextures(ERHIFeatureLevel::Type InFeatureLevel)
 			FPooledRenderTargetDesc Desc(FPooledRenderTargetDesc::Create2DDesc(FIntPoint(1, 1), PF_FloatRGBA, TexCreate_HideInVisualizeTexture, TexCreate_RenderTargetable, false));
 			GRenderTargetPool.FindFreeElement(Desc, MaxFP16Depth, TEXT("MaxFP16Depth"));
 
-			RHISetRenderTarget(MaxFP16Depth->GetRenderTargetItem().TargetableTexture, FTextureRHIRef());
-			RHIClear(true, FLinearColor(65000.0f, 65000.0f, 65000.0f, 65000.0f), false, 0, false, 0, FIntRect());
-			RHICopyToResolveTarget(MaxFP16Depth->GetRenderTargetItem().TargetableTexture, MaxFP16Depth->GetRenderTargetItem().ShaderResourceTexture, true, FResolveParams());
+			SetRenderTarget(RHICmdList, MaxFP16Depth->GetRenderTargetItem().TargetableTexture, FTextureRHIRef());
+			RHICmdList.Clear(true, FLinearColor(65000.0f, 65000.0f, 65000.0f, 65000.0f), false, 0, false, 0, FIntRect());
+			RHICmdList.CopyToResolveTarget(MaxFP16Depth->GetRenderTargetItem().TargetableTexture, MaxFP16Depth->GetRenderTargetItem().ShaderResourceTexture, true, FResolveParams());
 		}
 	}
 
@@ -196,7 +196,7 @@ void FSystemTextures::InitializeTextures(ERHIFeatureLevel::Type InFeatureLevel)
 
 		FUpdateTextureRegion3D Region(0, 0, 0, 0, 0, 0, Desc.Extent.X, Desc.Extent.Y, Desc.Depth);
 
-		RHIUpdateTexture3D(
+		RHICmdList.UpdateTexture3D(
 			(FTexture3DRHIRef&)PerlinNoise3D->GetRenderTargetItem().ShaderResourceTexture,
 			0,
 			Region,
@@ -236,7 +236,7 @@ void FSystemTextures::InitializeTextures(ERHIFeatureLevel::Type InFeatureLevel)
 		GRenderTargetPool.FindFreeElement(Desc, SSAORandomization, TEXT("SSAORandomization"));
 		// Write the contents of the texture.
 		uint32 DestStride;
-		uint8* DestBuffer = (uint8*)RHILockTexture2D((FTexture2DRHIRef&)SSAORandomization->GetRenderTargetItem().ShaderResourceTexture,0,RLM_WriteOnly,DestStride,false);
+		uint8* DestBuffer = (uint8*)RHICmdList.LockTexture2D((FTexture2DRHIRef&)SSAORandomization->GetRenderTargetItem().ShaderResourceTexture, 0, RLM_WriteOnly, DestStride, false);
 
 		for(int32 y = 0; y < Desc.Extent.Y; ++y)
 		{
@@ -249,7 +249,7 @@ void FSystemTextures::InitializeTextures(ERHIFeatureLevel::Type InFeatureLevel)
 				*Dest = Bases[Index];
 			}
 		}
-		RHIUnlockTexture2D((FTexture2DRHIRef&)SSAORandomization->GetRenderTargetItem().ShaderResourceTexture, 0, false);
+		RHICmdList.UnlockTexture2D((FTexture2DRHIRef&)SSAORandomization->GetRenderTargetItem().ShaderResourceTexture, 0, false);
 
 		{
 			// for testing, with 128x128 R8G8 we are very close to the reference (if lower res is needed we might have to add an offset to counter the 0.5f texel shift)
@@ -267,7 +267,7 @@ void FSystemTextures::InitializeTextures(ERHIFeatureLevel::Type InFeatureLevel)
 			GRenderTargetPool.FindFreeElement(Desc, PreintegratedGF, TEXT("PreintegratedGF"));
 			// Write the contents of the texture.
 			uint32 DestStride;
-			uint8* DestBuffer = (uint8*)RHILockTexture2D((FTexture2DRHIRef&)PreintegratedGF->GetRenderTargetItem().ShaderResourceTexture, 0, RLM_WriteOnly, DestStride, false);
+			uint8* DestBuffer = (uint8*)RHICmdList.LockTexture2D((FTexture2DRHIRef&)PreintegratedGF->GetRenderTargetItem().ShaderResourceTexture, 0, RLM_WriteOnly, DestStride, false);
 
 			// x is NoV, y is roughness
 			for (int32 y = 0; y < Desc.Extent.Y; y++)
@@ -343,7 +343,7 @@ void FSystemTextures::InitializeTextures(ERHIFeatureLevel::Type InFeatureLevel)
 					}
 				}
 			}
-			RHIUnlockTexture2D((FTexture2DRHIRef&)PreintegratedGF->GetRenderTargetItem().ShaderResourceTexture, 0, false);
+			RHICmdList.UnlockTexture2D((FTexture2DRHIRef&)PreintegratedGF->GetRenderTargetItem().ShaderResourceTexture, 0, false);
 		}
 	}
 
diff --git a/Engine/Source/Runtime/Renderer/Private/SystemTextures.h b/Engine/Source/Runtime/Renderer/Private/SystemTextures.h
--- a/Engine/Source/Runtime/Renderer/Private/SystemTextures.h
+++ b/Engine/Source/Runtime/Renderer/Private/SystemTextures.h
@@ -24,7 +24,7 @@ public:
 	/**
 	 * Initialize/allocate textures if not already.
 	 */
-	void InitializeTextures(ERHIFeatureLevel::Type InFeatureLevel);
+	void InitializeTextures(FRHICommandListImmediate& RHICmdList, ERHIFeatureLevel::Type InFeatureLevel);
 
 	// FRenderResource interface.
 	/**
diff --git a/Engine/Source/Runtime/Renderer/Private/TiledDeferredLightRendering.cpp b/Engine/Source/Runtime/Renderer/Private/TiledDeferredLightRendering.cpp
--- a/Engine/Source/Runtime/Renderer/Private/TiledDeferredLightRendering.cpp
+++ b/Engine/Source/Runtime/Renderer/Private/TiledDeferredLightRendering.cpp
@@ -281,7 +281,7 @@ bool FDeferredShadingSceneRenderer::ShouldUseTiledDeferred(int32 NumUnshadowedLi
 
 template <bool bVisualizeLightCulling>
 static void SetShaderTemplTiledLighting(
-	FRHICommandList& RHICmdList, 
+	FRHICommandListImmediate& RHICmdList,
 	const FSceneView& View, 
 	int32 ViewIndex,
 	int32 NumViews,
@@ -294,7 +294,7 @@ static void SetShaderTemplTiledLighting(
 	IPooledRenderTarget& OutTexture)
 {
 	TShaderMapRef<FTiledDeferredLightingCS<bVisualizeLightCulling> > ComputeShader(GetGlobalShaderMap());
-	RHISetComputeShader(ComputeShader->GetComputeShader());
+	RHICmdList.SetComputeShader(ComputeShader->GetComputeShader());
 
 	ComputeShader->SetParameters(RHICmdList, View, ViewIndex, NumViews, SortedLights, NumLightsToRenderInSortedLights, SimpleLights, StartIndex, NumThisPass, InTexture, OutTexture);
 
@@ -306,10 +306,8 @@ static void SetShaderTemplTiledLighting(
 }
 
 
-void FDeferredShadingSceneRenderer::RenderTiledDeferredLighting(FRHICommandList& RHICmdList, const TArray<FSortedLightSceneInfo, SceneRenderingAllocator>& SortedLights, int32 NumUnshadowedLights, const FSimpleLightArray& SimpleLights)
+void FDeferredShadingSceneRenderer::RenderTiledDeferredLighting(FRHICommandListImmediate& RHICmdList, const TArray<FSortedLightSceneInfo, SceneRenderingAllocator>& SortedLights, int32 NumUnshadowedLights, const FSimpleLightArray& SimpleLights)
 {
-	RHICmdList.CheckIsNull();
-
 	check(GUseTiledDeferredShading);
 	check(SortedLights.Num() >= NumUnshadowedLights);
 
@@ -322,7 +320,7 @@ void FDeferredShadingSceneRenderer::RenderTiledDeferredLighting(FRHICommandList&
 		INC_DWORD_STAT_BY(STAT_NumLightsUsingSimpleTiledDeferred, SimpleLights.InstanceData.Num());
 		SCOPE_CYCLE_COUNTER(STAT_DirectLightRenderingTime);
 
-		RHISetRenderTarget(NULL, NULL);
+		SetRenderTarget(RHICmdList, NULL, NULL);
 
 		// Determine how many compute shader passes will be needed to process all the lights
 		const int32 NumPassesNeeded = FMath::DivideAndRoundUp(NumLightsToRender, GMaxNumTiledDeferredLights);
@@ -335,7 +333,7 @@ void FDeferredShadingSceneRenderer::RenderTiledDeferredLighting(FRHICommandList&
 			// One some hardware we can read and write from the same UAV with a 32 bit format. We don't do that yet.
 			TRefCountPtr<IPooledRenderTarget> OutTexture;
 			{
-				GSceneRenderTargets.ResolveSceneColor(FResolveRect(0, 0, ViewFamily.FamilySizeX, ViewFamily.FamilySizeY));
+				GSceneRenderTargets.ResolveSceneColor(RHICmdList, FResolveRect(0, 0, ViewFamily.FamilySizeX, ViewFamily.FamilySizeY));
 
 				FPooledRenderTargetDesc Desc = GSceneRenderTargets.GetSceneColor()->GetDesc();
 				Desc.TargetableFlags |= TexCreate_UAV;
diff --git a/Engine/Source/Runtime/Renderer/Private/TranslucentLighting.cpp b/Engine/Source/Runtime/Renderer/Private/TranslucentLighting.cpp
--- a/Engine/Source/Runtime/Renderer/Private/TranslucentLighting.cpp
+++ b/Engine/Source/Runtime/Renderer/Private/TranslucentLighting.cpp
@@ -348,15 +348,13 @@ public:
 		}
 	}
 
-	void DrawShared(FRHICommandList& RHICmdList, const FSceneView* View,FBoundShaderStateRHIParamRef BoundShaderState) const
+	void DrawShared(FRHICommandList& RHICmdList, const FSceneView* View, FBoundShaderStateRHIParamRef BoundShaderState) const
 	{
-		RHICmdList.CheckIsNull();
-
 		// Set the shared mesh resources.
 		FMeshDrawingPolicy::DrawShared(RHICmdList, View);
 
 		// Set the actual shader & vertex declaration state
-		RHISetBoundShaderState(BoundShaderState);
+		RHICmdList.SetBoundShaderState(BoundShaderState);
 
 		VertexShader->SetParameters(RHICmdList, MaterialRenderProxy, *View, PolicyShadowInfo);
 		PixelShader->SetParameters(RHICmdList, MaterialRenderProxy, *View, PolicyShadowInfo);
@@ -488,11 +486,9 @@ public:
 };
 
 /** Renders shadow maps for translucent primitives. */
-void FProjectedShadowInfo::RenderTranslucencyDepths(FDeferredShadingSceneRenderer* SceneRenderer)
+void FProjectedShadowInfo::RenderTranslucencyDepths(FRHICommandListImmediate& RHICmdList, FDeferredShadingSceneRenderer* SceneRenderer)
 {
-	//@todo-rco: RHIPacketList
-	FRHICommandList& RHICmdList = FRHICommandList::GetNullRef();
-
+	
 	checkSlow(!bWholeSceneShadow);
 	SCOPE_CYCLE_COUNTER(STAT_RenderWholeSceneShadowDepthsTime);
 
@@ -541,10 +537,10 @@ void FProjectedShadowInfo::RenderTranslucencyDepths(FDeferredShadingSceneRendere
 			GSceneRenderTargets.TranslucencyShadowTransmission[0]->GetRenderTargetItem().TargetableTexture,
 			GSceneRenderTargets.TranslucencyShadowTransmission[1]->GetRenderTargetItem().TargetableTexture
 		};
-		RHISetRenderTargets(ARRAY_COUNT(RenderTargets), RenderTargets, FTextureRHIParamRef(), 0, NULL);
+		SetRenderTargets(RHICmdList, ARRAY_COUNT(RenderTargets), RenderTargets, FTextureRHIParamRef(), 0, NULL);
 
 		// Clear the shadow and its border
-		RHISetViewport(
+		RHICmdList.SetViewport(
 			X,
 			Y,
 			0.0f,
@@ -557,7 +553,7 @@ void FProjectedShadowInfo::RenderTranslucencyDepths(FDeferredShadingSceneRendere
 		DrawClearQuadMRT(RHICmdList, true, ARRAY_COUNT(ClearColors), ClearColors, false, 1.0f, false, 0);
 
 		// Set the viewport for the shadow.
-		RHISetViewport(
+		RHICmdList.SetViewport(
 			(X + SHADOW_BORDER),
 			(Y + SHADOW_BORDER),
 			0.0f,
@@ -566,8 +562,8 @@ void FProjectedShadowInfo::RenderTranslucencyDepths(FDeferredShadingSceneRendere
 			1.0f
 			);
 
-		RHISetDepthStencilState(TStaticDepthStencilState<false, CF_Always>::GetRHI());
-		RHISetBlendState(TStaticBlendState<
+		RHICmdList.SetDepthStencilState(TStaticDepthStencilState<false, CF_Always>::GetRHI());
+		RHICmdList.SetBlendState(TStaticBlendState<
 			CW_RGBA, BO_Add, BF_One, BF_One, BO_Add, BF_One, BF_One,
 			CW_RGBA, BO_Add, BF_One, BF_One, BO_Add, BF_One, BF_One>::GetRHI());
 
@@ -1019,7 +1015,7 @@ class FVolumeRasterizeVertexBuffer : public FVertexBuffer
 {
 public:
 
-	virtual void InitRHI()
+	virtual void InitRHI() override
 	{
 		// Used as a non-indexed triangle strip, so 4 vertices per quad
 		const uint32 Size = 4 * sizeof(FScreenVertex);
@@ -1046,7 +1042,7 @@ public:
 TGlobalResource<FVolumeRasterizeVertexBuffer> GVolumeRasterizeVertexBuffer;
 
 /** Draws a quad per volume texture slice to the subregion of the volume texture specified. */
-void RasterizeToVolumeTexture(FRHICommandList& RHICmdList, FVolumeBounds VolumeBounds)
+void RasterizeToVolumeTexture(FRHICommandListImmediate& RHICmdList, FVolumeBounds VolumeBounds)
 {
 	// Setup the viewport to only render to the given bounds
 	RHICmdList.SetViewport(VolumeBounds.MinX, VolumeBounds.MinY, 0, VolumeBounds.MaxX, VolumeBounds.MaxY, 0);
@@ -1058,12 +1054,12 @@ void RasterizeToVolumeTexture(FRHICommandList& RHICmdList, FVolumeBounds VolumeB
 
 /** Helper function that clears the given volume texture render targets. */
 template<int32 NumRenderTargets>
-void ClearVolumeTextures(FRHICommandList& RHICmdList, const FTextureRHIParamRef* RenderTargets, const FLinearColor* ClearColors)
+void ClearVolumeTextures(FRHICommandListImmediate& RHICmdList, const FTextureRHIParamRef* RenderTargets, const FLinearColor* ClearColors)
 {
 	SetRenderTargets(RHICmdList, NumRenderTargets, RenderTargets, FTextureRHIRef(), 0, NULL);
 
 	static FGlobalBoundShaderState VolumeClearBoundShaderState;
-	RHICmdList.CheckIsNull(); // need new approach for "static FGlobalBoundShaderState" for parallel rendering
+	
 
 	// Currently using a manual clear, which is ~10x faster than a hardware clear of the volume textures on AMD PC GPU's
 	if (false)
@@ -1100,7 +1096,7 @@ void ClearVolumeTextures(FRHICommandList& RHICmdList, const FTextureRHIParamRef*
 	}
 }
 
-void FDeferredShadingSceneRenderer::ClearTranslucentVolumeLighting(FRHICommandList& RHICmdList)
+void FDeferredShadingSceneRenderer::ClearTranslucentVolumeLighting(FRHICommandListImmediate& RHICmdList)
 {
 	if (GUseTranslucentLightingVolumes && GSupportsVolumeTextureRendering)
 	{
@@ -1163,13 +1159,13 @@ public:
 
 		FGlobalShader::SetParameters(RHICmdList, ShaderRHI, View);
 
-		CubemapShaderParameters.SetParameters(ShaderRHI, CubemapEntry);
+		CubemapShaderParameters.SetParameters(RHICmdList, ShaderRHI, CubemapEntry);
 	}
 };
 
 IMPLEMENT_SHADER_TYPE(,FInjectAmbientCubemapPS,TEXT("TranslucentLightingShaders"),TEXT("InjectAmbientCubemapMainPS"),SF_Pixel);
 
-void FDeferredShadingSceneRenderer::InjectAmbientCubemapTranslucentVolumeLighting(FRHICommandList& RHICmdList)
+void FDeferredShadingSceneRenderer::InjectAmbientCubemapTranslucentVolumeLighting(FRHICommandListImmediate& RHICmdList)
 {
 	//@todo - support multiple views
 	const FViewInfo& View = Views[0];
@@ -1194,7 +1190,7 @@ void FDeferredShadingSceneRenderer::InjectAmbientCubemapTranslucentVolumeLightin
 			TShaderMapRef<FInjectAmbientCubemapPS> PixelShader(GetGlobalShaderMap());
 
 			static FGlobalBoundShaderState BoundShaderState;
-			RHICmdList.CheckIsNull(); // need new approach for "static FGlobalBoundShaderState" for parallel rendering						
+									
 			SetGlobalBoundShaderState(RHICmdList, BoundShaderState, GScreenVertexDeclaration.VertexDeclarationRHI, *VertexShader, *PixelShader, *GeometryShader);
 
 			VertexShader->SetParameters(RHICmdList, VolumeBounds, GTranslucencyLightingVolumeDim);
@@ -1322,7 +1318,7 @@ private:
 
 IMPLEMENT_SHADER_TYPE(,FCompositeGIForTranslucencyPS,TEXT("TranslucentLightingShaders"),TEXT("CompositeGIMainPS"),SF_Pixel);
 
-void FDeferredShadingSceneRenderer::CompositeIndirectTranslucentVolumeLighting(FRHICommandList& RHICmdList)
+void FDeferredShadingSceneRenderer::CompositeIndirectTranslucentVolumeLighting(FRHICommandListImmediate& RHICmdList)
 {
 	bool bAnyViewAllowsIndirectLightingCache = false;
 
@@ -1365,7 +1361,7 @@ void FDeferredShadingSceneRenderer::CompositeIndirectTranslucentVolumeLighting(F
 				TShaderMapRef<FCompositeGIForTranslucencyPS> PixelShader(GetGlobalShaderMap());
 
 				static FGlobalBoundShaderState BoundShaderState;
-				RHICmdList.CheckIsNull(); // need new approach for "static FGlobalBoundShaderState" for parallel rendering
+				
 				SetGlobalBoundShaderState(RHICmdList, BoundShaderState, GScreenVertexDeclaration.VertexDeclarationRHI, *VertexShader, *PixelShader, *GeometryShader);
 
 				PixelShader->SetParameters(RHICmdList, View, ViewState, VolumeCascadeIndex);
@@ -1383,7 +1379,7 @@ void FDeferredShadingSceneRenderer::CompositeIndirectTranslucentVolumeLighting(F
 	}
 }
 
-void FDeferredShadingSceneRenderer::ClearTranslucentVolumePerObjectShadowing(FRHICommandList& RHICmdList)
+void FDeferredShadingSceneRenderer::ClearTranslucentVolumePerObjectShadowing(FRHICommandListImmediate& RHICmdList)
 {
 	if (GUseTranslucentLightingVolumes && GSupportsVolumeTextureRendering)
 	{
@@ -1431,7 +1427,7 @@ FVolumeBounds CalculateLightVolumeBounds(const FSphere& LightBounds, const FView
 
 FGlobalBoundShaderState ObjectShadowingBoundShaderState;
 
-void FDeferredShadingSceneRenderer::AccumulateTranslucentVolumeObjectShadowing(FRHICommandList& RHICmdList, const FProjectedShadowInfo* InProjectedShadowInfo, bool bClearVolume)
+void FDeferredShadingSceneRenderer::AccumulateTranslucentVolumeObjectShadowing(FRHICommandListImmediate& RHICmdList, const FProjectedShadowInfo* InProjectedShadowInfo, bool bClearVolume)
 {
 	const FLightSceneInfo* LightSceneInfo = InProjectedShadowInfo->LightSceneInfo;
 
@@ -1557,7 +1553,6 @@ void SetInjectionShader(
 
 	// Recreate the bound shader state if the shader map has changed since the last cache
 	// This can happen due to async shader compiling
-	RHICmdList.CheckIsNull(); // RHICreateBoundShaderState
 
 	if (!IsValidRef(BoundShaderState) || CachedShaderMap != MaterialShaderMap)
 	{
@@ -1654,7 +1649,7 @@ static void AddLightForInjection(
 }
 
 /** Injects all the lights in LightInjectionData into the translucent lighting volume textures. */
-static void InjectTranslucentLightArray(FRHICommandList& RHICmdList, const FViewInfo& View, const TArray<FTranslucentLightInjectionData, SceneRenderingAllocator>& LightInjectionData)
+static void InjectTranslucentLightArray(FRHICommandListImmediate& RHICmdList, const FViewInfo& View, const TArray<FTranslucentLightInjectionData, SceneRenderingAllocator>& LightInjectionData)
 {
 	INC_DWORD_STAT_BY(STAT_NumLightsInjectedIntoTranslucency, LightInjectionData.Num());
 
@@ -1737,7 +1732,7 @@ static void InjectTranslucentLightArray(FRHICommandList& RHICmdList, const FView
 	}
 }
 
-void FDeferredShadingSceneRenderer::InjectTranslucentVolumeLighting(FRHICommandList& RHICmdList, const FLightSceneInfo& LightSceneInfo, const FProjectedShadowInfo* InProjectedShadowInfo)
+void FDeferredShadingSceneRenderer::InjectTranslucentVolumeLighting(FRHICommandListImmediate& RHICmdList, const FLightSceneInfo& LightSceneInfo, const FProjectedShadowInfo* InProjectedShadowInfo)
 {
 	if (GUseTranslucentLightingVolumes && GSupportsVolumeTextureRendering)
 	{
@@ -1755,7 +1750,7 @@ void FDeferredShadingSceneRenderer::InjectTranslucentVolumeLighting(FRHICommandL
 	}
 }
 
-void FDeferredShadingSceneRenderer::InjectTranslucentVolumeLightingArray(FRHICommandList& RHICmdList, const TArray<FSortedLightSceneInfo, SceneRenderingAllocator>& SortedLights, int32 NumLights)
+void FDeferredShadingSceneRenderer::InjectTranslucentVolumeLightingArray(FRHICommandListImmediate& RHICmdList, const TArray<FSortedLightSceneInfo, SceneRenderingAllocator>& SortedLights, int32 NumLights)
 {
 	SCOPE_CYCLE_COUNTER(STAT_TranslucentInjectTime);
 
@@ -1827,7 +1822,7 @@ IMPLEMENT_SHADER_TYPE(,FSimpleLightTranslucentLightingInjectPS,TEXT("Translucent
 
 FGlobalBoundShaderState InjectSimpleLightBoundShaderState;
 
-void FDeferredShadingSceneRenderer::InjectSimpleTranslucentVolumeLightingArray(FRHICommandList& RHICmdList, const FSimpleLightArray& SimpleLights)
+void FDeferredShadingSceneRenderer::InjectSimpleTranslucentVolumeLightingArray(FRHICommandListImmediate& RHICmdList, const FSimpleLightArray& SimpleLights)
 {
 	SCOPE_CYCLE_COUNTER(STAT_TranslucentInjectTime);
 
@@ -1908,7 +1903,7 @@ void FDeferredShadingSceneRenderer::InjectSimpleTranslucentVolumeLightingArray(F
 
 FGlobalBoundShaderState FilterBoundShaderState;
 
-void FDeferredShadingSceneRenderer::FilterTranslucentVolumeLighting(FRHICommandList& RHICmdList)
+void FDeferredShadingSceneRenderer::FilterTranslucentVolumeLighting(FRHICommandListImmediate& RHICmdList)
 {
 	if (GUseTranslucentLightingVolumes && GSupportsVolumeTextureRendering)
 	{
diff --git a/Engine/Source/Runtime/Renderer/Private/TranslucentRendering.cpp b/Engine/Source/Runtime/Renderer/Private/TranslucentRendering.cpp
--- a/Engine/Source/Runtime/Renderer/Private/TranslucentRendering.cpp
+++ b/Engine/Source/Runtime/Renderer/Private/TranslucentRendering.cpp
@@ -10,28 +10,28 @@
 #include "SceneFilterRendering.h"
 #include "LightPropagationVolume.h"
 
-void SetTranslucentRenderTargetAndState(const FViewInfo& View, bool bSeparateTranslucencyPass)
+void SetTranslucentRenderTargetAndState(FRHICommandListImmediate& RHICmdList, const FViewInfo& View, bool bSeparateTranslucencyPass)
 {
 	bool bSetupTranslucentState = true;
 
 	if (bSeparateTranslucencyPass && GSceneRenderTargets.IsSeparateTranslucencyActive(View))
 	{
-		bSetupTranslucentState = GSceneRenderTargets.BeginRenderingSeparateTranslucency(View, false);
+		bSetupTranslucentState = GSceneRenderTargets.BeginRenderingSeparateTranslucency(RHICmdList, View, false);
 	}
 	else
 	{
-		GSceneRenderTargets.BeginRenderingTranslucency(View);
+		GSceneRenderTargets.BeginRenderingTranslucency(RHICmdList, View);
 	}
 
 	if (bSetupTranslucentState)
 	{
 		// Enable depth test, disable depth writes.
 		// Note, this is a reversed Z depth surface, using CF_GreaterEqual.
-		RHISetDepthStencilState(TStaticDepthStencilState<false,CF_GreaterEqual>::GetRHI());
+		RHICmdList.SetDepthStencilState(TStaticDepthStencilState<false, CF_GreaterEqual>::GetRHI());
 	}
 }
 
-const FProjectedShadowInfo* FDeferredShadingSceneRenderer::PrepareTranslucentShadowMap(const FViewInfo& View, FPrimitiveSceneInfo* PrimitiveSceneInfo, bool bSeparateTranslucencyPass)
+const FProjectedShadowInfo* FDeferredShadingSceneRenderer::PrepareTranslucentShadowMap(FRHICommandListImmediate& RHICmdList, const FViewInfo& View, FPrimitiveSceneInfo* PrimitiveSceneInfo, bool bSeparateTranslucencyPass)
 {
 	const FVisibleLightInfo* VisibleLightInfo = NULL;
 	FProjectedShadowInfo* TranslucentSelfShadow = NULL;
@@ -120,10 +120,10 @@ const FProjectedShadowInfo* FDeferredShadingSceneRenderer::PrepareTranslucentSha
 				check(TranslucentSelfShadow->bAllocatedInTranslucentLayout);
 
 				// Render the translucency shadow map
-				TranslucentSelfShadow->RenderTranslucencyDepths(this);
+				TranslucentSelfShadow->RenderTranslucencyDepths(RHICmdList, this);
 
 				// Restore state
-				SetTranslucentRenderTargetAndState(View, bSeparateTranslucencyPass);
+				SetTranslucentRenderTargetAndState(RHICmdList, View, bSeparateTranslucencyPass);
 			}
 		}
 	}
@@ -166,20 +166,18 @@ IMPLEMENT_SHADER_TYPE(,FCopySceneColorPS,TEXT("TranslucentLightingShaders"),TEXT
 
 FGlobalBoundShaderState CopySceneColorBoundShaderState;
 
-void FTranslucencyDrawingPolicyFactory::CopySceneColor(const FViewInfo& View, const FPrimitiveSceneProxy* PrimitiveSceneProxy)
+void FTranslucencyDrawingPolicyFactory::CopySceneColor(FRHICommandListImmediate& RHICmdList, const FViewInfo& View, const FPrimitiveSceneProxy* PrimitiveSceneProxy)
 {
-	//@todo-rco: RHIPacketList
-	FRHICommandList& RHICmdList = FRHICommandList::GetNullRef();
 
 	SCOPED_DRAW_EVENTF(EventCopy, DEC_SCENE_ITEMS, TEXT("CopySceneColor for %s %s"), *PrimitiveSceneProxy->GetOwnerName().ToString(), *PrimitiveSceneProxy->GetResourceName().ToString());
-	RHISetRasterizerState(TStaticRasterizerState<FM_Solid,CM_None>::GetRHI());
-	RHISetDepthStencilState(TStaticDepthStencilState<false,CF_Always>::GetRHI());
-	RHISetBlendState(TStaticBlendState<>::GetRHI());
+	RHICmdList.SetRasterizerState(TStaticRasterizerState<FM_Solid, CM_None>::GetRHI());
+	RHICmdList.SetDepthStencilState(TStaticDepthStencilState<false, CF_Always>::GetRHI());
+	RHICmdList.SetBlendState(TStaticBlendState<>::GetRHI());
 
-	GSceneRenderTargets.ResolveSceneColor();
+	GSceneRenderTargets.ResolveSceneColor(RHICmdList);
 
-	GSceneRenderTargets.BeginRenderingLightAttenuation();
-	RHISetViewport(View.ViewRect.Min.X, View.ViewRect.Min.Y, 0.0f, View.ViewRect.Max.X, View.ViewRect.Max.Y, 1.0f);
+	GSceneRenderTargets.BeginRenderingLightAttenuation(RHICmdList);
+	RHICmdList.SetViewport(View.ViewRect.Min.X, View.ViewRect.Min.Y, 0.0f, View.ViewRect.Max.X, View.ViewRect.Max.Y, 1.0f);
 
 	TShaderMapRef<FScreenVS> ScreenVertexShader(GetGlobalShaderMap());
 	TShaderMapRef<FCopySceneColorPS> PixelShader(GetGlobalShaderMap());
@@ -192,14 +190,14 @@ void FTranslucencyDrawingPolicyFactory::CopySceneColor(const FViewInfo& View, co
 		RHICmdList,
 		0, 0, 
 		View.ViewRect.Width(), View.ViewRect.Height(),
-		View.ViewRect.Min.X, View.ViewRect.Min.Y,
+		View.ViewRect.Min.X, View.ViewRect.Min.Y, 
 		View.ViewRect.Width(), View.ViewRect.Height(),
 		FIntPoint(View.ViewRect.Width(), View.ViewRect.Height()),
 		GSceneRenderTargets.GetBufferSizeXY(),
 		*ScreenVertexShader,
 		EDRF_UseTriangleOptimization);
 
-	GSceneRenderTargets.FinishRenderingLightAttenuation();
+	GSceneRenderTargets.FinishRenderingLightAttenuation(RHICmdList);
 }
 
 /** The parameters used to draw a translucent mesh. */
@@ -336,9 +334,13 @@ bool FTranslucencyDrawingPolicyFactory::DrawMesh(
 		{
 			if (DrawingContext.bSceneColorCopyIsUpToDate == false)
 			{
-				CopySceneColor(View, PrimitiveSceneProxy);
+				//hack here to get the code checked in, but this needs to be fixed. It was not sensible to wind this all the way out of the callstack
+				FRHICommandListImmediate& RHICmdListIm = FRHICommandListExecutor::GetImmediateCommandList();
+				check(&RHICmdList == &RHICmdListIm);
+
+				CopySceneColor(RHICmdListIm, View, PrimitiveSceneProxy);
 				// Restore state
-				SetTranslucentRenderTargetAndState(View, false);
+				SetTranslucentRenderTargetAndState(RHICmdListIm, View, false);
 
 				// separate translucency is not updating scene color so we don't need to copy it multiple times.
 				// optimization:
@@ -469,6 +471,7 @@ FTranslucentPrimSet
 * @return true if anything was drawn
 */
 void FTranslucentPrimSet::DrawPrimitives(
+	FRHICommandListImmediate& RHICmdList,
 	const FViewInfo& View,
 	FDeferredShadingSceneRenderer& Renderer,
 	bool bSeparateTranslucencyPass
@@ -491,14 +494,15 @@ void FTranslucentPrimSet::DrawPrimitives(
 
 			checkSlow(ViewRelevance.HasTranslucency());
 			
-			const FProjectedShadowInfo* TranslucentSelfShadow = Renderer.PrepareTranslucentShadowMap(View, PrimitiveSceneInfo, bSeparateTranslucencyPass);
+			const FProjectedShadowInfo* TranslucentSelfShadow = Renderer.PrepareTranslucentShadowMap(RHICmdList, View, PrimitiveSceneInfo, bSeparateTranslucencyPass);
 
-			RenderPrimitive(View, PrimitiveSceneInfo, ViewRelevance, TranslucentSelfShadow, bSeparateTranslucencyPass);
+			RenderPrimitive(RHICmdList, View, PrimitiveSceneInfo, ViewRelevance, TranslucentSelfShadow, bSeparateTranslucencyPass);
 		}
 	}
 }
 
 void FTranslucentPrimSet::RenderPrimitive(
+	FRHICommandListImmediate& RHICmdList,
 	const FViewInfo& View, 
 	FPrimitiveSceneInfo* PrimitiveSceneInfo, 
 	const FPrimitiveViewRelevance& ViewRelevance, 
@@ -533,7 +537,8 @@ void FTranslucentPrimSet::RenderPrimitive(
 			SCOPE_CYCLE_COUNTER(STAT_RHICounterTEMP);
 			static IConsoleVariable* RHICmdListCVar = IConsoleManager::Get().FindConsoleVariable(TEXT("r.RHICmd"));
 			bool bUseRHICmdList = (RHICmdListCVar->GetInt() >= 3);
-			FRHICommandList& RHICmdList = bUseRHICmdList ? GRHICommandList.CreateList() : FRHICommandList::GetNullRef();
+			FRHICommandList LocalRHICmdList;
+			FRHICommandList& UseRHICmdList = bUseRHICmdList ? LocalRHICmdList : RHICmdList;
 
 			// Render static meshes from static scene prim
 			for( int32 StaticMeshIdx=0; StaticMeshIdx < PrimitiveSceneInfo->StaticMeshes.Num(); StaticMeshIdx++ )
@@ -545,7 +550,7 @@ void FTranslucentPrimSet::RenderPrimitive(
 					&& (StaticMesh.MaterialRenderProxy->GetMaterial(FeatureLevel)->IsSeparateTranslucencyEnabled() == bSeparateTranslucencyPass))
 				{
 					FTranslucencyDrawingPolicyFactory::DrawStaticMesh(
-						RHICmdList, 
+						UseRHICmdList,
 						View,
 						FTranslucencyDrawingPolicyFactory::ContextType( TranslucentSelfShadow, bSeparateTranslucencyPass),
 						StaticMesh,
@@ -556,8 +561,6 @@ void FTranslucentPrimSet::RenderPrimitive(
 						);
 				}
 			}
-
-			GRHICommandList.ExecuteAndFreeList(RHICmdList);
 		}
 	}
 }
@@ -620,12 +623,10 @@ bool FSceneRenderer::ShouldRenderTranslucency() const
 	return bRender;
 }
 
-void FDeferredShadingSceneRenderer::RenderTranslucency()
+void FDeferredShadingSceneRenderer::RenderTranslucency(FRHICommandListImmediate& RHICmdList)
 {
 	if (ShouldRenderTranslucency())
 	{
-		//@todo-rco: RHIPacketList
-		FRHICommandList& RHICmdList = FRHICommandList::GetNullRef();
 
 		SCOPED_DRAW_EVENT(Translucency, DEC_SCENE_ITEMS);
 
@@ -635,10 +636,10 @@ void FDeferredShadingSceneRenderer::RenderTranslucency()
 
 			const FViewInfo& View = Views[ViewIndex];
 
-			SetTranslucentRenderTargetAndState(View, false);
+			SetTranslucentRenderTargetAndState(RHICmdList, View, false);
 
 			// Draw only translucent prims that don't read from scene color
-			View.TranslucentPrimSet.DrawPrimitives(View, *this, false);
+			View.TranslucentPrimSet.DrawPrimitives(RHICmdList, View, *this, false);
 			// Draw the view's mesh elements with the translucent drawing policy.
 			DrawViewElements<FTranslucencyDrawingPolicyFactory>(View,FTranslucencyDrawingPolicyFactory::ContextType(),SDPG_World,false);
 			// Draw the view's mesh elements with the translucent drawing policy.
@@ -660,20 +661,20 @@ void FDeferredShadingSceneRenderer::RenderTranslucency()
 				bool bRenderSeparateTranslucency = View.TranslucentPrimSet.NumSeparateTranslucencyPrims() > 0;
 
 				// always call BeginRenderingSeparateTranslucency() even if there are no primitives to we keep the RT allocated
-				bool bSetupTranslucency = GSceneRenderTargets.BeginRenderingSeparateTranslucency(View, true);
+				bool bSetupTranslucency = GSceneRenderTargets.BeginRenderingSeparateTranslucency(RHICmdList, View, true);
 
 				// Draw only translucent prims that are in the SeparateTranslucency pass
 				if (bRenderSeparateTranslucency)
 				{
 					if (bSetupTranslucency)
 					{
-						RHISetDepthStencilState(TStaticDepthStencilState<false,CF_GreaterEqual>::GetRHI());
+						RHICmdList.SetDepthStencilState(TStaticDepthStencilState<false, CF_GreaterEqual>::GetRHI());
 					}
 
-					View.TranslucentPrimSet.DrawPrimitives(View, *this, true);
+					View.TranslucentPrimSet.DrawPrimitives(RHICmdList, View, *this, true);
 				}
 
-				GSceneRenderTargets.FinishRenderingSeparateTranslucency(View);
+				GSceneRenderTargets.FinishRenderingSeparateTranslucency(RHICmdList, View);
 			}
 		}
 	}
diff --git a/Engine/Source/Runtime/Renderer/Private/TranslucentRendering.h b/Engine/Source/Runtime/Renderer/Private/TranslucentRendering.h
--- a/Engine/Source/Runtime/Renderer/Private/TranslucentRendering.h
+++ b/Engine/Source/Runtime/Renderer/Private/TranslucentRendering.h
@@ -84,7 +84,7 @@ private:
 	/**
 	* Resolves the scene color target and copies it for use as a source texture.
 	*/
-	static void CopySceneColor(const FViewInfo& View, const FPrimitiveSceneProxy* PrimitiveSceneProxy);
+	static void CopySceneColor(FRHICommandListImmediate& RHICmdList, const FViewInfo& View, const FPrimitiveSceneProxy* PrimitiveSceneProxy);
 };
 
 
@@ -240,4 +240,4 @@ private:
 	FShaderParameter MinZ;
 };
 
-extern void RasterizeToVolumeTexture(FRHICommandList& RHICmdList, FVolumeBounds VolumeBounds);
\ No newline at end of file
+extern void RasterizeToVolumeTexture(FRHICommandListImmediate& RHICmdList, FVolumeBounds VolumeBounds);
\ No newline at end of file
diff --git a/Engine/Source/Runtime/Renderer/Private/VelocityRendering.cpp b/Engine/Source/Runtime/Renderer/Private/VelocityRendering.cpp
--- a/Engine/Source/Runtime/Renderer/Private/VelocityRendering.cpp
+++ b/Engine/Source/Runtime/Renderer/Private/VelocityRendering.cpp
@@ -447,7 +447,7 @@ bool IsMotionBlurEnabled(const FViewInfo& View)
 		&& !(View.Family->Views.Num() > 1);
 }
 
-void FDeferredShadingSceneRenderer::RenderVelocities(const FViewInfo& View, TRefCountPtr<IPooledRenderTarget>& VelocityRT, bool bLastFrame)
+void FDeferredShadingSceneRenderer::RenderVelocities(FRHICommandListImmediate& RHICmdList, const FViewInfo& View, TRefCountPtr<IPooledRenderTarget>& VelocityRT, bool bLastFrame)
 {
 	SCOPE_CYCLE_COUNTER( STAT_RenderVelocities );
 
@@ -472,29 +472,26 @@ void FDeferredShadingSceneRenderer::RenderVelocities(const FViewInfo& View, TRef
 
 	GPrevPerBoneMotionBlur.LockData();
 
-	//@todo-rco: RHIPacketList
-	FRHICommandList& RHICmdList = FRHICommandList::GetNullRef();
-
 	const uint32 MinX = View.ViewRect.Min.X * VelocityBufferSize.X / BufferSize.X;
 	const uint32 MinY = View.ViewRect.Min.Y * VelocityBufferSize.Y / BufferSize.Y;
 	const uint32 MaxX = View.ViewRect.Max.X * VelocityBufferSize.X / BufferSize.X;
 	const uint32 MaxY = View.ViewRect.Max.Y * VelocityBufferSize.Y / BufferSize.Y;
-	RHISetRenderTarget(VelocityRT->GetRenderTargetItem().TargetableTexture, GSceneRenderTargets.GetSceneDepthTexture());
+	SetRenderTarget(RHICmdList, VelocityRT->GetRenderTargetItem().TargetableTexture, GSceneRenderTargets.GetSceneDepthTexture());
 
-	RHISetViewport(MinX, MinY, 0.0f, MaxX, MaxY, 1.0f);
+	RHICmdList.SetViewport(MinX, MinY, 0.0f, MaxX, MaxY, 1.0f);
 
 	// Clear the velocity buffer (0.0f means "use static background velocity").
-	RHIClear(true, FLinearColor::Black, false, 1.0f, false, 0, FIntRect());
+	RHICmdList.Clear(true, FLinearColor::Black, false, 1.0f, false, 0, FIntRect());
 
 	// Blending is not supported with the velocity buffer format on other platforms
 	// opaque velocities in R|G channels, B is used to identify object motion
-	RHISetBlendState(TStaticBlendState<CW_RGBA>::GetRHI());
+	RHICmdList.SetBlendState(TStaticBlendState<CW_RGBA>::GetRHI());
 	// Use depth tests, no z-writes, backface-culling. 
 
 	// Note, this is a reversed Z depth surface, using CF_GreaterEqual.
-	RHISetDepthStencilState(TStaticDepthStencilState<false,CF_GreaterEqual>::GetRHI());
+	RHICmdList.SetDepthStencilState(TStaticDepthStencilState<false,CF_GreaterEqual>::GetRHI());
 			
-	RHISetRasterizerState(GetStaticRasterizerState<true>(FM_Solid, CM_CW));
+	RHICmdList.SetRasterizerState(GetStaticRasterizerState<true>(FM_Solid, CM_CW));
 
 	// Draw velocities for movable static meshes.
 	Scene->VelocityDrawList.DrawVisible(RHICmdList, View,View.StaticMeshVelocityMap, View.StaticMeshBatchVisibility);
@@ -519,11 +516,11 @@ void FDeferredShadingSceneRenderer::RenderVelocities(const FViewInfo& View, TRef
 	}				
 	Drawer.IsDirty();
 
-	RHICopyToResolveTarget(VelocityRT->GetRenderTargetItem().TargetableTexture, VelocityRT->GetRenderTargetItem().ShaderResourceTexture, false, FResolveParams());
+	RHICmdList.CopyToResolveTarget(VelocityRT->GetRenderTargetItem().TargetableTexture, VelocityRT->GetRenderTargetItem().ShaderResourceTexture, false, FResolveParams());
 
 	// restore any color write state changes
-	RHISetBlendState(TStaticBlendState<>::GetRHI());
-	RHISetRasterizerState(TStaticRasterizerState<>::GetRHI());
+	RHICmdList.SetBlendState(TStaticBlendState<>::GetRHI());
+	RHICmdList.SetRasterizerState(TStaticRasterizerState<>::GetRHI());
 	GPrevPerBoneMotionBlur.UnlockData(bLastFrame);
 
 	// to be able to observe results with VisualizeTexture
diff --git a/Engine/Source/Runtime/Renderer/Public/PostProcessParameters.h b/Engine/Source/Runtime/Renderer/Public/PostProcessParameters.h
--- a/Engine/Source/Runtime/Renderer/Public/PostProcessParameters.h
+++ b/Engine/Source/Runtime/Renderer/Public/PostProcessParameters.h
@@ -90,7 +90,6 @@ public:
 	// @param FilterOverrideArray can be 0 if Filter is used
 	template< typename ShaderRHIParamRef >
 	void Set(
-		FRHICommandList& RHICmdList,
 		const ShaderRHIParamRef& ShaderRHI,
 		const FRenderingCompositePassContext& Context,
 		FSamplerStateRHIParamRef Filter,
diff --git a/Engine/Source/Runtime/Renderer/Public/RendererInterface.h b/Engine/Source/Runtime/Renderer/Public/RendererInterface.h
--- a/Engine/Source/Runtime/Renderer/Public/RendererInterface.h
+++ b/Engine/Source/Runtime/Renderer/Public/RendererInterface.h
@@ -452,7 +452,7 @@ public:
 	virtual void SceneRenderTargetsSetBufferSize(uint32 SizeX, uint32 SizeY) = 0;
 
 	/** Draws a tile mesh element with the specified view. */
-	virtual void DrawTileMesh(const FSceneView& View, const FMeshBatch& Mesh, bool bIsHitTesting, const class FHitProxyId& HitProxyId) = 0;
+	virtual void DrawTileMesh(FRHICommandListImmediate& RHICmdList, const FSceneView& View, const FMeshBatch& Mesh, bool bIsHitTesting, const class FHitProxyId& HitProxyId) = 0;
 
 	/** Render thread side, use TRefCountPtr<IPooledRenderTarget>, allows to use sharing and VisualizeTexture */
 	virtual void RenderTargetPoolFindFreeElement(const FPooledRenderTargetDesc& Desc, TRefCountPtr<IPooledRenderTarget> &Out, const TCHAR* InDebugName) = 0;
diff --git a/Engine/Source/Runtime/ShaderCore/Private/Shader.cpp b/Engine/Source/Runtime/ShaderCore/Private/Shader.cpp
--- a/Engine/Source/Runtime/ShaderCore/Private/Shader.cpp
+++ b/Engine/Source/Runtime/ShaderCore/Private/Shader.cpp
@@ -1034,29 +1034,25 @@ FShaderType* FindShaderTypeByName(const TCHAR* ShaderTypeName)
 
 
 void DispatchComputeShader(
-	FRHICommandList& RHICmdList, 
+	FRHICommandListImmediate& RHICmdList,
 	FShader* Shader,
 	uint32 ThreadGroupCountX,
 	uint32 ThreadGroupCountY,
 	uint32 ThreadGroupCountZ)
 {
-	RHICmdList.CheckIsNull();
-
 	Shader->VerifyBoundUniformBufferParameters();
-	RHIDispatchComputeShader(ThreadGroupCountX, ThreadGroupCountY, ThreadGroupCountZ);
+	RHICmdList.DispatchComputeShader(ThreadGroupCountX, ThreadGroupCountY, ThreadGroupCountZ);
 }
 
 
 void DispatchIndirectComputeShader(
-	FRHICommandList& RHICmdList, 
+	FRHICommandListImmediate& RHICmdList,
 	FShader* Shader,
 	FVertexBufferRHIParamRef ArgumentBuffer,
 	uint32 ArgumentOffset)
 {
-	RHICmdList.CheckIsNull();
-
 	Shader->VerifyBoundUniformBufferParameters();
-	RHIDispatchIndirectComputeShader(ArgumentBuffer, ArgumentOffset);
+	RHICmdList.DispatchIndirectComputeShader(ArgumentBuffer, ArgumentOffset);
 }
 
 
diff --git a/Engine/Source/Runtime/ShaderCore/Public/Shader.h b/Engine/Source/Runtime/ShaderCore/Public/Shader.h
--- a/Engine/Source/Runtime/ShaderCore/Public/Shader.h
+++ b/Engine/Source/Runtime/ShaderCore/Public/Shader.h
@@ -1320,7 +1320,7 @@ extern SHADERCORE_API FShaderType* FindShaderTypeByName(const TCHAR* ShaderTypeN
 
 /** Helper function to dispatch a compute shader while checking that parameters have been set correctly. */
 extern SHADERCORE_API void DispatchComputeShader(
-	FRHICommandList& RHICmdList,
+	FRHICommandListImmediate& RHICmdList,
 	FShader* Shader,
 	uint32 ThreadGroupCountX,
 	uint32 ThreadGroupCountY,
@@ -1328,7 +1328,7 @@ extern SHADERCORE_API void DispatchComputeShader(
 
 /** Helper function to dispatch a compute shader indirectly while checking that parameters have been set correctly. */
 extern SHADERCORE_API void DispatchIndirectComputeShader(
-	FRHICommandList& RHICmdList, 
+	FRHICommandListImmediate& RHICmdList,
 	FShader* Shader,
 	FVertexBufferRHIParamRef ArgumentBuffer,
 	uint32 ArgumentOffset);
diff --git a/Engine/Source/Runtime/SlateCore/Public/Rendering/RenderingCommon.h b/Engine/Source/Runtime/SlateCore/Public/Rendering/RenderingCommon.h
--- a/Engine/Source/Runtime/SlateCore/Public/Rendering/RenderingCommon.h
+++ b/Engine/Source/Runtime/SlateCore/Public/Rendering/RenderingCommon.h
@@ -461,5 +461,5 @@ public:
 	 *
 	 * @param RenderTarget	handle to the platform specific render target implementation.  Note this is already bound by Slate initially 
 	 */
-	virtual void DrawRenderThread( const void* RenderTarget ) = 0;
+	virtual void DrawRenderThread(class FRHICommandListImmediate& RHICmdList, const void* RenderTarget) = 0;
 };
\ No newline at end of file
