UE4 - Fix PS4 compile issues

diff --git a/Engine/Source/Runtime/Renderer/Private/ReflectionEnvironmentCapture.cpp b/Engine/Source/Runtime/Renderer/Private/ReflectionEnvironmentCapture.cpp
--- a/Engine/Source/Runtime/Renderer/Private/ReflectionEnvironmentCapture.cpp
+++ b/Engine/Source/Runtime/Renderer/Private/ReflectionEnvironmentCapture.cpp
@@ -835,7 +835,8 @@ void CaptureSceneToScratchCubemap(FRHICommandListImmediate& RHICmdList, FSceneRe
 #if PLATFORM_PS4 // @todo ps4 - this should be done a different way
 		// PS4 needs some code here to process the scene
 		extern void TEMP_PostReflectionCaptureRender();
-		RHICmdList.CheckIsNull(); //FDeferredUpdateResource::UpdateResources()
+		//@todo-rco: Fix this assert
+		check(&RHICmdList == &FRHICommandListExecutor::GetImmediateCommandList());
 		TEMP_PostReflectionCaptureRender();
 #endif
 
