Changed SSR roughness fade for low quality.

diff --git a/Engine/Source/Runtime/Renderer/Private/PostProcess/ScreenSpaceReflections.cpp b/Engine/Source/Runtime/Renderer/Private/PostProcess/ScreenSpaceReflections.cpp
--- a/Engine/Source/Runtime/Renderer/Private/PostProcess/ScreenSpaceReflections.cpp
+++ b/Engine/Source/Runtime/Renderer/Private/PostProcess/ScreenSpaceReflections.cpp
@@ -116,6 +116,7 @@ public:
 			// f(MaxRoughness/2) = 1
 
 			float RoughnessMaskScale = -2.0f / MaxRoughness;
+			RoughnessMaskScale *= SSRQuality < 3 ? 2.0f : 1.0f;
 
 			FLinearColor Value(
 				FMath::Clamp(Context.View.FinalPostProcessSettings.ScreenSpaceReflectionIntensity * 0.01f, 0.0f, 1.0f), 
