Interim fix for postprocessing of stereo buffers on ES2 devices.

- ES2 postprocessing really needs to have a case that can handle stereo views.  Until that is done, just run the post with the first view expanded to cover the entire stereo buffer.
- This very likely doesn't work in all cases as some of the view values will not be valid, but it works in simple cases.

diff --git a/Engine/Source/Runtime/Renderer/Private/ForwardShadingRenderer.cpp b/Engine/Source/Runtime/Renderer/Private/ForwardShadingRenderer.cpp
--- a/Engine/Source/Runtime/Renderer/Private/ForwardShadingRenderer.cpp
+++ b/Engine/Source/Runtime/Renderer/Private/ForwardShadingRenderer.cpp
@@ -140,7 +140,18 @@ void FForwardShadingSceneRenderer::Render(FRHICommandListImmediate& RHICmdList)
 		// Drop depth and stencil before post processing to avoid export.
 		RHICmdList.DiscardRenderTargets(true, true, 0);
 
-		// Finish rendering for each view.
+		// Finish rendering for each view, or the full stereo buffer if enabled
+		if (GEngine->IsStereoscopic3D())
+		{
+			check(Views.Num() > 1);
+
+			//@todo ES2 stereo post: until we get proper stereo postprocessing for ES2, process the stereo buffer as one view
+			FIntPoint OriginalMax0 = Views[0].ViewRect.Max;
+			Views[0].ViewRect.Max = Views[1].ViewRect.Max;
+			GPostProcessing.ProcessES2(RHICmdList, Views[0], bOnChipSunMask);
+			Views[0].ViewRect.Max = OriginalMax0;
+		}
+		else
 		{
 			SCOPED_DRAW_EVENT(FinishRendering, DEC_SCENE_ITEMS);
 			SCOPE_CYCLE_COUNTER(STAT_FinishRenderViewTargetTime);
