* added more debug info

diff --git a/Engine/Source/Runtime/Renderer/Private/DeferredShadingRenderer.cpp b/Engine/Source/Runtime/Renderer/Private/DeferredShadingRenderer.cpp
--- a/Engine/Source/Runtime/Renderer/Private/DeferredShadingRenderer.cpp
+++ b/Engine/Source/Runtime/Renderer/Private/DeferredShadingRenderer.cpp
@@ -666,6 +666,8 @@ void FDeferredShadingSceneRenderer::Render(FRHICommandListImmediate& RHICmdList)
 
 		RenderLights(RHICmdList);
 
+		GRenderTargetPool.AddPhaseEvent(TEXT("AfterRenderLights"));
+
 		InjectAmbientCubemapTranslucentVolumeLighting(RHICmdList);
 
 		CompositeIndirectTranslucentVolumeLighting(RHICmdList);
