* fixed experimental r.SSSS (Screen Space Subsurface Scattering feature)

diff --git a/Engine/Source/Runtime/Renderer/Private/PostProcess/PostProcessSubsurface.cpp b/Engine/Source/Runtime/Renderer/Private/PostProcess/PostProcessSubsurface.cpp
--- a/Engine/Source/Runtime/Renderer/Private/PostProcess/PostProcessSubsurface.cpp
+++ b/Engine/Source/Runtime/Renderer/Private/PostProcess/PostProcessSubsurface.cpp
@@ -385,6 +385,9 @@ void FRCPassPostProcessSubsurface::Process(FRenderingCompositePassContext& Conte
 			EDRF_UseTriangleOptimization);
 
 		Context.RHICmdList.CopyToResolveTarget(DestRenderTarget.TargetableTexture, DestRenderTarget.ShaderResourceTexture, false, FResolveParams());
+
+		// matching AdjustGBufferRefCount(1) call is PostProcessing.cpp
+		GSceneRenderTargets.AdjustGBufferRefCount(-1);
 	}
 }
 
diff --git a/Engine/Source/Runtime/Renderer/Private/PostProcess/PostProcessing.cpp b/Engine/Source/Runtime/Renderer/Private/PostProcess/PostProcessing.cpp
--- a/Engine/Source/Runtime/Renderer/Private/PostProcess/PostProcessing.cpp
+++ b/Engine/Source/Runtime/Renderer/Private/PostProcess/PostProcessing.cpp
@@ -812,6 +812,9 @@ void FPostProcessing::Process(FRHICommandListImmediate& RHICmdList, const FViewI
 
 				if(Radius > 0 && !bSimpleDynamicLighting)
 				{
+					// matching AdjustGBufferRefCount(-1) call is in FRCPassPostProcessSubsurface::Process()
+					GSceneRenderTargets.AdjustGBufferRefCount(1);
+
 					FRenderingCompositePass* PassSetup = Context.Graph.RegisterPass(new(FMemStack::Get()) FRCPassPostProcessSubsurfaceSetup());
 					PassSetup->SetInput(ePId_Input0, Context.FinalOutput);
 
