UE4 - fix distortion + shader complexity crash

diff --git a/Engine/Source/Runtime/Renderer/Private/DistortionRendering.cpp b/Engine/Source/Runtime/Renderer/Private/DistortionRendering.cpp
--- a/Engine/Source/Runtime/Renderer/Private/DistortionRendering.cpp
+++ b/Engine/Source/Runtime/Renderer/Private/DistortionRendering.cpp
@@ -426,21 +426,18 @@ void TDistortionMeshDrawingPolicy<DistortMeshPolicy>::DrawShared(
 
 	if (bOverrideWithShaderComplexity)
 	{
-		checkSlow(!bInitializeOffsets);
+		check(!bInitializeOffsets);
 //later		TShaderMapRef<FShaderComplexityAccumulatePixelShader> ShaderComplexityPixelShader(GetGlobalShaderMap());
 		//don't add any vertex complexity
 //later		ShaderComplexityPixelShader->SetParameters(0, DistortPixelShader->GetNumInstructions());
 	}
-	else
+	if (bInitializeOffsets)
 	{
-		if (bInitializeOffsets)
-		{
 //later			InitializePixelShader->SetParameters(0, 0);
-		}
-		else
-		{
-			DistortPixelShader->SetParameters(RHICmdList, MaterialRenderProxy,*View);
-		}
+	}
+	else
+	{
+		DistortPixelShader->SetParameters(RHICmdList, MaterialRenderProxy,*View);
 	}
 }
 
@@ -456,7 +453,7 @@ FBoundShaderStateRHIRef TDistortionMeshDrawingPolicy<DistortMeshPolicy>::CreateB
 
 	if (bOverrideWithShaderComplexity)
 	{
-		checkSlow(!bInitializeOffsets);
+		check(!bInitializeOffsets);
 //later		TShaderMapRef<FShaderComplexityAccumulatePixelShader> ShaderComplexityAccumulatePixelShader(GetGlobalShaderMap());
 //later		PixelShaderRHIRef = ShaderComplexityAccumulatePixelShader->GetPixelShader();
 	}
