Fixed crash and light functions.

diff --git a/Engine/Shaders/DeferredLightingCommon.usf b/Engine/Shaders/DeferredLightingCommon.usf
--- a/Engine/Shaders/DeferredLightingCommon.usf
+++ b/Engine/Shaders/DeferredLightingCommon.usf
@@ -574,6 +574,10 @@ float4 GetDynamicLighting(float3 WorldPosition, float3 CameraVector, float2 InUV
 			// greatly reduces shadow mapping artifacts
 			SurfaceShadow *= saturate(dot(N, L) * 6 - 0.2);
 		}
+		else
+		{
+			SurfaceShadow = ScreenSpaceData.AmbientOcclusion;
+		}
 
 		float SurfaceAttenuation	= (DistanceAttenuation * LightRadiusMask * SpotFalloff) * SurfaceShadow;
 		float SubsurfaceAttenuation	= (DistanceAttenuation * LightRadiusMask * SpotFalloff) * SubsurfaceShadow;
diff --git a/Engine/Source/Runtime/Renderer/Private/LightRendering.cpp b/Engine/Source/Runtime/Renderer/Private/LightRendering.cpp
--- a/Engine/Source/Runtime/Renderer/Private/LightRendering.cpp
+++ b/Engine/Source/Runtime/Renderer/Private/LightRendering.cpp
@@ -465,6 +465,7 @@ void FDeferredShadingSceneRenderer::RenderLights(FRHICommandListImmediate& RHICm
 			bool bDrawShadows = SortedLightInfo.SortKey.Fields.bShadowed;
 			bool bDrawLightFunction = SortedLightInfo.SortKey.Fields.bLightFunction;
 			bool bInjectedTranslucentVolume = false;
+			bool bUsedLightAttenuation = false;
 			FScopeCycleCounter Context(LightSceneInfo.Proxy->GetStatId());
 
 			FString LightNameWithLevel;
@@ -485,6 +486,8 @@ void FDeferredShadingSceneRenderer::RenderLights(FRHICommandListImmediate& RHICm
 				bool bRenderedTranslucentObjectShadows = RenderTranslucentProjectedShadows(RHICmdList, &LightSceneInfo );
 				// Render non-modulated projected shadows to the attenuation buffer.
 				RenderProjectedShadows(RHICmdList, &LightSceneInfo, bRenderedTranslucentObjectShadows, bInjectedTranslucentVolume );
+				
+				bUsedLightAttenuation = true;
 			}
 
 			// Render any reflective shadow maps (if necessary)
@@ -501,13 +504,13 @@ void FDeferredShadingSceneRenderer::RenderLights(FRHICommandListImmediate& RHICm
 			if (bDirectLighting)
 			{
 				const bool bLightFunctionRendered = RenderLightFunction(RHICmdList, &LightSceneInfo, bDrawShadows);
+				bUsedLightAttenuation |= bLightFunctionRendered;
 
 				if (ViewFamily.EngineShowFlags.PreviewShadowsIndicator
 					&& !LightSceneInfo.bPrecomputedLightingIsValid 
 					&& LightSceneInfo.Proxy->HasStaticShadowing())
 				{
-					const bool bLightAttenuationCleared = bDrawShadows || bLightFunctionRendered;
-					RenderPreviewShadowsIndicator(RHICmdList, &LightSceneInfo, bLightAttenuationCleared);
+					RenderPreviewShadowsIndicator(RHICmdList, &LightSceneInfo, bUsedLightAttenuation);
 				}
 
 				if (!bDrawShadows)
@@ -515,9 +518,12 @@ void FDeferredShadingSceneRenderer::RenderLights(FRHICommandListImmediate& RHICm
 					INC_DWORD_STAT(STAT_NumLightFunctionOnlyLights);
 				}
 			}
-				
-			// Resolve light attenuation buffer
-			GSceneRenderTargets.FinishRenderingLightAttenuation(RHICmdList);
+			
+			if( bUsedLightAttenuation )
+			{
+				// Resolve light attenuation buffer
+				GSceneRenderTargets.FinishRenderingLightAttenuation(RHICmdList);
+			}
 			
 			if(bDirectLighting && !bInjectedTranslucentVolume)
 			{
@@ -526,7 +532,7 @@ void FDeferredShadingSceneRenderer::RenderLights(FRHICommandListImmediate& RHICm
 				InjectTranslucentVolumeLighting(RHICmdList, LightSceneInfo, NULL);
 			}
 
-			GSceneRenderTargets.SetLightAttenuationMode(true);
+			GSceneRenderTargets.SetLightAttenuationMode(bUsedLightAttenuation);
 			GSceneRenderTargets.BeginRenderingSceneColor(RHICmdList);
 
 			// Render the light to the scene color buffer, conditionally using the attenuation buffer or a 1x1 white texture as input 
diff --git a/Engine/Source/Runtime/Renderer/Private/LightRendering.h b/Engine/Source/Runtime/Renderer/Private/LightRendering.h
--- a/Engine/Source/Runtime/Renderer/Private/LightRendering.h
+++ b/Engine/Source/Runtime/Renderer/Private/LightRendering.h
@@ -76,7 +76,8 @@ void SetDeferredLightParameters(
 		ShadowMapChannel == 2 ? 1 : 0,
 		ShadowMapChannel == 3 ? 1 : 0);
 
-	DeferredLightUniformsValue.bShadowed = LightSceneInfo->Proxy->CastsDynamicShadow() || LightSceneInfo->Proxy->CastsStaticShadow();
+	const bool bHasLightFunction = LightSceneInfo->Proxy->GetLightFunctionMaterial() != NULL;
+	DeferredLightUniformsValue.bShadowed = LightSceneInfo->Proxy->CastsDynamicShadow() || LightSceneInfo->Proxy->CastsStaticShadow() || bHasLightFunction;
 
 	if( LightSceneInfo->Proxy->IsInverseSquared() )
 	{
