* added comment

diff --git a/Engine/Source/Runtime/Renderer/Private/ForwardShadingRenderer.cpp b/Engine/Source/Runtime/Renderer/Private/ForwardShadingRenderer.cpp
--- a/Engine/Source/Runtime/Renderer/Private/ForwardShadingRenderer.cpp
+++ b/Engine/Source/Runtime/Renderer/Private/ForwardShadingRenderer.cpp
@@ -99,6 +99,10 @@ void FForwardShadingSceneRenderer::Render(FRHICommandListImmediate& RHICmdList)
 	{
 		SCOPE_CYCLE_COUNTER(STAT_TranslucencyDrawTime);
 
+		// Note: Forward pass has no SeparateTranslucency, so refraction effect order with Transluency is different.
+		// Having the distortion applied between two different translucency passes would make it consistent with the deferred pass.
+		// This is not done yet.
+
 		if (ViewFamily.EngineShowFlags.Refraction)
 		{
 			// to apply refraction effect by distorting the scene color
