UE4 - More RHI command list work. Remove more cases of GetImmediateCommandList

diff --git a/Engine/Source/Runtime/Engine/Classes/Engine/CanvasRenderTarget2D.h b/Engine/Source/Runtime/Engine/Classes/Engine/CanvasRenderTarget2D.h
--- a/Engine/Source/Runtime/Engine/Classes/Engine/CanvasRenderTarget2D.h
+++ b/Engine/Source/Runtime/Engine/Classes/Engine/CanvasRenderTarget2D.h
@@ -32,7 +32,7 @@ public:
 	 * texture to make it available for rendering.
 	 */
 	UFUNCTION(BlueprintCallable, Category="Canvas Render Target 2D")
-	virtual void UpdateResource();
+	virtual void UpdateResource() override;
 
 	/**
 	 * Creates a new canvas render target and initializes it to the specified dimensions
diff --git a/Engine/Source/Runtime/Engine/Private/Slate/SlateTextures.cpp b/Engine/Source/Runtime/Engine/Private/Slate/SlateTextures.cpp
--- a/Engine/Source/Runtime/Engine/Private/Slate/SlateTextures.cpp
+++ b/Engine/Source/Runtime/Engine/Private/Slate/SlateTextures.cpp
@@ -187,7 +187,8 @@ void FSlateTextureRenderTarget2DResource::InitDynamicRHI()
 		TextureRHI = (FTextureRHIRef&)Texture2DRHI;
 
 		// make sure the texture target gets cleared
-		UpdateResource();
+		FRHICommandListImmediate& RHICmdList = FRHICommandListExecutor::GetImmediateCommandList();
+		UpdateDeferredResource(RHICmdList);
 	}
 
 	// Create the sampler state RHI resource.
@@ -215,10 +216,9 @@ void FSlateTextureRenderTarget2DResource::ReleaseDynamicRHI()
 	RemoveFromDeferredUpdateList();
 }
 
-void FSlateTextureRenderTarget2DResource::UpdateResource()
+void FSlateTextureRenderTarget2DResource::UpdateDeferredResource(FRHICommandListImmediate& RHICmdList)
 {
 	check(IsInRenderingThread());
-	FRHICommandListImmediate& RHICmdList = FRHICommandListExecutor::GetImmediateCommandList();
 
 	// clear the target surface to green
 	SetRenderTarget(RHICmdList, RenderTargetTextureRHI,FTextureRHIRef());
diff --git a/Engine/Source/Runtime/Engine/Private/TextureMovie.cpp b/Engine/Source/Runtime/Engine/Private/TextureMovie.cpp
--- a/Engine/Source/Runtime/Engine/Private/TextureMovie.cpp
+++ b/Engine/Source/Runtime/Engine/Private/TextureMovie.cpp
@@ -413,12 +413,11 @@ void FTextureMovieResource::ReleaseDynamicRHI()
 * Decodes the next frame of the movie stream and renders the result to this movie texture target
 * This is only called by the rendering thread.
 */
-void FTextureMovieResource::UpdateResource()
+void FTextureMovieResource::UpdateDeferredResource(FRHICommandListImmediate& RHICmdList)
 {
 	if( Owner->Decoder )
 	{
 		SCOPED_DRAW_EVENTF(EventDecode, DEC_SCENE_ITEMS, TEXT("Movie[%s]"),*Owner->GetName());
-		FRHICommandListImmediate& RHICmdList = FRHICommandListExecutor::GetImmediateCommandList();
 		Owner->Decoder->GetFrame(RHICmdList, this);
 	}
 }
diff --git a/Engine/Source/Runtime/Engine/Private/TextureRenderTarget.cpp b/Engine/Source/Runtime/Engine/Private/TextureRenderTarget.cpp
--- a/Engine/Source/Runtime/Engine/Private/TextureRenderTarget.cpp
+++ b/Engine/Source/Runtime/Engine/Private/TextureRenderTarget.cpp
@@ -110,7 +110,7 @@ TLinkedList<FDeferredUpdateResource*>*& FDeferredUpdateResource::GetUpdateList()
 * Iterate over the global list of resources that need to
 * be updated and call UpdateResource on each one.
 */
-void FDeferredUpdateResource::UpdateResources()
+void FDeferredUpdateResource::UpdateResources(FRHICommandListImmediate& RHICmdList)
 {
 	if( bNeedsUpdate )
 	{
@@ -123,7 +123,7 @@ void FDeferredUpdateResource::UpdateResources()
 			if( RTResource )
 			{
 				// update each resource
-				RTResource->UpdateResource();
+				RTResource->UpdateDeferredResource(RHICmdList);
 				if( RTResource->bOnlyUpdateOnce )
 				{
 					// remove from list if only a single update was requested
diff --git a/Engine/Source/Runtime/Engine/Private/TextureRenderTarget2D.cpp b/Engine/Source/Runtime/Engine/Private/TextureRenderTarget2D.cpp
--- a/Engine/Source/Runtime/Engine/Private/TextureRenderTarget2D.cpp
+++ b/Engine/Source/Runtime/Engine/Private/TextureRenderTarget2D.cpp
@@ -86,7 +86,7 @@ void UTextureRenderTarget2D::UpdateResourceImmediate()
 		UpdateResourceImmediate,
 		FRenderResource*,Resource,Resource,
 		{
-			static_cast< FTextureRenderTarget2DResource* >( Resource )->UpdateResource();
+			static_cast<FTextureRenderTarget2DResource*>(Resource)->UpdateDeferredResource(RHICmdList);
 		}
 	);
 }
@@ -363,10 +363,9 @@ void FTextureRenderTarget2DResource::ReleaseDynamicRHI()
 /**
  * Clear contents of the render target. 
  */
-void FTextureRenderTarget2DResource::UpdateResource()
+void FTextureRenderTarget2DResource::UpdateDeferredResource(FRHICommandListImmediate& RHICmdList)
 {
 	RemoveFromDeferredUpdateList();
-	FRHICommandListImmediate& RHICmdList = FRHICommandListExecutor::GetImmediateCommandList();
 
  	// clear the target surface to green
 	SetRenderTarget(RHICmdList, RenderTargetTextureRHI, FTextureRHIRef());
diff --git a/Engine/Source/Runtime/Engine/Private/TextureRenderTargetCube.cpp b/Engine/Source/Runtime/Engine/Private/TextureRenderTargetCube.cpp
--- a/Engine/Source/Runtime/Engine/Private/TextureRenderTargetCube.cpp
+++ b/Engine/Source/Runtime/Engine/Private/TextureRenderTargetCube.cpp
@@ -279,10 +279,8 @@ void FTextureRenderTargetCubeResource::ReleaseDynamicRHI()
 /**
  * Clear contents of the render target. Clears each face of the cube
  */
-void FTextureRenderTargetCubeResource::UpdateResource()
+void FTextureRenderTargetCubeResource::UpdateDeferredResource(FRHICommandListImmediate& RHICmdList)
 {
-	FRHICommandListImmediate& RHICmdList = FRHICommandListExecutor::GetImmediateCommandList();
-
 	const FIntPoint Dims = GetSizeXY();
 	for(int32 FaceIdx = CubeFace_PosX; FaceIdx < CubeFace_MAX; FaceIdx++)
 	{
diff --git a/Engine/Source/Runtime/Engine/Public/Slate/SlateTextures.h b/Engine/Source/Runtime/Engine/Public/Slate/SlateTextures.h
--- a/Engine/Source/Runtime/Engine/Public/Slate/SlateTextures.h
+++ b/Engine/Source/Runtime/Engine/Public/Slate/SlateTextures.h
@@ -160,13 +160,16 @@ public:
 	virtual void InitDynamicRHI() override;
 	virtual void ReleaseDynamicRHI() override;
 
-	// FDeferredUpdateResource implementation
-	virtual void UpdateResource() override;
-
 	// FRenderTarget interface
 	virtual FIntPoint GetSizeXY() const override;
 	virtual float GetDisplayGamma() const override;
 
+protected:
+
+	// FDeferredUpdateResource implementation
+	virtual void UpdateDeferredResource(FRHICommandListImmediate& RHICmdList) override;
+
+
 private:
 	FTexture2DRHIRef Texture2DRHI;
 	FLinearColor ClearColor;
diff --git a/Engine/Source/Runtime/Engine/Public/Texture.h b/Engine/Source/Runtime/Engine/Public/Texture.h
--- a/Engine/Source/Runtime/Engine/Public/Texture.h
+++ b/Engine/Source/Runtime/Engine/Public/Texture.h
@@ -558,7 +558,7 @@ public:
 	 * Iterate over the global list of resources that need to
 	 * be updated and call UpdateResource on each one.
 	 */
-	ENGINE_API static void UpdateResources();
+	ENGINE_API static void UpdateResources(FRHICommandListImmediate& RHICmdList);
 
 	/** 
 	 * This is reset after all viewports have been rendered
@@ -570,13 +570,13 @@ public:
 
 	// FDeferredUpdateResource interface
 
-	/**
-	 * Updates the resource
-	 */
-	virtual void UpdateResource() = 0;
-
 protected:
 
+	/**
+	* Updates the resource
+	*/
+	virtual void UpdateDeferredResource(FRHICommandListImmediate& RHICmdList) = 0;
+
 	/**
 	 * Add this resource to deferred update list
 	 * @param OnlyUpdateOnce - flag this resource for a single update if true
@@ -695,11 +695,6 @@ public:
 
 	// FDeferredClearResource interface
 
-	/**
-	 * Clear contents of the render target
-	 */
-	virtual void UpdateResource() override;
-
 	// FRenderTarget interface.
 	virtual FIntPoint GetSizeXY() const;
 
@@ -714,6 +709,12 @@ public:
 	 * @return TextureRHI for rendering 
 	 */
 	FTexture2DRHIRef GetTextureRHI() { return Texture2DRHI; }
+protected:
+	/**
+	* Clear contents of the render target
+	*/
+	friend class UTextureRenderTarget2D;
+	virtual void UpdateDeferredResource(FRHICommandListImmediate& RHICmdList) override;
 
 private:
 	/** The UTextureRenderTarget2D which this resource represents. */
@@ -766,12 +767,6 @@ public:
 	 */
 	virtual void ReleaseDynamicRHI() override;
 
-	/**
-	 * Clear contents of the render target. Clears each face of the cube
-	 * This is only called by the rendering thread.
-	 */
-	virtual void UpdateResource() override;
-
 	// FRenderTarget interface.
 
 	/**
@@ -809,6 +804,13 @@ public:
 	*/
 	ENGINE_API bool ReadPixels(TArray<FFloat16Color>& OutImageData, FReadSurfaceDataFlags InFlags, FIntRect InRect = FIntRect(0, 0, 0, 0));
 
+protected:
+	/**
+	* Clear contents of the render target. Clears each face of the cube
+	* This is only called by the rendering thread.
+	*/
+	virtual void UpdateDeferredResource(FRHICommandListImmediate& RHICmdList) override;
+
 private:
 	/** The UTextureRenderTargetCube which this resource represents. */
 	const class UTextureRenderTargetCube* Owner;
@@ -860,14 +862,17 @@ public:
 
 	// FDeferredClearResource interface
 
-	/**
-	 * Decodes the next frame of the movie stream and renders the result to this movie texture target
-	 */
-	virtual void UpdateResource();
-
 	// FRenderTarget interface.
 	virtual FIntPoint GetSizeXY() const;
 
+protected:
+
+	/**
+	* Decodes the next frame of the movie stream and renders the result to this movie texture target
+	*/
+	virtual void UpdateDeferredResource(FRHICommandListImmediate& RHICmdList) override;
+
+
 private:
 	/** The UTextureRenderTarget2D which this resource represents. */
 	const class UTextureMovie* Owner;
diff --git a/Engine/Source/Runtime/Renderer/Private/CustomDepthRendering.cpp b/Engine/Source/Runtime/Renderer/Private/CustomDepthRendering.cpp
--- a/Engine/Source/Runtime/Renderer/Private/CustomDepthRendering.cpp
+++ b/Engine/Source/Runtime/Renderer/Private/CustomDepthRendering.cpp
@@ -16,14 +16,14 @@
 	FCustomDepthPrimSet
 -----------------------------------------------------------------------------*/
 
-bool FCustomDepthPrimSet::DrawPrims(FRHICommandList& RHICmdList, const FViewInfo* ViewInfo,bool bInitializeOffsets)
+bool FCustomDepthPrimSet::DrawPrims(FRHICommandListImmediate& RHICmdList, const FViewInfo* ViewInfo, bool bInitializeOffsets)
 {
 	bool bDirty=false;
 
 	if(Prims.Num())
 	{
 		SCOPED_DRAW_EVENT(DistortionAccum, DEC_SCENE_ITEMS);
-		TDynamicPrimitiveDrawer<FDepthDrawingPolicyFactory> Drawer(ViewInfo, FDepthDrawingPolicyFactory::ContextType(DDM_AllOccluders), true);
+		TDynamicPrimitiveDrawer<FDepthDrawingPolicyFactory> Drawer(RHICmdList, ViewInfo, FDepthDrawingPolicyFactory::ContextType(DDM_AllOccluders), true);
 
 		for (int32 PrimIdx = 0; PrimIdx < Prims.Num(); PrimIdx++)
 		{
diff --git a/Engine/Source/Runtime/Renderer/Private/CustomDepthRendering.h b/Engine/Source/Runtime/Renderer/Private/CustomDepthRendering.h
--- a/Engine/Source/Runtime/Renderer/Private/CustomDepthRendering.h
+++ b/Engine/Source/Runtime/Renderer/Private/CustomDepthRendering.h
@@ -18,7 +18,7 @@ public:
 	* @param ViewInfo - current view used to draw items
 	* @return true if anything was drawn
 	*/
-	bool DrawPrims(FRHICommandList& RHICmdList, const class FViewInfo* ViewInfo,bool bInitializeOffsets);
+	bool DrawPrims(FRHICommandListImmediate& RHICmdList, const class FViewInfo* ViewInfo, bool bInitializeOffsets);
 
 	/**
 	* Add a new primitive to the list of prims
diff --git a/Engine/Source/Runtime/Renderer/Private/DeferredShadingRenderer.cpp b/Engine/Source/Runtime/Renderer/Private/DeferredShadingRenderer.cpp
--- a/Engine/Source/Runtime/Renderer/Private/DeferredShadingRenderer.cpp
+++ b/Engine/Source/Runtime/Renderer/Private/DeferredShadingRenderer.cpp
@@ -307,13 +307,13 @@ bool FDeferredShadingSceneRenderer::RenderBasePassDynamicData(FRHICommandListImm
 		const bool bNeedToSwitchVerticalAxis = IsES2Platform(GRHIShaderPlatform) && !IsPCPlatform(GRHIShaderPlatform);
 
 		// Draw the base pass for the view's batched mesh elements.
-		bDirty = DrawViewElements<FBasePassOpaqueDrawingPolicyFactory>(View,FBasePassOpaqueDrawingPolicyFactory::ContextType(false, ESceneRenderTargetsMode::DontSet), SDPG_World, true) || bDirty;
+		bDirty = DrawViewElements<FBasePassOpaqueDrawingPolicyFactory>(RHICmdList, View, FBasePassOpaqueDrawingPolicyFactory::ContextType(false, ESceneRenderTargetsMode::DontSet), SDPG_World, true) || bDirty;
 
 		// Draw the view's batched simple elements(lines, sprites, etc).
 		bDirty = View.BatchedViewElements.Draw(RHICmdList, bNeedToSwitchVerticalAxis, View.ViewProjectionMatrix, View.ViewRect.Width(), View.ViewRect.Height(), false) || bDirty;
 		
 		// Draw foreground objects last
-		bDirty = DrawViewElements<FBasePassOpaqueDrawingPolicyFactory>(View,FBasePassOpaqueDrawingPolicyFactory::ContextType(false, ESceneRenderTargetsMode::DontSet), SDPG_Foreground, true) || bDirty;
+		bDirty = DrawViewElements<FBasePassOpaqueDrawingPolicyFactory>(RHICmdList, View, FBasePassOpaqueDrawingPolicyFactory::ContextType(false, ESceneRenderTargetsMode::DontSet), SDPG_Foreground, true) || bDirty;
 
 		// Draw the view's batched simple elements(lines, sprites, etc).
 		bDirty = View.TopBatchedViewElements.Draw(RHICmdList, bNeedToSwitchVerticalAxis, View.ViewProjectionMatrix, View.ViewRect.Width(), View.ViewRect.Height(), false) || bDirty;
@@ -351,7 +351,7 @@ bool FDeferredShadingSceneRenderer::RenderBasePass(FRHICommandListImmediate& RHI
 		{
 			// Draw the dynamic non-occluded primitives using a base pass drawing policy.
 			TDynamicPrimitiveDrawer<FBasePassOpaqueDrawingPolicyFactory> Drawer(
-				&View,FBasePassOpaqueDrawingPolicyFactory::ContextType(false, ESceneRenderTargetsMode::DontSet),true);
+				RHICmdList, &View, FBasePassOpaqueDrawingPolicyFactory::ContextType(false, ESceneRenderTargetsMode::DontSet), true);
 			for(int32 PrimitiveIndex = 0, Num = View.VisibleDynamicPrimitives.Num();PrimitiveIndex < Num;PrimitiveIndex++)
 			{
 				const FPrimitiveSceneInfo* PrimitiveSceneInfo = View.VisibleDynamicPrimitives[PrimitiveIndex];
@@ -854,7 +854,7 @@ bool FDeferredShadingSceneRenderer::RenderPrePass(FRHICommandListImmediate& RHIC
 			}
 
 			// Draw the dynamic occluder primitives using a depth drawing policy.
-			TDynamicPrimitiveDrawer<FDepthDrawingPolicyFactory> Drawer(&View,FDepthDrawingPolicyFactory::ContextType(EarlyZPassMode),true);
+			TDynamicPrimitiveDrawer<FDepthDrawingPolicyFactory> Drawer(RHICmdList, &View, FDepthDrawingPolicyFactory::ContextType(EarlyZPassMode), true);
 			{
 				SCOPED_DRAW_EVENT(Dynamic, DEC_SCENE_ITEMS);
 				for(int32 PrimitiveIndex = 0;PrimitiveIndex < View.VisibleDynamicPrimitives.Num();PrimitiveIndex++)
diff --git a/Engine/Source/Runtime/Renderer/Private/DistortionRendering.cpp b/Engine/Source/Runtime/Renderer/Private/DistortionRendering.cpp
--- a/Engine/Source/Runtime/Renderer/Private/DistortionRendering.cpp
+++ b/Engine/Source/Runtime/Renderer/Private/DistortionRendering.cpp
@@ -688,6 +688,7 @@ bool FDistortionPrimSet::DrawAccumulatedOffsets(FRHICommandListImmediate& RHICmd
 
 	// Draw the view's elements with the translucent drawing policy.
 	bDirty |= DrawViewElements<TDistortionMeshDrawingPolicyFactory<FDistortMeshAccumulatePolicy> >(
+		RHICmdList,
 		*ViewInfo,
 		bInitializeOffsets,
 		0,	// DPG Index?
@@ -698,6 +699,7 @@ bool FDistortionPrimSet::DrawAccumulatedOffsets(FRHICommandListImmediate& RHICmd
 	{
 		// For drawing scene prims with dynamic relevance.
 		TDynamicPrimitiveDrawer<TDistortionMeshDrawingPolicyFactory<FDistortMeshAccumulatePolicy> > Drawer(
+			RHICmdList,
 			ViewInfo,
 			bInitializeOffsets,
 			false // Distortion is rendered post fog.
diff --git a/Engine/Source/Runtime/Renderer/Private/DynamicPrimitiveDrawing.h b/Engine/Source/Runtime/Renderer/Private/DynamicPrimitiveDrawing.h
--- a/Engine/Source/Runtime/Renderer/Private/DynamicPrimitiveDrawing.h
+++ b/Engine/Source/Runtime/Renderer/Private/DynamicPrimitiveDrawing.h
@@ -25,6 +25,7 @@ public:
 	* @param bInIsVelocityRendering - rendering is occurring during velocity pass
 	*/
 	TDynamicPrimitiveDrawer(
+		FRHICommandListImmediate& InRHICmdList,
 		const FViewInfo* InView,
 		const typename DrawingPolicyFactoryType::ContextType& InDrawingContext,
 		bool InPreFog,
@@ -34,6 +35,7 @@ public:
 		bool bInIsSelectionOutlineRendering = false
 		):
 		FPrimitiveDrawInterface(InView),
+		RHICmdList(InRHICmdList),
 		View(InView),
 		DrawingContext(InDrawingContext),
 		PrimitiveSceneProxy(NULL),
@@ -106,6 +108,10 @@ public:
 	virtual bool IsRenderingVelocities() const { return bIsVelocityRendering; }
 
 private:
+
+	/** Command list to draw to. */
+	FRHICommandListImmediate& RHICmdList;
+
 	/** The view which is being rendered. */
 	const FViewInfo* const View;
 
@@ -141,6 +147,7 @@ private:
 
 	/** true if we are currently rendering the selection outline */
 	uint32 bIsSelectionOutlineRendering : 1;
+
 };
 
 /**
@@ -152,6 +159,7 @@ private:
  */
 template<class DrawingPolicyFactoryType>
 bool DrawViewElements(
+	FRHICommandListImmediate& RHICmdList,
 	const FViewInfo& View,
 	const typename DrawingPolicyFactoryType::ContextType& DrawingContext,
 	uint8 DPGIndex,
@@ -166,6 +174,7 @@ bool DrawViewElements(
  */
 template<class DrawingPolicyFactoryType>
 bool DrawDynamicPrimitiveSet(
+	FRHICommandListImmediate& RHICmdList, 
 	const FViewInfo& View,
 	const TArray<const FPrimitiveSceneInfo*,SceneRenderingAllocator>& PrimitiveSet,
 	const typename DrawingPolicyFactoryType::ContextType& DrawingContext,
diff --git a/Engine/Source/Runtime/Renderer/Private/DynamicPrimitiveDrawing.inl b/Engine/Source/Runtime/Renderer/Private/DynamicPrimitiveDrawing.inl
--- a/Engine/Source/Runtime/Renderer/Private/DynamicPrimitiveDrawing.inl
+++ b/Engine/Source/Runtime/Renderer/Private/DynamicPrimitiveDrawing.inl
@@ -23,8 +23,6 @@ TDynamicPrimitiveDrawer<DrawingPolicyFactoryType>::~TDynamicPrimitiveDrawer()
 			DepthTexture = GSceneRenderTargets.GetSceneDepthTexture();
 		}
 		
-		FRHICommandListImmediate& RHICmdList = FRHICommandListExecutor::GetImmediateCommandList();
-
 		// Draw the batched elements.
 		BatchedElements.Draw(
 			RHICmdList,
@@ -108,7 +106,6 @@ int32 TDynamicPrimitiveDrawer<DrawingPolicyFactoryType>::DrawMesh(const FMeshBat
 	Mesh.CheckUniformBuffers();
 #endif
 	
-	FRHICommandListImmediate& RHICmdList = FRHICommandListExecutor::GetImmediateCommandList();
 	INC_DWORD_STAT_BY(STAT_DynamicPathMeshDrawCalls,Mesh.Elements.Num());
 	const bool DrawDirty = DrawingPolicyFactoryType::DrawDynamicMesh(
 		RHICmdList,
@@ -217,15 +214,13 @@ void TDynamicPrimitiveDrawer<DrawingPolicyFactoryType>::DrawPoint(
 
 template<class DrawingPolicyFactoryType>
 bool DrawViewElements(
+	FRHICommandListImmediate& RHICmdList,
 	const FViewInfo& View,
 	const typename DrawingPolicyFactoryType::ContextType& DrawingContext,
 	uint8 DPGIndex,
 	bool bPreFog
 	)
 {
-	
-	FRHICommandListImmediate& RHICmdList = FRHICommandListExecutor::GetImmediateCommandList();
-
 	// Get the correct element list based on dpg index
 	const TIndirectArray<FHitProxyMeshPair>& ViewMeshElementList = ( DPGIndex == SDPG_Foreground ? View.TopViewMeshElements : View.ViewMeshElements );
 	// Draw the view's mesh elements.
@@ -258,6 +253,7 @@ bool DrawViewElements(
 
 template<class DrawingPolicyFactoryType>
 bool DrawDynamicPrimitiveSet(
+	FRHICommandListImmediate& RHICmdList,
 	const FViewInfo& View,
 	const TArray<const FPrimitiveSceneInfo*,SceneRenderingAllocator>& PrimitiveSet,
 	const typename DrawingPolicyFactoryType::ContextType& DrawingContext,
@@ -267,7 +263,7 @@ bool DrawDynamicPrimitiveSet(
 {
 	
 	// Draw the elements of each dynamic primitive.
-	TDynamicPrimitiveDrawer<DrawingPolicyFactoryType> Drawer(&View, DrawingContext, bPreFog, bIsHitTesting);
+	TDynamicPrimitiveDrawer<DrawingPolicyFactoryType> Drawer(RHICmdList, &View, DrawingContext, bPreFog, bIsHitTesting);
 	for(int32 PrimitiveIndex = 0; PrimitiveIndex < PrimitiveSet.Num(); PrimitiveIndex++)
 	{
 		const FPrimitiveSceneInfo* PrimitiveSceneInfo = PrimitiveSet[PrimitiveIndex];
diff --git a/Engine/Source/Runtime/Renderer/Private/ForwardBasePassRendering.cpp b/Engine/Source/Runtime/Renderer/Private/ForwardBasePassRendering.cpp
--- a/Engine/Source/Runtime/Renderer/Private/ForwardBasePassRendering.cpp
+++ b/Engine/Source/Runtime/Renderer/Private/ForwardBasePassRendering.cpp
@@ -361,7 +361,7 @@ void FForwardShadingSceneRenderer::RenderForwardShadingBasePass(FRHICommandListI
 			if (View.VisibleDynamicPrimitives.Num() > 0)
 			{
 				// Draw the dynamic non-occluded primitives using a base pass drawing policy.
-				TDynamicPrimitiveDrawer<FBasePassForwardOpaqueDrawingPolicyFactory> Drawer(&View, FBasePassForwardOpaqueDrawingPolicyFactory::ContextType(ESceneRenderTargetsMode::DontSet), true);
+				TDynamicPrimitiveDrawer<FBasePassForwardOpaqueDrawingPolicyFactory> Drawer(RHICmdList, &View, FBasePassForwardOpaqueDrawingPolicyFactory::ContextType(ESceneRenderTargetsMode::DontSet), true);
 
 				for (int32 PrimitiveIndex = 0; PrimitiveIndex < View.VisibleDynamicPrimitives.Num(); PrimitiveIndex++)
 				{
@@ -386,13 +386,13 @@ void FForwardShadingSceneRenderer::RenderForwardShadingBasePass(FRHICommandListI
 			const bool bNeedToSwitchVerticalAxis = IsES2Platform(GRHIShaderPlatform) && !IsPCPlatform(GRHIShaderPlatform);
 
 			// Draw the base pass for the view's batched mesh elements.
-			DrawViewElements<FBasePassForwardOpaqueDrawingPolicyFactory>(View,FBasePassForwardOpaqueDrawingPolicyFactory::ContextType(ESceneRenderTargetsMode::DontSet), SDPG_World, true);
+			DrawViewElements<FBasePassForwardOpaqueDrawingPolicyFactory>(RHICmdList, View, FBasePassForwardOpaqueDrawingPolicyFactory::ContextType(ESceneRenderTargetsMode::DontSet), SDPG_World, true);
 			
 			// Draw the view's batched simple elements(lines, sprites, etc).
 			View.BatchedViewElements.Draw(RHICmdList, bNeedToSwitchVerticalAxis, View.ViewProjectionMatrix, View.ViewRect.Width(), View.ViewRect.Height(), false);
 
 			// Draw foreground objects last
-			DrawViewElements<FBasePassForwardOpaqueDrawingPolicyFactory>(View,FBasePassForwardOpaqueDrawingPolicyFactory::ContextType(ESceneRenderTargetsMode::DontSet), SDPG_Foreground, true);
+			DrawViewElements<FBasePassForwardOpaqueDrawingPolicyFactory>(RHICmdList, View, FBasePassForwardOpaqueDrawingPolicyFactory::ContextType(ESceneRenderTargetsMode::DontSet), SDPG_Foreground, true);
 
 			// Draw the view's batched simple elements(lines, sprites, etc).
 			View.TopBatchedViewElements.Draw(RHICmdList, bNeedToSwitchVerticalAxis, View.ViewProjectionMatrix, View.ViewRect.Width(), View.ViewRect.Height(), false);
diff --git a/Engine/Source/Runtime/Renderer/Private/ForwardTranslucentRendering.cpp b/Engine/Source/Runtime/Renderer/Private/ForwardTranslucentRendering.cpp
--- a/Engine/Source/Runtime/Renderer/Private/ForwardTranslucentRendering.cpp
+++ b/Engine/Source/Runtime/Renderer/Private/ForwardTranslucentRendering.cpp
@@ -269,6 +269,7 @@ void FTranslucentPrimSet::RenderPrimitiveForForwardShading(
 		if( ViewRelevance.bDynamicRelevance )
 		{
 			TDynamicPrimitiveDrawer<FTranslucencyForwardShadingDrawingPolicyFactory> TranslucencyDrawer(
+				RHICmdList,
 				&View,
 				FTranslucencyForwardShadingDrawingPolicyFactory::ContextType(),
 				false
@@ -343,9 +344,9 @@ void FForwardShadingSceneRenderer::RenderTranslucency(FRHICommandListImmediate&
 			// Draw only translucent prims that don't read from scene color
 			View.TranslucentPrimSet.DrawPrimitivesForForwardShading(RHICmdList, View, *this);
 			// Draw the view's mesh elements with the translucent drawing policy.
-			DrawViewElements<FTranslucencyForwardShadingDrawingPolicyFactory>(View,FTranslucencyForwardShadingDrawingPolicyFactory::ContextType(),SDPG_World,false);
+			DrawViewElements<FTranslucencyForwardShadingDrawingPolicyFactory>(RHICmdList, View, FTranslucencyForwardShadingDrawingPolicyFactory::ContextType(), SDPG_World, false);
 			// Draw the view's mesh elements with the translucent drawing policy.
-			DrawViewElements<FTranslucencyForwardShadingDrawingPolicyFactory>(View,FTranslucencyForwardShadingDrawingPolicyFactory::ContextType(),SDPG_Foreground,false);
+			DrawViewElements<FTranslucencyForwardShadingDrawingPolicyFactory>(RHICmdList, View, FTranslucencyForwardShadingDrawingPolicyFactory::ContextType(), SDPG_Foreground, false);
 		}
 	}
 }
diff --git a/Engine/Source/Runtime/Renderer/Private/LightMapDensityRendering.cpp b/Engine/Source/Runtime/Renderer/Private/LightMapDensityRendering.cpp
--- a/Engine/Source/Runtime/Renderer/Private/LightMapDensityRendering.cpp
+++ b/Engine/Source/Runtime/Renderer/Private/LightMapDensityRendering.cpp
@@ -61,7 +61,7 @@ bool FDeferredShadingSceneRenderer::RenderLightMapDensities(FRHICommandListImmed
 				{
 					// Draw the dynamic non-occluded primitives using a base pass drawing policy.
 					TDynamicPrimitiveDrawer<FLightMapDensityDrawingPolicyFactory> Drawer(
-						&View,FLightMapDensityDrawingPolicyFactory::ContextType(),true);
+						RHICmdList, &View, FLightMapDensityDrawingPolicyFactory::ContextType(), true);
 					for (int32 PrimitiveIndex = 0;PrimitiveIndex < View.VisibleDynamicPrimitives.Num();PrimitiveIndex++)
 					{
 						const FPrimitiveSceneInfo* PrimitiveSceneInfo = View.VisibleDynamicPrimitives[PrimitiveIndex];
diff --git a/Engine/Source/Runtime/Renderer/Private/PostProcess/PostProcessCompositeEditorPrimitives.cpp b/Engine/Source/Runtime/Renderer/Private/PostProcess/PostProcessCompositeEditorPrimitives.cpp
--- a/Engine/Source/Runtime/Renderer/Private/PostProcess/PostProcessCompositeEditorPrimitives.cpp
+++ b/Engine/Source/Runtime/Renderer/Private/PostProcess/PostProcessCompositeEditorPrimitives.cpp
@@ -308,10 +308,10 @@ void FRCPassPostProcessCompositeEditorPrimitives::RenderPrimitivesToComposite(FR
 	{
 		// Draw the dynamic non-occluded primitives using a base pass drawing policy.
 		TDynamicPrimitiveDrawer<FBasePassOpaqueDrawingPolicyFactory> DepthTestDrawer(
-			&View,FBasePassOpaqueDrawingPolicyFactory::ContextType(bDepthTest, ESceneRenderTargetsMode::SetTextures),true,false,false,bDepthTest);
+			RHICmdList, &View, FBasePassOpaqueDrawingPolicyFactory::ContextType(bDepthTest, ESceneRenderTargetsMode::SetTextures), true, false, false, bDepthTest);
 
 		TDynamicPrimitiveDrawer<FBasePassOpaqueDrawingPolicyFactory> NoDepthTestDrawer(
-			&View,FBasePassOpaqueDrawingPolicyFactory::ContextType(!bDepthTest, ESceneRenderTargetsMode::SetTextures),true,false,false,!bDepthTest);
+			RHICmdList, &View, FBasePassOpaqueDrawingPolicyFactory::ContextType(!bDepthTest, ESceneRenderTargetsMode::SetTextures), true, false, false, !bDepthTest);
 
 		for(int32 PrimitiveIndex = 0;PrimitiveIndex < View.VisibleEditorPrimitives.Num();PrimitiveIndex++)
 		{
@@ -350,7 +350,7 @@ void FRCPassPostProcessCompositeEditorPrimitives::RenderPrimitivesToComposite(FR
 	const bool bNeedToSwitchVerticalAxis = IsES2Platform(GRHIShaderPlatform) && !IsPCPlatform(GRHIShaderPlatform);
 
 	// Draw the base pass for the view's batched mesh elements.
-	DrawViewElements<FBasePassOpaqueDrawingPolicyFactory>(View,FBasePassOpaqueDrawingPolicyFactory::ContextType(bDepthTest, ESceneRenderTargetsMode::SetTextures), SDPG_World, false);
+	DrawViewElements<FBasePassOpaqueDrawingPolicyFactory>(RHICmdList, View, FBasePassOpaqueDrawingPolicyFactory::ContextType(bDepthTest, ESceneRenderTargetsMode::SetTextures), SDPG_World, false);
 
 	// Draw the view's batched simple elements(lines, sprites, etc).
 	View.BatchedViewElements.Draw(RHICmdList, bNeedToSwitchVerticalAxis, View.ViewProjectionMatrix, View.ViewRect.Width(), View.ViewRect.Height(), false, 1.0f, &View, SceneDepth);
@@ -362,14 +362,14 @@ void FRCPassPostProcessCompositeEditorPrimitives::RenderPrimitivesToComposite(FR
 
 		RHICmdList.SetDepthStencilState(TStaticDepthStencilState<false, CF_Always>::GetRHI());
 
-		DrawViewElements<FBasePassOpaqueDrawingPolicyFactory>(View,FBasePassOpaqueDrawingPolicyFactory::ContextType(bDepthTest, ESceneRenderTargetsMode::SetTextures), SDPG_Foreground, false);
+		DrawViewElements<FBasePassOpaqueDrawingPolicyFactory>(RHICmdList, View, FBasePassOpaqueDrawingPolicyFactory::ContextType(bDepthTest, ESceneRenderTargetsMode::SetTextures), SDPG_Foreground, false);
 
 		View.TopBatchedViewElements.Draw(RHICmdList, bNeedToSwitchVerticalAxis, View.ViewProjectionMatrix, View.ViewRect.Width(), View.ViewRect.Height(), false);
 
 		// Note, this is a reversed Z depth surface, using CF_GreaterEqual.
 		RHICmdList.SetDepthStencilState(TStaticDepthStencilState<true, CF_GreaterEqual>::GetRHI());
 
-		DrawViewElements<FBasePassOpaqueDrawingPolicyFactory>(View, FBasePassOpaqueDrawingPolicyFactory::ContextType(bDepthTest, ESceneRenderTargetsMode::SetTextures), SDPG_Foreground, false);
+		DrawViewElements<FBasePassOpaqueDrawingPolicyFactory>(RHICmdList, View, FBasePassOpaqueDrawingPolicyFactory::ContextType(bDepthTest, ESceneRenderTargetsMode::SetTextures), SDPG_Foreground, false);
 
 		View.TopBatchedViewElements.Draw(RHICmdList, bNeedToSwitchVerticalAxis, View.ViewProjectionMatrix, View.ViewRect.Width(), View.ViewRect.Height(), false);
 	}
diff --git a/Engine/Source/Runtime/Renderer/Private/PostProcess/PostProcessSelectionOutline.cpp b/Engine/Source/Runtime/Renderer/Private/PostProcess/PostProcessSelectionOutline.cpp
--- a/Engine/Source/Runtime/Renderer/Private/PostProcess/PostProcessSelectionOutline.cpp
+++ b/Engine/Source/Runtime/Renderer/Private/PostProcess/PostProcessSelectionOutline.cpp
@@ -41,7 +41,7 @@ void FRCPassPostProcessSelectionOutlineColor::Process(FRenderingCompositePassCon
 
 	if(View.VisibleDynamicPrimitives.Num() && View.Family->EngineShowFlags.Selection)
 	{
-		TDynamicPrimitiveDrawer<FHitProxyDrawingPolicyFactory> Drawer(&View, FHitProxyDrawingPolicyFactory::ContextType(), true, false, false, false, true);
+		TDynamicPrimitiveDrawer<FHitProxyDrawingPolicyFactory> Drawer(Context.RHICmdList, &View, FHitProxyDrawingPolicyFactory::ContextType(), true, false, false, false, true);
 		TMultiMap<FName, const FPrimitiveSceneInfo*> PrimitivesByActor;
 
 		for (int32 PrimitiveIndex = 0; PrimitiveIndex < View.VisibleDynamicPrimitives.Num();PrimitiveIndex++)
diff --git a/Engine/Source/Runtime/Renderer/Private/ReflectionEnvironmentCapture.cpp b/Engine/Source/Runtime/Renderer/Private/ReflectionEnvironmentCapture.cpp
--- a/Engine/Source/Runtime/Renderer/Private/ReflectionEnvironmentCapture.cpp
+++ b/Engine/Source/Runtime/Renderer/Private/ReflectionEnvironmentCapture.cpp
@@ -824,7 +824,7 @@ void CaptureSceneToScratchCubemap(FRHICommandListImmediate& RHICmdList, FSceneRe
 {
 	FMemMark MemStackMark(FMemStack::Get());
 	// update any resources that needed a deferred update
-	FDeferredUpdateResource::UpdateResources();
+	FDeferredUpdateResource::UpdateResources(RHICmdList);
 	
 	{
 		SCOPED_DRAW_EVENT(CubeMapCapture, DEC_SCENE_ITEMS);
diff --git a/Engine/Source/Runtime/Renderer/Private/SceneCaptureRendering.cpp b/Engine/Source/Runtime/Renderer/Private/SceneCaptureRendering.cpp
--- a/Engine/Source/Runtime/Renderer/Private/SceneCaptureRendering.cpp
+++ b/Engine/Source/Runtime/Renderer/Private/SceneCaptureRendering.cpp
@@ -18,7 +18,7 @@ void UpdateSceneCaptureContent_RenderThread(FRHICommandListImmediate& RHICmdList
 	FMemMark MemStackMark(FMemStack::Get());
 
 	// update any resources that needed a deferred update
-	FDeferredUpdateResource::UpdateResources();
+	FDeferredUpdateResource::UpdateResources(RHICmdList);
 
 	{
 #if WANTS_DRAW_MESH_EVENTS
diff --git a/Engine/Source/Runtime/Renderer/Private/SceneHitProxyRendering.cpp b/Engine/Source/Runtime/Renderer/Private/SceneHitProxyRendering.cpp
--- a/Engine/Source/Runtime/Renderer/Private/SceneHitProxyRendering.cpp
+++ b/Engine/Source/Runtime/Renderer/Private/SceneHitProxyRendering.cpp
@@ -407,6 +407,7 @@ void FDeferredShadingSceneRenderer::RenderHitProxies(FRHICommandListImmediate& R
 
 		// Draw all dynamic primitives using the hit proxy drawing policy.
 		DrawDynamicPrimitiveSet<FHitProxyDrawingPolicyFactory>(
+			RHICmdList,
 			View,
 			View.VisibleDynamicPrimitives,
 			FHitProxyDrawingPolicyFactory::ContextType(),
@@ -416,6 +417,7 @@ void FDeferredShadingSceneRenderer::RenderHitProxies(FRHICommandListImmediate& R
 
 		// Draw all visible editor primitives using the hit proxy drawing policy
 		DrawDynamicPrimitiveSet<FHitProxyDrawingPolicyFactory>(
+			RHICmdList,
 			View,
 			View.VisibleEditorPrimitives,
 			FHitProxyDrawingPolicyFactory::ContextType(),
@@ -425,7 +427,7 @@ void FDeferredShadingSceneRenderer::RenderHitProxies(FRHICommandListImmediate& R
 
 
 		// Draw the view's elements.
-		DrawViewElements<FHitProxyDrawingPolicyFactory>(View,FHitProxyDrawingPolicyFactory::ContextType(),SDPG_World,bPreFog);
+		DrawViewElements<FHitProxyDrawingPolicyFactory>(RHICmdList, View, FHitProxyDrawingPolicyFactory::ContextType(), SDPG_World, bPreFog);
 
 		// Draw the view's batched simple elements(lines, sprites, etc).
 		View.BatchedViewElements.Draw(RHICmdList, bNeedToSwitchVerticalAxis, View.ViewProjectionMatrix, View.ViewRect.Width(), View.ViewRect.Height(), true);
@@ -435,14 +437,14 @@ void FDeferredShadingSceneRenderer::RenderHitProxies(FRHICommandListImmediate& R
 		// then again to allows proper occlusion within those elements.
 		RHICmdList.SetDepthStencilState(TStaticDepthStencilState<false, CF_Always>::GetRHI());
 		// Draw the view's foreground elements last.
-		DrawViewElements<FHitProxyDrawingPolicyFactory>(View,FHitProxyDrawingPolicyFactory::ContextType(),SDPG_Foreground,bPreFog);
+		DrawViewElements<FHitProxyDrawingPolicyFactory>(RHICmdList, View, FHitProxyDrawingPolicyFactory::ContextType(), SDPG_Foreground, bPreFog);
 
 		View.TopBatchedViewElements.Draw(RHICmdList, bNeedToSwitchVerticalAxis, View.ViewProjectionMatrix, View.ViewRect.Width(), View.ViewRect.Height(), true);
 
 		// Note, this is a reversed Z depth surface, using CF_GreaterEqual.
 		RHICmdList.SetDepthStencilState(TStaticDepthStencilState<true, CF_GreaterEqual>::GetRHI());
 		// Draw the view's foreground elements last.
-		DrawViewElements<FHitProxyDrawingPolicyFactory>(View,FHitProxyDrawingPolicyFactory::ContextType(),SDPG_Foreground,bPreFog);
+		DrawViewElements<FHitProxyDrawingPolicyFactory>(RHICmdList, View, FHitProxyDrawingPolicyFactory::ContextType(), SDPG_Foreground, bPreFog);
 
 		View.TopBatchedViewElements.Draw(RHICmdList, bNeedToSwitchVerticalAxis, View.ViewProjectionMatrix, View.ViewRect.Width(), View.ViewRect.Height(), true);
 	}
diff --git a/Engine/Source/Runtime/Renderer/Private/SceneRendering.cpp b/Engine/Source/Runtime/Renderer/Private/SceneRendering.cpp
--- a/Engine/Source/Runtime/Renderer/Private/SceneRendering.cpp
+++ b/Engine/Source/Runtime/Renderer/Private/SceneRendering.cpp
@@ -921,7 +921,7 @@ static void RenderViewFamily_RenderThread(FRHICommandListImmediate& RHICmdList,
     FMemMark MemStackMark(FMemStack::Get());
 
 	// update any resources that needed a deferred update
-	FDeferredUpdateResource::UpdateResources();
+	FDeferredUpdateResource::UpdateResources(RHICmdList);
 
     for( int ViewExt = 0; ViewExt < SceneRenderer->ViewFamily.ViewExtensions.Num(); ViewExt++ )
     {
diff --git a/Engine/Source/Runtime/Renderer/Private/ShadowRendering.cpp b/Engine/Source/Runtime/Renderer/Private/ShadowRendering.cpp
--- a/Engine/Source/Runtime/Renderer/Private/ShadowRendering.cpp
+++ b/Engine/Source/Runtime/Renderer/Private/ShadowRendering.cpp
@@ -1416,7 +1416,7 @@ void FProjectedShadowInfo::RenderDepth(FRHICommandListImmediate& RHICmdList, FDe
 			DrawPrimitiveFlags = EDrawDynamicFlags::ForceLowestLOD;
 		}
 
-		TDynamicPrimitiveDrawer<FShadowDepthDrawingPolicyFactory> Drawer(FoundView, FShadowDepthDrawingPolicyFactory::ContextType(this), true);
+		TDynamicPrimitiveDrawer<FShadowDepthDrawingPolicyFactory> Drawer(RHICmdList, FoundView, FShadowDepthDrawingPolicyFactory::ContextType(this), true);
 		uint32 PrimitiveCount = SubjectPrimitives.Num();
 		for(uint32 PrimitiveIndex = 0; PrimitiveIndex < PrimitiveCount; ++PrimitiveIndex)
 		{
@@ -1653,7 +1653,7 @@ void FProjectedShadowInfo::RenderProjection(FRHICommandListImmediate& RHICmdList
 			>::GetRHI(), 1);
 
 		// Draw the receiver's dynamic elements.
-		TDynamicPrimitiveDrawer<FDepthDrawingPolicyFactory> Drawer(View, FDepthDrawingPolicyFactory::ContextType(DDM_AllOccluders), true);
+		TDynamicPrimitiveDrawer<FDepthDrawingPolicyFactory> Drawer(RHICmdList, View, FDepthDrawingPolicyFactory::ContextType(DDM_AllOccluders), true);
 
 		for (int32 PrimitiveIndex = 0; PrimitiveIndex < ReceiverPrimitives.Num(); PrimitiveIndex++)
 		{
diff --git a/Engine/Source/Runtime/Renderer/Private/TranslucentLighting.cpp b/Engine/Source/Runtime/Renderer/Private/TranslucentLighting.cpp
--- a/Engine/Source/Runtime/Renderer/Private/TranslucentLighting.cpp
+++ b/Engine/Source/Runtime/Renderer/Private/TranslucentLighting.cpp
@@ -567,7 +567,7 @@ void FProjectedShadowInfo::RenderTranslucencyDepths(FRHICommandListImmediate& RH
 			CW_RGBA, BO_Add, BF_One, BF_One, BO_Add, BF_One, BF_One,
 			CW_RGBA, BO_Add, BF_One, BF_One, BO_Add, BF_One, BF_One>::GetRHI());
 
-		TDynamicPrimitiveDrawer<FTranslucencyShadowDepthDrawingPolicyFactory> OpacityDrawer(FoundView, FTranslucencyShadowDepthDrawingPolicyFactory::ContextType(bDirectionalLight), true);
+		TDynamicPrimitiveDrawer<FTranslucencyShadowDepthDrawingPolicyFactory> OpacityDrawer(RHICmdList, FoundView, FTranslucencyShadowDepthDrawingPolicyFactory::ContextType(bDirectionalLight), true);
 
 		for (int32 PrimitiveIndex = 0; PrimitiveIndex < SubjectTranslucentPrimitives.Num(); PrimitiveIndex++)
 		{
diff --git a/Engine/Source/Runtime/Renderer/Private/TranslucentRendering.cpp b/Engine/Source/Runtime/Renderer/Private/TranslucentRendering.cpp
--- a/Engine/Source/Runtime/Renderer/Private/TranslucentRendering.cpp
+++ b/Engine/Source/Runtime/Renderer/Private/TranslucentRendering.cpp
@@ -518,6 +518,7 @@ void FTranslucentPrimSet::RenderPrimitive(
 		if( ViewRelevance.bDynamicRelevance )
 		{
 			TDynamicPrimitiveDrawer<FTranslucencyDrawingPolicyFactory> TranslucencyDrawer(
+				RHICmdList,
 				&View,
 				FTranslucencyDrawingPolicyFactory::ContextType(TranslucentSelfShadow, bSeparateTranslucencyPass),
 				false
@@ -641,9 +642,9 @@ void FDeferredShadingSceneRenderer::RenderTranslucency(FRHICommandListImmediate&
 			// Draw only translucent prims that don't read from scene color
 			View.TranslucentPrimSet.DrawPrimitives(RHICmdList, View, *this, false);
 			// Draw the view's mesh elements with the translucent drawing policy.
-			DrawViewElements<FTranslucencyDrawingPolicyFactory>(View,FTranslucencyDrawingPolicyFactory::ContextType(),SDPG_World,false);
+			DrawViewElements<FTranslucencyDrawingPolicyFactory>(RHICmdList, View, FTranslucencyDrawingPolicyFactory::ContextType(), SDPG_World, false);
 			// Draw the view's mesh elements with the translucent drawing policy.
-			DrawViewElements<FTranslucencyDrawingPolicyFactory>(View,FTranslucencyDrawingPolicyFactory::ContextType(),SDPG_Foreground,false);
+			DrawViewElements<FTranslucencyDrawingPolicyFactory>(RHICmdList, View, FTranslucencyDrawingPolicyFactory::ContextType(), SDPG_Foreground, false);
 
 			const FSceneViewState* ViewState = (const FSceneViewState*)View.State;
 
diff --git a/Engine/Source/Runtime/Renderer/Private/VelocityRendering.cpp b/Engine/Source/Runtime/Renderer/Private/VelocityRendering.cpp
--- a/Engine/Source/Runtime/Renderer/Private/VelocityRendering.cpp
+++ b/Engine/Source/Runtime/Renderer/Private/VelocityRendering.cpp
@@ -498,7 +498,7 @@ void FDeferredShadingSceneRenderer::RenderVelocities(FRHICommandListImmediate& R
 
 	// Draw velocities for movable dynamic meshes.
 	TDynamicPrimitiveDrawer<FVelocityDrawingPolicyFactory> Drawer(
-		&View,FVelocityDrawingPolicyFactory::ContextType(DDM_AllOccluders),true,false,true
+		RHICmdList, &View, FVelocityDrawingPolicyFactory::ContextType(DDM_AllOccluders), true, false, true
 		);
 	for(int32 PrimitiveIndex = 0;PrimitiveIndex < View.VisibleDynamicPrimitives.Num();PrimitiveIndex++)
 	{
