UE4 - more command list work, remove a few GetImmediateCommandList calls from the FBasePassForwardOpaqueDrawingPolicyFactory pipe

diff --git a/Engine/Source/Runtime/Renderer/Private/ForwardBasePassRendering.cpp b/Engine/Source/Runtime/Renderer/Private/ForwardBasePassRendering.cpp
--- a/Engine/Source/Runtime/Renderer/Private/ForwardBasePassRendering.cpp
+++ b/Engine/Source/Runtime/Renderer/Private/ForwardBasePassRendering.cpp
@@ -86,7 +86,7 @@ public:
 	}
 };
 
-void FBasePassForwardOpaqueDrawingPolicyFactory::AddStaticMesh(FScene* Scene,FStaticMesh* StaticMesh)
+void FBasePassForwardOpaqueDrawingPolicyFactory::AddStaticMesh(FRHICommandList& RHICmdList, FScene* Scene, FStaticMesh* StaticMesh)
 {
 	// Determine the mesh's material and blend mode.
 	const auto FeatureLevel = Scene->GetFeatureLevel();
@@ -97,6 +97,7 @@ void FBasePassForwardOpaqueDrawingPolicyFactory::AddStaticMesh(FScene* Scene,FSt
 	if( !IsTranslucentBlendMode(BlendMode) )
 	{
 		ProcessBasePassMeshForForwardShading(
+			RHICmdList,
 			FProcessBasePassMeshParameters(
 				*StaticMesh,
 				Material,
@@ -221,6 +222,7 @@ bool FBasePassForwardOpaqueDrawingPolicyFactory::DrawDynamicMesh(
 	if(!IsTranslucentBlendMode(BlendMode))
 	{
 		ProcessBasePassMeshForForwardShading(
+			RHICmdList,
 			FProcessBasePassMeshParameters(
 				Mesh,
 				Material,
diff --git a/Engine/Source/Runtime/Renderer/Private/ForwardBasePassRendering.h b/Engine/Source/Runtime/Renderer/Private/ForwardBasePassRendering.h
--- a/Engine/Source/Runtime/Renderer/Private/ForwardBasePassRendering.h
+++ b/Engine/Source/Runtime/Renderer/Private/ForwardBasePassRendering.h
@@ -436,7 +436,7 @@ public:
 		ESceneRenderTargetsMode::Type TextureMode;
 	};
 
-	static void AddStaticMesh(FScene* Scene,FStaticMesh* StaticMesh);
+	static void AddStaticMesh(FRHICommandList& RHICmdList, FScene* Scene, FStaticMesh* StaticMesh);
 	static bool DrawDynamicMesh(
 		FRHICommandList& RHICmdList, 
 		const FSceneView& View,
@@ -458,6 +458,7 @@ public:
 /** Processes a base pass mesh using an unknown light map policy, and unknown fog density policy. */
 template<typename ProcessActionType>
 void ProcessBasePassMeshForForwardShading(
+	FRHICommandList& RHICmdList,
 	const FProcessBasePassMeshParameters& Parameters,
 	const ProcessActionType& Action
 	)
@@ -471,9 +472,6 @@ void ProcessBasePassMeshForForwardShading(
 
 	check(!AllowHighQualityLightmaps());
 
-	
-	FRHICommandListImmediate& RHICmdList = FRHICommandListExecutor::GetImmediateCommandList();
-
 	if (LightMapInteraction.GetType() == LMIT_Texture)
 	{
 		const FShadowMapInteraction ShadowMapInteraction = (Parameters.Mesh.LCI && bIsLitMaterial) 
diff --git a/Engine/Source/Runtime/Renderer/Private/ForwardTranslucentRendering.cpp b/Engine/Source/Runtime/Renderer/Private/ForwardTranslucentRendering.cpp
--- a/Engine/Source/Runtime/Renderer/Private/ForwardTranslucentRendering.cpp
+++ b/Engine/Source/Runtime/Renderer/Private/ForwardTranslucentRendering.cpp
@@ -187,6 +187,7 @@ bool FTranslucencyForwardShadingDrawingPolicyFactory::DrawDynamicMesh(
 	if (IsTranslucentBlendMode(BlendMode))
 	{
 		ProcessBasePassMeshForForwardShading(
+			RHICmdList,
 			FProcessBasePassMeshParameters(
 				Mesh,
 				Material,
diff --git a/Engine/Source/Runtime/Renderer/Private/SceneCore.cpp b/Engine/Source/Runtime/Renderer/Private/SceneCore.cpp
--- a/Engine/Source/Runtime/Renderer/Private/SceneCore.cpp
+++ b/Engine/Source/Runtime/Renderer/Private/SceneCore.cpp
@@ -321,7 +321,7 @@ void FStaticMesh::AddToDrawLists(FRHICommandListImmediate& RHICmdList, FScene* S
 		if (!bShadowOnly && !IsTranslucent(FeatureLevel))
 		{
 			// Add the static mesh to the DPG's base pass draw list.
-			FBasePassForwardOpaqueDrawingPolicyFactory::AddStaticMesh(Scene,this);
+			FBasePassForwardOpaqueDrawingPolicyFactory::AddStaticMesh(RHICmdList, Scene, this);
 		}
 	}
 }
