* refactor, better SceneColor naming convention

diff --git a/Engine/Source/Runtime/Renderer/Private/DistortionRendering.cpp b/Engine/Source/Runtime/Renderer/Private/DistortionRendering.cpp
--- a/Engine/Source/Runtime/Renderer/Private/DistortionRendering.cpp
+++ b/Engine/Source/Runtime/Renderer/Private/DistortionRendering.cpp
@@ -849,7 +849,8 @@ void FSceneRenderer::RenderDistortion(FRHICommandListImmediate& RHICmdList)
 
 // OCULUS BEGIN: select ONE render target for all views (eyes)
 		TRefCountPtr<IPooledRenderTarget> NewSceneColor;
-		GRenderTargetPool.FindFreeElement(GSceneRenderTargets.GetSceneColor()->GetDesc(), NewSceneColor, TEXT("RefractedSceneColor"));
+		// we don't create a new name to make it easier to use "vis SceneColor" and get the last HDRSceneColor
+		GRenderTargetPool.FindFreeElement(GSceneRenderTargets.GetSceneColor()->GetDesc(), NewSceneColor, TEXT("SceneColor"));
 		const FSceneRenderTargetItem& DestRenderTarget = NewSceneColor->GetRenderTargetItem();
 // OCULUS END
 
@@ -870,7 +871,8 @@ void FSceneRenderer::RenderDistortion(FRHICommandListImmediate& RHICmdList)
 			{
 // OCULUS BEGIN: Weird, different eyes may have different render targets!
 // 				TRefCountPtr<IPooledRenderTarget> NewSceneColor;
-// 				GRenderTargetPool.FindFreeElement(GSceneRenderTargets.SceneColor->GetDesc(), NewSceneColor, TEXT("RefractedSceneColor"));
+// we don't create a new name to make it easier to use "vis SceneColor" and get the last HDRSceneColor
+// 				GRenderTargetPool.FindFreeElement(GSceneRenderTargets.SceneColor->GetDesc(), NewSceneColor, TEXT("SceneColor"));
 // 				const FSceneRenderTargetItem& DestRenderTarget = NewSceneColor->GetRenderTargetItem(); 
 // OCULUS END
 
diff --git a/Engine/Source/Runtime/Renderer/Private/ReflectionEnvironment.cpp b/Engine/Source/Runtime/Renderer/Private/ReflectionEnvironment.cpp
--- a/Engine/Source/Runtime/Renderer/Private/ReflectionEnvironment.cpp
+++ b/Engine/Source/Runtime/Renderer/Private/ReflectionEnvironment.cpp
@@ -607,7 +607,8 @@ void FDeferredShadingSceneRenderer::RenderImageBasedReflectionsSM5ForAllViews(FR
 		FPooledRenderTargetDesc Desc = GSceneRenderTargets.GetSceneColor()->GetDesc();
 		Desc.TargetableFlags |= TexCreate_UAV;
 
-		GRenderTargetPool.FindFreeElement( Desc, NewSceneColor, TEXT("SceneColorEnv") );
+		// we don't create a new name to make it easier to use "vis SceneColor" and get the last HDRSceneColor
+		GRenderTargetPool.FindFreeElement( Desc, NewSceneColor, TEXT("SceneColor") );
 	}
 
 				// If we are in SM5, use the compute shader gather method
@@ -663,16 +664,6 @@ void FDeferredShadingSceneRenderer::RenderImageBasedReflectionsSM4ForAllViews(FR
 		&& Scene->SkyLight->ProcessedTexture
 		&& ViewFamily.EngineShowFlags.SkyLighting;
 
-	TRefCountPtr<IPooledRenderTarget> LightAccumulation;
-	{
-		const ERHIFeatureLevel::Type FeatureLevel = Scene->GetFeatureLevel();
-
-		uint32 LightAccumulationUAVFlag = (FeatureLevel == ERHIFeatureLevel::SM5) ? TexCreate_UAV : 0;
-		FPooledRenderTargetDesc Desc = GSceneRenderTargets.GetSceneColor()->GetDesc();
-
-		GRenderTargetPool.FindFreeElement(Desc, LightAccumulation, TEXT("LightAccumulation"));
-	}
-
 	static TArray<FReflectionCaptureSortData> SortData;
 
 	if (bReflectionEnv)
@@ -733,12 +724,23 @@ void FDeferredShadingSceneRenderer::RenderImageBasedReflectionsSM4ForAllViews(FR
 			ScreenSpaceReflections(RHICmdList, View, SSROutput);
 		}
 
+		TRefCountPtr<IPooledRenderTarget> LightAccumulation;
+
 		if (bReflectionEnv)
 		{
 			bRequiresApply = true;
 
 			SCOPED_DRAW_EVENT(StandardDeferredReflectionEnvironment, DEC_SCENE_ITEMS);
 
+			{
+				const ERHIFeatureLevel::Type FeatureLevel = Scene->GetFeatureLevel();
+
+				uint32 LightAccumulationUAVFlag = (FeatureLevel == ERHIFeatureLevel::SM5) ? TexCreate_UAV : 0;
+				FPooledRenderTargetDesc Desc = GSceneRenderTargets.GetSceneColor()->GetDesc();
+
+				GRenderTargetPool.FindFreeElement(Desc, LightAccumulation, TEXT("LightAccumulation"));
+			}
+
 			SetRenderTarget(RHICmdList, LightAccumulation->GetRenderTargetItem().TargetableTexture, NULL);
 
 			// Clear to no reflection contribution, alpha of 1 indicates full background contribution
@@ -814,7 +816,11 @@ void FDeferredShadingSceneRenderer::RenderImageBasedReflectionsSM4ForAllViews(FR
 
 			TShaderMapRef< FPostProcessVS >		VertexShader(GetGlobalShaderMap());
 
-			
+			if (!LightAccumulation)
+			{
+				// should never be used but during debugging it can happen
+				LightAccumulation = GSystemTextures.WhiteDummy;
+			}
 
 #define CASE(A,B,C) \
 			case ((A << 2) | (B << 1) | C) : \
