Return some fixed values for GPU synthetic benchmarking on OS X to workaround the lack of GPU timestamps. The values are keyed off the OpenGL driver in use and are quite coarse, but should help in the majority of cases.

diff --git a/Engine/Source/Runtime/Renderer/Private/GPUBenchmark.cpp b/Engine/Source/Runtime/Renderer/Private/GPUBenchmark.cpp
--- a/Engine/Source/Runtime/Renderer/Private/GPUBenchmark.cpp
+++ b/Engine/Source/Runtime/Renderer/Private/GPUBenchmark.cpp
@@ -387,7 +387,8 @@ void RendererGPUBenchmark(FRHICommandListImmediate& RHICmdList, FSynthBenchmarkR
 			TimerQueries[i] = GTimerQueryPool.AllocateQuery();
 		}
 
-		if(!TimerQueries[0])
+		const bool bSupportsTimerQueries = (TimerQueries[0] != NULL);
+		if(!bSupportsTimerQueries)
 		{
 			UE_LOG(LogSynthBenchmark, Warning, TEXT("GPU driver does not support timer queries."));
 		}
@@ -519,16 +520,84 @@ void RendererGPUBenchmark(FRHICommandListImmediate& RHICmdList, FSynthBenchmarkR
 				}
 #endif
 			}
+			
+			// Temporary workaround for GL_TIMESTAMP being unavailable and GL_TIME_ELAPSED workaround breaking drivers
+#if PLATFORM_MAC
+			if(!TimerQueries[0])
+			{
+				GLint RendererID = 0;
+				[[NSOpenGLContext currentContext] getValues:&RendererID forParameter:NSOpenGLCPCurrentRendererID];
+				{
+					switch((RendererID & kCGLRendererIDMatchingMask))
+					{
+						case kCGLRendererATIRadeonX4000ID: // AMD 7xx0 & Dx00 series - should be pretty beefy
+						{
+							InOut.GPUStats[0].SetMeasuredTime(1.2f * (1.0f / 4.601f));
+							InOut.GPUStats[1].SetMeasuredTime(1.2f * (1.0f / 7.447f));
+							InOut.GPUStats[2].SetMeasuredTime(1.2f * (1.0f / 3.847f));
+							InOut.GPUStats[3].SetMeasuredTime(1.2f * (1.0f / 25.463f));
+							InOut.GPUStats[4].SetMeasuredTime(1.2f * (1.0f / 1.072f));
+							break;
+						}
+						case kCGLRendererATIRadeonX3000ID: // AMD 5xx0, 6xx0 series - mostly OK
+						case kCGLRendererGeForceID: // Nvidia 6x0 & 7x0 series - mostly OK
+						{
+							InOut.GPUStats[0].SetMeasuredTime(2.f * (1.0f / 4.601f));
+							InOut.GPUStats[1].SetMeasuredTime(2.f * (1.0f / 7.447f));
+							InOut.GPUStats[2].SetMeasuredTime(2.f * (1.0f / 3.847f));
+							InOut.GPUStats[3].SetMeasuredTime(2.f * (1.0f / 25.463f));
+							InOut.GPUStats[4].SetMeasuredTime(2.f * (1.0f / 1.072f));
+							break;
+						}
+						case kCGLRendererIntelHD5000ID: // Intel HD 5000, Iris, Iris Pro - not dreadful
+						{
+							InOut.GPUStats[0].SetMeasuredTime(4.f * (1.0f / 4.601f));
+							InOut.GPUStats[1].SetMeasuredTime(4.f * (1.0f / 7.447f));
+							InOut.GPUStats[2].SetMeasuredTime(4.f * (1.0f / 3.847f));
+							InOut.GPUStats[3].SetMeasuredTime(4.f * (1.0f / 25.463f));
+							InOut.GPUStats[4].SetMeasuredTime(4.f * (1.0f / 1.072f));
+							break;
+						}
+						case kCGLRendererIntelHD4000ID: // Intel HD 4000 - quite slow
+						{
+							InOut.GPUStats[0].SetMeasuredTime(7.5f * (1.0f / 4.601f));
+							InOut.GPUStats[1].SetMeasuredTime(7.5f * (1.0f / 7.447f));
+							InOut.GPUStats[2].SetMeasuredTime(7.5f * (1.0f / 3.847f));
+							InOut.GPUStats[3].SetMeasuredTime(7.5f * (1.0f / 25.463f));
+							InOut.GPUStats[4].SetMeasuredTime(7.5f * (1.0f / 1.072f));
+							break;
+						}
+						case kCGLRendererATIRadeonX2000ID: // ATi 4xx0, 3xx0, 2xx0 - almost all very slow and drivers are now very buggy
+						case kCGLRendererGeForce8xxxID: // Nvidia 3x0, 2x0, 1x0, 9xx0, 8xx0 - almost all very slow
+						case kCGLRendererIntelHDID: // Intel HD 3000 - very, very slow and very buggy driver
+						default:
+						{
+							InOut.GPUStats[0].SetMeasuredTime(10.f * (1.0f / 4.601f));
+							InOut.GPUStats[1].SetMeasuredTime(10.f * (1.0f / 7.447f));
+							InOut.GPUStats[2].SetMeasuredTime(10.f * (1.0f / 3.847f));
+							InOut.GPUStats[3].SetMeasuredTime(10.f * (1.0f / 25.463f));
+							InOut.GPUStats[4].SetMeasuredTime(10.f * (1.0f / 1.072f));
+							break;
+						}
+					}
+				}
+			}
+			
+#endif
 
 			for(uint32 MethodId = 0; MethodId < MethodCount; ++MethodId)
 			{
 				float Confidence = 0.0f;
 				
-				float TimingValue = TimingSeries[MethodId].ComputeValue(Confidence);
-
-				if(Confidence > 0)
+				// Temporary workaround for GL_TIMESTAMP being unavailable and GL_TIME_ELAPSED workaround breaking drivers on OS X
+				if(!PLATFORM_MAC || bSupportsTimerQueries)
 				{
-					InOut.GPUStats[MethodId].SetMeasuredTime(TimingValue, Confidence);
+					float TimingValue = TimingSeries[MethodId].ComputeValue(Confidence);
+
+					if(Confidence > 0)
+					{
+						InOut.GPUStats[MethodId].SetMeasuredTime(TimingValue, Confidence);
+					}
 				}
 
 				UE_LOG(LogSynthBenchmark, Display, TEXT("         ... %.3f GigaPix/s, Confidence=%.0f%% '%s'"),
