Fixed translucency being 16x too bright from inv squared lights

diff --git a/Engine/Shaders/DynamicLightingCommon.usf b/Engine/Shaders/DynamicLightingCommon.usf
--- a/Engine/Shaders/DynamicLightingCommon.usf
+++ b/Engine/Shaders/DynamicLightingCommon.usf
@@ -64,9 +64,6 @@ float CalcLightAttenuation(float3 WorldPosition, out float3 WorldLightVector)
 	// Sphere falloff (technically just 1/d2 but this avoids inf)
 	DistanceAttenuation = 1 / ( DistanceSqr + 1 );
 	
-	// Correction for lumen units
-	DistanceAttenuation *= 16;
-	
 	float LightRadiusMask = Square( saturate( 1 - Square( DistanceSqr * DeferredLightUniforms.LightInvRadius * DeferredLightUniforms.LightInvRadius ) ) );
 	DistanceAttenuation *= LightRadiusMask;
 	
diff --git a/Engine/Shaders/TranslucentLightInjectionShaders.usf b/Engine/Shaders/TranslucentLightInjectionShaders.usf
--- a/Engine/Shaders/TranslucentLightInjectionShaders.usf
+++ b/Engine/Shaders/TranslucentLightInjectionShaders.usf
@@ -30,9 +30,6 @@ float CalcSimpleLightAttenuation(float3 WorldPosition)
 		// Sphere falloff (technically just 1/d2 but this avoids inf)
 		Attenuation = 1 / ( DistanceSqr + 1 );
 	
-		// Correction for lumen units
-		Attenuation *= 16;
-	
 		float LightRadiusMask = Square(saturate(1 - Square(DistanceSqr / (SimpleLightPositionAndRadius.w * SimpleLightPositionAndRadius.w))));
 		Attenuation *= LightRadiusMask;
 	}
diff --git a/Engine/Source/Runtime/Renderer/Private/TranslucentLighting.cpp b/Engine/Source/Runtime/Renderer/Private/TranslucentLighting.cpp
--- a/Engine/Source/Runtime/Renderer/Private/TranslucentLighting.cpp
+++ b/Engine/Source/Runtime/Renderer/Private/TranslucentLighting.cpp
@@ -1800,7 +1800,18 @@ public:
 		FVector4 PositionAndRadius(SimpleLightPerViewData.Position, SimpleLight.Radius);
 		SetShaderValue(RHICmdList, GetPixelShader(), VolumeCascadeIndex, VolumeCascadeIndexValue);
 		SetShaderValue(RHICmdList, GetPixelShader(), SimpleLightPositionAndRadius, PositionAndRadius);
-		SetShaderValue(RHICmdList, GetPixelShader(), SimpleLightColorAndExponent, FVector4(SimpleLight.Color, SimpleLight.Exponent));
+
+		FVector4 LightColorAndExponent(SimpleLight.Color, SimpleLight.Exponent);
+
+		if (SimpleLight.Exponent == 0)
+		{
+			// Correction for lumen units
+			LightColorAndExponent.X *= 16.0f;
+			LightColorAndExponent.Y *= 16.0f;
+			LightColorAndExponent.Z *= 16.0f;
+		}
+
+		SetShaderValue(RHICmdList, GetPixelShader(), SimpleLightColorAndExponent, LightColorAndExponent);
 	}
 
 	virtual bool Serialize(FArchive& Ar)
