TTP 340656:  Fix for screen percentage breaking when not in stereoscopic mode

diff --git a/Engine/Source/Runtime/Renderer/Private/PostProcess/PostProcessing.cpp b/Engine/Source/Runtime/Renderer/Private/PostProcess/PostProcessing.cpp
--- a/Engine/Source/Runtime/Renderer/Private/PostProcess/PostProcessing.cpp
+++ b/Engine/Source/Runtime/Renderer/Private/PostProcess/PostProcessing.cpp
@@ -1140,7 +1140,7 @@ void FPostProcessing::Process(FRHICommandListImmediate& RHICmdList, const FViewI
 		AddHighResScreenshotMask(Context, SeparateTranslucency);
 
 		// Do not use upscale if SeparateRenderTarget is in use!
-		if(View.UnscaledViewRect != View.ViewRect && !View.Family->EngineShowFlags.HMDDistortion && View.Family->EngineShowFlags.StereoRendering && !View.Family->bUseSeparateRenderTarget) //!!AB
+		if(View.UnscaledViewRect != View.ViewRect && !View.Family->bUseSeparateRenderTarget)
 		{
 			static const auto CVar = IConsoleManager::Get().FindTConsoleVariableDataInt(TEXT("r.UpsampleQuality")); 
 			int32 UpsampleMethod = CVar->GetValueOnRenderThread();
