TTP 340656:  Better fix, which doesn't break the non-direct rendering path

diff --git a/Engine/Source/Runtime/Renderer/Private/PostProcess/PostProcessing.cpp b/Engine/Source/Runtime/Renderer/Private/PostProcess/PostProcessing.cpp
--- a/Engine/Source/Runtime/Renderer/Private/PostProcess/PostProcessing.cpp
+++ b/Engine/Source/Runtime/Renderer/Private/PostProcess/PostProcessing.cpp
@@ -1140,7 +1140,7 @@ void FPostProcessing::Process(FRHICommandListImmediate& RHICmdList, const FViewI
 		AddHighResScreenshotMask(Context, SeparateTranslucency);
 
 		// Do not use upscale if SeparateRenderTarget is in use!
-		if(View.UnscaledViewRect != View.ViewRect && !View.Family->bUseSeparateRenderTarget)
+		if (View.UnscaledViewRect != View.ViewRect && (!View.Family->EngineShowFlags.StereoRendering || (!View.Family->EngineShowFlags.HMDDistortion && !View.Family->bUseSeparateRenderTarget)))
 		{
 			static const auto CVar = IConsoleManager::Get().FindTConsoleVariableDataInt(TEXT("r.UpsampleQuality")); 
 			int32 UpsampleMethod = CVar->GetValueOnRenderThread();
