Fixed Distance Field AO with cameras that constrain aspect ratio
Skylight LightColor and Intensity can be animated by Matinee now

diff --git a/Engine/Shaders/BasePassForForwardShadingPixelShader.usf b/Engine/Shaders/BasePassForForwardShadingPixelShader.usf
--- a/Engine/Shaders/BasePassForForwardShadingPixelShader.usf
+++ b/Engine/Shaders/BasePassForForwardShadingPixelShader.usf
@@ -84,7 +84,7 @@ void Main(
 	GetMaterialCoverageAndClipping(MaterialParameters);
 
 #if COMPILER_HLSL	// Not sure why this isn't needed everywhere
-	float2 PixelPos = SvPosition - View.ScreenTexelBias.xy;	
+	float2 PixelPos = SvPosition - View.ViewRectMin.xy;	
 	MaterialParameters.ScreenPosition.xy = ( PixelPos / View.ViewSizeAndSceneTexelSize.xy - 0.5 ) * float2(2,-2);
 	MaterialParameters.ScreenPosition.zw = SvPosition.zw;
 	
diff --git a/Engine/Shaders/DistanceFieldSurfaceCacheLighting.usf b/Engine/Shaders/DistanceFieldSurfaceCacheLighting.usf
--- a/Engine/Shaders/DistanceFieldSurfaceCacheLighting.usf
+++ b/Engine/Shaders/DistanceFieldSurfaceCacheLighting.usf
@@ -114,7 +114,7 @@ void ComputeDistanceFieldNormalPS(
 	uint2 TileCoordinate = (uint2)(SVPos) / uint2(THREADGROUP_SIZEX, THREADGROUP_SIZEY);
 	uint4 TileHead = GetTileHead(TileCoordinate);
 
-	float2 ScreenUV = float2((((uint2)SVPos.xy) * DOWNSAMPLE_FACTOR + .5f) * View.ViewSizeAndSceneTexelSize.zw);
+	float2 ScreenUV = float2((((uint2)SVPos.xy) * DOWNSAMPLE_FACTOR + View.ViewRectMin.xy + .5f) * View.ViewSizeAndSceneTexelSize.zw);
 	float2 ScreenPosition = (ScreenUV.xy - View.ScreenPositionScaleBias.wz) / View.ScreenPositionScaleBias.xy;
 
 	float SceneDepth = CalcSceneDepth(ScreenUV);
@@ -195,7 +195,7 @@ void ComputeDistanceFieldNormalPS(
 	{
 		WorldNormal = GBufferData.WorldNormal;
 	}
-	
+
 	float DepthSign = GBufferData.HasDistanceFieldRepresentation > 0 ? 1 : -1;
 	// Note: Encoding whether the pixel has a distance field representation in the sign bit of the alpha
 	OutColor = float4(EncodeNormalForAO(WorldNormal), DepthSign * SceneDepth);
@@ -222,7 +222,7 @@ void ComputeDistanceFieldNormalCS(
 	uint2 TileCoordinate = PixelCoordinate / uint2(THREADGROUP_SIZEX, THREADGROUP_SIZEY);
 	uint4 TileHead = GetTileHead(TileCoordinate);
 
-	float2 ScreenUV = float2((PixelCoordinate * DOWNSAMPLE_FACTOR + .5f) * View.ViewSizeAndSceneTexelSize.zw);
+	float2 ScreenUV = float2((PixelCoordinate * DOWNSAMPLE_FACTOR + View.ViewRectMin.xy + .5f) * View.ViewSizeAndSceneTexelSize.zw);
 	float2 ScreenPosition = (ScreenUV.xy - View.ScreenPositionScaleBias.wz) / View.ScreenPositionScaleBias.xy;
 
 	float SceneDepth = CalcSceneDepth(ScreenUV);
@@ -437,7 +437,7 @@ void IrradianceCacheSplatPS(
 
 #else
 
-	float2 ScreenUV = float2((((uint2)SVPos.xy) * CurrentLevelDownsampleFactor + float2(.5f, .5f)) * View.ViewSizeAndSceneTexelSize.zw);
+	float2 ScreenUV = float2((((uint2)SVPos.xy) * CurrentLevelDownsampleFactor + View.ViewRectMin.xy + float2(.5f, .5f)) * View.ViewSizeAndSceneTexelSize.zw);
 	float SceneW = CalcSceneDepth(ScreenUV); 
 	float2 ScreenPosition = (ScreenUV.xy - View.ScreenPositionScaleBias.wz) / View.ScreenPositionScaleBias.xy;
 	float4 HomogeneousWorldPosition = mul(float4(ScreenPosition.xy * SceneW, SceneW, 1), View.ScreenToWorld);
@@ -500,7 +500,7 @@ void AOCombinePS2(
 	in float4 SVPos : SV_POSITION,
 	out float4 OutColor : SV_Target0)
 {
-	float2 ScreenUV = float2((((uint2)SVPos.xy) * DOWNSAMPLE_FACTOR + .5f) * View.ViewSizeAndSceneTexelSize.zw);
+	float2 ScreenUV = float2((((uint2)SVPos.xy) * DOWNSAMPLE_FACTOR + View.ViewRectMin.xy + .5f) * View.ViewSizeAndSceneTexelSize.zw);
 	float SceneDepth = CalcSceneDepth(ScreenUV);
 
 #define VISUALIZE_ACCUMULATED_WEIGHTS 0
@@ -602,8 +602,11 @@ void UpdateHistoryPS(
     float2 OldUV = UVAndScreenPos.xy + mv / float2(0.5, -0.5) * View.ScreenPositionScaleBias.xy;
 	float2 OldScreenPosition = (OldUV.xy - View.ScreenPositionScaleBias.wz) / View.ScreenPositionScaleBias.xy;
 
-	float4 NewValue = Texture2DSampleLevel(BentNormalAOTexture, BentNormalAOSampler, UVAndScreenPos.xy, 0);
-	float4 HistoryValue = Texture2DSampleLevel(HistoryTexture, HistorySampler, OldUV, 0);
+	// Distance field AO was computed at 0,0 regardless of viewrect min
+	float2 DistanceFieldUVs = UVAndScreenPos.xy - View.ViewRectMin.xy * View.ViewSizeAndSceneTexelSize.zw;
+	float4 NewValue = Texture2DSampleLevel(BentNormalAOTexture, BentNormalAOSampler, DistanceFieldUVs, 0);
+	float2 OldDistanceFieldUVs = OldUV.xy - View.ViewRectMin.xy * View.ViewSizeAndSceneTexelSize.zw;
+	float4 HistoryValue = Texture2DSampleLevel(HistoryTexture, HistorySampler, OldDistanceFieldUVs, 0);
 	
 	//@todo - this doesn't take camera movement into account
 	float DepthDelta = abs(abs(NewValue.a) - abs(HistoryValue.a));
@@ -626,12 +629,15 @@ void UpdateHistoryPS(
 /** Upsamples the AO results to full resolution using a bilateral filter. */
 void AOUpsamplePS(in float4 UVAndScreenPos : TEXCOORD0, out float4 OutColor : SV_Target0)
 {
+	// Distance field AO was computed at 0,0 regardless of viewrect min
+	float2 DistanceFieldUVs = UVAndScreenPos.xy - View.ViewRectMin.xy * View.ViewSizeAndSceneTexelSize.zw;
+
 #define BILATERAL_UPSAMPLE 1
 #if BILATERAL_UPSAMPLE
 	float2 LowResBufferSize = floor(View.RenderTargetSize / DOWNSAMPLE_FACTOR);
 	float2 LowResTexelSize = 1.0f / LowResBufferSize;
-	float2 Corner00UV = floor(UVAndScreenPos.xy * LowResBufferSize - .5f) / LowResBufferSize + .5f * LowResTexelSize;
-	float2 BilinearWeights = (UVAndScreenPos.xy - Corner00UV) * LowResBufferSize;
+	float2 Corner00UV = floor(DistanceFieldUVs * LowResBufferSize - .5f) / LowResBufferSize + .5f * LowResTexelSize;
+	float2 BilinearWeights = (DistanceFieldUVs - Corner00UV) * LowResBufferSize;
 
 	float4 TextureValues00 = Texture2DSampleLevel(BentNormalAOTexture, BentNormalAOSampler, Corner00UV, 0);
 	float4 TextureValues10 = Texture2DSampleLevel(BentNormalAOTexture, BentNormalAOSampler, Corner00UV + float2(LowResTexelSize.x, 0), 0);
@@ -661,7 +667,7 @@ void AOUpsamplePS(in float4 UVAndScreenPos : TEXCOORD0, out float4 OutColor : SV
 	float3 BentNormal = InterpolatedResult.xyz;
 
 #else
-	float3 BentNormal = Texture2DSampleLevel(BentNormalAOTexture, BentNormalAOSampler, UVAndScreenPos.xy, 0).xyz;
+	float3 BentNormal = Texture2DSampleLevel(BentNormalAOTexture, BentNormalAOSampler, DistanceFieldUVs, 0).xyz;
 #endif
 
 	OutColor = float4(BentNormal, 1);
diff --git a/Engine/Shaders/DistanceFieldSurfaceCacheLightingCompute.usf b/Engine/Shaders/DistanceFieldSurfaceCacheLightingCompute.usf
--- a/Engine/Shaders/DistanceFieldSurfaceCacheLightingCompute.usf
+++ b/Engine/Shaders/DistanceFieldSurfaceCacheLightingCompute.usf
@@ -38,7 +38,7 @@ void BuildTileConesMain(
     uint3 GroupThreadId : SV_GroupThreadID) 
 {
     uint ThreadIndex = GroupThreadId.y * THREADGROUP_SIZEX + GroupThreadId.x;
-	float2 ScreenUV = (DispatchThreadId.xy * DOWNSAMPLE_FACTOR + float2(.5f, .5f)) * View.ViewSizeAndSceneTexelSize.zw;
+	float2 ScreenUV = (DispatchThreadId.xy * DOWNSAMPLE_FACTOR + View.ViewRectMin.xy + float2(.5f, .5f)) * View.ViewSizeAndSceneTexelSize.zw;
 	float SceneDepth = CalcSceneDepth(ScreenUV);
 
 	// Initialize per-tile variables
@@ -155,7 +155,7 @@ void DistanceFieldAOBuildTileListMain(
     uint3 GroupThreadId : SV_GroupThreadID) 
 {
     uint ThreadIndex = GroupThreadId.y * THREADGROUP_SIZEX + GroupThreadId.x;
-	float2 ScreenUV = (DispatchThreadId.xy * DOWNSAMPLE_FACTOR + float2(.5f, .5f)) * View.ViewSizeAndSceneTexelSize.zw;
+	float2 ScreenUV = (DispatchThreadId.xy * DOWNSAMPLE_FACTOR + View.ViewRectMin.xy + float2(.5f, .5f)) * View.ViewSizeAndSceneTexelSize.zw;
 	float SceneDepth = CalcSceneDepth(ScreenUV);
 
 	// Initialize per-tile variables
@@ -374,7 +374,7 @@ void PopulateCacheCS(
 
 	GroupMemoryBarrierWithGroupSync();
 
-	float2 ScreenUV = (DispatchThreadId.xy * CurrentLevelDownsampleFactor + float2(.5f, .5f)) * View.ViewSizeAndSceneTexelSize.zw;
+	float2 ScreenUV = (DispatchThreadId.xy * CurrentLevelDownsampleFactor + View.ViewRectMin.xy + float2(.5f, .5f)) * View.ViewSizeAndSceneTexelSize.zw;
 	float2 ScreenPosition = (ScreenUV.xy - View.ScreenPositionScaleBias.wz) / View.ScreenPositionScaleBias.xy;
 	float2 DownsampledScreenUV = (DispatchThreadId.xy + float2(.5f, .5f)) / AOBufferSize;
 
@@ -410,6 +410,7 @@ void PopulateCacheCS(
 				// Allocate a new record and store off attributes of the created record
 				int NextSampleIndex;
 				InterlockedAdd(NumQueuedIrradianceCacheRecords, 1, NextSampleIndex);
+				// W stores max allowed radius, used to limit overdraw from nearby samples placed in the high resolution passes
 				CachedIrradianceCachePositionRadius[NextSampleIndex] = float4(WorldShadingPosition, SceneDepth * CurrentLevelDownsampleFactor * .005f);
 				CachedIrradianceCacheNormal[NextSampleIndex] = float4(WorldNormal, .01f * CurrentLevelDownsampleFactor / 2);
 				CachedIrradianceCacheTileCoordinate[NextSampleIndex] = TileCoordinate;
diff --git a/Engine/Shaders/MaterialTemplate.usf b/Engine/Shaders/MaterialTemplate.usf
--- a/Engine/Shaders/MaterialTemplate.usf
+++ b/Engine/Shaders/MaterialTemplate.usf
@@ -1268,7 +1268,7 @@ void CalcMaterialParameters(
 		float2 PixelPos = (NDC * float2(0.5f, -0.5f) + float2(0.5f, 0.5f)) * View.ViewSizeAndSceneTexelSize.xy;
 		
 		// Pixel position relative to left top of screen (not viewport), center of pixel
-		Parameters.SVPosition.xy = PixelPos + View.ScreenTexelBias.xy;
+		Parameters.SVPosition.xy = PixelPos + View.ViewRectMin.xy;
 		// z: DeviceZ
 		Parameters.SVPosition.z = ConvertToDeviceZ(Parameters.ScreenPosition.w);
 		// w: SceneDepth
diff --git a/Engine/Source/Runtime/Engine/Classes/Components/LightComponent.h b/Engine/Source/Runtime/Engine/Classes/Components/LightComponent.h
--- a/Engine/Source/Runtime/Engine/Classes/Components/LightComponent.h
+++ b/Engine/Source/Runtime/Engine/Classes/Components/LightComponent.h
@@ -316,7 +316,7 @@ public:
 	 *
 	 * @param PropertyThatChanged	Property that changed
 	 */
-	virtual void PostInterpChange(UProperty* PropertyThatChanged);
+	virtual void PostInterpChange(UProperty* PropertyThatChanged) override;
 
 	/** 
 	 * Iterates over ALL stationary light components in the target world and assigns their preview shadowmap channel, and updates light icons accordingly.
diff --git a/Engine/Source/Runtime/Engine/Classes/Components/SkyLightComponent.h b/Engine/Source/Runtime/Engine/Classes/Components/SkyLightComponent.h
--- a/Engine/Source/Runtime/Engine/Classes/Components/SkyLightComponent.h
+++ b/Engine/Source/Runtime/Engine/Classes/Components/SkyLightComponent.h
@@ -113,6 +113,7 @@ class ENGINE_API USkyLightComponent : public ULightComponentBase
 	// Begin UObject Interface
 	virtual void PostInitProperties();	
 	virtual void PostLoad() override;
+	virtual void PostInterpChange(UProperty* PropertyThatChanged) override;
 #if WITH_EDITOR
 	virtual void PostEditChangeProperty(FPropertyChangedEvent& PropertyChangedEvent) override;
 	virtual bool CanEditChange(const UProperty* InProperty) const override;
diff --git a/Engine/Source/Runtime/Engine/Private/SkyLightComponent.cpp b/Engine/Source/Runtime/Engine/Private/SkyLightComponent.cpp
--- a/Engine/Source/Runtime/Engine/Private/SkyLightComponent.cpp
+++ b/Engine/Source/Runtime/Engine/Private/SkyLightComponent.cpp
@@ -181,6 +181,30 @@ void USkyLightComponent::PostLoad()
 	}
 }
 
+/** 
+* This is called when property is modified by InterpPropertyTracks
+*
+* @param PropertyThatChanged	Property that changed
+*/
+void USkyLightComponent::PostInterpChange(UProperty* PropertyThatChanged)
+{
+	static FName LightColorName(TEXT("LightColor"));
+	static FName IntensityName(TEXT("Intensity"));
+	static FName IndirectLightingIntensityName(TEXT("IndirectLightingIntensity"));
+
+	FName PropertyName = PropertyThatChanged->GetFName();
+	if (PropertyName == LightColorName
+		|| PropertyName == IntensityName
+		|| PropertyName == IndirectLightingIntensityName)
+	{
+		MarkRenderStateDirty();
+	}
+	else
+	{
+		Super::PostInterpChange(PropertyThatChanged);
+	}
+}
+
 void USkyLightComponent::DestroyRenderState_Concurrent()
 {
 	Super::DestroyRenderState_Concurrent();
@@ -340,16 +364,18 @@ void USkyLightComponent::ApplyComponentInstanceData(TSharedPtr<FComponentInstanc
 
 void USkyLightComponent::UpdateSkyCaptureContents(UWorld* WorldToUpdate)
 {
-	if (WorldToUpdate->Scene
-		// Don't process any sky capture requests until async shader compiling completes
-		&& (GShaderCompilingManager == NULL || !GShaderCompilingManager->IsCompiling()))
+	if (WorldToUpdate->Scene)
 	{
+		const bool bIsCompilingShaders = GShaderCompilingManager != NULL && GShaderCompilingManager->IsCompiling();
+
 		// Iterate backwards so we can remove elements without changing the index
 		for (int32 CaptureIndex = SkyCapturesToUpdate.Num() - 1; CaptureIndex >= 0; CaptureIndex--)
 		{
 			USkyLightComponent* CaptureComponent = SkyCapturesToUpdate[CaptureIndex];
 
-			if (!CaptureComponent->GetOwner() || WorldToUpdate->ContainsActor(CaptureComponent->GetOwner()))
+			if ((!CaptureComponent->GetOwner() || WorldToUpdate->ContainsActor(CaptureComponent->GetOwner()))
+				// Only process sky capture requests once async shader compiling completes, otherwise we will capture the scene with temporary shaders
+				&& (!bIsCompilingShaders || CaptureComponent->SourceType == SLS_SpecifiedCubemap))
 			{
 				// Only capture valid sky light components
 				if (CaptureComponent->SourceType != SLS_SpecifiedCubemap || CaptureComponent->Cubemap)
diff --git a/Engine/Source/Runtime/Engine/Public/SceneManagement.h b/Engine/Source/Runtime/Engine/Public/SceneManagement.h
--- a/Engine/Source/Runtime/Engine/Public/SceneManagement.h
+++ b/Engine/Source/Runtime/Engine/Public/SceneManagement.h
@@ -1265,7 +1265,7 @@ BEGIN_UNIFORM_BUFFER_STRUCT(FViewUniformShaderParameters,ENGINE_API)
 	DECLARE_UNIFORM_BUFFER_STRUCT_MEMBER_EX(FVector,ViewRight, EShaderPrecisionModifier::Half)
 	DECLARE_UNIFORM_BUFFER_STRUCT_MEMBER(FVector4,InvDeviceZToWorldZTransform)
 	DECLARE_UNIFORM_BUFFER_STRUCT_MEMBER_EX(FVector4,ScreenPositionScaleBias, EShaderPrecisionModifier::Half)
-	DECLARE_UNIFORM_BUFFER_STRUCT_MEMBER_EX(FVector4,ScreenTexelBias, EShaderPrecisionModifier::Half)
+	DECLARE_UNIFORM_BUFFER_STRUCT_MEMBER_EX(FVector4,ViewRectMin, EShaderPrecisionModifier::Half)
 	DECLARE_UNIFORM_BUFFER_STRUCT_MEMBER(FVector4,ViewSizeAndSceneTexelSize)
 	DECLARE_UNIFORM_BUFFER_STRUCT_MEMBER(FVector4,ViewOrigin)
 	DECLARE_UNIFORM_BUFFER_STRUCT_MEMBER(FVector4,TranslatedViewOrigin)
diff --git a/Engine/Source/Runtime/Renderer/Private/DistanceFieldSurfaceCacheLighting.cpp b/Engine/Source/Runtime/Renderer/Private/DistanceFieldSurfaceCacheLighting.cpp
--- a/Engine/Source/Runtime/Renderer/Private/DistanceFieldSurfaceCacheLighting.cpp
+++ b/Engine/Source/Runtime/Renderer/Private/DistanceFieldSurfaceCacheLighting.cpp
@@ -976,7 +976,7 @@ void ComputeDistanceFieldNormal(FRHICommandListImmediate& RHICmdList, const TArr
 
 			SCOPED_DRAW_EVENT(ComputeNormal, DEC_SCENE_ITEMS);
 
-			RHICmdList.SetViewport(View.ViewRect.Min.X / GAODownsampleFactor, View.ViewRect.Min.Y / GAODownsampleFactor, 0.0f, View.ViewRect.Max.X / GAODownsampleFactor, View.ViewRect.Max.Y / GAODownsampleFactor, 1.0f);
+			RHICmdList.SetViewport(0, 0, 0.0f, View.ViewRect.Width() / GAODownsampleFactor, View.ViewRect.Height() / GAODownsampleFactor, 1.0f);
 			RHICmdList.SetRasterizerState(TStaticRasterizerState<FM_Solid, CM_None>::GetRHI());
 			RHICmdList.SetDepthStencilState(TStaticDepthStencilState<false, CF_Always>::GetRHI());
 			RHICmdList.SetBlendState(TStaticBlendState<>::GetRHI());
@@ -994,7 +994,7 @@ void ComputeDistanceFieldNormal(FRHICommandListImmediate& RHICmdList, const TArr
 				RHICmdList,
 				0, 0, 
 				View.ViewRect.Width() / GAODownsampleFactor, View.ViewRect.Height() / GAODownsampleFactor,
-				View.ViewRect.Min.X, View.ViewRect.Min.Y, 
+				0, 0, 
 				View.ViewRect.Width(), View.ViewRect.Height(),
 				FIntPoint(View.ViewRect.Width() / GAODownsampleFactor, View.ViewRect.Height() / GAODownsampleFactor),
 				GSceneRenderTargets.GetBufferSizeXY(),
@@ -2076,7 +2076,7 @@ void UpdateHistory(
 				SCOPED_DRAW_EVENT(UpdateHistory, DEC_SCENE_ITEMS);
 				SetRenderTarget(RHICmdList, NewHistory->GetRenderTargetItem().TargetableTexture, NULL);
 
-				RHICmdList.SetViewport(View.ViewRect.Min.X / GAODownsampleFactor, View.ViewRect.Min.Y / GAODownsampleFactor, 0.0f, View.ViewRect.Max.X / GAODownsampleFactor, View.ViewRect.Max.Y / GAODownsampleFactor, 1.0f);
+				RHICmdList.SetViewport(0, 0, 0.0f, View.ViewRect.Width() / GAODownsampleFactor, View.ViewRect.Height() / GAODownsampleFactor, 1.0f);
 				RHICmdList.SetRasterizerState(TStaticRasterizerState<FM_Solid, CM_None>::GetRHI());
 				RHICmdList.SetDepthStencilState(TStaticDepthStencilState<false, CF_Always>::GetRHI());
 				RHICmdList.SetBlendState(TStaticBlendState<>::GetRHI());
@@ -2094,7 +2094,7 @@ void UpdateHistory(
 					RHICmdList,
 					0, 0, 
 					View.ViewRect.Width() / GAODownsampleFactor, View.ViewRect.Height() / GAODownsampleFactor,
-					View.ViewRect.Min.X / GAODownsampleFactor, View.ViewRect.Min.Y / GAODownsampleFactor, 
+					View.ViewRect.Min.X / GAODownsampleFactor, View.ViewRect.Min.Y / GAODownsampleFactor,
 					View.ViewRect.Width() / GAODownsampleFactor, View.ViewRect.Height() / GAODownsampleFactor,
 					FIntPoint(View.ViewRect.Width() / GAODownsampleFactor, View.ViewRect.Height() / GAODownsampleFactor),
 					GSceneRenderTargets.GetBufferSizeXY() / FIntPoint(GAODownsampleFactor, GAODownsampleFactor),
@@ -2143,7 +2143,8 @@ void PostProcessBentNormalAO(FRHICommandList& RHICmdList, TArray<FViewInfo>& Vie
 		for (int32 ViewIndex = 0; ViewIndex < Views.Num(); ViewIndex++)
 		{
 			FViewInfo& View = Views[ViewIndex];
-			RHICmdList.SetViewport(View.ViewRect.Min.X / GAODownsampleFactor, View.ViewRect.Min.Y / GAODownsampleFactor, 0.0f, View.ViewRect.Max.X / GAODownsampleFactor, View.ViewRect.Max.Y / GAODownsampleFactor, 1.0f);
+
+			RHICmdList.SetViewport(0, 0, 0.0f, View.ViewRect.Width() / GAODownsampleFactor, View.ViewRect.Height() / GAODownsampleFactor, 1.0f);
 			RHICmdList.SetRasterizerState(TStaticRasterizerState<FM_Solid, CM_None>::GetRHI());
 			RHICmdList.SetDepthStencilState(TStaticDepthStencilState<false, CF_Always>::GetRHI());
 			RHICmdList.SetBlendState(TStaticBlendState<>::GetRHI());
@@ -2157,12 +2158,11 @@ void PostProcessBentNormalAO(FRHICommandList& RHICmdList, TArray<FViewInfo>& Vie
 
 			PixelShader->SetParameters(RHICmdList, View, IrradianceCacheInterpolation);
 
-			//@todo - get the round up correct
 			DrawRectangle( 
 				RHICmdList,
 				0, 0, 
 				View.ViewRect.Width() / GAODownsampleFactor, View.ViewRect.Height() / GAODownsampleFactor,
-				View.ViewRect.Min.X / GAODownsampleFactor, View.ViewRect.Min.Y / GAODownsampleFactor, 
+				0, 0, 
 				View.ViewRect.Width() / GAODownsampleFactor, View.ViewRect.Height() / GAODownsampleFactor,
 				FIntPoint(View.ViewRect.Width() / GAODownsampleFactor, View.ViewRect.Height() / GAODownsampleFactor),
 				GetBufferSizeForAO() / GAODownsampleFactor,
@@ -2179,7 +2179,7 @@ void PostProcessBentNormalAO(FRHICommandList& RHICmdList, TArray<FViewInfo>& Vie
 		{
 			const FViewInfo& View = Views[ViewIndex];
 
-			RHICmdList.SetViewport(View.ViewRect.Min.X / GAODownsampleFactor, View.ViewRect.Min.Y / GAODownsampleFactor, 0.0f, View.ViewRect.Max.X / GAODownsampleFactor, View.ViewRect.Max.Y / GAODownsampleFactor, 1.0f);
+			RHICmdList.SetViewport(0, 0, 0.0f, View.ViewRect.Width() / GAODownsampleFactor, View.ViewRect.Height() / GAODownsampleFactor, 1.0f);
 			RHICmdList.SetRasterizerState(TStaticRasterizerState<FM_Solid, CM_None>::GetRHI());
 			RHICmdList.SetDepthStencilState(TStaticDepthStencilState<false, CF_Always>::GetRHI());
 			RHICmdList.SetBlendState(TStaticBlendState<>::GetRHI());
@@ -2198,7 +2198,7 @@ void PostProcessBentNormalAO(FRHICommandList& RHICmdList, TArray<FViewInfo>& Vie
 				RHICmdList,
 				0, 0, 
 				View.ViewRect.Width() / GAODownsampleFactor, View.ViewRect.Height() / GAODownsampleFactor,
-				View.ViewRect.Min.X / GAODownsampleFactor, View.ViewRect.Min.Y / GAODownsampleFactor, 
+				0, 0, 
 				View.ViewRect.Width() / GAODownsampleFactor, View.ViewRect.Height() / GAODownsampleFactor,
 				FIntPoint(View.ViewRect.Width() / GAODownsampleFactor, View.ViewRect.Height() / GAODownsampleFactor),
 				GSceneRenderTargets.GetBufferSizeXY() / FIntPoint(GAODownsampleFactor, GAODownsampleFactor),
@@ -2619,16 +2619,15 @@ void RenderIrradianceCacheInterpolation(FRHICommandListImmediate& RHICmdList, TA
 			const FScene* Scene = (const FScene*)View.Family->Scene;
 			FSurfaceCacheResources& SurfaceCacheResources = *Scene->SurfaceCacheResources;
 
-			RHICmdList.SetViewport(View.ViewRect.Min.X / DestLevelDownsampleFactor, View.ViewRect.Min.Y / DestLevelDownsampleFactor, 0.0f, FMath::DivideAndRoundUp(View.ViewRect.Max.X, DestLevelDownsampleFactor), FMath::DivideAndRoundUp(View.ViewRect.Max.Y, DestLevelDownsampleFactor), 1.0f);
+			const uint32 DownsampledViewSizeX = FMath::DivideAndRoundUp(View.ViewRect.Width(), DestLevelDownsampleFactor);
+			const uint32 DownsampledViewSizeY = FMath::DivideAndRoundUp(View.ViewRect.Height(), DestLevelDownsampleFactor);
 
+			RHICmdList.SetViewport(0, 0, 0.0f, DownsampledViewSizeX, DownsampledViewSizeY, 1.0f);
 			RHICmdList.SetRasterizerState(TStaticRasterizerState<FM_Solid, CM_None>::GetRHI());
 			//@todo - render front faces with depth testing, requires depth buffer
 			RHICmdList.SetDepthStencilState(TStaticDepthStencilState<false, CF_Always>::GetRHI());
 			RHICmdList.SetBlendState(TStaticBlendState<CW_RGBA, BO_Add, BF_One, BF_One, BO_Add, BF_One, BF_One>::GetRHI());
 
-			//FVector2D NormalizedOffsetToPixelCenter(-(-DestLevelDownsampleFactor / 2 + .5f) / (float)View.ViewRect.Width(), -(-DestLevelDownsampleFactor / 2 + .5f) / (float)View.ViewRect.Height());
-			uint32 DownsampledViewSizeX = FMath::DivideAndRoundUp(View.ViewRect.Width(), DestLevelDownsampleFactor);
-			uint32 DownsampledViewSizeY = FMath::DivideAndRoundUp(View.ViewRect.Height(), DestLevelDownsampleFactor);
 			FVector2D NormalizedOffsetToPixelCenter2(.5f / DownsampledViewSizeX - .5f / (DownsampledViewSizeX * DestLevelDownsampleFactor), .5f / DownsampledViewSizeY - .5f / (DownsampledViewSizeY * DestLevelDownsampleFactor));
 
 			// Convert to NDC 
@@ -2643,8 +2642,8 @@ void RenderIrradianceCacheInterpolation(FRHICommandListImmediate& RHICmdList, TA
 				View.ViewRect.Height() * InvBufferSizeY / (-2.0f * GProjectionSignY));
 
 			const FVector2D ScreenPositionBias(
-				(View.ViewRect.Width() / 2.0f + GPixelCenterOffset + View.ViewRect.Min.X) * InvBufferSizeX,
-				(View.ViewRect.Height() / 2.0f + GPixelCenterOffset + View.ViewRect.Min.Y) * InvBufferSizeY);
+				(View.ViewRect.Width() / 2.0f + GPixelCenterOffset) * InvBufferSizeX,
+				(View.ViewRect.Height() / 2.0f + GPixelCenterOffset) * InvBufferSizeY);
 
 			FVector2D OffsetToLowResCorner = (FVector2D(.5f, .5f) / FVector2D(DownsampledViewSizeX, DownsampledViewSizeY) - FVector2D(.5f, .5f)) * FVector2D(2, -2);
 			FVector2D OffsetToTopResPixel = (FVector2D(.5f, .5f) / BufferSize - ScreenPositionBias) / ScreenPositionScale;
@@ -3008,10 +3007,14 @@ bool FDeferredShadingSceneRenderer::RenderDistanceFieldAOSurfaceCache(FRHIComman
 					RenderIrradianceCacheInterpolation(RHICmdList, Views, IrradianceCacheAccumulation, DistanceFieldNormal->GetRenderTargetItem(), 0, GAODownsampleFactor, true);
 				}
 
+				GRenderTargetPool.VisualizeTexture.SetCheckPoint(RHICmdList, IrradianceCacheAccumulation);
+
 				// Post process the AO to cover over artifacts
 				TRefCountPtr<IPooledRenderTarget> AOOutput;
 				PostProcessBentNormalAO(RHICmdList, Views, IrradianceCacheAccumulation->GetRenderTargetItem(), AOOutput);
 
+				GRenderTargetPool.VisualizeTexture.SetCheckPoint(RHICmdList, AOOutput);
+
 				if (bApplyToSceneColor)
 				{
 					GSceneRenderTargets.BeginRenderingSceneColor(RHICmdList, false);
diff --git a/Engine/Source/Runtime/Renderer/Private/SceneRendering.cpp b/Engine/Source/Runtime/Renderer/Private/SceneRendering.cpp
--- a/Engine/Source/Runtime/Renderer/Private/SceneRendering.cpp
+++ b/Engine/Source/Runtime/Renderer/Private/SceneRendering.cpp
@@ -239,7 +239,7 @@ TUniformBufferRef<FViewUniformShaderParameters> FViewInfo::CreateUniformBuffer(
 	ViewUniformShaderParameters.ViewRight = EffectiveTranslatedViewMatrix.GetColumn(0);
 	ViewUniformShaderParameters.InvDeviceZToWorldZTransform = InvDeviceZToWorldZTransform;
 	ViewUniformShaderParameters.ScreenPositionScaleBias = ScreenPositionScaleBias;
-	ViewUniformShaderParameters.ScreenTexelBias = FVector4(ViewRect.Min.X, ViewRect.Min.Y, 0.0f, 0.0f);
+	ViewUniformShaderParameters.ViewRectMin = FVector4(ViewRect.Min.X, ViewRect.Min.Y, 0.0f, 0.0f);
 	ViewUniformShaderParameters.ViewSizeAndSceneTexelSize = FVector4(ViewRect.Width(), ViewRect.Height(), InvBufferSizeX, InvBufferSizeY);
 	ViewUniformShaderParameters.ViewOrigin = ViewMatrices.ViewOrigin;
 	ViewUniformShaderParameters.TranslatedViewOrigin = ViewMatrices.ViewOrigin + ViewMatrices.PreViewTranslation;
diff --git a/Engine/Source/Runtime/SlateRHIRenderer/Private/SlateRHIRenderingPolicy.cpp b/Engine/Source/Runtime/SlateRHIRenderer/Private/SlateRHIRenderingPolicy.cpp
--- a/Engine/Source/Runtime/SlateRHIRenderer/Private/SlateRHIRenderingPolicy.cpp
+++ b/Engine/Source/Runtime/SlateRHIRenderer/Private/SlateRHIRenderingPolicy.cpp
@@ -230,7 +230,7 @@ static FSceneView& CreateSceneView( FSceneViewFamilyContext& ViewFamilyContext,
 	ViewUniformShaderParameters.ViewRight = EffectiveTranslatedViewMatrix.GetColumn(0);
 	ViewUniformShaderParameters.InvDeviceZToWorldZTransform = View->InvDeviceZToWorldZTransform;
 	ViewUniformShaderParameters.ScreenPositionScaleBias = FVector4(0,0,0,0);
-	ViewUniformShaderParameters.ScreenTexelBias = FVector4(ViewRect.Min.X, ViewRect.Min.Y, 0.0f, 0.0f);
+	ViewUniformShaderParameters.ViewRectMin = FVector4(ViewRect.Min.X, ViewRect.Min.Y, 0.0f, 0.0f);
 	ViewUniformShaderParameters.ViewSizeAndSceneTexelSize = FVector4(ViewRect.Width(), ViewRect.Height(), 1.0f/ViewRect.Width(), 1.0f/ViewRect.Height() );
 	ViewUniformShaderParameters.ViewOrigin = View->ViewMatrices.ViewOrigin;
 	ViewUniformShaderParameters.TranslatedViewOrigin = View->ViewMatrices.ViewOrigin + View->ViewMatrices.PreViewTranslation;
