Fixed LPV, which was broken in fallout from Distance Field AO changes

diff --git a/Engine/Shaders/LPVClearLists.usf b/Engine/Shaders/LPVClearLists.usf
--- a/Engine/Shaders/LPVClearLists.usf
+++ b/Engine/Shaders/LPVClearLists.usf
@@ -20,8 +20,8 @@
 
 //------------------------------------------------------------------------------
 
-RWByteAddressBuffer								gVplListHeadBufferRW; 
-RWByteAddressBuffer								gGvListHeadBufferRW; 
+RWByteAddressBuffer								RWVplListHeadBuffer;
+RWByteAddressBuffer								RWGvListHeadBuffer;
 
 //------------------------------------------------------------------------------
 [numthreads(4,4,4)]
@@ -29,8 +29,8 @@ void CSClearLists(uint3 DTid : SV_DispatchThreadID)
 {
 	uint i = GetGridAddress( DTid );
 
-	gVplListHeadBufferRW.Store( i*4, -1 );
-	gGvListHeadBufferRW.Store( i*4, -1 );
+	RWVplListHeadBuffer.Store( i*4, -1 );
+	RWGvListHeadBuffer.Store( i*4, -1 );
 }
 
 //------------------------------------------------------------------------------
diff --git a/Engine/Source/Runtime/Renderer/Private/LightPropagationVolume.cpp b/Engine/Source/Runtime/Renderer/Private/LightPropagationVolume.cpp
--- a/Engine/Source/Runtime/Renderer/Private/LightPropagationVolume.cpp
+++ b/Engine/Source/Runtime/Renderer/Private/LightPropagationVolume.cpp
@@ -106,9 +106,9 @@ public:
 #endif
 
 		VplListHeadBufferSRV.Bind(Initializer.ParameterMap, TEXT("gVplListHeadBuffer") );
-		VplListHeadBufferUAV.Bind(Initializer.ParameterMap, TEXT("gVplListHeadBufferRW") );
+		VplListHeadBufferUAV.Bind(Initializer.ParameterMap, TEXT("RWVplListHeadBuffer") );
 		VplListBufferSRV.Bind(Initializer.ParameterMap, TEXT("gVplListBuffer") );
-		VplListBufferUAV.Bind(Initializer.ParameterMap, TEXT("gVplListBufferRW") );
+		VplListBufferUAV.Bind(Initializer.ParameterMap, TEXT("RWVplListBuffer") );
 
 #if LPV_GV_VOLUME_TEXTURE
 		for ( int i = 0; i < 3; i++ )
@@ -120,8 +120,8 @@ public:
 		GvBufferSRV.Bind(Initializer.ParameterMap, TEXT("gGvBuffer") );
 		GvBufferUAV.Bind(Initializer.ParameterMap, TEXT("gGvBufferRW") );
 #endif
-		GvListBufferUAV.Bind(Initializer.ParameterMap, TEXT("gGvListBufferRW") );
-		GvListHeadBufferUAV.Bind(Initializer.ParameterMap, TEXT("gGvListHeadBufferRW") );
+		GvListBufferUAV.Bind(Initializer.ParameterMap, TEXT("RWGvListBuffer") );
+		GvListHeadBufferUAV.Bind(Initializer.ParameterMap, TEXT("RWGvListHeadBuffer") );
 
 		GvListBufferSRV.Bind(Initializer.ParameterMap, TEXT("gGvListBuffer") );
 		GvListHeadBufferSRV.Bind(Initializer.ParameterMap, TEXT("gGvListHeadBuffer") );
@@ -470,11 +470,6 @@ public:
 
 		LinearTextureSampler.Bind(	Initializer.ParameterMap, TEXT("LinearSampler") );
 		PointTextureSampler.Bind(	Initializer.ParameterMap, TEXT("PointSampler") );
-
-		GvListBuffer.Bind(Initializer.ParameterMap,TEXT("GvListBuffer"));
-		GvListHeadBuffer.Bind(Initializer.ParameterMap,TEXT("GvListHeadBuffer"));
-		VplListBuffer.Bind(Initializer.ParameterMap,TEXT("VplListBuffer"));
-		VplListHeadBuffer.Bind(Initializer.ParameterMap,TEXT("VplListHeadBuffer"));
 	}
 
 	void SetParameters(
@@ -511,10 +506,6 @@ public:
 		Ar << RsmDepthTexture;
 		Ar << LinearTextureSampler;
 		Ar << PointTextureSampler;
-		Ar << GvListBuffer;
-		Ar << GvListHeadBuffer;
-		Ar << VplListBuffer;
-		Ar << VplListHeadBuffer;
 		return rv;
 	}
 protected:
@@ -524,11 +515,6 @@ protected:
 
 	FShaderResourceParameter LinearTextureSampler;
 	FShaderResourceParameter PointTextureSampler;
-
-	FRWShaderParameter GvListBuffer;
-	FRWShaderParameter GvListHeadBuffer;
-	FRWShaderParameter VplListBuffer;
-	FRWShaderParameter VplListHeadBuffer;
 };
 IMPLEMENT_SHADER_TYPE(,FLpvInject_GenerateVplListsCS,TEXT("LPVInject_GenerateVplLists"),TEXT("CSGenerateVplLists_LightDirectional"),SF_Compute);
 
