Fix crash TTP #334597.

diff --git a/Engine/Source/Runtime/Renderer/Private/ShadowSetup.cpp b/Engine/Source/Runtime/Renderer/Private/ShadowSetup.cpp
--- a/Engine/Source/Runtime/Renderer/Private/ShadowSetup.cpp
+++ b/Engine/Source/Runtime/Renderer/Private/ShadowSetup.cpp
@@ -1585,6 +1585,7 @@ void FDeferredShadingSceneRenderer::InitProjectedShadowVisibility(FRHICommandLis
 					bool bIsValidForView = View.StereoPass == eSSP_RIGHT_EYE
 						&& Views.IsValidIndex(ViewIndex - 1)
 						&& Views[ViewIndex - 1].StereoPass == eSSP_LEFT_EYE
+						&& ProjectedShadowInfo.FadeAlphas.IsValidIndex(ViewIndex)
 						&& ProjectedShadowInfo.FadeAlphas[ViewIndex] == 1.0f;
 					if (!bIsValidForView)
 					{
