UE4 - parallel rendering, added masked in the base pass and latency-hiding

diff --git a/Engine/Source/Runtime/RHI/Private/RHICommandList.cpp b/Engine/Source/Runtime/RHI/Private/RHICommandList.cpp
--- a/Engine/Source/Runtime/RHI/Private/RHICommandList.cpp
+++ b/Engine/Source/Runtime/RHI/Private/RHICommandList.cpp
@@ -19,14 +19,18 @@ static TAutoConsoleVariable<int32> CVarRHICmdBypass(
 	TEXT("0: Disable, 1: Enable"),
 	ECVF_Cheat);
 
-static int32 GRHICmdExec = 0;
-static FAutoConsoleVariableRef CVarRHICmdExec(
+static TAutoConsoleVariable<int32> CVarRHICmdExec(
 	TEXT("r.RHIExec"),
-	GRHICmdExec,
+	0,
 	TEXT("0: Executes RHI Command List using RHI calls (default)\n")
 	TEXT("1: Executes RHI Command List using platform specific APIs")
 	);
 
+TAutoConsoleVariable<int32> CVarRHICmdWidth(
+	TEXT("r.RHICmdWidth"), 
+	8,
+	TEXT("Number of threads."));
+
 
 RHI_API FRHICommandListExecutor GRHICommandList;
 RHI_API TLockFreeFixedSizeAllocator<FRHICommandList::FMemManager::FPage::PageSize> FRHICommandList::FMemManager::FPage::TheAllocator;
@@ -407,9 +411,8 @@ void FRHICommandListExecutor::ExecuteListGenericRHI(FRHICommandList& CmdList)
 
 void FRHICommandListExecutor::ExecuteInner(FRHICommandList& CmdList)
 {
-	static IConsoleVariable* RHIExecCVar = IConsoleManager::Get().FindConsoleVariable(TEXT("r.RHIExec"));
 	CmdList.bExecuting = true;
-	if (RHIExecCVar && RHIExecCVar->GetInt() != 0)
+	if (CVarRHICmdExec.GetValueOnRenderThread() != 0)
 	{
 		RHIExecuteCommandList(&CmdList);
 	}
diff --git a/Engine/Source/Runtime/RHI/Public/RHICommandList.h b/Engine/Source/Runtime/RHI/Public/RHICommandList.h
--- a/Engine/Source/Runtime/RHI/Public/RHICommandList.h
+++ b/Engine/Source/Runtime/RHI/Public/RHICommandList.h
@@ -18,6 +18,7 @@ DECLARE_CYCLE_STAT_EXTERN(TEXT("Immed. Command List Execute"), STAT_ImmedCmdList
 DECLARE_DWORD_COUNTER_STAT_EXTERN(TEXT("Immed. Command List memory"), STAT_ImmedCmdListMemory, STATGROUP_RHICMDLIST, RHI_API);
 DECLARE_DWORD_COUNTER_STAT_EXTERN(TEXT("Immed. Command count"), STAT_ImmedCmdListCount, STATGROUP_RHICMDLIST, RHI_API);
 
+
 enum ERHICommandType
 {
 	// Start the enum at non-zero to help debugging
@@ -71,6 +72,7 @@ enum ERHICommandType
 	ERCT_NUM,
 };
 
+extern RHI_API TAutoConsoleVariable<int32> CVarRHICmdWidth;
 
 struct FRHICommand
 {
diff --git a/Engine/Source/Runtime/Renderer/Private/DeferredShadingRenderer.cpp b/Engine/Source/Runtime/Renderer/Private/DeferredShadingRenderer.cpp
--- a/Engine/Source/Runtime/Renderer/Private/DeferredShadingRenderer.cpp
+++ b/Engine/Source/Runtime/Renderer/Private/DeferredShadingRenderer.cpp
@@ -192,7 +192,6 @@ void FDeferredShadingSceneRenderer::ClearGBufferAtMaxZ(FRHICommandList& RHICmdLi
 bool FDeferredShadingSceneRenderer::RenderBasePassStaticDataMasked(FRHICommandList& RHICmdList, FViewInfo& View)
 {
 	bool bDirty = false;
-
 	{
 		// Draw the scene's base pass draw lists.
 		FScene::EBasePassDrawListType MaskedDrawType = FScene::EBasePass_Masked;
@@ -210,47 +209,28 @@ bool FDeferredShadingSceneRenderer::RenderBasePassStaticDataMasked(FRHICommandLi
 			bDirty |= Scene->BasePassLowQualityLightMapDrawList[MaskedDrawType].DrawVisible(RHICmdList, View,View.StaticMeshVisibilityMap,View.StaticMeshBatchVisibility);
 		}
 	}
-
 	return bDirty;
 }
 
-static TAutoConsoleVariable<int32> CVarRHICmdWidth(
-	TEXT("r.RHICmdWidth"), 8,
-	TEXT("Number of threads."));
+FGraphEventRef FDeferredShadingSceneRenderer::RenderBasePassStaticDataMaskedParallel(FViewInfo& View, int32 Width, FGraphEventRef SubmitChain, bool& OutDirty)
+{
+	// Draw the scene's base pass draw lists.
+	FScene::EBasePassDrawListType MaskedDrawType = FScene::EBasePass_Masked;
+	SubmitChain = Scene->BasePassNoLightMapDrawList[MaskedDrawType].DrawVisibleParallel(View, View.StaticMeshVisibilityMap, View.StaticMeshBatchVisibility, Width, SubmitChain, OutDirty);
+	SubmitChain = Scene->BasePassSimpleDynamicLightingDrawList[MaskedDrawType].DrawVisibleParallel(View, View.StaticMeshVisibilityMap, View.StaticMeshBatchVisibility, Width, SubmitChain, OutDirty);
+	SubmitChain = Scene->BasePassCachedVolumeIndirectLightingDrawList[MaskedDrawType].DrawVisibleParallel(View, View.StaticMeshVisibilityMap, View.StaticMeshBatchVisibility, Width, SubmitChain, OutDirty);
+	SubmitChain = Scene->BasePassCachedPointIndirectLightingDrawList[MaskedDrawType].DrawVisibleParallel(View, View.StaticMeshVisibilityMap, View.StaticMeshBatchVisibility, Width, SubmitChain, OutDirty);
+
+	SubmitChain = Scene->BasePassHighQualityLightMapDrawList[MaskedDrawType].DrawVisibleParallel(View, View.StaticMeshVisibilityMap, View.StaticMeshBatchVisibility, Width, SubmitChain, OutDirty);
+	SubmitChain = Scene->BasePassDistanceFieldShadowMapLightMapDrawList[MaskedDrawType].DrawVisibleParallel(View, View.StaticMeshVisibilityMap, View.StaticMeshBatchVisibility, Width, SubmitChain, OutDirty);
+	SubmitChain = Scene->BasePassLowQualityLightMapDrawList[MaskedDrawType].DrawVisibleParallel(View, View.StaticMeshVisibilityMap, View.StaticMeshBatchVisibility, Width, SubmitChain, OutDirty);
+
+	return SubmitChain;
+}
 
 bool FDeferredShadingSceneRenderer::RenderBasePassStaticDataDefault(FRHICommandList& RHICmdList, FViewInfo& View)
 {
-	QUICK_SCOPE_CYCLE_COUNTER(AA_RenderBasePassStaticDataDefault);
 	bool bDirty = false;
-	if (!GRHICommandList.Bypass() && FApp::ShouldUseThreadingForPerformance())
-	{
-		FGraphEventRef SubmitChain;
-		int32 Width = CVarRHICmdWidth.GetValueOnRenderThread(); // we use a few more than needed to cover non-equal jobs
-
-		{
-			QUICK_SCOPE_CYCLE_COUNTER(AA_CreateTasks);
-
-			FScene::EBasePassDrawListType OpaqueDrawType = FScene::EBasePass_Default;
-			{
-				SubmitChain = Scene->BasePassNoLightMapDrawList[OpaqueDrawType].DrawVisibleParallel(View, View.StaticMeshVisibilityMap, View.StaticMeshBatchVisibility, Width, SubmitChain, bDirty);
-				SubmitChain = Scene->BasePassSimpleDynamicLightingDrawList[OpaqueDrawType].DrawVisibleParallel(View, View.StaticMeshVisibilityMap, View.StaticMeshBatchVisibility, Width, SubmitChain, bDirty);
-				SubmitChain = Scene->BasePassCachedVolumeIndirectLightingDrawList[OpaqueDrawType].DrawVisibleParallel(View, View.StaticMeshVisibilityMap, View.StaticMeshBatchVisibility, Width, SubmitChain, bDirty);
-				SubmitChain = Scene->BasePassCachedPointIndirectLightingDrawList[OpaqueDrawType].DrawVisibleParallel(View, View.StaticMeshVisibilityMap, View.StaticMeshBatchVisibility, Width, SubmitChain, bDirty);
-			}
-			{
-				SubmitChain = Scene->BasePassHighQualityLightMapDrawList[OpaqueDrawType].DrawVisibleParallel(View, View.StaticMeshVisibilityMap, View.StaticMeshBatchVisibility, Width, SubmitChain, bDirty);
-				SubmitChain = Scene->BasePassDistanceFieldShadowMapLightMapDrawList[OpaqueDrawType].DrawVisibleParallel(View, View.StaticMeshVisibilityMap, View.StaticMeshBatchVisibility, Width, SubmitChain, bDirty);
-				SubmitChain = Scene->BasePassLowQualityLightMapDrawList[OpaqueDrawType].DrawVisibleParallel(View, View.StaticMeshVisibilityMap, View.StaticMeshBatchVisibility, Width, SubmitChain, bDirty);
-			}
-		}
-		if (SubmitChain.GetReference())
-		{
-			QUICK_SCOPE_CYCLE_COUNTER(AA_WaitAndSubmit);
-			FTaskGraphInterface::Get().WaitUntilTaskCompletes(SubmitChain, ENamedThreads::RenderThread_Local);
-		}
-
-	}
-	else
 	{
 		FScene::EBasePassDrawListType OpaqueDrawType = FScene::EBasePass_Default;
 		{
@@ -271,6 +251,21 @@ bool FDeferredShadingSceneRenderer::RenderBasePassStaticDataDefault(FRHICommandL
 	return bDirty;
 }
 
+FGraphEventRef FDeferredShadingSceneRenderer::RenderBasePassStaticDataDefaultParallel(FViewInfo& View, int32 Width, FGraphEventRef SubmitChain, bool& OutDirty)
+{
+	FScene::EBasePassDrawListType OpaqueDrawType = FScene::EBasePass_Default;
+	SubmitChain = Scene->BasePassNoLightMapDrawList[OpaqueDrawType].DrawVisibleParallel(View, View.StaticMeshVisibilityMap, View.StaticMeshBatchVisibility, Width, SubmitChain, OutDirty);
+	SubmitChain = Scene->BasePassSimpleDynamicLightingDrawList[OpaqueDrawType].DrawVisibleParallel(View, View.StaticMeshVisibilityMap, View.StaticMeshBatchVisibility, Width, SubmitChain, OutDirty);
+	SubmitChain = Scene->BasePassCachedVolumeIndirectLightingDrawList[OpaqueDrawType].DrawVisibleParallel(View, View.StaticMeshVisibilityMap, View.StaticMeshBatchVisibility, Width, SubmitChain, OutDirty);
+	SubmitChain = Scene->BasePassCachedPointIndirectLightingDrawList[OpaqueDrawType].DrawVisibleParallel(View, View.StaticMeshVisibilityMap, View.StaticMeshBatchVisibility, Width, SubmitChain, OutDirty);
+
+	SubmitChain = Scene->BasePassHighQualityLightMapDrawList[OpaqueDrawType].DrawVisibleParallel(View, View.StaticMeshVisibilityMap, View.StaticMeshBatchVisibility, Width, SubmitChain, OutDirty);
+	SubmitChain = Scene->BasePassDistanceFieldShadowMapLightMapDrawList[OpaqueDrawType].DrawVisibleParallel(View, View.StaticMeshVisibilityMap, View.StaticMeshBatchVisibility, Width, SubmitChain, OutDirty);
+	SubmitChain = Scene->BasePassLowQualityLightMapDrawList[OpaqueDrawType].DrawVisibleParallel(View, View.StaticMeshVisibilityMap, View.StaticMeshBatchVisibility, Width, SubmitChain, OutDirty);
+
+	return SubmitChain;
+}
+
 void FDeferredShadingSceneRenderer::SortBasePassStaticData(FVector ViewPosition)
 {
 	// If we're not using a depth only pass, sort the static draw list buckets roughly front to back, to maximize HiZ culling
@@ -320,10 +315,34 @@ bool FDeferredShadingSceneRenderer::RenderBasePassStaticData(FRHICommandList& RH
 		bDirty |= RenderBasePassStaticDataDefault(RHICmdList, View);
 		bDirty |= RenderBasePassStaticDataMasked(RHICmdList, View);
 	}
-
 	return bDirty;
 }
 
+FGraphEventRef FDeferredShadingSceneRenderer::RenderBasePassStaticDataParallel(FViewInfo& View, int32 Width, FGraphEventRef SubmitChain, bool& OutDirty)
+{
+	SCOPE_CYCLE_COUNTER(STAT_StaticDrawListDrawTime);
+
+	// When using a depth-only pass, the default opaque geometry's depths are already
+	// in the depth buffer at this point, so rendering masked next will already cull
+	// as efficiently as it can, while also increasing the ZCull efficiency when
+	// rendering the default opaque geometry afterward.
+	if (EarlyZPassMode != DDM_None)
+	{
+		SubmitChain = RenderBasePassStaticDataMaskedParallel(View, Width, SubmitChain, OutDirty);
+		SubmitChain = RenderBasePassStaticDataDefaultParallel(View, Width, SubmitChain, OutDirty);
+	}
+	else
+	{
+		// Otherwise, in the case where we're not using a depth-only pre-pass, there
+		// is an advantage to rendering default opaque first to help cull the more
+		// expensive masked geometry.
+		SubmitChain = RenderBasePassStaticDataDefaultParallel(View, Width, SubmitChain, OutDirty);
+		SubmitChain = RenderBasePassStaticDataMaskedParallel(View, Width, SubmitChain, OutDirty);
+	}
+
+	return SubmitChain;
+}
+
 /**
 * Renders the basepass for the dynamic data of a given DPG and View.
 *
@@ -362,11 +381,35 @@ bool FDeferredShadingSceneRenderer::RenderBasePassDynamicData(FRHICommandListImm
 */
 bool FDeferredShadingSceneRenderer::RenderBasePass(FRHICommandListImmediate& RHICmdList, FViewInfo& View)
 {
-	bool bDirty=0;
+	bool bDirty = false;
+	bool bParallelDirty = false; // we keep this separate because |= is not threadsafe
+	FGraphEventRef SubmitChain;
 
 	// Render the base pass static data
 	{
-		bDirty |= RenderBasePassStaticData(RHICmdList, View);
+		if (!GRHICommandList.Bypass() && FApp::ShouldUseThreadingForPerformance())
+		{
+			int32 Width = CVarRHICmdWidth.GetValueOnRenderThread(); // we use a few more than needed to cover non-equal jobs
+
+			SubmitChain = RenderBasePassStaticDataParallel(View, Width, FGraphEventRef(), bParallelDirty);
+
+			static TAutoConsoleVariable<int32> CVarLatencyTest(
+				TEXT("r.LatencyTest"),
+				0,
+				TEXT("0: Executes static base pass stuff first (default)\n")
+				TEXT("1: Executes static base pass stuff after we have done the dynamics")
+				);
+
+			if (!CVarLatencyTest.GetValueOnRenderThread() && SubmitChain.GetReference())
+			{
+				FTaskGraphInterface::Get().WaitUntilTaskCompletes(SubmitChain, ENamedThreads::RenderThread_Local);
+				SubmitChain = nullptr;
+			}
+		}
+		else
+		{
+			bDirty |= RenderBasePassStaticData(RHICmdList, View);
+		}
 	}
 
 	{
@@ -406,8 +449,12 @@ bool FDeferredShadingSceneRenderer::RenderBasePass(FRHICommandListImmediate& RHI
 
 		bDirty |= RenderBasePassDynamicData(RHICmdList, View);
 	}
+	if (SubmitChain.GetReference())
+	{
+		FTaskGraphInterface::Get().WaitUntilTaskCompletes(SubmitChain, ENamedThreads::RenderThread_Local);
+	}
 
-	return bDirty;
+	return bDirty || bParallelDirty;
 }
 
 /** Render the TexturePool texture */
diff --git a/Engine/Source/Runtime/Renderer/Private/DeferredShadingRenderer.h b/Engine/Source/Runtime/Renderer/Private/DeferredShadingRenderer.h
--- a/Engine/Source/Runtime/Renderer/Private/DeferredShadingRenderer.h
+++ b/Engine/Source/Runtime/Renderer/Private/DeferredShadingRenderer.h
@@ -55,6 +55,11 @@ public:
 	bool RenderBasePassStaticDataMasked(FRHICommandList& RHICmdList, FViewInfo& View);
 	bool RenderBasePassStaticDataDefault(FRHICommandList& RHICmdList, FViewInfo& View);
 
+	/** Renders the basepass for the static data of a given View. Parallel versions.*/
+	FGraphEventRef RenderBasePassStaticDataParallel(FViewInfo& View, int32 Width, FGraphEventRef SubmitChain, bool& OutDirty);
+	FGraphEventRef RenderBasePassStaticDataMaskedParallel(FViewInfo& View, int32 Width, FGraphEventRef SubmitChain, bool& OutDirty);
+	FGraphEventRef RenderBasePassStaticDataDefaultParallel(FViewInfo& View, int32 Width, FGraphEventRef SubmitChain, bool& OutDirty);
+
 	/** Sorts base pass draw lists front to back for improved GPU culling. */
 	void SortBasePassStaticData(FVector ViewPosition);
 
