UE4 - bug fixes on parallel rendering toggle, plus draw dynamic basepass stuff in parallel

diff --git a/Engine/Source/Editor/MeshPaint/Private/MeshPaintEdMode.cpp b/Engine/Source/Editor/MeshPaint/Private/MeshPaintEdMode.cpp
--- a/Engine/Source/Editor/MeshPaint/Private/MeshPaintEdMode.cpp
+++ b/Engine/Source/Editor/MeshPaint/Private/MeshPaintEdMode.cpp
@@ -44,9 +44,9 @@ class FMeshPaintBatchedElementParameters : public FBatchedElementParameters
 public:
 
 	/** Binds vertex and pixel shaders for this element */
-	virtual void BindShaders_RenderThread(FRHICommandListImmediate& RHICmdList, const FMatrix& InTransform, const float InGamma, const FMatrix& ColorWeights, const FTexture* Texture) override
+	virtual void BindShaders(FRHICommandList& RHICmdList, const FMatrix& InTransform, const float InGamma, const FMatrix& ColorWeights, const FTexture* Texture) override
 	{
-		MeshPaintRendering::SetMeshPaintShaders_RenderThread(RHICmdList, InTransform, InGamma, ShaderParams );
+		MeshPaintRendering::SetMeshPaintShaders(RHICmdList, InTransform, InGamma, ShaderParams );
 	}
 
 public:
@@ -63,9 +63,9 @@ class FMeshPaintDilateBatchedElementParameters : public FBatchedElementParameter
 public:
 
 	/** Binds vertex and pixel shaders for this element */
-	virtual void BindShaders_RenderThread(FRHICommandListImmediate& RHICmdList, const FMatrix& InTransform, const float InGamma, const FMatrix& ColorWeights, const FTexture* Texture) override
+	virtual void BindShaders(FRHICommandList& RHICmdList, const FMatrix& InTransform, const float InGamma, const FMatrix& ColorWeights, const FTexture* Texture) override
 	{
-		MeshPaintRendering::SetMeshPaintDilateShaders_RenderThread(RHICmdList, InTransform, InGamma, ShaderParams );
+		MeshPaintRendering::SetMeshPaintDilateShaders(RHICmdList, InTransform, InGamma, ShaderParams );
 	}
 
 public:
diff --git a/Engine/Source/Editor/UnrealEd/Private/CubemapUnwrapUtils.cpp b/Engine/Source/Editor/UnrealEd/Private/CubemapUnwrapUtils.cpp
--- a/Engine/Source/Editor/UnrealEd/Private/CubemapUnwrapUtils.cpp
+++ b/Engine/Source/Editor/UnrealEd/Private/CubemapUnwrapUtils.cpp
@@ -124,20 +124,20 @@ void FCubemapTexturePropertiesPS<bHDROutput>::SetParameters( FRHICommandList& RH
 	SetShaderValue(RHICmdList, GetPixelShader(), Gamma, GammaValue);
 }
 
-void FMipLevelBatchedElementParameters::BindShaders_RenderThread(FRHICommandListImmediate& RHICmdList, const FMatrix& InTransform, const float InGamma, const FMatrix& ColorWeights, const FTexture* Texture)
+void FMipLevelBatchedElementParameters::BindShaders(FRHICommandList& RHICmdList, const FMatrix& InTransform, const float InGamma, const FMatrix& ColorWeights, const FTexture* Texture)
 {
 	if(bHDROutput)
 	{
-		BindShaders_RenderThread<FCubemapTexturePropertiesPS<true> >(RHICmdList, InTransform, InGamma, ColorWeights, Texture);
+		BindShaders<FCubemapTexturePropertiesPS<true> >(RHICmdList, InTransform, InGamma, ColorWeights, Texture);
 	}
 	else
 	{
-		BindShaders_RenderThread<FCubemapTexturePropertiesPS<false> >(RHICmdList, InTransform, InGamma, ColorWeights, Texture);
+		BindShaders<FCubemapTexturePropertiesPS<false> >(RHICmdList, InTransform, InGamma, ColorWeights, Texture);
 	}
 }
 
 template<typename TPixelShader>
-void FMipLevelBatchedElementParameters::BindShaders_RenderThread(FRHICommandListImmediate& RHICmdList, const FMatrix& InTransform, const float InGamma, const FMatrix& ColorWeights, const FTexture* Texture)
+void FMipLevelBatchedElementParameters::BindShaders(FRHICommandList& RHICmdList, const FMatrix& InTransform, const float InGamma, const FMatrix& ColorWeights, const FTexture* Texture)
 {
 	TShaderMapRef<FCubemapTexturePropertiesVS> VertexShader(GetGlobalShaderMap());
 	TShaderMapRef<TPixelShader> PixelShader(GetGlobalShaderMap());
@@ -161,7 +161,7 @@ void FIESLightProfilePS::SetParameters( FRHICommandList& RHICmdList, const FText
 	SetShaderValue(RHICmdList, ShaderRHI, BrightnessInLumens, InBrightnessInLumens);
 }
 
-void FIESLightProfileBatchedElementParameters::BindShaders_RenderThread( FRHICommandListImmediate& RHICmdList, const FMatrix& InTransform, const float InGamma, const FMatrix& ColorWeights, const FTexture* Texture )
+void FIESLightProfileBatchedElementParameters::BindShaders( FRHICommandList& RHICmdList, const FMatrix& InTransform, const float InGamma, const FMatrix& ColorWeights, const FTexture* Texture )
 {
 	TShaderMapRef<FSimpleElementVS> VertexShader(GetGlobalShaderMap());
 	TShaderMapRef<FIESLightProfilePS> PixelShader(GetGlobalShaderMap());
diff --git a/Engine/Source/Editor/UnrealEd/Private/MeshPaintRendering.cpp b/Engine/Source/Editor/UnrealEd/Private/MeshPaintRendering.cpp
--- a/Engine/Source/Editor/UnrealEd/Private/MeshPaintRendering.cpp
+++ b/Engine/Source/Editor/UnrealEd/Private/MeshPaintRendering.cpp
@@ -348,7 +348,7 @@ namespace MeshPaintRendering
 
 
 	/** Binds the mesh paint vertex and pixel shaders to the graphics device */
-	void SetMeshPaintShaders_RenderThread(FRHICommandListImmediate& RHICmdList, const FMatrix& InTransform,
+	void SetMeshPaintShaders(FRHICommandList& RHICmdList, const FMatrix& InTransform,
 										   const float InGamma,
 										   const FMeshPaintShaderParameters& InShaderParams )
 	{
@@ -369,7 +369,7 @@ namespace MeshPaintRendering
 	}
 
 	/** Binds the mesh paint vertex and pixel shaders to the graphics device */
-	void SetMeshPaintDilateShaders_RenderThread(FRHICommandListImmediate& RHICmdList, const FMatrix& InTransform,
+	void SetMeshPaintDilateShaders(FRHICommandList& RHICmdList, const FMatrix& InTransform,
 												 const float InGamma,
 												 const FMeshPaintDilateShaderParameters& InShaderParams )
 	{
diff --git a/Engine/Source/Editor/UnrealEd/Private/NormalMapPreview.cpp b/Engine/Source/Editor/UnrealEd/Private/NormalMapPreview.cpp
--- a/Engine/Source/Editor/UnrealEd/Private/NormalMapPreview.cpp
+++ b/Engine/Source/Editor/UnrealEd/Private/NormalMapPreview.cpp
@@ -74,8 +74,8 @@ private:
 IMPLEMENT_SHADER_TYPE(,FSimpleElementNormalMapPS,TEXT("SimpleElementNormalMapPixelShader"),TEXT("Main"),SF_Pixel);
 
 /** Binds vertex and pixel shaders for this element */
-void FNormalMapBatchedElementParameters::BindShaders_RenderThread(
-	FRHICommandListImmediate& RHICmdList,
+void FNormalMapBatchedElementParameters::BindShaders(
+	FRHICommandList& RHICmdList,
 	const FMatrix& InTransform,
 	const float InGamma,
 	const FMatrix& ColorWeights,
diff --git a/Engine/Source/Editor/UnrealEd/Private/Texture2DPreview.cpp b/Engine/Source/Editor/UnrealEd/Private/Texture2DPreview.cpp
--- a/Engine/Source/Editor/UnrealEd/Private/Texture2DPreview.cpp
+++ b/Engine/Source/Editor/UnrealEd/Private/Texture2DPreview.cpp
@@ -76,8 +76,8 @@ private:
 IMPLEMENT_SHADER_TYPE(,FSimpleElementTexture2DPreviewPS,TEXT("SimpleElementTexture2DPreviewPixelShader"),TEXT("Main"),SF_Pixel);
 
 /** Binds vertex and pixel shaders for this element */
-void FBatchedElementTexture2DPreviewParameters::BindShaders_RenderThread(
-	FRHICommandListImmediate& RHICmdList,
+void FBatchedElementTexture2DPreviewParameters::BindShaders(
+	FRHICommandList& RHICmdList,
 	const FMatrix& InTransform,
 	const float InGamma,
 	const FMatrix& ColorWeights,
diff --git a/Engine/Source/Editor/UnrealEd/Public/CubemapUnwrapUtils.h b/Engine/Source/Editor/UnrealEd/Public/CubemapUnwrapUtils.h
--- a/Engine/Source/Editor/UnrealEd/Public/CubemapUnwrapUtils.h
+++ b/Engine/Source/Editor/UnrealEd/Public/CubemapUnwrapUtils.h
@@ -120,10 +120,10 @@ public:
 	}
 
 	/** Binds vertex and pixel shaders for this element */
-	virtual void BindShaders_RenderThread(FRHICommandListImmediate& RHICmdList, const FMatrix& InTransform, const float InGamma, const FMatrix& ColorWeights, const FTexture* Texture) override;
+	virtual void BindShaders(FRHICommandList& RHICmdList, const FMatrix& InTransform, const float InGamma, const FMatrix& ColorWeights, const FTexture* Texture) override;
 
 private:
-	template<typename TPixelShader> void BindShaders_RenderThread(FRHICommandListImmediate& RHICmdList, const FMatrix& InTransform, const float InGamma, const FMatrix& ColorWeights, const FTexture* Texture);
+	template<typename TPixelShader> void BindShaders(FRHICommandList& RHICmdList, const FMatrix& InTransform, const float InGamma, const FMatrix& ColorWeights, const FTexture* Texture);
 
 	bool bHDROutput;
 	/** Parameters that need to be passed to the shader */
@@ -180,7 +180,7 @@ public:
 	}
 
 	/** Binds vertex and pixel shaders for this element */
-	virtual void BindShaders_RenderThread(FRHICommandListImmediate& RHICmdList, const FMatrix& InTransform, const float InGamma, const FMatrix& ColorWeights, const FTexture* Texture) override;
+	virtual void BindShaders(FRHICommandList& RHICmdList, const FMatrix& InTransform, const float InGamma, const FMatrix& ColorWeights, const FTexture* Texture) override;
 
 private:
 	float BrightnessInLumens;
diff --git a/Engine/Source/Editor/UnrealEd/Public/MeshPaintRendering.h b/Engine/Source/Editor/UnrealEd/Public/MeshPaintRendering.h
--- a/Engine/Source/Editor/UnrealEd/Public/MeshPaintRendering.h
+++ b/Engine/Source/Editor/UnrealEd/Public/MeshPaintRendering.h
@@ -57,12 +57,12 @@ namespace MeshPaintRendering
 
 
 	/** Binds the mesh paint vertex and pixel shaders to the graphics device */
-	void UNREALED_API SetMeshPaintShaders_RenderThread(FRHICommandListImmediate& RHICmdList, const FMatrix& InTransform,
+	void UNREALED_API SetMeshPaintShaders(FRHICommandList& RHICmdList, const FMatrix& InTransform,
 										   const float InGamma,
 										   const FMeshPaintShaderParameters& InShaderParams );
 
 	/** Binds the mesh paint dilation vertex and pixel shaders to the graphics device */
-	void UNREALED_API SetMeshPaintDilateShaders_RenderThread(FRHICommandListImmediate& RHICmdList, const FMatrix& InTransform,
+	void UNREALED_API SetMeshPaintDilateShaders(FRHICommandList& RHICmdList, const FMatrix& InTransform,
 												const float InGamma,
 												const FMeshPaintDilateShaderParameters& InShaderParams );
 
diff --git a/Engine/Source/Editor/UnrealEd/Public/NormalMapPreview.h b/Engine/Source/Editor/UnrealEd/Public/NormalMapPreview.h
--- a/Engine/Source/Editor/UnrealEd/Public/NormalMapPreview.h
+++ b/Engine/Source/Editor/UnrealEd/Public/NormalMapPreview.h
@@ -14,8 +14,8 @@
 class UNREALED_API FNormalMapBatchedElementParameters : public FBatchedElementParameters
 {
 	/** Binds vertex and pixel shaders for this element */
-	virtual void BindShaders_RenderThread(
-	FRHICommandListImmediate& RHICmdList,
+	virtual void BindShaders(
+		FRHICommandList& RHICmdList,
 		const FMatrix& InTransform,
 		const float InGamma,
 		const FMatrix& ColorWeights,
diff --git a/Engine/Source/Editor/UnrealEd/Public/Texture2DPreview.h b/Engine/Source/Editor/UnrealEd/Public/Texture2DPreview.h
--- a/Engine/Source/Editor/UnrealEd/Public/Texture2DPreview.h
+++ b/Engine/Source/Editor/UnrealEd/Public/Texture2DPreview.h
@@ -22,7 +22,7 @@ public:
 	}
 
 	/** Binds vertex and pixel shaders for this element */
-	virtual void BindShaders_RenderThread(FRHICommandListImmediate& RHICmdList, const FMatrix& InTransform, const float InGamma, const FMatrix& ColorWeights, const FTexture* Texture) override;
+	virtual void BindShaders(FRHICommandList& RHICmdList, const FMatrix& InTransform, const float InGamma, const FMatrix& ColorWeights, const FTexture* Texture) override;
 
 private:
 
diff --git a/Engine/Source/Runtime/Engine/Private/BatchedElements.cpp b/Engine/Source/Runtime/Engine/Private/BatchedElements.cpp
--- a/Engine/Source/Runtime/Engine/Private/BatchedElements.cpp
+++ b/Engine/Source/Runtime/Engine/Private/BatchedElements.cpp
@@ -350,7 +350,7 @@ float GBatchedElementSmoothWidth = 12;
  * Sets the appropriate vertex and pixel shader.
  */
 void FBatchedElements::PrepareShaders(
-	FRHICommandListImmediate& RHICmdList,
+	FRHICommandList& RHICmdList,
 	ESimpleElementBlendMode BlendMode,
 	const FMatrix& Transform,
 	bool bSwitchVerticalAxis,
@@ -430,7 +430,7 @@ void FBatchedElements::PrepareShaders(
 	if( BatchedElementParameters != NULL )
 	{
 		// Use the vertex/pixel shader that we were given
-		BatchedElementParameters->BindShaders_RenderThread(RHICmdList, Transform, Gamma, ColorWeights, Texture );
+		BatchedElementParameters->BindShaders(RHICmdList, Transform, Gamma, ColorWeights, Texture );
 	}
 	else
 	{
@@ -554,7 +554,7 @@ void FBatchedElements::PrepareShaders(
 PRAGMA_DISABLE_OPTIMIZATION
 #endif
 
-void FBatchedElements::DrawPointElements(FRHICommandListImmediate& RHICmdList, const FMatrix& Transform, const uint32 ViewportSizeX, const uint32 ViewportSizeY, const FVector& CameraX, const FVector& CameraY) const
+void FBatchedElements::DrawPointElements(FRHICommandList& RHICmdList, const FMatrix& Transform, const uint32 ViewportSizeX, const uint32 ViewportSizeY, const FVector& CameraX, const FVector& CameraY) const
 {
 	// Draw the point elements.
 	if( Points.Num() > 0 )
@@ -596,7 +596,7 @@ void FBatchedElements::DrawPointElements(FRHICommandListImmediate& RHICmdList, c
 }
 
 
-bool FBatchedElements::Draw(FRHICommandListImmediate& RHICmdList, bool bNeedToSwitchVerticalAxis, const FMatrix& Transform, uint32 ViewportSizeX, uint32 ViewportSizeY, bool bHitTesting, float Gamma, const FSceneView* View, FTexture2DRHIRef DepthTexture) const
+bool FBatchedElements::Draw(FRHICommandList& RHICmdList, bool bNeedToSwitchVerticalAxis, const FMatrix& Transform, uint32 ViewportSizeX, uint32 ViewportSizeY, bool bHitTesting, float Gamma, const FSceneView* View, FTexture2DRHIRef DepthTexture) const
 {
 	if( HasPrimsToDraw() )
 	{
diff --git a/Engine/Source/Runtime/Engine/Private/Materials/MaterialUniformExpressions.h b/Engine/Source/Runtime/Engine/Private/Materials/MaterialUniformExpressions.h
--- a/Engine/Source/Runtime/Engine/Private/Materials/MaterialUniformExpressions.h
+++ b/Engine/Source/Runtime/Engine/Private/Materials/MaterialUniformExpressions.h
@@ -238,7 +238,7 @@ public:
 	}
 	virtual void GetTextureValue(const FMaterialRenderContext& Context,const FMaterial& Material,const UTexture*& OutValue) const
 	{
-		check(IsInRenderingThread());
+		check(IsInParallelRenderingThread());
 		if( TransientOverrideValue_RenderThread != NULL )
 		{
 			OutValue = TransientOverrideValue_RenderThread;
diff --git a/Engine/Source/Runtime/Engine/Public/BatchedElements.h b/Engine/Source/Runtime/Engine/Public/BatchedElements.h
--- a/Engine/Source/Runtime/Engine/Public/BatchedElements.h
+++ b/Engine/Source/Runtime/Engine/Public/BatchedElements.h
@@ -72,7 +72,7 @@ class FBatchedElementParameters
 public:
 
 	/** Binds vertex and pixel shaders for this element */
-	virtual void BindShaders_RenderThread(FRHICommandListImmediate& RHICmdList, const FMatrix& InTransform, const float InGamma, const FMatrix& ColorWeights, const FTexture* Texture) = 0;
+	virtual void BindShaders(FRHICommandList& RHICmdList, const FMatrix& InTransform, const float InGamma, const FMatrix& ColorWeights, const FTexture* Texture) = 0;
 
 };
 
@@ -165,7 +165,7 @@ public:
 	 * @param View			Optional FSceneView for shaders that need access to view constants
 	 * @param DepthTexture	DepthTexture for manual depth testing with editor compositing in the pixel shader
 	 */
-	bool Draw(FRHICommandListImmediate& RHICmdList, bool bNeedToSwitchVerticalAxis, const FMatrix& Transform, uint32 ViewportSizeX, uint32 ViewportSizeY, bool bHitTesting, float Gamma = 1.0f, const FSceneView* View = NULL, FTexture2DRHIRef DepthTexture = FTexture2DRHIRef()) const;
+	bool Draw(FRHICommandList& RHICmdList, bool bNeedToSwitchVerticalAxis, const FMatrix& Transform, uint32 ViewportSizeX, uint32 ViewportSizeY, bool bHitTesting, float Gamma = 1.0f, const FSceneView* View = NULL, FTexture2DRHIRef DepthTexture = FTexture2DRHIRef()) const;
 	
 	FORCEINLINE bool HasPrimsToDraw() const
 	{
@@ -199,7 +199,7 @@ private:
 	 * @param	CameraX		Local space normalized view direction X vector
 	 * @param	CameraY		Local space normalized view direction Y vector
 	 */
-	void DrawPointElements(FRHICommandListImmediate& RHICmdList, const FMatrix& Transform, const uint32 ViewportSizeX, const uint32 ViewportSizeY, const FVector& CameraX, const FVector& CameraY) const;
+	void DrawPointElements(FRHICommandList& RHICmdList, const FMatrix& Transform, const uint32 ViewportSizeX, const uint32 ViewportSizeY, const FVector& CameraX, const FVector& CameraY) const;
 
 	TArray<FSimpleElementVertex> LineVertices;
 
@@ -300,7 +300,7 @@ private:
 	 * Sets the appropriate vertex and pixel shader.
 	 */
 	void PrepareShaders(
-		FRHICommandListImmediate& RHICmdList,
+		FRHICommandList& RHICmdList,
 		ESimpleElementBlendMode BlendMode,
 		const FMatrix& Transform,
 		bool bSwitchVerticalAxis,
diff --git a/Engine/Source/Runtime/RHI/Private/RHI.cpp b/Engine/Source/Runtime/RHI/Private/RHI.cpp
--- a/Engine/Source/Runtime/RHI/Private/RHI.cpp
+++ b/Engine/Source/Runtime/RHI/Private/RHI.cpp
@@ -88,6 +88,7 @@ void FRHIResource::FlushPendingDeletes()
 		for (int32 Index = 0; Index < ToDelete.Num(); Index++)
 		{
 			FRHIResource* Ref = ToDelete[Index];
+			check(Ref->MarkedForDelete == 1);
 			if (Ref->GetRefCount() == 0) // caches can bring dead objects back to life
 			{
 				CurrentlyDeleting = Ref;
@@ -96,7 +97,8 @@ void FRHIResource::FlushPendingDeletes()
 			}
 			else
 			{
-				verify(Ref->MarkedForDelete.Decrement() == 0);
+				Ref->MarkedForDelete = 0;
+				FPlatformMisc::MemoryBarrier();
 			}
 		}
 	}
diff --git a/Engine/Source/Runtime/RHI/Private/RHICommandList.cpp b/Engine/Source/Runtime/RHI/Private/RHICommandList.cpp
--- a/Engine/Source/Runtime/RHI/Private/RHICommandList.cpp
+++ b/Engine/Source/Runtime/RHI/Private/RHICommandList.cpp
@@ -86,6 +86,10 @@ void FRHICommandListExecutor::LatchBypass()
 #if !(UE_BUILD_SHIPPING || UE_BUILD_TEST)
 	check(GRHICommandList.OutstandingCmdListCount.GetValue() == 1 && !GRHICommandList.GetImmediateCommandList().HasCommands());
 	bool NewBypass = (CVarRHICmdBypass.GetValueOnRenderThread() >= 1);
+	if (NewBypass && !bLatchedBypass)
+	{
+		FRHIResource::FlushPendingDeletes();
+	}
 	bLatchedBypass = NewBypass;
 #endif
 }
diff --git a/Engine/Source/Runtime/RHI/Public/RHIResources.h b/Engine/Source/Runtime/RHI/Public/RHIResources.h
--- a/Engine/Source/Runtime/RHI/Public/RHIResources.h
+++ b/Engine/Source/Runtime/RHI/Public/RHIResources.h
@@ -15,7 +15,8 @@ class RHI_API FRHIResource
 {
 public:
 	FRHIResource(bool InbDoNotDeferDelete = false)
-		: bDoNotDeferDelete(InbDoNotDeferDelete)
+		: MarkedForDelete(0)
+		, bDoNotDeferDelete(InbDoNotDeferDelete)
 	{
 	}
 	virtual ~FRHIResource() 
@@ -39,14 +40,10 @@ public:
 			}
 			else
 			{
-				if (MarkedForDelete.Increment() == 1)
+				if (FPlatformAtomics::InterlockedCompareExchange(&MarkedForDelete, 1, 0) == 0)
 				{
 					PendingDeletes.Push(const_cast<FRHIResource*>(this));
 				}
-				else
-				{
-					MarkedForDelete.Decrement();
-				}
 			}
 		}
 		check(NewValue >= 0);
@@ -60,10 +57,10 @@ public:
 	}
 	void DoNoDeferDelete()
 	{
-		check(!MarkedForDelete.GetValue());
+		check(!MarkedForDelete);
 		bDoNotDeferDelete = true;
 		FPlatformMisc::MemoryBarrier();
-		check(!MarkedForDelete.GetValue());
+		check(!MarkedForDelete);
 	}
 
 	static void FlushPendingDeletes();
@@ -79,7 +76,7 @@ public:
 
 private:
 	mutable FThreadSafeCounter NumRefs;
-	mutable FThreadSafeCounter MarkedForDelete;
+	mutable int32 MarkedForDelete;
 	bool bDoNotDeferDelete;
 	static TLockFreePointerList<FRHIResource> PendingDeletes;
 	static FRHIResource* CurrentlyDeleting;
diff --git a/Engine/Source/Runtime/Renderer/Private/DeferredShadingRenderer.cpp b/Engine/Source/Runtime/Renderer/Private/DeferredShadingRenderer.cpp
--- a/Engine/Source/Runtime/Renderer/Private/DeferredShadingRenderer.cpp
+++ b/Engine/Source/Runtime/Renderer/Private/DeferredShadingRenderer.cpp
@@ -349,9 +349,43 @@ FGraphEventRef FDeferredShadingSceneRenderer::RenderBasePassStaticDataParallel(F
 * @return true if anything was rendered to scene color
 */
 
-bool FDeferredShadingSceneRenderer::RenderBasePassDynamicData(FRHICommandListImmediate& RHICmdList, FViewInfo& View)
+void FDeferredShadingSceneRenderer::RenderBasePassDynamicData(FRHICommandList& RHICmdList, const FViewInfo& View, bool& bOutDirty)
 {
-	bool bDirty=0;
+	bool bDirty = false;
+
+	SCOPE_CYCLE_COUNTER(STAT_DynamicPrimitiveDrawTime);
+	SCOPED_DRAW_EVENT(Dynamic, DEC_SCENE_ITEMS);
+
+	if (View.VisibleDynamicPrimitives.Num() > 0)
+	{
+		// Draw the dynamic non-occluded primitives using a base pass drawing policy.
+		TDynamicPrimitiveDrawer<FBasePassOpaqueDrawingPolicyFactory> Drawer(
+			RHICmdList, &View, FBasePassOpaqueDrawingPolicyFactory::ContextType(false, ESceneRenderTargetsMode::DontSet), true);
+		for (int32 PrimitiveIndex = 0, Num = View.VisibleDynamicPrimitives.Num(); PrimitiveIndex < Num; PrimitiveIndex++)
+		{
+			const FPrimitiveSceneInfo* PrimitiveSceneInfo = View.VisibleDynamicPrimitives[PrimitiveIndex];
+			int32 PrimitiveId = PrimitiveSceneInfo->GetIndex();
+			const FPrimitiveViewRelevance& PrimitiveViewRelevance = View.PrimitiveViewRelevanceMap[PrimitiveId];
+
+			const bool bVisible = View.PrimitiveVisibilityMap[PrimitiveId];
+
+			// Only draw the primitive if it's visible
+			if (bVisible &&
+				// only draw opaque and masked primitives if wireframe is disabled
+				(PrimitiveViewRelevance.bOpaqueRelevance || ViewFamily.EngineShowFlags.Wireframe) &&
+				PrimitiveViewRelevance.bRenderInMainPass
+				)
+			{
+				FScopeCycleCounter Context(PrimitiveSceneInfo->Proxy->GetStatId());
+				Drawer.SetPrimitive(PrimitiveSceneInfo->Proxy);
+				PrimitiveSceneInfo->Proxy->DrawDynamicElements(
+					&Drawer,
+					&View
+					);
+			}
+		}
+		bDirty |= Drawer.IsDirty();
+	}
 
 	if( !View.Family->EngineShowFlags.CompositeEditorPrimitives )
 	{
@@ -372,7 +406,69 @@ bool FDeferredShadingSceneRenderer::RenderBasePassDynamicData(FRHICommandListImm
 
 	}
 
-	return bDirty;
+	// this little bit of code is required because multiple threads might be writing bOutDirty...this you cannot use || bDirty - type things.
+	if (bDirty)
+	{
+		bOutDirty = true;
+	}
+}
+
+
+class FRenderBasePassDynamicDataThreadTask
+{
+	FDeferredShadingSceneRenderer& ThisRenderer;
+	FRHICommandList& RHICmdList;
+	const FViewInfo& View;
+	bool& OutDirty;
+
+public:
+
+	FRenderBasePassDynamicDataThreadTask(
+		FDeferredShadingSceneRenderer* InThisRenderer,
+		FRHICommandList* InRHICmdList,
+		const FViewInfo* InView,
+		bool* InOutDirty
+		)
+		: ThisRenderer(*InThisRenderer)
+		, RHICmdList(*InRHICmdList)
+		, View(*InView)
+		, OutDirty(*InOutDirty)
+	{
+	}
+
+	FORCEINLINE TStatId GetStatId() const
+	{
+		RETURN_QUICK_DECLARE_CYCLE_STAT(FRenderBasePassDynamicDataThreadTask, STATGROUP_TaskGraphTasks);
+	}
+
+	ENamedThreads::Type GetDesiredThread()
+	{
+		return ENamedThreads::AnyThread;
+	}
+
+	static ESubsequentsMode::Type GetSubsequentsMode() { return ESubsequentsMode::TrackSubsequents; }
+
+	void DoTask(ENamedThreads::Type CurrentThread, const FGraphEventRef& MyCompletionGraphEvent)
+	{
+		ThisRenderer.RenderBasePassDynamicData(RHICmdList, View, OutDirty);
+	}
+};
+
+FGraphEventRef FDeferredShadingSceneRenderer::RenderBasePassDynamicDataParallel(FViewInfo& View, FGraphEventRef SubmitChain, bool& bOutDirty)
+{
+	FRHICommandList* CmdList = new FRHICommandList;
+
+	FGraphEventRef AnyThreadCompletionEvent = TGraphTask<FRenderBasePassDynamicDataThreadTask>::CreateTask(nullptr, ENamedThreads::RenderThread)
+		.ConstructAndDispatchWhenReady(this, CmdList, &View, &bOutDirty);
+
+	FGraphEventArray Prereqs;
+	Prereqs.Add(AnyThreadCompletionEvent);
+	if (SubmitChain.GetReference())
+	{
+		Prereqs.Add(SubmitChain);
+	}
+
+	return TGraphTask<FSubmitCommandlistThreadTask>::CreateTask(&Prereqs, ENamedThreads::RenderThread).ConstructAndDispatchWhenReady(CmdList);
 }
 
 /**
@@ -381,9 +477,7 @@ bool FDeferredShadingSceneRenderer::RenderBasePassDynamicData(FRHICommandListImm
 */
 bool FDeferredShadingSceneRenderer::RenderBasePass(FRHICommandListImmediate& RHICmdList, FViewInfo& View)
 {
-	bool bDirty = false;
-	bool bParallelDirty = false; // we keep this separate because |= is not threadsafe
-	FGraphEventRef SubmitChain;
+	bool bDirty = false; 
 
 	// Render the base pass static data
 	{
@@ -391,70 +485,23 @@ bool FDeferredShadingSceneRenderer::RenderBasePass(FRHICommandListImmediate& RHI
 		{
 			int32 Width = CVarRHICmdWidth.GetValueOnRenderThread(); // we use a few more than needed to cover non-equal jobs
 
-			SubmitChain = RenderBasePassStaticDataParallel(View, Width, FGraphEventRef(), bParallelDirty);
+			FGraphEventRef SubmitChain = RenderBasePassStaticDataParallel(View, Width, FGraphEventRef(), bDirty);
+			SubmitChain = RenderBasePassDynamicDataParallel(View, SubmitChain, bDirty);
 
-			static TAutoConsoleVariable<int32> CVarLatencyTest(
-				TEXT("r.LatencyTest"),
-				0,
-				TEXT("0: Executes static base pass stuff first (default)\n")
-				TEXT("1: Executes static base pass stuff after we have done the dynamics")
-				);
-
-			if (!CVarLatencyTest.GetValueOnRenderThread() && SubmitChain.GetReference())
+			if (SubmitChain.GetReference())
 			{
 				FTaskGraphInterface::Get().WaitUntilTaskCompletes(SubmitChain, ENamedThreads::RenderThread_Local);
-				SubmitChain = nullptr;
 			}
 		}
 		else
 		{
+			bool bDirty = false;
 			bDirty |= RenderBasePassStaticData(RHICmdList, View);
+			RenderBasePassDynamicData(RHICmdList, View, bDirty);
 		}
 	}
 
-	{
-		SCOPE_CYCLE_COUNTER(STAT_DynamicPrimitiveDrawTime);
-		SCOPED_DRAW_EVENT(Dynamic, DEC_SCENE_ITEMS);
-
-		if( View.VisibleDynamicPrimitives.Num() > 0 )
-		{
-			// Draw the dynamic non-occluded primitives using a base pass drawing policy.
-			TDynamicPrimitiveDrawer<FBasePassOpaqueDrawingPolicyFactory> Drawer(
-				RHICmdList, &View, FBasePassOpaqueDrawingPolicyFactory::ContextType(false, ESceneRenderTargetsMode::DontSet), true);
-			for(int32 PrimitiveIndex = 0, Num = View.VisibleDynamicPrimitives.Num();PrimitiveIndex < Num;PrimitiveIndex++)
-			{
-				const FPrimitiveSceneInfo* PrimitiveSceneInfo = View.VisibleDynamicPrimitives[PrimitiveIndex];
-				int32 PrimitiveId = PrimitiveSceneInfo->GetIndex();
-				const FPrimitiveViewRelevance& PrimitiveViewRelevance = View.PrimitiveViewRelevanceMap[PrimitiveId];
-
-				const bool bVisible = View.PrimitiveVisibilityMap[PrimitiveId];
-
-				// Only draw the primitive if it's visible
-				if( bVisible && 
-					// only draw opaque and masked primitives if wireframe is disabled
-					(PrimitiveViewRelevance.bOpaqueRelevance || ViewFamily.EngineShowFlags.Wireframe) &&
-					PrimitiveViewRelevance.bRenderInMainPass
-					)
-				{
-					FScopeCycleCounter Context(PrimitiveSceneInfo->Proxy->GetStatId());
-					Drawer.SetPrimitive(PrimitiveSceneInfo->Proxy);
-					PrimitiveSceneInfo->Proxy->DrawDynamicElements(
-						&Drawer,
-						&View			
-						);
-				}
-			}
-			bDirty |= Drawer.IsDirty(); 
-		}
-
-		bDirty |= RenderBasePassDynamicData(RHICmdList, View);
-	}
-	if (SubmitChain.GetReference())
-	{
-		FTaskGraphInterface::Get().WaitUntilTaskCompletes(SubmitChain, ENamedThreads::RenderThread_Local);
-	}
-
-	return bDirty || bParallelDirty;
+	return bDirty;
 }
 
 /** Render the TexturePool texture */
@@ -979,7 +1026,6 @@ bool FDeferredShadingSceneRenderer::RenderBasePass(FRHICommandListImmediate& RHI
 {
 	bool bDirty = false;
 
-
 	if(ViewFamily.EngineShowFlags.LightMapDensity && AllowDebugViewmodes())
 	{
 		// Override the base pass with the lightmap density pass if the viewmode is enabled.
diff --git a/Engine/Source/Runtime/Renderer/Private/DeferredShadingRenderer.h b/Engine/Source/Runtime/Renderer/Private/DeferredShadingRenderer.h
--- a/Engine/Source/Runtime/Renderer/Private/DeferredShadingRenderer.h
+++ b/Engine/Source/Runtime/Renderer/Private/DeferredShadingRenderer.h
@@ -64,7 +64,10 @@ public:
 	void SortBasePassStaticData(FVector ViewPosition);
 
     /** Renders the basepass for the dynamic data of a given View. */
-	bool RenderBasePassDynamicData(FRHICommandListImmediate& RHICmdList, FViewInfo& View);
+	void RenderBasePassDynamicData(FRHICommandList& RHICmdList, const FViewInfo& View, bool& bOutDirty);
+
+	/** Renders the basepass for the dynamic data of a given View, in parallel. */
+	FGraphEventRef RenderBasePassDynamicDataParallel(FViewInfo& View, FGraphEventRef SubmitChain, bool& bOutDirty);
 
     /** Renders the basepass for a given View. */
 	bool RenderBasePass(FRHICommandListImmediate& RHICmdList, FViewInfo& View);
diff --git a/Engine/Source/Runtime/Renderer/Private/DynamicPrimitiveDrawing.h b/Engine/Source/Runtime/Renderer/Private/DynamicPrimitiveDrawing.h
--- a/Engine/Source/Runtime/Renderer/Private/DynamicPrimitiveDrawing.h
+++ b/Engine/Source/Runtime/Renderer/Private/DynamicPrimitiveDrawing.h
@@ -25,7 +25,7 @@ public:
 	* @param bInIsVelocityRendering - rendering is occurring during velocity pass
 	*/
 	TDynamicPrimitiveDrawer(
-		FRHICommandListImmediate& InRHICmdList,
+		FRHICommandList& InRHICmdList,
 		const FViewInfo* InView,
 		const typename DrawingPolicyFactoryType::ContextType& InDrawingContext,
 		bool InPreFog,
@@ -110,7 +110,7 @@ public:
 private:
 
 	/** Command list to draw to. */
-	FRHICommandListImmediate& RHICmdList;
+	FRHICommandList& RHICmdList;
 
 	/** The view which is being rendered. */
 	const FViewInfo* const View;
@@ -159,7 +159,7 @@ private:
  */
 template<class DrawingPolicyFactoryType>
 bool DrawViewElements(
-	FRHICommandListImmediate& RHICmdList,
+	FRHICommandList& RHICmdList,
 	const FViewInfo& View,
 	const typename DrawingPolicyFactoryType::ContextType& DrawingContext,
 	uint8 DPGIndex,
diff --git a/Engine/Source/Runtime/Renderer/Private/DynamicPrimitiveDrawing.inl b/Engine/Source/Runtime/Renderer/Private/DynamicPrimitiveDrawing.inl
--- a/Engine/Source/Runtime/Renderer/Private/DynamicPrimitiveDrawing.inl
+++ b/Engine/Source/Runtime/Renderer/Private/DynamicPrimitiveDrawing.inl
@@ -214,7 +214,7 @@ void TDynamicPrimitiveDrawer<DrawingPolicyFactoryType>::DrawPoint(
 
 template<class DrawingPolicyFactoryType>
 bool DrawViewElements(
-	FRHICommandListImmediate& RHICmdList,
+	FRHICommandList& RHICmdList,
 	const FViewInfo& View,
 	const typename DrawingPolicyFactoryType::ContextType& DrawingContext,
 	uint8 DPGIndex,
diff --git a/Engine/Source/Runtime/Renderer/Private/StaticMeshDrawList.inl b/Engine/Source/Runtime/Renderer/Private/StaticMeshDrawList.inl
--- a/Engine/Source/Runtime/Renderer/Private/StaticMeshDrawList.inl
+++ b/Engine/Source/Runtime/Renderer/Private/StaticMeshDrawList.inl
@@ -351,12 +351,12 @@ public:
 	}
 };
 
-class FDrawVisibleRenderThreadTask
+class FSubmitCommandlistThreadTask
 {
 	FRHICommandList* RHICmdList;
 public:
 
-	FDrawVisibleRenderThreadTask(FRHICommandList* InRHICmdList)
+	FSubmitCommandlistThreadTask(FRHICommandList* InRHICmdList)
 		: RHICmdList(InRHICmdList)
 	{
 	}
@@ -417,7 +417,7 @@ FGraphEventRef TStaticMeshDrawList<DrawingPolicyType>::DrawVisibleParallel(
 				Prereqs.Add(SubmitChain);
 			}
 
-			SubmitChain = TGraphTask<FDrawVisibleRenderThreadTask>::CreateTask(&Prereqs, ENamedThreads::RenderThread).ConstructAndDispatchWhenReady(CmdList);
+			SubmitChain = TGraphTask<FSubmitCommandlistThreadTask>::CreateTask(&Prereqs, ENamedThreads::RenderThread).ConstructAndDispatchWhenReady(CmdList);
 
 		}
 	}
