* fixed TTP 341929 REGRESSION: CRASH: EDITOR: Turning on GBuffer Hints in the main viewport crashes the Editor

diff --git a/Engine/Source/Runtime/Renderer/Private/PostProcess/PostProcessGBufferHints.cpp b/Engine/Source/Runtime/Renderer/Private/PostProcess/PostProcessGBufferHints.cpp
--- a/Engine/Source/Runtime/Renderer/Private/PostProcess/PostProcessGBufferHints.cpp
+++ b/Engine/Source/Runtime/Renderer/Private/PostProcess/PostProcessGBufferHints.cpp
@@ -76,6 +76,14 @@ public:
 
 IMPLEMENT_SHADER_TYPE(,FPostProcessGBufferHintsPS,TEXT("PostProcessGBufferHints"),TEXT("MainPS"),SF_Pixel);
 
+
+FRCPassPostProcessGBufferHints::FRCPassPostProcessGBufferHints()
+{
+	// AdjustGBufferRefCount(-1) call is done when the pass gets executed
+	GSceneRenderTargets.AdjustGBufferRefCount(1);
+}
+
+
 void FRCPassPostProcessGBufferHints::Process(FRenderingCompositePassContext& Context)
 {
 	SCOPED_DRAW_EVENT(GBufferHints, DEC_SCENE_ITEMS);
@@ -172,6 +180,9 @@ void FRCPassPostProcessGBufferHints::Process(FRenderingCompositePassContext& Con
 	Canvas.Flush_RenderThread(Context.RHICmdList);
 
 	Context.RHICmdList.CopyToResolveTarget(DestRenderTarget.TargetableTexture, DestRenderTarget.ShaderResourceTexture, false, FResolveParams());
+
+	// AdjustGBufferRefCount(1) call is done in constructor
+	GSceneRenderTargets.AdjustGBufferRefCount(-1);
 }
 
 FPooledRenderTargetDesc FRCPassPostProcessGBufferHints::ComputeOutputDesc(EPassOutputId InPassOutputId) const
diff --git a/Engine/Source/Runtime/Renderer/Private/PostProcess/PostProcessGBufferHints.h b/Engine/Source/Runtime/Renderer/Private/PostProcess/PostProcessGBufferHints.h
--- a/Engine/Source/Runtime/Renderer/Private/PostProcess/PostProcessGBufferHints.h
+++ b/Engine/Source/Runtime/Renderer/Private/PostProcess/PostProcessGBufferHints.h
@@ -14,6 +14,9 @@
 class FRCPassPostProcessGBufferHints : public TRenderingCompositePassBase<2, 1>
 {
 public:
+	// constructor
+	FRCPassPostProcessGBufferHints();
+
 	// interface FRenderingCompositePass ---------
 	virtual void Process(FRenderingCompositePassContext& Context);
 	virtual void Release() override { delete this; }
