UE4 - parallel rendering - fix a few bugs with shader creation and static states, do prepass in parallel

diff --git a/Engine/Source/Runtime/Launch/Private/LaunchEngineLoop.cpp b/Engine/Source/Runtime/Launch/Private/LaunchEngineLoop.cpp
--- a/Engine/Source/Runtime/Launch/Private/LaunchEngineLoop.cpp
+++ b/Engine/Source/Runtime/Launch/Private/LaunchEngineLoop.cpp
@@ -2031,6 +2031,7 @@ void FEngineLoop::Tick()
 		ENQUEUE_UNIQUE_RENDER_COMMAND(
 				BeginFrame,
 			{
+				GRHICommandList.GetImmediateCommandList().Flush();
 				GRHICommandList.LatchBypass();
 				GFrameNumberRenderThread++;
 				GDynamicRHI->PushEvent(*FString::Printf(TEXT("Frame%d"),GFrameNumberRenderThread));
diff --git a/Engine/Source/Runtime/RHI/Private/DynamicRHI.cpp b/Engine/Source/Runtime/RHI/Private/DynamicRHI.cpp
--- a/Engine/Source/Runtime/RHI/Private/DynamicRHI.cpp
+++ b/Engine/Source/Runtime/RHI/Private/DynamicRHI.cpp
@@ -37,8 +37,7 @@ void RHIInit(bool bHasEditorToken)
 {
 	if(!GDynamicRHI)
 	{
-		const TCHAR* CmdLine = FCommandLine::Get();
-		FString Token = FParse::Token(CmdLine, false);		
+		GRHICommandList.LatchBypass(); // read commandline for bypass flag
 
 		if(USE_NULL_RHI || FParse::Param(FCommandLine::Get(),TEXT("nullrhi")) || IsRunningCommandlet() || IsRunningDedicatedServer())
 		{
diff --git a/Engine/Source/Runtime/RHI/Private/RHICommandList.cpp b/Engine/Source/Runtime/RHI/Private/RHICommandList.cpp
--- a/Engine/Source/Runtime/RHI/Private/RHICommandList.cpp
+++ b/Engine/Source/Runtime/RHI/Private/RHICommandList.cpp
@@ -84,8 +84,20 @@ void FRHICommandListExecutor::ExecuteList(FRHICommandListImmediate& CmdList)
 void FRHICommandListExecutor::LatchBypass()
 {
 #if !(UE_BUILD_SHIPPING || UE_BUILD_TEST)
+
+	static bool bOnce = false;
+	if (!bOnce)
+	{
+		bOnce = true;
+		if (FParse::Param(FCommandLine::Get(),TEXT("parallelrendering")) && CVarRHICmdBypass.GetValueOnRenderThread() >= 1)
+		{
+			IConsoleVariable* BypassVar = IConsoleManager::Get().FindConsoleVariable(TEXT("r.RHICmdBypass"));
+			BypassVar->Set(0);
+		}
+	}
+
 	check(GRHICommandList.OutstandingCmdListCount.GetValue() == 1 && !GRHICommandList.GetImmediateCommandList().HasCommands());
-	bool NewBypass = (CVarRHICmdBypass.GetValueOnRenderThread() >= 1);
+	bool NewBypass = (CVarRHICmdBypass.GetValueOnAnyThread() >= 1);
 	if (NewBypass && !bLatchedBypass)
 	{
 		FRHIResource::FlushPendingDeletes();
diff --git a/Engine/Source/Runtime/RHI/Private/StaticRHI.cpp b/Engine/Source/Runtime/RHI/Private/StaticRHI.cpp
--- a/Engine/Source/Runtime/RHI/Private/StaticRHI.cpp
+++ b/Engine/Source/Runtime/RHI/Private/StaticRHI.cpp
@@ -19,6 +19,7 @@ void RHIInit(bool bHasEditorToken)
 {
 	if(!GStaticRHI)
 	{		
+		GRHICommandList.LatchBypass(); // read commandline for bypass flag
 		GStaticRHI = CreateStaticRHI();
 	}
 }
diff --git a/Engine/Source/Runtime/RHI/Public/RHICommandList.h b/Engine/Source/Runtime/RHI/Public/RHICommandList.h
--- a/Engine/Source/Runtime/RHI/Public/RHICommandList.h
+++ b/Engine/Source/Runtime/RHI/Public/RHICommandList.h
@@ -1632,7 +1632,7 @@ public:
 #if !(UE_BUILD_SHIPPING || UE_BUILD_TEST)
 		return bLatchedBypass;
 #else
-		return !!DefaultBypass;
+		return  !FApp::ShouldUseThreadingForPerformance() || !!DefaultBypass;
 #endif
 	}
 	static void CheckNoOutstandingCmdLists();
diff --git a/Engine/Source/Runtime/RHI/Public/RHIStaticStates.h b/Engine/Source/Runtime/RHI/Public/RHIStaticStates.h
--- a/Engine/Source/Runtime/RHI/Public/RHIStaticStates.h
+++ b/Engine/Source/Runtime/RHI/Public/RHIStaticStates.h
@@ -8,6 +8,43 @@
 
 #include "StaticArray.h"
 
+
+extern RHI_API FCriticalSection StaticStateRHICriticalSection;
+
+/**
+ * Helper task to initialize a static resource on the render thread.
+ */
+class FInitStaticResourceRenderThreadTask
+{
+	void (*DoConstruct)();
+	FScopedEvent& Event;
+public:
+
+	FInitStaticResourceRenderThreadTask(void (*InDoConstruct)(), FScopedEvent* InEvent)
+		: DoConstruct(InDoConstruct)
+		, Event(*InEvent)
+	{
+	}
+
+	FORCEINLINE TStatId GetStatId() const
+	{
+		RETURN_QUICK_DECLARE_CYCLE_STAT(FInitStaticResourceRenderThreadTask, STATGROUP_TaskGraphTasks);
+	}
+
+	ENamedThreads::Type GetDesiredThread()
+	{
+		return ENamedThreads::RenderThread_Local;
+	}
+
+	static ESubsequentsMode::Type GetSubsequentsMode() { return ESubsequentsMode::FireAndForget; }
+
+	void DoTask(ENamedThreads::Type CurrentThread, const FGraphEventRef& MyCompletionGraphEvent)
+	{
+		DoConstruct();
+		Event.Trigger();
+	}
+};
+
 /**
  * The base class of the static RHI state classes.
  */
@@ -16,20 +53,31 @@ class TStaticStateRHI
 {
 public:
 
+	static void GetRHI_WithNoReturnValue()
+	{
+		GetRHI();
+	}
 	static RHIParamRefType GetRHI()
 	{
-		static FStaticStateResource Resource;
+		// This is super-duper nasty. We rely upon the fact that all compilers will assign uninitialized static, POD data (like StaticResource) to a zero seg and not do any construction on it whatsoever.
+		static FStaticStateResource* StaticResource; // Must be left uninitialized!
 
-		static bool bRecreate = false;
-
-		// Set bRecreate to 1 to debug the resource creation
-		if (bRecreate)
+		if (!StaticResource)
 		{
-			bRecreate = false;
-			Resource = FStaticStateResource();
+			if (!IsInRenderingThread())
+			{
+				check(IsInParallelRenderingThread());
+				{
+					FScopedEvent Event;
+					TGraphTask<FInitStaticResourceRenderThreadTask>::CreateTask().ConstructAndDispatchWhenReady(&GetRHI_WithNoReturnValue, &Event);
+				}
+			}
+			else
+			{
+				StaticResource = new FStaticStateResource();
+			}
 		}
-
-		return Resource.StateRHI;
+		return StaticResource->StateRHI;
 	};
 
 private:
@@ -38,7 +86,6 @@ private:
 	class FStaticStateResource : public FRenderResource
 	{
 	public:
-
 		RHIRefType StateRHI;
 
 		FStaticStateResource()
diff --git a/Engine/Source/Runtime/RenderCore/Private/RenderResource.cpp b/Engine/Source/Runtime/RenderCore/Private/RenderResource.cpp
--- a/Engine/Source/Runtime/RenderCore/Private/RenderResource.cpp
+++ b/Engine/Source/Runtime/RenderCore/Private/RenderResource.cpp
@@ -28,6 +28,7 @@ void FRenderResource::InitResource()
 			InitDynamicRHI();
 			InitRHI();
 		}
+		FPlatformMisc::MemoryBarrier(); // there are some multithreaded reads of bInitialized
 		bInitialized = true;
 	}
 }
@@ -62,6 +63,53 @@ void FRenderResource::UpdateRHI()
 	}
 }
 
+void FRenderResource::InitResourceFromPossiblyParallelRendering()
+{
+	if (IsInRenderingThread())
+	{
+		InitResource();
+	}
+	else
+	{
+		check(IsInParallelRenderingThread());
+		class FInitResourceRenderThreadTask
+		{
+			FRenderResource& Resource;
+			FScopedEvent& Event;
+		public:
+
+			FInitResourceRenderThreadTask(FRenderResource* InResource, FScopedEvent* InEvent)
+				: Resource(*InResource)
+				, Event(*InEvent)
+			{
+			}
+
+			FORCEINLINE TStatId GetStatId() const
+			{
+				RETURN_QUICK_DECLARE_CYCLE_STAT(FInitResourceRenderThreadTask, STATGROUP_TaskGraphTasks);
+			}
+
+			ENamedThreads::Type GetDesiredThread()
+			{
+				return ENamedThreads::RenderThread_Local;
+			}
+
+			static ESubsequentsMode::Type GetSubsequentsMode() { return ESubsequentsMode::FireAndForget; }
+
+			void DoTask(ENamedThreads::Type CurrentThread, const FGraphEventRef& MyCompletionGraphEvent)
+			{
+				Resource.InitResource();
+				Event.Trigger();
+			}
+		};
+		{
+			FScopedEvent Event;
+			TGraphTask<FInitResourceRenderThreadTask>::CreateTask().ConstructAndDispatchWhenReady(this, &Event);
+		}
+	}
+}
+
+
 FRenderResource::~FRenderResource()
 {
 	if (bInitialized && !GIsCriticalError)
diff --git a/Engine/Source/Runtime/RenderCore/Public/RenderResource.h b/Engine/Source/Runtime/RenderCore/Public/RenderResource.h
--- a/Engine/Source/Runtime/RenderCore/Public/RenderResource.h
+++ b/Engine/Source/Runtime/RenderCore/Public/RenderResource.h
@@ -83,6 +83,9 @@ public:
 	 */
 	void UpdateRHI();
 
+	// Probably temporary code that sends a task back to renderthread_local and blocks waiting for it to call InitResource
+	void InitResourceFromPossiblyParallelRendering();
+
 	/** @return The resource's friendly name.  Typically a UObject name. */
 	virtual FString GetFriendlyName() const { return TEXT("undefined"); }
 
diff --git a/Engine/Source/Runtime/Renderer/Private/DeferredShadingRenderer.cpp b/Engine/Source/Runtime/Renderer/Private/DeferredShadingRenderer.cpp
--- a/Engine/Source/Runtime/Renderer/Private/DeferredShadingRenderer.cpp
+++ b/Engine/Source/Runtime/Renderer/Private/DeferredShadingRenderer.cpp
@@ -481,7 +481,7 @@ bool FDeferredShadingSceneRenderer::RenderBasePass(FRHICommandListImmediate& RHI
 
 	// Render the base pass static data
 	{
-		if (!GRHICommandList.Bypass() && FApp::ShouldUseThreadingForPerformance())
+		if (!GRHICommandList.Bypass())
 		{
 			int32 Width = CVarRHICmdWidth.GetValueOnRenderThread(); // we use a few more than needed to cover non-equal jobs
 
@@ -920,6 +920,213 @@ void FDeferredShadingSceneRenderer::Render(FRHICommandListImmediate& RHICmdList)
 	RenderFinish(RHICmdList);
 }
 
+bool FDeferredShadingSceneRenderer::RenderPrePassViewDynamic(FRHICommandList& RHICmdList, const FViewInfo& View)
+{
+	// Draw the dynamic occluder primitives using a depth drawing policy.
+	TDynamicPrimitiveDrawer<FDepthDrawingPolicyFactory> Drawer(RHICmdList, &View, FDepthDrawingPolicyFactory::ContextType(EarlyZPassMode), true);
+	{
+		SCOPED_DRAW_EVENT(Dynamic, DEC_SCENE_ITEMS);
+		for(int32 PrimitiveIndex = 0;PrimitiveIndex < View.VisibleDynamicPrimitives.Num();PrimitiveIndex++)
+		{
+			const FPrimitiveSceneInfo* PrimitiveSceneInfo = View.VisibleDynamicPrimitives[PrimitiveIndex];
+			int32 PrimitiveId = PrimitiveSceneInfo->GetIndex();
+			const FPrimitiveViewRelevance& PrimitiveViewRelevance = View.PrimitiveViewRelevanceMap[PrimitiveId];
+
+			bool bShouldUseAsOccluder = true;
+
+			if(EarlyZPassMode != DDM_AllOccluders)
+			{
+				extern float GMinScreenRadiusForDepthPrepass;
+				const float LODFactorDistanceSquared = (PrimitiveSceneInfo->Proxy->GetBounds().Origin - View.ViewMatrices.ViewOrigin).SizeSquared() * FMath::Square(View.LODDistanceFactor);
+
+				// Only render primitives marked as occluders
+				bShouldUseAsOccluder = PrimitiveSceneInfo->Proxy->ShouldUseAsOccluder()
+					// Only render static objects unless movable are requested
+					&& (!PrimitiveSceneInfo->Proxy->IsMovable() || GEarlyZPassMovable)
+					&& (FMath::Square(PrimitiveSceneInfo->Proxy->GetBounds().SphereRadius) > GMinScreenRadiusForDepthPrepass * GMinScreenRadiusForDepthPrepass * LODFactorDistanceSquared);
+			}
+
+			// Only render opaque primitives marked as occluders
+			if (bShouldUseAsOccluder && PrimitiveViewRelevance.bOpaqueRelevance && PrimitiveViewRelevance.bRenderInMainPass)
+			{
+				FScopeCycleCounter Context(PrimitiveSceneInfo->Proxy->GetStatId());
+				Drawer.SetPrimitive(PrimitiveSceneInfo->Proxy);
+				PrimitiveSceneInfo->Proxy->DrawDynamicElements(
+					&Drawer,
+					&View
+					);
+			}
+		}
+	}
+	return Drawer.IsDirty();
+}
+
+bool FDeferredShadingSceneRenderer::RenderPrePassView(FRHICommandList& RHICmdList, const FViewInfo& View)
+{
+	bool bDirty = false;
+
+	// Disable color writes, enable depth tests and writes.
+	RHICmdList.SetBlendState(TStaticBlendState<CW_NONE>::GetRHI());
+	// Note, this is a reversed Z depth surface, using CF_GreaterEqual.
+	RHICmdList.SetDepthStencilState(TStaticDepthStencilState<true,CF_GreaterEqual>::GetRHI());
+	RHICmdList.SetViewport(View.ViewRect.Min.X, View.ViewRect.Min.Y, 0, View.ViewRect.Max.X, View.ViewRect.Max.Y, 1);
+
+	// Draw the static occluder primitives using a depth drawing policy.
+	{
+		// Draw opaque occluders which support a separate position-only
+		// vertex buffer to minimize vertex fetch bandwidth, which is
+		// often the bottleneck during the depth only pass.
+		SCOPED_DRAW_EVENT(PosOnlyOpaque, DEC_SCENE_ITEMS);
+		bDirty |= Scene->PositionOnlyDepthDrawList.DrawVisible(RHICmdList, View,View.StaticMeshOccluderMap,View.StaticMeshBatchVisibility);
+	}
+	{
+		// Draw opaque occluders, using double speed z where supported.
+		SCOPED_DRAW_EVENT(Opaque, DEC_SCENE_ITEMS);
+		bDirty |= Scene->DepthDrawList.DrawVisible(RHICmdList, View,View.StaticMeshOccluderMap,View.StaticMeshBatchVisibility);
+	}
+
+	if(EarlyZPassMode >= DDM_AllOccluders)
+	{
+		// Draw opaque occluders with masked materials
+		SCOPED_DRAW_EVENT(Opaque, DEC_SCENE_ITEMS);
+		bDirty |= Scene->MaskedDepthDrawList.DrawVisible(RHICmdList, View,View.StaticMeshOccluderMap,View.StaticMeshBatchVisibility);
+	}
+
+	bDirty |= RenderPrePassViewDynamic(RHICmdList, View);
+	return bDirty;
+}
+
+class FSubmitCommandlistSetupPrepassView
+{
+	const FViewInfo& View;
+public:
+
+	FSubmitCommandlistSetupPrepassView(const FViewInfo* InView)
+		: View(*InView)
+	{
+	}
+
+	FORCEINLINE TStatId GetStatId() const
+	{
+		RETURN_QUICK_DECLARE_CYCLE_STAT(FSubmitCommandlistSetupPrepassView, STATGROUP_TaskGraphTasks);
+	}
+
+	ENamedThreads::Type GetDesiredThread()
+	{
+		return ENamedThreads::RenderThread_Local;
+	}
+
+	static ESubsequentsMode::Type GetSubsequentsMode() { return ESubsequentsMode::TrackSubsequents; }
+
+	void DoTask(ENamedThreads::Type CurrentThread, const FGraphEventRef& MyCompletionGraphEvent)
+	{
+		FRHICommandList RHICmdList;
+		// Disable color writes, enable depth tests and writes.
+		RHICmdList.SetBlendState(TStaticBlendState<CW_NONE>::GetRHI());
+		// Note, this is a reversed Z depth surface, using CF_GreaterEqual.
+		RHICmdList.SetDepthStencilState(TStaticDepthStencilState<true,CF_GreaterEqual>::GetRHI());
+		RHICmdList.SetViewport(View.ViewRect.Min.X, View.ViewRect.Min.Y, 0, View.ViewRect.Max.X, View.ViewRect.Max.Y, 1);
+	}
+};
+
+class FRenderPrepassDynamicDataThreadTask
+{
+	FDeferredShadingSceneRenderer& ThisRenderer;
+	FRHICommandList& RHICmdList;
+	const FViewInfo& View;
+	bool& OutDirty;
+
+public:
+
+	FRenderPrepassDynamicDataThreadTask(
+		FDeferredShadingSceneRenderer* InThisRenderer,
+		FRHICommandList* InRHICmdList,
+		const FViewInfo* InView,
+		bool* InOutDirty
+		)
+		: ThisRenderer(*InThisRenderer)
+		, RHICmdList(*InRHICmdList)
+		, View(*InView)
+		, OutDirty(*InOutDirty)
+	{
+	}
+
+	FORCEINLINE TStatId GetStatId() const
+	{
+		RETURN_QUICK_DECLARE_CYCLE_STAT(FRenderPrepassDynamicDataThreadTask, STATGROUP_TaskGraphTasks);
+	}
+
+	ENamedThreads::Type GetDesiredThread()
+	{
+		return ENamedThreads::AnyThread;
+	}
+
+	static ESubsequentsMode::Type GetSubsequentsMode() { return ESubsequentsMode::TrackSubsequents; }
+
+	void DoTask(ENamedThreads::Type CurrentThread, const FGraphEventRef& MyCompletionGraphEvent)
+	{
+		if (ThisRenderer.RenderPrePassViewDynamic(RHICmdList, View))
+		{
+			OutDirty = true;
+		}
+	}
+};
+
+
+FGraphEventRef FDeferredShadingSceneRenderer::RenderPrePassViewParallel(FRHICommandList& RHICmdList, const FViewInfo& View, int32 Width, FGraphEventRef SubmitChain, bool& OutDirty)
+{
+	{
+		//optimization: this is probably only needed on the first view
+		FGraphEventArray Prereqs;
+		if (SubmitChain.GetReference())
+		{
+			Prereqs.Add(SubmitChain);
+		}
+		SubmitChain = TGraphTask<FSubmitCommandlistSetupPrepassView>::CreateTask(&Prereqs, ENamedThreads::RenderThread).ConstructAndDispatchWhenReady(&View);
+
+	}
+
+	// Draw the static occluder primitives using a depth drawing policy.
+	{
+		// Draw opaque occluders which support a separate position-only
+		// vertex buffer to minimize vertex fetch bandwidth, which is
+		// often the bottleneck during the depth only pass.
+		//SCOPED_DRAW_EVENT(PosOnlyOpaque, DEC_SCENE_ITEMS);
+		SubmitChain = Scene->PositionOnlyDepthDrawList.DrawVisibleParallel(View, View.StaticMeshOccluderMap, View.StaticMeshBatchVisibility, Width, SubmitChain, OutDirty);
+	}
+	{
+		// Draw opaque occluders, using double speed z where supported.
+		//SCOPED_DRAW_EVENT(Opaque, DEC_SCENE_ITEMS);
+		SubmitChain = Scene->DepthDrawList.DrawVisibleParallel(View,View.StaticMeshOccluderMap,View.StaticMeshBatchVisibility, Width, SubmitChain, OutDirty);
+	}
+
+	if(EarlyZPassMode >= DDM_AllOccluders)
+	{
+		// Draw opaque occluders with masked materials
+		//SCOPED_DRAW_EVENT(Opaque, DEC_SCENE_ITEMS);
+		SubmitChain = Scene->MaskedDepthDrawList.DrawVisibleParallel(View,View.StaticMeshOccluderMap,View.StaticMeshBatchVisibility, Width, SubmitChain, OutDirty);
+	}
+
+	{
+		FRHICommandList* CmdList = new FRHICommandList;
+
+		FGraphEventRef AnyThreadCompletionEvent = TGraphTask<FRenderPrepassDynamicDataThreadTask>::CreateTask(nullptr, ENamedThreads::RenderThread)
+			.ConstructAndDispatchWhenReady(this, CmdList, &View, &OutDirty);
+
+		FGraphEventArray Prereqs;
+		Prereqs.Add(AnyThreadCompletionEvent);
+		if (SubmitChain.GetReference())
+		{
+			Prereqs.Add(SubmitChain);
+		}
+
+		SubmitChain = TGraphTask<FSubmitCommandlistThreadTask>::CreateTask(&Prereqs, ENamedThreads::RenderThread).ConstructAndDispatchWhenReady(CmdList);
+	}
+
+	return SubmitChain;
+}
+
+
 /** Renders the scene's prepass and occlusion queries */
 bool FDeferredShadingSceneRenderer::RenderPrePass(FRHICommandListImmediate& RHICmdList)
 {
@@ -938,76 +1145,31 @@ bool FDeferredShadingSceneRenderer::RenderPrePass(FRHICommandListImmediate& RHIC
 	// Draw a depth pass to avoid overdraw in the other passes.
 	if(EarlyZPassMode != DDM_None)
 	{
-		for(int32 ViewIndex = 0;ViewIndex < Views.Num();ViewIndex++)
+		if (!GRHICommandList.Bypass())
 		{
-			SCOPED_CONDITIONAL_DRAW_EVENTF(EventView, Views.Num() > 1, DEC_SCENE_ITEMS, TEXT("View%d"), ViewIndex);
-
-			const FViewInfo& View = Views[ViewIndex];
-
-			// Disable color writes, enable depth tests and writes.
-			RHICmdList.SetBlendState(TStaticBlendState<CW_NONE>::GetRHI());
-			// Note, this is a reversed Z depth surface, using CF_GreaterEqual.
-			RHICmdList.SetDepthStencilState(TStaticDepthStencilState<true,CF_GreaterEqual>::GetRHI());
-			RHICmdList.SetViewport(View.ViewRect.Min.X, View.ViewRect.Min.Y, 0, View.ViewRect.Max.X, View.ViewRect.Max.Y, 1);
+			int32 Width = CVarRHICmdWidth.GetValueOnRenderThread(); // we use a few more than needed to cover non-equal jobs
 
-			// Draw the static occluder primitives using a depth drawing policy.
+			FGraphEventRef SubmitChain;
+			for(int32 ViewIndex = 0;ViewIndex < Views.Num();ViewIndex++)
 			{
-				// Draw opaque occluders which support a separate position-only
-				// vertex buffer to minimize vertex fetch bandwidth, which is
-				// often the bottleneck during the depth only pass.
-				SCOPED_DRAW_EVENT(PosOnlyOpaque, DEC_SCENE_ITEMS);
-				bDirty |= Scene->PositionOnlyDepthDrawList.DrawVisible(RHICmdList, View,View.StaticMeshOccluderMap,View.StaticMeshBatchVisibility);
+				SCOPED_CONDITIONAL_DRAW_EVENTF(EventView, Views.Num() > 1, DEC_SCENE_ITEMS, TEXT("View%d"), ViewIndex);
+				const FViewInfo& View = Views[ViewIndex];
+				SubmitChain = RenderPrePassViewParallel(RHICmdList, View, Width, SubmitChain, bDirty);
 			}
+			
+			if (SubmitChain.GetReference())
 			{
-				// Draw opaque occluders, using double speed z where supported.
-				SCOPED_DRAW_EVENT(Opaque, DEC_SCENE_ITEMS);
-				bDirty |= Scene->DepthDrawList.DrawVisible(RHICmdList, View,View.StaticMeshOccluderMap,View.StaticMeshBatchVisibility);
-			}
-
-			if(EarlyZPassMode >= DDM_AllOccluders)
-			{
-				// Draw opaque occluders with masked materials
-				SCOPED_DRAW_EVENT(Opaque, DEC_SCENE_ITEMS);
-				bDirty |= Scene->MaskedDepthDrawList.DrawVisible(RHICmdList, View,View.StaticMeshOccluderMap,View.StaticMeshBatchVisibility);
+				FTaskGraphInterface::Get().WaitUntilTaskCompletes(SubmitChain, ENamedThreads::RenderThread_Local);
 			}
-
-			// Draw the dynamic occluder primitives using a depth drawing policy.
-			TDynamicPrimitiveDrawer<FDepthDrawingPolicyFactory> Drawer(RHICmdList, &View, FDepthDrawingPolicyFactory::ContextType(EarlyZPassMode), true);
+		}
+		else
+		{
+			for(int32 ViewIndex = 0;ViewIndex < Views.Num();ViewIndex++)
 			{
-				SCOPED_DRAW_EVENT(Dynamic, DEC_SCENE_ITEMS);
-				for(int32 PrimitiveIndex = 0;PrimitiveIndex < View.VisibleDynamicPrimitives.Num();PrimitiveIndex++)
-				{
-					const FPrimitiveSceneInfo* PrimitiveSceneInfo = View.VisibleDynamicPrimitives[PrimitiveIndex];
-					int32 PrimitiveId = PrimitiveSceneInfo->GetIndex();
-					const FPrimitiveViewRelevance& PrimitiveViewRelevance = View.PrimitiveViewRelevanceMap[PrimitiveId];
-
-					bool bShouldUseAsOccluder = true;
-					
-					if(EarlyZPassMode != DDM_AllOccluders)
-					{
-						extern float GMinScreenRadiusForDepthPrepass;
-						const float LODFactorDistanceSquared = (PrimitiveSceneInfo->Proxy->GetBounds().Origin - View.ViewMatrices.ViewOrigin).SizeSquared() * FMath::Square(View.LODDistanceFactor);
-
-						// Only render primitives marked as occluders
-						bShouldUseAsOccluder = PrimitiveSceneInfo->Proxy->ShouldUseAsOccluder()
-							// Only render static objects unless movable are requested
-							&& (!PrimitiveSceneInfo->Proxy->IsMovable() || GEarlyZPassMovable)
-							&& (FMath::Square(PrimitiveSceneInfo->Proxy->GetBounds().SphereRadius) > GMinScreenRadiusForDepthPrepass * GMinScreenRadiusForDepthPrepass * LODFactorDistanceSquared);
-					}
-
-					// Only render opaque primitives marked as occluders
-					if (bShouldUseAsOccluder && PrimitiveViewRelevance.bOpaqueRelevance && PrimitiveViewRelevance.bRenderInMainPass)
-					{
-						FScopeCycleCounter Context(PrimitiveSceneInfo->Proxy->GetStatId());
-						Drawer.SetPrimitive(PrimitiveSceneInfo->Proxy);
-						PrimitiveSceneInfo->Proxy->DrawDynamicElements(
-							&Drawer,
-							&View
-							);
-					}
-				}
+				SCOPED_CONDITIONAL_DRAW_EVENTF(EventView, Views.Num() > 1, DEC_SCENE_ITEMS, TEXT("View%d"), ViewIndex);
+				const FViewInfo& View = Views[ViewIndex];
+				bDirty |= RenderPrePassView(RHICmdList, View);
 			}
-			bDirty |= Drawer.IsDirty();
 		}
 	}
 
diff --git a/Engine/Source/Runtime/Renderer/Private/DeferredShadingRenderer.h b/Engine/Source/Runtime/Renderer/Private/DeferredShadingRenderer.h
--- a/Engine/Source/Runtime/Renderer/Private/DeferredShadingRenderer.h
+++ b/Engine/Source/Runtime/Renderer/Private/DeferredShadingRenderer.h
@@ -50,6 +50,24 @@ public:
 	/** Propagates LPVs for all views */
 	void PropagateLPVs(FRHICommandListImmediate& RHICmdList);
 
+	/**
+	 * Renders the dynamic scene's prepass for a particular view
+	 * @return true if anything was rendered
+	 */
+	bool RenderPrePassViewDynamic(FRHICommandList& RHICmdList, const FViewInfo& View);
+
+	/**
+	 * Renders the scene's prepass for a particular view
+	 * @return true if anything was rendered
+	 */
+	bool RenderPrePassView(FRHICommandList& RHICmdList, const FViewInfo& View);
+
+	/**
+	 * Renders the scene's prepass for a particular view in parallel
+	 * @return the submit chain
+	 */
+	FGraphEventRef RenderPrePassViewParallel(FRHICommandList& RHICmdList, const FViewInfo& View, int32 Width, FGraphEventRef SubmitChain, bool& OutDirty);
+
     /** Renders the basepass for the static data of a given View. */
     bool RenderBasePassStaticData(FRHICommandList& RHICmdList, FViewInfo& View);
 	bool RenderBasePassStaticDataMasked(FRHICommandList& RHICmdList, FViewInfo& View);
diff --git a/Engine/Source/Runtime/Renderer/Private/StaticMeshDrawList.inl b/Engine/Source/Runtime/Renderer/Private/StaticMeshDrawList.inl
--- a/Engine/Source/Runtime/Renderer/Private/StaticMeshDrawList.inl
+++ b/Engine/Source/Runtime/Renderer/Private/StaticMeshDrawList.inl
@@ -325,8 +325,6 @@ public:
 
 	void DoTask(ENamedThreads::Type CurrentThread, const FGraphEventRef& MyCompletionGraphEvent)
 	{
-		QUICK_SCOPE_CYCLE_COUNTER(AA_FDrawVisibleAnyThreadTask);
-
 		for (int32 Index = this->FirstPolicy; Index <= this->LastPolicy; Index++)
 		{
 			typename TStaticMeshDrawList<DrawingPolicyType>::FDrawingPolicyLink* DrawingPolicyLink = &this->Caller.DrawingPolicySet[this->Caller.OrderedDrawingPolicies[Index]];
diff --git a/Engine/Source/Runtime/ShaderCore/Private/Shader.cpp b/Engine/Source/Runtime/ShaderCore/Private/Shader.cpp
--- a/Engine/Source/Runtime/ShaderCore/Private/Shader.cpp
+++ b/Engine/Source/Runtime/ShaderCore/Private/Shader.cpp
@@ -457,7 +457,6 @@ void FShaderResource::ReleaseRHI()
 	ComputeShader.SafeRelease();
 }
 
-
 void FShaderResource::InitializeVertexShaderRHI() 
 { 
 	if (!IsInitialized())
@@ -467,7 +466,7 @@ void FShaderResource::InitializeVertexShaderRHI()
 			SCOPE_CYCLE_COUNTER(STAT_Shaders_FrameRTShaderInitForRenderingTime);
 			SCOPE_SECONDS_COUNTER(ShaderInitializationTime);
 
-			InitResource();
+			InitResourceFromPossiblyParallelRendering();
 
 		}
 
@@ -487,7 +486,7 @@ void FShaderResource::InitializePixelShaderRHI()
 			SCOPE_CYCLE_COUNTER(STAT_Shaders_FrameRTShaderInitForRenderingTime);
 			SCOPE_SECONDS_COUNTER(ShaderInitializationTime);
 
-			InitResource();
+			InitResourceFromPossiblyParallelRendering();
 		}
 
 		INC_FLOAT_STAT_BY(STAT_Shaders_TotalRTShaderInitForRenderingTime,(float)ShaderInitializationTime);
@@ -507,7 +506,7 @@ const FHullShaderRHIRef& FShaderResource::GetHullShader()
 			SCOPE_CYCLE_COUNTER(STAT_Shaders_FrameRTShaderInitForRenderingTime);
 			SCOPE_SECONDS_COUNTER(ShaderInitializationTime);
 
-			InitResource();
+			InitResourceFromPossiblyParallelRendering();
 		}
 
 		INC_FLOAT_STAT_BY(STAT_Shaders_TotalRTShaderInitForRenderingTime,(float)ShaderInitializationTime);
@@ -530,7 +529,7 @@ const FDomainShaderRHIRef& FShaderResource::GetDomainShader()
 			SCOPE_CYCLE_COUNTER(STAT_Shaders_FrameRTShaderInitForRenderingTime);
 			SCOPE_SECONDS_COUNTER(ShaderInitializationTime);
 
-			InitResource();
+			InitResourceFromPossiblyParallelRendering();
 		}
 		
 		INC_FLOAT_STAT_BY(STAT_Shaders_TotalRTShaderInitForRenderingTime,(float)ShaderInitializationTime);
@@ -553,7 +552,7 @@ const FGeometryShaderRHIRef& FShaderResource::GetGeometryShader()
 			SCOPE_CYCLE_COUNTER(STAT_Shaders_FrameRTShaderInitForRenderingTime);
 			SCOPE_SECONDS_COUNTER(ShaderInitializationTime);
 
-			InitResource();
+			InitResourceFromPossiblyParallelRendering();
 		}
 
 		INC_FLOAT_STAT_BY(STAT_Shaders_TotalRTShaderInitForRenderingTime,(float)ShaderInitializationTime);
@@ -576,7 +575,7 @@ const FComputeShaderRHIRef& FShaderResource::GetComputeShader()
 			SCOPE_CYCLE_COUNTER(STAT_Shaders_FrameRTShaderInitForRenderingTime);
 			SCOPE_SECONDS_COUNTER(ShaderInitializationTime);
 
-			InitResource();
+			InitResourceFromPossiblyParallelRendering();
 		}
 
 		INC_FLOAT_STAT_BY(STAT_Shaders_TotalRTShaderInitForRenderingTime,(float)ShaderInitializationTime);
