UE4 - parallel rendering, fix bug with dirty, slight refactor to avoid copy-paste, now do all views in the basepass before waiting

diff --git a/Engine/Source/Runtime/Renderer/Private/DeferredShadingRenderer.cpp b/Engine/Source/Runtime/Renderer/Private/DeferredShadingRenderer.cpp
--- a/Engine/Source/Runtime/Renderer/Private/DeferredShadingRenderer.cpp
+++ b/Engine/Source/Runtime/Renderer/Private/DeferredShadingRenderer.cpp
@@ -471,35 +471,82 @@ FGraphEventRef FDeferredShadingSceneRenderer::RenderBasePassDynamicDataParallel(
 	return TGraphTask<FSubmitCommandlistThreadTask>::CreateTask(&Prereqs, ENamedThreads::RenderThread).ConstructAndDispatchWhenReady(CmdList);
 }
 
-/**
-* Renders the basepass for a given DPG and View.
-* @return true if anything was rendered to scene color
-*/
-bool FDeferredShadingSceneRenderer::RenderBasePass(FRHICommandListImmediate& RHICmdList, FViewInfo& View)
+static void SetupBasePassView(FRHICommandList& RHICmdList, const FViewInfo& View, bool bShaderComplexity)
 {
-	bool bDirty = false; 
+	if (bShaderComplexity)
+	{
+		// Additive blending when shader complexity viewmode is enabled.
+		RHICmdList.SetBlendState(TStaticBlendState<CW_RGBA,BO_Add,BF_One,BF_One,BO_Add,BF_Zero,BF_One>::GetRHI());
+		// Disable depth writes as we have a full depth prepass.
+		RHICmdList.SetDepthStencilState(TStaticDepthStencilState<false,CF_GreaterEqual>::GetRHI());
+	}
+	else
+	{
+		// Opaque blending for all G buffer targets, depth tests and writes.
+		RHICmdList.SetBlendState(TStaticBlendStateWriteMask<CW_RGBA, CW_RGBA, CW_RGBA, CW_RGBA>::GetRHI());
+		// Note, this is a reversed Z depth surface, using CF_GreaterEqual.
+		RHICmdList.SetDepthStencilState(TStaticDepthStencilState<true,CF_GreaterEqual>::GetRHI());
+	}
+	RHICmdList.SetViewport(View.ViewRect.Min.X, View.ViewRect.Min.Y, 0, View.ViewRect.Max.X, View.ViewRect.Max.Y, 1);
+}
+
+class FSubmitCommandlistSetupBasePassView
+{
+	const FViewInfo& View;
+	bool bShaderComplexity;
+public:
 
-	// Render the base pass static data
+	FSubmitCommandlistSetupBasePassView(const FViewInfo* InView, bool bInShaderComplexity)
+		: View(*InView)
+		, bShaderComplexity(bInShaderComplexity)
 	{
-		if (!GRHICommandList.Bypass())
-		{
-			int32 Width = CVarRHICmdWidth.GetValueOnRenderThread(); // we use a few more than needed to cover non-equal jobs
+	}
+
+	FORCEINLINE TStatId GetStatId() const
+	{
+		RETURN_QUICK_DECLARE_CYCLE_STAT(FSubmitCommandlistSetupBasePassView, STATGROUP_TaskGraphTasks);
+	}
 
-			FGraphEventRef SubmitChain = RenderBasePassStaticDataParallel(View, Width, FGraphEventRef(), bDirty);
-			SubmitChain = RenderBasePassDynamicDataParallel(View, SubmitChain, bDirty);
+	ENamedThreads::Type GetDesiredThread()
+	{
+		return ENamedThreads::RenderThread_Local;
+	}
 
-			if (SubmitChain.GetReference())
-			{
-				FTaskGraphInterface::Get().WaitUntilTaskCompletes(SubmitChain, ENamedThreads::RenderThread_Local);
-			}
-		}
-		else
+	static ESubsequentsMode::Type GetSubsequentsMode() { return ESubsequentsMode::TrackSubsequents; }
+
+	void DoTask(ENamedThreads::Type CurrentThread, const FGraphEventRef& MyCompletionGraphEvent)
+	{
+		FRHICommandList RHICmdList;
+		SetupBasePassView(RHICmdList, View, bShaderComplexity);
+	}
+};
+
+FGraphEventRef FDeferredShadingSceneRenderer::RenderBasePassViewParallel(FRHICommandListImmediate& RHICmdList, FViewInfo& View, int32 Width, FGraphEventRef SubmitChain, bool& OutDirty)
+{
+	{
+		//optimization: this is probably only needed on the first view
+		FGraphEventArray Prereqs;
+		if (SubmitChain.GetReference())
 		{
-			bool bDirty = false;
-			bDirty |= RenderBasePassStaticData(RHICmdList, View);
-			RenderBasePassDynamicData(RHICmdList, View, bDirty);
+			Prereqs.Add(SubmitChain);
 		}
-	}
+		SubmitChain = TGraphTask<FSubmitCommandlistSetupBasePassView>::CreateTask(&Prereqs, ENamedThreads::RenderThread).ConstructAndDispatchWhenReady(&View, !!ViewFamily.EngineShowFlags.ShaderComplexity);
+
+	}			
+
+	SubmitChain = RenderBasePassStaticDataParallel(View, Width, SubmitChain, OutDirty);
+	SubmitChain = RenderBasePassDynamicDataParallel(View, SubmitChain, OutDirty);
+
+	return SubmitChain;
+}
+
+bool FDeferredShadingSceneRenderer::RenderBasePassView(FRHICommandListImmediate& RHICmdList, FViewInfo& View)
+{
+	bool bDirty = false; 
+
+	SetupBasePassView(RHICmdList, View, ViewFamily.EngineShowFlags.ShaderComplexity);
+	bDirty |= RenderBasePassStaticData(RHICmdList, View);
+	RenderBasePassDynamicData(RHICmdList, View, bDirty);
 
 	return bDirty;
 }
@@ -961,15 +1008,20 @@ bool FDeferredShadingSceneRenderer::RenderPrePassViewDynamic(FRHICommandList& RH
 	return Drawer.IsDirty();
 }
 
-bool FDeferredShadingSceneRenderer::RenderPrePassView(FRHICommandList& RHICmdList, const FViewInfo& View)
+static void SetupPrePassView(FRHICommandList& RHICmdList, const FViewInfo& View)
 {
-	bool bDirty = false;
-
 	// Disable color writes, enable depth tests and writes.
 	RHICmdList.SetBlendState(TStaticBlendState<CW_NONE>::GetRHI());
 	// Note, this is a reversed Z depth surface, using CF_GreaterEqual.
 	RHICmdList.SetDepthStencilState(TStaticDepthStencilState<true,CF_GreaterEqual>::GetRHI());
 	RHICmdList.SetViewport(View.ViewRect.Min.X, View.ViewRect.Min.Y, 0, View.ViewRect.Max.X, View.ViewRect.Max.Y, 1);
+}
+
+bool FDeferredShadingSceneRenderer::RenderPrePassView(FRHICommandList& RHICmdList, const FViewInfo& View)
+{
+	bool bDirty = false;
+
+	SetupPrePassView(RHICmdList, View);
 
 	// Draw the static occluder primitives using a depth drawing policy.
 	{
@@ -1021,11 +1073,7 @@ public:
 	void DoTask(ENamedThreads::Type CurrentThread, const FGraphEventRef& MyCompletionGraphEvent)
 	{
 		FRHICommandList RHICmdList;
-		// Disable color writes, enable depth tests and writes.
-		RHICmdList.SetBlendState(TStaticBlendState<CW_NONE>::GetRHI());
-		// Note, this is a reversed Z depth surface, using CF_GreaterEqual.
-		RHICmdList.SetDepthStencilState(TStaticDepthStencilState<true,CF_GreaterEqual>::GetRHI());
-		RHICmdList.SetViewport(View.ViewRect.Min.X, View.ViewRect.Min.Y, 0, View.ViewRect.Max.X, View.ViewRect.Max.Y, 1);
+		SetupPrePassView(RHICmdList, View);
 	}
 };
 
@@ -1196,30 +1244,34 @@ bool FDeferredShadingSceneRenderer::RenderBasePass(FRHICommandListImmediate& RHI
 		SCOPED_DRAW_EVENT(BasePass, DEC_SCENE_ITEMS);
 		SCOPE_CYCLE_COUNTER(STAT_BasePassDrawTime);
 
-		// Draw the scene's emissive and light-map color.
-		for(int32 ViewIndex = 0;ViewIndex < Views.Num();ViewIndex++)
+		if (!GRHICommandList.Bypass())
 		{
-			SCOPED_CONDITIONAL_DRAW_EVENTF(EventView, Views.Num() > 1, DEC_SCENE_ITEMS, TEXT("View%d"), ViewIndex);
-			FViewInfo& View = Views[ViewIndex];
+			int32 Width = CVarRHICmdWidth.GetValueOnRenderThread(); // we use a few more than needed to cover non-equal jobs
+			FGraphEventRef SubmitChain;
 
-			if (ViewFamily.EngineShowFlags.ShaderComplexity)
+			for(int32 ViewIndex = 0;ViewIndex < Views.Num();ViewIndex++)
 			{
-				// Additive blending when shader complexity viewmode is enabled.
-				RHICmdList.SetBlendState(TStaticBlendState<CW_RGBA,BO_Add,BF_One,BF_One,BO_Add,BF_Zero,BF_One>::GetRHI());
-				// Disable depth writes as we have a full depth prepass.
-				RHICmdList.SetDepthStencilState(TStaticDepthStencilState<false,CF_GreaterEqual>::GetRHI());
+				SCOPED_CONDITIONAL_DRAW_EVENTF(EventView, Views.Num() > 1, DEC_SCENE_ITEMS, TEXT("View%d"), ViewIndex);
+				FViewInfo& View = Views[ViewIndex];
+				SubmitChain = RenderBasePassViewParallel(RHICmdList, View, Width, SubmitChain, bDirty);
 			}
-			else
+
+			if (SubmitChain.GetReference())
 			{
-				// Opaque blending for all G buffer targets, depth tests and writes.
-				RHICmdList.SetBlendState(TStaticBlendStateWriteMask<CW_RGBA, CW_RGBA, CW_RGBA, CW_RGBA>::GetRHI());
-				// Note, this is a reversed Z depth surface, using CF_GreaterEqual.
-				RHICmdList.SetDepthStencilState(TStaticDepthStencilState<true,CF_GreaterEqual>::GetRHI());
+				FTaskGraphInterface::Get().WaitUntilTaskCompletes(SubmitChain, ENamedThreads::RenderThread_Local);
 			}
-			RHICmdList.SetViewport(View.ViewRect.Min.X, View.ViewRect.Min.Y, 0, View.ViewRect.Max.X, View.ViewRect.Max.Y, 1);
+		}
+		else
+		{
+			for(int32 ViewIndex = 0;ViewIndex < Views.Num();ViewIndex++)
+			{
+				SCOPED_CONDITIONAL_DRAW_EVENTF(EventView, Views.Num() > 1, DEC_SCENE_ITEMS, TEXT("View%d"), ViewIndex);
+				FViewInfo& View = Views[ViewIndex];
 
-			bDirty |= RenderBasePass(RHICmdList, View);
+				bDirty |= RenderBasePassView(RHICmdList, View);
+			}
 		}
+
 	}
 
 	return bDirty;
diff --git a/Engine/Source/Runtime/Renderer/Private/DeferredShadingRenderer.h b/Engine/Source/Runtime/Renderer/Private/DeferredShadingRenderer.h
--- a/Engine/Source/Runtime/Renderer/Private/DeferredShadingRenderer.h
+++ b/Engine/Source/Runtime/Renderer/Private/DeferredShadingRenderer.h
@@ -87,8 +87,11 @@ public:
 	/** Renders the basepass for the dynamic data of a given View, in parallel. */
 	FGraphEventRef RenderBasePassDynamicDataParallel(FViewInfo& View, FGraphEventRef SubmitChain, bool& bOutDirty);
 
-    /** Renders the basepass for a given View. */
-	bool RenderBasePass(FRHICommandListImmediate& RHICmdList, FViewInfo& View);
+    /** Renders the basepass for a given View, in parallel */
+	FGraphEventRef RenderBasePassViewParallel(FRHICommandListImmediate& RHICmdList, FViewInfo& View, int32 Width, FGraphEventRef SubmitChain, bool& OutDirty);
+
+	/** Renders the basepass for a given View. */
+	bool RenderBasePassView(FRHICommandListImmediate& RHICmdList, FViewInfo& View);
 	
 	/** 
 	* Renders the scene's base pass 
