Fixed indirect lighting scale in tiled deferred.

diff --git a/Engine/Source/Runtime/Renderer/Private/TiledDeferredLightRendering.cpp b/Engine/Source/Runtime/Renderer/Private/TiledDeferredLightRendering.cpp
--- a/Engine/Source/Runtime/Renderer/Private/TiledDeferredLightRendering.cpp
+++ b/Engine/Source/Runtime/Renderer/Private/TiledDeferredLightRendering.cpp
@@ -174,6 +174,14 @@ public:
 					LightData.LightColorAndFalloffExponent[LightIndex].W = 0;
 				}
 
+				// When rendering reflection captures, the direct lighting of the light is actually the indirect specular from the main view
+				if (View.bIsReflectionCapture)
+				{
+					LightData.LightColorAndFalloffExponent[LightIndex].X *= LightSceneInfo->Proxy->GetIndirectLightingScale();
+					LightData.LightColorAndFalloffExponent[LightIndex].Y *= LightSceneInfo->Proxy->GetIndirectLightingScale();
+					LightData.LightColorAndFalloffExponent[LightIndex].Z *= LightSceneInfo->Proxy->GetIndirectLightingScale();
+				}
+
 				{
 					// SpotlightMaskAndMinRoughness, >0:Spotlight, MinRoughness = abs();
 					float W = FMath::Max(0.0001f, MinRoughness) * ((LightSceneInfo->Proxy->GetLightType() == LightType_Spot) ? 1 : -1);
