Enable Hit proxy with ES2 emulation.

ES2 emulation 'ShouldCache' hit proxy shaders.
Add static meshes to hitproxy drawing policy factory
Moved RenderHitProxies implementation out of FDeferredShadingSceneRenderer scope and split in to two so classes derrived from FSceneRenderer can call their respective InitViews(RHICmdList); in the correct place.

diff --git a/Engine/Source/Runtime/Renderer/Private/SceneCore.cpp b/Engine/Source/Runtime/Renderer/Private/SceneCore.cpp
--- a/Engine/Source/Runtime/Renderer/Private/SceneCore.cpp
+++ b/Engine/Source/Runtime/Renderer/Private/SceneCore.cpp
@@ -277,52 +277,55 @@ void FStaticMesh::UnlinkDrawList(FStaticMesh::FDrawListElementLink* Link)
 void FStaticMesh::AddToDrawLists(FRHICommandListImmediate& RHICmdList, FScene* Scene)
 {
 	const auto FeatureLevel = Scene->GetFeatureLevel();
-	if (FeatureLevel >= ERHIFeatureLevel::SM3)
+
+	if (CastShadow && FeatureLevel >= ERHIFeatureLevel::SM3)
+	{
+		FShadowDepthDrawingPolicyFactory::AddStaticMesh(Scene, this);
+	}
+
+	if (!bShadowOnly && PrimitiveSceneInfo->Proxy->ShouldRenderInMainPass())
 	{
-		if (CastShadow)
+		const bool bRequiresHitProxies = Scene->RequiresHitProxies();
+		if (bRequiresHitProxies && PrimitiveSceneInfo->Proxy->IsSelectable())
 		{
-			FShadowDepthDrawingPolicyFactory::AddStaticMesh(Scene,this);
+			// Add the static mesh to the DPG's hit proxy draw list.
+			FHitProxyDrawingPolicyFactory::AddStaticMesh(Scene, this);
 		}
+	}
+	else
+	{
+		return;
+	}
 
-		if (!bShadowOnly && PrimitiveSceneInfo->Proxy->ShouldRenderInMainPass())
-		{
-			// not all platforms need this
-			const bool bRequiresHitProxies = Scene->RequiresHitProxies();
-			if ( bRequiresHitProxies && PrimitiveSceneInfo->Proxy->IsSelectable() )
-			{
-				// Add the static mesh to the DPG's hit proxy draw list.
-				FHitProxyDrawingPolicyFactory::AddStaticMesh(Scene,this);
-			}
+	if (IsTranslucent(FeatureLevel))
+	{
+		return;
+	}
 
-			if (!IsTranslucent(FeatureLevel))
-			{
-				extern TAutoConsoleVariable<int32> CVarEarlyZPass;
-				int32 EarlyZPass = CVarEarlyZPass.GetValueOnRenderThread();
+	if (FeatureLevel >= ERHIFeatureLevel::SM3)
+	{
+		extern TAutoConsoleVariable<int32> CVarEarlyZPass;
+		int32 EarlyZPass = CVarEarlyZPass.GetValueOnRenderThread();
 
-				extern int32 GEarlyZPassMovable;
+		extern int32 GEarlyZPassMovable;
 
-				// Render non-masked materials in the depth only pass
-				if (PrimitiveSceneInfo->Proxy->ShouldUseAsOccluder() 
-					&& (!IsMasked(FeatureLevel) || EarlyZPass == 2)
-					&& (!PrimitiveSceneInfo->Proxy->IsMovable() || GEarlyZPassMovable))
-				{
-					FDepthDrawingPolicyFactory::AddStaticMesh(Scene,this);
-				}
+		// Render non-masked materials in the depth only pass
+		if (PrimitiveSceneInfo->Proxy->ShouldUseAsOccluder() 
+			&& (!IsMasked(FeatureLevel) || EarlyZPass == 2)
+			&& (!PrimitiveSceneInfo->Proxy->IsMovable() || GEarlyZPassMovable))
+		{
+			FDepthDrawingPolicyFactory::AddStaticMesh(Scene,this);
+		}
 
-				// Add the static mesh to the DPG's base pass draw list.
-				FBasePassOpaqueDrawingPolicyFactory::AddStaticMesh(RHICmdList, Scene, this);
+		// Add the static mesh to the DPG's base pass draw list.
+		FBasePassOpaqueDrawingPolicyFactory::AddStaticMesh(RHICmdList, Scene, this);
 
-				FVelocityDrawingPolicyFactory::AddStaticMesh(Scene, this);
-			}
-		}
+		FVelocityDrawingPolicyFactory::AddStaticMesh(Scene, this);
 	}
 	else
 	{
-		if (!bShadowOnly && !IsTranslucent(FeatureLevel))
-		{
-			// Add the static mesh to the DPG's base pass draw list.
-			FBasePassForwardOpaqueDrawingPolicyFactory::AddStaticMesh(RHICmdList, Scene, this);
-		}
+		// Add the static mesh to the DPG's base pass draw list.
+		FBasePassForwardOpaqueDrawingPolicyFactory::AddStaticMesh(RHICmdList, Scene, this);
 	}
 }
 
diff --git a/Engine/Source/Runtime/Renderer/Private/SceneHitProxyRendering.cpp b/Engine/Source/Runtime/Renderer/Private/SceneHitProxyRendering.cpp
--- a/Engine/Source/Runtime/Renderer/Private/SceneHitProxyRendering.cpp
+++ b/Engine/Source/Runtime/Renderer/Private/SceneHitProxyRendering.cpp
@@ -36,10 +36,9 @@ public:
 	static bool ShouldCache(EShaderPlatform Platform,const FMaterial* Material,const FVertexFactoryType* VertexFactoryType)
 	{
 		// Only compile the hit proxy vertex shader on PC
-		return IsPCPlatform(Platform) && 
-			IsFeatureLevelSupported(Platform, ERHIFeatureLevel::SM3) &&
+		return IsPCPlatform(Platform)
 			// and only compile for the default material or materials that are masked.
-			(Material->IsSpecialEngineMaterial() || Material->IsMasked() || Material->MaterialMayModifyMeshPosition() || Material->IsTwoSided());
+			&& (Material->IsSpecialEngineMaterial() || Material->IsMasked() || Material->MaterialMayModifyMeshPosition() || Material->IsTwoSided());
 	}
 
 protected:
@@ -109,8 +108,7 @@ public:
 	static bool ShouldCache(EShaderPlatform Platform,const FMaterial* Material,const FVertexFactoryType* VertexFactoryType)
 	{
 		// Only compile the hit proxy vertex shader on PC
-		return IsPCPlatform(Platform) 
-			&& IsFeatureLevelSupported(Platform, ERHIFeatureLevel::SM3) 
+		return IsPCPlatform(Platform)
 			// and only compile for default materials or materials that are masked.
 			&& (Material->IsSpecialEngineMaterial() || Material->IsMasked() || Material->MaterialMayModifyMeshPosition() || Material->IsTwoSided());
 	}
@@ -322,10 +320,11 @@ bool FHitProxyDrawingPolicyFactory::DrawDynamicMesh(
 	return false;
 }
 
-void FDeferredShadingSceneRenderer::RenderHitProxies(FRHICommandListImmediate& RHICmdList)
-{
 #if WITH_EDITOR
 
+TRefCountPtr<IPooledRenderTarget> InitHitProxyRender(FRHICommandListImmediate& RHICmdList, const FSceneRenderer * SceneRenderer)
+{
+	auto& ViewFamily = SceneRenderer->ViewFamily;
 	auto FeatureLevel = ViewFamily.Scene->GetFeatureLevel();
 
 	// Initialize global system textures (pass-through if already initialized).
@@ -343,24 +342,29 @@ void FDeferredShadingSceneRenderer::RenderHitProxies(FRHICommandListImmediate& R
 		GRenderTargetPool.FindFreeElement(Desc, HitProxyRT, TEXT("HitProxy"));
 	}
 
-	if(!HitProxyRT)
+	if (!HitProxyRT)
 	{
 		// HitProxyRT==0 should never happen but better we don't crash
-		return;
+		return HitProxyRT;
 	}
-	
+
 	SetRenderTarget(RHICmdList, HitProxyRT->GetRenderTargetItem().TargetableTexture, GSceneRenderTargets.GetSceneDepthSurface());
 
 	// Clear color for each view.
-	for(int32 ViewIndex = 0;ViewIndex < Views.Num();ViewIndex++)
+	auto& Views = SceneRenderer->Views;
+	for (int32 ViewIndex = 0; ViewIndex < Views.Num(); ViewIndex++)
 	{
 		const FViewInfo& View = Views[ViewIndex];
 		RHICmdList.SetViewport(View.ViewRect.Min.X, View.ViewRect.Min.Y, 0.0f, View.ViewRect.Max.X, View.ViewRect.Max.Y, 1.0f);
 		RHICmdList.Clear(true, FLinearColor::White, false, 1, false, 0, FIntRect());
 	}
+	return HitProxyRT;
+}
 
-	// Find the visible primitives.
-	InitViews(RHICmdList);
+void RenderHitProxies(FRHICommandListImmediate& RHICmdList, const FSceneRenderer * SceneRenderer, TRefCountPtr<IPooledRenderTarget> HitProxyRT)
+{
+	auto & ViewFamily = SceneRenderer->ViewFamily;
+	auto & Views = SceneRenderer->Views;
 
 	// Dynamic vertex and index buffers need to be committed before rendering.
 	FGlobalDynamicVertexBuffer::Get().Commit();
@@ -373,7 +377,7 @@ void FDeferredShadingSceneRenderer::RenderHitProxies(FRHICommandListImmediate& R
 
 	const bool bNeedToSwitchVerticalAxis = IsES2Platform(GRHIShaderPlatform) && !IsPCPlatform(GRHIShaderPlatform);
 
-	for(int32 ViewIndex = 0;ViewIndex < Views.Num();ViewIndex++)
+	for (int32 ViewIndex = 0; ViewIndex < Views.Num(); ViewIndex++)
 	{
 		const FViewInfo& View = Views[ViewIndex];
 
@@ -385,15 +389,15 @@ void FDeferredShadingSceneRenderer::RenderHitProxies(FRHICommandListImmediate& R
 		RHICmdList.Clear(false, FLinearColor::Black, true, 0.0f, true, 0, FIntRect());
 
 		// Adjust the visibility map for this view
-		if( !View.bAllowTranslucentPrimitivesInHitProxy )
+		if (!View.bAllowTranslucentPrimitivesInHitProxy)
 		{
 			// Draw the scene's hit proxy draw lists. (opaque primitives only)
-			Scene->HitProxyDrawList_OpaqueOnly.DrawVisible(RHICmdList, View,View.StaticMeshVisibilityMap,View.StaticMeshBatchVisibility);
+			SceneRenderer->Scene->HitProxyDrawList_OpaqueOnly.DrawVisible(RHICmdList, View, View.StaticMeshVisibilityMap, View.StaticMeshBatchVisibility);
 		}
 		else
 		{
 			// Draw the scene's hit proxy draw lists.
-			Scene->HitProxyDrawList.DrawVisible(RHICmdList, View,View.StaticMeshVisibilityMap,View.StaticMeshBatchVisibility);
+			SceneRenderer->Scene->HitProxyDrawList.DrawVisible(RHICmdList, View, View.StaticMeshVisibilityMap, View.StaticMeshBatchVisibility);
 		}
 
 		const bool bPreFog = true;
@@ -462,7 +466,7 @@ void FDeferredShadingSceneRenderer::RenderHitProxies(FRHICommandListImmediate& R
 
 	// Generate the vertices and triangles mapping the hit proxy RT pixels into the view family's RT pixels.
 	FBatchedElements BatchedElements;
-	for(int32 ViewIndex = 0;ViewIndex < Views.Num();ViewIndex++)
+	for (int32 ViewIndex = 0; ViewIndex < Views.Num(); ViewIndex++)
 	{
 		const FViewInfo& View = Views[ViewIndex];
 
@@ -509,5 +513,33 @@ void FDeferredShadingSceneRenderer::RenderHitProxies(FRHICommandListImmediate& R
 				);
 
 	RHICmdList.EndScene();
+}
+#endif
+
+void FForwardShadingSceneRenderer::RenderHitProxies(FRHICommandListImmediate& RHICmdList)
+{
+#if WITH_EDITOR
+	TRefCountPtr<IPooledRenderTarget> HitProxyRT = ::InitHitProxyRender(RHICmdList, this);
+	// HitProxyRT==0 should never happen but better we don't crash
+	if (HitProxyRT)
+	{
+		// Find the visible primitives.
+		InitViews(RHICmdList);
+		::RenderHitProxies(RHICmdList, this, HitProxyRT);
+	}
+#endif
+}
+
+void FDeferredShadingSceneRenderer::RenderHitProxies(FRHICommandListImmediate& RHICmdList)
+{
+#if WITH_EDITOR
+	TRefCountPtr<IPooledRenderTarget> HitProxyRT = ::InitHitProxyRender(RHICmdList, this);
+	// HitProxyRT==0 should never happen but better we don't crash
+	if (HitProxyRT)
+	{
+		// Find the visible primitives.
+		InitViews(RHICmdList);
+		::RenderHitProxies(RHICmdList, this, HitProxyRT);
+	}
 #endif
 }
diff --git a/Engine/Source/Runtime/Renderer/Private/SceneRendering.h b/Engine/Source/Runtime/Renderer/Private/SceneRendering.h
--- a/Engine/Source/Runtime/Renderer/Private/SceneRendering.h
+++ b/Engine/Source/Runtime/Renderer/Private/SceneRendering.h
@@ -579,6 +579,8 @@ public:
 
 	virtual void Render(FRHICommandListImmediate& RHICmdList) override;
 
+	virtual void RenderHitProxies(FRHICommandListImmediate& RHICmdList) override;
+
 protected:
 
 	void InitViews(FRHICommandListImmediate& RHICmdList);
