* fixed "vis SSProfiles"

diff --git a/Engine/Source/Runtime/Engine/Classes/Engine/SubsurfaceProfile.h b/Engine/Source/Runtime/Engine/Classes/Engine/SubsurfaceProfile.h
--- a/Engine/Source/Runtime/Engine/Classes/Engine/SubsurfaceProfile.h
+++ b/Engine/Source/Runtime/Engine/Classes/Engine/SubsurfaceProfile.h
@@ -104,7 +104,7 @@ public:
 	// @param InRendererModule first call needs to have it != 0
 	void UpdateProfile(int32 AllocationId, const FSubsurfaceProfileStruct Settings);
 
-	const struct FSceneRenderTargetItem* GetTexture(FRHICommandListImmediate& RHICmdList);
+	const struct IPooledRenderTarget* GetTexture(FRHICommandListImmediate& RHICmdList);
 
 
 	// FRenderResource interface.
diff --git a/Engine/Source/Runtime/Engine/Private/Rendering/SubsurfaceProfile.cpp b/Engine/Source/Runtime/Engine/Private/Rendering/SubsurfaceProfile.cpp
--- a/Engine/Source/Runtime/Engine/Private/Rendering/SubsurfaceProfile.cpp
+++ b/Engine/Source/Runtime/Engine/Private/Rendering/SubsurfaceProfile.cpp
@@ -97,7 +97,7 @@ void FSubsurfaceProfileTexture::UpdateProfile(int32 AllocationId, const FSubsurf
 	GSSProfiles.SafeRelease();
 }
 
-const FSceneRenderTargetItem* FSubsurfaceProfileTexture::GetTexture(FRHICommandListImmediate& RHICmdList)
+const IPooledRenderTarget* FSubsurfaceProfileTexture::GetTexture(FRHICommandListImmediate& RHICmdList)
 {
 	// call SetRendererModule() is missing
 	check(RendererModule);
@@ -107,7 +107,7 @@ const FSceneRenderTargetItem* FSubsurfaceProfileTexture::GetTexture(FRHICommandL
 		CreateTexture(RHICmdList);
 	}
 
-	return &GSSProfiles->GetRenderTargetItem();
+	return GSSProfiles;
 }
 
 void FSubsurfaceProfileTexture::ReleaseDynamicRHI()
@@ -224,7 +224,7 @@ void FSubsurfaceProfileTexture::Dump()
 
 
 
-ENGINE_API const FSceneRenderTargetItem* GetSubsufaceProfileTexture_RT(FRHICommandListImmediate& RHICmdList)
+ENGINE_API const IPooledRenderTarget* GetSubsufaceProfileTexture_RT(FRHICommandListImmediate& RHICmdList)
 {
 	check(IsInRenderingThread());
 
diff --git a/Engine/Source/Runtime/Renderer/Private/PostProcess/PostProcessSubsurface.cpp b/Engine/Source/Runtime/Renderer/Private/PostProcess/PostProcessSubsurface.cpp
--- a/Engine/Source/Runtime/Renderer/Private/PostProcess/PostProcessSubsurface.cpp
+++ b/Engine/Source/Runtime/Renderer/Private/PostProcess/PostProcessSubsurface.cpp
@@ -10,7 +10,7 @@
 #include "PostProcessSubsurface.h"
 #include "PostProcessing.h"
 
-
+ENGINE_API const IPooledRenderTarget* GetSubsufaceProfileTexture_RT(FRHICommandListImmediate& RHICmdList);
 
 /** Encapsulates the post processing ambient occlusion pixel shader. */
 template <uint32 SpecularCorrection>
@@ -266,17 +266,13 @@ public:
 		}
 
 		{
-			ENGINE_API const FSceneRenderTargetItem* GetSubsufaceProfileTexture_RT(FRHICommandListImmediate& RHICmdList);
+			const IPooledRenderTarget* PooledRT = GetSubsufaceProfileTexture_RT(Context.RHICmdList);
 
-			const FSceneRenderTargetItem* Item = GetSubsufaceProfileTexture_RT(Context.RHICmdList);
+			check(PooledRT);
 
-			if (!Item)
-			{
-				// should never happen
-				Item = &GSystemTextures.BlackDummy->GetRenderTargetItem();
-			}
+			const FSceneRenderTargetItem& Item = PooledRT->GetRenderTargetItem();
 
-			SetTextureParameter(Context.RHICmdList, ShaderRHI, SSProfilesTexture, SSProfilesTextureSampler, TStaticSamplerState<SF_Point, AM_Wrap, AM_Wrap, AM_Wrap>::GetRHI(), Item->ShaderResourceTexture);
+			SetTextureParameter(Context.RHICmdList, ShaderRHI, SSProfilesTexture, SSProfilesTextureSampler, TStaticSamplerState<SF_Point, AM_Wrap, AM_Wrap, AM_Wrap>::GetRHI(), Item.ShaderResourceTexture);
 		}
 	}
 
@@ -325,10 +321,15 @@ void FRCPassPostProcessSubsurface::Process(FRenderingCompositePassContext& Conte
 {
 	const FPooledRenderTargetDesc* InputDesc = GetInputDesc(ePId_Input1);
 
-	if(!InputDesc)
+	check(InputDesc);
+
 	{
-		// input is not hooked up correctly
-		return;
+		const IPooledRenderTarget* PooledRT = GetSubsufaceProfileTexture_RT(Context.RHICmdList);
+
+		check(PooledRT);
+
+		// for debugging
+		GRenderTargetPool.VisualizeTexture.SetCheckPoint(Context.RHICmdList, PooledRT);
 	}
 
 	const FSceneView& View = Context.View;
