* fixed TTP 341929 REGRESSION: CRASH: EDITOR: Turning on GBuffer Hints in the main viewport crashes the Editor

diff --git a/Engine/Source/Runtime/Renderer/Private/PostProcess/PostProcessGBufferHints.cpp b/Engine/Source/Runtime/Renderer/Private/PostProcess/PostProcessGBufferHints.cpp
--- a/Engine/Source/Runtime/Renderer/Private/PostProcess/PostProcessGBufferHints.cpp
+++ b/Engine/Source/Runtime/Renderer/Private/PostProcess/PostProcessGBufferHints.cpp
@@ -19,7 +19,7 @@ class FPostProcessGBufferHintsPS : public FGlobalShader
 
 	static bool ShouldCache(EShaderPlatform Platform)
 	{
-		return IsFeatureLevelSupported(Platform, ERHIFeatureLevel::SM5);
+		return IsFeatureLevelSupported(Platform, ERHIFeatureLevel::SM4);
 	}
 
 	static void ModifyCompilationEnvironment(EShaderPlatform Platform, FShaderCompilerEnvironment& OutEnvironment)
diff --git a/Engine/Source/Runtime/Renderer/Private/PostProcess/PostProcessing.cpp b/Engine/Source/Runtime/Renderer/Private/PostProcess/PostProcessing.cpp
--- a/Engine/Source/Runtime/Renderer/Private/PostProcess/PostProcessing.cpp
+++ b/Engine/Source/Runtime/Renderer/Private/PostProcess/PostProcessing.cpp
@@ -1050,7 +1050,7 @@ void FPostProcessing::Process(FRHICommandListImmediate& RHICmdList, const FViewI
 			Context.FinalOutput = FRenderingCompositeOutputRef(Node);
 		}
 
-		if(View.Family->EngineShowFlags.GBufferHints)
+		if (View.Family->EngineShowFlags.GBufferHints && FeatureLevel >= ERHIFeatureLevel::SM4)
 		{
 			FRenderingCompositePass* Node = Context.Graph.RegisterPass(new(FMemStack::Get()) FRCPassPostProcessGBufferHints());
 			Node->SetInput(ePId_Input0, FRenderingCompositeOutputRef(Context.FinalOutput));
