Turn color grading off when it is not needed.

diff --git a/Engine/Shaders/PostProcessTonemap.usf b/Engine/Shaders/PostProcessTonemap.usf
--- a/Engine/Shaders/PostProcessTonemap.usf
+++ b/Engine/Shaders/PostProcessTonemap.usf
@@ -294,8 +294,10 @@ void MainPS(
 			TonemappedColor.rgb += GrainAdd;
 		#endif
 
-		// apply color grading
-		TonemappedColor = ColorLookupTable(TonemappedColor);
+		#if USE_COLOR_GRADING
+			// apply color grading
+			TonemappedColor = ColorLookupTable(TonemappedColor);
+		#endif
 
 		// RETURN_COLOR not needed unless writing to SceneColor
 		OutColor = float4(TonemappedColor, LuminanceForPostProcessAA);
diff --git a/Engine/Source/Runtime/Renderer/Private/PostProcess/PostProcessCombineLUTs.cpp b/Engine/Source/Runtime/Renderer/Private/PostProcess/PostProcessCombineLUTs.cpp
--- a/Engine/Source/Runtime/Renderer/Private/PostProcess/PostProcessCombineLUTs.cpp
+++ b/Engine/Source/Runtime/Renderer/Private/PostProcess/PostProcessCombineLUTs.cpp
@@ -452,3 +452,52 @@ FPooledRenderTargetDesc FRCPassPostProcessCombineLUTs::ComputeOutputDesc(EPassOu
 
 	return Ret;
 }
+
+bool FRCPassPostProcessCombineLUTs::IsColorGradingLUTNeeded(const FRenderingCompositePassContext* RESTRICT Context)
+{
+	FColorTransform ColorTransform;
+	ColorTransform.MinValue = FMath::Clamp(CVarColorMin.GetValueOnRenderThread(), -10.0f, 10.0f);
+	ColorTransform.MidValue = FMath::Clamp(CVarColorMid.GetValueOnRenderThread(), -10.0f, 10.0f);
+	ColorTransform.MaxValue = FMath::Clamp(CVarColorMax.GetValueOnRenderThread(), -10.0f, 10.0f);
+
+	if(ColorTransform.MinValue != 0.0f)
+	{
+		return true;
+	}
+	if(ColorTransform.MidValue != 0.5f)
+	{
+		return true;
+	}
+	if(ColorTransform.MaxValue != 1.0f)
+	{
+		return true;
+	}
+
+	if(Context->View.ColorScale.R != 1.0f || Context->View.ColorScale.G != 1.0f || Context->View.ColorScale.B != 1.0f)
+	{
+		return true;
+	}
+
+	if(Context->View.OverlayColor.A > (1.0f/512.0f))
+	{
+		return true;
+	}
+
+	if(Context->View.GetFeatureLevel() >= ERHIFeatureLevel::SM4)
+	{
+		const FFinalPostProcessSettings Settings = Context->View.FinalPostProcessSettings;
+		for(uint32 i = 0; i < (uint32)Settings.ContributingLUTs.Num(); ++i)
+		{
+			UTexture* LUTTexture = Settings.ContributingLUTs[i].LUTTexture;
+			if(LUTTexture == NULL) 
+			{
+				continue;
+			}
+			if(Settings.ContributingLUTs[i].Weight > 1.0f / 512.0f)
+			{
+				return true;
+			}
+		}
+	}
+	return false;
+}
diff --git a/Engine/Source/Runtime/Renderer/Private/PostProcess/PostProcessCombineLUTs.h b/Engine/Source/Runtime/Renderer/Private/PostProcess/PostProcessCombineLUTs.h
--- a/Engine/Source/Runtime/Renderer/Private/PostProcess/PostProcessCombineLUTs.h
+++ b/Engine/Source/Runtime/Renderer/Private/PostProcess/PostProcessCombineLUTs.h
@@ -22,6 +22,8 @@ public:
 	uint32 GenerateFinalTable(const FFinalPostProcessSettings& Settings, FTexture* OutTextures[], float OutWeights[], uint32 MaxCount) const;
 	/** @return 0xffffffff if not found */
 	uint32 FindIndex(const FFinalPostProcessSettings& Settings, UTexture* Tex) const;
+
+	static bool IsColorGradingLUTNeeded(const FRenderingCompositePassContext* RESTRICT Context);
 };
 
 
diff --git a/Engine/Source/Runtime/Renderer/Private/PostProcess/PostProcessTonemap.cpp b/Engine/Source/Runtime/Renderer/Private/PostProcess/PostProcessTonemap.cpp
--- a/Engine/Source/Runtime/Renderer/Private/PostProcess/PostProcessTonemap.cpp
+++ b/Engine/Source/Runtime/Renderer/Private/PostProcess/PostProcessTonemap.cpp
@@ -35,6 +35,7 @@ typedef enum {
 	TonemapperLightShafts       = (1<<11),
 	TonemapperMosaic            = (1<<12),
 	TonemapperColorFringe       = (1<<13),
+	TonemapperColorGrading      = (1<<14),
 
 } TonemapperOption;
 
@@ -56,6 +57,7 @@ static uint8 TonemapperCostTab[] = {
 	1, //TonemapperLightShafts
 	1, //TonemapperMosaic
 	1, //TonemapperColorFringe
+	1, //TonemapperColorGrading
 };
 
 // Edit the following to add and remove configurations.
@@ -63,7 +65,7 @@ static uint8 TonemapperCostTab[] = {
 // Place most common first (faster when searching in TonemapperFindLeastExpensive()).
 
 // List of configurations compiled for PC.
-static uint32 TonemapperConfBitmaskPC[11] = { 
+static uint32 TonemapperConfBitmaskPC[22] = { 
 
 	TonemapperBloom +
 	TonemapperContrast +
@@ -75,6 +77,7 @@ static uint32 TonemapperConfBitmaskPC[11] = {
 	TonemapperVignette +
 	TonemapperVignetteColor +
 	TonemapperColorFringe +
+	TonemapperColorGrading +
 	0,
 
 	TonemapperBloom +
@@ -84,6 +87,7 @@ static uint32 TonemapperConfBitmaskPC[11] = {
 	TonemapperVignette +
 	TonemapperGrainQuantization +
 	TonemapperColorFringe +
+	TonemapperColorGrading +
 	0,
 
 	TonemapperBloom + 
@@ -91,17 +95,20 @@ static uint32 TonemapperConfBitmaskPC[11] = {
 	TonemapperColorMatrix + 
 	TonemapperShadowTint +
 	TonemapperGrainQuantization +
+	TonemapperColorGrading +
 	0,
 
 	TonemapperBloom + 
 	TonemapperContrast + 
 	TonemapperColorMatrix +
 	TonemapperGrainQuantization +
+	TonemapperColorGrading +
 	0,
 
 	TonemapperBloom	+ 
 	TonemapperContrast +
 	TonemapperGrainQuantization +
+	TonemapperColorGrading +
 	0,
 
 	// same without TonemapperGrainQuantization
@@ -115,6 +122,7 @@ static uint32 TonemapperConfBitmaskPC[11] = {
 	TonemapperVignette +
 	TonemapperVignetteColor +
 	TonemapperColorFringe +
+	TonemapperColorGrading +
 	0,
 
 	TonemapperBloom +
@@ -123,28 +131,116 @@ static uint32 TonemapperConfBitmaskPC[11] = {
 	TonemapperShadowTint +
 	TonemapperVignette +
 	TonemapperColorFringe +
+	TonemapperColorGrading +
 	0,
 
 	TonemapperBloom + 
 	TonemapperContrast + 
 	TonemapperColorMatrix + 
 	TonemapperShadowTint +
+	TonemapperColorGrading +
 	0,
 
 	TonemapperBloom + 
 	TonemapperContrast + 
 	TonemapperColorMatrix +
+	TonemapperColorGrading +
 	0,
 
 	TonemapperBloom	+ 
 	TonemapperContrast +
+	TonemapperColorGrading +
 	0,
 	
 	//
 
 	TonemapperGammaOnly +
+	TonemapperColorGrading +
 	0,
 
+// SAME WITHOUT COLORGRADING
+
+	TonemapperBloom +
+	TonemapperContrast +
+	TonemapperColorMatrix +
+	TonemapperShadowTint +
+	TonemapperGrainJitter +
+	TonemapperGrainIntensity +
+	TonemapperGrainQuantization +
+	TonemapperVignette +
+	TonemapperVignetteColor +
+	TonemapperColorFringe +
+	0,
+
+	TonemapperBloom +
+	TonemapperContrast +
+	TonemapperColorMatrix +
+	TonemapperShadowTint +
+	TonemapperVignette +
+	TonemapperGrainQuantization +
+	TonemapperColorFringe +
+	0,
+
+	TonemapperBloom + 
+	TonemapperContrast + 
+	TonemapperColorMatrix + 
+	TonemapperShadowTint +
+	TonemapperGrainQuantization +
+	0,
+
+	TonemapperBloom + 
+	TonemapperContrast + 
+	TonemapperColorMatrix +
+	TonemapperGrainQuantization +
+	0,
+
+	TonemapperBloom	+ 
+	TonemapperContrast +
+	TonemapperGrainQuantization +
+	0,
+
+	// same without TonemapperGrainQuantization
+
+	TonemapperBloom +
+	TonemapperContrast +
+	TonemapperColorMatrix +
+	TonemapperShadowTint +
+	TonemapperGrainJitter +
+	TonemapperGrainIntensity +
+	TonemapperVignette +
+	TonemapperVignetteColor +
+	TonemapperColorFringe +
+	0,
+
+	TonemapperBloom +
+	TonemapperContrast +
+	TonemapperColorMatrix +
+	TonemapperShadowTint +
+	TonemapperVignette +
+	TonemapperColorFringe +
+	0,
+
+	TonemapperBloom + 
+	TonemapperContrast + 
+	TonemapperColorMatrix + 
+	TonemapperShadowTint +
+	0,
+
+	TonemapperBloom + 
+	TonemapperContrast + 
+	TonemapperColorMatrix +
+	0,
+
+	TonemapperBloom	+ 
+	TonemapperContrast +
+	0,
+	
+	//
+
+	TonemapperGammaOnly +
+	0,
+
+
 };
 
 // List of configurations compiled for Mobile.
@@ -521,6 +617,11 @@ static uint32 TonemapperGenerateBitmaskPC(const FRenderingCompositePassContext*
 		Bitmask |= TonemapperColorFringe;
 	}
 
+	if(FRCPassPostProcessCombineLUTs::IsColorGradingLUTNeeded(Context))
+	{
+		Bitmask |= TonemapperColorGrading;
+	}
+
 	return Bitmask + TonemapperGenerateBitmaskPost(Context);
 }
 
@@ -775,6 +876,7 @@ class FPostProcessTonemapPS : public FGlobalShader
 		OutEnvironment.SetDefine(TEXT("USE_COLOR_FRINGE"),		 TonemapperIsDefined(ConfigBitmask, TonemapperColorFringe));
 		// @todo Mac OS X: in order to share precompiled shaders between GL 3.3 & GL 4.1 devices we mustn't use volume-texture rendering as it isn't universally supported.
 		OutEnvironment.SetDefine(TEXT("USE_VOLUME_LUT"),		 (IsFeatureLevelSupported(Platform,ERHIFeatureLevel::SM4) && GSupportsVolumeTextureRendering && !PLATFORM_MAC));
+		OutEnvironment.SetDefine(TEXT("USE_COLOR_GRADING"),		 TonemapperIsDefined(ConfigBitmask, TonemapperColorGrading));
 
 		if( !IsFeatureLevelSupported(Platform,ERHIFeatureLevel::SM5) )
 		{
@@ -978,6 +1080,10 @@ public:
 	VARIATION1(0)  VARIATION1(1)  VARIATION1(2)  VARIATION1(3)  VARIATION1(4)  VARIATION1(5)
 	VARIATION1(6)  VARIATION1(7)  VARIATION1(8)  VARIATION1(9)  VARIATION1(10)
 
+	VARIATION1(11)  VARIATION1(12)  VARIATION1(13)  VARIATION1(14)  VARIATION1(15)  VARIATION1(16)
+	VARIATION1(17)  VARIATION1(18)  VARIATION1(19)  VARIATION1(20)  VARIATION1(21)
+
+
 #undef VARIATION1
 
 
@@ -1058,6 +1164,17 @@ void FRCPassPostProcessTonemap::Process(FRenderingCompositePassContext& Context)
 		case 8: SetShaderTempl<8>(Context); break;
 		case 9: SetShaderTempl<9>(Context); break;
 		case 10: SetShaderTempl<10>(Context); break;
+		case 11: SetShaderTempl<11>(Context); break;
+		case 12: SetShaderTempl<12>(Context); break;
+		case 13: SetShaderTempl<13>(Context); break;
+		case 14: SetShaderTempl<14>(Context); break;
+		case 15: SetShaderTempl<15>(Context); break;
+		case 16: SetShaderTempl<16>(Context); break;
+		case 17: SetShaderTempl<17>(Context); break;
+		case 18: SetShaderTempl<18>(Context); break;
+		case 19: SetShaderTempl<19>(Context); break;
+		case 20: SetShaderTempl<20>(Context); break;
+		case 21: SetShaderTempl<21>(Context); break;
 		default:
 			check(0);
 	}
