Unbreak the non-unity build. Forgot a header.

diff --git a/Engine/Source/Runtime/Renderer/Private/PostProcess/PostProcessTonemap.cpp b/Engine/Source/Runtime/Renderer/Private/PostProcess/PostProcessTonemap.cpp
--- a/Engine/Source/Runtime/Renderer/Private/PostProcess/PostProcessTonemap.cpp
+++ b/Engine/Source/Runtime/Renderer/Private/PostProcess/PostProcessTonemap.cpp
@@ -11,7 +11,7 @@
 #include "PostProcessTonemap.h"
 #include "PostProcessing.h"
 #include "PostProcessEyeAdaptation.h"
-
+#include "PostProcessCombineLUTs.h"
 
 
 //
