Fixed specular on movable skylights
 * Removed double normalization from stationary skylight specular

diff --git a/Engine/Shaders/BasePassPixelShader.usf b/Engine/Shaders/BasePassPixelShader.usf
--- a/Engine/Shaders/BasePassPixelShader.usf
+++ b/Engine/Shaders/BasePassPixelShader.usf
@@ -73,23 +73,32 @@ half3 GetImageBasedReflectionLighting(FMaterialPixelParameters MaterialParameter
 	half AbsoluteSpecularMip = ComputeReflectionCaptureMipFromRoughness(Roughness);
 	float4 SpecularIBL = TextureCubeArraySampleLevel(ReflectionCubemap, ReflectionCubemapSampler, MaterialParameters.ReflectionVector, CubemapArrayIndex, AbsoluteSpecularMip);
 
+	FLATTEN
+	if (View.UseLightmaps > 0)
+	{
+		// Note: make sure this matches the lightmap mixing done on opaque (ReflectionEnvironmentTiledDeferredMain)
+		SpecularIBL.rgb *= IndirectIrradiance;
+	}
+
 #if ENABLE_SKY_LIGHT
 	BRANCH
 	if (SkyLightParameters.y > 0 && SpecularIBL.a < .001f)
 	{
-		float3 SkyLighting = GetSkyLightReflection(MaterialParameters.ReflectionVector, Roughness, View.UseLightmaps);
+		// Normalize for static skylight types which mix with lightmaps
+		bool bNormalize = SkyLightParameters.z < 1 && View.UseLightmaps;
+		float3 SkyLighting = GetSkyLightReflection(MaterialParameters.ReflectionVector, Roughness, bNormalize);
+
+		FLATTEN
+		if (bNormalize)
+		{
+			SkyLighting *= IndirectIrradiance;
+		}
+
 		// Add in sky wherever reflection captures don't have coverage
 		SpecularIBL.rgb += (1 - SpecularIBL.a) * SkyLighting;
 	}
 #endif
 	
-	BRANCH
-	if( View.UseLightmaps > 0 )
-	{
-		// Note: make sure this matches the lightmap mixing done on opaque (ReflectionEnvironmentTiledDeferredMain)
-		SpecularIBL.rgb *= IndirectIrradiance;
-	}
-	
 	half NoV = saturate(dot(MaterialParameters.WorldNormal, MaterialParameters.CameraVector));
 	SpecularColor = EnvBRDFApprox(SpecularColor, Roughness, NoV);
 
diff --git a/Engine/Shaders/ReflectionEnvironmentShaders.usf b/Engine/Shaders/ReflectionEnvironmentShaders.usf
--- a/Engine/Shaders/ReflectionEnvironmentShaders.usf
+++ b/Engine/Shaders/ReflectionEnvironmentShaders.usf
@@ -505,15 +505,10 @@ void ReflectionApplyPS(
 		float4 SpecularLighting = float4(0, 0, 0, 1);
 		
 #if APPLY_REFLECTION_ENV
+		// Alpha of the reflection environment texture contains coverage, where 0 the sky reflection should be applied
 		SpecularLighting = Texture2DSample( ReflectionEnvTexture, ReflectionEnvSampler, UV );
-#endif
 
-#if APPLY_SKYLIGHT
-		float3 ReflectionVector = reflect(-V, GBuffer.WorldNormal);
-		SpecularLighting.rgb += SpecularLighting.a * GetSkyLightReflection(ReflectionVector, GBuffer.Roughness, true);
-#endif
-		
-		BRANCH
+		FLATTEN
 		if( View.UseLightmaps > 0 )
 		{
 			// We have high frequency directional data but low frequency spatial data in the envmap.
@@ -523,6 +518,22 @@ void ReflectionApplyPS(
 			// Note: make sure this matches the lightmap mixing done for translucency (BasePassPixelShader.usf)
 			SpecularLighting.rgb *= GBuffer.IndirectIrradiance;
 		}
+#endif
+
+#if APPLY_SKYLIGHT
+		float3 ReflectionVector = reflect(-V, GBuffer.WorldNormal);
+		// Normalize for static skylight types which mix with lightmaps
+		bool bNormalize = SkyLightParameters.z < 1 && View.UseLightmaps > 0;
+		float3 SkyLighting = GetSkyLightReflection(ReflectionVector, GBuffer.Roughness, bNormalize);
+
+		FLATTEN
+		if (bNormalize)
+		{
+			SpecularLighting.rgb *= GBuffer.IndirectIrradiance;
+		}
+		
+		SpecularLighting.rgb += SpecularLighting.a * SkyLighting;
+#endif
 
 #if APPLY_SSR
 		float4 SSR = Texture2DSample( ScreenSpaceReflectionsTexture, ScreenSpaceReflectionsSampler, UV );
diff --git a/Engine/Shaders/ReflectionEnvironmentShared.usf b/Engine/Shaders/ReflectionEnvironmentShared.usf
--- a/Engine/Shaders/ReflectionEnvironmentShared.usf
+++ b/Engine/Shaders/ReflectionEnvironmentShared.usf
@@ -32,15 +32,15 @@ float ComputeReflectionCaptureRoughnessFromMip( float Mip )
 TextureCube SkyLightCubemap;
 SamplerState SkyLightCubemapSampler;
 
-/** X = max mip, Y = 1 if sky light should be rendered, 0 otherwise. */
-float2 SkyLightParameters;
+/** X = max mip, Y = 1 if sky light should be rendered, 0 otherwise, Z = 1 if sky light is dynamic, 0 otherwise. */
+float3 SkyLightParameters;
 
 float3 GetSkyLightReflection(float3 ReflectionVector, float Roughness, bool Normalized)
 {
 	float AbsoluteSpecularMip = ComputeReflectionCaptureMipFromRoughness(Roughness);
 	float3 Reflection = TextureCubeSampleLevel(SkyLightCubemap, SkyLightCubemapSampler, ReflectionVector, AbsoluteSpecularMip).rgb;
 
-	FLATTEN
+	BRANCH
 	if (Normalized)
 	{
 		// Sample the lowest resolution mip to get the average color
diff --git a/Engine/Source/Runtime/Renderer/Private/BasePassRendering.cpp b/Engine/Source/Runtime/Renderer/Private/BasePassRendering.cpp
--- a/Engine/Source/Runtime/Renderer/Private/BasePassRendering.cpp
+++ b/Engine/Source/Runtime/Renderer/Private/BasePassRendering.cpp
@@ -47,10 +47,11 @@ IMPLEMENT_BASEPASS_LIGHTMAPPED_SHADER_TYPE( FCachedVolumeIndirectLightingPolicy,
 IMPLEMENT_BASEPASS_LIGHTMAPPED_SHADER_TYPE( FCachedPointIndirectLightingPolicy, FCachedPointIndirectLightingPolicy );
 IMPLEMENT_BASEPASS_LIGHTMAPPED_SHADER_TYPE( FSimpleDynamicLightingPolicy, FSimpleDynamicLightingPolicy );
 
-void FSkyLightReflectionParameters::GetSkyParametersFromScene(const FScene* Scene, bool bApplySkyLight, FTexture*& OutSkyLightTextureResource, float& OutApplySkyLightMask, float& OutSkyMipCount)
+void FSkyLightReflectionParameters::GetSkyParametersFromScene(const FScene* Scene, bool bApplySkyLight, FTexture*& OutSkyLightTextureResource, float& OutApplySkyLightMask, float& OutSkyMipCount, bool& bOutSkyLightIsDynamic)
 {
 	OutSkyLightTextureResource = GBlackTextureCube;
 	OutApplySkyLightMask = 0;
+	bOutSkyLightIsDynamic = false;
 
 	if (Scene
 		&& Scene->SkyLight 
@@ -59,6 +60,7 @@ void FSkyLightReflectionParameters::GetSkyParametersFromScene(const FScene* Scen
 	{
 		OutSkyLightTextureResource = Scene->SkyLight->ProcessedTexture;
 		OutApplySkyLightMask = 1;
+		bOutSkyLightIsDynamic = !Scene->SkyLight->bHasStaticLighting && !Scene->SkyLight->bWantsStaticShadowing;
 	}
 
 	OutSkyMipCount = 1;
diff --git a/Engine/Source/Runtime/Renderer/Private/BasePassRendering.h b/Engine/Source/Runtime/Renderer/Private/BasePassRendering.h
--- a/Engine/Source/Runtime/Renderer/Private/BasePassRendering.h
+++ b/Engine/Source/Runtime/Renderer/Private/BasePassRendering.h
@@ -209,11 +209,12 @@ public:
 		FTexture* SkyLightTextureResource = GBlackTextureCube;
 		float ApplySkyLightMask = 0;
 		float SkyMipCount = 1;
+		bool bSkyLightIsDynamic = false;
 
-		GetSkyParametersFromScene(Scene, bApplySkyLight, SkyLightTextureResource, ApplySkyLightMask, SkyMipCount);
+		GetSkyParametersFromScene(Scene, bApplySkyLight, SkyLightTextureResource, ApplySkyLightMask, SkyMipCount, bSkyLightIsDynamic);
 
 		SetTextureParameter(RHICmdList, ShaderRHI, SkyLightCubemap, SkyLightCubemapSampler, SkyLightTextureResource);
-		const FVector2D SkyParametersValue(SkyMipCount - 1.0f, ApplySkyLightMask);
+		const FVector SkyParametersValue(SkyMipCount - 1.0f, ApplySkyLightMask, bSkyLightIsDynamic ? 1.0f : 0.0f);
 		SetShaderValue(RHICmdList, ShaderRHI, SkyLightParameters, SkyParametersValue);
 	}
 
@@ -229,7 +230,7 @@ private:
 	FShaderResourceParameter SkyLightCubemapSampler;
 	FShaderParameter SkyLightParameters;
 
-	void GetSkyParametersFromScene(const FScene* Scene, bool bApplySkyLight, FTexture*& OutSkyLightTextureResource, float& OutApplySkyLightMask, float& OutSkyMipCount);
+	void GetSkyParametersFromScene(const FScene* Scene, bool bApplySkyLight, FTexture*& OutSkyLightTextureResource, float& OutApplySkyLightMask, float& OutSkyMipCount, bool& bSkyLightIsDynamic);
 };
 
 /** Parameters needed for lighting translucency, shared by multiple shaders. */
diff --git a/Engine/Source/Runtime/Renderer/Private/ReflectionEnvironmentCapture.cpp b/Engine/Source/Runtime/Renderer/Private/ReflectionEnvironmentCapture.cpp
--- a/Engine/Source/Runtime/Renderer/Private/ReflectionEnvironmentCapture.cpp
+++ b/Engine/Source/Runtime/Renderer/Private/ReflectionEnvironmentCapture.cpp
@@ -302,7 +302,7 @@ void ComputeAverageBrightness(FRHICommandList& RHICmdList)
 }
 
 /** Generates mips for glossiness and filters the cubemap for a given reflection. */
-void FilterReflectionEnvironment(FRHICommandListImmediate& RHICmdList, FSHVectorRGB3* OutIrradianceEnvironmentMap)
+void FilterReflectionEnvironment(FRHICommandListImmediate& RHICmdList, FSHVectorRGB3* OutIrradianceEnvironmentMap, bool bNormalize)
 {
 	const int32 EffectiveTopMipSize = GReflectionCaptureSize;
 	const int32 NumMips = FMath::CeilLogTwo(EffectiveTopMipSize) + 1;
@@ -311,7 +311,7 @@ void FilterReflectionEnvironment(FRHICommandListImmediate& RHICmdList, FSHVector
 	FSceneRenderTargetItem* DiffuseConvolutionSource = NULL;
 
 	static const auto CVar = IConsoleManager::Get().FindConsoleVariable(TEXT("r.DiffuseFromCaptures"));
-	const bool bNormalize = CVar->GetInt() == 0;
+	bNormalize = bNormalize && CVar->GetInt() == 0;
 
 	{	
 		SCOPED_DRAW_EVENT(DownsampleCubeMips, DEC_SCENE_ITEMS);
@@ -1490,7 +1490,8 @@ void FScene::UpdateReflectionCaptureContents(UReflectionCaptureComponent* Captur
 			ENQUEUE_UNIQUE_RENDER_COMMAND( 
 				FilterCommand,
 			{
-				FilterReflectionEnvironment(RHICmdList, NULL);
+				bool bNormalize = true;
+				FilterReflectionEnvironment(RHICmdList, NULL, bNormalize);
 			});
 
 			// Create a proxy to represent the reflection capture to the rendering thread
@@ -1574,7 +1575,9 @@ void FScene::UpdateSkyCaptureContents(const USkyLightComponent* CaptureComponent
 			FilterCommand,
 			FSHVectorRGB3*, IrradianceEnvironmentMap, &OutIrradianceEnvironmentMap,
 		{
-			FilterReflectionEnvironment(RHICmdList, IrradianceEnvironmentMap);
+			// Skylight is normalized manually in the shader
+			bool bNormalize = false;
+			FilterReflectionEnvironment(RHICmdList, IrradianceEnvironmentMap, bNormalize);
 		});
 
 		// Wait until the SH coefficients have been written out by the RT before returning
