Fixed SSAO being disabled when there's only a skylight in the scene

diff --git a/Engine/Source/Runtime/Renderer/Private/CompositionLighting/CompositionLighting.cpp b/Engine/Source/Runtime/Renderer/Private/CompositionLighting/CompositionLighting.cpp
--- a/Engine/Source/Runtime/Renderer/Private/CompositionLighting/CompositionLighting.cpp
+++ b/Engine/Source/Runtime/Renderer/Private/CompositionLighting/CompositionLighting.cpp
@@ -86,6 +86,14 @@ static bool IsReflectionEnvironmentActive(FPostprocessContext& Context)
 	return (Scene->GetFeatureLevel() == ERHIFeatureLevel::SM5 && IsReflectingEnvironment && (HasReflectionCaptures || HasSSR) && !IsSimpleDynamicLightingEnabled());
 }
 
+static bool IsSkylightActive(FPostprocessContext& Context)
+{
+	FScene* Scene = (FScene*)Context.View.Family->Scene;
+	return Scene->SkyLight 
+		&& Scene->SkyLight->ProcessedTexture
+		&& Context.View.Family->EngineShowFlags.SkyLighting;
+}
+
 static bool IsBasePassAmbientOcclusionRequired(FPostprocessContext& Context)
 {
 	// the BaseAO pass is only worth with some AO
@@ -103,7 +111,7 @@ static uint32 ComputeAmbientOcclusionPassCount(FPostprocessContext& Context)
 	{
 		bEnabled = Context.View.FinalPostProcessSettings.AmbientOcclusionIntensity > 0 
 			&& Context.View.FinalPostProcessSettings.AmbientOcclusionRadius >= 0.1f 
-			&& (IsBasePassAmbientOcclusionRequired(Context) || IsAmbientCubemapPassRequired(Context) || IsReflectionEnvironmentActive(Context) || Context.View.Family->EngineShowFlags.VisualizeBuffer )
+			&& (IsBasePassAmbientOcclusionRequired(Context) || IsAmbientCubemapPassRequired(Context) || IsReflectionEnvironmentActive(Context) || IsSkylightActive(Context) || Context.View.Family->EngineShowFlags.VisualizeBuffer )
 			&& !IsSimpleDynamicLightingEnabled();
 	}
 
