UE4 - work on translucent pass toward parallel rendering. RHCommandlists are now recycled. Refactor DrawVisibleParallel so there is no duplicated code.

diff --git a/Engine/Source/Runtime/RHI/Private/RHICommandList.cpp b/Engine/Source/Runtime/RHI/Private/RHICommandList.cpp
--- a/Engine/Source/Runtime/RHI/Private/RHICommandList.cpp
+++ b/Engine/Source/Runtime/RHI/Private/RHICommandList.cpp
@@ -48,7 +48,7 @@ void FRHICommandListExecutor::ExecuteList(FRHICommandListBase& CmdList)
 {
 	check(IsInRenderingThread() && &CmdList != &GetImmediateCommandList());
 
-	if (!GetImmediateCommandList().bExecuting) // don't flush if this is a recursive call and we are already executing the immediate command list
+	if (!GetImmediateCommandList().IsExecuting()) // don't flush if this is a recursive call and we are already executing the immediate command list
 	{
 		GetImmediateCommandList().Flush();
 	}
@@ -139,4 +139,23 @@ void FRHICommandListBase::Reset()
 	UID = GRHICommandList.UIDCounter.Increment();
 }
 
+static TLockFreeFixedSizeAllocator<sizeof(FRHICommandList), FThreadSafeCounter> RHICommandListAllocator;
+
+void* FRHICommandList::operator new(size_t Size)
+{
+	// doesn't support derived classes with a different size
+	check(Size == sizeof(FRHICommandList));
+	return RHICommandListAllocator.Allocate();
+	//return FMemory::Malloc(Size);
+}
+
+/**
+ * Custom delete
+ */
+void FRHICommandList::operator delete(void *RawMemory)
+{
+	RHICommandListAllocator.Free(RawMemory);
+	//FMemory::Free(RawMemory);
+}	
+
 
diff --git a/Engine/Source/Runtime/RHI/Public/RHICommandList.h b/Engine/Source/Runtime/RHI/Public/RHICommandList.h
--- a/Engine/Source/Runtime/RHI/Public/RHICommandList.h
+++ b/Engine/Source/Runtime/RHI/Public/RHICommandList.h
@@ -61,7 +61,7 @@ public:
 	template <typename TCmd>
 	FORCEINLINE_DEBUGGABLE void* AllocCommand()
 	{
-		checkSlow(!bExecuting);
+		checkSlow(!IsExecuting());
 		TCmd* Result = (TCmd*)Alloc<TCmd>();
 		++NumCommands;
 		*CommandLink = Result;
@@ -76,6 +76,10 @@ public:
 	{
 		return (NumCommands > 0);
 	}
+	FORCEINLINE bool IsExecuting() const
+	{
+		return bExecuting;
+	}
 
 	inline bool Bypass();
 
@@ -1049,6 +1053,11 @@ struct FRHICommandSetLocalUniformBuffer : public FRHICommand<FRHICommandSetLocal
 class RHI_API FRHICommandList : public FRHICommandListBase
 {
 public:
+
+	/** Custom new/delete with recycling */
+	void* operator new(size_t Size);
+	void operator delete(void *RawMemory);
+
 	FORCEINLINE_DEBUGGABLE FLocalBoundShaderState BuildLocalBoundShaderState(const FBoundShaderStateInput& BoundShaderStateInput)
 	{
 		return BuildLocalBoundShaderState(
@@ -1625,7 +1634,7 @@ public:
 
 	FORCEINLINE_DEBUGGABLE void Verify()
 	{
-		check(CommandListImmediate.bExecuting || !CommandListImmediate.HasCommands());
+		check(CommandListImmediate.IsExecuting() || !CommandListImmediate.HasCommands());
 	}
 	FORCEINLINE_DEBUGGABLE bool Bypass()
 	{
diff --git a/Engine/Source/Runtime/Renderer/Private/PostProcess/SceneRenderTargets.cpp b/Engine/Source/Runtime/Renderer/Private/PostProcess/SceneRenderTargets.cpp
--- a/Engine/Source/Runtime/Renderer/Private/PostProcess/SceneRenderTargets.cpp
+++ b/Engine/Source/Runtime/Renderer/Private/PostProcess/SceneRenderTargets.cpp
@@ -251,7 +251,7 @@ void FSceneRenderTargets::ClearGBufferTargets(FRHICommandListImmediate& RHICmdLi
 	}
 }
 
-void FSceneRenderTargets::BeginRenderingSceneColor(FRHICommandListImmediate& RHICmdList, bool bGBufferPass)
+void FSceneRenderTargets::BeginRenderingSceneColor(FRHICommandList& RHICmdList, bool bGBufferPass)
 {
 	SCOPED_DRAW_EVENT(BeginRenderingSceneColor, DEC_SCENE_ITEMS);
 
@@ -743,7 +743,7 @@ void FSceneRenderTargets::FinishRenderingSceneAlphaCopy(FRHICommandListImmediate
 }
 
 
-void FSceneRenderTargets::BeginRenderingLightAttenuation(FRHICommandListImmediate& RHICmdList)
+void FSceneRenderTargets::BeginRenderingLightAttenuation(FRHICommandList& RHICmdList)
 {
 	SCOPED_DRAW_EVENT(BeginRenderingLightAttenuation, DEC_SCENE_ITEMS);
 
@@ -755,7 +755,7 @@ void FSceneRenderTargets::BeginRenderingLightAttenuation(FRHICommandListImmediat
 	SetRenderTarget(RHICmdList, GetLightAttenuationSurface(),GetSceneDepthSurface());
 }
 
-void FSceneRenderTargets::FinishRenderingLightAttenuation(FRHICommandListImmediate& RHICmdList)
+void FSceneRenderTargets::FinishRenderingLightAttenuation(FRHICommandList& RHICmdList)
 {
 	SCOPED_DRAW_EVENT(FinishRenderingLightAttenuation, DEC_SCENE_ITEMS);
 
@@ -765,7 +765,7 @@ void FSceneRenderTargets::FinishRenderingLightAttenuation(FRHICommandListImmedia
 	GRenderTargetPool.VisualizeTexture.SetCheckPoint(RHICmdList, GSceneRenderTargets.GetLightAttenuation());
 }
 
-void FSceneRenderTargets::BeginRenderingTranslucency(FRHICommandListImmediate& RHICmdList, const FViewInfo& View)
+void FSceneRenderTargets::BeginRenderingTranslucency(FRHICommandList& RHICmdList, const FViewInfo& View)
 {
 	// Use the scene color buffer.
 	GSceneRenderTargets.BeginRenderingSceneColor(RHICmdList);
@@ -775,7 +775,7 @@ void FSceneRenderTargets::BeginRenderingTranslucency(FRHICommandListImmediate& R
 }
 
 
-bool FSceneRenderTargets::BeginRenderingSeparateTranslucency(FRHICommandListImmediate& RHICmdList, const FViewInfo& View, bool bFirstTimeThisFrame)
+bool FSceneRenderTargets::BeginRenderingSeparateTranslucency(FRHICommandList& RHICmdList, const FViewInfo& View, bool bFirstTimeThisFrame)
 {
 	if(IsSeparateTranslucencyActive(View))
 	{
diff --git a/Engine/Source/Runtime/Renderer/Private/PostProcess/SceneRenderTargets.h b/Engine/Source/Runtime/Renderer/Private/PostProcess/SceneRenderTargets.h
--- a/Engine/Source/Runtime/Renderer/Private/PostProcess/SceneRenderTargets.h
+++ b/Engine/Source/Runtime/Renderer/Private/PostProcess/SceneRenderTargets.h
@@ -148,7 +148,7 @@ public:
 	 * Sets the scene color target and restores its contents if necessary
 	 * @param bGBufferPass - Whether the pass about to be rendered is the GBuffer population pass
 	 */
-	void BeginRenderingSceneColor(FRHICommandListImmediate& RHICmdList, bool bGBufferPass = false);
+	void BeginRenderingSceneColor(FRHICommandList& RHICmdList, bool bGBufferPass = false);
 	/**
 	 * Called when finished rendering to the scene color surface
 	 * @param bKeepChanges - if true then the SceneColorSurface is resolved to the SceneColorTexture
@@ -185,9 +185,9 @@ public:
 	/** Resolves the appropriate shadow depth cube map and restores default state. */
 	void FinishRenderingCubeShadowDepth(FRHICommandList& RHICmdList, int32 ShadowResolution, const FResolveParams& ResolveParams = FResolveParams());
 	
-	void BeginRenderingTranslucency(FRHICommandListImmediate& RHICmdList, const class FViewInfo& View);
+	void BeginRenderingTranslucency(FRHICommandList& RHICmdList, const class FViewInfo& View);
 
-	bool BeginRenderingSeparateTranslucency(FRHICommandListImmediate& RHICmdList, const FViewInfo& View, bool bFirstTimeThisFrame);
+	bool BeginRenderingSeparateTranslucency(FRHICommandList& RHICmdList, const FViewInfo& View, bool bFirstTimeThisFrame);
 	void FinishRenderingSeparateTranslucency(FRHICommandList& RHICmdList, const FViewInfo& View);
 	void FreeSeparateTranslucency();
 
@@ -200,8 +200,8 @@ public:
 	void BeginRenderingSceneAlphaCopy(FRHICommandListImmediate& RHICmdList);
 	void FinishRenderingSceneAlphaCopy(FRHICommandListImmediate& RHICmdList);
 
-	void BeginRenderingLightAttenuation(FRHICommandListImmediate& RHICmdList);
-	void FinishRenderingLightAttenuation(FRHICommandListImmediate& RHICmdList);
+	void BeginRenderingLightAttenuation(FRHICommandList& RHICmdList);
+	void FinishRenderingLightAttenuation(FRHICommandList& RHICmdList);
 
 	/**
 	 * Cleans up editor primitive targets that we no longer need
diff --git a/Engine/Source/Runtime/Renderer/Private/PostProcess/VisualizeTexture.cpp b/Engine/Source/Runtime/Renderer/Private/PostProcess/VisualizeTexture.cpp
--- a/Engine/Source/Runtime/Renderer/Private/PostProcess/VisualizeTexture.cpp
+++ b/Engine/Source/Runtime/Renderer/Private/PostProcess/VisualizeTexture.cpp
@@ -620,10 +620,11 @@ void FVisualizeTexture::SetObserveTarget(const FString& InObservedDebugName, uin
 	ObservedDebugNameReusedGoal = InObservedDebugNameReusedGoal;
 }
 
-void FVisualizeTexture::SetCheckPoint(FRHICommandListImmediate& RHICmdList, const IPooledRenderTarget* PooledRenderTarget)
+void FVisualizeTexture::SetCheckPoint(FRHICommandList& RHICmdList, const IPooledRenderTarget* PooledRenderTarget)
 {
 
 #if !(UE_BUILD_SHIPPING || UE_BUILD_TEST)
+	check(IsInRenderingThread());
 	if (!PooledRenderTarget || !bEnabled)
 	{
 		// Don't checkpoint on ES2 to avoid TMap alloc/reallocations
@@ -648,7 +649,24 @@ void FVisualizeTexture::SetCheckPoint(FRHICommandListImmediate& RHICmdList, cons
 		// if multiple times reused during the frame, is that the one we want to look at?
 		if(*UsageCountPtr == ObservedDebugNameReusedGoal || ObservedDebugNameReusedGoal == 0xffffffff)
 		{
-			GenerateContent(RHICmdList, RenderTargetItem, Desc);
+			FRHICommandListImmediate& RHICmdListIm = FRHICommandListExecutor::GetImmediateCommandList();
+			if (RHICmdListIm.IsExecuting())
+			{
+				UE_LOG(LogConsoleResponse, Fatal, TEXT("We can't create a checkpoint because that requires the immediate commandlist, which is currently executing. You might try disabling parallel rendering."));
+			}
+			else
+			{
+				if (&RHICmdList != &RHICmdListIm)
+				{
+					UE_LOG(LogConsoleResponse, Warning, TEXT("Attempt to checkpoint a render target from a non-immediate command list. We will flush it and hope that works. If it doesn't you might try disabling parallel rendering."));
+					RHICmdList.Flush();
+				}
+				GenerateContent(RHICmdListIm, RenderTargetItem, Desc);
+				if (&RHICmdList != &RHICmdListIm)
+				{
+					RHICmdListIm.Flush();
+				}
+			}
 		}
 	}
 	// only needed for VisualizeTexture (todo: optimize out when possible)
diff --git a/Engine/Source/Runtime/Renderer/Private/PostProcess/VisualizeTexture.h b/Engine/Source/Runtime/Renderer/Private/PostProcess/VisualizeTexture.h
--- a/Engine/Source/Runtime/Renderer/Private/PostProcess/VisualizeTexture.h
+++ b/Engine/Source/Runtime/Renderer/Private/PostProcess/VisualizeTexture.h
@@ -23,7 +23,7 @@ public:
 	 * calling this allows to grab the state of the texture at this point to be queried by visualizetexture e.g. "vis LightAttenuation@2"
 	 * @param PooledRenderTarget 0 is silently ignored
 	 */
-	void SetCheckPoint(FRHICommandListImmediate& RHICmdList, const IPooledRenderTarget* PooledRenderTarget);
+	void SetCheckPoint(FRHICommandList& RHICmdList, const IPooledRenderTarget* PooledRenderTarget);
 
 	// @param bExtended true: with more convenience - not needed for crashes but useful from the console
 	void DebugLog(bool bExtended);
diff --git a/Engine/Source/Runtime/Renderer/Private/StaticMeshDrawList.h b/Engine/Source/Runtime/Renderer/Private/StaticMeshDrawList.h
--- a/Engine/Source/Runtime/Renderer/Private/StaticMeshDrawList.h
+++ b/Engine/Source/Runtime/Renderer/Private/StaticMeshDrawList.h
@@ -267,6 +267,17 @@ public:
 	 */
 	bool DrawVisible(const FViewInfo& View, const TBitArray<SceneRenderingBitArrayAllocator>& StaticMeshVisibilityMap);
 
+	/**
+	* Draws only the static meshes which are in the visibility map, limited to a range of policies
+	* @param View - The view of the meshes to render.
+	* @param StaticMeshVisibilityMap - An map from FStaticMesh::Id to visibility state.
+	* @param BatchVisibilityArray - An array of batch element visibility bitmasks.
+	* @param FirstPolicy - First policy to render
+	* @param LastPolicy - Last policy to render
+	* @return True if any static meshes were drawn.
+	*/
+	bool DrawVisibleInner(FRHICommandList& RHICmdList, const FViewInfo& View, const TBitArray<SceneRenderingBitArrayAllocator>& StaticMeshVisibilityMap, const TArray<uint64, SceneRenderingAllocator>& BatchVisibilityArray, int32 FirstPolicy, int32 LastPolicy);
+
 	/**
 	 * Draws only the static meshes which are in the visibility map.
 	 * @param View - The view of the meshes to render.
diff --git a/Engine/Source/Runtime/Renderer/Private/StaticMeshDrawList.inl b/Engine/Source/Runtime/Renderer/Private/StaticMeshDrawList.inl
--- a/Engine/Source/Runtime/Renderer/Private/StaticMeshDrawList.inl
+++ b/Engine/Source/Runtime/Renderer/Private/StaticMeshDrawList.inl
@@ -241,28 +241,29 @@ bool TStaticMeshDrawList<DrawingPolicyType>::DrawVisible(
 }
 
 template<typename DrawingPolicyType>
-bool TStaticMeshDrawList<DrawingPolicyType>::DrawVisible(
+bool TStaticMeshDrawList<DrawingPolicyType>::DrawVisibleInner(
 	FRHICommandList& RHICmdList,
 	const FViewInfo& View,
 	const TBitArray<SceneRenderingBitArrayAllocator>& StaticMeshVisibilityMap,
-	const TArray<uint64,SceneRenderingAllocator>& BatchVisibilityArray
+	const TArray<uint64, SceneRenderingAllocator>& BatchVisibilityArray,
+	int32 FirstPolicy, int32 LastPolicy
 	)
 {
 	bool bDirty = false;
-	for(typename TArray<FSetElementId>::TConstIterator PolicyIt(OrderedDrawingPolicies); PolicyIt; ++PolicyIt)
+	for (int32 Index = FirstPolicy; Index <= LastPolicy; Index++)
 	{
-		FDrawingPolicyLink* DrawingPolicyLink = &DrawingPolicySet[*PolicyIt];
+		FDrawingPolicyLink* DrawingPolicyLink = &DrawingPolicySet[OrderedDrawingPolicies[Index]];
 		bool bDrawnShared = false;
 		FPlatformMisc::Prefetch(DrawingPolicyLink->CompactElements.GetTypedData());
 		const int32 NumElements = DrawingPolicyLink->Elements.Num();
 		FPlatformMisc::Prefetch(&DrawingPolicyLink->CompactElements.GetTypedData()->MeshId);
 		const FElementCompact* CompactElementPtr = DrawingPolicyLink->CompactElements.GetTypedData();
-		for(int32 ElementIndex = 0; ElementIndex < NumElements; ElementIndex++, CompactElementPtr++)
+		for (int32 ElementIndex = 0; ElementIndex < NumElements; ElementIndex++, CompactElementPtr++)
 		{
-			if(StaticMeshVisibilityMap.AccessCorrespondingBit(FRelativeBitReference(CompactElementPtr->MeshId)))
+			if (StaticMeshVisibilityMap.AccessCorrespondingBit(FRelativeBitReference(CompactElementPtr->MeshId)))
 			{
 				const FElement& Element = DrawingPolicyLink->Elements[ElementIndex];
-				INC_DWORD_STAT_BY(STAT_StaticMeshTriangles,Element.Mesh->GetNumPrimitives());
+				INC_DWORD_STAT_BY(STAT_StaticMeshTriangles, Element.Mesh->GetNumPrimitives());
 				// Avoid the cache miss looking up batch visibility if there is only one element.
 				uint64 BatchElementMask = Element.Mesh->Elements.Num() == 1 ? 1 : BatchVisibilityArray[Element.Mesh->Id];
 				DrawElement(RHICmdList, View, Element, BatchElementMask, DrawingPolicyLink, bDrawnShared);
@@ -273,6 +274,17 @@ bool TStaticMeshDrawList<DrawingPolicyType>::DrawVisible(
 	return bDirty;
 }
 
+template<typename DrawingPolicyType>
+bool TStaticMeshDrawList<DrawingPolicyType>::DrawVisible(
+	FRHICommandList& RHICmdList,
+	const FViewInfo& View,
+	const TBitArray<SceneRenderingBitArrayAllocator>& StaticMeshVisibilityMap,
+	const TArray<uint64,SceneRenderingAllocator>& BatchVisibilityArray
+	)
+{
+	return DrawVisibleInner(RHICmdList, View, StaticMeshVisibilityMap, BatchVisibilityArray, 0, OrderedDrawingPolicies.Num() - 1);
+}
+
 
 template<typename DrawingPolicyType>
 class FDrawVisibleAnyThreadTask
@@ -325,26 +337,9 @@ public:
 
 	void DoTask(ENamedThreads::Type CurrentThread, const FGraphEventRef& MyCompletionGraphEvent)
 	{
-		for (int32 Index = this->FirstPolicy; Index <= this->LastPolicy; Index++)
+		if (this->Caller.DrawVisibleInner(this->RHICmdList, this->View, this->StaticMeshVisibilityMap, this->BatchVisibilityArray, this->FirstPolicy, this->LastPolicy))
 		{
-			typename TStaticMeshDrawList<DrawingPolicyType>::FDrawingPolicyLink* DrawingPolicyLink = &this->Caller.DrawingPolicySet[this->Caller.OrderedDrawingPolicies[Index]];
-			bool bDrawnShared = false;
-			FPlatformMisc::Prefetch(DrawingPolicyLink->CompactElements.GetTypedData());
-			const int32 NumElements = DrawingPolicyLink->Elements.Num();
-			FPlatformMisc::Prefetch(&DrawingPolicyLink->CompactElements.GetTypedData()->MeshId);
-			const typename TStaticMeshDrawList<DrawingPolicyType>::FElementCompact* CompactElementPtr = DrawingPolicyLink->CompactElements.GetTypedData();
-			for (int32 ElementIndex = 0; ElementIndex < NumElements; ElementIndex++, CompactElementPtr++)
-			{
-				if (this->StaticMeshVisibilityMap.AccessCorrespondingBit(FRelativeBitReference(CompactElementPtr->MeshId)))
-				{
-					const typename TStaticMeshDrawList<DrawingPolicyType>::FElement& Element = DrawingPolicyLink->Elements[ElementIndex];
-					INC_DWORD_STAT_BY(STAT_StaticMeshTriangles, Element.Mesh->GetNumPrimitives());
-					// Avoid the cache miss looking up batch visibility if there is only one element.
-					uint64 BatchElementMask = Element.Mesh->Elements.Num() == 1 ? 1 : this->BatchVisibilityArray[Element.Mesh->Id];
-					this->Caller.DrawElement(this->RHICmdList, this->View, Element, BatchElementMask, DrawingPolicyLink, bDrawnShared);
-					this->OutDirty = true;
-				}
-			}
+			this->OutDirty = true;
 		}
 	}
 };
diff --git a/Engine/Source/Runtime/Renderer/Private/TranslucentRendering.cpp b/Engine/Source/Runtime/Renderer/Private/TranslucentRendering.cpp
--- a/Engine/Source/Runtime/Renderer/Private/TranslucentRendering.cpp
+++ b/Engine/Source/Runtime/Renderer/Private/TranslucentRendering.cpp
@@ -10,7 +10,7 @@
 #include "SceneFilterRendering.h"
 #include "LightPropagationVolume.h"
 
-void SetTranslucentRenderTargetAndState(FRHICommandListImmediate& RHICmdList, const FViewInfo& View, bool bSeparateTranslucencyPass)
+void SetTranslucentRenderTargetAndState(FRHICommandList& RHICmdList, const FViewInfo& View, bool bSeparateTranslucencyPass)
 {
 	bool bSetupTranslucentState = true;
 
@@ -70,6 +70,7 @@ const FProjectedShadowInfo* FDeferredShadingSceneRenderer::PrepareTranslucentSha
 		// Allocate and render the shadow's depth map if needed
 		if (TranslucentSelfShadow && !TranslucentSelfShadow->bAllocatedInTranslucentLayout)
 		{
+			check(IsInRenderingThread());
 			bool bPossibleToAllocate = true;
 
 			// Attempt to find space in the layout
@@ -166,7 +167,7 @@ IMPLEMENT_SHADER_TYPE(,FCopySceneColorPS,TEXT("TranslucentLightingShaders"),TEXT
 
 FGlobalBoundShaderState CopySceneColorBoundShaderState;
 
-void FTranslucencyDrawingPolicyFactory::CopySceneColor(FRHICommandListImmediate& RHICmdList, const FViewInfo& View, const FPrimitiveSceneProxy* PrimitiveSceneProxy)
+void FTranslucencyDrawingPolicyFactory::CopySceneColor(FRHICommandList& RHICmdList, const FViewInfo& View, const FPrimitiveSceneProxy* PrimitiveSceneProxy)
 {
 
 	SCOPED_DRAW_EVENTF(EventCopy, DEC_SCENE_ITEMS, TEXT("CopySceneColor for %s %s"), *PrimitiveSceneProxy->GetOwnerName().ToString(), *PrimitiveSceneProxy->GetResourceName().ToString());
@@ -300,6 +301,31 @@ public:
 	}
 };
 
+static void CopySceneColorAndRestore(FRHICommandList& RHICmdList, const FViewInfo& View, const FPrimitiveSceneProxy* PrimitiveSceneProxy)
+{
+	FTranslucencyDrawingPolicyFactory::CopySceneColor(RHICmdList, View, PrimitiveSceneProxy);
+	// Restore state
+	SetTranslucentRenderTargetAndState(RHICmdList, View, false);
+}
+
+struct FRHICommandCopySceneColor : public FRHICommand<FRHICommandCopySceneColor>
+{
+	const FViewInfo& View;
+	const FPrimitiveSceneProxy* PrimitiveSceneProxy;
+
+	FORCEINLINE_DEBUGGABLE FRHICommandCopySceneColor(const FViewInfo& InView, const FPrimitiveSceneProxy* InPrimitiveSceneProxy)
+		: View(InView)
+		, PrimitiveSceneProxy(InPrimitiveSceneProxy)
+	{
+	}
+	void Execute()
+	{
+		check(IsInRenderingThread());
+		FRHICommandList RHICmdList;
+		CopySceneColorAndRestore(RHICmdList, View, PrimitiveSceneProxy);
+	}
+};
+
 /**
 * Render a dynamic or static mesh using a translucent draw policy
 * @return true if the mesh rendered
@@ -331,14 +357,17 @@ bool FTranslucencyDrawingPolicyFactory::DrawMesh(
 		{
 			if (DrawingContext.bSceneColorCopyIsUpToDate == false)
 			{
-				//hack here to get the code checked in, but this needs to be fixed. It was not sensible to wind this all the way out of the callstack
 				FRHICommandListImmediate& RHICmdListIm = FRHICommandListExecutor::GetImmediateCommandList();
-				check(&RHICmdList == &RHICmdListIm);
-
-				CopySceneColor(RHICmdListIm, View, PrimitiveSceneProxy);
-				// Restore state
-				SetTranslucentRenderTargetAndState(RHICmdListIm, View, false);
-
+				if (&RHICmdList != &RHICmdListIm)
+				{
+					// if we are buffering commands on a non-immediate command list, we need to defer the copy until commandlist execution
+					new (RHICmdList.AllocCommand<FRHICommandCopySceneColor>()) FRHICommandCopySceneColor(View, PrimitiveSceneProxy);
+				}
+				else
+				{
+					// otherwise, just do it now. We don't want to defer in this case because that can interfere with render target visualization (a debugging tool).
+					CopySceneColorAndRestore(RHICmdList, View, PrimitiveSceneProxy);
+				}
 				// separate translucency is not updating scene color so we don't need to copy it multiple times.
 				// optimization:
 				// we should consider the same for non separate translucency (could cause artifacts but will be much faster)
diff --git a/Engine/Source/Runtime/Renderer/Private/TranslucentRendering.h b/Engine/Source/Runtime/Renderer/Private/TranslucentRendering.h
--- a/Engine/Source/Runtime/Renderer/Private/TranslucentRendering.h
+++ b/Engine/Source/Runtime/Renderer/Private/TranslucentRendering.h
@@ -63,6 +63,10 @@ public:
 		return !IsTranslucentBlendMode(MaterialRenderProxy->GetMaterial(InFeatureLevel)->GetBlendMode());
 	}
 
+	/**
+	* Resolves the scene color target and copies it for use as a source texture.
+	*/
+	static void CopySceneColor(FRHICommandList& RHICmdList, const FViewInfo& View, const FPrimitiveSceneProxy* PrimitiveSceneProxy);
 
 private:
 	/**
@@ -81,10 +85,6 @@ private:
 		FHitProxyId HitProxyId
 		);
 
-	/**
-	* Resolves the scene color target and copies it for use as a source texture.
-	*/
-	static void CopySceneColor(FRHICommandListImmediate& RHICmdList, const FViewInfo& View, const FPrimitiveSceneProxy* PrimitiveSceneProxy);
 };
 
 
