UE4 - parallel velocity and shadow rendering

diff --git a/Engine/Source/Runtime/Renderer/Private/DeferredShadingRenderer.cpp b/Engine/Source/Runtime/Renderer/Private/DeferredShadingRenderer.cpp
--- a/Engine/Source/Runtime/Renderer/Private/DeferredShadingRenderer.cpp
+++ b/Engine/Source/Runtime/Renderer/Private/DeferredShadingRenderer.cpp
@@ -461,14 +461,7 @@ FGraphEventRef FDeferredShadingSceneRenderer::RenderBasePassDynamicDataParallel(
 	FGraphEventRef AnyThreadCompletionEvent = TGraphTask<FRenderBasePassDynamicDataThreadTask>::CreateTask(nullptr, ENamedThreads::RenderThread)
 		.ConstructAndDispatchWhenReady(this, CmdList, &View, &bOutDirty);
 
-	FGraphEventArray Prereqs;
-	Prereqs.Add(AnyThreadCompletionEvent);
-	if (SubmitChain.GetReference())
-	{
-		Prereqs.Add(SubmitChain);
-	}
-
-	return TGraphTask<FSubmitCommandlistThreadTask>::CreateTask(&Prereqs, ENamedThreads::RenderThread).ConstructAndDispatchWhenReady(CmdList);
+	return FSubmitCommandlistThreadTask::AddToChain(AnyThreadCompletionEvent, SubmitChain, CmdList);
 }
 
 static void SetupBasePassView(FRHICommandList& RHICmdList, const FViewInfo& View, bool bShaderComplexity)
@@ -1169,14 +1162,7 @@ FGraphEventRef FDeferredShadingSceneRenderer::RenderPrePassViewParallel(FRHIComm
 		FGraphEventRef AnyThreadCompletionEvent = TGraphTask<FRenderPrepassDynamicDataThreadTask>::CreateTask(nullptr, ENamedThreads::RenderThread)
 			.ConstructAndDispatchWhenReady(this, CmdList, &View, &OutDirty);
 
-		FGraphEventArray Prereqs;
-		Prereqs.Add(AnyThreadCompletionEvent);
-		if (SubmitChain.GetReference())
-		{
-			Prereqs.Add(SubmitChain);
-		}
-
-		SubmitChain = TGraphTask<FSubmitCommandlistThreadTask>::CreateTask(&Prereqs, ENamedThreads::RenderThread).ConstructAndDispatchWhenReady(CmdList);
+		SubmitChain = FSubmitCommandlistThreadTask::AddToChain(AnyThreadCompletionEvent, SubmitChain, CmdList);
 	}
 
 	return SubmitChain;
diff --git a/Engine/Source/Runtime/Renderer/Private/DeferredShadingRenderer.h b/Engine/Source/Runtime/Renderer/Private/DeferredShadingRenderer.h
--- a/Engine/Source/Runtime/Renderer/Private/DeferredShadingRenderer.h
+++ b/Engine/Source/Runtime/Renderer/Private/DeferredShadingRenderer.h
@@ -247,6 +247,13 @@ private:
 	/** Renders the velocities of movable objects for the motion blur effect. */
 	void RenderVelocities(FRHICommandListImmediate& RHICmdList, const FViewInfo& View, TRefCountPtr<IPooledRenderTarget>& VelocityRT, bool bLastFrame);
 
+	/** Renders the velocities of movable objects for the motion blur effect. */
+	friend class FRenderVelocityDynamicThreadTask;
+	void RenderDynamicVelocities(FRHICommandList& RHICmdList, const FViewInfo& View);
+
+	/** Renders the velocities of movable objects for the motion blur effect. */
+	void RenderVelocitiesInner(FRHICommandListImmediate& RHICmdList, const FViewInfo& View);
+
 	/** Renders world-space lightmap density instead of the normal color. */
 	bool RenderLightMapDensities(FRHICommandListImmediate& RHICmdList);
 
diff --git a/Engine/Source/Runtime/Renderer/Private/RendererScene.cpp b/Engine/Source/Runtime/Renderer/Private/RendererScene.cpp
--- a/Engine/Source/Runtime/Renderer/Private/RendererScene.cpp
+++ b/Engine/Source/Runtime/Renderer/Private/RendererScene.cpp
@@ -2365,7 +2365,7 @@ FMotionBlurInfo* FMotionBlurInfoData::FindMBInfoIndex(FPrimitiveComponentId Comp
 
 bool FMotionBlurInfoData::GetPrimitiveMotionBlurInfo(const FPrimitiveSceneInfo* PrimitiveSceneInfo, FMatrix& OutPreviousLocalToWorld)
 {
-	check(IsInRenderingThread());
+	check(IsInParallelRenderingThread());
 
 	if (PrimitiveSceneInfo && PrimitiveSceneInfo->PrimitiveComponentId.IsValid())
 	{
diff --git a/Engine/Source/Runtime/Renderer/Private/ShadowRendering.cpp b/Engine/Source/Runtime/Renderer/Private/ShadowRendering.cpp
--- a/Engine/Source/Runtime/Renderer/Private/ShadowRendering.cpp
+++ b/Engine/Source/Runtime/Renderer/Private/ShadowRendering.cpp
@@ -1263,6 +1263,232 @@ void DrawShadowMeshElements(FRHICommandList& RHICmdList, const FViewInfo& View,
 	}
 }
 
+void FProjectedShadowInfo::RenderDepthDynamic(FRHICommandList& RHICmdList, FDeferredShadingSceneRenderer* SceneRenderer, const FViewInfo* FoundView)
+{
+	// Draw the subject's dynamic elements.
+	SCOPE_CYCLE_COUNTER(STAT_WholeSceneDynamicShadowDepthsTime);
+
+	uint32 DrawPrimitiveFlags = 0;
+	if (bReflectiveShadowmap)
+	{
+		// force lowest LOD for RSMs
+		DrawPrimitiveFlags = EDrawDynamicFlags::ForceLowestLOD;
+	}
+
+	TDynamicPrimitiveDrawer<FShadowDepthDrawingPolicyFactory> Drawer(RHICmdList, FoundView, FShadowDepthDrawingPolicyFactory::ContextType(this), true);
+	uint32 PrimitiveCount = SubjectPrimitives.Num();
+	for (uint32 PrimitiveIndex = 0; PrimitiveIndex < PrimitiveCount; ++PrimitiveIndex)
+	{
+		const FPrimitiveSceneInfo* PrimitiveSceneInfo = SubjectPrimitives[PrimitiveIndex];
+
+		// Lookup the primitive's cached view relevance
+		FPrimitiveViewRelevance ViewRelevance = FoundView->PrimitiveViewRelevanceMap[PrimitiveSceneInfo->GetIndex()];
+
+		if (!ViewRelevance.bInitializedThisFrame)
+		{
+			// Compute the subject primitive's view relevance since it wasn't cached
+			ViewRelevance = PrimitiveSceneInfo->Proxy->GetViewRelevance(FoundView);
+		}
+
+		// Only draw if the subject primitive is shadow relevant.
+		if (ViewRelevance.bShadowRelevance)
+		{
+			FScopeCycleCounter Context(SubjectPrimitives[PrimitiveIndex]->Proxy->GetStatId());
+			Drawer.SetPrimitive(SubjectPrimitives[PrimitiveIndex]->Proxy);
+
+			SubjectPrimitives[PrimitiveIndex]->Proxy->DrawDynamicElements(&Drawer, FoundView, DrawPrimitiveFlags);
+		}
+	}
+}
+
+class FDrawShadowMeshElementsThreadTask
+{
+	FProjectedShadowInfo& ThisRenderer;
+	FRHICommandList& RHICmdList;
+	const FViewInfo& View;
+	bool bReflective;
+
+public:
+
+	FDrawShadowMeshElementsThreadTask(
+		FProjectedShadowInfo* InThisRenderer,
+		FRHICommandList* InRHICmdList,
+		const FViewInfo* InView,
+		bool InbReflective
+		)
+		: ThisRenderer(*InThisRenderer)
+		, RHICmdList(*InRHICmdList)
+		, View(*InView)
+		, bReflective(InbReflective)
+	{
+	}
+
+	FORCEINLINE TStatId GetStatId() const
+	{
+		RETURN_QUICK_DECLARE_CYCLE_STAT(FDrawShadowMeshElementsThreadTask, STATGROUP_TaskGraphTasks);
+	}
+
+	ENamedThreads::Type GetDesiredThread()
+	{
+		return ENamedThreads::AnyThread;
+	}
+
+	static ESubsequentsMode::Type GetSubsequentsMode() { return ESubsequentsMode::TrackSubsequents; }
+
+	void DoTask(ENamedThreads::Type CurrentThread, const FGraphEventRef& MyCompletionGraphEvent)
+	{
+		SCOPE_CYCLE_COUNTER(STAT_WholeSceneStaticShadowDepthsTime);
+
+		if (bReflective)
+		{
+			// reflective shadow map
+			DrawShadowMeshElements<true>(RHICmdList, View, ThisRenderer);
+		}
+		else
+		{
+			// normal shadow map
+			DrawShadowMeshElements<false>(RHICmdList, View, ThisRenderer);
+		}
+	}
+};
+
+class FRenderDepthDynamicThreadTask
+{
+	FProjectedShadowInfo& ThisRenderer;
+	FRHICommandList& RHICmdList;
+	const FViewInfo& View;
+	FDeferredShadingSceneRenderer* SceneRenderer;
+
+public:
+
+	FRenderDepthDynamicThreadTask(
+		FProjectedShadowInfo* InThisRenderer,
+		FRHICommandList* InRHICmdList,
+		const FViewInfo* InView,
+		FDeferredShadingSceneRenderer* InSceneRenderer
+		)
+		: ThisRenderer(*InThisRenderer)
+		, RHICmdList(*InRHICmdList)
+		, View(*InView)
+		, SceneRenderer(InSceneRenderer)
+	{
+	}
+
+	FORCEINLINE TStatId GetStatId() const
+	{
+		RETURN_QUICK_DECLARE_CYCLE_STAT(FRenderDepthDynamicThreadTask, STATGROUP_TaskGraphTasks);
+	}
+
+	ENamedThreads::Type GetDesiredThread()
+	{
+		return ENamedThreads::AnyThread;
+	}
+
+	static ESubsequentsMode::Type GetSubsequentsMode() { return ESubsequentsMode::TrackSubsequents; }
+
+	void DoTask(ENamedThreads::Type CurrentThread, const FGraphEventRef& MyCompletionGraphEvent)
+	{
+		ThisRenderer.RenderDepthDynamic(RHICmdList, SceneRenderer, &View);
+	}
+};
+
+void FProjectedShadowInfo::RenderDepthInner(FRHICommandListImmediate& RHICmdList, FDeferredShadingSceneRenderer* SceneRenderer, const FViewInfo* FoundView)
+{
+	static TAutoConsoleVariable<int32> CVarParallelShadows(
+		TEXT("r.ParallelShadows"),
+		1,
+		TEXT("Toggles parallel shadow rendering. Parallel rendering must be enabled for this to have an effect.\n"),
+		ECVF_RenderThreadSafe
+		);
+
+	if (!GRHICommandList.Bypass() && CVarParallelShadows.GetValueOnRenderThread())
+	{
+		// parallel version
+
+		int32 Width = CVarRHICmdWidth.GetValueOnRenderThread(); // we use a few more than needed to cover non-equal jobs
+		bool OutDirty = false; // unused
+
+		FGraphEventRef SubmitChain;
+
+		// Draw the subject's static elements using static draw lists
+		if (IsWholeSceneDirectionalShadow())
+		{
+			SCOPE_CYCLE_COUNTER(STAT_WholeSceneStaticDrawListShadowDepthsTime);
+
+			if (bReflectiveShadowmap)
+			{
+				SubmitChain = SceneRenderer->Scene->WholeSceneReflectiveShadowMapDrawList.DrawVisibleParallel(*FoundView, StaticMeshWholeSceneShadowDepthMap, StaticMeshWholeSceneShadowBatchVisibility, Width, SubmitChain, OutDirty);
+			}
+			else
+			{
+				// Use the scene's shadow depth draw list with this shadow's visibility map
+				SubmitChain = SceneRenderer->Scene->WholeSceneShadowDepthDrawList.DrawVisibleParallel(*FoundView, StaticMeshWholeSceneShadowDepthMap, StaticMeshWholeSceneShadowBatchVisibility, Width, SubmitChain, OutDirty);
+			}
+		}
+		// Draw the subject's static elements using manual state filtering
+		else if (SubjectMeshElements.Num() > 0)
+		{
+			FRHICommandList* CmdList = new FRHICommandList;
+
+			FGraphEventRef AnyThreadCompletionEvent = TGraphTask<FDrawShadowMeshElementsThreadTask>::CreateTask(nullptr, ENamedThreads::RenderThread)
+				.ConstructAndDispatchWhenReady(this, CmdList, FoundView, bReflectiveShadowmap && !bOnePassPointLightShadow);
+
+			SubmitChain = FSubmitCommandlistThreadTask::AddToChain(AnyThreadCompletionEvent, SubmitChain, CmdList);
+		}
+		{
+			FRHICommandList* CmdList = new FRHICommandList;
+
+			FGraphEventRef AnyThreadCompletionEvent = TGraphTask<FRenderDepthDynamicThreadTask>::CreateTask(nullptr, ENamedThreads::RenderThread)
+				.ConstructAndDispatchWhenReady(this, CmdList, FoundView, SceneRenderer);
+
+			SubmitChain = FSubmitCommandlistThreadTask::AddToChain(AnyThreadCompletionEvent, SubmitChain, CmdList);
+		}
+
+		if (SubmitChain.GetReference())
+		{
+			RHICmdList.Flush(); // this would happen later anyway, but we might as well do these while we wait for async tasks
+			FTaskGraphInterface::Get().WaitUntilTaskCompletes(SubmitChain, ENamedThreads::RenderThread_Local);
+		}
+	}
+	else
+	{
+		// single threaded version
+
+		// Draw the subject's static elements using static draw lists
+		if (IsWholeSceneDirectionalShadow())
+		{
+			SCOPE_CYCLE_COUNTER(STAT_WholeSceneStaticDrawListShadowDepthsTime);
+
+			if (bReflectiveShadowmap)
+			{
+				SceneRenderer->Scene->WholeSceneReflectiveShadowMapDrawList.DrawVisible(RHICmdList, *FoundView, StaticMeshWholeSceneShadowDepthMap, StaticMeshWholeSceneShadowBatchVisibility);
+			}
+			else
+			{
+				// Use the scene's shadow depth draw list with this shadow's visibility map
+				SceneRenderer->Scene->WholeSceneShadowDepthDrawList.DrawVisible(RHICmdList, *FoundView, StaticMeshWholeSceneShadowDepthMap, StaticMeshWholeSceneShadowBatchVisibility);
+			}
+		}
+		// Draw the subject's static elements using manual state filtering
+		else if (SubjectMeshElements.Num() > 0)
+		{
+			SCOPE_CYCLE_COUNTER(STAT_WholeSceneStaticShadowDepthsTime);
+
+			if (bReflectiveShadowmap && !bOnePassPointLightShadow)
+			{
+				// reflective shadow map
+				DrawShadowMeshElements<true>(RHICmdList, *FoundView, *this);
+			}
+			else
+			{
+				// normal shadow map
+				DrawShadowMeshElements<false>(RHICmdList, *FoundView, *this);
+			}
+		}
+		RenderDepthDynamic(RHICmdList, SceneRenderer, FoundView);
+	}
+}
+
 void FProjectedShadowInfo::RenderDepth(FRHICommandListImmediate& RHICmdList, FDeferredShadingSceneRenderer* SceneRenderer)
 {
 #if WANTS_DRAW_MESH_EVENTS
@@ -1341,7 +1567,8 @@ void FProjectedShadowInfo::RenderDepth(FRHICommandListImmediate& RHICmdList, FDe
 			break;
 		}
 	}
-	check(FoundView);
+	check(FoundView 
+		&& IsInRenderingThread()); // we should not hack the view in parallel
 
 	// Backup properties of the view that we will override
 	TUniformBufferRef<FViewUniformShaderParameters> OriginalUniformBuffer = FoundView->UniformBuffer;
@@ -1363,77 +1590,7 @@ void FProjectedShadowInfo::RenderDepth(FRHICommandListImmediate& RHICmdList, FDe
 	FShadowDepthDrawingPolicy<false>::PolicyShadowInfo = this;
 	FShadowDepthDrawingPolicy<true>::PolicyShadowInfo = this;
 
-	{
-		// Draw the subject's static elements using static draw lists
-		if (IsWholeSceneDirectionalShadow())
-		{
-			SCOPE_CYCLE_COUNTER(STAT_WholeSceneStaticDrawListShadowDepthsTime);
-
-			if (bReflectiveShadowmap)
-			{
-				SceneRenderer->Scene->WholeSceneReflectiveShadowMapDrawList.DrawVisible(RHICmdList, *FoundView, StaticMeshWholeSceneShadowDepthMap, StaticMeshWholeSceneShadowBatchVisibility);
-			}
-			else
-			{
-				// Use the scene's shadow depth draw list with this shadow's visibility map
-				SceneRenderer->Scene->WholeSceneShadowDepthDrawList.DrawVisible(RHICmdList, *FoundView, StaticMeshWholeSceneShadowDepthMap, StaticMeshWholeSceneShadowBatchVisibility);
-			}
-		}
-		// Draw the subject's static elements using manual state filtering
-		else if (SubjectMeshElements.Num() > 0)
-		{
-			SCOPE_CYCLE_COUNTER(STAT_WholeSceneStaticShadowDepthsTime);
-
-			if(bReflectiveShadowmap && !bOnePassPointLightShadow)
-			{
-				// reflective shadow map
-				DrawShadowMeshElements<true>(RHICmdList, *FoundView, *this);
-			}
-			else
-			{
-				// normal shadow map
-				DrawShadowMeshElements<false>(RHICmdList, *FoundView, *this);
-			}
-		}
-
-	}
-
-	// Draw the subject's dynamic elements.
-	{
-		SCOPE_CYCLE_COUNTER(STAT_WholeSceneDynamicShadowDepthsTime);
-
-		uint32 DrawPrimitiveFlags = 0;
-		if(bReflectiveShadowmap)
-		{
-			// force lowest LOD for RSMs
-			DrawPrimitiveFlags = EDrawDynamicFlags::ForceLowestLOD;
-		}
-
-		TDynamicPrimitiveDrawer<FShadowDepthDrawingPolicyFactory> Drawer(RHICmdList, FoundView, FShadowDepthDrawingPolicyFactory::ContextType(this), true);
-		uint32 PrimitiveCount = SubjectPrimitives.Num();
-		for(uint32 PrimitiveIndex = 0; PrimitiveIndex < PrimitiveCount; ++PrimitiveIndex)
-		{
-			const FPrimitiveSceneInfo* PrimitiveSceneInfo = SubjectPrimitives[PrimitiveIndex];
-
-			// Lookup the primitive's cached view relevance
-			FPrimitiveViewRelevance ViewRelevance = FoundView->PrimitiveViewRelevanceMap[PrimitiveSceneInfo->GetIndex()];
-
-			if (!ViewRelevance.bInitializedThisFrame)
-			{
-				// Compute the subject primitive's view relevance since it wasn't cached
-				ViewRelevance = PrimitiveSceneInfo->Proxy->GetViewRelevance(FoundView);
-			}
-		
-			// Only draw if the subject primitive is shadow relevant.
-			if (ViewRelevance.bShadowRelevance)
-			{
-				FScopeCycleCounter Context(SubjectPrimitives[PrimitiveIndex]->Proxy->GetStatId());
-				Drawer.SetPrimitive(SubjectPrimitives[PrimitiveIndex]->Proxy);
-
-				SubjectPrimitives[PrimitiveIndex]->Proxy->DrawDynamicElements(&Drawer, FoundView, DrawPrimitiveFlags);
-			}
-		}
-	}
+	RenderDepthInner(RHICmdList, SceneRenderer, FoundView);
 
 	FShadowDepthDrawingPolicy<false>::PolicyShadowInfo = NULL;
 	FShadowDepthDrawingPolicy<true>::PolicyShadowInfo = NULL;
diff --git a/Engine/Source/Runtime/Renderer/Private/ShadowRendering.h b/Engine/Source/Runtime/Renderer/Private/ShadowRendering.h
--- a/Engine/Source/Runtime/Renderer/Private/ShadowRendering.h
+++ b/Engine/Source/Runtime/Renderer/Private/ShadowRendering.h
@@ -741,6 +741,17 @@ public:
 
 private:
 
+	/**
+	* Renders the shadow subject depth, to a particular hacked view
+	*/
+	void RenderDepthInner(FRHICommandListImmediate& RHICmdList, class FDeferredShadingSceneRenderer* SceneRenderer, const FViewInfo* FoundView);
+
+	/**
+	* Renders the dynamic shadow subject depth, to a particular hacked view
+	*/
+	friend class FRenderDepthDynamicThreadTask;
+	void RenderDepthDynamic(FRHICommandList& RHICmdList, class FDeferredShadingSceneRenderer* SceneRenderer, const FViewInfo* FoundView);
+
 	void GetShadowTypeNameForDrawEvent(FString& TypeName) const;
 
 	template <bool bReflectiveShadowmap> friend void DrawShadowMeshElements(FRHICommandList& RHICmdList, const FViewInfo& View, const FProjectedShadowInfo& ShadowInfo);
diff --git a/Engine/Source/Runtime/Renderer/Private/StaticMeshDrawList.inl b/Engine/Source/Runtime/Renderer/Private/StaticMeshDrawList.inl
--- a/Engine/Source/Runtime/Renderer/Private/StaticMeshDrawList.inl
+++ b/Engine/Source/Runtime/Renderer/Private/StaticMeshDrawList.inl
@@ -371,6 +371,18 @@ public:
 		QUICK_SCOPE_CYCLE_COUNTER(AA_FDrawVisibleRenderThreadTask);
 		delete RHICmdList;
 	}
+
+	/*** Helper function set up a submit chain of the these   */
+	static FGraphEventRef AddToChain(FGraphEventRef AnyThreadCompletionEvent, FGraphEventRef SubmitChain, FRHICommandList* CmdList)
+	{
+		FGraphEventArray Prereqs;
+		Prereqs.Add(AnyThreadCompletionEvent);
+		if (SubmitChain.GetReference())
+		{
+			Prereqs.Add(SubmitChain);
+		}
+		return TGraphTask<FSubmitCommandlistThreadTask>::CreateTask(&Prereqs, ENamedThreads::RenderThread).ConstructAndDispatchWhenReady(CmdList);
+	}
 };
 
 template<typename DrawingPolicyType>
@@ -403,15 +415,7 @@ FGraphEventRef TStaticMeshDrawList<DrawingPolicyType>::DrawVisibleParallel(
 				.ConstructAndDispatchWhenReady(this, CmdList, &View, &StaticMeshVisibilityMap, &BatchVisibilityArray, Start, Last, &OutDirty);
 			Start = Last + 1;
 
-			FGraphEventArray Prereqs;
-			Prereqs.Add(AnyThreadCompletionEvent);
-			if (SubmitChain.GetReference())
-			{
-				Prereqs.Add(SubmitChain);
-			}
-
-			SubmitChain = TGraphTask<FSubmitCommandlistThreadTask>::CreateTask(&Prereqs, ENamedThreads::RenderThread).ConstructAndDispatchWhenReady(CmdList);
-
+			SubmitChain = FSubmitCommandlistThreadTask::AddToChain(AnyThreadCompletionEvent, SubmitChain, CmdList);
 		}
 	}
 	check(Start == OrderedDrawingPolicies.Num());
diff --git a/Engine/Source/Runtime/Renderer/Private/VelocityRendering.cpp b/Engine/Source/Runtime/Renderer/Private/VelocityRendering.cpp
--- a/Engine/Source/Runtime/Renderer/Private/VelocityRendering.cpp
+++ b/Engine/Source/Runtime/Renderer/Private/VelocityRendering.cpp
@@ -441,6 +441,110 @@ bool IsMotionBlurEnabled(const FViewInfo& View)
 		&& !(View.Family->Views.Num() > 1);
 }
 
+void FDeferredShadingSceneRenderer::RenderDynamicVelocities(FRHICommandList& RHICmdList, const FViewInfo& View)
+{
+	// Draw velocities for movable dynamic meshes.
+	TDynamicPrimitiveDrawer<FVelocityDrawingPolicyFactory> Drawer(
+		RHICmdList, &View, FVelocityDrawingPolicyFactory::ContextType(DDM_AllOccluders), true, false, true
+		);
+	for (int32 PrimitiveIndex = 0; PrimitiveIndex < View.VisibleDynamicPrimitives.Num(); PrimitiveIndex++)
+	{
+		const FPrimitiveSceneInfo* PrimitiveSceneInfo = View.VisibleDynamicPrimitives[PrimitiveIndex];
+
+		if (!PrimitiveSceneInfo->ShouldRenderVelocity(View))
+		{
+			continue;
+		}
+
+		FScopeCycleCounter Context(PrimitiveSceneInfo->Proxy->GetStatId());
+
+		Drawer.SetPrimitive(PrimitiveSceneInfo->Proxy);
+		PrimitiveSceneInfo->Proxy->DrawDynamicElements(&Drawer, &View);
+	}
+	Drawer.IsDirty();
+}
+
+class FRenderVelocityDynamicThreadTask
+{
+	FDeferredShadingSceneRenderer& ThisRenderer;
+	FRHICommandList& RHICmdList;
+	const FViewInfo& View;
+	FDeferredShadingSceneRenderer* SceneRenderer;
+
+public:
+
+	FRenderVelocityDynamicThreadTask(
+		FDeferredShadingSceneRenderer* InThisRenderer,
+		FRHICommandList* InRHICmdList,
+		const FViewInfo* InView
+		)
+		: ThisRenderer(*InThisRenderer)
+		, RHICmdList(*InRHICmdList)
+		, View(*InView)
+	{
+	}
+
+	FORCEINLINE TStatId GetStatId() const
+	{
+		RETURN_QUICK_DECLARE_CYCLE_STAT(FRenderVelocityDynamicThreadTask, STATGROUP_TaskGraphTasks);
+	}
+
+	ENamedThreads::Type GetDesiredThread()
+	{
+		return ENamedThreads::AnyThread;
+	}
+
+	static ESubsequentsMode::Type GetSubsequentsMode() { return ESubsequentsMode::TrackSubsequents; }
+
+	void DoTask(ENamedThreads::Type CurrentThread, const FGraphEventRef& MyCompletionGraphEvent)
+	{
+		ThisRenderer.RenderDynamicVelocities(RHICmdList, View);
+	}
+};
+
+
+void FDeferredShadingSceneRenderer::RenderVelocitiesInner(FRHICommandListImmediate& RHICmdList, const FViewInfo& View)
+{
+	static TAutoConsoleVariable<int32> CVarParallelVelocity(
+		TEXT("r.ParallelVelocity"),
+		1,
+		TEXT("Toggles parallel velocity rendering. Parallel rendering must be enabled for this to have an effect.\n"),
+		ECVF_RenderThreadSafe
+		);
+
+	if (!GRHICommandList.Bypass() && CVarParallelVelocity.GetValueOnRenderThread())
+	{
+		// parallel version
+
+		int32 Width = CVarRHICmdWidth.GetValueOnRenderThread(); // we use a few more than needed to cover non-equal jobs
+		bool OutDirty = false; // unused
+
+		FGraphEventRef SubmitChain;
+
+		SubmitChain = Scene->VelocityDrawList.DrawVisibleParallel(View, View.StaticMeshVelocityMap, View.StaticMeshBatchVisibility, Width, SubmitChain, OutDirty);
+		{
+			FRHICommandList* CmdList = new FRHICommandList;
+
+			FGraphEventRef AnyThreadCompletionEvent = TGraphTask<FRenderVelocityDynamicThreadTask>::CreateTask(nullptr, ENamedThreads::RenderThread)
+				.ConstructAndDispatchWhenReady(this, CmdList, &View);
+
+			SubmitChain = FSubmitCommandlistThreadTask::AddToChain(AnyThreadCompletionEvent, SubmitChain, CmdList);
+		}
+		if (SubmitChain.GetReference())
+		{
+			RHICmdList.Flush(); // this would happen later anyway, but we might as well do these while we wait for async tasks
+			FTaskGraphInterface::Get().WaitUntilTaskCompletes(SubmitChain, ENamedThreads::RenderThread_Local);
+		}
+	}
+	else
+	{
+		// single threaded version
+		// Draw velocities for movable static meshes.
+		Scene->VelocityDrawList.DrawVisible(RHICmdList, View, View.StaticMeshVelocityMap, View.StaticMeshBatchVisibility);
+		RenderDynamicVelocities(RHICmdList, View);
+	}
+}
+
 void FDeferredShadingSceneRenderer::RenderVelocities(FRHICommandListImmediate& RHICmdList, const FViewInfo& View, TRefCountPtr<IPooledRenderTarget>& VelocityRT, bool bLastFrame)
 {
 	check(FeatureLevel >= ERHIFeatureLevel::SM4);
@@ -488,28 +592,7 @@ void FDeferredShadingSceneRenderer::RenderVelocities(FRHICommandListImmediate& R
 			
 	RHICmdList.SetRasterizerState(GetStaticRasterizerState<true>(FM_Solid, CM_CW));
 
-	// Draw velocities for movable static meshes.
-	Scene->VelocityDrawList.DrawVisible(RHICmdList, View,View.StaticMeshVelocityMap, View.StaticMeshBatchVisibility);
-
-	// Draw velocities for movable dynamic meshes.
-	TDynamicPrimitiveDrawer<FVelocityDrawingPolicyFactory> Drawer(
-		RHICmdList, &View, FVelocityDrawingPolicyFactory::ContextType(DDM_AllOccluders), true, false, true
-		);
-	for(int32 PrimitiveIndex = 0;PrimitiveIndex < View.VisibleDynamicPrimitives.Num();PrimitiveIndex++)
-	{
-		const FPrimitiveSceneInfo* PrimitiveSceneInfo = View.VisibleDynamicPrimitives[PrimitiveIndex];
-
-		if(!PrimitiveSceneInfo->ShouldRenderVelocity(View))
-		{
-			continue;
-		}
-
-		FScopeCycleCounter Context(PrimitiveSceneInfo->Proxy->GetStatId());
-
-		Drawer.SetPrimitive(PrimitiveSceneInfo->Proxy);
-		PrimitiveSceneInfo->Proxy->DrawDynamicElements(&Drawer, &View);
-	}				
-	Drawer.IsDirty();
+	RenderVelocitiesInner(RHICmdList, View);
 
 	RHICmdList.CopyToResolveTarget(VelocityRT->GetRenderTargetItem().TargetableTexture, VelocityRT->GetRenderTargetItem().ShaderResourceTexture, false, FResolveParams());
 
