UE4 - parallel rendering, add immedaite commandlist flushes to eat latency and remove a stat

diff --git a/Engine/Source/Runtime/Renderer/Private/DeferredShadingRenderer.cpp b/Engine/Source/Runtime/Renderer/Private/DeferredShadingRenderer.cpp
--- a/Engine/Source/Runtime/Renderer/Private/DeferredShadingRenderer.cpp
+++ b/Engine/Source/Runtime/Renderer/Private/DeferredShadingRenderer.cpp
@@ -1122,7 +1122,7 @@ public:
 };
 
 
-FGraphEventRef FDeferredShadingSceneRenderer::RenderPrePassViewParallel(FRHICommandList& RHICmdList, const FViewInfo& View, int32 Width, FGraphEventRef SubmitChain, bool& OutDirty)
+FGraphEventRef FDeferredShadingSceneRenderer::RenderPrePassViewParallel(const FViewInfo& View, int32 Width, FGraphEventRef SubmitChain, bool& OutDirty)
 {
 	{
 		//optimization: this is probably only needed on the first view
@@ -1196,11 +1196,11 @@ bool FDeferredShadingSceneRenderer::RenderPrePass(FRHICommandListImmediate& RHIC
 			{
 				SCOPED_CONDITIONAL_DRAW_EVENTF(EventView, Views.Num() > 1, DEC_SCENE_ITEMS, TEXT("View%d"), ViewIndex);
 				const FViewInfo& View = Views[ViewIndex];
-				SubmitChain = RenderPrePassViewParallel(RHICmdList, View, Width, SubmitChain, bDirty);
+				SubmitChain = RenderPrePassViewParallel(View, Width, SubmitChain, bDirty);
 			}
-			
 			if (SubmitChain.GetReference())
 			{
+				RHICmdList.Flush(); // this would happen later anyway, but we might as well do these while we wait for async tasks
 				FTaskGraphInterface::Get().WaitUntilTaskCompletes(SubmitChain, ENamedThreads::RenderThread_Local);
 			}
 		}
@@ -1252,6 +1252,7 @@ bool FDeferredShadingSceneRenderer::RenderBasePass(FRHICommandListImmediate& RHI
 
 			if (SubmitChain.GetReference())
 			{
+				RHICmdList.Flush(); // this would happen later anyway, but we might as well do these while we wait for async tasks
 				FTaskGraphInterface::Get().WaitUntilTaskCompletes(SubmitChain, ENamedThreads::RenderThread_Local);
 			}
 		}
diff --git a/Engine/Source/Runtime/Renderer/Private/DeferredShadingRenderer.h b/Engine/Source/Runtime/Renderer/Private/DeferredShadingRenderer.h
--- a/Engine/Source/Runtime/Renderer/Private/DeferredShadingRenderer.h
+++ b/Engine/Source/Runtime/Renderer/Private/DeferredShadingRenderer.h
@@ -66,7 +66,7 @@ public:
 	 * Renders the scene's prepass for a particular view in parallel
 	 * @return the submit chain
 	 */
-	FGraphEventRef RenderPrePassViewParallel(FRHICommandList& RHICmdList, const FViewInfo& View, int32 Width, FGraphEventRef SubmitChain, bool& OutDirty);
+	FGraphEventRef RenderPrePassViewParallel(const FViewInfo& View, int32 Width, FGraphEventRef SubmitChain, bool& OutDirty);
 
     /** Renders the basepass for the static data of a given View. */
     bool RenderBasePassStaticData(FRHICommandList& RHICmdList, FViewInfo& View);
diff --git a/Engine/Source/Runtime/Renderer/Private/StaticMeshDrawList.inl b/Engine/Source/Runtime/Renderer/Private/StaticMeshDrawList.inl
--- a/Engine/Source/Runtime/Renderer/Private/StaticMeshDrawList.inl
+++ b/Engine/Source/Runtime/Renderer/Private/StaticMeshDrawList.inl
@@ -368,7 +368,6 @@ public:
 
 	void DoTask(ENamedThreads::Type CurrentThread, const FGraphEventRef& MyCompletionGraphEvent)
 	{
-		QUICK_SCOPE_CYCLE_COUNTER(AA_FDrawVisibleRenderThreadTask);
 		delete RHICmdList;
 	}
 
