UE4 - parallel rendering of translucency pass

diff --git a/Engine/Source/Runtime/Engine/Private/Materials/Material.cpp b/Engine/Source/Runtime/Engine/Private/Materials/Material.cpp
--- a/Engine/Source/Runtime/Engine/Private/Materials/Material.cpp
+++ b/Engine/Source/Runtime/Engine/Private/Materials/Material.cpp
@@ -3395,7 +3395,7 @@ static void ListSceneColorMaterials()
 	{
 		UMaterialInterface* Mat = *It;
 		const FMaterial* MatRes = Mat->GetRenderProxy(false)->GetMaterial(GRHIFeatureLevel);
-		if (MatRes && MatRes->RequiresSceneColorCopy())
+		if (MatRes && MatRes->RequiresSceneColorCopy_GameThread())
 		{
 			UMaterial* BaseMat = Mat->GetMaterial();
 			UE_LOG(LogConsoleResponse,Display,TEXT("[SepTrans=%d] %s"),
diff --git a/Engine/Source/Runtime/Engine/Private/Materials/MaterialInstance.cpp b/Engine/Source/Runtime/Engine/Private/Materials/MaterialInstance.cpp
--- a/Engine/Source/Runtime/Engine/Private/Materials/MaterialInstance.cpp
+++ b/Engine/Source/Runtime/Engine/Private/Materials/MaterialInstance.cpp
@@ -139,7 +139,7 @@ FMaterialInstanceResource::FMaterialInstanceResource(UMaterialInstance* InOwner,
 
 const FMaterial* FMaterialInstanceResource::GetMaterial(ERHIFeatureLevel::Type FeatureLevel) const
 {
-	checkSlow(IsInRenderingThread());
+	checkSlow(IsInParallelRenderingThread());
 
 	if (Owner->bHasStaticPermutationResource)
 	{
@@ -176,7 +176,7 @@ const FMaterial* FMaterialInstanceResource::GetMaterial(ERHIFeatureLevel::Type F
 
 FMaterial* FMaterialInstanceResource::GetMaterialNoFallback(ERHIFeatureLevel::Type FeatureLevel) const
 {
-	checkSlow(IsInRenderingThread());
+	checkSlow(IsInParallelRenderingThread());
 
 	if (Owner->bHasStaticPermutationResource)
 	{
@@ -205,7 +205,7 @@ bool FMaterialInstanceResource::GetScalarValue(
 	const FMaterialRenderContext& Context
 	) const
 {
-	checkSlow(IsInRenderingThread());
+	checkSlow(IsInParallelRenderingThread());
 	const float* Value = RenderThread_FindParameterByName<float>(ParameterName);
 	if(Value)
 	{
@@ -228,7 +228,7 @@ bool FMaterialInstanceResource::GetVectorValue(
 	const FMaterialRenderContext& Context
 	) const
 {
-	checkSlow(IsInRenderingThread());
+	checkSlow(IsInParallelRenderingThread());
 	const FLinearColor* Value = RenderThread_FindParameterByName<FLinearColor>(ParameterName);
 	if(Value)
 	{
@@ -251,7 +251,7 @@ bool FMaterialInstanceResource::GetTextureValue(
 	const FMaterialRenderContext& Context
 	) const
 {
-	checkSlow(IsInRenderingThread());
+	checkSlow(IsInParallelRenderingThread());
 	const UTexture* const * Value = RenderThread_FindParameterByName<const UTexture*>(ParameterName);
 	if(Value && *Value)
 	{
diff --git a/Engine/Source/Runtime/Engine/Private/Materials/MaterialShared.cpp b/Engine/Source/Runtime/Engine/Private/Materials/MaterialShared.cpp
--- a/Engine/Source/Runtime/Engine/Private/Materials/MaterialShared.cpp
+++ b/Engine/Source/Runtime/Engine/Private/Materials/MaterialShared.cpp
@@ -326,7 +326,7 @@ const TArray<TRefCountPtr<FMaterialUniformExpressionTexture> >& FMaterial::GetUn
 	}
 	else 
 	{
-		checkSlow(IsInRenderingThread());
+		check(IsInRenderingThread());
 		ShaderMapToUse = GetRenderingThreadShaderMap();
 	}
 
@@ -351,9 +351,9 @@ const TArray<TRefCountPtr<FMaterialUniformExpressionTexture> >& FMaterial::GetUn
 	}
 	else 
 	{
-		checkSlow(IsInRenderingThread());
+		check(IsInRenderingThread());
 		ShaderMapToUse = GetRenderingThreadShaderMap();
-}
+	}
 
 	if (ShaderMapToUse)
 	{
@@ -364,23 +364,29 @@ const TArray<TRefCountPtr<FMaterialUniformExpressionTexture> >& FMaterial::GetUn
 	return EmptyExpressions;
 }
 
-bool FMaterial::RequiresSceneColorCopy() const 
+bool FMaterial::RequiresSceneColorCopy_GameThread() const
 {
-	if (IsInGameThread() && GameThreadShaderMap)
+	check(IsInGameThread());
+	if (GameThreadShaderMap)
 	{
 		return GameThreadShaderMap->RequiresSceneColorCopy();
 	}
-	else if (IsInRenderingThread() && RenderingThreadShaderMap)
+	return false;
+}
+
+bool FMaterial::RequiresSceneColorCopy_RenderThread() const
+{
+	check(IsInParallelRenderingThread());
+	if (RenderingThreadShaderMap)
 	{
 		return RenderingThreadShaderMap->RequiresSceneColorCopy();
 	}
-
 	return false;
 }
 
 bool FMaterial::NeedsSceneTextures() const 
 {
-	checkSlow(IsInRenderingThread());
+	check(IsInParallelRenderingThread());
 
 	if (RenderingThreadShaderMap)
 	{
@@ -392,7 +398,7 @@ bool FMaterial::NeedsSceneTextures() const
 
 bool FMaterial::NeedsGBuffer() const
 {
-	checkSlow(IsInRenderingThread());
+	check(IsInParallelRenderingThread());
 
 	if (IsOpenGLPlatform(GRHIShaderPlatform)) // @todo: TTP #341211
 	{
@@ -410,7 +416,7 @@ bool FMaterial::NeedsGBuffer() const
 
 bool FMaterial::UsesEyeAdaptation() const 
 {
-	checkSlow(IsInRenderingThread());
+	check(IsInParallelRenderingThread());
 
 	if (RenderingThreadShaderMap)
 	{
@@ -420,9 +426,18 @@ bool FMaterial::UsesEyeAdaptation() const
 	return false;
 }
 
-bool FMaterial::MaterialModifiesMeshPosition() const 
+bool FMaterial::MaterialModifiesMeshPosition_RenderThread() const
 { 
-	FMaterialShaderMap* ShaderMap = IsInRenderingThread() ? RenderingThreadShaderMap : GameThreadShaderMap.GetReference();
+	check(IsInParallelRenderingThread());
+	bool bUsesWPO = RenderingThreadShaderMap ? RenderingThreadShaderMap->ModifiesMeshPosition() : false;
+
+	return bUsesWPO || GetTessellationMode() != MTM_NoTessellation;
+}
+
+bool FMaterial::MaterialModifiesMeshPosition_GameThread() const
+{
+	check(IsInGameThread());
+	FMaterialShaderMap* ShaderMap = GameThreadShaderMap.GetReference();
 	bool bUsesWPO = ShaderMap ? ShaderMap->ModifiesMeshPosition() : false;
 
 	return bUsesWPO || GetTessellationMode() != MTM_NoTessellation;
@@ -437,13 +452,13 @@ bool FMaterial::MaterialMayModifyMeshPosition() const
 
 FMaterialShaderMap* FMaterial::GetRenderingThreadShaderMap() const 
 { 
-	checkSlow(IsInParallelRenderingThread());
+	check(IsInParallelRenderingThread());
 	return RenderingThreadShaderMap; 
 }
 
 void FMaterial::SetRenderingThreadShaderMap(FMaterialShaderMap* InMaterialShaderMap)
 {
-	checkSlow(IsInRenderingThread());
+	check(IsInRenderingThread());
 	RenderingThreadShaderMap = InMaterialShaderMap;
 }
 
@@ -1114,7 +1129,7 @@ void FMaterial::SetupMaterialEnvironment(
 	OutEnvironment.SetDefine(TEXT("MATERIAL_TWOSIDED"), IsTwoSided() ? TEXT("1") : TEXT("0"));
 	OutEnvironment.SetDefine(TEXT("MATERIAL_TANGENTSPACENORMAL"), IsTangentSpaceNormal() ? TEXT("1") : TEXT("0"));
 	OutEnvironment.SetDefine(TEXT("GENERATE_SPHERICAL_PARTICLE_NORMALS"),ShouldGenerateSphericalParticleNormals() ? TEXT("1") : TEXT("0"));
-	OutEnvironment.SetDefine(TEXT("MATERIAL_USES_SCENE_COLOR_COPY"), RequiresSceneColorCopy() ? TEXT("1") : TEXT("0"));
+	OutEnvironment.SetDefine(TEXT("MATERIAL_USES_SCENE_COLOR_COPY"), RequiresSceneColorCopy_GameThread() ? TEXT("1") : TEXT("0"));
 	OutEnvironment.SetDefine(TEXT("MATERIAL_FULLY_ROUGH"), IsFullyRough() ? TEXT("1") : TEXT("0"));
 	OutEnvironment.SetDefine(TEXT("MATERIAL_NONMETAL"), IsNonmetal() ? TEXT("1") : TEXT("0"));
 	OutEnvironment.SetDefine(TEXT("MATERIAL_USE_LM_DIRECTIONALITY"), UseLmDirectionality() ? TEXT("1") : TEXT("0"));
@@ -1534,7 +1549,7 @@ void FMaterialRenderProxy::CacheUniformExpressions_GameThread()
 
 void FMaterialRenderProxy::InvalidateUniformExpressionCache()
 {
-	checkSlow(IsInRenderingThread());
+	check(IsInRenderingThread());
 	for (int32 i = 0; i < ERHIFeatureLevel::Num; ++i)
 	{
 		UniformExpressionCache[i].bUpToDate = false;
diff --git a/Engine/Source/Runtime/Engine/Private/PrimitiveDrawingUtils.cpp b/Engine/Source/Runtime/Engine/Private/PrimitiveDrawingUtils.cpp
--- a/Engine/Source/Runtime/Engine/Private/PrimitiveDrawingUtils.cpp
+++ b/Engine/Source/Runtime/Engine/Private/PrimitiveDrawingUtils.cpp
@@ -1303,7 +1303,7 @@ int32 DrawRichMesh(
 			BaseColor = LevelColor;
 		}
 
-		if (Mesh.MaterialRenderProxy->GetMaterial(FeatureLevel)->MaterialModifiesMeshPosition())
+		if (Mesh.MaterialRenderProxy->GetMaterial(FeatureLevel)->MaterialModifiesMeshPosition_RenderThread())
 		{
 			// If the material is mesh-modifying, we cannot rely on substitution
 			const FOverrideSelectionColorMaterialRenderProxy WireframeMaterialInstance(
@@ -1472,7 +1472,7 @@ int32 DrawRichMesh(
 		if(EngineShowFlags.MeshEdges)
 		{
 			// Draw the mesh's edges in blue, on top of the base geometry.
-			if (Mesh.MaterialRenderProxy->GetMaterial(FeatureLevel)->MaterialModifiesMeshPosition())
+			if (Mesh.MaterialRenderProxy->GetMaterial(FeatureLevel)->MaterialModifiesMeshPosition_RenderThread())
 			{
 				const FSceneView* View = PDI->View;
 
diff --git a/Engine/Source/Runtime/Engine/Private/StaticMeshRender.cpp b/Engine/Source/Runtime/Engine/Private/StaticMeshRender.cpp
--- a/Engine/Source/Runtime/Engine/Private/StaticMeshRender.cpp
+++ b/Engine/Source/Runtime/Engine/Private/StaticMeshRender.cpp
@@ -409,7 +409,7 @@ inline bool AllowShadowOnlyMesh()
 
 void FStaticMeshSceneProxy::DrawStaticElements(FStaticPrimitiveDrawInterface* PDI)
 {
-	checkSlow(IsInRenderingThread());
+	checkSlow(IsInParallelRenderingThread());
 	if (!HasViewDependentDPG())
 	{
 		// Determine the DPG the primitive should be drawn in.
@@ -476,7 +476,7 @@ void FStaticMeshSceneProxy::DrawStaticElements(FStaticPrimitiveDrawInterface* PD
 							Section.bCastShadow
 							&& Material->GetBlendMode() == BLEND_Opaque
 							&& !Material->IsTwoSided()
-							&& !Material->MaterialModifiesMeshPosition();
+							&& !Material->MaterialModifiesMeshPosition_RenderThread();
 						bAnySectionCastsShadow |= Section.bCastShadow;
 					}
 
@@ -556,7 +556,7 @@ void FStaticMeshSceneProxy::DrawDynamicElements(FPrimitiveDrawInterface* PDI,con
 {
 	QUICK_SCOPE_CYCLE_COUNTER( STAT_StaticMeshSceneProxy_DrawDynamicElements );
 
-	checkSlow(IsInRenderingThread());
+	checkSlow(IsInParallelRenderingThread());
 
 	const bool bIsLightmapSettingError = HasStaticLighting() && !HasValidSettingsForStaticLighting();
 
@@ -854,7 +854,7 @@ bool FStaticMeshSceneProxy::CanBeOccluded() const
 
 FPrimitiveViewRelevance FStaticMeshSceneProxy::GetViewRelevance(const FSceneView* View)
 {   
-	checkSlow(IsInRenderingThread());
+	checkSlow(IsInParallelRenderingThread());
 
 	FPrimitiveViewRelevance Result;
 	Result.bDrawRelevance = IsShown(View) && View->Family->EngineShowFlags.StaticMeshes;
@@ -1058,7 +1058,7 @@ FStaticMeshSceneProxy::FLODInfo::FLODInfo(const UStaticMeshComponent* InComponen
 		FMaterialResource const* MaterialResource = const_cast<UMaterialInterface const*>(SectionInfo.Material)->GetMaterial_Concurrent(RecursionGuard)->GetMaterialResource(FeatureLevel);
 		if(MaterialResource)
 		{
-			if(MaterialResource->MaterialModifiesMeshPosition())
+			if (MaterialResource->MaterialModifiesMeshPosition_GameThread())
 			{
 				bUsesMeshModifyingMaterials = true;
 			}
diff --git a/Engine/Source/Runtime/Engine/Public/MaterialShared.h b/Engine/Source/Runtime/Engine/Public/MaterialShared.h
--- a/Engine/Source/Runtime/Engine/Public/MaterialShared.h
+++ b/Engine/Source/Runtime/Engine/Public/MaterialShared.h
@@ -973,13 +973,15 @@ public:
 		//@todo - remove non-dynamic parameter particle VF and always support dynamic parameter
 		return true; 
 	}
-	ENGINE_API bool RequiresSceneColorCopy() const;
+	ENGINE_API bool RequiresSceneColorCopy_GameThread() const;
+	ENGINE_API bool RequiresSceneColorCopy_RenderThread() const;
 	ENGINE_API bool NeedsSceneTextures() const;
 	ENGINE_API bool NeedsGBuffer() const;
 	ENGINE_API bool UsesEyeAdaptation() const;	
 
 	/** Does the material modify the mesh position. */
-	ENGINE_API bool MaterialModifiesMeshPosition() const;
+	ENGINE_API bool MaterialModifiesMeshPosition_RenderThread() const;
+	ENGINE_API bool MaterialModifiesMeshPosition_GameThread() const;
 
 	/** Note: This function is only intended for use in deciding whether or not shader permutations are required before material translation occurs. */
 	ENGINE_API bool MaterialMayModifyMeshPosition() const;
diff --git a/Engine/Source/Runtime/Renderer/Private/DeferredShadingRenderer.h b/Engine/Source/Runtime/Renderer/Private/DeferredShadingRenderer.h
--- a/Engine/Source/Runtime/Renderer/Private/DeferredShadingRenderer.h
+++ b/Engine/Source/Runtime/Renderer/Private/DeferredShadingRenderer.h
@@ -232,8 +232,13 @@ private:
 	void RenderStationaryLightOverlap(FRHICommandListImmediate& RHICmdList);
 	
 	/** 
-	 * Renders the scene's translucency.
+	 * Renders the scene's translucency, parallel version
 	 */
+	void RenderTranslucencyParallel(FRHICommandListImmediate& RHICmdList);
+
+	/**
+	* Renders the scene's translucency.
+	*/
 	void RenderTranslucency(FRHICommandListImmediate& RHICmdList);
 
 	/** Renders the scene's light shafts */
@@ -242,7 +247,7 @@ private:
 	void RenderLightShaftBloom(FRHICommandListImmediate& RHICmdList);
 
 	/** Reuses an existing translucent shadow map if possible or re-renders one if necessary. */
-	const FProjectedShadowInfo* PrepareTranslucentShadowMap(FRHICommandListImmediate& RHICmdList, const FViewInfo& View, FPrimitiveSceneInfo* PrimitiveSceneInfo, bool bSeparateTranslucencyPass);
+	const FProjectedShadowInfo* PrepareTranslucentShadowMap(FRHICommandList& RHICmdList, const FViewInfo& View, FPrimitiveSceneInfo* PrimitiveSceneInfo, bool bSeparateTranslucencyPass);
 
 	/** Renders the velocities of movable objects for the motion blur effect. */
 	void RenderVelocities(FRHICommandListImmediate& RHICmdList, const FViewInfo& View, TRefCountPtr<IPooledRenderTarget>& VelocityRT, bool bLastFrame);
diff --git a/Engine/Source/Runtime/Renderer/Private/DepthRendering.cpp b/Engine/Source/Runtime/Renderer/Private/DepthRendering.cpp
--- a/Engine/Source/Runtime/Renderer/Private/DepthRendering.cpp
+++ b/Engine/Source/Runtime/Renderer/Private/DepthRendering.cpp
@@ -324,7 +324,7 @@ void FDepthDrawingPolicyFactory::AddStaticMesh(FScene* Scene,FStaticMesh* Static
 	else
 	{
 		if (StaticMesh->VertexFactory->SupportsPositionOnlyStream() 
-			&& !Material->MaterialModifiesMeshPosition())
+			&& !Material->MaterialModifiesMeshPosition_RenderThread())
 		{
 			const FMaterialRenderProxy* DefaultProxy = UMaterial::GetDefaultMaterial(MD_Surface)->GetRenderProxy(false);
 			// Add the static mesh to the position-only depth draw list.
@@ -343,7 +343,7 @@ void FDepthDrawingPolicyFactory::AddStaticMesh(FScene* Scene,FStaticMesh* Static
 		}
 		else
 		{
-			if (!Material->MaterialModifiesMeshPosition())
+			if (!Material->MaterialModifiesMeshPosition_RenderThread())
 			{
 				// Override with the default material for everything but opaque two sided materials
 				MaterialRenderProxy = UMaterial::GetDefaultMaterial(MD_Surface)->GetRenderProxy(false);
@@ -393,7 +393,7 @@ bool FDepthDrawingPolicyFactory::DrawMesh(
 
 		if ( BlendMode == BLEND_Opaque 
 			&& Mesh.VertexFactory->SupportsPositionOnlyStream() 
-			&& !Material->MaterialModifiesMeshPosition())
+			&& !Material->MaterialModifiesMeshPosition_RenderThread())
 		{
 			//render opaque primitives that support a separate position-only vertex buffer
 			const FMaterialRenderProxy* DefaultProxy = UMaterial::GetDefaultMaterial(MD_Surface)->GetRenderProxy(false);
@@ -435,7 +435,7 @@ bool FDepthDrawingPolicyFactory::DrawMesh(
 
 			if(bDraw)
 			{
-				if (!bMaterialMasked && !Material->MaterialModifiesMeshPosition())
+				if (!bMaterialMasked && !Material->MaterialModifiesMeshPosition_RenderThread())
 				{
 					// Override with the default material for opaque materials that are not two sided
 					MaterialRenderProxy = UMaterial::GetDefaultMaterial(MD_Surface)->GetRenderProxy(false);
diff --git a/Engine/Source/Runtime/Renderer/Private/DynamicPrimitiveDrawing.h b/Engine/Source/Runtime/Renderer/Private/DynamicPrimitiveDrawing.h
--- a/Engine/Source/Runtime/Renderer/Private/DynamicPrimitiveDrawing.h
+++ b/Engine/Source/Runtime/Renderer/Private/DynamicPrimitiveDrawing.h
@@ -150,6 +150,25 @@ private:
 
 };
 
+/**
+* Draws a range of view's elements with the specified drawing policy factory type.
+* @param View - The view to draw the meshes for.
+* @param DrawingContext - The drawing policy type specific context for the drawing.
+* @param DPGIndex World or Foreground DPG index for draw order
+* @param FirstIndex - Element range
+* @param LastIndex - Element range
+*/
+template<class DrawingPolicyFactoryType>
+void DrawViewElementsInner(
+	FRHICommandList& RHICmdList,
+	const FViewInfo& View,
+	const typename DrawingPolicyFactoryType::ContextType& DrawingContext,
+	uint8 DPGIndex,
+	bool bPreFog,
+	int32 FirstIndex,
+	int32 LastIndex
+	);
+
 /**
  * Draws a view's elements with the specified drawing policy factory type.
  * @param View - The view to draw the meshes for.
@@ -166,6 +185,25 @@ bool DrawViewElements(
 	bool bPreFog
 	);
 
+/**
+* Draws a view's elements with the specified drawing policy factory type.
+* @param View - The view to draw the meshes for.
+* @param DrawingContext - The drawing policy type specific context for the drawing.
+* @param DPGIndex World or Foreground DPG index for draw order
+* @param bPreFog - true if the draw call is occurring before fog has been rendered.
+* @param SubmitChain - parallel submit chain
+* @param Width - parallel width
+*/
+template<class DrawingPolicyFactoryType>
+FGraphEventRef DrawViewElementsParallel(
+	const FViewInfo& View,
+	const typename DrawingPolicyFactoryType::ContextType& DrawingContext,
+	uint8 DPGIndex,
+	bool bPreFog,
+	FGraphEventRef SubmitChain,
+	int32 Width
+	);
+
 /**
  * Draws a given set of dynamic primitives to a RHI command interface, using the specified drawing policy type.
  * @param View - The view to draw the meshes for.
diff --git a/Engine/Source/Runtime/Renderer/Private/DynamicPrimitiveDrawing.inl b/Engine/Source/Runtime/Renderer/Private/DynamicPrimitiveDrawing.inl
--- a/Engine/Source/Runtime/Renderer/Private/DynamicPrimitiveDrawing.inl
+++ b/Engine/Source/Runtime/Renderer/Private/DynamicPrimitiveDrawing.inl
@@ -213,18 +213,21 @@ void TDynamicPrimitiveDrawer<DrawingPolicyFactoryType>::DrawPoint(
 }
 
 template<class DrawingPolicyFactoryType>
-bool DrawViewElements(
+void DrawViewElementsInner(
 	FRHICommandList& RHICmdList,
 	const FViewInfo& View,
 	const typename DrawingPolicyFactoryType::ContextType& DrawingContext,
 	uint8 DPGIndex,
-	bool bPreFog
+	bool bPreFog,
+	int32 FirstIndex,
+	int32 LastIndex
 	)
 {
 	// Get the correct element list based on dpg index
 	const TIndirectArray<FHitProxyMeshPair>& ViewMeshElementList = ( DPGIndex == SDPG_Foreground ? View.TopViewMeshElements : View.ViewMeshElements );
 	// Draw the view's mesh elements.
-	for(int32 MeshIndex = 0;MeshIndex < ViewMeshElementList.Num();MeshIndex++)
+	check(LastIndex < ViewMeshElementList.Num());
+	for (int32 MeshIndex = FirstIndex; MeshIndex <= LastIndex; MeshIndex++)
 	{
 		const FHitProxyMeshPair& Mesh = ViewMeshElementList[MeshIndex];
 		const auto FeatureLevel = View.GetFeatureLevel();
@@ -247,8 +250,121 @@ bool DrawViewElements(
 			--bBackFace;
 		} while( bBackFace >= 0 );
 	}
+}
+
+template<typename DrawingPolicyFactoryType>
+class FDrawViewElementsAnyThreadTask
+{
+	FRHICommandList& RHICmdList;
+	const FViewInfo& View;
+	const typename DrawingPolicyFactoryType::ContextType& DrawingContext;
+	uint8 DPGIndex;
+	bool bPreFog;
+
+
+	const int32 FirstIndex;
+	const int32 LastIndex;
+
+public:
+
+	FDrawViewElementsAnyThreadTask(
+		FRHICommandList* InRHICmdList,
+		const FViewInfo* InView,
+		const typename DrawingPolicyFactoryType::ContextType& InDrawingContext,
+		uint8 InDPGIndex,
+		bool InbPreFog,
+		int32 InFirstIndex,
+		int32 InLastIndex
+		)
+		: RHICmdList(*InRHICmdList)
+		, View(*InView)
+		, DrawingContext(InDrawingContext)
+		, DPGIndex(InDPGIndex)
+		, bPreFog(InbPreFog)
+		, FirstIndex(InFirstIndex)
+		, LastIndex(InLastIndex)
+	{
+	}
+
+	FORCEINLINE TStatId GetStatId() const
+	{
+		RETURN_QUICK_DECLARE_CYCLE_STAT(FDrawViewElementsAnyThreadTask, STATGROUP_TaskGraphTasks);
+	}
+
+	ENamedThreads::Type GetDesiredThread()
+	{
+		return ENamedThreads::AnyThread;
+	}
+
+	static ESubsequentsMode::Type GetSubsequentsMode() { return ESubsequentsMode::TrackSubsequents; }
+
+	void DoTask(ENamedThreads::Type CurrentThread, const FGraphEventRef& MyCompletionGraphEvent)
+	{
+		DrawViewElementsInner<DrawingPolicyFactoryType>(RHICmdList, View, DrawingContext, DPGIndex, bPreFog, FirstIndex, LastIndex);
+	}
+};
+
+template<class DrawingPolicyFactoryType>
+FGraphEventRef DrawViewElementsParallel(
+	const FViewInfo& View,
+	const typename DrawingPolicyFactoryType::ContextType& DrawingContext,
+	uint8 DPGIndex,
+	bool bPreFog,
+	FGraphEventRef SubmitChain,
+	int32 Width
+	)
+{
+	// Get the correct element list based on dpg index
+	const TIndirectArray<FHitProxyMeshPair>& ViewMeshElementList = (DPGIndex == SDPG_Foreground ? View.TopViewMeshElements : View.ViewMeshElements);
+
+	{
+		int32 NumPrims = ViewMeshElementList.Num();
+		int32 EffectiveThreads = FMath::Min<int32>(NumPrims, Width);
+
+		int32 Start = 0;
+		if (EffectiveThreads)
+		{
+
+			int32 NumPer = NumPrims / EffectiveThreads;
+			int32 Extra = NumPrims - NumPer * EffectiveThreads;
+
+
+			for (int32 ThreadIndex = 0; ThreadIndex < EffectiveThreads; ThreadIndex++)
+			{
+				int32 Last = Start + (NumPer - 1) + (ThreadIndex < Extra);
+				check(Last >= Start);
+
+				FRHICommandList* CmdList = new FRHICommandList;
+
+				FGraphEventRef AnyThreadCompletionEvent = TGraphTask<FDrawViewElementsAnyThreadTask<DrawingPolicyFactoryType> >::CreateTask(nullptr, ENamedThreads::RenderThread)
+					.ConstructAndDispatchWhenReady(CmdList, &View, DrawingContext, DPGIndex, bPreFog, Start, Last);
+				Start = Last + 1;
+
+				SubmitChain = FSubmitCommandlistThreadTask::AddToChain(AnyThreadCompletionEvent, SubmitChain, CmdList);
+			}
+		}
+	}
 
-	return View.ViewMeshElements.Num() != 0;
+	return SubmitChain;
+}
+
+template<class DrawingPolicyFactoryType>
+bool DrawViewElements(
+	FRHICommandList& RHICmdList,
+	const FViewInfo& View,
+	const typename DrawingPolicyFactoryType::ContextType& DrawingContext,
+	uint8 DPGIndex,
+	bool bPreFog
+	)
+{
+	// Get the correct element list based on dpg index
+	const TIndirectArray<FHitProxyMeshPair>& ViewMeshElementList = (DPGIndex == SDPG_Foreground ? View.TopViewMeshElements : View.ViewMeshElements);
+	if (View.ViewMeshElements.Num() != 0)
+	{
+		DrawViewElementsInner<DrawingPolicyFactoryType>(RHICmdList, View, DrawingContext, DPGIndex, bPreFog, 0, ViewMeshElementList.Num() - 1);
+		return true;
+	}
+	return false;
 }
 
 template<class DrawingPolicyFactoryType>
diff --git a/Engine/Source/Runtime/Renderer/Private/LightMapDensityRendering.cpp b/Engine/Source/Runtime/Renderer/Private/LightMapDensityRendering.cpp
--- a/Engine/Source/Runtime/Renderer/Private/LightMapDensityRendering.cpp
+++ b/Engine/Source/Runtime/Renderer/Private/LightMapDensityRendering.cpp
@@ -108,7 +108,7 @@ bool FLightMapDensityDrawingPolicyFactory::DrawDynamicMesh(
 	const EBlendMode BlendMode = Material->GetBlendMode();
 
 	const bool bMaterialMasked = Material->IsMasked();
-	const bool bMaterialModifiesMesh = Material->MaterialModifiesMeshPosition();
+	const bool bMaterialModifiesMesh = Material->MaterialModifiesMeshPosition_RenderThread();
 	if (!bMaterialMasked && !bMaterialModifiesMesh)
 	{
 		// Override with the default material for opaque materials that are not two sided
diff --git a/Engine/Source/Runtime/Renderer/Private/LightMapDensityRendering.h b/Engine/Source/Runtime/Renderer/Private/LightMapDensityRendering.h
--- a/Engine/Source/Runtime/Renderer/Private/LightMapDensityRendering.h
+++ b/Engine/Source/Runtime/Renderer/Private/LightMapDensityRendering.h
@@ -454,7 +454,7 @@ public:
 		return (Material && 
 				!(	Material->IsSpecialEngineMaterial() || 
 					Material->IsMasked() ||
-					Material->MaterialModifiesMeshPosition()
+					Material->MaterialModifiesMeshPosition_RenderThread()
 				));
 	}
 };
diff --git a/Engine/Source/Runtime/Renderer/Private/PostProcess/SceneRenderTargets.cpp b/Engine/Source/Runtime/Renderer/Private/PostProcess/SceneRenderTargets.cpp
--- a/Engine/Source/Runtime/Renderer/Private/PostProcess/SceneRenderTargets.cpp
+++ b/Engine/Source/Runtime/Renderer/Private/PostProcess/SceneRenderTargets.cpp
@@ -354,6 +354,8 @@ void FSceneRenderTargets::AllocLightAttenuation()
 		return;
 	}
 
+	check(IsInRenderingThread());
+
 	// create LightAttenuation on demand so it can be shared with other pooled RT
 
 	// Create a texture to store the resolved light attenuation values, and a render-targetable surface to hold the unresolved light attenuation values.
diff --git a/Engine/Source/Runtime/Renderer/Private/PostProcess/SceneRenderTargets.h b/Engine/Source/Runtime/Renderer/Private/PostProcess/SceneRenderTargets.h
--- a/Engine/Source/Runtime/Renderer/Private/PostProcess/SceneRenderTargets.h
+++ b/Engine/Source/Runtime/Renderer/Private/PostProcess/SceneRenderTargets.h
@@ -339,10 +339,10 @@ public:
 	// ---
 
 	/** Get the current translucent ambient lighting volume texture. Can vary depending on whether volume filtering is enabled */
-	TRefCountPtr<IPooledRenderTarget> GetTranslucencyVolumeAmbient(ETranslucencyVolumeCascade Cascade) { return TranslucencyLightingVolumeAmbient[SelectTranslucencyVolumeTarget(Cascade)]; }
+	IPooledRenderTarget* GetTranslucencyVolumeAmbient(ETranslucencyVolumeCascade Cascade) { return TranslucencyLightingVolumeAmbient[SelectTranslucencyVolumeTarget(Cascade)].GetReference(); }
 
 	/** Get the current translucent directional lighting volume texture. Can vary depending on whether volume filtering is enabled */
-	TRefCountPtr<IPooledRenderTarget> GetTranslucencyVolumeDirectional(ETranslucencyVolumeCascade Cascade) { return TranslucencyLightingVolumeDirectional[SelectTranslucencyVolumeTarget(Cascade)]; }
+	IPooledRenderTarget* GetTranslucencyVolumeDirectional(ETranslucencyVolumeCascade Cascade) { return TranslucencyLightingVolumeDirectional[SelectTranslucencyVolumeTarget(Cascade)].GetReference(); }
 
 	// ---
 	/** Get the uniform buffer containing GBuffer resources. */
@@ -396,6 +396,8 @@ public:
 	//
 	void AllocGBufferTargets();
 
+	void AllocLightAttenuation();
+
 private: // Get...() methods instead of direct access
 
 	// 0 before BeginRenderingSceneColor and after tone mapping, for ES2 it's not released during the frame
@@ -512,8 +514,6 @@ private:
 
 	void AllocSceneColor();
 
-	void AllocLightAttenuation();
-
 	// internal method, used by AdjustGBufferRefCount()
 	void ReleaseGBufferTargets();
 
diff --git a/Engine/Source/Runtime/Renderer/Private/SceneHitProxyRendering.cpp b/Engine/Source/Runtime/Renderer/Private/SceneHitProxyRendering.cpp
--- a/Engine/Source/Runtime/Renderer/Private/SceneHitProxyRendering.cpp
+++ b/Engine/Source/Runtime/Renderer/Private/SceneHitProxyRendering.cpp
@@ -250,7 +250,7 @@ void FHitProxyDrawingPolicyFactory::AddStaticMesh(FScene* Scene,FStaticMesh* Sta
 	// Add the static mesh to the DPG's hit proxy draw list.
 	const FMaterialRenderProxy* MaterialRenderProxy = StaticMesh->MaterialRenderProxy;
 	const FMaterial* Material = MaterialRenderProxy->GetMaterial(Scene->GetFeatureLevel());
-	if ( !Material->IsMasked() && !Material->IsTwoSided() && !Material->MaterialModifiesMeshPosition() )
+	if (!Material->IsMasked() && !Material->IsTwoSided() && !Material->MaterialModifiesMeshPosition_RenderThread())
 	{
 		// Default material doesn't handle masked, and doesn't have the correct bIsTwoSided setting.
 		MaterialRenderProxy = UMaterial::GetDefaultMaterial(MD_Surface)->GetRenderProxy(false);
@@ -301,7 +301,7 @@ bool FHitProxyDrawingPolicyFactory::DrawDynamicMesh(
 		if( View.bAllowTranslucentPrimitivesInHitProxy || !IsTranslucentBlendMode( Material->GetBlendMode() ) || ( HitProxy && HitProxy->AlwaysAllowsTranslucentPrimitives() ) )
 #endif
 		{
-			if ( !Material->IsMasked() && !Material->IsTwoSided() && !Material->MaterialModifiesMeshPosition() )
+			if (!Material->IsMasked() && !Material->IsTwoSided() && !Material->MaterialModifiesMeshPosition_RenderThread())
 			{
 				// Default material doesn't handle masked, and doesn't have the correct bIsTwoSided setting.
 				MaterialRenderProxy = UMaterial::GetDefaultMaterial(MD_Surface)->GetRenderProxy(false);
diff --git a/Engine/Source/Runtime/Renderer/Private/SceneRendering.h b/Engine/Source/Runtime/Renderer/Private/SceneRendering.h
--- a/Engine/Source/Runtime/Renderer/Private/SceneRendering.h
+++ b/Engine/Source/Runtime/Renderer/Private/SceneRendering.h
@@ -86,6 +86,25 @@ public:
 	*/
 	void DrawPrimitives(FRHICommandListImmediate& RHICmdList, const class FViewInfo& View, class FDeferredShadingSceneRenderer& Renderer, bool bSeparateTranslucencyPass) const;
 
+	/**
+	* Iterate over the sorted list of prims and draw them
+	* @param View - current view used to draw items
+	* @param PhaseSortedPrimitives - array with the primitives we want to draw
+	* @param bSeparateTranslucencyPass
+	* @param FirstIndex, range of elements to render
+	* @param LastIndex, range of elements to render
+	*/
+	void DrawPrimitivesParallel(FRHICommandList& RHICmdList, const class FViewInfo& View, class FDeferredShadingSceneRenderer& Renderer, bool bSeparateTranslucencyPass, int32 FirstIndex, int32 LastIndex) const;
+
+	/**
+	* Draw a single primitive...this is used when we are rendering in parallel and we need to handlke a translucent shadow
+	* @param View - current view used to draw items
+	* @param PhaseSortedPrimitives - array with the primitives we want to draw
+	* @param bSeparateTranslucencyPass
+	* @param Index, element to render
+	*/
+	void DrawAPrimitive(FRHICommandList& RHICmdList, const class FViewInfo& View, class FDeferredShadingSceneRenderer& Renderer, bool bSeparateTranslucencyPass, int32 Index) const;
+
 	/** Draw all the primitives in this set for the forward shading pipeline. */
 	void DrawPrimitivesForForwardShading(FRHICommandListImmediate& RHICmdList, const class FViewInfo& View, class FSceneRenderer& Renderer) const;
 
@@ -184,7 +203,7 @@ private:
 	TArray<FSortedPrim,SceneRenderingAllocator> SortedSeparateTranslucencyPrims;
 
 	/** Renders a single primitive for the deferred shading pipeline. */
-	void RenderPrimitive(FRHICommandListImmediate& RHICmdList, const FViewInfo& View, FPrimitiveSceneInfo* PrimitiveSceneInfo, const FPrimitiveViewRelevance& ViewRelevance, const FProjectedShadowInfo* TranslucentSelfShadow, bool bSeparateTranslucencyPass) const;
+	void RenderPrimitive(FRHICommandList& RHICmdList, const FViewInfo& View, FPrimitiveSceneInfo* PrimitiveSceneInfo, const FPrimitiveViewRelevance& ViewRelevance, const FProjectedShadowInfo* TranslucentSelfShadow, bool bSeparateTranslucencyPass) const;
 
 	/** Renders a single primitive for the forward shading pipeline. */
 	void RenderPrimitiveForForwardShading(FRHICommandListImmediate& RHICmdList, const FViewInfo& View, FPrimitiveSceneInfo* PrimitiveSceneInfo, const FPrimitiveViewRelevance& ViewRelevance) const;
diff --git a/Engine/Source/Runtime/Renderer/Private/ShadowRendering.cpp b/Engine/Source/Runtime/Renderer/Private/ShadowRendering.cpp
--- a/Engine/Source/Runtime/Renderer/Private/ShadowRendering.cpp
+++ b/Engine/Source/Runtime/Renderer/Private/ShadowRendering.cpp
@@ -668,7 +668,7 @@ void OverrideWithDefaultMaterialForShadowDepth(
 {
 	// Override with the default material when possible.
 	if (!InOutMaterialResource->IsMasked() &&						// Don't override masked materials.
-		!InOutMaterialResource->MaterialModifiesMeshPosition() &&	// Don't override materials using world position offset.
+		!InOutMaterialResource->MaterialModifiesMeshPosition_RenderThread() &&	// Don't override materials using world position offset.
 		!bReflectiveShadowmap)										// Don't override when rendering reflective shadow maps.
 	{
 		const FMaterialRenderProxy* DefaultProxy = UMaterial::GetDefaultMaterial(MD_Surface)->GetRenderProxy(false);
@@ -749,7 +749,7 @@ FShadowDepthDrawingPolicy<bRenderingReflectiveShadowMaps>::FShadowDepthDrawingPo
 	bUsePositionOnlyVS = !bRenderingReflectiveShadowMaps
 		&& VertexFactory->SupportsPositionOnlyStream()
 		&& !MaterialResource->IsMasked()
-		&& !MaterialResource->MaterialModifiesMeshPosition();
+		&& !MaterialResource->MaterialModifiesMeshPosition_RenderThread();
 
 	// Vertex related shaders
 	if (bOnePassPointLightShadow)
@@ -1392,7 +1392,7 @@ public:
 	}
 };
 
-void FProjectedShadowInfo::RenderDepthInner(FRHICommandListImmediate& RHICmdList, FDeferredShadingSceneRenderer* SceneRenderer, const FViewInfo* FoundView)
+void FProjectedShadowInfo::RenderDepthInner(FRHICommandList& RHICmdList, FDeferredShadingSceneRenderer* SceneRenderer, const FViewInfo* FoundView)
 {
 	static TAutoConsoleVariable<int32> CVarParallelShadows(
 		TEXT("r.ParallelShadows"),
@@ -1401,7 +1401,8 @@ void FProjectedShadowInfo::RenderDepthInner(FRHICommandListImmediate& RHICmdList
 		ECVF_RenderThreadSafe
 		);
 
-	if (!GRHICommandList.Bypass() && CVarParallelShadows.GetValueOnRenderThread())
+	if (&RHICmdList == &FRHICommandListExecutor::GetImmediateCommandList() &&  // translucent shadows are draw on the render thread, using a recursive cmdlist, so this path seems like a loser in that case, but it should be tried
+		!GRHICommandList.Bypass() && CVarParallelShadows.GetValueOnRenderThread())
 	{
 		// parallel version
 
@@ -1446,7 +1447,7 @@ void FProjectedShadowInfo::RenderDepthInner(FRHICommandListImmediate& RHICmdList
 
 		if (SubmitChain.GetReference())
 		{
-			RHICmdList.Flush(); // this would happen later anyway, but we might as well do these while we wait for async tasks
+			FRHICommandListExecutor::GetImmediateCommandList().Flush(); // this would happen later anyway, but we might as well do these while we wait for async tasks
 			FTaskGraphInterface::Get().WaitUntilTaskCompletes(SubmitChain, ENamedThreads::RenderThread_Local);
 		}
 	}
@@ -1489,7 +1490,7 @@ void FProjectedShadowInfo::RenderDepthInner(FRHICommandListImmediate& RHICmdList
 	}
 }
 
-void FProjectedShadowInfo::RenderDepth(FRHICommandListImmediate& RHICmdList, FDeferredShadingSceneRenderer* SceneRenderer)
+void FProjectedShadowInfo::RenderDepth(FRHICommandList& RHICmdList, FDeferredShadingSceneRenderer* SceneRenderer)
 {
 #if WANTS_DRAW_MESH_EVENTS
 	FString EventName;
diff --git a/Engine/Source/Runtime/Renderer/Private/ShadowRendering.h b/Engine/Source/Runtime/Renderer/Private/ShadowRendering.h
--- a/Engine/Source/Runtime/Renderer/Private/ShadowRendering.h
+++ b/Engine/Source/Runtime/Renderer/Private/ShadowRendering.h
@@ -662,12 +662,12 @@ public:
 	/**
 	 * Renders the shadow subject depth.
 	 */
-	void RenderDepth(FRHICommandListImmediate& RHICmdList, class FDeferredShadingSceneRenderer* SceneRenderer);
+	void RenderDepth(FRHICommandList& RHICmdList, class FDeferredShadingSceneRenderer* SceneRenderer);
 
 	void ClearDepth(FRHICommandList& RHICmdList, class FDeferredShadingSceneRenderer* SceneRenderer);
 
 	/** Renders shadow maps for translucent primitives. */
-	void RenderTranslucencyDepths(FRHICommandListImmediate& RHICmdList, class FDeferredShadingSceneRenderer* SceneRenderer);
+	void RenderTranslucencyDepths(FRHICommandList& RHICmdList, class FDeferredShadingSceneRenderer* SceneRenderer);
 
 	/**
 	 * Projects the shadow onto the scene for a particular view.
@@ -744,7 +744,7 @@ private:
 	/**
 	* Renders the shadow subject depth, to a particular hacked view
 	*/
-	void RenderDepthInner(FRHICommandListImmediate& RHICmdList, class FDeferredShadingSceneRenderer* SceneRenderer, const FViewInfo* FoundView);
+	void RenderDepthInner(FRHICommandList& RHICmdList, class FDeferredShadingSceneRenderer* SceneRenderer, const FViewInfo* FoundView);
 
 	/**
 	* Renders the dynamic shadow subject depth, to a particular hacked view
diff --git a/Engine/Source/Runtime/Renderer/Private/TranslucentLighting.cpp b/Engine/Source/Runtime/Renderer/Private/TranslucentLighting.cpp
--- a/Engine/Source/Runtime/Renderer/Private/TranslucentLighting.cpp
+++ b/Engine/Source/Runtime/Renderer/Private/TranslucentLighting.cpp
@@ -484,7 +484,7 @@ public:
 };
 
 /** Renders shadow maps for translucent primitives. */
-void FProjectedShadowInfo::RenderTranslucencyDepths(FRHICommandListImmediate& RHICmdList, FDeferredShadingSceneRenderer* SceneRenderer)
+void FProjectedShadowInfo::RenderTranslucencyDepths(FRHICommandList& RHICmdList, FDeferredShadingSceneRenderer* SceneRenderer)
 {
 	
 	checkSlow(!bWholeSceneShadow);
@@ -503,7 +503,7 @@ void FProjectedShadowInfo::RenderTranslucencyDepths(FRHICommandListImmediate& RH
 			break;
 		}
 	}
-	check(FoundView);
+	check(FoundView && IsInRenderingThread());
 	
 	// Backup properties of the view that we will override
 	TUniformBufferRef<FViewUniformShaderParameters> OriginalUniformBuffer = FoundView->UniformBuffer;
diff --git a/Engine/Source/Runtime/Renderer/Private/TranslucentRendering.cpp b/Engine/Source/Runtime/Renderer/Private/TranslucentRendering.cpp
--- a/Engine/Source/Runtime/Renderer/Private/TranslucentRendering.cpp
+++ b/Engine/Source/Runtime/Renderer/Private/TranslucentRendering.cpp
@@ -10,7 +10,7 @@
 #include "SceneFilterRendering.h"
 #include "LightPropagationVolume.h"
 
-void SetTranslucentRenderTargetAndState(FRHICommandList& RHICmdList, const FViewInfo& View, bool bSeparateTranslucencyPass)
+static void SetTranslucentRenderTargetAndState(FRHICommandList& RHICmdList, const FViewInfo& View, bool bSeparateTranslucencyPass)
 {
 	bool bSetupTranslucentState = true;
 
@@ -31,7 +31,7 @@ void SetTranslucentRenderTargetAndState(FRHICommandList& RHICmdList, const FView
 	}
 }
 
-const FProjectedShadowInfo* FDeferredShadingSceneRenderer::PrepareTranslucentShadowMap(FRHICommandListImmediate& RHICmdList, const FViewInfo& View, FPrimitiveSceneInfo* PrimitiveSceneInfo, bool bSeparateTranslucencyPass)
+const FProjectedShadowInfo* FDeferredShadingSceneRenderer::PrepareTranslucentShadowMap(FRHICommandList& RHICmdList, const FViewInfo& View, FPrimitiveSceneInfo* PrimitiveSceneInfo, bool bSeparateTranslucencyPass)
 {
 	const FVisibleLightInfo* VisibleLightInfo = NULL;
 	FProjectedShadowInfo* TranslucentSelfShadow = NULL;
@@ -353,7 +353,7 @@ bool FTranslucencyDrawingPolicyFactory::DrawMesh(
 	// Only render translucent materials.
 	if(IsTranslucentBlendMode(BlendMode))
 	{
-		if (Material->RequiresSceneColorCopy())
+		if (Material->RequiresSceneColorCopy_RenderThread())
 		{
 			if (DrawingContext.bSceneColorCopyIsUpToDate == false)
 			{
@@ -490,12 +490,91 @@ bool FTranslucencyDrawingPolicyFactory::DrawStaticMesh(
 FTranslucentPrimSet
 -----------------------------------------------------------------------------*/
 
-/**
-* Iterate over the sorted list of prims and draw them
-* @param ViewInfo - current view used to draw items
-* @param bSeparateTranslucencyPass
-* @return true if anything was drawn
-*/
+void FTranslucentPrimSet::DrawAPrimitive(
+	FRHICommandList& RHICmdList,
+	const FViewInfo& View,
+	FDeferredShadingSceneRenderer& Renderer,
+	bool bSeparateTranslucencyPass,
+	int32 PrimIdx
+	) const
+{
+	const TArray<FSortedPrim, SceneRenderingAllocator>* PhaseSortedPrimitivesPtr = NULL;
+
+	// copy to reference for easier access
+	const TArray<FSortedPrim, SceneRenderingAllocator>& PhaseSortedPrimitives =
+		bSeparateTranslucencyPass ? SortedSeparateTranslucencyPrims : SortedPrims;
+	check(PrimIdx < PhaseSortedPrimitives.Num());
+
+	FPrimitiveSceneInfo* PrimitiveSceneInfo = PhaseSortedPrimitives[PrimIdx].PrimitiveSceneInfo;
+	int32 PrimitiveId = PrimitiveSceneInfo->GetIndex();
+	const FPrimitiveViewRelevance& ViewRelevance = View.PrimitiveViewRelevanceMap[PrimitiveId];
+
+	checkSlow(ViewRelevance.HasTranslucency());
+
+	const FProjectedShadowInfo* TranslucentSelfShadow = Renderer.PrepareTranslucentShadowMap(RHICmdList, View, PrimitiveSceneInfo, bSeparateTranslucencyPass);
+
+	RenderPrimitive(RHICmdList, View, PrimitiveSceneInfo, ViewRelevance, TranslucentSelfShadow, bSeparateTranslucencyPass);
+}
+
+struct FRHICommandVolumetricTranslucentShadow : public FRHICommand<FRHICommandVolumetricTranslucentShadow>
+{
+	const FTranslucentPrimSet &PrimSet;
+	const FViewInfo& View;
+	FDeferredShadingSceneRenderer& Renderer;
+	bool bSeparateTranslucencyPass;
+	int32 Index;
+
+	FORCEINLINE_DEBUGGABLE FRHICommandVolumetricTranslucentShadow(const FTranslucentPrimSet &InPrimSet, const FViewInfo& InView, FDeferredShadingSceneRenderer& InRenderer, bool InbSeparateTranslucencyPass, int32 InIndex)
+		: PrimSet(InPrimSet)
+		, View(InView)
+		, Renderer(InRenderer)
+		, bSeparateTranslucencyPass(InbSeparateTranslucencyPass)
+		, Index(InIndex)
+	{
+	}
+	void Execute()
+	{
+		check(IsInRenderingThread());
+		FRHICommandList RHICmdList;
+		PrimSet.DrawAPrimitive(RHICmdList, View, Renderer, bSeparateTranslucencyPass, Index);
+	}
+};
+
+void FTranslucentPrimSet::DrawPrimitivesParallel(
+	FRHICommandList& RHICmdList,
+	const FViewInfo& View,
+	FDeferredShadingSceneRenderer& Renderer,
+	bool bSeparateTranslucencyPass,
+	int32 FirstIndex, int32 LastIndex
+	) const
+{
+	const TArray<FSortedPrim, SceneRenderingAllocator>* PhaseSortedPrimitivesPtr = NULL;
+
+	// copy to reference for easier access
+	const TArray<FSortedPrim, SceneRenderingAllocator>& PhaseSortedPrimitives =
+		bSeparateTranslucencyPass ? SortedSeparateTranslucencyPrims : SortedPrims;
+	check(LastIndex < PhaseSortedPrimitives.Num());
+	// Draw sorted scene prims
+	for (int32 PrimIdx = FirstIndex; PrimIdx <= LastIndex; PrimIdx++)
+	{
+		FPrimitiveSceneInfo* PrimitiveSceneInfo = PhaseSortedPrimitives[PrimIdx].PrimitiveSceneInfo;
+		int32 PrimitiveId = PrimitiveSceneInfo->GetIndex();
+		const FPrimitiveViewRelevance& ViewRelevance = View.PrimitiveViewRelevanceMap[PrimitiveId];
+
+		checkSlow(ViewRelevance.HasTranslucency());
+
+		if (PrimitiveSceneInfo->Proxy && PrimitiveSceneInfo->Proxy->CastsVolumetricTranslucentShadow())
+		{
+			// can't do this in parallel, defer
+			new (RHICmdList.AllocCommand<FRHICommandVolumetricTranslucentShadow>()) FRHICommandVolumetricTranslucentShadow(*this, View, Renderer, bSeparateTranslucencyPass, PrimIdx);
+		}
+		else
+		{
+			RenderPrimitive(RHICmdList, View, PrimitiveSceneInfo, ViewRelevance, nullptr, bSeparateTranslucencyPass);
+		}
+	}
+}
+
 void FTranslucentPrimSet::DrawPrimitives(
 	FRHICommandListImmediate& RHICmdList,
 	const FViewInfo& View,
@@ -528,7 +607,7 @@ void FTranslucentPrimSet::DrawPrimitives(
 }
 
 void FTranslucentPrimSet::RenderPrimitive(
-	FRHICommandListImmediate& RHICmdList,
+	FRHICommandList& RHICmdList,
 	const FViewInfo& View, 
 	FPrimitiveSceneInfo* PrimitiveSceneInfo, 
 	const FPrimitiveViewRelevance& ViewRelevance, 
@@ -643,6 +722,234 @@ bool FSceneRenderer::ShouldRenderTranslucency() const
 	return bRender;
 }
 
+struct FRHICommandSetTranslucentRenderTargetAndState : public FRHICommand<FRHICommandSetTranslucentRenderTargetAndState>
+{
+	const FViewInfo& View;
+	bool bSeparateTranslucencyPass;
+
+	FORCEINLINE_DEBUGGABLE FRHICommandSetTranslucentRenderTargetAndState(const FViewInfo& InView, bool InbSeparateTranslucencyPass)
+		: View(InView)
+		, bSeparateTranslucencyPass(InbSeparateTranslucencyPass)
+	{
+	}
+	void Execute()
+	{
+		check(IsInRenderingThread());
+		FRHICommandList RHICmdList;
+		if (!bSeparateTranslucencyPass)
+		{
+			SetTranslucentRenderTargetAndState(RHICmdList, View, false);
+		}
+		else
+		{
+			bool bRenderSeparateTranslucency = View.TranslucentPrimSet.NumSeparateTranslucencyPrims() > 0;
+
+			// always call BeginRenderingSeparateTranslucency() even if there are no primitives to we keep the RT allocated
+			bool bSetupTranslucency = GSceneRenderTargets.BeginRenderingSeparateTranslucency(RHICmdList, View, true);
+
+			// Draw only translucent prims that are in the SeparateTranslucency pass
+			if (bRenderSeparateTranslucency)
+			{
+				if (bSetupTranslucency)
+				{
+					RHICmdList.SetDepthStencilState(TStaticDepthStencilState<false, CF_GreaterEqual>::GetRHI());
+				}
+			}
+		}
+	}
+};
+
+struct FRHICommandFinishRenderingSeparateTranslucency : public FRHICommand<FRHICommandFinishRenderingSeparateTranslucency>
+{
+	const FViewInfo& View;
+
+	FORCEINLINE_DEBUGGABLE FRHICommandFinishRenderingSeparateTranslucency(const FViewInfo& InView)
+		: View(InView)
+	{
+	}
+	void Execute()
+	{
+		check(IsInRenderingThread());
+		FRHICommandList RHICmdList;
+		GSceneRenderTargets.FinishRenderingSeparateTranslucency(RHICmdList, View);
+	}
+};
+
+class FDrawSortedTransAnyThreadTask
+{
+	FDeferredShadingSceneRenderer& Renderer;
+	FRHICommandList& RHICmdList;
+	const FViewInfo& View;
+	bool bSeparateTranslucencyPass;
+
+	const int32 FirstIndex;
+	const int32 LastIndex;
+
+public:
+
+	FDrawSortedTransAnyThreadTask(
+		FDeferredShadingSceneRenderer* InRenderer,
+		FRHICommandList* InRHICmdList,
+		const FViewInfo* InView,
+		bool InbSeparateTranslucencyPass,
+		int32 InFirstIndex,
+		int32 InLastIndex
+		)
+		: Renderer(*InRenderer)
+		, RHICmdList(*InRHICmdList)
+		, View(*InView)
+		, bSeparateTranslucencyPass(InbSeparateTranslucencyPass)
+		, FirstIndex(InFirstIndex)
+		, LastIndex(InLastIndex)
+	{
+	}
+
+	FORCEINLINE TStatId GetStatId() const
+	{
+		RETURN_QUICK_DECLARE_CYCLE_STAT(FDrawSortedTransAnyThreadTask, STATGROUP_TaskGraphTasks);
+	}
+
+	ENamedThreads::Type GetDesiredThread()
+	{
+		return ENamedThreads::AnyThread;
+	}
+
+	static ESubsequentsMode::Type GetSubsequentsMode() { return ESubsequentsMode::TrackSubsequents; }
+
+	void DoTask(ENamedThreads::Type CurrentThread, const FGraphEventRef& MyCompletionGraphEvent)
+	{
+		if (FirstIndex == 0)
+		{
+			new (RHICmdList.AllocCommand<FRHICommandSetTranslucentRenderTargetAndState>()) FRHICommandSetTranslucentRenderTargetAndState(View, bSeparateTranslucencyPass);
+		}
+		if (LastIndex >= FirstIndex)
+		{
+			View.TranslucentPrimSet.DrawPrimitivesParallel(RHICmdList, View, Renderer, bSeparateTranslucencyPass, FirstIndex, LastIndex);
+		}
+		if (bSeparateTranslucencyPass && (LastIndex < FirstIndex || LastIndex == View.TranslucentPrimSet.NumSeparateTranslucencyPrims() - 1))
+		{
+			new (RHICmdList.AllocCommand<FRHICommandFinishRenderingSeparateTranslucency>()) FRHICommandFinishRenderingSeparateTranslucency(View);
+		}
+	}
+};
+
+
+void FDeferredShadingSceneRenderer::RenderTranslucencyParallel(FRHICommandListImmediate& RHICmdList)
+{
+	GSceneRenderTargets.AllocLightAttenuation(); // materials will attempt to get this texture before the deferred command to set it up executes
+	int32 Width = CVarRHICmdWidth.GetValueOnRenderThread(); // we use a few more than needed to cover non-equal jobs
+
+	FTranslucencyDrawingPolicyFactory::ContextType ThisContext;
+
+	FGraphEventRef SubmitChain;
+
+	for (int32 ViewIndex = 0; ViewIndex < Views.Num(); ViewIndex++)
+	{
+		SCOPED_CONDITIONAL_DRAW_EVENTF(EventView, Views.Num() > 1, DEC_SCENE_ITEMS, TEXT("View%d"), ViewIndex);
+
+		const FViewInfo& View = Views[ViewIndex];
+
+		{
+			int32 NumPrims = View.TranslucentPrimSet.NumPrims() - View.TranslucentPrimSet.NumSeparateTranslucencyPrims();
+			int32 EffectiveThreads = FMath::Min<int32>(NumPrims, Width);
+
+			int32 Start = 0;
+			if (EffectiveThreads)
+			{
+
+				int32 NumPer = NumPrims / EffectiveThreads;
+				int32 Extra = NumPrims - NumPer * EffectiveThreads;
+
+
+				for (int32 ThreadIndex = 0; ThreadIndex < EffectiveThreads; ThreadIndex++)
+				{
+					int32 Last = Start + (NumPer - 1) + (ThreadIndex < Extra);
+					check(Last >= Start);
+
+					FRHICommandList* CmdList = new FRHICommandList;
+
+					FGraphEventRef AnyThreadCompletionEvent = TGraphTask<FDrawSortedTransAnyThreadTask>::CreateTask(nullptr, ENamedThreads::RenderThread)
+						.ConstructAndDispatchWhenReady(this, CmdList, &View, false, Start, Last);
+					Start = Last + 1;
+
+					SubmitChain = FSubmitCommandlistThreadTask::AddToChain(AnyThreadCompletionEvent, SubmitChain, CmdList);
+				}
+			}
+			else
+			{
+				// still need to setup the render target, a bit of a waste...
+				FRHICommandList* CmdList = new FRHICommandList;
+
+				FGraphEventRef AnyThreadCompletionEvent = TGraphTask<FDrawSortedTransAnyThreadTask>::CreateTask(nullptr, ENamedThreads::RenderThread)
+					.ConstructAndDispatchWhenReady(this, CmdList, &View, false, 0, -1);
+
+				SubmitChain = FSubmitCommandlistThreadTask::AddToChain(AnyThreadCompletionEvent, SubmitChain, CmdList);
+			}
+		}
+
+		// Draw the view's mesh elements with the translucent drawing policy.
+		SubmitChain = DrawViewElementsParallel<FTranslucencyDrawingPolicyFactory>(View, ThisContext, SDPG_World, false, SubmitChain, Width);
+		// Draw the view's mesh elements with the translucent drawing policy.
+		SubmitChain = DrawViewElementsParallel<FTranslucencyDrawingPolicyFactory>(View, ThisContext, SDPG_Foreground, false, SubmitChain, Width);
+
+#if 0 // unsupported visualization in the parallel case
+		const FSceneViewState* ViewState = (const FSceneViewState*)View.State;
+		if (ViewState && View.Family->EngineShowFlags.VisualizeLPV)
+		{
+			FLightPropagationVolume* LightPropagationVolume = ViewState->GetLightPropagationVolume();
+
+			if (LightPropagationVolume)
+			{
+				LightPropagationVolume->Visualise(RHICmdList, View);
+			}
+		}
+#endif
+		{
+			int32 NumPrims = View.TranslucentPrimSet.NumSeparateTranslucencyPrims();
+			int32 EffectiveThreads = FMath::Min<int32>(NumPrims, Width);
+
+			int32 Start = 0;
+			if (EffectiveThreads)
+			{
+
+				int32 NumPer = NumPrims / EffectiveThreads;
+				int32 Extra = NumPrims - NumPer * EffectiveThreads;
+
+
+				for (int32 ThreadIndex = 0; ThreadIndex < EffectiveThreads; ThreadIndex++)
+				{
+					int32 Last = Start + (NumPer - 1) + (ThreadIndex < Extra);
+					check(Last >= Start);
+
+					FRHICommandList* CmdList = new FRHICommandList;
+
+					FGraphEventRef AnyThreadCompletionEvent = TGraphTask<FDrawSortedTransAnyThreadTask>::CreateTask(nullptr, ENamedThreads::RenderThread)
+						.ConstructAndDispatchWhenReady(this, CmdList, &View, true, Start, Last);
+					Start = Last + 1;
+
+					SubmitChain = FSubmitCommandlistThreadTask::AddToChain(AnyThreadCompletionEvent, SubmitChain, CmdList);
+				}
+			}
+			else
+			{
+				// still need to setup the render target, a bit of a waste...
+				FRHICommandList* CmdList = new FRHICommandList;
+
+				FGraphEventRef AnyThreadCompletionEvent = TGraphTask<FDrawSortedTransAnyThreadTask>::CreateTask(nullptr, ENamedThreads::RenderThread)
+					.ConstructAndDispatchWhenReady(this, CmdList, &View, true, 0, -1);
+				SubmitChain = FSubmitCommandlistThreadTask::AddToChain(AnyThreadCompletionEvent, SubmitChain, CmdList);
+			}
+		}
+	}
+
+	if (SubmitChain.GetReference())
+	{
+		RHICmdList.Flush(); // this would happen later anyway, but we might as well do these while we wait for async tasks
+		FTaskGraphInterface::Get().WaitUntilTaskCompletes(SubmitChain, ENamedThreads::RenderThread_Local);
+	}
+}
+
+
 void FDeferredShadingSceneRenderer::RenderTranslucency(FRHICommandListImmediate& RHICmdList)
 {
 	if (ShouldRenderTranslucency())
@@ -650,6 +957,12 @@ void FDeferredShadingSceneRenderer::RenderTranslucency(FRHICommandListImmediate&
 
 		SCOPED_DRAW_EVENT(Translucency, DEC_SCENE_ITEMS);
 
+		if (!GRHICommandList.Bypass())
+		{
+			RenderTranslucencyParallel(RHICmdList);
+			return;
+		}
+
 		for (int32 ViewIndex = 0; ViewIndex < Views.Num(); ViewIndex++)
 		{
 			SCOPED_CONDITIONAL_DRAW_EVENTF(EventView, Views.Num() > 1, DEC_SCENE_ITEMS, TEXT("View%d"), ViewIndex);
diff --git a/Engine/Source/Runtime/Renderer/Private/VelocityRendering.cpp b/Engine/Source/Runtime/Renderer/Private/VelocityRendering.cpp
--- a/Engine/Source/Runtime/Renderer/Private/VelocityRendering.cpp
+++ b/Engine/Source/Runtime/Renderer/Private/VelocityRendering.cpp
@@ -265,7 +265,7 @@ void FVelocityDrawingPolicy::SetMeshRenderState(
 
 bool FVelocityDrawingPolicy::HasVelocity(const FViewInfo& View, const FPrimitiveSceneInfo* PrimitiveSceneInfo)
 {
-	checkSlow(IsInRenderingThread());
+	checkSlow(IsInParallelRenderingThread());
 
 	// No velocity if motionblur is off, or if it's a non-moving object (treat as background in that case)
 	if(View.bCameraCut || !PrimitiveSceneInfo->Proxy->IsMovable())
@@ -348,7 +348,7 @@ void FVelocityDrawingPolicyFactory::AddStaticMesh(FScene* Scene, FStaticMesh* St
 	    EBlendMode BlendMode = Material->GetBlendMode();
 	    if(BlendMode == BLEND_Opaque || BlendMode == BLEND_Masked)
 	    {
-		    if(!Material->IsMasked() && !Material->IsTwoSided() && !Material->MaterialModifiesMeshPosition())
+			if (!Material->IsMasked() && !Material->IsTwoSided() && !Material->MaterialModifiesMeshPosition_RenderThread())
 		    {
 			    // Default material doesn't handle masked or mesh-mod, and doesn't have the correct bIsTwoSided setting.
 			    MaterialRenderProxy = UMaterial::GetDefaultMaterial(MD_Surface)->GetRenderProxy(false);
@@ -386,7 +386,7 @@ bool FVelocityDrawingPolicyFactory::DrawDynamicMesh(
 		// This should be enforced at a higher level
 		//@todo - figure out why this is failing and re-enable
 		//check(FVelocityDrawingPolicy::HasVelocity(View, PrimitiveSceneInfo));
-		if(!Material->IsMasked() && !Material->IsTwoSided() && !Material->MaterialModifiesMeshPosition())
+		if (!Material->IsMasked() && !Material->IsTwoSided() && !Material->MaterialModifiesMeshPosition_RenderThread())
 		{
 			// Default material doesn't handle masked, and doesn't have the correct bIsTwoSided setting.
 			MaterialRenderProxy = UMaterial::GetDefaultMaterial(MD_Surface)->GetRenderProxy(false);
