UE4 - fix regression with widgets

diff --git a/Engine/Source/Runtime/Renderer/Private/DynamicPrimitiveDrawing.inl b/Engine/Source/Runtime/Renderer/Private/DynamicPrimitiveDrawing.inl
--- a/Engine/Source/Runtime/Renderer/Private/DynamicPrimitiveDrawing.inl
+++ b/Engine/Source/Runtime/Renderer/Private/DynamicPrimitiveDrawing.inl
@@ -359,7 +359,7 @@ bool DrawViewElements(
 {
 	// Get the correct element list based on dpg index
 	const TIndirectArray<FHitProxyMeshPair>& ViewMeshElementList = (DPGIndex == SDPG_Foreground ? View.TopViewMeshElements : View.ViewMeshElements);
-	if (View.ViewMeshElements.Num() != 0)
+	if (ViewMeshElementList.Num() != 0)
 	{
 		DrawViewElementsInner<DrawingPolicyFactoryType>(RHICmdList, View, DrawingContext, DPGIndex, bPreFog, 0, ViewMeshElementList.Num() - 1);
 		return true;
