Highres screenshot - Fix incorrect test for a capture region existing when saving out buffer visualization targets

diff --git a/Engine/Source/Runtime/Renderer/Private/PostProcess/RenderingCompositionGraph.cpp b/Engine/Source/Runtime/Renderer/Private/PostProcess/RenderingCompositionGraph.cpp
--- a/Engine/Source/Runtime/Renderer/Private/PostProcess/RenderingCompositionGraph.cpp
+++ b/Engine/Source/Runtime/Renderer/Private/PostProcess/RenderingCompositionGraph.cpp
@@ -270,7 +270,7 @@ void FRenderingCompositionGraph::DumpOutputToFile(FRenderingCompositePassContext
 	{
 		SourceRect = GetHighResScreenshotConfig().CaptureRegion;
 
-		if (SourceRect.Area() > 0)
+		if (SourceRect.Area() == 0)
 		{
 			SourceRect = Context.View.ViewRect;
 		}
