UE4 - Integrate from FNM

diff --git a/Engine/Shaders/VisualizeTexture.usf b/Engine/Shaders/VisualizeTexture.usf
--- a/Engine/Shaders/VisualizeTexture.usf
+++ b/Engine/Shaders/VisualizeTexture.usf
@@ -93,6 +93,7 @@ void VisualizeTexturePS(
 	int InputValueMapping = (int)VisualizeParam[2].r;
 	if(InputValueMapping == 1)
 	{
+#if FEATURE_LEVEL > FEATURE_LEVEL_ES2
 		// current z buffer is negated (1 is near, 0 is far)
 		float Depth = 1 - TexLookup.r;
 
@@ -109,6 +110,9 @@ void VisualizeTexturePS(
 		fStencil /= 255.0f;
 
 		TexLookup.a = fStencil;
+#else
+		TexLookup.r = 1 - TexLookup.r;
+#endif
 	}	
 
 	float Add = TexLookup.a * VisualizeParam[0].y + VisualizeParam[0].z;
diff --git a/Engine/Source/Runtime/Core/Private/Misc/CoreMisc.cpp b/Engine/Source/Runtime/Core/Private/Misc/CoreMisc.cpp
--- a/Engine/Source/Runtime/Core/Private/Misc/CoreMisc.cpp
+++ b/Engine/Source/Runtime/Core/Private/Misc/CoreMisc.cpp
@@ -318,7 +318,6 @@ bool FFileHelper::CreateBitmap( const TCHAR* Pattern, int32 SourceWidth, int32 S
 		FString Filename;
 		if (GenerateNextBitmapFilename(Pattern, Filename, FileManager))
 		{
-			Filename += TEXT(".bmp");
 			FCString::Strcpy(File, *Filename);
 			if ( OutFilename )
 			{
diff --git a/Engine/Source/Runtime/Renderer/Private/PostProcess/VisualizeTexture.cpp b/Engine/Source/Runtime/Renderer/Private/PostProcess/VisualizeTexture.cpp
--- a/Engine/Source/Runtime/Renderer/Private/PostProcess/VisualizeTexture.cpp
+++ b/Engine/Source/Runtime/Renderer/Private/PostProcess/VisualizeTexture.cpp
@@ -420,7 +420,11 @@ void FVisualizeTexture::GenerateContent(FRHICommandListImmediate& RHICmdList, co
 		// continue rendering to HDR if necessary
 		RenderVisualizeTexture(RHICmdList, VisualizeTextureData);
 	}
-	RHICmdList.CopyToResolveTarget(DestRenderTarget.TargetableTexture, DestRenderTarget.ShaderResourceTexture, false, FResolveParams());
+
+	{
+		SCOPED_DRAW_EVENT(VisCopy, DEC_SCENE_ITEMS);
+		RHICmdList.CopyToResolveTarget(DestRenderTarget.TargetableTexture, DestRenderTarget.ShaderResourceTexture, false, FResolveParams());
+	}
 
 	VisualizeTextureDesc = Desc;
 
@@ -517,17 +521,20 @@ void FVisualizeTexture::PresentContent(FRHICommandListImmediate& RHICmdList, con
 
 	FIntRect VisualizeTextureRect = ComputeVisualizeTextureRect(Desc.Extent);
 
-	// Draw a quad mapping scene color to the view's render target
-	DrawRectangle(
-		RHICmdList,
-		VisualizeTextureRect.Min.X, VisualizeTextureRect.Min.Y,
-		VisualizeTextureRect.Width(), VisualizeTextureRect.Height(),
-		0, 0,
-		VisualizeTextureRect.Width(), VisualizeTextureRect.Height(),
-		FIntPoint(RenderTarget->GetSizeX(), RenderTarget->GetSizeY()),
-		VisualizeTextureRect.Size(),
-		*VertexShader,
-		EDRF_Default);
+	{
+		SCOPED_DRAW_EVENT(VisCopyToMain, DEC_SCENE_ITEMS);
+		// Draw a quad mapping scene color to the view's render target
+		DrawRectangle(
+			RHICmdList,
+			VisualizeTextureRect.Min.X, VisualizeTextureRect.Min.Y,
+			VisualizeTextureRect.Width(), VisualizeTextureRect.Height(),
+			0, 0,
+			VisualizeTextureRect.Width(), VisualizeTextureRect.Height(),
+			FIntPoint(RenderTarget->GetSizeX(), RenderTarget->GetSizeY()),
+			VisualizeTextureRect.Size(),
+			*VertexShader,
+			EDRF_Default);
+	}
 
 	// this is a helper class for FCanvas to be able to get screen size
 	class FRenderTargetTemp : public FRenderTarget
@@ -898,7 +905,7 @@ IPooledRenderTarget* FVisualizeTexture::GetObservedElement() const
 
 void FVisualizeTexture::OnStartFrame(const FSceneView& View)
 {
-	bEnabled = GMaxRHIFeatureLevel >= ERHIFeatureLevel::SM3;
+	bEnabled = GMaxRHIFeatureLevel >= ERHIFeatureLevel::ES2;
 	ViewRect = View.UnscaledViewRect;
 	AspectRatioConstrainedViewRect = View.Family->EngineShowFlags.CameraAspectRatioBars ? View.CameraConstrainedViewRect : ViewRect;
 
