Added basic support for mesh based ES2 stereo HMD distortion passes.

Currently uses the Oculus HMD mesh distortion shader and vert format.

diff --git a/Engine/Source/Runtime/HeadMountedDisplay/Public/IHeadMountedDisplay.h b/Engine/Source/Runtime/HeadMountedDisplay/Public/IHeadMountedDisplay.h
--- a/Engine/Source/Runtime/HeadMountedDisplay/Public/IHeadMountedDisplay.h
+++ b/Engine/Source/Runtime/HeadMountedDisplay/Public/IHeadMountedDisplay.h
@@ -16,7 +16,8 @@ namespace EHMDDeviceType
 	enum Type
 	{
 		DT_OculusRift,
-		DT_Morpheus
+		DT_Morpheus,
+		DT_ES2GenericStereoMesh
 	};
 }
 
diff --git a/Engine/Source/Runtime/Renderer/Private/PostProcess/PostProcessing.cpp b/Engine/Source/Runtime/Renderer/Private/PostProcess/PostProcessing.cpp
--- a/Engine/Source/Runtime/Renderer/Private/PostProcess/PostProcessing.cpp
+++ b/Engine/Source/Runtime/Renderer/Private/PostProcess/PostProcessing.cpp
@@ -1492,6 +1492,23 @@ void FPostProcessing::ProcessES2(FRHICommandListImmediate& RHICmdList, const FVi
 			Context.FinalOutput = FRenderingCompositeOutputRef(Node);
 		}
 
+		bool bStereoRenderingAndHMD = View.Family->EngineShowFlags.StereoRendering && View.Family->EngineShowFlags.HMDDistortion;
+		if (bStereoRenderingAndHMD)
+		{
+			FRenderingCompositePass* Node = NULL;
+			const EHMDDeviceType::Type DeviceType = GEngine->HMDDevice->GetHMDDeviceType();
+			if (DeviceType == EHMDDeviceType::DT_ES2GenericStereoMesh)
+			{
+				Node = Context.Graph.RegisterPass(new FRCPassPostProcessHMD());
+			}
+
+			if (Node)
+			{
+				Node->SetInput(ePId_Input0, FRenderingCompositeOutputRef(Context.FinalOutput));
+				Context.FinalOutput = FRenderingCompositeOutputRef(Node);
+			}
+		}
+
 		// The graph setup should be finished before this line ----------------------------------------
 
 		{
