fix non unity build

diff --git a/Engine/Source/Runtime/Renderer/Private/RendererScene.cpp b/Engine/Source/Runtime/Renderer/Private/RendererScene.cpp
--- a/Engine/Source/Runtime/Renderer/Private/RendererScene.cpp
+++ b/Engine/Source/Runtime/Renderer/Private/RendererScene.cpp
@@ -11,6 +11,7 @@
 #include "ParameterCollection.h"
 #include "DistanceFieldSurfaceCacheLighting.h"
 #include "EngineModule.h"
+#include "PrecomputedLightVolume.h"
 
 // Enable this define to do slow checks for components being added to the wrong
 // world's scene, when using PIE. This can happen if a PIE component is reattached
