UE4 - rework the submission and tasking of "sub command lists"

diff --git a/Engine/Source/Runtime/Engine/Private/GPUSkinVertexFactory.cpp b/Engine/Source/Runtime/Engine/Private/GPUSkinVertexFactory.cpp
--- a/Engine/Source/Runtime/Engine/Private/GPUSkinVertexFactory.cpp
+++ b/Engine/Source/Runtime/Engine/Private/GPUSkinVertexFactory.cpp
@@ -428,7 +428,7 @@ public:
 					const FGPUBaseSkinVertexFactory* GPUVertexFactory = (const FGPUBaseSkinVertexFactory*)VertexFactory;
 
 					// copy the bone data and tell the instance where it can pick it up next frame
-
+					check(IsInRenderingThread());
 					// append data to a buffer we bind next frame to read old matrix data for motion blur
 					uint32 OldBoneDataStartIndex = GPrevPerBoneMotionBlur.AppendData(ShaderData.BoneMatrices.GetTypedData(), ShaderData.BoneMatrices.Num());
 					GPUVertexFactory->SetOldBoneDataStartIndex(View.FrameNumber, OldBoneDataStartIndex);
diff --git a/Engine/Source/Runtime/RHI/Private/RHICommandList.cpp b/Engine/Source/Runtime/RHI/Private/RHICommandList.cpp
--- a/Engine/Source/Runtime/RHI/Private/RHICommandList.cpp
+++ b/Engine/Source/Runtime/RHI/Private/RHICommandList.cpp
@@ -4,20 +4,31 @@
 #include "RHI.h"
 #include "RHICommandList.h"
 
-DEFINE_STAT(STAT_ImmedCmdListExecuteTime);
-DEFINE_STAT(STAT_ImmedCmdListMemory);
-DEFINE_STAT(STAT_ImmedCmdListCount);
+DECLARE_CYCLE_STAT(TEXT("Nonimmed. Command List Execute"), STAT_NonImmedCmdListExecuteTime, STATGROUP_RHICMDLIST);
+DECLARE_DWORD_COUNTER_STAT(TEXT("Nonimmed. Command List memory"), STAT_NonImmedCmdListMemory, STATGROUP_RHICMDLIST);
+DECLARE_DWORD_COUNTER_STAT(TEXT("Nonimmed. Command count"), STAT_NonImmedCmdListCount, STATGROUP_RHICMDLIST);
+
+DECLARE_CYCLE_STAT(TEXT("Immed. Command List Execute"), STAT_ImmedCmdListExecuteTime, STATGROUP_RHICMDLIST);
+DECLARE_DWORD_COUNTER_STAT(TEXT("Immed. Command List memory"), STAT_ImmedCmdListMemory, STATGROUP_RHICMDLIST);
+DECLARE_DWORD_COUNTER_STAT(TEXT("Immed. Command count"), STAT_ImmedCmdListCount, STATGROUP_RHICMDLIST);
 
-DEFINE_STAT(STAT_NonImmedCmdListExecuteTime);
-DEFINE_STAT(STAT_NonImmedCmdListMemory);
-DEFINE_STAT(STAT_NonImmedCmdListCount);
 
 static TAutoConsoleVariable<int32> CVarRHICmdBypass(
 	TEXT("r.RHICmdBypass"),
 	FRHICommandListExecutor::DefaultBypass,
 	TEXT("Whether to bypass the rhi command list and send the rhi commands immediately.\n")
-	TEXT("0: Disable, 1: Enable"),
-	ECVF_Cheat);
+	TEXT("0: Disable, 1: Enable"));
+
+static TAutoConsoleVariable<int32> CVarRHICmdUseParallelAlgorithms(
+	TEXT("r.RHICmdUseParallelAlgorithms"),
+	FRHICommandListExecutor::DefaultBypass,
+	TEXT("True to use parallel algorithms. Ignored if r.RHICmdBypass is 1.\n"));
+
+static TAutoConsoleVariable<int32> CVarRHICmdFlushBeforeWait(
+	TEXT("r.RHICmdFlushBeforeWait"),
+	1,
+	TEXT("Whether to flush the immediate command list before waiting for a command list chain.\n")
+	TEXT("0: Disable, 1: Enable"));
 
 TAutoConsoleVariable<int32> CVarRHICmdWidth(
 	TEXT("r.RHICmdWidth"), 
@@ -66,13 +77,28 @@ void FRHICommandListExecutor::ExecuteList(FRHICommandListImmediate& CmdList)
 
 	INC_MEMORY_STAT_BY(STAT_ImmedCmdListMemory, CmdList.GetUsedMemory());
 	INC_DWORD_STAT_BY(STAT_ImmedCmdListCount, CmdList.NumCommands);
-
+#if 0
+	static TAutoConsoleVariable<int32> CVarRHICmdMemDump(
+		TEXT("r.RHICmdMemDump"),
+		0,
+		TEXT("dumps callstacks and sizes of the immediate command lists to the console.\n")
+		TEXT("0: Disable, 1: Enable"),
+		ECVF_Cheat);
+	if (CVarRHICmdMemDump.GetValueOnRenderThread() > 0)
+	{
+		FPlatformMisc::LowLevelOutputDebugStringf(TEXT("Mem %d\n"), CmdList.GetUsedMemory());
+		if (CmdList.GetUsedMemory() > 300)
+		{
+			FPlatformMisc::LowLevelOutputDebugStringf(TEXT("big\n"));
+		}
+	}
+#endif
 	{
 		SCOPE_CYCLE_COUNTER(STAT_ImmedCmdListExecuteTime);
 		ExecuteInner(CmdList);
 	}
 
-#if !(UE_BUILD_SHIPPING || UE_BUILD_TEST)
+#if !UE_BUILD_SHIPPING
 	if (GRHICommandList.OutstandingCmdListCount.GetValue() == 1)
 	{
 		LatchBypass();
@@ -83,7 +109,7 @@ void FRHICommandListExecutor::ExecuteList(FRHICommandListImmediate& CmdList)
 
 void FRHICommandListExecutor::LatchBypass()
 {
-#if !(UE_BUILD_SHIPPING || UE_BUILD_TEST)
+#if !UE_BUILD_SHIPPING
 
 	static bool bOnce = false;
 	if (!bOnce)
@@ -103,6 +129,14 @@ void FRHICommandListExecutor::LatchBypass()
 		FRHIResource::FlushPendingDeletes();
 	}
 	bLatchedBypass = NewBypass;
+	if (!bLatchedBypass)
+	{
+		bLatchedUseParallelAlgorithms = (CVarRHICmdUseParallelAlgorithms.GetValueOnAnyThread() >= 1);
+	}
+	else
+	{
+		bLatchedUseParallelAlgorithms = false;
+	}
 #endif
 }
 
@@ -139,6 +173,39 @@ void FRHICommandListBase::Reset()
 	UID = GRHICommandList.UIDCounter.Increment();
 }
 
+DECLARE_DWORD_COUNTER_STAT(TEXT("Num Chains Links"), STAT_ChainLinkCount, STATGROUP_RHICMDLIST);
+DECLARE_CYCLE_STAT(TEXT("Wait for CmdList"), STAT_ChainWait, STATGROUP_RHICMDLIST);
+DECLARE_CYCLE_STAT(TEXT("Chain Execute"), STAT_ChainExecute, STATGROUP_RHICMDLIST);
+
+struct FRHICommandWaitForAndSubmitSubList : public FRHICommand<FRHICommandWaitForAndSubmitSubList>
+{
+	FGraphEventRef EventToWaitFor;
+	FRHICommandList* RHICmdList;
+	FORCEINLINE_DEBUGGABLE FRHICommandWaitForAndSubmitSubList(FGraphEventRef& InEventToWaitFor, FRHICommandList* InRHICmdList)
+		: EventToWaitFor(InEventToWaitFor)
+		, RHICmdList(InRHICmdList)
+	{
+	}
+	void Execute()
+	{
+		INC_DWORD_STAT_BY(STAT_ChainLinkCount, 1);
+		{
+			SCOPE_CYCLE_COUNTER(STAT_ChainWait);
+			FTaskGraphInterface::Get().WaitUntilTaskCompletes(EventToWaitFor, ENamedThreads::RenderThread_Local);
+		}
+		{
+			SCOPE_CYCLE_COUNTER(STAT_ChainExecute);
+			delete RHICmdList;
+		}
+	}
+};
+
+void FRHICommandListBase::QueueAsyncCommandListSubmit(FGraphEventRef AnyThreadCompletionEvent, class FRHICommandList* CmdList)
+{
+	new (AllocCommand<FRHICommandWaitForAndSubmitSubList>()) FRHICommandWaitForAndSubmitSubList(AnyThreadCompletionEvent, CmdList);
+}
+
+
 static TLockFreeFixedSizeAllocator<sizeof(FRHICommandList), FThreadSafeCounter> RHICommandListAllocator;
 
 void* FRHICommandList::operator new(size_t Size)
@@ -158,4 +225,3 @@ void FRHICommandList::operator delete(void *RawMemory)
 	//FMemory::Free(RawMemory);
 }	
 
-
diff --git a/Engine/Source/Runtime/RHI/Public/RHICommandList.h b/Engine/Source/Runtime/RHI/Public/RHICommandList.h
--- a/Engine/Source/Runtime/RHI/Public/RHICommandList.h
+++ b/Engine/Source/Runtime/RHI/Public/RHICommandList.h
@@ -6,20 +6,10 @@
 
 #pragma once
 #include "LockFreeFixedSizeAllocator.h"
-
+#include "TaskGraphInterfaces.h"
 
 DECLARE_STATS_GROUP(TEXT("RHICmdList"), STATGROUP_RHICMDLIST, STATCAT_Advanced);
 
-DECLARE_CYCLE_STAT_EXTERN(TEXT("Nonimmed. Command List Execute"), STAT_NonImmedCmdListExecuteTime, STATGROUP_RHICMDLIST, );
-DECLARE_DWORD_COUNTER_STAT_EXTERN(TEXT("Nonimmed. Command List memory"), STAT_NonImmedCmdListMemory, STATGROUP_RHICMDLIST, RHI_API);
-DECLARE_DWORD_COUNTER_STAT_EXTERN(TEXT("Nonimmed. Command count"), STAT_NonImmedCmdListCount, STATGROUP_RHICMDLIST, RHI_API);
-
-DECLARE_CYCLE_STAT_EXTERN(TEXT("Immed. Command List Execute"), STAT_ImmedCmdListExecuteTime, STATGROUP_RHICMDLIST, );
-DECLARE_DWORD_COUNTER_STAT_EXTERN(TEXT("Immed. Command List memory"), STAT_ImmedCmdListMemory, STATGROUP_RHICMDLIST, RHI_API);
-DECLARE_DWORD_COUNTER_STAT_EXTERN(TEXT("Immed. Command count"), STAT_ImmedCmdListCount, STATGROUP_RHICMDLIST, RHI_API);
-
-
-
 extern RHI_API TAutoConsoleVariable<int32> CVarRHICmdWidth;
 
 struct FRHICommandBase
@@ -45,7 +35,9 @@ public:
 	~FRHICommandListBase();
 
 	inline void Flush();
+	inline bool IsImmediate();
 	const int32 GetUsedMemory() const;
+	void QueueAsyncCommandListSubmit(FGraphEventRef AnyThreadCompletionEvent, class FRHICommandList* CmdList);
 
 	FORCEINLINE_DEBUGGABLE void* Alloc(int32 AllocSize, int32 Alignment)
 	{
@@ -1624,6 +1616,7 @@ public:
 	};
 	FRHICommandListExecutor()
 		: bLatchedBypass(!!DefaultBypass)
+		, bLatchedUseParallelAlgorithms(!DefaultBypass && FApp::ShouldUseThreadingForPerformance())
 	{
 	}
 	static inline FRHICommandListImmediate& GetImmediateCommandList();
@@ -1638,10 +1631,18 @@ public:
 	}
 	FORCEINLINE_DEBUGGABLE bool Bypass()
 	{
-#if !(UE_BUILD_SHIPPING || UE_BUILD_TEST)
+#if !UE_BUILD_SHIPPING
 		return bLatchedBypass;
 #else
-		return  !FApp::ShouldUseThreadingForPerformance() || !!DefaultBypass;
+		return !!DefaultBypass;
+#endif
+	}
+	FORCEINLINE_DEBUGGABLE bool UseParallelAlgorithms()
+	{
+#if !UE_BUILD_SHIPPING
+		return bLatchedUseParallelAlgorithms;
+#else
+		return  FApp::ShouldUseThreadingForPerformance() && !Bypass();
 #endif
 	}
 	static void CheckNoOutstandingCmdLists();
@@ -1651,6 +1652,7 @@ private:
 	void ExecuteInner(FRHICommandListBase& CmdList);
 
 	bool bLatchedBypass;
+	bool bLatchedUseParallelAlgorithms;
 	friend class FRHICommandListBase;
 	FThreadSafeCounter UIDCounter;
 	FThreadSafeCounter OutstandingCmdListCount;
@@ -1664,6 +1666,19 @@ FORCEINLINE_DEBUGGABLE FRHICommandListImmediate& FRHICommandListExecutor::GetImm
 	return GRHICommandList.CommandListImmediate;
 }
 
+struct FScopedCommandListFlush
+{
+	FRHICommandListImmediate& RHICmdList;
+
+	FScopedCommandListFlush(FRHICommandListImmediate& InRHICmdList = FRHICommandListExecutor::GetImmediateCommandList())
+		: RHICmdList(InRHICmdList)
+	{
+	}
+	~FScopedCommandListFlush()
+	{
+		RHICmdList.Flush();
+	}
+};
 
 #define DEFINE_RHIMETHOD_CMDLIST(Type,Name,ParameterTypesAndNames,ParameterNames,ReturnStatement,NullImplementation)
 #define DEFINE_RHIMETHOD(Type, Name, ParameterTypesAndNames, ParameterNames, ReturnStatement, NullImplementation)
diff --git a/Engine/Source/Runtime/Renderer/Private/DeferredShadingRenderer.cpp b/Engine/Source/Runtime/Renderer/Private/DeferredShadingRenderer.cpp
--- a/Engine/Source/Runtime/Renderer/Private/DeferredShadingRenderer.cpp
+++ b/Engine/Source/Runtime/Renderer/Private/DeferredShadingRenderer.cpp
@@ -212,20 +212,18 @@ bool FDeferredShadingSceneRenderer::RenderBasePassStaticDataMasked(FRHICommandLi
 	return bDirty;
 }
 
-FGraphEventRef FDeferredShadingSceneRenderer::RenderBasePassStaticDataMaskedParallel(FViewInfo& View, int32 Width, FGraphEventRef SubmitChain, bool& OutDirty)
+void FDeferredShadingSceneRenderer::RenderBasePassStaticDataMaskedParallel(FViewInfo& View, int32 Width, FRHICommandList& ParentCmdList, bool& OutDirty)
 {
 	// Draw the scene's base pass draw lists.
 	FScene::EBasePassDrawListType MaskedDrawType = FScene::EBasePass_Masked;
-	SubmitChain = Scene->BasePassNoLightMapDrawList[MaskedDrawType].DrawVisibleParallel(View, View.StaticMeshVisibilityMap, View.StaticMeshBatchVisibility, Width, SubmitChain, OutDirty);
-	SubmitChain = Scene->BasePassSimpleDynamicLightingDrawList[MaskedDrawType].DrawVisibleParallel(View, View.StaticMeshVisibilityMap, View.StaticMeshBatchVisibility, Width, SubmitChain, OutDirty);
-	SubmitChain = Scene->BasePassCachedVolumeIndirectLightingDrawList[MaskedDrawType].DrawVisibleParallel(View, View.StaticMeshVisibilityMap, View.StaticMeshBatchVisibility, Width, SubmitChain, OutDirty);
-	SubmitChain = Scene->BasePassCachedPointIndirectLightingDrawList[MaskedDrawType].DrawVisibleParallel(View, View.StaticMeshVisibilityMap, View.StaticMeshBatchVisibility, Width, SubmitChain, OutDirty);
-
-	SubmitChain = Scene->BasePassHighQualityLightMapDrawList[MaskedDrawType].DrawVisibleParallel(View, View.StaticMeshVisibilityMap, View.StaticMeshBatchVisibility, Width, SubmitChain, OutDirty);
-	SubmitChain = Scene->BasePassDistanceFieldShadowMapLightMapDrawList[MaskedDrawType].DrawVisibleParallel(View, View.StaticMeshVisibilityMap, View.StaticMeshBatchVisibility, Width, SubmitChain, OutDirty);
-	SubmitChain = Scene->BasePassLowQualityLightMapDrawList[MaskedDrawType].DrawVisibleParallel(View, View.StaticMeshVisibilityMap, View.StaticMeshBatchVisibility, Width, SubmitChain, OutDirty);
-
-	return SubmitChain;
+	Scene->BasePassNoLightMapDrawList[MaskedDrawType].DrawVisibleParallel(View, View.StaticMeshVisibilityMap, View.StaticMeshBatchVisibility, Width, ParentCmdList, OutDirty);
+	Scene->BasePassSimpleDynamicLightingDrawList[MaskedDrawType].DrawVisibleParallel(View, View.StaticMeshVisibilityMap, View.StaticMeshBatchVisibility, Width, ParentCmdList, OutDirty);
+	Scene->BasePassCachedVolumeIndirectLightingDrawList[MaskedDrawType].DrawVisibleParallel(View, View.StaticMeshVisibilityMap, View.StaticMeshBatchVisibility, Width, ParentCmdList, OutDirty);
+	Scene->BasePassCachedPointIndirectLightingDrawList[MaskedDrawType].DrawVisibleParallel(View, View.StaticMeshVisibilityMap, View.StaticMeshBatchVisibility, Width, ParentCmdList, OutDirty);
+
+	Scene->BasePassHighQualityLightMapDrawList[MaskedDrawType].DrawVisibleParallel(View, View.StaticMeshVisibilityMap, View.StaticMeshBatchVisibility, Width, ParentCmdList, OutDirty);
+	Scene->BasePassDistanceFieldShadowMapLightMapDrawList[MaskedDrawType].DrawVisibleParallel(View, View.StaticMeshVisibilityMap, View.StaticMeshBatchVisibility, Width, ParentCmdList, OutDirty);
+	Scene->BasePassLowQualityLightMapDrawList[MaskedDrawType].DrawVisibleParallel(View, View.StaticMeshVisibilityMap, View.StaticMeshBatchVisibility, Width, ParentCmdList, OutDirty);
 }
 
 bool FDeferredShadingSceneRenderer::RenderBasePassStaticDataDefault(FRHICommandList& RHICmdList, FViewInfo& View)
@@ -251,19 +249,17 @@ bool FDeferredShadingSceneRenderer::RenderBasePassStaticDataDefault(FRHICommandL
 	return bDirty;
 }
 
-FGraphEventRef FDeferredShadingSceneRenderer::RenderBasePassStaticDataDefaultParallel(FViewInfo& View, int32 Width, FGraphEventRef SubmitChain, bool& OutDirty)
+void FDeferredShadingSceneRenderer::RenderBasePassStaticDataDefaultParallel(FViewInfo& View, int32 Width, FRHICommandList& ParentCmdList, bool& OutDirty)
 {
 	FScene::EBasePassDrawListType OpaqueDrawType = FScene::EBasePass_Default;
-	SubmitChain = Scene->BasePassNoLightMapDrawList[OpaqueDrawType].DrawVisibleParallel(View, View.StaticMeshVisibilityMap, View.StaticMeshBatchVisibility, Width, SubmitChain, OutDirty);
-	SubmitChain = Scene->BasePassSimpleDynamicLightingDrawList[OpaqueDrawType].DrawVisibleParallel(View, View.StaticMeshVisibilityMap, View.StaticMeshBatchVisibility, Width, SubmitChain, OutDirty);
-	SubmitChain = Scene->BasePassCachedVolumeIndirectLightingDrawList[OpaqueDrawType].DrawVisibleParallel(View, View.StaticMeshVisibilityMap, View.StaticMeshBatchVisibility, Width, SubmitChain, OutDirty);
-	SubmitChain = Scene->BasePassCachedPointIndirectLightingDrawList[OpaqueDrawType].DrawVisibleParallel(View, View.StaticMeshVisibilityMap, View.StaticMeshBatchVisibility, Width, SubmitChain, OutDirty);
-
-	SubmitChain = Scene->BasePassHighQualityLightMapDrawList[OpaqueDrawType].DrawVisibleParallel(View, View.StaticMeshVisibilityMap, View.StaticMeshBatchVisibility, Width, SubmitChain, OutDirty);
-	SubmitChain = Scene->BasePassDistanceFieldShadowMapLightMapDrawList[OpaqueDrawType].DrawVisibleParallel(View, View.StaticMeshVisibilityMap, View.StaticMeshBatchVisibility, Width, SubmitChain, OutDirty);
-	SubmitChain = Scene->BasePassLowQualityLightMapDrawList[OpaqueDrawType].DrawVisibleParallel(View, View.StaticMeshVisibilityMap, View.StaticMeshBatchVisibility, Width, SubmitChain, OutDirty);
-
-	return SubmitChain;
+	Scene->BasePassNoLightMapDrawList[OpaqueDrawType].DrawVisibleParallel(View, View.StaticMeshVisibilityMap, View.StaticMeshBatchVisibility, Width, ParentCmdList, OutDirty);
+	Scene->BasePassSimpleDynamicLightingDrawList[OpaqueDrawType].DrawVisibleParallel(View, View.StaticMeshVisibilityMap, View.StaticMeshBatchVisibility, Width, ParentCmdList, OutDirty);
+	Scene->BasePassCachedVolumeIndirectLightingDrawList[OpaqueDrawType].DrawVisibleParallel(View, View.StaticMeshVisibilityMap, View.StaticMeshBatchVisibility, Width, ParentCmdList, OutDirty);
+	Scene->BasePassCachedPointIndirectLightingDrawList[OpaqueDrawType].DrawVisibleParallel(View, View.StaticMeshVisibilityMap, View.StaticMeshBatchVisibility, Width, ParentCmdList, OutDirty);
+
+	Scene->BasePassHighQualityLightMapDrawList[OpaqueDrawType].DrawVisibleParallel(View, View.StaticMeshVisibilityMap, View.StaticMeshBatchVisibility, Width, ParentCmdList, OutDirty);
+	Scene->BasePassDistanceFieldShadowMapLightMapDrawList[OpaqueDrawType].DrawVisibleParallel(View, View.StaticMeshVisibilityMap, View.StaticMeshBatchVisibility, Width, ParentCmdList, OutDirty);
+	Scene->BasePassLowQualityLightMapDrawList[OpaqueDrawType].DrawVisibleParallel(View, View.StaticMeshVisibilityMap, View.StaticMeshBatchVisibility, Width, ParentCmdList, OutDirty);
 }
 
 void FDeferredShadingSceneRenderer::SortBasePassStaticData(FVector ViewPosition)
@@ -318,7 +314,7 @@ bool FDeferredShadingSceneRenderer::RenderBasePassStaticData(FRHICommandList& RH
 	return bDirty;
 }
 
-FGraphEventRef FDeferredShadingSceneRenderer::RenderBasePassStaticDataParallel(FViewInfo& View, int32 Width, FGraphEventRef SubmitChain, bool& OutDirty)
+void FDeferredShadingSceneRenderer::RenderBasePassStaticDataParallel(FViewInfo& View, int32 Width, FRHICommandList& ParentCmdList, bool& OutDirty)
 {
 	SCOPE_CYCLE_COUNTER(STAT_StaticDrawListDrawTime);
 
@@ -328,19 +324,17 @@ FGraphEventRef FDeferredShadingSceneRenderer::RenderBasePassStaticDataParallel(F
 	// rendering the default opaque geometry afterward.
 	if (EarlyZPassMode != DDM_None)
 	{
-		SubmitChain = RenderBasePassStaticDataMaskedParallel(View, Width, SubmitChain, OutDirty);
-		SubmitChain = RenderBasePassStaticDataDefaultParallel(View, Width, SubmitChain, OutDirty);
+		RenderBasePassStaticDataMaskedParallel(View, Width, ParentCmdList, OutDirty);
+		RenderBasePassStaticDataDefaultParallel(View, Width, ParentCmdList, OutDirty);
 	}
 	else
 	{
 		// Otherwise, in the case where we're not using a depth-only pre-pass, there
 		// is an advantage to rendering default opaque first to help cull the more
 		// expensive masked geometry.
-		SubmitChain = RenderBasePassStaticDataDefaultParallel(View, Width, SubmitChain, OutDirty);
-		SubmitChain = RenderBasePassStaticDataMaskedParallel(View, Width, SubmitChain, OutDirty);
+		RenderBasePassStaticDataDefaultParallel(View, Width, ParentCmdList, OutDirty);
+		RenderBasePassStaticDataMaskedParallel(View, Width, ParentCmdList, OutDirty);
 	}
-
-	return SubmitChain;
 }
 
 /**
@@ -454,17 +448,17 @@ public:
 	}
 };
 
-FGraphEventRef FDeferredShadingSceneRenderer::RenderBasePassDynamicDataParallel(FViewInfo& View, FGraphEventRef SubmitChain, bool& bOutDirty)
+void FDeferredShadingSceneRenderer::RenderBasePassDynamicDataParallel(FViewInfo& View, FRHICommandList& ParentCmdList, bool& bOutDirty)
 {
 	FRHICommandList* CmdList = new FRHICommandList;
 
 	FGraphEventRef AnyThreadCompletionEvent = TGraphTask<FRenderBasePassDynamicDataThreadTask>::CreateTask(nullptr, ENamedThreads::RenderThread)
 		.ConstructAndDispatchWhenReady(this, CmdList, &View, &bOutDirty);
 
-	return FSubmitCommandlistThreadTask::AddToChain(AnyThreadCompletionEvent, SubmitChain, CmdList);
+	ParentCmdList.QueueAsyncCommandListSubmit(AnyThreadCompletionEvent, CmdList);
 }
 
-static void SetupBasePassView(FRHICommandList& RHICmdList, const FViewInfo& View, bool bShaderComplexity)
+static void SetupBasePassView(FRHICommandList& RHICmdList, const FIntRect& ViewRect, bool bShaderComplexity)
 {
 	if (bShaderComplexity)
 	{
@@ -480,64 +474,39 @@ static void SetupBasePassView(FRHICommandList& RHICmdList, const FViewInfo& View
 		// Note, this is a reversed Z depth surface, using CF_GreaterEqual.
 		RHICmdList.SetDepthStencilState(TStaticDepthStencilState<true,CF_GreaterEqual>::GetRHI());
 	}
-	RHICmdList.SetViewport(View.ViewRect.Min.X, View.ViewRect.Min.Y, 0, View.ViewRect.Max.X, View.ViewRect.Max.Y, 1);
+	RHICmdList.SetViewport(ViewRect.Min.X, ViewRect.Min.Y, 0, ViewRect.Max.X, ViewRect.Max.Y, 1);
 }
 
-class FSubmitCommandlistSetupBasePassView
+struct FRHICommandSetupBasePassView : public FRHICommand<FRHICommandSetupBasePassView>
 {
-	const FViewInfo& View;
+	FIntRect ViewRect;
 	bool bShaderComplexity;
-public:
-
-	FSubmitCommandlistSetupBasePassView(const FViewInfo* InView, bool bInShaderComplexity)
-		: View(*InView)
-		, bShaderComplexity(bInShaderComplexity)
+	FORCEINLINE_DEBUGGABLE FRHICommandSetupBasePassView(const FViewInfo& InView, bool InbShaderComplexity)
+		: ViewRect(InView.ViewRect)
+		, bShaderComplexity(InbShaderComplexity)
 	{
 	}
-
-	FORCEINLINE TStatId GetStatId() const
-	{
-		RETURN_QUICK_DECLARE_CYCLE_STAT(FSubmitCommandlistSetupBasePassView, STATGROUP_TaskGraphTasks);
-	}
-
-	ENamedThreads::Type GetDesiredThread()
-	{
-		return ENamedThreads::RenderThread_Local;
-	}
-
-	static ESubsequentsMode::Type GetSubsequentsMode() { return ESubsequentsMode::TrackSubsequents; }
-
-	void DoTask(ENamedThreads::Type CurrentThread, const FGraphEventRef& MyCompletionGraphEvent)
+	void Execute()
 	{
 		FRHICommandList RHICmdList;
-		SetupBasePassView(RHICmdList, View, bShaderComplexity);
+		SetupBasePassView(RHICmdList, ViewRect, bShaderComplexity);
 	}
 };
 
-FGraphEventRef FDeferredShadingSceneRenderer::RenderBasePassViewParallel(FRHICommandListImmediate& RHICmdList, FViewInfo& View, int32 Width, FGraphEventRef SubmitChain, bool& OutDirty)
+void FDeferredShadingSceneRenderer::RenderBasePassViewParallel(FViewInfo& View, int32 Width, FRHICommandList& ParentCmdList, bool& OutDirty)
 {
-	{
-		//optimization: this is probably only needed on the first view
-		FGraphEventArray Prereqs;
-		if (SubmitChain.GetReference())
-		{
-			Prereqs.Add(SubmitChain);
-		}
-		SubmitChain = TGraphTask<FSubmitCommandlistSetupBasePassView>::CreateTask(&Prereqs, ENamedThreads::RenderThread).ConstructAndDispatchWhenReady(&View, !!ViewFamily.EngineShowFlags.ShaderComplexity);
-
-	}			
+	// this is probably only needed on the first view
+	new (ParentCmdList.AllocCommand<FRHICommandSetupBasePassView>()) FRHICommandSetupBasePassView(View, !!ViewFamily.EngineShowFlags.ShaderComplexity);
 
-	SubmitChain = RenderBasePassStaticDataParallel(View, Width, SubmitChain, OutDirty);
-	SubmitChain = RenderBasePassDynamicDataParallel(View, SubmitChain, OutDirty);
-
-	return SubmitChain;
+	RenderBasePassStaticDataParallel(View, Width, ParentCmdList, OutDirty);
+	RenderBasePassDynamicDataParallel(View, ParentCmdList, OutDirty);
 }
 
 bool FDeferredShadingSceneRenderer::RenderBasePassView(FRHICommandListImmediate& RHICmdList, FViewInfo& View)
 {
 	bool bDirty = false; 
 
-	SetupBasePassView(RHICmdList, View, ViewFamily.EngineShowFlags.ShaderComplexity);
+	SetupBasePassView(RHICmdList, View.ViewRect, ViewFamily.EngineShowFlags.ShaderComplexity);
 	bDirty |= RenderBasePassStaticData(RHICmdList, View);
 	RenderBasePassDynamicData(RHICmdList, View, bDirty);
 
@@ -926,7 +895,6 @@ void FDeferredShadingSceneRenderer::Render(FRHICommandListImmediate& RHICmdList)
 			RenderDistortion(RHICmdList);
 		}
 	}
-
 	if (ViewFamily.EngineShowFlags.LightShafts)
 	{
 		RenderLightShaftBloom(RHICmdList);
@@ -1009,20 +977,20 @@ bool FDeferredShadingSceneRenderer::RenderPrePassViewDynamic(FRHICommandList& RH
 	return Drawer.IsDirty();
 }
 
-static void SetupPrePassView(FRHICommandList& RHICmdList, const FViewInfo& View)
+static void SetupPrePassView(FRHICommandList& RHICmdList, const FIntRect& ViewRect)
 {
 	// Disable color writes, enable depth tests and writes.
 	RHICmdList.SetBlendState(TStaticBlendState<CW_NONE>::GetRHI());
 	// Note, this is a reversed Z depth surface, using CF_GreaterEqual.
 	RHICmdList.SetDepthStencilState(TStaticDepthStencilState<true,CF_GreaterEqual>::GetRHI());
-	RHICmdList.SetViewport(View.ViewRect.Min.X, View.ViewRect.Min.Y, 0, View.ViewRect.Max.X, View.ViewRect.Max.Y, 1);
+	RHICmdList.SetViewport(ViewRect.Min.X, ViewRect.Min.Y, 0, ViewRect.Max.X, ViewRect.Max.Y, 1);
 }
 
 bool FDeferredShadingSceneRenderer::RenderPrePassView(FRHICommandList& RHICmdList, const FViewInfo& View)
 {
 	bool bDirty = false;
 
-	SetupPrePassView(RHICmdList, View);
+	SetupPrePassView(RHICmdList, View.ViewRect);
 
 	// Draw the static occluder primitives using a depth drawing policy.
 	{
@@ -1049,35 +1017,21 @@ bool FDeferredShadingSceneRenderer::RenderPrePassView(FRHICommandList& RHICmdLis
 	return bDirty;
 }
 
-class FSubmitCommandlistSetupPrepassView
+struct FRHICommandSetupPrepassView : public FRHICommand<FRHICommandSetupPrepassView>
 {
-	const FViewInfo& View;
-public:
-
-	FSubmitCommandlistSetupPrepassView(const FViewInfo* InView)
-		: View(*InView)
-	{
-	}
-
-	FORCEINLINE TStatId GetStatId() const
-	{
-		RETURN_QUICK_DECLARE_CYCLE_STAT(FSubmitCommandlistSetupPrepassView, STATGROUP_TaskGraphTasks);
-	}
-
-	ENamedThreads::Type GetDesiredThread()
+	FIntRect ViewRect;
+	FORCEINLINE_DEBUGGABLE FRHICommandSetupPrepassView(const FViewInfo& InView)
+		: ViewRect(InView.ViewRect)
 	{
-		return ENamedThreads::RenderThread_Local;
 	}
-
-	static ESubsequentsMode::Type GetSubsequentsMode() { return ESubsequentsMode::TrackSubsequents; }
-
-	void DoTask(ENamedThreads::Type CurrentThread, const FGraphEventRef& MyCompletionGraphEvent)
+	void Execute()
 	{
 		FRHICommandList RHICmdList;
-		SetupPrePassView(RHICmdList, View);
+		SetupPrePassView(RHICmdList, ViewRect);
 	}
 };
 
+
 class FRenderPrepassDynamicDataThreadTask
 {
 	FDeferredShadingSceneRenderer& ThisRenderer;
@@ -1122,18 +1076,10 @@ public:
 };
 
 
-FGraphEventRef FDeferredShadingSceneRenderer::RenderPrePassViewParallel(const FViewInfo& View, int32 Width, FGraphEventRef SubmitChain, bool& OutDirty)
+void FDeferredShadingSceneRenderer::RenderPrePassViewParallel(const FViewInfo& View, int32 Width, FRHICommandList& ParentCmdList, bool& OutDirty)
 {
-	{
-		//optimization: this is probably only needed on the first view
-		FGraphEventArray Prereqs;
-		if (SubmitChain.GetReference())
-		{
-			Prereqs.Add(SubmitChain);
-		}
-		SubmitChain = TGraphTask<FSubmitCommandlistSetupPrepassView>::CreateTask(&Prereqs, ENamedThreads::RenderThread).ConstructAndDispatchWhenReady(&View);
-
-	}
+	// this is probably only needed on the first view
+	new (ParentCmdList.AllocCommand<FRHICommandSetupPrepassView>()) FRHICommandSetupPrepassView(View);
 
 	// Draw the static occluder primitives using a depth drawing policy.
 	{
@@ -1141,31 +1087,28 @@ FGraphEventRef FDeferredShadingSceneRenderer::RenderPrePassViewParallel(const FV
 		// vertex buffer to minimize vertex fetch bandwidth, which is
 		// often the bottleneck during the depth only pass.
 		//SCOPED_DRAW_EVENT(PosOnlyOpaque, DEC_SCENE_ITEMS);
-		SubmitChain = Scene->PositionOnlyDepthDrawList.DrawVisibleParallel(View, View.StaticMeshOccluderMap, View.StaticMeshBatchVisibility, Width, SubmitChain, OutDirty);
+		Scene->PositionOnlyDepthDrawList.DrawVisibleParallel(View, View.StaticMeshOccluderMap, View.StaticMeshBatchVisibility, Width, ParentCmdList, OutDirty);
 	}
 	{
 		// Draw opaque occluders, using double speed z where supported.
 		//SCOPED_DRAW_EVENT(Opaque, DEC_SCENE_ITEMS);
-		SubmitChain = Scene->DepthDrawList.DrawVisibleParallel(View,View.StaticMeshOccluderMap,View.StaticMeshBatchVisibility, Width, SubmitChain, OutDirty);
+		Scene->DepthDrawList.DrawVisibleParallel(View,View.StaticMeshOccluderMap,View.StaticMeshBatchVisibility, Width, ParentCmdList, OutDirty);
 	}
 
 	if(EarlyZPassMode >= DDM_AllOccluders)
 	{
 		// Draw opaque occluders with masked materials
 		//SCOPED_DRAW_EVENT(Opaque, DEC_SCENE_ITEMS);
-		SubmitChain = Scene->MaskedDepthDrawList.DrawVisibleParallel(View,View.StaticMeshOccluderMap,View.StaticMeshBatchVisibility, Width, SubmitChain, OutDirty);
+		Scene->MaskedDepthDrawList.DrawVisibleParallel(View,View.StaticMeshOccluderMap,View.StaticMeshBatchVisibility, Width, ParentCmdList, OutDirty);
 	}
-
 	{
 		FRHICommandList* CmdList = new FRHICommandList;
 
 		FGraphEventRef AnyThreadCompletionEvent = TGraphTask<FRenderPrepassDynamicDataThreadTask>::CreateTask(nullptr, ENamedThreads::RenderThread)
 			.ConstructAndDispatchWhenReady(this, CmdList, &View, &OutDirty);
 
-		SubmitChain = FSubmitCommandlistThreadTask::AddToChain(AnyThreadCompletionEvent, SubmitChain, CmdList);
+		ParentCmdList.QueueAsyncCommandListSubmit(AnyThreadCompletionEvent, CmdList);
 	}
-
-	return SubmitChain;
 }
 
 
@@ -1187,21 +1130,16 @@ bool FDeferredShadingSceneRenderer::RenderPrePass(FRHICommandListImmediate& RHIC
 	// Draw a depth pass to avoid overdraw in the other passes.
 	if(EarlyZPassMode != DDM_None)
 	{
-		if (!GRHICommandList.Bypass())
+		if (GRHICommandList.UseParallelAlgorithms())
 		{
+			FScopedCommandListFlush Flusher(RHICmdList);
 			int32 Width = CVarRHICmdWidth.GetValueOnRenderThread(); // we use a few more than needed to cover non-equal jobs
 
-			FGraphEventRef SubmitChain;
 			for(int32 ViewIndex = 0;ViewIndex < Views.Num();ViewIndex++)
 			{
 				SCOPED_CONDITIONAL_DRAW_EVENTF(EventView, Views.Num() > 1, DEC_SCENE_ITEMS, TEXT("View%d"), ViewIndex);
 				const FViewInfo& View = Views[ViewIndex];
-				SubmitChain = RenderPrePassViewParallel(View, Width, SubmitChain, bDirty);
-			}
-			if (SubmitChain.GetReference())
-			{
-				RHICmdList.Flush(); // this would happen later anyway, but we might as well do these while we wait for async tasks
-				FTaskGraphInterface::Get().WaitUntilTaskCompletes(SubmitChain, ENamedThreads::RenderThread_Local);
+				RenderPrePassViewParallel(View, Width, RHICmdList, bDirty);
 			}
 		}
 		else
@@ -1238,22 +1176,16 @@ bool FDeferredShadingSceneRenderer::RenderBasePass(FRHICommandListImmediate& RHI
 		SCOPED_DRAW_EVENT(BasePass, DEC_SCENE_ITEMS);
 		SCOPE_CYCLE_COUNTER(STAT_BasePassDrawTime);
 
-		if (!GRHICommandList.Bypass())
+		if (GRHICommandList.UseParallelAlgorithms())
 		{
+			FScopedCommandListFlush Flusher(RHICmdList);
 			int32 Width = CVarRHICmdWidth.GetValueOnRenderThread(); // we use a few more than needed to cover non-equal jobs
-			FGraphEventRef SubmitChain;
 
 			for(int32 ViewIndex = 0;ViewIndex < Views.Num();ViewIndex++)
 			{
 				SCOPED_CONDITIONAL_DRAW_EVENTF(EventView, Views.Num() > 1, DEC_SCENE_ITEMS, TEXT("View%d"), ViewIndex);
 				FViewInfo& View = Views[ViewIndex];
-				SubmitChain = RenderBasePassViewParallel(RHICmdList, View, Width, SubmitChain, bDirty);
-			}
-
-			if (SubmitChain.GetReference())
-			{
-				RHICmdList.Flush(); // this would happen later anyway, but we might as well do these while we wait for async tasks
-				FTaskGraphInterface::Get().WaitUntilTaskCompletes(SubmitChain, ENamedThreads::RenderThread_Local);
+				RenderBasePassViewParallel(View, Width, RHICmdList, bDirty);
 			}
 		}
 		else
diff --git a/Engine/Source/Runtime/Renderer/Private/DeferredShadingRenderer.h b/Engine/Source/Runtime/Renderer/Private/DeferredShadingRenderer.h
--- a/Engine/Source/Runtime/Renderer/Private/DeferredShadingRenderer.h
+++ b/Engine/Source/Runtime/Renderer/Private/DeferredShadingRenderer.h
@@ -64,9 +64,8 @@ public:
 
 	/**
 	 * Renders the scene's prepass for a particular view in parallel
-	 * @return the submit chain
 	 */
-	FGraphEventRef RenderPrePassViewParallel(const FViewInfo& View, int32 Width, FGraphEventRef SubmitChain, bool& OutDirty);
+	void RenderPrePassViewParallel(const FViewInfo& View, int32 Width, FRHICommandList& ParentCmdList, bool& OutDirty);
 
     /** Renders the basepass for the static data of a given View. */
     bool RenderBasePassStaticData(FRHICommandList& RHICmdList, FViewInfo& View);
@@ -74,9 +73,9 @@ public:
 	bool RenderBasePassStaticDataDefault(FRHICommandList& RHICmdList, FViewInfo& View);
 
 	/** Renders the basepass for the static data of a given View. Parallel versions.*/
-	FGraphEventRef RenderBasePassStaticDataParallel(FViewInfo& View, int32 Width, FGraphEventRef SubmitChain, bool& OutDirty);
-	FGraphEventRef RenderBasePassStaticDataMaskedParallel(FViewInfo& View, int32 Width, FGraphEventRef SubmitChain, bool& OutDirty);
-	FGraphEventRef RenderBasePassStaticDataDefaultParallel(FViewInfo& View, int32 Width, FGraphEventRef SubmitChain, bool& OutDirty);
+	void RenderBasePassStaticDataParallel(FViewInfo& View, int32 Width, FRHICommandList& ParentCmdList, bool& OutDirty);
+	void RenderBasePassStaticDataMaskedParallel(FViewInfo& View, int32 Width, FRHICommandList& ParentCmdList, bool& OutDirty);
+	void RenderBasePassStaticDataDefaultParallel(FViewInfo& View, int32 Width, FRHICommandList& ParentCmdList, bool& OutDirty);
 
 	/** Sorts base pass draw lists front to back for improved GPU culling. */
 	void SortBasePassStaticData(FVector ViewPosition);
@@ -85,10 +84,10 @@ public:
 	void RenderBasePassDynamicData(FRHICommandList& RHICmdList, const FViewInfo& View, bool& bOutDirty);
 
 	/** Renders the basepass for the dynamic data of a given View, in parallel. */
-	FGraphEventRef RenderBasePassDynamicDataParallel(FViewInfo& View, FGraphEventRef SubmitChain, bool& bOutDirty);
+	void RenderBasePassDynamicDataParallel(FViewInfo& View, FRHICommandList& ParentCmdList, bool& bOutDirty);
 
     /** Renders the basepass for a given View, in parallel */
-	FGraphEventRef RenderBasePassViewParallel(FRHICommandListImmediate& RHICmdList, FViewInfo& View, int32 Width, FGraphEventRef SubmitChain, bool& OutDirty);
+	void RenderBasePassViewParallel(FViewInfo& View, int32 Width, FRHICommandList& ParentCmdList, bool& OutDirty);
 
 	/** Renders the basepass for a given View. */
 	bool RenderBasePassView(FRHICommandListImmediate& RHICmdList, FViewInfo& View);
@@ -252,9 +251,9 @@ private:
 	/** Renders the velocities of movable objects for the motion blur effect. */
 	void RenderVelocities(FRHICommandListImmediate& RHICmdList, const FViewInfo& View, TRefCountPtr<IPooledRenderTarget>& VelocityRT, bool bLastFrame);
 
-	/** Renders the velocities of movable objects for the motion blur effect. */
+	/** Renders the velocities for a subset of movable objects for the motion blur effect. */
 	friend class FRenderVelocityDynamicThreadTask;
-	void RenderDynamicVelocities(FRHICommandList& RHICmdList, const FViewInfo& View);
+	void RenderDynamicVelocitiesInner(FRHICommandList& RHICmdList, const FViewInfo& View, int32 FirstIndex, int32 LastIndex);
 
 	/** Renders the velocities of movable objects for the motion blur effect. */
 	void RenderVelocitiesInner(FRHICommandListImmediate& RHICmdList, const FViewInfo& View);
diff --git a/Engine/Source/Runtime/Renderer/Private/DynamicPrimitiveDrawing.h b/Engine/Source/Runtime/Renderer/Private/DynamicPrimitiveDrawing.h
--- a/Engine/Source/Runtime/Renderer/Private/DynamicPrimitiveDrawing.h
+++ b/Engine/Source/Runtime/Renderer/Private/DynamicPrimitiveDrawing.h
@@ -191,16 +191,16 @@ bool DrawViewElements(
 * @param DrawingContext - The drawing policy type specific context for the drawing.
 * @param DPGIndex World or Foreground DPG index for draw order
 * @param bPreFog - true if the draw call is occurring before fog has been rendered.
-* @param SubmitChain - parallel submit chain
+* @param ParentCmdList - cmdlist to put the wait and execute task on
 * @param Width - parallel width
 */
 template<class DrawingPolicyFactoryType>
-FGraphEventRef DrawViewElementsParallel(
+void DrawViewElementsParallel(
 	const FViewInfo& View,
 	const typename DrawingPolicyFactoryType::ContextType& DrawingContext,
 	uint8 DPGIndex,
 	bool bPreFog,
-	FGraphEventRef SubmitChain,
+	FRHICommandList& ParentCmdList,
 	int32 Width
 	);
 
diff --git a/Engine/Source/Runtime/Renderer/Private/DynamicPrimitiveDrawing.inl b/Engine/Source/Runtime/Renderer/Private/DynamicPrimitiveDrawing.inl
--- a/Engine/Source/Runtime/Renderer/Private/DynamicPrimitiveDrawing.inl
+++ b/Engine/Source/Runtime/Renderer/Private/DynamicPrimitiveDrawing.inl
@@ -305,12 +305,12 @@ public:
 };
 
 template<class DrawingPolicyFactoryType>
-FGraphEventRef DrawViewElementsParallel(
+void DrawViewElementsParallel(
 	const FViewInfo& View,
 	const typename DrawingPolicyFactoryType::ContextType& DrawingContext,
 	uint8 DPGIndex,
 	bool bPreFog,
-	FGraphEventRef SubmitChain,
+	FRHICommandList& ParentCmdList,
 	int32 Width
 	)
 {
@@ -334,18 +334,18 @@ FGraphEventRef DrawViewElementsParallel(
 				int32 Last = Start + (NumPer - 1) + (ThreadIndex < Extra);
 				check(Last >= Start);
 
-				FRHICommandList* CmdList = new FRHICommandList;
+				{
+					FRHICommandList* CmdList = new FRHICommandList;
 
-				FGraphEventRef AnyThreadCompletionEvent = TGraphTask<FDrawViewElementsAnyThreadTask<DrawingPolicyFactoryType> >::CreateTask(nullptr, ENamedThreads::RenderThread)
-					.ConstructAndDispatchWhenReady(CmdList, &View, DrawingContext, DPGIndex, bPreFog, Start, Last);
-				Start = Last + 1;
+					FGraphEventRef AnyThreadCompletionEvent = TGraphTask<FDrawViewElementsAnyThreadTask<DrawingPolicyFactoryType> >::CreateTask(nullptr, ENamedThreads::RenderThread)
+						.ConstructAndDispatchWhenReady(CmdList, &View, DrawingContext, DPGIndex, bPreFog, Start, Last);
 
-				SubmitChain = FSubmitCommandlistThreadTask::AddToChain(AnyThreadCompletionEvent, SubmitChain, CmdList);
+					ParentCmdList.QueueAsyncCommandListSubmit(AnyThreadCompletionEvent, CmdList);
+				}
+				Start = Last + 1;
 			}
 		}
 	}
-
-	return SubmitChain;
 }
 
 template<class DrawingPolicyFactoryType>
diff --git a/Engine/Source/Runtime/Renderer/Private/SceneRendering.cpp b/Engine/Source/Runtime/Renderer/Private/SceneRendering.cpp
--- a/Engine/Source/Runtime/Renderer/Private/SceneRendering.cpp
+++ b/Engine/Source/Runtime/Renderer/Private/SceneRendering.cpp
@@ -976,6 +976,7 @@ static void RenderViewFamily_RenderThread(FRHICommandListImmediate& RHICmdList,
 
         // Delete the scene renderer.
 		delete SceneRenderer;
+		FRHICommandListExecutor::GetImmediateCommandList().Flush(); // we want to make sure this all gets to the GPU this frame and doesn't hang around
 	}
 
 #if STATS
diff --git a/Engine/Source/Runtime/Renderer/Private/ShadowRendering.cpp b/Engine/Source/Runtime/Renderer/Private/ShadowRendering.cpp
--- a/Engine/Source/Runtime/Renderer/Private/ShadowRendering.cpp
+++ b/Engine/Source/Runtime/Renderer/Private/ShadowRendering.cpp
@@ -1392,25 +1392,24 @@ public:
 	}
 };
 
+static TAutoConsoleVariable<int32> CVarParallelShadows(
+	TEXT("r.ParallelShadows"),
+	1,
+	TEXT("Toggles parallel shadow rendering. Parallel rendering must be enabled for this to have an effect.\n"),
+	ECVF_RenderThreadSafe
+	);
+
 void FProjectedShadowInfo::RenderDepthInner(FRHICommandList& RHICmdList, FDeferredShadingSceneRenderer* SceneRenderer, const FViewInfo* FoundView)
 {
-	static TAutoConsoleVariable<int32> CVarParallelShadows(
-		TEXT("r.ParallelShadows"),
-		1,
-		TEXT("Toggles parallel shadow rendering. Parallel rendering must be enabled for this to have an effect.\n"),
-		ECVF_RenderThreadSafe
-		);
-
-	if (&RHICmdList == &FRHICommandListExecutor::GetImmediateCommandList() &&  // translucent shadows are draw on the render thread, using a recursive cmdlist, so this path seems like a loser in that case, but it should be tried
-		!GRHICommandList.Bypass() && CVarParallelShadows.GetValueOnRenderThread())
+	if (RHICmdList.IsImmediate() &&  // translucent shadows are draw on the render thread, using a recursive cmdlist
+		GRHICommandList.UseParallelAlgorithms() && CVarParallelShadows.GetValueOnRenderThread())
 	{
 		// parallel version
+		FScopedCommandListFlush Flusher;
 
 		int32 Width = CVarRHICmdWidth.GetValueOnRenderThread(); // we use a few more than needed to cover non-equal jobs
 		bool OutDirty = false; // unused
 
-		FGraphEventRef SubmitChain;
-
 		// Draw the subject's static elements using static draw lists
 		if (IsWholeSceneDirectionalShadow())
 		{
@@ -1418,12 +1417,12 @@ void FProjectedShadowInfo::RenderDepthInner(FRHICommandList& RHICmdList, FDeferr
 
 			if (bReflectiveShadowmap)
 			{
-				SubmitChain = SceneRenderer->Scene->WholeSceneReflectiveShadowMapDrawList.DrawVisibleParallel(*FoundView, StaticMeshWholeSceneShadowDepthMap, StaticMeshWholeSceneShadowBatchVisibility, Width, SubmitChain, OutDirty);
+				SceneRenderer->Scene->WholeSceneReflectiveShadowMapDrawList.DrawVisibleParallel(*FoundView, StaticMeshWholeSceneShadowDepthMap, StaticMeshWholeSceneShadowBatchVisibility, Width, RHICmdList, OutDirty);
 			}
 			else
 			{
 				// Use the scene's shadow depth draw list with this shadow's visibility map
-				SubmitChain = SceneRenderer->Scene->WholeSceneShadowDepthDrawList.DrawVisibleParallel(*FoundView, StaticMeshWholeSceneShadowDepthMap, StaticMeshWholeSceneShadowBatchVisibility, Width, SubmitChain, OutDirty);
+				SceneRenderer->Scene->WholeSceneShadowDepthDrawList.DrawVisibleParallel(*FoundView, StaticMeshWholeSceneShadowDepthMap, StaticMeshWholeSceneShadowBatchVisibility, Width, RHICmdList, OutDirty);
 			}
 		}
 		// Draw the subject's static elements using manual state filtering
@@ -1434,7 +1433,7 @@ void FProjectedShadowInfo::RenderDepthInner(FRHICommandList& RHICmdList, FDeferr
 			FGraphEventRef AnyThreadCompletionEvent = TGraphTask<FDrawShadowMeshElementsThreadTask>::CreateTask(nullptr, ENamedThreads::RenderThread)
 				.ConstructAndDispatchWhenReady(this, CmdList, FoundView, bReflectiveShadowmap && !bOnePassPointLightShadow);
 
-			SubmitChain = FSubmitCommandlistThreadTask::AddToChain(AnyThreadCompletionEvent, SubmitChain, CmdList);
+			RHICmdList.QueueAsyncCommandListSubmit(AnyThreadCompletionEvent, CmdList);
 		}
 		{
 			FRHICommandList* CmdList = new FRHICommandList;
@@ -1442,13 +1441,7 @@ void FProjectedShadowInfo::RenderDepthInner(FRHICommandList& RHICmdList, FDeferr
 			FGraphEventRef AnyThreadCompletionEvent = TGraphTask<FRenderDepthDynamicThreadTask>::CreateTask(nullptr, ENamedThreads::RenderThread)
 				.ConstructAndDispatchWhenReady(this, CmdList, FoundView, SceneRenderer);
 
-			SubmitChain = FSubmitCommandlistThreadTask::AddToChain(AnyThreadCompletionEvent, SubmitChain, CmdList);
-		}
-
-		if (SubmitChain.GetReference())
-		{
-			FRHICommandListExecutor::GetImmediateCommandList().Flush(); // this would happen later anyway, but we might as well do these while we wait for async tasks
-			FTaskGraphInterface::Get().WaitUntilTaskCompletes(SubmitChain, ENamedThreads::RenderThread_Local);
+			RHICmdList.QueueAsyncCommandListSubmit(AnyThreadCompletionEvent, CmdList);
 		}
 	}
 	else
diff --git a/Engine/Source/Runtime/Renderer/Private/StaticMeshDrawList.h b/Engine/Source/Runtime/Renderer/Private/StaticMeshDrawList.h
--- a/Engine/Source/Runtime/Renderer/Private/StaticMeshDrawList.h
+++ b/Engine/Source/Runtime/Renderer/Private/StaticMeshDrawList.h
@@ -292,9 +292,11 @@ public:
 	* @param View - The view of the meshes to render.
 	* @param StaticMeshVisibilityMap - An map from FStaticMesh::Id to visibility state.
 	* @param BatchVisibilityArray - An array of batch element visibility bitmasks.
-	* @return True if any static meshes were drawn.
+	* @param Width - parallel width
+	* @param ParentCmdList - cmdlist to put the wait and execute task on
+	* @param OutDirty - set to true if anything is drawn
 	*/
-	FGraphEventRef DrawVisibleParallel(const FViewInfo& View, const TBitArray<SceneRenderingBitArrayAllocator>& StaticMeshVisibilityMap, const TArray<uint64, SceneRenderingAllocator>& BatchVisibilityArray, int32 Width, FGraphEventRef SubmitChain, bool& OutDirty);
+	void DrawVisibleParallel(const FViewInfo& View, const TBitArray<SceneRenderingBitArrayAllocator>& StaticMeshVisibilityMap, const TArray<uint64, SceneRenderingAllocator>& BatchVisibilityArray, int32 Width, FRHICommandList& ParentCmdList, bool& OutDirty);
 
 	/**
 	 * Draws only the static meshes which are in the visibility map, sorted front-to-back.
diff --git a/Engine/Source/Runtime/Renderer/Private/StaticMeshDrawList.inl b/Engine/Source/Runtime/Renderer/Private/StaticMeshDrawList.inl
--- a/Engine/Source/Runtime/Renderer/Private/StaticMeshDrawList.inl
+++ b/Engine/Source/Runtime/Renderer/Private/StaticMeshDrawList.inl
@@ -344,52 +344,12 @@ public:
 	}
 };
 
-class FSubmitCommandlistThreadTask
-{
-	FRHICommandList* RHICmdList;
-public:
-
-	FSubmitCommandlistThreadTask(FRHICommandList* InRHICmdList)
-		: RHICmdList(InRHICmdList)
-	{
-	}
-
-	FORCEINLINE TStatId GetStatId() const
-	{
-		RETURN_QUICK_DECLARE_CYCLE_STAT(FDrawVisibleRenderThreadTask, STATGROUP_TaskGraphTasks);
-	}
-
-	ENamedThreads::Type GetDesiredThread()
-	{
-		return ENamedThreads::RenderThread_Local;
-	}
-
-	static ESubsequentsMode::Type GetSubsequentsMode() { return ESubsequentsMode::TrackSubsequents; }
-
-	void DoTask(ENamedThreads::Type CurrentThread, const FGraphEventRef& MyCompletionGraphEvent)
-	{
-		delete RHICmdList;
-	}
-
-	/*** Helper function set up a submit chain of the these   */
-	static FGraphEventRef AddToChain(FGraphEventRef AnyThreadCompletionEvent, FGraphEventRef SubmitChain, FRHICommandList* CmdList)
-	{
-		FGraphEventArray Prereqs;
-		Prereqs.Add(AnyThreadCompletionEvent);
-		if (SubmitChain.GetReference())
-		{
-			Prereqs.Add(SubmitChain);
-		}
-		return TGraphTask<FSubmitCommandlistThreadTask>::CreateTask(&Prereqs, ENamedThreads::RenderThread).ConstructAndDispatchWhenReady(CmdList);
-	}
-};
-
 template<typename DrawingPolicyType>
-FGraphEventRef TStaticMeshDrawList<DrawingPolicyType>::DrawVisibleParallel(
+void TStaticMeshDrawList<DrawingPolicyType>::DrawVisibleParallel(
 	const FViewInfo& View,
 	const TBitArray<SceneRenderingBitArrayAllocator>& StaticMeshVisibilityMap,
 	const TArray<uint64, SceneRenderingAllocator>& BatchVisibilityArray,
-	int32 Width, FGraphEventRef SubmitChain, bool& OutDirty
+	int32 Width, FRHICommandList& ParentCmdList, bool& OutDirty
 	)
 {
 
@@ -408,18 +368,19 @@ FGraphEventRef TStaticMeshDrawList<DrawingPolicyType>::DrawVisibleParallel(
 			int32 Last = Start + (NumPer - 1) + (ThreadIndex < Extra);
 			check(Last >= Start);
 
-			FRHICommandList* CmdList = new FRHICommandList;
+			{
+				FRHICommandList* CmdList = new FRHICommandList;
 
-			FGraphEventRef AnyThreadCompletionEvent = TGraphTask<FDrawVisibleAnyThreadTask<DrawingPolicyType> >::CreateTask(nullptr, ENamedThreads::RenderThread)
-				.ConstructAndDispatchWhenReady(this, CmdList, &View, &StaticMeshVisibilityMap, &BatchVisibilityArray, Start, Last, &OutDirty);
-			Start = Last + 1;
+				FGraphEventRef AnyThreadCompletionEvent = TGraphTask<FDrawVisibleAnyThreadTask<DrawingPolicyType> >::CreateTask(nullptr, ENamedThreads::RenderThread)
+					.ConstructAndDispatchWhenReady(this, CmdList, &View, &StaticMeshVisibilityMap, &BatchVisibilityArray, Start, Last, &OutDirty);
 
-			SubmitChain = FSubmitCommandlistThreadTask::AddToChain(AnyThreadCompletionEvent, SubmitChain, CmdList);
+				ParentCmdList.QueueAsyncCommandListSubmit(AnyThreadCompletionEvent, CmdList);
+			}
+
+			Start = Last + 1;
 		}
 	}
 	check(Start == OrderedDrawingPolicies.Num());
-
-	return SubmitChain;
 }
 
 template<typename DrawingPolicyType>
diff --git a/Engine/Source/Runtime/Renderer/Private/TranslucentRendering.cpp b/Engine/Source/Runtime/Renderer/Private/TranslucentRendering.cpp
--- a/Engine/Source/Runtime/Renderer/Private/TranslucentRendering.cpp
+++ b/Engine/Source/Runtime/Renderer/Private/TranslucentRendering.cpp
@@ -839,9 +839,10 @@ void FDeferredShadingSceneRenderer::RenderTranslucencyParallel(FRHICommandListIm
 	GSceneRenderTargets.AllocLightAttenuation(); // materials will attempt to get this texture before the deferred command to set it up executes
 	int32 Width = CVarRHICmdWidth.GetValueOnRenderThread(); // we use a few more than needed to cover non-equal jobs
 
+	FScopedCommandListFlush Flusher(RHICmdList);
 	FTranslucencyDrawingPolicyFactory::ContextType ThisContext;
 
-	FGraphEventRef SubmitChain;
+	check(IsInRenderingThread());
 
 	for (int32 ViewIndex = 0; ViewIndex < Views.Num(); ViewIndex++)
 	{
@@ -866,13 +867,15 @@ void FDeferredShadingSceneRenderer::RenderTranslucencyParallel(FRHICommandListIm
 					int32 Last = Start + (NumPer - 1) + (ThreadIndex < Extra);
 					check(Last >= Start);
 
-					FRHICommandList* CmdList = new FRHICommandList;
+					{
+						FRHICommandList* CmdList = new FRHICommandList;
 
-					FGraphEventRef AnyThreadCompletionEvent = TGraphTask<FDrawSortedTransAnyThreadTask>::CreateTask(nullptr, ENamedThreads::RenderThread)
-						.ConstructAndDispatchWhenReady(this, CmdList, &View, false, Start, Last);
-					Start = Last + 1;
+						FGraphEventRef AnyThreadCompletionEvent = TGraphTask<FDrawSortedTransAnyThreadTask>::CreateTask(nullptr, ENamedThreads::RenderThread)
+							.ConstructAndDispatchWhenReady(this, CmdList, &View, false, Start, Last);
 
-					SubmitChain = FSubmitCommandlistThreadTask::AddToChain(AnyThreadCompletionEvent, SubmitChain, CmdList);
+						RHICmdList.QueueAsyncCommandListSubmit(AnyThreadCompletionEvent, CmdList); 
+					}
+					Start = Last + 1;
 				}
 			}
 			else
@@ -883,14 +886,13 @@ void FDeferredShadingSceneRenderer::RenderTranslucencyParallel(FRHICommandListIm
 				FGraphEventRef AnyThreadCompletionEvent = TGraphTask<FDrawSortedTransAnyThreadTask>::CreateTask(nullptr, ENamedThreads::RenderThread)
 					.ConstructAndDispatchWhenReady(this, CmdList, &View, false, 0, -1);
 
-				SubmitChain = FSubmitCommandlistThreadTask::AddToChain(AnyThreadCompletionEvent, SubmitChain, CmdList);
+				RHICmdList.QueueAsyncCommandListSubmit(AnyThreadCompletionEvent, CmdList);
 			}
 		}
-
 		// Draw the view's mesh elements with the translucent drawing policy.
-		SubmitChain = DrawViewElementsParallel<FTranslucencyDrawingPolicyFactory>(View, ThisContext, SDPG_World, false, SubmitChain, Width);
+		DrawViewElementsParallel<FTranslucencyDrawingPolicyFactory>(View, ThisContext, SDPG_World, false, RHICmdList, Width);
 		// Draw the view's mesh elements with the translucent drawing policy.
-		SubmitChain = DrawViewElementsParallel<FTranslucencyDrawingPolicyFactory>(View, ThisContext, SDPG_Foreground, false, SubmitChain, Width);
+		DrawViewElementsParallel<FTranslucencyDrawingPolicyFactory>(View, ThisContext, SDPG_Foreground, false, RHICmdList, Width);
 
 #if 0 // unsupported visualization in the parallel case
 		const FSceneViewState* ViewState = (const FSceneViewState*)View.State;
@@ -921,34 +923,35 @@ void FDeferredShadingSceneRenderer::RenderTranslucencyParallel(FRHICommandListIm
 					int32 Last = Start + (NumPer - 1) + (ThreadIndex < Extra);
 					check(Last >= Start);
 
-					FRHICommandList* CmdList = new FRHICommandList;
+					{
+						FRHICommandList* CmdList = new FRHICommandList;
 
-					FGraphEventRef AnyThreadCompletionEvent = TGraphTask<FDrawSortedTransAnyThreadTask>::CreateTask(nullptr, ENamedThreads::RenderThread)
-						.ConstructAndDispatchWhenReady(this, CmdList, &View, true, Start, Last);
-					Start = Last + 1;
+						FGraphEventRef AnyThreadCompletionEvent = TGraphTask<FDrawSortedTransAnyThreadTask>::CreateTask(nullptr, ENamedThreads::RenderThread)
+							.ConstructAndDispatchWhenReady(this, CmdList, &View, true, Start, Last);
 
-					SubmitChain = FSubmitCommandlistThreadTask::AddToChain(AnyThreadCompletionEvent, SubmitChain, CmdList);
+						RHICmdList.QueueAsyncCommandListSubmit(AnyThreadCompletionEvent, CmdList); 
+					}
+					Start = Last + 1;
 				}
 			}
 			else
 			{
 				// still need to setup the render target, a bit of a waste...
 				FRHICommandList* CmdList = new FRHICommandList;
-
 				FGraphEventRef AnyThreadCompletionEvent = TGraphTask<FDrawSortedTransAnyThreadTask>::CreateTask(nullptr, ENamedThreads::RenderThread)
 					.ConstructAndDispatchWhenReady(this, CmdList, &View, true, 0, -1);
-				SubmitChain = FSubmitCommandlistThreadTask::AddToChain(AnyThreadCompletionEvent, SubmitChain, CmdList);
+				RHICmdList.QueueAsyncCommandListSubmit(AnyThreadCompletionEvent, CmdList);
 			}
 		}
 	}
-
-	if (SubmitChain.GetReference())
-	{
-		RHICmdList.Flush(); // this would happen later anyway, but we might as well do these while we wait for async tasks
-		FTaskGraphInterface::Get().WaitUntilTaskCompletes(SubmitChain, ENamedThreads::RenderThread_Local);
-	}
 }
 
+static TAutoConsoleVariable<int32> CVarParallelTranslucency(
+	TEXT("r.ParallelTranslucency"),
+	1,
+	TEXT("Toggles parallel translucency rendering. Parallel rendering must be enabled for this to have an effect.\n"),
+	ECVF_RenderThreadSafe
+	);
 
 void FDeferredShadingSceneRenderer::RenderTranslucency(FRHICommandListImmediate& RHICmdList)
 {
@@ -957,7 +960,7 @@ void FDeferredShadingSceneRenderer::RenderTranslucency(FRHICommandListImmediate&
 
 		SCOPED_DRAW_EVENT(Translucency, DEC_SCENE_ITEMS);
 
-		if (!GRHICommandList.Bypass())
+		if (GRHICommandList.UseParallelAlgorithms() && CVarParallelTranslucency.GetValueOnRenderThread())
 		{
 			RenderTranslucencyParallel(RHICmdList);
 			return;
diff --git a/Engine/Source/Runtime/Renderer/Private/VelocityRendering.cpp b/Engine/Source/Runtime/Renderer/Private/VelocityRendering.cpp
--- a/Engine/Source/Runtime/Renderer/Private/VelocityRendering.cpp
+++ b/Engine/Source/Runtime/Renderer/Private/VelocityRendering.cpp
@@ -441,13 +441,13 @@ bool IsMotionBlurEnabled(const FViewInfo& View)
 		&& !(View.Family->Views.Num() > 1);
 }
 
-void FDeferredShadingSceneRenderer::RenderDynamicVelocities(FRHICommandList& RHICmdList, const FViewInfo& View)
+void FDeferredShadingSceneRenderer::RenderDynamicVelocitiesInner(FRHICommandList& RHICmdList, const FViewInfo& View, int32 FirstIndex, int32 LastIndex)
 {
 	// Draw velocities for movable dynamic meshes.
 	TDynamicPrimitiveDrawer<FVelocityDrawingPolicyFactory> Drawer(
 		RHICmdList, &View, FVelocityDrawingPolicyFactory::ContextType(DDM_AllOccluders), true, false, true
 		);
-	for (int32 PrimitiveIndex = 0; PrimitiveIndex < View.VisibleDynamicPrimitives.Num(); PrimitiveIndex++)
+	for (int32 PrimitiveIndex = FirstIndex; PrimitiveIndex <= LastIndex; PrimitiveIndex++)
 	{
 		const FPrimitiveSceneInfo* PrimitiveSceneInfo = View.VisibleDynamicPrimitives[PrimitiveIndex];
 
@@ -461,7 +461,6 @@ void FDeferredShadingSceneRenderer::RenderDynamicVelocities(FRHICommandList& RHI
 		Drawer.SetPrimitive(PrimitiveSceneInfo->Proxy);
 		PrimitiveSceneInfo->Proxy->DrawDynamicElements(&Drawer, &View);
 	}
-	Drawer.IsDirty();
 }
 
 class FRenderVelocityDynamicThreadTask
@@ -470,17 +469,22 @@ class FRenderVelocityDynamicThreadTask
 	FRHICommandList& RHICmdList;
 	const FViewInfo& View;
 	FDeferredShadingSceneRenderer* SceneRenderer;
+	int32 FirstIndex;
+	int32 LastIndex;
 
 public:
 
 	FRenderVelocityDynamicThreadTask(
 		FDeferredShadingSceneRenderer* InThisRenderer,
 		FRHICommandList* InRHICmdList,
-		const FViewInfo* InView
+		const FViewInfo* InView,
+		int32 InFirstIndex, int32 InLastIndex
 		)
 		: ThisRenderer(*InThisRenderer)
 		, RHICmdList(*InRHICmdList)
 		, View(*InView)
+		, FirstIndex(InFirstIndex)
+		, LastIndex(InLastIndex)
 	{
 	}
 
@@ -498,42 +502,58 @@ public:
 
 	void DoTask(ENamedThreads::Type CurrentThread, const FGraphEventRef& MyCompletionGraphEvent)
 	{
-		ThisRenderer.RenderDynamicVelocities(RHICmdList, View);
+		ThisRenderer.RenderDynamicVelocitiesInner(RHICmdList, View, FirstIndex, LastIndex);
 	}
 };
 
 
+static TAutoConsoleVariable<int32> CVarParallelVelocity(
+	TEXT("r.ParallelVelocity"),
+	0,  // not possible because of per-bone motion blur
+	TEXT("Toggles parallel velocity rendering. Parallel rendering must be enabled for this to have an effect.\n"),
+	ECVF_RenderThreadSafe
+	);
+
 void FDeferredShadingSceneRenderer::RenderVelocitiesInner(FRHICommandListImmediate& RHICmdList, const FViewInfo& View)
 {
-	static TAutoConsoleVariable<int32> CVarParallelVelocity(
-		TEXT("r.ParallelVelocity"),
-		1,
-		TEXT("Toggles parallel velocity rendering. Parallel rendering must be enabled for this to have an effect.\n"),
-		ECVF_RenderThreadSafe
-		);
-
-	if (!GRHICommandList.Bypass() && CVarParallelVelocity.GetValueOnRenderThread())
+	if (GRHICommandList.UseParallelAlgorithms() && CVarParallelVelocity.GetValueOnRenderThread())
 	{
 		// parallel version
+		FScopedCommandListFlush Flusher(RHICmdList);
 
 		int32 Width = CVarRHICmdWidth.GetValueOnRenderThread(); // we use a few more than needed to cover non-equal jobs
 		bool OutDirty = false; // unused
 
-		FGraphEventRef SubmitChain;
+		Scene->VelocityDrawList.DrawVisibleParallel(View, View.StaticMeshVelocityMap, View.StaticMeshBatchVisibility, Width, RHICmdList, OutDirty);
 
-		SubmitChain = Scene->VelocityDrawList.DrawVisibleParallel(View, View.StaticMeshVelocityMap, View.StaticMeshBatchVisibility, Width, SubmitChain, OutDirty);
 		{
-			FRHICommandList* CmdList = new FRHICommandList;
+			int32 NumPrims = View.VisibleDynamicPrimitives.Num();
+			int32 EffectiveThreads = FMath::Min<int32>(NumPrims, Width);
 
-			FGraphEventRef AnyThreadCompletionEvent = TGraphTask<FRenderVelocityDynamicThreadTask>::CreateTask(nullptr, ENamedThreads::RenderThread)
-				.ConstructAndDispatchWhenReady(this, CmdList, &View);
+			int32 Start = 0;
+			if (EffectiveThreads)
+			{
+				check(IsInRenderingThread());
 
-			SubmitChain = FSubmitCommandlistThreadTask::AddToChain(AnyThreadCompletionEvent, SubmitChain, CmdList);
-		}
-		if (SubmitChain.GetReference())
-		{
-			RHICmdList.Flush(); // this would happen later anyway, but we might as well do these while we wait for async tasks
-			FTaskGraphInterface::Get().WaitUntilTaskCompletes(SubmitChain, ENamedThreads::RenderThread_Local);
+				int32 NumPer = NumPrims / EffectiveThreads;
+				int32 Extra = NumPrims - NumPer * EffectiveThreads;
+
+
+				for (int32 ThreadIndex = 0; ThreadIndex < EffectiveThreads; ThreadIndex++)
+				{
+					int32 Last = Start + (NumPer - 1) + (ThreadIndex < Extra);
+					check(Last >= Start);
+
+					FRHICommandList* CmdList = new FRHICommandList;
+
+					FGraphEventRef AnyThreadCompletionEvent = TGraphTask<FRenderVelocityDynamicThreadTask>::CreateTask(nullptr, ENamedThreads::RenderThread)
+						.ConstructAndDispatchWhenReady(this, CmdList, &View, Start, Last);
+
+					RHICmdList.QueueAsyncCommandListSubmit(AnyThreadCompletionEvent, CmdList);
+
+					Start = Last + 1;
+				}
+			}
 		}
 	}
 	else
@@ -541,7 +561,7 @@ void FDeferredShadingSceneRenderer::RenderVelocitiesInner(FRHICommandListImmedia
 		// single threaded version
 		// Draw velocities for movable static meshes.
 		Scene->VelocityDrawList.DrawVisible(RHICmdList, View, View.StaticMeshVelocityMap, View.StaticMeshBatchVisibility);
-		RenderDynamicVelocities(RHICmdList, View);
+		RenderDynamicVelocitiesInner(RHICmdList, View, 0, View.VisibleDynamicPrimitives.Num() - 1);
 	}
 }
 
