FPrimitiveSceneProxy::PreRenderView / DrawDynamicElements refactored into GetDynamicMeshElements

diff --git a/Engine/Source/Runtime/Renderer/Private/ForwardBasePassRendering.cpp b/Engine/Source/Runtime/Renderer/Private/ForwardBasePassRendering.cpp
--- a/Engine/Source/Runtime/Renderer/Private/ForwardBasePassRendering.cpp
+++ b/Engine/Source/Runtime/Renderer/Private/ForwardBasePassRendering.cpp
@@ -353,7 +353,26 @@ void FForwardShadingSceneRenderer::RenderForwardShadingBasePass(FRHICommandListI
 			SCOPE_CYCLE_COUNTER(STAT_DynamicPrimitiveDrawTime);
 			SCOPED_DRAW_EVENT(Dynamic, DEC_SCENE_ITEMS);
 
-			if (View.VisibleDynamicPrimitives.Num() > 0)
+			const bool bUseGetMeshElements = ShouldUseGetDynamicMeshElements();
+
+			if (bUseGetMeshElements)
+			{
+				FBasePassForwardOpaqueDrawingPolicyFactory::ContextType Context(ESceneRenderTargetsMode::DontSet);
+
+				for (int32 MeshBatchIndex = 0; MeshBatchIndex < View.DynamicMeshElements.Num(); MeshBatchIndex++)
+				{
+					const FMeshBatchAndRelevance& MeshBatchAndRelevance = View.DynamicMeshElements[MeshBatchIndex];
+
+					if (MeshBatchAndRelevance.bHasOpaqueOrMaskedMaterial || ViewFamily.EngineShowFlags.Wireframe)
+					{
+						const FMeshBatch& MeshBatch = *MeshBatchAndRelevance.Mesh;
+						FBasePassForwardOpaqueDrawingPolicyFactory::DrawDynamicMesh(RHICmdList, View, Context, MeshBatch, false, true, MeshBatchAndRelevance.PrimitiveSceneProxy, MeshBatch.BatchHitProxyId);
+					}
+				}
+
+				View.SimpleElementCollector.DrawBatchedElements(RHICmdList, View, NULL, EBlendModeFilter::OpaqueAndMasked);
+			}
+			else if (View.VisibleDynamicPrimitives.Num() > 0)
 			{
 				// Draw the dynamic non-occluded primitives using a base pass drawing policy.
 				TDynamicPrimitiveDrawer<FBasePassForwardOpaqueDrawingPolicyFactory> Drawer(RHICmdList, &View, FBasePassForwardOpaqueDrawingPolicyFactory::ContextType(ESceneRenderTargetsMode::DontSet), true);
@@ -378,7 +397,7 @@ void FForwardShadingSceneRenderer::RenderForwardShadingBasePass(FRHICommandListI
 				}
 			}
 
-			const bool bNeedToSwitchVerticalAxis = IsES2Platform(GRHIShaderPlatform) && !IsPCPlatform(GRHIShaderPlatform);
+			const bool bNeedToSwitchVerticalAxis = RHINeedsToSwitchVerticalAxis(GRHIShaderPlatform);
 
 			// Draw the base pass for the view's batched mesh elements.
 			DrawViewElements<FBasePassForwardOpaqueDrawingPolicyFactory>(RHICmdList, View, FBasePassForwardOpaqueDrawingPolicyFactory::ContextType(ESceneRenderTargetsMode::DontSet), SDPG_World, true);
diff --git a/Engine/Source/Runtime/Renderer/Private/ForwardShadingRenderer.cpp b/Engine/Source/Runtime/Renderer/Private/ForwardShadingRenderer.cpp
--- a/Engine/Source/Runtime/Renderer/Private/ForwardShadingRenderer.cpp
+++ b/Engine/Source/Runtime/Renderer/Private/ForwardShadingRenderer.cpp
@@ -115,7 +115,7 @@ void FForwardShadingSceneRenderer::Render(FRHICommandListImmediate& RHICmdList)
 	{
 		// This might eventually be a problem with multiple views.
 		// Using only view 0 to check to do on-chip transform of alpha.
-		const FViewInfo& View = Views[0];
+		FViewInfo& View = Views[0];
 
 		bool bOnChipSunMask = 
 			GSupportsRenderTargetFormat_PF_FloatRGBA && 
diff --git a/Engine/Source/Runtime/Renderer/Private/ForwardTranslucentRendering.cpp b/Engine/Source/Runtime/Renderer/Private/ForwardTranslucentRendering.cpp
--- a/Engine/Source/Runtime/Renderer/Private/ForwardTranslucentRendering.cpp
+++ b/Engine/Source/Runtime/Renderer/Private/ForwardTranslucentRendering.cpp
@@ -242,6 +242,8 @@ FTranslucentPrimSet
 
 void FTranslucentPrimSet::DrawPrimitivesForForwardShading(FRHICommandListImmediate& RHICmdList, const FViewInfo& View, FSceneRenderer& Renderer) const
 {
+	const bool bUseGetMeshElements = ShouldUseGetDynamicMeshElements();
+
 	// Draw sorted scene prims
 	for (int32 PrimIdx = 0; PrimIdx < SortedPrims.Num(); PrimIdx++)
 	{
@@ -249,61 +251,73 @@ void FTranslucentPrimSet::DrawPrimitivesForForwardShading(FRHICommandListImmedia
 		int32 PrimitiveId = PrimitiveSceneInfo->GetIndex();
 		const FPrimitiveViewRelevance& ViewRelevance = View.PrimitiveViewRelevanceMap[PrimitiveId];
 
-		RenderPrimitiveForForwardShading(RHICmdList, View, PrimitiveSceneInfo, ViewRelevance);
-	}
-}
-
-void FTranslucentPrimSet::RenderPrimitiveForForwardShading(
-	FRHICommandListImmediate& RHICmdList,
-	const FViewInfo& View, 
-	FPrimitiveSceneInfo* PrimitiveSceneInfo, 
-	const FPrimitiveViewRelevance& ViewRelevance) const
-{
-	checkSlow(ViewRelevance.HasTranslucency());
+		checkSlow(ViewRelevance.HasTranslucency());
 
-	if(ViewRelevance.bDrawRelevance)
-	{
-		// Render dynamic scene prim
-		if( ViewRelevance.bDynamicRelevance )
+		if(ViewRelevance.bDrawRelevance)
 		{
-			TDynamicPrimitiveDrawer<FTranslucencyForwardShadingDrawingPolicyFactory> TranslucencyDrawer(
-				RHICmdList,
-				&View,
-				FTranslucencyForwardShadingDrawingPolicyFactory::ContextType(),
-				false
-				);
+			if (bUseGetMeshElements)
+			{
+				FTranslucencyForwardShadingDrawingPolicyFactory::ContextType Context;
 
-			TranslucencyDrawer.SetPrimitive(PrimitiveSceneInfo->Proxy);
-			PrimitiveSceneInfo->Proxy->DrawDynamicElements(
-				&TranslucencyDrawer,
-				&View
-				);
-		}
+				//@todo parallelrendering - come up with a better way to filter these by primitive
+				for (int32 MeshBatchIndex = 0; MeshBatchIndex < View.DynamicMeshElements.Num(); MeshBatchIndex++)
+				{
+					const FMeshBatchAndRelevance& MeshBatchAndRelevance = View.DynamicMeshElements[MeshBatchIndex];
 
-		// Render static scene prim
-		if( ViewRelevance.bStaticRelevance )
-		{
-			// Render static meshes from static scene prim
-			for( int32 StaticMeshIdx=0; StaticMeshIdx < PrimitiveSceneInfo->StaticMeshes.Num(); StaticMeshIdx++ )
+					if (MeshBatchAndRelevance.PrimitiveSceneProxy == PrimitiveSceneInfo->Proxy)
+					{
+						const FMeshBatch& MeshBatch = *MeshBatchAndRelevance.Mesh;
+						FTranslucencyForwardShadingDrawingPolicyFactory::DrawDynamicMesh(RHICmdList, View, Context, MeshBatch, false, false, MeshBatchAndRelevance.PrimitiveSceneProxy, MeshBatch.BatchHitProxyId);
+					}
+				}
+			}
+			// Render dynamic scene prim
+			else if( ViewRelevance.bDynamicRelevance )
+			{
+				TDynamicPrimitiveDrawer<FTranslucencyForwardShadingDrawingPolicyFactory> TranslucencyDrawer(
+					RHICmdList,
+					&View,
+					FTranslucencyForwardShadingDrawingPolicyFactory::ContextType(),
+					false
+					);
+
+				TranslucencyDrawer.SetPrimitive(PrimitiveSceneInfo->Proxy);
+				PrimitiveSceneInfo->Proxy->DrawDynamicElements(
+					&TranslucencyDrawer,
+					&View
+					);
+			}
+
+			// Render static scene prim
+			if( ViewRelevance.bStaticRelevance )
 			{
-				FStaticMesh& StaticMesh = PrimitiveSceneInfo->StaticMeshes[StaticMeshIdx];
-				if (View.StaticMeshVisibilityMap[StaticMesh.Id]
-					// Only render static mesh elements using translucent materials
-					&& StaticMesh.IsTranslucent(View.GetFeatureLevel()) )
+				// Render static meshes from static scene prim
+				for( int32 StaticMeshIdx=0; StaticMeshIdx < PrimitiveSceneInfo->StaticMeshes.Num(); StaticMeshIdx++ )
 				{
-					FTranslucencyForwardShadingDrawingPolicyFactory::DrawStaticMesh(
-						RHICmdList, 
-						View,
-						FTranslucencyForwardShadingDrawingPolicyFactory::ContextType(),
-						StaticMesh,
-						false,
-						PrimitiveSceneInfo->Proxy,
-						StaticMesh.HitProxyId
-						);
+					FStaticMesh& StaticMesh = PrimitiveSceneInfo->StaticMeshes[StaticMeshIdx];
+					if (View.StaticMeshVisibilityMap[StaticMesh.Id]
+						// Only render static mesh elements using translucent materials
+						&& StaticMesh.IsTranslucent(View.GetFeatureLevel()) )
+					{
+						FTranslucencyForwardShadingDrawingPolicyFactory::DrawStaticMesh(
+							RHICmdList, 
+							View,
+							FTranslucencyForwardShadingDrawingPolicyFactory::ContextType(),
+							StaticMesh,
+							false,
+							PrimitiveSceneInfo->Proxy,
+							StaticMesh.BatchHitProxyId
+							);
+					}
 				}
 			}
 		}
 	}
+
+	if (bUseGetMeshElements)
+	{
+		View.SimpleElementCollector.DrawBatchedElements(RHICmdList, View, FTexture2DRHIRef(), EBlendModeFilter::Translucent);
+	}
 }
 
 void FForwardShadingSceneRenderer::RenderTranslucency(FRHICommandListImmediate& RHICmdList)
diff --git a/Engine/Source/Runtime/Renderer/Private/LightMapDensityRendering.cpp b/Engine/Source/Runtime/Renderer/Private/LightMapDensityRendering.cpp
--- a/Engine/Source/Runtime/Renderer/Private/LightMapDensityRendering.cpp
+++ b/Engine/Source/Runtime/Renderer/Private/LightMapDensityRendering.cpp
@@ -57,7 +57,24 @@ bool FDeferredShadingSceneRenderer::RenderLightMapDensities(FRHICommandListImmed
 			{
 				SCOPED_DRAW_EVENT(Dynamic, DEC_SCENE_ITEMS);
 
-				if( View.VisibleDynamicPrimitives.Num() > 0 )
+				const bool bUseGetMeshElements = ShouldUseGetDynamicMeshElements();
+
+				if (bUseGetMeshElements)
+				{
+					FLightMapDensityDrawingPolicyFactory::ContextType Context;
+
+					for (int32 MeshBatchIndex = 0; MeshBatchIndex < View.DynamicMeshElements.Num(); MeshBatchIndex++)
+					{
+						const FMeshBatchAndRelevance& MeshBatchAndRelevance = View.DynamicMeshElements[MeshBatchIndex];
+
+						if (MeshBatchAndRelevance.bHasOpaqueOrMaskedMaterial || ViewFamily.EngineShowFlags.Wireframe)
+						{
+							const FMeshBatch& MeshBatch = *MeshBatchAndRelevance.Mesh;
+							FLightMapDensityDrawingPolicyFactory::DrawDynamicMesh(RHICmdList, View, Context, MeshBatch, false, true, MeshBatchAndRelevance.PrimitiveSceneProxy, MeshBatch.BatchHitProxyId);
+						}
+					}
+				}
+				else if( View.VisibleDynamicPrimitives.Num() > 0 )
 				{
 					// Draw the dynamic non-occluded primitives using a base pass drawing policy.
 					TDynamicPrimitiveDrawer<FLightMapDensityDrawingPolicyFactory> Drawer(
@@ -119,7 +136,7 @@ bool FLightMapDensityDrawingPolicyFactory::DrawDynamicMesh(
 	/*const */FLightMapInteraction LightMapInteraction = (Mesh.LCI && bIsLitMaterial) ? Mesh.LCI->GetLightMapInteraction() : FLightMapInteraction();
 	// force simple lightmaps based on system settings
 	bool bAllowHighQualityLightMaps = AllowHighQualityLightmaps(View.GetFeatureLevel()) && LightMapInteraction.AllowsHighQualityLightmaps();
-	if (bIsLitMaterial && PrimitiveSceneProxy && (PrimitiveSceneProxy->GetLightMapType() == LMIT_Texture))
+	if (bIsLitMaterial && PrimitiveSceneProxy && (LightMapInteraction.GetType() == LMIT_Texture))
 	{
 		// Should this object be texture lightmapped? Ie, is lighting not built for it??
 		bool bUseDummyLightMapPolicy = false;
diff --git a/Engine/Source/Runtime/Renderer/Private/LightMapDensityRendering.h b/Engine/Source/Runtime/Renderer/Private/LightMapDensityRendering.h
--- a/Engine/Source/Runtime/Renderer/Private/LightMapDensityRendering.h
+++ b/Engine/Source/Runtime/Renderer/Private/LightMapDensityRendering.h
@@ -357,23 +357,9 @@ public:
 		else
  		if (PrimitiveSceneProxy)
  		{
- 			LMResolutionScale = PrimitiveSceneProxy->GetLightMapResolutionScale();
+ 			LMResolutionScale = FVector2D(0, 0);
 			BuiltLightingAndSelectedFlags.X = 0.0f;
 			BuiltLightingAndSelectedFlags.Y = 1.0f;
-			if (PrimitiveSceneProxy->GetLightMapType() == LMIT_Texture)
-			{
-				if (PrimitiveSceneProxy->IsLightMapResolutionPadded() == true)
-				{
-					LMResolutionScale.X -= 2.0f;
-					LMResolutionScale.Y -= 2.0f;
-				}
-				bTextureMapped = true;
-				if (PrimitiveSceneProxy->HasStaticLighting())
-				{
-					BuiltLightingAndSelectedFlags.X = 1.0f;
-					BuiltLightingAndSelectedFlags.Y = 0.0f;
-				}
-			}
 		}
 
 		if (Mesh.MaterialRenderProxy && (Mesh.MaterialRenderProxy->IsSelected() == true))
diff --git a/Engine/Source/Runtime/Renderer/Private/PostProcess/PostProcessCompositeEditorPrimitives.cpp b/Engine/Source/Runtime/Renderer/Private/PostProcess/PostProcessCompositeEditorPrimitives.cpp
--- a/Engine/Source/Runtime/Renderer/Private/PostProcess/PostProcessCompositeEditorPrimitives.cpp
+++ b/Engine/Source/Runtime/Renderer/Private/PostProcess/PostProcessCompositeEditorPrimitives.cpp
@@ -302,16 +302,36 @@ void FRCPassPostProcessCompositeEditorPrimitives::RenderPrimitivesToComposite(FR
 	// most objects should be occluded by the existing scene so we do a manual depth test in the shader
 	bool bDepthTest = true;
 
+	const bool bNeedToSwitchVerticalAxis = RHINeedsToSwitchVerticalAxis(GRHIShaderPlatform);
 	FTexture2DRHIRef SceneDepth = GSceneRenderTargets.GetSceneDepthTexture();
 
-	if( View.VisibleEditorPrimitives.Num() > 0 )
+	const bool bUseGetMeshElements = ShouldUseGetDynamicMeshElements();
+
+	if (bUseGetMeshElements)
+	{
+		FBasePassOpaqueDrawingPolicyFactory::ContextType Context(false, ESceneRenderTargetsMode::SetTextures);
+
+		for (int32 MeshBatchIndex = 0; MeshBatchIndex < View.DynamicEditorMeshElements.Num(); MeshBatchIndex++)
+		{
+			const FMeshBatchAndRelevance& MeshBatchAndRelevance = View.DynamicEditorMeshElements[MeshBatchIndex];
+
+			if (MeshBatchAndRelevance.bHasOpaqueOrMaskedMaterial || View.Family->EngineShowFlags.Wireframe)
+			{
+				const FMeshBatch& MeshBatch = *MeshBatchAndRelevance.Mesh;
+				FBasePassOpaqueDrawingPolicyFactory::DrawDynamicMesh(RHICmdList, View, Context, MeshBatch, false, true, MeshBatchAndRelevance.PrimitiveSceneProxy, MeshBatch.BatchHitProxyId);
+			}
+		}
+
+		View.EditorSimpleElementCollector.DrawBatchedElements(RHICmdList, View, GSceneRenderTargets.GetSceneDepthTexture(), EBlendModeFilter::OpaqueAndMasked);
+	}
+	else if( View.VisibleEditorPrimitives.Num() > 0 )
 	{
 		// Draw the dynamic non-occluded primitives using a base pass drawing policy.
 		TDynamicPrimitiveDrawer<FBasePassOpaqueDrawingPolicyFactory> DepthTestDrawer(
-			RHICmdList, &View, FBasePassOpaqueDrawingPolicyFactory::ContextType(bDepthTest, ESceneRenderTargetsMode::SetTextures), true, false, false, bDepthTest);
+			RHICmdList, &View, FBasePassOpaqueDrawingPolicyFactory::ContextType(bDepthTest, ESceneRenderTargetsMode::SetTextures), true, false, bDepthTest);
 
 		TDynamicPrimitiveDrawer<FBasePassOpaqueDrawingPolicyFactory> NoDepthTestDrawer(
-			RHICmdList, &View, FBasePassOpaqueDrawingPolicyFactory::ContextType(!bDepthTest, ESceneRenderTargetsMode::SetTextures), true, false, false, !bDepthTest);
+			RHICmdList, &View, FBasePassOpaqueDrawingPolicyFactory::ContextType(!bDepthTest, ESceneRenderTargetsMode::SetTextures), true, false, !bDepthTest);
 
 		for(int32 PrimitiveIndex = 0;PrimitiveIndex < View.VisibleEditorPrimitives.Num();PrimitiveIndex++)
 		{
@@ -347,8 +367,6 @@ void FRCPassPostProcessCompositeEditorPrimitives::RenderPrimitivesToComposite(FR
 		}
 	}
 
-	const bool bNeedToSwitchVerticalAxis = IsES2Platform(GRHIShaderPlatform) && !IsPCPlatform(GRHIShaderPlatform);
-
 	// Draw the base pass for the view's batched mesh elements.
 	DrawViewElements<FBasePassOpaqueDrawingPolicyFactory>(RHICmdList, View, FBasePassOpaqueDrawingPolicyFactory::ContextType(bDepthTest, ESceneRenderTargetsMode::SetTextures), SDPG_World, false);
 
diff --git a/Engine/Source/Runtime/Renderer/Private/PostProcess/PostProcessSelectionOutline.cpp b/Engine/Source/Runtime/Renderer/Private/PostProcess/PostProcessSelectionOutline.cpp
--- a/Engine/Source/Runtime/Renderer/Private/PostProcess/PostProcessSelectionOutline.cpp
+++ b/Engine/Source/Runtime/Renderer/Private/PostProcess/PostProcessSelectionOutline.cpp
@@ -39,62 +39,98 @@ void FRCPassPostProcessSelectionOutlineColor::Process(FRenderingCompositePassCon
 	Context.RHICmdList.Clear(false, FLinearColor(0, 0, 0, 0), true, 0.0f, true, 0, FIntRect());
 	Context.SetViewportAndCallRHI(ViewRect);
 
-	if(View.VisibleDynamicPrimitives.Num() && View.Family->EngineShowFlags.Selection)
+	if (View.Family->EngineShowFlags.Selection)
 	{
-		TDynamicPrimitiveDrawer<FHitProxyDrawingPolicyFactory> Drawer(Context.RHICmdList, &View, FHitProxyDrawingPolicyFactory::ContextType(), true, false, false, false, true);
-		TMultiMap<FName, const FPrimitiveSceneInfo*> PrimitivesByActor;
+		const bool bUseGetMeshElements = ShouldUseGetDynamicMeshElements();
 
-		for (int32 PrimitiveIndex = 0; PrimitiveIndex < View.VisibleDynamicPrimitives.Num();PrimitiveIndex++)
+		if (bUseGetMeshElements)
 		{
-			const FPrimitiveSceneInfo* PrimitiveSceneInfo = View.VisibleDynamicPrimitives[PrimitiveIndex];
+			FHitProxyDrawingPolicyFactory::ContextType FactoryContext;
 
-			// Only draw the primitive if relevant
-			if(PrimitiveSceneInfo->Proxy->IsSelected())
-			{
-				PrimitivesByActor.Add(PrimitiveSceneInfo->Proxy->GetOwnerName(), PrimitiveSceneInfo);
-			}
-		}
-
-		if (PrimitivesByActor.Num())
-		{
-			// 0 means no object, 1 means BSP so we start with 1
-			uint32 StencilValue = 2;
-
-			TShaderMapRef< TOneColorVS<false> > VertexShader(GetGlobalShaderMap());
-			TShaderMapRef<FOneColorPS> PixelShader(GetGlobalShaderMap());
+			//@todo - use memstack
+			TMap<FName, int32> ActorNameToStencilIndex;
+			ActorNameToStencilIndex.Add(NAME_BSP, 1);
 
 			Context.RHICmdList.SetRasterizerState(TStaticRasterizerState<>::GetRHI());
 			Context.RHICmdList.SetBlendState(TStaticBlendStateWriteMask<CW_NONE, CW_NONE, CW_NONE, CW_NONE>::GetRHI());
 
-			// Sort by actor
-			TArray<FName> Actors;
-			PrimitivesByActor.GetKeys(Actors);
-			for( TArray<FName>::TConstIterator ActorIt(Actors); ActorIt; ++ActorIt )
+			for (int32 MeshBatchIndex = 0; MeshBatchIndex < View.DynamicMeshElements.Num(); MeshBatchIndex++)
 			{
-				bool bBSP = *ActorIt == NAME_BSP;
-				if (bBSP)
+				const FMeshBatchAndRelevance& MeshBatchAndRelevance = View.DynamicMeshElements[MeshBatchIndex];
+				const FPrimitiveSceneProxy* PrimitiveSceneProxy = MeshBatchAndRelevance.PrimitiveSceneProxy;
+
+				if (PrimitiveSceneProxy->IsSelected() && MeshBatchAndRelevance.Mesh->bUseSelectionOutline)
 				{
+					const int32* AssignedStencilIndexPtr = ActorNameToStencilIndex.Find(PrimitiveSceneProxy->GetOwnerName());
+
+					if (!AssignedStencilIndexPtr)
+					{
+						AssignedStencilIndexPtr = &ActorNameToStencilIndex.Add(PrimitiveSceneProxy->GetOwnerName(), ActorNameToStencilIndex.Num() + 1);
+					}
+
 					// This is a reversed Z depth surface, using CF_GreaterEqual.
-					Context.RHICmdList.SetDepthStencilState(TStaticDepthStencilState<true, CF_GreaterEqual, true, CF_Always, SO_Keep, SO_Keep, SO_Replace>::GetRHI(), 1);
+					// Note that the stencil value will overflow with enough selected objects
+					Context.RHICmdList.SetDepthStencilState(TStaticDepthStencilState<true, CF_GreaterEqual, true, CF_Always, SO_Keep, SO_Keep, SO_Replace>::GetRHI(), *AssignedStencilIndexPtr);
+
+					const FMeshBatch& MeshBatch = *MeshBatchAndRelevance.Mesh;
+					FHitProxyDrawingPolicyFactory::DrawDynamicMesh(Context.RHICmdList, View, FactoryContext, MeshBatch, false, true, MeshBatchAndRelevance.PrimitiveSceneProxy, MeshBatch.BatchHitProxyId);
 				}
-				else
-				{
-					// This is a reversed Z depth surface, using CF_GreaterEqual.
-					Context.RHICmdList.SetDepthStencilState(TStaticDepthStencilState<true, CF_GreaterEqual, true, CF_Always, SO_Keep, SO_Keep, SO_Replace>::GetRHI(), StencilValue);
+			}
+		}
+		else if (View.VisibleDynamicPrimitives.Num() > 0)
+		{
+			TDynamicPrimitiveDrawer<FHitProxyDrawingPolicyFactory> Drawer(Context.RHICmdList, &View, FHitProxyDrawingPolicyFactory::ContextType(), true, false, false, true);
+			TMultiMap<FName, const FPrimitiveSceneInfo*> PrimitivesByActor;
+
+			for (int32 PrimitiveIndex = 0; PrimitiveIndex < View.VisibleDynamicPrimitives.Num();PrimitiveIndex++)
+			{
+				const FPrimitiveSceneInfo* PrimitiveSceneInfo = View.VisibleDynamicPrimitives[PrimitiveIndex];
 
-					// we want to use 1..255 for all objects, not correct silhouettes after that is acceptable
-					StencilValue = (StencilValue == 255) ? 2 : (StencilValue + 1);
+				// Only draw the primitive if relevant
+				if(PrimitiveSceneInfo->Proxy->IsSelected())
+				{
+					PrimitivesByActor.Add(PrimitiveSceneInfo->Proxy->GetOwnerName(), PrimitiveSceneInfo);
 				}
+			}
+
+			if (PrimitivesByActor.Num())
+			{
+				// 0 means no object, 1 means BSP so we start with 2
+				uint32 StencilValue = 2;
 
-				TArray<const FPrimitiveSceneInfo*> Primitives;
-				PrimitivesByActor.MultiFind(*ActorIt, Primitives);
+				Context.RHICmdList.SetRasterizerState(TStaticRasterizerState<>::GetRHI());
+				Context.RHICmdList.SetBlendState(TStaticBlendStateWriteMask<CW_NONE, CW_NONE, CW_NONE, CW_NONE>::GetRHI());
 
-				for( TArray<const FPrimitiveSceneInfo*>::TConstIterator PrimIt(Primitives); PrimIt; ++PrimIt )
+				// Sort by actor
+				TArray<FName> Actors;
+				PrimitivesByActor.GetKeys(Actors);
+				for( TArray<FName>::TConstIterator ActorIt(Actors); ActorIt; ++ActorIt )
 				{
-					const FPrimitiveSceneInfo* PrimitiveSceneInfo = *PrimIt;
-					// Render the object to the stencil/depth buffer
-					Drawer.SetPrimitive(PrimitiveSceneInfo->Proxy);
-					PrimitiveSceneInfo->Proxy->DrawDynamicElements(&Drawer, &View);
+					bool bBSP = *ActorIt == NAME_BSP;
+					if (bBSP)
+					{
+						// This is a reversed Z depth surface, using CF_GreaterEqual.
+						Context.RHICmdList.SetDepthStencilState(TStaticDepthStencilState<true, CF_GreaterEqual, true, CF_Always, SO_Keep, SO_Keep, SO_Replace>::GetRHI(), 1);
+					}
+					else
+					{
+						// This is a reversed Z depth surface, using CF_GreaterEqual.
+						Context.RHICmdList.SetDepthStencilState(TStaticDepthStencilState<true, CF_GreaterEqual, true, CF_Always, SO_Keep, SO_Keep, SO_Replace>::GetRHI(), StencilValue);
+
+						// we want to use 1..255 for all objects, not correct silhouettes after that is acceptable
+						StencilValue = (StencilValue == 255) ? 2 : (StencilValue + 1);
+					}
+
+					TArray<const FPrimitiveSceneInfo*> Primitives;
+					PrimitivesByActor.MultiFind(*ActorIt, Primitives);
+
+					for( TArray<const FPrimitiveSceneInfo*>::TConstIterator PrimIt(Primitives); PrimIt; ++PrimIt )
+					{
+						const FPrimitiveSceneInfo* PrimitiveSceneInfo = *PrimIt;
+						// Render the object to the stencil/depth buffer
+						Drawer.SetPrimitive(PrimitiveSceneInfo->Proxy);
+						PrimitiveSceneInfo->Proxy->DrawDynamicElements(&Drawer, &View);
+					}
 				}
 			}
 		}
@@ -237,7 +273,7 @@ public:
 
 			if(EditorPrimitivesStencil.IsBound())
 			{
-			// cache the stencil SRV to avoid create calls each frame (the cache element is stored in the state)
+				// cache the stencil SRV to avoid create calls each frame (the cache element is stored in the state)
 				if(ViewState->SelectionOutlineCacheKey != TargetableTexture)
 				{
 					// release if not the right one (as the internally SRV stores a pointer to the texture we cannot get a false positive)
diff --git a/Engine/Source/Runtime/Renderer/Private/PostProcess/PostProcessTemporalAA.cpp b/Engine/Source/Runtime/Renderer/Private/PostProcess/PostProcessTemporalAA.cpp
--- a/Engine/Source/Runtime/Renderer/Private/PostProcess/PostProcessTemporalAA.cpp
+++ b/Engine/Source/Runtime/Renderer/Private/PostProcess/PostProcessTemporalAA.cpp
@@ -198,7 +198,7 @@ void FRCPassPostProcessSSRTemporalAA::Process(FRenderingCompositePassContext& Co
 		return;
 	}
 
-	FViewInfo& View = const_cast< FViewInfo& >( Context.View );
+	const FViewInfo& View = Context.View;
 
 	FIntPoint TexSize = InputDesc->Extent;
 
@@ -272,8 +272,8 @@ void FRCPassPostProcessDOFTemporalAA::Process(FRenderingCompositePassContext& Co
 		return;
 	}
 
-	FViewInfo& View = const_cast< FViewInfo& >( Context.View );
-	FSceneViewState* ViewState = (FSceneViewState*)Context.View.State;
+	const FViewInfo& View = Context.View;
+	FSceneViewState* ViewState = Context.ViewState;
 
 	FIntPoint TexSize = InputDesc->Extent;
 
@@ -353,8 +353,8 @@ void FRCPassPostProcessLightShaftTemporalAA::Process(FRenderingCompositePassCont
 
 	SCOPED_DRAW_EVENT(LSTemporalAA, DEC_SCENE_ITEMS);
 
-	FViewInfo& View = const_cast< FViewInfo& >( Context.View );
-	FSceneViewState* ViewState = (FSceneViewState*)Context.View.State;
+	const FViewInfo& View = Context.View;
+	FSceneViewState* ViewState = Context.ViewState;
 
 	FIntPoint TexSize = InputDesc->Extent;
 
@@ -431,8 +431,8 @@ void FRCPassPostProcessTemporalAA::Process(FRenderingCompositePassContext& Conte
 		return;
 	}
 
-	FViewInfo& View = const_cast< FViewInfo& >( Context.View );
-	FSceneViewState* ViewState = (FSceneViewState*)Context.View.State;
+	FViewInfo& View = Context.View;
+	FSceneViewState* ViewState = Context.ViewState;
 
 	FIntPoint TexSize = InputDesc->Extent;
 
diff --git a/Engine/Source/Runtime/Renderer/Private/PostProcess/PostProcessing.cpp b/Engine/Source/Runtime/Renderer/Private/PostProcess/PostProcessing.cpp
--- a/Engine/Source/Runtime/Renderer/Private/PostProcess/PostProcessing.cpp
+++ b/Engine/Source/Runtime/Renderer/Private/PostProcess/PostProcessing.cpp
@@ -742,7 +742,7 @@ void OverrideRenderTarget(FRenderingCompositeOutputRef It, TRefCountPtr<IPooledR
 	}
 }
 
-void FPostProcessing::Process(FRHICommandListImmediate& RHICmdList, const FViewInfo& View, TRefCountPtr<IPooledRenderTarget>& VelocityRT)
+void FPostProcessing::Process(FRHICommandListImmediate& RHICmdList, FViewInfo& View, TRefCountPtr<IPooledRenderTarget>& VelocityRT)
 {
 	QUICK_SCOPE_CYCLE_COUNTER( STAT_PostProcessing_Process );
 
@@ -1165,7 +1165,7 @@ void FPostProcessing::Process(FRHICommandListImmediate& RHICmdList, const FViewI
 }
 
 
-void FPostProcessing::ProcessES2(FRHICommandListImmediate& RHICmdList, const FViewInfo& View, bool bUsedFramebufferFetch)
+void FPostProcessing::ProcessES2(FRHICommandListImmediate& RHICmdList, FViewInfo& View, bool bUsedFramebufferFetch)
 {
 	check(IsInRenderingThread());
 
diff --git a/Engine/Source/Runtime/Renderer/Private/PostProcess/PostProcessing.h b/Engine/Source/Runtime/Renderer/Private/PostProcess/PostProcessing.h
--- a/Engine/Source/Runtime/Renderer/Private/PostProcess/PostProcessing.h
+++ b/Engine/Source/Runtime/Renderer/Private/PostProcess/PostProcessing.h
@@ -67,9 +67,9 @@ class FPostProcessing
 {
 public:
 	// @param VelocityRT only valid if motion blur is supported
-	void Process(FRHICommandListImmediate& RHICmdList, const FViewInfo& View, TRefCountPtr<IPooledRenderTarget>& VelocityRT);
+	void Process(FRHICommandListImmediate& RHICmdList, FViewInfo& View, TRefCountPtr<IPooledRenderTarget>& VelocityRT);
 
-	void ProcessES2(FRHICommandListImmediate& RHICmdList, const FViewInfo& View, bool bUsedFramebufferFetch);
+	void ProcessES2(FRHICommandListImmediate& RHICmdList, FViewInfo& View, bool bUsedFramebufferFetch);
 };
 
 /** The global used for post processing. */
diff --git a/Engine/Source/Runtime/Renderer/Private/PostProcess/RenderingCompositionGraph.cpp b/Engine/Source/Runtime/Renderer/Private/PostProcess/RenderingCompositionGraph.cpp
--- a/Engine/Source/Runtime/Renderer/Private/PostProcess/RenderingCompositionGraph.cpp
+++ b/Engine/Source/Runtime/Renderer/Private/PostProcess/RenderingCompositionGraph.cpp
@@ -103,8 +103,9 @@ FAutoConsoleCommand CmdCompositionGraphDebug(
 	);
 #endif
 
-FRenderingCompositePassContext::FRenderingCompositePassContext(FRHICommandListImmediate& InRHICmdList, const FViewInfo& InView/*, const FSceneRenderTargetItem& InRenderTargetItem*/)
+FRenderingCompositePassContext::FRenderingCompositePassContext(FRHICommandListImmediate& InRHICmdList, FViewInfo& InView/*, const FSceneRenderTargetItem& InRenderTargetItem*/)
 	: View(InView)
+	, ViewState((FSceneViewState*)InView.State)
 	//		, CompositingOutputRTItem(InRenderTargetItem)
 	, Pass(0)
 	, RHICmdList(InRHICmdList)
diff --git a/Engine/Source/Runtime/Renderer/Private/PostProcess/RenderingCompositionGraph.h b/Engine/Source/Runtime/Renderer/Private/PostProcess/RenderingCompositionGraph.h
--- a/Engine/Source/Runtime/Renderer/Private/PostProcess/RenderingCompositionGraph.h
+++ b/Engine/Source/Runtime/Renderer/Private/PostProcess/RenderingCompositionGraph.h
@@ -79,7 +79,7 @@ private:
 
 struct FRenderingCompositePassContext
 {
-	FRenderingCompositePassContext(FRHICommandListImmediate& RHICmdList, const FViewInfo& InView);
+	FRenderingCompositePassContext(FRHICommandListImmediate& RHICmdList, FViewInfo& InView);
 
 	~FRenderingCompositePassContext();
 
@@ -87,7 +87,8 @@ struct FRenderingCompositePassContext
 	void Process(const TCHAR *GraphDebugName);
 
 	//
-	const FViewInfo& View;
+	FViewInfo& View;
+	FSceneViewState* ViewState;
 	// is updated before each Pass->Process() call
 	FRenderingCompositePass* Pass;
 
diff --git a/Engine/Source/Runtime/Renderer/Private/PostProcess/ScreenSpaceReflections.cpp b/Engine/Source/Runtime/Renderer/Private/PostProcess/ScreenSpaceReflections.cpp
--- a/Engine/Source/Runtime/Renderer/Private/PostProcess/ScreenSpaceReflections.cpp
+++ b/Engine/Source/Runtime/Renderer/Private/PostProcess/ScreenSpaceReflections.cpp
@@ -257,7 +257,7 @@ FPooledRenderTargetDesc FRCPassPostProcessScreenSpaceReflections::ComputeOutputD
 }
 void BuildHZB(FRHICommandListImmediate& RHICmdList, const FViewInfo& View);
 
-void ScreenSpaceReflections(FRHICommandListImmediate& RHICmdList, const FViewInfo& View, TRefCountPtr<IPooledRenderTarget>& SSROutput)
+void ScreenSpaceReflections(FRHICommandListImmediate& RHICmdList, FViewInfo& View, TRefCountPtr<IPooledRenderTarget>& SSROutput)
 {
 	BuildHZB(RHICmdList, View);
 
diff --git a/Engine/Source/Runtime/Renderer/Private/PostProcess/ScreenSpaceReflections.h b/Engine/Source/Runtime/Renderer/Private/PostProcess/ScreenSpaceReflections.h
--- a/Engine/Source/Runtime/Renderer/Private/PostProcess/ScreenSpaceReflections.h
+++ b/Engine/Source/Runtime/Renderer/Private/PostProcess/ScreenSpaceReflections.h
@@ -49,6 +49,6 @@ public:
 	virtual FPooledRenderTargetDesc ComputeOutputDesc(EPassOutputId InPassOutputId) const;
 };
 
-void ScreenSpaceReflections(FRHICommandListImmediate& RHICmdList, const FViewInfo& View, TRefCountPtr<IPooledRenderTarget>& SSROutput);
+void ScreenSpaceReflections(FRHICommandListImmediate& RHICmdList, FViewInfo& View, TRefCountPtr<IPooledRenderTarget>& SSROutput);
 
 bool DoScreenSpaceReflections(const FViewInfo& View);
diff --git a/Engine/Source/Runtime/Renderer/Private/SceneCore.h b/Engine/Source/Runtime/Renderer/Private/SceneCore.h
--- a/Engine/Source/Runtime/Renderer/Private/SceneCore.h
+++ b/Engine/Source/Runtime/Renderer/Private/SceneCore.h
@@ -129,9 +129,6 @@ public:
 	/** The render info for the primitive which created this mesh. */
 	FPrimitiveSceneInfo* PrimitiveSceneInfo;
 
-	/** The ID of the hit proxy which represents this static mesh. */
-	FHitProxyId HitProxyId;
-
 	/** The index of the mesh in the scene's static meshes array. */
 	int32 Id;
 
@@ -149,10 +146,12 @@ public:
 		FMeshBatch(InMesh),
 		ScreenSize(InScreenSize),
 		PrimitiveSceneInfo(InPrimitiveSceneInfo),
-		HitProxyId(InHitProxyId),
 		Id(INDEX_NONE),
 		bShadowOnly(bInShadowOnly)
-	{}
+	{
+		BatchHitProxyId = InHitProxyId;
+	}
+
 	~FStaticMesh();
 
 	/** Adds a link from the mesh to its entry in a draw list. */
@@ -179,7 +178,6 @@ private:
 		FMeshBatch(InStaticMesh),
 		ScreenSize(InStaticMesh.ScreenSize),
 		PrimitiveSceneInfo(InStaticMesh.PrimitiveSceneInfo),
-		HitProxyId(InStaticMesh.HitProxyId),
 		Id(InStaticMesh.Id)
 	{}
 };
diff --git a/Engine/Source/Runtime/Renderer/Private/SceneHitProxyRendering.cpp b/Engine/Source/Runtime/Renderer/Private/SceneHitProxyRendering.cpp
--- a/Engine/Source/Runtime/Renderer/Private/SceneHitProxyRendering.cpp
+++ b/Engine/Source/Runtime/Renderer/Private/SceneHitProxyRendering.cpp
@@ -258,7 +258,7 @@ void FHitProxyDrawingPolicyFactory::AddStaticMesh(FScene* Scene,FStaticMesh* Sta
 	
 	Scene->HitProxyDrawList.AddMesh(
 		StaticMesh,
-		StaticMesh->HitProxyId,
+		StaticMesh->BatchHitProxyId,
 		FHitProxyDrawingPolicy(StaticMesh->VertexFactory, MaterialRenderProxy, Scene->GetFeatureLevel()),
 		Scene->GetFeatureLevel()
 		);
@@ -271,7 +271,7 @@ void FHitProxyDrawingPolicyFactory::AddStaticMesh(FScene* Scene,FStaticMesh* Sta
 	{
 		Scene->HitProxyDrawList_OpaqueOnly.AddMesh(
 			StaticMesh,
-			StaticMesh->HitProxyId,
+			StaticMesh->BatchHitProxyId,
 			FHitProxyDrawingPolicy(StaticMesh->VertexFactory, MaterialRenderProxy, Scene->GetFeatureLevel()),
 			Scene->GetFeatureLevel()
 			);
@@ -375,7 +375,7 @@ void RenderHitProxies(FRHICommandListImmediate& RHICmdList, const FSceneRenderer
 	RHICmdList.SetDepthStencilState(TStaticDepthStencilState<true, CF_GreaterEqual>::GetRHI());
 	RHICmdList.SetBlendState(TStaticBlendState<>::GetRHI());
 
-	const bool bNeedToSwitchVerticalAxis = IsES2Platform(GRHIShaderPlatform) && !IsPCPlatform(GRHIShaderPlatform);
+	const bool bNeedToSwitchVerticalAxis = RHINeedsToSwitchVerticalAxis(GRHIShaderPlatform);
 
 	for (int32 ViewIndex = 0; ViewIndex < Views.Num(); ViewIndex++)
 	{
@@ -403,26 +403,62 @@ void RenderHitProxies(FRHICommandListImmediate& RHICmdList, const FSceneRenderer
 		const bool bPreFog = true;
 		const bool bHitTesting = true;
 
-		// Draw all dynamic primitives using the hit proxy drawing policy.
-		DrawDynamicPrimitiveSet<FHitProxyDrawingPolicyFactory>(
-			RHICmdList,
-			View,
-			View.VisibleDynamicPrimitives,
-			FHitProxyDrawingPolicyFactory::ContextType(),
-			bPreFog,
-			bHitTesting
-			);
+		const bool bUseGetMeshElements = ShouldUseGetDynamicMeshElements();
 
-		// Draw all visible editor primitives using the hit proxy drawing policy
-		DrawDynamicPrimitiveSet<FHitProxyDrawingPolicyFactory>(
-			RHICmdList,
-			View,
-			View.VisibleEditorPrimitives,
-			FHitProxyDrawingPolicyFactory::ContextType(),
-			bPreFog,
-			bHitTesting
-			);
+		if (bUseGetMeshElements)
+		{
+			FHitProxyDrawingPolicyFactory::ContextType DrawingContext;
+
+			for (int32 MeshBatchIndex = 0; MeshBatchIndex < View.DynamicMeshElements.Num(); MeshBatchIndex++)
+			{
+				const FMeshBatchAndRelevance& MeshBatchAndRelevance = View.DynamicMeshElements[MeshBatchIndex];
+				const FMeshBatch& MeshBatch = *MeshBatchAndRelevance.Mesh;
+				
+				if (MeshBatch.bSelectable)
+				{
+					const FHitProxyId EffectiveHitProxyId = MeshBatch.BatchHitProxyId == FHitProxyId() ? MeshBatchAndRelevance.PrimitiveSceneProxy->GetPrimitiveSceneInfo()->DefaultDynamicHitProxyId : MeshBatch.BatchHitProxyId;
+					FHitProxyDrawingPolicyFactory::DrawDynamicMesh(RHICmdList, View, DrawingContext, MeshBatch, false, bPreFog, MeshBatchAndRelevance.PrimitiveSceneProxy, EffectiveHitProxyId);
+				}
+			}
+
+			View.SimpleElementCollector.DrawBatchedElements(RHICmdList, View, FTexture2DRHIRef(), EBlendModeFilter::All);
+
+			for (int32 MeshBatchIndex = 0; MeshBatchIndex < View.DynamicEditorMeshElements.Num(); MeshBatchIndex++)
+			{
+				const FMeshBatchAndRelevance& MeshBatchAndRelevance = View.DynamicEditorMeshElements[MeshBatchIndex];
+				const FMeshBatch& MeshBatch = *MeshBatchAndRelevance.Mesh;
+
+				if (MeshBatch.bSelectable)
+				{
+					const FHitProxyId EffectiveHitProxyId = MeshBatch.BatchHitProxyId == FHitProxyId() ? MeshBatchAndRelevance.PrimitiveSceneProxy->GetPrimitiveSceneInfo()->DefaultDynamicHitProxyId : MeshBatch.BatchHitProxyId;
+					FHitProxyDrawingPolicyFactory::DrawDynamicMesh(RHICmdList, View, DrawingContext, MeshBatch, false, bPreFog, MeshBatchAndRelevance.PrimitiveSceneProxy, EffectiveHitProxyId);
+				}
+			}
 
+			View.EditorSimpleElementCollector.DrawBatchedElements(RHICmdList, View, FTexture2DRHIRef(), EBlendModeFilter::All);
+		}
+		else 
+		{
+			// Draw all dynamic primitives using the hit proxy drawing policy.
+			DrawDynamicPrimitiveSet<FHitProxyDrawingPolicyFactory>(
+				RHICmdList,
+				View,
+				View.VisibleDynamicPrimitives,
+				FHitProxyDrawingPolicyFactory::ContextType(),
+				bPreFog,
+				bHitTesting
+				);
+
+			// Draw all visible editor primitives using the hit proxy drawing policy
+			DrawDynamicPrimitiveSet<FHitProxyDrawingPolicyFactory>(
+				RHICmdList,
+				View,
+				View.VisibleEditorPrimitives,
+				FHitProxyDrawingPolicyFactory::ContextType(),
+				bPreFog,
+				bHitTesting
+				);
+		}
 
 		// Draw the view's elements.
 		DrawViewElements<FHitProxyDrawingPolicyFactory>(RHICmdList, View, FHitProxyDrawingPolicyFactory::ContextType(), SDPG_World, bPreFog);
diff --git a/Engine/Source/Runtime/Renderer/Private/ScenePostProcessing.cpp b/Engine/Source/Runtime/Renderer/Private/ScenePostProcessing.cpp
--- a/Engine/Source/Runtime/Renderer/Private/ScenePostProcessing.cpp
+++ b/Engine/Source/Runtime/Renderer/Private/ScenePostProcessing.cpp
@@ -80,7 +80,7 @@ IMPLEMENT_SHADER_TYPE(,FGammaCorrectionVS,TEXT("GammaCorrection"),TEXT("MainVS")
 * Finish rendering a view, writing the contents to ViewFamily.RenderTarget.
 * @param View - The view to process.
 */
-void FDeferredShadingSceneRenderer::FinishRenderViewTarget(FRHICommandListImmediate& RHICmdList, const FViewInfo* View, bool bLastView)
+void FDeferredShadingSceneRenderer::FinishRenderViewTarget(FRHICommandListImmediate& RHICmdList, FViewInfo* View, bool bLastView)
 {
 	TRefCountPtr<IPooledRenderTarget> VelocityRT;
 
diff --git a/Engine/Source/Runtime/Renderer/Private/ScenePrivate.h b/Engine/Source/Runtime/Renderer/Private/ScenePrivate.h
--- a/Engine/Source/Runtime/Renderer/Private/ScenePrivate.h
+++ b/Engine/Source/Runtime/Renderer/Private/ScenePrivate.h
@@ -36,6 +36,8 @@ bool IsMobileHDR();
 /** True if the mobile renderer is emulating HDR in a 32bpp render target. */
 bool IsMobileHDR32bpp();
 
+extern bool ShouldUseGetDynamicMeshElements();
+
 // Dependencies.
 #include "StaticBoundShaderState.h"
 #include "BatchedElements.h"
diff --git a/Engine/Source/Runtime/Renderer/Private/SceneRendering.cpp b/Engine/Source/Runtime/Renderer/Private/SceneRendering.cpp
--- a/Engine/Source/Runtime/Renderer/Private/SceneRendering.cpp
+++ b/Engine/Source/Runtime/Renderer/Private/SceneRendering.cpp
@@ -664,7 +664,7 @@ void FSceneRenderer::RenderFinish(FRHICommandListImmediate& RHICmdList)
 	if(FRCPassPostProcessBusyWait::IsPassRequired())
 	{
 		// mostly view independent but to be safe we use the first view
-		const FViewInfo& View = Views[0];
+		FViewInfo& View = Views[0];
 
 		FMemMark Mark(FMemStack::Get());
 		FRenderingCompositePassContext CompositeContext(RHICmdList, View);
@@ -883,7 +883,7 @@ void FSceneRenderer::RenderCustomDepthPass(FRHICommandListImmediate& RHICmdList)
 			RHICmdList.SetDepthStencilState(TStaticDepthStencilState<true, CF_GreaterEqual>::GetRHI());
 			RHICmdList.SetBlendState(TStaticBlendState<>::GetRHI());
 
-			View.CustomDepthSet.DrawPrims(RHICmdList, &View, false);
+			View.CustomDepthSet.DrawPrims(RHICmdList, View);
 		}
 
 		// resolve using the current ResolveParams 
@@ -1065,3 +1065,10 @@ bool IsMobileHDR32bpp()
 	static auto* MobileHDR32bppCvar = IConsoleManager::Get().FindTConsoleVariableDataInt(TEXT("r.MobileHDR32bpp"));
 	return IsMobileHDR() && (GSupportsRenderTargetFormat_PF_FloatRGBA == false || MobileHDR32bppCvar->GetValueOnRenderThread() == 1);
 }
+
+bool ShouldUseGetDynamicMeshElements()
+{
+	static const auto CVarUseGetMeshElements = IConsoleManager::Get().FindTConsoleVariableDataInt(TEXT("r.UseGetMeshElements"));
+	const bool bUseGetMeshElements = CVarUseGetMeshElements->GetValueOnRenderThread() != 0;
+	return bUseGetMeshElements;
+}
\ No newline at end of file
diff --git a/Engine/Source/Runtime/Renderer/Private/SceneRendering.h b/Engine/Source/Runtime/Renderer/Private/SceneRendering.h
--- a/Engine/Source/Runtime/Renderer/Private/SceneRendering.h
+++ b/Engine/Source/Runtime/Renderer/Private/SceneRendering.h
@@ -13,19 +13,6 @@
 // Forward declarations.
 class FPostprocessContext;
 
-/** An association between a hit proxy and a mesh. */
-class FHitProxyMeshPair : public FMeshBatch
-{
-public:
-	FHitProxyId HitProxyId;
-
-	/** Initialization constructor. */
-	FHitProxyMeshPair(const FMeshBatch& InMesh,FHitProxyId InHitProxyId):
-		FMeshBatch(InMesh),
-		HitProxyId(InHitProxyId)
-	{}
-};
-
 /** Information about a visible light which is specific to the view it's visible in. */
 class FVisibleLightViewInfo
 {
@@ -204,9 +191,6 @@ private:
 
 	/** Renders a single primitive for the deferred shading pipeline. */
 	void RenderPrimitive(FRHICommandList& RHICmdList, const FViewInfo& View, FPrimitiveSceneInfo* PrimitiveSceneInfo, const FPrimitiveViewRelevance& ViewRelevance, const FProjectedShadowInfo* TranslucentSelfShadow, bool bSeparateTranslucencyPass) const;
-
-	/** Renders a single primitive for the forward shading pipeline. */
-	void RenderPrimitiveForForwardShading(FRHICommandListImmediate& RHICmdList, const FViewInfo& View, FPrimitiveSceneInfo* PrimitiveSceneInfo, const FPrimitiveViewRelevance& ViewRelevance) const;
 };
 
 /** A batched occlusion primitive. */
@@ -369,14 +353,23 @@ public:
 	FBatchedElements TopBatchedViewElements;
 
 	/** The view's mesh elements. */
-	TIndirectArray<FHitProxyMeshPair> ViewMeshElements;
+	TIndirectArray<FMeshBatch> ViewMeshElements;
 
 	/** The view's mesh elements for the foreground (editor gizmos and primitives )*/
-	TIndirectArray<FHitProxyMeshPair> TopViewMeshElements;
+	TIndirectArray<FMeshBatch> TopViewMeshElements;
 
 	/** The dynamic resources used by the view elements. */
 	TArray<FDynamicPrimitiveResource*> DynamicResources;
 
+	/** Gathered in initviews from all the primitives with dynamic view relevance, used in each mesh pass. */
+	TArray<FMeshBatchAndRelevance,SceneRenderingAllocator> DynamicMeshElements;
+
+	TArray<FMeshBatchAndRelevance,SceneRenderingAllocator> DynamicEditorMeshElements;
+
+	FSimpleElementCollector SimpleElementCollector;
+
+	FSimpleElementCollector EditorSimpleElementCollector;
+
 	/** Parameters for exponential height fog. */
 	FVector4 ExponentialFogParameters;
 	FVector ExponentialFogColor;
@@ -501,6 +494,11 @@ struct FCombinedShadowStats
 	{}
 };
 
+/**
+ * Masks indicating for which views a primitive needs to have a certain operation on.
+ * One entry per primitive in the scene.
+ */
+typedef TArray<uint8, SceneRenderingAllocator> FPrimitiveViewMasks;
 
 /**
  * Used as the scope for scene rendering functions.
@@ -520,6 +518,8 @@ public:
 	/** The views being rendered. */
 	TArray<FViewInfo> Views;
 
+	FMeshElementCollector MeshCollector;
+
 	/** Information about the visible lights. */
 	TArray<FVisibleLightInfo,SceneRenderingAllocator> VisibleLightInfos;
 
@@ -563,6 +563,14 @@ protected:
 	/** Performs once per frame setup after to visibility determination. */
 	void PostVisibilityFrameSetup();
 
+	void GatherDynamicMeshElements(
+		TArray<FViewInfo>& Views, 
+		const FScene* Scene, 
+		const FSceneViewFamily& ViewFamily, 
+		const FPrimitiveViewMasks& HasDynamicMeshElementsMasks, 
+		const FPrimitiveViewMasks& HasDynamicEditorMeshElementsMasks, 
+		FMeshElementCollector& Collector);
+
 	/** Initialized the fog constants for each view. */
 	void InitFogConstants();
 
diff --git a/Engine/Source/Runtime/Renderer/Private/SceneVisibility.cpp b/Engine/Source/Runtime/Renderer/Private/SceneVisibility.cpp
--- a/Engine/Source/Runtime/Renderer/Private/SceneVisibility.cpp
+++ b/Engine/Source/Runtime/Renderer/Private/SceneVisibility.cpp
@@ -620,12 +620,6 @@ static int32 OcclusionCull(FRHICommandListImmediate& RHICmdList, const FScene* S
 /** List of relevant static primitives. Indexes into Scene->Primitives. */
 typedef TArray<int32, SceneRenderingAllocator> FRelevantStaticPrimitives;
 
-/**
- * Masks indicating for which views a primitive needs to have PreRenderView
- * called. One entry per primitive in the scene.
- */
-typedef TArray<uint8, SceneRenderingAllocator> FPrimitivePreRenderViewMasks;
-
 /**
  * Computes view relevance for visible primitives in the view and adds them to
  * appropriate per-view rendering lists.
@@ -643,7 +637,9 @@ static void ComputeRelevanceForView(
 	FViewInfo& View,
 	uint8 ViewBit,
 	FRelevantStaticPrimitives& OutRelevantStaticPrimitives,
-	FPrimitivePreRenderViewMasks& OutPreRenderViewMasks
+	FPrimitiveViewMasks& OutPreRenderViewMasks,
+	FPrimitiveViewMasks& OutHasDynamicMeshElementsMasks,
+	FPrimitiveViewMasks& OutHasDynamicEditorMeshElementsMasks
 	)
 {
 	SCOPE_CYCLE_COUNTER(STAT_ComputeViewRelevance);
@@ -655,6 +651,7 @@ static void ComputeRelevanceForView(
 	// But it's hard to say since the view relevance is going to be available right now.
 
 	check(OutPreRenderViewMasks.Num() == Scene->Primitives.Num());
+	check(OutHasDynamicMeshElementsMasks.Num() == Scene->Primitives.Num());
 
 	float CurrentWorldTime = View.Family->CurrentWorldTime;
 	float DeltaWorldTime = View.Family->DeltaWorldTime;
@@ -694,11 +691,17 @@ static void ComputeRelevanceForView(
 		{
 			// Editor primitives are rendered after post processing and composited onto the scene
 			View.VisibleEditorPrimitives.Add(PrimitiveSceneInfo);
+
+			if (GIsEditor)
+			{
+				OutHasDynamicEditorMeshElementsMasks[BitIt.GetIndex()] |= ViewBit;
+			}
 		}
 		else if(bDynamicRelevance)
 		{
 			// Keep track of visible dynamic primitives.
 			View.VisibleDynamicPrimitives.Add(PrimitiveSceneInfo);
+			OutHasDynamicMeshElementsMasks[BitIt.GetIndex()] |= ViewBit;
 		}
 
 		if (ViewRelevance.HasTranslucency() && !bEditorRelevance && ViewRelevance.bRenderInMainPass)
@@ -721,7 +724,7 @@ static void ComputeRelevanceForView(
 		if (ViewRelevance.bRenderCustomDepth)
 		{
 			// Add to set of dynamic distortion primitives
-			View.CustomDepthSet.AddScenePrimitive(PrimitiveSceneInfo->Proxy, View);
+			View.CustomDepthSet.AddScenePrimitive(PrimitiveSceneInfo->Proxy);
 		}
 
 		// INITVIEWS_TODO: Do this in a separate pass? There are no dependencies
@@ -759,6 +762,62 @@ static void ComputeRelevanceForView(
 	}
 }
 
+void FSceneRenderer::GatherDynamicMeshElements(
+	TArray<FViewInfo>& Views, 
+	const FScene* Scene, 
+	const FSceneViewFamily& ViewFamily, 
+	const FPrimitiveViewMasks& HasDynamicMeshElementsMasks, 
+	const FPrimitiveViewMasks& HasDynamicEditorMeshElementsMasks, 
+	FMeshElementCollector& Collector)
+{
+	SCOPE_CYCLE_COUNTER(STAT_GetDynamicMeshElements);
+
+	int32 NumPrimitives = Scene->Primitives.Num();
+	check(HasDynamicMeshElementsMasks.Num() == NumPrimitives);
+
+	{
+		Collector.ClearViewMeshArrays();
+
+		for (int32 ViewIndex = 0; ViewIndex < Views.Num(); ViewIndex++)
+		{
+			Collector.AddViewMeshArrays(&Views[ViewIndex], &Views[ViewIndex].DynamicMeshElements, &Views[ViewIndex].SimpleElementCollector, ViewFamily.GetFeatureLevel());
+		}
+
+		for (int32 PrimitiveIndex = 0; PrimitiveIndex < NumPrimitives; ++PrimitiveIndex)
+		{
+			const uint8 ViewMask = HasDynamicMeshElementsMasks[PrimitiveIndex];
+
+			if (ViewMask != 0)
+			{
+				FPrimitiveSceneInfo* PrimitiveSceneInfo = Scene->Primitives[PrimitiveIndex];
+				Collector.SetPrimitive(PrimitiveSceneInfo->Proxy, PrimitiveSceneInfo->DefaultDynamicHitProxyId);
+				PrimitiveSceneInfo->Proxy->GetDynamicMeshElements(ViewFamily.Views, ViewFamily, ViewMask, Collector);
+			}
+		}
+	}
+
+	{
+		Collector.ClearViewMeshArrays();
+
+		for (int32 ViewIndex = 0; ViewIndex < Views.Num(); ViewIndex++)
+		{
+			Collector.AddViewMeshArrays(&Views[ViewIndex], &Views[ViewIndex].DynamicEditorMeshElements, &Views[ViewIndex].EditorSimpleElementCollector, ViewFamily.GetFeatureLevel());
+		}
+
+		for (int32 PrimitiveIndex = 0; PrimitiveIndex < NumPrimitives; ++PrimitiveIndex)
+		{
+			const uint8 ViewMask = HasDynamicEditorMeshElementsMasks[PrimitiveIndex];
+
+			if (ViewMask != 0)
+			{
+				FPrimitiveSceneInfo* PrimitiveSceneInfo = Scene->Primitives[PrimitiveIndex];
+				Collector.SetPrimitive(PrimitiveSceneInfo->Proxy, PrimitiveSceneInfo->DefaultDynamicHitProxyId);
+				PrimitiveSceneInfo->Proxy->GetDynamicMeshElements(ViewFamily.Views, ViewFamily, ViewMask, Collector);
+			}
+		}
+	}
+}
+
 static void MarkAllPrimitivesForReflectionProxyUpdate(FScene* Scene)
 {
 	if (Scene->ReflectionSceneData.bRegisteredReflectionCapturesHasChanged)
@@ -871,7 +930,7 @@ static void MarkRelevantStaticMeshesForView(
 static void DispatchPreRenderView(
 	const FScene* Scene,
 	const FSceneViewFamily* ViewFamily,
-	const FPrimitivePreRenderViewMasks& PreRenderViewMasks,
+	const FPrimitiveViewMasks& PreRenderViewMasks,
 	uint32 FrameNumber
 	)
 {
@@ -1224,9 +1283,19 @@ void FSceneRenderer::ComputeViewVisibility(FRHICommandListImmediate& RHICmdList)
 
 	// This array contains a per-primitive mask specifying for which views a
 	// primitive must have PreRenderView called.
-	FPrimitivePreRenderViewMasks PreRenderViewMasks;
+	FPrimitiveViewMasks PreRenderViewMasks;
 	PreRenderViewMasks.AddZeroed(NumPrimitives);
 
+	FPrimitiveViewMasks HasDynamicMeshElementsMasks;
+	HasDynamicMeshElementsMasks.AddZeroed(NumPrimitives);
+
+	FPrimitiveViewMasks HasDynamicEditorMeshElementsMasks;
+
+	if (GIsEditor)
+	{
+		HasDynamicEditorMeshElementsMasks.AddZeroed(NumPrimitives);
+	}
+
 	// This array contains a list of relevant static primities.
 	FRelevantStaticPrimitives RelevantStaticPrimitives;
 
@@ -1385,7 +1454,7 @@ void FSceneRenderer::ComputeViewVisibility(FRHICommandListImmediate& RHICmdList)
 
 		// Compute view relevance for all visible primitives.
 		RelevantStaticPrimitives.Reset();
-		ComputeRelevanceForView(RHICmdList, Scene, View, ViewBit, RelevantStaticPrimitives, PreRenderViewMasks);
+		ComputeRelevanceForView(RHICmdList, Scene, View, ViewBit, RelevantStaticPrimitives, PreRenderViewMasks, HasDynamicMeshElementsMasks, HasDynamicEditorMeshElementsMasks);
 		MarkRelevantStaticMeshesForView(Scene, View, RelevantStaticPrimitives);
 
 #if !(UE_BUILD_SHIPPING || UE_BUILD_TEST)
@@ -1411,7 +1480,16 @@ void FSceneRenderer::ComputeViewVisibility(FRHICommandListImmediate& RHICmdList)
 #endif
 	}
 
-	DispatchPreRenderView(Scene, &ViewFamily, PreRenderViewMasks, FrameNumber);
+	const bool bUseGetMeshElements = ShouldUseGetDynamicMeshElements();
+
+	if (bUseGetMeshElements)
+	{
+		GatherDynamicMeshElements(Views, Scene, ViewFamily, HasDynamicMeshElementsMasks, HasDynamicEditorMeshElementsMasks, MeshCollector);
+	}
+	else
+	{
+		DispatchPreRenderView(Scene, &ViewFamily, PreRenderViewMasks, FrameNumber);
+	}
 
 	INC_DWORD_STAT_BY(STAT_ProcessedPrimitives,NumProcessedPrimitives);
 	INC_DWORD_STAT_BY(STAT_CulledPrimitives,NumCulledPrimitives);
diff --git a/Engine/Source/Runtime/Renderer/Private/ShadowRendering.cpp b/Engine/Source/Runtime/Renderer/Private/ShadowRendering.cpp
--- a/Engine/Source/Runtime/Renderer/Private/ShadowRendering.cpp
+++ b/Engine/Source/Runtime/Renderer/Private/ShadowRendering.cpp
@@ -1275,11 +1275,14 @@ void FProjectedShadowInfo::RenderDepthDynamic(FRHICommandList& RHICmdList, FDefe
 		DrawPrimitiveFlags = EDrawDynamicFlags::ForceLowestLOD;
 	}
 
+	const bool bUseGetMeshElements = ShouldUseGetDynamicMeshElements();
 	TDynamicPrimitiveDrawer<FShadowDepthDrawingPolicyFactory> Drawer(RHICmdList, FoundView, FShadowDepthDrawingPolicyFactory::ContextType(this), true);
 	uint32 PrimitiveCount = SubjectPrimitives.Num();
+
 	for (uint32 PrimitiveIndex = 0; PrimitiveIndex < PrimitiveCount; ++PrimitiveIndex)
 	{
 		const FPrimitiveSceneInfo* PrimitiveSceneInfo = SubjectPrimitives[PrimitiveIndex];
+		const FPrimitiveSceneProxy* PrimitiveSceneProxy = PrimitiveSceneInfo->Proxy;
 
 		// Lookup the primitive's cached view relevance
 		FPrimitiveViewRelevance ViewRelevance = FoundView->PrimitiveViewRelevanceMap[PrimitiveSceneInfo->GetIndex()];
@@ -1293,10 +1296,29 @@ void FProjectedShadowInfo::RenderDepthDynamic(FRHICommandList& RHICmdList, FDefe
 		// Only draw if the subject primitive is shadow relevant.
 		if (ViewRelevance.bShadowRelevance)
 		{
-			FScopeCycleCounter Context(SubjectPrimitives[PrimitiveIndex]->Proxy->GetStatId());
-			Drawer.SetPrimitive(SubjectPrimitives[PrimitiveIndex]->Proxy);
+			if (bUseGetMeshElements)
+			{
+				FShadowDepthDrawingPolicyFactory::ContextType Context(this);
+
+				//@todo parallelrendering - come up with a better way to filter these by primitive
+				for (int32 MeshBatchIndex = 0; MeshBatchIndex < FoundView->DynamicMeshElements.Num(); MeshBatchIndex++)
+				{
+					const FMeshBatchAndRelevance& MeshBatchAndRelevance = FoundView->DynamicMeshElements[MeshBatchIndex];
+
+					if (MeshBatchAndRelevance.PrimitiveSceneProxy == PrimitiveSceneProxy)
+					{
+						const FMeshBatch& MeshBatch = *MeshBatchAndRelevance.Mesh;
+						FShadowDepthDrawingPolicyFactory::DrawDynamicMesh(RHICmdList, *FoundView, Context, MeshBatch, false, true, MeshBatchAndRelevance.PrimitiveSceneProxy, MeshBatch.BatchHitProxyId);
+					}
+				}
+			}
+			else 
+			{
+				FScopeCycleCounter Context(PrimitiveSceneProxy->GetStatId());
+				Drawer.SetPrimitive(PrimitiveSceneProxy);
 
-			SubjectPrimitives[PrimitiveIndex]->Proxy->DrawDynamicElements(&Drawer, FoundView, DrawPrimitiveFlags);
+				SubjectPrimitives[PrimitiveIndex]->Proxy->DrawDynamicElements(&Drawer, FoundView, DrawPrimitiveFlags);
+			}
 		}
 	}
 }
@@ -1798,6 +1820,7 @@ void FProjectedShadowInfo::RenderProjection(FRHICommandListImmediate& RHICmdList
 
 		// Draw the receiver's dynamic elements.
 		TDynamicPrimitiveDrawer<FDepthDrawingPolicyFactory> Drawer(RHICmdList, View, FDepthDrawingPolicyFactory::ContextType(DDM_AllOccluders), true);
+		const bool bUseGetMeshElements = ShouldUseGetDynamicMeshElements();
 
 		for (int32 PrimitiveIndex = 0; PrimitiveIndex < ReceiverPrimitives.Num(); PrimitiveIndex++)
 		{
@@ -1808,7 +1831,22 @@ void FProjectedShadowInfo::RenderProjection(FRHICommandListImmediate& RHICmdList
 				const FPrimitiveViewRelevance& ViewRelevance = View->PrimitiveViewRelevanceMap[ReceiverPrimitiveSceneInfo->GetIndex()];
 				if (ViewRelevance.bRenderInMainPass)
 				{
-					if (ViewRelevance.bDynamicRelevance)
+					if (bUseGetMeshElements)
+					{
+						FDepthDrawingPolicyFactory::ContextType Context(DDM_AllOccluders);
+
+						for (int32 MeshBatchIndex = 0; MeshBatchIndex < View->DynamicMeshElements.Num(); MeshBatchIndex++)
+						{
+							const FMeshBatchAndRelevance& MeshBatchAndRelevance = View->DynamicMeshElements[MeshBatchIndex];
+
+							if (MeshBatchAndRelevance.PrimitiveSceneProxy == ReceiverPrimitiveSceneInfo->Proxy)
+							{
+								const FMeshBatch& MeshBatch = *MeshBatchAndRelevance.Mesh;
+								FDepthDrawingPolicyFactory::DrawDynamicMesh(RHICmdList, *View, Context, MeshBatch, false, true, MeshBatchAndRelevance.PrimitiveSceneProxy, MeshBatch.BatchHitProxyId);
+							}
+						}
+					}
+					else if (ViewRelevance.bDynamicRelevance)
 					{
 						Drawer.SetPrimitive(ReceiverPrimitiveSceneInfo->Proxy);
 						ReceiverPrimitiveSceneInfo->Proxy->DrawDynamicElements(&Drawer, View);
@@ -1830,7 +1868,7 @@ void FProjectedShadowInfo::RenderProjection(FRHICommandListImmediate& RHICmdList
 									StaticMesh.Elements.Num() == 1 ? 1 : View->StaticMeshBatchVisibility[StaticMesh.Id],
 									true,
 									ReceiverPrimitiveSceneInfo->Proxy,
-									StaticMesh.HitProxyId
+									StaticMesh.BatchHitProxyId
 									);
 							}
 						}
diff --git a/Engine/Source/Runtime/Renderer/Private/ShadowRendering.h b/Engine/Source/Runtime/Renderer/Private/ShadowRendering.h
--- a/Engine/Source/Runtime/Renderer/Private/ShadowRendering.h
+++ b/Engine/Source/Runtime/Renderer/Private/ShadowRendering.h
@@ -352,8 +352,6 @@ public:
 		const ElementDataType& ElementData
 		) const;
 
-	void UnSetMeshRenderState() const;
-
 	template<bool T2>
 	friend int32 CompareDrawingPolicy(const FShadowDepthDrawingPolicy<T2>& A,const FShadowDepthDrawingPolicy<T2>& B);
 
diff --git a/Engine/Source/Runtime/Renderer/Private/TranslucentLighting.cpp b/Engine/Source/Runtime/Renderer/Private/TranslucentLighting.cpp
--- a/Engine/Source/Runtime/Renderer/Private/TranslucentLighting.cpp
+++ b/Engine/Source/Runtime/Renderer/Private/TranslucentLighting.cpp
@@ -565,6 +565,8 @@ void FProjectedShadowInfo::RenderTranslucencyDepths(FRHICommandList& RHICmdList,
 			CW_RGBA, BO_Add, BF_One, BF_One, BO_Add, BF_One, BF_One,
 			CW_RGBA, BO_Add, BF_One, BF_One, BO_Add, BF_One, BF_One>::GetRHI());
 
+		const bool bUseGetMeshElements = ShouldUseGetDynamicMeshElements();
+
 		TDynamicPrimitiveDrawer<FTranslucencyShadowDepthDrawingPolicyFactory> OpacityDrawer(RHICmdList, FoundView, FTranslucencyShadowDepthDrawingPolicyFactory::ContextType(bDirectionalLight), true);
 
 		for (int32 PrimitiveIndex = 0; PrimitiveIndex < SubjectTranslucentPrimitives.Num(); PrimitiveIndex++)
@@ -581,7 +583,22 @@ void FProjectedShadowInfo::RenderTranslucencyDepths(FRHICommandList& RHICmdList,
 
 			if(ViewRelevance.bDrawRelevance)
 			{
-				if (ViewRelevance.bDynamicRelevance)
+				if (bUseGetMeshElements)
+				{
+					FTranslucencyShadowDepthDrawingPolicyFactory::ContextType Context(bDirectionalLight);
+
+					for (int32 MeshBatchIndex = 0; MeshBatchIndex < FoundView->DynamicMeshElements.Num(); MeshBatchIndex++)
+					{
+						const FMeshBatchAndRelevance& MeshBatchAndRelevance = FoundView->DynamicMeshElements[MeshBatchIndex];
+
+						if (MeshBatchAndRelevance.PrimitiveSceneProxy == PrimitiveSceneInfo->Proxy)
+						{
+							const FMeshBatch& MeshBatch = *MeshBatchAndRelevance.Mesh;
+							FTranslucencyShadowDepthDrawingPolicyFactory::DrawDynamicMesh(RHICmdList, *FoundView, Context, MeshBatch, false, true, MeshBatchAndRelevance.PrimitiveSceneProxy, MeshBatch.BatchHitProxyId);
+						}
+					}
+				}
+				else if (ViewRelevance.bDynamicRelevance)
 				{
 					OpacityDrawer.SetPrimitive(PrimitiveSceneInfo->Proxy);
 					PrimitiveSceneInfo->Proxy->DrawDynamicElements(&OpacityDrawer, FoundView);
diff --git a/Engine/Source/Runtime/Renderer/Private/TranslucentRendering.cpp b/Engine/Source/Runtime/Renderer/Private/TranslucentRendering.cpp
--- a/Engine/Source/Runtime/Renderer/Private/TranslucentRendering.cpp
+++ b/Engine/Source/Runtime/Renderer/Private/TranslucentRendering.cpp
@@ -604,6 +604,13 @@ void FTranslucentPrimSet::DrawPrimitives(
 			RenderPrimitive(RHICmdList, View, PrimitiveSceneInfo, ViewRelevance, TranslucentSelfShadow, bSeparateTranslucencyPass);
 		}
 	}
+
+	const bool bUseGetMeshElements = ShouldUseGetDynamicMeshElements();
+
+	if (bUseGetMeshElements)
+	{
+		View.SimpleElementCollector.DrawBatchedElements(RHICmdList, View, FTexture2DRHIRef(), EBlendModeFilter::Translucent);
+	}
 }
 
 void FTranslucentPrimSet::RenderPrimitive(
@@ -619,8 +626,26 @@ void FTranslucentPrimSet::RenderPrimitive(
 
 	if(ViewRelevance.bDrawRelevance)
 	{
+		const bool bUseGetMeshElements = ShouldUseGetDynamicMeshElements();
+		
+		if (bUseGetMeshElements)
+		{
+			FTranslucencyDrawingPolicyFactory::ContextType Context(TranslucentSelfShadow, bSeparateTranslucencyPass);
+
+			//@todo parallelrendering - come up with a better way to filter these by primitive
+			for (int32 MeshBatchIndex = 0; MeshBatchIndex < View.DynamicMeshElements.Num(); MeshBatchIndex++)
+			{
+				const FMeshBatchAndRelevance& MeshBatchAndRelevance = View.DynamicMeshElements[MeshBatchIndex];
+
+				if (MeshBatchAndRelevance.PrimitiveSceneProxy == PrimitiveSceneInfo->Proxy)
+				{
+					const FMeshBatch& MeshBatch = *MeshBatchAndRelevance.Mesh;
+					FTranslucencyDrawingPolicyFactory::DrawDynamicMesh(RHICmdList, View, Context, MeshBatch, false, false, MeshBatchAndRelevance.PrimitiveSceneProxy, MeshBatch.BatchHitProxyId);
+				}
+			}
+		}
 		// Render dynamic scene prim
-		if( ViewRelevance.bDynamicRelevance )
+		else if( ViewRelevance.bDynamicRelevance )
 		{
 			TDynamicPrimitiveDrawer<FTranslucencyDrawingPolicyFactory> TranslucencyDrawer(
 				RHICmdList,
@@ -656,7 +681,7 @@ void FTranslucentPrimSet::RenderPrimitive(
 						StaticMesh.Elements.Num() == 1 ? 1 : View.StaticMeshBatchVisibility[StaticMesh.Id],
 						false,
 						PrimitiveSceneInfo->Proxy,
-						StaticMesh.HitProxyId
+						StaticMesh.BatchHitProxyId
 						);
 				}
 			}
diff --git a/Engine/Source/Runtime/Renderer/Private/VelocityRendering.cpp b/Engine/Source/Runtime/Renderer/Private/VelocityRendering.cpp
--- a/Engine/Source/Runtime/Renderer/Private/VelocityRendering.cpp
+++ b/Engine/Source/Runtime/Renderer/Private/VelocityRendering.cpp
@@ -463,6 +463,23 @@ void FDeferredShadingSceneRenderer::RenderDynamicVelocitiesInner(FRHICommandList
 	}
 }
 
+void FDeferredShadingSceneRenderer::RenderDynamicVelocitiesMeshElementsInner(FRHICommandList& RHICmdList, const FViewInfo& View)
+{
+	FVelocityDrawingPolicyFactory::ContextType Context(DDM_AllOccluders);
+
+	for (int32 MeshBatchIndex = 0; MeshBatchIndex < View.DynamicMeshElements.Num(); MeshBatchIndex++)
+	{
+		const FMeshBatchAndRelevance& MeshBatchAndRelevance = View.DynamicMeshElements[MeshBatchIndex];
+
+		if (MeshBatchAndRelevance.bHasOpaqueOrMaskedMaterial
+			&& MeshBatchAndRelevance.PrimitiveSceneProxy->GetPrimitiveSceneInfo()->ShouldRenderVelocity(View))
+		{
+			const FMeshBatch& MeshBatch = *MeshBatchAndRelevance.Mesh;
+			FVelocityDrawingPolicyFactory::DrawDynamicMesh(RHICmdList, View, Context, MeshBatch, false, true, MeshBatchAndRelevance.PrimitiveSceneProxy, MeshBatch.BatchHitProxyId);
+		}
+	}
+}
+
 class FRenderVelocityDynamicThreadTask
 {
 	FDeferredShadingSceneRenderer& ThisRenderer;
@@ -516,6 +533,8 @@ static TAutoConsoleVariable<int32> CVarParallelVelocity(
 
 void FDeferredShadingSceneRenderer::RenderVelocitiesInner(FRHICommandListImmediate& RHICmdList, const FViewInfo& View)
 {
+	const bool bUseGetMeshElements = ShouldUseGetDynamicMeshElements();
+
 	if (GRHICommandList.UseParallelAlgorithms() && CVarParallelVelocity.GetValueOnRenderThread())
 	{
 		// parallel version
@@ -526,6 +545,11 @@ void FDeferredShadingSceneRenderer::RenderVelocitiesInner(FRHICommandListImmedia
 
 		Scene->VelocityDrawList.DrawVisibleParallel(View, View.StaticMeshVelocityMap, View.StaticMeshBatchVisibility, Width, RHICmdList, OutDirty);
 
+		if (bUseGetMeshElements)
+		{
+			RenderDynamicVelocitiesMeshElementsInner(RHICmdList, View);
+		}
+		else
 		{
 			int32 NumPrims = View.VisibleDynamicPrimitives.Num();
 			int32 EffectiveThreads = FMath::Min<int32>(NumPrims, Width);
@@ -561,7 +585,15 @@ void FDeferredShadingSceneRenderer::RenderVelocitiesInner(FRHICommandListImmedia
 		// single threaded version
 		// Draw velocities for movable static meshes.
 		Scene->VelocityDrawList.DrawVisible(RHICmdList, View, View.StaticMeshVelocityMap, View.StaticMeshBatchVisibility);
-		RenderDynamicVelocitiesInner(RHICmdList, View, 0, View.VisibleDynamicPrimitives.Num() - 1);
+
+		if (bUseGetMeshElements)
+		{
+			RenderDynamicVelocitiesMeshElementsInner(RHICmdList, View);
+		}
+		else
+		{
+			RenderDynamicVelocitiesInner(RHICmdList, View, 0, View.VisibleDynamicPrimitives.Num() - 1);
+		}
 	}
 }
 
diff --git a/Engine/Source/Runtime/ShaderCore/Private/Shader.cpp b/Engine/Source/Runtime/ShaderCore/Private/Shader.cpp
--- a/Engine/Source/Runtime/ShaderCore/Private/Shader.cpp
+++ b/Engine/Source/Runtime/ShaderCore/Private/Shader.cpp
@@ -908,6 +908,7 @@ void FShader::VerifyBoundUniformBufferParameters()
 {
 	// Support being called on a NULL pointer
 // TODO: doesn't work with uniform buffer parameters on helper structs like FDeferredPixelShaderParameters
+	//@todo parallelrendering
 	if (0&&this)
 	{
 		for (int32 StructIndex = 0; StructIndex < UniformBufferParameters.Num(); StructIndex++)
diff --git a/Engine/Source/Runtime/Windows/D3D11RHI/Private/D3D11Commands.cpp b/Engine/Source/Runtime/Windows/D3D11RHI/Private/D3D11Commands.cpp
--- a/Engine/Source/Runtime/Windows/D3D11RHI/Private/D3D11Commands.cpp
+++ b/Engine/Source/Runtime/Windows/D3D11RHI/Private/D3D11Commands.cpp
@@ -1103,7 +1103,7 @@ static D3D11_PRIMITIVE_TOPOLOGY GetD3D11PrimitiveType(uint32 PrimitiveType, bool
 		case PT_TriangleStrip:
 		case PT_QuadList:
 		case PT_PointList:
-			UE_LOG(LogD3D11RHI, Fatal,TEXT("Invalid type specified for tessellated render, probably missing a case in FSkeletalMeshSceneProxy::DrawDynamicElementsByMaterial or FStaticMeshSceneProxy::GetMeshElement"));
+			UE_LOG(LogD3D11RHI, Fatal,TEXT("Invalid type specified for tessellated render, probably missing a case in FStaticMeshSceneProxy::GetMeshElement"));
 			break;
 		default:
 			// Other cases are valid.
