UE4 - remove "custom RHI commands", fix a few bugs with GDME.

diff --git a/Engine/Source/Runtime/Core/Private/Misc/MemStack.cpp b/Engine/Source/Runtime/Core/Private/Misc/MemStack.cpp
--- a/Engine/Source/Runtime/Core/Private/Misc/MemStack.cpp
+++ b/Engine/Source/Runtime/Core/Private/Misc/MemStack.cpp
@@ -37,7 +37,7 @@ int32 FMemStackBase::GetByteCount() const
 		}
 		else
 		{
-			Count += Top - Chunk->Data;
+			Count += Top - Chunk->Data();
 		}
 	}
 	return Count;
@@ -47,22 +47,31 @@ void FMemStackBase::AllocateNewChunk(int32 MinSize)
 {
 	FTaggedMemory* Chunk=nullptr;
 	// Create new chunk.
-	const int32 DataSize = AlignArbitrary<int32>(MinSize + (int32)sizeof(FTaggedMemory), FPageAllocator::PageSize) - sizeof(FTaggedMemory);
-	const uint32 AllocSize = DataSize + sizeof(FTaggedMemory);
-	if (AllocSize == FPageAllocator::PageSize)
+	int32 TotalSize = MinSize + (int32)sizeof(FTaggedMemory);
+	uint32 AllocSize;
+	if (TopChunk || TotalSize > SMALL_START_SIZE)
 	{
-		Chunk = (FTaggedMemory*)FPageAllocator::Alloc();
+		AllocSize = AlignArbitrary<int32>(TotalSize, FPageAllocator::PageSize);
+		if (AllocSize == FPageAllocator::PageSize)
+		{
+			Chunk = (FTaggedMemory*)FPageAllocator::Alloc();
+		}
+		else
+		{
+			Chunk = (FTaggedMemory*)FMemory::Malloc(AllocSize);
+			INC_MEMORY_STAT_BY(STAT_MemStackLargeBLock, AllocSize);
+		}
 	}
 	else
 	{
-		Chunk = (FTaggedMemory*)FMemory::Malloc(AllocSize);
-		INC_MEMORY_STAT_BY(STAT_MemStackLargeBLock, AllocSize);
+		AllocSize = SMALL_START_SIZE;
+		Chunk = (FTaggedMemory*)SmallStart;
 	}
-	Chunk->DataSize        = DataSize;
+	Chunk->DataSize = AllocSize - sizeof(FTaggedMemory);
 
 	Chunk->Next = TopChunk;
 	TopChunk    = Chunk;
-	Top         = Chunk->Data;
+	Top         = Chunk->Data();
 	End         = Top + Chunk->DataSize;
 }
 
@@ -72,23 +81,24 @@ void FMemStackBase::FreeChunks(FTaggedMemory* NewTopChunk)
 	{
 		FTaggedMemory* RemoveChunk = TopChunk;
 		TopChunk                   = TopChunk->Next;
-
-		if (RemoveChunk->DataSize + sizeof(FTaggedMemory) == FPageAllocator::PageSize)
+		if (RemoveChunk != (FTaggedMemory*)SmallStart)
 		{
-			FPageAllocator::Free(RemoveChunk);
+			if (RemoveChunk->DataSize + sizeof(FTaggedMemory) == FPageAllocator::PageSize)
+			{
+				FPageAllocator::Free(RemoveChunk);
+			}
+			else
+			{
+				DEC_MEMORY_STAT_BY(STAT_MemStackLargeBLock, RemoveChunk->DataSize + sizeof(FTaggedMemory));
+				FMemory::Free(RemoveChunk);
+			}
 		}
-		else
-		{
-			DEC_MEMORY_STAT_BY(STAT_MemStackLargeBLock, RemoveChunk->DataSize + sizeof(FTaggedMemory));
-			FMemory::Free(RemoveChunk);
-		}
-
 	}
 	Top = nullptr;
 	End = nullptr;
 	if( TopChunk )
 	{
-		Top = TopChunk->Data;
+		Top = TopChunk->Data();
 		End = Top + TopChunk->DataSize;
 	}
 }
@@ -98,7 +108,7 @@ bool FMemStackBase::ContainsPointer(const void* Pointer) const
 	const uint8* Ptr = (const uint8*)Pointer;
 	for (const FTaggedMemory* Chunk = TopChunk; Chunk; Chunk = Chunk->Next)
 	{
-		if (Ptr >= Chunk->Data && Ptr < Chunk->Data + Chunk->DataSize)
+		if (Ptr >= Chunk->Data() && Ptr < Chunk->Data() + Chunk->DataSize)
 		{
 			return true;
 		}
diff --git a/Engine/Source/Runtime/Core/Public/Misc/MemStack.h b/Engine/Source/Runtime/Core/Public/Misc/MemStack.h
--- a/Engine/Source/Runtime/Core/Public/Misc/MemStack.h
+++ b/Engine/Source/Runtime/Core/Public/Misc/MemStack.h
@@ -136,7 +136,11 @@ public:
 	{
 		FTaggedMemory* Next;
 		int32 DataSize;
-		uint8 Data[1];
+
+		uint8 *Data() const
+		{
+			return ((uint8*)this) + sizeof(FTaggedMemory);
+		}
 	};
 
 private:
@@ -163,6 +167,13 @@ private:
 
 	/** Used for a checkSlow. Most stacks require a mark to allocate. Command lists don't because they never mark, only flush*/
 	int32 MinMarksToAlloc;
+
+	// we are going to use an inline allocator for the first chunk to support things that are small and never get big
+	enum 
+	{
+		SMALL_START_SIZE = 1024 - sizeof(FTaggedMemory)
+	};
+	MS_ALIGN(16) uint8 SmallStart[SMALL_START_SIZE] GCC_ALIGN(16); 
 };
 
 class CORE_API FMemStack : public TThreadSingleton<FMemStack>, public FMemStackBase
diff --git a/Engine/Source/Runtime/Engine/Private/StaticBoundShaderState.cpp b/Engine/Source/Runtime/Engine/Private/StaticBoundShaderState.cpp
--- a/Engine/Source/Runtime/Engine/Private/StaticBoundShaderState.cpp
+++ b/Engine/Source/Runtime/Engine/Private/StaticBoundShaderState.cpp
@@ -86,13 +86,21 @@ FBoundShaderStateRHIParamRef FGlobalBoundShaderStateResource::GetInitializedRHI(
 	return BoundShaderState;
 }
 
+FBoundShaderStateRHIParamRef FGlobalBoundShaderStateResource::GetPreinitializedRHI()
+{
+	check(IsInitialized());
+	check(GIsRHIInitialized);
+	check(IsInParallelRenderingThread());
+	return BoundShaderState;
+}
+
 void FGlobalBoundShaderStateResource::ReleaseRHI()
 {
 	// Release the cached bound shader state.
 	BoundShaderState.SafeRelease();
 }
 
-static void SetGlobalBoundShaderState_Internal(FGlobalBoundShaderState& GlobalBoundShaderState)
+static FBoundShaderStateRHIParamRef GetGlobalBoundShaderState_Internal(FGlobalBoundShaderState& GlobalBoundShaderState)
 {
 	// Check for unset uniform buffer parameters
 	// Technically you can set uniform buffer parameters after calling RHISetBoundShaderState, but this is the most global place to check for unset parameters
@@ -101,35 +109,56 @@ static void SetGlobalBoundShaderState_Internal(FGlobalBoundShaderState& GlobalBo
 	GlobalBoundShaderState.Get()->Args.GeometryShader->VerifyBoundUniformBufferParameters();
 
 	FGlobalBoundShaderState_Internal* BSS = GlobalBoundShaderState.Get()->BSS;
+	bool bNewBSS = false;
 	if (!BSS)
 	{
 		BSS = new FGlobalBoundShaderState_Internal(); // these are simply leaked and never freed
-		GlobalBoundShaderState.Get()->BSS = BSS;
+		bNewBSS = true;
 	}
-
-	SetBoundShaderState_Internal(
-		BSS->GetInitializedRHI(
-			GlobalBoundShaderState.Get()->Args.VertexDeclarationRHI,
-			GETSAFERHISHADER_VERTEX(GlobalBoundShaderState.Get()->Args.VertexShader),
-			GETSAFERHISHADER_PIXEL(GlobalBoundShaderState.Get()->Args.PixelShader),
-			(FGeometryShaderRHIParamRef)GETSAFERHISHADER_GEOMETRY(GlobalBoundShaderState.Get()->Args.GeometryShader))
-		);
+	FBoundShaderStateRHIParamRef Result = BSS->GetInitializedRHI(
+		GlobalBoundShaderState.Get()->Args.VertexDeclarationRHI,
+		GETSAFERHISHADER_VERTEX(GlobalBoundShaderState.Get()->Args.VertexShader),
+		GETSAFERHISHADER_PIXEL(GlobalBoundShaderState.Get()->Args.PixelShader),
+		(FGeometryShaderRHIParamRef)GETSAFERHISHADER_GEOMETRY(GlobalBoundShaderState.Get()->Args.GeometryShader));
+	if (bNewBSS)
+	{
+		FPlatformMisc::MemoryBarrier();
+		// this probably doesn't need to be an atomic, assuming a pointer write is atomic, but this checks alignment, etc
+		FGlobalBoundShaderState_Internal* OldBSS = (FGlobalBoundShaderState_Internal*)FPlatformAtomics::InterlockedCompareExchangePointer((void**)&GlobalBoundShaderState.Get()->BSS, BSS, nullptr);
+		check(!OldBSS); // should not be any concurrent writes
+	}
+	return Result;
 }
 
-
-struct FRHICommandSetGlobalBoundShaderState : public FRHICommand<FRHICommandSetGlobalBoundShaderState>
+class FSetGlobalBoundShaderStateRenderThreadTask
 {
-	FGlobalBoundShaderState GlobalBoundShaderState;
-	FORCEINLINE_DEBUGGABLE FRHICommandSetGlobalBoundShaderState(FGlobalBoundShaderState& InGlobalBoundShaderState)
-		: GlobalBoundShaderState(InGlobalBoundShaderState)
+	FRHICommandList& RHICmdList;
+	FGlobalBoundShaderState& GlobalBoundShaderState;
+public:
+
+	FSetGlobalBoundShaderStateRenderThreadTask(FRHICommandList* InRHICmdList, FGlobalBoundShaderState* InGlobalBoundShaderState)
+		: RHICmdList(*InRHICmdList)
+		, GlobalBoundShaderState(*InGlobalBoundShaderState)
 	{
 	}
-	void Execute()
+
+	FORCEINLINE TStatId GetStatId() const
 	{
-		SetGlobalBoundShaderState_Internal(GlobalBoundShaderState);
+		RETURN_QUICK_DECLARE_CYCLE_STAT(FSetGlobalBoundShaderStateRenderThreadTask, STATGROUP_TaskGraphTasks);
 	}
-};
 
+	ENamedThreads::Type GetDesiredThread()
+	{
+		return ENamedThreads::RenderThread_Local;
+	}
+
+	static ESubsequentsMode::Type GetSubsequentsMode() { return ESubsequentsMode::TrackSubsequents; }
+
+	void DoTask(ENamedThreads::Type CurrentThread, const FGraphEventRef& MyCompletionGraphEvent)
+	{
+		RHICmdList.SetBoundShaderState(GetGlobalBoundShaderState_Internal(GlobalBoundShaderState));
+	}
+};
 
 void SetGlobalBoundShaderState(
 	FRHICommandList& RHICmdList,
@@ -156,7 +185,13 @@ void SetGlobalBoundShaderState(
 		{
 			//we lost
 			delete NewGlobalBoundShaderState;
-			check(OldGlobalBoundShaderState == ExistingGlobalBoundShaderState);
+			check(OldGlobalBoundShaderState == GlobalBoundShaderState.Get());
+			ExistingGlobalBoundShaderState = OldGlobalBoundShaderState;
+		}
+		else
+		{
+			check(NewGlobalBoundShaderState == GlobalBoundShaderState.Get());
+			ExistingGlobalBoundShaderState = NewGlobalBoundShaderState;
 		}
 	}
 	else if (!(VertexDeclarationRHI == ExistingGlobalBoundShaderState->Args.VertexDeclarationRHI &&
@@ -171,12 +206,27 @@ void SetGlobalBoundShaderState(
 		ExistingGlobalBoundShaderState->Args.GeometryShader = GeometryShader;
 	}
 
-	if (RHICmdList.Bypass())
+	if (RHICmdList.Bypass() || IsInRenderingThread())
 	{
-		SetGlobalBoundShaderState_Internal(GlobalBoundShaderState);
+		// we can just do it as usual
+		RHICmdList.SetBoundShaderState(GetGlobalBoundShaderState_Internal(GlobalBoundShaderState));
 		return;
 	}
-	new (RHICmdList.AllocCommand<FRHICommandSetGlobalBoundShaderState>()) FRHICommandSetGlobalBoundShaderState(GlobalBoundShaderState);
+	if (ExistingGlobalBoundShaderState->BSS)
+	{
+		FBoundShaderStateRHIParamRef BoundShaderState = ExistingGlobalBoundShaderState->BSS->GetPreinitializedRHI();
+		if (BoundShaderState)
+		{
+			// this is on a thread, but we already have a valid BSS
+			RHICmdList.SetBoundShaderState(BoundShaderState);
+			return;
+		}
+	}
+	// We need to do this on the render thread
+
+	FRHICommandList* CmdList = new FRHICommandList;
+	FGraphEventRef RenderThreadCompletionEvent = TGraphTask<FSetGlobalBoundShaderStateRenderThreadTask>::CreateTask().ConstructAndDispatchWhenReady(CmdList, &GlobalBoundShaderState);
+	RHICmdList.QueueRenderThreadCommandListSubmit(RenderThreadCompletionEvent, CmdList);
 }
 
 
diff --git a/Engine/Source/Runtime/Engine/Public/StaticBoundShaderState.h b/Engine/Source/Runtime/Engine/Public/StaticBoundShaderState.h
--- a/Engine/Source/Runtime/Engine/Public/StaticBoundShaderState.h
+++ b/Engine/Source/Runtime/Engine/Public/StaticBoundShaderState.h
@@ -36,6 +36,12 @@ public:
 		FGeometryShaderRHIParamRef GeometryShader
 		);
 
+	/**
+	 * If this global bound shader state has been initialized return it, otherwise return null. Can be called from any thread.
+	 * @return The bound shader state RHI.
+	 */
+	FBoundShaderStateRHIParamRef GetPreinitializedRHI();
+
 private:
 
 	/** The cached bound shader state. */
diff --git a/Engine/Source/Runtime/RHI/Private/RHICommandList.cpp b/Engine/Source/Runtime/RHI/Private/RHICommandList.cpp
--- a/Engine/Source/Runtime/RHI/Private/RHICommandList.cpp
+++ b/Engine/Source/Runtime/RHI/Private/RHICommandList.cpp
@@ -8,7 +8,7 @@ DECLARE_CYCLE_STAT(TEXT("Nonimmed. Command List Execute"), STAT_NonImmedCmdListE
 DECLARE_DWORD_COUNTER_STAT(TEXT("Nonimmed. Command List memory"), STAT_NonImmedCmdListMemory, STATGROUP_RHICMDLIST);
 DECLARE_DWORD_COUNTER_STAT(TEXT("Nonimmed. Command count"), STAT_NonImmedCmdListCount, STATGROUP_RHICMDLIST);
 
-DECLARE_CYCLE_STAT(TEXT("Immed. Command List Execute"), STAT_ImmedCmdListExecuteTime, STATGROUP_RHICMDLIST);
+DECLARE_CYCLE_STAT(TEXT("All Command List Execute"), STAT_ImmedCmdListExecuteTime, STATGROUP_RHICMDLIST);
 DECLARE_DWORD_COUNTER_STAT(TEXT("Immed. Command List memory"), STAT_ImmedCmdListMemory, STATGROUP_RHICMDLIST);
 DECLARE_DWORD_COUNTER_STAT(TEXT("Immed. Command count"), STAT_ImmedCmdListCount, STATGROUP_RHICMDLIST);
 
@@ -95,6 +95,7 @@ void FRHICommandListExecutor::ExecuteList(FRHICommandListImmediate& CmdList)
 #endif
 	{
 		SCOPE_CYCLE_COUNTER(STAT_ImmedCmdListExecuteTime);
+		//FTaskGraphInterface::Get().ProcessThreadUntilIdle(ENamedThreads::RenderThread_Local); 
 		ExecuteInner(CmdList);
 	}
 
@@ -173,9 +174,9 @@ void FRHICommandListBase::Reset()
 	UID = GRHICommandList.UIDCounter.Increment();
 }
 
-DECLARE_DWORD_COUNTER_STAT(TEXT("Num Chains Links"), STAT_ChainLinkCount, STATGROUP_RHICMDLIST);
-DECLARE_CYCLE_STAT(TEXT("Wait for CmdList"), STAT_ChainWait, STATGROUP_RHICMDLIST);
-DECLARE_CYCLE_STAT(TEXT("Chain Execute"), STAT_ChainExecute, STATGROUP_RHICMDLIST);
+DECLARE_DWORD_COUNTER_STAT(TEXT("Num Async Chains Links"), STAT_ChainLinkCount, STATGROUP_RHICMDLIST);
+DECLARE_CYCLE_STAT(TEXT("Wait for Async CmdList"), STAT_ChainWait, STATGROUP_RHICMDLIST);
+DECLARE_CYCLE_STAT(TEXT("Async Chain Execute"), STAT_ChainExecute, STATGROUP_RHICMDLIST);
 
 struct FRHICommandWaitForAndSubmitSubList : public FRHICommand<FRHICommandWaitForAndSubmitSubList>
 {
@@ -200,11 +201,84 @@ struct FRHICommandWaitForAndSubmitSubList : public FRHICommand<FRHICommandWaitFo
 	}
 };
 
-void FRHICommandListBase::QueueAsyncCommandListSubmit(FGraphEventRef AnyThreadCompletionEvent, class FRHICommandList* CmdList)
+void FRHICommandListBase::QueueAsyncCommandListSubmit(FGraphEventRef& AnyThreadCompletionEvent, class FRHICommandList* CmdList)
 {
 	new (AllocCommand<FRHICommandWaitForAndSubmitSubList>()) FRHICommandWaitForAndSubmitSubList(AnyThreadCompletionEvent, CmdList);
 }
 
+DECLARE_DWORD_COUNTER_STAT(TEXT("Num RT Chains Links"), STAT_RTChainLinkCount, STATGROUP_RHICMDLIST);
+DECLARE_CYCLE_STAT(TEXT("Wait for RT CmdList"), STAT_RTChainWait, STATGROUP_RHICMDLIST);
+DECLARE_CYCLE_STAT(TEXT("RT Chain Execute"), STAT_RTChainExecute, STATGROUP_RHICMDLIST);
+
+struct FRHICommandWaitForAndSubmitRTSubList : public FRHICommand<FRHICommandWaitForAndSubmitRTSubList>
+{
+	FGraphEventRef EventToWaitFor;
+	FRHICommandList* RHICmdList;
+	FORCEINLINE_DEBUGGABLE FRHICommandWaitForAndSubmitRTSubList(FGraphEventRef& InEventToWaitFor, FRHICommandList* InRHICmdList)
+		: EventToWaitFor(InEventToWaitFor)
+		, RHICmdList(InRHICmdList)
+	{
+	}
+	void Execute()
+	{
+		INC_DWORD_STAT_BY(STAT_RTChainLinkCount, 1);
+		{
+			SCOPE_CYCLE_COUNTER(STAT_RTChainWait);
+			if (IsInRenderingThread())
+			{
+				if (!EventToWaitFor->IsComplete())
+				{
+					if (FTaskGraphInterface::Get().IsThreadProcessingTasks(ENamedThreads::RenderThread_Local))
+					{
+						// this is a deadlock. RT tasks must be done by now or they won't be done. We could add a third queue...
+						UE_LOG(LogRHI, Fatal, TEXT("Deadlock in command list processing."));
+					}
+#if 0
+					while (!EventToWaitFor->IsComplete())
+					{
+						FTaskGraphInterface::Get().ProcessThreadUntilIdle(ENamedThreads::RenderThread_Local);
+					}
+#endif
+					FTaskGraphInterface::Get().WaitUntilTaskCompletes(EventToWaitFor, ENamedThreads::RenderThread_Local);
+				}
+			}
+			else
+			{
+				FScopedEvent Waiter;
+				FTaskGraphInterface::Get().TriggerEventWhenTaskCompletes(Waiter.Get(), EventToWaitFor);
+			}
+		}
+		{
+			SCOPE_CYCLE_COUNTER(STAT_RTChainExecute);
+			delete RHICmdList;
+		}
+	}
+};
+
+void FRHICommandListBase::QueueRenderThreadCommandListSubmit(FGraphEventRef& RenderThreadCompletionEvent, class FRHICommandList* CmdList)
+{
+	new (AllocCommand<FRHICommandWaitForAndSubmitRTSubList>()) FRHICommandWaitForAndSubmitRTSubList(RenderThreadCompletionEvent, CmdList);
+}
+
+struct FRHICommandSubmitSubList : public FRHICommand<FRHICommandSubmitSubList>
+{
+	FRHICommandList* RHICmdList;
+	FORCEINLINE_DEBUGGABLE FRHICommandSubmitSubList(FRHICommandList* InRHICmdList)
+		: RHICmdList(InRHICmdList)
+	{
+	}
+	void Execute()
+	{
+		INC_DWORD_STAT_BY(STAT_ChainLinkCount, 1);
+		SCOPE_CYCLE_COUNTER(STAT_ChainExecute);
+		delete RHICmdList;
+	}
+};
+
+void FRHICommandListBase::QueueCommandListSubmit(class FRHICommandList* CmdList)
+{
+	new (AllocCommand<FRHICommandSubmitSubList>()) FRHICommandSubmitSubList(CmdList);
+}
 
 static TLockFreeFixedSizeAllocator<sizeof(FRHICommandList), FThreadSafeCounter> RHICommandListAllocator;
 
diff --git a/Engine/Source/Runtime/RHI/Public/RHICommandList.h b/Engine/Source/Runtime/RHI/Public/RHICommandList.h
--- a/Engine/Source/Runtime/RHI/Public/RHICommandList.h
+++ b/Engine/Source/Runtime/RHI/Public/RHICommandList.h
@@ -37,7 +37,9 @@ public:
 	inline void Flush();
 	inline bool IsImmediate();
 	const int32 GetUsedMemory() const;
-	void QueueAsyncCommandListSubmit(FGraphEventRef AnyThreadCompletionEvent, class FRHICommandList* CmdList);
+	void QueueAsyncCommandListSubmit(FGraphEventRef& AnyThreadCompletionEvent, class FRHICommandList* CmdList);
+	void QueueRenderThreadCommandListSubmit(FGraphEventRef& RenderThreadCompletionEvent, class FRHICommandList* CmdList);
+	void QueueCommandListSubmit(class FRHICommandList* CmdList);
 
 	FORCEINLINE_DEBUGGABLE void* Alloc(int32 AllocSize, int32 Alignment)
 	{
diff --git a/Engine/Source/Runtime/RenderCore/Public/UniformBuffer.h b/Engine/Source/Runtime/RenderCore/Public/UniformBuffer.h
--- a/Engine/Source/Runtime/RenderCore/Public/UniformBuffer.h
+++ b/Engine/Source/Runtime/RenderCore/Public/UniformBuffer.h
@@ -48,6 +48,7 @@ public:
 	// FRenderResource interface.
 	virtual void InitDynamicRHI() override
 	{
+		check(IsInRenderingThread());
 		UniformBufferRHI = RHICreateUniformBuffer(Contents,TBufferStruct::StaticStruct.GetLayout(),BufferUsage);
 	}
 	virtual void ReleaseDynamicRHI()
@@ -86,6 +87,11 @@ public:
 		check(IsInRenderingThread());
 		return TUniformBufferRef<TBufferStruct>(RHICreateUniformBuffer(&Value,TBufferStruct::StaticStruct.GetLayout(),Usage));
 	}
+	/** Creates a uniform buffer with the given value, and returns a structured reference to it. */
+	static FLocalUniformBuffer CreateLocalUniformBuffer(FRHICommandList& RHICmdList, const TBufferStruct& Value, EUniformBufferUsage Usage)
+	{
+		return RHICmdList.BuildLocalUniformBuffer(&Value, sizeof(TBufferStruct), TBufferStruct::StaticStruct.GetLayout());
+	}
 
 private:
 
diff --git a/Engine/Source/Runtime/Renderer/Private/DeferredShadingRenderer.cpp b/Engine/Source/Runtime/Renderer/Private/DeferredShadingRenderer.cpp
--- a/Engine/Source/Runtime/Renderer/Private/DeferredShadingRenderer.cpp
+++ b/Engine/Source/Runtime/Renderer/Private/DeferredShadingRenderer.cpp
@@ -358,7 +358,7 @@ void FDeferredShadingSceneRenderer::RenderBasePassDynamicData(FRHICommandList& R
 	SCOPE_CYCLE_COUNTER(STAT_DynamicPrimitiveDrawTime);
 	SCOPED_DRAW_EVENT(Dynamic, DEC_SCENE_ITEMS);
 
-	const bool bUseGetMeshElements = CVarUseGetMeshElements.GetValueOnRenderThread() != 0;
+	const bool bUseGetMeshElements = ShouldUseGetDynamicMeshElements();
 
 	if (bUseGetMeshElements)
 	{
@@ -505,26 +505,10 @@ static void SetupBasePassView(FRHICommandList& RHICmdList, const FIntRect& ViewR
 	RHICmdList.SetViewport(ViewRect.Min.X, ViewRect.Min.Y, 0, ViewRect.Max.X, ViewRect.Max.Y, 1);
 }
 
-struct FRHICommandSetupBasePassView : public FRHICommand<FRHICommandSetupBasePassView>
-{
-	FIntRect ViewRect;
-	bool bShaderComplexity;
-	FORCEINLINE_DEBUGGABLE FRHICommandSetupBasePassView(const FViewInfo& InView, bool InbShaderComplexity)
-		: ViewRect(InView.ViewRect)
-		, bShaderComplexity(InbShaderComplexity)
-	{
-	}
-	void Execute()
-	{
-		FRHICommandList RHICmdList;
-		SetupBasePassView(RHICmdList, ViewRect, bShaderComplexity);
-	}
-};
-
 void FDeferredShadingSceneRenderer::RenderBasePassViewParallel(FViewInfo& View, int32 Width, FRHICommandList& ParentCmdList, bool& OutDirty)
 {
 	// this is probably only needed on the first view
-	new (ParentCmdList.AllocCommand<FRHICommandSetupBasePassView>()) FRHICommandSetupBasePassView(View, !!ViewFamily.EngineShowFlags.ShaderComplexity);
+	SetupBasePassView(ParentCmdList, View.ViewRect, !!ViewFamily.EngineShowFlags.ShaderComplexity);
 
 	RenderBasePassStaticDataParallel(View, Width, ParentCmdList, OutDirty);
 	RenderBasePassDynamicDataParallel(View, ParentCmdList, OutDirty);
@@ -965,7 +949,7 @@ void FDeferredShadingSceneRenderer::Render(FRHICommandListImmediate& RHICmdList)
 
 bool FDeferredShadingSceneRenderer::RenderPrePassViewDynamic(FRHICommandList& RHICmdList, const FViewInfo& View)
 {
-	const bool bUseGetMeshElements = CVarUseGetMeshElements.GetValueOnRenderThread() != 0;
+	const bool bUseGetMeshElements = ShouldUseGetDynamicMeshElements();
 
 	if (bUseGetMeshElements)
 	{
@@ -1084,21 +1068,6 @@ bool FDeferredShadingSceneRenderer::RenderPrePassView(FRHICommandList& RHICmdLis
 	return bDirty;
 }
 
-struct FRHICommandSetupPrepassView : public FRHICommand<FRHICommandSetupPrepassView>
-{
-	FIntRect ViewRect;
-	FORCEINLINE_DEBUGGABLE FRHICommandSetupPrepassView(const FViewInfo& InView)
-		: ViewRect(InView.ViewRect)
-	{
-	}
-	void Execute()
-	{
-		FRHICommandList RHICmdList;
-		SetupPrePassView(RHICmdList, ViewRect);
-	}
-};
-
-
 class FRenderPrepassDynamicDataThreadTask
 {
 	FDeferredShadingSceneRenderer& ThisRenderer;
@@ -1145,8 +1114,7 @@ public:
 
 void FDeferredShadingSceneRenderer::RenderPrePassViewParallel(const FViewInfo& View, int32 Width, FRHICommandList& ParentCmdList, bool& OutDirty)
 {
-	// this is probably only needed on the first view
-	new (ParentCmdList.AllocCommand<FRHICommandSetupPrepassView>()) FRHICommandSetupPrepassView(View);
+	SetupPrePassView(ParentCmdList, View.ViewRect);
 
 	// Draw the static occluder primitives using a depth drawing policy.
 	{
diff --git a/Engine/Source/Runtime/Renderer/Private/SceneRendering.cpp b/Engine/Source/Runtime/Renderer/Private/SceneRendering.cpp
--- a/Engine/Source/Runtime/Renderer/Private/SceneRendering.cpp
+++ b/Engine/Source/Runtime/Renderer/Private/SceneRendering.cpp
@@ -1069,6 +1069,6 @@ bool IsMobileHDR32bpp()
 bool ShouldUseGetDynamicMeshElements()
 {
 	static const auto CVarUseGetMeshElements = IConsoleManager::Get().FindTConsoleVariableDataInt(TEXT("r.UseGetMeshElements"));
-	const bool bUseGetMeshElements = CVarUseGetMeshElements->GetValueOnRenderThread() != 0;
+	const bool bUseGetMeshElements = !GRHICommandList.Bypass() ||  CVarUseGetMeshElements->GetValueOnRenderThread() != 0;
 	return bUseGetMeshElements;
 }
\ No newline at end of file
diff --git a/Engine/Source/Runtime/Renderer/Private/SceneVisibility.cpp b/Engine/Source/Runtime/Renderer/Private/SceneVisibility.cpp
--- a/Engine/Source/Runtime/Renderer/Private/SceneVisibility.cpp
+++ b/Engine/Source/Runtime/Renderer/Private/SceneVisibility.cpp
@@ -796,6 +796,7 @@ void FSceneRenderer::GatherDynamicMeshElements(
 		}
 	}
 
+	if (GIsEditor)
 	{
 		Collector.ClearViewMeshArrays();
 
diff --git a/Engine/Source/Runtime/Renderer/Private/TranslucentRendering.cpp b/Engine/Source/Runtime/Renderer/Private/TranslucentRendering.cpp
--- a/Engine/Source/Runtime/Renderer/Private/TranslucentRendering.cpp
+++ b/Engine/Source/Runtime/Renderer/Private/TranslucentRendering.cpp
@@ -303,29 +303,45 @@ public:
 
 static void CopySceneColorAndRestore(FRHICommandList& RHICmdList, const FViewInfo& View, const FPrimitiveSceneProxy* PrimitiveSceneProxy)
 {
+	check(IsInRenderingThread());
 	FTranslucencyDrawingPolicyFactory::CopySceneColor(RHICmdList, View, PrimitiveSceneProxy);
 	// Restore state
 	SetTranslucentRenderTargetAndState(RHICmdList, View, false);
 }
 
-struct FRHICommandCopySceneColor : public FRHICommand<FRHICommandCopySceneColor>
+class FCopySceneColorAndRestoreRenderThreadTask
 {
+	FRHICommandList& RHICmdList;
 	const FViewInfo& View;
 	const FPrimitiveSceneProxy* PrimitiveSceneProxy;
+public:
 
-	FORCEINLINE_DEBUGGABLE FRHICommandCopySceneColor(const FViewInfo& InView, const FPrimitiveSceneProxy* InPrimitiveSceneProxy)
-		: View(InView)
+	FCopySceneColorAndRestoreRenderThreadTask(FRHICommandList* InRHICmdList, const FViewInfo* InView, const FPrimitiveSceneProxy* InPrimitiveSceneProxy)
+		: RHICmdList(*InRHICmdList)
+		, View(*InView)
 		, PrimitiveSceneProxy(InPrimitiveSceneProxy)
 	{
 	}
-	void Execute()
+
+	FORCEINLINE TStatId GetStatId() const
+	{
+		RETURN_QUICK_DECLARE_CYCLE_STAT(FSetGlobalBoundShaderStateRenderThreadTask, STATGROUP_TaskGraphTasks);
+	}
+
+	ENamedThreads::Type GetDesiredThread()
+	{
+		return ENamedThreads::RenderThread_Local;
+	}
+
+	static ESubsequentsMode::Type GetSubsequentsMode() { return ESubsequentsMode::TrackSubsequents; }
+
+	void DoTask(ENamedThreads::Type CurrentThread, const FGraphEventRef& MyCompletionGraphEvent)
 	{
-		check(IsInRenderingThread());
-		FRHICommandList RHICmdList;
 		CopySceneColorAndRestore(RHICmdList, View, PrimitiveSceneProxy);
 	}
 };
 
+
 /**
 * Render a dynamic or static mesh using a translucent draw policy
 * @return true if the mesh rendered
@@ -357,11 +373,11 @@ bool FTranslucencyDrawingPolicyFactory::DrawMesh(
 		{
 			if (DrawingContext.bSceneColorCopyIsUpToDate == false)
 			{
-				FRHICommandListImmediate& RHICmdListIm = FRHICommandListExecutor::GetImmediateCommandList();
-				if (&RHICmdList != &RHICmdListIm)
+				if (!RHICmdList.Bypass() && !IsInRenderingThread())
 				{
-					// if we are buffering commands on a non-immediate command list, we need to defer the copy until commandlist execution
-					new (RHICmdList.AllocCommand<FRHICommandCopySceneColor>()) FRHICommandCopySceneColor(View, PrimitiveSceneProxy);
+					FRHICommandList* CmdList = new FRHICommandList;
+					FGraphEventRef RenderThreadCompletionEvent = TGraphTask<FCopySceneColorAndRestoreRenderThreadTask>::CreateTask().ConstructAndDispatchWhenReady(CmdList, &View, PrimitiveSceneProxy);
+					RHICmdList.QueueRenderThreadCommandListSubmit(RenderThreadCompletionEvent, CmdList);
 				}
 				else
 				{
@@ -516,26 +532,41 @@ void FTranslucentPrimSet::DrawAPrimitive(
 	RenderPrimitive(RHICmdList, View, PrimitiveSceneInfo, ViewRelevance, TranslucentSelfShadow, bSeparateTranslucencyPass);
 }
 
-struct FRHICommandVolumetricTranslucentShadow : public FRHICommand<FRHICommandVolumetricTranslucentShadow>
+class FVolumetricTranslucentShadowRenderThreadTask
 {
+	FRHICommandList& RHICmdList;
 	const FTranslucentPrimSet &PrimSet;
 	const FViewInfo& View;
 	FDeferredShadingSceneRenderer& Renderer;
 	bool bSeparateTranslucencyPass;
 	int32 Index;
 
-	FORCEINLINE_DEBUGGABLE FRHICommandVolumetricTranslucentShadow(const FTranslucentPrimSet &InPrimSet, const FViewInfo& InView, FDeferredShadingSceneRenderer& InRenderer, bool InbSeparateTranslucencyPass, int32 InIndex)
-		: PrimSet(InPrimSet)
-		, View(InView)
-		, Renderer(InRenderer)
+public:
+
+	FORCEINLINE_DEBUGGABLE FVolumetricTranslucentShadowRenderThreadTask(FRHICommandList* InRHICmdList, const FTranslucentPrimSet* InPrimSet, const FViewInfo* InView, FDeferredShadingSceneRenderer* InRenderer, bool InbSeparateTranslucencyPass, int32 InIndex)
+		: RHICmdList(*InRHICmdList)
+		, PrimSet(*InPrimSet)
+		, View(*InView)
+		, Renderer(*InRenderer)
 		, bSeparateTranslucencyPass(InbSeparateTranslucencyPass)
 		, Index(InIndex)
 	{
 	}
-	void Execute()
+
+	FORCEINLINE TStatId GetStatId() const
+	{
+		RETURN_QUICK_DECLARE_CYCLE_STAT(FVolumetricTranslucentShadowRenderThreadTask, STATGROUP_TaskGraphTasks);
+	}
+
+	ENamedThreads::Type GetDesiredThread()
+	{
+		return ENamedThreads::RenderThread_Local;
+	}
+
+	static ESubsequentsMode::Type GetSubsequentsMode() { return ESubsequentsMode::TrackSubsequents; }
+
+	void DoTask(ENamedThreads::Type CurrentThread, const FGraphEventRef& MyCompletionGraphEvent)
 	{
-		check(IsInRenderingThread());
-		FRHICommandList RHICmdList;
 		PrimSet.DrawAPrimitive(RHICmdList, View, Renderer, bSeparateTranslucencyPass, Index);
 	}
 };
@@ -566,7 +597,9 @@ void FTranslucentPrimSet::DrawPrimitivesParallel(
 		if (PrimitiveSceneInfo->Proxy && PrimitiveSceneInfo->Proxy->CastsVolumetricTranslucentShadow())
 		{
 			// can't do this in parallel, defer
-			new (RHICmdList.AllocCommand<FRHICommandVolumetricTranslucentShadow>()) FRHICommandVolumetricTranslucentShadow(*this, View, Renderer, bSeparateTranslucencyPass, PrimIdx);
+			FRHICommandList* CmdList = new FRHICommandList;
+			FGraphEventRef RenderThreadCompletionEvent = TGraphTask<FVolumetricTranslucentShadowRenderThreadTask>::CreateTask().ConstructAndDispatchWhenReady(CmdList, this, &View, &Renderer, bSeparateTranslucencyPass, PrimIdx);
+			RHICmdList.QueueRenderThreadCommandListSubmit(RenderThreadCompletionEvent, CmdList);
 		}
 		else
 		{
@@ -747,59 +780,6 @@ bool FSceneRenderer::ShouldRenderTranslucency() const
 	return bRender;
 }
 
-struct FRHICommandSetTranslucentRenderTargetAndState : public FRHICommand<FRHICommandSetTranslucentRenderTargetAndState>
-{
-	const FViewInfo& View;
-	bool bSeparateTranslucencyPass;
-
-	FORCEINLINE_DEBUGGABLE FRHICommandSetTranslucentRenderTargetAndState(const FViewInfo& InView, bool InbSeparateTranslucencyPass)
-		: View(InView)
-		, bSeparateTranslucencyPass(InbSeparateTranslucencyPass)
-	{
-	}
-	void Execute()
-	{
-		check(IsInRenderingThread());
-		FRHICommandList RHICmdList;
-		if (!bSeparateTranslucencyPass)
-		{
-			SetTranslucentRenderTargetAndState(RHICmdList, View, false);
-		}
-		else
-		{
-			bool bRenderSeparateTranslucency = View.TranslucentPrimSet.NumSeparateTranslucencyPrims() > 0;
-
-			// always call BeginRenderingSeparateTranslucency() even if there are no primitives to we keep the RT allocated
-			bool bSetupTranslucency = GSceneRenderTargets.BeginRenderingSeparateTranslucency(RHICmdList, View, true);
-
-			// Draw only translucent prims that are in the SeparateTranslucency pass
-			if (bRenderSeparateTranslucency)
-			{
-				if (bSetupTranslucency)
-				{
-					RHICmdList.SetDepthStencilState(TStaticDepthStencilState<false, CF_GreaterEqual>::GetRHI());
-				}
-			}
-		}
-	}
-};
-
-struct FRHICommandFinishRenderingSeparateTranslucency : public FRHICommand<FRHICommandFinishRenderingSeparateTranslucency>
-{
-	const FViewInfo& View;
-
-	FORCEINLINE_DEBUGGABLE FRHICommandFinishRenderingSeparateTranslucency(const FViewInfo& InView)
-		: View(InView)
-	{
-	}
-	void Execute()
-	{
-		check(IsInRenderingThread());
-		FRHICommandList RHICmdList;
-		GSceneRenderTargets.FinishRenderingSeparateTranslucency(RHICmdList, View);
-	}
-};
-
 class FDrawSortedTransAnyThreadTask
 {
 	FDeferredShadingSceneRenderer& Renderer;
@@ -843,18 +823,7 @@ public:
 
 	void DoTask(ENamedThreads::Type CurrentThread, const FGraphEventRef& MyCompletionGraphEvent)
 	{
-		if (FirstIndex == 0)
-		{
-			new (RHICmdList.AllocCommand<FRHICommandSetTranslucentRenderTargetAndState>()) FRHICommandSetTranslucentRenderTargetAndState(View, bSeparateTranslucencyPass);
-		}
-		if (LastIndex >= FirstIndex)
-		{
-			View.TranslucentPrimSet.DrawPrimitivesParallel(RHICmdList, View, Renderer, bSeparateTranslucencyPass, FirstIndex, LastIndex);
-		}
-		if (bSeparateTranslucencyPass && (LastIndex < FirstIndex || LastIndex == View.TranslucentPrimSet.NumSeparateTranslucencyPrims() - 1))
-		{
-			new (RHICmdList.AllocCommand<FRHICommandFinishRenderingSeparateTranslucency>()) FRHICommandFinishRenderingSeparateTranslucency(View);
-		}
+		View.TranslucentPrimSet.DrawPrimitivesParallel(RHICmdList, View, Renderer, bSeparateTranslucencyPass, FirstIndex, LastIndex);
 	}
 };
 
@@ -875,6 +844,8 @@ void FDeferredShadingSceneRenderer::RenderTranslucencyParallel(FRHICommandListIm
 
 		const FViewInfo& View = Views[ViewIndex];
 
+		SetTranslucentRenderTargetAndState(RHICmdList, View, false);
+
 		{
 			int32 NumPrims = View.TranslucentPrimSet.NumPrims() - View.TranslucentPrimSet.NumSeparateTranslucencyPrims();
 			int32 EffectiveThreads = FMath::Min<int32>(NumPrims, Width);
@@ -903,16 +874,6 @@ void FDeferredShadingSceneRenderer::RenderTranslucencyParallel(FRHICommandListIm
 					Start = Last + 1;
 				}
 			}
-			else
-			{
-				// still need to setup the render target, a bit of a waste...
-				FRHICommandList* CmdList = new FRHICommandList;
-
-				FGraphEventRef AnyThreadCompletionEvent = TGraphTask<FDrawSortedTransAnyThreadTask>::CreateTask(nullptr, ENamedThreads::RenderThread)
-					.ConstructAndDispatchWhenReady(this, CmdList, &View, false, 0, -1);
-
-				RHICmdList.QueueAsyncCommandListSubmit(AnyThreadCompletionEvent, CmdList);
-			}
 		}
 		// Draw the view's mesh elements with the translucent drawing policy.
 		DrawViewElementsParallel<FTranslucencyDrawingPolicyFactory>(View, ThisContext, SDPG_World, false, RHICmdList, Width);
@@ -932,41 +893,48 @@ void FDeferredShadingSceneRenderer::RenderTranslucencyParallel(FRHICommandListIm
 		}
 #endif
 		{
-			int32 NumPrims = View.TranslucentPrimSet.NumSeparateTranslucencyPrims();
-			int32 EffectiveThreads = FMath::Min<int32>(NumPrims, Width);
+			bool bRenderSeparateTranslucency = View.TranslucentPrimSet.NumSeparateTranslucencyPrims() > 0;
 
-			int32 Start = 0;
-			if (EffectiveThreads)
 			{
+				// always call BeginRenderingSeparateTranslucency() even if there are no primitives to we keep the RT allocated
+				bool bSetupTranslucency = GSceneRenderTargets.BeginRenderingSeparateTranslucency(RHICmdList, View, true);
+				if (bRenderSeparateTranslucency && bSetupTranslucency)
+				{
+					RHICmdList.SetDepthStencilState(TStaticDepthStencilState<false, CF_GreaterEqual>::GetRHI());
+				}
+			}
 
-				int32 NumPer = NumPrims / EffectiveThreads;
-				int32 Extra = NumPrims - NumPer * EffectiveThreads;
-
+			// Draw only translucent prims that are in the SeparateTranslucency pass
+			if (bRenderSeparateTranslucency)
+			{
+				int32 NumPrims = View.TranslucentPrimSet.NumSeparateTranslucencyPrims();
+				int32 EffectiveThreads = FMath::Min<int32>(NumPrims, Width);
 
-				for (int32 ThreadIndex = 0; ThreadIndex < EffectiveThreads; ThreadIndex++)
+				int32 Start = 0;
+				check(EffectiveThreads);
 				{
-					int32 Last = Start + (NumPer - 1) + (ThreadIndex < Extra);
-					check(Last >= Start);
+					int32 NumPer = NumPrims / EffectiveThreads;
+					int32 Extra = NumPrims - NumPer * EffectiveThreads;
 
+
+					for (int32 ThreadIndex = 0; ThreadIndex < EffectiveThreads; ThreadIndex++)
 					{
-						FRHICommandList* CmdList = new FRHICommandList;
+						int32 Last = Start + (NumPer - 1) + (ThreadIndex < Extra);
+						check(Last >= Start);
 
-						FGraphEventRef AnyThreadCompletionEvent = TGraphTask<FDrawSortedTransAnyThreadTask>::CreateTask(nullptr, ENamedThreads::RenderThread)
-							.ConstructAndDispatchWhenReady(this, CmdList, &View, true, Start, Last);
+						{
+							FRHICommandList* CmdList = new FRHICommandList;
 
-						RHICmdList.QueueAsyncCommandListSubmit(AnyThreadCompletionEvent, CmdList); 
+							FGraphEventRef AnyThreadCompletionEvent = TGraphTask<FDrawSortedTransAnyThreadTask>::CreateTask(nullptr, ENamedThreads::RenderThread)
+								.ConstructAndDispatchWhenReady(this, CmdList, &View, true, Start, Last);
+
+							RHICmdList.QueueAsyncCommandListSubmit(AnyThreadCompletionEvent, CmdList); 
+						}
+						Start = Last + 1;
 					}
-					Start = Last + 1;
 				}
 			}
-			else
-			{
-				// still need to setup the render target, a bit of a waste...
-				FRHICommandList* CmdList = new FRHICommandList;
-				FGraphEventRef AnyThreadCompletionEvent = TGraphTask<FDrawSortedTransAnyThreadTask>::CreateTask(nullptr, ENamedThreads::RenderThread)
-					.ConstructAndDispatchWhenReady(this, CmdList, &View, true, 0, -1);
-				RHICmdList.QueueAsyncCommandListSubmit(AnyThreadCompletionEvent, CmdList);
-			}
+			GSceneRenderTargets.FinishRenderingSeparateTranslucency(RHICmdList, View);
 		}
 	}
 }
