PrimitiveComponents have a new advanced lighting property - IndirectLightingCacheQuality
* Settings are Off, Point and Volume
* The default quality is Point, which is a downgrade in quality from the previous effective default of Volume, but has much less overhead (ILC update time reduced by 4x).  Components now have to opt-in to higher quality.  This also fixes the issue where ILC allocations would fail after 1728 movable components.
* Fixed Point sample location to be the snapped center of the bounds, previously it was the Min which was often below the ground
* Opaque materials can now apply the Point cached lighting from shader constants

diff --git a/Engine/Plugins/Experimental/Paper2D/Source/Paper2D/Private/PaperRenderSceneProxy.cpp b/Engine/Plugins/Experimental/Paper2D/Source/Paper2D/Private/PaperRenderSceneProxy.cpp
--- a/Engine/Plugins/Experimental/Paper2D/Source/Paper2D/Private/PaperRenderSceneProxy.cpp
+++ b/Engine/Plugins/Experimental/Paper2D/Source/Paper2D/Private/PaperRenderSceneProxy.cpp
@@ -163,6 +163,7 @@ FPaperRenderSceneProxy::FPaperRenderSceneProxy(const UPrimitiveComponent* InComp
 	, CollisionResponse(InComponent->GetCollisionResponseToChannels())
 {
 	WireframeColor = FLinearColor::White;
+	bWillEverBeLit = false;
 }
 
 FPaperRenderSceneProxy::~FPaperRenderSceneProxy()
diff --git a/Engine/Plugins/Experimental/Paper2D/Source/Paper2D/Private/Rendering/PaperBatchSceneProxy.cpp b/Engine/Plugins/Experimental/Paper2D/Source/Paper2D/Private/Rendering/PaperBatchSceneProxy.cpp
--- a/Engine/Plugins/Experimental/Paper2D/Source/Paper2D/Private/Rendering/PaperBatchSceneProxy.cpp
+++ b/Engine/Plugins/Experimental/Paper2D/Source/Paper2D/Private/Rendering/PaperBatchSceneProxy.cpp
@@ -10,6 +10,7 @@
 FPaperBatchSceneProxy::FPaperBatchSceneProxy(const UPrimitiveComponent* InComponent)
 	: FPrimitiveSceneProxy(InComponent)
 {
+	bWillEverBeLit = false;
 }
 
 uint32 FPaperBatchSceneProxy::GetMemoryFootprint() const
diff --git a/Engine/Shaders/BasePassPixelShader.usf b/Engine/Shaders/BasePassPixelShader.usf
--- a/Engine/Shaders/BasePassPixelShader.usf
+++ b/Engine/Shaders/BasePassPixelShader.usf
@@ -240,35 +240,6 @@ float3 GetTranslucencyLighting(FMaterialPixelParameters MaterialParameters, floa
 
 	#endif
 
-	#if CACHED_POINT_INDIRECT_LIGHTING 
-		#if TRANSLUCENCY_LIGHTING_VOLUMETRIC_DIRECTIONAL || TRANSLUCENCY_LIGHTING_SURFACE
-
-			FTwoBandSHVectorRGB PointIndirectLighting;
-			PointIndirectLighting.R.V = IndirectLightingSHCoefficients[0];
-			PointIndirectLighting.G.V = IndirectLightingSHCoefficients[1];
-			PointIndirectLighting.B.V = IndirectLightingSHCoefficients[2];
-
-			// Compute diffuse lighting which takes the normal into account
-			float3 DiffuseGI = max(half3(0,0,0), DotSH(PointIndirectLighting, DiffuseTransferSH));
-			IndirectIrradiance += Luminance(DiffuseGI);
-
-			InterpolatedLighting += DiffuseColor * DiffuseGI * MaterialAO;
-
-		#elif TRANSLUCENCY_LIGHTING_VOLUMETRIC_NONDIRECTIONAL
-
-			FOneBandSHVectorRGB PointIndirectLighting;
-			PointIndirectLighting.R.V = IndirectLightingSHCoefficients[0].x;
-			PointIndirectLighting.G.V = IndirectLightingSHCoefficients[1].x;
-			PointIndirectLighting.B.V = IndirectLightingSHCoefficients[2].x;
-
-			float3 DiffuseGI = DotSH1(PointIndirectLighting, DiffuseTransferSH);
-			IndirectIrradiance += Luminance(DiffuseGI);
-
-			InterpolatedLighting += DiffuseColor * DiffuseGI * MaterialAO;
-
-		#endif
-	#endif
-
 	#if TRANSLUCENCY_LIGHTING_SURFACE
 		InterpolatedLighting += GetImageBasedReflectionLighting(MaterialParameters, Roughness, SpecularColor, IndirectIrradiance);
 	#endif
@@ -299,7 +270,7 @@ float3 GetSkyLighting(float3 WorldNormal, float2 LightmapUV)
 
 #if ENABLE_SKY_LIGHT
 
-	float SkyOcclusion = 1;
+	float SkyVisibility = 1;
 	float3 SkyLightingNormal = WorldNormal;
 	
 	#if HQ_TEXTURE_LIGHTMAP || CACHED_POINT_INDIRECT_LIGHTING || CACHED_VOLUME_INDIRECT_LIGHTING
@@ -312,13 +283,13 @@ float3 GetSkyLighting(float3 WorldNormal, float2 LightmapUV)
 				float4 WorldSkyBentNormalAndOcclusion = GetSkyBentNormalAndOcclusion(LightmapUV * float2(1, 2));
 				// Renormalize as vector was quantized and compressed
 				float3 NormalizedBentNormal = normalize(WorldSkyBentNormalAndOcclusion.xyz);
-				SkyOcclusion = WorldSkyBentNormalAndOcclusion.w;
+				SkyVisibility = WorldSkyBentNormalAndOcclusion.w;
 
 			#elif CACHED_POINT_INDIRECT_LIGHTING || CACHED_VOLUME_INDIRECT_LIGHTING
 
 				// Bent normal from the indirect lighting cache - one value for the whole object
 				float3 NormalizedBentNormal = PointSkyBentNormal.xyz;
-				SkyOcclusion = PointSkyBentNormal.w;
+				SkyVisibility = PointSkyBentNormal.w;
 
 			#endif
 
@@ -328,7 +299,7 @@ float3 GetSkyLighting(float3 WorldNormal, float2 LightmapUV)
 			#else
 				
 				// Weight toward the material normal to increase directionality
-				float BentNormalWeightFactor = 1 - (1 - SkyOcclusion) * (1 - SkyOcclusion);
+				float BentNormalWeightFactor = 1 - (1 - SkyVisibility) * (1 - SkyVisibility);
 
 				// We are lerping between the inputs of two lighting scenarios based on occlusion
 				// In the mostly unoccluded case, evaluate sky lighting with the material normal, because it has higher detail
@@ -338,7 +309,7 @@ float3 GetSkyLighting(float3 WorldNormal, float2 LightmapUV)
 
 				float DotProductFactor = lerp(saturate(dot(NormalizedBentNormal, WorldNormal)), 1, BentNormalWeightFactor);
 				// Account for darkening due to the geometry term
-				SkyOcclusion *= DotProductFactor;
+				SkyVisibility *= DotProductFactor;
 			#endif
 		}
 	#endif
@@ -347,13 +318,13 @@ float3 GetSkyLighting(float3 WorldNormal, float2 LightmapUV)
 	float3 DiffuseLookup = GetSkySHDiffuse(SkyLightingNormal) * View.SkyLightColor.rgb;
 
 	// Apply AO to the sky diffuse
-	Lighting += DiffuseLookup * SkyOcclusion;
+	Lighting += DiffuseLookup * SkyVisibility;
 #endif
 
 	return Lighting;
 }
 
-/** Calculates indirect lighting contribution on this object. */
+/** Calculates indirect lighting contribution on this object from precomputed data. */
 float3 GetPrecomputedIndirectLighting(
 	FMaterialPixelParameters MaterialParameters, 
 	FVertexFactoryInterpolantsVSToPS Interpolants,
@@ -364,8 +335,7 @@ float3 GetPrecomputedIndirectLighting(
 	float3 Lighting = 0;
 	float2 SkyOcclusionUV = 0;
 
-	#define AMBIENTONLY 0
-
+	// Method for movable components which want to use a volume texture of interpolated SH samples
 	#if CACHED_VOLUME_INDIRECT_LIGHTING
 	
 		// Compute volume teture UVs from world position
@@ -375,6 +345,8 @@ float3 GetPrecomputedIndirectLighting(
 		VolumeUVs = clamp(VolumeUVs, IndirectlightingCacheMinUV, IndirectlightingCacheMaxUV);
 		float4 Vector0 = Texture3DSample(IndirectLightingCacheTexture0, IndirectLightingCacheTextureSampler0, VolumeUVs);
 
+		// For debugging
+		#define AMBIENTONLY 0
 		#if AMBIENTONLY
 
 			Lighting = Vector0.rgb / SHAmbientFunction() / PI;
@@ -396,6 +368,32 @@ float3 GetPrecomputedIndirectLighting(
 
 		#endif
 
+	// Method for movable components which want to use a single interpolated SH sample
+	#elif CACHED_POINT_INDIRECT_LIGHTING 
+		#if TRANSLUCENCY_LIGHTING_VOLUMETRIC_NONDIRECTIONAL
+
+			FOneBandSHVectorRGB PointIndirectLighting;
+			PointIndirectLighting.R.V = IndirectLightingSHCoefficients[0].x;
+			PointIndirectLighting.G.V = IndirectLightingSHCoefficients[1].x;
+			PointIndirectLighting.B.V = IndirectLightingSHCoefficients[2].x;
+
+			FOneBandSHVector DiffuseTransferSH = CalcDiffuseTransferSH1(MaterialParameters.WorldNormal, 1);
+			Lighting = DotSH1(PointIndirectLighting, DiffuseTransferSH);
+
+		#else
+
+			FTwoBandSHVectorRGB PointIndirectLighting;
+			PointIndirectLighting.R.V = IndirectLightingSHCoefficients[0];
+			PointIndirectLighting.G.V = IndirectLightingSHCoefficients[1];
+			PointIndirectLighting.B.V = IndirectLightingSHCoefficients[2];
+
+			FTwoBandSHVector DiffuseTransferSH = CalcDiffuseTransferSH(MaterialParameters.WorldNormal, 1);
+			// Compute diffuse lighting which takes the normal into account
+			Lighting = max(half3(0,0,0), DotSH(PointIndirectLighting, DiffuseTransferSH));
+
+		#endif
+
+	// High quality texture lightmaps
 	#elif HQ_TEXTURE_LIGHTMAP
 
 		float2 LightmapUV0, LightmapUV1;
@@ -403,6 +401,7 @@ float3 GetPrecomputedIndirectLighting(
 		SkyOcclusionUV = LightmapUV0;
 		Lighting = GetLightMapColorHQ(LightmapUV0, LightmapUV1, MaterialParameters.WorldNormal).rgb;
 
+	// Low quality texture lightmaps
 	#elif LQ_TEXTURE_LIGHTMAP
 
 		float2 LightmapUV0, LightmapUV1;
diff --git a/Engine/Shaders/TranslucentLightingShaders.usf b/Engine/Shaders/TranslucentLightingShaders.usf
--- a/Engine/Shaders/TranslucentLightingShaders.usf
+++ b/Engine/Shaders/TranslucentLightingShaders.usf
@@ -76,51 +76,6 @@ void WriteToSliceMainGS(triangle FWriteToSliceVertexOutput Input[3], inout Trian
 	OutStream.Append(Vertex2);
 }
 
-/** Stores the SH ambient term and Coefficient3.x. */
-Texture3D IndirectLightingCacheTexture;
-SamplerState IndirectLightingCacheTextureSampler;
-/** Stores Coefficient1 and Coefficient3.y. */
-Texture3D IndirectLightingCacheTexture1;
-SamplerState IndirectLightingCacheTexture1Sampler;
-/** Stores Coefficient2 and Coefficient3.z. */
-Texture3D IndirectLightingCacheTexture2;
-SamplerState IndirectLightingCacheTexture2Sampler;
-
-/** Add and Scale to convert world space position into indirect lighting cache volume texture UVs. */
-float3 IndirectlightingCacheAdd;
-float3 IndirectlightingCacheScale;
-
-/**  */
-void CompositeGIMainPS(
-	FWriteToSliceGeometryOutput Input,
-	out float4 OutColor0 : SV_Target0,
-	out float4 OutColor1 : SV_Target1
-	)
-{
-	float ZPosition = View.TranslucencyLightingVolumeMin[VolumeCascadeIndex].z + Input.LayerIndex * View.TranslucencyLightingVolumeInvSize[VolumeCascadeIndex].w;
-	float3 WorldPosition = float3(View.TranslucencyLightingVolumeMin[VolumeCascadeIndex].xy + Input.Vertex.UV / View.TranslucencyLightingVolumeInvSize[VolumeCascadeIndex].xy - .5f * View.TranslucencyLightingVolumeInvSize[VolumeCascadeIndex].w, ZPosition);
-	float3 CacheUV = WorldPosition * IndirectlightingCacheScale + IndirectlightingCacheAdd;
-
-	float4 Vector0 = Texture3DSample(IndirectLightingCacheTexture, IndirectLightingCacheTextureSampler, CacheUV);
-	float4 Vector1 = Texture3DSample(IndirectLightingCacheTexture1, IndirectLightingCacheTexture1Sampler, CacheUV);
-	float4 Vector2 = Texture3DSample(IndirectLightingCacheTexture2, IndirectLightingCacheTexture2Sampler, CacheUV);
-
-	// Construct the SH environment
-	FTwoBandSHVectorRGB CachedSH;
-	CachedSH.R.V = float4(Vector0.x, Vector1.x, Vector2.x, Vector0.w);
-	CachedSH.G.V = float4(Vector0.y, Vector1.y, Vector2.y, Vector1.w);
-	CachedSH.B.V = float4(Vector0.z, Vector1.z, Vector2.z, Vector2.w);
-
-	OutColor0 = float4(CachedSH.R.V.x, CachedSH.G.V.x, CachedSH.B.V.x, 0);
-
-	float3 LuminanceWeights = float3(.3, .59, .11);
-	float3 Coefficient0 = float3(CachedSH.R.V.y, CachedSH.G.V.y, CachedSH.B.V.y);
-	float3 Coefficient1 = float3(CachedSH.R.V.z, CachedSH.G.V.z, CachedSH.B.V.z);
-	float3 Coefficient2 = float3(CachedSH.R.V.w, CachedSH.G.V.w, CachedSH.B.V.w);
-	OutColor1 = float4(dot(Coefficient0, LuminanceWeights), dot(Coefficient1, LuminanceWeights), dot(Coefficient2, LuminanceWeights), 0);
-}
-
-
 /** Filter pass inputs. */
 Texture3D TranslucencyLightingVolumeAmbient;
 SamplerState TranslucencyLightingVolumeAmbientSampler;
diff --git a/Engine/Source/Editor/UnrealEd/Private/PreviewMaterial.cpp b/Engine/Source/Editor/UnrealEd/Private/PreviewMaterial.cpp
--- a/Engine/Source/Editor/UnrealEd/Private/PreviewMaterial.cpp
+++ b/Engine/Source/Editor/UnrealEd/Private/PreviewMaterial.cpp
@@ -80,10 +80,7 @@ public:
 		{
 			bShaderTypeMatches = true;
 		}
-		else if (FCString::Stristr(ShaderType->GetName(), TEXT("BasePassVSFCachedVolumeIndirectLightingPolicy"))
-			|| FCString::Stristr(ShaderType->GetName(), TEXT("BasePassPSFCachedVolumeIndirectLightingPolicy"))
-			|| FCString::Stristr(ShaderType->GetName(), TEXT("BasePassHSFCachedVolumeIndirectLightingPolicy"))
-			|| FCString::Stristr(ShaderType->GetName(), TEXT("BasePassDSFCachedVolumeIndirectLightingPolicy")))
+		else if (FCString::Stristr(ShaderType->GetName(), TEXT("CachedPointIndirectLightingPolicy")))
 		{
 			bShaderTypeMatches = true;
 		}
diff --git a/Engine/Source/Runtime/Engine/Classes/Components/PrimitiveComponent.h b/Engine/Source/Runtime/Engine/Classes/Components/PrimitiveComponent.h
--- a/Engine/Source/Runtime/Engine/Classes/Components/PrimitiveComponent.h
+++ b/Engine/Source/Runtime/Engine/Classes/Components/PrimitiveComponent.h
@@ -215,7 +215,7 @@ public:
 	UPROPERTY(EditAnywhere, AdvancedDisplay, BlueprintReadWrite, Category=Rendering)
 	uint32 bRenderInMainPass:1;
 
-	/** Whether to completely hide the primitive in the game; if true, the primitive is not drawn, does not cast a shadow, and does not affect voxel lighting. */
+	/** Whether to completely hide the primitive in the game; if true, the primitive is not drawn, does not cast a shadow. */
 	UPROPERTY()
 	uint32 HiddenGame_DEPRECATED:1;
 
@@ -320,6 +320,10 @@ public:
 	UPROPERTY(EditAnywhere, AdvancedDisplay, BlueprintReadOnly, Category=Lighting)
 	uint32 bLightAttachmentsAsGroup:1;
 
+	/** Quality of indirect lighting for Movable primitives.  This has a large effect on Indirect Lighting Cache update time. */
+	UPROPERTY(EditAnywhere, AdvancedDisplay, BlueprintReadOnly, Category=Lighting)
+	TEnumAsByte<EIndirectLightingCacheQuality> IndirectLightingCacheQuality;
+
 	UPROPERTY()
 	bool bHasCachedStaticLighting;
 
@@ -347,10 +351,6 @@ public:
 	UPROPERTY()
 	uint32 AlwaysLoadOnServer:1;
 
-	/** Determines whether or not we allow shadowing fading.  Some objects (especially in cinematics) having the shadow fade/pop out looks really bad. **/
-	UPROPERTY()
-	uint32 bAllowShadowFade:1;
-
 	/** Composite the drawing of this component onto the scene after post processing (only applies to editor drawing) */
 	UPROPERTY()
 	uint32 bUseEditorCompositing:1;
diff --git a/Engine/Source/Runtime/Engine/Classes/Engine/EngineTypes.h b/Engine/Source/Runtime/Engine/Classes/Engine/EngineTypes.h
--- a/Engine/Source/Runtime/Engine/Classes/Engine/EngineTypes.h
+++ b/Engine/Source/Runtime/Engine/Classes/Engine/EngineTypes.h
@@ -64,6 +64,17 @@ enum ESceneDepthPriorityGroup
 	SDPG_MAX,
 };
 
+UENUM()
+enum EIndirectLightingCacheQuality
+{
+	/** The indirect lighting cache will be disabled for this object, so no GI from stationary lights on movable objects. */
+	ILCQ_Off,
+	/** A single indirect lighting sample computed at the bounds origin will be interpolated which fades over time to newer results. */
+	ILCQ_Point,
+	/** The object will get a 5x5x5 stable volume of interpolated indirect lighting, which allows gradients of lighting intensity across the receiving object. */
+	ILCQ_Volume
+};
+
 /** Note: This is mirrored in Lightmass, be sure to update the blend mode structure and logic there if this changes. */
 // Note: Check UMaterialInstance::Serialize if changed!!
 UENUM(BlueprintType)
diff --git a/Engine/Source/Runtime/Engine/Private/PrimitiveComponent.cpp b/Engine/Source/Runtime/Engine/Private/PrimitiveComponent.cpp
--- a/Engine/Source/Runtime/Engine/Private/PrimitiveComponent.cpp
+++ b/Engine/Source/Runtime/Engine/Private/PrimitiveComponent.cpp
@@ -79,10 +79,10 @@ UPrimitiveComponent::UPrimitiveComponent(const class FPostConstructInitializePro
 	LpvBiasMultiplier = 1.0f;
 	bCastStaticShadow = true;
 	bCastVolumetricTranslucentShadow = false;
+	IndirectLightingCacheQuality = ILCQ_Point;
 	bSelectable = true;
 	AlwaysLoadOnClient = true;
 	AlwaysLoadOnServer = true;
-	bAllowShadowFade = true;
 	BodyInstance.bEnableCollision_DEPRECATED = true;
 	SetCollisionProfileName(UCollisionProfile::BlockAll_ProfileName);
 	bAlwaysCreatePhysicsState = false;
@@ -554,6 +554,11 @@ bool UPrimitiveComponent::CanEditChange(const UProperty* InProperty) const
 		{
 			return Mobility != EComponentMobility::Movable || bLightAsIfStatic;
 		}
+
+		if (PropertyName == TEXT("IndirectLightingCacheQuality"))
+		{
+			return Mobility == EComponentMobility::Movable;
+		}
 	}
 
 	return bIsEditable;
diff --git a/Engine/Source/Runtime/Engine/Private/PrimitiveSceneProxy.cpp b/Engine/Source/Runtime/Engine/Private/PrimitiveSceneProxy.cpp
--- a/Engine/Source/Runtime/Engine/Private/PrimitiveSceneProxy.cpp
+++ b/Engine/Source/Runtime/Engine/Private/PrimitiveSceneProxy.cpp
@@ -49,10 +49,10 @@ FPrimitiveSceneProxy::FPrimitiveSceneProxy(const UPrimitiveComponent* InComponen
 ,	bUseAsOccluder(InComponent->bUseAsOccluder)
 ,	bAllowApproximateOcclusion(InComponent->Mobility != EComponentMobility::Movable)
 ,	bSelectable(InComponent->bSelectable)
-,	bAllowShadowFade(InComponent->bAllowShadowFade)
 ,	bHasPerInstanceHitProxies(InComponent->bHasPerInstanceHitProxies)
 ,	bUseEditorCompositing(InComponent->bUseEditorCompositing)
 ,	LpvBiasMultiplier(InComponent->LpvBiasMultiplier)
+,	IndirectLightingCacheQuality(InComponent->IndirectLightingCacheQuality)
 ,	Scene(InComponent->GetScene())
 ,	PrimitiveComponentId(InComponent->ComponentId)
 ,	PrimitiveSceneInfo(NULL)
diff --git a/Engine/Source/Runtime/Engine/Public/PrimitiveSceneProxy.h b/Engine/Source/Runtime/Engine/Public/PrimitiveSceneProxy.h
--- a/Engine/Source/Runtime/Engine/Public/PrimitiveSceneProxy.h
+++ b/Engine/Source/Runtime/Engine/Public/PrimitiveSceneProxy.h
@@ -334,6 +334,7 @@ public:
 	inline bool CastsDynamicShadow() const { return bCastDynamicShadow; }
 	inline bool AffectsDynamicIndirectLighting() const { return bAffectDynamicIndirectLighting; }
 	inline float GetLpvBiasMultiplier() const { return LpvBiasMultiplier; }
+	inline EIndirectLightingCacheQuality GetIndirectLightingCacheQuality() const { return IndirectLightingCacheQuality; }
 	inline bool CastsVolumetricTranslucentShadow() const { return bCastVolumetricTranslucentShadow; }
 	inline bool CastsHiddenShadow() const { return bCastHiddenShadow; }
 	inline bool CastsShadowAsTwoSided() const { return bCastShadowAsTwoSided; }
@@ -540,9 +541,6 @@ private:
 	/** If this is True, this primitive can be selected in the editor. */
 	uint32 bSelectable : 1;
 
-	/** Determines whether or not we allow shadowing fading.  Some objects (especially in cinematics) having the shadow fade/pop out looks really bad. **/
-	uint32 bAllowShadowFade : 1;
-
 	/** If this primitive has per-instance hit proxies. */
 	uint32 bHasPerInstanceHitProxies : 1;
 
@@ -553,6 +551,9 @@ protected:
 	/** The bias applied to LPV injection */
 	float LpvBiasMultiplier;
 
+	/** Quality of interpolated indirect lighting for Movable components. */
+	EIndirectLightingCacheQuality IndirectLightingCacheQuality;
+
 private:
 	/** The primitive's local to world transform. */
 	FMatrix LocalToWorld;
diff --git a/Engine/Source/Runtime/Renderer/Private/BasePassRendering.h b/Engine/Source/Runtime/Renderer/Private/BasePassRendering.h
--- a/Engine/Source/Runtime/Renderer/Private/BasePassRendering.h
+++ b/Engine/Source/Runtime/Renderer/Private/BasePassRendering.h
@@ -918,23 +918,36 @@ void ProcessBasePassMesh(
 							}
 							else if (IsIndirectLightingCacheAllowed(Parameters.FeatureLevel)
 								&& Action.AllowIndirectLightingCache()
-								&& Parameters.PrimitiveSceneProxy
+								&& Parameters.PrimitiveSceneProxy)
+							{
+								const FIndirectLightingCacheAllocation* IndirectLightingCacheAllocation = Parameters.PrimitiveSceneProxy->GetPrimitiveSceneInfo()->IndirectLightingCacheAllocation;
+								const bool bPrimitiveIsMovable = Parameters.PrimitiveSceneProxy->IsMovable();
+
 								// Use the indirect lighting cache shaders if the object has a cache allocation
 								// This happens for objects with unbuilt lighting
-								&& ((Parameters.PrimitiveSceneProxy->GetPrimitiveSceneInfo()->IndirectLightingCacheAllocation
-									&& Parameters.PrimitiveSceneProxy->GetPrimitiveSceneInfo()->IndirectLightingCacheAllocation->IsValid())
+								if (IndirectLightingCacheAllocation && IndirectLightingCacheAllocation->IsValid()
 									// Use the indirect lighting cache shaders if the object is movable, it may not have a cache allocation yet because that is done in InitViews
-									|| Parameters.PrimitiveSceneProxy->IsMovable()))
-							{
-								if (CanIndirectLightingCacheUseVolumeTexture(Parameters.FeatureLevel) && Action.AllowIndirectLightingCacheVolumeTexture())
+									// And movable objects are sometimes rendered in the static draw lists
+									|| bPrimitiveIsMovable)
 								{
-									// Use a lightmap policy that supports reading indirect lighting from a volume texture for dynamic objects
-									Action.template Process<FCachedVolumeIndirectLightingPolicy>(RHICmdList, Parameters, FCachedVolumeIndirectLightingPolicy(), FCachedVolumeIndirectLightingPolicy::ElementDataType());
+									if (CanIndirectLightingCacheUseVolumeTexture(Parameters.FeatureLevel) 
+										// Translucency forces point sample for pixel performance
+										&& Action.AllowIndirectLightingCacheVolumeTexture()
+										&& (IndirectLightingCacheAllocation && !IndirectLightingCacheAllocation->bPointSample
+											|| bPrimitiveIsMovable && Parameters.PrimitiveSceneProxy->GetIndirectLightingCacheQuality() == ILCQ_Volume))
+									{
+										// Use a lightmap policy that supports reading indirect lighting from a volume texture for dynamic objects
+										Action.template Process<FCachedVolumeIndirectLightingPolicy>(RHICmdList, Parameters, FCachedVolumeIndirectLightingPolicy(), FCachedVolumeIndirectLightingPolicy::ElementDataType());
+									}
+									else
+									{
+										// Use a lightmap policy that supports reading indirect lighting from a single SH sample
+										Action.template Process<FCachedPointIndirectLightingPolicy>(RHICmdList, Parameters, FCachedPointIndirectLightingPolicy(), FCachedPointIndirectLightingPolicy::ElementDataType(false));
+									}
 								}
 								else
 								{
-									// Use a lightmap policy that supports reading indirect lighting from a single SH sample
-									Action.template Process<FCachedPointIndirectLightingPolicy>(RHICmdList, Parameters, FCachedPointIndirectLightingPolicy(), FCachedPointIndirectLightingPolicy::ElementDataType(false));
+									Action.template Process<FNoLightMapPolicy>(RHICmdList, Parameters, FNoLightMapPolicy(), FNoLightMapPolicy::ElementDataType());
 								}
 							}
 							else
diff --git a/Engine/Source/Runtime/Renderer/Private/DeferredShadingRenderer.cpp b/Engine/Source/Runtime/Renderer/Private/DeferredShadingRenderer.cpp
--- a/Engine/Source/Runtime/Renderer/Private/DeferredShadingRenderer.cpp
+++ b/Engine/Source/Runtime/Renderer/Private/DeferredShadingRenderer.cpp
@@ -816,8 +816,6 @@ void FDeferredShadingSceneRenderer::Render(FRHICommandListImmediate& RHICmdList)
 
 		InjectAmbientCubemapTranslucentVolumeLighting(RHICmdList);
 
-		CompositeIndirectTranslucentVolumeLighting(RHICmdList);
-
 		// Filter the translucency lighting volume now that it is complete
 		FilterTranslucentVolumeLighting(RHICmdList);
 
diff --git a/Engine/Source/Runtime/Renderer/Private/DeferredShadingRenderer.h b/Engine/Source/Runtime/Renderer/Private/DeferredShadingRenderer.h
--- a/Engine/Source/Runtime/Renderer/Private/DeferredShadingRenderer.h
+++ b/Engine/Source/Runtime/Renderer/Private/DeferredShadingRenderer.h
@@ -346,9 +346,6 @@ private:
 	/** Add AmbientCubemap to the lighting volumes. */
 	void InjectAmbientCubemapTranslucentVolumeLighting(FRHICommandList& RHICmdList);
 
-	/** Renders indirect lighting into the translucent lighting volumes. */
-	void CompositeIndirectTranslucentVolumeLighting(FRHICommandList& RHICmdList);
-
 	/** Clears the volume texture used to accumulate per object shadows for translucency. */
 	void ClearTranslucentVolumePerObjectShadowing(FRHICommandList& RHICmdList);
 
diff --git a/Engine/Source/Runtime/Renderer/Private/IndirectLightingCache.cpp b/Engine/Source/Runtime/Renderer/Private/IndirectLightingCache.cpp
--- a/Engine/Source/Runtime/Renderer/Private/IndirectLightingCache.cpp
+++ b/Engine/Source/Runtime/Renderer/Private/IndirectLightingCache.cpp
@@ -109,11 +109,12 @@ static FAutoConsoleVariableRef CVarLightingCacheMovableObjectAllocationSize(
 	TEXT("r.Cache.LightingCacheMovableObjectAllocationSize"),
 	GLightingCacheMovableObjectAllocationSize,
 	TEXT("Resolution of the interpolation sample volume used to light a dynamic object.  \n")
-	TEXT("Values of 1 or 2 will result in a single interpolation sample per object which does not provide continuous lighting under movement, so interpolation over time is done.  \n")
+	TEXT("Values of 1 or 2 will result in a single interpolation sample per object which does not provide continuous lighting under movement, so interpolation is done over time.  \n")
 	TEXT("Values of 3 or more support the necessary padding to provide continuous results under movement."),
 	ECVF_ReadOnly
 	);
 
+// PS4 wants multiple of 16 for fast RHIUpdateTexture3D
 int32 GLightingCacheDimension = 64;
 static FAutoConsoleVariableRef CVarLightingCacheDimension(
 	TEXT("r.Cache.LightingCacheDimension"),
@@ -122,14 +123,6 @@ static FAutoConsoleVariableRef CVarLightingCacheDimension(
 	ECVF_ReadOnly
 	);
 
-int32 GLightingCacheUnbuiltPreviewAllocationSize = 1;
-static FAutoConsoleVariableRef CVarLightingCacheUnbuiltPreviewAllocationSize(
-	TEXT("r.Cache.LightingCacheUnbuiltPreviewAllocationSize"),
-	GLightingCacheUnbuiltPreviewAllocationSize,
-	TEXT("Resolution of the interpolation sample volume used to light an object due to unbuilt lighting."),
-	ECVF_ReadOnly
-	);
-
 bool IsIndirectLightingCacheAllowed(ERHIFeatureLevel::Type InFeatureLevel) 
 {
 	static const auto AllowStaticLightingVar = IConsoleManager::Get().FindTConsoleVariableDataInt(TEXT("r.AllowStaticLighting"));
@@ -234,22 +227,33 @@ void FIndirectLightingCache::CalculateBlockPositionAndSize(const FBoxSphereBound
 	const FVector CellSize = RoundedBoundsSize / (EffectiveTexelSize - 2);
 	const FVector BoundsMin = Bounds.Origin - Bounds.BoxExtent;
 
-	FVector SnappedMin;
-	SnappedMin.X = CellSize.X * FMath::FloorToFloat(BoundsMin.X / CellSize.X);
-	SnappedMin.Y = CellSize.Y * FMath::FloorToFloat(BoundsMin.Y / CellSize.Y);
-	SnappedMin.Z = CellSize.Z * FMath::FloorToFloat(BoundsMin.Z / CellSize.Z);
-
 	if (TexelSize > 2)
 	{
+		FVector SnappedMin;
+		SnappedMin.X = CellSize.X * FMath::FloorToFloat(BoundsMin.X / CellSize.X);
+		SnappedMin.Y = CellSize.Y * FMath::FloorToFloat(BoundsMin.Y / CellSize.Y);
+		SnappedMin.Z = CellSize.Z * FMath::FloorToFloat(BoundsMin.Z / CellSize.Z);
+
 		// Shift the min down so that the center of the voxel is at the min
 		// This is necessary so that all pixels inside the bounds only interpolate from valid voxels
 		SnappedMin -= CellSize * .5f;
-	}
 
-	const FVector SnappedSize = TexelSize * CellSize;
+		const FVector SnappedSize = TexelSize * CellSize;
 
-	OutMin = SnappedMin;
-	OutSize = TexelSize > 2 ? SnappedSize : FVector(0);
+		OutMin = SnappedMin;
+		OutSize = SnappedSize;
+	}
+	else
+	{
+		FVector SnappedCenter;
+		SnappedCenter.X = CellSize.X * FMath::FloorToFloat((Bounds.Origin.X + .5f * CellSize.X) / CellSize.X);
+		SnappedCenter.Y = CellSize.Y * FMath::FloorToFloat((Bounds.Origin.Y + .5f * CellSize.Y) / CellSize.Y);
+		SnappedCenter.Z = CellSize.Z * FMath::FloorToFloat((Bounds.Origin.Z + .5f * CellSize.Z) / CellSize.Z);
+
+		// Place the min at the snapped center of the object, which will be most representative for single sample lighting
+		OutMin = SnappedCenter;
+		OutSize = FVector(0);
+	}
 }
 
 void FIndirectLightingCache::CalculateBlockScaleAndAdd(FIntVector InTexelMin, int32 AllocationTexelSize, FVector InMin, FVector InSize, FVector& OutScale, FVector& OutAdd, FVector& OutMinUV, FVector& OutMaxUV) const
@@ -294,15 +298,17 @@ void FIndirectLightingCache::CalculateBlockScaleAndAdd(FIntVector InTexelMin, in
 
 FIndirectLightingCacheAllocation* FIndirectLightingCache::AllocatePrimitive(const FPrimitiveSceneInfo* PrimitiveSceneInfo, bool bUnbuiltPreview)
 {
-	const int32 BlockSize = bUnbuiltPreview ? GLightingCacheUnbuiltPreviewAllocationSize : GLightingCacheMovableObjectAllocationSize;
-	return PrimitiveAllocations.Add(PrimitiveSceneInfo->PrimitiveComponentId, CreateAllocation(BlockSize, PrimitiveSceneInfo->Proxy->GetBounds(), true));
+	const bool bPointSample = PrimitiveSceneInfo->Proxy->GetIndirectLightingCacheQuality() == ILCQ_Point || bUnbuiltPreview;
+	const int32 BlockSize = bPointSample ? 1 : GLightingCacheMovableObjectAllocationSize;
+	return PrimitiveAllocations.Add(PrimitiveSceneInfo->PrimitiveComponentId, CreateAllocation(BlockSize, PrimitiveSceneInfo->Proxy->GetBounds(), bPointSample));
 }
 
-FIndirectLightingCacheAllocation* FIndirectLightingCache::CreateAllocation(int32 BlockSize, const FBoxSphereBounds& Bounds, bool bOpaqueRelevance)
+FIndirectLightingCacheAllocation* FIndirectLightingCache::CreateAllocation(int32 BlockSize, const FBoxSphereBounds& Bounds, bool bPointSample)
 {
 	FIndirectLightingCacheAllocation* NewAllocation = new FIndirectLightingCacheAllocation();
 	FIndirectLightingCacheBlock NewBlock;
 
+	//@todo - don't allocate point samples from the layout, they don't go through the volume texture
 	if (AllocateBlock(BlockSize, NewBlock.MinTexel))
 	{
 		NewBlock.TexelSize = BlockSize;
@@ -315,7 +321,7 @@ FIndirectLightingCacheAllocation* FIndirectLightingCache::CreateAllocation(int32
 		CalculateBlockScaleAndAdd(NewBlock.MinTexel, NewBlock.TexelSize, NewBlock.Min, NewBlock.Size, Scale, Add, MinUV, MaxUV);
 
 		VolumeBlocks.Add(NewBlock.MinTexel, NewBlock);
-		NewAllocation->SetParameters(NewBlock.MinTexel, NewBlock.TexelSize, Scale, Add, MinUV, MaxUV, bOpaqueRelevance);
+		NewAllocation->SetParameters(NewBlock.MinTexel, NewBlock.TexelSize, Scale, Add, MinUV, MaxUV, bPointSample);
 	}
 
 	return NewAllocation;
@@ -389,8 +395,6 @@ void FIndirectLightingCache::UpdateCache(FScene* Scene, FSceneRenderer& Renderer
 						UpdateCachePrimitive(Scene, PrimitiveSceneInfo, bAllowUnbuiltPreview, PrimitiveRelevance.bOpaqueRelevance, BlocksToUpdate, TransitionsOverTimeToUpdate);
 					}
 				}
-
-				UpdateTranslucentVolumeCache(View, BlocksToUpdate, TransitionsOverTimeToUpdate);
 			}
 
 			UpdateBlocks(Scene, Renderer.Views.GetTypedData(), BlocksToUpdate);
@@ -417,7 +421,7 @@ void FIndirectLightingCache::UpdateCache(FScene* Scene, FSceneRenderer& Renderer
 void FIndirectLightingCache::UpdateCacheAllocation(
 	const FBoxSphereBounds& Bounds, 
 	int32 BlockSize,
-	bool bOpaqueRelevance,
+	bool bPointSample,
 	FIndirectLightingCacheAllocation*& Allocation, 
 	TMap<FIntVector, FBlockUpdateInfo>& BlocksToUpdate,
 	TArray<FIndirectLightingCacheAllocation*>& TransitionsOverTimeToUpdate)
@@ -444,7 +448,7 @@ void FIndirectLightingCache::UpdateCacheAllocation(
 			FVector MaxUV;
 			CalculateBlockScaleAndAdd(Allocation->MinTexel, Allocation->AllocationTexelSize, NewMin, NewSize, NewScale, NewAdd, MinUV, MaxUV);
 
-			Allocation->SetParameters(Allocation->MinTexel, Allocation->AllocationTexelSize, NewScale, NewAdd, MinUV, MaxUV, bOpaqueRelevance);
+			Allocation->SetParameters(Allocation->MinTexel, Allocation->AllocationTexelSize, NewScale, NewAdd, MinUV, MaxUV, bPointSample);
 			BlocksToUpdate.Add(Block.MinTexel, FBlockUpdateInfo(Block, Allocation));
 		}
 
@@ -456,7 +460,7 @@ void FIndirectLightingCache::UpdateCacheAllocation(
 	else
 	{
 		delete Allocation;
-		Allocation = CreateAllocation(BlockSize, Bounds, bOpaqueRelevance);
+		Allocation = CreateAllocation(BlockSize, Bounds, bPointSample);
 
 		if (Allocation->IsValid())
 		{
@@ -466,24 +470,6 @@ void FIndirectLightingCache::UpdateCacheAllocation(
 	}
 }
 
-void FIndirectLightingCache::UpdateTranslucentVolumeCache(FViewInfo& View, TMap<FIntVector, FBlockUpdateInfo>& BlocksToUpdate, TArray<FIndirectLightingCacheAllocation*>& TransitionsOverTimeToUpdate)
-{
-	extern int32 GUseIndirectLightingCacheInLightingVolume;
-
-	if (View.State && GUseIndirectLightingCacheInLightingVolume && GSupportsVolumeTextureRendering)
-	{
-		FSceneViewState* ViewState = (FSceneViewState*)View.State;
-
-		for (int32 CascadeIndex = 0; CascadeIndex < ARRAY_COUNT(ViewState->TranslucencyLightingCacheAllocations); CascadeIndex++)
-		{
-			FIndirectLightingCacheAllocation*& Allocation = ViewState->TranslucencyLightingCacheAllocations[CascadeIndex];
-			const FBoxSphereBounds Bounds(FBox(View.TranslucencyLightingVolumeMin[CascadeIndex], View.TranslucencyLightingVolumeMin[CascadeIndex] + View.TranslucencyLightingVolumeSize[CascadeIndex]));
-
-			UpdateCacheAllocation(Bounds, GTranslucencyLightingVolumeDim / 4, true, Allocation, BlocksToUpdate, TransitionsOverTimeToUpdate);
-		}
-	}
-}
-
 void FIndirectLightingCache::UpdateCachePrimitive(
 	FScene* Scene, 
 	FPrimitiveSceneInfo* PrimitiveSceneInfo,
@@ -496,9 +482,9 @@ void FIndirectLightingCache::UpdateCachePrimitive(
 	FIndirectLightingCacheAllocation** PrimitiveAllocationPtr = PrimitiveAllocations.Find(PrimitiveSceneInfo->PrimitiveComponentId);
 	FIndirectLightingCacheAllocation* PrimitiveAllocation = PrimitiveAllocationPtr != NULL ? *PrimitiveAllocationPtr : NULL;
 
-	if (PrimitiveSceneProxy->WillEverBeLit() 
-		&& ((bAllowUnbuiltPreview && PrimitiveSceneProxy->HasStaticLighting() && PrimitiveAllocation && PrimitiveAllocation->bIsDirty)
-		|| PrimitiveSceneProxy->IsMovable()))
+	if (PrimitiveSceneProxy->WillEverBeLit()
+		&& (bAllowUnbuiltPreview && PrimitiveSceneProxy->HasStaticLighting() && PrimitiveAllocation && PrimitiveAllocation->bIsDirty
+		|| (PrimitiveSceneProxy->IsMovable() && PrimitiveSceneProxy->GetIndirectLightingCacheQuality() != ILCQ_Off)))
 	{
 		const FIndirectLightingCacheAllocation* AttachmentParentAllocation = NULL;
 
@@ -521,10 +507,11 @@ void FIndirectLightingCache::UpdateCachePrimitive(
 		{
 			FIndirectLightingCacheAllocation* OriginalAllocation = PrimitiveAllocation;
 			const bool bUnbuiltPreview = bAllowUnbuiltPreview && !PrimitiveSceneProxy->IsMovable();
-			const int32 BlockSize = bUnbuiltPreview ? GLightingCacheUnbuiltPreviewAllocationSize : GLightingCacheMovableObjectAllocationSize;
+			const bool bPointSample = PrimitiveSceneProxy->GetIndirectLightingCacheQuality() == ILCQ_Point || bUnbuiltPreview || !bOpaqueRelevance;
+			const int32 BlockSize = bPointSample ? 1 : GLightingCacheMovableObjectAllocationSize;
 
 			// Light with the cumulative bounds of the entire attachment group
-			UpdateCacheAllocation(PrimitiveSceneInfo->GetAttachmentGroupBounds(), BlockSize, bOpaqueRelevance, PrimitiveAllocation, BlocksToUpdate, TransitionsOverTimeToUpdate);
+			UpdateCacheAllocation(PrimitiveSceneInfo->GetAttachmentGroupBounds(), BlockSize, bPointSample, PrimitiveAllocation, BlocksToUpdate, TransitionsOverTimeToUpdate);
 
 			// Cache the primitive allocation pointer on the FPrimitiveSceneInfo for base pass rendering
 			PrimitiveSceneInfo->IndirectLightingCacheAllocation = PrimitiveAllocation;
@@ -609,7 +596,7 @@ void FIndirectLightingCache::UpdateBlock(FScene* Scene, FViewInfo* DebugDrawingV
 	float DirectionalShadowing = 1;
 	FVector SkyBentNormal(0, 0, 1);
 
-	if (CanIndirectLightingCacheUseVolumeTexture(GetFeatureLevel()) && BlockInfo.Allocation->bOpaqueRelevance)
+	if (CanIndirectLightingCacheUseVolumeTexture(GetFeatureLevel()) && !BlockInfo.Allocation->bPointSample)
 	{
 		static TArray<float> AccumulatedWeight;
 		AccumulatedWeight.Reset(NumSamplesPerBlock);
@@ -640,6 +627,7 @@ void FIndirectLightingCache::UpdateBlock(FScene* Scene, FViewInfo* DebugDrawingV
 		check(FormatSize == sizeof(FFloat16Color));
 
 		// Encode the SH samples into a texture format
+		// Note the single sample is updated even if this is a volume allocation, because translucent materials only use the single sample
 		EncodeBlock(DebugDrawingView, BlockInfo.Block, AccumulatedWeight, AccumulatedIncidentRadiance, AccumulatedSkyBentNormal, Texture0Data, Texture1Data, Texture2Data, SingleSample, SkyBentNormal);
 
 		// Setup an update region
@@ -662,6 +650,13 @@ void FIndirectLightingCache::UpdateBlock(FScene* Scene, FViewInfo* DebugDrawingV
 	else
 	{
 		InterpolatePoint(Scene, BlockInfo.Block, DirectionalShadowing, SingleSample, SkyBentNormal);
+
+		if (GCacheDrawInterpolationPoints != 0 && DebugDrawingView)
+		{
+			FViewElementPDI DebugPDI(DebugDrawingView, NULL);
+			const FVector WorldPosition = BlockInfo.Block.Min;
+			DebugPDI.DrawPoint(WorldPosition, FLinearColor(0, 0, 1), 10, SDPG_World);
+		}
 	}
 
 	// Record the position that the sample was taken at
diff --git a/Engine/Source/Runtime/Renderer/Private/LightMapRendering.h b/Engine/Source/Runtime/Renderer/Private/LightMapRendering.h
--- a/Engine/Source/Runtime/Renderer/Private/LightMapRendering.h
+++ b/Engine/Source/Runtime/Renderer/Private/LightMapRendering.h
@@ -731,11 +731,8 @@ public:
 	static bool ShouldCache(EShaderPlatform Platform,const FMaterial* Material,const FVertexFactoryType* VertexFactoryType)
 	{
 		static const auto AllowStaticLightingVar = IConsoleManager::Get().FindTConsoleVariableDataInt(TEXT("r.AllowStaticLighting"));
-		
-		// @todo Mac OS X: For GL 3.3 devices which don't support volume-texture rendering we need to cache the simpler point indirect lighting shaders.
+	
 		return Material->GetShadingModel() != MSM_Unlit
-			&& (IsTranslucentBlendMode(Material->GetBlendMode()) || GetMaxSupportedFeatureLevel(Platform) == ERHIFeatureLevel::ES2
-				|| (Platform == SP_OPENGL_SM4_MAC))
 			&& (!AllowStaticLightingVar || AllowStaticLightingVar->GetValueOnAnyThread() != 0);
 	}
 
diff --git a/Engine/Source/Runtime/Renderer/Private/PostProcess/SceneRenderTargets.cpp b/Engine/Source/Runtime/Renderer/Private/PostProcess/SceneRenderTargets.cpp
--- a/Engine/Source/Runtime/Renderer/Private/PostProcess/SceneRenderTargets.cpp
+++ b/Engine/Source/Runtime/Renderer/Private/PostProcess/SceneRenderTargets.cpp
@@ -1146,24 +1146,6 @@ void FSceneRenderTargets::AllocateDeferredShadingPathRenderTargets()
 					);
 			}
 		}
-
-		extern int32 GUseIndirectLightingCacheInLightingVolume;
-
-		if (GUseIndirectLightingCacheInLightingVolume && GSupportsVolumeTextureRendering)
-		{
-			GRenderTargetPool.FindFreeElement(
-				FPooledRenderTargetDesc(FPooledRenderTargetDesc::CreateVolumeDesc(
-					GTranslucencyLightingVolumeDim / TranslucentVolumeGISratchDownsampleFactor,
-					GTranslucencyLightingVolumeDim / TranslucentVolumeGISratchDownsampleFactor,
-					GTranslucencyLightingVolumeDim / TranslucentVolumeGISratchDownsampleFactor,
-					PF_FloatRGB,
-					0,
-					TexCreate_ShaderResource | TexCreate_RenderTargetable,
-					false)),
-				TranslucencyLightingVolumeGIScratch,
-				TEXT("TranslucentVolumeGIScratch")
-				);
-		}
 	}
 
 	extern ENGINE_API int32 GReflectionCaptureSize;
@@ -1331,8 +1313,6 @@ void FSceneRenderTargets::ReleaseAllTargets()
 		TranslucencyLightingVolumeDirectional[RTSetIndex].SafeRelease();
 	}
 
-	TranslucencyLightingVolumeGIScratch.SafeRelease();
-	
 	EditorPrimitivesColor.SafeRelease();
 	EditorPrimitivesDepth.SafeRelease();
 }
diff --git a/Engine/Source/Runtime/Renderer/Private/PostProcess/SceneRenderTargets.h b/Engine/Source/Runtime/Renderer/Private/PostProcess/SceneRenderTargets.h
--- a/Engine/Source/Runtime/Renderer/Private/PostProcess/SceneRenderTargets.h
+++ b/Engine/Source/Runtime/Renderer/Private/PostProcess/SceneRenderTargets.h
@@ -74,9 +74,6 @@ inline int SelectTranslucencyVolumeTarget(ETranslucencyVolumeCascade InCascade)
 /** Number of surfaces used for translucent shadows. */
 static const int32 NumTranslucencyShadowSurfaces = 2;
 
-/** Downsample factor to use on the volume texture used for computing GI on translucency. */
-static const int32 TranslucentVolumeGISratchDownsampleFactor = 1;
-
 BEGIN_UNIFORM_BUFFER_STRUCT(FGBufferResourceStruct, )
 	DECLARE_UNIFORM_BUFFER_STRUCT_MEMBER_TEXTURE( Texture2D,			GBufferATexture )
 	DECLARE_UNIFORM_BUFFER_STRUCT_MEMBER_TEXTURE( Texture2D,			GBufferBTexture )
@@ -462,9 +459,6 @@ public:
 	TRefCountPtr<IPooledRenderTarget> TranslucencyLightingVolumeAmbient[NumTranslucentVolumeRenderTargetSets];
 	TRefCountPtr<IPooledRenderTarget> TranslucencyLightingVolumeDirectional[NumTranslucentVolumeRenderTargetSets];
 
-	/** Volume texture used to gather GI at a lower resolution than the rest of translucent lighting. */
-	TRefCountPtr<IPooledRenderTarget> TranslucencyLightingVolumeGIScratch;
-
 	/** Color and opacity for editor primitives (i.e editor gizmos). */
 	TRefCountPtr<IPooledRenderTarget> EditorPrimitivesColor;
 
diff --git a/Engine/Source/Runtime/Renderer/Private/PrimitiveSceneInfo.h b/Engine/Source/Runtime/Renderer/Private/PrimitiveSceneInfo.h
--- a/Engine/Source/Runtime/Renderer/Private/PrimitiveSceneInfo.h
+++ b/Engine/Source/Runtime/Renderer/Private/PrimitiveSceneInfo.h
@@ -28,7 +28,7 @@ public:
 		CurrentDirectionalShadowing(1),
 		CurrentSkyBentNormal(FVector4(0, 0, 1, 1)),
 		bHasEverUpdatedSingleSample(false),
-		bOpaqueRelevance(true),
+		bPointSample(true),
 		bIsDirty(false)
 	{
 		for (int32 VectorIndex = 0; VectorIndex < ARRAY_COUNT(TargetSamplePacked); VectorIndex++)
@@ -83,8 +83,8 @@ public:
 	/** Whether SingleSamplePacked has ever been populated with valid results, used to initialize. */
 	bool bHasEverUpdatedSingleSample;
 
-	/** Whether the primitive this allocation is used for has opaque relevance. */
-	bool bOpaqueRelevance;
+	/** Whether this allocation is a point sample and therefore was not put into the volume texture atlas. */
+	bool bPointSample;
 
 	/** Whether the primitive allocation is dirty and should be updated regardless of having moved. */
 	bool bIsDirty;
@@ -99,7 +99,7 @@ public:
 		return MinTexel.X >= 0 && MinTexel.Y >= 0 && MinTexel.Z >= 0 && AllocationTexelSize > 0;
 	}
 
-	void SetParameters(FIntVector InMinTexel, int32 InAllocationTexelSize, FVector InScale, FVector InAdd, FVector InMinUV, FVector InMaxUV, bool bInOpaqueRelevance)
+	void SetParameters(FIntVector InMinTexel, int32 InAllocationTexelSize, FVector InScale, FVector InAdd, FVector InMinUV, FVector InMaxUV, bool bInPointSample)
 	{
 		Add = InAdd;
 		Scale = InScale;
@@ -108,7 +108,7 @@ public:
 		MinTexel = InMinTexel;
 		AllocationTexelSize = InAllocationTexelSize;
 		bIsDirty = false;
-		bOpaqueRelevance = bInOpaqueRelevance;
+		bPointSample = bInPointSample;
 	}
 };
 
diff --git a/Engine/Source/Runtime/Renderer/Private/ScenePrivate.h b/Engine/Source/Runtime/Renderer/Private/ScenePrivate.h
--- a/Engine/Source/Runtime/Renderer/Private/ScenePrivate.h
+++ b/Engine/Source/Runtime/Renderer/Private/ScenePrivate.h
@@ -883,16 +883,11 @@ private:
 	void UpdateCacheAllocation(
 		const FBoxSphereBounds& Bounds, 
 		int32 BlockSize,
-		bool bOpaqueRelevance,
+		bool bPointSample,
 		FIndirectLightingCacheAllocation*& Allocation, 
 		TMap<FIntVector, FBlockUpdateInfo>& BlocksToUpdate,
 		TArray<FIndirectLightingCacheAllocation*>& TransitionsOverTimeToUpdate);
 
-	void UpdateTranslucentVolumeCache(
-		FViewInfo& View, 
-		TMap<FIntVector, FBlockUpdateInfo>& BlocksToUpdate, 
-		TArray<FIndirectLightingCacheAllocation*>& TransitionsOverTimeToUpdate);
-
 	/** 
 	 * Creates a new allocation if needed, caches the result in PrimitiveSceneInfo->IndirectLightingCacheAllocation, 
 	 * And adds an entry to the update lists when an update is needed. 
@@ -912,7 +907,7 @@ private:
 	void UpdateTransitionsOverTime(const TArray<FIndirectLightingCacheAllocation*>& TransitionsOverTimeToUpdate, float DeltaWorldTime) const;
 
 	/** Creates an allocation to be used outside the indirect lighting cache and a block to be used internally. */
-	FIndirectLightingCacheAllocation* CreateAllocation(int32 BlockSize, const FBoxSphereBounds& Bounds, bool bOpaqueRelevance);
+	FIndirectLightingCacheAllocation* CreateAllocation(int32 BlockSize, const FBoxSphereBounds& Bounds, bool bPointSample);
 
 	/** Block accessors. */
 	FIndirectLightingCacheBlock& FindBlock(FIntVector TexelMin);
diff --git a/Engine/Source/Runtime/Renderer/Private/TranslucentLighting.cpp b/Engine/Source/Runtime/Renderer/Private/TranslucentLighting.cpp
--- a/Engine/Source/Runtime/Renderer/Private/TranslucentLighting.cpp
+++ b/Engine/Source/Runtime/Renderer/Private/TranslucentLighting.cpp
@@ -26,15 +26,6 @@ FAutoConsoleVariableRef CVarUseTranslucentLightingVolumes(
 	ECVF_Cheat | ECVF_RenderThreadSafe
 	);
 
-int32 GUseIndirectLightingCacheInLightingVolume = 0;
-FAutoConsoleVariableRef CUseIndirectLightingCacheInLightingVolume(
-	TEXT("r.IndirectLightingCacheInLightingVolume"),
-	GUseIndirectLightingCacheInLightingVolume,
-	TEXT("Whether to use the indirect lighting cache to inject GI into the translucency lighting volume.  Causes extreme hitching at the moment."),
-	ECVF_Cheat | ECVF_RenderThreadSafe
-	);
-
-
 float GTranslucentVolumeMinFOV = 45;
 static FAutoConsoleVariableRef CVarTranslucentVolumeMinFOV(
 	TEXT("r.TranslucentVolumeMinFOV"),
@@ -146,9 +137,9 @@ void FViewInfo::CalcTranslucencyLightingVolumeBounds(FBox* InOutCascadeBoundsArr
 		FSphere SphereBounds(Center, FMath::Sqrt(RadiusSquared));
 
 		// Snap the center to a multiple of the volume dimension for stability
-		SphereBounds.Center.X = SphereBounds.Center.X - FMath::Fmod(SphereBounds.Center.X, SphereBounds.W * 2 * TranslucentVolumeGISratchDownsampleFactor / GTranslucencyLightingVolumeDim);
-		SphereBounds.Center.Y = SphereBounds.Center.Y - FMath::Fmod(SphereBounds.Center.Y, SphereBounds.W * 2 * TranslucentVolumeGISratchDownsampleFactor / GTranslucencyLightingVolumeDim);
-		SphereBounds.Center.Z = SphereBounds.Center.Z - FMath::Fmod(SphereBounds.Center.Z, SphereBounds.W * 2 * TranslucentVolumeGISratchDownsampleFactor / GTranslucencyLightingVolumeDim);
+		SphereBounds.Center.X = SphereBounds.Center.X - FMath::Fmod(SphereBounds.Center.X, SphereBounds.W * 2 / GTranslucencyLightingVolumeDim);
+		SphereBounds.Center.Y = SphereBounds.Center.Y - FMath::Fmod(SphereBounds.Center.Y, SphereBounds.W * 2 / GTranslucencyLightingVolumeDim);
+		SphereBounds.Center.Z = SphereBounds.Center.Z - FMath::Fmod(SphereBounds.Center.Z, SphereBounds.W * 2 / GTranslucencyLightingVolumeDim);
 
 		InOutCascadeBoundsArray[CascadeIndex] = FBox(SphereBounds.Center - SphereBounds.W, SphereBounds.Center + SphereBounds.W);
 	}
@@ -1227,173 +1218,6 @@ void FDeferredShadingSceneRenderer::InjectAmbientCubemapTranslucentVolumeLightin
 	}
 }
 
-
-/** Pixel shader used to filter a single volume lighting cascade. */
-class FCompositeGIForTranslucencyPS : public FGlobalShader
-{
-	DECLARE_SHADER_TYPE(FCompositeGIForTranslucencyPS,Global);
-public:
-
-	static bool ShouldCache(EShaderPlatform Platform) 
-	{ 
-		return IsFeatureLevelSupported(Platform, ERHIFeatureLevel::SM4); 
-	}
-
-	FCompositeGIForTranslucencyPS(const ShaderMetaType::CompiledShaderInitializerType& Initializer):
-		FGlobalShader(Initializer)
-	{
-		IndirectLightingCacheTexture.Bind(Initializer.ParameterMap, TEXT("IndirectLightingCacheTexture"));
-		IndirectLightingCacheTexture1.Bind(Initializer.ParameterMap, TEXT("IndirectLightingCacheTexture1"));
-		IndirectLightingCacheTexture2.Bind(Initializer.ParameterMap, TEXT("IndirectLightingCacheTexture2"));
-		IndirectLightingCacheTextureSampler.Bind(Initializer.ParameterMap, TEXT("IndirectLightingCacheTextureSampler"));
-		IndirectLightingCacheTexture1Sampler.Bind(Initializer.ParameterMap, TEXT("IndirectLightingCacheTexture1Sampler"));
-		IndirectLightingCacheTexture2Sampler.Bind(Initializer.ParameterMap, TEXT("IndirectLightingCacheTexture2Sampler"));
-		IndirectlightingCacheAdd.Bind(Initializer.ParameterMap, TEXT("IndirectlightingCacheAdd"));
-		IndirectlightingCacheScale.Bind(Initializer.ParameterMap, TEXT("IndirectlightingCacheScale"));
-		VolumeCascadeIndex.Bind(Initializer.ParameterMap,TEXT("VolumeCascadeIndex"));
-	}
-	FCompositeGIForTranslucencyPS() {}
-
-	void SetParameters(FRHICommandList& RHICmdList, const FViewInfo& View, const FSceneViewState* ViewState, int32 VolumeCascadeIndexValue)
-	{
-		const FPixelShaderRHIParamRef ShaderRHI = GetPixelShader();
-
-		FGlobalShader::SetParameters(RHICmdList, ShaderRHI, View);
-
-		SetShaderValue(RHICmdList, ShaderRHI, VolumeCascadeIndex, VolumeCascadeIndexValue);
-
-		const FIndirectLightingCacheAllocation* LightingAllocation = ViewState->TranslucencyLightingCacheAllocations[VolumeCascadeIndexValue];
-
-		if (LightingAllocation && LightingAllocation->IsValid())
-		{
-			SetShaderValue(RHICmdList, ShaderRHI, IndirectlightingCacheAdd, LightingAllocation->Add);
-			SetShaderValue(RHICmdList, ShaderRHI, IndirectlightingCacheScale, LightingAllocation->Scale);
-
-			FScene* Scene = (FScene*)View.Family->Scene;
-			FIndirectLightingCache& LightingCache = Scene->IndirectLightingCache;
-
-			if (LightingCache.IsInitialized())
-			{
-				SetTextureParameter(
-					RHICmdList, 
-					ShaderRHI,
-					IndirectLightingCacheTexture,
-					IndirectLightingCacheTextureSampler,
-					TStaticSamplerState<SF_Bilinear,AM_Clamp,AM_Clamp,AM_Clamp>::GetRHI(),
-					LightingCache.GetTexture0().ShaderResourceTexture
-					);
-
-				SetTextureParameter(
-					RHICmdList, 
-					ShaderRHI,
-					IndirectLightingCacheTexture1,
-					IndirectLightingCacheTexture1Sampler,
-					TStaticSamplerState<SF_Bilinear,AM_Clamp,AM_Clamp,AM_Clamp>::GetRHI(),
-					LightingCache.GetTexture1().ShaderResourceTexture
-					);
-
-				SetTextureParameter(
-					RHICmdList, 
-					ShaderRHI,
-					IndirectLightingCacheTexture2,
-					IndirectLightingCacheTexture2Sampler,
-					TStaticSamplerState<SF_Bilinear,AM_Clamp,AM_Clamp,AM_Clamp>::GetRHI(),
-					LightingCache.GetTexture2().ShaderResourceTexture
-					);
-			}
-		}
-	}
-
-	virtual bool Serialize(FArchive& Ar)
-	{
-		bool bShaderHasOutdatedParameters = FGlobalShader::Serialize(Ar);
-		Ar << IndirectLightingCacheTexture;
-		Ar << IndirectLightingCacheTexture1;
-		Ar << IndirectLightingCacheTexture2;
-		Ar << IndirectLightingCacheTextureSampler;
-		Ar << IndirectLightingCacheTexture1Sampler;
-		Ar << IndirectLightingCacheTexture2Sampler;
-		Ar << IndirectlightingCacheAdd;
-		Ar << IndirectlightingCacheScale;
-		Ar << VolumeCascadeIndex;
-		return bShaderHasOutdatedParameters;
-	}
-
-private:
-	FShaderResourceParameter IndirectLightingCacheTexture;
-	FShaderResourceParameter IndirectLightingCacheTexture1;
-	FShaderResourceParameter IndirectLightingCacheTexture2;
-	FShaderResourceParameter IndirectLightingCacheTextureSampler;
-	FShaderResourceParameter IndirectLightingCacheTexture1Sampler;
-	FShaderResourceParameter IndirectLightingCacheTexture2Sampler;
-	FShaderParameter IndirectlightingCacheAdd;
-	FShaderParameter IndirectlightingCacheScale;
-	FShaderParameter VolumeCascadeIndex;
-};
-
-IMPLEMENT_SHADER_TYPE(,FCompositeGIForTranslucencyPS,TEXT("TranslucentLightingShaders"),TEXT("CompositeGIMainPS"),SF_Pixel);
-
-void FDeferredShadingSceneRenderer::CompositeIndirectTranslucentVolumeLighting(FRHICommandList& RHICmdList)
-{
-	bool bAnyViewAllowsIndirectLightingCache = false;
-
-	for (int32 ViewIndex = 0; ViewIndex < Views.Num(); ViewIndex++)
-	{
-		bAnyViewAllowsIndirectLightingCache |= Views[ViewIndex].Family->EngineShowFlags.IndirectLightingCache;
-	}
-
-	if (GUseTranslucentLightingVolumes && GSupportsVolumeTextureRendering
-		&& GUseIndirectLightingCacheInLightingVolume
-		&& bAnyViewAllowsIndirectLightingCache)
-	{
-		SCOPED_DRAW_EVENT(CompositeIndirectTranslucentVolumeLighting, DEC_SCENE_ITEMS);
-
-		RHICmdList.SetRasterizerState(TStaticRasterizerState<FM_Solid, CM_None>::GetRHI());
-		RHICmdList.SetDepthStencilState(TStaticDepthStencilState<false, CF_Always>::GetRHI());
-		RHICmdList.SetBlendState(TStaticBlendState<CW_RGBA, BO_Add, BF_One, BF_One, BO_Add, BF_One, BF_One>::GetRHI());
-
-		const FVolumeBounds VolumeBounds(GTranslucencyLightingVolumeDim);
-
-		for (int32 VolumeCascadeIndex = 0; VolumeCascadeIndex < TVC_MAX; VolumeCascadeIndex++)
-		{
-			FTextureRHIParamRef RenderTargets[2];
-			RenderTargets[0] = GSceneRenderTargets.TranslucencyLightingVolumeAmbient[VolumeCascadeIndex]->GetRenderTargetItem().TargetableTexture;
-			RenderTargets[1] = GSceneRenderTargets.TranslucencyLightingVolumeDirectional[VolumeCascadeIndex]->GetRenderTargetItem().TargetableTexture;
-
-			SetRenderTargets(RHICmdList, ARRAY_COUNT(RenderTargets), RenderTargets, FTextureRHIRef(), 0, NULL);
-
-
-			//@todo - support multiple views
-			const FViewInfo& View = Views[0];
-			const FSceneViewState* ViewState = (const FSceneViewState*)View.State;
-
-			if (ViewState 
-				&& View.Family->EngineShowFlags.IndirectLightingCache
-				&& View.Family->EngineShowFlags.GlobalIllumination)
-			{
-				TShaderMapRef<FWriteToSliceVS> VertexShader(GetGlobalShaderMap());
-				TShaderMapRef<FWriteToSliceGS> GeometryShader(GetGlobalShaderMap());
-				TShaderMapRef<FCompositeGIForTranslucencyPS> PixelShader(GetGlobalShaderMap());
-
-				static FGlobalBoundShaderState BoundShaderState;
-				
-				SetGlobalBoundShaderState(RHICmdList, BoundShaderState, GScreenVertexDeclaration.VertexDeclarationRHI, *VertexShader, *PixelShader, *GeometryShader);
-
-				PixelShader->SetParameters(RHICmdList, View, ViewState, VolumeCascadeIndex);
-				VertexShader->SetParameters(RHICmdList, VolumeBounds, GTranslucencyLightingVolumeDim);
-				GeometryShader->SetParameters(RHICmdList, VolumeBounds);
-
-				RasterizeToVolumeTexture(RHICmdList, VolumeBounds);
-
-				RHICmdList.CopyToResolveTarget(GSceneRenderTargets.TranslucencyLightingVolumeAmbient[VolumeCascadeIndex]->GetRenderTargetItem().TargetableTexture,
-					GSceneRenderTargets.TranslucencyLightingVolumeAmbient[VolumeCascadeIndex]->GetRenderTargetItem().ShaderResourceTexture, true, FResolveParams());
-				RHICmdList.CopyToResolveTarget(GSceneRenderTargets.TranslucencyLightingVolumeDirectional[VolumeCascadeIndex]->GetRenderTargetItem().TargetableTexture,
-					GSceneRenderTargets.TranslucencyLightingVolumeDirectional[VolumeCascadeIndex]->GetRenderTargetItem().ShaderResourceTexture, true, FResolveParams());
-			}
-		}
-	}
-}
-
 void FDeferredShadingSceneRenderer::ClearTranslucentVolumePerObjectShadowing(FRHICommandList& RHICmdList)
 {
 	if (GUseTranslucentLightingVolumes && GSupportsVolumeTextureRendering)
