First half of getting r.MobileMSAA working with reversible tonemap.

diff --git a/Engine/Shaders/PostProcessMobile.usf b/Engine/Shaders/PostProcessMobile.usf
--- a/Engine/Shaders/PostProcessMobile.usf
+++ b/Engine/Shaders/PostProcessMobile.usf
@@ -135,6 +135,11 @@ void BloomPS_ES2(
 	// Output color is average.
 	OutColor.rgb = (C0.rgb * 0.25) + (C1.rgb * 0.25) + (C2.rgb * 0.25) + (C3.rgb * 0.25);
 
+	#if ES2_USE_MSAA
+		// This should really happen before the average, instead doing after average as optimization.
+		OutColor.rgb *= rcp(OutColor.r*(-0.299) + OutColor.g*(-0.587) + OutColor.b*(-0.114) + 1.0);
+	#endif
+
 	// Try to kill negatives and NaNs here
 	OutColor.rgb = max(OutColor.rgb, 0);
 
diff --git a/Engine/Shaders/PostProcessTonemap.usf b/Engine/Shaders/PostProcessTonemap.usf
--- a/Engine/Shaders/PostProcessTonemap.usf
+++ b/Engine/Shaders/PostProcessTonemap.usf
@@ -407,6 +407,9 @@ void MainPS_ES2(
 	half4 SceneColor = PostprocessInput0.Sample(PostprocessInput0Sampler, UV);
 
 	#if USE_GAMMA_ONLY
+		#if USE_MSAA
+			SceneColor.rgb *= rcp(SceneColor.r*(-0.299) + SceneColor.g*(-0.587) + SceneColor.b*(-0.114) + 1.0);
+		#endif
 		OutColor.rgb = sqrt(SceneColor.rgb);
 	#else
 
@@ -447,6 +450,11 @@ void MainPS_ES2(
 			#endif
 		#endif
 
+		#if USE_MSAA
+			// Do after jitter for grain as an optimization.
+			SceneColor.rgb *= rcp(SceneColor.r*(-0.299) + SceneColor.g*(-0.587) + SceneColor.b*(-0.114) + 1.0);
+		#endif
+
 		#if USE_HDR_MOSAIC
 			// TODO: Support odd frame inversion of mosaic pattern?
 			#if 0
diff --git a/Engine/Source/Runtime/Engine/Private/SceneView.cpp b/Engine/Source/Runtime/Engine/Private/SceneView.cpp
--- a/Engine/Source/Runtime/Engine/Private/SceneView.cpp
+++ b/Engine/Source/Runtime/Engine/Private/SceneView.cpp
@@ -826,6 +826,16 @@ void FSceneView::StartFinalPostprocessSettings(FVector InViewLocation)
 
 void FSceneView::EndFinalPostprocessSettings()
 {
+	{
+		static const auto CVarMobileMSAA = IConsoleManager::Get().FindTConsoleVariableDataInt(TEXT("r.MobileMSAA"));
+		if(CVarMobileMSAA ? CVarMobileMSAA->GetValueOnGameThread() > 1 : false)
+		{
+			// Turn off various features which won't work with mobile MSAA.
+			FinalPostProcessSettings.DepthOfFieldScale = 0.0f;
+			FinalPostProcessSettings.AntiAliasingMethod = AAM_None;
+		}
+	}
+
 	{
 		static const auto CVar = IConsoleManager::Get().FindTConsoleVariableDataInt(TEXT("r.BloomQuality"));
 
diff --git a/Engine/Source/Runtime/Renderer/Private/PostProcess/PostProcessMobile.cpp b/Engine/Source/Runtime/Renderer/Private/PostProcess/PostProcessMobile.cpp
--- a/Engine/Source/Runtime/Renderer/Private/PostProcess/PostProcessMobile.cpp
+++ b/Engine/Source/Runtime/Renderer/Private/PostProcess/PostProcessMobile.cpp
@@ -57,7 +57,7 @@ public:
 	}
 };
 
-template <uint32 UseSunDof> // 0=none, 1=dof, 2=sun, 3=sun&dof
+template <uint32 UseSunDof> // 0=none, 1=dof, 2=sun, 3=sun&dof, 4=msaa
 class FPostProcessBloomSetupPS_ES2 : public FGlobalShader
 {
 	DECLARE_SHADER_TYPE(FPostProcessBloomSetupPS_ES2, Global);
@@ -74,6 +74,7 @@ class FPostProcessBloomSetupPS_ES2 : public FGlobalShader
 		//Need to hack in exposure scale for < SM5
 		OutEnvironment.SetDefine(TEXT("NO_EYEADAPTATION_EXPOSURE_FIX"), 1);
 
+		OutEnvironment.SetDefine(TEXT("ES2_USE_MSAA"), (UseSunDof & 4) ? (uint32)1 : (uint32)0);
 		OutEnvironment.SetDefine(TEXT("ES2_USE_SUN"), (UseSunDof & 2) ? (uint32)1 : (uint32)0);
 		OutEnvironment.SetDefine(TEXT("ES2_USE_DOF"), (UseSunDof & 1) ? (uint32)1 : (uint32)0);
 	}
@@ -123,10 +124,12 @@ typedef FPostProcessBloomSetupPS_ES2<0> FPostProcessBloomSetupPS_ES2_0;
 typedef FPostProcessBloomSetupPS_ES2<1> FPostProcessBloomSetupPS_ES2_1;
 typedef FPostProcessBloomSetupPS_ES2<2> FPostProcessBloomSetupPS_ES2_2;
 typedef FPostProcessBloomSetupPS_ES2<3> FPostProcessBloomSetupPS_ES2_3;
+typedef FPostProcessBloomSetupPS_ES2<4> FPostProcessBloomSetupPS_ES2_4;
 IMPLEMENT_SHADER_TYPE(template<>,FPostProcessBloomSetupPS_ES2_0,TEXT("PostProcessMobile"),TEXT("BloomPS_ES2"),SF_Pixel);
 IMPLEMENT_SHADER_TYPE(template<>,FPostProcessBloomSetupPS_ES2_1,TEXT("PostProcessMobile"),TEXT("BloomPS_ES2"),SF_Pixel);
 IMPLEMENT_SHADER_TYPE(template<>,FPostProcessBloomSetupPS_ES2_2,TEXT("PostProcessMobile"),TEXT("BloomPS_ES2"),SF_Pixel);
 IMPLEMENT_SHADER_TYPE(template<>,FPostProcessBloomSetupPS_ES2_3,TEXT("PostProcessMobile"),TEXT("BloomPS_ES2"),SF_Pixel);
+IMPLEMENT_SHADER_TYPE(template<>,FPostProcessBloomSetupPS_ES2_4,TEXT("PostProcessMobile"),TEXT("BloomPS_ES2"),SF_Pixel);
 
 template <uint32 UseSunDof>
 static void BloomSetup_SetShader(const FRenderingCompositePassContext& Context)
@@ -149,12 +152,16 @@ void FRCPassPostProcessBloomSetupES2::SetShader(const FRenderingCompositePassCon
 	uint32 UseDof =  (Context.View.FinalPostProcessSettings.DepthOfFieldScale > 0.0f) ? 1 : 0;
 	uint32 UseSunDof = (UseSun << 1) + UseDof;
 
+	static const auto CVarMobileMSAA = IConsoleManager::Get().FindTConsoleVariableDataInt(TEXT("r.MobileMSAA"));
+	UseSunDof += CVarMobileMSAA ? (CVarMobileMSAA->GetValueOnRenderThread() > 1 ? 4 : 0) : 0;
+
 	switch(UseSunDof)
 	{
 		case 0: BloomSetup_SetShader<0>(Context); break;
 		case 1: BloomSetup_SetShader<1>(Context); break;
 		case 2: BloomSetup_SetShader<2>(Context); break;
 		case 3: BloomSetup_SetShader<3>(Context); break;
+		case 4: BloomSetup_SetShader<4>(Context); break;
 	}
 }
 
diff --git a/Engine/Source/Runtime/Renderer/Private/PostProcess/PostProcessTonemap.cpp b/Engine/Source/Runtime/Renderer/Private/PostProcess/PostProcessTonemap.cpp
--- a/Engine/Source/Runtime/Renderer/Private/PostProcess/PostProcessTonemap.cpp
+++ b/Engine/Source/Runtime/Renderer/Private/PostProcess/PostProcessTonemap.cpp
@@ -36,7 +36,7 @@ typedef enum {
 	TonemapperMosaic            = (1<<12),
 	TonemapperColorFringe       = (1<<13),
 	TonemapperColorGrading      = (1<<14),
-
+	TonemapperMsaa              = (1<<15),
 } TonemapperOption;
 
 // Tonemapper option cost (0 = no cost, 255 = max cost).
@@ -58,6 +58,7 @@ static uint8 TonemapperCostTab[] = {
 	1, //TonemapperMosaic
 	1, //TonemapperColorFringe
 	1, //TonemapperColorGrading
+	1, //TonemapperMsaa
 };
 
 // Edit the following to add and remove configurations.
@@ -244,7 +245,7 @@ static uint32 TonemapperConfBitmaskPC[22] = {
 };
 
 // List of configurations compiled for Mobile.
-static uint32 TonemapperConfBitmaskMobile[29] = { 
+static uint32 TonemapperConfBitmaskMobile[39] = { 
 
 	// 
 	//  15 for NON-MOSAIC 
@@ -461,6 +462,83 @@ static uint32 TonemapperConfBitmaskMobile[29] = {
 	TonemapperGrainIntensity +
 	0,
 
+
+	// 
+	//  10 for MSAA
+	//
+
+	TonemapperMsaa +
+	TonemapperContrast +
+	0,
+
+	TonemapperMsaa +
+	TonemapperContrast +
+	TonemapperColorMatrix +
+	0,
+
+	TonemapperMsaa +
+	TonemapperContrast +
+	TonemapperBloom +
+	TonemapperVignette +
+	TonemapperVignetteColor +
+	0,
+
+	TonemapperMsaa +
+	TonemapperContrast +
+	TonemapperColorMatrix +
+	TonemapperBloom +
+	TonemapperVignette +
+	TonemapperVignetteColor +
+	0,
+
+	TonemapperMsaa +
+	TonemapperContrast +
+	TonemapperColorMatrix +
+	TonemapperShadowTint +
+	TonemapperBloom +
+	TonemapperVignette +
+	TonemapperVignetteColor +
+	0,
+
+	// Same with grain.
+	TonemapperMsaa +
+	TonemapperContrast +
+	TonemapperGrainIntensity +
+	0,
+
+	TonemapperMsaa +
+	TonemapperContrast +
+	TonemapperColorMatrix +
+	TonemapperGrainIntensity +
+	0,
+
+	TonemapperMsaa +
+	TonemapperContrast +
+	TonemapperBloom +
+	TonemapperVignette +
+	TonemapperVignetteColor +
+	TonemapperGrainIntensity +
+	0,
+
+	TonemapperMsaa +
+	TonemapperContrast +
+	TonemapperColorMatrix +
+	TonemapperBloom +
+	TonemapperVignette +
+	TonemapperVignetteColor +
+	TonemapperGrainIntensity +
+	0,
+
+	TonemapperMsaa +
+	TonemapperContrast +
+	TonemapperColorMatrix +
+	TonemapperShadowTint +
+	TonemapperBloom +
+	TonemapperVignette +
+	TonemapperVignetteColor +
+	TonemapperGrainIntensity +
+	0,
+
 };
 
 // Returns 1 if option is defined otherwise 0.
@@ -476,6 +554,7 @@ static uint32 TonemapperFindLeastExpensive(uint32* RESTRICT Table, uint32 TableE
 	uint32 MustNotHaveBitmask = 0;
 	MustNotHaveBitmask += ((RequiredOptionsBitmask & TonemapperDOF) == 0) ? TonemapperDOF : 0;
 	MustNotHaveBitmask += ((RequiredOptionsBitmask & TonemapperMosaic) == 0) ? TonemapperMosaic : 0;
+	MustNotHaveBitmask += ((RequiredOptionsBitmask & TonemapperMsaa) == 0) ? TonemapperMsaa : 0;
 
 	// Search for exact match first.
 	uint32 Index;
@@ -644,6 +723,12 @@ static uint32 TonemapperGenerateBitmaskMobile(const FRenderingCompositePassConte
 		return Bitmask + TonemapperMosaic;
 	}
 
+	static const auto CVarMobileMSAA = IConsoleManager::Get().FindTConsoleVariableDataInt(TEXT("r.MobileMSAA"));
+	if(CVarMobileMSAA ? CVarMobileMSAA->GetValueOnGameThread() > 1 : false)
+	{
+		Bitmask += TonemapperMsaa;
+	}
+
 	// Only add mobile post if FP16 is supported.
 	if(GSupportsRenderTargetFormat_PF_FloatRGBA)
 	{
@@ -1256,6 +1341,7 @@ class FPostProcessTonemapPS_ES2 : public FGlobalShader
 		OutEnvironment.SetDefine(TEXT("USE_VIGNETTE_COLOR"),     TonemapperIsDefined(ConfigBitmask, TonemapperVignetteColor));
 		OutEnvironment.SetDefine(TEXT("USE_LIGHT_SHAFTS"),       TonemapperIsDefined(ConfigBitmask, TonemapperLightShafts));
 		OutEnvironment.SetDefine(TEXT("USE_DOF"),                TonemapperIsDefined(ConfigBitmask, TonemapperDOF));
+		OutEnvironment.SetDefine(TEXT("USE_MSAA"),               TonemapperIsDefined(ConfigBitmask, TonemapperMsaa));
 
 		//Need to hack in exposure scale for < SM5
 		OutEnvironment.SetDefine(TEXT("NO_EYEADAPTATION_EXPOSURE_FIX"), 1);
@@ -1403,7 +1489,9 @@ public:
 
 	VARIATION2(0)  VARIATION2(1)  VARIATION2(2)  VARIATION2(3)	VARIATION2(4)  VARIATION2(5)  VARIATION2(6)  VARIATION2(7)  VARIATION2(8)  VARIATION2(9)  
 	VARIATION2(10) VARIATION2(11) VARIATION2(12) VARIATION2(13)	VARIATION2(14) VARIATION2(15) VARIATION2(16) VARIATION2(17) VARIATION2(18) VARIATION2(19)  
-	VARIATION2(20) VARIATION2(21) VARIATION2(22) VARIATION2(23)	VARIATION2(24) VARIATION2(25) VARIATION2(26) VARIATION2(27) VARIATION2(28)
+	VARIATION2(20) VARIATION2(21) VARIATION2(22) VARIATION2(23)	VARIATION2(24) VARIATION2(25) VARIATION2(26) VARIATION2(27) VARIATION2(28) VARIATION2(29)
+	VARIATION2(30) VARIATION2(31) VARIATION2(32) VARIATION2(33)	VARIATION2(34) VARIATION2(35) VARIATION2(36) VARIATION2(37) VARIATION2(38)
+
 
 #undef VARIATION2
 	
@@ -1542,6 +1630,16 @@ void FRCPassPostProcessTonemapES2::Process(FRenderingCompositePassContext& Conte
 		case 26: SetShaderTemplES2<26>(Context, bUsedFramebufferFetch); break;
 		case 27: SetShaderTemplES2<27>(Context, bUsedFramebufferFetch); break;
 		case 28: SetShaderTemplES2<28>(Context, bUsedFramebufferFetch); break;
+		case 29: SetShaderTemplES2<29>(Context, bUsedFramebufferFetch); break;
+		case 30: SetShaderTemplES2<30>(Context, bUsedFramebufferFetch); break;
+		case 31: SetShaderTemplES2<31>(Context, bUsedFramebufferFetch); break;
+		case 32: SetShaderTemplES2<32>(Context, bUsedFramebufferFetch); break;
+		case 33: SetShaderTemplES2<33>(Context, bUsedFramebufferFetch); break;
+		case 34: SetShaderTemplES2<34>(Context, bUsedFramebufferFetch); break;
+		case 35: SetShaderTemplES2<35>(Context, bUsedFramebufferFetch); break;
+		case 36: SetShaderTemplES2<36>(Context, bUsedFramebufferFetch); break;
+		case 37: SetShaderTemplES2<37>(Context, bUsedFramebufferFetch); break;
+		case 38: SetShaderTemplES2<38>(Context, bUsedFramebufferFetch); break;
 		default:
 			check(0);
 	}
diff --git a/Engine/Source/Runtime/Renderer/Private/SceneVisibility.cpp b/Engine/Source/Runtime/Renderer/Private/SceneVisibility.cpp
--- a/Engine/Source/Runtime/Renderer/Private/SceneVisibility.cpp
+++ b/Engine/Source/Runtime/Renderer/Private/SceneVisibility.cpp
@@ -1559,6 +1559,9 @@ void FSceneRenderer::PostVisibilityFrameSetup()
 			{
 				VisibleLightViewInfo.bInViewFrustum = true;
 
+				static const auto CVarMobileMSAA = IConsoleManager::Get().FindTConsoleVariableDataInt(TEXT("r.MobileMSAA"));
+				bool bNotMobileMSAA = !(CVarMobileMSAA ? CVarMobileMSAA->GetValueOnRenderThread() > 1 : false);
+
 				// Setup single sun-shaft from direction lights for mobile.
 				if(bCheckLightShafts && LightSceneInfo->bEnableLightShaftBloom)
 				{
@@ -1579,7 +1582,7 @@ void FSceneRenderer::PostVisibilityFrameSetup()
 							&& (LightSceneInfo->Proxy->GetLightType() == LightType_Directional 
 							|| DistanceToBlurOrigin < LightSceneInfo->Proxy->GetRadius() * PointLightRadiusFadeFactor))
 					{
-						View.bLightShaftUse = true;
+						View.bLightShaftUse = bNotMobileMSAA;
 						View.LightShaftCenter.X = ProjectedBlurOrigin.X / ProjectedBlurOrigin.W;
 						View.LightShaftCenter.Y = ProjectedBlurOrigin.Y / ProjectedBlurOrigin.W;
 						// TODO: Might want to hookup different colors for these.
