Fixed compilation errors

diff --git a/Engine/Source/Runtime/Renderer/Private/BasePassRendering.h b/Engine/Source/Runtime/Renderer/Private/BasePassRendering.h
--- a/Engine/Source/Runtime/Renderer/Private/BasePassRendering.h
+++ b/Engine/Source/Runtime/Renderer/Private/BasePassRendering.h
@@ -925,7 +925,7 @@ void ProcessBasePassMesh(
 
 								// Use the indirect lighting cache shaders if the object has a cache allocation
 								// This happens for objects with unbuilt lighting
-								if (IndirectLightingCacheAllocation && IndirectLightingCacheAllocation->IsValid()
+								if ((IndirectLightingCacheAllocation && IndirectLightingCacheAllocation->IsValid())
 									// Use the indirect lighting cache shaders if the object is movable, it may not have a cache allocation yet because that is done in InitViews
 									// And movable objects are sometimes rendered in the static draw lists
 									|| bPrimitiveIsMovable)
@@ -933,8 +933,8 @@ void ProcessBasePassMesh(
 									if (CanIndirectLightingCacheUseVolumeTexture(Parameters.FeatureLevel) 
 										// Translucency forces point sample for pixel performance
 										&& Action.AllowIndirectLightingCacheVolumeTexture()
-										&& (IndirectLightingCacheAllocation && !IndirectLightingCacheAllocation->bPointSample
-											|| bPrimitiveIsMovable && Parameters.PrimitiveSceneProxy->GetIndirectLightingCacheQuality() == ILCQ_Volume))
+										&& ((IndirectLightingCacheAllocation && !IndirectLightingCacheAllocation->bPointSample)
+											|| (bPrimitiveIsMovable && Parameters.PrimitiveSceneProxy->GetIndirectLightingCacheQuality() == ILCQ_Volume)))
 									{
 										// Use a lightmap policy that supports reading indirect lighting from a volume texture for dynamic objects
 										Action.template Process<FCachedVolumeIndirectLightingPolicy>(RHICmdList, Parameters, FCachedVolumeIndirectLightingPolicy(), FCachedVolumeIndirectLightingPolicy::ElementDataType());
diff --git a/Engine/Source/Runtime/Renderer/Private/IndirectLightingCache.cpp b/Engine/Source/Runtime/Renderer/Private/IndirectLightingCache.cpp
--- a/Engine/Source/Runtime/Renderer/Private/IndirectLightingCache.cpp
+++ b/Engine/Source/Runtime/Renderer/Private/IndirectLightingCache.cpp
@@ -483,7 +483,7 @@ void FIndirectLightingCache::UpdateCachePrimitive(
 	FIndirectLightingCacheAllocation* PrimitiveAllocation = PrimitiveAllocationPtr != NULL ? *PrimitiveAllocationPtr : NULL;
 
 	if (PrimitiveSceneProxy->WillEverBeLit()
-		&& (bAllowUnbuiltPreview && PrimitiveSceneProxy->HasStaticLighting() && PrimitiveAllocation && PrimitiveAllocation->bIsDirty
+		&& ((bAllowUnbuiltPreview && PrimitiveSceneProxy->HasStaticLighting() && PrimitiveAllocation && PrimitiveAllocation->bIsDirty)
 		|| (PrimitiveSceneProxy->IsMovable() && PrimitiveSceneProxy->GetIndirectLightingCacheQuality() != ILCQ_Off)))
 	{
 		const FIndirectLightingCacheAllocation* AttachmentParentAllocation = NULL;
