Fix for a crash on ios

CVar was being queried on an incorrect thread.

diff --git a/Engine/Source/Runtime/Renderer/Private/PostProcess/PostProcessTonemap.cpp b/Engine/Source/Runtime/Renderer/Private/PostProcess/PostProcessTonemap.cpp
--- a/Engine/Source/Runtime/Renderer/Private/PostProcess/PostProcessTonemap.cpp
+++ b/Engine/Source/Runtime/Renderer/Private/PostProcess/PostProcessTonemap.cpp
@@ -724,7 +724,7 @@ static uint32 TonemapperGenerateBitmaskMobile(const FRenderingCompositePassConte
 	}
 
 	static const auto CVarMobileMSAA = IConsoleManager::Get().FindTConsoleVariableDataInt(TEXT("r.MobileMSAA"));
-	if(CVarMobileMSAA ? CVarMobileMSAA->GetValueOnGameThread() > 1 : false)
+	if(CVarMobileMSAA ? CVarMobileMSAA->GetValueOnAnyThread() > 1 : false)
 	{
 		Bitmask += TonemapperMsaa;
 	}
