Buffered occlusion queries for ps4

diff --git a/Engine/Source/Runtime/Core/Private/HAL/ConsoleManager.cpp b/Engine/Source/Runtime/Core/Private/HAL/ConsoleManager.cpp
--- a/Engine/Source/Runtime/Core/Private/HAL/ConsoleManager.cpp
+++ b/Engine/Source/Runtime/Core/Private/HAL/ConsoleManager.cpp
@@ -1575,6 +1575,14 @@ static TAutoConsoleVariable<int32> CVarAllowStaticLighting(
 	TEXT("Games that only use dynamic lighting should set this to 0 to save some static lighting overhead."),
 	ECVF_ReadOnly | ECVF_RenderThreadSafe);
 
+
+static TAutoConsoleVariable<int32> CVarNumBufferedOcclusionQueries(
+	TEXT("r.NumBufferedOcclusionQueries"),
+	1,
+	TEXT("Number of frames to buffer occlusion queries (including the current renderthread frame).\n")
+	TEXT("More frames reduces the chance of stalling the CPU waiting for results, but increases out of date query artifacts."),
+	ECVF_ReadOnly);
+
 static TAutoConsoleVariable<int32> CVarDistField(
 	TEXT("r.AllowMeshDistanceFieldRepresentations"),
 	0,	
diff --git a/Engine/Source/Runtime/Renderer/Private/RendererScene.cpp b/Engine/Source/Runtime/Renderer/Private/RendererScene.cpp
--- a/Engine/Source/Runtime/Renderer/Private/RendererScene.cpp
+++ b/Engine/Source/Runtime/Renderer/Private/RendererScene.cpp
@@ -64,6 +64,12 @@ FSceneViewState::FSceneViewState()
 	{
 		TranslucencyLightingCacheAllocations[CascadeIndex] = NULL;
 	}
+
+#if BUFFERED_OCCLUSION_QUERIES
+	NumBufferedFrames = FOcclusionQueryHelpers::GetNumBufferedFrames();
+	ShadowOcclusionQueryMaps.Empty(NumBufferedFrames);
+	ShadowOcclusionQueryMaps.AddZeroed(NumBufferedFrames);	
+#endif
 }
 
 /**
diff --git a/Engine/Source/Runtime/Renderer/Private/SceneOcclusion.cpp b/Engine/Source/Runtime/Renderer/Private/SceneOcclusion.cpp
--- a/Engine/Source/Runtime/Renderer/Private/SceneOcclusion.cpp
+++ b/Engine/Source/Runtime/Renderer/Private/SceneOcclusion.cpp
@@ -228,6 +228,14 @@ bool FSceneViewState::IsShadowOccluded(FRHICommandListImmediate& RHICmdList, FPr
 {
 	// Find the shadow's occlusion query from the previous frame.
 	const FSceneViewState::FProjectedShadowKey Key(PrimitiveId, Light, SplitIndex, bTranslucentShadow);
+
+#if BUFFERED_OCCLUSION_QUERIES
+	// Get the oldest occlusion query	
+	const uint32 QueryIndex = FOcclusionQueryHelpers::GetQueryLookupIndex(PendingPrevFrameNumber, NumBufferedFrames);
+	FSceneViewState::ShadowOcclusionQueryMap ShadowOcclusionQueryMap = ShadowOcclusionQueryMaps[QueryIndex];	
+#endif
+
+
 	const FRenderQueryRHIRef* Query = ShadowOcclusionQueryMap.Find(Key);
 
 	// Read the occlusion query results.
@@ -267,8 +275,18 @@ void FSceneViewState::Destroy()
 
 SIZE_T FSceneViewState::GetSizeBytes() const
 {
+#if BUFFERED_OCCLUSION_QUERIES
+	uint32 ShadowOcclusionQuerySize = ShadowOcclusionQueryMaps.GetAllocatedSize();
+	for (int32 i = 0; i < ShadowOcclusionQueryMaps.Num(); ++i)
+	{
+		ShadowOcclusionQuerySize += ShadowOcclusionQueryMaps[i].GetAllocatedSize();
+	}
+#else
+	uint32 ShadowOcclusionQuerySize = ShadowOcclusionQueryMap.GetAllocatedSize();
+#endif
+
 	return sizeof(*this) 
-		+ ShadowOcclusionQueryMap.GetAllocatedSize() 
+		+ ShadowOcclusionQuerySize
 		+ ParentPrimitives.GetAllocatedSize() 
 		+ PrimitiveFadingStates.GetAllocatedSize()
 		+ PrimitiveOcclusionHistorySet.GetAllocatedSize();
@@ -385,9 +403,15 @@ FRenderQueryRHIParamRef FOcclusionQueryBatcher::BatchPrimitive(const FVector& Bo
 }
 
 static void IssueProjectedShadowOcclusionQuery(FRHICommandListImmediate& RHICmdList, FViewInfo& View, const FProjectedShadowInfo& ProjectedShadowInfo, FOcclusionQueryVS* VertexShader)
-{
+{	
 	FSceneViewState* ViewState = (FSceneViewState*)View.State;
-	TMap<FSceneViewState::FProjectedShadowKey, FRenderQueryRHIRef>& ShadowOcclusionQueryMap = ViewState->ShadowOcclusionQueryMap;
+
+#if BUFFERED_OCCLUSION_QUERIES		
+	const uint32 QueryIndex = FOcclusionQueryHelpers::GetQueryIssueIndex(ViewState->PendingPrevFrameNumber, ViewState->NumBufferedFrames);
+	FSceneViewState::ShadowOcclusionQueryMap& ShadowOcclusionQueryMap = ViewState->ShadowOcclusionQueryMaps[QueryIndex];
+#else
+	FSceneViewState::ShadowOcclusionQueryMap& ShadowOcclusionQueryMap = ViewState->ShadowOcclusionQueryMap;
+#endif
 
 	// The shadow transforms and view transforms are relative to different origins, so the world coordinates need to
 	// be translated.
@@ -982,7 +1006,12 @@ void FDeferredShadingSceneRenderer::BeginOcclusionTests(FRHICommandListImmediate
 			VertexShader->SetParameters(RHICmdList, View);
 
 			// Issue this frame's occlusion queries (occlusion queries from last frame may still be in flight)
-			TMap<FSceneViewState::FProjectedShadowKey, FRenderQueryRHIRef>& ShadowOcclusionQueryMap = ViewState->ShadowOcclusionQueryMap;
+#if BUFFERED_OCCLUSION_QUERIES				
+			const uint32 QueryIndex = FOcclusionQueryHelpers::GetQueryIssueIndex(ViewState->PendingPrevFrameNumber, ViewState->NumBufferedFrames);
+			FSceneViewState::ShadowOcclusionQueryMap& ShadowOcclusionQueryMap = ViewState->ShadowOcclusionQueryMaps[QueryIndex];
+#else
+			FSceneViewState::ShadowOcclusionQueryMap& ShadowOcclusionQueryMap = ViewState->ShadowOcclusionQueryMap;
+#endif			
 
 			// Clear primitives which haven't been visible recently out of the occlusion history, and reset old pending occlusion queries.
 			ViewState->TrimOcclusionHistory(RHICmdList, ViewFamily.CurrentRealTime - GEngine->PrimitiveProbablyVisibleTime, ViewFamily.CurrentRealTime, FrameNumber);
diff --git a/Engine/Source/Runtime/Renderer/Private/ScenePrivate.h b/Engine/Source/Runtime/Renderer/Private/ScenePrivate.h
--- a/Engine/Source/Runtime/Renderer/Private/ScenePrivate.h
+++ b/Engine/Source/Runtime/Renderer/Private/ScenePrivate.h
@@ -69,9 +69,48 @@ extern bool ShouldUseGetDynamicMeshElements();
 #include "SpeedTreeWind.h"
 #include "AtmosphereRendering.h"
 
+#if WITH_SLI || PLATFORM_PS4
+#define BUFFERED_OCCLUSION_QUERIES 1
+#endif
 /** Factor by which to grow occlusion tests **/
 #define OCCLUSION_SLOP (1.0f)
 
+class FOcclusionQueryHelpers
+{
+public:
+
+	// get the system-wide number of frames of buffered occlusion queries.
+	static int32 GetNumBufferedFrames()
+	{
+		int32 NumBufferedFrames = 1;
+
+#if BUFFERED_OCCLUSION_QUERIES		
+#	if WITH_SLI
+		// If we're running with SLI, assume throughput is more important than latency, and buffer an extra frame
+		NumBufferedFrames = GNumActiveGPUsForRendering == 1 ? 1 : GNumActiveGPUsForRendering + 1;
+#	else
+		static const auto NumBufferedQueriesVar = IConsoleManager::Get().FindTConsoleVariableDataInt(TEXT("r.NumBufferedOcclusionQueries"));
+		NumBufferedFrames = NumBufferedQueriesVar->GetValueOnAnyThread();
+#	endif		
+#endif
+		return NumBufferedFrames;
+	}
+
+	// get the index of the oldest query based on the current frame and number of buffered frames.
+	static uint32 GetQueryLookupIndex(int32 CurrentFrame, int32 NumBufferedFrames)
+	{
+		const uint32 QueryIndex = (CurrentFrame - NumBufferedFrames + 1) % NumBufferedFrames;
+		return QueryIndex;
+	}
+
+	// get the index of the query to overwrite for new queries.
+	static uint32 GetQueryIssueIndex(int32 CurrentFrame, int32 NumBufferedFrames)
+	{
+		const uint32 QueryIndex = CurrentFrame % NumBufferedFrames;
+		return QueryIndex;
+	}
+};
+
 /** Holds information about a single primitive's occlusion. */
 class FPrimitiveOcclusionHistory
 {
@@ -79,7 +118,7 @@ public:
 	/** The primitive the occlusion information is about. */
 	FPrimitiveComponentId PrimitiveId;
 
-#if WITH_SLI
+#if BUFFERED_OCCLUSION_QUERIES
 	/** The occlusion query which contains the primitive's pending occlusion results. */
 	TArray<FRenderQueryRHIRef, TInlineAllocator<1> > PendingOcclusionQuery;
 #else
@@ -121,9 +160,8 @@ public:
 		, LastPixelsPercentage(0.0f)
 		, bGroupedQuery(false)
 	{
-#if WITH_SLI
-		// If we're running with SLI, assume throughput is more important than latency, and buffer an extra frame
-		NumBufferedFrames = GNumActiveGPUsForRendering == 1 ? 1 : GNumActiveGPUsForRendering + 1;
+#if BUFFERED_OCCLUSION_QUERIES
+		NumBufferedFrames = FOcclusionQueryHelpers::GetNumBufferedFrames();
 		PendingOcclusionQuery.Empty(NumBufferedFrames);
 		PendingOcclusionQuery.AddZeroed(NumBufferedFrames);
 #endif
@@ -138,7 +176,7 @@ public:
 	template<class TOcclusionQueryPool> // here we use a template just to allow this to be inlined without sorting out the header order
 	FORCEINLINE void ReleaseQueries(FRHICommandListImmediate& RHICmdList, TOcclusionQueryPool& Pool)
 	{
-#if WITH_SLI
+#if BUFFERED_OCCLUSION_QUERIES
 		for (int32 QueryIndex = 0; QueryIndex < NumBufferedFrames; QueryIndex++)
 		{
 			Pool.ReleaseQuery(RHICmdList, PendingOcclusionQuery[QueryIndex]);
@@ -150,9 +188,9 @@ public:
 
 	FORCEINLINE FRenderQueryRHIRef& GetPastQuery(uint32 FrameNumber)
 	{
-#if WITH_SLI
+#if BUFFERED_OCCLUSION_QUERIES
 		// Get the oldest occlusion query
-		const uint32 QueryIndex = (FrameNumber - NumBufferedFrames + 1) % NumBufferedFrames;
+		const uint32 QueryIndex = FOcclusionQueryHelpers::GetQueryLookupIndex(FrameNumber, NumBufferedFrames);
 		return PendingOcclusionQuery[QueryIndex];
 #else
 		return PendingOcclusionQuery;
@@ -161,9 +199,9 @@ public:
 
 	FORCEINLINE void SetCurrentQuery(uint32 FrameNumber, FRenderQueryRHIParamRef NewQuery)
 	{
-#if WITH_SLI
+#if BUFFERED_OCCLUSION_QUERIES
 		// Get the current occlusion query
-		const uint32 QueryIndex = FrameNumber % NumBufferedFrames;
+		const uint32 QueryIndex = FOcclusionQueryHelpers::GetQueryIssueIndex(FrameNumber, NumBufferedFrames);
 		PendingOcclusionQuery[QueryIndex] = NewQuery;
 #else
 		PendingOcclusionQuery = NewQuery;
@@ -397,7 +435,13 @@ public:
 		bool bTranslucentShadow;
     };
 
-    TMap<FSceneViewState::FProjectedShadowKey, FRenderQueryRHIRef> ShadowOcclusionQueryMap;
+	int32 NumBufferedFrames;
+	typedef TMap<FSceneViewState::FProjectedShadowKey, FRenderQueryRHIRef> ShadowOcclusionQueryMap;
+#if BUFFERED_OCCLUSION_QUERIES
+	TArray<ShadowOcclusionQueryMap> ShadowOcclusionQueryMaps;
+#else
+	ShadowOcclusionQueryMap ShadowOcclusionQueryMap;
+#endif
 
 	/** The view's occlusion query pool. */
 	FRenderQueryPool OcclusionQueryPool;
@@ -615,7 +659,14 @@ public:
 
 	virtual void ReleaseDynamicRHI() override
     {
+#if BUFFERED_OCCLUSION_QUERIES
+		for (int i = 0; i < ShadowOcclusionQueryMaps.Num(); ++i)
+		{
+			ShadowOcclusionQueryMaps[i].Reset();
+		}
+#else
         ShadowOcclusionQueryMap.Reset();
+#endif
 		PrimitiveOcclusionHistorySet.Empty();
 		PrimitiveFadingStates.Empty();
 		OcclusionQueryPool.Release();
