Fix Xbox/Android/HTML5 compile

diff --git a/Engine/Source/Runtime/Renderer/Private/SceneOcclusion.cpp b/Engine/Source/Runtime/Renderer/Private/SceneOcclusion.cpp
--- a/Engine/Source/Runtime/Renderer/Private/SceneOcclusion.cpp
+++ b/Engine/Source/Runtime/Renderer/Private/SceneOcclusion.cpp
@@ -232,7 +232,7 @@ bool FSceneViewState::IsShadowOccluded(FRHICommandListImmediate& RHICmdList, FPr
 #if BUFFERED_OCCLUSION_QUERIES
 	// Get the oldest occlusion query	
 	const uint32 QueryIndex = FOcclusionQueryHelpers::GetQueryLookupIndex(PendingPrevFrameNumber, NumBufferedFrames);
-	FSceneViewState::ShadowOcclusionQueryMap ShadowOcclusionQueryMap = ShadowOcclusionQueryMaps[QueryIndex];	
+	FSceneViewState::ShadowKeyOcclusionQueryMap ShadowOcclusionQueryMap = ShadowOcclusionQueryMaps[QueryIndex];	
 #endif
 
 
@@ -408,9 +408,9 @@ static void IssueProjectedShadowOcclusionQuery(FRHICommandListImmediate& RHICmdL
 
 #if BUFFERED_OCCLUSION_QUERIES		
 	const uint32 QueryIndex = FOcclusionQueryHelpers::GetQueryIssueIndex(ViewState->PendingPrevFrameNumber, ViewState->NumBufferedFrames);
-	FSceneViewState::ShadowOcclusionQueryMap& ShadowOcclusionQueryMap = ViewState->ShadowOcclusionQueryMaps[QueryIndex];
+	FSceneViewState::ShadowKeyOcclusionQueryMap& ShadowOcclusionQueryMap = ViewState->ShadowOcclusionQueryMaps[QueryIndex];
 #else
-	FSceneViewState::ShadowOcclusionQueryMap& ShadowOcclusionQueryMap = ViewState->ShadowOcclusionQueryMap;
+	FSceneViewState::ShadowKeyOcclusionQueryMap& ShadowOcclusionQueryMap = ViewState->ShadowOcclusionQueryMap;
 #endif
 
 	// The shadow transforms and view transforms are relative to different origins, so the world coordinates need to
@@ -1008,9 +1008,9 @@ void FDeferredShadingSceneRenderer::BeginOcclusionTests(FRHICommandListImmediate
 			// Issue this frame's occlusion queries (occlusion queries from last frame may still be in flight)
 #if BUFFERED_OCCLUSION_QUERIES				
 			const uint32 QueryIndex = FOcclusionQueryHelpers::GetQueryIssueIndex(ViewState->PendingPrevFrameNumber, ViewState->NumBufferedFrames);
-			FSceneViewState::ShadowOcclusionQueryMap& ShadowOcclusionQueryMap = ViewState->ShadowOcclusionQueryMaps[QueryIndex];
+			FSceneViewState::ShadowKeyOcclusionQueryMap& ShadowOcclusionQueryMap = ViewState->ShadowOcclusionQueryMaps[QueryIndex];
 #else
-			FSceneViewState::ShadowOcclusionQueryMap& ShadowOcclusionQueryMap = ViewState->ShadowOcclusionQueryMap;
+			FSceneViewState::ShadowKeyOcclusionQueryMap& ShadowOcclusionQueryMap = ViewState->ShadowOcclusionQueryMap;
 #endif			
 
 			// Clear primitives which haven't been visible recently out of the occlusion history, and reset old pending occlusion queries.
diff --git a/Engine/Source/Runtime/Renderer/Private/ScenePrivate.h b/Engine/Source/Runtime/Renderer/Private/ScenePrivate.h
--- a/Engine/Source/Runtime/Renderer/Private/ScenePrivate.h
+++ b/Engine/Source/Runtime/Renderer/Private/ScenePrivate.h
@@ -72,6 +72,7 @@ extern bool ShouldUseGetDynamicMeshElements();
 #if WITH_SLI || PLATFORM_PS4
 #define BUFFERED_OCCLUSION_QUERIES 1
 #endif
+
 /** Factor by which to grow occlusion tests **/
 #define OCCLUSION_SLOP (1.0f)
 
@@ -436,11 +437,11 @@ public:
     };
 
 	int32 NumBufferedFrames;
-	typedef TMap<FSceneViewState::FProjectedShadowKey, FRenderQueryRHIRef> ShadowOcclusionQueryMap;
+	typedef TMap<FSceneViewState::FProjectedShadowKey, FRenderQueryRHIRef> ShadowKeyOcclusionQueryMap;
 #if BUFFERED_OCCLUSION_QUERIES
-	TArray<ShadowOcclusionQueryMap> ShadowOcclusionQueryMaps;
+	TArray<ShadowKeyOcclusionQueryMap> ShadowOcclusionQueryMaps;
 #else
-	ShadowOcclusionQueryMap ShadowOcclusionQueryMap;
+	ShadowKeyOcclusionQueryMap ShadowOcclusionQueryMap;
 #endif
 
 	/** The view's occlusion query pool. */
