* fixed minor banding issues with bright image based lighting

diff --git a/Engine/Source/Runtime/Renderer/Private/SystemTextures.cpp b/Engine/Source/Runtime/Renderer/Private/SystemTextures.cpp
--- a/Engine/Source/Runtime/Renderer/Private/SystemTextures.cpp
+++ b/Engine/Source/Runtime/Renderer/Private/SystemTextures.cpp
@@ -255,6 +255,9 @@ void FSystemTextures::InitializeTextures(FRHICommandListImmediate& RHICmdList, E
 
 			FPooledRenderTargetDesc Desc(FPooledRenderTargetDesc::Create2DDesc(FIntPoint(128, 32), PF_R8G8, TexCreate_HideInVisualizeTexture | TexCreate_FastVRAM, TexCreate_None, false));
 
+			// for low roughness we would get banding with PF_R8G8 but for low spec it could be used, for now we don't do this optimization
+			Desc.Format = PF_G16R16;
+
 			if (bReference)
 			{
 				Desc.Extent.X = 128;
@@ -327,7 +330,7 @@ void FSystemTextures::InitializeTextures(FRHICommandListImmediate& RHICmdList, E
 					A *= G_SchlickV;
 					B *= G_SchlickV;
 
-					if (bReference)
+					if (Desc.Format == PF_G16R16)
 					{
 						uint16* Dest = (uint16*)(DestBuffer + x * sizeof(uint32)+y * DestStride);
 						Dest[0] = (int32)(FMath::Clamp(A, 0.0f, 1.0f) * 65535.0f + 0.5f);
@@ -335,6 +338,8 @@ void FSystemTextures::InitializeTextures(FRHICommandListImmediate& RHICmdList, E
 					}
 					else
 					{
+						check(Desc.Format == PF_R8G8);
+
 						uint8* Dest = (uint8*)(DestBuffer + x * sizeof(uint16)+y * DestStride);
 						Dest[0] = (int32)(FMath::Clamp(A, 0.0f, 1.0f) * 255.9999f);
 						Dest[1] = (int32)(FMath::Clamp(B, 0.0f, 1.0f) * 255.9999f);
