* better comment

diff --git a/Engine/Source/Runtime/Renderer/Private/PostProcess/SceneRenderTargets.h b/Engine/Source/Runtime/Renderer/Private/PostProcess/SceneRenderTargets.h
--- a/Engine/Source/Runtime/Renderer/Private/PostProcess/SceneRenderTargets.h
+++ b/Engine/Source/Runtime/Renderer/Private/PostProcess/SceneRenderTargets.h
@@ -347,7 +347,9 @@ public:
 	/** Get the uniform buffer containing GBuffer resources. */
 	FUniformBufferRHIParamRef GetGBufferResourcesUniformBuffer() const 
 	{ 
+		// if this triggers you need to make sure the GBuffer is not getting released before (using AdjustGBufferRefCount(1) and AdjustGBufferRefCount(-1))
 		check(IsValidRef(GBufferResourcesUniformBuffer));
+
 		return GBufferResourcesUniformBuffer; 
 	}
 	/** */
