Fixed speedtree bug

diff --git a/Engine/Source/Runtime/Renderer/Private/RendererScene.cpp b/Engine/Source/Runtime/Renderer/Private/RendererScene.cpp
--- a/Engine/Source/Runtime/Renderer/Private/RendererScene.cpp
+++ b/Engine/Source/Runtime/Renderer/Private/RendererScene.cpp
@@ -1379,7 +1379,7 @@ void FScene::UpdateSpeedTreeWind(double CurrentTime)
 		{
 			FVector4 WindInfo = Scene->GetDirectionalWindParameters();
 
-			for (TMap<UStaticMesh*, FSpeedTreeWindComputation*>::TIterator It(Scene->SpeedTreeWindComputationMap); It; )
+			for (TMap<UStaticMesh*, FSpeedTreeWindComputation*>::TIterator It(Scene->SpeedTreeWindComputationMap); It; ++It )
 			{
 				UStaticMesh* StaticMesh = It.Key();
 				FSpeedTreeWindComputation* WindComputation = It.Value();
@@ -1434,7 +1434,6 @@ void FScene::UpdateSpeedTreeWind(double CurrentTime)
 				SET_SPEEDTREE_TABLE_FLOAT4V(WindRollingNoise, SH_ROLLING_NOISE_PERIOD);
 
 				WindComputation->UniformBuffer.SetContents(UniformParameters);
-				++It;
 			}
 		});
 	
