* added support for MaterialInstances for SubsurfaceProfile
* added VisualizeSSS to debug ScreenSpaceSubsurfaceScattering

diff --git a/Engine/Config/BaseEditor.ini b/Engine/Config/BaseEditor.ini
--- a/Engine/Config/BaseEditor.ini
+++ b/Engine/Config/BaseEditor.ini
@@ -34,6 +34,7 @@ MaxNumberOfLogLines=2000
 +RenderableThumbnailTypes=(ClassNeedingThumbnailName="/Script/Engine.World",RendererClassName="/Script/UnrealEd.WorldThumbnailRenderer")
 +RenderableThumbnailTypes=(ClassNeedingThumbnailName="/Script/Engine.SoundWave",RendererClassName="/Script/UnrealEd.SoundWaveThumbnailRenderer")
 +RenderableThumbnailTypes=(ClassNeedingThumbnailName="/Script/Engine.SlateBrushAsset",RendererClassName="/Script/UnrealEd.SlateBrushThumbnailRenderer")
++RenderableThumbnailTypes=(ClassNeedingThumbnailName="/Script/Engine.SubsurfaceProfile",RendererClassName="/Script/UnrealEd.SubsurfaceProfileRenderer")
 
 [/Script/UnrealEd.DefaultSizedThumbnailRenderer]
 DefaultSizeX=512
diff --git a/Engine/Shaders/MiniFontCommon.usf b/Engine/Shaders/MiniFontCommon.usf
new file mode 100644
--- /dev/null
+++ b/Engine/Shaders/MiniFontCommon.usf
@@ -0,0 +1,61 @@
+// Copyright 1998-2014 Epic Games, Inc. All Rights Reserved.
+
+/*=============================================================================
+	TinyFontCommon.usf: TinyFont shared functions and structures.
+=============================================================================*/
+
+#ifndef __TINY_FONT_COMMON__
+#define __TINY_FONT_COMMON__
+
+Texture2D MiniFontTexture;
+
+
+// for printf debugging in the shader
+// @param LeftTop - is advanced, in pixels
+void PrintCharacter(int2 PixelPos, inout float3 OutColor, float3 FontColor, inout int2 LeftTop, int CharacterID)
+{
+	uint2 Rel = (uint2)(PixelPos - LeftTop);
+
+	FLATTEN if (Rel.x < 8 && Rel.y < 8)
+	{
+		OutColor = lerp(OutColor, FontColor, MiniFontTexture.Load(int3(CharacterID * 8 + Rel.x, Rel.y, 0)).r);
+	}
+
+	LeftTop.x += 8;
+}
+
+// only for positive numbers
+// @param DigitValue - e.g. 1 for frist digit before period, 10 for second, 0.1 for first digit behind period
+uint ExtractDigitFromFloat(float Number, float DigitValue)
+{
+	uint Temp = (uint)(Number / DigitValue);
+
+	return Temp % 10;
+}
+
+
+// for printf debugging in the shader, has to be positive
+// outputs a float number in the form: xxx.yyy
+// @param LeftTop - in pixels
+void PrintFloat(int2 PixelPos, inout float3 OutColor, float3 FontColor, int2 LeftTop, float Number)
+{
+	int2 Cursor = LeftTop;
+
+	float BorderDistance = ComputeDistanceToRect(PixelPos, LeftTop, int2(7, 1) * 8 - 1);
+
+	// black border around number
+	//	OutColor = lerp(0, OutColor, saturate(BorderDistance - 2));
+
+	// before period
+	PrintCharacter(PixelPos, OutColor, FontColor, Cursor, ExtractDigitFromFloat(Number, 100));
+	PrintCharacter(PixelPos, OutColor, FontColor, Cursor, ExtractDigitFromFloat(Number, 10));
+	PrintCharacter(PixelPos, OutColor, FontColor, Cursor, ExtractDigitFromFloat(Number, 1));
+	// period
+	PrintCharacter(PixelPos, OutColor, FontColor, Cursor, 512 / 8 - 3);
+	// after period
+	PrintCharacter(PixelPos, OutColor, FontColor, Cursor, ExtractDigitFromFloat(Number, 0.1));
+	PrintCharacter(PixelPos, OutColor, FontColor, Cursor, ExtractDigitFromFloat(Number, 0.01));
+	PrintCharacter(PixelPos, OutColor, FontColor, Cursor, ExtractDigitFromFloat(Number, 0.001));
+}
+
+#endif // __TINY_FONT_COMMON__
diff --git a/Engine/Shaders/PostProcessSubsurface.usf b/Engine/Shaders/PostProcessSubsurface.usf
--- a/Engine/Shaders/PostProcessSubsurface.usf
+++ b/Engine/Shaders/PostProcessSubsurface.usf
@@ -8,6 +8,9 @@
 #include "PostProcessCommon.usf"
 #include "DeferredShadingCommon.usf"
 
+// for VisualizeSSS
+#include "MiniFontCommon.usf"
+
 /** RGBA8 linear texture containing random normals */
 Texture2D SSProfilesTexture;
 SamplerState SSProfilesTextureSampler;
@@ -37,6 +40,20 @@ float GetSubsurfaceStrength(float2 UV)
 // 0:speculars leak into SSS / 1:requires alpha channel for scenecolor
 // #define SSSS_SPECULAR_CORRECTION 1
 
+// define by C++ for the setup pass
+#ifndef SETUP_MODE
+	#define SETUP_MODE 0
+#endif
+
+#if SETUP_MODE == 0
+	#define SSSS_SPECULAR_CORRECTION 0
+#elif SETUP_MODE == 1
+	#define SSSS_SPECULAR_CORRECTION 1
+#else
+	// doesn't matter
+	#define SSSS_SPECULAR_CORRECTION 0
+#endif
+
 // 0: faster
 // 1: no color bleeding in z direction
 #define SSSS_FOLLOW_SURFACE 1
@@ -58,11 +75,11 @@ float GetSubsurfaceStrength(float2 UV)
 	error
 #endif
 
-#if METHOD == 0
+#if SSS_METHOD == 0
 	// horizontal
-#elif METHOD == 1
+#elif SSS_METHOD == 1
 	// vertical
-#elif METHOD == 2
+#elif SSS_METHOD == 2
 	// vertical and reconstruct specular
 #else
 	error
@@ -71,6 +88,10 @@ float GetSubsurfaceStrength(float2 UV)
 // 0 / 1
 #define VISUALIZE_KERNEL 0
 
+bool InUnitBox(float2 UV)
+{
+	return UV.x >= 0 && UV.y >= 0 && UV.y < 1 && UV.y < 1;
+}
 
 void SetupPS(in float4 UVAndScreenPos : TEXCOORD0, out float4 OutColor : SV_Target0)
 {
@@ -78,75 +99,100 @@ void SetupPS(in float4 UVAndScreenPos : TEXCOORD0, out float4 OutColor : SV_Targ
 
 	OutColor = 0;
 
-	FScreenSpaceData ScreenSpaceData = GetScreenSpaceData(UV);
-
-	FLATTEN if(ScreenSpaceData.GBuffer.ShadingModelID >= SHADINGMODELID_SUBSURFACE_PROFILE)
+	// if not VisualizeSSS
+	if (SETUP_MODE != 2)
 	{
-		float4 SceneColor4 = Texture2DSample(PostprocessInput0, PostprocessInput0Sampler, UV);
+		FScreenSpaceData ScreenSpaceData = GetScreenSpaceData(UV);
+
+		FLATTEN if (ScreenSpaceData.GBuffer.ShadingModelID >= SHADINGMODELID_SUBSURFACE_PROFILE)
+		{
+			float4 SceneColor4 = Texture2DSample(PostprocessInput0, PostprocessInput0Sampler, UV);
 
 #if SSSS_SPECULAR_CORRECTION == 1
-		// we take out the specular highlights
-		float3 Lighting = LightAccumulator_ReconstructDiffuseLighting(SceneColor4);
+			// we take out the specular highlights
+			float3 Lighting = LightAccumulator_ReconstructDiffuseLighting(SceneColor4);
 #else
-		float3 Lighting = SceneColor4.rgb;
+			float3 Lighting = SceneColor4.rgb;
 #endif
 
-		OutColor.rgb = Lighting;
+			OutColor.rgb = Lighting;
+		}
 	}
-}
+	else
+	{
+		// visualization (doesn't have to be fast)
+		OutColor = Texture2DSample(PostprocessInput0, PostprocessInput0Sampler, UV);
 
-// input0 is created by the SetupPS shader
-void MainPS(float4 UVAndScreenPos : TEXCOORD0, out float4 OutColor : SV_Target0)
-{
-	float2 UV = UVAndScreenPos.xy;
+		int2 PixelPos = (int2)(UVAndScreenPos.zw * ScreenPosToPixel.xy + ScreenPosToPixel.zw + 0.5f);
 
-	FScreenSpaceData ScreenSpaceData = GetScreenSpaceData(UV);
-	//
+		float2 ViewLocalUV = (PixelPos - View.ViewRectMin) / View.ViewSizeAndSceneTexelSize;
 
-	// test
-//	if (UV.y > 0.5f)
-	if (0)
-	{
-		uint SubsurfaceProfileInt = ExtractSubsurfaceProfileInt(ScreenSpaceData.GBuffer);
-		OutColor = float4(0.5f, 0.5f, 0.5f, 0);
-		if (SubsurfaceProfileInt == 0)
-		{
-			OutColor = float4(0.8f, 0.7f, 0.6f, 0);
-		}
-		if (SubsurfaceProfileInt == 1)
-		{
-			OutColor = float4(1, 0, 0, 0);
-		}
-		if (SubsurfaceProfileInt == 2)
-		{
-			OutColor = float4(0, 1, 0, 0);
-		}
-		if (SubsurfaceProfileInt == 3)
-		{ 
-			OutColor = float4(0, 0, 1, 0);
-		}
-		if (SubsurfaceProfileInt == 4)
-		{
-			OutColor = float4(1, 0, 1, 0);
-		}
-		if (SubsurfaceProfileInt == 5)
-		{
-			OutColor = float4(0, 1, 1, 0);
-		}
-		if (SubsurfaceProfileInt == 6)
-		{
-			OutColor = float4(1, 1, 0, 0);
-		}
-		if (SubsurfaceProfileInt == 100)
-		{
-			OutColor = float4(0, 0.2f, 0, 0);
-		}
-		if (SubsurfaceProfileInt == 255)
+
+		float2 IDAreaLocalUV = ViewLocalUV * 2 - 1.0f;
+
+//		OutColor = float4(IDAreaLocalUV.xy, 0, 0);		return;
+
+		if (InUnitBox(IDAreaLocalUV))
 		{
-			OutColor = float4(1, 1, 1, 0);
+			FScreenSpaceData ScreenSpaceData = GetScreenSpaceData(View.ViewRectMin.xy * View.ViewSizeAndSceneTexelSize.zw + IDAreaLocalUV * (View.ViewSizeAndSceneTexelSize.xy * View.ViewSizeAndSceneTexelSize.zw));
+
+			uint SubsurfaceProfileInt = ExtractSubsurfaceProfileInt(ScreenSpaceData.GBuffer);
+
+//			BRANCH if (ScreenSpaceData.GBuffer.ShadingModelID == SHADINGMODELID_SUBSURFACE_PROFILE)
+			{
+
+				OutColor = float4(0.5f, 0.5f, 0.5f, 0);
+
+				if (SubsurfaceProfileInt == 0)
+				{
+					// default (no Profile)
+					OutColor = float4(0.8f, 0.7f, 0.6f, 0);
+				}
+				if (SubsurfaceProfileInt == 1)
+				{
+					OutColor = float4(1, 0, 0, 0);
+				}
+				if (SubsurfaceProfileInt == 2)
+				{
+					OutColor = float4(0, 1, 0, 0);
+				}
+				if (SubsurfaceProfileInt == 3)
+				{
+					OutColor = float4(0, 0, 1, 0);
+				}
+				if (SubsurfaceProfileInt == 4)
+				{
+					OutColor = float4(1, 0, 1, 0);
+				}
+				if (SubsurfaceProfileInt == 5)
+				{
+					OutColor = float4(0, 1, 1, 0);
+				}
+				if (SubsurfaceProfileInt == 6)
+				{
+					OutColor = float4(1, 1, 0, 0);
+				}
+				if (SubsurfaceProfileInt == 100)
+				{
+					OutColor = float4(0, 0.2f, 0, 0);
+				}
+				if (SubsurfaceProfileInt == 255)
+				{
+					OutColor = float4(1, 1, 1, 0);
+				}
+
+				uint2 LeftTop = (PixelPos / 8) * 8;
+				PrintCharacter(PixelPos, OutColor.rgb, float3(1, 1, 1), LeftTop, SubsurfaceProfileInt);
+			}
 		}
-		return;
 	}
+}
+
+
+// input0 is created by the SetupPS shader
+void MainPS(float4 UVAndScreenPos : TEXCOORD0, out float4 OutColor : SV_Target0)
+{
+	float2 UV = UVAndScreenPos.xy;
 
 	// call into "SeparableSSS.usf"
 
@@ -157,7 +203,7 @@ void MainPS(float4 UVAndScreenPos : TEXCOORD0, out float4 OutColor : SV_Target0)
 	// *1000 as we store the sample distances / 1000 to allows for scaling
 	float Scaler = View.ViewSizeAndSceneTexelSize.x * View.ViewSizeAndSceneTexelSize.z * 1000.0f;
 
-#if METHOD == 0
+#if SSS_METHOD == 0
 	// first pass
 	float2 Direction = float2(1, 0) * Scaler;
 #else
@@ -170,16 +216,18 @@ void MainPS(float4 UVAndScreenPos : TEXCOORD0, out float4 OutColor : SV_Target0)
 
 	// Lighting
 
-#if METHOD > 0
+#if SSS_METHOD > 0
 	// on second pass we recombine with the scene color
 
 	float4 SceneColor4 = Texture2DSample(PostprocessInput0, PostprocessInput0Sampler, UV);
 
 	float4 SSSColor = OutColor;
 
+	FScreenSpaceData ScreenSpaceData = GetScreenSpaceData(UV);
+
 	SSSColor.a = (ScreenSpaceData.GBuffer.ShadingModelID >= SHADINGMODELID_SUBSURFACE_PROFILE);
 
-#if METHOD > 1
+#if SSS_METHOD > 1
 	// we took the specular highlights out, now we add them back in
 	SSSColor.rgb += LightAccumulator_ReconstructNonDiffuseLighting(SceneColor4);
 #endif
diff --git a/Engine/Shaders/PostProcessVisualizeHDR.usf b/Engine/Shaders/PostProcessVisualizeHDR.usf
--- a/Engine/Shaders/PostProcessVisualizeHDR.usf
+++ b/Engine/Shaders/PostProcessVisualizeHDR.usf
@@ -9,8 +9,7 @@
 #include "PostProcessHistogramCommon.usf"
 #include "DeferredShadingCommon.usf"
 #include "TonemapCommon.usf"
-
-Texture2D MiniFontTexture;
+#include "MiniFontCommon.usf"		// for PrintFloat()
 
 // only needed for nice visualization
 float ComputeHistogramMax(Texture2D HistogramTexture)
@@ -25,53 +24,6 @@ float ComputeHistogramMax(Texture2D HistogramTexture)
 	return Max;
 }
 
-// for printf debugging in the shader
-// @param LeftTop - is advanced, in pixels
-void PrintCharacter(int2 PixelPos, inout float3 OutColor, float3 FontColor, inout int2 LeftTop, int CharacterID)
-{
-	uint2 Rel = (uint2)(PixelPos - LeftTop);
-
-	FLATTEN if(Rel.x < 8 && Rel.y < 8)
-	{
-		OutColor = lerp(OutColor, FontColor, MiniFontTexture.Load(int3(CharacterID * 8 + Rel.x, Rel.y, 0)).r);
-	}
-
-	LeftTop.x += 8;
-}
-
-// only for positive numbers
-// @param DigitValue - e.g. 1 for frist digit before period, 10 for second, 0.1 for first digit behind period
-uint ExtractDigitFromFloat(float Number, float DigitValue)
-{
-	uint Temp = (uint)(Number / DigitValue);
-
-	return Temp % 10;
-}
-
-
-// for printf debugging in the shader, has to be positive
-// outputs a float number in the form: xxx.yyy
-// @param LeftTop - in pixels
-void PrintFloat(int2 PixelPos, inout float3 OutColor, float3 FontColor, int2 LeftTop, float Number)
-{
-	int2 Cursor = LeftTop;
-
-	float BorderDistance = ComputeDistanceToRect(PixelPos, LeftTop, int2(7, 1) * 8 - 1);
-	
-	// black border around number
-//	OutColor = lerp(0, OutColor, saturate(BorderDistance - 2));
-	
-	// before period
-	PrintCharacter(PixelPos, OutColor, FontColor, Cursor, ExtractDigitFromFloat(Number, 100));
-	PrintCharacter(PixelPos, OutColor, FontColor, Cursor, ExtractDigitFromFloat(Number, 10));
-	PrintCharacter(PixelPos, OutColor, FontColor, Cursor, ExtractDigitFromFloat(Number, 1));
-	// period
-	PrintCharacter(PixelPos, OutColor, FontColor, Cursor, 512 / 8 - 3);
-	// after period
-	PrintCharacter(PixelPos, OutColor, FontColor, Cursor, ExtractDigitFromFloat(Number, 0.1));
-	PrintCharacter(PixelPos, OutColor, FontColor, Cursor, ExtractDigitFromFloat(Number, 0.01));
-	PrintCharacter(PixelPos, OutColor, FontColor, Cursor, ExtractDigitFromFloat(Number, 0.001));
-}
 
 uint ComputeAdvice(float3 HDRColor)
 {
diff --git a/Engine/Shaders/ShaderVersion.usf b/Engine/Shaders/ShaderVersion.usf
--- a/Engine/Shaders/ShaderVersion.usf
+++ b/Engine/Shaders/ShaderVersion.usf
@@ -2,4 +2,4 @@
 // Each time the console command is executed it generates a new GUID. As this file is included
 // in common.usf (which should be included in any shader) it allows to invalidate the shader DDC.
 // 
-// GUID = {BAC53009-383D-4F58-A20C-AF3BC0ABAFB5}.
+// GUID = B08E107742F276154CF002A94C779319
\ No newline at end of file
diff --git a/Engine/Source/Editor/MaterialEditor/Private/MaterialEditorInstanceDetailCustomization.cpp b/Engine/Source/Editor/MaterialEditor/Private/MaterialEditorInstanceDetailCustomization.cpp
--- a/Engine/Source/Editor/MaterialEditor/Private/MaterialEditorInstanceDetailCustomization.cpp
+++ b/Engine/Source/Editor/MaterialEditor/Private/MaterialEditorInstanceDetailCustomization.cpp
@@ -55,8 +55,20 @@ void FMaterialInstanceParameterDetails::CustomizeDetails( IDetailLayoutBuilder&
 	DetailLayout.HideProperty("ParameterGroups");
 	DetailLayout.HideProperty("BasePropertyOverrides");
 
-	IDetailPropertyRow& RefractionDepthBias = DefaultCategory.AddProperty("RefractionDepthBias");
-	RefractionDepthBias.Visibility(TAttribute<EVisibility>::Create(TAttribute<EVisibility>::FGetter::CreateSP(this, &FMaterialInstanceParameterDetails::ShouldShowMaterialRefractionSettings)));
+	{
+		IDetailPropertyRow& PropertyRow = DefaultCategory.AddProperty("RefractionDepthBias");
+		PropertyRow.Visibility(TAttribute<EVisibility>::Create(TAttribute<EVisibility>::FGetter::CreateSP(this, &FMaterialInstanceParameterDetails::ShouldShowMaterialRefractionSettings)));
+	}
+
+	{
+		IDetailPropertyRow& PropertyRow = DefaultCategory.AddProperty("bOverrideSubsurfaceProfile");
+		PropertyRow.Visibility(TAttribute<EVisibility>::Create(TAttribute<EVisibility>::FGetter::CreateSP(this, &FMaterialInstanceParameterDetails::ShouldShowSubsurfaceProfile)));
+	}
+
+	{
+		IDetailPropertyRow& PropertyRow = DefaultCategory.AddProperty("SubsurfaceProfile");
+		PropertyRow.Visibility(TAttribute<EVisibility>::Create(TAttribute<EVisibility>::FGetter::CreateSP(this, &FMaterialInstanceParameterDetails::ShouldShowSubsurfaceProfile)));
+	}
 
 	// Create a new category for a custom layout for the MIC parameters
 	FName GroupsCategoryName = TEXT("ParameterGroups");
@@ -373,3 +385,10 @@ EVisibility FMaterialInstanceParameterDetails::ShouldShowMaterialRefractionSetti
 {
 	return (MaterialEditorInstance->SourceInstance->GetMaterial()->bUsesDistortion) ? EVisibility::Visible : EVisibility::Collapsed;
 }
+
+EVisibility FMaterialInstanceParameterDetails::ShouldShowSubsurfaceProfile() const
+{
+	EMaterialShadingModel Model = MaterialEditorInstance->SourceInstance->GetMaterial()->GetShadingModel_Internal();
+
+	return (Model == MSM_SubsurfaceProfile) ? EVisibility::Visible : EVisibility::Collapsed;
+}
diff --git a/Engine/Source/Editor/MaterialEditor/Private/MaterialEditorInstanceDetailCustomization.h b/Engine/Source/Editor/MaterialEditor/Private/MaterialEditorInstanceDetailCustomization.h
--- a/Engine/Source/Editor/MaterialEditor/Private/MaterialEditorInstanceDetailCustomization.h
+++ b/Engine/Source/Editor/MaterialEditor/Private/MaterialEditorInstanceDetailCustomization.h
@@ -58,6 +58,9 @@ private:
 	/** Returns true if the refraction options should be displayed */
 	EVisibility ShouldShowMaterialRefractionSettings() const;
 
+	/** Returns true if the refraction options should be displayed */
+	EVisibility ShouldShowSubsurfaceProfile() const;
+
 private:
 	/** Object that stores all of the possible parameters we can edit */
 	UMaterialEditorInstanceConstant* MaterialEditorInstance;
diff --git a/Engine/Source/Editor/UnrealEd/Classes/MaterialEditor/MaterialEditorInstanceConstant.h b/Engine/Source/Editor/UnrealEd/Classes/MaterialEditor/MaterialEditorInstanceConstant.h
--- a/Engine/Source/Editor/UnrealEd/Classes/MaterialEditor/MaterialEditorInstanceConstant.h
+++ b/Engine/Source/Editor/UnrealEd/Classes/MaterialEditor/MaterialEditorInstanceConstant.h
@@ -219,7 +219,15 @@ class UNREALED_API UMaterialEditorInstanceConstant : public UObject
 	UPROPERTY(EditAnywhere, Category=MaterialEditorInstanceConstant)
 	float RefractionDepthBias;
 
-	UPROPERTY(EditAnywhere, Category=MaterialOverrides)
+	/** SubsurfaceProfile, for Screen Space Subsurface Scattering */
+	UPROPERTY(EditAnywhere, BlueprintReadOnly, Category = Material, meta = (DisplayName = "Subsurface Profile"))
+	class USubsurfaceProfile* SubsurfaceProfile;
+
+	/** Defines if SubsurfaceProfile from tis instance is used or it uses the parent one. */
+	UPROPERTY(EditAnywhere, Category = MaterialEditorInstanceConstant)
+	uint32 bOverrideSubsurfaceProfile : 1;
+
+	UPROPERTY(EditAnywhere, Category = MaterialOverrides)
 	uint32 bOverrideBaseProperties:1;
 
 	UPROPERTY(EditAnywhere, Category=MaterialOverrides, meta=(editcondition="bOverrideBaseProperties"))
diff --git a/Engine/Source/Editor/UnrealEd/Classes/ThumbnailRendering/SubsurfaceProfileRenderer.h b/Engine/Source/Editor/UnrealEd/Classes/ThumbnailRendering/SubsurfaceProfileRenderer.h
new file mode 100644
--- /dev/null
+++ b/Engine/Source/Editor/UnrealEd/Classes/ThumbnailRendering/SubsurfaceProfileRenderer.h
@@ -0,0 +1,22 @@
+// Copyright 1998-2014 Epic Games, Inc. All Rights Reserved.
+
+/**
+ *
+ * This thumbnail renderer displays the texture for the object in question
+ */
+
+#pragma once
+#include "SubsurfaceProfileRenderer.generated.h"
+
+UCLASS()
+class USubsurfaceProfileRenderer : public UTextureThumbnailRenderer
+{
+	GENERATED_UCLASS_BODY()
+
+
+	// Begin UThumbnailRenderer Object
+	virtual void GetThumbnailSize(UObject* Object, float Zoom, uint32& OutWidth, uint32& OutHeight) const override;
+	virtual void Draw(UObject* Object, int32 X, int32 Y, uint32 Width, uint32 Height, FRenderTarget*, FCanvas* Canvas) override;
+	// End UThumbnailRenderer Object
+};
+
diff --git a/Engine/Source/Editor/UnrealEd/Private/PreviewMaterial.cpp b/Engine/Source/Editor/UnrealEd/Private/PreviewMaterial.cpp
--- a/Engine/Source/Editor/UnrealEd/Private/PreviewMaterial.cpp
+++ b/Engine/Source/Editor/UnrealEd/Private/PreviewMaterial.cpp
@@ -598,6 +598,9 @@ void UMaterialEditorInstanceConstant::CopyToSourceInstance()
 		// Copy Refraction bias setting
 		SourceInstance->SetScalarParameterValueEditorOnly(TEXT("RefractionDepthBias"), RefractionDepthBias);
 
+		SourceInstance->bOverrideSubsurfaceProfile = bOverrideSubsurfaceProfile;
+		SourceInstance->SubsurfaceProfile = SubsurfaceProfile;
+
 		// Update object references and parameter names.
 		SourceInstance->UpdateParameterNames();
 		VisibleExpressions.Empty();
@@ -688,6 +691,9 @@ void UMaterialEditorInstanceConstant::SetSourceInstance(UMaterialInstanceConstan
 	//Copy refraction settings
 	SourceInstance->GetRefractionSettings(RefractionDepthBias);
 
+	bOverrideSubsurfaceProfile = SourceInstance->bOverrideSubsurfaceProfile;
+	SubsurfaceProfile = SourceInstance->SubsurfaceProfile;
+
 	RegenerateArrays();
 
 	//propagate changes to the base material so the instance will be updated if it has a static permutation resource
diff --git a/Engine/Source/Editor/UnrealEd/Private/ThumbnailRendering/SubsurfaceProfileRenderer.cpp b/Engine/Source/Editor/UnrealEd/Private/ThumbnailRendering/SubsurfaceProfileRenderer.cpp
new file mode 100644
--- /dev/null
+++ b/Engine/Source/Editor/UnrealEd/Private/ThumbnailRendering/SubsurfaceProfileRenderer.cpp
@@ -0,0 +1,34 @@
+// Copyright 1998-2014 Epic Games, Inc. All Rights Reserved.
+
+#include "UnrealEd.h"
+
+USubsurfaceProfileRenderer::USubsurfaceProfileRenderer(const class FPostConstructInitializeProperties& PCIP)
+	: Super(PCIP)
+{
+}
+
+void USubsurfaceProfileRenderer::GetThumbnailSize(UObject* Object, float Zoom, uint32& OutWidth, uint32& OutHeight) const
+{
+	OutWidth = 128;
+	OutHeight = 128;
+}
+
+void USubsurfaceProfileRenderer::Draw(UObject* Object, int32 X, int32 Y, uint32 Width, uint32 Height, FRenderTarget*, FCanvas* Canvas)
+{
+	USubsurfaceProfile* LocalSubsurfaceProfile = Cast<USubsurfaceProfile>(Object);
+	if (LocalSubsurfaceProfile)
+	{
+		FLinearColor Col;
+
+		Col = LocalSubsurfaceProfile->Settings.SubsurfaceColor; Col.A = 1;
+		Canvas->DrawTile(0,          0, Width, Height / 2, 0, 0, 1, 1, Col);
+		Col = LocalSubsurfaceProfile->Settings.FalloffColor; Col.A = 1;
+		Canvas->DrawTile(0, Height / 2, Width, Height / 2, 0, 0, 1, 1, Col);
+
+		FText ScatterRadiusText = FText::AsNumber(LocalSubsurfaceProfile->Settings.ScatterRadius);
+		FCanvasTextItem TextItem(FVector2D(5.0f, 5.0f), ScatterRadiusText, GEngine->GetLargeFont(), FLinearColor::White);
+		TextItem.EnableShadow(FLinearColor::Black);
+		TextItem.Scale = FVector2D(Width / 128.0f, Height / 128.0f);
+		TextItem.Draw(Canvas);
+	}
+}
\ No newline at end of file
diff --git a/Engine/Source/Runtime/Engine/Classes/Engine/SubsurfaceProfile.h b/Engine/Source/Runtime/Engine/Classes/Engine/SubsurfaceProfile.h
--- a/Engine/Source/Runtime/Engine/Classes/Engine/SubsurfaceProfile.h
+++ b/Engine/Source/Runtime/Engine/Classes/Engine/SubsurfaceProfile.h
@@ -39,6 +39,13 @@ struct FSubsurfaceProfileStruct
 		FalloffColor = FLinearColor(1.0f, 0.37f, 0.3f);
 	}
 
+	void Invalidate()
+	{
+		// nicer for debugging / VisualizeSSS
+		ScatterRadius = 0.0f;
+		SubsurfaceColor = FLinearColor(0, 0, 0);
+		FalloffColor = FLinearColor(0, 0, 0);
+	}
 };
 
 /**
@@ -59,6 +66,8 @@ class USubsurfaceProfile : public UObject
 	// End UObject interface
 };
 
+// For render thread side, to identify the right index in a map, is never dereferenced
+typedef void* USubsurfaceProfilePointer;
 
 
 // render thread
@@ -74,7 +83,7 @@ public:
 	void SetRendererModule(class IRendererModule* InRendererModule) { RendererModule = InRendererModule; }
 
 	// convenience, can be optimized 
-	int32 AddOrUpdateProfile(const FSubsurfaceProfileStruct Settings, const USubsurfaceProfile& Profile)
+	int32 AddOrUpdateProfile(const FSubsurfaceProfileStruct Settings, const USubsurfaceProfilePointer Profile)
 	{
 		int32 AllocationId = FindAllocationId(Profile); 
 		
@@ -93,13 +102,13 @@ public:
 	// O(n) n is a small number
 	// @param Profile must not be 0
 	// @return AllocationId -1: no allocation, should be deallocated with DeallocateSubsurfaceProfile()
-	int32 AddProfile(const FSubsurfaceProfileStruct Settings, const USubsurfaceProfile& Profile);
+	int32 AddProfile(const FSubsurfaceProfileStruct Settings, const USubsurfaceProfilePointer Profile);
 
 	// O(n) to find the element, n is the SSProfile count and usually quite small
-	void RemoveProfile(USubsurfaceProfile& InProfile);
+	void RemoveProfile(const USubsurfaceProfilePointer InProfile);
 
 	// @param InRendererModule first call needs to have it != 0
-	void UpdateProfile(const FSubsurfaceProfileStruct Settings, const USubsurfaceProfile& Profile) { int32 AllocationId = FindAllocationId(Profile); UpdateProfile(AllocationId, Settings); }
+	void UpdateProfile(const FSubsurfaceProfileStruct Settings, const USubsurfaceProfilePointer Profile) { int32 AllocationId = FindAllocationId(Profile); UpdateProfile(AllocationId, Settings); }
 
 	// @param InRendererModule first call needs to have it != 0
 	void UpdateProfile(int32 AllocationId, const FSubsurfaceProfileStruct Settings);
@@ -116,14 +125,17 @@ public:
 	// for debugging, can be removed
 	void Dump();
 
+	// for debugging / VisualizeSSS
+	ENGINE_API bool GetEntryString(uint32 Index, FString& Out) const;
+
 	// @return -1 if not found
-	int32 FindAllocationId(const USubsurfaceProfile& InProfile) const;
+	int32 FindAllocationId(const USubsurfaceProfilePointer InProfile) const;
 
 private:
 
 	struct FSubsurfaceProfileEntry
 	{
-		FSubsurfaceProfileEntry(FSubsurfaceProfileStruct InSubsurfaceProfileStruct, const USubsurfaceProfile* InGameThreadObject)
+		FSubsurfaceProfileEntry(FSubsurfaceProfileStruct InSubsurfaceProfileStruct, const USubsurfaceProfilePointer InGameThreadObject)
 			: Settings(InSubsurfaceProfileStruct)
 			, GameThreadObject(InGameThreadObject)
 		{
@@ -132,7 +144,7 @@ private:
 		FSubsurfaceProfileStruct Settings;
 
 		// not referenced, void and not USubsurfaceProfile* to prevent accidential access, 0 if the entry can be reused or it's [0] which is used as default
-		const void* GameThreadObject;
+		USubsurfaceProfilePointer GameThreadObject;
 	};
 
 	// 0 if there if no user of the feature
@@ -146,4 +158,4 @@ private:
 };
 
 // lives on the render thread
-extern TGlobalResource<FSubsurfaceProfileTexture> GSubsufaceProfileTextureObject;
+extern ENGINE_API TGlobalResource<FSubsurfaceProfileTexture> GSubsufaceProfileTextureObject;
diff --git a/Engine/Source/Runtime/Engine/Classes/Materials/Material.h b/Engine/Source/Runtime/Engine/Classes/Materials/Material.h
--- a/Engine/Source/Runtime/Engine/Classes/Materials/Material.h
+++ b/Engine/Source/Runtime/Engine/Classes/Materials/Material.h
@@ -778,6 +778,7 @@ public:
 	ENGINE_API virtual EMaterialShadingModel GetShadingModel_Internal() const;
 	ENGINE_API virtual bool IsTwoSided_Internal() const;
 	ENGINE_API virtual bool IsMasked_Internal() const;
+	ENGINE_API virtual USubsurfaceProfile* GetSubsurfaceProfile_Internal() const;
 
 	ENGINE_API void SetShadingModel(EMaterialShadingModel NewModel) {ShadingModel = NewModel;}
 
diff --git a/Engine/Source/Runtime/Engine/Classes/Materials/MaterialInstance.h b/Engine/Source/Runtime/Engine/Classes/Materials/MaterialInstance.h
--- a/Engine/Source/Runtime/Engine/Classes/Materials/MaterialInstance.h
+++ b/Engine/Source/Runtime/Engine/Classes/Materials/MaterialInstance.h
@@ -146,6 +146,10 @@ class UMaterialInstance : public UMaterialInterface
 	/** Flag to detect cycles in the material instance graph. */
 	uint32 ReentrantFlag:1;
 
+	/** Defines if SubsurfaceProfile from tis instance is used or it uses the parent one. */
+	UPROPERTY(EditAnywhere, BlueprintReadOnly, Category = MaterialInstance)
+	uint32 bOverrideSubsurfaceProfile:1;
+
 	/** Unique ID for this material, used for caching during distributed lighting */
 	UPROPERTY()
 	FGuid ParentLightingGuid;
@@ -245,6 +249,7 @@ public:
 	ENGINE_API virtual EMaterialShadingModel GetShadingModel_Internal() const;
 	ENGINE_API virtual bool IsTwoSided_Internal() const;
 	ENGINE_API virtual bool IsMasked_Internal() const;
+	ENGINE_API virtual USubsurfaceProfile* GetSubsurfaceProfile_Internal() const;
 
 	/** Returns true and sets Result if the property was overridden in the instance. Otherwise, returns false and the base material property should be used. */
 	ENGINE_API virtual bool GetOpacityMaskClipValueOverride(float& OutResult) const;
@@ -376,7 +381,7 @@ protected:
 	ENGINE_API void CopyMaterialInstanceParameters(UMaterialInterface* MaterialInterface);
 
 	// to share code between PostLoad() and PostEditChangeProperty()
-	void UpdateMaterialInstanceData();
+	void PropagateDataToMaterialProxy();
 
 	/** Allow resource to access private members. */
 	friend class FMaterialInstanceResource;
diff --git a/Engine/Source/Runtime/Engine/Classes/Materials/MaterialInterface.h b/Engine/Source/Runtime/Engine/Classes/Materials/MaterialInterface.h
--- a/Engine/Source/Runtime/Engine/Classes/Materials/MaterialInterface.h
+++ b/Engine/Source/Runtime/Engine/Classes/Materials/MaterialInterface.h
@@ -528,6 +528,8 @@ public:
 	ENGINE_API virtual EMaterialShadingModel GetShadingModel_Internal() const;
 	ENGINE_API virtual bool IsTwoSided_Internal() const;
 	ENGINE_API virtual bool IsMasked_Internal() const;
+	ENGINE_API virtual USubsurfaceProfile* GetSubsurfaceProfile_Internal() const { return 0; }
+
 	/**
 	 * Force the streaming system to disregard the normal logic for the specified duration and
 	 * instead always load all mip-levels for all textures used by this material.
diff --git a/Engine/Source/Runtime/Engine/Private/Materials/HLSLMaterialTranslator.h b/Engine/Source/Runtime/Engine/Private/Materials/HLSLMaterialTranslator.h
--- a/Engine/Source/Runtime/Engine/Private/Materials/HLSLMaterialTranslator.h
+++ b/Engine/Source/Runtime/Engine/Private/Materials/HLSLMaterialTranslator.h
@@ -343,12 +343,12 @@ public:
 			{
 				int32 SubsurfaceColor = Material->CompilePropertyAndSetMaterialProperty(MP_SubsurfaceColor, this);
 
-				static FName SubsurfaceProfileName(TEXT("SubsurfaceProfile"));
+				static FName NameSubsurfaceProfile(TEXT("__SubsurfaceProfile"));
 
 				// 1.0f is is a not used profile - later this gets replaced with the actual profile
-				int32 SubsurfaceProfile = ForceCast(ScalarParameter(SubsurfaceProfileName, 1.0f), MCT_Float1);
+				int32 CodeSubsurfaceProfile = ForceCast(ScalarParameter(NameSubsurfaceProfile, 1.0f), MCT_Float1);
 
-				Chunk[MP_SubsurfaceColor] = AppendVector(SubsurfaceColor, SubsurfaceProfile);		
+				Chunk[MP_SubsurfaceColor] = AppendVector(SubsurfaceColor, CodeSubsurfaceProfile);		
 			}
 
 			Chunk[MP_ClearCoat]						= Material->CompilePropertyAndSetMaterialProperty(MP_ClearCoat			   ,this);
diff --git a/Engine/Source/Runtime/Engine/Private/Materials/Material.cpp b/Engine/Source/Runtime/Engine/Private/Materials/Material.cpp
--- a/Engine/Source/Runtime/Engine/Private/Materials/Material.cpp
+++ b/Engine/Source/Runtime/Engine/Private/Materials/Material.cpp
@@ -210,21 +210,24 @@ public:
 		const FMaterialResource* MaterialResource = Material->GetMaterialResource(Context.Material.GetFeatureLevel());
 		if(MaterialResource && MaterialResource->GetRenderingThreadShaderMap())
 		{
-			static FName SubsurfaceProfile(TEXT("SubsurfaceProfile"));
-			if (ParameterName == SubsurfaceProfile)
+			static FName NameSubsurfaceProfile(TEXT("__SubsurfaceProfile"));
+			if (ParameterName == NameSubsurfaceProfile)
 			{
-				const USubsurfaceProfile* SubsurfaceProfileRT = GetSubsurfaceProfile();
+				const USubsurfaceProfilePointer SubsurfaceProfileRT = GetSubsurfaceProfileRT();
 
 				if(SubsurfaceProfileRT)
 				{
 					// can be optimized (cached)
-					*OutValue = GSubsufaceProfileTextureObject.FindAllocationId(*SubsurfaceProfileRT) / 255.0f;
+					*OutValue = GSubsufaceProfileTextureObject.FindAllocationId(SubsurfaceProfileRT) / 255.0f;
 				}
 				else
 				{
 					// no profile specified means we use the default one stored at [0] which is human skin
 					*OutValue = 0.0f;
 				}
+
+				FPlatformMisc::LowLevelOutputDebugStringf(TEXT(">>>> Material %p GetScalarValue %f\n"), this, *OutValue);
+
 				return true;
 			}
 
@@ -3480,6 +3483,12 @@ bool UMaterial::IsMasked_Internal() const
 	return bIsMasked != 0;
 }
 
+USubsurfaceProfile* UMaterial::GetSubsurfaceProfile_Internal() const
+{
+	checkSlow(IsInGameThread());
+	return SubsurfaceProfile; 
+}
+
 bool UMaterial::IsPropertyActive(EMaterialProperty InProperty)const 
 {
 	if(MaterialDomain == MD_PostProcess)
diff --git a/Engine/Source/Runtime/Engine/Private/Materials/MaterialInstance.cpp b/Engine/Source/Runtime/Engine/Private/Materials/MaterialInstance.cpp
--- a/Engine/Source/Runtime/Engine/Private/Materials/MaterialInstance.cpp
+++ b/Engine/Source/Runtime/Engine/Private/Materials/MaterialInstance.cpp
@@ -206,6 +206,28 @@ bool FMaterialInstanceResource::GetScalarValue(
 	) const
 {
 	checkSlow(IsInParallelRenderingThread());
+
+	static FName NameSubsurfaceProfile(TEXT("__SubsurfaceProfile"));
+	if (ParameterName == NameSubsurfaceProfile)
+	{
+		const USubsurfaceProfilePointer SubsurfaceProfileRT = GetSubsurfaceProfileRT();
+
+		if (SubsurfaceProfileRT)
+		{
+			// can be optimized (cached)
+			*OutValue = GSubsufaceProfileTextureObject.FindAllocationId(SubsurfaceProfileRT) / 255.0f;
+		}
+		else
+		{
+			// no profile specified means we use the default one stored at [0] which is human skin
+			*OutValue = 0.0f;
+		}
+
+		FPlatformMisc::LowLevelOutputDebugStringf(TEXT(">>>> MaterialInstance %p GetScalarValue %f\n"), this, *OutValue);
+
+		return true;
+	}
+
 	const float* Value = RenderThread_FindParameterByName<float>(ParameterName);
 	if(Value)
 	{
@@ -280,7 +302,7 @@ void FMaterialInstanceResource::GameThread_UpdateDistanceFieldPenumbraScale(floa
 	});
 }
 
-void UMaterialInstance::UpdateMaterialInstanceData()
+void UMaterialInstance::PropagateDataToMaterialProxy()
 {
 	for (int32 i = 0; i < ARRAY_COUNT(Resources); i++)
 	{
@@ -1665,6 +1687,7 @@ void UMaterialInstance::PostLoad()
 			Texture->ConditionalPostLoad();
 		}
 	}
+
 	// do the same for font textures
 	for( int32 ValueIndex=0; ValueIndex < FontParameterValues.Num(); ValueIndex++ )
 	{
@@ -1676,6 +1699,9 @@ void UMaterialInstance::PostLoad()
 		}
 	}
 
+	// called before we cache the uniform expression as a call to SubsurfaceProfileRT affects the dta in there
+	PropagateDataToMaterialProxy();
+
 	// Update bHasStaticPermutationResource in case the parent was not found
 	bHasStaticPermutationResource = (!StaticParameters.IsEmpty() || bOverrideBaseProperties) && Parent;
 
@@ -1699,6 +1725,7 @@ void UMaterialInstance::PostLoad()
 			}
 		}
 	}
+
 	INC_FLOAT_STAT_BY(STAT_ShaderCompiling_MaterialLoading,(float)MaterialLoadTime);
 
 	if (GIsEditor && GEngine != NULL && !IsTemplate() && Parent)
@@ -1707,8 +1734,6 @@ void UMaterialInstance::PostLoad()
 		UpdateLightmassTextureTracking();
 	}
 
-	UpdateMaterialInstanceData();
-
 	// Fixup for legacy instances which didn't recreate the lighting guid properly on duplication
 	if (GetLinker() && GetLinker()->UE4Ver() < VER_UE4_BUMPED_MATERIAL_EXPORT_GUIDS)
 	{
@@ -2014,7 +2039,7 @@ void UMaterialInstance::PostEditChangeProperty(FPropertyChangedEvent& PropertyCh
 		UpdateLightmassTextureTracking();
 	}
 
-	UpdateMaterialInstanceData();
+	PropagateDataToMaterialProxy();
 
 	InitResources();
 
@@ -2393,6 +2418,18 @@ bool UMaterialInstance::IsMasked_Internal() const
 	return GetMaterial()->IsMasked();
 }
 
+USubsurfaceProfile* UMaterialInstance::GetSubsurfaceProfile_Internal() const
+{
+	checkSlow(IsInGameThread());
+	if (bOverrideSubsurfaceProfile)
+	{
+		return SubsurfaceProfile;
+	}
+
+	// go up the chain if possible
+	return Parent ? Parent->GetSubsurfaceProfile_Internal() : 0;
+}
+
 bool UMaterialInstance::GetOpacityMaskClipValueOverride(float& OutResult) const
 {
 	if( bOverrideBaseProperties && BasePropertyOverrides.bOverride_OpacityMaskClipValue )
diff --git a/Engine/Source/Runtime/Engine/Private/Materials/MaterialInterface.cpp b/Engine/Source/Runtime/Engine/Private/Materials/MaterialInterface.cpp
--- a/Engine/Source/Runtime/Engine/Private/Materials/MaterialInterface.cpp
+++ b/Engine/Source/Runtime/Engine/Private/Materials/MaterialInterface.cpp
@@ -335,53 +335,57 @@ void UMaterialInterface::UpdateMaterialRenderProxy(FMaterialRenderProxy& Proxy)
 	// no 0 pointer
 	check(&Proxy);
 
-	FSubsurfaceProfileStruct Settings;
-
-	if (SubsurfaceProfile)
-	{
-		Settings = SubsurfaceProfile->Settings;
-	}
+	EMaterialShadingModel MaterialShadingModel = GetShadingModel_Internal();
 
-	// this can be improved, it doesn't support Renderer hot reload
+	// for better performance we only update SubsurfaceProfileRT if the feature is used
+	if (MaterialShadingModel == MSM_SubsurfaceProfile)
 	{
-		static bool bFirst = true;
+		FSubsurfaceProfileStruct Settings;
 
-		if(bFirst)
+		USubsurfaceProfile* LocalSubsurfaceProfile = GetSubsurfaceProfile_Internal();
+		
+		if (LocalSubsurfaceProfile)
 		{
-			bFirst = false;
+			Settings = LocalSubsurfaceProfile->Settings;
+		}
 
-			static const FName RendererModuleName("Renderer");
-			IRendererModule& RendererModule = FModuleManager::GetModuleChecked<IRendererModule>(RendererModuleName);
+		// this can be improved, it doesn't support Renderer hot reload
+		{
+			static bool bFirst = true;
 
-			ENQUEUE_UNIQUE_RENDER_COMMAND_ONEPARAMETER(
-				UpdateMaterialRenderProxySubsurfaceSetup,
-				IRendererModule&, RendererModule, RendererModule,
+			if (bFirst)
 			{
-				GSubsufaceProfileTextureObject.SetRendererModule(&RendererModule);
-			});
-		}
-	}
+				bFirst = false;
 
-	EMaterialShadingModel MaterialShadingModel = GetShadingModel_Internal();
+				static const FName RendererModuleName("Renderer");
+				IRendererModule& RendererModule = FModuleManager::GetModuleChecked<IRendererModule>(RendererModuleName);
 
-	// for better performance we only update SubsurfaceProfileRT if the feature is used
-	if (MaterialShadingModel == MSM_SubsurfaceProfile)
-	{
-		ENQUEUE_UNIQUE_RENDER_COMMAND_THREEPARAMETER(
-			UpdateMaterialRenderProxySubsurface,
-			const FSubsurfaceProfileStruct, Settings, Settings,
-			const USubsurfaceProfile*, SubsurfaceProfile, SubsurfaceProfile,
-			FMaterialRenderProxy&, Proxy, Proxy,
-		{
-			uint32 AllocationId = 0;
+				ENQUEUE_UNIQUE_RENDER_COMMAND_ONEPARAMETER(
+					UpdateMaterialRenderProxySubsurfaceSetup,
+					IRendererModule&, RendererModule, RendererModule,
+					{
+					GSubsufaceProfileTextureObject.SetRendererModule(&RendererModule);
+				});
+			}
+		}
 
-			if(SubsurfaceProfile)
+		{
+			ENQUEUE_UNIQUE_RENDER_COMMAND_THREEPARAMETER(
+				UpdateMaterialRenderProxySubsurface,
+				const FSubsurfaceProfileStruct, Settings, Settings,
+				USubsurfaceProfilePointer, LocalSubsurfaceProfile, (USubsurfaceProfilePointer)LocalSubsurfaceProfile,
+				FMaterialRenderProxy&, Proxy, Proxy,
 			{
-				AllocationId = GSubsufaceProfileTextureObject.AddOrUpdateProfile(Settings, *SubsurfaceProfile);
+				uint32 AllocationId = 0;
 
-				check(AllocationId >= 0 && AllocationId <= 255);
-			}
-			Proxy.SetSubsurfaceProfile(SubsurfaceProfile);
-		});
+				if (LocalSubsurfaceProfile)
+				{
+					AllocationId = GSubsufaceProfileTextureObject.AddOrUpdateProfile(Settings, LocalSubsurfaceProfile);
+
+					check(AllocationId >= 0 && AllocationId <= 255);
+				}
+				Proxy.SetSubsurfaceProfileRT(LocalSubsurfaceProfile);
+			});
+		}
 	}
 }
diff --git a/Engine/Source/Runtime/Engine/Private/Rendering/SubsurfaceProfile.cpp b/Engine/Source/Runtime/Engine/Private/Rendering/SubsurfaceProfile.cpp
--- a/Engine/Source/Runtime/Engine/Private/Rendering/SubsurfaceProfile.cpp
+++ b/Engine/Source/Runtime/Engine/Private/Rendering/SubsurfaceProfile.cpp
@@ -7,7 +7,7 @@
 DEFINE_LOG_CATEGORY_STATIC(LogSubsurfaceProfile, Log, All);
 
 // lives on the render thread
-TGlobalResource<FSubsurfaceProfileTexture> GSubsufaceProfileTextureObject;
+ENGINE_API TGlobalResource<FSubsurfaceProfileTexture> GSubsufaceProfileTextureObject;
 
 // Texture with one or more SubSurfaceProfiles or 0 if there is no user
 static TRefCountPtr<IPooledRenderTarget> GSSProfiles;
@@ -31,9 +31,9 @@ FSubsurfaceProfileTexture::~FSubsurfaceProfileTexture()
 	}
 }
 
-int32 FSubsurfaceProfileTexture::AddProfile(const FSubsurfaceProfileStruct Settings, const USubsurfaceProfile& InProfile)
+int32 FSubsurfaceProfileTexture::AddProfile(const FSubsurfaceProfileStruct Settings, const USubsurfaceProfilePointer InProfile)
 {
-	check(&InProfile);
+	check(InProfile);
 	check(FindAllocationId(InProfile) == -1);
 
 	int32 RetAllocationId = -1;
@@ -49,7 +49,7 @@ int32 FSubsurfaceProfileTexture::AddProfile(const FSubsurfaceProfileStruct Setti
 		if(RetAllocationId == -1)
 		{
 			RetAllocationId = SubsurfaceProfileEntries.Num();
-			SubsurfaceProfileEntries.Add(FSubsurfaceProfileEntry(Settings, &InProfile));
+			SubsurfaceProfileEntries.Add(FSubsurfaceProfileEntry(Settings, InProfile));
 		}
 	}
 
@@ -59,7 +59,7 @@ int32 FSubsurfaceProfileTexture::AddProfile(const FSubsurfaceProfileStruct Setti
 }
 
 
-void FSubsurfaceProfileTexture::RemoveProfile(USubsurfaceProfile& InProfile)
+void FSubsurfaceProfileTexture::RemoveProfile(const USubsurfaceProfilePointer InProfile)
 {
 	int32 AllocationId = FindAllocationId(InProfile);
 
@@ -72,10 +72,11 @@ void FSubsurfaceProfileTexture::RemoveProfile(USubsurfaceProfile& InProfile)
 	// >0 as 0 is used as default profile which should never be removed
 	check(AllocationId > 0);
 
-	check(SubsurfaceProfileEntries[AllocationId].GameThreadObject == &InProfile);
+	check(SubsurfaceProfileEntries[AllocationId].GameThreadObject == InProfile);
 
 	// make it available for reuse
 	SubsurfaceProfileEntries[AllocationId].GameThreadObject = 0;
+	SubsurfaceProfileEntries[AllocationId].Settings.Invalidate();
 }
 
 void FSubsurfaceProfileTexture::UpdateProfile(int32 AllocationId, const FSubsurfaceProfileStruct Settings)
@@ -192,12 +193,49 @@ void FSubsurfaceProfileTexture::CreateTexture(FRHICommandListImmediate& RHICmdLi
 	RHICmdList.UnlockTexture2D((FTexture2DRHIRef&)GSSProfiles->GetRenderTargetItem().ShaderResourceTexture, 0, false);
 }
 
-int32 FSubsurfaceProfileTexture::FindAllocationId(const USubsurfaceProfile& InProfile) const
+TCHAR MiniFontCharFromIndex(uint32 Index)
+{
+	if (Index <= 9)
+	{
+		return (TCHAR)('0' + Index);
+	}
+
+	Index -= 10;
+
+	if (Index <= 'Z' - 'A')
+	{
+		return (TCHAR)('A' + Index);
+	}
+
+	return (TCHAR)'?';
+}
+
+bool FSubsurfaceProfileTexture::GetEntryString(uint32 Index, FString& Out) const
+{
+	if (Index >= (uint32)SubsurfaceProfileEntries.Num())
+	{
+		return false;
+	}
+
+	const FSubsurfaceProfileStruct& ref = SubsurfaceProfileEntries[Index].Settings;
+
+
+	Out = FString::Printf(TEXT(" %c. %p ScatterRadius=%.1f, SubsurfaceColor=%.1f %.1f %.1f, FalloffColor=%.1f %.1f %.1f"), 
+		MiniFontCharFromIndex(Index), 
+		SubsurfaceProfileEntries[Index].GameThreadObject,
+		ref.ScatterRadius,
+		ref.SubsurfaceColor.R, ref.SubsurfaceColor.G, ref.SubsurfaceColor.B,
+		ref.FalloffColor.R, ref.FalloffColor.G, ref.FalloffColor.B);
+
+	return true;
+}
+
+int32 FSubsurfaceProfileTexture::FindAllocationId(const USubsurfaceProfilePointer InProfile) const
 {
 	// we start at 1 because [0] is the default profile and always [0].GameThreadObject = 0 so we don't need to iterate that one
 	for (int32 i = 1; i < SubsurfaceProfileEntries.Num(); ++i)
 	{
-		if (SubsurfaceProfileEntries[i].GameThreadObject == &InProfile)
+		if (SubsurfaceProfileEntries[i].GameThreadObject == InProfile)
 		{
 			return i;
 		}
@@ -250,7 +288,7 @@ void USubsurfaceProfile::BeginDestroy()
 {
 	ENQUEUE_UNIQUE_RENDER_COMMAND_ONEPARAMETER(
 		RemoveSubsurfaceProfile,
-		USubsurfaceProfile&, Ref, *this,
+		USubsurfaceProfilePointer, Ref, (USubsurfaceProfilePointer)this,
 		{
 			GSubsufaceProfileTextureObject.RemoveProfile(Ref);
 		});
@@ -263,9 +301,9 @@ void USubsurfaceProfile::PostEditChangeProperty(struct FPropertyChangedEvent& Pr
 	ENQUEUE_UNIQUE_RENDER_COMMAND_TWOPARAMETER(
 		UpdateSubsurfaceProfile,
 		const FSubsurfaceProfileStruct, Settings, this->Settings,
-		USubsurfaceProfile*, Profile, this,
+		USubsurfaceProfilePointer, Profile, (USubsurfaceProfilePointer)this,
 	{
 		// any changes to the setting require an update of the texture
-		GSubsufaceProfileTextureObject.UpdateProfile(Settings, *Profile);
+		GSubsufaceProfileTextureObject.UpdateProfile(Settings, Profile);
 	});
 }
\ No newline at end of file
diff --git a/Engine/Source/Runtime/Engine/Public/MaterialShared.h b/Engine/Source/Runtime/Engine/Public/MaterialShared.h
--- a/Engine/Source/Runtime/Engine/Public/MaterialShared.h
+++ b/Engine/Source/Runtime/Engine/Public/MaterialShared.h
@@ -1209,6 +1209,7 @@ struct FUniformExpressionCache
 };
 
 class USubsurfaceProfile;
+typedef void* USubsurfaceProfilePointer;
 
 /**
  * A material render proxy used by the renderer.
@@ -1273,8 +1274,8 @@ public:
 		return MaterialRenderProxyMap;
 	}
 
-	void SetSubsurfaceProfile(const USubsurfaceProfile* Ptr) { SubsurfaceProfileRT = Ptr; }
-	const USubsurfaceProfile* GetSubsurfaceProfile() const { return SubsurfaceProfileRT; }
+	void SetSubsurfaceProfileRT(const USubsurfaceProfilePointer Ptr) { SubsurfaceProfileRT = Ptr; }
+	USubsurfaceProfilePointer GetSubsurfaceProfileRT() const { return SubsurfaceProfileRT; }
 
 private:
 
@@ -1282,8 +1283,8 @@ private:
 	bool bSelected : 1;
 	/** true if the material is hovered. */
 	bool bHovered : 1;
-	/** 0 if not set */
-	const USubsurfaceProfile* SubsurfaceProfileRT;
+	/** 0 if not set, for the render thread */
+	USubsurfaceProfilePointer SubsurfaceProfileRT;
 
 	/** 
 	 * Tracks all material render proxies in all scenes, can only be accessed on the rendering thread.
diff --git a/Engine/Source/Runtime/Engine/Public/ShowFlags.h b/Engine/Source/Runtime/Engine/Public/ShowFlags.h
--- a/Engine/Source/Runtime/Engine/Public/ShowFlags.h
+++ b/Engine/Source/Runtime/Engine/Public/ShowFlags.h
@@ -208,6 +208,7 @@ struct FEngineShowFlags
 		SetVisualizeDistanceFieldAO(false);
 		SetVisualizeSSR(false);
 		SetForceGBuffer(false);
+		SetVisualizeSSS(false);
 	}
 
 	static bool CanBeToggledInEditor(const TCHAR* Name)
diff --git a/Engine/Source/Runtime/Engine/Public/ShowFlagsValues.inl b/Engine/Source/Runtime/Engine/Public/ShowFlagsValues.inl
--- a/Engine/Source/Runtime/Engine/Public/ShowFlagsValues.inl
+++ b/Engine/Source/Runtime/Engine/Public/ShowFlagsValues.inl
@@ -99,7 +99,7 @@ SHOWFLAG_ALWAYS_ACCESSIBLE(ReflectionEnvironment, SFG_LightingFeatures, LOCTEXT(
 /** Reflection environment mixes with lightmap shading. */
 SHOWFLAG_ALWAYS_ACCESSIBLE(ReflectionEnvironmentLightmapMixing, SFG_LightingFeatures, LOCTEXT("ReflectionEnvironmentLightmapMixingSF", "Reflection Environment Lightmap Mixing"))
 /** Visualize pixels that are outside of their object's bounding box (content error). */
-SHOWFLAG_ALWAYS_ACCESSIBLE(VisualizeOutOfBoundsPixels, SFG_Visualize, LOCTEXT("VisualizeOutOfBoundsPixelsSF", "Visualize Out of Bounds Pixels"))
+SHOWFLAG_ALWAYS_ACCESSIBLE(VisualizeOutOfBoundsPixels, SFG_Visualize, LOCTEXT("VisualizeOutOfBoundsPixelsSF", "Out of Bounds Pixels"))
 /** Whether to display the scene's diffuse. */
 SHOWFLAG_ALWAYS_ACCESSIBLE(Diffuse, SFG_LightingComponents, LOCTEXT("DiffuseSF", "Diffuse"))
 /** Whether to display the scene's specular, including reflections. */
@@ -109,7 +109,9 @@ SHOWFLAG_ALWAYS_ACCESSIBLE(SelectionOutline, SFG_Hidden, LOCTEXT("SelectionOutli
 /** If screen space reflections are enabled */
 SHOWFLAG_ALWAYS_ACCESSIBLE(ScreenSpaceReflections, SFG_LightingFeatures, LOCTEXT("ScreenSpaceReflectionsSF", "Screen Space Reflections"))
 /** If Screen Space Subsurface Scattering enabled */
-SHOWFLAG_ALWAYS_ACCESSIBLE(SubsurfaceScattering, SFG_LightingFeatures, LOCTEXT("SubsurfaceScatteringSF", "Screen Space Subsurface Scattering"))
+SHOWFLAG_FIXED_IN_SHIPPING(SubsurfaceScattering, 1, SFG_LightingFeatures, LOCTEXT("SubsurfaceScatteringSF", "Subsurface Scattering (Screen Space)"))
+/** If Screen Space Subsurface Scattering visualization is enabled */
+SHOWFLAG_FIXED_IN_SHIPPING(VisualizeSSS, 0, SFG_Visualize, LOCTEXT("VisualizeSSSSF", "Subsurface Scattering (Screen Space)"))
 /** If the indirect lighting cache is enabled */
 SHOWFLAG_ALWAYS_ACCESSIBLE(IndirectLightingCache, SFG_LightingFeatures, LOCTEXT("IndirectLightingCacheSF", "Indirect Lighting Cache"))
 /** calls debug drawing for AIs */
@@ -253,7 +255,7 @@ SHOWFLAG_ALWAYS_ACCESSIBLE(PrecomputedVisibility, SFG_Advanced, LOCTEXT("Precomp
 /** Contribution from sky light */
 SHOWFLAG_ALWAYS_ACCESSIBLE(SkyLighting, SFG_LightingComponents, LOCTEXT("SkyLightingSF", "Sky Lighting"))
 /** Visualize Light Propagation Volume, for developer (by default off): */
-SHOWFLAG_ALWAYS_ACCESSIBLE(VisualizeLPV, SFG_Visualize, LOCTEXT("VisualizeLPVSF", "Visualize Light Propagation Volume"))
+SHOWFLAG_ALWAYS_ACCESSIBLE(VisualizeLPV, SFG_Visualize, LOCTEXT("VisualizeLPVSF", "Light Propagation Volume"))
 /** Visualize preview shadow indicator */
 SHOWFLAG_ALWAYS_ACCESSIBLE(PreviewShadowsIndicator, SFG_Visualize, LOCTEXT("PreviewShadowIndicatorSF", "Preview Shadows Indicator"))
 /** Visualize precomputed visibilty cells */
@@ -265,11 +267,11 @@ SHOWFLAG_ALWAYS_ACCESSIBLE(LpvLightingOnly, SFG_Hidden, LOCTEXT("VisualizeLPVSF_
 /** Render Paper2D sprites */
 SHOWFLAG_ALWAYS_ACCESSIBLE(Paper2DSprites, SFG_Advanced, LOCTEXT("Paper2DSpritesSF", "Paper 2D Sprites"))
 /** Visualization of distance field AO */
-SHOWFLAG_ALWAYS_ACCESSIBLE(VisualizeDistanceFieldAO, SFG_Visualize, LOCTEXT("VisualizeDistanceFieldAOSF", "Visualize Distance Field AO"))
+SHOWFLAG_ALWAYS_ACCESSIBLE(VisualizeDistanceFieldAO, SFG_Visualize, LOCTEXT("VisualizeDistanceFieldAOSF", "Distance Field AO"))
 /** Distance field AO */
 SHOWFLAG_ALWAYS_ACCESSIBLE(DistanceFieldAO, SFG_LightingFeatures, LOCTEXT("DistanceFieldAOSF", "Distance Field AO"))
 /** Visualize screen space reflections, for developer (by default off): */
-SHOWFLAG_FIXED_IN_SHIPPING(VisualizeSSR, 0, SFG_Visualize, LOCTEXT("VisualizeSSR", "Visualize ScreenSpaceReflections"))
+SHOWFLAG_FIXED_IN_SHIPPING(VisualizeSSR, 0, SFG_Visualize, LOCTEXT("VisualizeSSR", "Screen Space Reflections"))
 /** Force the use of the GBuffer. */
 SHOWFLAG_ALWAYS_ACCESSIBLE(ForceGBuffer, SFG_Hidden, LOCTEXT("ForceGBuffer", "Force usage of GBuffer"))
 /** Visualize the senses of the Senses component */
diff --git a/Engine/Source/Runtime/Renderer/Private/CompositionLighting/CompositionLighting.cpp b/Engine/Source/Runtime/Renderer/Private/CompositionLighting/CompositionLighting.cpp
--- a/Engine/Source/Runtime/Renderer/Private/CompositionLighting/CompositionLighting.cpp
+++ b/Engine/Source/Runtime/Renderer/Private/CompositionLighting/CompositionLighting.cpp
@@ -357,7 +357,7 @@ void FCompositionLighting::ProcessLighting(FRHICommandListImmediate& RHICmdList,
 
 			if (View.bScreenSpaceSubsurfacePassNeeded && Radius > 0 && !bSimpleDynamicLighting && View.Family->EngineShowFlags.SubsurfaceScattering)
 			{
-				FRenderingCompositePass* PassSetup = Context.Graph.RegisterPass(new(FMemStack::Get()) FRCPassPostProcessSubsurfaceSetup());
+				FRenderingCompositePass* PassSetup = Context.Graph.RegisterPass(new(FMemStack::Get()) FRCPassPostProcessSubsurfaceSetup(false));
 				PassSetup->SetInput(ePId_Input0, Context.FinalOutput);
 
 				FRenderingCompositePass* Pass0 = Context.Graph.RegisterPass(new(FMemStack::Get()) FRCPassPostProcessSubsurface(0, Radius));
diff --git a/Engine/Source/Runtime/Renderer/Private/PostProcess/PostProcessSubsurface.cpp b/Engine/Source/Runtime/Renderer/Private/PostProcess/PostProcessSubsurface.cpp
--- a/Engine/Source/Runtime/Renderer/Private/PostProcess/PostProcessSubsurface.cpp
+++ b/Engine/Source/Runtime/Renderer/Private/PostProcess/PostProcessSubsurface.cpp
@@ -13,8 +13,11 @@
 
 ENGINE_API const IPooledRenderTarget* GetSubsufaceProfileTexture_RT(FRHICommandListImmediate& RHICmdList);
 
-/** Encapsulates the post processing ambient occlusion pixel shader. */
-template <uint32 SpecularCorrection>
+/**
+ * Encapsulates the post processing ambient occlusion pixel shader.
+ * @param SetupMode 0:without specular correction, 1: with specular correction, 2:visualize
+ */
+template <uint32 SetupMode>
 class FPostProcessSubsurfaceSetupPS : public FGlobalShader
 {
 	DECLARE_SHADER_TYPE(FPostProcessSubsurfaceSetupPS , Global);
@@ -27,7 +30,7 @@ class FPostProcessSubsurfaceSetupPS : public FGlobalShader
 	static void ModifyCompilationEnvironment(EShaderPlatform Platform, FShaderCompilerEnvironment& OutEnvironment)
 	{
 		FGlobalShader::ModifyCompilationEnvironment(Platform, OutEnvironment);
-		OutEnvironment.SetDefine(TEXT("SSSS_SPECULAR_CORRECTION"), SpecularCorrection);
+		OutEnvironment.SetDefine(TEXT("SETUP_MODE"), SetupMode);
 	}
 
 	/** Default constructor. */
@@ -36,6 +39,7 @@ class FPostProcessSubsurfaceSetupPS : public FGlobalShader
 public:
 	FPostProcessPassParameters PostprocessParameter;
 	FDeferredPixelShaderParameters DeferredParameters;
+	FShaderResourceParameter MiniFontTexture;
 
 	/** Initialization constructor. */
 	FPostProcessSubsurfaceSetupPS(const ShaderMetaType::CompiledShaderInitializerType& Initializer)
@@ -43,6 +47,7 @@ public:
 	{
 		PostprocessParameter.Bind(Initializer.ParameterMap);
 		DeferredParameters.Bind(Initializer.ParameterMap);
+		MiniFontTexture.Bind(Initializer.ParameterMap, TEXT("MiniFontTexture"));
 	}
 
 	void SetParameters(const FRenderingCompositePassContext& Context)
@@ -53,13 +58,14 @@ public:
 		FGlobalShader::SetParameters(Context.RHICmdList, ShaderRHI, Context.View);
 		PostprocessParameter.SetPS(ShaderRHI, Context, TStaticSamplerState<SF_Point,AM_Clamp,AM_Clamp,AM_Clamp>::GetRHI());
 		DeferredParameters.Set(Context.RHICmdList, ShaderRHI, Context.View);
+		SetTextureParameter(Context.RHICmdList, ShaderRHI, MiniFontTexture, GEngine->MiniFontTexture ? GEngine->MiniFontTexture->Resource->TextureRHI : GSystemTextures.WhiteDummy->GetRenderTargetItem().TargetableTexture);
 	}
 
 	// FShader interface.
 	virtual bool Serialize(FArchive& Ar)
 	{
 		bool bShaderHasOutdatedParameters = FGlobalShader::Serialize(Ar);
-		Ar << PostprocessParameter << DeferredParameters;
+		Ar << PostprocessParameter << DeferredParameters << MiniFontTexture;
 		return bShaderHasOutdatedParameters;
 	}
 
@@ -78,19 +84,18 @@ public:
 #define VARIATION1(A) typedef FPostProcessSubsurfaceSetupPS<A> FPostProcessSubsurfaceSetupPS##A; \
 	IMPLEMENT_SHADER_TYPE2(FPostProcessSubsurfaceSetupPS##A, SF_Pixel);
 
-	VARIATION1(0)			VARIATION1(1)
+	VARIATION1(0) VARIATION1(1) VARIATION1(2)
 
 #undef VARIATION1
 
 
-template <uint32 SpecularCorrection>
+template <uint32 SetupMode>
 void SetSubsurfaceSetupShader(const FRenderingCompositePassContext& Context)
 {
 	TShaderMapRef<FPostProcessVS> VertexShader(GetGlobalShaderMap());
-	TShaderMapRef<FPostProcessSubsurfaceSetupPS<SpecularCorrection> > PixelShader(GetGlobalShaderMap());
+	TShaderMapRef<FPostProcessSubsurfaceSetupPS<SetupMode> > PixelShader(GetGlobalShaderMap());
 
 	static FGlobalBoundShaderState BoundShaderState;
-	
 
 	SetGlobalBoundShaderState(Context.RHICmdList, Context.GetFeatureLevel(), BoundShaderState, GFilterVertexDeclaration.VertexDeclarationRHI, *VertexShader, *PixelShader);
 
@@ -122,6 +127,15 @@ static bool DoSpecularCorrection()
 	return CVarState && (SceneColorFormat >= 4);
 }
 
+FRCPassPostProcessSubsurfaceSetup::FRCPassPostProcessSubsurfaceSetup(bool bInVisualize)
+	: bVisualize(bInVisualize)
+{
+	if (bVisualize)
+	{
+		GSceneRenderTargets.AdjustGBufferRefCount(1);
+	}
+}
+
 void FRCPassPostProcessSubsurfaceSetup::Process(FRenderingCompositePassContext& Context)
 {
 	SCOPED_DRAW_EVENT(PostProcessSubsurfaceSetup, DEC_SCENE_ITEMS);
@@ -161,15 +175,22 @@ void FRCPassPostProcessSubsurfaceSetup::Process(FRenderingCompositePassContext&
 	Context.RHICmdList.SetRasterizerState(TStaticRasterizerState<>::GetRHI());
 	Context.RHICmdList.SetDepthStencilState(TStaticDepthStencilState<false, CF_Always>::GetRHI());
 
-	if(DoSpecularCorrection())
+	if (bVisualize)
 	{
-		// with separate specular
-		SetSubsurfaceSetupShader<1>(Context);
+		SetSubsurfaceSetupShader<2>(Context);
 	}
 	else
 	{
-		// no separate specular
-		SetSubsurfaceSetupShader<0>(Context);
+		if (DoSpecularCorrection())
+		{
+			// with separate specular
+			SetSubsurfaceSetupShader<1>(Context);
+		}
+		else
+		{
+			// no separate specular
+			SetSubsurfaceSetupShader<0>(Context);
+		}
 	}
 
 	// Draw a quad mapping scene color to the view's render target
@@ -185,7 +206,58 @@ void FRCPassPostProcessSubsurfaceSetup::Process(FRenderingCompositePassContext&
 		*VertexShader,
 		EDRF_UseTriangleOptimization);
 
+
+	if (bVisualize)
+	{
+		// this is a helper class for FCanvas to be able to get screen size
+		class FRenderTargetTemp : public FRenderTarget
+		{
+		public:
+			const FSceneView& View;
+			const FTexture2DRHIRef Texture;
+
+			FRenderTargetTemp(const FSceneView& InView, const FTexture2DRHIRef InTexture)
+				: View(InView), Texture(InTexture)
+			{
+			}
+			virtual FIntPoint GetSizeXY() const
+			{
+				return View.ViewRect.Size();
+			};
+			virtual const FTexture2DRHIRef& GetRenderTargetTexture() const
+			{
+				return Texture;
+			}
+		} TempRenderTarget(View, (const FTexture2DRHIRef&)DestRenderTarget.TargetableTexture);
+
+		FCanvas Canvas(&TempRenderTarget, NULL, ViewFamily.CurrentRealTime, ViewFamily.CurrentWorldTime, ViewFamily.DeltaWorldTime, Context.GetFeatureLevel());
+
+		float X = 30;
+		float Y = 28;
+		const float YStep = 14;
+
+		FString Line;
+
+		Line = FString::Printf(TEXT("Visualize Screen Space Subsurface Scattering"));
+		Canvas.DrawShadowedString(X, Y += YStep, *Line, GetStatsFont(), FLinearColor(1, 1, 1));
+
+		Y += YStep;
+
+		uint32 Index = 0;
+		while (GSubsufaceProfileTextureObject.GetEntryString(Index++, Line))
+		{
+			Canvas.DrawShadowedString(X, Y += YStep, *Line, GetStatsFont(), FLinearColor(1, 1, 1));
+		}
+
+		Canvas.Flush_RenderThread(Context.RHICmdList);
+	}
+
 	Context.RHICmdList.CopyToResolveTarget(DestRenderTarget.TargetableTexture, DestRenderTarget.ShaderResourceTexture, false, FResolveParams());
+
+	if (bVisualize)
+	{
+		GSceneRenderTargets.AdjustGBufferRefCount(-1);
+	}
 }
 
 FPooledRenderTargetDesc FRCPassPostProcessSubsurfaceSetup::ComputeOutputDesc(EPassOutputId InPassOutputId) const
@@ -215,7 +287,7 @@ class FPostProcessSubsurfacePS : public FGlobalShader
 	static void ModifyCompilationEnvironment(EShaderPlatform Platform, FShaderCompilerEnvironment& OutEnvironment)
 	{
 		FGlobalShader::ModifyCompilationEnvironment(Platform,OutEnvironment);
-		OutEnvironment.SetDefine(TEXT("METHOD"), Method);
+		OutEnvironment.SetDefine(TEXT("SSS_METHOD"), Method);
 	}
 
 	/** Default constructor. */
diff --git a/Engine/Source/Runtime/Renderer/Private/PostProcess/PostProcessSubsurface.h b/Engine/Source/Runtime/Renderer/Private/PostProcess/PostProcessSubsurface.h
--- a/Engine/Source/Runtime/Renderer/Private/PostProcess/PostProcessSubsurface.h
+++ b/Engine/Source/Runtime/Renderer/Private/PostProcess/PostProcessSubsurface.h
@@ -16,10 +16,14 @@
 class FRCPassPostProcessSubsurfaceSetup : public TRenderingCompositePassBase<1, 1>
 {
 public:
+	FRCPassPostProcessSubsurfaceSetup(bool bInVisualize);
+
 	// interface FRenderingCompositePass ---------
 	virtual void Process(FRenderingCompositePassContext& Context);
 	virtual FPooledRenderTargetDesc ComputeOutputDesc(EPassOutputId InPassOutputId) const;
 	virtual void Release() override { delete this; }
+
+	bool bVisualize;
 };
 
 
diff --git a/Engine/Source/Runtime/Renderer/Private/PostProcess/PostProcessing.cpp b/Engine/Source/Runtime/Renderer/Private/PostProcess/PostProcessing.cpp
--- a/Engine/Source/Runtime/Renderer/Private/PostProcess/PostProcessing.cpp
+++ b/Engine/Source/Runtime/Renderer/Private/PostProcess/PostProcessing.cpp
@@ -42,7 +42,7 @@
 #include "PostProcessTestImage.h"
 #include "PostProcessUIBlur.h"
 #include "HighResScreenshot.h"
-
+#include "PostProcessSubsurface.h"
 #include "PostProcessMorpheus.h"
 #include "IHeadMountedDisplay.h"
 #include "BufferVisualizationData.h"
@@ -1065,6 +1065,14 @@ void FPostProcessing::Process(FRHICommandListImmediate& RHICmdList, FViewInfo& V
 		
 		AddPostProcessMaterial(Context, BL_AfterTonemapping, SeparateTranslucency);
 
+		if (View.Family->EngineShowFlags.VisualizeSSS)
+		{
+			// the setup pass also does visualization, based on EngineShowFlags.VisualizeSSS
+			FRenderingCompositePass* PassVisualize = Context.Graph.RegisterPass(new(FMemStack::Get()) FRCPassPostProcessSubsurfaceSetup(true));
+			PassVisualize->SetInput(ePId_Input0, Context.FinalOutput);
+			Context.FinalOutput = FRenderingCompositeOutputRef(PassVisualize);
+		}
+
 		AddGBufferVisualization(Context, SeparateTranslucency);
 
 		bool bStereoRenderingAndHMD = View.Family->EngineShowFlags.StereoRendering && View.Family->EngineShowFlags.HMDDistortion;
