PC emulation of mobile MSAA works correctly.

diff --git a/Engine/Shaders/HdrCustomResolveShaders.usf b/Engine/Shaders/HdrCustomResolveShaders.usf
new file mode 100644
--- /dev/null
+++ b/Engine/Shaders/HdrCustomResolveShaders.usf
@@ -0,0 +1,75 @@
+// Copyright 1998-2014 Epic Games, Inc. All Rights Reserved.
+
+/*=============================================================================
+	HdrCustomResolve.usf: Custom resolve for HDR color and box filter.
+=============================================================================*/
+
+#include "Common.usf"
+
+float4 HdrCustomResolveVS(uint Id : SV_VertexID) : SV_POSITION
+{
+	int x = Id &  1;
+	int y = Id >> 1;
+
+	return float4(x * 4 - 1, y * 4 - 1, 0, 1);
+}
+
+float4 Encode(float4 Color)
+{
+	return float4(Color.rgb * rcp(Color.r*0.299 + Color.g*0.587 + Color.b*0.114 + 1.0), Color.a);
+}
+
+float4 Decode(float4 Color)
+{
+	return float4(Color.rgb * rcp(Color.r*(-0.299) + Color.g*(-0.587) + Color.b*(-0.114) + 1.0), Color.a);
+}
+
+#if HDR_CUSTOM_RESOLVE_2X
+
+	Texture2DMS<float4,2> Tex;
+	float4 HdrCustomResolve2xPS(float4 Pos : SV_POSITION) : SV_Target0
+	{
+		uint2 P = uint2(Pos.xy);
+		float4 C0 = Encode(Tex.Load(P, 0));
+		float4 C1 = Encode(Tex.Load(P, 1));
+		return Decode(C0*0.5 + C1*0.5);
+	}
+	
+#endif
+
+#if HDR_CUSTOM_RESOLVE_4X
+
+	Texture2DMS<float4,4> Tex;
+	float4 HdrCustomResolve4xPS(float4 Pos : SV_POSITION) : SV_Target0
+	{
+		uint2 P = uint2(Pos.xy);
+		float4 C0 = Encode(Tex.Load(P, 0));
+		float4 C1 = Encode(Tex.Load(P, 1));
+		float4 C2 = Encode(Tex.Load(P, 2));
+		float4 C3 = Encode(Tex.Load(P, 3));
+		return Decode(C0*0.25 + C1*0.25 + C2*0.25 + C3*0.25);
+	}
+
+#endif
+
+#if HDR_CUSTOM_RESOLVE_8X
+
+	Texture2DMS<float4,8> Tex;
+	float4 HdrCustomResolve8xPS(float4 Pos : SV_POSITION) : SV_Target0
+	{
+		uint2 P = uint2(Pos.xy);
+		float4 C0 = Encode(Tex.Load(P, 0));
+		float4 C1 = Encode(Tex.Load(P, 1));
+		float4 C2 = Encode(Tex.Load(P, 2));
+		float4 C3 = Encode(Tex.Load(P, 3));
+		float4 C4 = Encode(Tex.Load(P, 4));
+		float4 C5 = Encode(Tex.Load(P, 5));
+		float4 C6 = Encode(Tex.Load(P, 6));
+		float4 C7 = Encode(Tex.Load(P, 7));
+		return Decode(C0*0.125 + C1*0.125 + C2*0.125 + C3*0.125 + C4*0.125 + C5*0.125 + C6*0.125 + C7*0.125);
+	}
+
+#endif
+
+
+
diff --git a/Engine/Source/Runtime/Renderer/Private/HdrCustomResolveShaders.cpp b/Engine/Source/Runtime/Renderer/Private/HdrCustomResolveShaders.cpp
new file mode 100644
--- /dev/null
+++ b/Engine/Source/Runtime/Renderer/Private/HdrCustomResolveShaders.cpp
@@ -0,0 +1,11 @@
+// Copyright 1998-2014 Epic Games, Inc. All Rights Reserved.
+
+#include "RendererPrivate.h"
+#include "HdrCustomResolveShaders.h"
+#include "ShaderParameterUtils.h"
+
+IMPLEMENT_SHADER_TYPE(,FHdrCustomResolveVS,TEXT("HdrCustomResolveShaders"),TEXT("HdrCustomResolveVS"),SF_Vertex);
+IMPLEMENT_SHADER_TYPE(,FHdrCustomResolve2xPS,TEXT("HdrCustomResolveShaders"),TEXT("HdrCustomResolve2xPS"),SF_Pixel);
+IMPLEMENT_SHADER_TYPE(,FHdrCustomResolve4xPS,TEXT("HdrCustomResolveShaders"),TEXT("HdrCustomResolve4xPS"),SF_Pixel);
+IMPLEMENT_SHADER_TYPE(,FHdrCustomResolve8xPS,TEXT("HdrCustomResolveShaders"),TEXT("HdrCustomResolve8xPS"),SF_Pixel);
+
diff --git a/Engine/Source/Runtime/Renderer/Private/PostProcess/PostProcessMobile.cpp b/Engine/Source/Runtime/Renderer/Private/PostProcess/PostProcessMobile.cpp
--- a/Engine/Source/Runtime/Renderer/Private/PostProcess/PostProcessMobile.cpp
+++ b/Engine/Source/Runtime/Renderer/Private/PostProcess/PostProcessMobile.cpp
@@ -154,7 +154,7 @@ void FRCPassPostProcessBloomSetupES2::SetShader(const FRenderingCompositePassCon
 	uint32 UseSunDof = (UseSun << 1) + UseDof;
 
 	static const auto CVarMobileMSAA = IConsoleManager::Get().FindTConsoleVariableDataInt(TEXT("r.MobileMSAA"));
-	UseSunDof += CVarMobileMSAA ? (CVarMobileMSAA->GetValueOnRenderThread() > 1 ? 4 : 0) : 0;
+	UseSunDof += (GRHIShaderPlatform == SP_METAL) ? (CVarMobileMSAA ? (CVarMobileMSAA->GetValueOnRenderThread() > 1 ? 4 : 0) : 0) : 0;
 
 	switch(UseSunDof)
 	{
diff --git a/Engine/Source/Runtime/Renderer/Private/PostProcess/PostProcessTonemap.cpp b/Engine/Source/Runtime/Renderer/Private/PostProcess/PostProcessTonemap.cpp
--- a/Engine/Source/Runtime/Renderer/Private/PostProcess/PostProcessTonemap.cpp
+++ b/Engine/Source/Runtime/Renderer/Private/PostProcess/PostProcessTonemap.cpp
@@ -725,7 +725,7 @@ static uint32 TonemapperGenerateBitmaskMobile(const FRenderingCompositePassConte
 	}
 
 	static const auto CVarMobileMSAA = IConsoleManager::Get().FindTConsoleVariableDataInt(TEXT("r.MobileMSAA"));
-	if(CVarMobileMSAA ? CVarMobileMSAA->GetValueOnAnyThread() > 1 : false)
+	if((GRHIShaderPlatform == SP_METAL) && (CVarMobileMSAA ? CVarMobileMSAA->GetValueOnAnyThread() > 1 : false))
 	{
 		Bitmask += TonemapperMsaa;
 	}
diff --git a/Engine/Source/Runtime/Renderer/Private/PostProcess/SceneRenderTargets.cpp b/Engine/Source/Runtime/Renderer/Private/PostProcess/SceneRenderTargets.cpp
--- a/Engine/Source/Runtime/Renderer/Private/PostProcess/SceneRenderTargets.cpp
+++ b/Engine/Source/Runtime/Renderer/Private/PostProcess/SceneRenderTargets.cpp
@@ -10,6 +10,8 @@
 #include "ReflectionEnvironment.h"
 #include "LightPropagationVolume.h"
 #include "SceneUtils.h"
+#include "HdrCustomResolveShaders.h"
+
 
 // for LightPropagationVolume feature, could be exposed
 const int ReflectiveShadowMapResolution = 256;
@@ -640,7 +642,55 @@ void FSceneRenderTargets::ResolveSceneColor(FRHICommandList& RHICmdList, const F
 {
     SCOPED_DRAW_EVENT(ResolveSceneColor, DEC_SCENE_ITEMS);
 
-	RHICmdList.CopyToResolveTarget(GetSceneColorSurface(), GetSceneColorTexture(), true, FResolveParams(ResolveRect));
+	uint32 samples = GSceneRenderTargets.SceneColor->GetDesc().NumSamples;
+
+	if(samples <= 1)
+	{
+		RHICmdList.CopyToResolveTarget(GetSceneColorSurface(), GetSceneColorTexture(), true, FResolveParams(ResolveRect));
+	}
+	else 
+	{
+		// Custom shader based color resolve for HDR color to emulate mobile.
+		// TODO: Currently this resolves the entire surface, need to respect ResolveRect.
+		
+		SetRenderTarget(RHICmdList, GetSceneColorTexture(), FTextureRHIParamRef());
+		
+		RHICmdList.SetBlendState(TStaticBlendState<>::GetRHI());
+		RHICmdList.SetRasterizerState(TStaticRasterizerState<>::GetRHI());
+		RHICmdList.SetDepthStencilState(TStaticDepthStencilState<false, CF_Always>::GetRHI());
+
+		TShaderMapRef<FHdrCustomResolveVS> VertexShader(GetGlobalShaderMap());
+
+		if(samples == 2)
+		{
+			TShaderMapRef<FHdrCustomResolve2xPS> PixelShader(GetGlobalShaderMap());
+			static FGlobalBoundShaderState BoundShaderState;
+			SetGlobalBoundShaderState(RHICmdList, CurrentFeatureLevel, BoundShaderState, GetVertexDeclarationFVector4(), *VertexShader, *PixelShader);
+			PixelShader->SetParameters(RHICmdList, GSceneRenderTargets.SceneColor->GetRenderTargetItem().TargetableTexture);
+			RHICmdList.DrawPrimitive(PT_TriangleList, 0, 1, 1);
+		}
+		else if(samples == 4)
+		{
+			TShaderMapRef<FHdrCustomResolve4xPS> PixelShader(GetGlobalShaderMap());
+			static FGlobalBoundShaderState BoundShaderState;
+			SetGlobalBoundShaderState(RHICmdList, CurrentFeatureLevel, BoundShaderState, GetVertexDeclarationFVector4(), *VertexShader, *PixelShader);
+			PixelShader->SetParameters(RHICmdList, GSceneRenderTargets.SceneColor->GetRenderTargetItem().TargetableTexture);
+			RHICmdList.DrawPrimitive(PT_TriangleList, 0, 1, 1);
+		}
+		else if(samples == 8)
+		{
+			TShaderMapRef<FHdrCustomResolve8xPS> PixelShader(GetGlobalShaderMap());
+			static FGlobalBoundShaderState BoundShaderState;
+			SetGlobalBoundShaderState(RHICmdList, CurrentFeatureLevel, BoundShaderState, GetVertexDeclarationFVector4(), *VertexShader, *PixelShader);
+			PixelShader->SetParameters(RHICmdList, GSceneRenderTargets.SceneColor->GetRenderTargetItem().TargetableTexture);
+			RHICmdList.DrawPrimitive(PT_TriangleList, 0, 1, 1);
+		}
+		else
+		{
+			// Everything other than 2,4,8 samples is not implemented.
+			check(0);
+		}
+	}
 }
 
 /** Resolves the GBuffer targets so that their resolved textures can be sampled. */
diff --git a/Engine/Source/Runtime/Renderer/Public/HdrCustomResolveShaders.h b/Engine/Source/Runtime/Renderer/Public/HdrCustomResolveShaders.h
new file mode 100644
--- /dev/null
+++ b/Engine/Source/Runtime/Renderer/Public/HdrCustomResolveShaders.h
@@ -0,0 +1,151 @@
+// Copyright 1998-2014 Epic Games, Inc. All Rights Reserved.
+
+
+#pragma once
+
+#include "GlobalShader.h"
+#include "ShaderParameters.h"
+
+class FHdrCustomResolveVS : public FGlobalShader
+{
+	DECLARE_SHADER_TYPE(FHdrCustomResolveVS,Global);
+public:
+	FHdrCustomResolveVS() {}
+	FHdrCustomResolveVS(const ShaderMetaType::CompiledShaderInitializerType& Initializer)
+		: FGlobalShader( Initializer )
+	{
+	}
+
+	virtual bool Serialize(FArchive& Ar)
+	{
+		bool bShaderHasOutdatedParameters = FGlobalShader::Serialize(Ar);
+		return bShaderHasOutdatedParameters;
+	}
+
+	static bool ShouldCache(EShaderPlatform Platform)
+	{
+		return IsFeatureLevelSupported(Platform, ERHIFeatureLevel::SM5) || (Platform == SP_PCD3D_ES2);
+	}
+};
+
+class FHdrCustomResolve2xPS : public FGlobalShader
+{
+	DECLARE_SHADER_TYPE(FHdrCustomResolve2xPS,Global);
+public:
+	FHdrCustomResolve2xPS() {}
+	FHdrCustomResolve2xPS(const ShaderMetaType::CompiledShaderInitializerType& Initializer)
+		: FGlobalShader( Initializer )
+	{
+		Tex.Bind(Initializer.ParameterMap, TEXT("Tex"), SPF_Mandatory);
+	}
+
+	virtual bool Serialize(FArchive& Ar)
+	{
+		bool bShaderHasOutdatedParameters = FGlobalShader::Serialize(Ar);
+		Ar << Tex;
+		return bShaderHasOutdatedParameters;
+	}
+
+	static bool ShouldCache(EShaderPlatform Platform)
+	{
+		return IsFeatureLevelSupported(Platform, ERHIFeatureLevel::SM5) || (Platform == SP_PCD3D_ES2);
+	}
+
+	void SetParameters(FRHICommandList& RHICmdList, FTextureRHIParamRef Texture2DMS)
+	{
+		FPixelShaderRHIParamRef PixelShaderRHI = GetPixelShader();
+		SetTextureParameter(RHICmdList, PixelShaderRHI, Tex, Texture2DMS);
+	}
+
+	static void ModifyCompilationEnvironment(EShaderPlatform Platform, FShaderCompilerEnvironment& OutEnvironment)
+	{
+		FGlobalShader::ModifyCompilationEnvironment(Platform, OutEnvironment);
+		OutEnvironment.SetDefine(TEXT("HDR_CUSTOM_RESOLVE_2X"), 1);
+	}
+
+protected:
+	FShaderResourceParameter Tex;
+};
+
+class FHdrCustomResolve4xPS : public FGlobalShader
+{
+	DECLARE_SHADER_TYPE(FHdrCustomResolve4xPS,Global);
+public:
+	FHdrCustomResolve4xPS() {}
+	FHdrCustomResolve4xPS(const ShaderMetaType::CompiledShaderInitializerType& Initializer)
+		: FGlobalShader( Initializer )
+	{
+		Tex.Bind(Initializer.ParameterMap, TEXT("Tex"), SPF_Mandatory);
+	}
+
+	virtual bool Serialize(FArchive& Ar)
+	{
+		bool bShaderHasOutdatedParameters = FGlobalShader::Serialize(Ar);
+		Ar << Tex;
+		return bShaderHasOutdatedParameters;
+	}
+
+	static bool ShouldCache(EShaderPlatform Platform)
+	{
+		return IsFeatureLevelSupported(Platform, ERHIFeatureLevel::SM5) || (Platform == SP_PCD3D_ES2);
+	}
+
+	void SetParameters(FRHICommandList& RHICmdList, FTextureRHIParamRef Texture2DMS)
+	{
+		FPixelShaderRHIParamRef PixelShaderRHI = GetPixelShader();
+		SetTextureParameter(RHICmdList, PixelShaderRHI, Tex, Texture2DMS);
+	}
+
+	static void ModifyCompilationEnvironment(EShaderPlatform Platform, FShaderCompilerEnvironment& OutEnvironment)
+	{
+		FGlobalShader::ModifyCompilationEnvironment(Platform, OutEnvironment);
+		OutEnvironment.SetDefine(TEXT("HDR_CUSTOM_RESOLVE_4X"), 1);
+	}
+
+protected:
+	FShaderResourceParameter Tex;
+};
+
+
+class FHdrCustomResolve8xPS : public FGlobalShader
+{
+	DECLARE_SHADER_TYPE(FHdrCustomResolve8xPS,Global);
+public:
+	FHdrCustomResolve8xPS() {}
+	FHdrCustomResolve8xPS(const ShaderMetaType::CompiledShaderInitializerType& Initializer)
+		: FGlobalShader( Initializer )
+	{
+		Tex.Bind(Initializer.ParameterMap, TEXT("Tex"), SPF_Mandatory);
+	}
+
+	virtual bool Serialize(FArchive& Ar)
+	{
+		bool bShaderHasOutdatedParameters = FGlobalShader::Serialize(Ar);
+		Ar << Tex;
+		return bShaderHasOutdatedParameters;
+	}
+
+	static bool ShouldCache(EShaderPlatform Platform)
+	{
+		return IsFeatureLevelSupported(Platform, ERHIFeatureLevel::SM5) || (Platform == SP_PCD3D_ES2);
+	}
+
+	void SetParameters(FRHICommandList& RHICmdList, FTextureRHIParamRef Texture2DMS)
+	{
+		FPixelShaderRHIParamRef PixelShaderRHI = GetPixelShader();
+		SetTextureParameter(RHICmdList, PixelShaderRHI, Tex, Texture2DMS);
+	}
+
+	static void ModifyCompilationEnvironment(EShaderPlatform Platform, FShaderCompilerEnvironment& OutEnvironment)
+	{
+		FGlobalShader::ModifyCompilationEnvironment(Platform, OutEnvironment);
+		OutEnvironment.SetDefine(TEXT("HDR_CUSTOM_RESOLVE_8X"), 1);
+	}
+
+protected:
+	FShaderResourceParameter Tex;
+};
+
+
+
+
