Hopefully this will fix MSAA for Zen.

diff --git a/Engine/Shaders/PostProcessMobile.usf b/Engine/Shaders/PostProcessMobile.usf
--- a/Engine/Shaders/PostProcessMobile.usf
+++ b/Engine/Shaders/PostProcessMobile.usf
@@ -52,6 +52,7 @@ half2 SunConstDepthMaskScaleBias()
 
 //
 // Convert depth in alpha into combined circle of confusion and sun intensity.
+// Or pre-tonemap before hardware box-filtered resolve.
 //
 
 void SunMaskVS_ES2(
@@ -76,26 +77,31 @@ void SunMaskPS_ES2(
 		OutColor = PostprocessInput0.Sample(PostprocessInput0Sampler, InUVPos.xy);
 	#endif
 
-	half InDepth = OutColor.a;
-
-	#if ES2_USE_SUN
-		half2 DepthMaskScaleBias = SunConstDepthMaskScaleBias();
-		half FarAmount = saturate(InDepth * DepthMaskScaleBias.x + DepthMaskScaleBias.y);
-		half3 SunAmount = OutColor.rgb * SunColorApertureDiv2.rgb;
-
-		half2 Pos = InUVPos.zw * 0.5 + 0.5;
-		half EdgeMask = 1.0f - Pos.x * (1.0f - Pos.x) * Pos.y * (1.0f - Pos.y) * 8.0f;
-		EdgeMask = EdgeMask * EdgeMask;
-
-		FarAmount *= 1.0-EdgeMask;
-
-		OutColor.a = min(min(SunAmount.r, SunAmount.g), SunAmount.b) * FarAmount;
+	#if ES2_USE_MSAA
+		// On-chip pre-tonemap before MSAA resolve.
+		OutColor.rgb *= rcp(OutColor.r*0.299 + OutColor.g*0.587 + OutColor.b*0.114 + 1.0);
 	#else
-		OutColor.a = 0.0;
-	#endif
+		half InDepth = OutColor.a;
 	
-	#if ES2_USE_DOF
-		OutColor.a += Coc(InDepth);
+		#if ES2_USE_SUN
+			half2 DepthMaskScaleBias = SunConstDepthMaskScaleBias();
+			half FarAmount = saturate(InDepth * DepthMaskScaleBias.x + DepthMaskScaleBias.y);
+			half3 SunAmount = OutColor.rgb * SunColorApertureDiv2.rgb;
+	
+			half2 Pos = InUVPos.zw * 0.5 + 0.5;
+			half EdgeMask = 1.0f - Pos.x * (1.0f - Pos.x) * Pos.y * (1.0f - Pos.y) * 8.0f;
+			EdgeMask = EdgeMask * EdgeMask;
+	
+			FarAmount *= 1.0-EdgeMask;
+	
+			OutColor.a = min(min(SunAmount.r, SunAmount.g), SunAmount.b) * FarAmount;
+		#else
+			OutColor.a = 0.0;
+		#endif
+		
+		#if ES2_USE_DOF
+			OutColor.a += Coc(InDepth);
+		#endif
 	#endif
 }
 
diff --git a/Engine/Source/Runtime/Renderer/Private/ForwardShadingRenderer.cpp b/Engine/Source/Runtime/Renderer/Private/ForwardShadingRenderer.cpp
--- a/Engine/Source/Runtime/Renderer/Private/ForwardShadingRenderer.cpp
+++ b/Engine/Source/Runtime/Renderer/Private/ForwardShadingRenderer.cpp
@@ -117,11 +117,15 @@ void FForwardShadingSceneRenderer::Render(FRHICommandListImmediate& RHICmdList)
 		// Using only view 0 to check to do on-chip transform of alpha.
 		FViewInfo& View = Views[0];
 
+		static const auto CVarMobileMSAA = IConsoleManager::Get().FindTConsoleVariableDataInt(TEXT("r.MobileMSAA"));
+
 		bool bOnChipSunMask = 
 			GSupportsRenderTargetFormat_PF_FloatRGBA && 
 			GSupportsShaderFramebufferFetch &&
 			ViewFamily.EngineShowFlags.PostProcessing &&
-			((View.bLightShaftUse) || (View.FinalPostProcessSettings.DepthOfFieldScale > 0.0));
+			((View.bLightShaftUse) || (View.FinalPostProcessSettings.DepthOfFieldScale > 0.0) || 
+			((GRHIShaderPlatform == SP_METAL) && (CVarMobileMSAA ? CVarMobileMSAA->GetValueOnAnyThread() > 1 : false))
+			);
 
 		// Convert alpha from depth to circle of confusion with sunshaft intensity.
 		// This is done before resolve on hardware with framebuffer fetch.
diff --git a/Engine/Source/Runtime/Renderer/Private/PostProcess/PostProcessMobile.cpp b/Engine/Source/Runtime/Renderer/Private/PostProcess/PostProcessMobile.cpp
--- a/Engine/Source/Runtime/Renderer/Private/PostProcess/PostProcessMobile.cpp
+++ b/Engine/Source/Runtime/Renderer/Private/PostProcess/PostProcessMobile.cpp
@@ -791,7 +791,7 @@ FPooledRenderTargetDesc FRCPassPostProcessBloomUpES2::ComputeOutputDesc(EPassOut
 // SUN MASK
 //
 
-template <uint32 UseFetchSunDof> // 0=none, 1=dof, 2=sun, 3=sun&dof, {4,5,6,7}=ES2_USE_FETCH
+template <uint32 UseFetchSunDof> // 0=none, 1=dof, 2=sun, 3=sun&dof, {4,5,6,7}=ES2_USE_FETCH, 8=MSAA-pre-resolve
 class FPostProcessSunMaskPS_ES2 : public FGlobalShader
 {
 	DECLARE_SHADER_TYPE(FPostProcessSunMaskPS_ES2, Global);
@@ -804,6 +804,7 @@ class FPostProcessSunMaskPS_ES2 : public FGlobalShader
 	static void ModifyCompilationEnvironment(EShaderPlatform Platform, FShaderCompilerEnvironment& OutEnvironment)
 	{
 		FGlobalShader::ModifyCompilationEnvironment(Platform, OutEnvironment);
+		OutEnvironment.SetDefine(TEXT("ES2_USE_MSAA"), (UseFetchSunDof & 8) ? (uint32)1 : (uint32)0);
 		OutEnvironment.SetDefine(TEXT("ES2_USE_FETCH"), (UseFetchSunDof & 4) ? (uint32)1 : (uint32)0);
 		OutEnvironment.SetDefine(TEXT("ES2_USE_SUN"), (UseFetchSunDof & 2) ? (uint32)1 : (uint32)0);
 		OutEnvironment.SetDefine(TEXT("ES2_USE_DOF"), (UseFetchSunDof & 1) ? (uint32)1 : (uint32)0);
@@ -853,6 +854,7 @@ typedef FPostProcessSunMaskPS_ES2<4> FPostProcessSunMaskPS_ES2_4;
 typedef FPostProcessSunMaskPS_ES2<5> FPostProcessSunMaskPS_ES2_5;
 typedef FPostProcessSunMaskPS_ES2<6> FPostProcessSunMaskPS_ES2_6;
 typedef FPostProcessSunMaskPS_ES2<7> FPostProcessSunMaskPS_ES2_7;
+typedef FPostProcessSunMaskPS_ES2<8> FPostProcessSunMaskPS_ES2_8;
 IMPLEMENT_SHADER_TYPE(template<>,FPostProcessSunMaskPS_ES2_0,TEXT("PostProcessMobile"),TEXT("SunMaskPS_ES2"),SF_Pixel);
 IMPLEMENT_SHADER_TYPE(template<>,FPostProcessSunMaskPS_ES2_1,TEXT("PostProcessMobile"),TEXT("SunMaskPS_ES2"),SF_Pixel);
 IMPLEMENT_SHADER_TYPE(template<>,FPostProcessSunMaskPS_ES2_2,TEXT("PostProcessMobile"),TEXT("SunMaskPS_ES2"),SF_Pixel);
@@ -861,6 +863,7 @@ IMPLEMENT_SHADER_TYPE(template<>,FPostProcessSunMaskPS_ES2_4,TEXT("PostProcessMo
 IMPLEMENT_SHADER_TYPE(template<>,FPostProcessSunMaskPS_ES2_5,TEXT("PostProcessMobile"),TEXT("SunMaskPS_ES2"),SF_Pixel);
 IMPLEMENT_SHADER_TYPE(template<>,FPostProcessSunMaskPS_ES2_6,TEXT("PostProcessMobile"),TEXT("SunMaskPS_ES2"),SF_Pixel);
 IMPLEMENT_SHADER_TYPE(template<>,FPostProcessSunMaskPS_ES2_7,TEXT("PostProcessMobile"),TEXT("SunMaskPS_ES2"),SF_Pixel);
+IMPLEMENT_SHADER_TYPE(template<>,FPostProcessSunMaskPS_ES2_8,TEXT("PostProcessMobile"),TEXT("SunMaskPS_ES2"),SF_Pixel);
 
 
 class FPostProcessSunMaskVS_ES2 : public FGlobalShader
@@ -922,6 +925,12 @@ void FRCPassPostProcessSunMaskES2::SetShader(const FRenderingCompositePassContex
 	uint32 UseFetch = GSupportsShaderFramebufferFetch ? 1 : 0;
 	uint32 UseFetchSunDof = (UseFetch << 2) + (UseSun << 1) + UseDof;
 
+	static const auto CVarMobileMSAA = IConsoleManager::Get().FindTConsoleVariableDataInt(TEXT("r.MobileMSAA"));
+	if((GRHIShaderPlatform == SP_METAL) && (CVarMobileMSAA ? CVarMobileMSAA->GetValueOnAnyThread() > 1 : false))
+	{
+		UseFetchSunDof = 8;
+	}
+
 	switch(UseFetchSunDof)
 	{
 		case 0: SunMask_SetShader<0>(Context); break;
@@ -932,6 +941,7 @@ void FRCPassPostProcessSunMaskES2::SetShader(const FRenderingCompositePassContex
 		case 5: SunMask_SetShader<5>(Context); break;
 		case 6: SunMask_SetShader<6>(Context); break;
 		case 7: SunMask_SetShader<7>(Context); break;
+		case 8: SunMask_SetShader<8>(Context); break;
 	}
 }
 
