* fixed ScreenSpaceSubsurfaceScattering not affected by [REDACTED_PERSON]

diff --git a/Engine/Source/Runtime/Renderer/Private/CompositionLighting/PostProcessAmbientOcclusion.cpp b/Engine/Source/Runtime/Renderer/Private/CompositionLighting/PostProcessAmbientOcclusion.cpp
--- a/Engine/Source/Runtime/Renderer/Private/CompositionLighting/PostProcessAmbientOcclusion.cpp
+++ b/Engine/Source/Runtime/Renderer/Private/CompositionLighting/PostProcessAmbientOcclusion.cpp
@@ -602,7 +602,7 @@ void FRCPassPostProcessBasePassAO::Process(FRenderingCompositePassContext& Conte
 	Context.SetViewportAndCallRHI(View.ViewRect);
 
 	// set the state
-	Context.RHICmdList.SetBlendState(TStaticBlendState<CW_RGB, BO_Add, BF_DestColor, BF_Zero>::GetRHI());
+	Context.RHICmdList.SetBlendState(TStaticBlendState<CW_RGBA, BO_Add, BF_DestColor, BF_Zero, BO_Add, BF_DestAlpha, BF_Zero>::GetRHI());
 	Context.RHICmdList.SetRasterizerState(TStaticRasterizerState<>::GetRHI());
 	Context.RHICmdList.SetDepthStencilState(TStaticDepthStencilState<false, CF_Always>::GetRHI());
 
