Improved specular handling in the shading model.

Switched to joint Smith geometric shadowing and masking.

Now using off-specular peak for IBL sampling.

diff --git a/Engine/Shaders/BRDF.usf b/Engine/Shaders/BRDF.usf
--- a/Engine/Shaders/BRDF.usf
+++ b/Engine/Shaders/BRDF.usf
@@ -43,16 +43,16 @@
 
 float3 Diffuse_Lambert( float3 DiffuseColor )
 {
-	return DiffuseColor / PI;
+	return DiffuseColor * (1 / PI);
 }
 
 // [Burley 2012, "Physically-Based Shading at Disney"]
 float3 Diffuse_Burley( float3 DiffuseColor, float Roughness, float NoV, float NoL, float VoH )
 {
 	float FD90 = 0.5 + 2 * VoH * VoH * Roughness;
-	float FdV = 1 + (FD90 - 1) * exp2( (-5.55473 * NoV - 6.98316) * NoV );
-	float FdL = 1 + (FD90 - 1) * exp2( (-5.55473 * NoL - 6.98316) * NoL );
-	return DiffuseColor / PI * FdV * FdL;
+	float FdV = 1 + (FD90 - 1) * pow( 1 - NoV, 5 );
+	float FdL = 1 + (FD90 - 1) * pow( 1 - NoL, 5 );
+	return DiffuseColor * ( 1 / PI * FdV * FdL );
 }
 
 // [Gotanda 2012, "Beyond a Simple Physically Based Blinn-Phong Model in Real-Time"]
@@ -146,6 +146,16 @@ float Vis_Smith( float Roughness, float NoV, float NoL )
 	return rcp( Vis_SmithV * Vis_SmithL );
 }
 
+// Appoximation of joint Smith term for GGX
+// [Heitz 2014, "Understanding the Masking-Shadowing Function in Microfacet-Based BRDFs"]
+float Vis_SmithJointApprox( float Roughness, float NoV, float NoL )
+{
+	float a = Square( Roughness );
+	float Vis_SmithV = NoL * ( NoV * ( 1 - a ) + a );
+	float Vis_SmithL = NoV * ( NoL * ( 1 - a ) + a );
+	return 0.5 * rcp( Vis_SmithV + Vis_SmithL );
+}
+
 float3 F_None( float3 SpecularColor )
 {
 	return SpecularColor;
diff --git a/Engine/Shaders/DeferredLightingCommon.usf b/Engine/Shaders/DeferredLightingCommon.usf
--- a/Engine/Shaders/DeferredLightingCommon.usf
+++ b/Engine/Shaders/DeferredLightingCommon.usf
@@ -52,17 +52,17 @@ struct FSimpleDeferredLightData
 #if 0
 void StandardShadingShared( float3 DiffuseColor, float3 SpecularColor, float Roughness, float3 V, half3 N )
 {
-	float NoV = saturate( dot(N, V) );
+	//float NoV = saturate( dot(N, V) );
+	float NoV = abs( dot(N, V) ) + 1e-5;
 
 	// Diffuse_Lambert
 	Shared.DiffuseMul = DiffuseColor * (1.0 / PI);
 
-	// D_GGX, Vis_Schlick
+	// D_GGX, Vis_SmithJointApprox
 	float m = Roughness * Roughness;
 	Shared.m2 = m * m;
-	Shared.k = Square( Roughness ) * 0.5;
-	float Vis_SchlickV = NoV * (1 - Shared.k) + Shared.k;
-	Shared.SpecularMul = (0.25 / PI) * (Shared.m2 / Vis_SchlickV);
+	Shared.SpecularMul = (0.5 / PI) * Shared.m2;
+	Shared.VisMad = float2( 2 * NoV * ( 1 - m ) + m, NoV * m );
 	
 	// F_Schlick
 	Shared.SpecularMul *= saturate( 50.0 * SpecularColor.g );
@@ -75,15 +75,14 @@ void StandardShadingPerLight( Shared, float3 L, float3 V, half3 N )
 	float NoH = saturate( dot(N, H) );		// 2 mad, 1 mul
 	float VoH = saturate( dot(V, H) );		// 2 mad, 1 mul
 
-	// D_GGX, Vis_Schlick
-	float d = ( NoH * Shared.m2 - NoH ) * NoH + 1;					// 2 mad
-	float Vis_SchlickL = NoL * (1 - Shared.k) + Shared.k;			// 1 mad, 1 add
-	float D_Vis = Shared.SpecularMul * rcp( d * d * Vis_SchlickL );	// 3 mul, 1 rcp
+	// D_GGX, Vis_SmithJointApprox
+	float d = ( NoH * Shared.m2 - NoH ) * NoH + 1;			// 2 mad
+	float v = NoL * Shared.VisMad.x + Shared.VisMad.y;		// 1 mad
+	float D_Vis = Shared.SpecularMul * rcp( d * d * v );	// 3 mul, 1 rcp
 
 	// F_Schlick
 	float Fc = pow( 1 - VoH, 5 );							// 1 sub, 3 mul
-	//float Fc = exp2( (-5.55473 * VoH - 6.98316) * VoH );	// 1 mad, 1 mul, 1 exp
-	float3 F = Fc + (1 - Fc) * SpecularColor;				// 1 add, 3 mad
+	float3 F = Fc + (1 - Fc) * SpecularColor;				// 1 sub, 3 mad
 
 	return Shared.DiffuseMul + D_Vis * F;					// 3 mad
 }
@@ -94,13 +93,14 @@ float3 StandardShading( FGBufferData GBuffer, float3 LobeRoughness, float3 LobeE
 {
 	float3 H = normalize(V + L);
 	float NoL = saturate( dot(N, L) );
-	float NoV = saturate( dot(N, V) );
+	//float NoV = saturate( dot(N, V) );
+	float NoV = abs( dot(N, V) ) + 1e-5;
 	float NoH = saturate( dot(N, H) );
 	float VoH = saturate( dot(V, H) );
 	
 	// Generalized microfacet specular
 	float D = D_GGX( LobeRoughness[1], NoH ) * LobeEnergy[1];
-	float Vis = Vis_Schlick( LobeRoughness[1], NoV, NoL );
+	float Vis = Vis_SmithJointApprox( LobeRoughness[1], NoV, NoL );
 	float3 F = F_Schlick( GBuffer.SpecularColor, VoH );
 
 	return Diffuse_Lambert(GBuffer.DiffuseColor) * (LobeEnergy[2] * DiffSpecMask.r) + (D * Vis * DiffSpecMask.g) * F;
diff --git a/Engine/Shaders/PostProcessAmbient.usf b/Engine/Shaders/PostProcessAmbient.usf
--- a/Engine/Shaders/PostProcessAmbient.usf
+++ b/Engine/Shaders/PostProcessAmbient.usf
@@ -82,13 +82,8 @@ float3 StandardShading( FGBufferData GBuffer, float Roughness, float3 L, float3
 	
 	// Generalized microfacet specular
 	float D = D_GGX( Roughness, NoH );
-	float Vis = Vis_Schlick( Roughness, NoV, NoL );
+	float Vis = Vis_SmithJointApprox( Roughness, NoV, NoL );
 	float3 F = F_Schlick( GBuffer.SpecularColor, VoH );
-	
-	//float k = Square( Roughness ) * 0.25;
-	//float Vis = 0.25 / Square( VoH * (1 - k) + k );
-	//float Vis = 0.25 / ( VoH * VoH * (1 - k*k) + k*k );
-	//float Vis = Vis_Smith( Roughness, NoV, NoL );
 
 	return Diffuse_Lambert( GBuffer.DiffuseColor ) + (D * Vis) * F;
 }
@@ -234,7 +229,7 @@ float3 ImageBasedLightingMIS( FGBufferData GBuffer, float3 V, float3 N, uint2 Ra
 	{
 		16,
 		16,
-		16,
+		0,
 	};
 
 	UNROLL
@@ -349,7 +344,7 @@ float3 PrefilterEnvMap( uint2 Random, float Roughness, float3 R )
 	return FilteredColor / max( Weight, 0.001 );
 }
 
-float2 IntegrateBRDF( uint2 Random, float Roughness, float NoV )
+float3 IntegrateBRDF( uint2 Random, float Roughness, float NoV )
 {
 	float3 V;
 	V.x = sqrt( 1.0f - NoV * NoV );	// sin
@@ -358,51 +353,59 @@ float2 IntegrateBRDF( uint2 Random, float Roughness, float NoV )
 
 	float A = 0;
 	float B = 0;
+	float C = 0;
 
 	const uint NumSamples = 64;
 	for( uint i = 0; i < NumSamples; i++ )
 	{
 		float2 E = Hammersley( i, NumSamples, Random );
-		float3 H = ImportanceSampleGGX( E, Roughness ).xyz;
-		float3 L = 2 * dot( V, H ) * H - V;
 
-		float NoL = saturate( L.z );
-		float NoH = saturate( H.z );
-		float VoH = saturate( dot( V, H ) );
-
-		if( NoL > 0 )
 		{
-			//float Vis = Vis_Schlick( Roughness, NoV, NoL );
+			float3 H = ImportanceSampleGGX( E, Roughness ).xyz;
+			float3 L = 2 * dot( V, H ) * H - V;
 
-			float a = Square( Roughness );
-			float a2 = a*a;
+			float NoL = saturate( L.z );
+			float NoH = saturate( H.z );
+			float VoH = saturate( dot( V, H ) );
 
-#if 1
-			float Vis_SchlickV = NoL * ( NoV * ( 1 - a ) + a );
-			float Vis_SchlickL = NoV * ( NoL * ( 1 - a ) + a );
-			float Vis = 0.5 * rcp( Vis_SchlickV + Vis_SchlickL );
-#elif 0
-			float Vis_SmithV = NoV + sqrt( NoV * (NoV - NoV * a2) + a2 );
-			float Vis_SmithL = NoL + sqrt( NoL * (NoL - NoL * a2) + a2 );
-			float Vis = rcp( Vis_SmithV * Vis_SmithL );
-#else
-			float Vis_SmithV = NoL * sqrt( NoV * (NoV - NoV * a2) + a2 );
-			float Vis_SmithL = NoV * sqrt( NoL * (NoL - NoL * a2) + a2 );
-			float Vis = 0.5 * rcp( Vis_SmithV + Vis_SmithL );
-#endif
-
-			// Incident light = NoL
-			// pdf = D * NoH / (4 * VoH)
-			// NoL * Vis / pdf
-			float NoL_Vis_PDF = NoL * Vis * (4 * VoH / NoH);
+			if( NoL > 0 )
+			{
+				float Vis = Vis_SmithJointApprox( Roughness, NoV, NoL );
+
+				float a = Square( Roughness );
+				float a2 = a*a;
+				float Vis_SmithV = NoL * sqrt( NoV * (NoV - NoV * a2) + a2 );
+				float Vis_SmithL = NoV * sqrt( NoL * (NoL - NoL * a2) + a2 );
+				//float Vis = 0.5 * rcp( Vis_SmithV + Vis_SmithL );
+
+				// Incident light = NoL
+				// pdf = D * NoH / (4 * VoH)
+				// NoL * Vis / pdf
+				float NoL_Vis_PDF = NoL * Vis * (4 * VoH / NoH);
+
+				float Fc = pow( 1 - VoH, 5 );
+				A += (1 - Fc) * NoL_Vis_PDF;
+				B += Fc * NoL_Vis_PDF;
+			}
+		}
 
-			float Fc = pow( 1 - VoH, 5 );
-			A += (1 - Fc) * NoL_Vis_PDF;
-			B += Fc * NoL_Vis_PDF;
+		{
+			float3 L = CosineSampleHemisphere( E ).xyz;
+			float3 H = normalize(V + L);
+
+			float NoL = saturate( L.z );
+			float NoH = saturate( H.z );
+			float VoH = saturate( dot( V, H ) );
+
+			//float FD90 = 0.5 * Roughness + 2 * VoH * VoH * Roughness;
+			float FD90 = 0.5 + 2 * VoH * VoH * Roughness;
+			float FdV = 1 + (FD90 - 1) * pow( 1 - NoV, 5 );
+			float FdL = 1 + (FD90 - 1) * pow( 1 - NoL, 5 );
+			C += FdV * FdL * lerp( 1, 1.0/1.5, Roughness );
 		}
 	}
 
-	return float2( A, B ) / NumSamples;
+	return float3( A, B, C ) / NumSamples;
 }
 
 float3 ApproximateSpecularIBL( uint2 Random, float3 SpecularColor, float Roughness, float3 N, float3 V )
@@ -432,8 +435,16 @@ void MainPS(in float4 UVAndScreenPos : TEXCOORD0, out float4 OutColor : SV_Targe
 	int2 PixelPos = int2(UVAndScreenPos.zw * ScreenPosToPixel.xy + ScreenPosToPixel.zw + 0.5f);
 
 	float3 ScreenVector = normalize(mul(float4(ScreenSpacePos, 1, 0), View.ScreenToWorld).xyz);
-	float NoV = saturate( dot( GBuffer.WorldNormal, -ScreenVector ) );
-	float3 R = reflect( ScreenVector, GBuffer.WorldNormal );
+
+	float3 N = GBuffer.WorldNormal;
+	float3 V = -ScreenVector;
+	float3 R = 2 * dot( V, N ) * N - V;
+
+	float NoV = abs( dot(N, V) ) + 1e-5;
+
+	// Point lobe in off-specular peak direction
+	float a = Square( GBuffer.Roughness );
+	R = lerp( N, R, (1 - a) * ( sqrt(1 - a) + a ) );
 	
 	uint2 Random = ScrambleTEA( PixelPos );
 	Random.x ^= View.Random;
@@ -442,7 +453,16 @@ void MainPS(in float4 UVAndScreenPos : TEXCOORD0, out float4 OutColor : SV_Targe
 	float3 NonSpecularContribution = 0;
 	float3 SpecularContribution = 0;
 
-	float3 DiffuseLookup =  TextureCubeSampleLevel(AmbientCubemap, AmbientCubemapSampler, GBuffer.WorldNormal, AbsoluteDiffuseMip).rgb;
+#if 0
+	float3 RetroReflectionDir = lerp( N, V, saturate( ( NoV * (1.02341 * a - 1.51174) + -0.511705 * a + 0.755868 ) * a ) );
+	float3 DiffuseLookup =  TextureCubeSampleLevel(AmbientCubemap, AmbientCubemapSampler, RetroReflectionDir, AbsoluteDiffuseMip).rgb;
+
+	float3 ABC = PreIntegratedGF.SampleLevel( PreIntegratedGFSampler, float2( NoV, GBuffer.Roughness ), 0 ).rgb;
+	//float3 ABC = IntegrateBRDF( Random, GBuffer.Roughness, NoV );
+	DiffuseLookup *= ABC.z;
+#else
+	float3 DiffuseLookup =  TextureCubeSampleLevel(AmbientCubemap, AmbientCubemapSampler, N, AbsoluteDiffuseMip).rgb;
+#endif
 
 	// Diffuse
 	{
diff --git a/Engine/Shaders/ReflectionEnvironmentComputeShaders.usf b/Engine/Shaders/ReflectionEnvironmentComputeShaders.usf
--- a/Engine/Shaders/ReflectionEnvironmentComputeShaders.usf
+++ b/Engine/Shaders/ReflectionEnvironmentComputeShaders.usf
@@ -583,8 +583,15 @@ void ReflectionEnvironmentTiledDeferredMain(
 	BRANCH
 	if( GBuffer.ShadingModelID != SHADINGMODELID_UNLIT )
 	{
-		// Save GPRs by using R instead of V
-		float NoV = saturate( dot( GBuffer.WorldNormal, ReflectionVector ) );
+		float3 N = GBuffer.WorldNormal;
+		float3 V = -CameraToPixel;
+		float3 R = 2 * dot( V, N ) * N - V;
+
+		float NoV = saturate( dot( N, V ) );
+
+		// Point lobe in off-specular peak direction
+		float a = Square( GBuffer.Roughness );
+		R = lerp( N, R, (1 - a) * ( sqrt(1 - a) + a ) );
 
 	#if 1//APPLY_SSR
 		float4 SSR = ScreenSpaceReflections.Load( int3(PixelPos, 0) );
@@ -595,13 +602,11 @@ void ReflectionEnvironmentTiledDeferredMain(
 		if( GBuffer.ShadingModelID == SHADINGMODELID_CLEAR_COAT )
 		{
 			const float ClearCoat = GBuffer.CustomData.x;
-			//Color = float4(0,0,0,1);
 			Color = lerp( Color, float4(0,0,0,1), ClearCoat );
 		}
 	
 		float AO = ScreenSpaceData.AmbientOcclusion;
-		//float SpecularOcclusion = saturate( Square( NoV + AO ) - 1 + AO );
-		float SpecularOcclusion = saturate( pow( NoV + AO, GBuffer.Roughness ) - 1 + AO );
+		float SpecularOcclusion = saturate( pow( NoV + AO, a ) - 1 + AO );
 
 		Color.a *= SpecularOcclusion;
 
@@ -614,7 +619,7 @@ void ReflectionEnvironmentTiledDeferredMain(
 		Color.a *= GBuffer.IndirectIrradiance;
 	#endif
 	
-		GatherRadiance( Color, WorldPosition, ReflectionVector, GBuffer.Roughness );
+		GatherRadiance( Color, WorldPosition, R, GBuffer.Roughness );
 
 		BRANCH
 		if( GBuffer.ShadingModelID == SHADINGMODELID_CLEAR_COAT )
@@ -650,7 +655,7 @@ void ReflectionEnvironmentTiledDeferredMain(
 			Color.a *= GBuffer.IndirectIrradiance;
 		#endif
 
-			GatherRadiance( Color, WorldPosition, ReflectionVector, ClearCoatRoughness );
+			GatherRadiance( Color, WorldPosition, R, ClearCoatRoughness );
 		}
 		else
 		{
diff --git a/Engine/Source/Runtime/Renderer/Private/SystemTextures.cpp b/Engine/Source/Runtime/Renderer/Private/SystemTextures.cpp
--- a/Engine/Source/Runtime/Renderer/Private/SystemTextures.cpp
+++ b/Engine/Source/Runtime/Renderer/Private/SystemTextures.cpp
@@ -253,16 +253,16 @@ void FSystemTextures::InitializeTextures(FRHICommandListImmediate& RHICmdList, E
 			// for testing, with 128x128 R8G8 we are very close to the reference (if lower res is needed we might have to add an offset to counter the 0.5f texel shift)
 			const bool bReference = false;
 
-			FPooledRenderTargetDesc Desc(FPooledRenderTargetDesc::Create2DDesc(FIntPoint(128, 32), PF_R8G8, TexCreate_HideInVisualizeTexture | TexCreate_FastVRAM, TexCreate_None, false));
+			FPooledRenderTargetDesc Desc(FPooledRenderTargetDesc::Create2DDesc(FIntPoint(128, 32), PF_R8G8, TexCreate_FastVRAM, TexCreate_None, false));
 
 			// for low roughness we would get banding with PF_R8G8 but for low spec it could be used, for now we don't do this optimization
 			Desc.Format = PF_G16R16;
+			//Desc.Format = PF_A16B16G16R16;
 
 			if (bReference)
 			{
 				Desc.Extent.X = 128;
 				Desc.Extent.Y = 128;
-				Desc.Format = PF_G16R16;
 			}
 
 			GRenderTargetPool.FindFreeElement(Desc, PreintegratedGF, TEXT("PreintegratedGF"));
@@ -276,13 +276,10 @@ void FSystemTextures::InitializeTextures(FRHICommandListImmediate& RHICmdList, E
 				float Roughness = (float)(y + 0.5f) / Desc.Extent.Y;
 				float m = Roughness * Roughness;
 				float m2 = m*m;
-				float k = m * 0.5f;
 
 				for (int32 x = 0; x < Desc.Extent.X; x++)
 				{
 					float NoV = (float)(x + 0.5f) / Desc.Extent.X;
-					float G_SchlickV = 1.0f / (NoV * (1.0f - k) + k);
-					float G_SmithV = 2.0f / (NoV + FMath::Sqrt((NoV*NoV) * (1.0f - m2) + m2));
 
 					FVector V;
 					V.X = FMath::Sqrt(1.0f - NoV * NoV);	// sin
@@ -291,6 +288,7 @@ void FSystemTextures::InitializeTextures(FRHICommandListImmediate& RHICmdList, E
 
 					float A = 0.0f;
 					float B = 0.0f;
+					float C = 0.0f;
 
 					const uint32 NumSamples = 128;
 					for (uint32 i = 0; i < NumSamples; i++)
@@ -298,41 +296,68 @@ void FSystemTextures::InitializeTextures(FRHICommandListImmediate& RHICmdList, E
 						float E1 = (float)i / NumSamples;
 						float E2 = (double)ReverseBits(i) / (double)0x100000000LL;
 
-						float Phi = 2.0f * PI * E1;
-						float CosPhi = FMath::Cos(Phi);
-						float SinPhi = FMath::Sin(Phi);
-						float CosTheta = FMath::Sqrt((1.0f - E2) / (1.0f + (m2 - 1.0f) * E2));
-						float SinTheta = FMath::Sqrt(1.0f - CosTheta * CosTheta);
-
-						FVector H(SinTheta * FMath::Cos(Phi), SinTheta * FMath::Sin(Phi), CosTheta);
-
-						FVector L = 2.0f * (V | H) * H - V;
-
-						float NoL = FMath::Max(L.Z, 0.0f);
-						float NoH = FMath::Max(H.Z, 0.0f);
-						float VoH = FMath::Max(V | H, 0.0f);
+						{
+							float Phi = 2.0f * PI * E1;
+							float CosPhi = FMath::Cos(Phi);
+							float SinPhi = FMath::Sin(Phi);
+							float CosTheta = FMath::Sqrt((1.0f - E2) / (1.0f + (m2 - 1.0f) * E2));
+							float SinTheta = FMath::Sqrt(1.0f - CosTheta * CosTheta);
+
+							FVector H(SinTheta * FMath::Cos(Phi), SinTheta * FMath::Sin(Phi), CosTheta);
+							FVector L = 2.0f * (V | H) * H - V;
+
+							float NoL = FMath::Max(L.Z, 0.0f);
+							float NoH = FMath::Max(H.Z, 0.0f);
+							float VoH = FMath::Max(V | H, 0.0f);
+
+							if (NoL > 0.0f)
+							{
+								float Vis_SmithV = NoL * ( NoV * ( 1 - m ) + m );
+								float Vis_SmithL = NoV * ( NoL * ( 1 - m ) + m );
+								float Vis = 0.5f / ( Vis_SmithV + Vis_SmithL );
+
+								float NoL_Vis_PDF = NoL * Vis * (4.0f * VoH / NoH);
+								float Fc = 1.0f - VoH;
+								Fc *= FMath::Square(Fc*Fc);
+								A += NoL_Vis_PDF * (1.0f - Fc);
+								B += NoL_Vis_PDF * Fc;
+							}
+						}
 
-						if (NoL > 0.0f)
 						{
-							float G_SchlickL = 1.0f / (NoL * (1.0f - k) + k);
-							float G_SmithL = 2.0f / (NoL + FMath::Sqrt((NoL*NoL) * (1.0f - m2) + m2));
-
-							float Vis_G = VoH * (NoL / NoH) * G_SchlickL;
-							float Fc = 1.0f - VoH;
-							Fc *= FMath::Square(Fc*Fc);
-							A += Vis_G * (1.0f - Fc);
-							B += Vis_G * Fc;
+							float Phi = 2.0f * PI * E1;
+							float CosPhi = FMath::Cos(Phi);
+							float SinPhi = FMath::Sin(Phi);
+							float CosTheta = FMath::Sqrt(E2);
+							float SinTheta = FMath::Sqrt(1.0f - CosTheta * CosTheta);
+
+							FVector L(SinTheta * FMath::Cos(Phi), SinTheta * FMath::Sin(Phi), CosTheta);
+							FVector H = ( V + L ).UnsafeNormal();
+
+							float NoL = FMath::Max(L.Z, 0.0f);
+							float NoH = FMath::Max(H.Z, 0.0f);
+							float VoH = FMath::Max(V | H, 0.0f);
+
+							float FD90 = 0.5f + 2.0f * VoH * VoH * Roughness;
+							float FdV = 1.0f + (FD90 - 1.0f) * pow( 1.0f - NoV, 5 );
+							float FdL = 1.0f + (FD90 - 1.0f) * pow( 1.0f - NoL, 5 );
+							C += FdV * FdL;// * ( 1.0f - 0.3333f * Roughness );
 						}
 					}
 					A /= NumSamples;
 					B /= NumSamples;
+					C /= NumSamples;
 
-					A *= G_SchlickV;
-					B *= G_SchlickV;
-
-					if (Desc.Format == PF_G16R16)
+					if( Desc.Format == PF_A16B16G16R16 )
+					{
+						uint16* Dest = (uint16*)(DestBuffer + x * 8 + y * DestStride);
+						Dest[0] = (int32)(FMath::Clamp(A, 0.0f, 1.0f) * 65535.0f + 0.5f);
+						Dest[1] = (int32)(FMath::Clamp(B, 0.0f, 1.0f) * 65535.0f + 0.5f);
+						Dest[2] = (int32)(FMath::Clamp(C, 0.0f, 1.0f) * 65535.0f + 0.5f);
+					}
+					else if (Desc.Format == PF_G16R16)
 					{
-						uint16* Dest = (uint16*)(DestBuffer + x * sizeof(uint32)+y * DestStride);
+						uint16* Dest = (uint16*)(DestBuffer + x * 4 + y * DestStride);
 						Dest[0] = (int32)(FMath::Clamp(A, 0.0f, 1.0f) * 65535.0f + 0.5f);
 						Dest[1] = (int32)(FMath::Clamp(B, 0.0f, 1.0f) * 65535.0f + 0.5f);
 					}
@@ -340,7 +365,7 @@ void FSystemTextures::InitializeTextures(FRHICommandListImmediate& RHICmdList, E
 					{
 						check(Desc.Format == PF_R8G8);
 
-						uint8* Dest = (uint8*)(DestBuffer + x * sizeof(uint16)+y * DestStride);
+						uint8* Dest = (uint8*)(DestBuffer + x * 2 + y * DestStride);
 						Dest[0] = (int32)(FMath::Clamp(A, 0.0f, 1.0f) * 255.9999f);
 						Dest[1] = (int32)(FMath::Clamp(B, 0.0f, 1.0f) * 255.9999f);
 					}
