Fixed responsive AA.

diff --git a/Engine/Source/Runtime/Renderer/Private/DistortionRendering.cpp b/Engine/Source/Runtime/Renderer/Private/DistortionRendering.cpp
--- a/Engine/Source/Runtime/Renderer/Private/DistortionRendering.cpp
+++ b/Engine/Source/Runtime/Renderer/Private/DistortionRendering.cpp
@@ -824,8 +824,8 @@ void FSceneRenderer::RenderDistortion(FRHICommandListImmediate& RHICmdList)
 				// viewport to match view size
 				RHICmdList.SetViewport(View.ViewRect.Min.X, View.ViewRect.Min.Y, 0.0f, View.ViewRect.Max.X, View.ViewRect.Max.Y, 1.0f);
 
-				// clear offsets to 0, stencil to 0
-				RHICmdList.Clear(true, FLinearColor(0, 0, 0, 0), false, 0, true, 0, FIntRect());
+				// clear offsets to 0
+				RHICmdList.Clear(true, FLinearColor(0, 0, 0, 0), false, 0, false, 0, FIntRect());
 
 				// enable depth test but disable depth writes
 				// Note, this is a reversed Z depth surface, using CF_GreaterEqual.
