Resolve rect support in the custom HDR MSAA resolve for mobile emulation.

diff --git a/Engine/Source/Runtime/Renderer/Private/PostProcess/SceneRenderTargets.cpp b/Engine/Source/Runtime/Renderer/Private/PostProcess/SceneRenderTargets.cpp
--- a/Engine/Source/Runtime/Renderer/Private/PostProcess/SceneRenderTargets.cpp
+++ b/Engine/Source/Runtime/Renderer/Private/PostProcess/SceneRenderTargets.cpp
@@ -651,10 +651,17 @@ void FSceneRenderTargets::ResolveSceneColor(FRHICommandList& RHICmdList, const F
 	else 
 	{
 		// Custom shader based color resolve for HDR color to emulate mobile.
-		// TODO: Currently this resolves the entire surface, need to respect ResolveRect.
-		
 		SetRenderTarget(RHICmdList, GetSceneColorTexture(), FTextureRHIParamRef());
 		
+		if(ResolveRect.IsValid())
+		{
+			RHICmdList.SetScissorRect(true, ResolveRect.X1, ResolveRect.Y1, ResolveRect.X2, ResolveRect.Y2);
+		}
+		else
+		{
+			RHICmdList.SetScissorRect(false, 0, 0, 0, 0);
+		}
+
 		RHICmdList.SetBlendState(TStaticBlendState<>::GetRHI());
 		RHICmdList.SetRasterizerState(TStaticRasterizerState<>::GetRHI());
 		RHICmdList.SetDepthStencilState(TStaticDepthStencilState<false, CF_Always>::GetRHI());
