* nicer profile info for WeightedSampleSum shader

diff --git a/Engine/Source/Runtime/Renderer/Private/PostProcess/PostProcessWeightedSampleSum.cpp b/Engine/Source/Runtime/Renderer/Private/PostProcess/PostProcessWeightedSampleSum.cpp
--- a/Engine/Source/Runtime/Renderer/Private/PostProcess/PostProcessWeightedSampleSum.cpp
+++ b/Engine/Source/Runtime/Renderer/Private/PostProcess/PostProcessWeightedSampleSum.cpp
@@ -338,8 +338,6 @@ void FRCPassPostProcessWeightedSampleSum::Process(FRenderingCompositePassContext
 	// input is not hooked up correctly
 	check(InputDesc);
 
-	SCOPED_DRAW_EVENT(PostProcessWeightedSampleSum, DEC_SCENE_ITEMS);
-	
 	const FSceneViewFamily& ViewFamily = *(View.Family);
 
 	FIntPoint SrcSize = InputDesc->Extent;
@@ -373,6 +371,8 @@ void FRCPassPostProcessWeightedSampleSum::Process(FRenderingCompositePassContext
 
 	uint32 NumSamples = Compute1DGaussianFilterKernel(FeatureLevel, EffectiveBlurRadius, OffsetAndWeight, MaxNumSamples);
 
+	SCOPED_DRAW_EVENTF(PostProcessWeightedSampleSum, DEC_SCENE_ITEMS, TEXT("PostProcessWeightedSampleSum#%d"), NumSamples);
+
 	// compute weights as weighted contributions of the TintValue
 	for(uint32 i = 0; i < NumSamples; ++i)
 	{
