Soul : Fixed ShouldCache for ES2 material

diff --git a/Engine/Source/Runtime/Renderer/Private/BasePassRendering.h b/Engine/Source/Runtime/Renderer/Private/BasePassRendering.h
--- a/Engine/Source/Runtime/Renderer/Private/BasePassRendering.h
+++ b/Engine/Source/Runtime/Renderer/Private/BasePassRendering.h
@@ -103,6 +103,7 @@ public:
 	{
 		bool bShouldCache = TBasePassVertexShaderBaseType<LightMapPolicyType>::ShouldCache(Platform, Material, VertexFactoryType);
 		return bShouldCache 
+			&& IsFeatureLevelSupported(Platform, ERHIFeatureLevel::SM3)
 			&& (!bEnableAtmosphericFog || IsFeatureLevelSupported(Platform, ERHIFeatureLevel::SM4));
 	}
 
@@ -426,6 +427,7 @@ public:
 		const bool bCacheShaders = !bEnableSkyLight || (Material->GetShadingModel() != MSM_Unlit);
 
 		return bCacheShaders
+			&& IsFeatureLevelSupported(Platform, ERHIFeatureLevel::SM3)
 			&& TBasePassPixelShaderBaseType<LightMapPolicyType>::ShouldCache(Platform, Material, VertexFactoryType);
 	}
 
