Fix crash when opening Persona
 - Only add TexCreate_TargetArraySlicesIndependently flag to reflection scratch cubemaps if we don't support GS addressing

diff --git a/Engine/Source/Runtime/Renderer/Private/PostProcess/SceneRenderTargets.cpp b/Engine/Source/Runtime/Renderer/Private/PostProcess/SceneRenderTargets.cpp
--- a/Engine/Source/Runtime/Renderer/Private/PostProcess/SceneRenderTargets.cpp
+++ b/Engine/Source/Runtime/Renderer/Private/PostProcess/SceneRenderTargets.cpp
@@ -1114,7 +1114,11 @@ void FSceneRenderTargets::AllocateReflectionTargets()
 		const int32 NumReflectionCaptureMips = FMath::CeilLogTwo(GReflectionCaptureSize) + 1;
 
 		{
-			uint32 TexFlags = TexCreate_TargetArraySlicesIndependently;
+			uint32 TexFlags = 0;
+			if (!GSupportsGSRenderTargetLayerSwitchingToMips)
+			{
+				TexFlags = TexCreate_TargetArraySlicesIndependently;
+			}
 
 			// Create scratch cubemaps for filtering passes
 			FPooledRenderTargetDesc Desc2(FPooledRenderTargetDesc::CreateCubemapDesc(GReflectionCaptureSize, PF_FloatRGBA, TexFlags, TexCreate_RenderTargetable, false, 1, NumReflectionCaptureMips));
