Fix TTP: 344871  Bad tex flags due to a partial integration.

diff --git a/Engine/Source/Runtime/Renderer/Private/PostProcess/SceneRenderTargets.cpp b/Engine/Source/Runtime/Renderer/Private/PostProcess/SceneRenderTargets.cpp
--- a/Engine/Source/Runtime/Renderer/Private/PostProcess/SceneRenderTargets.cpp
+++ b/Engine/Source/Runtime/Renderer/Private/PostProcess/SceneRenderTargets.cpp
@@ -1112,16 +1112,16 @@ void FSceneRenderTargets::AllocateReflectionTargets()
 	{
 		extern ENGINE_API int32 GReflectionCaptureSize;
 		const int32 NumReflectionCaptureMips = FMath::CeilLogTwo(GReflectionCaptureSize) + 1;
-
+		
+		uint32 CubeTexFlags = 0;
+		if (!GSupportsGSRenderTargetLayerSwitchingToMips)
 		{
-			uint32 TexFlags = 0;
-			if (!GSupportsGSRenderTargetLayerSwitchingToMips)
-			{
-				TexFlags = TexCreate_TargetArraySlicesIndependently;
-			}
+			CubeTexFlags = TexCreate_TargetArraySlicesIndependently;
+		}
 
+		{
 			// Create scratch cubemaps for filtering passes
-			FPooledRenderTargetDesc Desc2(FPooledRenderTargetDesc::CreateCubemapDesc(GReflectionCaptureSize, PF_FloatRGBA, TexFlags, TexCreate_RenderTargetable, false, 1, NumReflectionCaptureMips));
+			FPooledRenderTargetDesc Desc2(FPooledRenderTargetDesc::CreateCubemapDesc(GReflectionCaptureSize, PF_FloatRGBA, CubeTexFlags, TexCreate_RenderTargetable, false, 1, NumReflectionCaptureMips));
 			GRenderTargetPool.FindFreeElement(Desc2, ReflectionColorScratchCubemap[0], TEXT("ReflectionColorScratchCubemap0"));
 			GRenderTargetPool.FindFreeElement(Desc2, ReflectionColorScratchCubemap[1], TEXT("ReflectionColorScratchCubemap1"));
 
@@ -1134,7 +1134,7 @@ void FSceneRenderTargets::AllocateReflectionTargets()
 		const int32 NumDiffuseIrradianceMips = FMath::CeilLogTwo(GDiffuseIrradianceCubemapSize) + 1;
 
 		{
-			FPooledRenderTargetDesc Desc2(FPooledRenderTargetDesc::CreateCubemapDesc(GDiffuseIrradianceCubemapSize, PF_FloatRGBA, TexCreate_TargetArraySlicesIndependently, TexCreate_RenderTargetable, false, 1, NumDiffuseIrradianceMips));
+			FPooledRenderTargetDesc Desc2(FPooledRenderTargetDesc::CreateCubemapDesc(GDiffuseIrradianceCubemapSize, PF_FloatRGBA, CubeTexFlags, TexCreate_RenderTargetable, false, 1, NumDiffuseIrradianceMips));
 			GRenderTargetPool.FindFreeElement(Desc2, DiffuseIrradianceScratchCubemap[0], TEXT("DiffuseIrradianceScratchCubemap0"));
 			GRenderTargetPool.FindFreeElement(Desc2, DiffuseIrradianceScratchCubemap[1], TEXT("DiffuseIrradianceScratchCubemap1"));
 		}
