Added check to prevent decals without scale from being rendered.

diff --git a/Engine/Source/Runtime/Renderer/Private/CompositionLighting/PostProcessDeferredDecals.cpp b/Engine/Source/Runtime/Renderer/Private/CompositionLighting/PostProcessDeferredDecals.cpp
--- a/Engine/Source/Runtime/Renderer/Private/CompositionLighting/PostProcessDeferredDecals.cpp
+++ b/Engine/Source/Runtime/Renderer/Private/CompositionLighting/PostProcessDeferredDecals.cpp
@@ -793,7 +793,7 @@ void FRCPassPostProcessDeferredDecals::Process(FRenderingCompositePassContext& C
 			ComponentToWorldMatrix.GetScaledAxis( EAxis::Z ).SizeSquared() * FMath::Square(GDefaultDecalSize.Z));
 
 		// can be optimized as the test is too conservative (sphere instead of OBB)
-		if(!View.ViewFrustum.IntersectSphere(ComponentToWorldMatrix.GetOrigin(), ConservativeRadius))
+		if(ConservativeRadius < SMALL_NUMBER || !View.ViewFrustum.IntersectSphere(ComponentToWorldMatrix.GetOrigin(), ConservativeRadius))
 		{
 			bIsShown = false;
 		}
