Update component's LastRenderTime when rendering shadows.

#ttp 341917 -  LIVE: RENDERING: Shadows from hidden animated actors freeze after about 1 second

diff --git a/Engine/Source/Runtime/Renderer/Private/ShadowSetup.cpp b/Engine/Source/Runtime/Renderer/Private/ShadowSetup.cpp
--- a/Engine/Source/Runtime/Renderer/Private/ShadowSetup.cpp
+++ b/Engine/Source/Runtime/Renderer/Private/ShadowSetup.cpp
@@ -699,6 +699,13 @@ void FProjectedShadowInfo::AddSubjectPrimitive(FPrimitiveSceneInfo* PrimitiveSce
 			bShadowRelevance |= ViewRelevance.bShadowRelevance;
 		}
 
+		if (bShadowRelevance)
+		{
+			// Update the primitive component's last render time. Allows the component to update when using bCastWhenHidden.
+			const float CurrentWorldTime = Views[0]->Family->CurrentWorldTime;
+			*(PrimitiveSceneInfo->ComponentLastRenderTime) = CurrentWorldTime;
+		}
+
 		if (bNeedsPreRenderView)
 		{
 			// Call PreRenderView on primitives that weren't visible in any of the main views, but need to be rendered in this shadow's depth pass
