Continued in-editor mobile preview work
 - Removed default shader platform parameter from GetGlobalShaderMap()
 - Added an inline overload of GetGlobalShaderMap() that takes a feature level instead, and translates to the correct shader platform
 - Cached feature level on FSceneView for faster access
 - Cached shader map on FViewInfo for faster access
 - Cached featurelevel/shadermap on rendering composition graph execution context, for faster access

Fixed a couple of crashes when switching feature level dynamically
 - Needed to allow certain shader permutations to be generated if feature level was ES2, but still on a PC platform. (i.e. hitproxy shaders, and basepass shaders used by editor primitives system)

diff --git a/Engine/Source/Editor/UnrealEd/Private/CubemapUnwrapUtils.cpp b/Engine/Source/Editor/UnrealEd/Private/CubemapUnwrapUtils.cpp
--- a/Engine/Source/Editor/UnrealEd/Private/CubemapUnwrapUtils.cpp
+++ b/Engine/Source/Editor/UnrealEd/Private/CubemapUnwrapUtils.cpp
@@ -139,8 +139,8 @@ void FMipLevelBatchedElementParameters::BindShaders(FRHICommandList& RHICmdList,
 template<typename TPixelShader>
 void FMipLevelBatchedElementParameters::BindShaders(FRHICommandList& RHICmdList, ERHIFeatureLevel::Type InFeatureLevel, const FMatrix& InTransform, const float InGamma, const FMatrix& ColorWeights, const FTexture* Texture)
 {
-	TShaderMapRef<FCubemapTexturePropertiesVS> VertexShader(GetGlobalShaderMap());
-	TShaderMapRef<TPixelShader> PixelShader(GetGlobalShaderMap());
+	TShaderMapRef<FCubemapTexturePropertiesVS> VertexShader(GetGlobalShaderMap(InFeatureLevel));
+	TShaderMapRef<TPixelShader> PixelShader(GetGlobalShaderMap(InFeatureLevel));
 	
 
 	static FGlobalBoundShaderState BoundShaderState;
@@ -163,8 +163,8 @@ void FIESLightProfilePS::SetParameters( FRHICommandList& RHICmdList, const FText
 
 void FIESLightProfileBatchedElementParameters::BindShaders( FRHICommandList& RHICmdList, ERHIFeatureLevel::Type InFeatureLevel, const FMatrix& InTransform, const float InGamma, const FMatrix& ColorWeights, const FTexture* Texture )
 {
-	TShaderMapRef<FSimpleElementVS> VertexShader(GetGlobalShaderMap());
-	TShaderMapRef<FIESLightProfilePS> PixelShader(GetGlobalShaderMap());
+	TShaderMapRef<FSimpleElementVS> VertexShader(GetGlobalShaderMap(InFeatureLevel));
+	TShaderMapRef<FIESLightProfilePS> PixelShader(GetGlobalShaderMap(InFeatureLevel));
 
 	
 	static FGlobalBoundShaderState BoundShaderState;
diff --git a/Engine/Source/Editor/UnrealEd/Private/MeshPaintRendering.cpp b/Engine/Source/Editor/UnrealEd/Private/MeshPaintRendering.cpp
--- a/Engine/Source/Editor/UnrealEd/Private/MeshPaintRendering.cpp
+++ b/Engine/Source/Editor/UnrealEd/Private/MeshPaintRendering.cpp
@@ -352,8 +352,8 @@ namespace MeshPaintRendering
 										   const float InGamma,
 										   const FMeshPaintShaderParameters& InShaderParams )
 	{
-		TShaderMapRef< TMeshPaintVertexShader > VertexShader( GetGlobalShaderMap() );
-		TShaderMapRef< TMeshPaintPixelShader > PixelShader( GetGlobalShaderMap() );
+		TShaderMapRef< TMeshPaintVertexShader > VertexShader(GetGlobalShaderMap(InFeatureLevel));
+		TShaderMapRef< TMeshPaintPixelShader > PixelShader(GetGlobalShaderMap(InFeatureLevel));
 
 		
 		static FGlobalBoundShaderState BoundShaderState;
@@ -373,8 +373,8 @@ namespace MeshPaintRendering
 												 const float InGamma,
 												 const FMeshPaintDilateShaderParameters& InShaderParams )
 	{
-		TShaderMapRef< TMeshPaintDilateVertexShader > VertexShader( GetGlobalShaderMap() );
-		TShaderMapRef< TMeshPaintDilatePixelShader > PixelShader( GetGlobalShaderMap() );
+		TShaderMapRef< TMeshPaintDilateVertexShader > VertexShader(GetGlobalShaderMap(InFeatureLevel));
+		TShaderMapRef< TMeshPaintDilatePixelShader > PixelShader(GetGlobalShaderMap(InFeatureLevel));
 
 		
 		static FGlobalBoundShaderState BoundShaderState;
diff --git a/Engine/Source/Editor/UnrealEd/Private/NormalMapPreview.cpp b/Engine/Source/Editor/UnrealEd/Private/NormalMapPreview.cpp
--- a/Engine/Source/Editor/UnrealEd/Private/NormalMapPreview.cpp
+++ b/Engine/Source/Editor/UnrealEd/Private/NormalMapPreview.cpp
@@ -82,8 +82,8 @@ void FNormalMapBatchedElementParameters::BindShaders(
 	const FMatrix& ColorWeights,
 	const FTexture* Texture)
 {
-	TShaderMapRef<FSimpleElementVS> VertexShader(GetGlobalShaderMap());
-	TShaderMapRef<FSimpleElementNormalMapPS> PixelShader(GetGlobalShaderMap());
+	TShaderMapRef<FSimpleElementVS> VertexShader(GetGlobalShaderMap(InFeatureLevel));
+	TShaderMapRef<FSimpleElementNormalMapPS> PixelShader(GetGlobalShaderMap(InFeatureLevel));
 
 	
 	static FGlobalBoundShaderState BoundShaderState;
diff --git a/Engine/Source/Editor/UnrealEd/Private/Texture2DPreview.cpp b/Engine/Source/Editor/UnrealEd/Private/Texture2DPreview.cpp
--- a/Engine/Source/Editor/UnrealEd/Private/Texture2DPreview.cpp
+++ b/Engine/Source/Editor/UnrealEd/Private/Texture2DPreview.cpp
@@ -84,8 +84,8 @@ void FBatchedElementTexture2DPreviewParameters::BindShaders(
 	const FMatrix& ColorWeights,
 	const FTexture* Texture)
 {
-	TShaderMapRef<FSimpleElementVS> VertexShader(GetGlobalShaderMap());
-	TShaderMapRef<FSimpleElementTexture2DPreviewPS> PixelShader(GetGlobalShaderMap());
+	TShaderMapRef<FSimpleElementVS> VertexShader(GetGlobalShaderMap(InFeatureLevel));
+	TShaderMapRef<FSimpleElementTexture2DPreviewPS> PixelShader(GetGlobalShaderMap(InFeatureLevel));
 
 	
 	static FGlobalBoundShaderState BoundShaderState;
diff --git a/Engine/Source/Runtime/Engine/Private/BatchedElements.cpp b/Engine/Source/Runtime/Engine/Private/BatchedElements.cpp
--- a/Engine/Source/Runtime/Engine/Private/BatchedElements.cpp
+++ b/Engine/Source/Runtime/Engine/Private/BatchedElements.cpp
@@ -456,11 +456,11 @@ void FBatchedElements::PrepareShaders(
 	}
 	else
 	{
-		TShaderMapRef<FSimpleElementVS> VertexShader(GetGlobalShaderMap());
+		TShaderMapRef<FSimpleElementVS> VertexShader(GetGlobalShaderMap(FeatureLevel));
 		    
 	    if (bHitTesting)
 	    {
-		    TShaderMapRef<FSimpleElementHitProxyPS> HitTestingPixelShader(GetGlobalShaderMap());
+			TShaderMapRef<FSimpleElementHitProxyPS> HitTestingPixelShader(GetGlobalShaderMap(FeatureLevel));
 			SetGlobalBoundShaderState(RHICmdList, FeatureLevel, HitTestingBoundShaderState, GSimpleElementVertexDeclaration.VertexDeclarationRHI,
 				*VertexShader, *HitTestingPixelShader);
 
@@ -476,7 +476,7 @@ void FBatchedElements::PrepareShaders(
 			    // also don't support alpha testing to floating point render targets
 				RHICmdList.SetBlendState(TStaticBlendState<>::GetRHI());
 
-			    TShaderMapRef<FSimpleElementMaskedGammaPS> MaskedPixelShader(GetGlobalShaderMap());
+				TShaderMapRef<FSimpleElementMaskedGammaPS> MaskedPixelShader(GetGlobalShaderMap(FeatureLevel));
 				SetGlobalBoundShaderState(RHICmdList, FeatureLevel, MaskedBoundShaderState, GSimpleElementVertexDeclaration.VertexDeclarationRHI,
 					*VertexShader, *MaskedPixelShader);
 
@@ -504,7 +504,7 @@ void FBatchedElements::PrepareShaders(
 					RHICmdList.SetBlendState(TStaticBlendState<>::GetRHI());
 			    }
 			    
-			    TShaderMapRef<FSimpleElementDistanceFieldGammaPS> DistanceFieldPixelShader(GetGlobalShaderMap());			
+				TShaderMapRef<FSimpleElementDistanceFieldGammaPS> DistanceFieldPixelShader(GetGlobalShaderMap(FeatureLevel));
 				SetGlobalBoundShaderState(RHICmdList, FeatureLevel, DistanceFieldBoundShaderState, GSimpleElementVertexDeclaration.VertexDeclarationRHI,
 					*VertexShader, *DistanceFieldPixelShader );			
 
@@ -534,7 +534,7 @@ void FBatchedElements::PrepareShaders(
 		    }
 			else if(BlendMode >= SE_BLEND_RGBA_MASK_START && BlendMode <= SE_BLEND_RGBA_MASK_END)
 			{
-				TShaderMapRef<FSimpleElementColorChannelMaskPS> ColorChannelMaskPixelShader(GetGlobalShaderMap());
+				TShaderMapRef<FSimpleElementColorChannelMaskPS> ColorChannelMaskPixelShader(GetGlobalShaderMap(FeatureLevel));
 				SetGlobalBoundShaderState(RHICmdList, FeatureLevel, ColorChannelMaskShaderState, GSimpleElementVertexDeclaration.VertexDeclarationRHI,
 					*VertexShader, *ColorChannelMaskPixelShader );
 			
@@ -546,7 +546,7 @@ void FBatchedElements::PrepareShaders(
     
 			    if (FMath::Abs(Gamma - 1.0f) < KINDA_SMALL_NUMBER)
 			    {
-				    TShaderMapRef<FSimpleElementPS> RegularPixelShader(GetGlobalShaderMap());
+					TShaderMapRef<FSimpleElementPS> RegularPixelShader(GetGlobalShaderMap(FeatureLevel));
 					SetGlobalBoundShaderState(RHICmdList, FeatureLevel, SimpleBoundShaderState, GSimpleElementVertexDeclaration.VertexDeclarationRHI,
 						*VertexShader, *RegularPixelShader );
 
@@ -555,7 +555,7 @@ void FBatchedElements::PrepareShaders(
 			    }
 			    else
 			    {
-				    TShaderMapRef<FSimpleElementGammaPS> RegularPixelShader(GetGlobalShaderMap());
+					TShaderMapRef<FSimpleElementGammaPS> RegularPixelShader(GetGlobalShaderMap(FeatureLevel));
 					SetGlobalBoundShaderState(RHICmdList, FeatureLevel, RegularBoundShaderState, GSimpleElementVertexDeclaration.VertexDeclarationRHI,
 						*VertexShader, *RegularPixelShader );
 
diff --git a/Engine/Source/Runtime/Engine/Private/GPUSkinCache.cpp b/Engine/Source/Runtime/Engine/Private/GPUSkinCache.cpp
--- a/Engine/Source/Runtime/Engine/Private/GPUSkinCache.cpp
+++ b/Engine/Source/Runtime/Engine/Private/GPUSkinCache.cpp
@@ -469,7 +469,10 @@ FGPUSkinCache::FElementCacheStatusInfo* FGPUSkinCache::FindEvictableCacheStatusI
 
 void FGPUSkinCache::DispatchSkinCacheProcess(FRHICommandListImmediate& RHICmdList, uint32 InputStreamFloatOffset, uint32 OutputBufferFloatOffset, FBoneBufferTypeRef BoneBuffer, FUniformBufferRHIRef UniformBuffer, const FVertexBufferInfo* VBInfo, uint32 VertexStride, uint32 VertexCount, const FVector& MeshOrigin, const FVector& MeshExtension, bool bUseExtraBoneInfluences)
 {
-	if (GRHIFeatureLevel < ERHIFeatureLevel::SM5)
+	// MOBILEPREVIEWTODO: Proper value for this
+	const auto FeatureLevel = GRHIFeatureLevel;
+
+	if (FeatureLevel < ERHIFeatureLevel::SM5)
 	{
 		return;
 	}
@@ -480,7 +483,7 @@ void FGPUSkinCache::DispatchSkinCacheProcess(FRHICommandListImmediate& RHICmdLis
 
 	if (bUseExtraBoneInfluences)
 	{
-		TShaderMapRef<FGPUSkinCacheCS<true> > SkinCacheCS(GetGlobalShaderMap());
+		TShaderMapRef<FGPUSkinCacheCS<true> > SkinCacheCS(GetGlobalShaderMap(FeatureLevel));
 		SkinCacheCS->SetParameters(RHICmdList, VertexStride, InputStreamFloatOffset, VertexCount, OutputBufferFloatOffset, BoneBuffer, UniformBuffer, VBInfo->VertexBufferSRV, SkinCacheBufferRW, MeshOrigin, MeshExtension);
 
 		RHICmdList.SetComputeShader(SkinCacheCS->GetComputeShader());
@@ -489,7 +492,7 @@ void FGPUSkinCache::DispatchSkinCacheProcess(FRHICommandListImmediate& RHICmdLis
 	}
 	else
 	{
-		TShaderMapRef<FGPUSkinCacheCS<false> > SkinCacheCS(GetGlobalShaderMap());
+		TShaderMapRef<FGPUSkinCacheCS<false> > SkinCacheCS(GetGlobalShaderMap(FeatureLevel));
 		SkinCacheCS->SetParameters(RHICmdList, VertexStride, InputStreamFloatOffset, VertexCount, OutputBufferFloatOffset, BoneBuffer, UniformBuffer, VBInfo->VertexBufferSRV, SkinCacheBufferRW, MeshOrigin, MeshExtension);
 
 		RHICmdList.SetComputeShader(SkinCacheCS->GetComputeShader());
diff --git a/Engine/Source/Runtime/Engine/Private/GPUSort.cpp b/Engine/Source/Runtime/Engine/Private/GPUSort.cpp
--- a/Engine/Source/Runtime/Engine/Private/GPUSort.cpp
+++ b/Engine/Source/Runtime/Engine/Private/GPUSort.cpp
@@ -690,7 +690,8 @@ int32 SortGPUBuffers(FRHICommandListImmediate& RHICmdList, FGPUSortBuffers SortB
 	const bool bDebugOffsets = CVarDebugOffsets.GetValueOnRenderThread() != 0;
 	const bool bDebugSort = CVarDebugSort.GetValueOnRenderThread() != 0;
 
-	check(GRHIFeatureLevel == ERHIFeatureLevel::SM5);
+	const auto FeatureLevel = GRHIFeatureLevel;
+	check(FeatureLevel == ERHIFeatureLevel::SM5);
 
 	SCOPED_DRAW_EVENTF(SortGPU, DEC_PARTICLE, TEXT("SortGPU_%d"), Count);
 
@@ -729,10 +730,11 @@ int32 SortGPUBuffers(FRHICommandListImmediate& RHICmdList, FGPUSortBuffers SortB
 	SortParameters.GroupCount = GroupCount;
 
 	// Grab shaders.
-	TShaderMapRef<FRadixSortClearOffsetsCS> ClearOffsetsCS(GetGlobalShaderMap());
-	TShaderMapRef<FRadixSortUpsweepCS> UpsweepCS(GetGlobalShaderMap());
-	TShaderMapRef<FRadixSortSpineCS> SpineCS(GetGlobalShaderMap());
-	TShaderMapRef<FRadixSortDownsweepCS> DownsweepCS(GetGlobalShaderMap());
+	auto ShaderMap = GetGlobalShaderMap(FeatureLevel);
+	TShaderMapRef<FRadixSortClearOffsetsCS> ClearOffsetsCS(ShaderMap);
+	TShaderMapRef<FRadixSortUpsweepCS> UpsweepCS(ShaderMap);
+	TShaderMapRef<FRadixSortSpineCS> SpineCS(ShaderMap);
+	TShaderMapRef<FRadixSortDownsweepCS> DownsweepCS(ShaderMap);
 
 	// Constant buffer workaround. Both shaders must use either the constant buffer or vertex buffer.
 	check( UpsweepCS->RequiresConstantBufferWorkaround() == DownsweepCS->RequiresConstantBufferWorkaround() );
diff --git a/Engine/Source/Runtime/Engine/Private/GlobalShader.cpp b/Engine/Source/Runtime/Engine/Private/GlobalShader.cpp
--- a/Engine/Source/Runtime/Engine/Private/GlobalShader.cpp
+++ b/Engine/Source/Runtime/Engine/Private/GlobalShader.cpp
@@ -478,8 +478,8 @@ void RecompileGlobalShaders()
 		// Flush pending accesses to the existing global shaders.
 		FlushRenderingCommands();
 
+		// MOBILEPREVIEWTODO: Use global feature level bitfield to recompile for all active platforms
 		GetGlobalShaderMap(GRHIShaderPlatform)->Empty();
-
 		VerifyGlobalShaders(GRHIShaderPlatform, false);
 
 		GShaderCompilingManager->ProcessAsyncResults(false, true);
diff --git a/Engine/Source/Runtime/Engine/Private/Particles/ParticleCurveTexture.cpp b/Engine/Source/Runtime/Engine/Private/Particles/ParticleCurveTexture.cpp
--- a/Engine/Source/Runtime/Engine/Private/Particles/ParticleCurveTexture.cpp
+++ b/Engine/Source/Runtime/Engine/Private/Particles/ParticleCurveTexture.cpp
@@ -222,8 +222,8 @@ static void InjectCurves(
 			(float)Curve.TexelAllocation.Y / GParticleCurveTextureSizeY );
 
 		// Grab shaders.
-		TShaderMapRef<FParticleCurveInjectionVS> VertexShader( GetGlobalShaderMap() );
-		TShaderMapRef<FParticleCurveInjectionPS> PixelShader( GetGlobalShaderMap() );
+		TShaderMapRef<FParticleCurveInjectionVS> VertexShader(GetGlobalShaderMap(GMaxRHIFeatureLevel));
+		TShaderMapRef<FParticleCurveInjectionPS> PixelShader(GetGlobalShaderMap(GMaxRHIFeatureLevel));
 
 		// Bound shader state.
 		
diff --git a/Engine/Source/Runtime/Engine/Private/Particles/ParticleGpuSimulation.cpp b/Engine/Source/Runtime/Engine/Private/Particles/ParticleGpuSimulation.cpp
--- a/Engine/Source/Runtime/Engine/Private/Particles/ParticleGpuSimulation.cpp
+++ b/Engine/Source/Runtime/Engine/Private/Particles/ParticleGpuSimulation.cpp
@@ -1236,8 +1236,8 @@ void ExecuteSimulationCommands(
 	)
 {
 	// Grab shaders.
-	TShaderMapRef<FParticleTileVS> VertexShader(GetGlobalShaderMap());
-	TShaderMapRef<TParticleSimulationPS<bUseDepthBufferCollision> > PixelShader(GetGlobalShaderMap());
+	TShaderMapRef<FParticleTileVS> VertexShader(GetGlobalShaderMap(FeatureLevel));
+	TShaderMapRef<TParticleSimulationPS<bUseDepthBufferCollision> > PixelShader(GetGlobalShaderMap(FeatureLevel));
 
 	// Bound shader state.
 	
@@ -1293,8 +1293,8 @@ void ClearTiles(FRHICommandList& RHICmdList, ERHIFeatureLevel::Type FeatureLevel
 	int32 FirstTile = 0;
 
 	// Grab shaders.
-	TShaderMapRef<FParticleTileVS> VertexShader( GetGlobalShaderMap() );
-	TShaderMapRef<FParticleSimulationClearPS> PixelShader( GetGlobalShaderMap() );
+	TShaderMapRef<FParticleTileVS> VertexShader(GetGlobalShaderMap(FeatureLevel));
+	TShaderMapRef<FParticleSimulationClearPS> PixelShader(GetGlobalShaderMap(FeatureLevel));
 	
 	// Bound shader state.
 	
@@ -1539,8 +1539,8 @@ void InjectNewParticles(FRHICommandList& RHICmdList, ERHIFeatureLevel::Type Feat
 		FirstParticle += ParticlesThisDrawCall;
 
 		// Grab shaders.
-		TShaderMapRef<FParticleInjectionVS> VertexShader( GetGlobalShaderMap() );
-		TShaderMapRef<FParticleInjectionPS> PixelShader( GetGlobalShaderMap() );
+		TShaderMapRef<FParticleInjectionVS> VertexShader(GetGlobalShaderMap(FeatureLevel));
+		TShaderMapRef<FParticleInjectionPS> PixelShader(GetGlobalShaderMap(FeatureLevel));
 
 		// Bound shader state.
 		
@@ -1762,8 +1762,8 @@ static void VisualizeGPUSimulation(
 	RHICmdList.SetBlendState(TStaticBlendState<>::GetRHI());
 
 	// Grab shaders.
-	TShaderMapRef<FParticleSimVisualizeVS> VertexShader( GetGlobalShaderMap() );
-	TShaderMapRef<FParticleSimVisualizePS> PixelShader( GetGlobalShaderMap() );
+	TShaderMapRef<FParticleSimVisualizeVS> VertexShader(GetGlobalShaderMap(FeatureLevel));
+	TShaderMapRef<FParticleSimVisualizePS> PixelShader(GetGlobalShaderMap(FeatureLevel));
 
 	// Bound shader state.
 	
@@ -2012,7 +2012,10 @@ static FBox ComputeParticleBounds(
 	FParticleBoundsParameters Parameters;
 	FParticleBoundsUniformBufferRef UniformBuffer;
 
-	if (ParticleCount > 0 && GRHIFeatureLevel == ERHIFeatureLevel::SM5)
+	// MOBILEPREVIEWTODO: Proper value for this
+	const auto FeatureLevel = GRHIFeatureLevel;
+
+	if (ParticleCount > 0 && FeatureLevel == ERHIFeatureLevel::SM5)
 	{
 		// Determine how to break the work up over individual work groups.
 		const uint32 MaxGroupCount = 128;
@@ -2038,7 +2041,7 @@ static FBox ComputeParticleBounds(
 			PF_A32B32G32R32F );
 
 		// Grab the shader.
-		TShaderMapRef<FParticleBoundsCS> ParticleBoundsCS( GetGlobalShaderMap() );
+		TShaderMapRef<FParticleBoundsCS> ParticleBoundsCS(GetGlobalShaderMap(FeatureLevel));
 		RHICmdList.SetComputeShader(ParticleBoundsCS->GetComputeShader());
 
 		// Dispatch shader to compute bounds.
diff --git a/Engine/Source/Runtime/Engine/Private/Particles/ParticleSortingGPU.cpp b/Engine/Source/Runtime/Engine/Private/Particles/ParticleSortingGPU.cpp
--- a/Engine/Source/Runtime/Engine/Private/Particles/ParticleSortingGPU.cpp
+++ b/Engine/Source/Runtime/Engine/Private/Particles/ParticleSortingGPU.cpp
@@ -185,7 +185,10 @@ static int32 GenerateParticleSortKeys(
 {
 	SCOPED_DRAW_EVENT(ParticleSortKeyGen, DEC_PARTICLE);
 
-	check(GRHIFeatureLevel == ERHIFeatureLevel::SM5);
+	// MOBILEPREVIEWTODO: Proper value for this
+	const auto FeatureLevel = GRHIFeatureLevel;
+
+	check(FeatureLevel == ERHIFeatureLevel::SM5);
 
 	FParticleKeyGenParameters KeyGenParameters;
 	FParticleKeyGenUniformBufferRef KeyGenUniformBuffer;
@@ -193,7 +196,7 @@ static int32 GenerateParticleSortKeys(
 	int32 TotalParticleCount = 0;
 
 	// Grab the shader, set output.
-	TShaderMapRef<FParticleSortKeyGenCS> KeyGenCS(GetGlobalShaderMap());
+	TShaderMapRef<FParticleSortKeyGenCS> KeyGenCS(GetGlobalShaderMap(FeatureLevel));
 	RHICmdList.SetComputeShader(KeyGenCS->GetComputeShader());
 	KeyGenCS->SetOutput(RHICmdList, KeyBufferUAV, SortedVertexBufferUAV);
 	KeyGenCS->SetPositionTextures(RHICmdList, PositionTextureRHI);
diff --git a/Engine/Source/Runtime/Engine/Private/SceneView.cpp b/Engine/Source/Runtime/Engine/Private/SceneView.cpp
--- a/Engine/Source/Runtime/Engine/Private/SceneView.cpp
+++ b/Engine/Source/Runtime/Engine/Private/SceneView.cpp
@@ -250,6 +250,7 @@ FSceneView::FSceneView(const FSceneViewInitOptions& InitOptions)
 	, OverrideLODViewOrigin(InitOptions.OverrideLODViewOrigin)
 	, bAllowTranslucentPrimitivesInHitProxy( true )
 #endif
+	, FeatureLevel(InitOptions.ViewFamily ? InitOptions.ViewFamily->GetFeatureLevel() : GRHIFeatureLevel)
 {
 	check(UnscaledViewRect.Min.X >= 0);
 	check(UnscaledViewRect.Min.Y >= 0);
@@ -1222,11 +1223,6 @@ void FSceneView::ConfigureBufferVisualizationSettings()
 	}
 }
 
-ERHIFeatureLevel::Type FSceneView::GetFeatureLevel() const
-{ 
-	return (Family != nullptr) ? Family->GetFeatureLevel() : GRHIFeatureLevel;
-}
-
 EShaderPlatform FSceneView::GetShaderPlatform() const
 {
 	return GShaderPlatformForFeatureLevel[GetFeatureLevel()];
diff --git a/Engine/Source/Runtime/Engine/Private/TessellationRendering.cpp b/Engine/Source/Runtime/Engine/Private/TessellationRendering.cpp
--- a/Engine/Source/Runtime/Engine/Private/TessellationRendering.cpp
+++ b/Engine/Source/Runtime/Engine/Private/TessellationRendering.cpp
@@ -8,7 +8,7 @@ bool RequiresAdjacencyInformation( UMaterialInterface* Material, const FVertexFa
 {
 	EMaterialTessellationMode TessellationMode = MTM_NoTessellation;
 	bool bEnableCrackFreeDisplacement = false;
-	if ( RHISupportsTessellation(GRHIShaderPlatform) && VertexFactoryType->SupportsTessellationShaders() && Material )
+	if ( RHISupportsTessellation(GShaderPlatformForFeatureLevel[InFeatureLevel]) && VertexFactoryType->SupportsTessellationShaders() && Material )
 	{
 		if ( IsInRenderingThread() )
 		{
diff --git a/Engine/Source/Runtime/Engine/Private/UserInterface/Canvas.cpp b/Engine/Source/Runtime/Engine/Private/UserInterface/Canvas.cpp
--- a/Engine/Source/Runtime/Engine/Private/UserInterface/Canvas.cpp
+++ b/Engine/Source/Runtime/Engine/Private/UserInterface/Canvas.cpp
@@ -245,7 +245,7 @@ bool FCanvasBatchedElementRenderItem::Render_GameThread(const FCanvas* Canvas)
 			BatchedDrawCommand,
 			FBatchedDrawParameters,Parameters,DrawParameters,
 		{
-			const bool bNeedsToSwitchVerticalAxis = RHINeedsToSwitchVerticalAxis(GRHIShaderPlatform);
+			const bool bNeedsToSwitchVerticalAxis = RHINeedsToSwitchVerticalAxis(Parameters.ShaderPlatform);
 				
 			// draw batched items
 			Parameters.RenderData->BatchedElements.Draw(
diff --git a/Engine/Source/Runtime/Engine/Private/VectorField.cpp b/Engine/Source/Runtime/Engine/Private/VectorField.cpp
--- a/Engine/Source/Runtime/Engine/Private/VectorField.cpp
+++ b/Engine/Source/Runtime/Engine/Private/VectorField.cpp
@@ -843,7 +843,7 @@ public:
 			FCompositeAnimatedVectorFieldUniformBufferRef UniformBuffer = 
 				FCompositeAnimatedVectorFieldUniformBufferRef::CreateUniformBufferImmediate(Parameters, UniformBuffer_SingleDraw);
 
-			TShaderMapRef<FCompositeAnimatedVectorFieldCS> CompositeCS(GetGlobalShaderMap());
+			TShaderMapRef<FCompositeAnimatedVectorFieldCS> CompositeCS(GetGlobalShaderMap(GetFeatureLevel()));
 			FTextureRHIParamRef NoiseVolumeTextureRHI = GBlackVolumeTexture->TextureRHI;
 			if (AnimatedVectorField->NoiseField && AnimatedVectorField->NoiseField->Resource)
 			{
diff --git a/Engine/Source/Runtime/Engine/Private/World.cpp b/Engine/Source/Runtime/Engine/Private/World.cpp
--- a/Engine/Source/Runtime/Engine/Private/World.cpp
+++ b/Engine/Source/Runtime/Engine/Private/World.cpp
@@ -5171,7 +5171,7 @@ void UWorld::ChangeFeatureLevel(ERHIFeatureLevel::Type InFeatureLevel)
 
 		UMaterial::AllMaterialsCacheResourceShadersForRendering();
 		UMaterialInstance::AllMaterialsCacheResourceShadersForRendering();
-		GetGlobalShaderMap(GRHIShaderPlatform, false);
+		GetGlobalShaderMap(InFeatureLevel, false);
 		GShaderCompilingManager->ProcessAsyncResults(false, true);
 
 		//invalidate global bound shader states so they will be created with the new shaders the next time they are set (in SetGlobalBoundShaderState)
@@ -5214,7 +5214,7 @@ void UWorld::ChangeAllWorldFeatureLevels(ERHIFeatureLevel::Type InFeatureLevel)
 
 		UMaterial::AllMaterialsCacheResourceShadersForRendering();
 		UMaterialInstance::AllMaterialsCacheResourceShadersForRendering();
-		GetGlobalShaderMap(GRHIShaderPlatform, false);
+		GetGlobalShaderMap(InFeatureLevel, false);
 		GShaderCompilingManager->ProcessAsyncResults(false, true);
 		
 		//invalidate global bound shader states so they will be created with the new shaders the next time they are set (in SetGlobalBoundShaderState)
diff --git a/Engine/Source/Runtime/Engine/Public/GlobalShader.h b/Engine/Source/Runtime/Engine/Public/GlobalShader.h
--- a/Engine/Source/Runtime/Engine/Public/GlobalShader.h
+++ b/Engine/Source/Runtime/Engine/Public/GlobalShader.h
@@ -7,6 +7,7 @@
 #pragma once
 
 #include "Shader.h"
+#include "SceneManagement.h"
 
 class FShaderCompileJob;
 
@@ -176,7 +177,20 @@ extern ENGINE_API FString SaveGlobalShaderFile(EShaderPlatform Platform, FString
  * @param bRefreshShaderMap If true, the existing global shader map will be tossed first
  * @return A reference to the global shader map.
  */
-extern ENGINE_API TShaderMap<FGlobalShaderType>* GetGlobalShaderMap(EShaderPlatform Platform = GRHIShaderPlatform, bool bRefreshShaderMap = false);
+extern ENGINE_API TShaderMap<FGlobalShaderType>* GetGlobalShaderMap(EShaderPlatform Platform, bool bRefreshShaderMap = false);
+
+/**
+  * Overload for the above GetGlobalShaderMap which takes a feature level and translates to the appropriate shader platform
+  *
+  * @param FeatureLevel - Which feature levels shader map to use
+  * @param bRefreshShaderMap If true, the existing global shader map will be tossed first
+  * @return A reference to the global shader map.
+  *
+  **/
+inline TShaderMap<FGlobalShaderType>* GetGlobalShaderMap(ERHIFeatureLevel::Type FeatureLevel, bool bRefreshShaderMap = false) 
+{ 
+	return GetGlobalShaderMap(GShaderPlatformForFeatureLevel[FeatureLevel], bRefreshShaderMap); 
+}
 
 /**
  * Forces a recompile of the global shaders.
diff --git a/Engine/Source/Runtime/Engine/Public/SceneView.h b/Engine/Source/Runtime/Engine/Public/SceneView.h
--- a/Engine/Source/Runtime/Engine/Public/SceneView.h
+++ b/Engine/Source/Runtime/Engine/Public/SceneView.h
@@ -372,6 +372,9 @@ public:
 	/** Areas of the screen where the depth of field should be forced to the max. */
 	TArray< FIntRect > UIBlurOverrideRectangles;
 
+	/** Feature level for this scene */
+	ERHIFeatureLevel::Type FeatureLevel;
+
 	/** Initialization constructor. */
 	FSceneView(const FSceneViewInitOptions& InitOptions);
 
@@ -451,7 +454,7 @@ public:
 	void ConfigureBufferVisualizationSettings();
 
 	/** Get the feature level for this view **/
-	ERHIFeatureLevel::Type GetFeatureLevel() const;
+	ERHIFeatureLevel::Type GetFeatureLevel() const { return FeatureLevel; }
 
 	/** Get the feature level for this view **/
 	EShaderPlatform GetShaderPlatform() const;
diff --git a/Engine/Source/Runtime/GameLiveStreaming/GameLiveStreaming.cpp b/Engine/Source/Runtime/GameLiveStreaming/GameLiveStreaming.cpp
--- a/Engine/Source/Runtime/GameLiveStreaming/GameLiveStreaming.cpp
+++ b/Engine/Source/Runtime/GameLiveStreaming/GameLiveStreaming.cpp
@@ -315,6 +315,8 @@ void FGameLiveStreaming::StartCopyingNextGameVideoFrame( const FViewportRHIRef&
 	{
 		FPooledRenderTargetDesc OutputDesc(FPooledRenderTargetDesc::Create2DDesc(Context.ResizeTo, PF_B8G8R8A8, TexCreate_None, TexCreate_RenderTargetable, false));
 			
+		const auto FeatureLevel = GRHIFeatureLevel;
+
 		TRefCountPtr<IPooledRenderTarget> ResampleTexturePooledRenderTarget;
 		Context.RendererModule->RenderTargetPoolFindFreeElement(OutputDesc, ResampleTexturePooledRenderTarget, TEXT("ResampleTexture"));
 		check( ResampleTexturePooledRenderTarget );
@@ -330,11 +332,12 @@ void FGameLiveStreaming::StartCopyingNextGameVideoFrame( const FViewportRHIRef&
 
 		FTexture2DRHIRef ViewportBackBuffer = RHICmdList.GetViewportBackBuffer(Context.ViewportRHI);
 
-		TShaderMapRef<FScreenVS> VertexShader(GetGlobalShaderMap());
-		TShaderMapRef<FScreenPS> PixelShader(GetGlobalShaderMap());
+		auto ShaderMap = GetGlobalShaderMap(FeatureLevel);
+		TShaderMapRef<FScreenVS> VertexShader(ShaderMap);
+		TShaderMapRef<FScreenPS> PixelShader(ShaderMap);
 
 		static FGlobalBoundShaderState BoundShaderState;
-		SetGlobalBoundShaderState(RHICmdList, GRHIFeatureLevel, BoundShaderState, Context.RendererModule->GetFilterVertexDeclaration().VertexDeclarationRHI, *VertexShader, *PixelShader);
+		SetGlobalBoundShaderState(RHICmdList, FeatureLevel, BoundShaderState, Context.RendererModule->GetFilterVertexDeclaration().VertexDeclarationRHI, *VertexShader, *PixelShader);
 
 		if( Context.ResizeTo != FIntPoint( ViewportBackBuffer->GetSizeX(), ViewportBackBuffer->GetSizeY() ) )
 		{
diff --git a/Engine/Source/Runtime/Launch/Private/LaunchEngineLoop.cpp b/Engine/Source/Runtime/Launch/Private/LaunchEngineLoop.cpp
--- a/Engine/Source/Runtime/Launch/Private/LaunchEngineLoop.cpp
+++ b/Engine/Source/Runtime/Launch/Private/LaunchEngineLoop.cpp
@@ -1242,7 +1242,7 @@ int32 FEngineLoop::PreInit( const TCHAR* CmdLine )
 		InitializeShaderTypes();
 
 		// Load the global shaders.
-		if (GetGlobalShaderMap() == NULL && GIsRequestingExit)
+		if (GetGlobalShaderMap(GMaxRHIFeatureLevel) == NULL && GIsRequestingExit)
 		{
 			// This means we can't continue without the global shader map.
 			return 1;
diff --git a/Engine/Source/Runtime/OpenGLDrv/Private/OpenGLDrv.cpp b/Engine/Source/Runtime/OpenGLDrv/Private/OpenGLDrv.cpp
--- a/Engine/Source/Runtime/OpenGLDrv/Private/OpenGLDrv.cpp
+++ b/Engine/Source/Runtime/OpenGLDrv/Private/OpenGLDrv.cpp
@@ -404,16 +404,19 @@ void FOpenGLDynamicRHI::IssueLongGPUTask()
 	{
 		FOpenGLViewport* Viewport = Viewports[LargestViewportIndex];
 
+		const auto FeatureLevel = GMaxRHIFeatureLevel;
+
 		FRHICommandList_RecursiveHazardous RHICmdList;
 		SetRenderTarget(RHICmdList, Viewport->GetBackBuffer(), FTextureRHIRef());
 		RHICmdList.SetBlendState(TStaticBlendState<CW_RGBA, BO_Add, BF_One, BF_One>::GetRHI(), FLinearColor::Black);
 		RHICmdList.SetDepthStencilState(TStaticDepthStencilState<false, CF_Always>::GetRHI(), 0);
 		RHICmdList.SetRasterizerState(TStaticRasterizerState<FM_Solid, CM_None>::GetRHI());
 
-		TShaderMapRef<TOneColorVS<true> > VertexShader(GetGlobalShaderMap());
-		TShaderMapRef<FOpenGLRHILongGPUTaskPS> PixelShader(GetGlobalShaderMap());
+		auto ShaderMap = GetGlobalShaderMap(FeatureLevel);
+		TShaderMapRef<TOneColorVS<true> > VertexShader(ShaderMap);
+		TShaderMapRef<FOpenGLRHILongGPUTaskPS> PixelShader(ShaderMap);
 
-		SetGlobalBoundShaderState(RHICmdList, GMaxRHIFeatureLevel, LongGPUTaskBoundShaderState, GOpenGLVector4VertexDeclaration.VertexDeclarationRHI, *VertexShader, *PixelShader, 0);
+		SetGlobalBoundShaderState(RHICmdList, FeatureLevel, LongGPUTaskBoundShaderState, GOpenGLVector4VertexDeclaration.VertexDeclarationRHI, *VertexShader, *PixelShader, 0);
 
 		// Draw a fullscreen quad
 		FVector4 Vertices[4];
diff --git a/Engine/Source/Runtime/OpenGLDrv/Private/OpenGLShaders.cpp b/Engine/Source/Runtime/OpenGLDrv/Private/OpenGLShaders.cpp
--- a/Engine/Source/Runtime/OpenGLDrv/Private/OpenGLShaders.cpp
+++ b/Engine/Source/Runtime/OpenGLDrv/Private/OpenGLShaders.cpp
@@ -1379,7 +1379,7 @@ FBoundShaderStateRHIRef FOpenGLDynamicRHI::RHICreateBoundShaderState(
 	if(!PixelShaderRHI)
 	{
 		// use special null pixel shader when PixelShader was set to NULL
-		PixelShaderRHI = TShaderMapRef<FNULLPS>(GetGlobalShaderMap())->GetPixelShader();
+		PixelShaderRHI = TShaderMapRef<FNULLPS>(GetGlobalShaderMap(GMaxRHIFeatureLevel))->GetPixelShader();
 	}
 
 	// Check for an existing bound shader state which matches the parameters
diff --git a/Engine/Source/Runtime/Renderer/Private/AtmosphereRendering.cpp b/Engine/Source/Runtime/Renderer/Private/AtmosphereRendering.cpp
--- a/Engine/Source/Runtime/Renderer/Private/AtmosphereRendering.cpp
+++ b/Engine/Source/Runtime/Renderer/Private/AtmosphereRendering.cpp
@@ -355,34 +355,36 @@ FGlobalBoundShaderState AtmosphereBoundShaderState[EAtmosphereRenderFlag::E_Rend
 
 void SetAtmosphericFogShaders(FRHICommandList& RHICmdList, FScene* Scene, const FViewInfo& View, TRefCountPtr<IPooledRenderTarget>& LightShaftOcclusion)
 {
-	TShaderMapRef<FAtmosphericVS> VertexShader(GetGlobalShaderMap());
+	auto ShaderMap = View.ShaderMap;
+
+	TShaderMapRef<FAtmosphericVS> VertexShader(ShaderMap);
 	FAtmosphericFogPS* PixelShader = NULL;
 	switch (Scene->AtmosphericFog->RenderFlag)
 	{
 	default:
 	case EAtmosphereRenderFlag::E_EnableAll:
-		PixelShader = *TShaderMapRef<TAtmosphericFogPS<EAtmosphereRenderFlag::E_EnableAll> >(GetGlobalShaderMap());
+		PixelShader = *TShaderMapRef<TAtmosphericFogPS<EAtmosphereRenderFlag::E_EnableAll> >(ShaderMap);
 		break;
 	case EAtmosphereRenderFlag::E_DisableSunDisk:
-		PixelShader = *TShaderMapRef<TAtmosphericFogPS<EAtmosphereRenderFlag::E_DisableSunDisk> >(GetGlobalShaderMap());
+		PixelShader = *TShaderMapRef<TAtmosphericFogPS<EAtmosphereRenderFlag::E_DisableSunDisk> >(ShaderMap);
 		break;
 	case EAtmosphereRenderFlag::E_DisableGroundScattering:
-		PixelShader = *TShaderMapRef<TAtmosphericFogPS<EAtmosphereRenderFlag::E_DisableGroundScattering> >(GetGlobalShaderMap());
+		PixelShader = *TShaderMapRef<TAtmosphericFogPS<EAtmosphereRenderFlag::E_DisableGroundScattering> >(ShaderMap);
 		break;
 	case EAtmosphereRenderFlag::E_DisableSunAndGround :
-		PixelShader = *TShaderMapRef<TAtmosphericFogPS<EAtmosphereRenderFlag::E_DisableSunAndGround> >(GetGlobalShaderMap());
+		PixelShader = *TShaderMapRef<TAtmosphericFogPS<EAtmosphereRenderFlag::E_DisableSunAndGround> >(ShaderMap);
 		break;
 	case EAtmosphereRenderFlag::E_DisableLightShaft :
-		PixelShader = *TShaderMapRef<TAtmosphericFogPS<EAtmosphereRenderFlag::E_DisableLightShaft> >(GetGlobalShaderMap());
+		PixelShader = *TShaderMapRef<TAtmosphericFogPS<EAtmosphereRenderFlag::E_DisableLightShaft> >(ShaderMap);
 		break;
 	case EAtmosphereRenderFlag::E_DisableSunAndLightShaft :
-		PixelShader = *TShaderMapRef<TAtmosphericFogPS<EAtmosphereRenderFlag::E_DisableSunAndLightShaft> >(GetGlobalShaderMap());
+		PixelShader = *TShaderMapRef<TAtmosphericFogPS<EAtmosphereRenderFlag::E_DisableSunAndLightShaft> >(ShaderMap);
 		break;
 	case EAtmosphereRenderFlag::E_DisableGroundAndLightShaft :
-		PixelShader = *TShaderMapRef<TAtmosphericFogPS<EAtmosphereRenderFlag::E_DisableGroundAndLightShaft> >(GetGlobalShaderMap());
+		PixelShader = *TShaderMapRef<TAtmosphericFogPS<EAtmosphereRenderFlag::E_DisableGroundAndLightShaft> >(ShaderMap);
 		break;
 	case EAtmosphereRenderFlag::E_DisableAll:
-		PixelShader = *TShaderMapRef<TAtmosphericFogPS<EAtmosphereRenderFlag::E_DisableAll> >(GetGlobalShaderMap());
+		PixelShader = *TShaderMapRef<TAtmosphericFogPS<EAtmosphereRenderFlag::E_DisableAll> >(ShaderMap);
 		break;
 	}
 	
@@ -1085,6 +1087,9 @@ void FAtmosphericFogSceneInfo::GetLayerValue(int Layer, float& AtmosphereR, FVec
 
 void FAtmosphericFogSceneInfo::RenderAtmosphereShaders(FRHICommandList& RHICmdList, const FViewInfo& View, const FIntRect& ViewRect)
 {
+	const auto FeatureLevel = View.GetFeatureLevel();
+	auto ShaderMap = View.ShaderMap;
+
 	check(Component != NULL);
 	switch(AtmospherePhase)
 	{
@@ -1093,11 +1098,11 @@ void FAtmosphericFogSceneInfo::RenderAtmosphereShaders(FRHICommandList& RHICmdLi
 			const FSceneRenderTargetItem& DestRenderTarget = AtmosphereTextures->AtmosphereTransmittance->GetRenderTargetItem();
 			SetRenderTarget(RHICmdList, DestRenderTarget.TargetableTexture, FTextureRHIRef());
 
-			TShaderMapRef<FAtmospherePrecomputeVS> VertexShader(GetGlobalShaderMap());
-			TShaderMapRef<FAtmosphereTransmittancePS> PixelShader(GetGlobalShaderMap());
+			TShaderMapRef<FAtmospherePrecomputeVS> VertexShader(ShaderMap);
+			TShaderMapRef<FAtmosphereTransmittancePS> PixelShader(ShaderMap);
 
 			static FGlobalBoundShaderState BoundShaderState;
-			SetGlobalBoundShaderState(RHICmdList, View.GetFeatureLevel(), BoundShaderState, GFilterVertexDeclaration.VertexDeclarationRHI, *VertexShader, *PixelShader);
+			SetGlobalBoundShaderState(RHICmdList, FeatureLevel, BoundShaderState, GFilterVertexDeclaration.VertexDeclarationRHI, *VertexShader, *PixelShader);
 			PixelShader->SetParameters(RHICmdList, View);
 			DrawQuad(RHICmdList, ViewRect, *VertexShader);
 
@@ -1109,11 +1114,11 @@ void FAtmosphericFogSceneInfo::RenderAtmosphereShaders(FRHICommandList& RHICmdLi
 			const FSceneRenderTargetItem& DestRenderTarget = AtmosphereTextures->AtmosphereDeltaE->GetRenderTargetItem();
 			SetRenderTarget(RHICmdList, DestRenderTarget.TargetableTexture, FTextureRHIRef());
 
-			TShaderMapRef<FAtmospherePrecomputeVS> VertexShader(GetGlobalShaderMap());
-			TShaderMapRef<FAtmosphereIrradiance1PS> PixelShader(GetGlobalShaderMap());
+			TShaderMapRef<FAtmospherePrecomputeVS> VertexShader(ShaderMap);
+			TShaderMapRef<FAtmosphereIrradiance1PS> PixelShader(ShaderMap);
 
 			static FGlobalBoundShaderState BoundShaderState;
-			SetGlobalBoundShaderState(RHICmdList, View.GetFeatureLevel(), BoundShaderState, GFilterVertexDeclaration.VertexDeclarationRHI, *VertexShader, *PixelShader);
+			SetGlobalBoundShaderState(RHICmdList, FeatureLevel, BoundShaderState, GFilterVertexDeclaration.VertexDeclarationRHI, *VertexShader, *PixelShader);
 
 			PixelShader->SetParameters(RHICmdList, AtmosphereTextures);
 
@@ -1131,13 +1136,13 @@ void FAtmosphericFogSceneInfo::RenderAtmosphereShaders(FRHICommandList& RHICmdLi
 				RenderTargets[1] = AtmosphereTextures->AtmosphereDeltaSM->GetRenderTargetItem().TargetableTexture;
 				SetRenderTargets(RHICmdList, ARRAY_COUNT(RenderTargets), RenderTargets, FTextureRHIRef(), 0, NULL);
 
-				TShaderMapRef<FAtmospherePrecomputeVS> VertexShader(GetGlobalShaderMap());
-				TShaderMapRef<FAtmosphereGS> GeometryShader(GetGlobalShaderMap());
-				TShaderMapRef<FAtmosphereInscatter1PS> PixelShader(GetGlobalShaderMap());
+				TShaderMapRef<FAtmospherePrecomputeVS> VertexShader(ShaderMap);
+				TShaderMapRef<FAtmosphereGS> GeometryShader(ShaderMap);
+				TShaderMapRef<FAtmosphereInscatter1PS> PixelShader(ShaderMap);
 
 				static FGlobalBoundShaderState BoundShaderState;
 				
-				SetGlobalBoundShaderState(RHICmdList, View.GetFeatureLevel(), BoundShaderState, GFilterVertexDeclaration.VertexDeclarationRHI, *VertexShader, *PixelShader, *GeometryShader);
+				SetGlobalBoundShaderState(RHICmdList, FeatureLevel, BoundShaderState, GFilterVertexDeclaration.VertexDeclarationRHI, *VertexShader, *PixelShader, *GeometryShader);
 
 				//
 				float r;
@@ -1173,13 +1178,13 @@ void FAtmosphericFogSceneInfo::RenderAtmosphereShaders(FRHICommandList& RHICmdLi
 				const FSceneRenderTargetItem& DestRenderTarget = AtmosphereTextures->AtmosphereInscatter->GetRenderTargetItem();
 				SetRenderTarget(RHICmdList, DestRenderTarget.TargetableTexture, FTextureRHIRef());
 
-				TShaderMapRef<FAtmospherePrecomputeVS> VertexShader(GetGlobalShaderMap());
-				TShaderMapRef<FAtmosphereGS> GeometryShader(GetGlobalShaderMap());
-				TShaderMapRef<FAtmosphereCopyInscatter1PS> PixelShader(GetGlobalShaderMap());
+				TShaderMapRef<FAtmospherePrecomputeVS> VertexShader(ShaderMap);
+				TShaderMapRef<FAtmosphereGS> GeometryShader(ShaderMap);
+				TShaderMapRef<FAtmosphereCopyInscatter1PS> PixelShader(ShaderMap);
 
 				static FGlobalBoundShaderState BoundShaderState;
 				
-				SetGlobalBoundShaderState(RHICmdList, View.GetFeatureLevel(), BoundShaderState, GFilterVertexDeclaration.VertexDeclarationRHI, *VertexShader, *PixelShader, *GeometryShader);
+				SetGlobalBoundShaderState(RHICmdList, FeatureLevel, BoundShaderState, GFilterVertexDeclaration.VertexDeclarationRHI, *VertexShader, *PixelShader, *GeometryShader);
 
 				float r;
 				FVector4 DhdH;
@@ -1202,13 +1207,13 @@ void FAtmosphericFogSceneInfo::RenderAtmosphereShaders(FRHICommandList& RHICmdLi
 				const FSceneRenderTargetItem& DestRenderTarget = AtmosphereTextures->AtmosphereDeltaJ->GetRenderTargetItem();
 				SetRenderTarget(RHICmdList, DestRenderTarget.TargetableTexture, FTextureRHIRef());
 
-				TShaderMapRef<FAtmospherePrecomputeVS> VertexShader(GetGlobalShaderMap());
-				TShaderMapRef<FAtmosphereGS> GeometryShader(GetGlobalShaderMap());
-				TShaderMapRef<FAtmosphereInscatterSPS> PixelShader(GetGlobalShaderMap());
+				TShaderMapRef<FAtmospherePrecomputeVS> VertexShader(ShaderMap);
+				TShaderMapRef<FAtmosphereGS> GeometryShader(ShaderMap);
+				TShaderMapRef<FAtmosphereInscatterSPS> PixelShader(ShaderMap);
 
 				static FGlobalBoundShaderState BoundShaderState;
 				
-				SetGlobalBoundShaderState(RHICmdList, View.GetFeatureLevel(), BoundShaderState, GFilterVertexDeclaration.VertexDeclarationRHI, *VertexShader, *PixelShader, *GeometryShader);
+				SetGlobalBoundShaderState(RHICmdList, FeatureLevel, BoundShaderState, GFilterVertexDeclaration.VertexDeclarationRHI, *VertexShader, *PixelShader, *GeometryShader);
 
 				float r;
 				FVector4 DhdH;
@@ -1229,12 +1234,12 @@ void FAtmosphericFogSceneInfo::RenderAtmosphereShaders(FRHICommandList& RHICmdLi
 			const FSceneRenderTargetItem& DestRenderTarget = AtmosphereTextures->AtmosphereDeltaE->GetRenderTargetItem();
 			SetRenderTarget(RHICmdList, DestRenderTarget.TargetableTexture, FTextureRHIRef());
 
-			TShaderMapRef<FAtmospherePrecomputeVS> VertexShader(GetGlobalShaderMap());
-			TShaderMapRef<FAtmosphereIrradianceNPS> PixelShader(GetGlobalShaderMap());
+			TShaderMapRef<FAtmospherePrecomputeVS> VertexShader(ShaderMap);
+			TShaderMapRef<FAtmosphereIrradianceNPS> PixelShader(ShaderMap);
 
 			static FGlobalBoundShaderState BoundShaderState;
 			
-			SetGlobalBoundShaderState(RHICmdList, View.GetFeatureLevel(), BoundShaderState, GFilterVertexDeclaration.VertexDeclarationRHI, *VertexShader, *PixelShader);
+			SetGlobalBoundShaderState(RHICmdList, FeatureLevel, BoundShaderState, GFilterVertexDeclaration.VertexDeclarationRHI, *VertexShader, *PixelShader);
 
 			PixelShader->SetParameters(RHICmdList, View, AtmoshpereOrder == 2 ? 1.f : 0.f, AtmosphereTextures);
 
@@ -1251,13 +1256,13 @@ void FAtmosphericFogSceneInfo::RenderAtmosphereShaders(FRHICommandList& RHICmdLi
 				const FSceneRenderTargetItem& DestRenderTarget = AtmosphereTextures->AtmosphereDeltaSR->GetRenderTargetItem();
 				SetRenderTarget(RHICmdList, DestRenderTarget.TargetableTexture, FTextureRHIRef());
 
-				TShaderMapRef<FAtmospherePrecomputeVS> VertexShader(GetGlobalShaderMap());
-				TShaderMapRef<FAtmosphereGS> GeometryShader(GetGlobalShaderMap());
-				TShaderMapRef<FAtmosphereInscatterNPS> PixelShader(GetGlobalShaderMap());
+				TShaderMapRef<FAtmospherePrecomputeVS> VertexShader(ShaderMap);
+				TShaderMapRef<FAtmosphereGS> GeometryShader(ShaderMap);
+				TShaderMapRef<FAtmosphereInscatterNPS> PixelShader(ShaderMap);
 
 				static FGlobalBoundShaderState BoundShaderState;
 				
-				SetGlobalBoundShaderState(RHICmdList, View.GetFeatureLevel(), BoundShaderState, GFilterVertexDeclaration.VertexDeclarationRHI, *VertexShader, *PixelShader, *GeometryShader);
+				SetGlobalBoundShaderState(RHICmdList, FeatureLevel, BoundShaderState, GFilterVertexDeclaration.VertexDeclarationRHI, *VertexShader, *PixelShader, *GeometryShader);
 
 				float r;
 				FVector4 DhdH;
@@ -1280,12 +1285,12 @@ void FAtmosphericFogSceneInfo::RenderAtmosphereShaders(FRHICommandList& RHICmdLi
 			const FSceneRenderTargetItem& DestRenderTarget = AtmosphereTextures->AtmosphereIrradiance->GetRenderTargetItem();
 			SetRenderTarget(RHICmdList, DestRenderTarget.TargetableTexture, FTextureRHIRef());
 
-			TShaderMapRef<FAtmospherePrecomputeVS> VertexShader(GetGlobalShaderMap());
-			TShaderMapRef<FAtmosphereCopyIrradiancePS> PixelShader(GetGlobalShaderMap());
+			TShaderMapRef<FAtmospherePrecomputeVS> VertexShader(ShaderMap);
+			TShaderMapRef<FAtmosphereCopyIrradiancePS> PixelShader(ShaderMap);
 
 			static FGlobalBoundShaderState BoundShaderState;
 			
-			SetGlobalBoundShaderState(RHICmdList, View.GetFeatureLevel(), BoundShaderState, GFilterVertexDeclaration.VertexDeclarationRHI, *VertexShader, *PixelShader);
+			SetGlobalBoundShaderState(RHICmdList, FeatureLevel, BoundShaderState, GFilterVertexDeclaration.VertexDeclarationRHI, *VertexShader, *PixelShader);
 
 			PixelShader->SetParameters(RHICmdList, AtmosphereTextures);
 
@@ -1306,13 +1311,13 @@ void FAtmosphericFogSceneInfo::RenderAtmosphereShaders(FRHICommandList& RHICmdLi
 				const FSceneRenderTargetItem& DestRenderTarget = AtmosphereTextures->AtmosphereInscatter->GetRenderTargetItem();
 				SetRenderTarget(RHICmdList, DestRenderTarget.TargetableTexture, FTextureRHIRef());
 
-				TShaderMapRef<FAtmospherePrecomputeVS> VertexShader(GetGlobalShaderMap());
-				TShaderMapRef<FAtmosphereGS> GeometryShader(GetGlobalShaderMap());
-				TShaderMapRef<FAtmosphereCopyInscatterNPS> PixelShader(GetGlobalShaderMap());
+				TShaderMapRef<FAtmospherePrecomputeVS> VertexShader(ShaderMap);
+				TShaderMapRef<FAtmosphereGS> GeometryShader(ShaderMap);
+				TShaderMapRef<FAtmosphereCopyInscatterNPS> PixelShader(ShaderMap);
 
 				static FGlobalBoundShaderState BoundShaderState;
 				
-				SetGlobalBoundShaderState(RHICmdList, View.GetFeatureLevel(), BoundShaderState, GFilterVertexDeclaration.VertexDeclarationRHI, *VertexShader, *PixelShader, *GeometryShader);
+				SetGlobalBoundShaderState(RHICmdList, FeatureLevel, BoundShaderState, GFilterVertexDeclaration.VertexDeclarationRHI, *VertexShader, *PixelShader, *GeometryShader);
 
 				float r;
 				FVector4 DhdH;
@@ -1338,13 +1343,13 @@ void FAtmosphericFogSceneInfo::RenderAtmosphereShaders(FRHICommandList& RHICmdLi
 				const FSceneRenderTargetItem& DestRenderTarget = AtmosphereTextures->AtmosphereDeltaSR->GetRenderTargetItem();
 				SetRenderTarget(RHICmdList, DestRenderTarget.TargetableTexture, FTextureRHIRef());
 
-				TShaderMapRef<FAtmospherePrecomputeVS> VertexShader(GetGlobalShaderMap());
-				TShaderMapRef<FAtmosphereGS> GeometryShader(GetGlobalShaderMap());
-				TShaderMapRef<FAtmosphereCopyInscatterFPS> PixelShader(GetGlobalShaderMap());
+				TShaderMapRef<FAtmospherePrecomputeVS> VertexShader(ShaderMap);
+				TShaderMapRef<FAtmosphereGS> GeometryShader(ShaderMap);
+				TShaderMapRef<FAtmosphereCopyInscatterFPS> PixelShader(ShaderMap);
 
 				static FGlobalBoundShaderState BoundShaderState;
 				
-				SetGlobalBoundShaderState(RHICmdList, View.GetFeatureLevel(), BoundShaderState, GFilterVertexDeclaration.VertexDeclarationRHI, *VertexShader, *PixelShader, *GeometryShader);
+				SetGlobalBoundShaderState(RHICmdList, FeatureLevel, BoundShaderState, GFilterVertexDeclaration.VertexDeclarationRHI, *VertexShader, *PixelShader, *GeometryShader);
 
 				float r;
 				FVector4 DhdH;
@@ -1367,13 +1372,13 @@ void FAtmosphericFogSceneInfo::RenderAtmosphereShaders(FRHICommandList& RHICmdLi
 				const FSceneRenderTargetItem& DestRenderTarget = AtmosphereTextures->AtmosphereInscatter->GetRenderTargetItem();
 				SetRenderTarget(RHICmdList, DestRenderTarget.TargetableTexture, FTextureRHIRef());
 
-				TShaderMapRef<FAtmospherePrecomputeVS> VertexShader(GetGlobalShaderMap());
-				TShaderMapRef<FAtmosphereGS> GeometryShader(GetGlobalShaderMap());
-				TShaderMapRef<FAtmosphereCopyInscatterFBackPS> PixelShader(GetGlobalShaderMap());
+				TShaderMapRef<FAtmospherePrecomputeVS> VertexShader(ShaderMap);
+				TShaderMapRef<FAtmosphereGS> GeometryShader(ShaderMap);
+				TShaderMapRef<FAtmosphereCopyInscatterFBackPS> PixelShader(ShaderMap);
 
 				static FGlobalBoundShaderState BoundShaderState;
 				
-				SetGlobalBoundShaderState(RHICmdList, View.GetFeatureLevel(), BoundShaderState, GFilterVertexDeclaration.VertexDeclarationRHI, *VertexShader, *PixelShader, *GeometryShader);
+				SetGlobalBoundShaderState(RHICmdList, FeatureLevel, BoundShaderState, GFilterVertexDeclaration.VertexDeclarationRHI, *VertexShader, *PixelShader, *GeometryShader);
 
 				float r;
 				FVector4 DhdH;
diff --git a/Engine/Source/Runtime/Renderer/Private/BasePassRendering.cpp b/Engine/Source/Runtime/Renderer/Private/BasePassRendering.cpp
--- a/Engine/Source/Runtime/Renderer/Private/BasePassRendering.cpp
+++ b/Engine/Source/Runtime/Renderer/Private/BasePassRendering.cpp
@@ -231,14 +231,14 @@ class FDrawBasePassDynamicMeshAction
 {
 public:
 
-	const FSceneView& View;
+	const FViewInfo& View;
 	bool bBackFace;
 	bool bPreFog;
 	FHitProxyId HitProxyId;
 
 	/** Initialization constructor. */
 	FDrawBasePassDynamicMeshAction(
-		const FSceneView& InView,
+		const FViewInfo& InView,
 		const bool bInBackFace,
 		const FHitProxyId InHitProxyId
 		)
@@ -328,7 +328,7 @@ public:
 
 bool FBasePassOpaqueDrawingPolicyFactory::DrawDynamicMesh(
 	FRHICommandList& RHICmdList, 
-	const FSceneView& View,
+	const FViewInfo& View,
 	ContextType DrawingContext,
 	const FMeshBatch& Mesh,
 	bool bBackFace,
diff --git a/Engine/Source/Runtime/Renderer/Private/BasePassRendering.h b/Engine/Source/Runtime/Renderer/Private/BasePassRendering.h
--- a/Engine/Source/Runtime/Renderer/Private/BasePassRendering.h
+++ b/Engine/Source/Runtime/Renderer/Private/BasePassRendering.h
@@ -103,7 +103,7 @@ public:
 	{
 		bool bShouldCache = TBasePassVertexShaderBaseType<LightMapPolicyType>::ShouldCache(Platform, Material, VertexFactoryType);
 		return bShouldCache 
-			&& IsFeatureLevelSupported(Platform, ERHIFeatureLevel::SM3)
+			&& (IsFeatureLevelSupported(Platform, ERHIFeatureLevel::SM3) || (IsPCPlatform(Platform)))
 			&& (!bEnableAtmosphericFog || IsFeatureLevelSupported(Platform, ERHIFeatureLevel::SM4));
 	}
 
@@ -427,7 +427,7 @@ public:
 		const bool bCacheShaders = !bEnableSkyLight || (Material->GetShadingModel() != MSM_Unlit);
 
 		return bCacheShaders
-			&& IsFeatureLevelSupported(Platform, ERHIFeatureLevel::SM3)
+			&& (IsFeatureLevelSupported(Platform, ERHIFeatureLevel::SM3) || (IsPCPlatform(Platform)))
 			&& TBasePassPixelShaderBaseType<LightMapPolicyType>::ShouldCache(Platform, Material, VertexFactoryType);
 	}
 
@@ -553,7 +553,7 @@ public:
 			LightMapPolicy == Other.LightMapPolicy;
 	}
 
-	void SetSharedState(FRHICommandList& RHICmdList, const FSceneView* View) const
+	void SetSharedState(FRHICommandList& RHICmdList, const FViewInfo* View) const
 	{
 		// Set the light-map policy.
 		LightMapPolicy.Set(RHICmdList, VertexShader,bOverrideWithShaderComplexity ? NULL : PixelShader,VertexShader,PixelShader,VertexFactory,MaterialRenderProxy,View);
@@ -579,7 +579,7 @@ public:
 				RHICmdList.SetBlendState( TStaticBlendState<CW_RGBA, BO_Add, BF_One, BF_One, BO_Add, BF_Zero, BF_One>::GetRHI());
 			}
 
-			TShaderMapRef<FShaderComplexityAccumulatePS> ShaderComplexityPixelShader(GetGlobalShaderMap());
+			TShaderMapRef<FShaderComplexityAccumulatePS> ShaderComplexityPixelShader(View->ShaderMap);
 			const uint32 NumPixelShaderInstructions = PixelShader->GetNumInstructions();
 			const uint32 NumVertexShaderInstructions = VertexShader->GetNumInstructions();
 			ShaderComplexityPixelShader->SetParameters(RHICmdList, NumVertexShaderInstructions, NumPixelShaderInstructions, View->GetFeatureLevel());
@@ -629,7 +629,7 @@ public:
 #if !(UE_BUILD_SHIPPING || UE_BUILD_TEST)
 		if (bOverrideWithShaderComplexity)
 		{
-			TShaderMapRef<FShaderComplexityAccumulatePS> ShaderComplexityAccumulatePixelShader(GetGlobalShaderMap());
+			TShaderMapRef<FShaderComplexityAccumulatePS> ShaderComplexityAccumulatePixelShader(GetGlobalShaderMap(InFeatureLevel));
 			PixelShaderRHIRef = ShaderComplexityAccumulatePixelShader->GetPixelShader();
 		}
 #endif
@@ -645,7 +645,7 @@ public:
 
 	void SetMeshRenderState(
 		FRHICommandList& RHICmdList, 
-		const FSceneView& View,
+		const FViewInfo& View,
 		const FPrimitiveSceneProxy* PrimitiveSceneProxy,
 		const FMeshBatch& Mesh,
 		int32 BatchElementIndex,
@@ -685,10 +685,11 @@ public:
 				RHICmdList.SetBlendState(TStaticBlendState<CW_RGB,BO_Add,BF_One,BF_One>::GetRHI());
 			}
 
-			TShaderMapRef<FShaderComplexityAccumulatePS> ShaderComplexityPixelShader(GetGlobalShaderMap());
+			const auto FeatureLevel = View.GetFeatureLevel();
+			TShaderMapRef<FShaderComplexityAccumulatePS> ShaderComplexityPixelShader(View.ShaderMap);
 			const uint32 NumPixelShaderInstructions = PixelShader->GetNumInstructions();
 			const uint32 NumVertexShaderInstructions = VertexShader->GetNumInstructions();
-			ShaderComplexityPixelShader->SetParameters(RHICmdList, NumVertexShaderInstructions,NumPixelShaderInstructions, View.GetFeatureLevel());
+			ShaderComplexityPixelShader->SetParameters(RHICmdList, NumVertexShaderInstructions,NumPixelShaderInstructions, FeatureLevel);
 		}
 		else
 #endif
@@ -759,7 +760,7 @@ public:
 	static void AddStaticMesh(FRHICommandList& RHICmdList, FScene* Scene, FStaticMesh* StaticMesh);
 	static bool DrawDynamicMesh(
 		FRHICommandList& RHICmdList, 
-		const FSceneView& View,
+		const FViewInfo& View,
 		ContextType DrawingContext,
 		const FMeshBatch& Mesh,
 		bool bBackFace,
diff --git a/Engine/Source/Runtime/Renderer/Private/ClearQuad.cpp b/Engine/Source/Runtime/Renderer/Private/ClearQuad.cpp
--- a/Engine/Source/Runtime/Renderer/Private/ClearQuad.cpp
+++ b/Engine/Source/Runtime/Renderer/Private/ClearQuad.cpp
@@ -47,8 +47,11 @@ void DrawClearQuadMRT(FRHICommandList& RHICmdList, ERHIFeatureLevel::Type Featur
 	RHICmdList.SetBlendState(BlendStateRHI);
 	RHICmdList.SetDepthStencilState(DepthStencilStateRHI);
 
+	auto ShaderMap = GetGlobalShaderMap(FeatureLevel);
+
+
 	// Set the new shaders
-	TShaderMapRef<TOneColorVS<true> > VertexShader(GetGlobalShaderMap());
+	TShaderMapRef<TOneColorVS<true> > VertexShader(ShaderMap);
 
 	FOneColorPS* PixelShader = NULL;
 
@@ -56,42 +59,42 @@ void DrawClearQuadMRT(FRHICommandList& RHICmdList, ERHIFeatureLevel::Type Featur
 	// On AMD PC hardware, outputting to a color index in the shader without a matching render target set has a significant performance hit
 	if (NumClearColors <= 1)
 	{
-		TShaderMapRef<TOneColorPixelShaderMRT<1> > MRTPixelShader(GetGlobalShaderMap());
+		TShaderMapRef<TOneColorPixelShaderMRT<1> > MRTPixelShader(ShaderMap);
 		PixelShader = *MRTPixelShader;
 	}
 	else if (NumClearColors == 2)
 	{
-		TShaderMapRef<TOneColorPixelShaderMRT<2> > MRTPixelShader(GetGlobalShaderMap());
+		TShaderMapRef<TOneColorPixelShaderMRT<2> > MRTPixelShader(ShaderMap);
 		PixelShader = *MRTPixelShader;
 	}
 	else if (NumClearColors == 3)
 	{
-		TShaderMapRef<TOneColorPixelShaderMRT<3> > MRTPixelShader(GetGlobalShaderMap());
+		TShaderMapRef<TOneColorPixelShaderMRT<3> > MRTPixelShader(ShaderMap);
 		PixelShader = *MRTPixelShader;
 	}
 	else if (NumClearColors == 4)
 	{
-		TShaderMapRef<TOneColorPixelShaderMRT<4> > MRTPixelShader(GetGlobalShaderMap());
+		TShaderMapRef<TOneColorPixelShaderMRT<4> > MRTPixelShader(ShaderMap);
 		PixelShader = *MRTPixelShader;
 	}
 	else if (NumClearColors == 5)
 	{
-		TShaderMapRef<TOneColorPixelShaderMRT<5> > MRTPixelShader(GetGlobalShaderMap());
+		TShaderMapRef<TOneColorPixelShaderMRT<5> > MRTPixelShader(ShaderMap);
 		PixelShader = *MRTPixelShader;
 	}
 	else if (NumClearColors == 6)
 	{
-		TShaderMapRef<TOneColorPixelShaderMRT<6> > MRTPixelShader(GetGlobalShaderMap());
+		TShaderMapRef<TOneColorPixelShaderMRT<6> > MRTPixelShader(ShaderMap);
 		PixelShader = *MRTPixelShader;
 	}
 	else if (NumClearColors == 7)
 	{
-		TShaderMapRef<TOneColorPixelShaderMRT<7> > MRTPixelShader(GetGlobalShaderMap());
+		TShaderMapRef<TOneColorPixelShaderMRT<7> > MRTPixelShader(ShaderMap);
 		PixelShader = *MRTPixelShader;
 	}
 	else if (NumClearColors == 8)
 	{
-		TShaderMapRef<TOneColorPixelShaderMRT<8> > MRTPixelShader(GetGlobalShaderMap());
+		TShaderMapRef<TOneColorPixelShaderMRT<8> > MRTPixelShader(ShaderMap);
 		PixelShader = *MRTPixelShader;
 	}
 
diff --git a/Engine/Source/Runtime/Renderer/Private/CompositionLighting/PostProcessAmbient.cpp b/Engine/Source/Runtime/Renderer/Private/CompositionLighting/PostProcessAmbient.cpp
--- a/Engine/Source/Runtime/Renderer/Private/CompositionLighting/PostProcessAmbient.cpp
+++ b/Engine/Source/Runtime/Renderer/Private/CompositionLighting/PostProcessAmbient.cpp
@@ -92,7 +92,7 @@ void FRCPassPostProcessAmbient::Process(FRenderingCompositePassContext& Context)
 	Context.RHICmdList.SetRasterizerState(TStaticRasterizerState<>::GetRHI());
 	Context.RHICmdList.SetDepthStencilState(TStaticDepthStencilState<false, CF_Always>::GetRHI());
 
-	TShaderMapRef<FPostProcessVS> VertexShader(GetGlobalShaderMap());
+	TShaderMapRef<FPostProcessVS> VertexShader(Context.GetShaderMap());
 
 	FScene* Scene = ViewFamily.Scene->GetRenderScene();
 	check(Scene);
@@ -100,7 +100,7 @@ void FRCPassPostProcessAmbient::Process(FRenderingCompositePassContext& Context)
 	
 	// Ambient cubemap specular will be applied in the reflection environment pass if it is enabled
 	const bool bApplySpecular = View.Family->EngineShowFlags.ReflectionEnvironment == 0 || NumReflectionCaptures == 0;
-	TShaderMapRef<FPostProcessAmbientPS> PixelShader(GetGlobalShaderMap());
+	TShaderMapRef<FPostProcessAmbientPS> PixelShader(Context.GetShaderMap());
 
 	static FGlobalBoundShaderState BoundShaderState;
 	
diff --git a/Engine/Source/Runtime/Renderer/Private/CompositionLighting/PostProcessAmbientOcclusion.cpp b/Engine/Source/Runtime/Renderer/Private/CompositionLighting/PostProcessAmbientOcclusion.cpp
--- a/Engine/Source/Runtime/Renderer/Private/CompositionLighting/PostProcessAmbientOcclusion.cpp
+++ b/Engine/Source/Runtime/Renderer/Private/CompositionLighting/PostProcessAmbientOcclusion.cpp
@@ -347,8 +347,8 @@ IMPLEMENT_SHADER_TYPE(,FPostProcessBasePassAOPS,TEXT("PostProcessAmbientOcclusio
 template <uint32 bInitialSetup>
 void FRCPassPostProcessAmbientOcclusionSetup::SetShaderSetupTempl(const FRenderingCompositePassContext& Context)
 {
-	TShaderMapRef<FPostProcessVS> VertexShader(GetGlobalShaderMap());
-	TShaderMapRef<FPostProcessAmbientOcclusionSetupPS<bInitialSetup> > PixelShader(GetGlobalShaderMap());
+	TShaderMapRef<FPostProcessVS> VertexShader(Context.GetShaderMap());
+	TShaderMapRef<FPostProcessAmbientOcclusionSetupPS<bInitialSetup> > PixelShader(Context.GetShaderMap());
 
 	static FGlobalBoundShaderState BoundShaderState;
 	
@@ -394,7 +394,7 @@ void FRCPassPostProcessAmbientOcclusionSetup::Process(FRenderingCompositePassCon
 		SetShaderSetupTempl<0>(Context);
 	}
 
-	TShaderMapRef<FPostProcessVS> VertexShader(GetGlobalShaderMap());
+	TShaderMapRef<FPostProcessVS> VertexShader(Context.GetShaderMap());
 
 	// Draw a quad mapping scene color to the view's render target
 	DrawRectangle( 
@@ -458,8 +458,8 @@ bool FRCPassPostProcessAmbientOcclusionSetup::IsInitialPass() const
 template <uint32 bTAOSetupAsInput, uint32 bDoUpsample, uint32 SampleSetQuality>
 void FRCPassPostProcessAmbientOcclusion::SetShaderTempl(const FRenderingCompositePassContext& Context)
 {
-	TShaderMapRef<FPostProcessVS> VertexShader(GetGlobalShaderMap());
-	TShaderMapRef<FPostProcessAmbientOcclusionPS<bTAOSetupAsInput, bDoUpsample, SampleSetQuality> > PixelShader(GetGlobalShaderMap());
+	TShaderMapRef<FPostProcessVS> VertexShader(Context.GetShaderMap());
+	TShaderMapRef<FPostProcessAmbientOcclusionPS<bTAOSetupAsInput, bDoUpsample, SampleSetQuality> > PixelShader(Context.GetShaderMap());
 
 	static FGlobalBoundShaderState BoundShaderState;
 	
@@ -548,7 +548,7 @@ void FRCPassPostProcessAmbientOcclusion::Process(FRenderingCompositePassContext&
 		}
 	}
 
-	TShaderMapRef<FPostProcessVS> VertexShader(GetGlobalShaderMap());
+	TShaderMapRef<FPostProcessVS> VertexShader(Context.GetShaderMap());
 
 	// Draw a quad mapping scene color to the view's render target
 	DrawRectangle( 
@@ -606,8 +606,8 @@ void FRCPassPostProcessBasePassAO::Process(FRenderingCompositePassContext& Conte
 	Context.RHICmdList.SetRasterizerState(TStaticRasterizerState<>::GetRHI());
 	Context.RHICmdList.SetDepthStencilState(TStaticDepthStencilState<false, CF_Always>::GetRHI());
 
-	TShaderMapRef<FPostProcessVS> VertexShader(GetGlobalShaderMap());
-	TShaderMapRef<FPostProcessBasePassAOPS> PixelShader(GetGlobalShaderMap());
+	TShaderMapRef<FPostProcessVS> VertexShader(Context.GetShaderMap());
+	TShaderMapRef<FPostProcessBasePassAOPS> PixelShader(Context.GetShaderMap());
 
 	static FGlobalBoundShaderState BoundShaderState;
 	
diff --git a/Engine/Source/Runtime/Renderer/Private/CompositionLighting/PostProcessDeferredDecals.cpp b/Engine/Source/Runtime/Renderer/Private/CompositionLighting/PostProcessDeferredDecals.cpp
--- a/Engine/Source/Runtime/Renderer/Private/CompositionLighting/PostProcessDeferredDecals.cpp
+++ b/Engine/Source/Runtime/Renderer/Private/CompositionLighting/PostProcessDeferredDecals.cpp
@@ -351,7 +351,7 @@ IMPLEMENT_SHADER_TYPE(,FStencilDecalMaskPS,TEXT("DeferredDecal"),TEXT("StencilDe
 FGlobalBoundShaderState StencilDecalMaskBoundShaderState;
 
 /** Draws a full view quad that sets stencil to 1 anywhere that decals should not be projected. */
-void StencilDecalMask(FRHICommandList& RHICmdList, const FSceneView& View)
+void StencilDecalMask(FRHICommandList& RHICmdList, const FViewInfo& View)
 {
 	SCOPED_DRAW_EVENT(StencilDecalMask, DEC_SCENE_ITEMS);
 	RHICmdList.SetRasterizerState(TStaticRasterizerState<FM_Solid, CM_None>::GetRHI());
@@ -362,10 +362,12 @@ void StencilDecalMask(FRHICommandList& RHICmdList, const FSceneView& View)
 	// Write 1 to highest bit of stencil to areas that should not receive decals
 	RHICmdList.SetDepthStencilState(TStaticDepthStencilState<false, CF_Always, true, CF_Always, SO_Replace, SO_Replace, SO_Replace>::GetRHI(), 0x80);
 
-	TShaderMapRef<FScreenVS> ScreenVertexShader(GetGlobalShaderMap());
-	TShaderMapRef<FStencilDecalMaskPS> PixelShader(GetGlobalShaderMap());
+	const auto FeatureLevel = View.GetFeatureLevel();
+	auto ShaderMap = View.ShaderMap;
+	TShaderMapRef<FScreenVS> ScreenVertexShader(ShaderMap);
+	TShaderMapRef<FStencilDecalMaskPS> PixelShader(ShaderMap);
 	
-	SetGlobalBoundShaderState(RHICmdList, View.GetFeatureLevel(), StencilDecalMaskBoundShaderState, GFilterVertexDeclaration.VertexDeclarationRHI, *ScreenVertexShader, *PixelShader);
+	SetGlobalBoundShaderState(RHICmdList, FeatureLevel, StencilDecalMaskBoundShaderState, GFilterVertexDeclaration.VertexDeclarationRHI, *ScreenVertexShader, *PixelShader);
 
 	PixelShader->SetParameters(RHICmdList, View);
 
@@ -746,7 +748,7 @@ void FRCPassPostProcessDeferredDecals::Process(FRenderingCompositePassContext& C
 	}
 
 	// this cast is safe as only the dedicated server implements this differently and this pass should not be executed on the dedicated server
-	const FSceneView& View = Context.View;
+	const FViewInfo& View = Context.View;
 	const FSceneViewFamily& ViewFamily = *(View.Family);
 	FScene& Scene = *(FScene*)ViewFamily.Scene;
 
@@ -863,7 +865,7 @@ void FRCPassPostProcessDeferredDecals::Process(FRenderingCompositePassContext& C
 		int32 LastDecalHasNormal = -1; // Decal state can change based on its normal property.(SM5)
 		ERenderTargetMode LastRenderTargetMode = RTM_Unknown;
 		int32 WasInsideDecal = -1;
-		const ERHIFeatureLevel::Type SMFeatureLevel = GetMaxSupportedFeatureLevel(GRHIShaderPlatform);
+		const ERHIFeatureLevel::Type SMFeatureLevel = Context.GetFeatureLevel();
 
 		SCOPED_DRAW_EVENT(Decals, DEC_SCENE_ITEMS);
 		INC_DWORD_STAT_BY(STAT_Decals, SortedDecals.Num());
diff --git a/Engine/Source/Runtime/Renderer/Private/CompositionLighting/PostProcessLpvIndirect.cpp b/Engine/Source/Runtime/Renderer/Private/CompositionLighting/PostProcessLpvIndirect.cpp
--- a/Engine/Source/Runtime/Renderer/Private/CompositionLighting/PostProcessLpvIndirect.cpp
+++ b/Engine/Source/Runtime/Renderer/Private/CompositionLighting/PostProcessLpvIndirect.cpp
@@ -169,7 +169,7 @@ void FRCPassPostProcessLpvIndirect::Process(FRenderingCompositePassContext& Cont
 	Context.RHICmdList.SetRasterizerState(TStaticRasterizerState<>::GetRHI());
 	Context.RHICmdList.SetDepthStencilState(TStaticDepthStencilState<false, CF_Always>::GetRHI());
 
-	TShaderMapRef<FPostProcessVS> VertexShader(GetGlobalShaderMap());
+	TShaderMapRef<FPostProcessVS> VertexShader(Context.GetShaderMap());
 
 	FSceneViewState* ViewState = (FSceneViewState*)View.State;
 	FLightPropagationVolume* Lpv = NULL;
@@ -186,7 +186,7 @@ void FRCPassPostProcessLpvIndirect::Process(FRenderingCompositePassContext& Cont
 		return;
 	}
 
-	TShaderMapRef<FPostProcessLpvIndirectPS> PixelShader(GetGlobalShaderMap());
+	TShaderMapRef<FPostProcessLpvIndirectPS> PixelShader(Context.GetShaderMap());
 
 	static FGlobalBoundShaderState BoundShaderState;
 	
diff --git a/Engine/Source/Runtime/Renderer/Private/CompositionLighting/PostProcessPassThrough.cpp b/Engine/Source/Runtime/Renderer/Private/CompositionLighting/PostProcessPassThrough.cpp
--- a/Engine/Source/Runtime/Renderer/Private/CompositionLighting/PostProcessPassThrough.cpp
+++ b/Engine/Source/Runtime/Renderer/Private/CompositionLighting/PostProcessPassThrough.cpp
@@ -115,8 +115,8 @@ void FRCPassPostProcessPassThrough::Process(FRenderingCompositePassContext& Cont
 	Context.RHICmdList.SetRasterizerState(TStaticRasterizerState<>::GetRHI());
 	Context.RHICmdList.SetDepthStencilState(TStaticDepthStencilState<false, CF_Always>::GetRHI());
 
-	TShaderMapRef<FPostProcessVS> VertexShader(GetGlobalShaderMap());
-	TShaderMapRef<FPostProcessPassThroughPS> PixelShader(GetGlobalShaderMap());
+	TShaderMapRef<FPostProcessVS> VertexShader(Context.GetShaderMap());
+	TShaderMapRef<FPostProcessPassThroughPS> PixelShader(Context.GetShaderMap());
 
 	static FGlobalBoundShaderState BoundShaderState;
 	
diff --git a/Engine/Source/Runtime/Renderer/Private/DeferredShadingRenderer.cpp b/Engine/Source/Runtime/Renderer/Private/DeferredShadingRenderer.cpp
--- a/Engine/Source/Runtime/Renderer/Private/DeferredShadingRenderer.cpp
+++ b/Engine/Source/Runtime/Renderer/Private/DeferredShadingRenderer.cpp
@@ -146,7 +146,9 @@ void FDeferredShadingSceneRenderer::ClearGBufferAtMaxZ(FRHICommandList& RHICmdLi
 
 	uint32 NumActiveRenderTargets = GSceneRenderTargets.GetNumGBufferTargets();
 	
-	TShaderMapRef<TOneColorVS<true> > VertexShader(GetGlobalShaderMap());
+	auto ShaderMap = GetGlobalShaderMap(FeatureLevel);
+
+	TShaderMapRef<TOneColorVS<true> > VertexShader(ShaderMap);
 	FOneColorPS* PixelShader = NULL; 
 
 	// Assume for now all code path supports SM4, otherwise render target numbers are changed
@@ -154,20 +156,20 @@ void FDeferredShadingSceneRenderer::ClearGBufferAtMaxZ(FRHICommandList& RHICmdLi
 	{
 	case 5:
 		{
-			TShaderMapRef<TOneColorPixelShaderMRT<5> > MRTPixelShader(GetGlobalShaderMap());
+			TShaderMapRef<TOneColorPixelShaderMRT<5> > MRTPixelShader(ShaderMap);
 			PixelShader = *MRTPixelShader;
 		}
 		break;
 	case 6:
 		{
-			TShaderMapRef<TOneColorPixelShaderMRT<6> > MRTPixelShader(GetGlobalShaderMap());
+			TShaderMapRef<TOneColorPixelShaderMRT<6> > MRTPixelShader(ShaderMap);
 			PixelShader = *MRTPixelShader;
 		}
 		break;
 	default:
 	case 1:
 		{
-			TShaderMapRef<TOneColorPixelShaderMRT<1> > MRTPixelShader(GetGlobalShaderMap());
+			TShaderMapRef<TOneColorPixelShaderMRT<1> > MRTPixelShader(ShaderMap);
 			PixelShader = *MRTPixelShader;
 		}
 		break;
@@ -1257,7 +1259,7 @@ void FDeferredShadingSceneRenderer::ClearLPVs(FRHICommandListImmediate& RHICmdLi
 				SCOPED_DRAW_EVENT(ClearLPVs, DEC_SCENE_ITEMS);
 				SCOPE_CYCLE_COUNTER(STAT_UpdateLPVs);
 				LightPropagationVolume->InitSettings(RHICmdList, Views[ViewIndex]);
-				LightPropagationVolume->Clear(RHICmdList);
+				LightPropagationVolume->Clear(RHICmdList, View);
 			}
 		}
 	}
@@ -1279,7 +1281,7 @@ void FDeferredShadingSceneRenderer::PropagateLPVs(FRHICommandListImmediate& RHIC
 				SCOPED_DRAW_EVENT(UpdateLPVs, DEC_SCENE_ITEMS);
 				SCOPE_CYCLE_COUNTER(STAT_UpdateLPVs);
 				
-				LightPropagationVolume->Propagate(RHICmdList);
+				LightPropagationVolume->Propagate(RHICmdList, View);
 			}
 		}
 	}
@@ -1360,8 +1362,8 @@ void FDeferredShadingSceneRenderer::UpdateDownsampledDepthSurface(FRHICommandLis
 		{
 			const FViewInfo& View = Views[ViewIndex];
 			// Set shaders and texture
-			TShaderMapRef<FScreenVS> ScreenVertexShader(GetGlobalShaderMap());
-			TShaderMapRef<FDownsampleSceneDepthPS> PixelShader(GetGlobalShaderMap());
+			TShaderMapRef<FScreenVS> ScreenVertexShader(View.ShaderMap);
+			TShaderMapRef<FDownsampleSceneDepthPS> PixelShader(View.ShaderMap);
 
 			extern TGlobalResource<FFilterVertexDeclaration> GFilterVertexDeclaration;
 
diff --git a/Engine/Source/Runtime/Renderer/Private/DistanceFieldSurfaceCacheLighting.cpp b/Engine/Source/Runtime/Renderer/Private/DistanceFieldSurfaceCacheLighting.cpp
--- a/Engine/Source/Runtime/Renderer/Private/DistanceFieldSurfaceCacheLighting.cpp
+++ b/Engine/Source/Runtime/Renderer/Private/DistanceFieldSurfaceCacheLighting.cpp
@@ -957,7 +957,7 @@ void ComputeDistanceFieldNormal(FRHICommandListImmediate& RHICmdList, const TArr
 
 			{
 				SCOPED_DRAW_EVENT(ComputeNormalCS, DEC_SCENE_ITEMS);
-				TShaderMapRef<FComputeDistanceFieldNormalCS> ComputeShader(GetGlobalShaderMap());
+				TShaderMapRef<FComputeDistanceFieldNormalCS> ComputeShader(View.ShaderMap);
 
 				RHICmdList.SetComputeShader(ComputeShader->GetComputeShader());
 				ComputeShader->SetParameters(RHICmdList, View, DistanceFieldNormal, NumObjects);
@@ -981,13 +981,13 @@ void ComputeDistanceFieldNormal(FRHICommandListImmediate& RHICmdList, const TArr
 			RHICmdList.SetRasterizerState(TStaticRasterizerState<FM_Solid, CM_None>::GetRHI());
 			RHICmdList.SetDepthStencilState(TStaticDepthStencilState<false, CF_Always>::GetRHI());
 			RHICmdList.SetBlendState(TStaticBlendState<>::GetRHI());
-
-			TShaderMapRef<FPostProcessVS> VertexShader( GetGlobalShaderMap() );
-			TShaderMapRef<FComputeDistanceFieldNormalPS> PixelShader( GetGlobalShaderMap() );
+			
+			TShaderMapRef<FPostProcessVS> VertexShader(View.ShaderMap);
+			TShaderMapRef<FComputeDistanceFieldNormalPS> PixelShader(View.ShaderMap);
 
 			static FGlobalBoundShaderState BoundShaderState;
 			
-			SetGlobalBoundShaderState(RHICmdList, View.GetFeatureLevel(), BoundShaderState, GFilterVertexDeclaration.VertexDeclarationRHI, *VertexShader, *PixelShader);
+			SetGlobalBoundShaderState(RHICmdList, View.FeatureLevel, BoundShaderState, GFilterVertexDeclaration.VertexDeclarationRHI, *VertexShader, *PixelShader);
 
 			PixelShader->SetParameters(RHICmdList, View, NumObjects);
 
@@ -2082,8 +2082,8 @@ void UpdateHistory(
 				RHICmdList.SetDepthStencilState(TStaticDepthStencilState<false, CF_Always>::GetRHI());
 				RHICmdList.SetBlendState(TStaticBlendState<>::GetRHI());
 
-				TShaderMapRef<FPostProcessVS> VertexShader( GetGlobalShaderMap() );
-				TShaderMapRef<FUpdateHistoryPS> PixelShader( GetGlobalShaderMap() );
+				TShaderMapRef<FPostProcessVS> VertexShader(View.ShaderMap);
+				TShaderMapRef<FUpdateHistoryPS> PixelShader(View.ShaderMap);
 
 				static FGlobalBoundShaderState BoundShaderState;
 				
@@ -2150,8 +2150,8 @@ void PostProcessBentNormalAO(FRHICommandList& RHICmdList, TArray<FViewInfo>& Vie
 			RHICmdList.SetDepthStencilState(TStaticDepthStencilState<false, CF_Always>::GetRHI());
 			RHICmdList.SetBlendState(TStaticBlendState<>::GetRHI());
 
-			TShaderMapRef<FPostProcessVS> VertexShader(GetGlobalShaderMap());
-			TShaderMapRef<FDistanceFieldAOCombinePS2> PixelShader(GetGlobalShaderMap());
+			TShaderMapRef<FPostProcessVS> VertexShader(View.ShaderMap);
+			TShaderMapRef<FDistanceFieldAOCombinePS2> PixelShader(View.ShaderMap);
 
 			static FGlobalBoundShaderState BoundShaderState;
 			
@@ -2185,8 +2185,8 @@ void PostProcessBentNormalAO(FRHICommandList& RHICmdList, TArray<FViewInfo>& Vie
 			RHICmdList.SetDepthStencilState(TStaticDepthStencilState<false, CF_Always>::GetRHI());
 			RHICmdList.SetBlendState(TStaticBlendState<>::GetRHI());
 
-			TShaderMapRef<FPostProcessVS> VertexShader( GetGlobalShaderMap() );
-			TShaderMapRef<FFillGapsPS> PixelShader( GetGlobalShaderMap() );
+			TShaderMapRef<FPostProcessVS> VertexShader(View.ShaderMap);
+			TShaderMapRef<FFillGapsPS> PixelShader(View.ShaderMap);
 
 			static FGlobalBoundShaderState BoundShaderState;
 			
@@ -2293,8 +2293,8 @@ void UpsampleBentNormalAO(FRHICommandList& RHICmdList, const TArray<FViewInfo>&
 		RHICmdList.SetDepthStencilState(TStaticDepthStencilState<false, CF_Always>::GetRHI());
 		RHICmdList.SetBlendState(TStaticBlendState<>::GetRHI());
 
-		TShaderMapRef<FPostProcessVS> VertexShader( GetGlobalShaderMap() );
-		TShaderMapRef<FDistanceFieldAOUpsamplePS> PixelShader( GetGlobalShaderMap() );
+		TShaderMapRef<FPostProcessVS> VertexShader(View.ShaderMap);
+		TShaderMapRef<FDistanceFieldAOUpsamplePS> PixelShader(View.ShaderMap);
 
 		static FGlobalBoundShaderState BoundShaderState;
 		
@@ -2509,7 +2509,7 @@ FIntPoint BuildTileObjectLists(FRHICommandListImmediate& RHICmdList, FScene* Sce
 
 			{
 				SCOPED_DRAW_EVENT(BuildTileCones, DEC_SCENE_ITEMS);
-				TShaderMapRef<FBuildTileConesCS> ComputeShader(GetGlobalShaderMap());
+				TShaderMapRef<FBuildTileConesCS> ComputeShader(View.ShaderMap);
 
 				RHICmdList.SetComputeShader(ComputeShader->GetComputeShader());
 				ComputeShader->SetParameters(RHICmdList, View, Scene, FVector2D(GroupSizeX, GroupSizeY));
@@ -2521,8 +2521,8 @@ FIntPoint BuildTileObjectLists(FRHICommandListImmediate& RHICmdList, FScene* Sce
 			{
 				SCOPED_DRAW_EVENT(CullObjects, DEC_SCENE_ITEMS);
 
-				TShaderMapRef<FObjectCullVS> VertexShader(GetGlobalShaderMap());
-				TShaderMapRef<FObjectCullPS> PixelShader(GetGlobalShaderMap());
+				TShaderMapRef<FObjectCullVS> VertexShader(View.ShaderMap);
+				TShaderMapRef<FObjectCullPS> PixelShader(View.ShaderMap);
 
 				TArray<FUnorderedAccessViewRHIParamRef> UAVs;
 				PixelShader->GetUAVs(Views[0], UAVs);
@@ -2589,7 +2589,7 @@ FIntPoint BuildTileObjectLists(FRHICommandListImmediate& RHICmdList, FScene* Sce
 			RHICmdList.ClearUAV(TileIntersectionResources->TileArrayNextAllocation.UAV, ClearValues);
 			RHICmdList.ClearUAV(TileIntersectionResources->TileHeadData.UAV, ClearValues);
 
-			TShaderMapRef<FDistanceFieldBuildTileListCS > ComputeShader(GetGlobalShaderMap());
+			TShaderMapRef<FDistanceFieldBuildTileListCS > ComputeShader(View.ShaderMap);
 
 			RHICmdList.SetComputeShader(ComputeShader->GetComputeShader());
 			ComputeShader->SetParameters(RHICmdList, View, Scene, FVector2D(GroupSizeX, GroupSizeY), NumObjects);
@@ -2654,8 +2654,8 @@ void RenderIrradianceCacheInterpolation(FRHICommandListImmediate& RHICmdList, TA
 
 			if (bFinalInterpolation)
 			{
-				TShaderMapRef<TIrradianceCacheSplatVS<true> > VertexShader(GetGlobalShaderMap());
-				TShaderMapRef<TIrradianceCacheSplatPS<true> > PixelShader(GetGlobalShaderMap());
+				TShaderMapRef<TIrradianceCacheSplatVS<true> > VertexShader(View.ShaderMap);
+				TShaderMapRef<TIrradianceCacheSplatPS<true> > PixelShader(View.ShaderMap);
 
 				for (int32 SplatSourceDepthLevel = MaxAllowedLevel; SplatSourceDepthLevel >= DepthLevel; SplatSourceDepthLevel--)
 				{
@@ -2671,8 +2671,8 @@ void RenderIrradianceCacheInterpolation(FRHICommandListImmediate& RHICmdList, TA
 			}
 			else
 			{
-				TShaderMapRef<TIrradianceCacheSplatVS<false> > VertexShader(GetGlobalShaderMap());
-				TShaderMapRef<TIrradianceCacheSplatPS<false> > PixelShader(GetGlobalShaderMap());
+				TShaderMapRef<TIrradianceCacheSplatVS<false> > VertexShader(View.ShaderMap);
+				TShaderMapRef<TIrradianceCacheSplatPS<false> > PixelShader(View.ShaderMap);
 
 				for (int32 SplatSourceDepthLevel = MaxAllowedLevel; SplatSourceDepthLevel >= DepthLevel; SplatSourceDepthLevel--)
 				{
@@ -2891,7 +2891,7 @@ bool FDeferredShadingSceneRenderer::RenderDistanceFieldAOSurfaceCache(FRHIComman
 					for (int32 DepthLevel = FMath::Min(GAOMaxLevel, MaxAllowedLevel); DepthLevel >= FMath::Max(GAOMinLevel, 0); DepthLevel--)
 					{
 						{	
-							TShaderMapRef<FSetupCopyIndirectArgumentsCS> ComputeShader(GetGlobalShaderMap());
+							TShaderMapRef<FSetupCopyIndirectArgumentsCS> ComputeShader(View.ShaderMap);
 
 							RHICmdList.SetComputeShader(ComputeShader->GetComputeShader());
 							ComputeShader->SetParameters(RHICmdList, View, DepthLevel);
@@ -2901,7 +2901,7 @@ bool FDeferredShadingSceneRenderer::RenderDistanceFieldAOSurfaceCache(FRHIComman
 						}
 
 						{
-							TShaderMapRef<FCopyIrradianceCacheSamplesCS> ComputeShader(GetGlobalShaderMap());
+							TShaderMapRef<FCopyIrradianceCacheSamplesCS> ComputeShader(View.ShaderMap);
 
 							RHICmdList.SetComputeShader(ComputeShader->GetComputeShader());
 							ComputeShader->SetParameters(RHICmdList, View, DepthLevel);
@@ -2940,7 +2940,7 @@ bool FDeferredShadingSceneRenderer::RenderDistanceFieldAOSurfaceCache(FRHIComman
 
 						// Save off the current record count before adding any more
 						{	
-							TShaderMapRef<FSaveStartIndexCS> ComputeShader(GetGlobalShaderMap());
+							TShaderMapRef<FSaveStartIndexCS> ComputeShader(View.ShaderMap);
 
 							RHICmdList.SetComputeShader(ComputeShader->GetComputeShader());
 							ComputeShader->SetParameters(RHICmdList, View, DepthLevel);
@@ -2958,7 +2958,7 @@ bool FDeferredShadingSceneRenderer::RenderDistanceFieldAOSurfaceCache(FRHIComman
 							uint32 GroupSizeX = FMath::DivideAndRoundUp(DownsampledViewSize.X, GDistanceFieldAOTileSizeX);
 							uint32 GroupSizeY = FMath::DivideAndRoundUp(DownsampledViewSize.Y, GDistanceFieldAOTileSizeY);
 
-							TShaderMapRef<FPopulateCacheCS> ComputeShader(GetGlobalShaderMap());
+							TShaderMapRef<FPopulateCacheCS> ComputeShader(View.ShaderMap);
 
 							RHICmdList.SetComputeShader(ComputeShader->GetComputeShader());
 							ComputeShader->SetParameters(RHICmdList, View, DistanceFieldAOIrradianceCacheSplat->GetRenderTargetItem(), DistanceFieldNormal->GetRenderTargetItem(), DestLevelDownsampleFactor, DepthLevel, TileListGroupSize);
@@ -2972,7 +2972,7 @@ bool FDeferredShadingSceneRenderer::RenderDistanceFieldAOSurfaceCache(FRHIComman
 						SCOPED_DRAW_EVENT(ShadeIrradianceCache, DEC_SCENE_ITEMS);
 
 						{	
-							TShaderMapRef<FSetupFinalGatherIndirectArgumentsCS> ComputeShader(GetGlobalShaderMap());
+							TShaderMapRef<FSetupFinalGatherIndirectArgumentsCS> ComputeShader(View.ShaderMap);
 
 							RHICmdList.SetComputeShader(ComputeShader->GetComputeShader());
 							ComputeShader->SetParameters(RHICmdList, View, DepthLevel);
@@ -2983,7 +2983,7 @@ bool FDeferredShadingSceneRenderer::RenderDistanceFieldAOSurfaceCache(FRHIComman
 
 						// Compute lighting for the new surface cache records by cone-stepping through the object distance fields
 						{
-							TShaderMapRef<FFinalGatherCS> ComputeShader(GetGlobalShaderMap());
+							TShaderMapRef<FFinalGatherCS> ComputeShader(View.ShaderMap);
 
 							RHICmdList.SetComputeShader(ComputeShader->GetComputeShader());
 							ComputeShader->SetParameters(RHICmdList, View, DistanceFieldAOIrradianceCacheSplat->GetRenderTargetItem(), DistanceFieldNormal->GetRenderTargetItem(), NumObjects, DestLevelDownsampleFactor, DepthLevel, TileListGroupSize);
@@ -3139,11 +3139,11 @@ void FDeferredShadingSceneRenderer::RenderDynamicSkyLighting(FRHICommandListImme
 			RHICmdList.SetDepthStencilState(TStaticDepthStencilState< false, CF_Always >::GetRHI());
 			RHICmdList.SetBlendState(TStaticBlendState< CW_RGB, BO_Add, BF_One, BF_One >::GetRHI());
 
-			TShaderMapRef< FPostProcessVS > VertexShader( GetGlobalShaderMap() );
+			TShaderMapRef< FPostProcessVS > VertexShader(View.ShaderMap);
 
 			if (bApplyShadowing)
 			{
-				TShaderMapRef<TDynamicSkyLightDiffusePS<true> > PixelShader( GetGlobalShaderMap() );
+				TShaderMapRef<TDynamicSkyLightDiffusePS<true> > PixelShader(View.ShaderMap);
 
 				static FGlobalBoundShaderState BoundShaderState;
 				
@@ -3153,7 +3153,7 @@ void FDeferredShadingSceneRenderer::RenderDynamicSkyLighting(FRHICommandListImme
 			}
 			else
 			{
-				TShaderMapRef<TDynamicSkyLightDiffusePS<false> > PixelShader( GetGlobalShaderMap() );
+				TShaderMapRef<TDynamicSkyLightDiffusePS<false> > PixelShader(View.ShaderMap);
 
 				static FGlobalBoundShaderState BoundShaderState;
 				
diff --git a/Engine/Source/Runtime/Renderer/Private/DistortionRendering.cpp b/Engine/Source/Runtime/Renderer/Private/DistortionRendering.cpp
--- a/Engine/Source/Runtime/Renderer/Private/DistortionRendering.cpp
+++ b/Engine/Source/Runtime/Renderer/Private/DistortionRendering.cpp
@@ -369,7 +369,7 @@ TDistortionMeshDrawingPolicy<DistortMeshPolicy>::TDistortionMeshDrawingPolicy(
 
 	if (bInitializeOffsets)
 	{
-//later		InitializePixelShader = GetGlobalShaderMap()->GetShader<FShaderComplexityAccumulatePixelShader>();
+//later		InitializePixelShader = GetGlobalShaderMap(View.ShaderMap)->GetShader<FShaderComplexityAccumulatePixelShader>();
 		DistortPixelShader = NULL;
 	}
 	else
@@ -424,7 +424,7 @@ void TDistortionMeshDrawingPolicy<DistortMeshPolicy>::SetSharedState(
 	if (bOverrideWithShaderComplexity)
 	{
 		check(!bInitializeOffsets);
-//later		TShaderMapRef<FShaderComplexityAccumulatePixelShader> ShaderComplexityPixelShader(GetGlobalShaderMap());
+//later		TShaderMapRef<FShaderComplexityAccumulatePixelShader> ShaderComplexityPixelShader(GetGlobalShaderMap(View->ShaderMap));
 		//don't add any vertex complexity
 //later		ShaderComplexityPixelShader->SetParameters(0, DistortPixelShader->GetNumInstructions());
 	}
@@ -451,7 +451,7 @@ FBoundShaderStateInput TDistortionMeshDrawingPolicy<DistortMeshPolicy>::GetBound
 	if (bOverrideWithShaderComplexity)
 	{
 		check(!bInitializeOffsets);
-//later		TShaderMapRef<FShaderComplexityAccumulatePixelShader> ShaderComplexityAccumulatePixelShader(GetGlobalShaderMap());
+//later		TShaderMapRef<FShaderComplexityAccumulatePixelShader> ShaderComplexityAccumulatePixelShader(GetGlobalShaderMap(InFeatureLevel));
 //later		PixelShaderRHIRef = ShaderComplexityAccumulatePixelShader->GetPixelShader();
 	}
 
@@ -895,8 +895,8 @@ void FSceneRenderer::RenderDistortion(FRHICommandListImmediate& RHICmdList)
 				// Set the view family's render target/viewport.
 				Context.SetViewportAndCallRHI(View.ViewRect);
 
-				TShaderMapRef<FPostProcessVS> VertexShader(GetGlobalShaderMap());
-				TShaderMapRef<FDistortionApplyScreenPS> PixelShader(GetGlobalShaderMap());
+				TShaderMapRef<FPostProcessVS> VertexShader(View.ShaderMap);
+				TShaderMapRef<FDistortionApplyScreenPS> PixelShader(View.ShaderMap);
 
 				static FGlobalBoundShaderState BoundShaderState;
 				
diff --git a/Engine/Source/Runtime/Renderer/Private/FogRendering.cpp b/Engine/Source/Runtime/Renderer/Private/FogRendering.cpp
--- a/Engine/Source/Runtime/Renderer/Private/FogRendering.cpp
+++ b/Engine/Source/Runtime/Renderer/Private/FogRendering.cpp
@@ -280,8 +280,8 @@ void SetFogShaders(FRHICommandList& RHICmdList, FScene* Scene, const FViewInfo&
 {
 	if (Scene->ExponentialFogs.Num() > 0)
 	{
-		TShaderMapRef<FHeightFogVS> VertexShader(GetGlobalShaderMap());
-		TShaderMapRef<FExponentialHeightFogPS> ExponentialHeightFogPixelShader(GetGlobalShaderMap());
+		TShaderMapRef<FHeightFogVS> VertexShader(View.ShaderMap);
+		TShaderMapRef<FExponentialHeightFogPS> ExponentialHeightFogPixelShader(View.ShaderMap);
 
 		SetGlobalBoundShaderState(RHICmdList, View.GetFeatureLevel(), ExponentialBoundShaderState, GFogVertexDeclaration.VertexDeclarationRHI, *VertexShader, *ExponentialHeightFogPixelShader);
 		VertexShader->SetParameters(RHICmdList, View);
diff --git a/Engine/Source/Runtime/Renderer/Private/ForwardBasePassRendering.cpp b/Engine/Source/Runtime/Renderer/Private/ForwardBasePassRendering.cpp
--- a/Engine/Source/Runtime/Renderer/Private/ForwardBasePassRendering.cpp
+++ b/Engine/Source/Runtime/Renderer/Private/ForwardBasePassRendering.cpp
@@ -121,7 +121,7 @@ class FDrawBasePassForwardShadingDynamicMeshAction
 {
 public:
 
-	const FSceneView& View;
+	const FViewInfo& View;
 	bool bBackFace;
 	FHitProxyId HitProxyId;
 
@@ -137,7 +137,7 @@ public:
 
 	/** Initialization constructor. */
 	FDrawBasePassForwardShadingDynamicMeshAction(
-		const FSceneView& InView,
+		const FViewInfo& InView,
 		const bool bInBackFace,
 		const FHitProxyId InHitProxyId
 		)
@@ -205,7 +205,7 @@ public:
 
 bool FBasePassForwardOpaqueDrawingPolicyFactory::DrawDynamicMesh(
 	FRHICommandList& RHICmdList, 
-	const FSceneView& View,
+	const FViewInfo& View,
 	ContextType DrawingContext,
 	const FMeshBatch& Mesh,
 	bool bBackFace,
diff --git a/Engine/Source/Runtime/Renderer/Private/ForwardBasePassRendering.h b/Engine/Source/Runtime/Renderer/Private/ForwardBasePassRendering.h
--- a/Engine/Source/Runtime/Renderer/Private/ForwardBasePassRendering.h
+++ b/Engine/Source/Runtime/Renderer/Private/ForwardBasePassRendering.h
@@ -264,7 +264,7 @@ public:
 			SceneTextureMode == Other.SceneTextureMode;
 	}
 
-	void SetSharedState(FRHICommandList& RHICmdList, const FSceneView* View) const
+	void SetSharedState(FRHICommandList& RHICmdList, const FViewInfo* View) const
 	{
 		VertexShader->SetParameters(RHICmdList, MaterialRenderProxy, VertexFactory, *MaterialResource, *View, SceneTextureMode);
 
@@ -281,7 +281,7 @@ public:
 				RHICmdList.SetBlendState( TStaticBlendState<CW_RGBA,BO_Add,BF_One,BF_One,BO_Add,BF_Zero,BF_One>::GetRHI());
 			}
 
-			TShaderMapRef<FShaderComplexityAccumulatePS> ShaderComplexityPixelShader(GetGlobalShaderMap());
+			TShaderMapRef<FShaderComplexityAccumulatePS> ShaderComplexityPixelShader(View->ShaderMap);
 			const uint32 NumPixelShaderInstructions = PixelShader->GetNumInstructions();
 			const uint32 NumVertexShaderInstructions = VertexShader->GetNumInstructions();
 			ShaderComplexityPixelShader->SetParameters(RHICmdList, NumVertexShaderInstructions,NumPixelShaderInstructions, View->GetFeatureLevel());
@@ -330,7 +330,7 @@ public:
 #if !(UE_BUILD_SHIPPING || UE_BUILD_TEST)
 		if (bOverrideWithShaderComplexity)
 		{
-			TShaderMapRef<FShaderComplexityAccumulatePS> ShaderComplexityAccumulatePixelShader(GetGlobalShaderMap());
+			TShaderMapRef<FShaderComplexityAccumulatePS> ShaderComplexityAccumulatePixelShader(GetGlobalShaderMap(InFeatureLevel));
 			PixelShaderRHIRef = ShaderComplexityAccumulatePixelShader->GetPixelShader();
 		}
 #endif
@@ -380,7 +380,7 @@ public:
 				RHICmdList.SetBlendState(TStaticBlendState<CW_RGB,BO_Add,BF_One,BF_One>::GetRHI());
 			}
 
-			TShaderMapRef<FShaderComplexityAccumulatePS> ShaderComplexityPixelShader(GetGlobalShaderMap());
+			TShaderMapRef<FShaderComplexityAccumulatePS> ShaderComplexityPixelShader(GetGlobalShaderMap(View.FeatureLevel));
 			const uint32 NumPixelShaderInstructions = PixelShader->GetNumInstructions();
 			const uint32 NumVertexShaderInstructions = VertexShader->GetNumInstructions();
 			ShaderComplexityPixelShader->SetParameters(RHICmdList, NumVertexShaderInstructions,NumPixelShaderInstructions, View.GetFeatureLevel());
@@ -434,7 +434,7 @@ public:
 	static void AddStaticMesh(FRHICommandList& RHICmdList, FScene* Scene, FStaticMesh* StaticMesh);
 	static bool DrawDynamicMesh(
 		FRHICommandList& RHICmdList, 
-		const FSceneView& View,
+		const FViewInfo& View,
 		ContextType DrawingContext,
 		const FMeshBatch& Mesh,
 		bool bBackFace,
diff --git a/Engine/Source/Runtime/Renderer/Private/ForwardTranslucentRendering.cpp b/Engine/Source/Runtime/Renderer/Private/ForwardTranslucentRendering.cpp
--- a/Engine/Source/Runtime/Renderer/Private/ForwardTranslucentRendering.cpp
+++ b/Engine/Source/Runtime/Renderer/Private/ForwardTranslucentRendering.cpp
@@ -45,7 +45,7 @@ IMPLEMENT_SHADER_TYPE(,FForwardCopySceneAlphaPS,TEXT("TranslucentLightingShaders
 
 FGlobalBoundShaderState ForwardCopySceneAlphaBoundShaderState;
 
-void FForwardShadingSceneRenderer::CopySceneAlpha(FRHICommandListImmediate& RHICmdList, const FSceneView& View)
+void FForwardShadingSceneRenderer::CopySceneAlpha(FRHICommandListImmediate& RHICmdList, const FViewInfo& View)
 {
 	SCOPED_DRAW_EVENTF(EventCopy, DEC_SCENE_ITEMS, TEXT("CopySceneAlpha"));
 	RHICmdList.SetRasterizerState(TStaticRasterizerState<FM_Solid, CM_None>::GetRHI());
@@ -61,8 +61,8 @@ void FForwardShadingSceneRenderer::CopySceneAlpha(FRHICommandListImmediate& RHIC
 
 	RHICmdList.SetViewport(0, 0, 0.0f, X, Y, 1.0f);
 
-	TShaderMapRef<FScreenVS> ScreenVertexShader(GetGlobalShaderMap());
-	TShaderMapRef<FForwardCopySceneAlphaPS> PixelShader(GetGlobalShaderMap());
+	TShaderMapRef<FScreenVS> ScreenVertexShader(View.ShaderMap);
+	TShaderMapRef<FForwardCopySceneAlphaPS> PixelShader(View.ShaderMap);
 	SetGlobalBoundShaderState(RHICmdList, FeatureLevel, ForwardCopySceneAlphaBoundShaderState, GFilterVertexDeclaration.VertexDeclarationRHI, *ScreenVertexShader, *PixelShader);
 
 	PixelShader->SetParameters(RHICmdList, View);
diff --git a/Engine/Source/Runtime/Renderer/Private/GPUBenchmark.cpp b/Engine/Source/Runtime/Renderer/Private/GPUBenchmark.cpp
--- a/Engine/Source/Runtime/Renderer/Private/GPUBenchmark.cpp
+++ b/Engine/Source/Runtime/Renderer/Private/GPUBenchmark.cpp
@@ -128,8 +128,10 @@ IMPLEMENT_SHADER_TYPE(,FPostProcessBenchmarkVS,TEXT("GPUBenchmark"),TEXT("MainBe
 template <uint32 Method>
 void RunBenchmarkShader(FRHICommandList& RHICmdList, const FSceneView& View, TRefCountPtr<IPooledRenderTarget>& Src, uint32 Count)
 {
-	TShaderMapRef<FPostProcessBenchmarkVS> VertexShader(GetGlobalShaderMap());
-	TShaderMapRef<FPostProcessBenchmarkPS<Method> > PixelShader(GetGlobalShaderMap());
+	auto ShaderMap = GetGlobalShaderMap(View.GetFeatureLevel());
+
+	TShaderMapRef<FPostProcessBenchmarkVS> VertexShader(ShaderMap);
+	TShaderMapRef<FPostProcessBenchmarkPS<Method> > PixelShader(ShaderMap);
 
 	static FGlobalBoundShaderState BoundShaderState;
 	
diff --git a/Engine/Source/Runtime/Renderer/Private/LightPropagationVolume.cpp b/Engine/Source/Runtime/Renderer/Private/LightPropagationVolume.cpp
--- a/Engine/Source/Runtime/Renderer/Private/LightPropagationVolume.cpp
+++ b/Engine/Source/Runtime/Renderer/Private/LightPropagationVolume.cpp
@@ -611,22 +611,22 @@ IMPLEMENT_SHADER_TYPE(template<>,TLpvPropagateCS<PROPAGATE_SECONDARY_OCCLUSION>,
 IMPLEMENT_SHADER_TYPE(template<>,TLpvPropagateCS<PROPAGATE_MULTIPLE_BOUNCES>,							TEXT("LPVPropagate"),TEXT("CSPropogate"),SF_Compute);
 IMPLEMENT_SHADER_TYPE(template<>,TLpvPropagateCS<PROPAGATE_SECONDARY_OCCLUSION_AND_MULTIPLE_BOUNCES>,	TEXT("LPVPropagate"),TEXT("CSPropogate"),SF_Compute);
 
-FLpvWriteShaderCSBase* GetPropagateShader( bool bSecondaryOcclusion, bool bSecondaryBounces )
+FLpvWriteShaderCSBase* GetPropagateShader(FViewInfo& View, bool bSecondaryOcclusion, bool bSecondaryBounces )
 {
 	FLpvWriteShaderCSBase* ShaderOut = NULL;
 	if ( bSecondaryBounces )
 	{
 		if ( bSecondaryOcclusion ) 
-			ShaderOut = (FLpvWriteShaderCSBase*)*TShaderMapRef<TLpvPropagateCS<PROPAGATE_SECONDARY_OCCLUSION_AND_MULTIPLE_BOUNCES> >( GetGlobalShaderMap() );
+			ShaderOut = (FLpvWriteShaderCSBase*)*TShaderMapRef<TLpvPropagateCS<PROPAGATE_SECONDARY_OCCLUSION_AND_MULTIPLE_BOUNCES> >(View.ShaderMap);
 		else 
-			ShaderOut = (FLpvWriteShaderCSBase*)*TShaderMapRef<TLpvPropagateCS<PROPAGATE_MULTIPLE_BOUNCES> >( GetGlobalShaderMap() );
+			ShaderOut = (FLpvWriteShaderCSBase*)*TShaderMapRef<TLpvPropagateCS<PROPAGATE_MULTIPLE_BOUNCES> >(View.ShaderMap);
 	}
 	else
 	{
 		if ( bSecondaryOcclusion ) 
-			ShaderOut = (FLpvWriteShaderCSBase*)*TShaderMapRef<TLpvPropagateCS<PROPAGATE_SECONDARY_OCCLUSION> >( GetGlobalShaderMap() );
+			ShaderOut = (FLpvWriteShaderCSBase*)*TShaderMapRef<TLpvPropagateCS<PROPAGATE_SECONDARY_OCCLUSION> >(View.ShaderMap);
 		else 
-			ShaderOut = (FLpvWriteShaderCSBase*)*TShaderMapRef<TLpvPropagateCS<0> >( GetGlobalShaderMap() );
+			ShaderOut = (FLpvWriteShaderCSBase*)*TShaderMapRef<TLpvPropagateCS<0> >(View.ShaderMap);
 	}
 	return ShaderOut;
 }
@@ -922,7 +922,7 @@ void FLightPropagationVolume::InitSettings(FRHICommandList& RHICmdList, const FS
 /**
 * Clears the LPV
 */
-void FLightPropagationVolume::Clear(FRHICommandListImmediate& RHICmdList)
+void FLightPropagationVolume::Clear(FRHICommandListImmediate& RHICmdList, FViewInfo& View)
 {
 	if ( !bEnabled )
 	{
@@ -941,7 +941,7 @@ void FLightPropagationVolume::Clear(FRHICommandListImmediate& RHICmdList)
 	//RHIBatchComputeShaderBegin();
 	// Clear the list buffers
 	{
-		TShaderMapRef<FLpvClearListsCS> Shader(GetGlobalShaderMap());
+		TShaderMapRef<FLpvClearListsCS> Shader(View.ShaderMap);
 		RHICmdList.SetComputeShader(Shader->GetComputeShader());
 
 		FLpvBaseWriteShaderParams ShaderParams;
@@ -953,7 +953,7 @@ void FLightPropagationVolume::Clear(FRHICommandListImmediate& RHICmdList)
 
 	// Clear the LPV (or fade, if REFINE_OVER_TIME is enabled)
 	{
-		TShaderMapRef<FLpvClearCS> Shader(GetGlobalShaderMap());
+		TShaderMapRef<FLpvClearCS> Shader(View.ShaderMap);
 		RHICmdList.SetComputeShader(Shader->GetComputeShader());
 
 		FLpvBaseWriteShaderParams ShaderParams;
@@ -966,7 +966,7 @@ void FLightPropagationVolume::Clear(FRHICommandListImmediate& RHICmdList)
 	// Clear the geometry volume if necessary
 	if ( SecondaryOcclusionStrength > 0.001f )
 	{
-		TShaderMapRef<FLpvClearGeometryVolumeCS> Shader(GetGlobalShaderMap());
+		TShaderMapRef<FLpvClearGeometryVolumeCS> Shader(View.ShaderMap);
 		RHICmdList.SetComputeShader(Shader->GetComputeShader());
 
 		FLpvBaseWriteShaderParams ShaderParams;
@@ -1033,6 +1033,7 @@ void FLightPropagationVolume::SetVplInjectionConstants(
 */
 void FLightPropagationVolume::InjectDirectionalLightRSM(
 	FRHICommandListImmediate& RHICmdList,
+	FViewInfo&					View,
 	const FTexture2DRHIRef&		RsmDiffuseTex, 
 	const FTexture2DRHIRef&		RsmNormalTex, 
 	const FTexture2DRHIRef&		RsmDepthTex, 
@@ -1045,7 +1046,7 @@ void FLightPropagationVolume::InjectDirectionalLightRSM(
 
 		SetVplInjectionConstants(ProjectedShadowInfo, LightProxy );
 
-		TShaderMapRef<FLpvInject_GenerateVplListsCS> Shader(GetGlobalShaderMap());
+		TShaderMapRef<FLpvInject_GenerateVplListsCS> Shader(View.ShaderMap);
 		RHICmdList.SetComputeShader(Shader->GetComputeShader());
 
 		// Clear the list counter the first time this function is called in a frame
@@ -1070,7 +1071,7 @@ void FLightPropagationVolume::InjectDirectionalLightRSM(
 		}
 		LpvWriteUniformBufferParams->GeometryVolumeCaptureLightDirection = LightDirection;
 
-		TShaderMapRef<FLpvBuildGeometryVolumeCS> Shader(GetGlobalShaderMap());
+		TShaderMapRef<FLpvBuildGeometryVolumeCS> Shader(View.ShaderMap);
 		RHICmdList.SetComputeShader(Shader->GetComputeShader());
 
 		LpvWriteUniformBuffer.SetContents( *LpvWriteUniformBufferParams ); // FIXME: is this causing a stall? Double-buffer?
@@ -1090,7 +1091,7 @@ void FLightPropagationVolume::InjectDirectionalLightRSM(
 /**
 * Propagates light in the LPV 
 */
-void FLightPropagationVolume::Propagate(FRHICommandListImmediate& RHICmdList)
+void FLightPropagationVolume::Propagate(FRHICommandListImmediate& RHICmdList, FViewInfo& View)
 {
 	if ( !bEnabled )
 	{
@@ -1103,7 +1104,7 @@ void FLightPropagationVolume::Propagate(FRHICommandListImmediate& RHICmdList)
 		SCOPED_DRAW_EVENT(LpvAccumulateVplLists, DEC_LIGHT);
 		mWriteBufferIndex = 1-mWriteBufferIndex; // Swap buffers with each iteration
 
-		TShaderMapRef<FLpvInject_AccumulateVplListsCS> Shader(GetGlobalShaderMap());
+		TShaderMapRef<FLpvInject_AccumulateVplListsCS> Shader(View.ShaderMap);
 		RHICmdList.SetComputeShader(Shader->GetComputeShader());
 
 		LpvWriteUniformBuffer.SetContents( *LpvWriteUniformBufferParams );
@@ -1129,7 +1130,7 @@ void FLightPropagationVolume::Propagate(FRHICommandListImmediate& RHICmdList)
 			bool bSecondaryOcclusion	= (SecondaryOcclusionStrength > 0.001f); 
 			bool bSecondaryBounces		= (SecondaryBounceStrength > 0.001f); 
 			mWriteBufferIndex = 1-mWriteBufferIndex; // Swap buffers with each iteration
-			FLpvWriteShaderCSBase* Shader = GetPropagateShader( bSecondaryOcclusion, bSecondaryBounces );
+			FLpvWriteShaderCSBase* Shader = GetPropagateShader(View, bSecondaryOcclusion, bSecondaryBounces);
 			RHICmdList.SetComputeShader(Shader->GetComputeShader());
 
 			LpvWriteUniformBufferParams->PropagationIndex = PropagationStep;
@@ -1190,7 +1191,7 @@ void FLightPropagationVolume::GetShaderParams( FLpvBaseWriteShaderParams& OutPar
 	OutParams.UniformBuffer = LpvWriteUniformBuffer;
 }
 
-void FLightPropagationVolume::InjectLightDirect(FRHICommandListImmediate& RHICmdList, const FLightSceneProxy& Light)
+void FLightPropagationVolume::InjectLightDirect(FRHICommandListImmediate& RHICmdList, const FLightSceneProxy& Light, const FViewInfo& View)
 {
 	if(!bEnabled)
 	{
@@ -1224,11 +1225,11 @@ void FLightPropagationVolume::InjectLightDirect(FRHICommandListImmediate& RHICmd
 
 			if ( Light.CastsStaticShadow() || Light.CastsDynamicShadow() )
 			{
-				Shader = (FLpvInjectShader_Base*)*TShaderMapRef<TLpvInject_PointLightCS<INJECT_SHADOWED> >( GetGlobalShaderMap() );
+				Shader = (FLpvInjectShader_Base*)*TShaderMapRef<TLpvInject_PointLightCS<INJECT_SHADOWED> >(View.ShaderMap);
 			}
 			else
 			{
-				Shader = (FLpvInjectShader_Base*)*TShaderMapRef<TLpvInject_PointLightCS<0> >( GetGlobalShaderMap() );
+				Shader = (FLpvInjectShader_Base*)*TShaderMapRef<TLpvInject_PointLightCS<0> >(View.ShaderMap);
 			}
 		}
 		RHICmdList.SetComputeShader(Shader->GetComputeShader());
diff --git a/Engine/Source/Runtime/Renderer/Private/LightPropagationVolume.h b/Engine/Source/Runtime/Renderer/Private/LightPropagationVolume.h
--- a/Engine/Source/Runtime/Renderer/Private/LightPropagationVolume.h
+++ b/Engine/Source/Runtime/Renderer/Private/LightPropagationVolume.h
@@ -126,7 +126,7 @@ public:
 
 	void InitSettings(FRHICommandList& RHICmdList, const FSceneView& View);
 
-	void Clear(FRHICommandListImmediate& RHICmdList);
+	void Clear(FRHICommandListImmediate& RHICmdList, FViewInfo& View);
 
 	void SetVplInjectionConstants(
 		const FProjectedShadowInfo&	ProjectedShadowInfo,
@@ -134,17 +134,18 @@ public:
 
 	void InjectDirectionalLightRSM(
 		FRHICommandListImmediate& RHICmdList,
+		FViewInfo&					View,
 		const FTexture2DRHIRef&		RsmDiffuseTex, 
 		const FTexture2DRHIRef&		RsmNormalTex, 
 		const FTexture2DRHIRef&		RsmDepthTex, 
 		const FProjectedShadowInfo&	ProjectedShadowInfo,
 		const FLinearColor&			LightColour );
 
-	void InjectLightDirect(FRHICommandListImmediate& RHICmdList, const FLightSceneProxy& Light);
+	void InjectLightDirect(FRHICommandListImmediate& RHICmdList, const FLightSceneProxy& Light, const FViewInfo& View);
 
-	void Propagate(FRHICommandListImmediate& RHICmdList);
+	void Propagate(FRHICommandListImmediate& RHICmdList, FViewInfo& View);
 
-	void Visualise(FRHICommandList& RHICmdList, const FSceneView& View) const;
+	void Visualise(FRHICommandList& RHICmdList, const FViewInfo& View) const;
 
 	const FIntVector& GetGridOffset() const { return mGridOffset; }
 
diff --git a/Engine/Source/Runtime/Renderer/Private/LightPropagationVolumeVisualisation.cpp b/Engine/Source/Runtime/Renderer/Private/LightPropagationVolumeVisualisation.cpp
--- a/Engine/Source/Runtime/Renderer/Private/LightPropagationVolumeVisualisation.cpp
+++ b/Engine/Source/Runtime/Renderer/Private/LightPropagationVolumeVisualisation.cpp
@@ -248,14 +248,14 @@ IMPLEMENT_SHADER_TYPE(,FLpvVisualiseVS,TEXT("LPVVisualise"),TEXT("VShader"),SF_V
 IMPLEMENT_SHADER_TYPE(,FLpvVisualisePS,TEXT("LPVVisualise"),TEXT("PShader"),SF_Pixel);
 
 
-void FLightPropagationVolume::Visualise(FRHICommandList& RHICmdList, const FSceneView& View) const
+void FLightPropagationVolume::Visualise(FRHICommandList& RHICmdList, const FViewInfo& View) const
 {
 	SCOPED_DRAW_EVENT(LpvVisualise, DEC_LIGHT);
 	check(View.GetFeatureLevel() == ERHIFeatureLevel::SM5);
 
-	TShaderMapRef<FLpvVisualiseVS> VertexShader(GetGlobalShaderMap());
-	TShaderMapRef<FLpvVisualiseGS> GeometryShader(GetGlobalShaderMap());
-	TShaderMapRef<FLpvVisualisePS> PixelShader(GetGlobalShaderMap());
+	TShaderMapRef<FLpvVisualiseVS> VertexShader(View.ShaderMap);
+	TShaderMapRef<FLpvVisualiseGS> GeometryShader(View.ShaderMap);
+	TShaderMapRef<FLpvVisualisePS> PixelShader(View.ShaderMap);
 
 	RHICmdList.SetDepthStencilState(TStaticDepthStencilState<false, CF_Always>::GetRHI());
 	RHICmdList.SetRasterizerState(TStaticRasterizerState<FM_Solid, CM_None>::GetRHI());
diff --git a/Engine/Source/Runtime/Renderer/Private/LightRendering.cpp b/Engine/Source/Runtime/Renderer/Private/LightRendering.cpp
--- a/Engine/Source/Runtime/Renderer/Private/LightRendering.cpp
+++ b/Engine/Source/Runtime/Renderer/Private/LightRendering.cpp
@@ -564,17 +564,18 @@ void FDeferredShadingSceneRenderer::RenderLights(FRHICommandListImmediate& RHICm
 					// Inject the light directly into all relevant LPVs
 					for (int32 ViewIndex = 0; ViewIndex < Views.Num(); ViewIndex++)
 					{
-						FSceneViewState* ViewState = (FSceneViewState*)Views[ViewIndex].State;
+						FViewInfo& View = Views[ViewIndex];
+						FSceneViewState* ViewState = (FSceneViewState*)View.State;
 
 						if (ViewState)
 						{
-							if (LightSceneInfo->ShouldRenderLight(Views[ViewIndex]))
+							if (LightSceneInfo->ShouldRenderLight(View))
 							{
 								FLightPropagationVolume* Lpv = ViewState->GetLightPropagationVolume();
 								if (Lpv && LightSceneInfo->Proxy)
 								{
 									
-									Lpv->InjectLightDirect(RHICmdList, *LightSceneInfo->Proxy);
+									Lpv->InjectLightDirect(RHICmdList, *LightSceneInfo->Proxy, View);
 								}
 							}
 						}
@@ -666,19 +667,19 @@ static FVertexDeclarationRHIParamRef GetDeferredLightingVertexDeclaration()
 template<bool bUseIESProfile, bool bRadialAttenuation, bool bInverseSquaredFalloff>
 static void SetShaderTemplLighting(
 	FRHICommandList& RHICmdList,
-	const FSceneView& View, 
+	const FViewInfo& View,
 	FShader* VertexShader,
 	const FLightSceneInfo* LightSceneInfo)
 {
 	if(View.Family->EngineShowFlags.VisualizeLightCulling)
 	{
-		TShaderMapRef<TDeferredLightPS<false, bRadialAttenuation, false, true> > PixelShader(GetGlobalShaderMap());
+		TShaderMapRef<TDeferredLightPS<false, bRadialAttenuation, false, true> > PixelShader(View.ShaderMap);
 		SetGlobalBoundShaderState(RHICmdList, View.GetFeatureLevel(), PixelShader->GetBoundShaderState(), GetDeferredLightingVertexDeclaration<bRadialAttenuation>(), VertexShader, *PixelShader);
 		PixelShader->SetParameters(RHICmdList, View, LightSceneInfo);
 	}
 	else
 	{
-		TShaderMapRef<TDeferredLightPS<bUseIESProfile, bRadialAttenuation, bInverseSquaredFalloff, false> > PixelShader(GetGlobalShaderMap());
+		TShaderMapRef<TDeferredLightPS<bUseIESProfile, bRadialAttenuation, bInverseSquaredFalloff, false> > PixelShader(View.ShaderMap);
 		SetGlobalBoundShaderState(RHICmdList, View.GetFeatureLevel(), PixelShader->GetBoundShaderState(), GetDeferredLightingVertexDeclaration<bRadialAttenuation>(), VertexShader, *PixelShader);
 		PixelShader->SetParameters(RHICmdList, View, LightSceneInfo);
 	}
@@ -687,20 +688,20 @@ static void SetShaderTemplLighting(
 template<bool bUseIESProfile, bool bRadialAttenuation, bool bInverseSquaredFalloff>
 static void SetShaderTemplLightingSimple(
 	FRHICommandList& RHICmdList,
-	const FSceneView& View, 
+	const FViewInfo& View,
 	FShader* VertexShader,
 	const FSimpleLightEntry& SimpleLight,
 	const FSimpleLightPerViewEntry& SimpleLightPerViewData)
 {
 	if(View.Family->EngineShowFlags.VisualizeLightCulling)
 	{
-		TShaderMapRef<TDeferredLightPS<false, bRadialAttenuation, false, true> > PixelShader(GetGlobalShaderMap());
+		TShaderMapRef<TDeferredLightPS<false, bRadialAttenuation, false, true> > PixelShader(View.ShaderMap);
 		SetGlobalBoundShaderState(RHICmdList, View.GetFeatureLevel(), PixelShader->GetBoundShaderState(), GetDeferredLightingVertexDeclaration<bRadialAttenuation>(), VertexShader, *PixelShader);
 		PixelShader->SetParametersSimpleLight(RHICmdList, View, SimpleLight, SimpleLightPerViewData);
 	}
 	else
 	{
-		TShaderMapRef<TDeferredLightPS<bUseIESProfile, bRadialAttenuation, bInverseSquaredFalloff, false> > PixelShader(GetGlobalShaderMap());
+		TShaderMapRef<TDeferredLightPS<bUseIESProfile, bRadialAttenuation, bInverseSquaredFalloff, false> > PixelShader(View.ShaderMap);
 		SetGlobalBoundShaderState(RHICmdList, View.GetFeatureLevel(), PixelShader->GetBoundShaderState(), GetDeferredLightingVertexDeclaration<bRadialAttenuation>(), VertexShader, *PixelShader);
 		PixelShader->SetParametersSimpleLight(RHICmdList, View, SimpleLight, SimpleLightPerViewData);
 	}
@@ -741,14 +742,14 @@ void FDeferredShadingSceneRenderer::RenderLight(FRHICommandList& RHICmdList, con
 
 		if (LightSceneInfo->Proxy->GetLightType() == LightType_Directional)
 		{
-			TShaderMapRef<TDeferredLightVS<false> > VertexShader(GetGlobalShaderMap());
+			TShaderMapRef<TDeferredLightVS<false> > VertexShader(View.ShaderMap);
 
 			RHICmdList.SetRasterizerState(TStaticRasterizerState<FM_Solid, CM_None>::GetRHI());
 			RHICmdList.SetDepthStencilState(TStaticDepthStencilState<false, CF_Always>::GetRHI());
 
 			if (bRenderOverlap)
 			{
-				TShaderMapRef<TDeferredLightOverlapPS<false> > PixelShader(GetGlobalShaderMap());
+				TShaderMapRef<TDeferredLightOverlapPS<false> > PixelShader(View.ShaderMap);
 				SetGlobalBoundShaderState(RHICmdList, FeatureLevel, PixelShader->GetBoundShaderState(), GetDeferredLightingVertexDeclaration<false>(), *VertexShader, *PixelShader);
 				PixelShader->SetParameters(RHICmdList, View, LightSceneInfo);
 			}
@@ -780,13 +781,13 @@ void FDeferredShadingSceneRenderer::RenderLight(FRHICommandList& RHICmdList, con
 		}
 		else
 		{
-			TShaderMapRef<TDeferredLightVS<true> > VertexShader(GetGlobalShaderMap());
+			TShaderMapRef<TDeferredLightVS<true> > VertexShader(View.ShaderMap);
 
 			SetBoundingGeometryRasterizerAndDepthState(RHICmdList, View, LightBounds);
 
 			if (bRenderOverlap)
 			{
-				TShaderMapRef<TDeferredLightOverlapPS<true> > PixelShader(GetGlobalShaderMap());
+				TShaderMapRef<TDeferredLightOverlapPS<true> > PixelShader(View.ShaderMap);
 				SetGlobalBoundShaderState(RHICmdList, FeatureLevel, PixelShader->GetBoundShaderState(), GetDeferredLightingVertexDeclaration<true>(), *VertexShader, *PixelShader);
 				PixelShader->SetParameters(RHICmdList, View, LightSceneInfo);
 			}
@@ -861,7 +862,7 @@ void FDeferredShadingSceneRenderer::RenderSimpleLightsStandardDeferred(FRHIComma
 			// Set the device viewport for the view.
 			RHICmdList.SetViewport(View.ViewRect.Min.X, View.ViewRect.Min.Y, 0.0f, View.ViewRect.Max.X, View.ViewRect.Max.Y, 1.0f);
 
-			TShaderMapRef<TDeferredLightVS<true> > VertexShader(GetGlobalShaderMap());
+			TShaderMapRef<TDeferredLightVS<true> > VertexShader(View.ShaderMap);
 
 			SetBoundingGeometryRasterizerAndDepthState(RHICmdList, View, LightBounds);
 
diff --git a/Engine/Source/Runtime/Renderer/Private/LightShaftRendering.cpp b/Engine/Source/Runtime/Renderer/Private/LightShaftRendering.cpp
--- a/Engine/Source/Runtime/Renderer/Private/LightShaftRendering.cpp
+++ b/Engine/Source/Runtime/Renderer/Private/LightShaftRendering.cpp
@@ -464,21 +464,21 @@ void DownsamplePass(FRHICommandListImmediate& RHICmdList, const FViewInfo& View,
 	RHICmdList.SetViewport(DownSampledXY.X, DownSampledXY.Y, 0.0f, DownSampledXY.X + DownsampledSizeX, DownSampledXY.Y + DownsampledSizeY, 1.0f);
 
 	// Set shaders and texture
-	TShaderMapRef<FDownsampleLightShaftsVertexShader> DownsampleLightShaftsVertexShader(GetGlobalShaderMap());
+	TShaderMapRef<FDownsampleLightShaftsVertexShader> DownsampleLightShaftsVertexShader(View.ShaderMap);
 	TRefCountPtr<IPooledRenderTarget> UnusedRT;
 
 	switch(LightSceneInfo->Proxy->GetLightType())
 	{
 	case LightType_Directional:
 		{
-			TShaderMapRef<TDownsampleLightShaftsPixelShader<LightType_Directional, bDownsampleOcclusion> > DownsampleLightShaftsPixelShader(GetGlobalShaderMap());
+			TShaderMapRef<TDownsampleLightShaftsPixelShader<LightType_Directional, bDownsampleOcclusion> > DownsampleLightShaftsPixelShader(View.ShaderMap);
 			SetGlobalBoundShaderState(RHICmdList, View.GetFeatureLevel(), DownsampleDirectionalLightShaftsBoundShaderState[bDownsampleOcclusion], GFilterVertexDeclaration.VertexDeclarationRHI, *DownsampleLightShaftsVertexShader, *DownsampleLightShaftsPixelShader);
 			DownsampleLightShaftsPixelShader->SetParameters(RHICmdList, LightSceneInfo, View, UnusedRT);
 		}
 		break;
 	case LightType_Spot:
 		{
-			TShaderMapRef<TDownsampleLightShaftsPixelShader<LightType_Spot, bDownsampleOcclusion> > DownsampleLightShaftsPixelShader(GetGlobalShaderMap());
+			TShaderMapRef<TDownsampleLightShaftsPixelShader<LightType_Spot, bDownsampleOcclusion> > DownsampleLightShaftsPixelShader(View.ShaderMap);
 			SetGlobalBoundShaderState(RHICmdList, View.GetFeatureLevel(), DownsampleSpotLightShaftsBoundShaderState[bDownsampleOcclusion], GFilterVertexDeclaration.VertexDeclarationRHI, *DownsampleLightShaftsVertexShader, *DownsampleLightShaftsPixelShader);
 			DownsampleLightShaftsPixelShader->SetParameters(RHICmdList, LightSceneInfo, View, UnusedRT);
 		}
@@ -486,7 +486,7 @@ void DownsamplePass(FRHICommandListImmediate& RHICmdList, const FViewInfo& View,
 	default:
 	case LightType_Point:
 		{
-			TShaderMapRef<TDownsampleLightShaftsPixelShader<LightType_Point, bDownsampleOcclusion> > DownsampleLightShaftsPixelShader(GetGlobalShaderMap());
+			TShaderMapRef<TDownsampleLightShaftsPixelShader<LightType_Point, bDownsampleOcclusion> > DownsampleLightShaftsPixelShader(View.ShaderMap);
 			SetGlobalBoundShaderState(RHICmdList, View.GetFeatureLevel(), DownsamplePointLightShaftsBoundShaderState[bDownsampleOcclusion], GFilterVertexDeclaration.VertexDeclarationRHI, *DownsampleLightShaftsVertexShader, *DownsampleLightShaftsPixelShader);
 			DownsampleLightShaftsPixelShader->SetParameters(RHICmdList, LightSceneInfo, View, UnusedRT);
 		}
@@ -597,7 +597,7 @@ void ApplyRadialBlurPasses(
 	/** First pass dest. */
 	TRefCountPtr<IPooledRenderTarget>& LightShaftsDest)
 {
-	TShaderMapRef<FScreenVS> ScreenVertexShader(GetGlobalShaderMap());
+	TShaderMapRef<FScreenVS> ScreenVertexShader(View.ShaderMap);
 
 	const FIntPoint BufferSize = GSceneRenderTargets.GetBufferSizeXY();
 	const uint32 DownsampleFactor	= GetLightShaftDownsampleFactor();
@@ -616,7 +616,7 @@ void ApplyRadialBlurPasses(
 		RHICmdList.SetRasterizerState(TStaticRasterizerState<>::GetRHI());
 		RHICmdList.SetDepthStencilState(TStaticDepthStencilState<false, CF_Always>::GetRHI());
 
-		TShaderMapRef<FBlurLightShaftsPixelShader> BlurLightShaftsPixelShader(GetGlobalShaderMap());
+		TShaderMapRef<FBlurLightShaftsPixelShader> BlurLightShaftsPixelShader(View.ShaderMap);
 		SetGlobalBoundShaderState(RHICmdList, View.GetFeatureLevel(), BlurLightShaftsBoundShaderState, GFilterVertexDeclaration.VertexDeclarationRHI, *ScreenVertexShader, *BlurLightShaftsPixelShader);
 
 		TRefCountPtr<IPooledRenderTarget>& EffectiveSource = PassIndex == 0 ? FirstPassSource : LightShaftsSource;
@@ -646,7 +646,7 @@ void ApplyRadialBlurPasses(
 
 void FinishOcclusionTerm(FRHICommandList& RHICmdList, const FViewInfo& View, const FLightSceneInfo* const LightSceneInfo, TRefCountPtr<IPooledRenderTarget>& LightShaftsSource, TRefCountPtr<IPooledRenderTarget>& LightShaftsDest)
 {
-	TShaderMapRef<FScreenVS> ScreenVertexShader(GetGlobalShaderMap());
+	TShaderMapRef<FScreenVS> ScreenVertexShader(View.ShaderMap);
 
 	const FIntPoint BufferSize = GSceneRenderTargets.GetBufferSizeXY();
 	const uint32 DownsampleFactor	= GetLightShaftDownsampleFactor();
@@ -662,7 +662,7 @@ void FinishOcclusionTerm(FRHICommandList& RHICmdList, const FViewInfo& View, con
 	RHICmdList.SetRasterizerState(TStaticRasterizerState<>::GetRHI());
 	RHICmdList.SetDepthStencilState(TStaticDepthStencilState<false, CF_Always>::GetRHI());
 
-	TShaderMapRef<FFinishOcclusionPixelShader> MaskOcclusionTermPixelShader(GetGlobalShaderMap());
+	TShaderMapRef<FFinishOcclusionPixelShader> MaskOcclusionTermPixelShader(View.ShaderMap);
 	SetGlobalBoundShaderState(RHICmdList, View.GetFeatureLevel(), AccumulateTermBoundShaderState, GFilterVertexDeclaration.VertexDeclarationRHI, *ScreenVertexShader, *MaskOcclusionTermPixelShader);
 	/// ?
 	MaskOcclusionTermPixelShader->SetParameters(RHICmdList, LightSceneInfo, View, LightShaftsSource);
@@ -845,8 +845,8 @@ void ApplyLightShaftBloom(FRHICommandListImmediate& RHICmdList, const FViewInfo&
 	RHICmdList.SetRasterizerState(TStaticRasterizerState<>::GetRHI());
 	RHICmdList.SetDepthStencilState(TStaticDepthStencilState<false, CF_Always>::GetRHI());
 
-	TShaderMapRef<FScreenVS> ScreenVertexShader(GetGlobalShaderMap());
-	TShaderMapRef<FApplyLightShaftsPixelShader> ApplyLightShaftsPixelShader(GetGlobalShaderMap());
+	TShaderMapRef<FScreenVS> ScreenVertexShader(View.ShaderMap);
+	TShaderMapRef<FApplyLightShaftsPixelShader> ApplyLightShaftsPixelShader(View.ShaderMap);
 
 	SetGlobalBoundShaderState(RHICmdList, View.GetFeatureLevel(), ApplyLightShaftsBoundShaderState, GFilterVertexDeclaration.VertexDeclarationRHI, *ScreenVertexShader, *ApplyLightShaftsPixelShader);
 
diff --git a/Engine/Source/Runtime/Renderer/Private/PostProcess/PostProcessAA.cpp b/Engine/Source/Runtime/Renderer/Private/PostProcess/PostProcessAA.cpp
--- a/Engine/Source/Runtime/Renderer/Private/PostProcess/PostProcessAA.cpp
+++ b/Engine/Source/Runtime/Renderer/Private/PostProcess/PostProcessAA.cpp
@@ -216,8 +216,8 @@ IMPLEMENT_SHADER_TYPE(,FFXAAVS,TEXT("FXAAShader"),TEXT("FxaaVS"),SF_Vertex);
 template <uint32 Quality>
 static void SetShaderTemplAA(const FRenderingCompositePassContext& Context)
 {
-	TShaderMapRef<FFXAAVS> VertexShader(GetGlobalShaderMap());
-	TShaderMapRef<FPostProcessAntiAliasingPS<Quality> > PixelShader(GetGlobalShaderMap());
+	TShaderMapRef<FFXAAVS> VertexShader(Context.GetShaderMap());
+	TShaderMapRef<FPostProcessAntiAliasingPS<Quality> > PixelShader(Context.GetShaderMap());
 
 	static FGlobalBoundShaderState BoundShaderState;
 	
@@ -272,7 +272,7 @@ void FRCPassPostProcessAA::Process(FRenderingCompositePassContext& Context)
 	}
 	
 	// Draw a quad mapping scene color to the view's render target
-	TShaderMapRef<FFXAAVS> VertexShader(GetGlobalShaderMap());
+	TShaderMapRef<FFXAAVS> VertexShader(Context.GetShaderMap());
 
 	DrawRectangle(
 		Context.RHICmdList,
diff --git a/Engine/Source/Runtime/Renderer/Private/PostProcess/PostProcessBloomSetup.cpp b/Engine/Source/Runtime/Renderer/Private/PostProcess/PostProcessBloomSetup.cpp
--- a/Engine/Source/Runtime/Renderer/Private/PostProcess/PostProcessBloomSetup.cpp
+++ b/Engine/Source/Runtime/Renderer/Private/PostProcess/PostProcessBloomSetup.cpp
@@ -176,8 +176,8 @@ void FRCPassPostProcessBloomSetup::Process(FRenderingCompositePassContext& Conte
 	Context.RHICmdList.SetRasterizerState(TStaticRasterizerState<>::GetRHI());
 	Context.RHICmdList.SetDepthStencilState(TStaticDepthStencilState<false, CF_Always>::GetRHI());
 
-	TShaderMapRef<FPostProcessBloomSetupVS> VertexShader(GetGlobalShaderMap());
-	TShaderMapRef<FPostProcessBloomSetupPS> PixelShader(GetGlobalShaderMap());
+	TShaderMapRef<FPostProcessBloomSetupVS> VertexShader(Context.GetShaderMap());
+	TShaderMapRef<FPostProcessBloomSetupPS> PixelShader(Context.GetShaderMap());
 
 	static FGlobalBoundShaderState BoundShaderState;
 	
diff --git a/Engine/Source/Runtime/Renderer/Private/PostProcess/PostProcessBokehDOF.cpp b/Engine/Source/Runtime/Renderer/Private/PostProcess/PostProcessBokehDOF.cpp
--- a/Engine/Source/Runtime/Renderer/Private/PostProcess/PostProcessBokehDOF.cpp
+++ b/Engine/Source/Runtime/Renderer/Private/PostProcess/PostProcessBokehDOF.cpp
@@ -186,10 +186,10 @@ void FRCPassPostProcessVisualizeDOF::Process(FRenderingCompositePassContext& Con
 	Context.RHICmdList.SetDepthStencilState(TStaticDepthStencilState<false, CF_Always>::GetRHI());
 
 	// setup shader
-	TShaderMapRef<FPostProcessVS> VertexShader(GetGlobalShaderMap());
+	TShaderMapRef<FPostProcessVS> VertexShader(Context.GetShaderMap());
 	
 	{
-		TShaderMapRef<PostProcessVisualizeDOFPS> PixelShader(GetGlobalShaderMap());
+		TShaderMapRef<PostProcessVisualizeDOFPS> PixelShader(Context.GetShaderMap());
 
 		static FGlobalBoundShaderState BoundShaderState;
 
@@ -345,9 +345,9 @@ void FRCPassPostProcessBokehDOFSetup::Process(FRenderingCompositePassContext& Co
 	Context.RHICmdList.SetDepthStencilState(TStaticDepthStencilState<false, CF_Always>::GetRHI());
 
 	// setup shader
-	TShaderMapRef<FPostProcessVS> VertexShader(GetGlobalShaderMap());
+	TShaderMapRef<FPostProcessVS> VertexShader(Context.GetShaderMap());
 	{
-		TShaderMapRef<PostProcessBokehDOFSetupPS> PixelShader(GetGlobalShaderMap());
+		TShaderMapRef<PostProcessBokehDOFSetupPS> PixelShader(Context.GetShaderMap());
 
 		static FGlobalBoundShaderState BoundShaderState;
 
@@ -560,8 +560,8 @@ VARIATION1(0,1)			VARIATION1(1,1)			VARIATION1(2,1)
 template <uint32 DOFMethod, uint32 DOFIndexStyle>
 void FRCPassPostProcessBokehDOF::SetShaderTempl(const FRenderingCompositePassContext& Context, FIntPoint LeftTop, FIntPoint TileCount, uint32 TileSize, float PixelKernelSize)
 {
-	TShaderMapRef<FPostProcessBokehDOFVS<DOFMethod, DOFIndexStyle> > VertexShader(GetGlobalShaderMap());
-	TShaderMapRef<FPostProcessBokehDOFPS> PixelShader(GetGlobalShaderMap());
+	TShaderMapRef<FPostProcessBokehDOFVS<DOFMethod, DOFIndexStyle> > VertexShader(Context.GetShaderMap());
+	TShaderMapRef<FPostProcessBokehDOFPS> PixelShader(Context.GetShaderMap());
 
 	static FGlobalBoundShaderState BoundShaderState;
 	
diff --git a/Engine/Source/Runtime/Renderer/Private/PostProcess/PostProcessBokehDOFRecombine.cpp b/Engine/Source/Runtime/Renderer/Private/PostProcess/PostProcessBokehDOFRecombine.cpp
--- a/Engine/Source/Runtime/Renderer/Private/PostProcess/PostProcessBokehDOFRecombine.cpp
+++ b/Engine/Source/Runtime/Renderer/Private/PostProcess/PostProcessBokehDOFRecombine.cpp
@@ -98,8 +98,8 @@ VARIATION1(1)			VARIATION1(2)			VARIATION1(3)
 template <uint32 Method>
 void FRCPassPostProcessBokehDOFRecombine::SetShader(const FRenderingCompositePassContext& Context)
 {
-	TShaderMapRef<FPostProcessVS> VertexShader(GetGlobalShaderMap());
-	TShaderMapRef<FPostProcessBokehDOFRecombinePS<Method> > PixelShader(GetGlobalShaderMap());
+	TShaderMapRef<FPostProcessVS> VertexShader(Context.GetShaderMap());
+	TShaderMapRef<FPostProcessBokehDOFRecombinePS<Method> > PixelShader(Context.GetShaderMap());
 
 	static FGlobalBoundShaderState BoundShaderState;
 	
@@ -168,7 +168,7 @@ void FRCPassPostProcessBokehDOFRecombine::Process(FRenderingCompositePassContext
 			check(0);
 	}
 
-	TShaderMapRef<FPostProcessVS> VertexShader(GetGlobalShaderMap());
+	TShaderMapRef<FPostProcessVS> VertexShader(Context.GetShaderMap());
 
 	// Draw a quad mapping scene color to the view's render target
 	DrawRectangle(
diff --git a/Engine/Source/Runtime/Renderer/Private/PostProcess/PostProcessBusyWait.cpp b/Engine/Source/Runtime/Renderer/Private/PostProcess/PostProcessBusyWait.cpp
--- a/Engine/Source/Runtime/Renderer/Private/PostProcess/PostProcessBusyWait.cpp
+++ b/Engine/Source/Runtime/Renderer/Private/PostProcess/PostProcessBusyWait.cpp
@@ -103,8 +103,8 @@ void FRCPassPostProcessBusyWait::Process(FRenderingCompositePassContext& Context
 	Context.RHICmdList.SetRasterizerState(TStaticRasterizerState<>::GetRHI());
 	Context.RHICmdList.SetDepthStencilState(TStaticDepthStencilState<false, CF_Always>::GetRHI());
 
-	TShaderMapRef<FPostProcessVS> VertexShader(GetGlobalShaderMap());
-	TShaderMapRef<FPostProcessBusyWaitPS> PixelShader(GetGlobalShaderMap());
+	TShaderMapRef<FPostProcessVS> VertexShader(Context.GetShaderMap());
+	TShaderMapRef<FPostProcessBusyWaitPS> PixelShader(Context.GetShaderMap());
 
 	static FGlobalBoundShaderState BoundShaderState;
 
diff --git a/Engine/Source/Runtime/Renderer/Private/PostProcess/PostProcessCombineLUTs.cpp b/Engine/Source/Runtime/Renderer/Private/PostProcess/PostProcessCombineLUTs.cpp
--- a/Engine/Source/Runtime/Renderer/Private/PostProcess/PostProcessCombineLUTs.cpp
+++ b/Engine/Source/Runtime/Renderer/Private/PostProcess/PostProcessCombineLUTs.cpp
@@ -203,11 +203,14 @@ void SetLUTBlenderShader(FRenderingCompositePassContext& Context, uint32 BlendCo
 	FGlobalBoundShaderState* LocalBoundShaderState = 0;
 	const FSceneView& View = Context.View;
 
+	const auto FeatureLevel = Context.GetFeatureLevel();
+	auto ShaderMap = Context.GetShaderMap();
+
 	// A macro to handle setting the filter shader for a specific number of samples.
 #define CASE_COUNT(BlendCount) \
 	case BlendCount: \
 	{ \
-		TShaderMapRef<FLUTBlenderPS<BlendCount> > PixelShader(GetGlobalShaderMap()); \
+		TShaderMapRef<FLUTBlenderPS<BlendCount> > PixelShader(ShaderMap); \
 		static FGlobalBoundShaderState BoundShaderState; \
 		LocalBoundShaderState = &BoundShaderState;\
 		LocalPixelShader = *PixelShader;\
@@ -229,26 +232,26 @@ void SetLUTBlenderShader(FRenderingCompositePassContext& Context, uint32 BlendCo
 	check(LocalBoundShaderState != NULL);
 	if(UseVolumeTextureLUT(Context.View.GetShaderPlatform()))
 	{
-		TShaderMapRef<FWriteToSliceVS> VertexShader(GetGlobalShaderMap());
-		TShaderMapRef<FWriteToSliceGS> GeometryShader(GetGlobalShaderMap());
+		TShaderMapRef<FWriteToSliceVS> VertexShader(ShaderMap);
+		TShaderMapRef<FWriteToSliceGS> GeometryShader(ShaderMap);
 
-		SetGlobalBoundShaderState(Context.RHICmdList, Context.GetFeatureLevel(), *LocalBoundShaderState, GScreenVertexDeclaration.VertexDeclarationRHI, *VertexShader, LocalPixelShader, *GeometryShader);
+		SetGlobalBoundShaderState(Context.RHICmdList, FeatureLevel, *LocalBoundShaderState, GScreenVertexDeclaration.VertexDeclarationRHI, *VertexShader, LocalPixelShader, *GeometryShader);
 
 		VertexShader->SetParameters(Context.RHICmdList, VolumeBounds, VolumeBounds.MaxX - VolumeBounds.MinX);
 		GeometryShader->SetParameters(Context.RHICmdList, VolumeBounds);
 	}
 	else
 	{
-		TShaderMapRef<FPostProcessVS> VertexShader(GetGlobalShaderMap());
+		TShaderMapRef<FPostProcessVS> VertexShader(ShaderMap);
 
-		SetGlobalBoundShaderState(Context.RHICmdList, Context.GetFeatureLevel(), *LocalBoundShaderState, GFilterVertexDeclaration.VertexDeclarationRHI, *VertexShader, LocalPixelShader);
+		SetGlobalBoundShaderState(Context.RHICmdList, FeatureLevel, *LocalBoundShaderState, GFilterVertexDeclaration.VertexDeclarationRHI, *VertexShader, LocalPixelShader);
 
 		VertexShader->SetParameters(Context);
 	}
 #define CASE_COUNT(BlendCount) \
 	case BlendCount: \
 	{ \
-	TShaderMapRef<FLUTBlenderPS<BlendCount> > PixelShader(GetGlobalShaderMap()); \
+	TShaderMapRef<FLUTBlenderPS<BlendCount> > PixelShader(ShaderMap); \
 	PixelShader->SetParameters(Context.RHICmdList, View, Texture, Weights); \
 	}; \
 	break;
@@ -422,7 +425,7 @@ void FRCPassPostProcessCombineLUTs::Process(FRenderingCompositePassContext& Cont
 	else
 	{
 		// use unwrapped 2d texture 256x16
-		TShaderMapRef<FPostProcessVS> VertexShader(GetGlobalShaderMap());
+		TShaderMapRef<FPostProcessVS> VertexShader(Context.GetShaderMap());
 
 		DrawRectangle( 
 			Context.RHICmdList,
@@ -484,7 +487,7 @@ bool FRCPassPostProcessCombineLUTs::IsColorGradingLUTNeeded(const FRenderingComp
 		return true;
 	}
 
-	if(Context->View.GetFeatureLevel() >= ERHIFeatureLevel::SM4)
+	if(Context->GetFeatureLevel() >= ERHIFeatureLevel::SM4)
 	{
 		const FFinalPostProcessSettings Settings = Context->View.FinalPostProcessSettings;
 		for(uint32 i = 0; i < (uint32)Settings.ContributingLUTs.Num(); ++i)
diff --git a/Engine/Source/Runtime/Renderer/Private/PostProcess/PostProcessCompositeEditorPrimitives.cpp b/Engine/Source/Runtime/Renderer/Private/PostProcess/PostProcessCompositeEditorPrimitives.cpp
--- a/Engine/Source/Runtime/Renderer/Private/PostProcess/PostProcessCompositeEditorPrimitives.cpp
+++ b/Engine/Source/Runtime/Renderer/Private/PostProcess/PostProcessCompositeEditorPrimitives.cpp
@@ -165,13 +165,16 @@ private:
 template <uint32 MSAASampleCount>
 static void SetCompositePrimitivesShaderTempl(const FRenderingCompositePassContext& Context)
 {
-	TShaderMapRef<FPostProcessVS> VertexShader(GetGlobalShaderMap());
-	TShaderMapRef<FPostProcessCompositeEditorPrimitivesPS<MSAASampleCount> > PixelShader(GetGlobalShaderMap());
+	const auto FeatureLevel = Context.GetFeatureLevel();
+	auto ShaderMap = Context.GetShaderMap();
+
+	TShaderMapRef<FPostProcessVS> VertexShader(ShaderMap);
+	TShaderMapRef<FPostProcessCompositeEditorPrimitivesPS<MSAASampleCount> > PixelShader(ShaderMap);
 
 	static FGlobalBoundShaderState BoundShaderState;
 	
 
-	SetGlobalBoundShaderState(Context.RHICmdList, Context.GetFeatureLevel(), BoundShaderState, GFilterVertexDeclaration.VertexDeclarationRHI, *VertexShader, *PixelShader);
+	SetGlobalBoundShaderState(Context.RHICmdList, FeatureLevel, BoundShaderState, GFilterVertexDeclaration.VertexDeclarationRHI, *VertexShader, *PixelShader);
 
 	VertexShader->SetParameters(Context);
 	PixelShader->SetParameters(Context);
@@ -262,7 +265,7 @@ void FRCPassPostProcessCompositeEditorPrimitives::Process(FRenderingCompositePas
 		check(0);
 	}
 
-	TShaderMapRef<FPostProcessVS> VertexShader(GetGlobalShaderMap());
+	TShaderMapRef<FPostProcessVS> VertexShader(Context.GetShaderMap());
 
 	// Draw a quad mapping our render targets to the view's render target
 	DrawRectangle(
diff --git a/Engine/Source/Runtime/Renderer/Private/PostProcess/PostProcessDOF.cpp b/Engine/Source/Runtime/Renderer/Private/PostProcess/PostProcessDOF.cpp
--- a/Engine/Source/Runtime/Renderer/Private/PostProcess/PostProcessDOF.cpp
+++ b/Engine/Source/Runtime/Renderer/Private/PostProcess/PostProcessDOF.cpp
@@ -95,6 +95,9 @@ void FRCPassPostProcessDOFSetup::Process(FRenderingCompositePassContext& Context
 	const FSceneView& View = Context.View;
 	const FSceneViewFamily& ViewFamily = *(View.Family);
 
+	const auto FeatureLevel = Context.GetFeatureLevel();
+	auto ShaderMap = Context.GetShaderMap();
+
 	FIntPoint SrcSize = InputDesc->Extent;
 	FIntPoint DestSize = PassOutputs[0].RenderTargetDesc.Extent;
 
@@ -130,15 +133,15 @@ void FRCPassPostProcessDOFSetup::Process(FRenderingCompositePassContext& Context
 	Context.RHICmdList.SetRasterizerState(TStaticRasterizerState<>::GetRHI());
 	Context.RHICmdList.SetDepthStencilState(TStaticDepthStencilState<false, CF_Always>::GetRHI());
 	
-	TShaderMapRef<FPostProcessVS> VertexShader(GetGlobalShaderMap());
+	TShaderMapRef<FPostProcessVS> VertexShader(ShaderMap);
 
 	if (bNearBlurEnabled)
 	{
 		static FGlobalBoundShaderState BoundShaderState;
 		
 
-		TShaderMapRef< FPostProcessDOFSetupPS<1> > PixelShader(GetGlobalShaderMap());
-		SetGlobalBoundShaderState(Context.RHICmdList, Context.GetFeatureLevel(), BoundShaderState, GFilterVertexDeclaration.VertexDeclarationRHI, *VertexShader, *PixelShader);
+		TShaderMapRef< FPostProcessDOFSetupPS<1> > PixelShader(ShaderMap);
+		SetGlobalBoundShaderState(Context.RHICmdList, FeatureLevel, BoundShaderState, GFilterVertexDeclaration.VertexDeclarationRHI, *VertexShader, *PixelShader);
 		
 		PixelShader->SetParameters(Context);
 	}
@@ -146,8 +149,8 @@ void FRCPassPostProcessDOFSetup::Process(FRenderingCompositePassContext& Context
 	{
 		static FGlobalBoundShaderState BoundShaderState;
 		
-		TShaderMapRef< FPostProcessDOFSetupPS<0> > PixelShader(GetGlobalShaderMap());
-		SetGlobalBoundShaderState(Context.RHICmdList, Context.GetFeatureLevel(), BoundShaderState, GFilterVertexDeclaration.VertexDeclarationRHI, *VertexShader, *PixelShader);
+		TShaderMapRef< FPostProcessDOFSetupPS<0> > PixelShader(ShaderMap);
+		SetGlobalBoundShaderState(Context.RHICmdList, FeatureLevel, BoundShaderState, GFilterVertexDeclaration.VertexDeclarationRHI, *VertexShader, *PixelShader);
 
 		PixelShader->SetParameters(Context);
 	}
@@ -273,6 +276,9 @@ void FRCPassPostProcessDOFRecombine::Process(FRenderingCompositePassContext& Con
 
 	const FSceneView& View = Context.View;
 
+	const auto FeatureLevel = Context.GetFeatureLevel();
+	auto ShaderMap = Context.GetShaderMap();
+
 	FIntPoint TexSize = InputDesc->Extent;
 
 	// usually 1, 2, 4 or 8
@@ -295,22 +301,22 @@ void FRCPassPostProcessDOFRecombine::Process(FRenderingCompositePassContext& Con
 	Context.RHICmdList.SetRasterizerState(TStaticRasterizerState<>::GetRHI());
 	Context.RHICmdList.SetDepthStencilState(TStaticDepthStencilState<false, CF_Always>::GetRHI());
 
-	TShaderMapRef<FPostProcessVS> VertexShader(GetGlobalShaderMap());
+	TShaderMapRef<FPostProcessVS> VertexShader(ShaderMap);
 
 	if (bNearBlurEnabled)
 	{
 		static FGlobalBoundShaderState BoundShaderState;
 		
-		TShaderMapRef< FPostProcessDOFRecombinePS<1> > PixelShader(GetGlobalShaderMap());
-		SetGlobalBoundShaderState(Context.RHICmdList, Context.GetFeatureLevel(), BoundShaderState, GFilterVertexDeclaration.VertexDeclarationRHI, *VertexShader, *PixelShader);
+		TShaderMapRef< FPostProcessDOFRecombinePS<1> > PixelShader(ShaderMap);
+		SetGlobalBoundShaderState(Context.RHICmdList, FeatureLevel, BoundShaderState, GFilterVertexDeclaration.VertexDeclarationRHI, *VertexShader, *PixelShader);
 		PixelShader->SetParameters(Context);
 	}
 	else
 	{
 		static FGlobalBoundShaderState BoundShaderState;
 		
-		TShaderMapRef< FPostProcessDOFRecombinePS<0> > PixelShader(GetGlobalShaderMap());
-		SetGlobalBoundShaderState(Context.RHICmdList, Context.GetFeatureLevel(), BoundShaderState, GFilterVertexDeclaration.VertexDeclarationRHI, *VertexShader, *PixelShader);
+		TShaderMapRef< FPostProcessDOFRecombinePS<0> > PixelShader(ShaderMap);
+		SetGlobalBoundShaderState(Context.RHICmdList, FeatureLevel, BoundShaderState, GFilterVertexDeclaration.VertexDeclarationRHI, *VertexShader, *PixelShader);
 		PixelShader->SetParameters(Context);
 	}
 
diff --git a/Engine/Source/Runtime/Renderer/Private/PostProcess/PostProcessDownsample.cpp b/Engine/Source/Runtime/Renderer/Private/PostProcess/PostProcessDownsample.cpp
--- a/Engine/Source/Runtime/Renderer/Private/PostProcess/PostProcessDownsample.cpp
+++ b/Engine/Source/Runtime/Renderer/Private/PostProcess/PostProcessDownsample.cpp
@@ -139,8 +139,9 @@ FRCPassPostProcessDownsample::FRCPassPostProcessDownsample(EPixelFormat InOverri
 template <uint32 Method>
 void FRCPassPostProcessDownsample::SetShader(const FRenderingCompositePassContext& Context)
 {
-	TShaderMapRef<FPostProcessDownsampleVS> VertexShader(GetGlobalShaderMap());
-	TShaderMapRef<FPostProcessDownsamplePS<Method> > PixelShader(GetGlobalShaderMap());
+	auto ShaderMap = Context.GetShaderMap();
+	TShaderMapRef<FPostProcessDownsampleVS> VertexShader(ShaderMap);
+	TShaderMapRef<FPostProcessDownsamplePS<Method> > PixelShader(ShaderMap);
 
 	static FGlobalBoundShaderState BoundShaderState;
 	
@@ -220,7 +221,7 @@ void FRCPassPostProcessDownsample::Process(FRenderingCompositePassContext& Conte
 		bHasCleared = true;
 	}
 
-	TShaderMapRef<FPostProcessDownsampleVS> VertexShader(GetGlobalShaderMap());
+	TShaderMapRef<FPostProcessDownsampleVS> VertexShader(Context.GetShaderMap());
 
 	switch (RectSource)
 	{
diff --git a/Engine/Source/Runtime/Renderer/Private/PostProcess/PostProcessEyeAdaptation.cpp b/Engine/Source/Runtime/Renderer/Private/PostProcess/PostProcessEyeAdaptation.cpp
--- a/Engine/Source/Runtime/Renderer/Private/PostProcess/PostProcessEyeAdaptation.cpp
+++ b/Engine/Source/Runtime/Renderer/Private/PostProcess/PostProcessEyeAdaptation.cpp
@@ -91,8 +91,8 @@ void FRCPassPostProcessEyeAdaptation::Process(FRenderingCompositePassContext& Co
 	Context.RHICmdList.SetRasterizerState(TStaticRasterizerState<>::GetRHI());
 	Context.RHICmdList.SetDepthStencilState(TStaticDepthStencilState<false, CF_Always>::GetRHI());
 
-	TShaderMapRef<FPostProcessVS> VertexShader(GetGlobalShaderMap());
-	TShaderMapRef<FPostProcessEyeAdaptationPS> PixelShader(GetGlobalShaderMap());
+	TShaderMapRef<FPostProcessVS> VertexShader(Context.GetShaderMap());
+	TShaderMapRef<FPostProcessEyeAdaptationPS> PixelShader(Context.GetShaderMap());
 
 	static FGlobalBoundShaderState BoundShaderState;
 	
diff --git a/Engine/Source/Runtime/Renderer/Private/PostProcess/PostProcessGBufferHints.cpp b/Engine/Source/Runtime/Renderer/Private/PostProcess/PostProcessGBufferHints.cpp
--- a/Engine/Source/Runtime/Renderer/Private/PostProcess/PostProcessGBufferHints.cpp
+++ b/Engine/Source/Runtime/Renderer/Private/PostProcess/PostProcessGBufferHints.cpp
@@ -114,8 +114,8 @@ void FRCPassPostProcessGBufferHints::Process(FRenderingCompositePassContext& Con
 	Context.RHICmdList.SetRasterizerState(TStaticRasterizerState<>::GetRHI());
 	Context.RHICmdList.SetDepthStencilState(TStaticDepthStencilState<false, CF_Always>::GetRHI());
 
-	TShaderMapRef<FPostProcessVS> VertexShader(GetGlobalShaderMap());
-	TShaderMapRef<FPostProcessGBufferHintsPS> PixelShader(GetGlobalShaderMap());
+	TShaderMapRef<FPostProcessVS> VertexShader(Context.GetShaderMap());
+	TShaderMapRef<FPostProcessGBufferHintsPS> PixelShader(Context.GetShaderMap());
 
 	static FGlobalBoundShaderState BoundShaderState;
 	
diff --git a/Engine/Source/Runtime/Renderer/Private/PostProcess/PostProcessHMD.cpp b/Engine/Source/Runtime/Renderer/Private/PostProcess/PostProcessHMD.cpp
--- a/Engine/Source/Runtime/Renderer/Private/PostProcess/PostProcessHMD.cpp
+++ b/Engine/Source/Runtime/Renderer/Private/PostProcess/PostProcessHMD.cpp
@@ -219,8 +219,8 @@ void FRCPassPostProcessHMD::Process(FRenderingCompositePassContext& Context)
 	check(GEngine->HMDDevice.IsValid());
 
 	{
-		TShaderMapRef<FPostProcessHMDVS<false> > VertexShader(GetGlobalShaderMap());
-		TShaderMapRef<FPostProcessHMDPS<false> > PixelShader(GetGlobalShaderMap());
+		TShaderMapRef<FPostProcessHMDVS<false> > VertexShader(Context.GetShaderMap());
+		TShaderMapRef<FPostProcessHMDPS<false> > PixelShader(Context.GetShaderMap());
 		static FGlobalBoundShaderState BoundShaderState;
 		
 		SetGlobalBoundShaderState(Context.RHICmdList, Context.GetFeatureLevel(), BoundShaderState, GDistortionVertexDeclaration.VertexDeclarationRHI, *VertexShader, *PixelShader);
diff --git a/Engine/Source/Runtime/Renderer/Private/PostProcess/PostProcessHistogram.cpp b/Engine/Source/Runtime/Renderer/Private/PostProcess/PostProcessHistogram.cpp
--- a/Engine/Source/Runtime/Renderer/Private/PostProcess/PostProcessHistogram.cpp
+++ b/Engine/Source/Runtime/Renderer/Private/PostProcess/PostProcessHistogram.cpp
@@ -108,7 +108,7 @@ void FRCPassPostProcessHistogram::Process(FRenderingCompositePassContext& Contex
 
 	const FSceneRenderTargetItem& DestRenderTarget = PassOutputs[0].RequestSurface(Context);
 
-	TShaderMapRef<FPostProcessHistogramCS> ComputeShader(GetGlobalShaderMap());
+	TShaderMapRef<FPostProcessHistogramCS> ComputeShader(Context.GetShaderMap());
 
 	Context.RHICmdList.SetComputeShader(ComputeShader->GetComputeShader());
 	SetRenderTarget(Context.RHICmdList, FTextureRHIRef(), FTextureRHIRef());
diff --git a/Engine/Source/Runtime/Renderer/Private/PostProcess/PostProcessHistogramReduce.cpp b/Engine/Source/Runtime/Renderer/Private/PostProcess/PostProcessHistogramReduce.cpp
--- a/Engine/Source/Runtime/Renderer/Private/PostProcess/PostProcessHistogramReduce.cpp
+++ b/Engine/Source/Runtime/Renderer/Private/PostProcess/PostProcessHistogramReduce.cpp
@@ -114,8 +114,8 @@ void FRCPassPostProcessHistogramReduce::Process(FRenderingCompositePassContext&
 	Context.RHICmdList.SetRasterizerState(TStaticRasterizerState<>::GetRHI());
 	Context.RHICmdList.SetDepthStencilState(TStaticDepthStencilState<false, CF_Always>::GetRHI());
 
-	TShaderMapRef<FPostProcessVS> VertexShader(GetGlobalShaderMap());
-	TShaderMapRef<FPostProcessHistogramReducePS> PixelShader(GetGlobalShaderMap());
+	TShaderMapRef<FPostProcessVS> VertexShader(Context.GetShaderMap());
+	TShaderMapRef<FPostProcessHistogramReducePS> PixelShader(Context.GetShaderMap());
 
 	static FGlobalBoundShaderState BoundShaderState;
 	
diff --git a/Engine/Source/Runtime/Renderer/Private/PostProcess/PostProcessLensBlur.cpp b/Engine/Source/Runtime/Renderer/Private/PostProcess/PostProcessLensBlur.cpp
--- a/Engine/Source/Runtime/Renderer/Private/PostProcess/PostProcessLensBlur.cpp
+++ b/Engine/Source/Runtime/Renderer/Private/PostProcess/PostProcessLensBlur.cpp
@@ -196,8 +196,8 @@ void FRCPassPostProcessLensBlur::Process(FRenderingCompositePassContext& Context
 	Context.RHICmdList.SetRasterizerState(TStaticRasterizerState<>::GetRHI());
 	Context.RHICmdList.SetDepthStencilState(TStaticDepthStencilState<false, CF_Always>::GetRHI());
 
-	TShaderMapRef<FPostProcessLensBlurVS> VertexShader(GetGlobalShaderMap());
-	TShaderMapRef<FPostProcessLensBlurPS> PixelShader(GetGlobalShaderMap());
+	TShaderMapRef<FPostProcessLensBlurVS> VertexShader(Context.GetShaderMap());
+	TShaderMapRef<FPostProcessLensBlurPS> PixelShader(Context.GetShaderMap());
 
 	static FGlobalBoundShaderState BoundShaderState;
 	
diff --git a/Engine/Source/Runtime/Renderer/Private/PostProcess/PostProcessLensFlares.cpp b/Engine/Source/Runtime/Renderer/Private/PostProcess/PostProcessLensFlares.cpp
--- a/Engine/Source/Runtime/Renderer/Private/PostProcess/PostProcessLensFlares.cpp
+++ b/Engine/Source/Runtime/Renderer/Private/PostProcess/PostProcessLensFlares.cpp
@@ -154,12 +154,12 @@ void FRCPassPostProcessLensFlares::Process(FRenderingCompositePassContext& Conte
 	Context.RHICmdList.SetRasterizerState(TStaticRasterizerState<>::GetRHI());
 	Context.RHICmdList.SetDepthStencilState(TStaticDepthStencilState<false, CF_Always>::GetRHI());
 
-	TShaderMapRef<FPostProcessVS> VertexShader(GetGlobalShaderMap());
+	TShaderMapRef<FPostProcessVS> VertexShader(Context.GetShaderMap());
 
 	
 	// setup background (bloom), can be implemented to use additive blending to avoid the read here
 	{
-		TShaderMapRef<FPostProcessLensFlareBasePS> PixelShader(GetGlobalShaderMap());
+		TShaderMapRef<FPostProcessLensFlareBasePS> PixelShader(Context.GetShaderMap());
 
 		static FGlobalBoundShaderState BoundShaderState;
 		
@@ -186,7 +186,7 @@ void FRCPassPostProcessLensFlares::Process(FRenderingCompositePassContext& Conte
 
 	// add lens flares on top of that
 	{
-		TShaderMapRef<FPostProcessLensFlaresPS> PixelShader(GetGlobalShaderMap());
+		TShaderMapRef<FPostProcessLensFlaresPS> PixelShader(Context.GetShaderMap());
 
 		static FGlobalBoundShaderState BoundShaderState;
 		
diff --git a/Engine/Source/Runtime/Renderer/Private/PostProcess/PostProcessMobile.cpp b/Engine/Source/Runtime/Renderer/Private/PostProcess/PostProcessMobile.cpp
--- a/Engine/Source/Runtime/Renderer/Private/PostProcess/PostProcessMobile.cpp
+++ b/Engine/Source/Runtime/Renderer/Private/PostProcess/PostProcessMobile.cpp
@@ -135,8 +135,8 @@ IMPLEMENT_SHADER_TYPE(template<>,FPostProcessBloomSetupPS_ES2_4,TEXT("PostProces
 template <uint32 UseSunDof>
 static void BloomSetup_SetShader(const FRenderingCompositePassContext& Context)
 {
-	TShaderMapRef<FPostProcessBloomSetupVS_ES2> VertexShader(GetGlobalShaderMap());
-	TShaderMapRef<FPostProcessBloomSetupPS_ES2<UseSunDof> > PixelShader(GetGlobalShaderMap());
+	TShaderMapRef<FPostProcessBloomSetupVS_ES2> VertexShader(Context.GetShaderMap());
+	TShaderMapRef<FPostProcessBloomSetupPS_ES2<UseSunDof> > PixelShader(Context.GetShaderMap());
 
 	static FGlobalBoundShaderState BoundShaderState;
 	
@@ -216,7 +216,7 @@ void FRCPassPostProcessBloomSetupES2::Process(FRenderingCompositePassContext& Co
 
 	SetShader(Context);
 
-	TShaderMapRef<FPostProcessBloomSetupVS_ES2> VertexShader(GetGlobalShaderMap());
+	TShaderMapRef<FPostProcessBloomSetupVS_ES2> VertexShader(Context.GetShaderMap());
 
 	DrawRectangle(
 		Context.RHICmdList,
@@ -356,8 +356,8 @@ IMPLEMENT_SHADER_TYPE(,FPostProcessBloomSetupSmallPS_ES2,TEXT("PostProcessMobile
 
 void FRCPassPostProcessBloomSetupSmallES2::SetShader(const FRenderingCompositePassContext& Context)
 {
-	TShaderMapRef<FPostProcessBloomSetupSmallVS_ES2> VertexShader(GetGlobalShaderMap());
-	TShaderMapRef<FPostProcessBloomSetupSmallPS_ES2> PixelShader(GetGlobalShaderMap());
+	TShaderMapRef<FPostProcessBloomSetupSmallVS_ES2> VertexShader(Context.GetShaderMap());
+	TShaderMapRef<FPostProcessBloomSetupSmallPS_ES2> PixelShader(Context.GetShaderMap());
 
 	static FGlobalBoundShaderState BoundShaderState;
 	
@@ -417,7 +417,7 @@ void FRCPassPostProcessBloomSetupSmallES2::Process(FRenderingCompositePassContex
 
 	SetShader(Context);
 
-	TShaderMapRef<FPostProcessBloomSetupSmallVS_ES2> VertexShader(GetGlobalShaderMap());
+	TShaderMapRef<FPostProcessBloomSetupSmallVS_ES2> VertexShader(Context.GetShaderMap());
 	DrawRectangle(
 		Context.RHICmdList,
 		0, 0,
@@ -568,8 +568,8 @@ void FRCPassPostProcessBloomDownES2::Process(FRenderingCompositePassContext& Con
 	Context.RHICmdList.SetRasterizerState(TStaticRasterizerState<>::GetRHI());
 	Context.RHICmdList.SetDepthStencilState(TStaticDepthStencilState<false, CF_Always>::GetRHI());
 
-	TShaderMapRef<FPostProcessBloomDownVS_ES2> VertexShader(GetGlobalShaderMap());
-	TShaderMapRef<FPostProcessBloomDownPS_ES2> PixelShader(GetGlobalShaderMap());
+	TShaderMapRef<FPostProcessBloomDownVS_ES2> VertexShader(Context.GetShaderMap());
+	TShaderMapRef<FPostProcessBloomDownPS_ES2> PixelShader(Context.GetShaderMap());
 
 	static FGlobalBoundShaderState BoundShaderState;
 	
@@ -736,8 +736,8 @@ void FRCPassPostProcessBloomUpES2::Process(FRenderingCompositePassContext& Conte
 	Context.RHICmdList.SetRasterizerState(TStaticRasterizerState<>::GetRHI());
 	Context.RHICmdList.SetDepthStencilState(TStaticDepthStencilState<false, CF_Always>::GetRHI());
 
-	TShaderMapRef<FPostProcessBloomUpVS_ES2> VertexShader(GetGlobalShaderMap());
-	TShaderMapRef<FPostProcessBloomUpPS_ES2> PixelShader(GetGlobalShaderMap());
+	TShaderMapRef<FPostProcessBloomUpVS_ES2> VertexShader(Context.GetShaderMap());
+	TShaderMapRef<FPostProcessBloomUpPS_ES2> PixelShader(Context.GetShaderMap());
 
 	static FGlobalBoundShaderState BoundShaderState;
 	
@@ -906,8 +906,8 @@ IMPLEMENT_SHADER_TYPE(,FPostProcessSunMaskVS_ES2,TEXT("PostProcessMobile"),TEXT(
 template <uint32 UseFetchSunDof>
 static void SunMask_SetShader(const FRenderingCompositePassContext& Context)
 {
-	TShaderMapRef<FPostProcessSunMaskVS_ES2> VertexShader(GetGlobalShaderMap());
-	TShaderMapRef<FPostProcessSunMaskPS_ES2<UseFetchSunDof> > PixelShader(GetGlobalShaderMap());
+	TShaderMapRef<FPostProcessSunMaskVS_ES2> VertexShader(Context.GetShaderMap());
+	TShaderMapRef<FPostProcessSunMaskPS_ES2<UseFetchSunDof> > PixelShader(Context.GetShaderMap());
 
 	static FGlobalBoundShaderState BoundShaderState;
 	
@@ -966,7 +966,7 @@ void FRCPassPostProcessSunMaskES2::Process(FRenderingCompositePassContext& Conte
 	FIntRect SrcRect;
 	const FSceneView& View = Context.View;
 
-	TShaderMapRef<FPostProcessSunMaskVS_ES2> VertexShader(GetGlobalShaderMap());
+	TShaderMapRef<FPostProcessSunMaskVS_ES2> VertexShader(Context.GetShaderMap());
 	if(bOnChip)
 	{
 		SrcSize = DstSize;
@@ -1144,8 +1144,8 @@ IMPLEMENT_SHADER_TYPE(,FPostProcessSunAlphaVS_ES2,TEXT("PostProcessMobile"),TEXT
 template <uint32 UseDof>
 static void SunAlpha_SetShader(const FRenderingCompositePassContext& Context)
 {
-	TShaderMapRef<FPostProcessSunAlphaVS_ES2> VertexShader(GetGlobalShaderMap());
-	TShaderMapRef<FPostProcessSunAlphaPS_ES2<UseDof> > PixelShader(GetGlobalShaderMap());
+	TShaderMapRef<FPostProcessSunAlphaVS_ES2> VertexShader(Context.GetShaderMap());
+	TShaderMapRef<FPostProcessSunAlphaPS_ES2<UseDof> > PixelShader(Context.GetShaderMap());
 
 	static FGlobalBoundShaderState BoundShaderState;
 
@@ -1195,7 +1195,7 @@ void FRCPassPostProcessSunAlphaES2::Process(FRenderingCompositePassContext& Cont
 	SetShader(Context);
 
 	FIntPoint SrcDstSize = PrePostSourceViewportSize / 4;
-	TShaderMapRef<FPostProcessSunAlphaVS_ES2> VertexShader(GetGlobalShaderMap());
+	TShaderMapRef<FPostProcessSunAlphaVS_ES2> VertexShader(Context.GetShaderMap());
 
 	DrawRectangle(
 		Context.RHICmdList,
@@ -1349,8 +1349,8 @@ void FRCPassPostProcessSunBlurES2::Process(FRenderingCompositePassContext& Conte
 	Context.RHICmdList.SetRasterizerState(TStaticRasterizerState<>::GetRHI());
 	Context.RHICmdList.SetDepthStencilState(TStaticDepthStencilState<false,CF_Always>::GetRHI());
 
-	TShaderMapRef<FPostProcessSunBlurVS_ES2> VertexShader(GetGlobalShaderMap());
-	TShaderMapRef<FPostProcessSunBlurPS_ES2> PixelShader(GetGlobalShaderMap());
+	TShaderMapRef<FPostProcessSunBlurVS_ES2> VertexShader(Context.GetShaderMap());
+	TShaderMapRef<FPostProcessSunBlurPS_ES2> PixelShader(Context.GetShaderMap());
 
 	static FGlobalBoundShaderState BoundShaderState;
 	
@@ -1520,8 +1520,8 @@ IMPLEMENT_SHADER_TYPE(,FPostProcessSunMergeVS_ES2,TEXT("PostProcessMobile"),TEXT
 template <uint32 UseSunBloom>
 static void SunMerge_SetShader(const FRenderingCompositePassContext& Context)
 {
-	TShaderMapRef<FPostProcessSunMergeVS_ES2> VertexShader(GetGlobalShaderMap());
-	TShaderMapRef<FPostProcessSunMergePS_ES2<UseSunBloom> > PixelShader(GetGlobalShaderMap());
+	TShaderMapRef<FPostProcessSunMergeVS_ES2> VertexShader(Context.GetShaderMap());
+	TShaderMapRef<FPostProcessSunMergePS_ES2<UseSunBloom> > PixelShader(Context.GetShaderMap());
 
 	static FGlobalBoundShaderState BoundShaderState;
 	
@@ -1576,7 +1576,7 @@ void FRCPassPostProcessSunMergeES2::Process(FRenderingCompositePassContext& Cont
 	SetShader(Context);
 
 	FIntPoint SrcDstSize = PrePostSourceViewportSize / 4;
-	TShaderMapRef<FPostProcessSunMergeVS_ES2> VertexShader(GetGlobalShaderMap());
+	TShaderMapRef<FPostProcessSunMergeVS_ES2> VertexShader(Context.GetShaderMap());
 
 	DrawRectangle(
 		Context.RHICmdList,
@@ -1731,8 +1731,8 @@ IMPLEMENT_SHADER_TYPE(,FPostProcessSunMergeSmallVS_ES2,TEXT("PostProcessMobile")
 
 void FRCPassPostProcessSunMergeSmallES2::SetShader(const FRenderingCompositePassContext& Context)
 {
-	TShaderMapRef<FPostProcessSunMergeSmallVS_ES2> VertexShader(GetGlobalShaderMap());
-	TShaderMapRef<FPostProcessSunMergeSmallPS_ES2> PixelShader(GetGlobalShaderMap());
+	TShaderMapRef<FPostProcessSunMergeSmallVS_ES2> VertexShader(Context.GetShaderMap());
+	TShaderMapRef<FPostProcessSunMergeSmallPS_ES2> PixelShader(Context.GetShaderMap());
 
 	static FGlobalBoundShaderState BoundShaderState;
 	
@@ -1771,7 +1771,7 @@ void FRCPassPostProcessSunMergeSmallES2::Process(FRenderingCompositePassContext&
 	SetShader(Context);
 
 	FIntPoint SrcDstSize = PrePostSourceViewportSize / 4;
-	TShaderMapRef<FPostProcessSunMergeSmallVS_ES2> VertexShader(GetGlobalShaderMap());
+	TShaderMapRef<FPostProcessSunMergeSmallVS_ES2> VertexShader(Context.GetShaderMap());
 
 	DrawRectangle(
 		Context.RHICmdList, 
@@ -1921,8 +1921,8 @@ IMPLEMENT_SHADER_TYPE(template<>,FPostProcessDofDownPS_ES2_1,TEXT("PostProcessMo
 template <uint32 UseSun>
 static void DofDown_SetShader(const FRenderingCompositePassContext& Context)
 {
-	TShaderMapRef<FPostProcessDofDownVS_ES2> VertexShader(GetGlobalShaderMap());
-	TShaderMapRef<FPostProcessDofDownPS_ES2<UseSun> > PixelShader(GetGlobalShaderMap());
+	TShaderMapRef<FPostProcessDofDownVS_ES2> VertexShader(Context.GetShaderMap());
+	TShaderMapRef<FPostProcessDofDownPS_ES2<UseSun> > PixelShader(Context.GetShaderMap());
 
 	static FGlobalBoundShaderState BoundShaderState;
 	
@@ -1996,7 +1996,7 @@ void FRCPassPostProcessDofDownES2::Process(FRenderingCompositePassContext& Conte
 
 	SetShader(Context);
 
-	TShaderMapRef<FPostProcessDofDownVS_ES2> VertexShader(GetGlobalShaderMap());
+	TShaderMapRef<FPostProcessDofDownVS_ES2> VertexShader(Context.GetShaderMap());
 
 	DrawRectangle(
 		Context.RHICmdList,
@@ -2129,8 +2129,8 @@ IMPLEMENT_SHADER_TYPE(template<>,FPostProcessDofNearPS_ES2_1,TEXT("PostProcessMo
 template <uint32 UseSun>
 static void DofNear_SetShader(const FRenderingCompositePassContext& Context)
 {
-	TShaderMapRef<FPostProcessDofNearVS_ES2> VertexShader(GetGlobalShaderMap());
-	TShaderMapRef<FPostProcessDofNearPS_ES2<UseSun> > PixelShader(GetGlobalShaderMap());
+	TShaderMapRef<FPostProcessDofNearVS_ES2> VertexShader(Context.GetShaderMap());
+	TShaderMapRef<FPostProcessDofNearPS_ES2<UseSun> > PixelShader(Context.GetShaderMap());
 
 	static FGlobalBoundShaderState BoundShaderState;
 	
@@ -2182,7 +2182,7 @@ void FRCPassPostProcessDofNearES2::Process(FRenderingCompositePassContext& Conte
 	SetShader(Context);
 
 	FIntPoint SrcDstSize = PrePostSourceViewportSize / 4;
-	TShaderMapRef<FPostProcessDofNearVS_ES2> VertexShader(GetGlobalShaderMap());
+	TShaderMapRef<FPostProcessDofNearVS_ES2> VertexShader(Context.GetShaderMap());
 
 	DrawRectangle(
 		Context.RHICmdList,
@@ -2329,8 +2329,8 @@ void FRCPassPostProcessDofBlurES2::Process(FRenderingCompositePassContext& Conte
 	Context.RHICmdList.SetRasterizerState(TStaticRasterizerState<>::GetRHI());
 	Context.RHICmdList.SetDepthStencilState(TStaticDepthStencilState<false,CF_Always>::GetRHI());
 
-	TShaderMapRef<FPostProcessDofBlurVS_ES2> VertexShader(GetGlobalShaderMap());
-	TShaderMapRef<FPostProcessDofBlurPS_ES2> PixelShader(GetGlobalShaderMap());
+	TShaderMapRef<FPostProcessDofBlurVS_ES2> VertexShader(Context.GetShaderMap());
+	TShaderMapRef<FPostProcessDofBlurPS_ES2> PixelShader(Context.GetShaderMap());
 
 	static FGlobalBoundShaderState BoundShaderState;
 	
@@ -2467,8 +2467,8 @@ IMPLEMENT_SHADER_TYPE(,FPostProcessSunAvgVS_ES2,TEXT("PostProcessMobile"),TEXT("
 
 static void SunAvg_SetShader(const FRenderingCompositePassContext& Context)
 {
-	TShaderMapRef<FPostProcessSunAvgVS_ES2> VertexShader(GetGlobalShaderMap());
-	TShaderMapRef<FPostProcessSunAvgPS_ES2> PixelShader(GetGlobalShaderMap());
+	TShaderMapRef<FPostProcessSunAvgVS_ES2> VertexShader(Context.GetShaderMap());
+	TShaderMapRef<FPostProcessSunAvgPS_ES2> PixelShader(Context.GetShaderMap());
 
 	static FGlobalBoundShaderState BoundShaderState;
 	
@@ -2512,7 +2512,7 @@ void FRCPassPostProcessSunAvgES2::Process(FRenderingCompositePassContext& Contex
 	SetShader(Context);
 
 	FIntPoint SrcDstSize = PrePostSourceViewportSize / 4;
-	TShaderMapRef<FPostProcessSunAvgVS_ES2> VertexShader(GetGlobalShaderMap());
+	TShaderMapRef<FPostProcessSunAvgVS_ES2> VertexShader(Context.GetShaderMap());
 
 	DrawRectangle(
 		Context.RHICmdList,
@@ -2710,8 +2710,8 @@ IMPLEMENT_SHADER_TYPE(,FPostProcessAaVS_ES2,TEXT("PostProcessMobile"),TEXT("AaVS
 
 static void Aa_SetShader(const FRenderingCompositePassContext& Context)
 {
-	TShaderMapRef<FPostProcessAaVS_ES2> VertexShader(GetGlobalShaderMap());
-	TShaderMapRef<FPostProcessAaPS_ES2> PixelShader(GetGlobalShaderMap());
+	TShaderMapRef<FPostProcessAaVS_ES2> VertexShader(Context.GetShaderMap());
+	TShaderMapRef<FPostProcessAaPS_ES2> PixelShader(Context.GetShaderMap());
 
 	static FGlobalBoundShaderState BoundShaderState;
 	
@@ -2755,7 +2755,7 @@ void FRCPassPostProcessAaES2::Process(FRenderingCompositePassContext& Context)
 	SetShader(Context);
 
 	FIntPoint SrcDstSize = PrePostSourceViewportSize;
-	TShaderMapRef<FPostProcessAaVS_ES2> VertexShader(GetGlobalShaderMap());
+	TShaderMapRef<FPostProcessAaVS_ES2> VertexShader(Context.GetShaderMap());
 
 	DrawRectangle(
 		Context.RHICmdList,
diff --git a/Engine/Source/Runtime/Renderer/Private/PostProcess/PostProcessMorpheus.cpp b/Engine/Source/Runtime/Renderer/Private/PostProcess/PostProcessMorpheus.cpp
--- a/Engine/Source/Runtime/Renderer/Private/PostProcess/PostProcessMorpheus.cpp
+++ b/Engine/Source/Runtime/Renderer/Private/PostProcess/PostProcessMorpheus.cpp
@@ -205,8 +205,8 @@ void FRCPassPostProcessMorpheus::Process(FRenderingCompositePassContext& Context
 	Context.RHICmdList.SetRasterizerState(TStaticRasterizerState<>::GetRHI());
 	Context.RHICmdList.SetDepthStencilState(TStaticDepthStencilState<false, CF_Always>::GetRHI());
 
-	TShaderMapRef<FPostProcessMorpheusVS> VertexShader(GetGlobalShaderMap());
-	TShaderMapRef<FPostProcessMorpheusPS> PixelShader(GetGlobalShaderMap());
+	TShaderMapRef<FPostProcessMorpheusVS> VertexShader(Context.GetShaderMap());
+	TShaderMapRef<FPostProcessMorpheusPS> PixelShader(Context.GetShaderMap());
 
 	static FGlobalBoundShaderState BoundShaderState;
 	
diff --git a/Engine/Source/Runtime/Renderer/Private/PostProcess/PostProcessMotionBlur.cpp b/Engine/Source/Runtime/Renderer/Private/PostProcess/PostProcessMotionBlur.cpp
--- a/Engine/Source/Runtime/Renderer/Private/PostProcess/PostProcessMotionBlur.cpp
+++ b/Engine/Source/Runtime/Renderer/Private/PostProcess/PostProcessMotionBlur.cpp
@@ -135,10 +135,10 @@ void FRCPassPostProcessMotionBlurSetup::Process(FRenderingCompositePassContext&
 	Context.RHICmdList.SetRasterizerState(TStaticRasterizerState<>::GetRHI());
 	Context.RHICmdList.SetDepthStencilState(TStaticDepthStencilState<false, CF_Always>::GetRHI());
 
-	TShaderMapRef<FPostProcessVS> VertexShader(GetGlobalShaderMap());
+	TShaderMapRef<FPostProcessVS> VertexShader(Context.GetShaderMap());
 
 	{		
-		TShaderMapRef<FPostProcessMotionBlurSetupPS > PixelShader(GetGlobalShaderMap());
+		TShaderMapRef<FPostProcessMotionBlurSetupPS > PixelShader(Context.GetShaderMap());
 		static FGlobalBoundShaderState BoundShaderState;
 		
 
@@ -374,8 +374,8 @@ VARIATION1(0)			VARIATION1(1)			VARIATION1(2)			VARIATION1(3)			VARIATION1(4)
 template <uint32 Quality>
 static void SetMotionBlurShaderTempl(const FRenderingCompositePassContext& Context)
 {
-	TShaderMapRef<FPostProcessVS> VertexShader(GetGlobalShaderMap());
-	TShaderMapRef<FPostProcessMotionBlurPS<Quality> > PixelShader(GetGlobalShaderMap());
+	TShaderMapRef<FPostProcessVS> VertexShader(Context.GetShaderMap());
+	TShaderMapRef<FPostProcessMotionBlurPS<Quality> > PixelShader(Context.GetShaderMap());
 
 	static FGlobalBoundShaderState BoundShaderState;
 	
@@ -453,7 +453,7 @@ void FRCPassPostProcessMotionBlur::Process(FRenderingCompositePassContext& Conte
 		SetMotionBlurShaderTempl<4>(Context);
 	}
 
-	TShaderMapRef<FPostProcessVS> VertexShader(GetGlobalShaderMap());
+	TShaderMapRef<FPostProcessVS> VertexShader(Context.GetShaderMap());
 
 	// Draw a quad mapping scene color to the view's render target
 	DrawRectangle(
@@ -570,8 +570,8 @@ void FRCPassPostProcessMotionBlurRecombine::Process(FRenderingCompositePassConte
 	Context.RHICmdList.SetRasterizerState(TStaticRasterizerState<>::GetRHI());
 	Context.RHICmdList.SetDepthStencilState(TStaticDepthStencilState<false, CF_Always>::GetRHI());
 
-	TShaderMapRef<FPostProcessVS> VertexShader(GetGlobalShaderMap());
-	TShaderMapRef<FPostProcessMotionBlurRecombinePS> PixelShader(GetGlobalShaderMap());
+	TShaderMapRef<FPostProcessVS> VertexShader(Context.GetShaderMap());
+	TShaderMapRef<FPostProcessMotionBlurRecombinePS> PixelShader(Context.GetShaderMap());
 
 	static FGlobalBoundShaderState BoundShaderState;
 	
@@ -655,7 +655,7 @@ void FRCPassPostProcessVisualizeMotionBlur::Process(FRenderingCompositePassConte
 	SetMotionBlurShaderTempl<0>(Context);
 
 	// Draw a quad mapping scene color to the view's render target
-	TShaderMapRef<FPostProcessVS> VertexShader(GetGlobalShaderMap());
+	TShaderMapRef<FPostProcessVS> VertexShader(Context.GetShaderMap());
 
 	DrawRectangle(
 		Context.RHICmdList,
diff --git a/Engine/Source/Runtime/Renderer/Private/PostProcess/PostProcessNoiseBlur.cpp b/Engine/Source/Runtime/Renderer/Private/PostProcess/PostProcessNoiseBlur.cpp
--- a/Engine/Source/Runtime/Renderer/Private/PostProcess/PostProcessNoiseBlur.cpp
+++ b/Engine/Source/Runtime/Renderer/Private/PostProcess/PostProcessNoiseBlur.cpp
@@ -97,8 +97,8 @@ FRCPassPostProcessNoiseBlur::FRCPassPostProcessNoiseBlur(float InRadius, EPixelF
 template <uint32 Method>
 void SetNoiseBlurShader(const FRenderingCompositePassContext& Context, float InRadius)
 {
-	TShaderMapRef<FPostProcessVS> VertexShader(GetGlobalShaderMap());
-	TShaderMapRef<FPostProcessNoiseBlurPS<Method> > PixelShader(GetGlobalShaderMap());
+	TShaderMapRef<FPostProcessVS> VertexShader(Context.GetShaderMap());
+	TShaderMapRef<FPostProcessNoiseBlurPS<Method> > PixelShader(Context.GetShaderMap());
 
 	static FGlobalBoundShaderState BoundShaderState;
 	
@@ -163,7 +163,7 @@ void FRCPassPostProcessNoiseBlur::Process(FRenderingCompositePassContext& Contex
 	}
 
 	// Draw a quad mapping scene color to the view's render target
-	TShaderMapRef<FPostProcessVS> VertexShader(GetGlobalShaderMap());
+	TShaderMapRef<FPostProcessVS> VertexShader(Context.GetShaderMap());
 	DrawRectangle(
 		Context.RHICmdList,
 		DestRect.Min.X, DestRect.Min.Y,
diff --git a/Engine/Source/Runtime/Renderer/Private/PostProcess/PostProcessSceneColorFringe.cpp b/Engine/Source/Runtime/Renderer/Private/PostProcess/PostProcessSceneColorFringe.cpp
--- a/Engine/Source/Runtime/Renderer/Private/PostProcess/PostProcessSceneColorFringe.cpp
+++ b/Engine/Source/Runtime/Renderer/Private/PostProcess/PostProcessSceneColorFringe.cpp
@@ -167,8 +167,8 @@ void FRCPassPostProcessSceneColorFringe::Process(FRenderingCompositePassContext&
 	Context.RHICmdList.SetRasterizerState(TStaticRasterizerState<>::GetRHI());
 	Context.RHICmdList.SetDepthStencilState(TStaticDepthStencilState<false, CF_Always>::GetRHI());
 	
-	TShaderMapRef<FPostProcessSceneColorFringeVS> VertexShader(GetGlobalShaderMap());
-	TShaderMapRef<FPostProcessSceneColorFringePS> PixelShader(GetGlobalShaderMap());
+	TShaderMapRef<FPostProcessSceneColorFringeVS> VertexShader(Context.GetShaderMap());
+	TShaderMapRef<FPostProcessSceneColorFringePS> PixelShader(Context.GetShaderMap());
 
 	static FGlobalBoundShaderState BoundShaderState;
 	
diff --git a/Engine/Source/Runtime/Renderer/Private/PostProcess/PostProcessSelectionOutline.cpp b/Engine/Source/Runtime/Renderer/Private/PostProcess/PostProcessSelectionOutline.cpp
--- a/Engine/Source/Runtime/Renderer/Private/PostProcess/PostProcessSelectionOutline.cpp
+++ b/Engine/Source/Runtime/Renderer/Private/PostProcess/PostProcessSelectionOutline.cpp
@@ -351,8 +351,8 @@ VARIATION1(1)			VARIATION1(2)			VARIATION1(4)			VARIATION1(8)
 template <uint32 MSAASampleCount>
 static void SetSelectionOutlineShaderTempl(const FRenderingCompositePassContext& Context)
 {
-	TShaderMapRef<FPostProcessVS> VertexShader(GetGlobalShaderMap());
-	TShaderMapRef<FPostProcessSelectionOutlinePS<MSAASampleCount> > PixelShader(GetGlobalShaderMap());
+	TShaderMapRef<FPostProcessVS> VertexShader(Context.GetShaderMap());
+	TShaderMapRef<FPostProcessSelectionOutlinePS<MSAASampleCount> > PixelShader(Context.GetShaderMap());
 
 	static FGlobalBoundShaderState BoundShaderState;
 	
@@ -415,7 +415,7 @@ void FRCPassPostProcessSelectionOutline::Process(FRenderingCompositePassContext&
 	}
 
 	// Draw a quad mapping scene color to the view's render target
-	TShaderMapRef<FPostProcessVS> VertexShader(GetGlobalShaderMap());
+	TShaderMapRef<FPostProcessVS> VertexShader(Context.GetShaderMap());
 	DrawRectangle(
 		Context.RHICmdList,
 		0, 0,
diff --git a/Engine/Source/Runtime/Renderer/Private/PostProcess/PostProcessSubsurface.cpp b/Engine/Source/Runtime/Renderer/Private/PostProcess/PostProcessSubsurface.cpp
--- a/Engine/Source/Runtime/Renderer/Private/PostProcess/PostProcessSubsurface.cpp
+++ b/Engine/Source/Runtime/Renderer/Private/PostProcess/PostProcessSubsurface.cpp
@@ -92,8 +92,8 @@ public:
 template <uint32 SetupMode>
 void SetSubsurfaceSetupShader(const FRenderingCompositePassContext& Context)
 {
-	TShaderMapRef<FPostProcessVS> VertexShader(GetGlobalShaderMap());
-	TShaderMapRef<FPostProcessSubsurfaceSetupPS<SetupMode> > PixelShader(GetGlobalShaderMap());
+	TShaderMapRef<FPostProcessVS> VertexShader(Context.GetShaderMap());
+	TShaderMapRef<FPostProcessSubsurfaceSetupPS<SetupMode> > PixelShader(Context.GetShaderMap());
 
 	static FGlobalBoundShaderState BoundShaderState;
 
@@ -194,7 +194,7 @@ void FRCPassPostProcessSubsurfaceSetup::Process(FRenderingCompositePassContext&
 	}
 
 	// Draw a quad mapping scene color to the view's render target
-	TShaderMapRef<FPostProcessVS> VertexShader(GetGlobalShaderMap());
+	TShaderMapRef<FPostProcessVS> VertexShader(Context.GetShaderMap());
 	DrawRectangle(
 		Context.RHICmdList,
 		DestRect.Min.X, DestRect.Min.Y,
@@ -378,8 +378,8 @@ FRCPassPostProcessSubsurface::FRCPassPostProcessSubsurface(uint32 InPass, float
 template <uint32 Method>
 void SetSubsurfaceShader(const FRenderingCompositePassContext& Context, float InRadius)
 {
-	TShaderMapRef<FPostProcessVS> VertexShader(GetGlobalShaderMap());
-	TShaderMapRef<FPostProcessSubsurfacePS<Method> > PixelShader(GetGlobalShaderMap());
+	TShaderMapRef<FPostProcessVS> VertexShader(Context.GetShaderMap());
+	TShaderMapRef<FPostProcessSubsurfacePS<Method> > PixelShader(Context.GetShaderMap());
 
 	static FGlobalBoundShaderState BoundShaderState;
 
@@ -431,7 +431,7 @@ void FRCPassPostProcessSubsurface::Process(FRenderingCompositePassContext& Conte
 	Context.RHICmdList.SetRasterizerState(TStaticRasterizerState<>::GetRHI());
 	Context.RHICmdList.SetDepthStencilState(TStaticDepthStencilState<false, CF_Always>::GetRHI());
 
-	TShaderMapRef<FPostProcessVS> VertexShader(GetGlobalShaderMap());
+	TShaderMapRef<FPostProcessVS> VertexShader(Context.GetShaderMap());
 
 	SCOPED_DRAW_EVENT(SubsurfacePass, DEC_SCENE_ITEMS);
 
diff --git a/Engine/Source/Runtime/Renderer/Private/PostProcess/PostProcessTemporalAA.cpp b/Engine/Source/Runtime/Renderer/Private/PostProcess/PostProcessTemporalAA.cpp
--- a/Engine/Source/Runtime/Renderer/Private/PostProcess/PostProcessTemporalAA.cpp
+++ b/Engine/Source/Runtime/Renderer/Private/PostProcess/PostProcessTemporalAA.cpp
@@ -226,8 +226,8 @@ void FRCPassPostProcessSSRTemporalAA::Process(FRenderingCompositePassContext& Co
 	Context.RHICmdList.SetRasterizerState(TStaticRasterizerState<>::GetRHI());
 	Context.RHICmdList.SetDepthStencilState(TStaticDepthStencilState<false,CF_Always>::GetRHI());
 
-	TShaderMapRef< FPostProcessTonemapVS >			VertexShader( GetGlobalShaderMap() );
-	TShaderMapRef< FPostProcessTemporalAAPS<2,0> >	PixelShader( GetGlobalShaderMap() );
+	TShaderMapRef< FPostProcessTonemapVS >			VertexShader(Context.GetShaderMap());
+	TShaderMapRef< FPostProcessTemporalAAPS<2, 0> >	PixelShader(Context.GetShaderMap());
 
 	static FGlobalBoundShaderState BoundShaderState;
 	
@@ -303,8 +303,8 @@ void FRCPassPostProcessDOFTemporalAA::Process(FRenderingCompositePassContext& Co
 	Context.RHICmdList.SetRasterizerState(TStaticRasterizerState<>::GetRHI());
 	Context.RHICmdList.SetDepthStencilState(TStaticDepthStencilState<false,CF_Always>::GetRHI());
 
-	TShaderMapRef< FPostProcessTonemapVS >			VertexShader( GetGlobalShaderMap() );
-	TShaderMapRef< FPostProcessTemporalAAPS<0,0> >	PixelShader( GetGlobalShaderMap() );
+	TShaderMapRef< FPostProcessTonemapVS >			VertexShader(Context.GetShaderMap());
+	TShaderMapRef< FPostProcessTemporalAAPS<0, 0> >	PixelShader(Context.GetShaderMap());
 
 	static FGlobalBoundShaderState BoundShaderState;
 	
@@ -384,8 +384,8 @@ void FRCPassPostProcessLightShaftTemporalAA::Process(FRenderingCompositePassCont
 	Context.RHICmdList.SetRasterizerState(TStaticRasterizerState<>::GetRHI());
 	Context.RHICmdList.SetDepthStencilState(TStaticDepthStencilState<false,CF_Always>::GetRHI());
 
-	TShaderMapRef< FPostProcessTonemapVS >			VertexShader( GetGlobalShaderMap() );
-	TShaderMapRef< FPostProcessTemporalAAPS<3,0> >	PixelShader( GetGlobalShaderMap() );
+	TShaderMapRef< FPostProcessTonemapVS >			VertexShader(Context.GetShaderMap());
+	TShaderMapRef< FPostProcessTemporalAAPS<3, 0> >	PixelShader(Context.GetShaderMap());
 
 	static FGlobalBoundShaderState BoundShaderState;
 	
@@ -473,10 +473,10 @@ void FRCPassPostProcessTemporalAA::Process(FRenderingCompositePassContext& Conte
 		// Normal temporal feedback
 		Context.RHICmdList.SetDepthStencilState(TStaticDepthStencilState<false,CF_Always>::GetRHI());
 
-		TShaderMapRef< FPostProcessTonemapVS >			VertexShader(GetGlobalShaderMap());
+		TShaderMapRef< FPostProcessTonemapVS >			VertexShader(Context.GetShaderMap());
 		if (bUseFast)
 		{
-			TShaderMapRef< FPostProcessTemporalAAPS<4,1> >	PixelShader( GetGlobalShaderMap() );
+			TShaderMapRef< FPostProcessTemporalAAPS<4, 1> >	PixelShader(Context.GetShaderMap());
 
 			static FGlobalBoundShaderState BoundShaderState;
 			
@@ -488,7 +488,7 @@ void FRCPassPostProcessTemporalAA::Process(FRenderingCompositePassContext& Conte
 		}
 		else
 		{
-			TShaderMapRef< FPostProcessTemporalAAPS<1,1> >	PixelShader( GetGlobalShaderMap() );
+			TShaderMapRef< FPostProcessTemporalAAPS<1, 1> >	PixelShader(Context.GetShaderMap());
 
 			static FGlobalBoundShaderState BoundShaderState;
 			
@@ -518,10 +518,10 @@ void FRCPassPostProcessTemporalAA::Process(FRenderingCompositePassContext& Conte
 			// Draw to pixels where stencil == 0
 			Context.RHICmdList.SetDepthStencilState(TStaticDepthStencilState<false,CF_Always,true,CF_Equal,SO_Keep,SO_Keep,SO_Keep>::GetRHI(), 0);
 	
-			TShaderMapRef< FPostProcessTonemapVS >			VertexShader(GetGlobalShaderMap());
+			TShaderMapRef< FPostProcessTonemapVS >			VertexShader(Context.GetShaderMap());
 			if (bUseFast)
 			{
-				TShaderMapRef< FPostProcessTemporalAAPS<4,0> >	PixelShader( GetGlobalShaderMap() );
+				TShaderMapRef< FPostProcessTemporalAAPS<4, 0> >	PixelShader(Context.GetShaderMap());
 	
 				static FGlobalBoundShaderState BoundShaderState;
 				
@@ -533,7 +533,7 @@ void FRCPassPostProcessTemporalAA::Process(FRenderingCompositePassContext& Conte
 			}
 			else
 			{
-				TShaderMapRef< FPostProcessTemporalAAPS<1,0> >	PixelShader( GetGlobalShaderMap() );
+				TShaderMapRef< FPostProcessTemporalAAPS<1, 0> >	PixelShader(Context.GetShaderMap());
 	
 				static FGlobalBoundShaderState BoundShaderState;
 				
@@ -561,10 +561,10 @@ void FRCPassPostProcessTemporalAA::Process(FRenderingCompositePassContext& Conte
 			// Draw to pixels where stencil != 0
 			Context.RHICmdList.SetDepthStencilState(TStaticDepthStencilState<false,CF_Always,true,CF_NotEqual,SO_Keep,SO_Keep,SO_Keep>::GetRHI(), 0);
 			
-			TShaderMapRef< FPostProcessTonemapVS >			VertexShader( GetGlobalShaderMap() );
+			TShaderMapRef< FPostProcessTonemapVS >			VertexShader(Context.GetShaderMap());
 			if(bUseFast)
 			{
-				TShaderMapRef< FPostProcessTemporalAAPS<4,1> >	PixelShader( GetGlobalShaderMap() );
+				TShaderMapRef< FPostProcessTemporalAAPS<4, 1> >	PixelShader(Context.GetShaderMap());
 	
 				static FGlobalBoundShaderState BoundShaderState;
 				
@@ -576,7 +576,7 @@ void FRCPassPostProcessTemporalAA::Process(FRenderingCompositePassContext& Conte
 			}
 			else
 			{
-				TShaderMapRef< FPostProcessTemporalAAPS<1,1> >	PixelShader( GetGlobalShaderMap() );
+				TShaderMapRef< FPostProcessTemporalAAPS<1, 1> >	PixelShader(Context.GetShaderMap());
 	
 				static FGlobalBoundShaderState BoundShaderState;
 				
diff --git a/Engine/Source/Runtime/Renderer/Private/PostProcess/PostProcessTestImage.cpp b/Engine/Source/Runtime/Renderer/Private/PostProcess/PostProcessTestImage.cpp
--- a/Engine/Source/Runtime/Renderer/Private/PostProcess/PostProcessTestImage.cpp
+++ b/Engine/Source/Runtime/Renderer/Private/PostProcess/PostProcessTestImage.cpp
@@ -108,8 +108,8 @@ void FRCPassPostProcessTestImage::Process(FRenderingCompositePassContext& Contex
 	Context.RHICmdList.SetRasterizerState(TStaticRasterizerState<>::GetRHI());
 	Context.RHICmdList.SetDepthStencilState(TStaticDepthStencilState<false, CF_Always>::GetRHI());
 
-	TShaderMapRef<FPostProcessVS> VertexShader(GetGlobalShaderMap());
-	TShaderMapRef<FPostProcessTestImagePS> PixelShader(GetGlobalShaderMap());
+	TShaderMapRef<FPostProcessVS> VertexShader(Context.GetShaderMap());
+	TShaderMapRef<FPostProcessTestImagePS> PixelShader(Context.GetShaderMap());
 
 	static FGlobalBoundShaderState BoundShaderState;
 	
diff --git a/Engine/Source/Runtime/Renderer/Private/PostProcess/PostProcessTonemap.cpp b/Engine/Source/Runtime/Renderer/Private/PostProcess/PostProcessTonemap.cpp
--- a/Engine/Source/Runtime/Renderer/Private/PostProcess/PostProcessTonemap.cpp
+++ b/Engine/Source/Runtime/Renderer/Private/PostProcess/PostProcessTonemap.cpp
@@ -1184,8 +1184,8 @@ FRCPassPostProcessTonemap::FRCPassPostProcessTonemap(bool bInDoGammaOnly)
 template <uint32 ConfigIndex>
 static void SetShaderTempl(const FRenderingCompositePassContext& Context)
 {
-	TShaderMapRef<FPostProcessTonemapVS> VertexShader(GetGlobalShaderMap());
-	TShaderMapRef<FPostProcessTonemapPS<ConfigIndex> > PixelShader(GetGlobalShaderMap());
+	TShaderMapRef<FPostProcessTonemapVS> VertexShader(Context.GetShaderMap());
+	TShaderMapRef<FPostProcessTonemapPS<ConfigIndex> > PixelShader(Context.GetShaderMap());
 
 	static FGlobalBoundShaderState BoundShaderState;
 	
@@ -1266,7 +1266,7 @@ void FRCPassPostProcessTonemap::Process(FRenderingCompositePassContext& Context)
 	}
 
 	// Draw a quad mapping scene color to the view's render target
-	TShaderMapRef<FPostProcessTonemapVS> VertexShader(GetGlobalShaderMap());
+	TShaderMapRef<FPostProcessTonemapVS> VertexShader(Context.GetShaderMap());
 
 	DrawRectangle(
 		Context.RHICmdList,
@@ -1549,8 +1549,8 @@ IMPLEMENT_SHADER_TYPE(,FPostProcessTonemapVS_ES2,TEXT("PostProcessTonemap"),TEXT
 template <uint32 ConfigIndex>
 static void SetShaderTemplES2(const FRenderingCompositePassContext& Context, bool bUsedFramebufferFetch)
 {
-	TShaderMapRef<FPostProcessTonemapVS_ES2> VertexShader(GetGlobalShaderMap());
-	TShaderMapRef<FPostProcessTonemapPS_ES2<ConfigIndex> > PixelShader(GetGlobalShaderMap());
+	TShaderMapRef<FPostProcessTonemapVS_ES2> VertexShader(Context.GetShaderMap());
+	TShaderMapRef<FPostProcessTonemapPS_ES2<ConfigIndex> > PixelShader(Context.GetShaderMap());
 
 	VertexShader->bUsedFramebufferFetch = bUsedFramebufferFetch;
 
@@ -1646,7 +1646,7 @@ void FRCPassPostProcessTonemapES2::Process(FRenderingCompositePassContext& Conte
 	}
 
 	// Draw a quad mapping scene color to the view's render target
-	TShaderMapRef<FPostProcessTonemapVS_ES2> VertexShader(GetGlobalShaderMap());
+	TShaderMapRef<FPostProcessTonemapVS_ES2> VertexShader(Context.GetShaderMap());
 
 	DrawRectangle(
 		Context.RHICmdList,
diff --git a/Engine/Source/Runtime/Renderer/Private/PostProcess/PostProcessUIBlur.cpp b/Engine/Source/Runtime/Renderer/Private/PostProcess/PostProcessUIBlur.cpp
--- a/Engine/Source/Runtime/Renderer/Private/PostProcess/PostProcessUIBlur.cpp
+++ b/Engine/Source/Runtime/Renderer/Private/PostProcess/PostProcessUIBlur.cpp
@@ -86,8 +86,8 @@ void FRCPassPostProcessUIBlur::Process(FRenderingCompositePassContext& Context)
 	Context.RHICmdList.SetRasterizerState(TStaticRasterizerState<>::GetRHI());
 	Context.RHICmdList.SetDepthStencilState(TStaticDepthStencilState<false, CF_Always>::GetRHI());
 
-	TShaderMapRef<FPostProcessVS> VertexShader(GetGlobalShaderMap());
-	TShaderMapRef<FPostProcessUIBlurPS> PixelShader(GetGlobalShaderMap());
+	TShaderMapRef<FPostProcessVS> VertexShader(Context.GetShaderMap());
+	TShaderMapRef<FPostProcessUIBlurPS> PixelShader(Context.GetShaderMap());
 
 	static FGlobalBoundShaderState BoundShaderState;
 	
diff --git a/Engine/Source/Runtime/Renderer/Private/PostProcess/PostProcessUpscale.cpp b/Engine/Source/Runtime/Renderer/Private/PostProcess/PostProcessUpscale.cpp
--- a/Engine/Source/Runtime/Renderer/Private/PostProcess/PostProcessUpscale.cpp
+++ b/Engine/Source/Runtime/Renderer/Private/PostProcess/PostProcessUpscale.cpp
@@ -115,8 +115,8 @@ FRCPassPostProcessUpscale::FRCPassPostProcessUpscale(uint32 InUpscaleMethod)
 template <uint32 Method>
 void FRCPassPostProcessUpscale::SetShader(const FRenderingCompositePassContext& Context)
 {
-	TShaderMapRef<FPostProcessVS> VertexShader(GetGlobalShaderMap());
-	TShaderMapRef<FPostProcessUpscalePS<Method> > PixelShader(GetGlobalShaderMap());
+	TShaderMapRef<FPostProcessVS> VertexShader(Context.GetShaderMap());
+	TShaderMapRef<FPostProcessUpscalePS<Method> > PixelShader(Context.GetShaderMap());
 
 	static FGlobalBoundShaderState BoundShaderState;
 	
@@ -175,7 +175,7 @@ void FRCPassPostProcessUpscale::Process(FRenderingCompositePassContext& Context)
 	}
 
 	// Draw a quad mapping scene color to the view's render target
-	TShaderMapRef<FPostProcessVS> VertexShader(GetGlobalShaderMap());
+	TShaderMapRef<FPostProcessVS> VertexShader(Context.GetShaderMap());
 	DrawRectangle(
 		Context.RHICmdList,
 		0, 0,
diff --git a/Engine/Source/Runtime/Renderer/Private/PostProcess/PostProcessVisualizeBuffer.cpp b/Engine/Source/Runtime/Renderer/Private/PostProcess/PostProcessVisualizeBuffer.cpp
--- a/Engine/Source/Runtime/Renderer/Private/PostProcess/PostProcessVisualizeBuffer.cpp
+++ b/Engine/Source/Runtime/Renderer/Private/PostProcess/PostProcessVisualizeBuffer.cpp
@@ -113,8 +113,8 @@ IMPLEMENT_SHADER_TYPE2(FPostProcessVisualizeBufferPS<false>, SF_Pixel);
 template <bool bDrawingTile>
 FShader* FRCPassPostProcessVisualizeBuffer::SetShaderTempl(const FRenderingCompositePassContext& Context)
 {
-	TShaderMapRef<FPostProcessVS> VertexShader(GetGlobalShaderMap());
-	TShaderMapRef<FPostProcessVisualizeBufferPS<bDrawingTile> > PixelShader(GetGlobalShaderMap());
+	TShaderMapRef<FPostProcessVS> VertexShader(Context.GetShaderMap());
+	TShaderMapRef<FPostProcessVisualizeBufferPS<bDrawingTile> > PixelShader(Context.GetShaderMap());
 
 	static FGlobalBoundShaderState BoundShaderState;
 	
@@ -173,8 +173,8 @@ void FRCPassPostProcessVisualizeBuffer::Process(FRenderingCompositePassContext&
 
 	Context.RHICmdList.SetBlendState(TStaticBlendState<CW_RGB, BO_Add, BF_SourceAlpha, BF_InverseSourceAlpha>::GetRHI());
 
-	TShaderMapRef<FPostProcessVS> VertexShader(GetGlobalShaderMap());
-	TShaderMapRef<FPostProcessVisualizeBufferPS<true> > PixelShader(GetGlobalShaderMap());
+	TShaderMapRef<FPostProcessVS> VertexShader(Context.GetShaderMap());
+	TShaderMapRef<FPostProcessVisualizeBufferPS<true> > PixelShader(Context.GetShaderMap());
 
 	static FGlobalBoundShaderState BoundShaderState;
 	
diff --git a/Engine/Source/Runtime/Renderer/Private/PostProcess/PostProcessVisualizeHDR.cpp b/Engine/Source/Runtime/Renderer/Private/PostProcess/PostProcessVisualizeHDR.cpp
--- a/Engine/Source/Runtime/Renderer/Private/PostProcess/PostProcessVisualizeHDR.cpp
+++ b/Engine/Source/Runtime/Renderer/Private/PostProcess/PostProcessVisualizeHDR.cpp
@@ -163,8 +163,8 @@ void FRCPassPostProcessVisualizeHDR::Process(FRenderingCompositePassContext& Con
 	Context.RHICmdList.SetRasterizerState(TStaticRasterizerState<>::GetRHI());
 	Context.RHICmdList.SetDepthStencilState(TStaticDepthStencilState<false, CF_Always>::GetRHI());
 
-	TShaderMapRef<FPostProcessVS> VertexShader(GetGlobalShaderMap());
-	TShaderMapRef<FPostProcessVisualizeHDRPS> PixelShader(GetGlobalShaderMap());
+	TShaderMapRef<FPostProcessVS> VertexShader(Context.GetShaderMap());
+	TShaderMapRef<FPostProcessVisualizeHDRPS> PixelShader(Context.GetShaderMap());
 
 	static FGlobalBoundShaderState BoundShaderState;
 	
diff --git a/Engine/Source/Runtime/Renderer/Private/PostProcess/PostProcessWeightedSampleSum.cpp b/Engine/Source/Runtime/Renderer/Private/PostProcess/PostProcessWeightedSampleSum.cpp
--- a/Engine/Source/Runtime/Renderer/Private/PostProcess/PostProcessWeightedSampleSum.cpp
+++ b/Engine/Source/Runtime/Renderer/Private/PostProcess/PostProcessWeightedSampleSum.cpp
@@ -172,17 +172,17 @@ void SetFilterShaders(
 	)
 {
 	check(CombineMethodInt <= 2);
-	
+	auto ShaderMap = GetGlobalShaderMap(FeatureLevel);
 
 	// A macro to handle setting the filter shader for a specific number of samples.
 #define SET_FILTER_SHADER_TYPE(NumSamples) \
 	case NumSamples: \
 	{ \
-		TShaderMapRef<TFilterVS<NumSamples> > VertexShader(GetGlobalShaderMap()); \
+		TShaderMapRef<TFilterVS<NumSamples> > VertexShader(ShaderMap); \
 		*OutVertexShader = *VertexShader; \
 		if(CombineMethodInt == 0) \
 		{ \
-			TShaderMapRef<TFilterPS<NumSamples, 0> > PixelShader(GetGlobalShaderMap()); \
+			TShaderMapRef<TFilterPS<NumSamples, 0> > PixelShader(ShaderMap); \
 			{ \
 				static FGlobalBoundShaderState BoundShaderState; \
 				SetGlobalBoundShaderState(RHICmdList, FeatureLevel, BoundShaderState, GFilterVertexDeclaration.VertexDeclarationRHI, *VertexShader, *PixelShader); \
@@ -191,7 +191,7 @@ void SetFilterShaders(
 		} \
 		else if(CombineMethodInt == 1) \
 		{ \
-			TShaderMapRef<TFilterPS<NumSamples, 1> > PixelShader(GetGlobalShaderMap()); \
+			TShaderMapRef<TFilterPS<NumSamples, 1> > PixelShader(ShaderMap); \
 			{ \
 				static FGlobalBoundShaderState BoundShaderState; \
 				SetGlobalBoundShaderState(RHICmdList, FeatureLevel, BoundShaderState, GFilterVertexDeclaration.VertexDeclarationRHI, *VertexShader, *PixelShader); \
@@ -200,7 +200,7 @@ void SetFilterShaders(
 		} \
 		else\
 		{ \
-			TShaderMapRef<TFilterPS<NumSamples, 2> > PixelShader(GetGlobalShaderMap()); \
+			TShaderMapRef<TFilterPS<NumSamples, 2> > PixelShader(ShaderMap); \
 			{ \
 				static FGlobalBoundShaderState BoundShaderState; \
 				SetGlobalBoundShaderState(RHICmdList, FeatureLevel, BoundShaderState, GFilterVertexDeclaration.VertexDeclarationRHI, *VertexShader, *PixelShader); \
diff --git a/Engine/Source/Runtime/Renderer/Private/PostProcess/RenderTargetPool.cpp b/Engine/Source/Runtime/Renderer/Private/PostProcess/RenderTargetPool.cpp
--- a/Engine/Source/Runtime/Renderer/Private/PostProcess/RenderTargetPool.cpp
+++ b/Engine/Source/Runtime/Renderer/Private/PostProcess/RenderTargetPool.cpp
@@ -687,7 +687,7 @@ inline void DrawBorder(FCanvas& Canvas, const FIntRect Rect, FLinearColor Color)
 	Canvas.DrawTile(Rect.Max.X - 1, Rect.Min.Y + 1, 1, Rect.Max.Y - Rect.Min.Y - 2, 0, 0, 1, 1, Color);
 }
 
-void FRenderTargetPool::PresentContent(FRHICommandListImmediate& RHICmdList, const FSceneView& View)
+void FRenderTargetPool::PresentContent(FRHICommandListImmediate& RHICmdList, const FViewInfo& View)
 {
 	if (RenderTargetPoolEvents.Num())
 	{
diff --git a/Engine/Source/Runtime/Renderer/Private/PostProcess/RenderTargetPool.h b/Engine/Source/Runtime/Renderer/Private/PostProcess/RenderTargetPool.h
--- a/Engine/Source/Runtime/Renderer/Private/PostProcess/RenderTargetPool.h
+++ b/Engine/Source/Runtime/Renderer/Private/PostProcess/RenderTargetPool.h
@@ -221,7 +221,7 @@ public:
 	void AddPhaseEvent(const TCHAR* InPhaseName);
 
 	/** renders the VisualizeTextureContent to the current render target */
-	void PresentContent(FRHICommandListImmediate& RHICmdList, const FSceneView& View);
+	void PresentContent(FRHICommandListImmediate& RHICmdList, const FViewInfo& View);
 
 	FVisualizeTexture VisualizeTexture;
 
diff --git a/Engine/Source/Runtime/Renderer/Private/PostProcess/RenderingCompositionGraph.cpp b/Engine/Source/Runtime/Renderer/Private/PostProcess/RenderingCompositionGraph.cpp
--- a/Engine/Source/Runtime/Renderer/Private/PostProcess/RenderingCompositionGraph.cpp
+++ b/Engine/Source/Runtime/Renderer/Private/PostProcess/RenderingCompositionGraph.cpp
@@ -110,6 +110,8 @@ FRenderingCompositePassContext::FRenderingCompositePassContext(FRHICommandListIm
 	, Pass(0)
 	, RHICmdList(InRHICmdList)
 	, ViewPortRect(0, 0, 0 ,0)
+	, FeatureLevel(View.GetFeatureLevel())
+	, ShaderMap(InView.ShaderMap)
 {
 	check(!IsViewportValid());
 
@@ -123,11 +125,6 @@ FRenderingCompositePassContext::~FRenderingCompositePassContext()
 	Graph.Free();
 }
 
-ERHIFeatureLevel::Type FRenderingCompositePassContext::GetFeatureLevel() const
-{
-	return View.GetFeatureLevel();
-}
-
 void FRenderingCompositePassContext::Process(const TCHAR *GraphDebugName)
 {
 	if(Root)
diff --git a/Engine/Source/Runtime/Renderer/Private/PostProcess/RenderingCompositionGraph.h b/Engine/Source/Runtime/Renderer/Private/PostProcess/RenderingCompositionGraph.h
--- a/Engine/Source/Runtime/Renderer/Private/PostProcess/RenderingCompositionGraph.h
+++ b/Engine/Source/Runtime/Renderer/Private/PostProcess/RenderingCompositionGraph.h
@@ -131,7 +131,8 @@ struct FRenderingCompositePassContext
 		return ViewPortRect.Min != ViewPortRect.Max;
 	}
 
-	ERHIFeatureLevel::Type GetFeatureLevel() const;
+	ERHIFeatureLevel::Type GetFeatureLevel() const { return FeatureLevel; }
+	TShaderMap<FGlobalShaderType>* GetShaderMap() const { check(ShaderMap); return ShaderMap; }
 
 	FRenderingCompositePass* Root;
 
@@ -142,6 +143,9 @@ struct FRenderingCompositePassContext
 private:
 	// cached state to map between ScreenPos and pixels
 	FIntRect ViewPortRect;
+
+	ERHIFeatureLevel::Type FeatureLevel;
+	TShaderMap<FGlobalShaderType>* ShaderMap;
 };
 
 // ---------------------------------------------------------------------------
diff --git a/Engine/Source/Runtime/Renderer/Private/PostProcess/SceneRenderTargets.cpp b/Engine/Source/Runtime/Renderer/Private/PostProcess/SceneRenderTargets.cpp
--- a/Engine/Source/Runtime/Renderer/Private/PostProcess/SceneRenderTargets.cpp
+++ b/Engine/Source/Runtime/Renderer/Private/PostProcess/SceneRenderTargets.cpp
@@ -667,11 +667,12 @@ void FSceneRenderTargets::ResolveSceneColor(FRHICommandList& RHICmdList, const F
 		RHICmdList.SetRasterizerState(TStaticRasterizerState<>::GetRHI());
 		RHICmdList.SetDepthStencilState(TStaticDepthStencilState<false, CF_Always>::GetRHI());
 
-		TShaderMapRef<FHdrCustomResolveVS> VertexShader(GetGlobalShaderMap());
+		auto ShaderMap = GetGlobalShaderMap(CurrentFeatureLevel);
+		TShaderMapRef<FHdrCustomResolveVS> VertexShader(ShaderMap);
 
 		if(samples == 2)
 		{
-			TShaderMapRef<FHdrCustomResolve2xPS> PixelShader(GetGlobalShaderMap());
+			TShaderMapRef<FHdrCustomResolve2xPS> PixelShader(ShaderMap);
 			static FGlobalBoundShaderState BoundShaderState;
 			SetGlobalBoundShaderState(RHICmdList, CurrentFeatureLevel, BoundShaderState, GetVertexDeclarationFVector4(), *VertexShader, *PixelShader);
 			PixelShader->SetParameters(RHICmdList, GSceneRenderTargets.SceneColor->GetRenderTargetItem().TargetableTexture);
@@ -679,7 +680,7 @@ void FSceneRenderTargets::ResolveSceneColor(FRHICommandList& RHICmdList, const F
 		}
 		else if(samples == 4)
 		{
-			TShaderMapRef<FHdrCustomResolve4xPS> PixelShader(GetGlobalShaderMap());
+			TShaderMapRef<FHdrCustomResolve4xPS> PixelShader(ShaderMap);
 			static FGlobalBoundShaderState BoundShaderState;
 			SetGlobalBoundShaderState(RHICmdList, CurrentFeatureLevel, BoundShaderState, GetVertexDeclarationFVector4(), *VertexShader, *PixelShader);
 			PixelShader->SetParameters(RHICmdList, GSceneRenderTargets.SceneColor->GetRenderTargetItem().TargetableTexture);
@@ -687,7 +688,7 @@ void FSceneRenderTargets::ResolveSceneColor(FRHICommandList& RHICmdList, const F
 		}
 		else if(samples == 8)
 		{
-			TShaderMapRef<FHdrCustomResolve8xPS> PixelShader(GetGlobalShaderMap());
+			TShaderMapRef<FHdrCustomResolve8xPS> PixelShader(ShaderMap);
 			static FGlobalBoundShaderState BoundShaderState;
 			SetGlobalBoundShaderState(RHICmdList, CurrentFeatureLevel, BoundShaderState, GetVertexDeclarationFVector4(), *VertexShader, *PixelShader);
 			PixelShader->SetParameters(RHICmdList, GSceneRenderTargets.SceneColor->GetRenderTargetItem().TargetableTexture);
diff --git a/Engine/Source/Runtime/Renderer/Private/PostProcess/ScreenSpaceReflections.cpp b/Engine/Source/Runtime/Renderer/Private/PostProcess/ScreenSpaceReflections.cpp
--- a/Engine/Source/Runtime/Renderer/Private/PostProcess/ScreenSpaceReflections.cpp
+++ b/Engine/Source/Runtime/Renderer/Private/PostProcess/ScreenSpaceReflections.cpp
@@ -207,12 +207,12 @@ void FRCPassPostProcessScreenSpaceReflections::Process(FRenderingCompositePassCo
 		SSRQuality = 0;
 	}
 
-	TShaderMapRef< FPostProcessVS > VertexShader(GetGlobalShaderMap());
+	TShaderMapRef< FPostProcessVS > VertexShader(Context.GetShaderMap());
 
 	#define CASE(A, B) \
 		case (A + 2 * (B + 3 * 0 )): \
 		{ \
-			TShaderMapRef< FPostProcessScreenSpaceReflectionsPS<A, B> > PixelShader(GetGlobalShaderMap()); \
+			TShaderMapRef< FPostProcessScreenSpaceReflectionsPS<A, B> > PixelShader(Context.GetShaderMap()); \
 			static FGlobalBoundShaderState BoundShaderState; \
 			SetGlobalBoundShaderState(Context.RHICmdList, FeatureLevel, BoundShaderState, GFilterVertexDeclaration.VertexDeclarationRHI, *VertexShader, *PixelShader); \
 			VertexShader->SetParameters(Context); \
diff --git a/Engine/Source/Runtime/Renderer/Private/PostProcess/VisualizeTexture.cpp b/Engine/Source/Runtime/Renderer/Private/PostProcess/VisualizeTexture.cpp
--- a/Engine/Source/Runtime/Renderer/Private/PostProcess/VisualizeTexture.cpp
+++ b/Engine/Source/Runtime/Renderer/Private/PostProcess/VisualizeTexture.cpp
@@ -194,8 +194,9 @@ IMPLEMENT_SHADER_TYPE(,FVisualizeTexturePresentPS,TEXT("VisualizeTexture"),TEXT(
 
 template<uint32 TextureType> void VisualizeTextureForTextureType(FRHICommandList& RHICmdList, ERHIFeatureLevel::Type FeatureLevel, const FVisualizeTextureData& Data)
 {
-	TShaderMapRef<FScreenVS> VertexShader(GetGlobalShaderMap());
-	TShaderMapRef<VisualizeTexturePS<TextureType> > PixelShader(GetGlobalShaderMap());
+	auto ShaderMap = GetGlobalShaderMap(FeatureLevel);
+	TShaderMapRef<FScreenVS> VertexShader(ShaderMap);
+	TShaderMapRef<VisualizeTexturePS<TextureType> > PixelShader(ShaderMap);
 
 	static FGlobalBoundShaderState BoundShaderState;
 	
@@ -475,7 +476,7 @@ void FVisualizeTexture::GenerateContent(FRHICommandListImmediate& RHICmdList, co
 	}
 }
 
-void FVisualizeTexture::PresentContent(FRHICommandListImmediate& RHICmdList, const FSceneView& View)
+void FVisualizeTexture::PresentContent(FRHICommandListImmediate& RHICmdList, const FViewInfo& View)
 {
 	if(Mode != 0)
 	{
@@ -507,8 +508,9 @@ void FVisualizeTexture::PresentContent(FRHICommandListImmediate& RHICmdList, con
 	RHICmdList.SetRasterizerState(TStaticRasterizerState<>::GetRHI());
 	RHICmdList.SetDepthStencilState(TStaticDepthStencilState<false, CF_Always>::GetRHI());
 
-	TShaderMapRef<FPostProcessVS> VertexShader(GetGlobalShaderMap());
-	TShaderMapRef<FVisualizeTexturePresentPS> PixelShader(GetGlobalShaderMap());
+	auto ShaderMap = View.ShaderMap;
+	TShaderMapRef<FPostProcessVS> VertexShader(ShaderMap);
+	TShaderMapRef<FVisualizeTexturePresentPS> PixelShader(ShaderMap);
 
 	static FGlobalBoundShaderState BoundShaderState;
 	
diff --git a/Engine/Source/Runtime/Renderer/Private/PostProcess/VisualizeTexture.h b/Engine/Source/Runtime/Renderer/Private/PostProcess/VisualizeTexture.h
--- a/Engine/Source/Runtime/Renderer/Private/PostProcess/VisualizeTexture.h
+++ b/Engine/Source/Runtime/Renderer/Private/PostProcess/VisualizeTexture.h
@@ -7,6 +7,7 @@
 #pragma once
 
 struct IPooledRenderTarget;
+class FViewInfo;
 
 class FVisualizeTexture
 {
@@ -14,7 +15,7 @@ public:
 	FVisualizeTexture();
 
 	/** renders the VisualizeTextureContent to the current render target */
-	void PresentContent(FRHICommandListImmediate& RHICmdList, const FSceneView& View);
+	void PresentContent(FRHICommandListImmediate& RHICmdList, const FViewInfo& View);
 	/** */
 	void OnStartFrame(const FSceneView& View);
 	/** */
diff --git a/Engine/Source/Runtime/Renderer/Private/ReflectionEnvironment.cpp b/Engine/Source/Runtime/Renderer/Private/ReflectionEnvironment.cpp
--- a/Engine/Source/Runtime/Renderer/Private/ReflectionEnvironment.cpp
+++ b/Engine/Source/Runtime/Renderer/Private/ReflectionEnvironment.cpp
@@ -559,8 +559,9 @@ void FDeferredShadingSceneRenderer::RenderReflectionCaptureSpecularBounceForAllV
 	RHICmdList.SetDepthStencilState(TStaticDepthStencilState< false, CF_Always >::GetRHI());
 	RHICmdList.SetBlendState(TStaticBlendState< CW_RGB, BO_Add, BF_One, BF_One >::GetRHI());
 
-	TShaderMapRef< FPostProcessVS > VertexShader( GetGlobalShaderMap() );
-	TShaderMapRef< FReflectionCaptureSpecularBouncePS > PixelShader( GetGlobalShaderMap() );
+	auto ShaderMap = GetGlobalShaderMap(FeatureLevel);
+	TShaderMapRef< FPostProcessVS > VertexShader(ShaderMap);
+	TShaderMapRef< FReflectionCaptureSpecularBouncePS > PixelShader(ShaderMap);
 
 	static FGlobalBoundShaderState BoundShaderState;
 	
@@ -635,11 +636,11 @@ void FDeferredShadingSceneRenderer::RenderImageBasedReflectionsSM5ForAllViews(FR
 			FReflectionEnvironmentTiledDeferredCS* ComputeShader = NULL;
 			if( bUseLightmaps )
 			{
-				ComputeShader = *TShaderMapRef< TReflectionEnvironmentTiledDeferredCS<1> >( GetGlobalShaderMap() );
+				ComputeShader = *TShaderMapRef< TReflectionEnvironmentTiledDeferredCS<1> >(View.ShaderMap);
 			}
 			else
 			{
-				ComputeShader = *TShaderMapRef< TReflectionEnvironmentTiledDeferredCS<0> >( GetGlobalShaderMap() );
+				ComputeShader = *TShaderMapRef< TReflectionEnvironmentTiledDeferredCS<0> >(View.ShaderMap);
 			}
 
 			RHICmdList.SetComputeShader(ComputeShader->GetComputeShader());
@@ -760,12 +761,12 @@ void FDeferredShadingSceneRenderer::RenderImageBasedReflectionsSM4ForAllViews(FR
 				{
 					const FSphere LightBounds(ReflectionCapture.PositionAndRadius, ReflectionCapture.PositionAndRadius.W);
 
-					TShaderMapRef<TDeferredLightVS<true> > VertexShader(GetGlobalShaderMap());
+					TShaderMapRef<TDeferredLightVS<true> > VertexShader(View.ShaderMap);
 
 					// Use the appropriate shader for the capture shape
 					if (ReflectionCapture.CaptureProperties.Z == 0)
 					{
-						TShaderMapRef<TStandardDeferredReflectionPS<true> > PixelShader(GetGlobalShaderMap());
+						TShaderMapRef<TStandardDeferredReflectionPS<true> > PixelShader(View.ShaderMap);
 
 						static FGlobalBoundShaderState BoundShaderState;
 						
@@ -775,7 +776,7 @@ void FDeferredShadingSceneRenderer::RenderImageBasedReflectionsSM4ForAllViews(FR
 					}
 					else
 					{
-						TShaderMapRef<TStandardDeferredReflectionPS<false> > PixelShader(GetGlobalShaderMap());
+						TShaderMapRef<TStandardDeferredReflectionPS<false> > PixelShader(View.ShaderMap);
 
 						static FGlobalBoundShaderState BoundShaderState;
 						
@@ -815,7 +816,7 @@ void FDeferredShadingSceneRenderer::RenderImageBasedReflectionsSM4ForAllViews(FR
 				RHICmdList.SetBlendState(TStaticBlendState<CW_RGBA, BO_Add, BF_One, BF_One, BO_Add, BF_One, BF_One>::GetRHI());
 			}
 
-			TShaderMapRef< FPostProcessVS >		VertexShader(GetGlobalShaderMap());
+			TShaderMapRef< FPostProcessVS >		VertexShader(View.ShaderMap);
 
 			if (!LightAccumulation)
 			{
@@ -826,7 +827,7 @@ void FDeferredShadingSceneRenderer::RenderImageBasedReflectionsSM4ForAllViews(FR
 #define CASE(A,B,C) \
 			case ((A << 2) | (B << 1) | C) : \
 			{ \
-			TShaderMapRef< FReflectionApplyPS<A, B, C> > PixelShader(GetGlobalShaderMap()); \
+			TShaderMapRef< FReflectionApplyPS<A, B, C> > PixelShader(View.ShaderMap); \
 			static FGlobalBoundShaderState BoundShaderState; \
 			SetGlobalBoundShaderState(RHICmdList, FeatureLevel, BoundShaderState, GFilterVertexDeclaration.VertexDeclarationRHI, *VertexShader, *PixelShader); \
 			PixelShader->SetParameters(RHICmdList, View, LightAccumulation->GetRenderTargetItem().ShaderResourceTexture, SSROutput->GetRenderTargetItem().ShaderResourceTexture); \
diff --git a/Engine/Source/Runtime/Renderer/Private/ReflectionEnvironmentCapture.cpp b/Engine/Source/Runtime/Renderer/Private/ReflectionEnvironmentCapture.cpp
--- a/Engine/Source/Runtime/Renderer/Private/ReflectionEnvironmentCapture.cpp
+++ b/Engine/Source/Runtime/Renderer/Private/ReflectionEnvironmentCapture.cpp
@@ -282,8 +282,9 @@ void ComputeAverageBrightness(FRHICommandList& RHICmdList, ERHIFeatureLevel::Typ
 	RHICmdList.SetDepthStencilState(TStaticDepthStencilState<false, CF_Always>::GetRHI());
 	RHICmdList.SetBlendState(TStaticBlendState<>::GetRHI());
 
-	TShaderMapRef<FPostProcessVS> VertexShader(GetGlobalShaderMap());
-	TShaderMapRef<FComputeBrightnessPS> PixelShader(GetGlobalShaderMap());
+	auto ShaderMap = GetGlobalShaderMap(FeatureLevel);
+	TShaderMapRef<FPostProcessVS> VertexShader(ShaderMap);
+	TShaderMapRef<FComputeBrightnessPS> PixelShader(ShaderMap);
 
 	static FGlobalBoundShaderState BoundShaderState;
 	
@@ -314,6 +315,8 @@ void FilterReflectionEnvironment(FRHICommandListImmediate& RHICmdList, ERHIFeatu
 	static const auto CVar = IConsoleManager::Get().FindConsoleVariable(TEXT("r.DiffuseFromCaptures"));
 	bNormalize = bNormalize && CVar->GetInt() == 0;
 
+	auto ShaderMap = GetGlobalShaderMap(FeatureLevel);
+
 	{	
 		SCOPED_DRAW_EVENT(DownsampleCubeMips, DEC_SCENE_ITEMS);
 		// Downsample all the mips, each one reads from the mip above it
@@ -337,9 +340,9 @@ void FilterReflectionEnvironment(FRHICommandListImmediate& RHICmdList, ERHIFeatu
 				RHICmdList.SetDepthStencilState(TStaticDepthStencilState<false, CF_Always>::GetRHI());
 				RHICmdList.SetBlendState(TStaticBlendState<>::GetRHI());
 
-				TShaderMapRef<FScreenVSForGS> VertexShader(GetGlobalShaderMap());
-				TShaderMapRef<FDownsampleGS> GeometryShader(GetGlobalShaderMap());
-				TShaderMapRef<FDownsamplePS> PixelShader(GetGlobalShaderMap());
+				TShaderMapRef<FScreenVSForGS> VertexShader(ShaderMap);
+				TShaderMapRef<FDownsampleGS> GeometryShader(ShaderMap);
+				TShaderMapRef<FDownsamplePS> PixelShader(ShaderMap);
 
 				SetGlobalBoundShaderState(RHICmdList, FeatureLevel, DownsampleBoundShaderState, GFilterVertexDeclaration.VertexDeclarationRHI, *VertexShader, *PixelShader, *GeometryShader);
 
@@ -377,8 +380,8 @@ void FilterReflectionEnvironment(FRHICommandListImmediate& RHICmdList, ERHIFeatu
 					RHICmdList.SetDepthStencilState(TStaticDepthStencilState<false, CF_Always>::GetRHI());
 					RHICmdList.SetBlendState(TStaticBlendState<>::GetRHI());
 
-					TShaderMapRef<FScreenVSForGS> VertexShader(GetGlobalShaderMap());
-					TShaderMapRef<FDownsamplePS> PixelShader(GetGlobalShaderMap());
+					TShaderMapRef<FScreenVSForGS> VertexShader(ShaderMap);
+					TShaderMapRef<FDownsamplePS> PixelShader(ShaderMap);
 
 					SetGlobalBoundShaderState(RHICmdList, FeatureLevel, DownsampleBoundShaderState, GFilterVertexDeclaration.VertexDeclarationRHI, *VertexShader, *PixelShader);
 
@@ -436,20 +439,20 @@ void FilterReflectionEnvironment(FRHICommandListImmediate& RHICmdList, ERHIFeatu
 				RHICmdList.SetDepthStencilState(TStaticDepthStencilState<false, CF_Always>::GetRHI());
 				RHICmdList.SetBlendState(TStaticBlendState<>::GetRHI());
 
-				TShaderMapRef<FScreenVSForGS> VertexShader(GetGlobalShaderMap());
-				TShaderMapRef<FDownsampleGS> GeometryShader(GetGlobalShaderMap());
+				TShaderMapRef<FScreenVSForGS> VertexShader(ShaderMap);
+				TShaderMapRef<FDownsampleGS> GeometryShader(ShaderMap);
 
 				FCubeFilterPS* PixelShader;
 				if( bNormalize )
 				{
-					PixelShader = *TShaderMapRef< TCubeFilterPS<1> >( GetGlobalShaderMap() );
+					PixelShader = *TShaderMapRef< TCubeFilterPS<1> >(ShaderMap);
 					static FGlobalBoundShaderState BoundShaderState;
 					
 					SetGlobalBoundShaderState(RHICmdList, FeatureLevel, BoundShaderState, GFilterVertexDeclaration.VertexDeclarationRHI, *VertexShader, PixelShader, *GeometryShader);
 				}
 				else
 				{
-					PixelShader = *TShaderMapRef< TCubeFilterPS<0> >( GetGlobalShaderMap() );
+					PixelShader = *TShaderMapRef< TCubeFilterPS<0> >(ShaderMap);
 					static FGlobalBoundShaderState BoundShaderState;
 					
 					SetGlobalBoundShaderState(RHICmdList, FeatureLevel, BoundShaderState, GFilterVertexDeclaration.VertexDeclarationRHI, *VertexShader, PixelShader, *GeometryShader);
@@ -488,21 +491,21 @@ void FilterReflectionEnvironment(FRHICommandListImmediate& RHICmdList, ERHIFeatu
 					RHICmdList.SetDepthStencilState(TStaticDepthStencilState<false, CF_Always>::GetRHI());
 					RHICmdList.SetBlendState(TStaticBlendState<>::GetRHI());
 
-					TShaderMapRef<FScreenVSForGS> VertexShader(GetGlobalShaderMap());
-					TShaderMapRef< TCubeFilterPS<1> > CaptureCubemapArrayPixelShader(GetGlobalShaderMap());
+					TShaderMapRef<FScreenVSForGS> VertexShader(ShaderMap);
+					TShaderMapRef< TCubeFilterPS<1> > CaptureCubemapArrayPixelShader(ShaderMap);
 					//FFilterPS* PixelShader = (FFilterPS*)*CaptureCubemapArrayPixelShader;
 
 					FCubeFilterPS* PixelShader;
 					if( bNormalize )
 					{
-						PixelShader = *TShaderMapRef< TCubeFilterPS<1> >( GetGlobalShaderMap() );
+						PixelShader = *TShaderMapRef< TCubeFilterPS<1> >(ShaderMap);
 						static FGlobalBoundShaderState BoundShaderState;
 						
 						SetGlobalBoundShaderState(RHICmdList, FeatureLevel, BoundShaderState, GFilterVertexDeclaration.VertexDeclarationRHI, *VertexShader, PixelShader);
 					}
 					else
 					{
-						PixelShader = *TShaderMapRef< TCubeFilterPS<0> >( GetGlobalShaderMap() );
+						PixelShader = *TShaderMapRef< TCubeFilterPS<0> >(ShaderMap);
 						static FGlobalBoundShaderState BoundShaderState;
 						
 						SetGlobalBoundShaderState(RHICmdList, FeatureLevel, BoundShaderState, GFilterVertexDeclaration.VertexDeclarationRHI, *VertexShader, PixelShader);
@@ -843,6 +846,8 @@ void CaptureSceneToScratchCubemap(FRHICommandListImmediate& RHICmdList, FSceneRe
 		TEMP_PostReflectionCaptureRender();
 #endif
 
+		auto ShaderMap = GetGlobalShaderMap(FeatureLevel);
+
 		const int32 EffectiveSize = GReflectionCaptureSize;
 		FSceneRenderTargetItem& EffectiveColorRT =  GSceneRenderTargets.ReflectionColorScratchCubemap[0]->GetRenderTargetItem();
 
@@ -859,29 +864,29 @@ void CaptureSceneToScratchCubemap(FRHICommandListImmediate& RHICmdList, FSceneRe
 				RHICmdList.SetDepthStencilState(TStaticDepthStencilState<false, CF_Always>::GetRHI());
 				RHICmdList.SetBlendState(TStaticBlendState<>::GetRHI());
 
-			TShaderMapRef<FCopyToCubeFaceVS> VertexShader(GetGlobalShaderMap());
-			TShaderMapRef<FRouteToCubeFaceGS> GeometryShader(GetGlobalShaderMap());
+				TShaderMapRef<FCopyToCubeFaceVS> VertexShader(ShaderMap);
+				TShaderMapRef<FRouteToCubeFaceGS> GeometryShader(ShaderMap);
 
-			TShaderMapRef<FCopySceneColorToCubeFacePS> PixelShader(GetGlobalShaderMap());
-			SetGlobalBoundShaderState(RHICmdList, FeatureLevel, CopyColorCubemapBoundShaderState, GFilterVertexDeclaration.VertexDeclarationRHI, *VertexShader, *PixelShader, *GeometryShader);
-			
-			PixelShader->SetParameters(RHICmdList, SceneRenderer->Views[0], bCapturingForSkyLight, bLowerHemisphereIsBlack);
-			VertexShader->SetParameters(RHICmdList, SceneRenderer->Views[0]);
-			GeometryShader->SetParameters(RHICmdList, (uint32)CubeFace);
-			
-			DrawRectangle( 
-				RHICmdList,
-				ViewRect.Min.X, ViewRect.Min.Y, 
-				ViewRect.Width(), ViewRect.Height(),
-				ViewRect.Min.X, ViewRect.Min.Y, 
-				ViewRect.Width() * GSupersampleCaptureFactor, ViewRect.Height() * GSupersampleCaptureFactor,
-				FIntPoint(ViewRect.Width(), ViewRect.Height()),
-				GSceneRenderTargets.GetBufferSizeXY(),
-				*VertexShader,
-				EDRF_UseTriangleOptimization);
+				TShaderMapRef<FCopySceneColorToCubeFacePS> PixelShader(ShaderMap);
+				SetGlobalBoundShaderState(RHICmdList, FeatureLevel, CopyColorCubemapBoundShaderState, GFilterVertexDeclaration.VertexDeclarationRHI, *VertexShader, *PixelShader, *GeometryShader);
+
+				PixelShader->SetParameters(RHICmdList, SceneRenderer->Views[0], bCapturingForSkyLight, bLowerHemisphereIsBlack);
+				VertexShader->SetParameters(RHICmdList, SceneRenderer->Views[0]);
+				GeometryShader->SetParameters(RHICmdList, (uint32)CubeFace);
+
+				DrawRectangle(
+					RHICmdList,
+					ViewRect.Min.X, ViewRect.Min.Y,
+					ViewRect.Width(), ViewRect.Height(),
+					ViewRect.Min.X, ViewRect.Min.Y,
+					ViewRect.Width() * GSupersampleCaptureFactor, ViewRect.Height() * GSupersampleCaptureFactor,
+					FIntPoint(ViewRect.Width(), ViewRect.Height()),
+					GSceneRenderTargets.GetBufferSizeXY(),
+					*VertexShader,
+					EDRF_UseTriangleOptimization);
 
 				RHICmdList.CopyToResolveTarget(EffectiveColorRT.TargetableTexture, EffectiveColorRT.ShaderResourceTexture, true, FResolveParams(FResolveRect()));
-		}
+			}
 			else
 			{
 				// Copy the captured scene into the cubemap face
@@ -893,8 +898,8 @@ void CaptureSceneToScratchCubemap(FRHICommandListImmediate& RHICmdList, FSceneRe
 				RHICmdList.SetDepthStencilState(TStaticDepthStencilState<false, CF_Always>::GetRHI());
 				RHICmdList.SetBlendState(TStaticBlendState<>::GetRHI());
 
-				TShaderMapRef<FCopyToCubeFaceVS> VertexShader(GetGlobalShaderMap());
-				TShaderMapRef<FCopySceneColorToCubeFacePS> PixelShader(GetGlobalShaderMap());
+				TShaderMapRef<FCopyToCubeFaceVS> VertexShader(ShaderMap);
+				TShaderMapRef<FCopySceneColorToCubeFacePS> PixelShader(ShaderMap);
 				SetGlobalBoundShaderState(RHICmdList, FeatureLevel, CopyColorCubemapBoundShaderState, GFilterVertexDeclaration.VertexDeclarationRHI, *VertexShader, *PixelShader);
 
 				PixelShader->SetParameters(RHICmdList, SceneRenderer->Views[0], bCapturingForSkyLight, bLowerHemisphereIsBlack);
@@ -927,6 +932,8 @@ void CopyCubemapToScratchCubemap(FRHICommandList& RHICmdList, ERHIFeatureLevel::
 	const int32 EffectiveSize = GReflectionCaptureSize;
 	FSceneRenderTargetItem& EffectiveColorRT =  GSceneRenderTargets.ReflectionColorScratchCubemap[0]->GetRenderTargetItem();
 
+	auto ShaderMap = GetGlobalShaderMap(FeatureLevel);
+
 	if (GSupportsGSRenderTargetLayerSwitchingToMips)
 	{
 		// Copy the captured scene into the cubemap face
@@ -940,9 +947,9 @@ void CopyCubemapToScratchCubemap(FRHICommandList& RHICmdList, ERHIFeatureLevel::
 		RHICmdList.SetDepthStencilState(TStaticDepthStencilState<false, CF_Always>::GetRHI());
 		RHICmdList.SetBlendState(TStaticBlendState<>::GetRHI());
 
-		TShaderMapRef<FScreenVSForGS> VertexShader(GetGlobalShaderMap());
-		TShaderMapRef<FDownsampleGS> GeometryShader(GetGlobalShaderMap());
-		TShaderMapRef<FCopyCubemapToCubeFacePS> PixelShader(GetGlobalShaderMap());
+		TShaderMapRef<FScreenVSForGS> VertexShader(ShaderMap);
+		TShaderMapRef<FDownsampleGS> GeometryShader(ShaderMap);
+		TShaderMapRef<FCopyCubemapToCubeFacePS> PixelShader(ShaderMap);
 
 		for (uint32 CubeFace = 0; CubeFace < CubeFace_MAX; CubeFace++)
 		{
@@ -978,8 +985,8 @@ void CopyCubemapToScratchCubemap(FRHICommandList& RHICmdList, ERHIFeatureLevel::
 			RHICmdList.SetDepthStencilState(TStaticDepthStencilState<false, CF_Always>::GetRHI());
 			RHICmdList.SetBlendState(TStaticBlendState<>::GetRHI());
 
-			TShaderMapRef<FScreenVSForGS> VertexShader(GetGlobalShaderMap());
-			TShaderMapRef<FCopyCubemapToCubeFacePS> PixelShader(GetGlobalShaderMap());
+			TShaderMapRef<FScreenVSForGS> VertexShader(ShaderMap);
+			TShaderMapRef<FCopyCubemapToCubeFacePS> PixelShader(ShaderMap);
 
 			SetGlobalBoundShaderState(RHICmdList, FeatureLevel, CopyFromCubemapToCubemapBoundShaderState, GFilterVertexDeclaration.VertexDeclarationRHI, *VertexShader, *PixelShader);
 			PixelShader->SetParameters(RHICmdList, SourceCubemapResource, CubeFace, bIsSkyLight, bLowerHemisphereIsBlack);
diff --git a/Engine/Source/Runtime/Renderer/Private/ReflectionEnvironmentDiffuseIrradiance.cpp b/Engine/Source/Runtime/Renderer/Private/ReflectionEnvironmentDiffuseIrradiance.cpp
--- a/Engine/Source/Runtime/Renderer/Private/ReflectionEnvironmentDiffuseIrradiance.cpp
+++ b/Engine/Source/Runtime/Renderer/Private/ReflectionEnvironmentDiffuseIrradiance.cpp
@@ -237,6 +237,8 @@ FGlobalBoundShaderState AccumulateCubeFacesBoundShaderState;
 
 void ComputeDiffuseIrradiance(FRHICommandListImmediate& RHICmdList, ERHIFeatureLevel::Type FeatureLevel, FTextureRHIRef LightingSource, int32 LightingSourceMipIndex, FSHVectorRGB3* OutIrradianceEnvironmentMap)
 {
+	auto ShaderMap = GetGlobalShaderMap(FeatureLevel);
+
 	for (int32 CoefficientIndex = 0; CoefficientIndex < FSHVector3::MaxSHBasis; CoefficientIndex++)
 	{
 		// Copy the starting mip from the lighting texture, apply texel area weighting and appropriate SH coefficient
@@ -254,9 +256,9 @@ void ComputeDiffuseIrradiance(FRHICommandListImmediate& RHICmdList, ERHIFeatureL
 				RHICmdList.SetDepthStencilState(TStaticDepthStencilState<false, CF_Always>::GetRHI());
 				RHICmdList.SetBlendState(TStaticBlendState<>::GetRHI());
 
-				TShaderMapRef<FScreenVSForGS> VertexShader(GetGlobalShaderMap());
-				TShaderMapRef<FDownsampleGS> GeometryShader(GetGlobalShaderMap());
-				TShaderMapRef<FCopyDiffuseIrradiancePS> PixelShader(GetGlobalShaderMap());
+				TShaderMapRef<FScreenVSForGS> VertexShader(ShaderMap);
+				TShaderMapRef<FDownsampleGS> GeometryShader(ShaderMap);
+				TShaderMapRef<FCopyDiffuseIrradiancePS> PixelShader(ShaderMap);
 
 				SetGlobalBoundShaderState(RHICmdList, FeatureLevel, CopyDiffuseIrradianceShaderState, GFilterVertexDeclaration.VertexDeclarationRHI, *VertexShader, *PixelShader, *GeometryShader);
 
@@ -292,8 +294,8 @@ void ComputeDiffuseIrradiance(FRHICommandListImmediate& RHICmdList, ERHIFeatureL
 					RHICmdList.SetDepthStencilState(TStaticDepthStencilState<false, CF_Always>::GetRHI());
 					RHICmdList.SetBlendState(TStaticBlendState<>::GetRHI());
 					
-					TShaderMapRef<FScreenVSForGS> VertexShader(GetGlobalShaderMap());
-					TShaderMapRef<FCopyDiffuseIrradiancePS> PixelShader(GetGlobalShaderMap());
+					TShaderMapRef<FScreenVSForGS> VertexShader(ShaderMap);
+					TShaderMapRef<FCopyDiffuseIrradiancePS> PixelShader(ShaderMap);
 					
 					SetGlobalBoundShaderState(RHICmdList, FeatureLevel, CopyDiffuseIrradianceShaderState, GFilterVertexDeclaration.VertexDeclarationRHI, *VertexShader, *PixelShader);
 					
@@ -337,9 +339,9 @@ void ComputeDiffuseIrradiance(FRHICommandListImmediate& RHICmdList, ERHIFeatureL
 					RHICmdList.SetDepthStencilState(TStaticDepthStencilState<false, CF_Always>::GetRHI());
 					RHICmdList.SetBlendState(TStaticBlendState<>::GetRHI());
 
-					TShaderMapRef<FScreenVSForGS> VertexShader(GetGlobalShaderMap());
-					TShaderMapRef<FDownsampleGS> GeometryShader(GetGlobalShaderMap());
-					TShaderMapRef<FAccumulateDiffuseIrradiancePS> PixelShader(GetGlobalShaderMap());
+					TShaderMapRef<FScreenVSForGS> VertexShader(ShaderMap);
+					TShaderMapRef<FDownsampleGS> GeometryShader(ShaderMap);
+					TShaderMapRef<FAccumulateDiffuseIrradiancePS> PixelShader(ShaderMap);
 
 					SetGlobalBoundShaderState(RHICmdList, FeatureLevel, DiffuseIrradianceAccumulateShaderState, GFilterVertexDeclaration.VertexDeclarationRHI, *VertexShader, *PixelShader, *GeometryShader);
 
@@ -375,8 +377,8 @@ void ComputeDiffuseIrradiance(FRHICommandListImmediate& RHICmdList, ERHIFeatureL
 						RHICmdList.SetDepthStencilState(TStaticDepthStencilState<false, CF_Always>::GetRHI());
 						RHICmdList.SetBlendState(TStaticBlendState<>::GetRHI());
 						
-						TShaderMapRef<FScreenVSForGS> VertexShader(GetGlobalShaderMap());
-						TShaderMapRef<FAccumulateDiffuseIrradiancePS> PixelShader(GetGlobalShaderMap());
+						TShaderMapRef<FScreenVSForGS> VertexShader(ShaderMap);
+						TShaderMapRef<FAccumulateDiffuseIrradiancePS> PixelShader(ShaderMap);
 						
 						SetGlobalBoundShaderState(RHICmdList, FeatureLevel, DiffuseIrradianceAccumulateShaderState, GFilterVertexDeclaration.VertexDeclarationRHI, *VertexShader, *PixelShader);
 						
@@ -409,8 +411,8 @@ void ComputeDiffuseIrradiance(FRHICommandListImmediate& RHICmdList, ERHIFeatureL
 			RHICmdList.SetDepthStencilState(TStaticDepthStencilState<false, CF_Always>::GetRHI());
 			RHICmdList.SetBlendState(TStaticBlendState<>::GetRHI());
 
-			TShaderMapRef<FScreenVS> VertexShader(GetGlobalShaderMap());
-			TShaderMapRef<FAccumulateCubeFacesPS> PixelShader(GetGlobalShaderMap());
+			TShaderMapRef<FScreenVS> VertexShader(ShaderMap);
+			TShaderMapRef<FAccumulateCubeFacesPS> PixelShader(ShaderMap);
 			SetGlobalBoundShaderState(RHICmdList, FeatureLevel, AccumulateCubeFacesBoundShaderState, GFilterVertexDeclaration.VertexDeclarationRHI, *VertexShader, *PixelShader);
 
 			const int32 SourceMipIndex = NumMips - 1;
diff --git a/Engine/Source/Runtime/Renderer/Private/SceneCaptureRendering.cpp b/Engine/Source/Runtime/Renderer/Private/SceneCaptureRendering.cpp
--- a/Engine/Source/Runtime/Renderer/Private/SceneCaptureRendering.cpp
+++ b/Engine/Source/Runtime/Renderer/Private/SceneCaptureRendering.cpp
@@ -29,8 +29,9 @@ void UpdateSceneCaptureContent_RenderThread(FRHICommandListImmediate& RHICmdList
 
 		// Render the scene normally
 		const FRenderTarget* Target = SceneRenderer->ViewFamily.RenderTarget;
-		FIntRect ViewRect = SceneRenderer->Views[0].ViewRect;
-		FIntRect UnconstrainedViewRect = SceneRenderer->Views[0].UnconstrainedViewRect;
+		FViewInfo& View = SceneRenderer->Views[0];
+		FIntRect ViewRect = View.ViewRect;
+		FIntRect UnconstrainedViewRect = View.UnconstrainedViewRect;
 		SetRenderTarget(RHICmdList, Target->GetRenderTargetTexture(), NULL);
 		RHICmdList.Clear(true, FLinearColor::Black, false, 1.0f, false, 0, ViewRect);
 		SceneRenderer->Render(RHICmdList);
@@ -45,14 +46,14 @@ void UpdateSceneCaptureContent_RenderThread(FRHICommandListImmediate& RHICmdList
 			RHICmdList.SetDepthStencilState(TStaticDepthStencilState<false, CF_Always>::GetRHI());
 			RHICmdList.SetBlendState(TStaticBlendState<>::GetRHI());
 
-			TShaderMapRef<FScreenVS> VertexShader(GetGlobalShaderMap());
-			TShaderMapRef<FScreenPS> PixelShader(GetGlobalShaderMap());
+			TShaderMapRef<FScreenVS> VertexShader(View.ShaderMap);
+			TShaderMapRef<FScreenPS> PixelShader(View.ShaderMap);
 			static FGlobalBoundShaderState BoundShaderState;
-			SetGlobalBoundShaderState(RHICmdList, SceneRenderer->FeatureLevel, BoundShaderState, GFilterVertexDeclaration.VertexDeclarationRHI, *VertexShader, *PixelShader);
+			SetGlobalBoundShaderState(RHICmdList, View.GetFeatureLevel(), BoundShaderState, GFilterVertexDeclaration.VertexDeclarationRHI, *VertexShader, *PixelShader);
 
-			FRenderingCompositePassContext Context(RHICmdList, SceneRenderer->Views[0]);
+			FRenderingCompositePassContext Context(RHICmdList, View);
 
-			VertexShader->SetParameters(RHICmdList, SceneRenderer->Views[0]);
+			VertexShader->SetParameters(RHICmdList, View);
 			PixelShader->SetParameters(RHICmdList, TStaticSamplerState<SF_Point>::GetRHI(), GSceneRenderTargets.GetSceneColorTexture());
 
 			FIntPoint TargetSize(UnconstrainedViewRect.Width(), UnconstrainedViewRect.Height());
diff --git a/Engine/Source/Runtime/Renderer/Private/SceneHitProxyRendering.cpp b/Engine/Source/Runtime/Renderer/Private/SceneHitProxyRendering.cpp
--- a/Engine/Source/Runtime/Renderer/Private/SceneHitProxyRendering.cpp
+++ b/Engine/Source/Runtime/Renderer/Private/SceneHitProxyRendering.cpp
@@ -37,7 +37,7 @@ public:
 	{
 		// Only compile the hit proxy vertex shader on PC
 		return IsPCPlatform(Platform)
-			&& IsFeatureLevelSupported(Platform, ERHIFeatureLevel::SM4)
+			&& (IsFeatureLevelSupported(Platform, ERHIFeatureLevel::SM4) || IsPCPlatform(Platform))
 			// and only compile for the default material or materials that are masked.
 			&& (Material->IsSpecialEngineMaterial() || Material->IsMasked() || Material->MaterialMayModifyMeshPosition() || Material->IsTwoSided());
 	}
@@ -110,7 +110,7 @@ public:
 	{
 		// Only compile the hit proxy vertex shader on PC
 		return IsPCPlatform(Platform) 
-			&& IsFeatureLevelSupported(Platform, ERHIFeatureLevel::SM4) 
+			&& IsFeatureLevelSupported(Platform, ERHIFeatureLevel::ES2) 
 			// and only compile for default materials or materials that are masked.
 			&& (Material->IsSpecialEngineMaterial() || Material->IsMasked() || Material->MaterialMayModifyMeshPosition() || Material->IsTwoSided());
 	}
diff --git a/Engine/Source/Runtime/Renderer/Private/SceneOcclusion.cpp b/Engine/Source/Runtime/Renderer/Private/SceneOcclusion.cpp
--- a/Engine/Source/Runtime/Renderer/Private/SceneOcclusion.cpp
+++ b/Engine/Source/Runtime/Renderer/Private/SceneOcclusion.cpp
@@ -746,8 +746,8 @@ void FHZBOcclusionTester::Submit(FRHICommandListImmediate& RHICmdList, const FVi
 
 		SetRenderTarget(RHICmdList, ResultsTextureGPU->GetRenderTargetItem().TargetableTexture, NULL);
 
-		TShaderMapRef< FScreenVS >	VertexShader( GetGlobalShaderMap() );
-		TShaderMapRef< FHZBTestPS >	PixelShader( GetGlobalShaderMap() );
+		TShaderMapRef< FScreenVS >	VertexShader(View.ShaderMap);
+		TShaderMapRef< FHZBTestPS >	PixelShader(View.ShaderMap);
 
 		static FGlobalBoundShaderState BoundShaderState;
 		SetGlobalBoundShaderState(RHICmdList, View.GetFeatureLevel(), BoundShaderState, GFilterVertexDeclaration.VertexDeclarationRHI, *VertexShader, *PixelShader);
@@ -890,8 +890,8 @@ void BuildHZB(FRHICommandListImmediate& RHICmdList, const FViewInfo& View)
 
 		SetRenderTarget(RHICmdList, HZBRenderTarget.TargetableTexture, 0, NULL);
 
-		TShaderMapRef< FScreenVS >		VertexShader( GetGlobalShaderMap() );
-		TShaderMapRef< THZBBuildPS<0> >	PixelShader( GetGlobalShaderMap() );
+		TShaderMapRef< FScreenVS >		VertexShader(View.ShaderMap);
+		TShaderMapRef< THZBBuildPS<0> >	PixelShader(View.ShaderMap);
 
 		static FGlobalBoundShaderState BoundShaderState;
 		
@@ -924,8 +924,8 @@ void BuildHZB(FRHICommandListImmediate& RHICmdList, const FViewInfo& View)
 
 		SetRenderTarget(RHICmdList, HZBRenderTarget.TargetableTexture, MipIndex, NULL);
 
-		TShaderMapRef< FScreenVS >		VertexShader( GetGlobalShaderMap() );
-		TShaderMapRef< THZBBuildPS<1> >	PixelShader( GetGlobalShaderMap() );
+		TShaderMapRef< FScreenVS >		VertexShader(View.ShaderMap);
+		TShaderMapRef< THZBBuildPS<1> >	PixelShader(View.ShaderMap);
 
 		static FGlobalBoundShaderState BoundShaderState;
 		
@@ -1002,7 +1002,7 @@ void FDeferredShadingSceneRenderer::BeginOcclusionTests(FRHICommandListImmediate
 			RHICmdList.SetBlendState(TStaticBlendState<CW_NONE>::GetRHI());
 
 			// Lookup the vertex shader.
-			TShaderMapRef<FOcclusionQueryVS> VertexShader(GetGlobalShaderMap());
+			TShaderMapRef<FOcclusionQueryVS> VertexShader(View.ShaderMap);
 			SetGlobalBoundShaderState(RHICmdList, View.GetFeatureLevel(), OcclusionTestBoundShaderState, GetVertexDeclarationFVector3(), *VertexShader, NULL);
 			VertexShader->SetParameters(RHICmdList, View);
 
diff --git a/Engine/Source/Runtime/Renderer/Private/ScenePostProcessing.cpp b/Engine/Source/Runtime/Renderer/Private/ScenePostProcessing.cpp
--- a/Engine/Source/Runtime/Renderer/Private/ScenePostProcessing.cpp
+++ b/Engine/Source/Runtime/Renderer/Private/ScenePostProcessing.cpp
@@ -120,8 +120,8 @@ void FSceneRenderer::GammaCorrectToViewportRenderTarget(FRHICommandList& RHICmdL
 	// turn off depth reads/writes
 	RHICmdList.SetDepthStencilState(TStaticDepthStencilState<false, CF_Always>::GetRHI());
 
-	TShaderMapRef<FGammaCorrectionVS> VertexShader(GetGlobalShaderMap());
-	TShaderMapRef<FGammaCorrectionPS> PixelShader(GetGlobalShaderMap());
+	TShaderMapRef<FGammaCorrectionVS> VertexShader(View->ShaderMap);
+	TShaderMapRef<FGammaCorrectionPS> PixelShader(View->ShaderMap);
 
 	static FGlobalBoundShaderState PostProcessBoundShaderState;
 	
diff --git a/Engine/Source/Runtime/Renderer/Private/SceneRendering.cpp b/Engine/Source/Runtime/Renderer/Private/SceneRendering.cpp
--- a/Engine/Source/Runtime/Renderer/Private/SceneRendering.cpp
+++ b/Engine/Source/Runtime/Renderer/Private/SceneRendering.cpp
@@ -129,6 +129,8 @@ void FViewInfo::Init()
 	}
 
 	MaxShadowCascades = FMath::Clamp<int32>(CVarMaxShadowCascades.GetValueOnAnyThread(), 1, 10);
+
+	ShaderMap = GetGlobalShaderMap(FeatureLevel);
 }
 
 FViewInfo::~FViewInfo()
diff --git a/Engine/Source/Runtime/Renderer/Private/SceneRendering.h b/Engine/Source/Runtime/Renderer/Private/SceneRendering.h
--- a/Engine/Source/Runtime/Renderer/Private/SceneRendering.h
+++ b/Engine/Source/Runtime/Renderer/Private/SceneRendering.h
@@ -426,6 +426,8 @@ public:
 	FLinearColor LightShaftColorApply;
 	bool bLightShaftUse;
 
+	TShaderMap<FGlobalShaderType>* ShaderMap;
+
 	/** 
 	 * Initialization constructor. Passes all parameters to FSceneView constructor
 	 */
@@ -618,7 +620,7 @@ protected:
 	void RenderForwardShadingBasePass(FRHICommandListImmediate& RHICmdList);
 
 	/** Makes a copy of scene alpha so PC can emulate ES2 framebuffer fetch. */
-	void CopySceneAlpha(FRHICommandListImmediate& RHICmdList, const FSceneView& View);
+	void CopySceneAlpha(FRHICommandListImmediate& RHICmdList, const FViewInfo& View);
 
 	/** Renders the base pass for translucency. */
 	void RenderTranslucency(FRHICommandListImmediate& RHICmdList);
diff --git a/Engine/Source/Runtime/Renderer/Private/ShaderComplexityRendering.cpp b/Engine/Source/Runtime/Renderer/Private/ShaderComplexityRendering.cpp
--- a/Engine/Source/Runtime/Renderer/Private/ShaderComplexityRendering.cpp
+++ b/Engine/Source/Runtime/Renderer/Private/ShaderComplexityRendering.cpp
@@ -180,8 +180,8 @@ void FRCPassPostProcessVisualizeComplexity::Process(FRenderingCompositePassConte
 	Context.RHICmdList.SetDepthStencilState(TStaticDepthStencilState<false, CF_Always>::GetRHI());
 
 	//reuse this generic vertex shader
-	TShaderMapRef<FShaderComplexityApplyVS> VertexShader(GetGlobalShaderMap());
-	TShaderMapRef<FShaderComplexityApplyPS> PixelShader(GetGlobalShaderMap());
+	TShaderMapRef<FShaderComplexityApplyVS> VertexShader(Context.GetShaderMap());
+	TShaderMapRef<FShaderComplexityApplyPS> PixelShader(Context.GetShaderMap());
 
 	static FGlobalBoundShaderState ShaderComplexityBoundShaderState;
 	
diff --git a/Engine/Source/Runtime/Renderer/Private/ShadowRendering.cpp b/Engine/Source/Runtime/Renderer/Private/ShadowRendering.cpp
--- a/Engine/Source/Runtime/Renderer/Private/ShadowRendering.cpp
+++ b/Engine/Source/Runtime/Renderer/Private/ShadowRendering.cpp
@@ -1658,10 +1658,10 @@ static void SetShadowProjectionShaderTemplNew(FRHICommandList& RHICmdList, int32
 	if (ShadowInfo->bTranslucentShadow)
 	{
 		// Get the Shadow Projection Vertex Shader (with transforms)
-		FShadowProjectionVS* ShadowProjVS = GetGlobalShaderMap()->GetShader<FShadowProjectionVS>();
+		FShadowProjectionVS* ShadowProjVS = View.ShaderMap->GetShader<FShadowProjectionVS>();
 
 		// Get the translucency pixel shader
-		FShadowProjectionPixelShaderInterface* ShadowProjPS = GetGlobalShaderMap()->GetShader<TShadowProjectionFromTranslucencyPS<Quality> >();
+		FShadowProjectionPixelShaderInterface* ShadowProjPS = View.ShaderMap->GetShader<TShadowProjectionFromTranslucencyPS<Quality> >();
 
 		// Bind shader
 		static FGlobalBoundShaderState BoundShaderState;
@@ -1675,13 +1675,13 @@ static void SetShadowProjectionShaderTemplNew(FRHICommandList& RHICmdList, int32
 	else if (ShadowInfo->IsWholeSceneDirectionalShadow())
 	{
 		// Get the Shadow Projection Vertex Shader which does not use a transform
-		FShadowProjectionNoTransformVS* ShadowProjVS = GetGlobalShaderMap()->GetShader<FShadowProjectionNoTransformVS>();
+		FShadowProjectionNoTransformVS* ShadowProjVS = View.ShaderMap->GetShader<FShadowProjectionNoTransformVS>();
 
 		// Get the Shadow Projection Pixel Shader for PSSM
 		if (ShadowInfo->CascadeSettings.FadePlaneLength > 0)
 		{
 			// This shader fades the shadow towards the end of the split subfrustum.
-			FShadowProjectionPixelShaderInterface* ShadowProjPS = GetGlobalShaderMap()->GetShader<TShadowProjectionPS<Quality, true> >();
+			FShadowProjectionPixelShaderInterface* ShadowProjPS = View.ShaderMap->GetShader<TShadowProjectionPS<Quality, true> >();
 
 			static FGlobalBoundShaderState BoundShaderState;
 			
@@ -1692,7 +1692,7 @@ static void SetShadowProjectionShaderTemplNew(FRHICommandList& RHICmdList, int32
 		else
 		{
 			// Do not use the fade plane shader if the fade plane region length is 0 (avoids divide by 0).
-			FShadowProjectionPixelShaderInterface* ShadowProjPS = GetGlobalShaderMap()->GetShader<TShadowProjectionPS<Quality, false> >();
+			FShadowProjectionPixelShaderInterface* ShadowProjPS = View.ShaderMap->GetShader<TShadowProjectionPS<Quality, false> >();
 
 			static FGlobalBoundShaderState BoundShaderState;
 			
@@ -1706,11 +1706,11 @@ static void SetShadowProjectionShaderTemplNew(FRHICommandList& RHICmdList, int32
 	else
 	{
 		// Get the Shadow Projection Vertex Shader
-		FShadowProjectionVS* ShadowProjVS = GetGlobalShaderMap()->GetShader<FShadowProjectionVS>();
+		FShadowProjectionVS* ShadowProjVS = View.ShaderMap->GetShader<FShadowProjectionVS>();
 
 		// Get the Shadow Projection Pixel Shader
 		// This shader is the ordinary projection shader used by point/spot lights.		
-		FShadowProjectionPixelShaderInterface* ShadowProjPS = GetGlobalShaderMap()->GetShader<TShadowProjectionPS<Quality, false> >();
+		FShadowProjectionPixelShaderInterface* ShadowProjPS = View.ShaderMap->GetShader<TShadowProjectionPS<Quality, false> >();
 
 		static FGlobalBoundShaderState BoundShaderState;
 		
@@ -1811,7 +1811,7 @@ void FProjectedShadowInfo::RenderProjection(FRHICommandListImmediate& RHICmdList
 	}
 
 	// Find the projection shaders.
-	TShaderMapRef<FShadowProjectionVS> VertexShader(GetGlobalShaderMap());
+	TShaderMapRef<FShadowProjectionVS> VertexShader(View->ShaderMap);
 
 	// Depth test wo/ writes, no color writing.
 	// Note, this is a reversed Z depth surface, using CF_GreaterEqual.
@@ -1934,7 +1934,7 @@ void FProjectedShadowInfo::RenderProjection(FRHICommandListImmediate& RHICmdList
 			// Draw 2 fullscreen planes, front facing one at the near subfrustum plane, and back facing one at the far.
 
 			// Find the projection shaders.
-			TShaderMapRef<FShadowProjectionNoTransformVS> VertexShaderNoTransform(GetGlobalShaderMap());
+			TShaderMapRef<FShadowProjectionNoTransformVS> VertexShaderNoTransform(View->ShaderMap);
 			VertexShaderNoTransform->SetParameters(RHICmdList, *View);
 			SetGlobalBoundShaderState(RHICmdList, View->GetFeatureLevel(), MaskBoundShaderState[0], GetVertexDeclarationFVector4(), *VertexShaderNoTransform, nullptr);
 
@@ -2000,7 +2000,7 @@ void FProjectedShadowInfo::RenderProjection(FRHICommandListImmediate& RHICmdList
 		}
 		
 		// Find the projection shaders.
-		TShaderMapRef<FShadowProjectionVS> VertexShader(GetGlobalShaderMap());
+		TShaderMapRef<FShadowProjectionVS> VertexShader(View->ShaderMap);
 
 		// Cache the bound shader state
 		SetGlobalBoundShaderState(RHICmdList, View->GetFeatureLevel(), MaskBoundShaderState[1], GetVertexDeclarationFVector4(), *VertexShader, NULL);
@@ -2148,8 +2148,8 @@ void FProjectedShadowInfo::RenderProjection(FRHICommandListImmediate& RHICmdList
 template <uint32 Quality>
 static void SetPointLightShaderTempl(FRHICommandList& RHICmdList, int32 ViewIndex, const FViewInfo& View, const FProjectedShadowInfo* ShadowInfo)
 {
-	TShaderMapRef<FShadowProjectionVS> VertexShader(GetGlobalShaderMap());
-	TShaderMapRef<TOnePassPointShadowProjectionPS<Quality> > PixelShader(GetGlobalShaderMap());
+	TShaderMapRef<FShadowProjectionVS> VertexShader(View.ShaderMap);
+	TShaderMapRef<TOnePassPointShadowProjectionPS<Quality> > PixelShader(View.ShaderMap);
 
 	static FGlobalBoundShaderState BoundShaderState;
 	
@@ -2992,6 +2992,7 @@ bool FDeferredShadingSceneRenderer::RenderReflectiveShadowMaps(FRHICommandListIm
 					{
 						LightPropagationVolume->InjectDirectionalLightRSM( 
 							RHICmdList, 
+							*ProjectedShadowInfo->DependentView,
 							GSceneRenderTargets.GetReflectiveShadowMapDiffuseTexture(), 
 							GSceneRenderTargets.GetReflectiveShadowMapNormalTexture(),
 							GSceneRenderTargets.GetReflectiveShadowMapDepthTexture(),
diff --git a/Engine/Source/Runtime/Renderer/Private/TiledDeferredLightRendering.cpp b/Engine/Source/Runtime/Renderer/Private/TiledDeferredLightRendering.cpp
--- a/Engine/Source/Runtime/Renderer/Private/TiledDeferredLightRendering.cpp
+++ b/Engine/Source/Runtime/Renderer/Private/TiledDeferredLightRendering.cpp
@@ -291,7 +291,7 @@ bool FDeferredShadingSceneRenderer::ShouldUseTiledDeferred(int32 NumUnshadowedLi
 template <bool bVisualizeLightCulling>
 static void SetShaderTemplTiledLighting(
 	FRHICommandListImmediate& RHICmdList,
-	const FSceneView& View, 
+	const FViewInfo& View,
 	int32 ViewIndex,
 	int32 NumViews,
 	const TArray<FSortedLightSceneInfo, SceneRenderingAllocator>& SortedLights, 
@@ -302,7 +302,7 @@ static void SetShaderTemplTiledLighting(
 	IPooledRenderTarget& InTexture,
 	IPooledRenderTarget& OutTexture)
 {
-	TShaderMapRef<FTiledDeferredLightingCS<bVisualizeLightCulling> > ComputeShader(GetGlobalShaderMap());
+	TShaderMapRef<FTiledDeferredLightingCS<bVisualizeLightCulling> > ComputeShader(View.ShaderMap);
 	RHICmdList.SetComputeShader(ComputeShader->GetComputeShader());
 
 	ComputeShader->SetParameters(RHICmdList, View, ViewIndex, NumViews, SortedLights, NumLightsToRenderInSortedLights, SimpleLights, StartIndex, NumThisPass, InTexture, OutTexture);
diff --git a/Engine/Source/Runtime/Renderer/Private/TranslucentLighting.cpp b/Engine/Source/Runtime/Renderer/Private/TranslucentLighting.cpp
--- a/Engine/Source/Runtime/Renderer/Private/TranslucentLighting.cpp
+++ b/Engine/Source/Runtime/Renderer/Private/TranslucentLighting.cpp
@@ -1080,9 +1080,10 @@ void ClearVolumeTextures(FRHICommandList& RHICmdList, ERHIFeatureLevel::Type Fea
 		RHICmdList.SetBlendState(TStaticBlendState<>::GetRHI());
 
 		const FVolumeBounds VolumeBounds(GTranslucencyLightingVolumeDim);
-		TShaderMapRef<FWriteToSliceVS> VertexShader(GetGlobalShaderMap());
-		TShaderMapRef<FWriteToSliceGS> GeometryShader(GetGlobalShaderMap());
-		TShaderMapRef<TOneColorPixelShaderMRT<NumRenderTargets> > PixelShader(GetGlobalShaderMap());
+		auto ShaderMap = GetGlobalShaderMap(FeatureLevel);
+		TShaderMapRef<FWriteToSliceVS> VertexShader(ShaderMap);
+		TShaderMapRef<FWriteToSliceGS> GeometryShader(ShaderMap);
+		TShaderMapRef<TOneColorPixelShaderMRT<NumRenderTargets> > PixelShader(ShaderMap);
 		
 		SetGlobalBoundShaderState(RHICmdList, FeatureLevel, VolumeClearBoundShaderState, GScreenVertexDeclaration.VertexDeclarationRHI, *VertexShader, *PixelShader, *GeometryShader);
 
@@ -1186,15 +1187,17 @@ void FDeferredShadingSceneRenderer::InjectAmbientCubemapTranslucentVolumeLightin
 		RHICmdList.SetBlendState(TStaticBlendState<CW_RGBA, BO_Add, BF_One, BF_One, BO_Add, BF_One, BF_One>::GetRHI());
 
 		const FVolumeBounds VolumeBounds(GTranslucencyLightingVolumeDim);
+
+		auto ShaderMap = GetGlobalShaderMap(FeatureLevel);
 		
 		for (int32 VolumeCascadeIndex = 0; VolumeCascadeIndex < TVC_MAX; VolumeCascadeIndex++)
 		{
 			// we don't update the directional volume (could be a HQ option)
 			SetRenderTarget(RHICmdList, GSceneRenderTargets.TranslucencyLightingVolumeAmbient[VolumeCascadeIndex]->GetRenderTargetItem().TargetableTexture, FTextureRHIRef());
 
-			TShaderMapRef<FWriteToSliceVS> VertexShader(GetGlobalShaderMap());
-			TShaderMapRef<FWriteToSliceGS> GeometryShader(GetGlobalShaderMap());
-			TShaderMapRef<FInjectAmbientCubemapPS> PixelShader(GetGlobalShaderMap());
+			TShaderMapRef<FWriteToSliceVS> VertexShader(ShaderMap);
+			TShaderMapRef<FWriteToSliceGS> GeometryShader(ShaderMap);
+			TShaderMapRef<FInjectAmbientCubemapPS> PixelShader(ShaderMap);
 
 			static FGlobalBoundShaderState BoundShaderState;
 									
@@ -1286,6 +1289,8 @@ void FDeferredShadingSceneRenderer::AccumulateTranslucentVolumeObjectShadowing(F
 		// Modulate the contribution of multiple object shadows in rgb
 		RHICmdList.SetBlendState(TStaticBlendState<CW_RGB, BO_Add, BF_DestColor, BF_Zero>::GetRHI());
 
+		auto ShaderMap = GetGlobalShaderMap(FeatureLevel);
+
 		// Inject into each volume cascade
 		for (int32 VolumeCascadeIndex = 0; VolumeCascadeIndex < TVC_MAX; VolumeCascadeIndex++)
 		{
@@ -1309,9 +1314,9 @@ void FDeferredShadingSceneRenderer::AccumulateTranslucentVolumeObjectShadowing(F
 
 				SetRenderTarget(RHICmdList, RenderTarget, FTextureRHIRef());
 
-				TShaderMapRef<FWriteToSliceVS> VertexShader(GetGlobalShaderMap());
-				TShaderMapRef<FWriteToSliceGS> GeometryShader(GetGlobalShaderMap());
-				TShaderMapRef<FTranslucentObjectShadowingPS> PixelShader(GetGlobalShaderMap());
+				TShaderMapRef<FWriteToSliceVS> VertexShader(ShaderMap);
+				TShaderMapRef<FWriteToSliceGS> GeometryShader(ShaderMap);
+				TShaderMapRef<FTranslucentObjectShadowingPS> PixelShader(ShaderMap);
 
 				SetGlobalBoundShaderState(RHICmdList, FeatureLevel, ObjectShadowingBoundShaderState, GScreenVertexDeclaration.VertexDeclarationRHI, *VertexShader, *PixelShader, *GeometryShader);
 				VertexShader->SetParameters(RHICmdList, VolumeBounds, GTranslucencyLightingVolumeDim);
@@ -1523,8 +1528,8 @@ static void InjectTranslucentLightArray(FRHICommandListImmediate& RHICmdList, co
 
 			if (VolumeBounds.IsValid())
 			{
-				TShaderMapRef<FWriteToSliceVS> VertexShader(GetGlobalShaderMap());
-				TShaderMapRef<FWriteToSliceGS> GeometryShader(GetGlobalShaderMap());
+				TShaderMapRef<FWriteToSliceVS> VertexShader(View.ShaderMap);
+				TShaderMapRef<FWriteToSliceGS> GeometryShader(View.ShaderMap);
 
 				if (bDirectionalLight)
 				{
@@ -1728,9 +1733,9 @@ void FDeferredShadingSceneRenderer::InjectSimpleTranslucentVolumeLightingArray(F
 
 					if (VolumeBounds.IsValid())
 					{
-						TShaderMapRef<FWriteToSliceVS> VertexShader(GetGlobalShaderMap());
-						TShaderMapRef<FWriteToSliceGS> GeometryShader(GetGlobalShaderMap());
-						TShaderMapRef<FSimpleLightTranslucentLightingInjectPS> PixelShader(GetGlobalShaderMap());
+						TShaderMapRef<FWriteToSliceVS> VertexShader(View.ShaderMap);
+						TShaderMapRef<FWriteToSliceGS> GeometryShader(View.ShaderMap);
+						TShaderMapRef<FSimpleLightTranslucentLightingInjectPS> PixelShader(View.ShaderMap);
 						SetGlobalBoundShaderState(RHICmdList, FeatureLevel, InjectSimpleLightBoundShaderState, GScreenVertexDeclaration.VertexDeclarationRHI, *VertexShader, *PixelShader, *GeometryShader);
 
 						VertexShader->SetParameters(RHICmdList, VolumeBounds, GTranslucencyLightingVolumeDim);
@@ -1786,9 +1791,9 @@ void FDeferredShadingSceneRenderer::FilterTranslucentVolumeLighting(FRHICommandL
 				const FViewInfo& View = Views[0];
 
 				const FVolumeBounds VolumeBounds(GTranslucencyLightingVolumeDim);
-				TShaderMapRef<FWriteToSliceVS> VertexShader(GetGlobalShaderMap());
-				TShaderMapRef<FWriteToSliceGS> GeometryShader(GetGlobalShaderMap());
-				TShaderMapRef<FFilterTranslucentVolumePS> PixelShader(GetGlobalShaderMap());
+				TShaderMapRef<FWriteToSliceVS> VertexShader(View.ShaderMap);
+				TShaderMapRef<FWriteToSliceGS> GeometryShader(View.ShaderMap);
+				TShaderMapRef<FFilterTranslucentVolumePS> PixelShader(View.ShaderMap);
 				SetGlobalBoundShaderState(RHICmdList, FeatureLevel, FilterBoundShaderState, GScreenVertexDeclaration.VertexDeclarationRHI, *VertexShader, *PixelShader, *GeometryShader);
 
 				VertexShader->SetParameters(RHICmdList, VolumeBounds, GTranslucencyLightingVolumeDim);
diff --git a/Engine/Source/Runtime/Renderer/Private/TranslucentRendering.cpp b/Engine/Source/Runtime/Renderer/Private/TranslucentRendering.cpp
--- a/Engine/Source/Runtime/Renderer/Private/TranslucentRendering.cpp
+++ b/Engine/Source/Runtime/Renderer/Private/TranslucentRendering.cpp
@@ -181,8 +181,8 @@ void FTranslucencyDrawingPolicyFactory::CopySceneColor(FRHICommandList& RHICmdLi
 	GSceneRenderTargets.BeginRenderingLightAttenuation(RHICmdList);
 	RHICmdList.SetViewport(View.ViewRect.Min.X, View.ViewRect.Min.Y, 0.0f, View.ViewRect.Max.X, View.ViewRect.Max.Y, 1.0f);
 
-	TShaderMapRef<FScreenVS> ScreenVertexShader(GetGlobalShaderMap());
-	TShaderMapRef<FCopySceneColorPS> PixelShader(GetGlobalShaderMap());
+	TShaderMapRef<FScreenVS> ScreenVertexShader(View.ShaderMap);
+	TShaderMapRef<FCopySceneColorPS> PixelShader(View.ShaderMap);
 	SetGlobalBoundShaderState(RHICmdList, View.GetFeatureLevel(), CopySceneColorBoundShaderState, GFilterVertexDeclaration.VertexDeclarationRHI, *ScreenVertexShader, *PixelShader);
 
 	/// ?
diff --git a/Engine/Source/Runtime/Renderer/Private/VelocityRendering.cpp b/Engine/Source/Runtime/Renderer/Private/VelocityRendering.cpp
--- a/Engine/Source/Runtime/Renderer/Private/VelocityRendering.cpp
+++ b/Engine/Source/Runtime/Renderer/Private/VelocityRendering.cpp
@@ -173,7 +173,8 @@ IMPLEMENT_MATERIAL_SHADER_TYPE(,FVelocityPS,TEXT("VelocityShader"),TEXT("MainPix
 FVelocityDrawingPolicy::FVelocityDrawingPolicy(
 	const FVertexFactory* InVertexFactory,
 	const FMaterialRenderProxy* InMaterialRenderProxy,
-	const FMaterial& InMaterialResource
+	const FMaterial& InMaterialResource,
+	ERHIFeatureLevel::Type InFeatureLevel
 	)
 	:	FMeshDrawingPolicy(InVertexFactory,InMaterialRenderProxy,InMaterialResource)	
 {
@@ -184,7 +185,7 @@ FVelocityDrawingPolicy::FVelocityDrawingPolicy(
 	DomainShader = NULL;
 
 	const EMaterialTessellationMode MaterialTessellationMode = InMaterialResource.GetTessellationMode();
-	if(RHISupportsTessellation(GRHIShaderPlatform)
+	if (RHISupportsTessellation(GShaderPlatformForFeatureLevel[InFeatureLevel])
 		&& InVertexFactory->GetType()->SupportsTessellationShaders() 
 		&& MaterialTessellationMode != MTM_NoTessellation)
 	{
@@ -355,7 +356,7 @@ void FVelocityDrawingPolicyFactory::AddStaticMesh(FScene* Scene, FStaticMesh* St
 			    MaterialRenderProxy = UMaterial::GetDefaultMaterial(MD_Surface)->GetRenderProxy(false);
 		    }
 
-			FVelocityDrawingPolicy DrawingPolicy(StaticMesh->VertexFactory, MaterialRenderProxy, *MaterialRenderProxy->GetMaterial(FeatureLevel));
+			FVelocityDrawingPolicy DrawingPolicy(StaticMesh->VertexFactory, MaterialRenderProxy, *MaterialRenderProxy->GetMaterial(FeatureLevel), FeatureLevel);
 			if (DrawingPolicy.SupportsVelocity())
 			{
 				// Add the static mesh to the depth-only draw list.
@@ -392,7 +393,7 @@ bool FVelocityDrawingPolicyFactory::DrawDynamicMesh(
 			// Default material doesn't handle masked, and doesn't have the correct bIsTwoSided setting.
 			MaterialRenderProxy = UMaterial::GetDefaultMaterial(MD_Surface)->GetRenderProxy(false);
 		}
-		FVelocityDrawingPolicy DrawingPolicy(Mesh.VertexFactory, MaterialRenderProxy, *MaterialRenderProxy->GetMaterial(FeatureLevel));
+		FVelocityDrawingPolicy DrawingPolicy(Mesh.VertexFactory, MaterialRenderProxy, *MaterialRenderProxy->GetMaterial(FeatureLevel), FeatureLevel);
 		if(DrawingPolicy.SupportsVelocity())
 		{			
 			RHICmdList.BuildAndSetLocalBoundShaderState(DrawingPolicy.GetBoundShaderStateInput(View.GetFeatureLevel()));
diff --git a/Engine/Source/Runtime/Renderer/Private/VelocityRendering.h b/Engine/Source/Runtime/Renderer/Private/VelocityRendering.h
--- a/Engine/Source/Runtime/Renderer/Private/VelocityRendering.h
+++ b/Engine/Source/Runtime/Renderer/Private/VelocityRendering.h
@@ -15,7 +15,8 @@ public:
 	FVelocityDrawingPolicy(
 		const FVertexFactory* InVertexFactory,
 		const FMaterialRenderProxy* InMaterialRenderProxy,
-		const FMaterial& InMaterialResource
+		const FMaterial& InMaterialResource,
+		ERHIFeatureLevel::Type InFeatureLevel
 		);
 
 	// FMeshDrawingPolicy interface.
diff --git a/Engine/Source/Runtime/ShaderCore/Private/Shader.cpp b/Engine/Source/Runtime/ShaderCore/Private/Shader.cpp
--- a/Engine/Source/Runtime/ShaderCore/Private/Shader.cpp
+++ b/Engine/Source/Runtime/ShaderCore/Private/Shader.cpp
@@ -668,7 +668,7 @@ FShader::FShader() :
 {
 	// set to undefined (currently shared with SF_Vertex)
 	Target.Frequency = 0;
-	Target.Platform = GRHIShaderPlatform;
+	Target.Platform = GShaderPlatformForFeatureLevel[GMaxRHIFeatureLevel];
 }
 
 /**
diff --git a/Engine/Source/Runtime/SlateRHIRenderer/Private/SlateRHIRenderer.cpp b/Engine/Source/Runtime/SlateRHIRenderer/Private/SlateRHIRenderer.cpp
--- a/Engine/Source/Runtime/SlateRHIRenderer/Private/SlateRHIRenderer.cpp
+++ b/Engine/Source/Runtime/SlateRHIRenderer/Private/SlateRHIRenderer.cpp
@@ -835,11 +835,14 @@ void FSlateRHIRenderer::CopyWindowsToVirtualScreenBuffer(const TArray<FString>&
 				DrawWindowToBuffer,
 				FDrawWindowToBufferContext,Context,DrawWindowToBufferContext,
 			{
-				TShaderMapRef<FScreenVS> VertexShader(GetGlobalShaderMap());
-				TShaderMapRef<FScreenPS> PixelShader(GetGlobalShaderMap());
+				const auto FeatureLevel = GRHIFeatureLevel;
+				auto ShaderMap = GetGlobalShaderMap(FeatureLevel);
+
+				TShaderMapRef<FScreenVS> VertexShader(ShaderMap);
+				TShaderMapRef<FScreenPS> PixelShader(ShaderMap);
 
 				static FGlobalBoundShaderState BoundShaderState;
-				SetGlobalBoundShaderState(RHICmdList, GRHIFeatureLevel, BoundShaderState, Context.RendererModule->GetFilterVertexDeclaration().VertexDeclarationRHI, *VertexShader, *PixelShader);
+				SetGlobalBoundShaderState(RHICmdList, FeatureLevel, BoundShaderState, Context.RendererModule->GetFilterVertexDeclaration().VertexDeclarationRHI, *VertexShader, *PixelShader);
 
 				if( Context.WindowRect.Width() != Context.InViewportInfo->Width || Context.WindowRect.Height() != Context.InViewportInfo->Height )
 				{
diff --git a/Engine/Source/Runtime/SlateRHIRenderer/Private/SlateRHIRenderingPolicy.cpp b/Engine/Source/Runtime/SlateRHIRenderer/Private/SlateRHIRenderingPolicy.cpp
--- a/Engine/Source/Runtime/SlateRHIRenderer/Private/SlateRHIRenderingPolicy.cpp
+++ b/Engine/Source/Runtime/SlateRHIRenderer/Private/SlateRHIRenderingPolicy.cpp
@@ -283,7 +283,7 @@ void FSlateRHIRenderingPolicy::DrawElements(FRHICommandListImmediate& RHICmdList
 
 	FSceneView* SceneView = NULL;
 
-	TShaderMapRef<FSlateElementVS> VertexShader(GetGlobalShaderMap());
+	TShaderMapRef<FSlateElementVS> VertexShader(GetGlobalShaderMap(GRHIFeatureLevel));
 
 	// Disabled stencil test state
 	FDepthStencilStateRHIRef DSOff = TStaticDepthStencilState<false,CF_Always>::GetRHI();
@@ -468,6 +468,8 @@ void FSlateRHIRenderingPolicy::DrawElements(FRHICommandListImmediate& RHICmdList
 FSlateElementPS* FSlateRHIRenderingPolicy::GetTexturePixelShader( ESlateShader::Type ShaderType, ESlateDrawEffect::Type DrawEffects )
 {
 	FSlateElementPS* PixelShader = NULL;
+	const auto FeatureLevel = GRHIFeatureLevel;
+	auto ShaderMap = GetGlobalShaderMap(FeatureLevel);
 
 #if !DEBUG_OVERDRAW
 	const bool bDrawDisabled = (DrawEffects & ESlateDrawEffect::DisabledEffect) != 0;
@@ -480,28 +482,28 @@ FSlateElementPS* FSlateRHIRenderingPolicy::GetTexturePixelShader( ESlateShader::
 		case ESlateShader::Default:
 			if( bUseTextureAlpha )
 			{
-				PixelShader = *TShaderMapRef<TSlateElementPS<ESlateShader::Default,true,true> >( GetGlobalShaderMap() );
+				PixelShader = *TShaderMapRef<TSlateElementPS<ESlateShader::Default, true, true> >(ShaderMap);
 			}
 			else
 			{
-				PixelShader = *TShaderMapRef<TSlateElementPS<ESlateShader::Default,true,false> >( GetGlobalShaderMap() );
+				PixelShader = *TShaderMapRef<TSlateElementPS<ESlateShader::Default, true, false> >(ShaderMap);
 			}
 			break;
 		case ESlateShader::Border:
 			if( bUseTextureAlpha )
 			{
-				PixelShader = *TShaderMapRef<TSlateElementPS<ESlateShader::Border,true,true> >( GetGlobalShaderMap() );
+				PixelShader = *TShaderMapRef<TSlateElementPS<ESlateShader::Border, true, true> >(ShaderMap);
 			}
 			else
 			{
-				PixelShader = *TShaderMapRef<TSlateElementPS<ESlateShader::Border,true,false> >( GetGlobalShaderMap() );
+				PixelShader = *TShaderMapRef<TSlateElementPS<ESlateShader::Border, true, false> >(ShaderMap);
 			}
 			break;
 		case ESlateShader::Font:
-			PixelShader = *TShaderMapRef<TSlateElementPS<ESlateShader::Font,true> >( GetGlobalShaderMap() );
+			PixelShader = *TShaderMapRef<TSlateElementPS<ESlateShader::Font, true> >(ShaderMap);
 			break;
 		case ESlateShader::LineSegment:
-			PixelShader = *TShaderMapRef<TSlateElementPS<ESlateShader::LineSegment,true> >( GetGlobalShaderMap() );
+			PixelShader = *TShaderMapRef<TSlateElementPS<ESlateShader::LineSegment, true> >(ShaderMap);
 			break;
 		}
 	}
@@ -513,33 +515,33 @@ FSlateElementPS* FSlateRHIRenderingPolicy::GetTexturePixelShader( ESlateShader::
 		case ESlateShader::Default:
 			if( bUseTextureAlpha )
 			{
-				PixelShader = *TShaderMapRef<TSlateElementPS<ESlateShader::Default,false,true> >( GetGlobalShaderMap() );
+				PixelShader = *TShaderMapRef<TSlateElementPS<ESlateShader::Default, false, true> >(ShaderMap);
 			}
 			else
 			{
-				PixelShader = *TShaderMapRef<TSlateElementPS<ESlateShader::Default,false,false> >( GetGlobalShaderMap() );
+				PixelShader = *TShaderMapRef<TSlateElementPS<ESlateShader::Default, false, false> >(ShaderMap);
 			}
 			break;
 		case ESlateShader::Border:
 			if( bUseTextureAlpha )
 			{
-				PixelShader = *TShaderMapRef<TSlateElementPS<ESlateShader::Border,false,true> >( GetGlobalShaderMap() );
+				PixelShader = *TShaderMapRef<TSlateElementPS<ESlateShader::Border, false, true> >(ShaderMap);
 			}
 			else
 			{
-				PixelShader = *TShaderMapRef<TSlateElementPS<ESlateShader::Border,false,false> >( GetGlobalShaderMap() );
+				PixelShader = *TShaderMapRef<TSlateElementPS<ESlateShader::Border, false, false> >(ShaderMap);
 			}
 			break;
 		case ESlateShader::Font:
-			PixelShader = *TShaderMapRef<TSlateElementPS<ESlateShader::Font,false> >( GetGlobalShaderMap() );
+			PixelShader = *TShaderMapRef<TSlateElementPS<ESlateShader::Font, false> >(ShaderMap);
 			break;
 		case ESlateShader::LineSegment:
-			PixelShader = *TShaderMapRef<TSlateElementPS<ESlateShader::LineSegment,false> >( GetGlobalShaderMap() );
+			PixelShader = *TShaderMapRef<TSlateElementPS<ESlateShader::LineSegment, false> >(ShaderMap);
 			break;
 		}
 	}
 #else
-	PixelShader = *TShaderMapRef<FSlateDebugOverdrawPS>( GetGlobalShaderMap() );
+	PixelShader = *TShaderMapRef<FSlateDebugOverdrawPS>(ShaderMap);
 #endif
 
 	return PixelShader;
diff --git a/Engine/Source/Runtime/Windows/D3D11RHI/Private/D3D11Commands.cpp b/Engine/Source/Runtime/Windows/D3D11RHI/Private/D3D11Commands.cpp
--- a/Engine/Source/Runtime/Windows/D3D11RHI/Private/D3D11Commands.cpp
+++ b/Engine/Source/Runtime/Windows/D3D11RHI/Private/D3D11Commands.cpp
@@ -1894,7 +1894,8 @@ void FD3D11DynamicRHI::RHIClearMRT(bool bClearColor,int32 NumClearColors,const F
 		OriginalResourceState.ClearCurrentVertexResources(StateCache);		
 
 		// Set the new shaders
-		TShaderMapRef<TOneColorVS<true> > VertexShader(GetGlobalShaderMap());
+		auto ShaderMap = GetGlobalShaderMap(GMaxRHIFeatureLevel);
+		TShaderMapRef<TOneColorVS<true> > VertexShader(ShaderMap);
 
 		FOneColorPS* PixelShader = NULL;
 
@@ -1902,42 +1903,42 @@ void FD3D11DynamicRHI::RHIClearMRT(bool bClearColor,int32 NumClearColors,const F
 		// On AMD PC hardware, outputting to a color index in the shader without a matching render target set has a significant performance hit
 		if (BoundRenderTargets.GetNumActiveTargets() <= 1)
 		{
-			TShaderMapRef<TOneColorPixelShaderMRT<1> > MRTPixelShader(GetGlobalShaderMap());
+			TShaderMapRef<TOneColorPixelShaderMRT<1> > MRTPixelShader(ShaderMap);
 			PixelShader = *MRTPixelShader;
 		}
 		else if (BoundRenderTargets.GetNumActiveTargets() == 2)
 		{
-			TShaderMapRef<TOneColorPixelShaderMRT<2> > MRTPixelShader(GetGlobalShaderMap());
+			TShaderMapRef<TOneColorPixelShaderMRT<2> > MRTPixelShader(ShaderMap);
 			PixelShader = *MRTPixelShader;
 		}
 		else if (BoundRenderTargets.GetNumActiveTargets() == 3)
 		{
-			TShaderMapRef<TOneColorPixelShaderMRT<3> > MRTPixelShader(GetGlobalShaderMap());
+			TShaderMapRef<TOneColorPixelShaderMRT<3> > MRTPixelShader(ShaderMap);
 			PixelShader = *MRTPixelShader;
 		}
 		else if (BoundRenderTargets.GetNumActiveTargets() == 4)
 		{
-			TShaderMapRef<TOneColorPixelShaderMRT<4> > MRTPixelShader(GetGlobalShaderMap());
+			TShaderMapRef<TOneColorPixelShaderMRT<4> > MRTPixelShader(ShaderMap);
 			PixelShader = *MRTPixelShader;
 		}
 		else if (BoundRenderTargets.GetNumActiveTargets() == 5)
 		{
-			TShaderMapRef<TOneColorPixelShaderMRT<5> > MRTPixelShader(GetGlobalShaderMap());
+			TShaderMapRef<TOneColorPixelShaderMRT<5> > MRTPixelShader(ShaderMap);
 			PixelShader = *MRTPixelShader;
 		}
 		else if (BoundRenderTargets.GetNumActiveTargets() == 6)
 		{
-			TShaderMapRef<TOneColorPixelShaderMRT<6> > MRTPixelShader(GetGlobalShaderMap());
+			TShaderMapRef<TOneColorPixelShaderMRT<6> > MRTPixelShader(ShaderMap);
 			PixelShader = *MRTPixelShader;
 		}
 		else if (BoundRenderTargets.GetNumActiveTargets() == 7)
 		{
-			TShaderMapRef<TOneColorPixelShaderMRT<7> > MRTPixelShader(GetGlobalShaderMap());
+			TShaderMapRef<TOneColorPixelShaderMRT<7> > MRTPixelShader(ShaderMap);
 			PixelShader = *MRTPixelShader;
 		}
 		else if (BoundRenderTargets.GetNumActiveTargets() == 8)
 		{
-			TShaderMapRef<TOneColorPixelShaderMRT<8> > MRTPixelShader(GetGlobalShaderMap());
+			TShaderMapRef<TOneColorPixelShaderMRT<8> > MRTPixelShader(ShaderMap);
 			PixelShader = *MRTPixelShader;
 		}
 
diff --git a/Engine/Source/Runtime/Windows/D3D11RHI/Private/D3D11RHI.cpp b/Engine/Source/Runtime/Windows/D3D11RHI/Private/D3D11RHI.cpp
--- a/Engine/Source/Runtime/Windows/D3D11RHI/Private/D3D11RHI.cpp
+++ b/Engine/Source/Runtime/Windows/D3D11RHI/Private/D3D11RHI.cpp
@@ -209,8 +209,9 @@ void FD3D11DynamicRHI::IssueLongGPUTask()
 			RHICmdList.SetDepthStencilState(TStaticDepthStencilState<false, CF_Always>::GetRHI(), 0);
 			RHICmdList.SetRasterizerState(TStaticRasterizerState<FM_Solid, CM_None>::GetRHI());
 
-			TShaderMapRef<TOneColorVS<true> > VertexShader(GetGlobalShaderMap());
-			TShaderMapRef<FLongGPUTaskPS> PixelShader(GetGlobalShaderMap());
+			auto ShaderMap = GetGlobalShaderMap(GMaxRHIFeatureLevel);
+			TShaderMapRef<TOneColorVS<true> > VertexShader(ShaderMap);
+			TShaderMapRef<FLongGPUTaskPS> PixelShader(ShaderMap);
 
 			SetGlobalBoundShaderState(RHICmdList, GMaxRHIFeatureLevel, LongGPUTaskBoundShaderState, GD3D11Vector4VertexDeclaration.VertexDeclarationRHI, *VertexShader, *PixelShader);
 
@@ -221,7 +222,7 @@ void FD3D11DynamicRHI::IssueLongGPUTask()
 			Vertices[2].Set( -1.0f, -1.0f, 0, 1.0f );
 			Vertices[3].Set(  1.0f, -1.0f, 0, 1.0f );
 			DrawPrimitiveUP(RHICmdList, PT_TriangleStrip, 2, Vertices, sizeof(Vertices[0]));
-			// Implicit fluish. Always call flush when using a command list in RHI implementations before doing anything else. This is super hazardous.
+			// Implicit flush. Always call flush when using a command list in RHI implementations before doing anything else. This is super hazardous.
 		}
 	}
 }
diff --git a/Engine/Source/Runtime/Windows/D3D11RHI/Private/D3D11RenderTarget.cpp b/Engine/Source/Runtime/Windows/D3D11RHI/Private/D3D11RenderTarget.cpp
--- a/Engine/Source/Runtime/Windows/D3D11RHI/Private/D3D11RenderTarget.cpp
+++ b/Engine/Source/Runtime/Windows/D3D11RHI/Private/D3D11RenderTarget.cpp
@@ -249,8 +249,9 @@ void FD3D11DynamicRHI::ResolveTextureUsingShader(
 	static FGlobalBoundShaderState ResolveBoundShaderState;
 
 	// Set the vertex and pixel shader
-	TShaderMapRef<FResolveVS> ResolveVertexShader(GetGlobalShaderMap());
-	TShaderMapRef<TPixelShader> ResolvePixelShader(GetGlobalShaderMap());
+	auto ShaderMap = GetGlobalShaderMap(GMaxRHIFeatureLevel);
+	TShaderMapRef<FResolveVS> ResolveVertexShader(ShaderMap);
+	TShaderMapRef<TPixelShader> ResolvePixelShader(ShaderMap);
 	SetGlobalBoundShaderState(RHICmdList, GMaxRHIFeatureLevel, ResolveBoundShaderState, GScreenVertexDeclaration.VertexDeclarationRHI, *ResolveVertexShader, *ResolvePixelShader);
 
 	ResolvePixelShader->SetParameters(RHICmdList, Direct3DDeviceContext,PixelShaderParameter);
