Fixed compile error

diff --git a/Engine/Source/Runtime/Renderer/Private/DistanceFieldSurfaceCacheLighting.cpp b/Engine/Source/Runtime/Renderer/Private/DistanceFieldSurfaceCacheLighting.cpp
--- a/Engine/Source/Runtime/Renderer/Private/DistanceFieldSurfaceCacheLighting.cpp
+++ b/Engine/Source/Runtime/Renderer/Private/DistanceFieldSurfaceCacheLighting.cpp
@@ -3224,7 +3224,7 @@ void FDeferredShadingSceneRenderer::RenderMeshDistanceFieldVisualization(FRHICom
 
 						{
 							SCOPED_DRAW_EVENT(VisualizeMeshDistanceFieldCS, DEC_SCENE_ITEMS);
-							TShaderMapRef<FVisualizeMeshDistanceFieldCS> ComputeShader(GetGlobalShaderMap());
+							TShaderMapRef<FVisualizeMeshDistanceFieldCS> ComputeShader(View.ShaderMap);
 
 							RHICmdList.SetComputeShader(ComputeShader->GetComputeShader());
 							ComputeShader->SetParameters(RHICmdList, View, VisualizeResultRT->GetRenderTargetItem(), NumObjects, Parameters);
@@ -3249,8 +3249,8 @@ void FDeferredShadingSceneRenderer::RenderMeshDistanceFieldVisualization(FRHICom
 						RHICmdList.SetDepthStencilState(TStaticDepthStencilState<false, CF_Always>::GetRHI());
 						RHICmdList.SetBlendState(TStaticBlendState<>::GetRHI());
 
-						TShaderMapRef<FPostProcessVS> VertexShader( GetGlobalShaderMap() );
-						TShaderMapRef<FVisualizeDistanceFieldUpsamplePS> PixelShader( GetGlobalShaderMap() );
+						TShaderMapRef<FPostProcessVS> VertexShader( View.ShaderMap );
+						TShaderMapRef<FVisualizeDistanceFieldUpsamplePS> PixelShader( View.ShaderMap );
 
 						static FGlobalBoundShaderState BoundShaderState;
 
