fixed TTP 34489 UE4: RENDERING: Having a material with Decal Blend Mode using DBuffer crashes on DX10

diff --git a/Engine/Shaders/DeferredShadingCommon.usf b/Engine/Shaders/DeferredShadingCommon.usf
--- a/Engine/Shaders/DeferredShadingCommon.usf
+++ b/Engine/Shaders/DeferredShadingCommon.usf
@@ -502,6 +502,19 @@ void EncodeDBufferData(FGBufferData GBufferData, float3 MultiOpacity,
 
 	// UNORM 2 channel
 	DBufferC = float4(GBufferData.Roughness, 0, 0, MultiOpacity.z);
+
+	// no need to output any data if the channel isn't used anyway (for DX10 we don't adjust the FrameBufferBlend state - it doesn't support different blend modes for different MRT)
+	{
+#if !MATERIAL_DBUFFERA
+		DBufferA = 0;
+#endif
+#if !MATERIAL_DBUFFERB
+		DBufferB = 0;
+#endif
+#if !MATERIAL_DBUFFERC
+		DBufferC = 0;
+#endif
+	}
 }
 
 float3 ExtractSubsurfaceColor(FGBufferData BufferData)
diff --git a/Engine/Source/Runtime/Engine/Private/Materials/MaterialShared.cpp b/Engine/Source/Runtime/Engine/Private/Materials/MaterialShared.cpp
--- a/Engine/Source/Runtime/Engine/Private/Materials/MaterialShared.cpp
+++ b/Engine/Source/Runtime/Engine/Private/Materials/MaterialShared.cpp
@@ -1031,6 +1031,31 @@ FMaterial::~FMaterial()
 	FMaterialShaderMap::RemovePendingMaterial(this);
 }
 
+// could be more to a more central DBuffer file
+// @return e.g. 1+2+4 means DBufferA(1) + DBufferB(2) + DBufferC(4) is used
+static uint8 ComputeDBufferMRTMask(EDecalBlendMode DecalBlendMode)
+{
+	switch(DecalBlendMode)
+	{
+		case DBM_DBuffer_ColorNormalRoughness:
+			return 1 + 2 + 4;
+		case DBM_DBuffer_Color:
+			return 1;
+		case DBM_DBuffer_ColorNormal:
+			return 1 + 2;
+		case DBM_DBuffer_ColorRoughness:
+			return 1 + 4;
+		case DBM_DBuffer_Normal:
+			return 2;
+		case DBM_DBuffer_NormalRoughness:
+			return 2 + 4;
+		case DBM_DBuffer_Roughness:
+			return 4;
+	}
+
+	return 0;
+}
+
 /** Populates OutEnvironment with defines needed to compile shaders for this material. */
 void FMaterial::SetupMaterialEnvironment(
 	EShaderPlatform Platform,
@@ -1138,6 +1163,16 @@ void FMaterial::SetupMaterialEnvironment(
 	OutEnvironment.SetDefine(TEXT("MATERIAL_USE_LM_DIRECTIONALITY"), UseLmDirectionality() ? TEXT("1") : TEXT("0"));
 	OutEnvironment.SetDefine(TEXT("MATERIAL_INJECT_EMISSIVE_INTO_LPV"), ShouldInjectEmissiveIntoLPV() ? TEXT("1") : TEXT("0"));
 
+	{
+		auto DecalBlendMode = (EDecalBlendMode)GetDecalBlendMode();
+
+		uint8 bDBufferMask = ComputeDBufferMRTMask(DecalBlendMode);
+
+		OutEnvironment.SetDefine(TEXT("MATERIAL_DBUFFERA"), (bDBufferMask & 0x1) != 0);
+		OutEnvironment.SetDefine(TEXT("MATERIAL_DBUFFERB"), (bDBufferMask & 0x2) != 0);
+		OutEnvironment.SetDefine(TEXT("MATERIAL_DBUFFERC"), (bDBufferMask & 0x4) != 0);
+	}
+
 	switch(GetShadingModel())
 	{
 		case MSM_Unlit:				OutEnvironment.SetDefine(TEXT("MATERIAL_SHADINGMODEL_UNLIT"),				TEXT("1")); break;
diff --git a/Engine/Source/Runtime/Renderer/Private/CompositionLighting/PostProcessDeferredDecals.cpp b/Engine/Source/Runtime/Renderer/Private/CompositionLighting/PostProcessDeferredDecals.cpp
--- a/Engine/Source/Runtime/Renderer/Private/CompositionLighting/PostProcessDeferredDecals.cpp
+++ b/Engine/Source/Runtime/Renderer/Private/CompositionLighting/PostProcessDeferredDecals.cpp
@@ -141,74 +141,95 @@ void SetDecalBlendState(FRHICommandList& RHICmdList, const ERHIFeatureLevel::Typ
 {
 	if(RenderStage == 0)
 	{
-		// todo if(SMFeatureLevel == ERHIFeatureLevel::SM4)
-		switch(DecalBlendMode)
-		{
-		case DBM_DBuffer_ColorNormalRoughness:
-			RHICmdList.SetBlendState( TStaticBlendState< 
-				CW_RGBA, BO_Add, BF_SourceAlpha, BF_InverseSourceAlpha,			BO_Add,BF_Zero,BF_InverseSourceAlpha,
-				CW_RGBA, BO_Add, BF_SourceAlpha, BF_InverseSourceAlpha,			BO_Add,BF_Zero,BF_InverseSourceAlpha,
-				CW_RGBA, BO_Add, BF_SourceAlpha, BF_InverseSourceAlpha,			BO_Add,BF_Zero,BF_InverseSourceAlpha >::GetRHI() );		
-			break;
+		// before base pass (for DBuffer decals)
 
-		case DBM_DBuffer_Color:
-			// we can optimize using less MRT later
-			RHICmdList.SetBlendState( TStaticBlendState< 
-				CW_RGBA, BO_Add, BF_SourceAlpha, BF_InverseSourceAlpha,			BO_Add,BF_Zero,BF_InverseSourceAlpha,
-				CW_RGBA, BO_Add, BF_Zero, BF_One,								BO_Add,BF_Zero,BF_One,
-				CW_RGBA, BO_Add, BF_Zero, BF_One,								BO_Add,BF_Zero,BF_One>::GetRHI() );		
-			break;
-
-		case DBM_DBuffer_ColorNormal:
-			// we can optimize using less MRT later
-			RHICmdList.SetBlendState( TStaticBlendState< 
-				CW_RGBA, BO_Add, BF_SourceAlpha, BF_InverseSourceAlpha,			BO_Add,BF_Zero,BF_InverseSourceAlpha,
-				CW_RGBA, BO_Add, BF_SourceAlpha, BF_InverseSourceAlpha,			BO_Add,BF_Zero,BF_InverseSourceAlpha,
-				CW_RGBA, BO_Add, BF_Zero, BF_One,								BO_Add,BF_Zero,BF_One >::GetRHI() );		
-			break;
-
-		case DBM_DBuffer_ColorRoughness:
-			// we can optimize using less MRT later
-			RHICmdList.SetBlendState( TStaticBlendState< 
-				CW_RGBA, BO_Add, BF_SourceAlpha, BF_InverseSourceAlpha,			BO_Add,BF_Zero,BF_InverseSourceAlpha,
-				CW_RGBA, BO_Add, BF_Zero, BF_One,								BO_Add,BF_Zero,BF_One,
-				CW_RGBA, BO_Add, BF_SourceAlpha, BF_InverseSourceAlpha,			BO_Add,BF_Zero,BF_InverseSourceAlpha >::GetRHI() );		
-			break;
+		if(SMFeatureLevel == ERHIFeatureLevel::SM4)
+		{
+			// DX10 doesn't support masking/using different blend modes per MRT.
+			// We set the opacity in the shader to 0 so we can use the same frame buffer blend.
 
-		case DBM_DBuffer_Normal:
-			// we can optimize using less MRT later
 			RHICmdList.SetBlendState( TStaticBlendState< 
-				CW_RGBA, BO_Add, BF_Zero, BF_One,								BO_Add,BF_Zero,BF_One,
 				CW_RGBA, BO_Add, BF_SourceAlpha, BF_InverseSourceAlpha,			BO_Add,BF_Zero,BF_InverseSourceAlpha,
-				CW_RGBA, BO_Add, BF_Zero, BF_One,								BO_Add,BF_Zero,BF_One>::GetRHI() );		
-			break;
-
-		case DBM_DBuffer_NormalRoughness:
-			// we can optimize using less MRT later
-			RHICmdList.SetBlendState( TStaticBlendState< 
-				CW_RGBA, BO_Add, BF_Zero, BF_One,								BO_Add,BF_Zero,BF_One,
 				CW_RGBA, BO_Add, BF_SourceAlpha, BF_InverseSourceAlpha,			BO_Add,BF_Zero,BF_InverseSourceAlpha,
-				CW_RGBA, BO_Add, BF_SourceAlpha, BF_InverseSourceAlpha,			BO_Add,BF_Zero,BF_InverseSourceAlpha >::GetRHI() );		
-			break;
-
-		case DBM_DBuffer_Roughness:
-			// we can optimize using less MRT later
-			RHICmdList.SetBlendState( TStaticBlendState< 
-				CW_RGBA, BO_Add, BF_Zero, BF_One,								BO_Add,BF_Zero,BF_One,
-				CW_RGBA, BO_Add, BF_Zero, BF_One,								BO_Add,BF_Zero,BF_One,
-				CW_RGBA, BO_Add, BF_SourceAlpha, BF_InverseSourceAlpha,			BO_Add,BF_Zero,BF_InverseSourceAlpha >::GetRHI() );		
-			break;
+				CW_RGBA, BO_Add, BF_SourceAlpha, BF_InverseSourceAlpha,			BO_Add,BF_Zero,BF_InverseSourceAlpha >::GetRHI() );
+			return;
+		}
+		else
+		{
+			// see DX10 comment above
+			// As we set the opacity in the shader we don't need to set different frame buffer blend modes but we like to hint to the driver that we
+			// don't need to output there. We also could replace this with many SetRenderTarget calls but it might be slower (needs to be tested).
 
-		default:
-			// the decal type should not be rendered in this pass - internal error
-			check(0);	
-			break;
+			switch(DecalBlendMode)
+			{
+			case DBM_DBuffer_ColorNormalRoughness:
+				RHICmdList.SetBlendState( TStaticBlendState< 
+					CW_RGBA, BO_Add, BF_SourceAlpha, BF_InverseSourceAlpha,			BO_Add,BF_Zero,BF_InverseSourceAlpha,
+					CW_RGBA, BO_Add, BF_SourceAlpha, BF_InverseSourceAlpha,			BO_Add,BF_Zero,BF_InverseSourceAlpha,
+					CW_RGBA, BO_Add, BF_SourceAlpha, BF_InverseSourceAlpha,			BO_Add,BF_Zero,BF_InverseSourceAlpha >::GetRHI() );		
+				break;
+
+			case DBM_DBuffer_Color:
+				// we can optimize using less MRT later
+				RHICmdList.SetBlendState( TStaticBlendState< 
+					CW_RGBA, BO_Add, BF_SourceAlpha, BF_InverseSourceAlpha,			BO_Add,BF_Zero,BF_InverseSourceAlpha,
+					CW_RGBA, BO_Add, BF_Zero, BF_One,								BO_Add,BF_Zero,BF_One,
+					CW_RGBA, BO_Add, BF_Zero, BF_One,								BO_Add,BF_Zero,BF_One>::GetRHI() );		
+				break;
+
+			case DBM_DBuffer_ColorNormal:
+				// we can optimize using less MRT later
+				RHICmdList.SetBlendState( TStaticBlendState< 
+					CW_RGBA, BO_Add, BF_SourceAlpha, BF_InverseSourceAlpha,			BO_Add,BF_Zero,BF_InverseSourceAlpha,
+					CW_RGBA, BO_Add, BF_SourceAlpha, BF_InverseSourceAlpha,			BO_Add,BF_Zero,BF_InverseSourceAlpha,
+					CW_RGBA, BO_Add, BF_Zero, BF_One,								BO_Add,BF_Zero,BF_One >::GetRHI() );		
+				break;
+
+			case DBM_DBuffer_ColorRoughness:
+				// we can optimize using less MRT later
+				RHICmdList.SetBlendState( TStaticBlendState< 
+					CW_RGBA, BO_Add, BF_SourceAlpha, BF_InverseSourceAlpha,			BO_Add,BF_Zero,BF_InverseSourceAlpha,
+					CW_RGBA, BO_Add, BF_Zero, BF_One,								BO_Add,BF_Zero,BF_One,
+					CW_RGBA, BO_Add, BF_SourceAlpha, BF_InverseSourceAlpha,			BO_Add,BF_Zero,BF_InverseSourceAlpha >::GetRHI() );		
+				break;
+
+			case DBM_DBuffer_Normal:
+				// we can optimize using less MRT later
+				RHICmdList.SetBlendState( TStaticBlendState< 
+					CW_RGBA, BO_Add, BF_Zero, BF_One,								BO_Add,BF_Zero,BF_One,
+					CW_RGBA, BO_Add, BF_SourceAlpha, BF_InverseSourceAlpha,			BO_Add,BF_Zero,BF_InverseSourceAlpha,
+					CW_RGBA, BO_Add, BF_Zero, BF_One,								BO_Add,BF_Zero,BF_One>::GetRHI() );		
+				break;
+
+			case DBM_DBuffer_NormalRoughness:
+				// we can optimize using less MRT later
+				RHICmdList.SetBlendState( TStaticBlendState< 
+					CW_RGBA, BO_Add, BF_Zero, BF_One,								BO_Add,BF_Zero,BF_One,
+					CW_RGBA, BO_Add, BF_SourceAlpha, BF_InverseSourceAlpha,			BO_Add,BF_Zero,BF_InverseSourceAlpha,
+					CW_RGBA, BO_Add, BF_SourceAlpha, BF_InverseSourceAlpha,			BO_Add,BF_Zero,BF_InverseSourceAlpha >::GetRHI() );		
+				break;
+
+			case DBM_DBuffer_Roughness:
+				// we can optimize using less MRT later
+				RHICmdList.SetBlendState( TStaticBlendState< 
+					CW_RGBA, BO_Add, BF_Zero, BF_One,								BO_Add,BF_Zero,BF_One,
+					CW_RGBA, BO_Add, BF_Zero, BF_One,								BO_Add,BF_Zero,BF_One,
+					CW_RGBA, BO_Add, BF_SourceAlpha, BF_InverseSourceAlpha,			BO_Add,BF_Zero,BF_InverseSourceAlpha >::GetRHI() );		
+				break;
+
+			default:
+				// the decal type should not be rendered in this pass - internal error
+				check(0);	
+				break;
+			}
 		}
 
 		return;
 	}
 	else
 	{
+		// before lighting (for non DBuffer decals)
+
 		switch(DecalBlendMode)
 		{
 		case DBM_Translucent:
