Fixed sky lighting past the max AO view distance

diff --git a/Engine/Shaders/DistanceFieldSurfaceCacheLighting.usf b/Engine/Shaders/DistanceFieldSurfaceCacheLighting.usf
--- a/Engine/Shaders/DistanceFieldSurfaceCacheLighting.usf
+++ b/Engine/Shaders/DistanceFieldSurfaceCacheLighting.usf
@@ -393,8 +393,8 @@ void AOCombinePS2(
 	in float4 SVPos : SV_POSITION,
 	out float4 OutColor : SV_Target0)
 {
-	float2 ScreenUV = float2((((uint2)SVPos.xy) * DOWNSAMPLE_FACTOR + View.ViewRectMin.xy + .5f) * View.ViewSizeAndSceneTexelSize.zw);
-	float SceneDepth = CalcSceneDepth(ScreenUV);
+	float4 DistanceFieldNormalValue = Texture2DSampleLevel(DistanceFieldNormalTexture, DistanceFieldNormalSampler, UVAndScreenPos.xy, 0);
+	float SceneDepth = DistanceFieldNormalValue.w;
 
 #define VISUALIZE_ACCUMULATED_WEIGHTS 0
 #if VISUALIZE_ACCUMULATED_WEIGHTS
@@ -418,13 +418,12 @@ void AOCombinePS2(
 		OutColor.a = -SceneDepth;
 	}
 
-	float3 DefaultVisibility = float3(0, 0, 1);
-
-	// Fade to unoccluded in the distance
-	OutColor.rgb = lerp(DefaultVisibility, OutColor.rgb, saturate(.003f * (AOMaxViewDistance - SceneDepth)));
-
+	BRANCH
 	if (SceneDepth > AOMaxViewDistance)
 	{
+		// Fade to unoccluded in the distance
+		OutColor.rgb = lerp(DistanceFieldNormalValue.xyz, OutColor.rgb, saturate(.003f * (AOMaxViewDistance - SceneDepth)));
+
 		OutColor.a = abs(OutColor.a);
 	}
 
diff --git a/Engine/Source/Runtime/Engine/Private/DistanceFieldAtlas.cpp b/Engine/Source/Runtime/Engine/Private/DistanceFieldAtlas.cpp
--- a/Engine/Source/Runtime/Engine/Private/DistanceFieldAtlas.cpp
+++ b/Engine/Source/Runtime/Engine/Private/DistanceFieldAtlas.cpp
@@ -243,7 +243,7 @@ private:
 	/** The runnable thread */
 	TScopedPointer<FRunnableThread> Thread;
 
-	TScopedPointer<class FQueuedThreadPool> WorkerThreadPool;
+	TScopedPointer<FQueuedThreadPool> WorkerThreadPool;
 
 	volatile bool bIsRunning;
 	volatile bool bForceFinish;
diff --git a/Engine/Source/Runtime/Renderer/Private/DistanceFieldSurfaceCacheLighting.cpp b/Engine/Source/Runtime/Renderer/Private/DistanceFieldSurfaceCacheLighting.cpp
--- a/Engine/Source/Runtime/Renderer/Private/DistanceFieldSurfaceCacheLighting.cpp
+++ b/Engine/Source/Runtime/Renderer/Private/DistanceFieldSurfaceCacheLighting.cpp
@@ -1856,9 +1856,11 @@ public:
 		AOLevelParameters.Bind(Initializer.ParameterMap);
 		BentNormalAOTexture3.Bind(Initializer.ParameterMap,TEXT("BentNormalAOTexture3"));
 		BentNormalAOSampler3.Bind(Initializer.ParameterMap,TEXT("BentNormalAOSampler3"));
+		DistanceFieldNormalTexture.Bind(Initializer.ParameterMap, TEXT("DistanceFieldNormalTexture"));
+		DistanceFieldNormalSampler.Bind(Initializer.ParameterMap, TEXT("DistanceFieldNormalSampler"));
 	}
 
-	void SetParameters(FRHICommandList& RHICmdList, const FSceneView& View, FSceneRenderTargetItem& InBentNormalTexture, const FDistanceFieldAOParameters& Parameters)
+	void SetParameters(FRHICommandList& RHICmdList, const FSceneView& View, FSceneRenderTargetItem& InBentNormalTexture, FSceneRenderTargetItem& DistanceFieldNormal, const FDistanceFieldAOParameters& Parameters)
 	{
 		const FPixelShaderRHIParamRef ShaderRHI = GetPixelShader();
 
@@ -1868,6 +1870,15 @@ public:
 		AOLevelParameters.Set(RHICmdList, ShaderRHI, View, GAODownsampleFactor);
 
 		SetTextureParameter(RHICmdList, ShaderRHI, BentNormalAOTexture3, BentNormalAOSampler3, TStaticSamplerState<SF_Point>::GetRHI(), InBentNormalTexture.ShaderResourceTexture);
+
+		SetTextureParameter(
+			RHICmdList,
+			ShaderRHI,
+			DistanceFieldNormalTexture,
+			DistanceFieldNormalSampler,
+			TStaticSamplerState<SF_Point,AM_Clamp,AM_Clamp,AM_Clamp>::GetRHI(),
+			DistanceFieldNormal.ShaderResourceTexture
+			);
 	}
 	// FShader interface.
 	virtual bool Serialize(FArchive& Ar)
@@ -1878,6 +1889,8 @@ public:
 		Ar << AOLevelParameters;
 		Ar << BentNormalAOTexture3;
 		Ar << BentNormalAOSampler3;
+		Ar << DistanceFieldNormalTexture;
+		Ar << DistanceFieldNormalSampler;
 		return bShaderHasOutdatedParameters;
 	}
 
@@ -1888,6 +1901,8 @@ private:
 	FAOLevelParameters AOLevelParameters;
 	FShaderResourceParameter BentNormalAOTexture3;
 	FShaderResourceParameter BentNormalAOSampler3;
+	FShaderResourceParameter DistanceFieldNormalTexture;
+	FShaderResourceParameter DistanceFieldNormalSampler;
 };
 
 IMPLEMENT_SHADER_TYPE(,FDistanceFieldAOCombinePS2,TEXT("DistanceFieldSurfaceCacheLighting"),TEXT("AOCombinePS2"),SF_Pixel);
@@ -2124,7 +2139,13 @@ void UpdateHistory(
 	}
 }
 
-void PostProcessBentNormalAO(FRHICommandList& RHICmdList, const FDistanceFieldAOParameters& Parameters, TArray<FViewInfo>& Views, FSceneRenderTargetItem& IrradianceCacheInterpolation, TRefCountPtr<IPooledRenderTarget>& AOOutput)
+void PostProcessBentNormalAO(
+	FRHICommandList& RHICmdList, 
+	const FDistanceFieldAOParameters& Parameters, 
+	TArray<FViewInfo>& Views, 
+	FSceneRenderTargetItem& IrradianceCacheInterpolation, 
+	FSceneRenderTargetItem& DistanceFieldNormal,
+	TRefCountPtr<IPooledRenderTarget>& AOOutput)
 {
 	TRefCountPtr<IPooledRenderTarget> DistanceFieldAOBentNormal;
 	AllocateOrReuseAORenderTarget(DistanceFieldAOBentNormal, TEXT("DistanceFieldBentNormalAO"));
@@ -2160,7 +2181,7 @@ void PostProcessBentNormalAO(FRHICommandList& RHICmdList, const FDistanceFieldAO
 			
 			SetGlobalBoundShaderState(RHICmdList, View.GetFeatureLevel(), BoundShaderState, GFilterVertexDeclaration.VertexDeclarationRHI, *VertexShader, *PixelShader);
 
-			PixelShader->SetParameters(RHICmdList, View, IrradianceCacheInterpolation, Parameters);
+			PixelShader->SetParameters(RHICmdList, View, IrradianceCacheInterpolation, DistanceFieldNormal, Parameters);
 
 			DrawRectangle( 
 				RHICmdList,
@@ -3022,7 +3043,7 @@ bool FDeferredShadingSceneRenderer::RenderDistanceFieldAOSurfaceCache(FRHIComman
 
 				// Post process the AO to cover over artifacts
 				TRefCountPtr<IPooledRenderTarget> AOOutput;
-				PostProcessBentNormalAO(RHICmdList, Parameters, Views, IrradianceCacheAccumulation->GetRenderTargetItem(), AOOutput);
+				PostProcessBentNormalAO(RHICmdList, Parameters, Views, IrradianceCacheAccumulation->GetRenderTargetItem(), DistanceFieldNormal->GetRenderTargetItem(), AOOutput);
 
 				GRenderTargetPool.VisualizeTexture.SetCheckPoint(RHICmdList, AOOutput);
 
