fix from TimL for utilizing the correct resolve shader on mobile
#ue4

diff --git a/Engine/Source/Runtime/Renderer/Private/PostProcess/SceneRenderTargets.cpp b/Engine/Source/Runtime/Renderer/Private/PostProcess/SceneRenderTargets.cpp
--- a/Engine/Source/Runtime/Renderer/Private/PostProcess/SceneRenderTargets.cpp
+++ b/Engine/Source/Runtime/Renderer/Private/PostProcess/SceneRenderTargets.cpp
@@ -645,7 +645,7 @@ void FSceneRenderTargets::ResolveSceneColor(FRHICommandList& RHICmdList, const F
 
 	uint32 samples = GSceneRenderTargets.SceneColor->GetDesc().NumSamples;
 
-	if(samples <= 1)
+	if(samples <= 1 || GRHIShaderPlatform == SP_METAL)
 	{
 		RHICmdList.CopyToResolveTarget(GetSceneColorSurface(), GetSceneColorTexture(), true, FResolveParams(ResolveRect));
 	}
