Fix for non-unity compile error

diff --git a/Engine/Source/Runtime/Renderer/Private/SceneVisibility.cpp b/Engine/Source/Runtime/Renderer/Private/SceneVisibility.cpp
--- a/Engine/Source/Runtime/Renderer/Private/SceneVisibility.cpp
+++ b/Engine/Source/Runtime/Renderer/Private/SceneVisibility.cpp
@@ -10,6 +10,7 @@
 #include "FXSystem.h"
 #include "../../Engine/Private/SkeletalRenderGPUSkin.h"		// GPrevPerBoneMotionBlur
 #include "SceneUtils.h"
+#include "PostProcessing.h"
 
 /*------------------------------------------------------------------------------
 	Globals
