fixed LightFunctions getting disabled with low FOV
 see https://udn.unrealengine.com/questions/218679/low-fov-settings-makes-our-light-function-material-1.html

diff --git a/Engine/Source/Runtime/Renderer/Private/LightFunctionRendering.cpp b/Engine/Source/Runtime/Renderer/Private/LightFunctionRendering.cpp
--- a/Engine/Source/Runtime/Renderer/Private/LightFunctionRendering.cpp
+++ b/Engine/Source/Runtime/Renderer/Private/LightFunctionRendering.cpp
@@ -150,6 +150,8 @@ IMPLEMENT_MATERIAL_SHADER_TYPE(,FLightFunctionPS,TEXT("LightFunctionPixelShader"
 /** Returns a fade fraction for a light function and a given view based on the appropriate fade settings. */
 static float GetLightFunctionFadeFraction(const FViewInfo& View, FSphere LightBounds)
 {
+	extern float CalculateShadowFadeAlpha(float MaxUnclampedResolution, int32 ShadowFadeResolution, int32 MinShadowResolution);
+
 	// Override the global settings with the light's settings if the light has them specified
 	static auto CVarMinShadowResolution = IConsoleManager::Get().FindTConsoleVariableDataInt(TEXT("r.Shadow.MinResolution"));
 	static auto CVarShadowFadeResolution = IConsoleManager::Get().FindTConsoleVariableDataInt(TEXT("r.Shadow.FadeResolution"));
@@ -169,8 +171,7 @@ static float GetLightFunctionFadeFraction(const FViewInfo& View, FSphere LightBo
 		FMath::Max(ScreenPosition.W, 1.0f);
 
 	static auto CVarShadowTexelsPerPixel = IConsoleManager::Get().FindTConsoleVariableDataFloat(TEXT("r.Shadow.TexelsPerPixel"));
-	const uint32 UnclampedResolution = FMath::TruncToInt(ScreenRadius * CVarShadowTexelsPerPixel->GetValueOnRenderThread());
-	extern float CalculateShadowFadeAlpha(int32 MaxUnclampedResolution, int32 ShadowFadeResolution, int32 MinShadowResolution);
+	const float UnclampedResolution = ScreenRadius * CVarShadowTexelsPerPixel->GetValueOnRenderThread();
 	const float ResolutionFadeAlpha = CalculateShadowFadeAlpha(UnclampedResolution, ShadowFadeResolution, MinShadowResolution);
 	return ResolutionFadeAlpha;
 }
diff --git a/Engine/Source/Runtime/Renderer/Private/ShadowSetup.cpp b/Engine/Source/Runtime/Renderer/Private/ShadowSetup.cpp
--- a/Engine/Source/Runtime/Renderer/Private/ShadowSetup.cpp
+++ b/Engine/Source/Runtime/Renderer/Private/ShadowSetup.cpp
@@ -143,8 +143,12 @@ static TAutoConsoleVariable<int32> CVarUseConservativeShadowBounds(
  *
  * @return	fade value between 0 and 1
  */
-float CalculateShadowFadeAlpha(int32 MaxUnclampedResolution, int32 ShadowFadeResolution, int32 MinShadowResolution)
+float CalculateShadowFadeAlpha(float MaxUnclampedResolution, int32 ShadowFadeResolution, int32 MinShadowResolution)
 {
+	check(MaxUnclampedResolution >= 0);
+	check(ShadowFadeResolution >= 0);
+	check(MinShadowResolution >= 0);
+
 	float FadeAlpha = 0.0f;
 	// Shadow size is above fading resolution.
 	if (MaxUnclampedResolution > ShadowFadeResolution)
@@ -1161,7 +1165,7 @@ void FDeferredShadingSceneRenderer::CreatePerObjectProjectedShadow(
 
 	// Compute the maximum resolution required for the shadow by any view. Also keep track of the unclamped resolution for fading.
 	uint32 MaxDesiredResolution = 0;
-	uint32 MaxUnclampedResolution	= 0;
+	float MaxUnclampedResolution = 0;
 	float MaxScreenPercent = 0;
 	TArray<float, TInlineAllocator<2> > ResolutionFadeAlphas;
 	TArray<float, TInlineAllocator<2> > ResolutionPreShadowFadeAlphas;
@@ -1195,7 +1199,7 @@ void FDeferredShadingSceneRenderer::CreatePerObjectProjectedShadow(
 		MaxScreenPercent = FMath::Max(MaxScreenPercent, ScreenPercent);
 
 		// Determine the amount of shadow buffer resolution needed for this view.
-		const uint32 UnclampedResolution = FMath::TruncToInt(ScreenRadius * CVarShadowTexelsPerPixel.GetValueOnRenderThread());
+		const float UnclampedResolution = ScreenRadius * CVarShadowTexelsPerPixel.GetValueOnRenderThread();
 		MaxUnclampedResolution = FMath::Max( MaxUnclampedResolution, UnclampedResolution );
 		MaxDesiredResolution = FMath::Max(
 			MaxDesiredResolution,
@@ -1211,7 +1215,7 @@ void FDeferredShadingSceneRenderer::CreatePerObjectProjectedShadow(
 		MaxResolutionFadeAlpha = FMath::Max(MaxResolutionFadeAlpha, ViewSpecificAlpha);
 		ResolutionFadeAlphas.Add(ViewSpecificAlpha);
 
-		const float ViewSpecificPreShadowAlpha = CalculateShadowFadeAlpha( FMath::TruncToInt(UnclampedResolution * CVarPreShadowResolutionFactor.GetValueOnRenderThread()), CVarPreShadowFadeResolution.GetValueOnRenderThread(), CVarMinPreShadowResolution.GetValueOnRenderThread() );
+		const float ViewSpecificPreShadowAlpha = CalculateShadowFadeAlpha( UnclampedResolution * CVarPreShadowResolutionFactor.GetValueOnRenderThread(), CVarPreShadowFadeResolution.GetValueOnRenderThread(), CVarMinPreShadowResolution.GetValueOnRenderThread() );
 		MaxResolutionPreShadowFadeAlpha = FMath::Max(MaxResolutionPreShadowFadeAlpha, ViewSpecificPreShadowAlpha);
 		ResolutionPreShadowFadeAlphas.Add(ViewSpecificPreShadowAlpha);
 	}
@@ -1429,8 +1433,8 @@ void FDeferredShadingSceneRenderer::CreateWholeSceneProjectedShadow(FLightSceneI
 		const int32 ShadowFadeResolution = CVarShadowFadeResolution.GetValueOnRenderThread();
 
 		// Compute the maximum resolution required for the shadow by any view. Also keep track of the unclamped resolution for fading.
-		uint32 MaxDesiredResolution = 0;
-		uint32 MaxUnclampedResolution	= 0;
+		float MaxDesiredResolution = 0;
+		float MaxUnclampedResolution = 0;
 		TArray<float, TInlineAllocator<2> > FadeAlphas;
 		float MaxFadeAlpha = 0;
 		bool bStaticSceneOnly = false;
@@ -1453,13 +1457,13 @@ void FDeferredShadingSceneRenderer::CreateWholeSceneProjectedShadow(FLightSceneI
 				FMath::Max(ScreenPosition.W,1.0f);
 
 			// Determine the amount of shadow buffer resolution needed for this view.
-			const uint32 UnclampedResolution = FMath::TruncToInt(ScreenRadius * CVarShadowTexelsPerPixel.GetValueOnRenderThread());
+			const float UnclampedResolution = ScreenRadius * CVarShadowTexelsPerPixel.GetValueOnRenderThread();
 			MaxUnclampedResolution = FMath::Max( MaxUnclampedResolution, UnclampedResolution );
 			MaxDesiredResolution = FMath::Max(
 				MaxDesiredResolution,
-				FMath::Clamp<uint32>(
+				FMath::Clamp<float>(
 					UnclampedResolution,
-					FMath::Min<int32>(MinShadowResolution,ShadowBufferResolution.X - EffectiveDoubleShadowBorder),
+					FMath::Min<float>(MinShadowResolution,ShadowBufferResolution.X - EffectiveDoubleShadowBorder),
 					MaxShadowResolution
 					)
 				);
