Add PolicyContext to avoid global variables during shadow rendering

diff --git a/Engine/Source/Runtime/Core/Public/Async/TaskGraphInterfaces.h b/Engine/Source/Runtime/Core/Public/Async/TaskGraphInterfaces.h
--- a/Engine/Source/Runtime/Core/Public/Async/TaskGraphInterfaces.h
+++ b/Engine/Source/Runtime/Core/Public/Async/TaskGraphInterfaces.h
@@ -577,6 +577,13 @@ public:
 			new ((void *)&Owner->TaskStorage) TTask(Arg1, Arg2, Arg3, Arg4, Arg5, Arg6, Arg7, Arg8);
 			return Owner->Setup(Prerequisites, CurrentThreadIfKnown);
 		}
+		/** Passthrough internal task constructor and dispatch. Note! Generally speaking references will not pass through; use pointers */
+		template<typename T1, typename T2, typename T3, typename T4, typename T5, typename T6, typename T7, typename T8, typename T9>
+		FGraphEventRef ConstructAndDispatchWhenReady(const T1& Arg1, const T2& Arg2, const T3& Arg3, const T4& Arg4, const T5& Arg5, const T6& Arg6, const T7& Arg7, const T8& Arg8, const T9& Arg9)
+		{
+			new ((void *)&Owner->TaskStorage) TTask(Arg1, Arg2, Arg3, Arg4, Arg5, Arg6, Arg7, Arg8, Arg9);
+			return Owner->Setup(Prerequisites, CurrentThreadIfKnown);
+		}
 	private:
 		friend class TGraphTask;
 
diff --git a/Engine/Source/Runtime/Renderer/Private/BasePassRendering.cpp b/Engine/Source/Runtime/Renderer/Private/BasePassRendering.cpp
--- a/Engine/Source/Runtime/Renderer/Private/BasePassRendering.cpp
+++ b/Engine/Source/Runtime/Renderer/Private/BasePassRendering.cpp
@@ -300,7 +300,7 @@ public:
 			Parameters.bEditorCompositeDepthTest
 			);
 		RHICmdList.BuildAndSetLocalBoundShaderState(DrawingPolicy.GetBoundShaderStateInput(View.GetFeatureLevel()));
-		DrawingPolicy.SetSharedState(RHICmdList, &View);
+		DrawingPolicy.SetSharedState(RHICmdList, &View, typename TBasePassDrawingPolicy<LightMapPolicyType>::ContextDataType());
 
 		for( int32 BatchElementIndex = 0, Num = Parameters.Mesh.Elements.Num(); BatchElementIndex < Num; BatchElementIndex++ )
 		{
@@ -311,7 +311,8 @@ public:
 				Parameters.Mesh,
 				BatchElementIndex,
 				bBackFace,
-				typename TBasePassDrawingPolicy<LightMapPolicyType>::ElementDataType(LightMapElementData)
+				typename TBasePassDrawingPolicy<LightMapPolicyType>::ElementDataType(LightMapElementData),
+				typename TBasePassDrawingPolicy<LightMapPolicyType>::ContextDataType()
 				);
 			DrawingPolicy.DrawMesh(RHICmdList, Parameters.Mesh, BatchElementIndex);
 		}
diff --git a/Engine/Source/Runtime/Renderer/Private/BasePassRendering.h b/Engine/Source/Runtime/Renderer/Private/BasePassRendering.h
--- a/Engine/Source/Runtime/Renderer/Private/BasePassRendering.h
+++ b/Engine/Source/Runtime/Renderer/Private/BasePassRendering.h
@@ -553,7 +553,7 @@ public:
 			LightMapPolicy == Other.LightMapPolicy;
 	}
 
-	void SetSharedState(FRHICommandList& RHICmdList, const FViewInfo* View) const
+	void SetSharedState(FRHICommandList& RHICmdList, const FViewInfo* View, const ContextDataType PolicyContext) const
 	{
 		// Set the light-map policy.
 		LightMapPolicy.Set(RHICmdList, VertexShader,bOverrideWithShaderComplexity ? NULL : PixelShader,VertexShader,PixelShader,VertexFactory,MaterialRenderProxy,View);
@@ -650,7 +650,8 @@ public:
 		const FMeshBatch& Mesh,
 		int32 BatchElementIndex,
 		bool bBackFace,
-		const ElementDataType& ElementData
+		const ElementDataType& ElementData,
+		const ContextDataType PolicyContext
 		) const
 	{
 		// Set the light-map policy's mesh-specific settings.
@@ -697,7 +698,7 @@ public:
 			PixelShader->SetMesh(RHICmdList, VertexFactory,View,PrimitiveSceneProxy,BatchElement,BlendMode);
 		}
 
-		FMeshDrawingPolicy::SetMeshRenderState(RHICmdList, View,PrimitiveSceneProxy,Mesh,BatchElementIndex,bBackFace,FMeshDrawingPolicy::ElementDataType());
+		FMeshDrawingPolicy::SetMeshRenderState(RHICmdList, View,PrimitiveSceneProxy,Mesh,BatchElementIndex,bBackFace,FMeshDrawingPolicy::ElementDataType(),PolicyContext);
 	}
 
 	friend int32 CompareDrawingPolicy(const TBasePassDrawingPolicy& A,const TBasePassDrawingPolicy& B)
diff --git a/Engine/Source/Runtime/Renderer/Private/DepthRendering.cpp b/Engine/Source/Runtime/Renderer/Private/DepthRendering.cpp
--- a/Engine/Source/Runtime/Renderer/Private/DepthRendering.cpp
+++ b/Engine/Source/Runtime/Renderer/Private/DepthRendering.cpp
@@ -169,7 +169,7 @@ FDepthDrawingPolicy::FDepthDrawingPolicy(
 
 }
 
-void FDepthDrawingPolicy::SetSharedState(FRHICommandList& RHICmdList, const FSceneView* View) const
+void FDepthDrawingPolicy::SetSharedState(FRHICommandList& RHICmdList, const FSceneView* View, const ContextDataType PolicyContext) const
 {
 	// Set the depth-only shader parameters for the material.
 	VertexShader->SetParameters(RHICmdList, MaterialRenderProxy,*MaterialResource,*View);
@@ -185,7 +185,7 @@ void FDepthDrawingPolicy::SetSharedState(FRHICommandList& RHICmdList, const FSce
 	}
 
 	// Set the shared mesh resources.
-	FMeshDrawingPolicy::DrawShared(RHICmdList, View);
+	FMeshDrawingPolicy::SetSharedState(RHICmdList, View, PolicyContext);
 }
 
 /** 
@@ -211,7 +211,8 @@ void FDepthDrawingPolicy::SetMeshRenderState(
 	const FMeshBatch& Mesh,
 	int32 BatchElementIndex,
 	bool bBackFace,
-	const ElementDataType& ElementData
+	const ElementDataType& ElementData,
+	const ContextDataType PolicyContext
 	) const
 {
 	const FMeshBatchElement& BatchElement = Mesh.Elements[BatchElementIndex];
@@ -226,7 +227,7 @@ void FDepthDrawingPolicy::SetMeshRenderState(
 	{
 		PixelShader->SetMesh(RHICmdList, VertexFactory,View,PrimitiveSceneProxy,BatchElement);
 	}
-	FMeshDrawingPolicy::SetMeshRenderState(RHICmdList, View,PrimitiveSceneProxy,Mesh,BatchElementIndex,bBackFace,ElementData);
+	FMeshDrawingPolicy::SetMeshRenderState(RHICmdList, View,PrimitiveSceneProxy,Mesh,BatchElementIndex,bBackFace,ElementData,PolicyContext);
 }
 
 int32 CompareDrawingPolicy(const FDepthDrawingPolicy& A,const FDepthDrawingPolicy& B)
@@ -254,7 +255,7 @@ FPositionOnlyDepthDrawingPolicy::FPositionOnlyDepthDrawingPolicy(
 	VertexShader = InMaterialResource.GetShader<TDepthOnlyVS<true> >(InVertexFactory->GetType());
 }
 
-void FPositionOnlyDepthDrawingPolicy::SetSharedState(FRHICommandList& RHICmdList, const FSceneView* View) const
+void FPositionOnlyDepthDrawingPolicy::SetSharedState(FRHICommandList& RHICmdList, const FSceneView* View, const ContextDataType PolicyContext) const
 {
 	// Set the depth-only shader parameters for the material.
 	VertexShader->SetParameters(RHICmdList, MaterialRenderProxy,*MaterialResource,*View);
@@ -284,11 +285,12 @@ void FPositionOnlyDepthDrawingPolicy::SetMeshRenderState(
 	const FMeshBatch& Mesh,
 	int32 BatchElementIndex,
 	bool bBackFace,
-	const ElementDataType& ElementData
+	const ElementDataType& ElementData,
+	const ContextDataType PolicyContext
 	) const
 {
 	VertexShader->SetMesh(RHICmdList, VertexFactory,View,PrimitiveSceneProxy,Mesh.Elements[BatchElementIndex]);
-	FMeshDrawingPolicy::SetMeshRenderState(RHICmdList, View,PrimitiveSceneProxy,Mesh,BatchElementIndex,bBackFace,ElementData);
+	FMeshDrawingPolicy::SetMeshRenderState(RHICmdList, View,PrimitiveSceneProxy,Mesh,BatchElementIndex,bBackFace,ElementData,PolicyContext);
 }
 
 int32 CompareDrawingPolicy(const FPositionOnlyDepthDrawingPolicy& A,const FPositionOnlyDepthDrawingPolicy& B)
@@ -399,7 +401,7 @@ bool FDepthDrawingPolicyFactory::DrawMesh(
 			const FMaterialRenderProxy* DefaultProxy = UMaterial::GetDefaultMaterial(MD_Surface)->GetRenderProxy(false);
 			FPositionOnlyDepthDrawingPolicy DrawingPolicy(Mesh.VertexFactory, DefaultProxy, *DefaultProxy->GetMaterial(View.GetFeatureLevel()), Material->IsTwoSided(), Material->IsWireframe());
 			RHICmdList.BuildAndSetLocalBoundShaderState(DrawingPolicy.GetBoundShaderStateInput(View.GetFeatureLevel()));
-			DrawingPolicy.SetSharedState(RHICmdList, &View);
+			DrawingPolicy.SetSharedState(RHICmdList, &View, FDepthDrawingPolicy::ContextDataType());
 
 			int32 BatchElementIndex = 0;
 			uint64 Mask = BatchElementMask;
@@ -407,7 +409,7 @@ bool FDepthDrawingPolicyFactory::DrawMesh(
 			{
 				if(Mask & 1)
 				{
-					DrawingPolicy.SetMeshRenderState(RHICmdList, View,PrimitiveSceneProxy,Mesh,BatchElementIndex,bBackFace,FPositionOnlyDepthDrawingPolicy::ElementDataType());
+					DrawingPolicy.SetMeshRenderState(RHICmdList, View,PrimitiveSceneProxy,Mesh,BatchElementIndex,bBackFace,FPositionOnlyDepthDrawingPolicy::ElementDataType(),FDepthDrawingPolicy::ContextDataType());
 					DrawingPolicy.DrawMesh(RHICmdList, Mesh,BatchElementIndex);
 				}
 				Mask >>= 1;
@@ -443,7 +445,7 @@ bool FDepthDrawingPolicyFactory::DrawMesh(
 
 				FDepthDrawingPolicy DrawingPolicy(Mesh.VertexFactory, MaterialRenderProxy, *MaterialRenderProxy->GetMaterial(View.GetFeatureLevel()), Material->IsTwoSided());
 				RHICmdList.BuildAndSetLocalBoundShaderState(DrawingPolicy.GetBoundShaderStateInput(View.GetFeatureLevel()));
-				DrawingPolicy.SetSharedState(RHICmdList, &View);
+				DrawingPolicy.SetSharedState(RHICmdList, &View, FDepthDrawingPolicy::ContextDataType());
 
 				int32 BatchElementIndex = 0;
 				uint64 Mask = BatchElementMask;
@@ -451,7 +453,7 @@ bool FDepthDrawingPolicyFactory::DrawMesh(
 				{
 					if(Mask & 1)
 					{
-						DrawingPolicy.SetMeshRenderState(RHICmdList, View,PrimitiveSceneProxy,Mesh,BatchElementIndex,bBackFace,FMeshDrawingPolicy::ElementDataType());
+						DrawingPolicy.SetMeshRenderState(RHICmdList, View,PrimitiveSceneProxy,Mesh,BatchElementIndex,bBackFace,FMeshDrawingPolicy::ElementDataType(),FDepthDrawingPolicy::ContextDataType());
 						DrawingPolicy.DrawMesh(RHICmdList, Mesh,BatchElementIndex);
 					}
 					Mask >>= 1;
diff --git a/Engine/Source/Runtime/Renderer/Private/DepthRendering.h b/Engine/Source/Runtime/Renderer/Private/DepthRendering.h
--- a/Engine/Source/Runtime/Renderer/Private/DepthRendering.h
+++ b/Engine/Source/Runtime/Renderer/Private/DepthRendering.h
@@ -44,7 +44,7 @@ public:
 			&& PixelShader == Other.PixelShader;
 	}
 
-	void SetSharedState(FRHICommandList& RHICmdList, const FSceneView* View) const;
+	void SetSharedState(FRHICommandList& RHICmdList, const FSceneView* View, const ContextDataType PolicyContext) const;
 
 	/** 
 	* Create bound shader state using the vertex decl from the mesh draw policy
@@ -61,7 +61,8 @@ public:
 		const FMeshBatch& Mesh,
 		int32 BatchElementIndex,
 		bool bBackFace,
-		const ElementDataType& ElementData
+		const ElementDataType& ElementData,
+		const ContextDataType PolicyContext
 		) const;
 
 	friend int32 CompareDrawingPolicy(const FDepthDrawingPolicy& A,const FDepthDrawingPolicy& B);
@@ -97,7 +98,7 @@ public:
 		return FMeshDrawingPolicy::Matches(Other) && VertexShader == Other.VertexShader;
 	}
 
-	void SetSharedState(FRHICommandList& RHICmdList, const FSceneView* View) const;
+	void SetSharedState(FRHICommandList& RHICmdList, const FSceneView* View, const ContextDataType PolicyContext) const;
 
 	/** 
 	* Create bound shader state using the vertex decl from the mesh draw policy
@@ -113,7 +114,8 @@ public:
 		const FMeshBatch& Mesh,
 		int32 BatchElementIndex,
 		bool bBackFace,
-		const ElementDataType& ElementData
+		const ElementDataType& ElementData,
+		const ContextDataType PolicyContext
 		) const;
 
 	friend int32 CompareDrawingPolicy(const FPositionOnlyDepthDrawingPolicy& A,const FPositionOnlyDepthDrawingPolicy& B);
diff --git a/Engine/Source/Runtime/Renderer/Private/DistortionRendering.cpp b/Engine/Source/Runtime/Renderer/Private/DistortionRendering.cpp
--- a/Engine/Source/Runtime/Renderer/Private/DistortionRendering.cpp
+++ b/Engine/Source/Runtime/Renderer/Private/DistortionRendering.cpp
@@ -296,7 +296,7 @@ public:
 	* @param CI - The command interface to execute the draw commands on.
 	* @param View - The view of the scene being drawn.
 	*/
-	void SetSharedState(FRHICommandList& RHICmdList, const FSceneView* View) const;
+	void SetSharedState(FRHICommandList& RHICmdList, const FSceneView* View, const ContextDataType PolicyContext) const;
 	
 	/** 
 	* Create bound shader state using the vertex decl from the mesh draw policy
@@ -318,7 +318,8 @@ public:
 		const FMeshBatch& Mesh,
 		int32 BatchElementIndex,
 		bool bBackFace,
-		const ElementDataType& ElementData
+		const ElementDataType& ElementData,
+		const ContextDataType PolicyContext
 		) const;
 
 private:
@@ -406,11 +407,12 @@ bool TDistortionMeshDrawingPolicy<DistortMeshPolicy>::Matches(
 template<class DistortMeshPolicy>
 void TDistortionMeshDrawingPolicy<DistortMeshPolicy>::SetSharedState(
 	FRHICommandList& RHICmdList, 
-	const FSceneView* View
+	const FSceneView* View,
+	const ContextDataType PolicyContext
 	) const
 {
 	// Set shared mesh resources
-	FMeshDrawingPolicy::DrawShared(RHICmdList, View);
+	FMeshDrawingPolicy::SetSharedState(RHICmdList, View, PolicyContext);
 
 	// Set the translucent shader parameters for the material instance
 	VertexShader->SetParameters(RHICmdList, VertexFactory,MaterialRenderProxy,View);
@@ -488,7 +490,8 @@ void TDistortionMeshDrawingPolicy<DistortMeshPolicy>::SetMeshRenderState(
 	const FMeshBatch& Mesh,
 	int32 BatchElementIndex,
 	bool bBackFace,
-	const ElementDataType& ElementData
+	const ElementDataType& ElementData,
+	const ContextDataType PolicyContext
 	) const
 {
 
@@ -601,10 +604,10 @@ bool TDistortionMeshDrawingPolicyFactory<DistortMeshPolicy>::DrawDynamicMesh(
 			View.Family->EngineShowFlags.ShaderComplexity
 			);
 		RHICmdList.BuildAndSetLocalBoundShaderState(DrawingPolicy.GetBoundShaderStateInput(View.GetFeatureLevel()));
-		DrawingPolicy.SetSharedState(RHICmdList, &View);
+		DrawingPolicy.SetSharedState(RHICmdList, &View, typename TDistortionMeshDrawingPolicy<DistortMeshPolicy>::ContextDataType());
 		for (int32 BatchElementIndex = 0; BatchElementIndex < Mesh.Elements.Num(); BatchElementIndex++)
 		{
-			DrawingPolicy.SetMeshRenderState(RHICmdList, View,PrimitiveSceneProxy,Mesh,BatchElementIndex,bBackFace,typename TDistortionMeshDrawingPolicy<DistortMeshPolicy>::ElementDataType());
+			DrawingPolicy.SetMeshRenderState(RHICmdList, View,PrimitiveSceneProxy,Mesh,BatchElementIndex,bBackFace,typename TDistortionMeshDrawingPolicy<DistortMeshPolicy>::ElementDataType(), typename TDistortionMeshDrawingPolicy<DistortMeshPolicy>::ContextDataType());
 			DrawingPolicy.DrawMesh(RHICmdList, Mesh,BatchElementIndex);
 		}
 
@@ -646,13 +649,16 @@ bool TDistortionMeshDrawingPolicyFactory<DistortMeshPolicy>::DrawStaticMesh(
 			View->Family->EngineShowFlags.ShaderComplexity
 			);
 		RHICmdList.BuildAndSetLocalBoundShaderState(DrawingPolicy.GetBoundShaderStateInput(View->GetFeatureLevel()));
-		DrawingPolicy.SetSharedState(RHICmdList, View);
+		DrawingPolicy.SetSharedState(RHICmdList, View, typename TDistortionMeshDrawingPolicy<DistortMeshPolicy>::ContextDataType());
 		int32 BatchElementIndex = 0;
 		do
 		{
 			if(BatchElementMask & 1)
 			{
-				DrawingPolicy.SetMeshRenderState(RHICmdList, *View,PrimitiveSceneProxy,StaticMesh,BatchElementIndex,bBackFace,typename TDistortionMeshDrawingPolicy<DistortMeshPolicy>::ElementDataType());
+				DrawingPolicy.SetMeshRenderState(RHICmdList, *View,PrimitiveSceneProxy,StaticMesh,BatchElementIndex,bBackFace,
+					typename TDistortionMeshDrawingPolicy<DistortMeshPolicy>::ElementDataType(),
+					typename TDistortionMeshDrawingPolicy<DistortMeshPolicy>::ContextDataType()
+					);
 				DrawingPolicy.DrawMesh(RHICmdList, StaticMesh, BatchElementIndex);
 			}
 			BatchElementMask >>= 1;
diff --git a/Engine/Source/Runtime/Renderer/Private/DrawingPolicy.cpp b/Engine/Source/Runtime/Renderer/Private/DrawingPolicy.cpp
--- a/Engine/Source/Runtime/Renderer/Private/DrawingPolicy.cpp
+++ b/Engine/Source/Runtime/Renderer/Private/DrawingPolicy.cpp
@@ -49,7 +49,8 @@ void FMeshDrawingPolicy::SetMeshRenderState(
 	const FMeshBatch& Mesh,
 	int32 BatchElementIndex,
 	bool bBackFace,
-	const ElementDataType& ElementData
+	const ElementDataType& ElementData,
+	const ContextDataType PolicyContext
 	) const
 {
 	EmitMeshDrawEvents(PrimitiveSceneProxy, Mesh);
@@ -125,7 +126,7 @@ void FMeshDrawingPolicy::DrawMesh(FRHICommandList& RHICmdList, const FMeshBatch&
 	}
 }
 
-void FMeshDrawingPolicy::DrawShared(FRHICommandList& RHICmdList, const FSceneView* View) const
+void FMeshDrawingPolicy::SetSharedState(FRHICommandList& RHICmdList, const FSceneView* View, const ContextDataType PolicyContext) const
 {
 	check(VertexFactory && VertexFactory->IsInitialized());
 	VertexFactory->Set(RHICmdList);
diff --git a/Engine/Source/Runtime/Renderer/Private/DrawingPolicy.h b/Engine/Source/Runtime/Renderer/Private/DrawingPolicy.h
--- a/Engine/Source/Runtime/Renderer/Private/DrawingPolicy.h
+++ b/Engine/Source/Runtime/Renderer/Private/DrawingPolicy.h
@@ -26,8 +26,12 @@
 class FMeshDrawingPolicy
 {
 public:
+	/** Per-element data required by the drawing policy that static mesh draw lists will cache. */
 	struct ElementDataType {};
 
+	/** Context data required by the drawing policy that is not known when caching policies in static mesh draw lists. */
+	struct ContextDataType {};
+
 	FMeshDrawingPolicy(
 		const FVertexFactory* InVertexFactory,
 		const FMaterialRenderProxy* InMaterialRenderProxy,
@@ -74,7 +78,8 @@ public:
 		const FMeshBatch& Mesh,
 		int32 BatchElementIndex,
 		bool bBackFace,
-		const ElementDataType& ElementData
+		const ElementDataType& ElementData,
+		const ContextDataType PolicyContext
 		) const;
 
 	/**
@@ -87,7 +92,7 @@ public:
 	 * @param CI - The command interface to execute the draw commands on.
 	 * @param View - The view of the scene being drawn.
 	 */
-	void DrawShared(FRHICommandList& RHICmdList, const FSceneView* View) const;
+	void SetSharedState(FRHICommandList& RHICmdList, const FSceneView* View, const ContextDataType PolicyContext) const;
 
 	/**
 	* Get the decl for this mesh policy type and vertexfactory
diff --git a/Engine/Source/Runtime/Renderer/Private/ForwardBasePassRendering.cpp b/Engine/Source/Runtime/Renderer/Private/ForwardBasePassRendering.cpp
--- a/Engine/Source/Runtime/Renderer/Private/ForwardBasePassRendering.cpp
+++ b/Engine/Source/Runtime/Renderer/Private/ForwardBasePassRendering.cpp
@@ -176,7 +176,7 @@ public:
 			View.Family->EngineShowFlags.ShaderComplexity
 			);
 		RHICmdList.BuildAndSetLocalBoundShaderState(DrawingPolicy.GetBoundShaderStateInput(View.GetFeatureLevel()));
-		DrawingPolicy.SetSharedState(RHICmdList, &View);
+		DrawingPolicy.SetSharedState(RHICmdList, &View, typename TBasePassForForwardShadingDrawingPolicy<LightMapPolicyType>::ContextDataType());
 
 		for( int32 BatchElementIndex=0;BatchElementIndex<Parameters.Mesh.Elements.Num();BatchElementIndex++ )
 		{
@@ -187,7 +187,8 @@ public:
 				Parameters.Mesh,
 				BatchElementIndex,
 				bBackFace,
-				typename TBasePassForForwardShadingDrawingPolicy<LightMapPolicyType>::ElementDataType(LightMapElementData)
+				typename TBasePassForForwardShadingDrawingPolicy<LightMapPolicyType>::ElementDataType(LightMapElementData),
+				typename TBasePassForForwardShadingDrawingPolicy<LightMapPolicyType>::ContextDataType()
 				);
 			DrawingPolicy.DrawMesh(RHICmdList, Parameters.Mesh, BatchElementIndex);
 		}
diff --git a/Engine/Source/Runtime/Renderer/Private/ForwardBasePassRendering.h b/Engine/Source/Runtime/Renderer/Private/ForwardBasePassRendering.h
--- a/Engine/Source/Runtime/Renderer/Private/ForwardBasePassRendering.h
+++ b/Engine/Source/Runtime/Renderer/Private/ForwardBasePassRendering.h
@@ -264,7 +264,7 @@ public:
 			SceneTextureMode == Other.SceneTextureMode;
 	}
 
-	void SetSharedState(FRHICommandList& RHICmdList, const FViewInfo* View) const
+	void SetSharedState(FRHICommandList& RHICmdList, const FViewInfo* View, const ContextDataType PolicyContext) const
 	{
 		VertexShader->SetParameters(RHICmdList, MaterialRenderProxy, VertexFactory, *MaterialResource, *View, SceneTextureMode);
 
@@ -351,7 +351,8 @@ public:
 		const FMeshBatch& Mesh,
 		int32 BatchElementIndex,
 		bool bBackFace,
-		const ElementDataType& ElementData
+		const ElementDataType& ElementData,
+		const ContextDataType PolicyContext
 		) const
 	{
 		// Set the light-map policy's mesh-specific settings.
@@ -391,7 +392,7 @@ public:
 			PixelShader->SetMesh(RHICmdList, VertexFactory,View,PrimitiveSceneProxy,BatchElement);
 		}
 
-		FMeshDrawingPolicy::SetMeshRenderState(RHICmdList, View,PrimitiveSceneProxy,Mesh,BatchElementIndex,bBackFace,FMeshDrawingPolicy::ElementDataType());
+		FMeshDrawingPolicy::SetMeshRenderState(RHICmdList, View,PrimitiveSceneProxy,Mesh,BatchElementIndex,bBackFace,FMeshDrawingPolicy::ElementDataType(),PolicyContext);
 	}
 
 	friend int32 CompareDrawingPolicy(const TBasePassForForwardShadingDrawingPolicy& A,const TBasePassForForwardShadingDrawingPolicy& B)
diff --git a/Engine/Source/Runtime/Renderer/Private/ForwardTranslucentRendering.cpp b/Engine/Source/Runtime/Renderer/Private/ForwardTranslucentRendering.cpp
--- a/Engine/Source/Runtime/Renderer/Private/ForwardTranslucentRendering.cpp
+++ b/Engine/Source/Runtime/Renderer/Private/ForwardTranslucentRendering.cpp
@@ -140,7 +140,7 @@ public:
 			);
 
 		RHICmdList.BuildAndSetLocalBoundShaderState(DrawingPolicy.GetBoundShaderStateInput(View.GetFeatureLevel()));
-		DrawingPolicy.SetSharedState(RHICmdList, &View);
+		DrawingPolicy.SetSharedState(RHICmdList, &View, typename TBasePassForForwardShadingDrawingPolicy<LightMapPolicyType>::ContextDataType());
 
 		for (int32 BatchElementIndex = 0; BatchElementIndex<Parameters.Mesh.Elements.Num(); BatchElementIndex++)
 		{
@@ -151,7 +151,8 @@ public:
 				Parameters.Mesh,
 				BatchElementIndex,
 				bBackFace,
-				typename TBasePassForForwardShadingDrawingPolicy<LightMapPolicyType>::ElementDataType(LightMapElementData)
+				typename TBasePassForForwardShadingDrawingPolicy<LightMapPolicyType>::ElementDataType(LightMapElementData),
+				typename TBasePassForForwardShadingDrawingPolicy<LightMapPolicyType>::ContextDataType()
 				);
 			DrawingPolicy.DrawMesh(RHICmdList, Parameters.Mesh, BatchElementIndex);
 		}
diff --git a/Engine/Source/Runtime/Renderer/Private/LightMapDensityRendering.cpp b/Engine/Source/Runtime/Renderer/Private/LightMapDensityRendering.cpp
--- a/Engine/Source/Runtime/Renderer/Private/LightMapDensityRendering.cpp
+++ b/Engine/Source/Runtime/Renderer/Private/LightMapDensityRendering.cpp
@@ -155,11 +155,13 @@ bool FLightMapDensityDrawingPolicyFactory::DrawDynamicMesh(
 			{
 				TLightMapDensityDrawingPolicy< TLightMapPolicy<HQ_LIGHTMAP> > DrawingPolicy(View, Mesh.VertexFactory, MaterialRenderProxy, TLightMapPolicy<HQ_LIGHTMAP>(), BlendMode);
 				RHICmdList.BuildAndSetLocalBoundShaderState(DrawingPolicy.GetBoundShaderStateInput(View.GetFeatureLevel()));
-				DrawingPolicy.SetSharedState(RHICmdList, &View);
+				DrawingPolicy.SetSharedState(RHICmdList, &View, TLightMapDensityDrawingPolicy< TLightMapPolicy<HQ_LIGHTMAP> >::ContextDataType());
 				for (int32 BatchElementIndex = 0; BatchElementIndex < Mesh.Elements.Num(); BatchElementIndex++)
 				{
 					DrawingPolicy.SetMeshRenderState(RHICmdList, View,PrimitiveSceneProxy,Mesh,BatchElementIndex,bBackFace,
-						TLightMapDensityDrawingPolicy< TLightMapPolicy<HQ_LIGHTMAP> >::ElementDataType(LightMapInteraction));
+						TLightMapDensityDrawingPolicy< TLightMapPolicy<HQ_LIGHTMAP> >::ElementDataType(LightMapInteraction),
+						TLightMapDensityDrawingPolicy< TLightMapPolicy<HQ_LIGHTMAP> >::ContextDataType()
+						);
 					DrawingPolicy.DrawMesh(RHICmdList, Mesh,BatchElementIndex);
 				}
 				bDirty = true;
@@ -168,11 +170,13 @@ bool FLightMapDensityDrawingPolicyFactory::DrawDynamicMesh(
 			{
 				TLightMapDensityDrawingPolicy< TLightMapPolicy<LQ_LIGHTMAP> > DrawingPolicy(View, Mesh.VertexFactory, MaterialRenderProxy, TLightMapPolicy<LQ_LIGHTMAP>(), BlendMode);
 				RHICmdList.BuildAndSetLocalBoundShaderState(DrawingPolicy.GetBoundShaderStateInput(View.GetFeatureLevel()));
-				DrawingPolicy.SetSharedState(RHICmdList, &View);
+				DrawingPolicy.SetSharedState(RHICmdList, &View, TLightMapDensityDrawingPolicy< TLightMapPolicy<LQ_LIGHTMAP> >::ContextDataType());
 				for (int32 BatchElementIndex = 0; BatchElementIndex < Mesh.Elements.Num(); BatchElementIndex++)
 				{
 					DrawingPolicy.SetMeshRenderState(RHICmdList, View,PrimitiveSceneProxy,Mesh,BatchElementIndex,bBackFace,
-						TLightMapDensityDrawingPolicy< TLightMapPolicy<LQ_LIGHTMAP> >::ElementDataType(LightMapInteraction));
+						TLightMapDensityDrawingPolicy< TLightMapPolicy<LQ_LIGHTMAP> >::ElementDataType(LightMapInteraction),
+						TLightMapDensityDrawingPolicy< TLightMapPolicy<LQ_LIGHTMAP> >::ContextDataType()
+						);
 					DrawingPolicy.DrawMesh(RHICmdList, Mesh,BatchElementIndex);
 				}
 				bDirty = true;
@@ -182,11 +186,13 @@ bool FLightMapDensityDrawingPolicyFactory::DrawDynamicMesh(
 		{
 			TLightMapDensityDrawingPolicy<FDummyLightMapPolicy> DrawingPolicy(View, Mesh.VertexFactory, MaterialRenderProxy, FDummyLightMapPolicy(), BlendMode);
 			RHICmdList.BuildAndSetLocalBoundShaderState(DrawingPolicy.GetBoundShaderStateInput(View.GetFeatureLevel()));
-			DrawingPolicy.SetSharedState(RHICmdList, &View);
+			DrawingPolicy.SetSharedState(RHICmdList, &View, TLightMapDensityDrawingPolicy<FDummyLightMapPolicy>::ContextDataType());
 			for (int32 BatchElementIndex = 0; BatchElementIndex < Mesh.Elements.Num(); BatchElementIndex++)
 			{
 				DrawingPolicy.SetMeshRenderState(RHICmdList, View,PrimitiveSceneProxy,Mesh,BatchElementIndex,bBackFace,
-					TLightMapDensityDrawingPolicy<FDummyLightMapPolicy>::ElementDataType(LightMapInteraction));
+					TLightMapDensityDrawingPolicy<FDummyLightMapPolicy>::ElementDataType(LightMapInteraction),
+					TLightMapDensityDrawingPolicy<FDummyLightMapPolicy>::ContextDataType()
+					);
 				DrawingPolicy.DrawMesh(RHICmdList, Mesh,BatchElementIndex);
 			}
 			bDirty = true;
@@ -196,10 +202,13 @@ bool FLightMapDensityDrawingPolicyFactory::DrawDynamicMesh(
 	{
 		TLightMapDensityDrawingPolicy<FNoLightMapPolicy> DrawingPolicy(View, Mesh.VertexFactory, MaterialRenderProxy, FNoLightMapPolicy(), BlendMode);
 		RHICmdList.BuildAndSetLocalBoundShaderState(DrawingPolicy.GetBoundShaderStateInput(View.GetFeatureLevel()));
-		DrawingPolicy.SetSharedState(RHICmdList, &View);
+		DrawingPolicy.SetSharedState(RHICmdList, &View, TLightMapDensityDrawingPolicy<FNoLightMapPolicy>::ContextDataType());
 		for (int32 BatchElementIndex = 0; BatchElementIndex < Mesh.Elements.Num(); BatchElementIndex++)
 		{
-			DrawingPolicy.SetMeshRenderState(RHICmdList, View,PrimitiveSceneProxy,Mesh,BatchElementIndex,bBackFace,TLightMapDensityDrawingPolicy<FNoLightMapPolicy>::ElementDataType());
+			DrawingPolicy.SetMeshRenderState(RHICmdList, View,PrimitiveSceneProxy,Mesh,BatchElementIndex,bBackFace,
+				TLightMapDensityDrawingPolicy<FNoLightMapPolicy>::ElementDataType(),
+				TLightMapDensityDrawingPolicy<FNoLightMapPolicy>::ContextDataType()
+				);
 			DrawingPolicy.DrawMesh(RHICmdList, Mesh,BatchElementIndex);
 		}
 		bDirty = true;
diff --git a/Engine/Source/Runtime/Renderer/Private/LightMapDensityRendering.h b/Engine/Source/Runtime/Renderer/Private/LightMapDensityRendering.h
--- a/Engine/Source/Runtime/Renderer/Private/LightMapDensityRendering.h
+++ b/Engine/Source/Runtime/Renderer/Private/LightMapDensityRendering.h
@@ -296,7 +296,7 @@ public:
 			LightMapPolicy == Other.LightMapPolicy;
 	}
 
-	void SetSharedState(FRHICommandList& RHICmdList, const FSceneView* View) const
+	void SetSharedState(FRHICommandList& RHICmdList, const FSceneView* View, const ContextDataType) const
 	{
 		// Set the base pass shader parameters for the material.
 		VertexShader->SetParameters(RHICmdList, MaterialRenderProxy,*View);
@@ -321,7 +321,8 @@ public:
 		const FMeshBatch& Mesh,
 		int32 BatchElementIndex,
 		bool bBackFace,
-		const ElementDataType& ElementData
+		const ElementDataType& ElementData,
+		const ContextDataType PolicyContext
 		) const
 	{
 		const FMeshBatchElement& BatchElement = Mesh.Elements[BatchElementIndex];
@@ -374,7 +375,7 @@ public:
 		// Adjust for the grid texture being 2x2 repeating pattern...
 		LMResolutionScale *= 0.5f;
 		PixelShader->SetMesh(RHICmdList, VertexFactory,PrimitiveSceneProxy, BatchElement, View, bBackFace, BuiltLightingAndSelectedFlags, LMResolutionScale, bTextureMapped);	
-		FMeshDrawingPolicy::SetMeshRenderState(RHICmdList, View,PrimitiveSceneProxy,Mesh,BatchElementIndex,bBackFace,FMeshDrawingPolicy::ElementDataType());
+		FMeshDrawingPolicy::SetMeshRenderState(RHICmdList, View,PrimitiveSceneProxy,Mesh,BatchElementIndex,bBackFace,FMeshDrawingPolicy::ElementDataType(),PolicyContext);
 	}
 
 	/** 
diff --git a/Engine/Source/Runtime/Renderer/Private/SceneHitProxyRendering.cpp b/Engine/Source/Runtime/Renderer/Private/SceneHitProxyRendering.cpp
--- a/Engine/Source/Runtime/Renderer/Private/SceneHitProxyRendering.cpp
+++ b/Engine/Source/Runtime/Renderer/Private/SceneHitProxyRendering.cpp
@@ -173,7 +173,7 @@ FHitProxyDrawingPolicy::FHitProxyDrawingPolicy(
 	PixelShader = MaterialResource->GetShader<FHitProxyPS>(InVertexFactory->GetType());
 }
 
-void FHitProxyDrawingPolicy::SetSharedState(FRHICommandList& RHICmdList, const FSceneView* View) const
+void FHitProxyDrawingPolicy::SetSharedState(FRHICommandList& RHICmdList, const FSceneView* View, const ContextDataType PolicyContext) const
 {
 	// Set the depth-only shader parameters for the material.
 	VertexShader->SetParameters(RHICmdList, MaterialRenderProxy,*View);
@@ -186,7 +186,7 @@ void FHitProxyDrawingPolicy::SetSharedState(FRHICommandList& RHICmdList, const F
 	}
 
 	// Set the shared mesh resources.
-	FMeshDrawingPolicy::DrawShared(RHICmdList, View);
+	FMeshDrawingPolicy::SetSharedState(RHICmdList, View, PolicyContext);
 }
 
 /** 
@@ -212,7 +212,8 @@ void FHitProxyDrawingPolicy::SetMeshRenderState(
 	const FMeshBatch& Mesh,
 	int32 BatchElementIndex,
 	bool bBackFace,
-	const FHitProxyId HitProxyId
+	const FHitProxyId HitProxyId,
+	const ContextDataType PolicyContext
 	) const
 {
 	EmitMeshDrawEvents(PrimitiveSceneProxy, Mesh);
@@ -310,10 +311,10 @@ bool FHitProxyDrawingPolicyFactory::DrawDynamicMesh(
 			}
 			FHitProxyDrawingPolicy DrawingPolicy( Mesh.VertexFactory, MaterialRenderProxy, View.GetFeatureLevel() );
 			RHICmdList.BuildAndSetLocalBoundShaderState(DrawingPolicy.GetBoundShaderStateInput(View.GetFeatureLevel()));
-			DrawingPolicy.SetSharedState(RHICmdList, &View);
+			DrawingPolicy.SetSharedState(RHICmdList, &View, FHitProxyDrawingPolicy::ContextDataType());
 			for (int32 BatchElementIndex = 0; BatchElementIndex < Mesh.Elements.Num(); BatchElementIndex++)
 			{
-				DrawingPolicy.SetMeshRenderState(RHICmdList, View, PrimitiveSceneProxy, Mesh, BatchElementIndex, bBackFace, HitProxyId);
+				DrawingPolicy.SetMeshRenderState(RHICmdList, View, PrimitiveSceneProxy, Mesh, BatchElementIndex, bBackFace, HitProxyId, FHitProxyDrawingPolicy::ContextDataType());
 				DrawingPolicy.DrawMesh(RHICmdList, Mesh,BatchElementIndex);
 			}
 			return true;
diff --git a/Engine/Source/Runtime/Renderer/Private/SceneHitProxyRendering.h b/Engine/Source/Runtime/Renderer/Private/SceneHitProxyRendering.h
--- a/Engine/Source/Runtime/Renderer/Private/SceneHitProxyRendering.h
+++ b/Engine/Source/Runtime/Renderer/Private/SceneHitProxyRendering.h
@@ -30,7 +30,7 @@ public:
 			VertexShader == Other.VertexShader &&
 			PixelShader == Other.PixelShader;
 	}
-	void SetSharedState(FRHICommandList& RHICmdList, const FSceneView* View) const;
+	void SetSharedState(FRHICommandList& RHICmdList, const FSceneView* View, const ContextDataType PolicyContext) const;
 
 	/** 
 	* Create bound shader state using the vertex decl from the mesh draw policy
@@ -46,7 +46,8 @@ public:
 		const FMeshBatch& Mesh,
 		int32 BatchElementIndex,
 		bool bBackFace,
-		const FHitProxyId HitProxyId
+		const FHitProxyId HitProxyId,
+		const ContextDataType PolicyContext
 		) const;
 
 private:
diff --git a/Engine/Source/Runtime/Renderer/Private/ShadowRendering.cpp b/Engine/Source/Runtime/Renderer/Private/ShadowRendering.cpp
--- a/Engine/Source/Runtime/Renderer/Private/ShadowRendering.cpp
+++ b/Engine/Source/Runtime/Renderer/Private/ShadowRendering.cpp
@@ -700,11 +700,6 @@ void OverrideWithDefaultMaterialForShadowDepth(
 	FShadowDepthDrawingPolicy
 -----------------------------------------------------------------------------*/
 
-template<>
-const FProjectedShadowInfo* FShadowDepthDrawingPolicy<false>::PolicyShadowInfo = NULL;
-template<>
-const FProjectedShadowInfo* FShadowDepthDrawingPolicy<true>::PolicyShadowInfo = NULL;
-
 template <bool bRenderingReflectiveShadowMaps>
 void FShadowDepthDrawingPolicy<bRenderingReflectiveShadowMaps>::UpdateElementState(FShadowStaticMeshElement& State)
 {
@@ -862,26 +857,26 @@ FShadowDepthDrawingPolicy<bRenderingReflectiveShadowMaps>::FShadowDepthDrawingPo
 }
 
 template <bool bRenderingReflectiveShadowMaps>
-void FShadowDepthDrawingPolicy<bRenderingReflectiveShadowMaps>::SetSharedState(FRHICommandList& RHICmdList, const FSceneView* View) const
+void FShadowDepthDrawingPolicy<bRenderingReflectiveShadowMaps>::SetSharedState(FRHICommandList& RHICmdList, const FSceneView* View, const ContextDataType PolicyContext) const
 {
-	checkSlow(bDirectionalLight == PolicyShadowInfo->bDirectionalLight && bPreShadow == PolicyShadowInfo->bPreShadow);
+	checkSlow(bDirectionalLight == PolicyContext.ShadowInfo->bDirectionalLight && bPreShadow == PolicyContext.ShadowInfo->bPreShadow);
 
-	VertexShader->SetParameters(RHICmdList, MaterialRenderProxy,*MaterialResource,*View,PolicyShadowInfo);
+	VertexShader->SetParameters(RHICmdList, MaterialRenderProxy,*MaterialResource,*View,PolicyContext.ShadowInfo);
 
 	if (GeometryShader)
 	{
-		GeometryShader->SetParameters(RHICmdList, *View,PolicyShadowInfo);
+		GeometryShader->SetParameters(RHICmdList, *View,PolicyContext.ShadowInfo);
 	}
 
 	if(HullShader && DomainShader)
 	{
 		HullShader->SetParameters(RHICmdList, MaterialRenderProxy,*View);
-		DomainShader->SetParameters(RHICmdList, MaterialRenderProxy,*View,PolicyShadowInfo);
+		DomainShader->SetParameters(RHICmdList, MaterialRenderProxy,*View,PolicyContext.ShadowInfo);
 	}
 
 	if (PixelShader)
 	{
-		PixelShader->SetParameters(RHICmdList, MaterialRenderProxy,*MaterialResource,*View,PolicyShadowInfo);
+		PixelShader->SetParameters(RHICmdList, MaterialRenderProxy,*MaterialResource,*View,PolicyContext.ShadowInfo);
 	}
 
 	// Set the shared mesh resources.
@@ -891,12 +886,12 @@ void FShadowDepthDrawingPolicy<bRenderingReflectiveShadowMaps>::SetSharedState(F
 	}
 	else
 	{
-		FMeshDrawingPolicy::DrawShared(RHICmdList, View);
+		FMeshDrawingPolicy::SetSharedState(RHICmdList, View, PolicyContext);
 	}
 
 	// @TODO: only render directional light shadows as two sided, and only when blocking is enabled (required by geometry volume injection)
 	bool bIsTwoSided = IsTwoSided();
-	if ( PolicyShadowInfo->bReflectiveShadowmap )
+	if ( PolicyContext.ShadowInfo->bReflectiveShadowmap )
 	{
 		bIsTwoSided = true;
 	}
@@ -945,7 +940,8 @@ void FShadowDepthDrawingPolicy<bRenderingReflectiveShadowMaps>::SetMeshRenderSta
 	const FMeshBatch& Mesh,
 	int32 BatchElementIndex,
 	bool bBackFace,
-	const ElementDataType& ElementData
+	const ElementDataType& ElementData,
+	const ContextDataType PolicyContext
 	) const
 {
 	const FMeshBatchElement& BatchElement = Mesh.Elements[BatchElementIndex];
@@ -961,13 +957,13 @@ void FShadowDepthDrawingPolicy<bRenderingReflectiveShadowMaps>::SetMeshRenderSta
 	}
 	if (GeometryShader)
 	{
-		GeometryShader->SetMesh(RHICmdList, PrimitiveSceneProxy, PolicyShadowInfo, View);
+		GeometryShader->SetMesh(RHICmdList, PrimitiveSceneProxy, PolicyContext.ShadowInfo, View);
 	}
 	if (PixelShader)
 	{
 		PixelShader->SetMesh(RHICmdList, VertexFactory,View,PrimitiveSceneProxy,BatchElement);
 	}
-	// Not calling FMeshDrawingPolicy::SetMeshRenderState as DrawShared sets the rasterizer state
+	// Not calling FMeshDrawingPolicy::SetMeshRenderState as SetSharedState sets the rasterizer state
 }
 
 template <bool bRenderingReflectiveShadowMaps>
@@ -981,7 +977,7 @@ int32 CompareDrawingPolicy(const FShadowDepthDrawingPolicy<bRenderingReflectiveS
 	COMPAREDRAWINGPOLICYMEMBERS(VertexFactory);
 	COMPAREDRAWINGPOLICYMEMBERS(MaterialRenderProxy);
 	COMPAREDRAWINGPOLICYMEMBERS(bDirectionalLight);
-	// Have to sort on two sidedness because rasterizer state is set in DrawShared and not SetMeshRenderState
+	// Have to sort on two sidedness because rasterizer state is set in SetSharedState and not SetMeshRenderState
 	COMPAREDRAWINGPOLICYMEMBERS(bIsTwoSidedMaterial);
 	COMPAREDRAWINGPOLICYMEMBERS(bReverseCulling);
 	COMPAREDRAWINGPOLICYMEMBERS(bOnePassPointLightShadow);
@@ -1075,6 +1071,7 @@ bool FShadowDepthDrawingPolicyFactory::DrawDynamicMesh(
 			const bool bDirectionalLight = Context.ShadowInfo->bDirectionalLight;
 			const bool bPreShadow = Context.ShadowInfo->bPreShadow;
 			const bool bTwoSided = Material->IsTwoSided() || PrimitiveSceneProxy->CastsShadowAsTwoSided();
+			const FShadowDepthDrawingPolicyContext PolicyContext(Context.ShadowInfo);
 
 			OverrideWithDefaultMaterialForShadowDepth(MaterialRenderProxy, Material, bReflectiveShadowmap, View.GetFeatureLevel());
 
@@ -1092,10 +1089,10 @@ bool FShadowDepthDrawingPolicyFactory::DrawDynamicMesh(
 					);
 
 				RHICmdList.BuildAndSetLocalBoundShaderState(DrawingPolicy.GetBoundShaderStateInput(View.GetFeatureLevel()));
-				DrawingPolicy.SetSharedState(RHICmdList, &View);
+				DrawingPolicy.SetSharedState(RHICmdList, &View, PolicyContext);
 				for (int32 BatchElementIndex = 0, Num = Mesh.Elements.Num(); BatchElementIndex < Num; BatchElementIndex++)
 				{
-					DrawingPolicy.SetMeshRenderState(RHICmdList, View,PrimitiveSceneProxy,Mesh,BatchElementIndex,bBackFace,FMeshDrawingPolicy::ElementDataType());
+					DrawingPolicy.SetMeshRenderState(RHICmdList, View,PrimitiveSceneProxy,Mesh,BatchElementIndex,bBackFace,FMeshDrawingPolicy::ElementDataType(),PolicyContext);
 					DrawingPolicy.DrawMesh(RHICmdList, Mesh,BatchElementIndex);
 				}
 			}
@@ -1113,10 +1110,10 @@ bool FShadowDepthDrawingPolicyFactory::DrawDynamicMesh(
 					);
 
 				RHICmdList.BuildAndSetLocalBoundShaderState(DrawingPolicy.GetBoundShaderStateInput(View.GetFeatureLevel()));
-				DrawingPolicy.SetSharedState(RHICmdList, &View);
+				DrawingPolicy.SetSharedState(RHICmdList, &View, PolicyContext);
 				for (int32 BatchElementIndex = 0; BatchElementIndex < Mesh.Elements.Num(); BatchElementIndex++)
 				{
-					DrawingPolicy.SetMeshRenderState(RHICmdList, View,PrimitiveSceneProxy,Mesh,BatchElementIndex,bBackFace,FMeshDrawingPolicy::ElementDataType());
+					DrawingPolicy.SetMeshRenderState(RHICmdList, View,PrimitiveSceneProxy,Mesh,BatchElementIndex,bBackFace,FMeshDrawingPolicy::ElementDataType(),PolicyContext);
 					DrawingPolicy.DrawMesh(RHICmdList, Mesh,BatchElementIndex);
 				}
 			}
@@ -1198,7 +1195,7 @@ void FProjectedShadowInfo::ClearDepth(FRHICommandList& RHICmdList, FDeferredShad
 
 
 template <bool bReflectiveShadowmap>
-void DrawMeshElements(FRHICommandList& RHICmdList, FShadowDepthDrawingPolicy<bReflectiveShadowmap>& SharedDrawingPolicy, const FShadowStaticMeshElement& State, const FViewInfo& View, const FStaticMesh* Mesh)
+void DrawMeshElements(FRHICommandList& RHICmdList, FShadowDepthDrawingPolicy<bReflectiveShadowmap>& SharedDrawingPolicy, const FShadowStaticMeshElement& State, const FViewInfo& View, FShadowDepthDrawingPolicyContext PolicyContext, const FStaticMesh* Mesh)
 {
 #if UE_BUILD_DEBUG
 	// During shadow setup we should have already overridden materials with default material where needed.
@@ -1228,7 +1225,7 @@ void DrawMeshElements(FRHICommandList& RHICmdList, FShadowDepthDrawingPolicy<bRe
 	{
 		if(BatchElementMask & 1)
 		{
-			SharedDrawingPolicy.SetMeshRenderState(RHICmdList, View, Mesh->PrimitiveSceneInfo->Proxy, *Mesh, BatchElementIndex, false, FMeshDrawingPolicy::ElementDataType());
+			SharedDrawingPolicy.SetMeshRenderState(RHICmdList, View, Mesh->PrimitiveSceneInfo->Proxy, *Mesh, BatchElementIndex, false, FMeshDrawingPolicy::ElementDataType(),PolicyContext);
 			SharedDrawingPolicy.DrawMesh(RHICmdList, *Mesh, BatchElementIndex);
 			INC_DWORD_STAT(STAT_ShadowDynamicPathDrawCalls);
 		}
@@ -1241,6 +1238,7 @@ void DrawMeshElements(FRHICommandList& RHICmdList, FShadowDepthDrawingPolicy<bRe
 template <bool bReflectiveShadowmap>
 void DrawShadowMeshElements(FRHICommandList& RHICmdList, const FViewInfo& View, const FProjectedShadowInfo& ShadowInfo)
 {
+	FShadowDepthDrawingPolicyContext PolicyContext(&ShadowInfo);
 	const FShadowStaticMeshElement& FirstShadowMesh = ShadowInfo.SubjectMeshElements[0];
 	const FMaterial* FirstMaterialResource = FirstShadowMesh.MaterialResource;
 
@@ -1272,10 +1270,10 @@ void DrawShadowMeshElements(FRHICommandList& RHICmdList, const FViewInfo& View,
 
 			SharedDrawingPolicy.UpdateElementState(CurrentState);
 			RHICmdList.BuildAndSetLocalBoundShaderState(SharedDrawingPolicy.GetBoundShaderStateInput(View.GetFeatureLevel()));
-			SharedDrawingPolicy.SetSharedState(RHICmdList, &View);
+			SharedDrawingPolicy.SetSharedState(RHICmdList, &View, PolicyContext);
 		}
 
-		DrawMeshElements(RHICmdList, SharedDrawingPolicy, OldState, View, ShadowMesh.Mesh);
+		DrawMeshElements(RHICmdList, SharedDrawingPolicy, OldState, View, PolicyContext, ShadowMesh.Mesh);
 	}
 }
 
@@ -1439,6 +1437,8 @@ static TAutoConsoleVariable<int32> CVarParallelShadows(
 
 void FProjectedShadowInfo::RenderDepthInner(FRHICommandList& RHICmdList, FDeferredShadingSceneRenderer* SceneRenderer, const FViewInfo* FoundView)
 {
+	FShadowDepthDrawingPolicyContext PolicyContext(this);
+
 	if (RHICmdList.IsImmediate() &&  // translucent shadows are draw on the render thread, using a recursive cmdlist
 		GRHICommandList.UseParallelAlgorithms() && CVarParallelShadows.GetValueOnRenderThread())
 	{
@@ -1455,12 +1455,12 @@ void FProjectedShadowInfo::RenderDepthInner(FRHICommandList& RHICmdList, FDeferr
 
 			if (bReflectiveShadowmap)
 			{
-				SceneRenderer->Scene->WholeSceneReflectiveShadowMapDrawList.DrawVisibleParallel(*FoundView, StaticMeshWholeSceneShadowDepthMap, StaticMeshWholeSceneShadowBatchVisibility, Width, RHICmdList, OutDirty);
+				SceneRenderer->Scene->WholeSceneReflectiveShadowMapDrawList.DrawVisibleParallel(*FoundView, PolicyContext, StaticMeshWholeSceneShadowDepthMap, StaticMeshWholeSceneShadowBatchVisibility, Width, RHICmdList, OutDirty);
 			}
 			else
 			{
 				// Use the scene's shadow depth draw list with this shadow's visibility map
-				SceneRenderer->Scene->WholeSceneShadowDepthDrawList.DrawVisibleParallel(*FoundView, StaticMeshWholeSceneShadowDepthMap, StaticMeshWholeSceneShadowBatchVisibility, Width, RHICmdList, OutDirty);
+				SceneRenderer->Scene->WholeSceneShadowDepthDrawList.DrawVisibleParallel(*FoundView, PolicyContext, StaticMeshWholeSceneShadowDepthMap, StaticMeshWholeSceneShadowBatchVisibility, Width, RHICmdList, OutDirty);
 			}
 		}
 		// Draw the subject's static elements using manual state filtering
@@ -1493,12 +1493,12 @@ void FProjectedShadowInfo::RenderDepthInner(FRHICommandList& RHICmdList, FDeferr
 
 			if (bReflectiveShadowmap)
 			{
-				SceneRenderer->Scene->WholeSceneReflectiveShadowMapDrawList.DrawVisible(RHICmdList, *FoundView, StaticMeshWholeSceneShadowDepthMap, StaticMeshWholeSceneShadowBatchVisibility);
+				SceneRenderer->Scene->WholeSceneReflectiveShadowMapDrawList.DrawVisible(RHICmdList, *FoundView, PolicyContext, StaticMeshWholeSceneShadowDepthMap, StaticMeshWholeSceneShadowBatchVisibility);
 			}
 			else
 			{
 				// Use the scene's shadow depth draw list with this shadow's visibility map
-				SceneRenderer->Scene->WholeSceneShadowDepthDrawList.DrawVisible(RHICmdList, *FoundView, StaticMeshWholeSceneShadowDepthMap, StaticMeshWholeSceneShadowBatchVisibility);
+				SceneRenderer->Scene->WholeSceneShadowDepthDrawList.DrawVisible(RHICmdList, *FoundView, PolicyContext, StaticMeshWholeSceneShadowDepthMap, StaticMeshWholeSceneShadowBatchVisibility);
 			}
 		}
 		// Draw the subject's static elements using manual state filtering
@@ -1619,14 +1619,8 @@ void FProjectedShadowInfo::RenderDepth(FRHICommandList& RHICmdList, FDeferredSha
 	// Lighting only should only affect the material used with direct lighting, not the indirect lighting
 	FoundView->bForceShowMaterials = true;
 
-	FShadowDepthDrawingPolicy<false>::PolicyShadowInfo = this;
-	FShadowDepthDrawingPolicy<true>::PolicyShadowInfo = this;
-
 	RenderDepthInner(RHICmdList, SceneRenderer, FoundView);
 
-	FShadowDepthDrawingPolicy<false>::PolicyShadowInfo = NULL;
-	FShadowDepthDrawingPolicy<true>::PolicyShadowInfo = NULL;
-
 	FoundView->bForceShowMaterials = false;
 	FoundView->UniformBuffer = OriginalUniformBuffer;
 	FoundView->ViewMatrices.ViewMatrix = OriginalViewMatrix;
@@ -2453,7 +2447,7 @@ float FProjectedShadowInfo::ComputeTransitionSize() const
 
 void FProjectedShadowInfo::SortSubjectMeshElements()
 {
-	// Note: this should match the criteria in FProjectedShadowInfo::RenderDepth for deciding when to call DrawShared on a static mesh element for best performance
+	// Note: this should match the criteria in FProjectedShadowInfo::RenderDepth for deciding when to call SetSharedState on a static mesh element for best performance
 	struct FCompareFShadowStaticMeshElement
 	{
 		FORCEINLINE bool operator()( const FShadowStaticMeshElement& A, const FShadowStaticMeshElement& B ) const
diff --git a/Engine/Source/Runtime/Renderer/Private/ShadowRendering.h b/Engine/Source/Runtime/Renderer/Private/ShadowRendering.h
--- a/Engine/Source/Runtime/Renderer/Private/ShadowRendering.h
+++ b/Engine/Source/Runtime/Renderer/Private/ShadowRendering.h
@@ -282,6 +282,20 @@ const static uint32 SHADOW_BORDER = 4;
 template <bool bRenderingReflectiveShadowMaps> class TShadowDepthBasePS;
 class FShadowStaticMeshElement;
 
+/**
+ * The shadow depth drawing policy's context data.
+ */
+struct FShadowDepthDrawingPolicyContext : FMeshDrawingPolicy::ContextDataType
+{
+	/** The projected shadow info for which we are rendering shadow depths. */
+	const FProjectedShadowInfo* ShadowInfo;
+
+	/** Initialization constructor. */
+	explicit FShadowDepthDrawingPolicyContext(const FProjectedShadowInfo* InShadowInfo)
+		: ShadowInfo(InShadowInfo)
+	{}
+};
+
 /**
  * Outputs no color, but can be used to write the mesh's depth values to the depth buffer.
  */
@@ -289,6 +303,8 @@ template <bool bRenderingReflectiveShadowMaps>
 class FShadowDepthDrawingPolicy : public FMeshDrawingPolicy
 {
 public:
+	typedef FShadowDepthDrawingPolicyContext ContextDataType;
+
 	FShadowDepthDrawingPolicy(
 		const FMaterial* InMaterialResource,
 		bool bInDirectionalLight,
@@ -333,7 +349,7 @@ public:
 			&& bUsePositionOnlyVS == Other.bUsePositionOnlyVS
 			&& bPreShadow == Other.bPreShadow;
 	}
-	void SetSharedState(FRHICommandList& RHICmdList, const FSceneView* View) const;
+	void SetSharedState(FRHICommandList& RHICmdList, const FSceneView* View, const ContextDataType PolicyContext) const;
 
 	/** 
 	 * Create bound shader state using the vertex decl from the mesh draw policy
@@ -349,18 +365,13 @@ public:
 		const FMeshBatch& Mesh,
 		int32 BatchElementIndex,
 		bool bBackFace,
-		const ElementDataType& ElementData
+		const ElementDataType& ElementData,
+		const ContextDataType PolicyContext
 		) const;
 
 	template<bool T2>
 	friend int32 CompareDrawingPolicy(const FShadowDepthDrawingPolicy<T2>& A,const FShadowDepthDrawingPolicy<T2>& B);
 
-	/** 
-	 * Shadow currently being rendered by a FShadowDepthDrawingPolicy.  
-	 * This is global so that different shadows can be used with the same static draw list. 
-	 */
-	static const FProjectedShadowInfo* PolicyShadowInfo;
-
 	bool IsReversingCulling() const
 	{
 		return bReverseCulling;
@@ -383,13 +394,6 @@ public:
 	uint32 bPreShadow:1;
 };
 
-/** 
- * Shadow currently being rendered by a FShadowDepthDrawingPolicy.  
- * This is global so that different shadows can be used with the same static draw list. 
- */
-extern const FProjectedShadowInfo* GShadowInfo;
-
-
 /**
  * A drawing policy factory for the shadow depth drawing policy.
  */
diff --git a/Engine/Source/Runtime/Renderer/Private/StaticMeshDrawList.h b/Engine/Source/Runtime/Renderer/Private/StaticMeshDrawList.h
--- a/Engine/Source/Runtime/Renderer/Private/StaticMeshDrawList.h
+++ b/Engine/Source/Runtime/Renderer/Private/StaticMeshDrawList.h
@@ -241,7 +241,7 @@ private:
 	* @param DrawingPolicyLink - the drawing policy link
 	* @param bDrawnShared - determines whether to draw shared 
 	*/
-	void DrawElement(FRHICommandList& RHICmdList, const FViewInfo& View, const FElement& Element, uint64 BatchElementMask, FDrawingPolicyLink* DrawingPolicyLink, bool &bDrawnShared);
+	void DrawElement(FRHICommandList& RHICmdList, const FViewInfo& View, const typename DrawingPolicyType::ContextDataType PolicyContext, const FElement& Element, uint64 BatchElementMask, FDrawingPolicyLink* DrawingPolicyLink, bool &bDrawnShared);
 
 public:
 
@@ -265,7 +265,7 @@ public:
 	 * @param StaticMeshVisibilityMap - An map from FStaticMesh::Id to visibility state.
 	 * @return True if any static meshes were drawn.
 	 */
-	bool DrawVisible(const FViewInfo& View, const TBitArray<SceneRenderingBitArrayAllocator>& StaticMeshVisibilityMap);
+	bool DrawVisible(const FViewInfo& View, const typename DrawingPolicyType::ContextDataType PolicyContext, const TBitArray<SceneRenderingBitArrayAllocator>& StaticMeshVisibilityMap);
 
 	/**
 	* Draws only the static meshes which are in the visibility map, limited to a range of policies
@@ -276,7 +276,7 @@ public:
 	* @param LastPolicy - Last policy to render
 	* @return True if any static meshes were drawn.
 	*/
-	bool DrawVisibleInner(FRHICommandList& RHICmdList, const FViewInfo& View, const TBitArray<SceneRenderingBitArrayAllocator>& StaticMeshVisibilityMap, const TArray<uint64, SceneRenderingAllocator>& BatchVisibilityArray, int32 FirstPolicy, int32 LastPolicy);
+	bool DrawVisibleInner(FRHICommandList& RHICmdList, const FViewInfo& View, const typename DrawingPolicyType::ContextDataType PolicyContext, const TBitArray<SceneRenderingBitArrayAllocator>& StaticMeshVisibilityMap, const TArray<uint64, SceneRenderingAllocator>& BatchVisibilityArray, int32 FirstPolicy, int32 LastPolicy);
 
 	/**
 	 * Draws only the static meshes which are in the visibility map.
@@ -285,7 +285,7 @@ public:
 	 * @param BatchVisibilityArray - An array of batch element visibility bitmasks.
 	 * @return True if any static meshes were drawn.
 	 */
-	bool DrawVisible(FRHICommandList& RHICmdList, const FViewInfo& View, const TBitArray<SceneRenderingBitArrayAllocator>& StaticMeshVisibilityMap, const TArray<uint64,SceneRenderingAllocator>& BatchVisibilityArray);
+	bool DrawVisible(FRHICommandList& RHICmdList, const FViewInfo& View, const typename DrawingPolicyType::ContextDataType PolicyContext, const TBitArray<SceneRenderingBitArrayAllocator>& StaticMeshVisibilityMap, const TArray<uint64,SceneRenderingAllocator>& BatchVisibilityArray);
 
 	/**
 	* Draws only the static meshes which are in the visibility map.
@@ -296,7 +296,7 @@ public:
 	* @param ParentCmdList - cmdlist to put the wait and execute task on
 	* @param OutDirty - set to true if anything is drawn
 	*/
-	void DrawVisibleParallel(const FViewInfo& View, const TBitArray<SceneRenderingBitArrayAllocator>& StaticMeshVisibilityMap, const TArray<uint64, SceneRenderingAllocator>& BatchVisibilityArray, int32 Width, FRHICommandList& ParentCmdList, bool& OutDirty);
+	void DrawVisibleParallel(const FViewInfo& View, const typename DrawingPolicyType::ContextDataType PolicyContext, const TBitArray<SceneRenderingBitArrayAllocator>& StaticMeshVisibilityMap, const TArray<uint64, SceneRenderingAllocator>& BatchVisibilityArray, int32 Width, FRHICommandList& ParentCmdList, bool& OutDirty);
 
 	/**
 	 * Draws only the static meshes which are in the visibility map, sorted front-to-back.
@@ -306,7 +306,30 @@ public:
 	 * @param MaxToDraw - The maximum number of meshes to be drawn.
 	 * @return The number of static meshes drawn.
 	 */
-	int32 DrawVisibleFrontToBack(FRHICommandList& RHICmdList, const FViewInfo& View, const TBitArray<SceneRenderingBitArrayAllocator>& StaticMeshVisibilityMap, const TArray<uint64,SceneRenderingAllocator>& BatchVisibilityArray, int32 MaxToDraw);
+	int32 DrawVisibleFrontToBack(FRHICommandList& RHICmdList, const FViewInfo& View, const typename DrawingPolicyType::ContextDataType PolicyContext, const TBitArray<SceneRenderingBitArrayAllocator>& StaticMeshVisibilityMap, const TArray<uint64,SceneRenderingAllocator>& BatchVisibilityArray, int32 MaxToDraw);
+
+	/**
+	 * Helper functions when policy context is not needed.
+	 */
+	inline bool DrawVisible(const FViewInfo& View, const TBitArray<SceneRenderingBitArrayAllocator>& StaticMeshVisibilityMap)
+	{
+		return DrawVisible(View, typename DrawingPolicyType::ContextDataType(), StaticMeshVisibilityMap);
+	}
+
+	inline bool DrawVisible(FRHICommandList& RHICmdList, const FViewInfo& View, const TBitArray<SceneRenderingBitArrayAllocator>& StaticMeshVisibilityMap, const TArray<uint64,SceneRenderingAllocator>& BatchVisibilityArray)
+	{
+		return DrawVisible(RHICmdList, View, typename DrawingPolicyType::ContextDataType(), StaticMeshVisibilityMap, BatchVisibilityArray);
+	}
+
+	inline void DrawVisibleParallel(const FViewInfo& View, const TBitArray<SceneRenderingBitArrayAllocator>& StaticMeshVisibilityMap, const TArray<uint64, SceneRenderingAllocator>& BatchVisibilityArray, int32 Width, FRHICommandList& ParentCmdList, bool& OutDirty)
+	{
+		DrawVisibleParallel(View, typename DrawingPolicyType::ContextDataType(), StaticMeshVisibilityMap, BatchVisibilityArray, Width, ParentCmdList, OutDirty);
+	}
+
+	inline int32 DrawVisibleFrontToBack(FRHICommandList& RHICmdList, const FViewInfo& View, const TBitArray<SceneRenderingBitArrayAllocator>& StaticMeshVisibilityMap, const TArray<uint64,SceneRenderingAllocator>& BatchVisibilityArray, int32 MaxToDraw)
+	{
+		return DrawVisibleFrontToBack(RHICmdList, View, typename DrawingPolicyType::ContextDataType(), StaticMeshVisibilityMap, BatchVisibilityArray, MaxToDraw);
+	}
 
 	/** Sorts OrderedDrawingPolicies front to back. */
 	void SortFrontToBack(FVector ViewPosition);
diff --git a/Engine/Source/Runtime/Renderer/Private/StaticMeshDrawList.inl b/Engine/Source/Runtime/Renderer/Private/StaticMeshDrawList.inl
--- a/Engine/Source/Runtime/Renderer/Private/StaticMeshDrawList.inl
+++ b/Engine/Source/Runtime/Renderer/Private/StaticMeshDrawList.inl
@@ -60,6 +60,7 @@ template<typename DrawingPolicyType>
 void TStaticMeshDrawList<DrawingPolicyType>::DrawElement(
 	FRHICommandList& RHICmdList,
 	const FViewInfo& View,
+	const typename DrawingPolicyType::ContextDataType PolicyContext,
 	const FElement& Element,
 	uint64 BatchElementMask,
 	FDrawingPolicyLink* DrawingPolicyLink,
@@ -80,7 +81,7 @@ void TStaticMeshDrawList<DrawingPolicyType>::DrawElement(
 		{
 			RHICmdList.BuildAndSetLocalBoundShaderState(DrawingPolicyLink->DrawingPolicy.GetBoundShaderStateInput(View.GetFeatureLevel()));
 		}
-		DrawingPolicyLink->DrawingPolicy.SetSharedState(RHICmdList, &View);
+		DrawingPolicyLink->DrawingPolicy.SetSharedState(RHICmdList, &View, PolicyContext);
 		bDrawnShared = true;
 	}
 	
@@ -102,7 +103,8 @@ void TStaticMeshDrawList<DrawingPolicyType>::DrawElement(
 					*Element.Mesh,
 					BatchElementIndex,
 					!!BackFace,
-					Element.PolicyData
+					Element.PolicyData,
+					PolicyContext
 					);
 				DrawingPolicyLink->DrawingPolicy.DrawMesh(RHICmdList, *Element.Mesh,BatchElementIndex);
 			}
@@ -212,6 +214,7 @@ void TStaticMeshDrawList<DrawingPolicyType>::ReleaseRHI()
 template<typename DrawingPolicyType>
 bool TStaticMeshDrawList<DrawingPolicyType>::DrawVisible(
 	const FViewInfo& View,
+	const typename DrawingPolicyType::ContextDataType PolicyContext,
 	const TBitArray<SceneRenderingBitArrayAllocator>& StaticMeshVisibilityMap
 	)
 {
@@ -232,7 +235,7 @@ bool TStaticMeshDrawList<DrawingPolicyType>::DrawVisible(
 				INC_DWORD_STAT_BY(STAT_StaticMeshTriangles,Element.Mesh->GetNumPrimitives());
 				// Avoid the virtual call looking up batch visibility if there is only one element.
 				uint32 BatchElementMask = Element.Mesh->Elements.Num() == 1 ? 1 : Element.Mesh->VertexFactory->GetStaticBatchElementVisibility(View,Element.Mesh);
-				DrawElement(View, Element, BatchElementMask, DrawingPolicyLink, bDrawnShared);
+				DrawElement(View, PolicyContext, Element, BatchElementMask, DrawingPolicyLink, bDrawnShared);
 				bDirty = true;
 			}
 		}
@@ -244,6 +247,7 @@ template<typename DrawingPolicyType>
 bool TStaticMeshDrawList<DrawingPolicyType>::DrawVisibleInner(
 	FRHICommandList& RHICmdList,
 	const FViewInfo& View,
+	const typename DrawingPolicyType::ContextDataType PolicyContext,
 	const TBitArray<SceneRenderingBitArrayAllocator>& StaticMeshVisibilityMap,
 	const TArray<uint64, SceneRenderingAllocator>& BatchVisibilityArray,
 	int32 FirstPolicy, int32 LastPolicy
@@ -266,7 +270,7 @@ bool TStaticMeshDrawList<DrawingPolicyType>::DrawVisibleInner(
 				INC_DWORD_STAT_BY(STAT_StaticMeshTriangles, Element.Mesh->GetNumPrimitives());
 				// Avoid the cache miss looking up batch visibility if there is only one element.
 				uint64 BatchElementMask = Element.Mesh->Elements.Num() == 1 ? 1 : BatchVisibilityArray[Element.Mesh->Id];
-				DrawElement(RHICmdList, View, Element, BatchElementMask, DrawingPolicyLink, bDrawnShared);
+				DrawElement(RHICmdList, View, PolicyContext, Element, BatchElementMask, DrawingPolicyLink, bDrawnShared);
 				bDirty = true;
 			}
 		}
@@ -278,11 +282,12 @@ template<typename DrawingPolicyType>
 bool TStaticMeshDrawList<DrawingPolicyType>::DrawVisible(
 	FRHICommandList& RHICmdList,
 	const FViewInfo& View,
+	const typename DrawingPolicyType::ContextDataType PolicyContext,
 	const TBitArray<SceneRenderingBitArrayAllocator>& StaticMeshVisibilityMap,
 	const TArray<uint64,SceneRenderingAllocator>& BatchVisibilityArray
 	)
 {
-	return DrawVisibleInner(RHICmdList, View, StaticMeshVisibilityMap, BatchVisibilityArray, 0, OrderedDrawingPolicies.Num() - 1);
+	return DrawVisibleInner(RHICmdList, View, PolicyContext, StaticMeshVisibilityMap, BatchVisibilityArray, 0, OrderedDrawingPolicies.Num() - 1);
 }
 
 
@@ -292,6 +297,7 @@ class FDrawVisibleAnyThreadTask
 	TStaticMeshDrawList<DrawingPolicyType>& Caller;
 	FRHICommandList& RHICmdList;
 	const FViewInfo& View;
+	const typename DrawingPolicyType::ContextDataType PolicyContext;
 	const TBitArray<SceneRenderingBitArrayAllocator>& StaticMeshVisibilityMap;
 	const TArray<uint64, SceneRenderingAllocator>& BatchVisibilityArray;
 
@@ -306,6 +312,7 @@ public:
 		TStaticMeshDrawList<DrawingPolicyType>* InCaller,
 		FRHICommandList* InRHICmdList,
 		const FViewInfo* InView,
+		const typename DrawingPolicyType::ContextDataType InPolicyContext,
 		const TBitArray<SceneRenderingBitArrayAllocator>* InStaticMeshVisibilityMap,
 		const TArray<uint64, SceneRenderingAllocator>* InBatchVisibilityArray,
 		int32 InFirstPolicy,
@@ -315,6 +322,7 @@ public:
 		: Caller(*InCaller)
 		, RHICmdList(*InRHICmdList)
 		, View(*InView)
+		, PolicyContext(InPolicyContext)
 		, StaticMeshVisibilityMap(*InStaticMeshVisibilityMap)
 		, BatchVisibilityArray(*InBatchVisibilityArray)
 		, FirstPolicy(InFirstPolicy)
@@ -337,7 +345,7 @@ public:
 
 	void DoTask(ENamedThreads::Type CurrentThread, const FGraphEventRef& MyCompletionGraphEvent)
 	{
-		if (this->Caller.DrawVisibleInner(this->RHICmdList, this->View, this->StaticMeshVisibilityMap, this->BatchVisibilityArray, this->FirstPolicy, this->LastPolicy))
+		if (this->Caller.DrawVisibleInner(this->RHICmdList, this->View, this->PolicyContext, this->StaticMeshVisibilityMap, this->BatchVisibilityArray, this->FirstPolicy, this->LastPolicy))
 		{
 			this->OutDirty = true;
 		}
@@ -347,6 +355,7 @@ public:
 template<typename DrawingPolicyType>
 void TStaticMeshDrawList<DrawingPolicyType>::DrawVisibleParallel(
 	const FViewInfo& View,
+	const typename DrawingPolicyType::ContextDataType PolicyContext,
 	const TBitArray<SceneRenderingBitArrayAllocator>& StaticMeshVisibilityMap,
 	const TArray<uint64, SceneRenderingAllocator>& BatchVisibilityArray,
 	int32 Width, FRHICommandList& ParentCmdList, bool& OutDirty
@@ -372,7 +381,7 @@ void TStaticMeshDrawList<DrawingPolicyType>::DrawVisibleParallel(
 				FRHICommandList* CmdList = new FRHICommandList;
 
 				FGraphEventRef AnyThreadCompletionEvent = TGraphTask<FDrawVisibleAnyThreadTask<DrawingPolicyType> >::CreateTask(nullptr, ENamedThreads::RenderThread)
-					.ConstructAndDispatchWhenReady(this, CmdList, &View, &StaticMeshVisibilityMap, &BatchVisibilityArray, Start, Last, &OutDirty);
+					.ConstructAndDispatchWhenReady(this, CmdList, &View, PolicyContext, &StaticMeshVisibilityMap, &BatchVisibilityArray, Start, Last, &OutDirty);
 
 				ParentCmdList.QueueAsyncCommandListSubmit(AnyThreadCompletionEvent, CmdList);
 			}
@@ -387,6 +396,7 @@ template<typename DrawingPolicyType>
 int32 TStaticMeshDrawList<DrawingPolicyType>::DrawVisibleFrontToBack(
 	FRHICommandList& RHICmdList,
 	const FViewInfo& View,
+	const typename DrawingPolicyType::ContextDataType PolicyContext,
 	const TBitArray<SceneRenderingBitArrayAllocator>& StaticMeshVisibilityMap,
 	const TArray<uint64,SceneRenderingAllocator>& BatchVisibilityArray,
 	int32 MaxToDraw
@@ -438,7 +448,7 @@ int32 TStaticMeshDrawList<DrawingPolicyType>::DrawVisibleFrontToBack(
 		INC_DWORD_STAT_BY(STAT_StaticMeshTriangles,Element.Mesh->GetNumPrimitives());
 		// Avoid the cache miss looking up batch visibility if there is only one element.
 		uint64 BatchElementMask = Element.Mesh->Elements.Num() == 1 ? 1 : BatchVisibilityArray[Element.Mesh->Id];
-		DrawElement(RHICmdList, View, Element, BatchElementMask, DrawingPolicyLink, bDrawnShared);
+		DrawElement(RHICmdList, View, PolicyContext, Element, BatchElementMask, DrawingPolicyLink, bDrawnShared);
 		NumDraws++;
 	}
 
diff --git a/Engine/Source/Runtime/Renderer/Private/TranslucentLighting.cpp b/Engine/Source/Runtime/Renderer/Private/TranslucentLighting.cpp
--- a/Engine/Source/Runtime/Renderer/Private/TranslucentLighting.cpp
+++ b/Engine/Source/Runtime/Renderer/Private/TranslucentLighting.cpp
@@ -318,6 +318,15 @@ class FTranslucencyShadowDepthDrawingPolicy : public FMeshDrawingPolicy
 {
 public:
 
+	struct ContextDataType : public FMeshDrawingPolicy::ContextDataType
+	{
+		const FProjectedShadowInfo* ShadowInfo;
+
+		explicit ContextDataType(const FProjectedShadowInfo* InShadowInfo)
+			: ShadowInfo(InShadowInfo)
+		{}
+	};
+
 	FTranslucencyShadowDepthDrawingPolicy(
 		const FVertexFactory* InVertexFactory,
 		const FMaterialRenderProxy* InMaterialRenderProxy,
@@ -340,13 +349,13 @@ public:
 		}
 	}
 
-	void SetSharedState(FRHICommandList& RHICmdList, const FSceneView* View) const
+	void SetSharedState(FRHICommandList& RHICmdList, const FSceneView* View, const ContextDataType PolicyContext) const
 	{
 		// Set the shared mesh resources.
-		FMeshDrawingPolicy::DrawShared(RHICmdList, View);
+		FMeshDrawingPolicy::SetSharedState(RHICmdList, View, PolicyContext);
 
-		VertexShader->SetParameters(RHICmdList, MaterialRenderProxy, *View, PolicyShadowInfo);
-		PixelShader->SetParameters(RHICmdList, MaterialRenderProxy, *View, PolicyShadowInfo);
+		VertexShader->SetParameters(RHICmdList, MaterialRenderProxy, *View, PolicyContext.ShadowInfo);
+		PixelShader->SetParameters(RHICmdList, MaterialRenderProxy, *View, PolicyContext.ShadowInfo);
 	}
 
 	/** 
@@ -372,25 +381,22 @@ public:
 		const FMeshBatch& Mesh,
 		int32 BatchElementIndex,
 		bool bBackFace,
-		const ElementDataType& ElementData
+		const ElementDataType& ElementData,
+		const ContextDataType PolicyContext
 		) const
 	{
 		const FMeshBatchElement& BatchElement = Mesh.Elements[BatchElementIndex];
 		VertexShader->SetMesh(RHICmdList, VertexFactory,View,PrimitiveSceneProxy,BatchElement);
 		PixelShader->SetMesh(RHICmdList, VertexFactory,View,PrimitiveSceneProxy,BatchElement);
 
-		FMeshDrawingPolicy::SetMeshRenderState(RHICmdList, View,PrimitiveSceneProxy,Mesh,BatchElementIndex,bBackFace,ElementData);
+		FMeshDrawingPolicy::SetMeshRenderState(RHICmdList, View,PrimitiveSceneProxy,Mesh,BatchElementIndex,bBackFace,ElementData,PolicyContext);
 	}
 
-	static const FProjectedShadowInfo* PolicyShadowInfo;
-
 private:
 	FTranslucencyShadowDepthVS* VertexShader;
 	FTranslucencyShadowDepthPS* PixelShader;
 };
 
-const FProjectedShadowInfo* FTranslucencyShadowDepthDrawingPolicy::PolicyShadowInfo = NULL;
-
 class FTranslucencyShadowDepthDrawingPolicyFactory
 {
 public:
@@ -398,10 +404,12 @@ public:
 	enum { bAllowSimpleElements = false };
 	struct ContextType 
 	{
-		ContextType(bool bInDirectionalLight) :
-			bDirectionalLight(bInDirectionalLight)
+		ContextType(const FProjectedShadowInfo* InShadowInfo, bool bInDirectionalLight)
+			: ShadowInfo(InShadowInfo)
+			, bDirectionalLight(bInDirectionalLight)
 		{}
 
+		const FProjectedShadowInfo* ShadowInfo;
 		bool bDirectionalLight;
 	};
 
@@ -429,11 +437,14 @@ public:
 			{
 				FTranslucencyShadowDepthDrawingPolicy DrawingPolicy(Mesh.VertexFactory, MaterialRenderProxy, *MaterialRenderProxy->GetMaterial(FeatureLevel), DrawingContext.bDirectionalLight);
 				RHICmdList.BuildAndSetLocalBoundShaderState(DrawingPolicy.GetBoundShaderStateInput(View.GetFeatureLevel()));
-				DrawingPolicy.SetSharedState(RHICmdList, &View);
+				DrawingPolicy.SetSharedState(RHICmdList, &View, FTranslucencyShadowDepthDrawingPolicy::ContextDataType(DrawingContext.ShadowInfo));
 
 				for (int32 BatchElementIndex = 0; BatchElementIndex < Mesh.Elements.Num(); BatchElementIndex++)
 				{
-					DrawingPolicy.SetMeshRenderState(RHICmdList, View,PrimitiveSceneProxy,Mesh,BatchElementIndex,bBackFace,FMeshDrawingPolicy::ElementDataType());
+					DrawingPolicy.SetMeshRenderState(RHICmdList, View,PrimitiveSceneProxy,Mesh,BatchElementIndex,bBackFace,
+						FTranslucencyShadowDepthDrawingPolicy::ElementDataType(),
+						FTranslucencyShadowDepthDrawingPolicy::ContextDataType(DrawingContext.ShadowInfo)
+						);
 					DrawingPolicy.DrawMesh(RHICmdList, Mesh,BatchElementIndex);
 				}
 				bDirty = true;
@@ -513,8 +524,6 @@ void FProjectedShadowInfo::RenderTranslucencyDepths(FRHICommandList& RHICmdList,
 	// Lighting only should only affect the material used with direct lighting, not the indirect lighting
 	FoundView->bForceShowMaterials = true;
 
-	FTranslucencyShadowDepthDrawingPolicy::PolicyShadowInfo = this;
-
 	{
 #if WANTS_DRAW_MESH_EVENTS
 		FString EventName;
@@ -559,7 +568,8 @@ void FProjectedShadowInfo::RenderTranslucencyDepths(FRHICommandList& RHICmdList,
 
 		const bool bUseGetMeshElements = ShouldUseGetDynamicMeshElements();
 
-		TDynamicPrimitiveDrawer<FTranslucencyShadowDepthDrawingPolicyFactory> OpacityDrawer(RHICmdList, FoundView, FTranslucencyShadowDepthDrawingPolicyFactory::ContextType(bDirectionalLight), true);
+		FTranslucencyShadowDepthDrawingPolicyFactory::ContextType DrawingContext(this,bDirectionalLight);
+		TDynamicPrimitiveDrawer<FTranslucencyShadowDepthDrawingPolicyFactory> OpacityDrawer(RHICmdList, FoundView, DrawingContext, true);
 
 		for (int32 PrimitiveIndex = 0; PrimitiveIndex < SubjectTranslucentPrimitives.Num(); PrimitiveIndex++)
 		{
@@ -577,8 +587,6 @@ void FProjectedShadowInfo::RenderTranslucencyDepths(FRHICommandList& RHICmdList,
 			{
 				if (bUseGetMeshElements)
 				{
-					FTranslucencyShadowDepthDrawingPolicyFactory::ContextType Context(bDirectionalLight);
-
 					for (int32 MeshBatchIndex = 0; MeshBatchIndex < FoundView->DynamicMeshElements.Num(); MeshBatchIndex++)
 					{
 						const FMeshBatchAndRelevance& MeshBatchAndRelevance = FoundView->DynamicMeshElements[MeshBatchIndex];
@@ -586,7 +594,7 @@ void FProjectedShadowInfo::RenderTranslucencyDepths(FRHICommandList& RHICmdList,
 						if (MeshBatchAndRelevance.PrimitiveSceneProxy == PrimitiveSceneInfo->Proxy)
 						{
 							const FMeshBatch& MeshBatch = *MeshBatchAndRelevance.Mesh;
-							FTranslucencyShadowDepthDrawingPolicyFactory::DrawDynamicMesh(RHICmdList, *FoundView, Context, MeshBatch, false, true, MeshBatchAndRelevance.PrimitiveSceneProxy, MeshBatch.BatchHitProxyId);
+							FTranslucencyShadowDepthDrawingPolicyFactory::DrawDynamicMesh(RHICmdList, *FoundView, DrawingContext, MeshBatch, false, true, MeshBatchAndRelevance.PrimitiveSceneProxy, MeshBatch.BatchHitProxyId);
 						}
 					}
 				}
@@ -603,7 +611,7 @@ void FProjectedShadowInfo::RenderTranslucencyDepths(FRHICommandList& RHICmdList,
 						FTranslucencyShadowDepthDrawingPolicyFactory::DrawStaticMesh(
 							RHICmdList, 
 							*FoundView,
-							FTranslucencyShadowDepthDrawingPolicyFactory::ContextType(bDirectionalLight),
+							DrawingContext,
 							PrimitiveSceneInfo->StaticMeshes[MeshIndex],
 							true,
 							PrimitiveSceneInfo->Proxy,
@@ -615,8 +623,6 @@ void FProjectedShadowInfo::RenderTranslucencyDepths(FRHICommandList& RHICmdList,
 	}
 
 	// Restore overridden properties
-	FTranslucencyShadowDepthDrawingPolicy::PolicyShadowInfo = NULL;
-
 	FoundView->bForceShowMaterials = false;
 	FoundView->UniformBuffer = OriginalUniformBuffer;
 	FoundView->ViewMatrices.ViewMatrix = OriginalViewMatrix;
diff --git a/Engine/Source/Runtime/Renderer/Private/TranslucentRendering.cpp b/Engine/Source/Runtime/Renderer/Private/TranslucentRendering.cpp
--- a/Engine/Source/Runtime/Renderer/Private/TranslucentRendering.cpp
+++ b/Engine/Source/Runtime/Renderer/Private/TranslucentRendering.cpp
@@ -277,7 +277,7 @@ public:
 			Parameters.bAllowFog
 			);
 		RHICmdList.BuildAndSetLocalBoundShaderState(DrawingPolicy.GetBoundShaderStateInput(View.GetFeatureLevel()));
-		DrawingPolicy.SetSharedState(RHICmdList, &View);
+		DrawingPolicy.SetSharedState(RHICmdList, &View, typename TBasePassDrawingPolicy<LightMapPolicyType>::ContextDataType());
 
 		int32 BatchElementIndex = 0;
 		uint64 BatchElementMask = Parameters.BatchElementMask;
@@ -292,7 +292,8 @@ public:
 					Parameters.Mesh,
 					BatchElementIndex,
 					bBackFace,
-					typename TBasePassDrawingPolicy<LightMapPolicyType>::ElementDataType(LightMapElementData)
+					typename TBasePassDrawingPolicy<LightMapPolicyType>::ElementDataType(LightMapElementData),
+					typename TBasePassDrawingPolicy<LightMapPolicyType>::ContextDataType()
 					);
 				DrawingPolicy.DrawMesh(RHICmdList, Parameters.Mesh,BatchElementIndex);
 			}
diff --git a/Engine/Source/Runtime/Renderer/Private/VelocityRendering.cpp b/Engine/Source/Runtime/Renderer/Private/VelocityRendering.cpp
--- a/Engine/Source/Runtime/Renderer/Private/VelocityRendering.cpp
+++ b/Engine/Source/Runtime/Renderer/Private/VelocityRendering.cpp
@@ -208,7 +208,7 @@ bool FVelocityDrawingPolicy::SupportsVelocity() const
 	return (VertexShader && PixelShader) ? VertexShader->SupportsVelocity() : false;
 }
 
-void FVelocityDrawingPolicy::SetSharedState(FRHICommandList& RHICmdList, const FSceneView* SceneView) const
+void FVelocityDrawingPolicy::SetSharedState(FRHICommandList& RHICmdList, const FSceneView* SceneView, const ContextDataType PolicyContext) const
 {
 	// NOTE: Assuming this cast is always safe!
 	FViewInfo* View = (FViewInfo*)SceneView;
@@ -223,7 +223,7 @@ void FVelocityDrawingPolicy::SetSharedState(FRHICommandList& RHICmdList, const F
 	}
 
 	// Set the shared mesh resources.
-	FMeshDrawingPolicy::DrawShared(RHICmdList, View );
+	FMeshDrawingPolicy::SetSharedState(RHICmdList, View, PolicyContext);
 }
 
 void FVelocityDrawingPolicy::SetMeshRenderState(
@@ -233,7 +233,8 @@ void FVelocityDrawingPolicy::SetMeshRenderState(
 	const FMeshBatch& Mesh,
 	int32 BatchElementIndex,
 	bool bBackFace,
-	const ElementDataType& ElementData
+	const ElementDataType& ElementData,
+	const ContextDataType PolicyContext
 	) const
 {
 	const FMeshBatchElement& BatchElement = Mesh.Elements[BatchElementIndex];
@@ -262,7 +263,7 @@ void FVelocityDrawingPolicy::SetMeshRenderState(
 	}
 
 	PixelShader->SetMesh(RHICmdList, VertexFactory, Mesh, BatchElementIndex, View, PrimitiveSceneProxy, bBackFace);
-	FMeshDrawingPolicy::SetMeshRenderState(RHICmdList, View, PrimitiveSceneProxy, Mesh, BatchElementIndex, bBackFace, ElementData);
+	FMeshDrawingPolicy::SetMeshRenderState(RHICmdList, View, PrimitiveSceneProxy, Mesh, BatchElementIndex, bBackFace, ElementData, PolicyContext);
 }
 
 bool FVelocityDrawingPolicy::HasVelocity(const FViewInfo& View, const FPrimitiveSceneInfo* PrimitiveSceneInfo)
@@ -397,10 +398,10 @@ bool FVelocityDrawingPolicyFactory::DrawDynamicMesh(
 		if(DrawingPolicy.SupportsVelocity())
 		{			
 			RHICmdList.BuildAndSetLocalBoundShaderState(DrawingPolicy.GetBoundShaderStateInput(View.GetFeatureLevel()));
-			DrawingPolicy.SetSharedState(RHICmdList, &View);
+			DrawingPolicy.SetSharedState(RHICmdList, &View, FVelocityDrawingPolicy::ContextDataType());
 			for (int32 BatchElementIndex = 0; BatchElementIndex < Mesh.Elements.Num(); ++BatchElementIndex)
 			{
-				DrawingPolicy.SetMeshRenderState(RHICmdList, View, PrimitiveSceneProxy, Mesh, BatchElementIndex, bBackFace, FMeshDrawingPolicy::ElementDataType());
+				DrawingPolicy.SetMeshRenderState(RHICmdList, View, PrimitiveSceneProxy, Mesh, BatchElementIndex, bBackFace, FMeshDrawingPolicy::ElementDataType(), FVelocityDrawingPolicy::ContextDataType());
 				DrawingPolicy.DrawMesh(RHICmdList, Mesh, BatchElementIndex);
 			}
 			return true;
diff --git a/Engine/Source/Runtime/Renderer/Private/VelocityRendering.h b/Engine/Source/Runtime/Renderer/Private/VelocityRendering.h
--- a/Engine/Source/Runtime/Renderer/Private/VelocityRendering.h
+++ b/Engine/Source/Runtime/Renderer/Private/VelocityRendering.h
@@ -28,7 +28,7 @@ public:
 			&& VertexShader == Other.VertexShader
 			&& PixelShader == Other.PixelShader;
 	}
-	void SetSharedState(FRHICommandList& RHICmdList, const FSceneView* View) const;
+	void SetSharedState(FRHICommandList& RHICmdList, const FSceneView* View, const ContextDataType PolicyContext) const;
 	void SetMeshRenderState(
 		FRHICommandList& RHICmdList, 
 		const FSceneView& View,
@@ -36,7 +36,8 @@ public:
 		const FMeshBatch& Mesh,
 		int32 BatchElementIndex,
 		bool bBackFace,
-		const ElementDataType& ElementData
+		const ElementDataType& ElementData, 
+		const ContextDataType PolicyContext
 		) const;
 
 	FBoundShaderStateInput GetBoundShaderStateInput(ERHIFeatureLevel::Type InFeatureLevel);
