fixed CustomDepth rendering

diff --git a/Engine/Source/Runtime/Renderer/Private/SceneRendering.cpp b/Engine/Source/Runtime/Renderer/Private/SceneRendering.cpp
--- a/Engine/Source/Runtime/Renderer/Private/SceneRendering.cpp
+++ b/Engine/Source/Runtime/Renderer/Private/SceneRendering.cpp
@@ -841,7 +841,7 @@ FSceneRenderer* FSceneRenderer::CreateSceneRenderer(const FSceneViewFamily* InVi
 
 void FSceneRenderer::RenderCustomDepthPass(FRHICommandListImmediate& RHICmdList)
 {
-	if(FeatureLevel >= ERHIFeatureLevel::SM3)
+	if(FeatureLevel < ERHIFeatureLevel::SM3)
 	{
 		// not yet supported on lower end platforms
 		return;
