fixed ScreenSpaceSubsurfaceScattering with LPV

diff --git a/Engine/Source/Runtime/Renderer/Private/CompositionLighting/PostProcessLpvIndirect.cpp b/Engine/Source/Runtime/Renderer/Private/CompositionLighting/PostProcessLpvIndirect.cpp
--- a/Engine/Source/Runtime/Renderer/Private/CompositionLighting/PostProcessLpvIndirect.cpp
+++ b/Engine/Source/Runtime/Renderer/Private/CompositionLighting/PostProcessLpvIndirect.cpp
@@ -132,6 +132,16 @@ void FRCPassPostProcessLpvIndirect::Process(FRenderingCompositePassContext& Cont
 {
 	SCOPED_DRAW_EVENT(PostProcessLpvIndirect, DEC_SCENE_ITEMS);
 
+	{
+		FRenderingCompositeOutput* OutputOfMyInput = GetInput(ePId_Input0)->GetOutput();
+		PassOutputs[0].PooledRenderTarget = OutputOfMyInput->PooledRenderTarget;
+		OutputOfMyInput->RenderTargetDesc.DebugName = PassOutputs[0].RenderTargetDesc.DebugName;
+		PassOutputs[0].RenderTargetDesc = OutputOfMyInput->RenderTargetDesc;
+
+		check(PassOutputs[0].RenderTargetDesc.Extent.X);
+		check(PassOutputs[0].RenderTargetDesc.Extent.Y);
+	}
+
 	const FPostProcessSettings& PostprocessSettings = Context.View.FinalPostProcessSettings;
 	const FSceneView& View = Context.View;
 	const FSceneViewFamily& ViewFamily = *(View.Family);
diff --git a/Engine/Source/Runtime/Renderer/Private/PostProcess/PostProcessSubsurface.cpp b/Engine/Source/Runtime/Renderer/Private/PostProcess/PostProcessSubsurface.cpp
--- a/Engine/Source/Runtime/Renderer/Private/PostProcess/PostProcessSubsurface.cpp
+++ b/Engine/Source/Runtime/Renderer/Private/PostProcess/PostProcessSubsurface.cpp
@@ -262,7 +262,7 @@ void FRCPassPostProcessSubsurfaceSetup::Process(FRenderingCompositePassContext&
 
 FPooledRenderTargetDesc FRCPassPostProcessSubsurfaceSetup::ComputeOutputDesc(EPassOutputId InPassOutputId) const
 {
-	FPooledRenderTargetDesc Ret = PassInputs[0].GetOutput()->RenderTargetDesc;
+	FPooledRenderTargetDesc Ret = GSceneRenderTargets.GetSceneColor()->GetDesc();
 
 	Ret.Reset();
 	Ret.DebugName = TEXT("SubsurfaceSetup");
