rework buffered occlusion queries to only need 2 array entries.

diff --git a/Engine/Source/Runtime/Renderer/Private/ScenePrivate.h b/Engine/Source/Runtime/Renderer/Private/ScenePrivate.h
--- a/Engine/Source/Runtime/Renderer/Private/ScenePrivate.h
+++ b/Engine/Source/Runtime/Renderer/Private/ScenePrivate.h
@@ -71,7 +71,7 @@ extern bool ShouldUseGetDynamicMeshElements();
 #include "SpeedTreeWind.h"
 #include "AtmosphereRendering.h"
 
-#if WITH_SLI || PLATFORM_PS4
+#if WITH_SLI || PLATFORM_SHOULD_BUFFER_QUERIES
 #define BUFFERED_OCCLUSION_QUERIES 1
 #endif
 
@@ -90,7 +90,7 @@ public:
 #if BUFFERED_OCCLUSION_QUERIES		
 #	if WITH_SLI
 		// If we're running with SLI, assume throughput is more important than latency, and buffer an extra frame
-		NumBufferedFrames = GNumActiveGPUsForRendering == 1 ? 1 : GNumActiveGPUsForRendering + 1;
+		NumBufferedFrames = GNumActiveGPUsForRendering == 1 ? 1 : GNumActiveGPUsForRendering;
 #	else
 		static const auto NumBufferedQueriesVar = IConsoleManager::Get().FindTConsoleVariableDataInt(TEXT("r.NumBufferedOcclusionQueries"));
 		NumBufferedFrames = NumBufferedQueriesVar->GetValueOnAnyThread();
@@ -102,13 +102,19 @@ public:
 	// get the index of the oldest query based on the current frame and number of buffered frames.
 	static uint32 GetQueryLookupIndex(int32 CurrentFrame, int32 NumBufferedFrames)
 	{
-		const uint32 QueryIndex = (CurrentFrame - NumBufferedFrames + 1) % NumBufferedFrames;
+		// queries are currently always requested earlier in the frame than they are issued.
+		// thus we can always overwrite the oldest query with the current one as we never need them
+		// to coexist.  This saves us a buffer entry.
+		const uint32 QueryIndex = CurrentFrame % NumBufferedFrames;
 		return QueryIndex;
 	}
 
 	// get the index of the query to overwrite for new queries.
 	static uint32 GetQueryIssueIndex(int32 CurrentFrame, int32 NumBufferedFrames)
 	{
+		// queries are currently always requested earlier in the frame than they are issued.
+		// thus we can always overwrite the oldest query with the current one as we never need them
+		// to coexist.  This saves us a buffer entry.
 		const uint32 QueryIndex = CurrentFrame % NumBufferedFrames;
 		return QueryIndex;
 	}
