Fixed velocity rendering in stereo.

Moved velocity rendering outside of post process graph.

diff --git a/Engine/Source/Runtime/Engine/Private/SkeletalRenderGPUSkin.cpp b/Engine/Source/Runtime/Engine/Private/SkeletalRenderGPUSkin.cpp
--- a/Engine/Source/Runtime/Engine/Private/SkeletalRenderGPUSkin.cpp
+++ b/Engine/Source/Runtime/Engine/Private/SkeletalRenderGPUSkin.cpp
@@ -1153,7 +1153,7 @@ uint32 FPreviousPerBoneMotionBlur::AppendData(FBoneSkinning *DataStart, uint32 B
 	}
 }
 
-void FPreviousPerBoneMotionBlur::UnlockData(bool bAdvance)
+void FPreviousPerBoneMotionBlur::UnlockData()
 {
 	if(IsLocked())
 	{
@@ -1163,10 +1163,7 @@ void FPreviousPerBoneMotionBlur::UnlockData(bool bAdvance)
 
 		PerChunkBoneMatricesTexture[GetWriteBufferIndex()].UnlockData();
 
-		if(bAdvance) 
-		{
-			AdvanceBufferIndex();
-		}
+		AdvanceBufferIndex();
 	}
 	
 	{
@@ -1178,7 +1175,7 @@ void FPreviousPerBoneMotionBlur::UnlockData(bool bAdvance)
 
 			if((LogSpawmPrevent % 16) == 0)
 			{
-				UE_LOG(LogSkeletalGPUSkinMesh, Warning, TEXT("Exceeded buffer for per bone motionblur for skinned mesh veclocity rendering. Artifacts can occur. Change Content, increase buffer size or change to use FGlobalDynamicVertexBuffer."));
+				UE_LOG(LogSkeletalGPUSkinMesh, Warning, TEXT("Exceeded buffer for per bone motionblur for skinned mesh velocity rendering. Artifacts can occur. Change Content, increase buffer size or change to use FGlobalDynamicVertexBuffer."));
 			}
 			++LogSpawmPrevent;
 		}
diff --git a/Engine/Source/Runtime/Engine/Private/SkeletalRenderGPUSkin.h b/Engine/Source/Runtime/Engine/Private/SkeletalRenderGPUSkin.h
--- a/Engine/Source/Runtime/Engine/Private/SkeletalRenderGPUSkin.h
+++ b/Engine/Source/Runtime/Engine/Private/SkeletalRenderGPUSkin.h
@@ -202,7 +202,7 @@ public:
 	uint32 AppendData(FBoneSkinning *DataStart, uint32 BoneCount);
 
 	/** only call if LockData() */
-	ENGINE_API void UnlockData(bool bAdvance);
+	ENGINE_API void UnlockData();
 
 	/** @return 0 if there should be no bone based motion blur (no previous data available or it's not active) */
 	FBoneDataVertexBuffer* GetReadData();
diff --git a/Engine/Source/Runtime/Renderer/Private/DeferredShadingRenderer.cpp b/Engine/Source/Runtime/Renderer/Private/DeferredShadingRenderer.cpp
--- a/Engine/Source/Runtime/Renderer/Private/DeferredShadingRenderer.cpp
+++ b/Engine/Source/Runtime/Renderer/Private/DeferredShadingRenderer.cpp
@@ -936,15 +936,26 @@ void FDeferredShadingSceneRenderer::Render(FRHICommandListImmediate& RHICmdList)
 	}
 #endif
 
+	TRefCountPtr<IPooledRenderTarget> VelocityRT;
+
+	// Render the velocities of movable objects
+	RenderVelocities(RHICmdList, VelocityRT);
+
 	// Finish rendering for each view.
 	if(ViewFamily.bResolveScene)
 	{
-		SCOPED_DRAW_EVENT(FinishRendering, DEC_SCENE_ITEMS);
+		SCOPED_DRAW_EVENT(PostProcessing, DEC_SCENE_ITEMS);
 		SCOPE_CYCLE_COUNTER(STAT_FinishRenderViewTargetTime);
-		for(int32 ViewIndex = 0;ViewIndex < Views.Num();ViewIndex++)
-		{	
+
+		for(int32 ViewIndex = 0; ViewIndex < Views.Num(); ViewIndex++)
+		{
 			SCOPED_CONDITIONAL_DRAW_EVENTF(EventView, Views.Num() > 1, DEC_SCENE_ITEMS, TEXT("View%d"), ViewIndex);
-			FinishRenderViewTarget(RHICmdList, &Views[ViewIndex], ViewIndex == (Views.Num() - 1));
+
+			GPostProcessing.Process(RHICmdList, Views[ ViewIndex ], VelocityRT);
+
+			// we rendered to it during the frame, seems we haven't made use of it, because it should be released
+			FSceneViewState* ViewState = (FSceneViewState*)Views[ ViewIndex ].State;
+			check( !ViewState || !ViewState->SeparateTranslucencyRT );
 		}
 	}
 	else
@@ -953,6 +964,8 @@ void FDeferredShadingSceneRenderer::Render(FRHICommandListImmediate& RHICmdList)
 		GSceneRenderTargets.AdjustGBufferRefCount(-1);
 	}
 
+	VelocityRT.SafeRelease();
+
 	RenderFinish(RHICmdList);
 }
 
diff --git a/Engine/Source/Runtime/Renderer/Private/DeferredShadingRenderer.h b/Engine/Source/Runtime/Renderer/Private/DeferredShadingRenderer.h
--- a/Engine/Source/Runtime/Renderer/Private/DeferredShadingRenderer.h
+++ b/Engine/Source/Runtime/Renderer/Private/DeferredShadingRenderer.h
@@ -251,7 +251,7 @@ private:
 	const FProjectedShadowInfo* PrepareTranslucentShadowMap(FRHICommandList& RHICmdList, const FViewInfo& View, FPrimitiveSceneInfo* PrimitiveSceneInfo, bool bSeparateTranslucencyPass);
 
 	/** Renders the velocities of movable objects for the motion blur effect. */
-	void RenderVelocities(FRHICommandListImmediate& RHICmdList, const FViewInfo& View, TRefCountPtr<IPooledRenderTarget>& VelocityRT, bool bLastFrame);
+	void RenderVelocities(FRHICommandListImmediate& RHICmdList, TRefCountPtr<IPooledRenderTarget>& VelocityRT);
 
 	/** Renders the velocities for a subset of movable objects for the motion blur effect. */
 	friend class FRenderVelocityDynamicThreadTask;
@@ -267,12 +267,6 @@ private:
 	/** Updates the downsized depth buffer with the current full resolution depth buffer. */
 	void UpdateDownsampledDepthSurface(FRHICommandList& RHICmdList);
 
-	/**
-	 * Finish rendering a view, writing the contents to ViewFamily.RenderTarget.
-	 * @param View - The view to process.
-	*/
-	void FinishRenderViewTarget(FRHICommandListImmediate& RHICmdList, FViewInfo* View, bool bLastView);
-
 	/**
 	  * Used by RenderLights to figure out if projected shadows need to be rendered to the attenuation buffer.
 	  *
diff --git a/Engine/Source/Runtime/Renderer/Private/ForwardShadingRenderer.cpp b/Engine/Source/Runtime/Renderer/Private/ForwardShadingRenderer.cpp
--- a/Engine/Source/Runtime/Renderer/Private/ForwardShadingRenderer.cpp
+++ b/Engine/Source/Runtime/Renderer/Private/ForwardShadingRenderer.cpp
@@ -161,7 +161,7 @@ void FForwardShadingSceneRenderer::Render(FRHICommandListImmediate& RHICmdList)
 		}
 		else
 		{
-			SCOPED_DRAW_EVENT(FinishRendering, DEC_SCENE_ITEMS);
+			SCOPED_DRAW_EVENT(PostProcessing, DEC_SCENE_ITEMS);
 			SCOPE_CYCLE_COUNTER(STAT_FinishRenderViewTargetTime);
 			for(int32 ViewIndex = 0;ViewIndex < Views.Num();ViewIndex++)
 			{	
diff --git a/Engine/Source/Runtime/Renderer/Private/GammaCorrection.cpp b/Engine/Source/Runtime/Renderer/Private/GammaCorrection.cpp
new file mode 100644
--- /dev/null
+++ b/Engine/Source/Runtime/Renderer/Private/GammaCorrection.cpp
@@ -0,0 +1,149 @@
+// Copyright 1998-2014 Epic Games, Inc. All Rights Reserved.
+
+/*=============================================================================
+	GammaCorrection.cpp
+=============================================================================*/
+
+#include "RendererPrivate.h"
+#include "ScenePrivate.h"
+#include "SceneFilterRendering.h"
+#include "PostProcessing.h"
+#include "SceneUtils.h"
+
+/** Encapsulates the gamma correction pixel shader. */
+class FGammaCorrectionPS : public FGlobalShader
+{
+	DECLARE_SHADER_TYPE(FGammaCorrectionPS,Global);
+
+	static bool ShouldCache(EShaderPlatform Platform)
+	{
+		return true;
+	}
+
+	/** Default constructor. */
+	FGammaCorrectionPS() {}
+
+public:
+
+	FShaderResourceParameter SceneTexture;
+	FShaderResourceParameter SceneTextureSampler;
+	FShaderParameter InverseGamma;
+	FShaderParameter ColorScale;
+	FShaderParameter OverlayColor;
+
+	/** Initialization constructor. */
+	FGammaCorrectionPS(const ShaderMetaType::CompiledShaderInitializerType& Initializer)
+		: FGlobalShader(Initializer)
+	{
+		SceneTexture.Bind(Initializer.ParameterMap,TEXT("SceneColorTexture"));
+		SceneTextureSampler.Bind(Initializer.ParameterMap,TEXT("SceneColorTextureSampler"));
+		InverseGamma.Bind(Initializer.ParameterMap,TEXT("InverseGamma"));
+		ColorScale.Bind(Initializer.ParameterMap,TEXT("ColorScale"));
+		OverlayColor.Bind(Initializer.ParameterMap,TEXT("OverlayColor"));
+	}
+
+	// FShader interface.
+	virtual bool Serialize(FArchive& Ar)
+	{
+		bool bShaderHasOutdatedParameters = FGlobalShader::Serialize(Ar);
+		
+		Ar << SceneTexture << SceneTextureSampler << InverseGamma << ColorScale << OverlayColor;
+		return bShaderHasOutdatedParameters;
+	}
+};
+
+/** Encapsulates the gamma correction vertex shader. */
+class FGammaCorrectionVS : public FGlobalShader
+{
+	DECLARE_SHADER_TYPE(FGammaCorrectionVS,Global);
+
+	static bool ShouldCache(EShaderPlatform Platform)
+	{
+		return true;
+	}
+
+	/** Default constructor. */
+	FGammaCorrectionVS() {}
+
+public:
+
+	/** Initialization constructor. */
+	FGammaCorrectionVS(const ShaderMetaType::CompiledShaderInitializerType& Initializer)
+		:	FGlobalShader(Initializer)
+	{
+	}
+};
+
+IMPLEMENT_SHADER_TYPE(,FGammaCorrectionPS,TEXT("GammaCorrection"),TEXT("MainPS"),SF_Pixel);
+IMPLEMENT_SHADER_TYPE(,FGammaCorrectionVS,TEXT("GammaCorrection"),TEXT("MainVS"),SF_Vertex);
+
+// TODO: REMOVE if no longer needed:
+void FSceneRenderer::GammaCorrectToViewportRenderTarget(FRHICommandList& RHICmdList, const FViewInfo* View, float OverrideGamma)
+{
+	// Set the view family's render target/viewport.
+	SetRenderTarget(RHICmdList, ViewFamily.RenderTarget->GetRenderTargetTexture(), FTextureRHIRef());
+
+	// Deferred the clear until here so the garbage left in the non rendered regions by the post process effects do not show up
+	if( ViewFamily.bDeferClear )
+	{
+		RHICmdList.Clear(true, FLinearColor::Black, false, 0.0f, false, 0, FIntRect());
+		ViewFamily.bDeferClear = false;
+	}
+
+	SCOPED_DRAW_EVENT(GammaCorrection, DEC_SCENE_ITEMS);
+
+	// turn off culling and blending
+	RHICmdList.SetRasterizerState(TStaticRasterizerState<FM_Solid, CM_None>::GetRHI());
+	RHICmdList.SetBlendState(TStaticBlendState<>::GetRHI());
+
+	// turn off depth reads/writes
+	RHICmdList.SetDepthStencilState(TStaticDepthStencilState<false, CF_Always>::GetRHI());
+
+	TShaderMapRef<FGammaCorrectionVS> VertexShader(View->ShaderMap);
+	TShaderMapRef<FGammaCorrectionPS> PixelShader(View->ShaderMap);
+
+	static FGlobalBoundShaderState PostProcessBoundShaderState;
+	
+	SetGlobalBoundShaderState(RHICmdList, View->GetFeatureLevel(), PostProcessBoundShaderState, GFilterVertexDeclaration.VertexDeclarationRHI, *VertexShader, *PixelShader);
+
+	float InvDisplayGamma = 1.0f / ViewFamily.RenderTarget->GetDisplayGamma();
+
+	if (OverrideGamma != 0)
+	{
+		InvDisplayGamma = 1 / OverrideGamma;
+	}
+
+	const FPixelShaderRHIParamRef ShaderRHI = PixelShader->GetPixelShader();
+
+	SetShaderValue(
+		RHICmdList, 
+		ShaderRHI,
+		PixelShader->InverseGamma,
+		InvDisplayGamma
+		);
+	SetShaderValue(RHICmdList, ShaderRHI,PixelShader->ColorScale,View->ColorScale);
+	SetShaderValue(RHICmdList, ShaderRHI,PixelShader->OverlayColor,View->OverlayColor);
+
+	const FTextureRHIRef DesiredSceneColorTexture = GSceneRenderTargets.GetSceneColorTexture();
+
+	SetTextureParameter(
+		RHICmdList, 
+		ShaderRHI,
+		PixelShader->SceneTexture,
+		PixelShader->SceneTextureSampler,
+		TStaticSamplerState<SF_Bilinear>::GetRHI(),
+		DesiredSceneColorTexture
+		);
+
+	// Draw a quad mapping scene color to the view's render target
+	DrawRectangle(
+		RHICmdList,
+		View->UnscaledViewRect.Min.X,View->UnscaledViewRect.Min.Y,
+		View->UnscaledViewRect.Width(),View->UnscaledViewRect.Height(),
+		View->ViewRect.Min.X,View->ViewRect.Min.Y,
+		View->ViewRect.Width(),View->ViewRect.Height(),
+		ViewFamily.RenderTarget->GetSizeXY(),
+		GSceneRenderTargets.GetBufferSizeXY(),
+		*VertexShader,
+		EDRF_UseTriangleOptimization);
+}
diff --git a/Engine/Source/Runtime/Renderer/Private/PostProcess/PostProcessing.cpp b/Engine/Source/Runtime/Renderer/Private/PostProcess/PostProcessing.cpp
--- a/Engine/Source/Runtime/Renderer/Private/PostProcess/PostProcessing.cpp
+++ b/Engine/Source/Runtime/Renderer/Private/PostProcess/PostProcessing.cpp
@@ -810,9 +810,7 @@ void FPostProcessing::Process(FRHICommandListImmediate& RHICmdList, FViewInfo& V
 			FRenderingCompositeOutputRef VelocityInput;
 			if(VelocityRT)
 			{
-				VelocityInput = Context.Graph.RegisterPass(new(FMemStack::Get()) FRCPassPostProcessInputSingleUse(VelocityRT));
-				// the graph keeps a reference until it's no longer needed
-				VelocityRT.SafeRelease();
+				VelocityInput = Context.Graph.RegisterPass(new(FMemStack::Get()) FRCPassPostProcessInput(VelocityRT));
 			}
 			
 			if( Context.View.FinalPostProcessSettings.AntiAliasingMethod != AAM_TemporalAA && ViewState )			
diff --git a/Engine/Source/Runtime/Renderer/Private/ScenePostProcessing.cpp b/Engine/Source/Runtime/Renderer/Private/ScenePostProcessing.cpp
deleted file mode 100644
diff --git a/Engine/Source/Runtime/Renderer/Private/VelocityRendering.cpp b/Engine/Source/Runtime/Renderer/Private/VelocityRendering.cpp
--- a/Engine/Source/Runtime/Renderer/Private/VelocityRendering.cpp
+++ b/Engine/Source/Runtime/Renderer/Private/VelocityRendering.cpp
@@ -600,14 +600,22 @@ void FDeferredShadingSceneRenderer::RenderVelocitiesInner(FRHICommandListImmedia
 	}
 }
 
-void FDeferredShadingSceneRenderer::RenderVelocities(FRHICommandListImmediate& RHICmdList, const FViewInfo& View, TRefCountPtr<IPooledRenderTarget>& VelocityRT, bool bLastFrame)
+void FDeferredShadingSceneRenderer::RenderVelocities(FRHICommandListImmediate& RHICmdList, TRefCountPtr<IPooledRenderTarget>& VelocityRT)
 {
 	check(FeatureLevel >= ERHIFeatureLevel::SM4);
 	SCOPE_CYCLE_COUNTER(STAT_RenderVelocities);
 
-	bool bTemporalAA = (View.FinalPostProcessSettings.AntiAliasingMethod == AAM_TemporalAA) && !View.bCameraCut;
+	bool bNeedsVelocity = false;
+	for(int32 ViewIndex = 0; ViewIndex < Views.Num(); ViewIndex++)
+	{
+		const FViewInfo& View = Views[ViewIndex];
+	
+		bool bTemporalAA = (View.FinalPostProcessSettings.AntiAliasingMethod == AAM_TemporalAA) && !View.bCameraCut;
+		bool bMotionBlur = IsMotionBlurEnabled(View);
 
-	if(!IsMotionBlurEnabled(View) && !bTemporalAA)
+		bNeedsVelocity |= bMotionBlur || bTemporalAA;
+	}
+	if( !bNeedsVelocity )
 	{
 		return;
 	}
@@ -617,44 +625,42 @@ void FDeferredShadingSceneRenderer::RenderVelocities(FRHICommandListImmediate& R
 	const FIntPoint BufferSize = GSceneRenderTargets.GetBufferSizeXY();
 	const FIntPoint VelocityBufferSize = BufferSize;		// full resolution so we can reuse the existing full res z buffer
 
-	if(!VelocityRT)
-	{
-		FPooledRenderTargetDesc Desc(FPooledRenderTargetDesc::Create2DDesc(VelocityBufferSize, PF_G16R16, TexCreate_None, TexCreate_RenderTargetable, false));
-
-		GRenderTargetPool.FindFreeElement(Desc, VelocityRT, TEXT("Velocity"));
-	}
+	FPooledRenderTargetDesc Desc(FPooledRenderTargetDesc::Create2DDesc(VelocityBufferSize, PF_G16R16, TexCreate_None, TexCreate_RenderTargetable, false));
+	GRenderTargetPool.FindFreeElement(Desc, VelocityRT, TEXT("Velocity"));
 
 	GPrevPerBoneMotionBlur.LockData();
-
-	const uint32 MinX = View.ViewRect.Min.X * VelocityBufferSize.X / BufferSize.X;
-	const uint32 MinY = View.ViewRect.Min.Y * VelocityBufferSize.Y / BufferSize.Y;
-	const uint32 MaxX = View.ViewRect.Max.X * VelocityBufferSize.X / BufferSize.X;
-	const uint32 MaxY = View.ViewRect.Max.Y * VelocityBufferSize.Y / BufferSize.Y;
+	
 	SetRenderTarget(RHICmdList, VelocityRT->GetRenderTargetItem().TargetableTexture, GSceneRenderTargets.GetSceneDepthTexture());
-
-	RHICmdList.SetViewport(MinX, MinY, 0.0f, MaxX, MaxY, 1.0f);
-
+	
 	// Clear the velocity buffer (0.0f means "use static background velocity").
 	RHICmdList.Clear(true, FLinearColor::Black, false, 1.0f, false, 0, FIntRect());
 
-	// Blending is not supported with the velocity buffer format on other platforms
-	// opaque velocities in R|G channels, B is used to identify object motion
-	RHICmdList.SetBlendState(TStaticBlendState<CW_RGBA>::GetRHI());
-	// Use depth tests, no z-writes, backface-culling. 
+	for(int32 ViewIndex = 0; ViewIndex < Views.Num(); ViewIndex++)
+	{
+		const FViewInfo& View = Views[ViewIndex];
+
+		const uint32 MinX = View.ViewRect.Min.X * VelocityBufferSize.X / BufferSize.X;
+		const uint32 MinY = View.ViewRect.Min.Y * VelocityBufferSize.Y / BufferSize.Y;
+		const uint32 MaxX = View.ViewRect.Max.X * VelocityBufferSize.X / BufferSize.X;
+		const uint32 MaxY = View.ViewRect.Max.Y * VelocityBufferSize.Y / BufferSize.Y;
+		
+		RHICmdList.SetViewport(MinX, MinY, 0.0f, MaxX, MaxY, 1.0f);
 
-	// Note, this is a reversed Z depth surface, using CF_GreaterEqual.
-	RHICmdList.SetDepthStencilState(TStaticDepthStencilState<false,CF_GreaterEqual>::GetRHI());
-			
-	RHICmdList.SetRasterizerState(GetStaticRasterizerState<true>(FM_Solid, CM_CW));
+		RHICmdList.SetBlendState(TStaticBlendState<CW_RGBA>::GetRHI());
+		// Note, this is a reversed Z depth surface, using CF_GreaterEqual.
+		RHICmdList.SetDepthStencilState(TStaticDepthStencilState<false,CF_GreaterEqual>::GetRHI());
+		RHICmdList.SetRasterizerState(GetStaticRasterizerState<true>(FM_Solid, CM_CW));
 
-	RenderVelocitiesInner(RHICmdList, View);
+		RenderVelocitiesInner(RHICmdList, View);
+	}
 
 	RHICmdList.CopyToResolveTarget(VelocityRT->GetRenderTargetItem().TargetableTexture, VelocityRT->GetRenderTargetItem().ShaderResourceTexture, false, FResolveParams());
 
 	// restore any color write state changes
 	RHICmdList.SetBlendState(TStaticBlendState<>::GetRHI());
 	RHICmdList.SetRasterizerState(TStaticRasterizerState<>::GetRHI());
-	GPrevPerBoneMotionBlur.UnlockData(bLastFrame);
+	
+	GPrevPerBoneMotionBlur.UnlockData();
 
 	// to be able to observe results with VisualizeTexture
 	GRenderTargetPool.VisualizeTexture.SetCheckPoint(RHICmdList, VelocityRT);
