Fix non-unity in (GPostProcessing)

diff --git a/Engine/Source/Runtime/Renderer/Private/DeferredShadingRenderer.cpp b/Engine/Source/Runtime/Renderer/Private/DeferredShadingRenderer.cpp
--- a/Engine/Source/Runtime/Renderer/Private/DeferredShadingRenderer.cpp
+++ b/Engine/Source/Runtime/Renderer/Private/DeferredShadingRenderer.cpp
@@ -18,6 +18,7 @@
 #include "DeferredShadingRenderer.h"
 #include "SceneUtils.h"
 #include "DistanceFieldLightingShared.h"
+#include "PostProcess/PostProcessing.h"
 
 TAutoConsoleVariable<int32> CVarEarlyZPass(
 	TEXT("r.EarlyZPass"),
