minor optimization, run LUTBlender only if really needed

diff --git a/Engine/Source/Runtime/Renderer/Private/PostProcess/PostProcessCombineLUTs.cpp b/Engine/Source/Runtime/Renderer/Private/PostProcess/PostProcessCombineLUTs.cpp
--- a/Engine/Source/Runtime/Renderer/Private/PostProcess/PostProcessCombineLUTs.cpp
+++ b/Engine/Source/Runtime/Renderer/Private/PostProcess/PostProcessCombineLUTs.cpp
@@ -457,8 +457,10 @@ FPooledRenderTargetDesc FRCPassPostProcessCombineLUTs::ComputeOutputDesc(EPassOu
 	return Ret;
 }
 
-bool FRCPassPostProcessCombineLUTs::IsColorGradingLUTNeeded(const FRenderingCompositePassContext* RESTRICT Context)
+bool FRCPassPostProcessCombineLUTs::IsColorGradingLUTNeeded(const FViewInfo* RESTRICT View)
 {
+	check(View);
+
 	FColorTransform ColorTransform;
 	ColorTransform.MinValue = FMath::Clamp(CVarColorMin.GetValueOnRenderThread(), -10.0f, 10.0f);
 	ColorTransform.MidValue = FMath::Clamp(CVarColorMid.GetValueOnRenderThread(), -10.0f, 10.0f);
@@ -477,19 +479,19 @@ bool FRCPassPostProcessCombineLUTs::IsColorGradingLUTNeeded(const FRenderingComp
 		return true;
 	}
 
-	if(Context->View.ColorScale.R != 1.0f || Context->View.ColorScale.G != 1.0f || Context->View.ColorScale.B != 1.0f)
+	if(View->ColorScale.R != 1.0f || View->ColorScale.G != 1.0f || View->ColorScale.B != 1.0f)
 	{
 		return true;
 	}
 
-	if(Context->View.OverlayColor.A > (1.0f/512.0f))
+	if(View->OverlayColor.A > (1.0f/512.0f))
 	{
 		return true;
 	}
 
-	if(Context->GetFeatureLevel() >= ERHIFeatureLevel::SM4)
+	if(View->GetFeatureLevel() >= ERHIFeatureLevel::SM4)
 	{
-		const FFinalPostProcessSettings Settings = Context->View.FinalPostProcessSettings;
+		const FFinalPostProcessSettings Settings = View->FinalPostProcessSettings;
 		for(uint32 i = 0; i < (uint32)Settings.ContributingLUTs.Num(); ++i)
 		{
 			UTexture* LUTTexture = Settings.ContributingLUTs[i].LUTTexture;
diff --git a/Engine/Source/Runtime/Renderer/Private/PostProcess/PostProcessCombineLUTs.h b/Engine/Source/Runtime/Renderer/Private/PostProcess/PostProcessCombineLUTs.h
--- a/Engine/Source/Runtime/Renderer/Private/PostProcess/PostProcessCombineLUTs.h
+++ b/Engine/Source/Runtime/Renderer/Private/PostProcess/PostProcessCombineLUTs.h
@@ -28,7 +28,7 @@ public:
 	/** @return 0xffffffff if not found */
 	uint32 FindIndex(const FFinalPostProcessSettings& Settings, UTexture* Tex) const;
 
-	static bool IsColorGradingLUTNeeded(const FRenderingCompositePassContext* RESTRICT Context);
+	static bool IsColorGradingLUTNeeded(const FViewInfo* RESTRICT View);
 
 private:
 
diff --git a/Engine/Source/Runtime/Renderer/Private/PostProcess/PostProcessTonemap.cpp b/Engine/Source/Runtime/Renderer/Private/PostProcess/PostProcessTonemap.cpp
--- a/Engine/Source/Runtime/Renderer/Private/PostProcess/PostProcessTonemap.cpp
+++ b/Engine/Source/Runtime/Renderer/Private/PostProcess/PostProcessTonemap.cpp
@@ -607,11 +607,13 @@ static uint32 TonemapperFindLeastExpensive(uint32* RESTRICT Table, uint32 TableE
 }
 
 // Common conversion of engine settings into a bitmask which describes the shader options required.
-static uint32 TonemapperGenerateBitmask(const FRenderingCompositePassContext* RESTRICT Context, bool bGammaOnly, bool bMobile)
+static uint32 TonemapperGenerateBitmask(const FViewInfo* RESTRICT View, bool bGammaOnly, bool bMobile)
 {
+	check(View);
+
 	bGammaOnly |= !IsMobileHDR();
 
-	const FSceneViewFamily* RESTRICT Family = Context->View.Family;
+	const FSceneViewFamily* RESTRICT Family = View->Family;
 	if(
 		bGammaOnly ||
 		(Family->EngineShowFlags.Tonemapper == 0) || 
@@ -622,7 +624,7 @@ static uint32 TonemapperGenerateBitmask(const FRenderingCompositePassContext* RE
 
 	uint32 Bitmask = 0;
 
-	const FPostProcessSettings* RESTRICT Settings = &(Context->View.FinalPostProcessSettings);
+	const FPostProcessSettings* RESTRICT Settings = &(View->FinalPostProcessSettings);
 
 	FVector MixerR(Settings->FilmChannelMixerRed);
 	FVector MixerG(Settings->FilmChannelMixerGreen);
@@ -652,10 +654,10 @@ static uint32 TonemapperGenerateBitmask(const FRenderingCompositePassContext* RE
 
 // Common post.
 // These are separated because mosiac mode doesn't support them.
-static uint32 TonemapperGenerateBitmaskPost(const FRenderingCompositePassContext* RESTRICT Context)
+static uint32 TonemapperGenerateBitmaskPost(const FViewInfo* RESTRICT View)
 {
-	const FPostProcessSettings* RESTRICT Settings = &(Context->View.FinalPostProcessSettings);
-	const FSceneViewFamily* RESTRICT Family = Context->View.Family;
+	const FPostProcessSettings* RESTRICT Settings = &(View->FinalPostProcessSettings);
+	const FSceneViewFamily* RESTRICT Family = View->Family;
 
 	uint32 Bitmask = (Settings->GrainJitter > 0.0f) ? TonemapperGrainJitter : 0;
 
@@ -665,9 +667,9 @@ static uint32 TonemapperGenerateBitmaskPost(const FRenderingCompositePassContext
 }
 
 // PC only.
-static uint32 TonemapperGenerateBitmaskPC(const FRenderingCompositePassContext* RESTRICT Context, bool bGammaOnly)
+static uint32 TonemapperGenerateBitmaskPC(const FViewInfo* RESTRICT View, bool bGammaOnly, bool bUseLUT)
 {
-	uint32 Bitmask = TonemapperGenerateBitmask(Context, bGammaOnly, false);
+	uint32 Bitmask = TonemapperGenerateBitmask(View, bGammaOnly, false);
 
 	// Must early exit if gamma only.
 	if(Bitmask == TonemapperGammaOnly) 
@@ -688,27 +690,29 @@ static uint32 TonemapperGenerateBitmaskPC(const FRenderingCompositePassContext*
 	static const auto CVar = IConsoleManager::Get().FindTConsoleVariableDataInt(TEXT("r.SceneColorFringeQuality")); 
 
 	int32 FringeQuality = CVar->GetValueOnRenderThread();
-	if(Context->View.Family->EngineShowFlags.SceneColorFringe
+	if(View->Family->EngineShowFlags.SceneColorFringe
 		// Removed this from the camera imperfections toggle because this no longer takes an extra pass.
 		// && Context->View.Family->EngineShowFlags.CameraImperfections
-		&& Context->View.FinalPostProcessSettings.SceneFringeIntensity > 0.01f
+		&& View->FinalPostProcessSettings.SceneFringeIntensity > 0.01f
 		&& FringeQuality > 0)
 	{
 		Bitmask |= TonemapperColorFringe;
 	}
 
-	if(FRCPassPostProcessCombineLUTs::IsColorGradingLUTNeeded(Context))
+	if(bUseLUT)
 	{
 		Bitmask |= TonemapperColorGrading;
 	}
 
-	return Bitmask + TonemapperGenerateBitmaskPost(Context);
+	return Bitmask + TonemapperGenerateBitmaskPost(View);
 }
 
 // Mobile only.
-static uint32 TonemapperGenerateBitmaskMobile(const FRenderingCompositePassContext* RESTRICT Context, bool bGammaOnly)
+static uint32 TonemapperGenerateBitmaskMobile(const FViewInfo* RESTRICT View, bool bGammaOnly)
 {
-	uint32 Bitmask = TonemapperGenerateBitmask(Context, bGammaOnly, true);
+	check(View);
+
+	uint32 Bitmask = TonemapperGenerateBitmask(View, bGammaOnly, true);
 
 	bool bUseMosaic = IsMobileHDR32bpp();
 
@@ -733,9 +737,9 @@ static uint32 TonemapperGenerateBitmaskMobile(const FRenderingCompositePassConte
 	// Only add mobile post if FP16 is supported.
 	if(GSupportsRenderTargetFormat_PF_FloatRGBA)
 	{
-		Bitmask += TonemapperGenerateBitmaskPost(Context);
-		Bitmask += (Context->View.FinalPostProcessSettings.DepthOfFieldScale > 0.0f) ? TonemapperDOF         : 0;
-		Bitmask += (Context->View.bLightShaftUse)                                    ? TonemapperLightShafts : 0;
+		Bitmask += TonemapperGenerateBitmaskPost(View);
+		Bitmask += (View->FinalPostProcessSettings.DepthOfFieldScale > 0.0f) ? TonemapperDOF         : 0;
+		Bitmask += (View->bLightShaftUse)                                    ? TonemapperLightShafts : 0;
 
 		// Mobile is not supporting grain quantization and grain jitter currently.
 		Bitmask &= ~(TonemapperGrainQuantization | TonemapperGrainJitter);
@@ -1176,9 +1180,18 @@ public:
 IMPLEMENT_SHADER_TYPE(,FPostProcessTonemapVS,TEXT("PostProcessTonemap"),TEXT("MainVS"),SF_Vertex);
 
 
-FRCPassPostProcessTonemap::FRCPassPostProcessTonemap(bool bInDoGammaOnly)
+FRCPassPostProcessTonemap::FRCPassPostProcessTonemap(const FViewInfo& View, bool bInDoGammaOnly)
 	: bDoGammaOnly(bInDoGammaOnly)
 {
+	uint32 ConfigBitmask = TonemapperGenerateBitmaskPC(&View, bDoGammaOnly, FRCPassPostProcessCombineLUTs::IsColorGradingLUTNeeded(&View));
+	ConfigIndex = TonemapperFindLeastExpensive(TonemapperConfBitmaskPC, sizeof(TonemapperConfBitmaskPC)/4, TonemapperCostTab, ConfigBitmask);
+}
+
+bool FRCPassPostProcessTonemap::IsLUTNeeded() const
+{
+	uint32 ConfigBitmask = TonemapperConfBitmaskMobile[ConfigIndex];
+
+	return TonemapperIsDefined(ConfigBitmask, TonemapperColorGrading) != 0;
 }
 
 template <uint32 ConfigIndex>
@@ -1198,9 +1211,6 @@ static void SetShaderTempl(const FRenderingCompositePassContext& Context)
 
 void FRCPassPostProcessTonemap::Process(FRenderingCompositePassContext& Context)
 {
-	uint32 ConfigBitmask = TonemapperGenerateBitmaskPC(&Context, bDoGammaOnly);
-	uint32 ConfigIndex = TonemapperFindLeastExpensive(TonemapperConfBitmaskPC, sizeof(TonemapperConfBitmaskPC)/4, TonemapperCostTab, ConfigBitmask);
-
 	SCOPED_DRAW_EVENTF(PostProcessTonemap, DEC_SCENE_ITEMS, TEXT("Tonemapper#%d%s"), ConfigIndex, bDoGammaOnly ? TEXT(" GammaOnly") : TEXT(""));
 
 	const FPooledRenderTargetDesc* InputDesc = GetInputDesc(ePId_Input0);
@@ -1563,12 +1573,17 @@ static void SetShaderTemplES2(const FRenderingCompositePassContext& Context, boo
 	PixelShader->SetPS(Context);
 }
 
-void FRCPassPostProcessTonemapES2::Process(FRenderingCompositePassContext& Context)
+
+FRCPassPostProcessTonemapES2::FRCPassPostProcessTonemapES2(const FViewInfo& View, bool bInUsedFramebufferFetch) 
+	: bUsedFramebufferFetch(bInUsedFramebufferFetch)
 {
-	uint32 ConfigBitmask = TonemapperGenerateBitmaskMobile(&Context, false);
-	uint32 ConfigIndex = TonemapperFindLeastExpensive(TonemapperConfBitmaskMobile, sizeof(TonemapperConfBitmaskMobile)/4, TonemapperCostTab, ConfigBitmask);
+	uint32 ConfigBitmask = TonemapperGenerateBitmaskMobile(&View, false);
+	ConfigIndex = TonemapperFindLeastExpensive(TonemapperConfBitmaskMobile, sizeof(TonemapperConfBitmaskMobile)/4, TonemapperCostTab, ConfigBitmask);
+}
 
-	SCOPED_DRAW_EVENT(PostProcessTonemap, DEC_SCENE_ITEMS);
+void FRCPassPostProcessTonemapES2::Process(FRenderingCompositePassContext& Context)
+{
+	SCOPED_DRAW_EVENTF(PostProcessTonemap, DEC_SCENE_ITEMS, TEXT("Tonemapper#%d%s"), ConfigIndex, bUsedFramebufferFetch ? TEXT(" FramebufferFetch=0") : TEXT("FramebufferFetch=1"));
 
 	const FPooledRenderTargetDesc* InputDesc = GetInputDesc(ePId_Input0);
 
diff --git a/Engine/Source/Runtime/Renderer/Private/PostProcess/PostProcessTonemap.h b/Engine/Source/Runtime/Renderer/Private/PostProcess/PostProcessTonemap.h
--- a/Engine/Source/Runtime/Renderer/Private/PostProcess/PostProcessTonemap.h
+++ b/Engine/Source/Runtime/Renderer/Private/PostProcess/PostProcessTonemap.h
@@ -41,7 +41,9 @@ class FRCPassPostProcessTonemap : public TRenderingCompositePassBase<4, 1>
 {
 public:
 	// constructor
-	FRCPassPostProcessTonemap(bool bInDoGammaOnly = false);
+	FRCPassPostProcessTonemap(const FViewInfo& View, bool bInDoGammaOnly = false);
+
+	bool IsLUTNeeded() const;
 
 	// interface FRenderingCompositePass ---------
 
@@ -50,8 +52,11 @@ public:
 	virtual FPooledRenderTargetDesc ComputeOutputDesc(EPassOutputId InPassOutputId) const;
 
 private:
-	void SetShader(const FRenderingCompositePassContext& Context);
 	bool bDoGammaOnly;
+	// set in constructor
+	uint32 ConfigIndex;
+
+	void SetShader(const FRenderingCompositePassContext& Context);
 };
 
 // derives from TRenderingCompositePassBase<InputCount, OutputCount>
@@ -61,13 +66,19 @@ private:
 class FRCPassPostProcessTonemapES2 : public TRenderingCompositePassBase<3, 1>
 {
 public:
-	FRCPassPostProcessTonemapES2(bool bInUsedFramebufferFetch) : bUsedFramebufferFetch(bInUsedFramebufferFetch) { }
+	FRCPassPostProcessTonemapES2(const FViewInfo& View, bool bInUsedFramebufferFetch);
+
+	// interface FRenderingCompositePass ---------
+
 	virtual void Process(FRenderingCompositePassContext& Context);
 	virtual void Release() override { delete this; }
 	virtual FPooledRenderTargetDesc ComputeOutputDesc(EPassOutputId InPassOutputId) const;
 
 private:
 	bool bUsedFramebufferFetch;
+	// set in constructor
+	uint32 ConfigIndex;
+
 	void SetShader(const FRenderingCompositePassContext& Context);
 };
 
diff --git a/Engine/Source/Runtime/Renderer/Private/PostProcess/PostProcessing.cpp b/Engine/Source/Runtime/Renderer/Private/PostProcess/PostProcessing.cpp
--- a/Engine/Source/Runtime/Renderer/Private/PostProcess/PostProcessing.cpp
+++ b/Engine/Source/Runtime/Renderer/Private/PostProcess/PostProcessing.cpp
@@ -144,27 +144,23 @@ static FRenderingCompositeOutputRef RenderBloom(
 	return RenderGaussianBlur(Context, EPostProcessRectSource::GBS_ViewRect, TEXT("BloomBlurX"), TEXT("BloomBlurY"), PreviousBloom, Size, Tint, Additive);
 }
 
-
-static FRenderingCompositeOutputRef CombineColorGradingLUTs(
-	FPostprocessContext& Context)
-{
-	FRenderingCompositePass* Pass = Context.Graph.RegisterPass(new(FMemStack::Get()) FRCPassPostProcessCombineLUTs(Context.View.GetShaderPlatform()));
-
-	return FRenderingCompositeOutputRef(Pass);
-}
-
 static void AddTonemapper(
 	FPostprocessContext& Context,
 	const FRenderingCompositeOutputRef& BloomOutputCombined,
-	const FRenderingCompositeOutputRef& EyeAdaptation,
-	const FRenderingCompositeOutputRef& CombinedLUT)
+	const FRenderingCompositeOutputRef& EyeAdaptation)
 {
-	FRenderingCompositePass* PostProcessTonemap = Context.Graph.RegisterPass(new(FMemStack::Get()) FRCPassPostProcessTonemap());
+	FRCPassPostProcessTonemap* PostProcessTonemap = Context.Graph.RegisterPass(new(FMemStack::Get()) FRCPassPostProcessTonemap(Context.View));
 
 	PostProcessTonemap->SetInput(ePId_Input0, Context.FinalOutput);
 	PostProcessTonemap->SetInput(ePId_Input1, BloomOutputCombined);
 	PostProcessTonemap->SetInput(ePId_Input2, EyeAdaptation);
-	PostProcessTonemap->SetInput(ePId_Input3, CombinedLUT);
+
+	if(PostProcessTonemap->IsLUTNeeded())
+	{
+		FRenderingCompositePass* CombinedLUT = Context.Graph.RegisterPass(new(FMemStack::Get()) FRCPassPostProcessCombineLUTs(Context.View.GetShaderPlatform()));
+
+		PostProcessTonemap->SetInput(ePId_Input3, CombinedLUT);
+	}
 
 	Context.FinalOutput = FRenderingCompositeOutputRef(PostProcessTonemap);
 }
@@ -183,7 +179,7 @@ static void AddSelectionOutline(FPostprocessContext& Context)
 
 static void AddGammaOnlyTonemapper(FPostprocessContext& Context)
 {
-	FRenderingCompositePass* PostProcessTonemap = Context.Graph.RegisterPass(new(FMemStack::Get()) FRCPassPostProcessTonemap(true));
+	FRenderingCompositePass* PostProcessTonemap = Context.Graph.RegisterPass(new(FMemStack::Get()) FRCPassPostProcessTonemap(Context.View, true));
 
 	PostProcessTonemap->SetInput(ePId_Input0, Context.FinalOutput);
 
@@ -993,11 +989,7 @@ void FPostProcessing::Process(FRHICommandListImmediate& RHICmdList, FViewInfo& V
 
 			if(bAllowTonemapper && FeatureLevel >= ERHIFeatureLevel::SM4)
 			{
-				FRenderingCompositeOutputRef CombinedLUT;
-
-				CombinedLUT = CombineColorGradingLUTs(Context);
-
-				AddTonemapper(Context, BloomOutputCombined, EyeAdaptation, CombinedLUT);
+				AddTonemapper(Context, BloomOutputCombined, EyeAdaptation);
 			}
 
 			if(AntiAliasingMethod == AAM_FXAA)
@@ -1476,7 +1468,7 @@ void FPostProcessing::ProcessES2(FRHICommandListImmediate& RHICmdList, FViewInfo
 		}
 
 		// Must run to blit to back buffer even if post processing is off.
-		FRenderingCompositePass* PostProcessTonemap = Context.Graph.RegisterPass(new(FMemStack::Get()) FRCPassPostProcessTonemapES2(bUsedFramebufferFetch));
+		FRenderingCompositePass* PostProcessTonemap = Context.Graph.RegisterPass(new(FMemStack::Get()) FRCPassPostProcessTonemapES2(Context.View, bUsedFramebufferFetch));
 		PostProcessTonemap->SetInput(ePId_Input0, Context.FinalOutput);
 		PostProcessTonemap->SetInput(ePId_Input1, BloomOutput);
 		PostProcessTonemap->SetInput(ePId_Input2, DofOutput);
