Distance field area shadows
* This new shadow method is computed by cone tracing through the mesh distance fields when the 'bGenerateMeshDistanceFields' project setting is enabled
* It has high pixel cost but low object and triangle cost and is useful for covering a large depth range for a dynamic sun.  It inherits all the limitations of mesh distance fields (no vertex animation, only static meshes supported for now, only uniform scaling).  Components can use an inset shadow to use shadowmapping even while the ray traced shadows are used for the rest of the world.
* Currently only implemented for directional lights, as a final cascade after the shadowmapped cascades
* Current implementation groups pixels in a tile into a far depth range and near, tile culls the objects in the scene with those two groups. ~6ms in AITestBed_2 at 1080p on 7870 with a view range of 200k.
* bAllowMeshDistanceFieldRepresentations project setting renamed to bGenerateMeshDistanceFields

diff --git a/Engine/Shaders/DistanceFieldLightingShared.usf b/Engine/Shaders/DistanceFieldLightingShared.usf
--- a/Engine/Shaders/DistanceFieldLightingShared.usf
+++ b/Engine/Shaders/DistanceFieldLightingShared.usf
@@ -1,7 +1,7 @@
-// Copyright 1998-2013 Epic Games, Inc. All Rights Reserved.
+// Copyright 1998-2014 Epic Games, Inc. All Rights Reserved.
 
 /*=============================================================================
-
+	DistanceFieldLightingShared.usf
 =============================================================================*/
 
 #ifndef THREADGROUP_SIZEX
diff --git a/Engine/Shaders/DistanceFieldShadowing.usf b/Engine/Shaders/DistanceFieldShadowing.usf
new file mode 100644
--- /dev/null
+++ b/Engine/Shaders/DistanceFieldShadowing.usf
@@ -0,0 +1,308 @@
+// Copyright 1998-2014 Epic Games, Inc. All Rights Reserved.
+
+/*=============================================================================
+	DistanceFieldShadowing.usf
+=============================================================================*/
+
+#include "Common.usf"
+#include "DeferredShadingCommon.usf"
+#include "DistanceFieldLightingShared.usf"
+
+RWTexture2D<float2> RWShadowFactors;
+
+bool RayHitSphere( float3 RayOrigin, float3 UnitRayDirection, float3 SphereCenter, float SphereRadius )
+{
+	float3 ClosestPointOnRay = max( 0, dot( SphereCenter - RayOrigin, UnitRayDirection ) ) * UnitRayDirection;
+	float3 CenterToRay = RayOrigin + ClosestPointOnRay - SphereCenter;
+	return dot( CenterToRay, CenterToRay ) <= Square( SphereRadius );
+}
+
+float2 NumGroups;
+float3 LightDirection;
+
+/** Min and Max depth for this tile. */
+groupshared uint IntegerTileMinZ;
+groupshared uint IntegerTileMaxZ;
+
+/** Inner Min and Max depth for this tile. */
+groupshared uint IntegerTileMinZ2;
+groupshared uint IntegerTileMaxZ2;
+
+/** Number of objects affecting the tile, after culling. */
+groupshared uint TileNumObjects0;
+groupshared uint TileNumObjects1;
+
+#define MAX_INTERSECTING_OBJECTS 256
+groupshared uint IntersectingObjectIndices[MAX_INTERSECTING_OBJECTS * 2];
+
+[numthreads(THREADGROUP_SIZEX, THREADGROUP_SIZEY, 1)]
+void DistanceFieldShadowingCS(
+	uint3 GroupId : SV_GroupID,
+	uint3 DispatchThreadId : SV_DispatchThreadID,
+    uint3 GroupThreadId : SV_GroupThreadID) 
+{
+	uint ThreadIndex = GroupThreadId.y * THREADGROUP_SIZEX + GroupThreadId.x;
+
+	float2 ScreenUV = float2((DispatchThreadId.xy * DOWNSAMPLE_FACTOR + View.ViewRectMin.xy + .5f) * View.ViewSizeAndSceneTexelSize.zw);
+	float2 ScreenPosition = (ScreenUV.xy - View.ScreenPositionScaleBias.wz) / View.ScreenPositionScaleBias.xy;
+
+	float SceneDepth = CalcSceneDepth(ScreenUV);
+
+	//@todo - handle invalid tile depth ranges due to MinDepth
+	float MinDepth = 0;
+
+	// Initialize per-tile variables
+    if (ThreadIndex == 0) 
+	{
+        IntegerTileMinZ = 0x7F7FFFFF;     
+        IntegerTileMaxZ = 0;
+		IntegerTileMinZ2 = 0x7F7FFFFF;  
+		IntegerTileMaxZ2 = 0;
+		TileNumObjects0 = 0;
+		TileNumObjects1 = 0;
+    }
+
+    GroupMemoryBarrierWithGroupSync();
+    
+	if (SceneDepth > MinDepth)
+	{
+		// Use shared memory atomics to build the depth bounds for this tile
+		// Each thread is assigned to a pixel at this point
+		//@todo - move depth range computation to a central point where it can be reused by all the frame's tiled deferred passes!
+		InterlockedMin(IntegerTileMinZ, asuint(SceneDepth));
+		InterlockedMax(IntegerTileMaxZ, asuint(SceneDepth));
+	}
+
+    GroupMemoryBarrierWithGroupSync();
+
+    float MinTileZ = asfloat(IntegerTileMinZ);
+    float MaxTileZ = asfloat(IntegerTileMaxZ);
+
+	float HalfZ = .5f * (MinTileZ + MaxTileZ);
+
+	// Compute a second min and max Z, clipped by HalfZ, so that we get two depth bounds per tile
+	// This results in more conservative tile depth bounds and fewer intersections
+	if (SceneDepth >= HalfZ && SceneDepth > MinDepth)
+	{
+		InterlockedMin(IntegerTileMinZ2, asuint(SceneDepth));
+	}
+
+	if (SceneDepth <= HalfZ && SceneDepth > MinDepth)
+	{
+		InterlockedMax(IntegerTileMaxZ2, asuint(SceneDepth));
+	}
+
+	GroupMemoryBarrierWithGroupSync();
+	
+	float MinTileZ2 = asfloat(IntegerTileMinZ2);
+	float MaxTileZ2 = asfloat(IntegerTileMaxZ2);
+
+	float3 ViewTileMin;
+	float3 ViewTileMax;
+
+	float3 ViewTileMin2;
+	float3 ViewTileMax2;
+
+	float ExpandRadius = 0;
+
+	float2 TanViewFOV = float2(1 / View.ViewToClip[0][0], 1 / View.ViewToClip[1][1]);
+	float2 TileSize = TanViewFOV * 2 / NumGroups;
+	float2 UnitPlaneTileMin = (GroupId.xy * TileSize - TanViewFOV) * float2(1, -1);
+	float2 UnitPlaneTileMax = ((GroupId.xy + 1) * TileSize - TanViewFOV) * float2(1, -1);
+
+	ViewTileMin.xy = min(MinTileZ * UnitPlaneTileMin, MaxTileZ2 * UnitPlaneTileMin) - ExpandRadius;
+	ViewTileMax.xy = max(MinTileZ * UnitPlaneTileMax, MaxTileZ2 * UnitPlaneTileMax) + ExpandRadius;
+	ViewTileMin.z = MinTileZ - ExpandRadius;
+	ViewTileMax.z = MaxTileZ2 + ExpandRadius;
+	ViewTileMin2.xy = min(MinTileZ2 * UnitPlaneTileMin, MaxTileZ * UnitPlaneTileMin) - ExpandRadius;
+	ViewTileMax2.xy = max(MinTileZ2 * UnitPlaneTileMax, MaxTileZ * UnitPlaneTileMax) + ExpandRadius;
+	ViewTileMin2.z = MinTileZ2 - ExpandRadius;
+	ViewTileMax2.z = MaxTileZ + ExpandRadius;
+
+	float3 ViewGroup0Center = (ViewTileMax + ViewTileMin) / 2;
+	float3 WorldGroup0Center = mul(float4(ViewGroup0Center, 1), View.ViewToTranslatedWorld).xyz - View.PreViewTranslation;
+	float Group0BoundingRadius = length(ViewGroup0Center - ViewTileMax);
+
+	float3 ViewGroup1Center = (ViewTileMax2 + ViewTileMin2) / 2;
+	float3 WorldGroup1Center = mul(float4(ViewGroup1Center, 1), View.ViewToTranslatedWorld).xyz - View.PreViewTranslation;
+	float Group1BoundingRadius = length(ViewGroup1Center - ViewTileMax2);
+
+	// Compute per-tile lists of affecting objects through bounds culling
+	// Each thread now operates on a sample instead of a pixel
+	LOOP
+	for (uint ObjectIndex = ThreadIndex; ObjectIndex < NumObjects; ObjectIndex += THREADGROUP_TOTALSIZE)
+	{
+		float4 SphereCenterAndRadius = LoadObjectPositionAndRadius(ObjectIndex);
+
+		BRANCH
+		if (RayHitSphere(WorldGroup0Center, -LightDirection, SphereCenterAndRadius.xyz, SphereCenterAndRadius.w + Group0BoundingRadius))
+		{
+			uint ListIndex;
+			InterlockedAdd(TileNumObjects0, 1U, ListIndex);
+			// Don't overwrite on overflow
+			ListIndex = min(ListIndex, MAX_INTERSECTING_OBJECTS - 1);
+			IntersectingObjectIndices[MAX_INTERSECTING_OBJECTS * 0 + ListIndex] = ObjectIndex; 
+		}
+
+		BRANCH
+		if (RayHitSphere(WorldGroup1Center, -LightDirection, SphereCenterAndRadius.xyz, SphereCenterAndRadius.w + Group1BoundingRadius))
+		{
+			uint ListIndex;
+			InterlockedAdd(TileNumObjects1, 1U, ListIndex);
+			// Don't write out of bounds on overflow
+			ListIndex = min(ListIndex, MAX_INTERSECTING_OBJECTS - 1);
+			IntersectingObjectIndices[MAX_INTERSECTING_OBJECTS * 1 + ListIndex] = ObjectIndex; 
+		}
+	}
+
+	GroupMemoryBarrierWithGroupSync();
+
+	float Result = 0;
+
+	BRANCH
+	if (SceneDepth > MinDepth)
+	{
+		float4 HomogeneousWorldPosition = mul(float4(ScreenPosition * SceneDepth, SceneDepth, 1), View.ScreenToWorld);
+		float3 OpaqueWorldPosition = HomogeneousWorldPosition.xyz / HomogeneousWorldPosition.w;
+
+		float TraceDistance = 10000;
+		float RayStartOffset = 10;
+		float3 WorldRayStart = OpaqueWorldPosition - LightDirection * RayStartOffset;
+		float3 WorldRayEnd = OpaqueWorldPosition - LightDirection * TraceDistance;
+		float3 WorldRayDirection = WorldRayEnd - WorldRayStart;
+		float3 UnitWorldRayDirection = normalize(WorldRayDirection);
+		uint GroupIndex = SceneDepth > MaxTileZ2 ? 1 : 0;
+		uint NumIntersectingObjects = GroupIndex == 0 ? TileNumObjects0 : TileNumObjects1;
+
+		//@todo - expose
+		float SunAngle = 1 * PI / 180.0f;
+		float TanSunAngle = tan(SunAngle);
+		float MinRayTime = TraceDistance;
+		float ClosestDistanceToIntersection = 1000;
+		float TotalStepsTaken = 0;
+		float MinConeVisibility = 1;
+
+		LOOP
+		for (uint ListObjectIndex = 0; ListObjectIndex < NumIntersectingObjects && MinRayTime >= TraceDistance; ListObjectIndex++)
+		{
+			uint ObjectIndex = IntersectingObjectIndices[MAX_INTERSECTING_OBJECTS * GroupIndex + ListObjectIndex];
+			float3 LocalPositionExtent = LoadObjectLocalPositionExtent(ObjectIndex);
+			float4x4 WorldToVolume = LoadObjectWorldToVolume(ObjectIndex);
+			float4 UVScaleAndVolumeScale = LoadObjectUVScale(ObjectIndex);
+			float3 UVAdd = LoadObjectUVAdd(ObjectIndex);
+
+			float3 VolumeRayStart = mul(float4(WorldRayStart, 1), WorldToVolume).xyz;
+			float3 VolumeRayEnd = mul(float4(WorldRayEnd, 1), WorldToVolume).xyz;
+			float3 VolumeRayDirection = VolumeRayEnd - VolumeRayStart;
+			float VolumeRayLength = length(VolumeRayDirection);
+			VolumeRayDirection /= VolumeRayLength;
+
+			float2 IntersectionTimes = LineBoxIntersect(VolumeRayStart, VolumeRayEnd, -LocalPositionExtent, LocalPositionExtent);
+
+			if (IntersectionTimes.x < IntersectionTimes.y && IntersectionTimes.x < 1)
+			{
+				float SampleRayTime = IntersectionTimes.x * VolumeRayLength;
+
+				float MinDistance = 1000000;
+				float3 IntersectionPosition = float3(0, 0, 0);
+
+				uint StepIndex = 0;
+				uint MaxSteps = 256;
+
+				LOOP
+				for (; StepIndex < MaxSteps; StepIndex++)
+				{
+					float3 SampleVolumePosition = VolumeRayStart + VolumeRayDirection * SampleRayTime;
+					float3 ClampedSamplePosition = clamp(SampleVolumePosition, -LocalPositionExtent, LocalPositionExtent);
+					float3 VolumeUV = DistanceFieldVolumePositionToUV(ClampedSamplePosition, UVScaleAndVolumeScale.xyz, UVAdd);
+					float DistanceField = Texture3DSampleLevel(DistanceFieldTexture, DistanceFieldSampler, VolumeUV, 0).x;
+					MinDistance = min(MinDistance, DistanceField);
+					MinConeVisibility = min(MinConeVisibility, saturate(DistanceField / (TanSunAngle * SampleRayTime)));
+					IntersectionPosition = SampleVolumePosition;
+
+					float MinStepSize = 1.0f / (4 * MaxSteps);
+					float StepDistance = max(DistanceField, MinStepSize);
+					SampleRayTime += StepDistance;
+
+					// Terminate the trace if we reached a negative area or went past the end of the ray
+					if (DistanceField < 0 
+						|| SampleRayTime > IntersectionTimes.y * VolumeRayLength)
+					{
+						break;
+					}
+				}
+
+				if (MinDistance * UVScaleAndVolumeScale.w < 0 || StepIndex == MaxSteps)
+				{
+					MinRayTime = min(MinRayTime, SampleRayTime * UVScaleAndVolumeScale.w);
+				}
+
+				if (SampleRayTime * UVScaleAndVolumeScale.w > 100)
+				{
+					ClosestDistanceToIntersection = min(ClosestDistanceToIntersection, MinDistance * UVScaleAndVolumeScale.w);
+				}
+				TotalStepsTaken += StepIndex;
+			}
+		}
+
+		Result = MinConeVisibility;
+	}
+
+	//Result = (GroupIndex == 0 ? Group0BoundingRadius : Group1BoundingRadius) / 1000.0f;
+	//Result = TotalStepsTaken / 1000.0f;
+	//Result = NumIntersectingObjects / 100.0f;
+
+	RWShadowFactors[DispatchThreadId.xy] = float2(Result, SceneDepth);
+}
+
+Texture2D ShadowFactorsTexture;
+SamplerState ShadowFactorsSampler;
+
+void DistanceFieldShadowingUpsamplePS(
+	in float4 UVAndScreenPos : TEXCOORD0,
+	in float4 SVPos : SV_POSITION,
+	out float4 OutColor : SV_Target0)
+{
+	// Distance field AO was computed at 0,0 regardless of viewrect min
+	float2 DistanceFieldUVs = UVAndScreenPos.xy - View.ViewRectMin.xy * View.ViewSizeAndSceneTexelSize.zw;
+
+#define BILATERAL_UPSAMPLE 1
+#if BILATERAL_UPSAMPLE
+	float2 LowResBufferSize = floor(View.RenderTargetSize / DOWNSAMPLE_FACTOR);
+	float2 LowResTexelSize = 1.0f / LowResBufferSize;
+	float2 Corner00UV = floor(DistanceFieldUVs * LowResBufferSize - .5f) / LowResBufferSize + .5f * LowResTexelSize;
+	float2 BilinearWeights = (DistanceFieldUVs - Corner00UV) * LowResBufferSize;
+
+	float2 TextureValues00 = Texture2DSampleLevel(ShadowFactorsTexture, ShadowFactorsSampler, Corner00UV, 0).xy;
+	float2 TextureValues10 = Texture2DSampleLevel(ShadowFactorsTexture, ShadowFactorsSampler, Corner00UV + float2(LowResTexelSize.x, 0), 0).xy;
+	float2 TextureValues01 = Texture2DSampleLevel(ShadowFactorsTexture, ShadowFactorsSampler, Corner00UV + float2(0, LowResTexelSize.y), 0).xy;
+	float2 TextureValues11 = Texture2DSampleLevel(ShadowFactorsTexture, ShadowFactorsSampler, Corner00UV + LowResTexelSize, 0).xy;
+
+	float4 CornerWeights = float4(
+		(1 - BilinearWeights.y) * (1 - BilinearWeights.x), 
+		(1 - BilinearWeights.y) * BilinearWeights.x,
+		BilinearWeights.y * (1 - BilinearWeights.x),
+		BilinearWeights.y * BilinearWeights.x);
+
+	float Epsilon = .0001f;
+
+	float4 CornerDepths = abs(float4(TextureValues00.y, TextureValues10.y, TextureValues01.y, TextureValues11.y));
+	float SceneDepth = CalcSceneDepth(UVAndScreenPos.xy);
+	float4 DepthWeights = 1.0f / (abs(CornerDepths - SceneDepth.xxxx) + Epsilon);
+	float4 FinalWeights = CornerWeights * DepthWeights;
+
+	float InterpolatedResult = 
+		(FinalWeights.x * TextureValues00.x 
+		+ FinalWeights.y * TextureValues10.x
+		+ FinalWeights.z * TextureValues01.x 
+		+ FinalWeights.w * TextureValues11.x)
+		/ dot(FinalWeights, 1);
+
+	float Output = InterpolatedResult;
+
+#else
+	float Output = Texture2DSampleLevel(ShadowFactorsTexture, ShadowFactorsSampler, DistanceFieldUVs, 0).x;
+#endif
+
+	OutColor = EncodeLightAttenuation(half4(Output, Output, Output, Output));
+}
\ No newline at end of file
diff --git a/Engine/Shaders/DistanceFieldSurfaceCacheLighting.usf b/Engine/Shaders/DistanceFieldSurfaceCacheLighting.usf
--- a/Engine/Shaders/DistanceFieldSurfaceCacheLighting.usf
+++ b/Engine/Shaders/DistanceFieldSurfaceCacheLighting.usf
@@ -1,12 +1,11 @@
-// Copyright 1998-2013 Epic Games, Inc. All Rights Reserved.
+// Copyright 1998-2014 Epic Games, Inc. All Rights Reserved.
 
 /*=============================================================================
-
+	DistanceFieldSurfaceCacheLighting.usf
 =============================================================================*/
 
 #include "Common.usf"
 #include "DeferredShadingCommon.usf"
-#include "SHCommon.usf"
 #include "DistanceFieldLightingShared.usf"
 
 struct FObjectCullVertexOutput
@@ -390,7 +389,6 @@ SamplerState BentNormalAOSampler3;
 /** Normalizes the splatted surface cache values, packs depth in alpha. */
 void AOCombinePS2(
 	in float4 UVAndScreenPos : TEXCOORD0, 
-	in float4 SVPos : SV_POSITION,
 	out float4 OutColor : SV_Target0)
 {
 	float4 DistanceFieldNormalValue = Texture2DSampleLevel(DistanceFieldNormalTexture, DistanceFieldNormalSampler, UVAndScreenPos.xy, 0);
@@ -421,10 +419,18 @@ void AOCombinePS2(
 	BRANCH
 	if (SceneDepth > AOMaxViewDistance)
 	{
+		float BentNormalLength = length(OutColor.rgb);
 		// Fade to unoccluded in the distance
-		OutColor.rgb = lerp(DistanceFieldNormalValue.xyz, OutColor.rgb, saturate(.003f * (AOMaxViewDistance - SceneDepth)));
+		float FadeAlpha = saturate(.001f * (.95f * AOMaxViewDistance - SceneDepth));
+		OutColor.rgb = lerp(DistanceFieldNormalValue.xyz, OutColor.rgb, FadeAlpha);
+		OutColor.rgb = normalize(OutColor.rgb) * lerp(1, BentNormalLength, FadeAlpha);
 
-		OutColor.a = abs(OutColor.a);
+		FLATTEN
+		if (SceneDepth > AOMaxViewDistance)
+		{
+			// Mark as valid for the gap filling pass, so we don't get overwritten
+			OutColor.a = abs(OutColor.a);
+		}
 	}
 
 #endif
@@ -439,7 +445,6 @@ SamplerState BentNormalAOSampler;
 /** Fills in texels with no splatted weight from screenspace neighbors. */
 void FillGapsPS(
 	in float4 UVAndScreenPos : TEXCOORD0, 
-	in float4 SVPos : SV_POSITION,
 	out float4 OutColor : SV_Target0)
 {
 	float4 CenterValue = Texture2DSampleLevel(BentNormalAOTexture, BentNormalAOSampler, UVAndScreenPos.xy, 0);
@@ -576,6 +581,11 @@ bool RayHitSphere( float3 RayOrigin, float3 UnitRayDirection, float3 SphereCente
 	return dot( CenterToRay, CenterToRay ) <= Square( SphereRadius );
 }
 
+#define MAX_INTERSECTING_OBJECTS 256
+groupshared uint IntersectingObjectIndices[MAX_INTERSECTING_OBJECTS];
+
+groupshared uint NumIntersectingObjects;
+
 [numthreads(THREADGROUP_SIZEX, THREADGROUP_SIZEY, 1)]
 void VisualizeMeshDistanceFieldCS(
 	uint3 GroupId : SV_GroupID,
@@ -590,21 +600,48 @@ void VisualizeMeshDistanceFieldCS(
 	float SceneDepth = CalcSceneDepth(ScreenUV);
 	float4 HomogeneousWorldPosition = mul(float4(ScreenPosition * SceneDepth, SceneDepth, 1), View.ScreenToWorld);
 	float3 OpaqueWorldPosition = HomogeneousWorldPosition.xyz / HomogeneousWorldPosition.w;
-	
+
 	float TraceDistance = 10000;
 	float3 WorldRayStart = View.ViewOrigin.xyz;
 	float3 WorldRayEnd = View.ViewOrigin.xyz + normalize(OpaqueWorldPosition - View.ViewOrigin.xyz) * TraceDistance;
 	float3 WorldRayDirection = WorldRayEnd - WorldRayStart;
 	float3 UnitWorldRayDirection = normalize(WorldRayDirection);
 
+	if (ThreadIndex == 0)
+	{
+		NumIntersectingObjects = 0;
+	}
+
+	GroupMemoryBarrierWithGroupSync();
+
+	LOOP
+	for (uint ObjectIndex = ThreadIndex; ObjectIndex < NumObjects; ObjectIndex += THREADGROUP_TOTALSIZE)
+	{
+		float4 SphereCenterAndRadius = LoadObjectPositionAndRadius(ObjectIndex);
+
+		//@todo - make independent of current pixel
+		BRANCH
+		if (RayHitSphere(WorldRayStart, UnitWorldRayDirection, SphereCenterAndRadius.xyz, SphereCenterAndRadius.w))
+		{
+			uint ListIndex;
+			InterlockedAdd(NumIntersectingObjects, 1U, ListIndex);
+			IntersectingObjectIndices[ListIndex] = ObjectIndex; 
+		}
+	}
+
+	GroupMemoryBarrierWithGroupSync();
+
+	uint LocalNumIntersectingObjects = NumIntersectingObjects;
+
 	float3 Result = 0;
 	float MinRayTime = TraceDistance;
 	float MinDistanceToSurface = 10000;
 	float TotalStepsTaken = 0;
 
 	LOOP
-	for (uint ObjectIndex = 0; ObjectIndex < NumObjects; ObjectIndex++)
+	for (uint ListObjectIndex = 0; ListObjectIndex < LocalNumIntersectingObjects; ListObjectIndex++)
 	{
+		uint ObjectIndex = IntersectingObjectIndices[ListObjectIndex];
 		float4 SphereCenterAndRadius = LoadObjectPositionAndRadius(ObjectIndex);
 
 #define TRACE_THROUGH_WORLD 0
@@ -620,16 +657,10 @@ void VisualizeMeshDistanceFieldCS(
 
 		float3 VolumeRayStart = mul(float4(WorldRayStart, 1), WorldToVolume).xyz;
 		float3 VolumeRayEnd = mul(float4(WorldRayEnd, 1), WorldToVolume).xyz;
-		//float3 VolumeRayEnd = mul(float4(OpaqueWorldPosition, 1), DistanceFieldObjectData.WorldToVolume[ObjectIndex]).xyz;
 		float3 VolumeRayDirection = VolumeRayEnd - VolumeRayStart;
 		float VolumeRayLength = length(VolumeRayDirection);
 		VolumeRayDirection /= VolumeRayLength;
-		/*
-		float3 TempPosition = mul(float4(OpaqueWorldPosition, 1), WorldToVolume).xyz;
-		float3 TempVolumeUV = DistanceFieldVolumePositionToUV(TempPosition, UVScaleAndVolumeScale.xyz, UVAdd);
-		float TempDistance = Texture3DSampleLevel(DistanceFieldTexture, DistanceFieldSampler, TempVolumeUV, 0).x;
-		Temp = min(Temp, TempDistance * UVScaleAndVolumeScale.w / 10.0f);
-		*/
+
 		float2 IntersectionTimes = LineBoxIntersect(VolumeRayStart, VolumeRayEnd, -LocalPositionExtent, LocalPositionExtent);
 
 		if (IntersectionTimes.x < IntersectionTimes.y && IntersectionTimes.x < 1)
@@ -640,7 +671,7 @@ void VisualizeMeshDistanceFieldCS(
 			float3 IntersectionPosition = float3(0, 0, 0);
 
 			uint StepIndex = 0;
-			uint MaxSteps = 512;
+			uint MaxSteps = 256;
 
 			LOOP
 			for (; StepIndex < MaxSteps; StepIndex++)
diff --git a/Engine/Shaders/DistanceFieldSurfaceCacheLightingCompute.usf b/Engine/Shaders/DistanceFieldSurfaceCacheLightingCompute.usf
--- a/Engine/Shaders/DistanceFieldSurfaceCacheLightingCompute.usf
+++ b/Engine/Shaders/DistanceFieldSurfaceCacheLightingCompute.usf
@@ -1,12 +1,11 @@
-// Copyright 1998-2013 Epic Games, Inc. All Rights Reserved.
+// Copyright 1998-2014 Epic Games, Inc. All Rights Reserved.
 
 /*=============================================================================
-
+	DistanceFieldSurfaceCacheLightingCompute.usf
 =============================================================================*/
 
 #include "Common.usf"
 #include "DeferredShadingCommon.usf"
-#include "SHCommon.usf"
 #include "ReflectionEnvironmentShared.usf"
 #include "DistanceFieldLightingShared.usf"
 
@@ -542,7 +541,7 @@ void FinalGatherCS(
 				float4 UVScaleAndVolumeScale = LoadObjectUVScale(ObjectIndex);
 				float3 VolumeShadingPosition = mul(float4(WorldShadingPosition, 1), WorldToVolume).xyz;
 
-				float ObjectOccluderRadius = length(LocalPositionExtent) * .5f;
+				float ObjectOccluderRadius = length(LocalPositionExtent) * .5f * UVScaleAndVolumeScale.w;
 				float BoxDistance = ComputeDistanceFromBoxToPoint(-LocalPositionExtent, LocalPositionExtent, VolumeShadingPosition) * UVScaleAndVolumeScale.w;
 
 				BRANCH
@@ -580,26 +579,42 @@ void FinalGatherCS(
 						float3 ClampedSamplePosition = clamp(StepSamplePosition, -LocalPositionExtent, LocalPositionExtent);
 						float DistanceToClamped = length(StepSamplePosition - ClampedSamplePosition);
 
-						float LocalStepOffset = length(StepSamplePosition - VolumeShadingPosition);
-
 						float3 StepVolumeUV = DistanceFieldVolumePositionToUV(ClampedSamplePosition, UVScaleAndVolumeScale.xyz, UVAdd);
-						float DistanceToOccluder = Texture3DSampleLevel(DistanceFieldTexture, DistanceFieldSampler, StepVolumeUV, 0).x + DistanceToClamped;
+						float DistanceToOccluder = (Texture3DSampleLevel(DistanceFieldTexture, DistanceFieldSampler, StepVolumeUV, 0).x + DistanceToClamped) * UVScaleAndVolumeScale.w;
+
+						float SphereRadius = WorldStepOffset * TanConeHalfAngle;
+
+						#define CIRCLE_LINE_INTERSECTION 0
+
+						#if CIRCLE_LINE_INTERSECTION
+
+							float Visibility = 1;
+
+							BRANCH
+							if (DistanceToOccluder < SphereRadius)
+							{
+								float SphereRadiusSquared = SphereRadius * SphereRadius;
+								// Area of the circle cut off by an intersecting line segment of distance DistanceToOccluder to the circle center
+								float CircleSegmentArea = SphereRadiusSquared * acos(DistanceToOccluder / SphereRadius) - DistanceToOccluder * sqrt(SphereRadiusSquared - DistanceToOccluder * DistanceToOccluder);
+								Visibility = 1 - saturate(CircleSegmentArea / (.5f * PI * SphereRadiusSquared));
+							}
 
-						float EffectiveDistanceToOccluder = DistanceToOccluder;
-						float SphereRadius = LocalStepOffset * TanConeHalfAngle;
-						// Allows use of negative distances in the small steps to create full penumbra range, reduces over-occlusion
-						float MinOcclusionFraction = 0;//-.6f * (1 - saturate(4 * StepIndex / (float)NUM_CONE_STEPS));
-						float Visibility = saturate((EffectiveDistanceToOccluder / SphereRadius - MinOcclusionFraction) / (1 - MinOcclusionFraction));
+						#else
+							// Allows use of negative distances in the small steps to create full penumbra range, reduces over-occlusion
+							float MinOcclusionFraction = 0;//-.6f * (1 - saturate(4 * StepIndex / (float)NUM_CONE_STEPS));
+							// Derive visibility from 1d intersection
+							float Visibility = saturate((DistanceToOccluder / SphereRadius - MinOcclusionFraction) / (1 - MinOcclusionFraction));
+						#endif
 
 						// Don't allow small objects to fully occlude a cone step
 						Visibility = max(Visibility, 1 - saturate(ObjectOccluderRadius / SphereRadius));
 
 						InterlockedMin(SharedConeVisibility[ConeIndex], asuint(Visibility));
 
-						if (EffectiveDistanceToOccluder < .9f * SphereRadius)
+						if (DistanceToOccluder < .9f * SphereRadius)
 						{
 							// Assuming occluder is straight forward along the cone
-							float WorldDistanceToOccluder = (LocalStepOffset + DistanceToOccluder) * UVScaleAndVolumeScale.w;
+							float WorldDistanceToOccluder = WorldStepOffset + DistanceToOccluder;
 							InterlockedMin(SharedMinOcclusionDistance, asuint(WorldDistanceToOccluder));
 						}
 					}
@@ -614,8 +629,9 @@ void FinalGatherCS(
 	{
 		float3 UnoccludedDirection = 0;
 		float Visibility = 0;
+		float NumSamplesToDisplay = NUM_SAMPLES;
 
-		for (uint ConeIndex = 0; ConeIndex < NUM_SAMPLES; ConeIndex++)
+		for (uint ConeIndex = 0; ConeIndex < NumSamplesToDisplay; ConeIndex++)
 		{
 			float3 ConeDirection = AOSamples2.SampleDirections[ConeIndex].xyz;
 			float3 RotatedConeDirection = ConeDirection.x * TangentX + ConeDirection.y * TangentY + ConeDirection.z * WorldNormal;
@@ -625,8 +641,8 @@ void FinalGatherCS(
 			Visibility += ConeVisibility;
 		}
 
-		UnoccludedDirection = UnoccludedDirection * BentNormalNormalizeFactor / (float)NUM_SAMPLES;
-		Visibility = Visibility / (float)NUM_SAMPLES;
+		UnoccludedDirection = UnoccludedDirection * BentNormalNormalizeFactor / (float)NumSamplesToDisplay;
+		Visibility = Visibility / (float)NumSamplesToDisplay;
 
 		RWOccluderRadius[RecordIndex] = RecordRadiusScale * asfloat(SharedMinOcclusionDistance);
 		RWIrradianceCacheBentNormal[RecordIndex] = float4(UnoccludedDirection, 0);
diff --git a/Engine/Shaders/ReflectionEnvironmentShaders.usf b/Engine/Shaders/ReflectionEnvironmentShaders.usf
--- a/Engine/Shaders/ReflectionEnvironmentShaders.usf
+++ b/Engine/Shaders/ReflectionEnvironmentShaders.usf
@@ -533,7 +533,10 @@ void ReflectionApplyPS(
 			SkyLighting *= GBuffer.IndirectIrradiance;
 		}
 
-		float Visibility = GetDistanceFieldAOSpecularOcclusion(UV, ReflectionVector, GBuffer.Roughness); 
+		float Visibility = 1;
+		
+		if (View.GeneralPurposeTweak > 0)
+		Visibility = GetDistanceFieldAOSpecularOcclusion(UV, ReflectionVector, GBuffer.Roughness); 
 
 		SpecularLighting.rgb += (Visibility * SpecularLighting.a) * SkyLighting;
 #endif
diff --git a/Engine/Shaders/SkyLighting.usf b/Engine/Shaders/SkyLighting.usf
--- a/Engine/Shaders/SkyLighting.usf
+++ b/Engine/Shaders/SkyLighting.usf
@@ -1,4 +1,4 @@
-// Copyright 1998-2013 Epic Games, Inc. All Rights Reserved.
+// Copyright 1998-2014 Epic Games, Inc. All Rights Reserved.
 
 /*=============================================================================
 	SkyLighting.usf
diff --git a/Engine/Shaders/SkyLightingShared.usf b/Engine/Shaders/SkyLightingShared.usf
--- a/Engine/Shaders/SkyLightingShared.usf
+++ b/Engine/Shaders/SkyLightingShared.usf
@@ -36,7 +36,7 @@ float GetDistanceFieldAOSpecularOcclusion(float2 ScreenUV, float3 ReflectionVect
 		Visibility = ApproximateConeConeIntersection(ReflectionConeAngle, UnoccludedAngle, AngleBetween);
 
 		// Can't rely on the direction of the bent normal when close to fully occluded, lerp to shadowed
-		Visibility = lerp(0, Visibility, saturate(UnoccludedAngle / .1f));
+		Visibility = lerp(0, Visibility, saturate((UnoccludedAngle - .1f) / .2f));
 	}
 
 	return Visibility;
diff --git a/Engine/Source/Runtime/Core/Private/HAL/ConsoleManager.cpp b/Engine/Source/Runtime/Core/Private/HAL/ConsoleManager.cpp
--- a/Engine/Source/Runtime/Core/Private/HAL/ConsoleManager.cpp
+++ b/Engine/Source/Runtime/Core/Private/HAL/ConsoleManager.cpp
@@ -1576,7 +1576,7 @@ static TAutoConsoleVariable<int32> CVarNumBufferedOcclusionQueries(
 	ECVF_ReadOnly);
 
 static TAutoConsoleVariable<int32> CVarDistField(
-	TEXT("r.AllowMeshDistanceFieldRepresentations"),
+	TEXT("r.GenerateMeshDistanceFields"),
 	0,	
 	TEXT("Whether to build distance fields of static meshes, needed for distance field AO, which is used to implement Movable SkyLight shadows.\n")
 	TEXT("Enabling will increase mesh build times and memory usage.  Changing this value will cause a rebuild of all static meshes."),
diff --git a/Engine/Source/Runtime/Engine/Classes/Components/DirectionalLightComponent.h b/Engine/Source/Runtime/Engine/Classes/Components/DirectionalLightComponent.h
--- a/Engine/Source/Runtime/Engine/Classes/Components/DirectionalLightComponent.h
+++ b/Engine/Source/Runtime/Engine/Classes/Components/DirectionalLightComponent.h
@@ -94,6 +94,19 @@ class UDirectionalLightComponent : public ULightComponent
 	UPROPERTY(EditAnywhere, AdvancedDisplay, BlueprintReadOnly, Category=CascadedShadowMaps)
 	uint32 bUseInsetShadowsForMovableObjects : 1;
 
+	/** 
+	 * Whether to use ray traced distance field area shadows.  The project setting bGenerateMeshDistanceFields must be enabled for this to have effect.
+	 * Distance field shadows have a large per-pixel cost but a low per-object cost (and don't depend on triangle count) so they are effective for distant shadows from a dynamic sun.
+	 */
+	UPROPERTY(EditAnywhere, BlueprintReadOnly, Category=DistanceFieldShadows)
+	bool bUseDistanceFieldShadows;
+
+	/** 
+	 * Distance at which the ray traced shadow cascade should end.  Distance field shadows will cover the range between 'Dynamic Shadow Distance' this distance. 
+	 */
+	UPROPERTY(EditAnywhere, BlueprintReadOnly, Category=DistanceFieldShadows, meta=(UIMin = "0", UIMax = "100000"), DisplayName = "DistanceField Shadow Distance")
+	float DistanceFieldShadowDistance;
+
 	/** The Lightmass settings for this object. */
 	UPROPERTY(EditAnywhere, BlueprintReadOnly, Category=Light, meta=(ShowOnlyInnerProperties))
 	struct FLightmassDirectionalLightSettings LightmassSettings;
diff --git a/Engine/Source/Runtime/Engine/Classes/Engine/RendererSettings.h b/Engine/Source/Runtime/Engine/Classes/Engine/RendererSettings.h
--- a/Engine/Source/Runtime/Engine/Classes/Engine/RendererSettings.h
+++ b/Engine/Source/Runtime/Engine/Classes/Engine/RendererSettings.h
@@ -142,9 +142,9 @@ class ENGINE_API URendererSettings
 
 
 	UPROPERTY(config, EditAnywhere, Category=Lighting, meta=(
-		ConsoleVariable="r.AllowMeshDistanceFieldRepresentations",
-		ToolTip="Whether to build distance fields of static meshes, needed for distance field AO, which is used to implement Movable SkyLight shadows.  Enabling will increase mesh build times and memory usage.  Changing this setting requires restarting the editor."))
-	uint32 bAllowMeshDistanceFieldRepresentations:1;
+		ConsoleVariable="r.GenerateMeshDistanceFields",
+		ToolTip="Whether to build distance fields of static meshes, needed for distance field AO, which is used to implement Movable SkyLight shadows, and ray traced distance field shadows on directional lights.  Enabling will increase mesh build times and memory usage.  Changing this setting requires restarting the editor."))
+	uint32 bGenerateMeshDistanceFields:1;
 
 	UPROPERTY(config, EditAnywhere, Category=Lighting, meta=(
 		ConsoleVariable="r.Shadow.DistanceFieldPenumbraSize",
diff --git a/Engine/Source/Runtime/Engine/Private/DirectionalLightComponent.cpp b/Engine/Source/Runtime/Engine/Private/DirectionalLightComponent.cpp
--- a/Engine/Source/Runtime/Engine/Private/DirectionalLightComponent.cpp
+++ b/Engine/Source/Runtime/Engine/Private/DirectionalLightComponent.cpp
@@ -86,6 +86,9 @@ public:
 
 	bool bUseInsetShadowsForMovableObjects;
 
+	/** If greater than WholeSceneDynamicShadowRadius, a cascade will be created to support ray traced distance field shadows covering up to this distance. */
+	float DistanceFieldShadowDistance;
+
 	/** Initialization constructor. */
 	FDirectionalLightSceneProxy(const UDirectionalLightComponent* Component):
 		FLightSceneProxy(Component),
@@ -97,7 +100,8 @@ public:
 		CascadeDistributionExponent(Component->CascadeDistributionExponent),
 		CascadeTransitionFraction(Component->CascadeTransitionFraction),
 		ShadowDistanceFadeoutFraction(Component->ShadowDistanceFadeoutFraction),
-		bUseInsetShadowsForMovableObjects(Component->bUseInsetShadowsForMovableObjects)
+		bUseInsetShadowsForMovableObjects(Component->bUseInsetShadowsForMovableObjects),
+		DistanceFieldShadowDistance(Component->bUseDistanceFieldShadows ? Component->DistanceFieldShadowDistance : 0)
 	{
 		LightShaftOverrideDirection.Normalize();
 
@@ -123,24 +127,6 @@ public:
 			Proxy->UpdateLightShaftOverrideDirection_RenderThread(NewLightShaftOverrideDirection);
 		});
 	}
-		
-	float ComputeShadowDistance() const
-	{
-		static const auto CVar = IConsoleManager::Get().FindTConsoleVariableDataFloat(TEXT("r.Shadow.DistanceScale"));
-
-		float Scale = FMath::Clamp(CVar->GetValueOnRenderThread(), 0.1f, 2.0f);
-
-		return WholeSceneDynamicShadowRadius * Scale;
-	}
-
-	float GetShadowTransiationScale() const
-	{
-		static const auto CVar = IConsoleManager::Get().FindTConsoleVariableDataFloat(TEXT("r.Shadow.CSM.TransitionScale"));
-
-		float Scale = FMath::Clamp(CVar->GetValueOnRenderThread(), 0.0f, 2.0f);
-
-		return Scale;
-	}
 
 	/** Accesses parameters needed for rendering the light. */
 	virtual void GetParameters(FVector4& LightPositionAndInvRadius, FVector4& LightColorAndFalloffExponent, FVector& NormalizedLightDirection, FVector2D& SpotAngles, float& LightSourceRadius, float& LightSourceLength, float& LightMinRoughness) const
@@ -202,7 +188,15 @@ public:
 	virtual int32 GetNumViewDependentWholeSceneShadows(const FSceneView& View) const 
 	{ 
 		int32 NumCascades = WholeSceneDynamicShadowRadius > 0.0f ? DynamicShadowCascades : 0;
-		return FMath::Min<int32>(NumCascades, View.MaxShadowCascades);
+		NumCascades = FMath::Min<int32>(NumCascades, View.MaxShadowCascades);
+
+		if (ShouldCreateRayTracedCascade())
+		{
+			// Add a cascade for the ray traced shadows
+			NumCascades++;
+		}
+
+		return NumCascades;
 	}
 
 	/**
@@ -211,34 +205,29 @@ public:
 	 */
 	virtual bool GetViewDependentWholeSceneProjectedShadowInitializer(const FSceneView& View, int32 SplitIndex, FWholeSceneProjectedShadowInitializer& OutInitializer) const
 	{
-		// this checks for WholeSceneDynamicShadowRadius and DynamicShadowCascades
-		int32 NumCascades = GetNumViewDependentWholeSceneShadows(View);
+		const FMatrix& WorldToLight = GetWorldToLight();
 
-		if(NumCascades)
-		{
-			const FMatrix& WorldToLight = GetWorldToLight();
-
-			FShadowCascadeSettings CascadeSettings;
-
-			FSphere Bounds = FDirectionalLightSceneProxy::GetShadowSplitBounds(View, SplitIndex, &CascadeSettings);
-			OutInitializer.CascadeSettings = CascadeSettings;
-
-			const float ShadowExtent = Bounds.W / FMath::Sqrt(3.0f);
-			const FBoxSphereBounds SubjectBounds(Bounds.Center, FVector(ShadowExtent, ShadowExtent, ShadowExtent), Bounds.W);
-			OutInitializer.bDirectionalLight = true;
-			OutInitializer.bOnePassPointLightShadow = false;
-			OutInitializer.PreShadowTranslation = -Bounds.Center;
-			OutInitializer.WorldToLight = FInverseRotationMatrix(FVector(WorldToLight.M[0][0],WorldToLight.M[1][0],WorldToLight.M[2][0]).SafeNormal().Rotation());
-			OutInitializer.Scales = FVector(1.0f,1.0f / Bounds.W,1.0f / Bounds.W);
-			OutInitializer.FaceDirection = FVector(1,0,0);
-			OutInitializer.SubjectBounds = FBoxSphereBounds(FVector::ZeroVector,SubjectBounds.BoxExtent,SubjectBounds.SphereRadius);
-			OutInitializer.WAxis = FVector4(0,0,0,1);
-			OutInitializer.MinLightW = -HALF_WORLD_MAX;
-			OutInitializer.MaxDistanceToCastInLightW = HALF_WORLD_MAX / 8.0f;
-			OutInitializer.SplitIndex = SplitIndex;
-			return true;
-		}
-		return false;
+		FShadowCascadeSettings CascadeSettings;
+
+		FSphere Bounds = FDirectionalLightSceneProxy::GetShadowSplitBounds(View, SplitIndex, &CascadeSettings);
+		OutInitializer.CascadeSettings = CascadeSettings;
+
+		const float ShadowExtent = Bounds.W / FMath::Sqrt(3.0f);
+		const FBoxSphereBounds SubjectBounds(Bounds.Center, FVector(ShadowExtent, ShadowExtent, ShadowExtent), Bounds.W);
+		OutInitializer.bDirectionalLight = true;
+		OutInitializer.bOnePassPointLightShadow = false;
+		OutInitializer.PreShadowTranslation = -Bounds.Center;
+		OutInitializer.WorldToLight = FInverseRotationMatrix(FVector(WorldToLight.M[0][0],WorldToLight.M[1][0],WorldToLight.M[2][0]).SafeNormal().Rotation());
+		OutInitializer.Scales = FVector(1.0f,1.0f / Bounds.W,1.0f / Bounds.W);
+		OutInitializer.FaceDirection = FVector(1,0,0);
+		OutInitializer.SubjectBounds = FBoxSphereBounds(FVector::ZeroVector,SubjectBounds.BoxExtent,SubjectBounds.SphereRadius);
+		OutInitializer.WAxis = FVector4(0,0,0,1);
+		OutInitializer.MinLightW = -HALF_WORLD_MAX;
+		OutInitializer.MaxDistanceToCastInLightW = HALF_WORLD_MAX / 8.0f;
+		OutInitializer.SplitIndex = SplitIndex;
+		// Last cascade is the ray traced shadow, if enabled
+		OutInitializer.bRayTracedDistanceFieldShadow = ShouldCreateRayTracedCascade() && SplitIndex == FDirectionalLightSceneProxy::GetNumViewDependentWholeSceneShadows(View) - 1;
+		return true;
 	}
 
 	// reflective shadow map for Light Propagation Volume
@@ -309,6 +298,52 @@ public:
 
 private:
 
+	float GetCSMMaxDistance() const
+	{
+		static const auto CVar = IConsoleManager::Get().FindTConsoleVariableDataFloat(TEXT("r.Shadow.DistanceScale"));
+		 
+		float Scale = FMath::Clamp(CVar->GetValueOnRenderThread(), 0.1f, 2.0f);
+		float Distance = WholeSceneDynamicShadowRadius * Scale;
+		return Distance;
+	}
+
+	float GetDistanceFieldShadowDistance() const
+	{
+		static const auto CVar = IConsoleManager::Get().FindTConsoleVariableDataInt(TEXT("r.GenerateMeshDistanceFields"));
+
+		if (CVar->GetValueOnRenderThread() == 0)
+		{
+			// Meshes must have distance fields generated
+			return 0;
+		}
+
+		return DistanceFieldShadowDistance;
+	}
+
+	bool ShouldCreateRayTracedCascade() const
+	{
+		return GetDistanceFieldShadowDistance() > GetCSMMaxDistance() && DoesPlatformSupportDistanceFieldShadowing(GRHIShaderPlatform);
+	}
+
+	float ComputeShadowDistance() const
+	{
+		float Distance = GetCSMMaxDistance();
+
+		if (ShouldCreateRayTracedCascade())
+		{
+			Distance = GetDistanceFieldShadowDistance();
+		}
+
+		return Distance;
+	}
+
+	float GetShadowTransitionScale() const
+	{
+		static const auto CVar = IConsoleManager::Get().FindTConsoleVariableDataFloat(TEXT("r.Shadow.CSM.TransitionScale"));
+		float Scale = FMath::Clamp(CVar->GetValueOnRenderThread(), 0.0f, 2.0f);
+		return Scale;
+	}
+
 	void UpdateLightShaftOverrideDirection_RenderThread(FVector NewLightShaftOverrideDirection)
 	{
 		LightShaftOverrideDirection = NewLightShaftOverrideDirection;
@@ -423,10 +458,19 @@ private:
 	}
 
 	// Given a particular split index and view near plane, returns the view space Z distance to the PSSM split plane.
-	inline float GetSplitDistance(int32 SplitIndex, int32 NumCascades, float ShadowNear, float FarDistance) const
+	inline float GetSplitDistance(int32 SplitIndex, int32 NumCascades, float ShadowNear, bool bIgnoreLastCascade) const
 	{
 		check(SplitIndex >= 0);
 
+		float FarDistance = ComputeShadowDistance();
+
+		if (bIgnoreLastCascade)
+		{
+			// Compute shadowmap cascade distances as if the ray traced cascade didn't exist
+			NumCascades--;
+			FarDistance = GetCSMMaxDistance();
+		}
+
 		return ShadowNear + ComputeAccumulatedScale(CascadeDistributionExponent, SplitIndex, NumCascades) * (FarDistance - ShadowNear);
 	}
 
@@ -435,6 +479,9 @@ private:
 		// this checks for WholeSceneDynamicShadowRadius and DynamicShadowCascades
 		int32 NumCascades = GetNumViewDependentWholeSceneShadows(View);
 
+		const bool bHasRayTracedCascade = ShouldCreateRayTracedCascade();
+		const bool bIsRayTracedCascade = bHasRayTracedCascade && SplitIndex == NumCascades - 1;
+
 		const FMatrix ViewMatrix = View.ShadowViewMatrices.ViewMatrix;
 		const FMatrix ProjectionMatrix = View.ShadowViewMatrices.ProjMatrix;
 		const FVector4 ViewOrigin = View.ShadowViewMatrices.ViewOrigin;
@@ -442,14 +489,13 @@ private:
 		const FVector CameraDirection = ViewMatrix.GetColumn(2);
 		const FVector LightDirection = -GetDirection();
 
-		float FarDistance = ComputeShadowDistance();
-
 		// Determine start and end distances to the current cascade's split planes
-		float SplitNear = GetSplitDistance(SplitIndex, NumCascades, View.NearClippingDistance, FarDistance);
-		float SplitFar = GetSplitDistance(SplitIndex + 1, NumCascades, View.NearClippingDistance, FarDistance);
+		// Presence of the ray traced cascade does not change depth ranges for the shadow-mapped cascades
+		float SplitNear = GetSplitDistance(SplitIndex, NumCascades, View.NearClippingDistance, bHasRayTracedCascade);
+		float SplitFar = GetSplitDistance(SplitIndex + 1, NumCascades, View.NearClippingDistance, bHasRayTracedCascade && !bIsRayTracedCascade);
 		float FadePlane = SplitFar;
 
-		float LocalCascadeTransitionFraction = CascadeTransitionFraction * GetShadowTransiationScale();
+		float LocalCascadeTransitionFraction = CascadeTransitionFraction * GetShadowTransitionScale();
 
 		float FadeExtension = (SplitFar - SplitNear) * LocalCascadeTransitionFraction;
 
@@ -467,8 +513,8 @@ private:
 			if(SplitIndex >= 1)
 			{
 				// todo: optimize
-				float BeforeSplitNear = GetSplitDistance(SplitIndex - 1, NumCascades, View.NearClippingDistance, FarDistance);
-				float BeforeSplitFar = GetSplitDistance(SplitIndex, NumCascades, View.NearClippingDistance, FarDistance);
+				float BeforeSplitNear = GetSplitDistance(SplitIndex - 1, NumCascades, View.NearClippingDistance, bHasRayTracedCascade);
+				float BeforeSplitFar = GetSplitDistance(SplitIndex, NumCascades, View.NearClippingDistance, bHasRayTracedCascade);
 
 				OutCascadeSettings->SplitNearFadeRegion = (BeforeSplitFar - BeforeSplitNear) * LocalCascadeTransitionFraction;
 			}
@@ -571,6 +617,8 @@ UDirectionalLightComponent::UDirectionalLightComponent(const class FPostConstruc
 	DynamicShadowDistanceMovableLight = 20000.0f;
 	DynamicShadowDistanceStationaryLight = 0.f;
 
+	DistanceFieldShadowDistance = 40000.0f;
+
 	DynamicShadowCascades = 3;
 	CascadeDistributionExponent = 3.0f;
 	CascadeTransitionFraction = 0.1f;
@@ -622,8 +670,24 @@ bool UDirectionalLightComponent::CanEditChange(const UProperty* InProperty) cons
 			|| PropertyName == GET_MEMBER_NAME_STRING_CHECKED(UDirectionalLightComponent, ShadowDistanceFadeoutFraction)
 			|| PropertyName == GET_MEMBER_NAME_STRING_CHECKED(UDirectionalLightComponent, bUseInsetShadowsForMovableObjects))
 		{
-			return (DynamicShadowDistanceMovableLight > 0 && Mobility == EComponentMobility::Movable
-				|| DynamicShadowDistanceStationaryLight > 0 && Mobility == EComponentMobility::Stationary);
+			return CastShadows 
+				&& CastDynamicShadows 
+				&& (DynamicShadowDistanceMovableLight > 0 && Mobility == EComponentMobility::Movable
+					|| DynamicShadowDistanceStationaryLight > 0 && Mobility == EComponentMobility::Stationary);
+		}
+
+		if (PropertyName == GET_MEMBER_NAME_STRING_CHECKED(UDirectionalLightComponent, bUseDistanceFieldShadows)
+			|| PropertyName == GET_MEMBER_NAME_STRING_CHECKED(UDirectionalLightComponent, DistanceFieldShadowDistance))
+		{
+			static const auto CVar = IConsoleManager::Get().FindTConsoleVariableDataInt(TEXT("r.GenerateMeshDistanceFields"));
+			bool bCanEdit = CastShadows && CastDynamicShadows && Mobility != EComponentMobility::Static && CVar->GetValueOnGameThread() != 0;
+
+			if (PropertyName == GET_MEMBER_NAME_STRING_CHECKED(UDirectionalLightComponent, DistanceFieldShadowDistance))
+			{
+				bCanEdit = bCanEdit && bUseDistanceFieldShadows;
+			}
+
+			return bCanEdit;
 		}
 
 		if (PropertyName == GET_MEMBER_NAME_STRING_CHECKED(UDirectionalLightComponent, OcclusionMaskDarkness)
diff --git a/Engine/Source/Runtime/Engine/Private/DistanceFieldAtlas.cpp b/Engine/Source/Runtime/Engine/Private/DistanceFieldAtlas.cpp
--- a/Engine/Source/Runtime/Engine/Private/DistanceFieldAtlas.cpp
+++ b/Engine/Source/Runtime/Engine/Private/DistanceFieldAtlas.cpp
@@ -1,7 +1,7 @@
 // Copyright 1998-2014 Epic Games, Inc. All Rights Reserved.
 
 /*=============================================================================
-	DistanceFieldAtlas.cpp: 
+	DistanceFieldAtlas.cpp
 =============================================================================*/
 
 #include "EnginePrivate.h"
diff --git a/Engine/Source/Runtime/Engine/Private/LightComponent.cpp b/Engine/Source/Runtime/Engine/Private/LightComponent.cpp
--- a/Engine/Source/Runtime/Engine/Private/LightComponent.cpp
+++ b/Engine/Source/Runtime/Engine/Private/LightComponent.cpp
@@ -526,6 +526,8 @@ void ULightComponent::PostEditChangeProperty(FPropertyChangedEvent& PropertyChan
 		PropertyName != GET_MEMBER_NAME_STRING_CHECKED(UDirectionalLightComponent, CascadeTransitionFraction) &&
 		PropertyName != GET_MEMBER_NAME_STRING_CHECKED(UDirectionalLightComponent, ShadowDistanceFadeoutFraction) &&
 		PropertyName != GET_MEMBER_NAME_STRING_CHECKED(UDirectionalLightComponent, bUseInsetShadowsForMovableObjects) &&
+		PropertyName != GET_MEMBER_NAME_STRING_CHECKED(UDirectionalLightComponent, bUseDistanceFieldShadows) &&
+		PropertyName != GET_MEMBER_NAME_STRING_CHECKED(UDirectionalLightComponent, DistanceFieldShadowDistance) &&
 		PropertyName != GET_MEMBER_NAME_STRING_CHECKED(UDirectionalLightComponent, bEnableLightShaftOcclusion) &&
 		PropertyName != GET_MEMBER_NAME_STRING_CHECKED(UDirectionalLightComponent, OcclusionMaskDarkness) &&
 		PropertyName != GET_MEMBER_NAME_STRING_CHECKED(UDirectionalLightComponent, OcclusionDepthRange) &&
diff --git a/Engine/Source/Runtime/Engine/Private/SkyLightComponent.cpp b/Engine/Source/Runtime/Engine/Private/SkyLightComponent.cpp
--- a/Engine/Source/Runtime/Engine/Private/SkyLightComponent.cpp
+++ b/Engine/Source/Runtime/Engine/Private/SkyLightComponent.cpp
@@ -247,7 +247,8 @@ bool USkyLightComponent::CanEditChange(const UProperty* InProperty) const
 		if (FCString::Strcmp(*PropertyName, TEXT("Contrast")) == 0
 			|| FCString::Strcmp(*PropertyName, TEXT("OcclusionMaxDistance")) == 0)
 		{
-			return Mobility == EComponentMobility::Movable && CastShadows;
+			static const auto CVar = IConsoleManager::Get().FindTConsoleVariableDataInt(TEXT("r.GenerateMeshDistanceFields"));
+			return Mobility == EComponentMobility::Movable && CastShadows && CVar->GetValueOnGameThread() != 0;
 		}
 	}
 
diff --git a/Engine/Source/Runtime/Engine/Private/StaticMesh.cpp b/Engine/Source/Runtime/Engine/Private/StaticMesh.cpp
--- a/Engine/Source/Runtime/Engine/Private/StaticMesh.cpp
+++ b/Engine/Source/Runtime/Engine/Private/StaticMesh.cpp
@@ -1142,7 +1142,7 @@ void FStaticMeshRenderData::Cache(UStaticMesh* Owner, const FStaticMeshLODSettin
 		FPlatformAtomics::InterlockedAdd(&StaticMeshDerivedDataTimings::BuildCycles, T1-T0);
 	}
 
-	static const auto CVar = IConsoleManager::Get().FindTConsoleVariableDataInt(TEXT("r.AllowMeshDistanceFieldRepresentations"));
+	static const auto CVar = IConsoleManager::Get().FindTConsoleVariableDataInt(TEXT("r.GenerateMeshDistanceFields"));
 
 	if (CVar->GetValueOnGameThread() != 0)
 	{
diff --git a/Engine/Source/Runtime/Engine/Public/DistanceFieldAtlas.h b/Engine/Source/Runtime/Engine/Public/DistanceFieldAtlas.h
--- a/Engine/Source/Runtime/Engine/Public/DistanceFieldAtlas.h
+++ b/Engine/Source/Runtime/Engine/Public/DistanceFieldAtlas.h
@@ -1,7 +1,7 @@
 // Copyright 1998-2014 Epic Games, Inc. All Rights Reserved.
 
 /*=============================================================================
-	DistanceFieldAtlas.h: 
+	DistanceFieldAtlas.h
 =============================================================================*/
 
 #pragma once
diff --git a/Engine/Source/Runtime/Engine/Public/SceneManagement.h b/Engine/Source/Runtime/Engine/Public/SceneManagement.h
--- a/Engine/Source/Runtime/Engine/Public/SceneManagement.h
+++ b/Engine/Source/Runtime/Engine/Public/SceneManagement.h
@@ -507,7 +507,7 @@ public:
 	float MaxDistanceToCastInLightW;
 
 	/** Whether the shadow is for a directional light. */
-	uint32 bDirectionalLight : 1;
+	bool bDirectionalLight;
 
 	/** Default constructor. */
 	FProjectedShadowInitializer()
@@ -531,14 +531,24 @@ public:
 	FShadowCascadeSettings CascadeSettings;
 
 	/** Whether the shadow is a point light shadow that renders all faces of a cubemap in one pass. */
-	uint32 bOnePassPointLightShadow : 1;
+	bool bOnePassPointLightShadow;
+
+	/** Whether the shadow is a directional light cascade that will be computed by ray tracing the distance field. */
+	bool bRayTracedDistanceFieldShadow;
 
 	FWholeSceneProjectedShadowInitializer()
 	:	SplitIndex(INDEX_NONE)
 	,	bOnePassPointLightShadow(false)
+	,	bRayTracedDistanceFieldShadow(false)
 	{}	
 };
 
+inline bool DoesPlatformSupportDistanceFieldShadowing(EShaderPlatform Platform)
+{
+	// Hasn't been tested elsewhere yet
+	return Platform == SP_PCD3D_SM5;
+}
+
 /** Represents a USkyLightComponent to the rendering thread. */
 class ENGINE_API FSkyLightSceneProxy
 {
diff --git a/Engine/Source/Runtime/Renderer/Private/DistanceFieldLightingShared.h b/Engine/Source/Runtime/Renderer/Private/DistanceFieldLightingShared.h
--- a/Engine/Source/Runtime/Renderer/Private/DistanceFieldLightingShared.h
+++ b/Engine/Source/Runtime/Renderer/Private/DistanceFieldLightingShared.h
@@ -1,7 +1,7 @@
-// Copyright 1998-2013 Epic Games, Inc. All Rights Reserved.
+// Copyright 1998-2014 Epic Games, Inc. All Rights Reserved.
 
 /*=============================================================================
-	
+	DistanceFieldLightingShared.h
 =============================================================================*/
 
 #pragma once
@@ -19,6 +19,8 @@ const int32 GAODownsampleFactor = 2;
 /** Number of cone traced directions. */
 const int32 NumConeSampleDirections = 9;
 
+extern FIntPoint GetBufferSizeForAO();
+
 /** Generates unit length, stratified and uniformly distributed direction samples in a hemisphere. */
 extern void GenerateStratifiedUniformHemisphereSamples2(int32 NumThetaSteps, int32 NumPhiSteps, FRandomStream& RandomStream, TArray<FVector4>& Samples);
 
@@ -50,9 +52,6 @@ public:
 		AOStepScale.Bind(ParameterMap,TEXT("AOStepScale"));
 		AOStepExponentScale.Bind(ParameterMap,TEXT("AOStepExponentScale"));
 		AOMaxViewDistance.Bind(ParameterMap,TEXT("AOMaxViewDistance"));
-		DistanceFieldTexture.Bind(ParameterMap, TEXT("DistanceFieldTexture"));
-		DistanceFieldSampler.Bind(ParameterMap, TEXT("DistanceFieldSampler"));
-		DistanceFieldAtlasTexelSize.Bind(ParameterMap, TEXT("DistanceFieldAtlasTexelSize"));
 	}
 
 	friend FArchive& operator<<(FArchive& Ar,FAOParameters& Parameters)
@@ -61,9 +60,6 @@ public:
 		Ar << Parameters.AOStepScale;
 		Ar << Parameters.AOStepExponentScale;
 		Ar << Parameters.AOMaxViewDistance;
-		Ar << Parameters.DistanceFieldTexture;
-		Ar << Parameters.DistanceFieldSampler;
-		Ar << Parameters.DistanceFieldAtlasTexelSize;
 		return Ar;
 	}
 
@@ -84,6 +80,91 @@ public:
 
 		extern float GAOMaxViewDistance;
 		SetShaderValue(RHICmdList, ShaderRHI, AOMaxViewDistance, GAOMaxViewDistance);
+	}
+
+private:
+	FShaderParameter AOMaxDistance;
+	FShaderParameter AOStepScale;
+	FShaderParameter AOStepExponentScale;
+	FShaderParameter AOMaxViewDistance;
+};
+
+class FDistanceFieldObjectData
+{
+public:
+
+	FVertexBufferRHIRef Bounds;
+	FVertexBufferRHIRef Data;
+	FVertexBufferRHIRef Data2;
+	FShaderResourceViewRHIRef BoundsSRV;
+	FShaderResourceViewRHIRef DataSRV;
+	FShaderResourceViewRHIRef Data2SRV;
+};
+
+class FDistanceFieldObjectBuffers : public FRenderResource
+{
+public:
+
+	// In float4's
+	static int32 ObjectDataStride;
+	static int32 ObjectData2Stride;
+
+	int32 MaxObjects;
+	FDistanceFieldObjectData ObjectData;
+
+	FDistanceFieldObjectBuffers()
+	{
+		MaxObjects = 0;
+	}
+
+	virtual void InitDynamicRHI()  override
+	{
+		if (MaxObjects > 0)
+		{
+			FRHIResourceCreateInfo CreateInfo;
+			ObjectData.Bounds = RHICreateVertexBuffer(MaxObjects * sizeof(FVector4), BUF_Volatile | BUF_ShaderResource, CreateInfo);
+			ObjectData.Data = RHICreateVertexBuffer(MaxObjects * sizeof(FVector4) * ObjectDataStride, BUF_Volatile | BUF_ShaderResource, CreateInfo);
+			ObjectData.Data2 = RHICreateVertexBuffer(MaxObjects * sizeof(FVector4) * ObjectData2Stride, BUF_Volatile | BUF_ShaderResource, CreateInfo);
+
+			ObjectData.BoundsSRV = RHICreateShaderResourceView(ObjectData.Bounds, sizeof(FVector4), PF_A32B32G32R32F);
+			ObjectData.DataSRV = RHICreateShaderResourceView(ObjectData.Data, sizeof(FVector4), PF_A32B32G32R32F);
+			ObjectData.Data2SRV = RHICreateShaderResourceView(ObjectData.Data2, sizeof(FVector4), PF_A32B32G32R32F);
+		}
+	}
+
+	virtual void ReleaseDynamicRHI() override
+	{
+		ObjectData.Bounds.SafeRelease();
+		ObjectData.Data.SafeRelease();
+		ObjectData.Data2.SafeRelease();
+		ObjectData.BoundsSRV.SafeRelease(); 
+		ObjectData.DataSRV.SafeRelease(); 
+		ObjectData.Data2SRV.SafeRelease(); 
+	}
+};
+
+class FDistanceFieldObjectBufferParameters
+{
+public:
+	void Bind(const FShaderParameterMap& ParameterMap)
+	{
+		ObjectBounds.Bind(ParameterMap, TEXT("ObjectBounds"));
+		ObjectData.Bind(ParameterMap, TEXT("ObjectData"));
+		ObjectData2.Bind(ParameterMap, TEXT("ObjectData2"));
+		NumObjects.Bind(ParameterMap, TEXT("NumObjects"));
+		DistanceFieldTexture.Bind(ParameterMap, TEXT("DistanceFieldTexture"));
+		DistanceFieldSampler.Bind(ParameterMap, TEXT("DistanceFieldSampler"));
+		DistanceFieldAtlasTexelSize.Bind(ParameterMap, TEXT("DistanceFieldAtlasTexelSize"));
+	}
+
+	template<typename TParamRef>
+	void Set(FRHICommandList& RHICmdList, const TParamRef& ShaderRHI, const FDistanceFieldObjectBuffers& ObjectBuffers, int32 NumObjectsValue)
+	{
+		SetSRVParameter(RHICmdList, ShaderRHI, ObjectBounds, ObjectBuffers.ObjectData.BoundsSRV);
+		SetSRVParameter(RHICmdList, ShaderRHI, ObjectData, ObjectBuffers.ObjectData.DataSRV);
+		SetSRVParameter(RHICmdList, ShaderRHI, ObjectData2, ObjectBuffers.ObjectData.Data2SRV);
+
+		SetShaderValue(RHICmdList, ShaderRHI, NumObjects, NumObjectsValue);
 
 		SetTextureParameter(
 			RHICmdList,
@@ -99,13 +180,20 @@ public:
 		const int32 NumTexelsOneDimZ = GDistanceFieldVolumeTextureAtlas.GetSizeZ();
 		const FVector InvTextureDim(1.0f / NumTexelsOneDimX, 1.0f / NumTexelsOneDimY, 1.0f / NumTexelsOneDimZ);
 		SetShaderValue(RHICmdList, ShaderRHI, DistanceFieldAtlasTexelSize, InvTextureDim);
+
+	}
+
+	friend FArchive& operator<<(FArchive& Ar, FDistanceFieldObjectBufferParameters& P)
+	{
+		Ar << P.ObjectBounds << P.ObjectData << P.ObjectData2 << P.NumObjects << P.DistanceFieldTexture << P.DistanceFieldSampler << P.DistanceFieldAtlasTexelSize;
+		return Ar;
 	}
 
 private:
-	FShaderParameter AOMaxDistance;
-	FShaderParameter AOStepScale;
-	FShaderParameter AOStepExponentScale;
-	FShaderParameter AOMaxViewDistance;
+	FShaderResourceParameter ObjectBounds;
+	FShaderResourceParameter ObjectData;
+	FShaderResourceParameter ObjectData2;
+	FShaderParameter NumObjects;
 	FShaderResourceParameter DistanceFieldTexture;
 	FShaderResourceParameter DistanceFieldSampler;
 	FShaderParameter DistanceFieldAtlasTexelSize;
diff --git a/Engine/Source/Runtime/Renderer/Private/DistanceFieldShadowing.cpp b/Engine/Source/Runtime/Renderer/Private/DistanceFieldShadowing.cpp
new file mode 100644
--- /dev/null
+++ b/Engine/Source/Runtime/Renderer/Private/DistanceFieldShadowing.cpp
@@ -0,0 +1,364 @@
+// Copyright 1998-2014 Epic Games, Inc. All Rights Reserved.
+
+/*=============================================================================
+	DistanceFieldShadowing.cpp
+=============================================================================*/
+
+#include "RendererPrivate.h"
+#include "ScenePrivate.h"
+#include "UniformBuffer.h"
+#include "ShaderParameters.h"
+#include "PostProcessing.h"
+#include "SceneFilterRendering.h"
+#include "ScreenRendering.h"
+#include "DistanceFieldSurfaceCacheLighting.h"
+#include "DistanceFieldLightingShared.h"
+#include "RHICommandList.h"
+#include "SceneUtils.h"
+#include "DistanceFieldAtlas.h"
+
+int32 GDistanceFieldShadowing = 1;
+FAutoConsoleVariableRef CVarDistanceFieldShadowing(
+	TEXT("r.DistanceFieldShadowing"),
+	GDistanceFieldShadowing,
+	TEXT("Whether the distance field shadowing feature is allowed."),
+	ECVF_Cheat | ECVF_RenderThreadSafe
+	);
+
+//@todo - per light, breaks with multiple directional lights
+TGlobalResource<FDistanceFieldObjectBuffers> GDirectionalLightObjectBuffers;
+
+class FDistanceFieldShadowingCS : public FGlobalShader
+{
+	DECLARE_SHADER_TYPE(FDistanceFieldShadowingCS, Global);
+public:
+
+	static bool ShouldCache(EShaderPlatform Platform)
+	{
+		return IsFeatureLevelSupported(Platform, ERHIFeatureLevel::SM5) && DoesPlatformSupportDistanceFieldShadowing(Platform);
+	}
+
+	static void ModifyCompilationEnvironment(EShaderPlatform Platform, FShaderCompilerEnvironment& OutEnvironment)
+	{
+		OutEnvironment.SetDefine(TEXT("DOWNSAMPLE_FACTOR"), GAODownsampleFactor);
+		OutEnvironment.SetDefine(TEXT("THREADGROUP_SIZEX"), GDistanceFieldAOTileSizeX);
+		OutEnvironment.SetDefine(TEXT("THREADGROUP_SIZEY"), GDistanceFieldAOTileSizeY);
+	}
+
+	/** Default constructor. */
+	FDistanceFieldShadowingCS() {}
+
+	/** Initialization constructor. */
+	FDistanceFieldShadowingCS(const ShaderMetaType::CompiledShaderInitializerType& Initializer)
+		: FGlobalShader(Initializer)
+	{
+		ShadowFactors.Bind(Initializer.ParameterMap, TEXT("ShadowFactors"));
+		NumGroups.Bind(Initializer.ParameterMap, TEXT("NumGroups"));
+		LightDirection.Bind(Initializer.ParameterMap, TEXT("LightDirection"));
+		ObjectParameters.Bind(Initializer.ParameterMap);
+		DeferredParameters.Bind(Initializer.ParameterMap);
+	}
+
+	void SetParameters(
+		FRHICommandList& RHICmdList, 
+		const FSceneView& View, 
+		const FProjectedShadowInfo* ProjectedShadowInfo,
+		FSceneRenderTargetItem& ShadowFactorsValue, 
+		int32 NumObjects, 
+		FVector2D NumGroupsValue)
+	{
+		const FComputeShaderRHIParamRef ShaderRHI = GetComputeShader();
+
+		FGlobalShader::SetParameters(RHICmdList, ShaderRHI, View);
+
+		ShadowFactors.SetTexture(RHICmdList, ShaderRHI, ShadowFactorsValue.ShaderResourceTexture, ShadowFactorsValue.UAV);
+
+		ObjectParameters.Set(RHICmdList, ShaderRHI, GDirectionalLightObjectBuffers, NumObjects);
+		DeferredParameters.Set(RHICmdList, ShaderRHI, View);
+
+		SetShaderValue(RHICmdList, ShaderRHI, NumGroups, NumGroupsValue);
+		SetShaderValue(RHICmdList, ShaderRHI, LightDirection, ProjectedShadowInfo->LightSceneInfo->Proxy->GetDirection());
+	}
+
+	void UnsetParameters(FRHICommandList& RHICmdList)
+	{
+		ShadowFactors.UnsetUAV(RHICmdList, GetComputeShader());
+	}
+
+	// FShader interface.
+	virtual bool Serialize(FArchive& Ar)
+	{
+		bool bShaderHasOutdatedParameters = FGlobalShader::Serialize(Ar);
+		Ar << ShadowFactors;
+		Ar << NumGroups;
+		Ar << LightDirection;
+		Ar << ObjectParameters;
+		Ar << DeferredParameters;
+		return bShaderHasOutdatedParameters;
+	}
+
+private:
+
+	FRWShaderParameter ShadowFactors;
+	FShaderParameter NumGroups;
+	FShaderParameter LightDirection;
+	FDistanceFieldObjectBufferParameters ObjectParameters;
+	FDeferredPixelShaderParameters DeferredParameters;
+};
+
+IMPLEMENT_SHADER_TYPE(,FDistanceFieldShadowingCS,TEXT("DistanceFieldShadowing"),TEXT("DistanceFieldShadowingCS"),SF_Compute);
+
+class FDistanceFieldShadowingUpsamplePS : public FGlobalShader
+{
+	DECLARE_SHADER_TYPE(FDistanceFieldShadowingUpsamplePS, Global);
+public:
+
+	static bool ShouldCache(EShaderPlatform Platform)
+	{
+		return IsFeatureLevelSupported(Platform, ERHIFeatureLevel::SM5) && DoesPlatformSupportDistanceFieldShadowing(Platform);
+	}
+
+	static void ModifyCompilationEnvironment(EShaderPlatform Platform, FShaderCompilerEnvironment& OutEnvironment)
+	{
+		OutEnvironment.SetDefine(TEXT("DOWNSAMPLE_FACTOR"), GAODownsampleFactor);
+	}
+
+	/** Default constructor. */
+	FDistanceFieldShadowingUpsamplePS() {}
+
+	/** Initialization constructor. */
+	FDistanceFieldShadowingUpsamplePS(const ShaderMetaType::CompiledShaderInitializerType& Initializer)
+		: FGlobalShader(Initializer)
+	{
+		DeferredParameters.Bind(Initializer.ParameterMap);
+		ShadowFactorsTexture.Bind(Initializer.ParameterMap,TEXT("ShadowFactorsTexture"));
+		ShadowFactorsSampler.Bind(Initializer.ParameterMap,TEXT("ShadowFactorsSampler"));
+	}
+
+	void SetParameters(FRHICommandList& RHICmdList, const FSceneView& View, const FProjectedShadowInfo* ShadowInfo, TRefCountPtr<IPooledRenderTarget>& ShadowFactorsTextureValue)
+	{
+		const FPixelShaderRHIParamRef ShaderRHI = GetPixelShader();
+
+		FGlobalShader::SetParameters(RHICmdList, ShaderRHI, View);
+		DeferredParameters.Set(RHICmdList, ShaderRHI, View);
+
+		SetTextureParameter(RHICmdList, ShaderRHI, ShadowFactorsTexture, ShadowFactorsSampler, TStaticSamplerState<SF_Bilinear>::GetRHI(), ShadowFactorsTextureValue->GetRenderTargetItem().ShaderResourceTexture);
+	}
+	// FShader interface.
+	virtual bool Serialize(FArchive& Ar)
+	{
+		bool bShaderHasOutdatedParameters = FGlobalShader::Serialize(Ar);
+		Ar << DeferredParameters;
+		Ar << ShadowFactorsTexture;
+		Ar << ShadowFactorsSampler;
+		return bShaderHasOutdatedParameters;
+	}
+
+private:
+
+	FDeferredPixelShaderParameters DeferredParameters;
+	FShaderResourceParameter ShadowFactorsTexture;
+	FShaderResourceParameter ShadowFactorsSampler;
+};
+
+IMPLEMENT_SHADER_TYPE(,FDistanceFieldShadowingUpsamplePS,TEXT("DistanceFieldShadowing"),TEXT("DistanceFieldShadowingUpsamplePS"),SF_Pixel);
+
+int32 FProjectedShadowInfo::UpdateShadowCastingObjectBuffers() const
+{
+	int32 NumObjects = 0;
+
+	//@todo - scene rendering allocator
+	TArray<FVector4> ObjectBoundsData;
+	TArray<FVector4> ObjectData;
+	TArray<FVector4> ObjectData2;
+
+	{
+		QUICK_SCOPE_CYCLE_COUNTER(STAT_GatherObjectData);
+
+		const int32 NumTexelsOneDimX = GDistanceFieldVolumeTextureAtlas.GetSizeX();
+		const int32 NumTexelsOneDimY = GDistanceFieldVolumeTextureAtlas.GetSizeY();
+		const int32 NumTexelsOneDimZ = GDistanceFieldVolumeTextureAtlas.GetSizeZ();
+		const FVector InvTextureDim(1.0f / NumTexelsOneDimX, 1.0f / NumTexelsOneDimY, 1.0f / NumTexelsOneDimZ);
+
+		ObjectBoundsData.Empty(SubjectPrimitives.Num());
+		ObjectData.Empty(FDistanceFieldObjectBuffers::ObjectDataStride * SubjectPrimitives.Num());
+		ObjectData2.Empty(FDistanceFieldObjectBuffers::ObjectData2Stride * SubjectPrimitives.Num());
+
+		for (int32 PrimitiveIndex = 0; PrimitiveIndex < SubjectPrimitives.Num() && NumObjects < MAX_uint16; PrimitiveIndex++)
+		{
+			const FPrimitiveSceneInfo* PrimitiveSceneInfo = SubjectPrimitives[PrimitiveIndex];
+
+			FBox LocalVolumeBounds;
+			FIntVector BlockMin;
+			FIntVector BlockSize;
+			bool bBuiltAsIfTwoSided;
+			PrimitiveSceneInfo->Proxy->GetDistancefieldAtlasData(LocalVolumeBounds, BlockMin, BlockSize, bBuiltAsIfTwoSided);
+
+			if (BlockMin.X >= 0 && BlockMin.Y >= 0 && BlockMin.Z >= 0)
+			{
+				const FMatrix LocalToWorld = PrimitiveSceneInfo->Proxy->GetLocalToWorld();
+
+				const FMatrix VolumeToWorld = FScaleMatrix(LocalVolumeBounds.GetExtent()) 
+					* FTranslationMatrix(LocalVolumeBounds.GetCenter())
+					* LocalToWorld;
+
+				const FVector4 ObjectBoundingSphere(VolumeToWorld.GetOrigin(), VolumeToWorld.GetScaleVector().Size());
+
+				ObjectBoundsData.Add(ObjectBoundingSphere);
+
+				const float MaxExtent = LocalVolumeBounds.GetExtent().GetMax();
+
+				const FMatrix UniformScaleVolumeToWorld = FScaleMatrix(MaxExtent) 
+					* FTranslationMatrix(LocalVolumeBounds.GetCenter())
+					* LocalToWorld;
+
+				const FVector InvBlockSize(1.0f / BlockSize.X, 1.0f / BlockSize.Y, 1.0f / BlockSize.Z);
+
+				//float3 VolumeUV = (VolumePosition / LocalPositionExtent * .5f * UVScale + .5f * UVScale + UVAdd;
+				const FVector LocalPositionExtent = LocalVolumeBounds.GetExtent() / FVector(MaxExtent);
+				const FVector UVScale = FVector(BlockSize) * InvTextureDim;
+				const float VolumeScale = UniformScaleVolumeToWorld.GetMaximumAxisScale();
+				ObjectData2.Add(*(FVector4*)&UniformScaleVolumeToWorld.M[0]);
+				ObjectData2.Add(*(FVector4*)&UniformScaleVolumeToWorld.M[1]);
+				ObjectData2.Add(*(FVector4*)&UniformScaleVolumeToWorld.M[2]);
+				ObjectData2.Add(*(FVector4*)&UniformScaleVolumeToWorld.M[3]);
+
+				const FMatrix WorldToVolume = UniformScaleVolumeToWorld.InverseFast();
+				// WorldToVolume
+				ObjectData.Add(*(FVector4*)&WorldToVolume.M[0]);
+				ObjectData.Add(*(FVector4*)&WorldToVolume.M[1]);
+				ObjectData.Add(*(FVector4*)&WorldToVolume.M[2]);
+				ObjectData.Add(*(FVector4*)&WorldToVolume.M[3]);
+
+				// Clamp to texel center by subtracting a half texel in the [-1,1] position space
+				// LocalPositionExtent
+				ObjectData.Add(FVector4(LocalPositionExtent - InvBlockSize, bBuiltAsIfTwoSided ? 1.0f : 0.0f));
+
+				// UVScale
+				ObjectData.Add(FVector4(FVector(BlockSize) * InvTextureDim * .5f / LocalPositionExtent, VolumeScale));
+
+				// UVAdd
+				ObjectData.Add(FVector4(FVector(BlockMin) * InvTextureDim + .5f * UVScale, 0));
+
+				// Padding
+				ObjectData.Add(FVector4(0, 0, 0, 0));
+
+				checkSlow(ObjectData.Num() % FDistanceFieldObjectBuffers::ObjectDataStride == 0);
+				checkSlow(ObjectData2.Num() % FDistanceFieldObjectBuffers::ObjectData2Stride == 0);
+
+				NumObjects++;
+			}
+		}
+	}
+
+	if (NumObjects > 0)
+	{
+		QUICK_SCOPE_CYCLE_COUNTER(STAT_UpdateObjectBuffers);
+
+		if (NumObjects > GDirectionalLightObjectBuffers.MaxObjects)
+		{
+			// Allocate with slack
+			GDirectionalLightObjectBuffers.MaxObjects = NumObjects * 5 / 4;
+			GDirectionalLightObjectBuffers.UpdateRHI();
+		}
+
+		checkSlow(NumObjects < MAX_uint16);
+
+		void* LockedBuffer = RHILockVertexBuffer(GDirectionalLightObjectBuffers.ObjectData.Bounds, 0, GDirectionalLightObjectBuffers.ObjectData.Bounds->GetSize(), RLM_WriteOnly);
+		FPlatformMemory::Memcpy(LockedBuffer, ObjectBoundsData.GetData(), ObjectBoundsData.GetTypeSize() * ObjectBoundsData.Num());
+		RHIUnlockVertexBuffer(GDirectionalLightObjectBuffers.ObjectData.Bounds);
+
+		LockedBuffer = RHILockVertexBuffer(GDirectionalLightObjectBuffers.ObjectData.Data, 0, GDirectionalLightObjectBuffers.ObjectData.Data->GetSize(), RLM_WriteOnly);
+		FPlatformMemory::Memcpy(LockedBuffer, ObjectData.GetData(), ObjectData.GetTypeSize() * ObjectData.Num());
+		RHIUnlockVertexBuffer(GDirectionalLightObjectBuffers.ObjectData.Data);
+
+		LockedBuffer = RHILockVertexBuffer(GDirectionalLightObjectBuffers.ObjectData.Data2, 0, GDirectionalLightObjectBuffers.ObjectData.Data2->GetSize(), RLM_WriteOnly);
+		FPlatformMemory::Memcpy(LockedBuffer, ObjectData2.GetData(), ObjectData2.GetTypeSize() * ObjectData2.Num());
+		RHIUnlockVertexBuffer(GDirectionalLightObjectBuffers.ObjectData.Data2);
+	}
+
+	return NumObjects;
+}
+
+void FProjectedShadowInfo::RenderRayTracedDistanceFieldProjection(FRHICommandListImmediate& RHICmdList, int32 ViewIndex, const FViewInfo& View) const
+{
+	if (GDistanceFieldShadowing
+		&& View.GetFeatureLevel() >= ERHIFeatureLevel::SM5
+		&& DoesPlatformSupportDistanceFieldShadowing(GRHIShaderPlatform))
+	{
+		QUICK_SCOPE_CYCLE_COUNTER(STAT_RenderRayTracedDistanceFieldShadows);
+		SCOPED_DRAW_EVENT(RayTracedDistanceFieldShadow, DEC_SCENE_ITEMS);
+
+		// Update the global distance field atlas
+		GDistanceFieldVolumeTextureAtlas.UpdateAllocations();
+
+		if (GDistanceFieldVolumeTextureAtlas.VolumeTextureRHI)
+		{
+			const int32 NumObjects = UpdateShadowCastingObjectBuffers();
+
+			if (NumObjects > 0)
+			{
+				TRefCountPtr<IPooledRenderTarget> RayTracedShadowsRT;
+
+				{
+					const FIntPoint ExpandedBufferSize = GetBufferSizeForAO();
+					FPooledRenderTargetDesc Desc(FPooledRenderTargetDesc::Create2DDesc(ExpandedBufferSize / GAODownsampleFactor, PF_G16R16F, TexCreate_None, TexCreate_RenderTargetable | TexCreate_UAV, false));
+					GRenderTargetPool.FindFreeElement(Desc, RayTracedShadowsRT, TEXT("RayTracedShadows"));
+				}
+
+				{
+					GSceneRenderTargets.FinishRenderingLightAttenuation(RHICmdList);
+					SetRenderTarget(RHICmdList, NULL, NULL);
+
+					uint32 GroupSizeX = FMath::DivideAndRoundUp(View.ViewRect.Size().X / GAODownsampleFactor, GDistanceFieldAOTileSizeX);
+					uint32 GroupSizeY = FMath::DivideAndRoundUp(View.ViewRect.Size().Y / GAODownsampleFactor, GDistanceFieldAOTileSizeY);
+
+					{
+						SCOPED_DRAW_EVENT(RayTraceShadows, DEC_SCENE_ITEMS);
+						TShaderMapRef<FDistanceFieldShadowingCS> ComputeShader(View.ShaderMap);
+
+						RHICmdList.SetComputeShader(ComputeShader->GetComputeShader());
+						ComputeShader->SetParameters(RHICmdList, View, this, RayTracedShadowsRT->GetRenderTargetItem(), NumObjects, FVector2D(GroupSizeX, GroupSizeY));
+						DispatchComputeShader(RHICmdList, *ComputeShader, GroupSizeX, GroupSizeY, 1);
+
+						ComputeShader->UnsetParameters(RHICmdList);
+					}
+				}
+
+				{
+					GSceneRenderTargets.BeginRenderingLightAttenuation(RHICmdList);
+
+					SCOPED_DRAW_EVENT(Upsample, DEC_SCENE_ITEMS);
+
+					RHICmdList.SetViewport(View.ViewRect.Min.X, View.ViewRect.Min.Y, 0.0f, View.ViewRect.Max.X, View.ViewRect.Max.Y, 1.0f);
+					RHICmdList.SetRasterizerState(TStaticRasterizerState<FM_Solid, CM_None>::GetRHI());
+					RHICmdList.SetDepthStencilState(TStaticDepthStencilState<false, CF_Always>::GetRHI());
+					
+					check(CascadeSettings.FadePlaneLength == 0);
+					// first cascade rendered or old method doesn't require fading (CO_Min is needed to combine multiple shadow passes)
+					// The ray traced cascade should always be first
+					RHICmdList.SetBlendState(TStaticBlendState<CW_RG, BO_Min, BF_One, BF_One>::GetRHI());
+
+					TShaderMapRef<FPostProcessVS> VertexShader( View.ShaderMap );
+					TShaderMapRef<FDistanceFieldShadowingUpsamplePS> PixelShader( View.ShaderMap );
+
+					static FGlobalBoundShaderState BoundShaderState;
+
+					SetGlobalBoundShaderState(RHICmdList, View.GetFeatureLevel(), BoundShaderState, GFilterVertexDeclaration.VertexDeclarationRHI, *VertexShader, *PixelShader);
+
+					PixelShader->SetParameters(RHICmdList, View, this, RayTracedShadowsRT);
+
+					DrawRectangle( 
+						RHICmdList,
+						0, 0, 
+						View.ViewRect.Width(), View.ViewRect.Height(),
+						View.ViewRect.Min.X / GAODownsampleFactor, View.ViewRect.Min.Y / GAODownsampleFactor, 
+						View.ViewRect.Width() / GAODownsampleFactor, View.ViewRect.Height() / GAODownsampleFactor,
+						FIntPoint(View.ViewRect.Width(), View.ViewRect.Height()),
+						GetBufferSizeForAO() / GAODownsampleFactor,
+						*VertexShader);
+				}
+			}
+		}
+	}
+}
diff --git a/Engine/Source/Runtime/Renderer/Private/DistanceFieldSurfaceCacheLighting.cpp b/Engine/Source/Runtime/Renderer/Private/DistanceFieldSurfaceCacheLighting.cpp
--- a/Engine/Source/Runtime/Renderer/Private/DistanceFieldSurfaceCacheLighting.cpp
+++ b/Engine/Source/Runtime/Renderer/Private/DistanceFieldSurfaceCacheLighting.cpp
@@ -1,7 +1,7 @@
-// Copyright 1998-2013 Epic Games, Inc. All Rights Reserved.
+// Copyright 1998-2014 Epic Games, Inc. All Rights Reserved.
 
 /*=============================================================================
-	
+	DistanceFieldSurfaceCacheLighting.cpp
 =============================================================================*/
 
 #include "RendererPrivate.h"
@@ -246,66 +246,12 @@ uint32 GAONumConeSteps = 10;
 
 TGlobalResource<FCircleVertexBuffer> GCircleVertexBuffer;
 
-class FAOObjectData
-{
-public:
-
-	FVertexBufferRHIRef Bounds;
-	FVertexBufferRHIRef Data;
-	FVertexBufferRHIRef Data2;
-	FShaderResourceViewRHIRef BoundsSRV;
-	FShaderResourceViewRHIRef DataSRV;
-	FShaderResourceViewRHIRef Data2SRV;
-};
-
-class FAOObjectBuffers : public FRenderResource
-{
-public:
-
-	// In float4's
-	static int32 ObjectDataStride;
-	static int32 ObjectData2Stride;
-
-	int32 MaxObjects;
-	FAOObjectData ObjectData;
-
-	FAOObjectBuffers()
-	{
-		MaxObjects = 0;
-	}
-
-	virtual void InitDynamicRHI()  override
-	{
-		if (MaxObjects > 0)
-		{
-			FRHIResourceCreateInfo CreateInfo;
-			ObjectData.Bounds = RHICreateVertexBuffer(MaxObjects * sizeof(FVector4), BUF_Volatile | BUF_ShaderResource, CreateInfo);
-			ObjectData.Data = RHICreateVertexBuffer(MaxObjects * sizeof(FVector4) * ObjectDataStride, BUF_Volatile | BUF_ShaderResource, CreateInfo);
-			ObjectData.Data2 = RHICreateVertexBuffer(MaxObjects * sizeof(FVector4) * ObjectData2Stride, BUF_Volatile | BUF_ShaderResource, CreateInfo);
-
-			ObjectData.BoundsSRV = RHICreateShaderResourceView(ObjectData.Bounds, sizeof(FVector4), PF_A32B32G32R32F);
-			ObjectData.DataSRV = RHICreateShaderResourceView(ObjectData.Data, sizeof(FVector4), PF_A32B32G32R32F);
-			ObjectData.Data2SRV = RHICreateShaderResourceView(ObjectData.Data2, sizeof(FVector4), PF_A32B32G32R32F);
-		}
-	}
-
-	virtual void ReleaseDynamicRHI() override
-	{
-		ObjectData.Bounds.SafeRelease();
-		ObjectData.Data.SafeRelease();
-		ObjectData.Data2.SafeRelease();
-		ObjectData.BoundsSRV.SafeRelease(); 
-		ObjectData.DataSRV.SafeRelease(); 
-		ObjectData.Data2SRV.SafeRelease(); 
-	}
-};
-
 // Only 7 needed atm, pad to cache line size
 // Must match equivalent shader defines
-int32 FAOObjectBuffers::ObjectDataStride = 8;
-int32 FAOObjectBuffers::ObjectData2Stride = 4;
+int32 FDistanceFieldObjectBuffers::ObjectDataStride = 8;
+int32 FDistanceFieldObjectBuffers::ObjectData2Stride = 4;
 
-TGlobalResource<FAOObjectBuffers> GAOObjectBuffers;
+TGlobalResource<FDistanceFieldObjectBuffers> GAOObjectBuffers;
 
 void FSceneViewState::DestroyAOTileResources()
 {
@@ -341,41 +287,6 @@ FAutoConsoleCommandWithWorld ClearCacheConsoleCommand(
 	FConsoleCommandWithWorldDelegate::CreateStatic(OnClearIrradianceCache)
 	);
 
-class FAOObjectBufferParameters
-{
-public:
-	void Bind(const FShaderParameterMap& ParameterMap)
-	{
-		ObjectBounds.Bind(ParameterMap, TEXT("ObjectBounds"));
-		ObjectData.Bind(ParameterMap, TEXT("ObjectData"));
-		ObjectData2.Bind(ParameterMap, TEXT("ObjectData2"));
-		NumObjects.Bind(ParameterMap, TEXT("NumObjects"));
-	}
-
-	template<typename TParamRef>
-	void Set(FRHICommandList& RHICmdList, const TParamRef& ShaderRHI, int32 NumObjectsValue)
-	{
-		SetSRVParameter(RHICmdList, ShaderRHI, ObjectBounds, GAOObjectBuffers.ObjectData.BoundsSRV);
-		SetSRVParameter(RHICmdList, ShaderRHI, ObjectData, GAOObjectBuffers.ObjectData.DataSRV);
-		SetSRVParameter(RHICmdList, ShaderRHI, ObjectData2, GAOObjectBuffers.ObjectData.Data2SRV);
-
-		SetShaderValue(RHICmdList, ShaderRHI, NumObjects, NumObjectsValue);
-	}
-
-	friend FArchive& operator<<(FArchive& Ar, FAOObjectBufferParameters& P)
-	{
-		Ar << P.ObjectBounds << P.ObjectData << P.ObjectData2 << P.NumObjects;
-		return Ar;
-	}
-
-private:
-	FShaderResourceParameter ObjectBounds;
-	FShaderResourceParameter ObjectData;
-	FShaderResourceParameter ObjectData2;
-	FShaderParameter NumObjects;
-};
-
-
 /**  */
 class FBuildTileConesCS : public FGlobalShader
 {
@@ -499,7 +410,7 @@ public:
 		const FVertexShaderRHIParamRef ShaderRHI = GetVertexShader();
 		FGlobalShader::SetParameters(RHICmdList, ShaderRHI, View);
 		
-		ObjectParameters.Set(RHICmdList, ShaderRHI, NumObjects);
+		ObjectParameters.Set(RHICmdList, ShaderRHI, GAOObjectBuffers, NumObjects);
 		AOParameters.Set(RHICmdList, ShaderRHI, Parameters);
 
 		const int32 NumRings = StencilingGeometry::GLowPolyStencilSphereVertexBuffer.GetNumRings();
@@ -521,7 +432,7 @@ public:
 	}
 
 private:
-	FAOObjectBufferParameters ObjectParameters;
+	FDistanceFieldObjectBufferParameters ObjectParameters;
 	FAOParameters AOParameters;
 	FShaderParameter ConservativeRadiusScale;
 };
@@ -665,7 +576,7 @@ public:
 
 		FGlobalShader::SetParameters(RHICmdList, ShaderRHI, View);
 		DeferredParameters.Set(RHICmdList, ShaderRHI, View);
-		ObjectParameters.Set(RHICmdList, ShaderRHI, NumObjects);
+		ObjectParameters.Set(RHICmdList, ShaderRHI, GAOObjectBuffers, NumObjects);
 		AOParameters.Set(RHICmdList, ShaderRHI, Parameters);
 
 		FTileIntersectionResources* TileIntersectionResources = ((FSceneViewState*)View.State)->AOTileIntersectionResources;
@@ -703,7 +614,7 @@ public:
 private:
 
 	FDeferredPixelShaderParameters DeferredParameters;
-	FAOObjectBufferParameters ObjectParameters;
+	FDistanceFieldObjectBufferParameters ObjectParameters;
 	FAOParameters AOParameters;
 	FRWShaderParameter TileHeadData;
 	FRWShaderParameter TileArrayData;
@@ -793,7 +704,7 @@ public:
 		const FPixelShaderRHIParamRef ShaderRHI = GetPixelShader();
 
 		FGlobalShader::SetParameters(RHICmdList, ShaderRHI, View);
-		ObjectParameters.Set(RHICmdList, ShaderRHI, NumObjects);
+		ObjectParameters.Set(RHICmdList, ShaderRHI, GAOObjectBuffers, NumObjects);
 		AOParameters.Set(RHICmdList, ShaderRHI, Parameters);
 		DeferredParameters.Set(RHICmdList, ShaderRHI, View);
 
@@ -821,7 +732,7 @@ public:
 
 private:
 
-	FAOObjectBufferParameters ObjectParameters;
+	FDistanceFieldObjectBufferParameters ObjectParameters;
 	FDeferredPixelShaderParameters DeferredParameters;
 	FShaderResourceParameter TileHeadData;
 	FShaderResourceParameter TileHeadDataUnpacked;
@@ -877,7 +788,7 @@ public:
 
 		DistanceFieldNormal.SetTexture(RHICmdList, ShaderRHI, DistanceFieldNormalValue.ShaderResourceTexture, DistanceFieldNormalValue.UAV);
 
-		ObjectParameters.Set(RHICmdList, ShaderRHI, NumObjects);
+		ObjectParameters.Set(RHICmdList, ShaderRHI, GAOObjectBuffers, NumObjects);
 		AOParameters.Set(RHICmdList, ShaderRHI, Parameters);
 		DeferredParameters.Set(RHICmdList, ShaderRHI, View);
 
@@ -916,7 +827,7 @@ public:
 private:
 
 	FRWShaderParameter DistanceFieldNormal;
-	FAOObjectBufferParameters ObjectParameters;
+	FDistanceFieldObjectBufferParameters ObjectParameters;
 	FShaderParameter UVToTileHead;
 	FDeferredPixelShaderParameters DeferredParameters;
 	FShaderResourceParameter TileHeadData;
@@ -1514,7 +1425,7 @@ public:
 		FComputeShaderRHIParamRef ShaderRHI = GetComputeShader();
 		FGlobalShader::SetParameters(RHICmdList, ShaderRHI, View);
 		DeferredParameters.Set(RHICmdList, ShaderRHI, View);
-		ObjectParameters.Set(RHICmdList, ShaderRHI, NumObjects);
+		ObjectParameters.Set(RHICmdList, ShaderRHI, GAOObjectBuffers, NumObjects);
 		AOParameters.Set(RHICmdList, ShaderRHI, Parameters);
 		AOLevelParameters.Set(RHICmdList, ShaderRHI, View, DownsampleFactorValue);
 
@@ -1634,7 +1545,7 @@ public:
 private:
 
 	FDeferredPixelShaderParameters DeferredParameters;
-	FAOObjectBufferParameters ObjectParameters;
+	FDistanceFieldObjectBufferParameters ObjectParameters;
 	FAOParameters AOParameters;
 	FAOLevelParameters AOLevelParameters;
 	FShaderResourceParameter IrradianceCacheNormal;
@@ -2738,11 +2649,12 @@ void UpdateVisibleObjectBuffers(const FScene* Scene, const FDistanceFieldAOParam
 		const FVector InvTextureDim(1.0f / NumTexelsOneDimX, 1.0f / NumTexelsOneDimY, 1.0f / NumTexelsOneDimZ);
 
 		ObjectBoundsData.Empty(Scene->Primitives.Num() / 2);
-		ObjectData.Empty(FAOObjectBuffers::ObjectDataStride * Scene->Primitives.Num() / 2);
-		ObjectData2.Empty(FAOObjectBuffers::ObjectData2Stride * Scene->Primitives.Num() / 2);
+		ObjectData.Empty(FDistanceFieldObjectBuffers::ObjectDataStride * Scene->Primitives.Num() / 2);
+		ObjectData2.Empty(FDistanceFieldObjectBuffers::ObjectData2Stride * Scene->Primitives.Num() / 2);
 
 		//@todo - dense data-coherent pass like frustum culling in InitViews
-		for (int32 PrimitiveIndex = 0; PrimitiveIndex < Scene->Primitives.Num(); PrimitiveIndex++)
+		// FTileIntersectionResources uses 16 bit object indices
+		for (int32 PrimitiveIndex = 0; PrimitiveIndex < Scene->Primitives.Num() && NumObjects < MAX_uint16; PrimitiveIndex++)
 		{
 			const FPrimitiveSceneInfo* PrimitiveSceneInfo = Scene->Primitives[PrimitiveIndex];
 
@@ -2808,8 +2720,8 @@ void UpdateVisibleObjectBuffers(const FScene* Scene, const FDistanceFieldAOParam
 						// Padding
 						ObjectData.Add(FVector4(0, 0, 0, 0));
 
-						checkSlow(ObjectData.Num() % FAOObjectBuffers::ObjectDataStride == 0);
-						checkSlow(ObjectData2.Num() % FAOObjectBuffers::ObjectData2Stride == 0);
+						checkSlow(ObjectData.Num() % FDistanceFieldObjectBuffers::ObjectDataStride == 0);
+						checkSlow(ObjectData2.Num() % FDistanceFieldObjectBuffers::ObjectData2Stride == 0);
 
 						NumObjects++;
 					}
@@ -3098,12 +3010,19 @@ public:
 		: FGlobalShader(Initializer)
 	{
 		VisualizeMeshDistanceFields.Bind(Initializer.ParameterMap, TEXT("VisualizeMeshDistanceFields"));
+		NumGroups.Bind(Initializer.ParameterMap, TEXT("NumGroups"));
 		ObjectParameters.Bind(Initializer.ParameterMap);
 		DeferredParameters.Bind(Initializer.ParameterMap);
 		AOParameters.Bind(Initializer.ParameterMap);
 	}
 
-	void SetParameters(FRHICommandList& RHICmdList, const FSceneView& View, FSceneRenderTargetItem& VisualizeMeshDistanceFieldsValue, int32 NumObjects, const FDistanceFieldAOParameters& Parameters)
+	void SetParameters(
+		FRHICommandList& RHICmdList, 
+		const FSceneView& View, 
+		FSceneRenderTargetItem& VisualizeMeshDistanceFieldsValue, 
+		int32 NumObjects, 
+		FVector2D NumGroupsValue,
+		const FDistanceFieldAOParameters& Parameters)
 	{
 		const FComputeShaderRHIParamRef ShaderRHI = GetComputeShader();
 
@@ -3111,9 +3030,11 @@ public:
 
 		VisualizeMeshDistanceFields.SetTexture(RHICmdList, ShaderRHI, VisualizeMeshDistanceFieldsValue.ShaderResourceTexture, VisualizeMeshDistanceFieldsValue.UAV);
 
-		ObjectParameters.Set(RHICmdList, ShaderRHI, NumObjects);
+		ObjectParameters.Set(RHICmdList, ShaderRHI, GAOObjectBuffers, NumObjects);
 		AOParameters.Set(RHICmdList, ShaderRHI, Parameters);
 		DeferredParameters.Set(RHICmdList, ShaderRHI, View);
+
+		SetShaderValue(RHICmdList, ShaderRHI, NumGroups, NumGroupsValue);
 	}
 
 	void UnsetParameters(FRHICommandList& RHICmdList)
@@ -3126,6 +3047,7 @@ public:
 	{
 		bool bShaderHasOutdatedParameters = FGlobalShader::Serialize(Ar);
 		Ar << VisualizeMeshDistanceFields;
+		Ar << NumGroups;
 		Ar << ObjectParameters;
 		Ar << DeferredParameters;
 		Ar << AOParameters;
@@ -3135,7 +3057,8 @@ public:
 private:
 
 	FRWShaderParameter VisualizeMeshDistanceFields;
-	FAOObjectBufferParameters ObjectParameters;
+	FShaderParameter NumGroups;
+	FDistanceFieldObjectBufferParameters ObjectParameters;
 	FDeferredPixelShaderParameters DeferredParameters;
 	FAOParameters AOParameters;
 };
@@ -3248,7 +3171,7 @@ void FDeferredShadingSceneRenderer::RenderMeshDistanceFieldVisualization(FRHICom
 							TShaderMapRef<FVisualizeMeshDistanceFieldCS> ComputeShader(View.ShaderMap);
 
 							RHICmdList.SetComputeShader(ComputeShader->GetComputeShader());
-							ComputeShader->SetParameters(RHICmdList, View, VisualizeResultRT->GetRenderTargetItem(), NumObjects, Parameters);
+							ComputeShader->SetParameters(RHICmdList, View, VisualizeResultRT->GetRenderTargetItem(), NumObjects, FVector2D(GroupSizeX, GroupSizeY), Parameters);
 							DispatchComputeShader(RHICmdList, *ComputeShader, GroupSizeX, GroupSizeY, 1);
 
 							ComputeShader->UnsetParameters(RHICmdList);
@@ -3384,7 +3307,8 @@ void FDeferredShadingSceneRenderer::RenderDynamicSkyLighting(FRHICommandListImme
 
 		if (Scene->SkyLight->bCastShadows
 			&& ViewFamily.EngineShowFlags.DistanceFieldAO 
-			&& !ViewFamily.EngineShowFlags.VisualizeDistanceFieldAO)
+			&& !ViewFamily.EngineShowFlags.VisualizeDistanceFieldAO
+			&& !ViewFamily.EngineShowFlags.VisualizeMeshDistanceFields)
 		{
 			bApplyShadowing = RenderDistanceFieldAOSurfaceCache(RHICmdList, Parameters, DynamicBentNormalAO, false);
 		}
diff --git a/Engine/Source/Runtime/Renderer/Private/ShadowRendering.cpp b/Engine/Source/Runtime/Renderer/Private/ShadowRendering.cpp
--- a/Engine/Source/Runtime/Renderer/Private/ShadowRendering.cpp
+++ b/Engine/Source/Runtime/Renderer/Private/ShadowRendering.cpp
@@ -1737,6 +1737,12 @@ void FProjectedShadowInfo::RenderProjection(FRHICommandListImmediate& RHICmdList
 		return;
 	}
 
+	if (bRayTracedDistanceFieldShadow)
+	{
+		RenderRayTracedDistanceFieldProjection(RHICmdList, ViewIndex, *View);
+		return;
+	}
+
 	// The shadow transforms and view transforms are relative to different origins, so the world coordinates need to be translated.
 	const FVector4 PreShadowToPreViewTranslation(View->ViewMatrices.PreViewTranslation - PreShadowTranslation,0);
 
@@ -3136,7 +3142,7 @@ bool FDeferredShadingSceneRenderer::RenderProjectedShadows(FRHICommandListImmedi
 			for (int32 ShadowIndex = 0; ShadowIndex < Shadows.Num(); ShadowIndex++)
 			{
 				FProjectedShadowInfo* ProjectedShadowInfo = Shadows[ShadowIndex];
-				if (ProjectedShadowInfo->bAllocated && !ProjectedShadowInfo->bTranslucentShadow)
+				if (ProjectedShadowInfo->bAllocated && !ProjectedShadowInfo->bTranslucentShadow && !ProjectedShadowInfo->bRayTracedDistanceFieldShadow)
 				{
 					ProjectedShadowInfo->RenderDepth(RHICmdList, this);
 				}
@@ -3158,6 +3164,8 @@ bool FDeferredShadingSceneRenderer::RenderProjectedShadows(FRHICommandListImmedi
 
 				if (ProjectedShadowInfo->bAllocated
 					&& ProjectedShadowInfo->bWholeSceneShadow
+					// Not supported on translucency yet
+					&& !ProjectedShadowInfo->bRayTracedDistanceFieldShadow
 					// Don't inject shadowed lighting with whole scene shadows used for previewing a light with static shadows,
 					// Since that would cause a mismatch with the built lighting
 					// However, stationary directional lights allow whole scene shadows that blend with precomputed shadowing
diff --git a/Engine/Source/Runtime/Renderer/Private/ShadowRendering.h b/Engine/Source/Runtime/Renderer/Private/ShadowRendering.h
--- a/Engine/Source/Runtime/Renderer/Private/ShadowRendering.h
+++ b/Engine/Source/Runtime/Renderer/Private/ShadowRendering.h
@@ -604,6 +604,9 @@ public:
 	/** Whether the shadow is a preshadow or not.  A preshadow is a per object shadow that handles the static environment casting on a dynamic receiver. */
 	uint32 bPreShadow : 1;
 
+	/** Whether the shadow is a directional light cascade that should be computed by ray tracing mesh distance fields. */
+	uint32 bRayTracedDistanceFieldShadow : 1;
+
 	uint32 bValidTransform : 1;
 
 	TBitArray<SceneRenderingBitArrayAllocator> StaticMeshWholeSceneShadowDepthMap;
@@ -676,6 +679,9 @@ public:
 	 */
 	void RenderProjection(FRHICommandListImmediate& RHICmdList, int32 ViewIndex, const class FViewInfo* View) const;
 
+	/** Renders ray traced distance field shadows. */
+	void RenderRayTracedDistanceFieldProjection(FRHICommandListImmediate& RHICmdList, int32 ViewIndex, const class FViewInfo& View) const;
+
 	/** Render one pass point light shadow projections. */
 	void RenderOnePassPointLightProjection(FRHICommandListImmediate& RHICmdList, int32 ViewIndex, const FViewInfo& View) const;
 
@@ -757,6 +763,9 @@ private:
 	void GetShadowTypeNameForDrawEvent(FString& TypeName) const;
 
 	template <bool bReflectiveShadowmap> friend void DrawShadowMeshElements(FRHICommandList& RHICmdList, const FViewInfo& View, const FProjectedShadowInfo& ShadowInfo);
+
+	/** Updates object buffers needed by ray traced distance field shadows. */
+	int32 UpdateShadowCastingObjectBuffers() const;
 };
 
 /** Shader parameters for rendering the depth of a mesh for shadowing. */
diff --git a/Engine/Source/Runtime/Renderer/Private/ShadowSetup.cpp b/Engine/Source/Runtime/Renderer/Private/ShadowSetup.cpp
--- a/Engine/Source/Runtime/Renderer/Private/ShadowSetup.cpp
+++ b/Engine/Source/Runtime/Renderer/Private/ShadowSetup.cpp
@@ -11,12 +11,21 @@
 static float GMinScreenRadiusForShadowCaster = 0.03f;
 static TAutoConsoleVariable<float> CVarMinScreenRadiusForShadowCaster(
 	TEXT("r.Shadow.RadiusThreshold"),
-	0.03f,
-	TEXT("Cull shadow casters if they are too small, values is the minimal screen space bounding sphere radius\n")
+	GMinScreenRadiusForShadowCaster,
+	TEXT("Cull shadow casters if they are too small, value is the minimal screen space bounding sphere radius\n")
 	TEXT("(default 0.03)"),
 	ECVF_Scalability | ECVF_RenderThreadSafe
 	);
 
+static float GMinScreenRadiusForRayTracedCasters = 0.01f;
+static TAutoConsoleVariable<float> CVarMinScreenRadiusForRayTracedCasters(
+	TEXT("r.Shadow.RadiusThresholdRayTraced"),
+	GMinScreenRadiusForRayTracedCasters,
+	TEXT("Cull shadow casters in the ray traced shadows if they are too small, value is the minimal screen space bounding sphere radius\n")
+	TEXT("(default 0.01)"),
+	ECVF_Scalability | ECVF_RenderThreadSafe
+	);
+
 static float GMinScreenRadiusForShadowCasterRSM = 0.06f;
 static FAutoConsoleVariableRef CVarMinScreenRadiusForShadowCasterRSM(
 	TEXT("r.Shadow.RadiusThresholdRSM"),
@@ -383,7 +392,8 @@ FProjectedShadowInfo::FProjectedShadowInfo(
 	bOnePassPointLightShadow(false),
 	bReflectiveShadowmap(false),
 	bTranslucentShadow(bInTranslucentShadow),
-	bPreShadow(bInPreShadow)
+	bPreShadow(bInPreShadow),
+	bRayTracedDistanceFieldShadow(false)
 {
 	const FMatrix WorldToLightScaled = Initializer.WorldToLight * FScaleMatrix(Initializer.Scales);
 
@@ -507,6 +517,7 @@ FProjectedShadowInfo::FProjectedShadowInfo(
 ,	bReflectiveShadowmap(bInReflectiveShadowMap) 
 ,	bTranslucentShadow(false)
 ,	bPreShadow(false)
+,	bRayTracedDistanceFieldShadow(Initializer.bRayTracedDistanceFieldShadow)
 ,	bValidTransform(true)
 {
 	FVector	XAxis, YAxis;
@@ -637,217 +648,224 @@ void FProjectedShadowInfo::AddSubjectPrimitive(FPrimitiveSceneInfo* PrimitiveSce
 {
 	QUICK_SCOPE_CYCLE_COUNTER(STAT_AddSubjectPrimitive);
 
-	if (!ReceiverPrimitives.Contains(PrimitiveSceneInfo))
+	if (bRayTracedDistanceFieldShadow)
 	{
-		TArray<FViewInfo*, TInlineAllocator<1> > Views;
-		const bool bWholeSceneDirectionalShadow = IsWholeSceneDirectionalShadow();
-
-		if (bWholeSceneDirectionalShadow)
-		{
-			Views.Add(DependentView);
-		}
-		else
+		SubjectPrimitives.Add(PrimitiveSceneInfo);
+	}
+	else
+	{
+		if (!ReceiverPrimitives.Contains(PrimitiveSceneInfo))
 		{
-			check(ViewArray);
+			TArray<FViewInfo*, TInlineAllocator<1> > Views;
+			const bool bWholeSceneDirectionalShadow = IsWholeSceneDirectionalShadow();
 
-			for (int32 ViewIndex = 0; ViewIndex < ViewArray->Num(); ViewIndex++)
+			if (bWholeSceneDirectionalShadow)
 			{
-				Views.Add(&(*ViewArray)[ViewIndex]);
+				Views.Add(DependentView);
 			}
-		}
+			else
+			{
+				check(ViewArray);
 
-		bool bOpaqueRelevance = false;
-		bool bTranslucentRelevance = false;
-		bool bShadowRelevance = false;
-		bool bNeedsPreRenderView = false;
-		uint32 ViewMask = 0;
-		int32 PrimitiveId = PrimitiveSceneInfo->GetIndex();
-		const auto FeatureLevel = PrimitiveSceneInfo->Scene->GetFeatureLevel();
+				for (int32 ViewIndex = 0; ViewIndex < ViewArray->Num(); ViewIndex++)
+				{
+					Views.Add(&(*ViewArray)[ViewIndex]);
+				}
+			}
 
-		for (int32 ViewIndex = 0, Num = Views.Num(); ViewIndex < Num; ViewIndex++)
-		{
-			FViewInfo& CurrentView = *Views[ViewIndex];
-			FPrimitiveViewRelevance& ViewRelevance = CurrentView.PrimitiveViewRelevanceMap[PrimitiveId];
+			bool bOpaqueRelevance = false;
+			bool bTranslucentRelevance = false;
+			bool bShadowRelevance = false;
+			bool bNeedsPreRenderView = false;
+			uint32 ViewMask = 0;
+			int32 PrimitiveId = PrimitiveSceneInfo->GetIndex();
+			const auto FeatureLevel = PrimitiveSceneInfo->Scene->GetFeatureLevel();
 
-			if (!ViewRelevance.bInitializedThisFrame)
+			for (int32 ViewIndex = 0, Num = Views.Num(); ViewIndex < Num; ViewIndex++)
 			{
-				if( CurrentView.IsPerspectiveProjection() )
-				{
-					const FPrimitiveSceneProxy* Proxy = PrimitiveSceneInfo->Proxy;
+				FViewInfo& CurrentView = *Views[ViewIndex];
+				FPrimitiveViewRelevance& ViewRelevance = CurrentView.PrimitiveViewRelevanceMap[PrimitiveId];
 
-					// Compute the distance between the view and the primitive.
-					float DistanceSquared = (Proxy->GetBounds().Origin - CurrentView.ShadowViewMatrices.ViewOrigin).SizeSquared();
-
-					bool bIsDistanceCulled = CurrentView.IsDistanceCulled(
-						DistanceSquared,
-						Proxy->GetMinDrawDistance(),
-						Proxy->GetMaxDrawDistance(),
-						PrimitiveSceneInfo
-						);
-					if( bIsDistanceCulled )
+				if (!ViewRelevance.bInitializedThisFrame)
+				{
+					if( CurrentView.IsPerspectiveProjection() )
 					{
-						continue;
-					}
-				}
+						const FPrimitiveSceneProxy* Proxy = PrimitiveSceneInfo->Proxy;
 
-				// Compute the subject primitive's view relevance since it wasn't cached
-				// Update the main view's PrimitiveViewRelevanceMap
-				ViewRelevance = PrimitiveSceneInfo->Proxy->GetViewRelevance(&CurrentView);
+						// Compute the distance between the view and the primitive.
+						float DistanceSquared = (Proxy->GetBounds().Origin - CurrentView.ShadowViewMatrices.ViewOrigin).SizeSquared();
 
-				bNeedsPreRenderView |= ViewRelevance.bNeedsPreRenderView;
-				ViewMask |= (1 << ViewIndex);
-			}
+						bool bIsDistanceCulled = CurrentView.IsDistanceCulled(
+							DistanceSquared,
+							Proxy->GetMinDrawDistance(),
+							Proxy->GetMaxDrawDistance(),
+							PrimitiveSceneInfo
+							);
+						if( bIsDistanceCulled )
+						{
+							continue;
+						}
+					}
 
-			bOpaqueRelevance |= ViewRelevance.bOpaqueRelevance;
-			bTranslucentRelevance |= ViewRelevance.HasTranslucency();
-			bShadowRelevance |= ViewRelevance.bShadowRelevance;
-		}
+					// Compute the subject primitive's view relevance since it wasn't cached
+					// Update the main view's PrimitiveViewRelevanceMap
+					ViewRelevance = PrimitiveSceneInfo->Proxy->GetViewRelevance(&CurrentView);
 
-		if (bShadowRelevance)
-		{
-			// Update the primitive component's last render time. Allows the component to update when using bCastWhenHidden.
-			const float CurrentWorldTime = Views[0]->Family->CurrentWorldTime;
-			*(PrimitiveSceneInfo->ComponentLastRenderTime) = CurrentWorldTime;
-		}
+					bNeedsPreRenderView |= ViewRelevance.bNeedsPreRenderView;
+					ViewMask |= (1 << ViewIndex);
+				}
 
-		if (bNeedsPreRenderView)
-		{
-			// Call PreRenderView on primitives that weren't visible in any of the main views, but need to be rendered in this shadow's depth pass
-			PrimitiveSceneInfo->Proxy->PreRenderView(Views[0]->Family, ViewMask, Views[0]->FrameNumber);
-		}
+				bOpaqueRelevance |= ViewRelevance.bOpaqueRelevance;
+				bTranslucentRelevance |= ViewRelevance.HasTranslucency();
+				bShadowRelevance |= ViewRelevance.bShadowRelevance;
+			}
 
-		if (bOpaqueRelevance && bShadowRelevance)
-		{
-			const FPrimitiveSceneProxy* Proxy = PrimitiveSceneInfo->Proxy;
-			const FBoxSphereBounds& Bounds = Proxy->GetBounds();
-			bool bDrawingStaticMeshes = false;
+			if (bShadowRelevance)
+			{
+				// Update the primitive component's last render time. Allows the component to update when using bCastWhenHidden.
+				const float CurrentWorldTime = Views[0]->Family->CurrentWorldTime;
+				*(PrimitiveSceneInfo->ComponentLastRenderTime) = CurrentWorldTime;
+			}
 
-			if (PrimitiveSceneInfo->StaticMeshes.Num() > 0)
+			if (bNeedsPreRenderView)
 			{
-				for (int32 ViewIndex = 0; ViewIndex < Views.Num(); ViewIndex++)
-				{
-					FViewInfo& CurrentView = *Views[ViewIndex];
+				// Call PreRenderView on primitives that weren't visible in any of the main views, but need to be rendered in this shadow's depth pass
+				PrimitiveSceneInfo->Proxy->PreRenderView(Views[0]->Family, ViewMask, Views[0]->FrameNumber);
+			}
 
-					const float DistanceSquared = ( Bounds.Origin - CurrentView.ShadowViewMatrices.ViewOrigin ).SizeSquared();
-					const bool bDrawShadowDepth = FMath::Square( Bounds.SphereRadius ) > FMath::Square( CVarMinScreenRadiusForShadowCaster.GetValueOnRenderThread() ) * DistanceSquared;
-					if( !bDrawShadowDepth )
-					{
-						// cull object if it's too small to be considered as shadow caster
-						continue;
-					}
+			if (bOpaqueRelevance && bShadowRelevance)
+			{
+				const FPrimitiveSceneProxy* Proxy = PrimitiveSceneInfo->Proxy;
+				const FBoxSphereBounds& Bounds = Proxy->GetBounds();
+				bool bDrawingStaticMeshes = false;
 
-					// Update visibility for meshes which weren't visible in the main views or were visible with static relevance
-					if (!CurrentView.PrimitiveVisibilityMap[PrimitiveId] || CurrentView.PrimitiveViewRelevanceMap[PrimitiveId].bStaticRelevance)
+				if (PrimitiveSceneInfo->StaticMeshes.Num() > 0)
+				{
+					for (int32 ViewIndex = 0; ViewIndex < Views.Num(); ViewIndex++)
 					{
-						bool bUseExistingVisibility = false;
+						FViewInfo& CurrentView = *Views[ViewIndex];
 
-						if(!bReflectiveShadowmap) // Don't use existing visibility for RSMs 
+						const float DistanceSquared = ( Bounds.Origin - CurrentView.ShadowViewMatrices.ViewOrigin ).SizeSquared();
+						const bool bDrawShadowDepth = FMath::Square( Bounds.SphereRadius ) > FMath::Square( CVarMinScreenRadiusForShadowCaster.GetValueOnRenderThread() ) * DistanceSquared;
+						if( !bDrawShadowDepth )
 						{
-							for (int32 MeshIndex = 0; MeshIndex < PrimitiveSceneInfo->StaticMeshes.Num(); MeshIndex++)
-							{
-								const FStaticMesh& StaticMesh = PrimitiveSceneInfo->StaticMeshes[MeshIndex];
-								bool bMeshIsVisible = CurrentView.StaticMeshShadowDepthMap[StaticMesh.Id] && StaticMesh.CastShadow;
-								bUseExistingVisibility = bUseExistingVisibility || bMeshIsVisible;
-
-								if (bMeshIsVisible && bWholeSceneDirectionalShadow)
-								{
-									StaticMeshWholeSceneShadowDepthMap[StaticMesh.Id] = true;
-									StaticMeshWholeSceneShadowBatchVisibility[StaticMesh.Id] = CurrentView.StaticMeshBatchVisibility[StaticMesh.Id];
-								}
-							}
+							// cull object if it's too small to be considered as shadow caster
+							continue;
 						}
 
-						if (bUseExistingVisibility)
+						// Update visibility for meshes which weren't visible in the main views or were visible with static relevance
+						if (!CurrentView.PrimitiveVisibilityMap[PrimitiveId] || CurrentView.PrimitiveViewRelevanceMap[PrimitiveId].bStaticRelevance)
 						{
-							bDrawingStaticMeshes = true;
-						}
-						// Don't overwrite visibility set by the main views
-						// This is necessary to avoid popping when transitioning between LODs, because on the frame of the transition, 
-						// The old LOD will continue to be drawn even though a different LOD would be chosen by distance.
-						else
-						{
-							int8 LODToRender = 0;
-							int32 ForcedLODLevel = (CurrentView.Family->EngineShowFlags.LOD) ? GetCVarForceLOD() : 0;
+							bool bUseExistingVisibility = false;
 
-							// Add the primitive's static mesh elements to the draw lists.
-							if ( bReflectiveShadowmap) 
+							if(!bReflectiveShadowmap) // Don't use existing visibility for RSMs 
 							{
-								LODToRender = -CHAR_MAX;
-								// Force the lowest detail LOD Level in reflective shadow maps.
-								for (int32 Index = 0; Index < PrimitiveSceneInfo->StaticMeshes.Num(); Index++)
+								for (int32 MeshIndex = 0; MeshIndex < PrimitiveSceneInfo->StaticMeshes.Num(); MeshIndex++)
 								{
-									LODToRender = FMath::Max<int8>(PrimitiveSceneInfo->StaticMeshes[Index].LODIndex, LODToRender);
+									const FStaticMesh& StaticMesh = PrimitiveSceneInfo->StaticMeshes[MeshIndex];
+									bool bMeshIsVisible = CurrentView.StaticMeshShadowDepthMap[StaticMesh.Id] && StaticMesh.CastShadow;
+									bUseExistingVisibility = bUseExistingVisibility || bMeshIsVisible;
+
+									if (bMeshIsVisible && bWholeSceneDirectionalShadow)
+									{
+										StaticMeshWholeSceneShadowDepthMap[StaticMesh.Id] = true;
+										StaticMeshWholeSceneShadowBatchVisibility[StaticMesh.Id] = CurrentView.StaticMeshBatchVisibility[StaticMesh.Id];
+									}
 								}
 							}
-							else
+
+							if (bUseExistingVisibility)
 							{
-								FPrimitiveBounds PrimitiveBounds;
-								PrimitiveBounds.Origin = Bounds.Origin;
-								PrimitiveBounds.SphereRadius = Bounds.SphereRadius;
-								LODToRender = ComputeLODForMeshes(PrimitiveSceneInfo->StaticMeshes, CurrentView, PrimitiveBounds.Origin, PrimitiveBounds.SphereRadius, ForcedLODLevel);
+								bDrawingStaticMeshes = true;
 							}
-
-							for (int32 MeshIndex = 0; MeshIndex < PrimitiveSceneInfo->StaticMeshes.Num(); MeshIndex++)
+							// Don't overwrite visibility set by the main views
+							// This is necessary to avoid popping when transitioning between LODs, because on the frame of the transition, 
+							// The old LOD will continue to be drawn even though a different LOD would be chosen by distance.
+							else
 							{
-								const FStaticMesh& StaticMesh = PrimitiveSceneInfo->StaticMeshes[MeshIndex];
-								if (StaticMesh.CastShadow && StaticMesh.LODIndex == LODToRender)
+								int8 LODToRender = 0;
+								int32 ForcedLODLevel = (CurrentView.Family->EngineShowFlags.LOD) ? GetCVarForceLOD() : 0;
+
+								// Add the primitive's static mesh elements to the draw lists.
+								if ( bReflectiveShadowmap) 
 								{
-									if (bWholeSceneDirectionalShadow)
+									LODToRender = -CHAR_MAX;
+									// Force the lowest detail LOD Level in reflective shadow maps.
+									for (int32 Index = 0; Index < PrimitiveSceneInfo->StaticMeshes.Num(); Index++)
 									{
-										StaticMeshWholeSceneShadowDepthMap[StaticMesh.Id] = true;
-										StaticMeshWholeSceneShadowBatchVisibility[StaticMesh.Id] = StaticMesh.Elements.Num() == 1 ? 1 : StaticMesh.VertexFactory->GetStaticBatchElementVisibility(*DependentView, &StaticMesh);
+										LODToRender = FMath::Max<int8>(PrimitiveSceneInfo->StaticMeshes[Index].LODIndex, LODToRender);
 									}
-									else
+								}
+								else
+								{
+									FPrimitiveBounds PrimitiveBounds;
+									PrimitiveBounds.Origin = Bounds.Origin;
+									PrimitiveBounds.SphereRadius = Bounds.SphereRadius;
+									LODToRender = ComputeLODForMeshes(PrimitiveSceneInfo->StaticMeshes, CurrentView, PrimitiveBounds.Origin, PrimitiveBounds.SphereRadius, ForcedLODLevel);
+								}
+
+								for (int32 MeshIndex = 0; MeshIndex < PrimitiveSceneInfo->StaticMeshes.Num(); MeshIndex++)
+								{
+									const FStaticMesh& StaticMesh = PrimitiveSceneInfo->StaticMeshes[MeshIndex];
+									if (StaticMesh.CastShadow && StaticMesh.LODIndex == LODToRender)
 									{
-										CurrentView.StaticMeshShadowDepthMap[StaticMesh.Id] = true;
-										CurrentView.StaticMeshBatchVisibility[StaticMesh.Id] = StaticMesh.Elements.Num() == 1 ? 1 : StaticMesh.VertexFactory->GetStaticBatchElementVisibility(CurrentView, &StaticMesh);
-									}
+										if (bWholeSceneDirectionalShadow)
+										{
+											StaticMeshWholeSceneShadowDepthMap[StaticMesh.Id] = true;
+											StaticMeshWholeSceneShadowBatchVisibility[StaticMesh.Id] = StaticMesh.Elements.Num() == 1 ? 1 : StaticMesh.VertexFactory->GetStaticBatchElementVisibility(*DependentView, &StaticMesh);
+										}
+										else
+										{
+											CurrentView.StaticMeshShadowDepthMap[StaticMesh.Id] = true;
+											CurrentView.StaticMeshBatchVisibility[StaticMesh.Id] = StaticMesh.Elements.Num() == 1 ? 1 : StaticMesh.VertexFactory->GetStaticBatchElementVisibility(CurrentView, &StaticMesh);
+										}
 
-									bDrawingStaticMeshes = true;
+										bDrawingStaticMeshes = true;
+									}
 								}
 							}
 						}
 					}
 				}
-			}
 
-			if (bDrawingStaticMeshes)
-			{
-				if (!bWholeSceneDirectionalShadow)
+				if (bDrawingStaticMeshes)
 				{
-					// Add the primitive's static mesh elements to the draw lists.
-					for (int32 MeshIndex = 0; MeshIndex < PrimitiveSceneInfo->StaticMeshes.Num(); MeshIndex++)
+					if (!bWholeSceneDirectionalShadow)
 					{
-						FStaticMesh& StaticMesh = PrimitiveSceneInfo->StaticMeshes[MeshIndex];
-						if (StaticMesh.CastShadow)
+						// Add the primitive's static mesh elements to the draw lists.
+						for (int32 MeshIndex = 0; MeshIndex < PrimitiveSceneInfo->StaticMeshes.Num(); MeshIndex++)
 						{
-							const FMaterialRenderProxy* MaterialRenderProxy = StaticMesh.MaterialRenderProxy;
-							const FMaterial* Material = MaterialRenderProxy->GetMaterial(FeatureLevel);
-							const EBlendMode BlendMode = Material->GetBlendMode();
-							const EMaterialShadingModel ShadingModel = Material->GetShadingModel();
-
-							if(((!IsTranslucentBlendMode(BlendMode)) && ShadingModel != MSM_Unlit) || (bReflectiveShadowmap && Material->ShouldInjectEmissiveIntoLPV())) 
+							FStaticMesh& StaticMesh = PrimitiveSceneInfo->StaticMeshes[MeshIndex];
+							if (StaticMesh.CastShadow)
 							{
-								const bool bTwoSided = Material->IsTwoSided() || PrimitiveSceneInfo->Proxy->CastsShadowAsTwoSided();
-								OverrideWithDefaultMaterialForShadowDepth(MaterialRenderProxy, Material, bReflectiveShadowmap, FeatureLevel);
-								SubjectMeshElements.Add(FShadowStaticMeshElement(MaterialRenderProxy, Material, &StaticMesh,bTwoSided));
+								const FMaterialRenderProxy* MaterialRenderProxy = StaticMesh.MaterialRenderProxy;
+								const FMaterial* Material = MaterialRenderProxy->GetMaterial(FeatureLevel);
+								const EBlendMode BlendMode = Material->GetBlendMode();
+								const EMaterialShadingModel ShadingModel = Material->GetShadingModel();
+
+								if(((!IsTranslucentBlendMode(BlendMode)) && ShadingModel != MSM_Unlit) || (bReflectiveShadowmap && Material->ShouldInjectEmissiveIntoLPV())) 
+								{
+									const bool bTwoSided = Material->IsTwoSided() || PrimitiveSceneInfo->Proxy->CastsShadowAsTwoSided();
+									OverrideWithDefaultMaterialForShadowDepth(MaterialRenderProxy, Material, bReflectiveShadowmap, FeatureLevel);
+									SubjectMeshElements.Add(FShadowStaticMeshElement(MaterialRenderProxy, Material, &StaticMesh,bTwoSided));
+								}
 							}
 						}
 					}
 				}
+				else
+				{
+					// Add the primitive to the subject primitive list.
+					SubjectPrimitives.Add(PrimitiveSceneInfo);
+				}
 			}
-			else
+
+			// Add translucent shadow casting primitives to SubjectTranslucentPrimitives
+			if (bTranslucentRelevance && bShadowRelevance && bTranslucentShadow)
 			{
-				// Add the primitive to the subject primitive list.
-				SubjectPrimitives.Add(PrimitiveSceneInfo);
+				SubjectTranslucentPrimitives.Add(PrimitiveSceneInfo);
 			}
 		}
-
-		// Add translucent shadow casting primitives to SubjectTranslucentPrimitives
-		if (bTranslucentRelevance && bShadowRelevance && bTranslucentShadow)
-		{
-			SubjectTranslucentPrimitives.Add(PrimitiveSceneInfo);
-		}
 	}
 }
 
@@ -1772,10 +1790,14 @@ inline void FDeferredShadingSceneRenderer::GatherShadowsForPrimitiveInner(
 					{
 						// Distance culling for RSMs
 						float MinScreenRadiusForShadowCaster = GMinScreenRadiusForShadowCaster;
-						if ( ProjectedShadowInfo->bReflectiveShadowmap )
+						if (ProjectedShadowInfo->bReflectiveShadowmap)
 						{
 							MinScreenRadiusForShadowCaster = GMinScreenRadiusForShadowCasterRSM;
 						}
+						else if (ProjectedShadowInfo->bRayTracedDistanceFieldShadow)
+						{
+							MinScreenRadiusForShadowCaster = GMinScreenRadiusForRayTracedCasters;
+						}
 
 						bool bScreenSpaceSizeCulled = false;
 						check( ProjectedShadowInfo->DependentView );
@@ -2027,6 +2049,7 @@ void FDeferredShadingSceneRenderer::InitDynamicShadows(FRHICommandListImmediate&
 										ViewDependentWholeSceneShadows.Add(ProjectedShadowInfo);
 									}
 								}
+
 								FSceneViewState* ViewState = (FSceneViewState*)View.State;
 								if(ViewState)
 								{
