fixed broken color grading (e.g. black/yellow image)

diff --git a/Engine/Source/Runtime/Renderer/Private/PostProcess/PostProcessTonemap.cpp b/Engine/Source/Runtime/Renderer/Private/PostProcess/PostProcessTonemap.cpp
--- a/Engine/Source/Runtime/Renderer/Private/PostProcess/PostProcessTonemap.cpp
+++ b/Engine/Source/Runtime/Renderer/Private/PostProcess/PostProcessTonemap.cpp
@@ -1184,12 +1184,12 @@ FRCPassPostProcessTonemap::FRCPassPostProcessTonemap(const FViewInfo& View, bool
 	: bDoGammaOnly(bInDoGammaOnly)
 {
 	uint32 ConfigBitmask = TonemapperGenerateBitmaskPC(&View, bDoGammaOnly, FRCPassPostProcessCombineLUTs::IsColorGradingLUTNeeded(&View));
-	ConfigIndex = TonemapperFindLeastExpensive(TonemapperConfBitmaskPC, sizeof(TonemapperConfBitmaskPC)/4, TonemapperCostTab, ConfigBitmask);
+	ConfigIndexPC = TonemapperFindLeastExpensive(TonemapperConfBitmaskPC, sizeof(TonemapperConfBitmaskPC)/4, TonemapperCostTab, ConfigBitmask);
 }
 
 bool FRCPassPostProcessTonemap::IsLUTNeeded() const
 {
-	uint32 ConfigBitmask = TonemapperConfBitmaskMobile[ConfigIndex];
+	uint32 ConfigBitmask = TonemapperConfBitmaskPC[ConfigIndexPC];
 
 	return TonemapperIsDefined(ConfigBitmask, TonemapperColorGrading) != 0;
 }
@@ -1211,7 +1211,7 @@ static void SetShaderTempl(const FRenderingCompositePassContext& Context)
 
 void FRCPassPostProcessTonemap::Process(FRenderingCompositePassContext& Context)
 {
-	SCOPED_DRAW_EVENTF(PostProcessTonemap, DEC_SCENE_ITEMS, TEXT("Tonemapper#%d%s"), ConfigIndex, bDoGammaOnly ? TEXT(" GammaOnly") : TEXT(""));
+	SCOPED_DRAW_EVENTF(PostProcessTonemap, DEC_SCENE_ITEMS, TEXT("Tonemapper#%d%s"), ConfigIndexPC, bDoGammaOnly ? TEXT(" GammaOnly") : TEXT(""));
 
 	const FPooledRenderTargetDesc* InputDesc = GetInputDesc(ePId_Input0);
 
@@ -1247,7 +1247,7 @@ void FRCPassPostProcessTonemap::Process(FRenderingCompositePassContext& Context)
 	Context.RHICmdList.SetRasterizerState(TStaticRasterizerState<>::GetRHI());
 	Context.RHICmdList.SetDepthStencilState(TStaticDepthStencilState<false, CF_Always>::GetRHI());
 
-	switch(ConfigIndex)
+	switch(ConfigIndexPC)
 	{
 		case 0:	SetShaderTempl<0>(Context); break;
 		case 1:	SetShaderTempl<1>(Context);	break;
@@ -1578,12 +1578,12 @@ FRCPassPostProcessTonemapES2::FRCPassPostProcessTonemapES2(const FViewInfo& View
 	: bUsedFramebufferFetch(bInUsedFramebufferFetch)
 {
 	uint32 ConfigBitmask = TonemapperGenerateBitmaskMobile(&View, false);
-	ConfigIndex = TonemapperFindLeastExpensive(TonemapperConfBitmaskMobile, sizeof(TonemapperConfBitmaskMobile)/4, TonemapperCostTab, ConfigBitmask);
+	ConfigIndexMobile = TonemapperFindLeastExpensive(TonemapperConfBitmaskMobile, sizeof(TonemapperConfBitmaskMobile)/4, TonemapperCostTab, ConfigBitmask);
 }
 
 void FRCPassPostProcessTonemapES2::Process(FRenderingCompositePassContext& Context)
 {
-	SCOPED_DRAW_EVENTF(PostProcessTonemap, DEC_SCENE_ITEMS, TEXT("Tonemapper#%d%s"), ConfigIndex, bUsedFramebufferFetch ? TEXT(" FramebufferFetch=0") : TEXT("FramebufferFetch=1"));
+	SCOPED_DRAW_EVENTF(PostProcessTonemap, DEC_SCENE_ITEMS, TEXT("Tonemapper#%d%s"), ConfigIndexMobile, bUsedFramebufferFetch ? TEXT(" FramebufferFetch=0") : TEXT("FramebufferFetch=1"));
 
 	const FPooledRenderTargetDesc* InputDesc = GetInputDesc(ePId_Input0);
 
@@ -1615,7 +1615,7 @@ void FRCPassPostProcessTonemapES2::Process(FRenderingCompositePassContext& Conte
 	Context.RHICmdList.SetRasterizerState(TStaticRasterizerState<>::GetRHI());
 	Context.RHICmdList.SetDepthStencilState(TStaticDepthStencilState<false, CF_Always>::GetRHI());
 
-	switch(ConfigIndex)
+	switch(ConfigIndexMobile)
 	{
 		case 0:	SetShaderTemplES2<0>(Context, bUsedFramebufferFetch); break;
 		case 1:	SetShaderTemplES2<1>(Context, bUsedFramebufferFetch); break;
diff --git a/Engine/Source/Runtime/Renderer/Private/PostProcess/PostProcessTonemap.h b/Engine/Source/Runtime/Renderer/Private/PostProcess/PostProcessTonemap.h
--- a/Engine/Source/Runtime/Renderer/Private/PostProcess/PostProcessTonemap.h
+++ b/Engine/Source/Runtime/Renderer/Private/PostProcess/PostProcessTonemap.h
@@ -54,7 +54,7 @@ public:
 private:
 	bool bDoGammaOnly;
 	// set in constructor
-	uint32 ConfigIndex;
+	uint32 ConfigIndexPC;
 
 	void SetShader(const FRenderingCompositePassContext& Context);
 };
@@ -77,7 +77,7 @@ public:
 private:
 	bool bUsedFramebufferFetch;
 	// set in constructor
-	uint32 ConfigIndex;
+	uint32 ConfigIndexMobile;
 
 	void SetShader(const FRenderingCompositePassContext& Context);
 };
