document

diff --git a/Engine/Source/Runtime/Renderer/Private/PostProcess/PostProcessSubsurface.h b/Engine/Source/Runtime/Renderer/Private/PostProcess/PostProcessSubsurface.h
--- a/Engine/Source/Runtime/Renderer/Private/PostProcess/PostProcessSubsurface.h
+++ b/Engine/Source/Runtime/Renderer/Private/PostProcess/PostProcessSubsurface.h
@@ -35,7 +35,7 @@ class FRCPassPostProcessSubsurface : public TRenderingCompositePassBase<2, 1>
 {
 public:
 	// constructor
-	// @param Pass 0/1
+	// @param Pass 0:horizontal/1:vertical
 	// @param InRadius in pixels in the full res image
 	FRCPassPostProcessSubsurface(uint32 Pass, float InRadius);
 
@@ -48,6 +48,6 @@ public:
 private:
 	// in pixels in the full res image
 	float Radius;
-	// 0/1
+	// 0:horizontal/1:vertical
 	uint32 Pass;
 };
