Optimized ScreenSpaceSubsurface
* Removed extra pass-through pass
* Added scalability option r.SSS.SampleSet 0/1/2, later this can be used per pixel
* added to BaseScalability
* Renamed r.SSSSS to r.SSS.Scale

diff --git a/Engine/Config/BaseScalability.ini b/Engine/Config/BaseScalability.ini
--- a/Engine/Config/BaseScalability.ini
+++ b/Engine/Config/BaseScalability.ini
@@ -192,6 +192,8 @@ r.DetailMode=0
 r.TranslucencyVolumeBlur=0
 UI.BlurRadius=0
 r.MaterialQualityLevel=0
+r.SSS.Scale = 0
+r.SSS.SampleSet=0
 
 [EffectsQuality@1]
 r.TranslucencyLightingVolumeDim=32
@@ -202,6 +204,8 @@ r.DetailMode=1
 r.TranslucencyVolumeBlur=0
 UI.BlurRadius=1.2
 r.MaterialQualityLevel=1
+r.SSS.Scale = 0.75
+r.SSS.SampleSet=0
 
 [EffectsQuality@2]
 r.TranslucencyLightingVolumeDim=48
@@ -212,6 +216,8 @@ r.DetailMode=1
 r.TranslucencyVolumeBlur=1
 UI.BlurRadius=1.2
 r.MaterialQualityLevel=1
+r.SSS.Scale = 1
+r.SSS.SampleSet=1
 
 [EffectsQuality@3]
 r.TranslucencyLightingVolumeDim=64
@@ -221,4 +227,6 @@ r.SceneColorFormat=4
 r.DetailMode=2
 r.TranslucencyVolumeBlur=1
 UI.BlurRadius=1.2
-r.MaterialQualityLevel=1
\ No newline at end of file
+r.MaterialQualityLevel=1
+r.SSS.Scale = 1
+r.SSS.SampleSet=2
\ No newline at end of file
diff --git a/Engine/Shaders/PostProcessSubsurface.usf b/Engine/Shaders/PostProcessSubsurface.usf
--- a/Engine/Shaders/PostProcessSubsurface.usf
+++ b/Engine/Shaders/PostProcessSubsurface.usf
@@ -54,6 +54,17 @@ float GetSubsurfaceStrength(float2 UV)
 	#define SSSS_SPECULAR_CORRECTION 0
 #endif
 
+#if SSS_SAMPLESET == 0
+	#define	SSSS_N_KERNELWEIGHTCOUNT 6
+	#define	SSSS_N_KERNELWEIGHTOFFSET (13 + 9)
+#elif SSS_SAMPLESET == 1
+	#define	SSSS_N_KERNELWEIGHTCOUNT 9
+	#define	SSSS_N_KERNELWEIGHTOFFSET 13
+#else // SSS_SAMPLESET == 2
+	#define	SSSS_N_KERNELWEIGHTCOUNT 13
+	#define	SSSS_N_KERNELWEIGHTOFFSET 0
+#endif
+
 // 0: faster
 // 1: no color bleeding in z direction
 #define SSSS_FOLLOW_SURFACE 1
diff --git a/Engine/Shaders/SeparableSSS.usf b/Engine/Shaders/SeparableSSS.usf
--- a/Engine/Shaders/SeparableSSS.usf
+++ b/Engine/Shaders/SeparableSSS.usf
@@ -482,7 +482,7 @@ float4 SSSSBlurPS(
 	for (int i = 1; i < SSSS_N_KERNELWEIGHTCOUNT; i++) 
 	{
 		// Kernel.a = 0..3
-		half4 Kernel = GetKernel(i, SubsurfaceProfileInt);
+		half4 Kernel = GetKernel(SSSS_N_KERNELWEIGHTOFFSET + i, SubsurfaceProfileInt);
 
 		float3 LocalAccum = 0;
 
diff --git a/Engine/Source/Runtime/Engine/Private/Rendering/SeparableSSS.cpp b/Engine/Source/Runtime/Engine/Private/Rendering/SeparableSSS.cpp
--- a/Engine/Source/Runtime/Engine/Private/Rendering/SeparableSSS.cpp
+++ b/Engine/Source/Runtime/Engine/Private/Rendering/SeparableSSS.cpp
@@ -83,21 +83,19 @@ inline FVector SeparableSSS_Profile(float r, FLinearColor FalloffColor)
 		0.078f * SeparableSSS_Gaussian(7.41f, r, FalloffColor);
 }
 
-void ComputeMirroredSSSKernel(TArray<FLinearColor>& Out, FLinearColor SubsurfaceColor, FLinearColor FalloffColor)
+void ComputeMirroredSSSKernel(FLinearColor* TargetBuffer, uint32 TargetBufferSize, FLinearColor SubsurfaceColor, FLinearColor FalloffColor)
 {
-	// needs to be preallocated with the expected size e.g. 5 for a 9 sample kernel with one center and 4 left, 4 right and one center sample
-	check(Out.Num());
+	check(TargetBuffer);
+	check(TargetBufferSize > 0);
 
-	uint32 nNonMirroredSamples = Out.Num();
+	uint32 nNonMirroredSamples = TargetBufferSize;
 	int32 nTotalSamples = nNonMirroredSamples * 2 - 1;
 
 	// we could generate Out directly but the original code form SeparableSSS wasn't done like that so we convert it later
 	// .a is in mm
-	TArray<FLinearColor> kernel;
+	check(nTotalSamples < 64);
+	FLinearColor kernel[64];
 	{
-
-		kernel.AddZeroed(nTotalSamples);
-
 		const float Range = nTotalSamples > 20 ? 3.0f : 2.0f;
 		// tweak constant
 		const float Exponent = 2.0f;
@@ -170,12 +168,12 @@ void ComputeMirroredSSSKernel(TArray<FLinearColor>& Out, FLinearColor Subsurface
 		check(kernel[0].A == 0.0f);
 
 		// center sample
-		Out[0] = kernel[0];
+		TargetBuffer[0] = kernel[0];
 
 		// all positive samples
 		for (uint32 i = 0; i < nNonMirroredSamples - 1; i++)
 		{
-			Out[i + 1] = kernel[nNonMirroredSamples + i];
+			TargetBuffer[i + 1] = kernel[nNonMirroredSamples + i];
 		}
 	}
 }
\ No newline at end of file
diff --git a/Engine/Source/Runtime/Engine/Private/Rendering/SeparableSSS.h b/Engine/Source/Runtime/Engine/Private/Rendering/SeparableSSS.h
--- a/Engine/Source/Runtime/Engine/Private/Rendering/SeparableSSS.h
+++ b/Engine/Source/Runtime/Engine/Private/Rendering/SeparableSSS.h
@@ -6,9 +6,10 @@ SeparableSSS.h: Computing the kernel for the Separable Screen Space Subsurface S
 
 #pragma once
 
-// @param Out, needs to be preallocated with the expected size (RGB is the sample weight, A is the offset), [0] is the center samples, following elements need to be mirrored with A, -A
+// @param TargetBuffer, needs to be preallocated with TargetBufferSize (RGB is the sample weight, A is the offset), [0] is the center samples, following elements need to be mirrored with A, -A
+// @param TargetBufferSize >0
 // @parma SubsurfaceColor see SubsurfaceProfile.h
 // @parma FalloffColor see SubsurfaceProfile.h
-void ComputeMirroredSSSKernel(TArray<FLinearColor>& Out, FLinearColor SubsurfaceColor, FLinearColor FalloffColor);
+void ComputeMirroredSSSKernel(FLinearColor* TargetBuffer, uint32 TargetBufferSize, FLinearColor SubsurfaceColor, FLinearColor FalloffColor);
 
 
diff --git a/Engine/Source/Runtime/Engine/Private/Rendering/SubsurfaceProfile.cpp b/Engine/Source/Runtime/Engine/Private/Rendering/SubsurfaceProfile.cpp
--- a/Engine/Source/Runtime/Engine/Private/Rendering/SubsurfaceProfile.cpp
+++ b/Engine/Source/Runtime/Engine/Private/Rendering/SubsurfaceProfile.cpp
@@ -129,8 +129,10 @@ void FSubsurfaceProfileTexture::CreateTexture(FRHICommandListImmediate& RHICmdLi
 	// true:16bit (currently required to have very small and very large kernel sizes), false: 8bit
 	const bool b16Bit = true;
 
+	const uint32 Width = 32;
+
 	// at minimum 64 lines (less reallocations)
-	FPooledRenderTargetDesc Desc(FPooledRenderTargetDesc::Create2DDesc(FIntPoint(16, FMath::Max(Height, (uint32)64)), PF_B8G8R8A8, TexCreate_None, TexCreate_None, false));
+	FPooledRenderTargetDesc Desc(FPooledRenderTargetDesc::Create2DDesc(FIntPoint(Width, FMath::Max(Height, (uint32)64)), PF_B8G8R8A8, TexCreate_FastVRAM, TexCreate_None, false));
 
 	if (b16Bit)
 	{
@@ -143,11 +145,15 @@ void FSubsurfaceProfileTexture::CreateTexture(FRHICommandListImmediate& RHICmdLi
 	uint32 DestStride;
 	uint8* DestBuffer = (uint8*)RHICmdList.LockTexture2D((FTexture2DRHIRef&)GSSProfiles->GetRenderTargetItem().ShaderResourceTexture, 0, RLM_WriteOnly, DestStride, false);
 
-	TArray<FLinearColor> kernel;
-	kernel.AddZeroed(13);	// highest quality
-	// todo: for optimizations
-//	kernel.AddZeroed(9);
-//	kernel.AddZeroed(6);	// lowest quality
+	// we precompute 3 kernels of different size and store in one line
+	const uint32 KernelSize0 = 13;
+	const uint32 KernelSize1 = 9; 
+	const uint32 KernelSize2 = 6;
+
+	const uint32 KernelTotalSize = KernelSize0 + KernelSize1 + KernelSize2;
+	check(KernelTotalSize < Width);
+
+	FLinearColor kernel[Width];
 
 	for (uint32 y = 0; y < Height; ++y)
 	{
@@ -161,12 +167,14 @@ void FSubsurfaceProfileTexture::CreateTexture(FRHICommandListImmediate& RHICmdLi
 		Data.SubsurfaceColor = Data.SubsurfaceColor.GetClamped(Bias);
 		Data.FalloffColor = Data.FalloffColor.GetClamped(Bias);
 
-		ComputeMirroredSSSKernel(kernel, Data.SubsurfaceColor, Data.FalloffColor);
+		ComputeMirroredSSSKernel(&kernel[0], KernelSize0, Data.SubsurfaceColor, Data.FalloffColor);
+		ComputeMirroredSSSKernel(&kernel[KernelSize0], KernelSize1, Data.SubsurfaceColor, Data.FalloffColor);
+		ComputeMirroredSSSKernel(&kernel[KernelSize0 + KernelSize1], KernelSize2, Data.SubsurfaceColor, Data.FalloffColor);
 
 		const float TableMaxRGB = 1.0f;
 		const float TableMaxA = 3.0f;
 
-		for (int32 Pos = 0; Pos < kernel.Num(); ++Pos)
+		for (int32 Pos = 0; Pos < KernelTotalSize; ++Pos)
 		{
 			FVector4 C = kernel[Pos] * FLinearColor(1.0f / TableMaxRGB, 1.0f / TableMaxRGB, 1.0f / TableMaxRGB, 1.0f / TableMaxA);
 
diff --git a/Engine/Source/Runtime/Renderer/Private/CompositionLighting/CompositionLighting.cpp b/Engine/Source/Runtime/Renderer/Private/CompositionLighting/CompositionLighting.cpp
--- a/Engine/Source/Runtime/Renderer/Private/CompositionLighting/CompositionLighting.cpp
+++ b/Engine/Source/Runtime/Renderer/Private/CompositionLighting/CompositionLighting.cpp
@@ -27,8 +27,8 @@ FCompositionLighting GCompositionLighting;
 
 // -------------------------------------------------------
 
-static TAutoConsoleVariable<float> CVarSSSSS(
-	TEXT("r.SSSSS"),
+static TAutoConsoleVariable<float> CVarSSSScale(
+	TEXT("r.SSS.Scale"),
 	1.0f,
 	TEXT("Experimental screen space subsurface scattering pass")
 	TEXT("(use shadingmodel SubsurfaceProfile, get near to the object as the default)\n")
@@ -351,7 +351,7 @@ void FCompositionLighting::ProcessLighting(FRHICommandListImmediate& RHICmdList,
 
 		// Screen Space Subsurface Scattering
 		{
-			float Radius = CVarSSSSS.GetValueOnRenderThread();
+			float Radius = CVarSSSScale.GetValueOnRenderThread();
 
 			bool bSimpleDynamicLighting = IsSimpleDynamicLightingEnabled();
 
@@ -367,11 +367,6 @@ void FCompositionLighting::ProcessLighting(FRHICommandListImmediate& RHICmdList,
 				Pass1->SetInput(ePId_Input0, Context.FinalOutput);
 				Pass1->SetInput(ePId_Input1, Pass0);
 				Context.FinalOutput = FRenderingCompositeOutputRef(Pass1);
-
-				// waste of performance?
-				FRenderingCompositePass* NullPass = Context.Graph.RegisterPass(new(FMemStack::Get()) FRCPassPostProcessPassThrough(0));
-				NullPass->SetInput(ePId_Input0, Context.FinalOutput);
-				Context.FinalOutput = FRenderingCompositeOutputRef(NullPass);
 			}
 		}
 
@@ -379,10 +374,7 @@ void FCompositionLighting::ProcessLighting(FRHICommandListImmediate& RHICmdList,
 
 		SCOPED_DRAW_EVENT(CompositionLighting, DEC_SCENE_ITEMS);
 
-		TRefCountPtr<IPooledRenderTarget>& SceneColor = GSceneRenderTargets.GetSceneColor();
-
-		Context.FinalOutput.GetOutput()->RenderTargetDesc = SceneColor->GetDesc();
-		Context.FinalOutput.GetOutput()->PooledRenderTarget = SceneColor;
+		// we don't replace the final element with the scenecolor because this is what those passes should do by themself
 
 		// you can add multiple dependencies
 		CompositeContext.Root->AddDependency(Context.FinalOutput);
diff --git a/Engine/Source/Runtime/Renderer/Private/PostProcess/PostProcessSubsurface.cpp b/Engine/Source/Runtime/Renderer/Private/PostProcess/PostProcessSubsurface.cpp
--- a/Engine/Source/Runtime/Renderer/Private/PostProcess/PostProcessSubsurface.cpp
+++ b/Engine/Source/Runtime/Renderer/Private/PostProcess/PostProcessSubsurface.cpp
@@ -13,6 +13,16 @@
 
 ENGINE_API const IPooledRenderTarget* GetSubsufaceProfileTexture_RT(FRHICommandListImmediate& RHICmdList);
 
+
+static TAutoConsoleVariable<int32> CVarSSSSampleSet(
+	TEXT("r.SSS.SampleSet"),
+	2,
+	TEXT("0: lowest quality\n")
+	TEXT("1: medium quality\n")
+	TEXT("2: high quality (default)"),
+	ECVF_RenderThreadSafe  | ECVF_Scalability);
+
+
 /**
  * Encapsulates the post processing ambient occlusion pixel shader.
  * @param SetupMode 0:without specular correction, 1: with specular correction, 2:visualize
@@ -274,10 +284,12 @@ FPooledRenderTargetDesc FRCPassPostProcessSubsurfaceSetup::ComputeOutputDesc(EPa
 
 
 /** Encapsulates the post processing down sample pixel shader. */
-template <uint32 Method>
-class FPostProcessSubsurfacePS : public FGlobalShader
+// @param Method 0: horizontal, 1:vertical, 2:vertical and reconstruct specular
+// @param SampleSet 0:low, 1:med, 2:high
+template <uint32 Method, uint32 SampleSet>
+class TPostProcessSubsurfacePS : public FGlobalShader
 {
-	DECLARE_SHADER_TYPE(FPostProcessSubsurfacePS, Global);
+	DECLARE_SHADER_TYPE(TPostProcessSubsurfacePS, Global);
 
 	static bool ShouldCache(EShaderPlatform Platform)
 	{
@@ -288,10 +300,11 @@ class FPostProcessSubsurfacePS : public FGlobalShader
 	{
 		FGlobalShader::ModifyCompilationEnvironment(Platform,OutEnvironment);
 		OutEnvironment.SetDefine(TEXT("SSS_METHOD"), Method);
+		OutEnvironment.SetDefine(TEXT("SSS_SAMPLESET"), SampleSet);
 	}
 
 	/** Default constructor. */
-	FPostProcessSubsurfacePS() {}
+	TPostProcessSubsurfacePS() {}
 
 public:
 	FPostProcessPassParameters PostprocessParameter;
@@ -301,7 +314,7 @@ public:
 	FShaderResourceParameter SSProfilesTextureSampler;
 
 	/** Initialization constructor. */
-	FPostProcessSubsurfacePS(const ShaderMetaType::CompiledShaderInitializerType& Initializer)
+	TPostProcessSubsurfacePS(const ShaderMetaType::CompiledShaderInitializerType& Initializer)
 		: FGlobalShader(Initializer)
 	{
 		PostprocessParameter.Bind(Initializer.ParameterMap);
@@ -360,12 +373,13 @@ public:
 	}
 };
 
-// #define avoids a lot of code duplication
-#define VARIATION1(A) typedef FPostProcessSubsurfacePS<A> FPostProcessSubsurfacePS##A; \
-	IMPLEMENT_SHADER_TYPE2(FPostProcessSubsurfacePS##A, SF_Pixel);
-
-VARIATION1(0)			VARIATION1(1)			VARIATION1(2)
+	// #define avoids a lot of code duplication
+#define VARIATION1(A)		VARIATION2(A,0)			VARIATION2(A,1)			VARIATION2(A,2)
+#define VARIATION2(A, B) typedef TPostProcessSubsurfacePS<A, B> TPostProcessSubsurfacePS##A##B; \
+	IMPLEMENT_SHADER_TYPE2(TPostProcessSubsurfacePS##A##B, SF_Pixel);
+	VARIATION1(0) VARIATION1(1) VARIATION1(2)
 #undef VARIATION1
+#undef VARIATION2
 
 
 FRCPassPostProcessSubsurface::FRCPassPostProcessSubsurface(uint32 InPass, float InRadius)
@@ -374,12 +388,11 @@ FRCPassPostProcessSubsurface::FRCPassPostProcessSubsurface(uint32 InPass, float
 {
 }
 
-
-template <uint32 Method>
+template <uint32 Method, uint32 SampleSet>
 void SetSubsurfaceShader(const FRenderingCompositePassContext& Context, float InRadius)
 {
 	TShaderMapRef<FPostProcessVS> VertexShader(Context.GetShaderMap());
-	TShaderMapRef<FPostProcessSubsurfacePS<Method> > PixelShader(Context.GetShaderMap());
+	TShaderMapRef<TPostProcessSubsurfacePS<Method,SampleSet> > PixelShader(Context.GetShaderMap());
 
 	static FGlobalBoundShaderState BoundShaderState;
 
@@ -389,6 +402,20 @@ void SetSubsurfaceShader(const FRenderingCompositePassContext& Context, float In
 	VertexShader->SetParameters(Context);
 }
 
+// 0:horizontal, 1:vertical with spec correction, 2:vertical without spec correction
+template <uint32 Method>
+void SetSubsurfaceShaderSampleSet(const FRenderingCompositePassContext& Context, float InRadius, uint32 SampleSet)
+{
+	switch(SampleSet)
+	{
+		case 0: SetSubsurfaceShader<Method, 0>(Context, InRadius); break;
+		case 1: SetSubsurfaceShader<Method, 1>(Context, InRadius); break;
+		case 2: SetSubsurfaceShader<Method, 2>(Context, InRadius); break;
+
+	default:
+		check(0);
+	}
+}
 
 void FRCPassPostProcessSubsurface::Process(FRenderingCompositePassContext& Context)
 {
@@ -417,10 +444,26 @@ void FRCPassPostProcessSubsurface::Process(FRenderingCompositePassContext& Conte
 	FIntRect SrcRect = View.ViewRect / ScaleFactor;
 	FIntRect DestRect = SrcRect;
 
-	const FSceneRenderTargetItem& DestRenderTarget = PassOutputs[0].RequestSurface(Context);
+	TRefCountPtr<IPooledRenderTarget> NewSceneColor;
+
+	const FSceneRenderTargetItem* DestRenderTarget;
+	{
+		if(Pass == 0)
+		{
+			DestRenderTarget = &PassOutputs[0].RequestSurface(Context);
+		}
+		else
+		{
+			// in that case we directly render to the HDR scene color
+			GRenderTargetPool.FindFreeElement(GSceneRenderTargets.GetSceneColor()->GetDesc(), NewSceneColor, TEXT("SceneColor"));
+			DestRenderTarget = &NewSceneColor->GetRenderTargetItem();
+		}
+
+		check(DestRenderTarget);
+	}
 
 	// Set the view family's render target/viewport.
-	SetRenderTarget(Context.RHICmdList, DestRenderTarget.TargetableTexture, FTextureRHIRef());
+	SetRenderTarget(Context.RHICmdList, DestRenderTarget->TargetableTexture, FTextureRHIRef());
 
 	// is optimized away if possible (RT size=view size, )
 	Context.RHICmdList.Clear(true, FLinearColor(0, 0, 0, 0), false, 1.0f, false, 0, DestRect);
@@ -435,9 +478,11 @@ void FRCPassPostProcessSubsurface::Process(FRenderingCompositePassContext& Conte
 
 	SCOPED_DRAW_EVENT(SubsurfacePass, DEC_SCENE_ITEMS);
 
+	uint32 SampleSet = FMath::Clamp(CVarSSSSampleSet.GetValueOnRenderThread(), 0, 2);
+
 	if (Pass == 0)
 	{
-		SetSubsurfaceShader<0>(Context, Radius);
+		SetSubsurfaceShaderSampleSet<0>(Context, Radius, SampleSet);
 	}
 	else
 	{
@@ -446,11 +491,11 @@ void FRCPassPostProcessSubsurface::Process(FRenderingCompositePassContext& Conte
 		if(DoSpecularCorrection())
 		{
 			// reconstruct specular and add it in final pass
-			SetSubsurfaceShader<2>(Context, Radius);
+			SetSubsurfaceShaderSampleSet<2>(Context, Radius, SampleSet);
 		}
 		else
 		{
-			SetSubsurfaceShader<1>(Context, Radius);
+			SetSubsurfaceShaderSampleSet<1>(Context, Radius, SampleSet);
 		}
 	}
 
@@ -465,7 +510,12 @@ void FRCPassPostProcessSubsurface::Process(FRenderingCompositePassContext& Conte
 		*VertexShader,
 		EDRF_UseTriangleOptimization);
 
-	Context.RHICmdList.CopyToResolveTarget(DestRenderTarget.TargetableTexture, DestRenderTarget.ShaderResourceTexture, false, FResolveParams());
+	Context.RHICmdList.CopyToResolveTarget(DestRenderTarget->TargetableTexture, DestRenderTarget->ShaderResourceTexture, false, FResolveParams());
+
+	if(Pass != 0)
+	{
+		GSceneRenderTargets.SetSceneColor(NewSceneColor);
+	}
 }
 
 
@@ -473,6 +523,12 @@ FPooledRenderTargetDesc FRCPassPostProcessSubsurface::ComputeOutputDesc(EPassOut
 {
 	FPooledRenderTargetDesc Ret;
 
+	if(Pass != 0)
+	{
+		// in that case we directly render to the HDR scene color
+		return Ret;
+	}
+
 	Ret = PassInputs[1].GetOutput()->RenderTargetDesc;
 
 	Ret.Reset();
