Fix for shader ES2 compile error
 - Don't compile motion blur post process shader setup VS below SM4

diff --git a/Engine/Source/Runtime/Renderer/Private/PostProcess/PostProcessMotionBlur.cpp b/Engine/Source/Runtime/Renderer/Private/PostProcess/PostProcessMotionBlur.cpp
--- a/Engine/Source/Runtime/Renderer/Private/PostProcess/PostProcessMotionBlur.cpp
+++ b/Engine/Source/Runtime/Renderer/Private/PostProcess/PostProcessMotionBlur.cpp
@@ -29,7 +29,7 @@ class FPostProcessMotionBlurSetupVS : public FGlobalShader
 
 	static bool ShouldCache(EShaderPlatform Platform)
 	{
-		return true;
+		return IsFeatureLevelSupported(Platform, ERHIFeatureLevel::SM4);
 	}
 
 	/** Default constructor. */
