fixed TTP 335245 UE4: RENDERING: LPV shader error

diff --git a/Engine/Source/Runtime/Renderer/Private/ShadowRendering.cpp b/Engine/Source/Runtime/Renderer/Private/ShadowRendering.cpp
--- a/Engine/Source/Runtime/Renderer/Private/ShadowRendering.cpp
+++ b/Engine/Source/Runtime/Renderer/Private/ShadowRendering.cpp
@@ -1060,7 +1060,7 @@ bool FShadowDepthDrawingPolicyFactory::DrawDynamicMesh(
 		const EMaterialShadingModel ShadingModel = Material->GetShadingModel();
 
 		const bool bOnePassPointLightShadow = Context.ShadowInfo->bOnePassPointLightShadow;
-		const bool bReflectiveShadowmap = Context.ShadowInfo->bReflectiveShadowmap && !bOnePassPointLightShadow;
+		const bool bReflectiveShadowmap = Context.ShadowInfo->bReflectiveShadowmap && !bOnePassPointLightShadow && Material->ShouldInjectEmissiveIntoLPV();
 
 		if ( (!IsTranslucentBlendMode(BlendMode) || bReflectiveShadowmap ) && ShadingModel != MSM_Unlit )
 		{
