* fixed TTP 341870 COMPAT: CRITICAL: Regression: LIVE: CRASH: Running Engine Benchmark freezes and crashes with certain video cards
* fixed SynthBenchmark for slow GPUs
* try with increasing number of work amount
* additionally tack total work time
* made WorkScale a float
* removed BenchMarkGraph (was for internal testing))
* updated version to 0.95

diff --git a/Engine/Source/Developer/SynthBenchmark/Private/SynthBenchmarkPrivatePCH.cpp b/Engine/Source/Developer/SynthBenchmark/Private/SynthBenchmarkPrivatePCH.cpp
--- a/Engine/Source/Developer/SynthBenchmark/Private/SynthBenchmarkPrivatePCH.cpp
+++ b/Engine/Source/Developer/SynthBenchmark/Private/SynthBenchmarkPrivatePCH.cpp
@@ -21,7 +21,7 @@ class FSynthBenchmark : public ISynthBenchmark
 	virtual void ShutdownModule() override;
 	
 	/** ISynthBenchmark implementation */
-	virtual void Run(FSynthBenchmarkResults& InOut, bool bGPUBenchmark, uint32 WorkScale, bool bDebugOut) const override;
+	virtual void Run(FSynthBenchmarkResults& InOut, bool bGPUBenchmark, float WorkScale) const override;
 };
 
 IMPLEMENT_MODULE( FSynthBenchmark, SynthBenchmark )
@@ -39,11 +39,13 @@ void FSynthBenchmark::ShutdownModule()
 
 // @param RunCount should be around 10 but can be adjusted for precision 
 // @param Function should run for about 3 ms
-// @return performance index: 100 is a developer machine: Intel Xeon E5-2660 2.2GHz
-static float RunBenchmark(uint32 RunCount, float (*Function)())
+static FTimeSample RunBenchmark(float WorkScale, float (*Function)())
 {
 	float Sum = 0;
 
+	// this test doesn't support fractional WorkScale
+	uint32 RunCount = FMath::Max((int32)1, (int32)WorkScale);
+
 	for(uint32 i = 0; i < RunCount; ++i)
 	{
 		FPlatformMisc::MemoryBarrier();
@@ -59,10 +61,10 @@ static float RunBenchmark(uint32 RunCount, float (*Function)())
 		FPlatformMisc::MemoryBarrier();
 	}
 	
-	return Sum / RunCount;
+	return FTimeSample(Sum, Sum / RunCount);
 }
 
-void FSynthBenchmark::Run(FSynthBenchmarkResults& InOut, bool bGPUBenchmark, uint32 WorkScale, bool bDebugOut) const
+void FSynthBenchmark::Run(FSynthBenchmarkResults& InOut, bool bGPUBenchmark, float WorkScale) const
 {
 	check(WorkScale > 0);
 
@@ -70,12 +72,12 @@ void FSynthBenchmark::Run(FSynthBenchmarkResults& InOut, bool bGPUBenchmark, uin
 	{
 		// run a very quick GPU benchmark (less confidence but at least we get some numbers)
 		// it costs little time and we get some stats
-		WorkScale = 1;
+		WorkScale = 1.0f;
 	}
 
 	const double StartTime = FPlatformTime::Seconds();
 
-	UE_LOG(LogSynthBenchmark, Display, TEXT("FSynthBenchmark (V0.92):"));
+	UE_LOG(LogSynthBenchmark, Display, TEXT("FSynthBenchmark (V0.95):  requested WorkScale=%.2f"), WorkScale);
 	UE_LOG(LogSynthBenchmark, Display, TEXT("==============="));
 	
 #if UE_BUILD_DEBUG
@@ -96,7 +98,7 @@ void FSynthBenchmark::Run(FSynthBenchmarkResults& InOut, bool bGPUBenchmark, uin
 
 	for(uint32 i = 0; i < ARRAY_COUNT(InOut.CPUStats); ++i)
 	{
-		UE_LOG(LogSynthBenchmark, Display, TEXT("         ... %f %s '%s'"), InOut.CPUStats[i].GetMeasuredTime(), InOut.CPUStats[i].GetValueType(), InOut.CPUStats[i].GetDesc());
+		UE_LOG(LogSynthBenchmark, Display, TEXT("         ... %f %s '%s'"), InOut.CPUStats[i].GetNormalizedTime(), InOut.CPUStats[i].GetValueType(), InOut.CPUStats[i].GetDesc());
 	}
 
 	UE_LOG(LogSynthBenchmark, Display, TEXT(""));
@@ -143,7 +145,53 @@ void FSynthBenchmark::Run(FSynthBenchmarkResults& InOut, bool bGPUBenchmark, uin
 	{
 		IRendererModule& RendererModule = FModuleManager::LoadModuleChecked<IRendererModule>(TEXT("Renderer"));
 
-		RendererModule.GPUBenchmark(InOut, WorkScale, bDebugOut);
+		// First we run a quick test. If that shows very bad performance we don't need another test
+		// The hardware is slow, we don't need a long test and risk driver TDR (driver recovery).
+		// We have seen this problem on very low end GPUs.
+		{
+			const float fFirstWorkScale = 0.01f;
+			const float fSecondWorkScale = 0.1f;
+
+			float GPUTime = 0.0f;
+
+			RendererModule.GPUBenchmark(InOut, fFirstWorkScale);
+			GPUTime = InOut.ComputeTotalGPUTime();
+			UE_LOG(LogSynthBenchmark, Display, TEXT("  GPU first test: %.2fs"), GPUTime);
+
+			for(uint32 MethodId = 0; MethodId < sizeof(InOut.GPUStats) / sizeof(InOut.GPUStats[0]); ++MethodId)
+			{
+				UE_LOG(LogSynthBenchmark, Display, TEXT("         ... %.3f GigaPix/s, Confidence=%.0f%% '%s' (likely to be very inaccurate)"),
+					1.0f / InOut.GPUStats[MethodId].GetNormalizedTime(), InOut.GPUStats[MethodId].GetConfidence(), InOut.GPUStats[MethodId].GetDesc());
+			}
+
+			if(GPUTime < 0.1f)
+			{
+				RendererModule.GPUBenchmark(InOut, fSecondWorkScale);
+				GPUTime = InOut.ComputeTotalGPUTime();
+				UE_LOG(LogSynthBenchmark, Display, TEXT("  GPU second test: %.2fs"), GPUTime);
+
+				// for testing
+				for(uint32 MethodId = 0; MethodId < sizeof(InOut.GPUStats) / sizeof(InOut.GPUStats[0]); ++MethodId)
+				{
+					UE_LOG(LogSynthBenchmark, Display, TEXT("         ... %.3f GigaPix/s, Confidence=%.0f%% '%s' (likely to be inaccurate)"),
+						1.0f / InOut.GPUStats[MethodId].GetNormalizedTime(), InOut.GPUStats[MethodId].GetConfidence(), InOut.GPUStats[MethodId].GetDesc());
+				}
+
+				if(GPUTime < 0.1f)
+				{
+					RendererModule.GPUBenchmark(InOut, WorkScale);
+					GPUTime = InOut.ComputeTotalGPUTime();
+					UE_LOG(LogSynthBenchmark, Display, TEXT("  GPU third test: %.2fs"), GPUTime);
+				}
+			}
+		}
+
+		for(uint32 MethodId = 0; MethodId < sizeof(InOut.GPUStats) / sizeof(InOut.GPUStats[0]); ++MethodId)
+		{
+			UE_LOG(LogSynthBenchmark, Display, TEXT("         ... %.3f GigaPix/s, Confidence=%.0f%% '%s'"),
+				1.0f / InOut.GPUStats[MethodId].GetNormalizedTime(), InOut.GPUStats[MethodId].GetConfidence(), InOut.GPUStats[MethodId].GetDesc());
+		}
+		UE_LOG(LogSynthBenchmark, Display, TEXT(""));
 
 		UE_LOG(LogSynthBenchmark, Display, TEXT("  GPU Perf Index 0: %.1f"), InOut.GPUStats[0].ComputePerfIndex());
 		UE_LOG(LogSynthBenchmark, Display, TEXT("  GPU Perf Index 1: %.1f"), InOut.GPUStats[1].ComputePerfIndex());
diff --git a/Engine/Source/Developer/SynthBenchmark/Public/SynthBenchmark.h b/Engine/Source/Developer/SynthBenchmark/Public/SynthBenchmark.h
--- a/Engine/Source/Developer/SynthBenchmark/Public/SynthBenchmark.h
+++ b/Engine/Source/Developer/SynthBenchmark/Public/SynthBenchmark.h
@@ -15,8 +15,7 @@ class ISynthBenchmark : public IModuleInterface
 public:
 
 	// @param >0, WorkScale 10 for normal precision and runtime of less than a second
-	// @param bDebugOut has no effect in shipping
-	virtual void Run(FSynthBenchmarkResults& Out, bool bGPUBenchmark = false, uint32 WorkScale = 10, bool bDebugOut = false) const = 0;
+	virtual void Run(FSynthBenchmarkResults& Out, bool bGPUBenchmark = false, float WorkScale = 10.0f) const = 0;
 
 	/**
 	 * Singleton-like access to this module's interface.  This is just for convenience!
diff --git a/Engine/Source/Runtime/Core/Public/GenericPlatform/GenericPlatformSurvey.h b/Engine/Source/Runtime/Core/Public/GenericPlatform/GenericPlatformSurvey.h
--- a/Engine/Source/Runtime/Core/Public/GenericPlatform/GenericPlatformSurvey.h
+++ b/Engine/Source/Runtime/Core/Public/GenericPlatform/GenericPlatformSurvey.h
@@ -7,22 +7,40 @@
 
 #pragma once
 
+
+// time and amount of work that was measured
+struct FTimeSample
+{
+	FTimeSample(float InTotalTime, float InNormalizedTime)
+		: TotalTime(InTotalTime)
+		, NormalizedTime(InNormalizedTime)
+	{
+	}
+
+	// in seconds, >=0
+	float TotalTime;
+	// Time / WorkScale, WorkScale might have been adjusted (e.g. quantized), >0
+	float NormalizedTime;
+};
+
 struct FSynthBenchmarkStat
 {
 	FSynthBenchmarkStat()
 		: Desc(0)
-		, MeasuredTime(-1)
-		, IndexTime(-1)
+		, MeasuredTotalTime(-1)
+		, MeasuredNormalizedTime(-1)
+		, IndexNormalizedTime(-1)
 		, ValueType(0)
 		, Confidence(0)
 	{
 	}
 
 	// @param InDesc descriptions
-	FSynthBenchmarkStat(const TCHAR* InDesc, float InIndexTime, const TCHAR* InValueType)
+	FSynthBenchmarkStat(const TCHAR* InDesc, float InIndexNormalizedTime, const TCHAR* InValueType)
 		: Desc(InDesc)
-		, MeasuredTime(-1)
-		, IndexTime(InIndexTime)
+		, MeasuredTotalTime(-1)
+		, MeasuredNormalizedTime(-1)
+		, IndexNormalizedTime(InIndexNormalizedTime)
 		, ValueType(InValueType)
 		, Confidence(0)
 	{
@@ -30,16 +48,19 @@ struct FSynthBenchmarkStat
 //		check(InValueType);
 	}
 
+	// Computes the linear performance index (>0), around 100 with good hardware but higher numbers are possible
 	float ComputePerfIndex() const
 	{
-		return 100.0f * IndexTime / MeasuredTime;
+		return 100.0f * IndexNormalizedTime / MeasuredNormalizedTime;
 	}
 
+	// @param TimeSample seconds and normalized time (e.g. seconds / GPixels)
 	// @param InConfidence 0..100
-	void SetMeasuredTime(float InMeasuredValue, float InConfidence = 90)
+	void SetMeasuredTime(const FTimeSample& TimeSample, float InConfidence = 90)
 	{
 //		check(InConfidence >= 0.0f && InConfidence <= 100.0f);
-		MeasuredTime = InMeasuredValue;
+		MeasuredTotalTime = TimeSample.TotalTime;
+		MeasuredNormalizedTime = TimeSample.NormalizedTime;
 		Confidence = InConfidence;
 	}
 
@@ -55,12 +76,17 @@ struct FSynthBenchmarkStat
 		return ValueType;
 	}
 
-	float GetMeasuredTime() const
+	float GetNormalizedTime() const
+	{
+		return MeasuredNormalizedTime;
+	}
+
+	float GetMeasuredTotalTime() const
 	{
-		return MeasuredTime;
+		return MeasuredTotalTime;
 	}
 
-	// @param 0..100
+	// @return 0=no..100=full
 	float GetConfidence() const
 	{
 		return Confidence;
@@ -69,10 +95,12 @@ struct FSynthBenchmarkStat
 private:
 	// 0 if not valid
 	const TCHAR *Desc;
-	// -1 if not defined
-	float MeasuredTime;
+	// -1 if not defined, in seconds, useful to see if a test did run too long (some slower GPUs might timeout)
+	float MeasuredTotalTime;
+	// -1 if not defined, depends on the test (e.g. s/GPixels), WorkScale is divided out
+	float MeasuredNormalizedTime;
 	// -1 if not defined, timing value expected on a norm GPU (index value 100, here NVidia 670)
-	float IndexTime;
+	float IndexNormalizedTime;
 	// 0 if not valid
 	const TCHAR *ValueType;
 	// 0..100, 100: fully confident
@@ -88,21 +116,44 @@ struct FSynthBenchmarkResults
 	float ComputeCPUPerfIndex() const
 	{
 		// if needed we can weight them differently
-		return (
-			CPUStats[0].ComputePerfIndex() +
-			CPUStats[1].ComputePerfIndex() ) / 2;
+		float Ret = 0.0f;
+		float Count = 0.0f;
+
+		for(uint32 i = 0; i < sizeof(CPUStats) / sizeof(CPUStats[0]); ++i)
+		{
+			Ret += CPUStats[i].ComputePerfIndex();
+			++Count;
+		}
+
+		return Ret / Count;
 	}
 
 	// 100: avg good GPU, <100:slower, >100:faster
 	float ComputeGPUPerfIndex() const
 	{
 		// if needed we can weight them differently
-		return (
-			GPUStats[0].ComputePerfIndex() +
-			GPUStats[1].ComputePerfIndex() +
-			GPUStats[2].ComputePerfIndex() +
-			GPUStats[3].ComputePerfIndex() +
-			GPUStats[4].ComputePerfIndex() ) / 5;
+		float Ret = 0.0f;
+		float Count = 0.0f;
+
+		for(uint32 i = 0; i < sizeof(GPUStats) / sizeof(GPUStats[0]); ++i)
+		{
+			Ret += GPUStats[i].ComputePerfIndex();
+			++Count;
+		}
+
+		return Ret / Count;
+	}
+
+	// @return in seconds, useful to check if a benchmark takes too long (very slow hardware, don't make tests with large WorkScale)
+	float ComputeTotalGPUTime() const
+	{
+		float Ret = 0.0f;
+		for(uint32 i = 0; i < sizeof(GPUStats) / sizeof(GPUStats[0]); ++i)
+		{
+			Ret += GPUStats[i].GetMeasuredTotalTime();
+		}
+
+		return Ret;
 	}
 };
 
diff --git a/Engine/Source/Runtime/Engine/Private/GameEngine.cpp b/Engine/Source/Runtime/Engine/Private/GameEngine.cpp
--- a/Engine/Source/Runtime/Engine/Private/GameEngine.cpp
+++ b/Engine/Source/Runtime/Engine/Private/GameEngine.cpp
@@ -37,28 +37,22 @@ ENGINE_API bool GDisallowNetworkTravel = false;
 /** Benchmark results to the log */
 static void RunSynthBenchmark(const TArray<FString>& Args)
 {
-	uint32 WorkScale = 10;
-	bool bDebugOut = false;
-
-	// 2 arguments mean with debug out
-	if(Args.Num() > 1)
-	{
-		bDebugOut = true;
-	}
+	float WorkScale = 10.0f;
 
 	if ( Args.Num() > 0 )
 	{
-		WorkScale = FCString::Atoi(*Args[0]);
-		WorkScale = FMath::Clamp(WorkScale, (uint32)1, (uint32)1000);
+		WorkScale = FCString::Atof(*Args[0]);
+		WorkScale = FMath::Clamp(WorkScale, 1.0f, 1000.0f);
 	}
 
 	FSynthBenchmarkResults Result;
-	ISynthBenchmark::Get().Run(Result, true, WorkScale, bDebugOut);
+	ISynthBenchmark::Get().Run(Result, true, WorkScale);
 }
 
 static FAutoConsoleCommand GDumpDrawListStatsCmd(
 	TEXT("SynthBenchmark"),
-	TEXT("Run simple benchmark to get some metrics to find reasonable game settings automatically."),
+	TEXT("Run simple benchmark to get some metrics to find reasonable game settings automatically\n")
+	TEXT("Optional (float) parameter allows to scale with work amount to trade time or precision (default: 10)."),
 	FConsoleCommandWithArgsDelegate::CreateStatic(&RunSynthBenchmark)
 	);
 
diff --git a/Engine/Source/Runtime/Renderer/Private/GPUBenchmark.cpp b/Engine/Source/Runtime/Renderer/Private/GPUBenchmark.cpp
--- a/Engine/Source/Runtime/Renderer/Private/GPUBenchmark.cpp
+++ b/Engine/Source/Runtime/Renderer/Private/GPUBenchmark.cpp
@@ -126,7 +126,7 @@ IMPLEMENT_SHADER_TYPE(,FPostProcessBenchmarkVS,TEXT("GPUBenchmark"),TEXT("MainBe
 
 
 template <uint32 Method>
-void RunBenchmarkShader(FRHICommandList& RHICmdList, const FSceneView& View, TRefCountPtr<IPooledRenderTarget>& Src, uint32 Count)
+void RunBenchmarkShader(FRHICommandList& RHICmdList, const FSceneView& View, TRefCountPtr<IPooledRenderTarget>& Src, float WorkScale)
 {
 	auto ShaderMap = GetGlobalShaderMap(View.GetFeatureLevel());
 
@@ -141,14 +141,24 @@ void RunBenchmarkShader(FRHICommandList& RHICmdList, const FSceneView& View, TRe
 	PixelShader->SetParameters(RHICmdList, View, Src);
 	VertexShader->SetParameters(RHICmdList, View);
 
-	for(uint32 i = 0; i < Count; ++i)
+	// single pass was not fine grained enough so we reduce the pass size based on the fractional part of WorkScale
+	float TotalHeight = GBenchmarkResolution * WorkScale;
+
+	// rounds up
+	uint32 PassCount = (uint32)FMath::CeilToFloat(TotalHeight / GBenchmarkResolution);
+
+	for(uint32 i = 0; i < PassCount; ++i)
 	{
+		float Top = i * GBenchmarkResolution;
+		float Bottom = FMath::Min(Top + GBenchmarkResolution, TotalHeight);
+		float LocalHeight = Bottom - Top;
+
 		DrawRectangle(
 			RHICmdList,
 			0, 0,
-			GBenchmarkResolution, GBenchmarkResolution,
+			GBenchmarkResolution, LocalHeight,
 			0, 0,
-			GBenchmarkResolution, GBenchmarkResolution,
+			GBenchmarkResolution, LocalHeight,
 			FIntPoint(GBenchmarkResolution, GBenchmarkResolution),
 			FIntPoint(GBenchmarkResolution, GBenchmarkResolution),
 			*VertexShader,
@@ -156,17 +166,17 @@ void RunBenchmarkShader(FRHICommandList& RHICmdList, const FSceneView& View, TRe
 	}
 }
 
-void RunBenchmarkShader(FRHICommandListImmediate& RHICmdList, const FSceneView& View, uint32 MethodId, TRefCountPtr<IPooledRenderTarget>& Src, uint32 Count)
+void RunBenchmarkShader(FRHICommandListImmediate& RHICmdList, const FSceneView& View, uint32 MethodId, TRefCountPtr<IPooledRenderTarget>& Src, float WorkScale)
 {
 	SCOPED_DRAW_EVENTF(Benchmark, DEC_SCENE_ITEMS, TEXT("Benchmark Method:%d"), MethodId);
 
 	switch(MethodId)
 	{
-		case 0: RunBenchmarkShader<0>(RHICmdList, View, Src, Count); return;
-		case 1: RunBenchmarkShader<1>(RHICmdList, View, Src, Count); return;
-		case 2: RunBenchmarkShader<2>(RHICmdList, View, Src, Count); return;
-		case 3: RunBenchmarkShader<3>(RHICmdList, View, Src, Count); return;
-		case 4: RunBenchmarkShader<4>(RHICmdList, View, Src, Count); return;
+		case 0: RunBenchmarkShader<0>(RHICmdList, View, Src, WorkScale); return;
+		case 1: RunBenchmarkShader<1>(RHICmdList, View, Src, WorkScale); return;
+		case 2: RunBenchmarkShader<2>(RHICmdList, View, Src, WorkScale); return;
+		case 3: RunBenchmarkShader<3>(RHICmdList, View, Src, WorkScale); return;
+		case 4: RunBenchmarkShader<4>(RHICmdList, View, Src, WorkScale); return;
 
 		default:
 			check(0);
@@ -296,63 +306,10 @@ private:
 	TArray<float> TimingValues;
 };
 
-#if !UE_BUILD_SHIPPING
-// for debugging only
-class FBenchmarkGraph
-{
-public:
-	FBenchmarkGraph(uint32 InWidth, uint32 InHeight, const TCHAR* InFilePath)
-		: FilePath(InFilePath)
-		, Height(InHeight)
-		, Width(InWidth)
-	{
-		Data.AddZeroed(Width * Height);
-	}
-
-
-	// @param x 0..Width-1
-	// @param y 0..Height-1
-	void DrawPixel(int32 X, int32 Y, FColor Color = FColor::White)
-	{
-		if((uint32)X < Width && (uint32)Y < Height)
-		{
-			Data[X + Y * Width] = Color;
-		}
-	}
-
-	// @param x 0..Width-1
-	// @param Value 0..1
-	void DrawBar(int32 X, float Value)
-	{
-		Value = FMath::Clamp(Value, 0.0f, 1.0f);
-
-		uint32 StartY = (uint32)((1.0f - Value) * Height);
-		for(uint32 Y = StartY; Y < Height; ++Y)
-		{
-			bool bBetterThanReference = Y < (Height / 2);
-			// better than reference system: green, otherwise red
-			DrawPixel(X, Y, bBetterThanReference ? FColor::Green : FColor::Red);
-		}
-	}
-
-	void Save() const
-	{
-		FFileHelper::CreateBitmap(*FilePath, Width, Height, Data.GetData());
-	}
-
-private:
-	TArray<FColor> Data;
-	FString FilePath;
-	uint32 Height;
-	uint32 Width;
-};
-#endif
-
-
-void RendererGPUBenchmark(FRHICommandListImmediate& RHICmdList, FSynthBenchmarkResults& InOut, const FSceneView& View, uint32 WorkScale, bool bDebugOut)
+void RendererGPUBenchmark(FRHICommandListImmediate& RHICmdList, FSynthBenchmarkResults& InOut, const FSceneView& View, float WorkScale, bool bDebugOut)
 {
 	check(IsInRenderingThread());
-	
+
 	// two RT to ping pong so we force the GPU to flush it's pipeline
 	TRefCountPtr<IPooledRenderTarget> RTItems[3];
 	{
@@ -383,7 +340,7 @@ void RendererGPUBenchmark(FRHICommandListImmediate& RHICmdList, FSynthBenchmarkR
 		const uint32 TimerSampleCount = IterationCount * MethodCount + 1;
 
 		static FRenderQueryRHIRef TimerQueries[TimerSampleCount];
-		static uint32 PassCount[IterationCount];
+		static float LocalWorkScale[IterationCount];
 
 		for(uint32  i = 0; i < TimerSampleCount; ++i)
 		{
@@ -394,13 +351,49 @@ void RendererGPUBenchmark(FRHICommandListImmediate& RHICmdList, FSynthBenchmarkR
 		if(!bSupportsTimerQueries)
 		{
 			UE_LOG(LogSynthBenchmark, Warning, TEXT("GPU driver does not support timer queries."));
+
+			// Temporary workaround for GL_TIMESTAMP being unavailable and GL_TIME_ELAPSED workaround breaking drivers
+#if PLATFORM_MAC
+			GLint RendererID = 0;
+			float PerfScale = 1.0f;
+			[[NSOpenGLContext currentContext] getValues:&RendererID forParameter:NSOpenGLCPCurrentRendererID];
+			{
+				switch((RendererID & kCGLRendererIDMatchingMask))
+				{
+					case kCGLRendererATIRadeonX4000ID: // AMD 7xx0 & Dx00 series - should be pretty beefy
+						PerfScale = 1.2f;break;
+					case kCGLRendererATIRadeonX3000ID: // AMD 5xx0, 6xx0 series - mostly OK
+					case kCGLRendererGeForceID: // Nvidia 6x0 & 7x0 series - mostly OK
+						PerfScale = 2.0f;break;
+					case kCGLRendererIntelHD5000ID: // Intel HD 5000, Iris, Iris Pro - not dreadful
+						PerfScale = 4.2f;break;
+					case kCGLRendererIntelHD4000ID: // Intel HD 4000 - quite slow
+						PerfScale = 7.5f;break;
+					case kCGLRendererATIRadeonX2000ID: // ATi 4xx0, 3xx0, 2xx0 - almost all very slow and drivers are now very buggy
+					case kCGLRendererGeForce8xxxID: // Nvidia 3x0, 2x0, 1x0, 9xx0, 8xx0 - almost all very slow
+					case kCGLRendererIntelHDID: // Intel HD 3000 - very, very slow and very buggy driver
+					default:
+						PerfScale = 10.0f;break;
+				}
+			}
+
+			InOut.GPUStats[0].SetMeasuredTime(PerfScale * (1.0f / 4.601f));
+			InOut.GPUStats[1].SetMeasuredTime(PerfScale * (1.0f / 7.447f));
+			InOut.GPUStats[2].SetMeasuredTime(PerfScale * (1.0f / 3.847f));
+			InOut.GPUStats[3].SetMeasuredTime(PerfScale * (1.0f / 25.463f));
+			InOut.GPUStats[4].SetMeasuredTime(PerfScale * (1.0f / 1.072f));
+#endif
+			return;
 		}
 
-		// TimingValues are in Seconds per GPixel
+		// TimingValues are in Seconds
 		FTimingSeries TimingSeries[MethodCount];
+		// in 1/1000000 Seconds
+		uint64 TotalTimes[MethodCount];
 		
 		for(uint32 MethodIterator = 0; MethodIterator < MethodCount; ++MethodIterator)
 		{
+			TotalTimes[MethodIterator] = 0;
 			TimingSeries[MethodIterator].Init(IterationCount);
 		}
 
@@ -435,9 +428,9 @@ void RendererGPUBenchmark(FRHICommandListImmediate& RHICmdList, FSynthBenchmarkR
 				SetRenderTarget(RHICmdList, RTItems[DestRTIndex]->GetRenderTargetItem().TargetableTexture, FTextureRHIRef());	
 
 				// decide how much work we do in this pass
-				PassCount[Iteration] = (Iteration / 10 + 1) * WorkScale;
+				LocalWorkScale[Iteration] = (Iteration / 10 + 1) * WorkScale;
 
-				RunBenchmarkShader(RHICmdList, View, MethodId, RTItems[SrcRTIndex], PassCount[Iteration]);
+				RunBenchmarkShader(RHICmdList, View, MethodId, RTItems[SrcRTIndex], LocalWorkScale[Iteration]);
 
 				RHICmdList.CopyToResolveTarget(RTItems[DestRTIndex]->GetRenderTargetItem().TargetableTexture, RTItems[DestRTIndex]->GetRenderTargetItem().ShaderResourceTexture, false, FResolveParams());
 
@@ -474,10 +467,6 @@ void RendererGPUBenchmark(FRHICommandListImmediate& RHICmdList, FSynthBenchmarkR
 			RHICmdList.GetRenderQueryResult(TimerQueries[0], OldAbsTime, true);
 			GTimerQueryPool.ReleaseQuery(RHICmdList, TimerQueries[0]);
 
-#if !UE_BUILD_SHIPPING
-			FBenchmarkGraph BenchmarkGraph(IterationCount, IterationCount, *(FPaths::ScreenShotDir() + TEXT("GPUSynthBenchmarkGraph.bmp")));
-#endif
-
 			for(uint32 Iteration = 0; Iteration < IterationCount; ++Iteration)
 			{
 				uint32 Results[MethodCount];
@@ -490,131 +479,43 @@ void RendererGPUBenchmark(FRHICommandListImmediate& RHICmdList, FSynthBenchmarkR
 					RHICmdList.GetRenderQueryResult(TimerQueries[QueryIndex], AbsTime, true);
 					GTimerQueryPool.ReleaseQuery(RHICmdList, TimerQueries[QueryIndex]);
 
-					Results[MethodId] = AbsTime - OldAbsTime;
-					OldAbsTime = AbsTime;
-				}
-
-				double SamplesInGPix = PassCount[Iteration] * GBenchmarkResolution * GBenchmarkResolution / 1000000000.0;
+					uint64 RelTime = AbsTime - OldAbsTime; 
 
-				for(uint32 MethodId = 0; MethodId < MethodCount; ++MethodId)
-				{
-					double TimeInSec = Results[MethodId] / 1000000.0;
-					double TimingValue = TimeInSec / SamplesInGPix;
+					TotalTimes[MethodId] += RelTime;
+					Results[MethodId] = RelTime;
 
-					// TimingValue in Seconds per GPixel
-					TimingSeries[MethodId].SetEntry(Iteration, (float)TimingValue);
+					OldAbsTime = AbsTime;
 				}
 
-#if !UE_BUILD_SHIPPING
+				for(uint32 MethodId = 0; MethodId < MethodCount; ++MethodId)
 				{
-					// This is for debugging and we don't want to change the output but we still use "InOut".
-					// That shouldn't hurt, as we override the values after that anyway.
+					float TimeInSec = Results[MethodId] / 1000000.0f;
 
-					for(uint32 MethodId = 0; MethodId < MethodCount; ++MethodId)
-					{
-						InOut.GPUStats[MethodId].SetMeasuredTime(TimingSeries[MethodId].GetEntry(Iteration));
-					}
+					// to normalize from seconds to seconds per GPixel
+					float SamplesInGPix = LocalWorkScale[Iteration] * GBenchmarkResolution * GBenchmarkResolution / 1000000000.0f;
 
-					float LocalGPUIndex = InOut.ComputeGPUPerfIndex();
-
-					// * 0.01 to get it in 0..1 range
-					// * 0.5f to have 100 is the middle
-					BenchmarkGraph.DrawBar(Iteration, LocalGPUIndex * 0.01f * 0.5f);
-				}
-#endif
-			}
-			
-			// Temporary workaround for GL_TIMESTAMP being unavailable and GL_TIME_ELAPSED workaround breaking drivers
-#if PLATFORM_MAC
-			if(!TimerQueries[0])
-			{
-				GLint RendererID = 0;
-				[[NSOpenGLContext currentContext] getValues:&RendererID forParameter:NSOpenGLCPCurrentRendererID];
-				{
-					switch((RendererID & kCGLRendererIDMatchingMask))
-					{
-						case kCGLRendererATIRadeonX4000ID: // AMD 7xx0 & Dx00 series - should be pretty beefy
-						{
-							InOut.GPUStats[0].SetMeasuredTime(1.2f * (1.0f / 4.601f));
-							InOut.GPUStats[1].SetMeasuredTime(1.2f * (1.0f / 7.447f));
-							InOut.GPUStats[2].SetMeasuredTime(1.2f * (1.0f / 3.847f));
-							InOut.GPUStats[3].SetMeasuredTime(1.2f * (1.0f / 25.463f));
-							InOut.GPUStats[4].SetMeasuredTime(1.2f * (1.0f / 1.072f));
-							break;
-						}
-						case kCGLRendererATIRadeonX3000ID: // AMD 5xx0, 6xx0 series - mostly OK
-						case kCGLRendererGeForceID: // Nvidia 6x0 & 7x0 series - mostly OK
-						{
-							InOut.GPUStats[0].SetMeasuredTime(2.f * (1.0f / 4.601f));
-							InOut.GPUStats[1].SetMeasuredTime(2.f * (1.0f / 7.447f));
-							InOut.GPUStats[2].SetMeasuredTime(2.f * (1.0f / 3.847f));
-							InOut.GPUStats[3].SetMeasuredTime(2.f * (1.0f / 25.463f));
-							InOut.GPUStats[4].SetMeasuredTime(2.f * (1.0f / 1.072f));
-							break;
-						}
-						case kCGLRendererIntelHD5000ID: // Intel HD 5000, Iris, Iris Pro - not dreadful
-						{
-							InOut.GPUStats[0].SetMeasuredTime(4.f * (1.0f / 4.601f));
-							InOut.GPUStats[1].SetMeasuredTime(4.f * (1.0f / 7.447f));
-							InOut.GPUStats[2].SetMeasuredTime(4.f * (1.0f / 3.847f));
-							InOut.GPUStats[3].SetMeasuredTime(4.f * (1.0f / 25.463f));
-							InOut.GPUStats[4].SetMeasuredTime(4.f * (1.0f / 1.072f));
-							break;
-						}
-						case kCGLRendererIntelHD4000ID: // Intel HD 4000 - quite slow
-						{
-							InOut.GPUStats[0].SetMeasuredTime(7.5f * (1.0f / 4.601f));
-							InOut.GPUStats[1].SetMeasuredTime(7.5f * (1.0f / 7.447f));
-							InOut.GPUStats[2].SetMeasuredTime(7.5f * (1.0f / 3.847f));
-							InOut.GPUStats[3].SetMeasuredTime(7.5f * (1.0f / 25.463f));
-							InOut.GPUStats[4].SetMeasuredTime(7.5f * (1.0f / 1.072f));
-							break;
-						}
-						case kCGLRendererATIRadeonX2000ID: // ATi 4xx0, 3xx0, 2xx0 - almost all very slow and drivers are now very buggy
-						case kCGLRendererGeForce8xxxID: // Nvidia 3x0, 2x0, 1x0, 9xx0, 8xx0 - almost all very slow
-						case kCGLRendererIntelHDID: // Intel HD 3000 - very, very slow and very buggy driver
-						default:
-						{
-							InOut.GPUStats[0].SetMeasuredTime(10.f * (1.0f / 4.601f));
-							InOut.GPUStats[1].SetMeasuredTime(10.f * (1.0f / 7.447f));
-							InOut.GPUStats[2].SetMeasuredTime(10.f * (1.0f / 3.847f));
-							InOut.GPUStats[3].SetMeasuredTime(10.f * (1.0f / 25.463f));
-							InOut.GPUStats[4].SetMeasuredTime(10.f * (1.0f / 1.072f));
-							break;
-						}
-					}
+					// TimingValue in Seconds per GPixel
+					TimingSeries[MethodId].SetEntry(Iteration, TimeInSec / SamplesInGPix);
 				}
 			}
-			
-#endif
 
-			for(uint32 MethodId = 0; MethodId < MethodCount; ++MethodId)
+			if(bSupportsTimerQueries)
 			{
-				float Confidence = 0.0f;
-				
 				// Temporary workaround for GL_TIMESTAMP being unavailable and GL_TIME_ELAPSED workaround breaking drivers on OS X
-				if(!PLATFORM_MAC || bSupportsTimerQueries)
+				for(uint32 MethodId = 0; MethodId < MethodCount; ++MethodId)
 				{
-					float TimingValue = TimingSeries[MethodId].ComputeValue(Confidence);
+					float Confidence = 0.0f;
+					// in seconds per GPixel
+					float NormalizedTime = TimingSeries[MethodId].ComputeValue(Confidence);
 
 					if(Confidence > 0)
 					{
-						InOut.GPUStats[MethodId].SetMeasuredTime(TimingValue, Confidence);
+						FTimeSample TimeSample(TotalTimes[MethodId] / 1000000.0f, NormalizedTime);
+
+						InOut.GPUStats[MethodId].SetMeasuredTime(TimeSample, Confidence);
 					}
 				}
-
-				UE_LOG(LogSynthBenchmark, Display, TEXT("         ... %.3f GigaPix/s, Confidence=%.0f%% '%s'"),
-					1.0f / InOut.GPUStats[MethodId].GetMeasuredTime(), Confidence, InOut.GPUStats[MethodId].GetDesc());
 			}
-
-			UE_LOG(LogSynthBenchmark, Display, TEXT(""));
-			
-#if !UE_BUILD_SHIPPING
-			if(bDebugOut)
-			{
-				BenchmarkGraph.Save();
-			}
-#endif
 		}
 	}
 }
diff --git a/Engine/Source/Runtime/Renderer/Private/GPUBenchmark.h b/Engine/Source/Runtime/Renderer/Private/GPUBenchmark.h
--- a/Engine/Source/Runtime/Renderer/Private/GPUBenchmark.h
+++ b/Engine/Source/Runtime/Renderer/Private/GPUBenchmark.h
@@ -8,4 +8,4 @@
 
 // @param >0, WorkScale 10 for normal precision and runtime of less than a second
 // @param bDebugOut has no effect in shipping
-void RendererGPUBenchmark(FRHICommandListImmediate& RHICmdList, FSynthBenchmarkResults& InOut, const FSceneView& View, uint32 WorkScale, bool bDebugOut = false);
+void RendererGPUBenchmark(FRHICommandListImmediate& RHICmdList, FSynthBenchmarkResults& InOut, const FSceneView& View, float WorkScale, bool bDebugOut = false);
diff --git a/Engine/Source/Runtime/Renderer/Private/Renderer.cpp b/Engine/Source/Runtime/Renderer/Private/Renderer.cpp
--- a/Engine/Source/Runtime/Renderer/Private/Renderer.cpp
+++ b/Engine/Source/Runtime/Renderer/Private/Renderer.cpp
@@ -124,7 +124,7 @@ void FRendererModule::DebugLogOnCrash()
 	}
 }
 
-void FRendererModule::GPUBenchmark(FSynthBenchmarkResults& InOut, uint32 WorkScale, bool bDebugOut)
+void FRendererModule::GPUBenchmark(FSynthBenchmarkResults& InOut, float WorkScale)
 {
 	check(IsInGameThread());
 
@@ -153,14 +153,13 @@ void FRendererModule::GPUBenchmark(FSynthBenchmarkResults& InOut, uint32 WorkSca
 
 	FSceneView DummyView(ViewInitOptions);
 
-	ENQUEUE_UNIQUE_RENDER_COMMAND_FOURPARAMETER(
+	ENQUEUE_UNIQUE_RENDER_COMMAND_THREEPARAMETER(
 	  RendererGPUBenchmarkCommand,
 	  FSceneView, DummyView, DummyView,
-	  uint32, WorkScale, WorkScale,
-	  bool, bDebugOut, bDebugOut,
+	  float, WorkScale, WorkScale,
 	  FSynthBenchmarkResults&, InOut, InOut,
 	{
-		RendererGPUBenchmark(RHICmdList, InOut, DummyView, WorkScale, bDebugOut);
+		RendererGPUBenchmark(RHICmdList, InOut, DummyView, WorkScale);
 	});
 	FlushRenderingCommands();
 }
diff --git a/Engine/Source/Runtime/Renderer/Private/RendererPrivate.h b/Engine/Source/Runtime/Renderer/Private/RendererPrivate.h
--- a/Engine/Source/Runtime/Renderer/Private/RendererPrivate.h
+++ b/Engine/Source/Runtime/Renderer/Private/RendererPrivate.h
@@ -34,7 +34,7 @@ public:
 	virtual void RenderTargetPoolFindFreeElement(const FPooledRenderTargetDesc& Desc, TRefCountPtr<IPooledRenderTarget> &Out, const TCHAR* InDebugName) override;
 	virtual void TickRenderTargetPool() override;
 	virtual void DebugLogOnCrash() override;
-	virtual void GPUBenchmark(FSynthBenchmarkResults& InOut, uint32 WorkScale, bool bDebugOut) override;
+	virtual void GPUBenchmark(FSynthBenchmarkResults& InOut, float WorkScale) override;
 	virtual void QueryVisualizeTexture(FQueryVisualizeTexureInfo& Out) override;
 	virtual void ExecVisualizeTextureCmd(const FString& Cmd) override;
 	virtual void UpdateMapNeedsLightingFullyRebuiltState(UWorld* World) override;
