fix build for last checkin

diff --git a/Engine/Source/Runtime/Renderer/Public/RendererInterface.h b/Engine/Source/Runtime/Renderer/Public/RendererInterface.h
--- a/Engine/Source/Runtime/Renderer/Public/RendererInterface.h
+++ b/Engine/Source/Runtime/Renderer/Public/RendererInterface.h
@@ -465,10 +465,8 @@ public:
 	/** Renderer gets a chance to log some useful crash data */
 	virtual void DebugLogOnCrash() = 0;
 
-	/**  */
 	// @param WorkScale >0, 10 for normal precision and runtime of less than a second
-	// @param bDebugOut has no effect in shipping
-	virtual void GPUBenchmark(FSynthBenchmarkResults& InOut, uint32 WorkScale = 10, bool bDebugOut = false) = 0;
+	virtual void GPUBenchmark(FSynthBenchmarkResults& InOut, float WorkScale = 10.0f) = 0;
 
 	virtual void QueryVisualizeTexture(FQueryVisualizeTexureInfo& Out) = 0;
 	virtual void ExecVisualizeTextureCmd(const FString& Cmd) = 0;
