fixed TTP 333766 UE4: RENDERING: BufferVisualization doesn't use the Before/AfterTonemapper material setting

diff --git a/Engine/Source/Runtime/Engine/Classes/Engine/BlendableInterface.h b/Engine/Source/Runtime/Engine/Classes/Engine/BlendableInterface.h
--- a/Engine/Source/Runtime/Engine/Classes/Engine/BlendableInterface.h
+++ b/Engine/Source/Runtime/Engine/Classes/Engine/BlendableInterface.h
@@ -40,14 +40,27 @@ class IBlendableInterface
 
 struct FPostProcessMaterialNode
 {
+	// default constructor
+	FPostProcessMaterialNode()
+		: MaterialInterface(0), bIsMID(false), Location(BL_MAX), Priority(0)
+	{
+	}
+
+	// constructor
+	FPostProcessMaterialNode(UMaterialInterface* InMaterialInterface, EBlendableLocation InLocation, int32 InPriority)
+		: MaterialInterface(InMaterialInterface), bIsMID(false), Location(InLocation), Priority(InPriority)
+	{
+	}
+
+	// constructor
 	FPostProcessMaterialNode(UMaterialInstanceDynamic* InMID, EBlendableLocation InLocation, int32 InPriority)
-		:MID(InMID), Location(InLocation), Priority(InPriority)
+		: MaterialInterface((UMaterialInterface*)InMID), bIsMID(true), Location(InLocation), Priority(InPriority)
 	{
 	}
 
-	UMaterialInstanceDynamic* MID;
-	EBlendableLocation Location;
-	int32 Priority;
+	UMaterialInterface* GetMaterialInterface() const { return MaterialInterface; }
+	// only call if you are sure it's an MID
+	UMaterialInstanceDynamic* GetMID() const { check(bIsMID); return (UMaterialInstanceDynamic*)MaterialInterface; }
 
 	// for type safety in FBlendableManager
 	static const FName& GetFName()
@@ -70,6 +83,19 @@ struct FPostProcessMaterialNode
 			return false;
 		}
 	};
+
+	EBlendableLocation GetLocation() const { return Location; }
+	int32 GetPriority() const { return Priority; }
+
+private:
+	UMaterialInterface* MaterialInterface;
+
+	// if MaterialInterface is an MID, needed for GetMID()
+	bool bIsMID;
+	// e.g. BL_BeforeTonemapping
+	EBlendableLocation Location;
+	// default 0
+	int32 Priority;
 };
 
 
diff --git a/Engine/Source/Runtime/Engine/Private/Materials/MaterialInstance.cpp b/Engine/Source/Runtime/Engine/Private/Materials/MaterialInstance.cpp
--- a/Engine/Source/Runtime/Engine/Private/Materials/MaterialInstance.cpp
+++ b/Engine/Source/Runtime/Engine/Private/Materials/MaterialInstance.cpp
@@ -2260,7 +2260,7 @@ FPostProcessMaterialNode* IteratePostProcessMaterialNodes(const FFinalPostProces
 			return 0;
 		}
 
-		if(DataPtr->MID->GetMaterial() == Material && DataPtr->Location == Location && DataPtr->Priority == Priority)
+		if(DataPtr->GetLocation() == Location && DataPtr->GetPriority() == Priority && DataPtr->GetMaterialInterface()->GetMaterial() == Material)
 		{
 			return DataPtr;
 		}
@@ -2318,10 +2318,10 @@ void UMaterialInstance::OverrideBlendableSettings(class FSceneView& View, float
 
 	if(PostProcessMaterialNode)
 	{
-		UMaterialInstanceDynamic* DestMID = PostProcessMaterialNode->MID;
-		UMaterialInstance* SrcMID = (UMaterialInstance*)this;
-
+		UMaterialInstanceDynamic* DestMID = PostProcessMaterialNode->GetMID();
 		check(DestMID);
+
+		UMaterialInstance* SrcMID = (UMaterialInstance*)this;
 		check(SrcMID);
 
 		// a material already exists, blend with existing ones
diff --git a/Engine/Source/Runtime/Engine/Private/ShowFlags.cpp b/Engine/Source/Runtime/Engine/Private/ShowFlags.cpp
--- a/Engine/Source/Runtime/Engine/Private/ShowFlags.cpp
+++ b/Engine/Source/Runtime/Engine/Private/ShowFlags.cpp
@@ -422,7 +422,7 @@ void EngineShowFlagOverride(EShowFlagInitMode ShowFlagInitMode, EViewModeIndex V
 	// disable AA in full screen GBuffer visualization
 	if(EngineShowFlags.VisualizeBuffer && CurrentBufferVisualizationMode != NAME_None)
 	{
-		EngineShowFlags.AntiAliasing = 0;
+		EngineShowFlags.Tonemapper = 0;
 	}
 
 #if !(UE_BUILD_SHIPPING || UE_BUILD_TEST)
diff --git a/Engine/Source/Runtime/Renderer/Private/PostProcess/PostProcessing.cpp b/Engine/Source/Runtime/Renderer/Private/PostProcess/PostProcessing.cpp
--- a/Engine/Source/Runtime/Renderer/Private/PostProcess/PostProcessing.cpp
+++ b/Engine/Source/Runtime/Renderer/Private/PostProcess/PostProcessing.cpp
@@ -480,7 +480,7 @@ FPostProcessMaterialNode* IteratePostProcessMaterialNodes(const FFinalPostProces
 	{
 		FPostProcessMaterialNode* DataPtr = Dest.BlendableManager.IterateBlendables<FPostProcessMaterialNode>(Iterator);
 
-		if(!DataPtr || DataPtr->Location == InLocation)
+		if(!DataPtr || DataPtr->GetLocation() == InLocation)
 		{
 			return DataPtr;
 		}
@@ -494,13 +494,26 @@ static void AddPostProcessMaterial(FPostprocessContext& Context, EBlendableLocat
 		return;
 	}
 
-	FBlendableEntry* Iterator = 0;
-
+	// hard coded - this should be a reasonable limit
 	const uint32 MAX_PPMATERIALNODES = 10;
+	FBlendableEntry* Iterator = 0;
+	FPostProcessMaterialNode PPNodes[MAX_PPMATERIALNODES];
+	uint32 PPNodeCount = 0;
 
-	FPostProcessMaterialNode* PPNodes[MAX_PPMATERIALNODES];
+	static const auto CVarDumpFrames = IConsoleManager::Get().FindTConsoleVariableDataInt(TEXT("r.BufferVisualizationDumpFrames"));
+	static const auto CVarDumpFramesAsHDR = IConsoleManager::Get().FindTConsoleVariableDataInt(TEXT("r.BufferVisualizationDumpFramesAsHDR"));
 
-	uint32 PPNodeCount = 0;
+	if(Context.View.Family->EngineShowFlags.VisualizeBuffer)
+	{	
+		// Apply requested material to the full screen
+		UMaterial* Material = GetBufferVisualizationData().GetMaterial(Context.View.CurrentBufferVisualizationMode);
+		
+		if(Material && Material->BlendableLocation == InLocation)
+		{
+			PPNodes[0] = FPostProcessMaterialNode(Material, InLocation, Material->BlendablePriority);
+			++PPNodeCount;
+		}
+	}
 	for(;PPNodeCount < MAX_PPMATERIALNODES; ++PPNodeCount)
 	{
 		FPostProcessMaterialNode* Data = IteratePostProcessMaterialNodes(Context.View.FinalPostProcessSettings, InLocation, Iterator);
@@ -510,18 +523,16 @@ static void AddPostProcessMaterial(FPostprocessContext& Context, EBlendableLocat
 			break;
 		}
 
-		check(Data->MID);
+		check(Data->GetMaterialInterface());
 
-		PPNodes[PPNodeCount] = Data;
+		PPNodes[PPNodeCount] = *Data;
 	}
 	
-	::Sort(PPNodes,PPNodeCount, FPostProcessMaterialNode::FCompare());
+	::Sort(PPNodes, PPNodeCount, FPostProcessMaterialNode::FCompare());
 
 	for(uint32 i = 0; i < PPNodeCount; ++i)
 	{
-		FPostProcessMaterialNode* Data = PPNodes[i];
-
-		UMaterialInterface* MaterialInterface = Data->MID;
+		UMaterialInterface* MaterialInterface = PPNodes[i].GetMaterialInterface();
 
 		FMaterialRenderProxy* Proxy = MaterialInterface->GetRenderProxy(false);
 
@@ -582,7 +593,7 @@ static void AddHighResScreenshotMask(FPostprocessContext& Context, FRenderingCom
 	}
 }
 
-static void AddGBufferVisualization(FPostprocessContext& Context, FRenderingCompositeOutputRef& SeparateTranslucencyInput)
+static void AddGBufferVisualizationOverview(FPostprocessContext& Context, FRenderingCompositeOutputRef& SeparateTranslucencyInput)
 {
 	static const auto CVarDumpFrames = IConsoleManager::Get().FindTConsoleVariableDataInt(TEXT("r.BufferVisualizationDumpFrames"));
 	static const auto CVarDumpFramesAsHDR = IConsoleManager::Get().FindTConsoleVariableDataInt(TEXT("r.BufferVisualizationDumpFramesAsHDR"));
@@ -664,22 +675,6 @@ static void AddGBufferVisualization(FPostprocessContext& Context, FRenderingComp
 				}
 			}
 		}
-		else if (bVisualizationEnabled && !bOverviewModeEnabled)
-		{
-			// Apply requested material to the full screen
-			UMaterial* Material = GetBufferVisualizationData().GetMaterial(Context.View.CurrentBufferVisualizationMode);
-			check(Material);
-			FRenderingCompositePass* MaterialPass = Context.Graph.RegisterPass(new(FMemStack::Get()) FRCPassPostProcessMaterial(Material));
-			MaterialPass->SetInput(ePId_Input0, FRenderingCompositeOutputRef(Context.FinalOutput));
-			MaterialPass->SetInput(ePId_Input1, FRenderingCompositeOutputRef(SeparateTranslucencyInput));
-			Context.FinalOutput = FRenderingCompositeOutputRef(MaterialPass);
-
-			if (Material->GetMaterialResource(Context.View.GetFeatureLevel())->NeedsGBuffer())
-			{
-				// AdjustGBufferRefCount(-1) call is done when the pass gets executed
-				GSceneRenderTargets.AdjustGBufferRefCount(1);
-			}
-		}
 	}
 }
 
@@ -1071,7 +1066,7 @@ void FPostProcessing::Process(FRHICommandListImmediate& RHICmdList, FViewInfo& V
 			Context.FinalOutput = FRenderingCompositeOutputRef(PassVisualize);
 		}
 
-		AddGBufferVisualization(Context, SeparateTranslucency);
+		AddGBufferVisualizationOverview(Context, SeparateTranslucency);
 
 		bool bStereoRenderingAndHMD = View.Family->EngineShowFlags.StereoRendering && View.Family->EngineShowFlags.HMDDistortion;
 		if (bStereoRenderingAndHMD)
