Fix up Mac workaround for Martin's latest changes.

diff --git a/Engine/Source/Runtime/Renderer/Private/GPUBenchmark.cpp b/Engine/Source/Runtime/Renderer/Private/GPUBenchmark.cpp
--- a/Engine/Source/Runtime/Renderer/Private/GPUBenchmark.cpp
+++ b/Engine/Source/Runtime/Renderer/Private/GPUBenchmark.cpp
@@ -361,27 +361,37 @@ void RendererGPUBenchmark(FRHICommandListImmediate& RHICmdList, FSynthBenchmarkR
 				switch((RendererID & kCGLRendererIDMatchingMask))
 				{
 					case kCGLRendererATIRadeonX4000ID: // AMD 7xx0 & Dx00 series - should be pretty beefy
-						PerfScale = 1.2f;break;
+						PerfScale = 1.2f;
+						break;
 					case kCGLRendererATIRadeonX3000ID: // AMD 5xx0, 6xx0 series - mostly OK
 					case kCGLRendererGeForceID: // Nvidia 6x0 & 7x0 series - mostly OK
-						PerfScale = 2.0f;break;
+						PerfScale = 2.0f;
+						break;
 					case kCGLRendererIntelHD5000ID: // Intel HD 5000, Iris, Iris Pro - not dreadful
-						PerfScale = 4.2f;break;
+						PerfScale = 4.2f;
+						break;
 					case kCGLRendererIntelHD4000ID: // Intel HD 4000 - quite slow
-						PerfScale = 7.5f;break;
+						PerfScale = 7.5f;
+						break;
 					case kCGLRendererATIRadeonX2000ID: // ATi 4xx0, 3xx0, 2xx0 - almost all very slow and drivers are now very buggy
 					case kCGLRendererGeForce8xxxID: // Nvidia 3x0, 2x0, 1x0, 9xx0, 8xx0 - almost all very slow
 					case kCGLRendererIntelHDID: // Intel HD 3000 - very, very slow and very buggy driver
 					default:
-						PerfScale = 10.0f;break;
+						PerfScale = 10.0f;
+						break;
 				}
 			}
-
-			InOut.GPUStats[0].SetMeasuredTime(PerfScale * (1.0f / 4.601f));
-			InOut.GPUStats[1].SetMeasuredTime(PerfScale * (1.0f / 7.447f));
-			InOut.GPUStats[2].SetMeasuredTime(PerfScale * (1.0f / 3.847f));
-			InOut.GPUStats[3].SetMeasuredTime(PerfScale * (1.0f / 25.463f));
-			InOut.GPUStats[4].SetMeasuredTime(PerfScale * (1.0f / 1.072f));
+			
+			InOut.GPUStats[0] = FSynthBenchmarkStat(TEXT("ALUHeavyNoise"), 1.0f / 4.601f, TEXT("s/GigaPix"));
+			InOut.GPUStats[1] = FSynthBenchmarkStat(TEXT("TexHeavy"), 1.0f / 7.447f, TEXT("s/GigaPix"));
+			InOut.GPUStats[2] = FSynthBenchmarkStat(TEXT("DepTexHeavy"), 1.0f / 3.847f, TEXT("s/GigaPix"));
+			InOut.GPUStats[3] = FSynthBenchmarkStat(TEXT("FillOnly"), 1.0f / 25.463f, TEXT("s/GigaPix"));
+			InOut.GPUStats[4] = FSynthBenchmarkStat(TEXT("Bandwidth"), 1.0f / 1.072f, TEXT("s/GigaPix"));
+			InOut.GPUStats[0].SetMeasuredTime( FTimeSample(PerfScale, PerfScale * (1.0f / 4.601f)) );
+			InOut.GPUStats[1].SetMeasuredTime( FTimeSample(PerfScale, PerfScale * (1.0f / 7.447f)) );
+			InOut.GPUStats[2].SetMeasuredTime( FTimeSample(PerfScale, PerfScale * (1.0f / 3.847f)) );
+			InOut.GPUStats[3].SetMeasuredTime( FTimeSample(PerfScale, PerfScale * (1.0f / 25.463f)) );
+			InOut.GPUStats[4].SetMeasuredTime( FTimeSample(PerfScale, PerfScale * (1.0f / 1.072f)) );
 #endif
 			return;
 		}
@@ -501,7 +511,6 @@ void RendererGPUBenchmark(FRHICommandListImmediate& RHICmdList, FSynthBenchmarkR
 
 			if(bSupportsTimerQueries)
 			{
-				// Temporary workaround for GL_TIMESTAMP being unavailable and GL_TIME_ELAPSED workaround breaking drivers on OS X
 				for(uint32 MethodId = 0; MethodId < MethodCount; ++MethodId)
 				{
 					float Confidence = 0.0f;
