UE4 - Fix ShadowDepth render targets for Mobile

diff --git a/Engine/Source/Runtime/OpenGLDrv/Private/OpenGLDevice.cpp b/Engine/Source/Runtime/OpenGLDrv/Private/OpenGLDevice.cpp
--- a/Engine/Source/Runtime/OpenGLDrv/Private/OpenGLDevice.cpp
+++ b/Engine/Source/Runtime/OpenGLDrv/Private/OpenGLDevice.cpp
@@ -680,7 +680,11 @@ static void InitRHICapabilitiesForGL()
 	SetupTextureFormat( PF_Unknown,				FOpenGLTextureFormat());
 	SetupTextureFormat( PF_A32B32G32R32F,		FOpenGLTextureFormat(GL_RGBA32F,						GL_NONE,									GL_RGBA,				GL_FLOAT,							false));
 	SetupTextureFormat( PF_UYVY,				FOpenGLTextureFormat());
+#if PLATFORM_IOS || PLATFORM_ANDROID
+	SetupTextureFormat( PF_ShadowDepth,			FOpenGLTextureFormat(GL_DEPTH_COMPONENT,				GL_NONE,									GL_DEPTH_COMPONENT,		GL_UNSIGNED_INT,					false));
+#else
 	SetupTextureFormat( PF_ShadowDepth,			FOpenGLTextureFormat(GL_DEPTH_COMPONENT16,				GL_NONE,									GL_DEPTH_COMPONENT,		GL_UNSIGNED_INT,					false));
+#endif
 	SetupTextureFormat( PF_A16B16G16R16,		FOpenGLTextureFormat(GL_RGBA16,							GL_RGBA16,									GL_RGBA,				GL_UNSIGNED_SHORT,					false));
 	SetupTextureFormat( PF_D24,					FOpenGLTextureFormat(depthFormat,						GL_NONE,									GL_DEPTH_COMPONENT,		GL_UNSIGNED_INT,					false));
 	SetupTextureFormat( PF_A1,					FOpenGLTextureFormat());
diff --git a/Engine/Source/Runtime/Renderer/Private/PostProcess/SceneRenderTargets.cpp b/Engine/Source/Runtime/Renderer/Private/PostProcess/SceneRenderTargets.cpp
--- a/Engine/Source/Runtime/Renderer/Private/PostProcess/SceneRenderTargets.cpp
+++ b/Engine/Source/Runtime/Renderer/Private/PostProcess/SceneRenderTargets.cpp
@@ -740,7 +740,14 @@ void FSceneRenderTargets::FinishRenderingPrePass(FRHICommandListImmediate& RHICm
 void FSceneRenderTargets::BeginRenderingShadowDepth(FRHICommandListImmediate& RHICmdList)
 {
 	GRenderTargetPool.VisualizeTexture.SetCheckPoint(RHICmdList, ShadowDepthZ);
-	SetRenderTarget(RHICmdList, FTextureRHIRef(), GetShadowDepthZSurface());
+	if (GSupportsDepthRenderTargetWithoutColorRenderTarget)
+	{
+		SetRenderTarget(RHICmdList, FTextureRHIRef(), GetShadowDepthZSurface());
+	}
+	else
+	{
+		SetRenderTarget(RHICmdList, GetOptionalShadowDepthColorSurface(), GetShadowDepthZSurface());
+	}
 }
 
 void FSceneRenderTargets::BeginRenderingCubeShadowDepth(FRHICommandListImmediate& RHICmdList, int32 ShadowResolution)
@@ -1073,6 +1080,12 @@ void FSceneRenderTargets::AllocateForwardShadingShadowDepthTarget(const FIntPoin
 {
 	FPooledRenderTargetDesc Desc(FPooledRenderTargetDesc::Create2DDesc(ShadowBufferResolution, PF_ShadowDepth, TexCreate_None, TexCreate_DepthStencilTargetable, false));
 	GRenderTargetPool.FindFreeElement(Desc, ShadowDepthZ, TEXT("ShadowDepthZ"));
+
+	if (!GSupportsDepthRenderTargetWithoutColorRenderTarget)
+	{
+		FPooledRenderTargetDesc Desc(FPooledRenderTargetDesc::Create2DDesc(ShadowBufferResolution, PF_B8G8R8A8, TexCreate_None, TexCreate_RenderTargetable, false));
+		GRenderTargetPool.FindFreeElement(Desc, OptionalShadowDepthColor, TEXT("OptionalShadowDepthColor"));
+	}
 }
 
 // for easier use of "VisualizeTexture"
