UE4 - Compile fix, missing file

diff --git a/Engine/Source/Runtime/Renderer/Private/PostProcess/SceneRenderTargets.h b/Engine/Source/Runtime/Renderer/Private/PostProcess/SceneRenderTargets.h
--- a/Engine/Source/Runtime/Renderer/Private/PostProcess/SceneRenderTargets.h
+++ b/Engine/Source/Runtime/Renderer/Private/PostProcess/SceneRenderTargets.h
@@ -317,6 +317,10 @@ public:
 	{ 
 		return (const FTexture2DRHIRef&)ShadowDepthZ->GetRenderTargetItem().TargetableTexture; 
 	}
+	const FTexture2DRHIRef& GetOptionalShadowDepthColorSurface() const 
+	{ 
+		return (const FTexture2DRHIRef&)OptionalShadowDepthColor->GetRenderTargetItem().TargetableTexture; 
+	}
 
 	const FTexture2DRHIRef& GetReflectiveShadowMapNormalSurface() const { return (const FTexture2DRHIRef&)ReflectiveShadowMapNormal->GetRenderTargetItem().TargetableTexture; }
 	const FTexture2DRHIRef& GetReflectiveShadowMapDiffuseSurface() const { return (const FTexture2DRHIRef&)ReflectiveShadowMapDiffuse->GetRenderTargetItem().TargetableTexture; }
@@ -438,8 +442,8 @@ public:
 	TRefCountPtr<IPooledRenderTarget> ScreenSpaceAO;
 	// used by the CustomDepth material feature, is allocated on demand or if r.CustomDepth is 2
 	TRefCountPtr<IPooledRenderTarget> CustomDepth;
-	// Render target for per-object shadow depths.
-	TRefCountPtr<IPooledRenderTarget> ShadowDepthZ;
+	// Render target for per-object shadow depths. (Also the optional in case this RHI requires a color render target)
+	TRefCountPtr<IPooledRenderTarget> ShadowDepthZ, OptionalShadowDepthColor;
 	// Cache of preshadow depths
 	//@todo - this should go in FScene
 	TRefCountPtr<IPooledRenderTarget> PreShadowCacheDepthZ;
