Reflection Environment uses standard deferred application which has a low constant cost, until the number of captures exceeds 'r.TiledReflectionEnvironmentMinimumCount', where it switches to tiled deferred.
* The threshold is currently 0, as clear coat has not been implemented with the standard deferred path

diff --git a/Engine/Shaders/ReflectionEnvironmentShaders.usf b/Engine/Shaders/ReflectionEnvironmentShaders.usf
--- a/Engine/Shaders/ReflectionEnvironmentShaders.usf
+++ b/Engine/Shaders/ReflectionEnvironmentShaders.usf
@@ -371,6 +371,11 @@ float4 CaptureProperties;
 float4x4 CaptureBoxTransform;
 float4 CaptureBoxScales;
 
+#if FEATURE_LEVEL >= FEATURE_LEVEL_SM5
+TextureCubeArray ReflectionEnvironmentColorTextureArray;
+int CaptureArrayIndex;
+#endif
+
 /** Standard deferred shading implementation of reflection environment, used in feature level SM4 */
 void StandardDeferredReflectionPS(
 	float4 InScreenPosition : TEXCOORD0,
@@ -461,7 +466,12 @@ void StandardDeferredReflectionPS(
 		if (DistanceAlpha > 0)
 		{
 			float AbsoluteSpecularMip = ComputeReflectionCaptureMipFromRoughness(InGBufferData.Roughness);
+
+#if FEATURE_LEVEL >= FEATURE_LEVEL_SM5
+			float4 Reflection = ReflectionEnvironmentColorTextureArray.SampleLevel(ReflectionEnvironmentColorSampler, float4(ProjectedCaptureVector, CaptureArrayIndex), AbsoluteSpecularMip);
+#else
 			float4 Reflection = TextureCubeSampleLevel(ReflectionEnvironmentColorTexture, ReflectionEnvironmentColorSampler, ProjectedCaptureVector, AbsoluteSpecularMip);
+#endif
 
 			Reflection.rgb *= CaptureProperties.r;
 
diff --git a/Engine/Source/Runtime/Renderer/Private/DeferredShadingRenderer.h b/Engine/Source/Runtime/Renderer/Private/DeferredShadingRenderer.h
--- a/Engine/Source/Runtime/Renderer/Private/DeferredShadingRenderer.h
+++ b/Engine/Source/Runtime/Renderer/Private/DeferredShadingRenderer.h
@@ -364,10 +364,10 @@ private:
 	void RenderReflectionCaptureSpecularBounceForAllViews(FRHICommandListImmediate& RHICmdList);
 
 	/** Render image based reflections (SSR, Env, SkyLight) with compute shaders */
-	void RenderImageBasedReflectionsSM5ForAllViews(FRHICommandListImmediate& RHICmdList, const TRefCountPtr<IPooledRenderTarget>& DynamicBentNormalAO);
+	void RenderTiledDeferredImageBasedReflections(FRHICommandListImmediate& RHICmdList, const TRefCountPtr<IPooledRenderTarget>& DynamicBentNormalAO);
 
 	/** Render image based reflections (SSR, Env, SkyLight) without compute shaders */
-	void RenderImageBasedReflectionsSM4ForAllViews(FRHICommandListImmediate& RHICmdList, bool bReflectionEnv, const TRefCountPtr<IPooledRenderTarget>& DynamicBentNormalAO);
+	void RenderStandardDeferredImageBasedReflections(FRHICommandListImmediate& RHICmdList, bool bReflectionEnv, const TRefCountPtr<IPooledRenderTarget>& DynamicBentNormalAO);
 
 	bool ShouldDoReflectionEnvironment() const;
 
diff --git a/Engine/Source/Runtime/Renderer/Private/ReflectionEnvironment.cpp b/Engine/Source/Runtime/Renderer/Private/ReflectionEnvironment.cpp
--- a/Engine/Source/Runtime/Renderer/Private/ReflectionEnvironment.cpp
+++ b/Engine/Source/Runtime/Renderer/Private/ReflectionEnvironment.cpp
@@ -40,6 +40,14 @@ static TAutoConsoleVariable<int32> CVarReflectionEnvironment(
 	ECVF_Cheat | ECVF_RenderThreadSafe);
 #endif 
 
+int32 GNumCapturesBeforeUsingTiledDeferred = 0;
+static FAutoConsoleVariableRef CVarNumCapturesBeforeUsingTiledDeferred(
+	TEXT("r.TiledReflectionEnvironmentMinimumCount"),
+	GNumCapturesBeforeUsingTiledDeferred,
+	TEXT("Number of applicable reflection captures that must be in the scene before switching to tiled deferred."),
+	ECVF_RenderThreadSafe
+	);
+
 // to avoid having direct access from many places
 static int GetReflectionEnvironmentCVar()
 {
@@ -151,6 +159,7 @@ private:
 struct FReflectionCaptureSortData
 {
 	uint32 Guid;
+	int32 CaptureIndex;
 	FVector4 PositionAndRadius;
 	FVector4 CaptureProperties;
 	FMatrix BoxTransform;
@@ -272,6 +281,7 @@ public:
 
 				FReflectionCaptureSortData NewSortEntry;
 
+				NewSortEntry.CaptureIndex = CubemapIndex;
 				NewSortEntry.SM4FullHDRCubemap = NULL;
 				NewSortEntry.Guid = CurrentCapture->Guid;
 				NewSortEntry.PositionAndRadius = FVector4(CurrentCapture->Position, CurrentCapture->InfluenceRadius);
@@ -553,7 +563,9 @@ public:
 		CaptureProperties.Bind(Initializer.ParameterMap, TEXT("CaptureProperties"));
 		CaptureBoxTransform.Bind(Initializer.ParameterMap, TEXT("CaptureBoxTransform"));
 		CaptureBoxScales.Bind(Initializer.ParameterMap, TEXT("CaptureBoxScales"));
+		CaptureArrayIndex.Bind(Initializer.ParameterMap, TEXT("CaptureArrayIndex"));
 		ReflectionEnvironmentColorTexture.Bind(Initializer.ParameterMap, TEXT("ReflectionEnvironmentColorTexture"));
+		ReflectionEnvironmentColorTextureArray.Bind(Initializer.ParameterMap, TEXT("ReflectionEnvironmentColorTextureArray"));
 		ReflectionEnvironmentColorSampler.Bind(Initializer.ParameterMap, TEXT("ReflectionEnvironmentColorSampler"));
 		DeferredParameters.Bind(Initializer.ParameterMap);
 	}
@@ -563,7 +575,22 @@ public:
 		const FPixelShaderRHIParamRef ShaderRHI = GetPixelShader();
 		FGlobalShader::SetParameters(RHICmdList, ShaderRHI, View);
 
-		SetTextureParameter(RHICmdList, ShaderRHI, ReflectionEnvironmentColorTexture, ReflectionEnvironmentColorSampler, TStaticSamplerState<SF_Trilinear,AM_Clamp,AM_Clamp,AM_Clamp>::GetRHI(), SortData.SM4FullHDRCubemap->TextureRHI);
+		if (GRHIFeatureLevel >= ERHIFeatureLevel::SM5)
+		{
+			FScene* Scene = (FScene*)View.Family->Scene;
+
+			check(Scene->ReflectionSceneData.CubemapArray.IsValid());
+			check(Scene->ReflectionSceneData.CubemapArray.GetRenderTarget().IsValid());
+
+			FSceneRenderTargetItem& CubemapArray = Scene->ReflectionSceneData.CubemapArray.GetRenderTarget();
+
+			SetTextureParameter(RHICmdList, ShaderRHI, ReflectionEnvironmentColorTextureArray, ReflectionEnvironmentColorSampler, TStaticSamplerState<SF_Trilinear,AM_Clamp,AM_Clamp,AM_Clamp>::GetRHI(), CubemapArray.ShaderResourceTexture);
+			SetShaderValue(RHICmdList, ShaderRHI, CaptureArrayIndex, SortData.CaptureIndex);
+		}
+		else
+		{
+			SetTextureParameter(RHICmdList, ShaderRHI, ReflectionEnvironmentColorTexture, ReflectionEnvironmentColorSampler, TStaticSamplerState<SF_Trilinear,AM_Clamp,AM_Clamp,AM_Clamp>::GetRHI(), SortData.SM4FullHDRCubemap->TextureRHI);
+		}
 
 		DeferredParameters.Set(RHICmdList, ShaderRHI, View);
 		SetShaderValue(RHICmdList, ShaderRHI, CapturePositionAndRadius, SortData.PositionAndRadius);
@@ -580,7 +607,9 @@ public:
 		Ar << CaptureProperties;
 		Ar << CaptureBoxTransform;
 		Ar << CaptureBoxScales;
+		Ar << CaptureArrayIndex;
 		Ar << ReflectionEnvironmentColorTexture;
+		Ar << ReflectionEnvironmentColorTextureArray;
 		Ar << ReflectionEnvironmentColorSampler;
 		Ar << DeferredParameters;
 		return bShaderHasOutdatedParameters;
@@ -592,7 +621,9 @@ private:
 	FShaderParameter CaptureProperties;
 	FShaderParameter CaptureBoxTransform;
 	FShaderParameter CaptureBoxScales;
+	FShaderParameter CaptureArrayIndex;
 	FShaderResourceParameter ReflectionEnvironmentColorTexture;
+	FShaderResourceParameter ReflectionEnvironmentColorTextureArray;
 	FShaderResourceParameter ReflectionEnvironmentColorSampler;
 	FDeferredPixelShaderParameters DeferredParameters;
 };
@@ -646,7 +677,7 @@ bool FDeferredShadingSceneRenderer::ShouldDoReflectionEnvironment() const
 		&& ViewFamily.EngineShowFlags.ReflectionEnvironment;
 }
 
-void FDeferredShadingSceneRenderer::RenderImageBasedReflectionsSM5ForAllViews(FRHICommandListImmediate& RHICmdList, const TRefCountPtr<IPooledRenderTarget>& DynamicBentNormalAO)
+void FDeferredShadingSceneRenderer::RenderTiledDeferredImageBasedReflections(FRHICommandListImmediate& RHICmdList, const TRefCountPtr<IPooledRenderTarget>& DynamicBentNormalAO)
 {
 	const uint32 bUseLightmaps = CVarDiffuseFromCaptures.GetValueOnRenderThread() == 0;
 
@@ -707,7 +738,7 @@ void FDeferredShadingSceneRenderer::RenderImageBasedReflectionsSM5ForAllViews(FR
 	check(GSceneRenderTargets.GetSceneColor());
 }
 
-void FDeferredShadingSceneRenderer::RenderImageBasedReflectionsSM4ForAllViews(FRHICommandListImmediate& RHICmdList, bool bReflectionEnv, const TRefCountPtr<IPooledRenderTarget>& DynamicBentNormalAO)
+void FDeferredShadingSceneRenderer::RenderStandardDeferredImageBasedReflections(FRHICommandListImmediate& RHICmdList, bool bReflectionEnv, const TRefCountPtr<IPooledRenderTarget>& DynamicBentNormalAO)
 {
 	const bool bSkyLight = Scene->SkyLight
 		&& Scene->SkyLight->ProcessedTexture
@@ -727,6 +758,14 @@ void FDeferredShadingSceneRenderer::RenderImageBasedReflectionsSM4ForAllViews(FR
 			FReflectionCaptureProxy* CurrentCapture = Scene->ReflectionSceneData.RegisteredReflectionCaptures[ReflectionProxyIndex];
 			FReflectionCaptureSortData NewSortEntry;
 
+			NewSortEntry.CaptureIndex = -1;
+
+			if (GRHIFeatureLevel >= ERHIFeatureLevel::SM5)
+			{
+				const FCaptureComponentSceneState* ComponentStatePtr = Scene->ReflectionSceneData.AllocatedReflectionCaptureState.Find(CurrentCapture->Component);
+				NewSortEntry.CaptureIndex = ComponentStatePtr ? ComponentStatePtr->CaptureIndex : -1;
+			}
+			
 			NewSortEntry.SM4FullHDRCubemap = CurrentCapture->SM4FullHDRCubemap;
 			NewSortEntry.Guid = CurrentCapture->Guid;
 			NewSortEntry.PositionAndRadius = FVector4(CurrentCapture->Position, CurrentCapture->InfluenceRadius);
@@ -756,7 +795,7 @@ void FDeferredShadingSceneRenderer::RenderImageBasedReflectionsSM4ForAllViews(FR
 		SortData.Sort();
 	}
 
-	// In SM4 use standard deferred shading to composite reflection capture contribution
+	// Use standard deferred shading to composite reflection capture contribution
 	for (int32 ViewIndex = 0, Num = Views.Num(); ViewIndex < Num; ViewIndex++)
 	{
 		FViewInfo& View = Views[ViewIndex];
@@ -802,7 +841,7 @@ void FDeferredShadingSceneRenderer::RenderImageBasedReflectionsSM4ForAllViews(FR
 			{
 				const FReflectionCaptureSortData& ReflectionCapture = SortData[ReflectionCaptureIndex];
 
-				if (ReflectionCapture.SM4FullHDRCubemap)
+				if (GRHIFeatureLevel >= ERHIFeatureLevel::SM5 || ReflectionCapture.SM4FullHDRCubemap)
 				{
 					const FSphere LightBounds(ReflectionCapture.PositionAndRadius, ReflectionCapture.PositionAndRadius.W);
 
@@ -905,6 +944,12 @@ void FDeferredShadingSceneRenderer::RenderImageBasedReflectionsSM4ForAllViews(FR
 	}
 }
 
+bool ShouldUseTiledDeferredReflectionEnvironment(FScene* Scene)
+{
+	const ERHIFeatureLevel::Type FeatureLevel = Scene->GetFeatureLevel();
+	return FeatureLevel >= ERHIFeatureLevel::SM5 && Scene->ReflectionSceneData.RegisteredReflectionCaptures.Num() > GNumCapturesBeforeUsingTiledDeferred;
+}
+
 void FDeferredShadingSceneRenderer::RenderDeferredReflections(FRHICommandListImmediate& RHICmdList, const TRefCountPtr<IPooledRenderTarget>& DynamicBentNormalAO)
 {
 	if (IsSimpleDynamicLightingEnabled() || ViewFamily.EngineShowFlags.VisualizeLightCulling)
@@ -927,19 +972,17 @@ void FDeferredShadingSceneRenderer::RenderDeferredReflections(FRHICommandListImm
 	}
 	else
 	{
-		const ERHIFeatureLevel::Type FeatureLevel = Scene->GetFeatureLevel();
-
 		const bool bReflectionEnv = ShouldDoReflectionEnvironment();
-		const bool bReflectionsWithCompute = (FeatureLevel >= ERHIFeatureLevel::SM5) && bReflectionEnv && Scene->ReflectionSceneData.CubemapArray.IsValid();
+		const bool bTiledDeferredReflections = ShouldUseTiledDeferredReflectionEnvironment(Scene) && bReflectionEnv && Scene->ReflectionSceneData.CubemapArray.IsValid();
 
-		if (bReflectionsWithCompute)
+		if (bTiledDeferredReflections)
 		{
 			check(bReflectionEnv);
-			RenderImageBasedReflectionsSM5ForAllViews(RHICmdList, DynamicBentNormalAO);
+			RenderTiledDeferredImageBasedReflections(RHICmdList, DynamicBentNormalAO);
 		}
 		else
 		{
-			RenderImageBasedReflectionsSM4ForAllViews(RHICmdList, bReflectionEnv, DynamicBentNormalAO);
+			RenderStandardDeferredImageBasedReflections(RHICmdList, bReflectionEnv, DynamicBentNormalAO);
 		}
 	}
 }
\ No newline at end of file
