Fixed to match code standard

diff --git a/Engine/Source/Runtime/Renderer/Private/PostProcess/SceneRenderTargets.h b/Engine/Source/Runtime/Renderer/Private/PostProcess/SceneRenderTargets.h
--- a/Engine/Source/Runtime/Renderer/Private/PostProcess/SceneRenderTargets.h
+++ b/Engine/Source/Runtime/Renderer/Private/PostProcess/SceneRenderTargets.h
@@ -442,8 +442,10 @@ public:
 	TRefCountPtr<IPooledRenderTarget> ScreenSpaceAO;
 	// used by the CustomDepth material feature, is allocated on demand or if r.CustomDepth is 2
 	TRefCountPtr<IPooledRenderTarget> CustomDepth;
-	// Render target for per-object shadow depths. (Also the optional in case this RHI requires a color render target)
-	TRefCountPtr<IPooledRenderTarget> ShadowDepthZ, OptionalShadowDepthColor;
+	// Render target for per-object shadow depths.
+	TRefCountPtr<IPooledRenderTarget> ShadowDepthZ;
+	// optional in case this RHI requires a color render target
+	TRefCountPtr<IPooledRenderTarget> OptionalShadowDepthColor;
 	// Cache of preshadow depths
 	//@todo - this should go in FScene
 	TRefCountPtr<IPooledRenderTarget> PreShadowCacheDepthZ;
