Fix for rendering reflection captures before they have been captured with the standard deferred path

diff --git a/Engine/Source/Runtime/Renderer/Private/ReflectionEnvironment.cpp b/Engine/Source/Runtime/Renderer/Private/ReflectionEnvironment.cpp
--- a/Engine/Source/Runtime/Renderer/Private/ReflectionEnvironment.cpp
+++ b/Engine/Source/Runtime/Renderer/Private/ReflectionEnvironment.cpp
@@ -972,12 +972,16 @@ void FDeferredShadingSceneRenderer::RenderDeferredReflections(FRHICommandListImm
 	}
 	else
 	{
-		const bool bReflectionEnv = ShouldDoReflectionEnvironment();
-		const bool bTiledDeferredReflections = ShouldUseTiledDeferredReflectionEnvironment(Scene) && bReflectionEnv && Scene->ReflectionSceneData.CubemapArray.IsValid();
+		bool bReflectionEnv = ShouldDoReflectionEnvironment();
+		const bool bWantsTiledReflections = ShouldUseTiledDeferredReflectionEnvironment(Scene);
 
-		if (bTiledDeferredReflections)
+		if (bWantsTiledReflections && !Scene->ReflectionSceneData.CubemapArray.IsValid())
+		{
+			bReflectionEnv = false;
+		}
+
+		if (bWantsTiledReflections && bReflectionEnv)
 		{
-			check(bReflectionEnv);
 			RenderTiledDeferredImageBasedReflections(RHICmdList, DynamicBentNormalAO);
 		}
 		else
