SM5 reflection environment requires scene cubemap array to be valid regardless of deferred path

diff --git a/Engine/Source/Runtime/Renderer/Private/ReflectionEnvironment.cpp b/Engine/Source/Runtime/Renderer/Private/ReflectionEnvironment.cpp
--- a/Engine/Source/Runtime/Renderer/Private/ReflectionEnvironment.cpp
+++ b/Engine/Source/Runtime/Renderer/Private/ReflectionEnvironment.cpp
@@ -674,7 +674,8 @@ bool FDeferredShadingSceneRenderer::ShouldDoReflectionEnvironment() const
 
 	return IsReflectionEnvironmentAvailable(FeatureLevel)
 		&& Scene->ReflectionSceneData.RegisteredReflectionCaptures.Num()
-		&& ViewFamily.EngineShowFlags.ReflectionEnvironment;
+		&& ViewFamily.EngineShowFlags.ReflectionEnvironment
+		&& (FeatureLevel == ERHIFeatureLevel::SM4 || Scene->ReflectionSceneData.CubemapArray.IsValid());
 }
 
 void FDeferredShadingSceneRenderer::RenderTiledDeferredImageBasedReflections(FRHICommandListImmediate& RHICmdList, const TRefCountPtr<IPooledRenderTarget>& DynamicBentNormalAO)
@@ -972,21 +973,16 @@ void FDeferredShadingSceneRenderer::RenderDeferredReflections(FRHICommandListImm
 	}
 	else
 	{
-		bool bReflectionEnv = ShouldDoReflectionEnvironment();
+		const bool bReflectionEnvironment = ShouldDoReflectionEnvironment();
 		const bool bWantsTiledReflections = ShouldUseTiledDeferredReflectionEnvironment(Scene);
 
-		if (bWantsTiledReflections && !Scene->ReflectionSceneData.CubemapArray.IsValid())
-		{
-			bReflectionEnv = false;
-		}
-
-		if (bWantsTiledReflections && bReflectionEnv)
+		if (bWantsTiledReflections && bReflectionEnvironment)
 		{
 			RenderTiledDeferredImageBasedReflections(RHICmdList, DynamicBentNormalAO);
 		}
 		else
 		{
-			RenderStandardDeferredImageBasedReflections(RHICmdList, bReflectionEnv, DynamicBentNormalAO);
+			RenderStandardDeferredImageBasedReflections(RHICmdList, bReflectionEnvironment, DynamicBentNormalAO);
 		}
 	}
 }
\ No newline at end of file
