UE4 - allow draw events to work in parallel rendering, some flush tweaks

diff --git a/Engine/Source/Runtime/EmptyRHI/Private/EmptyRHI.cpp b/Engine/Source/Runtime/EmptyRHI/Private/EmptyRHI.cpp
--- a/Engine/Source/Runtime/EmptyRHI/Private/EmptyRHI.cpp
+++ b/Engine/Source/Runtime/EmptyRHI/Private/EmptyRHI.cpp
@@ -143,12 +143,12 @@ void FEmptyDynamicRHI::RHIEndScene()
 
 }
 
-void FEmptyDynamicRHI::PushEvent(const TCHAR* Name)
+void FEmptyDynamicRHI::RHIPushEvent(const TCHAR* Name)
 {
 
 }
 
-void FEmptyDynamicRHI::PopEvent()
+void FEmptyDynamicRHI::RHIPopEvent()
 {
 
 }
diff --git a/Engine/Source/Runtime/EmptyRHI/Public/EmptyRHI.h b/Engine/Source/Runtime/EmptyRHI/Public/EmptyRHI.h
--- a/Engine/Source/Runtime/EmptyRHI/Public/EmptyRHI.h
+++ b/Engine/Source/Runtime/EmptyRHI/Public/EmptyRHI.h
@@ -32,8 +32,6 @@ public:
 	// FDynamicRHI interface.
 	virtual void Init();
 	virtual void Shutdown();
-	virtual void PushEvent(const TCHAR* Name);
-	virtual void PopEvent();
 
 	// The RHI methods are defined as virtual functions in URenderHardwareInterface.
 	#define DEFINE_RHIMETHOD(Type,Name,ParameterTypesAndNames,ParameterNames,ReturnStatement,NullImplementation) virtual Type RHI##Name ParameterTypesAndNames
diff --git a/Engine/Source/Runtime/Engine/Classes/Engine/NiagaraEffectRenderer.h b/Engine/Source/Runtime/Engine/Classes/Engine/NiagaraEffectRenderer.h
--- a/Engine/Source/Runtime/Engine/Classes/Engine/NiagaraEffectRenderer.h
+++ b/Engine/Source/Runtime/Engine/Classes/Engine/NiagaraEffectRenderer.h
@@ -266,8 +266,6 @@ public:
 		if (DynamicVertexAllocation.IsValid()
 			&& (bIsWireframe || !PDI->IsMaterialIgnored(MaterialRenderProxy, View->GetFeatureLevel())))
 		{
-			SCOPED_DRAW_EVENT(NiagaraRender, DEC_SCENE_ITEMS);
-
 			FMeshBatch MeshBatch;
 			MeshBatch.VertexFactory = &VertexFactory;
 			MeshBatch.CastShadow = SceneProxy->CastsDynamicShadow();
@@ -676,8 +674,6 @@ public:
 		if (DynamicVertexAllocation.IsValid()
 			&& (bIsWireframe || !PDI->IsMaterialIgnored(MaterialRenderProxy, View->GetFeatureLevel())))
 		{
-			SCOPED_DRAW_EVENT(NiagaraRender, DEC_SCENE_ITEMS);
-
 			FMeshBatch MeshBatch;
 			MeshBatch.VertexFactory = &VertexFactory;
 			MeshBatch.CastShadow = SceneProxy->CastsDynamicShadow();
diff --git a/Engine/Source/Runtime/Engine/Private/GPUSkinCache.cpp b/Engine/Source/Runtime/Engine/Private/GPUSkinCache.cpp
--- a/Engine/Source/Runtime/Engine/Private/GPUSkinCache.cpp
+++ b/Engine/Source/Runtime/Engine/Private/GPUSkinCache.cpp
@@ -479,7 +479,7 @@ void FGPUSkinCache::DispatchSkinCacheProcess(FRHICommandListImmediate& RHICmdLis
 
 	uint32 VertexCountAlign64 = (VertexCount + 63) / 64;
 
-	SCOPED_DRAW_EVENT(SkinCacheDispatch, DEC_SKEL_MESH);
+	SCOPED_DRAW_EVENT(RHICmdList, SkinCacheDispatch, DEC_SKEL_MESH);
 
 	if (bUseExtraBoneInfluences)
 	{
diff --git a/Engine/Source/Runtime/Engine/Private/GPUSort.cpp b/Engine/Source/Runtime/Engine/Private/GPUSort.cpp
--- a/Engine/Source/Runtime/Engine/Private/GPUSort.cpp
+++ b/Engine/Source/Runtime/Engine/Private/GPUSort.cpp
@@ -693,7 +693,7 @@ int32 SortGPUBuffers(FRHICommandListImmediate& RHICmdList, FGPUSortBuffers SortB
 	const auto FeatureLevel = GRHIFeatureLevel;
 	check(FeatureLevel == ERHIFeatureLevel::SM5);
 
-	SCOPED_DRAW_EVENTF(SortGPU, DEC_PARTICLE, TEXT("SortGPU_%d"), Count);
+	SCOPED_DRAW_EVENTF(RHICmdList, SortGPU, DEC_PARTICLE, TEXT("SortGPU_%d"), Count);
 
 	// Determine how many tiles need to be sorted.
 	const int32 TileCount = Count / TILE_SIZE;
diff --git a/Engine/Source/Runtime/Engine/Private/Particles/ParticleCurveTexture.cpp b/Engine/Source/Runtime/Engine/Private/Particles/ParticleCurveTexture.cpp
--- a/Engine/Source/Runtime/Engine/Private/Particles/ParticleCurveTexture.cpp
+++ b/Engine/Source/Runtime/Engine/Private/Particles/ParticleCurveTexture.cpp
@@ -184,7 +184,7 @@ static void InjectCurves(
 
 	check( IsInRenderingThread() );
 
-	SCOPED_DRAW_EVENT(InjectParticleCurves, DEC_PARTICLE);
+	SCOPED_DRAW_EVENT(RHICmdList, InjectParticleCurves, DEC_PARTICLE);
 
 	FVertexBufferRHIParamRef ScratchVertexBufferRHI = GParticleScratchVertexBuffer.VertexBufferRHI;
 
diff --git a/Engine/Source/Runtime/Engine/Private/Particles/ParticleGpuSimulation.cpp b/Engine/Source/Runtime/Engine/Private/Particles/ParticleGpuSimulation.cpp
--- a/Engine/Source/Runtime/Engine/Private/Particles/ParticleGpuSimulation.cpp
+++ b/Engine/Source/Runtime/Engine/Private/Particles/ParticleGpuSimulation.cpp
@@ -1279,7 +1279,7 @@ void ExecuteSimulationCommands(
  */
 void ClearTiles(FRHICommandList& RHICmdList, ERHIFeatureLevel::Type FeatureLevel, const TArray<uint32>& Tiles)
 {
-	SCOPED_DRAW_EVENT(ClearTiles, DEC_PARTICLE);
+	SCOPED_DRAW_EVENT(RHICmdList, ClearTiles, DEC_PARTICLE);
 
 	const int32 MaxTilesPerDrawCallUnaligned = GParticleScratchVertexBufferSize / sizeof(FVector2D);
 	const int32 MaxTilesPerDrawCall = MaxTilesPerDrawCallUnaligned & (~(TILES_PER_INSTANCE-1));
@@ -1745,7 +1745,7 @@ static void VisualizeGPUSimulation(
 	)
 {
 	check(IsInRenderingThread());
-	SCOPED_DRAW_EVENT(ParticleSimDebugDraw, DEC_PARTICLE);
+	SCOPED_DRAW_EVENT(RHICmdList, ParticleSimDebugDraw, DEC_PARTICLE);
 
 	// Some constants for laying out the debug view.
 	const float DisplaySizeX = 256.0f;
@@ -4265,7 +4265,7 @@ void FFXSystem::SimulateGPUParticles(
 	// Simulate particles in all active tiles.
 	if ( SimulationCommands.Num() )
 	{
-		SCOPED_DRAW_EVENT(ParticleSimulation, DEC_PARTICLE);
+		SCOPED_DRAW_EVENT(RHICmdList, ParticleSimulation, DEC_PARTICLE);
 
 		if (Phase == EParticleSimulatePhase::Collision && CollisionView)
 		{
@@ -4308,7 +4308,7 @@ void FFXSystem::SimulateGPUParticles(
 	// Inject any new particles that have spawned into the simulation.
 	if (NewParticles.Num())
 	{
-		SCOPED_DRAW_EVENT(ParticleInjection, DEC_PARTICLE);
+		SCOPED_DRAW_EVENT(RHICmdList, ParticleInjection, DEC_PARTICLE);
 		// Set render targets.
 		FTextureRHIParamRef InjectRenderTargets[4] =
 		{
diff --git a/Engine/Source/Runtime/Engine/Private/Particles/ParticleSortingGPU.cpp b/Engine/Source/Runtime/Engine/Private/Particles/ParticleSortingGPU.cpp
--- a/Engine/Source/Runtime/Engine/Private/Particles/ParticleSortingGPU.cpp
+++ b/Engine/Source/Runtime/Engine/Private/Particles/ParticleSortingGPU.cpp
@@ -183,7 +183,7 @@ static int32 GenerateParticleSortKeys(
 	const TArray<FParticleSimulationSortInfo>& SimulationsToSort
 	)
 {
-	SCOPED_DRAW_EVENT(ParticleSortKeyGen, DEC_PARTICLE);
+	SCOPED_DRAW_EVENT(RHICmdList, ParticleSortKeyGen, DEC_PARTICLE);
 
 	// MOBILEPREVIEWTODO: Proper value for this
 	const auto FeatureLevel = GRHIFeatureLevel;
@@ -317,7 +317,7 @@ int32 SortParticlesGPU(
 	const TArray<FParticleSimulationSortInfo>& SimulationsToSort
 	)
 {
-	SCOPED_DRAW_EVENTF(ParticleSort, DEC_PARTICLE, TEXT("ParticleSort_%d"), SimulationsToSort.Num());
+	SCOPED_DRAW_EVENTF(RHICmdList, ParticleSort, DEC_PARTICLE, TEXT("ParticleSort_%d"), SimulationsToSort.Num());
 
 	check(GRHIFeatureLevel == ERHIFeatureLevel::SM5);
 
diff --git a/Engine/Source/Runtime/Engine/Private/PrimitiveDrawingUtils.cpp b/Engine/Source/Runtime/Engine/Private/PrimitiveDrawingUtils.cpp
--- a/Engine/Source/Runtime/Engine/Private/PrimitiveDrawingUtils.cpp
+++ b/Engine/Source/Runtime/Engine/Private/PrimitiveDrawingUtils.cpp
@@ -10,30 +10,26 @@
 #include "RendererInterface.h"
 
 /** Emits draw events for a given FMeshBatch and the FPrimitiveSceneProxy corresponding to that mesh element. */
-void EmitMeshDrawEvents(const FPrimitiveSceneProxy* PrimitiveSceneProxy, const FMeshBatch& Mesh)
-{
 #if !(UE_BUILD_SHIPPING || UE_BUILD_TEST)
-	extern bool GShowMaterialDrawEvents;
-	if ( GShowMaterialDrawEvents )
+void EmitMeshDrawEvents_Inner(FRHICommandList& RHICmdList, const FPrimitiveSceneProxy* PrimitiveSceneProxy, const FMeshBatch& Mesh)
+{
+	// Only show material name at the top level
+	// Note: this is the parent's material name, not the material instance
+	SCOPED_DRAW_EVENTF(RHICmdList, MaterialEvent, DEC_SCENE_ITEMS, *Mesh.MaterialRenderProxy->GetMaterial(PrimitiveSceneProxy->GetScene()->GetFeatureLevel())->GetFriendlyName());
+	if (PrimitiveSceneProxy)
 	{
-		// Only show material name at the top level
-		// Note: this is the parent's material name, not the material instance
-		SCOPED_DRAW_EVENTF(MaterialEvent, DEC_SCENE_ITEMS, *Mesh.MaterialRenderProxy->GetMaterial(PrimitiveSceneProxy->GetScene()->GetFeatureLevel())->GetFriendlyName());
-		if (PrimitiveSceneProxy)
+		// Show Actor, level and resource name inside the material name
+		// These are separate draw events since some platforms only allow 32 character event names (xenon)
 		{
-			// Show Actor, level and resource name inside the material name
-			// These are separate draw events since some platforms only allow 32 character event names (xenon)
-			{
-				SCOPED_CONDITIONAL_DRAW_EVENTF(LevelEvent, PrimitiveSceneProxy->GetLevelName() != NAME_None, DEC_SCENE_ITEMS, PrimitiveSceneProxy->GetLevelName().IsValid() ? *PrimitiveSceneProxy->GetLevelName().ToString() : TEXT(""));
-			}
-			{
-				SCOPED_CONDITIONAL_DRAW_EVENTF(OwnerEvent,PrimitiveSceneProxy->GetOwnerName() != NAME_None, DEC_SCENE_ITEMS, *PrimitiveSceneProxy->GetOwnerName().ToString());
-			}
-			SCOPED_CONDITIONAL_DRAW_EVENTF(ResourceEvent,PrimitiveSceneProxy->GetResourceName() != NAME_None, DEC_SCENE_ITEMS, PrimitiveSceneProxy->GetResourceName().IsValid() ? *PrimitiveSceneProxy->GetResourceName().ToString() : TEXT(""));
+			SCOPED_CONDITIONAL_DRAW_EVENTF(RHICmdList, LevelEvent, PrimitiveSceneProxy->GetLevelName() != NAME_None, DEC_SCENE_ITEMS, PrimitiveSceneProxy->GetLevelName().IsValid() ? *PrimitiveSceneProxy->GetLevelName().ToString() : TEXT(""));
+		}
+		{
+			SCOPED_CONDITIONAL_DRAW_EVENTF(RHICmdList, OwnerEvent,PrimitiveSceneProxy->GetOwnerName() != NAME_None, DEC_SCENE_ITEMS, *PrimitiveSceneProxy->GetOwnerName().ToString());
 		}
+		SCOPED_CONDITIONAL_DRAW_EVENTF(RHICmdList, ResourceEvent,PrimitiveSceneProxy->GetResourceName() != NAME_None, DEC_SCENE_ITEMS, PrimitiveSceneProxy->GetResourceName().IsValid() ? *PrimitiveSceneProxy->GetResourceName().ToString() : TEXT(""));
 	}
-#endif
 }
+#endif
 
 void DrawPlane10x10(class FPrimitiveDrawInterface* PDI,const FMatrix& ObjectToWorld, float Radii, FVector2D UVMin, FVector2D UVMax, const FMaterialRenderProxy* MaterialRenderProxy, uint8 DepthPriorityGroup)
 {
diff --git a/Engine/Source/Runtime/Engine/Private/SceneUtils.cpp b/Engine/Source/Runtime/Engine/Private/SceneUtils.cpp
--- a/Engine/Source/Runtime/Engine/Private/SceneUtils.cpp
+++ b/Engine/Source/Runtime/Engine/Private/SceneUtils.cpp
@@ -6,24 +6,23 @@
 
 #if WANTS_DRAW_MESH_EVENTS && PLATFORM_SUPPORTS_DRAW_MESH_EVENTS
 
-bool FDrawEvent::IsInRenderingThread_Internal()
+void FDrawEvent::Start(FRHICommandList& InRHICmdList, const FColor& Color, const TCHAR* Fmt, ...)
 {
-	return IsInRenderingThread();
-}
-
-ENGINE_API void FDrawEvent::Start(const FColor& Color, const TCHAR* Fmt, ...)
-{
-	if (GRHICommandList.Bypass()) // todo interleave draw evens into command buffers
+	check(IsInParallelRenderingThread() || IsInRHIThread());
 	{
-		check(IsInRenderingThread_Internal());
 		va_list ptr;
 		va_start(ptr, Fmt);
 		TCHAR TempStr[256];
 		// Build the string in the temp buffer
 		FCString::GetVarArgs(TempStr, ARRAY_COUNT(TempStr), ARRAY_COUNT(TempStr) - 1, Fmt, ptr);
-		GDynamicRHI->PushEvent(TempStr);
-		bDrawEventHasBeenEmitted = true;
+		InRHICmdList.PushEvent(TempStr);
+		RHICmdList = &InRHICmdList;
 	}
 }
 
+void FDrawEvent::Stop()
+{
+	RHICmdList->PopEvent();
+}
+
 #endif // WANTS_DRAW_MESH_EVENTS && PLATFORM_SUPPORTS_DRAW_MESH_EVENTS
\ No newline at end of file
diff --git a/Engine/Source/Runtime/Engine/Private/UserInterface/Canvas.cpp b/Engine/Source/Runtime/Engine/Private/UserInterface/Canvas.cpp
--- a/Engine/Source/Runtime/Engine/Private/UserInterface/Canvas.cpp
+++ b/Engine/Source/Runtime/Engine/Private/UserInterface/Canvas.cpp
@@ -552,7 +552,7 @@ void FCanvas::Flush_RenderThread(FRHICommandListImmediate& RHICmdList, bool bFor
 	// sort the array of FCanvasSortElement entries so that higher sort keys render first (back-to-front)
 	SortedElements.Sort(FCompareFCanvasSortElement());
 
-	SCOPED_DRAW_EVENT(CanvasFlush, DEC_SCENE_ITEMS);
+	SCOPED_DRAW_EVENT(RHICmdList, CanvasFlush, DEC_SCENE_ITEMS);
 	const FTexture2DRHIRef& RenderTargetTexture = RenderTarget->GetRenderTargetTexture();
 
 	check(IsValidRef(RenderTargetTexture));
@@ -648,7 +648,7 @@ void FCanvas::Flush_GameThread(bool bForce)
 		CanvasFlushSetupCommand,
 		FCanvasFlushParameters,Parameters,FlushParameters,
 	{
-		SCOPED_DRAW_EVENT(CanvasFlush, DEC_SCENE_ITEMS);
+		SCOPED_DRAW_EVENT(RHICmdList, CanvasFlush, DEC_SCENE_ITEMS);
 
 		// Set the RHI render target.
 		::SetRenderTarget(RHICmdList, Parameters.CanvasRenderTarget->GetRenderTargetTexture(), FTextureRHIRef());
diff --git a/Engine/Source/Runtime/Engine/Private/VectorField.cpp b/Engine/Source/Runtime/Engine/Private/VectorField.cpp
--- a/Engine/Source/Runtime/Engine/Private/VectorField.cpp
+++ b/Engine/Source/Runtime/Engine/Private/VectorField.cpp
@@ -800,7 +800,7 @@ public:
 
 		if (GetFeatureLevel() == ERHIFeatureLevel::SM5 && AnimatedVectorField && AnimatedVectorField->Texture && AnimatedVectorField->Texture->Resource)
 		{
-			SCOPED_DRAW_EVENT(AnimateVectorField, DEC_PARTICLE);
+			SCOPED_DRAW_EVENT(RHICmdList, AnimateVectorField, DEC_PARTICLE);
 
 			// Move frame time forward.
 			FrameTime += AnimatedVectorField->FramesPerSecond * DeltaSeconds;
diff --git a/Engine/Source/Runtime/Engine/Public/SceneManagement.h b/Engine/Source/Runtime/Engine/Public/SceneManagement.h
--- a/Engine/Source/Runtime/Engine/Public/SceneManagement.h
+++ b/Engine/Source/Runtime/Engine/Public/SceneManagement.h
@@ -1719,8 +1719,25 @@ extern ENGINE_API EVertexColorViewMode::Type GVertexColorViewMode;
  */
 extern ENGINE_API bool IsRichView(const FSceneViewFamily& ViewFamily);
 
+#if !(UE_BUILD_SHIPPING || UE_BUILD_TEST)
+	/**
+	 * true if we debug material names with SCOPED_DRAW_EVENT.
+	 * Toggle with "ShowMaterialDrawEvents" console command.
+	 */
+	extern ENGINE_API bool GShowMaterialDrawEvents;
+	extern ENGINE_API void EmitMeshDrawEvents_Inner(FRHICommandList& RHICmdList, const class FPrimitiveSceneProxy* PrimitiveSceneProxy, const struct FMeshBatch& Mesh);
+#endif
+
 /** Emits draw events for a given FMeshBatch and the PrimitiveSceneProxy corresponding to that mesh element. */
-extern ENGINE_API void EmitMeshDrawEvents(const class FPrimitiveSceneProxy* PrimitiveSceneProxy, const struct FMeshBatch& Mesh);
+FORCEINLINE void EmitMeshDrawEvents(FRHICommandList& RHICmdList, const class FPrimitiveSceneProxy* PrimitiveSceneProxy, const struct FMeshBatch& Mesh)
+{
+#if !(UE_BUILD_SHIPPING || UE_BUILD_TEST)
+	if ( GShowMaterialDrawEvents )
+	{
+		EmitMeshDrawEvents_Inner(RHICmdList, PrimitiveSceneProxy, Mesh);
+	}
+#endif
+}
 
 /**
  * Draws a mesh, modifying the material which is used depending on the view's show flags.
diff --git a/Engine/Source/Runtime/Engine/Public/SceneUtils.h b/Engine/Source/Runtime/Engine/Public/SceneUtils.h
--- a/Engine/Source/Runtime/Engine/Public/SceneUtils.h
+++ b/Engine/Source/Runtime/Engine/Public/SceneUtils.h
@@ -39,42 +39,39 @@ inline void appSetCounterValue(const TCHAR* CounterName, float Value) {}
 	 */
 	struct ENGINE_API FDrawEvent
 	{
-		/** Whether a draw event has been emitted or not. */
-		bool bDrawEventHasBeenEmitted;
+		/** Cmdlist to push onto. */
+		class FRHICommandList* RHICmdList;
 
 		/** Default constructor, initializing all member variables. */
-		FORCEINLINE FDrawEvent() :
-			bDrawEventHasBeenEmitted( false )
+		FORCEINLINE FDrawEvent()
+			: RHICmdList(nullptr)
 		{}
 
 		/**
 		 * Terminate the event based upon scope
 		 */
-		~FDrawEvent()
+		FORCEINLINE ~FDrawEvent()
 		{
-			if (bDrawEventHasBeenEmitted)
+			if (RHICmdList)
 			{
-				GDynamicRHI->PopEvent();
+				Stop();
 			}
 		}
 
 		/**
 		 * Function for logging a PIX event with var args
 		 */
-		void CDECL Start(const FColor& Color,const TCHAR* Fmt, ...);
-
-		private:
-		/** Helper function to determine if we are currently int the render thread. This is needed to avoid using IsInRenderThread() in a public header since it comes from RenderCore */
-		bool IsInRenderingThread_Internal();
+		void CDECL Start(FRHICommandList& RHICmdList, const FColor& Color,const TCHAR* Fmt, ...);
+		void Stop();
 	};
 
 	// Macros to allow for scoping of draw events
 	#define SCOPED_DRAW_TOKENPASTE_INNER(x,y) x##y
 	#define SCOPED_DRAW_TOKENPASTE(x,y) SCOPED_DRAW_TOKENPASTE_INNER(x,y)
-	#define SCOPED_DRAW_EVENT(Name, Color) FDrawEvent SCOPED_DRAW_TOKENPASTE(Event_##Name,__LINE__); if(GEmitDrawEvents) SCOPED_DRAW_TOKENPASTE(Event_##Name,__LINE__).Start(Color, TEXT(#Name));
-	#define SCOPED_DRAW_EVENTF(Name, Color, Format, ...) FDrawEvent SCOPED_DRAW_TOKENPASTE(Event_##Name,__LINE__); if(GEmitDrawEvents) SCOPED_DRAW_TOKENPASTE(Event_##Name,__LINE__).Start(Color, Format, ##__VA_ARGS__);
-	#define SCOPED_CONDITIONAL_DRAW_EVENT(Name, Condition, Color) FDrawEvent SCOPED_DRAW_TOKENPASTE(Event_##Name,__LINE__); if(GEmitDrawEvents && (Condition)) SCOPED_DRAW_TOKENPASTE(Event_##Name,__LINE__).Start(Color, TEXT(#Name));
-	#define SCOPED_CONDITIONAL_DRAW_EVENTF(Name, Condition, Color, Format, ...) FDrawEvent SCOPED_DRAW_TOKENPASTE(Event_##Name,__LINE__); if(GEmitDrawEvents && (Condition)) SCOPED_DRAW_TOKENPASTE(Event_##Name,__LINE__).Start(Color, Format, ##__VA_ARGS__);
+	#define SCOPED_DRAW_EVENT(RHICmdList, Name, Color) FDrawEvent SCOPED_DRAW_TOKENPASTE(Event_##Name,__LINE__); if(GEmitDrawEvents) SCOPED_DRAW_TOKENPASTE(Event_##Name,__LINE__).Start(RHICmdList, Color, TEXT(#Name));
+	#define SCOPED_DRAW_EVENTF(RHICmdList, Name, Color, Format, ...) FDrawEvent SCOPED_DRAW_TOKENPASTE(Event_##Name,__LINE__); if(GEmitDrawEvents) SCOPED_DRAW_TOKENPASTE(Event_##Name,__LINE__).Start(RHICmdList, Color, Format, ##__VA_ARGS__);
+	#define SCOPED_CONDITIONAL_DRAW_EVENT(RHICmdList, Name, Condition, Color) FDrawEvent SCOPED_DRAW_TOKENPASTE(Event_##Name,__LINE__); if(GEmitDrawEvents && (Condition)) SCOPED_DRAW_TOKENPASTE(Event_##Name,__LINE__).Start(RHICmdList, Color, TEXT(#Name));
+	#define SCOPED_CONDITIONAL_DRAW_EVENTF(RHICmdList, Name, Condition, Color, Format, ...) FDrawEvent SCOPED_DRAW_TOKENPASTE(Event_##Name,__LINE__); if(GEmitDrawEvents && (Condition)) SCOPED_DRAW_TOKENPASTE(Event_##Name,__LINE__).Start(RHICmdList, Color, Format, ##__VA_ARGS__);
 
 #else
 
diff --git a/Engine/Source/Runtime/IOS/MetalRHI/Private/MetalRHI.cpp b/Engine/Source/Runtime/IOS/MetalRHI/Private/MetalRHI.cpp
--- a/Engine/Source/Runtime/IOS/MetalRHI/Private/MetalRHI.cpp
+++ b/Engine/Source/Runtime/IOS/MetalRHI/Private/MetalRHI.cpp
@@ -186,14 +186,14 @@ void FMetalDynamicRHI::RHIEndScene()
 	FMetalManager::Get()->EndScene();
 }
 
-void FMetalDynamicRHI::PushEvent(const TCHAR* Name)
+void FMetalDynamicRHI::RHIPushEvent(const TCHAR* Name)
 {
 #if ENABLE_METAL_GPUEVENTS
 	[FMetalManager::GetContext() pushDebugGroup: [NSString stringWithTCHARString:Name]];
 #endif
 }
 
-void FMetalDynamicRHI::PopEvent()
+void FMetalDynamicRHI::RHIPopEvent()
 {
 #if ENABLE_METAL_GPUEVENTS
 	[FMetalManager::GetContext() popDebugGroup];
diff --git a/Engine/Source/Runtime/IOS/MetalRHI/Public/MetalRHI.h b/Engine/Source/Runtime/IOS/MetalRHI/Public/MetalRHI.h
--- a/Engine/Source/Runtime/IOS/MetalRHI/Public/MetalRHI.h
+++ b/Engine/Source/Runtime/IOS/MetalRHI/Public/MetalRHI.h
@@ -34,8 +34,6 @@ public:
 
 	// FDynamicRHI interface.
 	virtual void Init();
-	virtual void PushEvent(const TCHAR* Name);
-	virtual void PopEvent();
 	virtual void Shutdown() {}
 
 	// The RHI methods are defined as virtual functions in URenderHardwareInterface.
diff --git a/Engine/Source/Runtime/Launch/Private/LaunchEngineLoop.cpp b/Engine/Source/Runtime/Launch/Private/LaunchEngineLoop.cpp
--- a/Engine/Source/Runtime/Launch/Private/LaunchEngineLoop.cpp
+++ b/Engine/Source/Runtime/Launch/Private/LaunchEngineLoop.cpp
@@ -2039,10 +2039,7 @@ void FEngineLoop::Tick()
 		{
 			GRHICommandList.LatchBypass();
 			GFrameNumberRenderThread++;
-			if (GRHICommandList.Bypass())
-			{
-				GDynamicRHI->PushEvent(*FString::Printf(TEXT("Frame%d"),GFrameNumberRenderThread));
-			}
+			RHICmdList.PushEvent(*FString::Printf(TEXT("Frame%d"),GFrameNumberRenderThread));
 			RHICmdList.BeginFrame();
 		});
 
@@ -2215,10 +2212,7 @@ void FEngineLoop::Tick()
 			EndFrame,
 		{
 			RHICmdList.EndFrame();
-			if (GRHICommandList.Bypass())
-			{
-				GDynamicRHI->PopEvent();
-			}
+			RHICmdList.PopEvent();
 		});
 	} 
 
diff --git a/Engine/Source/Runtime/OpenGLDrv/Private/OpenGLDrv.cpp b/Engine/Source/Runtime/OpenGLDrv/Private/OpenGLDrv.cpp
--- a/Engine/Source/Runtime/OpenGLDrv/Private/OpenGLDrv.cpp
+++ b/Engine/Source/Runtime/OpenGLDrv/Private/OpenGLDrv.cpp
@@ -19,7 +19,7 @@ IMPLEMENT_MODULE(FOpenGLDynamicRHIModule, OpenGLDrv);
 /** OpenGL Logging. */
 DEFINE_LOG_CATEGORY(LogOpenGL);
 
-void FOpenGLDynamicRHI::PushEvent(const TCHAR* Name)
+void FOpenGLDynamicRHI::RHIPushEvent(const TCHAR* Name)
 {
 #if ENABLE_OPENGL_DEBUG_GROUPS
 	// @todo-mobile: Fix string conversion ASAP!
@@ -33,7 +33,7 @@ void FOpenGLGPUProfiler::PushEvent(const TCHAR* Name)
 	FGPUProfiler::PushEvent(Name);
 }
 
-void FOpenGLDynamicRHI::PopEvent()
+void FOpenGLDynamicRHI::RHIPopEvent()
 {
 #if ENABLE_OPENGL_DEBUG_GROUPS
 	FOpenGL::PopGroupMarker();
diff --git a/Engine/Source/Runtime/OpenGLDrv/Public/OpenGLDrv.h b/Engine/Source/Runtime/OpenGLDrv/Public/OpenGLDrv.h
--- a/Engine/Source/Runtime/OpenGLDrv/Public/OpenGLDrv.h
+++ b/Engine/Source/Runtime/OpenGLDrv/Public/OpenGLDrv.h
@@ -302,9 +302,6 @@ public:
 	virtual void Init();
 	virtual void Shutdown();
 
-	virtual void PushEvent(const TCHAR* Name);
-	virtual void PopEvent();
-
 	#define DEFINE_RHIMETHOD(Type,Name,ParameterTypesAndNames,ParameterNames,ReturnStatement,NullImplementation) virtual Type RHI##Name ParameterTypesAndNames
 	#include "RHIMethods.h"
 	#undef DEFINE_RHIMETHOD
diff --git a/Engine/Source/Runtime/RHI/Private/RHICommandList.cpp b/Engine/Source/Runtime/RHI/Private/RHICommandList.cpp
--- a/Engine/Source/Runtime/RHI/Private/RHICommandList.cpp
+++ b/Engine/Source/Runtime/RHI/Private/RHICommandList.cpp
@@ -657,7 +657,7 @@ DECLARE_CYCLE_STAT(TEXT("Explicit wait for RHI thread"), STAT_ExplicitWaitRHIThr
 DECLARE_CYCLE_STAT(TEXT("Deep spin for stray resource init."), STAT_SpinWaitRHIThread, STATGROUP_RHICMDLIST);
 void FRHICommandListBase::WaitForRHIThreadTasks()
 {
-	check(IsImmediate() && IsInRenderingThread() && GRHIThread);
+	check(IsImmediate() && IsInRenderingThread());
 	if (RHIThreadTask.GetReference() && RHIThreadTask->IsComplete())
 	{
 		RHIThreadTask = nullptr;
diff --git a/Engine/Source/Runtime/RHI/Public/DynamicRHI.h b/Engine/Source/Runtime/RHI/Public/DynamicRHI.h
--- a/Engine/Source/Runtime/RHI/Public/DynamicRHI.h
+++ b/Engine/Source/Runtime/RHI/Public/DynamicRHI.h
@@ -31,9 +31,6 @@ public:
 	/** Shutdown the RHI; handle shutdown and resource destruction before the RHI's actual destructor is called (so that all resources of the RHI are still available for shutdown). */
 	virtual void Shutdown() = 0;
 
-	virtual void PushEvent(const TCHAR* Name) {}
-	virtual void PopEvent() {}
-
 	#define DEFINE_RHIMETHOD(Type,Name,ParameterTypesAndNames,ParameterNames,ReturnStatement,NullImplementation) virtual Type RHI##Name ParameterTypesAndNames = 0
 	#include "RHIMethods.h"
 	#undef DEFINE_RHIMETHOD
diff --git a/Engine/Source/Runtime/RHI/Public/RHICommandList.h b/Engine/Source/Runtime/RHI/Public/RHICommandList.h
--- a/Engine/Source/Runtime/RHI/Public/RHICommandList.h
+++ b/Engine/Source/Runtime/RHI/Public/RHICommandList.h
@@ -1062,6 +1062,21 @@ struct FRHICommandEndDrawingViewport : public FRHICommand<FRHICommandEndDrawingV
 };
 #endif
 
+struct FRHICommandPushEvent : public FRHICommand<FRHICommandPushEvent>
+{
+	const TCHAR *Name;
+
+	FORCEINLINE_DEBUGGABLE FRHICommandPushEvent(const TCHAR *InName)
+		: Name(InName)
+	{
+	}
+	RHI_API void Execute();
+};
+
+struct FRHICommandPopEvent : public FRHICommand<FRHICommandPopEvent>
+{
+	RHI_API void Execute();
+};
 
 class RHI_API FRHICommandList : public FRHICommandListBase
 {
@@ -1685,6 +1700,30 @@ public:
 	void BeginFrame();
 	void EndFrame();
 #endif
+
+	FORCEINLINE_DEBUGGABLE void PushEvent(const TCHAR* Name)
+	{
+		if (Bypass())
+		{
+			PushEvent_Internal(Name);
+			return;
+		}
+		int32 Len = FCString::Strlen(Name) + 1;
+		TCHAR* NameCopy  = (TCHAR*)Alloc(Len * (int32)sizeof(TCHAR), (int32)sizeof(TCHAR));
+		FCString::Strcpy(NameCopy, Len, Name);
+		new (AllocCommand<FRHICommandPushEvent>()) FRHICommandPushEvent(NameCopy);
+	}
+
+	FORCEINLINE_DEBUGGABLE void PopEvent()
+	{
+		if (Bypass())
+		{
+			PopEvent_Internal();
+			return;
+		}
+		new (AllocCommand<FRHICommandPopEvent>()) FRHICommandPopEvent();
+	}
+
 };
 
 namespace EImmediateFlushType
diff --git a/Engine/Source/Runtime/RHI/Public/RHICommandList.inl b/Engine/Source/Runtime/RHI/Public/RHICommandList.inl
--- a/Engine/Source/Runtime/RHI/Public/RHICommandList.inl
+++ b/Engine/Source/Runtime/RHI/Public/RHICommandList.inl
@@ -68,7 +68,6 @@ FORCEINLINE_DEBUGGABLE void FRHICommandListImmediate::ImmediateFlush(EImmediateF
 		break;
 	case EImmediateFlushType::FlushRHIThreadFlushResources:
 		{
-			check(GRHIThread);
 			if (HasCommands())
 			{
 				GRHICommandList.ExecuteList(*this);
diff --git a/Engine/Source/Runtime/RHI/Public/RHICommandListCommandExecutes.inl b/Engine/Source/Runtime/RHI/Public/RHICommandListCommandExecutes.inl
--- a/Engine/Source/Runtime/RHI/Public/RHICommandListCommandExecutes.inl
+++ b/Engine/Source/Runtime/RHI/Public/RHICommandListCommandExecutes.inl
@@ -394,4 +394,18 @@ void FRHICommandEndDrawingViewport::Execute()
 	RHISTAT(EndDrawingViewport);
 	INTERNAL_DECORATOR(EndDrawingViewport)(Viewport, bPresent, bLockToVsync);
 }
-#endif
\ No newline at end of file
+#endif
+
+void FRHICommandPushEvent::Execute()
+{
+	RHISTAT(PushEvent);
+	INTERNAL_DECORATOR(PushEvent)(Name);
+}
+
+void FRHICommandPopEvent::Execute()
+{
+	RHISTAT(PopEvent);
+	INTERNAL_DECORATOR(PopEvent)();
+}
+
+
diff --git a/Engine/Source/Runtime/RHI/Public/RHIMethods.h b/Engine/Source/Runtime/RHI/Public/RHIMethods.h
--- a/Engine/Source/Runtime/RHI/Public/RHIMethods.h
+++ b/Engine/Source/Runtime/RHI/Public/RHIMethods.h
@@ -2121,3 +2121,13 @@ DEFINE_RHIMETHOD_GLOBALFLUSH_0(
 	void*,GetNativeDevice,
 	return,return NULL
 	);
+
+DEFINE_RHIMETHOD_CMDLIST_1(
+	void,PushEvent,
+	const TCHAR*, Name,
+	,
+	);
+DEFINE_RHIMETHOD_CMDLIST_0(
+	void,PopEvent,
+	,
+	);
diff --git a/Engine/Source/Runtime/RenderCore/Private/RenderingThread.cpp b/Engine/Source/Runtime/RenderCore/Private/RenderingThread.cpp
--- a/Engine/Source/Runtime/RenderCore/Private/RenderingThread.cpp
+++ b/Engine/Source/Runtime/RenderCore/Private/RenderingThread.cpp
@@ -833,7 +833,12 @@ void AdvanceFrameRenderPrerequisite()
  */
 void FlushRenderingCommands()
 {
-	FlushPendingDeleteRHIResources_GameThread(); // we do this here because with modal dialogs or blueprint debugging, the engine does not tick
+	ENQUEUE_UNIQUE_RENDER_COMMAND(
+		FlushPendingDeleteRHIResources,
+	{
+		GRHICommandList.GetImmediateCommandList().ImmediateFlush(EImmediateFlushType::FlushRHIThreadFlushResources);
+	}
+	);
 
 	AdvanceFrameRenderPrerequisite();
 
diff --git a/Engine/Source/Runtime/Renderer/Private/AtmosphereRendering.cpp b/Engine/Source/Runtime/Renderer/Private/AtmosphereRendering.cpp
--- a/Engine/Source/Runtime/Renderer/Private/AtmosphereRendering.cpp
+++ b/Engine/Source/Runtime/Renderer/Private/AtmosphereRendering.cpp
@@ -398,7 +398,7 @@ void FDeferredShadingSceneRenderer::RenderAtmosphere(FRHICommandListImmediate& R
 	// Atmospheric fog?
 	if (Scene->GetFeatureLevel() >= ERHIFeatureLevel::SM4 && Scene->HasAtmosphericFog())
 	{
-		SCOPED_DRAW_EVENT(Fog, DEC_SCENE_ITEMS);
+		SCOPED_DRAW_EVENT(RHICmdList, Fog, DEC_SCENE_ITEMS);
 
 		static const FVector2D Vertices[4] =
 		{
diff --git a/Engine/Source/Runtime/Renderer/Private/CompositionLighting/CompositionLighting.cpp b/Engine/Source/Runtime/Renderer/Private/CompositionLighting/CompositionLighting.cpp
--- a/Engine/Source/Runtime/Renderer/Private/CompositionLighting/CompositionLighting.cpp
+++ b/Engine/Source/Runtime/Renderer/Private/CompositionLighting/CompositionLighting.cpp
@@ -252,7 +252,7 @@ void FCompositionLighting::ProcessBeforeBasePass(FRHICommandListImmediate& RHICm
 
 		// The graph setup should be finished before this line ----------------------------------------
 
-		SCOPED_DRAW_EVENT(CompositionBeforeBasePass, DEC_SCENE_ITEMS);
+		SCOPED_DRAW_EVENT(RHICmdList, CompositionBeforeBasePass, DEC_SCENE_ITEMS);
 
 		TRefCountPtr<IPooledRenderTarget>& SceneColor = GSceneRenderTargets.GetSceneColor();
 
@@ -312,7 +312,7 @@ void FCompositionLighting::ProcessAfterBasePass(FRHICommandListImmediate& RHICmd
 
 		// The graph setup should be finished before this line ----------------------------------------
     
-		SCOPED_DRAW_EVENT(LightCompositionTasks_PreLighting, DEC_SCENE_ITEMS);
+		SCOPED_DRAW_EVENT(RHICmdList, LightCompositionTasks_PreLighting, DEC_SCENE_ITEMS);
 
 		TRefCountPtr<IPooledRenderTarget>& SceneColor = GSceneRenderTargets.GetSceneColor();
 
@@ -372,7 +372,7 @@ void FCompositionLighting::ProcessLighting(FRHICommandListImmediate& RHICmdList,
 
 		// The graph setup should be finished before this line ----------------------------------------
 
-		SCOPED_DRAW_EVENT(CompositionLighting, DEC_SCENE_ITEMS);
+		SCOPED_DRAW_EVENT(RHICmdList, CompositionLighting, DEC_SCENE_ITEMS);
 
 		// we don't replace the final element with the scenecolor because this is what those passes should do by themself
 
diff --git a/Engine/Source/Runtime/Renderer/Private/CompositionLighting/PostProcessAmbient.cpp b/Engine/Source/Runtime/Renderer/Private/CompositionLighting/PostProcessAmbient.cpp
--- a/Engine/Source/Runtime/Renderer/Private/CompositionLighting/PostProcessAmbient.cpp
+++ b/Engine/Source/Runtime/Renderer/Private/CompositionLighting/PostProcessAmbient.cpp
@@ -71,7 +71,7 @@ IMPLEMENT_SHADER_TYPE(,FPostProcessAmbientPS,TEXT("PostProcessAmbient"),TEXT("Ma
 
 void FRCPassPostProcessAmbient::Process(FRenderingCompositePassContext& Context)
 {
-	SCOPED_DRAW_EVENT(PostProcessAmbient, DEC_SCENE_ITEMS);
+	SCOPED_DRAW_EVENT(Context.RHICmdList, PostProcessAmbient, DEC_SCENE_ITEMS);
 
 	const FSceneView& View = Context.View;
 	const FSceneViewFamily& ViewFamily = *(View.Family);
diff --git a/Engine/Source/Runtime/Renderer/Private/CompositionLighting/PostProcessAmbientOcclusion.cpp b/Engine/Source/Runtime/Renderer/Private/CompositionLighting/PostProcessAmbientOcclusion.cpp
--- a/Engine/Source/Runtime/Renderer/Private/CompositionLighting/PostProcessAmbientOcclusion.cpp
+++ b/Engine/Source/Runtime/Renderer/Private/CompositionLighting/PostProcessAmbientOcclusion.cpp
@@ -361,7 +361,7 @@ void FRCPassPostProcessAmbientOcclusionSetup::SetShaderSetupTempl(const FRenderi
 
 void FRCPassPostProcessAmbientOcclusionSetup::Process(FRenderingCompositePassContext& Context)
 {
-	SCOPED_DRAW_EVENT(AmbientOcclusionSetup, DEC_SCENE_ITEMS);
+	SCOPED_DRAW_EVENT(Context.RHICmdList, AmbientOcclusionSetup, DEC_SCENE_ITEMS);
 
 	const FSceneView& View = Context.View;
 
@@ -474,7 +474,7 @@ void FRCPassPostProcessAmbientOcclusion::SetShaderTempl(const FRenderingComposit
 
 void FRCPassPostProcessAmbientOcclusion::Process(FRenderingCompositePassContext& Context)
 {
-	SCOPED_DRAW_EVENT(AmbientOcclusion, DEC_SCENE_ITEMS);
+	SCOPED_DRAW_EVENT(Context.RHICmdList, AmbientOcclusion, DEC_SCENE_ITEMS);
 
 	const FSceneView& View = Context.View;
 
@@ -591,7 +591,7 @@ FPooledRenderTargetDesc FRCPassPostProcessAmbientOcclusion::ComputeOutputDesc(EP
 
 void FRCPassPostProcessBasePassAO::Process(FRenderingCompositePassContext& Context)
 {
-	SCOPED_DRAW_EVENT(ApplyAOToBasePassSceneColor, DEC_SCENE_ITEMS);
+	SCOPED_DRAW_EVENT(Context.RHICmdList, ApplyAOToBasePassSceneColor, DEC_SCENE_ITEMS);
 
 	const FSceneView& View = Context.View;
 
diff --git a/Engine/Source/Runtime/Renderer/Private/CompositionLighting/PostProcessDeferredDecals.cpp b/Engine/Source/Runtime/Renderer/Private/CompositionLighting/PostProcessDeferredDecals.cpp
--- a/Engine/Source/Runtime/Renderer/Private/CompositionLighting/PostProcessDeferredDecals.cpp
+++ b/Engine/Source/Runtime/Renderer/Private/CompositionLighting/PostProcessDeferredDecals.cpp
@@ -374,7 +374,7 @@ FGlobalBoundShaderState StencilDecalMaskBoundShaderState;
 /** Draws a full view quad that sets stencil to 1 anywhere that decals should not be projected. */
 void StencilDecalMask(FRHICommandList& RHICmdList, const FViewInfo& View)
 {
-	SCOPED_DRAW_EVENT(StencilDecalMask, DEC_SCENE_ITEMS);
+	SCOPED_DRAW_EVENT(RHICmdList, StencilDecalMask, DEC_SCENE_ITEMS);
 	RHICmdList.SetRasterizerState(TStaticRasterizerState<FM_Solid, CM_None>::GetRHI());
 	RHICmdList.SetBlendState(TStaticBlendState<CW_NONE>::GetRHI());
 	SetRenderTarget(RHICmdList, NULL, GSceneRenderTargets.GetSceneDepthSurface());
@@ -640,7 +640,7 @@ void SetShader(const FRenderingCompositePassContext& Context, const FTransientDe
 
 bool RenderPreStencil(FRenderingCompositePassContext& Context, const FMaterialShaderMap* MaterialShaderMap, const FMatrix& ComponentToWorldMatrix, const FMatrix& FrustumComponentToClip)
 {
-	SCOPED_DRAW_EVENT(RenderPreStencil, DEC_SCENE_ITEMS);
+	SCOPED_DRAW_EVENT(Context.RHICmdList, RenderPreStencil, DEC_SCENE_ITEMS);
 
 	const FSceneView& View = Context.View;
 
@@ -721,7 +721,7 @@ void FRCPassPostProcessDeferredDecals::Process(FRenderingCompositePassContext& C
 		PassOutputs[0].RenderTargetDesc = OutputOfMyInput->RenderTargetDesc;
 	}
 
-	SCOPED_DRAW_EVENT(PostProcessDeferredDecals, DEC_SCENE_ITEMS);
+	SCOPED_DRAW_EVENT(RHICmdList, PostProcessDeferredDecals, DEC_SCENE_ITEMS);
 
 	if(RenderStage == 0)
 	{
@@ -754,7 +754,7 @@ void FRCPassPostProcessDeferredDecals::Process(FRenderingCompositePassContext& C
 		}
 
 
-		SCOPED_DRAW_EVENT(DBufferClear, DEC_SCENE_ITEMS);
+		SCOPED_DRAW_EVENT(RHICmdList, DBufferClear, DEC_SCENE_ITEMS);
 		{
 			// could be optimized
 			SetRenderTarget(RHICmdList, GSceneRenderTargets.DBufferA->GetRenderTargetItem().TargetableTexture, FTextureRHIParamRef());
@@ -888,7 +888,7 @@ void FRCPassPostProcessDeferredDecals::Process(FRenderingCompositePassContext& C
 		int32 WasInsideDecal = -1;
 		const ERHIFeatureLevel::Type SMFeatureLevel = Context.GetFeatureLevel();
 
-		SCOPED_DRAW_EVENT(Decals, DEC_SCENE_ITEMS);
+		SCOPED_DRAW_EVENT(RHICmdList, Decals, DEC_SCENE_ITEMS);
 		INC_DWORD_STAT_BY(STAT_Decals, SortedDecals.Num());
 		
 		RHICmdList.SetStreamSource(0, GUnitCubeVertexBuffer.VertexBufferRHI, sizeof(FVector4), 0);
diff --git a/Engine/Source/Runtime/Renderer/Private/CompositionLighting/PostProcessLpvIndirect.cpp b/Engine/Source/Runtime/Renderer/Private/CompositionLighting/PostProcessLpvIndirect.cpp
--- a/Engine/Source/Runtime/Renderer/Private/CompositionLighting/PostProcessLpvIndirect.cpp
+++ b/Engine/Source/Runtime/Renderer/Private/CompositionLighting/PostProcessLpvIndirect.cpp
@@ -130,7 +130,7 @@ IMPLEMENT_SHADER_TYPE(,FPostProcessLpvIndirectPS,TEXT("PostProcessLpvIndirect"),
 
 void FRCPassPostProcessLpvIndirect::Process(FRenderingCompositePassContext& Context)
 {
-	SCOPED_DRAW_EVENT(PostProcessLpvIndirect, DEC_SCENE_ITEMS);
+	SCOPED_DRAW_EVENT(Context.RHICmdList, PostProcessLpvIndirect, DEC_SCENE_ITEMS);
 
 	{
 		FRenderingCompositeOutput* OutputOfMyInput = GetInput(ePId_Input0)->GetOutput();
diff --git a/Engine/Source/Runtime/Renderer/Private/CompositionLighting/PostProcessPassThrough.cpp b/Engine/Source/Runtime/Renderer/Private/CompositionLighting/PostProcessPassThrough.cpp
--- a/Engine/Source/Runtime/Renderer/Private/CompositionLighting/PostProcessPassThrough.cpp
+++ b/Engine/Source/Runtime/Renderer/Private/CompositionLighting/PostProcessPassThrough.cpp
@@ -70,7 +70,7 @@ FRCPassPostProcessPassThrough::FRCPassPostProcessPassThrough(FPooledRenderTarget
 
 void FRCPassPostProcessPassThrough::Process(FRenderingCompositePassContext& Context)
 {
-	SCOPED_DRAW_EVENT(PassThrough, DEC_SCENE_ITEMS);
+	SCOPED_DRAW_EVENT(Context.RHICmdList, PassThrough, DEC_SCENE_ITEMS);
 
 	const FPooledRenderTargetDesc* InputDesc = GetInputDesc(ePId_Input0);
 
diff --git a/Engine/Source/Runtime/Renderer/Private/CustomDepthRendering.cpp b/Engine/Source/Runtime/Renderer/Private/CustomDepthRendering.cpp
--- a/Engine/Source/Runtime/Renderer/Private/CustomDepthRendering.cpp
+++ b/Engine/Source/Runtime/Renderer/Private/CustomDepthRendering.cpp
@@ -23,7 +23,7 @@ bool FCustomDepthPrimSet::DrawPrims(FRHICommandListImmediate& RHICmdList, const
 
 	if(Prims.Num())
 	{
-		SCOPED_DRAW_EVENT(CustomDepth, DEC_SCENE_ITEMS);
+		SCOPED_DRAW_EVENT(RHICmdList, CustomDepth, DEC_SCENE_ITEMS);
 
 		const bool bUseGetMeshElements = ShouldUseGetDynamicMeshElements();
 
diff --git a/Engine/Source/Runtime/Renderer/Private/DeferredShadingRenderer.cpp b/Engine/Source/Runtime/Renderer/Private/DeferredShadingRenderer.cpp
--- a/Engine/Source/Runtime/Renderer/Private/DeferredShadingRenderer.cpp
+++ b/Engine/Source/Runtime/Renderer/Private/DeferredShadingRenderer.cpp
@@ -130,7 +130,7 @@ extern FGlobalBoundShaderState GClearMRTBoundShaderState[8];
 void FDeferredShadingSceneRenderer::ClearGBufferAtMaxZ(FRHICommandList& RHICmdList)
 {
 	// Assumes BeginRenderingSceneColor() has been called before this function
-	SCOPED_DRAW_EVENT(ClearGBufferAtMaxZ, DEC_SCENE_ITEMS);
+	SCOPED_DRAW_EVENT(RHICmdList, ClearGBufferAtMaxZ, DEC_SCENE_ITEMS);
 
 	// Clear the G Buffer render targets
 	const bool bClearBlack = Views[0].Family->EngineShowFlags.ShaderComplexity || Views[0].Family->EngineShowFlags.StationaryLightOverlap;
@@ -180,7 +180,7 @@ void FDeferredShadingSceneRenderer::ClearGBufferAtMaxZ(FRHICommandList& RHICmdLi
 	// Clear each viewport by drawing background color at MaxZ depth
 	for(int32 ViewIndex = 0;ViewIndex < Views.Num();ViewIndex++)
 	{
-		SCOPED_CONDITIONAL_DRAW_EVENTF(EventView, Views.Num() > 1, DEC_SCENE_ITEMS, TEXT("ClearView%d"), ViewIndex);
+		SCOPED_CONDITIONAL_DRAW_EVENTF(RHICmdList, EventView, Views.Num() > 1, DEC_SCENE_ITEMS, TEXT("ClearView%d"), ViewIndex);
 
 		FViewInfo& View = Views[ViewIndex];
 
@@ -202,14 +202,14 @@ bool FDeferredShadingSceneRenderer::RenderBasePassStaticDataMasked(FRHICommandLi
 		// Draw the scene's base pass draw lists.
 		FScene::EBasePassDrawListType MaskedDrawType = FScene::EBasePass_Masked;
 		{
-			SCOPED_DRAW_EVENT(StaticMaskedNoLightmap, DEC_SCENE_ITEMS);
+			SCOPED_DRAW_EVENT(RHICmdList, StaticMaskedNoLightmap, DEC_SCENE_ITEMS);
 			bDirty |= Scene->BasePassNoLightMapDrawList[MaskedDrawType].DrawVisible(RHICmdList, View,View.StaticMeshVisibilityMap,View.StaticMeshBatchVisibility);
 			bDirty |= Scene->BasePassSimpleDynamicLightingDrawList[MaskedDrawType].DrawVisible(RHICmdList, View,View.StaticMeshVisibilityMap,View.StaticMeshBatchVisibility);
 			bDirty |= Scene->BasePassCachedVolumeIndirectLightingDrawList[MaskedDrawType].DrawVisible(RHICmdList, View,View.StaticMeshVisibilityMap,View.StaticMeshBatchVisibility);
 			bDirty |= Scene->BasePassCachedPointIndirectLightingDrawList[MaskedDrawType].DrawVisible(RHICmdList, View,View.StaticMeshVisibilityMap,View.StaticMeshBatchVisibility);
 		}
 		{
-			SCOPED_DRAW_EVENT(StaticMaskedLightmapped, DEC_SCENE_ITEMS);
+			SCOPED_DRAW_EVENT(RHICmdList, StaticMaskedLightmapped, DEC_SCENE_ITEMS);
 			bDirty |= Scene->BasePassHighQualityLightMapDrawList[MaskedDrawType].DrawVisible(RHICmdList, View,View.StaticMeshVisibilityMap,View.StaticMeshBatchVisibility);
 			bDirty |= Scene->BasePassDistanceFieldShadowMapLightMapDrawList[MaskedDrawType].DrawVisible(RHICmdList, View,View.StaticMeshVisibilityMap,View.StaticMeshBatchVisibility);
 			bDirty |= Scene->BasePassLowQualityLightMapDrawList[MaskedDrawType].DrawVisible(RHICmdList, View,View.StaticMeshVisibilityMap,View.StaticMeshBatchVisibility);
@@ -222,14 +222,20 @@ void FDeferredShadingSceneRenderer::RenderBasePassStaticDataMaskedParallel(FView
 {
 	// Draw the scene's base pass draw lists.
 	FScene::EBasePassDrawListType MaskedDrawType = FScene::EBasePass_Masked;
-	Scene->BasePassNoLightMapDrawList[MaskedDrawType].DrawVisibleParallel(View, View.StaticMeshVisibilityMap, View.StaticMeshBatchVisibility, Width, ParentCmdList, OutDirty);
-	Scene->BasePassSimpleDynamicLightingDrawList[MaskedDrawType].DrawVisibleParallel(View, View.StaticMeshVisibilityMap, View.StaticMeshBatchVisibility, Width, ParentCmdList, OutDirty);
-	Scene->BasePassCachedVolumeIndirectLightingDrawList[MaskedDrawType].DrawVisibleParallel(View, View.StaticMeshVisibilityMap, View.StaticMeshBatchVisibility, Width, ParentCmdList, OutDirty);
-	Scene->BasePassCachedPointIndirectLightingDrawList[MaskedDrawType].DrawVisibleParallel(View, View.StaticMeshVisibilityMap, View.StaticMeshBatchVisibility, Width, ParentCmdList, OutDirty);
-
-	Scene->BasePassHighQualityLightMapDrawList[MaskedDrawType].DrawVisibleParallel(View, View.StaticMeshVisibilityMap, View.StaticMeshBatchVisibility, Width, ParentCmdList, OutDirty);
-	Scene->BasePassDistanceFieldShadowMapLightMapDrawList[MaskedDrawType].DrawVisibleParallel(View, View.StaticMeshVisibilityMap, View.StaticMeshBatchVisibility, Width, ParentCmdList, OutDirty);
-	Scene->BasePassLowQualityLightMapDrawList[MaskedDrawType].DrawVisibleParallel(View, View.StaticMeshVisibilityMap, View.StaticMeshBatchVisibility, Width, ParentCmdList, OutDirty);
+	{
+		SCOPED_DRAW_EVENT(ParentCmdList, StaticMaskedNoLightmap, DEC_SCENE_ITEMS);
+		Scene->BasePassNoLightMapDrawList[MaskedDrawType].DrawVisibleParallel(View, View.StaticMeshVisibilityMap, View.StaticMeshBatchVisibility, Width, ParentCmdList, OutDirty);
+		Scene->BasePassSimpleDynamicLightingDrawList[MaskedDrawType].DrawVisibleParallel(View, View.StaticMeshVisibilityMap, View.StaticMeshBatchVisibility, Width, ParentCmdList, OutDirty);
+		Scene->BasePassCachedVolumeIndirectLightingDrawList[MaskedDrawType].DrawVisibleParallel(View, View.StaticMeshVisibilityMap, View.StaticMeshBatchVisibility, Width, ParentCmdList, OutDirty);
+		Scene->BasePassCachedPointIndirectLightingDrawList[MaskedDrawType].DrawVisibleParallel(View, View.StaticMeshVisibilityMap, View.StaticMeshBatchVisibility, Width, ParentCmdList, OutDirty);
+	}
+
+	{
+		SCOPED_DRAW_EVENT(ParentCmdList, StaticMaskedLightmapped, DEC_SCENE_ITEMS);
+		Scene->BasePassHighQualityLightMapDrawList[MaskedDrawType].DrawVisibleParallel(View, View.StaticMeshVisibilityMap, View.StaticMeshBatchVisibility, Width, ParentCmdList, OutDirty);
+		Scene->BasePassDistanceFieldShadowMapLightMapDrawList[MaskedDrawType].DrawVisibleParallel(View, View.StaticMeshVisibilityMap, View.StaticMeshBatchVisibility, Width, ParentCmdList, OutDirty);
+		Scene->BasePassLowQualityLightMapDrawList[MaskedDrawType].DrawVisibleParallel(View, View.StaticMeshVisibilityMap, View.StaticMeshBatchVisibility, Width, ParentCmdList, OutDirty);
+	}
 }
 
 bool FDeferredShadingSceneRenderer::RenderBasePassStaticDataDefault(FRHICommandList& RHICmdList, FViewInfo& View)
@@ -238,14 +244,14 @@ bool FDeferredShadingSceneRenderer::RenderBasePassStaticDataDefault(FRHICommandL
 	{
 		FScene::EBasePassDrawListType OpaqueDrawType = FScene::EBasePass_Default;
 		{
-			SCOPED_DRAW_EVENT(StaticOpaqueNoLightmap, DEC_SCENE_ITEMS);
+			SCOPED_DRAW_EVENT(RHICmdList, StaticOpaqueNoLightmap, DEC_SCENE_ITEMS);
 			bDirty |= Scene->BasePassNoLightMapDrawList[OpaqueDrawType].DrawVisible(RHICmdList, View,View.StaticMeshVisibilityMap,View.StaticMeshBatchVisibility);
 			bDirty |= Scene->BasePassSimpleDynamicLightingDrawList[OpaqueDrawType].DrawVisible(RHICmdList, View,View.StaticMeshVisibilityMap,View.StaticMeshBatchVisibility);
 			bDirty |= Scene->BasePassCachedVolumeIndirectLightingDrawList[OpaqueDrawType].DrawVisible(RHICmdList, View,View.StaticMeshVisibilityMap,View.StaticMeshBatchVisibility);
 			bDirty |= Scene->BasePassCachedPointIndirectLightingDrawList[OpaqueDrawType].DrawVisible(RHICmdList, View,View.StaticMeshVisibilityMap,View.StaticMeshBatchVisibility);
 		}
 		{
-			SCOPED_DRAW_EVENT(StaticOpaqueLightmapped, DEC_SCENE_ITEMS);
+			SCOPED_DRAW_EVENT(RHICmdList, StaticOpaqueLightmapped, DEC_SCENE_ITEMS);
 			bDirty |= Scene->BasePassHighQualityLightMapDrawList[OpaqueDrawType].DrawVisible(RHICmdList, View,View.StaticMeshVisibilityMap,View.StaticMeshBatchVisibility);
 			bDirty |= Scene->BasePassDistanceFieldShadowMapLightMapDrawList[OpaqueDrawType].DrawVisible(RHICmdList, View,View.StaticMeshVisibilityMap,View.StaticMeshBatchVisibility);
 			bDirty |= Scene->BasePassLowQualityLightMapDrawList[OpaqueDrawType].DrawVisible(RHICmdList, View,View.StaticMeshVisibilityMap,View.StaticMeshBatchVisibility);
@@ -258,14 +264,20 @@ bool FDeferredShadingSceneRenderer::RenderBasePassStaticDataDefault(FRHICommandL
 void FDeferredShadingSceneRenderer::RenderBasePassStaticDataDefaultParallel(FViewInfo& View, int32 Width, FRHICommandList& ParentCmdList, bool& OutDirty)
 {
 	FScene::EBasePassDrawListType OpaqueDrawType = FScene::EBasePass_Default;
-	Scene->BasePassNoLightMapDrawList[OpaqueDrawType].DrawVisibleParallel(View, View.StaticMeshVisibilityMap, View.StaticMeshBatchVisibility, Width, ParentCmdList, OutDirty);
-	Scene->BasePassSimpleDynamicLightingDrawList[OpaqueDrawType].DrawVisibleParallel(View, View.StaticMeshVisibilityMap, View.StaticMeshBatchVisibility, Width, ParentCmdList, OutDirty);
-	Scene->BasePassCachedVolumeIndirectLightingDrawList[OpaqueDrawType].DrawVisibleParallel(View, View.StaticMeshVisibilityMap, View.StaticMeshBatchVisibility, Width, ParentCmdList, OutDirty);
-	Scene->BasePassCachedPointIndirectLightingDrawList[OpaqueDrawType].DrawVisibleParallel(View, View.StaticMeshVisibilityMap, View.StaticMeshBatchVisibility, Width, ParentCmdList, OutDirty);
-
-	Scene->BasePassHighQualityLightMapDrawList[OpaqueDrawType].DrawVisibleParallel(View, View.StaticMeshVisibilityMap, View.StaticMeshBatchVisibility, Width, ParentCmdList, OutDirty);
-	Scene->BasePassDistanceFieldShadowMapLightMapDrawList[OpaqueDrawType].DrawVisibleParallel(View, View.StaticMeshVisibilityMap, View.StaticMeshBatchVisibility, Width, ParentCmdList, OutDirty);
-	Scene->BasePassLowQualityLightMapDrawList[OpaqueDrawType].DrawVisibleParallel(View, View.StaticMeshVisibilityMap, View.StaticMeshBatchVisibility, Width, ParentCmdList, OutDirty);
+	{
+		SCOPED_DRAW_EVENT(ParentCmdList, StaticOpaqueNoLightmap, DEC_SCENE_ITEMS);
+		Scene->BasePassNoLightMapDrawList[OpaqueDrawType].DrawVisibleParallel(View, View.StaticMeshVisibilityMap, View.StaticMeshBatchVisibility, Width, ParentCmdList, OutDirty);
+		Scene->BasePassSimpleDynamicLightingDrawList[OpaqueDrawType].DrawVisibleParallel(View, View.StaticMeshVisibilityMap, View.StaticMeshBatchVisibility, Width, ParentCmdList, OutDirty);
+		Scene->BasePassCachedVolumeIndirectLightingDrawList[OpaqueDrawType].DrawVisibleParallel(View, View.StaticMeshVisibilityMap, View.StaticMeshBatchVisibility, Width, ParentCmdList, OutDirty);
+		Scene->BasePassCachedPointIndirectLightingDrawList[OpaqueDrawType].DrawVisibleParallel(View, View.StaticMeshVisibilityMap, View.StaticMeshBatchVisibility, Width, ParentCmdList, OutDirty);
+	}
+
+	{
+		SCOPED_DRAW_EVENT(ParentCmdList, StaticOpaqueLightmapped, DEC_SCENE_ITEMS);
+		Scene->BasePassHighQualityLightMapDrawList[OpaqueDrawType].DrawVisibleParallel(View, View.StaticMeshVisibilityMap, View.StaticMeshBatchVisibility, Width, ParentCmdList, OutDirty);
+		Scene->BasePassDistanceFieldShadowMapLightMapDrawList[OpaqueDrawType].DrawVisibleParallel(View, View.StaticMeshVisibilityMap, View.StaticMeshBatchVisibility, Width, ParentCmdList, OutDirty);
+		Scene->BasePassLowQualityLightMapDrawList[OpaqueDrawType].DrawVisibleParallel(View, View.StaticMeshVisibilityMap, View.StaticMeshBatchVisibility, Width, ParentCmdList, OutDirty);
+	}
 }
 
 void FDeferredShadingSceneRenderer::SortBasePassStaticData(FVector ViewPosition)
@@ -354,7 +366,7 @@ void FDeferredShadingSceneRenderer::RenderBasePassDynamicData(FRHICommandList& R
 	bool bDirty = false;
 
 	SCOPE_CYCLE_COUNTER(STAT_DynamicPrimitiveDrawTime);
-	SCOPED_DRAW_EVENT(Dynamic, DEC_SCENE_ITEMS);
+	SCOPED_DRAW_EVENT(RHICmdList, Dynamic, DEC_SCENE_ITEMS);
 
 	const bool bUseGetMeshElements = ShouldUseGetDynamicMeshElements();
 
@@ -606,7 +618,7 @@ void FDeferredShadingSceneRenderer::Render(FRHICommandListImmediate& RHICmdList)
 	{
 		return;
 	}
-	SCOPED_DRAW_EVENT(Scene,DEC_SCENE_ITEMS);
+	SCOPED_DRAW_EVENT(RHICmdList, Scene,DEC_SCENE_ITEMS);
 	
 	// Initialize global system textures (pass-through if already initialized).
 	GSystemTextures.InitializeTextures(RHICmdList, FeatureLevel);
@@ -730,7 +742,7 @@ void FDeferredShadingSceneRenderer::Render(FRHICommandListImmediate& RHICmdList)
 		// e.g. ambient cubemaps, ambient occlusion, deferred decals
 		for(int32 ViewIndex = 0;ViewIndex < Views.Num();ViewIndex++)
 		{	
-			SCOPED_CONDITIONAL_DRAW_EVENTF(EventView,Views.Num() > 1, DEC_SCENE_ITEMS, TEXT("View%d"), ViewIndex);
+			SCOPED_CONDITIONAL_DRAW_EVENTF(RHICmdList, EventView,Views.Num() > 1, DEC_SCENE_ITEMS, TEXT("View%d"), ViewIndex);
 			GCompositionLighting.ProcessBeforeBasePass(RHICmdList, Views[ViewIndex]);
 		}
 	}
@@ -824,7 +836,7 @@ void FDeferredShadingSceneRenderer::Render(FRHICommandListImmediate& RHICmdList)
 		// e.g. deferred decals, blurred GBuffer
 		for(int32 ViewIndex = 0;ViewIndex < Views.Num();ViewIndex++)
 		{	
-			SCOPED_CONDITIONAL_DRAW_EVENTF(EventView,Views.Num() > 1, DEC_SCENE_ITEMS, TEXT("View%d"), ViewIndex);
+			SCOPED_CONDITIONAL_DRAW_EVENTF(RHICmdList, EventView,Views.Num() > 1, DEC_SCENE_ITEMS, TEXT("View%d"), ViewIndex);
 			GCompositionLighting.ProcessAfterBasePass(RHICmdList, Views[ViewIndex]);
 		}
 
@@ -856,7 +868,7 @@ void FDeferredShadingSceneRenderer::Render(FRHICommandListImmediate& RHICmdList)
 		// e.g. ambient cubemaps, ambient occlusion, LPV indirect
 		for(int32 ViewIndex = 0; ViewIndex < Views.Num(); ++ViewIndex)
 		{	
-			SCOPED_CONDITIONAL_DRAW_EVENTF(EventView,Views.Num() > 1, DEC_SCENE_ITEMS, TEXT("View%d"), ViewIndex);
+			SCOPED_CONDITIONAL_DRAW_EVENTF(RHICmdList, EventView,Views.Num() > 1, DEC_SCENE_ITEMS, TEXT("View%d"), ViewIndex);
 			GCompositionLighting.ProcessLighting(RHICmdList, Views[ViewIndex]);
 		}
 	}
@@ -955,12 +967,12 @@ void FDeferredShadingSceneRenderer::Render(FRHICommandListImmediate& RHICmdList)
 	// Finish rendering for each view.
 	if(ViewFamily.bResolveScene)
 	{
-		SCOPED_DRAW_EVENT(PostProcessing, DEC_SCENE_ITEMS);
+		SCOPED_DRAW_EVENT(RHICmdList, PostProcessing, DEC_SCENE_ITEMS);
 		SCOPE_CYCLE_COUNTER(STAT_FinishRenderViewTargetTime);
 
 		for(int32 ViewIndex = 0; ViewIndex < Views.Num(); ViewIndex++)
 		{
-			SCOPED_CONDITIONAL_DRAW_EVENTF(EventView, Views.Num() > 1, DEC_SCENE_ITEMS, TEXT("View%d"), ViewIndex);
+			SCOPED_CONDITIONAL_DRAW_EVENTF(RHICmdList, EventView, Views.Num() > 1, DEC_SCENE_ITEMS, TEXT("View%d"), ViewIndex);
 
 			GPostProcessing.Process(RHICmdList, Views[ ViewIndex ], VelocityRT);
 
@@ -1023,7 +1035,7 @@ bool FDeferredShadingSceneRenderer::RenderPrePassViewDynamic(FRHICommandList& RH
 		// Draw the dynamic occluder primitives using a depth drawing policy.
 		TDynamicPrimitiveDrawer<FDepthDrawingPolicyFactory> Drawer(RHICmdList, &View, FDepthDrawingPolicyFactory::ContextType(EarlyZPassMode), true);
 		{
-			SCOPED_DRAW_EVENT(Dynamic, DEC_SCENE_ITEMS);
+			SCOPED_DRAW_EVENT(RHICmdList, Dynamic, DEC_SCENE_ITEMS);
 			for(int32 PrimitiveIndex = 0;PrimitiveIndex < View.VisibleDynamicPrimitives.Num();PrimitiveIndex++)
 			{
 				const FPrimitiveSceneInfo* PrimitiveSceneInfo = View.VisibleDynamicPrimitives[PrimitiveIndex];
@@ -1081,19 +1093,19 @@ bool FDeferredShadingSceneRenderer::RenderPrePassView(FRHICommandList& RHICmdLis
 		// Draw opaque occluders which support a separate position-only
 		// vertex buffer to minimize vertex fetch bandwidth, which is
 		// often the bottleneck during the depth only pass.
-		SCOPED_DRAW_EVENT(PosOnlyOpaque, DEC_SCENE_ITEMS);
+		SCOPED_DRAW_EVENT(RHICmdList, PosOnlyOpaque, DEC_SCENE_ITEMS);
 		bDirty |= Scene->PositionOnlyDepthDrawList.DrawVisible(RHICmdList, View,View.StaticMeshOccluderMap,View.StaticMeshBatchVisibility);
 	}
 	{
 		// Draw opaque occluders, using double speed z where supported.
-		SCOPED_DRAW_EVENT(Opaque, DEC_SCENE_ITEMS);
+		SCOPED_DRAW_EVENT(RHICmdList, Opaque, DEC_SCENE_ITEMS);
 		bDirty |= Scene->DepthDrawList.DrawVisible(RHICmdList, View,View.StaticMeshOccluderMap,View.StaticMeshBatchVisibility);
 	}
 
 	if(EarlyZPassMode >= DDM_AllOccluders)
 	{
 		// Draw opaque occluders with masked materials
-		SCOPED_DRAW_EVENT(Opaque, DEC_SCENE_ITEMS);
+		SCOPED_DRAW_EVENT(RHICmdList, Opaque, DEC_SCENE_ITEMS);
 		bDirty |= Scene->MaskedDepthDrawList.DrawVisible(RHICmdList, View,View.StaticMeshOccluderMap,View.StaticMeshBatchVisibility);
 	}
 
@@ -1155,19 +1167,19 @@ void FDeferredShadingSceneRenderer::RenderPrePassViewParallel(const FViewInfo& V
 		// Draw opaque occluders which support a separate position-only
 		// vertex buffer to minimize vertex fetch bandwidth, which is
 		// often the bottleneck during the depth only pass.
-		//SCOPED_DRAW_EVENT(PosOnlyOpaque, DEC_SCENE_ITEMS);
+		SCOPED_DRAW_EVENT(ParentCmdList, PosOnlyOpaque, DEC_SCENE_ITEMS);
 		Scene->PositionOnlyDepthDrawList.DrawVisibleParallel(View, View.StaticMeshOccluderMap, View.StaticMeshBatchVisibility, Width, ParentCmdList, OutDirty);
 	}
 	{
 		// Draw opaque occluders, using double speed z where supported.
-		//SCOPED_DRAW_EVENT(Opaque, DEC_SCENE_ITEMS);
+		SCOPED_DRAW_EVENT(ParentCmdList, Opaque, DEC_SCENE_ITEMS);
 		Scene->DepthDrawList.DrawVisibleParallel(View,View.StaticMeshOccluderMap,View.StaticMeshBatchVisibility, Width, ParentCmdList, OutDirty);
 	}
 
 	if(EarlyZPassMode >= DDM_AllOccluders)
 	{
 		// Draw opaque occluders with masked materials
-		//SCOPED_DRAW_EVENT(Opaque, DEC_SCENE_ITEMS);
+		SCOPED_DRAW_EVENT(ParentCmdList, Opaque, DEC_SCENE_ITEMS);
 		Scene->MaskedDepthDrawList.DrawVisibleParallel(View,View.StaticMeshOccluderMap,View.StaticMeshBatchVisibility, Width, ParentCmdList, OutDirty);
 	}
 	{
@@ -1184,7 +1196,7 @@ void FDeferredShadingSceneRenderer::RenderPrePassViewParallel(const FViewInfo& V
 /** Renders the scene's prepass and occlusion queries */
 bool FDeferredShadingSceneRenderer::RenderPrePass(FRHICommandListImmediate& RHICmdList)
 {
-	SCOPED_DRAW_EVENT(PrePass, DEC_SCENE_ITEMS);
+	SCOPED_DRAW_EVENT(RHICmdList, PrePass, DEC_SCENE_ITEMS);
 	SCOPE_CYCLE_COUNTER(STAT_DepthDrawTime);
 
 	bool bDirty = false;
@@ -1206,7 +1218,7 @@ bool FDeferredShadingSceneRenderer::RenderPrePass(FRHICommandListImmediate& RHIC
 
 			for(int32 ViewIndex = 0;ViewIndex < Views.Num();ViewIndex++)
 			{
-				SCOPED_CONDITIONAL_DRAW_EVENTF(EventView, Views.Num() > 1, DEC_SCENE_ITEMS, TEXT("View%d"), ViewIndex);
+				SCOPED_CONDITIONAL_DRAW_EVENTF(RHICmdList, EventView, Views.Num() > 1, DEC_SCENE_ITEMS, TEXT("View%d"), ViewIndex);
 				const FViewInfo& View = Views[ViewIndex];
 				RenderPrePassViewParallel(View, Width, RHICmdList, bDirty);
 			}
@@ -1215,7 +1227,7 @@ bool FDeferredShadingSceneRenderer::RenderPrePass(FRHICommandListImmediate& RHIC
 		{
 			for(int32 ViewIndex = 0;ViewIndex < Views.Num();ViewIndex++)
 			{
-				SCOPED_CONDITIONAL_DRAW_EVENTF(EventView, Views.Num() > 1, DEC_SCENE_ITEMS, TEXT("View%d"), ViewIndex);
+				SCOPED_CONDITIONAL_DRAW_EVENTF(RHICmdList, EventView, Views.Num() > 1, DEC_SCENE_ITEMS, TEXT("View%d"), ViewIndex);
 				const FViewInfo& View = Views[ViewIndex];
 				bDirty |= RenderPrePassView(RHICmdList, View);
 			}
@@ -1242,7 +1254,7 @@ bool FDeferredShadingSceneRenderer::RenderBasePass(FRHICommandListImmediate& RHI
 	}
 	else
 	{
-		SCOPED_DRAW_EVENT(BasePass, DEC_SCENE_ITEMS);
+		SCOPED_DRAW_EVENT(RHICmdList, BasePass, DEC_SCENE_ITEMS);
 		SCOPE_CYCLE_COUNTER(STAT_BasePassDrawTime);
 
 		if (GRHICommandList.UseParallelAlgorithms())
@@ -1252,7 +1264,7 @@ bool FDeferredShadingSceneRenderer::RenderBasePass(FRHICommandListImmediate& RHI
 
 			for(int32 ViewIndex = 0;ViewIndex < Views.Num();ViewIndex++)
 			{
-				SCOPED_CONDITIONAL_DRAW_EVENTF(EventView, Views.Num() > 1, DEC_SCENE_ITEMS, TEXT("View%d"), ViewIndex);
+				SCOPED_CONDITIONAL_DRAW_EVENTF(RHICmdList, EventView, Views.Num() > 1, DEC_SCENE_ITEMS, TEXT("View%d"), ViewIndex);
 				FViewInfo& View = Views[ViewIndex];
 				RenderBasePassViewParallel(View, Width, RHICmdList, bDirty);
 			}
@@ -1261,7 +1273,7 @@ bool FDeferredShadingSceneRenderer::RenderBasePass(FRHICommandListImmediate& RHI
 		{
 			for(int32 ViewIndex = 0;ViewIndex < Views.Num();ViewIndex++)
 			{
-				SCOPED_CONDITIONAL_DRAW_EVENTF(EventView, Views.Num() > 1, DEC_SCENE_ITEMS, TEXT("View%d"), ViewIndex);
+				SCOPED_CONDITIONAL_DRAW_EVENTF(RHICmdList, EventView, Views.Num() > 1, DEC_SCENE_ITEMS, TEXT("View%d"), ViewIndex);
 				FViewInfo& View = Views[ViewIndex];
 
 				bDirty |= RenderBasePassView(RHICmdList, View);
@@ -1288,7 +1300,7 @@ void FDeferredShadingSceneRenderer::ClearLPVs(FRHICommandListImmediate& RHICmdLi
 
 			if(LightPropagationVolume)
 			{
-				SCOPED_DRAW_EVENT(ClearLPVs, DEC_SCENE_ITEMS);
+				SCOPED_DRAW_EVENT(RHICmdList, ClearLPVs, DEC_SCENE_ITEMS);
 				SCOPE_CYCLE_COUNTER(STAT_UpdateLPVs);
 				LightPropagationVolume->InitSettings(RHICmdList, Views[ViewIndex]);
 				LightPropagationVolume->Clear(RHICmdList, View);
@@ -1310,7 +1322,7 @@ void FDeferredShadingSceneRenderer::PropagateLPVs(FRHICommandListImmediate& RHIC
 
 			if(LightPropagationVolume)
 			{
-				SCOPED_DRAW_EVENT(UpdateLPVs, DEC_SCENE_ITEMS);
+				SCOPED_DRAW_EVENT(RHICmdList, UpdateLPVs, DEC_SCENE_ITEMS);
 				SCOPE_CYCLE_COUNTER(STAT_UpdateLPVs);
 				
 				LightPropagationVolume->Propagate(RHICmdList, View);
@@ -1388,7 +1400,7 @@ void FDeferredShadingSceneRenderer::UpdateDownsampledDepthSurface(FRHICommandLis
 	{
 		SetRenderTarget(RHICmdList, NULL, GSceneRenderTargets.GetSmallDepthSurface());
 
-		SCOPED_DRAW_EVENT(DownsampleDepth, DEC_SCENE_ITEMS);
+		SCOPED_DRAW_EVENT(RHICmdList, DownsampleDepth, DEC_SCENE_ITEMS);
 
 		for (int32 ViewIndex = 0; ViewIndex < Views.Num(); ViewIndex++)
 		{
diff --git a/Engine/Source/Runtime/Renderer/Private/DistanceFieldShadowing.cpp b/Engine/Source/Runtime/Renderer/Private/DistanceFieldShadowing.cpp
--- a/Engine/Source/Runtime/Renderer/Private/DistanceFieldShadowing.cpp
+++ b/Engine/Source/Runtime/Renderer/Private/DistanceFieldShadowing.cpp
@@ -287,7 +287,7 @@ void FProjectedShadowInfo::RenderRayTracedDistanceFieldProjection(FRHICommandLis
 		&& DoesPlatformSupportDistanceFieldShadowing(GRHIShaderPlatform))
 	{
 		QUICK_SCOPE_CYCLE_COUNTER(STAT_RenderRayTracedDistanceFieldShadows);
-		SCOPED_DRAW_EVENT(RayTracedDistanceFieldShadow, DEC_SCENE_ITEMS);
+		SCOPED_DRAW_EVENT(RHICmdList, RayTracedDistanceFieldShadow, DEC_SCENE_ITEMS);
 
 		// Update the global distance field atlas
 		GDistanceFieldVolumeTextureAtlas.UpdateAllocations();
@@ -314,7 +314,7 @@ void FProjectedShadowInfo::RenderRayTracedDistanceFieldProjection(FRHICommandLis
 					uint32 GroupSizeY = FMath::DivideAndRoundUp(View.ViewRect.Size().Y / GAODownsampleFactor, GDistanceFieldAOTileSizeY);
 
 					{
-						SCOPED_DRAW_EVENT(RayTraceShadows, DEC_SCENE_ITEMS);
+						SCOPED_DRAW_EVENT(RHICmdList, RayTraceShadows, DEC_SCENE_ITEMS);
 						TShaderMapRef<FDistanceFieldShadowingCS> ComputeShader(View.ShaderMap);
 
 						RHICmdList.SetComputeShader(ComputeShader->GetComputeShader());
@@ -328,7 +328,7 @@ void FProjectedShadowInfo::RenderRayTracedDistanceFieldProjection(FRHICommandLis
 				{
 					GSceneRenderTargets.BeginRenderingLightAttenuation(RHICmdList);
 
-					SCOPED_DRAW_EVENT(Upsample, DEC_SCENE_ITEMS);
+					SCOPED_DRAW_EVENT(RHICmdList, Upsample, DEC_SCENE_ITEMS);
 
 					RHICmdList.SetViewport(View.ViewRect.Min.X, View.ViewRect.Min.Y, 0.0f, View.ViewRect.Max.X, View.ViewRect.Max.Y, 1.0f);
 					RHICmdList.SetRasterizerState(TStaticRasterizerState<FM_Solid, CM_None>::GetRHI());
diff --git a/Engine/Source/Runtime/Renderer/Private/DistanceFieldSurfaceCacheLighting.cpp b/Engine/Source/Runtime/Renderer/Private/DistanceFieldSurfaceCacheLighting.cpp
--- a/Engine/Source/Runtime/Renderer/Private/DistanceFieldSurfaceCacheLighting.cpp
+++ b/Engine/Source/Runtime/Renderer/Private/DistanceFieldSurfaceCacheLighting.cpp
@@ -853,7 +853,7 @@ void ComputeDistanceFieldNormal(FRHICommandListImmediate& RHICmdList, const TArr
 			uint32 GroupSizeY = FMath::DivideAndRoundUp(View.ViewRect.Size().Y / GAODownsampleFactor, GDistanceFieldAOTileSizeY);
 
 			{
-				SCOPED_DRAW_EVENT(ComputeNormalCS, DEC_SCENE_ITEMS);
+				SCOPED_DRAW_EVENT(RHICmdList, ComputeNormalCS, DEC_SCENE_ITEMS);
 				TShaderMapRef<FComputeDistanceFieldNormalCS> ComputeShader(View.ShaderMap);
 
 				RHICmdList.SetComputeShader(ComputeShader->GetComputeShader());
@@ -872,7 +872,7 @@ void ComputeDistanceFieldNormal(FRHICommandListImmediate& RHICmdList, const TArr
 		{
 			const FViewInfo& View = Views[ViewIndex];
 
-			SCOPED_DRAW_EVENT(ComputeNormal, DEC_SCENE_ITEMS);
+			SCOPED_DRAW_EVENT(RHICmdList, ComputeNormal, DEC_SCENE_ITEMS);
 
 			RHICmdList.SetViewport(0, 0, 0.0f, View.ViewRect.Width() / GAODownsampleFactor, View.ViewRect.Height() / GAODownsampleFactor, 1.0f);
 			RHICmdList.SetRasterizerState(TStaticRasterizerState<FM_Solid, CM_None>::GetRHI());
@@ -2003,7 +2003,7 @@ void UpdateHistory(
 			AllocateOrReuseAORenderTarget(NewHistory, HistoryRTName);
 
 			{
-				SCOPED_DRAW_EVENT(UpdateHistory, DEC_SCENE_ITEMS);
+				SCOPED_DRAW_EVENT(RHICmdList, UpdateHistory, DEC_SCENE_ITEMS);
 				SetRenderTarget(RHICmdList, NewHistory->GetRenderTargetItem().TargetableTexture, NULL);
 
 				RHICmdList.SetViewport(0, 0, 0.0f, View.ViewRect.Width() / GAODownsampleFactor, View.ViewRect.Height() / GAODownsampleFactor, 1.0f);
@@ -2069,7 +2069,7 @@ void PostProcessBentNormalAO(
 	}
 
 	{
-		SCOPED_DRAW_EVENT(AOCombine, DEC_SCENE_ITEMS);
+		SCOPED_DRAW_EVENT(RHICmdList, AOCombine, DEC_SCENE_ITEMS);
 
 		SetRenderTarget(RHICmdList, GAOFillGaps
 			? DistanceFieldAOBentNormal2->GetRenderTargetItem().TargetableTexture
@@ -2108,7 +2108,7 @@ void PostProcessBentNormalAO(
 
 	if (GAOFillGaps)
 	{
-		SCOPED_DRAW_EVENT(FillGaps, DEC_SCENE_ITEMS);
+		SCOPED_DRAW_EVENT(RHICmdList, FillGaps, DEC_SCENE_ITEMS);
 		SetRenderTarget(RHICmdList, DistanceFieldAOBentNormal->GetRenderTargetItem().TargetableTexture, NULL);
 
 		for (int32 ViewIndex = 0; ViewIndex < Views.Num(); ViewIndex++)
@@ -2221,7 +2221,7 @@ void UpsampleBentNormalAO(FRHICommandList& RHICmdList, const TArray<FViewInfo>&
 	{
 		const FViewInfo& View = Views[ViewIndex];
 
-		SCOPED_DRAW_EVENT(UpsampleAO, DEC_SCENE_ITEMS);
+		SCOPED_DRAW_EVENT(RHICmdList, UpsampleAO, DEC_SCENE_ITEMS);
 
 		RHICmdList.SetViewport(View.ViewRect.Min.X, View.ViewRect.Min.Y, 0.0f, View.ViewRect.Max.X, View.ViewRect.Max.Y, 1.0f);
 		RHICmdList.SetRasterizerState(TStaticRasterizerState<FM_Solid, CM_None>::GetRHI());
@@ -2408,7 +2408,7 @@ void GenerateBestSpacedVectors()
 
 FIntPoint BuildTileObjectLists(FRHICommandListImmediate& RHICmdList, FScene* Scene, TArray<FViewInfo>& Views, int32 NumObjects, const FDistanceFieldAOParameters& Parameters)
 {
-	SCOPED_DRAW_EVENT(BuildTileList, DEC_SCENE_ITEMS);
+	SCOPED_DRAW_EVENT(RHICmdList, BuildTileList, DEC_SCENE_ITEMS);
 	SetRenderTarget(RHICmdList, NULL, NULL);
 
 	FIntPoint TileListGroupSize;
@@ -2443,7 +2443,7 @@ FIntPoint BuildTileObjectLists(FRHICommandListImmediate& RHICmdList, FScene* Sce
 			}
 
 			{
-				SCOPED_DRAW_EVENT(BuildTileCones, DEC_SCENE_ITEMS);
+				SCOPED_DRAW_EVENT(RHICmdList, BuildTileCones, DEC_SCENE_ITEMS);
 				TShaderMapRef<FBuildTileConesCS> ComputeShader(View.ShaderMap);
 
 				RHICmdList.SetComputeShader(ComputeShader->GetComputeShader());
@@ -2454,7 +2454,7 @@ FIntPoint BuildTileObjectLists(FRHICommandListImmediate& RHICmdList, FScene* Sce
 			}
 
 			{
-				SCOPED_DRAW_EVENT(CullObjects, DEC_SCENE_ITEMS);
+				SCOPED_DRAW_EVENT(RHICmdList, CullObjects, DEC_SCENE_ITEMS);
 
 				TShaderMapRef<FObjectCullVS> VertexShader(View.ShaderMap);
 				TShaderMapRef<FObjectCullPS> PixelShader(View.ShaderMap);
@@ -2549,7 +2549,7 @@ void RenderIrradianceCacheInterpolation(
 	check(!(bFinalInterpolation && DepthLevel != 0));
 
 	{
-		SCOPED_DRAW_EVENT(IrradianceCacheSplat, DEC_SCENE_ITEMS);
+		SCOPED_DRAW_EVENT(RHICmdList, IrradianceCacheSplat, DEC_SCENE_ITEMS);
 		const int32 MaxAllowedLevel = GAOMaxSupportedLevel;
 
 		SetRenderTarget(RHICmdList, InterpolationTarget->GetRenderTargetItem().TargetableTexture, NULL);
@@ -2770,7 +2770,7 @@ bool FDeferredShadingSceneRenderer::RenderDistanceFieldAOSurfaceCache(FRHIComman
 		&& View.State)
 	{
 		QUICK_SCOPE_CYCLE_COUNTER(STAT_RenderDistanceFieldAOSurfaceCache);
-		SCOPED_DRAW_EVENT(DistanceFieldAO, DEC_SCENE_ITEMS);
+		SCOPED_DRAW_EVENT(RHICmdList, DistanceFieldAO, DEC_SCENE_ITEMS);
 
 		// Update the global distance field atlas
 		GDistanceFieldVolumeTextureAtlas.UpdateAllocations();
@@ -2828,7 +2828,7 @@ bool FDeferredShadingSceneRenderer::RenderDistanceFieldAOSurfaceCache(FRHIComman
 
 				if (GAOReuseAcrossFrames && GAOTrimOldRecordsFraction > 0)
 				{
-					SCOPED_DRAW_EVENT(TrimRecords, DEC_SCENE_ITEMS);
+					SCOPED_DRAW_EVENT(RHICmdList, TrimRecords, DEC_SCENE_ITEMS);
 
 					// Copy and trim last frame's surface cache samples
 					for (int32 DepthLevel = FMath::Min(GAOMaxLevel, MaxAllowedLevel); DepthLevel >= FMath::Max(GAOMinLevel, 0); DepthLevel--)
@@ -2864,7 +2864,7 @@ bool FDeferredShadingSceneRenderer::RenderDistanceFieldAOSurfaceCache(FRHIComman
 				{
 					int32 DestLevelDownsampleFactor = GAODownsampleFactor * (1 << (DepthLevel * 2));
 
-					SCOPED_DRAW_EVENTF(Level, DEC_SCENE_ITEMS, TEXT("Level_%d"), DepthLevel);
+					SCOPED_DRAW_EVENTF(RHICmdList, Level, DEC_SCENE_ITEMS, TEXT("Level_%d"), DepthLevel);
 
 					TRefCountPtr<IPooledRenderTarget> DistanceFieldAOIrradianceCacheSplat;
 
@@ -2878,7 +2878,7 @@ bool FDeferredShadingSceneRenderer::RenderDistanceFieldAOSurfaceCache(FRHIComman
 					RenderIrradianceCacheInterpolation(RHICmdList, Views, DistanceFieldAOIrradianceCacheSplat, DistanceFieldNormal->GetRenderTargetItem(), DepthLevel, DestLevelDownsampleFactor, Parameters, false);
 
 					{
-						SCOPED_DRAW_EVENT(PopulateIrradianceCache, DEC_SCENE_ITEMS);
+						SCOPED_DRAW_EVENT(RHICmdList, PopulateIrradianceCache, DEC_SCENE_ITEMS);
 						SetRenderTarget(RHICmdList, NULL, NULL);
 
 						// Save off the current record count before adding any more
@@ -2912,7 +2912,7 @@ bool FDeferredShadingSceneRenderer::RenderDistanceFieldAOSurfaceCache(FRHIComman
 					}
 
 					{
-						SCOPED_DRAW_EVENT(ShadeIrradianceCache, DEC_SCENE_ITEMS);
+						SCOPED_DRAW_EVENT(RHICmdList, ShadeIrradianceCache, DEC_SCENE_ITEMS);
 
 						{	
 							TShaderMapRef<FSetupFinalGatherIndirectArgumentsCS> ComputeShader(View.ShaderMap);
@@ -2947,7 +2947,7 @@ bool FDeferredShadingSceneRenderer::RenderDistanceFieldAOSurfaceCache(FRHIComman
 
 				// Splat the surface cache records onto the opaque pixels, using less strict weighting so the lighting is smoothed in world space
 				{
-					SCOPED_DRAW_EVENT(FinalIrradianceCacheSplat, DEC_SCENE_ITEMS);
+					SCOPED_DRAW_EVENT(RHICmdList, FinalIrradianceCacheSplat, DEC_SCENE_ITEMS);
 					RenderIrradianceCacheInterpolation(RHICmdList, Views, IrradianceCacheAccumulation, DistanceFieldNormal->GetRenderTargetItem(), 0, GAODownsampleFactor, Parameters, true);
 				}
 
@@ -3132,7 +3132,7 @@ void FDeferredShadingSceneRenderer::RenderMeshDistanceFieldVisualization(FRHICom
 		&& Views.Num() == 1)
 	{
 		QUICK_SCOPE_CYCLE_COUNTER(STAT_RenderMeshDistanceFieldVis);
-		SCOPED_DRAW_EVENT(VisualizeMeshDistanceFields, DEC_SCENE_ITEMS);
+		SCOPED_DRAW_EVENT(RHICmdList, VisualizeMeshDistanceFields, DEC_SCENE_ITEMS);
 
 		// Update the global distance field atlas
 		GDistanceFieldVolumeTextureAtlas.UpdateAllocations();
@@ -3167,7 +3167,7 @@ void FDeferredShadingSceneRenderer::RenderMeshDistanceFieldVisualization(FRHICom
 						uint32 GroupSizeY = FMath::DivideAndRoundUp(View.ViewRect.Size().Y / GAODownsampleFactor, GDistanceFieldAOTileSizeY);
 
 						{
-							SCOPED_DRAW_EVENT(VisualizeMeshDistanceFieldCS, DEC_SCENE_ITEMS);
+							SCOPED_DRAW_EVENT(RHICmdList, VisualizeMeshDistanceFieldCS, DEC_SCENE_ITEMS);
 							TShaderMapRef<FVisualizeMeshDistanceFieldCS> ComputeShader(View.ShaderMap);
 
 							RHICmdList.SetComputeShader(ComputeShader->GetComputeShader());
@@ -3186,7 +3186,7 @@ void FDeferredShadingSceneRenderer::RenderMeshDistanceFieldVisualization(FRHICom
 					{
 						const FViewInfo& View = Views[ViewIndex];
 
-						SCOPED_DRAW_EVENT(UpsampleAO, DEC_SCENE_ITEMS);
+						SCOPED_DRAW_EVENT(RHICmdList, UpsampleAO, DEC_SCENE_ITEMS);
 
 						RHICmdList.SetViewport(View.ViewRect.Min.X, View.ViewRect.Min.Y, 0.0f, View.ViewRect.Max.X, View.ViewRect.Max.Y, 1.0f);
 						RHICmdList.SetRasterizerState(TStaticRasterizerState<FM_Solid, CM_None>::GetRHI());
@@ -3299,7 +3299,7 @@ void FDeferredShadingSceneRenderer::RenderDynamicSkyLighting(FRHICommandListImme
 		&& !IsSimpleDynamicLightingEnabled() 
 		&& !ViewFamily.EngineShowFlags.VisualizeLightCulling)
 	{
-		SCOPED_DRAW_EVENT(SkyLightDiffuse, DEC_SCENE_ITEMS);
+		SCOPED_DRAW_EVENT(RHICmdList, SkyLightDiffuse, DEC_SCENE_ITEMS);
 
 		bool bApplyShadowing = false;
 
diff --git a/Engine/Source/Runtime/Renderer/Private/DistortionRendering.cpp b/Engine/Source/Runtime/Renderer/Private/DistortionRendering.cpp
--- a/Engine/Source/Runtime/Renderer/Private/DistortionRendering.cpp
+++ b/Engine/Source/Runtime/Renderer/Private/DistortionRendering.cpp
@@ -497,7 +497,7 @@ void TDistortionMeshDrawingPolicy<DistortMeshPolicy>::SetMeshRenderState(
 
 	const FMeshBatchElement& BatchElement = Mesh.Elements[BatchElementIndex];
 
-	EmitMeshDrawEvents(PrimitiveSceneProxy, Mesh);
+	EmitMeshDrawEvents(RHICmdList, PrimitiveSceneProxy, Mesh);
 
 	// Set transforms
 	VertexShader->SetMesh(RHICmdList, VertexFactory,View,PrimitiveSceneProxy,BatchElement);
@@ -781,7 +781,7 @@ void FDistortionPrimSet::AddScenePrimitive(FPrimitiveSceneProxy* PrimitiveSceneP
  */
 void FSceneRenderer::RenderDistortion(FRHICommandListImmediate& RHICmdList)
 {
-	SCOPED_DRAW_EVENT(Distortion, DEC_SCENE_ITEMS);
+	SCOPED_DRAW_EVENT(RHICmdList, Distortion, DEC_SCENE_ITEMS);
 
 	// do we need to render the distortion pass?
 	bool bRender=false;
@@ -802,7 +802,7 @@ void FSceneRenderer::RenderDistortion(FRHICommandListImmediate& RHICmdList)
 	// Render accumulated distortion offsets
 	if( bRender)
 	{
-		SCOPED_DRAW_EVENT(DistortionAccum, DEC_SCENE_ITEMS);
+		SCOPED_DRAW_EVENT(RHICmdList, DistortionAccum, DEC_SCENE_ITEMS);
 
 		// Create a texture to store the resolved light attenuation values, and a render-targetable surface to hold the unresolved light attenuation values.
 		{
@@ -824,7 +824,7 @@ void FSceneRenderer::RenderDistortion(FRHICommandListImmediate& RHICmdList)
 
 			for(int32 ViewIndex = 0;ViewIndex < Views.Num();ViewIndex++)
 			{
-				SCOPED_CONDITIONAL_DRAW_EVENTF(EventView, Views.Num() > 1, DEC_SCENE_ITEMS, TEXT("View%d"), ViewIndex);
+				SCOPED_CONDITIONAL_DRAW_EVENTF(RHICmdList, EventView, Views.Num() > 1, DEC_SCENE_ITEMS, TEXT("View%d"), ViewIndex);
 
 				FViewInfo& View = Views[ViewIndex];
 				// viewport to match view size
@@ -860,7 +860,7 @@ void FSceneRenderer::RenderDistortion(FRHICommandListImmediate& RHICmdList)
 
 	if(bDirty)
 	{
-		SCOPED_DRAW_EVENT(DistortionApply, DEC_SCENE_ITEMS);
+		SCOPED_DRAW_EVENT(RHICmdList, DistortionApply, DEC_SCENE_ITEMS);
 
 		GSceneRenderTargets.ResolveSceneColor(RHICmdList, FResolveRect(0, 0, ViewFamily.FamilySizeX, ViewFamily.FamilySizeY));
 
@@ -874,7 +874,7 @@ void FSceneRenderer::RenderDistortion(FRHICommandListImmediate& RHICmdList)
 		// Apply distortion as a full-screen pass		
 		for(int32 ViewIndex = 0, Num = Views.Num(); ViewIndex < Num; ++ViewIndex)
 		{
-			SCOPED_CONDITIONAL_DRAW_EVENTF(EventView, Views.Num() > 1, DEC_SCENE_ITEMS, TEXT("View%d"), ViewIndex);
+			SCOPED_CONDITIONAL_DRAW_EVENTF(RHICmdList, EventView, Views.Num() > 1, DEC_SCENE_ITEMS, TEXT("View%d"), ViewIndex);
 
 			FViewInfo& View = Views[ViewIndex];
 
diff --git a/Engine/Source/Runtime/Renderer/Private/DrawingPolicy.cpp b/Engine/Source/Runtime/Renderer/Private/DrawingPolicy.cpp
--- a/Engine/Source/Runtime/Renderer/Private/DrawingPolicy.cpp
+++ b/Engine/Source/Runtime/Renderer/Private/DrawingPolicy.cpp
@@ -53,7 +53,7 @@ void FMeshDrawingPolicy::SetMeshRenderState(
 	const ContextDataType PolicyContext
 	) const
 {
-	EmitMeshDrawEvents(PrimitiveSceneProxy, Mesh);
+	EmitMeshDrawEvents(RHICmdList, PrimitiveSceneProxy, Mesh);
 
 		// Use bitwise logic ops to avoid branches
 	RHICmdList.SetRasterizerState( GetStaticRasterizerState<true>(
@@ -65,7 +65,7 @@ void FMeshDrawingPolicy::SetMeshRenderState(
 void FMeshDrawingPolicy::DrawMesh(FRHICommandList& RHICmdList, const FMeshBatch& Mesh, int32 BatchElementIndex) const
 {
 	INC_DWORD_STAT(STAT_MeshDrawCalls);
-	SCOPED_CONDITIONAL_DRAW_EVENTF(MeshEvent, GEmitMeshDrawEvent != 0, DEC_SCENE_ITEMS, TEXT("Mesh Draw"));
+	SCOPED_CONDITIONAL_DRAW_EVENTF(RHICmdList, MeshEvent, GEmitMeshDrawEvent != 0, DEC_SCENE_ITEMS, TEXT("Mesh Draw"));
 
 	const FMeshBatchElement& BatchElement = Mesh.Elements[BatchElementIndex];
 
diff --git a/Engine/Source/Runtime/Renderer/Private/FogRendering.cpp b/Engine/Source/Runtime/Renderer/Private/FogRendering.cpp
--- a/Engine/Source/Runtime/Renderer/Private/FogRendering.cpp
+++ b/Engine/Source/Runtime/Renderer/Private/FogRendering.cpp
@@ -293,7 +293,7 @@ bool FDeferredShadingSceneRenderer::RenderFog(FRHICommandListImmediate& RHICmdLi
 {
 	if (Scene->ExponentialFogs.Num() > 0)
 	{
-		SCOPED_DRAW_EVENT(Fog, DEC_SCENE_ITEMS);
+		SCOPED_DRAW_EVENT(RHICmdList, Fog, DEC_SCENE_ITEMS);
 
 		static const FVector2D Vertices[4] =
 		{
diff --git a/Engine/Source/Runtime/Renderer/Private/ForwardBasePassRendering.cpp b/Engine/Source/Runtime/Renderer/Private/ForwardBasePassRendering.cpp
--- a/Engine/Source/Runtime/Renderer/Private/ForwardBasePassRendering.cpp
+++ b/Engine/Source/Runtime/Renderer/Private/ForwardBasePassRendering.cpp
@@ -297,7 +297,7 @@ EBasePassSort::Type GetSortMode()
 
 void FForwardShadingSceneRenderer::RenderForwardShadingBasePass(FRHICommandListImmediate& RHICmdList)
 {
-	SCOPED_DRAW_EVENT(BasePass, DEC_SCENE_ITEMS);
+	SCOPED_DRAW_EVENT(RHICmdList, BasePass, DEC_SCENE_ITEMS);
 	SCOPE_CYCLE_COUNTER(STAT_BasePassDrawTime);
 
 	EBasePassSort::Type SortMode = GetSortMode();
@@ -323,7 +323,7 @@ void FForwardShadingSceneRenderer::RenderForwardShadingBasePass(FRHICommandListI
 	// Draw the scene's emissive and light-map color.
 	for (int32 ViewIndex = 0; ViewIndex < Views.Num(); ViewIndex++)
 	{
-		SCOPED_CONDITIONAL_DRAW_EVENTF(EventView, Views.Num() > 1, DEC_SCENE_ITEMS, TEXT("View%d"), ViewIndex);
+		SCOPED_CONDITIONAL_DRAW_EVENTF(RHICmdList, EventView, Views.Num() > 1, DEC_SCENE_ITEMS, TEXT("View%d"), ViewIndex);
 		FViewInfo& View = Views[ViewIndex];
 
 		// Opaque blending
@@ -354,7 +354,7 @@ void FForwardShadingSceneRenderer::RenderForwardShadingBasePass(FRHICommandListI
 
 		{
 			SCOPE_CYCLE_COUNTER(STAT_DynamicPrimitiveDrawTime);
-			SCOPED_DRAW_EVENT(Dynamic, DEC_SCENE_ITEMS);
+			SCOPED_DRAW_EVENT(RHICmdList, Dynamic, DEC_SCENE_ITEMS);
 
 			const bool bUseGetMeshElements = ShouldUseGetDynamicMeshElements();
 
diff --git a/Engine/Source/Runtime/Renderer/Private/ForwardShadingRenderer.cpp b/Engine/Source/Runtime/Renderer/Private/ForwardShadingRenderer.cpp
--- a/Engine/Source/Runtime/Renderer/Private/ForwardShadingRenderer.cpp
+++ b/Engine/Source/Runtime/Renderer/Private/ForwardShadingRenderer.cpp
@@ -24,7 +24,7 @@ FForwardShadingSceneRenderer::FForwardShadingSceneRenderer(const FSceneViewFamil
  */
 void FForwardShadingSceneRenderer::InitViews(FRHICommandListImmediate& RHICmdList)
 {
-	SCOPED_DRAW_EVENT(InitViews, DEC_SCENE_ITEMS);
+	SCOPED_DRAW_EVENT(RHICmdList, InitViews, DEC_SCENE_ITEMS);
 
 	SCOPE_CYCLE_COUNTER(STAT_InitViewsTime);
 
@@ -161,11 +161,11 @@ void FForwardShadingSceneRenderer::Render(FRHICommandListImmediate& RHICmdList)
 		}
 		else
 		{
-			SCOPED_DRAW_EVENT(PostProcessing, DEC_SCENE_ITEMS);
+			SCOPED_DRAW_EVENT(RHICmdList, PostProcessing, DEC_SCENE_ITEMS);
 			SCOPE_CYCLE_COUNTER(STAT_FinishRenderViewTargetTime);
 			for(int32 ViewIndex = 0;ViewIndex < Views.Num();ViewIndex++)
 			{	
-				SCOPED_CONDITIONAL_DRAW_EVENTF(EventView, Views.Num() > 1, DEC_SCENE_ITEMS, TEXT("View%d"), ViewIndex);
+				SCOPED_CONDITIONAL_DRAW_EVENTF(RHICmdList, EventView, Views.Num() > 1, DEC_SCENE_ITEMS, TEXT("View%d"), ViewIndex);
 				GPostProcessing.ProcessES2(RHICmdList, Views[ViewIndex], bOnChipSunMask);
 			}
 		}
diff --git a/Engine/Source/Runtime/Renderer/Private/ForwardTranslucentRendering.cpp b/Engine/Source/Runtime/Renderer/Private/ForwardTranslucentRendering.cpp
--- a/Engine/Source/Runtime/Renderer/Private/ForwardTranslucentRendering.cpp
+++ b/Engine/Source/Runtime/Renderer/Private/ForwardTranslucentRendering.cpp
@@ -47,7 +47,7 @@ FGlobalBoundShaderState ForwardCopySceneAlphaBoundShaderState;
 
 void FForwardShadingSceneRenderer::CopySceneAlpha(FRHICommandListImmediate& RHICmdList, const FViewInfo& View)
 {
-	SCOPED_DRAW_EVENTF(EventCopy, DEC_SCENE_ITEMS, TEXT("CopySceneAlpha"));
+	SCOPED_DRAW_EVENTF(RHICmdList, EventCopy, DEC_SCENE_ITEMS, TEXT("CopySceneAlpha"));
 	RHICmdList.SetRasterizerState(TStaticRasterizerState<FM_Solid, CM_None>::GetRHI());
 	RHICmdList.SetDepthStencilState(TStaticDepthStencilState<false,CF_Always>::GetRHI());
 	RHICmdList.SetBlendState(TStaticBlendState<>::GetRHI());
@@ -329,11 +329,11 @@ void FForwardShadingSceneRenderer::RenderTranslucency(FRHICommandListImmediate&
 		const bool bGammaSpace = !IsMobileHDR();
 		const bool bLinearHDR64 = !bGammaSpace && !IsMobileHDR32bpp();
 
-		SCOPED_DRAW_EVENT(Translucency, DEC_SCENE_ITEMS);
+		SCOPED_DRAW_EVENT(RHICmdList, Translucency, DEC_SCENE_ITEMS);
 
 		for (int32 ViewIndex = 0; ViewIndex < Views.Num(); ViewIndex++)
 		{
-			SCOPED_CONDITIONAL_DRAW_EVENTF(EventView, Views.Num() > 1, DEC_SCENE_ITEMS, TEXT("View%d"), ViewIndex);
+			SCOPED_CONDITIONAL_DRAW_EVENTF(RHICmdList, EventView, Views.Num() > 1, DEC_SCENE_ITEMS, TEXT("View%d"), ViewIndex);
 
 			const FViewInfo& View = Views[ViewIndex];
 
diff --git a/Engine/Source/Runtime/Renderer/Private/GPUBenchmark.cpp b/Engine/Source/Runtime/Renderer/Private/GPUBenchmark.cpp
--- a/Engine/Source/Runtime/Renderer/Private/GPUBenchmark.cpp
+++ b/Engine/Source/Runtime/Renderer/Private/GPUBenchmark.cpp
@@ -168,7 +168,7 @@ void RunBenchmarkShader(FRHICommandList& RHICmdList, const FSceneView& View, TRe
 
 void RunBenchmarkShader(FRHICommandListImmediate& RHICmdList, const FSceneView& View, uint32 MethodId, TRefCountPtr<IPooledRenderTarget>& Src, float WorkScale)
 {
-	SCOPED_DRAW_EVENTF(Benchmark, DEC_SCENE_ITEMS, TEXT("Benchmark Method:%d"), MethodId);
+	SCOPED_DRAW_EVENTF(RHICmdList, Benchmark, DEC_SCENE_ITEMS, TEXT("Benchmark Method:%d"), MethodId);
 
 	switch(MethodId)
 	{
diff --git a/Engine/Source/Runtime/Renderer/Private/GammaCorrection.cpp b/Engine/Source/Runtime/Renderer/Private/GammaCorrection.cpp
--- a/Engine/Source/Runtime/Renderer/Private/GammaCorrection.cpp
+++ b/Engine/Source/Runtime/Renderer/Private/GammaCorrection.cpp
@@ -90,7 +90,7 @@ void FSceneRenderer::GammaCorrectToViewportRenderTarget(FRHICommandList& RHICmdL
 		ViewFamily.bDeferClear = false;
 	}
 
-	SCOPED_DRAW_EVENT(GammaCorrection, DEC_SCENE_ITEMS);
+	SCOPED_DRAW_EVENT(RHICmdList, GammaCorrection, DEC_SCENE_ITEMS);
 
 	// turn off culling and blending
 	RHICmdList.SetRasterizerState(TStaticRasterizerState<FM_Solid, CM_None>::GetRHI());
diff --git a/Engine/Source/Runtime/Renderer/Private/LightFunctionRendering.cpp b/Engine/Source/Runtime/Renderer/Private/LightFunctionRendering.cpp
--- a/Engine/Source/Runtime/Renderer/Private/LightFunctionRendering.cpp
+++ b/Engine/Source/Runtime/Renderer/Private/LightFunctionRendering.cpp
@@ -245,7 +245,7 @@ bool FDeferredShadingSceneRenderer::RenderLightFunctionForMaterial(FRHICommandLi
 		bRenderedLightFunction = true;
 
 		const FMaterial* Material = MaterialProxy->GetMaterial(Scene->GetFeatureLevel());
-		SCOPED_DRAW_EVENTF(LightFunction, DEC_SCENE_ITEMS, TEXT("LightFunction Material=%s"), *Material->GetFriendlyName());
+		SCOPED_DRAW_EVENTF(RHICmdList, LightFunction, DEC_SCENE_ITEMS, TEXT("LightFunction Material=%s"), *Material->GetFriendlyName());
 
 		const FMaterialShaderMap* MaterialShaderMap = Material->GetRenderingThreadShaderMap();
 		FLightFunctionVS* VertexShader = MaterialShaderMap->GetShader<FLightFunctionVS>();
@@ -258,7 +258,7 @@ bool FDeferredShadingSceneRenderer::RenderLightFunctionForMaterial(FRHICommandLi
 		// Render to the light attenuation buffer for all views.
 		for (int32 ViewIndex = 0; ViewIndex < Views.Num(); ViewIndex++)
 		{
-			SCOPED_CONDITIONAL_DRAW_EVENTF(EventView, Views.Num() > 1, DEC_SCENE_ITEMS, TEXT("View%d"), ViewIndex);
+			SCOPED_CONDITIONAL_DRAW_EVENTF(RHICmdList, EventView, Views.Num() > 1, DEC_SCENE_ITEMS, TEXT("View%d"), ViewIndex);
 
 			const FViewInfo& View = Views[ViewIndex];
 
diff --git a/Engine/Source/Runtime/Renderer/Private/LightMapDensityRendering.cpp b/Engine/Source/Runtime/Renderer/Private/LightMapDensityRendering.cpp
--- a/Engine/Source/Runtime/Renderer/Private/LightMapDensityRendering.cpp
+++ b/Engine/Source/Runtime/Renderer/Private/LightMapDensityRendering.cpp
@@ -40,12 +40,12 @@ bool FDeferredShadingSceneRenderer::RenderLightMapDensities(FRHICommandListImmed
 	bool bDirty=0;
 	if (Scene->GetFeatureLevel() >= ERHIFeatureLevel::SM4)
 	{
-		SCOPED_DRAW_EVENT(LightMapDensity, DEC_SCENE_ITEMS);
+		SCOPED_DRAW_EVENT(RHICmdList, LightMapDensity, DEC_SCENE_ITEMS);
 
 		// Draw the scene's emissive and light-map color.
 		for(int32 ViewIndex = 0;ViewIndex < Views.Num();ViewIndex++)
 		{
-			SCOPED_CONDITIONAL_DRAW_EVENTF(EventView, Views.Num() > 1, DEC_SCENE_ITEMS, TEXT("View%d"), ViewIndex);
+			SCOPED_CONDITIONAL_DRAW_EVENTF(RHICmdList, EventView, Views.Num() > 1, DEC_SCENE_ITEMS, TEXT("View%d"), ViewIndex);
 
 			FViewInfo& View = Views[ViewIndex];
 
@@ -56,7 +56,7 @@ bool FDeferredShadingSceneRenderer::RenderLightMapDensities(FRHICommandListImmed
 			RHICmdList.SetViewport(View.ViewRect.Min.X, View.ViewRect.Min.Y, 0, View.ViewRect.Max.X, View.ViewRect.Max.Y, 1);
 
 			{
-				SCOPED_DRAW_EVENT(Dynamic, DEC_SCENE_ITEMS);
+				SCOPED_DRAW_EVENT(RHICmdList, Dynamic, DEC_SCENE_ITEMS);
 
 				const bool bUseGetMeshElements = ShouldUseGetDynamicMeshElements();
 
diff --git a/Engine/Source/Runtime/Renderer/Private/LightPropagationVolume.cpp b/Engine/Source/Runtime/Renderer/Private/LightPropagationVolume.cpp
--- a/Engine/Source/Runtime/Renderer/Private/LightPropagationVolume.cpp
+++ b/Engine/Source/Runtime/Renderer/Private/LightPropagationVolume.cpp
@@ -929,7 +929,7 @@ void FLightPropagationVolume::Clear(FRHICommandListImmediate& RHICmdList, FViewI
 		return;
 	}
 
-	SCOPED_DRAW_EVENT(LpvClear, DEC_LIGHT);
+	SCOPED_DRAW_EVENT(RHICmdList, LpvClear, DEC_LIGHT);
 
 	if ( !LpvWriteUniformBuffer.IsInitialized() )
 	{
@@ -1042,7 +1042,7 @@ void FLightPropagationVolume::InjectDirectionalLightRSM(
 {
 	const FLightSceneProxy* LightProxy = ProjectedShadowInfo.LightSceneInfo->Proxy;
 	{
-		SCOPED_DRAW_EVENT(LpvInjectDirectionalLightRSM, DEC_LIGHT);
+		SCOPED_DRAW_EVENT(RHICmdList, LpvInjectDirectionalLightRSM, DEC_LIGHT);
 
 		SetVplInjectionConstants(ProjectedShadowInfo, LightProxy );
 
@@ -1062,7 +1062,7 @@ void FLightPropagationVolume::InjectDirectionalLightRSM(
 	// If this is the first directional light, build the geometry volume with it
 	if ( !GeometryVolumeGenerated && SecondaryOcclusionStrength > 0.001f )
 	{
-		SCOPED_DRAW_EVENT(LpvBuildGeometryVolume, DEC_LIGHT);
+		SCOPED_DRAW_EVENT(RHICmdList, LpvBuildGeometryVolume, DEC_LIGHT);
 		GeometryVolumeGenerated = true;
 		FVector LightDirection( 0.0f, 0.0f, 1.0f );
 		if ( LightProxy )
@@ -1101,7 +1101,7 @@ void FLightPropagationVolume::Propagate(FRHICommandListImmediate& RHICmdList, FV
 	// Inject the VPLs into the LPV
 	if ( mInjectedLightCount )
 	{
-		SCOPED_DRAW_EVENT(LpvAccumulateVplLists, DEC_LIGHT);
+		SCOPED_DRAW_EVENT(RHICmdList, LpvAccumulateVplLists, DEC_LIGHT);
 		mWriteBufferIndex = 1-mWriteBufferIndex; // Swap buffers with each iteration
 
 		TShaderMapRef<FLpvInject_AccumulateVplListsCS> Shader(View.ShaderMap);
@@ -1122,7 +1122,7 @@ void FLightPropagationVolume::Propagate(FRHICommandListImmediate& RHICmdList, FV
 	{
 		static int PropagationStep = 0;
 
-		SCOPED_DRAW_EVENT(LpvPropagate, DEC_LIGHT);
+		SCOPED_DRAW_EVENT(RHICmdList, LpvPropagate, DEC_LIGHT);
 		static int numIterations = 3;
 		for ( int i=0; i<numIterations; i++ )
 		{
@@ -1209,7 +1209,7 @@ void FLightPropagationVolume::InjectLightDirect(FRHICommandListImmediate& RHICmd
 	if ( GeometryVolumeGenerated )
 	{
 		// Inject the VPLs into the LPV
-		SCOPED_DRAW_EVENT(LpvDirectLightInjection, DEC_LIGHT);
+		SCOPED_DRAW_EVENT(RHICmdList, LpvDirectLightInjection, DEC_LIGHT);
 
 		FLpvDirectLightInjectParameters InjectUniformBufferParams;
 		InjectUniformBufferParams.LightColor = Light.GetColor() * Light.GetIndirectLightingScale();
diff --git a/Engine/Source/Runtime/Renderer/Private/LightPropagationVolumeVisualisation.cpp b/Engine/Source/Runtime/Renderer/Private/LightPropagationVolumeVisualisation.cpp
--- a/Engine/Source/Runtime/Renderer/Private/LightPropagationVolumeVisualisation.cpp
+++ b/Engine/Source/Runtime/Renderer/Private/LightPropagationVolumeVisualisation.cpp
@@ -250,7 +250,7 @@ IMPLEMENT_SHADER_TYPE(,FLpvVisualisePS,TEXT("LPVVisualise"),TEXT("PShader"),SF_P
 
 void FLightPropagationVolume::Visualise(FRHICommandList& RHICmdList, const FViewInfo& View) const
 {
-	SCOPED_DRAW_EVENT(LpvVisualise, DEC_LIGHT);
+	SCOPED_DRAW_EVENT(RHICmdList, LpvVisualise, DEC_LIGHT);
 	check(View.GetFeatureLevel() == ERHIFeatureLevel::SM5);
 
 	TShaderMapRef<FLpvVisualiseVS> VertexShader(View.ShaderMap);
diff --git a/Engine/Source/Runtime/Renderer/Private/LightRendering.cpp b/Engine/Source/Runtime/Renderer/Private/LightRendering.cpp
--- a/Engine/Source/Runtime/Renderer/Private/LightRendering.cpp
+++ b/Engine/Source/Runtime/Renderer/Private/LightRendering.cpp
@@ -279,7 +279,7 @@ uint32 GetShadowQuality();
 /** Renders the scene's lighting. */
 void FDeferredShadingSceneRenderer::RenderLights(FRHICommandListImmediate& RHICmdList)
 {
-	SCOPED_DRAW_EVENT(Lights, DEC_SCENE_ITEMS);
+	SCOPED_DRAW_EVENT(RHICmdList, Lights, DEC_SCENE_ITEMS);
 
 	if(IsSimpleDynamicLightingEnabled())
 	{
@@ -340,7 +340,7 @@ void FDeferredShadingSceneRenderer::RenderLights(FRHICommandListImmediate& RHICm
 	SortedLights.Sort( FCompareFSortedLightSceneInfo() );
 
 	{
-		SCOPED_DRAW_EVENT(DirectLighting, DEC_SCENE_ITEMS);
+		SCOPED_DRAW_EVENT(RHICmdList, DirectLighting, DEC_SCENE_ITEMS);
 
 		int32 AttenuationLightStart = SortedLights.Num();
 		int32 SupportedByTiledDeferredLightEnd = SortedLights.Num();
@@ -382,7 +382,7 @@ void FDeferredShadingSceneRenderer::RenderLights(FRHICommandListImmediate& RHICm
 		
 		if(ViewFamily.EngineShowFlags.DirectLighting)
 		{
-			SCOPED_DRAW_EVENT(NonShadowedLights, DEC_SCENE_ITEMS);
+			SCOPED_DRAW_EVENT(RHICmdList, NonShadowedLights, DEC_SCENE_ITEMS);
 			INC_DWORD_STAT_BY(STAT_NumUnshadowedLights, AttenuationLightStart);
 			GSceneRenderTargets.SetLightAttenuationMode(false);
 
@@ -422,7 +422,7 @@ void FDeferredShadingSceneRenderer::RenderLights(FRHICommandListImmediate& RHICm
 			}
 
 			{
-				SCOPED_DRAW_EVENT(StandardDeferredLighting, DEC_SCENE_ITEMS);
+				SCOPED_DRAW_EVENT(RHICmdList, StandardDeferredLighting, DEC_SCENE_ITEMS);
 
 				GSceneRenderTargets.BeginRenderingSceneColor(RHICmdList);
 
@@ -443,13 +443,13 @@ void FDeferredShadingSceneRenderer::RenderLights(FRHICommandListImmediate& RHICm
 				if (AttenuationLightStart)
 				{
 					// Inject non-shadowed, non-light function lights in to the volume.
-					SCOPED_DRAW_EVENT(InjectNonShadowedTranslucentLighting, DEC_SCENE_ITEMS);
+					SCOPED_DRAW_EVENT(RHICmdList, InjectNonShadowedTranslucentLighting, DEC_SCENE_ITEMS);
 					InjectTranslucentVolumeLightingArray(RHICmdList, SortedLights, AttenuationLightStart);
 				}
 				
 				if (SimpleLights.InstanceData.Num() > 0)
 				{
-					SCOPED_DRAW_EVENT(InjectSimpleLightsTranslucentLighting, DEC_SCENE_ITEMS);
+					SCOPED_DRAW_EVENT(RHICmdList, InjectSimpleLightsTranslucentLighting, DEC_SCENE_ITEMS);
 					InjectSimpleTranslucentVolumeLightingArray(RHICmdList, SimpleLights);
 				}
 			}
@@ -471,7 +471,7 @@ void FDeferredShadingSceneRenderer::RenderLights(FRHICommandListImmediate& RHICm
 
 			FString LightNameWithLevel;
 			GetLightNameForDrawEvent(LightSceneInfo.Proxy, LightNameWithLevel);
-			SCOPED_DRAW_EVENTF(EventLightPass, DEC_SCENE_ITEMS, *LightNameWithLevel);
+			SCOPED_DRAW_EVENTF(RHICmdList, EventLightPass, DEC_SCENE_ITEMS, *LightNameWithLevel);
 
 			// Do not resolve to scene color texture, this is done lazily
 			GSceneRenderTargets.FinishRenderingSceneColor(RHICmdList, false);
@@ -528,7 +528,7 @@ void FDeferredShadingSceneRenderer::RenderLights(FRHICommandListImmediate& RHICm
 			
 			if(bDirectLighting && !bInjectedTranslucentVolume)
 			{
-				SCOPED_DRAW_EVENT(InjectTranslucentVolume, DEC_SCENE_ITEMS);
+				SCOPED_DRAW_EVENT(RHICmdList, InjectTranslucentVolume, DEC_SCENE_ITEMS);
 				// Accumulate this light's unshadowed contribution to the translucency lighting volume
 				InjectTranslucentVolumeLighting(RHICmdList, LightSceneInfo, NULL);
 			}
@@ -718,7 +718,7 @@ void FDeferredShadingSceneRenderer::RenderLight(FRHICommandList& RHICmdList, con
 {
 	SCOPE_CYCLE_COUNTER(STAT_DirectLightRenderingTime);
 	INC_DWORD_STAT(STAT_NumLightsUsingStandardDeferred);
-	SCOPED_CONDITIONAL_DRAW_EVENT(StandardDeferredLighting, bIssueDrawEvent, DEC_SCENE_ITEMS);
+	SCOPED_CONDITIONAL_DRAW_EVENT(RHICmdList, StandardDeferredLighting, bIssueDrawEvent, DEC_SCENE_ITEMS);
 
 	// Use additive blending for color
 	RHICmdList.SetBlendState(TStaticBlendState<CW_RGBA, BO_Add, BF_One, BF_One, BO_Add, BF_One, BF_One>::GetRHI());
@@ -843,7 +843,7 @@ void FDeferredShadingSceneRenderer::RenderSimpleLightsStandardDeferred(FRHIComma
 {
 	SCOPE_CYCLE_COUNTER(STAT_DirectLightRenderingTime);
 	INC_DWORD_STAT_BY(STAT_NumLightsUsingStandardDeferred, SimpleLights.InstanceData.Num());
-	SCOPED_DRAW_EVENT(StandardDeferredSimpleLights, DEC_SCENE_ITEMS);
+	SCOPED_DRAW_EVENT(RHICmdList, StandardDeferredSimpleLights, DEC_SCENE_ITEMS);
 	
 	// Use additive blending for color
 	RHICmdList.SetBlendState(TStaticBlendState<CW_RGBA, BO_Add, BF_One, BF_One, BO_Add, BF_One, BF_One>::GetRHI());
diff --git a/Engine/Source/Runtime/Renderer/Private/LightShaftRendering.cpp b/Engine/Source/Runtime/Renderer/Private/LightShaftRendering.cpp
--- a/Engine/Source/Runtime/Renderer/Private/LightShaftRendering.cpp
+++ b/Engine/Source/Runtime/Renderer/Private/LightShaftRendering.cpp
@@ -451,7 +451,7 @@ void AllocateOrReuseLightShaftRenderTarget(FRHICommandListImmediate& RHICmdList,
 template<bool bDownsampleOcclusion>
 void DownsamplePass(FRHICommandListImmediate& RHICmdList, const FViewInfo& View, const FLightSceneInfo* LightSceneInfo, TRefCountPtr<IPooledRenderTarget>& LightShaftsSource, TRefCountPtr<IPooledRenderTarget>& LightShaftsDest)
 {
-	SCOPED_DRAW_EVENT(Downsample, DEC_SCENE_ITEMS);
+	SCOPED_DRAW_EVENT(RHICmdList, Downsample, DEC_SCENE_ITEMS);
 
 	const FIntPoint BufferSize = GSceneRenderTargets.GetBufferSizeXY();
 	const uint32 DownsampleFactor	= GetLightShaftDownsampleFactor();
@@ -624,7 +624,7 @@ void ApplyRadialBlurPasses(
 		BlurLightShaftsPixelShader->SetParameters(RHICmdList, LightSceneInfo, View, PassIndex, EffectiveSource);
 
 		{
-			SCOPED_DRAW_EVENT(RadialBlur, DEC_SCENE_ITEMS);
+			SCOPED_DRAW_EVENT(RHICmdList, RadialBlur, DEC_SCENE_ITEMS);
 			// Apply a radial blur to the bloom and occlusion mask
 			DrawRectangle( 
 				RHICmdList,
@@ -668,7 +668,7 @@ void FinishOcclusionTerm(FRHICommandList& RHICmdList, const FViewInfo& View, con
 	MaskOcclusionTermPixelShader->SetParameters(RHICmdList, LightSceneInfo, View, LightShaftsSource);
 
 	{
-		SCOPED_DRAW_EVENT(FinishOcclusion, DEC_SCENE_ITEMS);
+		SCOPED_DRAW_EVENT(RHICmdList, FinishOcclusion, DEC_SCENE_ITEMS);
 		// Apply a radial blur to the bloom and occlusion mask
 		DrawRectangle( 
 			RHICmdList,
@@ -733,7 +733,7 @@ FLightShaftsOutput FDeferredShadingSceneRenderer::RenderLightShaftOcclusion(FRHI
 
 			if (bEnableOcclusion && LightSceneInfo->Proxy->GetLightType() == LightType_Directional)
 			{
-				SCOPED_DRAW_EVENT(RenderLightShaftOcclusion, DEC_SCENE_ITEMS);
+				SCOPED_DRAW_EVENT(RHICmdList, RenderLightShaftOcclusion, DEC_SCENE_ITEMS);
 
 				// Allocate light shaft render targets on demand, using the pool
 				// Need two targets to ping pong between
@@ -836,7 +836,7 @@ IMPLEMENT_SHADER_TYPE(,FApplyLightShaftsPixelShader,TEXT("LightShaftShader"),TEX
 
 void ApplyLightShaftBloom(FRHICommandListImmediate& RHICmdList, const FViewInfo& View, const FLightSceneInfo* const LightSceneInfo, TRefCountPtr<IPooledRenderTarget>& LightShaftsSource)
 {
-	SCOPED_DRAW_EVENT(Apply, DEC_SCENE_ITEMS);
+	SCOPED_DRAW_EVENT(RHICmdList, Apply, DEC_SCENE_ITEMS);
 
 	GSceneRenderTargets.BeginRenderingSceneColor(RHICmdList);
 
@@ -886,7 +886,7 @@ void FDeferredShadingSceneRenderer::RenderLightShaftBloom(FRHICommandListImmedia
 
 			if (LightSceneInfo->bEnableLightShaftBloom)
 			{
-				SCOPED_DRAW_EVENT(RenderLightShaftBloom, DEC_SCENE_ITEMS);
+				SCOPED_DRAW_EVENT(RHICmdList, RenderLightShaftBloom, DEC_SCENE_ITEMS);
 
 				// Allocate light shaft render targets on demand, using the pool
 				AllocateOrReuseLightShaftRenderTarget(RHICmdList, LightShafts0, TEXT("LightShafts0"));
diff --git a/Engine/Source/Runtime/Renderer/Private/PostProcess/PostProcessAA.cpp b/Engine/Source/Runtime/Renderer/Private/PostProcess/PostProcessAA.cpp
--- a/Engine/Source/Runtime/Renderer/Private/PostProcess/PostProcessAA.cpp
+++ b/Engine/Source/Runtime/Renderer/Private/PostProcess/PostProcessAA.cpp
@@ -230,7 +230,7 @@ static void SetShaderTemplAA(const FRenderingCompositePassContext& Context)
 
 void FRCPassPostProcessAA::Process(FRenderingCompositePassContext& Context)
 {
-	SCOPED_DRAW_EVENT(PostProcessAA, DEC_SCENE_ITEMS);
+	SCOPED_DRAW_EVENT(Context.RHICmdList, PostProcessAA, DEC_SCENE_ITEMS);
 
 	const FPooledRenderTargetDesc* InputDesc = GetInputDesc(ePId_Input0);
 
diff --git a/Engine/Source/Runtime/Renderer/Private/PostProcess/PostProcessBloomSetup.cpp b/Engine/Source/Runtime/Renderer/Private/PostProcess/PostProcessBloomSetup.cpp
--- a/Engine/Source/Runtime/Renderer/Private/PostProcess/PostProcessBloomSetup.cpp
+++ b/Engine/Source/Runtime/Renderer/Private/PostProcess/PostProcessBloomSetup.cpp
@@ -139,7 +139,7 @@ IMPLEMENT_SHADER_TYPE(,FPostProcessBloomSetupVS,TEXT("PostProcessBloom"),TEXT("M
 
 void FRCPassPostProcessBloomSetup::Process(FRenderingCompositePassContext& Context)
 {
-	SCOPED_DRAW_EVENT(PostProcessBloomSetup, DEC_SCENE_ITEMS);
+	SCOPED_DRAW_EVENT(Context.RHICmdList, PostProcessBloomSetup, DEC_SCENE_ITEMS);
 
 	const FPooledRenderTargetDesc* InputDesc = GetInputDesc(ePId_Input0);
 
diff --git a/Engine/Source/Runtime/Renderer/Private/PostProcess/PostProcessBokehDOF.cpp b/Engine/Source/Runtime/Renderer/Private/PostProcess/PostProcessBokehDOF.cpp
--- a/Engine/Source/Runtime/Renderer/Private/PostProcess/PostProcessBokehDOF.cpp
+++ b/Engine/Source/Runtime/Renderer/Private/PostProcess/PostProcessBokehDOF.cpp
@@ -146,7 +146,7 @@ IMPLEMENT_SHADER_TYPE3(PostProcessVisualizeDOFPS, SF_Pixel);
 
 void FRCPassPostProcessVisualizeDOF::Process(FRenderingCompositePassContext& Context)
 {
-	SCOPED_DRAW_EVENT(VisualizeDOF, DEC_SCENE_ITEMS);
+	SCOPED_DRAW_EVENT(Context.RHICmdList, VisualizeDOF, DEC_SCENE_ITEMS);
 
 	const FPooledRenderTargetDesc* InputDesc = GetInputDesc(ePId_Input0);
 
@@ -305,7 +305,7 @@ IMPLEMENT_SHADER_TYPE3(PostProcessBokehDOFSetupPS, SF_Pixel);
 
 void FRCPassPostProcessBokehDOFSetup::Process(FRenderingCompositePassContext& Context)
 {
-	SCOPED_DRAW_EVENT(BokehDOFSetup, DEC_SCENE_ITEMS);
+	SCOPED_DRAW_EVENT(Context.RHICmdList, BokehDOFSetup, DEC_SCENE_ITEMS);
 
 	const FPooledRenderTargetDesc* InputDesc = GetInputDesc(ePId_Input0);
 
@@ -614,7 +614,7 @@ static TAutoConsoleVariable<int32> CVarBokehDOFIndexStyle(
 
 void FRCPassPostProcessBokehDOF::Process(FRenderingCompositePassContext& Context)
 {
-	SCOPED_DRAW_EVENT(PassPostProcessBokehDOF, DEC_SCENE_ITEMS);
+	SCOPED_DRAW_EVENT(Context.RHICmdList, PassPostProcessBokehDOF, DEC_SCENE_ITEMS);
 
 	const FPooledRenderTargetDesc* InputDesc = GetInputDesc(ePId_Input0);
 	
diff --git a/Engine/Source/Runtime/Renderer/Private/PostProcess/PostProcessBokehDOFRecombine.cpp b/Engine/Source/Runtime/Renderer/Private/PostProcess/PostProcessBokehDOFRecombine.cpp
--- a/Engine/Source/Runtime/Renderer/Private/PostProcess/PostProcessBokehDOFRecombine.cpp
+++ b/Engine/Source/Runtime/Renderer/Private/PostProcess/PostProcessBokehDOFRecombine.cpp
@@ -130,7 +130,7 @@ void FRCPassPostProcessBokehDOFRecombine::Process(FRenderingCompositePassContext
 		check(GetInput(ePId_Input2)->GetPass());
 	}
 
-	SCOPED_DRAW_EVENTF(BokehDOFRecombine, DEC_SCENE_ITEMS, TEXT("BokehDOFRecombine#%d"), Method);
+	SCOPED_DRAW_EVENTF(Context.RHICmdList, BokehDOFRecombine, DEC_SCENE_ITEMS, TEXT("BokehDOFRecombine#%d"), Method);
 
 	const FPooledRenderTargetDesc* InputDesc0 = GetInputDesc(ePId_Input0);
 	const FPooledRenderTargetDesc* InputDesc1 = GetInputDesc(ePId_Input1);
diff --git a/Engine/Source/Runtime/Renderer/Private/PostProcess/PostProcessBusyWait.cpp b/Engine/Source/Runtime/Renderer/Private/PostProcess/PostProcessBusyWait.cpp
--- a/Engine/Source/Runtime/Renderer/Private/PostProcess/PostProcessBusyWait.cpp
+++ b/Engine/Source/Runtime/Renderer/Private/PostProcess/PostProcessBusyWait.cpp
@@ -83,7 +83,7 @@ IMPLEMENT_SHADER_TYPE(,FPostProcessBusyWaitPS,TEXT("PostProcessBusyWait"),TEXT("
 
 void FRCPassPostProcessBusyWait::Process(FRenderingCompositePassContext& Context)
 {
-	SCOPED_DRAW_EVENT(BusyWait, DEC_SCENE_ITEMS);
+	SCOPED_DRAW_EVENT(Context.RHICmdList, BusyWait, DEC_SCENE_ITEMS);
 
 	const FSceneView& View = Context.View;
 	
diff --git a/Engine/Source/Runtime/Renderer/Private/PostProcess/PostProcessCombineLUTs.cpp b/Engine/Source/Runtime/Renderer/Private/PostProcess/PostProcessCombineLUTs.cpp
--- a/Engine/Source/Runtime/Renderer/Private/PostProcess/PostProcessCombineLUTs.cpp
+++ b/Engine/Source/Runtime/Renderer/Private/PostProcess/PostProcessCombineLUTs.cpp
@@ -380,7 +380,7 @@ uint32 FRCPassPostProcessCombineLUTs::GenerateFinalTable(const FFinalPostProcess
 
 void FRCPassPostProcessCombineLUTs::Process(FRenderingCompositePassContext& Context)
 {
-	SCOPED_DRAW_EVENT(PostProcessCombineLUTs, DEC_SCENE_ITEMS);
+	SCOPED_DRAW_EVENT(Context.RHICmdList, PostProcessCombineLUTs, DEC_SCENE_ITEMS);
 
 	FTexture* LocalTextures[GMaxLUTBlendCount];
 	float LocalWeights[GMaxLUTBlendCount];
diff --git a/Engine/Source/Runtime/Renderer/Private/PostProcess/PostProcessCompositeEditorPrimitives.cpp b/Engine/Source/Runtime/Renderer/Private/PostProcess/PostProcessCompositeEditorPrimitives.cpp
--- a/Engine/Source/Runtime/Renderer/Private/PostProcess/PostProcessCompositeEditorPrimitives.cpp
+++ b/Engine/Source/Runtime/Renderer/Private/PostProcess/PostProcessCompositeEditorPrimitives.cpp
@@ -182,7 +182,7 @@ static void SetCompositePrimitivesShaderTempl(const FRenderingCompositePassConte
 
 void FRCPassPostProcessCompositeEditorPrimitives::Process(FRenderingCompositePassContext& Context)
 {
-	SCOPED_DRAW_EVENT(CompositeEditorPrimitives, DEC_SCENE_ITEMS);
+	SCOPED_DRAW_EVENT(Context.RHICmdList, CompositeEditorPrimitives, DEC_SCENE_ITEMS);
 	const FPooledRenderTargetDesc* InputDesc = GetInputDesc(ePId_Input0);
 
 	if(!InputDesc)
@@ -213,13 +213,13 @@ void FRCPassPostProcessCompositeEditorPrimitives::Process(FRenderingCompositePas
 
 		if(bClearIsNeeded)
 		{
-			SCOPED_DRAW_EVENT(ClearViewEditorPrimitives, DEC_SCENE_ITEMS);
+			SCOPED_DRAW_EVENT(Context.RHICmdList, ClearViewEditorPrimitives, DEC_SCENE_ITEMS);
 			// Clear color and depth
 			// Note, this is a reversed Z depth surface, so 0.0f is the far plane.
 			Context.RHICmdList.Clear(true, FLinearColor(0, 0, 0, 0), true, 0.0f, false, 0, FIntRect());
 		}
 
-		SCOPED_DRAW_EVENT(RenderEditorPrimitives, DEC_SCENE_ITEMS);
+		SCOPED_DRAW_EVENT(Context.RHICmdList, RenderEditorPrimitives, DEC_SCENE_ITEMS);
 
 		Context.RHICmdList.SetRasterizerState(View.bReverseCulling ? TStaticRasterizerState<FM_Solid, CM_CW>::GetRHI() : TStaticRasterizerState<FM_Solid, CM_CCW>::GetRHI());
 
diff --git a/Engine/Source/Runtime/Renderer/Private/PostProcess/PostProcessDOF.cpp b/Engine/Source/Runtime/Renderer/Private/PostProcess/PostProcessDOF.cpp
--- a/Engine/Source/Runtime/Renderer/Private/PostProcess/PostProcessDOF.cpp
+++ b/Engine/Source/Runtime/Renderer/Private/PostProcess/PostProcessDOF.cpp
@@ -80,7 +80,7 @@ IMPLEMENT_SHADER_TYPE(template<>,FPostProcessDOFSetupPS<1>,TEXT("PostProcessDOF"
 
 void FRCPassPostProcessDOFSetup::Process(FRenderingCompositePassContext& Context)
 {
-	SCOPED_DRAW_EVENT(DOFSetup, DEC_SCENE_ITEMS);
+	SCOPED_DRAW_EVENT(Context.RHICmdList, DOFSetup, DEC_SCENE_ITEMS);
 
 	const FPooledRenderTargetDesc* InputDesc = GetInputDesc(ePId_Input0);
 
@@ -264,7 +264,7 @@ IMPLEMENT_SHADER_TYPE(template<>,FPostProcessDOFRecombinePS<1>,TEXT("PostProcess
 
 void FRCPassPostProcessDOFRecombine::Process(FRenderingCompositePassContext& Context)
 {
-	SCOPED_DRAW_EVENT(DOFRecombine, DEC_SCENE_ITEMS);
+	SCOPED_DRAW_EVENT(Context.RHICmdList, DOFRecombine, DEC_SCENE_ITEMS);
 
 	const FPooledRenderTargetDesc* InputDesc = GetInputDesc(ePId_Input1);
 
diff --git a/Engine/Source/Runtime/Renderer/Private/PostProcess/PostProcessDownsample.cpp b/Engine/Source/Runtime/Renderer/Private/PostProcess/PostProcessDownsample.cpp
--- a/Engine/Source/Runtime/Renderer/Private/PostProcess/PostProcessDownsample.cpp
+++ b/Engine/Source/Runtime/Renderer/Private/PostProcess/PostProcessDownsample.cpp
@@ -154,7 +154,7 @@ void FRCPassPostProcessDownsample::SetShader(const FRenderingCompositePassContex
 
 void FRCPassPostProcessDownsample::Process(FRenderingCompositePassContext& Context)
 {
-	SCOPED_DRAW_EVENT(Downsample, DEC_SCENE_ITEMS);
+	SCOPED_DRAW_EVENT(Context.RHICmdList, Downsample, DEC_SCENE_ITEMS);
 
 	const FPooledRenderTargetDesc* InputDesc = GetInputDesc(ePId_Input0);
 
diff --git a/Engine/Source/Runtime/Renderer/Private/PostProcess/PostProcessEyeAdaptation.cpp b/Engine/Source/Runtime/Renderer/Private/PostProcess/PostProcessEyeAdaptation.cpp
--- a/Engine/Source/Runtime/Renderer/Private/PostProcess/PostProcessEyeAdaptation.cpp
+++ b/Engine/Source/Runtime/Renderer/Private/PostProcess/PostProcessEyeAdaptation.cpp
@@ -71,7 +71,7 @@ IMPLEMENT_SHADER_TYPE(,FPostProcessEyeAdaptationPS,TEXT("PostProcessEyeAdaptatio
 
 void FRCPassPostProcessEyeAdaptation::Process(FRenderingCompositePassContext& Context)
 {
-	SCOPED_DRAW_EVENT(PostProcessEyeAdaptation, DEC_SCENE_ITEMS);
+	SCOPED_DRAW_EVENT(Context.RHICmdList, PostProcessEyeAdaptation, DEC_SCENE_ITEMS);
 
 	const FSceneView& View = Context.View;
 	const FSceneViewFamily& ViewFamily = *(View.Family);
diff --git a/Engine/Source/Runtime/Renderer/Private/PostProcess/PostProcessGBufferHints.cpp b/Engine/Source/Runtime/Renderer/Private/PostProcess/PostProcessGBufferHints.cpp
--- a/Engine/Source/Runtime/Renderer/Private/PostProcess/PostProcessGBufferHints.cpp
+++ b/Engine/Source/Runtime/Renderer/Private/PostProcess/PostProcessGBufferHints.cpp
@@ -87,7 +87,7 @@ FRCPassPostProcessGBufferHints::FRCPassPostProcessGBufferHints()
 
 void FRCPassPostProcessGBufferHints::Process(FRenderingCompositePassContext& Context)
 {
-	SCOPED_DRAW_EVENT(GBufferHints, DEC_SCENE_ITEMS);
+	SCOPED_DRAW_EVENT(Context.RHICmdList, GBufferHints, DEC_SCENE_ITEMS);
 	const FPooledRenderTargetDesc* InputDesc = GetInputDesc(ePId_Input0);
 
 	if(!InputDesc)
diff --git a/Engine/Source/Runtime/Renderer/Private/PostProcess/PostProcessHMD.cpp b/Engine/Source/Runtime/Renderer/Private/PostProcess/PostProcessHMD.cpp
--- a/Engine/Source/Runtime/Renderer/Private/PostProcess/PostProcessHMD.cpp
+++ b/Engine/Source/Runtime/Renderer/Private/PostProcess/PostProcessHMD.cpp
@@ -184,7 +184,7 @@ IMPLEMENT_SHADER_TYPE(template<>, FPostProcessHMDPS<false>, TEXT("PostProcessHMD
 
 void FRCPassPostProcessHMD::Process(FRenderingCompositePassContext& Context)
 {
-	SCOPED_DRAW_EVENT(PostProcessHMD, DEC_SCENE_ITEMS);
+	SCOPED_DRAW_EVENT(Context.RHICmdList, PostProcessHMD, DEC_SCENE_ITEMS);
 	const FPooledRenderTargetDesc* InputDesc = GetInputDesc(ePId_Input0);
 
 	if(!InputDesc)
diff --git a/Engine/Source/Runtime/Renderer/Private/PostProcess/PostProcessHistogram.cpp b/Engine/Source/Runtime/Renderer/Private/PostProcess/PostProcessHistogram.cpp
--- a/Engine/Source/Runtime/Renderer/Private/PostProcess/PostProcessHistogram.cpp
+++ b/Engine/Source/Runtime/Renderer/Private/PostProcess/PostProcessHistogram.cpp
@@ -91,7 +91,7 @@ IMPLEMENT_SHADER_TYPE(,FPostProcessHistogramCS,TEXT("PostProcessHistogram"),TEXT
 
 void FRCPassPostProcessHistogram::Process(FRenderingCompositePassContext& Context)
 {
-	SCOPED_DRAW_EVENT(PostProcessHistogram, DEC_SCENE_ITEMS);
+	SCOPED_DRAW_EVENT(Context.RHICmdList, PostProcessHistogram, DEC_SCENE_ITEMS);
 	const FPooledRenderTargetDesc* InputDesc = GetInputDesc(ePId_Input0);
 
 	if(!InputDesc)
diff --git a/Engine/Source/Runtime/Renderer/Private/PostProcess/PostProcessHistogramReduce.cpp b/Engine/Source/Runtime/Renderer/Private/PostProcess/PostProcessHistogramReduce.cpp
--- a/Engine/Source/Runtime/Renderer/Private/PostProcess/PostProcessHistogramReduce.cpp
+++ b/Engine/Source/Runtime/Renderer/Private/PostProcess/PostProcessHistogramReduce.cpp
@@ -89,7 +89,7 @@ IMPLEMENT_SHADER_TYPE(,FPostProcessHistogramReducePS,TEXT("PostProcessHistogramR
 
 void FRCPassPostProcessHistogramReduce::Process(FRenderingCompositePassContext& Context)
 {
-	SCOPED_DRAW_EVENT(PostProcessHistogramReduce, DEC_SCENE_ITEMS);
+	SCOPED_DRAW_EVENT(Context.RHICmdList, PostProcessHistogramReduce, DEC_SCENE_ITEMS);
 	const FPooledRenderTargetDesc* InputDesc = GetInputDesc(ePId_Input0);
 
 	if(!InputDesc)
diff --git a/Engine/Source/Runtime/Renderer/Private/PostProcess/PostProcessLensBlur.cpp b/Engine/Source/Runtime/Renderer/Private/PostProcess/PostProcessLensBlur.cpp
--- a/Engine/Source/Runtime/Renderer/Private/PostProcess/PostProcessLensBlur.cpp
+++ b/Engine/Source/Runtime/Renderer/Private/PostProcess/PostProcessLensBlur.cpp
@@ -162,7 +162,7 @@ FRCPassPostProcessLensBlur::FRCPassPostProcessLensBlur(float InPercentKernelSize
 
 void FRCPassPostProcessLensBlur::Process(FRenderingCompositePassContext& Context)
 {
-	SCOPED_DRAW_EVENT(PassPostProcessLensBlur, DEC_SCENE_ITEMS);
+	SCOPED_DRAW_EVENT(Context.RHICmdList, PassPostProcessLensBlur, DEC_SCENE_ITEMS);
 
 	const FPooledRenderTargetDesc* InputDesc = GetInputDesc(ePId_Input0);
 	
diff --git a/Engine/Source/Runtime/Renderer/Private/PostProcess/PostProcessLensFlares.cpp b/Engine/Source/Runtime/Renderer/Private/PostProcess/PostProcessLensFlares.cpp
--- a/Engine/Source/Runtime/Renderer/Private/PostProcess/PostProcessLensFlares.cpp
+++ b/Engine/Source/Runtime/Renderer/Private/PostProcess/PostProcessLensFlares.cpp
@@ -113,7 +113,7 @@ FRCPassPostProcessLensFlares::FRCPassPostProcessLensFlares(float InSizeScale)
 
 void FRCPassPostProcessLensFlares::Process(FRenderingCompositePassContext& Context)
 {
-	SCOPED_DRAW_EVENT(LensFlares, DEC_SCENE_ITEMS);
+	SCOPED_DRAW_EVENT(Context.RHICmdList, LensFlares, DEC_SCENE_ITEMS);
 
 	const FPooledRenderTargetDesc* InputDesc1 = GetInputDesc(ePId_Input0);
 	const FPooledRenderTargetDesc* InputDesc2 = GetInputDesc(ePId_Input1);
diff --git a/Engine/Source/Runtime/Renderer/Private/PostProcess/PostProcessMaterial.cpp b/Engine/Source/Runtime/Renderer/Private/PostProcess/PostProcessMaterial.cpp
--- a/Engine/Source/Runtime/Renderer/Private/PostProcess/PostProcessMaterial.cpp
+++ b/Engine/Source/Runtime/Renderer/Private/PostProcess/PostProcessMaterial.cpp
@@ -107,7 +107,7 @@ void FRCPassPostProcessMaterial::Process(FRenderingCompositePassContext& Context
 	
 	check(Material);
 
-	SCOPED_DRAW_EVENTF(PostProcessMaterial, DEC_SCENE_ITEMS, TEXT("PostProcessMaterial Material=%s"), *Material->GetFriendlyName());
+	SCOPED_DRAW_EVENTF(Context.RHICmdList, PostProcessMaterial, DEC_SCENE_ITEMS, TEXT("PostProcessMaterial Material=%s"), *Material->GetFriendlyName());
 
 	const FPooledRenderTargetDesc* InputDesc = GetInputDesc(ePId_Input0);
 
diff --git a/Engine/Source/Runtime/Renderer/Private/PostProcess/PostProcessMobile.cpp b/Engine/Source/Runtime/Renderer/Private/PostProcess/PostProcessMobile.cpp
--- a/Engine/Source/Runtime/Renderer/Private/PostProcess/PostProcessMobile.cpp
+++ b/Engine/Source/Runtime/Renderer/Private/PostProcess/PostProcessMobile.cpp
@@ -168,7 +168,7 @@ void FRCPassPostProcessBloomSetupES2::SetShader(const FRenderingCompositePassCon
 
 void FRCPassPostProcessBloomSetupES2::Process(FRenderingCompositePassContext& Context)
 {
-	SCOPED_DRAW_EVENT(PostProcessBloomSetup, DEC_SCENE_ITEMS);
+	SCOPED_DRAW_EVENT(Context.RHICmdList, PostProcessBloomSetup, DEC_SCENE_ITEMS);
 
 	const FPooledRenderTargetDesc* InputDesc = GetInputDesc(ePId_Input0);
 
@@ -369,7 +369,7 @@ void FRCPassPostProcessBloomSetupSmallES2::SetShader(const FRenderingCompositePa
 
 void FRCPassPostProcessBloomSetupSmallES2::Process(FRenderingCompositePassContext& Context)
 {
-	SCOPED_DRAW_EVENT(PostProcessBloomSetupSmall, DEC_SCENE_ITEMS);
+	SCOPED_DRAW_EVENT(Context.RHICmdList, PostProcessBloomSetupSmall, DEC_SCENE_ITEMS);
 
 	const FPooledRenderTargetDesc* InputDesc = GetInputDesc(ePId_Input0);
 
@@ -545,7 +545,7 @@ IMPLEMENT_SHADER_TYPE(,FPostProcessBloomDownVS_ES2,TEXT("PostProcessMobile"),TEX
 
 void FRCPassPostProcessBloomDownES2::Process(FRenderingCompositePassContext& Context)
 {
-	SCOPED_DRAW_EVENT(PostProcessBloomDown, DEC_SCENE_ITEMS);
+	SCOPED_DRAW_EVENT(Context.RHICmdList, PostProcessBloomDown, DEC_SCENE_ITEMS);
 
 	uint32 DstX = FMath::Max(1, PrePostSourceViewportSize.X/2);
 	uint32 DstY = FMath::Max(1, PrePostSourceViewportSize.Y/2);
@@ -713,7 +713,7 @@ IMPLEMENT_SHADER_TYPE(,FPostProcessBloomUpVS_ES2,TEXT("PostProcessMobile"),TEXT(
 
 void FRCPassPostProcessBloomUpES2::Process(FRenderingCompositePassContext& Context)
 {
-	SCOPED_DRAW_EVENT(PostProcessBloomUp, DEC_SCENE_ITEMS);
+	SCOPED_DRAW_EVENT(Context.RHICmdList, PostProcessBloomUp, DEC_SCENE_ITEMS);
 
 	uint32 DstX = FMath::Max(1, PrePostSourceViewportSize.X);
 	uint32 DstY = FMath::Max(1, PrePostSourceViewportSize.Y);
@@ -947,7 +947,7 @@ void FRCPassPostProcessSunMaskES2::SetShader(const FRenderingCompositePassContex
 
 void FRCPassPostProcessSunMaskES2::Process(FRenderingCompositePassContext& Context)
 {
-	SCOPED_DRAW_EVENT(PostProcessSunMask, DEC_SCENE_ITEMS);
+	SCOPED_DRAW_EVENT(Context.RHICmdList, PostProcessSunMask, DEC_SCENE_ITEMS);
 
 	const FPooledRenderTargetDesc* InputDesc = GetInputDesc(ePId_Input0);
 
@@ -1169,7 +1169,7 @@ void FRCPassPostProcessSunAlphaES2::SetShader(const FRenderingCompositePassConte
 
 void FRCPassPostProcessSunAlphaES2::Process(FRenderingCompositePassContext& Context)
 {
-	SCOPED_DRAW_EVENT(PostProcessSunAlpha, DEC_SCENE_ITEMS);
+	SCOPED_DRAW_EVENT(Context.RHICmdList, PostProcessSunAlpha, DEC_SCENE_ITEMS);
 
 	uint32 DstX = FMath::Max(1, PrePostSourceViewportSize.X/4);
 	uint32 DstY = FMath::Max(1, PrePostSourceViewportSize.Y/4);
@@ -1326,7 +1326,7 @@ IMPLEMENT_SHADER_TYPE(,FPostProcessSunBlurVS_ES2,TEXT("PostProcessMobile"),TEXT(
 
 void FRCPassPostProcessSunBlurES2::Process(FRenderingCompositePassContext& Context)
 {
-	SCOPED_DRAW_EVENT(PostProcessSunBlur, DEC_SCENE_ITEMS);
+	SCOPED_DRAW_EVENT(Context.RHICmdList, PostProcessSunBlur, DEC_SCENE_ITEMS);
 
 	uint32 DstX = FMath::Max(1, PrePostSourceViewportSize.X/4);
 	uint32 DstY = FMath::Max(1, PrePostSourceViewportSize.Y/4);
@@ -1550,7 +1550,7 @@ void FRCPassPostProcessSunMergeES2::SetShader(const FRenderingCompositePassConte
 
 void FRCPassPostProcessSunMergeES2::Process(FRenderingCompositePassContext& Context)
 {
-	SCOPED_DRAW_EVENT(PostProcessSunMerge, DEC_SCENE_ITEMS);
+	SCOPED_DRAW_EVENT(Context.RHICmdList, PostProcessSunMerge, DEC_SCENE_ITEMS);
 
 	uint32 DstX = FMath::Max(1, PrePostSourceViewportSize.X/4);
 	uint32 DstY = FMath::Max(1, PrePostSourceViewportSize.Y/4);
@@ -1745,7 +1745,7 @@ void FRCPassPostProcessSunMergeSmallES2::SetShader(const FRenderingCompositePass
 
 void FRCPassPostProcessSunMergeSmallES2::Process(FRenderingCompositePassContext& Context)
 {
-	SCOPED_DRAW_EVENT(PostProcessSunMergeSmall, DEC_SCENE_ITEMS);
+	SCOPED_DRAW_EVENT(Context.RHICmdList, PostProcessSunMergeSmall, DEC_SCENE_ITEMS);
 
 	uint32 DstX = FMath::Max(1, PrePostSourceViewportSize.X/4);
 	uint32 DstY = FMath::Max(1, PrePostSourceViewportSize.Y/4);
@@ -1948,7 +1948,7 @@ void FRCPassPostProcessDofDownES2::SetShader(const FRenderingCompositePassContex
 
 void FRCPassPostProcessDofDownES2::Process(FRenderingCompositePassContext& Context)
 {
-	SCOPED_DRAW_EVENT(PostProcessDofDown, DEC_SCENE_ITEMS);
+	SCOPED_DRAW_EVENT(Context.RHICmdList, PostProcessDofDown, DEC_SCENE_ITEMS);
 
 	const FPooledRenderTargetDesc* InputDesc = GetInputDesc(ePId_Input0);
 
@@ -2156,7 +2156,7 @@ void FRCPassPostProcessDofNearES2::SetShader(const FRenderingCompositePassContex
 
 void FRCPassPostProcessDofNearES2::Process(FRenderingCompositePassContext& Context)
 {
-	SCOPED_DRAW_EVENT(PostProcessDofNear, DEC_SCENE_ITEMS);
+	SCOPED_DRAW_EVENT(Context.RHICmdList, PostProcessDofNear, DEC_SCENE_ITEMS);
 
 	uint32 DstX = FMath::Max(1, PrePostSourceViewportSize.X/4);
 	uint32 DstY = FMath::Max(1, PrePostSourceViewportSize.Y/4);
@@ -2306,7 +2306,7 @@ IMPLEMENT_SHADER_TYPE(,FPostProcessDofBlurVS_ES2,TEXT("PostProcessMobile"),TEXT(
 
 void FRCPassPostProcessDofBlurES2::Process(FRenderingCompositePassContext& Context)
 {
-	SCOPED_DRAW_EVENT(PostProcessDofBlur, DEC_SCENE_ITEMS);
+	SCOPED_DRAW_EVENT(Context.RHICmdList, PostProcessDofBlur, DEC_SCENE_ITEMS);
 
 	uint32 DstX = FMath::Max(1, PrePostSourceViewportSize.X/2);
 	uint32 DstY = FMath::Max(1, PrePostSourceViewportSize.Y/2);
@@ -2486,7 +2486,7 @@ void FRCPassPostProcessSunAvgES2::SetShader(const FRenderingCompositePassContext
 
 void FRCPassPostProcessSunAvgES2::Process(FRenderingCompositePassContext& Context)
 {
-	SCOPED_DRAW_EVENT(PostProcessSunAvg, DEC_SCENE_ITEMS);
+	SCOPED_DRAW_EVENT(Context.RHICmdList, PostProcessSunAvg, DEC_SCENE_ITEMS);
 
 	uint32 DstX = FMath::Max(1, PrePostSourceViewportSize.X/4);
 	uint32 DstY = FMath::Max(1, PrePostSourceViewportSize.Y/4);
@@ -2729,7 +2729,7 @@ void FRCPassPostProcessAaES2::SetShader(const FRenderingCompositePassContext& Co
 
 void FRCPassPostProcessAaES2::Process(FRenderingCompositePassContext& Context)
 {
-	SCOPED_DRAW_EVENT(PostProcessAa, DEC_SCENE_ITEMS);
+	SCOPED_DRAW_EVENT(Context.RHICmdList, PostProcessAa, DEC_SCENE_ITEMS);
 
 	uint32 DstX = FMath::Max(1, PrePostSourceViewportSize.X);
 	uint32 DstY = FMath::Max(1, PrePostSourceViewportSize.Y);
diff --git a/Engine/Source/Runtime/Renderer/Private/PostProcess/PostProcessMorpheus.cpp b/Engine/Source/Runtime/Renderer/Private/PostProcess/PostProcessMorpheus.cpp
--- a/Engine/Source/Runtime/Renderer/Private/PostProcess/PostProcessMorpheus.cpp
+++ b/Engine/Source/Runtime/Renderer/Private/PostProcess/PostProcessMorpheus.cpp
@@ -177,7 +177,7 @@ IMPLEMENT_SHADER_TYPE(, FPostProcessMorpheusPS,TEXT("MorpheusInclude"),TEXT("Mai
 
 void FRCPassPostProcessMorpheus::Process(FRenderingCompositePassContext& Context)
 {
-	SCOPED_DRAW_EVENT(PostProcessMorpheus, DEC_SCENE_ITEMS);
+	SCOPED_DRAW_EVENT(Context.RHICmdList, PostProcessMorpheus, DEC_SCENE_ITEMS);
 	const FPooledRenderTargetDesc* InputDesc = GetInputDesc(ePId_Input0);
 
 	if(!InputDesc)
diff --git a/Engine/Source/Runtime/Renderer/Private/PostProcess/PostProcessMotionBlur.cpp b/Engine/Source/Runtime/Renderer/Private/PostProcess/PostProcessMotionBlur.cpp
--- a/Engine/Source/Runtime/Renderer/Private/PostProcess/PostProcessMotionBlur.cpp
+++ b/Engine/Source/Runtime/Renderer/Private/PostProcess/PostProcessMotionBlur.cpp
@@ -132,7 +132,7 @@ IMPLEMENT_SHADER_TYPE(,FPostProcessMotionBlurSetupPS, TEXT("PostProcessMotionBlu
 
 void FRCPassPostProcessMotionBlurSetup::Process(FRenderingCompositePassContext& Context)
 {
-	SCOPED_DRAW_EVENT(MotionBlurSetup, DEC_SCENE_ITEMS);
+	SCOPED_DRAW_EVENT(Context.RHICmdList, MotionBlurSetup, DEC_SCENE_ITEMS);
 
 	const FPooledRenderTargetDesc* InputDesc = GetInputDesc(ePId_Input0);
 
@@ -438,7 +438,7 @@ FRCPassPostProcessMotionBlur::FRCPassPostProcessMotionBlur(uint32 InQuality)
 
 void FRCPassPostProcessMotionBlur::Process(FRenderingCompositePassContext& Context)
 {
-	SCOPED_DRAW_EVENT(MotionBlur, DEC_SCENE_ITEMS);
+	SCOPED_DRAW_EVENT(Context.RHICmdList, MotionBlur, DEC_SCENE_ITEMS);
 
 	const FPooledRenderTargetDesc* InputDesc = GetInputDesc(ePId_Input0);
 
@@ -573,7 +573,7 @@ IMPLEMENT_SHADER_TYPE(,FPostProcessMotionBlurRecombinePS,TEXT("PostProcessMotion
 
 void FRCPassPostProcessMotionBlurRecombine::Process(FRenderingCompositePassContext& Context)
 {
-	SCOPED_DRAW_EVENT(MotionBlurRecombine, DEC_SCENE_ITEMS);
+	SCOPED_DRAW_EVENT(Context.RHICmdList, MotionBlurRecombine, DEC_SCENE_ITEMS);
 
 	const FPooledRenderTargetDesc* InputDesc = GetInputDesc(ePId_Input0);
 
@@ -654,7 +654,7 @@ FPooledRenderTargetDesc FRCPassPostProcessMotionBlurRecombine::ComputeOutputDesc
 
 void FRCPassPostProcessVisualizeMotionBlur::Process(FRenderingCompositePassContext& Context)
 {
-	SCOPED_DRAW_EVENT(VisualizeMotionBlur, DEC_SCENE_ITEMS);
+	SCOPED_DRAW_EVENT(Context.RHICmdList, VisualizeMotionBlur, DEC_SCENE_ITEMS);
 
 	const FPooledRenderTargetDesc* InputDesc = GetInputDesc(ePId_Input0);
 
diff --git a/Engine/Source/Runtime/Renderer/Private/PostProcess/PostProcessNoiseBlur.cpp b/Engine/Source/Runtime/Renderer/Private/PostProcess/PostProcessNoiseBlur.cpp
--- a/Engine/Source/Runtime/Renderer/Private/PostProcess/PostProcessNoiseBlur.cpp
+++ b/Engine/Source/Runtime/Renderer/Private/PostProcess/PostProcessNoiseBlur.cpp
@@ -112,7 +112,7 @@ void SetNoiseBlurShader(const FRenderingCompositePassContext& Context, float InR
 
 void FRCPassPostProcessNoiseBlur::Process(FRenderingCompositePassContext& Context)
 {
-	SCOPED_DRAW_EVENT(NoiseBlur, DEC_SCENE_ITEMS);
+	SCOPED_DRAW_EVENT(Context.RHICmdList, NoiseBlur, DEC_SCENE_ITEMS);
 
 	const FPooledRenderTargetDesc* InputDesc = GetInputDesc(ePId_Input0);
 
diff --git a/Engine/Source/Runtime/Renderer/Private/PostProcess/PostProcessSceneColorFringe.cpp b/Engine/Source/Runtime/Renderer/Private/PostProcess/PostProcessSceneColorFringe.cpp
--- a/Engine/Source/Runtime/Renderer/Private/PostProcess/PostProcessSceneColorFringe.cpp
+++ b/Engine/Source/Runtime/Renderer/Private/PostProcess/PostProcessSceneColorFringe.cpp
@@ -143,7 +143,7 @@ IMPLEMENT_SHADER_TYPE(,FPostProcessSceneColorFringePS,TEXT("PostProcessSceneColo
 
 void FRCPassPostProcessSceneColorFringe::Process(FRenderingCompositePassContext& Context)
 {
-	SCOPED_DRAW_EVENT(SceneColorFringe, DEC_SCENE_ITEMS);
+	SCOPED_DRAW_EVENT(Context.RHICmdList, SceneColorFringe, DEC_SCENE_ITEMS);
 
 	const FPooledRenderTargetDesc* InputDesc = GetInputDesc(ePId_Input0);
 
diff --git a/Engine/Source/Runtime/Renderer/Private/PostProcess/PostProcessSelectionOutline.cpp b/Engine/Source/Runtime/Renderer/Private/PostProcess/PostProcessSelectionOutline.cpp
--- a/Engine/Source/Runtime/Renderer/Private/PostProcess/PostProcessSelectionOutline.cpp
+++ b/Engine/Source/Runtime/Renderer/Private/PostProcess/PostProcessSelectionOutline.cpp
@@ -17,7 +17,7 @@
 
 void FRCPassPostProcessSelectionOutlineColor::Process(FRenderingCompositePassContext& Context)
 {
-	SCOPED_DRAW_EVENT(PostProcessSelectionOutlineBuffer, DEC_SCENE_ITEMS);
+	SCOPED_DRAW_EVENT(Context.RHICmdList, PostProcessSelectionOutlineBuffer, DEC_SCENE_ITEMS);
 
 	const FPooledRenderTargetDesc* SceneColorInputDesc = GetInputDesc(ePId_Input0);
 
@@ -364,7 +364,7 @@ static void SetSelectionOutlineShaderTempl(const FRenderingCompositePassContext&
 
 void FRCPassPostProcessSelectionOutline::Process(FRenderingCompositePassContext& Context)
 {
-	SCOPED_DRAW_EVENT(PostProcessSelectionOutline, DEC_SCENE_ITEMS);
+	SCOPED_DRAW_EVENT(Context.RHICmdList, PostProcessSelectionOutline, DEC_SCENE_ITEMS);
 	const FPooledRenderTargetDesc* SceneColorInputDesc = GetInputDesc(ePId_Input0);
 	const FPooledRenderTargetDesc* SelectionColorInputDesc = GetInputDesc(ePId_Input1);
 
diff --git a/Engine/Source/Runtime/Renderer/Private/PostProcess/PostProcessSubsurface.cpp b/Engine/Source/Runtime/Renderer/Private/PostProcess/PostProcessSubsurface.cpp
--- a/Engine/Source/Runtime/Renderer/Private/PostProcess/PostProcessSubsurface.cpp
+++ b/Engine/Source/Runtime/Renderer/Private/PostProcess/PostProcessSubsurface.cpp
@@ -148,7 +148,7 @@ FRCPassPostProcessSubsurfaceSetup::FRCPassPostProcessSubsurfaceSetup(bool bInVis
 
 void FRCPassPostProcessSubsurfaceSetup::Process(FRenderingCompositePassContext& Context)
 {
-	SCOPED_DRAW_EVENT(PostProcessSubsurfaceSetup, DEC_SCENE_ITEMS);
+	SCOPED_DRAW_EVENT(Context.RHICmdList, PostProcessSubsurfaceSetup, DEC_SCENE_ITEMS);
 
 	const FPooledRenderTargetDesc* InputDesc = GetInputDesc(ePId_Input0);
 
@@ -476,7 +476,7 @@ void FRCPassPostProcessSubsurface::Process(FRenderingCompositePassContext& Conte
 
 	TShaderMapRef<FPostProcessVS> VertexShader(Context.GetShaderMap());
 
-	SCOPED_DRAW_EVENT(SubsurfacePass, DEC_SCENE_ITEMS);
+	SCOPED_DRAW_EVENT(Context.RHICmdList, SubsurfacePass, DEC_SCENE_ITEMS);
 
 	uint32 SampleSet = FMath::Clamp(CVarSSSSampleSet.GetValueOnRenderThread(), 0, 2);
 
diff --git a/Engine/Source/Runtime/Renderer/Private/PostProcess/PostProcessTemporalAA.cpp b/Engine/Source/Runtime/Renderer/Private/PostProcess/PostProcessTemporalAA.cpp
--- a/Engine/Source/Runtime/Renderer/Private/PostProcess/PostProcessTemporalAA.cpp
+++ b/Engine/Source/Runtime/Renderer/Private/PostProcess/PostProcessTemporalAA.cpp
@@ -189,7 +189,7 @@ IMPLEMENT_TEMPORALAA_PIXELSHADER_TYPE(4, 1, TEXT("MainFastTemporalAAPS"));
 
 void FRCPassPostProcessSSRTemporalAA::Process(FRenderingCompositePassContext& Context)
 {
-	SCOPED_DRAW_EVENT(SSRTemporalAA, DEC_SCENE_ITEMS);
+	SCOPED_DRAW_EVENT(Context.RHICmdList, SSRTemporalAA, DEC_SCENE_ITEMS);
 
 	const FPooledRenderTargetDesc* InputDesc = GetInputDesc(ePId_Input0);
 
@@ -263,7 +263,7 @@ FPooledRenderTargetDesc FRCPassPostProcessSSRTemporalAA::ComputeOutputDesc(EPass
 
 void FRCPassPostProcessDOFTemporalAA::Process(FRenderingCompositePassContext& Context)
 {
-	SCOPED_DRAW_EVENT(DOFTemporalAA, DEC_SCENE_ITEMS);
+	SCOPED_DRAW_EVENT(Context.RHICmdList, DOFTemporalAA, DEC_SCENE_ITEMS);
 
 	const FPooledRenderTargetDesc* InputDesc = GetInputDesc(ePId_Input0);
 
@@ -352,7 +352,7 @@ void FRCPassPostProcessLightShaftTemporalAA::Process(FRenderingCompositePassCont
 		return;
 	}
 
-	SCOPED_DRAW_EVENT(LSTemporalAA, DEC_SCENE_ITEMS);
+	SCOPED_DRAW_EVENT(Context.RHICmdList, LSTemporalAA, DEC_SCENE_ITEMS);
 
 	const FViewInfo& View = Context.View;
 	FSceneViewState* ViewState = Context.ViewState;
@@ -422,7 +422,7 @@ FPooledRenderTargetDesc FRCPassPostProcessLightShaftTemporalAA::ComputeOutputDes
 
 void FRCPassPostProcessTemporalAA::Process(FRenderingCompositePassContext& Context)
 {
-	SCOPED_DRAW_EVENT(TemporalAA, DEC_SCENE_ITEMS);
+	SCOPED_DRAW_EVENT(Context.RHICmdList, TemporalAA, DEC_SCENE_ITEMS);
 
 	const FPooledRenderTargetDesc* InputDesc = GetInputDesc(ePId_Input0);
 
diff --git a/Engine/Source/Runtime/Renderer/Private/PostProcess/PostProcessTestImage.cpp b/Engine/Source/Runtime/Renderer/Private/PostProcess/PostProcessTestImage.cpp
--- a/Engine/Source/Runtime/Renderer/Private/PostProcess/PostProcessTestImage.cpp
+++ b/Engine/Source/Runtime/Renderer/Private/PostProcess/PostProcessTestImage.cpp
@@ -89,7 +89,7 @@ FRCPassPostProcessTestImage::FRCPassPostProcessTestImage()
 
 void FRCPassPostProcessTestImage::Process(FRenderingCompositePassContext& Context)
 {
-	SCOPED_DRAW_EVENT(TestImage, DEC_SCENE_ITEMS);
+	SCOPED_DRAW_EVENT(Context.RHICmdList, TestImage, DEC_SCENE_ITEMS);
 
 	const FSceneView& View = Context.View;
 	const FSceneViewFamily& ViewFamily = *(View.Family);
diff --git a/Engine/Source/Runtime/Renderer/Private/PostProcess/PostProcessTonemap.cpp b/Engine/Source/Runtime/Renderer/Private/PostProcess/PostProcessTonemap.cpp
--- a/Engine/Source/Runtime/Renderer/Private/PostProcess/PostProcessTonemap.cpp
+++ b/Engine/Source/Runtime/Renderer/Private/PostProcess/PostProcessTonemap.cpp
@@ -1211,7 +1211,7 @@ static void SetShaderTempl(const FRenderingCompositePassContext& Context)
 
 void FRCPassPostProcessTonemap::Process(FRenderingCompositePassContext& Context)
 {
-	SCOPED_DRAW_EVENTF(PostProcessTonemap, DEC_SCENE_ITEMS, TEXT("Tonemapper#%d%s"), ConfigIndexPC, bDoGammaOnly ? TEXT(" GammaOnly") : TEXT(""));
+	SCOPED_DRAW_EVENTF(Context.RHICmdList, PostProcessTonemap, DEC_SCENE_ITEMS, TEXT("Tonemapper#%d%s"), ConfigIndexPC, bDoGammaOnly ? TEXT(" GammaOnly") : TEXT(""));
 
 	const FPooledRenderTargetDesc* InputDesc = GetInputDesc(ePId_Input0);
 
@@ -1583,7 +1583,7 @@ FRCPassPostProcessTonemapES2::FRCPassPostProcessTonemapES2(const FViewInfo& View
 
 void FRCPassPostProcessTonemapES2::Process(FRenderingCompositePassContext& Context)
 {
-	SCOPED_DRAW_EVENTF(PostProcessTonemap, DEC_SCENE_ITEMS, TEXT("Tonemapper#%d%s"), ConfigIndexMobile, bUsedFramebufferFetch ? TEXT(" FramebufferFetch=0") : TEXT("FramebufferFetch=1"));
+	SCOPED_DRAW_EVENTF(Context.RHICmdList, PostProcessTonemap, DEC_SCENE_ITEMS, TEXT("Tonemapper#%d%s"), ConfigIndexMobile, bUsedFramebufferFetch ? TEXT(" FramebufferFetch=0") : TEXT("FramebufferFetch=1"));
 
 	const FPooledRenderTargetDesc* InputDesc = GetInputDesc(ePId_Input0);
 
diff --git a/Engine/Source/Runtime/Renderer/Private/PostProcess/PostProcessUpscale.cpp b/Engine/Source/Runtime/Renderer/Private/PostProcess/PostProcessUpscale.cpp
--- a/Engine/Source/Runtime/Renderer/Private/PostProcess/PostProcessUpscale.cpp
+++ b/Engine/Source/Runtime/Renderer/Private/PostProcess/PostProcessUpscale.cpp
@@ -128,7 +128,7 @@ void FRCPassPostProcessUpscale::SetShader(const FRenderingCompositePassContext&
 
 void FRCPassPostProcessUpscale::Process(FRenderingCompositePassContext& Context)
 {
-	SCOPED_DRAW_EVENT(PostProcessUpscale, DEC_SCENE_ITEMS);
+	SCOPED_DRAW_EVENT(Context.RHICmdList, PostProcessUpscale, DEC_SCENE_ITEMS);
 	const FPooledRenderTargetDesc* InputDesc = GetInputDesc(ePId_Input0);
 
 	if(!InputDesc)
diff --git a/Engine/Source/Runtime/Renderer/Private/PostProcess/PostProcessVisualizeBuffer.cpp b/Engine/Source/Runtime/Renderer/Private/PostProcess/PostProcessVisualizeBuffer.cpp
--- a/Engine/Source/Runtime/Renderer/Private/PostProcess/PostProcessVisualizeBuffer.cpp
+++ b/Engine/Source/Runtime/Renderer/Private/PostProcess/PostProcessVisualizeBuffer.cpp
@@ -128,7 +128,7 @@ FShader* FRCPassPostProcessVisualizeBuffer::SetShaderTempl(const FRenderingCompo
 
 void FRCPassPostProcessVisualizeBuffer::Process(FRenderingCompositePassContext& Context)
 {
-	SCOPED_DRAW_EVENT(VisualizeBuffer, DEC_SCENE_ITEMS);
+	SCOPED_DRAW_EVENT(Context.RHICmdList, VisualizeBuffer, DEC_SCENE_ITEMS);
 	const FPooledRenderTargetDesc* InputDesc = GetInputDesc(ePId_Input0);
 
 	if(!InputDesc)
diff --git a/Engine/Source/Runtime/Renderer/Private/PostProcess/PostProcessVisualizeHDR.cpp b/Engine/Source/Runtime/Renderer/Private/PostProcess/PostProcessVisualizeHDR.cpp
--- a/Engine/Source/Runtime/Renderer/Private/PostProcess/PostProcessVisualizeHDR.cpp
+++ b/Engine/Source/Runtime/Renderer/Private/PostProcess/PostProcessVisualizeHDR.cpp
@@ -136,7 +136,7 @@ FString LogToString(float LogValue)
 
 void FRCPassPostProcessVisualizeHDR::Process(FRenderingCompositePassContext& Context)
 {
-	SCOPED_DRAW_EVENT(PostProcessVisualizeHDR, DEC_SCENE_ITEMS);
+	SCOPED_DRAW_EVENT(Context.RHICmdList, PostProcessVisualizeHDR, DEC_SCENE_ITEMS);
 	const FPooledRenderTargetDesc* InputDesc = GetInputDesc(ePId_Input0);
 
 	if(!InputDesc)
diff --git a/Engine/Source/Runtime/Renderer/Private/PostProcess/PostProcessWeightedSampleSum.cpp b/Engine/Source/Runtime/Renderer/Private/PostProcess/PostProcessWeightedSampleSum.cpp
--- a/Engine/Source/Runtime/Renderer/Private/PostProcess/PostProcessWeightedSampleSum.cpp
+++ b/Engine/Source/Runtime/Renderer/Private/PostProcess/PostProcessWeightedSampleSum.cpp
@@ -370,7 +370,7 @@ void FRCPassPostProcessWeightedSampleSum::Process(FRenderingCompositePassContext
 
 	uint32 NumSamples = Compute1DGaussianFilterKernel(FeatureLevel, EffectiveBlurRadius, OffsetAndWeight, MaxNumSamples);
 
-	SCOPED_DRAW_EVENTF(PostProcessWeightedSampleSum, DEC_SCENE_ITEMS, TEXT("PostProcessWeightedSampleSum#%d"), NumSamples);
+	SCOPED_DRAW_EVENTF(Context.RHICmdList, PostProcessWeightedSampleSum, DEC_SCENE_ITEMS, TEXT("PostProcessWeightedSampleSum#%d"), NumSamples);
 
 	// compute weights as weighted contributions of the TintValue
 	for(uint32 i = 0; i < NumSamples; ++i)
diff --git a/Engine/Source/Runtime/Renderer/Private/PostProcess/SceneRenderTargets.cpp b/Engine/Source/Runtime/Renderer/Private/PostProcess/SceneRenderTargets.cpp
--- a/Engine/Source/Runtime/Renderer/Private/PostProcess/SceneRenderTargets.cpp
+++ b/Engine/Source/Runtime/Renderer/Private/PostProcess/SceneRenderTargets.cpp
@@ -233,7 +233,7 @@ void FSceneRenderTargets::Allocate(const FSceneViewFamily& ViewFamily)
 /** Clears the GBuffer render targets to default values. */
 void FSceneRenderTargets::ClearGBufferTargets(FRHICommandListImmediate& RHICmdList, const FLinearColor& ClearColor)
 {
-	SCOPED_DRAW_EVENT(ClearGBufferTargets, DEC_SCENE_ITEMS);
+	SCOPED_DRAW_EVENT(RHICmdList, ClearGBufferTargets, DEC_SCENE_ITEMS);
 
 	// Clear GBufferA, GBufferB, GBufferC, GBufferD, GBufferE
 	{
@@ -256,7 +256,7 @@ void FSceneRenderTargets::ClearGBufferTargets(FRHICommandListImmediate& RHICmdLi
 
 void FSceneRenderTargets::BeginRenderingSceneColor(FRHICommandList& RHICmdList, bool bGBufferPass)
 {
-	SCOPED_DRAW_EVENT(BeginRenderingSceneColor, DEC_SCENE_ITEMS);
+	SCOPED_DRAW_EVENT(RHICmdList, BeginRenderingSceneColor, DEC_SCENE_ITEMS);
 
 	if(IsSimpleDynamicLightingEnabled())
 	{
@@ -603,7 +603,7 @@ void FSceneRenderTargets::AdjustGBufferRefCount(int Delta)
 
 void FSceneRenderTargets::FinishRenderingSceneColor(FRHICommandListImmediate& RHICmdList, bool bKeepChanges, const FResolveRect& ResolveRect)
 {
-	SCOPED_DRAW_EVENT(FinishRenderingSceneColor, DEC_SCENE_ITEMS);
+	SCOPED_DRAW_EVENT(RHICmdList, FinishRenderingSceneColor, DEC_SCENE_ITEMS);
 
 	if(bKeepChanges)
 	{
@@ -617,7 +617,7 @@ bool FSceneRenderTargets::BeginRenderingCustomDepth(FRHICommandListImmediate& RH
 
 	if(CustomDepth)
 	{
-		SCOPED_DRAW_EVENT(BeginRenderingCustomDepth, DEC_SCENE_ITEMS);
+		SCOPED_DRAW_EVENT(RHICmdList, BeginRenderingCustomDepth, DEC_SCENE_ITEMS);
 		
 		SetRenderTarget(RHICmdList, FTextureRHIRef(), CustomDepth->GetRenderTargetItem().ShaderResourceTexture);
 
@@ -629,7 +629,7 @@ bool FSceneRenderTargets::BeginRenderingCustomDepth(FRHICommandListImmediate& RH
 
 void FSceneRenderTargets::FinishRenderingCustomDepth(FRHICommandListImmediate& RHICmdList, const FResolveRect& ResolveRect)
 {
-	SCOPED_DRAW_EVENT(FinishRenderingCustomDepth, DEC_SCENE_ITEMS);
+	SCOPED_DRAW_EVENT(RHICmdList, FinishRenderingCustomDepth, DEC_SCENE_ITEMS);
 
 	RHICmdList.CopyToResolveTarget(SceneColor->GetRenderTargetItem().TargetableTexture, SceneColor->GetRenderTargetItem().ShaderResourceTexture, true, FResolveParams(ResolveRect));
 
@@ -641,7 +641,7 @@ void FSceneRenderTargets::FinishRenderingCustomDepth(FRHICommandListImmediate& R
 */
 void FSceneRenderTargets::ResolveSceneColor(FRHICommandList& RHICmdList, const FResolveRect& ResolveRect)
 {
-	SCOPED_DRAW_EVENT(ResolveSceneColor, DEC_SCENE_ITEMS);
+	SCOPED_DRAW_EVENT(RHICmdList, ResolveSceneColor, DEC_SCENE_ITEMS);
 
 	uint32 samples = GSceneRenderTargets.SceneColor->GetDesc().NumSamples;
 
@@ -707,7 +707,7 @@ void FSceneRenderTargets::ResolveGBufferSurfaces(FRHICommandList& RHICmdList, co
 {
 	if (CurrentFeatureLevel >= ERHIFeatureLevel::SM4)
 	{
-		SCOPED_DRAW_EVENT(ResolveGBufferSurfaces, DEC_SCENE_ITEMS);
+		SCOPED_DRAW_EVENT(RHICmdList, ResolveGBufferSurfaces, DEC_SCENE_ITEMS);
 
 		RHICmdList.CopyToResolveTarget(GSceneRenderTargets.GBufferA->GetRenderTargetItem().TargetableTexture, GSceneRenderTargets.GBufferA->GetRenderTargetItem().ShaderResourceTexture, true, FResolveParams(ResolveRect));
 		RHICmdList.CopyToResolveTarget(GSceneRenderTargets.GBufferB->GetRenderTargetItem().TargetableTexture, GSceneRenderTargets.GBufferB->GetRenderTargetItem().ShaderResourceTexture, true, FResolveParams(ResolveRect));
@@ -723,7 +723,7 @@ void FSceneRenderTargets::ResolveGBufferSurfaces(FRHICommandList& RHICmdList, co
 
 void FSceneRenderTargets::BeginRenderingPrePass(FRHICommandListImmediate& RHICmdList)
 {
-	SCOPED_DRAW_EVENT(BeginRenderingPrePass, DEC_SCENE_ITEMS);
+	SCOPED_DRAW_EVENT(RHICmdList, BeginRenderingPrePass, DEC_SCENE_ITEMS);
 
 	// Set the scene depth surface and a DUMMY buffer as color buffer
 	// (as long as it's the same dimension as the depth buffer),
@@ -733,7 +733,7 @@ void FSceneRenderTargets::BeginRenderingPrePass(FRHICommandListImmediate& RHICmd
 
 void FSceneRenderTargets::FinishRenderingPrePass(FRHICommandListImmediate& RHICmdList)
 {
-	SCOPED_DRAW_EVENT(FinishRenderingPrePass, DEC_SCENE_ITEMS);
+	SCOPED_DRAW_EVENT(RHICmdList, FinishRenderingPrePass, DEC_SCENE_ITEMS);
 	GRenderTargetPool.VisualizeTexture.SetCheckPoint(RHICmdList, SceneDepthZ);
 }
 
@@ -752,7 +752,7 @@ void FSceneRenderTargets::BeginRenderingShadowDepth(FRHICommandListImmediate& RH
 
 void FSceneRenderTargets::BeginRenderingCubeShadowDepth(FRHICommandListImmediate& RHICmdList, int32 ShadowResolution)
 {
-	SCOPED_DRAW_EVENT(BeginRenderingCubeShadowDepth, DEC_SCENE_ITEMS);
+	SCOPED_DRAW_EVENT(RHICmdList, BeginRenderingCubeShadowDepth, DEC_SCENE_ITEMS);
 	SetRenderTarget(RHICmdList, FTextureRHIRef(), GetCubeShadowDepthZSurface(ShadowResolution));
 }
 
@@ -793,20 +793,20 @@ void FSceneRenderTargets::FinishRenderingReflectiveShadowMap(FRHICommandList& RH
 
 void FSceneRenderTargets::FinishRenderingCubeShadowDepth(FRHICommandList& RHICmdList, int32 ShadowResolution, const FResolveParams& ResolveParams)
 {
-	SCOPED_DRAW_EVENT(FinishRenderingCubeShadowDepth, DEC_SCENE_ITEMS);
+	SCOPED_DRAW_EVENT(RHICmdList, FinishRenderingCubeShadowDepth, DEC_SCENE_ITEMS);
 	RHICmdList.CopyToResolveTarget(GetCubeShadowDepthZSurface(ShadowResolution), GetCubeShadowDepthZTexture(ShadowResolution), false, ResolveParams);
 }
 
 void FSceneRenderTargets::BeginRenderingSceneAlphaCopy(FRHICommandListImmediate& RHICmdList)
 {
-	SCOPED_DRAW_EVENT(BeginRenderingSceneAlphaCopy, DEC_SCENE_ITEMS);
+	SCOPED_DRAW_EVENT(RHICmdList, BeginRenderingSceneAlphaCopy, DEC_SCENE_ITEMS);
 	GRenderTargetPool.VisualizeTexture.SetCheckPoint(RHICmdList, GSceneRenderTargets.SceneAlphaCopy);
 	SetRenderTarget(RHICmdList, GetSceneAlphaCopySurface(), 0);
 }
 
 void FSceneRenderTargets::FinishRenderingSceneAlphaCopy(FRHICommandListImmediate& RHICmdList)
 {
-	SCOPED_DRAW_EVENT(FinishRenderingSceneAlphaCopy, DEC_SCENE_ITEMS);
+	SCOPED_DRAW_EVENT(RHICmdList, FinishRenderingSceneAlphaCopy, DEC_SCENE_ITEMS);
 	RHICmdList.CopyToResolveTarget(GetSceneAlphaCopySurface(), SceneAlphaCopy->GetRenderTargetItem().ShaderResourceTexture, false, FResolveParams(FResolveRect()));
 	GRenderTargetPool.VisualizeTexture.SetCheckPoint(RHICmdList, GSceneRenderTargets.SceneAlphaCopy);
 }
@@ -814,7 +814,7 @@ void FSceneRenderTargets::FinishRenderingSceneAlphaCopy(FRHICommandListImmediate
 
 void FSceneRenderTargets::BeginRenderingLightAttenuation(FRHICommandList& RHICmdList)
 {
-	SCOPED_DRAW_EVENT(BeginRenderingLightAttenuation, DEC_SCENE_ITEMS);
+	SCOPED_DRAW_EVENT(RHICmdList, BeginRenderingLightAttenuation, DEC_SCENE_ITEMS);
 
 	AllocLightAttenuation();
 
@@ -826,7 +826,7 @@ void FSceneRenderTargets::BeginRenderingLightAttenuation(FRHICommandList& RHICmd
 
 void FSceneRenderTargets::FinishRenderingLightAttenuation(FRHICommandList& RHICmdList)
 {
-	SCOPED_DRAW_EVENT(FinishRenderingLightAttenuation, DEC_SCENE_ITEMS);
+	SCOPED_DRAW_EVENT(RHICmdList, FinishRenderingLightAttenuation, DEC_SCENE_ITEMS);
 
 	// Resolve the light attenuation surface.
 	RHICmdList.CopyToResolveTarget(GetLightAttenuationSurface(), LightAttenuation->GetRenderTargetItem().ShaderResourceTexture, false, FResolveParams(FResolveRect()));
@@ -848,7 +848,7 @@ bool FSceneRenderTargets::BeginRenderingSeparateTranslucency(FRHICommandList& RH
 {
 	if(IsSeparateTranslucencyActive(View))
 	{
-		SCOPED_DRAW_EVENT(BeginSeparateTranslucency, DEC_SCENE_ITEMS);
+		SCOPED_DRAW_EVENT(RHICmdList, BeginSeparateTranslucency, DEC_SCENE_ITEMS);
 
 		FSceneViewState* ViewState = (FSceneViewState*)View.State;
 
@@ -889,7 +889,7 @@ void FSceneRenderTargets::FinishRenderingSeparateTranslucency(FRHICommandList& R
 
 void FSceneRenderTargets::ResolveSceneDepthTexture(FRHICommandList& RHICmdList)
 {
-	SCOPED_DRAW_EVENT(ResolveSceneDepthTexture, DEC_SCENE_ITEMS);
+	SCOPED_DRAW_EVENT(RHICmdList, ResolveSceneDepthTexture, DEC_SCENE_ITEMS);
 
 	RHICmdList.CopyToResolveTarget(GetSceneDepthSurface(), GetSceneDepthTexture(), true, FResolveParams());
 }
@@ -901,7 +901,7 @@ void FSceneRenderTargets::ResolveSceneDepthToAuxiliaryTexture(FRHICommandList& R
 	// is unsupported in d3d feature level 10.
 	if(!GSupportsDepthFetchDuringDepthTest)
 	{
-		SCOPED_DRAW_EVENT(ResolveSceneDepthToAuxiliaryTexture, DEC_SCENE_ITEMS);
+		SCOPED_DRAW_EVENT(RHICmdList, ResolveSceneDepthToAuxiliaryTexture, DEC_SCENE_ITEMS);
 
 		RHICmdList.CopyToResolveTarget(GetSceneDepthSurface(), GetAuxiliarySceneDepthTexture(), true, FResolveParams());
 	}
diff --git a/Engine/Source/Runtime/Renderer/Private/PostProcess/ScreenSpaceReflections.cpp b/Engine/Source/Runtime/Renderer/Private/PostProcess/ScreenSpaceReflections.cpp
--- a/Engine/Source/Runtime/Renderer/Private/PostProcess/ScreenSpaceReflections.cpp
+++ b/Engine/Source/Runtime/Renderer/Private/PostProcess/ScreenSpaceReflections.cpp
@@ -178,7 +178,7 @@ static int32 ComputeSSRQuality(float Quality)
 
 void FRCPassPostProcessScreenSpaceReflections::Process(FRenderingCompositePassContext& Context)
 {
-	SCOPED_DRAW_EVENT(ScreenSpaceReflections, DEC_SCENE_ITEMS);
+	SCOPED_DRAW_EVENT(Context.RHICmdList, ScreenSpaceReflections, DEC_SCENE_ITEMS);
 
 	const FSceneView& View = Context.View;
 	const auto FeatureLevel = Context.GetFeatureLevel();
diff --git a/Engine/Source/Runtime/Renderer/Private/PostProcess/VisualizeTexture.cpp b/Engine/Source/Runtime/Renderer/Private/PostProcess/VisualizeTexture.cpp
--- a/Engine/Source/Runtime/Renderer/Private/PostProcess/VisualizeTexture.cpp
+++ b/Engine/Source/Runtime/Renderer/Private/PostProcess/VisualizeTexture.cpp
@@ -418,13 +418,13 @@ void FVisualizeTexture::GenerateContent(FRHICommandListImmediate& RHICmdList, co
 
 	if(!(Desc.Flags & TexCreate_CPUReadback))		// We cannot make a texture lookup on such elements
 	{	
-		SCOPED_DRAW_EVENT(VisualizeTexture, DEC_SCENE_ITEMS);
+		SCOPED_DRAW_EVENT(RHICmdList, VisualizeTexture, DEC_SCENE_ITEMS);
 		// continue rendering to HDR if necessary
 		RenderVisualizeTexture(RHICmdList, FeatureLevel, VisualizeTextureData);
 	}
 
 	{
-		SCOPED_DRAW_EVENT(VisCopy, DEC_SCENE_ITEMS);
+		SCOPED_DRAW_EVENT(RHICmdList, VisCopy, DEC_SCENE_ITEMS);
 		RHICmdList.CopyToResolveTarget(DestRenderTarget.TargetableTexture, DestRenderTarget.ShaderResourceTexture, false, FResolveParams());
 	}
 
@@ -525,7 +525,7 @@ void FVisualizeTexture::PresentContent(FRHICommandListImmediate& RHICmdList, con
 	FIntRect VisualizeTextureRect = ComputeVisualizeTextureRect(Desc.Extent);
 
 	{
-		SCOPED_DRAW_EVENT(VisCopyToMain, DEC_SCENE_ITEMS);
+		SCOPED_DRAW_EVENT(RHICmdList, VisCopyToMain, DEC_SCENE_ITEMS);
 		// Draw a quad mapping scene color to the view's render target
 		DrawRectangle(
 			RHICmdList,
diff --git a/Engine/Source/Runtime/Renderer/Private/ReflectionEnvironment.cpp b/Engine/Source/Runtime/Renderer/Private/ReflectionEnvironment.cpp
--- a/Engine/Source/Runtime/Renderer/Private/ReflectionEnvironment.cpp
+++ b/Engine/Source/Runtime/Renderer/Private/ReflectionEnvironment.cpp
@@ -709,7 +709,7 @@ void FDeferredShadingSceneRenderer::RenderTiledDeferredImageBasedReflections(FRH
 		// ReflectionEnv is assumed to be on when going into this method
 		{
 			// Render the reflection environment with tiled deferred culling
-			SCOPED_DRAW_EVENT(ReflectionEnvironmentGather, DEC_SCENE_ITEMS);
+			SCOPED_DRAW_EVENT(RHICmdList, ReflectionEnvironmentGather, DEC_SCENE_ITEMS);
 
 			SetRenderTarget(RHICmdList, NULL, NULL);
 
@@ -819,7 +819,7 @@ void FDeferredShadingSceneRenderer::RenderStandardDeferredImageBasedReflections(
 		{
 			bRequiresApply = true;
 
-			SCOPED_DRAW_EVENT(StandardDeferredReflectionEnvironment, DEC_SCENE_ITEMS);
+			SCOPED_DRAW_EVENT(RHICmdList, StandardDeferredReflectionEnvironment, DEC_SCENE_ITEMS);
 
 			{
 				FPooledRenderTargetDesc Desc = GSceneRenderTargets.GetSceneColor()->GetDesc();
@@ -882,7 +882,7 @@ void FDeferredShadingSceneRenderer::RenderStandardDeferredImageBasedReflections(
 		if (bRequiresApply)
 		{
 			// Apply reflections to screen
-			SCOPED_DRAW_EVENT(ReflectionApply, DEC_SCENE_ITEMS);
+			SCOPED_DRAW_EVENT(RHICmdList, ReflectionApply, DEC_SCENE_ITEMS);
 
 			GSceneRenderTargets.BeginRenderingSceneColor(RHICmdList);
 
diff --git a/Engine/Source/Runtime/Renderer/Private/ReflectionEnvironmentCapture.cpp b/Engine/Source/Runtime/Renderer/Private/ReflectionEnvironmentCapture.cpp
--- a/Engine/Source/Runtime/Renderer/Private/ReflectionEnvironmentCapture.cpp
+++ b/Engine/Source/Runtime/Renderer/Private/ReflectionEnvironmentCapture.cpp
@@ -318,11 +318,11 @@ void FilterReflectionEnvironment(FRHICommandListImmediate& RHICmdList, ERHIFeatu
 	auto ShaderMap = GetGlobalShaderMap(FeatureLevel);
 
 	{	
-		SCOPED_DRAW_EVENT(DownsampleCubeMips, DEC_SCENE_ITEMS);
+		SCOPED_DRAW_EVENT(RHICmdList, DownsampleCubeMips, DEC_SCENE_ITEMS);
 		// Downsample all the mips, each one reads from the mip above it
 		for (int32 MipIndex = 1; MipIndex < NumMips; MipIndex++)
 		{
-			SCOPED_DRAW_EVENT(DownsampleCubeMip, DEC_SCENE_ITEMS);
+			SCOPED_DRAW_EVENT(RHICmdList, DownsampleCubeMip, DEC_SCENE_ITEMS);
 			const int32 SourceMipIndex = FMath::Max(MipIndex - 1, 0);
 			const int32 MipSize = 1 << (NumMips - MipIndex - 1);
 
@@ -349,7 +349,7 @@ void FilterReflectionEnvironment(FRHICommandListImmediate& RHICmdList, ERHIFeatu
 				// Draw each face with a geometry shader that routes to the correct render target slice
 				for (int32 CubeFace = 0; CubeFace < CubeFace_MAX; CubeFace++)
 				{
-					SCOPED_DRAW_EVENT(DownsampleCubeFace, DEC_SCENE_ITEMS);
+					SCOPED_DRAW_EVENT(RHICmdList, DownsampleCubeFace, DEC_SCENE_ITEMS);
 					GeometryShader->SetParameters(RHICmdList, CubeFace);
 
 					PixelShader->SetParameters(RHICmdList, CubeFace, SourceMipIndex, EffectiveSource);
@@ -411,7 +411,7 @@ void FilterReflectionEnvironment(FRHICommandListImmediate& RHICmdList, ERHIFeatu
 
 	if (OutIrradianceEnvironmentMap)
 	{
-		SCOPED_DRAW_EVENT(ComputeDiffuseIrradiance, DEC_SCENE_ITEMS);
+		SCOPED_DRAW_EVENT(RHICmdList, ComputeDiffuseIrradiance, DEC_SCENE_ITEMS);
 		check(DiffuseConvolutionSource != NULL);
 		ComputeDiffuseIrradiance(RHICmdList, FeatureLevel, DiffuseConvolutionSource->ShaderResourceTexture, DiffuseConvolutionSourceMip, OutIrradianceEnvironmentMap);
 	}
@@ -419,11 +419,11 @@ void FilterReflectionEnvironment(FRHICommandListImmediate& RHICmdList, ERHIFeatu
 	ComputeAverageBrightness(RHICmdList, FeatureLevel);
 
 	{	
-		SCOPED_DRAW_EVENT(FilterCubeMap, DEC_SCENE_ITEMS);
+		SCOPED_DRAW_EVENT(RHICmdList, FilterCubeMap, DEC_SCENE_ITEMS);
 		// Filter all the mips, each one reads from whichever scratch render target holds the downsampled contents, and writes to the destination cubemap
 		for (int32 MipIndex = 0; MipIndex < NumMips; MipIndex++)
 		{
-			SCOPED_DRAW_EVENT(FilterCubeMip, DEC_SCENE_ITEMS);
+			SCOPED_DRAW_EVENT(RHICmdList, FilterCubeMip, DEC_SCENE_ITEMS);
 			FSceneRenderTargetItem& EffectiveRT = GetEffectiveRenderTarget(false, MipIndex);
 			FSceneRenderTargetItem& EffectiveSource = GetEffectiveSourceTexture(false, MipIndex);
 			check(EffectiveRT.TargetableTexture != EffectiveSource.ShaderResourceTexture);
@@ -460,7 +460,7 @@ void FilterReflectionEnvironment(FRHICommandListImmediate& RHICmdList, ERHIFeatu
 
 				for (int32 CubeFace = 0; CubeFace < CubeFace_MAX; CubeFace++)
 				{
-					SCOPED_DRAW_EVENT(FilterCubeFace, DEC_SCENE_ITEMS);
+					SCOPED_DRAW_EVENT(RHICmdList, FilterCubeFace, DEC_SCENE_ITEMS);
 					GeometryShader->SetParameters(RHICmdList, CubeFace);
 
 					PixelShader->SetParameters(RHICmdList, CubeFace, MipIndex, EffectiveSource);
@@ -833,7 +833,7 @@ void CaptureSceneToScratchCubemap(FRHICommandListImmediate& RHICmdList, FSceneRe
 	const auto FeatureLevel = SceneRenderer->FeatureLevel;
 	
 	{
-		SCOPED_DRAW_EVENT(CubeMapCapture, DEC_SCENE_ITEMS);
+		SCOPED_DRAW_EVENT(RHICmdList, CubeMapCapture, DEC_SCENE_ITEMS);
 
 		// Render the scene normally for one face of the cubemap
 		SceneRenderer->Render(RHICmdList);
@@ -856,7 +856,7 @@ void CaptureSceneToScratchCubemap(FRHICommandListImmediate& RHICmdList, FSceneRe
 		FSceneRenderTargetItem& EffectiveColorRT =  GSceneRenderTargets.ReflectionColorScratchCubemap[0]->GetRenderTargetItem();
 
 		{
-			SCOPED_DRAW_EVENT(CubeMapCopyScene, DEC_SCENE_ITEMS);
+			SCOPED_DRAW_EVENT(RHICmdList, CubeMapCopyScene, DEC_SCENE_ITEMS);
 			if (GSupportsGSRenderTargetLayerSwitchingToMips)
 			{
 			// Copy the captured scene into the cubemap face
diff --git a/Engine/Source/Runtime/Renderer/Private/SceneCaptureRendering.cpp b/Engine/Source/Runtime/Renderer/Private/SceneCaptureRendering.cpp
--- a/Engine/Source/Runtime/Renderer/Private/SceneCaptureRendering.cpp
+++ b/Engine/Source/Runtime/Renderer/Private/SceneCaptureRendering.cpp
@@ -24,7 +24,7 @@ void UpdateSceneCaptureContent_RenderThread(FRHICommandListImmediate& RHICmdList
 #if WANTS_DRAW_MESH_EVENTS
 		FString EventName;
 		OwnerName.ToString(EventName);
-		SCOPED_DRAW_EVENTF(SceneCapture, DEC_SCENE_ITEMS, TEXT("SceneCapture %s"), *EventName);
+		SCOPED_DRAW_EVENTF(RHICmdList, SceneCapture, DEC_SCENE_ITEMS, TEXT("SceneCapture %s"), *EventName);
 #endif
 
 		// Render the scene normally
diff --git a/Engine/Source/Runtime/Renderer/Private/SceneHitProxyRendering.cpp b/Engine/Source/Runtime/Renderer/Private/SceneHitProxyRendering.cpp
--- a/Engine/Source/Runtime/Renderer/Private/SceneHitProxyRendering.cpp
+++ b/Engine/Source/Runtime/Renderer/Private/SceneHitProxyRendering.cpp
@@ -216,7 +216,7 @@ void FHitProxyDrawingPolicy::SetMeshRenderState(
 	const ContextDataType PolicyContext
 	) const
 {
-	EmitMeshDrawEvents(PrimitiveSceneProxy, Mesh);
+	EmitMeshDrawEvents(RHICmdList, PrimitiveSceneProxy, Mesh);
 
 	const FMeshBatchElement& BatchElement = Mesh.Elements[BatchElementIndex];
 
diff --git a/Engine/Source/Runtime/Renderer/Private/SceneOcclusion.cpp b/Engine/Source/Runtime/Renderer/Private/SceneOcclusion.cpp
--- a/Engine/Source/Runtime/Renderer/Private/SceneOcclusion.cpp
+++ b/Engine/Source/Runtime/Renderer/Private/SceneOcclusion.cpp
@@ -675,7 +675,7 @@ IMPLEMENT_SHADER_TYPE(,FHZBTestPS,TEXT("HZBOcclusion"),TEXT("HZBTestPS"),SF_Pixe
 
 void FHZBOcclusionTester::Submit(FRHICommandListImmediate& RHICmdList, const FViewInfo& View)
 {
-	SCOPED_DRAW_EVENT(SubmitHZB, DEC_SCENE_ITEMS);
+	SCOPED_DRAW_EVENT(RHICmdList, SubmitHZB, DEC_SCENE_ITEMS);
 
 	FSceneViewState* ViewState = (FSceneViewState*)View.State;
 	if( !ViewState )
@@ -744,7 +744,7 @@ void FHZBOcclusionTester::Submit(FRHICommandListImmediate& RHICmdList, const FVi
 
 	// Draw test
 	{
-		SCOPED_DRAW_EVENT(TestHZB, DEC_SCENE_ITEMS);
+		SCOPED_DRAW_EVENT(RHICmdList, TestHZB, DEC_SCENE_ITEMS);
 
 		SetRenderTarget(RHICmdList, ResultsTextureGPU->GetRenderTargetItem().TargetableTexture, NULL);
 
@@ -855,7 +855,7 @@ IMPLEMENT_SHADER_TYPE(template<>,THZBBuildPS<1>,TEXT("HZBOcclusion"),TEXT("HZBBu
 void BuildHZB(FRHICommandListImmediate& RHICmdList, const FViewInfo& View)
 {
 	QUICK_SCOPE_CYCLE_COUNTER(STAT_BuildHZB);
-	SCOPED_DRAW_EVENT(BuildHZB, DEC_SCENE_ITEMS);
+	SCOPED_DRAW_EVENT(RHICmdList, BuildHZB, DEC_SCENE_ITEMS);
 
 	FSceneViewState* ViewState = (FSceneViewState*)View.State;
 
@@ -888,7 +888,7 @@ void BuildHZB(FRHICommandListImmediate& RHICmdList, const FViewInfo& View)
 	
 	// Mip 0
 	{
-		SCOPED_DRAW_EVENTF(BuildHZB, DEC_SCENE_ITEMS, TEXT("HZB#%d"), 0);
+		SCOPED_DRAW_EVENTF(RHICmdList, BuildHZB, DEC_SCENE_ITEMS, TEXT("HZB#%d"), 0);
 
 		SetRenderTarget(RHICmdList, HZBRenderTarget.TargetableTexture, 0, NULL);
 
@@ -922,7 +922,7 @@ void BuildHZB(FRHICommandListImmediate& RHICmdList, const FViewInfo& View)
 	// Mip 1-7
 	for( uint8 MipIndex = 1; MipIndex < NumMips; MipIndex++ )
 	{
-		SCOPED_DRAW_EVENTF(BuildHZB, DEC_SCENE_ITEMS, TEXT("HZB#%d"), MipIndex);
+		SCOPED_DRAW_EVENTF(RHICmdList, BuildHZB, DEC_SCENE_ITEMS, TEXT("HZB#%d"), MipIndex);
 
 		SetRenderTarget(RHICmdList, HZBRenderTarget.TargetableTexture, MipIndex, NULL);
 
@@ -974,7 +974,7 @@ void FDeferredShadingSceneRenderer::BeginOcclusionTests(FRHICommandListImmediate
 	// Perform occlusion queries for each view
 	for(int32 ViewIndex = 0;ViewIndex < Views.Num();ViewIndex++)
 	{
-		SCOPED_DRAW_EVENT(BeginOcclusionTests, DEC_SCENE_ITEMS);
+		SCOPED_DRAW_EVENT(RHICmdList, BeginOcclusionTests, DEC_SCENE_ITEMS);
 		FViewInfo& View = Views[ViewIndex];
 
 		if (bUseDownsampledDepth)
@@ -1029,7 +1029,7 @@ void FDeferredShadingSceneRenderer::BeginOcclusionTests(FRHICommandListImmediate
 			ShadowOcclusionQueryMap.Reset();
 
 			{
-				SCOPED_DRAW_EVENT(ShadowFrustumQueries, DEC_SCENE_ITEMS);
+				SCOPED_DRAW_EVENT(RHICmdList, ShadowFrustumQueries, DEC_SCENE_ITEMS);
 
 				for(TSparseArray<FLightSceneInfoCompact>::TConstIterator LightIt(Scene->Lights);LightIt;++LightIt)
 				{
@@ -1100,11 +1100,11 @@ void FDeferredShadingSceneRenderer::BeginOcclusionTests(FRHICommandListImmediate
 				VertexShader->SetParameters(RHICmdList, View);
 
 				{
-					SCOPED_DRAW_EVENT(IndividualQueries, DEC_SCENE_ITEMS);
+					SCOPED_DRAW_EVENT(RHICmdList, IndividualQueries, DEC_SCENE_ITEMS);
 					View.IndividualOcclusionQueries.Flush(RHICmdList);
 				}
 				{
-					SCOPED_DRAW_EVENT(GroupedQueries, DEC_SCENE_ITEMS);
+					SCOPED_DRAW_EVENT(RHICmdList, GroupedQueries, DEC_SCENE_ITEMS);
 					View.GroupedOcclusionQueries.Flush(RHICmdList);
 				}
 			}
@@ -1120,7 +1120,7 @@ void FDeferredShadingSceneRenderer::BeginOcclusionTests(FRHICommandListImmediate
 		{
 			check( ViewState->HZBOcclusionTests.IsValidFrame(View.FrameNumber) );
 
-			SCOPED_DRAW_EVENT(HZB, DEC_SCENE_ITEMS);
+			SCOPED_DRAW_EVENT(RHICmdList, HZB, DEC_SCENE_ITEMS);
 
 			BuildHZB(RHICmdList, View);
 			ViewState->HZBOcclusionTests.Submit(RHICmdList, View);
diff --git a/Engine/Source/Runtime/Renderer/Private/SceneRendering.cpp b/Engine/Source/Runtime/Renderer/Private/SceneRendering.cpp
--- a/Engine/Source/Runtime/Renderer/Private/SceneRendering.cpp
+++ b/Engine/Source/Runtime/Renderer/Private/SceneRendering.cpp
@@ -662,7 +662,7 @@ FSceneRenderer::~FSceneRenderer()
 */
 void FSceneRenderer::RenderFinish(FRHICommandListImmediate& RHICmdList)
 {
-	SCOPED_DRAW_EVENT(RenderFinish, DEC_SCENE_ITEMS);
+	SCOPED_DRAW_EVENT(RHICmdList, RenderFinish, DEC_SCENE_ITEMS);
 
 	if(FRCPassPostProcessBusyWait::IsPassRequired())
 	{
@@ -691,7 +691,7 @@ void FSceneRenderer::RenderFinish(FRHICommandListImmediate& RHICmdList)
 	{
 		for(int32 ViewIndex = 0;ViewIndex < Views.Num();ViewIndex++)
 		{	
-			SCOPED_CONDITIONAL_DRAW_EVENTF(EventView, Views.Num() > 1, DEC_SCENE_ITEMS, TEXT("View%d"), ViewIndex);
+			SCOPED_CONDITIONAL_DRAW_EVENTF(RHICmdList, EventView, Views.Num() > 1, DEC_SCENE_ITEMS, TEXT("View%d"), ViewIndex);
 			FViewInfo& View = Views[ViewIndex];
 
 			bool bShowPrecomputedVisibilityWarning = false;
@@ -866,11 +866,11 @@ void FSceneRenderer::RenderCustomDepthPass(FRHICommandListImmediate& RHICmdList)
 	// Render CustomDepth
 	if (GSceneRenderTargets.BeginRenderingCustomDepth(RHICmdList, bPrimitives))
 	{
-		SCOPED_DRAW_EVENT(CustomDepth, DEC_SCENE_ITEMS);
+		SCOPED_DRAW_EVENT(RHICmdList, CustomDepth, DEC_SCENE_ITEMS);
 
 		for(int32 ViewIndex = 0;ViewIndex < Views.Num();ViewIndex++)
 		{
-			SCOPED_CONDITIONAL_DRAW_EVENTF(EventView, Views.Num() > 1, DEC_SCENE_ITEMS, TEXT("View%d"), ViewIndex);
+			SCOPED_CONDITIONAL_DRAW_EVENTF(RHICmdList, EventView, Views.Num() > 1, DEC_SCENE_ITEMS, TEXT("View%d"), ViewIndex);
 
 			FViewInfo& View = Views[ViewIndex];
 
diff --git a/Engine/Source/Runtime/Renderer/Private/SceneVisibility.cpp b/Engine/Source/Runtime/Renderer/Private/SceneVisibility.cpp
--- a/Engine/Source/Runtime/Renderer/Private/SceneVisibility.cpp
+++ b/Engine/Source/Runtime/Renderer/Private/SceneVisibility.cpp
@@ -1609,7 +1609,7 @@ uint32 GetShadowQuality();
  */
 void FDeferredShadingSceneRenderer::InitViews(FRHICommandListImmediate& RHICmdList)
 {
-	SCOPED_DRAW_EVENT(InitViews, DEC_SCENE_ITEMS);
+	SCOPED_DRAW_EVENT(RHICmdList, InitViews, DEC_SCENE_ITEMS);
 
 	SCOPE_CYCLE_COUNTER(STAT_InitViewsTime);
 
diff --git a/Engine/Source/Runtime/Renderer/Private/ShaderComplexityRendering.cpp b/Engine/Source/Runtime/Renderer/Private/ShaderComplexityRendering.cpp
--- a/Engine/Source/Runtime/Renderer/Private/ShaderComplexityRendering.cpp
+++ b/Engine/Source/Runtime/Renderer/Private/ShaderComplexityRendering.cpp
@@ -150,7 +150,7 @@ extern TGlobalResource<FFilterVertexDeclaration> GFilterVertexDeclaration;
 
 void FRCPassPostProcessVisualizeComplexity::Process(FRenderingCompositePassContext& Context)
 {
-	SCOPED_DRAW_EVENT(PostProcessVisualizeComplexity, DEC_SCENE_ITEMS);
+	SCOPED_DRAW_EVENT(Context.RHICmdList, PostProcessVisualizeComplexity, DEC_SCENE_ITEMS);
 	const FPooledRenderTargetDesc* InputDesc = GetInputDesc(ePId_Input0);
 
 	if(!InputDesc)
diff --git a/Engine/Source/Runtime/Renderer/Private/ShadowRendering.cpp b/Engine/Source/Runtime/Renderer/Private/ShadowRendering.cpp
--- a/Engine/Source/Runtime/Renderer/Private/ShadowRendering.cpp
+++ b/Engine/Source/Runtime/Renderer/Private/ShadowRendering.cpp
@@ -942,7 +942,7 @@ void FShadowDepthDrawingPolicy<bRenderingReflectiveShadowMaps>::SetMeshRenderSta
 {
 	const FMeshBatchElement& BatchElement = Mesh.Elements[BatchElementIndex];
 
-	EmitMeshDrawEvents(PrimitiveSceneProxy, Mesh);
+	EmitMeshDrawEvents(RHICmdList, PrimitiveSceneProxy, Mesh);
 
 	VertexShader->SetMesh(RHICmdList, VertexFactory,View,PrimitiveSceneProxy,BatchElement);
 
@@ -1524,7 +1524,7 @@ void FProjectedShadowInfo::RenderDepth(FRHICommandList& RHICmdList, FDeferredSha
 #if WANTS_DRAW_MESH_EVENTS
 	FString EventName;
 	GetShadowTypeNameForDrawEvent(EventName);
-	SCOPED_DRAW_EVENTF(EventShadowDepthActor, DEC_SCENE_ITEMS, *EventName);
+	SCOPED_DRAW_EVENTF(RHICmdList, EventShadowDepthActor, DEC_SCENE_ITEMS, *EventName);
 #endif
 
 	CONDITIONAL_SCOPE_CYCLE_COUNTER(STAT_RenderWholeSceneShadowDepthsTime, bWholeSceneShadow);
@@ -1718,7 +1718,7 @@ void FProjectedShadowInfo::RenderProjection(FRHICommandListImmediate& RHICmdList
 #if WANTS_DRAW_MESH_EVENTS
 	FString EventName;
 	GetShadowTypeNameForDrawEvent(EventName);
-	SCOPED_DRAW_EVENTF(EventShadowProjectionActor, DEC_SCENE_ITEMS, *EventName);
+	SCOPED_DRAW_EVENTF(RHICmdList, EventShadowProjectionActor, DEC_SCENE_ITEMS, *EventName);
 #endif
 
 	FScopeCycleCounter Scope(bWholeSceneShadow ? GET_STATID(STAT_RenderWholeSceneShadowProjectionsTime) : GET_STATID(STAT_RenderPerObjectShadowProjectionsTime));
@@ -1821,7 +1821,7 @@ void FProjectedShadowInfo::RenderProjection(FRHICommandListImmediate& RHICmdList
 	// If this is a preshadow, mask the projection by the receiver primitives.
 	if (bPreShadow)
 	{
-		SCOPED_DRAW_EVENTF(EventMaskSubjects, DEC_SCENE_ITEMS, TEXT("Stencil Mask Subjects"));
+		SCOPED_DRAW_EVENTF(RHICmdList, EventMaskSubjects, DEC_SCENE_ITEMS, TEXT("Stencil Mask Subjects"));
 
 		// Set stencil to one.
 		// Note, this is a reversed Z depth surface, using CF_GreaterEqual.
@@ -2560,7 +2560,7 @@ bool FDeferredShadingSceneRenderer::RenderOnePassPointLightShadows(FRHICommandLi
 			INC_DWORD_STAT(STAT_WholeSceneShadows);
 
 			{
-				SCOPED_DRAW_EVENT(ShadowDepthsFromOpaque, DEC_SCENE_ITEMS);
+				SCOPED_DRAW_EVENT(RHICmdList, ShadowDepthsFromOpaque, DEC_SCENE_ITEMS);
 				GSceneRenderTargets.BeginRenderingCubeShadowDepth(RHICmdList, ProjectedShadowInfo->ResolutionX);
 				ProjectedShadowInfo->bAllocated = true;
 				//gilmerge, if clear depth doesn't go into a cmdlist, then how can renderdepth? This occurs in 4 spots
@@ -2572,11 +2572,11 @@ bool FDeferredShadingSceneRenderer::RenderOnePassPointLightShadows(FRHICommandLi
 			{
 				GSceneRenderTargets.BeginRenderingLightAttenuation(RHICmdList);
 
-				SCOPED_DRAW_EVENT(ShadowProjectionOnOpaque, DEC_SCENE_ITEMS);
+				SCOPED_DRAW_EVENT(RHICmdList, ShadowProjectionOnOpaque, DEC_SCENE_ITEMS);
 
 				for (int32 ViewIndex = 0; ViewIndex < Views.Num(); ViewIndex++)
 				{
-					SCOPED_CONDITIONAL_DRAW_EVENTF(EventView, Views.Num() > 1, DEC_SCENE_ITEMS, TEXT("View%d"), ViewIndex);
+					SCOPED_CONDITIONAL_DRAW_EVENTF(RHICmdList, EventView, Views.Num() > 1, DEC_SCENE_ITEMS, TEXT("View%d"), ViewIndex);
 
 					const FViewInfo& View = Views[ViewIndex];
 
@@ -2592,7 +2592,7 @@ bool FDeferredShadingSceneRenderer::RenderOnePassPointLightShadows(FRHICommandLi
 			if (!LightSceneInfo->Proxy->HasStaticShadowing())
 			{
 				bInjectedTranslucentVolume = true;
-				SCOPED_DRAW_EVENT(InjectTranslucentVolume, DEC_SCENE_ITEMS);
+				SCOPED_DRAW_EVENT(RHICmdList, InjectTranslucentVolume, DEC_SCENE_ITEMS);
 				// Inject the shadowed light into the translucency lighting volumes
 				InjectTranslucentVolumeLighting(RHICmdList, *LightSceneInfo, ProjectedShadowInfo);
 			}
@@ -2616,7 +2616,7 @@ void FDeferredShadingSceneRenderer::RenderProjections(
 
 	for (int32 ViewIndex = 0; ViewIndex < Views.Num(); ViewIndex++)
 	{
-		SCOPED_CONDITIONAL_DRAW_EVENTF(EventView, Views.Num() > 1, DEC_SCENE_ITEMS, TEXT("View%d"), ViewIndex);
+		SCOPED_CONDITIONAL_DRAW_EVENTF(RHICmdList, EventView, Views.Num() > 1, DEC_SCENE_ITEMS, TEXT("View%d"), ViewIndex);
 
 		const FViewInfo& View = Views[ViewIndex];
 
@@ -2793,7 +2793,7 @@ bool FDeferredShadingSceneRenderer::RenderTranslucentProjectedShadows(FRHIComman
 
 		{
 			// Render the translucent shadow depths.
-			SCOPED_DRAW_EVENT(TranslucencyShadowMapGeneration, DEC_SCENE_ITEMS);
+			SCOPED_DRAW_EVENT(RHICmdList, TranslucencyShadowMapGeneration, DEC_SCENE_ITEMS);
 
 			for (int32 ShadowIndex = 0; ShadowIndex < Shadows.Num(); ShadowIndex++)
 			{
@@ -2809,7 +2809,7 @@ bool FDeferredShadingSceneRenderer::RenderTranslucentProjectedShadows(FRHIComman
 		// Render the shadow projections.
 		{
 			{
-				SCOPED_DRAW_EVENT(ShadowProjectionOnOpaque, DEC_SCENE_ITEMS);
+				SCOPED_DRAW_EVENT(RHICmdList, ShadowProjectionOnOpaque, DEC_SCENE_ITEMS);
 				RenderProjections(RHICmdList, LightSceneInfo, Shadows);
 			}
 
@@ -2948,7 +2948,7 @@ bool FDeferredShadingSceneRenderer::RenderReflectiveShadowMaps(FRHICommandListIm
 
 		{
 			// Render the shadow depths.
-			SCOPED_DRAW_EVENT(ReflectiveShadowMapsFromOpaque, DEC_SCENE_ITEMS);
+			SCOPED_DRAW_EVENT(RHICmdList, ReflectiveShadowMapsFromOpaque, DEC_SCENE_ITEMS);
 
 			// render the RSMs
 			for (int32 ShadowIndex = 0; ShadowIndex < Shadows.Num(); ShadowIndex++)
@@ -3128,7 +3128,7 @@ bool FDeferredShadingSceneRenderer::RenderProjectedShadows(FRHICommandListImmedi
 
 		{
 			// Render the shadow depths.
-			SCOPED_DRAW_EVENT(ShadowDepthsFromOpaque, DEC_SCENE_ITEMS);
+			SCOPED_DRAW_EVENT(RHICmdList, ShadowDepthsFromOpaque, DEC_SCENE_ITEMS);
 
 			GSceneRenderTargets.BeginRenderingShadowDepth(RHICmdList);
 
@@ -3150,7 +3150,7 @@ bool FDeferredShadingSceneRenderer::RenderProjectedShadows(FRHICommandListImmedi
 		// Render the shadow projections.
 		{
 			{
-				SCOPED_DRAW_EVENT(ShadowProjectionOnOpaque, DEC_SCENE_ITEMS);
+				SCOPED_DRAW_EVENT(RHICmdList, ShadowProjectionOnOpaque, DEC_SCENE_ITEMS);
 				RenderProjections(RHICmdList, LightSceneInfo, Shadows);
 			}
 			
@@ -3168,7 +3168,7 @@ bool FDeferredShadingSceneRenderer::RenderProjectedShadows(FRHICommandListImmedi
 					&& (!LightSceneInfo->Proxy->HasStaticShadowing() || ProjectedShadowInfo->IsWholeSceneDirectionalShadow()))
 				{
 					bInjectedTranslucentVolume = true;
-					SCOPED_DRAW_EVENT(InjectTranslucentVolume, DEC_SCENE_ITEMS);
+					SCOPED_DRAW_EVENT(RHICmdList, InjectTranslucentVolume, DEC_SCENE_ITEMS);
 					// Inject the shadowed light into the translucency lighting volumes
 					InjectTranslucentVolumeLighting(RHICmdList, *LightSceneInfo, ProjectedShadowInfo);
 				}
@@ -3261,7 +3261,7 @@ bool FDeferredShadingSceneRenderer::RenderCachedPreshadows(FRHICommandListImmedi
 	{
 		if (bHasDepthsToRender)
 		{
-			SCOPED_DRAW_EVENTF(EventShadowDepths, DEC_SCENE_ITEMS, TEXT("Preshadow Cache Depths"));
+			SCOPED_DRAW_EVENTF(RHICmdList, EventShadowDepths, DEC_SCENE_ITEMS, TEXT("Preshadow Cache Depths"));
 
 			SetRenderTarget(RHICmdList, FTextureRHIRef(), GSceneRenderTargets.PreShadowCacheDepthZ->GetRenderTargetItem().TargetableTexture);
 
@@ -3294,7 +3294,7 @@ bool FDeferredShadingSceneRenderer::RenderCachedPreshadows(FRHICommandListImmedi
 
 		// Render the shadow projections.
 		{
-			SCOPED_DRAW_EVENTF(EventShadowProj, DEC_SCENE_ITEMS, TEXT("Cached Preshadow Projections"));
+			SCOPED_DRAW_EVENTF(RHICmdList, EventShadowProj, DEC_SCENE_ITEMS, TEXT("Cached Preshadow Projections"));
 			RenderProjections(RHICmdList, LightSceneInfo, OpaqueCachedPreshadows);
 
 			// Mark the attenuation buffer as dirty.
diff --git a/Engine/Source/Runtime/Renderer/Private/TiledDeferredLightRendering.cpp b/Engine/Source/Runtime/Renderer/Private/TiledDeferredLightRendering.cpp
--- a/Engine/Source/Runtime/Renderer/Private/TiledDeferredLightRendering.cpp
+++ b/Engine/Source/Runtime/Renderer/Private/TiledDeferredLightRendering.cpp
@@ -351,7 +351,7 @@ void FDeferredShadingSceneRenderer::RenderTiledDeferredLighting(FRHICommandListI
 			}
 
 			{
-				SCOPED_DRAW_EVENT(TiledDeferredLighting, DEC_SCENE_ITEMS);
+				SCOPED_DRAW_EVENT(RHICmdList, TiledDeferredLighting, DEC_SCENE_ITEMS);
 
 				IPooledRenderTarget& InTexture = *GSceneRenderTargets.GetSceneColor();
 
diff --git a/Engine/Source/Runtime/Renderer/Private/TranslucentLighting.cpp b/Engine/Source/Runtime/Renderer/Private/TranslucentLighting.cpp
--- a/Engine/Source/Runtime/Renderer/Private/TranslucentLighting.cpp
+++ b/Engine/Source/Runtime/Renderer/Private/TranslucentLighting.cpp
@@ -528,7 +528,7 @@ void FProjectedShadowInfo::RenderTranslucencyDepths(FRHICommandList& RHICmdList,
 #if WANTS_DRAW_MESH_EVENTS
 		FString EventName;
 		GetShadowTypeNameForDrawEvent(EventName);
-		SCOPED_DRAW_EVENTF(EventShadowDepthActor, DEC_SCENE_ITEMS, *EventName);
+		SCOPED_DRAW_EVENTF(RHICmdList, EventShadowDepthActor, DEC_SCENE_ITEMS, *EventName);
 #endif
 
 		FTextureRHIParamRef RenderTargets[2] =
@@ -1114,7 +1114,7 @@ void FDeferredShadingSceneRenderer::ClearTranslucentVolumeLighting(FRHICommandLi
 {
 	if (GUseTranslucentLightingVolumes && GSupportsVolumeTextureRendering)
 	{
-		SCOPED_DRAW_EVENT(ClearTranslucentVolumeLighting, DEC_SCENE_ITEMS);
+		SCOPED_DRAW_EVENT(RHICmdList, ClearTranslucentVolumeLighting, DEC_SCENE_ITEMS);
 
 		// Clear all volume textures in the same draw with MRT, which is faster than individually
 
@@ -1186,7 +1186,7 @@ void FDeferredShadingSceneRenderer::InjectAmbientCubemapTranslucentVolumeLightin
 
 	if (GUseTranslucentLightingVolumes && View.FinalPostProcessSettings.ContributingCubemaps.Num() && !IsSimpleDynamicLightingEnabled())
 	{
-		SCOPED_DRAW_EVENT(InjectAmbientCubemapTranslucentVolumeLighting, DEC_SCENE_ITEMS);
+		SCOPED_DRAW_EVENT(RHICmdList, InjectAmbientCubemapTranslucentVolumeLighting, DEC_SCENE_ITEMS);
 
 		RHICmdList.SetRasterizerState(TStaticRasterizerState<FM_Solid, CM_None>::GetRHI());
 		RHICmdList.SetDepthStencilState(TStaticDepthStencilState<false, CF_Always>::GetRHI());
@@ -1232,7 +1232,7 @@ void FDeferredShadingSceneRenderer::ClearTranslucentVolumePerObjectShadowing(FRH
 {
 	if (GUseTranslucentLightingVolumes && GSupportsVolumeTextureRendering)
 	{
-		SCOPED_DRAW_EVENT(ClearTranslucentVolumePerLightShadowing, DEC_SCENE_ITEMS);
+		SCOPED_DRAW_EVENT(RHICmdList, ClearTranslucentVolumePerLightShadowing, DEC_SCENE_ITEMS);
 
 		static_assert(TVC_MAX == 2, "Only expecting two translucency lighting cascades.");
 		FTextureRHIParamRef RenderTargets[2];
@@ -1287,7 +1287,7 @@ void FDeferredShadingSceneRenderer::AccumulateTranslucentVolumeObjectShadowing(F
 
 	if (GUseTranslucentLightingVolumes && GSupportsVolumeTextureRendering)
 	{
-		SCOPED_DRAW_EVENT(AccumulateTranslucentVolumeShadowing, DEC_SCENE_ITEMS);
+		SCOPED_DRAW_EVENT(RHICmdList, AccumulateTranslucentVolumeShadowing, DEC_SCENE_ITEMS);
 
 		RHICmdList.SetRasterizerState(TStaticRasterizerState<FM_Solid, CM_None>::GetRHI());
 		RHICmdList.SetDepthStencilState(TStaticDepthStencilState<false, CF_Always>::GetRHI());
@@ -1772,7 +1772,7 @@ void FDeferredShadingSceneRenderer::FilterTranslucentVolumeLighting(FRHICommandL
 	{
 		if (GUseTranslucencyVolumeBlur)
 		{
-			SCOPED_DRAW_EVENT(FilterTranslucentVolume, DEC_SCENE_ITEMS);
+			SCOPED_DRAW_EVENT(RHICmdList, FilterTranslucentVolume, DEC_SCENE_ITEMS);
 
 			RHICmdList.SetRasterizerState(TStaticRasterizerState<FM_Solid, CM_None>::GetRHI());
 			RHICmdList.SetDepthStencilState(TStaticDepthStencilState<false, CF_Always>::GetRHI());
diff --git a/Engine/Source/Runtime/Renderer/Private/TranslucentRendering.cpp b/Engine/Source/Runtime/Renderer/Private/TranslucentRendering.cpp
--- a/Engine/Source/Runtime/Renderer/Private/TranslucentRendering.cpp
+++ b/Engine/Source/Runtime/Renderer/Private/TranslucentRendering.cpp
@@ -171,7 +171,7 @@ FGlobalBoundShaderState CopySceneColorBoundShaderState;
 void FTranslucencyDrawingPolicyFactory::CopySceneColor(FRHICommandList& RHICmdList, const FViewInfo& View, const FPrimitiveSceneProxy* PrimitiveSceneProxy)
 {
 
-	SCOPED_DRAW_EVENTF(EventCopy, DEC_SCENE_ITEMS, TEXT("CopySceneColor for %s %s"), *PrimitiveSceneProxy->GetOwnerName().ToString(), *PrimitiveSceneProxy->GetResourceName().ToString());
+	SCOPED_DRAW_EVENTF(RHICmdList, EventCopy, DEC_SCENE_ITEMS, TEXT("CopySceneColor for %s %s"), *PrimitiveSceneProxy->GetOwnerName().ToString(), *PrimitiveSceneProxy->GetResourceName().ToString());
 	RHICmdList.SetRasterizerState(TStaticRasterizerState<FM_Solid, CM_None>::GetRHI());
 	RHICmdList.SetDepthStencilState(TStaticDepthStencilState<false, CF_Always>::GetRHI());
 	RHICmdList.SetBlendState(TStaticBlendState<>::GetRHI());
@@ -845,7 +845,7 @@ void FDeferredShadingSceneRenderer::RenderTranslucencyParallel(FRHICommandListIm
 
 	for (int32 ViewIndex = 0; ViewIndex < Views.Num(); ViewIndex++)
 	{
-		SCOPED_CONDITIONAL_DRAW_EVENTF(EventView, Views.Num() > 1, DEC_SCENE_ITEMS, TEXT("View%d"), ViewIndex);
+		SCOPED_CONDITIONAL_DRAW_EVENTF(RHICmdList, EventView, Views.Num() > 1, DEC_SCENE_ITEMS, TEXT("View%d"), ViewIndex);
 
 		const FViewInfo& View = Views[ViewIndex];
 
@@ -956,7 +956,7 @@ void FDeferredShadingSceneRenderer::RenderTranslucency(FRHICommandListImmediate&
 	if (ShouldRenderTranslucency())
 	{
 
-		SCOPED_DRAW_EVENT(Translucency, DEC_SCENE_ITEMS);
+		SCOPED_DRAW_EVENT(RHICmdList, Translucency, DEC_SCENE_ITEMS);
 
 		if (GRHICommandList.UseParallelAlgorithms() && CVarParallelTranslucency.GetValueOnRenderThread())
 		{
@@ -966,7 +966,7 @@ void FDeferredShadingSceneRenderer::RenderTranslucency(FRHICommandListImmediate&
 
 		for (int32 ViewIndex = 0; ViewIndex < Views.Num(); ViewIndex++)
 		{
-			SCOPED_CONDITIONAL_DRAW_EVENTF(EventView, Views.Num() > 1, DEC_SCENE_ITEMS, TEXT("View%d"), ViewIndex);
+			SCOPED_CONDITIONAL_DRAW_EVENTF(RHICmdList, EventView, Views.Num() > 1, DEC_SCENE_ITEMS, TEXT("View%d"), ViewIndex);
 
 			const FViewInfo& View = Views[ViewIndex];
 
diff --git a/Engine/Source/Runtime/Renderer/Private/VelocityRendering.cpp b/Engine/Source/Runtime/Renderer/Private/VelocityRendering.cpp
--- a/Engine/Source/Runtime/Renderer/Private/VelocityRendering.cpp
+++ b/Engine/Source/Runtime/Renderer/Private/VelocityRendering.cpp
@@ -621,7 +621,7 @@ void FDeferredShadingSceneRenderer::RenderVelocities(FRHICommandListImmediate& R
 		return;
 	}
 
-	SCOPED_DRAW_EVENT(RenderVelocities, DEC_SCENE_ITEMS);
+	SCOPED_DRAW_EVENT(RHICmdList, RenderVelocities, DEC_SCENE_ITEMS);
 
 	const FIntPoint BufferSize = GSceneRenderTargets.GetBufferSizeXY();
 	const FIntPoint VelocityBufferSize = BufferSize;		// full resolution so we can reuse the existing full res z buffer
diff --git a/Engine/Source/Runtime/SlateRHIRenderer/Private/SlateRHIRenderer.cpp b/Engine/Source/Runtime/SlateRHIRenderer/Private/SlateRHIRenderer.cpp
--- a/Engine/Source/Runtime/SlateRHIRenderer/Private/SlateRHIRenderer.cpp
+++ b/Engine/Source/Runtime/SlateRHIRenderer/Private/SlateRHIRenderer.cpp
@@ -386,7 +386,7 @@ void FSlateRHIRenderer::OnWindowDestroyed( const TSharedRef<SWindow>& InWindow )
 /** Draws windows from a FSlateDrawBuffer on the render thread */
 void FSlateRHIRenderer::DrawWindow_RenderThread(FRHICommandListImmediate& RHICmdList, const FViewportInfo& ViewportInfo, const FSlateWindowElementList& WindowElementList, bool bLockToVsync)
 {
-	SCOPED_DRAW_EVENT(SlateUI, DEC_SCENE_ITEMS);
+	SCOPED_DRAW_EVENT(RHICmdList, SlateUI, DEC_SCENE_ITEMS);
 
 	// Should only be called by the rendering thread
 	check(IsInRenderingThread());
diff --git a/Engine/Source/Runtime/Windows/D3D11RHI/Private/D3D11Commands.cpp b/Engine/Source/Runtime/Windows/D3D11RHI/Private/D3D11Commands.cpp
--- a/Engine/Source/Runtime/Windows/D3D11RHI/Private/D3D11Commands.cpp
+++ b/Engine/Source/Runtime/Windows/D3D11RHI/Private/D3D11Commands.cpp
@@ -1705,9 +1705,6 @@ void FD3D11DynamicRHI::RHIClearMRT(bool bClearColor,int32 NumClearColors,const F
 
 	if (UseDrawClear)
 	{
-		// So we can see when we are taking the slow path
-		SCOPED_DRAW_EVENT(DrawCallClear, DEC_SCENE_ITEMS);
-
 		// we don't support draw call clears before the RHI is initialized, reorder the code or make sure it's not a draw call clear
 		check(GIsRHIInitialized);
 
diff --git a/Engine/Source/Runtime/Windows/D3D11RHI/Private/D3D11Device.cpp b/Engine/Source/Runtime/Windows/D3D11RHI/Private/D3D11Device.cpp
--- a/Engine/Source/Runtime/Windows/D3D11RHI/Private/D3D11Device.cpp
+++ b/Engine/Source/Runtime/Windows/D3D11RHI/Private/D3D11Device.cpp
@@ -204,6 +204,16 @@ void FD3D11DynamicRHI::Shutdown()
 	ZeroBufferSize = 0;
 }
 
+void FD3D11DynamicRHI::RHIPushEvent(const TCHAR* Name)
+{ 
+	GPUProfilingData.PushEvent(Name); 
+}
+void FD3D11DynamicRHI::RHIPopEvent()
+{ 
+	GPUProfilingData.PopEvent(); 
+}
+
+
 /**
  * Returns a supported screen resolution that most closely matches the input.
  * @param Width - Input: Desired resolution width in pixels. Output: A width that the platform supports.
diff --git a/Engine/Source/Runtime/Windows/D3D11RHI/Private/D3D11RHIPrivate.h b/Engine/Source/Runtime/Windows/D3D11RHI/Private/D3D11RHIPrivate.h
--- a/Engine/Source/Runtime/Windows/D3D11RHI/Private/D3D11RHIPrivate.h
+++ b/Engine/Source/Runtime/Windows/D3D11RHI/Private/D3D11RHIPrivate.h
@@ -159,12 +159,20 @@ public:
 		Timing(InRHI, 1)
 	{
 		// Initialize Buffered timestamp queries 
-		Timing.InitResource();
+#if PLATFORM_SUPPORTS_RHI_THREAD
+		Timing.InitDynamicRHI();
+#else
+		Timing.InitResource(); // can't do this from the RHI thread
+#endif
 	}
 
 	virtual ~FD3D11EventNode()
 	{
-		Timing.ReleaseResource();
+#if PLATFORM_SUPPORTS_RHI_THREAD
+		Timing.ReleaseDynamicRHI();
+#else
+		Timing.ReleaseResource();  // can't do this from the RHI thread
+#endif
 	}
 
 	/** 
@@ -197,16 +205,26 @@ public:
 		RootEventTiming(InRHI, 1),
 		DisjointQuery(InRHI)
 	{
-
-	  RootEventTiming.InitResource();
-	  DisjointQuery.InitResource();
+#if PLATFORM_SUPPORTS_RHI_THREAD
+		RootEventTiming.InitDynamicRHI();
+		DisjointQuery.InitDynamicRHI();
+#else
+		// can't do these on the RHI thread
+		RootEventTiming.InitResource();
+		DisjointQuery.InitResource();
+#endif
 	}
 
 	~FD3D11EventNodeFrame()
 	{
-
+#if PLATFORM_SUPPORTS_RHI_THREAD
+		RootEventTiming.ReleaseDynamicRHI();
+		DisjointQuery.ReleaseDynamicRHI();
+#else
+		// can't do these on the RHI thread
 		RootEventTiming.ReleaseResource();
 		DisjointQuery.ReleaseResource();
+#endif
 	}
 
 	/** Start this frame of per tracking */
@@ -235,12 +253,22 @@ public:
 	FD3D11GPUProfile(FD3D11DynamicRHI* D3D11RHI)
 	: DisjointQuery(D3D11RHI)
 	{
+#if PLATFORM_SUPPORTS_RHI_THREAD
+		DisjointQuery.InitDynamicRHI();
+#else
+		// can't do these on the RHI thread
 		DisjointQuery.InitResource();
+#endif
 	}
 
 	~FD3D11GPUProfile()
 	{
+#if PLATFORM_SUPPORTS_RHI_THREAD
+		DisjointQuery.ReleaseDynamicRHI();
+#else
+		// can't do these on the RHI thread
 		DisjointQuery.ReleaseResource();
+#endif
 	}
 
 	FD3D11DisjointTimeStampQuery DisjointQuery;
@@ -320,8 +348,6 @@ public:
 	// FDynamicRHI interface.
 	virtual void Init() override;
 	virtual void Shutdown() override;
-	virtual void PushEvent(const TCHAR* Name) override { GPUProfilingData.PushEvent(Name); }
-	virtual void PopEvent() override { GPUProfilingData.PopEvent(); }
 
 	/**
 	 * Reads a D3D query's data into the provided buffer.
