UE4 - Cascaded shadow maps on mobile/forward
* Still disabled as Metal requires a horrible workaround/hack
* Remove GSupportsGSRenderTargetLayerSwitchingToMips workaround

diff --git a/Engine/Shaders/BasePassForForwardShadingPixelShader.usf b/Engine/Shaders/BasePassForForwardShadingPixelShader.usf
--- a/Engine/Shaders/BasePassForForwardShadingPixelShader.usf
+++ b/Engine/Shaders/BasePassForForwardShadingPixelShader.usf
@@ -12,6 +12,7 @@
 #include "LightmapCommon.usf"  
 #include "BRDF.usf"
 #include "SHCommon.usf"
+#include "ShadowFilteringCommon.usf"
 
 /** Prenormalized capture of the scene that's closest to the object being rendered. */
 #if !MATERIAL_FULLY_ROUGH
@@ -61,9 +62,7 @@ half DirectionalLightShadowing;
 void Main( 
 	FVertexFactoryInterpolantsVSToPS Interpolants,
 	FForwardShadingBasePassInterpolantsVSToPS BasePassInterpolants,
-	#if ES2_PROFILE && (USE_HDR_MOSAIC || !COMPILER_GLSL_ES2)
 		in half4 SvPosition	: SV_Position,
-	#endif
 	OPTIONAL_IsFrontFace,
 	out half4 OutColor	: SV_Target0
 	)
@@ -81,15 +80,17 @@ void Main(
 	);
 	  
 	//Clip if the blend mode requires it.
-	GetMaterialCoverageAndClipping(MaterialParameters);
+	GetMaterialCoverageAndClipping(MaterialParameters);	
 
-#if COMPILER_HLSL	// Not sure why this isn't needed everywhere
+	//@mw todo:
+	//Go over with Brian to determine the proper defines to wrap it and test cases to fix it.
+	/*
 	float2 PixelPos = SvPosition - View.ViewRectMin.xy;	
 	MaterialParameters.ScreenPosition.xy = ( PixelPos / View.ViewSizeAndSceneTexelSize.xy - 0.5 ) * float2(2,-2);
 	MaterialParameters.ScreenPosition.zw = SvPosition.zw;
 	
 	MaterialParameters.ScreenPosition.xyz *= MaterialParameters.ScreenPosition.w;
-#endif
+	*/
 
 	// Store the results in local variables and reuse instead of calling the functions multiple times.
 #if DIFFUSE_SPEC_INPUTS
@@ -177,17 +178,66 @@ void Main(
 		#if SIMPLE_DYNAMIC_LIGHTING
 			Shadow = DirectionalLightShadowing;
 		#endif
-		
+			
+#if MOVABLE_DIRECTIONAL_LIGHT_CSM
+		// Cascaded Shadow Map
+		{
+			FPCFSamplerSettings Settings;			
+			Settings.ShadowDepthTexture = View.DirectionalLightShadowTexture;
+			Settings.ShadowDepthTextureSampler = View.DirectionalLightShadowSampler;
+			Settings.TransitionScale = View.DirectionalLightShadowTransition;
+			Settings.ShadowBufferSize = View.DirectionalLightShadowSize;
+			Settings.bSubsurface = false;
+			Settings.DensityMulConstant = 0;
+			Settings.ProjectionDepthBiasParameters = 0;
+						
+			float4 ShadowPosition;	
+#if MAX_FORWARD_SHADOWCASCADES == 2			
+			{				
+				if (MaterialParameters.ScreenPosition.w >= View.DirectionalLightShadowDistances.x)
+				{
+					ShadowPosition = mul(float4(MaterialParameters.ScreenPosition.xyw, 1), View.DirectionalLightScreenToShadow[1]);
+				}
+				else
+				{
+					ShadowPosition = mul(float4(MaterialParameters.ScreenPosition.xyw, 1), View.DirectionalLightScreenToShadow[0]);
+				}
+			}
+#else			
+			{
+				ShadowPosition = mul(float4(MaterialParameters.ScreenPosition.xyw, 1), View.DirectionalLightScreenToShadow[0]);				
+			}
+#endif
+
+			// Clamp pixel depth in light space for shadowing opaque, because areas of the shadow depth buffer that weren't rendered to will have been cleared to 1
+			// We want to force the shadow comparison to result in 'unshadowed' in that case, regardless of whether the pixel being shaded is in front or behind that plane
+			float LightSpacePixelDepthForOpaque = min(ShadowPosition.z, 0.99999f);
+			Settings.SceneDepth = LightSpacePixelDepthForOpaque;
+
+#if 0
+			half ShadowMap = Manual1x1PCF(ShadowPosition.xy, Settings);
+#else
+			half ShadowMap = Manual2x2PCF(ShadowPosition.xy, Settings);
+#endif
+			Shadow = min(Shadow, ShadowMap);
+		}
+#endif
+
 		half NoL = max(0, dot(MaterialParameters.WorldNormal, View.DirectionalLightDirection));
 		half RoL = max(0, dot(MaterialParameters.ReflectionVector, View.DirectionalLightDirection));
 
 		Color += (Shadow * NoL) * View.DirectionalLightColor.rgb * (DiffuseColor + SpecularColor * PhongApprox(Roughness, RoL));
-
 		#if !MATERIAL_FULLY_ROUGH
 			half3 SpecularIBL = GetImageBasedReflectionLighting(MaterialParameters, Roughness);
 			// Environment map has been prenormalized, scale by lightmap luminance
 			Color += SpecularIBL * IndirectIrradiance * SpecularColor;
 		#endif
+
+		#if ENABLE_SKY_LIGHT
+			//@mw todo
+			// TODO: Also need to do specular.
+			Color += GetSkySHDiffuseSimple(MaterialParameters.WorldNormal) * View.SkyLightColor.rgb * DiffuseColor;    
+		#endif
 	#endif
 			 
 	half Opacity = GetMaterialOpacity(MaterialParameters);
@@ -232,8 +282,8 @@ void Main(
 		OutColor.rgb = sqrt( OutColor.rgb );
 	#endif
 
-	#if ES2_PROFILE && (!OUTPUT_GAMMA_SPACE)
-		#if COMPILER_GLSL_ES2
+	#if (ES2_PROFILE || ES3_1_PROFILE) && (!OUTPUT_GAMMA_SPACE)
+		#if COMPILER_GLSL_ES2 || METAL_PROFILE
 			#if USE_HDR_MOSAIC
 				OutColor.rgb = HdrMosaic(OutColor.rgb, SvPosition.xy);
 			#endif
diff --git a/Engine/Shaders/ReflectionEnvironmentShaders.usf b/Engine/Shaders/ReflectionEnvironmentShaders.usf
--- a/Engine/Shaders/ReflectionEnvironmentShaders.usf
+++ b/Engine/Shaders/ReflectionEnvironmentShaders.usf
@@ -29,34 +29,8 @@ void CopyToCubeFaceVS(
 	Out.ScreenVector = mul(float4(Out.Position.xy, 1, 0), View.ScreenToTranslatedWorld).xyz;
 }
 
-struct FRouteToCubeFaceGSOut
-{
-	FCopyToCubeFaceVSOutput Vertex;
-	uint LayerIndex : SV_RenderTargetArrayIndex;
-};
-
 int CubeFace;
 
-[maxvertexcount(3)]
-void RouteToCubeFaceGS(triangle FCopyToCubeFaceVSOutput Input[3], inout TriangleStream<FRouteToCubeFaceGSOut> OutStream)
-{
-	FRouteToCubeFaceGSOut Vertex0;
-	Vertex0.Vertex = Input[0];
-	Vertex0.LayerIndex = CubeFace;
-
-	FRouteToCubeFaceGSOut Vertex1;
-	Vertex1.Vertex = Input[1];
-	Vertex1.LayerIndex = CubeFace;
-
-	FRouteToCubeFaceGSOut Vertex2;
-	Vertex2.Vertex = Input[2];
-	Vertex2.LayerIndex = CubeFace;
-
-	OutStream.Append(Vertex0);
-	OutStream.Append(Vertex1);
-	OutStream.Append(Vertex2);
-}
-
 Texture2D InTexture;
 SamplerState InTextureSampler;
 
@@ -68,16 +42,16 @@ SamplerState InTextureSampler;
 float3 SkyLightCaptureParameters;
 
 void CopySceneColorToCubeFaceColorPS(
-	FRouteToCubeFaceGSOut Input,
+	FCopyToCubeFaceVSOutput Input,
 	out float4 OutColor : SV_Target0
 	)
 { 
-	float SceneDepth = CalcSceneDepth(Input.Vertex.UV);
-	float3 SceneColor = Texture2DSample(InTexture, InTextureSampler, Input.Vertex.UV).rgb;
+	float SceneDepth = CalcSceneDepth(Input.UV);
+	float3 SceneColor = Texture2DSample(InTexture, InTextureSampler, Input.UV).rgb;
 
 	// Convert INF's to valid values
 	SceneColor = min(SceneColor, 65503);
-	float3 WorldPosition = Input.Vertex.ScreenVector * SceneDepth + View.ViewOrigin.xyz;
+	float3 WorldPosition = Input.ScreenVector * SceneDepth + View.ViewOrigin.xyz;
 
 	float Alpha = 1;
 
@@ -85,7 +59,7 @@ void CopySceneColorToCubeFaceColorPS(
 	{
 		// Assuming we're on a planet and no sky lighting is coming from below the horizon
 		// This is important to avoid leaking from below since we are integrating incoming lighting and shadowing separately
-		float AboveHorizonMask = Input.Vertex.ScreenVector.z > 0 || SkyLightCaptureParameters.z < 1;
+		float AboveHorizonMask = Input.ScreenVector.z > 0 || SkyLightCaptureParameters.z < 1;
 		SceneColor *= AboveHorizonMask;
 	}
 	else if (SkyLightCaptureParameters.x == 1)
@@ -135,41 +109,15 @@ float3 GetCubemapVector(float2 ScaledUVs)
 	return CubeCoordinates;
 }
 
-struct FDownsampleGSOut
-{
-	FScreenVertexOutput Vertex;
-	uint LayerIndex : SV_RenderTargetArrayIndex;
-};
-
-[maxvertexcount(3)]
-void DownsampleGS(triangle FScreenVertexOutput Input[3], inout TriangleStream<FDownsampleGSOut> OutStream)
-{
-	FDownsampleGSOut Vertex0;
-	Vertex0.Vertex = Input[0];
-	Vertex0.LayerIndex = CubeFace;
-
-	FDownsampleGSOut Vertex1;
-	Vertex1.Vertex = Input[1];
-	Vertex1.LayerIndex = CubeFace;
-
-	FDownsampleGSOut Vertex2;
-	Vertex2.Vertex = Input[2];
-	Vertex2.LayerIndex = CubeFace;
-
-	OutStream.Append(Vertex0);
-	OutStream.Append(Vertex1);
-	OutStream.Append(Vertex2);
-}
-
 TextureCube SourceTexture;
 SamplerState SourceTextureSampler;
 
 void CopyCubemapToCubeFaceColorPS(
-	FDownsampleGSOut Input,
+	FScreenVertexOutput Input,
 	out float4 OutColor : SV_Target0
 	)
 { 
-	float2 ScaledUVs = Input.Vertex.UV * 2 - 1;
+	float2 ScaledUVs = Input.UV * 2 - 1;
 	float3 CubeCoordinates = GetCubemapVector(ScaledUVs);
 	OutColor = TextureCubeSampleLevel(SourceTexture, SourceTextureSampler, CubeCoordinates, 0);
 
@@ -190,17 +138,17 @@ float4 SampleCubemap(float3 Coordinates)
 }
 
 void DownsamplePS(
-	FDownsampleGSOut Input,
+	FScreenVertexOutput Input,
 	out float4 OutColor : SV_Target0
 	)
 { 
-	float2 ScaledUVs = Input.Vertex.UV * 2 - 1;
+	float2 ScaledUVs = Input.UV * 2 - 1;
 	float3 CubeCoordinates = GetCubemapVector(ScaledUVs);
 
 	OutColor = SampleCubemap(CubeCoordinates);
 }
 
-uint NumCaptureArrayMips;
+int NumCaptureArrayMips;
 
 /** Cube map array of reflection captures. */
 TextureCube ReflectionEnvironmentColorTexture;
@@ -232,14 +180,20 @@ Texture2D AverageBrightnessTexture;
 SamplerState AverageBrightnessSampler;
 
 void FilterPS(
-	FDownsampleGSOut Input,
+	FScreenVertexOutput Input,
 	out float4 OutColor : SV_Target0
 	)
 { 
-	float2 ScaledUVs = Input.Vertex.UV * 2 - 1;
+	float2 ScaledUVs = Input.UV * 2 - 1;
 	float3 CubeCoordinates = GetCubemapVector(ScaledUVs);
 
+#if ES2_PROFILE
+//@todo - mwassmer.  Write filtering that works for ES2 so specular will work.
+#define FILTER_CUBEMAP 0
+#else
 #define FILTER_CUBEMAP 1
+#endif
+
 #if FILTER_CUBEMAP
 
 	float3 N = normalize(CubeCoordinates);
@@ -249,7 +203,7 @@ void FilterPS(
 	float Weight = 0;
 		
 	LOOP
-	for( uint i = 0; i < NUM_FILTER_SAMPLES; i++ )
+	for( int i = 0; i < NUM_FILTER_SAMPLES; i++ )
 	{
 		float2 E = Hammersley( i, NUM_FILTER_SAMPLES, 0 );
 		float3 H = TangentToWorld( ImportanceSampleGGX( E, Roughness ).xyz, N );
@@ -285,11 +239,11 @@ float CoefficientMask2;
 int NumSamples;
 
 void DiffuseIrradianceCopyPS(
-	FDownsampleGSOut Input,
+	FScreenVertexOutput Input,
 	out float4 OutColor : SV_Target0
 	)
 { 
-	float2 ScaledUVs = Input.Vertex.UV * 2 - 1;
+	float2 ScaledUVs = Input.UV * 2 - 1;
 	float3 CubeCoordinates = normalize(GetCubemapVector(ScaledUVs));
 
 	float SquaredUVs = 1 + dot(ScaledUVs, ScaledUVs);
@@ -307,32 +261,32 @@ float4 Sample01;
 float4 Sample23;
 
 void DiffuseIrradianceAccumulatePS(
-	FDownsampleGSOut Input,
+	FScreenVertexOutput Input,
 	out float4 OutColor : SV_Target0
 	)
 { 
 	float4 AccumulatedValue = 0;
 
 	{
-		float2 ScaledUVs = saturate(Input.Vertex.UV + Sample01.xy) * 2 - 1;
+		float2 ScaledUVs = saturate(Input.UV + Sample01.xy) * 2 - 1;
 		float3 CubeCoordinates = GetCubemapVector(ScaledUVs);
 		AccumulatedValue += SampleCubemap(CubeCoordinates);
 	}
 	
 	{
-		float2 ScaledUVs = saturate(Input.Vertex.UV + Sample01.zw) * 2 - 1;
+		float2 ScaledUVs = saturate(Input.UV + Sample01.zw) * 2 - 1;
 		float3 CubeCoordinates = GetCubemapVector(ScaledUVs);
 		AccumulatedValue += SampleCubemap(CubeCoordinates);
 	}
 
 	{
-		float2 ScaledUVs = saturate(Input.Vertex.UV + Sample23.xy) * 2 - 1;
+		float2 ScaledUVs = saturate(Input.UV + Sample23.xy) * 2 - 1;
 		float3 CubeCoordinates = GetCubemapVector(ScaledUVs);
 		AccumulatedValue += SampleCubemap(CubeCoordinates);
 	}
 
 	{
-		float2 ScaledUVs = saturate(Input.Vertex.UV + Sample23.zw) * 2 - 1;
+		float2 ScaledUVs = saturate(Input.UV + Sample23.zw) * 2 - 1;
 		float3 CubeCoordinates = GetCubemapVector(ScaledUVs);
 		AccumulatedValue += SampleCubemap(CubeCoordinates);
 	}
@@ -351,7 +305,7 @@ void AccumulateCubeFacesPS(
 	AccumulatedValue += TextureCubeSampleLevel(SourceTexture, SourceTextureSampler, float3(0, -1, 0), SourceMipIndex);
 	AccumulatedValue += TextureCubeSampleLevel(SourceTexture, SourceTextureSampler, float3(0, 0, 1), SourceMipIndex);
 	AccumulatedValue += TextureCubeSampleLevel(SourceTexture, SourceTextureSampler, float3(0, 0, -1), SourceMipIndex);
-	OutColor = float4(4 * PI * AccumulatedValue.rgb / ( max(AccumulatedValue.a, .00001f)), 0);
+	OutColor = float4(4 * PI * AccumulatedValue.rgb / ( max(AccumulatedValue.a, .00001f)), 0);	
 }
 
 // Used during reflection captures to get bounce light from specular surfaces
diff --git a/Engine/Shaders/ShadowProjectionPixelShader.usf b/Engine/Shaders/ShadowProjectionPixelShader.usf
--- a/Engine/Shaders/ShadowProjectionPixelShader.usf
+++ b/Engine/Shaders/ShadowProjectionPixelShader.usf
@@ -23,6 +23,7 @@
 #endif
 
 #include "ShadowProjectionCommon.usf"
+#include "ShadowFilteringCommon.usf"
 #include "DeferredShadingCommon.usf"
 
 float ShadowFadeFraction;
@@ -30,379 +31,12 @@ float ShadowSharpen;
 float4x4 ScreenToShadowMatrix;
 // .x:DepthBias, y: MaxSubjectZ - MinSubjectZ
 float2 ProjectionDepthBiasParameters;
-Texture2D ShadowDepthTextureObject;
-SamplerState ShadowDepthSampler;
 
 #if USE_FADE_PLANE
 	float FadePlaneOffset;
 	float InvFadePlaneLength;
 #endif
 
-struct FPCFSamplerSettings
-{
-	float SceneDepth;
-
-	// set by the caller, constant for the code so only one code branch should be compiled
-	bool bSubsurface;
-
-	// only used if bSubsurface is true
-	float DensityMulConstant;
-};
-
-// linear PCF, input 4x4
-// @param Values0 in column 0 from top to bottom: x,y,z,w
-// @param Values1 in column 1 from top to bottom: x,y,z,w
-// @param Values2 in column 2 from top to bottom: x,y,z,w
-// @param Values3 in column 3 from top to bottom: x,y,z,w
-// can be optimized
-float PCF3x3(float2 Fraction, float4 Values0, float4 Values1, float4 Values2, float4 Values3)
-{
-	float2 VerticalLerp00 = lerp(float2(Values0.x, Values1.x), float2(Values0.y, Values1.y), Fraction.xx);
-	float PCFResult00 = lerp(VerticalLerp00.x, VerticalLerp00.y, Fraction.y);
-	float2 VerticalLerp10 = lerp(float2(Values0.y, Values1.y), float2(Values0.z, Values1.z), Fraction.xx);
-	float PCFResult10 = lerp(VerticalLerp10.x, VerticalLerp10.y, Fraction.y);
-	float2 VerticalLerp20 = lerp(float2(Values0.z, Values1.z), float2(Values0.w, Values1.w), Fraction.xx);
-	float PCFResult20 = lerp(VerticalLerp20.x, VerticalLerp20.y, Fraction.y);
-
-	float2 VerticalLerp01 = lerp(float2(Values1.x, Values2.x), float2(Values1.y, Values2.y), Fraction.xx);
-	float PCFResult01 = lerp(VerticalLerp01.x, VerticalLerp01.y, Fraction.y);
-	float2 VerticalLerp11 = lerp(float2(Values1.y, Values2.y), float2(Values1.z, Values2.z), Fraction.xx);
-	float PCFResult11 = lerp(VerticalLerp11.x, VerticalLerp11.y, Fraction.y);
-	float2 VerticalLerp21 = lerp(float2(Values1.z, Values2.z), float2(Values1.w, Values2.w), Fraction.xx);
-	float PCFResult21 = lerp(VerticalLerp21.x, VerticalLerp21.y, Fraction.y);
-
-	float2 VerticalLerp02 = lerp(float2(Values2.x, Values3.x), float2(Values2.y, Values3.y), Fraction.xx);
-	float PCFResult02 = lerp(VerticalLerp02.x, VerticalLerp02.y, Fraction.y);
-	float2 VerticalLerp12 = lerp(float2(Values2.y, Values3.y), float2(Values2.z, Values3.z), Fraction.xx);
-	float PCFResult12 = lerp(VerticalLerp12.x, VerticalLerp12.y, Fraction.y);
-	float2 VerticalLerp22 = lerp(float2(Values2.z, Values3.z), float2(Values2.w, Values3.w), Fraction.xx);
-	float PCFResult22 = lerp(VerticalLerp22.x, VerticalLerp22.y, Fraction.y);
-
-	return saturate((PCFResult00 + PCFResult10 + PCFResult20 + PCFResult01 + PCFResult11 + PCFResult21 + PCFResult02 + PCFResult12 + PCFResult22) * .11111f);
-}
-
-
-// linear PCF, input 4x4
-// using Gather: xyzw in counter clockwise order starting with the sample to the lower left of the queried location
-// @param Values0 left top
-// @param Values1 right top
-// @param Values2 left bottom
-// @param Values3 right bottom
-float PCF3x3gather(float2 Fraction, float4 Values0, float4 Values1, float4 Values2, float4 Values3)
-{
-	// mostly vectorized
-	float4 PCFResult0 = float4(lerp(Values0.wx, Values0.zy, Fraction.xx), lerp(Values2.wx, Values2.zy, Fraction.x));
-	float4 PCFResult1 = float4(lerp(Values0.zy, Values1.wx, Fraction.xx), lerp(Values2.zy, Values3.wx, Fraction.x));
-	float4 PCFResult2 = float4(lerp(Values1.wx, Values1.zy, Fraction.xx), lerp(Values3.wx, Values3.zy, Fraction.x));
-
-	float4 PCFResult012 = PCFResult0 + PCFResult1 + PCFResult2;
-
-	float3 PCFResults = lerp(PCFResult012.xyz, PCFResult012.yzw, Fraction.y);
-
-	return dot(PCFResults, 1.0f / 9.0f);
-}
-
-float PCF5x2(float2 Fraction, float4 Values00, float4 Values20, float4 Values40, float4 Values02, float4 Values22, float4 Values42)
-{
-	float4 Results0;
-	float2 Results1;
-	float2 HorizontalLerp00 = lerp(float2(Values00.w, Values00.x), float2(Values00.z, Values00.y), Fraction.xx);
-	Results0.x = lerp(HorizontalLerp00.x, HorizontalLerp00.y, Fraction.y);
-
-	float2 HorizontalLerp10 = lerp(float2(Values00.z, Values00.y), float2(Values20.w, Values20.x), Fraction.xx);
-	Results0.y = lerp(HorizontalLerp10.x, HorizontalLerp10.y, Fraction.y);
-
-	float2 HorizontalLerp20 = lerp(float2(Values20.w, Values20.x), float2(Values20.z, Values20.y), Fraction.xx);
-	Results0.z = lerp(HorizontalLerp20.x, HorizontalLerp20.y, Fraction.y);
-
-	float2 HorizontalLerp30 = lerp(float2(Values20.z, Values20.y), float2(Values40.w, Values40.x), Fraction.xx);
-	Results0.w = lerp(HorizontalLerp30.x, HorizontalLerp30.y, Fraction.y);
-
-	float2 HorizontalLerp40 = lerp(float2(Values40.w, Values40.x), float2(Values40.z, Values40.y), Fraction.xx);
-	Results1.x = lerp(HorizontalLerp40.x, HorizontalLerp40.y, Fraction.y);
-
-	float4 Results2;
-	float2 HorizontalLerp01 = lerp(float2(Values00.x, Values02.w), float2(Values00.y, Values02.z), Fraction.xx);
-	Results2.x = lerp(HorizontalLerp01.x, HorizontalLerp01.y, Fraction.y);
-
-	float2 HorizontalLerp11 = lerp(float2(Values00.y, Values02.z), float2(Values20.x, Values22.w), Fraction.xx);
-	Results2.y = lerp(HorizontalLerp11.x, HorizontalLerp11.y, Fraction.y);
-
-	float2 HorizontalLerp21 = lerp(float2(Values20.x, Values22.w), float2(Values20.y, Values22.z), Fraction.xx);
-	Results2.z = lerp(HorizontalLerp21.x, HorizontalLerp21.y, Fraction.y);
-
-	float2 HorizontalLerp31 = lerp(float2(Values20.y, Values22.z), float2(Values40.x, Values42.w), Fraction.xx);
-	Results2.w = lerp(HorizontalLerp31.x, HorizontalLerp31.y, Fraction.y);
-
-	float2 HorizontalLerp41 = lerp(float2(Values40.x, Values42.w), float2(Values40.y, Values42.z), Fraction.xx);
-	Results1.y = lerp(HorizontalLerp41.x, HorizontalLerp41.y, Fraction.y);
-
-	return dot(Results0, .04f) + dot(Results1, .04f) + dot(Results2, .04f);
-}
-
-float PCF5x1(float2 Fraction, float4 Values00, float4 Values20, float4 Values40)
-{
-	float4 Results0;
-	float2 HorizontalLerp00 = lerp(float2(Values00.w, Values00.x), float2(Values00.z, Values00.y), Fraction.xx);
-	Results0.x = lerp(HorizontalLerp00.x, HorizontalLerp00.y, Fraction.y);
-
-	float2 HorizontalLerp10 = lerp(float2(Values00.z, Values00.y), float2(Values20.w, Values20.x), Fraction.xx);
-	Results0.y = lerp(HorizontalLerp10.x, HorizontalLerp10.y, Fraction.y);
-
-	float2 HorizontalLerp20 = lerp(float2(Values20.w, Values20.x), float2(Values20.z, Values20.y), Fraction.xx);
-	Results0.z = lerp(HorizontalLerp20.x, HorizontalLerp20.y, Fraction.y);
-
-	float2 HorizontalLerp30 = lerp(float2(Values20.z, Values20.y), float2(Values40.w, Values40.x), Fraction.xx);
-	Results0.w = lerp(HorizontalLerp30.x, HorizontalLerp30.y, Fraction.y);
-
-	float2 HorizontalLerp40 = lerp(float2(Values40.w, Values40.x), float2(Values40.z, Values40.y), Fraction.xx);
-	float Results1 = lerp(HorizontalLerp40.x, HorizontalLerp40.y, Fraction.y);
-
-	return dot(Results0, .04f) + Results1 * .04f;
-}
-
-// lowest quality ith PCF
-float PCF1x1(float2 Fraction, float4 Values00)
-{
-	float2 HorizontalLerp00 = lerp(Values00.wx, Values00.zy, Fraction.xx);
-
-	return lerp(HorizontalLerp00.x, HorizontalLerp00.y, Fraction.y);
-}
-
-float4 CalculateOcclusion(float4 ShadowmapDepth, FPCFSamplerSettings Settings)
-{
-	if(Settings.bSubsurface)
-	{
-		// Determine the distance that the light traveled through the subsurface object
-		// This assumes that anything between this subsurface pixel and the light was also a subsurface material,
-		// As a result, subsurface materials receive leaked light based on how transparent they are
-		float4 Thickness = max(Settings.SceneDepth - ShadowmapDepth, 0);
-		float4 Occlusion = saturate(exp(-Thickness * Settings.DensityMulConstant));
-		// Never shadow from depths that were never written to (max depth value)
-		return ShadowmapDepth > .99f ? 1 : Occlusion;
-	}
-	else
-	{
-		// The standard comparison is SceneDepth < ShadowmapDepth
-		// Using a soft transition based on depth difference
-		// Offsets shadows a bit but reduces self shadowing artifacts considerably
-		float TransitionScale = SoftTransitionScale.z;
-		return saturate((ShadowmapDepth - Settings.SceneDepth) * TransitionScale + 1);
-	}
-}
-
-float3 CalculateOcclusion(float3 ShadowmapDepth, FPCFSamplerSettings Settings)
-{
-	if(Settings.bSubsurface)
-	{
-		// Determine the distance that the light traveled through the subsurface object
-		// This assumes that anything between this subsurface pixel and the light was also a subsurface material,
-		// As a result, subsurface materials receive leaked light based on how transparent they are
-		float3 Thickness = max(Settings.SceneDepth - (ShadowmapDepth - ProjectionDepthBiasParameters.x), 0);
-		float3 Occlusion = saturate(exp(-Thickness * Settings.DensityMulConstant));
-		// Never shadow from depths that were never written to (max depth value)
-		return ShadowmapDepth > .99f ? 1 : Occlusion;
-	}
-	else
-	{
-		// The standard comparison is Settings.SceneDepth < ShadowmapDepth
-		// Using a soft transition based on depth difference
-		// Offsets shadows a bit but reduces self shadowing artifacts considerably
-		float TransitionScale = SoftTransitionScale.z;
-		return saturate((ShadowmapDepth - Settings.SceneDepth) * TransitionScale + 1);
-	}
-}
-
-void FetchRowOfFour(float2 Sample00TexelCenter, float VerticalOffset, out float4 Values0, FPCFSamplerSettings Settings)
-{
-	Values0.x = Texture2DSampleLevel(ShadowDepthTexture, ShadowDepthTextureSampler, (Sample00TexelCenter + float2(0, VerticalOffset)) * ShadowBufferSize.zw, 0).r;
-	Values0.y = Texture2DSampleLevel(ShadowDepthTexture, ShadowDepthTextureSampler, (Sample00TexelCenter + float2(1, VerticalOffset)) * ShadowBufferSize.zw, 0).r;
-	Values0.z = Texture2DSampleLevel(ShadowDepthTexture, ShadowDepthTextureSampler, (Sample00TexelCenter + float2(2, VerticalOffset)) * ShadowBufferSize.zw, 0).r;
-	Values0.w = Texture2DSampleLevel(ShadowDepthTexture, ShadowDepthTextureSampler, (Sample00TexelCenter + float2(3, VerticalOffset)) * ShadowBufferSize.zw, 0).r;
-	Values0 = CalculateOcclusion(Values0, Settings);
-}
-
-void FetchRowOfThreeAfterFour(float2 Sample00TexelCenter, float VerticalOffset, out float3 Values1, FPCFSamplerSettings Settings)
-{
-	Values1.x = Texture2DSampleLevel(ShadowDepthTexture, ShadowDepthTextureSampler, (Sample00TexelCenter + float2(4, VerticalOffset)) * ShadowBufferSize.zw, 0).r;
-	Values1.y = Texture2DSampleLevel(ShadowDepthTexture, ShadowDepthTextureSampler, (Sample00TexelCenter + float2(5, VerticalOffset)) * ShadowBufferSize.zw, 0).r;
-	Values1.z = Texture2DSampleLevel(ShadowDepthTexture, ShadowDepthTextureSampler, (Sample00TexelCenter + float2(6, VerticalOffset)) * ShadowBufferSize.zw, 0).r;
-	Values1 = CalculateOcclusion(Values1, Settings);
-}
-
-float ManualPCF(float2 ShadowPosition, FPCFSamplerSettings Settings)
-{
-
-#if SHADOW_QUALITY == 1
-	// very low quality but very good performance, useful to profile, 1 sample, not using gather4
-	return CalculateOcclusion(Texture2DSampleLevel(ShadowDepthTexture, ShadowDepthTextureSampler, ShadowPosition, 0).rrr, Settings).r;
-#endif
-
-#if SHADOW_QUALITY == 2
-	// low quality, 2x2 samples, using and not using gather4
-	{
-		float2 TexelPos = ShadowPosition * ShadowBufferSize.xy - 0.5f;	// bias to be consistent with texture filtering hardware
-		float2 Fraction = frac(TexelPos);
-		float2 TexelCenter = floor(TexelPos) + 0.5f;	// bias to get reliable texel center content
-
-		// using Gather: xyzw in counter clockwise order starting with the sample to the lower left of the queried location
-		float4 Samples;
-
-#if FEATURE_GATHER4
-		Samples = ShadowDepthTextureObject.Gather(ShadowDepthSampler, TexelCenter * ShadowBufferSize.zw);
-#else
-		Samples.x = Texture2DSampleLevel(ShadowDepthTexture, ShadowDepthTextureSampler, (TexelCenter.xy + float2(0, 1)) * ShadowBufferSize.zw, 0).r;
-		Samples.y = Texture2DSampleLevel(ShadowDepthTexture, ShadowDepthTextureSampler, (TexelCenter.xy + float2(1, 1)) * ShadowBufferSize.zw, 0).r;
-		Samples.z = Texture2DSampleLevel(ShadowDepthTexture, ShadowDepthTextureSampler, (TexelCenter.xy + float2(1, 0)) * ShadowBufferSize.zw, 0).r;
-		Samples.w = Texture2DSampleLevel(ShadowDepthTexture, ShadowDepthTextureSampler, (TexelCenter.xy + float2(0, 0)) * ShadowBufferSize.zw, 0).r;
-#endif
-
-		float4 Values00 = CalculateOcclusion(Samples, Settings);
-		return PCF1x1(Fraction, Values00);
-	}
-#endif
-
-#if SHADOW_QUALITY == 3
-	// medium quality, 4x4 samples, using and not using gather4
-	{
-		float2 TexelPos = ShadowPosition * ShadowBufferSize.xy - 0.5f;	// bias to be consistent with texture filtering hardware
-		float2 Fraction = frac(TexelPos);
-		float2 TexelCenter = floor(TexelPos) + 0.5f;	// bias to get reliable texel center content
-		{
-			float2 Sample00TexelCenter = TexelCenter - float2(1, 1);
-
-			float4 SampleValues0, SampleValues1, SampleValues2, SampleValues3;
-
-#if FEATURE_GATHER4
-			float2 SamplePos = TexelCenter * ShadowBufferSize.zw;	// bias to get reliable texel center content
-			SampleValues0 = CalculateOcclusion(ShadowDepthTextureObject.Gather(ShadowDepthSampler, SamplePos, int2(-1, -1)), Settings);
-			SampleValues1 = CalculateOcclusion(ShadowDepthTextureObject.Gather(ShadowDepthSampler, SamplePos, int2( 1, -1)), Settings);
-			SampleValues2 = CalculateOcclusion(ShadowDepthTextureObject.Gather(ShadowDepthSampler, SamplePos, int2(-1,  1)), Settings);
-			SampleValues3 = CalculateOcclusion(ShadowDepthTextureObject.Gather(ShadowDepthSampler, SamplePos, int2( 1,  1)), Settings);
-			return PCF3x3gather(Fraction, SampleValues0, SampleValues1, SampleValues2, SampleValues3);
-#else // FEATURE_GATHER4 
-			FetchRowOfFour(Sample00TexelCenter, 0, SampleValues0, Settings);
-			FetchRowOfFour(Sample00TexelCenter, 1, SampleValues1, Settings);
-			FetchRowOfFour(Sample00TexelCenter, 2, SampleValues2, Settings);
-			FetchRowOfFour(Sample00TexelCenter, 3, SampleValues3, Settings);
-			return PCF3x3(Fraction, SampleValues0, SampleValues1, SampleValues2, SampleValues3);
-#endif // FEATURE_GATHER4
-		}
-	}
-#endif
-
-#if FEATURE_GATHER4
-	// high quality, 6x6 samples, using gather4
-	{
-		float2 TexelPos = ShadowPosition * ShadowBufferSize.xy - 0.5f;	// bias to be consistent with texture filtering hardware
-		float2 Fraction = frac(TexelPos);
-		float2 TexelCenter = floor(TexelPos);
-		float2 SamplePos = (TexelCenter + 0.5f) * ShadowBufferSize.zw;	// bias to get reliable texel center content
-
-		float3 Results;
-
-		float4 Values00 = CalculateOcclusion(ShadowDepthTextureObject.Gather(ShadowDepthSampler, SamplePos, int2(-2, -2)), Settings);
-		float4 Values20 = CalculateOcclusion(ShadowDepthTextureObject.Gather(ShadowDepthSampler, SamplePos, int2(0, -2)), Settings);
-		float4 Values40 = CalculateOcclusion(ShadowDepthTextureObject.Gather(ShadowDepthSampler, SamplePos, int2(2, -2)), Settings);
-
-		float4 Values02 = CalculateOcclusion(ShadowDepthTextureObject.Gather(ShadowDepthSampler, SamplePos, int2(-2, 0)), Settings);
-		float4 Values22 = CalculateOcclusion(ShadowDepthTextureObject.Gather(ShadowDepthSampler, SamplePos, int2(0, 0)), Settings);
-		float4 Values42 = CalculateOcclusion(ShadowDepthTextureObject.Gather(ShadowDepthSampler, SamplePos, int2(2, 0)), Settings);
-
-		Results.x = PCF5x2(Fraction, Values00, Values20, Values40, Values02, Values22, Values42);
-
-		float4 Values04 = CalculateOcclusion(ShadowDepthTextureObject.Gather(ShadowDepthSampler, SamplePos, int2(-2, 2)), Settings);
-		float4 Values24 = CalculateOcclusion(ShadowDepthTextureObject.Gather(ShadowDepthSampler, SamplePos, int2(0, 2)), Settings);
-		float4 Values44 = CalculateOcclusion(ShadowDepthTextureObject.Gather(ShadowDepthSampler, SamplePos, int2(2, 2)), Settings);
-
-		Results.y = PCF5x2(Fraction, Values02, Values22, Values42, Values04, Values24, Values44);
-		Results.z = PCF5x1(Fraction, Values04, Values24, Values44);
-
-		return dot(Results, 1);
-	}
-
-#else // FEATURE_GATHER4
-
-	// high quality, 7x7 samples, not using gather4 (todo: ideally we make this 6x6 to get same results with gather code)
-	{
-		float2 Fraction = frac(ShadowPosition * ShadowBufferSize.xy);
-		float2 Sample00TexelCenter = floor(ShadowPosition * ShadowBufferSize.xy) - float2(3, 3);
-	
-		// Fetch 7x7 shadowmap point samples
-		// Do 6x6 PCF samples, sharing the point samples between neighboring PCF samples
-		float4 Results;
-	
-		float4 SampleValues03;
-		float4 SampleValues13;
-	
-		{
-			float4 SampleValues10;
-			float4 SampleValues11;
-			float4 SampleValues12;
-		
-			// Group work to minimize temporary registers and to split texture work with PCF ALU operations to hide texture latency
-			// Without this layout (all texture lookups at the beginning, PCF ALU's at the end) this shader was 4x slower on Nvidia cards
-			{
-				float4 SampleValues00;
-				FetchRowOfFour(Sample00TexelCenter, 0, SampleValues00, Settings);
-				SampleValues10.x = SampleValues00.w;
-	
-				float4 SampleValues01;
-				FetchRowOfFour(Sample00TexelCenter, 1, SampleValues01, Settings);
-				SampleValues11.x = SampleValues01.w;
-	
-				float4 SampleValues02;
-				FetchRowOfFour(Sample00TexelCenter, 2, SampleValues02, Settings);
-				SampleValues12.x = SampleValues02.w;
-	
-				FetchRowOfFour(Sample00TexelCenter, 3, SampleValues03, Settings);
-				SampleValues13.x = SampleValues03.w;
-				Results.x = PCF3x3(Fraction, SampleValues00, SampleValues01, SampleValues02, SampleValues03);
-			}
-		
-			{
-				FetchRowOfThreeAfterFour(Sample00TexelCenter, 0, SampleValues10.yzw, Settings);
-				FetchRowOfThreeAfterFour(Sample00TexelCenter, 1, SampleValues11.yzw, Settings);
-				FetchRowOfThreeAfterFour(Sample00TexelCenter, 2, SampleValues12.yzw, Settings);
-				FetchRowOfThreeAfterFour(Sample00TexelCenter, 3, SampleValues13.yzw, Settings);
-				Results.y = PCF3x3(Fraction, SampleValues10, SampleValues11, SampleValues12, SampleValues13);
-			}
-		}
-
-		{
-			float4 SampleValues14;
-			float4 SampleValues15;
-			float4 SampleValues16;
-		
-			{
-				float4 SampleValues04;
-				FetchRowOfFour(Sample00TexelCenter, 4, SampleValues04, Settings);
-				SampleValues14.x = SampleValues04.w;
-	
-				float4 SampleValues05;
-				FetchRowOfFour(Sample00TexelCenter, 5, SampleValues05, Settings);
-				SampleValues15.x = SampleValues05.w;
-	
-				float4 SampleValues06;
-				FetchRowOfFour(Sample00TexelCenter, 6, SampleValues06, Settings);
-				SampleValues16.x = SampleValues06.w;
-	
-				Results.z = PCF3x3(Fraction, SampleValues03, SampleValues04, SampleValues05, SampleValues06);
-			}
-		
-			{
-				FetchRowOfThreeAfterFour(Sample00TexelCenter, 4, SampleValues14.yzw, Settings);
-				FetchRowOfThreeAfterFour(Sample00TexelCenter, 5, SampleValues15.yzw, Settings);
-				FetchRowOfThreeAfterFour(Sample00TexelCenter, 6, SampleValues16.yzw, Settings);
-				Results.w = PCF3x3(Fraction, SampleValues13, SampleValues14, SampleValues15, SampleValues16);
-			}
-		}
-	
-		return dot(Results, .25f);
-	}
-#endif	// FEATURE_GATHER4
-}
-
 /**
 * Entry point for uniform manual PCF that supports lights using normal shadows.
 */
@@ -447,9 +81,14 @@ void Main(
 
 		FPCFSamplerSettings Settings;
 
+		Settings.ShadowDepthTexture = ShadowDepthTexture;
+		Settings.ShadowDepthTextureSampler = ShadowDepthTextureSampler;
+		Settings.ShadowBufferSize = ShadowBufferSize;
+		Settings.TransitionScale = SoftTransitionScale.z;
 		Settings.SceneDepth = LightSpacePixelDepthForOpaque;
 		Settings.bSubsurface = false;
 		Settings.DensityMulConstant = 0;
+		Settings.ProjectionDepthBiasParameters = 0;
 
 		Shadow = ManualPCF(ShadowPosition.xy, Settings);
 
@@ -482,9 +121,14 @@ void Main(
 				// default code path
 				FPCFSamplerSettings Settings;
 				
+				Settings.ShadowDepthTexture = ShadowDepthTexture;
+				Settings.ShadowDepthTextureSampler = ShadowDepthTextureSampler;
+				Settings.ShadowBufferSize = ShadowBufferSize;
+				Settings.TransitionScale = SoftTransitionScale.z;
 				Settings.SceneDepth = LightSpacePixelDepthForSSS + ProjectionDepthBiasParameters.x;
 				Settings.bSubsurface = true;
 				Settings.DensityMulConstant = Density * ProjectionDepthBiasParameters.y;
+				Settings.ProjectionDepthBiasParameters = ProjectionDepthBiasParameters;
 				
 				// ideally we use a larger filter kernel for SSSbut as Gather4 makes that harder
 				SSSTransmission = ManualPCF(ShadowPosition.xy, Settings); 
diff --git a/Engine/Source/Developer/iOS/MetalShaderFormat/Private/MetalShaderCompiler.cpp b/Engine/Source/Developer/iOS/MetalShaderFormat/Private/MetalShaderCompiler.cpp
--- a/Engine/Source/Developer/iOS/MetalShaderFormat/Private/MetalShaderCompiler.cpp
+++ b/Engine/Source/Developer/iOS/MetalShaderFormat/Private/MetalShaderCompiler.cpp
@@ -617,6 +617,8 @@ static void BuildMetalShaderOutput(
 		BuildResourceTableTokenStream(GenericSRT.ShaderResourceViewMap, GenericSRT.MaxBoundResourceTable, Header.Bindings.ShaderResourceTable.ShaderResourceViewMap);
 		BuildResourceTableTokenStream(GenericSRT.SamplerMap, GenericSRT.MaxBoundResourceTable, Header.Bindings.ShaderResourceTable.SamplerMap);
 		BuildResourceTableTokenStream(GenericSRT.UnorderedAccessViewMap, GenericSRT.MaxBoundResourceTable, Header.Bindings.ShaderResourceTable.UnorderedAccessViewMap);
+
+		Header.Bindings.NumUniformBuffers = FMath::Max((uint8)GetNumUniformBuffersUsed(GenericSRT), Header.Bindings.NumUniformBuffers);
 	}
 
 	const int32 MaxSamplers = GetFeatureLevelMaxTextureSamplers(GetMaxSupportedFeatureLevel((EShaderPlatform)ShaderOutput.Target.Platform));
@@ -852,6 +854,13 @@ void CompileShader_Metal(const FShaderCompilerInput& Input,FShaderCompilerOutput
 			return;
 		}
 
+
+		// This requires removing the HLSLCC_NoPreprocess flag later on!
+		if (!RemoveUniformBuffersFromSource(PreprocessedShader))
+		{
+			return;
+		}
+
 		// Write out the preprocessed file and a batch file to compile it if requested (DumpDebugInfoPath is valid)
 		if (bDumpDebugInfo)
 		{
@@ -879,6 +888,9 @@ void CompileShader_Metal(const FShaderCompilerInput& Input,FShaderCompilerOutput
 			}
 		}
 
+		// Required as we added the RemoveUniformBuffersFromSource() function (the cross-compiler won't be able to interpret comments w/o a preprocessor)
+		CCFlags &= ~HLSLCC_NoPreprocess;
+
 		FMetalCodeBackend MetalBackEnd(CCFlags);
 		FMetalLanguageSpec MetalLanguageSpec;
 		int32 Result = HlslCrossCompile(
diff --git a/Engine/Source/Developer/iOS/MetalShaderFormat/Private/MetalShaderFormat.cpp b/Engine/Source/Developer/iOS/MetalShaderFormat/Private/MetalShaderFormat.cpp
--- a/Engine/Source/Developer/iOS/MetalShaderFormat/Private/MetalShaderFormat.cpp
+++ b/Engine/Source/Developer/iOS/MetalShaderFormat/Private/MetalShaderFormat.cpp
@@ -14,7 +14,7 @@ class FMetalShaderFormat : public IShaderFormat
 {
 	enum
 	{
-		HEADER_VERSION = 6,
+		HEADER_VERSION = 7,
 	};
 	struct FVersion
 	{
diff --git a/Engine/Source/Runtime/Engine/Private/Landscape/Landscape.cpp b/Engine/Source/Runtime/Engine/Private/Landscape/Landscape.cpp
--- a/Engine/Source/Runtime/Engine/Private/Landscape/Landscape.cpp
+++ b/Engine/Source/Runtime/Engine/Private/Landscape/Landscape.cpp
@@ -682,7 +682,7 @@ FPrimitiveSceneProxy* ULandscapeComponent::CreateSceneProxy()
 		Proxy = new FLandscapeComponentSceneProxy(this, NULL);
 #endif
 	}
-	else if (FeatureLevel >= ERHIFeatureLevel::ES2)
+	else
 	{
 #if WITH_EDITOR
 		if (!PlatformData.HasValidPlatformData()) // Deferred generation
diff --git a/Engine/Source/Runtime/Engine/Public/SceneManagement.h b/Engine/Source/Runtime/Engine/Public/SceneManagement.h
--- a/Engine/Source/Runtime/Engine/Public/SceneManagement.h
+++ b/Engine/Source/Runtime/Engine/Public/SceneManagement.h
@@ -28,6 +28,12 @@ struct FDynamicMeshVertex;
 
 DECLARE_LOG_CATEGORY_EXTERN(LogBufferVisualization, Log, All);
 
+static const int MAX_FORWARD_SHADOWCASCADES = 2;
+
+// -----------------------------------------------------------------------------
+
+
+
 /**
  * The scene manager's persistent view state.
  */
@@ -1552,6 +1558,10 @@ BEGIN_UNIFORM_BUFFER_STRUCT(FViewUniformShaderParameters,ENGINE_API)
 	DECLARE_UNIFORM_BUFFER_STRUCT_MEMBER_EX(float,UnlitViewmodeMask, EShaderPrecisionModifier::Half)
 	DECLARE_UNIFORM_BUFFER_STRUCT_MEMBER_EX(FLinearColor,DirectionalLightColor, EShaderPrecisionModifier::Half)
 	DECLARE_UNIFORM_BUFFER_STRUCT_MEMBER_EX(FVector,DirectionalLightDirection, EShaderPrecisionModifier::Half)
+	DECLARE_UNIFORM_BUFFER_STRUCT_MEMBER_EX(FVector4, DirectionalLightShadowSize, EShaderPrecisionModifier::Half)
+	DECLARE_UNIFORM_BUFFER_STRUCT_MEMBER_EX(float, DirectionalLightShadowTransition, EShaderPrecisionModifier::Half)			
+	DECLARE_UNIFORM_BUFFER_STRUCT_MEMBER_ARRAY(FMatrix, DirectionalLightScreenToShadow, [MAX_FORWARD_SHADOWCASCADES])
+	DECLARE_UNIFORM_BUFFER_STRUCT_MEMBER_EX(FVector4, DirectionalLightShadowDistances, EShaderPrecisionModifier::Half)
 	DECLARE_UNIFORM_BUFFER_STRUCT_MEMBER_EX(FLinearColor,UpperSkyColor, EShaderPrecisionModifier::Half)
 	DECLARE_UNIFORM_BUFFER_STRUCT_MEMBER_EX(FLinearColor,LowerSkyColor, EShaderPrecisionModifier::Half)
 	DECLARE_UNIFORM_BUFFER_STRUCT_MEMBER_ARRAY(FVector4,TranslucencyLightingVolumeMin,[TVC_MAX])
@@ -1599,6 +1609,8 @@ BEGIN_UNIFORM_BUFFER_STRUCT(FViewUniformShaderParameters,ENGINE_API)
 	DECLARE_UNIFORM_BUFFER_STRUCT_MEMBER(FLinearColor,SkyLightColor)
 	DECLARE_UNIFORM_BUFFER_STRUCT_MEMBER_ARRAY(FVector4,SkyIrradianceEnvironmentMap,[7])
 	DECLARE_UNIFORM_BUFFER_STRUCT_MEMBER(float, ES2PreviewMode)
+	DECLARE_UNIFORM_BUFFER_STRUCT_MEMBER_TEXTURE(Texture2D, DirectionalLightShadowTexture)	
+	DECLARE_UNIFORM_BUFFER_STRUCT_MEMBER_SAMPLER(SamplerState, DirectionalLightShadowSampler)
 END_UNIFORM_BUFFER_STRUCT(FViewUniformShaderParameters)
 
 
diff --git a/Engine/Source/Runtime/Engine/Public/ScreenRendering.h b/Engine/Source/Runtime/Engine/Public/ScreenRendering.h
--- a/Engine/Source/Runtime/Engine/Public/ScreenRendering.h
+++ b/Engine/Source/Runtime/Engine/Public/ScreenRendering.h
@@ -51,7 +51,7 @@ class FScreenPS : public FGlobalShader
 	DECLARE_EXPORTED_SHADER_TYPE(FScreenPS,Global,ENGINE_API);
 public:
 
-	static bool ShouldCache(EShaderPlatform Platform) { return IsFeatureLevelSupported(Platform, ERHIFeatureLevel::ES2); }
+	static bool ShouldCache(EShaderPlatform Platform) { return true; }
 
 	FScreenPS(const ShaderMetaType::CompiledShaderInitializerType& Initializer):
 		FGlobalShader(Initializer)
@@ -118,7 +118,7 @@ class FScreenVSForGS : public FScreenVS
 	DECLARE_EXPORTED_SHADER_TYPE(FScreenVSForGS,Global,ENGINE_API);
 public:
 
-	static bool ShouldCache(EShaderPlatform Platform) { return IsFeatureLevelSupported(Platform, ERHIFeatureLevel::ES2); }
+	static bool ShouldCache(EShaderPlatform Platform) { return IsFeatureLevelSupported(Platform, ERHIFeatureLevel::SM4); }
 
 	FScreenVSForGS(const ShaderMetaType::CompiledShaderInitializerType& Initializer):
 	  FScreenVS(Initializer)
diff --git a/Engine/Source/Runtime/OpenGLDrv/Private/Mac/OpenGLMac.h b/Engine/Source/Runtime/OpenGLDrv/Private/Mac/OpenGLMac.h
--- a/Engine/Source/Runtime/OpenGLDrv/Private/Mac/OpenGLMac.h
+++ b/Engine/Source/Runtime/OpenGLDrv/Private/Mac/OpenGLMac.h
@@ -71,9 +71,6 @@ struct FMacOpenGL : public FOpenGL3
 		MacGetQueryObject(QueryId, QueryMode, OutResult);
 	}
 	
-	// A driver bug with Nvidia cards on mac causes rendering with GL_LayerIndex to target different cubemap faces when rendering to mip > 0 to always go to the +X cube face
-	static FORCEINLINE bool SupportsGSRenderTargetLayerSwitchingToMips() { return false; }
-
 	static FORCEINLINE void InitDebugContext() UGL_OPTIONAL_VOID
 
 	static FORCEINLINE void LabelObject(GLenum Type, GLuint Object, const ANSICHAR* Name)
diff --git a/Engine/Source/Runtime/OpenGLDrv/Private/OpenGLDevice.cpp b/Engine/Source/Runtime/OpenGLDrv/Private/OpenGLDevice.cpp
--- a/Engine/Source/Runtime/OpenGLDrv/Private/OpenGLDevice.cpp
+++ b/Engine/Source/Runtime/OpenGLDrv/Private/OpenGLDevice.cpp
@@ -630,7 +630,6 @@ static void InitRHICapabilitiesForGL()
 #endif
 
 	GSupportsVolumeTextureRendering = FOpenGL::SupportsVolumeTextureRendering();
-	GSupportsGSRenderTargetLayerSwitchingToMips = FOpenGL::SupportsGSRenderTargetLayerSwitchingToMips();
 	GSupportsRenderDepthTargetableShaderResources = true;
 	GSupportsVertexTextureFetch = true;
 	GSupportsRenderTargetFormat_PF_G8 = true;
@@ -676,7 +675,7 @@ static void InitRHICapabilitiesForGL()
 
 	GLenum depthFormat = FOpenGL::GetDepthFormat();
 
-	// Initialize the platform pixel format map.						 InternalFormat						InternalFormatSRGB							Format				Type								bCompressed
+	// Initialize the platform pixel format map.					 InternalFormat						InternalFormatSRGB							Format				Type								bCompressed
 	SetupTextureFormat( PF_Unknown,				FOpenGLTextureFormat());
 	SetupTextureFormat( PF_A32B32G32R32F,		FOpenGLTextureFormat(GL_RGBA32F,						GL_NONE,									GL_RGBA,				GL_FLOAT,							false));
 	SetupTextureFormat( PF_UYVY,				FOpenGLTextureFormat());
diff --git a/Engine/Source/Runtime/OpenGLDrv/Public/OpenGL.h b/Engine/Source/Runtime/OpenGLDrv/Public/OpenGL.h
--- a/Engine/Source/Runtime/OpenGLDrv/Public/OpenGL.h
+++ b/Engine/Source/Runtime/OpenGLDrv/Public/OpenGL.h
@@ -102,7 +102,6 @@ public:
 	static FORCEINLINE bool SupportsTextureHalfFloat()					{ return true; }
 	static FORCEINLINE bool SupportsColorBufferHalfFloat()				{ return true; }
 	static FORCEINLINE bool SupportsVolumeTextureRendering()			{ return false; }
-	static FORCEINLINE bool SupportsGSRenderTargetLayerSwitchingToMips() { return true; }
 	static FORCEINLINE bool SupportsShaderFramebufferFetch()			{ return false; }
 	static FORCEINLINE bool SupportsVertexArrayBGRA()					{ return true; }
 	static FORCEINLINE bool SupportsBGRA8888()							{ return true; }
diff --git a/Engine/Source/Runtime/RHI/Private/RHI.cpp b/Engine/Source/Runtime/RHI/Private/RHI.cpp
--- a/Engine/Source/Runtime/RHI/Private/RHI.cpp
+++ b/Engine/Source/Runtime/RHI/Private/RHI.cpp
@@ -169,7 +169,6 @@ bool GHardwareHiddenSurfaceRemoval = false;
 bool GRHISupportsAsyncTextureCreation = false;
 bool GSupportsQuads = false;
 bool GSupportsVolumeTextureRendering = true;
-bool GSupportsGSRenderTargetLayerSwitchingToMips = true;
 bool GSupportsSeparateRenderTargetBlendState = false;
 bool GSupportsDepthRenderTargetWithoutColorRenderTarget = true;
 float GPixelCenterOffset = 0;
diff --git a/Engine/Source/Runtime/RHI/Public/RHI.h b/Engine/Source/Runtime/RHI/Public/RHI.h
--- a/Engine/Source/Runtime/RHI/Public/RHI.h
+++ b/Engine/Source/Runtime/RHI/Public/RHI.h
@@ -99,9 +99,6 @@ extern RHI_API bool GSupportsQuads;
 /** True if and only if the GPU support rendering to volume textures (2D Array, 3D). Some OpenGL 3.3 cards support SM4, but can't render to volume textures. */
 extern RHI_API bool GSupportsVolumeTextureRendering;
 
-/** Used to work around a bug with Nvidia cards on Mac: using GL_LayerIndex to specify render target layer doesn't work when rendering to mips. */
-extern RHI_API bool GSupportsGSRenderTargetLayerSwitchingToMips;
-
 /** True if the RHI supports separate blend states per render target. */
 extern RHI_API bool GSupportsSeparateRenderTargetBlendState;
 
diff --git a/Engine/Source/Runtime/Renderer/Private/DeferredShadingRenderer.h b/Engine/Source/Runtime/Renderer/Private/DeferredShadingRenderer.h
--- a/Engine/Source/Runtime/Renderer/Private/DeferredShadingRenderer.h
+++ b/Engine/Source/Runtime/Renderer/Private/DeferredShadingRenderer.h
@@ -80,7 +80,7 @@ public:
 	/** Sorts base pass draw lists front to back for improved GPU culling. */
 	void SortBasePassStaticData(FVector ViewPosition);
 
-    /** Renders the basepass for the dynamic data of a given View. */
+	/** Renders the basepass for the dynamic data of a given View. */
 	void RenderBasePassDynamicData(FRHICommandList& RHICmdList, const FViewInfo& View, bool& bOutDirty);
 
 	/** Renders the basepass for the dynamic data of a given View, in parallel. */
@@ -128,63 +128,48 @@ private:
 	/** Creates a per object projected shadow for the given interaction. */
 	void CreatePerObjectProjectedShadow(
 		FRHICommandListImmediate& RHICmdList,
-		FLightPrimitiveInteraction* Interaction, 
-		bool bCreateTranslucentObjectShadow, 
+		FLightPrimitiveInteraction* Interaction,
+		bool bCreateTranslucentObjectShadow,
 		bool bCreateInsetObjectShadow,
-		const TArray<FProjectedShadowInfo*,SceneRenderingAllocator>& ViewDependentWholeSceneShadows,
+		const TArray<FProjectedShadowInfo*, SceneRenderingAllocator>& ViewDependentWholeSceneShadows,
 		TArray<FProjectedShadowInfo*, SceneRenderingAllocator>& OutPreShadows);
 
 	/**
-	 * Creates a projected shadow for all primitives affected by a light.  
-	 * @param LightSceneInfo - The light to create a shadow for.
-	 */
+	* Creates a projected shadow for all primitives affected by a light.
+	* @param LightSceneInfo - The light to create a shadow for.
+	*/
 	void CreateWholeSceneProjectedShadow(FLightSceneInfo* LightSceneInfo);
 
-	/**
-	 * Checks to see if this primitive is affected by various shadow types
-	 *
-	 * @param PrimitiveSceneInfoCompact The primitive to check for shadow interaction
-	 * @param PreShadows The list of pre-shadows to check against
-	 */
-	void GatherShadowsForPrimitiveInner(const FPrimitiveSceneInfoCompact& PrimitiveSceneInfoCompact,
-		const TArray<FProjectedShadowInfo*,SceneRenderingAllocator>& PreShadows,
-		const TArray<FProjectedShadowInfo*,SceneRenderingAllocator>& ViewDependentWholeSceneShadows,
-		bool bReflectionCaptureScene);
-
-	/** Gathers the list of primitives used to draw various shadow types */
-	void GatherShadowPrimitives(
-		const TArray<FProjectedShadowInfo*,SceneRenderingAllocator>& PreShadows,
-		const TArray<FProjectedShadowInfo*,SceneRenderingAllocator>& ViewDependentWholeSceneShadows,
-		bool bReflectionCaptureScene);
-
 	/** Updates the preshadow cache, allocating new preshadows that can fit and evicting old ones. */
 	void UpdatePreshadowCache();
 
 	/** Finds a matching cached preshadow, if one exists. */
 	TRefCountPtr<FProjectedShadowInfo> GetCachedPreshadow(
-		const FLightPrimitiveInteraction* InParentInteraction, 
+		const FLightPrimitiveInteraction* InParentInteraction,
 		const FProjectedShadowInitializer& Initializer,
 		const FBoxSphereBounds& Bounds,
 		uint32 InResolutionX);
 
-	/** Calculates projected shadow visibility. */
-	void InitProjectedShadowVisibility(FRHICommandListImmediate& RHICmdList);
-
-	/** Returns whether a per object shadow should be created due to the light being a stationary light. */
-	bool ShouldCreateObjectShadowForStationaryLight(const FLightSceneInfo* LightSceneInfo, const FPrimitiveSceneProxy* PrimitiveSceneProxy, bool bInteractionShadowMapped) const;
-
 	/** Creates shadows for the given interaction. */
 	void SetupInteractionShadows(
 		FRHICommandListImmediate& RHICmdList,
-		FLightPrimitiveInteraction* Interaction, 
-		FVisibleLightInfo& VisibleLightInfo, 
+		FLightPrimitiveInteraction* Interaction,
+		FVisibleLightInfo& VisibleLightInfo,
 		bool bStaticSceneOnly,
-		const TArray<FProjectedShadowInfo*,SceneRenderingAllocator>& ViewDependentWholeSceneShadows,
-		TArray<FProjectedShadowInfo*,SceneRenderingAllocator>& PreShadows);
+		const TArray<FProjectedShadowInfo*, SceneRenderingAllocator>& ViewDependentWholeSceneShadows,
+		TArray<FProjectedShadowInfo*, SceneRenderingAllocator>& PreShadows);
 
 	/** Finds the visible dynamic shadows for each view. */
 	void InitDynamicShadows(FRHICommandListImmediate& RHICmdList);
 
+	/**
+	* Used by RenderLights to figure out if light functions need to be rendered to the attenuation buffer.
+	*
+	* @param LightSceneInfo Represents the current light
+	* @return true if anything got rendered
+	*/
+	bool CheckForLightFunction(const FLightSceneInfo* LightSceneInfo) const;
+
 	/** Determines which primitives are visible for each view. */
 	void InitViews(FRHICommandListImmediate& RHICmdList);
 
@@ -267,14 +252,6 @@ private:
 	/** Updates the downsized depth buffer with the current full resolution depth buffer. */
 	void UpdateDownsampledDepthSurface(FRHICommandList& RHICmdList);
 
-	/**
-	  * Used by RenderLights to figure out if projected shadows need to be rendered to the attenuation buffer.
-	  *
-	  * @param LightSceneInfo Represents the current light
-	  * @return true if anything needs to be rendered
-	  */
-	bool CheckForProjectedShadows( const FLightSceneInfo* LightSceneInfo ) const;
-
 	/** Renders one pass point light shadows. */
 	bool RenderOnePassPointLightShadows(FRHICommandListImmediate& RHICmdList, const FLightSceneInfo* LightSceneInfo, bool bRenderedTranslucentObjectShadows, bool& bInjectedTranslucentVolume);
 
@@ -302,14 +279,6 @@ private:
 	/** Caches the depths of any preshadows that should be cached, and renders their projections. */
 	bool RenderCachedPreshadows(FRHICommandListImmediate& RHICmdList, const FLightSceneInfo* LightSceneInfo);
 
-	/**
-	  * Used by RenderLights to figure out if light functions need to be rendered to the attenuation buffer.
-	  *
-	  * @param LightSceneInfo Represents the current light
-	  * @return true if anything got rendered
-	  */
-	bool CheckForLightFunction( const FLightSceneInfo* LightSceneInfo ) const;
-
 	/**
 	  * Used by RenderLights to render a light function to the attenuation buffer.
 	  *
diff --git a/Engine/Source/Runtime/Renderer/Private/ForwardBasePassRendering.cpp b/Engine/Source/Runtime/Renderer/Private/ForwardBasePassRendering.cpp
--- a/Engine/Source/Runtime/Renderer/Private/ForwardBasePassRendering.cpp
+++ b/Engine/Source/Runtime/Renderer/Private/ForwardBasePassRendering.cpp
@@ -13,18 +13,26 @@
 	typedef TBasePassForForwardShadingVS< LightMapPolicyType, HDR_LINEAR_64 > TBasePassForForwardShadingVS##LightMapPolicyName##HDRLinear64; \
 	IMPLEMENT_MATERIAL_SHADER_TYPE(template<>,TBasePassForForwardShadingVS##LightMapPolicyName##LDRGamma32,TEXT("BasePassForForwardShadingVertexShader"),TEXT("Main"),SF_Vertex); \
 	IMPLEMENT_MATERIAL_SHADER_TYPE(template<>,TBasePassForForwardShadingVS##LightMapPolicyName##HDRLinear64,TEXT("BasePassForForwardShadingVertexShader"),TEXT("Main"),SF_Vertex); \
-	typedef TBasePassForForwardShadingPS< LightMapPolicyType, LDR_GAMMA_32 > TBasePassForForwardShadingPS##LightMapPolicyName##LDRGamma32; \
-	typedef TBasePassForForwardShadingPS< LightMapPolicyType, HDR_LINEAR_32 > TBasePassForForwardShadingPS##LightMapPolicyName##HDRLinear32; \
-	typedef TBasePassForForwardShadingPS< LightMapPolicyType, HDR_LINEAR_64 > TBasePassForForwardShadingPS##LightMapPolicyName##HDRLinear64; \
+	typedef TBasePassForForwardShadingPS< LightMapPolicyType, LDR_GAMMA_32, false > TBasePassForForwardShadingPS##LightMapPolicyName##LDRGamma32; \
+	typedef TBasePassForForwardShadingPS< LightMapPolicyType, HDR_LINEAR_32, false > TBasePassForForwardShadingPS##LightMapPolicyName##HDRLinear32; \
+	typedef TBasePassForForwardShadingPS< LightMapPolicyType, HDR_LINEAR_64, false > TBasePassForForwardShadingPS##LightMapPolicyName##HDRLinear64; \
+	typedef TBasePassForForwardShadingPS< LightMapPolicyType, LDR_GAMMA_32, true > TBasePassForForwardShadingPS##LightMapPolicyName##LDRGamma32##Skylight; \
+	typedef TBasePassForForwardShadingPS< LightMapPolicyType, HDR_LINEAR_32, true > TBasePassForForwardShadingPS##LightMapPolicyName##HDRLinear32##Skylight; \
+	typedef TBasePassForForwardShadingPS< LightMapPolicyType, HDR_LINEAR_64, true > TBasePassForForwardShadingPS##LightMapPolicyName##HDRLinear64##Skylight; \
 	IMPLEMENT_MATERIAL_SHADER_TYPE(template<>,TBasePassForForwardShadingPS##LightMapPolicyName##LDRGamma32,TEXT("BasePassForForwardShadingPixelShader"),TEXT("Main"),SF_Pixel); \
 	IMPLEMENT_MATERIAL_SHADER_TYPE(template<>,TBasePassForForwardShadingPS##LightMapPolicyName##HDRLinear32,TEXT("BasePassForForwardShadingPixelShader"),TEXT("Main"),SF_Pixel); \
-	IMPLEMENT_MATERIAL_SHADER_TYPE(template<>,TBasePassForForwardShadingPS##LightMapPolicyName##HDRLinear64,TEXT("BasePassForForwardShadingPixelShader"),TEXT("Main"),SF_Pixel);
+	IMPLEMENT_MATERIAL_SHADER_TYPE(template<>,TBasePassForForwardShadingPS##LightMapPolicyName##HDRLinear64,TEXT("BasePassForForwardShadingPixelShader"),TEXT("Main"),SF_Pixel); \
+	IMPLEMENT_MATERIAL_SHADER_TYPE(template<>, TBasePassForForwardShadingPS##LightMapPolicyName##LDRGamma32##Skylight, TEXT("BasePassForForwardShadingPixelShader"), TEXT("Main"), SF_Pixel); \
+	IMPLEMENT_MATERIAL_SHADER_TYPE(template<>, TBasePassForForwardShadingPS##LightMapPolicyName##HDRLinear32##Skylight, TEXT("BasePassForForwardShadingPixelShader"), TEXT("Main"), SF_Pixel); \
+	IMPLEMENT_MATERIAL_SHADER_TYPE(template<>, TBasePassForForwardShadingPS##LightMapPolicyName##HDRLinear64##Skylight, TEXT("BasePassForForwardShadingPixelShader"), TEXT("Main"), SF_Pixel);
 
 // Implement shader types per lightmap policy
 IMPLEMENT_FORWARD_SHADING_BASEPASS_LIGHTMAPPED_SHADER_TYPE( FNoLightMapPolicy, FNoLightMapPolicy ); 
 IMPLEMENT_FORWARD_SHADING_BASEPASS_LIGHTMAPPED_SHADER_TYPE( TLightMapPolicy<LQ_LIGHTMAP>, TLightMapPolicyLQ );
 IMPLEMENT_FORWARD_SHADING_BASEPASS_LIGHTMAPPED_SHADER_TYPE( TDistanceFieldShadowsAndLightMapPolicy<LQ_LIGHTMAP>, TDistanceFieldShadowsAndLightMapPolicyLQ );
 IMPLEMENT_FORWARD_SHADING_BASEPASS_LIGHTMAPPED_SHADER_TYPE( FSimpleDirectionalLightAndSHIndirectPolicy, FSimpleDirectionalLightAndSHIndirectPolicy );
+IMPLEMENT_FORWARD_SHADING_BASEPASS_LIGHTMAPPED_SHADER_TYPE( FMovableDirectionalLightLightingPolicy, FMovableDirectionalLightLightingPolicy );
+IMPLEMENT_FORWARD_SHADING_BASEPASS_LIGHTMAPPED_SHADER_TYPE( FMovableDirectionalLightCSMLightingPolicy, FMovableDirectionalLightCSMLightingPolicy );
 
 /** The action used to draw a base pass static mesh element. */
 class FDrawBasePassForwardShadingStaticMeshAction
@@ -70,6 +78,7 @@ public:
 		TStaticMeshDrawList<TBasePassForForwardShadingDrawingPolicy<LightMapPolicyType> >& DrawList =
 			Scene->GetForwardShadingBasePassDrawList<LightMapPolicyType>(DrawType);
 
+		ERHIFeatureLevel::Type FeatureLevel = Scene->GetFeatureLevel();
 		// Add the static mesh to the draw list.
 		DrawList.AddMesh(
 			StaticMesh,
@@ -80,9 +89,12 @@ public:
 				*Parameters.Material,
 				LightMapPolicy,
 				Parameters.BlendMode,
-				Parameters.TextureMode
+				Parameters.TextureMode,
+				Parameters.ShadingModel != MSM_Unlit && Scene && Scene->ShouldRenderSkylight(),
+				false,
+				FeatureLevel
 				),
-			Scene->GetFeatureLevel()
+				FeatureLevel
 			);
 	}
 };
@@ -165,6 +177,7 @@ public:
 #endif
 
 		const bool bIsLitMaterial = Parameters.ShadingModel != MSM_Unlit;
+		const FScene* Scene = Parameters.PrimitiveSceneProxy ? Parameters.PrimitiveSceneProxy->GetPrimitiveSceneInfo()->Scene : NULL;
 
 		TBasePassForForwardShadingDrawingPolicy<LightMapPolicyType> DrawingPolicy(
 			Parameters.Mesh.VertexFactory,
@@ -173,7 +186,9 @@ public:
 			LightMapPolicy,
 			Parameters.BlendMode,
 			Parameters.TextureMode,
-			View.Family->EngineShowFlags.ShaderComplexity
+			Parameters.ShadingModel != MSM_Unlit && Scene && Scene->ShouldRenderSkylight(),
+			View.Family->EngineShowFlags.ShaderComplexity,
+			View.GetFeatureLevel()
 			);
 		RHICmdList.BuildAndSetLocalBoundShaderState(DrawingPolicy.GetBoundShaderStateInput(View.GetFeatureLevel()));
 		DrawingPolicy.SetSharedState(RHICmdList, &View, typename TBasePassForForwardShadingDrawingPolicy<LightMapPolicyType>::ContextDataType());
@@ -316,6 +331,8 @@ void FForwardShadingSceneRenderer::RenderForwardShadingBasePass(FRHICommandListI
 			Scene->BasePassForForwardShadingLowQualityLightMapDrawList[DrawType].SortFrontToBack(Views[0].ViewLocation);
 			Scene->BasePassForForwardShadingDistanceFieldShadowMapLightMapDrawList[DrawType].SortFrontToBack(Views[0].ViewLocation);
 			Scene->BasePassForForwardShadingDirectionalLightAndSHIndirectDrawList[DrawType].SortFrontToBack(Views[0].ViewLocation);
+			Scene->BasePassForForwardShadingMovableDirectionalLightCSMDrawList[DrawType].SortFrontToBack(Views[0].ViewLocation);
+			Scene->BasePassForForwardShadingMovableDirectionalLightDrawList[DrawType].SortFrontToBack(Views[0].ViewLocation);
 			Scene->BasePassForForwardShadingNoLightMapDrawList[DrawType].SortFrontToBack(Views[0].ViewLocation);
 		}
 	}
@@ -340,6 +357,8 @@ void FForwardShadingSceneRenderer::RenderForwardShadingBasePass(FRHICommandListI
 				MaxDraws -= Scene->BasePassForForwardShadingLowQualityLightMapDrawList[FScene::EBasePass_Default].DrawVisibleFrontToBack(RHICmdList, View, View.StaticMeshVisibilityMap, View.StaticMeshBatchVisibility, MaxDraws);
 				MaxDraws -= Scene->BasePassForForwardShadingDistanceFieldShadowMapLightMapDrawList[FScene::EBasePass_Default].DrawVisibleFrontToBack(RHICmdList, View, View.StaticMeshVisibilityMap, View.StaticMeshBatchVisibility, MaxDraws);
 				MaxDraws -= Scene->BasePassForForwardShadingDirectionalLightAndSHIndirectDrawList[FScene::EBasePass_Default].DrawVisibleFrontToBack(RHICmdList, View, View.StaticMeshVisibilityMap, View.StaticMeshBatchVisibility, MaxDraws);
+				MaxDraws -= Scene->BasePassForForwardShadingMovableDirectionalLightCSMDrawList[FScene::EBasePass_Default].DrawVisibleFrontToBack(RHICmdList, View, View.StaticMeshVisibilityMap, View.StaticMeshBatchVisibility, MaxDraws);
+				MaxDraws -= Scene->BasePassForForwardShadingMovableDirectionalLightDrawList[FScene::EBasePass_Default].DrawVisibleFrontToBack(RHICmdList, View, View.StaticMeshVisibilityMap, View.StaticMeshBatchVisibility, MaxDraws);
 				MaxDraws -= Scene->BasePassForForwardShadingNoLightMapDrawList[FScene::EBasePass_Default].DrawVisibleFrontToBack(RHICmdList, View, View.StaticMeshVisibilityMap, View.StaticMeshBatchVisibility, MaxDraws);
 			}
 			else
@@ -348,6 +367,8 @@ void FForwardShadingSceneRenderer::RenderForwardShadingBasePass(FRHICommandListI
 				Scene->BasePassForForwardShadingLowQualityLightMapDrawList[FScene::EBasePass_Default].DrawVisible(RHICmdList, View, View.StaticMeshVisibilityMap, View.StaticMeshBatchVisibility);
 				Scene->BasePassForForwardShadingDistanceFieldShadowMapLightMapDrawList[FScene::EBasePass_Default].DrawVisible(RHICmdList, View, View.StaticMeshVisibilityMap, View.StaticMeshBatchVisibility);
 				Scene->BasePassForForwardShadingDirectionalLightAndSHIndirectDrawList[FScene::EBasePass_Default].DrawVisible(RHICmdList, View, View.StaticMeshVisibilityMap, View.StaticMeshBatchVisibility);
+				Scene->BasePassForForwardShadingMovableDirectionalLightCSMDrawList[FScene::EBasePass_Default].DrawVisible(RHICmdList, View, View.StaticMeshVisibilityMap, View.StaticMeshBatchVisibility);
+				Scene->BasePassForForwardShadingMovableDirectionalLightDrawList[FScene::EBasePass_Default].DrawVisible(RHICmdList, View, View.StaticMeshVisibilityMap, View.StaticMeshBatchVisibility);
 				Scene->BasePassForForwardShadingNoLightMapDrawList[FScene::EBasePass_Default].DrawVisible(RHICmdList, View, View.StaticMeshVisibilityMap, View.StaticMeshBatchVisibility);
 			}
 		}
@@ -423,6 +444,8 @@ void FForwardShadingSceneRenderer::RenderForwardShadingBasePass(FRHICommandListI
 			MaxDraws -= Scene->BasePassForForwardShadingLowQualityLightMapDrawList[FScene::EBasePass_Masked].DrawVisibleFrontToBack(RHICmdList, View,View.StaticMeshVisibilityMap,View.StaticMeshBatchVisibility,MaxDraws);
 			MaxDraws -= Scene->BasePassForForwardShadingDistanceFieldShadowMapLightMapDrawList[FScene::EBasePass_Masked].DrawVisibleFrontToBack(RHICmdList, View,View.StaticMeshVisibilityMap,View.StaticMeshBatchVisibility,MaxDraws);
 			MaxDraws -= Scene->BasePassForForwardShadingDirectionalLightAndSHIndirectDrawList[FScene::EBasePass_Masked].DrawVisibleFrontToBack(RHICmdList, View,View.StaticMeshVisibilityMap,View.StaticMeshBatchVisibility,MaxDraws);
+			MaxDraws -= Scene->BasePassForForwardShadingMovableDirectionalLightCSMDrawList[FScene::EBasePass_Masked].DrawVisibleFrontToBack(RHICmdList, View, View.StaticMeshVisibilityMap, View.StaticMeshBatchVisibility, MaxDraws);
+			MaxDraws -= Scene->BasePassForForwardShadingMovableDirectionalLightDrawList[FScene::EBasePass_Masked].DrawVisibleFrontToBack(RHICmdList, View,View.StaticMeshVisibilityMap,View.StaticMeshBatchVisibility,MaxDraws);
 		}
 		else
 		{
@@ -431,6 +454,8 @@ void FForwardShadingSceneRenderer::RenderForwardShadingBasePass(FRHICommandListI
 			Scene->BasePassForForwardShadingLowQualityLightMapDrawList[FScene::EBasePass_Masked].DrawVisible(RHICmdList, View,View.StaticMeshVisibilityMap,View.StaticMeshBatchVisibility);
 			Scene->BasePassForForwardShadingDistanceFieldShadowMapLightMapDrawList[FScene::EBasePass_Masked].DrawVisible(RHICmdList, View,View.StaticMeshVisibilityMap,View.StaticMeshBatchVisibility);
 			Scene->BasePassForForwardShadingDirectionalLightAndSHIndirectDrawList[FScene::EBasePass_Masked].DrawVisible(RHICmdList, View,View.StaticMeshVisibilityMap,View.StaticMeshBatchVisibility);
+			Scene->BasePassForForwardShadingMovableDirectionalLightCSMDrawList[FScene::EBasePass_Masked].DrawVisible(RHICmdList, View, View.StaticMeshVisibilityMap, View.StaticMeshBatchVisibility);
+			Scene->BasePassForForwardShadingMovableDirectionalLightDrawList[FScene::EBasePass_Masked].DrawVisible(RHICmdList, View,View.StaticMeshVisibilityMap,View.StaticMeshBatchVisibility);
 		}
 	}
 }
diff --git a/Engine/Source/Runtime/Renderer/Private/ForwardBasePassRendering.h b/Engine/Source/Runtime/Renderer/Private/ForwardBasePassRendering.h
--- a/Engine/Source/Runtime/Renderer/Private/ForwardBasePassRendering.h
+++ b/Engine/Source/Runtime/Renderer/Private/ForwardBasePassRendering.h
@@ -16,6 +16,17 @@ enum EOutputFormat
 	HDR_LINEAR_64,
 };
 
+static bool ShouldCacheShaderByOutputFormat(EOutputFormat OutputFormat)
+{
+	bool bSupportsMobileHDR = IsMobileHDR();
+	bool bShaderUsesLDR = (OutputFormat == LDR_GAMMA_32);
+	bool bShaderUsesHDR = !bShaderUsesLDR;
+
+	// only cache this shader if the LDR/HDR output matches what we currently support.  IsMobileHDR can't change, so we don't need
+	// the LDR shaders if we are doing HDR, and vice-versa.	
+	return (bShaderUsesLDR && !bSupportsMobileHDR) || (bShaderUsesHDR && bSupportsMobileHDR);
+}
+
 /**
  * The base shader type for vertex shaders that render the emissive color, and light-mapped/ambient lighting of a mesh.
  */
@@ -38,7 +49,7 @@ public:
 
 	static bool ShouldCache(EShaderPlatform Platform,const FMaterial* Material,const FVertexFactoryType* VertexFactoryType)
 	{
-		return IsES2Platform(Platform) && LightMapPolicyType::ShouldCache(Platform,Material,VertexFactoryType);
+		return IsMobilePlatform(Platform) && LightMapPolicyType::ShouldCache(Platform,Material,VertexFactoryType);
 	}
 
 	static void ModifyCompilationEnvironment(EShaderPlatform Platform, const FMaterial* Material, FShaderCompilerEnvironment& OutEnvironment)
@@ -83,6 +94,11 @@ class TBasePassForForwardShadingVS : public TBasePassForForwardShadingVSBaseType
 	DECLARE_SHADER_TYPE(TBasePassForForwardShadingVS,MeshMaterial);
 public:
 	
+	static bool ShouldCache(EShaderPlatform Platform, const FMaterial* Material, const FVertexFactoryType* VertexFactoryType)
+	{		
+		return TBasePassForForwardShadingVSBaseType<LightMapPolicyType>::ShouldCache(Platform, Material, VertexFactoryType) && ShouldCacheShaderByOutputFormat(OutputFormat);
+	}
+
 	static void ModifyCompilationEnvironment(EShaderPlatform Platform, const FMaterial* Material, FShaderCompilerEnvironment& OutEnvironment)
 	{
 		TBasePassForForwardShadingVSBaseType<LightMapPolicyType>::ModifyCompilationEnvironment(Platform, Material, OutEnvironment);
@@ -111,7 +127,7 @@ public:
 
 	static bool ShouldCache(EShaderPlatform Platform,const FMaterial* Material,const FVertexFactoryType* VertexFactoryType)
 	{
-		return IsES2Platform(Platform) && LightMapPolicyType::ShouldCache(Platform,Material,VertexFactoryType);
+		return IsMobilePlatform(Platform) && LightMapPolicyType::ShouldCache(Platform,Material,VertexFactoryType);
 	}
 
 	static void ModifyCompilationEnvironment(EShaderPlatform Platform, const FMaterial* Material, FShaderCompilerEnvironment& OutEnvironment)
@@ -126,7 +142,7 @@ public:
 	{
 		LightMapPolicyType::PixelParametersType::Bind(Initializer.ParameterMap);
 		ReflectionCubemap.Bind(Initializer.ParameterMap, TEXT("ReflectionCubemap"));
-		ReflectionSampler.Bind(Initializer.ParameterMap, TEXT("ReflectionCubemapSampler"));
+		ReflectionSampler.Bind(Initializer.ParameterMap, TEXT("ReflectionCubemapSampler"));		
 	}
 	TBasePassForForwardShadingPSBaseType() {}
 
@@ -137,6 +153,7 @@ public:
 
 	void SetMesh(FRHICommandList& RHICmdList, const FVertexFactory* VertexFactory,const FSceneView& View,const FPrimitiveSceneProxy* Proxy,const FMeshBatchElement& BatchElement)
 	{
+		FRHIPixelShader* PixelShader = GetPixelShader();
 		if (ReflectionCubemap.IsBound())
 		{
 			FTexture* ReflectionTexture = GBlackTextureCube;
@@ -151,10 +168,10 @@ public:
 			}
 
 			// Set the reflection cubemap
-			SetTextureParameter(RHICmdList, GetPixelShader(), ReflectionCubemap, ReflectionSampler, ReflectionTexture);
+			SetTextureParameter(RHICmdList, PixelShader, ReflectionCubemap, ReflectionSampler, ReflectionTexture);
 		}
 
-		FMeshMaterialShader::SetMesh(RHICmdList, GetPixelShader(),VertexFactory,View,Proxy,BatchElement);
+		FMeshMaterialShader::SetMesh(RHICmdList, PixelShader,VertexFactory,View,Proxy,BatchElement);		
 	}
 
 	virtual bool Serialize(FArchive& Ar)
@@ -162,25 +179,34 @@ public:
 		bool bShaderHasOutdatedParameters = FMeshMaterialShader::Serialize(Ar);
 		LightMapPolicyType::PixelParametersType::Serialize(Ar);
 		Ar << ReflectionCubemap;
-		Ar << ReflectionSampler;
+		Ar << ReflectionSampler;		
 		return bShaderHasOutdatedParameters;
 	}
 
 private:
 	FShaderResourceParameter ReflectionCubemap;
-	FShaderResourceParameter ReflectionSampler;
+	FShaderResourceParameter ReflectionSampler;	
 };
 
 
-template< typename LightMapPolicyType, EOutputFormat OutputFormat>
+template< typename LightMapPolicyType, EOutputFormat OutputFormat, bool bEnableSkyLight>
 class TBasePassForForwardShadingPS : public TBasePassForForwardShadingPSBaseType<LightMapPolicyType>
 {
 	DECLARE_SHADER_TYPE(TBasePassForForwardShadingPS,MeshMaterial);
 public:
+
+	static bool ShouldCache(EShaderPlatform Platform, const FMaterial* Material, const FVertexFactoryType* VertexFactoryType)
+	{		
+		// Only compile skylight version for lit materials on ES2 (Metal) or higher
+		const bool bShouldCacheBySkylight = !bEnableSkyLight || (Material->GetShadingModel() != MSM_Unlit);
+
+		return TBasePassForForwardShadingPSBaseType<LightMapPolicyType>::ShouldCache(Platform, Material, VertexFactoryType) && ShouldCacheShaderByOutputFormat(OutputFormat) && bShouldCacheBySkylight;
+	}
 	
 	static void ModifyCompilationEnvironment(EShaderPlatform Platform, const FMaterial* Material, FShaderCompilerEnvironment& OutEnvironment)
-	{
+	{		
 		TBasePassForForwardShadingPSBaseType<LightMapPolicyType>::ModifyCompilationEnvironment(Platform, Material, OutEnvironment);
+		OutEnvironment.SetDefine(TEXT("ENABLE_SKY_LIGHT"), (uint32)(bEnableSkyLight ? 1 : 0));
 		OutEnvironment.SetDefine( TEXT("USE_HDR_MOSAIC"), OutputFormat == HDR_LINEAR_32 ? 1u : 0u );
 		OutEnvironment.SetDefine( TEXT("OUTPUT_GAMMA_SPACE"), OutputFormat == LDR_GAMMA_32 ? 1u : 0u );
 	}
@@ -229,7 +255,9 @@ public:
 		LightMapPolicyType InLightMapPolicy,
 		EBlendMode InBlendMode,
 		ESceneRenderTargetsMode::Type InSceneTextureMode,
-		bool bOverrideWithShaderComplexity = false
+		bool bInEnableSkyLight,
+		bool bOverrideWithShaderComplexity,
+		ERHIFeatureLevel::Type FeatureLevel
 		):
 		FMeshDrawingPolicy(InVertexFactory,InMaterialRenderProxy,InMaterialResource,bOverrideWithShaderComplexity),
 		LightMapPolicy(InLightMapPolicy),
@@ -239,17 +267,42 @@ public:
 		if (IsMobileHDR32bpp())
 		{
 			VertexShader = InMaterialResource.GetShader<TBasePassForForwardShadingVS<LightMapPolicyType, HDR_LINEAR_64> >(InVertexFactory->GetType());
-			PixelShader = InMaterialResource.GetShader< TBasePassForForwardShadingPS<LightMapPolicyType, HDR_LINEAR_32> >(InVertexFactory->GetType());
+
+			if (bInEnableSkyLight)
+			{
+				PixelShader = InMaterialResource.GetShader< TBasePassForForwardShadingPS<LightMapPolicyType, HDR_LINEAR_32, true> >(InVertexFactory->GetType());
+			}
+			else
+			{
+				PixelShader = InMaterialResource.GetShader< TBasePassForForwardShadingPS<LightMapPolicyType, HDR_LINEAR_32, false> >(InVertexFactory->GetType());
+			}
 		}
 		else if (IsMobileHDR())
 		{
 			VertexShader = InMaterialResource.GetShader<TBasePassForForwardShadingVS<LightMapPolicyType, HDR_LINEAR_64> >(InVertexFactory->GetType());
-			PixelShader = InMaterialResource.GetShader< TBasePassForForwardShadingPS<LightMapPolicyType, HDR_LINEAR_64> >(InVertexFactory->GetType());
+
+			if (bInEnableSkyLight)
+			{
+				PixelShader = InMaterialResource.GetShader< TBasePassForForwardShadingPS<LightMapPolicyType, HDR_LINEAR_64, true> >(InVertexFactory->GetType());
+			}
+			else
+			{
+				PixelShader = InMaterialResource.GetShader< TBasePassForForwardShadingPS<LightMapPolicyType, HDR_LINEAR_64, false> >(InVertexFactory->GetType());
+			}
+			
 		}
 		else
 		{
 			VertexShader = InMaterialResource.GetShader<TBasePassForForwardShadingVS<LightMapPolicyType, LDR_GAMMA_32> >(InVertexFactory->GetType());
-			PixelShader = InMaterialResource.GetShader< TBasePassForForwardShadingPS<LightMapPolicyType, LDR_GAMMA_32> >(InVertexFactory->GetType());
+
+			if (bInEnableSkyLight)
+			{
+				PixelShader = InMaterialResource.GetShader< TBasePassForForwardShadingPS<LightMapPolicyType, LDR_GAMMA_32, true> >(InVertexFactory->GetType());
+			}
+			else
+			{
+				PixelShader = InMaterialResource.GetShader< TBasePassForForwardShadingPS<LightMapPolicyType, LDR_GAMMA_32, false> >(InVertexFactory->GetType());
+			}			
 		}
 	}
 
@@ -466,33 +519,54 @@ void ProcessBasePassMeshForForwardShading(
 		? Parameters.Mesh.LCI->GetLightMapInteraction() 
 		: FLightMapInteraction();
 
-	if (LightMapInteraction.GetType() == LMIT_Texture)
+	if (bIsLitMaterial)
 	{
-		const FShadowMapInteraction ShadowMapInteraction = (Parameters.Mesh.LCI && bIsLitMaterial) 
-			? Parameters.Mesh.LCI->GetShadowMapInteraction() 
-			: FShadowMapInteraction();
+		const FLightSceneInfo* SimpleDirectionalLight = Action.GetSimpleDirectionalLight();
+		const bool bUseMovableLight = SimpleDirectionalLight && !SimpleDirectionalLight->Proxy->HasStaticShadowing();
 
-		if (ShadowMapInteraction.GetType() == SMIT_Texture)
+		if (bUseMovableLight)
 		{
-			Action.template Process< TDistanceFieldShadowsAndLightMapPolicy<LQ_LIGHTMAP> >(
-				RHICmdList, 
-				Parameters,
-				TDistanceFieldShadowsAndLightMapPolicy<LQ_LIGHTMAP>(),
-				TDistanceFieldShadowsAndLightMapPolicy<LQ_LIGHTMAP>::ElementDataType(ShadowMapInteraction, LightMapInteraction));
+			// final determination of whether CSMs are rendered can be view dependent, thus we always need to clear the CSMs even if we're not going to render to them based on the condition below.
+			if (SimpleDirectionalLight && SimpleDirectionalLight->ShouldRenderViewIndependentWholeSceneShadows())
+			{
+				Action.template Process<FMovableDirectionalLightCSMLightingPolicy>(RHICmdList, Parameters, FMovableDirectionalLightCSMLightingPolicy(), FMovableDirectionalLightCSMLightingPolicy::ElementDataType());
+			}
+			else
+			{
+				Action.template Process<FMovableDirectionalLightLightingPolicy>(RHICmdList, Parameters, FMovableDirectionalLightLightingPolicy(), FMovableDirectionalLightLightingPolicy::ElementDataType());
+			}		
+		}
+		else if (LightMapInteraction.GetType() == LMIT_Texture)
+		{
+			const FShadowMapInteraction ShadowMapInteraction = (Parameters.Mesh.LCI && bIsLitMaterial) 
+				? Parameters.Mesh.LCI->GetShadowMapInteraction() 
+				: FShadowMapInteraction();
+
+			if (ShadowMapInteraction.GetType() == SMIT_Texture)
+			{
+				Action.template Process< TDistanceFieldShadowsAndLightMapPolicy<LQ_LIGHTMAP> >(
+					RHICmdList, 
+					Parameters,
+					TDistanceFieldShadowsAndLightMapPolicy<LQ_LIGHTMAP>(),
+					TDistanceFieldShadowsAndLightMapPolicy<LQ_LIGHTMAP>::ElementDataType(ShadowMapInteraction, LightMapInteraction));
+			}
+			else
+			{
+				Action.template Process< TLightMapPolicy<LQ_LIGHTMAP> >(RHICmdList, Parameters, TLightMapPolicy<LQ_LIGHTMAP>(), LightMapInteraction );
+			}
+		}
+		else if (IsIndirectLightingCacheAllowed(Parameters.FeatureLevel)
+			&& Parameters.PrimitiveSceneProxy
+			// Movable objects need to get their GI from the indirect lighting cache
+			&& Parameters.PrimitiveSceneProxy->IsMovable())
+		{
+			Action.template Process<FSimpleDirectionalLightAndSHIndirectPolicy>(RHICmdList, Parameters,FSimpleDirectionalLightAndSHIndirectPolicy(),FSimpleDirectionalLightAndSHIndirectPolicy::ElementDataType(Action.ShouldPackAmbientSH())); 
 		}
 		else
 		{
-			Action.template Process< TLightMapPolicy<LQ_LIGHTMAP> >(RHICmdList, Parameters, TLightMapPolicy<LQ_LIGHTMAP>(), LightMapInteraction );
+			Action.template Process<FNoLightMapPolicy>(RHICmdList, Parameters,FNoLightMapPolicy(),FNoLightMapPolicy::ElementDataType());
 		}
 	}
-	else if (bIsLitMaterial 
-		&& IsIndirectLightingCacheAllowed(Parameters.FeatureLevel)
-		&& Parameters.PrimitiveSceneProxy
-		// Movable objects need to get their GI from the indirect lighting cache
-		&& Parameters.PrimitiveSceneProxy->IsMovable())
-	{
-		Action.template Process<FSimpleDirectionalLightAndSHIndirectPolicy>(RHICmdList, Parameters,FSimpleDirectionalLightAndSHIndirectPolicy(),FSimpleDirectionalLightAndSHIndirectPolicy::ElementDataType(Action.ShouldPackAmbientSH())); 
-	}
 	else
 	{
 		Action.template Process<FNoLightMapPolicy>(RHICmdList, Parameters,FNoLightMapPolicy(),FNoLightMapPolicy::ElementDataType());
diff --git a/Engine/Source/Runtime/Renderer/Private/ForwardShadingRenderer.cpp b/Engine/Source/Runtime/Renderer/Private/ForwardShadingRenderer.cpp
--- a/Engine/Source/Runtime/Renderer/Private/ForwardShadingRenderer.cpp
+++ b/Engine/Source/Runtime/Renderer/Private/ForwardShadingRenderer.cpp
@@ -13,6 +13,9 @@
 #include "PostProcessMobile.h"
 #include "SceneUtils.h"
 
+uint32 GetShadowQuality();
+
+
 FForwardShadingSceneRenderer::FForwardShadingSceneRenderer(const FSceneViewFamily* InViewFamily,FHitProxyConsumer* HitProxyConsumer)
 	:	FSceneRenderer(InViewFamily, HitProxyConsumer)
 {
@@ -32,6 +35,38 @@ void FForwardShadingSceneRenderer::InitViews(FRHICommandListImmediate& RHICmdLis
 	ComputeViewVisibility(RHICmdList);
 	PostVisibilityFrameSetup();
 
+	bool bDynamicShadows = ViewFamily.EngineShowFlags.DynamicShadows && GetShadowQuality() > 0;
+
+	if (bDynamicShadows && !IsSimpleDynamicLightingEnabled())
+	{
+		// Setup dynamic shadows.
+		InitDynamicShadows(RHICmdList);		
+	}
+
+	// initialize per-view uniform buffer.  Pass in shadow info as necessary.
+	for (int32 ViewIndex = 0; ViewIndex < Views.Num(); ViewIndex++)
+	{
+		const TArray<FProjectedShadowInfo*, SceneRenderingAllocator>* DirectionalLightShadowInfo = nullptr;
+
+		FViewInfo& ViewInfo = Views[ViewIndex];
+		FScene* Scene = (FScene*)ViewInfo.Family->Scene;
+		if (bDynamicShadows && Scene->SimpleDirectionalLight)
+		{
+			int32 LightId = Scene->SimpleDirectionalLight->Id;
+			if (VisibleLightInfos.IsValidIndex(LightId))
+			{
+				const FVisibleLightInfo& VisibleLightInfo = VisibleLightInfos[LightId];
+				if (VisibleLightInfo.AllProjectedShadows.Num() > 0)
+				{
+					DirectionalLightShadowInfo = &VisibleLightInfo.AllProjectedShadows;
+				}
+			}
+		}
+
+		// Initialize the view's RHI resources.
+		Views[ViewIndex].InitRHIResources(DirectionalLightShadowInfo);
+	}
+
 	OnStartFrame();
 }
 
@@ -56,6 +91,8 @@ void FForwardShadingSceneRenderer::Render(FRHICommandListImmediate& RHICmdList)
 	// Find the visible primitives.
 	InitViews(RHICmdList);
 	
+	RenderShadowDepthMaps(RHICmdList);
+
 	// Notify the FX system that the scene is about to be rendered.
 	if (Scene->FXSystem)
 	{
@@ -149,24 +186,27 @@ void FForwardShadingSceneRenderer::Render(FRHICommandListImmediate& RHICmdList)
 		RHICmdList.DiscardRenderTargets(true, true, 0);
 
 		// Finish rendering for each view, or the full stereo buffer if enabled
-		if (GEngine->IsStereoscopic3D())
+		if (ViewFamily.bResolveScene)
 		{
-			check(Views.Num() > 1);
-
-			//@todo ES2 stereo post: until we get proper stereo postprocessing for ES2, process the stereo buffer as one view
-			FIntPoint OriginalMax0 = Views[0].ViewRect.Max;
-			Views[0].ViewRect.Max = Views[1].ViewRect.Max;
-			GPostProcessing.ProcessES2(RHICmdList, Views[0], bOnChipSunMask);
-			Views[0].ViewRect.Max = OriginalMax0;
-		}
-		else
-		{
-			SCOPED_DRAW_EVENT(RHICmdList, PostProcessing, DEC_SCENE_ITEMS);
-			SCOPE_CYCLE_COUNTER(STAT_FinishRenderViewTargetTime);
-			for(int32 ViewIndex = 0;ViewIndex < Views.Num();ViewIndex++)
-			{	
-				SCOPED_CONDITIONAL_DRAW_EVENTF(RHICmdList, EventView, Views.Num() > 1, DEC_SCENE_ITEMS, TEXT("View%d"), ViewIndex);
-				GPostProcessing.ProcessES2(RHICmdList, Views[ViewIndex], bOnChipSunMask);
+			if (GEngine->IsStereoscopic3D())
+			{
+				check(Views.Num() > 1);
+
+				//@todo ES2 stereo post: until we get proper stereo postprocessing for ES2, process the stereo buffer as one view
+				FIntPoint OriginalMax0 = Views[0].ViewRect.Max;
+				Views[0].ViewRect.Max = Views[1].ViewRect.Max;
+				GPostProcessing.ProcessES2(RHICmdList, Views[0], bOnChipSunMask);
+				Views[0].ViewRect.Max = OriginalMax0;
+			}
+			else
+			{
+				SCOPED_DRAW_EVENT(RHICmdList, PostProcessing, DEC_SCENE_ITEMS);
+				SCOPE_CYCLE_COUNTER(STAT_FinishRenderViewTargetTime);
+				for(int32 ViewIndex = 0;ViewIndex < Views.Num();ViewIndex++)
+				{	
+					SCOPED_CONDITIONAL_DRAW_EVENTF(RHICmdList, EventView, Views.Num() > 1, DEC_SCENE_ITEMS, TEXT("View%d"), ViewIndex);
+					GPostProcessing.ProcessES2(RHICmdList, Views[ViewIndex], bOnChipSunMask);
+				}
 			}
 		}
 	}
diff --git a/Engine/Source/Runtime/Renderer/Private/ForwardTranslucentRendering.cpp b/Engine/Source/Runtime/Renderer/Private/ForwardTranslucentRendering.cpp
--- a/Engine/Source/Runtime/Renderer/Private/ForwardTranslucentRendering.cpp
+++ b/Engine/Source/Runtime/Renderer/Private/ForwardTranslucentRendering.cpp
@@ -128,6 +128,7 @@ public:
 		) const
 	{
 		const bool bIsLitMaterial = Parameters.ShadingModel != MSM_Unlit;
+		const FScene* Scene = Parameters.PrimitiveSceneProxy ? Parameters.PrimitiveSceneProxy->GetPrimitiveSceneInfo()->Scene : NULL;
 
 		TBasePassForForwardShadingDrawingPolicy<LightMapPolicyType> DrawingPolicy(
 			Parameters.Mesh.VertexFactory,
@@ -136,7 +137,9 @@ public:
 			LightMapPolicy,
 			Parameters.BlendMode,
 			Parameters.TextureMode,
-			View.Family->EngineShowFlags.ShaderComplexity
+			Parameters.ShadingModel != MSM_Unlit && Scene && Scene->ShouldRenderSkylight(),
+			View.Family->EngineShowFlags.ShaderComplexity,
+			View.GetFeatureLevel()
 			);
 
 		RHICmdList.BuildAndSetLocalBoundShaderState(DrawingPolicy.GetBoundShaderStateInput(View.GetFeatureLevel()));
diff --git a/Engine/Source/Runtime/Renderer/Private/LightMapRendering.cpp b/Engine/Source/Runtime/Renderer/Private/LightMapRendering.cpp
new file mode 100644
--- /dev/null
+++ b/Engine/Source/Runtime/Renderer/Private/LightMapRendering.cpp
@@ -0,0 +1,17 @@
+// Copyright 1998-2014 Epic Games, Inc. All Rights Reserved.
+
+/*=============================================================================
+LightMapRendering.cpp: Light map rendering implementations.
+=============================================================================*/
+
+#include "RendererPrivate.h"
+#include "SceneManagement.h"
+
+void FMovableDirectionalLightCSMLightingPolicy::ModifyCompilationEnvironment(EShaderPlatform Platform, const FMaterial* Material, FShaderCompilerEnvironment& OutEnvironment)
+{
+	OutEnvironment.SetDefine(TEXT("MOVABLE_DIRECTIONAL_LIGHT"), TEXT("1"));
+	OutEnvironment.SetDefine(TEXT("MOVABLE_DIRECTIONAL_LIGHT_CSM"), TEXT("1"));
+	OutEnvironment.SetDefine(TEXT("MAX_FORWARD_SHADOWCASCADES"), TEXT(PREPROCESSOR_TO_STRING(MAX_FORWARD_SHADOWCASCADES)));
+
+	FNoLightMapPolicy::ModifyCompilationEnvironment(Platform, Material, OutEnvironment);
+}
\ No newline at end of file
diff --git a/Engine/Source/Runtime/Renderer/Private/LightMapRendering.h b/Engine/Source/Runtime/Renderer/Private/LightMapRendering.h
--- a/Engine/Source/Runtime/Renderer/Private/LightMapRendering.h
+++ b/Engine/Source/Runtime/Renderer/Private/LightMapRendering.h
@@ -967,4 +967,31 @@ private:
 	FCachedPointIndirectLightingPolicy CachedPointIndirectLightingPolicy;
 };
 
+class FMovableDirectionalLightLightingPolicy : public FNoLightMapPolicy
+{
+public:
+
+	static bool ShouldCache(EShaderPlatform Platform,const FMaterial* Material,const FVertexFactoryType* VertexFactoryType)
+	{
+		return Material->GetShadingModel() != MSM_Unlit;
+	}
+
+	static void ModifyCompilationEnvironment(EShaderPlatform Platform, const FMaterial* Material, FShaderCompilerEnvironment& OutEnvironment)
+	{
+		OutEnvironment.SetDefine(TEXT("MOVABLE_DIRECTIONAL_LIGHT"),TEXT("1"));
+		FNoLightMapPolicy::ModifyCompilationEnvironment(Platform, Material, OutEnvironment);
+	}
+};
+
+class FMovableDirectionalLightCSMLightingPolicy : public FNoLightMapPolicy
+{
+public:
+	static bool ShouldCache(EShaderPlatform Platform, const FMaterial* Material, const FVertexFactoryType* VertexFactoryType)
+	{
+		return Material->GetShadingModel() != MSM_Unlit;
+	}	
+
+	static void ModifyCompilationEnvironment(EShaderPlatform Platform, const FMaterial* Material, FShaderCompilerEnvironment& OutEnvironment);
+};
+
 #endif // __LIGHTMAPRENDERING_H__
diff --git a/Engine/Source/Runtime/Renderer/Private/LightRendering.cpp b/Engine/Source/Runtime/Renderer/Private/LightRendering.cpp
--- a/Engine/Source/Runtime/Renderer/Private/LightRendering.cpp
+++ b/Engine/Source/Runtime/Renderer/Private/LightRendering.cpp
@@ -223,7 +223,7 @@ IMPLEMENT_SHADER_TYPE(template<>, TDeferredLightOverlapPS<true>, TEXT("Stationar
 IMPLEMENT_SHADER_TYPE(template<>, TDeferredLightOverlapPS<false>, TEXT("StationaryLightOverlapShaders"), TEXT("OverlapDirectionalPixelMain"), SF_Pixel);
 
 /** Gathers simple lights from visible primtives in the passed in views. */
-void GatherSimpleLights(const FSceneViewFamily& ViewFamily, const TArray<FViewInfo>& Views, FSimpleLightArray& SimpleLights)
+void FSceneRenderer::GatherSimpleLights(const FSceneViewFamily& ViewFamily, const TArray<FViewInfo>& Views, FSimpleLightArray& SimpleLights)
 {
 	TArray<const FPrimitiveSceneInfo*, SceneRenderingAllocator> PrimitivesWithSimpleLights;
 
@@ -254,7 +254,7 @@ void GatherSimpleLights(const FSceneViewFamily& ViewFamily, const TArray<FViewIn
 }
 
 /** Gets a readable light name for use with a draw event. */
-void GetLightNameForDrawEvent(const FLightSceneProxy* LightProxy, FString& LightNameWithLevel)
+void FSceneRenderer::GetLightNameForDrawEvent(const FLightSceneProxy* LightProxy, FString& LightNameWithLevel)
 {
 #if WANTS_DRAW_MESH_EVENTS
 	if (GEmitDrawEvents)
@@ -305,9 +305,7 @@ void FDeferredShadingSceneRenderer::RenderLights(FRHICommandListImmediate& RHICm
 		const FLightSceneInfoCompact& LightSceneInfoCompact = *LightIt;
 		const FLightSceneInfo* const LightSceneInfo = LightSceneInfoCompact.LightSceneInfo;
 
-		if (!LightSceneInfoCompact.Color.IsAlmostBlack()
-			// Only render lights with dynamic lighting or unbuilt static lights
-			&& (!LightSceneInfo->Proxy->HasStaticLighting() || !LightSceneInfo->bPrecomputedLightingIsValid)
+		if (LightSceneInfo->ShouldRenderLightViewIndependent()
 			// Reflection override skips direct specular because it tends to be blindingly bright with a perfectly smooth surface
 			&& !ViewFamily.EngineShowFlags.ReflectionOverride)
 		{
@@ -885,4 +883,3 @@ void FDeferredShadingSceneRenderer::RenderSimpleLightsStandardDeferred(FRHIComma
 		}
 	}
 }
-
diff --git a/Engine/Source/Runtime/Renderer/Private/LightSceneInfo.h b/Engine/Source/Runtime/Renderer/Private/LightSceneInfo.h
--- a/Engine/Source/Runtime/Renderer/Private/LightSceneInfo.h
+++ b/Engine/Source/Runtime/Renderer/Private/LightSceneInfo.h
@@ -183,6 +183,30 @@ public:
 
 	bool ShouldRenderLight(const FViewInfo& View) const;
 
+	/** Encapsulates all View-Independent reasons to have this light render. */
+	bool ShouldRenderLightViewIndependent() const
+	{
+		return !Proxy->GetColor().IsAlmostBlack()
+			// Only render lights with dynamic lighting or unbuilt static lights
+			&& (!Proxy->HasStaticLighting() || !bPrecomputedLightingIsValid);
+	}
+
+	/** Encapsulates all View-Independent reasons to render ViewIndependentWholeSceneShadows for this light */
+	bool ShouldRenderViewIndependentWholeSceneShadows() const
+	{
+		bool bShouldRenderLight = ShouldRenderLightViewIndependent();
+		bool bCastDynamicShadow = Proxy->CastsDynamicShadow();
+		
+		// Also create a whole scene shadow for lights with precomputed shadows that are unbuilt
+		const bool bCreateShadowToPreviewStaticLight =
+			Proxy->HasStaticShadowing()
+			&& bCastDynamicShadow
+			&& !bPrecomputedLightingIsValid;
+
+		bool bShouldRenderShadow = bShouldRenderLight && bCastDynamicShadow && (!Proxy->HasStaticLighting() || bCreateShadowToPreviewStaticLight);
+		return bShouldRenderShadow;
+	}
+
 	/** Hash function. */
 	friend uint32 GetTypeHash(const FLightSceneInfo* LightSceneInfo)
 	{
diff --git a/Engine/Source/Runtime/Renderer/Private/PostProcess/PostProcessTemporalAA.cpp b/Engine/Source/Runtime/Renderer/Private/PostProcess/PostProcessTemporalAA.cpp
--- a/Engine/Source/Runtime/Renderer/Private/PostProcess/PostProcessTemporalAA.cpp
+++ b/Engine/Source/Runtime/Renderer/Private/PostProcess/PostProcessTemporalAA.cpp
@@ -640,7 +640,7 @@ void FRCPassPostProcessTemporalAA::Process(FRenderingCompositePassContext& Conte
 	View.ViewMatrices.TranslatedViewProjectionMatrix = TranslatedViewMatrix * View.ViewMatrices.ProjMatrix;
 	View.ViewMatrices.InvTranslatedViewProjectionMatrix = View.ViewMatrices.TranslatedViewProjectionMatrix.Inverse();
 
-	View.InitRHIResources();
+	View.InitRHIResources(nullptr);
 }
 
 FPooledRenderTargetDesc FRCPassPostProcessTemporalAA::ComputeOutputDesc(EPassOutputId InPassOutputId) const
diff --git a/Engine/Source/Runtime/Renderer/Private/PostProcess/SceneRenderTargets.cpp b/Engine/Source/Runtime/Renderer/Private/PostProcess/SceneRenderTargets.cpp
--- a/Engine/Source/Runtime/Renderer/Private/PostProcess/SceneRenderTargets.cpp
+++ b/Engine/Source/Runtime/Renderer/Private/PostProcess/SceneRenderTargets.cpp
@@ -1122,16 +1122,13 @@ static TCHAR* const GetTranslucencyShadowTransmissionName(uint32 Id)
 
 void FSceneRenderTargets::AllocateReflectionTargets()
 {	
-	if (CurrentFeatureLevel >= ERHIFeatureLevel::ES2 && GSupportsRenderTargetFormat_PF_FloatRGBA)
+	if (GSupportsRenderTargetFormat_PF_FloatRGBA)
 	{
 		extern ENGINE_API int32 GReflectionCaptureSize;
 		const int32 NumReflectionCaptureMips = FMath::CeilLogTwo(GReflectionCaptureSize) + 1;
 		
-		uint32 CubeTexFlags = 0;
-		if (!GSupportsGSRenderTargetLayerSwitchingToMips)
-		{
-			CubeTexFlags = TexCreate_TargetArraySlicesIndependently;
-		}
+		// We write to these cubemap faces individually during filtering
+		uint32 CubeTexFlags = TexCreate_TargetArraySlicesIndependently;
 
 		{
 			// Create scratch cubemaps for filtering passes
diff --git a/Engine/Source/Runtime/Renderer/Private/PostProcess/VisualizeTexture.cpp b/Engine/Source/Runtime/Renderer/Private/PostProcess/VisualizeTexture.cpp
--- a/Engine/Source/Runtime/Renderer/Private/PostProcess/VisualizeTexture.cpp
+++ b/Engine/Source/Runtime/Renderer/Private/PostProcess/VisualizeTexture.cpp
@@ -24,7 +24,7 @@ public:
 
 	static bool ShouldCache(EShaderPlatform Platform)
 	{
-		return IsFeatureLevelSupported(Platform, ERHIFeatureLevel::ES2);
+		return true;
 	}
 
 	static void ModifyCompilationEnvironment(EShaderPlatform Platform, FShaderCompilerEnvironment& OutEnvironment)
@@ -152,8 +152,7 @@ class FVisualizeTexturePresentPS : public FGlobalShader
 
 	static bool ShouldCache(EShaderPlatform Platform)
 	{
-		return IsFeatureLevelSupported(Platform, ERHIFeatureLevel::ES2);
-
+		return true;
 	}
 
 	/** Default constructor. */
@@ -909,7 +908,7 @@ IPooledRenderTarget* FVisualizeTexture::GetObservedElement() const
 void FVisualizeTexture::OnStartFrame(const FSceneView& View)
 {
 	FeatureLevel = View.GetFeatureLevel();
-	bEnabled = FeatureLevel >= ERHIFeatureLevel::ES2;
+	bEnabled = true;
 	ViewRect = View.UnscaledViewRect;
 	AspectRatioConstrainedViewRect = View.Family->EngineShowFlags.CameraAspectRatioBars ? View.CameraConstrainedViewRect : ViewRect;
 
diff --git a/Engine/Source/Runtime/Renderer/Private/ReflectionEnvironmentCapture.cpp b/Engine/Source/Runtime/Renderer/Private/ReflectionEnvironmentCapture.cpp
--- a/Engine/Source/Runtime/Renderer/Private/ReflectionEnvironmentCapture.cpp
+++ b/Engine/Source/Runtime/Renderer/Private/ReflectionEnvironmentCapture.cpp
@@ -75,8 +75,6 @@ FSceneRenderTargetItem& GetEffectiveSourceTexture(bool bDownsamplePass, int32 Ta
 	return GSceneRenderTargets.ReflectionColorScratchCubemap[ScratchTextureIndex]->GetRenderTargetItem();
 }
 
-IMPLEMENT_SHADER_TYPE(,FDownsampleGS,TEXT("ReflectionEnvironmentShaders"),TEXT("DownsampleGS"),SF_Geometry);
-
 class FDownsamplePS : public FGlobalShader
 {
 	DECLARE_SHADER_TYPE(FDownsamplePS,Global);
@@ -84,7 +82,7 @@ public:
 
 	static bool ShouldCache(EShaderPlatform Platform)
 	{
-		return IsFeatureLevelSupported(Platform, ERHIFeatureLevel::SM4);
+		return true;
 	}
 
 	FDownsamplePS(const ShaderMetaType::CompiledShaderInitializerType& Initializer):
@@ -139,7 +137,7 @@ public:
 
 	static bool ShouldCache(EShaderPlatform Platform)
 	{
-		return IsFeatureLevelSupported(Platform, ERHIFeatureLevel::SM4);
+		return true;
 	}
 
 	static void ModifyCompilationEnvironment(EShaderPlatform Platform, FShaderCompilerEnvironment& OutEnvironment)
@@ -213,7 +211,7 @@ public:
 
 	static bool ShouldCache(EShaderPlatform Platform)
 	{
-		return IsFeatureLevelSupported(Platform, ERHIFeatureLevel::SM4);
+		return true;
 	}
 
 	static void ModifyCompilationEnvironment(EShaderPlatform Platform, FShaderCompilerEnvironment& OutEnvironment)
@@ -329,10 +327,10 @@ void FilterReflectionEnvironment(FRHICommandListImmediate& RHICmdList, ERHIFeatu
 			FSceneRenderTargetItem& EffectiveRT = GetEffectiveRenderTarget(true, MipIndex);
 			FSceneRenderTargetItem& EffectiveSource = GetEffectiveSourceTexture(true, MipIndex);
 			check(EffectiveRT.TargetableTexture != EffectiveSource.ShaderResourceTexture);
-
-			if (GSupportsGSRenderTargetLayerSwitchingToMips)
+			
+			for (int32 CubeFace = 0; CubeFace < CubeFace_MAX; CubeFace++)
 			{
-				SetRenderTarget(RHICmdList, EffectiveRT.TargetableTexture, MipIndex, NULL);
+				SetRenderTarget(RHICmdList, EffectiveRT.TargetableTexture, MipIndex, CubeFace, NULL);
 
 				const FIntRect ViewRect(0, 0, MipSize, MipSize);
 				RHICmdList.SetViewport(0, 0, 0.0f, MipSize, MipSize, 1.0f);
@@ -340,65 +338,15 @@ void FilterReflectionEnvironment(FRHICommandListImmediate& RHICmdList, ERHIFeatu
 				RHICmdList.SetDepthStencilState(TStaticDepthStencilState<false, CF_Always>::GetRHI());
 				RHICmdList.SetBlendState(TStaticBlendState<>::GetRHI());
 
-				TShaderMapRef<FScreenVSForGS> VertexShader(ShaderMap);
-				TShaderMapRef<FDownsampleGS> GeometryShader(ShaderMap);
-				TShaderMapRef<FDownsamplePS> PixelShader(ShaderMap);
-
-				SetGlobalBoundShaderState(RHICmdList, FeatureLevel, DownsampleBoundShaderState, GFilterVertexDeclaration.VertexDeclarationRHI, *VertexShader, *PixelShader, *GeometryShader);
-
-				// Draw each face with a geometry shader that routes to the correct render target slice
-				for (int32 CubeFace = 0; CubeFace < CubeFace_MAX; CubeFace++)
-				{
-					SCOPED_DRAW_EVENT(RHICmdList, DownsampleCubeFace, DEC_SCENE_ITEMS);
-					GeometryShader->SetParameters(RHICmdList, CubeFace);
-
-					PixelShader->SetParameters(RHICmdList, CubeFace, SourceMipIndex, EffectiveSource);
-
-					DrawRectangle( 
-						RHICmdList,
-						ViewRect.Min.X, ViewRect.Min.Y, 
-						ViewRect.Width(), ViewRect.Height(),
-						ViewRect.Min.X, ViewRect.Min.Y, 
-						ViewRect.Width(), ViewRect.Height(),
-						FIntPoint(ViewRect.Width(), ViewRect.Height()),
-						FIntPoint(MipSize, MipSize),
-						*VertexShader,
-						EDRF_UseTriangleOptimization);
-
-					RHICmdList.CopyToResolveTarget(EffectiveRT.TargetableTexture, EffectiveRT.ShaderResourceTexture, true, FResolveParams(FResolveRect(), (ECubeFace)CubeFace, MipIndex));
-				}
-			}
-			else
-			{
-				for (int32 CubeFace = 0; CubeFace < CubeFace_MAX; CubeFace++)
-				{
-					SetRenderTarget(RHICmdList, EffectiveRT.TargetableTexture, MipIndex, CubeFace, NULL);
-
-					const FIntRect ViewRect(0, 0, MipSize, MipSize);
-					RHICmdList.SetViewport(0, 0, 0.0f, MipSize, MipSize, 1.0f);
-					RHICmdList.SetRasterizerState(TStaticRasterizerState<FM_Solid, CM_None>::GetRHI());
-					RHICmdList.SetDepthStencilState(TStaticDepthStencilState<false, CF_Always>::GetRHI());
-					RHICmdList.SetBlendState(TStaticBlendState<>::GetRHI());
-
-					TShaderMapRef<FScreenVSForGS> VertexShader(ShaderMap);
-					TShaderMapRef<FDownsamplePS> PixelShader(ShaderMap);
+				TShaderMapRef<FScreenVS> VertexShader(GetGlobalShaderMap(FeatureLevel));
+				TShaderMapRef<FDownsamplePS> PixelShader(GetGlobalShaderMap(FeatureLevel));
 
-					SetGlobalBoundShaderState(RHICmdList, FeatureLevel, DownsampleBoundShaderState, GFilterVertexDeclaration.VertexDeclarationRHI, *VertexShader, *PixelShader);
+				SetGlobalBoundShaderState(RHICmdList, FeatureLevel, DownsampleBoundShaderState, GFilterVertexDeclaration.VertexDeclarationRHI, *VertexShader, *PixelShader);
 
-					PixelShader->SetParameters(RHICmdList, CubeFace, SourceMipIndex, EffectiveSource);
+				PixelShader->SetParameters(RHICmdList, CubeFace, SourceMipIndex, EffectiveSource);
 
-					DrawRectangle( 
-						RHICmdList,
-						ViewRect.Min.X, ViewRect.Min.Y, 
-						ViewRect.Width(), ViewRect.Height(),
-						ViewRect.Min.X, ViewRect.Min.Y, 
-						ViewRect.Width(), ViewRect.Height(),
-						FIntPoint(ViewRect.Width(), ViewRect.Height()),
-						FIntPoint(MipSize, MipSize),
-						*VertexShader);
 
-					RHICmdList.CopyToResolveTarget(EffectiveRT.TargetableTexture, EffectiveRT.ShaderResourceTexture, true, FResolveParams(FResolveRect(), (ECubeFace)CubeFace, MipIndex));
-				}
+				RHICmdList.CopyToResolveTarget(EffectiveRT.TargetableTexture, EffectiveRT.ShaderResourceTexture, true, FResolveParams(FResolveRect(), (ECubeFace)CubeFace, MipIndex));
 			}
 
 			if (MipSize == GDiffuseIrradianceCubemapSize)
@@ -429,9 +377,9 @@ void FilterReflectionEnvironment(FRHICommandListImmediate& RHICmdList, ERHIFeatu
 			check(EffectiveRT.TargetableTexture != EffectiveSource.ShaderResourceTexture);
 			const int32 MipSize = 1 << (NumMips - MipIndex - 1);
 
-			if (GSupportsGSRenderTargetLayerSwitchingToMips)
+			for (int32 CubeFace = 0; CubeFace < CubeFace_MAX; CubeFace++)
 			{
-				SetRenderTarget(RHICmdList, EffectiveRT.TargetableTexture, MipIndex, NULL);
+				SetRenderTarget(RHICmdList, EffectiveRT.TargetableTexture, MipIndex, CubeFace, NULL);
 
 				const FIntRect ViewRect(0, 0, MipSize, MipSize);
 				RHICmdList.SetViewport(0, 0, 0.0f, MipSize, MipSize, 1.0f);
@@ -439,8 +387,8 @@ void FilterReflectionEnvironment(FRHICommandListImmediate& RHICmdList, ERHIFeatu
 				RHICmdList.SetDepthStencilState(TStaticDepthStencilState<false, CF_Always>::GetRHI());
 				RHICmdList.SetBlendState(TStaticBlendState<>::GetRHI());
 
-				TShaderMapRef<FScreenVSForGS> VertexShader(ShaderMap);
-				TShaderMapRef<FDownsampleGS> GeometryShader(ShaderMap);
+				TShaderMapRef<FScreenVS> VertexShader(GetGlobalShaderMap(FeatureLevel));
+				TShaderMapRef< TCubeFilterPS<1> > CaptureCubemapArrayPixelShader(GetGlobalShaderMap(FeatureLevel));
 
 				FCubeFilterPS* PixelShader;
 				if( bNormalize )
@@ -448,83 +396,30 @@ void FilterReflectionEnvironment(FRHICommandListImmediate& RHICmdList, ERHIFeatu
 					PixelShader = *TShaderMapRef< TCubeFilterPS<1> >(ShaderMap);
 					static FGlobalBoundShaderState BoundShaderState;
 					
-					SetGlobalBoundShaderState(RHICmdList, FeatureLevel, BoundShaderState, GFilterVertexDeclaration.VertexDeclarationRHI, *VertexShader, PixelShader, *GeometryShader);
+					SetGlobalBoundShaderState(RHICmdList, FeatureLevel, BoundShaderState, GFilterVertexDeclaration.VertexDeclarationRHI, *VertexShader, PixelShader);
 				}
 				else
 				{
 					PixelShader = *TShaderMapRef< TCubeFilterPS<0> >(ShaderMap);
 					static FGlobalBoundShaderState BoundShaderState;
-					
-					SetGlobalBoundShaderState(RHICmdList, FeatureLevel, BoundShaderState, GFilterVertexDeclaration.VertexDeclarationRHI, *VertexShader, PixelShader, *GeometryShader);
-				}
-
-				for (int32 CubeFace = 0; CubeFace < CubeFace_MAX; CubeFace++)
-				{
-					SCOPED_DRAW_EVENT(RHICmdList, FilterCubeFace, DEC_SCENE_ITEMS);
-					GeometryShader->SetParameters(RHICmdList, CubeFace);
-
-					PixelShader->SetParameters(RHICmdList, CubeFace, MipIndex, EffectiveSource);
-
-					DrawRectangle( 
-						RHICmdList,
-						ViewRect.Min.X, ViewRect.Min.Y, 
-						ViewRect.Width(), ViewRect.Height(),
-						ViewRect.Min.X, ViewRect.Min.Y, 
-						ViewRect.Width(), ViewRect.Height(),
-						FIntPoint(ViewRect.Width(), ViewRect.Height()),
-						FIntPoint(MipSize, MipSize),
-						*VertexShader,
-						EDRF_UseTriangleOptimization);
-
-					RHICmdList.CopyToResolveTarget(EffectiveRT.TargetableTexture, EffectiveRT.ShaderResourceTexture, true, FResolveParams(FResolveRect(), (ECubeFace)CubeFace, MipIndex));
-				}
-			}
-			else
-			{
-				for (int32 CubeFace = 0; CubeFace < CubeFace_MAX; CubeFace++)
-				{
-					SetRenderTarget(RHICmdList, EffectiveRT.TargetableTexture, MipIndex, CubeFace, NULL);
-
-					const FIntRect ViewRect(0, 0, MipSize, MipSize);
-					RHICmdList.SetViewport(0, 0, 0.0f, MipSize, MipSize, 1.0f);
-					RHICmdList.SetRasterizerState(TStaticRasterizerState<FM_Solid, CM_None>::GetRHI());
-					RHICmdList.SetDepthStencilState(TStaticDepthStencilState<false, CF_Always>::GetRHI());
-					RHICmdList.SetBlendState(TStaticBlendState<>::GetRHI());
-
-					TShaderMapRef<FScreenVSForGS> VertexShader(ShaderMap);
-					TShaderMapRef< TCubeFilterPS<1> > CaptureCubemapArrayPixelShader(ShaderMap);
-					//FFilterPS* PixelShader = (FFilterPS*)*CaptureCubemapArrayPixelShader;
-
-					FCubeFilterPS* PixelShader;
-					if( bNormalize )
-					{
-						PixelShader = *TShaderMapRef< TCubeFilterPS<1> >(ShaderMap);
-						static FGlobalBoundShaderState BoundShaderState;
-						
-						SetGlobalBoundShaderState(RHICmdList, FeatureLevel, BoundShaderState, GFilterVertexDeclaration.VertexDeclarationRHI, *VertexShader, PixelShader);
-					}
-					else
-					{
-						PixelShader = *TShaderMapRef< TCubeFilterPS<0> >(ShaderMap);
-						static FGlobalBoundShaderState BoundShaderState;
+					SetGlobalBoundShaderState(RHICmdList, FeatureLevel, BoundShaderState, GFilterVertexDeclaration.VertexDeclarationRHI, *VertexShader, PixelShader);
 						
-						SetGlobalBoundShaderState(RHICmdList, FeatureLevel, BoundShaderState, GFilterVertexDeclaration.VertexDeclarationRHI, *VertexShader, PixelShader);
-					}
+					SetGlobalBoundShaderState(RHICmdList, FeatureLevel, BoundShaderState, GFilterVertexDeclaration.VertexDeclarationRHI, *VertexShader, PixelShader);
+				}
 
-					PixelShader->SetParameters(RHICmdList, CubeFace, MipIndex, EffectiveSource);
+				PixelShader->SetParameters(RHICmdList, CubeFace, MipIndex, EffectiveSource);
 
-					DrawRectangle( 
-						RHICmdList,
-						ViewRect.Min.X, ViewRect.Min.Y, 
-						ViewRect.Width(), ViewRect.Height(),
-						ViewRect.Min.X, ViewRect.Min.Y, 
-						ViewRect.Width(), ViewRect.Height(),
-						FIntPoint(ViewRect.Width(), ViewRect.Height()),
-						FIntPoint(MipSize, MipSize),
-						*VertexShader);
+				DrawRectangle( 
+					RHICmdList,
+					ViewRect.Min.X, ViewRect.Min.Y, 
+					ViewRect.Width(), ViewRect.Height(),
+					ViewRect.Min.X, ViewRect.Min.Y, 
+					ViewRect.Width(), ViewRect.Height(),
+					FIntPoint(ViewRect.Width(), ViewRect.Height()),
+					FIntPoint(MipSize, MipSize),
+					*VertexShader);
 
-					RHICmdList.CopyToResolveTarget(EffectiveRT.TargetableTexture, EffectiveRT.ShaderResourceTexture, true, FResolveParams(FResolveRect(), (ECubeFace)CubeFace, MipIndex));
-				}
+				RHICmdList.CopyToResolveTarget(EffectiveRT.TargetableTexture, EffectiveRT.ShaderResourceTexture, true, FResolveParams(FResolveRect(), (ECubeFace)CubeFace, MipIndex));
 			}
 		}
 	}
@@ -538,7 +433,7 @@ public:
 
 	static bool ShouldCache(EShaderPlatform Platform)
 	{
-		return IsFeatureLevelSupported(Platform, ERHIFeatureLevel::SM4);
+		return true;
 	}
 
 	FCopyToCubeFaceVS()	{}
@@ -561,42 +456,6 @@ public:
 
 IMPLEMENT_SHADER_TYPE(,FCopyToCubeFaceVS,TEXT("ReflectionEnvironmentShaders"),TEXT("CopyToCubeFaceVS"),SF_Vertex);
 
-/** Geometry shader used when writing to a cubemap face.  Supports writing to a cubemap array as well. */
-class FRouteToCubeFaceGS : public FGlobalShader
-{
-	DECLARE_SHADER_TYPE(FRouteToCubeFaceGS,Global);
-public:
-
-	static bool ShouldCache(EShaderPlatform Platform)
-	{
-		return IsFeatureLevelSupported(Platform, ERHIFeatureLevel::SM4);
-	}
-
-	FRouteToCubeFaceGS(const ShaderMetaType::CompiledShaderInitializerType& Initializer):
-		FGlobalShader(Initializer)
-	{
-		CubeFace.Bind(Initializer.ParameterMap,TEXT("CubeFace"));
-	}
-	FRouteToCubeFaceGS() {}
-
-	void SetParameters(FRHICommandList& RHICmdList, uint32 CubeFaceValue)
-	{
-		SetShaderValue(RHICmdList, GetGeometryShader(), CubeFace, CubeFaceValue);
-	}
-
-	virtual bool Serialize(FArchive& Ar)
-	{
-		bool bShaderHasOutdatedParameters = FGlobalShader::Serialize(Ar);
-		Ar << CubeFace;
-		return bShaderHasOutdatedParameters;
-	}
-
-private:
-	FShaderParameter CubeFace;
-};
-
-IMPLEMENT_SHADER_TYPE(,FRouteToCubeFaceGS,TEXT("ReflectionEnvironmentShaders"),TEXT("RouteToCubeFaceGS"),SF_Geometry);
-
 /** Pixel shader used when copying scene color from a scene render into a face of a reflection capture cubemap. */
 class FCopySceneColorToCubeFacePS : public FGlobalShader
 {
@@ -605,7 +464,7 @@ public:
 
 	static bool ShouldCache(EShaderPlatform Platform)
 	{
-		return IsFeatureLevelSupported(Platform, ERHIFeatureLevel::SM4);
+		return true;
 	}
 
 	static void ModifyCompilationEnvironment(EShaderPlatform Platform, FShaderCompilerEnvironment& OutEnvironment)
@@ -690,7 +549,7 @@ public:
 
 	static bool ShouldCache(EShaderPlatform Platform)
 	{
-		return IsFeatureLevelSupported(Platform, ERHIFeatureLevel::SM4);
+		return true;
 	}
 
 	FCopyCubemapToCubeFacePS(const ShaderMetaType::CompiledShaderInitializerType& Initializer):
@@ -788,38 +647,22 @@ void ClearScratchCubemaps(FRHICommandList& RHICmdList)
 
 	for (int32 MipIndex = 0; MipIndex < NumMips; MipIndex++)
 	{
-		if (GSupportsGSRenderTargetLayerSwitchingToMips)
+		for (int32 CubeFace = 0; CubeFace < CubeFace_MAX; CubeFace++)
 		{
-			SetRenderTarget(RHICmdList, RT0.TargetableTexture, MipIndex, NULL);
+			SetRenderTarget(RHICmdList, RT0.TargetableTexture, MipIndex, CubeFace, NULL);
 			RHICmdList.Clear(true, FLinearColor(0, 10000, 0, 0), false, 0, false, 0, FIntRect());
 		}
-		else
-		{
-			for (int32 CubeFace = 0; CubeFace < CubeFace_MAX; CubeFace++)
-			{
-				SetRenderTarget(RHICmdList, RT0.TargetableTexture, MipIndex, CubeFace, NULL);
-				RHICmdList.Clear(true, FLinearColor(0, 10000, 0, 0), false, 0, false, 0, FIntRect());
-			}
-		}
 	}
 
 	FSceneRenderTargetItem& RT1 = GSceneRenderTargets.ReflectionColorScratchCubemap[1]->GetRenderTargetItem();
 
 	for (int32 MipIndex = 0; MipIndex < NumMips; MipIndex++)
 	{
-		if (GSupportsGSRenderTargetLayerSwitchingToMips)
+		for (int32 CubeFace = 0; CubeFace < CubeFace_MAX; CubeFace++)
 		{
-			SetRenderTarget(RHICmdList, RT1.TargetableTexture, MipIndex, NULL);
+			SetRenderTarget(RHICmdList, RT1.TargetableTexture, MipIndex, CubeFace, NULL);
 			RHICmdList.Clear(true, FLinearColor(0, 10000, 0, 0), false, 0, false, 0, FIntRect());
 		}
-		else
-		{
-			for (int32 CubeFace = 0; CubeFace < CubeFace_MAX; CubeFace++)
-			{
-				SetRenderTarget(RHICmdList, RT1.TargetableTexture, MipIndex, CubeFace, NULL);
-				RHICmdList.Clear(true, FLinearColor(0, 10000, 0, 0), false, 0, false, 0, FIntRect());
-			}
-		}
 	}
 }
 
@@ -857,57 +700,24 @@ void CaptureSceneToScratchCubemap(FRHICommandListImmediate& RHICmdList, FSceneRe
 
 		{
 			SCOPED_DRAW_EVENT(RHICmdList, CubeMapCopyScene, DEC_SCENE_ITEMS);
-			if (GSupportsGSRenderTargetLayerSwitchingToMips)
-			{
+			
 			// Copy the captured scene into the cubemap face
-				SetRenderTarget(RHICmdList, EffectiveColorRT.TargetableTexture, NULL);
+			SetRenderTarget(RHICmdList, EffectiveColorRT.TargetableTexture, 0, CubeFace, NULL);
 
 			const FIntRect ViewRect(0, 0, EffectiveSize, EffectiveSize);
-				RHICmdList.SetViewport(0, 0, 0.0f, EffectiveSize, EffectiveSize, 1.0f);
-				RHICmdList.SetRasterizerState(TStaticRasterizerState<FM_Solid, CM_None>::GetRHI());
-				RHICmdList.SetDepthStencilState(TStaticDepthStencilState<false, CF_Always>::GetRHI());
-				RHICmdList.SetBlendState(TStaticBlendState<>::GetRHI());
-
-				TShaderMapRef<FCopyToCubeFaceVS> VertexShader(ShaderMap);
-				TShaderMapRef<FRouteToCubeFaceGS> GeometryShader(ShaderMap);
-
-				TShaderMapRef<FCopySceneColorToCubeFacePS> PixelShader(ShaderMap);
-				SetGlobalBoundShaderState(RHICmdList, FeatureLevel, CopyColorCubemapBoundShaderState, GFilterVertexDeclaration.VertexDeclarationRHI, *VertexShader, *PixelShader, *GeometryShader);
-
-				PixelShader->SetParameters(RHICmdList, SceneRenderer->Views[0], bCapturingForSkyLight, bLowerHemisphereIsBlack);
-				VertexShader->SetParameters(RHICmdList, SceneRenderer->Views[0]);
-				GeometryShader->SetParameters(RHICmdList, (uint32)CubeFace);
-
-				DrawRectangle(
-					RHICmdList,
-					ViewRect.Min.X, ViewRect.Min.Y,
-					ViewRect.Width(), ViewRect.Height(),
-					ViewRect.Min.X, ViewRect.Min.Y,
-					ViewRect.Width() * GSupersampleCaptureFactor, ViewRect.Height() * GSupersampleCaptureFactor,
-					FIntPoint(ViewRect.Width(), ViewRect.Height()),
-					GSceneRenderTargets.GetBufferSizeXY(),
-					*VertexShader,
-					EDRF_UseTriangleOptimization);
-
-				RHICmdList.CopyToResolveTarget(EffectiveColorRT.TargetableTexture, EffectiveColorRT.ShaderResourceTexture, true, FResolveParams(FResolveRect()));
-			}
-			else
-			{
-				// Copy the captured scene into the cubemap face
-				SetRenderTarget(RHICmdList, EffectiveColorRT.TargetableTexture, 0, CubeFace, NULL);
+			RHICmdList.SetViewport(0, 0, 0.0f, EffectiveSize, EffectiveSize, 1.0f);
+			RHICmdList.SetRasterizerState(TStaticRasterizerState<FM_Solid, CM_None>::GetRHI());
+			RHICmdList.SetDepthStencilState(TStaticDepthStencilState<false, CF_Always>::GetRHI());
+			RHICmdList.SetBlendState(TStaticBlendState<>::GetRHI());
 
-				const FIntRect ViewRect(0, 0, EffectiveSize, EffectiveSize);
-				RHICmdList.SetViewport(0, 0, 0.0f, EffectiveSize, EffectiveSize, 1.0f);
-				RHICmdList.SetRasterizerState(TStaticRasterizerState<FM_Solid, CM_None>::GetRHI());
-				RHICmdList.SetDepthStencilState(TStaticDepthStencilState<false, CF_Always>::GetRHI());
-				RHICmdList.SetBlendState(TStaticBlendState<>::GetRHI());
+			TShaderMapRef<FCopyToCubeFaceVS> VertexShader(GetGlobalShaderMap(FeatureLevel));
+			TShaderMapRef<FCopySceneColorToCubeFacePS> PixelShader(GetGlobalShaderMap(FeatureLevel));
+			SetGlobalBoundShaderState(RHICmdList, FeatureLevel, CopyColorCubemapBoundShaderState, GFilterVertexDeclaration.VertexDeclarationRHI, *VertexShader, *PixelShader);
 
-				TShaderMapRef<FCopyToCubeFaceVS> VertexShader(ShaderMap);
-				TShaderMapRef<FCopySceneColorToCubeFacePS> PixelShader(ShaderMap);
-				SetGlobalBoundShaderState(RHICmdList, FeatureLevel, CopyColorCubemapBoundShaderState, GFilterVertexDeclaration.VertexDeclarationRHI, *VertexShader, *PixelShader);
+			SetGlobalBoundShaderState(RHICmdList, FeatureLevel, CopyColorCubemapBoundShaderState, GFilterVertexDeclaration.VertexDeclarationRHI, *VertexShader, *PixelShader);
 
-				PixelShader->SetParameters(RHICmdList, SceneRenderer->Views[0], bCapturingForSkyLight, bLowerHemisphereIsBlack);
-				VertexShader->SetParameters(RHICmdList, SceneRenderer->Views[0]);
+			PixelShader->SetParameters(RHICmdList, SceneRenderer->Views[0], bCapturingForSkyLight, bLowerHemisphereIsBlack);
+			VertexShader->SetParameters(RHICmdList, SceneRenderer->Views[0]);
 
 				DrawRectangle( 
 					RHICmdList,
@@ -917,10 +727,9 @@ void CaptureSceneToScratchCubemap(FRHICommandListImmediate& RHICmdList, FSceneRe
 					ViewRect.Width() * GSupersampleCaptureFactor, ViewRect.Height() * GSupersampleCaptureFactor,
 					FIntPoint(ViewRect.Width(), ViewRect.Height()),
 					GSceneRenderTargets.GetBufferSizeXY(),
-					*VertexShader);
+				*VertexShader);
 
-				RHICmdList.CopyToResolveTarget(EffectiveColorRT.TargetableTexture, EffectiveColorRT.ShaderResourceTexture, true, FResolveParams(FResolveRect(), CubeFace));
-			}
+			RHICmdList.CopyToResolveTarget(EffectiveColorRT.TargetableTexture, EffectiveColorRT.ShaderResourceTexture, true, FResolveParams(FResolveRect(), CubeFace));			
 		}
 	}
 
@@ -936,12 +745,10 @@ void CopyCubemapToScratchCubemap(FRHICommandList& RHICmdList, ERHIFeatureLevel::
 	const int32 EffectiveSize = GReflectionCaptureSize;
 	FSceneRenderTargetItem& EffectiveColorRT =  GSceneRenderTargets.ReflectionColorScratchCubemap[0]->GetRenderTargetItem();
 
-	auto ShaderMap = GetGlobalShaderMap(FeatureLevel);
-
-	if (GSupportsGSRenderTargetLayerSwitchingToMips)
+	for (uint32 CubeFace = 0; CubeFace < CubeFace_MAX; CubeFace++)
 	{
 		// Copy the captured scene into the cubemap face
-		SetRenderTarget(RHICmdList, EffectiveColorRT.TargetableTexture, NULL);
+		SetRenderTarget(RHICmdList, EffectiveColorRT.TargetableTexture, 0, CubeFace, NULL);
 
 		const FTexture* SourceCubemapResource = SourceCubemap->Resource;
 		const FIntPoint SourceDimensions(SourceCubemapResource->GetSizeX(), SourceCubemapResource->GetSizeY());
@@ -951,62 +758,23 @@ void CopyCubemapToScratchCubemap(FRHICommandList& RHICmdList, ERHIFeatureLevel::
 		RHICmdList.SetDepthStencilState(TStaticDepthStencilState<false, CF_Always>::GetRHI());
 		RHICmdList.SetBlendState(TStaticBlendState<>::GetRHI());
 
-		TShaderMapRef<FScreenVSForGS> VertexShader(ShaderMap);
-		TShaderMapRef<FDownsampleGS> GeometryShader(ShaderMap);
-		TShaderMapRef<FCopyCubemapToCubeFacePS> PixelShader(ShaderMap);
+		TShaderMapRef<FScreenVS> VertexShader(GetGlobalShaderMap(FeatureLevel));
+		TShaderMapRef<FCopyCubemapToCubeFacePS> PixelShader(GetGlobalShaderMap(FeatureLevel));
 
-		for (uint32 CubeFace = 0; CubeFace < CubeFace_MAX; CubeFace++)
-		{
-			SetGlobalBoundShaderState(RHICmdList, FeatureLevel, CopyFromCubemapToCubemapBoundShaderState, GFilterVertexDeclaration.VertexDeclarationRHI, *VertexShader, *PixelShader, *GeometryShader);
-			GeometryShader->SetParameters(RHICmdList, CubeFace);
-			PixelShader->SetParameters(RHICmdList, SourceCubemapResource, CubeFace, bIsSkyLight, bLowerHemisphereIsBlack);
-
-			DrawRectangle( 
-				RHICmdList,
-				ViewRect.Min.X, ViewRect.Min.Y, 
-				ViewRect.Width(), ViewRect.Height(),
-				0, 0, 
-				SourceDimensions.X, SourceDimensions.Y,
-				FIntPoint(ViewRect.Width(), ViewRect.Height()),
-				SourceDimensions,
-				*VertexShader);
-		}
+		SetGlobalBoundShaderState(RHICmdList, FeatureLevel, CopyFromCubemapToCubemapBoundShaderState, GFilterVertexDeclaration.VertexDeclarationRHI, *VertexShader, *PixelShader);
+		PixelShader->SetParameters(RHICmdList, SourceCubemapResource, CubeFace, bIsSkyLight, bLowerHemisphereIsBlack);
 
-		RHICmdList.CopyToResolveTarget(EffectiveColorRT.TargetableTexture, EffectiveColorRT.ShaderResourceTexture, true, FResolveParams(FResolveRect()));
-	}
-	else
-	{
-		for (uint32 CubeFace = 0; CubeFace < CubeFace_MAX; CubeFace++)
-		{
-			// Copy the captured scene into the cubemap face
-			SetRenderTarget(RHICmdList, EffectiveColorRT.TargetableTexture, 0, CubeFace, NULL);
-
-			const FTexture* SourceCubemapResource = SourceCubemap->Resource;
-			const FIntPoint SourceDimensions(SourceCubemapResource->GetSizeX(), SourceCubemapResource->GetSizeY());
-			const FIntRect ViewRect(0, 0, EffectiveSize, EffectiveSize);
-			RHICmdList.SetViewport(0, 0, 0.0f, EffectiveSize, EffectiveSize, 1.0f);
-			RHICmdList.SetRasterizerState(TStaticRasterizerState<FM_Solid, CM_None>::GetRHI());
-			RHICmdList.SetDepthStencilState(TStaticDepthStencilState<false, CF_Always>::GetRHI());
-			RHICmdList.SetBlendState(TStaticBlendState<>::GetRHI());
-
-			TShaderMapRef<FScreenVSForGS> VertexShader(ShaderMap);
-			TShaderMapRef<FCopyCubemapToCubeFacePS> PixelShader(ShaderMap);
-
-			SetGlobalBoundShaderState(RHICmdList, FeatureLevel, CopyFromCubemapToCubemapBoundShaderState, GFilterVertexDeclaration.VertexDeclarationRHI, *VertexShader, *PixelShader);
-			PixelShader->SetParameters(RHICmdList, SourceCubemapResource, CubeFace, bIsSkyLight, bLowerHemisphereIsBlack);
-
-			DrawRectangle( 
-				RHICmdList,
-				ViewRect.Min.X, ViewRect.Min.Y, 
-				ViewRect.Width(), ViewRect.Height(),
-				0, 0, 
-				SourceDimensions.X, SourceDimensions.Y,
-				FIntPoint(ViewRect.Width(), ViewRect.Height()),
-				SourceDimensions,
-				*VertexShader);
-
-			RHICmdList.CopyToResolveTarget(EffectiveColorRT.TargetableTexture, EffectiveColorRT.ShaderResourceTexture, true, FResolveParams(FResolveRect(), (ECubeFace)CubeFace));
-		}
+		DrawRectangle(
+			RHICmdList,
+			ViewRect.Min.X, ViewRect.Min.Y, 
+			ViewRect.Width(), ViewRect.Height(),
+			0, 0, 
+			SourceDimensions.X, SourceDimensions.Y,
+			FIntPoint(ViewRect.Width(), ViewRect.Height()),
+			SourceDimensions,
+			*VertexShader);
+
+		RHICmdList.CopyToResolveTarget(EffectiveColorRT.TargetableTexture, EffectiveColorRT.ShaderResourceTexture, true, FResolveParams(FResolveRect(), (ECubeFace)CubeFace));
 	}
 }
 
@@ -1554,8 +1322,9 @@ void CopyToSkyTexture(FRHICommandList& RHICmdList, FScene* Scene, FTexture* Proc
 }
 
 void FScene::UpdateSkyCaptureContents(const USkyLightComponent* CaptureComponent, bool bCaptureEmissiveOnly, FTexture* OutProcessedTexture, FSHVectorRGB3& OutIrradianceEnvironmentMap)
-{
-	if (GetFeatureLevel() >= ERHIFeatureLevel::SM4)
+{	
+	//todo maybe use ShouldRenderSkylight here, but currently the Skylight hasn't been loaded in time.
+	if (GSupportsRenderTargetFormat_PF_FloatRGBA)
 	{
 		ENQUEUE_UNIQUE_RENDER_COMMAND( 
 			ClearCommand,
diff --git a/Engine/Source/Runtime/Renderer/Private/ReflectionEnvironmentCapture.h b/Engine/Source/Runtime/Renderer/Private/ReflectionEnvironmentCapture.h
--- a/Engine/Source/Runtime/Renderer/Private/ReflectionEnvironmentCapture.h
+++ b/Engine/Source/Runtime/Renderer/Private/ReflectionEnvironmentCapture.h
@@ -9,36 +9,3 @@
 extern ENGINE_API int32 GReflectionCaptureSize;
 
 extern void ComputeDiffuseIrradiance(FRHICommandListImmediate& RHICmdList, ERHIFeatureLevel::Type FeatureLevel, FTextureRHIRef LightingSource, int32 LightingSourceMipIndex, FSHVectorRGB3* OutIrradianceEnvironmentMap);
-
-class FDownsampleGS : public FGlobalShader
-{
-	DECLARE_SHADER_TYPE(FDownsampleGS,Global);
-public:
-
-	static bool ShouldCache(EShaderPlatform Platform)
-	{
-		return IsFeatureLevelSupported(Platform, ERHIFeatureLevel::SM4);
-	}
-
-	FDownsampleGS(const ShaderMetaType::CompiledShaderInitializerType& Initializer):
-		FGlobalShader(Initializer)
-	{
-		CubeFace.Bind(Initializer.ParameterMap,TEXT("CubeFace"));
-	}
-	FDownsampleGS() {}
-
-	void SetParameters(FRHICommandList& RHICmdList, int32 CubeFaceValue)
-	{
-		SetShaderValue(RHICmdList, GetGeometryShader(), CubeFace, CubeFaceValue);
-	}
-
-	virtual bool Serialize(FArchive& Ar)
-	{
-		bool bShaderHasOutdatedParameters = FGlobalShader::Serialize(Ar);
-		Ar << CubeFace;
-		return bShaderHasOutdatedParameters;
-	}
-
-private:
-	FShaderParameter CubeFace;
-};
\ No newline at end of file
diff --git a/Engine/Source/Runtime/Renderer/Private/ReflectionEnvironmentDiffuseIrradiance.cpp b/Engine/Source/Runtime/Renderer/Private/ReflectionEnvironmentDiffuseIrradiance.cpp
--- a/Engine/Source/Runtime/Renderer/Private/ReflectionEnvironmentDiffuseIrradiance.cpp
+++ b/Engine/Source/Runtime/Renderer/Private/ReflectionEnvironmentDiffuseIrradiance.cpp
@@ -38,7 +38,7 @@ public:
 
 	static bool ShouldCache(EShaderPlatform Platform)
 	{
-		return IsFeatureLevelSupported(Platform, ERHIFeatureLevel::SM4);
+		return true;
 	}
 
 	FCopyDiffuseIrradiancePS(const ShaderMetaType::CompiledShaderInitializerType& Initializer):
@@ -116,7 +116,7 @@ public:
 
 	static bool ShouldCache(EShaderPlatform Platform)
 	{
-		return IsFeatureLevelSupported(Platform, ERHIFeatureLevel::SM4);
+		return true;
 	}
 
 	FAccumulateDiffuseIrradiancePS(const ShaderMetaType::CompiledShaderInitializerType& Initializer):
@@ -191,7 +191,7 @@ public:
 
 	static bool ShouldCache(EShaderPlatform Platform)
 	{
-		return IsFeatureLevelSupported(Platform, ERHIFeatureLevel::SM4);
+		return true;
 	}
 
 	FAccumulateCubeFacesPS(const ShaderMetaType::CompiledShaderInitializerType& Initializer):
@@ -245,75 +245,37 @@ void ComputeDiffuseIrradiance(FRHICommandListImmediate& RHICmdList, ERHIFeatureL
 		{
 			const int32 MipIndex = 0;
 			const int32 MipSize = GDiffuseIrradianceCubemapSize;
-			FSceneRenderTargetItem& EffectiveRT = GetEffectiveDiffuseIrradianceRenderTarget(MipIndex);
-			if (GSupportsGSRenderTargetLayerSwitchingToMips)
+			FSceneRenderTargetItem& EffectiveRT = GetEffectiveDiffuseIrradianceRenderTarget(MipIndex);			
+			
+			for (int32 CubeFace = 0; CubeFace < CubeFace_MAX; CubeFace++)
 			{
-				SetRenderTarget(RHICmdList, EffectiveRT.TargetableTexture, 0, NULL);
-
+				SetRenderTarget(RHICmdList, EffectiveRT.TargetableTexture, 0, CubeFace, NULL);
+					
 				const FIntRect ViewRect(0, 0, MipSize, MipSize);
 				RHICmdList.SetViewport(0, 0, 0.0f, MipSize, MipSize, 1.0f);
 				RHICmdList.SetRasterizerState(TStaticRasterizerState<FM_Solid, CM_None>::GetRHI());
 				RHICmdList.SetDepthStencilState(TStaticDepthStencilState<false, CF_Always>::GetRHI());
 				RHICmdList.SetBlendState(TStaticBlendState<>::GetRHI());
-
-				TShaderMapRef<FScreenVSForGS> VertexShader(ShaderMap);
-				TShaderMapRef<FDownsampleGS> GeometryShader(ShaderMap);
 				TShaderMapRef<FCopyDiffuseIrradiancePS> PixelShader(ShaderMap);
-
-				SetGlobalBoundShaderState(RHICmdList, FeatureLevel, CopyDiffuseIrradianceShaderState, GFilterVertexDeclaration.VertexDeclarationRHI, *VertexShader, *PixelShader, *GeometryShader);
-
-				// Draw each face with a geometry shader that routes to the correct render target slice
-				for (int32 CubeFace = 0; CubeFace < CubeFace_MAX; CubeFace++)
-				{
-					GeometryShader->SetParameters(RHICmdList, CubeFace);
-
-					PixelShader->SetParameters(RHICmdList, CubeFace, LightingSourceMipIndex, CoefficientIndex, MipSize, LightingSource);
-
-					DrawRectangle( 
-						RHICmdList,
-						ViewRect.Min.X, ViewRect.Min.Y, 
-						ViewRect.Width(), ViewRect.Height(),
-						ViewRect.Min.X, ViewRect.Min.Y, 
-						ViewRect.Width(), ViewRect.Height(),
-						FIntPoint(ViewRect.Width(), ViewRect.Height()),
-						FIntPoint(MipSize, MipSize),
-						*VertexShader);
-
-					RHICmdList.CopyToResolveTarget(EffectiveRT.TargetableTexture, EffectiveRT.ShaderResourceTexture, true, FResolveParams(FResolveRect(), (ECubeFace)CubeFace, MipIndex));
-				}
-			}
-			else
-			{
-				for (int32 CubeFace = 0; CubeFace < CubeFace_MAX; CubeFace++)
-				{
-					SetRenderTarget(RHICmdList, EffectiveRT.TargetableTexture, 0, CubeFace, NULL);
 					
-					const FIntRect ViewRect(0, 0, MipSize, MipSize);
-					RHICmdList.SetViewport(0, 0, 0.0f, MipSize, MipSize, 1.0f);
-					RHICmdList.SetRasterizerState(TStaticRasterizerState<FM_Solid, CM_None>::GetRHI());
-					RHICmdList.SetDepthStencilState(TStaticDepthStencilState<false, CF_Always>::GetRHI());
-					RHICmdList.SetBlendState(TStaticBlendState<>::GetRHI());
+				TShaderMapRef<FScreenVS> VertexShader(GetGlobalShaderMap(FeatureLevel));
 					
-					TShaderMapRef<FScreenVSForGS> VertexShader(ShaderMap);
-					TShaderMapRef<FCopyDiffuseIrradiancePS> PixelShader(ShaderMap);
+				SetGlobalBoundShaderState(RHICmdList, FeatureLevel, CopyDiffuseIrradianceShaderState, GFilterVertexDeclaration.VertexDeclarationRHI, *VertexShader, *PixelShader);
 					
-					SetGlobalBoundShaderState(RHICmdList, FeatureLevel, CopyDiffuseIrradianceShaderState, GFilterVertexDeclaration.VertexDeclarationRHI, *VertexShader, *PixelShader);
+				PixelShader->SetParameters(RHICmdList, CubeFace, LightingSourceMipIndex, CoefficientIndex, MipSize, LightingSource);
 					
-					PixelShader->SetParameters(RHICmdList, CubeFace, LightingSourceMipIndex, CoefficientIndex, MipSize, LightingSource);
+				DrawRectangle(
+					RHICmdList,
+					ViewRect.Min.X, ViewRect.Min.Y,
+					ViewRect.Width(), ViewRect.Height(),
+					ViewRect.Min.X, ViewRect.Min.Y,
+					ViewRect.Width(), ViewRect.Height(),
+					FIntPoint(ViewRect.Width(), ViewRect.Height()),
+					FIntPoint(MipSize, MipSize),
+					*VertexShader);
 					
-					DrawRectangle(
-						RHICmdList,
-						ViewRect.Min.X, ViewRect.Min.Y,
-						ViewRect.Width(), ViewRect.Height(),
-						ViewRect.Min.X, ViewRect.Min.Y,
-						ViewRect.Width(), ViewRect.Height(),
-						FIntPoint(ViewRect.Width(), ViewRect.Height()),
-						FIntPoint(MipSize, MipSize),
-						*VertexShader);
-					
-					RHICmdList.CopyToResolveTarget(EffectiveRT.TargetableTexture, EffectiveRT.ShaderResourceTexture, true, FResolveParams(FResolveRect(), (ECubeFace)CubeFace, MipIndex));
-				}
-			}
+				RHICmdList.CopyToResolveTarget(EffectiveRT.TargetableTexture, EffectiveRT.ShaderResourceTexture, true, FResolveParams(FResolveRect(), (ECubeFace)CubeFace, MipIndex));
+			}			
 		}
 
 		const int32 NumMips = FMath::CeilLogTwo(GDiffuseIrradianceCubemapSize) + 1;
@@ -328,82 +290,46 @@ void ComputeDiffuseIrradiance(FRHICommandListImmediate& RHICmdList, ERHIFeatureL
 				FSceneRenderTargetItem& EffectiveRT = GetEffectiveDiffuseIrradianceRenderTarget(MipIndex);
 				FSceneRenderTargetItem& EffectiveSource = GetEffectiveDiffuseIrradianceSourceTexture(MipIndex);
 				check(EffectiveRT.TargetableTexture != EffectiveSource.ShaderResourceTexture);
-
-				if (GSupportsGSRenderTargetLayerSwitchingToMips)
+				
+				for (int32 CubeFace = 0; CubeFace < CubeFace_MAX; CubeFace++)
 				{
-					SetRenderTarget(RHICmdList, EffectiveRT.TargetableTexture, MipIndex, NULL);
-
+					SetRenderTarget(RHICmdList, EffectiveRT.TargetableTexture, MipIndex, CubeFace, NULL);
+						
 					const FIntRect ViewRect(0, 0, MipSize, MipSize);
 					RHICmdList.SetViewport(0, 0, 0.0f, MipSize, MipSize, 1.0f);
 					RHICmdList.SetRasterizerState(TStaticRasterizerState<FM_Solid, CM_None>::GetRHI());
 					RHICmdList.SetDepthStencilState(TStaticDepthStencilState<false, CF_Always>::GetRHI());
 					RHICmdList.SetBlendState(TStaticBlendState<>::GetRHI());
-
-					TShaderMapRef<FScreenVSForGS> VertexShader(ShaderMap);
-					TShaderMapRef<FDownsampleGS> GeometryShader(ShaderMap);
 					TShaderMapRef<FAccumulateDiffuseIrradiancePS> PixelShader(ShaderMap);
-
-					SetGlobalBoundShaderState(RHICmdList, FeatureLevel, DiffuseIrradianceAccumulateShaderState, GFilterVertexDeclaration.VertexDeclarationRHI, *VertexShader, *PixelShader, *GeometryShader);
-
-					// Draw each face with a geometry shader that routes to the correct render target slice
-					for (int32 CubeFace = 0; CubeFace < CubeFace_MAX; CubeFace++)
-					{
-						GeometryShader->SetParameters(RHICmdList, CubeFace);
-
-						PixelShader->SetParameters(RHICmdList, CubeFace, NumMips, SourceMipIndex, CoefficientIndex, EffectiveSource.ShaderResourceTexture);
-
-						DrawRectangle( 
-							RHICmdList,
-							ViewRect.Min.X, ViewRect.Min.Y, 
-							ViewRect.Width(), ViewRect.Height(),
-							ViewRect.Min.X, ViewRect.Min.Y, 
-							ViewRect.Width(), ViewRect.Height(),
-							FIntPoint(ViewRect.Width(), ViewRect.Height()),
-							FIntPoint(MipSize, MipSize),
-							*VertexShader);
-
-						RHICmdList.CopyToResolveTarget(EffectiveRT.TargetableTexture, EffectiveRT.ShaderResourceTexture, true, FResolveParams(FResolveRect(), (ECubeFace)CubeFace, MipIndex));
-					}
-				}
-				else
-				{
-					for (int32 CubeFace = 0; CubeFace < CubeFace_MAX; CubeFace++)
-					{
-						SetRenderTarget(RHICmdList, EffectiveRT.TargetableTexture, MipIndex, CubeFace, NULL);
 						
-						const FIntRect ViewRect(0, 0, MipSize, MipSize);
-						RHICmdList.SetViewport(0, 0, 0.0f, MipSize, MipSize, 1.0f);
-						RHICmdList.SetRasterizerState(TStaticRasterizerState<FM_Solid, CM_None>::GetRHI());
-						RHICmdList.SetDepthStencilState(TStaticDepthStencilState<false, CF_Always>::GetRHI());
-						RHICmdList.SetBlendState(TStaticBlendState<>::GetRHI());
+					TShaderMapRef<FScreenVS> VertexShader(GetGlobalShaderMap(FeatureLevel));
 						
-						TShaderMapRef<FScreenVSForGS> VertexShader(ShaderMap);
-						TShaderMapRef<FAccumulateDiffuseIrradiancePS> PixelShader(ShaderMap);
+					SetGlobalBoundShaderState(RHICmdList, FeatureLevel, DiffuseIrradianceAccumulateShaderState, GFilterVertexDeclaration.VertexDeclarationRHI, *VertexShader, *PixelShader);
 						
-						SetGlobalBoundShaderState(RHICmdList, FeatureLevel, DiffuseIrradianceAccumulateShaderState, GFilterVertexDeclaration.VertexDeclarationRHI, *VertexShader, *PixelShader);
+					PixelShader->SetParameters(RHICmdList, CubeFace, NumMips, SourceMipIndex, CoefficientIndex, EffectiveSource.ShaderResourceTexture);
 						
-						PixelShader->SetParameters(RHICmdList, CubeFace, NumMips, SourceMipIndex, CoefficientIndex, EffectiveSource.ShaderResourceTexture);
-						
-						DrawRectangle(
-							RHICmdList,
-							ViewRect.Min.X, ViewRect.Min.Y,
-							ViewRect.Width(), ViewRect.Height(),
-							ViewRect.Min.X, ViewRect.Min.Y,
-							ViewRect.Width(), ViewRect.Height(),
-							FIntPoint(ViewRect.Width(), ViewRect.Height()),
-							FIntPoint(MipSize, MipSize),
-							*VertexShader);
+					DrawRectangle(
+						RHICmdList,
+						ViewRect.Min.X, ViewRect.Min.Y,
+						ViewRect.Width(), ViewRect.Height(),
+						ViewRect.Min.X, ViewRect.Min.Y,
+						ViewRect.Width(), ViewRect.Height(),
+						FIntPoint(ViewRect.Width(), ViewRect.Height()),
+						FIntPoint(MipSize, MipSize),
+						*VertexShader);
 						
-						RHICmdList.CopyToResolveTarget(EffectiveRT.TargetableTexture, EffectiveRT.ShaderResourceTexture, true, FResolveParams(FResolveRect(), (ECubeFace)CubeFace, MipIndex));
-					}
-				}
+					RHICmdList.CopyToResolveTarget(EffectiveRT.TargetableTexture, EffectiveRT.ShaderResourceTexture, true, FResolveParams(FResolveRect(), (ECubeFace)CubeFace, MipIndex));
+				}				
 			}
 		}
 
 		{
 			// Gather the cubemap face results and normalize, copy this coefficient to GSceneRenderTargets.SkySHIrradianceMap
 			FSceneRenderTargetItem& EffectiveRT = GSceneRenderTargets.SkySHIrradianceMap->GetRenderTargetItem();
-			SetRenderTarget(RHICmdList, EffectiveRT.TargetableTexture, 0, NULL);
+
+			//load/store actions so we don't lose results as we render one pixel at a time on tile renderers.
+			FRHIRenderTargetView RTV(EffectiveRT.TargetableTexture, 0, -1, ERenderTargetLoadAction::ELoad, ERenderTargetStoreAction::EStore);
+			RHICmdList.SetRenderTargets(1, &RTV, nullptr, 0, nullptr);
 
 			const FIntRect ViewRect(CoefficientIndex, 0, CoefficientIndex + 1, 1);
 			RHICmdList.SetViewport(0, 0, 0.0f, FSHVector3::MaxSHBasis, 1, 1.0f);
diff --git a/Engine/Source/Runtime/Renderer/Private/Renderer.cpp b/Engine/Source/Runtime/Renderer/Private/Renderer.cpp
--- a/Engine/Source/Runtime/Renderer/Private/Renderer.cpp
+++ b/Engine/Source/Runtime/Renderer/Private/Renderer.cpp
@@ -37,7 +37,7 @@ void FRendererModule::DrawTileMesh(FRHICommandListImmediate& RHICmdList, const F
 	// Create an FViewInfo so we can initialize its RHI resources
 	//@todo - reuse this view for multiple tiles, this is going to be slow for each tile
 	FViewInfo View(&SceneView);
-	View.InitRHIResources();
+	View.InitRHIResources(nullptr);
 
 	const auto FeatureLevel = View.GetFeatureLevel();
 	
diff --git a/Engine/Source/Runtime/Renderer/Private/RendererScene.cpp b/Engine/Source/Runtime/Renderer/Private/RendererScene.cpp
--- a/Engine/Source/Runtime/Renderer/Private/RendererScene.cpp
+++ b/Engine/Source/Runtime/Renderer/Private/RendererScene.cpp
@@ -620,6 +620,9 @@ void FScene::AddLightSceneInfo_RenderThread(FLightSceneInfo* LightSceneInfo)
 		!LightSceneInfo->Proxy->HasStaticLighting())
 	{
 		SimpleDirectionalLight = LightSceneInfo;
+
+		// if we are forward rendered and this light is a dynamic shadowcast then we need to update the static draw lists to pick a new lightingpolicy
+		bScenesPrimitivesNeedStaticMeshElementUpdate = !ShouldUseDeferredRenderer() && !SimpleDirectionalLight->Proxy->HasStaticShadowing();
 	}
 
 	if (LightSceneInfo->Proxy->IsUsedAsAtmosphereSunLight() &&
@@ -1068,6 +1071,8 @@ void FScene::RemoveLightSceneInfo_RenderThread(FLightSceneInfo* LightSceneInfo)
 	{
 		if (LightSceneInfo == SimpleDirectionalLight)
 		{
+			// if we are forward rendered and this light is a dynamic shadowcast then we need to update the static draw lists to pick a new lightingpolicy
+			bScenesPrimitivesNeedStaticMeshElementUpdate = !ShouldUseDeferredRenderer() && !SimpleDirectionalLight->Proxy->HasStaticShadowing();
 			SimpleDirectionalLight = NULL;
 		}
 
@@ -1615,6 +1620,8 @@ void FScene::UpdateStaticDrawListsForMaterials_RenderThread(FRHICommandListImmed
 		BasePassForForwardShadingLowQualityLightMapDrawList[DrawType].GetUsedPrimitivesBasedOnMaterials(FeatureLevel, Materials, PrimitivesToUpdate);
 		BasePassForForwardShadingDistanceFieldShadowMapLightMapDrawList[DrawType].GetUsedPrimitivesBasedOnMaterials(FeatureLevel, Materials, PrimitivesToUpdate);
 		BasePassForForwardShadingDirectionalLightAndSHIndirectDrawList[DrawType].GetUsedPrimitivesBasedOnMaterials(FeatureLevel, Materials, PrimitivesToUpdate);
+		BasePassForForwardShadingMovableDirectionalLightDrawList[DrawType].GetUsedPrimitivesBasedOnMaterials(FeatureLevel, Materials, PrimitivesToUpdate);
+		BasePassForForwardShadingMovableDirectionalLightCSMDrawList[DrawType].GetUsedPrimitivesBasedOnMaterials(FeatureLevel, Materials, PrimitivesToUpdate);
 	}
 
 	PositionOnlyDepthDrawList.GetUsedPrimitivesBasedOnMaterials(FeatureLevel, Materials, PrimitivesToUpdate);
@@ -1668,7 +1675,7 @@ void FScene::Release()
 					*FString::Printf(TEXT("Component Name: %s World Name: %s Component Mesh: %s"), 
 										*ActorComponent->GetFullName(), 
 										*GetWorld()->GetFullName(), 
- 										Cast<UStaticMeshComponent>(ActorComponent) ? *CastChecked<UStaticMeshComponent>(ActorComponent)->StaticMesh->GetFullName() : TEXT("Not a static mesh"))) )
+										Cast<UStaticMeshComponent>(ActorComponent) ? *CastChecked<UStaticMeshComponent>(ActorComponent)->StaticMesh->GetFullName() : TEXT("Not a static mesh"))) )
 			{
 				bTriggeredOnce = true;
 				break;	
@@ -1826,6 +1833,10 @@ void FScene::DumpStaticMeshDrawListStats() const
 	DUMP_DRAW_LIST(BasePassForForwardShadingDistanceFieldShadowMapLightMapDrawList[EBasePass_Masked]);
 	DUMP_DRAW_LIST(BasePassForForwardShadingDirectionalLightAndSHIndirectDrawList[EBasePass_Default]);
 	DUMP_DRAW_LIST(BasePassForForwardShadingDirectionalLightAndSHIndirectDrawList[EBasePass_Masked]);
+	DUMP_DRAW_LIST(BasePassForForwardShadingMovableDirectionalLightDrawList[EBasePass_Default]);
+	DUMP_DRAW_LIST(BasePassForForwardShadingMovableDirectionalLightDrawList[EBasePass_Masked]);
+	DUMP_DRAW_LIST(BasePassForForwardShadingMovableDirectionalLightCSMDrawList[EBasePass_Default]);
+	DUMP_DRAW_LIST(BasePassForForwardShadingMovableDirectionalLightCSMDrawList[EBasePass_Masked]);
 	DUMP_DRAW_LIST(HitProxyDrawList);
 	DUMP_DRAW_LIST(HitProxyDrawList_OpaqueOnly);
 	DUMP_DRAW_LIST(VelocityDrawList);
@@ -1991,6 +2002,8 @@ void FScene::ApplyWorldOffset_RenderThread(FVector InOffset)
 	StaticMeshDrawListApplyWorldOffset(BasePassForForwardShadingNoLightMapDrawList, InOffset);
 	StaticMeshDrawListApplyWorldOffset(BasePassForForwardShadingLowQualityLightMapDrawList, InOffset);
 	StaticMeshDrawListApplyWorldOffset(BasePassForForwardShadingDirectionalLightAndSHIndirectDrawList, InOffset);
+	StaticMeshDrawListApplyWorldOffset(BasePassForForwardShadingMovableDirectionalLightDrawList, InOffset);
+	StaticMeshDrawListApplyWorldOffset(BasePassForForwardShadingMovableDirectionalLightCSMDrawList, InOffset);
 
 	// Motion blur 
 	MotionBlurInfoData.ApplyOffset(InOffset);
@@ -2251,6 +2264,18 @@ TStaticMeshDrawList<TBasePassForForwardShadingDrawingPolicy<FSimpleDirectionalLi
 	return BasePassForForwardShadingDirectionalLightAndSHIndirectDrawList[DrawType];
 }
 
+template<>
+TStaticMeshDrawList<TBasePassForForwardShadingDrawingPolicy<FMovableDirectionalLightLightingPolicy> >& FScene::GetForwardShadingBasePassDrawList<FMovableDirectionalLightLightingPolicy>(EBasePassDrawListType DrawType)
+{
+	return BasePassForForwardShadingMovableDirectionalLightDrawList[DrawType];
+}
+
+template<>
+TStaticMeshDrawList<TBasePassForForwardShadingDrawingPolicy<FMovableDirectionalLightCSMLightingPolicy> >& FScene::GetForwardShadingBasePassDrawList<FMovableDirectionalLightCSMLightingPolicy>(EBasePassDrawListType DrawType)
+{
+	return BasePassForForwardShadingMovableDirectionalLightCSMDrawList[DrawType];
+}
+
 /*-----------------------------------------------------------------------------
 	MotionBlurInfoData
 -----------------------------------------------------------------------------*/
diff --git a/Engine/Source/Runtime/Renderer/Private/SceneCore.cpp b/Engine/Source/Runtime/Renderer/Private/SceneCore.cpp
--- a/Engine/Source/Runtime/Renderer/Private/SceneCore.cpp
+++ b/Engine/Source/Runtime/Renderer/Private/SceneCore.cpp
@@ -278,7 +278,7 @@ void FStaticMesh::AddToDrawLists(FRHICommandListImmediate& RHICmdList, FScene* S
 {
 	const auto FeatureLevel = Scene->GetFeatureLevel();
 
-	if (CastShadow && FeatureLevel >= ERHIFeatureLevel::SM4)
+	if (CastShadow)
 	{
 		FShadowDepthDrawingPolicyFactory::AddStaticMesh(Scene, this);
 	}
@@ -302,7 +302,7 @@ void FStaticMesh::AddToDrawLists(FRHICommandListImmediate& RHICmdList, FScene* S
 		return;
 	}
 
-	if (FeatureLevel >= ERHIFeatureLevel::SM4)
+	if (Scene->ShouldUseDeferredRenderer())
 	{
 		extern TAutoConsoleVariable<int32> CVarEarlyZPass;
 		int32 EarlyZPass = CVarEarlyZPass.GetValueOnRenderThread();
diff --git a/Engine/Source/Runtime/Renderer/Private/SceneHitProxyRendering.cpp b/Engine/Source/Runtime/Renderer/Private/SceneHitProxyRendering.cpp
--- a/Engine/Source/Runtime/Renderer/Private/SceneHitProxyRendering.cpp
+++ b/Engine/Source/Runtime/Renderer/Private/SceneHitProxyRendering.cpp
@@ -110,7 +110,6 @@ public:
 	{
 		// Only compile the hit proxy vertex shader on PC
 		return IsPCPlatform(Platform) 
-			&& IsFeatureLevelSupported(Platform, ERHIFeatureLevel::ES2) 
 			// and only compile for default materials or materials that are masked.
 			&& (Material->IsSpecialEngineMaterial() || Material->IsMasked() || Material->MaterialMayModifyMeshPosition() || Material->IsTwoSided());
 	}
diff --git a/Engine/Source/Runtime/Renderer/Private/ScenePrivate.h b/Engine/Source/Runtime/Renderer/Private/ScenePrivate.h
--- a/Engine/Source/Runtime/Renderer/Private/ScenePrivate.h
+++ b/Engine/Source/Runtime/Renderer/Private/ScenePrivate.h
@@ -415,21 +415,21 @@ class FSceneViewState : public FSceneViewStateInterface, public FDeferredCleanup
 {
 public:
 
-    class FProjectedShadowKey
-    {
+	class FProjectedShadowKey
+	{
 	public:
 
-        FORCEINLINE bool operator == (const FProjectedShadowKey &Other) const
-        {
-            return (PrimitiveId == Other.PrimitiveId && Light == Other.Light && SplitIndex == Other.SplitIndex && bTranslucentShadow == Other.bTranslucentShadow);
-        }
+		FORCEINLINE bool operator == (const FProjectedShadowKey &Other) const
+		{
+			return (PrimitiveId == Other.PrimitiveId && Light == Other.Light && SplitIndex == Other.SplitIndex && bTranslucentShadow == Other.bTranslucentShadow);
+		}
 
-        FProjectedShadowKey(FPrimitiveComponentId InPrimitiveId, const ULightComponent* InLight, int32 InSplitIndex, bool bInTranslucentShadow):
-            PrimitiveId(InPrimitiveId),
-            Light(InLight),
+		FProjectedShadowKey(FPrimitiveComponentId InPrimitiveId, const ULightComponent* InLight, int32 InSplitIndex, bool bInTranslucentShadow):
+			PrimitiveId(InPrimitiveId),
+			Light(InLight),
 			SplitIndex(InSplitIndex),
 			bTranslucentShadow(bInTranslucentShadow)
-        {
+		{
 		}
 
 		friend FORCEINLINE uint32 GetTypeHash(const FSceneViewState::FProjectedShadowKey& Key)
@@ -439,10 +439,10 @@ public:
 
 	private:
 		FPrimitiveComponentId PrimitiveId;
-        const ULightComponent* Light;
+		const ULightComponent* Light;
 		int32 SplitIndex;
 		bool bTranslucentShadow;
-    };
+	};
 
 	int32 NumBufferedFrames;
 	typedef TMap<FSceneViewState::FProjectedShadowKey, FRenderQueryRHIRef> ShadowKeyOcclusionQueryMap;
@@ -590,7 +590,7 @@ public:
 	FLightPropagationVolume* GetLightPropagationVolume() const { return LightPropagationVolume; }
 
 	/** Default constructor. */
-    FSceneViewState();
+	FSceneViewState();
 
 	void DestroyAOTileResources();
 	void DestroyLightPropagationVolume();
@@ -660,21 +660,21 @@ public:
 		return SeparateTranslucencyRT;
 	}
 
-    // FRenderResource interface.
+	// FRenderResource interface.
 	virtual void InitDynamicRHI() override
-    {
+	{
 		HZBOcclusionTests.InitDynamicRHI();
 	}
 
 	virtual void ReleaseDynamicRHI() override
-    {
+	{
 #if BUFFERED_OCCLUSION_QUERIES
 		for (int i = 0; i < ShadowOcclusionQueryMaps.Num(); ++i)
 		{
 			ShadowOcclusionQueryMaps[i].Reset();
 		}
 #else
-        ShadowOcclusionQueryMap.Reset();
+		ShadowOcclusionQueryMap.Reset();
 #endif
 		PrimitiveOcclusionHistorySet.Empty();
 		PrimitiveFadingStates.Empty();
@@ -695,7 +695,7 @@ public:
 		HZB.SafeRelease();
 		SelectionOutlineCacheKey.SafeRelease();
 		SelectionOutlineCacheValue.SafeRelease();
-    }
+	}
 
 	// FSceneViewStateInterface
 	RENDERER_API virtual void Destroy();
@@ -1190,6 +1190,8 @@ public:
 	TStaticMeshDrawList< TBasePassForForwardShadingDrawingPolicy< TLightMapPolicy<LQ_LIGHTMAP> > >							BasePassForForwardShadingLowQualityLightMapDrawList[EBasePass_MAX];
 	TStaticMeshDrawList< TBasePassForForwardShadingDrawingPolicy< TDistanceFieldShadowsAndLightMapPolicy<LQ_LIGHTMAP> > >	BasePassForForwardShadingDistanceFieldShadowMapLightMapDrawList[EBasePass_MAX];
 	TStaticMeshDrawList<TBasePassForForwardShadingDrawingPolicy< FSimpleDirectionalLightAndSHIndirectPolicy > >				BasePassForForwardShadingDirectionalLightAndSHIndirectDrawList[EBasePass_MAX];
+	TStaticMeshDrawList<TBasePassForForwardShadingDrawingPolicy< FMovableDirectionalLightLightingPolicy > >					BasePassForForwardShadingMovableDirectionalLightDrawList[EBasePass_MAX];
+	TStaticMeshDrawList<TBasePassForForwardShadingDrawingPolicy< FMovableDirectionalLightCSMLightingPolicy > >				BasePassForForwardShadingMovableDirectionalLightCSMDrawList[EBasePass_MAX];
 
 	/** Maps a light-map type to the appropriate base pass draw list. */
 	template<typename LightMapPolicyType>
@@ -1230,8 +1232,8 @@ public:
 	/** Whether the early Z pass was force enabled when static draw lists were built. */
 	int32 StaticDrawListsEarlyZPassMode;
 
-	/** True if a change to SkyLight has occurred that requires static draw lists to be updated. */
-	bool bScenesPrimitivesNeedStaticMeshElementUpdate;
+	/** True if a change to SkyLight / Lighting has occurred that requires static draw lists to be updated. */
+	bool bScenesPrimitivesNeedStaticMeshElementUpdate;	
 
 	/** The scene's sky light, if any. */
 	FSkyLightSceneProxy* SkyLight;
@@ -1461,6 +1463,11 @@ public:
 
 	virtual ERHIFeatureLevel::Type GetFeatureLevel() const override { return GetWorld()->FeatureLevel; }
 
+	bool ShouldRenderSkylight() const
+	{
+		return SkyLight && !SkyLight->bHasStaticLighting && GSupportsRenderTargetFormat_PF_FloatRGBA;
+	}
+
 private:
 
 	/**
diff --git a/Engine/Source/Runtime/Renderer/Private/SceneRendering.cpp b/Engine/Source/Runtime/Renderer/Private/SceneRendering.cpp
--- a/Engine/Source/Runtime/Renderer/Private/SceneRendering.cpp
+++ b/Engine/Source/Runtime/Renderer/Private/SceneRendering.cpp
@@ -206,11 +206,13 @@ void FViewInfo::SetupSkyIrradianceEnvironmentMapConstants(FVector4* OutSkyIrradi
 
 /** Creates the view's uniform buffer given a set of view transforms. */
 TUniformBufferRef<FViewUniformShaderParameters> FViewInfo::CreateUniformBuffer(
+	const TArray<FProjectedShadowInfo*, SceneRenderingAllocator>* DirectionalLightShadowInfo,	
 	const FMatrix& EffectiveTranslatedViewMatrix, 
 	FBox* OutTranslucentCascadeBoundsArray, 
 	int32 NumTranslucentCascades) const
 {
 	check(Family);
+	check(!DirectionalLightShadowInfo || DirectionalLightShadowInfo->Num() > 0);
 
 	// Calculate the vector used by shaders to convert clip space coordinates to texture space.
 	const FIntPoint BufferSize = GSceneRenderTargets.GetBufferSizeXY();
@@ -292,11 +294,60 @@ TUniformBufferRef<FViewUniformShaderParameters> FViewInfo::CreateUniformBuffer(
 		ViewUniformShaderParameters.AdaptiveTessellationFactor = 0.5f * float(ViewRect.Height()) / TessellationAdaptivePixelsPerEdge;
 	}
 
+	//white texture should act like a shadowmap cleared to the farplane.
+	ViewUniformShaderParameters.DirectionalLightShadowTexture = GWhiteTexture->TextureRHI;
+	ViewUniformShaderParameters.DirectionalLightShadowSampler = TStaticSamplerState<SF_Point, AM_Clamp, AM_Clamp, AM_Clamp>::GetRHI();		
 	if (Family->Scene)
 	{
 		FScene* Scene = (FScene*)Family->Scene;
-		ViewUniformShaderParameters.DirectionalLightColor = Scene->SimpleDirectionalLight ? Scene->SimpleDirectionalLight->Proxy->GetColor() / PI : FLinearColor::Black;
-		ViewUniformShaderParameters.DirectionalLightDirection = Scene->SimpleDirectionalLight ? -Scene->SimpleDirectionalLight->Proxy->GetDirection() : FVector::ZeroVector;
+
+		if (Scene->SimpleDirectionalLight)
+		{			
+			ViewUniformShaderParameters.DirectionalLightColor = Scene->SimpleDirectionalLight->Proxy->GetColor() / PI;
+			ViewUniformShaderParameters.DirectionalLightDirection = -Scene->SimpleDirectionalLight->Proxy->GetDirection();			
+
+			static_assert(MAX_FORWARD_SHADOWCASCADES <= 4, "more than 4 cascades not supported by the shader and uniform buffer");
+			if (DirectionalLightShadowInfo)
+			{
+				FProjectedShadowInfo& ShadowInfo = *((*DirectionalLightShadowInfo)[0]);
+				FIntPoint ShadowBufferResolution = ShadowInfo.GetShadowBufferResolution();
+				FVector4 ShadowBufferSizeValue((float)ShadowBufferResolution.X, (float)ShadowBufferResolution.Y, 1.0f / (float)ShadowBufferResolution.X, 1.0f / (float)ShadowBufferResolution.Y);
+
+				ViewUniformShaderParameters.DirectionalLightShadowTexture = GSceneRenderTargets.GetShadowDepthZTexture();				
+				ViewUniformShaderParameters.DirectionalLightShadowTransition = 1.0f / ShadowInfo.ComputeTransitionSize();
+				ViewUniformShaderParameters.DirectionalLightShadowSize = ShadowBufferSizeValue;
+
+				int32 NumShadowsToCopy = FMath::Min(DirectionalLightShadowInfo->Num(), MAX_FORWARD_SHADOWCASCADES);				
+				for (int32 i = 0; i < NumShadowsToCopy; ++i)
+				{
+					const FProjectedShadowInfo& ShadowInfo = *(*DirectionalLightShadowInfo)[i];
+					ViewUniformShaderParameters.DirectionalLightScreenToShadow[i] = ShadowInfo.GetScreenToShadowMatrix(*this);
+					ViewUniformShaderParameters.DirectionalLightShadowDistances[i] = ShadowInfo.CascadeSettings.SplitFar;
+				}
+
+				for (int32 i = NumShadowsToCopy; i < MAX_FORWARD_SHADOWCASCADES; ++i)
+				{
+					ViewUniformShaderParameters.DirectionalLightScreenToShadow[i].SetIdentity();
+					ViewUniformShaderParameters.DirectionalLightShadowDistances[i] = 0.0f;
+				}
+			}
+			else
+			{
+				ViewUniformShaderParameters.DirectionalLightShadowTransition = 0.0f;
+				ViewUniformShaderParameters.DirectionalLightShadowSize = FVector::ZeroVector;
+				for (int32 i = 0; i < MAX_FORWARD_SHADOWCASCADES; ++i)
+				{
+					ViewUniformShaderParameters.DirectionalLightScreenToShadow[i].SetIdentity();
+					ViewUniformShaderParameters.DirectionalLightShadowDistances[i] = 0.0f;
+				}			
+			}			 
+		}
+		else
+		{
+			ViewUniformShaderParameters.DirectionalLightColor = FLinearColor::Black;
+			ViewUniformShaderParameters.DirectionalLightDirection = FVector::ZeroVector;
+		}
+		
 		ViewUniformShaderParameters.UpperSkyColor = Scene->UpperDynamicSkylightColor;
 		ViewUniformShaderParameters.LowerSkyColor = Scene->LowerDynamicSkylightColor;
 
@@ -500,7 +551,7 @@ TUniformBufferRef<FViewUniformShaderParameters> FViewInfo::CreateUniformBuffer(
 	return TUniformBufferRef<FViewUniformShaderParameters>::CreateUniformBufferImmediate(ViewUniformShaderParameters, UniformBuffer_SingleFrame);
 }
 
-void FViewInfo::InitRHIResources()
+void FViewInfo::InitRHIResources(const TArray<FProjectedShadowInfo*, SceneRenderingAllocator>* DirectionalLightShadowInfo)
 {
 	FBox VolumeBounds[TVC_MAX];
 
@@ -508,6 +559,7 @@ void FViewInfo::InitRHIResources()
 	FMatrix TranslatedViewMatrix = FTranslationMatrix(-ViewMatrices.PreViewTranslation) * ViewMatrices.ViewMatrix;
 
 	UniformBuffer = CreateUniformBuffer(
+		DirectionalLightShadowInfo,
 		TranslatedViewMatrix,
 		VolumeBounds,
 		TVC_MAX);
@@ -632,7 +684,7 @@ FSceneRenderer::FSceneRenderer(const FSceneViewFamily* InViewFamily,FHitProxyCon
 
 bool FSceneRenderer::DoOcclusionQueries(ERHIFeatureLevel::Type InFeatureLevel) const
 {
-	return !(InFeatureLevel == ERHIFeatureLevel::ES2)
+	return !IsMobilePlatform(GRHIShaderPlatform)
 		&& CVarAllowOcclusionQueries.GetValueOnRenderThread() != 0;
 }
 
@@ -827,21 +879,20 @@ void FSceneRenderer::RenderFinish(FRHICommandListImmediate& RHICmdList)
 
 FSceneRenderer* FSceneRenderer::CreateSceneRenderer(const FSceneViewFamily* InViewFamily, FHitProxyConsumer* HitProxyConsumer)
 {
-	ERHIFeatureLevel::Type FeatureLevel = InViewFamily->Scene->GetFeatureLevel();
-
-	if (FeatureLevel == ERHIFeatureLevel::ES2)
+	bool bUseDeferred = InViewFamily->Scene->ShouldUseDeferredRenderer();
+	if (bUseDeferred)
 	{
-		return new FForwardShadingSceneRenderer(InViewFamily, HitProxyConsumer);
+		return new FDeferredShadingSceneRenderer(InViewFamily, HitProxyConsumer);
 	}
 	else
 	{
-		return new FDeferredShadingSceneRenderer(InViewFamily, HitProxyConsumer);
+		return new FForwardShadingSceneRenderer(InViewFamily, HitProxyConsumer);
 	}
 }
 
 void FSceneRenderer::RenderCustomDepthPass(FRHICommandListImmediate& RHICmdList)
 {
-	if(FeatureLevel < ERHIFeatureLevel::SM4)
+	if(!IsFeatureLevelSupported(GRHIShaderPlatform, ERHIFeatureLevel::SM4))
 	{
 		// not yet supported on lower end platforms
 		return;
@@ -1063,7 +1114,7 @@ TGlobalResource<FFilterVertexDeclaration>& FRendererModule::GetFilterVertexDecla
 bool IsMobileHDR()
 {
 	static auto* MobileHDRCvar = IConsoleManager::Get().FindTConsoleVariableDataInt(TEXT("r.MobileHDR"));
-	return MobileHDRCvar->GetValueOnRenderThread() == 1;
+	return MobileHDRCvar->GetValueOnAnyThread() == 1;
 }
 
 bool IsMobileHDR32bpp()
diff --git a/Engine/Source/Runtime/Renderer/Private/SceneRendering.h b/Engine/Source/Runtime/Renderer/Private/SceneRendering.h
--- a/Engine/Source/Runtime/Renderer/Private/SceneRendering.h
+++ b/Engine/Source/Runtime/Renderer/Private/SceneRendering.h
@@ -446,12 +446,13 @@ public:
 
 	/** Creates the view's uniform buffer given a set of view transforms. */
 	TUniformBufferRef<FViewUniformShaderParameters> CreateUniformBuffer(
+		const TArray<FProjectedShadowInfo*, SceneRenderingAllocator>* DirectionalLightShadowInfo,
 		const FMatrix& EffectiveTranslatedViewMatrix, 
 		FBox* OutTranslucentCascadeBoundsArray, 
 		int32 NumTranslucentCascades) const;
 
 	/** Initializes the RHI resources used by this view. */
-	void InitRHIResources();
+	void InitRHIResources(const TArray<FProjectedShadowInfo*, SceneRenderingAllocator>* DirectionalLightShadowInfo);
 
 	/** Determines distance culling and fades if the state changes */
 	bool IsDistanceCulled(float DistanceSquared, float MaxDrawDistance, float MinDrawDistance, const FPrimitiveSceneInfo* PrimitiveSceneInfo);
@@ -556,6 +557,47 @@ protected:
 
 	// Shared functionality between all scene renderers
 
+	/** Generates FProjectedShadowInfos for all wholesceneshadows on the given light.*/
+	void AddViewDependentWholeSceneShadowsForView(TArray<FProjectedShadowInfo*, SceneRenderingAllocator>& ShadowInfos, FVisibleLightInfo& VisibleLightInfo, FLightSceneInfo& LightSceneInfo);
+
+	/**
+	* Used by RenderLights to figure out if projected shadows need to be rendered to the attenuation buffer.
+	* Or to render a given shadowdepth map for forward rendering.
+	*
+	* @param LightSceneInfo Represents the current light
+	* @return true if anything needs to be rendered
+	*/
+	bool CheckForProjectedShadows(const FLightSceneInfo* LightSceneInfo) const;
+
+	/** Returns whether a per object shadow should be created due to the light being a stationary light. */
+	bool ShouldCreateObjectShadowForStationaryLight(const FLightSceneInfo* LightSceneInfo, const FPrimitiveSceneProxy* PrimitiveSceneProxy, bool bInteractionShadowMapped) const;
+
+	/** Gathers the list of primitives used to draw various shadow types */
+	void GatherShadowPrimitives(
+		const TArray<FProjectedShadowInfo*, SceneRenderingAllocator>& PreShadows,
+		const TArray<FProjectedShadowInfo*, SceneRenderingAllocator>& ViewDependentWholeSceneShadows,
+		bool bReflectionCaptureScene);
+
+	/**
+	* Checks to see if this primitive is affected by various shadow types
+	*
+	* @param PrimitiveSceneInfoCompact The primitive to check for shadow interaction
+	* @param PreShadows The list of pre-shadows to check against
+	*/
+	void GatherShadowsForPrimitiveInner(const FPrimitiveSceneInfoCompact& PrimitiveSceneInfoCompact,
+		const TArray<FProjectedShadowInfo*, SceneRenderingAllocator>& PreShadows,
+		const TArray<FProjectedShadowInfo*, SceneRenderingAllocator>& ViewDependentWholeSceneShadows,
+		bool bReflectionCaptureScene);
+
+	/** Gets a readable light name for use with a draw event. */
+	static void GetLightNameForDrawEvent(const FLightSceneProxy* LightProxy, FString& LightNameWithLevel);
+
+	/** Gathers simple lights from visible primtives in the passed in views. */
+	static void GatherSimpleLights(const FSceneViewFamily& ViewFamily, const TArray<FViewInfo>& Views, FSimpleLightArray& SimpleLights);
+
+	/** Calculates projected shadow visibility. */
+	void InitProjectedShadowVisibility(FRHICommandListImmediate& RHICmdList);	
+
 	/** Performs once per frame setup prior to visibility determination. */
 	void PreVisibilityFrameSetup(FRHICommandListImmediate& RHICmdList);
 
@@ -593,7 +635,7 @@ protected:
 	void OnStartFrame();
 
 	/** Renders the scene's distortion */
-	void RenderDistortion(FRHICommandListImmediate& RHICmdList);
+	void RenderDistortion(FRHICommandListImmediate& RHICmdList);	
 };
 
 
@@ -616,6 +658,9 @@ protected:
 
 	void InitViews(FRHICommandListImmediate& RHICmdList);
 
+	/** Finds the visible dynamic shadows for each view. */
+	void InitDynamicShadows(FRHICommandListImmediate& RHICmdList);
+
 	/** Renders the opaque base pass for forward shading. */
 	void RenderForwardShadingBasePass(FRHICommandListImmediate& RHICmdList);
 
@@ -624,4 +669,15 @@ protected:
 
 	/** Renders the base pass for translucency. */
 	void RenderTranslucency(FRHICommandListImmediate& RHICmdList);
+
+	/** Renders any necessary shadowmaps. */
+	void RenderShadowDepthMaps(FRHICommandListImmediate& RHICmdList);
+
+	/**
+	  * Used by RenderShadowDepthMaps to render shadowmap for the given light.
+	  *
+	  * @param LightSceneInfo Represents the current light
+	  * @return true if anything got rendered
+	  */
+	bool RenderShadowDepthMap(FRHICommandListImmediate& RHICmdList, const FLightSceneInfo* LightSceneInfo);
 };
diff --git a/Engine/Source/Runtime/Renderer/Private/SceneVisibility.cpp b/Engine/Source/Runtime/Renderer/Private/SceneVisibility.cpp
--- a/Engine/Source/Runtime/Renderer/Private/SceneVisibility.cpp
+++ b/Engine/Source/Runtime/Renderer/Private/SceneVisibility.cpp
@@ -754,7 +754,7 @@ static void ComputeRelevanceForView(
 		if (PrimitiveSceneInfo->bNeedsCachedReflectionCaptureUpdate
 			// In ES2, the per-object reflection is used for everything
 			// Otherwise it is just used on translucency
-			&& (Scene->GetFeatureLevel() == ERHIFeatureLevel::ES2 || bTranslucentRelevance))
+			&& (!Scene->ShouldUseDeferredRenderer() || bTranslucentRelevance))
 		{
 			PrimitiveSceneInfo->CachedReflectionCaptureProxy = Scene->FindClosestReflectionCapture(Scene->PrimitiveBounds[BitIt.GetIndex()].Origin);
 			PrimitiveSceneInfo->bNeedsCachedReflectionCaptureUpdate = false;
@@ -1260,10 +1260,7 @@ void FSceneRenderer::PreVisibilityFrameSetup(FRHICommandListImmediate& RHICmdLis
 
 			// This finishes the update of view state
 			ViewState->UpdateLastRenderTime(*View.Family);
-		}
-
-		// Initialize the view's RHI resources.
-		View.InitRHIResources();
+		}		
 	}
 }
 
@@ -1510,7 +1507,7 @@ void FSceneRenderer::PostVisibilityFrameSetup()
 	}
 
 	bool bCheckLightShafts = false;
-	if (Scene->GetFeatureLevel() <= ERHIFeatureLevel::ES2)
+	if (Scene->GetFeatureLevel() <= ERHIFeatureLevel::ES3_1)
 	{
 		// Clear the mobile light shaft data.
 		for (int32 ViewIndex = 0; ViewIndex < Views.Num(); ViewIndex++)
@@ -1656,6 +1653,13 @@ void FDeferredShadingSceneRenderer::InitViews(FRHICommandListImmediate& RHICmdLi
 		GPrevPerBoneMotionBlur.RestoreForPausedMotionBlur();
 	}
 
+	// initialize per-view uniform buffer.
+	for (int32 ViewIndex = 0; ViewIndex < Views.Num(); ViewIndex++)
+	{
+		// Initialize the view's RHI resources.
+		Views[ViewIndex].InitRHIResources(nullptr);
+	}	
+
 	OnStartFrame();
 }
 
diff --git a/Engine/Source/Runtime/Renderer/Private/ShadowRendering.cpp b/Engine/Source/Runtime/Renderer/Private/ShadowRendering.cpp
--- a/Engine/Source/Runtime/Renderer/Private/ShadowRendering.cpp
+++ b/Engine/Source/Runtime/Renderer/Private/ShadowRendering.cpp
@@ -170,7 +170,8 @@ public:
 					&& (!bUsePositionOnlyStream || VertexFactoryType->SupportsPositionOnly())
 					// Don't render ShadowDepth for translucent unlit materials
 					&& (!IsTranslucentBlendMode(Material->GetBlendMode()) && Material->GetShadingModel() != MSM_Unlit)
-					&& IsFeatureLevelSupported(Platform, ERHIFeatureLevel::SM4);
+					// Only compile perspective correct light shaders for feature levels >= SM4
+					&& (ShaderMode != VertexShadowDepth_PerspectiveCorrect || IsFeatureLevelSupported(Platform, ERHIFeatureLevel::SM4));
 		}
 	}
 
@@ -1273,7 +1274,7 @@ void DrawShadowMeshElements(FRHICommandList& RHICmdList, const FViewInfo& View,
 	}
 }
 
-void FProjectedShadowInfo::RenderDepthDynamic(FRHICommandList& RHICmdList, FDeferredShadingSceneRenderer* SceneRenderer, const FViewInfo* FoundView)
+void FProjectedShadowInfo::RenderDepthDynamic(FRHICommandList& RHICmdList, FSceneRenderer* SceneRenderer, const FViewInfo* FoundView)
 {
 	// Draw the subject's dynamic elements.
 	SCOPE_CYCLE_COUNTER(STAT_WholeSceneDynamicShadowDepthsTime);
@@ -1390,7 +1391,7 @@ class FRenderDepthDynamicThreadTask
 	FProjectedShadowInfo& ThisRenderer;
 	FRHICommandList& RHICmdList;
 	const FViewInfo& View;
-	FDeferredShadingSceneRenderer* SceneRenderer;
+	FSceneRenderer* SceneRenderer;
 
 public:
 
@@ -1398,7 +1399,7 @@ public:
 		FProjectedShadowInfo* InThisRenderer,
 		FRHICommandList* InRHICmdList,
 		const FViewInfo* InView,
-		FDeferredShadingSceneRenderer* InSceneRenderer
+		FSceneRenderer* InSceneRenderer
 		)
 		: ThisRenderer(*InThisRenderer)
 		, RHICmdList(*InRHICmdList)
@@ -1433,7 +1434,7 @@ static TAutoConsoleVariable<int32> CVarParallelShadows(
 	ECVF_RenderThreadSafe
 	);
 
-void FProjectedShadowInfo::RenderDepthInner(FRHICommandList& RHICmdList, FDeferredShadingSceneRenderer* SceneRenderer, const FViewInfo* FoundView)
+void FProjectedShadowInfo::RenderDepthInner(FRHICommandList& RHICmdList, FSceneRenderer* SceneRenderer, const FViewInfo* FoundView)
 {
 	FShadowDepthDrawingPolicyContext PolicyContext(this);
 
@@ -1519,7 +1520,7 @@ void FProjectedShadowInfo::RenderDepthInner(FRHICommandList& RHICmdList, FDeferr
 	}
 }
 
-void FProjectedShadowInfo::RenderDepth(FRHICommandList& RHICmdList, FDeferredShadingSceneRenderer* SceneRenderer)
+void FProjectedShadowInfo::RenderDepth(FRHICommandList& RHICmdList, FSceneRenderer* SceneRenderer)
 {
 #if WANTS_DRAW_MESH_EVENTS
 	FString EventName;
@@ -1609,6 +1610,7 @@ void FProjectedShadowInfo::RenderDepth(FRHICommandList& RHICmdList, FDeferredSha
 	FoundView->ViewMatrices.ViewMatrix = ShadowViewMatrix;
 	FBox VolumeBounds[TVC_MAX];
 	FoundView->UniformBuffer = FoundView->CreateUniformBuffer(
+		nullptr,
 		ShadowViewMatrix, 
 		VolumeBounds,
 		TVC_MAX);
@@ -2349,7 +2351,17 @@ FMatrix FProjectedShadowInfo::GetWorldToShadowMatrix(FVector4& ShadowmapMinMax,
 /** Returns the resolution of the shadow buffer used for this shadow, based on the shadow's type. */
 FIntPoint FProjectedShadowInfo::GetShadowBufferResolution() const
 {
-	return bAllocatedInPreshadowCache ? GSceneRenderTargets.GetPreShadowCacheTextureResolution() : GSceneRenderTargets.GetShadowDepthTextureResolution();
+	const FTexture2DRHIRef& ShadowTexture = GSceneRenderTargets.GetShadowDepthZTexture(bAllocatedInPreshadowCache);
+
+	//prefer to return the actual size of the allocated texture if possible.  It may be larger than the size of a single shadowmap due to atlasing (see forward renderer CSM handling in InitDynamicShadows).
+	if (ShadowTexture)
+	{
+		return FIntPoint(ShadowTexture->GetSizeX(), ShadowTexture->GetSizeY());
+	}
+	else
+	{
+		return bAllocatedInPreshadowCache ? GSceneRenderTargets.GetPreShadowCacheTextureResolution() : GSceneRenderTargets.GetShadowDepthTextureResolution();
+	}
 }
 
 void FProjectedShadowInfo::UpdateShaderDepthBias()
@@ -2506,7 +2518,7 @@ FDeferredShadingSceneRenderer
  * @param LightSceneInfo Represents the current light
  * @return true if anything needs to be rendered
  */
-bool FDeferredShadingSceneRenderer::CheckForProjectedShadows( const FLightSceneInfo* LightSceneInfo ) const
+bool FSceneRenderer::CheckForProjectedShadows( const FLightSceneInfo* LightSceneInfo ) const
 {
 	// Find the projected shadows cast by this light.
 	const FVisibleLightInfo& VisibleLightInfo = VisibleLightInfos[LightSceneInfo->Id];
@@ -3305,3 +3317,149 @@ bool FDeferredShadingSceneRenderer::RenderCachedPreshadows(FRHICommandListImmedi
 	return bAttenuationBufferDirty;
 }
 
+/**
+ * Used by RenderLights to render shadows to the attenuation buffer.
+ *
+ * @param LightSceneInfo Represents the current light
+ * @return true if anything got rendered
+ */
+bool FForwardShadingSceneRenderer::RenderShadowDepthMap(FRHICommandListImmediate& RHICmdList, const FLightSceneInfo* LightSceneInfo)
+{
+	//@todo-rco: Disabled
+	bool bCurrentlyDisabled = true;
+	if (bCurrentlyDisabled)
+	{
+		return false;
+	}
+
+	SCOPE_CYCLE_COUNTER(STAT_ProjectedShadowDrawTime);	
+
+	bool bAttenuationBufferDirty = false;
+
+	// Find the projected shadows cast by this light.
+	FVisibleLightInfo& VisibleLightInfo = VisibleLightInfos[LightSceneInfo->Id];
+	TArray<FProjectedShadowInfo*, SceneRenderingAllocator> Shadows;
+	for (int32 ShadowIndex = 0; ShadowIndex < VisibleLightInfo.AllProjectedShadows.Num(); ShadowIndex++)
+	{
+		FProjectedShadowInfo* ProjectedShadowInfo = VisibleLightInfo.AllProjectedShadows[ShadowIndex];
+
+		// Check that the shadow is visible in at least one view before rendering it.
+		bool bShadowIsVisible = false;
+		for (int32 ViewIndex = 0; ViewIndex < Views.Num(); ViewIndex++)
+		{
+			const FViewInfo& View = Views[ViewIndex];
+			if (ProjectedShadowInfo->DependentView && ProjectedShadowInfo->DependentView != &View)
+			{
+				continue;
+			}
+			const FVisibleLightViewInfo& VisibleLightViewInfo = View.VisibleLightInfos[LightSceneInfo->Id];
+			const FPrimitiveViewRelevance ViewRelevance = VisibleLightViewInfo.ProjectedShadowViewRelevanceMap[ShadowIndex];
+			bShadowIsVisible |= !ProjectedShadowInfo->bTranslucentShadow && ViewRelevance.bOpaqueRelevance 
+				&& VisibleLightViewInfo.ProjectedShadowVisibilityMap[ShadowIndex];
+		}
+
+		// Skip shadows which will be handled in RenderOnePassPointLightShadows or RenderReflectiveShadowmaps
+		if (ProjectedShadowInfo->bOnePassPointLightShadow || ProjectedShadowInfo->bReflectiveShadowmap)
+		{
+			bShadowIsVisible = false;
+		}
+
+		if (ProjectedShadowInfo->bPreShadow && !LightSceneInfo->Proxy->HasStaticShadowing())
+		{
+			bShadowIsVisible = false;
+		}
+
+		if (bShadowIsVisible
+			&& (!ProjectedShadowInfo->bPreShadow || ProjectedShadowInfo->HasSubjectPrims())
+			&& !ProjectedShadowInfo->bAllocatedInPreshadowCache)
+		{
+			// Add the shadow to the list of visible shadows cast by this light.
+			if (ProjectedShadowInfo->bWholeSceneShadow)
+			{
+				INC_DWORD_STAT(STAT_WholeSceneShadows);
+			}
+			else if (ProjectedShadowInfo->bPreShadow)
+			{
+				INC_DWORD_STAT(STAT_PreShadows);
+			}
+			else
+			{
+				INC_DWORD_STAT(STAT_PerObjectShadows);
+			}
+			Shadows.Add(ProjectedShadowInfo);
+		}
+	}
+
+	// Sort the projected shadows by resolution.
+	Shadows.Sort( FCompareFProjectedShadowInfoBySplitIndex() );	
+
+	{
+		// Render the shadow depths.
+		SCOPED_DRAW_EVENT(RHICmdList, ShadowDepthsFromOpaque, DEC_SCENE_ITEMS);
+		GSceneRenderTargets.BeginRenderingShadowDepth(RHICmdList);
+		RHICmdList.Clear(false, FColor(255, 255, 255), true, 1.0f, false, 0, FIntRect());
+
+		// render depth for each shadow
+		for (int32 ShadowIndex = 0; ShadowIndex < Shadows.Num(); ShadowIndex++)
+		{
+			FProjectedShadowInfo* ProjectedShadowInfo = Shadows[ShadowIndex];
+			check(!ProjectedShadowInfo->bTranslucentShadow);
+			ProjectedShadowInfo->bRendered = false;
+			ProjectedShadowInfo->bAllocated = true;
+			ProjectedShadowInfo->RenderDepth(RHICmdList, this);
+			bAttenuationBufferDirty = true;
+		}
+
+		GSceneRenderTargets.FinishRenderingShadowDepth(RHICmdList);
+	}
+
+	return bAttenuationBufferDirty;
+}
+
+void FForwardShadingSceneRenderer::RenderShadowDepthMaps(FRHICommandListImmediate& RHICmdList)
+{
+	SCOPED_DRAW_EVENT(RHICmdList, Lights, DEC_SCENE_ITEMS);
+
+	bool bDynamicShadows = ViewFamily.EngineShowFlags.DynamicShadows && GetShadowQuality() > 0;
+	if (IsSimpleDynamicLightingEnabled() || !bDynamicShadows)
+	{
+		return;
+	}	
+
+	// render shadowmaps for relevant lights.
+	if (Scene->SimpleDirectionalLight)
+	{		
+		const FLightSceneInfo& LightSceneInfo = *Scene->SimpleDirectionalLight;
+
+		if (LightSceneInfo.ShouldRenderViewIndependentWholeSceneShadows()
+			// Only render movable shadowcasting lights
+			&& !LightSceneInfo.Proxy->HasStaticShadowing()
+			&& CheckForProjectedShadows(&LightSceneInfo))
+		{
+			bool bRender = false;
+			// Check if the light is visible in any of the views.
+			for (int32 ViewIndex = 0; ViewIndex < Views.Num(); ViewIndex++)
+			{
+				if (LightSceneInfo.ShouldRenderLight(Views[ViewIndex]))
+				{
+					bRender = true;
+					break;
+				}
+			}			
+
+			if (bRender)
+			{				
+				FScopeCycleCounter Context(LightSceneInfo.Proxy->GetStatId());
+
+				FString LightNameWithLevel;
+				GetLightNameForDrawEvent(LightSceneInfo.Proxy, LightNameWithLevel);
+				SCOPED_DRAW_EVENTF(RHICmdList, EventLightPass, DEC_SCENE_ITEMS, *LightNameWithLevel);
+				
+				INC_DWORD_STAT(STAT_NumShadowedLights);
+
+				//render shadowmap depth				
+				RenderShadowDepthMap(RHICmdList, &LightSceneInfo);				
+			}					
+		}		
+	}	
+}
diff --git a/Engine/Source/Runtime/Renderer/Private/ShadowRendering.h b/Engine/Source/Runtime/Renderer/Private/ShadowRendering.h
--- a/Engine/Source/Runtime/Renderer/Private/ShadowRendering.h
+++ b/Engine/Source/Runtime/Renderer/Private/ShadowRendering.h
@@ -667,7 +667,7 @@ public:
 	/**
 	 * Renders the shadow subject depth.
 	 */
-	void RenderDepth(FRHICommandList& RHICmdList, class FDeferredShadingSceneRenderer* SceneRenderer);
+	void RenderDepth(FRHICommandList& RHICmdList, class FSceneRenderer* SceneRenderer);
 
 	void ClearDepth(FRHICommandList& RHICmdList, class FDeferredShadingSceneRenderer* SceneRenderer);
 
@@ -752,13 +752,13 @@ private:
 	/**
 	* Renders the shadow subject depth, to a particular hacked view
 	*/
-	void RenderDepthInner(FRHICommandList& RHICmdList, class FDeferredShadingSceneRenderer* SceneRenderer, const FViewInfo* FoundView);
+	void RenderDepthInner(FRHICommandList& RHICmdList, class FSceneRenderer* SceneRenderer, const FViewInfo* FoundView);
 
 	/**
 	* Renders the dynamic shadow subject depth, to a particular hacked view
 	*/
 	friend class FRenderDepthDynamicThreadTask;
-	void RenderDepthDynamic(FRHICommandList& RHICmdList, class FDeferredShadingSceneRenderer* SceneRenderer, const FViewInfo* FoundView);
+	void RenderDepthDynamic(FRHICommandList& RHICmdList, class FSceneRenderer* SceneRenderer, const FViewInfo* FoundView);
 
 	void GetShadowTypeNameForDrawEvent(FString& TypeName) const;
 
@@ -1007,8 +1007,6 @@ public:
 		ShadowBufferSize.Bind(ParameterMap,TEXT("ShadowBufferSize"));
 		ShadowDepthTexture.Bind(ParameterMap,TEXT("ShadowDepthTexture"));
 		ShadowDepthTextureSampler.Bind(ParameterMap,TEXT("ShadowDepthTextureSampler"));
-		ShadowDepthTextureObject.Bind(ParameterMap,TEXT("ShadowDepthTextureObject"));
-		ShadowDepthSampler.Bind(ParameterMap,TEXT("ShadowDepthSampler"));
 		ProjectionDepthBias.Bind(ParameterMap,TEXT("ProjectionDepthBiasParameters"));
 		FadePlaneOffset.Bind(ParameterMap,TEXT("FadePlaneOffset"));
 		InvFadePlaneLength.Bind(ParameterMap,TEXT("InvFadePlaneLength"));
@@ -1044,14 +1042,13 @@ public:
 		FTexture2DRHIRef ShadowDepthTextureValue = GSceneRenderTargets.GetShadowDepthZTexture(ShadowInfo->bAllocatedInPreshadowCache);
 		FSamplerStateRHIParamRef DepthSamplerState = TStaticSamplerState<SF_Point,AM_Clamp,AM_Clamp,AM_Clamp>::GetRHI();
 
-		SetTextureParameter(RHICmdList, ShaderRHI, ShadowDepthTexture, ShadowDepthTextureSampler, DepthSamplerState, ShadowDepthTextureValue);
-		SetTextureParameter(RHICmdList, ShaderRHI, ShadowDepthTextureObject, ShadowDepthTextureValue);
+		SetTextureParameter(RHICmdList, ShaderRHI, ShadowDepthTexture, ShadowDepthTextureSampler, DepthSamplerState, ShadowDepthTextureValue);		
 
-		if (ShadowDepthSampler.IsBound())
+		if (ShadowDepthTextureSampler.IsBound())
 		{
 			RHICmdList.SetShaderSampler(
 				ShaderRHI, 
-				ShadowDepthSampler.GetBaseIndex(), 
+				ShadowDepthTextureSampler.GetBaseIndex(),
 				DepthSamplerState
 				);
 		}
@@ -1075,8 +1072,6 @@ public:
 		Ar << P.ShadowBufferSize;
 		Ar << P.ShadowDepthTexture;
 		Ar << P.ShadowDepthTextureSampler;
-		Ar << P.ShadowDepthTextureObject;
-		Ar << P.ShadowDepthSampler;
 		Ar << P.ProjectionDepthBias;
 		Ar << P.FadePlaneOffset;
 		Ar << P.InvFadePlaneLength;
@@ -1091,8 +1086,6 @@ private:
 	FShaderParameter ShadowBufferSize;
 	FShaderResourceParameter ShadowDepthTexture;
 	FShaderResourceParameter ShadowDepthTextureSampler;
-	FShaderResourceParameter ShadowDepthTextureObject;
-	FShaderResourceParameter ShadowDepthSampler;
 	FShaderParameter ProjectionDepthBias;
 	FShaderParameter FadePlaneOffset;
 	FShaderParameter InvFadePlaneLength;
diff --git a/Engine/Source/Runtime/Renderer/Private/ShadowSetup.cpp b/Engine/Source/Runtime/Renderer/Private/ShadowSetup.cpp
--- a/Engine/Source/Runtime/Renderer/Private/ShadowSetup.cpp
+++ b/Engine/Source/Runtime/Renderer/Private/ShadowSetup.cpp
@@ -373,7 +373,7 @@ FProjectedShadowInfo::FProjectedShadowInfo(
 	ParentSceneInfo(InParentSceneInfo),
 	DependentView(NULL),
 	ShadowId(INDEX_NONE),
-	PreShadowTranslation(Initializer.PreShadowTranslation),
+	PreShadowTranslation(Initializer.PreShadowTranslation),	
 	ShadowBounds(Initializer.SubjectBounds.Origin - Initializer.PreShadowTranslation, Initializer.SubjectBounds.SphereRadius),
 	X(0),
 	Y(0),
@@ -396,7 +396,7 @@ FProjectedShadowInfo::FProjectedShadowInfo(
 	bRayTracedDistanceFieldShadow(false)
 {
 	const FMatrix WorldToLightScaled = Initializer.WorldToLight * FScaleMatrix(Initializer.Scales);
-
+	
 	// Create an array of the extreme vertices of the subject's bounds.
 	FBoundingBoxVertexArray BoundsPoints;
 	FBoundingBoxEdgeArray BoundsEdges;
@@ -519,7 +519,7 @@ FProjectedShadowInfo::FProjectedShadowInfo(
 ,	bPreShadow(false)
 ,	bRayTracedDistanceFieldShadow(Initializer.bRayTracedDistanceFieldShadow)
 ,	bValidTransform(true)
-{
+{	
 	FVector	XAxis, YAxis;
 	Initializer.FaceDirection.FindBestAxisVectors(XAxis,YAxis);
 	const FMatrix WorldToLightScaled = Initializer.WorldToLight * FScaleMatrix(Initializer.Scales);
@@ -1039,7 +1039,7 @@ void FDeferredShadingSceneRenderer::UpdatePreshadowCache()
 	}
 }
 
-bool FDeferredShadingSceneRenderer::ShouldCreateObjectShadowForStationaryLight(const FLightSceneInfo* LightSceneInfo, const FPrimitiveSceneProxy* PrimitiveSceneProxy, bool bInteractionShadowMapped) const
+bool FSceneRenderer::ShouldCreateObjectShadowForStationaryLight(const FLightSceneInfo* LightSceneInfo, const FPrimitiveSceneProxy* PrimitiveSceneProxy, bool bInteractionShadowMapped) const
 {
 	const bool bCreateObjectShadowForStationaryLight = 
 		LightSceneInfo->bCreatePerObjectShadowsForDynamicObjects 
@@ -1580,7 +1580,7 @@ void FDeferredShadingSceneRenderer::CreateWholeSceneProjectedShadow(FLightSceneI
 	}
 }
 
-void FDeferredShadingSceneRenderer::InitProjectedShadowVisibility(FRHICommandListImmediate& RHICmdList)
+void FSceneRenderer::InitProjectedShadowVisibility(FRHICommandListImmediate& RHICmdList)
 {
 	// Initialize the views' ProjectedShadowVisibilityMaps and remove shadows without subjects.
 	for(TSparseArray<FLightSceneInfoCompact>::TConstIterator LightIt(Scene->Lights);LightIt;++LightIt)
@@ -1716,7 +1716,7 @@ void FDeferredShadingSceneRenderer::InitProjectedShadowVisibility(FRHICommandLis
 	}
 }
 
-inline void FDeferredShadingSceneRenderer::GatherShadowsForPrimitiveInner(
+inline void FSceneRenderer::GatherShadowsForPrimitiveInner(
 	const FPrimitiveSceneInfoCompact& PrimitiveSceneInfoCompact, 
 	const TArray<FProjectedShadowInfo*,SceneRenderingAllocator>& PreShadows,
 	const TArray<FProjectedShadowInfo*,SceneRenderingAllocator>& ViewDependentWholeSceneShadows,
@@ -1826,7 +1826,7 @@ inline void FDeferredShadingSceneRenderer::GatherShadowsForPrimitiveInner(
 	}
 }
 
-void FDeferredShadingSceneRenderer::GatherShadowPrimitives(
+void FSceneRenderer::GatherShadowPrimitives(
 	const TArray<FProjectedShadowInfo*,SceneRenderingAllocator>& PreShadows,
 	const TArray<FProjectedShadowInfo*,SceneRenderingAllocator>& ViewDependentWholeSceneShadows,
 	bool bStaticSceneOnly
@@ -1929,6 +1929,173 @@ void FDeferredShadingSceneRenderer::GatherShadowPrimitives(
 	}
 }
 
+void FSceneRenderer::AddViewDependentWholeSceneShadowsForView(TArray<FProjectedShadowInfo*, SceneRenderingAllocator>& ShadowInfos, FVisibleLightInfo& VisibleLightInfo, FLightSceneInfo& LightSceneInfo)
+{
+	// Allow each view to create a whole scene view dependent shadow
+	for (int32 ViewIndex = 0; ViewIndex < Views.Num(); ViewIndex++)
+	{
+		FViewInfo& View = Views[ViewIndex];
+
+		TArray<float, TInlineAllocator<2> > FadeAlphas;
+		FadeAlphas.Init(0.0f, Views.Num());
+		FadeAlphas[ViewIndex] = 1.0f;
+
+		if (View.StereoPass == eSSP_LEFT_EYE
+			&& Views.IsValidIndex(ViewIndex + 1)
+			&& Views[ViewIndex + 1].StereoPass == eSSP_RIGHT_EYE)
+		{
+			FadeAlphas[ViewIndex + 1] = 1.0f;
+		}		
+
+		// If rendering in stereo mode we render shadow depths only for the left eye, but project for both eyes!
+		if (View.StereoPass != eSSP_RIGHT_EYE)
+		{
+			const int32 NumSplits = LightSceneInfo.Proxy->GetNumViewDependentWholeSceneShadows(View);
+			for (int32 SplitIndex = 0; SplitIndex < NumSplits; SplitIndex++)
+			{
+				FWholeSceneProjectedShadowInitializer ProjectedShadowInitializer;
+
+				if (LightSceneInfo.Proxy->GetViewDependentWholeSceneProjectedShadowInitializer(View, SplitIndex, ProjectedShadowInitializer))
+				{
+					const FIntPoint ShadowBufferResolution = GSceneRenderTargets.GetShadowDepthTextureResolution();
+					// Create the projected shadow info.
+					FProjectedShadowInfo* ProjectedShadowInfo = new(FMemStack::Get(), 1, 16) FProjectedShadowInfo(
+						&LightSceneInfo,
+						&View,
+						ProjectedShadowInitializer,
+						//@todo - remove the shadow border for whole scene shadows
+						ShadowBufferResolution.X - SHADOW_BORDER * 2,
+						ShadowBufferResolution.Y - SHADOW_BORDER * 2,
+						false	// no RSM
+						);
+
+					ProjectedShadowInfo->FadeAlphas = FadeAlphas;
+
+					FVisibleLightInfo& LightViewInfo = VisibleLightInfos[LightSceneInfo.Id];
+					VisibleLightInfo.MemStackProjectedShadows.Add(ProjectedShadowInfo);
+					VisibleLightInfo.AllProjectedShadows.Add(ProjectedShadowInfo);
+					ShadowInfos.Add(ProjectedShadowInfo);
+				}
+			}
+
+			FSceneViewState* ViewState = (FSceneViewState*)View.State;
+			if (ViewState)
+			{
+				FLightPropagationVolume* LightPropagationVolume = ViewState->GetLightPropagationVolume();
+
+				if (LightPropagationVolume && View.FinalPostProcessSettings.LPVIntensity > 0)
+				{
+					// Generate the RSM shadow info
+					FWholeSceneProjectedShadowInitializer ProjectedShadowInitializer;
+					FLightPropagationVolume& Lpv = *LightPropagationVolume;
+
+					if (LightSceneInfo.Proxy->GetViewDependentRsmWholeSceneProjectedShadowInitializer(View, Lpv.GetBoundingBox(), ProjectedShadowInitializer))
+					{
+						const FIntPoint ShadowBufferResolution = GSceneRenderTargets.GetReflectiveShadowMapTextureResolution();
+
+						// Create the projected shadow info.
+						FProjectedShadowInfo* ProjectedShadowInfo = new(FMemStack::Get(), 1, 16) FProjectedShadowInfo(
+							&LightSceneInfo,
+							&View,
+							ProjectedShadowInitializer,
+							ShadowBufferResolution.X,
+							ShadowBufferResolution.Y,
+							true);
+
+						FVisibleLightInfo& LightViewInfo = VisibleLightInfos[LightSceneInfo.Id];
+						VisibleLightInfo.MemStackProjectedShadows.Add(ProjectedShadowInfo);
+						VisibleLightInfo.AllProjectedShadows.Add(ProjectedShadowInfo);
+						VisibleLightInfo.ReflectiveShadowMaps.Add(ProjectedShadowInfo);
+						ShadowInfos.Add(ProjectedShadowInfo); // or separate list?
+					}
+				}
+			}
+		}
+	}
+}
+
+void FForwardShadingSceneRenderer::InitDynamicShadows(FRHICommandListImmediate& RHICmdList)
+{	
+	TArray<FProjectedShadowInfo*, SceneRenderingAllocator> ViewDependentWholeSceneShadows;
+	TArray<FProjectedShadowInfo*, SceneRenderingAllocator> PreShadows;
+	{
+		SCOPE_CYCLE_COUNTER(STAT_InitDynamicShadowsTime);
+
+		for (TSparseArray<FLightSceneInfoCompact>::TConstIterator LightIt(Scene->Lights); LightIt; ++LightIt)
+		{
+			const FLightSceneInfoCompact& LightSceneInfoCompact = *LightIt;
+			FLightSceneInfo* LightSceneInfo = LightSceneInfoCompact.LightSceneInfo;
+			FVisibleLightInfo& VisibleLightInfo = VisibleLightInfos[LightSceneInfo->Id];
+			
+			// Only consider lights that may have shadows.
+			if (LightSceneInfo->ShouldRenderViewIndependentWholeSceneShadows())
+			{
+				// see if the light is visible in any view
+				bool bIsVisibleInAnyView = false;
+
+				for (int32 ViewIndex = 0; ViewIndex < Views.Num(); ViewIndex++)
+				{
+					// View frustums are only checked when lights have visible primitives or have modulated shadows,
+					// so we don't need to check for that again here
+					bIsVisibleInAnyView = LightSceneInfo->ShouldRenderLight(Views[ViewIndex]);
+
+					if (bIsVisibleInAnyView)
+					{
+						break;
+					}
+				}
+
+				if (bIsVisibleInAnyView)
+				{									
+					AddViewDependentWholeSceneShadowsForView(ViewDependentWholeSceneShadows, VisibleLightInfo, *LightSceneInfo);					
+				}
+
+				const int32 NumShadows = VisibleLightInfo.AllProjectedShadows.Num();
+				if (NumShadows > 0)
+				{
+					//create the shadow depth texture and/or surface
+					const FIntPoint ShadowBufferResolution = GSceneRenderTargets.GetShadowDepthTextureResolution();
+					const int32 MaxWide = GMaxShadowDepthBufferSizeX / ShadowBufferResolution.X;
+					const int32 MaxHigh = GMaxShadowDepthBufferSizeY / ShadowBufferResolution.Y;
+
+
+					const int32 NumWide = FMath::Min(NumShadows, MaxWide);
+					const int32 NumHigh = FMath::Min(((NumShadows - 1) / MaxWide) + 1, MaxHigh);
+
+					const FIntPoint AtlasShadowBufferResolution(ShadowBufferResolution.X * NumWide, ShadowBufferResolution.Y * NumHigh);
+					GSceneRenderTargets.AllocateForwardShadingShadowDepthTarget(AtlasShadowBufferResolution);
+
+
+					// Allocate atlas shadow texture space to the shadows.				
+					FTextureLayout ShadowLayout(1, 1, AtlasShadowBufferResolution.X, AtlasShadowBufferResolution.Y, false, false);
+					for (int32 ShadowIndex = 0; ShadowIndex < VisibleLightInfo.AllProjectedShadows.Num(); ShadowIndex++)
+					{
+						FProjectedShadowInfo* ProjectedShadowInfo = VisibleLightInfo.AllProjectedShadows[ShadowIndex];
+						if (!ProjectedShadowInfo->bRendered)
+						{
+							if (ShadowLayout.AddElement(
+								ProjectedShadowInfo->X,
+								ProjectedShadowInfo->Y,
+								ProjectedShadowInfo->ResolutionX + SHADOW_BORDER * 2,
+								ProjectedShadowInfo->ResolutionY + SHADOW_BORDER * 2)
+								)
+							{
+								ProjectedShadowInfo->bAllocated = true;
+							}
+						}
+					}
+				}
+			}
+		}
+
+		// Calculate visibility of the projected shadows.
+		InitProjectedShadowVisibility(RHICmdList);
+	}
+
+	// Gathers the list of primitives used to draw various shadow types
+	GatherShadowPrimitives(PreShadows, ViewDependentWholeSceneShadows, false);
+}
+
 void FDeferredShadingSceneRenderer::InitDynamicShadows(FRHICommandListImmediate& RHICmdList)
 {
 	SCOPE_CYCLE_COUNTER(STAT_DynamicShadowSetupTime);
@@ -2003,87 +2170,7 @@ void FDeferredShadingSceneRenderer::InitDynamicShadows(FRHICommandListImmediate&
 					// Allow movable and stationary lights to create CSM, or static lights that are unbuilt
 					if (!LightSceneInfo->Proxy->HasStaticLighting() || bCreateShadowToPreviewStaticLight)
 					{
-						TArray<float, TInlineAllocator<2> > FadeAlphas;
-						// Allow each view to create a whole scene view dependent shadow
-						for (int32 ViewIndex = 0; ViewIndex < Views.Num(); ViewIndex++)
-						{
-							FViewInfo& View = Views[ViewIndex];
-
-							// If rendering in stereo mode we render shadow depths only for the left eye, but project for both eyes!
-							if (View.StereoPass != eSSP_RIGHT_EYE)
-							{
-								FadeAlphas.Init(0.0f,Views.Num());
-								FadeAlphas[ViewIndex] = 1.0f;
-
-								if (View.StereoPass == eSSP_LEFT_EYE
-									&& Views.IsValidIndex(ViewIndex + 1)
-									&& Views[ViewIndex + 1].StereoPass == eSSP_RIGHT_EYE)
-								{
-									FadeAlphas[ViewIndex + 1] = 1.0f;
-								}
-
-								const int32 NumSplits = LightSceneInfo->Proxy->GetNumViewDependentWholeSceneShadows(View);
-								for (int32 SplitIndex = 0; SplitIndex < NumSplits; SplitIndex++)
-								{
-									FWholeSceneProjectedShadowInitializer ProjectedShadowInitializer;
-
-									if (LightSceneInfo->Proxy->GetViewDependentWholeSceneProjectedShadowInitializer(View, SplitIndex, ProjectedShadowInitializer))
-									{
-										const FIntPoint ShadowBufferResolution = GSceneRenderTargets.GetShadowDepthTextureResolution();
-										// Create the projected shadow info.
-										FProjectedShadowInfo* ProjectedShadowInfo = new(FMemStack::Get(),1,16) FProjectedShadowInfo(
-											LightSceneInfo,
-											&View,
-											ProjectedShadowInitializer,
-											//@todo - remove the shadow border for whole scene shadows
-											ShadowBufferResolution.X - SHADOW_BORDER * 2,
-											ShadowBufferResolution.Y - SHADOW_BORDER * 2,
-											false	// no RSM
-											);
-
-										ProjectedShadowInfo->FadeAlphas = FadeAlphas;
-
-										FVisibleLightInfo& LightViewInfo = VisibleLightInfos[LightSceneInfo->Id];
-										VisibleLightInfo.MemStackProjectedShadows.Add(ProjectedShadowInfo);
-										VisibleLightInfo.AllProjectedShadows.Add(ProjectedShadowInfo);
-										ViewDependentWholeSceneShadows.Add(ProjectedShadowInfo);
-									}
-								}
-
-								FSceneViewState* ViewState = (FSceneViewState*)View.State;
-								if(ViewState)
-								{
-									FLightPropagationVolume* LightPropagationVolume = ViewState->GetLightPropagationVolume();
-
-									if (LightPropagationVolume && View.FinalPostProcessSettings.LPVIntensity > 0 )
-									{
-										// Generate the RSM shadow info
-										FWholeSceneProjectedShadowInitializer ProjectedShadowInitializer;
-										FLightPropagationVolume& Lpv = *LightPropagationVolume;
-
-										if (LightSceneInfo->Proxy->GetViewDependentRsmWholeSceneProjectedShadowInitializer( View, Lpv.GetBoundingBox(), ProjectedShadowInitializer ))
-										{
-											const FIntPoint ShadowBufferResolution = GSceneRenderTargets.GetReflectiveShadowMapTextureResolution();
-
-											// Create the projected shadow info.
-											FProjectedShadowInfo* ProjectedShadowInfo = new(FMemStack::Get(),1,16) FProjectedShadowInfo(
-												LightSceneInfo,
-												&View,
-												ProjectedShadowInitializer,
-												ShadowBufferResolution.X,
-												ShadowBufferResolution.Y,
-												true);
-
-											FVisibleLightInfo& LightViewInfo = VisibleLightInfos[LightSceneInfo->Id];
-											VisibleLightInfo.MemStackProjectedShadows.Add(ProjectedShadowInfo);
-											VisibleLightInfo.AllProjectedShadows.Add(ProjectedShadowInfo);
-											VisibleLightInfo.ReflectiveShadowMaps.Add(ProjectedShadowInfo); 
-											ViewDependentWholeSceneShadows.Add(ProjectedShadowInfo); // or separate list?
-										}
-									}
-								}
-							}
-						}
+						AddViewDependentWholeSceneShadowsForView(ViewDependentWholeSceneShadows, VisibleLightInfo, *LightSceneInfo);
 
 						// Look for individual primitives with a dynamic shadow.
 						for (FLightPrimitiveInteraction* Interaction = LightSceneInfo->DynamicPrimitiveList;
diff --git a/Engine/Source/Runtime/Renderer/Private/TranslucentLighting.cpp b/Engine/Source/Runtime/Renderer/Private/TranslucentLighting.cpp
--- a/Engine/Source/Runtime/Renderer/Private/TranslucentLighting.cpp
+++ b/Engine/Source/Runtime/Renderer/Private/TranslucentLighting.cpp
@@ -516,6 +516,7 @@ void FProjectedShadowInfo::RenderTranslucencyDepths(FRHICommandList& RHICmdList,
 	FoundView->ViewMatrices.ViewMatrix = ShadowViewMatrix;
 	FBox VolumeBounds[TVC_MAX];
 	FoundView->UniformBuffer = FoundView->CreateUniformBuffer(
+		nullptr,
 		ShadowViewMatrix, 
 		VolumeBounds,
 		TVC_MAX);
