UE4 - Forward cascaded shadows
* Fixes blackness on ES2
* Fix uniform buffer alignment issue

diff --git a/Engine/Source/Developer/iOS/MetalShaderFormat/Private/MetalShaderCompiler.cpp b/Engine/Source/Developer/iOS/MetalShaderFormat/Private/MetalShaderCompiler.cpp
--- a/Engine/Source/Developer/iOS/MetalShaderFormat/Private/MetalShaderCompiler.cpp
+++ b/Engine/Source/Developer/iOS/MetalShaderFormat/Private/MetalShaderCompiler.cpp
@@ -1,4 +1,4 @@
-// Copyright 1998-2014 Epic Games, Inc. All Rights Reserved..
+// Copyright 1998-2014 Epic Games, Inc. All Rights Reserved. 
 // 
 
 #include "MetalShaderFormat.h"
@@ -621,7 +621,7 @@ static void BuildMetalShaderOutput(
 		Header.Bindings.NumUniformBuffers = FMath::Max((uint8)GetNumUniformBuffersUsed(GenericSRT), Header.Bindings.NumUniformBuffers);
 	}
 
-	const int32 MaxSamplers = GetFeatureLevelMaxTextureSamplers(GetMaxSupportedFeatureLevel((EShaderPlatform)ShaderOutput.Target.Platform));
+	const int32 MaxSamplers = GetFeatureLevelMaxTextureSamplers(ERHIFeatureLevel::ES3_1);
 
 	if (Header.Bindings.NumSamplers > MaxSamplers)
 	{
@@ -821,7 +821,7 @@ void CompileShader_Metal(const FShaderCompilerInput& Input,FShaderCompilerOutput
 	EHlslCompileTarget HlslCompilerTarget = HCT_FeatureLevelES3_1;
 	AdditionalDefines.SetDefine(TEXT("IOS"), 1);
 	AdditionalDefines.SetDefine(TEXT("METAL_PROFILE"), 1);
-	AdditionalDefines.SetDefine(TEXT("ES2_PROFILE"), 1);
+	AdditionalDefines.SetDefine(TEXT("ES3_1_PROFILE"), 1);
 	AdditionalDefines.SetDefine(TEXT("row_major"), TEXT(""));
 	
 	const bool bDumpDebugInfo = (Input.DumpDebugInfoPath != TEXT("") && IFileManager::Get().DirectoryExists(*Input.DumpDebugInfoPath));
diff --git a/Engine/Source/Developer/iOS/MetalShaderFormat/Private/MetalShaderFormat.cpp b/Engine/Source/Developer/iOS/MetalShaderFormat/Private/MetalShaderFormat.cpp
--- a/Engine/Source/Developer/iOS/MetalShaderFormat/Private/MetalShaderFormat.cpp
+++ b/Engine/Source/Developer/iOS/MetalShaderFormat/Private/MetalShaderFormat.cpp
@@ -14,7 +14,7 @@ class FMetalShaderFormat : public IShaderFormat
 {
 	enum
 	{
-		HEADER_VERSION = 7,
+		HEADER_VERSION = 8,
 	};
 	struct FVersion
 	{
diff --git a/Engine/Source/Runtime/IOS/MetalRHI/Private/MetalManager.cpp b/Engine/Source/Runtime/IOS/MetalRHI/Private/MetalManager.cpp
--- a/Engine/Source/Runtime/IOS/MetalRHI/Private/MetalManager.cpp
+++ b/Engine/Source/Runtime/IOS/MetalRHI/Private/MetalManager.cpp
@@ -36,7 +36,8 @@ static uint32 RTBitOffsets[] = { OFFSET_RENDER_TARGET_FORMAT0, OFFSET_RENDER_TAR
 	}
 #define GET_HASH(Offset, NumBits) ((Hash >> Offset) & ((1ULL << NumBits) - 1))
 
-
+//@todo-rco: HACK! This needs to exist until we merge SetRenderTargets()+Clear()
+extern int GMetalHackDepthActions;
 
 
 /*
@@ -628,8 +629,18 @@ void FMetalManager::UpdateContext()
 		// set up the depth attachment
 		DepthAttachment.texture = CurrentDepthRenderTexture;
 		[DepthAttachment setLoadAction:MTLLoadActionClear];
-		[DepthAttachment setStoreAction:MTLStoreActionDontCare];
-		[DepthAttachment setClearDepth:0.0];
+		if (GMetalHackDepthActions == 1)
+		{
+			// For Shadow Maps, clear to 1.0f and save the render target
+			[DepthAttachment setStoreAction:MTLStoreActionStore];
+			[DepthAttachment setClearDepth:1.0f];
+		}
+		else if (GMetalHackDepthActions == 0)
+		{
+			// Regular operation, clear to zero (due to inverted Z) and don't store it
+			[DepthAttachment setStoreAction:MTLStoreActionDontCare];
+			[DepthAttachment setClearDepth:0.0f];
+		}
 
 		Pipeline.DepthTargetFormat = CurrentDepthRenderTexture.pixelFormat;
 		if (Pipeline.SampleCount == 0)
diff --git a/Engine/Source/Runtime/Renderer/Private/ShadowRendering.cpp b/Engine/Source/Runtime/Renderer/Private/ShadowRendering.cpp
--- a/Engine/Source/Runtime/Renderer/Private/ShadowRendering.cpp
+++ b/Engine/Source/Runtime/Renderer/Private/ShadowRendering.cpp
@@ -3317,6 +3317,12 @@ bool FDeferredShadingSceneRenderer::RenderCachedPreshadows(FRHICommandListImmedi
 	return bAttenuationBufferDirty;
 }
 
+#if PLATFORM_IOS
+// 0: Regular operation, clear to zero (due to inverted Z) and don't store it
+// 1: For Shadow Maps, clear to 1.0f and save the render target
+int GMetalHackDepthActions = 0;
+#endif
+
 /**
  * Used by RenderLights to render shadows to the attenuation buffer.
  *
@@ -3325,13 +3331,6 @@ bool FDeferredShadingSceneRenderer::RenderCachedPreshadows(FRHICommandListImmedi
  */
 bool FForwardShadingSceneRenderer::RenderShadowDepthMap(FRHICommandListImmediate& RHICmdList, const FLightSceneInfo* LightSceneInfo)
 {
-	//@todo-rco: Disabled
-	bool bCurrentlyDisabled = true;
-	if (bCurrentlyDisabled)
-	{
-		return false;
-	}
-
 	SCOPE_CYCLE_COUNTER(STAT_ProjectedShadowDrawTime);	
 
 	bool bAttenuationBufferDirty = false;
@@ -3396,7 +3395,16 @@ bool FForwardShadingSceneRenderer::RenderShadowDepthMap(FRHICommandListImmediate
 	{
 		// Render the shadow depths.
 		SCOPED_DRAW_EVENT(RHICmdList, ShadowDepthsFromOpaque, DEC_SCENE_ITEMS);
+#if PLATFORM_IOS
+		// WILL BE FIXED WHEN WE HAVE SET & CLEAR RT IN THE SAME API! Also a way to tell it to be resolved!
+		GMetalHackDepthActions = 1;
+#endif
 		GSceneRenderTargets.BeginRenderingShadowDepth(RHICmdList);
+#if PLATFORM_IOS
+		// WILL BE FIXED WHEN WE HAVE SET & CLEAR RT IN THE SAME API!
+		GMetalHackDepthActions = 0;
+#endif
+
 		RHICmdList.Clear(false, FColor(255, 255, 255), true, 1.0f, false, 0, FIntRect());
 
 		// render depth for each shadow
