try to prevent a race condition in IsDBufferEnabled

diff --git a/Engine/Source/Runtime/Renderer/Private/DeferredShadingRenderer.cpp b/Engine/Source/Runtime/Renderer/Private/DeferredShadingRenderer.cpp
--- a/Engine/Source/Runtime/Renderer/Private/DeferredShadingRenderer.cpp
+++ b/Engine/Source/Runtime/Renderer/Private/DeferredShadingRenderer.cpp
@@ -607,6 +607,7 @@ void FDeferredShadingSceneRenderer::RenderFinish(FRHICommandListImmediate& RHICm
 
 void FDeferredShadingSceneRenderer::Render(FRHICommandListImmediate& RHICmdList)
 {
+	bool bDBuffer = IsDBufferEnabled();
 	static FGraphEventRef OcclusionSubmittedFence;
 	if (GRHIThread)
 	{
@@ -734,7 +735,7 @@ void FDeferredShadingSceneRenderer::Render(FRHICommandListImmediate& RHICmdList)
 	check(!GSceneRenderTargets.DBufferB);
 	check(!GSceneRenderTargets.DBufferC);
 
-	if(IsDBufferEnabled())
+	if(bDBuffer)
 	{
 		GSceneRenderTargets.ResolveSceneDepthTexture(RHICmdList);
 		GSceneRenderTargets.ResolveSceneDepthToAuxiliaryTexture(RHICmdList);
