UE4 - parallized the relevance part of init views, reducing that time to about 1/3 of its previous value. This isn't on the critical path for rhi thread platforms, and is currently disabled.

diff --git a/Engine/Plugins/Experimental/Paper2D/Source/Paper2D/Private/Rendering/PaperBatchSceneProxy.cpp b/Engine/Plugins/Experimental/Paper2D/Source/Paper2D/Private/Rendering/PaperBatchSceneProxy.cpp
--- a/Engine/Plugins/Experimental/Paper2D/Source/Paper2D/Private/Rendering/PaperBatchSceneProxy.cpp
+++ b/Engine/Plugins/Experimental/Paper2D/Source/Paper2D/Private/Rendering/PaperBatchSceneProxy.cpp
@@ -20,7 +20,7 @@ uint32 FPaperBatchSceneProxy::GetMemoryFootprint() const
 
 FPrimitiveViewRelevance FPaperBatchSceneProxy::GetViewRelevance(const FSceneView* View)
 {
-	checkSlow(IsInRenderingThread());
+	checkSlow(IsInParallelRenderingThread());
 
 	FPrimitiveViewRelevance Result;
 	Result.bDrawRelevance = IsShown(View) && View->Family->EngineShowFlags.Paper2DSprites;
diff --git a/Engine/Source/Runtime/Core/Public/Containers/BitArray.h b/Engine/Source/Runtime/Core/Public/Containers/BitArray.h
--- a/Engine/Source/Runtime/Core/Public/Containers/BitArray.h
+++ b/Engine/Source/Runtime/Core/Public/Containers/BitArray.h
@@ -60,6 +60,39 @@ public:
 			Data &= ~Mask;
 		}
 	}
+	FORCEINLINE void AtomicSet(const bool NewValue)
+	{
+		if(NewValue)
+		{
+			if (!(Data & Mask))
+			{
+				while (1)
+				{
+					uint32 Current = Data;
+					uint32 Desired = Current | Mask;
+					if (Current == Desired || FPlatformAtomics::InterlockedCompareExchange((volatile int32*)&Data, (int32)Desired, (int32)Current) == (int32)Current)
+					{
+						return;
+					}
+				}
+			}
+		}
+		else
+		{
+			if (Data & Mask)
+			{
+				while (1)
+				{
+					uint32 Current = Data;
+					uint32 Desired = Current & ~Mask;
+					if (Current == Desired || FPlatformAtomics::InterlockedCompareExchange((volatile int32*)&Data, (int32)Desired, (int32)Current) == (int32)Current)
+					{
+						return;
+					}
+				}
+			}
+		}
+	}
 	FORCEINLINE FBitReference& operator=(const FBitReference& Copy)
 	{
 		this->Data = Copy.Data;
diff --git a/Engine/Source/Runtime/Engine/Private/SceneManagement.cpp b/Engine/Source/Runtime/Engine/Private/SceneManagement.cpp
--- a/Engine/Source/Runtime/Engine/Private/SceneManagement.cpp
+++ b/Engine/Source/Runtime/Engine/Private/SceneManagement.cpp
@@ -247,7 +247,7 @@ int8 ComputeStaticMeshLOD( const FStaticMeshRenderData* RenderData, const FVecto
 	return 0;
 }
 
-int8 ComputeLODForMeshes( const TIndirectArray<class FStaticMesh>& StaticMeshes, FSceneView& View, const FVector4& Origin, float SphereRadius, int32 ForcedLODLevel, float ScreenSizeScale )
+int8 ComputeLODForMeshes( const TIndirectArray<class FStaticMesh>& StaticMeshes, const FSceneView& View, const FVector4& Origin, float SphereRadius, int32 ForcedLODLevel, float ScreenSizeScale )
 {
 	int8 LODToRender = 0;
 
diff --git a/Engine/Source/Runtime/Engine/Public/SceneManagement.h b/Engine/Source/Runtime/Engine/Public/SceneManagement.h
--- a/Engine/Source/Runtime/Engine/Public/SceneManagement.h
+++ b/Engine/Source/Runtime/Engine/Public/SceneManagement.h
@@ -1816,4 +1816,4 @@ int8 ENGINE_API ComputeStaticMeshLOD(const FStaticMeshRenderData* RenderData, co
  * @param Origin - Origin of the bounds of the mesh in world space
  * @param SphereRadius - Radius of the sphere to use to calculate screen coverage
  */
-int8 ENGINE_API ComputeLODForMeshes(const TIndirectArray<class FStaticMesh>& StaticMeshes, FSceneView& View, const FVector4& Origin, float SphereRadius, int32 ForcedLODLevel, float ScreenSizeScale = 1.0f);
+int8 ENGINE_API ComputeLODForMeshes(const TIndirectArray<class FStaticMesh>& StaticMeshes, const FSceneView& View, const FVector4& Origin, float SphereRadius, int32 ForcedLODLevel, float ScreenSizeScale = 1.0f);
diff --git a/Engine/Source/Runtime/RenderCore/Private/RenderCore.cpp b/Engine/Source/Runtime/RenderCore/Private/RenderCore.cpp
--- a/Engine/Source/Runtime/RenderCore/Private/RenderCore.cpp
+++ b/Engine/Source/Runtime/RenderCore/Private/RenderCore.cpp
@@ -77,6 +77,7 @@ DEFINE_STAT(STAT_PreRenderView);
 DEFINE_STAT(STAT_GetDynamicMeshElements);
 DEFINE_STAT(STAT_UpdateStaticMeshesTime);
 DEFINE_STAT(STAT_StaticRelevance);
+DEFINE_STAT(STAT_ViewRelevance);
 DEFINE_STAT(STAT_ComputeViewRelevance);
 DEFINE_STAT(STAT_OcclusionCull);
 DEFINE_STAT(STAT_UpdatePrimitiveFading);
diff --git a/Engine/Source/Runtime/RenderCore/Public/RenderCore.h b/Engine/Source/Runtime/RenderCore/Public/RenderCore.h
--- a/Engine/Source/Runtime/RenderCore/Public/RenderCore.h
+++ b/Engine/Source/Runtime/RenderCore/Public/RenderCore.h
@@ -45,6 +45,7 @@ DECLARE_CYCLE_STAT_EXTERN(TEXT("Decompress Occlusion"),STAT_DecompressPrecompute
 DECLARE_CYCLE_STAT_EXTERN(TEXT("Frustum Cull"),STAT_FrustumCull,STATGROUP_InitViews, RENDERCORE_API);
 DECLARE_CYCLE_STAT_EXTERN(TEXT("Update Fading"),STAT_UpdatePrimitiveFading,STATGROUP_InitViews, RENDERCORE_API);
 DECLARE_CYCLE_STAT_EXTERN(TEXT("Occlusion Cull"),STAT_OcclusionCull,STATGROUP_InitViews, RENDERCORE_API);
+DECLARE_CYCLE_STAT_EXTERN(TEXT("View Relevance"),STAT_ViewRelevance,STATGROUP_InitViews, RENDERCORE_API);
 DECLARE_CYCLE_STAT_EXTERN(TEXT("Compute View Relevance"),STAT_ComputeViewRelevance,STATGROUP_InitViews, RENDERCORE_API);
 DECLARE_CYCLE_STAT_EXTERN(TEXT("Static Mesh Relevance"),STAT_StaticRelevance,STATGROUP_InitViews, RENDERCORE_API);
 DECLARE_CYCLE_STAT_EXTERN(TEXT("UpdateStaticMeshes"),STAT_UpdateStaticMeshesTime,STATGROUP_InitViews, RENDERCORE_API);
diff --git a/Engine/Source/Runtime/Renderer/Private/CustomDepthRendering.h b/Engine/Source/Runtime/Renderer/Private/CustomDepthRendering.h
--- a/Engine/Source/Runtime/Renderer/Private/CustomDepthRendering.h
+++ b/Engine/Source/Runtime/Renderer/Private/CustomDepthRendering.h
@@ -29,6 +29,14 @@ public:
 	{
 		Prims.Add(PrimitiveSceneProxy);
 	}
+	/**
+	* Adds a new primitives to the list of distortion prims
+	* @param PrimitiveSceneProxies - primitive info to add.
+	*/
+	void Append(FPrimitiveSceneProxy** PrimitiveSceneProxies, int32 NumProxies)
+	{
+		Prims.Append(PrimitiveSceneProxies, NumProxies);
+	}
 
 	/** 
 	* @return number of prims to render
diff --git a/Engine/Source/Runtime/Renderer/Private/DistortionRendering.cpp b/Engine/Source/Runtime/Renderer/Private/DistortionRendering.cpp
--- a/Engine/Source/Runtime/Renderer/Private/DistortionRendering.cpp
+++ b/Engine/Source/Runtime/Renderer/Private/DistortionRendering.cpp
@@ -771,11 +771,6 @@ bool FDistortionPrimSet::DrawAccumulatedOffsets(FRHICommandListImmediate& RHICmd
 	return bDirty;
 }
 
-void FDistortionPrimSet::AddScenePrimitive(FPrimitiveSceneProxy* PrimitiveSceneProxy,const FViewInfo& ViewInfo)
-{
-	Prims.Add(PrimitiveSceneProxy);
-}
-
 /** 
  * Renders the scene's distortion 
  */
diff --git a/Engine/Source/Runtime/Renderer/Private/DistortionRendering.h b/Engine/Source/Runtime/Renderer/Private/DistortionRendering.h
--- a/Engine/Source/Runtime/Renderer/Private/DistortionRendering.h
+++ b/Engine/Source/Runtime/Renderer/Private/DistortionRendering.h
@@ -24,9 +24,19 @@ public:
 	/**
 	* Add a new primitive to the list of distortion prims
 	* @param PrimitiveSceneProxy - primitive info to add.
-	* @param ViewInfo - used to transform bounds to view space
 	*/
-	void AddScenePrimitive(FPrimitiveSceneProxy* PrimitiveSceneProxy,const FViewInfo& ViewInfo);
+	void AddScenePrimitive(FPrimitiveSceneProxy* PrimitiveSceneProxy)
+	{
+		Prims.Add(PrimitiveSceneProxy);
+	}
+	/**
+	* Adds a new primitives to the list of distortion prims
+	* @param PrimitiveSceneProxies - primitive info to add.
+	*/
+	void Append(FPrimitiveSceneProxy** PrimitiveSceneProxies, int32 NumProxies)
+	{
+		Prims.Append(PrimitiveSceneProxies, NumProxies);
+	}
 
 	/** 
 	* @return number of prims to render
diff --git a/Engine/Source/Runtime/Renderer/Private/PrimitiveSceneInfo.cpp b/Engine/Source/Runtime/Renderer/Private/PrimitiveSceneInfo.cpp
--- a/Engine/Source/Runtime/Renderer/Private/PrimitiveSceneInfo.cpp
+++ b/Engine/Source/Runtime/Renderer/Private/PrimitiveSceneInfo.cpp
@@ -293,18 +293,16 @@ void FPrimitiveSceneInfo::RemoveFromScene(bool bUpdateStaticDrawLists)
 	}
 }
 
-void FPrimitiveSceneInfo::ConditionalUpdateStaticMeshes(FRHICommandListImmediate& RHICmdList)
+void FPrimitiveSceneInfo::UpdateStaticMeshes(FRHICommandListImmediate& RHICmdList)
 {
-	if (bNeedsStaticMeshUpdate)
-	{
-		bNeedsStaticMeshUpdate = false;
+	checkSlow(bNeedsStaticMeshUpdate);
+	bNeedsStaticMeshUpdate = false;
 
-		// Remove the primitive's static meshes from the draw lists they're currently in, and re-add them to the appropriate draw lists.
-		for (int32 MeshIndex = 0; MeshIndex < StaticMeshes.Num(); MeshIndex++)
-		{
-			StaticMeshes[MeshIndex].RemoveFromDrawLists();
-			StaticMeshes[MeshIndex].AddToDrawLists(RHICmdList, Scene);
-		}
+	// Remove the primitive's static meshes from the draw lists they're currently in, and re-add them to the appropriate draw lists.
+	for (int32 MeshIndex = 0; MeshIndex < StaticMeshes.Num(); MeshIndex++)
+	{
+		StaticMeshes[MeshIndex].RemoveFromDrawLists();
+		StaticMeshes[MeshIndex].AddToDrawLists(RHICmdList, Scene);
 	}
 }
 
@@ -414,17 +412,19 @@ uint32 FPrimitiveSceneInfo::GetMemoryFootprint()
 	return( sizeof( *this ) + HitProxies.GetAllocatedSize() + StaticMeshes.GetAllocatedSize() );
 }
 
-bool FPrimitiveSceneInfo::ShouldRenderVelocity(const FViewInfo& View) const
+bool FPrimitiveSceneInfo::ShouldRenderVelocity(const FViewInfo& View, bool bCheckVisibility) const
 {
 	int32 PrimitiveId = GetIndex();
-	const bool bVisible = View.PrimitiveVisibilityMap[PrimitiveId];
-
-	// Only render if visible.
-	if(!bVisible)
+	if (bCheckVisibility)
 	{
-		return false;
-	}
+		const bool bVisible = View.PrimitiveVisibilityMap[PrimitiveId];
 
+		// Only render if visible.
+		if(!bVisible)
+		{
+			return false;
+		}
+	}
 	// Used to determine whether object is movable or not.
 	if(!Proxy->IsMovable())
 	{
diff --git a/Engine/Source/Runtime/Renderer/Private/PrimitiveSceneInfo.h b/Engine/Source/Runtime/Renderer/Private/PrimitiveSceneInfo.h
--- a/Engine/Source/Runtime/Renderer/Private/PrimitiveSceneInfo.h
+++ b/Engine/Source/Runtime/Renderer/Private/PrimitiveSceneInfo.h
@@ -234,8 +234,23 @@ public:
 	/** Removes the primitive from the scene. */
 	void RemoveFromScene(bool bUpdateStaticDrawLists);
 
+	/** return true if we need to call ConditionalUpdateStaticMeshes */
+	FORCEINLINE bool NeedsUpdateStaticMeshes()
+	{
+		return bNeedsStaticMeshUpdate;
+	}
+
+	/** Updates the primitive's static meshes in the scene. */
+	void UpdateStaticMeshes(FRHICommandListImmediate& RHICmdList);
+
 	/** Updates the primitive's static meshes in the scene. */
-	void ConditionalUpdateStaticMeshes(FRHICommandListImmediate& RHICmdList);
+	FORCEINLINE void ConditionalUpdateStaticMeshes(FRHICommandListImmediate& RHICmdList)
+	{
+		if (NeedsUpdateStaticMeshes())
+		{
+			UpdateStaticMeshes(RHICmdList);
+		}
+	}
 
 	/** Sets a flag to update the primitive's static meshes before it is next rendered. */
 	void BeginDeferredUpdateStaticMeshes();
@@ -278,7 +293,7 @@ public:
 	inline int32 GetIndex() const { return PackedIndex; }
 
 	/* @return true if the object needs to be rendered in the velocity pass (is not moving like the world, needed for motionblur and TemporalAA) */
-	bool ShouldRenderVelocity(const FViewInfo& View) const;
+	bool ShouldRenderVelocity(const FViewInfo& View, bool bCheckVisibility = true) const;
 	
 	/**
 	 * Shifts primitive position and all relevant data by an arbitrary delta.
diff --git a/Engine/Source/Runtime/Renderer/Private/RendererScene.cpp b/Engine/Source/Runtime/Renderer/Private/RendererScene.cpp
--- a/Engine/Source/Runtime/Renderer/Private/RendererScene.cpp
+++ b/Engine/Source/Runtime/Renderer/Private/RendererScene.cpp
@@ -891,7 +891,7 @@ void FScene::ReleaseReflectionCubemap(UReflectionCaptureComponent* CaptureCompon
 
 const FReflectionCaptureProxy* FScene::FindClosestReflectionCapture(FVector Position) const
 {
-	checkSlow(IsInRenderingThread());
+	checkSlow(IsInParallelRenderingThread());
 	int32 ClosestCaptureIndex = INDEX_NONE;
 	float ClosestDistanceSquared = FLT_MAX;
 
diff --git a/Engine/Source/Runtime/Renderer/Private/SceneRendering.h b/Engine/Source/Runtime/Renderer/Private/SceneRendering.h
--- a/Engine/Source/Runtime/Renderer/Private/SceneRendering.h
+++ b/Engine/Source/Runtime/Renderer/Private/SceneRendering.h
@@ -102,6 +102,11 @@ public:
 	*/
 	void AddScenePrimitive(FPrimitiveSceneInfo* PrimitiveSceneInfo, const FViewInfo& ViewInfo, bool bUseNormalTranslucency, bool bUseSeparateTranslucency);
 
+	/**
+	* Similar to AddScenePrimitive, but we are doing placement news and increasing counts when that happens
+	*/
+	static void PlaceScenePrimitive(FPrimitiveSceneInfo* PrimitiveSceneInfo, const FViewInfo& ViewInfo, bool bUseNormalTranslucency, bool bUseSeparateTranslucency, void *NormalPlace, int32& NormalNum, void* SeparatePlace, int32& SeparateNum);
+
 	/**
 	* Sort any primitives that were added to the set back-to-front
 	*/
@@ -132,8 +137,6 @@ public:
 		return SortedSeparateTranslucencyPrims[i].PrimitiveSceneInfo;
 	}
 
-private:
-
 	/** contains a sort key */
 	struct FDepthSortedPrim
 	{
@@ -165,6 +168,14 @@ private:
 		int32 SortPriority;
 	};
 
+	/**
+	* Adds primitives originally created with PlaceScenePrimitive
+	*/
+	void AppendScenePrimitives(FSortedPrim* Normal, int32 NumNormal, FSortedPrim* Separate, int32 NumSeparate);
+
+
+private:
+
 	/** sortkey compare class */
 	struct FCompareFDepthSortedPrim
 	{
@@ -193,6 +204,9 @@ private:
 	void RenderPrimitive(FRHICommandList& RHICmdList, const FViewInfo& View, FPrimitiveSceneInfo* PrimitiveSceneInfo, const FPrimitiveViewRelevance& ViewRelevance, const FProjectedShadowInfo* TranslucentSelfShadow, bool bSeparateTranslucencyPass) const;
 };
 
+template <> struct TIsPODType<FTranslucentPrimSet::FDepthSortedPrim> { enum { Value = true }; };
+template <> struct TIsPODType<FTranslucentPrimSet::FSortedPrim> { enum { Value = true }; };
+
 /** A batched occlusion primitive. */
 struct FOcclusionPrimitive
 {
diff --git a/Engine/Source/Runtime/Renderer/Private/SceneVisibility.cpp b/Engine/Source/Runtime/Renderer/Private/SceneVisibility.cpp
--- a/Engine/Source/Runtime/Renderer/Private/SceneVisibility.cpp
+++ b/Engine/Source/Runtime/Renderer/Private/SceneVisibility.cpp
@@ -658,7 +658,6 @@ static void ComputeRelevanceForView(
 	float CurrentWorldTime = View.Family->CurrentWorldTime;
 	float DeltaWorldTime = View.Family->DeltaWorldTime;
 
-	//OutRelevantStaticPrimitives.Empty(NumPrimitives);
 	for (FSceneSetBitIterator BitIt(View.PrimitiveVisibilityMap); BitIt; ++BitIt)
 	{
 		FPrimitiveSceneInfo* PrimitiveSceneInfo = Scene->Primitives[BitIt.GetIndex()];
@@ -714,7 +713,7 @@ static void ComputeRelevanceForView(
 			if (ViewRelevance.bDistortionRelevance)
 			{
 				// Add to set of dynamic distortion primitives
-				View.DistortionPrimSet.AddScenePrimitive(PrimitiveSceneInfo->Proxy, View);
+				View.DistortionPrimSet.AddScenePrimitive(PrimitiveSceneInfo->Proxy);
 			}
 		}
 		
@@ -764,6 +763,565 @@ static void ComputeRelevanceForView(
 	}
 }
 
+template<class T>
+struct FRelevancePrimSet
+{
+	enum
+	{
+		MaxPrims = 127 //like 128, but we leave space for NumPrims
+	};
+	int32 NumPrims;
+	T Prims[MaxPrims];
+	FORCEINLINE FRelevancePrimSet()
+		: NumPrims(0)
+	{
+		//FMemory::Memzero(Prims, sizeof(T) * MaxPrims);
+	}
+	FORCEINLINE void AddPrim(T Prim)
+	{
+		checkSlow(NumPrims < MaxPrims);
+		Prims[NumPrims++] = Prim;
+	}
+	FORCEINLINE bool IsFull() const
+	{
+		return NumPrims >= MaxPrims;
+	}
+	template<class TARRAY>
+	FORCEINLINE void AppendTo(TARRAY& Array)
+	{
+		Array.Append(Prims, NumPrims);
+	}
+};
+
+struct FMarkRelevantStaticMeshesForViewData
+{
+	FVector ViewOrigin;
+	float MaxDrawDistanceScaleSquared;
+	int32 ForcedLODLevel;
+	float LODScale;
+	float InvLODScale;
+	float MinScreenRadiusForCSMDepthSquared;
+	float MinScreenRadiusForDepthPrepassSquared;
+	bool bForceEarlyZPass;
+
+	FMarkRelevantStaticMeshesForViewData(FViewInfo& View)
+	{
+		ViewOrigin = View.ViewMatrices.ViewOrigin;
+
+		MaxDrawDistanceScaleSquared = GetCachedScalabilityCVars().ViewDistanceScaleSquared;
+
+		// outside of the loop to be more efficient
+		ForcedLODLevel = (View.Family->EngineShowFlags.LOD) ? GetCVarForceLOD() : 0;
+
+		LODScale = CVarStaticMeshLODDistanceScale.GetValueOnRenderThread();
+		InvLODScale = 1.0f / LODScale;
+
+		MinScreenRadiusForCSMDepthSquared = GMinScreenRadiusForCSMDepth * GMinScreenRadiusForCSMDepth;
+		MinScreenRadiusForDepthPrepassSquared = GMinScreenRadiusForDepthPrepass * GMinScreenRadiusForDepthPrepass;
+
+		extern TAutoConsoleVariable<int32> CVarEarlyZPass;
+		bForceEarlyZPass = CVarEarlyZPass.GetValueOnRenderThread() == 2;
+	}
+};
+
+namespace EMarkMaskBits
+{
+	enum Type
+	{
+		StaticMeshShadowDepthMapMask = 0x1,
+		StaticMeshVisibilityMapMask = 0x2,
+		StaticMeshVelocityMapMask = 0x4,
+		StaticMeshOccluderMapMask = 0x8,
+	};
+}
+
+struct FRelevancePacket
+{
+	const float CurrentWorldTime;
+	const float DeltaWorldTime;
+
+	FRHICommandListImmediate& RHICmdList;
+	const FScene* Scene;
+	const FViewInfo& View;
+	const uint8 ViewBit;
+	const FMarkRelevantStaticMeshesForViewData& ViewData;
+	FPrimitiveViewMasks& OutPreRenderViewMasks;
+	FPrimitiveViewMasks& OutHasDynamicMeshElementsMasks;
+	FPrimitiveViewMasks& OutHasDynamicEditorMeshElementsMasks;
+	uint8* RESTRICT MarkMasks;
+
+	FRelevancePrimSet<int32> Input;
+	FRelevancePrimSet<int32> RelevantStaticPrimitives;
+	FRelevancePrimSet<int32> NotDrawRelevant;
+	FRelevancePrimSet<FPrimitiveSceneInfo*> VisibleDynamicPrimitives;
+	FRelevancePrimSet<FTranslucentPrimSet::FSortedPrim> SortedSeparateTranslucencyPrims;
+	FRelevancePrimSet<FTranslucentPrimSet::FSortedPrim> SortedTranslucencyPrims;
+	FRelevancePrimSet<FPrimitiveSceneProxy*> DistortionPrimSet;
+	FRelevancePrimSet<FPrimitiveSceneProxy*> CustomDepthSet;
+	FRelevancePrimSet<FPrimitiveSceneInfo*> UpdateStaticMeshes;
+	FRelevancePrimSet<FPrimitiveSceneInfo*> VisibleEditorPrimitives;
+
+	FRelevancePacket(
+		FRHICommandListImmediate& InRHICmdList,
+		const FScene* InScene, 
+		const FViewInfo& InView, 
+		uint8 InViewBit,
+		const FMarkRelevantStaticMeshesForViewData& InViewData,
+		FPrimitiveViewMasks& InOutPreRenderViewMasks,
+		FPrimitiveViewMasks& InOutHasDynamicMeshElementsMasks,
+		FPrimitiveViewMasks& InOutHasDynamicEditorMeshElementsMasks,
+		uint8* InMarkMasks)
+
+		: CurrentWorldTime(InView.Family->CurrentWorldTime)
+		, DeltaWorldTime(InView.Family->DeltaWorldTime)
+		, RHICmdList(InRHICmdList)
+		, Scene(InScene)
+		, View(InView)
+		, ViewBit(InViewBit)
+		, ViewData(InViewData)
+		, OutPreRenderViewMasks(InOutPreRenderViewMasks)
+		, OutHasDynamicMeshElementsMasks(InOutHasDynamicMeshElementsMasks)
+		, OutHasDynamicEditorMeshElementsMasks(InOutHasDynamicEditorMeshElementsMasks)
+		, MarkMasks(InMarkMasks)
+	{
+	}
+
+	void AnyThreadTask()
+	{
+		ComputeRelevance();
+		MarkRelevant();
+	}
+
+	void ComputeRelevance()
+	{
+		SCOPE_CYCLE_COUNTER(STAT_ComputeViewRelevance);
+		for (int32 Index = 0; Index < Input.NumPrims; Index++)
+		{
+			int32 BitIndex = Input.Prims[Index];
+			FPrimitiveSceneInfo* PrimitiveSceneInfo = Scene->Primitives[BitIndex];
+			FPrimitiveViewRelevance& ViewRelevance = const_cast<FPrimitiveViewRelevance&>(View.PrimitiveViewRelevanceMap[BitIndex]);
+			ViewRelevance = PrimitiveSceneInfo->Proxy->GetViewRelevance(&View);
+			ViewRelevance.bInitializedThisFrame = true;
+
+			const bool bStaticRelevance = ViewRelevance.bStaticRelevance;
+			const bool bDrawRelevance = ViewRelevance.bDrawRelevance;
+			const bool bDynamicRelevance = ViewRelevance.bDynamicRelevance;
+			const bool bShadowRelevance = ViewRelevance.bShadowRelevance;
+			const bool bEditorRelevance = ViewRelevance.bEditorPrimitiveRelevance;
+			const bool bTranslucentRelevance = ViewRelevance.HasTranslucency();
+
+			if (bStaticRelevance && (bDrawRelevance || bShadowRelevance))
+			{
+				RelevantStaticPrimitives.AddPrim(BitIndex);
+			}
+
+			if (!bDrawRelevance)
+			{
+				NotDrawRelevant.AddPrim(BitIndex);
+				continue;
+			}
+
+			if (ViewRelevance.bNeedsPreRenderView)
+			{
+				OutPreRenderViewMasks[BitIndex] |= ViewBit;
+			}
+
+			if (bEditorRelevance)
+			{
+				// Editor primitives are rendered after post processing and composited onto the scene
+				VisibleEditorPrimitives.AddPrim(PrimitiveSceneInfo);
+
+				if (GIsEditor)
+				{
+					OutHasDynamicEditorMeshElementsMasks[BitIndex] |= ViewBit;
+				}
+			}
+			else if(bDynamicRelevance)
+			{
+				// Keep track of visible dynamic primitives.
+				VisibleDynamicPrimitives.AddPrim(PrimitiveSceneInfo);
+				OutHasDynamicMeshElementsMasks[BitIndex] |= ViewBit;
+			}
+
+			if (ViewRelevance.HasTranslucency() && !bEditorRelevance && ViewRelevance.bRenderInMainPass)
+			{
+				// Add to set of dynamic translucent primitives
+				FTranslucentPrimSet::PlaceScenePrimitive(PrimitiveSceneInfo, View, ViewRelevance.bNormalTranslucencyRelevance, ViewRelevance.bSeparateTranslucencyRelevance, 
+					&SortedTranslucencyPrims.Prims[SortedTranslucencyPrims.NumPrims], SortedTranslucencyPrims.NumPrims,
+					&SortedSeparateTranslucencyPrims.Prims[SortedSeparateTranslucencyPrims.NumPrims], SortedSeparateTranslucencyPrims.NumPrims
+					);
+
+				if (ViewRelevance.bDistortionRelevance)
+				{
+					// Add to set of dynamic distortion primitives
+					DistortionPrimSet.AddPrim(PrimitiveSceneInfo->Proxy);
+				}
+			}
+
+			if (ViewRelevance.bSubsurfaceProfileRelevance)
+			{
+				// this is ok because it is write only
+				const_cast<FViewInfo&>(View).bScreenSpaceSubsurfacePassNeeded = true;
+			}
+
+			if (ViewRelevance.bRenderCustomDepth)
+			{
+				// Add to set of dynamic distortion primitives
+				CustomDepthSet.AddPrim(PrimitiveSceneInfo->Proxy);
+			}
+
+			// INITVIEWS_TODO: Do this in a separate pass? There are no dependencies
+			// here except maybe ParentPrimitives. This could be done in a 
+			// low-priority background task and forgotten about.
+
+			// If the primitive's last render time is older than last frame, consider
+			// it newly visible and update its visibility change time
+			if (PrimitiveSceneInfo->LastRenderTime < CurrentWorldTime - DeltaWorldTime - DELTA)
+			{
+				PrimitiveSceneInfo->LastVisibilityChangeTime = CurrentWorldTime;
+			}
+			PrimitiveSceneInfo->LastRenderTime = CurrentWorldTime;
+
+			// If the primitive is definitely unoccluded or if in Wireframe mode and the primitive is estimated
+			// to be unoccluded, then update the primitive components's LastRenderTime 
+			// on the game thread. This signals that the primitive is visible.
+			if (View.PrimitiveDefinitelyUnoccludedMap[BitIndex] || (View.Family->EngineShowFlags.Wireframe && View.PrimitiveVisibilityMap[BitIndex]))
+			{
+				// Update the PrimitiveComponent's LastRenderTime.
+				*(PrimitiveSceneInfo->ComponentLastRenderTime) = CurrentWorldTime;
+			}
+
+			// Cache the nearest reflection proxy if needed
+			if (PrimitiveSceneInfo->bNeedsCachedReflectionCaptureUpdate
+				// In ES2, the per-object reflection is used for everything
+					// Otherwise it is just used on translucency
+						&& (Scene->GetFeatureLevel() == ERHIFeatureLevel::ES2 || bTranslucentRelevance))
+			{
+				PrimitiveSceneInfo->CachedReflectionCaptureProxy = Scene->FindClosestReflectionCapture(Scene->PrimitiveBounds[BitIndex].Origin);
+				PrimitiveSceneInfo->bNeedsCachedReflectionCaptureUpdate = false;
+			}
+			if (PrimitiveSceneInfo->NeedsUpdateStaticMeshes())
+			{
+				UpdateStaticMeshes.AddPrim(PrimitiveSceneInfo);
+			}
+		}
+	}
+	void MarkRelevant()
+	{
+		SCOPE_CYCLE_COUNTER(STAT_StaticRelevance);
+
+		// using a local counter to reduce memory traffic
+		int32 NumVisibleStaticMeshElements = 0;
+		FViewInfo& WriteView = const_cast<FViewInfo&>(View);
+
+		for (int32 StaticPrimIndex = 0, Num = RelevantStaticPrimitives.NumPrims; StaticPrimIndex < Num; ++StaticPrimIndex)
+		{
+			int32 PrimitiveIndex = RelevantStaticPrimitives.Prims[StaticPrimIndex];
+			const FPrimitiveSceneInfo * RESTRICT PrimitiveSceneInfo = Scene->Primitives[PrimitiveIndex];
+			const FPrimitiveBounds & Bounds = Scene->PrimitiveBounds[PrimitiveIndex];
+			const FPrimitiveViewRelevance & ViewRelevance = View.PrimitiveViewRelevanceMap[PrimitiveIndex];
+
+			int8 LODToRender = ComputeLODForMeshes( PrimitiveSceneInfo->StaticMeshes, View, Bounds.Origin, Bounds.SphereRadius, ViewData.ForcedLODLevel, ViewData.LODScale);
+
+			float DistanceSquared = (Bounds.Origin - ViewData.ViewOrigin).SizeSquared();
+			const float LODFactorDistanceSquared = DistanceSquared * FMath::Square(View.LODDistanceFactor * ViewData.InvLODScale);
+			const bool bDrawShadowDepth = FMath::Square(Bounds.SphereRadius) > ViewData.MinScreenRadiusForCSMDepthSquared * LODFactorDistanceSquared;
+			const bool bDrawDepthOnly = ViewData.bForceEarlyZPass || FMath::Square(Bounds.SphereRadius) > GMinScreenRadiusForDepthPrepass * GMinScreenRadiusForDepthPrepass * LODFactorDistanceSquared;
+
+			const int32 NumStaticMeshes = PrimitiveSceneInfo->StaticMeshes.Num();
+			for(int32 MeshIndex = 0;MeshIndex < NumStaticMeshes;MeshIndex++)
+			{
+				const FStaticMesh& StaticMesh = PrimitiveSceneInfo->StaticMeshes[MeshIndex];
+				if (StaticMesh.LODIndex == LODToRender)
+				{
+					uint8 MarkMask = 0;
+					bool bNeedsBatchVisibility = false;
+
+					if (ViewRelevance.bShadowRelevance && bDrawShadowDepth && StaticMesh.CastShadow)
+					{
+						// Mark static mesh as visible in shadows.
+						MarkMask |= EMarkMaskBits::StaticMeshShadowDepthMapMask;
+						bNeedsBatchVisibility = true;
+					}
+
+					if(ViewRelevance.bDrawRelevance && !StaticMesh.bShadowOnly && (ViewRelevance.bRenderInMainPass || ViewRelevance.bRenderCustomDepth))
+					{
+						// Mark static mesh as visible for rendering
+						MarkMask |= EMarkMaskBits::StaticMeshVisibilityMapMask;
+						if (PrimitiveSceneInfo->ShouldRenderVelocity(View, false))
+						{
+							MarkMask |= EMarkMaskBits::StaticMeshVelocityMapMask;
+						}
+						++NumVisibleStaticMeshElements;
+
+						// If the static mesh is an occluder, check whether it covers enough of the screen to be used as an occluder.
+						if(	StaticMesh.bUseAsOccluder && bDrawDepthOnly )
+						{
+							MarkMask |= EMarkMaskBits::StaticMeshOccluderMapMask;
+						}
+						bNeedsBatchVisibility = true;
+					}
+					if (MarkMask)
+					{
+						MarkMasks[StaticMesh.Id] = MarkMask;
+					}
+
+					// Static meshes with a single element always draw, as if the mask were 0x1.
+					if(bNeedsBatchVisibility && StaticMesh.Elements.Num() > 1)
+					{
+						WriteView.StaticMeshBatchVisibility[StaticMesh.Id] = StaticMesh.VertexFactory->GetStaticBatchElementVisibility(View, &StaticMesh);
+					}
+				}
+			}
+		}
+		static_assert(sizeof(WriteView.NumVisibleStaticMeshElements) == sizeof(int32), "Atomic is the wrong size");
+		FPlatformAtomics::InterlockedAdd((volatile int32*)&WriteView.NumVisibleStaticMeshElements, NumVisibleStaticMeshElements);
+	}
+
+	void RenderThreadFinalize()
+	{
+		FViewInfo& WriteView = const_cast<FViewInfo&>(View);
+
+		for (int32 Index = 0; Index < NotDrawRelevant.NumPrims; Index++)
+		{
+			WriteView.PrimitiveVisibilityMap[NotDrawRelevant.Prims[Index]] = false;
+		}
+		VisibleEditorPrimitives.AppendTo(WriteView.VisibleEditorPrimitives);
+		VisibleDynamicPrimitives.AppendTo(WriteView.VisibleDynamicPrimitives);
+		WriteView.TranslucentPrimSet.AppendScenePrimitives(SortedTranslucencyPrims.Prims, SortedTranslucencyPrims.NumPrims, SortedSeparateTranslucencyPrims.Prims, SortedSeparateTranslucencyPrims.NumPrims);
+		DistortionPrimSet.AppendTo(WriteView.DistortionPrimSet);
+		CustomDepthSet.AppendTo(WriteView.CustomDepthSet);
+
+		for (int32 Index = 0; Index < UpdateStaticMeshes.NumPrims; Index++)
+		{
+			UpdateStaticMeshes.Prims[Index]->UpdateStaticMeshes(RHICmdList);
+		}
+	}
+};
+
+class FRelevancePacketAnyThreadTask
+{
+	FRelevancePacket& Packet;
+	ENamedThreads::Type ThreadToUse;
+public:
+
+	FRelevancePacketAnyThreadTask(FRelevancePacket* InPacket, ENamedThreads::Type InThreadToUse)
+		: Packet(*InPacket)
+		, ThreadToUse(InThreadToUse)
+	{
+	}
+
+	FORCEINLINE TStatId GetStatId() const
+	{
+		RETURN_QUICK_DECLARE_CYCLE_STAT(FRelevancePacketAnyThreadTask, STATGROUP_TaskGraphTasks);
+	}
+
+	ENamedThreads::Type GetDesiredThread()
+	{
+		return ThreadToUse;
+	}
+
+	static ESubsequentsMode::Type GetSubsequentsMode() { return ESubsequentsMode::TrackSubsequents; }
+
+	void DoTask(ENamedThreads::Type CurrentThread, const FGraphEventRef& MyCompletionGraphEvent)
+	{
+		Packet.AnyThreadTask();
+	}
+};
+
+class FRelevancePacketRenderThreadTask
+{
+	FRelevancePacket& Packet;
+public:
+
+	FRelevancePacketRenderThreadTask(FRelevancePacket* InPacket)
+		: Packet(*InPacket)
+	{
+	}
+
+	FORCEINLINE TStatId GetStatId() const
+	{
+		RETURN_QUICK_DECLARE_CYCLE_STAT(FRelevancePacketRenderThreadTask, STATGROUP_TaskGraphTasks);
+	}
+
+	ENamedThreads::Type GetDesiredThread()
+	{
+		return ENamedThreads::RenderThread_Local;
+	}
+
+	static ESubsequentsMode::Type GetSubsequentsMode() { return ESubsequentsMode::TrackSubsequents; }
+
+	void DoTask(ENamedThreads::Type CurrentThread, const FGraphEventRef& MyCompletionGraphEvent)
+	{
+		Packet.RenderThreadFinalize();
+	}
+};
+
+static TAutoConsoleVariable<int32> CVarParallelInitViews(
+	TEXT("r.ParallelInitViews"),
+	0,  
+	TEXT("Toggles parallel init views.\n"),
+	ECVF_RenderThreadSafe
+	);
+
+
+/**
+ * Computes view relevance for visible primitives in the view and adds them to
+ * appropriate per-view rendering lists.
+ * @param Scene - The scene being rendered.
+ * @param View - The view for which to compute relevance.
+ * @param ViewBit - Bit mask: 1 << ViewIndex where Views(ViewIndex) == View.
+ * @param OutRelevantStaticPrimitives - Upon return contains a list of relevant
+ *                                      static primitives.
+ * @param OutPreRenderViewMasks - Primitives requiring the PreRenderView
+ *                                callback for this view will have ViewBit set.
+ */
+static void ComputeAndMarkRelevanceForViewParallel(
+	FRHICommandListImmediate& RHICmdList,
+	const FScene* Scene,
+	FViewInfo& View,
+	uint8 ViewBit,
+	FPrimitiveViewMasks& OutPreRenderViewMasks,
+	FPrimitiveViewMasks& OutHasDynamicMeshElementsMasks,
+	FPrimitiveViewMasks& OutHasDynamicEditorMeshElementsMasks
+	)
+{
+
+	check(OutPreRenderViewMasks.Num() == Scene->Primitives.Num());
+	check(OutHasDynamicMeshElementsMasks.Num() == Scene->Primitives.Num());
+
+	const FMarkRelevantStaticMeshesForViewData ViewData(View);
+
+	int32 NumMesh = View.StaticMeshVisibilityMap.Num();
+	check(View.StaticMeshShadowDepthMap.Num() == NumMesh && View.StaticMeshVelocityMap.Num() == NumMesh && View.StaticMeshOccluderMap.Num() == NumMesh);
+	uint8* RESTRICT MarkMasks = (uint8*)FMemStack::Get().Alloc(NumMesh + 31 , 8); // some padding to simplify the high speed transpose
+	FMemory::Memzero(MarkMasks, NumMesh + 31);
+
+	// I am going to do the render thread tasks in order, maybe that isn't necessary
+	FGraphEventRef LastRenderThread;
+	{
+		FSceneSetBitIterator BitIt(View.PrimitiveVisibilityMap);
+		if (BitIt)
+		{
+			int32 AnyThreadTasksPerRenderThreadTasks = 1;
+
+			if (FApp::ShouldUseThreadingForPerformance() && CVarParallelInitViews.GetValueOnRenderThread() > 0)
+			{
+				AnyThreadTasksPerRenderThreadTasks = FTaskGraphInterface::Get().GetNumWorkerThreads() + 1; // the idea is to put the render thread to work
+			}
+			// else AnyThreadTasksPerRenderThreadTasks == 1 means we never use a task thread
+			int32 WorkingAnyThreadTasksPerRenderThreadTasks = AnyThreadTasksPerRenderThreadTasks;
+			
+
+			FRelevancePacket* Packet = new(FMemStack::Get()) FRelevancePacket(
+				RHICmdList,
+				Scene, 
+				View, 
+				ViewBit,
+				ViewData,
+				OutPreRenderViewMasks,
+				OutHasDynamicMeshElementsMasks,
+				OutHasDynamicEditorMeshElementsMasks,
+				MarkMasks);
+
+			while (1)
+			{
+				Packet->Input.AddPrim(BitIt.GetIndex());
+				++BitIt;
+				if (Packet->Input.IsFull() || !BitIt)
+				{
+					// submit task
+					ENamedThreads::Type ThreadToUse = ENamedThreads::AnyThread;
+					if (!--WorkingAnyThreadTasksPerRenderThreadTasks)
+					{
+						WorkingAnyThreadTasksPerRenderThreadTasks = AnyThreadTasksPerRenderThreadTasks;
+						ThreadToUse = ENamedThreads::RenderThread_Local;
+					}
+					FGraphEventArray RenderPrereqs;
+					if (LastRenderThread.GetReference())
+					{
+						RenderPrereqs.Add(LastRenderThread); // this puts the render thread ones in order
+					}
+					FGraphEventRef AnyThread = TGraphTask<FRelevancePacketAnyThreadTask>::CreateTask(nullptr, ENamedThreads::RenderThread).ConstructAndDispatchWhenReady(Packet, ThreadToUse);
+					RenderPrereqs.Add(AnyThread);
+					LastRenderThread = TGraphTask<FRelevancePacketRenderThreadTask>::CreateTask(&RenderPrereqs, ENamedThreads::RenderThread).ConstructAndDispatchWhenReady(Packet);
+					if (!BitIt)
+					{
+						break;
+					}
+					else
+					{
+						Packet = new(FMemStack::Get()) FRelevancePacket(
+							RHICmdList,
+							Scene, 
+							View, 
+							ViewBit,
+							ViewData,
+							OutPreRenderViewMasks,
+							OutHasDynamicMeshElementsMasks,
+							OutHasDynamicEditorMeshElementsMasks,
+							MarkMasks);
+					}
+				}
+			}
+		}
+	}
+	if (LastRenderThread.GetReference())
+	{
+		QUICK_SCOPE_CYCLE_COUNTER(STAT_ComputeAndMarkRelevanceForViewParallel_Wait);
+		FTaskGraphInterface::Get().WaitUntilTaskCompletes(LastRenderThread, ENamedThreads::RenderThread_Local);
+	}
+	QUICK_SCOPE_CYCLE_COUNTER(STAT_ComputeAndMarkRelevanceForViewParallel_TransposeMeshBits);
+	check(View.StaticMeshVelocityMap.Num() == NumMesh && View.StaticMeshShadowDepthMap.Num() == NumMesh && View.StaticMeshVisibilityMap.Num() == NumMesh && View.StaticMeshOccluderMap.Num() == NumMesh);
+	uint32* RESTRICT StaticMeshVisibilityMap_Words = View.StaticMeshVisibilityMap.GetData();
+	uint32* RESTRICT StaticMeshVelocityMap_Words = View.StaticMeshVelocityMap.GetData();
+	uint32* RESTRICT StaticMeshShadowDepthMap_Words = View.StaticMeshShadowDepthMap.GetData();
+	uint32* RESTRICT StaticMeshOccluderMap_Words = View.StaticMeshOccluderMap.GetData();
+	const uint64* RESTRICT MarkMasks64 = (const uint64* RESTRICT)MarkMasks;
+	const uint8* RESTRICT MarkMasks8 = MarkMasks;
+	for (int32 BaseIndex = 0; BaseIndex < NumMesh; BaseIndex += 32)
+	{
+		uint32 StaticMeshVisibilityMap_Word = 0;
+		uint32 StaticMeshVelocityMap_Word = 0;
+		uint32 StaticMeshShadowDepthMap_Word = 0;
+		uint32 StaticMeshOccluderMap_Word = 0;
+		uint32 Mask = 1;
+		bool bAny = false;
+		for (int32 QWordIndex = 0; QWordIndex < 4; QWordIndex++)
+		{
+			if (*MarkMasks64++)
+			{
+				for (int32 ByteIndex = 0; ByteIndex < 8; ByteIndex++, Mask <<= 1, MarkMasks8++)
+				{
+					uint8 MaskMask = *MarkMasks8;
+					StaticMeshVisibilityMap_Word |= (MaskMask & EMarkMaskBits::StaticMeshVisibilityMapMask) ? Mask : 0;
+					StaticMeshVelocityMap_Word |= (MaskMask & EMarkMaskBits::StaticMeshVelocityMapMask) ? Mask : 0;
+					StaticMeshShadowDepthMap_Word |= (MaskMask & EMarkMaskBits::StaticMeshShadowDepthMapMask) ? Mask : 0;
+					StaticMeshOccluderMap_Word |= (MaskMask & EMarkMaskBits::StaticMeshOccluderMapMask) ? Mask : 0;
+				}
+				bAny = true;
+			}
+			else
+			{
+				MarkMasks8 += 8;
+				Mask <<= 8;
+			}
+		}
+		if (bAny)
+		{
+			checkSlow(!*StaticMeshVisibilityMap_Words && !*StaticMeshVelocityMap_Words && !*StaticMeshShadowDepthMap_Words && !*StaticMeshOccluderMap_Words);
+			*StaticMeshVisibilityMap_Words = StaticMeshVisibilityMap_Word;
+			*StaticMeshVelocityMap_Words = StaticMeshVelocityMap_Word;
+			*StaticMeshShadowDepthMap_Words = StaticMeshShadowDepthMap_Word;
+			*StaticMeshOccluderMap_Words = StaticMeshOccluderMap_Word;
+		}
+		StaticMeshVisibilityMap_Words++;
+		StaticMeshVelocityMap_Words++;
+		StaticMeshShadowDepthMap_Words++;
+		StaticMeshOccluderMap_Words++;
+	}
+}
+
 void FSceneRenderer::GatherDynamicMeshElements(
 	TArray<FViewInfo>& Views, 
 	const FScene* Scene, 
@@ -823,6 +1381,8 @@ void FSceneRenderer::GatherDynamicMeshElements(
 
 static void MarkAllPrimitivesForReflectionProxyUpdate(FScene* Scene)
 {
+	QUICK_SCOPE_CYCLE_COUNTER(STAT_MarkAllPrimitivesForReflectionProxyUpdate);
+
 	if (Scene->ReflectionSceneData.bRegisteredReflectionCapturesHasChanged)
 	{
 		// Mark all primitives as needing an update
@@ -848,21 +1408,7 @@ static void MarkRelevantStaticMeshesForView(
 {
 	SCOPE_CYCLE_COUNTER(STAT_StaticRelevance);
 
-	FVector ViewOrigin = View.ViewMatrices.ViewOrigin;
-
-	float MaxDrawDistanceScaleSquared = GetCachedScalabilityCVars().ViewDistanceScaleSquared;
-
-	// outside of the loop to be more efficient
-	int32 ForcedLODLevel = (View.Family->EngineShowFlags.LOD) ? GetCVarForceLOD() : 0;
-	
-	const float LODScale = CVarStaticMeshLODDistanceScale.GetValueOnRenderThread();
-	const float InvLODScale = 1.0f / LODScale;
-
-	const float MinScreenRadiusForCSMDepthSquared = GMinScreenRadiusForCSMDepth * GMinScreenRadiusForCSMDepth;
-	const float MinScreenRadiusForDepthPrepassSquared = GMinScreenRadiusForDepthPrepass * GMinScreenRadiusForDepthPrepass;
-	
-	extern TAutoConsoleVariable<int32> CVarEarlyZPass;
-	const bool bForceEarlyZPass = CVarEarlyZPass.GetValueOnRenderThread() == 2;
+	const FMarkRelevantStaticMeshesForViewData ViewData(View);
 
 	// using a local counter to reduce memory traffic
 	int32 NumVisibleStaticMeshElements = 0;
@@ -874,12 +1420,12 @@ static void MarkRelevantStaticMeshesForView(
 		const FPrimitiveBounds & Bounds = Scene->PrimitiveBounds[PrimitiveIndex];
 		const FPrimitiveViewRelevance & ViewRelevance = View.PrimitiveViewRelevanceMap[PrimitiveIndex];
 		
-		int8 LODToRender = ComputeLODForMeshes( PrimitiveSceneInfo->StaticMeshes, View, Bounds.Origin, Bounds.SphereRadius, ForcedLODLevel, LODScale);
+		int8 LODToRender = ComputeLODForMeshes( PrimitiveSceneInfo->StaticMeshes, View, Bounds.Origin, Bounds.SphereRadius, ViewData.ForcedLODLevel, ViewData.LODScale);
 
-		float DistanceSquared = (Bounds.Origin - ViewOrigin).SizeSquared();
-		const float LODFactorDistanceSquared = DistanceSquared * FMath::Square(View.LODDistanceFactor * InvLODScale);
-		const bool bDrawShadowDepth = FMath::Square(Bounds.SphereRadius) > MinScreenRadiusForCSMDepthSquared * LODFactorDistanceSquared;
-		const bool bDrawDepthOnly = bForceEarlyZPass || FMath::Square(Bounds.SphereRadius) > GMinScreenRadiusForDepthPrepass * GMinScreenRadiusForDepthPrepass * LODFactorDistanceSquared;
+		float DistanceSquared = (Bounds.Origin - ViewData.ViewOrigin).SizeSquared();
+		const float LODFactorDistanceSquared = DistanceSquared * FMath::Square(View.LODDistanceFactor * ViewData.InvLODScale);
+		const bool bDrawShadowDepth = FMath::Square(Bounds.SphereRadius) > ViewData.MinScreenRadiusForCSMDepthSquared * LODFactorDistanceSquared;
+		const bool bDrawDepthOnly = ViewData.bForceEarlyZPass || FMath::Square(Bounds.SphereRadius) > GMinScreenRadiusForDepthPrepass * GMinScreenRadiusForDepthPrepass * LODFactorDistanceSquared;
 		
 		// We could compute this only if required by the following code. Not sure if that is worth.
 		const bool bShouldRenderVelocity = PrimitiveSceneInfo->ShouldRenderVelocity(View);
@@ -1264,6 +1810,13 @@ void FSceneRenderer::PreVisibilityFrameSetup(FRHICommandListImmediate& RHICmdLis
 	}
 }
 
+static TAutoConsoleVariable<int32> CVarLegacySingleThreadedRelevance(
+	TEXT("r.LegacySingleThreadedRelevance"),
+	0,  
+	TEXT("Toggles the lgacy codepath for view relevance.\n"),
+	ECVF_RenderThreadSafe
+	);
+
 void FSceneRenderer::ComputeViewVisibility(FRHICommandListImmediate& RHICmdList)
 {
 	SCOPE_CYCLE_COUNTER(STAT_ViewVisibilityTime);
@@ -1296,9 +1849,6 @@ void FSceneRenderer::ComputeViewVisibility(FRHICommandListImmediate& RHICmdList)
 		HasDynamicEditorMeshElementsMasks.AddZeroed(NumPrimitives);
 	}
 
-	// This array contains a list of relevant static primities.
-	FRelevantStaticPrimitives RelevantStaticPrimitives;
-
 	uint8 ViewBit = 0x1;
 	for (int32 ViewIndex = 0; ViewIndex < Views.Num(); ++ViewIndex, ViewBit <<= 1)
 	{
@@ -1452,10 +2002,21 @@ void FSceneRenderer::ComputeViewVisibility(FRHICommandListImmediate& RHICmdList)
 		MarkAllPrimitivesForReflectionProxyUpdate(Scene);
 		Scene->ConditionalMarkStaticMeshElementsForUpdate();
 
-		// Compute view relevance for all visible primitives.
-		RelevantStaticPrimitives.Reset();
-		ComputeRelevanceForView(RHICmdList, Scene, View, ViewBit, RelevantStaticPrimitives, PreRenderViewMasks, HasDynamicMeshElementsMasks, HasDynamicEditorMeshElementsMasks);
-		MarkRelevantStaticMeshesForView(Scene, View, RelevantStaticPrimitives);
+		{
+			SCOPE_CYCLE_COUNTER(STAT_ViewRelevance);
+			// Compute view relevance for all visible primitives.
+			if (!CVarLegacySingleThreadedRelevance.GetValueOnRenderThread())
+			{
+				ComputeAndMarkRelevanceForViewParallel(RHICmdList, Scene, View, ViewBit, PreRenderViewMasks, HasDynamicMeshElementsMasks, HasDynamicEditorMeshElementsMasks);
+			}
+			else
+			{
+				// This array contains a list of relevant static primities.
+				FRelevantStaticPrimitives RelevantStaticPrimitives;
+				ComputeRelevanceForView(RHICmdList, Scene, View, ViewBit, RelevantStaticPrimitives, PreRenderViewMasks, HasDynamicMeshElementsMasks, HasDynamicEditorMeshElementsMasks);
+				MarkRelevantStaticMeshesForView(Scene, View, RelevantStaticPrimitives);
+			}
+		}
 
 #if !(UE_BUILD_SHIPPING || UE_BUILD_TEST)
 		// Store the primitive for parent occlusion rendering.
diff --git a/Engine/Source/Runtime/Renderer/Private/TranslucentRendering.cpp b/Engine/Source/Runtime/Renderer/Private/TranslucentRendering.cpp
--- a/Engine/Source/Runtime/Renderer/Private/TranslucentRendering.cpp
+++ b/Engine/Source/Runtime/Renderer/Private/TranslucentRendering.cpp
@@ -756,6 +756,39 @@ void FTranslucentPrimSet::AddScenePrimitive(FPrimitiveSceneInfo* PrimitiveSceneI
 	}
 }
 
+void FTranslucentPrimSet::AppendScenePrimitives(FSortedPrim* Normal, int32 NumNormal, FSortedPrim* Separate, int32 NumSeparate)
+{
+	SortedPrims.Append(Normal, NumNormal);
+	SortedSeparateTranslucencyPrims.Append(Separate, NumSeparate);
+}
+
+void FTranslucentPrimSet::PlaceScenePrimitive(FPrimitiveSceneInfo* PrimitiveSceneInfo, const FViewInfo& ViewInfo, bool bUseNormalTranslucency, bool bUseSeparateTranslucency, void *NormalPlace, int32& NormalNum, void* SeparatePlace, int32& SeparateNum)
+{
+	float SortKey=0.f;
+	//sort based on distance to the view position, view rotation is not a factor
+	SortKey = (PrimitiveSceneInfo->Proxy->GetBounds().Origin - ViewInfo.ViewMatrices.ViewOrigin).Size();
+	// UE4_TODO: also account for DPG in the sort key.
+
+	const auto FeatureLevel = ViewInfo.GetFeatureLevel();
+
+	if(bUseSeparateTranslucency 
+		&& FeatureLevel >= ERHIFeatureLevel::SM4)
+	{
+		// add to list of translucent prims that use scene color
+		new (SeparatePlace) FSortedPrim(PrimitiveSceneInfo,SortKey,PrimitiveSceneInfo->Proxy->GetTranslucencySortPriority());
+		SeparateNum++;
+	}
+
+	if (bUseNormalTranslucency 
+		// Force separate translucency to be rendered normally if the feature level does not support separate translucency
+			|| (bUseSeparateTranslucency && FeatureLevel < ERHIFeatureLevel::SM4))
+	{
+		// add to list of translucent prims
+		new (NormalPlace) FSortedPrim(PrimitiveSceneInfo,SortKey,PrimitiveSceneInfo->Proxy->GetTranslucencySortPriority());
+		NormalNum++;
+	}
+}
+
 /**
 * Sort any primitives that were added to the set back-to-front
 */
