UE4 - Revert Reflection captures changes

diff --git a/Engine/Shaders/ReflectionEnvironmentShaders.usf b/Engine/Shaders/ReflectionEnvironmentShaders.usf
--- a/Engine/Shaders/ReflectionEnvironmentShaders.usf
+++ b/Engine/Shaders/ReflectionEnvironmentShaders.usf
@@ -29,8 +29,34 @@ void CopyToCubeFaceVS(
 	Out.ScreenVector = mul(float4(Out.Position.xy, 1, 0), View.ScreenToTranslatedWorld).xyz;
 }
 
+struct FRouteToCubeFaceGSOut
+{
+	FCopyToCubeFaceVSOutput Vertex;
+	uint LayerIndex : SV_RenderTargetArrayIndex;
+};
+
 int CubeFace;
 
+[maxvertexcount(3)]
+void RouteToCubeFaceGS(triangle FCopyToCubeFaceVSOutput Input[3], inout TriangleStream<FRouteToCubeFaceGSOut> OutStream)
+{
+	FRouteToCubeFaceGSOut Vertex0;
+	Vertex0.Vertex = Input[0];
+	Vertex0.LayerIndex = CubeFace;
+
+	FRouteToCubeFaceGSOut Vertex1;
+	Vertex1.Vertex = Input[1];
+	Vertex1.LayerIndex = CubeFace;
+
+	FRouteToCubeFaceGSOut Vertex2;
+	Vertex2.Vertex = Input[2];
+	Vertex2.LayerIndex = CubeFace;
+
+	OutStream.Append(Vertex0);
+	OutStream.Append(Vertex1);
+	OutStream.Append(Vertex2);
+}
+
 Texture2D InTexture;
 SamplerState InTextureSampler;
 
@@ -42,16 +68,16 @@ SamplerState InTextureSampler;
 float3 SkyLightCaptureParameters;
 
 void CopySceneColorToCubeFaceColorPS(
-	FCopyToCubeFaceVSOutput Input,
+	FRouteToCubeFaceGSOut Input,
 	out float4 OutColor : SV_Target0
 	)
 { 
-	float SceneDepth = CalcSceneDepth(Input.UV);
-	float3 SceneColor = Texture2DSample(InTexture, InTextureSampler, Input.UV).rgb;
+	float SceneDepth = CalcSceneDepth(Input.Vertex.UV);
+	float3 SceneColor = Texture2DSample(InTexture, InTextureSampler, Input.Vertex.UV).rgb;
 
 	// Convert INF's to valid values
 	SceneColor = min(SceneColor, 65503);
-	float3 WorldPosition = Input.ScreenVector * SceneDepth + View.ViewOrigin.xyz;
+	float3 WorldPosition = Input.Vertex.ScreenVector * SceneDepth + View.ViewOrigin.xyz;
 
 	float Alpha = 1;
 
@@ -59,7 +85,7 @@ void CopySceneColorToCubeFaceColorPS(
 	{
 		// Assuming we're on a planet and no sky lighting is coming from below the horizon
 		// This is important to avoid leaking from below since we are integrating incoming lighting and shadowing separately
-		float AboveHorizonMask = Input.ScreenVector.z > 0 || SkyLightCaptureParameters.z < 1;
+		float AboveHorizonMask = Input.Vertex.ScreenVector.z > 0 || SkyLightCaptureParameters.z < 1;
 		SceneColor *= AboveHorizonMask;
 	}
 	else if (SkyLightCaptureParameters.x == 1)
@@ -109,15 +135,41 @@ float3 GetCubemapVector(float2 ScaledUVs)
 	return CubeCoordinates;
 }
 
+struct FDownsampleGSOut
+{
+	FScreenVertexOutput Vertex;
+	uint LayerIndex : SV_RenderTargetArrayIndex;
+};
+
+[maxvertexcount(3)]
+void DownsampleGS(triangle FScreenVertexOutput Input[3], inout TriangleStream<FDownsampleGSOut> OutStream)
+{
+	FDownsampleGSOut Vertex0;
+	Vertex0.Vertex = Input[0];
+	Vertex0.LayerIndex = CubeFace;
+
+	FDownsampleGSOut Vertex1;
+	Vertex1.Vertex = Input[1];
+	Vertex1.LayerIndex = CubeFace;
+
+	FDownsampleGSOut Vertex2;
+	Vertex2.Vertex = Input[2];
+	Vertex2.LayerIndex = CubeFace;
+
+	OutStream.Append(Vertex0);
+	OutStream.Append(Vertex1);
+	OutStream.Append(Vertex2);
+}
+
 TextureCube SourceTexture;
 SamplerState SourceTextureSampler;
 
 void CopyCubemapToCubeFaceColorPS(
-	FScreenVertexOutput Input,
+	FDownsampleGSOut Input,
 	out float4 OutColor : SV_Target0
 	)
 { 
-	float2 ScaledUVs = Input.UV * 2 - 1;
+	float2 ScaledUVs = Input.Vertex.UV * 2 - 1;
 	float3 CubeCoordinates = GetCubemapVector(ScaledUVs);
 	OutColor = TextureCubeSampleLevel(SourceTexture, SourceTextureSampler, CubeCoordinates, 0);
 
@@ -138,17 +190,17 @@ float4 SampleCubemap(float3 Coordinates)
 }
 
 void DownsamplePS(
-	FScreenVertexOutput Input,
+	FDownsampleGSOut Input,
 	out float4 OutColor : SV_Target0
 	)
 { 
-	float2 ScaledUVs = Input.UV * 2 - 1;
+	float2 ScaledUVs = Input.Vertex.UV * 2 - 1;
 	float3 CubeCoordinates = GetCubemapVector(ScaledUVs);
 
 	OutColor = SampleCubemap(CubeCoordinates);
 }
 
-int NumCaptureArrayMips;
+uint NumCaptureArrayMips;
 
 /** Cube map array of reflection captures. */
 TextureCube ReflectionEnvironmentColorTexture;
@@ -180,20 +232,14 @@ Texture2D AverageBrightnessTexture;
 SamplerState AverageBrightnessSampler;
 
 void FilterPS(
-	FScreenVertexOutput Input,
+	FDownsampleGSOut Input,
 	out float4 OutColor : SV_Target0
 	)
 { 
-	float2 ScaledUVs = Input.UV * 2 - 1;
+	float2 ScaledUVs = Input.Vertex.UV * 2 - 1;
 	float3 CubeCoordinates = GetCubemapVector(ScaledUVs);
 
-#if ES2_PROFILE
-//@todo - mwassmer.  Write filtering that works for ES2 so specular will work.
-#define FILTER_CUBEMAP 0
-#else
 #define FILTER_CUBEMAP 1
-#endif
-
 #if FILTER_CUBEMAP
 
 	float3 N = normalize(CubeCoordinates);
@@ -203,7 +249,7 @@ void FilterPS(
 	float Weight = 0;
 		
 	LOOP
-	for( int i = 0; i < NUM_FILTER_SAMPLES; i++ )
+	for( uint i = 0; i < NUM_FILTER_SAMPLES; i++ )
 	{
 		float2 E = Hammersley( i, NUM_FILTER_SAMPLES, 0 );
 		float3 H = TangentToWorld( ImportanceSampleGGX( E, Roughness ).xyz, N );
@@ -239,11 +285,11 @@ float CoefficientMask2;
 int NumSamples;
 
 void DiffuseIrradianceCopyPS(
-	FScreenVertexOutput Input,
+	FDownsampleGSOut Input,
 	out float4 OutColor : SV_Target0
 	)
 { 
-	float2 ScaledUVs = Input.UV * 2 - 1;
+	float2 ScaledUVs = Input.Vertex.UV * 2 - 1;
 	float3 CubeCoordinates = normalize(GetCubemapVector(ScaledUVs));
 
 	float SquaredUVs = 1 + dot(ScaledUVs, ScaledUVs);
@@ -261,32 +307,32 @@ float4 Sample01;
 float4 Sample23;
 
 void DiffuseIrradianceAccumulatePS(
-	FScreenVertexOutput Input,
+	FDownsampleGSOut Input,
 	out float4 OutColor : SV_Target0
 	)
 { 
 	float4 AccumulatedValue = 0;
 
 	{
-		float2 ScaledUVs = saturate(Input.UV + Sample01.xy) * 2 - 1;
+		float2 ScaledUVs = saturate(Input.Vertex.UV + Sample01.xy) * 2 - 1;
 		float3 CubeCoordinates = GetCubemapVector(ScaledUVs);
 		AccumulatedValue += SampleCubemap(CubeCoordinates);
 	}
 	
 	{
-		float2 ScaledUVs = saturate(Input.UV + Sample01.zw) * 2 - 1;
+		float2 ScaledUVs = saturate(Input.Vertex.UV + Sample01.zw) * 2 - 1;
 		float3 CubeCoordinates = GetCubemapVector(ScaledUVs);
 		AccumulatedValue += SampleCubemap(CubeCoordinates);
 	}
 
 	{
-		float2 ScaledUVs = saturate(Input.UV + Sample23.xy) * 2 - 1;
+		float2 ScaledUVs = saturate(Input.Vertex.UV + Sample23.xy) * 2 - 1;
 		float3 CubeCoordinates = GetCubemapVector(ScaledUVs);
 		AccumulatedValue += SampleCubemap(CubeCoordinates);
 	}
 
 	{
-		float2 ScaledUVs = saturate(Input.UV + Sample23.zw) * 2 - 1;
+		float2 ScaledUVs = saturate(Input.Vertex.UV + Sample23.zw) * 2 - 1;
 		float3 CubeCoordinates = GetCubemapVector(ScaledUVs);
 		AccumulatedValue += SampleCubemap(CubeCoordinates);
 	}
@@ -305,7 +351,7 @@ void AccumulateCubeFacesPS(
 	AccumulatedValue += TextureCubeSampleLevel(SourceTexture, SourceTextureSampler, float3(0, -1, 0), SourceMipIndex);
 	AccumulatedValue += TextureCubeSampleLevel(SourceTexture, SourceTextureSampler, float3(0, 0, 1), SourceMipIndex);
 	AccumulatedValue += TextureCubeSampleLevel(SourceTexture, SourceTextureSampler, float3(0, 0, -1), SourceMipIndex);
-	OutColor = float4(4 * PI * AccumulatedValue.rgb / ( max(AccumulatedValue.a, .00001f)), 0);	
+	OutColor = float4(4 * PI * AccumulatedValue.rgb / ( max(AccumulatedValue.a, .00001f)), 0);
 }
 
 // Used during reflection captures to get bounce light from specular surfaces
diff --git a/Engine/Source/Runtime/Renderer/Private/ReflectionEnvironmentCapture.cpp b/Engine/Source/Runtime/Renderer/Private/ReflectionEnvironmentCapture.cpp
--- a/Engine/Source/Runtime/Renderer/Private/ReflectionEnvironmentCapture.cpp
+++ b/Engine/Source/Runtime/Renderer/Private/ReflectionEnvironmentCapture.cpp
@@ -75,6 +75,8 @@ FSceneRenderTargetItem& GetEffectiveSourceTexture(bool bDownsamplePass, int32 Ta
 	return GSceneRenderTargets.ReflectionColorScratchCubemap[ScratchTextureIndex]->GetRenderTargetItem();
 }
 
+IMPLEMENT_SHADER_TYPE(,FDownsampleGS,TEXT("ReflectionEnvironmentShaders"),TEXT("DownsampleGS"),SF_Geometry);
+
 class FDownsamplePS : public FGlobalShader
 {
 	DECLARE_SHADER_TYPE(FDownsamplePS,Global);
@@ -82,7 +84,7 @@ public:
 
 	static bool ShouldCache(EShaderPlatform Platform)
 	{
-		return true;
+		return IsFeatureLevelSupported(Platform, ERHIFeatureLevel::SM4);
 	}
 
 	FDownsamplePS(const ShaderMetaType::CompiledShaderInitializerType& Initializer):
@@ -137,7 +139,7 @@ public:
 
 	static bool ShouldCache(EShaderPlatform Platform)
 	{
-		return true;
+		return IsFeatureLevelSupported(Platform, ERHIFeatureLevel::SM4);
 	}
 
 	static void ModifyCompilationEnvironment(EShaderPlatform Platform, FShaderCompilerEnvironment& OutEnvironment)
@@ -211,7 +213,7 @@ public:
 
 	static bool ShouldCache(EShaderPlatform Platform)
 	{
-		return true;
+		return IsFeatureLevelSupported(Platform, ERHIFeatureLevel::SM4);
 	}
 
 	static void ModifyCompilationEnvironment(EShaderPlatform Platform, FShaderCompilerEnvironment& OutEnvironment)
@@ -327,26 +329,37 @@ void FilterReflectionEnvironment(FRHICommandListImmediate& RHICmdList, ERHIFeatu
 			FSceneRenderTargetItem& EffectiveRT = GetEffectiveRenderTarget(true, MipIndex);
 			FSceneRenderTargetItem& EffectiveSource = GetEffectiveSourceTexture(true, MipIndex);
 			check(EffectiveRT.TargetableTexture != EffectiveSource.ShaderResourceTexture);
-			
-			for (int32 CubeFace = 0; CubeFace < CubeFace_MAX; CubeFace++)
+
 			{
-				SetRenderTarget(RHICmdList, EffectiveRT.TargetableTexture, MipIndex, CubeFace, NULL);
+				for (int32 CubeFace = 0; CubeFace < CubeFace_MAX; CubeFace++)
+				{
+					SetRenderTarget(RHICmdList, EffectiveRT.TargetableTexture, MipIndex, CubeFace, NULL);
 
-				const FIntRect ViewRect(0, 0, MipSize, MipSize);
-				RHICmdList.SetViewport(0, 0, 0.0f, MipSize, MipSize, 1.0f);
-				RHICmdList.SetRasterizerState(TStaticRasterizerState<FM_Solid, CM_None>::GetRHI());
-				RHICmdList.SetDepthStencilState(TStaticDepthStencilState<false, CF_Always>::GetRHI());
-				RHICmdList.SetBlendState(TStaticBlendState<>::GetRHI());
+					const FIntRect ViewRect(0, 0, MipSize, MipSize);
+					RHICmdList.SetViewport(0, 0, 0.0f, MipSize, MipSize, 1.0f);
+					RHICmdList.SetRasterizerState(TStaticRasterizerState<FM_Solid, CM_None>::GetRHI());
+					RHICmdList.SetDepthStencilState(TStaticDepthStencilState<false, CF_Always>::GetRHI());
+					RHICmdList.SetBlendState(TStaticBlendState<>::GetRHI());
 
-				TShaderMapRef<FScreenVS> VertexShader(GetGlobalShaderMap(FeatureLevel));
-				TShaderMapRef<FDownsamplePS> PixelShader(GetGlobalShaderMap(FeatureLevel));
+					TShaderMapRef<FScreenVSForGS> VertexShader(ShaderMap);
+					TShaderMapRef<FDownsamplePS> PixelShader(ShaderMap);
 
-				SetGlobalBoundShaderState(RHICmdList, FeatureLevel, DownsampleBoundShaderState, GFilterVertexDeclaration.VertexDeclarationRHI, *VertexShader, *PixelShader);
+					SetGlobalBoundShaderState(RHICmdList, FeatureLevel, DownsampleBoundShaderState, GFilterVertexDeclaration.VertexDeclarationRHI, *VertexShader, *PixelShader);
 
-				PixelShader->SetParameters(RHICmdList, CubeFace, SourceMipIndex, EffectiveSource);
+					PixelShader->SetParameters(RHICmdList, CubeFace, SourceMipIndex, EffectiveSource);
 
+					DrawRectangle( 
+						RHICmdList,
+						ViewRect.Min.X, ViewRect.Min.Y, 
+						ViewRect.Width(), ViewRect.Height(),
+						ViewRect.Min.X, ViewRect.Min.Y, 
+						ViewRect.Width(), ViewRect.Height(),
+						FIntPoint(ViewRect.Width(), ViewRect.Height()),
+						FIntPoint(MipSize, MipSize),
+						*VertexShader);
 
-				RHICmdList.CopyToResolveTarget(EffectiveRT.TargetableTexture, EffectiveRT.ShaderResourceTexture, true, FResolveParams(FResolveRect(), (ECubeFace)CubeFace, MipIndex));
+					RHICmdList.CopyToResolveTarget(EffectiveRT.TargetableTexture, EffectiveRT.ShaderResourceTexture, true, FResolveParams(FResolveRect(), (ECubeFace)CubeFace, MipIndex));
+				}
 			}
 
 			if (MipSize == GDiffuseIrradianceCubemapSize)
@@ -377,49 +390,51 @@ void FilterReflectionEnvironment(FRHICommandListImmediate& RHICmdList, ERHIFeatu
 			check(EffectiveRT.TargetableTexture != EffectiveSource.ShaderResourceTexture);
 			const int32 MipSize = 1 << (NumMips - MipIndex - 1);
 
-			for (int32 CubeFace = 0; CubeFace < CubeFace_MAX; CubeFace++)
 			{
-				SetRenderTarget(RHICmdList, EffectiveRT.TargetableTexture, MipIndex, CubeFace, NULL);
+				for (int32 CubeFace = 0; CubeFace < CubeFace_MAX; CubeFace++)
+				{
+					SetRenderTarget(RHICmdList, EffectiveRT.TargetableTexture, MipIndex, CubeFace, NULL);
 
-				const FIntRect ViewRect(0, 0, MipSize, MipSize);
-				RHICmdList.SetViewport(0, 0, 0.0f, MipSize, MipSize, 1.0f);
-				RHICmdList.SetRasterizerState(TStaticRasterizerState<FM_Solid, CM_None>::GetRHI());
-				RHICmdList.SetDepthStencilState(TStaticDepthStencilState<false, CF_Always>::GetRHI());
-				RHICmdList.SetBlendState(TStaticBlendState<>::GetRHI());
+					const FIntRect ViewRect(0, 0, MipSize, MipSize);
+					RHICmdList.SetViewport(0, 0, 0.0f, MipSize, MipSize, 1.0f);
+					RHICmdList.SetRasterizerState(TStaticRasterizerState<FM_Solid, CM_None>::GetRHI());
+					RHICmdList.SetDepthStencilState(TStaticDepthStencilState<false, CF_Always>::GetRHI());
+					RHICmdList.SetBlendState(TStaticBlendState<>::GetRHI());
 
-				TShaderMapRef<FScreenVS> VertexShader(GetGlobalShaderMap(FeatureLevel));
-				TShaderMapRef< TCubeFilterPS<1> > CaptureCubemapArrayPixelShader(GetGlobalShaderMap(FeatureLevel));
+					TShaderMapRef<FScreenVSForGS> VertexShader(ShaderMap);
+					TShaderMapRef< TCubeFilterPS<1> > CaptureCubemapArrayPixelShader(ShaderMap);
+					//FFilterPS* PixelShader = (FFilterPS*)*CaptureCubemapArrayPixelShader;
 
-				FCubeFilterPS* PixelShader;
-				if( bNormalize )
-				{
-					PixelShader = *TShaderMapRef< TCubeFilterPS<1> >(ShaderMap);
-					static FGlobalBoundShaderState BoundShaderState;
-					
-					SetGlobalBoundShaderState(RHICmdList, FeatureLevel, BoundShaderState, GFilterVertexDeclaration.VertexDeclarationRHI, *VertexShader, PixelShader);
-				}
-				else
-				{
-					PixelShader = *TShaderMapRef< TCubeFilterPS<0> >(ShaderMap);
-					static FGlobalBoundShaderState BoundShaderState;
-					SetGlobalBoundShaderState(RHICmdList, FeatureLevel, BoundShaderState, GFilterVertexDeclaration.VertexDeclarationRHI, *VertexShader, PixelShader);
+					FCubeFilterPS* PixelShader;
+					if( bNormalize )
+					{
+						PixelShader = *TShaderMapRef< TCubeFilterPS<1> >(ShaderMap);
+						static FGlobalBoundShaderState BoundShaderState;
 						
-					SetGlobalBoundShaderState(RHICmdList, FeatureLevel, BoundShaderState, GFilterVertexDeclaration.VertexDeclarationRHI, *VertexShader, PixelShader);
-				}
+						SetGlobalBoundShaderState(RHICmdList, FeatureLevel, BoundShaderState, GFilterVertexDeclaration.VertexDeclarationRHI, *VertexShader, PixelShader);
+					}
+					else
+					{
+						PixelShader = *TShaderMapRef< TCubeFilterPS<0> >(ShaderMap);
+						static FGlobalBoundShaderState BoundShaderState;
+						
+						SetGlobalBoundShaderState(RHICmdList, FeatureLevel, BoundShaderState, GFilterVertexDeclaration.VertexDeclarationRHI, *VertexShader, PixelShader);
+					}
 
-				PixelShader->SetParameters(RHICmdList, CubeFace, MipIndex, EffectiveSource);
+					PixelShader->SetParameters(RHICmdList, CubeFace, MipIndex, EffectiveSource);
 
-				DrawRectangle( 
-					RHICmdList,
-					ViewRect.Min.X, ViewRect.Min.Y, 
-					ViewRect.Width(), ViewRect.Height(),
-					ViewRect.Min.X, ViewRect.Min.Y, 
-					ViewRect.Width(), ViewRect.Height(),
-					FIntPoint(ViewRect.Width(), ViewRect.Height()),
-					FIntPoint(MipSize, MipSize),
-					*VertexShader);
+					DrawRectangle( 
+						RHICmdList,
+						ViewRect.Min.X, ViewRect.Min.Y, 
+						ViewRect.Width(), ViewRect.Height(),
+						ViewRect.Min.X, ViewRect.Min.Y, 
+						ViewRect.Width(), ViewRect.Height(),
+						FIntPoint(ViewRect.Width(), ViewRect.Height()),
+						FIntPoint(MipSize, MipSize),
+						*VertexShader);
 
-				RHICmdList.CopyToResolveTarget(EffectiveRT.TargetableTexture, EffectiveRT.ShaderResourceTexture, true, FResolveParams(FResolveRect(), (ECubeFace)CubeFace, MipIndex));
+					RHICmdList.CopyToResolveTarget(EffectiveRT.TargetableTexture, EffectiveRT.ShaderResourceTexture, true, FResolveParams(FResolveRect(), (ECubeFace)CubeFace, MipIndex));
+				}
 			}
 		}
 	}
@@ -433,7 +448,7 @@ public:
 
 	static bool ShouldCache(EShaderPlatform Platform)
 	{
-		return true;
+		return IsFeatureLevelSupported(Platform, ERHIFeatureLevel::SM4);
 	}
 
 	FCopyToCubeFaceVS()	{}
@@ -456,6 +471,42 @@ public:
 
 IMPLEMENT_SHADER_TYPE(,FCopyToCubeFaceVS,TEXT("ReflectionEnvironmentShaders"),TEXT("CopyToCubeFaceVS"),SF_Vertex);
 
+/** Geometry shader used when writing to a cubemap face.  Supports writing to a cubemap array as well. */
+class FRouteToCubeFaceGS : public FGlobalShader
+{
+	DECLARE_SHADER_TYPE(FRouteToCubeFaceGS,Global);
+public:
+
+	static bool ShouldCache(EShaderPlatform Platform)
+	{
+		return IsFeatureLevelSupported(Platform, ERHIFeatureLevel::SM4);
+	}
+
+	FRouteToCubeFaceGS(const ShaderMetaType::CompiledShaderInitializerType& Initializer):
+		FGlobalShader(Initializer)
+	{
+		CubeFace.Bind(Initializer.ParameterMap,TEXT("CubeFace"));
+	}
+	FRouteToCubeFaceGS() {}
+
+	void SetParameters(FRHICommandList& RHICmdList, uint32 CubeFaceValue)
+	{
+		SetShaderValue(RHICmdList, GetGeometryShader(), CubeFace, CubeFaceValue);
+	}
+
+	virtual bool Serialize(FArchive& Ar)
+	{
+		bool bShaderHasOutdatedParameters = FGlobalShader::Serialize(Ar);
+		Ar << CubeFace;
+		return bShaderHasOutdatedParameters;
+	}
+
+private:
+	FShaderParameter CubeFace;
+};
+
+IMPLEMENT_SHADER_TYPE(,FRouteToCubeFaceGS,TEXT("ReflectionEnvironmentShaders"),TEXT("RouteToCubeFaceGS"),SF_Geometry);
+
 /** Pixel shader used when copying scene color from a scene render into a face of a reflection capture cubemap. */
 class FCopySceneColorToCubeFacePS : public FGlobalShader
 {
@@ -464,7 +515,7 @@ public:
 
 	static bool ShouldCache(EShaderPlatform Platform)
 	{
-		return true;
+		return IsFeatureLevelSupported(Platform, ERHIFeatureLevel::SM4);
 	}
 
 	static void ModifyCompilationEnvironment(EShaderPlatform Platform, FShaderCompilerEnvironment& OutEnvironment)
@@ -549,7 +600,7 @@ public:
 
 	static bool ShouldCache(EShaderPlatform Platform)
 	{
-		return true;
+		return IsFeatureLevelSupported(Platform, ERHIFeatureLevel::SM4);
 	}
 
 	FCopyCubemapToCubeFacePS(const ShaderMetaType::CompiledShaderInitializerType& Initializer):
@@ -647,10 +698,12 @@ void ClearScratchCubemaps(FRHICommandList& RHICmdList)
 
 	for (int32 MipIndex = 0; MipIndex < NumMips; MipIndex++)
 	{
-		for (int32 CubeFace = 0; CubeFace < CubeFace_MAX; CubeFace++)
 		{
-			SetRenderTarget(RHICmdList, RT0.TargetableTexture, MipIndex, CubeFace, NULL);
-			RHICmdList.Clear(true, FLinearColor(0, 10000, 0, 0), false, 0, false, 0, FIntRect());
+			for (int32 CubeFace = 0; CubeFace < CubeFace_MAX; CubeFace++)
+			{
+				SetRenderTarget(RHICmdList, RT0.TargetableTexture, MipIndex, CubeFace, NULL);
+				RHICmdList.Clear(true, FLinearColor(0, 10000, 0, 0), false, 0, false, 0, FIntRect());
+			}
 		}
 	}
 
@@ -658,10 +711,12 @@ void ClearScratchCubemaps(FRHICommandList& RHICmdList)
 
 	for (int32 MipIndex = 0; MipIndex < NumMips; MipIndex++)
 	{
-		for (int32 CubeFace = 0; CubeFace < CubeFace_MAX; CubeFace++)
 		{
-			SetRenderTarget(RHICmdList, RT1.TargetableTexture, MipIndex, CubeFace, NULL);
-			RHICmdList.Clear(true, FLinearColor(0, 10000, 0, 0), false, 0, false, 0, FIntRect());
+			for (int32 CubeFace = 0; CubeFace < CubeFace_MAX; CubeFace++)
+			{
+				SetRenderTarget(RHICmdList, RT1.TargetableTexture, MipIndex, CubeFace, NULL);
+				RHICmdList.Clear(true, FLinearColor(0, 10000, 0, 0), false, 0, false, 0, FIntRect());
+			}
 		}
 	}
 }
@@ -700,24 +755,22 @@ void CaptureSceneToScratchCubemap(FRHICommandListImmediate& RHICmdList, FSceneRe
 
 		{
 			SCOPED_DRAW_EVENT(RHICmdList, CubeMapCopyScene, DEC_SCENE_ITEMS);
-			
-			// Copy the captured scene into the cubemap face
-			SetRenderTarget(RHICmdList, EffectiveColorRT.TargetableTexture, 0, CubeFace, NULL);
-
-			const FIntRect ViewRect(0, 0, EffectiveSize, EffectiveSize);
-			RHICmdList.SetViewport(0, 0, 0.0f, EffectiveSize, EffectiveSize, 1.0f);
-			RHICmdList.SetRasterizerState(TStaticRasterizerState<FM_Solid, CM_None>::GetRHI());
-			RHICmdList.SetDepthStencilState(TStaticDepthStencilState<false, CF_Always>::GetRHI());
-			RHICmdList.SetBlendState(TStaticBlendState<>::GetRHI());
+			{
+				// Copy the captured scene into the cubemap face
+				SetRenderTarget(RHICmdList, EffectiveColorRT.TargetableTexture, 0, CubeFace, NULL);
 
-			TShaderMapRef<FCopyToCubeFaceVS> VertexShader(GetGlobalShaderMap(FeatureLevel));
-			TShaderMapRef<FCopySceneColorToCubeFacePS> PixelShader(GetGlobalShaderMap(FeatureLevel));
-			SetGlobalBoundShaderState(RHICmdList, FeatureLevel, CopyColorCubemapBoundShaderState, GFilterVertexDeclaration.VertexDeclarationRHI, *VertexShader, *PixelShader);
+				const FIntRect ViewRect(0, 0, EffectiveSize, EffectiveSize);
+				RHICmdList.SetViewport(0, 0, 0.0f, EffectiveSize, EffectiveSize, 1.0f);
+				RHICmdList.SetRasterizerState(TStaticRasterizerState<FM_Solid, CM_None>::GetRHI());
+				RHICmdList.SetDepthStencilState(TStaticDepthStencilState<false, CF_Always>::GetRHI());
+				RHICmdList.SetBlendState(TStaticBlendState<>::GetRHI());
 
-			SetGlobalBoundShaderState(RHICmdList, FeatureLevel, CopyColorCubemapBoundShaderState, GFilterVertexDeclaration.VertexDeclarationRHI, *VertexShader, *PixelShader);
+				TShaderMapRef<FCopyToCubeFaceVS> VertexShader(ShaderMap);
+				TShaderMapRef<FCopySceneColorToCubeFacePS> PixelShader(ShaderMap);
+				SetGlobalBoundShaderState(RHICmdList, FeatureLevel, CopyColorCubemapBoundShaderState, GFilterVertexDeclaration.VertexDeclarationRHI, *VertexShader, *PixelShader);
 
-			PixelShader->SetParameters(RHICmdList, SceneRenderer->Views[0], bCapturingForSkyLight, bLowerHemisphereIsBlack);
-			VertexShader->SetParameters(RHICmdList, SceneRenderer->Views[0]);
+				PixelShader->SetParameters(RHICmdList, SceneRenderer->Views[0], bCapturingForSkyLight, bLowerHemisphereIsBlack);
+				VertexShader->SetParameters(RHICmdList, SceneRenderer->Views[0]);
 
 				DrawRectangle( 
 					RHICmdList,
@@ -727,9 +780,10 @@ void CaptureSceneToScratchCubemap(FRHICommandListImmediate& RHICmdList, FSceneRe
 					ViewRect.Width() * GSupersampleCaptureFactor, ViewRect.Height() * GSupersampleCaptureFactor,
 					FIntPoint(ViewRect.Width(), ViewRect.Height()),
 					GSceneRenderTargets.GetBufferSizeXY(),
-				*VertexShader);
+					*VertexShader);
 
-			RHICmdList.CopyToResolveTarget(EffectiveColorRT.TargetableTexture, EffectiveColorRT.ShaderResourceTexture, true, FResolveParams(FResolveRect(), CubeFace));			
+				RHICmdList.CopyToResolveTarget(EffectiveColorRT.TargetableTexture, EffectiveColorRT.ShaderResourceTexture, true, FResolveParams(FResolveRect(), CubeFace));
+			}
 		}
 	}
 
@@ -745,36 +799,40 @@ void CopyCubemapToScratchCubemap(FRHICommandList& RHICmdList, ERHIFeatureLevel::
 	const int32 EffectiveSize = GReflectionCaptureSize;
 	FSceneRenderTargetItem& EffectiveColorRT =  GSceneRenderTargets.ReflectionColorScratchCubemap[0]->GetRenderTargetItem();
 
-	for (uint32 CubeFace = 0; CubeFace < CubeFace_MAX; CubeFace++)
+	auto ShaderMap = GetGlobalShaderMap(FeatureLevel);
+
 	{
-		// Copy the captured scene into the cubemap face
-		SetRenderTarget(RHICmdList, EffectiveColorRT.TargetableTexture, 0, CubeFace, NULL);
+		for (uint32 CubeFace = 0; CubeFace < CubeFace_MAX; CubeFace++)
+		{
+			// Copy the captured scene into the cubemap face
+			SetRenderTarget(RHICmdList, EffectiveColorRT.TargetableTexture, 0, CubeFace, NULL);
+
+			const FTexture* SourceCubemapResource = SourceCubemap->Resource;
+			const FIntPoint SourceDimensions(SourceCubemapResource->GetSizeX(), SourceCubemapResource->GetSizeY());
+			const FIntRect ViewRect(0, 0, EffectiveSize, EffectiveSize);
+			RHICmdList.SetViewport(0, 0, 0.0f, EffectiveSize, EffectiveSize, 1.0f);
+			RHICmdList.SetRasterizerState(TStaticRasterizerState<FM_Solid, CM_None>::GetRHI());
+			RHICmdList.SetDepthStencilState(TStaticDepthStencilState<false, CF_Always>::GetRHI());
+			RHICmdList.SetBlendState(TStaticBlendState<>::GetRHI());
 
-		const FTexture* SourceCubemapResource = SourceCubemap->Resource;
-		const FIntPoint SourceDimensions(SourceCubemapResource->GetSizeX(), SourceCubemapResource->GetSizeY());
-		const FIntRect ViewRect(0, 0, EffectiveSize, EffectiveSize);
-		RHICmdList.SetViewport(0, 0, 0.0f, EffectiveSize, EffectiveSize, 1.0f);
-		RHICmdList.SetRasterizerState(TStaticRasterizerState<FM_Solid, CM_None>::GetRHI());
-		RHICmdList.SetDepthStencilState(TStaticDepthStencilState<false, CF_Always>::GetRHI());
-		RHICmdList.SetBlendState(TStaticBlendState<>::GetRHI());
+			TShaderMapRef<FScreenVSForGS> VertexShader(ShaderMap);
+			TShaderMapRef<FCopyCubemapToCubeFacePS> PixelShader(ShaderMap);
 
-		TShaderMapRef<FScreenVS> VertexShader(GetGlobalShaderMap(FeatureLevel));
-		TShaderMapRef<FCopyCubemapToCubeFacePS> PixelShader(GetGlobalShaderMap(FeatureLevel));
+			SetGlobalBoundShaderState(RHICmdList, FeatureLevel, CopyFromCubemapToCubemapBoundShaderState, GFilterVertexDeclaration.VertexDeclarationRHI, *VertexShader, *PixelShader);
+			PixelShader->SetParameters(RHICmdList, SourceCubemapResource, CubeFace, bIsSkyLight, bLowerHemisphereIsBlack);
 
-		SetGlobalBoundShaderState(RHICmdList, FeatureLevel, CopyFromCubemapToCubemapBoundShaderState, GFilterVertexDeclaration.VertexDeclarationRHI, *VertexShader, *PixelShader);
-		PixelShader->SetParameters(RHICmdList, SourceCubemapResource, CubeFace, bIsSkyLight, bLowerHemisphereIsBlack);
+			DrawRectangle( 
+				RHICmdList,
+				ViewRect.Min.X, ViewRect.Min.Y, 
+				ViewRect.Width(), ViewRect.Height(),
+				0, 0, 
+				SourceDimensions.X, SourceDimensions.Y,
+				FIntPoint(ViewRect.Width(), ViewRect.Height()),
+				SourceDimensions,
+				*VertexShader);
 
-		DrawRectangle(
-			RHICmdList,
-			ViewRect.Min.X, ViewRect.Min.Y, 
-			ViewRect.Width(), ViewRect.Height(),
-			0, 0, 
-			SourceDimensions.X, SourceDimensions.Y,
-			FIntPoint(ViewRect.Width(), ViewRect.Height()),
-			SourceDimensions,
-			*VertexShader);
-
-		RHICmdList.CopyToResolveTarget(EffectiveColorRT.TargetableTexture, EffectiveColorRT.ShaderResourceTexture, true, FResolveParams(FResolveRect(), (ECubeFace)CubeFace));
+			RHICmdList.CopyToResolveTarget(EffectiveColorRT.TargetableTexture, EffectiveColorRT.ShaderResourceTexture, true, FResolveParams(FResolveRect(), (ECubeFace)CubeFace));
+		}
 	}
 }
 
@@ -1322,9 +1380,8 @@ void CopyToSkyTexture(FRHICommandList& RHICmdList, FScene* Scene, FTexture* Proc
 }
 
 void FScene::UpdateSkyCaptureContents(const USkyLightComponent* CaptureComponent, bool bCaptureEmissiveOnly, FTexture* OutProcessedTexture, FSHVectorRGB3& OutIrradianceEnvironmentMap)
-{	
-	//todo maybe use ShouldRenderSkylight here, but currently the Skylight hasn't been loaded in time.
-	if (GSupportsRenderTargetFormat_PF_FloatRGBA)
+{
+	if (GetFeatureLevel() >= ERHIFeatureLevel::SM4)
 	{
 		ENQUEUE_UNIQUE_RENDER_COMMAND( 
 			ClearCommand,
diff --git a/Engine/Source/Runtime/Renderer/Private/ReflectionEnvironmentCapture.h b/Engine/Source/Runtime/Renderer/Private/ReflectionEnvironmentCapture.h
--- a/Engine/Source/Runtime/Renderer/Private/ReflectionEnvironmentCapture.h
+++ b/Engine/Source/Runtime/Renderer/Private/ReflectionEnvironmentCapture.h
@@ -9,3 +9,36 @@
 extern ENGINE_API int32 GReflectionCaptureSize;
 
 extern void ComputeDiffuseIrradiance(FRHICommandListImmediate& RHICmdList, ERHIFeatureLevel::Type FeatureLevel, FTextureRHIRef LightingSource, int32 LightingSourceMipIndex, FSHVectorRGB3* OutIrradianceEnvironmentMap);
+
+class FDownsampleGS : public FGlobalShader
+{
+	DECLARE_SHADER_TYPE(FDownsampleGS,Global);
+public:
+
+	static bool ShouldCache(EShaderPlatform Platform)
+	{
+		return IsFeatureLevelSupported(Platform, ERHIFeatureLevel::SM4);
+	}
+
+	FDownsampleGS(const ShaderMetaType::CompiledShaderInitializerType& Initializer):
+		FGlobalShader(Initializer)
+	{
+		CubeFace.Bind(Initializer.ParameterMap,TEXT("CubeFace"));
+	}
+	FDownsampleGS() {}
+
+	void SetParameters(FRHICommandList& RHICmdList, int32 CubeFaceValue)
+	{
+		SetShaderValue(RHICmdList, GetGeometryShader(), CubeFace, CubeFaceValue);
+	}
+
+	virtual bool Serialize(FArchive& Ar)
+	{
+		bool bShaderHasOutdatedParameters = FGlobalShader::Serialize(Ar);
+		Ar << CubeFace;
+		return bShaderHasOutdatedParameters;
+	}
+
+private:
+	FShaderParameter CubeFace;
+};
\ No newline at end of file
