UE4 - Fixed reflection captures on ES2

diff --git a/Engine/Shaders/ReflectionEnvironmentShaders.usf b/Engine/Shaders/ReflectionEnvironmentShaders.usf
--- a/Engine/Shaders/ReflectionEnvironmentShaders.usf
+++ b/Engine/Shaders/ReflectionEnvironmentShaders.usf
@@ -68,16 +68,16 @@ SamplerState InTextureSampler;
 float3 SkyLightCaptureParameters;
 
 void CopySceneColorToCubeFaceColorPS(
-	FRouteToCubeFaceGSOut Input,
+	FCopyToCubeFaceVSOutput Input,
 	out float4 OutColor : SV_Target0
 	)
 { 
-	float SceneDepth = CalcSceneDepth(Input.Vertex.UV);
-	float3 SceneColor = Texture2DSample(InTexture, InTextureSampler, Input.Vertex.UV).rgb;
+	float SceneDepth = CalcSceneDepth(Input.UV);
+	float3 SceneColor = Texture2DSample(InTexture, InTextureSampler, Input.UV).rgb;
 
 	// Convert INF's to valid values
 	SceneColor = min(SceneColor, 65503);
-	float3 WorldPosition = Input.Vertex.ScreenVector * SceneDepth + View.ViewOrigin.xyz;
+	float3 WorldPosition = Input.ScreenVector * SceneDepth + View.ViewOrigin.xyz;
 
 	float Alpha = 1;
 
@@ -85,7 +85,7 @@ void CopySceneColorToCubeFaceColorPS(
 	{
 		// Assuming we're on a planet and no sky lighting is coming from below the horizon
 		// This is important to avoid leaking from below since we are integrating incoming lighting and shadowing separately
-		float AboveHorizonMask = Input.Vertex.ScreenVector.z > 0 || SkyLightCaptureParameters.z < 1;
+		float AboveHorizonMask = Input.ScreenVector.z > 0 || SkyLightCaptureParameters.z < 1;
 		SceneColor *= AboveHorizonMask;
 	}
 	else if (SkyLightCaptureParameters.x == 1)
@@ -135,41 +135,15 @@ float3 GetCubemapVector(float2 ScaledUVs)
 	return CubeCoordinates;
 }
 
-struct FDownsampleGSOut
-{
-	FScreenVertexOutput Vertex;
-	uint LayerIndex : SV_RenderTargetArrayIndex;
-};
-
-[maxvertexcount(3)]
-void DownsampleGS(triangle FScreenVertexOutput Input[3], inout TriangleStream<FDownsampleGSOut> OutStream)
-{
-	FDownsampleGSOut Vertex0;
-	Vertex0.Vertex = Input[0];
-	Vertex0.LayerIndex = CubeFace;
-
-	FDownsampleGSOut Vertex1;
-	Vertex1.Vertex = Input[1];
-	Vertex1.LayerIndex = CubeFace;
-
-	FDownsampleGSOut Vertex2;
-	Vertex2.Vertex = Input[2];
-	Vertex2.LayerIndex = CubeFace;
-
-	OutStream.Append(Vertex0);
-	OutStream.Append(Vertex1);
-	OutStream.Append(Vertex2);
-}
-
 TextureCube SourceTexture;
 SamplerState SourceTextureSampler;
 
 void CopyCubemapToCubeFaceColorPS(
-	FDownsampleGSOut Input,
+	FScreenVertexOutput Input,
 	out float4 OutColor : SV_Target0
 	)
 { 
-	float2 ScaledUVs = Input.Vertex.UV * 2 - 1;
+	float2 ScaledUVs = Input.UV * 2 - 1;
 	float3 CubeCoordinates = GetCubemapVector(ScaledUVs);
 	OutColor = TextureCubeSampleLevel(SourceTexture, SourceTextureSampler, CubeCoordinates, 0);
 
@@ -190,17 +164,17 @@ float4 SampleCubemap(float3 Coordinates)
 }
 
 void DownsamplePS(
-	FDownsampleGSOut Input,
+	FScreenVertexOutput Input,
 	out float4 OutColor : SV_Target0
 	)
 { 
-	float2 ScaledUVs = Input.Vertex.UV * 2 - 1;
+	float2 ScaledUVs = Input.UV * 2 - 1;
 	float3 CubeCoordinates = GetCubemapVector(ScaledUVs);
 
 	OutColor = SampleCubemap(CubeCoordinates);
 }
 
-uint NumCaptureArrayMips;
+int NumCaptureArrayMips;
 
 /** Cube map array of reflection captures. */
 TextureCube ReflectionEnvironmentColorTexture;
@@ -232,14 +206,20 @@ Texture2D AverageBrightnessTexture;
 SamplerState AverageBrightnessSampler;
 
 void FilterPS(
-	FDownsampleGSOut Input,
+	FScreenVertexOutput Input,
 	out float4 OutColor : SV_Target0
 	)
 { 
-	float2 ScaledUVs = Input.Vertex.UV * 2 - 1;
+	float2 ScaledUVs = Input.UV * 2 - 1;
 	float3 CubeCoordinates = GetCubemapVector(ScaledUVs);
 
+#if ES2_PROFILE
+//@todo - mwassmer.  Write filtering that works for ES2 so specular will work.
+#define FILTER_CUBEMAP 0
+#else
 #define FILTER_CUBEMAP 1
+#endif
+
 #if FILTER_CUBEMAP
 
 	float3 N = normalize(CubeCoordinates);
@@ -249,7 +229,7 @@ void FilterPS(
 	float Weight = 0;
 		
 	LOOP
-	for( uint i = 0; i < NUM_FILTER_SAMPLES; i++ )
+	for( int i = 0; i < NUM_FILTER_SAMPLES; i++ )
 	{
 		float2 E = Hammersley( i, NUM_FILTER_SAMPLES, 0 );
 		float3 H = TangentToWorld( ImportanceSampleGGX( E, Roughness ).xyz, N );
@@ -285,11 +265,11 @@ float CoefficientMask2;
 int NumSamples;
 
 void DiffuseIrradianceCopyPS(
-	FDownsampleGSOut Input,
+	FScreenVertexOutput Input,
 	out float4 OutColor : SV_Target0
 	)
 { 
-	float2 ScaledUVs = Input.Vertex.UV * 2 - 1;
+	float2 ScaledUVs = Input.UV * 2 - 1;
 	float3 CubeCoordinates = normalize(GetCubemapVector(ScaledUVs));
 
 	float SquaredUVs = 1 + dot(ScaledUVs, ScaledUVs);
@@ -307,32 +287,32 @@ float4 Sample01;
 float4 Sample23;
 
 void DiffuseIrradianceAccumulatePS(
-	FDownsampleGSOut Input,
+	FScreenVertexOutput Input,
 	out float4 OutColor : SV_Target0
 	)
 { 
 	float4 AccumulatedValue = 0;
 
 	{
-		float2 ScaledUVs = saturate(Input.Vertex.UV + Sample01.xy) * 2 - 1;
+		float2 ScaledUVs = saturate(Input.UV + Sample01.xy) * 2 - 1;
 		float3 CubeCoordinates = GetCubemapVector(ScaledUVs);
 		AccumulatedValue += SampleCubemap(CubeCoordinates);
 	}
 	
 	{
-		float2 ScaledUVs = saturate(Input.Vertex.UV + Sample01.zw) * 2 - 1;
+		float2 ScaledUVs = saturate(Input.UV + Sample01.zw) * 2 - 1;
 		float3 CubeCoordinates = GetCubemapVector(ScaledUVs);
 		AccumulatedValue += SampleCubemap(CubeCoordinates);
 	}
 
 	{
-		float2 ScaledUVs = saturate(Input.Vertex.UV + Sample23.xy) * 2 - 1;
+		float2 ScaledUVs = saturate(Input.UV + Sample23.xy) * 2 - 1;
 		float3 CubeCoordinates = GetCubemapVector(ScaledUVs);
 		AccumulatedValue += SampleCubemap(CubeCoordinates);
 	}
 
 	{
-		float2 ScaledUVs = saturate(Input.Vertex.UV + Sample23.zw) * 2 - 1;
+		float2 ScaledUVs = saturate(Input.UV + Sample23.zw) * 2 - 1;
 		float3 CubeCoordinates = GetCubemapVector(ScaledUVs);
 		AccumulatedValue += SampleCubemap(CubeCoordinates);
 	}
@@ -351,7 +331,7 @@ void AccumulateCubeFacesPS(
 	AccumulatedValue += TextureCubeSampleLevel(SourceTexture, SourceTextureSampler, float3(0, -1, 0), SourceMipIndex);
 	AccumulatedValue += TextureCubeSampleLevel(SourceTexture, SourceTextureSampler, float3(0, 0, 1), SourceMipIndex);
 	AccumulatedValue += TextureCubeSampleLevel(SourceTexture, SourceTextureSampler, float3(0, 0, -1), SourceMipIndex);
-	OutColor = float4(4 * PI * AccumulatedValue.rgb / ( max(AccumulatedValue.a, .00001f)), 0);
+	OutColor = float4(4 * PI * AccumulatedValue.rgb / ( max(AccumulatedValue.a, .00001f)), 0);	
 }
 
 // Used during reflection captures to get bounce light from specular surfaces
diff --git a/Engine/Source/Runtime/Renderer/Private/ReflectionEnvironmentCapture.cpp b/Engine/Source/Runtime/Renderer/Private/ReflectionEnvironmentCapture.cpp
--- a/Engine/Source/Runtime/Renderer/Private/ReflectionEnvironmentCapture.cpp
+++ b/Engine/Source/Runtime/Renderer/Private/ReflectionEnvironmentCapture.cpp
@@ -75,8 +75,6 @@ FSceneRenderTargetItem& GetEffectiveSourceTexture(bool bDownsamplePass, int32 Ta
 	return GSceneRenderTargets.ReflectionColorScratchCubemap[ScratchTextureIndex]->GetRenderTargetItem();
 }
 
-IMPLEMENT_SHADER_TYPE(,FDownsampleGS,TEXT("ReflectionEnvironmentShaders"),TEXT("DownsampleGS"),SF_Geometry);
-
 class FDownsamplePS : public FGlobalShader
 {
 	DECLARE_SHADER_TYPE(FDownsamplePS,Global);
@@ -84,7 +82,7 @@ public:
 
 	static bool ShouldCache(EShaderPlatform Platform)
 	{
-		return IsFeatureLevelSupported(Platform, ERHIFeatureLevel::SM4);
+		return true;
 	}
 
 	FDownsamplePS(const ShaderMetaType::CompiledShaderInitializerType& Initializer):
@@ -139,7 +137,7 @@ public:
 
 	static bool ShouldCache(EShaderPlatform Platform)
 	{
-		return IsFeatureLevelSupported(Platform, ERHIFeatureLevel::SM4);
+		return true;
 	}
 
 	static void ModifyCompilationEnvironment(EShaderPlatform Platform, FShaderCompilerEnvironment& OutEnvironment)
@@ -213,7 +211,7 @@ public:
 
 	static bool ShouldCache(EShaderPlatform Platform)
 	{
-		return IsFeatureLevelSupported(Platform, ERHIFeatureLevel::SM4);
+		return true;
 	}
 
 	static void ModifyCompilationEnvironment(EShaderPlatform Platform, FShaderCompilerEnvironment& OutEnvironment)
@@ -329,37 +327,35 @@ void FilterReflectionEnvironment(FRHICommandListImmediate& RHICmdList, ERHIFeatu
 			FSceneRenderTargetItem& EffectiveRT = GetEffectiveRenderTarget(true, MipIndex);
 			FSceneRenderTargetItem& EffectiveSource = GetEffectiveSourceTexture(true, MipIndex);
 			check(EffectiveRT.TargetableTexture != EffectiveSource.ShaderResourceTexture);
-
+			
+			for (int32 CubeFace = 0; CubeFace < CubeFace_MAX; CubeFace++)
 			{
-				for (int32 CubeFace = 0; CubeFace < CubeFace_MAX; CubeFace++)
-				{
-					SetRenderTarget(RHICmdList, EffectiveRT.TargetableTexture, MipIndex, CubeFace, NULL);
+				SetRenderTarget(RHICmdList, EffectiveRT.TargetableTexture, MipIndex, CubeFace, NULL);
 
-					const FIntRect ViewRect(0, 0, MipSize, MipSize);
-					RHICmdList.SetViewport(0, 0, 0.0f, MipSize, MipSize, 1.0f);
-					RHICmdList.SetRasterizerState(TStaticRasterizerState<FM_Solid, CM_None>::GetRHI());
-					RHICmdList.SetDepthStencilState(TStaticDepthStencilState<false, CF_Always>::GetRHI());
-					RHICmdList.SetBlendState(TStaticBlendState<>::GetRHI());
+				const FIntRect ViewRect(0, 0, MipSize, MipSize);
+				RHICmdList.SetViewport(0, 0, 0.0f, MipSize, MipSize, 1.0f);
+				RHICmdList.SetRasterizerState(TStaticRasterizerState<FM_Solid, CM_None>::GetRHI());
+				RHICmdList.SetDepthStencilState(TStaticDepthStencilState<false, CF_Always>::GetRHI());
+				RHICmdList.SetBlendState(TStaticBlendState<>::GetRHI());
 
-					TShaderMapRef<FScreenVSForGS> VertexShader(ShaderMap);
-					TShaderMapRef<FDownsamplePS> PixelShader(ShaderMap);
+				TShaderMapRef<FScreenVS> VertexShader(GetGlobalShaderMap(FeatureLevel));
+				TShaderMapRef<FDownsamplePS> PixelShader(GetGlobalShaderMap(FeatureLevel));
 
-					SetGlobalBoundShaderState(RHICmdList, FeatureLevel, DownsampleBoundShaderState, GFilterVertexDeclaration.VertexDeclarationRHI, *VertexShader, *PixelShader);
+				SetGlobalBoundShaderState(RHICmdList, FeatureLevel, DownsampleBoundShaderState, GFilterVertexDeclaration.VertexDeclarationRHI, *VertexShader, *PixelShader);
 
-					PixelShader->SetParameters(RHICmdList, CubeFace, SourceMipIndex, EffectiveSource);
+				PixelShader->SetParameters(RHICmdList, CubeFace, SourceMipIndex, EffectiveSource);
 
-					DrawRectangle( 
-						RHICmdList,
-						ViewRect.Min.X, ViewRect.Min.Y, 
-						ViewRect.Width(), ViewRect.Height(),
-						ViewRect.Min.X, ViewRect.Min.Y, 
-						ViewRect.Width(), ViewRect.Height(),
-						FIntPoint(ViewRect.Width(), ViewRect.Height()),
-						FIntPoint(MipSize, MipSize),
-						*VertexShader);
+				DrawRectangle( 
+					RHICmdList,
+					ViewRect.Min.X, ViewRect.Min.Y, 
+					ViewRect.Width(), ViewRect.Height(),
+					ViewRect.Min.X, ViewRect.Min.Y, 
+					ViewRect.Width(), ViewRect.Height(),
+					FIntPoint(ViewRect.Width(), ViewRect.Height()),
+					FIntPoint(MipSize, MipSize),
+					*VertexShader);
 
-					RHICmdList.CopyToResolveTarget(EffectiveRT.TargetableTexture, EffectiveRT.ShaderResourceTexture, true, FResolveParams(FResolveRect(), (ECubeFace)CubeFace, MipIndex));
-				}
+				RHICmdList.CopyToResolveTarget(EffectiveRT.TargetableTexture, EffectiveRT.ShaderResourceTexture, true, FResolveParams(FResolveRect(), (ECubeFace)CubeFace, MipIndex));
 			}
 
 			if (MipSize == GDiffuseIrradianceCubemapSize)
@@ -390,51 +386,49 @@ void FilterReflectionEnvironment(FRHICommandListImmediate& RHICmdList, ERHIFeatu
 			check(EffectiveRT.TargetableTexture != EffectiveSource.ShaderResourceTexture);
 			const int32 MipSize = 1 << (NumMips - MipIndex - 1);
 
+			for (int32 CubeFace = 0; CubeFace < CubeFace_MAX; CubeFace++)
 			{
-				for (int32 CubeFace = 0; CubeFace < CubeFace_MAX; CubeFace++)
-				{
-					SetRenderTarget(RHICmdList, EffectiveRT.TargetableTexture, MipIndex, CubeFace, NULL);
+				SetRenderTarget(RHICmdList, EffectiveRT.TargetableTexture, MipIndex, CubeFace, NULL);
 
-					const FIntRect ViewRect(0, 0, MipSize, MipSize);
-					RHICmdList.SetViewport(0, 0, 0.0f, MipSize, MipSize, 1.0f);
-					RHICmdList.SetRasterizerState(TStaticRasterizerState<FM_Solid, CM_None>::GetRHI());
-					RHICmdList.SetDepthStencilState(TStaticDepthStencilState<false, CF_Always>::GetRHI());
-					RHICmdList.SetBlendState(TStaticBlendState<>::GetRHI());
+				const FIntRect ViewRect(0, 0, MipSize, MipSize);
+				RHICmdList.SetViewport(0, 0, 0.0f, MipSize, MipSize, 1.0f);
+				RHICmdList.SetRasterizerState(TStaticRasterizerState<FM_Solid, CM_None>::GetRHI());
+				RHICmdList.SetDepthStencilState(TStaticDepthStencilState<false, CF_Always>::GetRHI());
+				RHICmdList.SetBlendState(TStaticBlendState<>::GetRHI());
 
-					TShaderMapRef<FScreenVSForGS> VertexShader(ShaderMap);
-					TShaderMapRef< TCubeFilterPS<1> > CaptureCubemapArrayPixelShader(ShaderMap);
-					//FFilterPS* PixelShader = (FFilterPS*)*CaptureCubemapArrayPixelShader;
+				TShaderMapRef<FScreenVS> VertexShader(GetGlobalShaderMap(FeatureLevel));
+				TShaderMapRef< TCubeFilterPS<1> > CaptureCubemapArrayPixelShader(GetGlobalShaderMap(FeatureLevel));
 
-					FCubeFilterPS* PixelShader;
-					if( bNormalize )
-					{
-						PixelShader = *TShaderMapRef< TCubeFilterPS<1> >(ShaderMap);
-						static FGlobalBoundShaderState BoundShaderState;
-						
-						SetGlobalBoundShaderState(RHICmdList, FeatureLevel, BoundShaderState, GFilterVertexDeclaration.VertexDeclarationRHI, *VertexShader, PixelShader);
-					}
-					else
-					{
-						PixelShader = *TShaderMapRef< TCubeFilterPS<0> >(ShaderMap);
-						static FGlobalBoundShaderState BoundShaderState;
+				FCubeFilterPS* PixelShader;
+				if( bNormalize )
+				{
+					PixelShader = *TShaderMapRef< TCubeFilterPS<1> >(ShaderMap);
+					static FGlobalBoundShaderState BoundShaderState;
+					
+					SetGlobalBoundShaderState(RHICmdList, FeatureLevel, BoundShaderState, GFilterVertexDeclaration.VertexDeclarationRHI, *VertexShader, PixelShader);
+				}
+				else
+				{
+					PixelShader = *TShaderMapRef< TCubeFilterPS<0> >(ShaderMap);
+					static FGlobalBoundShaderState BoundShaderState;
+					SetGlobalBoundShaderState(RHICmdList, FeatureLevel, BoundShaderState, GFilterVertexDeclaration.VertexDeclarationRHI, *VertexShader, PixelShader);
 						
-						SetGlobalBoundShaderState(RHICmdList, FeatureLevel, BoundShaderState, GFilterVertexDeclaration.VertexDeclarationRHI, *VertexShader, PixelShader);
-					}
+					SetGlobalBoundShaderState(RHICmdList, FeatureLevel, BoundShaderState, GFilterVertexDeclaration.VertexDeclarationRHI, *VertexShader, PixelShader);
+				}
 
-					PixelShader->SetParameters(RHICmdList, CubeFace, MipIndex, EffectiveSource);
+				PixelShader->SetParameters(RHICmdList, CubeFace, MipIndex, EffectiveSource);
 
-					DrawRectangle( 
-						RHICmdList,
-						ViewRect.Min.X, ViewRect.Min.Y, 
-						ViewRect.Width(), ViewRect.Height(),
-						ViewRect.Min.X, ViewRect.Min.Y, 
-						ViewRect.Width(), ViewRect.Height(),
-						FIntPoint(ViewRect.Width(), ViewRect.Height()),
-						FIntPoint(MipSize, MipSize),
-						*VertexShader);
+				DrawRectangle( 
+					RHICmdList,
+					ViewRect.Min.X, ViewRect.Min.Y, 
+					ViewRect.Width(), ViewRect.Height(),
+					ViewRect.Min.X, ViewRect.Min.Y, 
+					ViewRect.Width(), ViewRect.Height(),
+					FIntPoint(ViewRect.Width(), ViewRect.Height()),
+					FIntPoint(MipSize, MipSize),
+					*VertexShader);
 
-					RHICmdList.CopyToResolveTarget(EffectiveRT.TargetableTexture, EffectiveRT.ShaderResourceTexture, true, FResolveParams(FResolveRect(), (ECubeFace)CubeFace, MipIndex));
-				}
+				RHICmdList.CopyToResolveTarget(EffectiveRT.TargetableTexture, EffectiveRT.ShaderResourceTexture, true, FResolveParams(FResolveRect(), (ECubeFace)CubeFace, MipIndex));
 			}
 		}
 	}
@@ -448,7 +442,7 @@ public:
 
 	static bool ShouldCache(EShaderPlatform Platform)
 	{
-		return IsFeatureLevelSupported(Platform, ERHIFeatureLevel::SM4);
+		return true;
 	}
 
 	FCopyToCubeFaceVS()	{}
@@ -471,42 +465,6 @@ public:
 
 IMPLEMENT_SHADER_TYPE(,FCopyToCubeFaceVS,TEXT("ReflectionEnvironmentShaders"),TEXT("CopyToCubeFaceVS"),SF_Vertex);
 
-/** Geometry shader used when writing to a cubemap face.  Supports writing to a cubemap array as well. */
-class FRouteToCubeFaceGS : public FGlobalShader
-{
-	DECLARE_SHADER_TYPE(FRouteToCubeFaceGS,Global);
-public:
-
-	static bool ShouldCache(EShaderPlatform Platform)
-	{
-		return IsFeatureLevelSupported(Platform, ERHIFeatureLevel::SM4);
-	}
-
-	FRouteToCubeFaceGS(const ShaderMetaType::CompiledShaderInitializerType& Initializer):
-		FGlobalShader(Initializer)
-	{
-		CubeFace.Bind(Initializer.ParameterMap,TEXT("CubeFace"));
-	}
-	FRouteToCubeFaceGS() {}
-
-	void SetParameters(FRHICommandList& RHICmdList, uint32 CubeFaceValue)
-	{
-		SetShaderValue(RHICmdList, GetGeometryShader(), CubeFace, CubeFaceValue);
-	}
-
-	virtual bool Serialize(FArchive& Ar)
-	{
-		bool bShaderHasOutdatedParameters = FGlobalShader::Serialize(Ar);
-		Ar << CubeFace;
-		return bShaderHasOutdatedParameters;
-	}
-
-private:
-	FShaderParameter CubeFace;
-};
-
-IMPLEMENT_SHADER_TYPE(,FRouteToCubeFaceGS,TEXT("ReflectionEnvironmentShaders"),TEXT("RouteToCubeFaceGS"),SF_Geometry);
-
 /** Pixel shader used when copying scene color from a scene render into a face of a reflection capture cubemap. */
 class FCopySceneColorToCubeFacePS : public FGlobalShader
 {
@@ -515,7 +473,7 @@ public:
 
 	static bool ShouldCache(EShaderPlatform Platform)
 	{
-		return IsFeatureLevelSupported(Platform, ERHIFeatureLevel::SM4);
+		return true;
 	}
 
 	static void ModifyCompilationEnvironment(EShaderPlatform Platform, FShaderCompilerEnvironment& OutEnvironment)
@@ -600,7 +558,7 @@ public:
 
 	static bool ShouldCache(EShaderPlatform Platform)
 	{
-		return IsFeatureLevelSupported(Platform, ERHIFeatureLevel::SM4);
+		return true;
 	}
 
 	FCopyCubemapToCubeFacePS(const ShaderMetaType::CompiledShaderInitializerType& Initializer):
@@ -698,12 +656,10 @@ void ClearScratchCubemaps(FRHICommandList& RHICmdList)
 
 	for (int32 MipIndex = 0; MipIndex < NumMips; MipIndex++)
 	{
+		for (int32 CubeFace = 0; CubeFace < CubeFace_MAX; CubeFace++)
 		{
-			for (int32 CubeFace = 0; CubeFace < CubeFace_MAX; CubeFace++)
-			{
-				SetRenderTarget(RHICmdList, RT0.TargetableTexture, MipIndex, CubeFace, NULL);
-				RHICmdList.Clear(true, FLinearColor(0, 10000, 0, 0), false, 0, false, 0, FIntRect());
-			}
+			SetRenderTarget(RHICmdList, RT0.TargetableTexture, MipIndex, CubeFace, NULL);
+			RHICmdList.Clear(true, FLinearColor(0, 10000, 0, 0), false, 0, false, 0, FIntRect());
 		}
 	}
 
@@ -711,12 +667,10 @@ void ClearScratchCubemaps(FRHICommandList& RHICmdList)
 
 	for (int32 MipIndex = 0; MipIndex < NumMips; MipIndex++)
 	{
+		for (int32 CubeFace = 0; CubeFace < CubeFace_MAX; CubeFace++)
 		{
-			for (int32 CubeFace = 0; CubeFace < CubeFace_MAX; CubeFace++)
-			{
-				SetRenderTarget(RHICmdList, RT1.TargetableTexture, MipIndex, CubeFace, NULL);
-				RHICmdList.Clear(true, FLinearColor(0, 10000, 0, 0), false, 0, false, 0, FIntRect());
-			}
+			SetRenderTarget(RHICmdList, RT1.TargetableTexture, MipIndex, CubeFace, NULL);
+			RHICmdList.Clear(true, FLinearColor(0, 10000, 0, 0), false, 0, false, 0, FIntRect());
 		}
 	}
 }
@@ -755,35 +709,36 @@ void CaptureSceneToScratchCubemap(FRHICommandListImmediate& RHICmdList, FSceneRe
 
 		{
 			SCOPED_DRAW_EVENT(RHICmdList, CubeMapCopyScene, DEC_SCENE_ITEMS);
-			{
-				// Copy the captured scene into the cubemap face
-				SetRenderTarget(RHICmdList, EffectiveColorRT.TargetableTexture, 0, CubeFace, NULL);
+			
+			// Copy the captured scene into the cubemap face
+			SetRenderTarget(RHICmdList, EffectiveColorRT.TargetableTexture, 0, CubeFace, NULL);
 
-				const FIntRect ViewRect(0, 0, EffectiveSize, EffectiveSize);
-				RHICmdList.SetViewport(0, 0, 0.0f, EffectiveSize, EffectiveSize, 1.0f);
-				RHICmdList.SetRasterizerState(TStaticRasterizerState<FM_Solid, CM_None>::GetRHI());
-				RHICmdList.SetDepthStencilState(TStaticDepthStencilState<false, CF_Always>::GetRHI());
-				RHICmdList.SetBlendState(TStaticBlendState<>::GetRHI());
+			const FIntRect ViewRect(0, 0, EffectiveSize, EffectiveSize);
+			RHICmdList.SetViewport(0, 0, 0.0f, EffectiveSize, EffectiveSize, 1.0f);
+			RHICmdList.SetRasterizerState(TStaticRasterizerState<FM_Solid, CM_None>::GetRHI());
+			RHICmdList.SetDepthStencilState(TStaticDepthStencilState<false, CF_Always>::GetRHI());
+			RHICmdList.SetBlendState(TStaticBlendState<>::GetRHI());
 
-				TShaderMapRef<FCopyToCubeFaceVS> VertexShader(ShaderMap);
-				TShaderMapRef<FCopySceneColorToCubeFacePS> PixelShader(ShaderMap);
-				SetGlobalBoundShaderState(RHICmdList, FeatureLevel, CopyColorCubemapBoundShaderState, GFilterVertexDeclaration.VertexDeclarationRHI, *VertexShader, *PixelShader);
+			TShaderMapRef<FCopyToCubeFaceVS> VertexShader(GetGlobalShaderMap(FeatureLevel));
+			TShaderMapRef<FCopySceneColorToCubeFacePS> PixelShader(GetGlobalShaderMap(FeatureLevel));
+			SetGlobalBoundShaderState(RHICmdList, FeatureLevel, CopyColorCubemapBoundShaderState, GFilterVertexDeclaration.VertexDeclarationRHI, *VertexShader, *PixelShader);
 
-				PixelShader->SetParameters(RHICmdList, SceneRenderer->Views[0], bCapturingForSkyLight, bLowerHemisphereIsBlack);
-				VertexShader->SetParameters(RHICmdList, SceneRenderer->Views[0]);
+			SetGlobalBoundShaderState(RHICmdList, FeatureLevel, CopyColorCubemapBoundShaderState, GFilterVertexDeclaration.VertexDeclarationRHI, *VertexShader, *PixelShader);
 
-				DrawRectangle( 
-					RHICmdList,
-					ViewRect.Min.X, ViewRect.Min.Y, 
-					ViewRect.Width(), ViewRect.Height(),
-					ViewRect.Min.X, ViewRect.Min.Y, 
-					ViewRect.Width() * GSupersampleCaptureFactor, ViewRect.Height() * GSupersampleCaptureFactor,
-					FIntPoint(ViewRect.Width(), ViewRect.Height()),
-					GSceneRenderTargets.GetBufferSizeXY(),
-					*VertexShader);
+			PixelShader->SetParameters(RHICmdList, SceneRenderer->Views[0], bCapturingForSkyLight, bLowerHemisphereIsBlack);
+			VertexShader->SetParameters(RHICmdList, SceneRenderer->Views[0]);
 
-				RHICmdList.CopyToResolveTarget(EffectiveColorRT.TargetableTexture, EffectiveColorRT.ShaderResourceTexture, true, FResolveParams(FResolveRect(), CubeFace));
-			}
+			DrawRectangle( 
+				RHICmdList,
+				ViewRect.Min.X, ViewRect.Min.Y, 
+				ViewRect.Width(), ViewRect.Height(),
+				ViewRect.Min.X, ViewRect.Min.Y, 
+				ViewRect.Width() * GSupersampleCaptureFactor, ViewRect.Height() * GSupersampleCaptureFactor,
+				FIntPoint(ViewRect.Width(), ViewRect.Height()),
+				GSceneRenderTargets.GetBufferSizeXY(),
+				*VertexShader);
+
+			RHICmdList.CopyToResolveTarget(EffectiveColorRT.TargetableTexture, EffectiveColorRT.ShaderResourceTexture, true, FResolveParams(FResolveRect(), CubeFace));			
 		}
 	}
 
@@ -799,40 +754,36 @@ void CopyCubemapToScratchCubemap(FRHICommandList& RHICmdList, ERHIFeatureLevel::
 	const int32 EffectiveSize = GReflectionCaptureSize;
 	FSceneRenderTargetItem& EffectiveColorRT =  GSceneRenderTargets.ReflectionColorScratchCubemap[0]->GetRenderTargetItem();
 
-	auto ShaderMap = GetGlobalShaderMap(FeatureLevel);
-
+	for (uint32 CubeFace = 0; CubeFace < CubeFace_MAX; CubeFace++)
 	{
-		for (uint32 CubeFace = 0; CubeFace < CubeFace_MAX; CubeFace++)
-		{
-			// Copy the captured scene into the cubemap face
-			SetRenderTarget(RHICmdList, EffectiveColorRT.TargetableTexture, 0, CubeFace, NULL);
+		// Copy the captured scene into the cubemap face
+		SetRenderTarget(RHICmdList, EffectiveColorRT.TargetableTexture, 0, CubeFace, NULL);
 
-			const FTexture* SourceCubemapResource = SourceCubemap->Resource;
-			const FIntPoint SourceDimensions(SourceCubemapResource->GetSizeX(), SourceCubemapResource->GetSizeY());
-			const FIntRect ViewRect(0, 0, EffectiveSize, EffectiveSize);
-			RHICmdList.SetViewport(0, 0, 0.0f, EffectiveSize, EffectiveSize, 1.0f);
-			RHICmdList.SetRasterizerState(TStaticRasterizerState<FM_Solid, CM_None>::GetRHI());
-			RHICmdList.SetDepthStencilState(TStaticDepthStencilState<false, CF_Always>::GetRHI());
-			RHICmdList.SetBlendState(TStaticBlendState<>::GetRHI());
+		const FTexture* SourceCubemapResource = SourceCubemap->Resource;
+		const FIntPoint SourceDimensions(SourceCubemapResource->GetSizeX(), SourceCubemapResource->GetSizeY());
+		const FIntRect ViewRect(0, 0, EffectiveSize, EffectiveSize);
+		RHICmdList.SetViewport(0, 0, 0.0f, EffectiveSize, EffectiveSize, 1.0f);
+		RHICmdList.SetRasterizerState(TStaticRasterizerState<FM_Solid, CM_None>::GetRHI());
+		RHICmdList.SetDepthStencilState(TStaticDepthStencilState<false, CF_Always>::GetRHI());
+		RHICmdList.SetBlendState(TStaticBlendState<>::GetRHI());
 
-			TShaderMapRef<FScreenVSForGS> VertexShader(ShaderMap);
-			TShaderMapRef<FCopyCubemapToCubeFacePS> PixelShader(ShaderMap);
+		TShaderMapRef<FScreenVS> VertexShader(GetGlobalShaderMap(FeatureLevel));
+		TShaderMapRef<FCopyCubemapToCubeFacePS> PixelShader(GetGlobalShaderMap(FeatureLevel));
 
-			SetGlobalBoundShaderState(RHICmdList, FeatureLevel, CopyFromCubemapToCubemapBoundShaderState, GFilterVertexDeclaration.VertexDeclarationRHI, *VertexShader, *PixelShader);
-			PixelShader->SetParameters(RHICmdList, SourceCubemapResource, CubeFace, bIsSkyLight, bLowerHemisphereIsBlack);
+		SetGlobalBoundShaderState(RHICmdList, FeatureLevel, CopyFromCubemapToCubemapBoundShaderState, GFilterVertexDeclaration.VertexDeclarationRHI, *VertexShader, *PixelShader);
+		PixelShader->SetParameters(RHICmdList, SourceCubemapResource, CubeFace, bIsSkyLight, bLowerHemisphereIsBlack);
 
-			DrawRectangle( 
-				RHICmdList,
-				ViewRect.Min.X, ViewRect.Min.Y, 
-				ViewRect.Width(), ViewRect.Height(),
-				0, 0, 
-				SourceDimensions.X, SourceDimensions.Y,
-				FIntPoint(ViewRect.Width(), ViewRect.Height()),
-				SourceDimensions,
-				*VertexShader);
-
-			RHICmdList.CopyToResolveTarget(EffectiveColorRT.TargetableTexture, EffectiveColorRT.ShaderResourceTexture, true, FResolveParams(FResolveRect(), (ECubeFace)CubeFace));
-		}
+		DrawRectangle(
+			RHICmdList,
+			ViewRect.Min.X, ViewRect.Min.Y, 
+			ViewRect.Width(), ViewRect.Height(),
+			0, 0, 
+			SourceDimensions.X, SourceDimensions.Y,
+			FIntPoint(ViewRect.Width(), ViewRect.Height()),
+			SourceDimensions,
+			*VertexShader);
+
+		RHICmdList.CopyToResolveTarget(EffectiveColorRT.TargetableTexture, EffectiveColorRT.ShaderResourceTexture, true, FResolveParams(FResolveRect(), (ECubeFace)CubeFace));
 	}
 }
 
@@ -1380,8 +1331,9 @@ void CopyToSkyTexture(FRHICommandList& RHICmdList, FScene* Scene, FTexture* Proc
 }
 
 void FScene::UpdateSkyCaptureContents(const USkyLightComponent* CaptureComponent, bool bCaptureEmissiveOnly, FTexture* OutProcessedTexture, FSHVectorRGB3& OutIrradianceEnvironmentMap)
-{
-	if (GetFeatureLevel() >= ERHIFeatureLevel::SM4)
+{	
+	//todo maybe use ShouldRenderSkylight here, but currently the Skylight hasn't been loaded in time.
+	if (GSupportsRenderTargetFormat_PF_FloatRGBA || GetFeatureLevel() >= ERHIFeatureLevel::SM4)
 	{
 		ENQUEUE_UNIQUE_RENDER_COMMAND( 
 			ClearCommand,
diff --git a/Engine/Source/Runtime/Renderer/Private/ReflectionEnvironmentCapture.h b/Engine/Source/Runtime/Renderer/Private/ReflectionEnvironmentCapture.h
--- a/Engine/Source/Runtime/Renderer/Private/ReflectionEnvironmentCapture.h
+++ b/Engine/Source/Runtime/Renderer/Private/ReflectionEnvironmentCapture.h
@@ -9,36 +9,3 @@
 extern ENGINE_API int32 GReflectionCaptureSize;
 
 extern void ComputeDiffuseIrradiance(FRHICommandListImmediate& RHICmdList, ERHIFeatureLevel::Type FeatureLevel, FTextureRHIRef LightingSource, int32 LightingSourceMipIndex, FSHVectorRGB3* OutIrradianceEnvironmentMap);
-
-class FDownsampleGS : public FGlobalShader
-{
-	DECLARE_SHADER_TYPE(FDownsampleGS,Global);
-public:
-
-	static bool ShouldCache(EShaderPlatform Platform)
-	{
-		return IsFeatureLevelSupported(Platform, ERHIFeatureLevel::SM4);
-	}
-
-	FDownsampleGS(const ShaderMetaType::CompiledShaderInitializerType& Initializer):
-		FGlobalShader(Initializer)
-	{
-		CubeFace.Bind(Initializer.ParameterMap,TEXT("CubeFace"));
-	}
-	FDownsampleGS() {}
-
-	void SetParameters(FRHICommandList& RHICmdList, int32 CubeFaceValue)
-	{
-		SetShaderValue(RHICmdList, GetGeometryShader(), CubeFace, CubeFaceValue);
-	}
-
-	virtual bool Serialize(FArchive& Ar)
-	{
-		bool bShaderHasOutdatedParameters = FGlobalShader::Serialize(Ar);
-		Ar << CubeFace;
-		return bShaderHasOutdatedParameters;
-	}
-
-private:
-	FShaderParameter CubeFace;
-};
\ No newline at end of file
