ES2 preview - ES2 post process path deals with sub rectangles of render targets.

diff --git a/Engine/Source/Runtime/Renderer/Private/PostProcess/PostProcessMobile.cpp b/Engine/Source/Runtime/Renderer/Private/PostProcess/PostProcessMobile.cpp
--- a/Engine/Source/Runtime/Renderer/Private/PostProcess/PostProcessMobile.cpp
+++ b/Engine/Source/Runtime/Renderer/Private/PostProcess/PostProcessMobile.cpp
@@ -185,7 +185,7 @@ void FRCPassPostProcessBloomSetupES2::Process(FRenderingCompositePassContext& Co
 
 	FIntPoint SrcSize;
 	FIntRect SrcRect;
-	if(bUsedFramebufferFetch)
+	if(bUseViewRectSource)
 	{
 		// Mobile with framebuffer fetch uses view rect as source.
 		const FSceneView& View = Context.View;
@@ -386,7 +386,7 @@ void FRCPassPostProcessBloomSetupSmallES2::Process(FRenderingCompositePassContex
 
 	FIntPoint SrcSize;
 	FIntRect SrcRect;
-	if(bUsedFramebufferFetch)
+	if(bUseViewRectSource)
 	{
 		// Mobile with framebuffer fetch uses view rect as source.
 		const FSceneView& View = Context.View;
@@ -1965,7 +1965,7 @@ void FRCPassPostProcessDofDownES2::Process(FRenderingCompositePassContext& Conte
 
 	FIntPoint SrcSize;
 	FIntRect SrcRect;
-	if(bUsedFramebufferFetch)
+	if(bUseViewRectSource)
 	{
 		// Mobile with framebuffer fetch uses view rect as source.
 		const FSceneView& View = Context.View;
@@ -2158,6 +2158,9 @@ void FRCPassPostProcessDofNearES2::Process(FRenderingCompositePassContext& Conte
 {
 	SCOPED_DRAW_EVENT(Context.RHICmdList, PostProcessDofNear, DEC_SCENE_ITEMS);
 
+	const FPooledRenderTargetDesc* InputDesc = GetInputDesc(ePId_Input0);
+	FIntPoint SrcSize = InputDesc->Extent;
+
 	uint32 DstX = FMath::Max(1, PrePostSourceViewportSize.X/4);
 	uint32 DstY = FMath::Max(1, PrePostSourceViewportSize.Y/4);
 
@@ -2191,7 +2194,7 @@ void FRCPassPostProcessDofNearES2::Process(FRenderingCompositePassContext& Conte
 		0, 0,
 		DstX, DstY,
 		SrcDstSize,
-		SrcDstSize,
+		SrcSize,
 		*VertexShader,
 		EDRF_UseTriangleOptimization);
 
@@ -2731,6 +2734,9 @@ void FRCPassPostProcessAaES2::Process(FRenderingCompositePassContext& Context)
 {
 	SCOPED_DRAW_EVENT(Context.RHICmdList, PostProcessAa, DEC_SCENE_ITEMS);
 
+	const FPooledRenderTargetDesc* InputDesc = GetInputDesc(ePId_Input0);
+	FIntPoint SrcSize = InputDesc->Extent;
+
 	uint32 DstX = FMath::Max(1, PrePostSourceViewportSize.X);
 	uint32 DstY = FMath::Max(1, PrePostSourceViewportSize.Y);
 
@@ -2764,7 +2770,7 @@ void FRCPassPostProcessAaES2::Process(FRenderingCompositePassContext& Context)
 		0, 0,
 		DstX, DstY,
 		SrcDstSize,
-		SrcDstSize,
+		SrcSize,
 		*VertexShader,
 		EDRF_UseTriangleOptimization);
 
diff --git a/Engine/Source/Runtime/Renderer/Private/PostProcess/PostProcessMobile.h b/Engine/Source/Runtime/Renderer/Private/PostProcess/PostProcessMobile.h
--- a/Engine/Source/Runtime/Renderer/Private/PostProcess/PostProcessMobile.h
+++ b/Engine/Source/Runtime/Renderer/Private/PostProcess/PostProcessMobile.h
@@ -11,26 +11,26 @@
 class FRCPassPostProcessBloomSetupES2 : public TRenderingCompositePassBase<1, 1>
 {
 public:
-	FRCPassPostProcessBloomSetupES2(FIntPoint InPrePostSourceViewportSize, bool bInUsedFramebufferFetch) : PrePostSourceViewportSize(InPrePostSourceViewportSize), bUsedFramebufferFetch(bInUsedFramebufferFetch) { }
+	FRCPassPostProcessBloomSetupES2(FIntPoint InPrePostSourceViewportSize, bool bInUseViewRectSource) : PrePostSourceViewportSize(InPrePostSourceViewportSize), bUseViewRectSource(bInUseViewRectSource) { }
 	virtual void Process(FRenderingCompositePassContext& Context);
 	virtual FPooledRenderTargetDesc ComputeOutputDesc(EPassOutputId InPassOutputId) const;
 	virtual void Release() override { delete this; }
 private:
 	FIntPoint PrePostSourceViewportSize;
-	bool bUsedFramebufferFetch;
+	bool bUseViewRectSource;
 	void SetShader(const FRenderingCompositePassContext& Context);
 };
 
 class FRCPassPostProcessBloomSetupSmallES2 : public TRenderingCompositePassBase<1, 1>
 {
 public:
-	FRCPassPostProcessBloomSetupSmallES2(FIntPoint InPrePostSourceViewportSize, bool bInUsedFramebufferFetch) : PrePostSourceViewportSize(InPrePostSourceViewportSize), bUsedFramebufferFetch(bInUsedFramebufferFetch) { }
+	FRCPassPostProcessBloomSetupSmallES2(FIntPoint InPrePostSourceViewportSize, bool bInUseViewRectSource) : PrePostSourceViewportSize(InPrePostSourceViewportSize), bUseViewRectSource(bInUseViewRectSource) { }
 	virtual void Process(FRenderingCompositePassContext& Context);
 	virtual FPooledRenderTargetDesc ComputeOutputDesc(EPassOutputId InPassOutputId) const;
 	virtual void Release() override { delete this; }
 private:
 	FIntPoint PrePostSourceViewportSize;
-	bool bUsedFramebufferFetch;
+	bool bUseViewRectSource;
 	void SetShader(const FRenderingCompositePassContext& Context);
 };
 
@@ -49,13 +49,13 @@ private:
 class FRCPassPostProcessDofDownES2 : public TRenderingCompositePassBase<2, 1>
 {
 public:
-	FRCPassPostProcessDofDownES2(FIntPoint InPrePostSourceViewportSize, bool bInUsedFramebufferFetch) : PrePostSourceViewportSize(InPrePostSourceViewportSize), bUsedFramebufferFetch(bInUsedFramebufferFetch) { }
+	FRCPassPostProcessDofDownES2(FIntPoint InPrePostSourceViewportSize, bool bInUseViewRectSource) : PrePostSourceViewportSize(InPrePostSourceViewportSize), bUseViewRectSource(bInUseViewRectSource) { }
 	virtual void Process(FRenderingCompositePassContext& Context);
 	virtual FPooledRenderTargetDesc ComputeOutputDesc(EPassOutputId InPassOutputId) const;
 	virtual void Release() { delete this; }
 private:
 	FIntPoint PrePostSourceViewportSize;
-	bool bUsedFramebufferFetch;
+	bool bUseViewRectSource;
 	void SetShader(const FRenderingCompositePassContext& Context);
 };
 
diff --git a/Engine/Source/Runtime/Renderer/Private/PostProcess/PostProcessTonemap.cpp b/Engine/Source/Runtime/Renderer/Private/PostProcess/PostProcessTonemap.cpp
--- a/Engine/Source/Runtime/Renderer/Private/PostProcess/PostProcessTonemap.cpp
+++ b/Engine/Source/Runtime/Renderer/Private/PostProcess/PostProcessTonemap.cpp
@@ -1595,12 +1595,12 @@ void FRCPassPostProcessTonemapES2::Process(FRenderingCompositePassContext& Conte
 	
 	const FSceneView& View = Context.View;
 	const FSceneViewFamily& ViewFamily = *(View.Family);
+	const FSceneRenderTargetItem& DestRenderTarget = PassOutputs[0].RequestSurface(Context);
 
 	FIntRect SrcRect = View.ViewRect;
 	FIntRect DestRect = View.ViewRect;
 	FIntPoint SrcSize = InputDesc->Extent;
-
-	const FSceneRenderTargetItem& DestRenderTarget = PassOutputs[0].RequestSurface(Context);
+	FIntPoint DstSize = DestRect.Size();
 
 	// Set the view family's render target/viewport.
 	SetRenderTarget(Context.RHICmdList, DestRenderTarget.TargetableTexture, FTextureRHIParamRef());
@@ -1669,8 +1669,8 @@ void FRCPassPostProcessTonemapES2::Process(FRenderingCompositePassContext& Conte
 		View.ViewRect.Width(), View.ViewRect.Height(),
 		View.ViewRect.Min.X, View.ViewRect.Min.Y, 
 		View.ViewRect.Width(), View.ViewRect.Height(),
-		View.ViewRect.Size(),
-		GSceneRenderTargets.GetBufferSizeXY(),
+		DstSize,
+		SrcSize,
 		*VertexShader,
 		EDRF_UseTriangleOptimization);
 
diff --git a/Engine/Source/Runtime/Renderer/Private/PostProcess/PostProcessing.cpp b/Engine/Source/Runtime/Renderer/Private/PostProcess/PostProcessing.cpp
--- a/Engine/Source/Runtime/Renderer/Private/PostProcess/PostProcessing.cpp
+++ b/Engine/Source/Runtime/Renderer/Private/PostProcess/PostProcessing.cpp
@@ -1162,6 +1162,10 @@ void FPostProcessing::ProcessES2(FRHICommandListImmediate& RHICmdList, FViewInfo
 			}
 		}
 
+		FIntPoint PrePostSourceViewportSize = View.ViewRect.Size();
+		// ES2 preview uses a subsection of the scene RT, bUsedFramebufferFetch == true deals with this case.  
+		bool bViewRectSource = bUsedFramebufferFetch || GSceneRenderTargets.GetBufferSizeXY() != PrePostSourceViewportSize;
+
 		// add the passes we want to add to the graph (commenting a line means the pass is not inserted into the graph) ---------
 		if( View.Family->EngineShowFlags.PostProcessing )
 		{
@@ -1181,7 +1185,6 @@ void FPostProcessing::ProcessES2(FRHICommandListImmediate& RHICmdList, FViewInfo
 			bUsePost &= GSupportsRenderTargetFormat_PF_FloatRGBA;
 			bUsePost &= IsMobileHDR();
 
-			FIntPoint PrePostSourceViewportSize = View.ViewRect.Size();
 
 			if(bUsePost)
 			{
@@ -1199,13 +1202,13 @@ void FPostProcessing::ProcessES2(FRHICommandListImmediate& RHICmdList, FViewInfo
 				{
 					if(bUseBloomSmall)
 					{
-						FRenderingCompositePass* Pass = Context.Graph.RegisterPass(new(FMemStack::Get()) FRCPassPostProcessBloomSetupSmallES2(PrePostSourceViewportSize, bUsedFramebufferFetch));
+						FRenderingCompositePass* Pass = Context.Graph.RegisterPass(new(FMemStack::Get()) FRCPassPostProcessBloomSetupSmallES2(PrePostSourceViewportSize, bViewRectSource));
 						Pass->SetInput(ePId_Input0, Context.FinalOutput);
 						PostProcessBloomSetup = FRenderingCompositeOutputRef(Pass);
 					}
 					else
 					{
-						FRenderingCompositePass* Pass = Context.Graph.RegisterPass(new(FMemStack::Get()) FRCPassPostProcessBloomSetupES2(PrePostSourceViewportSize, bUsedFramebufferFetch));
+						FRenderingCompositePass* Pass = Context.Graph.RegisterPass(new(FMemStack::Get()) FRCPassPostProcessBloomSetupES2(PrePostSourceViewportSize, bViewRectSource));
 						Pass->SetInput(ePId_Input0, Context.FinalOutput);
 						PostProcessBloomSetup = FRenderingCompositeOutputRef(Pass);
 					}
@@ -1226,7 +1229,7 @@ void FPostProcessing::ProcessES2(FRHICommandListImmediate& RHICmdList, FViewInfo
 					// Samples at full resolution, writes to 1/4 area.
 					FRenderingCompositeOutputRef PostProcessDofDown;
 					{
-						FRenderingCompositePass* Pass = Context.Graph.RegisterPass(new(FMemStack::Get()) FRCPassPostProcessDofDownES2(PrePostSourceViewportSize, bUsedFramebufferFetch));
+						FRenderingCompositePass* Pass = Context.Graph.RegisterPass(new(FMemStack::Get()) FRCPassPostProcessDofDownES2(PrePostSourceViewportSize, bViewRectSource));
 						Pass->SetInput(ePId_Input0, Context.FinalOutput);
 						Pass->SetInput(ePId_Input1, PostProcessNear);
 						PostProcessDofDown = FRenderingCompositeOutputRef(Pass);
@@ -1412,7 +1415,7 @@ void FPostProcessing::ProcessES2(FRHICommandListImmediate& RHICmdList, FViewInfo
 
 		// Composite editor primitives if we had any to draw and compositing is enabled
 		// TODO: Move FDeferredShadingSceneRenderer::ShouldCompositeEditorPrimitives somewhere more generic
-		if (FDeferredShadingSceneRenderer::ShouldCompositeEditorPrimitives(View))
+		if (FDeferredShadingSceneRenderer::ShouldCompositeEditorPrimitives(View) && !DofOutput.IsValid())
 		{
 			FRenderingCompositePass* EditorCompNode = Context.Graph.RegisterPass(new(FMemStack::Get()) FRCPassPostProcessCompositeEditorPrimitives());
 			EditorCompNode->SetInput(ePId_Input0, FRenderingCompositeOutputRef(Context.FinalOutput));
@@ -1421,7 +1424,7 @@ void FPostProcessing::ProcessES2(FRHICommandListImmediate& RHICmdList, FViewInfo
 		}
 
 		// Must run to blit to back buffer even if post processing is off.
-		FRenderingCompositePass* PostProcessTonemap = Context.Graph.RegisterPass(new(FMemStack::Get()) FRCPassPostProcessTonemapES2(Context.View, bUsedFramebufferFetch));
+		FRenderingCompositePass* PostProcessTonemap = Context.Graph.RegisterPass(new(FMemStack::Get()) FRCPassPostProcessTonemapES2(Context.View, bViewRectSource));
 		PostProcessTonemap->SetInput(ePId_Input0, Context.FinalOutput);
 		PostProcessTonemap->SetInput(ePId_Input1, BloomOutput);
 		PostProcessTonemap->SetInput(ePId_Input2, DofOutput);
@@ -1448,6 +1451,18 @@ void FPostProcessing::ProcessES2(FRHICommandListImmediate& RHICmdList, FViewInfo
 			Context.FinalOutput = FRenderingCompositeOutputRef(PostProcessAa);
 		}
 
+		{
+			// temporary work around for EditorPrimitives / Dof alpha issues.
+			if (FDeferredShadingSceneRenderer::ShouldCompositeEditorPrimitives(View) && DofOutput.IsValid())
+			{
+				// TODO: combine editor prims in tonemap pass or remove AA jitter
+				FRenderingCompositePass* EditorCompNode = Context.Graph.RegisterPass(new(FMemStack::Get()) FRCPassPostProcessCompositeEditorPrimitives());
+				EditorCompNode->SetInput(ePId_Input0, FRenderingCompositeOutputRef(Context.FinalOutput));
+				//Node->SetInput(ePId_Input1, FRenderingCompositeOutputRef(Context.SceneDepth));
+				Context.FinalOutput = FRenderingCompositeOutputRef(EditorCompNode);
+			}
+		}
+
 		if(View.Family->EngineShowFlags.ShaderComplexity)
 		{
 			FRenderingCompositePass* Node = Context.Graph.RegisterPass(new(FMemStack::Get()) FRCPassPostProcessVisualizeComplexity(GEngine->ShaderComplexityColors));
diff --git a/Engine/Source/Runtime/Renderer/Private/PostProcess/PostProcessing.h b/Engine/Source/Runtime/Renderer/Private/PostProcess/PostProcessing.h
--- a/Engine/Source/Runtime/Renderer/Private/PostProcess/PostProcessing.h
+++ b/Engine/Source/Runtime/Renderer/Private/PostProcess/PostProcessing.h
@@ -73,7 +73,7 @@ public:
 	// @param VelocityRT only valid if motion blur is supported
 	void Process(FRHICommandListImmediate& RHICmdList, FViewInfo& View, TRefCountPtr<IPooledRenderTarget>& VelocityRT);
 
-	void ProcessES2(FRHICommandListImmediate& RHICmdList, FViewInfo& View, bool bUsedFramebufferFetch);
+	void ProcessES2(FRHICommandListImmediate& RHICmdList, FViewInfo& View, bool bViewRectSource);
 };
 
 /** The global used for post processing. */
