UE4 - add stats for shadow setup

diff --git a/Engine/Source/Runtime/RenderCore/Private/RenderCore.cpp b/Engine/Source/Runtime/RenderCore/Private/RenderCore.cpp
--- a/Engine/Source/Runtime/RenderCore/Private/RenderCore.cpp
+++ b/Engine/Source/Runtime/RenderCore/Private/RenderCore.cpp
@@ -73,6 +73,11 @@ DEFINE_STAT(STAT_GatherShadowPrimitivesTime);
 DEFINE_STAT(STAT_UpdateIndirectLightingCache);
 DEFINE_STAT(STAT_SortStaticDrawLists);
 DEFINE_STAT(STAT_InitDynamicShadowsTime);
+DEFINE_STAT(STAT_InitProjectedShadowVisibility);
+DEFINE_STAT(STAT_UpdatePreshadowCache);
+DEFINE_STAT(STAT_CreateWholeSceneProjectedShadow);
+DEFINE_STAT(STAT_AddViewDependentWholeSceneShadowsForView);
+DEFINE_STAT(STAT_SetupInteractionShadows);
 DEFINE_STAT(STAT_PreRenderView);
 DEFINE_STAT(STAT_GetDynamicMeshElements);
 DEFINE_STAT(STAT_UpdateStaticMeshesTime);
diff --git a/Engine/Source/Runtime/RenderCore/Public/RenderCore.h b/Engine/Source/Runtime/RenderCore/Public/RenderCore.h
--- a/Engine/Source/Runtime/RenderCore/Public/RenderCore.h
+++ b/Engine/Source/Runtime/RenderCore/Public/RenderCore.h
@@ -52,6 +52,11 @@ DECLARE_CYCLE_STAT_EXTERN(TEXT("UpdateStaticMeshes"),STAT_UpdateStaticMeshesTime
 DECLARE_CYCLE_STAT_EXTERN(TEXT("PreRenderView"),STAT_PreRenderView,STATGROUP_InitViews, RENDERCORE_API);
 DECLARE_CYCLE_STAT_EXTERN(TEXT("GetDynamicMeshElements"),STAT_GetDynamicMeshElements,STATGROUP_InitViews, RENDERCORE_API);
 DECLARE_CYCLE_STAT_EXTERN(TEXT("Init dynamic shadows"),STAT_InitDynamicShadowsTime,STATGROUP_InitViews, RENDERCORE_API);
+DECLARE_CYCLE_STAT_EXTERN(TEXT("Update Preshadow Cache"),STAT_UpdatePreshadowCache,STATGROUP_InitViews, RENDERCORE_API);
+DECLARE_CYCLE_STAT_EXTERN(TEXT("Init Projected Shadow"),STAT_InitProjectedShadowVisibility,STATGROUP_InitViews, RENDERCORE_API);
+DECLARE_CYCLE_STAT_EXTERN(TEXT("Create WholeScene Shadow"),STAT_CreateWholeSceneProjectedShadow,STATGROUP_InitViews, RENDERCORE_API);
+DECLARE_CYCLE_STAT_EXTERN(TEXT("Add View Whole Scene Shadows"),STAT_AddViewDependentWholeSceneShadowsForView,STATGROUP_InitViews, RENDERCORE_API);
+DECLARE_CYCLE_STAT_EXTERN(TEXT("Setup Interaction Shadows"),STAT_SetupInteractionShadows,STATGROUP_InitViews, RENDERCORE_API);
 DECLARE_CYCLE_STAT_EXTERN(TEXT("SortStaticDrawLists"),STAT_SortStaticDrawLists,STATGROUP_InitViews, RENDERCORE_API);
 DECLARE_CYCLE_STAT_EXTERN(TEXT("Update Indirect Lighting Cache"),STAT_UpdateIndirectLightingCache,STATGROUP_InitViews, RENDERCORE_API);
 DECLARE_CYCLE_STAT_EXTERN(TEXT("GatherShadowPrimitives"),STAT_GatherShadowPrimitivesTime,STATGROUP_InitViews, RENDERCORE_API);
diff --git a/Engine/Source/Runtime/Renderer/Private/ShadowSetup.cpp b/Engine/Source/Runtime/Renderer/Private/ShadowSetup.cpp
--- a/Engine/Source/Runtime/Renderer/Private/ShadowSetup.cpp
+++ b/Engine/Source/Runtime/Renderer/Private/ShadowSetup.cpp
@@ -961,6 +961,7 @@ void FDeferredShadingSceneRenderer::UpdatePreshadowCache()
 {
 	if (ShouldUseCachePreshadows() && !Views[0].bIsSceneCapture)
 	{
+		SCOPE_CYCLE_COUNTER(STAT_UpdatePreshadowCache);
 		if (Scene->PreshadowCacheLayout.GetSizeX() == 0)
 		{
 			// Initialize the texture layout if necessary
@@ -1060,6 +1061,9 @@ void FDeferredShadingSceneRenderer::SetupInteractionShadows(
 	const TArray<FProjectedShadowInfo*,SceneRenderingAllocator>& ViewDependentWholeSceneShadows,
 	TArray<FProjectedShadowInfo*,SceneRenderingAllocator>& PreShadows)
 {
+	// too high on hit count to leave on
+	// SCOPE_CYCLE_COUNTER(STAT_SetupInteractionShadows);
+
 	FPrimitiveSceneInfo* PrimitiveSceneInfo = Interaction->GetPrimitiveSceneInfo();
 	FLightSceneProxy* LightProxy = Interaction->GetLight()->Proxy;
 	extern bool GUseTranslucencyShadowDepths;
@@ -1433,6 +1437,7 @@ void FDeferredShadingSceneRenderer::CreatePerObjectProjectedShadow(
  */
 void FDeferredShadingSceneRenderer::CreateWholeSceneProjectedShadow(FLightSceneInfo* LightSceneInfo)
 {
+	SCOPE_CYCLE_COUNTER(STAT_CreateWholeSceneProjectedShadow);
 	FVisibleLightInfo& VisibleLightInfo = VisibleLightInfos[LightSceneInfo->Id];
 
 	// Try to create a whole-scene projected shadow initializer for the light.
@@ -1582,6 +1587,7 @@ void FDeferredShadingSceneRenderer::CreateWholeSceneProjectedShadow(FLightSceneI
 
 void FSceneRenderer::InitProjectedShadowVisibility(FRHICommandListImmediate& RHICmdList)
 {
+	SCOPE_CYCLE_COUNTER(STAT_InitProjectedShadowVisibility);
 	// Initialize the views' ProjectedShadowVisibilityMaps and remove shadows without subjects.
 	for(TSparseArray<FLightSceneInfoCompact>::TConstIterator LightIt(Scene->Lights);LightIt;++LightIt)
 	{
@@ -1931,6 +1937,8 @@ void FSceneRenderer::GatherShadowPrimitives(
 
 void FSceneRenderer::AddViewDependentWholeSceneShadowsForView(TArray<FProjectedShadowInfo*, SceneRenderingAllocator>& ShadowInfos, FVisibleLightInfo& VisibleLightInfo, FLightSceneInfo& LightSceneInfo)
 {
+	SCOPE_CYCLE_COUNTER(STAT_AddViewDependentWholeSceneShadowsForView);
+
 	// Allow each view to create a whole scene view dependent shadow
 	for (int32 ViewIndex = 0; ViewIndex < Views.Num(); ViewIndex++)
 	{
@@ -2117,6 +2125,9 @@ void FDeferredShadingSceneRenderer::InitDynamicShadows(FRHICommandListImmediate&
 		{
 			const FLightSceneInfoCompact& LightSceneInfoCompact = *LightIt;
 			FLightSceneInfo* LightSceneInfo = LightSceneInfoCompact.LightSceneInfo;
+
+			FScopeCycleCounter Context(LightSceneInfo->Proxy->GetStatId());
+
 			FVisibleLightInfo& VisibleLightInfo = VisibleLightInfos[LightSceneInfo->Id];
 
 			// Only consider lights that may have shadows.
