UE4 - turn on parallel view relevance for non-editor builds (automatically disabled on platforms that don't want threading for performance)

diff --git a/Engine/Source/Runtime/Renderer/Private/SceneVisibility.cpp b/Engine/Source/Runtime/Renderer/Private/SceneVisibility.cpp
--- a/Engine/Source/Runtime/Renderer/Private/SceneVisibility.cpp
+++ b/Engine/Source/Runtime/Renderer/Private/SceneVisibility.cpp
@@ -1159,7 +1159,11 @@ public:
 
 static TAutoConsoleVariable<int32> CVarParallelInitViews(
 	TEXT("r.ParallelInitViews"),
+#if WITH_EDITOR
 	0,  
+#else
+	1,  
+#endif
 	TEXT("Toggles parallel init views.\n"),
 	ECVF_RenderThreadSafe
 	);
