Ray traced distance field soft shadows
* Support for point, spot and directional lights
* LightSourceRadius / LightSourceAngle controls how large the penumbras are, with proper sharp contacts
* Directional lights use fast scatter culling path, the others use tiled gather culling
* Local lights only operate on the screen bounding rect

diff --git a/Engine/Shaders/DistanceFieldAOShared.usf b/Engine/Shaders/DistanceFieldAOShared.usf
new file mode 100644
--- /dev/null
+++ b/Engine/Shaders/DistanceFieldAOShared.usf
@@ -0,0 +1,140 @@
+// Copyright 1998-2014 Epic Games, Inc. All Rights Reserved.
+
+/*=============================================================================
+	DistanceFieldAOShared.usf
+=============================================================================*/
+
+#ifndef SORTTHREADGROUP_SIZEX
+#define SORTTHREADGROUP_SIZEX 1
+#endif
+
+#ifndef SORTTHREADGROUP_SIZEY
+#define SORTTHREADGROUP_SIZEY 1
+#endif
+
+#ifndef NUM_CONE_DIRECTIONS
+#define NUM_CONE_DIRECTIONS 1
+#endif
+
+#ifndef CONE_TRACE_PHASE
+#define CONE_TRACE_PHASE -1
+#endif
+
+// Must match C++
+#define NUM_CONE_STEPS 10
+
+float AOMaxDistance;
+float AOStepScale;
+float AOStepExponentScale;
+float AOMaxViewDistance;
+
+#define NUM_CULLED_OBJECT_LISTS 3
+
+void GetPhaseParameters(uint Phase, out uint StartIndex, out uint EndIndex)
+{
+	if (Phase == 0)
+	{
+		StartIndex = 0;
+		EndIndex = 4;
+	}
+	else if (Phase == 1)
+	{
+		StartIndex = 5;
+		EndIndex = 7;
+	}
+	else
+	{
+		StartIndex = 8;
+		EndIndex = 9;
+	}
+}
+
+float GetStepOffset(float StepIndex)
+{
+	return AOStepScale * exp2(AOStepExponentScale * StepIndex);
+}
+
+Buffer<uint> TileHeadDataUnpacked;
+Buffer<uint4> TileHeadData;
+Buffer<uint> TileArrayData;
+uint2 TileListGroupSize;
+
+
+#define UNPACKED_TILE_HEAD 1
+#define COHERENT_OBJECT_LIST_WRITES 1
+
+uint4 GetTileHead(uint2 TileCoordinate)
+{
+	uint TileIndex = TileCoordinate.y * TileListGroupSize.x + TileCoordinate.x;
+#if UNPACKED_TILE_HEAD
+	return uint4(
+		TileHeadDataUnpacked[TileIndex * 4 + 0], 
+		min(TileHeadDataUnpacked[TileIndex * 4 + 1], (uint)MAX_OBJECTS_PER_TILE),
+		min(TileHeadDataUnpacked[TileIndex * 4 + 2], (uint)MAX_OBJECTS_PER_TILE),
+		min(TileHeadDataUnpacked[TileIndex * 4 + 3], (uint)MAX_OBJECTS_PER_TILE));
+#else
+	return TileHeadData[TileIndex];
+#endif
+}
+
+Texture2D DistanceFieldNormalTexture;
+SamplerState DistanceFieldNormalSampler;
+
+int CurrentLevelDownsampleFactor;
+float2 AOBufferSize;
+
+int DownsampleFactorToBaseLevel;
+float2 BaseLevelTexelSize;
+
+float3 EncodeNormalForAO(float3 InNormal)
+{
+	return InNormal.xyz;
+}
+
+float3 DecodeNormalForAO(float3 InNormal)
+{
+	return InNormal.xyz;
+}
+
+bool SphereIntersectCone(float4 SphereCenterAndRadius, float3 ConeVertex, float3 ConeAxis, float ConeAngleCos, float ConeAngleSin)
+{
+	float3 U = ConeVertex - (SphereCenterAndRadius.w / ConeAngleSin) * ConeAxis;
+	float3 D = SphereCenterAndRadius.xyz - U;
+	float DSizeSq = dot(D, D);
+	float E = dot(ConeAxis, D);
+
+	if (E > 0 && E * E >= DSizeSq * ConeAngleCos * ConeAngleCos)
+	{
+		D = SphereCenterAndRadius.xyz - ConeVertex;
+		DSizeSq = dot(D, D);
+		E = -dot(ConeAxis, D);
+
+		if (E > 0 && E * E >= DSizeSq * ConeAngleSin * ConeAngleSin)
+		{
+			return DSizeSq <= SphereCenterAndRadius.w * SphereCenterAndRadius.w;
+		}
+		else
+		{
+			return true;
+		}
+	}
+
+	return false;
+}
+
+bool SphereIntersectConeWithDepthRanges(float4 SphereCenterAndRadius, float3 ConeVertex, float3 ConeAxis, float ConeAngleCos, float ConeAngleSin, float4 ConeAxisDepthRanges)
+{
+	if (SphereIntersectCone(SphereCenterAndRadius, ConeVertex, ConeAxis, ConeAngleCos, ConeAngleSin))
+	{
+		float ConeAxisDistance = dot(SphereCenterAndRadius.xyz - ConeVertex, ConeAxis);
+		float2 ConeAxisDistanceMinMax = float2(ConeAxisDistance + SphereCenterAndRadius.w, ConeAxisDistance - SphereCenterAndRadius.w);
+
+		if (ConeAxisDistanceMinMax.x > ConeAxisDepthRanges.x && ConeAxisDistanceMinMax.y < ConeAxisDepthRanges.y
+			|| ConeAxisDistanceMinMax.x > ConeAxisDepthRanges.z && ConeAxisDistanceMinMax.y < ConeAxisDepthRanges.w)
+		{
+			return true;
+		}
+	}
+
+	return false;
+}
\ No newline at end of file
diff --git a/Engine/Shaders/DistanceFieldLightingShared.usf b/Engine/Shaders/DistanceFieldLightingShared.usf
--- a/Engine/Shaders/DistanceFieldLightingShared.usf
+++ b/Engine/Shaders/DistanceFieldLightingShared.usf
@@ -14,64 +14,14 @@
 
 #define THREADGROUP_TOTALSIZE (THREADGROUP_SIZEX * THREADGROUP_SIZEY)
 
-#ifndef SORTTHREADGROUP_SIZEX
-#define SORTTHREADGROUP_SIZEX 1
-#endif
-
-#ifndef SORTTHREADGROUP_SIZEY
-#define SORTTHREADGROUP_SIZEY 1
-#endif
-
-#ifndef NUM_CONE_DIRECTIONS
-#define NUM_CONE_DIRECTIONS 1
+#ifndef MAX_OBJECTS_PER_TILE
+#define MAX_OBJECTS_PER_TILE 1
 #endif
 
 #ifndef DOWNSAMPLE_FACTOR
 #define DOWNSAMPLE_FACTOR 1
 #endif
 
-#ifndef CONE_TRACE_PHASE
-#define CONE_TRACE_PHASE -1
-#endif
-
-#ifndef MAX_OBJECTS_PER_TILE
-#define MAX_OBJECTS_PER_TILE 1
-#endif
-
-// Must match C++
-#define NUM_CONE_STEPS 10
-
-float AOMaxDistance;
-float AOStepScale;
-float AOStepExponentScale;
-float AOMaxViewDistance;
-
-#define NUM_CULLED_OBJECT_LISTS 3
-
-void GetPhaseParameters(uint Phase, out uint StartIndex, out uint EndIndex)
-{
-	if (Phase == 0)
-	{
-		StartIndex = 0;
-		EndIndex = 4;
-	}
-	else if (Phase == 1)
-	{
-		StartIndex = 5;
-		EndIndex = 7;
-	}
-	else
-	{
-		StartIndex = 8;
-		EndIndex = 9;
-	}
-}
-
-float GetStepOffset(float StepIndex)
-{
-	return AOStepScale * exp2(AOStepExponentScale * StepIndex);
-}
-
 float3 DistanceFieldVolumePositionToUV(float3 VolumePosition, float3 UVScale, float3 UVAdd)
 {
 	float3 VolumeUV = VolumePosition * UVScale + UVAdd;
@@ -82,57 +32,17 @@ Texture3D DistanceFieldTexture;
 SamplerState DistanceFieldSampler;
 float3 DistanceFieldAtlasTexelSize;
 
-Buffer<uint> TileHeadDataUnpacked;
-Buffer<uint4> TileHeadData;
-Buffer<uint> TileArrayData;
-uint2 TileListGroupSize;
-
-
-#define UNPACKED_TILE_HEAD 1
-#define COHERENT_OBJECT_LIST_WRITES 1
-
-uint4 GetTileHead(uint2 TileCoordinate)
-{
-	uint TileIndex = TileCoordinate.y * TileListGroupSize.x + TileCoordinate.x;
-#if UNPACKED_TILE_HEAD
-	return uint4(
-		TileHeadDataUnpacked[TileIndex * 4 + 0], 
-		min(TileHeadDataUnpacked[TileIndex * 4 + 1], (uint)MAX_OBJECTS_PER_TILE),
-		min(TileHeadDataUnpacked[TileIndex * 4 + 2], (uint)MAX_OBJECTS_PER_TILE),
-		min(TileHeadDataUnpacked[TileIndex * 4 + 3], (uint)MAX_OBJECTS_PER_TILE));
-#else
-	return TileHeadData[TileIndex];
-#endif
-}
-
-Texture2D DistanceFieldNormalTexture;
-SamplerState DistanceFieldNormalSampler;
-
-int CurrentLevelDownsampleFactor;
-float2 AOBufferSize;
-
-int DownsampleFactorToBaseLevel;
-float2 BaseLevelTexelSize;
-
-float3 EncodeNormalForAO(float3 InNormal)
-{
-	return InNormal.xyz;
-}
-
-float3 DecodeNormalForAO(float3 InNormal)
-{
-	return InNormal.xyz;
-}
-
 uint NumObjects;
 
 // In float4's.  Must match equivalent C++ variables.
 #define OBJECT_DATA_STRIDE 8
 #define OBJECT_DATA2_STRIDE 4
+#define OBJECT_BOX_BOUNDS_STRIDE 5
 
 Buffer<float4> ObjectBounds;
 Buffer<float4> ObjectData;
 Buffer<float4> ObjectData2;
+Buffer<float4> ObjectBoxBounds;
 
 float4 LoadObjectPositionAndRadius(uint ObjectIndex)
 {
@@ -179,45 +89,15 @@ float4x4 LoadObjectVolumeToWorld(uint ObjectIndex)
 	return float4x4(M0, M1, M2, M3);
 }
 
-bool SphereIntersectCone(float4 SphereCenterAndRadius, float3 ConeVertex, float3 ConeAxis, float ConeAngleCos, float ConeAngleSin)
+void LoadObjectViewSpaceBox(uint ObjectIndex, out float3 ObjectViewSpaceMin, out float3 ObjectViewSpaceMax)
 {
-	float3 U = ConeVertex - (SphereCenterAndRadius.w / ConeAngleSin) * ConeAxis;
-	float3 D = SphereCenterAndRadius.xyz - U;
-	float DSizeSq = dot(D, D);
-	float E = dot(ConeAxis, D);
-
-	if (E > 0 && E * E >= DSizeSq * ConeAngleCos * ConeAngleCos)
-	{
-		D = SphereCenterAndRadius.xyz - ConeVertex;
-		DSizeSq = dot(D, D);
-		E = -dot(ConeAxis, D);
-
-		if (E > 0 && E * E >= DSizeSq * ConeAngleSin * ConeAngleSin)
-		{
-			return DSizeSq <= SphereCenterAndRadius.w * SphereCenterAndRadius.w;
-		}
-		else
-		{
-			return true;
-		}
-	}
-
-	return false;
+	ObjectViewSpaceMin = ObjectBoxBounds.Load(ObjectIndex * OBJECT_BOX_BOUNDS_STRIDE + 0).xyz;
+	ObjectViewSpaceMax = ObjectBoxBounds.Load(ObjectIndex * OBJECT_BOX_BOUNDS_STRIDE + 1).xyz;
 }
 
-bool SphereIntersectConeWithDepthRanges(float4 SphereCenterAndRadius, float3 ConeVertex, float3 ConeAxis, float ConeAngleCos, float ConeAngleSin, float4 ConeAxisDepthRanges)
+void LoadObjectAxes(uint ObjectIndex, out float3 ObjectAxisX, out float3 ObjectAxisY, out float3 ObjectAxisZ)
 {
-	if (SphereIntersectCone(SphereCenterAndRadius, ConeVertex, ConeAxis, ConeAngleCos, ConeAngleSin))
-	{
-		float ConeAxisDistance = dot(SphereCenterAndRadius.xyz - ConeVertex, ConeAxis);
-		float2 ConeAxisDistanceMinMax = float2(ConeAxisDistance + SphereCenterAndRadius.w, ConeAxisDistance - SphereCenterAndRadius.w);
-
-		if (ConeAxisDistanceMinMax.x > ConeAxisDepthRanges.x && ConeAxisDistanceMinMax.y < ConeAxisDepthRanges.y
-			|| ConeAxisDistanceMinMax.x > ConeAxisDepthRanges.z && ConeAxisDistanceMinMax.y < ConeAxisDepthRanges.w)
-		{
-			return true;
-		}
-	}
-
-	return false;
+	ObjectAxisX = ObjectBoxBounds.Load(ObjectIndex * OBJECT_BOX_BOUNDS_STRIDE + 2).xyz;
+	ObjectAxisY = ObjectBoxBounds.Load(ObjectIndex * OBJECT_BOX_BOUNDS_STRIDE + 3).xyz;
+	ObjectAxisZ = ObjectBoxBounds.Load(ObjectIndex * OBJECT_BOX_BOUNDS_STRIDE + 4).xyz;
 }
\ No newline at end of file
diff --git a/Engine/Shaders/DistanceFieldShadowing.usf b/Engine/Shaders/DistanceFieldShadowing.usf
--- a/Engine/Shaders/DistanceFieldShadowing.usf
+++ b/Engine/Shaders/DistanceFieldShadowing.usf
@@ -8,17 +8,176 @@
 #include "DeferredShadingCommon.usf"
 #include "DistanceFieldLightingShared.usf"
 
-RWTexture2D<float2> RWShadowFactors;
-
-bool RayHitSphere( float3 RayOrigin, float3 UnitRayDirection, float3 SphereCenter, float SphereRadius )
+bool RayHitSphere(float3 RayOrigin, float3 UnitRayDirection, float RayLength, float3 SphereCenter, float SphereRadius)
 {
-	float3 ClosestPointOnRay = max( 0, dot( SphereCenter - RayOrigin, UnitRayDirection ) ) * UnitRayDirection;
+	float DistanceAlongRay = dot(SphereCenter - RayOrigin, UnitRayDirection);
+	float3 ClosestPointOnRay = DistanceAlongRay * UnitRayDirection;
 	float3 CenterToRay = RayOrigin + ClosestPointOnRay - SphereCenter;
-	return dot( CenterToRay, CenterToRay ) <= Square( SphereRadius );
+	return dot(CenterToRay, CenterToRay) <= Square(SphereRadius) && DistanceAlongRay > -SphereRadius && DistanceAlongRay - SphereRadius < RayLength;
 }
 
+RWBuffer<uint> RWTileHeadDataUnpacked;
+RWBuffer<uint> RWTileArrayData;
+
 float2 NumGroups;
+
+[numthreads(THREADGROUP_SIZEX, THREADGROUP_SIZEY, 1)]
+void ClearTilesCS(
+	uint3 GroupId : SV_GroupID,
+	uint3 DispatchThreadId : SV_DispatchThreadID,
+    uint3 GroupThreadId : SV_GroupThreadID) 
+{
+    uint TileIndex = DispatchThreadId.y * NumGroups.x + DispatchThreadId.x;
+	
+	RWTileHeadDataUnpacked[TileIndex * 2 + 0] = TileIndex;
+	RWTileHeadDataUnpacked[TileIndex * 2 + 1] = 0;
+}
+
+struct FShadowObjectCullVertexOutput
+{
+	float4 Position : SV_POSITION;
+	nointerpolation float4 PositionAndRadius : TEXCOORD0;
+	nointerpolation uint ObjectIndex : TEXCOORD1;
+};
+
+float ConservativeRadiusScale;
+float MinRadius;
+float4x4 WorldToShadow;
+
+/** Used when culling objects into screenspace tile lists */
+void ShadowObjectCullVS(
+	float4 InPosition : ATTRIBUTE0,
+	uint ObjectIndex : SV_InstanceID,
+	out FShadowObjectCullVertexOutput Output
+	)
+{
+	float4 ObjectPositionAndRadius = LoadObjectPositionAndRadius(ObjectIndex);
+	// ConservativeRadiusScale pushes the sphere vertices out so the triangles between them lie completely outside the sphere
+	// MinRadius is for conservative rasterization
+	float EffectiveRadius = (ObjectPositionAndRadius.w + MinRadius) * ConservativeRadiusScale;
+	float3 WorldPosition = InPosition.xyz * EffectiveRadius + ObjectPositionAndRadius.xyz;
+	Output.Position = mul(float4(WorldPosition, 1), WorldToShadow);
+	Output.PositionAndRadius = ObjectPositionAndRadius;
+	Output.ObjectIndex = ObjectIndex;
+} 
+
+/** Intersects a single object with the tile and adds to the intersection list if needed. */
+void ShadowObjectCullPS(
+	FShadowObjectCullVertexOutput Input, 
+	in float4 SVPos : SV_POSITION,
+	out float4 OutColor : COLOR0)
+{
+	OutColor = 0;
+	
+	uint2 TilePosition = (uint2)SVPos.xy;
+	uint TileIndex = TilePosition.y * NumGroups.x + TilePosition.x;
+
+#define OBJECT_OBB_INTERSECTION 1
+#if OBJECT_OBB_INTERSECTION
+	
+	float3 ShadowTileMin;
+	float3 ShadowTileMax;
+
+	float2 TilePositionForCulling = float2(TilePosition.x, NumGroups.y - TilePosition.y);
+	//@todo - why is this expand needed
+	float TileExpand = 1;
+	ShadowTileMin.xy = (TilePositionForCulling - TileExpand) / (float2)NumGroups * 2 - 1;
+	ShadowTileMax.xy = (TilePositionForCulling + 1 + TileExpand) / (float2)NumGroups * 2 - 1;
+	ShadowTileMin.z = 0;
+	ShadowTileMax.z = 1;
+
+	float3 ObjectViewSpaceMin;
+	float3 ObjectViewSpaceMax;
+	LoadObjectViewSpaceBox(Input.ObjectIndex, ObjectViewSpaceMin, ObjectViewSpaceMax);
+
+	BRANCH
+	// Separating axis test on the AABB
+	if (all(ObjectViewSpaceMax > ShadowTileMin) && all(ObjectViewSpaceMin < ShadowTileMax))
+	{
+		float3 ObjectCenter = .5f * (ObjectViewSpaceMin + ObjectViewSpaceMax);
+		float3 MinProjections = 500000;
+		float3 MaxProjections = -500000;
+
+		{
+			float3 Corners[8];
+			Corners[0] = float3(ShadowTileMin.x, ShadowTileMin.y, ShadowTileMin.z);
+			Corners[1] = float3(ShadowTileMax.x, ShadowTileMin.y, ShadowTileMin.z);
+			Corners[2] = float3(ShadowTileMin.x, ShadowTileMax.y, ShadowTileMin.z);
+			Corners[3] = float3(ShadowTileMax.x, ShadowTileMax.y, ShadowTileMin.z);
+			Corners[4] = float3(ShadowTileMin.x, ShadowTileMin.y, ShadowTileMax.z);
+			Corners[5] = float3(ShadowTileMax.x, ShadowTileMin.y, ShadowTileMax.z);
+			Corners[6] = float3(ShadowTileMin.x, ShadowTileMax.y, ShadowTileMax.z);
+			Corners[7] = float3(ShadowTileMax.x, ShadowTileMax.y, ShadowTileMax.z);
+
+			float3 ObjectAxisX;
+			float3 ObjectAxisY; 
+			float3 ObjectAxisZ;
+			LoadObjectAxes(Input.ObjectIndex, ObjectAxisX, ObjectAxisY, ObjectAxisZ);
+
+			UNROLL
+			for (int i = 0; i < 8; i++)
+			{
+				float3 CenterToVertex = Corners[i] - ObjectCenter;
+				float3 Projections = float3(dot(CenterToVertex, ObjectAxisX), dot(CenterToVertex, ObjectAxisY), dot(CenterToVertex, ObjectAxisZ));
+				MinProjections = min(MinProjections, Projections);
+				MaxProjections = max(MaxProjections, Projections);
+			}
+		}
+
+		BRANCH
+		// Separating axis test on the OBB
+		if (all(MinProjections < 1) && all(MaxProjections > -1))
+		{
+			uint ArrayIndex;
+			InterlockedAdd(RWTileHeadDataUnpacked[TileIndex * 2 + 1], 1U, ArrayIndex);
+
+			if (ArrayIndex < MAX_OBJECTS_PER_TILE)
+			{
+				uint DataIndex = (ArrayIndex * (uint)(NumGroups.x * NumGroups.y + .5f) + TileIndex);
+				RWTileArrayData[DataIndex] = Input.ObjectIndex;
+			}
+		}
+	}
+	
+#else
+	{
+		uint ArrayIndex;
+		InterlockedAdd(RWTileHeadDataUnpacked[TileIndex * 2 + 1], 1U, ArrayIndex);
+
+		if (ArrayIndex < MAX_OBJECTS_PER_TILE)
+		{
+			uint DataIndex = (ArrayIndex * (uint)(NumGroups.x * NumGroups.y + .5f) + TileIndex);
+			RWTileArrayData[DataIndex] = Input.ObjectIndex;
+		}
+	}
+#endif
+}
+
+RWTexture2D<float2> RWShadowFactors;
+
+/** From point being shaded toward light, for directional lights. */
 float3 LightDirection;
+float4 LightPositionAndInvRadius;
+float LightSourceRadius;
+float2 TanLightAngleAndNormalThreshold;
+int4 ScissorRectMinAndSize;
+
+#if SCATTER_TILE_CULLING
+
+Buffer<uint> TileHeadDataUnpacked;
+Buffer<uint> TileArrayData;
+uint2 TileListGroupSize;
+
+uint2 GetTileHead(uint2 TileCoordinate)
+{
+	uint TileIndex = TileCoordinate.y * TileListGroupSize.x + TileCoordinate.x;
+
+	return uint2(
+		TileHeadDataUnpacked[TileIndex * 2 + 0], 
+		min(TileHeadDataUnpacked[TileIndex * 2 + 1], (uint)MAX_OBJECTS_PER_TILE)); 
+}
+
+#else // SCATTER_TILE_CULLING
 
 /** Min and Max depth for this tile. */
 groupshared uint IntegerTileMinZ;
@@ -35,6 +194,8 @@ groupshared uint TileNumObjects1;
 #define MAX_INTERSECTING_OBJECTS 256
 groupshared uint IntersectingObjectIndices[MAX_INTERSECTING_OBJECTS * 2];
 
+#endif
+
 [numthreads(THREADGROUP_SIZEX, THREADGROUP_SIZEY, 1)]
 void DistanceFieldShadowingCS(
 	uint3 GroupId : SV_GroupID,
@@ -43,14 +204,40 @@ void DistanceFieldShadowingCS(
 {
 	uint ThreadIndex = GroupThreadId.y * THREADGROUP_SIZEX + GroupThreadId.x;
 
-	float2 ScreenUV = float2((DispatchThreadId.xy * DOWNSAMPLE_FACTOR + View.ViewRectMin.xy + .5f) * View.ViewSizeAndSceneTexelSize.zw);
+	float2 ScreenUV = float2((DispatchThreadId.xy * DOWNSAMPLE_FACTOR + ScissorRectMinAndSize.xy + .5f) * View.ViewSizeAndSceneTexelSize.zw);
 	float2 ScreenPosition = (ScreenUV.xy - View.ScreenPositionScaleBias.wz) / View.ScreenPositionScaleBias.xy;
 
 	float SceneDepth = CalcSceneDepth(ScreenUV);
+	float4 HomogeneousWorldPosition = mul(float4(ScreenPosition * SceneDepth, SceneDepth, 1), View.ScreenToWorld);
+	float3 OpaqueWorldPosition = HomogeneousWorldPosition.xyz / HomogeneousWorldPosition.w;
 
+	// Distance for directional lights to trace
+	float TraceDistance = 10000;
+
+	//@todo - take advantage of CSM covering near pixels
 	//@todo - handle invalid tile depth ranges due to MinDepth
 	float MinDepth = 0;
 
+	uint NumIntersectingObjects = NumObjects;
+	bool bTileShouldComputeShadowing = SceneDepth > MinDepth;
+
+#define USE_CULLING 1
+#if USE_CULLING
+
+#if SCATTER_TILE_CULLING
+
+	// Transform into shadow space
+	float4 HomogeneousShadowPosition = mul(float4(OpaqueWorldPosition, 1), WorldToShadow);
+	float2 NormalizedShadowPosition = HomogeneousShadowPosition.xy * .5f + .5f;
+	NormalizedShadowPosition.y = 1 - NormalizedShadowPosition.y;
+	// Quantize the shadow position to get our tile position
+	uint2 TilePosition = (uint2)(NormalizedShadowPosition * TileListGroupSize);
+	// Fetch the tile head information
+	uint2 TileHead = GetTileHead(TilePosition);
+	NumIntersectingObjects = TileHead.y;
+
+#else
+
 	// Initialize per-tile variables
     if (ThreadIndex == 0) 
 	{
@@ -106,9 +293,12 @@ void DistanceFieldShadowingCS(
 	float ExpandRadius = 0;
 
 	float2 TanViewFOV = float2(1 / View.ViewToClip[0][0], 1 / View.ViewToClip[1][1]);
-	float2 TileSize = TanViewFOV * 2 / NumGroups;
-	float2 UnitPlaneTileMin = (GroupId.xy * TileSize - TanViewFOV) * float2(1, -1);
-	float2 UnitPlaneTileMax = ((GroupId.xy + 1) * TileSize - TanViewFOV) * float2(1, -1);
+	// tan(FOV) = HalfUnitPlaneWidth / 1, so TanViewFOV * 2 is the size of the whole unit view plane
+	// We are operating on a subset of that defined by ScissorRectMinAndSize
+	float2 TileSize = TanViewFOV * 2 * ScissorRectMinAndSize.zw / ((float2)View.ViewSizeAndSceneTexelSize.xy * NumGroups);
+	float2 UnitPlaneMin = -TanViewFOV + TanViewFOV * 2 * ScissorRectMinAndSize.xy / (float2)View.ViewSizeAndSceneTexelSize.xy;
+	float2 UnitPlaneTileMin = (GroupId.xy * TileSize + UnitPlaneMin) * float2(1, -1);
+	float2 UnitPlaneTileMax = ((GroupId.xy + 1) * TileSize + UnitPlaneMin) * float2(1, -1);
 
 	ViewTileMin.xy = min(MinTileZ * UnitPlaneTileMin, MaxTileZ2 * UnitPlaneTileMin) - ExpandRadius;
 	ViewTileMax.xy = max(MinTileZ * UnitPlaneTileMax, MaxTileZ2 * UnitPlaneTileMax) + ExpandRadius;
@@ -127,75 +317,124 @@ void DistanceFieldShadowingCS(
 	float3 WorldGroup1Center = mul(float4(ViewGroup1Center, 1), View.ViewToTranslatedWorld).xyz - View.PreViewTranslation;
 	float Group1BoundingRadius = length(ViewGroup1Center - ViewTileMax2);
 
-	// Compute per-tile lists of affecting objects through bounds culling
-	// Each thread now operates on a sample instead of a pixel
-	LOOP
-	for (uint ObjectIndex = ThreadIndex; ObjectIndex < NumObjects; ObjectIndex += THREADGROUP_TOTALSIZE)
-	{
-		float4 SphereCenterAndRadius = LoadObjectPositionAndRadius(ObjectIndex);
+	bool bCullTile = true;
 
-		BRANCH
-		if (RayHitSphere(WorldGroup0Center, -LightDirection, SphereCenterAndRadius.xyz, SphereCenterAndRadius.w + Group0BoundingRadius))
-		{
-			uint ListIndex;
-			InterlockedAdd(TileNumObjects0, 1U, ListIndex);
-			// Don't overwrite on overflow
-			ListIndex = min(ListIndex, MAX_INTERSECTING_OBJECTS - 1);
-			IntersectingObjectIndices[MAX_INTERSECTING_OBJECTS * 0 + ListIndex] = ObjectIndex; 
-		}
+#if POINT_LIGHT
+	float3 LightVector0 = LightPositionAndInvRadius.xyz - WorldGroup0Center;
+	float LightVector0Length = length(LightVector0);
+	float3 LightVector1 = LightPositionAndInvRadius.xyz - WorldGroup1Center;
+	float LightVector1Length = length(LightVector1);
+	float3 LightDirection0 = LightVector0 / LightVector0Length;
+	float3 LightDirection1 = LightVector1 / LightVector1Length;;
+	float RayLength0 = LightVector0Length;
+	float RayLength1 = LightVector1Length;
 
-		BRANCH
-		if (RayHitSphere(WorldGroup1Center, -LightDirection, SphereCenterAndRadius.xyz, SphereCenterAndRadius.w + Group1BoundingRadius))
+	// Don't operate on tiles completely outside of the light's influence
+	bCullTile = bTileShouldComputeShadowing = LightVector0Length < 1.0f / LightPositionAndInvRadius.w + Group0BoundingRadius
+		|| LightVector1Length < 1.0f / LightPositionAndInvRadius.w + Group1BoundingRadius;
+
+#else
+	float3 LightDirection0 = LightDirection;
+	float3 LightDirection1 = LightDirection;
+	float RayLength0 = TraceDistance;
+	float RayLength1 = TraceDistance;
+#endif
+
+	BRANCH
+	if (bCullTile)
+	{
+		// Compute per-tile lists of affecting objects through bounds culling
+		// Each thread now operates on a sample instead of a pixel
+		LOOP
+		for (uint ObjectIndex = ThreadIndex; ObjectIndex < NumObjects; ObjectIndex += THREADGROUP_TOTALSIZE)
 		{
-			uint ListIndex;
-			InterlockedAdd(TileNumObjects1, 1U, ListIndex);
-			// Don't write out of bounds on overflow
-			ListIndex = min(ListIndex, MAX_INTERSECTING_OBJECTS - 1);
-			IntersectingObjectIndices[MAX_INTERSECTING_OBJECTS * 1 + ListIndex] = ObjectIndex; 
+			float4 SphereCenterAndRadius = LoadObjectPositionAndRadius(ObjectIndex);
+
+			BRANCH
+			if (RayHitSphere(WorldGroup0Center, LightDirection0, RayLength0, SphereCenterAndRadius.xyz, SphereCenterAndRadius.w + Group0BoundingRadius))
+			{
+				uint ListIndex;
+				InterlockedAdd(TileNumObjects0, 1U, ListIndex);
+				// Don't overwrite on overflow
+				ListIndex = min(ListIndex, MAX_INTERSECTING_OBJECTS - 1);
+				IntersectingObjectIndices[MAX_INTERSECTING_OBJECTS * 0 + ListIndex] = ObjectIndex; 
+			}
+
+			BRANCH
+			if (RayHitSphere(WorldGroup1Center, LightDirection1, RayLength1, SphereCenterAndRadius.xyz, SphereCenterAndRadius.w + Group1BoundingRadius))
+			{
+				uint ListIndex;
+				InterlockedAdd(TileNumObjects1, 1U, ListIndex);
+				// Don't write out of bounds on overflow
+				ListIndex = min(ListIndex, MAX_INTERSECTING_OBJECTS - 1);
+				IntersectingObjectIndices[MAX_INTERSECTING_OBJECTS * 1 + ListIndex] = ObjectIndex; 
+			}
 		}
 	}
 
 	GroupMemoryBarrierWithGroupSync();
 
 	uint GroupIndex = SceneDepth > MaxTileZ2 ? 1 : 0;
-	uint NumIntersectingObjects = GroupIndex == 0 ? TileNumObjects0 : TileNumObjects1;
-	
+	NumIntersectingObjects = GroupIndex == 0 ? TileNumObjects0 : TileNumObjects1;
+#endif
+#endif // USE_CULLING
+
 	float Result = 0;
 
 #define COMPUTE_SHADOWING 1
 #if COMPUTE_SHADOWING
 	BRANCH
-	if (SceneDepth > MinDepth)
+	if (bTileShouldComputeShadowing)
 	{
 		FGBufferData GBufferData = GetGBufferData(ScreenUV);
 		float3 WorldNormal = GBufferData.WorldNormal;
 
-		//@todo - expose
-		float SunAngle = 2 * PI / 180.0f;
-
+#if !POINT_LIGHT
 		BRANCH
-		if (dot(WorldNormal, -LightDirection) > cos(PI / 2 + SunAngle))
+		if (dot(WorldNormal, LightDirection) > TanLightAngleAndNormalThreshold.y)
+#endif
 		{
-			float4 HomogeneousWorldPosition = mul(float4(ScreenPosition * SceneDepth, SceneDepth, 1), View.ScreenToWorld);
-			float3 OpaqueWorldPosition = HomogeneousWorldPosition.xyz / HomogeneousWorldPosition.w;
-
-			float TraceDistance = 10000;
+			// World space offset along the start of the ray to avoid incorrect self-shadowing
 			float RayStartOffset = 2;
-			float3 WorldRayStart = OpaqueWorldPosition - LightDirection * RayStartOffset;
-			float3 WorldRayEnd = OpaqueWorldPosition - LightDirection * TraceDistance;
-			float3 WorldRayDirection = WorldRayEnd - WorldRayStart;
-			float3 UnitWorldRayDirection = normalize(WorldRayDirection);
-
-			float TanSunAngle = tan(SunAngle);
-			float MinRayTime = TraceDistance;
-			float ClosestDistanceToIntersection = 1000;
-			float TotalStepsTaken = 0;
+			// Keeps result from going all the way sharp
+			float MinSphereRadius = .4f;
+			// Maintain reasonable culling bounds
+			float MaxSphereRadius = 100;
+
+			#if POINT_LIGHT
+
+				float3 LightVector = LightPositionAndInvRadius.xyz - OpaqueWorldPosition;
+				float LightVectorLength = length(LightVector);
+				float3 WorldRayStart = OpaqueWorldPosition + LightVector / LightVectorLength * RayStartOffset;
+				float3 WorldRayEnd = LightPositionAndInvRadius.xyz;
+				float MaxRayTime = LightVectorLength;
+				float TanLightAngle = LightSourceRadius / LightVectorLength;
+
+			#else
+
+				float3 WorldRayStart = OpaqueWorldPosition + LightDirection * RayStartOffset;
+				float3 WorldRayEnd = OpaqueWorldPosition + LightDirection * TraceDistance;
+				float MaxRayTime = TraceDistance;
+				float TanLightAngle = TanLightAngleAndNormalThreshold.x;
+
+			#endif
+
+			float MinRayTime = MaxRayTime;
 			float MinConeVisibility = 1;
 
 			LOOP
-			for (uint ListObjectIndex = 0; ListObjectIndex < NumIntersectingObjects && MinRayTime >= TraceDistance; ListObjectIndex++)
+			for (uint ListObjectIndex = 0; ListObjectIndex < NumIntersectingObjects && MinRayTime >= MaxRayTime; ListObjectIndex++)
 			{
-				uint ObjectIndex = IntersectingObjectIndices[MAX_INTERSECTING_OBJECTS * GroupIndex + ListObjectIndex];
+				#if USE_CULLING
+					#if SCATTER_TILE_CULLING
+						uint ObjectIndex = TileArrayData.Load(ListObjectIndex * TileListGroupSize.x * TileListGroupSize.y + TileHead.x);
+					#else
+						uint ObjectIndex = IntersectingObjectIndices[MAX_INTERSECTING_OBJECTS * GroupIndex + ListObjectIndex];
+					#endif
+				#else
+					uint ObjectIndex = ListObjectIndex;
+				#endif
+
 				float3 LocalPositionExtent = LoadObjectLocalPositionExtent(ObjectIndex);
 				float4x4 WorldToVolume = LoadObjectWorldToVolume(ObjectIndex);
 				float4 UVScaleAndVolumeScale = LoadObjectUVScale(ObjectIndex);
@@ -206,9 +445,17 @@ void DistanceFieldShadowingCS(
 				float3 VolumeRayDirection = VolumeRayEnd - VolumeRayStart;
 				float VolumeRayLength = length(VolumeRayDirection);
 				VolumeRayDirection /= VolumeRayLength;
+				float VolumeMinSphereRadius = MinSphereRadius / UVScaleAndVolumeScale.w;
+				float VolumeMaxSphereRadius = MaxSphereRadius / UVScaleAndVolumeScale.w;
+
+				float4 SphereCenterAndRadius = LoadObjectPositionAndRadius(ObjectIndex);
+				float ObjectCenterDistanceAlongRay = max(dot(SphereCenterAndRadius.xyz - WorldRayStart, WorldRayEnd - WorldRayStart), 0);
+				// Expand the intersection box by the radius of the cone at the distance of the object along the cone
+				float LocalConeRadiusAtObject = min(TanLightAngle * ObjectCenterDistanceAlongRay / UVScaleAndVolumeScale.w, VolumeMaxSphereRadius);
 
-				float2 IntersectionTimes = LineBoxIntersect(VolumeRayStart, VolumeRayEnd, -LocalPositionExtent, LocalPositionExtent);
+				float2 IntersectionTimes = LineBoxIntersect(VolumeRayStart, VolumeRayEnd, -LocalPositionExtent - LocalConeRadiusAtObject, LocalPositionExtent + LocalConeRadiusAtObject);
 
+				BRANCH
 				if (IntersectionTimes.x < IntersectionTimes.y && IntersectionTimes.x < 1)
 				{
 					float SampleRayTime = IntersectionTimes.x * VolumeRayLength;
@@ -225,11 +472,12 @@ void DistanceFieldShadowingCS(
 					{
 						float3 SampleVolumePosition = VolumeRayStart + VolumeRayDirection * SampleRayTime;
 						float3 ClampedSamplePosition = clamp(SampleVolumePosition, -LocalPositionExtent, LocalPositionExtent);
+						float DistanceToClamped = length(ClampedSamplePosition - SampleVolumePosition);
 						float3 VolumeUV = DistanceFieldVolumePositionToUV(ClampedSamplePosition, UVScaleAndVolumeScale.xyz, UVAdd);
-						float DistanceField = Texture3DSampleLevel(DistanceFieldTexture, DistanceFieldSampler, VolumeUV, 0).x;
+						float DistanceField = Texture3DSampleLevel(DistanceFieldTexture, DistanceFieldSampler, VolumeUV, 0).x + DistanceToClamped;
 
 						MinDistance = min(MinDistance, DistanceField);
-						float SphereRadius = TanSunAngle * SampleRayTime;
+						float SphereRadius = clamp(TanLightAngle * SampleRayTime, VolumeMinSphereRadius, VolumeMaxSphereRadius);
 
 						MinConeVisibility = min(MinConeVisibility, saturate(DistanceField / SphereRadius));
 						IntersectionPosition = SampleVolumePosition;
@@ -245,16 +493,11 @@ void DistanceFieldShadowingCS(
 						}
 					}
 
-					if (MinDistance * UVScaleAndVolumeScale.w < 0 || StepIndex == MaxSteps)
+					if (MinDistance < 0 || StepIndex == MaxSteps)
 					{
+						MinConeVisibility = 0;
 						MinRayTime = min(MinRayTime, SampleRayTime * UVScaleAndVolumeScale.w);
 					}
-
-					if (SampleRayTime * UVScaleAndVolumeScale.w > 100)
-					{
-						ClosestDistanceToIntersection = min(ClosestDistanceToIntersection, MinDistance * UVScaleAndVolumeScale.w);
-					}
-					TotalStepsTaken += StepIndex;
 				}
 			}
 
@@ -263,13 +506,10 @@ void DistanceFieldShadowingCS(
 	}
 
 #else
-	Result = .0001f * NumIntersectingObjects;
+	//Result = .0001f * NumIntersectingObjects;
+	Result = NumIntersectingObjects / 30.0f;
 #endif
 
-	//Result = (GroupIndex == 0 ? Group0BoundingRadius : Group1BoundingRadius) / 1000.0f;
-	//Result = TotalStepsTaken / 1000.0f;
-	//Result = NumIntersectingObjects / 100.0f;
-
 	RWShadowFactors[DispatchThreadId.xy] = float2(Result, SceneDepth);
 }
 
@@ -281,8 +521,8 @@ void DistanceFieldShadowingUpsamplePS(
 	in float4 SVPos : SV_POSITION,
 	out float4 OutColor : SV_Target0)
 {
-	// Distance field AO was computed at 0,0 regardless of viewrect min
-	float2 DistanceFieldUVs = UVAndScreenPos.xy - View.ViewRectMin.xy * View.ViewSizeAndSceneTexelSize.zw;
+	// Distance field shadowing was computed at 0,0 regardless of viewrect min
+	float2 DistanceFieldUVs = UVAndScreenPos.xy - ScissorRectMinAndSize.xy * View.ViewSizeAndSceneTexelSize.zw;
 
 #define BILATERAL_UPSAMPLE 1
 #if BILATERAL_UPSAMPLE
diff --git a/Engine/Shaders/DistanceFieldSurfaceCacheLighting.usf b/Engine/Shaders/DistanceFieldSurfaceCacheLighting.usf
--- a/Engine/Shaders/DistanceFieldSurfaceCacheLighting.usf
+++ b/Engine/Shaders/DistanceFieldSurfaceCacheLighting.usf
@@ -7,6 +7,7 @@
 #include "Common.usf"
 #include "DeferredShadingCommon.usf"
 #include "DistanceFieldLightingShared.usf"
+#include "DistanceFieldAOShared.usf"
 
 struct FObjectCullVertexOutput
 {
diff --git a/Engine/Shaders/DistanceFieldSurfaceCacheLightingCompute.usf b/Engine/Shaders/DistanceFieldSurfaceCacheLightingCompute.usf
--- a/Engine/Shaders/DistanceFieldSurfaceCacheLightingCompute.usf
+++ b/Engine/Shaders/DistanceFieldSurfaceCacheLightingCompute.usf
@@ -8,6 +8,7 @@
 #include "DeferredShadingCommon.usf"
 #include "ReflectionEnvironmentShared.usf"
 #include "DistanceFieldLightingShared.usf"
+#include "DistanceFieldAOShared.usf"
 
 /** Min and Max depth for this tile. */
 groupshared uint IntegerTileMinZ;
diff --git a/Engine/Source/Runtime/Engine/Classes/Components/DirectionalLightComponent.h b/Engine/Source/Runtime/Engine/Classes/Components/DirectionalLightComponent.h
--- a/Engine/Source/Runtime/Engine/Classes/Components/DirectionalLightComponent.h
+++ b/Engine/Source/Runtime/Engine/Classes/Components/DirectionalLightComponent.h
@@ -94,19 +94,16 @@ class UDirectionalLightComponent : public ULightComponent
 	UPROPERTY(EditAnywhere, AdvancedDisplay, BlueprintReadOnly, Category=CascadedShadowMaps)
 	uint32 bUseInsetShadowsForMovableObjects : 1;
 
-	/** 
-	 * Whether to use ray traced distance field area shadows.  The project setting bGenerateMeshDistanceFields must be enabled for this to have effect.
-	 * Distance field shadows have a large per-pixel cost but a low per-object cost (and don't depend on triangle count) so they are effective for distant shadows from a dynamic sun.
-	 */
-	UPROPERTY(EditAnywhere, BlueprintReadOnly, Category=DistanceFieldShadows, meta=(DisplayName = "Distance Field Shadows"))
-	bool bUseDistanceFieldShadows;
-
 	/** 
 	 * Distance at which the ray traced shadow cascade should end.  Distance field shadows will cover the range between 'Dynamic Shadow Distance' this distance. 
 	 */
 	UPROPERTY(EditAnywhere, BlueprintReadOnly, Category=DistanceFieldShadows, meta=(UIMin = "0", UIMax = "100000"), DisplayName = "DistanceField Shadow Distance")
 	float DistanceFieldShadowDistance;
 
+	/** Light source angle in degrees. */
+	UPROPERTY(EditAnywhere, BlueprintReadOnly, Category=DistanceFieldShadows, meta=(UIMin = "0", UIMax = "5"))
+	float LightSourceAngle;
+
 	/** The Lightmass settings for this object. */
 	UPROPERTY(EditAnywhere, BlueprintReadOnly, Category=Light, meta=(ShowOnlyInnerProperties))
 	struct FLightmassDirectionalLightSettings LightmassSettings;
diff --git a/Engine/Source/Runtime/Engine/Classes/Components/LightComponent.h b/Engine/Source/Runtime/Engine/Classes/Components/LightComponent.h
--- a/Engine/Source/Runtime/Engine/Classes/Components/LightComponent.h
+++ b/Engine/Source/Runtime/Engine/Classes/Components/LightComponent.h
@@ -149,6 +149,15 @@ class ENGINE_API ULightComponent : public ULightComponentBase
 	UPROPERTY(EditAnywhere, BlueprintReadOnly, Category=LightShafts)
 	FColor BloomTint;
 
+	/** 
+	 * Whether to use ray traced distance field area shadows.  The project setting bGenerateMeshDistanceFields must be enabled for this to have effect.
+	 * Distance field shadows support area lights so they create soft shadows with sharp contacts.  
+	 * They have less aliasing artifacts than standard shadowmaps, but inherit all the limitations of distance field representations (only uniform scale, no deformation).
+	 * These shadows have a low per-object cost (and don't depend on triangle count) so they are effective for distant shadows from a dynamic sun.
+	 */
+	UPROPERTY(EditAnywhere, BlueprintReadOnly, Category=DistanceFieldShadows, meta=(DisplayName = "Use RayTraced DistanceField Shadows"))
+	bool bUseRayTracedDistanceFieldShadows;
+
 public:
 	/** Set intensity of the light */
 	UFUNCTION(BlueprintCallable, Category="Rendering|Components|Light")
diff --git a/Engine/Source/Runtime/Engine/Private/DirectionalLightComponent.cpp b/Engine/Source/Runtime/Engine/Private/DirectionalLightComponent.cpp
--- a/Engine/Source/Runtime/Engine/Private/DirectionalLightComponent.cpp
+++ b/Engine/Source/Runtime/Engine/Private/DirectionalLightComponent.cpp
@@ -89,6 +89,9 @@ public:
 	/** If greater than WholeSceneDynamicShadowRadius, a cascade will be created to support ray traced distance field shadows covering up to this distance. */
 	float DistanceFieldShadowDistance;
 
+	/** Light source angle in degrees. */
+	float LightSourceAngle;
+
 	/** Initialization constructor. */
 	FDirectionalLightSceneProxy(const UDirectionalLightComponent* Component):
 		FLightSceneProxy(Component),
@@ -101,7 +104,8 @@ public:
 		CascadeTransitionFraction(Component->CascadeTransitionFraction),
 		ShadowDistanceFadeoutFraction(Component->ShadowDistanceFadeoutFraction),
 		bUseInsetShadowsForMovableObjects(Component->bUseInsetShadowsForMovableObjects),
-		DistanceFieldShadowDistance(Component->bUseDistanceFieldShadows ? Component->DistanceFieldShadowDistance : 0)
+		DistanceFieldShadowDistance(Component->bUseRayTracedDistanceFieldShadows ? Component->DistanceFieldShadowDistance : 0),
+		LightSourceAngle(Component->LightSourceAngle)
 	{
 		LightShaftOverrideDirection.Normalize();
 
@@ -147,6 +151,11 @@ public:
 		LightMinRoughness = MinRoughness;
 	}
 
+	virtual float GetLightSourceAngle() const override
+	{
+		return LightSourceAngle;
+	}
+
 	virtual bool GetLightShaftOcclusionParameters(float& OutOcclusionMaskDarkness, float& OutOcclusionDepthRange) const override
 	{
 		OutOcclusionMaskDarkness = OcclusionMaskDarkness;
@@ -627,7 +636,8 @@ UDirectionalLightComponent::UDirectionalLightComponent(const class FPostConstruc
 	DynamicShadowDistanceMovableLight = 20000.0f;
 	DynamicShadowDistanceStationaryLight = 0.f;
 
-	DistanceFieldShadowDistance = 40000.0f;
+	DistanceFieldShadowDistance = 30000.0f;
+	LightSourceAngle = 1;
 
 	DynamicShadowCascades = 3;
 	CascadeDistributionExponent = 3.0f;
@@ -686,17 +696,11 @@ bool UDirectionalLightComponent::CanEditChange(const UProperty* InProperty) cons
 					|| DynamicShadowDistanceStationaryLight > 0 && Mobility == EComponentMobility::Stationary);
 		}
 
-		if (PropertyName == GET_MEMBER_NAME_STRING_CHECKED(UDirectionalLightComponent, bUseDistanceFieldShadows)
-			|| PropertyName == GET_MEMBER_NAME_STRING_CHECKED(UDirectionalLightComponent, DistanceFieldShadowDistance))
+		if (PropertyName == GET_MEMBER_NAME_STRING_CHECKED(UDirectionalLightComponent, DistanceFieldShadowDistance)
+			|| PropertyName == GET_MEMBER_NAME_STRING_CHECKED(UDirectionalLightComponent, LightSourceAngle))
 		{
 			static const auto CVar = IConsoleManager::Get().FindTConsoleVariableDataInt(TEXT("r.GenerateMeshDistanceFields"));
-			bool bCanEdit = CastShadows && CastDynamicShadows && Mobility != EComponentMobility::Static && CVar->GetValueOnGameThread() != 0;
-
-			if (PropertyName == GET_MEMBER_NAME_STRING_CHECKED(UDirectionalLightComponent, DistanceFieldShadowDistance))
-			{
-				bCanEdit = bCanEdit && bUseDistanceFieldShadows;
-			}
-
+			bool bCanEdit = CastShadows && CastDynamicShadows && bUseRayTracedDistanceFieldShadows && Mobility != EComponentMobility::Static && CVar->GetValueOnGameThread() != 0;
 			return bCanEdit;
 		}
 
diff --git a/Engine/Source/Runtime/Engine/Private/LightComponent.cpp b/Engine/Source/Runtime/Engine/Private/LightComponent.cpp
--- a/Engine/Source/Runtime/Engine/Private/LightComponent.cpp
+++ b/Engine/Source/Runtime/Engine/Private/LightComponent.cpp
@@ -210,8 +210,9 @@ FLightSceneProxy::FLightSceneProxy(const ULightComponent* InLightComponent)
 	, bCastTranslucentShadows(InLightComponent->CastTranslucentShadows)
 	, bAffectTranslucentLighting(InLightComponent->bAffectTranslucentLighting)
 	, bUsedAsAtmosphereSunLight(InLightComponent->IsUsedAsAtmosphereSunLight())
-	, bAffectDynamicIndirectLighting( InLightComponent->bAffectDynamicIndirectLighting )
-	, bHasReflectiveShadowMap( InLightComponent->bAffectDynamicIndirectLighting && InLightComponent->GetLightType() == LightType_Directional )
+	, bAffectDynamicIndirectLighting(InLightComponent->bAffectDynamicIndirectLighting)
+	, bHasReflectiveShadowMap(InLightComponent->bAffectDynamicIndirectLighting && InLightComponent->GetLightType() == LightType_Directional)
+	, bUseRayTracedDistanceFieldShadows(InLightComponent->bUseRayTracedDistanceFieldShadows)
 	, LightType(InLightComponent->GetLightType())	
 	, ComponentName(InLightComponent->GetOwner() ? InLightComponent->GetOwner()->GetFName() : InLightComponent->GetFName())
 	, LevelName(InLightComponent->GetOutermost()->GetFName())
@@ -461,6 +462,13 @@ bool ULightComponent::CanEditChange(const UProperty* InProperty) const
 			}
 		}
 
+		if (PropertyName == GET_MEMBER_NAME_STRING_CHECKED(ULightComponent, bUseRayTracedDistanceFieldShadows))
+		{
+			static const auto CVar = IConsoleManager::Get().FindTConsoleVariableDataInt(TEXT("r.GenerateMeshDistanceFields"));
+			bool bCanEdit = CastShadows && CastDynamicShadows && Mobility != EComponentMobility::Static && CVar->GetValueOnGameThread() != 0;
+			return bCanEdit;
+		}
+
 		if (PropertyName == GET_MEMBER_NAME_STRING_CHECKED(ULightComponent, LightFunctionScale)
 			|| PropertyName == GET_MEMBER_NAME_STRING_CHECKED(ULightComponent, LightFunctionFadeDistance)
 			|| PropertyName == GET_MEMBER_NAME_STRING_CHECKED(ULightComponent, DisabledBrightness))
@@ -514,6 +522,7 @@ void ULightComponent::PostEditChangeProperty(FPropertyChangedEvent& PropertyChan
 		PropertyName != GET_MEMBER_NAME_STRING_CHECKED(ULightComponent, BloomScale) &&
 		PropertyName != GET_MEMBER_NAME_STRING_CHECKED(ULightComponent, BloomThreshold) &&
 		PropertyName != GET_MEMBER_NAME_STRING_CHECKED(ULightComponent, BloomTint) &&
+		PropertyName != GET_MEMBER_NAME_STRING_CHECKED(ULightComponent, bUseRayTracedDistanceFieldShadows) &&
 		PropertyName != GET_MEMBER_NAME_STRING_CHECKED(ULightComponent, bVisible) &&
 		// Point light properties that shouldn't unbuild lighting
 		PropertyName != GET_MEMBER_NAME_STRING_CHECKED(UPointLightComponent, SourceRadius) &&
@@ -526,7 +535,6 @@ void ULightComponent::PostEditChangeProperty(FPropertyChangedEvent& PropertyChan
 		PropertyName != GET_MEMBER_NAME_STRING_CHECKED(UDirectionalLightComponent, CascadeTransitionFraction) &&
 		PropertyName != GET_MEMBER_NAME_STRING_CHECKED(UDirectionalLightComponent, ShadowDistanceFadeoutFraction) &&
 		PropertyName != GET_MEMBER_NAME_STRING_CHECKED(UDirectionalLightComponent, bUseInsetShadowsForMovableObjects) &&
-		PropertyName != GET_MEMBER_NAME_STRING_CHECKED(UDirectionalLightComponent, bUseDistanceFieldShadows) &&
 		PropertyName != GET_MEMBER_NAME_STRING_CHECKED(UDirectionalLightComponent, DistanceFieldShadowDistance) &&
 		PropertyName != GET_MEMBER_NAME_STRING_CHECKED(UDirectionalLightComponent, bEnableLightShaftOcclusion) &&
 		PropertyName != GET_MEMBER_NAME_STRING_CHECKED(UDirectionalLightComponent, OcclusionMaskDarkness) &&
diff --git a/Engine/Source/Runtime/Engine/Private/PointLightComponent.cpp b/Engine/Source/Runtime/Engine/Private/PointLightComponent.cpp
--- a/Engine/Source/Runtime/Engine/Private/PointLightComponent.cpp
+++ b/Engine/Source/Runtime/Engine/Private/PointLightComponent.cpp
@@ -84,6 +84,7 @@ public:
 			OutInitializer.MinLightW = 0.1f;
 			OutInitializer.MaxDistanceToCastInLightW = Radius;
 			OutInitializer.SplitIndex = INDEX_NONE;
+			OutInitializer.bRayTracedDistanceFieldShadow = UseRayTracedDistanceFieldShadows() && DoesPlatformSupportDistanceFieldShadowing(GRHIShaderPlatform);
 		
 			return true;
 		}
diff --git a/Engine/Source/Runtime/Engine/Private/PointLightSceneProxy.h b/Engine/Source/Runtime/Engine/Private/PointLightSceneProxy.h
--- a/Engine/Source/Runtime/Engine/Private/PointLightSceneProxy.h
+++ b/Engine/Source/Runtime/Engine/Private/PointLightSceneProxy.h
@@ -148,15 +148,15 @@ public:
 		return true;
 	}
 
-	virtual void SetScissorRect(FRHICommandList& RHICmdList, const FSceneView& View) const override
+	virtual bool GetScissorRect(FIntRect& ScissorRect, const FSceneView& View) const override
 	{
+		ScissorRect = View.ViewRect;
+
 		// Calculate a scissor rectangle for the light's radius.
 		if((GetLightToWorld().GetOrigin() - View.ViewMatrices.ViewOrigin).Size() > Radius)
 		{
 			FVector LightVector = View.ViewMatrices.ViewMatrix.TransformPosition(GetLightToWorld().GetOrigin());
 
-			int32 ScissorMinX = View.ViewRect.Min.X;
-			int32 ScissorMaxX = View.ViewRect.Max.X;
 			if(!GetPointLightBounds(
 				LightVector.X,
 				LightVector.Z,
@@ -166,10 +166,10 @@ public:
 				View,
 				View.ViewRect.Min.X,
 				View.ViewRect.Width(),
-				ScissorMinX,
-				ScissorMaxX))
+				ScissorRect.Min.X,
+				ScissorRect.Max.X))
 			{
-				return;
+				return false;
 			}
 
 			int32 ScissorMinY = View.ViewRect.Min.Y;
@@ -183,13 +183,27 @@ public:
 				View,
 				View.ViewRect.Min.Y,
 				View.ViewRect.Height(),
-				ScissorMinY,
-				ScissorMaxY))
+				ScissorRect.Min.Y,
+				ScissorRect.Max.Y))
 			{
-				return;
+				return false;
 			}
 
-			RHICmdList.SetScissorRect(true, ScissorMinX, ScissorMinY, ScissorMaxX, ScissorMaxY);
+			return true;
+		}
+		else
+		{
+			return false;
+		}
+	}
+
+	virtual void SetScissorRect(FRHICommandList& RHICmdList, const FSceneView& View) const override
+	{
+		FIntRect ScissorRect;
+
+		if (FPointLightSceneProxyBase::GetScissorRect(ScissorRect, View))
+		{
+			RHICmdList.SetScissorRect(true, ScissorRect.Min.X, ScissorRect.Min.Y, ScissorRect.Max.X, ScissorRect.Max.Y);
 		}
 		else
 		{
diff --git a/Engine/Source/Runtime/Engine/Private/SpotLightComponent.cpp b/Engine/Source/Runtime/Engine/Private/SpotLightComponent.cpp
--- a/Engine/Source/Runtime/Engine/Private/SpotLightComponent.cpp
+++ b/Engine/Source/Runtime/Engine/Private/SpotLightComponent.cpp
@@ -137,6 +137,7 @@ public:
 		OutInitializer.MinLightW = 0.1f;
 		OutInitializer.MaxDistanceToCastInLightW = Radius;
 		OutInitializer.SplitIndex = INDEX_NONE;
+		OutInitializer.bRayTracedDistanceFieldShadow = UseRayTracedDistanceFieldShadows() && DoesPlatformSupportDistanceFieldShadowing(GRHIShaderPlatform);
 		return true;
 	}
 
diff --git a/Engine/Source/Runtime/Engine/Public/SceneManagement.h b/Engine/Source/Runtime/Engine/Public/SceneManagement.h
--- a/Engine/Source/Runtime/Engine/Public/SceneManagement.h
+++ b/Engine/Source/Runtime/Engine/Public/SceneManagement.h
@@ -539,7 +539,7 @@ public:
 	/** Whether the shadow is a point light shadow that renders all faces of a cubemap in one pass. */
 	bool bOnePassPointLightShadow;
 
-	/** Whether the shadow is a directional light cascade that will be computed by ray tracing the distance field. */
+	/** Whether the shadow will be computed by ray tracing the distance field. */
 	bool bRayTracedDistanceFieldShadow;
 
 	FWholeSceneProjectedShadowInitializer()
@@ -610,6 +610,7 @@ public:
 	virtual float GetOuterConeAngle() const { return 0.0f; }
 	virtual float GetSourceRadius() const { return 0.0f; }
 	virtual bool IsInverseSquared() const { return false; }
+	virtual float GetLightSourceAngle() const { return 0.0f; }
 
 	virtual FVector2D GetLightShaftConeParams() const
 	{
@@ -691,6 +692,12 @@ public:
 	// @param OutCascadeSettings can be 0
 	virtual FSphere GetShadowSplitBounds(const class FSceneView& View, int32 SplitIndex, FShadowCascadeSettings* OutCascadeSettings) const { return FSphere(FVector::ZeroVector, 0); }
 
+	virtual bool GetScissorRect(FIntRect& ScissorRect, const FSceneView& View) const
+	{
+		ScissorRect = View.ViewRect;
+		return false;
+	}
+
 	virtual void SetScissorRect(FRHICommandList& RHICmdList, const FSceneView& View) const
 	{
 	}
@@ -726,6 +733,7 @@ public:
 	inline bool CastsStaticShadow() const { return bCastStaticShadow; }
 	inline bool CastsTranslucentShadows() const { return bCastTranslucentShadows; }
 	inline bool AffectsTranslucentLighting() const { return bAffectTranslucentLighting; }
+	inline bool UseRayTracedDistanceFieldShadows() const { return bUseRayTracedDistanceFieldShadows; }
 	inline uint8 GetLightType() const { return LightType; }
 	inline FName GetComponentName() const { return ComponentName; }
 	inline FName GetLevelName() const { return LevelName; }
@@ -840,6 +848,9 @@ protected:
 	const uint32 bAffectDynamicIndirectLighting : 1;
 	const uint32 bHasReflectiveShadowMap : 1;
 
+	/** Whether to use ray traced distance field area shadows. */
+	const uint32 bUseRayTracedDistanceFieldShadows : 1;
+
 	/** The light type (ELightComponentType) */
 	const uint8 LightType;
 
diff --git a/Engine/Source/Runtime/Renderer/Private/DeferredShadingRenderer.cpp b/Engine/Source/Runtime/Renderer/Private/DeferredShadingRenderer.cpp
--- a/Engine/Source/Runtime/Renderer/Private/DeferredShadingRenderer.cpp
+++ b/Engine/Source/Runtime/Renderer/Private/DeferredShadingRenderer.cpp
@@ -17,7 +17,7 @@
 #include "LightPropagationVolume.h"
 #include "DeferredShadingRenderer.h"
 #include "SceneUtils.h"
-#include "DistanceFieldLightingShared.h"
+#include "DistanceFieldSurfaceCacheLighting.h"
 #include "PostProcess/PostProcessing.h"
 
 TAutoConsoleVariable<int32> CVarEarlyZPass(
diff --git a/Engine/Source/Runtime/Renderer/Private/DistanceFieldLightingShared.h b/Engine/Source/Runtime/Renderer/Private/DistanceFieldLightingShared.h
--- a/Engine/Source/Runtime/Renderer/Private/DistanceFieldLightingShared.h
+++ b/Engine/Source/Runtime/Renderer/Private/DistanceFieldLightingShared.h
@@ -13,79 +13,10 @@ const int32 GDistanceFieldAOTileSizeX = 8;
 const int32 GDistanceFieldAOTileSizeY = 8;
 /** Base downsample factor that all distance field AO operations are done at. */
 const int32 GAODownsampleFactor = 2;
-/** Number of cone traced directions. */
-const int32 NumConeSampleDirections = 9;
+static const int32 GMaxNumObjectsPerTile = 512;
 
 extern FIntPoint GetBufferSizeForAO();
 
-/** Generates unit length, stratified and uniformly distributed direction samples in a hemisphere. */
-extern void GenerateStratifiedUniformHemisphereSamples2(int32 NumThetaSteps, int32 NumPhiSteps, FRandomStream& RandomStream, TArray<FVector4>& Samples);
-
-extern void GetSpacedVectors(TArray<TInlineAllocator<9> >& OutVectors);
-
-class FDistanceFieldAOParameters
-{
-public:
-	float OcclusionMaxDistance;
-	float Contrast;
-
-	FDistanceFieldAOParameters(float InOcclusionMaxDistance = 600.0f, float InContrast = 0)
-	{
-		Contrast = FMath::Clamp(InContrast, .01f, 2.0f);
-		OcclusionMaxDistance = FMath::Clamp(InOcclusionMaxDistance, 200.0f, 1500.0f);
-	}
-};
-
-BEGIN_UNIFORM_BUFFER_STRUCT(FAOSampleData2,)
-	DECLARE_UNIFORM_BUFFER_STRUCT_MEMBER_ARRAY(FVector4,SampleDirections,[NumConeSampleDirections])
-END_UNIFORM_BUFFER_STRUCT(FAOSampleData2)
-
-class FAOParameters
-{
-public:
-	void Bind(const FShaderParameterMap& ParameterMap)
-	{
-		AOMaxDistance.Bind(ParameterMap,TEXT("AOMaxDistance"));
-		AOStepScale.Bind(ParameterMap,TEXT("AOStepScale"));
-		AOStepExponentScale.Bind(ParameterMap,TEXT("AOStepExponentScale"));
-		AOMaxViewDistance.Bind(ParameterMap,TEXT("AOMaxViewDistance"));
-	}
-
-	friend FArchive& operator<<(FArchive& Ar,FAOParameters& Parameters)
-	{
-		Ar << Parameters.AOMaxDistance;
-		Ar << Parameters.AOStepScale;
-		Ar << Parameters.AOStepExponentScale;
-		Ar << Parameters.AOMaxViewDistance;
-		return Ar;
-	}
-
-	template<typename ShaderRHIParamRef>
-	void Set(FRHICommandList& RHICmdList, const ShaderRHIParamRef ShaderRHI, const FDistanceFieldAOParameters& Parameters)
-	{
-		SetShaderValue(RHICmdList, ShaderRHI, AOMaxDistance, Parameters.OcclusionMaxDistance);
-
-		extern float GAOConeHalfAngle;
-		const float AOLargestSampleOffset = Parameters.OcclusionMaxDistance / (1 + FMath::Tan(GAOConeHalfAngle));
-
-		extern float GAOStepExponentScale;
-		extern uint32 GAONumConeSteps;
-		float AOStepScaleValue = AOLargestSampleOffset / FMath::Pow(2.0f, GAOStepExponentScale * (GAONumConeSteps - 1));
-		SetShaderValue(RHICmdList, ShaderRHI, AOStepScale, AOStepScaleValue);
-
-		SetShaderValue(RHICmdList, ShaderRHI, AOStepExponentScale, GAOStepExponentScale);
-
-		extern float GAOMaxViewDistance;
-		SetShaderValue(RHICmdList, ShaderRHI, AOMaxViewDistance, GAOMaxViewDistance);
-	}
-
-private:
-	FShaderParameter AOMaxDistance;
-	FShaderParameter AOStepScale;
-	FShaderParameter AOStepExponentScale;
-	FShaderParameter AOMaxViewDistance;
-};
-
 class FDistanceFieldObjectData
 {
 public:
@@ -93,50 +24,69 @@ public:
 	FVertexBufferRHIRef Bounds;
 	FVertexBufferRHIRef Data;
 	FVertexBufferRHIRef Data2;
+	FVertexBufferRHIRef BoxBounds;
 	FShaderResourceViewRHIRef BoundsSRV;
 	FShaderResourceViewRHIRef DataSRV;
 	FShaderResourceViewRHIRef Data2SRV;
+	FShaderResourceViewRHIRef BoxBoundsSRV;
 };
 
-class FDistanceFieldObjectBuffers : public FRenderResource
+class FDistanceFieldObjectBuffers
 {
 public:
 
 	// In float4's
 	static int32 ObjectDataStride;
 	static int32 ObjectData2Stride;
+	static int32 ObjectBoxBoundsStride;
 
+	bool bWantBoxBounds;
+	bool bWantVolumeToWorld;
 	int32 MaxObjects;
 	FDistanceFieldObjectData ObjectData;
 
 	FDistanceFieldObjectBuffers()
 	{
+		bWantBoxBounds = false;
+		bWantVolumeToWorld = true;
 		MaxObjects = 0;
 	}
 
-	virtual void InitDynamicRHI()  override
+	void Initialize()
 	{
 		if (MaxObjects > 0)
 		{
 			FRHIResourceCreateInfo CreateInfo;
 			ObjectData.Bounds = RHICreateVertexBuffer(MaxObjects * sizeof(FVector4), BUF_Volatile | BUF_ShaderResource, CreateInfo);
 			ObjectData.Data = RHICreateVertexBuffer(MaxObjects * sizeof(FVector4) * ObjectDataStride, BUF_Volatile | BUF_ShaderResource, CreateInfo);
-			ObjectData.Data2 = RHICreateVertexBuffer(MaxObjects * sizeof(FVector4) * ObjectData2Stride, BUF_Volatile | BUF_ShaderResource, CreateInfo);
 
 			ObjectData.BoundsSRV = RHICreateShaderResourceView(ObjectData.Bounds, sizeof(FVector4), PF_A32B32G32R32F);
 			ObjectData.DataSRV = RHICreateShaderResourceView(ObjectData.Data, sizeof(FVector4), PF_A32B32G32R32F);
-			ObjectData.Data2SRV = RHICreateShaderResourceView(ObjectData.Data2, sizeof(FVector4), PF_A32B32G32R32F);
+
+			if (bWantVolumeToWorld)
+			{
+				ObjectData.Data2 = RHICreateVertexBuffer(MaxObjects * sizeof(FVector4) * ObjectData2Stride, BUF_Volatile | BUF_ShaderResource, CreateInfo);
+				ObjectData.Data2SRV = RHICreateShaderResourceView(ObjectData.Data2, sizeof(FVector4), PF_A32B32G32R32F);
+			}
+
+			if (bWantBoxBounds)
+			{
+				ObjectData.BoxBounds = RHICreateVertexBuffer(MaxObjects * sizeof(FVector4) * ObjectBoxBoundsStride, BUF_Volatile | BUF_ShaderResource, CreateInfo);
+				ObjectData.BoxBoundsSRV = RHICreateShaderResourceView(ObjectData.BoxBounds, sizeof(FVector4), PF_A32B32G32R32F);
+			}
 		}
 	}
 
-	virtual void ReleaseDynamicRHI() override
+	void Release()
 	{
 		ObjectData.Bounds.SafeRelease();
 		ObjectData.Data.SafeRelease();
 		ObjectData.Data2.SafeRelease();
+		ObjectData.BoxBounds.SafeRelease();
 		ObjectData.BoundsSRV.SafeRelease(); 
 		ObjectData.DataSRV.SafeRelease(); 
 		ObjectData.Data2SRV.SafeRelease(); 
+		ObjectData.BoxBoundsSRV.SafeRelease(); 
 	}
 };
 
@@ -148,6 +98,7 @@ public:
 		ObjectBounds.Bind(ParameterMap, TEXT("ObjectBounds"));
 		ObjectData.Bind(ParameterMap, TEXT("ObjectData"));
 		ObjectData2.Bind(ParameterMap, TEXT("ObjectData2"));
+		ObjectBoxBounds.Bind(ParameterMap, TEXT("ObjectBoxBounds"));
 		NumObjects.Bind(ParameterMap, TEXT("NumObjects"));
 		DistanceFieldTexture.Bind(ParameterMap, TEXT("DistanceFieldTexture"));
 		DistanceFieldSampler.Bind(ParameterMap, TEXT("DistanceFieldSampler"));
@@ -159,7 +110,24 @@ public:
 	{
 		SetSRVParameter(RHICmdList, ShaderRHI, ObjectBounds, ObjectBuffers.ObjectData.BoundsSRV);
 		SetSRVParameter(RHICmdList, ShaderRHI, ObjectData, ObjectBuffers.ObjectData.DataSRV);
-		SetSRVParameter(RHICmdList, ShaderRHI, ObjectData2, ObjectBuffers.ObjectData.Data2SRV);
+
+		if (ObjectBuffers.bWantVolumeToWorld)
+		{
+			SetSRVParameter(RHICmdList, ShaderRHI, ObjectData2, ObjectBuffers.ObjectData.Data2SRV);
+		}
+		else
+		{
+			check(!ObjectData2.IsBound());
+		}
+
+		if (ObjectBuffers.bWantBoxBounds)
+		{
+			SetSRVParameter(RHICmdList, ShaderRHI, ObjectBoxBounds, ObjectBuffers.ObjectData.BoxBoundsSRV);
+		}
+		else
+		{
+			check(!ObjectBoxBounds.IsBound());
+		}
 
 		SetShaderValue(RHICmdList, ShaderRHI, NumObjects, NumObjectsValue);
 
@@ -182,7 +150,7 @@ public:
 
 	friend FArchive& operator<<(FArchive& Ar, FDistanceFieldObjectBufferParameters& P)
 	{
-		Ar << P.ObjectBounds << P.ObjectData << P.ObjectData2 << P.NumObjects << P.DistanceFieldTexture << P.DistanceFieldSampler << P.DistanceFieldAtlasTexelSize;
+		Ar << P.ObjectBounds << P.ObjectData << P.ObjectData2 << P.ObjectBoxBounds << P.NumObjects << P.DistanceFieldTexture << P.DistanceFieldSampler << P.DistanceFieldAtlasTexelSize;
 		return Ar;
 	}
 
@@ -190,77 +158,33 @@ private:
 	FShaderResourceParameter ObjectBounds;
 	FShaderResourceParameter ObjectData;
 	FShaderResourceParameter ObjectData2;
+	FShaderResourceParameter ObjectBoxBounds;
 	FShaderParameter NumObjects;
 	FShaderResourceParameter DistanceFieldTexture;
 	FShaderResourceParameter DistanceFieldSampler;
 	FShaderParameter DistanceFieldAtlasTexelSize;
 };
 
-class FQuadVertexBuffer : public FVertexBuffer
+/**  */
+class FLightTileIntersectionResources
 {
 public:
-
-	virtual void InitRHI() override
+	void Initialize()
 	{
-		// Used as a non-indexed triangle list, so 6 vertices per quad
-		const uint32 Size = 6 * sizeof(FScreenVertex);
-		FRHIResourceCreateInfo CreateInfo;
-		VertexBufferRHI = RHICreateVertexBuffer(Size, BUF_Static, CreateInfo);
-		void* Buffer = RHILockVertexBuffer(VertexBufferRHI, 0, Size, RLM_WriteOnly);
-		FScreenVertex* DestVertex = (FScreenVertex*)Buffer;
+		TileHeadDataUnpacked.Initialize(sizeof(uint32), TileDimensions.X * TileDimensions.Y * 2, PF_R32_UINT, BUF_Static);
 
-		DestVertex[0].Position = FVector2D(0, 0);
-		DestVertex[0].UV = FVector2D(1, 1);
-		DestVertex[1].Position = FVector2D(0, 0);
-		DestVertex[1].UV = FVector2D(1, 0);
-		DestVertex[2].Position = FVector2D(0, 0);
-		DestVertex[2].UV = FVector2D(0, 1);
-		DestVertex[3].Position = FVector2D(0, 0);
-		DestVertex[3].UV = FVector2D(1, 0);
-		DestVertex[4].Position = FVector2D(0, 0);
-		DestVertex[4].UV = FVector2D(0, 0);
-		DestVertex[5].Position = FVector2D(0, 0);
-		DestVertex[5].UV = FVector2D(0, 1);
-
-		RHIUnlockVertexBuffer(VertexBufferRHI);      
+		//@todo - handle max exceeded
+		TileArrayData.Initialize(sizeof(uint16), GMaxNumObjectsPerTile * TileDimensions.X * TileDimensions.Y, PF_R16_UINT, BUF_Static);
 	}
-};
 
-extern TGlobalResource<FQuadVertexBuffer> GQuadVertexBuffer;
-
-class FCircleVertexBuffer : public FVertexBuffer
-{
-public:
-
-	virtual void InitRHI() override
+	void Release()
 	{
-		int32 NumSections = 8;
-
-		// Used as a non-indexed triangle list, so 3 vertices per triangle
-		const uint32 Size = 3 * NumSections * sizeof(FScreenVertex);
-		FRHIResourceCreateInfo CreateInfo;
-		VertexBufferRHI = RHICreateVertexBuffer(Size, BUF_Static, CreateInfo);
-		void* Buffer = RHILockVertexBuffer(VertexBufferRHI, 0, Size, RLM_WriteOnly);
-		FScreenVertex* DestVertex = (FScreenVertex*)Buffer;
-
-		for (int32 SectionIndex = 0; SectionIndex < NumSections; SectionIndex++)
-		{
-			float Fraction = SectionIndex / (float)NumSections;
-			float CurrentAngle = Fraction * 2 * PI;
-			float NextAngle = ((SectionIndex + 1) / (float)NumSections) * 2 * PI;
-			FVector2D CurrentPosition(FMath::Cos(CurrentAngle), FMath::Sin(CurrentAngle));
-			FVector2D NextPosition(FMath::Cos(NextAngle), FMath::Sin(NextAngle));
+		TileHeadDataUnpacked.Release();
+		TileArrayData.Release();
+	}
 
-			DestVertex[SectionIndex * 3 + 0].Position = FVector2D(0, 0);
-			DestVertex[SectionIndex * 3 + 0].UV = CurrentPosition;
-			DestVertex[SectionIndex * 3 + 1].Position = FVector2D(0, 0);
-			DestVertex[SectionIndex * 3 + 1].UV = NextPosition;
-			DestVertex[SectionIndex * 3 + 2].Position = FVector2D(0, 0);
-			DestVertex[SectionIndex * 3 + 2].UV = FVector2D(.5f, .5f);
-		}
+	FIntPoint TileDimensions;
 
-		RHIUnlockVertexBuffer(VertexBufferRHI);      
-	}
+	FRWBuffer TileHeadDataUnpacked;
+	FRWBuffer TileArrayData;
 };
-
-extern TGlobalResource<FCircleVertexBuffer> GCircleVertexBuffer;
diff --git a/Engine/Source/Runtime/Renderer/Private/DistanceFieldShadowing.cpp b/Engine/Source/Runtime/Renderer/Private/DistanceFieldShadowing.cpp
--- a/Engine/Source/Runtime/Renderer/Private/DistanceFieldShadowing.cpp
+++ b/Engine/Source/Runtime/Renderer/Private/DistanceFieldShadowing.cpp
@@ -11,8 +11,8 @@
 #include "PostProcessing.h"
 #include "SceneFilterRendering.h"
 #include "ScreenRendering.h"
-#include "DistanceFieldSurfaceCacheLighting.h"
 #include "DistanceFieldLightingShared.h"
+#include "DistanceFieldSurfaceCacheLighting.h"
 #include "RHICommandList.h"
 #include "SceneUtils.h"
 #include "DistanceFieldAtlas.h"
@@ -25,12 +25,248 @@ FAutoConsoleVariableRef CVarDistanceFieldShadowing(
 	ECVF_Cheat | ECVF_RenderThreadSafe
 	);
 
-//@todo - per light, breaks with multiple directional lights
-TGlobalResource<FDistanceFieldObjectBuffers> GDirectionalLightObjectBuffers;
+int32 GShadowScatterTileCulling = 1;
+FAutoConsoleVariableRef CVarShadowScatterTileCulling(
+	TEXT("r.DFShadowScatterTileCulling"),
+	GShadowScatterTileCulling,
+	TEXT("Whether to use the rasterizer to scatter objects onto the tile grid for culling."),
+	ECVF_Cheat | ECVF_RenderThreadSafe
+	);
+
+float GShadowWorldTileSize = 200.0f;
+FAutoConsoleVariableRef CVarShadowWorldTileSize(
+	TEXT("r.DFShadowWorldTileSize"),
+	GShadowWorldTileSize,
+	TEXT("World space size of a tile used for culling for directional lights."),
+	ECVF_Cheat | ECVF_RenderThreadSafe
+	);
 
-class FDistanceFieldShadowingCS : public FGlobalShader
+/**  */
+class FClearTilesCS : public FGlobalShader
 {
-	DECLARE_SHADER_TYPE(FDistanceFieldShadowingCS, Global);
+	DECLARE_SHADER_TYPE(FClearTilesCS,Global)
+public:
+
+	static bool ShouldCache(EShaderPlatform Platform)
+	{
+		return IsFeatureLevelSupported(Platform, ERHIFeatureLevel::SM5) && DoesPlatformSupportDistanceFieldShadowing(Platform);
+	}
+
+	static void ModifyCompilationEnvironment(EShaderPlatform Platform, FShaderCompilerEnvironment& OutEnvironment)
+	{
+		FGlobalShader::ModifyCompilationEnvironment(Platform,OutEnvironment);
+		OutEnvironment.SetDefine(TEXT("THREADGROUP_SIZEX"), GDistanceFieldAOTileSizeX);
+		OutEnvironment.SetDefine(TEXT("THREADGROUP_SIZEY"), GDistanceFieldAOTileSizeY);
+	}
+
+	FClearTilesCS(const ShaderMetaType::CompiledShaderInitializerType& Initializer)
+		: FGlobalShader(Initializer)
+	{
+		TileHeadDataUnpacked.Bind(Initializer.ParameterMap, TEXT("TileHeadDataUnpacked"));
+		NumGroups.Bind(Initializer.ParameterMap, TEXT("NumGroups"));
+	}
+
+	FClearTilesCS()
+	{
+	}
+
+	void SetParameters(FRHICommandList& RHICmdList, const FSceneView& View, FVector2D NumGroupsValue, FLightTileIntersectionResources* TileIntersectionResources)
+	{
+		FComputeShaderRHIParamRef ShaderRHI = GetComputeShader();
+
+		FGlobalShader::SetParameters(RHICmdList, ShaderRHI, View);
+		TileHeadDataUnpacked.SetBuffer(RHICmdList, ShaderRHI, TileIntersectionResources->TileHeadDataUnpacked);
+
+		SetShaderValue(RHICmdList, ShaderRHI, NumGroups, NumGroupsValue);
+	}
+
+	void UnsetParameters(FRHICommandList& RHICmdList)
+	{
+		TileHeadDataUnpacked.UnsetUAV(RHICmdList, GetComputeShader());
+	}
+
+	virtual bool Serialize(FArchive& Ar)
+	{		
+		bool bShaderHasOutdatedParameters = FGlobalShader::Serialize(Ar);
+		Ar << TileHeadDataUnpacked;
+		Ar << NumGroups;
+		return bShaderHasOutdatedParameters;
+	}
+
+private:
+
+	FRWShaderParameter TileHeadDataUnpacked;
+	FShaderParameter NumGroups;
+};
+
+IMPLEMENT_SHADER_TYPE(,FClearTilesCS,TEXT("DistanceFieldShadowing"),TEXT("ClearTilesCS"),SF_Compute);
+
+
+/**  */
+class FShadowObjectCullVS : public FGlobalShader
+{
+	DECLARE_SHADER_TYPE(FShadowObjectCullVS,Global);
+public:
+
+	static bool ShouldCache(EShaderPlatform Platform) 
+	{ 
+		return IsFeatureLevelSupported(Platform, ERHIFeatureLevel::SM5) && DoesPlatformSupportDistanceFieldShadowing(Platform); 
+	}
+
+	static void ModifyCompilationEnvironment(EShaderPlatform Platform, FShaderCompilerEnvironment& OutEnvironment)
+	{
+		FGlobalShader::ModifyCompilationEnvironment(Platform,OutEnvironment);
+	}
+
+	FShadowObjectCullVS(const ShaderMetaType::CompiledShaderInitializerType& Initializer):
+		FGlobalShader(Initializer)
+	{
+		ObjectParameters.Bind(Initializer.ParameterMap);
+		ConservativeRadiusScale.Bind(Initializer.ParameterMap, TEXT("ConservativeRadiusScale"));
+		WorldToShadow.Bind(Initializer.ParameterMap, TEXT("WorldToShadow"));
+		MinRadius.Bind(Initializer.ParameterMap, TEXT("MinRadius"));
+	}
+
+	FShadowObjectCullVS() {}
+
+	void SetParameters(FRHICommandList& RHICmdList, const FSceneView& View, FVector2D NumGroupsValue, const FProjectedShadowInfo* ProjectedShadowInfo, int32 NumObjects)
+	{
+		const FVertexShaderRHIParamRef ShaderRHI = GetVertexShader();
+		FGlobalShader::SetParameters(RHICmdList, ShaderRHI, View);
+		
+		ObjectParameters.Set(RHICmdList, ShaderRHI, *ProjectedShadowInfo->LightSceneInfo->DistanceFieldObjectBuffers, NumObjects);
+
+		const int32 NumRings = StencilingGeometry::GLowPolyStencilSphereVertexBuffer.GetNumRings();
+		const float RadiansPerRingSegment = PI / (float)NumRings;
+
+		// Boost the effective radius so that the edges of the sphere approximation lie on the sphere, instead of the vertices
+		const float ConservativeRadiusScaleValue = 1.0f / FMath::Cos(RadiansPerRingSegment);
+
+		SetShaderValue(RHICmdList, ShaderRHI, ConservativeRadiusScale, ConservativeRadiusScaleValue);
+
+		FMatrix WorldToShadowMatrixValue = FTranslationMatrix(ProjectedShadowInfo->PreShadowTranslation) * ProjectedShadowInfo->SubjectAndReceiverMatrix;
+		SetShaderValue(RHICmdList, ShaderRHI, WorldToShadow, WorldToShadowMatrixValue);
+
+		float MinRadiusValue = 2 * ProjectedShadowInfo->ShadowBounds.W / FMath::Min(NumGroupsValue.X, NumGroupsValue.Y);
+		SetShaderValue(RHICmdList, ShaderRHI, MinRadius, MinRadiusValue);
+	}
+
+	virtual bool Serialize(FArchive& Ar)
+	{
+		bool bShaderHasOutdatedParameters = FGlobalShader::Serialize(Ar);
+		Ar << ObjectParameters;
+		Ar << ConservativeRadiusScale;
+		Ar << WorldToShadow;
+		Ar << MinRadius;
+		return bShaderHasOutdatedParameters;
+	}
+
+private:
+	FDistanceFieldObjectBufferParameters ObjectParameters;
+	FShaderParameter ConservativeRadiusScale;
+	FShaderParameter WorldToShadow;
+	FShaderParameter MinRadius;
+};
+
+IMPLEMENT_SHADER_TYPE(,FShadowObjectCullVS,TEXT("DistanceFieldShadowing"),TEXT("ShadowObjectCullVS"),SF_Vertex);
+
+class FShadowObjectCullPS : public FGlobalShader
+{
+	DECLARE_SHADER_TYPE(FShadowObjectCullPS, Global);
+public:
+
+	static bool ShouldCache(EShaderPlatform Platform)
+	{
+		return IsFeatureLevelSupported(Platform, ERHIFeatureLevel::SM5) && DoesPlatformSupportDistanceFieldShadowing(Platform);
+	}
+
+	static void ModifyCompilationEnvironment(EShaderPlatform Platform, FShaderCompilerEnvironment& OutEnvironment)
+	{
+		FGlobalShader::ModifyCompilationEnvironment(Platform,OutEnvironment);
+		
+		OutEnvironment.SetDefine(TEXT("DOWNSAMPLE_FACTOR"), GAODownsampleFactor);
+		OutEnvironment.SetDefine(TEXT("MAX_OBJECTS_PER_TILE"), GMaxNumObjectsPerTile);
+	}
+
+	/** Default constructor. */
+	FShadowObjectCullPS() {}
+
+	/** Initialization constructor. */
+	FShadowObjectCullPS(const ShaderMetaType::CompiledShaderInitializerType& Initializer)
+		: FGlobalShader(Initializer)
+	{
+		ObjectParameters.Bind(Initializer.ParameterMap);
+		TileHeadDataUnpacked.Bind(Initializer.ParameterMap, TEXT("TileHeadDataUnpacked"));
+		TileArrayData.Bind(Initializer.ParameterMap, TEXT("TileArrayData"));
+		NumGroups.Bind(Initializer.ParameterMap, TEXT("NumGroups"));
+	}
+
+	void SetParameters(
+		FRHICommandList& RHICmdList, 
+		const FSceneView& View, 
+		FVector2D NumGroupsValue, 
+		int32 NumObjects, 
+		const FProjectedShadowInfo* ProjectedShadowInfo,
+		FLightTileIntersectionResources* TileIntersectionResources)
+	{
+		const FPixelShaderRHIParamRef ShaderRHI = GetPixelShader();
+		FGlobalShader::SetParameters(RHICmdList, ShaderRHI, View);
+
+		ObjectParameters.Set(RHICmdList, ShaderRHI, *ProjectedShadowInfo->LightSceneInfo->DistanceFieldObjectBuffers, NumObjects);
+		SetShaderValue(RHICmdList, ShaderRHI, NumGroups, NumGroupsValue);
+	}
+
+	void GetUAVs(const FSceneView& View, FLightTileIntersectionResources* TileIntersectionResources, TArray<FUnorderedAccessViewRHIParamRef>& UAVs)
+	{
+		int32 MaxIndex = FMath::Max(TileHeadDataUnpacked.GetUAVIndex(), TileArrayData.GetUAVIndex());
+		UAVs.AddZeroed(MaxIndex + 1);
+
+		if (TileHeadDataUnpacked.IsBound())
+		{
+			UAVs[TileHeadDataUnpacked.GetUAVIndex()] = TileIntersectionResources->TileHeadDataUnpacked.UAV;
+		}
+
+		if (TileArrayData.IsBound())
+		{
+			UAVs[TileArrayData.GetUAVIndex()] = TileIntersectionResources->TileArrayData.UAV;
+		}
+
+		check(UAVs.Num() > 0);
+	}
+
+	// FShader interface.
+	virtual bool Serialize(FArchive& Ar)
+	{
+		bool bShaderHasOutdatedParameters = FGlobalShader::Serialize(Ar);
+		Ar << ObjectParameters;
+		Ar << TileHeadDataUnpacked;
+		Ar << TileArrayData;
+		Ar << NumGroups;
+		return bShaderHasOutdatedParameters;
+	}
+
+private:
+	FDistanceFieldObjectBufferParameters ObjectParameters;
+	FRWShaderParameter TileHeadDataUnpacked;
+	FRWShaderParameter TileArrayData;
+	FShaderParameter NumGroups;
+};
+
+IMPLEMENT_SHADER_TYPE(,FShadowObjectCullPS,TEXT("DistanceFieldShadowing"),TEXT("ShadowObjectCullPS"),SF_Pixel);
+
+
+
+enum EDistanceFieldShadowingType
+{
+	DFS_DirectionalLightScatterTileCulling,
+	DFS_DirectionalLightTiledCulling,
+	DFS_PointLightTiledCulling
+};
+
+template<EDistanceFieldShadowingType ShadowingType>
+class TDistanceFieldShadowingCS : public FGlobalShader
+{
+	DECLARE_SHADER_TYPE(TDistanceFieldShadowingCS, Global);
 public:
 
 	static bool ShouldCache(EShaderPlatform Platform)
@@ -43,20 +279,31 @@ public:
 		OutEnvironment.SetDefine(TEXT("DOWNSAMPLE_FACTOR"), GAODownsampleFactor);
 		OutEnvironment.SetDefine(TEXT("THREADGROUP_SIZEX"), GDistanceFieldAOTileSizeX);
 		OutEnvironment.SetDefine(TEXT("THREADGROUP_SIZEY"), GDistanceFieldAOTileSizeY);
+		OutEnvironment.SetDefine(TEXT("MAX_OBJECTS_PER_TILE"), GMaxNumObjectsPerTile);
+		OutEnvironment.SetDefine(TEXT("SCATTER_TILE_CULLING"), ShadowingType == DFS_DirectionalLightScatterTileCulling ? TEXT("1") : TEXT("0"));
+		OutEnvironment.SetDefine(TEXT("POINT_LIGHT"), ShadowingType == DFS_PointLightTiledCulling ? TEXT("1") : TEXT("0"));
 	}
 
 	/** Default constructor. */
-	FDistanceFieldShadowingCS() {}
+	TDistanceFieldShadowingCS() {}
 
 	/** Initialization constructor. */
-	FDistanceFieldShadowingCS(const ShaderMetaType::CompiledShaderInitializerType& Initializer)
+	TDistanceFieldShadowingCS(const ShaderMetaType::CompiledShaderInitializerType& Initializer)
 		: FGlobalShader(Initializer)
 	{
 		ShadowFactors.Bind(Initializer.ParameterMap, TEXT("ShadowFactors"));
 		NumGroups.Bind(Initializer.ParameterMap, TEXT("NumGroups"));
 		LightDirection.Bind(Initializer.ParameterMap, TEXT("LightDirection"));
+		LightSourceRadius.Bind(Initializer.ParameterMap, TEXT("LightSourceRadius"));
+		LightPositionAndInvRadius.Bind(Initializer.ParameterMap, TEXT("LightPositionAndInvRadius"));
+		TanLightAngleAndNormalThreshold.Bind(Initializer.ParameterMap, TEXT("TanLightAngleAndNormalThreshold"));
+		ScissorRectMinAndSize.Bind(Initializer.ParameterMap, TEXT("ScissorRectMinAndSize"));
 		ObjectParameters.Bind(Initializer.ParameterMap);
 		DeferredParameters.Bind(Initializer.ParameterMap);
+		TileHeadDataUnpacked.Bind(Initializer.ParameterMap, TEXT("TileHeadDataUnpacked"));
+		TileArrayData.Bind(Initializer.ParameterMap, TEXT("TileArrayData"));
+		TileListGroupSize.Bind(Initializer.ParameterMap, TEXT("TileListGroupSize"));
+		WorldToShadow.Bind(Initializer.ParameterMap, TEXT("WorldToShadow"));
 	}
 
 	void SetParameters(
@@ -65,7 +312,9 @@ public:
 		const FProjectedShadowInfo* ProjectedShadowInfo,
 		FSceneRenderTargetItem& ShadowFactorsValue, 
 		int32 NumObjects, 
-		FVector2D NumGroupsValue)
+		FVector2D NumGroupsValue,
+		const FIntRect& ScissorRect,
+		FLightTileIntersectionResources* TileIntersectionResources)
 	{
 		const FComputeShaderRHIParamRef ShaderRHI = GetComputeShader();
 
@@ -73,11 +322,43 @@ public:
 
 		ShadowFactors.SetTexture(RHICmdList, ShaderRHI, ShadowFactorsValue.ShaderResourceTexture, ShadowFactorsValue.UAV);
 
-		ObjectParameters.Set(RHICmdList, ShaderRHI, GDirectionalLightObjectBuffers, NumObjects);
+		ObjectParameters.Set(RHICmdList, ShaderRHI, *ProjectedShadowInfo->LightSceneInfo->DistanceFieldObjectBuffers, NumObjects);
 		DeferredParameters.Set(RHICmdList, ShaderRHI, View);
 
 		SetShaderValue(RHICmdList, ShaderRHI, NumGroups, NumGroupsValue);
-		SetShaderValue(RHICmdList, ShaderRHI, LightDirection, ProjectedShadowInfo->LightSceneInfo->Proxy->GetDirection());
+
+		FVector4 LightPositionAndInvRadiusValue;
+		FVector4 LightColorAndFalloffExponent;
+		FVector NormalizedLightDirection;
+		FVector2D SpotAngles;
+		float LightSourceRadiusValue;
+		float LightSourceLength;
+		float LightMinRoughness;
+
+		ProjectedShadowInfo->LightSceneInfo->Proxy->GetParameters(LightPositionAndInvRadiusValue, LightColorAndFalloffExponent, NormalizedLightDirection, SpotAngles, LightSourceRadiusValue, LightSourceLength, LightMinRoughness);
+
+		SetShaderValue(RHICmdList, ShaderRHI, LightDirection, NormalizedLightDirection);
+		SetShaderValue(RHICmdList, ShaderRHI, LightPositionAndInvRadius, LightPositionAndInvRadiusValue);
+		// Default light source radius of 0 gives poor results
+		SetShaderValue(RHICmdList, ShaderRHI, LightSourceRadius, LightSourceRadiusValue == 0 ? 20 : FMath::Clamp(LightSourceRadiusValue, .001f, 1.0f / (4 * LightPositionAndInvRadiusValue.W)));
+
+		const float LightSourceAngle = FMath::Clamp(ProjectedShadowInfo->LightSceneInfo->Proxy->GetLightSourceAngle(), 0.001f, 5.0f) * PI / 180.0f;
+		const FVector2D TanLightAngleAndNormalThresholdValue(FMath::Tan(LightSourceAngle), FMath::Cos(PI / 2 + LightSourceAngle));
+		SetShaderValue(RHICmdList, ShaderRHI, TanLightAngleAndNormalThreshold, TanLightAngleAndNormalThresholdValue);
+
+		SetShaderValue(RHICmdList, ShaderRHI, ScissorRectMinAndSize, FIntRect(ScissorRect.Min, ScissorRect.Size()));
+
+		check(TileIntersectionResources || !TileHeadDataUnpacked.IsBound());
+
+		if (TileIntersectionResources)
+		{
+			SetSRVParameter(RHICmdList, ShaderRHI, TileHeadDataUnpacked, TileIntersectionResources->TileHeadDataUnpacked.SRV);
+			SetSRVParameter(RHICmdList, ShaderRHI, TileArrayData, TileIntersectionResources->TileArrayData.SRV);
+			SetShaderValue(RHICmdList, ShaderRHI, TileListGroupSize, TileIntersectionResources->TileDimensions);
+		}
+
+		FMatrix WorldToShadowMatrixValue = FTranslationMatrix(ProjectedShadowInfo->PreShadowTranslation) * ProjectedShadowInfo->SubjectAndReceiverMatrix;
+		SetShaderValue(RHICmdList, ShaderRHI, WorldToShadow, WorldToShadowMatrixValue);
 	}
 
 	void UnsetParameters(FRHICommandList& RHICmdList)
@@ -92,8 +373,16 @@ public:
 		Ar << ShadowFactors;
 		Ar << NumGroups;
 		Ar << LightDirection;
+		Ar << LightPositionAndInvRadius;
+		Ar << LightSourceRadius;
+		Ar << TanLightAngleAndNormalThreshold;
+		Ar << ScissorRectMinAndSize;
 		Ar << ObjectParameters;
 		Ar << DeferredParameters;
+		Ar << TileHeadDataUnpacked;
+		Ar << TileArrayData;
+		Ar << TileListGroupSize;
+		Ar << WorldToShadow;
 		return bShaderHasOutdatedParameters;
 	}
 
@@ -102,11 +391,21 @@ private:
 	FRWShaderParameter ShadowFactors;
 	FShaderParameter NumGroups;
 	FShaderParameter LightDirection;
+	FShaderParameter LightPositionAndInvRadius;
+	FShaderParameter LightSourceRadius;
+	FShaderParameter TanLightAngleAndNormalThreshold;
+	FShaderParameter ScissorRectMinAndSize;
 	FDistanceFieldObjectBufferParameters ObjectParameters;
 	FDeferredPixelShaderParameters DeferredParameters;
+	FShaderResourceParameter TileHeadDataUnpacked;
+	FShaderResourceParameter TileArrayData;
+	FShaderParameter TileListGroupSize;
+	FShaderParameter WorldToShadow;
 };
 
-IMPLEMENT_SHADER_TYPE(,FDistanceFieldShadowingCS,TEXT("DistanceFieldShadowing"),TEXT("DistanceFieldShadowingCS"),SF_Compute);
+IMPLEMENT_SHADER_TYPE(template<>,TDistanceFieldShadowingCS<DFS_DirectionalLightScatterTileCulling>,TEXT("DistanceFieldShadowing"),TEXT("DistanceFieldShadowingCS"),SF_Compute);
+IMPLEMENT_SHADER_TYPE(template<>,TDistanceFieldShadowingCS<DFS_DirectionalLightTiledCulling>,TEXT("DistanceFieldShadowing"),TEXT("DistanceFieldShadowingCS"),SF_Compute);
+IMPLEMENT_SHADER_TYPE(template<>,TDistanceFieldShadowingCS<DFS_PointLightTiledCulling>,TEXT("DistanceFieldShadowing"),TEXT("DistanceFieldShadowingCS"),SF_Compute);
 
 class FDistanceFieldShadowingUpsamplePS : public FGlobalShader
 {
@@ -133,9 +432,10 @@ public:
 		DeferredParameters.Bind(Initializer.ParameterMap);
 		ShadowFactorsTexture.Bind(Initializer.ParameterMap,TEXT("ShadowFactorsTexture"));
 		ShadowFactorsSampler.Bind(Initializer.ParameterMap,TEXT("ShadowFactorsSampler"));
+		ScissorRectMinAndSize.Bind(Initializer.ParameterMap,TEXT("ScissorRectMinAndSize"));
 	}
 
-	void SetParameters(FRHICommandList& RHICmdList, const FSceneView& View, const FProjectedShadowInfo* ShadowInfo, TRefCountPtr<IPooledRenderTarget>& ShadowFactorsTextureValue)
+	void SetParameters(FRHICommandList& RHICmdList, const FSceneView& View, const FProjectedShadowInfo* ShadowInfo, const FIntRect& ScissorRect, TRefCountPtr<IPooledRenderTarget>& ShadowFactorsTextureValue)
 	{
 		const FPixelShaderRHIParamRef ShaderRHI = GetPixelShader();
 
@@ -143,7 +443,10 @@ public:
 		DeferredParameters.Set(RHICmdList, ShaderRHI, View);
 
 		SetTextureParameter(RHICmdList, ShaderRHI, ShadowFactorsTexture, ShadowFactorsSampler, TStaticSamplerState<SF_Bilinear>::GetRHI(), ShadowFactorsTextureValue->GetRenderTargetItem().ShaderResourceTexture);
+	
+		SetShaderValue(RHICmdList, ShaderRHI, ScissorRectMinAndSize, FIntRect(ScissorRect.Min, ScissorRect.Size()));
 	}
+
 	// FShader interface.
 	virtual bool Serialize(FArchive& Ar)
 	{
@@ -151,6 +454,7 @@ public:
 		Ar << DeferredParameters;
 		Ar << ShadowFactorsTexture;
 		Ar << ShadowFactorsSampler;
+		Ar << ScissorRectMinAndSize;
 		return bShaderHasOutdatedParameters;
 	}
 
@@ -159,6 +463,7 @@ private:
 	FDeferredPixelShaderParameters DeferredParameters;
 	FShaderResourceParameter ShadowFactorsTexture;
 	FShaderResourceParameter ShadowFactorsSampler;
+	FShaderParameter ScissorRectMinAndSize;
 };
 
 IMPLEMENT_SHADER_TYPE(,FDistanceFieldShadowingUpsamplePS,TEXT("DistanceFieldShadowing"),TEXT("DistanceFieldShadowingUpsamplePS"),SF_Pixel);
@@ -170,7 +475,7 @@ int32 FProjectedShadowInfo::UpdateShadowCastingObjectBuffers() const
 	//@todo - scene rendering allocator
 	TArray<FVector4> ObjectBoundsData;
 	TArray<FVector4> ObjectData;
-	TArray<FVector4> ObjectData2;
+	TArray<FVector4> ObjectBoxBounds;
 
 	{
 		QUICK_SCOPE_CYCLE_COUNTER(STAT_GatherObjectData);
@@ -179,10 +484,11 @@ int32 FProjectedShadowInfo::UpdateShadowCastingObjectBuffers() const
 		const int32 NumTexelsOneDimY = GDistanceFieldVolumeTextureAtlas.GetSizeY();
 		const int32 NumTexelsOneDimZ = GDistanceFieldVolumeTextureAtlas.GetSizeZ();
 		const FVector InvTextureDim(1.0f / NumTexelsOneDimX, 1.0f / NumTexelsOneDimY, 1.0f / NumTexelsOneDimZ);
+		const FMatrix WorldToShadow = FTranslationMatrix(PreShadowTranslation) * SubjectAndReceiverMatrix;
 
 		ObjectBoundsData.Empty(SubjectPrimitives.Num());
 		ObjectData.Empty(FDistanceFieldObjectBuffers::ObjectDataStride * SubjectPrimitives.Num());
-		ObjectData2.Empty(FDistanceFieldObjectBuffers::ObjectData2Stride * SubjectPrimitives.Num());
+		ObjectBoxBounds.Empty(FDistanceFieldObjectBuffers::ObjectBoxBoundsStride * SubjectPrimitives.Num());
 
 		for (int32 PrimitiveIndex = 0; PrimitiveIndex < SubjectPrimitives.Num() && NumObjects < MAX_uint16; PrimitiveIndex++)
 		{
@@ -219,11 +525,6 @@ int32 FProjectedShadowInfo::UpdateShadowCastingObjectBuffers() const
 				const FVector LocalPositionExtent = LocalVolumeBounds.GetExtent() / FVector(MaxExtent);
 				const FVector UVScale = FVector(BlockSize) * InvTextureDim;
 				const float VolumeScale = UniformScaleVolumeToWorld.GetMaximumAxisScale();
-				ObjectData2.Add(*(FVector4*)&UniformScaleVolumeToWorld.M[0]);
-				ObjectData2.Add(*(FVector4*)&UniformScaleVolumeToWorld.M[1]);
-				ObjectData2.Add(*(FVector4*)&UniformScaleVolumeToWorld.M[2]);
-				ObjectData2.Add(*(FVector4*)&UniformScaleVolumeToWorld.M[3]);
-
 				const FMatrix WorldToVolume = UniformScaleVolumeToWorld.InverseFast();
 				// WorldToVolume
 				ObjectData.Add(*(FVector4*)&WorldToVolume.M[0]);
@@ -244,8 +545,42 @@ int32 FProjectedShadowInfo::UpdateShadowCastingObjectBuffers() const
 				// Padding
 				ObjectData.Add(FVector4(0, 0, 0, 0));
 
+				{
+					FVector LocalBoundsVertices[8];
+					LocalBoundsVertices[0] = FVector(LocalVolumeBounds.Min.X, LocalVolumeBounds.Min.Y, LocalVolumeBounds.Min.Z);
+					LocalBoundsVertices[1] = FVector(LocalVolumeBounds.Max.X, LocalVolumeBounds.Min.Y, LocalVolumeBounds.Min.Z);
+					LocalBoundsVertices[2] = FVector(LocalVolumeBounds.Min.X, LocalVolumeBounds.Max.Y, LocalVolumeBounds.Min.Z);
+					LocalBoundsVertices[3] = FVector(LocalVolumeBounds.Max.X, LocalVolumeBounds.Max.Y, LocalVolumeBounds.Min.Z);
+					LocalBoundsVertices[4] = FVector(LocalVolumeBounds.Min.X, LocalVolumeBounds.Min.Y, LocalVolumeBounds.Max.Z);
+					LocalBoundsVertices[5] = FVector(LocalVolumeBounds.Max.X, LocalVolumeBounds.Min.Y, LocalVolumeBounds.Max.Z);
+					LocalBoundsVertices[6] = FVector(LocalVolumeBounds.Min.X, LocalVolumeBounds.Max.Y, LocalVolumeBounds.Max.Z);
+					LocalBoundsVertices[7] = FVector(LocalVolumeBounds.Max.X, LocalVolumeBounds.Max.Y, LocalVolumeBounds.Max.Z);
+
+					FVector ViewSpaceBoundsVertices[8];
+					FVector MinViewSpacePosition(WORLD_MAX);
+					FVector MaxViewSpacePosition(-WORLD_MAX);
+
+					for (int32 i = 0; i < ARRAY_COUNT(LocalBoundsVertices); i++)
+					{
+						const FVector ViewSpacePosition = WorldToShadow.TransformPosition(LocalToWorld.TransformPosition(LocalBoundsVertices[i]));
+						MinViewSpacePosition = MinViewSpacePosition.ComponentMin(ViewSpacePosition);
+						MaxViewSpacePosition = MaxViewSpacePosition.ComponentMax(ViewSpacePosition);
+						ViewSpaceBoundsVertices[i] = ViewSpacePosition;
+					}
+
+					FVector ObjectXAxis = (ViewSpaceBoundsVertices[1] - ViewSpaceBoundsVertices[0]) / 2.0f;
+					FVector ObjectYAxis = (ViewSpaceBoundsVertices[2] - ViewSpaceBoundsVertices[0]) / 2.0f;
+					FVector ObjectZAxis = (ViewSpaceBoundsVertices[4] - ViewSpaceBoundsVertices[0]) / 2.0f;
+
+					ObjectBoxBounds.Add(FVector4(MinViewSpacePosition, 0));
+					ObjectBoxBounds.Add(FVector4(MaxViewSpacePosition, 0));
+					ObjectBoxBounds.Add(FVector4(ObjectXAxis / ObjectXAxis.SizeSquared(), 0));
+					ObjectBoxBounds.Add(FVector4(ObjectYAxis / ObjectYAxis.SizeSquared(), 0));
+					ObjectBoxBounds.Add(FVector4(ObjectZAxis / ObjectZAxis.SizeSquared(), 0));
+				}
+
 				checkSlow(ObjectData.Num() % FDistanceFieldObjectBuffers::ObjectDataStride == 0);
-				checkSlow(ObjectData2.Num() % FDistanceFieldObjectBuffers::ObjectData2Stride == 0);
+				checkSlow(ObjectBoxBounds.Num() % FDistanceFieldObjectBuffers::ObjectBoxBoundsStride == 0);
 
 				NumObjects++;
 			}
@@ -256,32 +591,42 @@ int32 FProjectedShadowInfo::UpdateShadowCastingObjectBuffers() const
 	{
 		QUICK_SCOPE_CYCLE_COUNTER(STAT_UpdateObjectBuffers);
 
-		if (NumObjects > GDirectionalLightObjectBuffers.MaxObjects)
+		if (!LightSceneInfo->DistanceFieldObjectBuffers)
+		{
+			LightSceneInfo->DistanceFieldObjectBuffers = new FDistanceFieldObjectBuffers;
+		}
+
+		FDistanceFieldObjectBuffers& ObjectBuffers = *LightSceneInfo->DistanceFieldObjectBuffers;
+
+		if (NumObjects > ObjectBuffers.MaxObjects)
 		{
 			// Allocate with slack
-			GDirectionalLightObjectBuffers.MaxObjects = NumObjects * 5 / 4;
-			GDirectionalLightObjectBuffers.UpdateRHI();
+			ObjectBuffers.bWantBoxBounds = true;
+			ObjectBuffers.bWantVolumeToWorld = false;
+			ObjectBuffers.MaxObjects = NumObjects * 5 / 4;
+			ObjectBuffers.Release();
+			ObjectBuffers.Initialize();
 		}
 
 		checkSlow(NumObjects < MAX_uint16);
 
-		void* LockedBuffer = RHILockVertexBuffer(GDirectionalLightObjectBuffers.ObjectData.Bounds, 0, GDirectionalLightObjectBuffers.ObjectData.Bounds->GetSize(), RLM_WriteOnly);
+		void* LockedBuffer = RHILockVertexBuffer(ObjectBuffers.ObjectData.Bounds, 0, ObjectBuffers.ObjectData.Bounds->GetSize(), RLM_WriteOnly);
 		FPlatformMemory::Memcpy(LockedBuffer, ObjectBoundsData.GetData(), ObjectBoundsData.GetTypeSize() * ObjectBoundsData.Num());
-		RHIUnlockVertexBuffer(GDirectionalLightObjectBuffers.ObjectData.Bounds);
+		RHIUnlockVertexBuffer(ObjectBuffers.ObjectData.Bounds);
 
-		LockedBuffer = RHILockVertexBuffer(GDirectionalLightObjectBuffers.ObjectData.Data, 0, GDirectionalLightObjectBuffers.ObjectData.Data->GetSize(), RLM_WriteOnly);
+		LockedBuffer = RHILockVertexBuffer(ObjectBuffers.ObjectData.Data, 0, ObjectBuffers.ObjectData.Data->GetSize(), RLM_WriteOnly);
 		FPlatformMemory::Memcpy(LockedBuffer, ObjectData.GetData(), ObjectData.GetTypeSize() * ObjectData.Num());
-		RHIUnlockVertexBuffer(GDirectionalLightObjectBuffers.ObjectData.Data);
+		RHIUnlockVertexBuffer(ObjectBuffers.ObjectData.Data);
 
-		LockedBuffer = RHILockVertexBuffer(GDirectionalLightObjectBuffers.ObjectData.Data2, 0, GDirectionalLightObjectBuffers.ObjectData.Data2->GetSize(), RLM_WriteOnly);
-		FPlatformMemory::Memcpy(LockedBuffer, ObjectData2.GetData(), ObjectData2.GetTypeSize() * ObjectData2.Num());
-		RHIUnlockVertexBuffer(GDirectionalLightObjectBuffers.ObjectData.Data2);
+		LockedBuffer = RHILockVertexBuffer(ObjectBuffers.ObjectData.BoxBounds, 0, ObjectBuffers.ObjectData.BoxBounds->GetSize(), RLM_WriteOnly);
+		FPlatformMemory::Memcpy(LockedBuffer, ObjectBoxBounds.GetData(), ObjectBoxBounds.GetTypeSize() * ObjectBoxBounds.Num());
+		RHIUnlockVertexBuffer(ObjectBuffers.ObjectData.BoxBounds);
 	}
 
 	return NumObjects;
 }
 
-void FProjectedShadowInfo::RenderRayTracedDistanceFieldProjection(FRHICommandListImmediate& RHICmdList, int32 ViewIndex, const FViewInfo& View) const
+void FProjectedShadowInfo::RenderRayTracedDistanceFieldProjection(FRHICommandListImmediate& RHICmdList, const FViewInfo& View) const
 {
 	if (GDistanceFieldShadowing
 		&& View.GetFeatureLevel() >= ERHIFeatureLevel::SM5
@@ -299,6 +644,88 @@ void FProjectedShadowInfo::RenderRayTracedDistanceFieldProjection(FRHICommandLis
 
 			if (NumObjects > 0)
 			{
+				// Allocate tile resolution based on world space size
+				//@todo - light space perspective shadow maps would make much better use of the resolution
+				const float LightTiles = FMath::Min(ShadowBounds.W / GShadowWorldTileSize + 1, 256.0f);
+				FIntPoint LightTileDimensions(LightTiles, LightTiles);
+
+				FLightTileIntersectionResources* TileIntersectionResources = NULL;
+
+				GSceneRenderTargets.FinishRenderingLightAttenuation(RHICmdList);
+				SetRenderTarget(RHICmdList, NULL, NULL);
+
+				if (bDirectionalLight && GShadowScatterTileCulling)
+				{
+					if (!LightSceneInfo->TileIntersectionResources || LightSceneInfo->TileIntersectionResources->TileDimensions != LightTileDimensions)
+					{
+						if (LightSceneInfo->TileIntersectionResources)
+						{
+							LightSceneInfo->TileIntersectionResources->Release();
+						}
+						else
+						{
+							LightSceneInfo->TileIntersectionResources = new FLightTileIntersectionResources();
+						}
+
+						LightSceneInfo->TileIntersectionResources->TileDimensions = LightTileDimensions;
+
+						LightSceneInfo->TileIntersectionResources->Initialize();
+					}
+
+					TileIntersectionResources = LightSceneInfo->TileIntersectionResources;
+
+					{
+						SCOPED_DRAW_EVENT(RHICmdList, ClearTiles, DEC_SCENE_ITEMS);
+						TShaderMapRef<FClearTilesCS> ComputeShader(View.ShaderMap);
+
+						uint32 GroupSizeX = FMath::DivideAndRoundUp(LightTileDimensions.X, GDistanceFieldAOTileSizeX);
+						uint32 GroupSizeY = FMath::DivideAndRoundUp(LightTileDimensions.Y, GDistanceFieldAOTileSizeY);
+
+						RHICmdList.SetComputeShader(ComputeShader->GetComputeShader());
+						ComputeShader->SetParameters(RHICmdList, View, FVector2D(LightTileDimensions.X, LightTileDimensions.Y), TileIntersectionResources);
+						DispatchComputeShader(RHICmdList, *ComputeShader, GroupSizeX, GroupSizeY, 1);
+
+						ComputeShader->UnsetParameters(RHICmdList);
+					}
+
+					{
+						SCOPED_DRAW_EVENT(RHICmdList, CullObjects, DEC_SCENE_ITEMS);
+
+						TShaderMapRef<FShadowObjectCullVS> VertexShader(View.ShaderMap);
+						TShaderMapRef<FShadowObjectCullPS> PixelShader(View.ShaderMap);
+
+						TArray<FUnorderedAccessViewRHIParamRef> UAVs;
+						PixelShader->GetUAVs(View, TileIntersectionResources, UAVs);
+						RHICmdList.SetRenderTargets(0, (const FRHIRenderTargetView *)NULL, NULL, UAVs.Num(), UAVs.GetData());
+
+						RHICmdList.SetViewport(0, 0, 0.0f, LightTileDimensions.X, LightTileDimensions.Y, 1.0f);
+
+						// Render backfaces since camera may intersect
+						RHICmdList.SetRasterizerState(View.bReverseCulling ? TStaticRasterizerState<FM_Solid, CM_CW>::GetRHI() : TStaticRasterizerState<FM_Solid, CM_CCW>::GetRHI());
+						RHICmdList.SetDepthStencilState(TStaticDepthStencilState<false, CF_Always>::GetRHI());
+						RHICmdList.SetBlendState(TStaticBlendState<>::GetRHI());
+
+						static FGlobalBoundShaderState BoundShaderState;
+
+						SetGlobalBoundShaderState(RHICmdList, View.GetFeatureLevel(), BoundShaderState, GetVertexDeclarationFVector4(), *VertexShader, *PixelShader);
+
+						VertexShader->SetParameters(RHICmdList, View, FVector2D(LightTileDimensions.X, LightTileDimensions.Y), this, NumObjects);
+						PixelShader->SetParameters(RHICmdList, View, FVector2D(LightTileDimensions.X, LightTileDimensions.Y), NumObjects, this, TileIntersectionResources);
+
+						RHICmdList.SetStreamSource(0, StencilingGeometry::GLowPolyStencilSphereVertexBuffer.VertexBufferRHI, sizeof(FVector4), 0);
+
+						RHICmdList.DrawIndexedPrimitive(
+							StencilingGeometry::GLowPolyStencilSphereIndexBuffer.IndexBufferRHI, 
+							PT_TriangleList, 
+							0,
+							0, 
+							StencilingGeometry::GLowPolyStencilSphereVertexBuffer.GetVertexCount(), 
+							0, 
+							StencilingGeometry::GLowPolyStencilSphereIndexBuffer.GetIndexCount() / 3, 
+							NumObjects);
+					}
+				}
+
 				TRefCountPtr<IPooledRenderTarget> RayTracedShadowsRT;
 
 				{
@@ -307,22 +734,45 @@ void FProjectedShadowInfo::RenderRayTracedDistanceFieldProjection(FRHICommandLis
 					GRenderTargetPool.FindFreeElement(Desc, RayTracedShadowsRT, TEXT("RayTracedShadows"));
 				}
 
+				FIntRect ScissorRect;
+
 				{
-					GSceneRenderTargets.FinishRenderingLightAttenuation(RHICmdList);
-					SetRenderTarget(RHICmdList, NULL, NULL);
+					if (!LightSceneInfo->Proxy->GetScissorRect(ScissorRect, View))
+					{
+						ScissorRect = View.ViewRect;
+					}
 
-					uint32 GroupSizeX = FMath::DivideAndRoundUp(View.ViewRect.Size().X / GAODownsampleFactor, GDistanceFieldAOTileSizeX);
-					uint32 GroupSizeY = FMath::DivideAndRoundUp(View.ViewRect.Size().Y / GAODownsampleFactor, GDistanceFieldAOTileSizeY);
+					uint32 GroupSizeX = FMath::DivideAndRoundUp(ScissorRect.Size().X / GAODownsampleFactor, GDistanceFieldAOTileSizeX);
+					uint32 GroupSizeY = FMath::DivideAndRoundUp(ScissorRect.Size().Y / GAODownsampleFactor, GDistanceFieldAOTileSizeY);
 
 					{
 						SCOPED_DRAW_EVENT(RHICmdList, RayTraceShadows, DEC_SCENE_ITEMS);
-						TShaderMapRef<FDistanceFieldShadowingCS> ComputeShader(View.ShaderMap);
-
-						RHICmdList.SetComputeShader(ComputeShader->GetComputeShader());
-						ComputeShader->SetParameters(RHICmdList, View, this, RayTracedShadowsRT->GetRenderTargetItem(), NumObjects, FVector2D(GroupSizeX, GroupSizeY));
-						DispatchComputeShader(RHICmdList, *ComputeShader, GroupSizeX, GroupSizeY, 1);
-
-						ComputeShader->UnsetParameters(RHICmdList);
+						SetRenderTarget(RHICmdList, NULL, NULL);
+
+						if (bDirectionalLight && GShadowScatterTileCulling)
+						{
+							TShaderMapRef<TDistanceFieldShadowingCS<DFS_DirectionalLightScatterTileCulling> > ComputeShader(View.ShaderMap);
+							RHICmdList.SetComputeShader(ComputeShader->GetComputeShader());
+							ComputeShader->SetParameters(RHICmdList, View, this, RayTracedShadowsRT->GetRenderTargetItem(), NumObjects, FVector2D(GroupSizeX, GroupSizeY), ScissorRect, TileIntersectionResources);
+							DispatchComputeShader(RHICmdList, *ComputeShader, GroupSizeX, GroupSizeY, 1);
+							ComputeShader->UnsetParameters(RHICmdList);
+						}
+						else if (bDirectionalLight)
+						{
+							TShaderMapRef<TDistanceFieldShadowingCS<DFS_DirectionalLightTiledCulling> > ComputeShader(View.ShaderMap);
+							RHICmdList.SetComputeShader(ComputeShader->GetComputeShader());
+							ComputeShader->SetParameters(RHICmdList, View, this, RayTracedShadowsRT->GetRenderTargetItem(), NumObjects, FVector2D(GroupSizeX, GroupSizeY), ScissorRect, TileIntersectionResources);
+							DispatchComputeShader(RHICmdList, *ComputeShader, GroupSizeX, GroupSizeY, 1);
+							ComputeShader->UnsetParameters(RHICmdList);
+						}
+						else
+						{
+							TShaderMapRef<TDistanceFieldShadowingCS<DFS_PointLightTiledCulling> > ComputeShader(View.ShaderMap);
+							RHICmdList.SetComputeShader(ComputeShader->GetComputeShader());
+							ComputeShader->SetParameters(RHICmdList, View, this, RayTracedShadowsRT->GetRenderTargetItem(), NumObjects, FVector2D(GroupSizeX, GroupSizeY), ScissorRect, TileIntersectionResources);
+							DispatchComputeShader(RHICmdList, *ComputeShader, GroupSizeX, GroupSizeY, 1);
+							ComputeShader->UnsetParameters(RHICmdList);
+						}
 					}
 				}
 
@@ -331,31 +781,41 @@ void FProjectedShadowInfo::RenderRayTracedDistanceFieldProjection(FRHICommandLis
 
 					SCOPED_DRAW_EVENT(RHICmdList, Upsample, DEC_SCENE_ITEMS);
 
-					RHICmdList.SetViewport(View.ViewRect.Min.X, View.ViewRect.Min.Y, 0.0f, View.ViewRect.Max.X, View.ViewRect.Max.Y, 1.0f);
+					RHICmdList.SetViewport(ScissorRect.Min.X, ScissorRect.Min.Y, 0.0f, ScissorRect.Max.X, ScissorRect.Max.Y, 1.0f);
 					RHICmdList.SetRasterizerState(TStaticRasterizerState<FM_Solid, CM_None>::GetRHI());
 					RHICmdList.SetDepthStencilState(TStaticDepthStencilState<false, CF_Always>::GetRHI());
 					
-					check(CascadeSettings.FadePlaneLength == 0);
-					// first cascade rendered or old method doesn't require fading (CO_Min is needed to combine multiple shadow passes)
-					// The ray traced cascade should always be first
-					RHICmdList.SetBlendState(TStaticBlendState<CW_RG, BO_Min, BF_One, BF_One>::GetRHI());
+					if (bDirectionalLight)
+					{
+						check(CascadeSettings.FadePlaneLength == 0);
+						// first cascade rendered or old method doesn't require fading (CO_Min is needed to combine multiple shadow passes)
+						// The ray traced cascade should always be first
+						RHICmdList.SetBlendState(TStaticBlendState<CW_RG, BO_Min, BF_One, BF_One>::GetRHI());
+					}
+					else
+					{
+						// use B and A in Light Attenuation
+						// (CO_Min is needed to combine multiple shadow passes)
+						RHICmdList.SetBlendState(TStaticBlendState<CW_BA, BO_Min, BF_One, BF_One, BO_Min, BF_One, BF_One>::GetRHI());
+					}
 
-					TShaderMapRef<FPostProcessVS> VertexShader( View.ShaderMap );
-					TShaderMapRef<FDistanceFieldShadowingUpsamplePS> PixelShader( View.ShaderMap );
+					TShaderMapRef<FPostProcessVS> VertexShader(View.ShaderMap);
+					TShaderMapRef<FDistanceFieldShadowingUpsamplePS> PixelShader(View.ShaderMap);
 
 					static FGlobalBoundShaderState BoundShaderState;
 
 					SetGlobalBoundShaderState(RHICmdList, View.GetFeatureLevel(), BoundShaderState, GFilterVertexDeclaration.VertexDeclarationRHI, *VertexShader, *PixelShader);
 
-					PixelShader->SetParameters(RHICmdList, View, this, RayTracedShadowsRT);
+					VertexShader->SetParameters(RHICmdList, View);
+					PixelShader->SetParameters(RHICmdList, View, this, ScissorRect, RayTracedShadowsRT);
 
 					DrawRectangle( 
 						RHICmdList,
 						0, 0, 
-						View.ViewRect.Width(), View.ViewRect.Height(),
-						View.ViewRect.Min.X / GAODownsampleFactor, View.ViewRect.Min.Y / GAODownsampleFactor, 
-						View.ViewRect.Width() / GAODownsampleFactor, View.ViewRect.Height() / GAODownsampleFactor,
-						FIntPoint(View.ViewRect.Width(), View.ViewRect.Height()),
+						ScissorRect.Width(), ScissorRect.Height(),
+						ScissorRect.Min.X / GAODownsampleFactor, ScissorRect.Min.Y / GAODownsampleFactor, 
+						ScissorRect.Width() / GAODownsampleFactor, ScissorRect.Height() / GAODownsampleFactor,
+						FIntPoint(ScissorRect.Width(), ScissorRect.Height()),
 						GetBufferSizeForAO() / GAODownsampleFactor,
 						*VertexShader);
 				}
diff --git a/Engine/Source/Runtime/Renderer/Private/DistanceFieldSurfaceCacheLighting.cpp b/Engine/Source/Runtime/Renderer/Private/DistanceFieldSurfaceCacheLighting.cpp
--- a/Engine/Source/Runtime/Renderer/Private/DistanceFieldSurfaceCacheLighting.cpp
+++ b/Engine/Source/Runtime/Renderer/Private/DistanceFieldSurfaceCacheLighting.cpp
@@ -11,8 +11,8 @@
 #include "PostProcessing.h"
 #include "SceneFilterRendering.h"
 #include "ScreenRendering.h"
-#include "DistanceFieldSurfaceCacheLighting.h"
 #include "DistanceFieldLightingShared.h"
+#include "DistanceFieldSurfaceCacheLighting.h"
 #include "PostProcessAmbientOcclusion.h"
 #include "RHICommandList.h"
 #include "SceneUtils.h"
@@ -250,8 +250,38 @@ TGlobalResource<FCircleVertexBuffer> GCircleVertexBuffer;
 // Must match equivalent shader defines
 int32 FDistanceFieldObjectBuffers::ObjectDataStride = 8;
 int32 FDistanceFieldObjectBuffers::ObjectData2Stride = 4;
+int32 FDistanceFieldObjectBuffers::ObjectBoxBoundsStride = 5;
+
+class FDistanceFieldObjectBufferResource : public FRenderResource
+{
+public:
+	FDistanceFieldObjectBuffers Buffers;
 
-TGlobalResource<FDistanceFieldObjectBuffers> GAOObjectBuffers;
+	virtual void InitDynamicRHI()  override
+	{
+		Buffers.Initialize();
+	}
+
+	virtual void ReleaseDynamicRHI() override
+	{
+		Buffers.Release();
+	}
+};
+
+TGlobalResource<FDistanceFieldObjectBufferResource> GAOObjectBuffers;
+
+void FTileIntersectionResources::InitDynamicRHI()
+{
+	TileConeAxisAndCos.Initialize(sizeof(float)* 4, TileDimensions.X * TileDimensions.Y, PF_A32B32G32R32F, BUF_Static);
+	TileConeDepthRanges.Initialize(sizeof(float)* 4, TileDimensions.X * TileDimensions.Y, PF_A32B32G32R32F, BUF_Static);
+
+	TileHeadDataUnpacked.Initialize(sizeof(uint32), TileDimensions.X * TileDimensions.Y * 4, PF_R32_UINT, BUF_Static);
+	TileHeadData.Initialize(sizeof(uint32)* 4, TileDimensions.X * TileDimensions.Y, PF_R32G32B32A32_UINT, BUF_Static);
+
+	//@todo - handle max exceeded
+	TileArrayData.Initialize(sizeof(uint16), GMaxNumObjectsPerTile * TileDimensions.X * TileDimensions.Y * 3, PF_R16_UINT, BUF_Static);
+	TileArrayNextAllocation.Initialize(sizeof(uint32), 1, PF_R32_UINT, BUF_Static);
+}
 
 void FSceneViewState::DestroyAOTileResources()
 {
@@ -410,7 +440,7 @@ public:
 		const FVertexShaderRHIParamRef ShaderRHI = GetVertexShader();
 		FGlobalShader::SetParameters(RHICmdList, ShaderRHI, View);
 		
-		ObjectParameters.Set(RHICmdList, ShaderRHI, GAOObjectBuffers, NumObjects);
+		ObjectParameters.Set(RHICmdList, ShaderRHI, GAOObjectBuffers.Buffers, NumObjects);
 		AOParameters.Set(RHICmdList, ShaderRHI, Parameters);
 
 		const int32 NumRings = StencilingGeometry::GLowPolyStencilSphereVertexBuffer.GetNumRings();
@@ -576,7 +606,7 @@ public:
 
 		FGlobalShader::SetParameters(RHICmdList, ShaderRHI, View);
 		DeferredParameters.Set(RHICmdList, ShaderRHI, View);
-		ObjectParameters.Set(RHICmdList, ShaderRHI, GAOObjectBuffers, NumObjects);
+		ObjectParameters.Set(RHICmdList, ShaderRHI, GAOObjectBuffers.Buffers, NumObjects);
 		AOParameters.Set(RHICmdList, ShaderRHI, Parameters);
 
 		FTileIntersectionResources* TileIntersectionResources = ((FSceneViewState*)View.State)->AOTileIntersectionResources;
@@ -704,7 +734,7 @@ public:
 		const FPixelShaderRHIParamRef ShaderRHI = GetPixelShader();
 
 		FGlobalShader::SetParameters(RHICmdList, ShaderRHI, View);
-		ObjectParameters.Set(RHICmdList, ShaderRHI, GAOObjectBuffers, NumObjects);
+		ObjectParameters.Set(RHICmdList, ShaderRHI, GAOObjectBuffers.Buffers, NumObjects);
 		AOParameters.Set(RHICmdList, ShaderRHI, Parameters);
 		DeferredParameters.Set(RHICmdList, ShaderRHI, View);
 
@@ -788,7 +818,7 @@ public:
 
 		DistanceFieldNormal.SetTexture(RHICmdList, ShaderRHI, DistanceFieldNormalValue.ShaderResourceTexture, DistanceFieldNormalValue.UAV);
 
-		ObjectParameters.Set(RHICmdList, ShaderRHI, GAOObjectBuffers, NumObjects);
+		ObjectParameters.Set(RHICmdList, ShaderRHI, GAOObjectBuffers.Buffers, NumObjects);
 		AOParameters.Set(RHICmdList, ShaderRHI, Parameters);
 		DeferredParameters.Set(RHICmdList, ShaderRHI, View);
 
@@ -1425,7 +1455,7 @@ public:
 		FComputeShaderRHIParamRef ShaderRHI = GetComputeShader();
 		FGlobalShader::SetParameters(RHICmdList, ShaderRHI, View);
 		DeferredParameters.Set(RHICmdList, ShaderRHI, View);
-		ObjectParameters.Set(RHICmdList, ShaderRHI, GAOObjectBuffers, NumObjects);
+		ObjectParameters.Set(RHICmdList, ShaderRHI, GAOObjectBuffers.Buffers, NumObjects);
 		AOParameters.Set(RHICmdList, ShaderRHI, Parameters);
 		AOLevelParameters.Set(RHICmdList, ShaderRHI, View, DownsampleFactorValue);
 
@@ -2746,26 +2776,27 @@ void UpdateVisibleObjectBuffers(const FScene* Scene, const FDistanceFieldAOParam
 	{
 		QUICK_SCOPE_CYCLE_COUNTER(STAT_UpdateObjectBuffers);
 
-		if (NumObjects > GAOObjectBuffers.MaxObjects)
+		if (NumObjects > GAOObjectBuffers.Buffers.MaxObjects)
 		{
 			// Allocate with slack
-			GAOObjectBuffers.MaxObjects = NumObjects * 5 / 4;
-			GAOObjectBuffers.UpdateRHI();
+			GAOObjectBuffers.Buffers.MaxObjects = NumObjects * 5 / 4;
+			GAOObjectBuffers.Buffers.Release();
+			GAOObjectBuffers.Buffers.Initialize();
 		}
 
 		checkSlow(NumObjects < MAX_uint16);
 
-		void* LockedBuffer = RHILockVertexBuffer(GAOObjectBuffers.ObjectData.Bounds, 0, GAOObjectBuffers.ObjectData.Bounds->GetSize(), RLM_WriteOnly);
+		void* LockedBuffer = RHILockVertexBuffer(GAOObjectBuffers.Buffers.ObjectData.Bounds, 0, GAOObjectBuffers.Buffers.ObjectData.Bounds->GetSize(), RLM_WriteOnly);
 		FPlatformMemory::Memcpy(LockedBuffer, ObjectBoundsData.GetData(), ObjectBoundsData.GetTypeSize() * ObjectBoundsData.Num());
-		RHIUnlockVertexBuffer(GAOObjectBuffers.ObjectData.Bounds);
+		RHIUnlockVertexBuffer(GAOObjectBuffers.Buffers.ObjectData.Bounds);
 
-		LockedBuffer = RHILockVertexBuffer(GAOObjectBuffers.ObjectData.Data, 0, GAOObjectBuffers.ObjectData.Data->GetSize(), RLM_WriteOnly);
+		LockedBuffer = RHILockVertexBuffer(GAOObjectBuffers.Buffers.ObjectData.Data, 0, GAOObjectBuffers.Buffers.ObjectData.Data->GetSize(), RLM_WriteOnly);
 		FPlatformMemory::Memcpy(LockedBuffer, ObjectData.GetData(), ObjectData.GetTypeSize() * ObjectData.Num());
-		RHIUnlockVertexBuffer(GAOObjectBuffers.ObjectData.Data);
+		RHIUnlockVertexBuffer(GAOObjectBuffers.Buffers.ObjectData.Data);
 
-		LockedBuffer = RHILockVertexBuffer(GAOObjectBuffers.ObjectData.Data2, 0, GAOObjectBuffers.ObjectData.Data2->GetSize(), RLM_WriteOnly);
+		LockedBuffer = RHILockVertexBuffer(GAOObjectBuffers.Buffers.ObjectData.Data2, 0, GAOObjectBuffers.Buffers.ObjectData.Data2->GetSize(), RLM_WriteOnly);
 		FPlatformMemory::Memcpy(LockedBuffer, ObjectData2.GetData(), ObjectData2.GetTypeSize() * ObjectData2.Num());
-		RHIUnlockVertexBuffer(GAOObjectBuffers.ObjectData.Data2);
+		RHIUnlockVertexBuffer(GAOObjectBuffers.Buffers.ObjectData.Data2);
 	}
 }
 
@@ -3042,7 +3073,7 @@ public:
 
 		VisualizeMeshDistanceFields.SetTexture(RHICmdList, ShaderRHI, VisualizeMeshDistanceFieldsValue.ShaderResourceTexture, VisualizeMeshDistanceFieldsValue.UAV);
 
-		ObjectParameters.Set(RHICmdList, ShaderRHI, GAOObjectBuffers, NumObjects);
+		ObjectParameters.Set(RHICmdList, ShaderRHI, GAOObjectBuffers.Buffers, NumObjects);
 		AOParameters.Set(RHICmdList, ShaderRHI, Parameters);
 		DeferredParameters.Set(RHICmdList, ShaderRHI, View);
 
diff --git a/Engine/Source/Runtime/Renderer/Private/DistanceFieldSurfaceCacheLighting.h b/Engine/Source/Runtime/Renderer/Private/DistanceFieldSurfaceCacheLighting.h
--- a/Engine/Source/Runtime/Renderer/Private/DistanceFieldSurfaceCacheLighting.h
+++ b/Engine/Source/Runtime/Renderer/Private/DistanceFieldSurfaceCacheLighting.h
@@ -1,7 +1,7 @@
-// Copyright 1998-2013 Epic Games, Inc. All Rights Reserved.
+// Copyright 1998-2014 Epic Games, Inc. All Rights Reserved.
 
 /*=============================================================================
-	
+	DistanceFieldSurfaceCacheLighting.h
 =============================================================================*/
 
 #pragma once
@@ -9,7 +9,21 @@
 const static int32 GAOMaxSupportedLevel = 4;
 //@todo - derive from worst case
 const static uint32 GMaxIrradianceCacheSamples = 100000;
-static const int32 GMaxNumObjectsPerTile = 512;
+/** Number of cone traced directions. */
+const int32 NumConeSampleDirections = 9;
+
+class FDistanceFieldAOParameters
+{
+public:
+	float OcclusionMaxDistance;
+	float Contrast;
+
+	FDistanceFieldAOParameters(float InOcclusionMaxDistance = 600.0f, float InContrast = 0)
+	{
+		Contrast = FMath::Clamp(InContrast, .01f, 2.0f);
+		OcclusionMaxDistance = FMath::Clamp(InOcclusionMaxDistance, 200.0f, 1500.0f);
+	}
+};
 
 /**  */
 class FRefinementLevelResources
@@ -112,18 +126,7 @@ public:
 class FTileIntersectionResources : public FRenderResource
 {
 public:
-	virtual void InitDynamicRHI() override
-	{
-		TileConeAxisAndCos.Initialize(sizeof(float)* 4, TileDimensions.X * TileDimensions.Y, PF_A32B32G32R32F, BUF_Static);
-		TileConeDepthRanges.Initialize(sizeof(float)* 4, TileDimensions.X * TileDimensions.Y, PF_A32B32G32R32F, BUF_Static);
-
-		TileHeadDataUnpacked.Initialize(sizeof(uint32), TileDimensions.X * TileDimensions.Y * 4, PF_R32_UINT, BUF_Static);
-		TileHeadData.Initialize(sizeof(uint32)* 4, TileDimensions.X * TileDimensions.Y, PF_R32G32B32A32_UINT, BUF_Static);
-
-		//@todo - handle max exceeded
-		TileArrayData.Initialize(sizeof(uint16), GMaxNumObjectsPerTile * TileDimensions.X * TileDimensions.Y * 3, PF_R16_UINT, BUF_Static);
-		TileArrayNextAllocation.Initialize(sizeof(uint32), 1, PF_R32_UINT, BUF_Static);
-	}
+	virtual void InitDynamicRHI() override;
 
 	virtual void ReleaseDynamicRHI() override
 	{
@@ -144,3 +147,128 @@ public:
 	FRWBuffer TileArrayData;
 	FRWBuffer TileArrayNextAllocation;
 };
+
+
+/** Generates unit length, stratified and uniformly distributed direction samples in a hemisphere. */
+extern void GenerateStratifiedUniformHemisphereSamples2(int32 NumThetaSteps, int32 NumPhiSteps, FRandomStream& RandomStream, TArray<FVector4>& Samples);
+
+extern void GetSpacedVectors(TArray<TInlineAllocator<9> >& OutVectors);
+
+BEGIN_UNIFORM_BUFFER_STRUCT(FAOSampleData2,)
+	DECLARE_UNIFORM_BUFFER_STRUCT_MEMBER_ARRAY(FVector4,SampleDirections,[NumConeSampleDirections])
+END_UNIFORM_BUFFER_STRUCT(FAOSampleData2)
+
+class FAOParameters
+{
+public:
+	void Bind(const FShaderParameterMap& ParameterMap)
+	{
+		AOMaxDistance.Bind(ParameterMap,TEXT("AOMaxDistance"));
+		AOStepScale.Bind(ParameterMap,TEXT("AOStepScale"));
+		AOStepExponentScale.Bind(ParameterMap,TEXT("AOStepExponentScale"));
+		AOMaxViewDistance.Bind(ParameterMap,TEXT("AOMaxViewDistance"));
+	}
+
+	friend FArchive& operator<<(FArchive& Ar,FAOParameters& Parameters)
+	{
+		Ar << Parameters.AOMaxDistance;
+		Ar << Parameters.AOStepScale;
+		Ar << Parameters.AOStepExponentScale;
+		Ar << Parameters.AOMaxViewDistance;
+		return Ar;
+	}
+
+	template<typename ShaderRHIParamRef>
+	void Set(FRHICommandList& RHICmdList, const ShaderRHIParamRef ShaderRHI, const FDistanceFieldAOParameters& Parameters)
+	{
+		SetShaderValue(RHICmdList, ShaderRHI, AOMaxDistance, Parameters.OcclusionMaxDistance);
+
+		extern float GAOConeHalfAngle;
+		const float AOLargestSampleOffset = Parameters.OcclusionMaxDistance / (1 + FMath::Tan(GAOConeHalfAngle));
+
+		extern float GAOStepExponentScale;
+		extern uint32 GAONumConeSteps;
+		float AOStepScaleValue = AOLargestSampleOffset / FMath::Pow(2.0f, GAOStepExponentScale * (GAONumConeSteps - 1));
+		SetShaderValue(RHICmdList, ShaderRHI, AOStepScale, AOStepScaleValue);
+
+		SetShaderValue(RHICmdList, ShaderRHI, AOStepExponentScale, GAOStepExponentScale);
+
+		extern float GAOMaxViewDistance;
+		SetShaderValue(RHICmdList, ShaderRHI, AOMaxViewDistance, GAOMaxViewDistance);
+	}
+
+private:
+	FShaderParameter AOMaxDistance;
+	FShaderParameter AOStepScale;
+	FShaderParameter AOStepExponentScale;
+	FShaderParameter AOMaxViewDistance;
+};
+
+class FQuadVertexBuffer : public FVertexBuffer
+{
+public:
+
+	virtual void InitRHI() override
+	{
+		// Used as a non-indexed triangle list, so 6 vertices per quad
+		const uint32 Size = 6 * sizeof(FScreenVertex);
+		FRHIResourceCreateInfo CreateInfo;
+		VertexBufferRHI = RHICreateVertexBuffer(Size, BUF_Static, CreateInfo);
+		void* Buffer = RHILockVertexBuffer(VertexBufferRHI, 0, Size, RLM_WriteOnly);
+		FScreenVertex* DestVertex = (FScreenVertex*)Buffer;
+
+		DestVertex[0].Position = FVector2D(0, 0);
+		DestVertex[0].UV = FVector2D(1, 1);
+		DestVertex[1].Position = FVector2D(0, 0);
+		DestVertex[1].UV = FVector2D(1, 0);
+		DestVertex[2].Position = FVector2D(0, 0);
+		DestVertex[2].UV = FVector2D(0, 1);
+		DestVertex[3].Position = FVector2D(0, 0);
+		DestVertex[3].UV = FVector2D(1, 0);
+		DestVertex[4].Position = FVector2D(0, 0);
+		DestVertex[4].UV = FVector2D(0, 0);
+		DestVertex[5].Position = FVector2D(0, 0);
+		DestVertex[5].UV = FVector2D(0, 1);
+
+		RHIUnlockVertexBuffer(VertexBufferRHI);      
+	}
+};
+
+extern TGlobalResource<FQuadVertexBuffer> GQuadVertexBuffer;
+
+class FCircleVertexBuffer : public FVertexBuffer
+{
+public:
+
+	virtual void InitRHI() override
+	{
+		int32 NumSections = 8;
+
+		// Used as a non-indexed triangle list, so 3 vertices per triangle
+		const uint32 Size = 3 * NumSections * sizeof(FScreenVertex);
+		FRHIResourceCreateInfo CreateInfo;
+		VertexBufferRHI = RHICreateVertexBuffer(Size, BUF_Static, CreateInfo);
+		void* Buffer = RHILockVertexBuffer(VertexBufferRHI, 0, Size, RLM_WriteOnly);
+		FScreenVertex* DestVertex = (FScreenVertex*)Buffer;
+
+		for (int32 SectionIndex = 0; SectionIndex < NumSections; SectionIndex++)
+		{
+			float Fraction = SectionIndex / (float)NumSections;
+			float CurrentAngle = Fraction * 2 * PI;
+			float NextAngle = ((SectionIndex + 1) / (float)NumSections) * 2 * PI;
+			FVector2D CurrentPosition(FMath::Cos(CurrentAngle), FMath::Sin(CurrentAngle));
+			FVector2D NextPosition(FMath::Cos(NextAngle), FMath::Sin(NextAngle));
+
+			DestVertex[SectionIndex * 3 + 0].Position = FVector2D(0, 0);
+			DestVertex[SectionIndex * 3 + 0].UV = CurrentPosition;
+			DestVertex[SectionIndex * 3 + 1].Position = FVector2D(0, 0);
+			DestVertex[SectionIndex * 3 + 1].UV = NextPosition;
+			DestVertex[SectionIndex * 3 + 2].Position = FVector2D(0, 0);
+			DestVertex[SectionIndex * 3 + 2].UV = FVector2D(.5f, .5f);
+		}
+
+		RHIUnlockVertexBuffer(VertexBufferRHI);      
+	}
+};
+
+extern TGlobalResource<FCircleVertexBuffer> GCircleVertexBuffer;
diff --git a/Engine/Source/Runtime/Renderer/Private/LightSceneInfo.cpp b/Engine/Source/Runtime/Renderer/Private/LightSceneInfo.cpp
--- a/Engine/Source/Runtime/Renderer/Private/LightSceneInfo.cpp
+++ b/Engine/Source/Runtime/Renderer/Private/LightSceneInfo.cpp
@@ -6,6 +6,7 @@
 
 #include "RendererPrivate.h"
 #include "ScenePrivate.h"
+#include "DistanceFieldLightingShared.h"
 
 void FLightSceneInfoCompact::Init(FLightSceneInfo* InLightSceneInfo)
 {
@@ -29,6 +30,8 @@ FLightSceneInfo::FLightSceneInfo(FLightSceneProxy* InProxy, bool InbVisible)
 	: Proxy(InProxy)
 	, DynamicPrimitiveList(NULL)
 	, Id(INDEX_NONE)
+	, TileIntersectionResources(NULL)
+	, DistanceFieldObjectBuffers(NULL)
 	, bPrecomputedLightingIsValid(InProxy->GetLightComponent()->bPrecomputedLightingIsValid)
 	, bVisible(InbVisible)
 	, bEnableLightShaftBloom(InProxy->GetLightComponent()->bEnableLightShaftBloom)
@@ -190,6 +193,16 @@ void FLightSceneInfo::ReleaseRHI()
 			}
 		}
 	}
+
+	if (TileIntersectionResources)
+	{
+		TileIntersectionResources->Release();
+	}
+
+	if (DistanceFieldObjectBuffers)
+	{
+		DistanceFieldObjectBuffers->Release();
+	}
 }
 
 /** Determines whether two bounding spheres intersect. */
diff --git a/Engine/Source/Runtime/Renderer/Private/LightSceneInfo.h b/Engine/Source/Runtime/Renderer/Private/LightSceneInfo.h
--- a/Engine/Source/Runtime/Renderer/Private/LightSceneInfo.h
+++ b/Engine/Source/Runtime/Renderer/Private/LightSceneInfo.h
@@ -117,6 +117,12 @@ public:
 	 */
 	mutable const FMaterialShaderMap* TranslucentInjectCachedShaderMaps[LightType_MAX][2][2][2];
 
+	/** Tile intersection buffer for distance field shadowing, stored on the light to avoid reallocating each frame. */
+	mutable TScopedPointer<class FLightTileIntersectionResources> TileIntersectionResources;
+
+	/** Culled object buffers for distance field shadowing, stored on the light to avoid reallocating each frame. */
+	mutable TScopedPointer<class FDistanceFieldObjectBuffers> DistanceFieldObjectBuffers;
+
 	/** True if the light is built. */
 	uint32 bPrecomputedLightingIsValid : 1;
 
diff --git a/Engine/Source/Runtime/Renderer/Private/ShadowRendering.cpp b/Engine/Source/Runtime/Renderer/Private/ShadowRendering.cpp
--- a/Engine/Source/Runtime/Renderer/Private/ShadowRendering.cpp
+++ b/Engine/Source/Runtime/Renderer/Private/ShadowRendering.cpp
@@ -1737,7 +1737,7 @@ void FProjectedShadowInfo::RenderProjection(FRHICommandListImmediate& RHICmdList
 
 	if (bRayTracedDistanceFieldShadow)
 	{
-		RenderRayTracedDistanceFieldProjection(RHICmdList, ViewIndex, *View);
+		RenderRayTracedDistanceFieldProjection(RHICmdList, *View);
 		return;
 	}
 
@@ -2271,7 +2271,6 @@ FMatrix FProjectedShadowInfo::GetScreenToShadowMatrix(const FSceneView& View) co
 	const float InvBufferResolutionY = 1.0f / (float)ShadowBufferResolution.Y;
 	const float ShadowResolutionFractionY = 0.5f * (float)ResolutionY * InvBufferResolutionY;
 	// Calculate the matrix to transform a screenspace position into shadow map space
-	// Translucent preshadows start from post projection space since the position was not derived off of the depth buffer like deferred shadows
 	FMatrix ScreenToShadow = 
 		// Z of the position being transformed is actually view space Z, 
 		// Transform it into post projection space by applying the projection matrix,
@@ -2280,10 +2279,7 @@ FMatrix FProjectedShadowInfo::GetScreenToShadowMatrix(const FSceneView& View) co
 			FPlane(1,0,0,0),
 			FPlane(0,1,0,0),
 			FPlane(0,0,View.ViewMatrices.ProjMatrix.M[2][2],1),
-			FPlane(0,0,View.ViewMatrices.ProjMatrix.M[3][2],0)
-		);
-
-	ScreenToShadow *=
+			FPlane(0,0,View.ViewMatrices.ProjMatrix.M[3][2],0)) * 
 		// Transform the post projection space position into translated world space
 		// Translated world space is normal world space translated to the view's origin, 
 		// Which prevents floating point imprecision far from the world origin.
@@ -2571,6 +2567,7 @@ bool FDeferredShadingSceneRenderer::RenderOnePassPointLightShadows(FRHICommandLi
 		{
 			INC_DWORD_STAT(STAT_WholeSceneShadows);
 
+			if (!ProjectedShadowInfo->bRayTracedDistanceFieldShadow)
 			{
 				SCOPED_DRAW_EVENT(RHICmdList, ShadowDepthsFromOpaque, DEC_SCENE_ITEMS);
 				GSceneRenderTargets.BeginRenderingCubeShadowDepth(RHICmdList, ProjectedShadowInfo->ResolutionX);
@@ -2594,14 +2591,21 @@ bool FDeferredShadingSceneRenderer::RenderOnePassPointLightShadows(FRHICommandLi
 
 					// Set the device viewport for the view.
 					RHICmdList.SetViewport(View.ViewRect.Min.X, View.ViewRect.Min.Y, 0.0f, View.ViewRect.Max.X, View.ViewRect.Max.Y, 1.0f);
-
-					ProjectedShadowInfo->RenderOnePassPointLightProjection(RHICmdList, ViewIndex, View);
+					
+					if (ProjectedShadowInfo->bRayTracedDistanceFieldShadow)
+					{
+						ProjectedShadowInfo->RenderRayTracedDistanceFieldProjection(RHICmdList, View);
+					}
+					else
+					{
+						ProjectedShadowInfo->RenderOnePassPointLightProjection(RHICmdList, ViewIndex, View);
+					}
 				}
 			}
 
 			// Don't inject shadowed lighting with whole scene shadows used for previewing a light with static shadows,
 			// Since that would cause a mismatch with the built lighting
-			if (!LightSceneInfo->Proxy->HasStaticShadowing())
+			if (!LightSceneInfo->Proxy->HasStaticShadowing() && !ProjectedShadowInfo->bRayTracedDistanceFieldShadow)
 			{
 				bInjectedTranslucentVolume = true;
 				SCOPED_DRAW_EVENT(RHICmdList, InjectTranslucentVolume, DEC_SCENE_ITEMS);
diff --git a/Engine/Source/Runtime/Renderer/Private/ShadowRendering.h b/Engine/Source/Runtime/Renderer/Private/ShadowRendering.h
--- a/Engine/Source/Runtime/Renderer/Private/ShadowRendering.h
+++ b/Engine/Source/Runtime/Renderer/Private/ShadowRendering.h
@@ -680,7 +680,7 @@ public:
 	void RenderProjection(FRHICommandListImmediate& RHICmdList, int32 ViewIndex, const class FViewInfo* View) const;
 
 	/** Renders ray traced distance field shadows. */
-	void RenderRayTracedDistanceFieldProjection(FRHICommandListImmediate& RHICmdList, int32 ViewIndex, const class FViewInfo& View) const;
+	void RenderRayTracedDistanceFieldProjection(FRHICommandListImmediate& RHICmdList, const class FViewInfo& View) const;
 
 	/** Render one pass point light shadow projections. */
 	void RenderOnePassPointLightProjection(FRHICommandListImmediate& RHICmdList, int32 ViewIndex, const FViewInfo& View) const;
