Fix for linux compile error

diff --git a/Engine/Source/Runtime/Renderer/Private/DistanceFieldShadowing.cpp b/Engine/Source/Runtime/Renderer/Private/DistanceFieldShadowing.cpp
--- a/Engine/Source/Runtime/Renderer/Private/DistanceFieldShadowing.cpp
+++ b/Engine/Source/Runtime/Renderer/Private/DistanceFieldShadowing.cpp
@@ -10,7 +10,6 @@
 #include "ShaderParameters.h"
 #include "PostProcessing.h"
 #include "SceneFilterRendering.h"
-#include "ScreenRendering.h"
 #include "DistanceFieldLightingShared.h"
 #include "DistanceFieldSurfaceCacheLighting.h"
 #include "RHICommandList.h"
diff --git a/Engine/Source/Runtime/Renderer/Private/DistanceFieldSurfaceCacheLighting.cpp b/Engine/Source/Runtime/Renderer/Private/DistanceFieldSurfaceCacheLighting.cpp
--- a/Engine/Source/Runtime/Renderer/Private/DistanceFieldSurfaceCacheLighting.cpp
+++ b/Engine/Source/Runtime/Renderer/Private/DistanceFieldSurfaceCacheLighting.cpp
@@ -10,7 +10,6 @@
 #include "ShaderParameters.h"
 #include "PostProcessing.h"
 #include "SceneFilterRendering.h"
-#include "ScreenRendering.h"
 #include "DistanceFieldLightingShared.h"
 #include "DistanceFieldSurfaceCacheLighting.h"
 #include "PostProcessAmbientOcclusion.h"
diff --git a/Engine/Source/Runtime/Renderer/Private/DistanceFieldSurfaceCacheLighting.h b/Engine/Source/Runtime/Renderer/Private/DistanceFieldSurfaceCacheLighting.h
--- a/Engine/Source/Runtime/Renderer/Private/DistanceFieldSurfaceCacheLighting.h
+++ b/Engine/Source/Runtime/Renderer/Private/DistanceFieldSurfaceCacheLighting.h
@@ -6,6 +6,8 @@
 
 #pragma once
 
+#include "ScreenRendering.h"
+
 const static int32 GAOMaxSupportedLevel = 4;
 //@todo - derive from worst case
 const static uint32 GMaxIrradianceCacheSamples = 100000;
