Mobile Preview - Scene render targets can now allocate forward/deferred targets simultaneously

diff --git a/Engine/Source/Runtime/Renderer/Private/PostProcess/SceneRenderTargets.cpp b/Engine/Source/Runtime/Renderer/Private/PostProcess/SceneRenderTargets.cpp
--- a/Engine/Source/Runtime/Renderer/Private/PostProcess/SceneRenderTargets.cpp
+++ b/Engine/Source/Runtime/Renderer/Private/PostProcess/SceneRenderTargets.cpp
@@ -83,6 +83,28 @@ static TAutoConsoleVariable<int32> CVarGBufferFormat(
 /** The global render targets used for scene rendering. */
 TGlobalResource<FSceneRenderTargets> GSceneRenderTargets;
 
+inline const TCHAR* GetReflectionBrightnessTargetName(uint32 Index)
+{
+	const TCHAR* Names[2] = 
+	{ 
+		TEXT("ReflectionBrightness0"), 
+		TEXT("ReflectionBrightness1") 
+	};
+	check(Index < sizeof(Names));
+	return Names[Index];
+}
+
+inline const TCHAR* GetSceneColorTargetName(FSceneRenderTargets::EShadingPath ShadingPath)
+{
+	const TCHAR* SceneColorNames[(uint32)FSceneRenderTargets::EShadingPath::Num] =
+	{ 
+		TEXT("SceneColorForward"), 
+		TEXT("SceneColorDeferred") 
+	};
+	check((uint32)ShadingPath < sizeof(SceneColorNames));
+	return SceneColorNames[(uint32)ShadingPath];
+}
+
 FIntPoint FSceneRenderTargets::GetSceneRenderTargetSize(const FSceneViewFamily& ViewFamily) const
 {
 	// Don't expose Clamped to the cvar since you need to at least grow to the initial state.
@@ -154,6 +176,7 @@ void FSceneRenderTargets::Allocate(const FSceneViewFamily& ViewFamily)
 
 	// If feature level has changed, release all previously allocated targets to the pool. If feature level has changed but
 	const auto NewFeatureLevel = ViewFamily.Scene->GetFeatureLevel();
+	CurrentShadingPath = ViewFamily.Scene->ShouldUseDeferredRenderer() ? EShadingPath::Deferred : EShadingPath::Forward;
 
 	FIntPoint DesiredBufferSize = GetSceneRenderTargetSize(ViewFamily);
 	check(DesiredBufferSize.X > 0 && DesiredBufferSize.Y > 0);
@@ -206,8 +229,7 @@ void FSceneRenderTargets::Allocate(const FSceneViewFamily& ViewFamily)
 		(CurrentMobile32bpp != Mobile32bpp) ||
 		(CurrentMobileMSAA != MobileMSAA) ||
 		(bCurrentLightPropagationVolume != bLightPropagationVolume) ||
-		(CurrentMinShadowResolution != MinShadowResolution) ||
-		(NewFeatureLevel != CurrentFeatureLevel))
+		(CurrentMinShadowResolution != MinShadowResolution))
 	{
 		CurrentGBufferFormat = GBufferFormat;
 		CurrentSceneColorFormat = SceneColorFormat;
@@ -219,7 +241,6 @@ void FSceneRenderTargets::Allocate(const FSceneViewFamily& ViewFamily)
 		CurrentMobileMSAA = MobileMSAA;
 		CurrentMinShadowResolution = MinShadowResolution;
 		bCurrentLightPropagationVolume = bLightPropagationVolume;
-		CurrentFeatureLevel = NewFeatureLevel;
 		
 		// Reinitialize the render targets for the given size.
 		SetBufferSize(DesiredBufferSize.X, DesiredBufferSize.Y);
@@ -228,6 +249,10 @@ void FSceneRenderTargets::Allocate(const FSceneViewFamily& ViewFamily)
 
 		UpdateRHI();
 	}
+
+	// Do allocation of render targets if they aren't available for the current shading path
+	CurrentFeatureLevel = NewFeatureLevel;
+	AllocateRenderTargets();
 }
 
 /** Clears the GBuffer render targets to default values. */
@@ -305,9 +330,20 @@ int32 FSceneRenderTargets::GetNumGBufferTargets() const
 	return NumGBufferTargets;
 }
 
+inline uint16 GetNumSceneColorMSAASamples(ERHIFeatureLevel::Type InFeatureLevel)
+{
+	//@todo-rco: Fix when OpenGL supports MSAA
+	uint16 NumSamples = GShaderPlatformForFeatureLevel[InFeatureLevel] == SP_OPENGL_ES2_IOS ? 1 : CVarMobileMSAA.GetValueOnRenderThread();
+	if (NumSamples != 1 && NumSamples != 2 && NumSamples != 4)
+	{
+		NumSamples = 1;
+	}
+	return NumSamples;
+}
+
 void FSceneRenderTargets::AllocSceneColor()
 {
-	if(SceneColor)
+	if (GetSceneColorForCurrentShadingPath())
 	{
 		// no work needed
 		return;
@@ -334,20 +370,14 @@ void FSceneRenderTargets::AllocSceneColor()
 
 		if (CurrentFeatureLevel < ERHIFeatureLevel::SM4)
 		{
-			//@todo-rco: Fix when OpenGL supports MSAA
-			uint16 NumSamples = GRHIShaderPlatform == SP_OPENGL_ES2_IOS ? 1 : CVarMobileMSAA.GetValueOnRenderThread();
-			if (NumSamples != 1 && NumSamples != 2 && NumSamples != 4)
-			{
-				NumSamples = 1;
-			}
-			Desc.NumSamples = NumSamples;
+			Desc.NumSamples = GetNumSceneColorMSAASamples(CurrentFeatureLevel);
 		}
 
-		GRenderTargetPool.FindFreeElement(Desc, SceneColor, TEXT("SceneColor"));
+		GRenderTargetPool.FindFreeElement(Desc, GetSceneColorForCurrentShadingPath(), GetSceneColorTargetName(CurrentShadingPath));
 	}
 
 	// otherwise we have a severe problem
-	check(SceneColor);
+	check(GetSceneColorForCurrentShadingPath());
 }
 
 void FSceneRenderTargets::AllocLightAttenuation()
@@ -375,6 +405,13 @@ void FSceneRenderTargets::AllocLightAttenuation()
 	check(LightAttenuation);
 }
 
+TRefCountPtr<IPooledRenderTarget>& FSceneRenderTargets::GetReflectionBrightnessTarget()
+{
+	bool bSupportsR32Float = CurrentFeatureLevel > ERHIFeatureLevel::ES2;
+	int32 ReflectionBrightnessIndex = bSupportsR32Float ? 0 : 1;
+	return ReflectionBrightness[ReflectionBrightnessIndex];
+}
+
 void FSceneRenderTargets::ReleaseGBufferTargets()
 {
 	GBufferResourcesUniformBuffer.SafeRelease();
@@ -499,7 +536,7 @@ void FSceneRenderTargets::AllocGBufferTargets()
 
 const TRefCountPtr<IPooledRenderTarget>& FSceneRenderTargets::GetSceneColor() const
 {
-	if(!SceneColor)
+	if (!GetSceneColorForCurrentShadingPath())
 	{
 		// to avoid log/ensure spam
 		static bool bFirst = true;
@@ -508,18 +545,18 @@ const TRefCountPtr<IPooledRenderTarget>& FSceneRenderTargets::GetSceneColor() co
 			bFirst = false;
 
 			// the first called should be AllocSceneColor(), contact MartinM if that happens
-			ensure(SceneColor);
+			ensure(GetSceneColorForCurrentShadingPath());
 		}
 
 		return GSystemTextures.BlackDummy;
 	}
 
-	return SceneColor;
+	return GetSceneColorForCurrentShadingPath();
 }
 
 TRefCountPtr<IPooledRenderTarget>& FSceneRenderTargets::GetSceneColor()
 {
-	if(!SceneColor)
+	if (!GetSceneColorForCurrentShadingPath())
 	{
 		// to avoid log/ensure spam
 		static bool bFirst = true;
@@ -528,18 +565,19 @@ TRefCountPtr<IPooledRenderTarget>& FSceneRenderTargets::GetSceneColor()
 			bFirst = false;
 
 			// the first called should be AllocSceneColor(), contact MartinM if that happens
-			ensure(SceneColor);
+			ensure(GetSceneColorForCurrentShadingPath());
 		}
 
 		return GSystemTextures.BlackDummy;
 	}
 
-	return SceneColor;
+	return GetSceneColorForCurrentShadingPath();
 }
 
 void FSceneRenderTargets::SetSceneColor(IPooledRenderTarget* In)
 {
-	SceneColor = In;
+	check(CurrentShadingPath < EShadingPath::Num);
+	SceneColor[(int32)CurrentShadingPath] = In;
 }
 
 void FSceneRenderTargets::SetLightAttenuation(IPooledRenderTarget* In)
@@ -631,7 +669,8 @@ void FSceneRenderTargets::FinishRenderingCustomDepth(FRHICommandListImmediate& R
 {
 	SCOPED_DRAW_EVENT(RHICmdList, FinishRenderingCustomDepth, DEC_SCENE_ITEMS);
 
-	RHICmdList.CopyToResolveTarget(SceneColor->GetRenderTargetItem().TargetableTexture, SceneColor->GetRenderTargetItem().ShaderResourceTexture, true, FResolveParams(ResolveRect));
+	auto& CurrentSceneColor = GetSceneColor();
+	RHICmdList.CopyToResolveTarget(CurrentSceneColor->GetRenderTargetItem().TargetableTexture, CurrentSceneColor->GetRenderTargetItem().ShaderResourceTexture, true, FResolveParams(ResolveRect));
 
 	bCustomDepthIsValid = true;
 }
@@ -643,9 +682,11 @@ void FSceneRenderTargets::ResolveSceneColor(FRHICommandList& RHICmdList, const F
 {
 	SCOPED_DRAW_EVENT(RHICmdList, ResolveSceneColor, DEC_SCENE_ITEMS);
 
-	uint32 samples = GSceneRenderTargets.SceneColor->GetDesc().NumSamples;
+	auto& CurrentSceneColor = GetSceneColor();
+
+	uint32 samples = CurrentSceneColor->GetDesc().NumSamples;
 
-	if(samples <= 1 || GRHIShaderPlatform == SP_METAL)
+	if(samples <= 1 || GShaderPlatformForFeatureLevel[CurrentFeatureLevel] == SP_METAL)
 	{
 		RHICmdList.CopyToResolveTarget(GetSceneColorSurface(), GetSceneColorTexture(), true, FResolveParams(ResolveRect));
 	}
@@ -675,7 +716,7 @@ void FSceneRenderTargets::ResolveSceneColor(FRHICommandList& RHICmdList, const F
 			TShaderMapRef<FHdrCustomResolve2xPS> PixelShader(ShaderMap);
 			static FGlobalBoundShaderState BoundShaderState;
 			SetGlobalBoundShaderState(RHICmdList, CurrentFeatureLevel, BoundShaderState, GetVertexDeclarationFVector4(), *VertexShader, *PixelShader);
-			PixelShader->SetParameters(RHICmdList, GSceneRenderTargets.SceneColor->GetRenderTargetItem().TargetableTexture);
+			PixelShader->SetParameters(RHICmdList, CurrentSceneColor->GetRenderTargetItem().TargetableTexture);
 			RHICmdList.DrawPrimitive(PT_TriangleList, 0, 1, 1);
 		}
 		else if(samples == 4)
@@ -683,7 +724,7 @@ void FSceneRenderTargets::ResolveSceneColor(FRHICommandList& RHICmdList, const F
 			TShaderMapRef<FHdrCustomResolve4xPS> PixelShader(ShaderMap);
 			static FGlobalBoundShaderState BoundShaderState;
 			SetGlobalBoundShaderState(RHICmdList, CurrentFeatureLevel, BoundShaderState, GetVertexDeclarationFVector4(), *VertexShader, *PixelShader);
-			PixelShader->SetParameters(RHICmdList, GSceneRenderTargets.SceneColor->GetRenderTargetItem().TargetableTexture);
+			PixelShader->SetParameters(RHICmdList, CurrentSceneColor->GetRenderTargetItem().TargetableTexture);
 			RHICmdList.DrawPrimitive(PT_TriangleList, 0, 1, 1);
 		}
 		else if(samples == 8)
@@ -691,7 +732,7 @@ void FSceneRenderTargets::ResolveSceneColor(FRHICommandList& RHICmdList, const F
 			TShaderMapRef<FHdrCustomResolve8xPS> PixelShader(ShaderMap);
 			static FGlobalBoundShaderState BoundShaderState;
 			SetGlobalBoundShaderState(RHICmdList, CurrentFeatureLevel, BoundShaderState, GetVertexDeclarationFVector4(), *VertexShader, *PixelShader);
-			PixelShader->SetParameters(RHICmdList, GSceneRenderTargets.SceneColor->GetRenderTargetItem().TargetableTexture);
+			PixelShader->SetParameters(RHICmdList, CurrentSceneColor->GetRenderTargetItem().TargetableTexture);
 			RHICmdList.DrawPrimitive(PT_TriangleList, 0, 1, 1);
 		}
 		else
@@ -1042,15 +1083,10 @@ void FSceneRenderTargets::AllocateForwardShadingPathRenderTargets()
 {
 	// on ES2 we don't do on demand allocation of SceneColor yet (in non ES2 it's released in the Tonemapper Process())
 	AllocSceneColor();
+	AllocateCommonDepthTargets();
+	AllocateReflectionTargets();
 
-	EPixelFormat Format = SceneColor->GetDesc().Format;
-
-	{
-		// Create a texture to store the resolved scene depth, and a render-targetable surface to hold the unresolved scene depth.
-		FPooledRenderTargetDesc Desc(FPooledRenderTargetDesc::Create2DDesc(BufferSize, PF_DepthStencil, TexCreate_None, TexCreate_DepthStencilTargetable, false));
-		Desc.NumSamples = SceneColor->GetDesc().NumSamples;
-		GRenderTargetPool.FindFreeElement(Desc, SceneDepthZ, TEXT("SceneDepthZ"));
-	}
+	EPixelFormat Format = GetSceneColor()->GetDesc().Format;
 
 	// For 64-bit ES2 without framebuffer fetch, create extra render target for copy of alpha channel.
 	if((Format == PF_FloatRGBA) && (GSupportsShaderFramebufferFetch == false)) 
@@ -1068,15 +1104,6 @@ void FSceneRenderTargets::AllocateForwardShadingPathRenderTargets()
 	{
 		SceneAlphaCopy = GSystemTextures.MaxFP16Depth;
 	}
-
-	// When targeting DX Feature Level 10, create an auxiliary texture to store the resolved scene depth, and a render-targetable surface to hold the unresolved scene depth.
-	if (!GSupportsDepthFetchDuringDepthTest)
-	{
-		FPooledRenderTargetDesc Desc(FPooledRenderTargetDesc::Create2DDesc(BufferSize, PF_DepthStencil, TexCreate_None, TexCreate_DepthStencilTargetable, false));
-		GRenderTargetPool.FindFreeElement(Desc, AuxiliarySceneDepthZ, TEXT("AuxiliarySceneDepthZ"));
-	}
-
-	AllocateReflectionTargets();	
 }
 
 void FSceneRenderTargets::AllocateForwardShadingShadowDepthTarget(const FIntPoint& ShadowBufferResolution)
@@ -1127,54 +1154,74 @@ void FSceneRenderTargets::AllocateReflectionTargets()
 {	
 	if (GSupportsRenderTargetFormat_PF_FloatRGBA)
 	{
-		extern ENGINE_API int32 GReflectionCaptureSize;
-		const int32 NumReflectionCaptureMips = FMath::CeilLogTwo(GReflectionCaptureSize) + 1;
-		
-		// We write to these cubemap faces individually during filtering
-		uint32 CubeTexFlags = TexCreate_TargetArraySlicesIndependently;
+		// Reflection targets are shared between both forward and deferred shading paths. If we have already allocated for one and are now allocating for the other,
+		// we can skip these targets.
+		bool bSharedReflectionTargetsAllocated = ReflectionColorScratchCubemap[0] != nullptr;
 
+		if (!bSharedReflectionTargetsAllocated)
 		{
-			// Create scratch cubemaps for filtering passes
-			FPooledRenderTargetDesc Desc2(FPooledRenderTargetDesc::CreateCubemapDesc(GReflectionCaptureSize, PF_FloatRGBA, CubeTexFlags, TexCreate_RenderTargetable, false, 1, NumReflectionCaptureMips));
-			GRenderTargetPool.FindFreeElement(Desc2, ReflectionColorScratchCubemap[0], TEXT("ReflectionColorScratchCubemap0"));
-			GRenderTargetPool.FindFreeElement(Desc2, ReflectionColorScratchCubemap[1], TEXT("ReflectionColorScratchCubemap1"));
+			extern ENGINE_API int32 GReflectionCaptureSize;
+			const int32 NumReflectionCaptureMips = FMath::CeilLogTwo(GReflectionCaptureSize) + 1;
 
-			EPixelFormat BrightnessFormat = CurrentFeatureLevel > ERHIFeatureLevel::ES2 ? PF_R32_FLOAT : PF_FloatRGBA;
-			FPooledRenderTargetDesc Desc3(FPooledRenderTargetDesc::Create2DDesc(FIntPoint(1, 1), BrightnessFormat, TexCreate_None, TexCreate_RenderTargetable, false));
-			GRenderTargetPool.FindFreeElement(Desc3, ReflectionBrightness, TEXT("ReflectionBrightness"));
-		}
+			// We write to these cubemap faces individually during filtering
+			uint32 CubeTexFlags = TexCreate_TargetArraySlicesIndependently;
 
-		extern int32 GDiffuseIrradianceCubemapSize;
-		const int32 NumDiffuseIrradianceMips = FMath::CeilLogTwo(GDiffuseIrradianceCubemapSize) + 1;
+			{
+				// Create scratch cubemaps for filtering passes
+				FPooledRenderTargetDesc Desc2(FPooledRenderTargetDesc::CreateCubemapDesc(GReflectionCaptureSize, PF_FloatRGBA, CubeTexFlags, TexCreate_RenderTargetable, false, 1, NumReflectionCaptureMips));
+				GRenderTargetPool.FindFreeElement(Desc2, ReflectionColorScratchCubemap[0], TEXT("ReflectionColorScratchCubemap0"));
+				GRenderTargetPool.FindFreeElement(Desc2, ReflectionColorScratchCubemap[1], TEXT("ReflectionColorScratchCubemap1"));
+			}
 
-		{
-			FPooledRenderTargetDesc Desc2(FPooledRenderTargetDesc::CreateCubemapDesc(GDiffuseIrradianceCubemapSize, PF_FloatRGBA, CubeTexFlags, TexCreate_RenderTargetable, false, 1, NumDiffuseIrradianceMips));
-			GRenderTargetPool.FindFreeElement(Desc2, DiffuseIrradianceScratchCubemap[0], TEXT("DiffuseIrradianceScratchCubemap0"));
-			GRenderTargetPool.FindFreeElement(Desc2, DiffuseIrradianceScratchCubemap[1], TEXT("DiffuseIrradianceScratchCubemap1"));
+			extern int32 GDiffuseIrradianceCubemapSize;
+			const int32 NumDiffuseIrradianceMips = FMath::CeilLogTwo(GDiffuseIrradianceCubemapSize) + 1;
+
+			{
+				FPooledRenderTargetDesc Desc2(FPooledRenderTargetDesc::CreateCubemapDesc(GDiffuseIrradianceCubemapSize, PF_FloatRGBA, CubeTexFlags, TexCreate_RenderTargetable, false, 1, NumDiffuseIrradianceMips));
+				GRenderTargetPool.FindFreeElement(Desc2, DiffuseIrradianceScratchCubemap[0], TEXT("DiffuseIrradianceScratchCubemap0"));
+				GRenderTargetPool.FindFreeElement(Desc2, DiffuseIrradianceScratchCubemap[1], TEXT("DiffuseIrradianceScratchCubemap1"));
+			}
+
+			{
+				FPooledRenderTargetDesc Desc(FPooledRenderTargetDesc::Create2DDesc(FIntPoint(FSHVector3::MaxSHBasis, 1), PF_FloatRGBA, TexCreate_None, TexCreate_RenderTargetable, false));
+				GRenderTargetPool.FindFreeElement(Desc, SkySHIrradianceMap, TEXT("SkySHIrradianceMap"));
+			}
 		}
 
+		auto& ReflectionBrightnessTarget = GetReflectionBrightnessTarget();
+		if (!ReflectionBrightnessTarget)
 		{
-			FPooledRenderTargetDesc Desc(FPooledRenderTargetDesc::Create2DDesc(FIntPoint(FSHVector3::MaxSHBasis, 1), PF_FloatRGBA, TexCreate_None, TexCreate_RenderTargetable, false));
-			GRenderTargetPool.FindFreeElement(Desc, SkySHIrradianceMap, TEXT("SkySHIrradianceMap"));
+			bool bSupportsR32Float = CurrentFeatureLevel > ERHIFeatureLevel::ES2;
+			int32 ReflectionBrightnessIndex = bSupportsR32Float ? 0 : 1;
+			EPixelFormat BrightnessFormat = bSupportsR32Float ? PF_R32_FLOAT : PF_FloatRGBA;
+			FPooledRenderTargetDesc Desc3(FPooledRenderTargetDesc::Create2DDesc(FIntPoint(1, 1), BrightnessFormat, TexCreate_None, TexCreate_RenderTargetable, false));
+			GRenderTargetPool.FindFreeElement(Desc3, ReflectionBrightnessTarget, GetReflectionBrightnessTargetName(ReflectionBrightnessIndex));
 		}
 	}
 }
 
-void FSceneRenderTargets::AllocateDeferredShadingPathRenderTargets()
+void FSceneRenderTargets::AllocateCommonDepthTargets()
 {
+	if (!SceneDepthZ)
 	{
 		// Create a texture to store the resolved scene depth, and a render-targetable surface to hold the unresolved scene depth.
 		FPooledRenderTargetDesc Desc(FPooledRenderTargetDesc::Create2DDesc(BufferSize, PF_DepthStencil, TexCreate_None, TexCreate_DepthStencilTargetable, false));
+		Desc.NumSamples = GetNumSceneColorMSAASamples(CurrentFeatureLevel);
 		Desc.Flags |= TexCreate_FastVRAM;
 		GRenderTargetPool.FindFreeElement(Desc, SceneDepthZ, TEXT("SceneDepthZ"));
 	}
 
 	// When targeting DX Feature Level 10, create an auxiliary texture to store the resolved scene depth, and a render-targetable surface to hold the unresolved scene depth.
-	if(!GSupportsDepthFetchDuringDepthTest)
+	if (!AuxiliarySceneDepthZ && !GSupportsDepthFetchDuringDepthTest)
 	{
 		FPooledRenderTargetDesc Desc(FPooledRenderTargetDesc::Create2DDesc(BufferSize, PF_DepthStencil, TexCreate_None, TexCreate_DepthStencilTargetable, false));
 		GRenderTargetPool.FindFreeElement(Desc, AuxiliarySceneDepthZ, TEXT("AuxiliarySceneDepthZ"));
 	}
+}
+
+void FSceneRenderTargets::AllocateDeferredShadingPathRenderTargets()
+{
+	AllocateCommonDepthTargets();
 
 	// Create a quarter-sized version of the scene depth.
 	{
@@ -1329,16 +1376,17 @@ EPixelFormat FSceneRenderTargets::GetSceneColorFormat() const
 	return SceneColorBufferFormat;
 }
 
-void FSceneRenderTargets::InitDynamicRHI()
+void FSceneRenderTargets::AllocateRenderTargets()
 {
-	if(BufferSize.X > 0 && BufferSize.Y > 0)
+	if (BufferSize.X > 0 && BufferSize.Y > 0 && !AreShadingPathRenderTargetsAllocated(CurrentShadingPath))
 	{
 		// start with a defined state for the scissor rect (D3D11 was returning (0,0,0,0) which caused a clear to not execute correctly)
 		// todo: move this to an earlier place (for dx9 is has to be after device creation which is after window creation)
+		// todo: potentially redundant now? Seems like a strange thing to have here anyway???
 		FRHICommandListImmediate& RHICmdList = FRHICommandListExecutor::GetImmediateCommandList();
 		RHICmdList.SetScissorRect(false, 0, 0, 0, 0);
 
-		if (CurrentFeatureLevel < ERHIFeatureLevel::SM4)
+		if ((EShadingPath)CurrentShadingPath == EShadingPath::Forward)
 		{
 			AllocateForwardShadingPathRenderTargets();
 		}
@@ -1353,7 +1401,11 @@ void FSceneRenderTargets::ReleaseAllTargets()
 {
 	ReleaseGBufferTargets();
 
-	SceneColor.SafeRelease();
+	for (auto i = 0; i < (int32)EShadingPath::Num; ++i)
+	{
+		SceneColor[i].SafeRelease();
+	}
+
 	SceneAlphaCopy.SafeRelease();
 	SceneDepthZ.SafeRelease();
 	AuxiliarySceneDepthZ.SafeRelease();
@@ -1386,7 +1438,8 @@ void FSceneRenderTargets::ReleaseAllTargets()
 		ReflectionColorScratchCubemap[i].SafeRelease();
 	}
 
-	ReflectionBrightness.SafeRelease();
+	ReflectionBrightness[0].SafeRelease();
+	ReflectionBrightness[1].SafeRelease();
 
 	for (int32 i = 0; i < ARRAY_COUNT(DiffuseIrradianceScratchCubemap); i++)
 	{
@@ -1445,17 +1498,17 @@ FIntPoint FSceneRenderTargets::GetTranslucentShadowDepthTextureResolution() cons
 
 const FTextureRHIRef& FSceneRenderTargets::GetSceneColorSurface() const							
 {
-	if(!SceneColor)
+	if (!GetSceneColorForCurrentShadingPath())
 	{
 		return GBlackTexture->TextureRHI;
 	}
 
-	return (const FTextureRHIRef&)SceneColor->GetRenderTargetItem().TargetableTexture; 
+	return (const FTextureRHIRef&)GetSceneColor()->GetRenderTargetItem().TargetableTexture;
 }
 
 const FTextureRHIRef& FSceneRenderTargets::GetSceneColorTexture() const
 {
-	if(!SceneColor)
+	if (!GetSceneColorForCurrentShadingPath())
 	{
 		return GBlackTexture->TextureRHI;
 	}
@@ -1529,6 +1582,26 @@ int32 FSceneRenderTargets::GetCubeShadowDepthZResolution(int32 ShadowIndex) cons
 	return SurfaceSizes[ShadowIndex];
 }
 
+bool FSceneRenderTargets::AreShadingPathRenderTargetsAllocated(EShadingPath InShadingPath) const
+{
+	switch (InShadingPath)
+	{
+	case EShadingPath::Forward:
+		{
+			return (SceneColor[(int32)EShadingPath::Forward] != nullptr);
+		}
+	case EShadingPath::Deferred:
+		{
+			return (ShadowDepthZ != nullptr);
+		}
+	default:
+		{
+			checkNoEntry();
+			return false;
+		}
+	}
+}
+
 /*-----------------------------------------------------------------------------
 FSceneTextureShaderParameters
 -----------------------------------------------------------------------------*/
@@ -1849,4 +1922,4 @@ FArchive& operator<<(FArchive& Ar,FDeferredPixelShaderParameters& Parameters)
 	Ar << Parameters.CustomDepthTextureSampler;
 
 	return Ar;
-}
+}
\ No newline at end of file
diff --git a/Engine/Source/Runtime/Renderer/Private/PostProcess/SceneRenderTargets.h b/Engine/Source/Runtime/Renderer/Private/PostProcess/SceneRenderTargets.h
--- a/Engine/Source/Runtime/Renderer/Private/PostProcess/SceneRenderTargets.h
+++ b/Engine/Source/Runtime/Renderer/Private/PostProcess/SceneRenderTargets.h
@@ -126,10 +126,20 @@ protected:
 		CurrentTranslucencyLightingVolumeDim(64),
 		CurrentMobile32bpp(0),
 		bCurrentLightPropagationVolume(false),
-		CurrentFeatureLevel(ERHIFeatureLevel::Num)
-		{}
+		CurrentFeatureLevel(ERHIFeatureLevel::Num),
+		CurrentShadingPath(EShadingPath::Num)
+		{
+		}
 public:
 
+	enum class EShadingPath
+	{
+		Forward,
+		Deferred,
+
+		Num,
+	};
+
 	/**
 	 * Checks that scene render targets are ready for rendering a view family of the given dimensions.
 	 * If the allocated render targets are too small, they are reallocated.
@@ -235,7 +245,6 @@ public:
 
 
 	// FRenderResource interface.
-	virtual void InitDynamicRHI() override;
 	virtual void ReleaseDynamicRHI() override;
 
 	// Texture Accessors -----------
@@ -409,10 +418,13 @@ public:
 
 	void AllocLightAttenuation();
 
+	TRefCountPtr<IPooledRenderTarget>& GetReflectionBrightnessTarget();
+
 private: // Get...() methods instead of direct access
 
-	// 0 before BeginRenderingSceneColor and after tone mapping, for ES2 it's not released during the frame
-	TRefCountPtr<IPooledRenderTarget> SceneColor;
+	// 0 before BeginRenderingSceneColor and after tone mapping in deferred shading
+	// Permanently allocated for forward shading
+	TRefCountPtr<IPooledRenderTarget> SceneColor[(int32)EShadingPath::Num];
 	// also used as LDR scene color
 	TRefCountPtr<IPooledRenderTarget> LightAttenuation;
 public:
@@ -468,8 +480,11 @@ public:
 	/** Temporary storage during SH irradiance map generation. */
 	TRefCountPtr<IPooledRenderTarget> SkySHIrradianceMap;
 
-	/** Temporary storage, used during reflection capture filtering. */
-	TRefCountPtr<IPooledRenderTarget> ReflectionBrightness;
+	/** Temporary storage, used during reflection capture filtering. 
+	  * 0 - R32 version for > ES2
+	  * 1 - RGBAF version for ES2
+	  */
+	TRefCountPtr<IPooledRenderTarget> ReflectionBrightness[2];
 
 	/** Volume textures used for lighting translucency. */
 	TRefCountPtr<IPooledRenderTarget> TranslucencyLightingVolumeAmbient[NumTranslucentVolumeRenderTargetSets];
@@ -519,8 +534,14 @@ private:
 	/** Allocates render targets for use with the deferred shading path. */
 	void AllocateDeferredShadingPathRenderTargets();
 
+	/** Allocates render targets for use with the current shading path. */
+	void AllocateRenderTargets();
+
 	void AllocateReflectionTargets();
 
+	/** Allocates common depth render targets that are used by both forward and deferred rendering paths */
+	void AllocateCommonDepthTargets();
+
 	/** Determine the appropriate render target dimensions. */
 	FIntPoint GetSceneRenderTargetSize(const FSceneViewFamily & ViewFamily) const;
 
@@ -534,6 +555,16 @@ private:
 
 	EPixelFormat GetSceneColorFormat() const;
 
+	/** Get the current scene color target based on our current shading path. Will return a null ptr if there is no valid scene color target  */
+	const TRefCountPtr<IPooledRenderTarget>& GetSceneColorForCurrentShadingPath() const { check(CurrentShadingPath < EShadingPath::Num); return SceneColor[(int32)CurrentShadingPath]; }
+	TRefCountPtr<IPooledRenderTarget>& GetSceneColorForCurrentShadingPath() { check(CurrentShadingPath < EShadingPath::Num); return SceneColor[(int32)CurrentShadingPath]; }
+
+	/** Determine whether the render targets for a particular shading path have been allocated */
+	bool AreShadingPathRenderTargetsAllocated(EShadingPath InShadingPath) const;
+
+	/** Determine whether the render targets for any shading path have been allocated */
+	bool AreAnyShadingPathRenderTargetsAllocated() const { return AreShadingPathRenderTargetsAllocated(EShadingPath::Deferred) || AreShadingPathRenderTargetsAllocated(EShadingPath::Forward); }
+
 private:
 	/** Uniform buffer containing GBuffer resources. */
 	FUniformBufferRHIRef GBufferResourcesUniformBuffer;
@@ -565,6 +596,8 @@ private:
 	bool bCurrentLightPropagationVolume;
 	/** Feature level we were initialized for */
 	ERHIFeatureLevel::Type CurrentFeatureLevel;
+	/** Shading path that we are currently drawing through. Set when calling Allocate at the start of a scene render. */
+	EShadingPath CurrentShadingPath;
 };
 
 /** The global render targets used for scene rendering. */
diff --git a/Engine/Source/Runtime/Renderer/Private/ReflectionEnvironmentCapture.cpp b/Engine/Source/Runtime/Renderer/Private/ReflectionEnvironmentCapture.cpp
--- a/Engine/Source/Runtime/Renderer/Private/ReflectionEnvironmentCapture.cpp
+++ b/Engine/Source/Runtime/Renderer/Private/ReflectionEnvironmentCapture.cpp
@@ -163,7 +163,7 @@ public:
 			AverageBrightnessTexture, 
 			AverageBrightnessSampler, 
 			TStaticSamplerState<SF_Bilinear, AM_Clamp, AM_Clamp, AM_Clamp>::GetRHI(), 
-			GSceneRenderTargets.ReflectionBrightness->GetRenderTargetItem().ShaderResourceTexture);
+			GSceneRenderTargets.GetReflectionBrightnessTarget()->GetRenderTargetItem().ShaderResourceTexture);
 	}
 
 	virtual bool Serialize(FArchive& Ar)
@@ -275,7 +275,7 @@ IMPLEMENT_SHADER_TYPE(,FComputeBrightnessPS,TEXT("ReflectionEnvironmentShaders")
 /** Computes the average brightness of the given reflection capture and stores it in the scene. */
 void ComputeAverageBrightness(FRHICommandList& RHICmdList, ERHIFeatureLevel::Type FeatureLevel)
 {
-	SetRenderTarget(RHICmdList, GSceneRenderTargets.ReflectionBrightness->GetRenderTargetItem().TargetableTexture, NULL);
+	SetRenderTarget(RHICmdList, GSceneRenderTargets.GetReflectionBrightnessTarget()->GetRenderTargetItem().TargetableTexture, NULL);
 	RHICmdList.SetRasterizerState(TStaticRasterizerState<FM_Solid, CM_None>::GetRHI());
 	RHICmdList.SetDepthStencilState(TStaticDepthStencilState<false, CF_Always>::GetRHI());
 	RHICmdList.SetBlendState(TStaticBlendState<>::GetRHI());
