UE4 - Remove some geometry shaders from Metal

diff --git a/Engine/Source/Runtime/Renderer/Private/ShadowRendering.cpp b/Engine/Source/Runtime/Renderer/Private/ShadowRendering.cpp
--- a/Engine/Source/Runtime/Renderer/Private/ShadowRendering.cpp
+++ b/Engine/Source/Runtime/Renderer/Private/ShadowRendering.cpp
@@ -149,6 +149,11 @@ public:
 
 	static bool ShouldCache(EShaderPlatform Platform,const FMaterial* Material,const FVertexFactoryType* VertexFactoryType)
 	{
+		if (bIsForGeometryShader && !IsFeatureLevelSupported(Platform, ERHIFeatureLevel::SM4))
+		{
+			return false;
+		}
+
 		//Note: This logic needs to stay in sync with OverrideWithDefaultMaterialForShadowDepth!
 		// Compile for special engine materials.
 		if(bRenderReflectiveShadowMap)
