optimized screen space subsurface scattering with minor quality loss
r.SSS.HalfRes allows to go back to full quality

diff --git a/Engine/Shaders/PostProcessSubsurface.usf b/Engine/Shaders/PostProcessSubsurface.usf
--- a/Engine/Shaders/PostProcessSubsurface.usf
+++ b/Engine/Shaders/PostProcessSubsurface.usf
@@ -45,10 +45,19 @@ float GetSubsurfaceStrength(float2 UV)
 	#define SETUP_MODE 0
 #endif
 
+
 #if SETUP_MODE == 0
 	#define SSSS_SPECULAR_CORRECTION 0
+	#define SETUP_HALF_RES 0
 #elif SETUP_MODE == 1
 	#define SSSS_SPECULAR_CORRECTION 1
+	#define SETUP_HALF_RES 0
+#elif SETUP_MODE == 3
+	#define SSSS_SPECULAR_CORRECTION 0
+	#define SETUP_HALF_RES 1
+#elif SETUP_MODE == 4
+	#define SSSS_SPECULAR_CORRECTION 1
+	#define SETUP_HALF_RES 1
 #else
 	// doesn't matter
 	#define SSSS_SPECULAR_CORRECTION 0
@@ -104,6 +113,27 @@ bool InUnitBox(float2 UV)
 	return UV.x >= 0 && UV.y >= 0 && UV.y < 1 && UV.y < 1;
 }
 
+float3 SetupSubsurfaceForOnePixel(float2 UV)
+{
+	float3 Ret = 0;
+
+	FScreenSpaceData ScreenSpaceData = GetScreenSpaceData(UV);
+
+	FLATTEN if (ScreenSpaceData.GBuffer.ShadingModelID >= SHADINGMODELID_SUBSURFACE_PROFILE)
+	{
+		float4 SceneColor4 = Texture2DSample(PostprocessInput0, PostprocessInput0Sampler, UV);
+
+#if SSSS_SPECULAR_CORRECTION == 1
+		// we take out the specular highlights
+		Ret = LightAccumulator_ReconstructDiffuseLighting(SceneColor4);
+#else
+		Ret = SceneColor4.rgb;
+#endif
+	}
+
+	return Ret;
+}
+
 void SetupPS(in float4 UVAndScreenPos : TEXCOORD0, out float4 OutColor : SV_Target0)
 {
 	float2 UV = UVAndScreenPos.xy;
@@ -113,21 +143,16 @@ void SetupPS(in float4 UVAndScreenPos : TEXCOORD0, out float4 OutColor : SV_Targ
 	// if not VisualizeSSS
 	if (SETUP_MODE != 2)
 	{
-		FScreenSpaceData ScreenSpaceData = GetScreenSpaceData(UV);
-
-		FLATTEN if (ScreenSpaceData.GBuffer.ShadingModelID >= SHADINGMODELID_SUBSURFACE_PROFILE)
-		{
-			float4 SceneColor4 = Texture2DSample(PostprocessInput0, PostprocessInput0Sampler, UV);
+#if SETUP_HALF_RES
+		OutColor.rgb += SetupSubsurfaceForOnePixel(UV + float2(-0.5, -0.5f) * PostprocessInput0Size.zw);
+		OutColor.rgb += SetupSubsurfaceForOnePixel(UV + float2( 0.5, -0.5f) * PostprocessInput0Size.zw);
+		OutColor.rgb += SetupSubsurfaceForOnePixel(UV + float2(-0.5,  0.5f) * PostprocessInput0Size.zw);
+		OutColor.rgb += SetupSubsurfaceForOnePixel(UV + float2( 0.5,  0.5f) * PostprocessInput0Size.zw);
 
-#if SSSS_SPECULAR_CORRECTION == 1
-			// we take out the specular highlights
-			float3 Lighting = LightAccumulator_ReconstructDiffuseLighting(SceneColor4);
+		OutColor *= 0.25f;
 #else
-			float3 Lighting = SceneColor4.rgb;
+		OutColor.rgb = SetupSubsurfaceForOnePixel(UV);
 #endif
-
-			OutColor.rgb = Lighting;
-		}
 	}
 	else
 	{
@@ -136,7 +161,7 @@ void SetupPS(in float4 UVAndScreenPos : TEXCOORD0, out float4 OutColor : SV_Targ
 
 		int2 PixelPos = (int2)(UVAndScreenPos.zw * ScreenPosToPixel.xy + ScreenPosToPixel.zw + 0.5f);
 
-		float2 ViewLocalUV = (PixelPos - View.ViewRectMin) / View.ViewSizeAndSceneTexelSize;
+		float2 ViewLocalUV = (PixelPos - View.ViewRectMin.xy) * View.ViewSizeAndSceneTexelSize.zw;
 
 
 		float2 IDAreaLocalUV = ViewLocalUV * 2 - 1.0f;
@@ -204,6 +229,7 @@ void MainPS(float4 UVAndScreenPos : TEXCOORD0, out float4 OutColor : SV_Target0)
 {
 	float2 UV = UVAndScreenPos.xy;
 
+
 	// call into "SeparableSSS.usf"
 
 	// in world units, *0.01 to scale from units(cm) to m
@@ -229,11 +255,15 @@ void MainPS(float4 UVAndScreenPos : TEXCOORD0, out float4 OutColor : SV_Target0)
 #if SSS_METHOD > 0
 	// on second pass we recombine with the scene color
 
-	float4 SceneColor4 = Texture2DSample(PostprocessInput0, PostprocessInput0Sampler, UV);
+	// can be optimized
+	float2 PixelPos = UVAndScreenPos.zw * ScreenPosToPixel.xy + ScreenPosToPixel.zw;
+	float2 UVSceneColor = (PixelPos + 0.5f) * View.ViewSizeAndSceneTexelSize.zw;
+
+	float4 SceneColor4 = Texture2DSample(PostprocessInput0, PostprocessInput0Sampler, UVSceneColor);
 
 	float4 SSSColor = OutColor;
 
-	FScreenSpaceData ScreenSpaceData = GetScreenSpaceData(UV);
+	FScreenSpaceData ScreenSpaceData = GetScreenSpaceData(UVSceneColor);
 
 	SSSColor.a = (ScreenSpaceData.GBuffer.ShadingModelID >= SHADINGMODELID_SUBSURFACE_PROFILE);
 
diff --git a/Engine/Source/Runtime/Renderer/Private/CompositionLighting/CompositionLighting.cpp b/Engine/Source/Runtime/Renderer/Private/CompositionLighting/CompositionLighting.cpp
--- a/Engine/Source/Runtime/Renderer/Private/CompositionLighting/CompositionLighting.cpp
+++ b/Engine/Source/Runtime/Renderer/Private/CompositionLighting/CompositionLighting.cpp
@@ -39,6 +39,14 @@ static TAutoConsoleVariable<float> CVarSSSScale(
 	TEXT(">1: scale scatter radius up (for testing)"),
 	ECVF_RenderThreadSafe);
 
+static TAutoConsoleVariable<int32> CVarSSSHalfRes(
+	TEXT("r.SSS.HalfRes"),
+	2,
+	TEXT("0: full quality (default)\n")
+	TEXT("1: half resolution setup, low quality but faster (currently minor shift from SubsurfaceSetup to SubsurfaceTemp)\n")
+	TEXT("2: also half horizontal during first blurring step, lower quality, even faster"),
+	ECVF_RenderThreadSafe  | ECVF_Scalability);
+
 static bool IsAmbientCubemapPassRequired(FPostprocessContext& Context)
 {
 	FScene* Scene = (FScene*)Context.View.Family->Scene;
@@ -357,10 +365,12 @@ void FCompositionLighting::ProcessLighting(FRHICommandListImmediate& RHICmdList,
 
 			if (View.bScreenSpaceSubsurfacePassNeeded && Radius > 0 && !bSimpleDynamicLighting && View.Family->EngineShowFlags.SubsurfaceScattering)
 			{
-				FRenderingCompositePass* PassSetup = Context.Graph.RegisterPass(new(FMemStack::Get()) FRCPassPostProcessSubsurfaceSetup(false));
+				int32 HalfRes = CVarSSSHalfRes.GetValueOnRenderThread();
+
+				FRenderingCompositePass* PassSetup = Context.Graph.RegisterPass(new(FMemStack::Get()) FRCPassPostProcessSubsurfaceSetup(false, HalfRes > 0));
 				PassSetup->SetInput(ePId_Input0, Context.FinalOutput);
 
-				FRenderingCompositePass* Pass0 = Context.Graph.RegisterPass(new(FMemStack::Get()) FRCPassPostProcessSubsurface(0, Radius));
+				FRenderingCompositePass* Pass0 = Context.Graph.RegisterPass(new(FMemStack::Get()) FRCPassPostProcessSubsurface(0, Radius, HalfRes > 1));
 				Pass0->SetInput(ePId_Input1, PassSetup);
 
 				FRenderingCompositePass* Pass1 = Context.Graph.RegisterPass(new(FMemStack::Get()) FRCPassPostProcessSubsurface(1, Radius));
diff --git a/Engine/Source/Runtime/Renderer/Private/PostProcess/PostProcessSubsurface.cpp b/Engine/Source/Runtime/Renderer/Private/PostProcess/PostProcessSubsurface.cpp
--- a/Engine/Source/Runtime/Renderer/Private/PostProcess/PostProcessSubsurface.cpp
+++ b/Engine/Source/Runtime/Renderer/Private/PostProcess/PostProcessSubsurface.cpp
@@ -25,7 +25,7 @@ static TAutoConsoleVariable<int32> CVarSSSSampleSet(
 
 /**
  * Encapsulates the post processing ambient occlusion pixel shader.
- * @param SetupMode 0:without specular correction, 1: with specular correction, 2:visualize
+ * @param SetupMode 0:without specular correction, 1: with specular correction, 2:visualize, 3: half res without specular correction, 4:half res with specular correction
  */
 template <uint32 SetupMode>
 class FPostProcessSubsurfaceSetupPS : public FGlobalShader
@@ -94,7 +94,7 @@ public:
 #define VARIATION1(A) typedef FPostProcessSubsurfaceSetupPS<A> FPostProcessSubsurfaceSetupPS##A; \
 	IMPLEMENT_SHADER_TYPE2(FPostProcessSubsurfaceSetupPS##A, SF_Pixel);
 
-	VARIATION1(0) VARIATION1(1) VARIATION1(2)
+	VARIATION1(0) VARIATION1(1) VARIATION1(2) VARIATION1(3) VARIATION1(4)
 
 #undef VARIATION1
 
@@ -137,8 +137,9 @@ static bool DoSpecularCorrection()
 	return CVarState && (SceneColorFormat >= 4);
 }
 
-FRCPassPostProcessSubsurfaceSetup::FRCPassPostProcessSubsurfaceSetup(bool bInVisualize)
+FRCPassPostProcessSubsurfaceSetup::FRCPassPostProcessSubsurfaceSetup(bool bInVisualize, bool bInHalfRes)
 	: bVisualize(bInVisualize)
+	, bHalfRes(bInHalfRes)
 {
 	if (bVisualize)
 	{
@@ -148,7 +149,7 @@ FRCPassPostProcessSubsurfaceSetup::FRCPassPostProcessSubsurfaceSetup(bool bInVis
 
 void FRCPassPostProcessSubsurfaceSetup::Process(FRenderingCompositePassContext& Context)
 {
-	SCOPED_DRAW_EVENT(Context.RHICmdList, PostProcessSubsurfaceSetup, DEC_SCENE_ITEMS);
+	SCOPED_DRAW_EVENT(Context.RHICmdList, SubsurfaceSetup, DEC_SCENE_ITEMS);
 
 	const FPooledRenderTargetDesc* InputDesc = GetInputDesc(ePId_Input0);
 
@@ -168,7 +169,13 @@ void FRCPassPostProcessSubsurfaceSetup::Process(FRenderingCompositePassContext&
 	uint32 ScaleFactor = GSceneRenderTargets.GetBufferSizeXY().X / SrcSize.X;
 
 	FIntRect SrcRect = View.ViewRect / ScaleFactor;
-	FIntRect DestRect = SrcRect;
+	FIntRect DestRect = bHalfRes ? FIntRect::DivideAndRoundUp(SrcRect, 2) : SrcRect;
+
+	if(bHalfRes)
+	{
+		// upscale rectangle to not slightly scale
+		SrcRect = DestRect * 2;
+	}
 
 	const FSceneRenderTargetItem& DestRenderTarget = PassOutputs[0].RequestSurface(Context);
 
@@ -191,15 +198,30 @@ void FRCPassPostProcessSubsurfaceSetup::Process(FRenderingCompositePassContext&
 	}
 	else
 	{
-		if (DoSpecularCorrection())
+		// reconstruct specular and add it in final pass
+		bool bSpecularCorrection = DoSpecularCorrection();
+
+		if(bHalfRes)
 		{
-			// with separate specular
-			SetSubsurfaceSetupShader<1>(Context);
+			if(bSpecularCorrection)
+			{
+				SetSubsurfaceSetupShader<4>(Context);
+			}
+			else
+			{
+				SetSubsurfaceSetupShader<3>(Context);
+			}
 		}
 		else
 		{
-			// no separate specular
-			SetSubsurfaceSetupShader<0>(Context);
+			if(bSpecularCorrection)
+			{
+				SetSubsurfaceSetupShader<1>(Context);
+			}
+			else
+			{
+				SetSubsurfaceSetupShader<0>(Context);
+			}
 		}
 	}
 
@@ -279,6 +301,14 @@ FPooledRenderTargetDesc FRCPassPostProcessSubsurfaceSetup::ComputeOutputDesc(EPa
 	// we don't need alpha any more
 	Ret.Format = PF_FloatRGB;
 
+	if(bHalfRes)
+	{
+		Ret.Extent  = FIntPoint::DivideAndRoundUp(Ret.Extent, 2);
+
+		Ret.Extent.X = FMath::Max(1, Ret.Extent.X);
+		Ret.Extent.Y = FMath::Max(1, Ret.Extent.Y);
+	}
+
 	return Ret;
 }
 
@@ -382,9 +412,10 @@ public:
 #undef VARIATION2
 
 
-FRCPassPostProcessSubsurface::FRCPassPostProcessSubsurface(uint32 InPass, float InRadius)
+FRCPassPostProcessSubsurface::FRCPassPostProcessSubsurface(uint32 InPass, float InRadius, bool bInHalfRes)
 	: Radius(InRadius)
 	, Pass(InPass) 
+	, bHalfRes(bInHalfRes) 
 {
 }
 
@@ -436,13 +467,13 @@ void FRCPassPostProcessSubsurface::Process(FRenderingCompositePassContext& Conte
 	const FSceneViewFamily& ViewFamily = *(View.Family);
 
 	FIntPoint SrcSize = InputDesc->Extent;
-	FIntPoint DestSize = SrcSize;
+	FIntPoint DestSize = (Pass == 0) ? PassOutputs[0].RenderTargetDesc.Extent : GSceneRenderTargets.GetBufferSizeXY();
 
 	// e.g. 4 means the input texture is 4x smaller than the buffer size
 	uint32 ScaleFactor = GSceneRenderTargets.GetBufferSizeXY().X / SrcSize.X;
 
 	FIntRect SrcRect = View.ViewRect / ScaleFactor;
-	FIntRect DestRect = SrcRect;
+	FIntRect DestRect = bHalfRes ? SrcRect : View.ViewRect;
 
 	TRefCountPtr<IPooledRenderTarget> NewSceneColor;
 
@@ -476,7 +507,10 @@ void FRCPassPostProcessSubsurface::Process(FRenderingCompositePassContext& Conte
 
 	TShaderMapRef<FPostProcessVS> VertexShader(Context.GetShaderMap());
 
-	SCOPED_DRAW_EVENT(Context.RHICmdList, SubsurfacePass, DEC_SCENE_ITEMS);
+	const TCHAR* PassNames[] = { TEXT("X"), TEXT("Y") };
+	check(Pass < sizeof(PassNames) / sizeof(PassNames[0]));
+
+	SCOPED_DRAW_EVENTF(Context.RHICmdList, SubsurfacePass, DEC_SCENE_ITEMS, PassNames[Pass]);
 
 	uint32 SampleSet = FMath::Clamp(CVarSSSSampleSet.GetValueOnRenderThread(), 0, 2);
 
@@ -533,6 +567,16 @@ FPooledRenderTargetDesc FRCPassPostProcessSubsurface::ComputeOutputDesc(EPassOut
 
 	Ret.Reset();
 	Ret.DebugName = (Pass == 0) ? TEXT("SubsurfaceTemp") : TEXT("SceneColor");
+	// the setup was done in half res but the actual sampling happens in full resolution
+	Ret.Extent = GSceneRenderTargets.GetBufferSizeXY();
+
+	if(bHalfRes)
+	{
+		Ret.Extent  = FIntPoint::DivideAndRoundUp(Ret.Extent, 2);
+
+		Ret.Extent.X = FMath::Max(1, Ret.Extent.X);
+		Ret.Extent.Y = FMath::Max(1, Ret.Extent.Y);
+	}
 
 	return Ret;
 }
\ No newline at end of file
diff --git a/Engine/Source/Runtime/Renderer/Private/PostProcess/PostProcessSubsurface.h b/Engine/Source/Runtime/Renderer/Private/PostProcess/PostProcessSubsurface.h
--- a/Engine/Source/Runtime/Renderer/Private/PostProcess/PostProcessSubsurface.h
+++ b/Engine/Source/Runtime/Renderer/Private/PostProcess/PostProcessSubsurface.h
@@ -16,7 +16,8 @@
 class FRCPassPostProcessSubsurfaceSetup : public TRenderingCompositePassBase<1, 1>
 {
 public:
-	FRCPassPostProcessSubsurfaceSetup(bool bInVisualize);
+	// @param bInHalfRes if the output shoudl be half resolution to scene color
+	FRCPassPostProcessSubsurfaceSetup(bool bInVisualize, bool bInHalfRes = false);
 
 	// interface FRenderingCompositePass ---------
 	virtual void Process(FRenderingCompositePassContext& Context);
@@ -24,6 +25,8 @@ public:
 	virtual void Release() override { delete this; }
 
 	bool bVisualize;
+	// if the output shoudl be half resolution to scene color
+	bool bHalfRes;
 };
 
 
@@ -37,7 +40,7 @@ public:
 	// constructor
 	// @param Pass 0:horizontal/1:vertical
 	// @param InRadius in pixels in the full res image
-	FRCPassPostProcessSubsurface(uint32 Pass, float InRadius);
+	FRCPassPostProcessSubsurface(uint32 Pass, float InRadius, bool bInHalfRes = false);
 
 	// interface FRenderingCompositePass ---------
 
@@ -50,4 +53,6 @@ private:
 	float Radius;
 	// 0:horizontal/1:vertical
 	uint32 Pass;
+	// if the output should be half resolution to scene color
+	bool bHalfRes;
 };
