fixed profiling pass names

diff --git a/Engine/Source/Runtime/Renderer/Private/PostProcess/PostProcessSubsurface.cpp b/Engine/Source/Runtime/Renderer/Private/PostProcess/PostProcessSubsurface.cpp
--- a/Engine/Source/Runtime/Renderer/Private/PostProcess/PostProcessSubsurface.cpp
+++ b/Engine/Source/Runtime/Renderer/Private/PostProcess/PostProcessSubsurface.cpp
@@ -507,7 +507,7 @@ void FRCPassPostProcessSubsurface::Process(FRenderingCompositePassContext& Conte
 
 	TShaderMapRef<FPostProcessVS> VertexShader(Context.GetShaderMap());
 
-	const TCHAR* PassNames[] = { TEXT("X"), TEXT("Y") };
+	const TCHAR* PassNames[] = { TEXT("SubsurfaceX"), TEXT("SubsurfaceY0"), TEXT("SubsurfaceY1") };
 	check(Pass < sizeof(PassNames) / sizeof(PassNames[0]));
 
 	SCOPED_DRAW_EVENTF(Context.RHICmdList, SubsurfacePass, DEC_SCENE_ITEMS, PassNames[Pass]);
