Integrated pull request to fix 'r.Shadow.RadiusThreshold'
Contains changes from PR #319

diff --git a/Engine/Source/Runtime/Renderer/Private/ShadowSetup.cpp b/Engine/Source/Runtime/Renderer/Private/ShadowSetup.cpp
--- a/Engine/Source/Runtime/Renderer/Private/ShadowSetup.cpp
+++ b/Engine/Source/Runtime/Renderer/Private/ShadowSetup.cpp
@@ -9,7 +9,7 @@
 #include "LightPropagationVolume.h"
 
 static float GMinScreenRadiusForShadowCaster = 0.03f;
-static TAutoConsoleVariable<float> CVarMinScreenRadiusForShadowCaster(
+static FAutoConsoleVariableRef CVarMinScreenRadiusForShadowCaster(
 	TEXT("r.Shadow.RadiusThreshold"),
 	GMinScreenRadiusForShadowCaster,
 	TEXT("Cull shadow casters if they are too small, value is the minimal screen space bounding sphere radius\n")
@@ -18,7 +18,7 @@ static TAutoConsoleVariable<float> CVarMinScreenRadiusForShadowCaster(
 	);
 
 static float GMinScreenRadiusForRayTracedCasters = 0.01f;
-static TAutoConsoleVariable<float> CVarMinScreenRadiusForRayTracedCasters(
+static FAutoConsoleVariableRef CVarMinScreenRadiusForRayTracedCasters(
 	TEXT("r.Shadow.RadiusThresholdRayTraced"),
 	GMinScreenRadiusForRayTracedCasters,
 	TEXT("Cull shadow casters in the ray traced shadows if they are too small, value is the minimal screen space bounding sphere radius\n")
@@ -746,7 +746,7 @@ void FProjectedShadowInfo::AddSubjectPrimitive(FPrimitiveSceneInfo* PrimitiveSce
 						FViewInfo& CurrentView = *Views[ViewIndex];
 
 						const float DistanceSquared = ( Bounds.Origin - CurrentView.ShadowViewMatrices.ViewOrigin ).SizeSquared();
-						const bool bDrawShadowDepth = FMath::Square( Bounds.SphereRadius ) > FMath::Square( CVarMinScreenRadiusForShadowCaster.GetValueOnRenderThread() ) * DistanceSquared;
+						const bool bDrawShadowDepth = FMath::Square( Bounds.SphereRadius ) > FMath::Square( GMinScreenRadiusForShadowCaster ) * DistanceSquared;
 						if( !bDrawShadowDepth )
 						{
 							// cull object if it's too small to be considered as shadow caster
