moved log printout of VRam from XboxOne code into main code to also print texture name and format, avoid re-request of a VRam texture in case the limit was laready reached

diff --git a/Engine/Source/Runtime/Renderer/Public/RendererInterface.h b/Engine/Source/Runtime/Renderer/Public/RendererInterface.h
--- a/Engine/Source/Runtime/Renderer/Public/RendererInterface.h
+++ b/Engine/Source/Runtime/Renderer/Public/RendererInterface.h
@@ -147,8 +147,17 @@ public:
 	}
 
 	/** Comparison operator to test if a render target can be reused */
-	bool operator==(const FPooledRenderTargetDesc& rhs) const
+	bool Compare(const FPooledRenderTargetDesc& rhs, bool bExact) const
 	{
+		auto LhsFlags = Flags;
+		auto RhsFlags = rhs.Flags;
+
+		if(!bExact)
+		{
+			LhsFlags &= (~TexCreate_FastVRAM);
+			RhsFlags &= (~TexCreate_FastVRAM);
+		}
+		
 		return Extent == rhs.Extent
 			&& Depth == rhs.Depth
 			&& bIsArray == rhs.bIsArray
@@ -156,7 +165,7 @@ public:
 			&& NumMips == rhs.NumMips
 			&& NumSamples == rhs.NumSamples
 			&& Format == rhs.Format
-			&& Flags == rhs.Flags
+			&& LhsFlags == RhsFlags
 			&& TargetableFlags == rhs.TargetableFlags
 			&& bForceSeparateTargetAndShaderResource == rhs.bForceSeparateTargetAndShaderResource;
 	}
diff --git a/Engine/Source/Runtime/Windows/D3D11RHI/Private/D3D11Util.cpp b/Engine/Source/Runtime/Windows/D3D11RHI/Private/D3D11Util.cpp
--- a/Engine/Source/Runtime/Windows/D3D11RHI/Private/D3D11Util.cpp
+++ b/Engine/Source/Runtime/Windows/D3D11RHI/Private/D3D11Util.cpp
@@ -100,10 +100,13 @@ const TCHAR* GetD3D11TextureFormatString(DXGI_FORMAT TextureFormat)
 		D3DFORMATCASE(DXGI_FORMAT_D32_FLOAT_S8X24_UINT)
 		D3DFORMATCASE(DXGI_FORMAT_R32_FLOAT_X8X24_TYPELESS)
 #endif
+		D3DFORMATCASE(DXGI_FORMAT_R32G8X24_TYPELESS)
 		D3DFORMATCASE(DXGI_FORMAT_D24_UNORM_S8_UINT)
 		D3DFORMATCASE(DXGI_FORMAT_R24_UNORM_X8_TYPELESS)
 		D3DFORMATCASE(DXGI_FORMAT_R32_FLOAT)
 		D3DFORMATCASE(DXGI_FORMAT_R16G16_UINT)
+		D3DFORMATCASE(DXGI_FORMAT_R16G16_UNORM)
+		D3DFORMATCASE(DXGI_FORMAT_R16G16_SNORM)
 		D3DFORMATCASE(DXGI_FORMAT_R16G16_FLOAT)
 		D3DFORMATCASE(DXGI_FORMAT_R32G32_FLOAT)
 		D3DFORMATCASE(DXGI_FORMAT_R10G10B10A2_UNORM)
