added missing file for last checkin

diff --git a/Engine/Source/Runtime/Renderer/Private/PostProcess/RenderTargetPool.cpp b/Engine/Source/Runtime/Renderer/Private/PostProcess/RenderTargetPool.cpp
--- a/Engine/Source/Runtime/Renderer/Private/PostProcess/RenderTargetPool.cpp
+++ b/Engine/Source/Runtime/Renderer/Private/PostProcess/RenderTargetPool.cpp
@@ -120,6 +120,41 @@ FRenderTargetPool::FRenderTargetPool()
 {
 }
 
+static void LogVRamUsage(FPooledRenderTarget& Ref, FTextureRHIParamRef TexRef)
+{
+	if(IsValidRef(TexRef))
+	{
+		FRHIResourceInfo Info;
+
+		RHIGetResourceInfo(TexRef, Info);
+
+		if(Info.VRamAllocation.AllocationSize)
+		{
+			// note we do KB for more readable numbers but this can cause quantization loss
+			UE_LOG(LogShaders, Log, TEXT(" VRamInKB(Start/Size):%d/%d %s '%s'"), 
+				Info.VRamAllocation.AllocationStart / 1024, 
+				(Info.VRamAllocation.AllocationSize + 1023) / 1024,
+				*Ref.GetDesc().GenerateInfoString(),
+				Ref.GetDesc().DebugName);
+		}
+		else
+		{
+			UE_LOG(LogShaders, Log, TEXT(" VRamInKB request failed %s '%s'"), 
+				*Ref.GetDesc().GenerateInfoString(),
+				Ref.GetDesc().DebugName);
+		}
+	}
+}
+
+static void LogVRamUsage(FPooledRenderTarget& Ref)
+{
+	if(Ref.GetDesc().Flags & TexCreate_FastVRAM)
+	{
+		LogVRamUsage(Ref, Ref.GetRenderTargetItem().TargetableTexture);
+		LogVRamUsage(Ref, Ref.GetRenderTargetItem().ShaderResourceTexture);
+	}
+}
+
 bool FRenderTargetPool::FindFreeElement(const FPooledRenderTargetDesc& Desc, TRefCountPtr<IPooledRenderTarget> &Out, const TCHAR* InDebugName)
 {
 	check(IsInRenderingThread());
@@ -135,7 +170,9 @@ bool FRenderTargetPool::FindFreeElement(const FPooledRenderTargetDesc& Desc, TRe
 	{
 		FPooledRenderTarget* Current = (FPooledRenderTarget*)Out.GetReference();
 
-		if(Out->GetDesc() == Desc)
+		const bool bExactMatch = true;
+
+		if(Out->GetDesc().Compare(Desc, bExactMatch))
 		{
 			// we can reuse the same, but the debug name might have changed
 			Current->Desc.DebugName = InDebugName;
@@ -164,15 +201,25 @@ bool FRenderTargetPool::FindFreeElement(const FPooledRenderTargetDesc& Desc, TRe
 	uint32 FoundIndex = -1;
 
 	// try to find a suitable element in the pool
-	for(uint32 i = 0, Num = (uint32)PooledRenderTargets.Num(); i < Num; ++i)
 	{
-		FPooledRenderTarget* Element = PooledRenderTargets[i];
+		uint32 PassCount = (Desc.Flags & TexCreate_FastVRAM) ? 2 : 1;
 
-		if(Element && Element->IsFree() && Element->GetDesc() == Desc)
+		// first we try exact, if that fails we try without TexCreate_FastVRAM
+		for(uint32 Pass = 0; Pass < PassCount; ++Pass)
 		{
-			Found = Element;
-			FoundIndex = i;
-			break;
+			bool bExactMatch = (Pass == 0);
+
+			for(uint32 i = 0, Num = (uint32)PooledRenderTargets.Num(); i < Num; ++i)
+			{
+				FPooledRenderTarget* Element = PooledRenderTargets[i];
+
+				if(Element && Element->IsFree() && Element->GetDesc().Compare(Desc, bExactMatch))
+				{
+					Found = Element;
+					FoundIndex = i;
+					break;
+				}
+			}
 		}
 	}
 
@@ -313,6 +360,11 @@ bool FRenderTargetPool::FindFreeElement(const FPooledRenderTargetDesc& Desc, TRe
 		VerifyAllocationLevel();
 
 		FoundIndex = PooledRenderTargets.Num() - 1;
+
+		// done twice but it doesn't hurt an LogVRamUsage gets the new name this way
+		Found->Desc.DebugName = InDebugName;
+
+		LogVRamUsage(*Found);
 	}
 
 #if !(UE_BUILD_SHIPPING || UE_BUILD_TEST)
