added comment

diff --git a/Engine/Source/Runtime/Renderer/Private/PostProcess/RenderTargetPool.cpp b/Engine/Source/Runtime/Renderer/Private/PostProcess/RenderTargetPool.cpp
--- a/Engine/Source/Runtime/Renderer/Private/PostProcess/RenderTargetPool.cpp
+++ b/Engine/Source/Runtime/Renderer/Private/PostProcess/RenderTargetPool.cpp
@@ -205,6 +205,7 @@ bool FRenderTargetPool::FindFreeElement(const FPooledRenderTargetDesc& Desc, TRe
 		uint32 PassCount = (Desc.Flags & TexCreate_FastVRAM) ? 2 : 1;
 
 		// first we try exact, if that fails we try without TexCreate_FastVRAM
+		// (easily we can run out of VRam, if this search becomes a performance problem we can optimize or we should use less TexCreate_FastVRAM)
 		for(uint32 Pass = 0; Pass < PassCount; ++Pass)
 		{
 			bool bExactMatch = (Pass == 0);
