exposing new PostProcessMaterial location to replace the tonemapper, this is the base for upcoming postprocess changes

diff --git a/Engine/Source/Runtime/Engine/Classes/Engine/BlendableInterface.h b/Engine/Source/Runtime/Engine/Classes/Engine/BlendableInterface.h
--- a/Engine/Source/Runtime/Engine/Classes/Engine/BlendableInterface.h
+++ b/Engine/Source/Runtime/Engine/Classes/Engine/BlendableInterface.h
@@ -10,9 +10,14 @@ class UMaterialInstanceDynamic;
 UENUM()
 enum EBlendableLocation
 {
+	// Input0:former pass color, Input1:SeparateTranslucency
 	BL_AfterTonemapping UMETA(DisplayName="After Tonemapping"),
+	// Input0:former pass color, Input1:SeparateTranslucency
 	BL_BeforeTonemapping UMETA(DisplayName="Before Tonemapping"),
+	// Input0:former pass color, Input1:SeparateTranslucency
 	BL_BeforeTranslucency UMETA(DisplayName="Before Translucency"),
+	// Input0:former pass color, Input1:SeparateTranslucency, Input2: BloomOutput, Input3: EyeAdaptation
+	BL_ReplacingTonemapper UMETA(DisplayName="Replacing the Tonemapper"),
 //	BL_AfterOpaque,
 //	BL_AfterFog,
 //	BL_AfterTranslucency,
@@ -44,18 +49,21 @@ struct FPostProcessMaterialNode
 	FPostProcessMaterialNode()
 		: MaterialInterface(0), bIsMID(false), Location(BL_MAX), Priority(0)
 	{
+		check(!IsValid());
 	}
 
 	// constructor
 	FPostProcessMaterialNode(UMaterialInterface* InMaterialInterface, EBlendableLocation InLocation, int32 InPriority)
 		: MaterialInterface(InMaterialInterface), bIsMID(false), Location(InLocation), Priority(InPriority)
 	{
+		check(IsValid());
 	}
 
 	// constructor
 	FPostProcessMaterialNode(UMaterialInstanceDynamic* InMID, EBlendableLocation InLocation, int32 InPriority)
 		: MaterialInterface((UMaterialInterface*)InMID), bIsMID(true), Location(InLocation), Priority(InPriority)
 	{
+		check(IsValid());
 	}
 
 	UMaterialInterface* GetMaterialInterface() const { return MaterialInterface; }
@@ -87,6 +95,8 @@ struct FPostProcessMaterialNode
 	EBlendableLocation GetLocation() const { return Location; }
 	int32 GetPriority() const { return Priority; }
 
+	bool IsValid() const { return MaterialInterface != 0; }
+
 private:
 	UMaterialInterface* MaterialInterface;
 
diff --git a/Engine/Source/Runtime/Renderer/Private/PostProcess/PostProcessMaterial.h b/Engine/Source/Runtime/Renderer/Private/PostProcess/PostProcessMaterial.h
--- a/Engine/Source/Runtime/Renderer/Private/PostProcess/PostProcessMaterial.h
+++ b/Engine/Source/Runtime/Renderer/Private/PostProcess/PostProcessMaterial.h
@@ -9,9 +9,11 @@
 #include "RenderingCompositionGraph.h"
 
 // derives from TRenderingCompositePassBase<InputCount, OutputCount>
-// ePId_Input0: SceneColor
-// ePId_Input1: SeparateTranslucency
-class FRCPassPostProcessMaterial : public TRenderingCompositePassBase<2,1>
+// ePId_Input0: former pass
+// ePId_Input1: optional, depends on EBlendableLocation
+// ePId_Input2: optional, depends on EBlendableLocation 
+// ePId_Input3: optional, depends on EBlendableLocation
+class FRCPassPostProcessMaterial : public TRenderingCompositePassBase<4,1>
 {
 public:
 	// constructor
diff --git a/Engine/Source/Runtime/Renderer/Private/PostProcess/PostProcessing.cpp b/Engine/Source/Runtime/Renderer/Private/PostProcess/PostProcessing.cpp
--- a/Engine/Source/Runtime/Renderer/Private/PostProcess/PostProcessing.cpp
+++ b/Engine/Source/Runtime/Renderer/Private/PostProcess/PostProcessing.cpp
@@ -502,6 +502,59 @@ FPostProcessMaterialNode* IteratePostProcessMaterialNodes(const FFinalPostProces
 	}
 }
 
+
+static FRenderingCompositePass* AddSinglePostProcessMaterial(FPostprocessContext& Context, EBlendableLocation InLocation)
+{
+	if(!Context.View.Family->EngineShowFlags.PostProcessing || !Context.View.Family->EngineShowFlags.PostProcessMaterial)
+	{
+		return 0;
+	}
+
+	FBlendableEntry* Iterator = 0;
+	FPostProcessMaterialNode PPNode;
+
+	while(FPostProcessMaterialNode* Data = IteratePostProcessMaterialNodes(Context.View.FinalPostProcessSettings, InLocation, Iterator))
+	{
+		check(Data->GetMaterialInterface());
+
+		if(PPNode.IsValid())
+		{
+			FPostProcessMaterialNode::FCompare Dummy;
+
+			// take the one with the highest priority
+			if(!Dummy.operator()(PPNode, *Data))
+			{
+				continue;
+			}
+		}
+
+		PPNode = *Data;
+	}
+
+	if(UMaterialInterface* MaterialInterface = PPNode.GetMaterialInterface())
+	{
+		FMaterialRenderProxy* Proxy = MaterialInterface->GetRenderProxy(false);
+
+		check(Proxy);
+
+		const FMaterial* Material = Proxy->GetMaterial(Context.View.GetFeatureLevel());
+
+		check(Material);
+
+		if(Material->NeedsGBuffer())
+		{
+			// AdjustGBufferRefCount(-1) call is done when the pass gets executed
+			GSceneRenderTargets.AdjustGBufferRefCount(1);
+		}
+
+		FRenderingCompositePass* Node = Context.Graph.RegisterPass(new(FMemStack::Get()) FRCPassPostProcessMaterial(MaterialInterface));
+
+		return Node;
+	}
+
+	return 0;
+}
+
 static void AddPostProcessMaterial(FPostprocessContext& Context, EBlendableLocation InLocation, FRenderingCompositeOutputRef SeparateTranslucency)
 {
 	if(!Context.View.Family->EngineShowFlags.PostProcessing || !Context.View.Family->EngineShowFlags.PostProcessMaterial)
@@ -515,9 +568,6 @@ static void AddPostProcessMaterial(FPostprocessContext& Context, EBlendableLocat
 	FPostProcessMaterialNode PPNodes[MAX_PPMATERIALNODES];
 	uint32 PPNodeCount = 0;
 
-	static const auto CVarDumpFrames = IConsoleManager::Get().FindTConsoleVariableDataInt(TEXT("r.BufferVisualizationDumpFrames"));
-	static const auto CVarDumpFramesAsHDR = IConsoleManager::Get().FindTConsoleVariableDataInt(TEXT("r.BufferVisualizationDumpFramesAsHDR"));
-
 	if(Context.View.Family->EngineShowFlags.VisualizeBuffer)
 	{	
 		// Apply requested material to the full screen
@@ -1055,7 +1105,25 @@ void FPostProcessing::Process(FRHICommandListImmediate& RHICmdList, FViewInfo& V
 
 			if(bAllowTonemapper && FeatureLevel >= ERHIFeatureLevel::SM4)
 			{
-				AddTonemapper(Context, BloomOutputCombined, EyeAdaptation);
+				auto Node = AddSinglePostProcessMaterial(Context, BL_ReplacingTonemapper);
+
+				if(Node)
+				{
+					// a custom tonemapper is provided
+					Node->SetInput(ePId_Input0, Context.FinalOutput);
+
+					// We are binding separate translucency here because the post process SceneTexture node can reference 
+					// the separate translucency buffers through ePId_Input1. 
+					// TODO: Check if material actually uses this texture and only bind if needed.
+					Node->SetInput(ePId_Input1, SeparateTranslucency);
+					Node->SetInput(ePId_Input2, BloomOutputCombined);
+					Node->SetInput(ePId_Input3, EyeAdaptation);
+					Context.FinalOutput = Node;
+				}
+				else
+				{
+					AddTonemapper(Context, BloomOutputCombined, EyeAdaptation);
+				}
 			}
 
 			if(AntiAliasingMethod == AAM_FXAA)
