Distance Field AO optimizations profiling on Nvidia
Objects with no distance field representation (character) no longer receive DFAO, that was causing a huge amount of interpolation cost
Interpolation splat uses depth testing, splat bounding geometry is moved to front, with clamping to avoid near plane clipping
Background is stenciled out from interpolation splat pass
Experimental tiled compute interpolation implementation, disabled as it's slower
Covered over NaNs in interpolation pass which were getting expanded during filtering
Surface cache resources are allocated based on worst case of the current resolution, saves memory and prevents overflow

diff --git a/Engine/Shaders/DistanceFieldSurfaceCacheLighting.usf b/Engine/Shaders/DistanceFieldSurfaceCacheLighting.usf
--- a/Engine/Shaders/DistanceFieldSurfaceCacheLighting.usf
+++ b/Engine/Shaders/DistanceFieldSurfaceCacheLighting.usf
@@ -209,8 +209,9 @@ void ComputeDistanceFieldNormalPS(
 
 #endif
 
+	float DepthSign = GBufferData.HasDistanceFieldRepresentation > 0 ? 1 : -1;
 	// Note: Encoding whether the pixel has a distance field representation in the sign bit of the alpha
-	OutColor = float4(EncodeNormalForAO(WorldNormal), SceneDepth);
+	OutColor = float4(EncodeNormalForAO(WorldNormal), DepthSign * SceneDepth);
 }
 
 RWTexture2D<float4> RWDistanceFieldNormal;
@@ -234,10 +235,23 @@ void ComputeDistanceFieldNormalCS(
 
 #endif
 	
-	float4 OutValue = float4(EncodeNormalForAO(WorldNormal), SceneDepth);
+	float DepthSign = GBufferData.HasDistanceFieldRepresentation > 0 ? 1 : -1;
+	float4 OutValue = float4(EncodeNormalForAO(WorldNormal), DepthSign * SceneDepth);
 	RWDistanceFieldNormal[DispatchThreadId.xy] = OutValue;
 }
 
+void WriteDownsampledDepthPS(
+	in float4 UVAndScreenPos : TEXCOORD0, 
+	out float4 OutColor : SV_Target0,
+	out float OutDepth : SV_Depth)
+{
+	float2 DistanceFieldUVs = UVAndScreenPos.xy;
+	float Depth = Texture2DSampleLevel(DistanceFieldNormalTexture, DistanceFieldNormalSampler, DistanceFieldUVs, 0).w;
+
+	OutColor = 0;
+	OutDepth = ConvertToDeviceZ(Depth);
+}
+
 // For some reason gives innaccurate results at lower resolutions
 #define USE_SCREENVECTOR_WORLD_POSITION FINAL_INTERPOLATION_PASS
 
@@ -277,17 +291,36 @@ void IrradianceCacheSplatVS(
 
 	float ViewToCenterLength = length(PositionAndRadius.xyz - View.ViewOrigin.xyz);
 	float3 NormalizedViewToCenter = (PositionAndRadius.xyz - View.ViewOrigin.xyz) / ViewToCenterLength;
-	PositionAndRadius.w = max(PositionAndRadius.w, ViewToCenterLength * .01f) * InterpolationRadiusScale;
+
+	// Only allow full interpolation expanding for small samples, it won't be noticeable for large samples but adds a lot of splatting cost
+	float EffectiveInterpolationRadiusScale = lerp(InterpolationRadiusScale, 1, max(saturate(PositionAndRadius.w / 80), .2f));
+	PositionAndRadius.w = max(PositionAndRadius.w, ViewToCenterLength * .01f) * EffectiveInterpolationRadiusScale;
+
+	// +1 = behind volume (safe from near plane clipping without depth testing), 
+	// -1 = in front of volume (required for conservative results with depth testing)
+	float BoundingDirection = -1;
+	float OffsetFromCenter = BoundingDirection * PositionAndRadius.w;
+	// Distance from the camera position that NormalizedViewToCenter will intersect the near plane
+	float NearPlaneDistance = View.NearPlane / dot(View.ViewForward, NormalizedViewToCenter);
+
+	// Don't move closer than the near plane to avoid getting clipped
+	// Clamping at the vertex level only works because it's a screen-aligned triangle
+	// Small bias to push just off the near plane
+	if (ViewToCenterLength + OffsetFromCenter < NearPlaneDistance + .001f)
+	{
+		OffsetFromCenter = -ViewToCenterLength + NearPlaneDistance + .001f;
+	}
 
 	// Construct a virtual sample position that won't be near plane clipped and will have a positive z after projection
-	float3 VirtualPosition = PositionAndRadius.xyz + NormalizedViewToCenter * PositionAndRadius.w;
-	// The circle verts were placed directly on the circle, so edges will go inside
-	// Also clipping the edge of the circle a bit can save perf
-	float RadiusSlack = 1.1f;
+	float3 VirtualPosition = PositionAndRadius.xyz + OffsetFromCenter * NormalizedViewToCenter;
+	// Clipping the edge of the circle a bit can save perf, but introduces artifacts in the interpolation
+	float RadiusSlack = 1.0f;
 	// Compute new radius since we moved the center away from the camera, approximate as similar triangles
-	float VirtualRadius = RadiusSlack * PositionAndRadius.w * (PositionAndRadius.w + ViewToCenterLength) / ViewToCenterLength;
+	// This is really incorrect because it's not taking into account perspective correction on a sphere - the widest part in screenspace is not at the world space center
+	float VirtualRadius = RadiusSlack * PositionAndRadius.w * (ViewToCenterLength + OffsetFromCenter) / ViewToCenterLength;
 
 #if !FINAL_INTERPOLATION_PASS
+	// Combat the mysterious bug where samples from last frame do not cover this frame's pixel even with no camera movement, seems to be a precision issue
 	VirtualRadius += HackExpand * ViewToCenterLength * 4;
 #endif
 
@@ -300,6 +333,7 @@ void IrradianceCacheSplatVS(
 	Output.ScreenVector = mul(float4(Output.Position.xy / Output.Position.w, 1, 0), View.ScreenToWorld).xyz;
 #endif
 
+	// Move the vertex over to compensate for the difference between the low res pixel center and the high res pixel center where shading is done
 	Output.Position.xy += NormalizedOffsetToPixelCenter * Output.Position.w;
 
 	Output.PositionRadius = PositionAndRadius; 
@@ -357,33 +391,325 @@ void IrradianceCacheSplatPS(
 		float PointBehindPlaneError = min(max(PlaneDistance * InvMinCosPointBehindPlane, 0.0f), DistanceToVertex / 3.0f);
 
 		float PrecisionScale = .1f;
-		float OverlapWeightScale = 1;
-		float Weight = OverlapWeightScale * saturate(PrecisionScale * (1 - max(DistanceError, max(NormalError, PointBehindPlaneError))));
+		float Weight = saturate(PrecisionScale * (1 - max(DistanceError, max(NormalError, PointBehindPlaneError))));
 		//float Weight = saturate(PrecisionScale * (1 - max(DistanceError, NormalError)));
 		//float Weight = saturate(PrecisionScale * (1 - DistanceError));
 		//float Weight = saturate(PrecisionScale * (1 - NormalError));
 
-		FLATTEN
-		if (DistanceFieldNormalValue.w < 0)
-		{
-			// Don't splat on pixels that don't have a distance field representation
-			// Otherwise these will be spread to neighbors in the gap filling pass
-			Weight = 0;
-		}
-
 		float VisualizePlacement = Distance < .1f * Input.PositionRadius.w;
 		OutColor.rgb = Input.BentNormal * Weight;
 		OutColor.a = Weight;
 	}
+	
+	#define VISUALIZE_SPLAT_OVERDRAW 0
+	#if VISUALIZE_SPLAT_OVERDRAW && FINAL_INTERPOLATION_PASS
+		OutColor.rgb = 0;
+		//OutColor.a = Input.Position.w <= SceneW ? .01f : 0.0f;
+		OutColor.a = .005f;
+	#endif
 
 	#define VISUALIZE_RECORD_POINTS 0
 	#if VISUALIZE_RECORD_POINTS && FINAL_INTERPOLATION_PASS
-		// Visualize cache record points
 		OutColor.rgb = 1;
 		OutColor.a = (DistanceError < .05f) * .1f;
 	#endif
 }
 
+
+#ifndef NUM_INTERPOLATE_LEVELS
+#define NUM_INTERPOLATE_LEVELS 1
+#endif
+
+RWTexture2D<float4> RWIrradianceCacheSplat;
+
+Buffer<uint> LevelScatterDrawParameters[NUM_INTERPOLATE_LEVELS];
+Buffer<float4> LevelIrradianceCachePositionRadius[NUM_INTERPOLATE_LEVELS];
+Buffer<float> LevelIrradianceCacheOccluderRadius[NUM_INTERPOLATE_LEVELS];
+Buffer<float4> LevelIrradianceCacheNormal[NUM_INTERPOLATE_LEVELS];
+Buffer<float4> LevelIrradianceCacheBentNormal[NUM_INTERPOLATE_LEVELS];
+
+#define MAX_RECORDS_PER_TILE 512
+
+groupshared uint IntegerTileMinZ;
+groupshared uint IntegerTileMaxZ;
+groupshared uint IntegerTileMinZ2;
+groupshared uint IntegerTileMaxZ2;
+groupshared uint TileNumSamples;
+groupshared uint TileSampleIndices[MAX_RECORDS_PER_TILE];
+
+[numthreads(THREADGROUP_SIZEX, THREADGROUP_SIZEY, 1)]
+void InterpolateCacheCS(
+	uint3 GroupId : SV_GroupID,
+	uint3 DispatchThreadId : SV_DispatchThreadID,
+    uint3 GroupThreadId : SV_GroupThreadID) 
+{
+	uint ThreadIndex = GroupThreadId.y * THREADGROUP_SIZEX + GroupThreadId.x;
+	uint2 PixelIndex = DispatchThreadId.xy;
+
+	float2 ScreenUV = (PixelIndex * CurrentLevelDownsampleFactor + float2(.5f, .5f)) * View.ViewSizeAndSceneTexelSize.zw;
+	float2 ScreenPosition = (ScreenUV.xy - View.ScreenPositionScaleBias.wz) / View.ScreenPositionScaleBias.xy;
+	float SceneDepth = CalcSceneDepth(ScreenUV);
+
+	float4 HomogeneousWorldPosition = mul(float4(ScreenPosition * SceneDepth, SceneDepth, 1), View.ScreenToWorld);
+	float3 OpaqueWorldPosition = HomogeneousWorldPosition.xyz / HomogeneousWorldPosition.w;
+	float2 BaseLevelScreenUV = (PixelIndex * DownsampleFactorToBaseLevel + float2(.5f, .5f)) * BaseLevelTexelSize;
+	float3 WorldNormal = DecodeNormalForAO(Texture2DSampleLevel(DistanceFieldNormalTexture, DistanceFieldNormalSampler, BaseLevelScreenUV, 0).xyz);
+	
+	// Initialize per-tile variables
+    if (ThreadIndex == 0) 
+	{
+        IntegerTileMinZ = 0x7F7FFFFF;     
+        IntegerTileMaxZ = 0;
+		IntegerTileMinZ2 = 0x7F7FFFFF;  
+		IntegerTileMaxZ2 = 0;
+    }
+
+    GroupMemoryBarrierWithGroupSync();
+    
+    InterlockedMin(IntegerTileMinZ, asuint(SceneDepth));
+    InterlockedMax(IntegerTileMaxZ, asuint(SceneDepth));
+
+    GroupMemoryBarrierWithGroupSync();
+
+    float MinTileZ = asfloat(IntegerTileMinZ);
+    float MaxTileZ = asfloat(IntegerTileMaxZ);
+	
+#define USE_TILE_FRUSTUM 1
+
+#define USE_DUAL_DEPTH_RANGES (1 && !USE_TILE_FRUSTUM)
+#if USE_DUAL_DEPTH_RANGES
+	float HalfZ = .5f * (MinTileZ + MaxTileZ);
+
+	// Compute a second min and max Z, clipped by HalfZ, so that we get two depth bounds per tile
+	// This results in more conservative tile depth bounds and fewer intersections
+	if (SceneDepth >= HalfZ)
+	{
+		InterlockedMin(IntegerTileMinZ2, asuint(SceneDepth));
+	}
+
+	if (SceneDepth <= HalfZ)
+	{
+		InterlockedMax(IntegerTileMaxZ2, asuint(SceneDepth));
+	}
+
+	GroupMemoryBarrierWithGroupSync();
+	
+	float MinTileZ2 = asfloat(IntegerTileMinZ2);
+	float MaxTileZ2 = asfloat(IntegerTileMaxZ2);
+
+#endif
+
+#if USE_TILE_FRUSTUM
+	
+	// Setup tile frustum planes
+	//@todo - use view size not buffer size
+	float2 TileScale = AOBufferSize * rcp(2 * float2(THREADGROUP_SIZEX, THREADGROUP_SIZEY));
+    float2 TileBias = TileScale - GroupId.xy;
+
+    float4 C1 = float4(View.ViewToClip._11 * TileScale.x,	0.0f,								View.ViewToClip._31 * TileScale.x + TileBias.x,	0.0f);
+    float4 C2 = float4(0.0f,								-View.ViewToClip._22 * TileScale.y, View.ViewToClip._32 * TileScale.y + TileBias.y,	0.0f);
+    float4 C4 = float4(0.0f,								0.0f,								1.0f,											0.0f);
+
+    float4 frustumPlanes[6];
+    frustumPlanes[0] = C4 - C1;
+    frustumPlanes[1] = C4 + C1;
+    frustumPlanes[2] = C4 - C2;
+    frustumPlanes[3] = C4 + C2;
+    frustumPlanes[4] = float4(0.0f, 0.0f,  1.0f, -MinTileZ);
+    frustumPlanes[5] = float4(0.0f, 0.0f, -1.0f,  MaxTileZ);
+
+	// Normalize tile frustum planes
+    UNROLL 
+	for (uint i = 0; i < 4; ++i) 
+	{
+        frustumPlanes[i] *= rcp(length(frustumPlanes[i].xyz));
+    }
+
+#else
+
+	float3 ViewTileMin;
+	float3 ViewTileMax;
+
+	float3 ViewTileMin2;
+	float3 ViewTileMax2;
+
+	float2 TanViewFOV = float2(1 / View.ViewToClip[0][0], 1 / View.ViewToClip[1][1]);
+	float2 TileSize = TanViewFOV * 2 / NumGroups;
+	float2 UnitPlaneTileMin = (GroupId.xy * TileSize - TanViewFOV) * float2(1, -1);
+	float2 UnitPlaneTileMax = ((GroupId.xy + 1) * TileSize - TanViewFOV) * float2(1, -1);
+
+	float ExpandRadius = 0;
+
+#if USE_DUAL_DEPTH_RANGES
+	ViewTileMin.xy = min(MinTileZ * UnitPlaneTileMin, MaxTileZ2 * UnitPlaneTileMin) - ExpandRadius;
+	ViewTileMax.xy = max(MinTileZ * UnitPlaneTileMax, MaxTileZ2 * UnitPlaneTileMax) + ExpandRadius;
+	ViewTileMin.z = MinTileZ - ExpandRadius;
+	ViewTileMax.z = MaxTileZ2 + ExpandRadius;
+	ViewTileMin2.xy = min(MinTileZ2 * UnitPlaneTileMin, MaxTileZ * UnitPlaneTileMin) - ExpandRadius;
+	ViewTileMax2.xy = max(MinTileZ2 * UnitPlaneTileMax, MaxTileZ * UnitPlaneTileMax) + ExpandRadius;
+	ViewTileMin2.z = MinTileZ2 - ExpandRadius;
+	ViewTileMax2.z = MaxTileZ + ExpandRadius;
+#else
+	ViewTileMin.xy = min(MinTileZ * UnitPlaneTileMin, MaxTileZ * UnitPlaneTileMin) - ExpandRadius;
+	ViewTileMax.xy = max(MinTileZ * UnitPlaneTileMax, MaxTileZ * UnitPlaneTileMax) + ExpandRadius;
+	ViewTileMin.z = MinTileZ - ExpandRadius;
+	ViewTileMax.z = MaxTileZ + ExpandRadius;
+#endif
+
+#endif
+	
+	float4 LocalAccumulation = 0;
+
+	UNROLL
+	for (uint LevelIndex = 0; LevelIndex < NUM_INTERPOLATE_LEVELS; LevelIndex++)
+	{
+		GroupMemoryBarrierWithGroupSync();
+
+		if (ThreadIndex == 0) 
+		{
+			TileNumSamples = 0;
+		}
+
+		GroupMemoryBarrierWithGroupSync();
+
+		uint NumRecords = LevelScatterDrawParameters[LevelIndex][1];
+
+		LOOP
+		for (uint SampleIndex = ThreadIndex; SampleIndex < NumRecords; SampleIndex += THREADGROUP_TOTALSIZE)
+		{
+			float4 PositionAndRadius = LevelIrradianceCachePositionRadius[LevelIndex].Load(SampleIndex);
+
+			float OccluderRadius = LevelIrradianceCacheOccluderRadius[LevelIndex].Load(SampleIndex);
+			PositionAndRadius.w = min(PositionAndRadius.w, OccluderRadius);
+
+			float ViewToCenterLength = length(PositionAndRadius.xyz - View.ViewOrigin.xyz);
+			PositionAndRadius.w = max(PositionAndRadius.w, ViewToCenterLength * .01f) * InterpolationRadiusScale;
+
+			float3 ViewSpaceSamplePosition = mul(float4(PositionAndRadius.xyz + View.PreViewTranslation.xyz, 1), View.TranslatedWorldToView).xyz;
+
+#if USE_TILE_FRUSTUM
+			// Cull the light against the tile's frustum planes
+			// Note: this has some false positives, a light that is intersecting three different axis frustum planes yet not intersecting the volume of the tile will be treated as intersecting
+			bool bInTile = true;  
+				
+			// Test against the screen x and y oriented planes first
+			UNROLL
+			for (uint i = 0; i < 4; ++i) 
+			{
+				float PlaneDistance = dot(frustumPlanes[i], float4(ViewSpaceSamplePosition, 1.0f));
+				bInTile = bInTile && (PlaneDistance >= -PositionAndRadius.w);
+			}
+
+			BRANCH
+			if (bInTile)
+			{
+				bool bInDepthRange = true;  
+				
+				// Test against the depth range
+				UNROLL 
+				for (uint i = 4; i < 6; ++i) 
+				{
+					float PlaneDistance = dot(frustumPlanes[i], float4(ViewSpaceSamplePosition, 1.0f));
+					bInDepthRange = bInDepthRange && (PlaneDistance >= -PositionAndRadius.w);
+				}
+				 
+				// Add this light to the list of indices if it intersects
+				BRANCH
+				if (bInDepthRange) 
+				{
+					uint ListIndex;
+					InterlockedAdd(TileNumSamples, 1U, ListIndex);
+
+					if (ListIndex < MAX_RECORDS_PER_TILE)
+					{
+						TileSampleIndices[ListIndex] = SampleIndex; 
+					}
+				}
+			}
+#else
+			//@todo - why needed?
+			PositionAndRadius.w *= 2;
+
+			float BoxDistance = ComputeDistanceFromBoxToPoint(ViewTileMin, ViewTileMax, ViewSpaceSamplePosition);
+
+#if USE_DUAL_DEPTH_RANGES
+			float BoxDistance2 = ComputeDistanceFromBoxToPoint(ViewTileMin2, ViewTileMax2, ViewSpaceSamplePosition);
+#endif
+
+			BRANCH
+			if (BoxDistance < PositionAndRadius.w
+#if USE_DUAL_DEPTH_RANGES
+				|| BoxDistance2 < PositionAndRadius.w
+#endif
+				)
+			{
+				uint ListIndex;
+				InterlockedAdd(TileNumSamples, 1U, ListIndex);
+
+				if (ListIndex < MAX_RECORDS_PER_TILE)
+				{
+					TileSampleIndices[ListIndex] = SampleIndex; 
+				}
+			}
+#endif
+		}
+
+		GroupMemoryBarrierWithGroupSync();
+
+		uint LocalNumSamples = TileNumSamples;
+
+		#define VISUALIZE_OVERDRAW 0
+		#if VISUALIZE_OVERDRAW
+			LocalAccumulation += LocalNumSamples * .005f;
+		#else
+
+			LOOP
+			for (uint ListIndex = 0; ListIndex < LocalNumSamples; ListIndex++)
+			{
+				uint SampleIndex = TileSampleIndices[ListIndex];
+				float4 PositionAndRadius = LevelIrradianceCachePositionRadius[LevelIndex].Load(SampleIndex);
+				float OccluderRadius = LevelIrradianceCacheOccluderRadius[LevelIndex].Load(SampleIndex);
+
+				PositionAndRadius.w = min(PositionAndRadius.w, OccluderRadius);
+
+				float ViewToCenterLength = length(PositionAndRadius.xyz - View.ViewOrigin.xyz);
+				PositionAndRadius.w = max(PositionAndRadius.w, ViewToCenterLength * .01f) * InterpolationRadiusScale;
+
+				float Distance = length(OpaqueWorldPosition - PositionAndRadius.xyz);
+				float DistanceError = saturate(Distance / PositionAndRadius.w);
+
+				BRANCH
+				if (DistanceError < 1)
+				{ 
+					float3 RecordNormal = DecodeNormalForAO(LevelIrradianceCacheNormal[LevelIndex].Load(SampleIndex).xyz);
+					float3 RecordBentNormal = DecodeNormalForAO(LevelIrradianceCacheBentNormal[LevelIndex].Load(SampleIndex).xyz);
+
+					float NormalError = InterpolationAngleNormalization * sqrt(saturate(1 - dot(WorldNormal, RecordNormal)));
+
+					// Don't use a lighting record if it's in front of the query point.
+					// Query points behind the lighting record may have nearby occluders that the lighting record does not see.
+					float3 RecordToVertexVector = OpaqueWorldPosition - (PositionAndRadius.xyz - 1 * RecordNormal.xyz);
+					float DistanceToVertex = length(RecordToVertexVector);
+					float PlaneDistance = dot(RecordNormal, RecordToVertexVector) / DistanceToVertex;
+		
+					// Setup an error metric that goes from 0 if the points are coplanar, to 1 if the point being shaded is at the angle corresponding to MinCosPointBehindPlane behind the plane
+					float PointBehindPlaneError = min(max(PlaneDistance * InvMinCosPointBehindPlane, 0.0f), DistanceToVertex / 3.0f);
+
+					float PrecisionScale = .1f;
+					float Weight = saturate(PrecisionScale * (1 - max(DistanceError, max(NormalError, PointBehindPlaneError))));
+					//float Weight = saturate(PrecisionScale * (1 - DistanceError));
+
+
+					LocalAccumulation.rgb += RecordBentNormal * Weight;
+					LocalAccumulation.a += Weight;
+				}
+			}
+		#endif
+	}
+
+	RWIrradianceCacheSplat[PixelIndex.xy] = LocalAccumulation;
+}
+
 Texture2D BentNormalAOTexture3;
 SamplerState BentNormalAOSampler3;
 
@@ -393,7 +719,7 @@ void AOCombinePS2(
 	out float4 OutColor : SV_Target0)
 {
 	float4 DistanceFieldNormalValue = Texture2DSampleLevel(DistanceFieldNormalTexture, DistanceFieldNormalSampler, UVAndScreenPos.xy, 0);
-	float SceneDepth = DistanceFieldNormalValue.w;
+	float SceneDepth = abs(DistanceFieldNormalValue.w);
 
 #define VISUALIZE_ACCUMULATED_WEIGHTS 0
 #if VISUALIZE_ACCUMULATED_WEIGHTS
@@ -405,6 +731,13 @@ void AOCombinePS2(
 
 	float4 IrradianceCacheAccumulation = Texture2DSampleLevel(BentNormalAOTexture3, BentNormalAOSampler3, UVAndScreenPos.xy, 0);
 
+	if (DistanceFieldNormalValue.w < 0)
+	{
+		// Add a low weight completely unshadowed contribution for pixels with no distance field representation
+		float Weight = .001f;
+		IrradianceCacheAccumulation += float4(DistanceFieldNormalValue.xyz * Weight, Weight);
+	}
+
 	if (IrradianceCacheAccumulation.w > 0)
 	{
 		OutColor.rgb = IrradianceCacheAccumulation.xyz / IrradianceCacheAccumulation.w;
@@ -422,6 +755,9 @@ void AOCombinePS2(
 	float FadeAlpha = saturate(.001f * (.95f * AOMaxViewDistance - SceneDepth));
 	OutColor.rgb = normalize(OutColor.rgb) * lerp(1, BentNormalLength, FadeAlpha);
 
+	// Clamp Nan's before they get to the filter
+	OutColor.rgb = clamp(OutColor.rgb, -1, 1);
+
 	FLATTEN
 	if (SceneDepth > AOMaxViewDistance)
 	{
@@ -501,13 +837,13 @@ void UpdateHistoryPS(
 	float4 HistoryValue = Texture2DSampleLevel(HistoryTexture, HistorySampler, OldDistanceFieldUVs, 0);
 	
 	//@todo - this doesn't take camera movement into account
-	float DepthDelta = abs(abs(NewValue.a) - abs(HistoryValue.a));
+	float DepthDeltaFraction = abs(abs(NewValue.a) - abs(HistoryValue.a)) / abs(NewValue.a);
 	float EffectiveHistoryWeight = HistoryWeight;
 
 	FLATTEN
 	if (any(OldScreenPosition > 1) 
 		|| any(OldScreenPosition < -1)
-		|| DepthDelta > 100
+		|| DepthDeltaFraction > .03f
 		// Toss history value from pixels that don't have a distance field representation
 		|| HistoryValue.a < 0)
 	{
diff --git a/Engine/Shaders/DistanceFieldSurfaceCacheLightingCompute.usf b/Engine/Shaders/DistanceFieldSurfaceCacheLightingCompute.usf
--- a/Engine/Shaders/DistanceFieldSurfaceCacheLightingCompute.usf
+++ b/Engine/Shaders/DistanceFieldSurfaceCacheLightingCompute.usf
@@ -336,6 +336,7 @@ void DistanceFieldAOBuildTileListMain(
 
 /** View rect min in xy, max in zw. */
 float2 ThreadToCulledTile;
+uint NumCircleSections;
 
 RWBuffer<float4> RWIrradianceCachePositionRadius;
 RWBuffer<float4> RWIrradianceCacheNormal;
@@ -384,6 +385,7 @@ void PopulateCacheCS(
 
 	BRANCH
 	if (GBufferData.ShadingModelID > 0
+		&& GBufferData.HasDistanceFieldRepresentation > 0
 		&& IrradianceCacheSplat.w < .0001f 
 		&& all(DispatchThreadId.xy < AOBufferSize))
 	{
@@ -439,7 +441,7 @@ void PopulateCacheCS(
 	if (all(DispatchThreadId == 0))
 	{
 		// VertexCountPerInstance
-		RWScatterDrawParameters[0] = 8 * 3;
+		RWScatterDrawParameters[0] = NumCircleSections * 3;
 		// StartVertexLocation
 		RWScatterDrawParameters[2] = 0;
 		// StartInstanceLocation
@@ -466,7 +468,7 @@ RWBuffer<float4> RWIrradianceCacheBentNormal;
 // Have to disable surface caching dependencies for this to work
 #define VISUALIZE_ONE_CONE 0
 
-#define FINAL_GATHER_THREADGROUP_SIZE 256
+#define FINAL_GATHER_THREADGROUP_SIZE 64
 
 #if USE_CONE_TRACING
 	// Enforce one thread per cone direction
diff --git a/Engine/Source/Runtime/Renderer/Private/DistanceFieldLightingShared.h b/Engine/Source/Runtime/Renderer/Private/DistanceFieldLightingShared.h
--- a/Engine/Source/Runtime/Renderer/Private/DistanceFieldLightingShared.h
+++ b/Engine/Source/Runtime/Renderer/Private/DistanceFieldLightingShared.h
@@ -9,8 +9,8 @@
 #include "DistanceFieldAtlas.h"
 
 /** Tile sized used for most AO compute shaders. */
-const int32 GDistanceFieldAOTileSizeX = 8;
-const int32 GDistanceFieldAOTileSizeY = 8;
+const int32 GDistanceFieldAOTileSizeX = 16;
+const int32 GDistanceFieldAOTileSizeY = 16;
 /** Base downsample factor that all distance field AO operations are done at. */
 const int32 GAODownsampleFactor = 2;
 static const int32 GMaxNumObjectsPerTile = 512;
diff --git a/Engine/Source/Runtime/Renderer/Private/DistanceFieldShadowing.cpp b/Engine/Source/Runtime/Renderer/Private/DistanceFieldShadowing.cpp
--- a/Engine/Source/Runtime/Renderer/Private/DistanceFieldShadowing.cpp
+++ b/Engine/Source/Runtime/Renderer/Private/DistanceFieldShadowing.cpp
@@ -728,8 +728,8 @@ void FProjectedShadowInfo::RenderRayTracedDistanceFieldProjection(FRHICommandLis
 				TRefCountPtr<IPooledRenderTarget> RayTracedShadowsRT;
 
 				{
-					const FIntPoint ExpandedBufferSize = GetBufferSizeForAO();
-					FPooledRenderTargetDesc Desc(FPooledRenderTargetDesc::Create2DDesc(ExpandedBufferSize / GAODownsampleFactor, PF_G16R16F, TexCreate_None, TexCreate_RenderTargetable | TexCreate_UAV, false));
+					const FIntPoint BufferSize = GetBufferSizeForAO();
+					FPooledRenderTargetDesc Desc(FPooledRenderTargetDesc::Create2DDesc(BufferSize, PF_G16R16F, TexCreate_None, TexCreate_RenderTargetable | TexCreate_UAV, false));
 					GRenderTargetPool.FindFreeElement(Desc, RayTracedShadowsRT, TEXT("RayTracedShadows"));
 				}
 
@@ -815,7 +815,7 @@ void FProjectedShadowInfo::RenderRayTracedDistanceFieldProjection(FRHICommandLis
 						ScissorRect.Min.X / GAODownsampleFactor, ScissorRect.Min.Y / GAODownsampleFactor, 
 						ScissorRect.Width() / GAODownsampleFactor, ScissorRect.Height() / GAODownsampleFactor,
 						FIntPoint(ScissorRect.Width(), ScissorRect.Height()),
-						GetBufferSizeForAO() / GAODownsampleFactor,
+						GetBufferSizeForAO(),
 						*VertexShader);
 				}
 			}
diff --git a/Engine/Source/Runtime/Renderer/Private/DistanceFieldSurfaceCacheLighting.cpp b/Engine/Source/Runtime/Renderer/Private/DistanceFieldSurfaceCacheLighting.cpp
--- a/Engine/Source/Runtime/Renderer/Private/DistanceFieldSurfaceCacheLighting.cpp
+++ b/Engine/Source/Runtime/Renderer/Private/DistanceFieldSurfaceCacheLighting.cpp
@@ -24,6 +24,14 @@ FAutoConsoleVariableRef CVarDistanceFieldAO(
 	ECVF_Cheat | ECVF_RenderThreadSafe
 	);
 
+int32 GAOPowerOfTwoBetweenLevels = 2;
+FAutoConsoleVariableRef CVarAOPowerOfTwoBetweenLevels(
+	TEXT("r.AOPowerOfTwoBetweenLevels"),
+	GAOPowerOfTwoBetweenLevels,
+	TEXT("Power of two in resolution between refinement levels of the surface cache"),
+	ECVF_Cheat | ECVF_RenderThreadSafe
+	);
+
 int32 GAOMinLevel = 1;
 FAutoConsoleVariableRef CVarAOMinLevel(
 	TEXT("r.AOMinLevel"),
@@ -34,7 +42,7 @@ FAutoConsoleVariableRef CVarAOMinLevel(
 	ECVF_Cheat | ECVF_RenderThreadSafe
 	);
 
-int32 GAOMaxLevel = GAOMaxSupportedLevel;
+int32 GAOMaxLevel = FMath::Min(4, GAOMaxSupportedLevel);
 FAutoConsoleVariableRef CVarAOMaxLevel(
 	TEXT("r.AOMaxLevel"),
 	GAOMaxLevel,
@@ -59,6 +67,11 @@ FAutoConsoleVariableRef CVarAOTrimOldRecordsFraction(
 	ECVF_Cheat | ECVF_RenderThreadSafe
 	);
 
+float GetTrimOldRecordsFraction(int32 DepthLevel)
+{
+	return GAOTrimOldRecordsFraction;
+}
+
 int32 GAOFillGaps = 1;
 FAutoConsoleVariableRef CVarAOFillGaps(
 	TEXT("r.AOFillGaps"),
@@ -92,7 +105,7 @@ FAutoConsoleVariableRef CVarAORecordRadiusScale(
 	ECVF_Cheat | ECVF_RenderThreadSafe
 	);
 
-float GAOInterpolationRadiusScale = 1.5f;
+float GAOInterpolationRadiusScale = 1.3f;
 FAutoConsoleVariableRef CVarAOInterpolationRadiusScale(
 	TEXT("r.AOInterpolationRadiusScale"),
 	GAOInterpolationRadiusScale,
@@ -109,7 +122,7 @@ FAutoConsoleVariableRef CVarAOMinPointBehindPlaneAngle(
 	ECVF_Cheat | ECVF_RenderThreadSafe
 	);
 
-float GAOInterpolationMaxAngle = 10;
+float GAOInterpolationMaxAngle = 15;
 FAutoConsoleVariableRef CVarAOInterpolationMaxAngle(
 	TEXT("r.AOInterpolationMaxAngle"),
 	GAOInterpolationMaxAngle,
@@ -157,6 +170,14 @@ FAutoConsoleVariableRef CVarAOScatterTileCulling(
 	ECVF_Cheat | ECVF_RenderThreadSafe
 	);
 
+int32 GAOScatterInterpolation = 1;
+FAutoConsoleVariableRef CVarAOScatterInterpolation(
+	TEXT("r.AOScatterInterpolation"),
+	GAOScatterInterpolation,
+	TEXT(""),
+	ECVF_Cheat | ECVF_RenderThreadSafe
+	);
+
 int32 GAOComputeShaderNormalCalculation = 1;
 FAutoConsoleVariableRef CVarAOComputeShaderNormalCalculation(
 	TEXT("r.AOComputeShaderNormalCalculation"),
@@ -173,13 +194,27 @@ FAutoConsoleVariableRef CVarAOSampleSet(
 	ECVF_Cheat | ECVF_RenderThreadSafe
 	);
 
+int32 GAOInterpolationStencilTesting = 1;
+FAutoConsoleVariableRef CVarAOInterpolationStencilTesting(
+	TEXT("r.AOInterpolationStencilTesting"),
+	GAOInterpolationStencilTesting,
+	TEXT("Whether to stencil out distant pixels from the interpolation splat pass, useful for debugging"),
+	ECVF_Cheat | ECVF_RenderThreadSafe
+	);
+
+int32 GAOInterpolationDepthTesting = 1;
+FAutoConsoleVariableRef CVarAOInterpolationDepthTesting(
+	TEXT("r.AOInterpolationDepthTesting"),
+	GAOInterpolationDepthTesting,
+	TEXT("Whether to use depth testing during the interpolation splat pass, useful for debugging"),
+	ECVF_Cheat | ECVF_RenderThreadSafe
+	);
+
 IMPLEMENT_UNIFORM_BUFFER_STRUCT(FAOSampleData2,TEXT("AOSamples2"));
 
 FIntPoint GetBufferSizeForAO()
 {
-	return FIntPoint(
-		GAODownsampleFactor * (GSceneRenderTargets.GetBufferSizeXY().X / GAODownsampleFactor),
-		GAODownsampleFactor * (GSceneRenderTargets.GetBufferSizeXY().Y / GAODownsampleFactor));
+	return FIntPoint::DivideAndRoundDown(GSceneRenderTargets.GetBufferSizeXY(), GAODownsampleFactor);
 }
 
 bool DoesPlatformSupportDistanceFieldAO(EShaderPlatform Platform)
@@ -243,6 +278,51 @@ float GAOConeHalfAngle = FMath::Acos(1 - 1.0f / (float)ARRAY_COUNT(SpacedVectors
 // Must match shader code
 uint32 GAONumConeSteps = 10;
 
+class FCircleVertexBuffer : public FVertexBuffer
+{
+public:
+
+	int32 NumSections;
+
+	FCircleVertexBuffer()
+	{
+		NumSections = 8;
+	}
+
+	virtual void InitRHI() override
+	{
+		// Used as a non-indexed triangle list, so 3 vertices per triangle
+		const uint32 Size = 3 * NumSections * sizeof(FScreenVertex);
+		FRHIResourceCreateInfo CreateInfo;
+		VertexBufferRHI = RHICreateVertexBuffer(Size, BUF_Static, CreateInfo);
+		void* Buffer = RHILockVertexBuffer(VertexBufferRHI, 0, Size, RLM_WriteOnly);
+		FScreenVertex* DestVertex = (FScreenVertex*)Buffer;
+
+		const float RadiansPerRingSegment = PI / (float)NumSections;
+
+		// Boost the effective radius so that the edges of the circle approximation lie on the circle, instead of the vertices
+		const float Radius = 1.0f / FMath::Cos(RadiansPerRingSegment);
+
+		for (int32 SectionIndex = 0; SectionIndex < NumSections; SectionIndex++)
+		{
+			float Fraction = SectionIndex / (float)NumSections;
+			float CurrentAngle = Fraction * 2 * PI;
+			float NextAngle = ((SectionIndex + 1) / (float)NumSections) * 2 * PI;
+			FVector2D CurrentPosition(Radius * FMath::Cos(CurrentAngle), Radius * FMath::Sin(CurrentAngle));
+			FVector2D NextPosition(Radius * FMath::Cos(NextAngle), Radius * FMath::Sin(NextAngle));
+
+			DestVertex[SectionIndex * 3 + 0].Position = FVector2D(0, 0);
+			DestVertex[SectionIndex * 3 + 0].UV = CurrentPosition;
+			DestVertex[SectionIndex * 3 + 1].Position = FVector2D(0, 0);
+			DestVertex[SectionIndex * 3 + 1].UV = NextPosition;
+			DestVertex[SectionIndex * 3 + 2].Position = FVector2D(0, 0);
+			DestVertex[SectionIndex * 3 + 2].UV = FVector2D(.5f, .5f);
+		}
+
+		RHIUnlockVertexBuffer(VertexBufferRHI);      
+	}
+};
+
 TGlobalResource<FCircleVertexBuffer> GCircleVertexBuffer;
 
 // Only 7 needed atm, pad to cache line size
@@ -298,7 +378,7 @@ void FSceneViewState::DestroyAOTileResources()
 	AOTileIntersectionResources = NULL;
 }
 
-void OnClearIrradianceCache(UWorld* InWorld)
+void OnClearSurfaceCache(UWorld* InWorld)
 {
 	FlushRenderingCommands();
 
@@ -313,7 +393,7 @@ void OnClearIrradianceCache(UWorld* InWorld)
 FAutoConsoleCommandWithWorld ClearCacheConsoleCommand(
 	TEXT("r.AOClearCache"),
 	TEXT(""),
-	FConsoleCommandWithWorldDelegate::CreateStatic(OnClearIrradianceCache)
+	FConsoleCommandWithWorldDelegate::CreateStatic(OnClearSurfaceCache)
 	);
 
 /**  */
@@ -671,7 +751,7 @@ public:
 		SetShaderValue(RHICmdList, ShaderRHI, CurrentLevelDownsampleFactor, CurrentLevelDownsampleFactorValue);
 
 		// Round up, to match render target allocation
-		const FVector2D AOBufferSizeValue = FIntPoint::DivideAndRoundUp(GSceneRenderTargets.GetBufferSizeXY(), CurrentLevelDownsampleFactorValue);
+		const FVector2D AOBufferSizeValue = FIntPoint::DivideAndRoundDown(GSceneRenderTargets.GetBufferSizeXY(), CurrentLevelDownsampleFactorValue);
 		SetShaderValue(RHICmdList, ShaderRHI, AOBufferSize, AOBufferSizeValue);
 
 		SetShaderValue(RHICmdList, ShaderRHI, DownsampleFactorToBaseLevel, CurrentLevelDownsampleFactorValue / GAODownsampleFactor);
@@ -970,7 +1050,7 @@ public:
 
 		DispatchParameters.SetBuffer(RHICmdList, ShaderRHI, SurfaceCacheResources.DispatchParameters);
 
-		SetShaderValue(RHICmdList, ShaderRHI, TrimFraction, GAOTrimOldRecordsFraction);
+		SetShaderValue(RHICmdList, ShaderRHI, TrimFraction, GetTrimOldRecordsFraction(DepthLevel));
 	}
 
 	void UnsetParameters(FRHICommandList& RHICmdList)
@@ -1057,7 +1137,7 @@ public:
 		CopyIrradianceCacheTileCoordinate.SetBuffer(RHICmdList, ShaderRHI, SurfaceCacheResources.TempResources->TileCoordinate);
 		ScatterDrawParameters.SetBuffer(RHICmdList, ShaderRHI, SurfaceCacheResources.TempResources->ScatterDrawParameters);
 
-		SetShaderValue(RHICmdList, ShaderRHI, TrimFraction, GAOTrimOldRecordsFraction);
+		SetShaderValue(RHICmdList, ShaderRHI, TrimFraction, GetTrimOldRecordsFraction(DepthLevel));
 	}
 
 	void UnsetParameters(FRHICommandList& RHICmdList)
@@ -1197,6 +1277,7 @@ public:
 		ViewDimensionsParameter.Bind(Initializer.ParameterMap, TEXT("ViewDimensions"));
 		ThreadToCulledTile.Bind(Initializer.ParameterMap, TEXT("ThreadToCulledTile"));
 		TileListGroupSize.Bind(Initializer.ParameterMap, TEXT("TileListGroupSize"));
+		NumCircleSections.Bind(Initializer.ParameterMap, TEXT("NumCircleSections"));
 		DistanceFieldNormalTexture.Bind(Initializer.ParameterMap, TEXT("DistanceFieldNormalTexture"));
 		DistanceFieldNormalSampler.Bind(Initializer.ParameterMap, TEXT("DistanceFieldNormalSampler"));
 		IrradianceCacheSplatTexture.Bind(Initializer.ParameterMap, TEXT("IrradianceCacheSplatTexture"));
@@ -1251,6 +1332,8 @@ public:
 
 		SetShaderValue(RHICmdList, ShaderRHI, TileListGroupSize, TileListGroupSizeValue);
 
+		SetShaderValue(RHICmdList, ShaderRHI, NumCircleSections, GCircleVertexBuffer.NumSections);
+
 		SetTextureParameter(
 			RHICmdList,
 			ShaderRHI,
@@ -1291,6 +1374,7 @@ public:
 		Ar << ViewDimensionsParameter;
 		Ar << ThreadToCulledTile;
 		Ar << TileListGroupSize;
+		Ar << NumCircleSections;
 		Ar << DistanceFieldNormalTexture;
 		Ar << DistanceFieldNormalSampler;
 		Ar << IrradianceCacheSplatTexture;
@@ -1310,6 +1394,7 @@ private:
 	FShaderParameter ViewDimensionsParameter;
 	FShaderParameter ThreadToCulledTile;
 	FShaderParameter TileListGroupSize;
+	FShaderParameter NumCircleSections;
 	FShaderResourceParameter DistanceFieldNormalTexture;
 	FShaderResourceParameter DistanceFieldNormalSampler;
 	FShaderResourceParameter IrradianceCacheSplatTexture;
@@ -1602,6 +1687,62 @@ private:
 
 IMPLEMENT_SHADER_TYPE(,FFinalGatherCS,TEXT("DistanceFieldSurfaceCacheLightingCompute"),TEXT("FinalGatherCS"),SF_Compute);
 
+class FWriteDownsampledDepthPS : public FGlobalShader
+{
+	DECLARE_SHADER_TYPE(FWriteDownsampledDepthPS, Global);
+public:
+
+	static bool ShouldCache(EShaderPlatform Platform)
+	{
+		return IsFeatureLevelSupported(Platform, ERHIFeatureLevel::SM5) && DoesPlatformSupportDistanceFieldAO(Platform);
+	}
+
+	static void ModifyCompilationEnvironment(EShaderPlatform Platform, FShaderCompilerEnvironment& OutEnvironment)
+	{
+		OutEnvironment.SetDefine(TEXT("DOWNSAMPLE_FACTOR"), GAODownsampleFactor);
+	}
+
+	/** Default constructor. */
+	FWriteDownsampledDepthPS() {}
+
+	/** Initialization constructor. */
+	FWriteDownsampledDepthPS(const ShaderMetaType::CompiledShaderInitializerType& Initializer)
+		: FGlobalShader(Initializer)
+	{
+		DeferredParameters.Bind(Initializer.ParameterMap);
+		DistanceFieldNormalTexture.Bind(Initializer.ParameterMap,TEXT("DistanceFieldNormalTexture"));
+		DistanceFieldNormalSampler.Bind(Initializer.ParameterMap,TEXT("DistanceFieldNormalSampler"));
+	}
+
+	void SetParameters(FRHICommandList& RHICmdList, const FSceneView& View, FSceneRenderTargetItem& DistanceFieldNormal)
+	{
+		const FPixelShaderRHIParamRef ShaderRHI = GetPixelShader();
+
+		FGlobalShader::SetParameters(RHICmdList, ShaderRHI, View);
+		DeferredParameters.Set(RHICmdList, ShaderRHI, View);
+
+		SetTextureParameter(RHICmdList, ShaderRHI, DistanceFieldNormalTexture, DistanceFieldNormalSampler, TStaticSamplerState<SF_Point>::GetRHI(), DistanceFieldNormal.ShaderResourceTexture);
+	}
+	// FShader interface.
+	virtual bool Serialize(FArchive& Ar)
+	{
+		bool bShaderHasOutdatedParameters = FGlobalShader::Serialize(Ar);
+		Ar << DeferredParameters;
+		Ar << DistanceFieldNormalTexture;
+		Ar << DistanceFieldNormalSampler;
+		return bShaderHasOutdatedParameters;
+	}
+
+private:
+
+	FDeferredPixelShaderParameters DeferredParameters;
+	FShaderResourceParameter DistanceFieldNormalTexture;
+	FShaderResourceParameter DistanceFieldNormalSampler;
+};
+
+IMPLEMENT_SHADER_TYPE(,FWriteDownsampledDepthPS,TEXT("DistanceFieldSurfaceCacheLighting"),TEXT("WriteDownsampledDepthPS"),SF_Pixel);
+
+
 /**  */
 template<bool bFinalInterpolationPass>
 class TIrradianceCacheSplatVS : public FGlobalShader
@@ -1651,7 +1792,7 @@ public:
 		SetShaderValue(RHICmdList, ShaderRHI, InterpolationRadiusScale, (bFinalInterpolationPass ? GAOInterpolationRadiusScale : 1.0f));
 		SetShaderValue(RHICmdList, ShaderRHI, NormalizedOffsetToPixelCenter, NormalizedOffsetToPixelCenterValue);
 
-		const FIntPoint AOViewRectSize = FIntPoint::DivideAndRoundUp(View.ViewRect.Size(), CurrentLevelDownsampleFactorValue);
+		const FIntPoint AOViewRectSize = FIntPoint::DivideAndRoundDown(View.ViewRect.Size(), CurrentLevelDownsampleFactorValue);
 		const FVector2D HackExpandValue(.5f / AOViewRectSize.X, .5f / AOViewRectSize.Y);
 		SetShaderValue(RHICmdList, ShaderRHI, HackExpand, HackExpandValue);
 	}
@@ -1769,6 +1910,207 @@ private:
 IMPLEMENT_SHADER_TYPE(template<>,TIrradianceCacheSplatPS<true>,TEXT("DistanceFieldSurfaceCacheLighting"),TEXT("IrradianceCacheSplatPS"),SF_Pixel);
 IMPLEMENT_SHADER_TYPE(template<>,TIrradianceCacheSplatPS<false>,TEXT("DistanceFieldSurfaceCacheLighting"),TEXT("IrradianceCacheSplatPS"),SF_Pixel);
 
+
+template<uint32 NumLevels, bool bFinalInterpolationPass>
+class TInterpolateCacheCS : public FGlobalShader
+{
+	DECLARE_SHADER_TYPE(TInterpolateCacheCS,Global)
+public:
+
+	static bool ShouldCache(EShaderPlatform Platform)
+	{
+		return IsFeatureLevelSupported(Platform, ERHIFeatureLevel::SM5);
+	}
+
+	static void ModifyCompilationEnvironment(EShaderPlatform Platform, FShaderCompilerEnvironment& OutEnvironment)
+	{
+		FGlobalShader::ModifyCompilationEnvironment(Platform,OutEnvironment);
+
+		OutEnvironment.SetDefine(TEXT("FINAL_INTERPOLATION_PASS"), (uint32)bFinalInterpolationPass);
+		OutEnvironment.SetDefine(TEXT("NUM_INTERPOLATE_LEVELS"), NumLevels);
+		OutEnvironment.SetDefine(TEXT("THREADGROUP_SIZEX"), GDistanceFieldAOTileSizeX);
+		OutEnvironment.SetDefine(TEXT("THREADGROUP_SIZEY"), GDistanceFieldAOTileSizeY);
+		OutEnvironment.CompilerFlags.Add(CFLAG_StandardOptimization);
+	}
+
+	TInterpolateCacheCS(const ShaderMetaType::CompiledShaderInitializerType& Initializer)
+		: FGlobalShader(Initializer)
+	{
+		AOParameters.Bind(Initializer.ParameterMap);
+		DeferredParameters.Bind(Initializer.ParameterMap);
+		AOLevelParameters.Bind(Initializer.ParameterMap);
+		DistanceFieldNormalTexture.Bind(Initializer.ParameterMap, TEXT("DistanceFieldNormalTexture"));
+		DistanceFieldNormalSampler.Bind(Initializer.ParameterMap, TEXT("DistanceFieldNormalSampler"));
+		InterpolationAngleNormalization.Bind(Initializer.ParameterMap, TEXT("InterpolationAngleNormalization"));
+		InvMinCosPointBehindPlane.Bind(Initializer.ParameterMap, TEXT("InvMinCosPointBehindPlane"));
+		IrradianceCacheSplat.Bind(Initializer.ParameterMap, TEXT("IrradianceCacheSplat"));
+		IrradianceCachePositionRadius.Bind(Initializer.ParameterMap, TEXT("LevelIrradianceCachePositionRadius"));
+		IrradianceCacheOccluderRadius.Bind(Initializer.ParameterMap, TEXT("LevelIrradianceCacheOccluderRadius"));
+		IrradianceCacheNormal.Bind(Initializer.ParameterMap, TEXT("LevelIrradianceCacheNormal"));
+		IrradianceCacheBentNormal.Bind(Initializer.ParameterMap, TEXT("LevelIrradianceCacheBentNormal"));
+		InterpolationRadiusScale.Bind(Initializer.ParameterMap, TEXT("InterpolationRadiusScale"));
+		DrawParameters.Bind(Initializer.ParameterMap, TEXT("LevelScatterDrawParameters"));
+		NumGroups.Bind(Initializer.ParameterMap, TEXT("NumGroups"));
+	}
+
+	TInterpolateCacheCS()
+	{
+	}
+
+	void SetParameters(FRHICommandList& RHICmdList, const FSceneView& View, FSceneRenderTargetItem& DistanceFieldNormal, int32 DepthLevel, int32 DestLevelDownsampleFactor, IPooledRenderTarget* InterpolationTarget, const FDistanceFieldAOParameters& Parameters, FVector2D NumGroupsValue, FSurfaceCacheResources& SurfaceCacheResources)
+	{
+		FComputeShaderRHIParamRef ShaderRHI = GetComputeShader();
+		FGlobalShader::SetParameters(RHICmdList, ShaderRHI, View);
+		AOParameters.Set(RHICmdList, ShaderRHI, Parameters);
+		DeferredParameters.Set(RHICmdList, ShaderRHI, View);
+		AOLevelParameters.Set(RHICmdList, ShaderRHI, View, DestLevelDownsampleFactor);
+
+		SetTextureParameter(
+			RHICmdList,
+			ShaderRHI,
+			DistanceFieldNormalTexture,
+			DistanceFieldNormalSampler,
+			TStaticSamplerState<SF_Point,AM_Clamp,AM_Clamp,AM_Clamp>::GetRHI(),
+			DistanceFieldNormal.ShaderResourceTexture
+			);
+
+		const float EffectiveMaxAngle = bFinalInterpolationPass ? GAOInterpolationMaxAngle * GAOInterpolationAngleScale : GAOInterpolationMaxAngle;
+		const float InterpolationAngleNormalizationValue = 1.0f / FMath::Sqrt(1.0f - FMath::Cos(EffectiveMaxAngle * PI / 180.0f));
+		SetShaderValue(RHICmdList, ShaderRHI, InterpolationAngleNormalization, InterpolationAngleNormalizationValue);
+
+		const float MinCosPointBehindPlaneValue = FMath::Cos((GAOMinPointBehindPlaneAngle + 90.0f) * PI / 180.0f);
+		SetShaderValue(RHICmdList, ShaderRHI, InvMinCosPointBehindPlane, 1.0f / MinCosPointBehindPlaneValue);
+
+		IrradianceCacheSplat.SetTexture(
+			RHICmdList,
+			ShaderRHI, 
+			InterpolationTarget->GetRenderTargetItem().TargetableTexture, 
+			InterpolationTarget->GetRenderTargetItem().UAV);
+
+		const int32 MaxAllowedLevel = ARRAY_COUNT(SurfaceCacheResources.Level) - 1;
+		const int32 MinLevel = MaxAllowedLevel + 1 - NumLevels;
+
+		for (int32 LevelIndex = MaxAllowedLevel; LevelIndex >= MinLevel; LevelIndex--)
+		{
+			if (IrradianceCachePositionRadius.IsBound())
+			{
+				check(MaxAllowedLevel - LevelIndex < (int32)IrradianceCachePositionRadius.GetNumResources());
+
+				RHICmdList.SetShaderResourceViewParameter(
+					ShaderRHI,
+					IrradianceCachePositionRadius.GetBaseIndex() + MaxAllowedLevel - LevelIndex,
+					SurfaceCacheResources.Level[LevelIndex]->PositionAndRadius.SRV
+					);
+			}
+
+			if (IrradianceCacheOccluderRadius.IsBound())
+			{
+				check(MaxAllowedLevel - LevelIndex < (int32)IrradianceCacheOccluderRadius.GetNumResources());
+
+				RHICmdList.SetShaderResourceViewParameter(
+					ShaderRHI,
+					IrradianceCacheOccluderRadius.GetBaseIndex() + MaxAllowedLevel - LevelIndex,
+					SurfaceCacheResources.Level[LevelIndex]->OccluderRadius.SRV
+					);
+			}
+
+			if (IrradianceCacheNormal.IsBound())
+			{
+				check(MaxAllowedLevel - LevelIndex < (int32)IrradianceCacheNormal.GetNumResources());
+
+				RHICmdList.SetShaderResourceViewParameter(
+					ShaderRHI,
+					IrradianceCacheNormal.GetBaseIndex() + MaxAllowedLevel - LevelIndex,
+					SurfaceCacheResources.Level[LevelIndex]->Normal.SRV
+					);
+			}
+
+			if (IrradianceCacheBentNormal.IsBound())
+			{
+				check(MaxAllowedLevel - LevelIndex < (int32)IrradianceCacheBentNormal.GetNumResources());
+
+				RHICmdList.SetShaderResourceViewParameter(
+					ShaderRHI,
+					IrradianceCacheBentNormal.GetBaseIndex() + MaxAllowedLevel - LevelIndex,
+					SurfaceCacheResources.Level[LevelIndex]->BentNormal.SRV
+					);
+			}
+
+			if (DrawParameters.IsBound())
+			{
+				check(MaxAllowedLevel - LevelIndex < (int32)DrawParameters.GetNumResources());
+
+				RHICmdList.SetShaderResourceViewParameter(
+					ShaderRHI,
+					DrawParameters.GetBaseIndex() + MaxAllowedLevel - LevelIndex,
+					SurfaceCacheResources.Level[LevelIndex]->ScatterDrawParameters.SRV
+					);
+			}
+		}
+
+		SetShaderValue(RHICmdList, ShaderRHI, InterpolationRadiusScale, (bFinalInterpolationPass ? GAOInterpolationRadiusScale : 1.0f));
+
+		SetShaderValue(RHICmdList, ShaderRHI, NumGroups, NumGroupsValue);
+	}
+
+	void UnsetParameters(FRHICommandListImmediate& RHICmdList)
+	{
+		IrradianceCacheSplat.UnsetUAV(RHICmdList, GetComputeShader());
+	}
+
+	virtual bool Serialize(FArchive& Ar)
+	{		
+		bool bShaderHasOutdatedParameters = FGlobalShader::Serialize(Ar);
+		Ar << AOParameters;
+		Ar << AOLevelParameters;
+		Ar << DeferredParameters;
+		Ar << DistanceFieldNormalTexture;
+		Ar << DistanceFieldNormalSampler;
+		Ar << InterpolationAngleNormalization;
+		Ar << InvMinCosPointBehindPlane;
+		Ar << IrradianceCacheSplat;
+		Ar << IrradianceCachePositionRadius;
+		Ar << IrradianceCacheOccluderRadius;
+		Ar << IrradianceCacheNormal;
+		Ar << IrradianceCacheBentNormal;
+		Ar << InterpolationRadiusScale;
+		Ar << DrawParameters;
+		Ar << NumGroups;
+		return bShaderHasOutdatedParameters;
+	}
+
+private:
+
+	FAOParameters AOParameters;
+	FAOLevelParameters AOLevelParameters;
+	FDeferredPixelShaderParameters DeferredParameters;
+	FShaderResourceParameter DistanceFieldNormalTexture;
+	FShaderResourceParameter DistanceFieldNormalSampler;
+	FShaderParameter InterpolationAngleNormalization;
+	FShaderParameter InvMinCosPointBehindPlane;
+	FRWShaderParameter IrradianceCacheSplat;
+	FShaderResourceParameter IrradianceCachePositionRadius;
+	FShaderResourceParameter IrradianceCacheOccluderRadius;
+	FShaderResourceParameter IrradianceCacheNormal;
+	FShaderResourceParameter IrradianceCacheBentNormal;
+	FShaderParameter InterpolationRadiusScale;
+	FShaderResourceParameter DrawParameters;
+	FShaderParameter NumGroups;
+};
+
+#define IMPLEMENT_INTERPOLATION_TYPE(NumLevels,bFinalInterpolationPass) \
+	typedef TInterpolateCacheCS<NumLevels,bFinalInterpolationPass> TInterpolateCacheCS##NumLevels##bFinalInterpolationPass; \
+	IMPLEMENT_SHADER_TYPE(template<>,TInterpolateCacheCS##NumLevels##bFinalInterpolationPass,TEXT("DistanceFieldSurfaceCacheLighting"),TEXT("InterpolateCacheCS"),SF_Compute);
+
+IMPLEMENT_INTERPOLATION_TYPE(5, true);
+IMPLEMENT_INTERPOLATION_TYPE(4, true);
+IMPLEMENT_INTERPOLATION_TYPE(1, false);
+IMPLEMENT_INTERPOLATION_TYPE(2, false);
+IMPLEMENT_INTERPOLATION_TYPE(3, false);
+IMPLEMENT_INTERPOLATION_TYPE(4, false);
+IMPLEMENT_INTERPOLATION_TYPE(5, false);
+
+
 class FDistanceFieldAOCombinePS2 : public FGlobalShader
 {
 	DECLARE_SHADER_TYPE(FDistanceFieldAOCombinePS2, Global);
@@ -2004,10 +2346,9 @@ void AllocateOrReuseAORenderTarget(TRefCountPtr<IPooledRenderTarget>& Target, co
 {
 	if (!Target)
 	{
-		extern FIntPoint GetBufferSizeForAO();
-		FIntPoint ExpandedBufferSize = GetBufferSizeForAO();
+		FIntPoint BufferSize = GetBufferSizeForAO();
 
-		FPooledRenderTargetDesc Desc(FPooledRenderTargetDesc::Create2DDesc(ExpandedBufferSize / GAODownsampleFactor, PF_FloatRGBA, TexCreate_None, TexCreate_RenderTargetable | TexCreate_UAV, false));
+		FPooledRenderTargetDesc Desc(FPooledRenderTargetDesc::Create2DDesc(BufferSize, PF_FloatRGBA, TexCreate_None, TexCreate_RenderTargetable | TexCreate_UAV, false));
 		GRenderTargetPool.FindFreeElement(Desc, Target, Name);
 	}
 }
@@ -2130,7 +2471,7 @@ void PostProcessBentNormalAO(
 				0, 0, 
 				View.ViewRect.Width() / GAODownsampleFactor, View.ViewRect.Height() / GAODownsampleFactor,
 				FIntPoint(View.ViewRect.Width() / GAODownsampleFactor, View.ViewRect.Height() / GAODownsampleFactor),
-				GetBufferSizeForAO() / GAODownsampleFactor,
+				GetBufferSizeForAO(),
 				*VertexShader);
 		}
 	}
@@ -2284,7 +2625,7 @@ void UpsampleBentNormalAO(FRHICommandList& RHICmdList, const TArray<FViewInfo>&
 			View.ViewRect.Min.X / GAODownsampleFactor, View.ViewRect.Min.Y / GAODownsampleFactor, 
 			View.ViewRect.Width() / GAODownsampleFactor, View.ViewRect.Height() / GAODownsampleFactor,
 			FIntPoint(View.ViewRect.Width(), View.ViewRect.Height()),
-			GetBufferSizeForAO() / GAODownsampleFactor,
+			GetBufferSizeForAO(),
 			*VertexShader);
 	}
 }
@@ -2577,9 +2918,82 @@ FIntPoint BuildTileObjectLists(FRHICommandListImmediate& RHICmdList, FScene* Sce
 	return TileListGroupSize;
 }
 
+void SetupDepthStencil(FRHICommandListImmediate& RHICmdList, FViewInfo& View, FSceneRenderTargetItem& DistanceFieldNormal, int32 DepthLevel, int32 DestLevelDownsampleFactor)
+{
+	SCOPED_DRAW_EVENT(RHICmdList, SetupDepthStencil, DEC_SCENE_ITEMS);
+
+	{		
+		RHICmdList.SetViewport(0, 0, 0.0f, View.ViewRect.Width() / GAODownsampleFactor, View.ViewRect.Height() / GAODownsampleFactor, 1.0f);
+		RHICmdList.SetRasterizerState(TStaticRasterizerState<FM_Solid, CM_None>::GetRHI());
+		// Depth writes on
+		RHICmdList.SetDepthStencilState(TStaticDepthStencilState<true, CF_Always>::GetRHI());
+		// Color writes off
+		RHICmdList.SetBlendState(TStaticBlendState<CW_NONE>::GetRHI());
+
+		TShaderMapRef<FPostProcessVS> VertexShader(View.ShaderMap);
+		TShaderMapRef<FWriteDownsampledDepthPS> PixelShader(View.ShaderMap);
+
+		static FGlobalBoundShaderState BoundShaderState;
+		SetGlobalBoundShaderState(RHICmdList, View.GetFeatureLevel(), BoundShaderState, GFilterVertexDeclaration.VertexDeclarationRHI, *VertexShader, *PixelShader);
+		PixelShader->SetParameters(RHICmdList, View, DistanceFieldNormal);
+
+		// Draw a quad writing depth out to the depth stencil view
+		DrawRectangle( 
+			RHICmdList,
+			0, 0, 
+			View.ViewRect.Width() / GAODownsampleFactor, View.ViewRect.Height() / GAODownsampleFactor,
+			0, 0, 
+			View.ViewRect.Width() / GAODownsampleFactor, View.ViewRect.Height() / GAODownsampleFactor,
+			FIntPoint(View.ViewRect.Width() / GAODownsampleFactor, View.ViewRect.Height() / GAODownsampleFactor),
+			GetBufferSizeForAO(),
+			*VertexShader);
+	}
+
+	if (GAOInterpolationStencilTesting)
+	{
+		RHICmdList.SetViewport(0, 0, 0.0f, View.ViewRect.Width() / GAODownsampleFactor, View.ViewRect.Height() / GAODownsampleFactor, 1.0f);
+		RHICmdList.SetRasterizerState(TStaticRasterizerState<FM_Solid, CM_None>::GetRHI());
+		RHICmdList.SetBlendState(TStaticBlendState<CW_NONE>::GetRHI());
+
+		// Depth tests on, write 1 to stencil if depth test passed
+		RHICmdList.SetDepthStencilState(TStaticDepthStencilState<
+			// Note, this is a reversed Z depth surface, using CF_GreaterEqual.
+			false,CF_GreaterEqual,
+			true,CF_Always,SO_Keep,SO_Keep,SO_Replace,
+			false,CF_Always,SO_Keep,SO_Keep,SO_Keep,
+			0xff,0xff
+			>::GetRHI(), 1);
+
+		TShaderMapRef<TOneColorVS<true> > VertexShader(View.ShaderMap);
+		TShaderMapRef<FOneColorPS> PixelShader(View.ShaderMap);
+
+		static FGlobalBoundShaderState BoundShaderState;
+		SetGlobalBoundShaderState(RHICmdList, View.GetFeatureLevel(), BoundShaderState, GetVertexDeclarationFVector4(), *VertexShader, *PixelShader);
+
+		SetShaderValue(RHICmdList, PixelShader->GetPixelShader(), PixelShader->ColorParameter, FLinearColor::Black);
+
+		FVector ViewSpaceMaxDistance(0.0f, 0.0f, GAOMaxViewDistance);
+		FVector4 ClipSpaceMaxDistance = View.ViewMatrices.ProjMatrix.TransformPosition(ViewSpaceMaxDistance);
+		float ClipSpaceZ = ClipSpaceMaxDistance.Z / ClipSpaceMaxDistance.W;
+
+		// Place the quad's depth at the max AO view distance
+		// Any pixels that pass the depth test should not receive AO and therefore write a 1 to stencil
+		FVector4 ClearQuadVertices[4] = 
+		{
+			FVector4( -1.0f,  1.0f, ClipSpaceZ, 1.0f ),
+			FVector4(  1.0f,  1.0f, ClipSpaceZ, 1.0f ),
+			FVector4( -1.0f, -1.0f, ClipSpaceZ, 1.0f ),
+			FVector4(  1.0f, -1.0f, ClipSpaceZ, 1.0f )
+		};
+
+		DrawPrimitiveUP(RHICmdList, PT_TriangleStrip, 2, ClearQuadVertices, sizeof(ClearQuadVertices[0]));
+	}
+}
+
 void RenderIrradianceCacheInterpolation(
 	FRHICommandListImmediate& RHICmdList, 
-	TArray<FViewInfo>& Views, IPooledRenderTarget* InterpolationTarget, 
+	TArray<FViewInfo>& Views, 
+	IPooledRenderTarget* InterpolationTarget, 
 	FSceneRenderTargetItem& DistanceFieldNormal, 
 	int32 DepthLevel, 
 	int32 DestLevelDownsampleFactor, 
@@ -2588,12 +3002,65 @@ void RenderIrradianceCacheInterpolation(
 {
 	check(!(bFinalInterpolation && DepthLevel != 0));
 
+	if (bFinalInterpolation && !GAOScatterInterpolation)
 	{
+		for (int32 ViewIndex = 0; ViewIndex < Views.Num(); ViewIndex++)
+		{
+			FViewInfo& View = Views[ViewIndex];
+
+			const FScene* Scene = (const FScene*)View.Family->Scene;
+			FSurfaceCacheResources& SurfaceCacheResources = *Scene->SurfaceCacheResources;
+
+			FIntPoint DownsampledViewSize = FIntPoint::DivideAndRoundDown(View.ViewRect.Size(), DestLevelDownsampleFactor);
+			uint32 GroupSizeX = FMath::DivideAndRoundUp(DownsampledViewSize.X, GDistanceFieldAOTileSizeX);
+			uint32 GroupSizeY = FMath::DivideAndRoundUp(DownsampledViewSize.Y, GDistanceFieldAOTileSizeY);
+
+			TShaderMapRef<TInterpolateCacheCS<4, true> > ComputeShader(View.ShaderMap);
+
+			RHICmdList.SetComputeShader(ComputeShader->GetComputeShader());	
+			ComputeShader->SetParameters(RHICmdList, View, DistanceFieldNormal, DepthLevel, DestLevelDownsampleFactor, InterpolationTarget, Parameters, FVector2D(GroupSizeX, GroupSizeY), SurfaceCacheResources);
+			DispatchComputeShader(RHICmdList, *ComputeShader, GroupSizeX, GroupSizeY, 1);
+
+			ComputeShader->UnsetParameters(RHICmdList);
+		}
+	}
+	else
+	{
+		TRefCountPtr<IPooledRenderTarget> SplatDepthStencilBuffer;
+
+		if (bFinalInterpolation && (GAOInterpolationDepthTesting || GAOInterpolationStencilTesting))
+		{
+			FPooledRenderTargetDesc Desc(FPooledRenderTargetDesc::Create2DDesc(InterpolationTarget->GetDesc().Extent, PF_DepthStencil, TexCreate_None, TexCreate_DepthStencilTargetable, false));
+			GRenderTargetPool.FindFreeElement(Desc, SplatDepthStencilBuffer, TEXT("DistanceFieldAOSplatDepthBuffer"));
+
+			SetRenderTarget(RHICmdList, InterpolationTarget->GetRenderTargetItem().TargetableTexture, SplatDepthStencilBuffer->GetRenderTargetItem().TargetableTexture);
+			RHICmdList.Clear(true, FLinearColor(0, 0, 0, 0), true, 0, true, 0, FIntRect());
+
+			for (int32 ViewIndex = 0; ViewIndex < Views.Num(); ViewIndex++)
+			{
+				FViewInfo& View = Views[ViewIndex];
+				SetupDepthStencil(RHICmdList, View, DistanceFieldNormal, DepthLevel, DestLevelDownsampleFactor);
+			}
+
+			GRenderTargetPool.VisualizeTexture.SetCheckPoint(RHICmdList, SplatDepthStencilBuffer);
+
+			// Restore render target after checkpoint
+			SetRenderTarget(RHICmdList, InterpolationTarget->GetRenderTargetItem().TargetableTexture, SplatDepthStencilBuffer ? SplatDepthStencilBuffer->GetRenderTargetItem().TargetableTexture : NULL);
+		}
+		else
+		{
+			SetRenderTarget(RHICmdList, InterpolationTarget->GetRenderTargetItem().TargetableTexture, NULL);
+			RHICmdList.Clear(true, FLinearColor(0, 0, 0, 0), false, 0, false, 0, FIntRect());
+		}
+
 		SCOPED_DRAW_EVENT(RHICmdList, IrradianceCacheSplat, DEC_SCENE_ITEMS);
-		const int32 MaxAllowedLevel = GAOMaxSupportedLevel;
 
-		SetRenderTarget(RHICmdList, InterpolationTarget->GetRenderTargetItem().TargetableTexture, NULL);
-		RHICmdList.Clear(true, FLinearColor(0, 0, 0, 0), false, 0, false, 0, FIntRect());
+		static int32 NumInterpolationSections = 8;
+		if (GCircleVertexBuffer.NumSections != NumInterpolationSections)
+		{
+			GCircleVertexBuffer.NumSections = NumInterpolationSections;
+			GCircleVertexBuffer.UpdateRHI();
+		}
 
 		for (int32 ViewIndex = 0; ViewIndex < Views.Num(); ViewIndex++)
 		{
@@ -2602,19 +3069,50 @@ void RenderIrradianceCacheInterpolation(
 			const FScene* Scene = (const FScene*)View.Family->Scene;
 			FSurfaceCacheResources& SurfaceCacheResources = *Scene->SurfaceCacheResources;
 
-			const uint32 DownsampledViewSizeX = FMath::DivideAndRoundUp(View.ViewRect.Width(), DestLevelDownsampleFactor);
-			const uint32 DownsampledViewSizeY = FMath::DivideAndRoundUp(View.ViewRect.Height(), DestLevelDownsampleFactor);
+			const uint32 DownsampledViewSizeX = FMath::DivideAndRoundDown(View.ViewRect.Width(), DestLevelDownsampleFactor);
+			const uint32 DownsampledViewSizeY = FMath::DivideAndRoundDown(View.ViewRect.Height(), DestLevelDownsampleFactor);
 
 			RHICmdList.SetViewport(0, 0, 0.0f, DownsampledViewSizeX, DownsampledViewSizeY, 1.0f);
 			RHICmdList.SetRasterizerState(TStaticRasterizerState<FM_Solid, CM_None>::GetRHI());
-			//@todo - render front faces with depth testing, requires depth buffer
-			RHICmdList.SetDepthStencilState(TStaticDepthStencilState<false, CF_Always>::GetRHI());
-			RHICmdList.SetBlendState(TStaticBlendState<CW_RGBA, BO_Add, BF_One, BF_One, BO_Add, BF_One, BF_One>::GetRHI());
 
-			FVector2D NormalizedOffsetToPixelCenter2(.5f / DownsampledViewSizeX - .5f / (DownsampledViewSizeX * DestLevelDownsampleFactor), .5f / DownsampledViewSizeY - .5f / (DownsampledViewSizeY * DestLevelDownsampleFactor));
+			if (bFinalInterpolation && (GAOInterpolationDepthTesting || GAOInterpolationStencilTesting))
+			{
+				if (GAOInterpolationDepthTesting)
+				{
+					if (GAOInterpolationStencilTesting)
+					{
+						// Depth tests enabled, pass stencil test if stencil is zero
+						RHICmdList.SetDepthStencilState(TStaticDepthStencilState<
+							false,CF_GreaterEqual,
+							true,CF_Equal,SO_Keep,SO_Keep,SO_Keep,
+							false,CF_Always,SO_Keep,SO_Keep,SO_Keep,
+							0xff,0xff
+							>::GetRHI());
+					}
+					else
+					{
+						// Depth tests enabled
+						RHICmdList.SetDepthStencilState(TStaticDepthStencilState<false, CF_GreaterEqual>::GetRHI());
+					}
+				}
+				else if (GAOInterpolationStencilTesting)
+				{
+					// Pass stencil test if stencil is zero
+					RHICmdList.SetDepthStencilState(TStaticDepthStencilState<
+						false,CF_Always,
+						true,CF_Equal,SO_Keep,SO_Keep,SO_Keep,
+						false,CF_Always,SO_Keep,SO_Keep,SO_Keep,
+						0xff,0xff
+						>::GetRHI());
+				}
+			}
+			else
+			{
+				//@todo - render front faces with depth testing, requires depth buffer
+				RHICmdList.SetDepthStencilState(TStaticDepthStencilState<false, CF_Always>::GetRHI());
+			}
 
-			// Convert to NDC 
-			NormalizedOffsetToPixelCenter2 = NormalizedOffsetToPixelCenter2 * FVector2D(2, -2);
+			RHICmdList.SetBlendState(TStaticBlendState<CW_RGBA, BO_Add, BF_One, BF_One, BO_Add, BF_One, BF_One>::GetRHI());
 
 			const FIntPoint BufferSize = GSceneRenderTargets.GetBufferSizeXY();
 			const float InvBufferSizeX = 1.0f / BufferSize.X;
@@ -2639,7 +3137,7 @@ void RenderIrradianceCacheInterpolation(
 				TShaderMapRef<TIrradianceCacheSplatVS<true> > VertexShader(View.ShaderMap);
 				TShaderMapRef<TIrradianceCacheSplatPS<true> > PixelShader(View.ShaderMap);
 
-				for (int32 SplatSourceDepthLevel = MaxAllowedLevel; SplatSourceDepthLevel >= DepthLevel; SplatSourceDepthLevel--)
+				for (int32 SplatSourceDepthLevel = GAOMaxLevel; SplatSourceDepthLevel >= FMath::Max(DepthLevel, GAOMinLevel); SplatSourceDepthLevel--)
 				{
 					static FGlobalBoundShaderState BoundShaderState;
 					
@@ -2656,7 +3154,7 @@ void RenderIrradianceCacheInterpolation(
 				TShaderMapRef<TIrradianceCacheSplatVS<false> > VertexShader(View.ShaderMap);
 				TShaderMapRef<TIrradianceCacheSplatPS<false> > PixelShader(View.ShaderMap);
 
-				for (int32 SplatSourceDepthLevel = MaxAllowedLevel; SplatSourceDepthLevel >= DepthLevel; SplatSourceDepthLevel--)
+				for (int32 SplatSourceDepthLevel = GAOMaxLevel; SplatSourceDepthLevel >= FMath::Max(DepthLevel, GAOMinLevel); SplatSourceDepthLevel--)
 				{
 					static FGlobalBoundShaderState BoundShaderState;
 					
@@ -2836,17 +3334,20 @@ bool FDeferredShadingSceneRenderer::RenderDistanceFieldAOSurfaceCache(FRHIComman
 			if (!Scene->SurfaceCacheResources)
 			{
 				Scene->SurfaceCacheResources = new FSurfaceCacheResources();
-				Scene->SurfaceCacheResources->InitResource();
 			}
 
-			uint32 ClearValues[4] = {0};
-			const int32 MaxAllowedLevel = GAOMaxSupportedLevel;
 			FSurfaceCacheResources& SurfaceCacheResources = *Scene->SurfaceCacheResources;
+			FIntPoint DownsampledViewRect = FIntPoint::DivideAndRoundDown(View.ViewRect.Size(), GAODownsampleFactor);
+			int32 SlackAmount = 2;
+			int32 MaxSurfaceCacheSamples = SlackAmount * DownsampledViewRect.X * DownsampledViewRect.Y / (1 << (2 * GAOMinLevel * GAOPowerOfTwoBetweenLevels));
+			SurfaceCacheResources.AllocateFor(GAOMinLevel, GAOMaxLevel, MaxSurfaceCacheSamples);
+
+			uint32 ClearValues[4] = {0};
 
 			if (!SurfaceCacheResources.bClearedResources || !GAOReuseAcrossFrames)
 			{
 				// Reset the number of active cache records to 0
-				for (int32 DepthLevel = MaxAllowedLevel; DepthLevel >= 0; DepthLevel--)
+				for (int32 DepthLevel = GAOMaxLevel; DepthLevel >= GAOMinLevel; DepthLevel--)
 				{
 					RHICmdList.ClearUAV(SurfaceCacheResources.Level[DepthLevel]->ScatterDrawParameters.UAV, ClearValues);
 				}
@@ -2869,8 +3370,8 @@ bool FDeferredShadingSceneRenderer::RenderDistanceFieldAOSurfaceCache(FRHIComman
 				TRefCountPtr<IPooledRenderTarget> DistanceFieldNormal;
 
 				{
-					const FIntPoint ExpandedBufferSize = GetBufferSizeForAO();
-					FPooledRenderTargetDesc Desc(FPooledRenderTargetDesc::Create2DDesc(ExpandedBufferSize / GAODownsampleFactor, PF_FloatRGBA, TexCreate_None, TexCreate_RenderTargetable | TexCreate_UAV, false));
+					const FIntPoint BufferSize = GetBufferSizeForAO();
+					FPooledRenderTargetDesc Desc(FPooledRenderTargetDesc::Create2DDesc(BufferSize, PF_FloatRGBA, TexCreate_None, TexCreate_RenderTargetable | TexCreate_UAV, false));
 					GRenderTargetPool.FindFreeElement(Desc, DistanceFieldNormal, TEXT("DistanceFieldNormal"));
 				}
 
@@ -2879,50 +3380,53 @@ bool FDeferredShadingSceneRenderer::RenderDistanceFieldAOSurfaceCache(FRHIComman
 
 				GRenderTargetPool.VisualizeTexture.SetCheckPoint(RHICmdList, DistanceFieldNormal);
 
-				if (GAOReuseAcrossFrames && GAOTrimOldRecordsFraction > 0)
+				if (GAOReuseAcrossFrames)
 				{
 					SCOPED_DRAW_EVENT(RHICmdList, TrimRecords, DEC_SCENE_ITEMS);
 
 					// Copy and trim last frame's surface cache samples
-					for (int32 DepthLevel = FMath::Min(GAOMaxLevel, MaxAllowedLevel); DepthLevel >= FMath::Max(GAOMinLevel, 0); DepthLevel--)
+					for (int32 DepthLevel = GAOMaxLevel; DepthLevel >= GAOMinLevel; DepthLevel--)
 					{
-						{	
-							TShaderMapRef<FSetupCopyIndirectArgumentsCS> ComputeShader(View.ShaderMap);
+						if (GetTrimOldRecordsFraction(DepthLevel) > 0)
+						{
+							{	
+								TShaderMapRef<FSetupCopyIndirectArgumentsCS> ComputeShader(View.ShaderMap);
 
-							RHICmdList.SetComputeShader(ComputeShader->GetComputeShader());
-							ComputeShader->SetParameters(RHICmdList, View, DepthLevel);
-							DispatchComputeShader(RHICmdList, *ComputeShader, 1, 1, 1);
+								RHICmdList.SetComputeShader(ComputeShader->GetComputeShader());
+								ComputeShader->SetParameters(RHICmdList, View, DepthLevel);
+								DispatchComputeShader(RHICmdList, *ComputeShader, 1, 1, 1);
 
-							ComputeShader->UnsetParameters(RHICmdList);
-						}
+								ComputeShader->UnsetParameters(RHICmdList);
+							}
 
-						{
-							TShaderMapRef<FCopyIrradianceCacheSamplesCS> ComputeShader(View.ShaderMap);
+							{
+								TShaderMapRef<FCopyIrradianceCacheSamplesCS> ComputeShader(View.ShaderMap);
 
-							RHICmdList.SetComputeShader(ComputeShader->GetComputeShader());
-							ComputeShader->SetParameters(RHICmdList, View, DepthLevel);
-							DispatchIndirectComputeShader(RHICmdList, *ComputeShader, SurfaceCacheResources.DispatchParameters.Buffer, 0);
+								RHICmdList.SetComputeShader(ComputeShader->GetComputeShader());
+								ComputeShader->SetParameters(RHICmdList, View, DepthLevel);
+								DispatchIndirectComputeShader(RHICmdList, *ComputeShader, SurfaceCacheResources.DispatchParameters.Buffer, 0);
 
-							ComputeShader->UnsetParameters(RHICmdList);
-						}
+								ComputeShader->UnsetParameters(RHICmdList);
+							}
 
-						Swap(SurfaceCacheResources.Level[DepthLevel], SurfaceCacheResources.TempResources);
+							Swap(SurfaceCacheResources.Level[DepthLevel], SurfaceCacheResources.TempResources);
+						}
 					}
 				}
 
 				// Start from the highest depth, which will consider the fewest potential shading points
 				// Each level potentially prevents the next higher resolution level from doing expensive shading work
 				// This is the core of the surface cache algorithm
-				for (int32 DepthLevel = FMath::Min(GAOMaxLevel, MaxAllowedLevel); DepthLevel >= FMath::Max(GAOMinLevel, 0); DepthLevel--)
+				for (int32 DepthLevel = GAOMaxLevel; DepthLevel >= GAOMinLevel; DepthLevel--)
 				{
-					int32 DestLevelDownsampleFactor = GAODownsampleFactor * (1 << (DepthLevel * 2));
+					int32 DestLevelDownsampleFactor = GAODownsampleFactor * (1 << (DepthLevel * GAOPowerOfTwoBetweenLevels));
 
 					SCOPED_DRAW_EVENTF(RHICmdList, Level, DEC_SCENE_ITEMS, TEXT("Level_%d"), DepthLevel);
 
 					TRefCountPtr<IPooledRenderTarget> DistanceFieldAOIrradianceCacheSplat;
 
 					{
-						FIntPoint AOBufferSize = FIntPoint::DivideAndRoundUp(GSceneRenderTargets.GetBufferSizeXY(), DestLevelDownsampleFactor);
+						FIntPoint AOBufferSize = FIntPoint::DivideAndRoundDown(GSceneRenderTargets.GetBufferSizeXY(), DestLevelDownsampleFactor);
 						FPooledRenderTargetDesc Desc(FPooledRenderTargetDesc::Create2DDesc(AOBufferSize, PF_FloatRGBA, TexCreate_None, TexCreate_RenderTargetable | TexCreate_UAV, false));
 						GRenderTargetPool.FindFreeElement(Desc, DistanceFieldAOIrradianceCacheSplat, TEXT("DistanceFieldAOIrradianceCacheSplat"));
 					}
@@ -2950,7 +3454,7 @@ bool FDeferredShadingSceneRenderer::RenderDistanceFieldAOSurfaceCache(FRHIComman
 						{
 							const FViewInfo& View = Views[ViewIndex];
 
-							FIntPoint DownsampledViewSize = FIntPoint::DivideAndRoundUp(View.ViewRect.Size(), DestLevelDownsampleFactor);
+							FIntPoint DownsampledViewSize = FIntPoint::DivideAndRoundDown(View.ViewRect.Size(), DestLevelDownsampleFactor);
 							uint32 GroupSizeX = FMath::DivideAndRoundUp(DownsampledViewSize.X, GDistanceFieldAOTileSizeX);
 							uint32 GroupSizeY = FMath::DivideAndRoundUp(DownsampledViewSize.Y, GDistanceFieldAOTileSizeY);
 
@@ -2993,8 +3497,8 @@ bool FDeferredShadingSceneRenderer::RenderDistanceFieldAOSurfaceCache(FRHIComman
 				TRefCountPtr<IPooledRenderTarget> IrradianceCacheAccumulation;
 
 				{
-					FIntPoint ExpandedBufferSize = GetBufferSizeForAO();
-					FPooledRenderTargetDesc Desc(FPooledRenderTargetDesc::Create2DDesc(ExpandedBufferSize / GAODownsampleFactor, PF_FloatRGBA, TexCreate_None, TexCreate_RenderTargetable | TexCreate_UAV, false));
+					FIntPoint BufferSize = GetBufferSizeForAO();
+					FPooledRenderTargetDesc Desc(FPooledRenderTargetDesc::Create2DDesc(BufferSize, PF_FloatRGBA, TexCreate_None, TexCreate_RenderTargetable | TexCreate_UAV, false));
 					GRenderTargetPool.FindFreeElement(Desc, IrradianceCacheAccumulation, TEXT("DistanceFieldAOIrradianceCacheAccumulation"));
 				}
 
@@ -3204,8 +3708,8 @@ void FDeferredShadingSceneRenderer::RenderMeshDistanceFieldVisualization(FRHICom
 				TRefCountPtr<IPooledRenderTarget> VisualizeResultRT;
 
 				{
-					const FIntPoint ExpandedBufferSize = GetBufferSizeForAO();
-					FPooledRenderTargetDesc Desc(FPooledRenderTargetDesc::Create2DDesc(ExpandedBufferSize / GAODownsampleFactor, PF_FloatRGBA, TexCreate_None, TexCreate_RenderTargetable | TexCreate_UAV, false));
+					const FIntPoint BufferSize = GetBufferSizeForAO();
+					FPooledRenderTargetDesc Desc(FPooledRenderTargetDesc::Create2DDesc(BufferSize, PF_FloatRGBA, TexCreate_None, TexCreate_RenderTargetable | TexCreate_UAV, false));
 					GRenderTargetPool.FindFreeElement(Desc, VisualizeResultRT, TEXT("VisualizeDistanceField"));
 				}
 
@@ -3262,7 +3766,7 @@ void FDeferredShadingSceneRenderer::RenderMeshDistanceFieldVisualization(FRHICom
 							View.ViewRect.Min.X / GAODownsampleFactor, View.ViewRect.Min.Y / GAODownsampleFactor, 
 							View.ViewRect.Width() / GAODownsampleFactor, View.ViewRect.Height() / GAODownsampleFactor,
 							FIntPoint(View.ViewRect.Width(), View.ViewRect.Height()),
-							GetBufferSizeForAO() / GAODownsampleFactor,
+							GetBufferSizeForAO(),
 							*VertexShader);
 					}
 				}
diff --git a/Engine/Source/Runtime/Renderer/Private/DistanceFieldSurfaceCacheLighting.h b/Engine/Source/Runtime/Renderer/Private/DistanceFieldSurfaceCacheLighting.h
--- a/Engine/Source/Runtime/Renderer/Private/DistanceFieldSurfaceCacheLighting.h
+++ b/Engine/Source/Runtime/Renderer/Private/DistanceFieldSurfaceCacheLighting.h
@@ -8,9 +8,7 @@
 
 #include "ScreenRendering.h"
 
-const static int32 GAOMaxSupportedLevel = 4;
-//@todo - derive from worst case
-const static uint32 GMaxIrradianceCacheSamples = 100000;
+const static int32 GAOMaxSupportedLevel = 6;
 /** Number of cone traced directions. */
 const int32 NumConeSampleDirections = 9;
 
@@ -31,16 +29,25 @@ public:
 class FRefinementLevelResources
 {
 public:
+
+	FRefinementLevelResources()
+	{
+		MaxIrradianceCacheSamples = 0;
+	}
+
 	void InitDynamicRHI()
 	{
-		//@todo - handle max exceeded
-		PositionAndRadius.Initialize(sizeof(float) * 4, GMaxIrradianceCacheSamples, PF_A32B32G32R32F, BUF_Static);
-		OccluderRadius.Initialize(sizeof(float), GMaxIrradianceCacheSamples, PF_R32_FLOAT, BUF_Static);
-		Normal.Initialize(sizeof(FFloat16Color), GMaxIrradianceCacheSamples, PF_FloatRGBA, BUF_Static);
-		BentNormal.Initialize(sizeof(FFloat16Color), GMaxIrradianceCacheSamples, PF_FloatRGBA, BUF_Static);
-		ScatterDrawParameters.Initialize(sizeof(uint32), 4, PF_R32_UINT, BUF_Static | BUF_DrawIndirect);
-		SavedStartIndex.Initialize(sizeof(uint32), 1, PF_R32_UINT, BUF_Static);
-		TileCoordinate.Initialize(sizeof(uint16)* 2, GMaxIrradianceCacheSamples, PF_R16G16_UINT, BUF_Static);
+		if (MaxIrradianceCacheSamples > 0)
+		{
+			//@todo - handle max exceeded
+			PositionAndRadius.Initialize(sizeof(float) * 4, MaxIrradianceCacheSamples, PF_A32B32G32R32F, BUF_Static);
+			OccluderRadius.Initialize(sizeof(float), MaxIrradianceCacheSamples, PF_R32_FLOAT, BUF_Static);
+			Normal.Initialize(sizeof(FFloat16Color), MaxIrradianceCacheSamples, PF_FloatRGBA, BUF_Static);
+			BentNormal.Initialize(sizeof(FFloat16Color), MaxIrradianceCacheSamples, PF_FloatRGBA, BUF_Static);
+			ScatterDrawParameters.Initialize(sizeof(uint32), 4, PF_R32_UINT, BUF_Static | BUF_DrawIndirect);
+			SavedStartIndex.Initialize(sizeof(uint32), 1, PF_R32_UINT, BUF_Static);
+			TileCoordinate.Initialize(sizeof(uint16)* 2, MaxIrradianceCacheSamples, PF_R16G16_UINT, BUF_Static);
+		}
 	}
 
 	void ReleaseDynamicRHI()
@@ -54,6 +61,8 @@ public:
 		TileCoordinate.Release();
 	}
 
+	int32 MaxIrradianceCacheSamples;
+
 	FRWBuffer PositionAndRadius;
 	FRWBuffer OccluderRadius;
 	FRWBuffer Normal;
@@ -76,6 +85,9 @@ public:
 		}
 
 		TempResources = new FRefinementLevelResources();
+
+		MinLevel = 100;
+		MaxLevel = 0;
 	}
 
 	~FSurfaceCacheResources()
@@ -88,11 +100,55 @@ public:
 		delete TempResources;
 	}
 
+	void AllocateFor(int32 InMinLevel, int32 InMaxLevel, int32 MaxSurfaceCacheSamples)
+	{
+		check(InMaxLevel < ARRAY_COUNT(Level));
+
+		bool bReallocate = false;
+
+		if (MinLevel > InMinLevel)
+		{
+			MinLevel = InMinLevel;
+			bReallocate = true;
+		}
+
+		if (MaxLevel < InMaxLevel)
+		{
+			MaxLevel = InMaxLevel;
+			bReallocate = true;
+		}
+
+		for (int32 LevelIndex = MinLevel; LevelIndex <= MaxLevel; LevelIndex++)
+		{
+			//@todo - would like to allocate each level based on its worst case, but not possible right now due to swapping with TempResources
+			if (Level[LevelIndex]->MaxIrradianceCacheSamples < MaxSurfaceCacheSamples)
+			{
+				Level[LevelIndex]->MaxIrradianceCacheSamples = MaxSurfaceCacheSamples;
+				bReallocate = true;
+			}
+		}
+
+		if (TempResources->MaxIrradianceCacheSamples < MaxSurfaceCacheSamples)
+		{
+			bReallocate = true;
+			TempResources->MaxIrradianceCacheSamples = MaxSurfaceCacheSamples;
+		}
+
+		if (!IsInitialized())
+		{
+			InitResource();
+		}
+		else if (bReallocate)
+		{
+			UpdateRHI();
+		}
+	}
+
 	virtual void InitDynamicRHI() override
 	{
 		DispatchParameters.Initialize(sizeof(uint32), 3, PF_R32_UINT, BUF_Static | BUF_DrawIndirect);
 
-		for (int32 i = 0; i < ARRAY_COUNT(Level); i++)
+		for (int32 i = MinLevel; i <= MaxLevel; i++)
 		{
 			Level[i]->InitDynamicRHI();
 		}
@@ -106,7 +162,7 @@ public:
 	{
 		DispatchParameters.Release();
 
-		for (int32 i = 0; i < ARRAY_COUNT(Level); i++)
+		for (int32 i = MinLevel; i <= MaxLevel; i++)
 		{
 			Level[i]->ReleaseDynamicRHI();
 		}
@@ -122,6 +178,11 @@ public:
 	FRefinementLevelResources* TempResources;
 
 	bool bClearedResources;
+
+private:
+
+	int32 MinLevel;
+	int32 MaxLevel;
 };
 
 /**  */
@@ -206,71 +267,3 @@ private:
 	FShaderParameter AOMaxViewDistance;
 };
 
-class FQuadVertexBuffer : public FVertexBuffer
-{
-public:
-
-	virtual void InitRHI() override
-	{
-		// Used as a non-indexed triangle list, so 6 vertices per quad
-		const uint32 Size = 6 * sizeof(FScreenVertex);
-		FRHIResourceCreateInfo CreateInfo;
-		VertexBufferRHI = RHICreateVertexBuffer(Size, BUF_Static, CreateInfo);
-		void* Buffer = RHILockVertexBuffer(VertexBufferRHI, 0, Size, RLM_WriteOnly);
-		FScreenVertex* DestVertex = (FScreenVertex*)Buffer;
-
-		DestVertex[0].Position = FVector2D(0, 0);
-		DestVertex[0].UV = FVector2D(1, 1);
-		DestVertex[1].Position = FVector2D(0, 0);
-		DestVertex[1].UV = FVector2D(1, 0);
-		DestVertex[2].Position = FVector2D(0, 0);
-		DestVertex[2].UV = FVector2D(0, 1);
-		DestVertex[3].Position = FVector2D(0, 0);
-		DestVertex[3].UV = FVector2D(1, 0);
-		DestVertex[4].Position = FVector2D(0, 0);
-		DestVertex[4].UV = FVector2D(0, 0);
-		DestVertex[5].Position = FVector2D(0, 0);
-		DestVertex[5].UV = FVector2D(0, 1);
-
-		RHIUnlockVertexBuffer(VertexBufferRHI);      
-	}
-};
-
-extern TGlobalResource<FQuadVertexBuffer> GQuadVertexBuffer;
-
-class FCircleVertexBuffer : public FVertexBuffer
-{
-public:
-
-	virtual void InitRHI() override
-	{
-		int32 NumSections = 8;
-
-		// Used as a non-indexed triangle list, so 3 vertices per triangle
-		const uint32 Size = 3 * NumSections * sizeof(FScreenVertex);
-		FRHIResourceCreateInfo CreateInfo;
-		VertexBufferRHI = RHICreateVertexBuffer(Size, BUF_Static, CreateInfo);
-		void* Buffer = RHILockVertexBuffer(VertexBufferRHI, 0, Size, RLM_WriteOnly);
-		FScreenVertex* DestVertex = (FScreenVertex*)Buffer;
-
-		for (int32 SectionIndex = 0; SectionIndex < NumSections; SectionIndex++)
-		{
-			float Fraction = SectionIndex / (float)NumSections;
-			float CurrentAngle = Fraction * 2 * PI;
-			float NextAngle = ((SectionIndex + 1) / (float)NumSections) * 2 * PI;
-			FVector2D CurrentPosition(FMath::Cos(CurrentAngle), FMath::Sin(CurrentAngle));
-			FVector2D NextPosition(FMath::Cos(NextAngle), FMath::Sin(NextAngle));
-
-			DestVertex[SectionIndex * 3 + 0].Position = FVector2D(0, 0);
-			DestVertex[SectionIndex * 3 + 0].UV = CurrentPosition;
-			DestVertex[SectionIndex * 3 + 1].Position = FVector2D(0, 0);
-			DestVertex[SectionIndex * 3 + 1].UV = NextPosition;
-			DestVertex[SectionIndex * 3 + 2].Position = FVector2D(0, 0);
-			DestVertex[SectionIndex * 3 + 2].UV = FVector2D(.5f, .5f);
-		}
-
-		RHIUnlockVertexBuffer(VertexBufferRHI);      
-	}
-};
-
-extern TGlobalResource<FCircleVertexBuffer> GCircleVertexBuffer;
diff --git a/Engine/Source/Runtime/Renderer/Private/PostProcess/VisualizeTexture.h b/Engine/Source/Runtime/Renderer/Private/PostProcess/VisualizeTexture.h
--- a/Engine/Source/Runtime/Renderer/Private/PostProcess/VisualizeTexture.h
+++ b/Engine/Source/Runtime/Renderer/Private/PostProcess/VisualizeTexture.h
@@ -25,6 +25,7 @@ public:
 	/**
 	 * calling this allows to grab the state of the texture at this point to be queried by visualizetexture e.g. "vis LightAttenuation@2"
 	 * @param PooledRenderTarget 0 is silently ignored
+	 * Warning: this may change the active render target and other state
 	 */
 	void SetCheckPoint(FRHICommandList& RHICmdList, const IPooledRenderTarget* PooledRenderTarget);
 
