fixed build for clang

diff --git a/Engine/Source/Runtime/Renderer/Private/PostProcess/RenderTargetPool.cpp b/Engine/Source/Runtime/Renderer/Private/PostProcess/RenderTargetPool.cpp
--- a/Engine/Source/Runtime/Renderer/Private/PostProcess/RenderTargetPool.cpp
+++ b/Engine/Source/Runtime/Renderer/Private/PostProcess/RenderTargetPool.cpp
@@ -122,7 +122,7 @@ FRenderTargetPool::FRenderTargetPool()
 
 static void LogVRamUsage(FPooledRenderTarget& Ref, FTextureRHIParamRef TexRef)
 {
-	if(IsValidRef(TexRef))
+	if(TexRef)
 	{
 		FRHIResourceInfo Info;
 
