DFAO refinement buffers are rounded up again, solves over-sampling on the edges of the screen

diff --git a/Engine/Shaders/DistanceFieldSurfaceCacheLighting.usf b/Engine/Shaders/DistanceFieldSurfaceCacheLighting.usf
--- a/Engine/Shaders/DistanceFieldSurfaceCacheLighting.usf
+++ b/Engine/Shaders/DistanceFieldSurfaceCacheLighting.usf
@@ -261,7 +261,6 @@ struct FIrradianceCacheSplatVertexOutput
 	nointerpolation float4 PositionRadius : TEXCOORD0;
 	nointerpolation float3 Normal : TEXCOORD1;
 	nointerpolation float3 BentNormal : TEXCOORD2;
-	nointerpolation uint RecordIndex : TEXCOORD3;
 #if USE_SCREENVECTOR_WORLD_POSITION
 	float3 ScreenVector : TEXCOORD4;
 #endif
@@ -324,6 +323,8 @@ void IrradianceCacheSplatVS(
 	VirtualRadius += HackExpand * ViewToCenterLength * 4;
 #endif
 
+	float3 RecordNormal = DecodeNormalForAO(IrradianceCacheNormal.Load(JobIndex).xyz);
+
 	float2 CornerUVs = InUV;
 	float3 CornerPosition = VirtualPosition + (View.ViewRight * CornerUVs.x + View.ViewUp * CornerUVs.y) * VirtualRadius;
 
@@ -337,9 +338,8 @@ void IrradianceCacheSplatVS(
 	Output.Position.xy += NormalizedOffsetToPixelCenter * Output.Position.w;
 
 	Output.PositionRadius = PositionAndRadius; 
-	Output.Normal = DecodeNormalForAO(IrradianceCacheNormal.Load(JobIndex).xyz);
+	Output.Normal = RecordNormal;
 	Output.BentNormal = DecodeNormalForAO(IrradianceCacheBentNormal.Load(JobIndex).xyz);
-	Output.RecordIndex = JobIndex;
 } 
 
 float InterpolationAngleNormalization; 
@@ -373,6 +373,7 @@ void IrradianceCacheSplatPS(
 	
 	float Distance = length(OpaqueWorldPosition - Input.PositionRadius.xyz);
 	float DistanceError = saturate(Distance / Input.PositionRadius.w);
+	float Weight = 0;
 
 	BRANCH
 	if (DistanceError < 1 && SceneW < AOMaxViewDistance)
@@ -391,7 +392,7 @@ void IrradianceCacheSplatPS(
 		float PointBehindPlaneError = min(max(PlaneDistance * InvMinCosPointBehindPlane, 0.0f), DistanceToVertex / 3.0f);
 
 		float PrecisionScale = .1f;
-		float Weight = saturate(PrecisionScale * (1 - max(DistanceError, max(NormalError, PointBehindPlaneError))));
+		Weight = saturate(PrecisionScale * (1 - max(DistanceError, max(NormalError, PointBehindPlaneError))));
 		//float Weight = saturate(PrecisionScale * (1 - max(DistanceError, NormalError)));
 		//float Weight = saturate(PrecisionScale * (1 - DistanceError));
 		//float Weight = saturate(PrecisionScale * (1 - NormalError));
@@ -405,7 +406,7 @@ void IrradianceCacheSplatPS(
 	#if VISUALIZE_SPLAT_OVERDRAW && FINAL_INTERPOLATION_PASS
 		OutColor.rgb = 0;
 		//OutColor.a = Input.Position.w <= SceneW ? .01f : 0.0f;
-		OutColor.a = .005f;
+		OutColor.a = .005f + (Weight > .001f ? .01f : 0);
 	#endif
 
 	#define VISUALIZE_RECORD_POINTS 0
diff --git a/Engine/Source/Runtime/Renderer/Private/DistanceFieldSurfaceCacheLighting.cpp b/Engine/Source/Runtime/Renderer/Private/DistanceFieldSurfaceCacheLighting.cpp
--- a/Engine/Source/Runtime/Renderer/Private/DistanceFieldSurfaceCacheLighting.cpp
+++ b/Engine/Source/Runtime/Renderer/Private/DistanceFieldSurfaceCacheLighting.cpp
@@ -751,12 +751,12 @@ public:
 		SetShaderValue(RHICmdList, ShaderRHI, CurrentLevelDownsampleFactor, CurrentLevelDownsampleFactorValue);
 
 		// Round up, to match render target allocation
-		const FVector2D AOBufferSizeValue = FIntPoint::DivideAndRoundDown(GSceneRenderTargets.GetBufferSizeXY(), CurrentLevelDownsampleFactorValue);
+		const FVector2D AOBufferSizeValue = FIntPoint::DivideAndRoundUp(GSceneRenderTargets.GetBufferSizeXY(), CurrentLevelDownsampleFactorValue);
 		SetShaderValue(RHICmdList, ShaderRHI, AOBufferSize, AOBufferSizeValue);
 
 		SetShaderValue(RHICmdList, ShaderRHI, DownsampleFactorToBaseLevel, CurrentLevelDownsampleFactorValue / GAODownsampleFactor);
 
-		const FIntPoint DownsampledBufferSize(GSceneRenderTargets.GetBufferSizeXY() / FIntPoint(GAODownsampleFactor, GAODownsampleFactor));
+		const FIntPoint DownsampledBufferSize = GetBufferSizeForAO();
 		const FVector2D BaseLevelBufferSizeValue(1.0f / DownsampledBufferSize.X, 1.0f / DownsampledBufferSize.Y);
 		SetShaderValue(RHICmdList, ShaderRHI, BaseLevelTexelSize, BaseLevelBufferSizeValue);
 	}
@@ -1792,7 +1792,7 @@ public:
 		SetShaderValue(RHICmdList, ShaderRHI, InterpolationRadiusScale, (bFinalInterpolationPass ? GAOInterpolationRadiusScale : 1.0f));
 		SetShaderValue(RHICmdList, ShaderRHI, NormalizedOffsetToPixelCenter, NormalizedOffsetToPixelCenterValue);
 
-		const FIntPoint AOViewRectSize = FIntPoint::DivideAndRoundDown(View.ViewRect.Size(), CurrentLevelDownsampleFactorValue);
+		const FIntPoint AOViewRectSize = FIntPoint::DivideAndRoundUp(View.ViewRect.Size(), CurrentLevelDownsampleFactorValue);
 		const FVector2D HackExpandValue(.5f / AOViewRectSize.X, .5f / AOViewRectSize.Y);
 		SetShaderValue(RHICmdList, ShaderRHI, HackExpand, HackExpandValue);
 	}
@@ -3011,7 +3011,7 @@ void RenderIrradianceCacheInterpolation(
 			const FScene* Scene = (const FScene*)View.Family->Scene;
 			FSurfaceCacheResources& SurfaceCacheResources = *Scene->SurfaceCacheResources;
 
-			FIntPoint DownsampledViewSize = FIntPoint::DivideAndRoundDown(View.ViewRect.Size(), DestLevelDownsampleFactor);
+			FIntPoint DownsampledViewSize = FIntPoint::DivideAndRoundUp(View.ViewRect.Size(), DestLevelDownsampleFactor);
 			uint32 GroupSizeX = FMath::DivideAndRoundUp(DownsampledViewSize.X, GDistanceFieldAOTileSizeX);
 			uint32 GroupSizeY = FMath::DivideAndRoundUp(DownsampledViewSize.Y, GDistanceFieldAOTileSizeY);
 
@@ -3069,8 +3069,8 @@ void RenderIrradianceCacheInterpolation(
 			const FScene* Scene = (const FScene*)View.Family->Scene;
 			FSurfaceCacheResources& SurfaceCacheResources = *Scene->SurfaceCacheResources;
 
-			const uint32 DownsampledViewSizeX = FMath::DivideAndRoundDown(View.ViewRect.Width(), DestLevelDownsampleFactor);
-			const uint32 DownsampledViewSizeY = FMath::DivideAndRoundDown(View.ViewRect.Height(), DestLevelDownsampleFactor);
+			const uint32 DownsampledViewSizeX = FMath::DivideAndRoundUp(View.ViewRect.Width(), DestLevelDownsampleFactor);
+			const uint32 DownsampledViewSizeY = FMath::DivideAndRoundUp(View.ViewRect.Height(), DestLevelDownsampleFactor);
 
 			RHICmdList.SetViewport(0, 0, 0.0f, DownsampledViewSizeX, DownsampledViewSizeY, 1.0f);
 			RHICmdList.SetRasterizerState(TStaticRasterizerState<FM_Solid, CM_None>::GetRHI());
@@ -3426,7 +3426,7 @@ bool FDeferredShadingSceneRenderer::RenderDistanceFieldAOSurfaceCache(FRHIComman
 					TRefCountPtr<IPooledRenderTarget> DistanceFieldAOIrradianceCacheSplat;
 
 					{
-						FIntPoint AOBufferSize = FIntPoint::DivideAndRoundDown(GSceneRenderTargets.GetBufferSizeXY(), DestLevelDownsampleFactor);
+						FIntPoint AOBufferSize = FIntPoint::DivideAndRoundUp(GSceneRenderTargets.GetBufferSizeXY(), DestLevelDownsampleFactor);
 						FPooledRenderTargetDesc Desc(FPooledRenderTargetDesc::Create2DDesc(AOBufferSize, PF_FloatRGBA, TexCreate_None, TexCreate_RenderTargetable | TexCreate_UAV, false));
 						GRenderTargetPool.FindFreeElement(Desc, DistanceFieldAOIrradianceCacheSplat, TEXT("DistanceFieldAOIrradianceCacheSplat"));
 					}
@@ -3454,7 +3454,7 @@ bool FDeferredShadingSceneRenderer::RenderDistanceFieldAOSurfaceCache(FRHIComman
 						{
 							const FViewInfo& View = Views[ViewIndex];
 
-							FIntPoint DownsampledViewSize = FIntPoint::DivideAndRoundDown(View.ViewRect.Size(), DestLevelDownsampleFactor);
+							FIntPoint DownsampledViewSize = FIntPoint::DivideAndRoundUp(View.ViewRect.Size(), DestLevelDownsampleFactor);
 							uint32 GroupSizeX = FMath::DivideAndRoundUp(DownsampledViewSize.X, GDistanceFieldAOTileSizeX);
 							uint32 GroupSizeY = FMath::DivideAndRoundUp(DownsampledViewSize.Y, GDistanceFieldAOTileSizeY);
 
