Mobile Preview - Removal of almost all GRHIFeatureLevel usages.

diff --git a/Engine/Source/Developer/CrashTracker/Private/CrashTrackerModule.cpp b/Engine/Source/Developer/CrashTracker/Private/CrashTrackerModule.cpp
--- a/Engine/Source/Developer/CrashTracker/Private/CrashTrackerModule.cpp
+++ b/Engine/Source/Developer/CrashTracker/Private/CrashTrackerModule.cpp
@@ -112,7 +112,7 @@ void FCrashTrackerModule::StartupModule()
 	VideoCapture.Reset();
 #if CRASH_TRACKER_SUPPORTED
 	if (GIsEditor &&
-		(GRHIShaderPlatform == SP_PCD3D_SM4 || GRHIShaderPlatform == SP_PCD3D_SM5 || GRHIShaderPlatform == SP_OPENGL_SM4 || GRHIShaderPlatform == SP_OPENGL_SM4_MAC ))
+		(GMaxRHIShaderPlatform == SP_PCD3D_SM4 || GMaxRHIShaderPlatform == SP_PCD3D_SM5 || GMaxRHIShaderPlatform == SP_OPENGL_SM4 || GMaxRHIShaderPlatform == SP_OPENGL_SM4_MAC))
 	{
 		bool bCrashTrackerShouldBeEnabled = false;
 #if UE_BUILD_DEVELOPMENT
diff --git a/Engine/Source/Editor/MaterialEditor/Private/MaterialEditor.cpp b/Engine/Source/Editor/MaterialEditor/Private/MaterialEditor.cpp
--- a/Engine/Source/Editor/MaterialEditor/Private/MaterialEditor.cpp
+++ b/Engine/Source/Editor/MaterialEditor/Private/MaterialEditor.cpp
@@ -1114,7 +1114,7 @@ void FMaterialEditor::RegenerateCodeView(bool bForce)
 	}
 
 	FString MarkupSource;
-	if (Material->GetMaterialResource(GRHIFeatureLevel)->GetMaterialExpressionSource(MarkupSource, ExpressionCodeMap))
+	if (Material->GetMaterialResource(GMaxRHIFeatureLevel)->GetMaterialExpressionSource(MarkupSource, ExpressionCodeMap))
 	{
 		// Remove line-feeds and leave just CRs so the character counts match the selection ranges.
 		MarkupSource.ReplaceInline(TEXT("\r"), TEXT(""));
@@ -1474,7 +1474,7 @@ void FMaterialEditor::UpdateMaterialInfoList(bool bForceDisplay)
 	ERHIFeatureLevel::Type FeatureLevelsToDisplay[2];
 	int32 NumFeatureLevels = 0;
 	// Always show basic features so that errors aren't hidden
-	FeatureLevelsToDisplay[NumFeatureLevels++] = GRHIFeatureLevel;
+	FeatureLevelsToDisplay[NumFeatureLevels++] = GMaxRHIFeatureLevel;
 	if (bShowMobileStats)
 	{
 		FeatureLevelsToDisplay[NumFeatureLevels++] = ERHIFeatureLevel::ES2;
@@ -1626,7 +1626,7 @@ void FMaterialEditor::UpdateMaterialInfoList(bool bForceDisplay)
 
 void FMaterialEditor::UpdateGraphNodeStates()
 {
-	const FMaterialResource* ErrorMaterialResource = PreviewExpression ? ExpressionPreviewMaterial->GetMaterialResource(GRHIFeatureLevel) : Material->GetMaterialResource(GRHIFeatureLevel);
+	const FMaterialResource* ErrorMaterialResource = PreviewExpression ? ExpressionPreviewMaterial->GetMaterialResource(GMaxRHIFeatureLevel) : Material->GetMaterialResource(GMaxRHIFeatureLevel);
 	const FMaterialResource* ErrorMaterialResourceES2 = NULL;
 	if (bShowMobileStats)
 	{
@@ -3410,7 +3410,7 @@ FMatExpressionPreview* FMaterialEditor::GetExpressionPreview(UMaterialExpression
 		{
 			bNewlyCreated = true;
 			Preview = new(ExpressionPreviews) FMatExpressionPreview(MaterialExpression);
-			Preview->CacheShaders(GRHIShaderPlatform, true);
+			Preview->CacheShaders(GMaxRHIShaderPlatform, true);
 		}
 		return Preview;
 	}
diff --git a/Engine/Source/Editor/MaterialEditor/Private/MaterialEditor.h b/Engine/Source/Editor/MaterialEditor/Private/MaterialEditor.h
--- a/Engine/Source/Editor/MaterialEditor/Private/MaterialEditor.h
+++ b/Engine/Source/Editor/MaterialEditor/Private/MaterialEditor.h
@@ -17,7 +17,7 @@ public:
 	{
 		// Register this FMaterial derivative with AddEditorLoadedMaterialResource since it does not have a corresponding UMaterialInterface
 		FMaterial::AddEditorLoadedMaterialResource(this);
-		SetQualityLevelProperties(EMaterialQualityLevel::High, false, GRHIFeatureLevel);
+		SetQualityLevelProperties(EMaterialQualityLevel::High, false, GMaxRHIFeatureLevel);
 	}
 
 	FMatExpressionPreview(UMaterialExpression* InExpression)
@@ -29,7 +29,7 @@ public:
 
 		check(InExpression->Material && InExpression->Material->Expressions.Contains(InExpression));
 		InExpression->Material->AppendReferencedTextures(ReferencedTextures);
-		SetQualityLevelProperties(EMaterialQualityLevel::High, false, GRHIFeatureLevel);
+		SetQualityLevelProperties(EMaterialQualityLevel::High, false, GMaxRHIFeatureLevel);
 	}
 
 	~FMatExpressionPreview()
diff --git a/Engine/Source/Editor/MaterialEditor/Private/MaterialInstanceEditor.cpp b/Engine/Source/Editor/MaterialEditor/Private/MaterialInstanceEditor.cpp
--- a/Engine/Source/Editor/MaterialEditor/Private/MaterialInstanceEditor.cpp
+++ b/Engine/Source/Editor/MaterialEditor/Private/MaterialInstanceEditor.cpp
@@ -692,7 +692,7 @@ void FMaterialInstanceEditor::DrawMessages( FViewport* Viewport, FCanvas* Canvas
 	Canvas->PushAbsoluteTransform(FMatrix::Identity);
 	if ( MaterialEditorInstance->Parent && MaterialEditorInstance->SourceInstance )
 	{
-		const FMaterialResource* MaterialResource = MaterialEditorInstance->SourceInstance->GetMaterialResource(GRHIFeatureLevel);
+		const FMaterialResource* MaterialResource = MaterialEditorInstance->SourceInstance->GetMaterialResource(GMaxRHIFeatureLevel);
 		UMaterial* BaseMaterial = MaterialEditorInstance->SourceInstance->GetMaterial();
 		int32 DrawPositionY = 5;
 		if ( BaseMaterial && MaterialResource )
diff --git a/Engine/Source/Editor/Persona/Private/PersonaMeshDetails.cpp b/Engine/Source/Editor/Persona/Private/PersonaMeshDetails.cpp
--- a/Engine/Source/Editor/Persona/Private/PersonaMeshDetails.cpp
+++ b/Engine/Source/Editor/Persona/Private/PersonaMeshDetails.cpp
@@ -1676,7 +1676,7 @@ FReply FPersonaMeshDetails::OnOpenClothingFileClicked(IDetailLayoutBuilder* Deta
 					FClothingAssetData& AssetData = SkelMesh->ClothingAssets[AssetIndex];
 					int32 NumLODs = ApexClothingUtils::GetNumLODs(AssetData.ApexClothingAsset->GetAsset());
 
-					uint32 MaxClothVertices = ApexClothingUtils::GetMaxClothSimulVertices(GRHIFeatureLevel);
+					uint32 MaxClothVertices = ApexClothingUtils::GetMaxClothSimulVertices(GMaxRHIFeatureLevel);
 
 					// check whether there are sub-meshes which have over MaxClothVertices or not
 					// this checking will be removed after changing implementation way for supporting over MaxClothVertices
diff --git a/Engine/Source/Editor/PropertyEditor/Private/PropertyCustomizationHelpers.cpp b/Engine/Source/Editor/PropertyEditor/Private/PropertyCustomizationHelpers.cpp
--- a/Engine/Source/Editor/PropertyEditor/Private/PropertyCustomizationHelpers.cpp
+++ b/Engine/Source/Editor/PropertyEditor/Private/PropertyCustomizationHelpers.cpp
@@ -744,7 +744,7 @@ private:
 			UMaterialInterface* Material = MaterialItem.Material.Get();
 
 			TArray< UTexture* > Textures;
-			Material->GetUsedTextures(Textures, EMaterialQualityLevel::Num, false, GRHIFeatureLevel, false);
+			Material->GetUsedTextures(Textures, EMaterialQualityLevel::Num, false, ERHIFeatureLevel::Num, true);
 
 			// Add a menu item for each texture.  Clicking on the texture will display it in the content browser
 			for( int32 TextureIndex = 0; TextureIndex < Textures.Num(); ++TextureIndex )
diff --git a/Engine/Source/Editor/StatsViewer/Private/StatsPages/TextureStatsPage.cpp b/Engine/Source/Editor/StatsViewer/Private/StatsPages/TextureStatsPage.cpp
--- a/Engine/Source/Editor/StatsViewer/Private/StatsPages/TextureStatsPage.cpp
+++ b/Engine/Source/Editor/StatsViewer/Private/StatsPages/TextureStatsPage.cpp
@@ -299,7 +299,7 @@ struct TextureStatsGenerator : public FFindReferencedAssets
 						{
 							TArray<UTexture*> UsedTextures;
 
-							CurrentMaterial->GetUsedTextures(UsedTextures, EMaterialQualityLevel::Num, false, GRHIFeatureLevel, false);
+							CurrentMaterial->GetUsedTextures(UsedTextures, EMaterialQualityLevel::Num, false, GMaxRHIFeatureLevel, true);
 							for(int32 TextureIndex = 0; TextureIndex < UsedTextures.Num(); TextureIndex++)
 							{
 								UTexture* CurrentUsedTexture = UsedTextures[TextureIndex];
diff --git a/Engine/Source/Editor/TextureEditor/Private/Models/TextureEditorViewportClient.cpp b/Engine/Source/Editor/TextureEditor/Private/Models/TextureEditorViewportClient.cpp
--- a/Engine/Source/Editor/TextureEditor/Private/Models/TextureEditorViewportClient.cpp
+++ b/Engine/Source/Editor/TextureEditor/Private/Models/TextureEditorViewportClient.cpp
@@ -10,21 +10,16 @@ FTextureEditorViewportClient::FTextureEditorViewportClient( TWeakPtr<ITextureEdi
 	: TextureEditorPtr(InTextureEditor)
 	, TextureEditorViewportPtr(InTextureEditorViewport)
 	, CheckerboardTexture(NULL)
-	, World(NULL)
 {
 	check(TextureEditorPtr.IsValid() && TextureEditorViewportPtr.IsValid());
 
 	ModifyCheckerboardTextureColors();
-
-	// @todo This is a hack to get this working for now. This won't work with multiple worlds
-	GEditor->GetEditorWorldContext(true).AddRef(World);
 }
 
 
 FTextureEditorViewportClient::~FTextureEditorViewportClient( )
 {
 	DestroyCheckerboardTexture();
-	GEditor->GetEditorWorldContext(true).RemoveRef(World);
 }
 
 
@@ -79,7 +74,7 @@ void FTextureEditorViewportClient::Draw(FViewport* Viewport, FCanvas* Canvas)
 
 		TRefCountPtr<FBatchedElementParameters> BatchedElementParameters;
 
-		if (GRHIFeatureLevel >= ERHIFeatureLevel::SM4)
+		if (GMaxRHIFeatureLevel >= ERHIFeatureLevel::SM4)
 		{
 			if (TextureCube || RTTextureCube)
 			{
diff --git a/Engine/Source/Editor/TextureEditor/Private/Models/TextureEditorViewportClient.h b/Engine/Source/Editor/TextureEditor/Private/Models/TextureEditorViewportClient.h
--- a/Engine/Source/Editor/TextureEditor/Private/Models/TextureEditorViewportClient.h
+++ b/Engine/Source/Editor/TextureEditor/Private/Models/TextureEditorViewportClient.h
@@ -15,7 +15,7 @@ public:
 	/** FViewportClient interface */
 	virtual void Draw(FViewport* Viewport, FCanvas* Canvas) override;
 	virtual bool InputKey(FViewport* Viewport, int32 ControllerId, FKey Key, EInputEvent Event, float AmountDepressed = 1.0f, bool bGamepad = false) override;
-	virtual UWorld* GetWorld() const override { return World; }
+	virtual UWorld* GetWorld() const override { return nullptr; }
 
 	/** FGCObject interface */
 	virtual void AddReferencedObjects(FReferenceCollector& Collector) override;
@@ -52,7 +52,4 @@ private:
 
 	/** Checkerboard texture */
 	UTexture2D* CheckerboardTexture;
-
-	/** The UWorld that this level editor is viewing and allowing the user to interact with through */
-	UWorld* World;
 };
diff --git a/Engine/Source/Editor/UnrealEd/Private/EditorExporters.cpp b/Engine/Source/Editor/UnrealEd/Private/EditorExporters.cpp
--- a/Engine/Source/Editor/UnrealEd/Private/EditorExporters.cpp
+++ b/Engine/Source/Editor/UnrealEd/Private/EditorExporters.cpp
@@ -706,7 +706,7 @@ public:
 	FExportMaterialProxy()
 		: FMaterial()
 	{
-		SetQualityLevelProperties(EMaterialQualityLevel::High, false, GRHIFeatureLevel);
+		SetQualityLevelProperties(EMaterialQualityLevel::High, false, GMaxRHIFeatureLevel);
 	}
 
 	FExportMaterialProxy(UMaterialInterface* InMaterialInterface, EMaterialProperty InPropertyToCompile)
@@ -714,7 +714,7 @@ public:
 		, MaterialInterface(InMaterialInterface)
 		, PropertyToCompile(InPropertyToCompile)
 	{
-		SetQualityLevelProperties(EMaterialQualityLevel::High, false, GRHIFeatureLevel);
+		SetQualityLevelProperties(EMaterialQualityLevel::High, false, GMaxRHIFeatureLevel);
 		Material = InMaterialInterface->GetMaterial();
 		Material->AppendReferencedTextures(ReferencedTextures);
 		FPlatformMisc::CreateGuid(Id);
@@ -734,26 +734,26 @@ public:
 		// Special path for a MI with static parameters
 		if (MaterialInstance && MaterialInstance->bHasStaticPermutationResource && MaterialInstance->Parent)
 		{
-			FMaterialResource* MIResource = MaterialInstance->GetMaterialResource(GRHIFeatureLevel);
+			FMaterialResource* MIResource = MaterialInstance->GetMaterialResource(GMaxRHIFeatureLevel);
 
 			// Use the shader map Id from the static permutation
 			// This allows us to create a deterministic yet unique Id for the shader map that will be compiled for this FLightmassMaterialProxy
 			FMaterialShaderMapId ResourceId;
 			//@todo - always use highest quality level for static lighting
-			MaterialInstance->GetMaterialResourceId(GRHIShaderPlatform, EMaterialQualityLevel::Num, ResourceId);
+			MaterialInstance->GetMaterialResourceId(GMaxRHIShaderPlatform, EMaterialQualityLevel::Num, ResourceId);
 
 			// Override with a special usage so we won't re-use the shader map used by the MI for rendering
 			ResourceId.Usage = GetShaderMapUsage();
 
-			CacheShaders(ResourceId, GRHIShaderPlatform, true);
+			CacheShaders(ResourceId, GMaxRHIShaderPlatform, true);
 		}
 		else
 		{
-			FMaterialResource* MaterialResource = Material->GetMaterialResource(GRHIFeatureLevel);
+			FMaterialResource* MaterialResource = Material->GetMaterialResource(GMaxRHIFeatureLevel);
 			
 			// Copy the material resource Id
 			// The FLightmassMaterialProxy's GetShaderMapUsage will set it apart from the MI's resource when it comes to finding a shader map
-			CacheShaders(GRHIShaderPlatform, true);
+			CacheShaders(GMaxRHIShaderPlatform, true);
 		}
 	}
 
@@ -1147,7 +1147,7 @@ public:
 
 		TArray<UTexture*> MaterialTextures;
 
-		InMaterialInterface->GetUsedTextures(MaterialTextures, EMaterialQualityLevel::Num, false, GRHIFeatureLevel, false);
+		InMaterialInterface->GetUsedTextures(MaterialTextures, EMaterialQualityLevel::Num, false, GMaxRHIFeatureLevel, false);
 
 		// find the largest texture in the list (applying it's LOD bias)
 		FIntPoint MaxSize = MinimumSize;
diff --git a/Engine/Source/Editor/UnrealEd/Private/ObjectTools.cpp b/Engine/Source/Editor/UnrealEd/Private/ObjectTools.cpp
--- a/Engine/Source/Editor/UnrealEd/Private/ObjectTools.cpp
+++ b/Engine/Source/Editor/UnrealEd/Private/ObjectTools.cpp
@@ -3517,7 +3517,7 @@ namespace ThumbnailTools
 				FScopedSlowTask SlowTaskMessage( NSLOCTEXT( "ObjectTools", "FinishingCompilationStatus", "Finishing Shader Compilation..." ), bAllowNewSlowTask );
 
 				// Block until the shader maps that we will save have finished being compiled
-				InMaterial->GetMaterialResource(GRHIFeatureLevel)->FinishCompilation();
+				InMaterial->GetMaterialResource(GMaxRHIFeatureLevel)->FinishCompilation();
 			}
 
 			// Generate the thumbnail
diff --git a/Engine/Source/Editor/UnrealEd/Private/ThumbnailHelpers.cpp b/Engine/Source/Editor/UnrealEd/Private/ThumbnailHelpers.cpp
--- a/Engine/Source/Editor/UnrealEd/Private/ThumbnailHelpers.cpp
+++ b/Engine/Source/Editor/UnrealEd/Private/ThumbnailHelpers.cpp
@@ -138,7 +138,7 @@ FParticleSystemThumbnailScene::FParticleSystemThumbnailScene()
 	bForceAllUsedMipsResident = false;
 	PartComponent = NULL;
 
-	ThumbnailFXSystem = FFXSystemInterface::Create(GRHIFeatureLevel);
+	ThumbnailFXSystem = FFXSystemInterface::Create(GetScene()->GetFeatureLevel());
 	GetScene()->SetFXSystem( ThumbnailFXSystem );
 }
 
diff --git a/Engine/Source/Runtime/Engine/Classes/Materials/MaterialInterface.h b/Engine/Source/Runtime/Engine/Classes/Materials/MaterialInterface.h
--- a/Engine/Source/Runtime/Engine/Classes/Materials/MaterialInterface.h
+++ b/Engine/Source/Runtime/Engine/Classes/Materials/MaterialInterface.h
@@ -6,6 +6,7 @@
 #include "Runtime/RHI/Public/RHIDefinitions.h"
 #include "Runtime/RenderCore/Public/RenderCommandFence.h"
 #include "SceneTypes.h"
+#include <functional>
 #include "MaterialInterface.generated.h"
 
 class UMaterial;
@@ -587,6 +588,16 @@ public:
 	/** Get bitfield indicating which feature levels should be compiled by default */
 	ENGINE_API static uint32 GetFeatureLevelsToCompileForAllMaterials() { return FeatureLevelsForAllMaterials | (1 << GMaxRHIFeatureLevel); }
 
+	/** Iterate over all feature levels currently marked as active */
+	ENGINE_API static void IterateOverActiveFeatureLevels(std::function<void(ERHIFeatureLevel::Type FeatureLevel)> InHandler) 
+	{  
+		uint32 FeatureLevels = GetFeatureLevelsToCompileForAllMaterials();
+		while (FeatureLevels != 0)
+		{
+			InHandler((ERHIFeatureLevel::Type)FBitSet::GetAndClearNextBit(FeatureLevels));
+		}
+	}
+
 protected:
 
 	/** Returns a bitfield indicating which feature levels should be compiled for rendering. GMaxRHIFeatureLevel is always present */
diff --git a/Engine/Source/Runtime/Engine/Private/DirectionalLightComponent.cpp b/Engine/Source/Runtime/Engine/Private/DirectionalLightComponent.cpp
--- a/Engine/Source/Runtime/Engine/Private/DirectionalLightComponent.cpp
+++ b/Engine/Source/Runtime/Engine/Private/DirectionalLightComponent.cpp
@@ -194,11 +194,11 @@ public:
 	}
 
 	/** Returns the number of view dependent shadows this light will create. */
-	virtual int32 GetNumViewDependentWholeSceneShadows(const FSceneView& View) const 
+	virtual int32 GetNumViewDependentWholeSceneShadows(const FSceneView& View) const
 	{ 
 		int32 NumCascades = GetNumShadowMappedCascades(View.MaxShadowCascades);
 
-		if (ShouldCreateRayTracedCascade())
+		if (ShouldCreateRayTracedCascade(View.GetFeatureLevel()))
 		{
 			// Add a cascade for the ray traced shadows
 			NumCascades++;
@@ -234,7 +234,7 @@ public:
 		OutInitializer.MaxDistanceToCastInLightW = HALF_WORLD_MAX / 8.0f;
 		OutInitializer.SplitIndex = SplitIndex;
 		// Last cascade is the ray traced shadow, if enabled
-		OutInitializer.bRayTracedDistanceFieldShadow = ShouldCreateRayTracedCascade() && SplitIndex == FDirectionalLightSceneProxy::GetNumViewDependentWholeSceneShadows(View) - 1;
+		OutInitializer.bRayTracedDistanceFieldShadow = ShouldCreateRayTracedCascade(View.GetFeatureLevel()) && SplitIndex == FDirectionalLightSceneProxy::GetNumViewDependentWholeSceneShadows(View) - 1;
 		return true;
 	}
 
@@ -280,9 +280,9 @@ public:
 		return true;
 	}
 
-	virtual FVector2D GetDirectionalLightDistanceFadeParameters() const
+	virtual FVector2D GetDirectionalLightDistanceFadeParameters(ERHIFeatureLevel::Type InFeatureLevel) const
 	{
-		const float FarDistance = ComputeShadowDistance();
+		const float FarDistance = ComputeShadowDistance(InFeatureLevel);
         
 		// The far distance for the dynamic to static fade is the range of the directional light.
 		// The near distance is placed at a depth of 90% of the light's range.
@@ -334,21 +334,21 @@ private:
 		return DistanceFieldShadowDistance;
 	}
 
-	bool ShouldCreateRayTracedCascade() const
+	bool ShouldCreateRayTracedCascade(ERHIFeatureLevel::Type InFeatureLevel) const
 	{
 		//@todo - handle View.MaxShadowCascades properly
 		const int32 NumCascades = GetNumShadowMappedCascades(10);
 		const float RaytracedShadowDistance = GetDistanceFieldShadowDistance();
 		const bool bCreateWithCSM = NumCascades > 0 && RaytracedShadowDistance > GetCSMMaxDistance();
 		const bool bCreateWithoutCSM = NumCascades == 0 && RaytracedShadowDistance > 0;
-		return DoesPlatformSupportDistanceFieldShadowing(GRHIShaderPlatform) && (bCreateWithCSM || bCreateWithoutCSM);
+		return DoesPlatformSupportDistanceFieldShadowing(GShaderPlatformForFeatureLevel[InFeatureLevel]) && (bCreateWithCSM || bCreateWithoutCSM);
 	}
 
-	float ComputeShadowDistance() const
+	float ComputeShadowDistance(ERHIFeatureLevel::Type InFeatureLevel) const
 	{
 		float Distance = GetCSMMaxDistance();
 
-		if (ShouldCreateRayTracedCascade())
+		if (ShouldCreateRayTracedCascade(InFeatureLevel))
 		{
 			Distance = GetDistanceFieldShadowDistance();
 		}
@@ -477,11 +477,11 @@ private:
 	}
 
 	// Given a particular split index and view near plane, returns the view space Z distance to the PSSM split plane.
-	inline float GetSplitDistance(int32 SplitIndex, int32 NumCascades, float ShadowNear, bool bIgnoreLastCascade) const
+	inline float GetSplitDistance(const FSceneView& View, int32 SplitIndex, int32 NumCascades, float ShadowNear, bool bIgnoreLastCascade) const
 	{
 		check(SplitIndex >= 0);
 
-		float FarDistance = ComputeShadowDistance();
+		float FarDistance = ComputeShadowDistance(View.GetFeatureLevel());
 
 		if (bIgnoreLastCascade)
 		{
@@ -498,7 +498,7 @@ private:
 		// this checks for WholeSceneDynamicShadowRadius and DynamicShadowCascades
 		int32 NumCascades = GetNumViewDependentWholeSceneShadows(View);
 
-		const bool bHasRayTracedCascade = ShouldCreateRayTracedCascade();
+		const bool bHasRayTracedCascade = ShouldCreateRayTracedCascade(View.GetFeatureLevel());
 		const bool bIsRayTracedCascade = bHasRayTracedCascade && SplitIndex == NumCascades - 1;
 
 		const FMatrix ViewMatrix = View.ShadowViewMatrices.ViewMatrix;
@@ -510,8 +510,8 @@ private:
 
 		// Determine start and end distances to the current cascade's split planes
 		// Presence of the ray traced cascade does not change depth ranges for the shadow-mapped cascades
-		float SplitNear = GetSplitDistance(SplitIndex, NumCascades, View.NearClippingDistance, bHasRayTracedCascade);
-		float SplitFar = GetSplitDistance(SplitIndex + 1, NumCascades, View.NearClippingDistance, bHasRayTracedCascade && !bIsRayTracedCascade);
+		float SplitNear = GetSplitDistance(View, SplitIndex, NumCascades, View.NearClippingDistance, bHasRayTracedCascade);
+		float SplitFar = GetSplitDistance(View, SplitIndex + 1, NumCascades, View.NearClippingDistance, bHasRayTracedCascade && !bIsRayTracedCascade);
 		float FadePlane = SplitFar;
 
 		float LocalCascadeTransitionFraction = CascadeTransitionFraction * GetShadowTransitionScale();
@@ -532,8 +532,8 @@ private:
 			if(SplitIndex >= 1)
 			{
 				// todo: optimize
-				float BeforeSplitNear = GetSplitDistance(SplitIndex - 1, NumCascades, View.NearClippingDistance, bHasRayTracedCascade);
-				float BeforeSplitFar = GetSplitDistance(SplitIndex, NumCascades, View.NearClippingDistance, bHasRayTracedCascade);
+				float BeforeSplitNear = GetSplitDistance(View, SplitIndex - 1, NumCascades, View.NearClippingDistance, bHasRayTracedCascade);
+				float BeforeSplitFar = GetSplitDistance(View, SplitIndex, NumCascades, View.NearClippingDistance, bHasRayTracedCascade);
 
 				OutCascadeSettings->SplitNearFadeRegion = (BeforeSplitFar - BeforeSplitNear) * LocalCascadeTransitionFraction;
 			}
diff --git a/Engine/Source/Runtime/Engine/Private/GPUSort.cpp b/Engine/Source/Runtime/Engine/Private/GPUSort.cpp
--- a/Engine/Source/Runtime/Engine/Private/GPUSort.cpp
+++ b/Engine/Source/Runtime/Engine/Private/GPUSort.cpp
@@ -683,14 +683,13 @@ IMPLEMENT_SHADER_TYPE(,FRadixSortDownsweepCS,TEXT("RadixSortShaders"),TEXT("Radi
  * @param Count - How many items in the buffer need to be sorted.
  * @returns The index of the buffer containing sorted results.
  */
-int32 SortGPUBuffers(FRHICommandListImmediate& RHICmdList, FGPUSortBuffers SortBuffers, int32 BufferIndex, uint32 KeyMask, int32 Count)
+int32 SortGPUBuffers(FRHICommandListImmediate& RHICmdList, FGPUSortBuffers SortBuffers, int32 BufferIndex, uint32 KeyMask, int32 Count, ERHIFeatureLevel::Type FeatureLevel)
 {
 	FRadixSortParameters SortParameters;
 	FRadixSortUniformBufferRef SortUniformBufferRef;
 	const bool bDebugOffsets = CVarDebugOffsets.GetValueOnRenderThread() != 0;
 	const bool bDebugSort = CVarDebugSort.GetValueOnRenderThread() != 0;
 
-	const auto FeatureLevel = GRHIFeatureLevel;
 	check(FeatureLevel == ERHIFeatureLevel::SM5);
 
 	SCOPED_DRAW_EVENTF(RHICmdList, SortGPU, DEC_PARTICLE, TEXT("SortGPU_%d"), Count);
@@ -833,7 +832,7 @@ enum
  * @param TestSize - The number of elements to sort.
  * @returns true if the sort succeeded.
  */
-static bool RunGPUSortTest(FRHICommandListImmediate& RHICmdList, int32 TestSize)
+static bool RunGPUSortTest(FRHICommandListImmediate& RHICmdList, int32 TestSize, ERHIFeatureLevel::Type FeatureLevel)
 {
 	FRandomStream RandomStream(0x3819FFE4);
 	FGPUSortBuffers SortBuffers;
@@ -850,7 +849,7 @@ static bool RunGPUSortTest(FRHICommandListImmediate& RHICmdList, int32 TestSize)
 	const bool bDebugOffsets = CVarDebugOffsets.GetValueOnRenderThread() != 0;
 	const bool bDebugSort = CVarDebugSort.GetValueOnRenderThread() != 0;
 
-	if (GRHIFeatureLevel != ERHIFeatureLevel::SM5)
+	if (FeatureLevel != ERHIFeatureLevel::SM5)
 	{
 		return false;
 	}
@@ -899,7 +898,7 @@ static bool RunGPUSortTest(FRHICommandListImmediate& RHICmdList, int32 TestSize)
 		SortBuffers.RemoteValueSRVs[BufferIndex] = ValuesBufferSRV[BufferIndex];
 		SortBuffers.RemoteValueUAVs[BufferIndex] = ValuesBufferUAV[BufferIndex];
 	}
-	ResultBufferIndex = SortGPUBuffers(RHICmdList, SortBuffers, 0, 0xFFFFFFFF, TestSize);
+	ResultBufferIndex = SortGPUBuffers(RHICmdList, SortBuffers, 0, 0xFFFFFFFF, TestSize, FeatureLevel);
 
 	// Download results from the GPU.
 	{
@@ -964,7 +963,7 @@ static bool RunGPUSortTest(FRHICommandListImmediate& RHICmdList, int32 TestSize)
  * Executes a sort test with debug information enabled.
  * @param TestSize - The number of elements to sort.
  */
-static void RunGPUSortTestWithDebug(FRHICommandListImmediate& RHICmdList, int32 TestSize)
+static void RunGPUSortTestWithDebug(FRHICommandListImmediate& RHICmdList, int32 TestSize, ERHIFeatureLevel::Type FeatureLevel)
 {
 	static IConsoleVariable* IVarDebugOffsets = IConsoleManager::Get().FindConsoleVariable(TEXT("GPUSort.DebugOffsets"));
 	static IConsoleVariable* IVarDebugSort = IConsoleManager::Get().FindConsoleVariable(TEXT("GPUSort.DebugSort"));
@@ -972,7 +971,7 @@ static void RunGPUSortTestWithDebug(FRHICommandListImmediate& RHICmdList, int32
 	const bool bWasDebuggingSort = CVarDebugSort.GetValueOnRenderThread() != 0;
 	IVarDebugOffsets->Set(1);
 	IVarDebugSort->Set(1);
-	RunGPUSortTest(RHICmdList, TestSize);
+	RunGPUSortTest(RHICmdList, TestSize, FeatureLevel);
 	IVarDebugOffsets->Set(bWasDebuggingOffsets ? 1 : 0);
 	IVarDebugSort->Set(bWasDebuggingSort ? 1 : 0);
 }
@@ -982,13 +981,13 @@ static void RunGPUSortTestWithDebug(FRHICommandListImmediate& RHICmdList, int32
  * information enabled.
  * @param TestSize - The number of elements to sort.
  */
-static bool TestGPUSortForSize(FRHICommandListImmediate& RHICmdList, int32 TestSize)
+static bool TestGPUSortForSize(FRHICommandListImmediate& RHICmdList, int32 TestSize, ERHIFeatureLevel::Type FeatureLevel)
 {
 	check(IsInRenderingThread());
-	const bool bResult = RunGPUSortTest(RHICmdList, TestSize);
+	const bool bResult = RunGPUSortTest(RHICmdList, TestSize, FeatureLevel);
 	if (bResult == false)
 	{
-		RunGPUSortTestWithDebug(RHICmdList, TestSize);
+		RunGPUSortTestWithDebug(RHICmdList, TestSize, FeatureLevel);
 	}
 	return bResult;
 }
@@ -997,17 +996,17 @@ static bool TestGPUSortForSize(FRHICommandListImmediate& RHICmdList, int32 TestS
  * Test that GPU sorting works.
  * @param TestToRun - The test to run.
  */
-static bool TestGPUSort_RenderThread(FRHICommandListImmediate& RHICmdList, EGPUSortTest TestToRun)
+static bool TestGPUSort_RenderThread(FRHICommandListImmediate& RHICmdList, EGPUSortTest TestToRun, ERHIFeatureLevel::Type FeatureLevel)
 {
 	check(IsInRenderingThread());
 
 	switch (TestToRun)
 	{
 	case GPU_SORT_TEST_SMALL:
-		return TestGPUSortForSize(RHICmdList, GPU_SORT_TEST_SIZE_SMALL);
+		return TestGPUSortForSize(RHICmdList, GPU_SORT_TEST_SIZE_SMALL, FeatureLevel);
 
 	case GPU_SORT_TEST_LARGE:
-		return TestGPUSortForSize(RHICmdList, GPU_SORT_TEST_SIZE_LARGE);
+		return TestGPUSortForSize(RHICmdList, GPU_SORT_TEST_SIZE_LARGE, FeatureLevel);
 
 	case GPU_SORT_TEST_EXHAUSTIVE:
 		{
@@ -1015,7 +1014,7 @@ static bool TestGPUSort_RenderThread(FRHICommandListImmediate& RHICmdList, EGPUS
 			int32 TestSize = GPU_SORT_TEST_SIZE_MIN;
 			while (TestSize <= GPU_SORT_TEST_SIZE_MAX)
 			{
-				if (TestGPUSortForSize(RHICmdList, TestSize) == false)
+				if (TestGPUSortForSize(RHICmdList, TestSize, FeatureLevel) == false)
 				{
 					return false;
 				}
@@ -1026,7 +1025,7 @@ static bool TestGPUSort_RenderThread(FRHICommandListImmediate& RHICmdList, EGPUS
 			TestSize = GPU_SORT_TEST_SIZE_MIN;
 			while (TestSize <= GPU_SORT_TEST_SIZE_MAX)
 			{
-				if (TestGPUSortForSize(RHICmdList, TestSize - 1) == false)
+				if (TestGPUSortForSize(RHICmdList, TestSize - 1, FeatureLevel) == false)
 				{
 					return false;
 				}
@@ -1039,7 +1038,7 @@ static bool TestGPUSort_RenderThread(FRHICommandListImmediate& RHICmdList, EGPUS
 		for ( int32 i = 0; i < 1000; ++i )
 		{
 			int32 TestSize = FMath::TruncToInt(FMath::SRand() * (float)(GPU_SORT_TEST_SIZE_MAX - GPU_SORT_TEST_SIZE_MIN)) + GPU_SORT_TEST_SIZE_MIN;
-			if (TestGPUSortForSize(RHICmdList, (TestSize + 0xF) & 0xFFFFFFF0) == false)
+			if (TestGPUSortForSize(RHICmdList, (TestSize + 0xF) & 0xFFFFFFF0, FeatureLevel) == false)
 			{
 				return false;
 			}
@@ -1054,12 +1053,13 @@ static bool TestGPUSort_RenderThread(FRHICommandListImmediate& RHICmdList, EGPUS
  * Test that GPU sorting works.
  * @param TestToRun - The test to run.
  */
-void TestGPUSort(EGPUSortTest TestToRun)
+void TestGPUSort(EGPUSortTest TestToRun, ERHIFeatureLevel::Type FeatureLevel)
 {
-	ENQUEUE_UNIQUE_RENDER_COMMAND_ONEPARAMETER(
+	ENQUEUE_UNIQUE_RENDER_COMMAND_TWOPARAMETER(
 		FTestGPUSortCommand,
 		EGPUSortTest, TestToRun, TestToRun,
+		ERHIFeatureLevel::Type, FeatureLevel, FeatureLevel,
 	{
-		TestGPUSort_RenderThread(RHICmdList, TestToRun);
+		TestGPUSort_RenderThread(RHICmdList, TestToRun, FeatureLevel);
 	});
 }
diff --git a/Engine/Source/Runtime/Engine/Private/GPUSort.h b/Engine/Source/Runtime/Engine/Private/GPUSort.h
--- a/Engine/Source/Runtime/Engine/Private/GPUSort.h
+++ b/Engine/Source/Runtime/Engine/Private/GPUSort.h
@@ -36,7 +36,7 @@ struct FGPUSortBuffers
  * @param Count - How many items in the buffer need to be sorted.
  * @returns The index of the buffer containing sorted results.
  */
-int32 SortGPUBuffers(FRHICommandListImmediate& RHICmdList, FGPUSortBuffers SortBuffers, int32 BufferIndex, uint32 KeyMask, int32 Count);
+int32 SortGPUBuffers(FRHICommandListImmediate& RHICmdList, FGPUSortBuffers SortBuffers, int32 BufferIndex, uint32 KeyMask, int32 Count, ERHIFeatureLevel::Type FeatureLevel);
 
 /**
  * GPU sorting tests.
@@ -56,4 +56,4 @@ enum EGPUSortTest
  * Test that GPU sorting works.
  * @param TestToRun - The test to run.
  */
-void TestGPUSort(EGPUSortTest TestToRun);
+void TestGPUSort(EGPUSortTest TestToRun, ERHIFeatureLevel::Type FeatureLevel);
diff --git a/Engine/Source/Runtime/Engine/Private/GlobalShader.cpp b/Engine/Source/Runtime/Engine/Private/GlobalShader.cpp
--- a/Engine/Source/Runtime/Engine/Private/GlobalShader.cpp
+++ b/Engine/Source/Runtime/Engine/Private/GlobalShader.cpp
@@ -223,7 +223,7 @@ void VerifyGlobalShaders(EShaderPlatform Platform, bool bLoadedFromCacheFile)
 			// OpenGL requires that global shader maps are compiled before attaching
 			// primitives to the scene as it must be able to find FNULLPS.
 			// TODO_OPENGL: Allow shaders to be compiled asynchronously.
-			!IsOpenGLPlatform(GRHIShaderPlatform) &&
+			!IsOpenGLPlatform(GMaxRHIShaderPlatform) &&
 			GShaderCompilingManager->AllowAsynchronousShaderCompiling();
 
 		if (!bAllowAsynchronousGlobalShaderCompiling)
diff --git a/Engine/Source/Runtime/Engine/Private/Landscape/Landscape.cpp b/Engine/Source/Runtime/Engine/Private/Landscape/Landscape.cpp
--- a/Engine/Source/Runtime/Engine/Private/Landscape/Landscape.cpp
+++ b/Engine/Source/Runtime/Engine/Private/Landscape/Landscape.cpp
@@ -412,7 +412,7 @@ void ULandscapeComponent::PostLoad()
 		}
 
 		// If we're loading on a platform that doesn't require cooked data, attempt to load missing data from the DDC
-		if (!FPlatformProperties::RequiresCookedData() && GRHIFeatureLevel == ERHIFeatureLevel::ES2)
+		if (!FPlatformProperties::RequiresCookedData() && (GIsEditor || GRHIFeatureLevel == ERHIFeatureLevel::ES2))
 		{
 			// Only check the DDC if we don't already have it loaded
 			if (!PlatformData.HasValidPlatformData())
diff --git a/Engine/Source/Runtime/Engine/Private/Landscape/LandscapeRender.cpp b/Engine/Source/Runtime/Engine/Private/Landscape/LandscapeRender.cpp
--- a/Engine/Source/Runtime/Engine/Private/Landscape/LandscapeRender.cpp
+++ b/Engine/Source/Runtime/Engine/Private/Landscape/LandscapeRender.cpp
@@ -2697,7 +2697,7 @@ void ULandscapeComponent::GetStreamingTextureInfo(TArray<FStreamingTexturePrimit
 	if (MaterialInstance)
 	{
 		TArray<UTexture*> Textures;
-		MaterialInstance->GetUsedTextures(Textures, EMaterialQualityLevel::Num, false, GRHIFeatureLevel, false);
+		MaterialInstance->GetUsedTextures(Textures, EMaterialQualityLevel::Num, false, GetWorld()->FeatureLevel, false);
 		// Add each texture to the output with the appropriate parameters.
 		// TODO: Take into account which UVIndex is being used.
 		for (int32 TextureIndex = 0;TextureIndex < Textures.Num();TextureIndex++)
diff --git a/Engine/Source/Runtime/Engine/Private/LightComponent.cpp b/Engine/Source/Runtime/Engine/Private/LightComponent.cpp
--- a/Engine/Source/Runtime/Engine/Private/LightComponent.cpp
+++ b/Engine/Source/Runtime/Engine/Private/LightComponent.cpp
@@ -15,7 +15,7 @@
 
 void FStaticShadowDepthMap::InitRHI()
 {
-	if (ShadowMapSizeX > 0 && ShadowMapSizeY > 0 && GRHIFeatureLevel >= ERHIFeatureLevel::SM4)
+	if (ShadowMapSizeX > 0 && ShadowMapSizeY > 0 && GMaxRHIFeatureLevel >= ERHIFeatureLevel::SM4)
 	{
 		FRHIResourceCreateInfo CreateInfo;
 		FTexture2DRHIRef Texture2DRHI = RHICreateTexture2D(ShadowMapSizeX, ShadowMapSizeY, PF_R16F, 1, 1, 0, CreateInfo);
diff --git a/Engine/Source/Runtime/Engine/Private/LightMap.cpp b/Engine/Source/Runtime/Engine/Private/LightMap.cpp
--- a/Engine/Source/Runtime/Engine/Private/LightMap.cpp
+++ b/Engine/Source/Runtime/Engine/Private/LightMap.cpp
@@ -70,9 +70,9 @@ FLightMap::FLightMap()
 	: bAllowHighQualityLightMaps(true)
 	, NumRefs(0)
 {
-	bAllowHighQualityLightMaps = AllowHighQualityLightmaps();
+	bAllowHighQualityLightMaps = AllowHighQualityLightmaps(GMaxRHIFeatureLevel);
 #if !PLATFORM_DESKTOP 
-	checkf(bAllowHighQualityLightMaps || IsMobilePlatform(GRHIShaderPlatform), TEXT("Low quality lightmaps are not currently supported on consoles. Make sure console variable r.HighQualityLightMaps is true for this platform"));
+	checkf(bAllowHighQualityLightMaps || IsMobilePlatform(GMaxRHIShaderPlatform), TEXT("Low quality lightmaps are not currently supported on consoles. Make sure console variable r.HighQualityLightMaps is true for this platform"));
 #endif
 }
 
diff --git a/Engine/Source/Runtime/Engine/Private/Materials/Material.cpp b/Engine/Source/Runtime/Engine/Private/Materials/Material.cpp
--- a/Engine/Source/Runtime/Engine/Private/Materials/Material.cpp
+++ b/Engine/Source/Runtime/Engine/Private/Materials/Material.cpp
@@ -615,8 +615,7 @@ void UMaterial::GetUsedTextures(TArray<UTexture*>& OutTextures, EMaterialQuality
 
 						if (Texture)
 						{
-							OutTextures.Add(Texture);
-							//OutTextures.AddUnique(Texture); //AJB - maybe this?
+							OutTextures.AddUnique(Texture);
 						}
 					}
 				}
@@ -1518,11 +1517,14 @@ void UMaterial::FlushResourceShaderMaps()
 {
 	FPlatformMisc::CreateGuid(StateId);
 
-	for (int32 QualityLevelIndex = 0; QualityLevelIndex < EMaterialQualityLevel::Num; QualityLevelIndex++)
+	UMaterialInterface::IterateOverActiveFeatureLevels([&](ERHIFeatureLevel::Type InFeatureLevel)
 	{
-		FMaterialResource* CurrentResource = MaterialResources[QualityLevelIndex][GRHIFeatureLevel];
-		CurrentResource->ReleaseShaderMap();
-	}
+		for (int32 QualityLevelIndex = 0; QualityLevelIndex < EMaterialQualityLevel::Num; QualityLevelIndex++)
+		{
+			FMaterialResource* CurrentResource = MaterialResources[QualityLevelIndex][InFeatureLevel];
+			CurrentResource->ReleaseShaderMap();
+		}
+	});
 }
 
 void UMaterial::RebuildMaterialFunctionInfo()
@@ -2795,9 +2797,12 @@ void UMaterial::AddReferencedObjects(UObject* InThis, FReferenceCollector& Colle
 #if WITH_EDITOR
 void UMaterial::CancelOutstandingCompilation()
 {
-	if (FMaterialResource* Res = GetMaterialResource(GRHIFeatureLevel))
+	for (int32 FeatureLevel = 0; FeatureLevel < ERHIFeatureLevel::Num; ++FeatureLevel)
 	{
-		Res->CancelCompilation();
+		if (FMaterialResource* Res = GetMaterialResource((ERHIFeatureLevel::Type)FeatureLevel))
+		{
+			Res->CancelCompilation();
+		}
 	}
 }
 #endif
@@ -3443,20 +3448,28 @@ void UMaterial::AllMaterialsCacheResourceShadersForRendering()
 static void ListSceneColorMaterials()
 {
 	int32 NumSceneColorMaterials = 0;
-	for (TObjectIterator<UMaterialInterface> It; It; ++It)
+
+	UMaterialInterface::IterateOverActiveFeatureLevels([&](ERHIFeatureLevel::Type FeatureLevel) 
 	{
-		UMaterialInterface* Mat = *It;
-		const FMaterial* MatRes = Mat->GetRenderProxy(false)->GetMaterial(GRHIFeatureLevel);
-		if (MatRes && MatRes->RequiresSceneColorCopy_GameThread())
+		FString FeatureLevelName;
+		GetFeatureLevelName(FeatureLevel, FeatureLevelName);
+
+		for (TObjectIterator<UMaterialInterface> It; It; ++It)
 		{
-			UMaterial* BaseMat = Mat->GetMaterial();
-			UE_LOG(LogConsoleResponse,Display,TEXT("[SepTrans=%d] %s"),
-				BaseMat ? BaseMat->bEnableSeparateTranslucency : 3,
-				*Mat->GetPathName()
-				);
-			NumSceneColorMaterials++;
+			UMaterialInterface* Mat = *It;
+			const FMaterial* MatRes = Mat->GetRenderProxy(false)->GetMaterial(FeatureLevel);
+			if (MatRes && MatRes->RequiresSceneColorCopy_GameThread())
+			{
+				UMaterial* BaseMat = Mat->GetMaterial();
+				UE_LOG(LogConsoleResponse, Display, TEXT("[SepTrans=%d][FeatureLevel=%s] %s"),
+					BaseMat ? BaseMat->bEnableSeparateTranslucency : 3,
+					*FeatureLevelName,
+					*Mat->GetPathName()
+					);
+				NumSceneColorMaterials++;
+			}
 		}
-	}
+	});
 	UE_LOG(LogConsoleResponse,Display,TEXT("%d loaded materials read from scene color."),NumSceneColorMaterials);
 }
 
diff --git a/Engine/Source/Runtime/Engine/Private/Materials/MaterialInstance.cpp b/Engine/Source/Runtime/Engine/Private/Materials/MaterialInstance.cpp
--- a/Engine/Source/Runtime/Engine/Private/Materials/MaterialInstance.cpp
+++ b/Engine/Source/Runtime/Engine/Private/Materials/MaterialInstance.cpp
@@ -772,56 +772,47 @@ void UMaterialInstance::OverrideTexture(const UTexture* InTextureToOverride, UTe
 {
 #if WITH_EDITOR
 	bool bShouldRecacheMaterialExpressions = false;
-	const bool bES2Preview = false;
-	//ERHIFeatureLevel::Type FeatureLevelsToUpdate[2] = { GRHIFeatureLevel, ERHIFeatureLevel::ES2 };
-	//int32 NumFeatureLevelsToUpdate = bES2Preview ? 2 : 1;
-	ERHIFeatureLevel::Type FeatureLevelsToUpdate[1] = { InFeatureLevel };
-	int32 NumFeatureLevelsToUpdate = 1;
 	
-	for (int32 i = 0; i < NumFeatureLevelsToUpdate; ++i)
-	{
-		const TArray<TRefCountPtr<FMaterialUniformExpressionTexture> >* ExpressionsByType[2];
-		ERHIFeatureLevel::Type FeatureLevel = FeatureLevelsToUpdate[i];
+	const TArray<TRefCountPtr<FMaterialUniformExpressionTexture> >* ExpressionsByType[2];
 
-		const FMaterialResource* SourceMaterialResource = NULL;
-		if (bHasStaticPermutationResource)
-		{
-			SourceMaterialResource = GetMaterialResource(FeatureLevel);
-			// Iterate over both the 2D textures and cube texture expressions.
-			ExpressionsByType[0] = &SourceMaterialResource->GetUniform2DTextureExpressions();
-			ExpressionsByType[1] = &SourceMaterialResource->GetUniformCubeTextureExpressions();
-		}
-		else
-		{
-			//@todo - this isn't handling chained MIC's correctly, where a parent in the chain has static parameters
-			UMaterial* Material = GetMaterial();
-			SourceMaterialResource = Material->GetMaterialResource(FeatureLevel);
+	const FMaterialResource* SourceMaterialResource = NULL;
+	if (bHasStaticPermutationResource)
+	{
+		SourceMaterialResource = GetMaterialResource(InFeatureLevel);
+		// Iterate over both the 2D textures and cube texture expressions.
+		ExpressionsByType[0] = &SourceMaterialResource->GetUniform2DTextureExpressions();
+		ExpressionsByType[1] = &SourceMaterialResource->GetUniformCubeTextureExpressions();
+	}
+	else
+	{
+		//@todo - this isn't handling chained MIC's correctly, where a parent in the chain has static parameters
+		UMaterial* Material = GetMaterial();
+		SourceMaterialResource = Material->GetMaterialResource(InFeatureLevel);
 			
-			// Iterate over both the 2D textures and cube texture expressions.
-			ExpressionsByType[0] = &SourceMaterialResource->GetUniform2DTextureExpressions();
-			ExpressionsByType[1] = &SourceMaterialResource->GetUniformCubeTextureExpressions();
-		}
+		// Iterate over both the 2D textures and cube texture expressions.
+		ExpressionsByType[0] = &SourceMaterialResource->GetUniform2DTextureExpressions();
+		ExpressionsByType[1] = &SourceMaterialResource->GetUniformCubeTextureExpressions();
+	}
 		
-		for(int32 TypeIndex = 0;TypeIndex < ARRAY_COUNT(ExpressionsByType);TypeIndex++)
-		{
-			const TArray<TRefCountPtr<FMaterialUniformExpressionTexture> >& Expressions = *ExpressionsByType[TypeIndex];
+	for(int32 TypeIndex = 0;TypeIndex < ARRAY_COUNT(ExpressionsByType);TypeIndex++)
+	{
+		const TArray<TRefCountPtr<FMaterialUniformExpressionTexture> >& Expressions = *ExpressionsByType[TypeIndex];
 
-			// Iterate over each of the material's texture expressions.
-			for(int32 ExpressionIndex = 0;ExpressionIndex < Expressions.Num();ExpressionIndex++)
-			{
-				FMaterialUniformExpressionTexture* Expression = Expressions[ExpressionIndex];
+		// Iterate over each of the material's texture expressions.
+		for(int32 ExpressionIndex = 0;ExpressionIndex < Expressions.Num();ExpressionIndex++)
+		{
+			FMaterialUniformExpressionTexture* Expression = Expressions[ExpressionIndex];
 
-				// Evaluate the expression in terms of this material instance.
-				const bool bAllowOverride = false;
-				UTexture* Texture = NULL;
-				Expression->GetGameThreadTextureValue(this,*SourceMaterialResource,Texture,bAllowOverride);
+			// Evaluate the expression in terms of this material instance.
+			const bool bAllowOverride = false;
+			UTexture* Texture = NULL;
+			Expression->GetGameThreadTextureValue(this,*SourceMaterialResource,Texture,bAllowOverride);
 
-				if( Texture != NULL && Texture == InTextureToOverride )
-				{
-					// Override this texture!
-					Expression->SetTransientOverrideTextureValue( OverrideTexture );
-					bShouldRecacheMaterialExpressions = true;
-				}
+			if( Texture != NULL && Texture == InTextureToOverride )
+			{
+				// Override this texture!
+				Expression->SetTransientOverrideTextureValue( OverrideTexture );
+				bShouldRecacheMaterialExpressions = true;
 			}
 		}
 	}
@@ -1260,9 +1251,8 @@ void UMaterialInstance::GetMaterialResourceId(EShaderPlatform ShaderPlatform, EM
 	FStaticParameterSet CompositedStaticParameters;
 	GetStaticParameterValues(CompositedStaticParameters);
 
-	// TODO: Is this right?
-	// ERHIFeatureLevel::Type FeatureLevel = GetMaxSupportedFeatureLevel(ShaderPlatform);
-	const FMaterialResource* BaseResource = BaseMaterial->GetMaterialResource(GRHIFeatureLevel, QualityLevel);
+	ERHIFeatureLevel::Type FeatureLevel = GetMaxSupportedFeatureLevel(ShaderPlatform);
+	const FMaterialResource* BaseResource = BaseMaterial->GetMaterialResource(FeatureLevel, QualityLevel);
 
 	GetMaterialResourceId(BaseResource, ShaderPlatform, CompositedStaticParameters, OutId);
 }
@@ -2063,7 +2053,7 @@ bool UMaterialInstance::UpdateLightmassTextureTracking()
 #if WITH_EDITORONLY_DATA
 	TArray<UTexture*> UsedTextures;
 	
-	GetUsedTextures(UsedTextures, EMaterialQualityLevel::Num, true, GRHIFeatureLevel, true);
+	GetUsedTextures(UsedTextures, EMaterialQualityLevel::Num, true, GMaxRHIFeatureLevel, true);
 	if (UsedTextures.Num() != ReferencedTextureGuids.Num())
 	{
 		bTexturesHaveChanged = true;
diff --git a/Engine/Source/Runtime/Engine/Private/Materials/MaterialInterface.cpp b/Engine/Source/Runtime/Engine/Private/Materials/MaterialInterface.cpp
--- a/Engine/Source/Runtime/Engine/Private/Materials/MaterialInterface.cpp
+++ b/Engine/Source/Runtime/Engine/Private/Materials/MaterialInterface.cpp
@@ -92,7 +92,7 @@ void UMaterialInterface::SetForceMipLevelsToBeResident( bool OverrideForceMiplev
 {
 	TArray<UTexture*> Textures;
 	
-	GetUsedTextures(Textures, EMaterialQualityLevel::Num, false, GRHIFeatureLevel, false);
+	GetUsedTextures(Textures, EMaterialQualityLevel::Num, false, ERHIFeatureLevel::Num, true);
 	for ( int32 TextureIndex=0; TextureIndex < Textures.Num(); ++TextureIndex )
 	{
 		UTexture2D* Texture = Cast<UTexture2D>(Textures[TextureIndex]);
@@ -220,7 +220,7 @@ bool DoesMaterialUseTexture(const UMaterialInterface* Material,const UTexture* C
 	}
 
 	TArray<UTexture*> Textures;
-	Material->GetUsedTextures(Textures, EMaterialQualityLevel::Num, true, GRHIFeatureLevel, true);
+	Material->GetUsedTextures(Textures, EMaterialQualityLevel::Num, true, GMaxRHIFeatureLevel, true);
 	for (int32 i = 0; i < Textures.Num(); i++)
 	{
 		if (Textures[i] == CheckTexture)
@@ -231,7 +231,7 @@ bool DoesMaterialUseTexture(const UMaterialInterface* Material,const UTexture* C
 	return false;
 }
 
-float UMaterialInterface::GetOpacityMaskClipValue()const
+float UMaterialInterface::GetOpacityMaskClipValue() const
 {
 	if( IsInGameThread() )
 	{
@@ -239,10 +239,10 @@ float UMaterialInterface::GetOpacityMaskClipValue()const
 	}
 
 	//We're on the render thread so get it from the proxy.
-	return GetRenderProxy(0)->GetMaterial(GRHIFeatureLevel)->GetOpacityMaskClipValue();
+	return GetRenderProxy(0)->GetMaterial(GMaxRHIFeatureLevel)->GetOpacityMaskClipValue();
 }
 
-EBlendMode UMaterialInterface::GetBlendMode()const
+EBlendMode UMaterialInterface::GetBlendMode() const
 {
 	if( IsInGameThread() )
 	{
@@ -250,10 +250,10 @@ EBlendMode UMaterialInterface::GetBlendMode()const
 	}
 
 	//We're on the render thread so get it from the proxy.
-	return GetRenderProxy(0)->GetMaterial(GRHIFeatureLevel)->GetBlendMode();
+	return GetRenderProxy(0)->GetMaterial(GMaxRHIFeatureLevel)->GetBlendMode();
 }
 
-bool UMaterialInterface::IsTwoSided()const
+bool UMaterialInterface::IsTwoSided() const
 {
 	if( IsInGameThread() )
 	{
@@ -261,7 +261,7 @@ bool UMaterialInterface::IsTwoSided()const
 	}
 
 	//We're on the render thread so get it from the proxy.
-	return GetRenderProxy(0)->GetMaterial(GRHIFeatureLevel)->IsTwoSided();
+	return GetRenderProxy(0)->GetMaterial(GMaxRHIFeatureLevel)->IsTwoSided();
 }
 
 bool UMaterialInterface::IsMasked() const
@@ -272,10 +272,10 @@ bool UMaterialInterface::IsMasked() const
 	}
 
 	//We're on the render thread so get it from the proxy.
-	return GetRenderProxy(0)->GetMaterial(GRHIFeatureLevel)->IsMasked();
+	return GetRenderProxy(0)->GetMaterial(GMaxRHIFeatureLevel)->IsMasked();
 }
 
-EMaterialShadingModel UMaterialInterface::GetShadingModel()const
+EMaterialShadingModel UMaterialInterface::GetShadingModel() const
 {
 	if( IsInGameThread() )
 	{
@@ -283,20 +283,21 @@ EMaterialShadingModel UMaterialInterface::GetShadingModel()const
 	}
 
 	//We're on the render thread so get it from the proxy.
-	return GetRenderProxy(0)->GetMaterial(GRHIFeatureLevel)->GetShadingModel();
+	return GetRenderProxy(0)->GetMaterial(GMaxRHIFeatureLevel)->GetShadingModel();
+
 }
 
-float UMaterialInterface::GetOpacityMaskClipValue_Internal()const
+float UMaterialInterface::GetOpacityMaskClipValue_Internal() const
 {
 	return 0.0f;
 }
 
-EBlendMode UMaterialInterface::GetBlendMode_Internal()const
+EBlendMode UMaterialInterface::GetBlendMode_Internal() const
 {
 	return BLEND_Opaque;
 }
 
-bool UMaterialInterface::IsTwoSided_Internal()const
+bool UMaterialInterface::IsTwoSided_Internal() const
 {
 	return false;
 }
@@ -306,7 +307,7 @@ bool UMaterialInterface::IsMasked_Internal() const
 	return false;
 }
 
-EMaterialShadingModel UMaterialInterface::GetShadingModel_Internal()const
+EMaterialShadingModel UMaterialInterface::GetShadingModel_Internal() const
 {
 	return MSM_DefaultLit;
 }
diff --git a/Engine/Source/Runtime/Engine/Private/Materials/MaterialShader.cpp b/Engine/Source/Runtime/Engine/Private/Materials/MaterialShader.cpp
--- a/Engine/Source/Runtime/Engine/Private/Materials/MaterialShader.cpp
+++ b/Engine/Source/Runtime/Engine/Private/Materials/MaterialShader.cpp
@@ -1000,7 +1000,7 @@ void FMaterialShaderMap::LoadForRemoteRecompile(FArchive& Ar, EShaderPlatform Sh
 				FMaterialShaderMap* LoadedShaderMap = LoadedShaderMaps[ShaderMapIndex];
 
 				if (LoadedShaderMap->GetShaderPlatform() == ShaderPlatform 
-					&& LoadedShaderMap->GetShaderMapId().FeatureLevel == GRHIFeatureLevel)
+					&& LoadedShaderMap->GetShaderMapId().FeatureLevel == GetMaxSupportedFeatureLevel(ShaderPlatform))
 				{
 					EMaterialQualityLevel::Type LoadedQualityLevel = LoadedShaderMap->GetShaderMapId().QualityLevel;
 
@@ -1011,7 +1011,7 @@ void FMaterialShaderMap::LoadForRemoteRecompile(FArchive& Ar, EShaderPlatform Sh
 							// Second pass: assign shader maps with a specified quality level to only the appropriate material resource
 							|| (PassIndex == 1 && QualityLevelIndex == LoadedQualityLevel))
 						{
-							FMaterialResource* MaterialResource = MatchingMaterial->GetMaterialResource(GRHIFeatureLevel, (EMaterialQualityLevel::Type)QualityLevelIndex);
+							FMaterialResource* MaterialResource = MatchingMaterial->GetMaterialResource(GetMaxSupportedFeatureLevel(ShaderPlatform), (EMaterialQualityLevel::Type)QualityLevelIndex);
 
 							MaterialResource->SetGameThreadShaderMap(LoadedShaderMap);
 
diff --git a/Engine/Source/Runtime/Engine/Private/Materials/MaterialShared.cpp b/Engine/Source/Runtime/Engine/Private/Materials/MaterialShared.cpp
--- a/Engine/Source/Runtime/Engine/Private/Materials/MaterialShared.cpp
+++ b/Engine/Source/Runtime/Engine/Private/Materials/MaterialShared.cpp
@@ -1420,9 +1420,10 @@ FShader* FMaterial::GetShader(FMeshMaterialShaderType* ShaderType, FVertexFactor
 	if (!Shader)
 	{
 		// Get the ShouldCache results that determine whether the shader should be compiled
-		bool bMaterialShouldCache = ShouldCache(GRHIShaderPlatform, ShaderType, VertexFactoryType);
-		bool bVFShouldCache = VertexFactoryType->ShouldCache(GRHIShaderPlatform, this, ShaderType);
-		bool bShaderShouldCache = ShaderType->ShouldCache(GRHIShaderPlatform, this, VertexFactoryType);
+		auto ShaderPlatform = GShaderPlatformForFeatureLevel[GetFeatureLevel()];
+		bool bMaterialShouldCache = ShouldCache(ShaderPlatform, ShaderType, VertexFactoryType);
+		bool bVFShouldCache = VertexFactoryType->ShouldCache(ShaderPlatform, this, ShaderType);
+		bool bShaderShouldCache = ShaderType->ShouldCache(ShaderPlatform, this, VertexFactoryType);
 		FString MaterialUsage = GetMaterialUsageDescription();
 
 		int BreakPoint = 0;
@@ -1437,7 +1438,7 @@ FShader* FMaterial::GetShader(FMeshMaterialShaderType* ShaderType, FVertexFactor
 			ShaderType->GetName(), 
 			*GetFriendlyName(),
 			VertexFactoryType->GetName(),
-			*LegacyShaderPlatformToShaderFormat(GRHIShaderPlatform).ToString(),
+			*LegacyShaderPlatformToShaderFormat(ShaderPlatform).ToString(),
 			bMaterialShouldCache,
 			bVFShouldCache,
 			bShaderShouldCache,
@@ -1530,13 +1531,10 @@ void FMaterialRenderProxy::CacheUniformExpressions()
 
 	check(UMaterial::GetDefaultMaterial(MD_Surface));
 
-	uint32 FeatureLevelsToCompile = UMaterialInterface::GetFeatureLevelsToCompileForAllMaterials();
 	TArray<FMaterialResource*> ResourcesToCache;
 
-	while (FeatureLevelsToCompile != 0)
+	UMaterialInterface::IterateOverActiveFeatureLevels([&](ERHIFeatureLevel::Type FeatureLevel)
 	{
-		ERHIFeatureLevel::Type FeatureLevel = (ERHIFeatureLevel::Type)FBitSet::GetAndClearNextBit(FeatureLevelsToCompile);
-
 		const FMaterial* MaterialNoFallback = GetMaterialNoFallback(FeatureLevel);
 
 		if (MaterialNoFallback && MaterialNoFallback->GetRenderingThreadShaderMap())
@@ -1564,7 +1562,7 @@ void FMaterialRenderProxy::CacheUniformExpressions()
 			InvalidateUniformExpressionCache();
 			return;
 		}
-	}
+	});
 }
 
 void FMaterialRenderProxy::CacheUniformExpressions_GameThread()
@@ -1882,7 +1880,7 @@ bool FMaterial::GetMaterialExpressionSource( FString& OutSource, TMap<FMaterialE
 	};
 
 	FMaterialCompilationOutput TempOutput;
-	FViewSourceMaterialTranslator MaterialTranslator(this, TempOutput, FStaticParameterSet(), GRHIShaderPlatform, GetQualityLevel(), GetFeatureLevel());
+	FViewSourceMaterialTranslator MaterialTranslator(this, TempOutput, FStaticParameterSet(), GMaxRHIShaderPlatform, GetQualityLevel(), GetFeatureLevel());
 	bool bSuccess = MaterialTranslator.Translate();
 
 	if( bSuccess )
@@ -1977,10 +1975,19 @@ void FMaterialUpdateContext::AddMaterialInstance(UMaterialInstance* Instance)
 FMaterialUpdateContext::~FMaterialUpdateContext()
 {
 	double StartTime = FPlatformTime::Seconds();
+	bool bProcess = false;
 
 	// if the shader platform that was processed is not the currently rendering shader platform, 
 	// there's no reason to update all of the runtime components
-	if (ShaderPlatform != GRHIShaderPlatform)
+	UMaterialInterface::IterateOverActiveFeatureLevels([&](ERHIFeatureLevel::Type InFeatureLevel)
+	{
+		if (ShaderPlatform == GShaderPlatformForFeatureLevel[InFeatureLevel])
+		{
+			bProcess = true;
+		}
+	});
+
+	if (!bProcess)
 	{
 		return;
 	}
diff --git a/Engine/Source/Runtime/Engine/Private/ModelComponent.cpp b/Engine/Source/Runtime/Engine/Private/ModelComponent.cpp
--- a/Engine/Source/Runtime/Engine/Private/ModelComponent.cpp
+++ b/Engine/Source/Runtime/Engine/Private/ModelComponent.cpp
@@ -370,7 +370,7 @@ void UModelComponent::GetStreamingTextureInfo(TArray<FStreamingTexturePrimitiveI
 				// Enumerate the textures used by the surface's material.
 				TArray<UTexture*> Textures;
 				
-				Material->GetUsedTextures(Textures, EMaterialQualityLevel::Num, false, GRHIFeatureLevel, false);
+				Material->GetUsedTextures(Textures, EMaterialQualityLevel::Num, false, GMaxRHIFeatureLevel, true);
 
 				// Add each texture to the output with the appropriate parameters.
 				for(int32 TextureIndex = 0;TextureIndex < Textures.Num();TextureIndex++)
diff --git a/Engine/Source/Runtime/Engine/Private/ModelLight.cpp b/Engine/Source/Runtime/Engine/Private/ModelLight.cpp
--- a/Engine/Source/Runtime/Engine/Private/ModelLight.cpp
+++ b/Engine/Source/Runtime/Engine/Private/ModelLight.cpp
@@ -482,7 +482,9 @@ void UModelComponent::GetLightAndShadowMapMemoryUsage( int32& LightMapMemoryUsag
 	// Stored in texture.
 	const float MIP_FACTOR = 1.33f;
 	ShadowMapMemoryUsage	= FMath::TruncToInt( MIP_FACTOR * LightMapWidth * LightMapHeight ); // G8
-	if( AllowHighQualityLightmaps() )
+
+	auto FeatureLevel = GetWorld() ? GetWorld()->FeatureLevel : GMaxRHIFeatureLevel;
+	if (AllowHighQualityLightmaps(FeatureLevel))
 	{ 
 		LightMapMemoryUsage = FMath::TruncToInt( NUM_HQ_LIGHTMAP_COEF * MIP_FACTOR * LightMapWidth * LightMapHeight ); // DXT5
 	}
diff --git a/Engine/Source/Runtime/Engine/Private/Particles/FXSystem.cpp b/Engine/Source/Runtime/Engine/Private/Particles/FXSystem.cpp
--- a/Engine/Source/Runtime/Engine/Private/Particles/FXSystem.cpp
+++ b/Engine/Source/Runtime/Engine/Private/Particles/FXSystem.cpp
@@ -162,7 +162,7 @@ void FFXSystem::Tick(float DeltaSeconds)
 		// Test GPU sorting if requested.
 		if (FXConsoleVariables::TestGPUSort.GetValueOnGameThread() != 0)
 		{
-			TestGPUSort((EGPUSortTest)FXConsoleVariables::TestGPUSort.GetValueOnGameThread());
+			TestGPUSort((EGPUSortTest)FXConsoleVariables::TestGPUSort.GetValueOnGameThread(), GetFeatureLevel());
 			// Reset CVar
 			static IConsoleVariable* CVarTestGPUSort = IConsoleManager::Get().FindConsoleVariable(TEXT("FX.TestGPUSort"));
 			CVarTestGPUSort->Set(0);
diff --git a/Engine/Source/Runtime/Engine/Private/Particles/ParticleGpuSimulation.cpp b/Engine/Source/Runtime/Engine/Private/Particles/ParticleGpuSimulation.cpp
--- a/Engine/Source/Runtime/Engine/Private/Particles/ParticleGpuSimulation.cpp
+++ b/Engine/Source/Runtime/Engine/Private/Particles/ParticleGpuSimulation.cpp
@@ -2022,10 +2022,7 @@ static FBox ComputeParticleBounds(
 	FParticleBoundsParameters Parameters;
 	FParticleBoundsUniformBufferRef UniformBuffer;
 
-	// MOBILEPREVIEWTODO: Proper value for this
-	const auto FeatureLevel = GRHIFeatureLevel;
-
-	if (ParticleCount > 0 && FeatureLevel == ERHIFeatureLevel::SM5)
+	if (ParticleCount > 0 && GMaxRHIFeatureLevel == ERHIFeatureLevel::SM5)
 	{
 		// Determine how to break the work up over individual work groups.
 		const uint32 MaxGroupCount = 128;
@@ -2051,7 +2048,7 @@ static FBox ComputeParticleBounds(
 			PF_A32B32G32R32F );
 
 		// Grab the shader.
-		TShaderMapRef<FParticleBoundsCS> ParticleBoundsCS(GetGlobalShaderMap(FeatureLevel));
+		TShaderMapRef<FParticleBoundsCS> ParticleBoundsCS(GetGlobalShaderMap(GMaxRHIFeatureLevel));
 		RHICmdList.SetComputeShader(ParticleBoundsCS->GetComputeShader());
 
 		// Dispatch shader to compute bounds.
@@ -4077,7 +4074,8 @@ void FFXSystem::SortGPUParticles(FRHICommandListImmediate& RHICmdList)
 			RHICmdList,
 			GParticleSortBuffers,
 			ParticleSimulationResources->GetCurrentStateTextures().PositionTextureRHI,
-			ParticleSimulationResources->SimulationsToSort
+			ParticleSimulationResources->SimulationsToSort,
+			GetFeatureLevel()
 			);
 		ParticleSimulationResources->SortedVertexBuffer.VertexBufferRHI =
 			GParticleSortBuffers.GetSortedVertexBufferRHI(BufferIndex);
diff --git a/Engine/Source/Runtime/Engine/Private/Particles/ParticleSortingGPU.cpp b/Engine/Source/Runtime/Engine/Private/Particles/ParticleSortingGPU.cpp
--- a/Engine/Source/Runtime/Engine/Private/Particles/ParticleSortingGPU.cpp
+++ b/Engine/Source/Runtime/Engine/Private/Particles/ParticleSortingGPU.cpp
@@ -180,14 +180,11 @@ static int32 GenerateParticleSortKeys(
 	FUnorderedAccessViewRHIParamRef KeyBufferUAV,
 	FUnorderedAccessViewRHIParamRef SortedVertexBufferUAV,
 	FTexture2DRHIParamRef PositionTextureRHI,
-	const TArray<FParticleSimulationSortInfo>& SimulationsToSort
+	const TArray<FParticleSimulationSortInfo>& SimulationsToSort,
+	ERHIFeatureLevel::Type FeatureLevel
 	)
 {
 	SCOPED_DRAW_EVENT(RHICmdList, ParticleSortKeyGen, DEC_PARTICLE);
-
-	// MOBILEPREVIEWTODO: Proper value for this
-	const auto FeatureLevel = GRHIFeatureLevel;
-
 	check(FeatureLevel == ERHIFeatureLevel::SM5);
 
 	FParticleKeyGenParameters KeyGenParameters;
@@ -314,13 +311,12 @@ int32 SortParticlesGPU(
 	FRHICommandListImmediate& RHICmdList,
 	FParticleSortBuffers& ParticleSortBuffers,
 	FTexture2DRHIParamRef PositionTextureRHI,
-	const TArray<FParticleSimulationSortInfo>& SimulationsToSort
+	const TArray<FParticleSimulationSortInfo>& SimulationsToSort,
+	ERHIFeatureLevel::Type FeatureLevel
 	)
 {
 	SCOPED_DRAW_EVENTF(RHICmdList, ParticleSort, DEC_PARTICLE, TEXT("ParticleSort_%d"), SimulationsToSort.Num());
 
-	check(GRHIFeatureLevel == ERHIFeatureLevel::SM5);
-
 	// Ensure the sorted vertex buffers are not currently bound as input streams.
 	// They should only ever be bound to streams 0 or 1, so clear them.
 	{
@@ -338,7 +334,8 @@ int32 SortParticlesGPU(
 		ParticleSortBuffers.GetKeyBufferUAV(),
 		ParticleSortBuffers.GetVertexBufferUAV(),
 		PositionTextureRHI,
-		SimulationsToSort
+		SimulationsToSort,
+		FeatureLevel
 		);
 
 	// Update stats.
@@ -349,5 +346,5 @@ int32 SortParticlesGPU(
 	const uint32 EmitterKeyMask = (1 << FMath::CeilLogTwo( SimulationsToSort.Num() )) - 1;
 	const uint32 KeyMask = (EmitterKeyMask << 16) | 0xFFFF;
 	FGPUSortBuffers SortBuffers = ParticleSortBuffers.GetSortBuffers();
-	return SortGPUBuffers(RHICmdList, SortBuffers, 0, KeyMask, TotalParticleCount);
+	return SortGPUBuffers(RHICmdList, SortBuffers, 0, KeyMask, TotalParticleCount, FeatureLevel);
 }
diff --git a/Engine/Source/Runtime/Engine/Private/Particles/ParticleSortingGPU.h b/Engine/Source/Runtime/Engine/Private/Particles/ParticleSortingGPU.h
--- a/Engine/Source/Runtime/Engine/Private/Particles/ParticleSortingGPU.h
+++ b/Engine/Source/Runtime/Engine/Private/Particles/ParticleSortingGPU.h
@@ -125,5 +125,6 @@ int32 SortParticlesGPU(
 	FRHICommandListImmediate& RHICmdList,
 	FParticleSortBuffers& ParticleSortBuffers,
 	FTexture2DRHIParamRef PositionTextureRHI,
-	const TArray<FParticleSimulationSortInfo>& SimulationsToSort
+	const TArray<FParticleSimulationSortInfo>& SimulationsToSort,
+	ERHIFeatureLevel::Type FeatureLevel
 	);
diff --git a/Engine/Source/Runtime/Engine/Private/Particles/ParticleSystemRender.cpp b/Engine/Source/Runtime/Engine/Private/Particles/ParticleSystemRender.cpp
--- a/Engine/Source/Runtime/Engine/Private/Particles/ParticleSystemRender.cpp
+++ b/Engine/Source/Runtime/Engine/Private/Particles/ParticleSystemRender.cpp
@@ -1837,7 +1837,7 @@ void FDynamicMeshEmitterData::GetDynamicMeshElementsEmitter(const FParticleSyste
 {
 	SCOPE_CYCLE_COUNTER(STAT_MeshRenderingTime);
 
-	const bool bInstanced = GRHIFeatureLevel >= ERHIFeatureLevel::SM4;
+	const bool bInstanced = View->GetFeatureLevel() >= ERHIFeatureLevel::SM4;
 
 	if (bValid)
 	{
diff --git a/Engine/Source/Runtime/Engine/Private/PrecomputedLightVolume.cpp b/Engine/Source/Runtime/Engine/Private/PrecomputedLightVolume.cpp
--- a/Engine/Source/Runtime/Engine/Private/PrecomputedLightVolume.cpp
+++ b/Engine/Source/Runtime/Engine/Private/PrecomputedLightVolume.cpp
@@ -80,7 +80,7 @@ FArchive& operator<<(FArchive& Ar,FPrecomputedLightVolume& Volume)
 			Ar << HighQualitySamples;
 
 			if (FPlatformProperties::SupportsHighQualityLightmaps() 
-				&& (GIsEditor || AllowHighQualityLightmaps()))
+				&& (GIsEditor || AllowHighQualityLightmaps(GMaxRHIFeatureLevel)))
 			{
 				for(int32 SampleIndex = 0; SampleIndex < HighQualitySamples.Num(); SampleIndex++)
 				{
@@ -96,7 +96,7 @@ FArchive& operator<<(FArchive& Ar,FPrecomputedLightVolume& Volume)
 			}
 
 			if (FPlatformProperties::SupportsLowQualityLightmaps() 
-				&& (GIsEditor || !AllowHighQualityLightmaps()))
+				&& (GIsEditor || !AllowHighQualityLightmaps(GMaxRHIFeatureLevel)))
 			{
 				for(int32 SampleIndex = 0; SampleIndex < LowQualitySamples.Num(); SampleIndex++)
 				{
@@ -184,7 +184,7 @@ void FPrecomputedLightVolume::Initialize(const FBox& NewBounds)
 	HighQualityLightmapOctree = FLightVolumeOctree(NewBounds.GetCenter(), NewBounds.GetExtent().GetMax());
 	LowQualityLightmapOctree = FLightVolumeOctree(NewBounds.GetCenter(), NewBounds.GetExtent().GetMax());
 
-	OctreeForRendering = AllowHighQualityLightmaps() ? &HighQualityLightmapOctree : &LowQualityLightmapOctree;
+	OctreeForRendering = AllowHighQualityLightmaps(GMaxRHIFeatureLevel) ? &HighQualityLightmapOctree : &LowQualityLightmapOctree;
 }
 
 /** Adds a lighting sample. */
diff --git a/Engine/Source/Runtime/Engine/Private/PrimitiveComponent.cpp b/Engine/Source/Runtime/Engine/Private/PrimitiveComponent.cpp
--- a/Engine/Source/Runtime/Engine/Private/PrimitiveComponent.cpp
+++ b/Engine/Source/Runtime/Engine/Private/PrimitiveComponent.cpp
@@ -168,8 +168,10 @@ void UPrimitiveComponent::GetUsedTextures(TArray<UTexture*>& OutTextures, EMater
 		// Ensure we don't have any NULL elements.
 		if( UsedMaterials[ MatIndex ] )
 		{
+			auto World = GetWorld();
+
 			UsedTextures.Reset();
-			UsedMaterials[ MatIndex ]->GetUsedTextures( UsedTextures, QualityLevel, false, GRHIFeatureLevel, false );
+			UsedMaterials[MatIndex]->GetUsedTextures(UsedTextures, QualityLevel, false, World ? World->FeatureLevel : GMaxRHIFeatureLevel, false);
 
 			for( int32 TextureIndex=0; TextureIndex<UsedTextures.Num(); TextureIndex++ )
 			{
diff --git a/Engine/Source/Runtime/Engine/Private/ReflectionCaptureComponent.cpp b/Engine/Source/Runtime/Engine/Private/ReflectionCaptureComponent.cpp
--- a/Engine/Source/Runtime/Engine/Private/ReflectionCaptureComponent.cpp
+++ b/Engine/Source/Runtime/Engine/Private/ReflectionCaptureComponent.cpp
@@ -1245,11 +1245,11 @@ void ReadbackFromSM4Cubemap(FReflectionTextureCubeResource* SM4FullHDRCubemapTex
 
 void UReflectionCaptureComponent::ReadbackFromGPUAndSaveDerivedData(UWorld* WorldToUpdate)
 {
-	if (bDerivedDataDirty && !IsRunningCommandlet())
+	if (bDerivedDataDirty && !IsRunningCommandlet() && WorldToUpdate->FeatureLevel >= ERHIFeatureLevel::SM4)
 	{
 		FReflectionCaptureFullHDRDerivedData* NewDerivedData = new FReflectionCaptureFullHDRDerivedData();
 
-		if (GRHIFeatureLevel == ERHIFeatureLevel::SM4)
+		if (WorldToUpdate->FeatureLevel == ERHIFeatureLevel::SM4)
 		{
 			ReadbackFromSM4Cubemap(SM4FullHDRCubemapTexture, NewDerivedData);
 		}
diff --git a/Engine/Source/Runtime/Engine/Private/RendererSupport.cpp b/Engine/Source/Runtime/Engine/Private/RendererSupport.cpp
--- a/Engine/Source/Runtime/Engine/Private/RendererSupport.cpp
+++ b/Engine/Source/Runtime/Engine/Private/RendererSupport.cpp
@@ -159,8 +159,12 @@ void RestoreReferencesToRendererModuleClasses(
 	FShaderType::GetOutdatedTypes(OutdatedShaderTypes, OutdatedFactoryTypes);
 
 	// Recompile any missing shaders
-	BeginRecompileGlobalShaders(OutdatedShaderTypes);
-	UMaterial::UpdateMaterialShaders(OutdatedShaderTypes, OutdatedFactoryTypes, GRHIShaderPlatform);
+	UMaterialInterface::IterateOverActiveFeatureLevels([&](ERHIFeatureLevel::Type FeatureLevel) 
+	{
+		auto ShaderPlatform = GShaderPlatformForFeatureLevel[FeatureLevel];
+		BeginRecompileGlobalShaders(OutdatedShaderTypes, ShaderPlatform);
+		UMaterial::UpdateMaterialShaders(OutdatedShaderTypes, OutdatedFactoryTypes, ShaderPlatform);
+	});
 
 	// Block on global shader jobs
 	FinishRecompileGlobalShaders();
diff --git a/Engine/Source/Runtime/Engine/Private/SceneManagement.cpp b/Engine/Source/Runtime/Engine/Private/SceneManagement.cpp
--- a/Engine/Source/Runtime/Engine/Private/SceneManagement.cpp
+++ b/Engine/Source/Runtime/Engine/Private/SceneManagement.cpp
@@ -106,7 +106,7 @@ void FSimpleElementCollector::DrawBatchedElements(FRHICommandList& RHICmdList, c
 	const bool bIsMobileHDR = MobileHDRCvar->GetValueOnRenderThread() == 1;
 
 	// Mobile HDR does not execute post process, so does not need to render flipped
-	const bool bNeedToSwitchVerticalAxis = RHINeedsToSwitchVerticalAxis(GRHIShaderPlatform) && !bIsMobileHDR;
+	const bool bNeedToSwitchVerticalAxis = RHINeedsToSwitchVerticalAxis(View.GetShaderPlatform()) && !bIsMobileHDR;
 
 	// Draw the batched elements.
 	BatchedElements.Draw(
diff --git a/Engine/Source/Runtime/Engine/Private/SceneView.cpp b/Engine/Source/Runtime/Engine/Private/SceneView.cpp
--- a/Engine/Source/Runtime/Engine/Private/SceneView.cpp
+++ b/Engine/Source/Runtime/Engine/Private/SceneView.cpp
@@ -271,7 +271,7 @@ FSceneView::FSceneView(const FSceneViewInitOptions& InitOptions)
 	, OverrideLODViewOrigin(InitOptions.OverrideLODViewOrigin)
 	, bAllowTranslucentPrimitivesInHitProxy( true )
 #endif
-	, FeatureLevel(InitOptions.ViewFamily ? InitOptions.ViewFamily->GetFeatureLevel() : GRHIFeatureLevel)
+	, FeatureLevel(InitOptions.ViewFamily ? InitOptions.ViewFamily->GetFeatureLevel() : GMaxRHIFeatureLevel)
 {
 	check(UnscaledViewRect.Min.X >= 0);
 	check(UnscaledViewRect.Min.Y >= 0);
@@ -395,7 +395,8 @@ FSceneView::FSceneView(const FSceneViewInitOptions& InitOptions)
 	bReverseCulling = (Family && Family->Scene) ? FMath::IsNegativeFloat(ViewMatrices.ViewMatrix.Determinant()) : false;
 
 	// OpenGL Gamma space output in GLSL flips Y when rendering directly to the back buffer (so not needed on PC, as we never render directly into the back buffer)
-	static bool bPlatformRequiresReverseCulling = (IsOpenGLPlatform(GRHIShaderPlatform) && !IsPCPlatform(GRHIShaderPlatform));
+	auto ShaderPlatform = GShaderPlatformForFeatureLevel[FeatureLevel];
+	static bool bPlatformRequiresReverseCulling = (IsOpenGLPlatform(ShaderPlatform) && !IsPCPlatform(ShaderPlatform));
 	static auto* MobileHDRCvar = IConsoleManager::Get().FindTConsoleVariableDataInt(TEXT("r.MobileHDR"));
 	check(MobileHDRCvar);
 	bReverseCulling = (bPlatformRequiresReverseCulling && MobileHDRCvar->GetValueOnAnyThread() == 0) ? !bReverseCulling : bReverseCulling;
@@ -1106,13 +1107,14 @@ void FSceneView::EndFinalPostprocessSettings()
 
 	// Anti-Aliasing
 	{
+		const auto FeatureLevel = GetFeatureLevel();
+
 		static const auto CVar = IConsoleManager::Get().FindTConsoleVariableDataInt(TEXT("r.PostProcessAAQuality")); 
 		static auto* MobileHDRCvar = IConsoleManager::Get().FindTConsoleVariableDataInt(TEXT("r.MobileHDR"));
 		static auto* MobileMSAACvar = IConsoleManager::Get().FindTConsoleVariableDataInt(TEXT("r.MobileMSAA"));
-		static uint32 MSAAValue = GRHIShaderPlatform == SP_OPENGL_ES2_IOS ? 1 : MobileMSAACvar->GetValueOnGameThread();
+		static uint32 MSAAValue = GShaderPlatformForFeatureLevel[FeatureLevel] == SP_OPENGL_ES2_IOS ? 1 : MobileMSAACvar->GetValueOnGameThread();
 
 		int32 Quality = FMath::Clamp(CVar->GetValueOnGameThread(), 0, 6);
-		const auto FeatureLevel = GetFeatureLevel();
 
 		if( !Family->EngineShowFlags.PostProcessing || !Family->EngineShowFlags.AntiAliasing || Quality <= 0
 			// Disable antialiasing in GammaLDR mode to avoid jittering.
@@ -1379,7 +1381,7 @@ ERHIFeatureLevel::Type FSceneViewFamily::GetFeatureLevel() const
 	}
 	else
 	{
-		return GRHIFeatureLevel;
+		return GMaxRHIFeatureLevel;
 	}
 }
 
diff --git a/Engine/Source/Runtime/Engine/Private/ShaderCompiler/ShaderCompiler.cpp b/Engine/Source/Runtime/Engine/Private/ShaderCompiler/ShaderCompiler.cpp
--- a/Engine/Source/Runtime/Engine/Private/ShaderCompiler/ShaderCompiler.cpp
+++ b/Engine/Source/Runtime/Engine/Private/ShaderCompiler/ShaderCompiler.cpp
@@ -2332,8 +2332,11 @@ bool RecompileShaders(const TCHAR* Cmd, FOutputDevice& Ar)
 				FRecompileShadersTimer TestTimer(TEXT("RecompileShaders Changed"));
 
 				// Kick off global shader recompiles
-				BeginRecompileGlobalShaders(OutdatedShaderTypes);
-				UMaterial::UpdateMaterialShaders(OutdatedShaderTypes, OutdatedFactoryTypes, GRHIShaderPlatform);
+				UMaterialInterface::IterateOverActiveFeatureLevels([&](ERHIFeatureLevel::Type InFeatureLevel) {
+					auto ShaderPlatform = GShaderPlatformForFeatureLevel[InFeatureLevel];
+					BeginRecompileGlobalShaders(OutdatedShaderTypes, ShaderPlatform);
+					UMaterial::UpdateMaterialShaders(OutdatedShaderTypes, OutdatedFactoryTypes, ShaderPlatform);
+				});
 
 				GWarn->StatusUpdate(0, 1, NSLOCTEXT("ShaderCompilingManager", "CompilingGlobalShaderStatus", "Compiling global shaders..."));
 
@@ -2406,7 +2409,7 @@ bool RecompileShaders(const TCHAR* Cmd, FOutputDevice& Ar)
 				
 				TArray<const FVertexFactoryType*> FactoryTypes;
 
-				BeginRecompileGlobalShaders(ShaderTypes);
+				BeginRecompileGlobalShaders(ShaderTypes, GRHIShaderPlatform);
 				UMaterial::UpdateMaterialShaders(ShaderTypes, FactoryTypes, GRHIShaderPlatform);
 				FinishRecompileGlobalShaders();
 			}
diff --git a/Engine/Source/Runtime/Engine/Private/SkeletalMesh.cpp b/Engine/Source/Runtime/Engine/Private/SkeletalMesh.cpp
--- a/Engine/Source/Runtime/Engine/Private/SkeletalMesh.cpp
+++ b/Engine/Source/Runtime/Engine/Private/SkeletalMesh.cpp
@@ -1262,7 +1262,7 @@ void FStaticLODModel::InitResources(bool bNeedsVertexColors)
 		BeginInitResource(&APEXClothVertexBuffer);
 	}
 
-	if( RHISupportsTessellation(GRHIShaderPlatform) ) 
+	if( RHISupportsTessellation(GMaxRHIShaderPlatform) ) 
 	{
 		AdjacencyMultiSizeIndexContainer.InitResources();
 		INC_DWORD_STAT_BY( STAT_SkeletalMeshIndexMemory, AdjacencyMultiSizeIndexContainer.IsIndexBufferValid() ? (AdjacencyMultiSizeIndexContainer.GetIndexBuffer()->Num() * AdjacencyMultiSizeIndexContainer.GetDataTypeSize()) : 0 );
diff --git a/Engine/Source/Runtime/Engine/Private/SkinnedMeshComponent.cpp b/Engine/Source/Runtime/Engine/Private/SkinnedMeshComponent.cpp
--- a/Engine/Source/Runtime/Engine/Private/SkinnedMeshComponent.cpp
+++ b/Engine/Source/Runtime/Engine/Private/SkinnedMeshComponent.cpp
@@ -379,7 +379,8 @@ void USkinnedMeshComponent::GetStreamingTextureInfo(TArray<FStreamingTexturePrim
 			{
 				TArray<UTexture*> Textures;
 				
-				MaterialInterface->GetUsedTextures(Textures, EMaterialQualityLevel::Num, false, GRHIFeatureLevel, false);
+				auto World = GetWorld();
+				MaterialInterface->GetUsedTextures(Textures, EMaterialQualityLevel::Num, false, World ? World->FeatureLevel : GMaxRHIFeatureLevel, false);
 				for(int32 TextureIndex = 0; TextureIndex < Textures.Num(); TextureIndex++ )
 				{
 					FStreamingTexturePrimitiveInfo& StreamingTexture = *new(OutStreamingTextures) FStreamingTexturePrimitiveInfo;
diff --git a/Engine/Source/Runtime/Engine/Private/Slate/DebugCanvas.cpp b/Engine/Source/Runtime/Engine/Private/Slate/DebugCanvas.cpp
--- a/Engine/Source/Runtime/Engine/Private/Slate/DebugCanvas.cpp
+++ b/Engine/Source/Runtime/Engine/Private/Slate/DebugCanvas.cpp
@@ -27,7 +27,7 @@ class FCanvasProxy
 {
 public:
 	FCanvasProxy( FRenderTarget* RenderTarget, UWorld* InWorld )
-		: Canvas(RenderTarget, NULL, InWorld, InWorld ? InWorld->FeatureLevel : GRHIFeatureLevel)
+		: Canvas(RenderTarget, NULL, InWorld, InWorld ? InWorld->FeatureLevel : GMaxRHIFeatureLevel)
 	{
 		// Do not allow the canvas to be flushed outside of our debug rendering path
 		Canvas.SetAllowedModes( 0 );
diff --git a/Engine/Source/Runtime/Engine/Private/Slate/SceneViewport.cpp b/Engine/Source/Runtime/Engine/Private/Slate/SceneViewport.cpp
--- a/Engine/Source/Runtime/Engine/Private/Slate/SceneViewport.cpp
+++ b/Engine/Source/Runtime/Engine/Private/Slate/SceneViewport.cpp
@@ -1098,7 +1098,7 @@ void FSceneViewport::EndRenderFrame(FRHICommandListImmediate& RHICmdList, bool b
 	if (bUseSeparateRenderTarget)
 	{
 		// @todo-mobile
-		if (GRHIShaderPlatform == SP_OPENGL_ES2)
+		if (GMaxRHIShaderPlatform == SP_OPENGL_ES2)
 		{
 			check(0);
 		}
diff --git a/Engine/Source/Runtime/Engine/Private/StaticBoundShaderState.cpp b/Engine/Source/Runtime/Engine/Private/StaticBoundShaderState.cpp
--- a/Engine/Source/Runtime/Engine/Private/StaticBoundShaderState.cpp
+++ b/Engine/Source/Runtime/Engine/Private/StaticBoundShaderState.cpp
@@ -136,11 +136,14 @@ class FSetGlobalBoundShaderStateRenderThreadTask
 {
 	FRHICommandList& RHICmdList;
 	FGlobalBoundShaderState& GlobalBoundShaderState;
+	ERHIFeatureLevel::Type FeatureLevel;
+
 public:
 
-	FSetGlobalBoundShaderStateRenderThreadTask(FRHICommandList* InRHICmdList, FGlobalBoundShaderState* InGlobalBoundShaderState)
+	FSetGlobalBoundShaderStateRenderThreadTask(FRHICommandList* InRHICmdList, FGlobalBoundShaderState* InGlobalBoundShaderState, ERHIFeatureLevel::Type InFeatureLevel)
 		: RHICmdList(*InRHICmdList)
 		, GlobalBoundShaderState(*InGlobalBoundShaderState)
+		, FeatureLevel(InFeatureLevel)
 	{
 	}
 
@@ -158,7 +161,7 @@ public:
 
 	void DoTask(ENamedThreads::Type CurrentThread, const FGraphEventRef& MyCompletionGraphEvent)
 	{
-		RHICmdList.SetBoundShaderState(GetGlobalBoundShaderState_Internal(GlobalBoundShaderState, GRHIFeatureLevel));
+		RHICmdList.SetBoundShaderState(GetGlobalBoundShaderState_Internal(GlobalBoundShaderState, FeatureLevel));
 	}
 };
 
@@ -228,7 +231,7 @@ void SetGlobalBoundShaderState(
 	// We need to do this on the render thread
 
 	FRHICommandList* CmdList = new FRHICommandList;
-	FGraphEventRef RenderThreadCompletionEvent = TGraphTask<FSetGlobalBoundShaderStateRenderThreadTask>::CreateTask().ConstructAndDispatchWhenReady(CmdList, &GlobalBoundShaderState);
+	FGraphEventRef RenderThreadCompletionEvent = TGraphTask<FSetGlobalBoundShaderStateRenderThreadTask>::CreateTask().ConstructAndDispatchWhenReady(CmdList, &GlobalBoundShaderState, FeatureLevel);
 	RHICmdList.QueueRenderThreadCommandListSubmit(RenderThreadCompletionEvent, CmdList);
 }
 
diff --git a/Engine/Source/Runtime/Engine/Private/StaticMeshComponent.cpp b/Engine/Source/Runtime/Engine/Private/StaticMeshComponent.cpp
--- a/Engine/Source/Runtime/Engine/Private/StaticMeshComponent.cpp
+++ b/Engine/Source/Runtime/Engine/Private/StaticMeshComponent.cpp
@@ -1356,7 +1356,10 @@ void UStaticMeshComponent::GetTextureLightAndShadowMapMemoryUsage(int32 InWidth,
 	// Stored in texture.
 	const float MIP_FACTOR = 1.33f;
 	OutShadowMapMemoryUsage = FMath::TruncToInt(MIP_FACTOR * InWidth * InHeight); // G8
-	if( AllowHighQualityLightmaps() )
+
+	auto FeatureLevel = GetWorld() ? GetWorld()->FeatureLevel : GMaxRHIFeatureLevel;
+
+	if (AllowHighQualityLightmaps(FeatureLevel))
 	{
 		OutLightMapMemoryUsage = FMath::TruncToInt(NUM_HQ_LIGHTMAP_COEF * MIP_FACTOR * InWidth * InHeight); // DXT5
 	}
diff --git a/Engine/Source/Runtime/Engine/Private/UnrealClient.cpp b/Engine/Source/Runtime/Engine/Private/UnrealClient.cpp
--- a/Engine/Source/Runtime/Engine/Private/UnrealClient.cpp
+++ b/Engine/Source/Runtime/Engine/Private/UnrealClient.cpp
@@ -1011,7 +1011,7 @@ void FViewport::Draw( bool bShouldPresent /*= true */)
 				}
 
 				auto World = ViewportClient->GetWorld();
-				FCanvas Canvas(this, NULL, World, World ? World->FeatureLevel : GRHIFeatureLevel);
+				FCanvas Canvas(this, NULL, World, World ? World->FeatureLevel : GMaxRHIFeatureLevel);
 				{
 					// Make sure the Canvas is not rendered upside down
 					Canvas.SetAllowSwitchVerticalAxis(false);
@@ -1117,7 +1117,7 @@ const TArray<FColor>& FViewport::GetRawHitProxyData(FIntRect InRect)
 
 		// Let the viewport client draw its hit proxies.
 		auto World = ViewportClient->GetWorld();
-		FCanvas Canvas(&HitProxyMap, &HitProxyMap, World, World ? World->FeatureLevel : GRHIFeatureLevel);
+		FCanvas Canvas(&HitProxyMap, &HitProxyMap, World, World ? World->FeatureLevel : GMaxRHIFeatureLevel);
 		{
 			ViewportClient->Draw(this, &Canvas);
 		}
diff --git a/Engine/Source/Runtime/Engine/Private/UnrealEngine.cpp b/Engine/Source/Runtime/Engine/Private/UnrealEngine.cpp
--- a/Engine/Source/Runtime/Engine/Private/UnrealEngine.cpp
+++ b/Engine/Source/Runtime/Engine/Private/UnrealEngine.cpp
@@ -2956,7 +2956,7 @@ bool UEngine::HandleRecompileGlobalShadersCommand( const TCHAR* Cmd, FOutputDevi
 bool UEngine::HandleDumpShaderStatsCommand( const TCHAR* Cmd, FOutputDevice& Ar )
 {
 	FString FlagStr(FParse::Token(Cmd, 0));
-	EShaderPlatform Platform = GRHIShaderPlatform;
+	EShaderPlatform Platform = GMaxRHIShaderPlatform;
 	if (FlagStr.Len() > 0)
 	{
 		Platform = ShaderFormatToLegacyShaderPlatform(FName(*FlagStr));
@@ -2970,7 +2970,7 @@ bool UEngine::HandleDumpShaderStatsCommand( const TCHAR* Cmd, FOutputDevice& Ar
 bool UEngine::HandleDumpMaterialStatsCommand( const TCHAR* Cmd, FOutputDevice& Ar )
 {
 	FString FlagStr(FParse::Token(Cmd, 0));
-	EShaderPlatform Platform = GRHIShaderPlatform;
+	EShaderPlatform Platform = GMaxRHIShaderPlatform;
 	if (FlagStr.Len() > 0)
 	{
 		Platform = ShaderFormatToLegacyShaderPlatform(FName(*FlagStr));
@@ -3255,7 +3255,7 @@ bool UEngine::HandleRemoteTextureStatsCommand( const TCHAR* Cmd, FOutputDevice&
 			if( UsedMaterials[ MatIndex ] )
 			{
 				UsedTextures.Reset();
-				UsedMaterials[ MatIndex ]->GetUsedTextures( UsedTextures, EMaterialQualityLevel::Num, false, GRHIFeatureLevel, false );
+				UsedMaterials[MatIndex]->GetUsedTextures(UsedTextures, EMaterialQualityLevel::Num, false, GMaxRHIFeatureLevel, true);
 
 				// Increase usage count for all referenced textures
 				for( int32 TextureIndex=0; TextureIndex<UsedTextures.Num(); TextureIndex++ )
@@ -10239,11 +10239,6 @@ bool AllowHighQualityLightmaps(ERHIFeatureLevel::Type FeatureLevel)
 		&& !IsMobilePlatform(GShaderPlatformForFeatureLevel[FeatureLevel]);
 }
 
-bool AllowHighQualityLightmaps()
-{
-	return AllowHighQualityLightmaps(GRHIFeatureLevel);
-}
-
 // Helper function for changing system resolution via the r.setres console command
 void FSystemResolution::RequestResolutionChange(int32 InResX, int32 InResY, EWindowMode::Type InWindowMode)
 {
diff --git a/Engine/Source/Runtime/Engine/Private/UserInterface/Canvas.cpp b/Engine/Source/Runtime/Engine/Private/UserInterface/Canvas.cpp
--- a/Engine/Source/Runtime/Engine/Private/UserInterface/Canvas.cpp
+++ b/Engine/Source/Runtime/Engine/Private/UserInterface/Canvas.cpp
@@ -35,7 +35,6 @@ FCanvas::FCanvas(FRenderTarget* InRenderTarget, FHitProxyConsumer* InHitProxyCon
 ,	CurrentWorldTime(0)
 ,	CurrentDeltaWorldTime(0)
 ,	FeatureLevel(InFeatureLevel)
-,	ShaderPlatform(GRHIShaderPlatform)
 {
 	Construct();
 
@@ -57,7 +56,6 @@ FCanvas::FCanvas(FRenderTarget* InRenderTarget,FHitProxyConsumer* InHitProxyCons
 ,	CurrentWorldTime(InWorldTime)
 ,	CurrentDeltaWorldTime(InWorldDeltaTime)
 ,	FeatureLevel(InFeatureLevel)
-,	ShaderPlatform(GRHIShaderPlatform)
 {
 	Construct();
 }
diff --git a/Engine/Source/Runtime/Engine/Public/CanvasTypes.h b/Engine/Source/Runtime/Engine/Public/CanvasTypes.h
--- a/Engine/Source/Runtime/Engine/Public/CanvasTypes.h
+++ b/Engine/Source/Runtime/Engine/Public/CanvasTypes.h
@@ -310,7 +310,7 @@ public:
 	*
 	* @return shader platform that this canvas is rendering at
 	*/
-	EShaderPlatform GetShaderPlatform() const { return ShaderPlatform; }
+	EShaderPlatform GetShaderPlatform() const { return GShaderPlatformForFeatureLevel[FeatureLevel]; }
 
 	// Get/Set if this Canvas allows its batched elements to switch vertical axis (e.g., rendering to back buffer should never flip)
 	bool GetAllowSwitchVerticalAxis() const { return bAllowsToSwitchVerticalAxis; }
@@ -391,8 +391,6 @@ private:
 	bool bAllowsToSwitchVerticalAxis;
 	/** Feature level that we are currently rendering with */
 	ERHIFeatureLevel::Type FeatureLevel;
-	/** Shader platform that we are currently rendering with */
-	EShaderPlatform ShaderPlatform;
 
 	/** 
 	* Shared construction function
diff --git a/Engine/Source/Runtime/Engine/Public/GlobalShader.h b/Engine/Source/Runtime/Engine/Public/GlobalShader.h
--- a/Engine/Source/Runtime/Engine/Public/GlobalShader.h
+++ b/Engine/Source/Runtime/Engine/Public/GlobalShader.h
@@ -201,7 +201,7 @@ extern ENGINE_API void RecompileGlobalShaders();
  * Begins recompiling the specified global shader types, and flushes their bound shader states.
  * FinishRecompileGlobalShaders must be called after this and before using the global shaders for anything.
  */
-extern ENGINE_API void BeginRecompileGlobalShaders(const TArray<FShaderType*>& OutdatedShaderTypes, EShaderPlatform ShaderPlatform=GRHIShaderPlatform);
+extern ENGINE_API void BeginRecompileGlobalShaders(const TArray<FShaderType*>& OutdatedShaderTypes, EShaderPlatform ShaderPlatform);
 
 /** Finishes recompiling global shaders.  Must be called after BeginRecompileGlobalShaders. */
 extern ENGINE_API void FinishRecompileGlobalShaders();
diff --git a/Engine/Source/Runtime/Engine/Public/MaterialShared.h b/Engine/Source/Runtime/Engine/Public/MaterialShared.h
--- a/Engine/Source/Runtime/Engine/Public/MaterialShared.h
+++ b/Engine/Source/Runtime/Engine/Public/MaterialShared.h
@@ -1556,7 +1556,7 @@ public:
 	};
 
 	/** Initialization constructor. */
-	explicit ENGINE_API FMaterialUpdateContext(uint32 Options = EOptions::Default, EShaderPlatform InShaderPlatform=GRHIShaderPlatform);
+	explicit ENGINE_API FMaterialUpdateContext(uint32 Options = EOptions::Default, EShaderPlatform InShaderPlatform = GMaxRHIShaderPlatform);
 
 	/** Destructor. */
 	ENGINE_API ~FMaterialUpdateContext();
diff --git a/Engine/Source/Runtime/Engine/Public/SceneInterface.h b/Engine/Source/Runtime/Engine/Public/SceneInterface.h
--- a/Engine/Source/Runtime/Engine/Public/SceneInterface.h
+++ b/Engine/Source/Runtime/Engine/Public/SceneInterface.h
@@ -304,7 +304,7 @@ public:
 
 	virtual bool IsEditorScene() const { return false; }
 
-	virtual ERHIFeatureLevel::Type GetFeatureLevel() const { return GRHIFeatureLevel; }
+	virtual ERHIFeatureLevel::Type GetFeatureLevel() const { return GMaxRHIFeatureLevel; }
 
 	bool ShouldUseDeferredRenderer() const
 	{
diff --git a/Engine/Source/Runtime/Engine/Public/SceneManagement.h b/Engine/Source/Runtime/Engine/Public/SceneManagement.h
--- a/Engine/Source/Runtime/Engine/Public/SceneManagement.h
+++ b/Engine/Source/Runtime/Engine/Public/SceneManagement.h
@@ -620,7 +620,7 @@ public:
 	/** Accesses parameters needed for rendering the light. */
 	virtual void GetParameters(FVector4& LightPositionAndInvRadius, FVector4& LightColorAndFalloffExponent, FVector& NormalizedLightDirection, FVector2D& SpotAngles, float& LightSourceRadius, float& LightSourceLength, float& LightMinRoughness) const {}
 
-	virtual FVector2D GetDirectionalLightDistanceFadeParameters() const
+	virtual FVector2D GetDirectionalLightDistanceFadeParameters(ERHIFeatureLevel::Type InFeatureLevel) const
 	{
 		return FVector2D(0, 0);
 	}
diff --git a/Engine/Source/Runtime/Engine/Public/Tests/AutomationCommon.h b/Engine/Source/Runtime/Engine/Public/Tests/AutomationCommon.h
--- a/Engine/Source/Runtime/Engine/Public/Tests/AutomationCommon.h
+++ b/Engine/Source/Runtime/Engine/Public/Tests/AutomationCommon.h
@@ -43,7 +43,7 @@ namespace AutomationCommon
 		}
 
 		FString FeatureLevelString;
-		GetFeatureLevelName(GRHIFeatureLevel,FeatureLevelString);
+		GetFeatureLevelName(GMaxRHIFeatureLevel,FeatureLevelString);
 		{
 			HardwareDetailsString = ( HardwareDetailsString + TEXT( "_" ) ) + FeatureLevelString;
 		}
diff --git a/Engine/Source/Runtime/Engine/Public/UnrealEngine.h b/Engine/Source/Runtime/Engine/Public/UnrealEngine.h
--- a/Engine/Source/Runtime/Engine/Public/UnrealEngine.h
+++ b/Engine/Source/Runtime/Engine/Public/UnrealEngine.h
@@ -406,7 +406,6 @@ struct FCachedSystemScalabilityCVars
 	FCachedSystemScalabilityCVars();
 };
 
-ENGINE_API bool AllowHighQualityLightmaps();
 ENGINE_API bool AllowHighQualityLightmaps(ERHIFeatureLevel::Type FeatureLevel);
 
 ENGINE_API const FCachedSystemScalabilityCVars& GetCachedScalabilityCVars();
diff --git a/Engine/Source/Runtime/GameLiveStreaming/GameLiveStreaming.cpp b/Engine/Source/Runtime/GameLiveStreaming/GameLiveStreaming.cpp
--- a/Engine/Source/Runtime/GameLiveStreaming/GameLiveStreaming.cpp
+++ b/Engine/Source/Runtime/GameLiveStreaming/GameLiveStreaming.cpp
@@ -315,7 +315,7 @@ void FGameLiveStreaming::StartCopyingNextGameVideoFrame( const FViewportRHIRef&
 	{
 		FPooledRenderTargetDesc OutputDesc(FPooledRenderTargetDesc::Create2DDesc(Context.ResizeTo, PF_B8G8R8A8, TexCreate_None, TexCreate_RenderTargetable, false));
 			
-		const auto FeatureLevel = GRHIFeatureLevel;
+		const auto FeatureLevel = GMaxRHIFeatureLevel;
 
 		TRefCountPtr<IPooledRenderTarget> ResampleTexturePooledRenderTarget;
 		Context.RendererModule->RenderTargetPoolFindFreeElement(OutputDesc, ResampleTexturePooledRenderTarget, TEXT("ResampleTexture"));
diff --git a/Engine/Source/Runtime/OpenGLDrv/Private/OpenGLRenderTarget.cpp b/Engine/Source/Runtime/OpenGLDrv/Private/OpenGLRenderTarget.cpp
--- a/Engine/Source/Runtime/OpenGLDrv/Private/OpenGLRenderTarget.cpp
+++ b/Engine/Source/Runtime/OpenGLDrv/Private/OpenGLRenderTarget.cpp
@@ -737,7 +737,7 @@ void FOpenGLDynamicRHI::RHIReadSurfaceFloatData(FTextureRHIParamRef TextureRHI,F
 	VERIFY_GL_SCOPE();	
 
 	//reading from arrays only supported on SM5 and up.
-	check(FOpenGL::SupportsFloatReadSurface() && (ArrayIndex == 0 || GRHIFeatureLevel >= ERHIFeatureLevel::SM5));	
+	check(FOpenGL::SupportsFloatReadSurface() && (ArrayIndex == 0 || GMaxRHIFeatureLevel >= ERHIFeatureLevel::SM5));	
 	FOpenGLTextureBase* Texture = GetOpenGLTextureFromRHITexture(TextureRHI);
 	check(TextureRHI->GetFormat() == PF_FloatRGBA);
 
diff --git a/Engine/Source/Runtime/RHI/Public/RHI.h b/Engine/Source/Runtime/RHI/Public/RHI.h
--- a/Engine/Source/Runtime/RHI/Public/RHI.h
+++ b/Engine/Source/Runtime/RHI/Public/RHI.h
@@ -247,17 +247,17 @@ extern RHI_API EShaderPlatform GetMaxRHIShaderPlatform();
 /** treating GRHIFeatureLevel as a function allows for better usage tracking at small cost to performance. */
 #define RHI_FEATURE_LEVEL_AS_FUNCTION 0
 
+#define GMaxRHIFeatureLevel GMaxRHIFeatureLevelValue
+#define GMaxRHIShaderPlatform GMaxRHIShaderPlatformValue
+
 #if RHI_FEATURE_LEVEL_AS_FUNCTION
-	#define GMaxRHIFeatureLevel GetMaxRHIFeatureLevel()
-	#define GMaxRHIShaderPlatform GetMaxRHIShaderPlatform()
+	#define GRHIShaderPlatform GetMaxRHIShaderPlatform()
+	#define GRHIFeatureLevel GetMaxRHIFeatureLevel()
 #else
-	#define GMaxRHIFeatureLevel GMaxRHIFeatureLevelValue
-	#define GMaxRHIShaderPlatform GMaxRHIShaderPlatformValue
+	#define GRHIShaderPlatform GMaxRHIShaderPlatform
+	#define GRHIFeatureLevel GMaxRHIFeatureLevel
 #endif
 
-#define GRHIShaderPlatform GMaxRHIShaderPlatform
-#define GRHIFeatureLevel GMaxRHIFeatureLevel
-
 /** Table for finding out which shader platform corresponds to a given feature level for this RHI. */
 extern RHI_API EShaderPlatform GShaderPlatformForFeatureLevel[ERHIFeatureLevel::Num];
 
diff --git a/Engine/Source/Runtime/RHI/Public/RHIUtilities.h b/Engine/Source/Runtime/RHI/Public/RHIUtilities.h
--- a/Engine/Source/Runtime/RHI/Public/RHIUtilities.h
+++ b/Engine/Source/Runtime/RHI/Public/RHIUtilities.h
@@ -196,7 +196,7 @@ inline void RHICreateTargetableShaderResource2D(
 
 	if (NumSamples > 1)
 	{
-		bForceSeparateTargetAndShaderResource = RHISupportsSeparateMSAAAndResolveTextures(GRHIShaderPlatform);
+		bForceSeparateTargetAndShaderResource = RHISupportsSeparateMSAAAndResolveTextures(GMaxRHIShaderPlatform);
 	}
 
 	if (!bForceSeparateTargetAndShaderResource/* && GSupportsRenderDepthTargetableShaderResources*/)
diff --git a/Engine/Source/Runtime/Renderer/Private/DepthRendering.cpp b/Engine/Source/Runtime/Renderer/Private/DepthRendering.cpp
--- a/Engine/Source/Runtime/Renderer/Private/DepthRendering.cpp
+++ b/Engine/Source/Runtime/Renderer/Private/DepthRendering.cpp
@@ -132,7 +132,8 @@ FDepthDrawingPolicy::FDepthDrawingPolicy(
 	const FVertexFactory* InVertexFactory,
 	const FMaterialRenderProxy* InMaterialRenderProxy,
 	const FMaterial& InMaterialResource,
-	bool bIsTwoSided
+	bool bIsTwoSided,
+	ERHIFeatureLevel::Type InFeatureLevel
 	):
 	FMeshDrawingPolicy(InVertexFactory,InMaterialRenderProxy,InMaterialResource,false,/*bInTwoSidedOverride=*/ bIsTwoSided)
 {
@@ -154,7 +155,7 @@ FDepthDrawingPolicy::FDepthDrawingPolicy(
 	HullShader = NULL;	
 	DomainShader = NULL;
 
-	if(RHISupportsTessellation(GRHIShaderPlatform)
+	if (RHISupportsTessellation(GShaderPlatformForFeatureLevel[InFeatureLevel])
 		&& InVertexFactory->GetType()->SupportsTessellationShaders() 
 		&& TessellationMode != MTM_NoTessellation)
 	{
@@ -307,6 +308,7 @@ void FDepthDrawingPolicyFactory::AddStaticMesh(FScene* Scene,FStaticMesh* Static
 	const FMaterialRenderProxy* MaterialRenderProxy = StaticMesh->MaterialRenderProxy;
 	const FMaterial* Material = MaterialRenderProxy->GetMaterial(Scene->GetFeatureLevel());
 	const EBlendMode BlendMode = Material->GetBlendMode();
+	const auto FeatureLevel = Scene->GetFeatureLevel();
 
 	if (Material->IsMasked())
 	{
@@ -318,9 +320,10 @@ void FDepthDrawingPolicyFactory::AddStaticMesh(FScene* Scene,FStaticMesh* Static
 				StaticMesh->VertexFactory,
 				MaterialRenderProxy,
 				*Material,
-				Material->IsTwoSided()
+				Material->IsTwoSided(),
+				FeatureLevel
 				),
-			Scene->GetFeatureLevel()
+			FeatureLevel
 			);
 	}
 	else
@@ -340,7 +343,7 @@ void FDepthDrawingPolicyFactory::AddStaticMesh(FScene* Scene,FStaticMesh* Static
 					Material->IsTwoSided(),
 					Material->IsWireframe()
 					),
-				Scene->GetFeatureLevel()
+				FeatureLevel
 				);
 		}
 		else
@@ -359,9 +362,10 @@ void FDepthDrawingPolicyFactory::AddStaticMesh(FScene* Scene,FStaticMesh* Static
 					StaticMesh->VertexFactory,
 					MaterialRenderProxy,
 					*MaterialRenderProxy->GetMaterial(Scene->GetFeatureLevel()),
-					Material->IsTwoSided()
+					Material->IsTwoSided(),
+					FeatureLevel
 					),
-				Scene->GetFeatureLevel()
+				FeatureLevel
 				);
 		}
 	}
@@ -443,7 +447,7 @@ bool FDepthDrawingPolicyFactory::DrawMesh(
 					MaterialRenderProxy = UMaterial::GetDefaultMaterial(MD_Surface)->GetRenderProxy(false);
 				}
 
-				FDepthDrawingPolicy DrawingPolicy(Mesh.VertexFactory, MaterialRenderProxy, *MaterialRenderProxy->GetMaterial(View.GetFeatureLevel()), Material->IsTwoSided());
+				FDepthDrawingPolicy DrawingPolicy(Mesh.VertexFactory, MaterialRenderProxy, *MaterialRenderProxy->GetMaterial(View.GetFeatureLevel()), Material->IsTwoSided(), View.GetFeatureLevel());
 				RHICmdList.BuildAndSetLocalBoundShaderState(DrawingPolicy.GetBoundShaderStateInput(View.GetFeatureLevel()));
 				DrawingPolicy.SetSharedState(RHICmdList, &View, FDepthDrawingPolicy::ContextDataType());
 
diff --git a/Engine/Source/Runtime/Renderer/Private/DepthRendering.h b/Engine/Source/Runtime/Renderer/Private/DepthRendering.h
--- a/Engine/Source/Runtime/Renderer/Private/DepthRendering.h
+++ b/Engine/Source/Runtime/Renderer/Private/DepthRendering.h
@@ -32,7 +32,8 @@ public:
 		const FVertexFactory* InVertexFactory,
 		const FMaterialRenderProxy* InMaterialRenderProxy,
 		const FMaterial& InMaterialResource,
-		bool bIsTwoSided
+		bool bIsTwoSided,
+		ERHIFeatureLevel::Type InFeatureLevel
 		);
 
 	// FMeshDrawingPolicy interface.
diff --git a/Engine/Source/Runtime/Renderer/Private/DistortionRendering.cpp b/Engine/Source/Runtime/Renderer/Private/DistortionRendering.cpp
--- a/Engine/Source/Runtime/Renderer/Private/DistortionRendering.cpp
+++ b/Engine/Source/Runtime/Renderer/Private/DistortionRendering.cpp
@@ -279,7 +279,8 @@ public:
 		const FMaterialRenderProxy* InMaterialRenderProxy,
 		const FMaterial& MaterialResouce,
 		bool bInitializeOffsets,
-		bool bInOverrideWithShaderComplexity
+		bool bInOverrideWithShaderComplexity,
+		ERHIFeatureLevel::Type InFeatureLevel
 		);
 
 	// FMeshDrawingPolicy interface.
@@ -349,7 +350,8 @@ TDistortionMeshDrawingPolicy<DistortMeshPolicy>::TDistortionMeshDrawingPolicy(
 	const FMaterialRenderProxy* InMaterialRenderProxy,
 	const FMaterial& InMaterialResource,
 	bool bInInitializeOffsets,
-	bool bInOverrideWithShaderComplexity
+	bool bInOverrideWithShaderComplexity,
+	ERHIFeatureLevel::Type InFeatureLevel
 	)
 :	FMeshDrawingPolicy(InVertexFactory,InMaterialRenderProxy,InMaterialResource,bInOverrideWithShaderComplexity)
 ,	bInitializeOffsets(bInInitializeOffsets)
@@ -358,7 +360,7 @@ TDistortionMeshDrawingPolicy<DistortMeshPolicy>::TDistortionMeshDrawingPolicy(
 	DomainShader = NULL;
 
 	const EMaterialTessellationMode MaterialTessellationMode = MaterialResource->GetTessellationMode();
-	if(RHISupportsTessellation(GRHIShaderPlatform)
+	if (RHISupportsTessellation(GShaderPlatformForFeatureLevel[InFeatureLevel])
 		&& InVertexFactory->GetType()->SupportsTessellationShaders() 
 		&& MaterialTessellationMode != MTM_NoTessellation)
 	{
@@ -601,7 +603,8 @@ bool TDistortionMeshDrawingPolicyFactory<DistortMeshPolicy>::DrawDynamicMesh(
 			Mesh.MaterialRenderProxy,
 			*Mesh.MaterialRenderProxy->GetMaterial(FeatureLevel),
 			bInitializeOffsets,
-			View.Family->EngineShowFlags.ShaderComplexity
+			View.Family->EngineShowFlags.ShaderComplexity,
+			FeatureLevel
 			);
 		RHICmdList.BuildAndSetLocalBoundShaderState(DrawingPolicy.GetBoundShaderStateInput(View.GetFeatureLevel()));
 		DrawingPolicy.SetSharedState(RHICmdList, &View, typename TDistortionMeshDrawingPolicy<DistortMeshPolicy>::ContextDataType());
@@ -646,7 +649,8 @@ bool TDistortionMeshDrawingPolicyFactory<DistortMeshPolicy>::DrawStaticMesh(
 			StaticMesh.MaterialRenderProxy,
 			*StaticMesh.MaterialRenderProxy->GetMaterial(FeatureLevel),
 			bInitializeOffsets,
-			View->Family->EngineShowFlags.ShaderComplexity
+			View->Family->EngineShowFlags.ShaderComplexity,
+			FeatureLevel
 			);
 		RHICmdList.BuildAndSetLocalBoundShaderState(DrawingPolicy.GetBoundShaderStateInput(View->GetFeatureLevel()));
 		DrawingPolicy.SetSharedState(RHICmdList, View, typename TDistortionMeshDrawingPolicy<DistortMeshPolicy>::ContextDataType());
diff --git a/Engine/Source/Runtime/Renderer/Private/LightMapDensityRendering.cpp b/Engine/Source/Runtime/Renderer/Private/LightMapDensityRendering.cpp
--- a/Engine/Source/Runtime/Renderer/Private/LightMapDensityRendering.cpp
+++ b/Engine/Source/Runtime/Renderer/Private/LightMapDensityRendering.cpp
@@ -121,10 +121,10 @@ bool FLightMapDensityDrawingPolicyFactory::DrawDynamicMesh(
 {
 	bool bDirty = false;
 
+	const auto FeatureLevel = View.GetFeatureLevel();
 	const FMaterialRenderProxy* MaterialRenderProxy = Mesh.MaterialRenderProxy;
-	const FMaterial* Material = MaterialRenderProxy->GetMaterial(View.GetFeatureLevel());
+	const FMaterial* Material = MaterialRenderProxy->GetMaterial(FeatureLevel);
 	const EBlendMode BlendMode = Material->GetBlendMode();
-	const auto FeatureLevel = View.GetFeatureLevel();
 
 	const bool bMaterialMasked = Material->IsMasked();
 	const bool bMaterialModifiesMesh = Material->MaterialModifiesMeshPosition_RenderThread();
@@ -155,7 +155,7 @@ bool FLightMapDensityDrawingPolicyFactory::DrawDynamicMesh(
 			if (bAllowHighQualityLightMaps)
 			{
 				TLightMapDensityDrawingPolicy< TLightMapPolicy<HQ_LIGHTMAP> > DrawingPolicy(View, Mesh.VertexFactory, MaterialRenderProxy, TLightMapPolicy<HQ_LIGHTMAP>(), BlendMode);
-				RHICmdList.BuildAndSetLocalBoundShaderState(DrawingPolicy.GetBoundShaderStateInput(View.GetFeatureLevel()));
+				RHICmdList.BuildAndSetLocalBoundShaderState(DrawingPolicy.GetBoundShaderStateInput(FeatureLevel));
 				DrawingPolicy.SetSharedState(RHICmdList, &View, TLightMapDensityDrawingPolicy< TLightMapPolicy<HQ_LIGHTMAP> >::ContextDataType());
 				for (int32 BatchElementIndex = 0; BatchElementIndex < Mesh.Elements.Num(); BatchElementIndex++)
 				{
@@ -170,7 +170,7 @@ bool FLightMapDensityDrawingPolicyFactory::DrawDynamicMesh(
 			else
 			{
 				TLightMapDensityDrawingPolicy< TLightMapPolicy<LQ_LIGHTMAP> > DrawingPolicy(View, Mesh.VertexFactory, MaterialRenderProxy, TLightMapPolicy<LQ_LIGHTMAP>(), BlendMode);
-				RHICmdList.BuildAndSetLocalBoundShaderState(DrawingPolicy.GetBoundShaderStateInput(View.GetFeatureLevel()));
+				RHICmdList.BuildAndSetLocalBoundShaderState(DrawingPolicy.GetBoundShaderStateInput(FeatureLevel));
 				DrawingPolicy.SetSharedState(RHICmdList, &View, TLightMapDensityDrawingPolicy< TLightMapPolicy<LQ_LIGHTMAP> >::ContextDataType());
 				for (int32 BatchElementIndex = 0; BatchElementIndex < Mesh.Elements.Num(); BatchElementIndex++)
 				{
@@ -186,7 +186,7 @@ bool FLightMapDensityDrawingPolicyFactory::DrawDynamicMesh(
 		else
 		{
 			TLightMapDensityDrawingPolicy<FDummyLightMapPolicy> DrawingPolicy(View, Mesh.VertexFactory, MaterialRenderProxy, FDummyLightMapPolicy(), BlendMode);
-			RHICmdList.BuildAndSetLocalBoundShaderState(DrawingPolicy.GetBoundShaderStateInput(View.GetFeatureLevel()));
+			RHICmdList.BuildAndSetLocalBoundShaderState(DrawingPolicy.GetBoundShaderStateInput(FeatureLevel));
 			DrawingPolicy.SetSharedState(RHICmdList, &View, TLightMapDensityDrawingPolicy<FDummyLightMapPolicy>::ContextDataType());
 			for (int32 BatchElementIndex = 0; BatchElementIndex < Mesh.Elements.Num(); BatchElementIndex++)
 			{
@@ -202,7 +202,7 @@ bool FLightMapDensityDrawingPolicyFactory::DrawDynamicMesh(
 	else
 	{
 		TLightMapDensityDrawingPolicy<FNoLightMapPolicy> DrawingPolicy(View, Mesh.VertexFactory, MaterialRenderProxy, FNoLightMapPolicy(), BlendMode);
-		RHICmdList.BuildAndSetLocalBoundShaderState(DrawingPolicy.GetBoundShaderStateInput(View.GetFeatureLevel()));
+		RHICmdList.BuildAndSetLocalBoundShaderState(DrawingPolicy.GetBoundShaderStateInput(FeatureLevel));
 		DrawingPolicy.SetSharedState(RHICmdList, &View, TLightMapDensityDrawingPolicy<FNoLightMapPolicy>::ContextDataType());
 		for (int32 BatchElementIndex = 0; BatchElementIndex < Mesh.Elements.Num(); BatchElementIndex++)
 		{
diff --git a/Engine/Source/Runtime/Renderer/Private/LightRendering.h b/Engine/Source/Runtime/Renderer/Private/LightRendering.h
--- a/Engine/Source/Runtime/Renderer/Private/LightRendering.h
+++ b/Engine/Source/Runtime/Renderer/Private/LightRendering.h
@@ -55,7 +55,7 @@ void SetDeferredLightParameters(
 	DeferredLightUniformsValue.LightColor = LightColorAndFalloffExponent;
 	DeferredLightUniformsValue.LightFalloffExponent = LightColorAndFalloffExponent.W;
 
-	const FVector2D FadeParams = LightSceneInfo->Proxy->GetDirectionalLightDistanceFadeParameters();
+	const FVector2D FadeParams = LightSceneInfo->Proxy->GetDirectionalLightDistanceFadeParameters(View.GetFeatureLevel());
 
 	// use MAD for efficiency in the shader
 	DeferredLightUniformsValue.DistanceFadeMAD = FVector2D(FadeParams.Y, -FadeParams.X * FadeParams.Y);
diff --git a/Engine/Source/Runtime/Renderer/Private/ReflectionEnvironment.cpp b/Engine/Source/Runtime/Renderer/Private/ReflectionEnvironment.cpp
--- a/Engine/Source/Runtime/Renderer/Private/ReflectionEnvironment.cpp
+++ b/Engine/Source/Runtime/Renderer/Private/ReflectionEnvironment.cpp
@@ -597,7 +597,7 @@ public:
 		const FPixelShaderRHIParamRef ShaderRHI = GetPixelShader();
 		FGlobalShader::SetParameters(RHICmdList, ShaderRHI, View);
 
-		if (GRHIFeatureLevel >= ERHIFeatureLevel::SM5)
+		if (View.GetFeatureLevel() >= ERHIFeatureLevel::SM5)
 		{
 			FScene* Scene = (FScene*)View.Family->Scene;
 
@@ -800,7 +800,7 @@ void FDeferredShadingSceneRenderer::RenderStandardDeferredImageBasedReflections(
 
 			NewSortEntry.CaptureIndex = -1;
 
-			if (GRHIFeatureLevel >= ERHIFeatureLevel::SM5)
+			if (FeatureLevel >= ERHIFeatureLevel::SM5)
 			{
 				const FCaptureComponentSceneState* ComponentStatePtr = Scene->ReflectionSceneData.AllocatedReflectionCaptureState.Find(CurrentCapture->Component);
 				NewSortEntry.CaptureIndex = ComponentStatePtr ? ComponentStatePtr->CaptureIndex : -1;
@@ -881,7 +881,7 @@ void FDeferredShadingSceneRenderer::RenderStandardDeferredImageBasedReflections(
 			{
 				const FReflectionCaptureSortData& ReflectionCapture = SortData[ReflectionCaptureIndex];
 
-				if (GRHIFeatureLevel >= ERHIFeatureLevel::SM5 || ReflectionCapture.SM4FullHDRCubemap)
+				if (FeatureLevel >= ERHIFeatureLevel::SM5 || ReflectionCapture.SM4FullHDRCubemap)
 				{
 					const FSphere LightBounds(ReflectionCapture.PositionAndRadius, ReflectionCapture.PositionAndRadius.W);
 
diff --git a/Engine/Source/Runtime/Renderer/Private/SceneRendering.cpp b/Engine/Source/Runtime/Renderer/Private/SceneRendering.cpp
--- a/Engine/Source/Runtime/Renderer/Private/SceneRendering.cpp
+++ b/Engine/Source/Runtime/Renderer/Private/SceneRendering.cpp
@@ -684,7 +684,7 @@ FSceneRenderer::FSceneRenderer(const FSceneViewFamily* InViewFamily,FHitProxyCon
 
 bool FSceneRenderer::DoOcclusionQueries(ERHIFeatureLevel::Type InFeatureLevel) const
 {
-	return !IsMobilePlatform(GRHIShaderPlatform)
+	return !IsMobilePlatform(GShaderPlatformForFeatureLevel[InFeatureLevel])
 		&& CVarAllowOcclusionQueries.GetValueOnRenderThread() != 0;
 }
 
@@ -892,7 +892,7 @@ FSceneRenderer* FSceneRenderer::CreateSceneRenderer(const FSceneViewFamily* InVi
 
 void FSceneRenderer::RenderCustomDepthPass(FRHICommandListImmediate& RHICmdList)
 {
-	if(!IsFeatureLevelSupported(GRHIShaderPlatform, ERHIFeatureLevel::SM4))
+	if(FeatureLevel < ERHIFeatureLevel::SM4)
 	{
 		// not yet supported on lower end platforms
 		return;
diff --git a/Engine/Source/Runtime/Renderer/Private/ShadowRendering.cpp b/Engine/Source/Runtime/Renderer/Private/ShadowRendering.cpp
--- a/Engine/Source/Runtime/Renderer/Private/ShadowRendering.cpp
+++ b/Engine/Source/Runtime/Renderer/Private/ShadowRendering.cpp
@@ -703,7 +703,7 @@ void OverrideWithDefaultMaterialForShadowDepth(
 -----------------------------------------------------------------------------*/
 
 template <bool bRenderingReflectiveShadowMaps>
-void FShadowDepthDrawingPolicy<bRenderingReflectiveShadowMaps>::UpdateElementState(FShadowStaticMeshElement& State)
+void FShadowDepthDrawingPolicy<bRenderingReflectiveShadowMaps>::UpdateElementState(FShadowStaticMeshElement& State, ERHIFeatureLevel::Type FeatureLevel)
 {
 	// can be optimized
 	*this = FShadowDepthDrawingPolicy(
@@ -711,6 +711,7 @@ void FShadowDepthDrawingPolicy<bRenderingReflectiveShadowMaps>::UpdateElementSta
 		bDirectionalLight,
 		bOnePassPointLightShadow,
 		bPreShadow,
+		FeatureLevel,
 		State.Mesh->VertexFactory,
 		State.RenderProxy,
 		State.bIsTwoSided,
@@ -723,6 +724,7 @@ FShadowDepthDrawingPolicy<bRenderingReflectiveShadowMaps>::FShadowDepthDrawingPo
 	bool bInDirectionalLight,
 	bool bInOnePassPointLightShadow,
 	bool bInPreShadow,
+	ERHIFeatureLevel::Type InFeatureLevel,
 	const FVertexFactory* InVertexFactory,
 	const FMaterialRenderProxy* InMaterialRenderProxy,
 	bool bInCastShadowAsTwoSided,
@@ -756,7 +758,7 @@ FShadowDepthDrawingPolicy<bRenderingReflectiveShadowMaps>::FShadowDepthDrawingPo
 
 	const bool bInitializeTessellationShaders = 
 		MaterialResource->GetTessellationMode() != MTM_NoTessellation
-		&& RHISupportsTessellation(GRHIShaderPlatform)
+		&& RHISupportsTessellation(GShaderPlatformForFeatureLevel[InFeatureLevel])
 		&& VFType->SupportsTessellationShaders();
 
 	bUsePositionOnlyVS = !bRenderingReflectiveShadowMaps
@@ -992,8 +994,9 @@ void FShadowDepthDrawingPolicyFactory::AddStaticMesh(FScene* Scene,FStaticMesh*
 {
 	if (StaticMesh->CastShadow)
 	{
+		const auto FeatureLevel = Scene->GetFeatureLevel();
 		const FMaterialRenderProxy* MaterialRenderProxy = StaticMesh->MaterialRenderProxy;
-		const FMaterial* Material = MaterialRenderProxy->GetMaterial(Scene->GetFeatureLevel());
+		const FMaterial* Material = MaterialRenderProxy->GetMaterial(FeatureLevel);
 		const EBlendMode BlendMode = Material->GetBlendMode();
 		const EMaterialShadingModel ShadingModel = Material->GetShadingModel();
 
@@ -1013,6 +1016,7 @@ void FShadowDepthDrawingPolicyFactory::AddStaticMesh(FScene* Scene,FStaticMesh*
 						true,
 						false,
 						false,
+						Scene->GetFeatureLevel(),
 						StaticMesh->VertexFactory,
 						MaterialRenderProxy,
 						bTwoSided,
@@ -1034,6 +1038,7 @@ void FShadowDepthDrawingPolicyFactory::AddStaticMesh(FScene* Scene,FStaticMesh*
 					true,
 					false,
 					false,
+					Scene->GetFeatureLevel(),
 					StaticMesh->VertexFactory,
 					MaterialRenderProxy,
 					bTwoSided,
@@ -1060,8 +1065,9 @@ bool FShadowDepthDrawingPolicyFactory::DrawDynamicMesh(
 	// Use a per-FMeshBatch check on top of the per-primitive check because dynamic primitives can submit multiple FMeshElements.
 	if (Mesh.CastShadow)
 	{
+		const auto FeatureLevel = View.GetFeatureLevel();
 		const FMaterialRenderProxy* MaterialRenderProxy = Mesh.MaterialRenderProxy;
-		const FMaterial* Material = MaterialRenderProxy->GetMaterial(View.GetFeatureLevel());
+		const FMaterial* Material = MaterialRenderProxy->GetMaterial(FeatureLevel);
 		const EBlendMode BlendMode = Material->GetBlendMode();
 		const EMaterialShadingModel ShadingModel = Material->GetShadingModel();
 
@@ -1075,22 +1081,23 @@ bool FShadowDepthDrawingPolicyFactory::DrawDynamicMesh(
 			const bool bTwoSided = Material->IsTwoSided() || PrimitiveSceneProxy->CastsShadowAsTwoSided();
 			const FShadowDepthDrawingPolicyContext PolicyContext(Context.ShadowInfo);
 
-			OverrideWithDefaultMaterialForShadowDepth(MaterialRenderProxy, Material, bReflectiveShadowmap, View.GetFeatureLevel());
+			OverrideWithDefaultMaterialForShadowDepth(MaterialRenderProxy, Material, bReflectiveShadowmap, FeatureLevel);
 
 			if(bReflectiveShadowmap)
 			{
 				FShadowDepthDrawingPolicy<true> DrawingPolicy(
-					MaterialRenderProxy->GetMaterial(View.GetFeatureLevel()),
+					MaterialRenderProxy->GetMaterial(FeatureLevel),
 					bDirectionalLight,
 					bOnePassPointLightShadow,
 					bPreShadow,
+					FeatureLevel,
 					Mesh.VertexFactory,
 					MaterialRenderProxy,
 					bTwoSided,
 					Mesh.ReverseCulling
 					);
 
-				RHICmdList.BuildAndSetLocalBoundShaderState(DrawingPolicy.GetBoundShaderStateInput(View.GetFeatureLevel()));
+				RHICmdList.BuildAndSetLocalBoundShaderState(DrawingPolicy.GetBoundShaderStateInput(FeatureLevel));
 				DrawingPolicy.SetSharedState(RHICmdList, &View, PolicyContext);
 				for (int32 BatchElementIndex = 0, Num = Mesh.Elements.Num(); BatchElementIndex < Num; BatchElementIndex++)
 				{
@@ -1101,17 +1108,18 @@ bool FShadowDepthDrawingPolicyFactory::DrawDynamicMesh(
 			else
 			{
 				FShadowDepthDrawingPolicy<false> DrawingPolicy(
-					MaterialRenderProxy->GetMaterial(View.GetFeatureLevel()),
+					MaterialRenderProxy->GetMaterial(FeatureLevel),
 					bDirectionalLight,
 					bOnePassPointLightShadow,
 					bPreShadow,
+					FeatureLevel,
 					Mesh.VertexFactory,
 					MaterialRenderProxy,
 					bTwoSided,
 					Mesh.ReverseCulling
 					);
 
-				RHICmdList.BuildAndSetLocalBoundShaderState(DrawingPolicy.GetBoundShaderStateInput(View.GetFeatureLevel()));
+				RHICmdList.BuildAndSetLocalBoundShaderState(DrawingPolicy.GetBoundShaderStateInput(FeatureLevel));
 				DrawingPolicy.SetSharedState(RHICmdList, &View, PolicyContext);
 				for (int32 BatchElementIndex = 0; BatchElementIndex < Mesh.Elements.Num(); BatchElementIndex++)
 				{
@@ -1211,6 +1219,7 @@ void DrawMeshElements(FRHICommandList& RHICmdList, FShadowDepthDrawingPolicy<bRe
 		SharedDrawingPolicy.bDirectionalLight,
 		SharedDrawingPolicy.bOnePassPointLightShadow,
 		SharedDrawingPolicy.bPreShadow,
+		View.GetFeatureLevel(),
 		State.Mesh->VertexFactory,
 		State.RenderProxy,
 		State.bIsTwoSided,
@@ -1243,12 +1252,14 @@ void DrawShadowMeshElements(FRHICommandList& RHICmdList, const FViewInfo& View,
 	FShadowDepthDrawingPolicyContext PolicyContext(&ShadowInfo);
 	const FShadowStaticMeshElement& FirstShadowMesh = ShadowInfo.SubjectMeshElements[0];
 	const FMaterial* FirstMaterialResource = FirstShadowMesh.MaterialResource;
+	auto FeatureLevel = View.GetFeatureLevel();
 
 	FShadowDepthDrawingPolicy<bReflectiveShadowmap> SharedDrawingPolicy(
 		FirstMaterialResource,
 		ShadowInfo.bDirectionalLight,
 		ShadowInfo.bOnePassPointLightShadow,
-		ShadowInfo.bPreShadow);
+		ShadowInfo.bPreShadow,
+		FeatureLevel);
 
 	FShadowStaticMeshElement OldState;
 
@@ -1270,7 +1281,7 @@ void DrawShadowMeshElements(FRHICommandList& RHICmdList, const FViewInfo& View,
 		{
 			OldState = CurrentState;
 
-			SharedDrawingPolicy.UpdateElementState(CurrentState);
+			SharedDrawingPolicy.UpdateElementState(CurrentState, FeatureLevel);
 			RHICmdList.BuildAndSetLocalBoundShaderState(SharedDrawingPolicy.GetBoundShaderStateInput(View.GetFeatureLevel()));
 			SharedDrawingPolicy.SetSharedState(RHICmdList, &View, PolicyContext);
 		}
@@ -2873,8 +2884,8 @@ bool FDeferredShadingSceneRenderer::RenderTranslucentProjectedShadows(FRHIComman
 bool FDeferredShadingSceneRenderer::RenderReflectiveShadowMaps(FRHICommandListImmediate& RHICmdList, const FLightSceneInfo* LightSceneInfo)
 {
 	SCOPE_CYCLE_COUNTER(STAT_ReflectiveShadowMapDrawTime);
-
-	if ( !IsFeatureLevelSupported(GRHIShaderPlatform, ERHIFeatureLevel::SM5) )
+	
+	if (FeatureLevel < ERHIFeatureLevel::SM5)
 	{
 		return false;
 	}
diff --git a/Engine/Source/Runtime/Renderer/Private/ShadowRendering.h b/Engine/Source/Runtime/Renderer/Private/ShadowRendering.h
--- a/Engine/Source/Runtime/Renderer/Private/ShadowRendering.h
+++ b/Engine/Source/Runtime/Renderer/Private/ShadowRendering.h
@@ -310,13 +310,14 @@ public:
 		bool bInDirectionalLight,
 		bool bInOnePassPointLightShadow,
 		bool bInPreShadow,
+		ERHIFeatureLevel::Type InFeatureLevel,
 		const FVertexFactory* InVertexFactory = 0,
 		const FMaterialRenderProxy* InMaterialRenderProxy = 0,
 		bool bInCastShadowAsTwoSided = false,
 		bool bReverseCulling = false
 		);
 
-	void UpdateElementState(FShadowStaticMeshElement& State);
+	void UpdateElementState(FShadowStaticMeshElement& State, ERHIFeatureLevel::Type FeatureLevel);
 
 	FShadowDepthDrawingPolicy& operator = (const FShadowDepthDrawingPolicy& Other)
 	{ 
@@ -330,6 +331,7 @@ public:
 		bOnePassPointLightShadow = Other.bOnePassPointLightShadow;
 		bUsePositionOnlyVS = Other.bUsePositionOnlyVS;
 		bPreShadow = Other.bPreShadow;
+		FeatureLevel = Other.FeatureLevel;
 		(FMeshDrawingPolicy&)*this = (const FMeshDrawingPolicy&)Other;
 		return *this; 
 	}
@@ -347,7 +349,8 @@ public:
 			&& bReverseCulling == Other.bReverseCulling
 			&& bOnePassPointLightShadow == Other.bOnePassPointLightShadow
 			&& bUsePositionOnlyVS == Other.bUsePositionOnlyVS
-			&& bPreShadow == Other.bPreShadow;
+			&& bPreShadow == Other.bPreShadow
+			&& FeatureLevel == Other.FeatureLevel;
 	}
 	void SetSharedState(FRHICommandList& RHICmdList, const FSceneView* View, const ContextDataType PolicyContext) const;
 
@@ -384,6 +387,7 @@ private:
 	TShadowDepthBasePS<bRenderingReflectiveShadowMaps>* PixelShader;
 	class FBaseHS* HullShader;
 	class FShadowDepthDS* DomainShader;
+	ERHIFeatureLevel::Type FeatureLevel;
 
 public:
 
diff --git a/Engine/Source/Runtime/Renderer/Private/ShadowSetup.cpp b/Engine/Source/Runtime/Renderer/Private/ShadowSetup.cpp
--- a/Engine/Source/Runtime/Renderer/Private/ShadowSetup.cpp
+++ b/Engine/Source/Runtime/Renderer/Private/ShadowSetup.cpp
@@ -1368,7 +1368,7 @@ void FDeferredShadingSceneRenderer::CreatePerObjectProjectedShadow(
 			for (int32 i = 0; i < ViewDependentWholeSceneShadows.Num(); i++)
 			{
 				const FProjectedShadowInfo* WholeSceneShadow = ViewDependentWholeSceneShadows[i];
-				const FVector2D DistanceFadeValues = WholeSceneShadow->LightSceneInfo->Proxy->GetDirectionalLightDistanceFadeParameters();
+				const FVector2D DistanceFadeValues = WholeSceneShadow->LightSceneInfo->Proxy->GetDirectionalLightDistanceFadeParameters(Scene->GetFeatureLevel());
 				const float DistanceFromShadowCenterSquared = (WholeSceneShadow->ShadowBounds.Center - Bounds.Origin).SizeSquared();
 				//@todo - if view dependent whole scene shadows are ever supported in splitscreen, 
 				// We can only disable the preshadow at this point if it is inside a whole scene shadow for all views
diff --git a/Engine/Source/Runtime/Renderer/Private/TranslucentLighting.cpp b/Engine/Source/Runtime/Renderer/Private/TranslucentLighting.cpp
--- a/Engine/Source/Runtime/Renderer/Private/TranslucentLighting.cpp
+++ b/Engine/Source/Runtime/Renderer/Private/TranslucentLighting.cpp
@@ -896,7 +896,7 @@ public:
 				Planes[1] = FVector4((FVector)(ShadowCascadeSettings.FarFrustumPlane), -ShadowCascadeSettings.FarFrustumPlane.W);
 			}
 
-			const FVector2D FadeParams = LightSceneInfo->Proxy->GetDirectionalLightDistanceFadeParameters();
+			const FVector2D FadeParams = LightSceneInfo->Proxy->GetDirectionalLightDistanceFadeParameters(View.GetFeatureLevel());
 
 			// setup constants for the MAD in shader
 			ShadowInjectParamValue.Z = FadeParams.Y;
