removed not needed texture bindings (was needed when we had DX9)

diff --git a/Engine/Shaders/PostProcessCommon.usf b/Engine/Shaders/PostProcessCommon.usf
--- a/Engine/Shaders/PostProcessCommon.usf
+++ b/Engine/Shaders/PostProcessCommon.usf
@@ -7,8 +7,14 @@
 #ifndef __POST_PROCESS_COMMON__
 #define __POST_PROCESS_COMMON__
 
-// old method, will be removed
-// e.g.  tex2D(PostprocessInput0, UV.xy)
+/// e.g.
+//   PostprocessInput0.SampleLevel(PostprocessInput0Sampler, float2 UV, 0)
+//   Texture2DSample(PostprocessInput0, PostprocessInput0Sampler, float2 UV)
+//   PostprocessInput0.Load(Texel.xyz)
+//   PostprocessInput0.GetDimensions(uint, out uint width, out uint height, out uint levels)
+//   PostprocessInput0.GetDimensions(uint, out float width, out float height, out float levels)
+//   PostprocessInput0.GetDimensions(out uint width, out uint height)
+//   PostprocessInput0.GetDimensions(out float width, out float height)
 Texture2D PostprocessInput0;
 SamplerState PostprocessInput0Sampler;
 Texture2D PostprocessInput1;
@@ -24,19 +30,6 @@ SamplerState PostprocessInput5Sampler;
 Texture2D PostprocessInput6;
 SamplerState PostprocessInput6Sampler;
 
-// e.g. InputNew0.Load(Texel.xyz)
-// e.g. InputNew0.GetDimensions(uint, out uint width, out uint height, out uint levels)
-// e.g. InputNew0.GetDimensions(uint, out float width, out float height, out float levels)
-// e.g. InputNew0.GetDimensions(out uint width, out uint height)
-// e.g. InputNew0.GetDimensions(out float width, out float height)
-Texture2D InputNew0;
-Texture2D InputNew1;
-Texture2D InputNew2;
-Texture2D InputNew3;
-Texture2D InputNew4;
-Texture2D InputNew5;
-Texture2D InputNew6;
-
 // width, height, 1/width, 1/height
 float4 PostprocessInput0Size;
 float4 PostprocessInput1Size;
diff --git a/Engine/Shaders/PostProcessDilation.usf b/Engine/Shaders/PostProcessDilation.usf
--- a/Engine/Shaders/PostProcessDilation.usf
+++ b/Engine/Shaders/PostProcessDilation.usf
@@ -16,7 +16,7 @@ void MainPS(float4 UVAndScreenPos : TEXCOORD0, out float4 OutColor : SV_Target0)
 	float2 ScreenSpacePos = UVAndScreenPos.zw;
 	int3 PixelPos = int3((int2)(ScreenSpacePos * ScreenPosToPixel.xy + ScreenPosToPixel.zw + 0.5f), 0);
 
-	OutColor = InputNew0.Load(PixelPos);
+	OutColor = PostprocessInput0.Load(PixelPos);
 
 	int2 MaxSize = (int2)ViewportSize;
 
@@ -26,7 +26,7 @@ void MainPS(float4 UVAndScreenPos : TEXCOORD0, out float4 OutColor : SV_Target0)
 		FLATTEN\
 		if(all(OutColor.rgb == 0) && all(p.xy < MaxSize))\
 		{\
-			OutColor = InputNew0.Load(p);\
+			OutColor = PostprocessInput0.Load(p);\
 		}\
 	}
 
diff --git a/Engine/Shaders/PostProcessEyeAdaptation.usf b/Engine/Shaders/PostProcessEyeAdaptation.usf
--- a/Engine/Shaders/PostProcessEyeAdaptation.usf
+++ b/Engine/Shaders/PostProcessEyeAdaptation.usf
@@ -30,8 +30,8 @@ void MainPS(float4 UVAndScreenPos : TEXCOORD0, out float OutColor : SV_Target0)
 
 	float ExposureOffsetMultipler = EyeAdaptationParams[1].x;
 	
-	float TargetExposure = ComputeEyeAdaptationExposure(InputNew0);
-	float OldExposureScale = InputNew0.Load(int3(0, 1, 0)).x;
+	float TargetExposure = ComputeEyeAdaptationExposure(PostprocessInput0);
+	float OldExposureScale = PostprocessInput0.Load(int3(0, 1, 0)).x;
 	float OldExposure = ExposureOffsetMultipler / ( OldExposureScale != 0 ? OldExposureScale : 1.0f );
 	float FrameTime = EyeAdaptationParams[1].y;
 
diff --git a/Engine/Shaders/PostProcessHistogram.usf b/Engine/Shaders/PostProcessHistogram.usf
--- a/Engine/Shaders/PostProcessHistogram.usf
+++ b/Engine/Shaders/PostProcessHistogram.usf
@@ -54,7 +54,7 @@ void MainCS(
 
 				if(TexelPos.x < GatherExtentInt.x && TexelPos.y < GatherExtentInt.y)
 				{
-					float4 SceneColor = InputNew0.Load(int3(TexelPos, 0));
+					float4 SceneColor = PostprocessInput0.Load(int3(TexelPos, 0));
 	
 //					float LuminanceVal = dot(SceneColor.rgb, float3(0.3f, 0.59f, 0.11f));
 					float LuminanceVal = max(SceneColor.r, max(SceneColor.g, SceneColor.b));
diff --git a/Engine/Shaders/PostProcessVisualizeGBuffer.usf b/Engine/Shaders/PostProcessVisualizeGBuffer.usf
--- a/Engine/Shaders/PostProcessVisualizeGBuffer.usf
+++ b/Engine/Shaders/PostProcessVisualizeGBuffer.usf
@@ -113,7 +113,7 @@ void MainPS(float4 UVAndScreenPos : TEXCOORD0, out float4 OutColor : SV_Target0)
 	
 	// GBuffer size
 	uint2 BufferSize;
-	InputNew0.GetDimensions(BufferSize.x, BufferSize.y);
+	PostprocessInput0.GetDimensions(BufferSize.x, BufferSize.y);
 	float2 InvBufferSize = 1.0f / (float2)BufferSize;
 
 	int2 HalfTileLocalPos = PixelPos - QuarterInsetSize;
diff --git a/Engine/Shaders/PostProcessVisualizeHDR.usf b/Engine/Shaders/PostProcessVisualizeHDR.usf
--- a/Engine/Shaders/PostProcessVisualizeHDR.usf
+++ b/Engine/Shaders/PostProcessVisualizeHDR.usf
@@ -130,9 +130,7 @@ void MainPS(float4 UVAndScreenPos : TEXCOORD0, out float4 OutColor : SV_Target0)
 
 	uint Bucket = (uint)(InsetUV.x * HISTOGRAM_SIZE);
 	
-//	Texture2D HistogramTexture = InputNew1;
-#define HistogramTexture InputNew1 // WAR for HLSLCC not allowing assignment to samplers (yet)
-	float HistogramSum = ComputeHistogramSum(HistogramTexture);
+	float HistogramSum = ComputeHistogramSum(PostprocessInput1);
 
 	if(InsetUV.x < ComputeHistogramPositionFromLuminance(EyeAdaptationParams[0].z))
 	{
diff --git a/Engine/Shaders/ShaderVersion.usf b/Engine/Shaders/ShaderVersion.usf
--- a/Engine/Shaders/ShaderVersion.usf
+++ b/Engine/Shaders/ShaderVersion.usf
@@ -2,4 +2,4 @@
 // Each time the console command is executed it generates a new GUID. As this file is included
 // in common.usf (which should be included in any shader) it allows to invalidate the shader DDC.
 // 
-// GUID = CEB51B844C0F1E837BB98A92CFC5BE4C
\ No newline at end of file
+// GUID = DF48532F444DC92D3D0A8593C3F4749F
\ No newline at end of file
diff --git a/Engine/Source/Runtime/Renderer/Private/PostProcess/RenderingCompositionGraph.cpp b/Engine/Source/Runtime/Renderer/Private/PostProcess/RenderingCompositionGraph.cpp
--- a/Engine/Source/Runtime/Renderer/Private/PostProcess/RenderingCompositionGraph.cpp
+++ b/Engine/Source/Runtime/Renderer/Private/PostProcess/RenderingCompositionGraph.cpp
@@ -672,7 +672,6 @@ void FPostProcessPassParameters::Bind(const FShaderParameterMap& ParameterMap)
 	{
 		PostprocessInputParameter[i].Bind(ParameterMap, *FString::Printf(TEXT("PostprocessInput%d"), i));
 		PostprocessInputParameterSampler[i].Bind(ParameterMap, *FString::Printf(TEXT("PostprocessInput%dSampler"), i));
-		PostprocessInputNew[i].Bind(ParameterMap, *FString::Printf(TEXT("InputNew%d"), i));
 		PostprocessInputSizeParameter[i].Bind(ParameterMap, *FString::Printf(TEXT("PostprocessInput%dSize"), i));
 		PostProcessInputMinMaxParameter[i].Bind(ParameterMap, *FString::Printf(TEXT("PostprocessInput%dMinMax"), i));
 	}
@@ -800,7 +799,6 @@ void FPostProcessPassParameters::Set(
 			const FTextureRHIRef& SrcTexture = InputPooledElement->GetRenderTargetItem().ShaderResourceTexture;
 
 			SetTextureParameter(RHICmdList, ShaderRHI, PostprocessInputParameter[Id], PostprocessInputParameterSampler[Id], LocalFilter, SrcTexture);
-			SetTextureParameter(RHICmdList, ShaderRHI, PostprocessInputNew[Id], InputPooledElement->GetRenderTargetItem().TargetableTexture);
 
 			{
 				float Width = InputPooledElement->GetDesc().Extent.X;
@@ -823,7 +821,6 @@ void FPostProcessPassParameters::Set(
 			// if the input is not there but the shader request it we give it at least some data to avoid d3ddebug errors and shader permutations
 			// to make features optional we use default black for additive passes without shader permutations
 			SetTextureParameter(RHICmdList, ShaderRHI, PostprocessInputParameter[Id], PostprocessInputParameterSampler[Id], LocalFilter, Texture->GetRenderTargetItem().TargetableTexture);
-			SetTextureParameter(RHICmdList, ShaderRHI, PostprocessInputNew[Id], Texture->GetRenderTargetItem().TargetableTexture);
 
 			FVector4 Dummy(1, 1, 1, 1);
 			SetShaderValue(RHICmdList, ShaderRHI, PostprocessInputSizeParameter[Id], Dummy);
@@ -858,7 +855,6 @@ FArchive& operator<<(FArchive& Ar, FPostProcessPassParameters& P)
 	{
 		Ar << P.PostprocessInputParameter[i];
 		Ar << P.PostprocessInputParameterSampler[i];
-		Ar << P.PostprocessInputNew[i];
 		Ar << P.PostprocessInputSizeParameter[i];
 		Ar << P.PostProcessInputMinMaxParameter[i];
 	}
diff --git a/Engine/Source/Runtime/Renderer/Public/PostProcessParameters.h b/Engine/Source/Runtime/Renderer/Public/PostProcessParameters.h
--- a/Engine/Source/Runtime/Renderer/Public/PostProcessParameters.h
+++ b/Engine/Source/Runtime/Renderer/Public/PostProcessParameters.h
@@ -78,7 +78,6 @@ private:
 	FShaderParameter ViewportRect;
 	FShaderResourceParameter PostprocessInputParameter[ePId_Input_MAX];
 	FShaderResourceParameter PostprocessInputParameterSampler[ePId_Input_MAX];
-	FShaderResourceParameter PostprocessInputNew[ePId_Input_MAX];
 	FShaderParameter PostprocessInputSizeParameter[ePId_Input_MAX];
 	FShaderParameter PostProcessInputMinMaxParameter[ePId_Input_MAX];
 	FShaderParameter ScreenPosToPixel;
