fixed check() in debug, FeatureLevel wasn't set on the object

diff --git a/Engine/Source/Runtime/Renderer/Private/ShadowRendering.cpp b/Engine/Source/Runtime/Renderer/Private/ShadowRendering.cpp
--- a/Engine/Source/Runtime/Renderer/Private/ShadowRendering.cpp
+++ b/Engine/Source/Runtime/Renderer/Private/ShadowRendering.cpp
@@ -732,6 +732,7 @@ FShadowDepthDrawingPolicy<bRenderingReflectiveShadowMaps>::FShadowDepthDrawingPo
 	):
 	FMeshDrawingPolicy(InVertexFactory,InMaterialRenderProxy,*InMaterialResource,false,bInCastShadowAsTwoSided),
 	GeometryShader(NULL),
+	FeatureLevel(InFeatureLevel),
 	bDirectionalLight(bInDirectionalLight),
 	bReverseCulling(bInReverseCulling),
 	bOnePassPointLightShadow(bInOnePassPointLightShadow),
