avoid double printout of VRam allocs

diff --git a/Engine/Source/Runtime/Renderer/Private/PostProcess/RenderTargetPool.cpp b/Engine/Source/Runtime/Renderer/Private/PostProcess/RenderTargetPool.cpp
--- a/Engine/Source/Runtime/Renderer/Private/PostProcess/RenderTargetPool.cpp
+++ b/Engine/Source/Runtime/Renderer/Private/PostProcess/RenderTargetPool.cpp
@@ -151,7 +151,11 @@ static void LogVRamUsage(FPooledRenderTarget& Ref)
 	if(Ref.GetDesc().Flags & TexCreate_FastVRAM)
 	{
 		LogVRamUsage(Ref, Ref.GetRenderTargetItem().TargetableTexture);
-		LogVRamUsage(Ref, Ref.GetRenderTargetItem().ShaderResourceTexture);
+
+		if(Ref.GetRenderTargetItem().TargetableTexture != Ref.GetRenderTargetItem().ShaderResourceTexture)
+		{
+			LogVRamUsage(Ref, Ref.GetRenderTargetItem().ShaderResourceTexture);
+		}
 	}
 }
 
