Added OcclusionTint, MinOcclusion to sky light for DFAO artistic controls
Hopeful fix for ensure from zero scaled components

diff --git a/Engine/Shaders/SkyLighting.usf b/Engine/Shaders/SkyLighting.usf
--- a/Engine/Shaders/SkyLighting.usf
+++ b/Engine/Shaders/SkyLighting.usf
@@ -12,6 +12,7 @@ Texture2D BentNormalAOTexture;
 SamplerState BentNormalAOSampler;
 
 float3 ContrastAndNormalizeMulAdd;
+float4 OcclusionTintAndMinOcclusion;
 
 void SkyLightDiffusePS(
 	in float4 UVAndScreenPos : TEXCOORD0, 
@@ -62,17 +63,17 @@ void SkyLightDiffusePS(
 		*/
 
 		float ContrastCurve = 1 / (1 + exp(-ContrastAndNormalizeMulAdd.x * (SkyVisibility * 10 - 5)));
-		SkyVisibility = ContrastCurve * ContrastAndNormalizeMulAdd.y + ContrastAndNormalizeMulAdd.z;
+		SkyVisibility = saturate(ContrastCurve * ContrastAndNormalizeMulAdd.y + ContrastAndNormalizeMulAdd.z);
 		
 #endif
 
 		// Compute the preconvolved incoming lighting with the bent normal direction
 		float3 DiffuseLookup = GetSkySHDiffuse(SkyLightingNormal) * View.SkyLightColor.rgb;
 
-		float ScalarFactors = SkyVisibility * DotProductFactor * InGBufferData.GBufferAO * ScreenSpaceData.AmbientOcclusion;
+		float ScalarFactors = lerp(SkyVisibility, 1, OcclusionTintAndMinOcclusion.w) * DotProductFactor * InGBufferData.GBufferAO * ScreenSpaceData.AmbientOcclusion;
 
 		// Apply AO to the sky diffuse and account for darkening due to the geometry term
-		Lighting = ScalarFactors * DiffuseLookup * InGBufferData.DiffuseColor;
+		Lighting = ScalarFactors * DiffuseLookup * InGBufferData.DiffuseColor + (1 - SkyVisibility) * OcclusionTintAndMinOcclusion.xyz;
 
 		// For debugging
 		//Lighting = (Texture2DSampleLevel(BentNormalAOTexture, BentNormalAOSampler, UV, 0).xyz);
diff --git a/Engine/Source/Runtime/Engine/Classes/Components/SkyLightComponent.h b/Engine/Source/Runtime/Engine/Classes/Components/SkyLightComponent.h
--- a/Engine/Source/Runtime/Engine/Classes/Components/SkyLightComponent.h
+++ b/Engine/Source/Runtime/Engine/Classes/Components/SkyLightComponent.h
@@ -121,6 +121,16 @@ class ENGINE_API USkyLightComponent : public ULightComponentBase
 	UPROPERTY(EditAnywhere, BlueprintReadOnly, Category=DistanceFieldAmbientOcclusion, meta=(UIMin = "0", UIMax = "1"))
 	float Contrast;
 
+	/** 
+	 * Controls the darkest that a fully occluded area can get.
+	 */
+	UPROPERTY(EditAnywhere, BlueprintReadOnly, Category=DistanceFieldAmbientOcclusion, meta=(UIMin = "0", UIMax = "1"))
+	float MinOcclusion;
+
+	/** Tint color on occluded areas, artistic control. */
+	UPROPERTY(EditAnywhere, BlueprintReadOnly, Category=DistanceFieldAmbientOcclusion)
+	FColor OcclusionTint;
+
 	class FSkyLightSceneProxy* CreateSceneProxy() const;
 
 	// Begin UObject Interface
diff --git a/Engine/Source/Runtime/Engine/Private/SkyLightComponent.cpp b/Engine/Source/Runtime/Engine/Private/SkyLightComponent.cpp
--- a/Engine/Source/Runtime/Engine/Private/SkyLightComponent.cpp
+++ b/Engine/Source/Runtime/Engine/Private/SkyLightComponent.cpp
@@ -81,6 +81,8 @@ FSkyLightSceneProxy::FSkyLightSceneProxy(const USkyLightComponent* InLightCompon
 	, IrradianceEnvironmentMap(InLightComponent->IrradianceEnvironmentMap)
 	, OcclusionMaxDistance(InLightComponent->OcclusionMaxDistance)
 	, Contrast(InLightComponent->Contrast)
+	, MinOcclusion(InLightComponent->MinOcclusion)
+	, OcclusionTint(InLightComponent->OcclusionTint)
 {
 }
 
@@ -107,6 +109,8 @@ USkyLightComponent::USkyLightComponent(const class FPostConstructInitializePrope
 	bLowerHemisphereIsBlack = true;
 	bSavedConstructionScriptValuesValid = true;
 	OcclusionMaxDistance = 600;
+	MinOcclusion = 0;
+	OcclusionTint = FColor::Black;
 }
 
 FSkyLightSceneProxy* USkyLightComponent::CreateSceneProxy() const
diff --git a/Engine/Source/Runtime/Engine/Public/SceneManagement.h b/Engine/Source/Runtime/Engine/Public/SceneManagement.h
--- a/Engine/Source/Runtime/Engine/Public/SceneManagement.h
+++ b/Engine/Source/Runtime/Engine/Public/SceneManagement.h
@@ -574,6 +574,8 @@ public:
 	FSHVectorRGB3 IrradianceEnvironmentMap;
 	float OcclusionMaxDistance;
 	float Contrast;
+	float MinOcclusion;
+	FLinearColor OcclusionTint;
 };
 
 
diff --git a/Engine/Source/Runtime/Renderer/Private/DistanceFieldSurfaceCacheLighting.cpp b/Engine/Source/Runtime/Renderer/Private/DistanceFieldSurfaceCacheLighting.cpp
--- a/Engine/Source/Runtime/Renderer/Private/DistanceFieldSurfaceCacheLighting.cpp
+++ b/Engine/Source/Runtime/Renderer/Private/DistanceFieldSurfaceCacheLighting.cpp
@@ -3251,7 +3251,7 @@ void UpdateVisibleObjectBuffers(const FScene* Scene, const FDistanceFieldAOParam
 						ObjectData2.Add(*(FVector4*)&UniformScaleVolumeToWorld.M[2]);
 						ObjectData2.Add(*(FVector4*)&UniformScaleVolumeToWorld.M[3]);
 
-						const FMatrix WorldToVolume = UniformScaleVolumeToWorld.InverseFast();
+						const FMatrix WorldToVolume = UniformScaleVolumeToWorld.Inverse();
 						// WorldToVolume
 						ObjectData.Add(*(FVector4*)&WorldToVolume.M[0]);
 						ObjectData.Add(*(FVector4*)&WorldToVolume.M[1]);
@@ -3805,9 +3805,10 @@ public:
 		DynamicBentNormalAOTexture.Bind(Initializer.ParameterMap, TEXT("BentNormalAOTexture"));
 		DynamicBentNormalAOSampler.Bind(Initializer.ParameterMap, TEXT("BentNormalAOSampler"));
 		ContrastAndNormalizeMulAdd.Bind(Initializer.ParameterMap, TEXT("ContrastAndNormalizeMulAdd"));
+		OcclusionTintAndMinOcclusion.Bind(Initializer.ParameterMap, TEXT("OcclusionTintAndMinOcclusion"));
 	}
 
-	void SetParameters(FRHICommandList& RHICmdList, const FSceneView& View, FTextureRHIParamRef DynamicBentNormalAO, const FDistanceFieldAOParameters& Parameters)
+	void SetParameters(FRHICommandList& RHICmdList, const FSceneView& View, FTextureRHIParamRef DynamicBentNormalAO, const FDistanceFieldAOParameters& Parameters, const FSkyLightSceneProxy* SkyLight)
 	{
 		const FPixelShaderRHIParamRef ShaderRHI = GetPixelShader();
 		FGlobalShader::SetParameters(RHICmdList, ShaderRHI, View);
@@ -3822,6 +3823,10 @@ public:
 		const float Add = -Min / (Max - Min);
 
 		SetShaderValue(RHICmdList, ShaderRHI, ContrastAndNormalizeMulAdd, FVector(Parameters.Contrast, Mul, Add));
+
+		FVector4 OcclusionTintAndMinOcclusionValue = FVector4(SkyLight->OcclusionTint);
+		OcclusionTintAndMinOcclusionValue.W = SkyLight->MinOcclusion;
+		SetShaderValue(RHICmdList, ShaderRHI, OcclusionTintAndMinOcclusion, OcclusionTintAndMinOcclusionValue);
 	}
 
 	// FShader interface.
@@ -3832,6 +3837,7 @@ public:
 		Ar << DynamicBentNormalAOTexture;
 		Ar << DynamicBentNormalAOSampler;
 		Ar << ContrastAndNormalizeMulAdd;
+		Ar << OcclusionTintAndMinOcclusion;
 		return bShaderHasOutdatedParameters;
 	}
 
@@ -3841,6 +3847,7 @@ private:
 	FShaderResourceParameter DynamicBentNormalAOTexture;
 	FShaderResourceParameter DynamicBentNormalAOSampler;
 	FShaderParameter ContrastAndNormalizeMulAdd;
+	FShaderParameter OcclusionTintAndMinOcclusion;
 };
 
 IMPLEMENT_SHADER_TYPE(template<>,TDynamicSkyLightDiffusePS<true>,TEXT("SkyLighting"),TEXT("SkyLightDiffusePS"),SF_Pixel);
@@ -3890,7 +3897,7 @@ void FDeferredShadingSceneRenderer::RenderDynamicSkyLighting(FRHICommandListImme
 
 				static FGlobalBoundShaderState BoundShaderState;
 				SetGlobalBoundShaderState(RHICmdList, View.GetFeatureLevel(), BoundShaderState, GFilterVertexDeclaration.VertexDeclarationRHI, *VertexShader, *PixelShader);
-				PixelShader->SetParameters(RHICmdList, View, DynamicBentNormalAO->GetRenderTargetItem().ShaderResourceTexture, Parameters);
+				PixelShader->SetParameters(RHICmdList, View, DynamicBentNormalAO->GetRenderTargetItem().ShaderResourceTexture, Parameters, Scene->SkyLight);
 			}
 			else
 			{
@@ -3898,7 +3905,7 @@ void FDeferredShadingSceneRenderer::RenderDynamicSkyLighting(FRHICommandListImme
 
 				static FGlobalBoundShaderState BoundShaderState;
 				SetGlobalBoundShaderState(RHICmdList, View.GetFeatureLevel(), BoundShaderState, GFilterVertexDeclaration.VertexDeclarationRHI, *VertexShader, *PixelShader);
-				PixelShader->SetParameters(RHICmdList, View, GWhiteTexture->TextureRHI, Parameters);
+				PixelShader->SetParameters(RHICmdList, View, GWhiteTexture->TextureRHI, Parameters, Scene->SkyLight);
 			}
 
 			DrawRectangle( 
