GitHub pull request #68

diff --git a/Engine/Source/Runtime/Renderer/Private/Renderer.cpp b/Engine/Source/Runtime/Renderer/Private/Renderer.cpp
--- a/Engine/Source/Runtime/Renderer/Private/Renderer.cpp
+++ b/Engine/Source/Runtime/Renderer/Private/Renderer.cpp
@@ -409,6 +409,21 @@ static bool RendererExec( UWorld* InWorld, const TCHAR* Cmd, FOutputDevice& Ar )
 
 	return false;
 }
+
+ICustomCulling* GCustomCullingImpl = nullptr;
+
+void FRendererModule::RegisterCustomCullingImpl(ICustomCulling* impl)
+{
+	check(GCustomCullingImpl == nullptr);
+	GCustomCullingImpl = impl;
+}
+
+void FRendererModule::UnregisterCustomCullingImpl(ICustomCulling* impl)
+{
+	check(GCustomCullingImpl == impl);
+	GCustomCullingImpl = nullptr;
+}
+
 FStaticSelfRegisteringExec RendererExecRegistration(RendererExec);
 
 void FRendererModule::ExecVisualizeTextureCmd( const FString& Cmd )
diff --git a/Engine/Source/Runtime/Renderer/Private/RendererPrivate.h b/Engine/Source/Runtime/Renderer/Private/RendererPrivate.h
--- a/Engine/Source/Runtime/Renderer/Private/RendererPrivate.h
+++ b/Engine/Source/Runtime/Renderer/Private/RendererPrivate.h
@@ -60,8 +60,14 @@ public:
 		return AllocatedScenes;
 	}
 
+	virtual void RegisterCustomCullingImpl(ICustomCulling* impl) override;
+	virtual void UnregisterCustomCullingImpl(ICustomCulling* impl) override;
+
 private:
 	TSet<FSceneInterface*> AllocatedScenes;
+	ICustomCulling* CustomCullingImpl;
 };
 
+extern ICustomCulling* GCustomCullingImpl;
+
 #endif
diff --git a/Engine/Source/Runtime/Renderer/Private/SceneRendering.cpp b/Engine/Source/Runtime/Renderer/Private/SceneRendering.cpp
--- a/Engine/Source/Runtime/Renderer/Private/SceneRendering.cpp
+++ b/Engine/Source/Runtime/Renderer/Private/SceneRendering.cpp
@@ -155,6 +155,7 @@ FViewInfo::FViewInfo(const FSceneView* InView)
 	:	FSceneView(*InView)
 	,	IndividualOcclusionQueries((FSceneViewState*)InView->State,1)
 	,	GroupedOcclusionQueries((FSceneViewState*)InView->State,FOcclusionQueryBatcher::OccludedPrimitiveQueryBatchSize)
+	,	CustomVisibilityQuery(nullptr)
 {
 	Init();
 }
@@ -202,6 +203,10 @@ FViewInfo::~FViewInfo()
 	{
 		DynamicResources[ResourceIndex]->ReleasePrimitiveResource();
 	}
+	if (CustomVisibilityQuery)
+	{
+		CustomVisibilityQuery->Release();
+	}
 }
 
 void FViewInfo::SetupSkyIrradianceEnvironmentMapConstants(FVector4* OutSkyIrradianceEnvironmentMap) const
@@ -736,6 +741,16 @@ FSceneRenderer::FSceneRenderer(const FSceneViewFamily* InViewFamily,FHitProxyCon
 		ViewFamily.EngineShowFlags.HitProxies = 1;
 	}
 
+	// launch custom visibility queries for views
+	if (GCustomCullingImpl)
+	{
+		for(int32 ViewIndex = 0;ViewIndex < Views.Num();ViewIndex++)
+		{
+			FViewInfo& ViewInfo = Views[ViewIndex];
+			ViewInfo.CustomVisibilityQuery = GCustomCullingImpl->CreateQuery(ViewInfo);
+		}
+	}
+
 	ViewFamily.ComputeFamilySize();
 
 	// copy off the requests
diff --git a/Engine/Source/Runtime/Renderer/Private/SceneRendering.h b/Engine/Source/Runtime/Renderer/Private/SceneRendering.h
--- a/Engine/Source/Runtime/Renderer/Private/SceneRendering.h
+++ b/Engine/Source/Runtime/Renderer/Private/SceneRendering.h
@@ -469,6 +469,9 @@ public:
 
 	TShaderMap<FGlobalShaderType>* ShaderMap;
 
+	/** Custom visibility query for view */
+	ICustomVisibilityQuery* CustomVisibilityQuery;
+
 	/** 
 	 * Initialization constructor. Passes all parameters to FSceneView constructor
 	 */
diff --git a/Engine/Source/Runtime/Renderer/Private/SceneVisibility.cpp b/Engine/Source/Runtime/Renderer/Private/SceneVisibility.cpp
--- a/Engine/Source/Runtime/Renderer/Private/SceneVisibility.cpp
+++ b/Engine/Source/Runtime/Renderer/Private/SceneVisibility.cpp
@@ -223,6 +223,7 @@ bool FViewInfo::IsDistanceCulled( float DistanceSquared, float MinDrawDistance,
 /**
  * Frustum cull primitives in the scene against the view.
  */
+template<bool UseCustomCulling>
 static int32 FrustumCull(const FScene* Scene, FViewInfo& View)
 {
 	SCOPE_CYCLE_COUNTER(STAT_FrustumCull);
@@ -231,12 +232,20 @@ static int32 FrustumCull(const FScene* Scene, FViewInfo& View)
 	FVector ViewOriginForDistanceCulling = View.ViewMatrices.ViewOrigin;
 	float MaxDrawDistanceScale = GetCachedScalabilityCVars().ViewDistanceScale;
 	float FadeRadius = GDisableLODFade ? 0.0f : GDistanceFadeMaxTravel;
+	uint8 CustomVisibilityFlags = EOcclusionFlags::CanBeOccluded | EOcclusionFlags::HasPrecomputedVisibility;
 
 	for (FSceneBitArray::FIterator BitIt(View.PrimitiveVisibilityMap); BitIt; ++BitIt)
 	{
 		const FPrimitiveBounds& Bounds = Scene->PrimitiveBounds[BitIt.GetIndex()];
 		float DistanceSquared = (Bounds.Origin - ViewOriginForDistanceCulling).SizeSquared();
 		float MaxDrawDistance = Bounds.MaxDrawDistance * MaxDrawDistanceScale;
+		int32 VisibilityId = INDEX_NONE;
+
+		if (UseCustomCulling &&
+			((Scene->PrimitiveOcclusionFlags[BitIt.GetIndex()] & CustomVisibilityFlags) == CustomVisibilityFlags))
+		{
+			VisibilityId = Scene->PrimitiveVisibilityIds[BitIt.GetIndex()].ByteIndex;
+		}
 
 		// If cull distance is disabled, always show (except foliage)
 		if (View.Family->EngineShowFlags.DistanceCulledPrimitives
@@ -248,6 +257,7 @@ static int32 FrustumCull(const FScene* Scene, FViewInfo& View)
 		// The primitive is always culled if it exceeds the max fade distance or lay outside the view frustum.
 		if (DistanceSquared > FMath::Square(MaxDrawDistance + FadeRadius) ||
 			DistanceSquared < Bounds.MinDrawDistanceSq ||
+			(UseCustomCulling && !View.CustomVisibilityQuery->IsVisible(VisibilityId, FBoxSphereBounds(Bounds.Origin, Bounds.BoxExtent, Bounds.SphereRadius))) ||
 			View.ViewFrustum.IntersectSphere(Bounds.Origin, Bounds.SphereRadius) == false ||
 			View.ViewFrustum.IntersectBox(Bounds.Origin, Bounds.BoxExtent) == false)
 		{
@@ -1425,7 +1435,7 @@ static void MarkRelevantStaticMeshesForView(
 		const FPrimitiveSceneInfo* RESTRICT PrimitiveSceneInfo = Scene->Primitives[PrimitiveIndex];
 		const FPrimitiveBounds& Bounds = Scene->PrimitiveBounds[PrimitiveIndex];
 		const FPrimitiveViewRelevance& ViewRelevance = View.PrimitiveViewRelevanceMap[PrimitiveIndex];
-		
+
 		int8 LODToRender = ComputeLODForMeshes( PrimitiveSceneInfo->StaticMeshes, View, Bounds.Origin, Bounds.SphereRadius, ViewData.ForcedLODLevel, ViewData.LODScale);
 
 		float DistanceSquared = (Bounds.Origin - ViewData.ViewOrigin).SizeSquared();
@@ -1812,7 +1822,7 @@ void FSceneRenderer::PreVisibilityFrameSetup(FRHICommandListImmediate& RHICmdLis
 
 			// This finishes the update of view state
 			ViewState->UpdateLastRenderTime(*View.Family);
-		}		
+		}
 	}
 }
 
@@ -1955,7 +1965,15 @@ void FSceneRenderer::ComputeViewVisibility(FRHICommandListImmediate& RHICmdList)
 		// Most views use standard frustum culling.
 		if (bNeedsFrustumCulling)
 		{
-			int32 NumCulledPrimitivesForView = FrustumCull(Scene, View);
+			int32 NumCulledPrimitivesForView;
+			if (View.CustomVisibilityQuery && View.CustomVisibilityQuery->Prepare())
+			{
+				NumCulledPrimitivesForView = FrustumCull<true>(Scene, View);
+			}
+			else
+			{
+				NumCulledPrimitivesForView = FrustumCull<false>(Scene, View);
+			}
 			STAT(NumCulledPrimitives += NumCulledPrimitivesForView);
 			UpdatePrimitiveFading(Scene, View);			
 		}
@@ -2010,7 +2028,7 @@ void FSceneRenderer::ComputeViewVisibility(FRHICommandListImmediate& RHICmdList)
 
 		{
 			SCOPE_CYCLE_COUNTER(STAT_ViewRelevance);
-			// Compute view relevance for all visible primitives.
+		// Compute view relevance for all visible primitives.
 			if (!CVarLegacySingleThreadedRelevance.GetValueOnRenderThread())
 			{
 				ComputeAndMarkRelevanceForViewParallel(RHICmdList, Scene, View, ViewBit, PreRenderViewMasks, HasDynamicMeshElementsMasks, HasDynamicEditorMeshElementsMasks);
@@ -2020,7 +2038,7 @@ void FSceneRenderer::ComputeViewVisibility(FRHICommandListImmediate& RHICmdList)
 				// This array contains a list of relevant static primities.
 				FRelevantStaticPrimitives RelevantStaticPrimitives;
 				ComputeRelevanceForView(RHICmdList, Scene, View, ViewBit, RelevantStaticPrimitives, PreRenderViewMasks, HasDynamicMeshElementsMasks, HasDynamicEditorMeshElementsMasks);
-				MarkRelevantStaticMeshesForView(Scene, View, RelevantStaticPrimitives);
+		MarkRelevantStaticMeshesForView(Scene, View, RelevantStaticPrimitives);
 			}
 		}
 
@@ -2055,7 +2073,7 @@ void FSceneRenderer::ComputeViewVisibility(FRHICommandListImmediate& RHICmdList)
 	}
 	else
 	{
-		DispatchPreRenderView(Scene, &ViewFamily, PreRenderViewMasks, FrameNumber);
+	DispatchPreRenderView(Scene, &ViewFamily, PreRenderViewMasks, FrameNumber);
 	}
 
 	INC_DWORD_STAT_BY(STAT_ProcessedPrimitives,NumProcessedPrimitives);
diff --git a/Engine/Source/Runtime/Renderer/Public/RendererInterface.h b/Engine/Source/Runtime/Renderer/Public/RendererInterface.h
--- a/Engine/Source/Runtime/Renderer/Public/RendererInterface.h
+++ b/Engine/Source/Runtime/Renderer/Public/RendererInterface.h
@@ -421,6 +421,21 @@ enum EDrawRectangleFlags
 	EDRF_UseTriangleOptimization
 };
 
+class ICustomVisibilityQuery: public IRefCountedObject
+{
+public:
+	/** prepares the query for visibility tests */
+	virtual bool Prepare() = 0;
+
+	/** test primitive visiblity */
+	virtual bool IsVisible(int32 VisibilityId, const FBoxSphereBounds& Bounds) = 0;
+};
+
+class ICustomCulling
+{
+public:
+	virtual ICustomVisibilityQuery* CreateQuery (const FSceneView& View) = 0;
+};
 
 
 
@@ -515,6 +530,10 @@ public:
 
 	/** @return Returns a vertex declaration that can be used with with the DrawRectangle() function */
 	virtual TGlobalResource<FFilterVertexDeclaration>& GetFilterVertexDeclaration() = 0;
+
+	/** Register/unregister a custom occlusion culling implementation */
+	virtual void RegisterCustomCullingImpl(ICustomCulling* impl) = 0;
+	virtual void UnregisterCustomCullingImpl(ICustomCulling* impl) = 0;
 };
 
 
