UE4 - Disable SSS on OpenGL (needs cross-compiler fix)

diff --git a/Engine/Source/Runtime/Renderer/Private/CompositionLighting/CompositionLighting.cpp b/Engine/Source/Runtime/Renderer/Private/CompositionLighting/CompositionLighting.cpp
--- a/Engine/Source/Runtime/Renderer/Private/CompositionLighting/CompositionLighting.cpp
+++ b/Engine/Source/Runtime/Renderer/Private/CompositionLighting/CompositionLighting.cpp
@@ -319,17 +319,17 @@ void FCompositionLighting::ProcessAfterBasePass(FRHICommandListImmediate& RHICmd
 		}
 
 		// The graph setup should be finished before this line ----------------------------------------
-    
+
 		SCOPED_DRAW_EVENT(RHICmdList, LightCompositionTasks_PreLighting, DEC_SCENE_ITEMS);
 
 		TRefCountPtr<IPooledRenderTarget>& SceneColor = GSceneRenderTargets.GetSceneColor();
 
 		Context.FinalOutput.GetOutput()->RenderTargetDesc = SceneColor->GetDesc();
 		Context.FinalOutput.GetOutput()->PooledRenderTarget = SceneColor;
-    
+
 		// you can add multiple dependencies
 		CompositeContext.Root->AddDependency(Context.FinalOutput);
-    
+
 		CompositeContext.Process(TEXT("CompositionLighting_AfterBasePass"));
 	}
 }
@@ -363,7 +363,9 @@ void FCompositionLighting::ProcessLighting(FRHICommandListImmediate& RHICmdList,
 
 			bool bSimpleDynamicLighting = IsSimpleDynamicLightingEnabled();
 
-			if (View.bScreenSpaceSubsurfacePassNeeded && Radius > 0 && !bSimpleDynamicLighting && View.Family->EngineShowFlags.SubsurfaceScattering)
+			if (View.bScreenSpaceSubsurfacePassNeeded && Radius > 0 && !bSimpleDynamicLighting && View.Family->EngineShowFlags.SubsurfaceScattering && 
+				//@todo-rco: Remove this when we fix the cross-compiler
+				!IsOpenGLPlatform(GRHIShaderPlatform))
 			{
 				int32 HalfRes = CVarSSSHalfRes.GetValueOnRenderThread();
 
diff --git a/Engine/Source/Runtime/Renderer/Private/PostProcess/PostProcessSubsurface.cpp b/Engine/Source/Runtime/Renderer/Private/PostProcess/PostProcessSubsurface.cpp
--- a/Engine/Source/Runtime/Renderer/Private/PostProcess/PostProcessSubsurface.cpp
+++ b/Engine/Source/Runtime/Renderer/Private/PostProcess/PostProcessSubsurface.cpp
@@ -34,7 +34,8 @@ class FPostProcessSubsurfaceSetupPS : public FGlobalShader
 
 	static bool ShouldCache(EShaderPlatform Platform)
 	{
-		return IsFeatureLevelSupported(Platform, ERHIFeatureLevel::SM4);
+		//@todo-rco: Remove this when we fix the cross-compiler
+		return IsFeatureLevelSupported(Platform, ERHIFeatureLevel::SM4) && !IsOpenGLPlatform(Platform);
 	}
 
 	static void ModifyCompilationEnvironment(EShaderPlatform Platform, FShaderCompilerEnvironment& OutEnvironment)
@@ -323,7 +324,8 @@ class TPostProcessSubsurfacePS : public FGlobalShader
 
 	static bool ShouldCache(EShaderPlatform Platform)
 	{
-		return IsFeatureLevelSupported(Platform, ERHIFeatureLevel::SM4);
+		//@todo-rco: Remove this when we fix the cross-compiler
+		return IsFeatureLevelSupported(Platform, ERHIFeatureLevel::SM4) && !IsOpenGLPlatform(Platform);
 	}
 
 	static void ModifyCompilationEnvironment(EShaderPlatform Platform, FShaderCompilerEnvironment& OutEnvironment)
