* optimized ScreenSpaceSubsurfaceScattering (4 pass: setup, blurX, blurY, recombine)
* fixed reallocation of scene render targets because of Skylight, was wasting VRam

diff --git a/Engine/Shaders/PostProcessSubsurface.usf b/Engine/Shaders/PostProcessSubsurface.usf
--- a/Engine/Shaders/PostProcessSubsurface.usf
+++ b/Engine/Shaders/PostProcessSubsurface.usf
@@ -40,6 +40,12 @@ float GetSubsurfaceStrength(float2 UV)
 // 0:speculars leak into SSS / 1:requires alpha channel for scenecolor
 // #define SSSS_SPECULAR_CORRECTION 1
 
+// set by C++ for SubsurfaceRecombinePS
+// #define SSS_RECOMBINE_METHOD 0:without reconstruct specular/1:with reconstruct specular
+
+// set by C++ for SubsurfacePS
+// #define SSS_DIRECTION 0/1
+
 // define by C++ for the setup pass
 #ifndef SETUP_MODE
 	#define SETUP_MODE 0
@@ -95,16 +101,6 @@ float GetSubsurfaceStrength(float2 UV)
 	error
 #endif
 
-#if SSS_METHOD == 0
-	// horizontal
-#elif SSS_METHOD == 1
-	// vertical
-#elif SSS_METHOD == 2
-	// vertical and reconstruct specular
-#else
-	error
-#endif
-
 // 0 / 1
 #define VISUALIZE_KERNEL 0
 
@@ -239,27 +235,37 @@ void MainPS(float4 UVAndScreenPos : TEXCOORD0, out float4 OutColor : SV_Target0)
 	// *1000 as we store the sample distances / 1000 to allows for scaling
 	float Scaler = View.ViewSizeAndSceneTexelSize.x * View.ViewSizeAndSceneTexelSize.z * 1000.0f;
 
-#if SSS_METHOD == 0
-	// first pass
+#if SSS_DIRECTION == 0
+	// horizontal
 	float2 Direction = float2(1, 0) * Scaler;
 #else
-	// second pass
+	// vertical
 	float2 Direction = float2(0, View.ViewSizeAndSceneTexelSize.w / View.ViewSizeAndSceneTexelSize.z) * Scaler;
 #endif
-	OutColor = SSSSBlurPS(UV, sssWidth, Direction, false);
+
+	OutColor.rgb = SSSSBlurPS(UV, sssWidth, Direction, false).rgb;
 
 	OutColor.a = 1.0f;
+}
+
+
+
+
 
-	// Lighting
+void SubsurfaceRecombinePS(float4 UVAndScreenPos : TEXCOORD0, out float4 OutColor : SV_Target0)
+{
+	float2 UV = UVAndScreenPos.xy;
 
-#if SSS_METHOD > 0
-	// on second pass we recombine with the scene color
+	OutColor.rgb = Texture2DSample(PostprocessInput0, PostprocessInput0Sampler, UV).rgb;
+	OutColor.a = 1.0f;
+
+	// recombine with the scene color
 
 	// can be optimized
 	float2 PixelPos = UVAndScreenPos.zw * ScreenPosToPixel.xy + ScreenPosToPixel.zw;
 	float2 UVSceneColor = (PixelPos + 0.5f) * View.ViewSizeAndSceneTexelSize.zw;
 
-	float4 SceneColor4 = Texture2DSample(PostprocessInput0, PostprocessInput0Sampler, UVSceneColor);
+	float4 SceneColor4 = Texture2DSample(PostprocessInput1, PostprocessInput1Sampler, UVSceneColor);
 
 	float4 SSSColor = OutColor;
 
@@ -267,12 +273,11 @@ void MainPS(float4 UVAndScreenPos : TEXCOORD0, out float4 OutColor : SV_Target0)
 
 	SSSColor.a = (ScreenSpaceData.GBuffer.ShadingModelID >= SHADINGMODELID_SUBSURFACE_PROFILE);
 
-#if SSS_METHOD > 1
+#if SSS_RECOMBINE_METHOD
 	// we took the specular highlights out, now we add them back in
 	SSSColor.rgb += LightAccumulator_ReconstructNonDiffuseLighting(SceneColor4);
 #endif
 
 	// alpha channel can be anything - we consumed the value and likely output to a format that doesn't need the value
 	OutColor = float4(lerp(SceneColor4.rgb, SSSColor.rgb, SSSColor.a), 0);
-#endif
 }
diff --git a/Engine/Shaders/SeparableSSS.usf b/Engine/Shaders/SeparableSSS.usf
--- a/Engine/Shaders/SeparableSSS.usf
+++ b/Engine/Shaders/SeparableSSS.usf
@@ -253,8 +253,8 @@ float4 GetKernel(uint SampleIndex, uint SubsurfaceProfileInt)
 // EPIC changed start
 #define SSSSSampleShadowmap(coord)	float4(0,0,0,0)	// todo: not used yet
 #define SSSSSampleMasking(texcoord)	(GetScreenSpaceData(texcoord).GBuffer.ShadingModelID == SHADINGMODELID_SUBSURFACE_PROFILE)	// todo: masking doesn't work with filtering ideally put in alpha to get it filtered
-#define SSSSSampleSceneColor(coord) (SSSSSampleMasking(coord) * Texture2DSample(PostprocessInput1, PostprocessInput1Sampler, coord))
-#define SSSSSampleSceneColorPoint(coord) ( Texture2DSample(PostprocessInput1, PostprocessInput1Sampler, coord))		// todo: could be made point
+#define SSSSSampleSceneColor(coord) (SSSSSampleMasking(coord) * Texture2DSample(PostprocessInput0, PostprocessInput0Sampler, coord))
+#define SSSSSampleSceneColorPoint(coord) ( Texture2DSample(PostprocessInput0, PostprocessInput0Sampler, coord))		// todo: could be made point
 #define SSSSSampleSceneDepth(coord) CalcSceneDepth(coord)		// todo: could be made linear
 #define SSSSSampleSceneDepthPoint(coord) CalcSceneDepth(coord)
 // EPIC changed end
diff --git a/Engine/Source/Runtime/Renderer/Private/CompositionLighting/CompositionLighting.cpp b/Engine/Source/Runtime/Renderer/Private/CompositionLighting/CompositionLighting.cpp
--- a/Engine/Source/Runtime/Renderer/Private/CompositionLighting/CompositionLighting.cpp
+++ b/Engine/Source/Runtime/Renderer/Private/CompositionLighting/CompositionLighting.cpp
@@ -39,14 +39,6 @@ static TAutoConsoleVariable<float> CVarSSSScale(
 	TEXT(">1: scale scatter radius up (for testing)"),
 	ECVF_RenderThreadSafe);
 
-static TAutoConsoleVariable<int32> CVarSSSHalfRes(
-	TEXT("r.SSS.HalfRes"),
-	2,
-	TEXT("0: full quality (default)\n")
-	TEXT("1: half resolution setup, low quality but faster (currently minor shift from SubsurfaceSetup to SubsurfaceTemp)\n")
-	TEXT("2: also half horizontal during first blurring step, lower quality, even faster"),
-	ECVF_RenderThreadSafe  | ECVF_Scalability);
-
 static bool IsAmbientCubemapPassRequired(FPostprocessContext& Context)
 {
 	FScene* Scene = (FScene*)Context.View.Family->Scene;
@@ -367,18 +359,20 @@ void FCompositionLighting::ProcessLighting(FRHICommandListImmediate& RHICmdList,
 				//@todo-rco: Remove this when we fix the cross-compiler
 				!IsOpenGLPlatform(GRHIShaderPlatform))
 			{
-				int32 HalfRes = CVarSSSHalfRes.GetValueOnRenderThread();
-
-				FRenderingCompositePass* PassSetup = Context.Graph.RegisterPass(new(FMemStack::Get()) FRCPassPostProcessSubsurfaceSetup(false, HalfRes > 0));
+				FRenderingCompositePass* PassSetup = Context.Graph.RegisterPass(new(FMemStack::Get()) FRCPassPostProcessSubsurfaceSetup(false, true));
 				PassSetup->SetInput(ePId_Input0, Context.FinalOutput);
 
-				FRenderingCompositePass* Pass0 = Context.Graph.RegisterPass(new(FMemStack::Get()) FRCPassPostProcessSubsurface(0, Radius, HalfRes > 1));
-				Pass0->SetInput(ePId_Input1, PassSetup);
+				FRenderingCompositePass* Pass0 = Context.Graph.RegisterPass(new(FMemStack::Get()) FRCPassPostProcessSubsurface(0, Radius, true));
+				Pass0->SetInput(ePId_Input0, PassSetup);
+
+				FRenderingCompositePass* Pass1 = Context.Graph.RegisterPass(new(FMemStack::Get()) FRCPassPostProcessSubsurface(1, Radius, true));
+				Pass1->SetInput(ePId_Input0, Pass0);
 
-				FRenderingCompositePass* Pass1 = Context.Graph.RegisterPass(new(FMemStack::Get()) FRCPassPostProcessSubsurface(1, Radius));
-				Pass1->SetInput(ePId_Input0, Context.FinalOutput);
-				Pass1->SetInput(ePId_Input1, Pass0);
-				Context.FinalOutput = FRenderingCompositeOutputRef(Pass1);
+				// full res composite pass, no blurring (Radius=0)
+				FRenderingCompositePass* RecombinePass = Context.Graph.RegisterPass(new(FMemStack::Get()) FRCPassPostProcessSubsurfaceRecombine());
+				RecombinePass->SetInput(ePId_Input0, Pass1);
+				RecombinePass->SetInput(ePId_Input1, Context.FinalOutput);					
+				Context.FinalOutput = FRenderingCompositeOutputRef(RecombinePass);
 			}
 		}
 
diff --git a/Engine/Source/Runtime/Renderer/Private/PostProcess/PostProcessSubsurface.cpp b/Engine/Source/Runtime/Renderer/Private/PostProcess/PostProcessSubsurface.cpp
--- a/Engine/Source/Runtime/Renderer/Private/PostProcess/PostProcessSubsurface.cpp
+++ b/Engine/Source/Runtime/Renderer/Private/PostProcess/PostProcessSubsurface.cpp
@@ -16,12 +16,17 @@ ENGINE_API const IPooledRenderTarget* GetSubsufaceProfileTexture_RT(FRHICommandL
 
 static TAutoConsoleVariable<int32> CVarSSSSampleSet(
 	TEXT("r.SSS.SampleSet"),
-	2,
+	3,
 	TEXT("0: lowest quality\n")
 	TEXT("1: medium quality\n")
 	TEXT("2: high quality (default)"),
 	ECVF_RenderThreadSafe  | ECVF_Scalability);
 
+static bool SubsurfaceShouldCache(EShaderPlatform Platform)
+{
+	//@todo-rco: Remove this when we fix the cross-compiler
+	return IsFeatureLevelSupported(Platform, ERHIFeatureLevel::SM4) && !IsOpenGLPlatform(Platform);
+}
 
 /**
  * Encapsulates the post processing ambient occlusion pixel shader.
@@ -34,8 +39,7 @@ class FPostProcessSubsurfaceSetupPS : public FGlobalShader
 
 	static bool ShouldCache(EShaderPlatform Platform)
 	{
-		//@todo-rco: Remove this when we fix the cross-compiler
-		return IsFeatureLevelSupported(Platform, ERHIFeatureLevel::SM4) && !IsOpenGLPlatform(Platform);
+		return SubsurfaceShouldCache(Platform);
 	}
 
 	static void ModifyCompilationEnvironment(EShaderPlatform Platform, FShaderCompilerEnvironment& OutEnvironment)
@@ -314,24 +318,23 @@ FPooledRenderTargetDesc FRCPassPostProcessSubsurfaceSetup::ComputeOutputDesc(EPa
 }
 
 
-/** Encapsulates the post processing down sample pixel shader. */
-// @param Method 0: horizontal, 1:vertical, 2:vertical and reconstruct specular
+/** Encapsulates the post processing subsurface pixel shader. */
+// @param Direction 0: horizontal, 1:vertical
 // @param SampleSet 0:low, 1:med, 2:high
-template <uint32 Method, uint32 SampleSet>
+template <uint32 Direction, uint32 SampleSet>
 class TPostProcessSubsurfacePS : public FGlobalShader
 {
 	DECLARE_SHADER_TYPE(TPostProcessSubsurfacePS, Global);
 
 	static bool ShouldCache(EShaderPlatform Platform)
 	{
-		//@todo-rco: Remove this when we fix the cross-compiler
-		return IsFeatureLevelSupported(Platform, ERHIFeatureLevel::SM4) && !IsOpenGLPlatform(Platform);
+		return SubsurfaceShouldCache(Platform);
 	}
 
 	static void ModifyCompilationEnvironment(EShaderPlatform Platform, FShaderCompilerEnvironment& OutEnvironment)
 	{
 		FGlobalShader::ModifyCompilationEnvironment(Platform,OutEnvironment);
-		OutEnvironment.SetDefine(TEXT("SSS_METHOD"), Method);
+		OutEnvironment.SetDefine(TEXT("SSS_DIRECTION"), Direction);
 		OutEnvironment.SetDefine(TEXT("SSS_SAMPLESET"), SampleSet);
 	}
 
@@ -414,18 +417,19 @@ public:
 #undef VARIATION2
 
 
-FRCPassPostProcessSubsurface::FRCPassPostProcessSubsurface(uint32 InPass, float InRadius, bool bInHalfRes)
+FRCPassPostProcessSubsurface::FRCPassPostProcessSubsurface(uint32 InDirection, float InRadius, bool bInHalfRes)
 	: Radius(InRadius)
-	, Pass(InPass) 
+	, Direction(InDirection) 
 	, bHalfRes(bInHalfRes) 
 {
+	check(InDirection < 2);
+	check(Radius >= 0.0f);
 }
 
-template <uint32 Method, uint32 SampleSet>
-void SetSubsurfaceShader(const FRenderingCompositePassContext& Context, float InRadius)
+template <uint32 Direction, uint32 SampleSet>
+void SetSubsurfaceShader(const FRenderingCompositePassContext& Context, TShaderMapRef<FPostProcessVS> &VertexShader, float InRadius)
 {
-	TShaderMapRef<FPostProcessVS> VertexShader(Context.GetShaderMap());
-	TShaderMapRef<TPostProcessSubsurfacePS<Method,SampleSet> > PixelShader(Context.GetShaderMap());
+	TShaderMapRef<TPostProcessSubsurfacePS<Direction, SampleSet> > PixelShader(Context.GetShaderMap());
 
 	static FGlobalBoundShaderState BoundShaderState;
 
@@ -435,15 +439,15 @@ void SetSubsurfaceShader(const FRenderingCompositePassContext& Context, float In
 	VertexShader->SetParameters(Context);
 }
 
-// 0:horizontal, 1:vertical with spec correction, 2:vertical without spec correction
-template <uint32 Method>
-void SetSubsurfaceShaderSampleSet(const FRenderingCompositePassContext& Context, float InRadius, uint32 SampleSet)
+// 0:horizontal, 1: vertical
+template <uint32 Direction>
+void SetSubsurfaceShaderSampleSet(const FRenderingCompositePassContext& Context, TShaderMapRef<FPostProcessVS> &VertexShader, float InRadius, uint32 SampleSet)
 {
 	switch(SampleSet)
 	{
-		case 0: SetSubsurfaceShader<Method, 0>(Context, InRadius); break;
-		case 1: SetSubsurfaceShader<Method, 1>(Context, InRadius); break;
-		case 2: SetSubsurfaceShader<Method, 2>(Context, InRadius); break;
+		case 0: SetSubsurfaceShader<Direction, 0>(Context, VertexShader, InRadius); break;
+		case 1: SetSubsurfaceShader<Direction, 1>(Context, VertexShader, InRadius); break;
+		case 2: SetSubsurfaceShader<Direction, 2>(Context, VertexShader, InRadius); break;
 
 	default:
 		check(0);
@@ -452,7 +456,7 @@ void SetSubsurfaceShaderSampleSet(const FRenderingCompositePassContext& Context,
 
 void FRCPassPostProcessSubsurface::Process(FRenderingCompositePassContext& Context)
 {
-	const FPooledRenderTargetDesc* InputDesc = GetInputDesc(ePId_Input1);
+	const FPooledRenderTargetDesc* InputDesc = GetInputDesc(ePId_Input0);
 
 	check(InputDesc);
 
@@ -469,7 +473,12 @@ void FRCPassPostProcessSubsurface::Process(FRenderingCompositePassContext& Conte
 	const FSceneViewFamily& ViewFamily = *(View.Family);
 
 	FIntPoint SrcSize = InputDesc->Extent;
-	FIntPoint DestSize = (Pass == 0) ? PassOutputs[0].RenderTargetDesc.Extent : GSceneRenderTargets.GetBufferSizeXY();
+	FIntPoint DestSize = PassOutputs[0].RenderTargetDesc.Extent;
+
+	check(DestSize.X);
+	check(DestSize.Y);
+	check(SrcSize.X);
+	check(SrcSize.Y);
 
 	// e.g. 4 means the input texture is 4x smaller than the buffer size
 	uint32 ScaleFactor = GSceneRenderTargets.GetBufferSizeXY().X / SrcSize.X;
@@ -481,16 +490,7 @@ void FRCPassPostProcessSubsurface::Process(FRenderingCompositePassContext& Conte
 
 	const FSceneRenderTargetItem* DestRenderTarget;
 	{
-		if(Pass == 0)
-		{
-			DestRenderTarget = &PassOutputs[0].RequestSurface(Context);
-		}
-		else
-		{
-			// in that case we directly render to the HDR scene color
-			GRenderTargetPool.FindFreeElement(GSceneRenderTargets.GetSceneColor()->GetDesc(), NewSceneColor, TEXT("SceneColor"));
-			DestRenderTarget = &NewSceneColor->GetRenderTargetItem();
-		}
+		DestRenderTarget = &PassOutputs[0].RequestSurface(Context);
 
 		check(DestRenderTarget);
 	}
@@ -509,30 +509,17 @@ void FRCPassPostProcessSubsurface::Process(FRenderingCompositePassContext& Conte
 
 	TShaderMapRef<FPostProcessVS> VertexShader(Context.GetShaderMap());
 
-	const TCHAR* PassNames[] = { TEXT("SubsurfaceX"), TEXT("SubsurfaceY0"), TEXT("SubsurfaceY1") };
-	check(Pass < sizeof(PassNames) / sizeof(PassNames[0]));
-
-	SCOPED_DRAW_EVENTF(Context.RHICmdList, SubsurfacePass, DEC_SCENE_ITEMS, PassNames[Pass]);
+	SCOPED_DRAW_EVENTF(Context.RHICmdList, SubsurfacePass, DEC_SCENE_ITEMS, TEXT("SubsurfaceDirection#%d"), Direction);
 
 	uint32 SampleSet = FMath::Clamp(CVarSSSSampleSet.GetValueOnRenderThread(), 0, 2);
 
-	if (Pass == 0)
+	if (Direction == 0)
 	{
-		SetSubsurfaceShaderSampleSet<0>(Context, Radius, SampleSet);
+		SetSubsurfaceShaderSampleSet<0>(Context, VertexShader, Radius, SampleSet);
 	}
 	else
 	{
-		check(Pass == 1);
-
-		if(DoSpecularCorrection())
-		{
-			// reconstruct specular and add it in final pass
-			SetSubsurfaceShaderSampleSet<2>(Context, Radius, SampleSet);
-		}
-		else
-		{
-			SetSubsurfaceShaderSampleSet<1>(Context, Radius, SampleSet);
-		}
+		SetSubsurfaceShaderSampleSet<1>(Context, VertexShader, Radius, SampleSet);
 	}
 
 	DrawRectangle(
@@ -547,28 +534,16 @@ void FRCPassPostProcessSubsurface::Process(FRenderingCompositePassContext& Conte
 		EDRF_UseTriangleOptimization);
 
 	Context.RHICmdList.CopyToResolveTarget(DestRenderTarget->TargetableTexture, DestRenderTarget->ShaderResourceTexture, false, FResolveParams());
-
-	if(Pass != 0)
-	{
-		GSceneRenderTargets.SetSceneColor(NewSceneColor);
-	}
 }
 
 
 FPooledRenderTargetDesc FRCPassPostProcessSubsurface::ComputeOutputDesc(EPassOutputId InPassOutputId) const
 {
-	FPooledRenderTargetDesc Ret;
-
-	if(Pass != 0)
-	{
-		// in that case we directly render to the HDR scene color
-		return Ret;
-	}
-
-	Ret = PassInputs[1].GetOutput()->RenderTargetDesc;
+	FPooledRenderTargetDesc Ret = PassInputs[0].GetOutput()->RenderTargetDesc;
 
 	Ret.Reset();
-	Ret.DebugName = (Pass == 0) ? TEXT("SubsurfaceTemp") : TEXT("SceneColor");
+	Ret.DebugName = (Direction == 0) ? TEXT("SubsurfaceX") : TEXT("SubsurfaceY");
+
 	// the setup was done in half res but the actual sampling happens in full resolution
 	Ret.Extent = GSceneRenderTargets.GetBufferSizeXY();
 
@@ -581,4 +556,175 @@ FPooledRenderTargetDesc FRCPassPostProcessSubsurface::ComputeOutputDesc(EPassOut
 	}
 
 	return Ret;
-}
\ No newline at end of file
+}
+
+
+
+
+
+
+/** Encapsulates the post processing subsurafce recombine pixel shader. */
+// @param SpecularCorrection 0: reconstruct specular
+template <uint32 RecombineMethod>
+class TPostProcessSubsurfaceRecombinePS : public FGlobalShader
+{
+	DECLARE_SHADER_TYPE(TPostProcessSubsurfaceRecombinePS, Global);
+
+	static bool ShouldCache(EShaderPlatform Platform)
+	{
+		return SubsurfaceShouldCache(Platform);
+	}
+
+	static void ModifyCompilationEnvironment(EShaderPlatform Platform, FShaderCompilerEnvironment& OutEnvironment)
+	{
+		FGlobalShader::ModifyCompilationEnvironment(Platform,OutEnvironment);
+		OutEnvironment.SetDefine(TEXT("SSS_RECOMBINE_METHOD"), RecombineMethod);
+	}
+
+	/** Default constructor. */
+	TPostProcessSubsurfaceRecombinePS() {}
+
+public:
+	FPostProcessPassParameters PostprocessParameter;
+	FDeferredPixelShaderParameters DeferredParameters;
+
+	/** Initialization constructor. */
+	TPostProcessSubsurfaceRecombinePS(const ShaderMetaType::CompiledShaderInitializerType& Initializer)
+		: FGlobalShader(Initializer)
+	{
+		PostprocessParameter.Bind(Initializer.ParameterMap);
+		DeferredParameters.Bind(Initializer.ParameterMap);
+	}
+
+	// FShader interface.
+	virtual bool Serialize(FArchive& Ar)
+	{
+		bool bShaderHasOutdatedParameters = FGlobalShader::Serialize(Ar);
+		Ar << PostprocessParameter << DeferredParameters;
+		return bShaderHasOutdatedParameters;
+	}
+
+	void SetParameters(const FRenderingCompositePassContext& Context)
+	{
+		const FPixelShaderRHIParamRef ShaderRHI = GetPixelShader();
+
+		FGlobalShader::SetParameters(Context.RHICmdList, ShaderRHI, Context.View);
+		DeferredParameters.Set(Context.RHICmdList, ShaderRHI, Context.View);
+		PostprocessParameter.SetPS(ShaderRHI, Context, TStaticSamplerState<SF_Bilinear,AM_Border,AM_Border,AM_Border>::GetRHI());
+	}
+
+	static const TCHAR* GetSourceFilename()
+	{
+		return TEXT("PostProcessSubsurface");
+	}
+
+	static const TCHAR* GetFunctionName()
+	{
+		return TEXT("SubsurfaceRecombinePS");
+	}
+};
+
+// #define avoids a lot of code duplication
+#define VARIATION1(A) typedef TPostProcessSubsurfaceRecombinePS<A> TPostProcessSubsurfaceRecombinePS##A; \
+	IMPLEMENT_SHADER_TYPE2(TPostProcessSubsurfaceRecombinePS##A, SF_Pixel);
+	VARIATION1(0) VARIATION1(1)
+#undef VARIATION1
+
+
+template <uint32 RecombineMethod>
+void SetSubsurfaceRecombineShader(const FRenderingCompositePassContext& Context, TShaderMapRef<FPostProcessVS> &VertexShader)
+{
+	TShaderMapRef<TPostProcessSubsurfaceRecombinePS<RecombineMethod> > PixelShader(Context.GetShaderMap());
+
+	static FGlobalBoundShaderState BoundShaderState;
+
+	SetGlobalBoundShaderState(Context.RHICmdList, Context.GetFeatureLevel(), BoundShaderState, GFilterVertexDeclaration.VertexDeclarationRHI, *VertexShader, *PixelShader);
+
+	PixelShader->SetParameters(Context);
+	VertexShader->SetParameters(Context);
+}
+
+void FRCPassPostProcessSubsurfaceRecombine::Process(FRenderingCompositePassContext& Context)
+{
+	const FPooledRenderTargetDesc* InputDesc = GetInputDesc(ePId_Input0);
+
+	check(InputDesc);
+
+	const FSceneView& View = Context.View;
+	const FSceneViewFamily& ViewFamily = *(View.Family);
+
+	FIntPoint SrcSize = InputDesc->Extent;
+	FIntPoint DestSize = GSceneRenderTargets.GetBufferSizeXY();
+
+	check(DestSize.X);
+	check(DestSize.Y);
+	check(SrcSize.X);
+	check(SrcSize.Y);
+
+	// e.g. 4 means the input texture is 4x smaller than the buffer size
+	uint32 ScaleFactor = GSceneRenderTargets.GetBufferSizeXY().X / SrcSize.X;
+
+	FIntRect SrcRect = View.ViewRect / ScaleFactor;
+	FIntRect DestRect = View.ViewRect;
+
+	TRefCountPtr<IPooledRenderTarget> NewSceneColor;
+
+	const FSceneRenderTargetItem* DestRenderTarget;
+	{
+		// in that case we directly render to the HDR scene color
+		GRenderTargetPool.FindFreeElement(*GetInputDesc(ePId_Input1), NewSceneColor, TEXT("SceneColor"));
+		DestRenderTarget = &NewSceneColor->GetRenderTargetItem();
+	}
+
+	// Set the view family's render target/viewport.
+	SetRenderTarget(Context.RHICmdList, DestRenderTarget->TargetableTexture, FTextureRHIRef());
+
+	// is optimized away if possible (RT size=view size, )
+	Context.RHICmdList.Clear(true, FLinearColor(0, 0, 0, 0), false, 1.0f, false, 0, DestRect);
+
+	Context.SetViewportAndCallRHI(0, 0, 0.0f, DestSize.X, DestSize.Y, 1.0f );
+
+	Context.RHICmdList.SetBlendState(TStaticBlendState<>::GetRHI());
+	Context.RHICmdList.SetRasterizerState(TStaticRasterizerState<>::GetRHI());
+	Context.RHICmdList.SetDepthStencilState(TStaticDepthStencilState<false, CF_Always>::GetRHI());
+
+	TShaderMapRef<FPostProcessVS> VertexShader(Context.GetShaderMap());
+
+	bool bDoSpecularCorrection = DoSpecularCorrection();
+
+	SCOPED_DRAW_EVENTF(Context.RHICmdList, SubsurfacePass, DEC_SCENE_ITEMS, TEXT("SubsurfacePassRecombine#%d"), (int32)bDoSpecularCorrection);
+
+	if(bDoSpecularCorrection)
+	{
+		SetSubsurfaceRecombineShader<1>(Context, VertexShader);
+	}
+	else
+	{
+		SetSubsurfaceRecombineShader<0>(Context, VertexShader);
+	}
+
+	DrawRectangle(
+		Context.RHICmdList,
+		DestRect.Min.X, DestRect.Min.Y,
+		DestRect.Width(), DestRect.Height(),
+		SrcRect.Min.X, SrcRect.Min.Y,
+		SrcRect.Width(), SrcRect.Height(),
+		DestSize,
+		SrcSize,
+		*VertexShader,
+		EDRF_UseTriangleOptimization);
+
+	Context.RHICmdList.CopyToResolveTarget(DestRenderTarget->TargetableTexture, DestRenderTarget->ShaderResourceTexture, false, FResolveParams());
+
+	GSceneRenderTargets.SetSceneColor(NewSceneColor);
+}
+
+
+FPooledRenderTargetDesc FRCPassPostProcessSubsurfaceRecombine::ComputeOutputDesc(EPassOutputId InPassOutputId) const
+{
+	FPooledRenderTargetDesc Ret;
+
+	Ret.DebugName = TEXT("SceneColor");
+	// we directly render to the HDR scene color
+	return Ret;
+}
diff --git a/Engine/Source/Runtime/Renderer/Private/PostProcess/PostProcessSubsurface.h b/Engine/Source/Runtime/Renderer/Private/PostProcess/PostProcessSubsurface.h
--- a/Engine/Source/Runtime/Renderer/Private/PostProcess/PostProcessSubsurface.h
+++ b/Engine/Source/Runtime/Renderer/Private/PostProcess/PostProcessSubsurface.h
@@ -31,16 +31,15 @@ public:
 
 
 // derives from TRenderingCompositePassBase<InputCount, OutputCount>
-// ePId_Input0: optional SubsurfaceTemp input
-// ePId_Input1: SceneColor (horizontal blur) or the pass before (vertical blur)
+// ePId_Input0: SceneColor (horizontal blur) or the pass before (vertical blur)
 // modifies SceneColor, uses some GBuffer attributes
-class FRCPassPostProcessSubsurface : public TRenderingCompositePassBase<2, 1>
+class FRCPassPostProcessSubsurface : public TRenderingCompositePassBase<1, 1>
 {
 public:
 	// constructor
-	// @param Pass 0:horizontal/1:vertical
+	// @param InDirection 0:horizontal/1:vertical
 	// @param InRadius in pixels in the full res image
-	FRCPassPostProcessSubsurface(uint32 Pass, float InRadius, bool bInHalfRes = false);
+	FRCPassPostProcessSubsurface(uint32 InDirection, float InRadius, bool bInHalfRes = false);
 
 	// interface FRenderingCompositePass ---------
 
@@ -52,7 +51,22 @@ private:
 	// in pixels in the full res image
 	float Radius;
 	// 0:horizontal/1:vertical
-	uint32 Pass;
+	uint32 Direction;
 	// if the output should be half resolution to scene color
 	bool bHalfRes;
 };
+
+
+// derives from TRenderingCompositePassBase<InputCount, OutputCount>
+// ePId_Input0: output from FRCPassPostProcessSubsurface
+// ePId_Input1: SceneColor before Screen Space Subsurface input
+// modifies SceneColor, uses some GBuffer attributes
+class FRCPassPostProcessSubsurfaceRecombine : public TRenderingCompositePassBase<2, 1>
+{
+public:
+	// interface FRenderingCompositePass ---------
+
+	virtual void Process(FRenderingCompositePassContext& Context);
+	virtual void Release() override { delete this; }
+	virtual FPooledRenderTargetDesc ComputeOutputDesc(EPassOutputId InPassOutputId) const;
+};
\ No newline at end of file
diff --git a/Engine/Source/Runtime/Renderer/Private/PostProcess/RenderingCompositionGraph.h b/Engine/Source/Runtime/Renderer/Private/PostProcess/RenderingCompositionGraph.h
--- a/Engine/Source/Runtime/Renderer/Private/PostProcess/RenderingCompositionGraph.h
+++ b/Engine/Source/Runtime/Renderer/Private/PostProcess/RenderingCompositionGraph.h
@@ -163,7 +163,7 @@ struct FRenderingCompositePass
 	virtual FRenderingCompositeOutputRef* GetInput(EPassInputId InPassInputId) = 0;
 
 	/**
-	 * Each input is a dependency and will be processed before the node itself (don't generate cyles)
+	 * Each input is a dependency and will be processed before the node itself (don't generate cycles)
 	 * The index allows to access the input in Process() and on the shader side
 	 * @param InInputIndex silently ignores calls outside the range
 	 */
diff --git a/Engine/Source/Runtime/Renderer/Private/PostProcess/SceneRenderTargets.cpp b/Engine/Source/Runtime/Renderer/Private/PostProcess/SceneRenderTargets.cpp
--- a/Engine/Source/Runtime/Renderer/Private/PostProcess/SceneRenderTargets.cpp
+++ b/Engine/Source/Runtime/Renderer/Private/PostProcess/SceneRenderTargets.cpp
@@ -111,10 +111,12 @@ FIntPoint FSceneRenderTargets::GetSceneRenderTargetSize(const FSceneViewFamily&
 	enum ESizingMethods { RequestedSize, ScreenRes, Grow, VisibleSizingMethodsCount, Clamped};
 	ESizingMethods SceneTargetsSizingMethod = Grow;
 
-	bool bSceneCapture = false;
+	bool bCapture = false;
 	for (int32 ViewIndex = 0; ViewIndex < ViewFamily.Views.Num(); ++ViewIndex)
 	{
-		bSceneCapture |= ViewFamily.Views[ViewIndex]->bIsSceneCapture;
+		const FSceneView* View = ViewFamily.Views[ViewIndex];
+
+		bCapture |= View->bIsSceneCapture || View->bIsReflectionCapture;
 	}
 
 	if(!FPlatformProperties::SupportsWindowedMode())
@@ -133,7 +135,7 @@ FIntPoint FSceneRenderTargets::GetSceneRenderTargetSize(const FSceneViewFamily&
 		SceneTargetsSizingMethod = (ESizingMethods) FMath::Clamp(CVarSceneTargetsResizingMethod.GetValueOnRenderThread(), 0, (int32)VisibleSizingMethodsCount);
 	}
 	
-	if (bSceneCapture)
+	if (bCapture)
 	{
 		// In general, we don't want scenecapture to grow our buffers, because depending on the cvar for our game, we may not recover that memory.  This can be changed if necessary.
 		// However, in the editor a user might have a small editor window, but be capturing cubemaps or other dynamic assets for data distribution, 
