UE4 - Use perfect forwarding for task graph, plus some lambda work for shadow rendering

diff --git a/Engine/Source/Runtime/Core/Public/Async/TaskGraphInterfaces.h b/Engine/Source/Runtime/Core/Public/Async/TaskGraphInterfaces.h
--- a/Engine/Source/Runtime/Core/Public/Async/TaskGraphInterfaces.h
+++ b/Engine/Source/Runtime/Core/Public/Async/TaskGraphInterfaces.h
@@ -526,65 +526,65 @@ public:
 		}
 		/** Passthrough internal task constructor and dispatch. Note! Generally speaking references will not pass through; use pointers */
 		template<typename T>
-		FGraphEventRef ConstructAndDispatchWhenReady(const T& Arg)
+		FGraphEventRef ConstructAndDispatchWhenReady(T&& Arg1)
 		{
-			new ((void *)&Owner->TaskStorage) TTask(Arg);
+			new ((void *)&Owner->TaskStorage) TTask(Forward<T>(Arg1));
 			return Owner->Setup(Prerequisites, CurrentThreadIfKnown);
 		}
 		/** Passthrough internal task constructor and dispatch. Note! Generally speaking references will not pass through; use pointers */
 		template<typename T1,typename T2>
-		FGraphEventRef ConstructAndDispatchWhenReady(const T1& Arg1, const T2& Arg2)
+		FGraphEventRef ConstructAndDispatchWhenReady(T1&& Arg1, T2&& Arg2)
 		{
-			new ((void *)&Owner->TaskStorage) TTask(Arg1,Arg2);
+			new ((void *)&Owner->TaskStorage) TTask(Forward<T1>(Arg1), Forward<T2>(Arg2));
 			return Owner->Setup(Prerequisites, CurrentThreadIfKnown);
 		}
 		/** Passthrough internal task constructor and dispatch. Note! Generally speaking references will not pass through; use pointers */
 		template<typename T1,typename T2, typename T3>
-		FGraphEventRef ConstructAndDispatchWhenReady(const T1& Arg1, const T2& Arg2, const T3& Arg3)
+		FGraphEventRef ConstructAndDispatchWhenReady(T1&& Arg1, T2&& Arg2, T3&& Arg3)
 		{
-			new ((void *)&Owner->TaskStorage) TTask(Arg1,Arg2,Arg3);
+			new ((void *)&Owner->TaskStorage) TTask(Forward<T1>(Arg1), Forward<T2>(Arg2), Forward<T3>(Arg3));
 			return Owner->Setup(Prerequisites, CurrentThreadIfKnown);
 		}
 		/** Passthrough internal task constructor and dispatch. Note! Generally speaking references will not pass through; use pointers */
 		template<typename T1,typename T2, typename T3, typename T4>
-		FGraphEventRef ConstructAndDispatchWhenReady(const T1& Arg1, const T2& Arg2, const T3& Arg3, const T4& Arg4)
+		FGraphEventRef ConstructAndDispatchWhenReady(T1&& Arg1, T2&& Arg2, T3&& Arg3, T4&& Arg4)
 		{
-			new ((void *)&Owner->TaskStorage) TTask(Arg1,Arg2,Arg3,Arg4);
+			new ((void *)&Owner->TaskStorage) TTask(Forward<T1>(Arg1), Forward<T2>(Arg2), Forward<T3>(Arg3), Forward<T4>(Arg4));
 			return Owner->Setup(Prerequisites, CurrentThreadIfKnown);
 		}
 		/** Passthrough internal task constructor and dispatch. Note! Generally speaking references will not pass through; use pointers */
 		template<typename T1,typename T2, typename T3, typename T4, typename T5>
-		FGraphEventRef ConstructAndDispatchWhenReady(const T1& Arg1, const T2& Arg2, const T3& Arg3, const T4& Arg4, const T5& Arg5)
+		FGraphEventRef ConstructAndDispatchWhenReady(T1&& Arg1, T2&& Arg2, T3&& Arg3, T4&& Arg4, T5&& Arg5)
 		{
-			new ((void *)&Owner->TaskStorage) TTask(Arg1,Arg2,Arg3,Arg4,Arg5);
+			new ((void *)&Owner->TaskStorage) TTask(Forward<T1>(Arg1), Forward<T2>(Arg2), Forward<T3>(Arg3), Forward<T4>(Arg4), Forward<T5>(Arg5));
 			return Owner->Setup(Prerequisites, CurrentThreadIfKnown);
 		}
 		/** Passthrough internal task constructor and dispatch. Note! Generally speaking references will not pass through; use pointers */
 		template<typename T1, typename T2, typename T3, typename T4, typename T5, typename T6>
-		FGraphEventRef ConstructAndDispatchWhenReady(const T1& Arg1, const T2& Arg2, const T3& Arg3, const T4& Arg4, const T5& Arg5, const T6& Arg6)
+		FGraphEventRef ConstructAndDispatchWhenReady(T1&& Arg1, T2&& Arg2, T3&& Arg3, T4&& Arg4, T5&& Arg5, T6&& Arg6)
 		{
-			new ((void *)&Owner->TaskStorage) TTask(Arg1, Arg2, Arg3, Arg4, Arg5, Arg6);
+			new ((void *)&Owner->TaskStorage) TTask(Forward<T1>(Arg1), Forward<T2>(Arg2), Forward<T3>(Arg3), Forward<T4>(Arg4), Forward<T5>(Arg5), Forward<T6>(Arg6));
 			return Owner->Setup(Prerequisites, CurrentThreadIfKnown);
 		}
 		/** Passthrough internal task constructor and dispatch. Note! Generally speaking references will not pass through; use pointers */
 		template<typename T1, typename T2, typename T3, typename T4, typename T5, typename T6, typename T7>
-		FGraphEventRef ConstructAndDispatchWhenReady(const T1& Arg1, const T2& Arg2, const T3& Arg3, const T4& Arg4, const T5& Arg5, const T6& Arg6, const T7& Arg7)
+		FGraphEventRef ConstructAndDispatchWhenReady(T1&& Arg1, T2&& Arg2, T3&& Arg3, T4&& Arg4, T5&& Arg5, T6&& Arg6, T7&& Arg7)
 		{
-			new ((void *)&Owner->TaskStorage) TTask(Arg1, Arg2, Arg3, Arg4, Arg5, Arg6, Arg7);
+			new ((void *)&Owner->TaskStorage) TTask(Forward<T1>(Arg1), Forward<T2>(Arg2), Forward<T3>(Arg3), Forward<T4>(Arg4), Forward<T5>(Arg5), Forward<T6>(Arg6), Forward<T7>(Arg7));
 			return Owner->Setup(Prerequisites, CurrentThreadIfKnown);
 		}
 		/** Passthrough internal task constructor and dispatch. Note! Generally speaking references will not pass through; use pointers */
 		template<typename T1, typename T2, typename T3, typename T4, typename T5, typename T6, typename T7, typename T8>
-		FGraphEventRef ConstructAndDispatchWhenReady(const T1& Arg1, const T2& Arg2, const T3& Arg3, const T4& Arg4, const T5& Arg5, const T6& Arg6, const T7& Arg7, const T8& Arg8)
+		FGraphEventRef ConstructAndDispatchWhenReady(T1&& Arg1, T2&& Arg2, T3&& Arg3, T4&& Arg4, T5&& Arg5, T6&& Arg6, T7&& Arg7, T8&& Arg8)
 		{
-			new ((void *)&Owner->TaskStorage) TTask(Arg1, Arg2, Arg3, Arg4, Arg5, Arg6, Arg7, Arg8);
+			new ((void *)&Owner->TaskStorage) TTask(Forward<T1>(Arg1), Forward<T2>(Arg2), Forward<T3>(Arg3), Forward<T4>(Arg4), Forward<T5>(Arg5), Forward<T6>(Arg6), Forward<T7>(Arg7), Forward<T8>(Arg8));
 			return Owner->Setup(Prerequisites, CurrentThreadIfKnown);
 		}
 		/** Passthrough internal task constructor and dispatch. Note! Generally speaking references will not pass through; use pointers */
 		template<typename T1, typename T2, typename T3, typename T4, typename T5, typename T6, typename T7, typename T8, typename T9>
-		FGraphEventRef ConstructAndDispatchWhenReady(const T1& Arg1, const T2& Arg2, const T3& Arg3, const T4& Arg4, const T5& Arg5, const T6& Arg6, const T7& Arg7, const T8& Arg8, const T9& Arg9)
+		FGraphEventRef ConstructAndDispatchWhenReady(T1&& Arg1, T2&& Arg2, T3&& Arg3, T4&& Arg4, T5&& Arg5, T6&& Arg6, T7&& Arg7, T8&& Arg8, T9&& Arg9)
 		{
-			new ((void *)&Owner->TaskStorage) TTask(Arg1, Arg2, Arg3, Arg4, Arg5, Arg6, Arg7, Arg8, Arg9);
+			new ((void *)&Owner->TaskStorage) TTask(Forward<T1>(Arg1), Forward<T2>(Arg2), Forward<T3>(Arg3), Forward<T4>(Arg4), Forward<T5>(Arg5), Forward<T6>(Arg6), Forward<T7>(Arg7), Forward<T8>(Arg8), Forward<T9>(Arg9));
 			return Owner->Setup(Prerequisites, CurrentThreadIfKnown);
 		}
 	private:
@@ -976,7 +976,7 @@ public:
 	  *	@param StatId The stat id for this task.
 	  * @param InDesiredThread - Thread to run on
 	**/
-	FSimpleDelegateGraphTask(const FDelegate& InTaskDeletegate, const TStatId& StatId, ENamedThreads::Type InDesiredThread)
+	FSimpleDelegateGraphTask(const FDelegate& InTaskDeletegate, const TStatId StatId, ENamedThreads::Type InDesiredThread)
 		: FCustomStatIDGraphTaskBase(StatId)
 		, TaskDelegate(InTaskDeletegate)
 		, DesiredThread(InDesiredThread)
@@ -991,7 +991,7 @@ public:
 	  * @param InDesiredThread - Thread to run on
 	  * @return completion handle for the new task 
 	**/
-	static FGraphEventRef CreateAndDispatchWhenReady(const FDelegate& InTaskDeletegate, const TStatId& InStatId, const FGraphEventArray* InPrerequisites = NULL, ENamedThreads::Type InDesiredThread = ENamedThreads::AnyThread)
+	static FGraphEventRef CreateAndDispatchWhenReady(const FDelegate& InTaskDeletegate, const TStatId InStatId, const FGraphEventArray* InPrerequisites = NULL, ENamedThreads::Type InDesiredThread = ENamedThreads::AnyThread)
 	{
 		return TGraphTask<FSimpleDelegateGraphTask>::CreateTask(InPrerequisites).ConstructAndDispatchWhenReady<const FDelegate&>(InTaskDeletegate, InStatId, InDesiredThread);
 	}
@@ -1003,7 +1003,7 @@ public:
 	  * @param InDesiredThread - Thread to run on
 	  * @return completion handle for the new task 
 	**/
-	static FGraphEventRef CreateAndDispatchWhenReady(const FDelegate& InTaskDeletegate, const TStatId& InStatId, const FGraphEventRef& InPrerequisite, ENamedThreads::Type InDesiredThread = ENamedThreads::AnyThread)
+	static FGraphEventRef CreateAndDispatchWhenReady(const FDelegate& InTaskDeletegate, const TStatId&& InStatId, const FGraphEventRef& InPrerequisite, ENamedThreads::Type InDesiredThread = ENamedThreads::AnyThread)
 	{
 		FGraphEventArray Prerequisites;
 		check(InPrerequisite.GetReference());
@@ -1041,7 +1041,7 @@ public:
 	  *	@param InStatId - The stat id for this task.
 	  * @param InDesiredThread - Thread to run on
 	**/
-	FDelegateGraphTask(const FDelegate& InTaskDeletegate, const TStatId& InStatId, ENamedThreads::Type InDesiredThread)
+	FDelegateGraphTask(const FDelegate& InTaskDeletegate, const TStatId InStatId, ENamedThreads::Type InDesiredThread)
 		: FCustomStatIDGraphTaskBase(InStatId)
 		, TaskDelegate(InTaskDeletegate)
 		, DesiredThread(InDesiredThread)
@@ -1057,7 +1057,7 @@ public:
 	  * @param InDesiredThread - Thread to run on
 	  * @return completion handle for the new task 
 	**/
-	static FGraphEventRef CreateAndDispatchWhenReady(const FDelegate& InTaskDeletegate, const TStatId& InStatId, const FGraphEventArray* InPrerequisites = NULL, ENamedThreads::Type InCurrentThreadIfKnown = ENamedThreads::AnyThread, ENamedThreads::Type InDesiredThread = ENamedThreads::AnyThread)
+	static FGraphEventRef CreateAndDispatchWhenReady(const FDelegate& InTaskDeletegate, const TStatId InStatId, const FGraphEventArray* InPrerequisites = NULL, ENamedThreads::Type InCurrentThreadIfKnown = ENamedThreads::AnyThread, ENamedThreads::Type InDesiredThread = ENamedThreads::AnyThread)
 	{
 		return TGraphTask<FDelegateGraphTask>::CreateTask(InPrerequisites, InCurrentThreadIfKnown).ConstructAndDispatchWhenReady<const FDelegate&>(InTaskDeletegate, InStatId, InDesiredThread);
 	}
@@ -1070,7 +1070,7 @@ public:
 	  * @param InDesiredThread - Thread to run on
 	  * @return completion handle for the new task 
 	**/
-	static FGraphEventRef CreateAndDispatchWhenReady(const FDelegate& InTaskDeletegate, const TStatId& InStatId, const FGraphEventRef& InPrerequisite, ENamedThreads::Type InCurrentThreadIfKnown = ENamedThreads::AnyThread, ENamedThreads::Type InDesiredThread = ENamedThreads::AnyThread)
+	static FGraphEventRef CreateAndDispatchWhenReady(const FDelegate& InTaskDeletegate, const TStatId InStatId, const FGraphEventRef& InPrerequisite, ENamedThreads::Type InCurrentThreadIfKnown = ENamedThreads::AnyThread, ENamedThreads::Type InDesiredThread = ENamedThreads::AnyThread)
 	{
 		FGraphEventArray Prerequisites;
 		check(InPrerequisite.GetReference());
diff --git a/Engine/Source/Runtime/Engine/Private/StaticBoundShaderState.cpp b/Engine/Source/Runtime/Engine/Private/StaticBoundShaderState.cpp
--- a/Engine/Source/Runtime/Engine/Private/StaticBoundShaderState.cpp
+++ b/Engine/Source/Runtime/Engine/Private/StaticBoundShaderState.cpp
@@ -150,9 +150,9 @@ class FSetGlobalBoundShaderStateRenderThreadTask
 
 public:
 
-	FSetGlobalBoundShaderStateRenderThreadTask(FRHICommandList* InRHICmdList, FGlobalBoundShaderState* InGlobalBoundShaderState, ERHIFeatureLevel::Type InFeatureLevel)
-		: RHICmdList(*InRHICmdList)
-		, GlobalBoundShaderState(*InGlobalBoundShaderState)
+	FSetGlobalBoundShaderStateRenderThreadTask(FRHICommandList& InRHICmdList, FGlobalBoundShaderState& InGlobalBoundShaderState, ERHIFeatureLevel::Type InFeatureLevel)
+		: RHICmdList(InRHICmdList)
+		, GlobalBoundShaderState(InGlobalBoundShaderState)
 		, FeatureLevel(InFeatureLevel)
 	{
 	}
@@ -241,7 +241,7 @@ void SetGlobalBoundShaderState(
 	// We need to do this on the render thread
 
 	FRHICommandList* CmdList = new FRHICommandList;
-	FGraphEventRef RenderThreadCompletionEvent = TGraphTask<FSetGlobalBoundShaderStateRenderThreadTask>::CreateTask().ConstructAndDispatchWhenReady(CmdList, &GlobalBoundShaderState, FeatureLevel);
+	FGraphEventRef RenderThreadCompletionEvent = TGraphTask<FSetGlobalBoundShaderStateRenderThreadTask>::CreateTask().ConstructAndDispatchWhenReady(*CmdList, GlobalBoundShaderState, FeatureLevel);
 	RHICmdList.QueueRenderThreadCommandListSubmit(RenderThreadCompletionEvent, CmdList);
 }
 
diff --git a/Engine/Source/Runtime/RHI/Public/RHIStaticStates.h b/Engine/Source/Runtime/RHI/Public/RHIStaticStates.h
--- a/Engine/Source/Runtime/RHI/Public/RHIStaticStates.h
+++ b/Engine/Source/Runtime/RHI/Public/RHIStaticStates.h
@@ -20,9 +20,9 @@ class FInitStaticResourceRenderThreadTask
 	FScopedEvent& Event;
 public:
 
-	FInitStaticResourceRenderThreadTask(void (*InDoConstruct)(), FScopedEvent* InEvent)
+	FInitStaticResourceRenderThreadTask(void (*InDoConstruct)(), FScopedEvent& InEvent)
 		: DoConstruct(InDoConstruct)
-		, Event(*InEvent)
+		, Event(InEvent)
 	{
 	}
 
@@ -77,7 +77,7 @@ public:
 				check(IsInParallelRenderingThread());
 				{
 					FScopedEvent Event;
-					TGraphTask<FInitStaticResourceRenderThreadTask>::CreateTask().ConstructAndDispatchWhenReady(&GetRHI_WithNoReturnValue, &Event);
+					TGraphTask<FInitStaticResourceRenderThreadTask>::CreateTask().ConstructAndDispatchWhenReady(&GetRHI_WithNoReturnValue, Event);
 				}
 			}
 			else
diff --git a/Engine/Source/Runtime/RenderCore/Private/RenderResource.cpp b/Engine/Source/Runtime/RenderCore/Private/RenderResource.cpp
--- a/Engine/Source/Runtime/RenderCore/Private/RenderResource.cpp
+++ b/Engine/Source/Runtime/RenderCore/Private/RenderResource.cpp
@@ -78,9 +78,9 @@ void FRenderResource::InitResourceFromPossiblyParallelRendering()
 			FScopedEvent& Event;
 		public:
 
-			FInitResourceRenderThreadTask(FRenderResource* InResource, FScopedEvent* InEvent)
-				: Resource(*InResource)
-				, Event(*InEvent)
+			FInitResourceRenderThreadTask(FRenderResource& InResource, FScopedEvent& InEvent)
+				: Resource(InResource)
+				, Event(InEvent)
 			{
 			}
 
@@ -104,7 +104,7 @@ void FRenderResource::InitResourceFromPossiblyParallelRendering()
 		};
 		{
 			FScopedEvent Event;
-			TGraphTask<FInitResourceRenderThreadTask>::CreateTask().ConstructAndDispatchWhenReady(this, &Event);
+			TGraphTask<FInitResourceRenderThreadTask>::CreateTask().ConstructAndDispatchWhenReady(*this, Event);
 		}
 	}
 }
diff --git a/Engine/Source/Runtime/RenderCore/Public/RenderingThread.h b/Engine/Source/Runtime/RenderCore/Public/RenderingThread.h
--- a/Engine/Source/Runtime/RenderCore/Public/RenderingThread.h
+++ b/Engine/Source/Runtime/RenderCore/Public/RenderingThread.h
@@ -224,7 +224,7 @@ DECLARE_STATS_GROUP(TEXT("Render Thread Commands"), STATGROUP_RenderThreadComman
 		if(GIsThreadedRendering || !IsInGameThread()) \
 		{ \
 			CheckNotBlockedOnRenderThread(); \
-			TGraphTask<EURCMacro_##TypeName>::CreateTask().template ConstructAndDispatchWhenReady<ParamType1>(ParamValue1); \
+			TGraphTask<EURCMacro_##TypeName>::CreateTask().ConstructAndDispatchWhenReady(ParamValue1); \
 		} \
 		else \
 		{ \
@@ -270,7 +270,7 @@ DECLARE_STATS_GROUP(TEXT("Render Thread Commands"), STATGROUP_RenderThreadComman
 		if(GIsThreadedRendering || !IsInGameThread()) \
 		{ \
 			CheckNotBlockedOnRenderThread(); \
-			TGraphTask<EURCMacro_##TypeName>::CreateTask().template ConstructAndDispatchWhenReady<ParamType1,ParamType2>(ParamValue1,ParamValue2); \
+			TGraphTask<EURCMacro_##TypeName>::CreateTask().ConstructAndDispatchWhenReady(ParamValue1,ParamValue2); \
 		} \
 		else \
 		{ \
@@ -318,7 +318,7 @@ DECLARE_STATS_GROUP(TEXT("Render Thread Commands"), STATGROUP_RenderThreadComman
 		if(GIsThreadedRendering || !IsInGameThread()) \
 		{ \
 			CheckNotBlockedOnRenderThread(); \
-			TGraphTask<EURCMacro_##TypeName>::CreateTask().template ConstructAndDispatchWhenReady<ParamType1,ParamType2,ParamType3>(ParamValue1,ParamValue2,ParamValue3); \
+			TGraphTask<EURCMacro_##TypeName>::CreateTask().ConstructAndDispatchWhenReady(ParamValue1,ParamValue2,ParamValue3); \
 		} \
 		else \
 		{ \
@@ -369,7 +369,7 @@ DECLARE_STATS_GROUP(TEXT("Render Thread Commands"), STATGROUP_RenderThreadComman
 		if(GIsThreadedRendering || !IsInGameThread()) \
 		{ \
 			CheckNotBlockedOnRenderThread(); \
-			TGraphTask<EURCMacro_##TypeName>::CreateTask().template ConstructAndDispatchWhenReady<ParamType1,ParamType2,ParamType3,ParamType4>(ParamValue1,ParamValue2,ParamValue3,ParamValue4); \
+			TGraphTask<EURCMacro_##TypeName>::CreateTask().ConstructAndDispatchWhenReady(ParamValue1,ParamValue2,ParamValue3,ParamValue4); \
 		} \
 		else \
 		{ \
diff --git a/Engine/Source/Runtime/Renderer/Private/DeferredShadingRenderer.cpp b/Engine/Source/Runtime/Renderer/Private/DeferredShadingRenderer.cpp
--- a/Engine/Source/Runtime/Renderer/Private/DeferredShadingRenderer.cpp
+++ b/Engine/Source/Runtime/Renderer/Private/DeferredShadingRenderer.cpp
@@ -460,15 +460,15 @@ class FRenderBasePassDynamicDataThreadTask
 public:
 
 	FRenderBasePassDynamicDataThreadTask(
-		FDeferredShadingSceneRenderer* InThisRenderer,
-		FRHICommandList* InRHICmdList,
-		const FViewInfo* InView,
-		bool* InOutDirty
+		FDeferredShadingSceneRenderer& InThisRenderer,
+		FRHICommandList& InRHICmdList,
+		const FViewInfo& InView,
+		bool& InOutDirty
 		)
-		: ThisRenderer(*InThisRenderer)
-		, RHICmdList(*InRHICmdList)
-		, View(*InView)
-		, OutDirty(*InOutDirty)
+		: ThisRenderer(InThisRenderer)
+		, RHICmdList(InRHICmdList)
+		, View(InView)
+		, OutDirty(InOutDirty)
 	{
 	}
 
@@ -495,7 +495,7 @@ void FDeferredShadingSceneRenderer::RenderBasePassDynamicDataParallel(FParallelC
 {
 	FRHICommandList* CmdList = ParallelCommandListSet.NewParallelCommandList();
 	FGraphEventRef AnyThreadCompletionEvent = TGraphTask<FRenderBasePassDynamicDataThreadTask>::CreateTask(nullptr, ENamedThreads::RenderThread)
-		.ConstructAndDispatchWhenReady(this, CmdList, &ParallelCommandListSet.View, &ParallelCommandListSet.OutDirty);
+		.ConstructAndDispatchWhenReady(*this, *CmdList, ParallelCommandListSet.View, ParallelCommandListSet.OutDirty);
 
 	ParallelCommandListSet.AddParallelCommandList(CmdList, AnyThreadCompletionEvent);
 }
@@ -1179,15 +1179,15 @@ class FRenderPrepassDynamicDataThreadTask
 public:
 
 	FRenderPrepassDynamicDataThreadTask(
-		FDeferredShadingSceneRenderer* InThisRenderer,
-		FRHICommandList* InRHICmdList,
-		const FViewInfo* InView,
-		bool* InOutDirty
+		FDeferredShadingSceneRenderer& InThisRenderer,
+		FRHICommandList& InRHICmdList,
+		const FViewInfo& InView,
+		bool& InOutDirty
 		)
-		: ThisRenderer(*InThisRenderer)
-		, RHICmdList(*InRHICmdList)
-		, View(*InView)
-		, OutDirty(*InOutDirty)
+		: ThisRenderer(InThisRenderer)
+		, RHICmdList(InRHICmdList)
+		, View(InView)
+		, OutDirty(InOutDirty)
 	{
 	}
 
@@ -1265,7 +1265,7 @@ void FDeferredShadingSceneRenderer::RenderPrePassViewParallel(const FViewInfo& V
 		FRHICommandList* CmdList = ParallelCommandListSet.NewParallelCommandList();
 
 		FGraphEventRef AnyThreadCompletionEvent = TGraphTask<FRenderPrepassDynamicDataThreadTask>::CreateTask(nullptr, ENamedThreads::RenderThread)
-			.ConstructAndDispatchWhenReady(this, CmdList, &View, &OutDirty);
+			.ConstructAndDispatchWhenReady(*this, *CmdList, View, OutDirty);
 
 		ParallelCommandListSet.AddParallelCommandList(CmdList, AnyThreadCompletionEvent);
 	}
diff --git a/Engine/Source/Runtime/Renderer/Private/DynamicPrimitiveDrawing.inl b/Engine/Source/Runtime/Renderer/Private/DynamicPrimitiveDrawing.inl
--- a/Engine/Source/Runtime/Renderer/Private/DynamicPrimitiveDrawing.inl
+++ b/Engine/Source/Runtime/Renderer/Private/DynamicPrimitiveDrawing.inl
@@ -270,16 +270,16 @@ class FDrawViewElementsAnyThreadTask
 public:
 
 	FDrawViewElementsAnyThreadTask(
-		FRHICommandList* InRHICmdList,
-		const FViewInfo* InView,
+		FRHICommandList& InRHICmdList,
+		const FViewInfo& InView,
 		const typename DrawingPolicyFactoryType::ContextType& InDrawingContext,
 		uint8 InDPGIndex,
 		bool InbPreFog,
 		int32 InFirstIndex,
 		int32 InLastIndex
 		)
-		: RHICmdList(*InRHICmdList)
-		, View(*InView)
+		: RHICmdList(InRHICmdList)
+		, View(InView)
 		, DrawingContext(InDrawingContext)
 		, DPGIndex(InDPGIndex)
 		, bPreFog(InbPreFog)
@@ -339,7 +339,7 @@ void DrawViewElementsParallel(
 					FRHICommandList* CmdList = ParallelCommandListSet.NewParallelCommandList();
 
 					FGraphEventRef AnyThreadCompletionEvent = TGraphTask<FDrawViewElementsAnyThreadTask<DrawingPolicyFactoryType> >::CreateTask(nullptr, ENamedThreads::RenderThread)
-						.ConstructAndDispatchWhenReady(CmdList, &ParallelCommandListSet.View, DrawingContext, DPGIndex, bPreFog, Start, Last);
+						.ConstructAndDispatchWhenReady(*CmdList, ParallelCommandListSet.View, DrawingContext, DPGIndex, bPreFog, Start, Last);
 
 					ParallelCommandListSet.AddParallelCommandList(CmdList, AnyThreadCompletionEvent);
 				}
diff --git a/Engine/Source/Runtime/Renderer/Private/SceneVisibility.cpp b/Engine/Source/Runtime/Renderer/Private/SceneVisibility.cpp
--- a/Engine/Source/Runtime/Renderer/Private/SceneVisibility.cpp
+++ b/Engine/Source/Runtime/Renderer/Private/SceneVisibility.cpp
@@ -1117,8 +1117,8 @@ class FRelevancePacketAnyThreadTask
 	ENamedThreads::Type ThreadToUse;
 public:
 
-	FRelevancePacketAnyThreadTask(FRelevancePacket* InPacket, ENamedThreads::Type InThreadToUse)
-		: Packet(*InPacket)
+	FRelevancePacketAnyThreadTask(FRelevancePacket& InPacket, ENamedThreads::Type InThreadToUse)
+		: Packet(InPacket)
 		, ThreadToUse(InThreadToUse)
 	{
 	}
@@ -1146,8 +1146,8 @@ class FRelevancePacketRenderThreadTask
 	FRelevancePacket& Packet;
 public:
 
-	FRelevancePacketRenderThreadTask(FRelevancePacket* InPacket)
-		: Packet(*InPacket)
+	FRelevancePacketRenderThreadTask(FRelevancePacket& InPacket)
+		: Packet(InPacket)
 	{
 	}
 
@@ -1258,9 +1258,9 @@ static void ComputeAndMarkRelevanceForViewParallel(
 					{
 						RenderPrereqs.Add(LastRenderThread); // this puts the render thread ones in order
 					}
-					FGraphEventRef AnyThread = TGraphTask<FRelevancePacketAnyThreadTask>::CreateTask(nullptr, ENamedThreads::RenderThread).ConstructAndDispatchWhenReady(Packet, ThreadToUse);
+					FGraphEventRef AnyThread = TGraphTask<FRelevancePacketAnyThreadTask>::CreateTask(nullptr, ENamedThreads::RenderThread).ConstructAndDispatchWhenReady(*Packet, ThreadToUse);
 					RenderPrereqs.Add(AnyThread);
-					LastRenderThread = TGraphTask<FRelevancePacketRenderThreadTask>::CreateTask(&RenderPrereqs, ENamedThreads::RenderThread).ConstructAndDispatchWhenReady(Packet);
+					LastRenderThread = TGraphTask<FRelevancePacketRenderThreadTask>::CreateTask(&RenderPrereqs, ENamedThreads::RenderThread).ConstructAndDispatchWhenReady(*Packet);
 					if (!BitIt)
 					{
 						break;
diff --git a/Engine/Source/Runtime/Renderer/Private/ShadowRendering.cpp b/Engine/Source/Runtime/Renderer/Private/ShadowRendering.cpp
--- a/Engine/Source/Runtime/Renderer/Private/ShadowRendering.cpp
+++ b/Engine/Source/Runtime/Renderer/Private/ShadowRendering.cpp
@@ -1361,14 +1361,14 @@ class FDrawShadowMeshElementsThreadTask
 public:
 
 	FDrawShadowMeshElementsThreadTask(
-		FProjectedShadowInfo* InThisRenderer,
-		FRHICommandList* InRHICmdList,
-		const FViewInfo* InView,
+		FProjectedShadowInfo& InThisRenderer,
+		FRHICommandList& InRHICmdList,
+		const FViewInfo& InView,
 		bool InbReflective
 		)
-		: ThisRenderer(*InThisRenderer)
-		, RHICmdList(*InRHICmdList)
-		, View(*InView)
+		: ThisRenderer(InThisRenderer)
+		, RHICmdList(InRHICmdList)
+		, View(InView)
 		, bReflective(InbReflective)
 	{
 	}
@@ -1413,14 +1413,14 @@ class FRenderDepthDynamicThreadTask
 public:
 
 	FRenderDepthDynamicThreadTask(
-		FProjectedShadowInfo* InThisRenderer,
-		FRHICommandList* InRHICmdList,
-		const FViewInfo* InView,
+		FProjectedShadowInfo& InThisRenderer,
+		FRHICommandList& InRHICmdList,
+		const FViewInfo& InView,
 		FSceneRenderer* InSceneRenderer
 		)
-		: ThisRenderer(*InThisRenderer)
-		, RHICmdList(*InRHICmdList)
-		, View(*InView)
+		: ThisRenderer(InThisRenderer)
+		, RHICmdList(InRHICmdList)
+		, View(InView)
 		, SceneRenderer(InSceneRenderer)
 	{
 	}
@@ -1517,9 +1517,9 @@ static TAutoConsoleVariable<int32> CVarRHICmdShadowDeferredContexts(
 class FShadowParallelCommandListSet : public FParallelCommandListSet
 {
 	FProjectedShadowInfo& ProjectedShadowInfo;
-	ISetShadowRenderTargets& SetShadowRenderTargets;
+	IFunction<void, FRHICommandList&>& SetShadowRenderTargets;
 public:
-	FShadowParallelCommandListSet(const FViewInfo& InView, FRHICommandList& InParentCmdList, bool* InOutDirty, bool bInParallelExecute, FProjectedShadowInfo& InProjectedShadowInfo, ISetShadowRenderTargets& InSetShadowRenderTargets)
+	FShadowParallelCommandListSet(const FViewInfo& InView, FRHICommandList& InParentCmdList, bool* InOutDirty, bool bInParallelExecute, FProjectedShadowInfo& InProjectedShadowInfo, IFunction<void, FRHICommandList&>& InSetShadowRenderTargets)
 		: FParallelCommandListSet(InView, InParentCmdList, InOutDirty, bInParallelExecute)
 		, ProjectedShadowInfo(InProjectedShadowInfo)
 		, SetShadowRenderTargets(InSetShadowRenderTargets)
@@ -1528,12 +1528,12 @@ public:
 	}
 	virtual void SetStateOnCommandList(FRHICommandList& CmdList) override
 	{
-		SetShadowRenderTargets.SetRenderTargetsForShadows(CmdList);
+		SetShadowRenderTargets(CmdList);
 		ProjectedShadowInfo.SetStateForDepth(CmdList);
 	}
 };
 
-void FProjectedShadowInfo::RenderDepthInner(FRHICommandList& RHICmdList, FSceneRenderer* SceneRenderer, const FViewInfo* FoundView, ISetShadowRenderTargets& SetShadowRenderTargets)
+void FProjectedShadowInfo::RenderDepthInner(FRHICommandList& RHICmdList, FSceneRenderer* SceneRenderer, const FViewInfo* FoundView, IFunction<void, FRHICommandList&>& SetShadowRenderTargets)
 {
 	FShadowDepthDrawingPolicyContext PolicyContext(this);
 
@@ -1567,7 +1567,7 @@ void FProjectedShadowInfo::RenderDepthInner(FRHICommandList& RHICmdList, FSceneR
 				FRHICommandList* CmdList = ParallelCommandListSet.NewParallelCommandList();
 
 				FGraphEventRef AnyThreadCompletionEvent = TGraphTask<FDrawShadowMeshElementsThreadTask>::CreateTask(nullptr, ENamedThreads::RenderThread)
-					.ConstructAndDispatchWhenReady(this, CmdList, FoundView, bReflectiveShadowmap && !bOnePassPointLightShadow);
+					.ConstructAndDispatchWhenReady(*this, *CmdList, *FoundView, bReflectiveShadowmap && !bOnePassPointLightShadow);
 
 				ParallelCommandListSet.AddParallelCommandList(CmdList, AnyThreadCompletionEvent);
 			}
@@ -1575,7 +1575,7 @@ void FProjectedShadowInfo::RenderDepthInner(FRHICommandList& RHICmdList, FSceneR
 				FRHICommandList* CmdList = ParallelCommandListSet.NewParallelCommandList();
 
 				FGraphEventRef AnyThreadCompletionEvent = TGraphTask<FRenderDepthDynamicThreadTask>::CreateTask(nullptr, ENamedThreads::RenderThread)
-					.ConstructAndDispatchWhenReady(this, CmdList, FoundView, SceneRenderer);
+					.ConstructAndDispatchWhenReady(*this, *CmdList, *FoundView, SceneRenderer);
 
 				ParallelCommandListSet.AddParallelCommandList(CmdList, AnyThreadCompletionEvent);
 			}
@@ -1621,7 +1621,7 @@ void FProjectedShadowInfo::RenderDepthInner(FRHICommandList& RHICmdList, FSceneR
 	}
 }
 
-void FProjectedShadowInfo::RenderDepth(FRHICommandList& RHICmdList, FSceneRenderer* SceneRenderer, ISetShadowRenderTargets& SetShadowRenderTargets)
+void FProjectedShadowInfo::RenderDepth(FRHICommandList& RHICmdList, FSceneRenderer* SceneRenderer, IFunction<void, FRHICommandList&>& SetShadowRenderTargets)
 {
 #if WANTS_DRAW_MESH_EVENTS
 	FString EventName;
@@ -2617,19 +2617,16 @@ bool FDeferredShadingSceneRenderer::RenderOnePassPointLightShadows(FRHICommandLi
 			if (!ProjectedShadowInfo->bRayTracedDistanceFieldShadow)
 			{
 				SCOPED_DRAW_EVENT(RHICmdList, ShadowDepthsFromOpaque, DEC_SCENE_ITEMS);
-				GSceneRenderTargets.BeginRenderingCubeShadowDepth(RHICmdList, ProjectedShadowInfo->ResolutionX);
+				auto SetShadowRenderTargets = MakeFunction<void, FRHICommandList&>([ProjectedShadowInfo](FRHICommandList& RHICmdList)
+				{
+					GSceneRenderTargets.BeginRenderingCubeShadowDepth(RHICmdList, ProjectedShadowInfo->ResolutionX);
+				});
+
+				SetShadowRenderTargets(RHICmdList);  // run it now, maybe run it later for parallel command lists
 				ProjectedShadowInfo->bAllocated = true;
 				ProjectedShadowInfo->ClearDepth(RHICmdList, this);
-				struct FLocalShadowTargets : public ISetShadowRenderTargets
-				{
-					uint32 ResolutionX;
-					virtual void SetRenderTargetsForShadows(FRHICommandList& RHICmdList)
-					{
-						GSceneRenderTargets.BeginRenderingCubeShadowDepth(RHICmdList, ResolutionX);
-					}
-				} LocalShadowTargets;
-				LocalShadowTargets.ResolutionX = ProjectedShadowInfo->ResolutionX;
-				ProjectedShadowInfo->RenderDepth(RHICmdList, this, LocalShadowTargets);
+
+				ProjectedShadowInfo->RenderDepth(RHICmdList, this, SetShadowRenderTargets);
 				GSceneRenderTargets.FinishRenderingCubeShadowDepth(RHICmdList, ProjectedShadowInfo->ResolutionX);
 			}
 
@@ -3030,23 +3027,20 @@ bool FDeferredShadingSceneRenderer::RenderReflectiveShadowMaps(FRHICommandListIm
 				FLightPropagationVolume* LightPropagationVolume = ViewState->GetLightPropagationVolume();
 
 				check(LightPropagationVolume);
-				GSceneRenderTargets.BeginRenderingReflectiveShadowMap(RHICmdList, LightPropagationVolume);
+				auto SetShadowRenderTargets = MakeFunction<void, FRHICommandList&>([LightPropagationVolume](FRHICommandList& RHICmdList)
+				{
+					GSceneRenderTargets.BeginRenderingReflectiveShadowMap(RHICmdList, LightPropagationVolume);
+				});
+
+				SetShadowRenderTargets(RHICmdList);  // run it now, maybe run it later for parallel command lists
+
 
 				LightPropagationVolume->SetVplInjectionConstants( *ProjectedShadowInfo, LightSceneInfo->Proxy );
 
 				if (ProjectedShadowInfo->bAllocated && !ProjectedShadowInfo->bTranslucentShadow)
 				{
 					ProjectedShadowInfo->ClearDepth(RHICmdList, this);
-					struct FLocalShadowTargets : public ISetShadowRenderTargets
-					{
-						FLightPropagationVolume* LightPropagationVolume;
-						virtual void SetRenderTargetsForShadows(FRHICommandList& RHICmdList)
-						{
-							GSceneRenderTargets.BeginRenderingReflectiveShadowMap(RHICmdList, LightPropagationVolume);
-						}
-					} LocalShadowTargets;
-					LocalShadowTargets.LightPropagationVolume = LightPropagationVolume;
-					ProjectedShadowInfo->RenderDepth(RHICmdList, this, LocalShadowTargets);
+					ProjectedShadowInfo->RenderDepth(RHICmdList, this, SetShadowRenderTargets);
 				}
 
 				GSceneRenderTargets.FinishRenderingReflectiveShadowMap(RHICmdList);
@@ -3210,14 +3204,14 @@ bool FDeferredShadingSceneRenderer::RenderProjectedShadows(FRHICommandListImmedi
 			// Render the shadow depths.
 			SCOPED_DRAW_EVENT(RHICmdList, ShadowDepthsFromOpaque, DEC_SCENE_ITEMS);
 
-			GSceneRenderTargets.BeginRenderingShadowDepth(RHICmdList, true);
-			struct FLocalShadowTargets : public ISetShadowRenderTargets
+			bool bFirst = true;
+			auto SetShadowRenderTargets = MakeFunction<void, FRHICommandList&>([&bFirst](FRHICommandList& RHICmdList)
 			{
-				virtual void SetRenderTargetsForShadows(FRHICommandList& RHICmdList)
-				{
-					GSceneRenderTargets.BeginRenderingShadowDepth(RHICmdList, false);
-				}
-			} LocalShadowTargets;
+				GSceneRenderTargets.BeginRenderingShadowDepth(RHICmdList, bFirst);
+			});
+
+			SetShadowRenderTargets(RHICmdList);  // run it now, maybe run it later for parallel command lists
+			bFirst = false;
 
 			// render depth for each shadow
 			for (int32 ShadowIndex = 0; ShadowIndex < Shadows.Num(); ShadowIndex++)
@@ -3225,7 +3219,7 @@ bool FDeferredShadingSceneRenderer::RenderProjectedShadows(FRHICommandListImmedi
 				FProjectedShadowInfo* ProjectedShadowInfo = Shadows[ShadowIndex];
 				if (ProjectedShadowInfo->bAllocated && !ProjectedShadowInfo->bTranslucentShadow && !ProjectedShadowInfo->bRayTracedDistanceFieldShadow)
 				{
-					ProjectedShadowInfo->RenderDepth(RHICmdList, this, LocalShadowTargets);
+					ProjectedShadowInfo->RenderDepth(RHICmdList, this, SetShadowRenderTargets);
 				}
 			}
 
@@ -3348,7 +3342,14 @@ bool FDeferredShadingSceneRenderer::RenderCachedPreshadows(FRHICommandListImmedi
 		{
 			SCOPED_DRAW_EVENTF(RHICmdList, EventShadowDepths, DEC_SCENE_ITEMS, TEXT("Preshadow Cache Depths"));
 
-			SetRenderTarget(RHICmdList, FTextureRHIRef(), GSceneRenderTargets.PreShadowCacheDepthZ->GetRenderTargetItem().TargetableTexture);
+			auto SetShadowRenderTargets = MakeFunction<void, FRHICommandList&>([](FRHICommandList& RHICmdList)
+			{
+				SetRenderTarget(RHICmdList, FTextureRHIRef(), GSceneRenderTargets.PreShadowCacheDepthZ->GetRenderTargetItem().TargetableTexture);
+			});
+
+			SetShadowRenderTargets(RHICmdList);  // run it now, maybe run it later for parallel command lists
+
+
 
 			// Render depth for each shadow
 			for (int32 ShadowIndex = 0; ShadowIndex < CachedPreshadows.Num(); ShadowIndex++)
@@ -3360,14 +3361,7 @@ bool FDeferredShadingSceneRenderer::RenderCachedPreshadows(FRHICommandListImmedi
 					check(ProjectedShadowInfo->bAllocated);
 					ProjectedShadowInfo->ClearDepth(RHICmdList, this);
 
-					struct FLocalShadowTargets : public ISetShadowRenderTargets
-					{
-						virtual void SetRenderTargetsForShadows(FRHICommandList& RHICmdList)
-						{
-							SetRenderTarget(RHICmdList, FTextureRHIRef(), GSceneRenderTargets.PreShadowCacheDepthZ->GetRenderTargetItem().TargetableTexture);
-						}
-					} LocalShadowTargets;
-					ProjectedShadowInfo->RenderDepth(RHICmdList, this, LocalShadowTargets);
+					ProjectedShadowInfo->RenderDepth(RHICmdList, this, SetShadowRenderTargets);
 					ProjectedShadowInfo->bDepthsCached = true;
 
 					const FResolveParams ResolveParams(FResolveRect(
@@ -3471,8 +3465,14 @@ bool FForwardShadingSceneRenderer::RenderShadowDepthMap(FRHICommandListImmediate
 		// Render the shadow depths.
 		SCOPED_DRAW_EVENT(RHICmdList, ShadowDepthsFromOpaque, DEC_SCENE_ITEMS);
 
-		// Clearin the render target happens inside the Begin function
-		GSceneRenderTargets.BeginRenderingShadowDepth(RHICmdList, true);
+		bool bFirst = true;
+		auto SetShadowRenderTargets = MakeFunction<void, FRHICommandList&>([&bFirst](FRHICommandList& RHICmdList)
+		{
+			GSceneRenderTargets.BeginRenderingShadowDepth(RHICmdList, bFirst);
+		});
+
+		SetShadowRenderTargets(RHICmdList);  // run it now, maybe run it later for parallel command lists
+		bFirst = false;
 
 		// render depth for each shadow
 		for (int32 ShadowIndex = 0; ShadowIndex < Shadows.Num(); ShadowIndex++)
@@ -3481,14 +3481,7 @@ bool FForwardShadingSceneRenderer::RenderShadowDepthMap(FRHICommandListImmediate
 			check(!ProjectedShadowInfo->bTranslucentShadow);
 			ProjectedShadowInfo->bRendered = false;
 			ProjectedShadowInfo->bAllocated = true;
-			struct FLocalShadowTargets : public ISetShadowRenderTargets
-			{
-				virtual void SetRenderTargetsForShadows(FRHICommandList& RHICmdList)
-				{
-					GSceneRenderTargets.BeginRenderingShadowDepth(RHICmdList, false);
-				}
-			} LocalShadowTargets;
-			ProjectedShadowInfo->RenderDepth(RHICmdList, this, LocalShadowTargets);
+			ProjectedShadowInfo->RenderDepth(RHICmdList, this, SetShadowRenderTargets);
 			bAttenuationBufferDirty = true;
 		}
 
diff --git a/Engine/Source/Runtime/Renderer/Private/ShadowRendering.h b/Engine/Source/Runtime/Renderer/Private/ShadowRendering.h
--- a/Engine/Source/Runtime/Renderer/Private/ShadowRendering.h
+++ b/Engine/Source/Runtime/Renderer/Private/ShadowRendering.h
@@ -490,13 +490,54 @@ public:
 	bool bIsTwoSided;
 };
 
+template<typename TReturn, typename TArg>
+class IFunction
+{
+public:
+	virtual ~IFunction() {}
+	virtual TReturn operator()(TArg&& Arg) = 0;
+	virtual TReturn operator()(TArg&& Arg) const = 0;
+};
 
-struct ISetShadowRenderTargets
+template<typename TReturn, typename TArg, typename TLambda>
+class TFunction : public IFunction<TReturn, TArg>
 {
-	virtual ~ISetShadowRenderTargets() {}
-	virtual void SetRenderTargetsForShadows(FRHICommandList& RHICmdList) = 0;
+	TLambda Func;
+public:
+	TFunction(TLambda&& InFunc)
+		: Func(MoveTemp(InFunc))
+	{
+	}
+	TFunction(TLambda& InFunc)
+		: Func(InFunc)
+	{
+	}
+	virtual ~TFunction() override final
+	{
+	}
+	virtual TReturn operator()(TArg&& Arg) override final
+	{
+		return Func(Forward<TArg>(Arg));
+	}
+	virtual TReturn operator()(TArg&& Arg) const override final
+	{
+		return Func(Forward<TArg>(Arg));
+	}
 };
 
+template<typename TReturn, typename TArg, typename TLambda>
+TFunction<TReturn, TArg, TLambda> MakeFunction(TLambda&& Func)
+{
+	return TFunction<TReturn, TArg, TLambda>(MoveTemp(Func));
+}
+
+template<typename TReturn, typename TArg, typename TLambda>
+TFunction<TReturn, TArg, TLambda>* NewFunction(TLambda&& Func)
+{
+	return new TFunction<TReturn, TArg, TLambda>(MoveTemp(Func));
+}
+
+
 /**
  * Information about a projected shadow.
  */
@@ -678,7 +719,7 @@ public:
 	/**
 	 * Renders the shadow subject depth.
 	 */
-	void RenderDepth(FRHICommandList& RHICmdList, class FSceneRenderer* SceneRenderer, ISetShadowRenderTargets& SetShadowRenderTargets);
+	void RenderDepth(FRHICommandList& RHICmdList, class FSceneRenderer* SceneRenderer, IFunction<void, FRHICommandList&>& SetShadowRenderTargets);
 
 	/** Set state for depth rendering */
 	void SetStateForDepth(FRHICommandList& RHICmdList);
@@ -766,7 +807,7 @@ private:
 	/**
 	* Renders the shadow subject depth, to a particular hacked view
 	*/
-	void RenderDepthInner(FRHICommandList& RHICmdList, class FSceneRenderer* SceneRenderer, const FViewInfo* FoundView, ISetShadowRenderTargets& SetShadowRenderTargets);
+	void RenderDepthInner(FRHICommandList& RHICmdList, class FSceneRenderer* SceneRenderer, const FViewInfo* FoundView, IFunction<void, FRHICommandList&>& SetShadowRenderTargets);
 
 	/**
 	* Renders the dynamic shadow subject depth, to a particular hacked view
diff --git a/Engine/Source/Runtime/Renderer/Private/StaticMeshDrawList.inl b/Engine/Source/Runtime/Renderer/Private/StaticMeshDrawList.inl
--- a/Engine/Source/Runtime/Renderer/Private/StaticMeshDrawList.inl
+++ b/Engine/Source/Runtime/Renderer/Private/StaticMeshDrawList.inl
@@ -309,25 +309,25 @@ class FDrawVisibleAnyThreadTask
 public:
 
 	FDrawVisibleAnyThreadTask(
-		TStaticMeshDrawList<DrawingPolicyType>* InCaller,
-		FRHICommandList* InRHICmdList,
-		const FViewInfo* InView,
-		const typename DrawingPolicyType::ContextDataType InPolicyContext,
-		const TBitArray<SceneRenderingBitArrayAllocator>* InStaticMeshVisibilityMap,
-		const TArray<uint64, SceneRenderingAllocator>* InBatchVisibilityArray,
+		TStaticMeshDrawList<DrawingPolicyType>& InCaller,
+		FRHICommandList& InRHICmdList,
+		const FViewInfo& InView,
+		const typename DrawingPolicyType::ContextDataType& InPolicyContext,
+		const TBitArray<SceneRenderingBitArrayAllocator>& InStaticMeshVisibilityMap,
+		const TArray<uint64, SceneRenderingAllocator>& InBatchVisibilityArray,
 		int32 InFirstPolicy,
 		int32 InLastPolicy,
-		bool* InOutDirty
+		bool& InOutDirty
 		)
-		: Caller(*InCaller)
-		, RHICmdList(*InRHICmdList)
-		, View(*InView)
+		: Caller(InCaller)
+		, RHICmdList(InRHICmdList)
+		, View(InView)
 		, PolicyContext(InPolicyContext)
-		, StaticMeshVisibilityMap(*InStaticMeshVisibilityMap)
-		, BatchVisibilityArray(*InBatchVisibilityArray)
+		, StaticMeshVisibilityMap(InStaticMeshVisibilityMap)
+		, BatchVisibilityArray(InBatchVisibilityArray)
 		, FirstPolicy(InFirstPolicy)
 		, LastPolicy(InLastPolicy)
-		, OutDirty(*InOutDirty)
+		, OutDirty(InOutDirty)
 	{
 	}
 
@@ -381,7 +381,7 @@ void TStaticMeshDrawList<DrawingPolicyType>::DrawVisibleParallel(
 				FRHICommandList* CmdList = ParallelCommandListSet.NewParallelCommandList();
 
 				FGraphEventRef AnyThreadCompletionEvent = TGraphTask<FDrawVisibleAnyThreadTask<DrawingPolicyType> >::CreateTask(nullptr, ENamedThreads::RenderThread)
-					.ConstructAndDispatchWhenReady(this, CmdList, &ParallelCommandListSet.View, PolicyContext, &StaticMeshVisibilityMap, &BatchVisibilityArray, Start, Last, &ParallelCommandListSet.OutDirty);
+					.ConstructAndDispatchWhenReady(*this, *CmdList, ParallelCommandListSet.View, PolicyContext, StaticMeshVisibilityMap, BatchVisibilityArray, Start, Last, ParallelCommandListSet.OutDirty);
 
 				ParallelCommandListSet.AddParallelCommandList(CmdList, AnyThreadCompletionEvent);
 			}
diff --git a/Engine/Source/Runtime/Renderer/Private/TranslucentRendering.cpp b/Engine/Source/Runtime/Renderer/Private/TranslucentRendering.cpp
--- a/Engine/Source/Runtime/Renderer/Private/TranslucentRendering.cpp
+++ b/Engine/Source/Runtime/Renderer/Private/TranslucentRendering.cpp
@@ -319,9 +319,9 @@ class FCopySceneColorAndRestoreRenderThreadTask
 	const FPrimitiveSceneProxy* PrimitiveSceneProxy;
 public:
 
-	FCopySceneColorAndRestoreRenderThreadTask(FRHICommandList* InRHICmdList, const FViewInfo* InView, const FPrimitiveSceneProxy* InPrimitiveSceneProxy)
-		: RHICmdList(*InRHICmdList)
-		, View(*InView)
+	FCopySceneColorAndRestoreRenderThreadTask(FRHICommandList& InRHICmdList, const FViewInfo& InView, const FPrimitiveSceneProxy* InPrimitiveSceneProxy)
+		: RHICmdList(InRHICmdList)
+		, View(InView)
 		, PrimitiveSceneProxy(InPrimitiveSceneProxy)
 	{
 	}
@@ -379,7 +379,7 @@ bool FTranslucencyDrawingPolicyFactory::DrawMesh(
 				if (!RHICmdList.Bypass() && !IsInRenderingThread())
 				{
 					FRHICommandList* CmdList = new FRHICommandList;
-					FGraphEventRef RenderThreadCompletionEvent = TGraphTask<FCopySceneColorAndRestoreRenderThreadTask>::CreateTask().ConstructAndDispatchWhenReady(CmdList, &View, PrimitiveSceneProxy);
+					FGraphEventRef RenderThreadCompletionEvent = TGraphTask<FCopySceneColorAndRestoreRenderThreadTask>::CreateTask().ConstructAndDispatchWhenReady(*CmdList, View, PrimitiveSceneProxy);
 					RHICmdList.QueueRenderThreadCommandListSubmit(RenderThreadCompletionEvent, CmdList);
 				}
 				else
@@ -546,11 +546,11 @@ class FVolumetricTranslucentShadowRenderThreadTask
 
 public:
 
-	FORCEINLINE_DEBUGGABLE FVolumetricTranslucentShadowRenderThreadTask(FRHICommandList* InRHICmdList, const FTranslucentPrimSet* InPrimSet, const FViewInfo* InView, FDeferredShadingSceneRenderer* InRenderer, bool InbSeparateTranslucencyPass, int32 InIndex)
-		: RHICmdList(*InRHICmdList)
-		, PrimSet(*InPrimSet)
-		, View(*InView)
-		, Renderer(*InRenderer)
+	FORCEINLINE_DEBUGGABLE FVolumetricTranslucentShadowRenderThreadTask(FRHICommandList& InRHICmdList, const FTranslucentPrimSet& InPrimSet, const FViewInfo& InView, FDeferredShadingSceneRenderer& InRenderer, bool InbSeparateTranslucencyPass, int32 InIndex)
+		: RHICmdList(InRHICmdList)
+		, PrimSet(InPrimSet)
+		, View(InView)
+		, Renderer(InRenderer)
 		, bSeparateTranslucencyPass(InbSeparateTranslucencyPass)
 		, Index(InIndex)
 	{
@@ -602,7 +602,7 @@ void FTranslucentPrimSet::DrawPrimitivesParallel(
 			check(!IsInRenderingThread());
 			// can't do this in parallel, defer
 			FRHICommandList* CmdList = new FRHICommandList;
-			FGraphEventRef RenderThreadCompletionEvent = TGraphTask<FVolumetricTranslucentShadowRenderThreadTask>::CreateTask().ConstructAndDispatchWhenReady(CmdList, this, &View, &Renderer, bSeparateTranslucencyPass, PrimIdx);
+			FGraphEventRef RenderThreadCompletionEvent = TGraphTask<FVolumetricTranslucentShadowRenderThreadTask>::CreateTask().ConstructAndDispatchWhenReady(*CmdList, *this, View, Renderer, bSeparateTranslucencyPass, PrimIdx);
 			RHICmdList.QueueRenderThreadCommandListSubmit(RenderThreadCompletionEvent, CmdList);
 		}
 		else
@@ -830,16 +830,16 @@ class FDrawSortedTransAnyThreadTask
 public:
 
 	FDrawSortedTransAnyThreadTask(
-		FDeferredShadingSceneRenderer* InRenderer,
-		FRHICommandList* InRHICmdList,
-		const FViewInfo* InView,
+		FDeferredShadingSceneRenderer& InRenderer,
+		FRHICommandList& InRHICmdList,
+		const FViewInfo& InView,
 		bool InbSeparateTranslucencyPass,
 		int32 InFirstIndex,
 		int32 InLastIndex
 		)
-		: Renderer(*InRenderer)
-		, RHICmdList(*InRHICmdList)
-		, View(*InView)
+		: Renderer(InRenderer)
+		, RHICmdList(InRHICmdList)
+		, View(InView)
 		, bSeparateTranslucencyPass(InbSeparateTranslucencyPass)
 		, FirstIndex(InFirstIndex)
 		, LastIndex(InLastIndex)
@@ -929,7 +929,7 @@ void FDeferredShadingSceneRenderer::RenderTranslucencyParallel(FRHICommandListIm
 							FRHICommandList* CmdList = ParallelCommandListSet.NewParallelCommandList();
 
 							FGraphEventRef AnyThreadCompletionEvent = TGraphTask<FDrawSortedTransAnyThreadTask>::CreateTask(nullptr, ENamedThreads::RenderThread)
-								.ConstructAndDispatchWhenReady(this, CmdList, &View, false, Start, Last);
+								.ConstructAndDispatchWhenReady(*this, *CmdList, View, false, Start, Last);
 
 							ParallelCommandListSet.AddParallelCommandList(CmdList, AnyThreadCompletionEvent);
 						}
@@ -983,7 +983,7 @@ void FDeferredShadingSceneRenderer::RenderTranslucencyParallel(FRHICommandListIm
 								FRHICommandList* CmdList = ParallelCommandListSet.NewParallelCommandList();
 
 								FGraphEventRef AnyThreadCompletionEvent = TGraphTask<FDrawSortedTransAnyThreadTask>::CreateTask(nullptr, ENamedThreads::RenderThread)
-									.ConstructAndDispatchWhenReady(this, CmdList, &View, true, Start, Last);
+									.ConstructAndDispatchWhenReady(*this, *CmdList, View, true, Start, Last);
 
 								ParallelCommandListSet.AddParallelCommandList(CmdList, AnyThreadCompletionEvent);
 							}
diff --git a/Engine/Source/Runtime/Renderer/Private/VelocityRendering.cpp b/Engine/Source/Runtime/Renderer/Private/VelocityRendering.cpp
--- a/Engine/Source/Runtime/Renderer/Private/VelocityRendering.cpp
+++ b/Engine/Source/Runtime/Renderer/Private/VelocityRendering.cpp
@@ -495,21 +495,20 @@ class FRenderVelocityDynamicThreadTask
 	FDeferredShadingSceneRenderer& ThisRenderer;
 	FRHICommandList& RHICmdList;
 	const FViewInfo& View;
-	FDeferredShadingSceneRenderer* SceneRenderer;
 	int32 FirstIndex;
 	int32 LastIndex;
 
 public:
 
 	FRenderVelocityDynamicThreadTask(
-		FDeferredShadingSceneRenderer* InThisRenderer,
-		FRHICommandList* InRHICmdList,
-		const FViewInfo* InView,
+		FDeferredShadingSceneRenderer& InThisRenderer,
+		FRHICommandList& InRHICmdList,
+		const FViewInfo& InView,
 		int32 InFirstIndex, int32 InLastIndex
 		)
-		: ThisRenderer(*InThisRenderer)
-		, RHICmdList(*InRHICmdList)
-		, View(*InView)
+		: ThisRenderer(InThisRenderer)
+		, RHICmdList(InRHICmdList)
+		, View(InView)
 		, FirstIndex(InFirstIndex)
 		, LastIndex(InLastIndex)
 	{
@@ -608,7 +607,7 @@ void FDeferredShadingSceneRenderer::RenderVelocitiesInnerParallel(FRHICommandLis
 				FRHICommandList* CmdList = ParallelCommandListSet.NewParallelCommandList();
 
 				FGraphEventRef AnyThreadCompletionEvent = TGraphTask<FRenderVelocityDynamicThreadTask>::CreateTask(nullptr, ENamedThreads::RenderThread)
-					.ConstructAndDispatchWhenReady(this, CmdList, &View, Start, Last);
+					.ConstructAndDispatchWhenReady(*this, *CmdList, View, Start, Last);
 
 				ParallelCommandListSet.AddParallelCommandList(CmdList, AnyThreadCompletionEvent);
 
