Rebuild Parameter Collection Instances after new scene is created.
move FXSystem allocation in to FScene constructor.
Fixes TTP 348856

diff --git a/Engine/Source/Runtime/Engine/Private/RendererSupport.cpp b/Engine/Source/Runtime/Engine/Private/RendererSupport.cpp
--- a/Engine/Source/Runtime/Engine/Private/RendererSupport.cpp
+++ b/Engine/Source/Runtime/Engine/Private/RendererSupport.cpp
@@ -133,13 +133,7 @@ void RestoreReferencesToRendererModuleClasses(
 	{
 		UWorld* World = It.Key();
 
-		World->Scene = RendererModule.AllocateScene(World, World->RequiresHitProxies(), World->FeatureLevel);
-
-		if (It.Value())
-		{
-			World->FXSystem = FFXSystemInterface::Create(World->FeatureLevel);
-			World->Scene->SetFXSystem( World->FXSystem );
-		}
+		RendererModule.AllocateScene(World, World->RequiresHitProxies(), It.Value(), World->FeatureLevel);
 
 		for (int32 LevelIndex = 0; LevelIndex < World->GetNumLevels(); LevelIndex++)
 		{
diff --git a/Engine/Source/Runtime/Engine/Private/World.cpp b/Engine/Source/Runtime/Engine/Private/World.cpp
--- a/Engine/Source/Runtime/Engine/Private/World.cpp
+++ b/Engine/Source/Runtime/Engine/Private/World.cpp
@@ -803,16 +803,7 @@ void UWorld::InitWorld(const InitializationValues IVS)
 		
 		// Save off the value used to create the scene, so this UWorld can recreate its scene later
 		bRequiresHitProxies = IVS.bRequiresHitProxies;
-		Scene = GetRendererModule().AllocateScene( this, bRequiresHitProxies, FeatureLevel );
-		if ( IVS.bCreateFXSystem )
-		{
-			CreateFXSystem();
-		}
-		else
-		{
-			FXSystem = NULL;
-			Scene->SetFXSystem(NULL);
-		}
+		GetRendererModule().AllocateScene(this, bRequiresHitProxies, IVS.bCreateFXSystem, FeatureLevel);
 	}
 
 	// Prepare AI systems
@@ -5379,11 +5370,7 @@ void UWorld::ChangeFeatureLevel(ERHIFeatureLevel::Type InFeatureLevel)
 			Scene->Release();
 			GetRendererModule().RemoveScene(Scene);
 
-			Scene = GetRendererModule().AllocateScene(this, bRequiresHitProxies, InFeatureLevel);
-
-			FFXSystemInterface::Destroy(FXSystem);
-			FXSystem = FFXSystemInterface::Create(InFeatureLevel);
-			Scene->SetFXSystem(FXSystem);
+			GetRendererModule().AllocateScene(this, bRequiresHitProxies, FXSystem != nullptr, InFeatureLevel );
 		}
 
 		TriggerStreamingDataRebuild();
diff --git a/Engine/Source/Runtime/Renderer/Private/RendererPrivate.h b/Engine/Source/Runtime/Renderer/Private/RendererPrivate.h
--- a/Engine/Source/Runtime/Renderer/Private/RendererPrivate.h
+++ b/Engine/Source/Runtime/Renderer/Private/RendererPrivate.h
@@ -23,7 +23,7 @@ public:
 	virtual bool SupportsDynamicReloading() override { return true; }
 
 	virtual void BeginRenderingViewFamily(FCanvas* Canvas,const FSceneViewFamily* ViewFamily) override;
-	virtual FSceneInterface* AllocateScene(UWorld* World, bool bInRequiresHitProxies, ERHIFeatureLevel::Type InFeatureLevel) override;
+	virtual FSceneInterface* AllocateScene(UWorld* World, bool bInRequiresHitProxies, bool bCreateFXSystem, ERHIFeatureLevel::Type InFeatureLevel) override;
 	virtual void RemoveScene(FSceneInterface* Scene) override;
 	virtual void UpdateStaticDrawListsForMaterials(const TArray<const FMaterial*>& Materials) override;
 	virtual FSceneViewStateInterface* AllocateViewState() override;
diff --git a/Engine/Source/Runtime/Renderer/Private/RendererScene.cpp b/Engine/Source/Runtime/Renderer/Private/RendererScene.cpp
--- a/Engine/Source/Runtime/Renderer/Private/RendererScene.cpp
+++ b/Engine/Source/Runtime/Renderer/Private/RendererScene.cpp
@@ -12,6 +12,7 @@
 #include "DistanceFieldSurfaceCacheLighting.h"
 #include "EngineModule.h"
 #include "PrecomputedLightVolume.h"
+#include "FXSystem.h"
 
 // Enable this define to do slow checks for components being added to the wrong
 // world's scene, when using PIE. This can happen if a PIE component is reattached
@@ -189,7 +190,7 @@ FORCEINLINE static void VerifyProperPIEScene(UPrimitiveComponent* Component, UWo
 #endif
 }
 
-FScene::FScene(UWorld* InWorld, bool bInRequiresHitProxies, bool bInIsEditorScene, ERHIFeatureLevel::Type InFeatureLevel)
+FScene::FScene(UWorld* InWorld, bool bInRequiresHitProxies, bool bInIsEditorScene, bool bCreateFXSystem, ERHIFeatureLevel::Type InFeatureLevel)
 :	World(InWorld)
 ,	FXSystem(NULL)
 ,	bStaticDrawListsMobileHDR(false)
@@ -216,6 +217,7 @@ FScene::FScene(UWorld* InWorld, bool bInRequiresHitProxies, bool bInIsEditorScen
 ,	bHasSkyLight(false)
 {
 	check(World);
+	World->Scene = this;
 
 	FeatureLevel = World->FeatureLevel;
 
@@ -226,6 +228,23 @@ FScene::FScene(UWorld* InWorld, bool bInRequiresHitProxies, bool bInIsEditorScen
 
 	static auto* EarlyZPassCvar = IConsoleManager::Get().FindTConsoleVariableDataInt(TEXT("r.EarlyZPass"));
 	StaticDrawListsEarlyZPassMode = EarlyZPassCvar->GetValueOnAnyThread();
+
+	if (World->FXSystem)
+	{
+		FFXSystemInterface::Destroy(World->FXSystem);
+	}
+
+	if (bCreateFXSystem)
+	{
+		World->CreateFXSystem();
+	}
+	else
+	{
+		World->FXSystem = NULL;
+		SetFXSystem(NULL);
+	}
+
+	World->UpdateParameterCollectionInstances(false);
 }
 
 FScene::~FScene()
@@ -2139,14 +2158,14 @@ private:
 	class FFXSystemInterface* FXSystem;
 };
 
-FSceneInterface* FRendererModule::AllocateScene(UWorld* World, bool bInRequiresHitProxies, ERHIFeatureLevel::Type InFeatureLevel)
+FSceneInterface* FRendererModule::AllocateScene(UWorld* World, bool bInRequiresHitProxies, bool bCreateFXSystem, ERHIFeatureLevel::Type InFeatureLevel)
 {
 	check(IsInGameThread());
 
 	// Create a full fledged scene if we have something to render.
 	if( GIsClient && !IsRunningCommandlet() && !GUsingNullRHI )
 	{
-		FScene* NewScene = new FScene(World, bInRequiresHitProxies, GIsEditor && !World->IsGameWorld(), InFeatureLevel);
+		FScene* NewScene = new FScene(World, bInRequiresHitProxies, GIsEditor && !World->IsGameWorld(), bCreateFXSystem, InFeatureLevel);
 		AllocatedScenes.Add(NewScene);
 		return NewScene;
 	}
diff --git a/Engine/Source/Runtime/Renderer/Private/ScenePrivate.h b/Engine/Source/Runtime/Renderer/Private/ScenePrivate.h
--- a/Engine/Source/Runtime/Renderer/Private/ScenePrivate.h
+++ b/Engine/Source/Runtime/Renderer/Private/ScenePrivate.h
@@ -1314,7 +1314,7 @@ public:
 	TMap<FGuid, FUniformBufferRHIRef> ParameterCollections;
 
 	/** Initialization constructor. */
-	FScene(UWorld* InWorld, bool bInRequiresHitProxies,bool bInIsEditorScene, ERHIFeatureLevel::Type InFeatureLevel);
+	FScene(UWorld* InWorld, bool bInRequiresHitProxies,bool bInIsEditorScene, bool bCreateFXSystem, ERHIFeatureLevel::Type InFeatureLevel);
 
 	virtual ~FScene();
 
diff --git a/Engine/Source/Runtime/Renderer/Public/RendererInterface.h b/Engine/Source/Runtime/Renderer/Public/RendererInterface.h
--- a/Engine/Source/Runtime/Renderer/Public/RendererInterface.h
+++ b/Engine/Source/Runtime/Renderer/Public/RendererInterface.h
@@ -454,7 +454,7 @@ public:
 	 * @param World - An optional world to associate with the scene.
 	 * @param bInRequiresHitProxies - Indicates that hit proxies should be rendered in the scene.
 	 */
-	virtual FSceneInterface* AllocateScene(UWorld* World, bool bInRequiresHitProxies, ERHIFeatureLevel::Type InFeatureLevel) = 0;
+	virtual FSceneInterface* AllocateScene(UWorld* World, bool bInRequiresHitProxies, bool bCreateFXSystem, ERHIFeatureLevel::Type InFeatureLevel) = 0;
 	
 	virtual void RemoveScene(FSceneInterface* Scene) = 0;
 
