UE4 - fix stencil clear bug for shadows and parallel rendering

diff --git a/Engine/Source/Runtime/Renderer/Private/PostProcess/SceneRenderTargets.cpp b/Engine/Source/Runtime/Renderer/Private/PostProcess/SceneRenderTargets.cpp
--- a/Engine/Source/Runtime/Renderer/Private/PostProcess/SceneRenderTargets.cpp
+++ b/Engine/Source/Runtime/Renderer/Private/PostProcess/SceneRenderTargets.cpp
@@ -783,7 +783,12 @@ void FSceneRenderTargets::FinishRenderingPrePass(FRHICommandListImmediate& RHICm
 void FSceneRenderTargets::BeginRenderingShadowDepth(FRHICommandList& RHICmdList, bool bClear)
 {
 	GRenderTargetPool.VisualizeTexture.SetCheckPoint(RHICmdList, ShadowDepthZ);
-	FRHISetRenderTargetsInfo Info(0, nullptr, FRHIDepthRenderTargetView(GetShadowDepthZSurface(), ERenderTargetLoadAction::EClear, ERenderTargetStoreAction::EStore));
+	FRHISetRenderTargetsInfo Info(0, nullptr, 
+		FRHIDepthRenderTargetView(GetShadowDepthZSurface(), 
+			bClear ? ERenderTargetLoadAction::EClear : ERenderTargetLoadAction::ELoad, 
+			ERenderTargetStoreAction::EStore, 
+			bClear ? ERenderTargetLoadAction::EClear : ERenderTargetLoadAction::ELoad,
+			ERenderTargetStoreAction::EStore));
 	if (bClear)
 	{
 		Info.SetClearDepthStencil(true, 1.0f);
