UE4 - Feature Level ES 3.1 fixes

diff --git a/Engine/Source/Runtime/Engine/Private/Landscape/Landscape.cpp b/Engine/Source/Runtime/Engine/Private/Landscape/Landscape.cpp
--- a/Engine/Source/Runtime/Engine/Private/Landscape/Landscape.cpp
+++ b/Engine/Source/Runtime/Engine/Private/Landscape/Landscape.cpp
@@ -412,7 +412,7 @@ void ULandscapeComponent::PostLoad()
 		}
 
 		// If we're loading on a platform that doesn't require cooked data, attempt to load missing data from the DDC
-		if (!FPlatformProperties::RequiresCookedData() && (GIsEditor || GMaxRHIFeatureLevel == ERHIFeatureLevel::ES2))
+		if (!FPlatformProperties::RequiresCookedData() && (GIsEditor || (GMaxRHIFeatureLevel == ERHIFeatureLevel::ES2 || GMaxRHIFeatureLevel == ERHIFeatureLevel::ES3_1)))
 		{
 			// Only check the DDC if we don't already have it loaded
 			if (!PlatformData.HasValidPlatformData())
diff --git a/Engine/Source/Runtime/Engine/Private/Materials/MaterialExpressions.cpp b/Engine/Source/Runtime/Engine/Private/Materials/MaterialExpressions.cpp
--- a/Engine/Source/Runtime/Engine/Private/Materials/MaterialExpressions.cpp
+++ b/Engine/Source/Runtime/Engine/Private/Materials/MaterialExpressions.cpp
@@ -9027,7 +9027,8 @@ int32 UMaterialExpressionLandscapeLayerCoords::Compile(class FMaterialCompiler*
 		Compiler->Constant2(MappingPanU, MappingPanV)
 		);
 
-	if (Compiler->GetFeatureLevel() != ERHIFeatureLevel::ES2) // No need to localize UV
+	auto FeatureLevel = Compiler->GetFeatureLevel();
+	if (FeatureLevel != ERHIFeatureLevel::ES2 && FeatureLevel != ERHIFeatureLevel::ES3_1) // No need to localize UV
 	{
 		ResultUV = TransformedUV;
 	}
diff --git a/Engine/Source/Runtime/Engine/Private/SceneView.cpp b/Engine/Source/Runtime/Engine/Private/SceneView.cpp
--- a/Engine/Source/Runtime/Engine/Private/SceneView.cpp
+++ b/Engine/Source/Runtime/Engine/Private/SceneView.cpp
@@ -1022,7 +1022,8 @@ void FSceneView::EndFinalPostprocessSettings()
 		}
 
 		// Not supported in ES2.
-		if(Family->Scene->GetFeatureLevel() == ERHIFeatureLevel::ES2)
+		auto FeatureLevel = Family->Scene->GetFeatureLevel();
+		if(FeatureLevel == ERHIFeatureLevel::ES2 || FeatureLevel == ERHIFeatureLevel::ES3_1)
 		{
 			FinalPostProcessSettings.ScreenPercentage = 100.0f;
 		}
@@ -1118,7 +1119,7 @@ void FSceneView::EndFinalPostprocessSettings()
 		if( !Family->EngineShowFlags.PostProcessing || !Family->EngineShowFlags.AntiAliasing || Quality <= 0
 			// Disable antialiasing in GammaLDR mode to avoid jittering.
 			|| (FeatureLevel == ERHIFeatureLevel::ES2 && MobileHDRCvar->GetValueOnGameThread() == 0)
-			|| (FeatureLevel == ERHIFeatureLevel::ES2 && (MSAAValue > 1)))
+			|| (FeatureLevel <= ERHIFeatureLevel::ES3_1 && (MSAAValue > 1)))
 		{
 			FinalPostProcessSettings.AntiAliasingMethod = AAM_None;
 		}
diff --git a/Engine/Source/Runtime/Engine/Public/Landscape/LandscapeRenderMobile.h b/Engine/Source/Runtime/Engine/Public/Landscape/LandscapeRenderMobile.h
--- a/Engine/Source/Runtime/Engine/Public/Landscape/LandscapeRenderMobile.h
+++ b/Engine/Source/Runtime/Engine/Public/Landscape/LandscapeRenderMobile.h
@@ -48,7 +48,8 @@ public:
 	*/
 	static bool ShouldCache(EShaderPlatform Platform, const class FMaterial* Material, const class FShaderType* ShaderType)
 	{
-		return (GetMaxSupportedFeatureLevel(Platform) == ERHIFeatureLevel::ES2) &&
+		auto FeatureLevel = GetMaxSupportedFeatureLevel(Platform);
+		return (FeatureLevel == ERHIFeatureLevel::ES2 || FeatureLevel == ERHIFeatureLevel::ES3_1) &&
 			(Material->IsUsedWithLandscape() || Material->IsSpecialEngineMaterial());
 	}
 
diff --git a/Engine/Source/Runtime/Renderer/Private/PostProcess/PostProcessDownsample.cpp b/Engine/Source/Runtime/Renderer/Private/PostProcess/PostProcessDownsample.cpp
--- a/Engine/Source/Runtime/Renderer/Private/PostProcess/PostProcessDownsample.cpp
+++ b/Engine/Source/Runtime/Renderer/Private/PostProcess/PostProcessDownsample.cpp
@@ -212,7 +212,8 @@ void FRCPassPostProcessDownsample::Process(FRenderingCompositePassContext& Conte
 
 	// check if we have to clear the whole surface.
 	// Otherwise perform the clear when the dest rectangle has been computed.
-	if (Context.View.GetFeatureLevel() == ERHIFeatureLevel::ES2)
+	auto FeatureLevel = Context.View.GetFeatureLevel();
+	if (FeatureLevel == ERHIFeatureLevel::ES2 || FeatureLevel == ERHIFeatureLevel::ES3_1)
 	{
 		Context.RHICmdList.Clear(true, FLinearColor(0, 0, 0, 0), false, 1.0f, false, 0, FIntRect());
 		bHasCleared = true;
diff --git a/Engine/Source/Runtime/Renderer/Private/PostProcess/PostProcessWeightedSampleSum.cpp b/Engine/Source/Runtime/Renderer/Private/PostProcess/PostProcessWeightedSampleSum.cpp
--- a/Engine/Source/Runtime/Renderer/Private/PostProcess/PostProcessWeightedSampleSum.cpp
+++ b/Engine/Source/Runtime/Renderer/Private/PostProcess/PostProcessWeightedSampleSum.cpp
@@ -447,7 +447,7 @@ void FRCPassPostProcessWeightedSampleSum::Process(FRenderingCompositePassContext
 	bool bRequiresClear = true;
 	// check if we have to clear the whole surface.
 	// Otherwise perform the clear when the dest rectangle has been computed.
-	if (FeatureLevel == ERHIFeatureLevel::ES2)
+	if (FeatureLevel == ERHIFeatureLevel::ES2 || FeatureLevel == ERHIFeatureLevel::ES3_1)
 	{
 		Context.RHICmdList.Clear(true, FLinearColor(0, 0, 0, 0), false, 1.0f, false, 0, FIntRect());
 		bRequiresClear = false;
diff --git a/Engine/Source/Runtime/Renderer/Private/SceneVisibility.cpp b/Engine/Source/Runtime/Renderer/Private/SceneVisibility.cpp
--- a/Engine/Source/Runtime/Renderer/Private/SceneVisibility.cpp
+++ b/Engine/Source/Runtime/Renderer/Private/SceneVisibility.cpp
@@ -763,7 +763,7 @@ static void ComputeRelevanceForView(
 
 		// Cache the nearest reflection proxy if needed
 		if (PrimitiveSceneInfo->bNeedsCachedReflectionCaptureUpdate
-			// In ES2, the per-object reflection is used for everything
+			// During Forward Shading, the per-object reflection is used for everything
 			// Otherwise it is just used on translucency
 			&& (!Scene->ShouldUseDeferredRenderer() || bTranslucentRelevance))
 		{
@@ -1005,9 +1005,9 @@ struct FRelevancePacket
 
 			// Cache the nearest reflection proxy if needed
 			if (PrimitiveSceneInfo->bNeedsCachedReflectionCaptureUpdate
-				// In ES2, the per-object reflection is used for everything
-					// Otherwise it is just used on translucency
-						&& (Scene->GetFeatureLevel() == ERHIFeatureLevel::ES2 || bTranslucentRelevance))
+				// During Forward Shading, the per-object reflection is used for everything
+				// Otherwise it is just used on translucency
+				&& (!Scene->ShouldUseDeferredRenderer() || bTranslucentRelevance))
 			{
 				PrimitiveSceneInfo->CachedReflectionCaptureProxy = Scene->FindClosestReflectionCapture(Scene->PrimitiveBounds[BitIndex].Origin);
 				PrimitiveSceneInfo->bNeedsCachedReflectionCaptureUpdate = false;
