Fixed running with -nullrhi crashing because the world was never given a scene

diff --git a/Engine/Source/Runtime/Renderer/Private/RendererScene.cpp b/Engine/Source/Runtime/Renderer/Private/RendererScene.cpp
--- a/Engine/Source/Runtime/Renderer/Private/RendererScene.cpp
+++ b/Engine/Source/Runtime/Renderer/Private/RendererScene.cpp
@@ -2062,7 +2062,9 @@ class FNULLSceneInterface : public FSceneInterface
 public:
 	FNULLSceneInterface( UWorld* InWorld )
 		:	World( InWorld )
-	{}
+	{
+		World->Scene = this;
+	}
 
 	virtual void AddPrimitive(UPrimitiveComponent* Primitive){}
 	virtual void RemovePrimitive(UPrimitiveComponent* Primitive){}
