Fixed editor primitive depth testing with GetDynamicMeshElements

diff --git a/Engine/Source/Runtime/Renderer/Private/PostProcess/PostProcessCompositeEditorPrimitives.cpp b/Engine/Source/Runtime/Renderer/Private/PostProcess/PostProcessCompositeEditorPrimitives.cpp
--- a/Engine/Source/Runtime/Renderer/Private/PostProcess/PostProcessCompositeEditorPrimitives.cpp
+++ b/Engine/Source/Runtime/Renderer/Private/PostProcess/PostProcessCompositeEditorPrimitives.cpp
@@ -315,7 +315,7 @@ void FRCPassPostProcessCompositeEditorPrimitives::RenderPrimitivesToComposite(FR
 
 	if (bUseGetMeshElements)
 	{
-		FBasePassOpaqueDrawingPolicyFactory::ContextType Context(false, ESceneRenderTargetsMode::SetTextures);
+		FBasePassOpaqueDrawingPolicyFactory::ContextType Context(true, ESceneRenderTargetsMode::SetTextures);
 
 		for (int32 MeshBatchIndex = 0; MeshBatchIndex < View.DynamicEditorMeshElements.Num(); MeshBatchIndex++)
 		{
