improved SubsurfaceProfile quality (no artifacts on half res upsample, more crisp non subsurface), refactoring

diff --git a/Engine/Shaders/PostProcessSubsurface.usf b/Engine/Shaders/PostProcessSubsurface.usf
--- a/Engine/Shaders/PostProcessSubsurface.usf
+++ b/Engine/Shaders/PostProcessSubsurface.usf
@@ -14,15 +14,17 @@
 /** RGBA8 linear texture containing random normals */
 Texture2D SSProfilesTexture;
 SamplerState SSProfilesTextureSampler;
-
-// x:Radius, y:DistanceToProjectionWindow, zw:xy scale correction for BufferSize!=ViewportSize
+// x:Radius*DistanceToProjectionWindow/KernelSize*0.5, y:DistanceToProjectionWindow, zw:xy scale correction for BufferSize!=ViewportSize
 float4 SSSParams;
 
 // ------------------------------------------
 
 // setup for "SeparableSSS.usf"
- 
-#define SSSS_HLSL_4 1
+
+float ComputeMaskFromDepthInAlpha(float Alpha)
+{
+	return Alpha > 0;
+}
 
 // can be optimized
 float GetSubsurfaceStrength(float2 UV)
@@ -35,7 +37,7 @@ float GetSubsurfaceStrength(float2 UV)
 }
 
 //#define SSSS_STREGTH_SOURCE 1	// todo: masking fom LightingModel or SSSColor
-#define SSSS_STREGTH_SOURCE GetSubsurfaceStrength(texcoord)
+#define SSSS_STREGTH_SOURCE(UV) GetSubsurfaceStrength(UV)
 
 // 0:speculars leak into SSS / 1:requires alpha channel for scenecolor
 // #define SSSS_SPECULAR_CORRECTION 1
@@ -54,16 +56,8 @@ float GetSubsurfaceStrength(float2 UV)
 
 #if SETUP_MODE == 0
 	#define SSSS_SPECULAR_CORRECTION 0
-	#define SETUP_HALF_RES 0
 #elif SETUP_MODE == 1
 	#define SSSS_SPECULAR_CORRECTION 1
-	#define SETUP_HALF_RES 0
-#elif SETUP_MODE == 3
-	#define SSSS_SPECULAR_CORRECTION 0
-	#define SETUP_HALF_RES 1
-#elif SETUP_MODE == 4
-	#define SSSS_SPECULAR_CORRECTION 1
-	#define SETUP_HALF_RES 1
 #else
 	// doesn't matter
 	#define SSSS_SPECULAR_CORRECTION 0
@@ -84,6 +78,39 @@ float GetSubsurfaceStrength(float2 UV)
 // 1: no color bleeding in z direction
 #define SSSS_FOLLOW_SURFACE 1
 
+float4 Quantize8Bit(float4 In)
+{
+	return floor(In * 255.999f) / 255.0f;
+}
+
+float4 GetKernel(uint SampleIndex, uint SubsurfaceProfileInt)
+{
+	const float4 TableMax = float4(1, 1, 1, 3.0f);
+
+	// profiled on NV670 fullscreen, one pass (total: 2 pass), samples: 13+1+13, large object filling the screen
+
+	// texture lookup
+	float4 Value = SSProfilesTexture.Load(int3(SampleIndex, SubsurfaceProfileInt, 0)) * TableMax;		// 0.88ms for 8bit
+	// simulated  8 bit lookup
+//	float4 Value = Quantize8Bit(kernel[SampleIndex] / TableMax) * TableMax;			// 1.33ms same look as texture
+	// float reference
+//	float4 Value = kernel[SampleIndex];												// 0.85ms, quite a bit different in look, need to use more than 8 bit?
+
+	// debug if the kernal was a box filter
+//	Value.rgb = 1.0f / (SSSS_N_KERNELWEIGHTCOUNT * 2 - 1);
+
+
+	return Value;
+}
+
+float GetProfileRadiusScale(FGBufferData GBufferData)
+{
+	// 0..255, which SubSurface profile to pick
+	uint SubsurfaceProfileInt = ExtractSubsurfaceProfileInt(GBufferData);
+
+	return GetKernel(SSSS_N_KERNELWEIGHTOFFSET + SSSS_N_KERNELWEIGHTCOUNT - 1, SubsurfaceProfileInt).a;
+}
+
 // from [REDACTED_PROFILE_URL]/separable-sss/tree/master/Demo
 // [REDACTED_PERSON] [REDACTED_PROFILE_URL]
 // [REDACTED_PROFILE_URL]
@@ -101,17 +128,82 @@ float GetSubsurfaceStrength(float2 UV)
 	error
 #endif
 
-// 0 / 1
-#define VISUALIZE_KERNEL 0
-
 bool InUnitBox(float2 UV)
 {
 	return UV.x >= 0 && UV.y >= 0 && UV.y < 1 && UV.y < 1;
 }
 
-float3 SetupSubsurfaceForOnePixel(float2 UV)
+// visualization (doesn't have to be fast)
+void VisualizePS(in float4 UVAndScreenPos : TEXCOORD0, out float4 OutColor : SV_Target0)
+{
+	float2 UV = UVAndScreenPos.xy;
+
+	OutColor = Texture2DSample(PostprocessInput0, PostprocessInput0Sampler, UV);
+
+	int2 PixelPos = (int2)(UVAndScreenPos.zw * ScreenPosToPixel.xy + ScreenPosToPixel.zw + 0.5f);
+
+	float2 ViewLocalUV = (PixelPos - View.ViewRectMin.xy) * View.ViewSizeAndSceneTexelSize.zw;
+
+	float2 IDAreaLocalUV = ViewLocalUV * 2 - 1.0f;
+
+	if (InUnitBox(IDAreaLocalUV))
+	{
+		FScreenSpaceData ScreenSpaceData = GetScreenSpaceData(View.ViewRectMin.xy * View.ViewSizeAndSceneTexelSize.zw + IDAreaLocalUV * (View.ViewSizeAndSceneTexelSize.xy * View.ViewSizeAndSceneTexelSize.zw));
+
+		int SubsurfaceProfileInt = ExtractSubsurfaceProfileInt(ScreenSpaceData.GBuffer);
+
+		OutColor = float4(0.5f, 0.5f, 0.5f, 0);
+
+		BRANCH if (ScreenSpaceData.GBuffer.ShadingModelID == SHADINGMODELID_SUBSURFACE_PROFILE)
+		{
+			if (SubsurfaceProfileInt == 0)
+			{
+				// default (no Profile)
+				OutColor = float4(0.8f, 0.7f, 0.6f, 0);
+			}
+			if (SubsurfaceProfileInt == 1)
+			{
+				OutColor = float4(1, 0, 0, 0);
+			}
+			if (SubsurfaceProfileInt == 2)
+			{
+				OutColor = float4(0, 1, 0, 0);
+			}
+			if (SubsurfaceProfileInt == 3)
+			{
+				OutColor = float4(0, 0, 1, 0);
+			}
+			if (SubsurfaceProfileInt == 4)
+			{
+				OutColor = float4(1, 0, 1, 0);
+			}
+			if (SubsurfaceProfileInt == 5)
+			{
+				OutColor = float4(0, 1, 1, 0);
+			}
+			if (SubsurfaceProfileInt == 6)
+			{
+				OutColor = float4(1, 1, 0, 0);
+			}
+			if (SubsurfaceProfileInt == 100)
+			{
+				OutColor = float4(0, 0.2f, 0, 0);
+			}
+			if (SubsurfaceProfileInt == 255)
+			{
+				OutColor = float4(1, 1, 1, 0);
+			}
+
+			int2 LeftTop = (PixelPos / 8) * 8;
+			PrintCharacter(PixelPos, OutColor.rgb, float3(1, 1, 1), LeftTop, SubsurfaceProfileInt);
+		}
+	}
+}
+
+// @return .RGB Color that should be scattared, .A:1 for subsurface scattering material, 0 for not
+float4 SetupSubsurfaceForOnePixel(float2 UV)
 {
-	float3 Ret = 0;
+	float4 Ret = 0;
 
 	FScreenSpaceData ScreenSpaceData = GetScreenSpaceData(UV);
 
@@ -119,11 +211,14 @@ float3 SetupSubsurfaceForOnePixel(float2 UV)
 	{
 		float4 SceneColor4 = Texture2DSample(PostprocessInput0, PostprocessInput0Sampler, UV);
 
+		// it's a valid sample
+		Ret.a = 1;
+
 #if SSSS_SPECULAR_CORRECTION == 1
 		// we take out the specular highlights
-		Ret = LightAccumulator_ReconstructDiffuseLighting(SceneColor4);
+		Ret.rgb = LightAccumulator_ReconstructDiffuseLighting(SceneColor4);
 #else
-		Ret = SceneColor4.rgb;
+		Ret.rgb = SceneColor4.rgb;
 #endif
 	}
 
@@ -134,89 +229,66 @@ void SetupPS(in float4 UVAndScreenPos : TEXCOORD0, out float4 OutColor : SV_Targ
 {
 	float2 UV = UVAndScreenPos.xy;
 
-	OutColor = 0;
+	// order aligned with Gather() hardware implementation
+	// RGB: color*A, A:weight 0 if no subsurface scattering
+	float4 A = SetupSubsurfaceForOnePixel(UV + float2(-0.5,  0.5f) * PostprocessInput0Size.zw);
+	float4 B = SetupSubsurfaceForOnePixel(UV + float2( 0.5,  0.5f) * PostprocessInput0Size.zw);
+	float4 C = SetupSubsurfaceForOnePixel(UV + float2( 0.5, -0.5f) * PostprocessInput0Size.zw);
+	float4 D = SetupSubsurfaceForOnePixel(UV + float2(-0.5, -0.5f) * PostprocessInput0Size.zw);
 
-	// if not VisualizeSSS
-	if (SETUP_MODE != 2)
-	{
-#if SETUP_HALF_RES
-		OutColor.rgb += SetupSubsurfaceForOnePixel(UV + float2(-0.5, -0.5f) * PostprocessInput0Size.zw);
-		OutColor.rgb += SetupSubsurfaceForOnePixel(UV + float2( 0.5, -0.5f) * PostprocessInput0Size.zw);
-		OutColor.rgb += SetupSubsurfaceForOnePixel(UV + float2(-0.5,  0.5f) * PostprocessInput0Size.zw);
-		OutColor.rgb += SetupSubsurfaceForOnePixel(UV + float2( 0.5,  0.5f) * PostprocessInput0Size.zw);
+	OutColor = (A + B) + (C + D);
 
-		OutColor *= 0.25f;
-#else
-		OutColor.rgb = SetupSubsurfaceForOnePixel(UV);
-#endif
-	}
-	else
-	{
-		// visualization (doesn't have to be fast)
-		OutColor = Texture2DSample(PostprocessInput0, PostprocessInput0Sampler, UV);
+	float Div = 1.0f / max(OutColor.a, 0.00001f);
 
-		int2 PixelPos = (int2)(UVAndScreenPos.zw * ScreenPosToPixel.xy + ScreenPosToPixel.zw + 0.5f);
+	OutColor.rgba *= Div;
+	
+	float4 FourDepth = GatherSceneDepth(UV, PostprocessInput0Size.zw);
 
-		float2 ViewLocalUV = (PixelPos - View.ViewRectMin.xy) * View.ViewSizeAndSceneTexelSize.zw;
+	// average all valid depth values to a single one
+	float SingleDepth = dot(FourDepth, float4(A.a, B.a, C.a, D.a)) * Div;
 
+	OutColor.a = SingleDepth;
 
-		float2 IDAreaLocalUV = ViewLocalUV * 2 - 1.0f;
+#if 0
+	float SSSScaleX = SSSParams.x;
 
-//		OutColor = float4(IDAreaLocalUV.xy, 0, 0);		return;
+	float scale = SSSScaleX / SingleDepth;
+//	float scale = SSSScaleX / min(min(FourDepth.x, FourDepth.y), min(FourDepth.z, FourDepth.w));
 
-		if (InUnitBox(IDAreaLocalUV))
-		{
-			FScreenSpaceData ScreenSpaceData = GetScreenSpaceData(View.ViewRectMin.xy * View.ViewSizeAndSceneTexelSize.zw + IDAreaLocalUV * (View.ViewSizeAndSceneTexelSize.xy * View.ViewSizeAndSceneTexelSize.zw));
+	// the viewport is only a fraction of the buffersize, here we compensate for that
+	// *1024 as we store the sample distances / 1024 to allows for scaling
+	float HorizontalScaler = View.ViewSizeAndSceneTexelSize.x * View.ViewSizeAndSceneTexelSize.z * 1024.0f;
 
-			int SubsurfaceProfileInt = ExtractSubsurfaceProfileInt(ScreenSpaceData.GBuffer);
+	// Calculate the final step to fetch the surrounding pixels:
+	float finalStep = scale * HorizontalScaler;
+	
+	FScreenSpaceData ScreenSpaceData = GetScreenSpaceData(UV);
 
-			OutColor = float4(0.5f, 0.5f, 0.5f, 0);
+//	finalStep *= SSSS_STREGTH_SOURCE(UV); // Modulate it using the alpha channel.
+	finalStep *= GetProfileRadiusScale(ScreenSpaceData.GBuffer);
 
-			BRANCH if (ScreenSpaceData.GBuffer.ShadingModelID == SHADINGMODELID_SUBSURFACE_PROFILE)
-			{
-				if (SubsurfaceProfileInt == 0)
-				{
-					// default (no Profile)
-					OutColor = float4(0.8f, 0.7f, 0.6f, 0);
-				}
-				if (SubsurfaceProfileInt == 1)
-				{
-					OutColor = float4(1, 0, 0, 0);
-				}
-				if (SubsurfaceProfileInt == 2)
-				{
-					OutColor = float4(0, 1, 0, 0);
-				}
-				if (SubsurfaceProfileInt == 3)
-				{
-					OutColor = float4(0, 0, 1, 0);
-				}
-				if (SubsurfaceProfileInt == 4)
-				{
-					OutColor = float4(1, 0, 1, 0);
-				}
-				if (SubsurfaceProfileInt == 5)
-				{
-					OutColor = float4(0, 1, 1, 0);
-				}
-				if (SubsurfaceProfileInt == 6)
-				{
-					OutColor = float4(1, 1, 0, 0);
-				}
-				if (SubsurfaceProfileInt == 100)
-				{
-					OutColor = float4(0, 0.2f, 0, 0);
-				}
-				if (SubsurfaceProfileInt == 255)
-				{
-					OutColor = float4(1, 1, 1, 0);
-				}
-
-				int2 LeftTop = (PixelPos / 8) * 8;
-				PrintCharacter(PixelPos, OutColor.rgb, float3(1, 1, 1), LeftTop, SubsurfaceProfileInt);
-			}
-		}
+	// in full resolution
+	float PixelFracationThreshold = 3.0f;
+
+	// *0.5f as we read in half res in the blur pass
+	if(finalStep < PixelFracationThreshold * PostprocessInput0Size.z * 0.5f)
+	{
+		// very small radius doesn't require work
+		OutColor = 0;
+	}
+#endif
+
+	// debug the kernel size and shape
+	/*
+	if(length((UV - 0.5f) * float2(2,1)) > 0.005f)
+	{
+		OutColor.rgb = 0;
+	}
+	else
+	{
+		OutColor.rgb *= 50;
 	}
+	*/
 }
 
 
@@ -225,40 +297,36 @@ void MainPS(float4 UVAndScreenPos : TEXCOORD0, out float4 OutColor : SV_Target0)
 {
 	float2 UV = UVAndScreenPos.xy;
 
-
 	// call into "SeparableSSS.usf"
 
-	// in world units, *0.01 to scale from units(cm) to m
-	float sssWidth = SSSParams.x;
-
 	// the viewport is only a fraction of the buffersize, here we compensate for that
-	// *1000 as we store the sample distances / 1000 to allows for scaling
-	float Scaler = View.ViewSizeAndSceneTexelSize.x * View.ViewSizeAndSceneTexelSize.z * 1000.0f;
+	// *1024 as we store the sample distances / 1024 to allows for scaling
 
 #if SSS_DIRECTION == 0
 	// horizontal
-	float2 Direction = float2(1, 0) * Scaler;
+	float2 Direction = float2(1, 0) * View.ViewSizeAndSceneTexelSize.x * View.ViewSizeAndSceneTexelSize.z * 1024.0f;
 #else
 	// vertical
-	float2 Direction = float2(0, View.ViewSizeAndSceneTexelSize.w / View.ViewSizeAndSceneTexelSize.z) * Scaler;
+	float2 Direction = float2(0, 1) * View.ViewSizeAndSceneTexelSize.w * View.ViewSizeAndSceneTexelSize.x * 1024.0f;
 #endif
 
-	OutColor.rgb = SSSSBlurPS(UV, sssWidth, Direction, false).rgb;
+	OutColor = SSSSBlurPS(UV, Direction, false);
 
-	OutColor.a = 1.0f;
+#if SSS_DIRECTION == 1
+	// second pass prepares the setup fro the recombine pass which doesn't need depth but wants to reconstruct the color
+	OutColor.a = ComputeMaskFromDepthInAlpha(OutColor.a);
+#endif
 }
 
 
 
 
-
+// Recombines the half res Subsurface filtered lighting contribution (upsampled and renormalized with the alpha)
+// with the SceneColor.
 void SubsurfaceRecombinePS(float4 UVAndScreenPos : TEXCOORD0, out float4 OutColor : SV_Target0)
 {
 	float2 UV = UVAndScreenPos.xy;
 
-	OutColor.rgb = Texture2DSample(PostprocessInput0, PostprocessInput0Sampler, UV).rgb;
-	OutColor.a = 1.0f;
-
 	// recombine with the scene color
 
 	// can be optimized
@@ -267,17 +335,61 @@ void SubsurfaceRecombinePS(float4 UVAndScreenPos : TEXCOORD0, out float4 OutColo
 
 	float4 SceneColor4 = Texture2DSample(PostprocessInput1, PostprocessInput1Sampler, UVSceneColor);
 
-	float4 SSSColor = OutColor;
-
 	FScreenSpaceData ScreenSpaceData = GetScreenSpaceData(UVSceneColor);
 
-	SSSColor.a = (ScreenSpaceData.GBuffer.ShadingModelID >= SHADINGMODELID_SUBSURFACE_PROFILE);
+	// bilinear filtering
+	float3 SSSColor;
+	{
+		float4 SSSColorWithAlpha = Texture2DSample(PostprocessInput0, PostprocessInput0Sampler, UV);
+
+		// renormalize to dilate RGB to fix haf res upsampling artifacts
+		SSSColor = SSSColorWithAlpha.rgb / max(SSSColorWithAlpha.a, 0.00001f);
+
+		// 0..255, which SubSurface profile to pick
+		uint SubsurfaceProfileInt = ExtractSubsurfaceProfileInt(ScreenSpaceData.GBuffer);
+
+		// add center sample from full resolution input for better quality
+		SSSColor += SceneColor4.rgb * GetKernel(SSSS_N_KERNELWEIGHTOFFSET, SubsurfaceProfileInt).rgb;
+	}
+
+
+
+	// 0: dont' fade in subsurface half res,  1:fully fade in half res
+	float BlendFactor = (ScreenSpaceData.GBuffer.ShadingModelID >= SHADINGMODELID_SUBSURFACE_PROFILE);
+
+#if 0
+	// fade out subsurface scattering if radius is too small to be more crips (not blend with half resolution)
+	// minro quality improvemenst (faces are more detailed in distance)
+	{
+		float SSSScaleX = SSSParams.x;
+
+		float scale = SSSScaleX / CalcSceneDepth(UV);
+
+		// the viewport is only a fraction of the buffersize, here we compensate for that
+		// *1024 as we store the sample distances / 1000 to allows for scaling
+		float HorizontalScaler = View.ViewSizeAndSceneTexelSize.x * View.ViewSizeAndSceneTexelSize.z * 1024.0f;
+
+		// Calculate the final step to fetch the surrounding pixels:
+		float finalStep = scale * HorizontalScaler;
+//		finalStep *= SSSS_STREGTH_SOURCE(UV); // Modulate it using the alpha channel (0..1), might only make a difference in very few cases
+		finalStep *= GetProfileRadiusScale(ScreenSpaceData.GBuffer);
+
+		float PixelSizeRadius = finalStep / (PostprocessInput1Size.z * 0.5f);
+
+		// tweaked for skin, a more flat kernel might need a smaller value, around 2 seems reasonable because we do half res
+		const float PixelSize = 4.0f;
+
+		BlendFactor *= saturate(PixelSizeRadius - PixelSize);
+	}
+#endif
+
 
 #if SSS_RECOMBINE_METHOD
 	// we took the specular highlights out, now we add them back in
-	SSSColor.rgb += LightAccumulator_ReconstructNonDiffuseLighting(SceneColor4);
+//	SSSColor += LightAccumulator_ReconstructNonDiffuseLighting(SceneColor4);
+	SSSColor += LightAccumulator_ReconstructNonDiffuseLighting(SceneColor4);
 #endif
 
 	// alpha channel can be anything - we consumed the value and likely output to a format that doesn't need the value
-	OutColor = float4(lerp(SceneColor4.rgb, SSSColor.rgb, SSSColor.a), 0);
+	OutColor = float4(lerp(SceneColor4.rgb, SSSColor, BlendFactor), 0);
 }
diff --git a/Engine/Shaders/SeparableSSS.usf b/Engine/Shaders/SeparableSSS.usf
--- a/Engine/Shaders/SeparableSSS.usf
+++ b/Engine/Shaders/SeparableSSS.usf
@@ -118,229 +118,38 @@
 #define SSSS_STREGTH_SOURCE (colorM.a)
 #endif
 
-/**
- * If SSSS_N_KERNELWEIGHTCOUNT is defined at this point, a custom filter kernel must be
- * set by the runtime.
- */
-#ifdef SSSS_N_KERNELWEIGHTCOUNT
-/**
- * Filter kernel layout is as follows:
- *   - Weights in the RGB channels.
- *   - Offsets in the A channel.
- */
-float4 kernel[SSSS_N_KERNELWEIGHTCOUNT]; 
-#else
-/**
- * Here you have ready-to-use kernels for quickstarters. Three kernels are 
- * readily available, with varying quality.
- * To create new kernels take a look into SSS::calculateKernel, or simply
- * push CTRL+C in the demo to copy the customized kernel into the clipboard.
- *
- * Note: these preset kernels are not used by the demo. They are calculated on
- * the fly depending on the selected values in the interface, by directly using
- * SSS::calculateKernel.
- *
- * Quality ranges from 0 to 2, being 2 the highest quality available.
- * The quality is with respect to 1080p; for 720p Quality=0 suffices.
- */
-#define SSSS_QUALITY 2	// todo: expose
-
-#if SSSS_QUALITY == 2
-#define SSSS_N_KERNELWEIGHTCOUNT 13
-/*const static float4 kernel[] = {
-    float4(0.530605, 0.613514, 0.739601, 0),
-    float4(0.0211412, 0.0459286, 0.0378196, 0.0208333),
-    float4(0.0402784, 0.0657244, 0.04631, 0.0833333),
-    float4(0.0493588, 0.0367726, 0.0219485, 0.1875),
-    float4(0.0410172, 0.0199899, 0.0118481, 0.333333),
-    float4(0.0263642, 0.0119715, 0.00684598, 0.520833),
-    float4(0.017924, 0.00711691, 0.00347194, 0.75),
-    float4(0.0128496, 0.00356329, 0.00132016, 1.02083),
-    float4(0.0094389, 0.00139119, 0.000416598, 1.33333),
-    float4(0.00700976, 0.00049366, 0.000151938, 1.6875),
-    float4(0.00500364, 0.00020094, 5.28848e-005, 2.08333),
-    float4(0.00333804, 7.85443e-005, 1.2945e-005, 2.52083),
-    float4(0.000973794, 1.11862e-005, 9.43437e-007, 3),
-};*/
-const static float4 kernel[] = {
-	float4(0.540384, 0.622941, 0.739601, 0),
-	float4(0.0207007, 0.0448084, 0.0378196, 0.0208333),
-	float4(0.0394393, 0.0641214, 0.04631, 0.0833333),
-	float4(0.0483305, 0.0358757, 0.0219485, 0.1875),
-	float4(0.0401626, 0.0195023, 0.0118481, 0.333333),
-	float4(0.0258149, 0.0116795, 0.00684598, 0.520833),
-	float4(0.0175505, 0.00694333, 0.00347194, 0.75),
-	float4(0.0125819, 0.00347638, 0.00132016, 1.02083),
-	float4(0.00924226, 0.00135726, 0.000416598, 1.33333),
-	float4(0.00686372, 0.000481619, 0.000151938, 1.6875),
-	float4(0.0048994, 0.000196039, 5.28848e-005, 2.08333),
-	float4(0.0032685, 7.66286e-005, 1.2945e-005, 2.52083),
-	float4(0.000953506, 1.09134e-005, 9.43437e-007, 3),
-};
-#elif SSSS_QUALITY == 1
-#define SSSS_N_KERNELWEIGHTCOUNT 9
-const static float4 kernel[] = {
-    float4(0.536343, 0.624624, 0.748867, 0),
-    float4(0.0324462, 0.0656718, 0.0532821, 0.03125),
-    float4(0.0582416, 0.0659959, 0.0411329, 0.125),
-    float4(0.0571056, 0.0287432, 0.0172844, 0.28125),
-    float4(0.0347317, 0.0151085, 0.00871983, 0.5),
-    float4(0.0216301, 0.00794618, 0.00376991, 0.78125),
-    float4(0.0144609, 0.00317269, 0.00106399, 1.125),
-    float4(0.0100386, 0.000914679, 0.000275702, 1.53125),
-    float4(0.00317394, 0.000134823, 3.77269e-005, 2),
-};
-#elif SSSS_QUALITY == 0
-#define SSSS_N_KERNELWEIGHTCOUNT 6
-const static float4 kernel[] = {
-    float4(0.560479, 0.669086, 0.784728, 0),
-    float4(0.0771802, 0.113491, 0.0793803, 0.08),
-    float4(0.0821904, 0.0358608, 0.0209261, 0.32),
-    float4(0.03639, 0.0130999, 0.00643685, 0.72),
-    float4(0.0192831, 0.00282018, 0.00084214, 1.28),
-    float4(0.00471691, 0.000184771, 5.07565e-005, 2),
-};
-#else
-#error Quality must be one of {0,1,2}
-#endif
-#endif
-
-float4 Quantize8Bit(float4 In)
-{
-	return floor(In * 255.999f) / 255.0f;
-}
-
-float4 GetKernel(uint SampleIndex, uint SubsurfaceProfileInt)
-{
-	const float4 TableMax = float4(1, 1, 1, 3.0f);
-
-	// profiled on NV670 fullscreen, one pass (total: 2 pass), samples: 13+1+13, large object filling the screen
-
-	// texture lookup
-	float4 Value = SSProfilesTexture.Load(int3(SampleIndex, SubsurfaceProfileInt, 0)) * TableMax;		// 0.88ms for 8bit
-	// simulated  8 bit lookup
-//	float4 Value = Quantize8Bit(kernel[SampleIndex] / TableMax) * TableMax;			// 1.33ms same look as texture
-	// float reference
-//	float4 Value = kernel[SampleIndex];												// 0.85ms, quite a bit different in look, need to use more than 8 bit?
-
-	return Value;
-}
-
 //-----------------------------------------------------------------------------
 // Porting Functions
-
-#if SSSS_HLSL_3 == 1
-#define SSSSTexture2D sampler2D
-#define SSSSSampleLevelZero(tex, coord) tex2Dlod(tex, float4(coord, 0.0, 0.0))
-#define SSSSSampleLevelZeroPoint(tex, coord) tex2Dlod(tex, float4(coord, 0.0, 0.0))
-#define SSSSSample(tex, coord) tex2D(tex, coord)
-#define SSSSSamplePoint(tex, coord) tex2D(tex, coord)
-#define SSSSSampleLevelZeroOffset(tex, coord, offset) tex2Dlod(tex, float4(coord + offset * SSSS_PIXEL_SIZE, 0.0, 0.0))
-#define SSSSSampleOffset(tex, coord, offset) tex2D(tex, coord + offset * SSSS_PIXEL_SIZE)
-#define SSSSLerp(a, b, t) lerp(a, b, t)
-#define SSSSSaturate(a) saturate(a)
-#define SSSSMad(a, b, c) mad(a, b, c)
-#define SSSSMul(v, m) mul(v, m)
-#define SSSS_FLATTEN [flatten]
-#define SSSS_BRANCH [branch]
-#define SSSS_UNROLL [unroll]
-#endif
-#if SSSS_HLSL_4 == 1
 #define SSSSTexture2D Texture2D
 #define SSSSSampleLevelZero(tex, coord) tex.SampleLevel(LinearSampler, coord, 0)
 #define SSSSSampleLevelZeroPoint(tex, coord) tex.SampleLevel(PointSampler, coord, 0)
-
-// EPIC changed start
 #define SSSSSampleShadowmap(coord)	float4(0,0,0,0)	// todo: not used yet
 #define SSSSSampleMasking(texcoord)	(GetScreenSpaceData(texcoord).GBuffer.ShadingModelID == SHADINGMODELID_SUBSURFACE_PROFILE)	// todo: masking doesn't work with filtering ideally put in alpha to get it filtered
-#define SSSSSampleSceneColor(coord) (SSSSSampleMasking(coord) * Texture2DSample(PostprocessInput0, PostprocessInput0Sampler, coord))
+#define SSSSSampleSceneColor(coord) Texture2DSample(PostprocessInput0, PostprocessInput0Sampler, coord)
 #define SSSSSampleSceneColorPoint(coord) ( Texture2DSample(PostprocessInput0, PostprocessInput0Sampler, coord))		// todo: could be made point
-#define SSSSSampleSceneDepth(coord) CalcSceneDepth(coord)		// todo: could be made linear
-#define SSSSSampleSceneDepthPoint(coord) CalcSceneDepth(coord)
-// EPIC changed end
-
 #define SSSSSampleLevelZeroOffset(tex, coord, offset) tex.SampleLevel(LinearSampler, coord, 0, offset)
 #define SSSSSampleOffset(tex, coord, offset) SSSSSampleLevelZeroOffset(tex, coord, offset)
 #define SSSSLerp(a, b, t) lerp(a, b, t)
-#define SSSSSaturate(a) saturate(a)
 #define SSSSMad(a, b, c) mad(a, b, c)
 #define SSSSMul(v, m) mul(v, m)
-#define SSSS_FLATTEN [flatten]
-#define SSSS_BRANCH [branch]
-#define SSSS_UNROLL [unroll]
-#endif
-#if SSSS_GLSL_3 == 1
-#define SSSSTexture2D sampler2D
-#define SSSSSampleLevelZero(tex, coord) textureLod(tex, coord, 0.0)
-#define SSSSSampleLevelZeroPoint(tex, coord) textureLod(tex, coord, 0.0)
-#define SSSSSample(tex, coord) texture(tex, coord)
-#define SSSSSamplePoint(tex, coord) texture(tex, coord)
-#define SSSSSampleLevelZeroOffset(tex, coord, offset) textureLodOffset(tex, coord, 0.0, offset)
-#define SSSSSampleOffset(tex, coord, offset) texture(tex, coord, offset)
-#define SSSSLerp(a, b, t) mix(a, b, t)
-#define SSSSSaturate(a) clamp(a, 0.0, 1.0)
-#define SSSSMad(a, b, c) (a * b + c)
-#define SSSSMul(v, m) (m * v)
-#define SSSS_FLATTEN
-#define SSSS_BRANCH
-#define SSSS_UNROLL
-#define float2 vec2
-#define float3 vec3
-#define float4 vec4
-#define int2 ivec2
-#define int3 ivec3
-#define int4 ivec4
-#define float4x4 mat4x4
-#endif
+#define SSSS_FLATTEN FLATTEN
+#define SSSS_BRANCH BRANCH
+#define SSSS_UNROLL UNROLL
 
 
 //-----------------------------------------------------------------------------
 // Separable SSS Transmittance Function
 
-float3 SSSSTransmittance(
-        /**
-         * This parameter allows to control the transmittance effect. Its range
-         * should be 0..1. Higher values translate to a stronger effect.
-         */
-        float translucency,
-
-        /**
-         * This parameter should be the same as the 'SSSSBlurPS' one. See below
-         * for more details.
-         */
-        float sssWidth,
-
-        /**
-         * Position in world space.
-         */
-        float3 worldPosition,
-
-        /**
-         * Normal in world space.
-         */
-        float3 worldNormal,
-
-        /**
-         * Light vector: lightWorldPosition - worldPosition.
-         */
-        float3 light,
-		
-// EPIC changed start
-        /**
-         * Linear 0..1 shadow map.
-         */
-//        SSSSTexture2D shadowMap,
-// EPIC changed end
-
-        /**
-         * Regular world to light space matrix.
-         */
-        float4x4 lightViewProjection,
-
-        /**
-         * Far plane distance used in the light projection matrix.
-         */
-        float lightFarPlane) {
+// @param translucency This parameter allows to control the transmittance effect. Its range should be 0..1. Higher values translate to a stronger effect.
+// @param sssWidth this parameter should be the same as the 'SSSSBlurPS' one. See below for more details.
+// @param worldPosition Position in world space.
+// @param worldNormal Normal in world space.
+// @param light Light vector: lightWorldPosition - worldPosition.
+// @param lightViewProjection Regular world to light space matrix.
+// @param lightFarPlane Far plane distance used in the light projection matrix.
+
+float3 SSSSTransmittance(float translucency, float sssWidth, float3 worldPosition, float3 worldNormal, float3 light, float4x4 lightViewProjection, float lightFarPlane)
+ {
     /**
      * Calculate the scale of the effect.
      */
@@ -378,104 +187,62 @@ float3 SSSSTransmittance(
      * Using the profile, we finally approximate the transmitted lighting from
      * the back of the object:
      */
-    return profile * SSSSSaturate(0.3 + dot(light, -worldNormal));
-}
-
-//-----------------------------------------------------------------------------
-// Separable SSS Reflectance Vertex Shader
-
-void SSSSBlurVS(float4 position,
-                out float4 svPosition,
-                inout float2 texcoord) {
-    svPosition = position;
+    return profile * saturate(0.3 + dot(light, -worldNormal));
 }
 
 //-----------------------------------------------------------------------------
 // Separable SSS Reflectance Pixel Shader
 
-float4 SSSSBlurPS(
-        /**
-         * The usual quad texture coordinates.
-         */
-        float2 texcoord,
-
-// EPIC changed start
-        /**
-         * This is a SRGB or HDR color input buffer, which should be the final
-         * color frame, resolved in case of using multisampling. The desired
-         * SSS strength should be stored in the alpha channel (1 for full
-         * strength, 0 for disabling SSS). If this is not possible, you an
-         * customize the source of this value using SSSS_STREGTH_SOURCE.
-         *
-         * When using non-SRGB buffers, you
-         * should convert to linear before processing, and back again to gamma
-         * space before storing the pixels (see Chapter 24 of GPU Gems 3 for
-         * more info)
-         *
-         * IMPORTANT: WORKING IN A NON-LINEAR SPACE WILL TOTALLY RUIN SSS!
-         */
-//        SSSSTexture2D colorTex,
-		/**
-         * The linear depth buffer of the scene, resolved in case of using
-         * multisampling. The resolve should be a simple average to avoid
-         * artifacts in the silhouette of objects.
-         */
-//        SSSSTexture2D depthTex,
-// EPIC changed end
-
-        /**
-         * This parameter specifies the global level of subsurface scattering
-         * or, in other words, the width of the filter. It's specified in
-         * world space units.
-         */
-        float sssWidth,
-
-        /**
-         * Direction of the blur:
-         *   - First pass:   float2(1.0, 0.0)
-         *   - Second pass:  float2(0.0, 1.0)
-         */
-        float2 dir,
-
-        /**
-         * This parameter indicates whether the stencil buffer should be
-         * initialized. Should be set to 'true' for the first pass if not
-         * previously initialized, to enable optimization of the second
-         * pass.
-         */
-        bool initStencil) {
-
+// @param texcoord The usual quad texture coordinates.
+// @param dir Direction of the blur: First pass:   float2(1.0, 0.0), Second pass:  float2(0.0, 1.0)
+// @param initStencil indicates whether the stencil buffer should be initialized. Should be set to 'true' for the first pass if not previously initialized, to enable optimization of the second pass.
+float4 SSSSBlurPS(float2 texcoord, float2 dir, bool initStencil)
+{
     // Fetch color of current pixel:
     float4 colorM = SSSSSampleSceneColorPoint(texcoord);
 
+	// we store the depth in alpha
+	float depthM = colorM.a;
+
+	colorM.a = ComputeMaskFromDepthInAlpha(colorM.a);
+
+	// we don't need to process pixels that are not part of the subsurafce scattering (optimization, also prevents divide by later on)
+	BRANCH if(!colorM.a)
+	{
+		// todo: need to check for proper clear
+//		discard;
+		return 0.0f;
+	}
+
+	float SSSStrength = SSSS_STREGTH_SOURCE(texcoord);
+
     // Initialize the stencil buffer in case it was not already available:
     if (initStencil) // (Checked in compile time, it's optimized away)
-        if (SSSS_STREGTH_SOURCE == 0.0) discard;
-
-    // Fetch linear depth of current pixel:
-    float depthM = SSSSSampleSceneDepthPoint(texcoord).r;
+        if (SSSStrength == 0.0) discard;
 
-    // Calculate the sssWidth scale (1.0 for a unit plane sitting on the
-    // projection window):
-//    float distanceToProjectionWindow = 1.0 / tan(0.5 * radians(SSSS_FOVY));
-	float distanceToProjectionWindow = SSSParams.y;
-//	float distanceToProjectionWindow = 1.0 / tan(0.5 * radians(90.0f));
+	float SSSScaleX = SSSParams.x;
 
-	float scale = distanceToProjectionWindow / depthM;
+	float scale = SSSScaleX / depthM;
 
     // Calculate the final step to fetch the surrounding pixels:
-    float2 finalStep = sssWidth * scale * dir * 0.5f;										// hacked in 0.5 - -1..1 to 0..1 but why this isn't in demo code?
-	finalStep *= SSSS_STREGTH_SOURCE; // Modulate it using the alpha channel.
-    finalStep *= 1.0 / 3.0; // Divide by 3 as the kernels range from -3 to 3.
+    float2 finalStep = scale * dir;
 
+	finalStep *= SSSStrength; // Modulate it using the alpha channel (0..1 range)
+    
 	FGBufferData GBufferData = GetGBufferData(texcoord);
 
 	// 0..255, which SubSurface profile to pick
 	uint SubsurfaceProfileInt = ExtractSubsurfaceProfileInt(GBufferData);
 
     // Accumulate the center sample:
-    float4 colorBlurred = colorM;
-	colorBlurred.rgb *= GetKernel(0, SubsurfaceProfileInt).rgb;
+    float4 colorBlurred = 0;
+
+	// center sample is added in Recombine pass
+	colorBlurred.a = colorM.a;
+
+#if SSS_DIRECTION == 0
+	colorBlurred.rgb = colorM * GetKernel(SSSS_N_KERNELWEIGHTOFFSET, SubsurfaceProfileInt).rgb;
+#endif
 
 	// Accumulate the other samples:
     SSSS_UNROLL
@@ -484,8 +251,9 @@ float4 SSSSBlurPS(
 		// Kernel.a = 0..3
 		half4 Kernel = GetKernel(SSSS_N_KERNELWEIGHTOFFSET + i, SubsurfaceProfileInt);
 
-		float3 LocalAccum = 0;
+		float4 LocalAccum = 0;
 
+		float2 UVOffset = Kernel.a * finalStep;
 		// The kernel is symtrical, we want to use that property.
 		// Half the GetKernel() calls (more expensive if done by texture)
 		// Half the weighting computations (saves 3mul per lookup sample)
@@ -493,27 +261,39 @@ float4 SSSSBlurPS(
 		for (int Side = -1; Side <= 1; Side += 2)
 		{
 			// Fetch color and depth for current sample:
-			float2 offset = texcoord + Kernel.a * finalStep * Side;
-			float4 color = SSSSSampleSceneColor(offset);
+			float4 color = SSSSSampleSceneColor(texcoord + UVOffset * Side);
+
+			float LocalDepth = color.a;
+			
+			color.a = ComputeMaskFromDepthInAlpha(color.a);
+
+			// approximation, ideally we would reconstruct the mask with ComputeMaskFromDepthInAlpha() and do manual bilinear filter
+			// needed?
+//			color.rgb *= color.a;
+
 
 #if SSSS_FOLLOW_SURFACE == 1
-			// If the difference in depth is huge, we lerp color back to "colorM":
-			float depth = SSSSSampleSceneDepth(offset).r;
-			float s = SSSSSaturate(6000.0f / 400000 * distanceToProjectionWindow *
-				//        float s = SSSSSaturate(300.0f/400000 * distanceToProjectionWindow *
-				sssWidth * abs(depthM - depth));
-			color.rgb = SSSSLerp(color.rgb, colorM.rgb, s);
+			// If the difference in depth is huge, we weight the sample less or not at all
+			float s = saturate(12000.0f / 400000 * SSSParams.y *
+	//        float s = saturate(300.0f/400000 * SSSParams.y *
+				abs(depthM - LocalDepth));
+
+			color *= 1 - s;
 #endif
+
 			// Accumulate left and right 
-			LocalAccum += color.rgb;
+			LocalAccum += color;
 		}
 
 		// Accumulate to final value (left and right sample with the same weight)
-		colorBlurred.rgb += Kernel.rgb * LocalAccum;
+		colorBlurred += Kernel * LocalAccum;
 	}
 
-	// hack - normalize
-//	colorBlurred.rgb /= 1.0f - kernel[0].rgb;
+	// normalize (some samples are rejected because of depth or the other material is no SSS, compensate for that)
+	colorBlurred.rgb *= 1.0f / max(colorBlurred.a, 0.00001f);
+
+	// alpha stored the SceneDepth (0 if there is no subsurface scattering)
+	colorBlurred.a = depthM;
 
     return colorBlurred;
 }
diff --git a/Engine/Source/Runtime/Engine/Private/Rendering/SubsurfaceProfile.cpp b/Engine/Source/Runtime/Engine/Private/Rendering/SubsurfaceProfile.cpp
--- a/Engine/Source/Runtime/Engine/Private/Rendering/SubsurfaceProfile.cpp
+++ b/Engine/Source/Runtime/Engine/Private/Rendering/SubsurfaceProfile.cpp
@@ -180,7 +180,7 @@ void FSubsurfaceProfileTexture::CreateTexture(FRHICommandListImmediate& RHICmdLi
 			FVector4 C = kernel[Pos] * FLinearColor(1.0f / TableMaxRGB, 1.0f / TableMaxRGB, 1.0f / TableMaxRGB, 1.0f / TableMaxA);
 
 			// requires 16bit (could be made with 8 bit e.g. using sample0.w as 8bit scale applied to all samples (more multiplications in the shader))
-			C.W *= Data.ScatterRadius / 1000.0f;
+			C.W *= Data.ScatterRadius / 1024.0f;
 
 			if (b16Bit)
 			{
diff --git a/Engine/Source/Runtime/Renderer/Private/CompositionLighting/CompositionLighting.cpp b/Engine/Source/Runtime/Renderer/Private/CompositionLighting/CompositionLighting.cpp
--- a/Engine/Source/Runtime/Renderer/Private/CompositionLighting/CompositionLighting.cpp
+++ b/Engine/Source/Runtime/Renderer/Private/CompositionLighting/CompositionLighting.cpp
@@ -359,13 +359,13 @@ void FCompositionLighting::ProcessLighting(FRHICommandListImmediate& RHICmdList,
 				//@todo-rco: Remove this when we fix the cross-compiler
 				!IsOpenGLPlatform(View.GetShaderPlatform()))
 			{
-				FRenderingCompositePass* PassSetup = Context.Graph.RegisterPass(new(FMemStack::Get()) FRCPassPostProcessSubsurfaceSetup(false, true));
+				FRenderingCompositePass* PassSetup = Context.Graph.RegisterPass(new(FMemStack::Get()) FRCPassPostProcessSubsurfaceSetup());
 				PassSetup->SetInput(ePId_Input0, Context.FinalOutput);
 
-				FRenderingCompositePass* Pass0 = Context.Graph.RegisterPass(new(FMemStack::Get()) FRCPassPostProcessSubsurface(0, Radius, true));
+				FRenderingCompositePass* Pass0 = Context.Graph.RegisterPass(new(FMemStack::Get()) FRCPassPostProcessSubsurface(0));
 				Pass0->SetInput(ePId_Input0, PassSetup);
 
-				FRenderingCompositePass* Pass1 = Context.Graph.RegisterPass(new(FMemStack::Get()) FRCPassPostProcessSubsurface(1, Radius, true));
+				FRenderingCompositePass* Pass1 = Context.Graph.RegisterPass(new(FMemStack::Get()) FRCPassPostProcessSubsurface(1));
 				Pass1->SetInput(ePId_Input0, Pass0);
 
 				// full res composite pass, no blurring (Radius=0)
diff --git a/Engine/Source/Runtime/Renderer/Private/PostProcess/PostProcessSubsurface.cpp b/Engine/Source/Runtime/Renderer/Private/PostProcess/PostProcessSubsurface.cpp
--- a/Engine/Source/Runtime/Renderer/Private/PostProcess/PostProcessSubsurface.cpp
+++ b/Engine/Source/Runtime/Renderer/Private/PostProcess/PostProcessSubsurface.cpp
@@ -22,20 +22,101 @@ static TAutoConsoleVariable<int32> CVarSSSSampleSet(
 	TEXT("2: high quality (default)"),
 	ECVF_RenderThreadSafe  | ECVF_Scalability);
 
+static TAutoConsoleVariable<int> CVarSubsurfaceQuality(
+	TEXT("r.SubsurfaceQuality"),
+	1,
+	TEXT("Define the quality of the Screenspace subsurface scattering postprocess.\n")
+	TEXT(" 0: low quality for speculars on subsurface materials\n")
+	TEXT(" 1: higher quality as specular is separated before screenspace blurring (Only used if SceneColor has an alpha channel)"),
+	ECVF_Scalability | ECVF_RenderThreadSafe);
+
 static bool SubsurfaceShouldCache(EShaderPlatform Platform)
 {
 	//@todo-rco: Remove this when we fix the cross-compiler
 	return IsFeatureLevelSupported(Platform, ERHIFeatureLevel::SM4) && !IsOpenGLPlatform(Platform);
 }
 
+// -------------------------------------------------------------
+
+float GetSubsurfaceRadiusScale()
+{
+	static const auto CVar = IConsoleManager::Get().FindTConsoleVariableDataFloat(TEXT("r.SSS.Scale"));
+	check(CVar);
+	float Ret = CVar->GetValueOnRenderThread();
+
+	return FMath::Max(0.0f, Ret);
+}
+
+// -------------------------------------------------------------
+
+/** Shared shader parameters needed for screen space subsurface scattering. */
+class FSubsurfaceParameters
+{
+public:
+	void Bind(const FShaderParameterMap& ParameterMap)
+	{
+		SSSParams.Bind(ParameterMap, TEXT("SSSParams"));
+		SSProfilesTexture.Bind(ParameterMap, TEXT("SSProfilesTexture"));
+		SSProfilesTextureSampler.Bind(ParameterMap, TEXT("SSProfilesTextureSampler"));
+	}
+
+	void SetParameters(FRHICommandList& RHICmdList, const FPixelShaderRHIParamRef& ShaderRHI, const FRenderingCompositePassContext& Context) const
+	{
+		{
+			// from Separabale.usf: float distanceToProjectionWindow = 1.0 / tan(0.5 * radians(SSSS_FOVY))
+			// can be extracted out of projection matrix
+
+			float ScaleCorrectionX = Context.View.ViewRect.Width() / (float)GSceneRenderTargets.GetBufferSizeXY().X;
+			float ScaleCorrectionY = Context.View.ViewRect.Height() / (float)GSceneRenderTargets.GetBufferSizeXY().Y;
+
+			 // Divide by 3 as the kernels range from -3 to 3.
+			const float KernelSize = 3.0f;
+
+			// Calculate the sssWidth scale (1.0 for a unit plane sitting on the projection window):
+			float DistanceToProjectionWindow = Context.View.ViewMatrices.ProjMatrix.M[0][0]; 
+
+			float SSSScaleZ = DistanceToProjectionWindow * GetSubsurfaceRadiusScale();
+
+			// * 0.5f: hacked in 0.5 - -1..1 to 0..1 but why this isn't in demo code?
+			float SSSScaleX = SSSScaleZ / KernelSize * 0.5f;
+
+			FVector4 ColorScale(SSSScaleX, SSSScaleZ, ScaleCorrectionX, ScaleCorrectionY);
+			SetShaderValue(Context.RHICmdList, ShaderRHI, SSSParams, ColorScale);
+		}
+
+		{
+			const IPooledRenderTarget* PooledRT = GetSubsufaceProfileTexture_RT(Context.RHICmdList);
+
+			check(PooledRT);
+
+			const FSceneRenderTargetItem& Item = PooledRT->GetRenderTargetItem();
+
+			SetTextureParameter(Context.RHICmdList, ShaderRHI, SSProfilesTexture, SSProfilesTextureSampler, TStaticSamplerState<SF_Point, AM_Wrap, AM_Wrap, AM_Wrap>::GetRHI(), Item.ShaderResourceTexture);
+		}
+	}
+
+	friend FArchive& operator<<(FArchive& Ar,FSubsurfaceParameters& P)
+	{
+		Ar << P.SSSParams << P.SSProfilesTexture << P.SSProfilesTextureSampler;
+
+		return Ar;
+	}
+
+private:
+	FShaderParameter SSSParams;
+	FShaderResourceParameter SSProfilesTexture;
+	FShaderResourceParameter SSProfilesTextureSampler;
+};
+
+// ---------------------------------------------
+
+
 /**
- * Encapsulates the post processing ambient occlusion pixel shader.
- * @param SetupMode 0:without specular correction, 1: with specular correction, 2:visualize, 3: half res without specular correction, 4:half res with specular correction
+ * Encapsulates the post processing subsurface scattering pixel shader.
  */
-template <uint32 SetupMode>
-class FPostProcessSubsurfaceSetupPS : public FGlobalShader
+class FPostProcessSubsurfaceVisualizePS : public FGlobalShader
 {
-	DECLARE_SHADER_TYPE(FPostProcessSubsurfaceSetupPS , Global);
+	DECLARE_SHADER_TYPE(FPostProcessSubsurfaceVisualizePS , Global);
 
 	static bool ShouldCache(EShaderPlatform Platform)
 	{
@@ -45,24 +126,25 @@ class FPostProcessSubsurfaceSetupPS : public FGlobalShader
 	static void ModifyCompilationEnvironment(EShaderPlatform Platform, FShaderCompilerEnvironment& OutEnvironment)
 	{
 		FGlobalShader::ModifyCompilationEnvironment(Platform, OutEnvironment);
-		OutEnvironment.SetDefine(TEXT("SETUP_MODE"), SetupMode);
 	}
 
 	/** Default constructor. */
-	FPostProcessSubsurfaceSetupPS () {}
+	FPostProcessSubsurfaceVisualizePS () {}
 
 public:
 	FPostProcessPassParameters PostprocessParameter;
 	FDeferredPixelShaderParameters DeferredParameters;
 	FShaderResourceParameter MiniFontTexture;
+	FSubsurfaceParameters SubsurfaceParameters;
 
 	/** Initialization constructor. */
-	FPostProcessSubsurfaceSetupPS(const ShaderMetaType::CompiledShaderInitializerType& Initializer)
+	FPostProcessSubsurfaceVisualizePS(const ShaderMetaType::CompiledShaderInitializerType& Initializer)
 		: FGlobalShader(Initializer)
 	{
 		PostprocessParameter.Bind(Initializer.ParameterMap);
 		DeferredParameters.Bind(Initializer.ParameterMap);
 		MiniFontTexture.Bind(Initializer.ParameterMap, TEXT("MiniFontTexture"));
+		SubsurfaceParameters.Bind(Initializer.ParameterMap);
 	}
 
 	void SetParameters(const FRenderingCompositePassContext& Context)
@@ -74,13 +156,14 @@ public:
 		PostprocessParameter.SetPS(ShaderRHI, Context, TStaticSamplerState<SF_Point,AM_Clamp,AM_Clamp,AM_Clamp>::GetRHI());
 		DeferredParameters.Set(Context.RHICmdList, ShaderRHI, Context.View);
 		SetTextureParameter(Context.RHICmdList, ShaderRHI, MiniFontTexture, GEngine->MiniFontTexture ? GEngine->MiniFontTexture->Resource->TextureRHI : GSystemTextures.WhiteDummy->GetRenderTargetItem().TargetableTexture);
+		SubsurfaceParameters.SetParameters(Context.RHICmdList, ShaderRHI, Context);
 	}
 
 	// FShader interface.
 	virtual bool Serialize(FArchive& Ar)
 	{
 		bool bShaderHasOutdatedParameters = FGlobalShader::Serialize(Ar);
-		Ar << PostprocessParameter << DeferredParameters << MiniFontTexture;
+		Ar << PostprocessParameter << DeferredParameters << MiniFontTexture << SubsurfaceParameters;
 		return bShaderHasOutdatedParameters;
 	}
 
@@ -91,24 +174,16 @@ public:
 
 	static const TCHAR* GetFunctionName()
 	{
-		return TEXT("SetupPS");
+		return TEXT("VisualizePS");
 	}
 };
 
-// #define avoids a lot of code duplication
-#define VARIATION1(A) typedef FPostProcessSubsurfaceSetupPS<A> FPostProcessSubsurfaceSetupPS##A; \
-	IMPLEMENT_SHADER_TYPE2(FPostProcessSubsurfaceSetupPS##A, SF_Pixel);
-
-	VARIATION1(0) VARIATION1(1) VARIATION1(2) VARIATION1(3) VARIATION1(4)
+IMPLEMENT_SHADER_TYPE3(FPostProcessSubsurfaceVisualizePS, SF_Pixel);
 
-#undef VARIATION1
-
-
-template <uint32 SetupMode>
-void SetSubsurfaceSetupShader(const FRenderingCompositePassContext& Context)
+void SetSubsurfaceVisualizeShader(const FRenderingCompositePassContext& Context)
 {
 	TShaderMapRef<FPostProcessVS> VertexShader(Context.GetShaderMap());
-	TShaderMapRef<FPostProcessSubsurfaceSetupPS<SetupMode> > PixelShader(Context.GetShaderMap());
+	TShaderMapRef<FPostProcessSubsurfaceVisualizePS> PixelShader(Context.GetShaderMap());
 
 	static FGlobalBoundShaderState BoundShaderState;
 
@@ -118,41 +193,13 @@ void SetSubsurfaceSetupShader(const FRenderingCompositePassContext& Context)
 	VertexShader->SetParameters(Context);
 }
 
-static TAutoConsoleVariable<int> CVarSubsurfaceQuality(
-	TEXT("r.SubsurfaceQuality"),
-	1,
-	TEXT("Define the quality of the Screenspace subsurface scattering postprocess.\n")
-	TEXT(" 0: low quality for speculars on subsurface materials\n")
-	TEXT(" 1: higher quality as specular is separated before screenspace blurring (Only used if SceneColor has an alpha channel)"),
-	ECVF_Scalability | ECVF_RenderThreadSafe);
-
-
-static bool DoSpecularCorrection()
+FRCPassPostProcessSubsurfaceVisualize::FRCPassPostProcessSubsurfaceVisualize()
 {
-	bool CVarState = CVarSubsurfaceQuality.GetValueOnRenderThread() > 0;
-
-	int SceneColorFormat;
-	{
-		static IConsoleVariable* CVar = IConsoleManager::Get().FindConsoleVariable(TEXT("r.SceneColorFormat"));
-
-		SceneColorFormat = CVar->GetInt();
-	}
-
-	// we need an alpha channel for this feature
-	return CVarState && (SceneColorFormat >= 4);
-}
-
-FRCPassPostProcessSubsurfaceSetup::FRCPassPostProcessSubsurfaceSetup(bool bInVisualize, bool bInHalfRes)
-	: bVisualize(bInVisualize)
-	, bHalfRes(bInHalfRes)
-{
-	if (bVisualize)
-	{
-		GSceneRenderTargets.AdjustGBufferRefCount(1);
-	}
+	// we need the GBuffer, we release it Process()
+	GSceneRenderTargets.AdjustGBufferRefCount(1);
 }
 
-void FRCPassPostProcessSubsurfaceSetup::Process(FRenderingCompositePassContext& Context)
+void FRCPassPostProcessSubsurfaceVisualize::Process(FRenderingCompositePassContext& Context)
 {
 	SCOPED_DRAW_EVENT(Context.RHICmdList, SubsurfaceSetup, DEC_SCENE_ITEMS);
 
@@ -174,13 +221,7 @@ void FRCPassPostProcessSubsurfaceSetup::Process(FRenderingCompositePassContext&
 	uint32 ScaleFactor = GSceneRenderTargets.GetBufferSizeXY().X / SrcSize.X;
 
 	FIntRect SrcRect = View.ViewRect / ScaleFactor;
-	FIntRect DestRect = bHalfRes ? FIntRect::DivideAndRoundUp(SrcRect, 2) : SrcRect;
-
-	if(bHalfRes)
-	{
-		// upscale rectangle to not slightly scale
-		SrcRect = DestRect * 2;
-	}
+	FIntRect DestRect = SrcRect;
 
 	const FSceneRenderTargetItem& DestRenderTarget = PassOutputs[0].RequestSurface(Context);
 
@@ -197,38 +238,7 @@ void FRCPassPostProcessSubsurfaceSetup::Process(FRenderingCompositePassContext&
 	Context.RHICmdList.SetRasterizerState(TStaticRasterizerState<>::GetRHI());
 	Context.RHICmdList.SetDepthStencilState(TStaticDepthStencilState<false, CF_Always>::GetRHI());
 
-	if (bVisualize)
-	{
-		SetSubsurfaceSetupShader<2>(Context);
-	}
-	else
-	{
-		// reconstruct specular and add it in final pass
-		bool bSpecularCorrection = DoSpecularCorrection();
-
-		if(bHalfRes)
-		{
-			if(bSpecularCorrection)
-			{
-				SetSubsurfaceSetupShader<4>(Context);
-			}
-			else
-			{
-				SetSubsurfaceSetupShader<3>(Context);
-			}
-		}
-		else
-		{
-			if(bSpecularCorrection)
-			{
-				SetSubsurfaceSetupShader<1>(Context);
-			}
-			else
-			{
-				SetSubsurfaceSetupShader<0>(Context);
-			}
-		}
-	}
+	SetSubsurfaceVisualizeShader(Context);
 
 	// Draw a quad mapping scene color to the view's render target
 	TShaderMapRef<FPostProcessVS> VertexShader(Context.GetShaderMap());
@@ -243,8 +253,6 @@ void FRCPassPostProcessSubsurfaceSetup::Process(FRenderingCompositePassContext&
 		*VertexShader,
 		EDRF_UseTriangleOptimization);
 
-
-	if (bVisualize)
 	{
 		// this is a helper class for FCanvas to be able to get screen size
 		class FRenderTargetTemp : public FRenderTarget
@@ -291,29 +299,217 @@ void FRCPassPostProcessSubsurfaceSetup::Process(FRenderingCompositePassContext&
 
 	Context.RHICmdList.CopyToResolveTarget(DestRenderTarget.TargetableTexture, DestRenderTarget.ShaderResourceTexture, false, FResolveParams());
 
-	if (bVisualize)
-	{
-		GSceneRenderTargets.AdjustGBufferRefCount(-1);
-	}
+	// we no longer need the GBuffer
+	GSceneRenderTargets.AdjustGBufferRefCount(-1);
 }
 
-FPooledRenderTargetDesc FRCPassPostProcessSubsurfaceSetup::ComputeOutputDesc(EPassOutputId InPassOutputId) const
+FPooledRenderTargetDesc FRCPassPostProcessSubsurfaceVisualize::ComputeOutputDesc(EPassOutputId InPassOutputId) const
 {
 	FPooledRenderTargetDesc Ret = GSceneRenderTargets.GetSceneColor()->GetDesc();
 
 	Ret.Reset();
-	Ret.DebugName = TEXT("SubsurfaceSetup");
-	// we don't need alpha any more
-	Ret.Format = PF_FloatRGB;
+	Ret.DebugName = TEXT("SubsurfaceVisualize");
+	// alpha is used to store depth and renormalize (alpha==0 means there is no subsurface scattering)
+	Ret.Format = PF_FloatRGBA;
+
+	return Ret;
+}
+
+
+// ---------------------------------------------
+
+/**
+ * Encapsulates the post processing subsurface scattering pixel shader.
+ * @param SetupMode 0:without specular correction, 1:with specular correction
+ */
+template <uint32 SetupMode>
+class FPostProcessSubsurfaceSetupPS : public FGlobalShader
+{
+	DECLARE_SHADER_TYPE(FPostProcessSubsurfaceSetupPS , Global);
+
+	static bool ShouldCache(EShaderPlatform Platform)
+	{
+		return SubsurfaceShouldCache(Platform);
+	}
 
-	if(bHalfRes)
+	static void ModifyCompilationEnvironment(EShaderPlatform Platform, FShaderCompilerEnvironment& OutEnvironment)
 	{
-		Ret.Extent  = FIntPoint::DivideAndRoundUp(Ret.Extent, 2);
+		FGlobalShader::ModifyCompilationEnvironment(Platform, OutEnvironment);
+		OutEnvironment.SetDefine(TEXT("SETUP_MODE"), SetupMode);
+	}
+
+	/** Default constructor. */
+	FPostProcessSubsurfaceSetupPS () {}
+
+public:
+	FPostProcessPassParameters PostprocessParameter;
+	FDeferredPixelShaderParameters DeferredParameters;
+	FShaderResourceParameter MiniFontTexture;
+	FSubsurfaceParameters SubsurfaceParameters;
 
-		Ret.Extent.X = FMath::Max(1, Ret.Extent.X);
-		Ret.Extent.Y = FMath::Max(1, Ret.Extent.Y);
+	/** Initialization constructor. */
+	FPostProcessSubsurfaceSetupPS(const ShaderMetaType::CompiledShaderInitializerType& Initializer)
+		: FGlobalShader(Initializer)
+	{
+		PostprocessParameter.Bind(Initializer.ParameterMap);
+		DeferredParameters.Bind(Initializer.ParameterMap);
+		SubsurfaceParameters.Bind(Initializer.ParameterMap);
 	}
 
+	void SetParameters(const FRenderingCompositePassContext& Context)
+	{
+		const FFinalPostProcessSettings& Settings = Context.View.FinalPostProcessSettings;
+		const FPixelShaderRHIParamRef ShaderRHI = GetPixelShader();
+
+		FGlobalShader::SetParameters(Context.RHICmdList, ShaderRHI, Context.View);
+		PostprocessParameter.SetPS(ShaderRHI, Context, TStaticSamplerState<SF_Point,AM_Clamp,AM_Clamp,AM_Clamp>::GetRHI());
+		DeferredParameters.Set(Context.RHICmdList, ShaderRHI, Context.View);
+		SubsurfaceParameters.SetParameters(Context.RHICmdList, ShaderRHI, Context);
+	}
+
+	// FShader interface.
+	virtual bool Serialize(FArchive& Ar)
+	{
+		bool bShaderHasOutdatedParameters = FGlobalShader::Serialize(Ar);
+		Ar << PostprocessParameter << DeferredParameters << SubsurfaceParameters;
+		return bShaderHasOutdatedParameters;
+	}
+
+	static const TCHAR* GetSourceFilename()
+	{
+		return TEXT("PostProcessSubsurface");
+	}
+
+	static const TCHAR* GetFunctionName()
+	{
+		return TEXT("SetupPS");
+	}
+};
+
+// #define avoids a lot of code duplication
+#define VARIATION1(A) typedef FPostProcessSubsurfaceSetupPS<A> FPostProcessSubsurfaceSetupPS##A; \
+	IMPLEMENT_SHADER_TYPE2(FPostProcessSubsurfaceSetupPS##A, SF_Pixel);
+
+	VARIATION1(0) VARIATION1(1)
+
+#undef VARIATION1
+
+
+template <uint32 SetupMode>
+void SetSubsurfaceSetupShader(const FRenderingCompositePassContext& Context)
+{
+	TShaderMapRef<FPostProcessVS> VertexShader(Context.GetShaderMap());
+	TShaderMapRef<FPostProcessSubsurfaceSetupPS<SetupMode> > PixelShader(Context.GetShaderMap());
+
+	static FGlobalBoundShaderState BoundShaderState;
+
+	SetGlobalBoundShaderState(Context.RHICmdList, Context.GetFeatureLevel(), BoundShaderState, GFilterVertexDeclaration.VertexDeclarationRHI, *VertexShader, *PixelShader);
+
+	PixelShader->SetParameters(Context);
+	VertexShader->SetParameters(Context);
+}
+
+
+static bool DoSpecularCorrection()
+{
+	bool CVarState = CVarSubsurfaceQuality.GetValueOnRenderThread() > 0;
+
+	int SceneColorFormat;
+	{
+		static IConsoleVariable* CVar = IConsoleManager::Get().FindConsoleVariable(TEXT("r.SceneColorFormat"));
+
+		SceneColorFormat = CVar->GetInt();
+	}
+
+	// we need an alpha channel for this feature
+	return CVarState && (SceneColorFormat >= 4);
+}
+
+void FRCPassPostProcessSubsurfaceSetup::Process(FRenderingCompositePassContext& Context)
+{
+	SCOPED_DRAW_EVENT(Context.RHICmdList, SubsurfaceSetup, DEC_SCENE_ITEMS);
+
+	const FPooledRenderTargetDesc* InputDesc = GetInputDesc(ePId_Input0);
+
+	if(!InputDesc)
+	{
+		// input is not hooked up correctly
+		return;
+	}
+
+	const FSceneView& View = Context.View;
+	const FSceneViewFamily& ViewFamily = *(View.Family);
+
+	FIntPoint SrcSize = InputDesc->Extent;
+	FIntPoint DestSize = PassOutputs[0].RenderTargetDesc.Extent;
+
+	// e.g. 4 means the input texture is 4x smaller than the buffer size
+	uint32 ScaleFactor = GSceneRenderTargets.GetBufferSizeXY().X / SrcSize.X;
+
+	FIntRect SrcRect = View.ViewRect / ScaleFactor;
+	FIntRect DestRect = FIntRect::DivideAndRoundUp(SrcRect, 2);
+
+	// upscale rectangle to not slightly scale
+	SrcRect = DestRect * 2;
+
+	const FSceneRenderTargetItem& DestRenderTarget = PassOutputs[0].RequestSurface(Context);
+
+	// Set the view family's render target/viewport.
+	SetRenderTarget(Context.RHICmdList, DestRenderTarget.TargetableTexture, FTextureRHIRef());
+
+	// is optimized away if possible (RT size=view size, )
+	Context.RHICmdList.Clear(true, FLinearColor::Black, false, 1.0f, false, 0, DestRect);
+
+	Context.SetViewportAndCallRHI(0, 0, 0.0f, DestSize.X, DestSize.Y, 1.0f );
+
+	// set the state
+	Context.RHICmdList.SetBlendState(TStaticBlendState<>::GetRHI());
+	Context.RHICmdList.SetRasterizerState(TStaticRasterizerState<>::GetRHI());
+	Context.RHICmdList.SetDepthStencilState(TStaticDepthStencilState<false, CF_Always>::GetRHI());
+
+	// reconstruct specular and add it in final pass
+	bool bSpecularCorrection = DoSpecularCorrection();
+
+	if(bSpecularCorrection)
+	{
+		SetSubsurfaceSetupShader<1>(Context);
+	}
+	else
+	{
+		SetSubsurfaceSetupShader<0>(Context);
+	}
+
+	// Draw a quad mapping scene color to the view's render target
+	TShaderMapRef<FPostProcessVS> VertexShader(Context.GetShaderMap());
+	DrawRectangle(
+		Context.RHICmdList,
+		DestRect.Min.X, DestRect.Min.Y,
+		DestRect.Width(), DestRect.Height(),
+		SrcRect.Min.X, SrcRect.Min.Y,
+		SrcRect.Width(), SrcRect.Height(),
+		DestSize,
+		SrcSize,
+		*VertexShader,
+		EDRF_UseTriangleOptimization);
+
+
+	Context.RHICmdList.CopyToResolveTarget(DestRenderTarget.TargetableTexture, DestRenderTarget.ShaderResourceTexture, false, FResolveParams());
+}
+
+FPooledRenderTargetDesc FRCPassPostProcessSubsurfaceSetup::ComputeOutputDesc(EPassOutputId InPassOutputId) const
+{
+	FPooledRenderTargetDesc Ret = GSceneRenderTargets.GetSceneColor()->GetDesc();
+
+	Ret.Reset();
+	Ret.DebugName = TEXT("SubsurfaceSetup");
+	// alpha is used to store depth and renormalize (alpha==0 means there is no subsurface scattering)
+	Ret.Format = PF_FloatRGBA;
+
+	// half res
+	Ret.Extent  = FIntPoint::DivideAndRoundUp(Ret.Extent, 2);
+	Ret.Extent.X = FMath::Max(1, Ret.Extent.X);
+	Ret.Extent.Y = FMath::Max(1, Ret.Extent.Y);
+	
 	return Ret;
 }
 
@@ -344,9 +540,7 @@ class TPostProcessSubsurfacePS : public FGlobalShader
 public:
 	FPostProcessPassParameters PostprocessParameter;
 	FDeferredPixelShaderParameters DeferredParameters;
-	FShaderParameter SSSParams;
-	FShaderResourceParameter SSProfilesTexture;
-	FShaderResourceParameter SSProfilesTextureSampler;
+	FSubsurfaceParameters SubsurfaceParameters;
 
 	/** Initialization constructor. */
 	TPostProcessSubsurfacePS(const ShaderMetaType::CompiledShaderInitializerType& Initializer)
@@ -354,47 +548,25 @@ public:
 	{
 		PostprocessParameter.Bind(Initializer.ParameterMap);
 		DeferredParameters.Bind(Initializer.ParameterMap);
-		SSSParams.Bind(Initializer.ParameterMap, TEXT("SSSParams"));
-		SSProfilesTexture.Bind(Initializer.ParameterMap, TEXT("SSProfilesTexture"));
-		SSProfilesTextureSampler.Bind(Initializer.ParameterMap, TEXT("SSProfilesTextureSampler"));
+		SubsurfaceParameters.Bind(Initializer.ParameterMap);
 	}
 
 	// FShader interface.
 	virtual bool Serialize(FArchive& Ar)
 	{
 		bool bShaderHasOutdatedParameters = FGlobalShader::Serialize(Ar);
-		Ar << PostprocessParameter << DeferredParameters << SSSParams << SSProfilesTexture << SSProfilesTextureSampler;
+		Ar << PostprocessParameter << DeferredParameters << SubsurfaceParameters;
 		return bShaderHasOutdatedParameters;
 	}
 
-	void SetParameters(const FRenderingCompositePassContext& Context, float InRadius)
+	void SetParameters(const FRenderingCompositePassContext& Context)
 	{
 		const FPixelShaderRHIParamRef ShaderRHI = GetPixelShader();
 
 		FGlobalShader::SetParameters(Context.RHICmdList, ShaderRHI, Context.View);
 		DeferredParameters.Set(Context.RHICmdList, ShaderRHI, Context.View);
 		PostprocessParameter.SetPS(ShaderRHI, Context, TStaticSamplerState<SF_Bilinear,AM_Border,AM_Border,AM_Border>::GetRHI());
-
-		{
-			// from Separabale.usf: float distanceToProjectionWindow = 1.0 / tan(0.5 * radians(SSSS_FOVY))
-			// can be extracted out of projection matrix
-
-			float ScaleCorrectionX = Context.View.ViewRect.Width() / (float)GSceneRenderTargets.GetBufferSizeXY().X;
-			float ScaleCorrectionY = Context.View.ViewRect.Height() / (float)GSceneRenderTargets.GetBufferSizeXY().Y;
-
-			FVector4 ColorScale(InRadius, Context.View.ViewMatrices.ProjMatrix.M[0][0], ScaleCorrectionX, ScaleCorrectionY);
-			SetShaderValue(Context.RHICmdList, ShaderRHI, SSSParams, ColorScale);
-		}
-
-		{
-			const IPooledRenderTarget* PooledRT = GetSubsufaceProfileTexture_RT(Context.RHICmdList);
-
-			check(PooledRT);
-
-			const FSceneRenderTargetItem& Item = PooledRT->GetRenderTargetItem();
-
-			SetTextureParameter(Context.RHICmdList, ShaderRHI, SSProfilesTexture, SSProfilesTextureSampler, TStaticSamplerState<SF_Point, AM_Wrap, AM_Wrap, AM_Wrap>::GetRHI(), Item.ShaderResourceTexture);
-		}
+		SubsurfaceParameters.SetParameters(Context.RHICmdList, ShaderRHI, Context);
 	}
 
 	static const TCHAR* GetSourceFilename()
@@ -417,17 +589,14 @@ public:
 #undef VARIATION2
 
 
-FRCPassPostProcessSubsurface::FRCPassPostProcessSubsurface(uint32 InDirection, float InRadius, bool bInHalfRes)
-	: Radius(InRadius)
-	, Direction(InDirection) 
-	, bHalfRes(bInHalfRes) 
+FRCPassPostProcessSubsurface::FRCPassPostProcessSubsurface(uint32 InDirection)
+	: Direction(InDirection) 
 {
 	check(InDirection < 2);
-	check(Radius >= 0.0f);
 }
 
 template <uint32 Direction, uint32 SampleSet>
-void SetSubsurfaceShader(const FRenderingCompositePassContext& Context, TShaderMapRef<FPostProcessVS> &VertexShader, float InRadius)
+void SetSubsurfaceShader(const FRenderingCompositePassContext& Context, TShaderMapRef<FPostProcessVS> &VertexShader)
 {
 	TShaderMapRef<TPostProcessSubsurfacePS<Direction, SampleSet> > PixelShader(Context.GetShaderMap());
 
@@ -435,19 +604,19 @@ void SetSubsurfaceShader(const FRenderingCompositePassContext& Context, TShaderM
 
 	SetGlobalBoundShaderState(Context.RHICmdList, Context.GetFeatureLevel(), BoundShaderState, GFilterVertexDeclaration.VertexDeclarationRHI, *VertexShader, *PixelShader);
 
-	PixelShader->SetParameters(Context, InRadius);
+	PixelShader->SetParameters(Context);
 	VertexShader->SetParameters(Context);
 }
 
 // 0:horizontal, 1: vertical
 template <uint32 Direction>
-void SetSubsurfaceShaderSampleSet(const FRenderingCompositePassContext& Context, TShaderMapRef<FPostProcessVS> &VertexShader, float InRadius, uint32 SampleSet)
+void SetSubsurfaceShaderSampleSet(const FRenderingCompositePassContext& Context, TShaderMapRef<FPostProcessVS> &VertexShader, uint32 SampleSet)
 {
 	switch(SampleSet)
 	{
-		case 0: SetSubsurfaceShader<Direction, 0>(Context, VertexShader, InRadius); break;
-		case 1: SetSubsurfaceShader<Direction, 1>(Context, VertexShader, InRadius); break;
-		case 2: SetSubsurfaceShader<Direction, 2>(Context, VertexShader, InRadius); break;
+		case 0: SetSubsurfaceShader<Direction, 0>(Context, VertexShader); break;
+		case 1: SetSubsurfaceShader<Direction, 1>(Context, VertexShader); break;
+		case 2: SetSubsurfaceShader<Direction, 2>(Context, VertexShader); break;
 
 	default:
 		check(0);
@@ -484,7 +653,7 @@ void FRCPassPostProcessSubsurface::Process(FRenderingCompositePassContext& Conte
 	uint32 ScaleFactor = GSceneRenderTargets.GetBufferSizeXY().X / SrcSize.X;
 
 	FIntRect SrcRect = View.ViewRect / ScaleFactor;
-	FIntRect DestRect = bHalfRes ? SrcRect : View.ViewRect;
+	FIntRect DestRect = SrcRect;
 
 	TRefCountPtr<IPooledRenderTarget> NewSceneColor;
 
@@ -515,11 +684,11 @@ void FRCPassPostProcessSubsurface::Process(FRenderingCompositePassContext& Conte
 
 	if (Direction == 0)
 	{
-		SetSubsurfaceShaderSampleSet<0>(Context, VertexShader, Radius, SampleSet);
+		SetSubsurfaceShaderSampleSet<0>(Context, VertexShader, SampleSet);
 	}
 	else
 	{
-		SetSubsurfaceShaderSampleSet<1>(Context, VertexShader, Radius, SampleSet);
+		SetSubsurfaceShaderSampleSet<1>(Context, VertexShader, SampleSet);
 	}
 
 	DrawRectangle(
@@ -547,13 +716,10 @@ FPooledRenderTargetDesc FRCPassPostProcessSubsurface::ComputeOutputDesc(EPassOut
 	// the setup was done in half res but the actual sampling happens in full resolution
 	Ret.Extent = GSceneRenderTargets.GetBufferSizeXY();
 
-	if(bHalfRes)
-	{
-		Ret.Extent  = FIntPoint::DivideAndRoundUp(Ret.Extent, 2);
-
-		Ret.Extent.X = FMath::Max(1, Ret.Extent.X);
-		Ret.Extent.Y = FMath::Max(1, Ret.Extent.Y);
-	}
+	// half res
+	Ret.Extent  = FIntPoint::DivideAndRoundUp(Ret.Extent, 2);
+	Ret.Extent.X = FMath::Max(1, Ret.Extent.X);
+	Ret.Extent.Y = FMath::Max(1, Ret.Extent.Y);
 
 	return Ret;
 }
@@ -587,6 +753,7 @@ class TPostProcessSubsurfaceRecombinePS : public FGlobalShader
 public:
 	FPostProcessPassParameters PostprocessParameter;
 	FDeferredPixelShaderParameters DeferredParameters;
+	FSubsurfaceParameters SubsurfaceParameters;
 
 	/** Initialization constructor. */
 	TPostProcessSubsurfaceRecombinePS(const ShaderMetaType::CompiledShaderInitializerType& Initializer)
@@ -594,13 +761,14 @@ public:
 	{
 		PostprocessParameter.Bind(Initializer.ParameterMap);
 		DeferredParameters.Bind(Initializer.ParameterMap);
+		SubsurfaceParameters.Bind(Initializer.ParameterMap);
 	}
 
 	// FShader interface.
 	virtual bool Serialize(FArchive& Ar)
 	{
 		bool bShaderHasOutdatedParameters = FGlobalShader::Serialize(Ar);
-		Ar << PostprocessParameter << DeferredParameters;
+		Ar << PostprocessParameter << DeferredParameters << SubsurfaceParameters;
 		return bShaderHasOutdatedParameters;
 	}
 
@@ -611,6 +779,7 @@ public:
 		FGlobalShader::SetParameters(Context.RHICmdList, ShaderRHI, Context.View);
 		DeferredParameters.Set(Context.RHICmdList, ShaderRHI, Context.View);
 		PostprocessParameter.SetPS(ShaderRHI, Context, TStaticSamplerState<SF_Bilinear,AM_Border,AM_Border,AM_Border>::GetRHI());
+		SubsurfaceParameters.SetParameters(Context.RHICmdList, ShaderRHI, Context);
 	}
 
 	static const TCHAR* GetSourceFilename()
diff --git a/Engine/Source/Runtime/Renderer/Private/PostProcess/PostProcessSubsurface.h b/Engine/Source/Runtime/Renderer/Private/PostProcess/PostProcessSubsurface.h
--- a/Engine/Source/Runtime/Renderer/Private/PostProcess/PostProcessSubsurface.h
+++ b/Engine/Source/Runtime/Renderer/Private/PostProcess/PostProcessSubsurface.h
@@ -8,25 +8,32 @@
 
 #include "RenderingCompositionGraph.h"
 
-
 // ePId_Input0: SceneColor
 // derives from TRenderingCompositePassBase<InputCount, OutputCount> 
 // uses some GBuffer attributes
 // alpha is unused
-class FRCPassPostProcessSubsurfaceSetup : public TRenderingCompositePassBase<1, 1>
+class FRCPassPostProcessSubsurfaceVisualize : public TRenderingCompositePassBase<1, 1>
 {
 public:
-	// @param bInHalfRes if the output shoudl be half resolution to scene color
-	FRCPassPostProcessSubsurfaceSetup(bool bInVisualize, bool bInHalfRes = false);
+	FRCPassPostProcessSubsurfaceVisualize();
 
 	// interface FRenderingCompositePass ---------
 	virtual void Process(FRenderingCompositePassContext& Context);
 	virtual FPooledRenderTargetDesc ComputeOutputDesc(EPassOutputId InPassOutputId) const;
 	virtual void Release() override { delete this; }
+};
 
-	bool bVisualize;
-	// if the output shoudl be half resolution to scene color
-	bool bHalfRes;
+// ePId_Input0: SceneColor
+// derives from TRenderingCompositePassBase<InputCount, OutputCount> 
+// uses some GBuffer attributes
+// alpha is unused
+class FRCPassPostProcessSubsurfaceSetup : public TRenderingCompositePassBase<1, 1>
+{
+public:
+	// interface FRenderingCompositePass ---------
+	virtual void Process(FRenderingCompositePassContext& Context);
+	virtual FPooledRenderTargetDesc ComputeOutputDesc(EPassOutputId InPassOutputId) const;
+	virtual void Release() override { delete this; }
 };
 
 
@@ -38,8 +45,7 @@ class FRCPassPostProcessSubsurface : public TRenderingCompositePassBase<1, 1>
 public:
 	// constructor
 	// @param InDirection 0:horizontal/1:vertical
-	// @param InRadius in pixels in the full res image
-	FRCPassPostProcessSubsurface(uint32 InDirection, float InRadius, bool bInHalfRes = false);
+	FRCPassPostProcessSubsurface(uint32 InDirection);
 
 	// interface FRenderingCompositePass ---------
 
@@ -48,12 +54,8 @@ public:
 	virtual FPooledRenderTargetDesc ComputeOutputDesc(EPassOutputId InPassOutputId) const;
 
 private:
-	// in pixels in the full res image
-	float Radius;
 	// 0:horizontal/1:vertical
 	uint32 Direction;
-	// if the output should be half resolution to scene color
-	bool bHalfRes;
 };
 
 
diff --git a/Engine/Source/Runtime/Renderer/Private/PostProcess/PostProcessing.cpp b/Engine/Source/Runtime/Renderer/Private/PostProcess/PostProcessing.cpp
--- a/Engine/Source/Runtime/Renderer/Private/PostProcess/PostProcessing.cpp
+++ b/Engine/Source/Runtime/Renderer/Private/PostProcess/PostProcessing.cpp
@@ -1199,7 +1199,7 @@ void FPostProcessing::Process(FRHICommandListImmediate& RHICmdList, FViewInfo& V
 		if (View.Family->EngineShowFlags.VisualizeSSS)
 		{
 			// the setup pass also does visualization, based on EngineShowFlags.VisualizeSSS
-			FRenderingCompositePass* PassVisualize = Context.Graph.RegisterPass(new(FMemStack::Get()) FRCPassPostProcessSubsurfaceSetup(true));
+			FRenderingCompositePass* PassVisualize = Context.Graph.RegisterPass(new(FMemStack::Get()) FRCPassPostProcessSubsurfaceVisualize());
 			PassVisualize->SetInput(ePId_Input0, Context.FinalOutput);
 			Context.FinalOutput = FRenderingCompositeOutputRef(PassVisualize);
 		}
