added useful cvar

diff --git a/Engine/Source/Runtime/Renderer/Private/PostProcess/PostProcessSubsurface.cpp b/Engine/Source/Runtime/Renderer/Private/PostProcess/PostProcessSubsurface.cpp
--- a/Engine/Source/Runtime/Renderer/Private/PostProcess/PostProcessSubsurface.cpp
+++ b/Engine/Source/Runtime/Renderer/Private/PostProcess/PostProcessSubsurface.cpp
@@ -14,6 +14,13 @@
 ENGINE_API const IPooledRenderTarget* GetSubsufaceProfileTexture_RT(FRHICommandListImmediate& RHICmdList);
 
 
+static TAutoConsoleVariable<int32> CVarSSSFilter(
+	TEXT("r.SSS.Filter"),
+	1,
+	TEXT("0: point filter (useful for testing, could be cleaner)\n")
+	TEXT("1: bilinear filter"),
+	ECVF_RenderThreadSafe  | ECVF_Scalability);
+
 static TAutoConsoleVariable<int32> CVarSSSSampleSet(
 	TEXT("r.SSS.SampleSet"),
 	3,
@@ -565,7 +572,15 @@ public:
 
 		FGlobalShader::SetParameters(Context.RHICmdList, ShaderRHI, Context.View);
 		DeferredParameters.Set(Context.RHICmdList, ShaderRHI, Context.View);
-		PostprocessParameter.SetPS(ShaderRHI, Context, TStaticSamplerState<SF_Bilinear,AM_Border,AM_Border,AM_Border>::GetRHI());
+
+		if(CVarSSSFilter.GetValueOnRenderThread())
+		{
+			PostprocessParameter.SetPS(ShaderRHI, Context, TStaticSamplerState<SF_Bilinear,AM_Border,AM_Border,AM_Border>::GetRHI());
+		}
+		else
+		{
+			PostprocessParameter.SetPS(ShaderRHI, Context, TStaticSamplerState<SF_Point,AM_Border,AM_Border,AM_Border>::GetRHI());
+		}
 		SubsurfaceParameters.SetParameters(Context.RHICmdList, ShaderRHI, Context);
 	}
 
