minor cleanup

diff --git a/Engine/Source/Runtime/Renderer/Private/PostProcess/PostProcessSubsurface.cpp b/Engine/Source/Runtime/Renderer/Private/PostProcess/PostProcessSubsurface.cpp
--- a/Engine/Source/Runtime/Renderer/Private/PostProcess/PostProcessSubsurface.cpp
+++ b/Engine/Source/Runtime/Renderer/Private/PostProcess/PostProcessSubsurface.cpp
@@ -23,7 +23,7 @@ static TAutoConsoleVariable<int32> CVarSSSFilter(
 
 static TAutoConsoleVariable<int32> CVarSSSSampleSet(
 	TEXT("r.SSS.SampleSet"),
-	3,
+	2,
 	TEXT("0: lowest quality\n")
 	TEXT("1: medium quality\n")
 	TEXT("2: high quality (default)"),
