UE4 - Remove Color parameter on draw events (DEC_*)

diff --git a/Engine/Source/Runtime/Engine/Private/GPUSkinCache.cpp b/Engine/Source/Runtime/Engine/Private/GPUSkinCache.cpp
--- a/Engine/Source/Runtime/Engine/Private/GPUSkinCache.cpp
+++ b/Engine/Source/Runtime/Engine/Private/GPUSkinCache.cpp
@@ -479,7 +479,7 @@ void FGPUSkinCache::DispatchSkinCacheProcess(FRHICommandListImmediate& RHICmdLis
 
 	uint32 VertexCountAlign64 = (VertexCount + 63) / 64;
 
-	SCOPED_DRAW_EVENT(RHICmdList, SkinCacheDispatch, DEC_SKEL_MESH);
+	SCOPED_DRAW_EVENT(RHICmdList, SkinCacheDispatch);
 
 	if (bUseExtraBoneInfluences)
 	{
diff --git a/Engine/Source/Runtime/Engine/Private/GPUSort.cpp b/Engine/Source/Runtime/Engine/Private/GPUSort.cpp
--- a/Engine/Source/Runtime/Engine/Private/GPUSort.cpp
+++ b/Engine/Source/Runtime/Engine/Private/GPUSort.cpp
@@ -688,7 +688,7 @@ int32 SortGPUBuffers(FRHICommandListImmediate& RHICmdList, FGPUSortBuffers SortB
 
 	check(FeatureLevel == ERHIFeatureLevel::SM5);
 
-	SCOPED_DRAW_EVENTF(RHICmdList, SortGPU, DEC_PARTICLE, TEXT("SortGPU_%d"), Count);
+	SCOPED_DRAW_EVENTF(RHICmdList, SortGPU, TEXT("SortGPU_%d"), Count);
 
 	// Determine how many tiles need to be sorted.
 	const int32 TileCount = Count / TILE_SIZE;
diff --git a/Engine/Source/Runtime/Engine/Private/Particles/ParticleCurveTexture.cpp b/Engine/Source/Runtime/Engine/Private/Particles/ParticleCurveTexture.cpp
--- a/Engine/Source/Runtime/Engine/Private/Particles/ParticleCurveTexture.cpp
+++ b/Engine/Source/Runtime/Engine/Private/Particles/ParticleCurveTexture.cpp
@@ -184,7 +184,7 @@ static void InjectCurves(
 
 	check( IsInRenderingThread() );
 
-	SCOPED_DRAW_EVENT(RHICmdList, InjectParticleCurves, DEC_PARTICLE);
+	SCOPED_DRAW_EVENT(RHICmdList, InjectParticleCurves);
 
 	FVertexBufferRHIParamRef ScratchVertexBufferRHI = GParticleScratchVertexBuffer.VertexBufferRHI;
 
diff --git a/Engine/Source/Runtime/Engine/Private/Particles/ParticleGpuSimulation.cpp b/Engine/Source/Runtime/Engine/Private/Particles/ParticleGpuSimulation.cpp
--- a/Engine/Source/Runtime/Engine/Private/Particles/ParticleGpuSimulation.cpp
+++ b/Engine/Source/Runtime/Engine/Private/Particles/ParticleGpuSimulation.cpp
@@ -1279,7 +1279,7 @@ void ExecuteSimulationCommands(
  */
 void ClearTiles(FRHICommandList& RHICmdList, ERHIFeatureLevel::Type FeatureLevel, const TArray<uint32>& Tiles)
 {
-	SCOPED_DRAW_EVENT(RHICmdList, ClearTiles, DEC_PARTICLE);
+	SCOPED_DRAW_EVENT(RHICmdList, ClearTiles);
 
 	const int32 MaxTilesPerDrawCallUnaligned = GParticleScratchVertexBufferSize / sizeof(FVector2D);
 	const int32 MaxTilesPerDrawCall = MaxTilesPerDrawCallUnaligned & (~(TILES_PER_INSTANCE-1));
@@ -1755,7 +1755,7 @@ static void VisualizeGPUSimulation(
 	)
 {
 	check(IsInRenderingThread());
-	SCOPED_DRAW_EVENT(RHICmdList, ParticleSimDebugDraw, DEC_PARTICLE);
+	SCOPED_DRAW_EVENT(RHICmdList, ParticleSimDebugDraw);
 
 	// Some constants for laying out the debug view.
 	const float DisplaySizeX = 256.0f;
@@ -4266,7 +4266,7 @@ void FFXSystem::SimulateGPUParticles(
 	// Simulate particles in all active tiles.
 	if ( SimulationCommands.Num() )
 	{
-		SCOPED_DRAW_EVENT(RHICmdList, ParticleSimulation, DEC_PARTICLE);
+		SCOPED_DRAW_EVENT(RHICmdList, ParticleSimulation);
 
 		if (Phase == EParticleSimulatePhase::Collision && CollisionView)
 		{
@@ -4309,7 +4309,7 @@ void FFXSystem::SimulateGPUParticles(
 	// Inject any new particles that have spawned into the simulation.
 	if (NewParticles.Num())
 	{
-		SCOPED_DRAW_EVENT(RHICmdList, ParticleInjection, DEC_PARTICLE);
+		SCOPED_DRAW_EVENT(RHICmdList, ParticleInjection);
 		// Set render targets.
 		FTextureRHIParamRef InjectRenderTargets[4] =
 		{
diff --git a/Engine/Source/Runtime/Engine/Private/Particles/ParticleSortingGPU.cpp b/Engine/Source/Runtime/Engine/Private/Particles/ParticleSortingGPU.cpp
--- a/Engine/Source/Runtime/Engine/Private/Particles/ParticleSortingGPU.cpp
+++ b/Engine/Source/Runtime/Engine/Private/Particles/ParticleSortingGPU.cpp
@@ -184,7 +184,7 @@ static int32 GenerateParticleSortKeys(
 	ERHIFeatureLevel::Type FeatureLevel
 	)
 {
-	SCOPED_DRAW_EVENT(RHICmdList, ParticleSortKeyGen, DEC_PARTICLE);
+	SCOPED_DRAW_EVENT(RHICmdList, ParticleSortKeyGen);
 	check(FeatureLevel == ERHIFeatureLevel::SM5);
 
 	FParticleKeyGenParameters KeyGenParameters;
@@ -315,7 +315,7 @@ int32 SortParticlesGPU(
 	ERHIFeatureLevel::Type FeatureLevel
 	)
 {
-	SCOPED_DRAW_EVENTF(RHICmdList, ParticleSort, DEC_PARTICLE, TEXT("ParticleSort_%d"), SimulationsToSort.Num());
+	SCOPED_DRAW_EVENTF(RHICmdList, ParticleSort, TEXT("ParticleSort_%d"), SimulationsToSort.Num());
 
 	// Ensure the sorted vertex buffers are not currently bound as input streams.
 	// They should only ever be bound to streams 0 or 1, so clear them.
diff --git a/Engine/Source/Runtime/Engine/Private/PrimitiveDrawingUtils.cpp b/Engine/Source/Runtime/Engine/Private/PrimitiveDrawingUtils.cpp
--- a/Engine/Source/Runtime/Engine/Private/PrimitiveDrawingUtils.cpp
+++ b/Engine/Source/Runtime/Engine/Private/PrimitiveDrawingUtils.cpp
@@ -15,18 +15,18 @@ void EmitMeshDrawEvents_Inner(FRHICommandList& RHICmdList, const FPrimitiveScene
 {
 	// Only show material name at the top level
 	// Note: this is the parent's material name, not the material instance
-	SCOPED_DRAW_EVENTF(RHICmdList, MaterialEvent, DEC_SCENE_ITEMS, *Mesh.MaterialRenderProxy->GetMaterial(PrimitiveSceneProxy ? PrimitiveSceneProxy->GetScene().GetFeatureLevel() : GRHIFeatureLevel_DEPRECATED)->GetFriendlyName());
+	SCOPED_DRAW_EVENTF(RHICmdList, MaterialEvent, *Mesh.MaterialRenderProxy->GetMaterial(PrimitiveSceneProxy ? PrimitiveSceneProxy->GetScene().GetFeatureLevel() : GRHIFeatureLevel_DEPRECATED)->GetFriendlyName());
 	if (PrimitiveSceneProxy)
 	{
 		// Show Actor, level and resource name inside the material name
 		// These are separate draw events since some platforms only allow 32 character event names (xenon)
 		{
-			SCOPED_CONDITIONAL_DRAW_EVENTF(RHICmdList, LevelEvent, PrimitiveSceneProxy->GetLevelName() != NAME_None, DEC_SCENE_ITEMS, PrimitiveSceneProxy->GetLevelName().IsValid() ? *PrimitiveSceneProxy->GetLevelName().ToString() : TEXT(""));
+			SCOPED_CONDITIONAL_DRAW_EVENTF(RHICmdList, LevelEvent, PrimitiveSceneProxy->GetLevelName() != NAME_None, PrimitiveSceneProxy->GetLevelName().IsValid() ? *PrimitiveSceneProxy->GetLevelName().ToString() : TEXT(""));
 		}
 		{
-			SCOPED_CONDITIONAL_DRAW_EVENTF(RHICmdList, OwnerEvent,PrimitiveSceneProxy->GetOwnerName() != NAME_None, DEC_SCENE_ITEMS, *PrimitiveSceneProxy->GetOwnerName().ToString());
+			SCOPED_CONDITIONAL_DRAW_EVENTF(RHICmdList, OwnerEvent,PrimitiveSceneProxy->GetOwnerName() != NAME_None, *PrimitiveSceneProxy->GetOwnerName().ToString());
 		}
-		SCOPED_CONDITIONAL_DRAW_EVENTF(RHICmdList, ResourceEvent,PrimitiveSceneProxy->GetResourceName() != NAME_None, DEC_SCENE_ITEMS, PrimitiveSceneProxy->GetResourceName().IsValid() ? *PrimitiveSceneProxy->GetResourceName().ToString() : TEXT(""));
+		SCOPED_CONDITIONAL_DRAW_EVENTF(RHICmdList, ResourceEvent,PrimitiveSceneProxy->GetResourceName() != NAME_None, PrimitiveSceneProxy->GetResourceName().IsValid() ? *PrimitiveSceneProxy->GetResourceName().ToString() : TEXT(""));
 	}
 }
 #endif
diff --git a/Engine/Source/Runtime/Engine/Private/SceneUtils.cpp b/Engine/Source/Runtime/Engine/Private/SceneUtils.cpp
--- a/Engine/Source/Runtime/Engine/Private/SceneUtils.cpp
+++ b/Engine/Source/Runtime/Engine/Private/SceneUtils.cpp
@@ -6,7 +6,7 @@
 
 #if WANTS_DRAW_MESH_EVENTS && PLATFORM_SUPPORTS_DRAW_MESH_EVENTS
 
-void FDrawEvent::Start(FRHICommandList& InRHICmdList, const FColor& Color, const TCHAR* Fmt, ...)
+void FDrawEvent::Start(FRHICommandList& InRHICmdList, const TCHAR* Fmt, ...)
 {
 	check(IsInParallelRenderingThread() || IsInRHIThread());
 	{
diff --git a/Engine/Source/Runtime/Engine/Private/UnrealClient.cpp b/Engine/Source/Runtime/Engine/Private/UnrealClient.cpp
--- a/Engine/Source/Runtime/Engine/Private/UnrealClient.cpp
+++ b/Engine/Source/Runtime/Engine/Private/UnrealClient.cpp
@@ -292,18 +292,15 @@ int32 FStatUnitData::DrawStat(FViewport* InViewport, FCanvas* InCanvas, int32 In
 	/** Number of milliseconds the gamethread was used last frame. */
 	RawGameThreadTime = FPlatformTime::ToMilliseconds(GGameThreadTime);
 	GameThreadTime = 0.9 * GameThreadTime + 0.1 * RawGameThreadTime;
-	appSetCounterValue(TEXT("Game thread time"), FPlatformTime::ToMilliseconds(GGameThreadTime));
 
 	/** Number of milliseconds the renderthread was used last frame. */
 	RawRenderThreadTime = FPlatformTime::ToMilliseconds(GRenderThreadTime);
 	RenderThreadTime = 0.9 * RenderThreadTime + 0.1 * RawRenderThreadTime;
-	appSetCounterValue(TEXT("Render thread time"), FPlatformTime::ToMilliseconds(GRenderThreadTime));
 
 	/** Number of milliseconds the GPU was busy last frame. */
 	const uint32 GPUCycles = RHIGetGPUFrameCycles();
 	RawGPUFrameTime = FPlatformTime::ToMilliseconds(GPUCycles);
 	GPUFrameTime = 0.9 * GPUFrameTime + 0.1 * RawGPUFrameTime;
-	appSetCounterValue(TEXT("GPU time"), FPlatformTime::ToMilliseconds(GPUCycles));
 
 	SET_FLOAT_STAT(STAT_FPSChart_UnitFrame, FrameTime);
 	SET_FLOAT_STAT(STAT_FPSChart_UnitRender, RenderThreadTime);
diff --git a/Engine/Source/Runtime/Engine/Private/UserInterface/Canvas.cpp b/Engine/Source/Runtime/Engine/Private/UserInterface/Canvas.cpp
--- a/Engine/Source/Runtime/Engine/Private/UserInterface/Canvas.cpp
+++ b/Engine/Source/Runtime/Engine/Private/UserInterface/Canvas.cpp
@@ -250,7 +250,7 @@ bool FCanvasBatchedElementRenderItem::Render_GameThread(const FCanvas* Canvas)
 			BatchedDrawCommand,
 			FBatchedDrawParameters,Parameters,DrawParameters,
 		{
-			SCOPED_DRAW_EVENT(RHICmdList, CanvasBatchedElements, DEC_CANVAS);
+			SCOPED_DRAW_EVENT(RHICmdList, CanvasBatchedElements);
 
 			// draw batched items
 			Parameters.RenderData->BatchedElements.Draw(
@@ -400,7 +400,7 @@ bool FCanvasTileRendererItem::Render_GameThread(const FCanvas* Canvas)
 		DrawTileCommand,
 		FDrawTileParameters, Parameters, DrawTileParameters,
 	{
-		SCOPED_DRAW_EVENT(RHICmdList, CanvasDrawTile, DEC_CANVAS);
+		SCOPED_DRAW_EVENT(RHICmdList, CanvasDrawTile);
 		for (int32 TileIdx = 0; TileIdx < Parameters.RenderData->Tiles.Num(); TileIdx++)
 		{
 			const FRenderData::FTileInst& Tile = Parameters.RenderData->Tiles[TileIdx];
@@ -566,7 +566,7 @@ void FCanvas::Flush_RenderThread(FRHICommandListImmediate& RHICmdList, bool bFor
 	// sort the array of FCanvasSortElement entries so that higher sort keys render first (back-to-front)
 	SortedElements.Sort(FCompareFCanvasSortElement());
 
-	SCOPED_DRAW_EVENT(RHICmdList, CanvasFlush, DEC_CANVAS);
+	SCOPED_DRAW_EVENT(RHICmdList, CanvasFlush);
 	const FTexture2DRHIRef& RenderTargetTexture = RenderTarget->GetRenderTargetTexture();
 
 	check(IsValidRef(RenderTargetTexture));
@@ -664,7 +664,7 @@ void FCanvas::Flush_GameThread(bool bForce)
 		CanvasFlushSetupCommand,
 		FCanvasFlushParameters,Parameters,FlushParameters,
 	{
-		SCOPED_DRAW_EVENT(RHICmdList, CanvasFlush, DEC_CANVAS);
+		SCOPED_DRAW_EVENT(RHICmdList, CanvasFlush);
 
 		// Set the RHI render target.
 		::SetRenderTarget(RHICmdList, Parameters.CanvasRenderTarget->GetRenderTargetTexture(), FTextureRHIRef());
@@ -811,7 +811,7 @@ void FCanvas::Clear(const FLinearColor& Color)
 		FColor,Color,GammaCorrectedColor,
 		FRenderTarget*,CanvasRenderTarget,GetRenderTarget(),
 	{
-		SCOPED_DRAW_EVENT(RHICmdList, CanvasClear, DEC_CANVAS);
+		SCOPED_DRAW_EVENT(RHICmdList, CanvasClear);
 		if( CanvasRenderTarget )
 		{
 			::SetRenderTarget(RHICmdList, CanvasRenderTarget->GetRenderTargetTexture(),FTextureRHIRef());
diff --git a/Engine/Source/Runtime/Engine/Private/VectorField.cpp b/Engine/Source/Runtime/Engine/Private/VectorField.cpp
--- a/Engine/Source/Runtime/Engine/Private/VectorField.cpp
+++ b/Engine/Source/Runtime/Engine/Private/VectorField.cpp
@@ -801,7 +801,7 @@ public:
 
 		if (GetFeatureLevel() == ERHIFeatureLevel::SM5 && AnimatedVectorField && AnimatedVectorField->Texture && AnimatedVectorField->Texture->Resource)
 		{
-			SCOPED_DRAW_EVENT(RHICmdList, AnimateVectorField, DEC_PARTICLE);
+			SCOPED_DRAW_EVENT(RHICmdList, AnimateVectorField);
 
 			// Move frame time forward.
 			FrameTime += AnimatedVectorField->FramesPerSecond * DeltaSeconds;
diff --git a/Engine/Source/Runtime/Engine/Public/SceneUtils.h b/Engine/Source/Runtime/Engine/Public/SceneUtils.h
--- a/Engine/Source/Runtime/Engine/Public/SceneUtils.h
+++ b/Engine/Source/Runtime/Engine/Public/SceneUtils.h
@@ -10,20 +10,6 @@
 // Each event type will be displayed using the defined color
 #pragma once
 
-
-#define DEC_LIGHT			FColor(255,0,0,255)
-#define DEC_SKEL_MESH		FColor(255,0,255,255)
-#define DEC_STATIC_MESH		FColor(0,128,255,255)
-#define DEC_CANVAS			FColor(128,255,255,255)
-#define DEC_SHADOW			FColor(128,128,128,255)
-#define DEC_BSP				FColor(255,128,0,255)
-#define DEC_PARTICLE		FColor(128,0,128,255)
-// general scene rendering events
-#define DEC_SCENE_ITEMS		FColor(128,128,128,255)
-
-/** Platform specific function for setting the value of a counter that can be viewed in PIX. */
-inline void appSetCounterValue(const TCHAR* CounterName, float Value) {}
-
 // Disable draw mesh events for final builds
 #if !(UE_BUILD_SHIPPING || UE_BUILD_TEST)
 #define WANTS_DRAW_MESH_EVENTS 1
@@ -61,17 +47,17 @@ inline void appSetCounterValue(const TCHAR* CounterName, float Value) {}
 		/**
 		 * Function for logging a PIX event with var args
 		 */
-		void CDECL Start(FRHICommandList& RHICmdList, const FColor& Color,const TCHAR* Fmt, ...);
+		void CDECL Start(FRHICommandList& RHICmdList, const TCHAR* Fmt, ...);
 		void Stop();
 	};
 
 	// Macros to allow for scoping of draw events
 	#define SCOPED_DRAW_TOKENPASTE_INNER(x,y) x##y
 	#define SCOPED_DRAW_TOKENPASTE(x,y) SCOPED_DRAW_TOKENPASTE_INNER(x,y)
-	#define SCOPED_DRAW_EVENT(RHICmdList, Name, Color) FDrawEvent SCOPED_DRAW_TOKENPASTE(Event_##Name,__LINE__); if(GEmitDrawEvents) SCOPED_DRAW_TOKENPASTE(Event_##Name,__LINE__).Start(RHICmdList, Color, TEXT(#Name));
-	#define SCOPED_DRAW_EVENTF(RHICmdList, Name, Color, Format, ...) FDrawEvent SCOPED_DRAW_TOKENPASTE(Event_##Name,__LINE__); if(GEmitDrawEvents) SCOPED_DRAW_TOKENPASTE(Event_##Name,__LINE__).Start(RHICmdList, Color, Format, ##__VA_ARGS__);
-	#define SCOPED_CONDITIONAL_DRAW_EVENT(RHICmdList, Name, Condition, Color) FDrawEvent SCOPED_DRAW_TOKENPASTE(Event_##Name,__LINE__); if(GEmitDrawEvents && (Condition)) SCOPED_DRAW_TOKENPASTE(Event_##Name,__LINE__).Start(RHICmdList, Color, TEXT(#Name));
-	#define SCOPED_CONDITIONAL_DRAW_EVENTF(RHICmdList, Name, Condition, Color, Format, ...) FDrawEvent SCOPED_DRAW_TOKENPASTE(Event_##Name,__LINE__); if(GEmitDrawEvents && (Condition)) SCOPED_DRAW_TOKENPASTE(Event_##Name,__LINE__).Start(RHICmdList, Color, Format, ##__VA_ARGS__);
+	#define SCOPED_DRAW_EVENT(RHICmdList, Name) FDrawEvent SCOPED_DRAW_TOKENPASTE(Event_##Name,__LINE__); if(GEmitDrawEvents) SCOPED_DRAW_TOKENPASTE(Event_##Name,__LINE__).Start(RHICmdList, TEXT(#Name));
+	#define SCOPED_DRAW_EVENTF(RHICmdList, Name, Format, ...) FDrawEvent SCOPED_DRAW_TOKENPASTE(Event_##Name,__LINE__); if(GEmitDrawEvents) SCOPED_DRAW_TOKENPASTE(Event_##Name,__LINE__).Start(RHICmdList, Format, ##__VA_ARGS__);
+	#define SCOPED_CONDITIONAL_DRAW_EVENT(RHICmdList, Name, Condition) FDrawEvent SCOPED_DRAW_TOKENPASTE(Event_##Name,__LINE__); if(GEmitDrawEvents && (Condition)) SCOPED_DRAW_TOKENPASTE(Event_##Name,__LINE__).Start(RHICmdList, TEXT(#Name));
+	#define SCOPED_CONDITIONAL_DRAW_EVENTF(RHICmdList, Name, Condition, Format, ...) FDrawEvent SCOPED_DRAW_TOKENPASTE(Event_##Name,__LINE__); if(GEmitDrawEvents && (Condition)) SCOPED_DRAW_TOKENPASTE(Event_##Name,__LINE__).Start(RHICmdList, Format, ##__VA_ARGS__);
 
 #else
 
diff --git a/Engine/Source/Runtime/Renderer/Private/AtmosphereRendering.cpp b/Engine/Source/Runtime/Renderer/Private/AtmosphereRendering.cpp
--- a/Engine/Source/Runtime/Renderer/Private/AtmosphereRendering.cpp
+++ b/Engine/Source/Runtime/Renderer/Private/AtmosphereRendering.cpp
@@ -398,7 +398,7 @@ void FDeferredShadingSceneRenderer::RenderAtmosphere(FRHICommandListImmediate& R
 	// Atmospheric fog?
 	if (Scene->GetFeatureLevel() >= ERHIFeatureLevel::SM4 && Scene->HasAtmosphericFog())
 	{
-		SCOPED_DRAW_EVENT(RHICmdList, Fog, DEC_SCENE_ITEMS);
+		SCOPED_DRAW_EVENT(RHICmdList, Fog);
 
 		static const FVector2D Vertices[4] =
 		{
diff --git a/Engine/Source/Runtime/Renderer/Private/CompositionLighting/CompositionLighting.cpp b/Engine/Source/Runtime/Renderer/Private/CompositionLighting/CompositionLighting.cpp
--- a/Engine/Source/Runtime/Renderer/Private/CompositionLighting/CompositionLighting.cpp
+++ b/Engine/Source/Runtime/Renderer/Private/CompositionLighting/CompositionLighting.cpp
@@ -252,7 +252,7 @@ void FCompositionLighting::ProcessBeforeBasePass(FRHICommandListImmediate& RHICm
 
 		// The graph setup should be finished before this line ----------------------------------------
 
-		SCOPED_DRAW_EVENT(RHICmdList, CompositionBeforeBasePass, DEC_SCENE_ITEMS);
+		SCOPED_DRAW_EVENT(RHICmdList, CompositionBeforeBasePass);
 
 		TRefCountPtr<IPooledRenderTarget>& SceneColor = GSceneRenderTargets.GetSceneColor();
 
@@ -312,7 +312,7 @@ void FCompositionLighting::ProcessAfterBasePass(FRHICommandListImmediate& RHICmd
 
 		// The graph setup should be finished before this line ----------------------------------------
 
-		SCOPED_DRAW_EVENT(RHICmdList, LightCompositionTasks_PreLighting, DEC_SCENE_ITEMS);
+		SCOPED_DRAW_EVENT(RHICmdList, LightCompositionTasks_PreLighting);
 
 		TRefCountPtr<IPooledRenderTarget>& SceneColor = GSceneRenderTargets.GetSceneColor();
 
@@ -378,7 +378,7 @@ void FCompositionLighting::ProcessLighting(FRHICommandListImmediate& RHICmdList,
 
 		// The graph setup should be finished before this line ----------------------------------------
 
-		SCOPED_DRAW_EVENT(RHICmdList, CompositionLighting, DEC_SCENE_ITEMS);
+		SCOPED_DRAW_EVENT(RHICmdList, CompositionLighting);
 
 		// we don't replace the final element with the scenecolor because this is what those passes should do by themself
 
diff --git a/Engine/Source/Runtime/Renderer/Private/CompositionLighting/PostProcessAmbient.cpp b/Engine/Source/Runtime/Renderer/Private/CompositionLighting/PostProcessAmbient.cpp
--- a/Engine/Source/Runtime/Renderer/Private/CompositionLighting/PostProcessAmbient.cpp
+++ b/Engine/Source/Runtime/Renderer/Private/CompositionLighting/PostProcessAmbient.cpp
@@ -71,7 +71,7 @@ IMPLEMENT_SHADER_TYPE(,FPostProcessAmbientPS,TEXT("PostProcessAmbient"),TEXT("Ma
 
 void FRCPassPostProcessAmbient::Process(FRenderingCompositePassContext& Context)
 {
-	SCOPED_DRAW_EVENT(Context.RHICmdList, PostProcessAmbient, DEC_SCENE_ITEMS);
+	SCOPED_DRAW_EVENT(Context.RHICmdList, PostProcessAmbient);
 
 	const FSceneView& View = Context.View;
 	const FSceneViewFamily& ViewFamily = *(View.Family);
diff --git a/Engine/Source/Runtime/Renderer/Private/CompositionLighting/PostProcessAmbientOcclusion.cpp b/Engine/Source/Runtime/Renderer/Private/CompositionLighting/PostProcessAmbientOcclusion.cpp
--- a/Engine/Source/Runtime/Renderer/Private/CompositionLighting/PostProcessAmbientOcclusion.cpp
+++ b/Engine/Source/Runtime/Renderer/Private/CompositionLighting/PostProcessAmbientOcclusion.cpp
@@ -361,7 +361,7 @@ void FRCPassPostProcessAmbientOcclusionSetup::SetShaderSetupTempl(const FRenderi
 
 void FRCPassPostProcessAmbientOcclusionSetup::Process(FRenderingCompositePassContext& Context)
 {
-	SCOPED_DRAW_EVENT(Context.RHICmdList, AmbientOcclusionSetup, DEC_SCENE_ITEMS);
+	SCOPED_DRAW_EVENT(Context.RHICmdList, AmbientOcclusionSetup);
 
 	const FSceneView& View = Context.View;
 
@@ -474,7 +474,7 @@ void FRCPassPostProcessAmbientOcclusion::SetShaderTempl(const FRenderingComposit
 
 void FRCPassPostProcessAmbientOcclusion::Process(FRenderingCompositePassContext& Context)
 {
-	SCOPED_DRAW_EVENT(Context.RHICmdList, AmbientOcclusion, DEC_SCENE_ITEMS);
+	SCOPED_DRAW_EVENT(Context.RHICmdList, AmbientOcclusion);
 
 	const FSceneView& View = Context.View;
 
@@ -591,7 +591,7 @@ FPooledRenderTargetDesc FRCPassPostProcessAmbientOcclusion::ComputeOutputDesc(EP
 
 void FRCPassPostProcessBasePassAO::Process(FRenderingCompositePassContext& Context)
 {
-	SCOPED_DRAW_EVENT(Context.RHICmdList, ApplyAOToBasePassSceneColor, DEC_SCENE_ITEMS);
+	SCOPED_DRAW_EVENT(Context.RHICmdList, ApplyAOToBasePassSceneColor);
 
 	const FSceneView& View = Context.View;
 
diff --git a/Engine/Source/Runtime/Renderer/Private/CompositionLighting/PostProcessDeferredDecals.cpp b/Engine/Source/Runtime/Renderer/Private/CompositionLighting/PostProcessDeferredDecals.cpp
--- a/Engine/Source/Runtime/Renderer/Private/CompositionLighting/PostProcessDeferredDecals.cpp
+++ b/Engine/Source/Runtime/Renderer/Private/CompositionLighting/PostProcessDeferredDecals.cpp
@@ -374,7 +374,7 @@ FGlobalBoundShaderState StencilDecalMaskBoundShaderState;
 /** Draws a full view quad that sets stencil to 1 anywhere that decals should not be projected. */
 void StencilDecalMask(FRHICommandList& RHICmdList, const FViewInfo& View)
 {
-	SCOPED_DRAW_EVENT(RHICmdList, StencilDecalMask, DEC_SCENE_ITEMS);
+	SCOPED_DRAW_EVENT(RHICmdList, StencilDecalMask);
 	RHICmdList.SetRasterizerState(TStaticRasterizerState<FM_Solid, CM_None>::GetRHI());
 	RHICmdList.SetBlendState(TStaticBlendState<CW_NONE>::GetRHI());
 	SetRenderTarget(RHICmdList, NULL, GSceneRenderTargets.GetSceneDepthSurface());
@@ -640,7 +640,7 @@ void SetShader(const FRenderingCompositePassContext& Context, const FTransientDe
 
 bool RenderPreStencil(FRenderingCompositePassContext& Context, const FMaterialShaderMap* MaterialShaderMap, const FMatrix& ComponentToWorldMatrix, const FMatrix& FrustumComponentToClip)
 {
-	SCOPED_DRAW_EVENT(Context.RHICmdList, RenderPreStencil, DEC_SCENE_ITEMS);
+	SCOPED_DRAW_EVENT(Context.RHICmdList, RenderPreStencil);
 
 	const FSceneView& View = Context.View;
 
@@ -721,7 +721,7 @@ void FRCPassPostProcessDeferredDecals::Process(FRenderingCompositePassContext& C
 		PassOutputs[0].RenderTargetDesc = OutputOfMyInput->RenderTargetDesc;
 	}
 
-	SCOPED_DRAW_EVENT(RHICmdList, PostProcessDeferredDecals, DEC_SCENE_ITEMS);
+	SCOPED_DRAW_EVENT(RHICmdList, PostProcessDeferredDecals);
 
 	if(RenderStage == 0)
 	{
@@ -754,7 +754,7 @@ void FRCPassPostProcessDeferredDecals::Process(FRenderingCompositePassContext& C
 		}
 
 
-		SCOPED_DRAW_EVENT(RHICmdList, DBufferClear, DEC_SCENE_ITEMS);
+		SCOPED_DRAW_EVENT(RHICmdList, DBufferClear);
 		{
 			// could be optimized
 			SetRenderTarget(RHICmdList, GSceneRenderTargets.DBufferA->GetRenderTargetItem().TargetableTexture, FTextureRHIParamRef());
@@ -888,7 +888,7 @@ void FRCPassPostProcessDeferredDecals::Process(FRenderingCompositePassContext& C
 		int32 WasInsideDecal = -1;
 		const ERHIFeatureLevel::Type SMFeatureLevel = Context.GetFeatureLevel();
 
-		SCOPED_DRAW_EVENT(RHICmdList, Decals, DEC_SCENE_ITEMS);
+		SCOPED_DRAW_EVENT(RHICmdList, Decals);
 		INC_DWORD_STAT_BY(STAT_Decals, SortedDecals.Num());
 		
 		RHICmdList.SetStreamSource(0, GUnitCubeVertexBuffer.VertexBufferRHI, sizeof(FVector4), 0);
diff --git a/Engine/Source/Runtime/Renderer/Private/CompositionLighting/PostProcessLpvIndirect.cpp b/Engine/Source/Runtime/Renderer/Private/CompositionLighting/PostProcessLpvIndirect.cpp
--- a/Engine/Source/Runtime/Renderer/Private/CompositionLighting/PostProcessLpvIndirect.cpp
+++ b/Engine/Source/Runtime/Renderer/Private/CompositionLighting/PostProcessLpvIndirect.cpp
@@ -130,7 +130,7 @@ IMPLEMENT_SHADER_TYPE(,FPostProcessLpvIndirectPS,TEXT("PostProcessLpvIndirect"),
 
 void FRCPassPostProcessLpvIndirect::Process(FRenderingCompositePassContext& Context)
 {
-	SCOPED_DRAW_EVENT(Context.RHICmdList, PostProcessLpvIndirect, DEC_SCENE_ITEMS);
+	SCOPED_DRAW_EVENT(Context.RHICmdList, PostProcessLpvIndirect);
 
 	{
 		FRenderingCompositeOutput* OutputOfMyInput = GetInput(ePId_Input0)->GetOutput();
diff --git a/Engine/Source/Runtime/Renderer/Private/CompositionLighting/PostProcessPassThrough.cpp b/Engine/Source/Runtime/Renderer/Private/CompositionLighting/PostProcessPassThrough.cpp
--- a/Engine/Source/Runtime/Renderer/Private/CompositionLighting/PostProcessPassThrough.cpp
+++ b/Engine/Source/Runtime/Renderer/Private/CompositionLighting/PostProcessPassThrough.cpp
@@ -70,7 +70,7 @@ FRCPassPostProcessPassThrough::FRCPassPostProcessPassThrough(FPooledRenderTarget
 
 void FRCPassPostProcessPassThrough::Process(FRenderingCompositePassContext& Context)
 {
-	SCOPED_DRAW_EVENT(Context.RHICmdList, PassThrough, DEC_SCENE_ITEMS);
+	SCOPED_DRAW_EVENT(Context.RHICmdList, PassThrough);
 
 	const FPooledRenderTargetDesc* InputDesc = GetInputDesc(ePId_Input0);
 
diff --git a/Engine/Source/Runtime/Renderer/Private/CustomDepthRendering.cpp b/Engine/Source/Runtime/Renderer/Private/CustomDepthRendering.cpp
--- a/Engine/Source/Runtime/Renderer/Private/CustomDepthRendering.cpp
+++ b/Engine/Source/Runtime/Renderer/Private/CustomDepthRendering.cpp
@@ -23,7 +23,7 @@ bool FCustomDepthPrimSet::DrawPrims(FRHICommandListImmediate& RHICmdList, const
 
 	if(Prims.Num())
 	{
-		SCOPED_DRAW_EVENT(RHICmdList, CustomDepth, DEC_SCENE_ITEMS);
+		SCOPED_DRAW_EVENT(RHICmdList, CustomDepth);
 
 		const bool bUseGetMeshElements = ShouldUseGetDynamicMeshElements();
 
diff --git a/Engine/Source/Runtime/Renderer/Private/DeferredShadingRenderer.cpp b/Engine/Source/Runtime/Renderer/Private/DeferredShadingRenderer.cpp
--- a/Engine/Source/Runtime/Renderer/Private/DeferredShadingRenderer.cpp
+++ b/Engine/Source/Runtime/Renderer/Private/DeferredShadingRenderer.cpp
@@ -130,7 +130,7 @@ extern FGlobalBoundShaderState GClearMRTBoundShaderState[8];
 void FDeferredShadingSceneRenderer::ClearGBufferAtMaxZ(FRHICommandList& RHICmdList)
 {
 	// Assumes BeginRenderingSceneColor() has been called before this function
-	SCOPED_DRAW_EVENT(RHICmdList, ClearGBufferAtMaxZ, DEC_SCENE_ITEMS);
+	SCOPED_DRAW_EVENT(RHICmdList, ClearGBufferAtMaxZ);
 
 	// Clear the G Buffer render targets
 	const bool bClearBlack = Views[0].Family->EngineShowFlags.ShaderComplexity || Views[0].Family->EngineShowFlags.StationaryLightOverlap;
@@ -180,7 +180,7 @@ void FDeferredShadingSceneRenderer::ClearGBufferAtMaxZ(FRHICommandList& RHICmdLi
 	// Clear each viewport by drawing background color at MaxZ depth
 	for(int32 ViewIndex = 0;ViewIndex < Views.Num();ViewIndex++)
 	{
-		SCOPED_CONDITIONAL_DRAW_EVENTF(RHICmdList, EventView, Views.Num() > 1, DEC_SCENE_ITEMS, TEXT("ClearView%d"), ViewIndex);
+		SCOPED_CONDITIONAL_DRAW_EVENTF(RHICmdList, EventView, Views.Num() > 1, TEXT("ClearView%d"), ViewIndex);
 
 		FViewInfo& View = Views[ViewIndex];
 
@@ -202,14 +202,14 @@ bool FDeferredShadingSceneRenderer::RenderBasePassStaticDataMasked(FRHICommandLi
 		// Draw the scene's base pass draw lists.
 		FScene::EBasePassDrawListType MaskedDrawType = FScene::EBasePass_Masked;
 		{
-			SCOPED_DRAW_EVENT(RHICmdList, StaticMaskedNoLightmap, DEC_SCENE_ITEMS);
+			SCOPED_DRAW_EVENT(RHICmdList, StaticMaskedNoLightmap);
 			bDirty |= Scene->BasePassNoLightMapDrawList[MaskedDrawType].DrawVisible(RHICmdList, View,View.StaticMeshVisibilityMap,View.StaticMeshBatchVisibility);
 			bDirty |= Scene->BasePassSimpleDynamicLightingDrawList[MaskedDrawType].DrawVisible(RHICmdList, View,View.StaticMeshVisibilityMap,View.StaticMeshBatchVisibility);
 			bDirty |= Scene->BasePassCachedVolumeIndirectLightingDrawList[MaskedDrawType].DrawVisible(RHICmdList, View,View.StaticMeshVisibilityMap,View.StaticMeshBatchVisibility);
 			bDirty |= Scene->BasePassCachedPointIndirectLightingDrawList[MaskedDrawType].DrawVisible(RHICmdList, View,View.StaticMeshVisibilityMap,View.StaticMeshBatchVisibility);
 		}
 		{
-			SCOPED_DRAW_EVENT(RHICmdList, StaticMaskedLightmapped, DEC_SCENE_ITEMS);
+			SCOPED_DRAW_EVENT(RHICmdList, StaticMaskedLightmapped);
 			bDirty |= Scene->BasePassHighQualityLightMapDrawList[MaskedDrawType].DrawVisible(RHICmdList, View,View.StaticMeshVisibilityMap,View.StaticMeshBatchVisibility);
 			bDirty |= Scene->BasePassDistanceFieldShadowMapLightMapDrawList[MaskedDrawType].DrawVisible(RHICmdList, View,View.StaticMeshVisibilityMap,View.StaticMeshBatchVisibility);
 			bDirty |= Scene->BasePassLowQualityLightMapDrawList[MaskedDrawType].DrawVisible(RHICmdList, View,View.StaticMeshVisibilityMap,View.StaticMeshBatchVisibility);
@@ -224,7 +224,7 @@ void FDeferredShadingSceneRenderer::RenderBasePassStaticDataMaskedParallel(FPara
 	FScene::EBasePassDrawListType MaskedDrawType = FScene::EBasePass_Masked;
 	{
 		// we can't insert this event on the parent command list; need to pass it along to the ParallelCommandListSet
-		//SCOPED_DRAW_EVENT(ParentCmdList, StaticMaskedNoLightmap, DEC_SCENE_ITEMS);
+		//SCOPED_DRAW_EVENT(ParentCmdList, StaticMaskedNoLightmap);
 		Scene->BasePassNoLightMapDrawList[MaskedDrawType].DrawVisibleParallel(ParallelCommandListSet.View.StaticMeshVisibilityMap, ParallelCommandListSet.View.StaticMeshBatchVisibility, ParallelCommandListSet);
 		Scene->BasePassSimpleDynamicLightingDrawList[MaskedDrawType].DrawVisibleParallel(ParallelCommandListSet.View.StaticMeshVisibilityMap, ParallelCommandListSet.View.StaticMeshBatchVisibility, ParallelCommandListSet);
 		Scene->BasePassCachedVolumeIndirectLightingDrawList[MaskedDrawType].DrawVisibleParallel(ParallelCommandListSet.View.StaticMeshVisibilityMap, ParallelCommandListSet.View.StaticMeshBatchVisibility, ParallelCommandListSet);
@@ -233,7 +233,7 @@ void FDeferredShadingSceneRenderer::RenderBasePassStaticDataMaskedParallel(FPara
 
 	{
 		// we can't insert this event on the parent command list; need to pass it along to the ParallelCommandListSet
-		//SCOPED_DRAW_EVENT(ParentCmdList, StaticMaskedLightmapped, DEC_SCENE_ITEMS);
+		//SCOPED_DRAW_EVENT(ParentCmdList, StaticMaskedLightmapped);
 		Scene->BasePassHighQualityLightMapDrawList[MaskedDrawType].DrawVisibleParallel(ParallelCommandListSet.View.StaticMeshVisibilityMap, ParallelCommandListSet.View.StaticMeshBatchVisibility, ParallelCommandListSet);
 		Scene->BasePassDistanceFieldShadowMapLightMapDrawList[MaskedDrawType].DrawVisibleParallel(ParallelCommandListSet.View.StaticMeshVisibilityMap, ParallelCommandListSet.View.StaticMeshBatchVisibility, ParallelCommandListSet);
 		Scene->BasePassLowQualityLightMapDrawList[MaskedDrawType].DrawVisibleParallel(ParallelCommandListSet.View.StaticMeshVisibilityMap, ParallelCommandListSet.View.StaticMeshBatchVisibility, ParallelCommandListSet);
@@ -246,14 +246,14 @@ bool FDeferredShadingSceneRenderer::RenderBasePassStaticDataDefault(FRHICommandL
 	{
 		FScene::EBasePassDrawListType OpaqueDrawType = FScene::EBasePass_Default;
 		{
-			SCOPED_DRAW_EVENT(RHICmdList, StaticOpaqueNoLightmap, DEC_SCENE_ITEMS);
+			SCOPED_DRAW_EVENT(RHICmdList, StaticOpaqueNoLightmap);
 			bDirty |= Scene->BasePassNoLightMapDrawList[OpaqueDrawType].DrawVisible(RHICmdList, View,View.StaticMeshVisibilityMap,View.StaticMeshBatchVisibility);
 			bDirty |= Scene->BasePassSimpleDynamicLightingDrawList[OpaqueDrawType].DrawVisible(RHICmdList, View,View.StaticMeshVisibilityMap,View.StaticMeshBatchVisibility);
 			bDirty |= Scene->BasePassCachedVolumeIndirectLightingDrawList[OpaqueDrawType].DrawVisible(RHICmdList, View,View.StaticMeshVisibilityMap,View.StaticMeshBatchVisibility);
 			bDirty |= Scene->BasePassCachedPointIndirectLightingDrawList[OpaqueDrawType].DrawVisible(RHICmdList, View,View.StaticMeshVisibilityMap,View.StaticMeshBatchVisibility);
 		}
 		{
-			SCOPED_DRAW_EVENT(RHICmdList, StaticOpaqueLightmapped, DEC_SCENE_ITEMS);
+			SCOPED_DRAW_EVENT(RHICmdList, StaticOpaqueLightmapped);
 			bDirty |= Scene->BasePassHighQualityLightMapDrawList[OpaqueDrawType].DrawVisible(RHICmdList, View,View.StaticMeshVisibilityMap,View.StaticMeshBatchVisibility);
 			bDirty |= Scene->BasePassDistanceFieldShadowMapLightMapDrawList[OpaqueDrawType].DrawVisible(RHICmdList, View,View.StaticMeshVisibilityMap,View.StaticMeshBatchVisibility);
 			bDirty |= Scene->BasePassLowQualityLightMapDrawList[OpaqueDrawType].DrawVisible(RHICmdList, View,View.StaticMeshVisibilityMap,View.StaticMeshBatchVisibility);
@@ -268,7 +268,7 @@ void FDeferredShadingSceneRenderer::RenderBasePassStaticDataDefaultParallel(FPar
 	FScene::EBasePassDrawListType OpaqueDrawType = FScene::EBasePass_Default;
 	{
 		// we can't insert this event on the parent command list; need to pass it along to the ParallelCommandListSet
-		//SCOPED_DRAW_EVENT(ParentCmdList, StaticOpaqueNoLightmap, DEC_SCENE_ITEMS);
+		//SCOPED_DRAW_EVENT(ParentCmdList, StaticOpaqueNoLightmap);
 		Scene->BasePassNoLightMapDrawList[OpaqueDrawType].DrawVisibleParallel(ParallelCommandListSet.View.StaticMeshVisibilityMap, ParallelCommandListSet.View.StaticMeshBatchVisibility, ParallelCommandListSet);
 		Scene->BasePassSimpleDynamicLightingDrawList[OpaqueDrawType].DrawVisibleParallel(ParallelCommandListSet.View.StaticMeshVisibilityMap, ParallelCommandListSet.View.StaticMeshBatchVisibility, ParallelCommandListSet);
 		Scene->BasePassCachedVolumeIndirectLightingDrawList[OpaqueDrawType].DrawVisibleParallel(ParallelCommandListSet.View.StaticMeshVisibilityMap, ParallelCommandListSet.View.StaticMeshBatchVisibility, ParallelCommandListSet);
@@ -277,7 +277,7 @@ void FDeferredShadingSceneRenderer::RenderBasePassStaticDataDefaultParallel(FPar
 
 	{
 		// we can't insert this event on the parent command list; need to pass it along to the ParallelCommandListSet
-		//SCOPED_DRAW_EVENT(ParentCmdList, StaticOpaqueLightmapped, DEC_SCENE_ITEMS);
+		//SCOPED_DRAW_EVENT(ParentCmdList, StaticOpaqueLightmapped);
 		Scene->BasePassHighQualityLightMapDrawList[OpaqueDrawType].DrawVisibleParallel(ParallelCommandListSet.View.StaticMeshVisibilityMap, ParallelCommandListSet.View.StaticMeshBatchVisibility, ParallelCommandListSet);
 		Scene->BasePassDistanceFieldShadowMapLightMapDrawList[OpaqueDrawType].DrawVisibleParallel(ParallelCommandListSet.View.StaticMeshVisibilityMap, ParallelCommandListSet.View.StaticMeshBatchVisibility, ParallelCommandListSet);
 		Scene->BasePassLowQualityLightMapDrawList[OpaqueDrawType].DrawVisibleParallel(ParallelCommandListSet.View.StaticMeshVisibilityMap, ParallelCommandListSet.View.StaticMeshBatchVisibility, ParallelCommandListSet);
@@ -370,7 +370,7 @@ void FDeferredShadingSceneRenderer::RenderBasePassDynamicData(FRHICommandList& R
 	bool bDirty = false;
 
 	SCOPE_CYCLE_COUNTER(STAT_DynamicPrimitiveDrawTime);
-	SCOPED_DRAW_EVENT(RHICmdList, Dynamic, DEC_SCENE_ITEMS);
+	SCOPED_DRAW_EVENT(RHICmdList, Dynamic);
 
 	const bool bUseGetMeshElements = ShouldUseGetDynamicMeshElements();
 
@@ -661,7 +661,7 @@ void FDeferredShadingSceneRenderer::Render(FRHICommandListImmediate& RHICmdList)
 	{
 		return;
 	}
-	SCOPED_DRAW_EVENT(RHICmdList, Scene,DEC_SCENE_ITEMS);
+	SCOPED_DRAW_EVENT(RHICmdList, Scene);
 	
 	// Initialize global system textures (pass-through if already initialized).
 	GSystemTextures.InitializeTextures(RHICmdList, FeatureLevel);
@@ -785,7 +785,7 @@ void FDeferredShadingSceneRenderer::Render(FRHICommandListImmediate& RHICmdList)
 		// e.g. ambient cubemaps, ambient occlusion, deferred decals
 		for(int32 ViewIndex = 0;ViewIndex < Views.Num();ViewIndex++)
 		{	
-			SCOPED_CONDITIONAL_DRAW_EVENTF(RHICmdList, EventView,Views.Num() > 1, DEC_SCENE_ITEMS, TEXT("View%d"), ViewIndex);
+			SCOPED_CONDITIONAL_DRAW_EVENTF(RHICmdList, EventView,Views.Num() > 1, TEXT("View%d"), ViewIndex);
 			GCompositionLighting.ProcessBeforeBasePass(RHICmdList, Views[ViewIndex]);
 		}
 	}
@@ -881,7 +881,7 @@ void FDeferredShadingSceneRenderer::Render(FRHICommandListImmediate& RHICmdList)
 		// e.g. deferred decals, blurred GBuffer
 		for(int32 ViewIndex = 0;ViewIndex < Views.Num();ViewIndex++)
 		{	
-			SCOPED_CONDITIONAL_DRAW_EVENTF(RHICmdList, EventView,Views.Num() > 1, DEC_SCENE_ITEMS, TEXT("View%d"), ViewIndex);
+			SCOPED_CONDITIONAL_DRAW_EVENTF(RHICmdList, EventView,Views.Num() > 1, TEXT("View%d"), ViewIndex);
 			GCompositionLighting.ProcessAfterBasePass(RHICmdList, Views[ViewIndex]);
 		}
 
@@ -915,7 +915,7 @@ void FDeferredShadingSceneRenderer::Render(FRHICommandListImmediate& RHICmdList)
 		// e.g. ambient cubemaps, ambient occlusion, LPV indirect
 		for(int32 ViewIndex = 0; ViewIndex < Views.Num(); ++ViewIndex)
 		{	
-			SCOPED_CONDITIONAL_DRAW_EVENTF(RHICmdList, EventView,Views.Num() > 1, DEC_SCENE_ITEMS, TEXT("View%d"), ViewIndex);
+			SCOPED_CONDITIONAL_DRAW_EVENTF(RHICmdList, EventView,Views.Num() > 1, TEXT("View%d"), ViewIndex);
 			GCompositionLighting.ProcessLighting(RHICmdList, Views[ViewIndex]);
 		}
 	}
@@ -1018,12 +1018,12 @@ void FDeferredShadingSceneRenderer::Render(FRHICommandListImmediate& RHICmdList)
 	// Finish rendering for each view.
 	if(ViewFamily.bResolveScene)
 	{
-		SCOPED_DRAW_EVENT(RHICmdList, PostProcessing, DEC_SCENE_ITEMS);
+		SCOPED_DRAW_EVENT(RHICmdList, PostProcessing);
 		SCOPE_CYCLE_COUNTER(STAT_FinishRenderViewTargetTime);
 
 		for(int32 ViewIndex = 0; ViewIndex < Views.Num(); ViewIndex++)
 		{
-			SCOPED_CONDITIONAL_DRAW_EVENTF(RHICmdList, EventView, Views.Num() > 1, DEC_SCENE_ITEMS, TEXT("View%d"), ViewIndex);
+			SCOPED_CONDITIONAL_DRAW_EVENTF(RHICmdList, EventView, Views.Num() > 1, TEXT("View%d"), ViewIndex);
 
 			GPostProcessing.Process(RHICmdList, Views[ ViewIndex ], VelocityRT);
 
@@ -1087,7 +1087,7 @@ bool FDeferredShadingSceneRenderer::RenderPrePassViewDynamic(FRHICommandList& RH
 		// Draw the dynamic occluder primitives using a depth drawing policy.
 		TDynamicPrimitiveDrawer<FDepthDrawingPolicyFactory> Drawer(RHICmdList, &View, FDepthDrawingPolicyFactory::ContextType(EarlyZPassMode), true);
 		{
-			SCOPED_DRAW_EVENT(RHICmdList, Dynamic, DEC_SCENE_ITEMS);
+			SCOPED_DRAW_EVENT(RHICmdList, Dynamic);
 			for(int32 PrimitiveIndex = 0;PrimitiveIndex < View.VisibleDynamicPrimitives.Num();PrimitiveIndex++)
 			{
 				const FPrimitiveSceneInfo* PrimitiveSceneInfo = View.VisibleDynamicPrimitives[PrimitiveIndex];
@@ -1147,19 +1147,19 @@ bool FDeferredShadingSceneRenderer::RenderPrePassView(FRHICommandList& RHICmdLis
 		// Draw opaque occluders which support a separate position-only
 		// vertex buffer to minimize vertex fetch bandwidth, which is
 		// often the bottleneck during the depth only pass.
-		SCOPED_DRAW_EVENT(RHICmdList, PosOnlyOpaque, DEC_SCENE_ITEMS);
+		SCOPED_DRAW_EVENT(RHICmdList, PosOnlyOpaque);
 		bDirty |= Scene->PositionOnlyDepthDrawList.DrawVisible(RHICmdList, View,View.StaticMeshOccluderMap,View.StaticMeshBatchVisibility);
 	}
 	{
 		// Draw opaque occluders, using double speed z where supported.
-		SCOPED_DRAW_EVENT(RHICmdList, Opaque, DEC_SCENE_ITEMS);
+		SCOPED_DRAW_EVENT(RHICmdList, Opaque);
 		bDirty |= Scene->DepthDrawList.DrawVisible(RHICmdList, View,View.StaticMeshOccluderMap,View.StaticMeshBatchVisibility);
 	}
 
 	if(EarlyZPassMode >= DDM_AllOccluders)
 	{
 		// Draw opaque occluders with masked materials
-		SCOPED_DRAW_EVENT(RHICmdList, Opaque, DEC_SCENE_ITEMS);
+		SCOPED_DRAW_EVENT(RHICmdList, Opaque);
 		bDirty |= Scene->MaskedDepthDrawList.DrawVisible(RHICmdList, View,View.StaticMeshOccluderMap,View.StaticMeshBatchVisibility);
 	}
 
@@ -1242,13 +1242,13 @@ void FDeferredShadingSceneRenderer::RenderPrePassViewParallel(const FViewInfo& V
 		// vertex buffer to minimize vertex fetch bandwidth, which is
 		// often the bottleneck during the depth only pass.
 		// we can't insert this event on the parent command list; need to pass it along to the ParallelCommandListSet
-		//SCOPED_DRAW_EVENT(ParentCmdList, PosOnlyOpaque, DEC_SCENE_ITEMS);
+		//SCOPED_DRAW_EVENT(ParentCmdList, PosOnlyOpaque);
 		Scene->PositionOnlyDepthDrawList.DrawVisibleParallel(View.StaticMeshOccluderMap, View.StaticMeshBatchVisibility, ParallelCommandListSet);
 	}
 	{
 		// Draw opaque occluders, using double speed z where supported.
 		// we can't insert this event on the parent command list; need to pass it along to the ParallelCommandListSet
-		//SCOPED_DRAW_EVENT(ParentCmdList, Opaque, DEC_SCENE_ITEMS);
+		//SCOPED_DRAW_EVENT(ParentCmdList, Opaque);
 		Scene->DepthDrawList.DrawVisibleParallel(View.StaticMeshOccluderMap,View.StaticMeshBatchVisibility, ParallelCommandListSet);
 	}
 
@@ -1256,7 +1256,7 @@ void FDeferredShadingSceneRenderer::RenderPrePassViewParallel(const FViewInfo& V
 	{
 		// Draw opaque occluders with masked materials
 		// we can't insert this event on the parent command list; need to pass it along to the ParallelCommandListSet
-		//SCOPED_DRAW_EVENT(ParentCmdList, Opaque, DEC_SCENE_ITEMS);
+		//SCOPED_DRAW_EVENT(ParentCmdList, Opaque);
 		Scene->MaskedDepthDrawList.DrawVisibleParallel(View.StaticMeshOccluderMap,View.StaticMeshBatchVisibility, ParallelCommandListSet);
 	}
 	{
@@ -1272,7 +1272,7 @@ void FDeferredShadingSceneRenderer::RenderPrePassViewParallel(const FViewInfo& V
 /** Renders the scene's prepass and occlusion queries */
 bool FDeferredShadingSceneRenderer::RenderPrePass(FRHICommandListImmediate& RHICmdList)
 {
-	SCOPED_DRAW_EVENT(RHICmdList, PrePass, DEC_SCENE_ITEMS);
+	SCOPED_DRAW_EVENT(RHICmdList, PrePass);
 	SCOPE_CYCLE_COUNTER(STAT_DepthDrawTime);
 
 	bool bDirty = false;
@@ -1292,7 +1292,7 @@ bool FDeferredShadingSceneRenderer::RenderPrePass(FRHICommandListImmediate& RHIC
 
 			for(int32 ViewIndex = 0;ViewIndex < Views.Num();ViewIndex++)
 			{
-				SCOPED_CONDITIONAL_DRAW_EVENTF(RHICmdList, EventView, Views.Num() > 1, DEC_SCENE_ITEMS, TEXT("View%d"), ViewIndex);
+				SCOPED_CONDITIONAL_DRAW_EVENTF(RHICmdList, EventView, Views.Num() > 1, TEXT("View%d"), ViewIndex);
 				const FViewInfo& View = Views[ViewIndex];
 				RenderPrePassViewParallel(View, RHICmdList, bDirty);
 			}
@@ -1301,7 +1301,7 @@ bool FDeferredShadingSceneRenderer::RenderPrePass(FRHICommandListImmediate& RHIC
 		{
 			for(int32 ViewIndex = 0;ViewIndex < Views.Num();ViewIndex++)
 			{
-				SCOPED_CONDITIONAL_DRAW_EVENTF(RHICmdList, EventView, Views.Num() > 1, DEC_SCENE_ITEMS, TEXT("View%d"), ViewIndex);
+				SCOPED_CONDITIONAL_DRAW_EVENTF(RHICmdList, EventView, Views.Num() > 1, TEXT("View%d"), ViewIndex);
 				const FViewInfo& View = Views[ViewIndex];
 				bDirty |= RenderPrePassView(RHICmdList, View);
 			}
@@ -1328,7 +1328,7 @@ bool FDeferredShadingSceneRenderer::RenderBasePass(FRHICommandListImmediate& RHI
 	}
 	else
 	{
-		SCOPED_DRAW_EVENT(RHICmdList, BasePass, DEC_SCENE_ITEMS);
+		SCOPED_DRAW_EVENT(RHICmdList, BasePass);
 		SCOPE_CYCLE_COUNTER(STAT_BasePassDrawTime);
 
 		if (GRHICommandList.UseParallelAlgorithms())
@@ -1336,7 +1336,7 @@ bool FDeferredShadingSceneRenderer::RenderBasePass(FRHICommandListImmediate& RHI
 			FScopedCommandListWaitForTasks Flusher(RHICmdList);
 			for(int32 ViewIndex = 0;ViewIndex < Views.Num();ViewIndex++)
 			{
-				SCOPED_CONDITIONAL_DRAW_EVENTF(RHICmdList, EventView, Views.Num() > 1, DEC_SCENE_ITEMS, TEXT("View%d"), ViewIndex);
+				SCOPED_CONDITIONAL_DRAW_EVENTF(RHICmdList, EventView, Views.Num() > 1, TEXT("View%d"), ViewIndex);
 				FViewInfo& View = Views[ViewIndex];
 				RenderBasePassViewParallel(View, RHICmdList, bDirty);
 			}
@@ -1345,7 +1345,7 @@ bool FDeferredShadingSceneRenderer::RenderBasePass(FRHICommandListImmediate& RHI
 		{
 			for(int32 ViewIndex = 0;ViewIndex < Views.Num();ViewIndex++)
 			{
-				SCOPED_CONDITIONAL_DRAW_EVENTF(RHICmdList, EventView, Views.Num() > 1, DEC_SCENE_ITEMS, TEXT("View%d"), ViewIndex);
+				SCOPED_CONDITIONAL_DRAW_EVENTF(RHICmdList, EventView, Views.Num() > 1, TEXT("View%d"), ViewIndex);
 				FViewInfo& View = Views[ViewIndex];
 
 				bDirty |= RenderBasePassView(RHICmdList, View);
@@ -1372,7 +1372,7 @@ void FDeferredShadingSceneRenderer::ClearLPVs(FRHICommandListImmediate& RHICmdLi
 
 			if(LightPropagationVolume)
 			{
-				SCOPED_DRAW_EVENT(RHICmdList, ClearLPVs, DEC_SCENE_ITEMS);
+				SCOPED_DRAW_EVENT(RHICmdList, ClearLPVs);
 				SCOPE_CYCLE_COUNTER(STAT_UpdateLPVs);
 				LightPropagationVolume->InitSettings(RHICmdList, Views[ViewIndex]);
 				LightPropagationVolume->Clear(RHICmdList, View);
@@ -1394,7 +1394,7 @@ void FDeferredShadingSceneRenderer::PropagateLPVs(FRHICommandListImmediate& RHIC
 
 			if(LightPropagationVolume)
 			{
-				SCOPED_DRAW_EVENT(RHICmdList, UpdateLPVs, DEC_SCENE_ITEMS);
+				SCOPED_DRAW_EVENT(RHICmdList, UpdateLPVs);
 				SCOPE_CYCLE_COUNTER(STAT_UpdateLPVs);
 				
 				LightPropagationVolume->Propagate(RHICmdList, View);
@@ -1472,7 +1472,7 @@ void FDeferredShadingSceneRenderer::UpdateDownsampledDepthSurface(FRHICommandLis
 	{
 		SetRenderTarget(RHICmdList, NULL, GSceneRenderTargets.GetSmallDepthSurface());
 
-		SCOPED_DRAW_EVENT(RHICmdList, DownsampleDepth, DEC_SCENE_ITEMS);
+		SCOPED_DRAW_EVENT(RHICmdList, DownsampleDepth);
 
 		for (int32 ViewIndex = 0; ViewIndex < Views.Num(); ViewIndex++)
 		{
diff --git a/Engine/Source/Runtime/Renderer/Private/DistanceFieldShadowing.cpp b/Engine/Source/Runtime/Renderer/Private/DistanceFieldShadowing.cpp
--- a/Engine/Source/Runtime/Renderer/Private/DistanceFieldShadowing.cpp
+++ b/Engine/Source/Runtime/Renderer/Private/DistanceFieldShadowing.cpp
@@ -632,7 +632,7 @@ void FProjectedShadowInfo::RenderRayTracedDistanceFieldProjection(FRHICommandLis
 		&& DoesPlatformSupportDistanceFieldShadowing(View.GetShaderPlatform()))
 	{
 		QUICK_SCOPE_CYCLE_COUNTER(STAT_RenderRayTracedDistanceFieldShadows);
-		SCOPED_DRAW_EVENT(RHICmdList, RayTracedDistanceFieldShadow, DEC_SCENE_ITEMS);
+		SCOPED_DRAW_EVENT(RHICmdList, RayTracedDistanceFieldShadow);
 
 		// Update the global distance field atlas
 		GDistanceFieldVolumeTextureAtlas.UpdateAllocations();
@@ -674,7 +674,7 @@ void FProjectedShadowInfo::RenderRayTracedDistanceFieldProjection(FRHICommandLis
 					TileIntersectionResources = LightSceneInfo->TileIntersectionResources;
 
 					{
-						SCOPED_DRAW_EVENT(RHICmdList, ClearTiles, DEC_SCENE_ITEMS);
+						SCOPED_DRAW_EVENT(RHICmdList, ClearTiles);
 						TShaderMapRef<FClearTilesCS> ComputeShader(View.ShaderMap);
 
 						uint32 GroupSizeX = FMath::DivideAndRoundUp(LightTileDimensions.X, GDistanceFieldAOTileSizeX);
@@ -688,7 +688,7 @@ void FProjectedShadowInfo::RenderRayTracedDistanceFieldProjection(FRHICommandLis
 					}
 
 					{
-						SCOPED_DRAW_EVENT(RHICmdList, CullObjects, DEC_SCENE_ITEMS);
+						SCOPED_DRAW_EVENT(RHICmdList, CullObjects);
 
 						TShaderMapRef<FShadowObjectCullVS> VertexShader(View.ShaderMap);
 						TShaderMapRef<FShadowObjectCullPS> PixelShader(View.ShaderMap);
@@ -745,7 +745,7 @@ void FProjectedShadowInfo::RenderRayTracedDistanceFieldProjection(FRHICommandLis
 					uint32 GroupSizeY = FMath::DivideAndRoundUp(ScissorRect.Size().Y / GAODownsampleFactor, GDistanceFieldAOTileSizeY);
 
 					{
-						SCOPED_DRAW_EVENT(RHICmdList, RayTraceShadows, DEC_SCENE_ITEMS);
+						SCOPED_DRAW_EVENT(RHICmdList, RayTraceShadows);
 						SetRenderTarget(RHICmdList, NULL, NULL);
 
 						if (bDirectionalLight && GShadowScatterTileCulling)
@@ -778,7 +778,7 @@ void FProjectedShadowInfo::RenderRayTracedDistanceFieldProjection(FRHICommandLis
 				{
 					GSceneRenderTargets.BeginRenderingLightAttenuation(RHICmdList);
 
-					SCOPED_DRAW_EVENT(RHICmdList, Upsample, DEC_SCENE_ITEMS);
+					SCOPED_DRAW_EVENT(RHICmdList, Upsample);
 
 					RHICmdList.SetViewport(ScissorRect.Min.X, ScissorRect.Min.Y, 0.0f, ScissorRect.Max.X, ScissorRect.Max.Y, 1.0f);
 					RHICmdList.SetRasterizerState(TStaticRasterizerState<FM_Solid, CM_None>::GetRHI());
diff --git a/Engine/Source/Runtime/Renderer/Private/DistanceFieldSurfaceCacheLighting.cpp b/Engine/Source/Runtime/Renderer/Private/DistanceFieldSurfaceCacheLighting.cpp
--- a/Engine/Source/Runtime/Renderer/Private/DistanceFieldSurfaceCacheLighting.cpp
+++ b/Engine/Source/Runtime/Renderer/Private/DistanceFieldSurfaceCacheLighting.cpp
@@ -963,7 +963,7 @@ void ComputeDistanceFieldNormal(FRHICommandListImmediate& RHICmdList, const TArr
 			uint32 GroupSizeY = FMath::DivideAndRoundUp(View.ViewRect.Size().Y / GAODownsampleFactor, GDistanceFieldAOTileSizeY);
 
 			{
-				SCOPED_DRAW_EVENT(RHICmdList, ComputeNormalCS, DEC_SCENE_ITEMS);
+				SCOPED_DRAW_EVENT(RHICmdList, ComputeNormalCS);
 				TShaderMapRef<FComputeDistanceFieldNormalCS> ComputeShader(View.ShaderMap);
 
 				RHICmdList.SetComputeShader(ComputeShader->GetComputeShader());
@@ -982,7 +982,7 @@ void ComputeDistanceFieldNormal(FRHICommandListImmediate& RHICmdList, const TArr
 		{
 			const FViewInfo& View = Views[ViewIndex];
 
-			SCOPED_DRAW_EVENT(RHICmdList, ComputeNormal, DEC_SCENE_ITEMS);
+			SCOPED_DRAW_EVENT(RHICmdList, ComputeNormal);
 
 			RHICmdList.SetViewport(0, 0, 0.0f, View.ViewRect.Width() / GAODownsampleFactor, View.ViewRect.Height() / GAODownsampleFactor, 1.0f);
 			RHICmdList.SetRasterizerState(TStaticRasterizerState<FM_Solid, CM_None>::GetRHI());
@@ -2374,7 +2374,7 @@ void UpdateHistory(
 			AllocateOrReuseAORenderTarget(NewHistory, HistoryRTName);
 
 			{
-				SCOPED_DRAW_EVENT(RHICmdList, UpdateHistory, DEC_SCENE_ITEMS);
+				SCOPED_DRAW_EVENT(RHICmdList, UpdateHistory);
 				SetRenderTarget(RHICmdList, NewHistory->GetRenderTargetItem().TargetableTexture, NULL);
 
 				RHICmdList.SetViewport(0, 0, 0.0f, View.ViewRect.Width() / GAODownsampleFactor, View.ViewRect.Height() / GAODownsampleFactor, 1.0f);
@@ -2440,7 +2440,7 @@ void PostProcessBentNormalAO(
 	}
 
 	{
-		SCOPED_DRAW_EVENT(RHICmdList, AOCombine, DEC_SCENE_ITEMS);
+		SCOPED_DRAW_EVENT(RHICmdList, AOCombine);
 
 		SetRenderTarget(RHICmdList, GAOFillGaps
 			? DistanceFieldAOBentNormal2->GetRenderTargetItem().TargetableTexture
@@ -2479,7 +2479,7 @@ void PostProcessBentNormalAO(
 
 	if (GAOFillGaps)
 	{
-		SCOPED_DRAW_EVENT(RHICmdList, FillGaps, DEC_SCENE_ITEMS);
+		SCOPED_DRAW_EVENT(RHICmdList, FillGaps);
 		SetRenderTarget(RHICmdList, DistanceFieldAOBentNormal->GetRenderTargetItem().TargetableTexture, NULL);
 
 		for (int32 ViewIndex = 0; ViewIndex < Views.Num(); ViewIndex++)
@@ -2595,7 +2595,7 @@ void UpsampleBentNormalAO(FRHICommandList& RHICmdList, const TArray<FViewInfo>&
 	{
 		const FViewInfo& View = Views[ViewIndex];
 
-		SCOPED_DRAW_EVENT(RHICmdList, UpsampleAO, DEC_SCENE_ITEMS);
+		SCOPED_DRAW_EVENT(RHICmdList, UpsampleAO);
 
 		RHICmdList.SetViewport(View.ViewRect.Min.X, View.ViewRect.Min.Y, 0.0f, View.ViewRect.Max.X, View.ViewRect.Max.Y, 1.0f);
 		RHICmdList.SetRasterizerState(TStaticRasterizerState<FM_Solid, CM_None>::GetRHI());
@@ -2790,7 +2790,7 @@ void GenerateBestSpacedVectors()
 
 FIntPoint BuildTileObjectLists(FRHICommandListImmediate& RHICmdList, FScene* Scene, TArray<FViewInfo>& Views, int32 NumObjects, const FDistanceFieldAOParameters& Parameters)
 {
-	SCOPED_DRAW_EVENT(RHICmdList, BuildTileList, DEC_SCENE_ITEMS);
+	SCOPED_DRAW_EVENT(RHICmdList, BuildTileList);
 	SetRenderTarget(RHICmdList, NULL, NULL);
 
 	FIntPoint TileListGroupSize;
@@ -2825,7 +2825,7 @@ FIntPoint BuildTileObjectLists(FRHICommandListImmediate& RHICmdList, FScene* Sce
 			}
 
 			{
-				SCOPED_DRAW_EVENT(RHICmdList, BuildTileCones, DEC_SCENE_ITEMS);
+				SCOPED_DRAW_EVENT(RHICmdList, BuildTileCones);
 				TShaderMapRef<FBuildTileConesCS> ComputeShader(View.ShaderMap);
 
 				RHICmdList.SetComputeShader(ComputeShader->GetComputeShader());
@@ -2836,7 +2836,7 @@ FIntPoint BuildTileObjectLists(FRHICommandListImmediate& RHICmdList, FScene* Sce
 			}
 
 			{
-				SCOPED_DRAW_EVENT(RHICmdList, CullObjects, DEC_SCENE_ITEMS);
+				SCOPED_DRAW_EVENT(RHICmdList, CullObjects);
 
 				TShaderMapRef<FObjectCullVS> VertexShader(View.ShaderMap);
 				TShaderMapRef<FObjectCullPS> PixelShader(View.ShaderMap);
@@ -2921,7 +2921,7 @@ FIntPoint BuildTileObjectLists(FRHICommandListImmediate& RHICmdList, FScene* Sce
 
 void SetupDepthStencil(FRHICommandListImmediate& RHICmdList, FViewInfo& View, FSceneRenderTargetItem& DistanceFieldNormal, int32 DepthLevel, int32 DestLevelDownsampleFactor)
 {
-	SCOPED_DRAW_EVENT(RHICmdList, SetupDepthStencil, DEC_SCENE_ITEMS);
+	SCOPED_DRAW_EVENT(RHICmdList, SetupDepthStencil);
 
 	{		
 		RHICmdList.SetViewport(0, 0, 0.0f, View.ViewRect.Width() / GAODownsampleFactor, View.ViewRect.Height() / GAODownsampleFactor, 1.0f);
@@ -3054,7 +3054,7 @@ void RenderIrradianceCacheInterpolation(
 			RHICmdList.Clear(true, FLinearColor(0, 0, 0, 0), false, 0, false, 0, FIntRect());
 		}
 
-		SCOPED_DRAW_EVENT(RHICmdList, IrradianceCacheSplat, DEC_SCENE_ITEMS);
+		SCOPED_DRAW_EVENT(RHICmdList, IrradianceCacheSplat);
 
 		static int32 NumInterpolationSections = 8;
 		if (GCircleVertexBuffer.NumSections != NumInterpolationSections)
@@ -3322,7 +3322,7 @@ bool FDeferredShadingSceneRenderer::RenderDistanceFieldAOSurfaceCache(FRHIComman
 		&& View.State)
 	{
 		QUICK_SCOPE_CYCLE_COUNTER(STAT_RenderDistanceFieldAOSurfaceCache);
-		SCOPED_DRAW_EVENT(RHICmdList, DistanceFieldAO, DEC_SCENE_ITEMS);
+		SCOPED_DRAW_EVENT(RHICmdList, DistanceFieldAO);
 
 		// Update the global distance field atlas
 		GDistanceFieldVolumeTextureAtlas.UpdateAllocations();
@@ -3385,7 +3385,7 @@ bool FDeferredShadingSceneRenderer::RenderDistanceFieldAOSurfaceCache(FRHIComman
 
 				if (GAOReuseAcrossFrames)
 				{
-					SCOPED_DRAW_EVENT(RHICmdList, TrimRecords, DEC_SCENE_ITEMS);
+					SCOPED_DRAW_EVENT(RHICmdList, TrimRecords);
 
 					// Copy and trim last frame's surface cache samples
 					for (int32 DepthLevel = GAOMaxLevel; DepthLevel >= GAOMinLevel; DepthLevel--)
@@ -3424,7 +3424,7 @@ bool FDeferredShadingSceneRenderer::RenderDistanceFieldAOSurfaceCache(FRHIComman
 				{
 					int32 DestLevelDownsampleFactor = GAODownsampleFactor * (1 << (DepthLevel * GAOPowerOfTwoBetweenLevels));
 
-					SCOPED_DRAW_EVENTF(RHICmdList, Level, DEC_SCENE_ITEMS, TEXT("Level_%d"), DepthLevel);
+					SCOPED_DRAW_EVENTF(RHICmdList, Level, TEXT("Level_%d"), DepthLevel);
 
 					TRefCountPtr<IPooledRenderTarget> DistanceFieldAOIrradianceCacheSplat;
 
@@ -3438,7 +3438,7 @@ bool FDeferredShadingSceneRenderer::RenderDistanceFieldAOSurfaceCache(FRHIComman
 					RenderIrradianceCacheInterpolation(RHICmdList, Views, DistanceFieldAOIrradianceCacheSplat, DistanceFieldNormal->GetRenderTargetItem(), DepthLevel, DestLevelDownsampleFactor, Parameters, false);
 
 					{
-						SCOPED_DRAW_EVENT(RHICmdList, PopulateIrradianceCache, DEC_SCENE_ITEMS);
+						SCOPED_DRAW_EVENT(RHICmdList, PopulateIrradianceCache);
 						SetRenderTarget(RHICmdList, NULL, NULL);
 
 						// Save off the current record count before adding any more
@@ -3472,7 +3472,7 @@ bool FDeferredShadingSceneRenderer::RenderDistanceFieldAOSurfaceCache(FRHIComman
 					}
 
 					{
-						SCOPED_DRAW_EVENT(RHICmdList, ShadeIrradianceCache, DEC_SCENE_ITEMS);
+						SCOPED_DRAW_EVENT(RHICmdList, ShadeIrradianceCache);
 
 						{	
 							TShaderMapRef<FSetupFinalGatherIndirectArgumentsCS> ComputeShader(View.ShaderMap);
@@ -3507,7 +3507,7 @@ bool FDeferredShadingSceneRenderer::RenderDistanceFieldAOSurfaceCache(FRHIComman
 
 				// Splat the surface cache records onto the opaque pixels, using less strict weighting so the lighting is smoothed in world space
 				{
-					SCOPED_DRAW_EVENT(RHICmdList, FinalIrradianceCacheSplat, DEC_SCENE_ITEMS);
+					SCOPED_DRAW_EVENT(RHICmdList, FinalIrradianceCacheSplat);
 					RenderIrradianceCacheInterpolation(RHICmdList, Views, IrradianceCacheAccumulation, DistanceFieldNormal->GetRenderTargetItem(), 0, GAODownsampleFactor, Parameters, true);
 				}
 
@@ -3692,7 +3692,7 @@ void FDeferredShadingSceneRenderer::RenderMeshDistanceFieldVisualization(FRHICom
 		&& Views.Num() == 1)
 	{
 		QUICK_SCOPE_CYCLE_COUNTER(STAT_RenderMeshDistanceFieldVis);
-		SCOPED_DRAW_EVENT(RHICmdList, VisualizeMeshDistanceFields, DEC_SCENE_ITEMS);
+		SCOPED_DRAW_EVENT(RHICmdList, VisualizeMeshDistanceFields);
 
 		// Update the global distance field atlas
 		GDistanceFieldVolumeTextureAtlas.UpdateAllocations();
@@ -3727,7 +3727,7 @@ void FDeferredShadingSceneRenderer::RenderMeshDistanceFieldVisualization(FRHICom
 						uint32 GroupSizeY = FMath::DivideAndRoundUp(View.ViewRect.Size().Y / GAODownsampleFactor, GDistanceFieldAOTileSizeY);
 
 						{
-							SCOPED_DRAW_EVENT(RHICmdList, VisualizeMeshDistanceFieldCS, DEC_SCENE_ITEMS);
+							SCOPED_DRAW_EVENT(RHICmdList, VisualizeMeshDistanceFieldCS);
 							TShaderMapRef<FVisualizeMeshDistanceFieldCS> ComputeShader(View.ShaderMap);
 
 							RHICmdList.SetComputeShader(ComputeShader->GetComputeShader());
@@ -3746,7 +3746,7 @@ void FDeferredShadingSceneRenderer::RenderMeshDistanceFieldVisualization(FRHICom
 					{
 						const FViewInfo& View = Views[ViewIndex];
 
-						SCOPED_DRAW_EVENT(RHICmdList, UpsampleAO, DEC_SCENE_ITEMS);
+						SCOPED_DRAW_EVENT(RHICmdList, UpsampleAO);
 
 						RHICmdList.SetViewport(View.ViewRect.Min.X, View.ViewRect.Min.Y, 0.0f, View.ViewRect.Max.X, View.ViewRect.Max.Y, 1.0f);
 						RHICmdList.SetRasterizerState(TStaticRasterizerState<FM_Solid, CM_None>::GetRHI());
@@ -3866,7 +3866,7 @@ void FDeferredShadingSceneRenderer::RenderDynamicSkyLighting(FRHICommandListImme
 		&& !IsSimpleDynamicLightingEnabled() 
 		&& !ViewFamily.EngineShowFlags.VisualizeLightCulling)
 	{
-		SCOPED_DRAW_EVENT(RHICmdList, SkyLightDiffuse, DEC_SCENE_ITEMS);
+		SCOPED_DRAW_EVENT(RHICmdList, SkyLightDiffuse);
 
 		bool bApplyShadowing = false;
 
diff --git a/Engine/Source/Runtime/Renderer/Private/DistortionRendering.cpp b/Engine/Source/Runtime/Renderer/Private/DistortionRendering.cpp
--- a/Engine/Source/Runtime/Renderer/Private/DistortionRendering.cpp
+++ b/Engine/Source/Runtime/Renderer/Private/DistortionRendering.cpp
@@ -780,7 +780,7 @@ bool FDistortionPrimSet::DrawAccumulatedOffsets(FRHICommandListImmediate& RHICmd
  */
 void FSceneRenderer::RenderDistortion(FRHICommandListImmediate& RHICmdList)
 {
-	SCOPED_DRAW_EVENT(RHICmdList, Distortion, DEC_SCENE_ITEMS);
+	SCOPED_DRAW_EVENT(RHICmdList, Distortion);
 
 	// do we need to render the distortion pass?
 	bool bRender=false;
@@ -801,7 +801,7 @@ void FSceneRenderer::RenderDistortion(FRHICommandListImmediate& RHICmdList)
 	// Render accumulated distortion offsets
 	if( bRender)
 	{
-		SCOPED_DRAW_EVENT(RHICmdList, DistortionAccum, DEC_SCENE_ITEMS);
+		SCOPED_DRAW_EVENT(RHICmdList, DistortionAccum);
 
 		// Create a texture to store the resolved light attenuation values, and a render-targetable surface to hold the unresolved light attenuation values.
 		{
@@ -823,7 +823,7 @@ void FSceneRenderer::RenderDistortion(FRHICommandListImmediate& RHICmdList)
 
 			for(int32 ViewIndex = 0;ViewIndex < Views.Num();ViewIndex++)
 			{
-				SCOPED_CONDITIONAL_DRAW_EVENTF(RHICmdList, EventView, Views.Num() > 1, DEC_SCENE_ITEMS, TEXT("View%d"), ViewIndex);
+				SCOPED_CONDITIONAL_DRAW_EVENTF(RHICmdList, EventView, Views.Num() > 1, TEXT("View%d"), ViewIndex);
 
 				FViewInfo& View = Views[ViewIndex];
 				// viewport to match view size
@@ -859,7 +859,7 @@ void FSceneRenderer::RenderDistortion(FRHICommandListImmediate& RHICmdList)
 
 	if(bDirty)
 	{
-		SCOPED_DRAW_EVENT(RHICmdList, DistortionApply, DEC_SCENE_ITEMS);
+		SCOPED_DRAW_EVENT(RHICmdList, DistortionApply);
 
 		GSceneRenderTargets.ResolveSceneColor(RHICmdList, FResolveRect(0, 0, ViewFamily.FamilySizeX, ViewFamily.FamilySizeY));
 
@@ -873,7 +873,7 @@ void FSceneRenderer::RenderDistortion(FRHICommandListImmediate& RHICmdList)
 		// Apply distortion as a full-screen pass		
 		for(int32 ViewIndex = 0, Num = Views.Num(); ViewIndex < Num; ++ViewIndex)
 		{
-			SCOPED_CONDITIONAL_DRAW_EVENTF(RHICmdList, EventView, Views.Num() > 1, DEC_SCENE_ITEMS, TEXT("View%d"), ViewIndex);
+			SCOPED_CONDITIONAL_DRAW_EVENTF(RHICmdList, EventView, Views.Num() > 1, TEXT("View%d"), ViewIndex);
 
 			FViewInfo& View = Views[ViewIndex];
 
diff --git a/Engine/Source/Runtime/Renderer/Private/DrawingPolicy.cpp b/Engine/Source/Runtime/Renderer/Private/DrawingPolicy.cpp
--- a/Engine/Source/Runtime/Renderer/Private/DrawingPolicy.cpp
+++ b/Engine/Source/Runtime/Renderer/Private/DrawingPolicy.cpp
@@ -65,7 +65,7 @@ void FMeshDrawingPolicy::SetMeshRenderState(
 void FMeshDrawingPolicy::DrawMesh(FRHICommandList& RHICmdList, const FMeshBatch& Mesh, int32 BatchElementIndex) const
 {
 	INC_DWORD_STAT(STAT_MeshDrawCalls);
-	SCOPED_CONDITIONAL_DRAW_EVENTF(RHICmdList, MeshEvent, GEmitMeshDrawEvent != 0, DEC_SCENE_ITEMS, TEXT("Mesh Draw"));
+	SCOPED_CONDITIONAL_DRAW_EVENTF(RHICmdList, MeshEvent, GEmitMeshDrawEvent != 0, TEXT("Mesh Draw"));
 
 	const FMeshBatchElement& BatchElement = Mesh.Elements[BatchElementIndex];
 
diff --git a/Engine/Source/Runtime/Renderer/Private/FogRendering.cpp b/Engine/Source/Runtime/Renderer/Private/FogRendering.cpp
--- a/Engine/Source/Runtime/Renderer/Private/FogRendering.cpp
+++ b/Engine/Source/Runtime/Renderer/Private/FogRendering.cpp
@@ -293,7 +293,7 @@ bool FDeferredShadingSceneRenderer::RenderFog(FRHICommandListImmediate& RHICmdLi
 {
 	if (Scene->ExponentialFogs.Num() > 0)
 	{
-		SCOPED_DRAW_EVENT(RHICmdList, Fog, DEC_SCENE_ITEMS);
+		SCOPED_DRAW_EVENT(RHICmdList, Fog);
 
 		static const FVector2D Vertices[4] =
 		{
diff --git a/Engine/Source/Runtime/Renderer/Private/ForwardBasePassRendering.cpp b/Engine/Source/Runtime/Renderer/Private/ForwardBasePassRendering.cpp
--- a/Engine/Source/Runtime/Renderer/Private/ForwardBasePassRendering.cpp
+++ b/Engine/Source/Runtime/Renderer/Private/ForwardBasePassRendering.cpp
@@ -312,7 +312,7 @@ EBasePassSort::Type GetSortMode()
 
 void FForwardShadingSceneRenderer::RenderForwardShadingBasePass(FRHICommandListImmediate& RHICmdList)
 {
-	SCOPED_DRAW_EVENT(RHICmdList, BasePass, DEC_SCENE_ITEMS);
+	SCOPED_DRAW_EVENT(RHICmdList, BasePass);
 	SCOPE_CYCLE_COUNTER(STAT_BasePassDrawTime);
 
 	EBasePassSort::Type SortMode = GetSortMode();
@@ -340,7 +340,7 @@ void FForwardShadingSceneRenderer::RenderForwardShadingBasePass(FRHICommandListI
 	// Draw the scene's emissive and light-map color.
 	for (int32 ViewIndex = 0; ViewIndex < Views.Num(); ViewIndex++)
 	{
-		SCOPED_CONDITIONAL_DRAW_EVENTF(RHICmdList, EventView, Views.Num() > 1, DEC_SCENE_ITEMS, TEXT("View%d"), ViewIndex);
+		SCOPED_CONDITIONAL_DRAW_EVENTF(RHICmdList, EventView, Views.Num() > 1, TEXT("View%d"), ViewIndex);
 		FViewInfo& View = Views[ViewIndex];
 
 		// Opaque blending
@@ -375,7 +375,7 @@ void FForwardShadingSceneRenderer::RenderForwardShadingBasePass(FRHICommandListI
 
 		{
 			SCOPE_CYCLE_COUNTER(STAT_DynamicPrimitiveDrawTime);
-			SCOPED_DRAW_EVENT(RHICmdList, Dynamic, DEC_SCENE_ITEMS);
+			SCOPED_DRAW_EVENT(RHICmdList, Dynamic);
 
 			const bool bUseGetMeshElements = ShouldUseGetDynamicMeshElements();
 
diff --git a/Engine/Source/Runtime/Renderer/Private/ForwardShadingRenderer.cpp b/Engine/Source/Runtime/Renderer/Private/ForwardShadingRenderer.cpp
--- a/Engine/Source/Runtime/Renderer/Private/ForwardShadingRenderer.cpp
+++ b/Engine/Source/Runtime/Renderer/Private/ForwardShadingRenderer.cpp
@@ -27,7 +27,7 @@ FForwardShadingSceneRenderer::FForwardShadingSceneRenderer(const FSceneViewFamil
  */
 void FForwardShadingSceneRenderer::InitViews(FRHICommandListImmediate& RHICmdList)
 {
-	SCOPED_DRAW_EVENT(RHICmdList, InitViews, DEC_SCENE_ITEMS);
+	SCOPED_DRAW_EVENT(RHICmdList, InitViews);
 
 	SCOPE_CYCLE_COUNTER(STAT_InitViewsTime);
 
@@ -200,11 +200,11 @@ void FForwardShadingSceneRenderer::Render(FRHICommandListImmediate& RHICmdList)
 			}
 			else
 			{
-				SCOPED_DRAW_EVENT(RHICmdList, PostProcessing, DEC_SCENE_ITEMS);
+				SCOPED_DRAW_EVENT(RHICmdList, PostProcessing);
 				SCOPE_CYCLE_COUNTER(STAT_FinishRenderViewTargetTime);
 				for(int32 ViewIndex = 0;ViewIndex < Views.Num();ViewIndex++)
 				{	
-					SCOPED_CONDITIONAL_DRAW_EVENTF(RHICmdList, EventView, Views.Num() > 1, DEC_SCENE_ITEMS, TEXT("View%d"), ViewIndex);
+					SCOPED_CONDITIONAL_DRAW_EVENTF(RHICmdList, EventView, Views.Num() > 1, TEXT("View%d"), ViewIndex);
 					GPostProcessing.ProcessES2(RHICmdList, Views[ViewIndex], bOnChipSunMask);
 				}
 			}
diff --git a/Engine/Source/Runtime/Renderer/Private/ForwardTranslucentRendering.cpp b/Engine/Source/Runtime/Renderer/Private/ForwardTranslucentRendering.cpp
--- a/Engine/Source/Runtime/Renderer/Private/ForwardTranslucentRendering.cpp
+++ b/Engine/Source/Runtime/Renderer/Private/ForwardTranslucentRendering.cpp
@@ -47,7 +47,7 @@ FGlobalBoundShaderState ForwardCopySceneAlphaBoundShaderState;
 
 void FForwardShadingSceneRenderer::CopySceneAlpha(FRHICommandListImmediate& RHICmdList, const FViewInfo& View)
 {
-	SCOPED_DRAW_EVENTF(RHICmdList, EventCopy, DEC_SCENE_ITEMS, TEXT("CopySceneAlpha"));
+	SCOPED_DRAW_EVENTF(RHICmdList, EventCopy, TEXT("CopySceneAlpha"));
 	RHICmdList.SetRasterizerState(TStaticRasterizerState<FM_Solid, CM_None>::GetRHI());
 	RHICmdList.SetDepthStencilState(TStaticDepthStencilState<false,CF_Always>::GetRHI());
 	RHICmdList.SetBlendState(TStaticBlendState<>::GetRHI());
@@ -332,11 +332,11 @@ void FForwardShadingSceneRenderer::RenderTranslucency(FRHICommandListImmediate&
 		const bool bGammaSpace = !IsMobileHDR();
 		const bool bLinearHDR64 = !bGammaSpace && !IsMobileHDR32bpp();
 
-		SCOPED_DRAW_EVENT(RHICmdList, Translucency, DEC_SCENE_ITEMS);
+		SCOPED_DRAW_EVENT(RHICmdList, Translucency);
 
 		for (int32 ViewIndex = 0; ViewIndex < Views.Num(); ViewIndex++)
 		{
-			SCOPED_CONDITIONAL_DRAW_EVENTF(RHICmdList, EventView, Views.Num() > 1, DEC_SCENE_ITEMS, TEXT("View%d"), ViewIndex);
+			SCOPED_CONDITIONAL_DRAW_EVENTF(RHICmdList, EventView, Views.Num() > 1, TEXT("View%d"), ViewIndex);
 
 			const FViewInfo& View = Views[ViewIndex];
 
diff --git a/Engine/Source/Runtime/Renderer/Private/GPUBenchmark.cpp b/Engine/Source/Runtime/Renderer/Private/GPUBenchmark.cpp
--- a/Engine/Source/Runtime/Renderer/Private/GPUBenchmark.cpp
+++ b/Engine/Source/Runtime/Renderer/Private/GPUBenchmark.cpp
@@ -168,7 +168,7 @@ void RunBenchmarkShader(FRHICommandList& RHICmdList, const FSceneView& View, TRe
 
 void RunBenchmarkShader(FRHICommandListImmediate& RHICmdList, const FSceneView& View, uint32 MethodId, TRefCountPtr<IPooledRenderTarget>& Src, float WorkScale)
 {
-	SCOPED_DRAW_EVENTF(RHICmdList, Benchmark, DEC_SCENE_ITEMS, TEXT("Benchmark Method:%d"), MethodId);
+	SCOPED_DRAW_EVENTF(RHICmdList, Benchmark, TEXT("Benchmark Method:%d"), MethodId);
 
 	switch(MethodId)
 	{
diff --git a/Engine/Source/Runtime/Renderer/Private/GammaCorrection.cpp b/Engine/Source/Runtime/Renderer/Private/GammaCorrection.cpp
--- a/Engine/Source/Runtime/Renderer/Private/GammaCorrection.cpp
+++ b/Engine/Source/Runtime/Renderer/Private/GammaCorrection.cpp
@@ -90,7 +90,7 @@ void FSceneRenderer::GammaCorrectToViewportRenderTarget(FRHICommandList& RHICmdL
 		ViewFamily.bDeferClear = false;
 	}
 
-	SCOPED_DRAW_EVENT(RHICmdList, GammaCorrection, DEC_SCENE_ITEMS);
+	SCOPED_DRAW_EVENT(RHICmdList, GammaCorrection);
 
 	// turn off culling and blending
 	RHICmdList.SetRasterizerState(TStaticRasterizerState<FM_Solid, CM_None>::GetRHI());
diff --git a/Engine/Source/Runtime/Renderer/Private/LightFunctionRendering.cpp b/Engine/Source/Runtime/Renderer/Private/LightFunctionRendering.cpp
--- a/Engine/Source/Runtime/Renderer/Private/LightFunctionRendering.cpp
+++ b/Engine/Source/Runtime/Renderer/Private/LightFunctionRendering.cpp
@@ -245,7 +245,7 @@ bool FDeferredShadingSceneRenderer::RenderLightFunctionForMaterial(FRHICommandLi
 		bRenderedLightFunction = true;
 
 		const FMaterial* Material = MaterialProxy->GetMaterial(Scene->GetFeatureLevel());
-		SCOPED_DRAW_EVENTF(RHICmdList, LightFunction, DEC_SCENE_ITEMS, TEXT("LightFunction Material=%s"), *Material->GetFriendlyName());
+		SCOPED_DRAW_EVENTF(RHICmdList, LightFunction, TEXT("LightFunction Material=%s"), *Material->GetFriendlyName());
 
 		const FMaterialShaderMap* MaterialShaderMap = Material->GetRenderingThreadShaderMap();
 		FLightFunctionVS* VertexShader = MaterialShaderMap->GetShader<FLightFunctionVS>();
@@ -258,7 +258,7 @@ bool FDeferredShadingSceneRenderer::RenderLightFunctionForMaterial(FRHICommandLi
 		// Render to the light attenuation buffer for all views.
 		for (int32 ViewIndex = 0; ViewIndex < Views.Num(); ViewIndex++)
 		{
-			SCOPED_CONDITIONAL_DRAW_EVENTF(RHICmdList, EventView, Views.Num() > 1, DEC_SCENE_ITEMS, TEXT("View%d"), ViewIndex);
+			SCOPED_CONDITIONAL_DRAW_EVENTF(RHICmdList, EventView, Views.Num() > 1, TEXT("View%d"), ViewIndex);
 
 			const FViewInfo& View = Views[ViewIndex];
 
diff --git a/Engine/Source/Runtime/Renderer/Private/LightMapDensityRendering.cpp b/Engine/Source/Runtime/Renderer/Private/LightMapDensityRendering.cpp
--- a/Engine/Source/Runtime/Renderer/Private/LightMapDensityRendering.cpp
+++ b/Engine/Source/Runtime/Renderer/Private/LightMapDensityRendering.cpp
@@ -40,12 +40,12 @@ bool FDeferredShadingSceneRenderer::RenderLightMapDensities(FRHICommandListImmed
 	bool bDirty=0;
 	if (Scene->GetFeatureLevel() >= ERHIFeatureLevel::SM4)
 	{
-		SCOPED_DRAW_EVENT(RHICmdList, LightMapDensity, DEC_SCENE_ITEMS);
+		SCOPED_DRAW_EVENT(RHICmdList, LightMapDensity);
 
 		// Draw the scene's emissive and light-map color.
 		for(int32 ViewIndex = 0;ViewIndex < Views.Num();ViewIndex++)
 		{
-			SCOPED_CONDITIONAL_DRAW_EVENTF(RHICmdList, EventView, Views.Num() > 1, DEC_SCENE_ITEMS, TEXT("View%d"), ViewIndex);
+			SCOPED_CONDITIONAL_DRAW_EVENTF(RHICmdList, EventView, Views.Num() > 1, TEXT("View%d"), ViewIndex);
 
 			FViewInfo& View = Views[ViewIndex];
 
@@ -56,7 +56,7 @@ bool FDeferredShadingSceneRenderer::RenderLightMapDensities(FRHICommandListImmed
 			RHICmdList.SetViewport(View.ViewRect.Min.X, View.ViewRect.Min.Y, 0, View.ViewRect.Max.X, View.ViewRect.Max.Y, 1);
 
 			{
-				SCOPED_DRAW_EVENT(RHICmdList, Dynamic, DEC_SCENE_ITEMS);
+				SCOPED_DRAW_EVENT(RHICmdList, Dynamic);
 
 				const bool bUseGetMeshElements = ShouldUseGetDynamicMeshElements();
 
diff --git a/Engine/Source/Runtime/Renderer/Private/LightPropagationVolume.cpp b/Engine/Source/Runtime/Renderer/Private/LightPropagationVolume.cpp
--- a/Engine/Source/Runtime/Renderer/Private/LightPropagationVolume.cpp
+++ b/Engine/Source/Runtime/Renderer/Private/LightPropagationVolume.cpp
@@ -929,7 +929,7 @@ void FLightPropagationVolume::Clear(FRHICommandListImmediate& RHICmdList, FViewI
 		return;
 	}
 
-	SCOPED_DRAW_EVENT(RHICmdList, LpvClear, DEC_LIGHT);
+	SCOPED_DRAW_EVENT(RHICmdList, LpvClear);
 
 	if ( !LpvWriteUniformBuffer.IsInitialized() )
 	{
@@ -1042,7 +1042,7 @@ void FLightPropagationVolume::InjectDirectionalLightRSM(
 {
 	const FLightSceneProxy* LightProxy = ProjectedShadowInfo.LightSceneInfo->Proxy;
 	{
-		SCOPED_DRAW_EVENT(RHICmdList, LpvInjectDirectionalLightRSM, DEC_LIGHT);
+		SCOPED_DRAW_EVENT(RHICmdList, LpvInjectDirectionalLightRSM);
 
 		SetVplInjectionConstants(ProjectedShadowInfo, LightProxy );
 
@@ -1062,7 +1062,7 @@ void FLightPropagationVolume::InjectDirectionalLightRSM(
 	// If this is the first directional light, build the geometry volume with it
 	if ( !GeometryVolumeGenerated && SecondaryOcclusionStrength > 0.001f )
 	{
-		SCOPED_DRAW_EVENT(RHICmdList, LpvBuildGeometryVolume, DEC_LIGHT);
+		SCOPED_DRAW_EVENT(RHICmdList, LpvBuildGeometryVolume);
 		GeometryVolumeGenerated = true;
 		FVector LightDirection( 0.0f, 0.0f, 1.0f );
 		if ( LightProxy )
@@ -1101,7 +1101,7 @@ void FLightPropagationVolume::Propagate(FRHICommandListImmediate& RHICmdList, FV
 	// Inject the VPLs into the LPV
 	if ( mInjectedLightCount )
 	{
-		SCOPED_DRAW_EVENT(RHICmdList, LpvAccumulateVplLists, DEC_LIGHT);
+		SCOPED_DRAW_EVENT(RHICmdList, LpvAccumulateVplLists);
 		mWriteBufferIndex = 1-mWriteBufferIndex; // Swap buffers with each iteration
 
 		TShaderMapRef<FLpvInject_AccumulateVplListsCS> Shader(View.ShaderMap);
@@ -1122,7 +1122,7 @@ void FLightPropagationVolume::Propagate(FRHICommandListImmediate& RHICmdList, FV
 	{
 		static int PropagationStep = 0;
 
-		SCOPED_DRAW_EVENT(RHICmdList, LpvPropagate, DEC_LIGHT);
+		SCOPED_DRAW_EVENT(RHICmdList, LpvPropagate);
 		static int numIterations = 3;
 		for ( int i=0; i<numIterations; i++ )
 		{
@@ -1209,7 +1209,7 @@ void FLightPropagationVolume::InjectLightDirect(FRHICommandListImmediate& RHICmd
 	if ( GeometryVolumeGenerated )
 	{
 		// Inject the VPLs into the LPV
-		SCOPED_DRAW_EVENT(RHICmdList, LpvDirectLightInjection, DEC_LIGHT);
+		SCOPED_DRAW_EVENT(RHICmdList, LpvDirectLightInjection);
 
 		FLpvDirectLightInjectParameters InjectUniformBufferParams;
 		InjectUniformBufferParams.LightColor = Light.GetColor() * Light.GetIndirectLightingScale();
diff --git a/Engine/Source/Runtime/Renderer/Private/LightPropagationVolumeVisualisation.cpp b/Engine/Source/Runtime/Renderer/Private/LightPropagationVolumeVisualisation.cpp
--- a/Engine/Source/Runtime/Renderer/Private/LightPropagationVolumeVisualisation.cpp
+++ b/Engine/Source/Runtime/Renderer/Private/LightPropagationVolumeVisualisation.cpp
@@ -250,7 +250,7 @@ IMPLEMENT_SHADER_TYPE(,FLpvVisualisePS,TEXT("LPVVisualise"),TEXT("PShader"),SF_P
 
 void FLightPropagationVolume::Visualise(FRHICommandList& RHICmdList, const FViewInfo& View) const
 {
-	SCOPED_DRAW_EVENT(RHICmdList, LpvVisualise, DEC_LIGHT);
+	SCOPED_DRAW_EVENT(RHICmdList, LpvVisualise);
 	check(View.GetFeatureLevel() == ERHIFeatureLevel::SM5);
 
 	TShaderMapRef<FLpvVisualiseVS> VertexShader(View.ShaderMap);
diff --git a/Engine/Source/Runtime/Renderer/Private/LightRendering.cpp b/Engine/Source/Runtime/Renderer/Private/LightRendering.cpp
--- a/Engine/Source/Runtime/Renderer/Private/LightRendering.cpp
+++ b/Engine/Source/Runtime/Renderer/Private/LightRendering.cpp
@@ -289,7 +289,7 @@ uint32 GetShadowQuality();
 /** Renders the scene's lighting. */
 void FDeferredShadingSceneRenderer::RenderLights(FRHICommandListImmediate& RHICmdList)
 {
-	SCOPED_DRAW_EVENT(RHICmdList, Lights, DEC_SCENE_ITEMS);
+	SCOPED_DRAW_EVENT(RHICmdList, Lights);
 
 	if(IsSimpleDynamicLightingEnabled())
 	{
@@ -348,7 +348,7 @@ void FDeferredShadingSceneRenderer::RenderLights(FRHICommandListImmediate& RHICm
 	SortedLights.Sort( FCompareFSortedLightSceneInfo() );
 
 	{
-		SCOPED_DRAW_EVENT(RHICmdList, DirectLighting, DEC_SCENE_ITEMS);
+		SCOPED_DRAW_EVENT(RHICmdList, DirectLighting);
 
 		int32 AttenuationLightStart = SortedLights.Num();
 		int32 SupportedByTiledDeferredLightEnd = SortedLights.Num();
@@ -390,7 +390,7 @@ void FDeferredShadingSceneRenderer::RenderLights(FRHICommandListImmediate& RHICm
 		
 		if(ViewFamily.EngineShowFlags.DirectLighting)
 		{
-			SCOPED_DRAW_EVENT(RHICmdList, NonShadowedLights, DEC_SCENE_ITEMS);
+			SCOPED_DRAW_EVENT(RHICmdList, NonShadowedLights);
 			INC_DWORD_STAT_BY(STAT_NumUnshadowedLights, AttenuationLightStart);
 			GSceneRenderTargets.SetLightAttenuationMode(false);
 
@@ -430,7 +430,7 @@ void FDeferredShadingSceneRenderer::RenderLights(FRHICommandListImmediate& RHICm
 			}
 
 			{
-				SCOPED_DRAW_EVENT(RHICmdList, StandardDeferredLighting, DEC_SCENE_ITEMS);
+				SCOPED_DRAW_EVENT(RHICmdList, StandardDeferredLighting);
 
 				GSceneRenderTargets.BeginRenderingSceneColor(RHICmdList);
 
@@ -451,13 +451,13 @@ void FDeferredShadingSceneRenderer::RenderLights(FRHICommandListImmediate& RHICm
 				if (AttenuationLightStart)
 				{
 					// Inject non-shadowed, non-light function lights in to the volume.
-					SCOPED_DRAW_EVENT(RHICmdList, InjectNonShadowedTranslucentLighting, DEC_SCENE_ITEMS);
+					SCOPED_DRAW_EVENT(RHICmdList, InjectNonShadowedTranslucentLighting);
 					InjectTranslucentVolumeLightingArray(RHICmdList, SortedLights, AttenuationLightStart);
 				}
 				
 				if (SimpleLights.InstanceData.Num() > 0)
 				{
-					SCOPED_DRAW_EVENT(RHICmdList, InjectSimpleLightsTranslucentLighting, DEC_SCENE_ITEMS);
+					SCOPED_DRAW_EVENT(RHICmdList, InjectSimpleLightsTranslucentLighting);
 					InjectSimpleTranslucentVolumeLightingArray(RHICmdList, SimpleLights);
 				}
 			}
@@ -479,7 +479,7 @@ void FDeferredShadingSceneRenderer::RenderLights(FRHICommandListImmediate& RHICm
 
 			FString LightNameWithLevel;
 			GetLightNameForDrawEvent(LightSceneInfo.Proxy, LightNameWithLevel);
-			SCOPED_DRAW_EVENTF(RHICmdList, EventLightPass, DEC_SCENE_ITEMS, *LightNameWithLevel);
+			SCOPED_DRAW_EVENTF(RHICmdList, EventLightPass, *LightNameWithLevel);
 
 			// Do not resolve to scene color texture, this is done lazily
 			GSceneRenderTargets.FinishRenderingSceneColor(RHICmdList, false);
@@ -536,7 +536,7 @@ void FDeferredShadingSceneRenderer::RenderLights(FRHICommandListImmediate& RHICm
 			
 			if(bDirectLighting && !bInjectedTranslucentVolume)
 			{
-				SCOPED_DRAW_EVENT(RHICmdList, InjectTranslucentVolume, DEC_SCENE_ITEMS);
+				SCOPED_DRAW_EVENT(RHICmdList, InjectTranslucentVolume);
 				// Accumulate this light's unshadowed contribution to the translucency lighting volume
 				InjectTranslucentVolumeLighting(RHICmdList, LightSceneInfo, NULL);
 			}
@@ -757,7 +757,7 @@ void FDeferredShadingSceneRenderer::RenderLight(FRHICommandList& RHICmdList, con
 {
 	SCOPE_CYCLE_COUNTER(STAT_DirectLightRenderingTime);
 	INC_DWORD_STAT(STAT_NumLightsUsingStandardDeferred);
-	SCOPED_CONDITIONAL_DRAW_EVENT(RHICmdList, StandardDeferredLighting, bIssueDrawEvent, DEC_SCENE_ITEMS);
+	SCOPED_CONDITIONAL_DRAW_EVENT(RHICmdList, StandardDeferredLighting, bIssueDrawEvent);
 
 	// Use additive blending for color
 	RHICmdList.SetBlendState(TStaticBlendState<CW_RGBA, BO_Add, BF_One, BF_One, BO_Add, BF_One, BF_One>::GetRHI());
@@ -901,7 +901,7 @@ void FDeferredShadingSceneRenderer::RenderSimpleLightsStandardDeferred(FRHIComma
 {
 	SCOPE_CYCLE_COUNTER(STAT_DirectLightRenderingTime);
 	INC_DWORD_STAT_BY(STAT_NumLightsUsingStandardDeferred, SimpleLights.InstanceData.Num());
-	SCOPED_DRAW_EVENT(RHICmdList, StandardDeferredSimpleLights, DEC_SCENE_ITEMS);
+	SCOPED_DRAW_EVENT(RHICmdList, StandardDeferredSimpleLights);
 	
 	// Use additive blending for color
 	RHICmdList.SetBlendState(TStaticBlendState<CW_RGBA, BO_Add, BF_One, BF_One, BO_Add, BF_One, BF_One>::GetRHI());
diff --git a/Engine/Source/Runtime/Renderer/Private/LightShaftRendering.cpp b/Engine/Source/Runtime/Renderer/Private/LightShaftRendering.cpp
--- a/Engine/Source/Runtime/Renderer/Private/LightShaftRendering.cpp
+++ b/Engine/Source/Runtime/Renderer/Private/LightShaftRendering.cpp
@@ -451,7 +451,7 @@ void AllocateOrReuseLightShaftRenderTarget(FRHICommandListImmediate& RHICmdList,
 template<bool bDownsampleOcclusion>
 void DownsamplePass(FRHICommandListImmediate& RHICmdList, const FViewInfo& View, const FLightSceneInfo* LightSceneInfo, TRefCountPtr<IPooledRenderTarget>& LightShaftsSource, TRefCountPtr<IPooledRenderTarget>& LightShaftsDest)
 {
-	SCOPED_DRAW_EVENT(RHICmdList, Downsample, DEC_SCENE_ITEMS);
+	SCOPED_DRAW_EVENT(RHICmdList, Downsample);
 
 	const FIntPoint BufferSize = GSceneRenderTargets.GetBufferSizeXY();
 	const uint32 DownsampleFactor	= GetLightShaftDownsampleFactor();
@@ -624,7 +624,7 @@ void ApplyRadialBlurPasses(
 		BlurLightShaftsPixelShader->SetParameters(RHICmdList, LightSceneInfo, View, PassIndex, EffectiveSource);
 
 		{
-			SCOPED_DRAW_EVENT(RHICmdList, RadialBlur, DEC_SCENE_ITEMS);
+			SCOPED_DRAW_EVENT(RHICmdList, RadialBlur);
 			// Apply a radial blur to the bloom and occlusion mask
 			DrawRectangle( 
 				RHICmdList,
@@ -668,7 +668,7 @@ void FinishOcclusionTerm(FRHICommandList& RHICmdList, const FViewInfo& View, con
 	MaskOcclusionTermPixelShader->SetParameters(RHICmdList, LightSceneInfo, View, LightShaftsSource);
 
 	{
-		SCOPED_DRAW_EVENT(RHICmdList, FinishOcclusion, DEC_SCENE_ITEMS);
+		SCOPED_DRAW_EVENT(RHICmdList, FinishOcclusion);
 		// Apply a radial blur to the bloom and occlusion mask
 		DrawRectangle( 
 			RHICmdList,
@@ -733,7 +733,7 @@ FLightShaftsOutput FDeferredShadingSceneRenderer::RenderLightShaftOcclusion(FRHI
 
 			if (bEnableOcclusion && LightSceneInfo->Proxy->GetLightType() == LightType_Directional)
 			{
-				SCOPED_DRAW_EVENT(RHICmdList, RenderLightShaftOcclusion, DEC_SCENE_ITEMS);
+				SCOPED_DRAW_EVENT(RHICmdList, RenderLightShaftOcclusion);
 
 				// Allocate light shaft render targets on demand, using the pool
 				// Need two targets to ping pong between
@@ -836,7 +836,7 @@ IMPLEMENT_SHADER_TYPE(,FApplyLightShaftsPixelShader,TEXT("LightShaftShader"),TEX
 
 void ApplyLightShaftBloom(FRHICommandListImmediate& RHICmdList, const FViewInfo& View, const FLightSceneInfo* const LightSceneInfo, TRefCountPtr<IPooledRenderTarget>& LightShaftsSource)
 {
-	SCOPED_DRAW_EVENT(RHICmdList, Apply, DEC_SCENE_ITEMS);
+	SCOPED_DRAW_EVENT(RHICmdList, Apply);
 
 	GSceneRenderTargets.BeginRenderingSceneColor(RHICmdList);
 
@@ -886,7 +886,7 @@ void FDeferredShadingSceneRenderer::RenderLightShaftBloom(FRHICommandListImmedia
 
 			if (LightSceneInfo->bEnableLightShaftBloom)
 			{
-				SCOPED_DRAW_EVENT(RHICmdList, RenderLightShaftBloom, DEC_SCENE_ITEMS);
+				SCOPED_DRAW_EVENT(RHICmdList, RenderLightShaftBloom);
 
 				// Allocate light shaft render targets on demand, using the pool
 				AllocateOrReuseLightShaftRenderTarget(RHICmdList, LightShafts0, TEXT("LightShafts0"));
diff --git a/Engine/Source/Runtime/Renderer/Private/PostProcess/PostProcessAA.cpp b/Engine/Source/Runtime/Renderer/Private/PostProcess/PostProcessAA.cpp
--- a/Engine/Source/Runtime/Renderer/Private/PostProcess/PostProcessAA.cpp
+++ b/Engine/Source/Runtime/Renderer/Private/PostProcess/PostProcessAA.cpp
@@ -230,7 +230,7 @@ static void SetShaderTemplAA(const FRenderingCompositePassContext& Context)
 
 void FRCPassPostProcessAA::Process(FRenderingCompositePassContext& Context)
 {
-	SCOPED_DRAW_EVENT(Context.RHICmdList, PostProcessAA, DEC_SCENE_ITEMS);
+	SCOPED_DRAW_EVENT(Context.RHICmdList, PostProcessAA);
 
 	const FPooledRenderTargetDesc* InputDesc = GetInputDesc(ePId_Input0);
 
diff --git a/Engine/Source/Runtime/Renderer/Private/PostProcess/PostProcessBloomSetup.cpp b/Engine/Source/Runtime/Renderer/Private/PostProcess/PostProcessBloomSetup.cpp
--- a/Engine/Source/Runtime/Renderer/Private/PostProcess/PostProcessBloomSetup.cpp
+++ b/Engine/Source/Runtime/Renderer/Private/PostProcess/PostProcessBloomSetup.cpp
@@ -139,7 +139,7 @@ IMPLEMENT_SHADER_TYPE(,FPostProcessBloomSetupVS,TEXT("PostProcessBloom"),TEXT("M
 
 void FRCPassPostProcessBloomSetup::Process(FRenderingCompositePassContext& Context)
 {
-	SCOPED_DRAW_EVENT(Context.RHICmdList, PostProcessBloomSetup, DEC_SCENE_ITEMS);
+	SCOPED_DRAW_EVENT(Context.RHICmdList, PostProcessBloomSetup);
 
 	const FPooledRenderTargetDesc* InputDesc = GetInputDesc(ePId_Input0);
 
diff --git a/Engine/Source/Runtime/Renderer/Private/PostProcess/PostProcessBokehDOF.cpp b/Engine/Source/Runtime/Renderer/Private/PostProcess/PostProcessBokehDOF.cpp
--- a/Engine/Source/Runtime/Renderer/Private/PostProcess/PostProcessBokehDOF.cpp
+++ b/Engine/Source/Runtime/Renderer/Private/PostProcess/PostProcessBokehDOF.cpp
@@ -163,7 +163,7 @@ IMPLEMENT_SHADER_TYPE3(PostProcessVisualizeDOFPS, SF_Pixel);
 
 void FRCPassPostProcessVisualizeDOF::Process(FRenderingCompositePassContext& Context)
 {
-	SCOPED_DRAW_EVENT(Context.RHICmdList, VisualizeDOF, DEC_SCENE_ITEMS);
+	SCOPED_DRAW_EVENT(Context.RHICmdList, VisualizeDOF);
 
 	const FPooledRenderTargetDesc* InputDesc = GetInputDesc(ePId_Input0);
 
@@ -385,7 +385,7 @@ IMPLEMENT_SHADER_TYPE3(PostProcessBokehDOFSetupPS, SF_Pixel);
 
 void FRCPassPostProcessBokehDOFSetup::Process(FRenderingCompositePassContext& Context)
 {
-	SCOPED_DRAW_EVENT(Context.RHICmdList, BokehDOFSetup, DEC_SCENE_ITEMS);
+	SCOPED_DRAW_EVENT(Context.RHICmdList, BokehDOFSetup);
 
 	const FPooledRenderTargetDesc* InputDesc = GetInputDesc(ePId_Input0);
 
@@ -694,7 +694,7 @@ static TAutoConsoleVariable<int32> CVarBokehDOFIndexStyle(
 
 void FRCPassPostProcessBokehDOF::Process(FRenderingCompositePassContext& Context)
 {
-	SCOPED_DRAW_EVENT(Context.RHICmdList, PassPostProcessBokehDOF, DEC_SCENE_ITEMS);
+	SCOPED_DRAW_EVENT(Context.RHICmdList, PassPostProcessBokehDOF);
 
 	const FPooledRenderTargetDesc* InputDesc = GetInputDesc(ePId_Input0);
 	
diff --git a/Engine/Source/Runtime/Renderer/Private/PostProcess/PostProcessBokehDOFRecombine.cpp b/Engine/Source/Runtime/Renderer/Private/PostProcess/PostProcessBokehDOFRecombine.cpp
--- a/Engine/Source/Runtime/Renderer/Private/PostProcess/PostProcessBokehDOFRecombine.cpp
+++ b/Engine/Source/Runtime/Renderer/Private/PostProcess/PostProcessBokehDOFRecombine.cpp
@@ -130,7 +130,7 @@ void FRCPassPostProcessBokehDOFRecombine::Process(FRenderingCompositePassContext
 		check(GetInput(ePId_Input2)->GetPass());
 	}
 
-	SCOPED_DRAW_EVENTF(Context.RHICmdList, BokehDOFRecombine, DEC_SCENE_ITEMS, TEXT("BokehDOFRecombine#%d"), Method);
+	SCOPED_DRAW_EVENTF(Context.RHICmdList, BokehDOFRecombine, TEXT("BokehDOFRecombine#%d"), Method);
 
 	const FPooledRenderTargetDesc* InputDesc0 = GetInputDesc(ePId_Input0);
 	const FPooledRenderTargetDesc* InputDesc1 = GetInputDesc(ePId_Input1);
diff --git a/Engine/Source/Runtime/Renderer/Private/PostProcess/PostProcessBusyWait.cpp b/Engine/Source/Runtime/Renderer/Private/PostProcess/PostProcessBusyWait.cpp
--- a/Engine/Source/Runtime/Renderer/Private/PostProcess/PostProcessBusyWait.cpp
+++ b/Engine/Source/Runtime/Renderer/Private/PostProcess/PostProcessBusyWait.cpp
@@ -83,7 +83,7 @@ IMPLEMENT_SHADER_TYPE(,FPostProcessBusyWaitPS,TEXT("PostProcessBusyWait"),TEXT("
 
 void FRCPassPostProcessBusyWait::Process(FRenderingCompositePassContext& Context)
 {
-	SCOPED_DRAW_EVENT(Context.RHICmdList, BusyWait, DEC_SCENE_ITEMS);
+	SCOPED_DRAW_EVENT(Context.RHICmdList, BusyWait);
 
 	const FSceneView& View = Context.View;
 	
diff --git a/Engine/Source/Runtime/Renderer/Private/PostProcess/PostProcessCombineLUTs.cpp b/Engine/Source/Runtime/Renderer/Private/PostProcess/PostProcessCombineLUTs.cpp
--- a/Engine/Source/Runtime/Renderer/Private/PostProcess/PostProcessCombineLUTs.cpp
+++ b/Engine/Source/Runtime/Renderer/Private/PostProcess/PostProcessCombineLUTs.cpp
@@ -380,7 +380,7 @@ uint32 FRCPassPostProcessCombineLUTs::GenerateFinalTable(const FFinalPostProcess
 
 void FRCPassPostProcessCombineLUTs::Process(FRenderingCompositePassContext& Context)
 {
-	SCOPED_DRAW_EVENT(Context.RHICmdList, PostProcessCombineLUTs, DEC_SCENE_ITEMS);
+	SCOPED_DRAW_EVENT(Context.RHICmdList, PostProcessCombineLUTs);
 
 	FTexture* LocalTextures[GMaxLUTBlendCount];
 	float LocalWeights[GMaxLUTBlendCount];
diff --git a/Engine/Source/Runtime/Renderer/Private/PostProcess/PostProcessCompositeEditorPrimitives.cpp b/Engine/Source/Runtime/Renderer/Private/PostProcess/PostProcessCompositeEditorPrimitives.cpp
--- a/Engine/Source/Runtime/Renderer/Private/PostProcess/PostProcessCompositeEditorPrimitives.cpp
+++ b/Engine/Source/Runtime/Renderer/Private/PostProcess/PostProcessCompositeEditorPrimitives.cpp
@@ -182,7 +182,7 @@ static void SetCompositePrimitivesShaderTempl(const FRenderingCompositePassConte
 
 void FRCPassPostProcessCompositeEditorPrimitives::Process(FRenderingCompositePassContext& Context)
 {
-	SCOPED_DRAW_EVENT(Context.RHICmdList, CompositeEditorPrimitives, DEC_SCENE_ITEMS);
+	SCOPED_DRAW_EVENT(Context.RHICmdList, CompositeEditorPrimitives);
 	const FPooledRenderTargetDesc* InputDesc = GetInputDesc(ePId_Input0);
 
 	if(!InputDesc)
@@ -213,13 +213,13 @@ void FRCPassPostProcessCompositeEditorPrimitives::Process(FRenderingCompositePas
 
 		if(bClearIsNeeded)
 		{
-			SCOPED_DRAW_EVENT(Context.RHICmdList, ClearViewEditorPrimitives, DEC_SCENE_ITEMS);
+			SCOPED_DRAW_EVENT(Context.RHICmdList, ClearViewEditorPrimitives);
 			// Clear color and depth
 			// Note, this is a reversed Z depth surface, so 0.0f is the far plane.
 			Context.RHICmdList.Clear(true, FLinearColor(0, 0, 0, 0), true, 0.0f, false, 0, FIntRect());
 		}
 
-		SCOPED_DRAW_EVENT(Context.RHICmdList, RenderEditorPrimitives, DEC_SCENE_ITEMS);
+		SCOPED_DRAW_EVENT(Context.RHICmdList, RenderEditorPrimitives);
 
 		Context.RHICmdList.SetRasterizerState(View.bReverseCulling ? TStaticRasterizerState<FM_Solid, CM_CW>::GetRHI() : TStaticRasterizerState<FM_Solid, CM_CCW>::GetRHI());
 
diff --git a/Engine/Source/Runtime/Renderer/Private/PostProcess/PostProcessDOF.cpp b/Engine/Source/Runtime/Renderer/Private/PostProcess/PostProcessDOF.cpp
--- a/Engine/Source/Runtime/Renderer/Private/PostProcess/PostProcessDOF.cpp
+++ b/Engine/Source/Runtime/Renderer/Private/PostProcess/PostProcessDOF.cpp
@@ -80,7 +80,7 @@ IMPLEMENT_SHADER_TYPE(template<>,FPostProcessDOFSetupPS<1>,TEXT("PostProcessDOF"
 
 void FRCPassPostProcessDOFSetup::Process(FRenderingCompositePassContext& Context)
 {
-	SCOPED_DRAW_EVENT(Context.RHICmdList, DOFSetup, DEC_SCENE_ITEMS);
+	SCOPED_DRAW_EVENT(Context.RHICmdList, DOFSetup);
 
 	const FPooledRenderTargetDesc* InputDesc = GetInputDesc(ePId_Input0);
 
@@ -264,7 +264,7 @@ IMPLEMENT_SHADER_TYPE(template<>,FPostProcessDOFRecombinePS<1>,TEXT("PostProcess
 
 void FRCPassPostProcessDOFRecombine::Process(FRenderingCompositePassContext& Context)
 {
-	SCOPED_DRAW_EVENT(Context.RHICmdList, DOFRecombine, DEC_SCENE_ITEMS);
+	SCOPED_DRAW_EVENT(Context.RHICmdList, DOFRecombine);
 
 	const FPooledRenderTargetDesc* InputDesc = GetInputDesc(ePId_Input1);
 
diff --git a/Engine/Source/Runtime/Renderer/Private/PostProcess/PostProcessDownsample.cpp b/Engine/Source/Runtime/Renderer/Private/PostProcess/PostProcessDownsample.cpp
--- a/Engine/Source/Runtime/Renderer/Private/PostProcess/PostProcessDownsample.cpp
+++ b/Engine/Source/Runtime/Renderer/Private/PostProcess/PostProcessDownsample.cpp
@@ -154,7 +154,7 @@ void FRCPassPostProcessDownsample::SetShader(const FRenderingCompositePassContex
 
 void FRCPassPostProcessDownsample::Process(FRenderingCompositePassContext& Context)
 {
-	SCOPED_DRAW_EVENT(Context.RHICmdList, Downsample, DEC_SCENE_ITEMS);
+	SCOPED_DRAW_EVENT(Context.RHICmdList, Downsample);
 
 	const FPooledRenderTargetDesc* InputDesc = GetInputDesc(ePId_Input0);
 
diff --git a/Engine/Source/Runtime/Renderer/Private/PostProcess/PostProcessEyeAdaptation.cpp b/Engine/Source/Runtime/Renderer/Private/PostProcess/PostProcessEyeAdaptation.cpp
--- a/Engine/Source/Runtime/Renderer/Private/PostProcess/PostProcessEyeAdaptation.cpp
+++ b/Engine/Source/Runtime/Renderer/Private/PostProcess/PostProcessEyeAdaptation.cpp
@@ -71,7 +71,7 @@ IMPLEMENT_SHADER_TYPE(,FPostProcessEyeAdaptationPS,TEXT("PostProcessEyeAdaptatio
 
 void FRCPassPostProcessEyeAdaptation::Process(FRenderingCompositePassContext& Context)
 {
-	SCOPED_DRAW_EVENT(Context.RHICmdList, PostProcessEyeAdaptation, DEC_SCENE_ITEMS);
+	SCOPED_DRAW_EVENT(Context.RHICmdList, PostProcessEyeAdaptation);
 
 	const FSceneView& View = Context.View;
 	const FSceneViewFamily& ViewFamily = *(View.Family);
diff --git a/Engine/Source/Runtime/Renderer/Private/PostProcess/PostProcessGBufferHints.cpp b/Engine/Source/Runtime/Renderer/Private/PostProcess/PostProcessGBufferHints.cpp
--- a/Engine/Source/Runtime/Renderer/Private/PostProcess/PostProcessGBufferHints.cpp
+++ b/Engine/Source/Runtime/Renderer/Private/PostProcess/PostProcessGBufferHints.cpp
@@ -87,7 +87,7 @@ FRCPassPostProcessGBufferHints::FRCPassPostProcessGBufferHints()
 
 void FRCPassPostProcessGBufferHints::Process(FRenderingCompositePassContext& Context)
 {
-	SCOPED_DRAW_EVENT(Context.RHICmdList, GBufferHints, DEC_SCENE_ITEMS);
+	SCOPED_DRAW_EVENT(Context.RHICmdList, GBufferHints);
 	const FPooledRenderTargetDesc* InputDesc = GetInputDesc(ePId_Input0);
 
 	if(!InputDesc)
diff --git a/Engine/Source/Runtime/Renderer/Private/PostProcess/PostProcessHMD.cpp b/Engine/Source/Runtime/Renderer/Private/PostProcess/PostProcessHMD.cpp
--- a/Engine/Source/Runtime/Renderer/Private/PostProcess/PostProcessHMD.cpp
+++ b/Engine/Source/Runtime/Renderer/Private/PostProcess/PostProcessHMD.cpp
@@ -184,7 +184,7 @@ IMPLEMENT_SHADER_TYPE(template<>, FPostProcessHMDPS<false>, TEXT("PostProcessHMD
 
 void FRCPassPostProcessHMD::Process(FRenderingCompositePassContext& Context)
 {
-	SCOPED_DRAW_EVENT(Context.RHICmdList, PostProcessHMD, DEC_SCENE_ITEMS);
+	SCOPED_DRAW_EVENT(Context.RHICmdList, PostProcessHMD);
 	const FPooledRenderTargetDesc* InputDesc = GetInputDesc(ePId_Input0);
 
 	if(!InputDesc)
@@ -200,7 +200,7 @@ void FRCPassPostProcessHMD::Process(FRenderingCompositePassContext& Context)
 	const FIntRect DestRect = View.UnscaledViewRect;
 	const FIntPoint SrcSize = InputDesc->Extent;
 
-    const FSceneRenderTargetItem& DestRenderTarget = PassOutputs[0].RequestSurface(Context);
+	const FSceneRenderTargetItem& DestRenderTarget = PassOutputs[0].RequestSurface(Context);
 
 	// Set the view family's render target/viewport.
 	SetRenderTarget(Context.RHICmdList, DestRenderTarget.TargetableTexture, FTextureRHIRef());
diff --git a/Engine/Source/Runtime/Renderer/Private/PostProcess/PostProcessHistogram.cpp b/Engine/Source/Runtime/Renderer/Private/PostProcess/PostProcessHistogram.cpp
--- a/Engine/Source/Runtime/Renderer/Private/PostProcess/PostProcessHistogram.cpp
+++ b/Engine/Source/Runtime/Renderer/Private/PostProcess/PostProcessHistogram.cpp
@@ -91,7 +91,7 @@ IMPLEMENT_SHADER_TYPE(,FPostProcessHistogramCS,TEXT("PostProcessHistogram"),TEXT
 
 void FRCPassPostProcessHistogram::Process(FRenderingCompositePassContext& Context)
 {
-	SCOPED_DRAW_EVENT(Context.RHICmdList, PostProcessHistogram, DEC_SCENE_ITEMS);
+	SCOPED_DRAW_EVENT(Context.RHICmdList, PostProcessHistogram);
 	const FPooledRenderTargetDesc* InputDesc = GetInputDesc(ePId_Input0);
 
 	if(!InputDesc)
diff --git a/Engine/Source/Runtime/Renderer/Private/PostProcess/PostProcessHistogramReduce.cpp b/Engine/Source/Runtime/Renderer/Private/PostProcess/PostProcessHistogramReduce.cpp
--- a/Engine/Source/Runtime/Renderer/Private/PostProcess/PostProcessHistogramReduce.cpp
+++ b/Engine/Source/Runtime/Renderer/Private/PostProcess/PostProcessHistogramReduce.cpp
@@ -89,7 +89,7 @@ IMPLEMENT_SHADER_TYPE(,FPostProcessHistogramReducePS,TEXT("PostProcessHistogramR
 
 void FRCPassPostProcessHistogramReduce::Process(FRenderingCompositePassContext& Context)
 {
-	SCOPED_DRAW_EVENT(Context.RHICmdList, PostProcessHistogramReduce, DEC_SCENE_ITEMS);
+	SCOPED_DRAW_EVENT(Context.RHICmdList, PostProcessHistogramReduce);
 	const FPooledRenderTargetDesc* InputDesc = GetInputDesc(ePId_Input0);
 
 	if(!InputDesc)
diff --git a/Engine/Source/Runtime/Renderer/Private/PostProcess/PostProcessLensBlur.cpp b/Engine/Source/Runtime/Renderer/Private/PostProcess/PostProcessLensBlur.cpp
--- a/Engine/Source/Runtime/Renderer/Private/PostProcess/PostProcessLensBlur.cpp
+++ b/Engine/Source/Runtime/Renderer/Private/PostProcess/PostProcessLensBlur.cpp
@@ -162,7 +162,7 @@ FRCPassPostProcessLensBlur::FRCPassPostProcessLensBlur(float InPercentKernelSize
 
 void FRCPassPostProcessLensBlur::Process(FRenderingCompositePassContext& Context)
 {
-	SCOPED_DRAW_EVENT(Context.RHICmdList, PassPostProcessLensBlur, DEC_SCENE_ITEMS);
+	SCOPED_DRAW_EVENT(Context.RHICmdList, PassPostProcessLensBlur);
 
 	const FPooledRenderTargetDesc* InputDesc = GetInputDesc(ePId_Input0);
 	
diff --git a/Engine/Source/Runtime/Renderer/Private/PostProcess/PostProcessLensFlares.cpp b/Engine/Source/Runtime/Renderer/Private/PostProcess/PostProcessLensFlares.cpp
--- a/Engine/Source/Runtime/Renderer/Private/PostProcess/PostProcessLensFlares.cpp
+++ b/Engine/Source/Runtime/Renderer/Private/PostProcess/PostProcessLensFlares.cpp
@@ -113,7 +113,7 @@ FRCPassPostProcessLensFlares::FRCPassPostProcessLensFlares(float InSizeScale)
 
 void FRCPassPostProcessLensFlares::Process(FRenderingCompositePassContext& Context)
 {
-	SCOPED_DRAW_EVENT(Context.RHICmdList, LensFlares, DEC_SCENE_ITEMS);
+	SCOPED_DRAW_EVENT(Context.RHICmdList, LensFlares);
 
 	const FPooledRenderTargetDesc* InputDesc1 = GetInputDesc(ePId_Input0);
 	const FPooledRenderTargetDesc* InputDesc2 = GetInputDesc(ePId_Input1);
diff --git a/Engine/Source/Runtime/Renderer/Private/PostProcess/PostProcessMaterial.cpp b/Engine/Source/Runtime/Renderer/Private/PostProcess/PostProcessMaterial.cpp
--- a/Engine/Source/Runtime/Renderer/Private/PostProcess/PostProcessMaterial.cpp
+++ b/Engine/Source/Runtime/Renderer/Private/PostProcess/PostProcessMaterial.cpp
@@ -107,7 +107,7 @@ void FRCPassPostProcessMaterial::Process(FRenderingCompositePassContext& Context
 	
 	check(Material);
 
-	SCOPED_DRAW_EVENTF(Context.RHICmdList, PostProcessMaterial, DEC_SCENE_ITEMS, TEXT("PostProcessMaterial Material=%s"), *Material->GetFriendlyName());
+	SCOPED_DRAW_EVENTF(Context.RHICmdList, PostProcessMaterial, TEXT("PostProcessMaterial Material=%s"), *Material->GetFriendlyName());
 
 	const FPooledRenderTargetDesc* InputDesc = GetInputDesc(ePId_Input0);
 
diff --git a/Engine/Source/Runtime/Renderer/Private/PostProcess/PostProcessMobile.cpp b/Engine/Source/Runtime/Renderer/Private/PostProcess/PostProcessMobile.cpp
--- a/Engine/Source/Runtime/Renderer/Private/PostProcess/PostProcessMobile.cpp
+++ b/Engine/Source/Runtime/Renderer/Private/PostProcess/PostProcessMobile.cpp
@@ -168,7 +168,7 @@ void FRCPassPostProcessBloomSetupES2::SetShader(const FRenderingCompositePassCon
 
 void FRCPassPostProcessBloomSetupES2::Process(FRenderingCompositePassContext& Context)
 {
-	SCOPED_DRAW_EVENT(Context.RHICmdList, PostProcessBloomSetup, DEC_SCENE_ITEMS);
+	SCOPED_DRAW_EVENT(Context.RHICmdList, PostProcessBloomSetup);
 
 	const FPooledRenderTargetDesc* InputDesc = GetInputDesc(ePId_Input0);
 
@@ -369,7 +369,7 @@ void FRCPassPostProcessBloomSetupSmallES2::SetShader(const FRenderingCompositePa
 
 void FRCPassPostProcessBloomSetupSmallES2::Process(FRenderingCompositePassContext& Context)
 {
-	SCOPED_DRAW_EVENT(Context.RHICmdList, PostProcessBloomSetupSmall, DEC_SCENE_ITEMS);
+	SCOPED_DRAW_EVENT(Context.RHICmdList, PostProcessBloomSetupSmall);
 
 	const FPooledRenderTargetDesc* InputDesc = GetInputDesc(ePId_Input0);
 
@@ -545,7 +545,7 @@ IMPLEMENT_SHADER_TYPE(,FPostProcessBloomDownVS_ES2,TEXT("PostProcessMobile"),TEX
 
 void FRCPassPostProcessBloomDownES2::Process(FRenderingCompositePassContext& Context)
 {
-	SCOPED_DRAW_EVENT(Context.RHICmdList, PostProcessBloomDown, DEC_SCENE_ITEMS);
+	SCOPED_DRAW_EVENT(Context.RHICmdList, PostProcessBloomDown);
 
 	uint32 DstX = FMath::Max(1, PrePostSourceViewportSize.X/2);
 	uint32 DstY = FMath::Max(1, PrePostSourceViewportSize.Y/2);
@@ -713,7 +713,7 @@ IMPLEMENT_SHADER_TYPE(,FPostProcessBloomUpVS_ES2,TEXT("PostProcessMobile"),TEXT(
 
 void FRCPassPostProcessBloomUpES2::Process(FRenderingCompositePassContext& Context)
 {
-	SCOPED_DRAW_EVENT(Context.RHICmdList, PostProcessBloomUp, DEC_SCENE_ITEMS);
+	SCOPED_DRAW_EVENT(Context.RHICmdList, PostProcessBloomUp);
 
 	uint32 DstX = FMath::Max(1, PrePostSourceViewportSize.X);
 	uint32 DstY = FMath::Max(1, PrePostSourceViewportSize.Y);
@@ -947,7 +947,7 @@ void FRCPassPostProcessSunMaskES2::SetShader(const FRenderingCompositePassContex
 
 void FRCPassPostProcessSunMaskES2::Process(FRenderingCompositePassContext& Context)
 {
-	SCOPED_DRAW_EVENT(Context.RHICmdList, PostProcessSunMask, DEC_SCENE_ITEMS);
+	SCOPED_DRAW_EVENT(Context.RHICmdList, PostProcessSunMask);
 
 	const FPooledRenderTargetDesc* InputDesc = GetInputDesc(ePId_Input0);
 
@@ -1169,7 +1169,7 @@ void FRCPassPostProcessSunAlphaES2::SetShader(const FRenderingCompositePassConte
 
 void FRCPassPostProcessSunAlphaES2::Process(FRenderingCompositePassContext& Context)
 {
-	SCOPED_DRAW_EVENT(Context.RHICmdList, PostProcessSunAlpha, DEC_SCENE_ITEMS);
+	SCOPED_DRAW_EVENT(Context.RHICmdList, PostProcessSunAlpha);
 
 	uint32 DstX = FMath::Max(1, PrePostSourceViewportSize.X/4);
 	uint32 DstY = FMath::Max(1, PrePostSourceViewportSize.Y/4);
@@ -1326,7 +1326,7 @@ IMPLEMENT_SHADER_TYPE(,FPostProcessSunBlurVS_ES2,TEXT("PostProcessMobile"),TEXT(
 
 void FRCPassPostProcessSunBlurES2::Process(FRenderingCompositePassContext& Context)
 {
-	SCOPED_DRAW_EVENT(Context.RHICmdList, PostProcessSunBlur, DEC_SCENE_ITEMS);
+	SCOPED_DRAW_EVENT(Context.RHICmdList, PostProcessSunBlur);
 
 	uint32 DstX = FMath::Max(1, PrePostSourceViewportSize.X/4);
 	uint32 DstY = FMath::Max(1, PrePostSourceViewportSize.Y/4);
@@ -1555,7 +1555,7 @@ FShader* FRCPassPostProcessSunMergeES2::SetShader(const FRenderingCompositePassC
 
 void FRCPassPostProcessSunMergeES2::Process(FRenderingCompositePassContext& Context)
 {
-	SCOPED_DRAW_EVENT(Context.RHICmdList, PostProcessSunMerge, DEC_SCENE_ITEMS);
+	SCOPED_DRAW_EVENT(Context.RHICmdList, PostProcessSunMerge);
 
 	uint32 DstX = FMath::Max(1, PrePostSourceViewportSize.X/4);
 	uint32 DstY = FMath::Max(1, PrePostSourceViewportSize.Y/4);
@@ -1749,7 +1749,7 @@ void FRCPassPostProcessSunMergeSmallES2::SetShader(const FRenderingCompositePass
 
 void FRCPassPostProcessSunMergeSmallES2::Process(FRenderingCompositePassContext& Context)
 {
-	SCOPED_DRAW_EVENT(Context.RHICmdList, PostProcessSunMergeSmall, DEC_SCENE_ITEMS);
+	SCOPED_DRAW_EVENT(Context.RHICmdList, PostProcessSunMergeSmall);
 
 	uint32 DstX = FMath::Max(1, PrePostSourceViewportSize.X/4);
 	uint32 DstY = FMath::Max(1, PrePostSourceViewportSize.Y/4);
@@ -1952,7 +1952,7 @@ void FRCPassPostProcessDofDownES2::SetShader(const FRenderingCompositePassContex
 
 void FRCPassPostProcessDofDownES2::Process(FRenderingCompositePassContext& Context)
 {
-	SCOPED_DRAW_EVENT(Context.RHICmdList, PostProcessDofDown, DEC_SCENE_ITEMS);
+	SCOPED_DRAW_EVENT(Context.RHICmdList, PostProcessDofDown);
 
 	const FPooledRenderTargetDesc* InputDesc = GetInputDesc(ePId_Input0);
 
@@ -2160,7 +2160,7 @@ void FRCPassPostProcessDofNearES2::SetShader(const FRenderingCompositePassContex
 
 void FRCPassPostProcessDofNearES2::Process(FRenderingCompositePassContext& Context)
 {
-	SCOPED_DRAW_EVENT(Context.RHICmdList, PostProcessDofNear, DEC_SCENE_ITEMS);
+	SCOPED_DRAW_EVENT(Context.RHICmdList, PostProcessDofNear);
 
 	const FPooledRenderTargetDesc* InputDesc = GetInputDesc(ePId_Input0);
 	FIntPoint SrcSize = InputDesc->Extent;
@@ -2313,7 +2313,7 @@ IMPLEMENT_SHADER_TYPE(,FPostProcessDofBlurVS_ES2,TEXT("PostProcessMobile"),TEXT(
 
 void FRCPassPostProcessDofBlurES2::Process(FRenderingCompositePassContext& Context)
 {
-	SCOPED_DRAW_EVENT(Context.RHICmdList, PostProcessDofBlur, DEC_SCENE_ITEMS);
+	SCOPED_DRAW_EVENT(Context.RHICmdList, PostProcessDofBlur);
 
 	uint32 DstX = FMath::Max(1, PrePostSourceViewportSize.X/2);
 	uint32 DstY = FMath::Max(1, PrePostSourceViewportSize.Y/2);
@@ -2493,7 +2493,7 @@ void FRCPassPostProcessSunAvgES2::SetShader(const FRenderingCompositePassContext
 
 void FRCPassPostProcessSunAvgES2::Process(FRenderingCompositePassContext& Context)
 {
-	SCOPED_DRAW_EVENT(Context.RHICmdList, PostProcessSunAvg, DEC_SCENE_ITEMS);
+	SCOPED_DRAW_EVENT(Context.RHICmdList, PostProcessSunAvg);
 
 	uint32 DstX = FMath::Max(1, PrePostSourceViewportSize.X/4);
 	uint32 DstY = FMath::Max(1, PrePostSourceViewportSize.Y/4);
@@ -2736,7 +2736,7 @@ void FRCPassPostProcessAaES2::SetShader(const FRenderingCompositePassContext& Co
 
 void FRCPassPostProcessAaES2::Process(FRenderingCompositePassContext& Context)
 {
-	SCOPED_DRAW_EVENT(Context.RHICmdList, PostProcessAa, DEC_SCENE_ITEMS);
+	SCOPED_DRAW_EVENT(Context.RHICmdList, PostProcessAa);
 
 	const FPooledRenderTargetDesc* InputDesc = GetInputDesc(ePId_Input0);
 	FIntPoint SrcSize = InputDesc->Extent;
diff --git a/Engine/Source/Runtime/Renderer/Private/PostProcess/PostProcessMorpheus.cpp b/Engine/Source/Runtime/Renderer/Private/PostProcess/PostProcessMorpheus.cpp
--- a/Engine/Source/Runtime/Renderer/Private/PostProcess/PostProcessMorpheus.cpp
+++ b/Engine/Source/Runtime/Renderer/Private/PostProcess/PostProcessMorpheus.cpp
@@ -112,8 +112,8 @@ public:
 				SetShaderValue(Context.RHICmdList, ShaderRHI, TextureOffset, HMDDevice->GetTextureOffsetRight());
 				SetShaderValue(Context.RHICmdList, ShaderRHI, TextureUVOffset, -0.5f);
 			}				
-			      
-            QuadTexTransform = FMatrix::Identity;            
+				  
+			QuadTexTransform = FMatrix::Identity;            
 		}
 	}
 	
@@ -177,7 +177,7 @@ IMPLEMENT_SHADER_TYPE(, FPostProcessMorpheusPS,TEXT("MorpheusInclude"),TEXT("Mai
 
 void FRCPassPostProcessMorpheus::Process(FRenderingCompositePassContext& Context)
 {
-	SCOPED_DRAW_EVENT(Context.RHICmdList, PostProcessMorpheus, DEC_SCENE_ITEMS);
+	SCOPED_DRAW_EVENT(Context.RHICmdList, PostProcessMorpheus);
 	const FPooledRenderTargetDesc* InputDesc = GetInputDesc(ePId_Input0);
 
 	if(!InputDesc)
@@ -193,7 +193,7 @@ void FRCPassPostProcessMorpheus::Process(FRenderingCompositePassContext& Context
 	FIntRect DestRect = View.UnscaledViewRect;
 	FIntPoint SrcSize = InputDesc->Extent;
 
-    const FSceneRenderTargetItem& DestRenderTarget = PassOutputs[0].RequestSurface(Context);
+	const FSceneRenderTargetItem& DestRenderTarget = PassOutputs[0].RequestSurface(Context);
 
 	// Set the view family's render target/viewport.
 	SetRenderTarget(Context.RHICmdList, DestRenderTarget.TargetableTexture, FTextureRHIRef());
@@ -213,23 +213,23 @@ void FRCPassPostProcessMorpheus::Process(FRenderingCompositePassContext& Context
 
 	SetGlobalBoundShaderState(Context.RHICmdList, Context.GetFeatureLevel(), BoundShaderState, GFilterVertexDeclaration.VertexDeclarationRHI, *VertexShader, *PixelShader);
 
-    FMatrix QuadTexTransform;
-    FMatrix QuadPosTransform = FMatrix::Identity;
+	FMatrix QuadTexTransform;
+	FMatrix QuadPosTransform = FMatrix::Identity;
 
 	PixelShader->SetPS(Context, SrcRect, SrcSize, View.StereoPass, QuadTexTransform);
 
 	// Draw a quad mapping scene color to the view's render target
-    DrawTransformedRectangle(
+	DrawTransformedRectangle(
 		Context.RHICmdList,
-        0, 0,
-        DestRect.Width(), DestRect.Height(),
-        QuadPosTransform,
-        SrcRect.Min.X, SrcRect.Min.Y,
-        SrcRect.Width(), SrcRect.Height(),
-        QuadTexTransform,
-        DestRect.Size(),
-        SrcSize
-        );
+		0, 0,
+		DestRect.Width(), DestRect.Height(),
+		QuadPosTransform,
+		SrcRect.Min.X, SrcRect.Min.Y,
+		SrcRect.Width(), SrcRect.Height(),
+		QuadTexTransform,
+		DestRect.Size(),
+		SrcSize
+		);
 
 	Context.RHICmdList.CopyToResolveTarget(DestRenderTarget.TargetableTexture, DestRenderTarget.ShaderResourceTexture, false, FResolveParams());
 }
diff --git a/Engine/Source/Runtime/Renderer/Private/PostProcess/PostProcessMotionBlur.cpp b/Engine/Source/Runtime/Renderer/Private/PostProcess/PostProcessMotionBlur.cpp
--- a/Engine/Source/Runtime/Renderer/Private/PostProcess/PostProcessMotionBlur.cpp
+++ b/Engine/Source/Runtime/Renderer/Private/PostProcess/PostProcessMotionBlur.cpp
@@ -132,7 +132,7 @@ IMPLEMENT_SHADER_TYPE(,FPostProcessMotionBlurSetupPS, TEXT("PostProcessMotionBlu
 
 void FRCPassPostProcessMotionBlurSetup::Process(FRenderingCompositePassContext& Context)
 {
-	SCOPED_DRAW_EVENT(Context.RHICmdList, MotionBlurSetup, DEC_SCENE_ITEMS);
+	SCOPED_DRAW_EVENT(Context.RHICmdList, MotionBlurSetup);
 
 	const FPooledRenderTargetDesc* InputDesc = GetInputDesc(ePId_Input0);
 
@@ -438,7 +438,7 @@ FRCPassPostProcessMotionBlur::FRCPassPostProcessMotionBlur(uint32 InQuality)
 
 void FRCPassPostProcessMotionBlur::Process(FRenderingCompositePassContext& Context)
 {
-	SCOPED_DRAW_EVENT(Context.RHICmdList, MotionBlur, DEC_SCENE_ITEMS);
+	SCOPED_DRAW_EVENT(Context.RHICmdList, MotionBlur);
 
 	const FPooledRenderTargetDesc* InputDesc = GetInputDesc(ePId_Input0);
 
@@ -573,7 +573,7 @@ IMPLEMENT_SHADER_TYPE(,FPostProcessMotionBlurRecombinePS,TEXT("PostProcessMotion
 
 void FRCPassPostProcessMotionBlurRecombine::Process(FRenderingCompositePassContext& Context)
 {
-	SCOPED_DRAW_EVENT(Context.RHICmdList, MotionBlurRecombine, DEC_SCENE_ITEMS);
+	SCOPED_DRAW_EVENT(Context.RHICmdList, MotionBlurRecombine);
 
 	const FPooledRenderTargetDesc* InputDesc = GetInputDesc(ePId_Input0);
 
@@ -654,7 +654,7 @@ FPooledRenderTargetDesc FRCPassPostProcessMotionBlurRecombine::ComputeOutputDesc
 
 void FRCPassPostProcessVisualizeMotionBlur::Process(FRenderingCompositePassContext& Context)
 {
-	SCOPED_DRAW_EVENT(Context.RHICmdList, VisualizeMotionBlur, DEC_SCENE_ITEMS);
+	SCOPED_DRAW_EVENT(Context.RHICmdList, VisualizeMotionBlur);
 
 	const FPooledRenderTargetDesc* InputDesc = GetInputDesc(ePId_Input0);
 
diff --git a/Engine/Source/Runtime/Renderer/Private/PostProcess/PostProcessNoiseBlur.cpp b/Engine/Source/Runtime/Renderer/Private/PostProcess/PostProcessNoiseBlur.cpp
--- a/Engine/Source/Runtime/Renderer/Private/PostProcess/PostProcessNoiseBlur.cpp
+++ b/Engine/Source/Runtime/Renderer/Private/PostProcess/PostProcessNoiseBlur.cpp
@@ -112,7 +112,7 @@ void SetNoiseBlurShader(const FRenderingCompositePassContext& Context, float InR
 
 void FRCPassPostProcessNoiseBlur::Process(FRenderingCompositePassContext& Context)
 {
-	SCOPED_DRAW_EVENT(Context.RHICmdList, NoiseBlur, DEC_SCENE_ITEMS);
+	SCOPED_DRAW_EVENT(Context.RHICmdList, NoiseBlur);
 
 	const FPooledRenderTargetDesc* InputDesc = GetInputDesc(ePId_Input0);
 
diff --git a/Engine/Source/Runtime/Renderer/Private/PostProcess/PostProcessSceneColorFringe.cpp b/Engine/Source/Runtime/Renderer/Private/PostProcess/PostProcessSceneColorFringe.cpp
--- a/Engine/Source/Runtime/Renderer/Private/PostProcess/PostProcessSceneColorFringe.cpp
+++ b/Engine/Source/Runtime/Renderer/Private/PostProcess/PostProcessSceneColorFringe.cpp
@@ -143,7 +143,7 @@ IMPLEMENT_SHADER_TYPE(,FPostProcessSceneColorFringePS,TEXT("PostProcessSceneColo
 
 void FRCPassPostProcessSceneColorFringe::Process(FRenderingCompositePassContext& Context)
 {
-	SCOPED_DRAW_EVENT(Context.RHICmdList, SceneColorFringe, DEC_SCENE_ITEMS);
+	SCOPED_DRAW_EVENT(Context.RHICmdList, SceneColorFringe);
 
 	const FPooledRenderTargetDesc* InputDesc = GetInputDesc(ePId_Input0);
 
diff --git a/Engine/Source/Runtime/Renderer/Private/PostProcess/PostProcessSelectionOutline.cpp b/Engine/Source/Runtime/Renderer/Private/PostProcess/PostProcessSelectionOutline.cpp
--- a/Engine/Source/Runtime/Renderer/Private/PostProcess/PostProcessSelectionOutline.cpp
+++ b/Engine/Source/Runtime/Renderer/Private/PostProcess/PostProcessSelectionOutline.cpp
@@ -17,7 +17,7 @@
 
 void FRCPassPostProcessSelectionOutlineColor::Process(FRenderingCompositePassContext& Context)
 {
-	SCOPED_DRAW_EVENT(Context.RHICmdList, PostProcessSelectionOutlineBuffer, DEC_SCENE_ITEMS);
+	SCOPED_DRAW_EVENT(Context.RHICmdList, PostProcessSelectionOutlineBuffer);
 
 	const FPooledRenderTargetDesc* SceneColorInputDesc = GetInputDesc(ePId_Input0);
 
@@ -364,7 +364,7 @@ static void SetSelectionOutlineShaderTempl(const FRenderingCompositePassContext&
 
 void FRCPassPostProcessSelectionOutline::Process(FRenderingCompositePassContext& Context)
 {
-	SCOPED_DRAW_EVENT(Context.RHICmdList, PostProcessSelectionOutline, DEC_SCENE_ITEMS);
+	SCOPED_DRAW_EVENT(Context.RHICmdList, PostProcessSelectionOutline);
 	const FPooledRenderTargetDesc* SceneColorInputDesc = GetInputDesc(ePId_Input0);
 	const FPooledRenderTargetDesc* SelectionColorInputDesc = GetInputDesc(ePId_Input1);
 
diff --git a/Engine/Source/Runtime/Renderer/Private/PostProcess/PostProcessSubsurface.cpp b/Engine/Source/Runtime/Renderer/Private/PostProcess/PostProcessSubsurface.cpp
--- a/Engine/Source/Runtime/Renderer/Private/PostProcess/PostProcessSubsurface.cpp
+++ b/Engine/Source/Runtime/Renderer/Private/PostProcess/PostProcessSubsurface.cpp
@@ -208,7 +208,7 @@ FRCPassPostProcessSubsurfaceVisualize::FRCPassPostProcessSubsurfaceVisualize()
 
 void FRCPassPostProcessSubsurfaceVisualize::Process(FRenderingCompositePassContext& Context)
 {
-	SCOPED_DRAW_EVENT(Context.RHICmdList, SubsurfaceSetup, DEC_SCENE_ITEMS);
+	SCOPED_DRAW_EVENT(Context.RHICmdList, SubsurfaceSetup);
 
 	const FPooledRenderTargetDesc* InputDesc = GetInputDesc(ePId_Input0);
 
@@ -434,7 +434,7 @@ static bool DoSpecularCorrection()
 
 void FRCPassPostProcessSubsurfaceSetup::Process(FRenderingCompositePassContext& Context)
 {
-	SCOPED_DRAW_EVENT(Context.RHICmdList, SubsurfaceSetup, DEC_SCENE_ITEMS);
+	SCOPED_DRAW_EVENT(Context.RHICmdList, SubsurfaceSetup);
 
 	const FPooledRenderTargetDesc* InputDesc = GetInputDesc(ePId_Input0);
 
@@ -693,7 +693,7 @@ void FRCPassPostProcessSubsurface::Process(FRenderingCompositePassContext& Conte
 
 	TShaderMapRef<FPostProcessVS> VertexShader(Context.GetShaderMap());
 
-	SCOPED_DRAW_EVENTF(Context.RHICmdList, SubsurfacePass, DEC_SCENE_ITEMS, TEXT("SubsurfaceDirection#%d"), Direction);
+	SCOPED_DRAW_EVENTF(Context.RHICmdList, SubsurfacePass, TEXT("SubsurfaceDirection#%d"), Direction);
 
 	uint32 SampleSet = FMath::Clamp(CVarSSSSampleSet.GetValueOnRenderThread(), 0, 2);
 
@@ -876,7 +876,7 @@ void FRCPassPostProcessSubsurfaceRecombine::Process(FRenderingCompositePassConte
 
 	bool bDoSpecularCorrection = DoSpecularCorrection();
 
-	SCOPED_DRAW_EVENTF(Context.RHICmdList, SubsurfacePass, DEC_SCENE_ITEMS, TEXT("SubsurfacePassRecombine#%d"), (int32)bDoSpecularCorrection);
+	SCOPED_DRAW_EVENTF(Context.RHICmdList, SubsurfacePass, TEXT("SubsurfacePassRecombine#%d"), (int32)bDoSpecularCorrection);
 
 	if(bDoSpecularCorrection)
 	{
diff --git a/Engine/Source/Runtime/Renderer/Private/PostProcess/PostProcessTemporalAA.cpp b/Engine/Source/Runtime/Renderer/Private/PostProcess/PostProcessTemporalAA.cpp
--- a/Engine/Source/Runtime/Renderer/Private/PostProcess/PostProcessTemporalAA.cpp
+++ b/Engine/Source/Runtime/Renderer/Private/PostProcess/PostProcessTemporalAA.cpp
@@ -205,7 +205,7 @@ IMPLEMENT_TEMPORALAA_PIXELSHADER_TYPE(4, 1, TEXT("MainFastTemporalAAPS"));
 
 void FRCPassPostProcessSSRTemporalAA::Process(FRenderingCompositePassContext& Context)
 {
-	SCOPED_DRAW_EVENT(Context.RHICmdList, SSRTemporalAA, DEC_SCENE_ITEMS);
+	SCOPED_DRAW_EVENT(Context.RHICmdList, SSRTemporalAA);
 
 	const FPooledRenderTargetDesc* InputDesc = GetInputDesc(ePId_Input0);
 
@@ -279,7 +279,7 @@ FPooledRenderTargetDesc FRCPassPostProcessSSRTemporalAA::ComputeOutputDesc(EPass
 
 void FRCPassPostProcessDOFTemporalAA::Process(FRenderingCompositePassContext& Context)
 {
-	SCOPED_DRAW_EVENT(Context.RHICmdList, DOFTemporalAA, DEC_SCENE_ITEMS);
+	SCOPED_DRAW_EVENT(Context.RHICmdList, DOFTemporalAA);
 
 	const FPooledRenderTargetDesc* InputDesc = GetInputDesc(ePId_Input0);
 
@@ -368,7 +368,7 @@ void FRCPassPostProcessLightShaftTemporalAA::Process(FRenderingCompositePassCont
 		return;
 	}
 
-	SCOPED_DRAW_EVENT(Context.RHICmdList, LSTemporalAA, DEC_SCENE_ITEMS);
+	SCOPED_DRAW_EVENT(Context.RHICmdList, LSTemporalAA);
 
 	const FViewInfo& View = Context.View;
 	FSceneViewState* ViewState = Context.ViewState;
@@ -438,7 +438,7 @@ FPooledRenderTargetDesc FRCPassPostProcessLightShaftTemporalAA::ComputeOutputDes
 
 void FRCPassPostProcessTemporalAA::Process(FRenderingCompositePassContext& Context)
 {
-	SCOPED_DRAW_EVENT(Context.RHICmdList, TemporalAA, DEC_SCENE_ITEMS);
+	SCOPED_DRAW_EVENT(Context.RHICmdList, TemporalAA);
 
 	const FPooledRenderTargetDesc* InputDesc = GetInputDesc(ePId_Input0);
 
diff --git a/Engine/Source/Runtime/Renderer/Private/PostProcess/PostProcessTestImage.cpp b/Engine/Source/Runtime/Renderer/Private/PostProcess/PostProcessTestImage.cpp
--- a/Engine/Source/Runtime/Renderer/Private/PostProcess/PostProcessTestImage.cpp
+++ b/Engine/Source/Runtime/Renderer/Private/PostProcess/PostProcessTestImage.cpp
@@ -89,7 +89,7 @@ FRCPassPostProcessTestImage::FRCPassPostProcessTestImage()
 
 void FRCPassPostProcessTestImage::Process(FRenderingCompositePassContext& Context)
 {
-	SCOPED_DRAW_EVENT(Context.RHICmdList, TestImage, DEC_SCENE_ITEMS);
+	SCOPED_DRAW_EVENT(Context.RHICmdList, TestImage);
 
 	const FSceneView& View = Context.View;
 	const FSceneViewFamily& ViewFamily = *(View.Family);
diff --git a/Engine/Source/Runtime/Renderer/Private/PostProcess/PostProcessTonemap.cpp b/Engine/Source/Runtime/Renderer/Private/PostProcess/PostProcessTonemap.cpp
--- a/Engine/Source/Runtime/Renderer/Private/PostProcess/PostProcessTonemap.cpp
+++ b/Engine/Source/Runtime/Renderer/Private/PostProcess/PostProcessTonemap.cpp
@@ -1219,7 +1219,7 @@ static void SetShaderTempl(const FRenderingCompositePassContext& Context)
 
 void FRCPassPostProcessTonemap::Process(FRenderingCompositePassContext& Context)
 {
-	SCOPED_DRAW_EVENTF(Context.RHICmdList, PostProcessTonemap, DEC_SCENE_ITEMS, TEXT("Tonemapper#%d%s"), ConfigIndexPC, bDoGammaOnly ? TEXT(" GammaOnly") : TEXT(""));
+	SCOPED_DRAW_EVENTF(Context.RHICmdList, PostProcessTonemap, TEXT("Tonemapper#%d%s"), ConfigIndexPC, bDoGammaOnly ? TEXT(" GammaOnly") : TEXT(""));
 
 	const FPooledRenderTargetDesc* InputDesc = GetInputDesc(ePId_Input0);
 
@@ -1591,7 +1591,7 @@ FRCPassPostProcessTonemapES2::FRCPassPostProcessTonemapES2(const FViewInfo& View
 
 void FRCPassPostProcessTonemapES2::Process(FRenderingCompositePassContext& Context)
 {
-	SCOPED_DRAW_EVENTF(Context.RHICmdList, PostProcessTonemap, DEC_SCENE_ITEMS, TEXT("Tonemapper#%d%s"), ConfigIndexMobile, bUsedFramebufferFetch ? TEXT(" FramebufferFetch=0") : TEXT("FramebufferFetch=1"));
+	SCOPED_DRAW_EVENTF(Context.RHICmdList, PostProcessTonemap, TEXT("Tonemapper#%d%s"), ConfigIndexMobile, bUsedFramebufferFetch ? TEXT(" FramebufferFetch=0") : TEXT("FramebufferFetch=1"));
 
 	const FPooledRenderTargetDesc* InputDesc = GetInputDesc(ePId_Input0);
 
diff --git a/Engine/Source/Runtime/Renderer/Private/PostProcess/PostProcessUpscale.cpp b/Engine/Source/Runtime/Renderer/Private/PostProcess/PostProcessUpscale.cpp
--- a/Engine/Source/Runtime/Renderer/Private/PostProcess/PostProcessUpscale.cpp
+++ b/Engine/Source/Runtime/Renderer/Private/PostProcess/PostProcessUpscale.cpp
@@ -128,7 +128,7 @@ void FRCPassPostProcessUpscale::SetShader(const FRenderingCompositePassContext&
 
 void FRCPassPostProcessUpscale::Process(FRenderingCompositePassContext& Context)
 {
-	SCOPED_DRAW_EVENT(Context.RHICmdList, PostProcessUpscale, DEC_SCENE_ITEMS);
+	SCOPED_DRAW_EVENT(Context.RHICmdList, PostProcessUpscale);
 	const FPooledRenderTargetDesc* InputDesc = GetInputDesc(ePId_Input0);
 
 	if(!InputDesc)
diff --git a/Engine/Source/Runtime/Renderer/Private/PostProcess/PostProcessVisualizeBuffer.cpp b/Engine/Source/Runtime/Renderer/Private/PostProcess/PostProcessVisualizeBuffer.cpp
--- a/Engine/Source/Runtime/Renderer/Private/PostProcess/PostProcessVisualizeBuffer.cpp
+++ b/Engine/Source/Runtime/Renderer/Private/PostProcess/PostProcessVisualizeBuffer.cpp
@@ -128,7 +128,7 @@ FShader* FRCPassPostProcessVisualizeBuffer::SetShaderTempl(const FRenderingCompo
 
 void FRCPassPostProcessVisualizeBuffer::Process(FRenderingCompositePassContext& Context)
 {
-	SCOPED_DRAW_EVENT(Context.RHICmdList, VisualizeBuffer, DEC_SCENE_ITEMS);
+	SCOPED_DRAW_EVENT(Context.RHICmdList, VisualizeBuffer);
 	const FPooledRenderTargetDesc* InputDesc = GetInputDesc(ePId_Input0);
 
 	if(!InputDesc)
diff --git a/Engine/Source/Runtime/Renderer/Private/PostProcess/PostProcessVisualizeHDR.cpp b/Engine/Source/Runtime/Renderer/Private/PostProcess/PostProcessVisualizeHDR.cpp
--- a/Engine/Source/Runtime/Renderer/Private/PostProcess/PostProcessVisualizeHDR.cpp
+++ b/Engine/Source/Runtime/Renderer/Private/PostProcess/PostProcessVisualizeHDR.cpp
@@ -136,7 +136,7 @@ FString LogToString(float LogValue)
 
 void FRCPassPostProcessVisualizeHDR::Process(FRenderingCompositePassContext& Context)
 {
-	SCOPED_DRAW_EVENT(Context.RHICmdList, PostProcessVisualizeHDR, DEC_SCENE_ITEMS);
+	SCOPED_DRAW_EVENT(Context.RHICmdList, PostProcessVisualizeHDR);
 	const FPooledRenderTargetDesc* InputDesc = GetInputDesc(ePId_Input0);
 
 	if(!InputDesc)
diff --git a/Engine/Source/Runtime/Renderer/Private/PostProcess/PostProcessWeightedSampleSum.cpp b/Engine/Source/Runtime/Renderer/Private/PostProcess/PostProcessWeightedSampleSum.cpp
--- a/Engine/Source/Runtime/Renderer/Private/PostProcess/PostProcessWeightedSampleSum.cpp
+++ b/Engine/Source/Runtime/Renderer/Private/PostProcess/PostProcessWeightedSampleSum.cpp
@@ -370,7 +370,7 @@ void FRCPassPostProcessWeightedSampleSum::Process(FRenderingCompositePassContext
 
 	uint32 NumSamples = Compute1DGaussianFilterKernel(FeatureLevel, EffectiveBlurRadius, OffsetAndWeight, MaxNumSamples);
 
-	SCOPED_DRAW_EVENTF(Context.RHICmdList, PostProcessWeightedSampleSum, DEC_SCENE_ITEMS, TEXT("PostProcessWeightedSampleSum#%d"), NumSamples);
+	SCOPED_DRAW_EVENTF(Context.RHICmdList, PostProcessWeightedSampleSum, TEXT("PostProcessWeightedSampleSum#%d"), NumSamples);
 
 	// compute weights as weighted contributions of the TintValue
 	for(uint32 i = 0; i < NumSamples; ++i)
diff --git a/Engine/Source/Runtime/Renderer/Private/PostProcess/SceneRenderTargets.cpp b/Engine/Source/Runtime/Renderer/Private/PostProcess/SceneRenderTargets.cpp
--- a/Engine/Source/Runtime/Renderer/Private/PostProcess/SceneRenderTargets.cpp
+++ b/Engine/Source/Runtime/Renderer/Private/PostProcess/SceneRenderTargets.cpp
@@ -260,7 +260,7 @@ void FSceneRenderTargets::Allocate(const FSceneViewFamily& ViewFamily)
 /** Clears the GBuffer render targets to default values. */
 void FSceneRenderTargets::ClearGBufferTargets(FRHICommandListImmediate& RHICmdList, const FLinearColor& ClearColor)
 {
-	SCOPED_DRAW_EVENT(RHICmdList, ClearGBufferTargets, DEC_SCENE_ITEMS);
+	SCOPED_DRAW_EVENT(RHICmdList, ClearGBufferTargets);
 
 	// Clear GBufferA, GBufferB, GBufferC, GBufferD, GBufferE
 	{
@@ -283,7 +283,7 @@ void FSceneRenderTargets::ClearGBufferTargets(FRHICommandListImmediate& RHICmdLi
 
 void FSceneRenderTargets::BeginRenderingSceneColor(FRHICommandList& RHICmdList, bool bGBufferPass)
 {
-	SCOPED_DRAW_EVENT(RHICmdList, BeginRenderingSceneColor, DEC_SCENE_ITEMS);
+	SCOPED_DRAW_EVENT(RHICmdList, BeginRenderingSceneColor);
 
 	if(IsSimpleDynamicLightingEnabled())
 	{
@@ -643,7 +643,7 @@ void FSceneRenderTargets::AdjustGBufferRefCount(int Delta)
 
 void FSceneRenderTargets::FinishRenderingSceneColor(FRHICommandListImmediate& RHICmdList, bool bKeepChanges, const FResolveRect& ResolveRect)
 {
-	SCOPED_DRAW_EVENT(RHICmdList, FinishRenderingSceneColor, DEC_SCENE_ITEMS);
+	SCOPED_DRAW_EVENT(RHICmdList, FinishRenderingSceneColor);
 
 	if(bKeepChanges)
 	{
@@ -657,7 +657,7 @@ bool FSceneRenderTargets::BeginRenderingCustomDepth(FRHICommandListImmediate& RH
 
 	if(CustomDepth)
 	{
-		SCOPED_DRAW_EVENT(RHICmdList, BeginRenderingCustomDepth, DEC_SCENE_ITEMS);
+		SCOPED_DRAW_EVENT(RHICmdList, BeginRenderingCustomDepth);
 		
 		SetRenderTarget(RHICmdList, FTextureRHIRef(), CustomDepth->GetRenderTargetItem().ShaderResourceTexture);
 
@@ -669,7 +669,7 @@ bool FSceneRenderTargets::BeginRenderingCustomDepth(FRHICommandListImmediate& RH
 
 void FSceneRenderTargets::FinishRenderingCustomDepth(FRHICommandListImmediate& RHICmdList, const FResolveRect& ResolveRect)
 {
-	SCOPED_DRAW_EVENT(RHICmdList, FinishRenderingCustomDepth, DEC_SCENE_ITEMS);
+	SCOPED_DRAW_EVENT(RHICmdList, FinishRenderingCustomDepth);
 
 	auto& CurrentSceneColor = GetSceneColor();
 	RHICmdList.CopyToResolveTarget(CurrentSceneColor->GetRenderTargetItem().TargetableTexture, CurrentSceneColor->GetRenderTargetItem().ShaderResourceTexture, true, FResolveParams(ResolveRect));
@@ -682,7 +682,7 @@ void FSceneRenderTargets::FinishRenderingCustomDepth(FRHICommandListImmediate& R
 */
 void FSceneRenderTargets::ResolveSceneColor(FRHICommandList& RHICmdList, const FResolveRect& ResolveRect)
 {
-	SCOPED_DRAW_EVENT(RHICmdList, ResolveSceneColor, DEC_SCENE_ITEMS);
+	SCOPED_DRAW_EVENT(RHICmdList, ResolveSceneColor);
 
 	auto& CurrentSceneColor = GetSceneColor();
 
@@ -750,7 +750,7 @@ void FSceneRenderTargets::ResolveGBufferSurfaces(FRHICommandList& RHICmdList, co
 {
 	if (CurrentFeatureLevel >= ERHIFeatureLevel::SM4)
 	{
-		SCOPED_DRAW_EVENT(RHICmdList, ResolveGBufferSurfaces, DEC_SCENE_ITEMS);
+		SCOPED_DRAW_EVENT(RHICmdList, ResolveGBufferSurfaces);
 
 		RHICmdList.CopyToResolveTarget(GSceneRenderTargets.GBufferA->GetRenderTargetItem().TargetableTexture, GSceneRenderTargets.GBufferA->GetRenderTargetItem().ShaderResourceTexture, true, FResolveParams(ResolveRect));
 		RHICmdList.CopyToResolveTarget(GSceneRenderTargets.GBufferB->GetRenderTargetItem().TargetableTexture, GSceneRenderTargets.GBufferB->GetRenderTargetItem().ShaderResourceTexture, true, FResolveParams(ResolveRect));
@@ -766,7 +766,7 @@ void FSceneRenderTargets::ResolveGBufferSurfaces(FRHICommandList& RHICmdList, co
 
 void FSceneRenderTargets::BeginRenderingPrePass(FRHICommandList& RHICmdList)
 {
-	SCOPED_DRAW_EVENT(RHICmdList, BeginRenderingPrePass, DEC_SCENE_ITEMS);
+	SCOPED_DRAW_EVENT(RHICmdList, BeginRenderingPrePass);
 
 	// Set the scene depth surface and a DUMMY buffer as color buffer
 	// (as long as it's the same dimension as the depth buffer),
@@ -776,7 +776,7 @@ void FSceneRenderTargets::BeginRenderingPrePass(FRHICommandList& RHICmdList)
 
 void FSceneRenderTargets::FinishRenderingPrePass(FRHICommandListImmediate& RHICmdList)
 {
-	SCOPED_DRAW_EVENT(RHICmdList, FinishRenderingPrePass, DEC_SCENE_ITEMS);
+	SCOPED_DRAW_EVENT(RHICmdList, FinishRenderingPrePass);
 	GRenderTargetPool.VisualizeTexture.SetCheckPoint(RHICmdList, SceneDepthZ);
 }
 
@@ -806,7 +806,7 @@ void FSceneRenderTargets::BeginRenderingShadowDepth(FRHICommandList& RHICmdList,
 
 void FSceneRenderTargets::BeginRenderingCubeShadowDepth(FRHICommandList& RHICmdList, int32 ShadowResolution)
 {
-	SCOPED_DRAW_EVENT(RHICmdList, BeginRenderingCubeShadowDepth, DEC_SCENE_ITEMS);
+	SCOPED_DRAW_EVENT(RHICmdList, BeginRenderingCubeShadowDepth);
 	SetRenderTarget(RHICmdList, FTextureRHIRef(), GetCubeShadowDepthZSurface(ShadowResolution));
 }
 
@@ -847,20 +847,20 @@ void FSceneRenderTargets::FinishRenderingReflectiveShadowMap(FRHICommandList& RH
 
 void FSceneRenderTargets::FinishRenderingCubeShadowDepth(FRHICommandList& RHICmdList, int32 ShadowResolution, const FResolveParams& ResolveParams)
 {
-	SCOPED_DRAW_EVENT(RHICmdList, FinishRenderingCubeShadowDepth, DEC_SCENE_ITEMS);
+	SCOPED_DRAW_EVENT(RHICmdList, FinishRenderingCubeShadowDepth);
 	RHICmdList.CopyToResolveTarget(GetCubeShadowDepthZSurface(ShadowResolution), GetCubeShadowDepthZTexture(ShadowResolution), false, ResolveParams);
 }
 
 void FSceneRenderTargets::BeginRenderingSceneAlphaCopy(FRHICommandListImmediate& RHICmdList)
 {
-	SCOPED_DRAW_EVENT(RHICmdList, BeginRenderingSceneAlphaCopy, DEC_SCENE_ITEMS);
+	SCOPED_DRAW_EVENT(RHICmdList, BeginRenderingSceneAlphaCopy);
 	GRenderTargetPool.VisualizeTexture.SetCheckPoint(RHICmdList, GSceneRenderTargets.SceneAlphaCopy);
 	SetRenderTarget(RHICmdList, GetSceneAlphaCopySurface(), 0);
 }
 
 void FSceneRenderTargets::FinishRenderingSceneAlphaCopy(FRHICommandListImmediate& RHICmdList)
 {
-	SCOPED_DRAW_EVENT(RHICmdList, FinishRenderingSceneAlphaCopy, DEC_SCENE_ITEMS);
+	SCOPED_DRAW_EVENT(RHICmdList, FinishRenderingSceneAlphaCopy);
 	RHICmdList.CopyToResolveTarget(GetSceneAlphaCopySurface(), SceneAlphaCopy->GetRenderTargetItem().ShaderResourceTexture, false, FResolveParams(FResolveRect()));
 	GRenderTargetPool.VisualizeTexture.SetCheckPoint(RHICmdList, GSceneRenderTargets.SceneAlphaCopy);
 }
@@ -868,7 +868,7 @@ void FSceneRenderTargets::FinishRenderingSceneAlphaCopy(FRHICommandListImmediate
 
 void FSceneRenderTargets::BeginRenderingLightAttenuation(FRHICommandList& RHICmdList)
 {
-	SCOPED_DRAW_EVENT(RHICmdList, BeginRenderingLightAttenuation, DEC_SCENE_ITEMS);
+	SCOPED_DRAW_EVENT(RHICmdList, BeginRenderingLightAttenuation);
 
 	AllocLightAttenuation();
 
@@ -880,7 +880,7 @@ void FSceneRenderTargets::BeginRenderingLightAttenuation(FRHICommandList& RHICmd
 
 void FSceneRenderTargets::FinishRenderingLightAttenuation(FRHICommandList& RHICmdList)
 {
-	SCOPED_DRAW_EVENT(RHICmdList, FinishRenderingLightAttenuation, DEC_SCENE_ITEMS);
+	SCOPED_DRAW_EVENT(RHICmdList, FinishRenderingLightAttenuation);
 
 	// Resolve the light attenuation surface.
 	RHICmdList.CopyToResolveTarget(GetLightAttenuationSurface(), LightAttenuation->GetRenderTargetItem().ShaderResourceTexture, false, FResolveParams(FResolveRect()));
@@ -902,7 +902,7 @@ bool FSceneRenderTargets::BeginRenderingSeparateTranslucency(FRHICommandList& RH
 {
 	if(IsSeparateTranslucencyActive(View))
 	{
-		SCOPED_DRAW_EVENT(RHICmdList, BeginSeparateTranslucency, DEC_SCENE_ITEMS);
+		SCOPED_DRAW_EVENT(RHICmdList, BeginSeparateTranslucency);
 
 		FSceneViewState* ViewState = (FSceneViewState*)View.State;
 
@@ -943,7 +943,7 @@ void FSceneRenderTargets::FinishRenderingSeparateTranslucency(FRHICommandList& R
 
 void FSceneRenderTargets::ResolveSceneDepthTexture(FRHICommandList& RHICmdList)
 {
-	SCOPED_DRAW_EVENT(RHICmdList, ResolveSceneDepthTexture, DEC_SCENE_ITEMS);
+	SCOPED_DRAW_EVENT(RHICmdList, ResolveSceneDepthTexture);
 
 	RHICmdList.CopyToResolveTarget(GetSceneDepthSurface(), GetSceneDepthTexture(), true, FResolveParams());
 }
@@ -955,7 +955,7 @@ void FSceneRenderTargets::ResolveSceneDepthToAuxiliaryTexture(FRHICommandList& R
 	// is unsupported in d3d feature level 10.
 	if(!GSupportsDepthFetchDuringDepthTest)
 	{
-		SCOPED_DRAW_EVENT(RHICmdList, ResolveSceneDepthToAuxiliaryTexture, DEC_SCENE_ITEMS);
+		SCOPED_DRAW_EVENT(RHICmdList, ResolveSceneDepthToAuxiliaryTexture);
 
 		RHICmdList.CopyToResolveTarget(GetSceneDepthSurface(), GetAuxiliarySceneDepthTexture(), true, FResolveParams());
 	}
diff --git a/Engine/Source/Runtime/Renderer/Private/PostProcess/ScreenSpaceReflections.cpp b/Engine/Source/Runtime/Renderer/Private/PostProcess/ScreenSpaceReflections.cpp
--- a/Engine/Source/Runtime/Renderer/Private/PostProcess/ScreenSpaceReflections.cpp
+++ b/Engine/Source/Runtime/Renderer/Private/PostProcess/ScreenSpaceReflections.cpp
@@ -178,7 +178,7 @@ static int32 ComputeSSRQuality(float Quality)
 
 void FRCPassPostProcessScreenSpaceReflections::Process(FRenderingCompositePassContext& Context)
 {
-	SCOPED_DRAW_EVENT(Context.RHICmdList, ScreenSpaceReflections, DEC_SCENE_ITEMS);
+	SCOPED_DRAW_EVENT(Context.RHICmdList, ScreenSpaceReflections);
 
 	const FSceneView& View = Context.View;
 	const auto FeatureLevel = Context.GetFeatureLevel();
diff --git a/Engine/Source/Runtime/Renderer/Private/PostProcess/VisualizeTexture.cpp b/Engine/Source/Runtime/Renderer/Private/PostProcess/VisualizeTexture.cpp
--- a/Engine/Source/Runtime/Renderer/Private/PostProcess/VisualizeTexture.cpp
+++ b/Engine/Source/Runtime/Renderer/Private/PostProcess/VisualizeTexture.cpp
@@ -417,13 +417,13 @@ void FVisualizeTexture::GenerateContent(FRHICommandListImmediate& RHICmdList, co
 
 	if(!(Desc.Flags & TexCreate_CPUReadback))		// We cannot make a texture lookup on such elements
 	{	
-		SCOPED_DRAW_EVENT(RHICmdList, VisualizeTexture, DEC_SCENE_ITEMS);
+		SCOPED_DRAW_EVENT(RHICmdList, VisualizeTexture);
 		// continue rendering to HDR if necessary
 		RenderVisualizeTexture(RHICmdList, FeatureLevel, VisualizeTextureData);
 	}
 
 	{
-		SCOPED_DRAW_EVENT(RHICmdList, VisCopy, DEC_SCENE_ITEMS);
+		SCOPED_DRAW_EVENT(RHICmdList, VisCopy);
 		RHICmdList.CopyToResolveTarget(DestRenderTarget.TargetableTexture, DestRenderTarget.ShaderResourceTexture, false, FResolveParams());
 	}
 
@@ -524,7 +524,7 @@ void FVisualizeTexture::PresentContent(FRHICommandListImmediate& RHICmdList, con
 	FIntRect VisualizeTextureRect = ComputeVisualizeTextureRect(Desc.Extent);
 
 	{
-		SCOPED_DRAW_EVENT(RHICmdList, VisCopyToMain, DEC_SCENE_ITEMS);
+		SCOPED_DRAW_EVENT(RHICmdList, VisCopyToMain);
 		// Draw a quad mapping scene color to the view's render target
 		DrawRectangle(
 			RHICmdList,
diff --git a/Engine/Source/Runtime/Renderer/Private/ReflectionEnvironment.cpp b/Engine/Source/Runtime/Renderer/Private/ReflectionEnvironment.cpp
--- a/Engine/Source/Runtime/Renderer/Private/ReflectionEnvironment.cpp
+++ b/Engine/Source/Runtime/Renderer/Private/ReflectionEnvironment.cpp
@@ -716,7 +716,7 @@ void FDeferredShadingSceneRenderer::RenderTiledDeferredImageBasedReflections(FRH
 		// ReflectionEnv is assumed to be on when going into this method
 		{
 			// Render the reflection environment with tiled deferred culling
-			SCOPED_DRAW_EVENT(RHICmdList, ReflectionEnvironmentGather, DEC_SCENE_ITEMS);
+			SCOPED_DRAW_EVENT(RHICmdList, ReflectionEnvironmentGather);
 
 			SetRenderTarget(RHICmdList, NULL, NULL);
 
@@ -842,7 +842,7 @@ void FDeferredShadingSceneRenderer::RenderStandardDeferredImageBasedReflections(
 		{
 			bRequiresApply = true;
 
-			SCOPED_DRAW_EVENT(RHICmdList, StandardDeferredReflectionEnvironment, DEC_SCENE_ITEMS);
+			SCOPED_DRAW_EVENT(RHICmdList, StandardDeferredReflectionEnvironment);
 
 			{
 				FPooledRenderTargetDesc Desc = GSceneRenderTargets.GetSceneColor()->GetDesc();
@@ -905,7 +905,7 @@ void FDeferredShadingSceneRenderer::RenderStandardDeferredImageBasedReflections(
 		if (bRequiresApply)
 		{
 			// Apply reflections to screen
-			SCOPED_DRAW_EVENT(RHICmdList, ReflectionApply, DEC_SCENE_ITEMS);
+			SCOPED_DRAW_EVENT(RHICmdList, ReflectionApply);
 
 			GSceneRenderTargets.BeginRenderingSceneColor(RHICmdList);
 
diff --git a/Engine/Source/Runtime/Renderer/Private/ReflectionEnvironmentCapture.cpp b/Engine/Source/Runtime/Renderer/Private/ReflectionEnvironmentCapture.cpp
--- a/Engine/Source/Runtime/Renderer/Private/ReflectionEnvironmentCapture.cpp
+++ b/Engine/Source/Runtime/Renderer/Private/ReflectionEnvironmentCapture.cpp
@@ -316,11 +316,11 @@ void FilterReflectionEnvironment(FRHICommandListImmediate& RHICmdList, ERHIFeatu
 	auto ShaderMap = GetGlobalShaderMap(FeatureLevel);
 
 	{	
-		SCOPED_DRAW_EVENT(RHICmdList, DownsampleCubeMips, DEC_SCENE_ITEMS);
+		SCOPED_DRAW_EVENT(RHICmdList, DownsampleCubeMips);
 		// Downsample all the mips, each one reads from the mip above it
 		for (int32 MipIndex = 1; MipIndex < NumMips; MipIndex++)
 		{
-			SCOPED_DRAW_EVENT(RHICmdList, DownsampleCubeMip, DEC_SCENE_ITEMS);
+			SCOPED_DRAW_EVENT(RHICmdList, DownsampleCubeMip);
 			const int32 SourceMipIndex = FMath::Max(MipIndex - 1, 0);
 			const int32 MipSize = 1 << (NumMips - MipIndex - 1);
 
@@ -368,7 +368,7 @@ void FilterReflectionEnvironment(FRHICommandListImmediate& RHICmdList, ERHIFeatu
 
 	if (OutIrradianceEnvironmentMap)
 	{
-		SCOPED_DRAW_EVENT(RHICmdList, ComputeDiffuseIrradiance, DEC_SCENE_ITEMS);
+		SCOPED_DRAW_EVENT(RHICmdList, ComputeDiffuseIrradiance);
 		check(DiffuseConvolutionSource != NULL);
 		ComputeDiffuseIrradiance(RHICmdList, FeatureLevel, DiffuseConvolutionSource->ShaderResourceTexture, DiffuseConvolutionSourceMip, OutIrradianceEnvironmentMap);
 	}
@@ -376,11 +376,11 @@ void FilterReflectionEnvironment(FRHICommandListImmediate& RHICmdList, ERHIFeatu
 	ComputeAverageBrightness(RHICmdList, FeatureLevel);
 
 	{	
-		SCOPED_DRAW_EVENT(RHICmdList, FilterCubeMap, DEC_SCENE_ITEMS);
+		SCOPED_DRAW_EVENT(RHICmdList, FilterCubeMap);
 		// Filter all the mips, each one reads from whichever scratch render target holds the downsampled contents, and writes to the destination cubemap
 		for (int32 MipIndex = 0; MipIndex < NumMips; MipIndex++)
 		{
-			SCOPED_DRAW_EVENT(RHICmdList, FilterCubeMip, DEC_SCENE_ITEMS);
+			SCOPED_DRAW_EVENT(RHICmdList, FilterCubeMip);
 			FSceneRenderTargetItem& EffectiveRT = GetEffectiveRenderTarget(false, MipIndex);
 			FSceneRenderTargetItem& EffectiveSource = GetEffectiveSourceTexture(false, MipIndex);
 			check(EffectiveRT.TargetableTexture != EffectiveSource.ShaderResourceTexture);
@@ -685,7 +685,7 @@ void CaptureSceneToScratchCubemap(FRHICommandListImmediate& RHICmdList, FSceneRe
 	const auto FeatureLevel = SceneRenderer->FeatureLevel;
 	
 	{
-		SCOPED_DRAW_EVENT(RHICmdList, CubeMapCapture, DEC_SCENE_ITEMS);
+		SCOPED_DRAW_EVENT(RHICmdList, CubeMapCapture);
 
 		// Render the scene normally for one face of the cubemap
 		SceneRenderer->Render(RHICmdList);
@@ -708,7 +708,7 @@ void CaptureSceneToScratchCubemap(FRHICommandListImmediate& RHICmdList, FSceneRe
 		FSceneRenderTargetItem& EffectiveColorRT =  GSceneRenderTargets.ReflectionColorScratchCubemap[0]->GetRenderTargetItem();
 
 		{
-			SCOPED_DRAW_EVENT(RHICmdList, CubeMapCopyScene, DEC_SCENE_ITEMS);
+			SCOPED_DRAW_EVENT(RHICmdList, CubeMapCopyScene);
 			
 			// Copy the captured scene into the cubemap face
 			SetRenderTarget(RHICmdList, EffectiveColorRT.TargetableTexture, 0, CubeFace, NULL);
diff --git a/Engine/Source/Runtime/Renderer/Private/SceneCaptureRendering.cpp b/Engine/Source/Runtime/Renderer/Private/SceneCaptureRendering.cpp
--- a/Engine/Source/Runtime/Renderer/Private/SceneCaptureRendering.cpp
+++ b/Engine/Source/Runtime/Renderer/Private/SceneCaptureRendering.cpp
@@ -24,7 +24,7 @@ void UpdateSceneCaptureContent_RenderThread(FRHICommandListImmediate& RHICmdList
 #if WANTS_DRAW_MESH_EVENTS
 		FString EventName;
 		OwnerName.ToString(EventName);
-		SCOPED_DRAW_EVENTF(RHICmdList, SceneCapture, DEC_SCENE_ITEMS, TEXT("SceneCapture %s"), *EventName);
+		SCOPED_DRAW_EVENTF(RHICmdList, SceneCapture, TEXT("SceneCapture %s"), *EventName);
 #endif
 
 		// Render the scene normally
diff --git a/Engine/Source/Runtime/Renderer/Private/SceneOcclusion.cpp b/Engine/Source/Runtime/Renderer/Private/SceneOcclusion.cpp
--- a/Engine/Source/Runtime/Renderer/Private/SceneOcclusion.cpp
+++ b/Engine/Source/Runtime/Renderer/Private/SceneOcclusion.cpp
@@ -675,7 +675,7 @@ IMPLEMENT_SHADER_TYPE(,FHZBTestPS,TEXT("HZBOcclusion"),TEXT("HZBTestPS"),SF_Pixe
 
 void FHZBOcclusionTester::Submit(FRHICommandListImmediate& RHICmdList, const FViewInfo& View)
 {
-	SCOPED_DRAW_EVENT(RHICmdList, SubmitHZB, DEC_SCENE_ITEMS);
+	SCOPED_DRAW_EVENT(RHICmdList, SubmitHZB);
 
 	FSceneViewState* ViewState = (FSceneViewState*)View.State;
 	if( !ViewState )
@@ -744,7 +744,7 @@ void FHZBOcclusionTester::Submit(FRHICommandListImmediate& RHICmdList, const FVi
 
 	// Draw test
 	{
-		SCOPED_DRAW_EVENT(RHICmdList, TestHZB, DEC_SCENE_ITEMS);
+		SCOPED_DRAW_EVENT(RHICmdList, TestHZB);
 
 		SetRenderTarget(RHICmdList, ResultsTextureGPU->GetRenderTargetItem().TargetableTexture, NULL);
 
@@ -855,7 +855,7 @@ IMPLEMENT_SHADER_TYPE(template<>,THZBBuildPS<1>,TEXT("HZBOcclusion"),TEXT("HZBBu
 void BuildHZB(FRHICommandListImmediate& RHICmdList, const FViewInfo& View)
 {
 	QUICK_SCOPE_CYCLE_COUNTER(STAT_BuildHZB);
-	SCOPED_DRAW_EVENT(RHICmdList, BuildHZB, DEC_SCENE_ITEMS);
+	SCOPED_DRAW_EVENT(RHICmdList, BuildHZB);
 
 	FSceneViewState* ViewState = (FSceneViewState*)View.State;
 
@@ -888,7 +888,7 @@ void BuildHZB(FRHICommandListImmediate& RHICmdList, const FViewInfo& View)
 	
 	// Mip 0
 	{
-		SCOPED_DRAW_EVENTF(RHICmdList, BuildHZB, DEC_SCENE_ITEMS, TEXT("HZB#%d"), 0);
+		SCOPED_DRAW_EVENTF(RHICmdList, BuildHZB, TEXT("HZB#%d"), 0);
 
 		SetRenderTarget(RHICmdList, HZBRenderTarget.TargetableTexture, 0, NULL);
 
@@ -922,7 +922,7 @@ void BuildHZB(FRHICommandListImmediate& RHICmdList, const FViewInfo& View)
 	// Mip 1-7
 	for( uint8 MipIndex = 1; MipIndex < NumMips; MipIndex++ )
 	{
-		SCOPED_DRAW_EVENTF(RHICmdList, BuildHZB, DEC_SCENE_ITEMS, TEXT("HZB#%d"), MipIndex);
+		SCOPED_DRAW_EVENTF(RHICmdList, BuildHZB, TEXT("HZB#%d"), MipIndex);
 
 		SetRenderTarget(RHICmdList, HZBRenderTarget.TargetableTexture, MipIndex, NULL);
 
@@ -974,7 +974,7 @@ void FDeferredShadingSceneRenderer::BeginOcclusionTests(FRHICommandListImmediate
 	// Perform occlusion queries for each view
 	for(int32 ViewIndex = 0;ViewIndex < Views.Num();ViewIndex++)
 	{
-		SCOPED_DRAW_EVENT(RHICmdList, BeginOcclusionTests, DEC_SCENE_ITEMS);
+		SCOPED_DRAW_EVENT(RHICmdList, BeginOcclusionTests);
 		FViewInfo& View = Views[ViewIndex];
 
 		if (bUseDownsampledDepth)
@@ -1029,7 +1029,7 @@ void FDeferredShadingSceneRenderer::BeginOcclusionTests(FRHICommandListImmediate
 			ShadowOcclusionQueryMap.Reset();
 
 			{
-				SCOPED_DRAW_EVENT(RHICmdList, ShadowFrustumQueries, DEC_SCENE_ITEMS);
+				SCOPED_DRAW_EVENT(RHICmdList, ShadowFrustumQueries);
 
 				for(TSparseArray<FLightSceneInfoCompact>::TConstIterator LightIt(Scene->Lights);LightIt;++LightIt)
 				{
@@ -1100,11 +1100,11 @@ void FDeferredShadingSceneRenderer::BeginOcclusionTests(FRHICommandListImmediate
 				VertexShader->SetParameters(RHICmdList, View);
 
 				{
-					SCOPED_DRAW_EVENT(RHICmdList, IndividualQueries, DEC_SCENE_ITEMS);
+					SCOPED_DRAW_EVENT(RHICmdList, IndividualQueries);
 					View.IndividualOcclusionQueries.Flush(RHICmdList);
 				}
 				{
-					SCOPED_DRAW_EVENT(RHICmdList, GroupedQueries, DEC_SCENE_ITEMS);
+					SCOPED_DRAW_EVENT(RHICmdList, GroupedQueries);
 					View.GroupedOcclusionQueries.Flush(RHICmdList);
 				}
 			}
@@ -1120,7 +1120,7 @@ void FDeferredShadingSceneRenderer::BeginOcclusionTests(FRHICommandListImmediate
 		{
 			check( ViewState->HZBOcclusionTests.IsValidFrame(View.FrameNumber) );
 
-			SCOPED_DRAW_EVENT(RHICmdList, HZB, DEC_SCENE_ITEMS);
+			SCOPED_DRAW_EVENT(RHICmdList, HZB);
 
 			BuildHZB(RHICmdList, View);
 			ViewState->HZBOcclusionTests.Submit(RHICmdList, View);
diff --git a/Engine/Source/Runtime/Renderer/Private/SceneRendering.cpp b/Engine/Source/Runtime/Renderer/Private/SceneRendering.cpp
--- a/Engine/Source/Runtime/Renderer/Private/SceneRendering.cpp
+++ b/Engine/Source/Runtime/Renderer/Private/SceneRendering.cpp
@@ -792,7 +792,7 @@ FSceneRenderer::~FSceneRenderer()
 */
 void FSceneRenderer::RenderFinish(FRHICommandListImmediate& RHICmdList)
 {
-	SCOPED_DRAW_EVENT(RHICmdList, RenderFinish, DEC_SCENE_ITEMS);
+	SCOPED_DRAW_EVENT(RHICmdList, RenderFinish);
 
 	if(FRCPassPostProcessBusyWait::IsPassRequired())
 	{
@@ -821,7 +821,7 @@ void FSceneRenderer::RenderFinish(FRHICommandListImmediate& RHICmdList)
 	{
 		for(int32 ViewIndex = 0;ViewIndex < Views.Num();ViewIndex++)
 		{	
-			SCOPED_CONDITIONAL_DRAW_EVENTF(RHICmdList, EventView, Views.Num() > 1, DEC_SCENE_ITEMS, TEXT("View%d"), ViewIndex);
+			SCOPED_CONDITIONAL_DRAW_EVENTF(RHICmdList, EventView, Views.Num() > 1, TEXT("View%d"), ViewIndex);
 			FViewInfo& View = Views[ViewIndex];
 
 			bool bShowPrecomputedVisibilityWarning = false;
@@ -995,11 +995,11 @@ void FSceneRenderer::RenderCustomDepthPass(FRHICommandListImmediate& RHICmdList)
 	// Render CustomDepth
 	if (GSceneRenderTargets.BeginRenderingCustomDepth(RHICmdList, bPrimitives))
 	{
-		SCOPED_DRAW_EVENT(RHICmdList, CustomDepth, DEC_SCENE_ITEMS);
+		SCOPED_DRAW_EVENT(RHICmdList, CustomDepth);
 
 		for(int32 ViewIndex = 0;ViewIndex < Views.Num();ViewIndex++)
 		{
-			SCOPED_CONDITIONAL_DRAW_EVENTF(RHICmdList, EventView, Views.Num() > 1, DEC_SCENE_ITEMS, TEXT("View%d"), ViewIndex);
+			SCOPED_CONDITIONAL_DRAW_EVENTF(RHICmdList, EventView, Views.Num() > 1, TEXT("View%d"), ViewIndex);
 
 			FViewInfo& View = Views[ViewIndex];
 
diff --git a/Engine/Source/Runtime/Renderer/Private/SceneVisibility.cpp b/Engine/Source/Runtime/Renderer/Private/SceneVisibility.cpp
--- a/Engine/Source/Runtime/Renderer/Private/SceneVisibility.cpp
+++ b/Engine/Source/Runtime/Renderer/Private/SceneVisibility.cpp
@@ -2191,7 +2191,7 @@ uint32 GetShadowQuality();
  */
 void FDeferredShadingSceneRenderer::InitViews(FRHICommandListImmediate& RHICmdList)
 {
-	SCOPED_DRAW_EVENT(RHICmdList, InitViews, DEC_SCENE_ITEMS);
+	SCOPED_DRAW_EVENT(RHICmdList, InitViews);
 
 	SCOPE_CYCLE_COUNTER(STAT_InitViewsTime);
 
diff --git a/Engine/Source/Runtime/Renderer/Private/ShaderComplexityRendering.cpp b/Engine/Source/Runtime/Renderer/Private/ShaderComplexityRendering.cpp
--- a/Engine/Source/Runtime/Renderer/Private/ShaderComplexityRendering.cpp
+++ b/Engine/Source/Runtime/Renderer/Private/ShaderComplexityRendering.cpp
@@ -150,7 +150,7 @@ extern TGlobalResource<FFilterVertexDeclaration> GFilterVertexDeclaration;
 
 void FRCPassPostProcessVisualizeComplexity::Process(FRenderingCompositePassContext& Context)
 {
-	SCOPED_DRAW_EVENT(Context.RHICmdList, PostProcessVisualizeComplexity, DEC_SCENE_ITEMS);
+	SCOPED_DRAW_EVENT(Context.RHICmdList, PostProcessVisualizeComplexity);
 	const FPooledRenderTargetDesc* InputDesc = GetInputDesc(ePId_Input0);
 
 	if(!InputDesc)
diff --git a/Engine/Source/Runtime/Renderer/Private/ShadowRendering.cpp b/Engine/Source/Runtime/Renderer/Private/ShadowRendering.cpp
--- a/Engine/Source/Runtime/Renderer/Private/ShadowRendering.cpp
+++ b/Engine/Source/Runtime/Renderer/Private/ShadowRendering.cpp
@@ -1626,7 +1626,7 @@ void FProjectedShadowInfo::RenderDepth(FRHICommandList& RHICmdList, FSceneRender
 #if WANTS_DRAW_MESH_EVENTS
 	FString EventName;
 	GetShadowTypeNameForDrawEvent(EventName);
-	SCOPED_DRAW_EVENTF(RHICmdList, EventShadowDepthActor, DEC_SCENE_ITEMS, *EventName);
+	SCOPED_DRAW_EVENTF(RHICmdList, EventShadowDepthActor, *EventName);
 #endif
 
 	CONDITIONAL_SCOPE_CYCLE_COUNTER(STAT_RenderWholeSceneShadowDepthsTime, bWholeSceneShadow);
@@ -1767,7 +1767,7 @@ void FProjectedShadowInfo::RenderProjection(FRHICommandListImmediate& RHICmdList
 #if WANTS_DRAW_MESH_EVENTS
 	FString EventName;
 	GetShadowTypeNameForDrawEvent(EventName);
-	SCOPED_DRAW_EVENTF(RHICmdList, EventShadowProjectionActor, DEC_SCENE_ITEMS, *EventName);
+	SCOPED_DRAW_EVENTF(RHICmdList, EventShadowProjectionActor, *EventName);
 #endif
 
 	FScopeCycleCounter Scope(bWholeSceneShadow ? GET_STATID(STAT_RenderWholeSceneShadowProjectionsTime) : GET_STATID(STAT_RenderPerObjectShadowProjectionsTime));
@@ -1870,7 +1870,7 @@ void FProjectedShadowInfo::RenderProjection(FRHICommandListImmediate& RHICmdList
 	// If this is a preshadow, mask the projection by the receiver primitives.
 	if (bPreShadow)
 	{
-		SCOPED_DRAW_EVENTF(RHICmdList, EventMaskSubjects, DEC_SCENE_ITEMS, TEXT("Stencil Mask Subjects"));
+		SCOPED_DRAW_EVENTF(RHICmdList, EventMaskSubjects, TEXT("Stencil Mask Subjects"));
 
 		// Set stencil to one.
 		// Note, this is a reversed Z depth surface, using CF_GreaterEqual.
@@ -2616,7 +2616,7 @@ bool FDeferredShadingSceneRenderer::RenderOnePassPointLightShadows(FRHICommandLi
 
 			if (!ProjectedShadowInfo->bRayTracedDistanceFieldShadow)
 			{
-				SCOPED_DRAW_EVENT(RHICmdList, ShadowDepthsFromOpaque, DEC_SCENE_ITEMS);
+				SCOPED_DRAW_EVENT(RHICmdList, ShadowDepthsFromOpaque);
 				auto SetShadowRenderTargets = MakeFunction<void, FRHICommandList&>([ProjectedShadowInfo](FRHICommandList& RHICmdList)
 				{
 					GSceneRenderTargets.BeginRenderingCubeShadowDepth(RHICmdList, ProjectedShadowInfo->ResolutionX);
@@ -2633,11 +2633,11 @@ bool FDeferredShadingSceneRenderer::RenderOnePassPointLightShadows(FRHICommandLi
 			{
 				GSceneRenderTargets.BeginRenderingLightAttenuation(RHICmdList);
 
-				SCOPED_DRAW_EVENT(RHICmdList, ShadowProjectionOnOpaque, DEC_SCENE_ITEMS);
+				SCOPED_DRAW_EVENT(RHICmdList, ShadowProjectionOnOpaque);
 
 				for (int32 ViewIndex = 0; ViewIndex < Views.Num(); ViewIndex++)
 				{
-					SCOPED_CONDITIONAL_DRAW_EVENTF(RHICmdList, EventView, Views.Num() > 1, DEC_SCENE_ITEMS, TEXT("View%d"), ViewIndex);
+					SCOPED_CONDITIONAL_DRAW_EVENTF(RHICmdList, EventView, Views.Num() > 1, TEXT("View%d"), ViewIndex);
 
 					const FViewInfo& View = Views[ViewIndex];
 
@@ -2660,7 +2660,7 @@ bool FDeferredShadingSceneRenderer::RenderOnePassPointLightShadows(FRHICommandLi
 			if (!LightSceneInfo->Proxy->HasStaticShadowing() && !ProjectedShadowInfo->bRayTracedDistanceFieldShadow)
 			{
 				bInjectedTranslucentVolume = true;
-				SCOPED_DRAW_EVENT(RHICmdList, InjectTranslucentVolume, DEC_SCENE_ITEMS);
+				SCOPED_DRAW_EVENT(RHICmdList, InjectTranslucentVolume);
 				// Inject the shadowed light into the translucency lighting volumes
 				InjectTranslucentVolumeLighting(RHICmdList, *LightSceneInfo, ProjectedShadowInfo);
 			}
@@ -2684,7 +2684,7 @@ void FDeferredShadingSceneRenderer::RenderProjections(
 
 	for (int32 ViewIndex = 0; ViewIndex < Views.Num(); ViewIndex++)
 	{
-		SCOPED_CONDITIONAL_DRAW_EVENTF(RHICmdList, EventView, Views.Num() > 1, DEC_SCENE_ITEMS, TEXT("View%d"), ViewIndex);
+		SCOPED_CONDITIONAL_DRAW_EVENTF(RHICmdList, EventView, Views.Num() > 1, TEXT("View%d"), ViewIndex);
 
 		const FViewInfo& View = Views[ViewIndex];
 
@@ -2861,7 +2861,7 @@ bool FDeferredShadingSceneRenderer::RenderTranslucentProjectedShadows(FRHIComman
 
 		{
 			// Render the translucent shadow depths.
-			SCOPED_DRAW_EVENT(RHICmdList, TranslucencyShadowMapGeneration, DEC_SCENE_ITEMS);
+			SCOPED_DRAW_EVENT(RHICmdList, TranslucencyShadowMapGeneration);
 
 			for (int32 ShadowIndex = 0; ShadowIndex < Shadows.Num(); ShadowIndex++)
 			{
@@ -2877,7 +2877,7 @@ bool FDeferredShadingSceneRenderer::RenderTranslucentProjectedShadows(FRHIComman
 		// Render the shadow projections.
 		{
 			{
-				SCOPED_DRAW_EVENT(RHICmdList, ShadowProjectionOnOpaque, DEC_SCENE_ITEMS);
+				SCOPED_DRAW_EVENT(RHICmdList, ShadowProjectionOnOpaque);
 				RenderProjections(RHICmdList, LightSceneInfo, Shadows);
 			}
 
@@ -3016,7 +3016,7 @@ bool FDeferredShadingSceneRenderer::RenderReflectiveShadowMaps(FRHICommandListIm
 
 		{
 			// Render the shadow depths.
-			SCOPED_DRAW_EVENT(RHICmdList, ReflectiveShadowMapsFromOpaque, DEC_SCENE_ITEMS);
+			SCOPED_DRAW_EVENT(RHICmdList, ReflectiveShadowMapsFromOpaque);
 
 			// render the RSMs
 			for (int32 ShadowIndex = 0; ShadowIndex < Shadows.Num(); ShadowIndex++)
@@ -3202,7 +3202,7 @@ bool FDeferredShadingSceneRenderer::RenderProjectedShadows(FRHICommandListImmedi
 
 		{
 			// Render the shadow depths.
-			SCOPED_DRAW_EVENT(RHICmdList, ShadowDepthsFromOpaque, DEC_SCENE_ITEMS);
+			SCOPED_DRAW_EVENT(RHICmdList, ShadowDepthsFromOpaque);
 
 			bool bFirst = true;
 			auto SetShadowRenderTargets = MakeFunction<void, FRHICommandList&>([&bFirst](FRHICommandList& RHICmdList)
@@ -3229,7 +3229,7 @@ bool FDeferredShadingSceneRenderer::RenderProjectedShadows(FRHICommandListImmedi
 		// Render the shadow projections.
 		{
 			{
-				SCOPED_DRAW_EVENT(RHICmdList, ShadowProjectionOnOpaque, DEC_SCENE_ITEMS);
+				SCOPED_DRAW_EVENT(RHICmdList, ShadowProjectionOnOpaque);
 				RenderProjections(RHICmdList, LightSceneInfo, Shadows);
 			}
 			
@@ -3247,7 +3247,7 @@ bool FDeferredShadingSceneRenderer::RenderProjectedShadows(FRHICommandListImmedi
 					&& (!LightSceneInfo->Proxy->HasStaticShadowing() || ProjectedShadowInfo->IsWholeSceneDirectionalShadow()))
 				{
 					bInjectedTranslucentVolume = true;
-					SCOPED_DRAW_EVENT(RHICmdList, InjectTranslucentVolume, DEC_SCENE_ITEMS);
+					SCOPED_DRAW_EVENT(RHICmdList, InjectTranslucentVolume);
 					// Inject the shadowed light into the translucency lighting volumes
 					InjectTranslucentVolumeLighting(RHICmdList, *LightSceneInfo, ProjectedShadowInfo);
 				}
@@ -3340,7 +3340,7 @@ bool FDeferredShadingSceneRenderer::RenderCachedPreshadows(FRHICommandListImmedi
 	{
 		if (bHasDepthsToRender)
 		{
-			SCOPED_DRAW_EVENTF(RHICmdList, EventShadowDepths, DEC_SCENE_ITEMS, TEXT("Preshadow Cache Depths"));
+			SCOPED_DRAW_EVENTF(RHICmdList, EventShadowDepths, TEXT("Preshadow Cache Depths"));
 
 			auto SetShadowRenderTargets = MakeFunction<void, FRHICommandList&>([](FRHICommandList& RHICmdList)
 			{
@@ -3381,7 +3381,7 @@ bool FDeferredShadingSceneRenderer::RenderCachedPreshadows(FRHICommandListImmedi
 
 		// Render the shadow projections.
 		{
-			SCOPED_DRAW_EVENTF(RHICmdList, EventShadowProj, DEC_SCENE_ITEMS, TEXT("Cached Preshadow Projections"));
+			SCOPED_DRAW_EVENTF(RHICmdList, EventShadowProj, TEXT("Cached Preshadow Projections"));
 			RenderProjections(RHICmdList, LightSceneInfo, OpaqueCachedPreshadows);
 
 			// Mark the attenuation buffer as dirty.
@@ -3463,7 +3463,7 @@ bool FForwardShadingSceneRenderer::RenderShadowDepthMap(FRHICommandListImmediate
 
 	{
 		// Render the shadow depths.
-		SCOPED_DRAW_EVENT(RHICmdList, ShadowDepthsFromOpaque, DEC_SCENE_ITEMS);
+		SCOPED_DRAW_EVENT(RHICmdList, ShadowDepthsFromOpaque);
 
 		bool bFirst = true;
 		auto SetShadowRenderTargets = MakeFunction<void, FRHICommandList&>([&bFirst](FRHICommandList& RHICmdList)
@@ -3493,7 +3493,7 @@ bool FForwardShadingSceneRenderer::RenderShadowDepthMap(FRHICommandListImmediate
 
 void FForwardShadingSceneRenderer::RenderShadowDepthMaps(FRHICommandListImmediate& RHICmdList)
 {
-	SCOPED_DRAW_EVENT(RHICmdList, Lights, DEC_SCENE_ITEMS);
+	SCOPED_DRAW_EVENT(RHICmdList, Lights);
 
 	bool bDynamicShadows = ViewFamily.EngineShowFlags.DynamicShadows && GetShadowQuality() > 0;
 	if (IsSimpleDynamicLightingEnabled() || !bDynamicShadows)
@@ -3528,7 +3528,7 @@ void FForwardShadingSceneRenderer::RenderShadowDepthMaps(FRHICommandListImmediat
 
 				FString LightNameWithLevel;
 				GetLightNameForDrawEvent(LightSceneInfo.Proxy, LightNameWithLevel);
-				SCOPED_DRAW_EVENTF(RHICmdList, EventLightPass, DEC_SCENE_ITEMS, *LightNameWithLevel);
+				SCOPED_DRAW_EVENTF(RHICmdList, EventLightPass, *LightNameWithLevel);
 				
 				INC_DWORD_STAT(STAT_NumShadowedLights);
 
diff --git a/Engine/Source/Runtime/Renderer/Private/TiledDeferredLightRendering.cpp b/Engine/Source/Runtime/Renderer/Private/TiledDeferredLightRendering.cpp
--- a/Engine/Source/Runtime/Renderer/Private/TiledDeferredLightRendering.cpp
+++ b/Engine/Source/Runtime/Renderer/Private/TiledDeferredLightRendering.cpp
@@ -351,7 +351,7 @@ void FDeferredShadingSceneRenderer::RenderTiledDeferredLighting(FRHICommandListI
 			}
 
 			{
-				SCOPED_DRAW_EVENT(RHICmdList, TiledDeferredLighting, DEC_SCENE_ITEMS);
+				SCOPED_DRAW_EVENT(RHICmdList, TiledDeferredLighting);
 
 				IPooledRenderTarget& InTexture = *GSceneRenderTargets.GetSceneColor();
 
diff --git a/Engine/Source/Runtime/Renderer/Private/TranslucentLighting.cpp b/Engine/Source/Runtime/Renderer/Private/TranslucentLighting.cpp
--- a/Engine/Source/Runtime/Renderer/Private/TranslucentLighting.cpp
+++ b/Engine/Source/Runtime/Renderer/Private/TranslucentLighting.cpp
@@ -529,7 +529,7 @@ void FProjectedShadowInfo::RenderTranslucencyDepths(FRHICommandList& RHICmdList,
 #if WANTS_DRAW_MESH_EVENTS
 		FString EventName;
 		GetShadowTypeNameForDrawEvent(EventName);
-		SCOPED_DRAW_EVENTF(RHICmdList, EventShadowDepthActor, DEC_SCENE_ITEMS, *EventName);
+		SCOPED_DRAW_EVENTF(RHICmdList, EventShadowDepthActor, *EventName);
 #endif
 
 		FTextureRHIParamRef RenderTargets[2] =
@@ -1115,7 +1115,7 @@ void FDeferredShadingSceneRenderer::ClearTranslucentVolumeLighting(FRHICommandLi
 {
 	if (GUseTranslucentLightingVolumes && GSupportsVolumeTextureRendering)
 	{
-		SCOPED_DRAW_EVENT(RHICmdList, ClearTranslucentVolumeLighting, DEC_SCENE_ITEMS);
+		SCOPED_DRAW_EVENT(RHICmdList, ClearTranslucentVolumeLighting);
 
 		// Clear all volume textures in the same draw with MRT, which is faster than individually
 
@@ -1187,7 +1187,7 @@ void FDeferredShadingSceneRenderer::InjectAmbientCubemapTranslucentVolumeLightin
 
 	if (GUseTranslucentLightingVolumes && View.FinalPostProcessSettings.ContributingCubemaps.Num() && !IsSimpleDynamicLightingEnabled())
 	{
-		SCOPED_DRAW_EVENT(RHICmdList, InjectAmbientCubemapTranslucentVolumeLighting, DEC_SCENE_ITEMS);
+		SCOPED_DRAW_EVENT(RHICmdList, InjectAmbientCubemapTranslucentVolumeLighting);
 
 		RHICmdList.SetRasterizerState(TStaticRasterizerState<FM_Solid, CM_None>::GetRHI());
 		RHICmdList.SetDepthStencilState(TStaticDepthStencilState<false, CF_Always>::GetRHI());
@@ -1233,7 +1233,7 @@ void FDeferredShadingSceneRenderer::ClearTranslucentVolumePerObjectShadowing(FRH
 {
 	if (GUseTranslucentLightingVolumes && GSupportsVolumeTextureRendering)
 	{
-		SCOPED_DRAW_EVENT(RHICmdList, ClearTranslucentVolumePerLightShadowing, DEC_SCENE_ITEMS);
+		SCOPED_DRAW_EVENT(RHICmdList, ClearTranslucentVolumePerLightShadowing);
 
 		static_assert(TVC_MAX == 2, "Only expecting two translucency lighting cascades.");
 		FTextureRHIParamRef RenderTargets[2];
@@ -1288,7 +1288,7 @@ void FDeferredShadingSceneRenderer::AccumulateTranslucentVolumeObjectShadowing(F
 
 	if (GUseTranslucentLightingVolumes && GSupportsVolumeTextureRendering)
 	{
-		SCOPED_DRAW_EVENT(RHICmdList, AccumulateTranslucentVolumeShadowing, DEC_SCENE_ITEMS);
+		SCOPED_DRAW_EVENT(RHICmdList, AccumulateTranslucentVolumeShadowing);
 
 		RHICmdList.SetRasterizerState(TStaticRasterizerState<FM_Solid, CM_None>::GetRHI());
 		RHICmdList.SetDepthStencilState(TStaticDepthStencilState<false, CF_Always>::GetRHI());
@@ -1775,7 +1775,7 @@ void FDeferredShadingSceneRenderer::FilterTranslucentVolumeLighting(FRHICommandL
 	{
 		if (GUseTranslucencyVolumeBlur)
 		{
-			SCOPED_DRAW_EVENT(RHICmdList, FilterTranslucentVolume, DEC_SCENE_ITEMS);
+			SCOPED_DRAW_EVENT(RHICmdList, FilterTranslucentVolume);
 
 			RHICmdList.SetRasterizerState(TStaticRasterizerState<FM_Solid, CM_None>::GetRHI());
 			RHICmdList.SetDepthStencilState(TStaticDepthStencilState<false, CF_Always>::GetRHI());
diff --git a/Engine/Source/Runtime/Renderer/Private/TranslucentRendering.cpp b/Engine/Source/Runtime/Renderer/Private/TranslucentRendering.cpp
--- a/Engine/Source/Runtime/Renderer/Private/TranslucentRendering.cpp
+++ b/Engine/Source/Runtime/Renderer/Private/TranslucentRendering.cpp
@@ -171,7 +171,7 @@ FGlobalBoundShaderState CopySceneColorBoundShaderState;
 void FTranslucencyDrawingPolicyFactory::CopySceneColor(FRHICommandList& RHICmdList, const FViewInfo& View, const FPrimitiveSceneProxy* PrimitiveSceneProxy)
 {
 
-	SCOPED_DRAW_EVENTF(RHICmdList, EventCopy, DEC_SCENE_ITEMS, TEXT("CopySceneColor for %s %s"), *PrimitiveSceneProxy->GetOwnerName().ToString(), *PrimitiveSceneProxy->GetResourceName().ToString());
+	SCOPED_DRAW_EVENTF(RHICmdList, EventCopy, TEXT("CopySceneColor for %s %s"), *PrimitiveSceneProxy->GetOwnerName().ToString(), *PrimitiveSceneProxy->GetResourceName().ToString());
 	RHICmdList.SetRasterizerState(TStaticRasterizerState<FM_Solid, CM_None>::GetRHI());
 	RHICmdList.SetDepthStencilState(TStaticDepthStencilState<false, CF_Always>::GetRHI());
 	RHICmdList.SetBlendState(TStaticBlendState<>::GetRHI());
@@ -900,7 +900,7 @@ void FDeferredShadingSceneRenderer::RenderTranslucencyParallel(FRHICommandListIm
 
 	for (int32 ViewIndex = 0; ViewIndex < Views.Num(); ViewIndex++)
 	{
-		SCOPED_CONDITIONAL_DRAW_EVENTF(RHICmdList, EventView, Views.Num() > 1, DEC_SCENE_ITEMS, TEXT("View%d"), ViewIndex);
+		SCOPED_CONDITIONAL_DRAW_EVENTF(RHICmdList, EventView, Views.Num() > 1, TEXT("View%d"), ViewIndex);
 
 		const FViewInfo& View = Views[ViewIndex];
 
@@ -1008,7 +1008,7 @@ void FDeferredShadingSceneRenderer::RenderTranslucency(FRHICommandListImmediate&
 	if (ShouldRenderTranslucency())
 	{
 
-		SCOPED_DRAW_EVENT(RHICmdList, Translucency, DEC_SCENE_ITEMS);
+		SCOPED_DRAW_EVENT(RHICmdList, Translucency);
 
 		if (GRHICommandList.UseParallelAlgorithms() && CVarParallelTranslucency.GetValueOnRenderThread())
 		{
@@ -1018,7 +1018,7 @@ void FDeferredShadingSceneRenderer::RenderTranslucency(FRHICommandListImmediate&
 
 		for (int32 ViewIndex = 0; ViewIndex < Views.Num(); ViewIndex++)
 		{
-			SCOPED_CONDITIONAL_DRAW_EVENTF(RHICmdList, EventView, Views.Num() > 1, DEC_SCENE_ITEMS, TEXT("View%d"), ViewIndex);
+			SCOPED_CONDITIONAL_DRAW_EVENTF(RHICmdList, EventView, Views.Num() > 1, TEXT("View%d"), ViewIndex);
 
 			const FViewInfo& View = Views[ViewIndex];
 
diff --git a/Engine/Source/Runtime/Renderer/Private/VelocityRendering.cpp b/Engine/Source/Runtime/Renderer/Private/VelocityRendering.cpp
--- a/Engine/Source/Runtime/Renderer/Private/VelocityRendering.cpp
+++ b/Engine/Source/Runtime/Renderer/Private/VelocityRendering.cpp
@@ -670,7 +670,7 @@ void FDeferredShadingSceneRenderer::RenderVelocities(FRHICommandListImmediate& R
 		return;
 	}
 
-	SCOPED_DRAW_EVENT(RHICmdList, RenderVelocities, DEC_SCENE_ITEMS);
+	SCOPED_DRAW_EVENT(RHICmdList, RenderVelocities);
 
 	const FIntPoint BufferSize = GSceneRenderTargets.GetBufferSizeXY();
 	const FIntPoint VelocityBufferSize = BufferSize;		// full resolution so we can reuse the existing full res z buffer
diff --git a/Engine/Source/Runtime/SlateRHIRenderer/Private/SlateRHIRenderer.cpp b/Engine/Source/Runtime/SlateRHIRenderer/Private/SlateRHIRenderer.cpp
--- a/Engine/Source/Runtime/SlateRHIRenderer/Private/SlateRHIRenderer.cpp
+++ b/Engine/Source/Runtime/SlateRHIRenderer/Private/SlateRHIRenderer.cpp
@@ -388,7 +388,7 @@ void FSlateRHIRenderer::OnWindowDestroyed( const TSharedRef<SWindow>& InWindow )
 /** Draws windows from a FSlateDrawBuffer on the render thread */
 void FSlateRHIRenderer::DrawWindow_RenderThread(FRHICommandListImmediate& RHICmdList, const FViewportInfo& ViewportInfo, const FSlateWindowElementList& WindowElementList, bool bLockToVsync)
 {
-	SCOPED_DRAW_EVENT(RHICmdList, SlateUI, DEC_SCENE_ITEMS);
+	SCOPED_DRAW_EVENT(RHICmdList, SlateUI);
 
 	// Should only be called by the rendering thread
 	check(IsInRenderingThread());
@@ -1213,6 +1213,6 @@ void FSlateRHIRenderer::SetWindowRenderTarget(const SWindow& Window, FTexture2DR
 	FViewportInfo* ViewInfo = WindowToViewportInfo.FindRef(&Window);
 	if (ViewInfo)
 	{
- 		ViewInfo->RenderTargetTexture = RT;
- 	}
+		ViewInfo->RenderTargetTexture = RT;
+	}
 }
