added histogram visuzalization ot HDRVisualize for debugging (disabled by default)

diff --git a/Engine/Shaders/PostProcessHistogram.usf b/Engine/Shaders/PostProcessHistogram.usf
--- a/Engine/Shaders/PostProcessHistogram.usf
+++ b/Engine/Shaders/PostProcessHistogram.usf
@@ -26,12 +26,12 @@ groupshared float SharedHistogram[HISTOGRAM_SIZE][THREADGROUP_SIZEX][THREADGROUP
 void MainCS(
 	uint3 GroupId : SV_GroupID,
 	uint3 DispatchThreadId : SV_DispatchThreadID,	// DispatchThreadId = GroupId * int2(dimx,dimy) + GroupThreadId
-    uint3 GroupThreadId : SV_GroupThreadID,
+    uint3 GroupThreadId : SV_GroupThreadID,	// 0..THREADGROUP_SIZEX 0..THREADGROUP_SIZEY
 	uint GroupIndex: SV_GroupIndex) // SV_GroupIndex = SV_GroupThreadID.z*dimx*dimy + SV_GroupThreadID.y*dimx + SV_GroupThreadID.x
 {
 	// todo: can be cleared more efficiently
 	// clear all THREADGROUP_SIZEX*THREADGROUP_SIZEY histograms
-	LOOP for(uint i = 0; i < HISTOGRAM_SIZE; ++i)
+	UNROLL for(uint i = 0; i < HISTOGRAM_SIZE; ++i)
 	{
 		SharedHistogram[i][GroupThreadId.x][GroupThreadId.y] = 0.0f;
 	}
@@ -88,8 +88,10 @@ void MainCS(
 
 	// accumulate all histograms into one
 
-	float4 Sum = 0;
+	if(GroupIndex < HISTOGRAM_SIZE / 4)
 	{
+		float4 Sum = 0;
+
 		LOOP for(uint y = 0; y < THREADGROUP_SIZEY; ++y)
 		{
 			LOOP for(uint x = 0; x < THREADGROUP_SIZEX; ++x)
@@ -101,16 +103,17 @@ void MainCS(
 					SharedHistogram[GroupIndex * 4 + 3][x][y]);
 			}
 		}
-	}
 
-	float2 MaxExtent = float2(THREADGROUP_SIZEX * LOOP_SIZEX, THREADGROUP_SIZEY * LOOP_SIZEY);
-//	float2 Extent = min(MaxExtent, HistogramParameters.xy - LeftTop);
-//	float Area = Extent.x * Extent.y;
-	float Area = MaxExtent.x * MaxExtent.y;
+		float2 MaxExtent = float2(THREADGROUP_SIZEX * LOOP_SIZEX, THREADGROUP_SIZEY * LOOP_SIZEY);
+	//	float2 Extent = min(MaxExtent, HistogramParameters.xy - LeftTop);
+	//	float Area = Extent.x * Extent.y;
+		float Area = MaxExtent.x * MaxExtent.y;
 	
-	// doesn't take borders into account but the error should be minor
-	float NormalizeFactor = 1.0f / Area;
+		// doesn't take borders into account but the error should be minor
+		float NormalizeFactor = 1.0f / Area;
 
-	// output texture with one histogram per line, x and y unwrapped into all the lines
-	HistogramRWTexture[uint2(GroupIndex, GroupId.x + GroupId.y * ThreadGroupCount.x)] = Sum * NormalizeFactor;
+		// output texture with one histogram per line, x and y unwrapped into all the lines
+		HistogramRWTexture[uint2(GroupIndex, GroupId.x + GroupId.y * ThreadGroupCount.x)] = Sum * NormalizeFactor;
+//		HistogramRWTexture[uint2(GroupIndex, GroupId.x + GroupId.y * ThreadGroupCount.x)] = 0.5f;
+	}
 }
diff --git a/Engine/Shaders/PostProcessSubsurface.usf b/Engine/Shaders/PostProcessSubsurface.usf
--- a/Engine/Shaders/PostProcessSubsurface.usf
+++ b/Engine/Shaders/PostProcessSubsurface.usf
@@ -36,9 +36,6 @@ float GetSubsurfaceStrength(float2 UV)
 	return Mask * ScreenSpaceData.GBuffer.Opacity;
 }
 
-//#define SSSS_STREGTH_SOURCE 1	// todo: masking fom LightingModel or SSSColor
-#define SSSS_STREGTH_SOURCE(UV) GetSubsurfaceStrength(UV)
-
 // 0:speculars leak into SSS / 1:requires alpha channel for scenecolor
 // #define SSSS_SPECULAR_CORRECTION 1
 
@@ -264,7 +261,6 @@ void SetupPS(in float4 UVAndScreenPos : TEXCOORD0, out float4 OutColor : SV_Targ
 	
 	FScreenSpaceData ScreenSpaceData = GetScreenSpaceData(UV);
 
-//	finalStep *= SSSS_STREGTH_SOURCE(UV); // Modulate it using the alpha channel.
 	finalStep *= GetProfileRadiusScale(ScreenSpaceData.GBuffer);
 
 	// in full resolution
@@ -371,7 +367,7 @@ void SubsurfaceRecombinePS(float4 UVAndScreenPos : TEXCOORD0, out float4 OutColo
 
 		// Calculate the final step to fetch the surrounding pixels:
 		float finalStep = scale * HorizontalScaler;
-//		finalStep *= SSSS_STREGTH_SOURCE(UV); // Modulate it using the alpha channel (0..1), might only make a difference in very few cases
+
 		finalStep *= GetProfileRadiusScale(ScreenSpaceData.GBuffer);
 
 		float PixelSizeRadius = finalStep / (PostprocessInput1Size.z * 0.5f);
diff --git a/Engine/Shaders/PostProcessVisualizeHDR.usf b/Engine/Shaders/PostProcessVisualizeHDR.usf
--- a/Engine/Shaders/PostProcessVisualizeHDR.usf
+++ b/Engine/Shaders/PostProcessVisualizeHDR.usf
@@ -11,6 +11,9 @@
 #include "TonemapCommon.usf"
 #include "MiniFontCommon.usf"		// for PrintFloat()
 
+// .xy:GatherExtent.xy, .zw:TexelPerThreadGroupXY
+uint4 HistogramParams;
+
 // only needed for nice visualization
 float ComputeHistogramMax(Texture2D HistogramTexture)
 {
@@ -86,16 +89,62 @@ void HighlightAdvice(inout float3 OutColor, float2 UV, int2 PixelPos)
 	}
 }
 
+bool InUnitBox(float2 UV)
+{
+	return UV.x >= 0 && UV.y >= 0 && UV.y < 1 && UV.y < 1;
+}
 
 //
 void MainPS(float4 UVAndScreenPos : TEXCOORD0, out float4 OutColor : SV_Target0)
 {
 	float2 UV = UVAndScreenPos.xy;
- 	int2 PixelPos = (int2)(UVAndScreenPos.zw * ScreenPosToPixel.xy + ScreenPosToPixel.zw + 0.5f);
-	
+	int2 PixelPos = (int2)(UVAndScreenPos.zw * ScreenPosToPixel.xy + ScreenPosToPixel.zw + 0.5f);
+	float2 ViewLocalUV = float2(UVAndScreenPos.z * 0.5f + 0.5f, 0.5f - 0.5f * UVAndScreenPos.w);
+
 	// background is the scene color
 	OutColor = Texture2DSample(PostprocessInput0, PostprocessInput0Sampler, UV);
 
+	float2 IDAreaLocalUV = ViewLocalUV * 2 + float2(-1, 0);
+
+	// disabled for now, useful for debgugging the histogram content
+	if(0)
+	BRANCH if(InUnitBox(IDAreaLocalUV))
+	{
+		float2 HistogramUV = frac(IDAreaLocalUV * HistogramParams.xy / HistogramParams.zw);
+
+		int2 HistogramXY = int2(IDAreaLocalUV * HistogramParams.xy / HistogramParams.zw);
+
+		bool bChecker = dot(HistogramXY, 1) % 2;
+	
+		int HistogramId = HistogramXY.x + HistogramXY.y * ceil(HistogramParams.x / (float)HistogramParams.z);
+
+		float2 UnrolledHistogramUV = float2(HistogramUV.x, (HistogramId + 0.5f) * PostprocessInput3Size.w);
+
+		float4 HistogramRGBA = Texture2DSampleLevel(PostprocessInput3, PostprocessInput3Sampler, UnrolledHistogramUV, 0);
+
+		// 16 buckets are good enough for the small space there
+		float Value = max(max(HistogramRGBA.r,HistogramRGBA.g), max(HistogramRGBA.b,HistogramRGBA.a));
+
+		// we want the histogram to be 0 at the bottom
+		float HistogramY = 1 - HistogramUV.y;
+		// /2 as we rea covering a quarter of the screen
+		float HistogramPixelHeight = ViewportSize.y / 2 / HistogramXY.y; 
+
+		// checked overlay
+		OutColor = bChecker ? float4(0.2f,0,0,0) : float4(0,0.15f,0,0);
+		
+		// SceneColor overlay
+		OutColor += 0.1f * Texture2DSampleLevel(PostprocessInput0, PostprocessInput0Sampler, IDAreaLocalUV * View.ViewSizeAndSceneTexelSize.xy * View.ViewSizeAndSceneTexelSize.zw, 0);
+
+		// could be the same as (HistogramY < Value) but a bit more antialiased
+//		float HistogramMask = HistogramY < Value;
+		float HistogramMask = saturate((Value - HistogramY) * HistogramPixelHeight);
+
+		OutColor = lerp(OutColor, 1, HistogramMask);
+	}
+
+
+
 	// not fully functional yet
 //	HighlightAdvice(OutColor.rgb, UV, PixelPos);
 
diff --git a/Engine/Source/Runtime/Renderer/Private/PostProcess/PostProcessHistogram.cpp b/Engine/Source/Runtime/Renderer/Private/PostProcess/PostProcessHistogram.cpp
--- a/Engine/Source/Runtime/Renderer/Private/PostProcess/PostProcessHistogram.cpp
+++ b/Engine/Source/Runtime/Renderer/Private/PostProcess/PostProcessHistogram.cpp
@@ -118,9 +118,7 @@ void FRCPassPostProcessHistogram::Process(FRenderingCompositePassContext& Contex
 	check(DestRenderTarget.UAV);
 	Context.RHICmdList.SetUAVParameter(ComputeShader->GetComputeShader(), ComputeShader->HistogramRWTexture.GetBaseIndex(), DestRenderTarget.UAV);
 
-	// we currently assume the input is half res, one full res pixel less to avoid getting bilinear filtered input
-	FIntPoint GatherExtent = (DestRect.Size() - FIntPoint(1, 1)) / 2;
-
+	FIntPoint GatherExtent = ComputeGatherExtent(View);
 	FIntPoint ThreadGroupCountValue = ComputeThreadGroupCount(GatherExtent);
 
 	ComputeShader->SetCS(Context.RHICmdList, Context, ThreadGroupCountValue, (DestRect.Min + FIntPoint(1, 1)) / 2, GatherExtent);
@@ -133,6 +131,13 @@ void FRCPassPostProcessHistogram::Process(FRenderingCompositePassContext& Contex
 	Context.RHICmdList.CopyToResolveTarget(DestRenderTarget.TargetableTexture, DestRenderTarget.ShaderResourceTexture, false, FResolveParams());
 }
 
+
+FIntPoint FRCPassPostProcessHistogram::ComputeGatherExtent(const FSceneView& View)
+{
+	// we currently assume the input is half res, one full res pixel less to avoid getting bilinear filtered input
+	return (View.ViewRect.Size() - FIntPoint(1, 1)) / 2;
+}
+
 FIntPoint FRCPassPostProcessHistogram::ComputeThreadGroupCount(FIntPoint PixelExtent)
 {
 	uint32 TexelPerThreadGroupX = ThreadGroupSizeX * LoopCountX;
diff --git a/Engine/Source/Runtime/Renderer/Private/PostProcess/PostProcessHistogram.h b/Engine/Source/Runtime/Renderer/Private/PostProcess/PostProcessHistogram.h
--- a/Engine/Source/Runtime/Renderer/Private/PostProcess/PostProcessHistogram.h
+++ b/Engine/Source/Runtime/Renderer/Private/PostProcess/PostProcessHistogram.h
@@ -33,5 +33,6 @@ public:
 
 	// -------------------------------------------
 
+	static FIntPoint ComputeGatherExtent(const FSceneView& View);
 	static FIntPoint ComputeThreadGroupCount(FIntPoint PixelExtent);
 };
diff --git a/Engine/Source/Runtime/Renderer/Private/PostProcess/PostProcessVisualizeHDR.cpp b/Engine/Source/Runtime/Renderer/Private/PostProcess/PostProcessVisualizeHDR.cpp
--- a/Engine/Source/Runtime/Renderer/Private/PostProcess/PostProcessVisualizeHDR.cpp
+++ b/Engine/Source/Runtime/Renderer/Private/PostProcess/PostProcessVisualizeHDR.cpp
@@ -41,6 +41,7 @@ public:
 	FShaderParameter EyeAdaptationParams;
 	FShaderResourceParameter MiniFontTexture;
 	FShaderParameter InverseGamma;
+	FShaderParameter HistogramParams;
 
 	FShaderParameter ColorMatrixR_ColorCurveCd1;
 	FShaderParameter ColorMatrixG_ColorCurveCd3Cm3;
@@ -59,6 +60,7 @@ public:
 		EyeAdaptationParams.Bind(Initializer.ParameterMap, TEXT("EyeAdaptationParams"));
 		MiniFontTexture.Bind(Initializer.ParameterMap, TEXT("MiniFontTexture"));
 		InverseGamma.Bind(Initializer.ParameterMap,TEXT("InverseGamma"));
+		HistogramParams.Bind(Initializer.ParameterMap,TEXT("HistogramParams"));
 
 		ColorMatrixR_ColorCurveCd1.Bind(Initializer.ParameterMap, TEXT("ColorMatrixR_ColorCurveCd1"));
 		ColorMatrixG_ColorCurveCd3Cm3.Bind(Initializer.ParameterMap, TEXT("ColorMatrixG_ColorCurveCd3Cm3"));
@@ -90,6 +92,17 @@ public:
 
 		SetTextureParameter(Context.RHICmdList, ShaderRHI, MiniFontTexture, GEngine->MiniFontTexture ? GEngine->MiniFontTexture->Resource->TextureRHI : GSystemTextures.WhiteDummy->GetRenderTargetItem().TargetableTexture);
 
+		{
+			FIntPoint GatherExtent = FRCPassPostProcessHistogram::ComputeGatherExtent(Context.View);
+
+			uint32 TexelPerThreadGroupX = FRCPassPostProcessHistogram::ThreadGroupSizeX * FRCPassPostProcessHistogram::LoopCountX;
+			uint32 TexelPerThreadGroupY = FRCPassPostProcessHistogram::ThreadGroupSizeY * FRCPassPostProcessHistogram::LoopCountY;
+
+			FIntRect Value(GatherExtent, FIntPoint(TexelPerThreadGroupX, TexelPerThreadGroupY));
+
+			SetShaderValue(Context.RHICmdList, ShaderRHI, HistogramParams, Value);
+		}
+
 		{
 			float InvDisplayGammaValue = 1.0f / ViewFamily.RenderTarget->GetDisplayGamma();
 
@@ -114,7 +127,7 @@ public:
 	virtual bool Serialize(FArchive& Ar)
 	{
 		bool bShaderHasOutdatedParameters = FGlobalShader::Serialize(Ar);
-		Ar << PostprocessParameter << EyeAdaptationParams << MiniFontTexture << InverseGamma 
+		Ar << PostprocessParameter << EyeAdaptationParams << MiniFontTexture << InverseGamma << HistogramParams
 			<< ColorMatrixR_ColorCurveCd1 << ColorMatrixG_ColorCurveCd3Cm3 << ColorMatrixB_ColorCurveCm2 << ColorCurve_Cm0Cd0_Cd2_Ch0Cm1_Ch3 << ColorCurve_Ch1_Ch2 << ColorShadow_Luma << ColorShadow_Tint1 << ColorShadow_Tint2;
 		return bShaderHasOutdatedParameters;
 	}
diff --git a/Engine/Source/Runtime/Renderer/Private/PostProcess/PostProcessVisualizeHDR.h b/Engine/Source/Runtime/Renderer/Private/PostProcess/PostProcessVisualizeHDR.h
--- a/Engine/Source/Runtime/Renderer/Private/PostProcess/PostProcessVisualizeHDR.h
+++ b/Engine/Source/Runtime/Renderer/Private/PostProcess/PostProcessVisualizeHDR.h
@@ -12,7 +12,8 @@
 // ePId_Input0: LDR SceneColor
 // ePId_Input1: output of the Histogram pass
 // ePId_Input2: HDR SceneColor
-class FRCPassPostProcessVisualizeHDR : public TRenderingCompositePassBase<3, 1>
+// ePId_Input3: output of the Histogram pass over screen (not reduced yet)
+class FRCPassPostProcessVisualizeHDR : public TRenderingCompositePassBase<4, 1>
 {
 public:
 	// interface FRenderingCompositePass ---------
diff --git a/Engine/Source/Runtime/Renderer/Private/PostProcess/PostProcessing.cpp b/Engine/Source/Runtime/Renderer/Private/PostProcess/PostProcessing.cpp
--- a/Engine/Source/Runtime/Renderer/Private/PostProcess/PostProcessing.cpp
+++ b/Engine/Source/Runtime/Renderer/Private/PostProcess/PostProcessing.cpp
@@ -342,19 +342,6 @@ static void AddPostProcessDepthOfFieldGaussian(FPostprocessContext& Context, FDe
 	Context.FinalOutput = FRenderingCompositeOutputRef(NodeRecombined);
 }
 
-static FRenderingCompositeOutputRef AddPostProcessHistogram(FPostprocessContext& Context, FRenderingCompositeOutputRef& HalfResScene)
-{
-	FRenderingCompositePass* NodeHistogram = Context.Graph.RegisterPass(new(FMemStack::Get()) FRCPassPostProcessHistogram());
-
-	NodeHistogram->SetInput(ePId_Input0, HalfResScene);
-
-	FRenderingCompositePass* NodeHistogramReduce = Context.Graph.RegisterPass(new(FMemStack::Get()) FRCPassPostProcessHistogramReduce());
-
-	NodeHistogramReduce->SetInput(ePId_Input0, NodeHistogram);
-
-	return FRenderingCompositeOutputRef(NodeHistogramReduce);
-}
-
 static FRenderingCompositeOutputRef AddBloom(FPostprocessContext Context, FRenderingCompositeOutputRef PostProcessDownsample0)
 {
 	// Quality level to bloom stages table. Note: 0 is omitted, ensure element count tallys with the range documented with 'r.BloomQuality' definition.
@@ -797,6 +784,8 @@ void FPostProcessing::Process(FRHICommandListImmediate& RHICmdList, FViewInfo& V
 		// not always valid
 		FRenderingCompositeOutputRef HDRColor; 
 		// not always valid
+		FRenderingCompositeOutputRef HistogramOverScreen;
+		// not always valid
 		FRenderingCompositeOutputRef Histogram;
 		// not always valid
 		FRenderingCompositeOutputRef EyeAdaptation;
@@ -988,7 +977,17 @@ void FPostProcessing::Process(FRHICommandListImmediate& RHICmdList, FViewInfo& V
 
 				if (!GIsHighResScreenshot && bHistogramNeeded && FeatureLevel >= ERHIFeatureLevel::SM5 && StereoPass != eSSP_RIGHT_EYE)
 				{
-					Histogram = AddPostProcessHistogram(Context, SceneColorHalfRes);
+					FRenderingCompositePass* NodeHistogram = Context.Graph.RegisterPass(new(FMemStack::Get()) FRCPassPostProcessHistogram());
+
+					NodeHistogram->SetInput(ePId_Input0, SceneColorHalfRes);
+
+					HistogramOverScreen = FRenderingCompositeOutputRef(NodeHistogram);
+
+					FRenderingCompositePass* NodeHistogramReduce = Context.Graph.RegisterPass(new(FMemStack::Get()) FRCPassPostProcessHistogramReduce());
+
+					NodeHistogramReduce->SetInput(ePId_Input0, NodeHistogram);
+
+					Histogram = FRenderingCompositeOutputRef(NodeHistogramReduce);
 				}
 			}
 
@@ -1235,6 +1234,7 @@ void FPostProcessing::Process(FRHICommandListImmediate& RHICmdList, FViewInfo& V
 			Node->SetInput(ePId_Input0, FRenderingCompositeOutputRef(Context.FinalOutput));
 			Node->SetInput(ePId_Input1, Histogram);
 			Node->SetInput(ePId_Input2, HDRColor);
+			Node->SetInput(ePId_Input3, HistogramOverScreen);
 			Context.FinalOutput = FRenderingCompositeOutputRef(Node);
 		}
 
