VisualizeHDR now shows Histogram bucket as color

diff --git a/Engine/Shaders/PostProcessVisualizeHDR.usf b/Engine/Shaders/PostProcessVisualizeHDR.usf
--- a/Engine/Shaders/PostProcessVisualizeHDR.usf
+++ b/Engine/Shaders/PostProcessVisualizeHDR.usf
@@ -117,7 +117,13 @@ void MainPS(float4 UVAndScreenPos : TEXCOORD0, out float4 OutColor : SV_Target0)
 	float2 ViewLocalUV = float2(UVAndScreenPos.z * 0.5f + 0.5f, 0.5f - 0.5f * UVAndScreenPos.w);
 
 	// background is the scene color
-	OutColor = Texture2DSample(PostprocessInput0, PostprocessInput0Sampler, UV);
+	float4 SceneColor = Texture2DSample(PostprocessInput0, PostprocessInput0Sampler, UV);
+
+//	OutColor = SceneColor;
+	float LuminanceVal = max(SceneColor.r, max(SceneColor.g, SceneColor.b));
+	OutColor = float4(Colorize(ComputeHistogramPositionFromLuminance(LuminanceVal)), 1.0f);
+
+
 
 	float2 IDAreaLocalUV = ViewLocalUV * 2 + float2(-1, 0);
 
diff --git a/Engine/Source/Runtime/Renderer/Private/PostProcess/PostProcessing.cpp b/Engine/Source/Runtime/Renderer/Private/PostProcess/PostProcessing.cpp
--- a/Engine/Source/Runtime/Renderer/Private/PostProcess/PostProcessing.cpp
+++ b/Engine/Source/Runtime/Renderer/Private/PostProcess/PostProcessing.cpp
@@ -809,7 +809,14 @@ void FPostProcessing::Process(FRHICommandListImmediate& RHICmdList, FViewInfo& V
 				check(!ViewState->SeparateTranslucencyRT);
 			}
 		}
-		
+
+		bool bVisualizeHDR = View.Family->EngineShowFlags.VisualizeHDR && FeatureLevel >= ERHIFeatureLevel::SM5;
+
+		if(bVisualizeHDR)
+		{
+			bAllowTonemapper = false;
+		}
+
 		// add the passes we want to add to the graph (commenting a line means the pass is not inserted into the graph) ---------
 
 		if (AllowFullPostProcessing(View, FeatureLevel))
@@ -961,7 +968,7 @@ void FPostProcessing::Process(FRHICommandListImmediate& RHICmdList, FViewInfo& V
 			}
 
 			{
-				bool bHistogramNeeded = View.Family->EngineShowFlags.VisualizeHDR;
+				bool bHistogramNeeded = false;
 
 				if(View.Family->EngineShowFlags.EyeAdaptation
 					&& View.FinalPostProcessSettings.AutoExposureMinBrightness < View.FinalPostProcessSettings.AutoExposureMaxBrightness
@@ -975,6 +982,11 @@ void FPostProcessing::Process(FRHICommandListImmediate& RHICmdList, FViewInfo& V
 					bHistogramNeeded = false;
 				}
 
+				if(View.Family->EngineShowFlags.VisualizeHDR)
+				{
+					bHistogramNeeded = true;
+				}
+
 				if (!GIsHighResScreenshot && bHistogramNeeded && FeatureLevel >= ERHIFeatureLevel::SM5 && StereoPass != eSSP_RIGHT_EYE)
 				{
 					FRenderingCompositePass* NodeHistogram = Context.Graph.RegisterPass(new(FMemStack::Get()) FRCPassPostProcessHistogram());
@@ -1228,7 +1240,7 @@ void FPostProcessing::Process(FRHICommandListImmediate& RHICmdList, FViewInfo& V
 			}
 		}
 
-		if(View.Family->EngineShowFlags.VisualizeHDR && FeatureLevel >= ERHIFeatureLevel::SM5)
+		if(bVisualizeHDR)
 		{
 			FRenderingCompositePass* Node = Context.Graph.RegisterPass(new(FMemStack::Get()) FRCPassPostProcessVisualizeHDR());
 			Node->SetInput(ePId_Input0, FRenderingCompositeOutputRef(Context.FinalOutput));
