UE4 - Fix mobile emulation crash at boot

diff --git a/Engine/Source/Runtime/Renderer/Private/ForwardBasePassRendering.cpp b/Engine/Source/Runtime/Renderer/Private/ForwardBasePassRendering.cpp
--- a/Engine/Source/Runtime/Renderer/Private/ForwardBasePassRendering.cpp
+++ b/Engine/Source/Runtime/Renderer/Private/ForwardBasePassRendering.cpp
@@ -143,8 +143,9 @@ public:
 	}
 
 	const FLightSceneInfo* GetSimpleDirectionalLight() const 
-	{ 
-		return ((FScene*)View.Family->Scene)->SimpleDirectionalLight;
+	{
+		auto* Scene = (FScene*)View.Family->Scene;
+		return Scene ? Scene->SimpleDirectionalLight : nullptr;
 	}
 
 	/** Initialization constructor. */
