fixed SubsurfaceSSS for splitscreen (minor extra cost even for non splitscreen but that can be optimized later)

diff --git a/Engine/Shaders/PostProcessSubsurface.usf b/Engine/Shaders/PostProcessSubsurface.usf
--- a/Engine/Shaders/PostProcessSubsurface.usf
+++ b/Engine/Shaders/PostProcessSubsurface.usf
@@ -39,7 +39,7 @@ float GetSubsurfaceStrength(float2 UV)
 // 0:speculars leak into SSS / 1:requires alpha channel for scenecolor
 // #define SSSS_SPECULAR_CORRECTION 1
 
-// set by C++ for SubsurfaceRecombinePS
+// set by C++ for SubsurfaceRecombinePS and ExtractSpecularPS
 // #define SSS_RECOMBINE_METHOD 0:without reconstruct specular/1:with reconstruct specular
 
 // set by C++ for SubsurfacePS
@@ -291,7 +291,7 @@ void SetupPS(in float4 UVAndScreenPos : TEXCOORD0, out float4 OutColor : SV_Targ
 // input0 is created by the SetupPS shader
 void MainPS(float4 UVAndScreenPos : TEXCOORD0, out float4 OutColor : SV_Target0)
 {
-	float2 UV = UVAndScreenPos.xy;
+	float2 ViewportUV = UVAndScreenPos.xy;
 
 	// call into "SeparableSSS.usf"
 
@@ -306,7 +306,10 @@ void MainPS(float4 UVAndScreenPos : TEXCOORD0, out float4 OutColor : SV_Target0)
 	float2 Direction = float2(0, 1) * View.ViewSizeAndSceneTexelSize.w * View.ViewSizeAndSceneTexelSize.x * 1024.0f;
 #endif
 
-	OutColor = SSSSBlurPS(UV, Direction, false);
+	// can be optimized
+	float2 GBufferUV = (ViewportUV * PostprocessInput0Size.xy * 2 + View.ViewRectMin.xy) * View.ViewSizeAndSceneTexelSize.zw;
+
+	OutColor = SSSSBlurPS(GBufferUV, ViewportUV, Direction, false);
 
 #if SSS_DIRECTION == 1
 	// second pass prepares the setup fro the recombine pass which doesn't need depth but wants to reconstruct the color
@@ -321,7 +324,7 @@ void MainPS(float4 UVAndScreenPos : TEXCOORD0, out float4 OutColor : SV_Target0)
 // with the SceneColor.
 void SubsurfaceRecombinePS(float4 UVAndScreenPos : TEXCOORD0, out float4 OutColor : SV_Target0)
 {
-	float2 UV = UVAndScreenPos.xy;
+	float2 ViewportUV = UVAndScreenPos.xy;
 
 	// recombine with the scene color
 
@@ -336,7 +339,7 @@ void SubsurfaceRecombinePS(float4 UVAndScreenPos : TEXCOORD0, out float4 OutColo
 	// bilinear filtering
 	float3 SSSColor;
 	{
-		float4 SSSColorWithAlpha = Texture2DSample(PostprocessInput0, PostprocessInput0Sampler, UV);
+		float4 SSSColorWithAlpha = Texture2DSample(PostprocessInput0, PostprocessInput0Sampler, ViewportUV);
 
 		// renormalize to dilate RGB to fix haf res upsampling artifacts
 		SSSColor = SSSColorWithAlpha.rgb / max(SSSColorWithAlpha.a, 0.00001f);
@@ -348,8 +351,6 @@ void SubsurfaceRecombinePS(float4 UVAndScreenPos : TEXCOORD0, out float4 OutColo
 //		SSSColor += SceneColor4.rgb * GetKernel(SSSS_N_KERNELWEIGHTOFFSET, SubsurfaceProfileInt).rgb;
 	}
 
-
-
 	// 0: dont' fade in subsurface half res,  1:fully fade in half res
 	float BlendFactor = (ScreenSpaceData.GBuffer.ShadingModelID >= SHADINGMODELID_SUBSURFACE_PROFILE);
 
@@ -359,7 +360,7 @@ void SubsurfaceRecombinePS(float4 UVAndScreenPos : TEXCOORD0, out float4 OutColo
 	{
 		float SSSScaleX = SSSParams.x;
 
-		float scale = SSSScaleX / CalcSceneDepth(UV);
+		float scale = SSSScaleX / CalcSceneDepth(UVSceneColor);
 
 		// the viewport is only a fraction of the buffersize, here we compensate for that
 		// *1024 as we store the sample distances / 1000 to allows for scaling
@@ -379,13 +380,29 @@ void SubsurfaceRecombinePS(float4 UVAndScreenPos : TEXCOORD0, out float4 OutColo
 	}
 #endif
 
+	// add back in non subsurface pixels and specular
+	OutColor = float4(SSSColor * BlendFactor + SceneColor4.rgb, 0);
+}
+
+
+void ExtractSpecularPS(float4 UVAndScreenPos : TEXCOORD0, out float4 OutColor : SV_Target0)
+{
+	float2 GBufferUV = UVAndScreenPos.xy;
+
+	float4 SceneColor4 = Texture2DSample(PostprocessInput0, PostprocessInput0Sampler, GBufferUV);
+
+	FScreenSpaceData ScreenSpaceData = GetScreenSpaceData(GBufferUV);
+
+	// 0: dont' fade in subsurface half res,  1:fully fade in half res
+	float BlendFactor = (ScreenSpaceData.GBuffer.ShadingModelID >= SHADINGMODELID_SUBSURFACE_PROFILE);
 
+	float3 ExtractedSpecular = 0;
+	 
 #if SSS_RECOMBINE_METHOD
 	// we took the specular highlights out, now we add them back in
-//	SSSColor += LightAccumulator_ReconstructNonDiffuseLighting(SceneColor4);
-	SSSColor += LightAccumulator_ReconstructNonDiffuseLighting(SceneColor4);
+	ExtractedSpecular += LightAccumulator_ReconstructNonDiffuseLighting(SceneColor4);
 #endif
 
 	// alpha channel can be anything - we consumed the value and likely output to a format that doesn't need the value
-	OutColor = float4(lerp(SceneColor4.rgb, SSSColor, BlendFactor), 0);
+	OutColor = float4(lerp(SceneColor4.rgb, ExtractedSpecular, BlendFactor), 0);
 }
diff --git a/Engine/Shaders/SeparableSSS.usf b/Engine/Shaders/SeparableSSS.usf
--- a/Engine/Shaders/SeparableSSS.usf
+++ b/Engine/Shaders/SeparableSSS.usf
@@ -186,10 +186,10 @@ float3 SSSSTransmittance(float translucency, float sssWidth, float3 worldPositio
 // @param texcoord The usual quad texture coordinates.
 // @param dir Direction of the blur: First pass:   float2(1.0, 0.0), Second pass:  float2(0.0, 1.0)
 // @param initStencil indicates whether the stencil buffer should be initialized. Should be set to 'true' for the first pass if not previously initialized, to enable optimization of the second pass.
-float4 SSSSBlurPS(float2 texcoord, float2 dir, bool initStencil)
+float4 SSSSBlurPS(float2 GBufferUV, float2 ViewportUV, float2 dir, bool initStencil)
 {
     // Fetch color of current pixel:
-    float4 colorM = SSSSSampleSceneColorPoint(texcoord);
+    float4 colorM = SSSSSampleSceneColorPoint(ViewportUV);
 
 	// we store the depth in alpha
 	float depthM = colorM.a;
@@ -204,7 +204,7 @@ float4 SSSSBlurPS(float2 texcoord, float2 dir, bool initStencil)
 		return 0.0f;
 	}
 
-	float SSSStrength = GetSubsurfaceStrength(texcoord);
+	float SSSStrength = GetSubsurfaceStrength(GBufferUV);
 
 	// Initialize the stencil buffer in case it was not already available:
     if (initStencil) // (Checked in compile time, it's optimized away)
@@ -220,7 +220,7 @@ float4 SSSSBlurPS(float2 texcoord, float2 dir, bool initStencil)
 	// ideally this comes from a half res buffer as well - there are some minor artifacts
 	finalStep *= SSSStrength; // Modulate it using the opacity (0..1 range)
     
-	FGBufferData GBufferData = GetGBufferData(texcoord);
+	FGBufferData GBufferData = GetGBufferData(GBufferUV);
 
 	// 0..255, which SubSurface profile to pick
 	// ideally this comes from a half res buffer as well - there are some minor artifacts
@@ -233,7 +233,7 @@ float4 SSSSBlurPS(float2 texcoord, float2 dir, bool initStencil)
 	colorBlurred.a = colorM.a;
 
 //#if SSS_DIRECTION == 0
-	colorBlurred.rgb = colorM * GetKernel(SSSS_N_KERNELWEIGHTOFFSET, SubsurfaceProfileInt).rgb;
+	colorBlurred.rgb = colorM.rgb * GetKernel(SSSS_N_KERNELWEIGHTOFFSET, SubsurfaceProfileInt).rgb;
 //#endif
 
 	// Accumulate the other samples:
@@ -253,7 +253,7 @@ float4 SSSSBlurPS(float2 texcoord, float2 dir, bool initStencil)
 		for (int Side = -1; Side <= 1; Side += 2)
 		{
 			// Fetch color and depth for current sample:
-			float4 color = SSSSSampleSceneColor(texcoord + UVOffset * Side);
+			float4 color = SSSSSampleSceneColor(ViewportUV + UVOffset * Side);
 
 			float LocalDepth = color.a;
 			
diff --git a/Engine/Source/Runtime/Renderer/Private/CompositionLighting/CompositionLighting.cpp b/Engine/Source/Runtime/Renderer/Private/CompositionLighting/CompositionLighting.cpp
--- a/Engine/Source/Runtime/Renderer/Private/CompositionLighting/CompositionLighting.cpp
+++ b/Engine/Source/Runtime/Renderer/Private/CompositionLighting/CompositionLighting.cpp
@@ -359,7 +359,11 @@ void FCompositionLighting::ProcessLighting(FRHICommandListImmediate& RHICmdList,
 				//@todo-rco: Remove this when we fix the cross-compiler
 				!IsOpenGLPlatform(View.GetShaderPlatform()))
 			{
-				FRenderingCompositePass* PassSetup = Context.Graph.RegisterPass(new(FMemStack::Get()) FRCPassPostProcessSubsurfaceSetup());
+				// can be optimized out if we don't do split screen/stereo rendering (should be done after we some post process refactoring)
+				FRenderingCompositePass* PassExtractSpecular = Context.Graph.RegisterPass(new(FMemStack::Get()) FRCPassPostProcessSubsurfaceExtractSpecular());
+				PassExtractSpecular->SetInput(ePId_Input0, Context.FinalOutput);
+
+				FRenderingCompositePass* PassSetup = Context.Graph.RegisterPass(new(FMemStack::Get()) FRCPassPostProcessSubsurfaceSetup(View));
 				PassSetup->SetInput(ePId_Input0, Context.FinalOutput);
 
 				FRenderingCompositePass* Pass0 = Context.Graph.RegisterPass(new(FMemStack::Get()) FRCPassPostProcessSubsurface(0));
@@ -371,8 +375,8 @@ void FCompositionLighting::ProcessLighting(FRHICommandListImmediate& RHICmdList,
 				// full res composite pass, no blurring (Radius=0)
 				FRenderingCompositePass* RecombinePass = Context.Graph.RegisterPass(new(FMemStack::Get()) FRCPassPostProcessSubsurfaceRecombine());
 				RecombinePass->SetInput(ePId_Input0, Pass1);
-				RecombinePass->SetInput(ePId_Input1, Context.FinalOutput);					
-				Context.FinalOutput = FRenderingCompositeOutputRef(RecombinePass);
+				RecombinePass->SetInput(ePId_Input1, PassExtractSpecular);
+				CompositeContext.Root->AddDependency(RecombinePass);
 			}
 		}
 
diff --git a/Engine/Source/Runtime/Renderer/Private/PostProcess/PostProcessSubsurface.cpp b/Engine/Source/Runtime/Renderer/Private/PostProcess/PostProcessSubsurface.cpp
--- a/Engine/Source/Runtime/Renderer/Private/PostProcess/PostProcessSubsurface.cpp
+++ b/Engine/Source/Runtime/Renderer/Private/PostProcess/PostProcessSubsurface.cpp
@@ -432,6 +432,178 @@ static bool DoSpecularCorrection()
 	return CVarState && (SceneColorFormat >= 4);
 }
 
+// --------------------------------------
+
+
+/**
+ * Encapsulates the post processing subsurface scattering pixel shader.
+ */
+template <uint32 RecombineMethod>
+class FPostProcessSubsurfaceExtractSpecularPS : public FGlobalShader
+{
+	DECLARE_SHADER_TYPE(FPostProcessSubsurfaceExtractSpecularPS, Global);
+
+	static bool ShouldCache(EShaderPlatform Platform)
+	{
+		return SubsurfaceShouldCache(Platform);
+	}
+
+	static void ModifyCompilationEnvironment(EShaderPlatform Platform, FShaderCompilerEnvironment& OutEnvironment)
+	{
+		FGlobalShader::ModifyCompilationEnvironment(Platform, OutEnvironment);
+		OutEnvironment.SetDefine(TEXT("SSS_RECOMBINE_METHOD"), RecombineMethod);
+	}
+
+	/** Default constructor. */
+	FPostProcessSubsurfaceExtractSpecularPS () {}
+
+public:
+	FPostProcessPassParameters PostprocessParameter;
+	FDeferredPixelShaderParameters DeferredParameters;
+
+	/** Initialization constructor. */
+	FPostProcessSubsurfaceExtractSpecularPS(const ShaderMetaType::CompiledShaderInitializerType& Initializer)
+		: FGlobalShader(Initializer)
+	{
+		PostprocessParameter.Bind(Initializer.ParameterMap);
+		DeferredParameters.Bind(Initializer.ParameterMap);
+	}
+
+	void SetParameters(const FRenderingCompositePassContext& Context)
+	{
+		const FFinalPostProcessSettings& Settings = Context.View.FinalPostProcessSettings;
+		const FPixelShaderRHIParamRef ShaderRHI = GetPixelShader();
+
+		FGlobalShader::SetParameters(Context.RHICmdList, ShaderRHI, Context.View);
+		PostprocessParameter.SetPS(ShaderRHI, Context, TStaticSamplerState<SF_Point,AM_Clamp,AM_Clamp,AM_Clamp>::GetRHI());
+		DeferredParameters.Set(Context.RHICmdList, ShaderRHI, Context.View);
+	}
+
+	// FShader interface.
+	virtual bool Serialize(FArchive& Ar)
+	{
+		bool bShaderHasOutdatedParameters = FGlobalShader::Serialize(Ar);
+		Ar << PostprocessParameter << DeferredParameters;
+		return bShaderHasOutdatedParameters;
+	}
+
+	static const TCHAR* GetSourceFilename()
+	{
+		return TEXT("PostProcessSubsurface");
+	}
+
+	static const TCHAR* GetFunctionName()
+	{
+		return TEXT("ExtractSpecularPS");
+	}
+};
+
+// #define avoids a lot of code duplication
+#define VARIATION1(A) typedef FPostProcessSubsurfaceExtractSpecularPS<A> FPostProcessSubsurfaceExtractSpecularPS##A; \
+	IMPLEMENT_SHADER_TYPE2(FPostProcessSubsurfaceExtractSpecularPS##A, SF_Pixel);
+VARIATION1(0) VARIATION1(1)
+#undef VARIATION1
+
+
+
+template <uint32 RecombineMethod>
+void SetSubsurfaceExtractSpecular(const FRenderingCompositePassContext& Context)
+{
+	TShaderMapRef<FPostProcessVS> VertexShader(Context.GetShaderMap());
+	TShaderMapRef<FPostProcessSubsurfaceExtractSpecularPS<RecombineMethod> > PixelShader(Context.GetShaderMap());
+
+	static FGlobalBoundShaderState BoundShaderState;
+
+	SetGlobalBoundShaderState(Context.RHICmdList, Context.GetFeatureLevel(), BoundShaderState, GFilterVertexDeclaration.VertexDeclarationRHI, *VertexShader, *PixelShader);
+
+	PixelShader->SetParameters(Context);
+	VertexShader->SetParameters(Context);
+}
+
+void FRCPassPostProcessSubsurfaceExtractSpecular::Process(FRenderingCompositePassContext& Context)
+{
+	const FPooledRenderTargetDesc* InputDesc = GetInputDesc(ePId_Input0);
+
+	check(InputDesc);
+
+	const FSceneView& View = Context.View;
+	const FSceneViewFamily& ViewFamily = *(View.Family);
+
+	FIntPoint SrcSize = InputDesc->Extent;
+	FIntPoint DestSize = PassOutputs[0].RenderTargetDesc.Extent;
+
+	check(DestSize.X);
+	check(DestSize.Y);
+	check(SrcSize.X);
+	check(SrcSize.Y);
+
+	//	FIntRect SrcRect = View.ViewRect / ScaleFactor;
+	FIntRect SrcRect = View.ViewRect;
+	FIntRect DestRect = View.ViewRect;
+
+	TRefCountPtr<IPooledRenderTarget> NewSceneColor;
+
+	const FSceneRenderTargetItem& DestRenderTarget = PassOutputs[0].RequestSurface(Context);
+
+	// Set the view family's render target/viewport.
+	SetRenderTarget(Context.RHICmdList, DestRenderTarget.TargetableTexture, FTextureRHIRef());
+
+	Context.SetViewportAndCallRHI(0, 0, 0.0f, DestSize.X, DestSize.Y, 1.0f );
+
+	Context.RHICmdList.SetBlendState(TStaticBlendState<>::GetRHI());
+	Context.RHICmdList.SetRasterizerState(TStaticRasterizerState<>::GetRHI());
+	Context.RHICmdList.SetDepthStencilState(TStaticDepthStencilState<false, CF_Always>::GetRHI());
+
+	TShaderMapRef<FPostProcessVS> VertexShader(Context.GetShaderMap());
+
+	bool bDoSpecularCorrection = DoSpecularCorrection();
+
+	SCOPED_DRAW_EVENTF(Context.RHICmdList, SubsurfacePass, TEXT("SubsurfacePassExtractSpecular#%d"), (int32)bDoSpecularCorrection);
+
+	if(bDoSpecularCorrection)
+	{
+		SetSubsurfaceExtractSpecular<1>(Context);
+	}
+	else
+	{
+		SetSubsurfaceExtractSpecular<0>(Context);
+	}
+
+	DrawRectangle(
+		Context.RHICmdList,
+		DestRect.Min.X, DestRect.Min.Y,
+		DestRect.Width(), DestRect.Height(),
+		SrcRect.Min.X, SrcRect.Min.Y,
+		SrcRect.Width(), SrcRect.Height(),
+		DestSize,
+		SrcSize,
+		*VertexShader,
+		EDRF_UseTriangleOptimization);
+
+	Context.RHICmdList.CopyToResolveTarget(DestRenderTarget.TargetableTexture, DestRenderTarget.ShaderResourceTexture, false, FResolveParams());
+}
+
+
+FPooledRenderTargetDesc FRCPassPostProcessSubsurfaceExtractSpecular::ComputeOutputDesc(EPassOutputId InPassOutputId) const
+{
+	FPooledRenderTargetDesc Ret = PassInputs[0].GetOutput()->RenderTargetDesc;
+
+	Ret.Reset();
+	Ret.DebugName = TEXT("SubsurfaceExtractSpecular");
+	// alpha is not used
+	Ret.Format = PF_FloatRGB;
+
+	return Ret;
+}
+
+
+// --------------------------------------
+
+FRCPassPostProcessSubsurfaceSetup::FRCPassPostProcessSubsurfaceSetup(FViewInfo& View)
+	: ViewRect(View.ViewRect)
+{
+}
+
 void FRCPassPostProcessSubsurfaceSetup::Process(FRenderingCompositePassContext& Context)
 {
 	SCOPED_DRAW_EVENT(Context.RHICmdList, SubsurfaceSetup);
@@ -450,23 +622,15 @@ void FRCPassPostProcessSubsurfaceSetup::Process(FRenderingCompositePassContext&
 	FIntPoint SrcSize = InputDesc->Extent;
 	FIntPoint DestSize = PassOutputs[0].RenderTargetDesc.Extent;
 
-	// e.g. 4 means the input texture is 4x smaller than the buffer size
-	uint32 ScaleFactor = GSceneRenderTargets.GetBufferSizeXY().X / SrcSize.X;
-
-	FIntRect SrcRect = View.ViewRect / ScaleFactor;
-	FIntRect DestRect = FIntRect::DivideAndRoundUp(SrcRect, 2);
-
-	// upscale rectangle to not slightly scale
-	SrcRect = DestRect * 2;
+	FIntRect DestRect = FIntRect(0, 0, DestSize.X, DestSize.Y);
+	// upscale rectangle to not slightly scale (might miss a pixel)
+	FIntRect SrcRect = DestRect * 2 + View.ViewRect.Min;
 
 	const FSceneRenderTargetItem& DestRenderTarget = PassOutputs[0].RequestSurface(Context);
 
 	// Set the view family's render target/viewport.
 	SetRenderTarget(Context.RHICmdList, DestRenderTarget.TargetableTexture, FTextureRHIRef());
 
-	// is optimized away if possible (RT size=view size, )
-	Context.RHICmdList.Clear(true, FLinearColor::Black, false, 1.0f, false, 0, DestRect);
-
 	Context.SetViewportAndCallRHI(0, 0, 0.0f, DestSize.X, DestSize.Y, 1.0f );
 
 	// set the state
@@ -512,6 +676,8 @@ FPooledRenderTargetDesc FRCPassPostProcessSubsurfaceSetup::ComputeOutputDesc(EPa
 	// alpha is used to store depth and renormalize (alpha==0 means there is no subsurface scattering)
 	Ret.Format = PF_FloatRGBA;
 
+	Ret.Extent = ViewRect.Size();
+
 	// half res
 	Ret.Extent  = FIntPoint::DivideAndRoundUp(Ret.Extent, 2);
 	Ret.Extent.X = FMath::Max(1, Ret.Extent.X);
@@ -664,10 +830,7 @@ void FRCPassPostProcessSubsurface::Process(FRenderingCompositePassContext& Conte
 	check(SrcSize.X);
 	check(SrcSize.Y);
 
-	// e.g. 4 means the input texture is 4x smaller than the buffer size
-	uint32 ScaleFactor = GSceneRenderTargets.GetBufferSizeXY().X / SrcSize.X;
-
-	FIntRect SrcRect = View.ViewRect / ScaleFactor;
+	FIntRect SrcRect = FIntRect(0, 0, DestSize.X, DestSize.Y);
 	FIntRect DestRect = SrcRect;
 
 	TRefCountPtr<IPooledRenderTarget> NewSceneColor;
@@ -682,9 +845,6 @@ void FRCPassPostProcessSubsurface::Process(FRenderingCompositePassContext& Conte
 	// Set the view family's render target/viewport.
 	SetRenderTarget(Context.RHICmdList, DestRenderTarget->TargetableTexture, FTextureRHIRef());
 
-	// is optimized away if possible (RT size=view size, )
-	Context.RHICmdList.Clear(true, FLinearColor(0, 0, 0, 0), false, 1.0f, false, 0, DestRect);
-
 	Context.SetViewportAndCallRHI(0, 0, 0.0f, DestSize.X, DestSize.Y, 1.0f );
 
 	Context.RHICmdList.SetBlendState(TStaticBlendState<>::GetRHI());
@@ -728,14 +888,6 @@ FPooledRenderTargetDesc FRCPassPostProcessSubsurface::ComputeOutputDesc(EPassOut
 	Ret.Reset();
 	Ret.DebugName = (Direction == 0) ? TEXT("SubsurfaceX") : TEXT("SubsurfaceY");
 
-	// the setup was done in half res but the actual sampling happens in full resolution
-	Ret.Extent = GSceneRenderTargets.GetBufferSizeXY();
-
-	// half res
-	Ret.Extent  = FIntPoint::DivideAndRoundUp(Ret.Extent, 2);
-	Ret.Extent.X = FMath::Max(1, Ret.Extent.X);
-	Ret.Extent.Y = FMath::Max(1, Ret.Extent.Y);
-
 	return Ret;
 }
 
@@ -845,26 +997,15 @@ void FRCPassPostProcessSubsurfaceRecombine::Process(FRenderingCompositePassConte
 	check(SrcSize.X);
 	check(SrcSize.Y);
 
-	// e.g. 4 means the input texture is 4x smaller than the buffer size
-	uint32 ScaleFactor = GSceneRenderTargets.GetBufferSizeXY().X / SrcSize.X;
-
-	FIntRect SrcRect = View.ViewRect / ScaleFactor;
+	FIntRect SrcRect = FIntRect(0, 0, InputDesc->Extent.X, InputDesc->Extent.Y);
 	FIntRect DestRect = View.ViewRect;
 
-	TRefCountPtr<IPooledRenderTarget> NewSceneColor;
+	TRefCountPtr<IPooledRenderTarget>& SceneColor = GSceneRenderTargets.GetSceneColor();
 
-	const FSceneRenderTargetItem* DestRenderTarget;
-	{
-		// in that case we directly render to the HDR scene color
-		GRenderTargetPool.FindFreeElement(*GetInputDesc(ePId_Input1), NewSceneColor, TEXT("SceneColor"));
-		DestRenderTarget = &NewSceneColor->GetRenderTargetItem();
-	}
+	const FSceneRenderTargetItem& DestRenderTarget = SceneColor->GetRenderTargetItem();
 
 	// Set the view family's render target/viewport.
-	SetRenderTarget(Context.RHICmdList, DestRenderTarget->TargetableTexture, FTextureRHIRef());
-
-	// is optimized away if possible (RT size=view size, )
-	Context.RHICmdList.Clear(true, FLinearColor(0, 0, 0, 0), false, 1.0f, false, 0, DestRect);
+	SetRenderTarget(Context.RHICmdList, DestRenderTarget.TargetableTexture, FTextureRHIRef());
 
 	Context.SetViewportAndCallRHI(0, 0, 0.0f, DestSize.X, DestSize.Y, 1.0f );
 
@@ -898,9 +1039,7 @@ void FRCPassPostProcessSubsurfaceRecombine::Process(FRenderingCompositePassConte
 		*VertexShader,
 		EDRF_UseTriangleOptimization);
 
-	Context.RHICmdList.CopyToResolveTarget(DestRenderTarget->TargetableTexture, DestRenderTarget->ShaderResourceTexture, false, FResolveParams());
-
-	GSceneRenderTargets.SetSceneColor(NewSceneColor);
+	Context.RHICmdList.CopyToResolveTarget(DestRenderTarget.TargetableTexture, DestRenderTarget.ShaderResourceTexture, false, FResolveParams());
 }
 
 
diff --git a/Engine/Source/Runtime/Renderer/Private/PostProcess/PostProcessSubsurface.h b/Engine/Source/Runtime/Renderer/Private/PostProcess/PostProcessSubsurface.h
--- a/Engine/Source/Runtime/Renderer/Private/PostProcess/PostProcessSubsurface.h
+++ b/Engine/Source/Runtime/Renderer/Private/PostProcess/PostProcessSubsurface.h
@@ -23,6 +23,18 @@ public:
 	virtual void Release() override { delete this; }
 };
 
+// ePId_Input0: SceneColor
+// derives from TRenderingCompositePassBase<InputCount, OutputCount> 
+// uses some GBuffer attributes
+class FRCPassPostProcessSubsurfaceExtractSpecular : public TRenderingCompositePassBase<1, 1>
+{
+public:
+	// interface FRenderingCompositePass ---------
+	virtual void Process(FRenderingCompositePassContext& Context);
+	virtual FPooledRenderTargetDesc ComputeOutputDesc(EPassOutputId InPassOutputId) const;
+	virtual void Release() override { delete this; }
+};
+
 // ePId_Input0: SceneColor
 // derives from TRenderingCompositePassBase<InputCount, OutputCount> 
 // uses some GBuffer attributes
@@ -30,10 +42,15 @@ public:
 class FRCPassPostProcessSubsurfaceSetup : public TRenderingCompositePassBase<1, 1>
 {
 public:
+	// constructor
+	FRCPassPostProcessSubsurfaceSetup(FViewInfo& View);
+
 	// interface FRenderingCompositePass ---------
 	virtual void Process(FRenderingCompositePassContext& Context);
 	virtual FPooledRenderTargetDesc ComputeOutputDesc(EPassOutputId InPassOutputId) const;
 	virtual void Release() override { delete this; }
+
+	FIntRect ViewRect;
 };
 
 
