fixed Highres Screen Shot visualizebuffer output (recently got broken)

diff --git a/Engine/Source/Runtime/Engine/Private/HighResScreenshot.cpp b/Engine/Source/Runtime/Engine/Private/HighResScreenshot.cpp
--- a/Engine/Source/Runtime/Engine/Private/HighResScreenshot.cpp
+++ b/Engine/Source/Runtime/Engine/Private/HighResScreenshot.cpp
@@ -134,8 +134,10 @@ bool FHighResScreenshotConfig::SaveImage(const FString& File, const TArray<TPixe
 	const size_t BitsPerPixel = (sizeof(TPixelType) / 4) * 8;
 
 	TSharedPtr<class IImageWrapper> ImageCompressor = bIsWritingHDRImage ? ImageCompressorHDR : ImageCompressorLDR;
-	FString Extension = bIsWritingHDRImage ? TEXT(".exr") : TEXT(".png");
-	FString Filename = (FPaths::GetExtension(File, true) == Extension) ? File : FPaths::GetBaseFilename(File, false) + Extension;
+
+	// here we require the input file name to have an extension
+	FString NewExtension = bIsWritingHDRImage ? TEXT(".exr") : TEXT(".png");
+	FString Filename = FPaths::GetBaseFilename(File, false) + NewExtension;
 
 	if (OutFilename != nullptr)
 	{
diff --git a/Engine/Source/Runtime/Renderer/Private/PostProcess/PostProcessing.cpp b/Engine/Source/Runtime/Renderer/Private/PostProcess/PostProcessing.cpp
--- a/Engine/Source/Runtime/Renderer/Private/PostProcess/PostProcessing.cpp
+++ b/Engine/Source/Runtime/Renderer/Private/PostProcess/PostProcessing.cpp
@@ -694,7 +694,7 @@ static void AddGBufferVisualizationOverview(FPostprocessContext& Context, FRende
 
 					if (BaseFilename.Len())
 					{
-						FString MaterialFilename = BaseFilename + TEXT("_") + (*It)->GetName();
+						FString MaterialFilename = BaseFilename + TEXT("_") + (*It)->GetName() + TEXT(".png");
 						MaterialPass->SetOutputDumpFilename(ePId_Output0, *MaterialFilename);
 					}
 
diff --git a/Engine/Source/Runtime/Renderer/Private/PostProcess/RenderingCompositionGraph.h b/Engine/Source/Runtime/Renderer/Private/PostProcess/RenderingCompositionGraph.h
--- a/Engine/Source/Runtime/Renderer/Private/PostProcess/RenderingCompositionGraph.h
+++ b/Engine/Source/Runtime/Renderer/Private/PostProcess/RenderingCompositionGraph.h
@@ -198,14 +198,14 @@ struct FRenderingCompositePass
 
 	/**
 	 * Allows access to dump filename for a given output
-	 * @return Filename for output dump
+	 * @return Filename for output dump with extension
 	 */
 	virtual const FString& GetOutputDumpFilename(EPassOutputId OutputId) = 0;
 
 	/**
 	 * Allows setting of a dump filename for a given output
 	 * @param Index - Output index
-	 * @param Filename - Output dump filename
+	 * @param Filename - Output dump filename, needs to have extension, gets modified if we have an HDR image e.g. ".png"
 	 */
 	virtual void SetOutputDumpFilename(EPassOutputId OutputId, const TCHAR* Filename) = 0;
 
