- Part 2 of Metal MRT support (adding ESimpleRenderTargetMode to make SetRenderTarget() helper usable with SetRenderTargetsAndClear)
- Made a new EmptyVertexDeclaration, for vertex shaders that don't need any inputs (eg, they generate verts based on vertex index)
- Split BeginRenderingSceneColor into 2 functions, adding BeginRenderingGBuffer

diff --git a/Engine/Source/Runtime/RHI/Public/RHIDefinitions.h b/Engine/Source/Runtime/RHI/Public/RHIDefinitions.h
--- a/Engine/Source/Runtime/RHI/Public/RHIDefinitions.h
+++ b/Engine/Source/Runtime/RHI/Public/RHIDefinitions.h
@@ -539,6 +539,26 @@ enum class ERenderTargetStoreAction
 	EMultisampleResolve,
 };
 
+/**
+ * Common render target use cases
+ */
+enum class ESimpleRenderTargetMode
+{
+	// These will all store out color and depth
+	EExistingColorAndDepth,			// Color = Existing, Depth = Existing
+	EUninitializedColorAndDepth,	// Color = ????, Depth = ????
+	EUninitializedColorExistingDepth,// Color = ????, Depth = Existing
+	EUninitializedColorClearDepth,	// Color = ????, Depth = Default
+	EClearToDefault,				// Default Color = (0,0,0,0), Default Depth = 0.0f
+	EClearColorToBlack,				// Color = (0,0,0,0), Depth = Existing
+	EClearColorToBlackWithFullAlpha,// Color = (0,0,0,1), Depth = Existing
+	EClearColorToWhite,				// Color = (1,1,1,1), Depth = Existing
+	EClearDepthToOne,				// Color = Existing, Depth = 1.0
+	EExistingContents_NoDepthStore,	// Load existing contents, but don't store depth out
+
+	// If you add an item here, make sure to add it to DecodeRenderTargetMode() as well!
+};
+
 inline bool IsPCPlatform(const EShaderPlatform Platform)
 {
 	return Platform == SP_PCD3D_SM5 || Platform == SP_PCD3D_SM4 || Platform == SP_PCD3D_ES2 || Platform ==  SP_OPENGL_SM4 || Platform == SP_OPENGL_SM4_MAC || Platform == SP_OPENGL_SM5 || Platform == SP_OPENGL_PCES2;
diff --git a/Engine/Source/Runtime/RHI/Public/RHIUtilities.h b/Engine/Source/Runtime/RHI/Public/RHIUtilities.h
--- a/Engine/Source/Runtime/RHI/Public/RHIUtilities.h
+++ b/Engine/Source/Runtime/RHI/Public/RHIUtilities.h
@@ -120,6 +120,70 @@ struct FRWBufferByteAddress
 	}
 };
 
+/**
+ * Convert the ESimpleRenderTargetMode into usable values 
+ * @todo: Can we easily put this into a .cpp somewhere?
+ */
+inline void DecodeRenderTargetMode(ESimpleRenderTargetMode Mode, ERenderTargetLoadAction& ColorLoadAction, ERenderTargetStoreAction& ColorStoreAction, ERenderTargetLoadAction& DepthLoadAction, ERenderTargetStoreAction& DepthStoreAction, FLinearColor& ClearColor, float& ClearDepth)
+{
+	// set defaults
+	ColorStoreAction = ERenderTargetStoreAction::EStore;
+	DepthStoreAction = ERenderTargetStoreAction::EStore;
+	ClearColor = FLinearColor(0, 0, 0, 0);
+	ClearDepth = 0.0f;
+
+	switch (Mode)
+	{
+	case ESimpleRenderTargetMode::EExistingColorAndDepth:
+		ColorLoadAction = ERenderTargetLoadAction::ELoad;
+		DepthLoadAction = ERenderTargetLoadAction::ELoad;
+		break;
+	case ESimpleRenderTargetMode::EUninitializedColorAndDepth:
+		ColorLoadAction = ERenderTargetLoadAction::ENoAction;
+		DepthLoadAction = ERenderTargetLoadAction::ENoAction;
+		break;
+	case ESimpleRenderTargetMode::EUninitializedColorExistingDepth:
+		ColorLoadAction = ERenderTargetLoadAction::ENoAction;
+		DepthLoadAction = ERenderTargetLoadAction::ELoad;
+		break;
+	case ESimpleRenderTargetMode::EUninitializedColorClearDepth:
+		ColorLoadAction = ERenderTargetLoadAction::ENoAction;
+		DepthLoadAction = ERenderTargetLoadAction::EClear;
+		break;
+	case ESimpleRenderTargetMode::EClearToDefault:
+		ColorLoadAction = ERenderTargetLoadAction::EClear;
+		DepthLoadAction = ERenderTargetLoadAction::EClear;
+		break;
+	case ESimpleRenderTargetMode::EClearColorToBlack:
+		ColorLoadAction = ERenderTargetLoadAction::EClear;
+		DepthLoadAction = ERenderTargetLoadAction::ELoad;
+		ClearColor = FLinearColor::Black;
+		break;
+	case ESimpleRenderTargetMode::EClearColorToBlackWithFullAlpha:
+		ColorLoadAction = ERenderTargetLoadAction::EClear;
+		DepthLoadAction = ERenderTargetLoadAction::ELoad;
+		ClearColor = FLinearColor(0, 0, 0, 1);
+		break;
+	case ESimpleRenderTargetMode::EClearColorToWhite:
+		ColorLoadAction = ERenderTargetLoadAction::EClear;
+		DepthLoadAction = ERenderTargetLoadAction::ELoad;
+		ClearColor = FLinearColor::White;
+		break;
+	case ESimpleRenderTargetMode::EClearDepthToOne:
+		ColorLoadAction = ERenderTargetLoadAction::ELoad;
+		DepthLoadAction = ERenderTargetLoadAction::EClear;
+		ClearDepth = 1.0f;
+		break;
+	case ESimpleRenderTargetMode::EExistingContents_NoDepthStore:
+		ColorLoadAction = ERenderTargetLoadAction::ELoad;
+		DepthLoadAction = ERenderTargetLoadAction::ELoad;
+		DepthStoreAction = ERenderTargetStoreAction::ENoAction;
+		break;
+	default:
+		UE_LOG(LogRHI, Fatal, TEXT("Using a ESimpleRenderTargetMode that wasn't decoded in DecodeRenderTargetMode [value = %d]"), (int32)Mode);
+	}
+}
+
 /** Helper for the common case of using a single color and depth render target. */
 inline void SetRenderTarget(FRHICommandList& RHICmdList, FTextureRHIParamRef NewRenderTarget, FTextureRHIParamRef NewDepthStencilTarget)
 {
@@ -127,6 +191,23 @@ inline void SetRenderTarget(FRHICommandList& RHICmdList, FTextureRHIParamRef New
 	RHICmdList.SetRenderTargets(1, &RTV, NewDepthStencilTarget, 0, NULL);
 }
 
+/** Helper for the common case of using a single color and depth render target. */
+inline void SetRenderTarget(FRHICommandList& RHICmdList, FTextureRHIParamRef NewRenderTarget, FTextureRHIParamRef NewDepthStencilTarget, ESimpleRenderTargetMode Mode)
+{
+	ERenderTargetLoadAction ColorLoadAction, DepthLoadAction;
+	ERenderTargetStoreAction ColorStoreAction, DepthStoreAction;
+	FLinearColor ClearColor;
+	float ClearDepth;
+	DecodeRenderTargetMode(Mode, ColorLoadAction, ColorStoreAction, DepthLoadAction, DepthStoreAction, ClearColor, ClearDepth);
+
+	// now make the FRHISetRenderTargetsInfo that encapsulates all of the info
+	FRHIRenderTargetView ColorView(NewRenderTarget, 0, -1, ColorLoadAction, ColorStoreAction);
+	FRHISetRenderTargetsInfo Info(1, &ColorView, FRHIDepthRenderTargetView(NewDepthStencilTarget, DepthLoadAction, DepthStoreAction));
+	Info.ClearColors[0] = ClearColor;
+	Info.DepthClearValue = ClearDepth;
+	RHICmdList.SetRenderTargetsAndClear(Info);
+}
+
 /** Helper for the common case of using a single color and depth render target, with a mip index for the color target. */
 inline void SetRenderTarget(FRHICommandList& RHICmdList, FTextureRHIParamRef NewRenderTarget, int32 MipIndex, FTextureRHIParamRef NewDepthStencilTarget)
 {
diff --git a/Engine/Source/Runtime/Renderer/Private/DeferredShadingRenderer.cpp b/Engine/Source/Runtime/Renderer/Private/DeferredShadingRenderer.cpp
--- a/Engine/Source/Runtime/Renderer/Private/DeferredShadingRenderer.cpp
+++ b/Engine/Source/Runtime/Renderer/Private/DeferredShadingRenderer.cpp
@@ -101,16 +101,6 @@ FDeferredShadingSceneRenderer::FDeferredShadingSceneRenderer(const FSceneViewFam
 	}
 }
 
-/** 
-* Clears a view. 
-*/
-void FDeferredShadingSceneRenderer::ClearView(FRHICommandListImmediate& RHICmdList)
-{
-	// Clear the G Buffer render targets
-	const bool bClearBlack = Views[0].Family->EngineShowFlags.ShaderComplexity || Views[0].Family->EngineShowFlags.StationaryLightOverlap;
-	GSceneRenderTargets.ClearGBufferTargets(RHICmdList, bClearBlack ? FLinearColor(0, 0, 0, 0) : Views[0].BackgroundColor);
-}
-
 namespace
 {
 	FVector4 ClearQuadVertices[4] = 
@@ -534,7 +524,7 @@ public:
 	}
 	virtual void SetStateOnCommandList(FRHICommandList& CmdList) override
 	{
-		GSceneRenderTargets.BeginRenderingSceneColor(CmdList, true);
+		GSceneRenderTargets.BeginRenderingGBuffer(CmdList, ERenderTargetLoadAction::ELoad, ERenderTargetLoadAction::ELoad);
 		SetupBasePassView(CmdList, View.ViewRect, !!ViewFamily.EngineShowFlags.ShaderComplexity);
 	}
 };
@@ -646,6 +636,14 @@ DECLARE_CYCLE_STAT(TEXT("Velocity"), STAT_CLM_Velocity, STATGROUP_CommandListMar
 DECLARE_CYCLE_STAT(TEXT("AfterVelocity"), STAT_CLM_AfterVelocity, STATGROUP_CommandListMarkers);
 DECLARE_CYCLE_STAT(TEXT("AfterFrame"), STAT_CLM_AfterFrame, STATGROUP_CommandListMarkers);
 
+/**
+ * Returns true if the depth Prepass needs to run
+ */
+static FORCEINLINE bool NeedsPrePass(const FDeferredShadingSceneRenderer* Renderer)
+{
+	return (RHIHasTiledGPU(Renderer->ViewFamily.GetShaderPlatform()) == false) && 
+		(Renderer->EarlyZPassMode != DDM_None || GEarlyZPassMovable != 0);
+}
 
 void FDeferredShadingSceneRenderer::Render(FRHICommandListImmediate& RHICmdList)
 {
@@ -751,21 +749,18 @@ void FDeferredShadingSceneRenderer::Render(FRHICommandListImmediate& RHICmdList)
 	GRenderTargetPool.AddPhaseEvent(TEXT("EarlyZPass"));
 
 	// Draw the scene pre-pass / early z pass, populating the scene depth buffer and HiZ
-	RHICmdList.SetCurrentStat(GET_STATID(STAT_CLM_PrePass));
-	RenderPrePass(RHICmdList);
+	bool bDepthWasCleared = false;
+	if (NeedsPrePass(this))
+	{
+		RHICmdList.SetCurrentStat(GET_STATID(STAT_CLM_PrePass));
+		RenderPrePass(RHICmdList);
+		// at this point, the depth was cleared
+		bDepthWasCleared = true;
+	}
 	RHICmdList.SetCurrentStat(GET_STATID(STAT_CLM_AfterPrePass));
 	
 	GSceneRenderTargets.AllocGBufferTargets();
 	
-	// Clear scene color buffer if necessary.
-	if ( bRequiresRHIClear )
-	{
-		ClearView(RHICmdList);
-
-		// Only clear once.
-		bRequiresRHIClear = false;
-	}
-
 	// Clear LPVs for all views
 	if ( FeatureLevel >= ERHIFeatureLevel::SM5 )
 	{
@@ -798,8 +793,22 @@ void FDeferredShadingSceneRenderer::Render(FRHICommandListImmediate& RHICmdList)
 	}
 	else
 	{
-		// Begin rendering to scene color
-		GSceneRenderTargets.BeginRenderingSceneColor(RHICmdList, true);
+		// if we didn't to the prepass above, then we will need to clear now, otherwise, it's already been cleared and rendered to
+		ERenderTargetLoadAction DepthLoadAction = bDepthWasCleared ? ERenderTargetLoadAction::ELoad : ERenderTargetLoadAction::EClear;
+
+		if (bRequiresRHIClear)
+		{
+			// Clear the G Buffer render targets
+			const bool bClearBlack = Views[0].Family->EngineShowFlags.ShaderComplexity || Views[0].Family->EngineShowFlags.StationaryLightOverlap;
+			const FLinearColor ClearColor = (bClearBlack ? FLinearColor(0, 0, 0, 0) : Views[0].BackgroundColor);
+
+			// render to the GBuffer, clearing it
+			GSceneRenderTargets.BeginRenderingGBuffer(RHICmdList, ERenderTargetLoadAction::EClear, DepthLoadAction, ClearColor);
+		}
+		else
+		{
+			GSceneRenderTargets.BeginRenderingGBuffer(RHICmdList, ERenderTargetLoadAction::ENoAction, DepthLoadAction);
+		}
 	}
 
 	GRenderTargetPool.AddPhaseEvent(TEXT("BasePass"));
@@ -811,7 +820,7 @@ void FDeferredShadingSceneRenderer::Render(FRHICommandListImmediate& RHICmdList)
 	if(ViewFamily.EngineShowFlags.VisualizeLightCulling)
 	{
 		// clear out emissive and baked lighting (not too efficient but simple and only needed for this debug view)
-		GSceneRenderTargets.BeginRenderingSceneColor(RHICmdList, false);
+		GSceneRenderTargets.BeginRenderingSceneColor(RHICmdList);
 		RHICmdList.Clear(true, FLinearColor(0, 0, 0, 0), false, 0, false, 0, FIntRect());
 	}
 
diff --git a/Engine/Source/Runtime/Renderer/Private/DistanceFieldSurfaceCacheLighting.cpp b/Engine/Source/Runtime/Renderer/Private/DistanceFieldSurfaceCacheLighting.cpp
--- a/Engine/Source/Runtime/Renderer/Private/DistanceFieldSurfaceCacheLighting.cpp
+++ b/Engine/Source/Runtime/Renderer/Private/DistanceFieldSurfaceCacheLighting.cpp
@@ -3521,7 +3521,7 @@ bool FDeferredShadingSceneRenderer::RenderDistanceFieldAOSurfaceCache(FRHIComman
 
 				if (bApplyToSceneColor)
 				{
-					GSceneRenderTargets.BeginRenderingSceneColor(RHICmdList, false);
+					GSceneRenderTargets.BeginRenderingSceneColor(RHICmdList);
 				}
 				else
 				{
@@ -3740,7 +3740,7 @@ void FDeferredShadingSceneRenderer::RenderMeshDistanceFieldVisualization(FRHICom
 				}
 
 				{
-					GSceneRenderTargets.BeginRenderingSceneColor(RHICmdList, false);
+					GSceneRenderTargets.BeginRenderingSceneColor(RHICmdList);
 
 					for (int32 ViewIndex = 0; ViewIndex < Views.Num(); ViewIndex++)
 					{
diff --git a/Engine/Source/Runtime/Renderer/Private/ForwardShadingRenderer.cpp b/Engine/Source/Runtime/Renderer/Private/ForwardShadingRenderer.cpp
--- a/Engine/Source/Runtime/Renderer/Private/ForwardShadingRenderer.cpp
+++ b/Engine/Source/Runtime/Renderer/Private/ForwardShadingRenderer.cpp
@@ -113,7 +113,7 @@ void FForwardShadingSceneRenderer::Render(FRHICommandListImmediate& RHICmdList)
 	else
 	{
 		// Begin rendering to scene color
-		GSceneRenderTargets.BeginRenderingSceneColor(RHICmdList, false);
+		GSceneRenderTargets.BeginRenderingSceneColor(RHICmdList);
 	}
 
 	// Clear color and depth buffer
diff --git a/Engine/Source/Runtime/Renderer/Private/PostProcess/PostProcessLensBlur.cpp b/Engine/Source/Runtime/Renderer/Private/PostProcess/PostProcessLensBlur.cpp
--- a/Engine/Source/Runtime/Renderer/Private/PostProcess/PostProcessLensBlur.cpp
+++ b/Engine/Source/Runtime/Renderer/Private/PostProcess/PostProcessLensBlur.cpp
@@ -185,9 +185,7 @@ void FRCPassPostProcessLensBlur::Process(FRenderingCompositePassContext& Context
 	const FSceneRenderTargetItem& DestRenderTarget = PassOutputs[0].RequestSurface(Context);
 
 	// Set the view family's render target/viewport.
-	SetRenderTarget(Context.RHICmdList, DestRenderTarget.TargetableTexture, FTextureRHIRef());
-
-	Context.RHICmdList.Clear(true, FLinearColor::Black, false, 1.0f, false, 0, FIntRect());
+	SetRenderTarget(Context.RHICmdList, DestRenderTarget.TargetableTexture, FTextureRHIRef(), ESimpleRenderTargetMode::EClearColorToBlack);
 
 	Context.SetViewportAndCallRHI(ViewRect);
 
@@ -201,7 +199,7 @@ void FRCPassPostProcessLensBlur::Process(FRenderingCompositePassContext& Context
 
 	static FGlobalBoundShaderState BoundShaderState;
 	
-	SetGlobalBoundShaderState(Context.RHICmdList, Context.GetFeatureLevel(), BoundShaderState, GFilterVertexDeclaration.VertexDeclarationRHI, *VertexShader, *PixelShader);
+	SetGlobalBoundShaderState(Context.RHICmdList, Context.GetFeatureLevel(), BoundShaderState, GEmptyVertexDeclaration.VertexDeclarationRHI, *VertexShader, *PixelShader);
 
 	uint32 TileSize = 1;
 
diff --git a/Engine/Source/Runtime/Renderer/Private/PostProcess/SceneFilterRendering.cpp b/Engine/Source/Runtime/Renderer/Private/PostProcess/SceneFilterRendering.cpp
--- a/Engine/Source/Runtime/Renderer/Private/PostProcess/SceneFilterRendering.cpp
+++ b/Engine/Source/Runtime/Renderer/Private/PostProcess/SceneFilterRendering.cpp
@@ -72,6 +72,8 @@ static TGlobalResource<FScreenRectangleIndexBuffer> GScreenRectangleIndexBuffer;
 
 /** Vertex declaration for the 2D screen rectangle. */
 TGlobalResource<FFilterVertexDeclaration> GFilterVertexDeclaration;
+/** Vertex declaration for vertex shaders that don't require any inputs (eg. generated via vertex ID). */
+TGlobalResource<FEmptyVertexDeclaration> GEmptyVertexDeclaration;
 
 /** Uniform buffer for computing the vertex positional and UV adjustments in the vertex shader. */
 BEGIN_UNIFORM_BUFFER_STRUCT( FDrawRectangleParameters,)
diff --git a/Engine/Source/Runtime/Renderer/Private/PostProcess/SceneFilterRendering.h b/Engine/Source/Runtime/Renderer/Private/PostProcess/SceneFilterRendering.h
--- a/Engine/Source/Runtime/Renderer/Private/PostProcess/SceneFilterRendering.h
+++ b/Engine/Source/Runtime/Renderer/Private/PostProcess/SceneFilterRendering.h
@@ -59,6 +59,7 @@ extern void DrawTransformedRectangle(
     );
 
 extern TGlobalResource<FFilterVertexDeclaration> GFilterVertexDeclaration;
+extern TGlobalResource<FEmptyVertexDeclaration> GEmptyVertexDeclaration;
 
 
 
diff --git a/Engine/Source/Runtime/Renderer/Private/PostProcess/SceneRenderTargets.cpp b/Engine/Source/Runtime/Renderer/Private/PostProcess/SceneRenderTargets.cpp
--- a/Engine/Source/Runtime/Renderer/Private/PostProcess/SceneRenderTargets.cpp
+++ b/Engine/Source/Runtime/Renderer/Private/PostProcess/SceneRenderTargets.cpp
@@ -257,69 +257,64 @@ void FSceneRenderTargets::Allocate(const FSceneViewFamily& ViewFamily)
 	AllocateRenderTargets();
 }
 
-/** Clears the GBuffer render targets to default values. */
-void FSceneRenderTargets::ClearGBufferTargets(FRHICommandListImmediate& RHICmdList, const FLinearColor& ClearColor)
+void FSceneRenderTargets::BeginRenderingSceneColor(FRHICommandList& RHICmdList, ESimpleRenderTargetMode RenderTargetMode/*=EUninitializedColorExistingDepth*/)
 {
-	SCOPED_DRAW_EVENT(RHICmdList, ClearGBufferTargets);
-
-	// Clear GBufferA, GBufferB, GBufferC, GBufferD, GBufferE
-	{
-		GSceneRenderTargets.BeginRenderingSceneColor(RHICmdList, true);
+	SCOPED_DRAW_EVENT(RHICmdList, BeginRenderingSceneColor);
 
-		int32 NumToClear = GSceneRenderTargets.GetNumGBufferTargets();
-		if (NumToClear > 1)
-		{
-			// Using 0 and 1 ensures we go through the fast path on Intel integrated GPUs.
-			// Normal could be 0.5,0.5,0.5 but then it would not use the fast path.
-			FLinearColor ClearColors[6] = {ClearColor, FLinearColor(0, 0, 0, 0), FLinearColor(0,0,0,0), FLinearColor(0,0,0,0), FLinearColor(0,1,1,1), FLinearColor(1,1,1,1)};
-			RHICmdList.ClearMRT(true, NumToClear, ClearColors, false, 0, false, 0, FIntRect());
-		}
-		else
-		{
-			RHICmdList.Clear(true, ClearColor, false, 0, false, 0, FIntRect());
-		}
-	}
-}
+	AllocSceneColor();
+	
+	SetRenderTarget(RHICmdList, GetSceneColorSurface(), GetSceneDepthSurface(), RenderTargetMode);
+} 
 
-void FSceneRenderTargets::BeginRenderingSceneColor(FRHICommandList& RHICmdList, bool bGBufferPass)
+void FSceneRenderTargets::BeginRenderingGBuffer(FRHICommandList& RHICmdList, ERenderTargetLoadAction ColorLoadAction, ERenderTargetLoadAction DepthLoadAction, const FLinearColor& ClearColor/*=(0,0,0,1)*/)
 {
 	SCOPED_DRAW_EVENT(RHICmdList, BeginRenderingSceneColor);
 
-	if(IsSimpleDynamicLightingEnabled())
+	if (IsSimpleDynamicLightingEnabled())
 	{
-		bGBufferPass = false;
+		// in this non-standard case, just render to scene color with default mode
+		BeginRenderingSceneColor(RHICmdList);
+		return;
 	}
-	
+
 	AllocSceneColor();
-	
+
 	// Set the scene color surface as the render target, and the scene depth surface as the depth-stencil target.
-	if (bGBufferPass && CurrentFeatureLevel >= ERHIFeatureLevel::SM4)
+	if (CurrentFeatureLevel >= ERHIFeatureLevel::SM4)
 	{
-		FTextureRHIParamRef RenderTargets[6] = {0};
-		RenderTargets[0] = GetSceneColorSurface();
-		RenderTargets[1] = GSceneRenderTargets.GBufferA->GetRenderTargetItem().TargetableTexture;
-		RenderTargets[2] = GSceneRenderTargets.GBufferB->GetRenderTargetItem().TargetableTexture;
-		RenderTargets[3] = GSceneRenderTargets.GBufferC->GetRenderTargetItem().TargetableTexture;
-		RenderTargets[4] = GSceneRenderTargets.GBufferD->GetRenderTargetItem().TargetableTexture;
+		FRHIRenderTargetView RenderTargets[6];
+		RenderTargets[0] = FRHIRenderTargetView(GetSceneColorSurface(), 0, -1, ColorLoadAction, ERenderTargetStoreAction::EStore);
+		RenderTargets[1] = FRHIRenderTargetView(GSceneRenderTargets.GBufferA->GetRenderTargetItem().TargetableTexture, 0, -1, ColorLoadAction, ERenderTargetStoreAction::EStore);
+		RenderTargets[2] = FRHIRenderTargetView(GSceneRenderTargets.GBufferB->GetRenderTargetItem().TargetableTexture, 0, -1, ColorLoadAction, ERenderTargetStoreAction::EStore);
+		RenderTargets[3] = FRHIRenderTargetView(GSceneRenderTargets.GBufferC->GetRenderTargetItem().TargetableTexture, 0, -1, ColorLoadAction, ERenderTargetStoreAction::EStore);
+		RenderTargets[4] = FRHIRenderTargetView(GSceneRenderTargets.GBufferD->GetRenderTargetItem().TargetableTexture, 0, -1, ColorLoadAction, ERenderTargetStoreAction::EStore);
 
 		uint32 MRTCount = ARRAY_COUNT(RenderTargets);
 
 		if (bAllowStaticLighting)
 		{
-			RenderTargets[5] = GSceneRenderTargets.GBufferE->GetRenderTargetItem().TargetableTexture;
+			RenderTargets[5] = FRHIRenderTargetView(GSceneRenderTargets.GBufferE->GetRenderTargetItem().TargetableTexture, 0, -1, ColorLoadAction, ERenderTargetStoreAction::EStore);
 		}
 		else
 		{
 			MRTCount--;
 		}
-		
-		SetRenderTargets(RHICmdList, MRTCount, RenderTargets, GetSceneDepthSurface(), 0, NULL);
-	}
-	else
-	{
-		SetRenderTarget(RHICmdList, GetSceneColorSurface(), GetSceneDepthSurface());
+
+		FRHIDepthRenderTargetView DepthView(GetSceneDepthSurface(), DepthLoadAction, ERenderTargetStoreAction::EStore);
+		FRHISetRenderTargetsInfo Info(MRTCount, RenderTargets, DepthView);
+		if (ColorLoadAction == ERenderTargetLoadAction::EClear)
+		{
+			Info.ClearColors[0] = ClearColor;
+			Info.ClearColors[1] = Info.ClearColors[2] = Info.ClearColors[3] = FLinearColor::Transparent;
+			Info.ClearColors[4] = FLinearColor(0, 1, 1, 1);
+			Info.ClearColors[5] = FLinearColor(1, 1, 1, 1);
+		}
+		Info.DepthClearValue = 0.0f;
+
+		// set the render target
+		RHICmdList.SetRenderTargetsAndClear(Info);
 	}
-} 
+}
 
 int32 FSceneRenderTargets::GetNumGBufferTargets() const
 {
@@ -783,16 +778,8 @@ void FSceneRenderTargets::FinishRenderingPrePass(FRHICommandListImmediate& RHICm
 void FSceneRenderTargets::BeginRenderingShadowDepth(FRHICommandList& RHICmdList, bool bClear)
 {
 	GRenderTargetPool.VisualizeTexture.SetCheckPoint(RHICmdList, ShadowDepthZ);
-	FRHISetRenderTargetsInfo Info(0, nullptr, 
-		FRHIDepthRenderTargetView(GetShadowDepthZSurface(), 
-			bClear ? ERenderTargetLoadAction::EClear : ERenderTargetLoadAction::ELoad, 
-			ERenderTargetStoreAction::EStore, 
-			bClear ? ERenderTargetLoadAction::EClear : ERenderTargetLoadAction::ELoad,
-			ERenderTargetStoreAction::EStore));
-	if (bClear)
-	{
-		Info.SetClearDepthStencil(true, 1.0f);
-	}
+	FRHISetRenderTargetsInfo Info(0, nullptr, FRHIDepthRenderTargetView(GetShadowDepthZSurface(), bClear ? ERenderTargetLoadAction::EClear : ERenderTargetLoadAction::ELoad, ERenderTargetStoreAction::EStore));
+	Info.DepthClearValue = 1.0f;
 	Info.ColorRenderTarget[0].StoreAction = ERenderTargetStoreAction::ENoAction;
 
 	if (!GSupportsDepthRenderTargetWithoutColorRenderTarget)
@@ -866,7 +853,7 @@ void FSceneRenderTargets::FinishRenderingSceneAlphaCopy(FRHICommandListImmediate
 }
 
 
-void FSceneRenderTargets::BeginRenderingLightAttenuation(FRHICommandList& RHICmdList)
+void FSceneRenderTargets::BeginRenderingLightAttenuation(FRHICommandList& RHICmdList, bool bClearToWhite)
 {
 	SCOPED_DRAW_EVENT(RHICmdList, BeginRenderingLightAttenuation);
 
@@ -875,7 +862,7 @@ void FSceneRenderTargets::BeginRenderingLightAttenuation(FRHICommandList& RHICmd
 	GRenderTargetPool.VisualizeTexture.SetCheckPoint(RHICmdList, GSceneRenderTargets.GetLightAttenuation());
 
 	// Set the light attenuation surface as the render target, and the scene depth buffer as the depth-stencil surface.
-	SetRenderTarget(RHICmdList, GetLightAttenuationSurface(),GetSceneDepthSurface());
+	SetRenderTarget(RHICmdList, GetLightAttenuationSurface(), GetSceneDepthSurface(), bClearToWhite ? ESimpleRenderTargetMode::EClearColorToWhite : ESimpleRenderTargetMode::EExistingColorAndDepth);
 }
 
 void FSceneRenderTargets::FinishRenderingLightAttenuation(FRHICommandList& RHICmdList)
@@ -891,7 +878,7 @@ void FSceneRenderTargets::FinishRenderingLightAttenuation(FRHICommandList& RHICm
 void FSceneRenderTargets::BeginRenderingTranslucency(FRHICommandList& RHICmdList, const FViewInfo& View)
 {
 	// Use the scene color buffer.
-	GSceneRenderTargets.BeginRenderingSceneColor(RHICmdList);
+	GSceneRenderTargets.BeginRenderingSceneColor(RHICmdList, ESimpleRenderTargetMode::EExistingColorAndDepth);
 	
 	// viewport to match view size
 	RHICmdList.SetViewport(View.ViewRect.Min.X, View.ViewRect.Min.Y, 0.0f, View.ViewRect.Max.X, View.ViewRect.Max.Y, 1.0f);
@@ -914,13 +901,9 @@ bool FSceneRenderTargets::BeginRenderingSeparateTranslucency(FRHICommandList& RH
 
 		TRefCountPtr<IPooledRenderTarget>& SeparateTranslucency = ViewState->GetSeparateTranslucency(View);
 
-		// Use a separate render target for translucency
-		SetRenderTarget(RHICmdList, SeparateTranslucency->GetRenderTargetItem().TargetableTexture, GetSceneDepthSurface());
-		
-		if(bFirstTimeThisFrame)
-		{
-			RHICmdList.Clear(true, FLinearColor(0, 0, 0, 1), false, 0, false, 0, FIntRect());
-		}
+		// clear the render target the first time, re-use afterwards
+		SetRenderTarget(RHICmdList, SeparateTranslucency->GetRenderTargetItem().TargetableTexture, GetSceneDepthSurface(), 
+			bFirstTimeThisFrame ? ESimpleRenderTargetMode::EClearColorToBlackWithFullAlpha : ESimpleRenderTargetMode::EExistingColorAndDepth);
 
 		RHICmdList.SetViewport(View.ViewRect.Min.X, View.ViewRect.Min.Y, 0.0f, View.ViewRect.Max.X, View.ViewRect.Max.Y, 1.0f);
 
diff --git a/Engine/Source/Runtime/Renderer/Private/PostProcess/SceneRenderTargets.h b/Engine/Source/Runtime/Renderer/Private/PostProcess/SceneRenderTargets.h
--- a/Engine/Source/Runtime/Renderer/Private/PostProcess/SceneRenderTargets.h
+++ b/Engine/Source/Runtime/Renderer/Private/PostProcess/SceneRenderTargets.h
@@ -155,15 +155,11 @@ public:
 	 *
 	 */
 	void SetBufferSize(int32 InBufferSizeX, int32 InBufferSizeY);
-	/**
-	 * Clears the GBuffer render targets to default values.
-	 */
-	void ClearGBufferTargets(FRHICommandListImmediate& RHICmdList, const FLinearColor& ClearColor);
 	/**
 	 * Sets the scene color target and restores its contents if necessary
-	 * @param bGBufferPass - Whether the pass about to be rendered is the GBuffer population pass
 	 */
-	void BeginRenderingSceneColor(FRHICommandList& RHICmdList, bool bGBufferPass = false);
+	void BeginRenderingSceneColor(FRHICommandList& RHICmdList, ESimpleRenderTargetMode RenderTargetMode=ESimpleRenderTargetMode::EUninitializedColorExistingDepth);
+	void BeginRenderingGBuffer(FRHICommandList& RHICmdList, ERenderTargetLoadAction ColorLoadAction, ERenderTargetLoadAction DepthLoadAction, const FLinearColor& ClearColor=FLinearColor(0,0,0,1));
 	/**
 	 * Called when finished rendering to the scene color surface
 	 * @param bKeepChanges - if true then the SceneColorSurface is resolved to the SceneColorTexture
@@ -215,7 +211,7 @@ public:
 	void BeginRenderingSceneAlphaCopy(FRHICommandListImmediate& RHICmdList);
 	void FinishRenderingSceneAlphaCopy(FRHICommandListImmediate& RHICmdList);
 
-	void BeginRenderingLightAttenuation(FRHICommandList& RHICmdList);
+	void BeginRenderingLightAttenuation(FRHICommandList& RHICmdList, bool bClearToWhite=false);
 	void FinishRenderingLightAttenuation(FRHICommandList& RHICmdList);
 
 	/**
diff --git a/Engine/Source/Runtime/Renderer/Public/RendererInterface.h b/Engine/Source/Runtime/Renderer/Public/RendererInterface.h
--- a/Engine/Source/Runtime/Renderer/Public/RendererInterface.h
+++ b/Engine/Source/Runtime/Renderer/Public/RendererInterface.h
@@ -412,6 +412,27 @@ public:
 	}
 };
 
+/** The empty vertex declaration resource type. */
+class FEmptyVertexDeclaration : public FRenderResource
+{
+public:
+	FVertexDeclarationRHIRef VertexDeclarationRHI;
+
+	/** Destructor. */
+	virtual ~FEmptyVertexDeclaration() {}
+
+	virtual void InitRHI()
+	{
+		FVertexDeclarationElementList Elements;
+		VertexDeclarationRHI = RHICreateVertexDeclaration(Elements);
+	}
+
+	virtual void ReleaseRHI()
+	{
+		VertexDeclarationRHI.SafeRelease();
+	}
+};
+
 // use r.DrawDenormalizedQuadMode to override the function call setting (quick way to see if an artifact is caused why this optimization)
 enum EDrawRectangleFlags
 {
