fixed UE-4597

diff --git a/Engine/Source/Runtime/Renderer/Private/CompositionLighting/CompositionLighting.cpp b/Engine/Source/Runtime/Renderer/Private/CompositionLighting/CompositionLighting.cpp
--- a/Engine/Source/Runtime/Renderer/Private/CompositionLighting/CompositionLighting.cpp
+++ b/Engine/Source/Runtime/Renderer/Private/CompositionLighting/CompositionLighting.cpp
@@ -254,11 +254,6 @@ void FCompositionLighting::ProcessBeforeBasePass(FRHICommandListImmediate& RHICm
 
 		SCOPED_DRAW_EVENT(RHICmdList, CompositionBeforeBasePass);
 
-		TRefCountPtr<IPooledRenderTarget>& SceneColor = GSceneRenderTargets.GetSceneColor();
-
-		Context.FinalOutput.GetOutput()->RenderTargetDesc = GSceneRenderTargets.GetSceneColor()->GetDesc();
-		Context.FinalOutput.GetOutput()->PooledRenderTarget = SceneColor;
-
 		// you can add multiple dependencies
 		CompositeContext.Root->AddDependency(Context.FinalOutput);
 
diff --git a/Engine/Source/Runtime/Renderer/Private/CompositionLighting/PostProcessDeferredDecals.cpp b/Engine/Source/Runtime/Renderer/Private/CompositionLighting/PostProcessDeferredDecals.cpp
--- a/Engine/Source/Runtime/Renderer/Private/CompositionLighting/PostProcessDeferredDecals.cpp
+++ b/Engine/Source/Runtime/Renderer/Private/CompositionLighting/PostProcessDeferredDecals.cpp
@@ -714,18 +714,11 @@ void FRCPassPostProcessDeferredDecals::Process(FRenderingCompositePassContext& C
 	bool bDBuffer = IsDBufferEnabled();
 	float bDecalPreStencil = CVarStencilSizeThreshold.GetValueOnRenderThread() >= 0;
 
-	{
-		FRenderingCompositeOutput* OutputOfMyInput = GetInput(ePId_Input0)->GetOutput();
-		PassOutputs[0].PooledRenderTarget = OutputOfMyInput->PooledRenderTarget;
-		OutputOfMyInput->RenderTargetDesc.DebugName = PassOutputs[0].RenderTargetDesc.DebugName;
-		PassOutputs[0].RenderTargetDesc = OutputOfMyInput->RenderTargetDesc;
-	}
-
 	SCOPED_DRAW_EVENT(RHICmdList, PostProcessDeferredDecals);
 
 	if(RenderStage == 0)
 	{
-		// before BasePass, onkly if DBuffer is enabled
+		// before BasePass, only if DBuffer is enabled
 
 		check(bDBuffer);
 
@@ -836,8 +829,7 @@ void FRCPassPostProcessDeferredDecals::Process(FRenderingCompositePassContext& C
 	}
 
 	if(SortedDecals.Num() > 0)
-	{		
-		FIntPoint DestSize = PassOutputs[0].RenderTargetDesc.Extent;
+	{
 		FIntRect SrcRect = View.ViewRect;
 		FIntRect DestRect = View.ViewRect;
 
@@ -938,7 +930,7 @@ void FRCPassPostProcessDeferredDecals::Process(FRenderingCompositePassContext& C
 					case RTM_SceneColorAndGBuffer:
 						{
 							FTextureRHIParamRef RenderTargets[5];
-							RenderTargets[0] = PassOutputs[0].PooledRenderTarget->GetRenderTargetItem().TargetableTexture;
+							RenderTargets[0] = GSceneRenderTargets.GetSceneColor()->GetRenderTargetItem().TargetableTexture;
 							RenderTargets[1] = GSceneRenderTargets.GBufferA->GetRenderTargetItem().TargetableTexture;
 							RenderTargets[2] = GSceneRenderTargets.GBufferB->GetRenderTargetItem().TargetableTexture;
 							RenderTargets[3] = GSceneRenderTargets.GBufferC->GetRenderTargetItem().TargetableTexture;
@@ -952,7 +944,7 @@ void FRCPassPostProcessDeferredDecals::Process(FRenderingCompositePassContext& C
 						break;
 					
 					case RTM_SceneColor:
-						SetRenderTarget(RHICmdList, PassOutputs[0].PooledRenderTarget->GetRenderTargetItem().TargetableTexture, GSceneRenderTargets.GetSceneDepthSurface());
+						SetRenderTarget(RHICmdList, GSceneRenderTargets.GetSceneColor()->GetRenderTargetItem().TargetableTexture, GSceneRenderTargets.GetSceneDepthSurface());
 						break;
 
 					case RTM_DBuffer:
diff --git a/Engine/Source/Runtime/Renderer/Private/PostProcess/PostProcessing.cpp b/Engine/Source/Runtime/Renderer/Private/PostProcess/PostProcessing.cpp
--- a/Engine/Source/Runtime/Renderer/Private/PostProcess/PostProcessing.cpp
+++ b/Engine/Source/Runtime/Renderer/Private/PostProcess/PostProcessing.cpp
@@ -83,7 +83,7 @@ static TAutoConsoleVariable<int32> CVarUpsampleQuality(
 static TAutoConsoleVariable<float> CVarMotionBlurSoftEdgeSize(
 	TEXT("r.MotionBlurSoftEdgeSize"),
 	1.0f,
-	TEXT("Defines how lange object motion blur is blurred (percent of screen width) to allow soft edge motion blur.\n")
+	TEXT("Defines how wide the object motion blur is blurred (percent of screen width) to allow soft edge motion blur.\n")
 	TEXT("This scales linearly with the size (up to a maximum of 32 samples, 2.5 is about 18 samples) and with screen resolution\n")
 	TEXT("Smaller values are better for performance and provide more accurate motion vectors but the blurring outside the object is reduced.\n")
 	TEXT("If needed this can be exposed like the other motionblur settings.\n")
@@ -97,8 +97,14 @@ IMPLEMENT_SHADER_TYPE(,FPostProcessVS,TEXT("PostProcessBloom"),TEXT("MainPostpro
 FPostprocessContext::FPostprocessContext(class FRenderingCompositionGraph& InGraph, const FViewInfo& InView)
 	: Graph(InGraph)
 	, View(InView)
+	, SceneColor(0)
+	, SceneDepth(0)
 {
-	SceneColor = Graph.RegisterPass(new(FMemStack::Get()) FRCPassPostProcessInput(GSceneRenderTargets.GetSceneColor()));
+	if(GSceneRenderTargets.IsSceneColorAllocated())
+	{
+		SceneColor = Graph.RegisterPass(new(FMemStack::Get()) FRCPassPostProcessInput(GSceneRenderTargets.GetSceneColor()));
+	}
+
 	SceneDepth = Graph.RegisterPass(new(FMemStack::Get()) FRCPassPostProcessInput(GSceneRenderTargets.SceneDepthZ));
 
 	FinalOutput = FRenderingCompositeOutputRef(SceneColor);
diff --git a/Engine/Source/Runtime/Renderer/Private/PostProcess/PostProcessing.h b/Engine/Source/Runtime/Renderer/Private/PostProcess/PostProcessing.h
--- a/Engine/Source/Runtime/Renderer/Private/PostProcess/PostProcessing.h
+++ b/Engine/Source/Runtime/Renderer/Private/PostProcess/PostProcessing.h
@@ -20,6 +20,7 @@ public:
 	FRenderingCompositionGraph& Graph;
 	const FViewInfo& View;
 
+	// 0 if there was no scene color available at constructor call time
 	FRenderingCompositePass* SceneColor;
 	FRenderingCompositePass* SceneDepth;
 
diff --git a/Engine/Source/Runtime/Renderer/Private/PostProcess/SceneRenderTargets.cpp b/Engine/Source/Runtime/Renderer/Private/PostProcess/SceneRenderTargets.cpp
--- a/Engine/Source/Runtime/Renderer/Private/PostProcess/SceneRenderTargets.cpp
+++ b/Engine/Source/Runtime/Renderer/Private/PostProcess/SceneRenderTargets.cpp
@@ -551,6 +551,11 @@ const TRefCountPtr<IPooledRenderTarget>& FSceneRenderTargets::GetSceneColor() co
 	return GetSceneColorForCurrentShadingPath();
 }
 
+bool FSceneRenderTargets::IsSceneColorAllocated() const
+{
+	return GetSceneColorForCurrentShadingPath() != 0;
+}
+
 TRefCountPtr<IPooledRenderTarget>& FSceneRenderTargets::GetSceneColor()
 {
 	if (!GetSceneColorForCurrentShadingPath())
diff --git a/Engine/Source/Runtime/Renderer/Private/PostProcess/SceneRenderTargets.h b/Engine/Source/Runtime/Renderer/Private/PostProcess/SceneRenderTargets.h
--- a/Engine/Source/Runtime/Renderer/Private/PostProcess/SceneRenderTargets.h
+++ b/Engine/Source/Runtime/Renderer/Private/PostProcess/SceneRenderTargets.h
@@ -392,6 +392,9 @@ public:
 
 	TRefCountPtr<IPooledRenderTarget>& GetSceneColor();
 
+	// changes depending at which part of the frame this is called
+	bool IsSceneColorAllocated() const;
+
 	void SetSceneColor(IPooledRenderTarget* In);
 
 	// ---
