Integrate from Speedtree branch: Fix SpeedTree wind crash when releasing instanced static mesh resources

diff --git a/Engine/Source/Runtime/Engine/Private/InstancedStaticMesh.cpp b/Engine/Source/Runtime/Engine/Private/InstancedStaticMesh.cpp
--- a/Engine/Source/Runtime/Engine/Private/InstancedStaticMesh.cpp
+++ b/Engine/Source/Runtime/Engine/Private/InstancedStaticMesh.cpp
@@ -620,7 +620,6 @@ public:
 
 	~FInstancedStaticMeshRenderData()
 	{
-		ReleaseResources();
 	}
 
 	void InitResources()
@@ -652,14 +651,14 @@ public:
 		}
 	}
 
-	void ReleaseResources()
+	void ReleaseResources(FSceneInterface* Scene, const UStaticMesh* StaticMesh)
 	{
 		// unregister SpeedTree wind with the scene
-		if (Component && Component->GetScene() && Component->StaticMesh && Component->StaticMesh->SpeedTreeWind.IsValid())
+		if (Scene && StaticMesh && StaticMesh->SpeedTreeWind.IsValid())
 		{
 			for (int32 LODIndex = 0; LODIndex < VertexFactories.Num(); LODIndex++)
 			{
-				Component->GetScene()->RemoveSpeedTreeWind(&VertexFactories[LODIndex], Component->StaticMesh);
+				Scene->RemoveSpeedTreeWind(&VertexFactories[LODIndex], StaticMesh);
 			}
 		}
 
@@ -905,6 +904,11 @@ public:
 		UserData_DeselectedInstances.bRenderSelected = false;
 	}
 
+	~FInstancedStaticMeshSceneProxy()
+	{
+		InstancedRenderData.ReleaseResources(GetScene( ), StaticMesh);
+	}
+
 	// FPrimitiveSceneProxy interface.
 	
 	virtual FPrimitiveViewRelevance GetViewRelevance(const FSceneView* View) override
diff --git a/Engine/Source/Runtime/Engine/Public/SceneInterface.h b/Engine/Source/Runtime/Engine/Public/SceneInterface.h
--- a/Engine/Source/Runtime/Engine/Public/SceneInterface.h
+++ b/Engine/Source/Runtime/Engine/Public/SceneInterface.h
@@ -200,13 +200,13 @@ public:
 	 * Adds a SpeedTree wind computation object to the scene.
 	 * @param StaticMesh - The SpeedTree static mesh whose wind to add.
 	 */
-	virtual void AddSpeedTreeWind(class FVertexFactory* VertexFactory, class UStaticMesh* StaticMesh) = 0;
+	virtual void AddSpeedTreeWind(class FVertexFactory* VertexFactory, const class UStaticMesh* StaticMesh) = 0;
 
 	/** 
 	 * Removes a SpeedTree wind computation object to the scene.
 	 * @param StaticMesh - The SpeedTree static mesh whose wind to remove.
 	 */
-	virtual void RemoveSpeedTreeWind(class FVertexFactory* VertexFactory, class UStaticMesh* StaticMesh) = 0;
+	virtual void RemoveSpeedTreeWind(class FVertexFactory* VertexFactory, const class UStaticMesh* StaticMesh) = 0;
 
 	/** Ticks the SpeedTree wind object and updates the uniform buffer. */
 	virtual void UpdateSpeedTreeWind(double CurrentTime) = 0;
diff --git a/Engine/Source/Runtime/Renderer/Private/RendererScene.cpp b/Engine/Source/Runtime/Renderer/Private/RendererScene.cpp
--- a/Engine/Source/Runtime/Renderer/Private/RendererScene.cpp
+++ b/Engine/Source/Runtime/Renderer/Private/RendererScene.cpp
@@ -1341,14 +1341,14 @@ FVector4 FScene::GetDirectionalWindParameters(void) const
 	return NumActiveWindSources > 0 ? AccumulatedDirectionAndSpeed / NumActiveWindSources : FVector4(0,0,1,0);
 }
 
-void FScene::AddSpeedTreeWind(FVertexFactory* VertexFactory, UStaticMesh* StaticMesh)
+void FScene::AddSpeedTreeWind(FVertexFactory* VertexFactory, const UStaticMesh* StaticMesh)
 {
 	if (StaticMesh != NULL && StaticMesh->SpeedTreeWind.IsValid() && StaticMesh->RenderData.IsValid())
 	{
 		ENQUEUE_UNIQUE_RENDER_COMMAND_THREEPARAMETER(
 			FAddSpeedTreeWindCommand,
 			FScene*,Scene,this,
-			UStaticMesh*,StaticMesh,StaticMesh,
+			const UStaticMesh*,StaticMesh,StaticMesh,
 			FVertexFactory*,VertexFactory,VertexFactory,
 			{
 				Scene->SpeedTreeVertexFactoryMap.Add(VertexFactory, StaticMesh);
@@ -1369,14 +1369,14 @@ void FScene::AddSpeedTreeWind(FVertexFactory* VertexFactory, UStaticMesh* Static
 	}
 }
 
-void FScene::RemoveSpeedTreeWind(class FVertexFactory* VertexFactory, class UStaticMesh* StaticMesh)
+void FScene::RemoveSpeedTreeWind(class FVertexFactory* VertexFactory, const class UStaticMesh* StaticMesh)
 {
 	if (StaticMesh != NULL && StaticMesh->SpeedTreeWind.IsValid() && StaticMesh->RenderData.IsValid())
 	{
 		ENQUEUE_UNIQUE_RENDER_COMMAND_THREEPARAMETER(
 			FRemoveSpeedTreeWindCommand,
 			FScene*,Scene,this,
-			UStaticMesh*, StaticMesh, StaticMesh,
+			const UStaticMesh*, StaticMesh, StaticMesh,
 			FVertexFactory*,VertexFactory,VertexFactory,
 			{
 				FSpeedTreeWindComputation** WindComputationRef = Scene->SpeedTreeWindComputationMap.Find(StaticMesh);
@@ -1415,9 +1415,9 @@ void FScene::UpdateSpeedTreeWind(double CurrentTime)
 		{
 			FVector4 WindInfo = Scene->GetDirectionalWindParameters();
 
-			for (TMap<UStaticMesh*, FSpeedTreeWindComputation*>::TIterator It(Scene->SpeedTreeWindComputationMap); It; ++It )
+			for (TMap<const UStaticMesh*, FSpeedTreeWindComputation*>::TIterator It(Scene->SpeedTreeWindComputationMap); It; ++It )
 			{
-				UStaticMesh* StaticMesh = It.Key();
+				const UStaticMesh* StaticMesh = It.Key();
 				FSpeedTreeWindComputation* WindComputation = It.Value();
 
 				if( !StaticMesh->RenderData )
@@ -1426,7 +1426,7 @@ void FScene::UpdateSpeedTreeWind(double CurrentTime)
 					continue;
 				}
 
-				if (StaticMesh->SpeedTreeWind->NeedsReload( ))
+				if (GIsEditor && StaticMesh->SpeedTreeWind->NeedsReload( ))
 				{
 					// reload the wind since it may have changed or been scaled differently during reimport
 					StaticMesh->SpeedTreeWind->SetNeedsReload(false);
@@ -1480,7 +1480,7 @@ FUniformBufferRHIParamRef FScene::GetSpeedTreeUniformBuffer(const FVertexFactory
 {
 	if (VertexFactory != NULL)
 	{
-		UStaticMesh** StaticMesh = SpeedTreeVertexFactoryMap.Find(VertexFactory);
+		const UStaticMesh** StaticMesh = SpeedTreeVertexFactoryMap.Find(VertexFactory);
 		if (StaticMesh != NULL)
 		{
 			FSpeedTreeWindComputation** WindComputation = SpeedTreeWindComputationMap.Find(*StaticMesh);
@@ -2102,8 +2102,8 @@ public:
 	}
 	virtual FVector4 GetWindParameters(const FVector& Position) const { return FVector4(0,0,1,0); }
 	virtual FVector4 GetDirectionalWindParameters() const { return FVector4(0,0,1,0); }
-	virtual void AddSpeedTreeWind(class FVertexFactory* VertexFactory, class UStaticMesh* StaticMesh) {}
-	virtual void RemoveSpeedTreeWind(class FVertexFactory* VertexFactory, class UStaticMesh* StaticMesh) {}
+	virtual void AddSpeedTreeWind(class FVertexFactory* VertexFactory, const class UStaticMesh* StaticMesh) {}
+	virtual void RemoveSpeedTreeWind(class FVertexFactory* VertexFactory, const class UStaticMesh* StaticMesh) {}
 	virtual void UpdateSpeedTreeWind(double CurrentTime) {}
 	virtual FUniformBufferRHIParamRef GetSpeedTreeUniformBuffer(const FVertexFactory* VertexFactory) { return FUniformBufferRHIParamRef(); }
 
diff --git a/Engine/Source/Runtime/Renderer/Private/ScenePrivate.h b/Engine/Source/Runtime/Renderer/Private/ScenePrivate.h
--- a/Engine/Source/Runtime/Renderer/Private/ScenePrivate.h
+++ b/Engine/Source/Runtime/Renderer/Private/ScenePrivate.h
@@ -1281,8 +1281,8 @@ public:
 	TArray<class FWindSourceSceneProxy*> WindSources;
 
 	/** SpeedTree wind objects in the scene. FLocalVertexFactoryShaderParameters needs to lookup by FVertexFactory, but wind objects are per tree (i.e. per UStaticMesh)*/
-	TMap<UStaticMesh*, FSpeedTreeWindComputation*> SpeedTreeWindComputationMap;
-	TMap<FVertexFactory*, UStaticMesh*> SpeedTreeVertexFactoryMap;
+	TMap<const UStaticMesh*, FSpeedTreeWindComputation*> SpeedTreeWindComputationMap;
+	TMap<FVertexFactory*, const UStaticMesh*> SpeedTreeVertexFactoryMap;
 
 	/** The attachment groups in the scene.  The map key is the attachment group's root primitive. */
 	TMap<FPrimitiveComponentId,FAttachmentGroupSceneInfo> AttachmentGroups;
@@ -1355,8 +1355,8 @@ public:
 	virtual const TArray<FWindSourceSceneProxy*>& GetWindSources_RenderThread() const;
 	virtual FVector4 GetWindParameters(const FVector& Position) const;
 	virtual FVector4 GetDirectionalWindParameters() const;
-	virtual void AddSpeedTreeWind(FVertexFactory* VertexFactory, UStaticMesh* StaticMesh);
-	virtual void RemoveSpeedTreeWind(FVertexFactory* VertexFactory, UStaticMesh* StaticMesh);
+	virtual void AddSpeedTreeWind(FVertexFactory* VertexFactory, const UStaticMesh* StaticMesh);
+	virtual void RemoveSpeedTreeWind(FVertexFactory* VertexFactory, const UStaticMesh* StaticMesh);
 	virtual void UpdateSpeedTreeWind(double CurrentTime);
 	virtual FUniformBufferRHIParamRef GetSpeedTreeUniformBuffer(const FVertexFactory* VertexFactory);
 	virtual void DumpUnbuiltLightIteractions( FOutputDevice& Ar ) const;
