Integrate from speedtree branch: Made RemoveSpeedTreeWind_RenderThread function for calling from instanced static mesh

diff --git a/Engine/Source/Runtime/Engine/Private/InstancedStaticMesh.cpp b/Engine/Source/Runtime/Engine/Private/InstancedStaticMesh.cpp
--- a/Engine/Source/Runtime/Engine/Private/InstancedStaticMesh.cpp
+++ b/Engine/Source/Runtime/Engine/Private/InstancedStaticMesh.cpp
@@ -658,7 +658,7 @@ public:
 		{
 			for (int32 LODIndex = 0; LODIndex < VertexFactories.Num(); LODIndex++)
 			{
-				Scene->RemoveSpeedTreeWind(&VertexFactories[LODIndex], StaticMesh);
+				Scene->RemoveSpeedTreeWind_RenderThread(&VertexFactories[LODIndex], StaticMesh);
 			}
 		}
 
diff --git a/Engine/Source/Runtime/Engine/Public/SceneInterface.h b/Engine/Source/Runtime/Engine/Public/SceneInterface.h
--- a/Engine/Source/Runtime/Engine/Public/SceneInterface.h
+++ b/Engine/Source/Runtime/Engine/Public/SceneInterface.h
@@ -207,6 +207,7 @@ public:
 	 * @param StaticMesh - The SpeedTree static mesh whose wind to remove.
 	 */
 	virtual void RemoveSpeedTreeWind(class FVertexFactory* VertexFactory, const class UStaticMesh* StaticMesh) = 0;
+	virtual void RemoveSpeedTreeWind_RenderThread(class FVertexFactory* VertexFactory, const class UStaticMesh* StaticMesh) = 0;
 
 	/** Ticks the SpeedTree wind object and updates the uniform buffer. */
 	virtual void UpdateSpeedTreeWind(double CurrentTime) = 0;
diff --git a/Engine/Source/Runtime/Renderer/Private/RendererScene.cpp b/Engine/Source/Runtime/Renderer/Private/RendererScene.cpp
--- a/Engine/Source/Runtime/Renderer/Private/RendererScene.cpp
+++ b/Engine/Source/Runtime/Renderer/Private/RendererScene.cpp
@@ -1379,28 +1379,33 @@ void FScene::RemoveSpeedTreeWind(class FVertexFactory* VertexFactory, const clas
 			const UStaticMesh*, StaticMesh, StaticMesh,
 			FVertexFactory*,VertexFactory,VertexFactory,
 			{
-				FSpeedTreeWindComputation** WindComputationRef = Scene->SpeedTreeWindComputationMap.Find(StaticMesh);
-				if (WindComputationRef != NULL)
-				{
-					FSpeedTreeWindComputation* WindComputation = *WindComputationRef;
+				Scene->RemoveSpeedTreeWind_RenderThread(VertexFactory, StaticMesh);
+			});
+	}
+}
 
-					WindComputation->ReferenceCount--;
-					if (WindComputation->ReferenceCount < 1)
-					{
-						for (auto Iter = Scene->SpeedTreeVertexFactoryMap.CreateIterator(); Iter; ++Iter )
-						{
-							if (Iter.Value() == StaticMesh)
-							{
-								Iter.RemoveCurrent();
-							}
-						}
-
-						Scene->SpeedTreeWindComputationMap.Remove(StaticMesh);
-						WindComputation->UniformBuffer.ReleaseResource();
-						delete WindComputation;
-					}
+void FScene::RemoveSpeedTreeWind_RenderThread(class FVertexFactory* VertexFactory, const class UStaticMesh* StaticMesh)
+{
+	FSpeedTreeWindComputation** WindComputationRef = SpeedTreeWindComputationMap.Find(StaticMesh);
+	if (WindComputationRef != NULL)
+	{
+		FSpeedTreeWindComputation* WindComputation = *WindComputationRef;
+
+		WindComputation->ReferenceCount--;
+		if (WindComputation->ReferenceCount < 1)
+		{
+			for (auto Iter = SpeedTreeVertexFactoryMap.CreateIterator(); Iter; ++Iter )
+			{
+				if (Iter.Value() == StaticMesh)
+				{
+					Iter.RemoveCurrent();
 				}
-			});
+			}
+
+			SpeedTreeWindComputationMap.Remove(StaticMesh);
+			WindComputation->UniformBuffer.ReleaseResource();
+			delete WindComputation;
+		}
 	}
 }
 
@@ -2104,6 +2109,7 @@ public:
 	virtual FVector4 GetDirectionalWindParameters() const { return FVector4(0,0,1,0); }
 	virtual void AddSpeedTreeWind(class FVertexFactory* VertexFactory, const class UStaticMesh* StaticMesh) {}
 	virtual void RemoveSpeedTreeWind(class FVertexFactory* VertexFactory, const class UStaticMesh* StaticMesh) {}
+	virtual void RemoveSpeedTreeWind_RenderThread(class FVertexFactory* VertexFactory, const class UStaticMesh* StaticMesh) {}
 	virtual void UpdateSpeedTreeWind(double CurrentTime) {}
 	virtual FUniformBufferRHIParamRef GetSpeedTreeUniformBuffer(const FVertexFactory* VertexFactory) { return FUniformBufferRHIParamRef(); }
 
diff --git a/Engine/Source/Runtime/Renderer/Private/ScenePrivate.h b/Engine/Source/Runtime/Renderer/Private/ScenePrivate.h
--- a/Engine/Source/Runtime/Renderer/Private/ScenePrivate.h
+++ b/Engine/Source/Runtime/Renderer/Private/ScenePrivate.h
@@ -1357,6 +1357,7 @@ public:
 	virtual FVector4 GetDirectionalWindParameters() const;
 	virtual void AddSpeedTreeWind(FVertexFactory* VertexFactory, const UStaticMesh* StaticMesh);
 	virtual void RemoveSpeedTreeWind(FVertexFactory* VertexFactory, const UStaticMesh* StaticMesh);
+	virtual void RemoveSpeedTreeWind_RenderThread(FVertexFactory* VertexFactory, const UStaticMesh* StaticMesh);
 	virtual void UpdateSpeedTreeWind(double CurrentTime);
 	virtual FUniformBufferRHIParamRef GetSpeedTreeUniformBuffer(const FVertexFactory* VertexFactory);
 	virtual void DumpUnbuiltLightIteractions( FOutputDevice& Ar ) const;
