Removed dead material texture code

diff --git a/Engine/Source/Runtime/Renderer/Private/ShaderBaseClasses.cpp b/Engine/Source/Runtime/Renderer/Private/ShaderBaseClasses.cpp
--- a/Engine/Source/Runtime/Renderer/Private/ShaderBaseClasses.cpp
+++ b/Engine/Source/Runtime/Renderer/Private/ShaderBaseClasses.cpp
@@ -149,66 +149,6 @@ void FMaterialShader::SetParameters(
 		}
 	}
 
-	// Set 2D texture uniform expressions.
-	{
-		int32 Count = Uniform2DTextureShaderResourceParameters.Num();
-		for(int32 ParameterIndex = 0;ParameterIndex < Count;ParameterIndex++)
-		{
-			const TUniformParameter<FShaderResourceParameter>& TextureResourceParameter = Uniform2DTextureShaderResourceParameters[ParameterIndex];
-			check(UniformExpressionCache->Textures.IsValidIndex(TextureResourceParameter.Index));
-			const TUniformParameter<FShaderResourceParameter>& SamplerResourceParameter = Uniform2DTextureSamplerShaderResourceParameters[ParameterIndex];
-
-			const UTexture* Texture2D = UniformExpressionCache->Textures[TextureResourceParameter.Index];
-			const FTexture* Value = Texture2D ? Texture2D->Resource : NULL;
-
-			if( !Value )
-			{
-				Value = GWhiteTexture;
-			}
-
-			checkSlow(Value);
-			Value->LastRenderTime = FApp::GetCurrentTime();
-
-			SetTextureParameter(
-				RHICmdList, 
-				ShaderRHI, 
-				TextureResourceParameter.ShaderParameter, 
-				SamplerResourceParameter.ShaderParameter, 
-				bDeferredPass && Value->DeferredPassSamplerStateRHI ? Value->DeferredPassSamplerStateRHI : Value->SamplerStateRHI, 
-				Value->TextureRHI);
-		}
-	}
-
-	// Set cube texture uniform expressions.
-	{
-		int32 Count = UniformCubeTextureShaderResourceParameters.Num(); 
-		for(int32 ParameterIndex = 0;ParameterIndex < Count;ParameterIndex++)
-		{
-			const TUniformParameter<FShaderResourceParameter>& TextureResourceParameter = UniformCubeTextureShaderResourceParameters[ParameterIndex];
-			check(UniformExpressionCache->CubeTextures.IsValidIndex(TextureResourceParameter.Index));
-			const TUniformParameter<FShaderResourceParameter>& SamplerResourceParameter = UniformCubeTextureSamplerShaderResourceParameters[ParameterIndex];
-
-			const UTexture* CubeTexture = UniformExpressionCache->CubeTextures[TextureResourceParameter.Index];
-			const FTexture* Value = CubeTexture ? CubeTexture->Resource : NULL;
-
-			if( !Value )
-			{
-				Value = GWhiteTextureCube;
-			}
-
-			checkSlow(Value);
-			Value->LastRenderTime = FApp::GetCurrentTime();
-
-			SetTextureParameter(
-				RHICmdList, 
-				ShaderRHI,
-				TextureResourceParameter.ShaderParameter,
-				SamplerResourceParameter.ShaderParameter,
-				bDeferredPass && Value->DeferredPassSamplerStateRHI ? Value->DeferredPassSamplerStateRHI : Value->SamplerStateRHI, 
-				Value->TextureRHI);
-		}
-	}
-
 	DeferredParameters.Set(RHICmdList, ShaderRHI, View, TextureMode);
 
 	AtmosphericFogTextureParameters.Set(RHICmdList, ShaderRHI, View);
diff --git a/Engine/Source/Runtime/Renderer/Public/MaterialShader.h b/Engine/Source/Runtime/Renderer/Public/MaterialShader.h
--- a/Engine/Source/Runtime/Renderer/Public/MaterialShader.h
+++ b/Engine/Source/Runtime/Renderer/Public/MaterialShader.h
@@ -128,10 +128,6 @@ private:
 
 	FShaderUniformBufferParameter MaterialUniformBuffer;
 	TArray<FShaderUniformBufferParameter> ParameterCollectionUniformBuffers;
-	TArray<TUniformParameter<FShaderResourceParameter> > Uniform2DTextureShaderResourceParameters;
-	TArray<TUniformParameter<FShaderResourceParameter> > Uniform2DTextureSamplerShaderResourceParameters;
-	TArray<TUniformParameter<FShaderResourceParameter> > UniformCubeTextureShaderResourceParameters;
-	TArray<TUniformParameter<FShaderResourceParameter> > UniformCubeTextureSamplerShaderResourceParameters;
 	FDeferredPixelShaderParameters DeferredParameters;
 	FShaderResourceParameter LightAttenuation;
 	FShaderResourceParameter LightAttenuationSampler;
