UE4 - UE-3057 - Scene captures not working on non-mobilehdr projects

diff --git a/Engine/Source/Runtime/Renderer/Private/SceneCaptureRendering.cpp b/Engine/Source/Runtime/Renderer/Private/SceneCaptureRendering.cpp
--- a/Engine/Source/Runtime/Renderer/Private/SceneCaptureRendering.cpp
+++ b/Engine/Source/Runtime/Renderer/Private/SceneCaptureRendering.cpp
@@ -12,8 +12,56 @@
 #include "ScreenRendering.h"
 #include "PostProcessAmbient.h"
 #include "PostProcessing.h"
+#include "SceneUtils.h"
 
-void UpdateSceneCaptureContent_RenderThread(FRHICommandListImmediate& RHICmdList, FSceneRenderer* SceneRenderer, FTextureRenderTargetResource* TextureRenderTarget, const FName OwnerName, const FResolveParams& ResolveParams, bool bUseSceneColorTexture)
+// Copies into render target, optionally flipping it in the Y-axis
+static void CopyCaptureToTarget(FRHICommandListImmediate& RHICmdList, const FRenderTarget* Target, const FIntPoint& TargetSize, FViewInfo& View, const FIntRect& ViewRect, FTextureRHIParamRef SourceTextureRHI, bool bNeedsFlippedRenderTarget)
+{
+	SetRenderTarget(RHICmdList, Target->GetRenderTargetTexture(), NULL);
+
+	RHICmdList.SetRasterizerState(TStaticRasterizerState<FM_Solid, CM_None>::GetRHI());
+	RHICmdList.SetDepthStencilState(TStaticDepthStencilState<false, CF_Always>::GetRHI());
+	RHICmdList.SetBlendState(TStaticBlendState<>::GetRHI());
+
+	TShaderMapRef<FScreenVS> VertexShader(View.ShaderMap);
+	TShaderMapRef<FScreenPS> PixelShader(View.ShaderMap);
+	static FGlobalBoundShaderState BoundShaderState;
+	SetGlobalBoundShaderState(RHICmdList, View.GetFeatureLevel(), BoundShaderState, GFilterVertexDeclaration.VertexDeclarationRHI, *VertexShader, *PixelShader);
+
+	FRenderingCompositePassContext Context(RHICmdList, View);
+
+	VertexShader->SetParameters(RHICmdList, View);
+	PixelShader->SetParameters(RHICmdList, TStaticSamplerState<SF_Point>::GetRHI(), SourceTextureRHI);
+
+	if (bNeedsFlippedRenderTarget)
+	{
+		DrawRectangle(
+			RHICmdList,
+			ViewRect.Min.X, ViewRect.Min.Y,
+			ViewRect.Width(), ViewRect.Height(),
+			ViewRect.Min.X, ViewRect.Height() - ViewRect.Min.Y,
+			ViewRect.Width(), -ViewRect.Height(),
+			TargetSize,
+			TargetSize,
+			*VertexShader,
+			EDRF_UseTriangleOptimization);
+	}
+	else
+	{
+		DrawRectangle(
+			RHICmdList,
+			ViewRect.Min.X, ViewRect.Min.Y,
+			ViewRect.Width(), ViewRect.Height(),
+			ViewRect.Min.X, ViewRect.Min.Y,
+			ViewRect.Width(), ViewRect.Height(),
+			TargetSize,
+			GSceneRenderTargets.GetBufferSizeXY(),
+			*VertexShader,
+			EDRF_UseTriangleOptimization);
+	}
+}
+
+static void UpdateSceneCaptureContent_RenderThread(FRHICommandListImmediate& RHICmdList, FSceneRenderer* SceneRenderer, FTextureRenderTargetResource* TextureRenderTarget, const FName OwnerName, const FResolveParams& ResolveParams, bool bUseSceneColorTexture)
 {
 	FMemMark MemStackMark(FMemStack::Get());
 
@@ -25,49 +73,82 @@ void UpdateSceneCaptureContent_RenderThread(FRHICommandListImmediate& RHICmdList
 		FString EventName;
 		OwnerName.ToString(EventName);
 		SCOPED_DRAW_EVENTF(RHICmdList, SceneCapture, TEXT("SceneCapture %s"), *EventName);
+#else
+		SCOPED_DRAW_EVENT(RHICmdList, UpdateSceneCaptureContent_RenderThread);
 #endif
 
-		// Render the scene normally
+		const bool bIsMobileHDR = IsMobileHDR();
+		const bool bRHINeedsFlip = RHINeedsToSwitchVerticalAxis(GMaxRHIShaderPlatform);
+		const bool bNeedsFlippedRenderTarget = !bIsMobileHDR && bRHINeedsFlip;
+
+		// Intermediate render target that will need to be flipped (needed on !IsMobileHDR())
+		TRefCountPtr<IPooledRenderTarget> FlippedPooledRenderTarget;
+
 		const FRenderTarget* Target = SceneRenderer->ViewFamily.RenderTarget;
+		if (bNeedsFlippedRenderTarget)
+		{
+			// We need to use an intermediate render target since the result will be flipped
+			auto& RenderTarget = Target->GetRenderTargetTexture();
+			FPooledRenderTargetDesc Desc(FPooledRenderTargetDesc::Create2DDesc(Target->GetSizeXY(), 
+				RenderTarget.GetReference()->GetFormat(), 
+				TexCreate_None, 
+				TexCreate_RenderTargetable,
+				false));
+			GRenderTargetPool.FindFreeElement(Desc, FlippedPooledRenderTarget, TEXT("SceneCaptureFlipped"));
+		}
+
+		// Helper class to allow setting render target
+		struct FRenderTargetOverride : public FRenderTarget
+		{
+			FRenderTargetOverride(FRHITexture2D* In)
+			{
+				RenderTargetTextureRHI = In;
+			}
+
+			virtual FIntPoint GetSizeXY() const { return FIntPoint(RenderTargetTextureRHI->GetSizeX(), RenderTargetTextureRHI->GetSizeY()); }
+
+			FTexture2DRHIRef GetTextureParamRef() { return RenderTargetTextureRHI; }
+		} FlippedRenderTarget(
+			FlippedPooledRenderTarget.GetReference()
+			? FlippedPooledRenderTarget.GetReference()->GetRenderTargetItem().TargetableTexture->GetTexture2D() 
+			: nullptr);
 		FViewInfo& View = SceneRenderer->Views[0];
 		FIntRect ViewRect = View.ViewRect;
 		FIntRect UnconstrainedViewRect = View.UnconstrainedViewRect;
 		SetRenderTarget(RHICmdList, Target->GetRenderTargetTexture(), NULL);
 		RHICmdList.Clear(true, FLinearColor::Black, false, 1.0f, false, 0, ViewRect);
-		SceneRenderer->Render(RHICmdList);
 
-		// Copy the captured scene into the destination texture
-		if (bUseSceneColorTexture)
+		// Render the scene normally
 		{
-			// Copy the captured scene into the destination texture
-			SetRenderTarget(RHICmdList, Target->GetRenderTargetTexture(), NULL);
-
-			RHICmdList.SetRasterizerState(TStaticRasterizerState<FM_Solid, CM_None>::GetRHI());
-			RHICmdList.SetDepthStencilState(TStaticDepthStencilState<false, CF_Always>::GetRHI());
-			RHICmdList.SetBlendState(TStaticBlendState<>::GetRHI());
+			SCOPED_DRAW_EVENT(RHICmdList, RenderScene);
 
-			TShaderMapRef<FScreenVS> VertexShader(View.ShaderMap);
-			TShaderMapRef<FScreenPS> PixelShader(View.ShaderMap);
-			static FGlobalBoundShaderState BoundShaderState;
-			SetGlobalBoundShaderState(RHICmdList, View.GetFeatureLevel(), BoundShaderState, GFilterVertexDeclaration.VertexDeclarationRHI, *VertexShader, *PixelShader);
-
-			FRenderingCompositePassContext Context(RHICmdList, View);
+			if (bNeedsFlippedRenderTarget)
+			{
+				// Hijack the render target
+				SceneRenderer->ViewFamily.RenderTarget = &FlippedRenderTarget;
+			}
+			SceneRenderer->Render(RHICmdList);
+			if (bNeedsFlippedRenderTarget)
+			{
+				// And restore it
+				SceneRenderer->ViewFamily.RenderTarget = Target;
+			}
+		}
 
-			VertexShader->SetParameters(RHICmdList, View);
-			PixelShader->SetParameters(RHICmdList, TStaticSamplerState<SF_Point>::GetRHI(), GSceneRenderTargets.GetSceneColorTexture());
+		if (bNeedsFlippedRenderTarget)
+		{
+			// We need to flip this texture upside down (since we depended on tonemapping to fix this on the hdr path)
+			SCOPED_DRAW_EVENT(RHICmdList, FlipCapture);
 
 			FIntPoint TargetSize(UnconstrainedViewRect.Width(), UnconstrainedViewRect.Height());
-
-			DrawRectangle(
-				RHICmdList,
-				ViewRect.Min.X, ViewRect.Min.Y,
-				ViewRect.Width(), ViewRect.Height(),
-				ViewRect.Min.X, ViewRect.Min.Y,
-				ViewRect.Width(), ViewRect.Height(),
-				TargetSize,
-				GSceneRenderTargets.GetBufferSizeXY(),
-				*VertexShader,
-				EDRF_UseTriangleOptimization);
+			CopyCaptureToTarget(RHICmdList, Target, TargetSize, View, ViewRect, FlippedRenderTarget.GetTextureParamRef(), true);
+		}
+		else if (bUseSceneColorTexture && (bIsMobileHDR || SceneRenderer->FeatureLevel >= ERHIFeatureLevel::SM4))
+		{
+			// Copy the captured scene into the destination texture (only required on HDR or deferred as that implies post-processing)
+			SCOPED_DRAW_EVENT(RHICmdList, CaptureSceneColor);
+			FIntPoint TargetSize(UnconstrainedViewRect.Width(), UnconstrainedViewRect.Height());
+			CopyCaptureToTarget(RHICmdList, Target, TargetSize, View, ViewRect, GSceneRenderTargets.GetSceneColorTexture(), false);
 		}
 
 		RHICmdList.CopyToResolveTarget(TextureRenderTarget->GetRenderTargetTexture(), TextureRenderTarget->TextureRHI, false, ResolveParams);
