- Fixed depth buffer not clearing in Metal forward renderer

diff --git a/Engine/Source/Runtime/Renderer/Private/ForwardShadingRenderer.cpp b/Engine/Source/Runtime/Renderer/Private/ForwardShadingRenderer.cpp
--- a/Engine/Source/Runtime/Renderer/Private/ForwardShadingRenderer.cpp
+++ b/Engine/Source/Runtime/Renderer/Private/ForwardShadingRenderer.cpp
@@ -113,13 +113,9 @@ void FForwardShadingSceneRenderer::Render(FRHICommandListImmediate& RHICmdList)
 	else
 	{
 		// Begin rendering to scene color
-		GSceneRenderTargets.BeginRenderingSceneColor(RHICmdList);
+        GSceneRenderTargets.BeginRenderingSceneColor(RHICmdList, ESimpleRenderTargetMode::EClearToDefault);
 	}
 
-	// Clear color and depth buffer
-	// Note, this is a reversed Z depth surface, so 0.0f is the far plane.
-	RHICmdList.Clear(true, FLinearColor::Black, true, 0.0f, true, 0, FIntRect());
-
 	RenderForwardShadingBasePass(RHICmdList);
 
 	// Make a copy of the scene depth if the current hardware doesn't support reading and writing to the same depth buffer
