Fix PS4 stencil clear assert.

diff --git a/Engine/Source/Runtime/Renderer/Private/PostProcess/SceneRenderTargets.cpp b/Engine/Source/Runtime/Renderer/Private/PostProcess/SceneRenderTargets.cpp
--- a/Engine/Source/Runtime/Renderer/Private/PostProcess/SceneRenderTargets.cpp
+++ b/Engine/Source/Runtime/Renderer/Private/PostProcess/SceneRenderTargets.cpp
@@ -783,7 +783,14 @@ void FSceneRenderTargets::FinishRenderingPrePass(FRHICommandListImmediate& RHICm
 void FSceneRenderTargets::BeginRenderingShadowDepth(FRHICommandList& RHICmdList, bool bClear)
 {
 	GRenderTargetPool.VisualizeTexture.SetCheckPoint(RHICmdList, ShadowDepthZ);
-	FRHISetRenderTargetsInfo Info(0, nullptr, FRHIDepthRenderTargetView(GetShadowDepthZSurface(), bClear ? ERenderTargetLoadAction::EClear : ERenderTargetLoadAction::ELoad, ERenderTargetStoreAction::EStore));
+
+	//All of the shadow passes are written using 'clear stencil after', so keep stencil as-is
+	FRHISetRenderTargetsInfo Info(0, nullptr, FRHIDepthRenderTargetView(GetShadowDepthZSurface(), 
+		bClear ? ERenderTargetLoadAction::EClear : ERenderTargetLoadAction::ELoad, 
+		ERenderTargetStoreAction::EStore, 
+		ERenderTargetLoadAction::ELoad, 
+		ERenderTargetStoreAction::EStore));
+
 	Info.DepthClearValue = 1.0f;
 	Info.ColorRenderTarget[0].StoreAction = ERenderTargetStoreAction::ENoAction;
 
