Scene captures that update every frame now benefit from occlusion queries

diff --git a/Engine/Source/Runtime/Engine/Classes/Components/SceneCaptureComponent.h b/Engine/Source/Runtime/Engine/Classes/Components/SceneCaptureComponent.h
--- a/Engine/Source/Runtime/Engine/Classes/Components/SceneCaptureComponent.h
+++ b/Engine/Source/Runtime/Engine/Classes/Components/SceneCaptureComponent.h
@@ -31,5 +31,16 @@ public:
 	/** Adds all primitive components in the actor to our list of hidden components. */
 	UFUNCTION(BlueprintCallable, Category = "Rendering|SceneCapture")
 	ENGINE_API void HideActorComponents(AActor* InActor);
+
+	/** Returns the view state, if any, and allocates one if needed. This function can return NULL, e.g. when bCaptureEveryFrame is false. */
+	ENGINE_API FSceneViewStateInterface* GetViewState();
+	
+private:
+
+	/**
+	 * The view state holds persistent scene rendering state and enables occlusion culling in scene captures.
+	 * NOTE: This object is used by the rendering thread. The smart pointer's destructor will destroy it and that
+	 */
+	FSceneViewStateReference ViewState;
 };
 
diff --git a/Engine/Source/Runtime/Engine/Private/SceneCaptureComponent.cpp b/Engine/Source/Runtime/Engine/Private/SceneCaptureComponent.cpp
--- a/Engine/Source/Runtime/Engine/Private/SceneCaptureComponent.cpp
+++ b/Engine/Source/Runtime/Engine/Private/SceneCaptureComponent.cpp
@@ -172,6 +172,22 @@ void USceneCaptureComponent::HideActorComponents(AActor* InActor)
 	}
 }
 
+FSceneViewStateInterface* USceneCaptureComponent::GetViewState()
+{
+	FSceneViewStateInterface* ViewStateInterface = ViewState.GetReference();
+	if (bCaptureEveryFrame && ViewStateInterface == NULL)
+	{
+		ViewState.Allocate();
+		ViewStateInterface = ViewState.GetReference();
+	}
+	else if (!bCaptureEveryFrame && ViewStateInterface)
+	{
+		ViewState.Destroy();
+		ViewStateInterface = NULL;
+	}
+	return ViewStateInterface;
+}
+
 // -----------------------------------------------
 
 
diff --git a/Engine/Source/Runtime/Renderer/Private/RendererScene.cpp b/Engine/Source/Runtime/Renderer/Private/RendererScene.cpp
--- a/Engine/Source/Runtime/Renderer/Private/RendererScene.cpp
+++ b/Engine/Source/Runtime/Renderer/Private/RendererScene.cpp
@@ -31,6 +31,7 @@ SIZE_T FStaticMeshDrawListBase::TotalBytesUsed = 0;
 FSceneViewState::FSceneViewState()
 	: OcclusionQueryPool(RQT_Occlusion)
 {
+	OcclusionFrameCounter = 0;
 	LastRenderTime = -FLT_MAX;
 	LastRenderTimeDelta = 0.0f;
 	MotionBlurTimeScale = 1.0f;
diff --git a/Engine/Source/Runtime/Renderer/Private/SceneCaptureRendering.cpp b/Engine/Source/Runtime/Renderer/Private/SceneCaptureRendering.cpp
--- a/Engine/Source/Runtime/Renderer/Private/SceneCaptureRendering.cpp
+++ b/Engine/Source/Runtime/Renderer/Private/SceneCaptureRendering.cpp
@@ -179,6 +179,7 @@ FSceneRenderer* FScene::CreateSceneRenderer( USceneCaptureComponent* SceneCaptur
 	ViewInitOptions.ViewMatrix = ViewMatrix;
 	ViewInitOptions.BackgroundColor = FLinearColor::Black;
 	ViewInitOptions.OverrideFarClippingPlaneDistance = MaxViewDistance;
+	ViewInitOptions.SceneViewStateInterface = SceneCaptureComponent->GetViewState();
 
 	if (bCaptureSceneColour)
 	{
diff --git a/Engine/Source/Runtime/Renderer/Private/SceneOcclusion.cpp b/Engine/Source/Runtime/Renderer/Private/SceneOcclusion.cpp
--- a/Engine/Source/Runtime/Renderer/Private/SceneOcclusion.cpp
+++ b/Engine/Source/Runtime/Renderer/Private/SceneOcclusion.cpp
@@ -1017,7 +1017,7 @@ void FDeferredShadingSceneRenderer::BeginOcclusionTests(FRHICommandListImmediate
 #endif			
 
 			// Clear primitives which haven't been visible recently out of the occlusion history, and reset old pending occlusion queries.
-			ViewState->TrimOcclusionHistory(RHICmdList, ViewFamily.CurrentRealTime - GEngine->PrimitiveProbablyVisibleTime, ViewFamily.CurrentRealTime, FrameNumber);
+			ViewState->TrimOcclusionHistory(RHICmdList, ViewFamily.CurrentRealTime - GEngine->PrimitiveProbablyVisibleTime, ViewFamily.CurrentRealTime, ViewState->OcclusionFrameCounter);
 
 			// Give back all these occlusion queries to the pool.
 			for ( TMap<FSceneViewState::FProjectedShadowKey, FRenderQueryRHIRef>::TIterator QueryIt(ShadowOcclusionQueryMap); QueryIt; ++QueryIt )
@@ -1118,7 +1118,7 @@ void FDeferredShadingSceneRenderer::BeginOcclusionTests(FRHICommandListImmediate
 		
 		if( ViewState && ViewState->HZBOcclusionTests.GetNum() != 0 )
 		{
-			check( ViewState->HZBOcclusionTests.IsValidFrame(View.FrameNumber) );
+			check( ViewState->HZBOcclusionTests.IsValidFrame(ViewState->OcclusionFrameCounter) );
 
 			SCOPED_DRAW_EVENT(RHICmdList, HZB);
 
diff --git a/Engine/Source/Runtime/Renderer/Private/ScenePrivate.h b/Engine/Source/Runtime/Renderer/Private/ScenePrivate.h
--- a/Engine/Source/Runtime/Renderer/Private/ScenePrivate.h
+++ b/Engine/Source/Runtime/Renderer/Private/ScenePrivate.h
@@ -478,6 +478,9 @@ public:
 	FMatrix		PrevViewMatrixForOcclusionQuery;
 	FVector		PrevViewOriginForOcclusionQuery;
 
+	// A counter incremented once each time this view is rendered.
+	uint32 OcclusionFrameCounter;
+
 	/** Used by states that have IsViewParent() == true to store primitives for child states. */
 	TSet<FPrimitiveComponentId> ParentPrimitives;
 
diff --git a/Engine/Source/Runtime/Renderer/Private/SceneVisibility.cpp b/Engine/Source/Runtime/Renderer/Private/SceneVisibility.cpp
--- a/Engine/Source/Runtime/Renderer/Private/SceneVisibility.cpp
+++ b/Engine/Source/Runtime/Renderer/Private/SceneVisibility.cpp
@@ -385,7 +385,7 @@ static int32 OcclusionCull(FRHICommandListImmediate& RHICmdList, const FScene* S
 			if( bHZBOcclusion )
 			{
 				QUICK_SCOPE_CYCLE_COUNTER(STAT_MapHZBResults);
-				check(!ViewState->HZBOcclusionTests.IsValidFrame(View.FrameNumber));
+				check(!ViewState->HZBOcclusionTests.IsValidFrame(ViewState->OcclusionFrameCounter));
 				ViewState->HZBOcclusionTests.MapResults(RHICmdList);
 			}
 
@@ -449,7 +449,7 @@ static int32 OcclusionCull(FRHICommandListImmediate& RHICmdList, const FScene* S
 						{
 							// Read the occlusion query results.
 							uint64 NumSamples = 0;
-							FRenderQueryRHIRef& PastQuery = PrimitiveOcclusionHistory->GetPastQuery(View.FrameNumber);
+							FRenderQueryRHIRef& PastQuery = PrimitiveOcclusionHistory->GetPastQuery(ViewState->OcclusionFrameCounter);
 							if (IsValidRef(PastQuery))
 							{
 								// NOTE: RHIGetOcclusionQueryResult should never fail when using a blocking call, rendering artifacts may show up.
@@ -513,7 +513,7 @@ static int32 OcclusionCull(FRHICommandListImmediate& RHICmdList, const FScene* S
 
 					if (bClearQueries)
 					{
-						ViewState->OcclusionQueryPool.ReleaseQuery(RHICmdList, PrimitiveOcclusionHistory->GetPastQuery(View.FrameNumber));
+						ViewState->OcclusionQueryPool.ReleaseQuery(RHICmdList, PrimitiveOcclusionHistory->GetPastQuery(ViewState->OcclusionFrameCounter));
 					}
 				}
 
@@ -540,7 +540,7 @@ static int32 OcclusionCull(FRHICommandListImmediate& RHICmdList, const FScene* S
 						{
 							// Always run
 							PrimitiveOcclusionHistory->HZBTestIndex = ViewState->HZBOcclusionTests.AddBounds( OcclusionBounds.Origin, OcclusionBounds.BoxExtent );
-							PrimitiveOcclusionHistory->HZBTestFrameNumber = View.FrameNumber;
+							PrimitiveOcclusionHistory->HZBTestFrameNumber = ViewState->OcclusionFrameCounter;
 						}
 						else
 						{
@@ -577,7 +577,7 @@ static int32 OcclusionCull(FRHICommandListImmediate& RHICmdList, const FScene* S
 
 							if (bRunQuery)
 							{
-								PrimitiveOcclusionHistory->SetCurrentQuery(View.FrameNumber, 
+								PrimitiveOcclusionHistory->SetCurrentQuery(ViewState->OcclusionFrameCounter, 
 									bGroupedQuery ? 
 									View.GroupedOcclusionQueries.BatchPrimitive(OcclusionBounds.Origin + View.ViewMatrices.PreViewTranslation,OcclusionBounds.BoxExtent) :
 									View.IndividualOcclusionQueries.BatchPrimitive(OcclusionBounds.Origin + View.ViewMatrices.PreViewTranslation,OcclusionBounds.BoxExtent)
@@ -614,7 +614,7 @@ static int32 OcclusionCull(FRHICommandListImmediate& RHICmdList, const FScene* S
 
 				if( bSubmitQueries )
 				{
-					ViewState->HZBOcclusionTests.SetValidFrameNumber(View.FrameNumber);
+					ViewState->HZBOcclusionTests.SetValidFrameNumber(ViewState->OcclusionFrameCounter);
 				}
 			}
 		}
@@ -1603,6 +1603,12 @@ void FSceneRenderer::PreVisibilityFrameSetup(FRHICommandListImmediate& RHICmdLis
 		FSceneViewState* ViewState = (FSceneViewState*) View.State;
 		static bool bEnableTimeScale = true;
 
+		// Once per render increment the occlusion frame counter.
+		if (ViewState)
+		{
+			ViewState->OcclusionFrameCounter++;
+		}
+
 		// HighResScreenshot should get best results so we don't do the occlusion optimization based on the former frame
 		extern bool GIsHighResScreenshot;
 		const bool bIsHitTesting = ViewFamily.EngineShowFlags.HitProxies;
