- Fixed my fix for metal with forward (for non-LDR case this time)

diff --git a/Engine/Source/Runtime/Renderer/Private/ForwardShadingRenderer.cpp b/Engine/Source/Runtime/Renderer/Private/ForwardShadingRenderer.cpp
--- a/Engine/Source/Runtime/Renderer/Private/ForwardShadingRenderer.cpp
+++ b/Engine/Source/Runtime/Renderer/Private/ForwardShadingRenderer.cpp
@@ -108,7 +108,7 @@ void FForwardShadingSceneRenderer::Render(FRHICommandListImmediate& RHICmdList)
 	const bool bGammaSpace = !IsMobileHDR();
 	if( bGammaSpace )
 	{
-		SetRenderTarget(RHICmdList, ViewFamily.RenderTarget->GetRenderTargetTexture(), GSceneRenderTargets.GetSceneDepthTexture());
+		SetRenderTarget(RHICmdList, ViewFamily.RenderTarget->GetRenderTargetTexture(), GSceneRenderTargets.GetSceneDepthTexture(), ESimpleRenderTargetMode::EClearToDefault);
 	}
 	else
 	{
