- Fixed wireframe in the editor [UE-4611]

diff --git a/Engine/Source/Runtime/Renderer/Private/DeferredShadingRenderer.cpp b/Engine/Source/Runtime/Renderer/Private/DeferredShadingRenderer.cpp
--- a/Engine/Source/Runtime/Renderer/Private/DeferredShadingRenderer.cpp
+++ b/Engine/Source/Runtime/Renderer/Private/DeferredShadingRenderer.cpp
@@ -645,6 +645,18 @@ static FORCEINLINE bool NeedsPrePass(const FDeferredShadingSceneRenderer* Render
 		(Renderer->EarlyZPassMode != DDM_None || GEarlyZPassMovable != 0);
 }
 
+static void SetAndClearViewGBuffer(FRHICommandListImmediate& RHICmdList, FViewInfo& View, bool bClearDepth)
+{
+	// if we didn't to the prepass above, then we will need to clear now, otherwise, it's already been cleared and rendered to
+	ERenderTargetLoadAction DepthLoadAction = bClearDepth ? ERenderTargetLoadAction::EClear : ERenderTargetLoadAction::ELoad;
+
+	const bool bClearBlack = View.Family->EngineShowFlags.ShaderComplexity || View.Family->EngineShowFlags.StationaryLightOverlap;
+	const FLinearColor ClearColor = (bClearBlack ? FLinearColor(0, 0, 0, 0) : View.BackgroundColor);
+
+	// clearing the GBuffer
+	GSceneRenderTargets.BeginRenderingGBuffer(RHICmdList, ERenderTargetLoadAction::EClear, DepthLoadAction, ClearColor);
+}
+
 void FDeferredShadingSceneRenderer::Render(FRHICommandListImmediate& RHICmdList)
 {
 	bool bDBuffer = IsDBufferEnabled();
@@ -785,30 +797,29 @@ void FDeferredShadingSceneRenderer::Render(FRHICommandListImmediate& RHICmdList)
 		}
 	}
 
+	// Clear the G Buffer render targets
+	bool bIsGBufferCurrent = false;
+	if (bRequiresRHIClear)
+	{
+		// set GBuffer to be current, and clear it
+		SetAndClearViewGBuffer(RHICmdList, Views[0], !bDepthWasCleared);
+
+		// depth was cleared now no matter what
+		bDepthWasCleared = true;
+		bIsGBufferCurrent = true;
+	}
+
 	if(bIsWireframe && FDeferredShadingSceneRenderer::ShouldCompositeEditorPrimitives(Views[0]))
 	{
 		// In Editor we want wire frame view modes to be MSAA for better quality. Resolve will be done with EditorPrimitives
 		SetRenderTarget(RHICmdList, GSceneRenderTargets.GetEditorPrimitivesColor(), GSceneRenderTargets.GetEditorPrimitivesDepth());
 		RHICmdList.Clear(true, FLinearColor(0, 0, 0, 0), true, 0.0f, false, 0, FIntRect());
 	}
-	else
+	else if (!bIsGBufferCurrent)
 	{
-		// if we didn't to the prepass above, then we will need to clear now, otherwise, it's already been cleared and rendered to
+		// make sure the GBuffer is set, in case we didn't need to clear above
 		ERenderTargetLoadAction DepthLoadAction = bDepthWasCleared ? ERenderTargetLoadAction::ELoad : ERenderTargetLoadAction::EClear;
-
-		if (bRequiresRHIClear)
-		{
-			// Clear the G Buffer render targets
-			const bool bClearBlack = Views[0].Family->EngineShowFlags.ShaderComplexity || Views[0].Family->EngineShowFlags.StationaryLightOverlap;
-			const FLinearColor ClearColor = (bClearBlack ? FLinearColor(0, 0, 0, 0) : Views[0].BackgroundColor);
-
-			// render to the GBuffer, clearing it
-			GSceneRenderTargets.BeginRenderingGBuffer(RHICmdList, ERenderTargetLoadAction::EClear, DepthLoadAction, ClearColor);
-		}
-		else
-		{
-			GSceneRenderTargets.BeginRenderingGBuffer(RHICmdList, ERenderTargetLoadAction::ENoAction, DepthLoadAction);
-		}
+		GSceneRenderTargets.BeginRenderingGBuffer(RHICmdList, ERenderTargetLoadAction::ENoAction, DepthLoadAction);
 	}
 
 	GRenderTargetPool.AddPhaseEvent(TEXT("BasePass"));
