#ue4: fixing a rare crash that occur when RenderTarget exists but there is no texture in it.

diff --git a/Engine/Source/Runtime/Renderer/Private/PostProcess/PostProcessing.cpp b/Engine/Source/Runtime/Renderer/Private/PostProcess/PostProcessing.cpp
--- a/Engine/Source/Runtime/Renderer/Private/PostProcess/PostProcessing.cpp
+++ b/Engine/Source/Runtime/Renderer/Private/PostProcess/PostProcessing.cpp
@@ -1299,8 +1299,15 @@ void FPostProcessing::Process(FRHICommandListImmediate& RHICmdList, FViewInfo& V
 			FPooledRenderTargetDesc Desc;
 
 			// Texture could be bigger than viewport
-			Desc.Extent.X = View.Family->RenderTarget->GetRenderTargetTexture()->GetSizeX();
-			Desc.Extent.Y = View.Family->RenderTarget->GetRenderTargetTexture()->GetSizeY();
+			if (View.Family->RenderTarget->GetRenderTargetTexture())
+			{
+				Desc.Extent.X = View.Family->RenderTarget->GetRenderTargetTexture()->GetSizeX();
+				Desc.Extent.Y = View.Family->RenderTarget->GetRenderTargetTexture()->GetSizeY();
+			}
+			else
+			{
+				Desc.Extent = View.Family->RenderTarget->GetSizeXY();
+			}
 			// todo: this should come from View.Family->RenderTarget
 			Desc.Format = PF_B8G8R8A8;
 			Desc.NumMips = 1;
