FNULLScene creates FXSystem.

Fixes UE-4540

diff --git a/Engine/Source/Runtime/Renderer/Private/RendererScene.cpp b/Engine/Source/Runtime/Renderer/Private/RendererScene.cpp
--- a/Engine/Source/Runtime/Renderer/Private/RendererScene.cpp
+++ b/Engine/Source/Runtime/Renderer/Private/RendererScene.cpp
@@ -2066,11 +2066,21 @@ void FScene::OnLevelAddedToWorld_RenderThread(FName InLevelName)
 class FNULLSceneInterface : public FSceneInterface
 {
 public:
-	FNULLSceneInterface( UWorld* InWorld )
+	FNULLSceneInterface(UWorld* InWorld, bool bCreateFXSystem )
 		:	World( InWorld )
 		,	FXSystem( NULL )
 	{
 		World->Scene = this;
+
+		if (bCreateFXSystem)
+		{
+			World->CreateFXSystem();
+		}
+		else
+		{
+			World->FXSystem = NULL;
+			SetFXSystem(NULL);
+		}
 	}
 
 	virtual void AddPrimitive(UPrimitiveComponent* Primitive){}
@@ -2182,7 +2192,7 @@ FSceneInterface* FRendererModule::AllocateScene(UWorld* World, bool bInRequiresH
 	// And fall back to a dummy/ NULL implementation for commandlets and dedicated server.
 	else
 	{
-		return new FNULLSceneInterface( World );
+		return new FNULLSceneInterface(World, bCreateFXSystem);
 	}
 }
 
