cvar override for scene capture render targets being clamped in game mode.

workaround for UE-2907.

diff --git a/Engine/Source/Runtime/Renderer/Private/PostProcess/SceneRenderTargets.cpp b/Engine/Source/Runtime/Renderer/Private/PostProcess/SceneRenderTargets.cpp
--- a/Engine/Source/Runtime/Renderer/Private/PostProcess/SceneRenderTargets.cpp
+++ b/Engine/Source/Runtime/Renderer/Private/PostProcess/SceneRenderTargets.cpp
@@ -41,6 +41,17 @@ static TAutoConsoleVariable<int32> CVarSceneTargetsResizingMethod(
 	ECVF_RenderThreadSafe
 	);
 
+static TAutoConsoleVariable<int32> CVarSceneCaptureResizingMethod(
+	TEXT("r.SceneCaptureResizeMethod"),
+	0,
+	TEXT("Control the scene render target resize method for scene captures:\n")
+	TEXT("(This value is only used in game mode and on windowing platforms.)\n")
+	TEXT("0: All scene capture renders are limited to screen resolution or smaller. (Default - prevents allocation when requested dimensions are too large.)\n")
+	TEXT("1: Allows scene capture targets to expand to encompass the dimensions requested.\n")
+	TEXT("   (large sizes could cause stalling without 'r.SceneRenderTargetResizeMethod 2'. Out of memory issues can occur if size is too large)"),
+	ECVF_RenderThreadSafe
+	);
+
 static TAutoConsoleVariable<int32> CVarOptimizeForUAVPerformance(
 	TEXT("r.OptimizeForUAVPerformance"),
 	0,
@@ -146,8 +157,13 @@ FIntPoint FSceneRenderTargets::GetSceneRenderTargetSize(const FSceneViewFamily&
 			SceneTargetsSizingMethod = Grow;
 		}
 		else
-		{			
+		{
 			SceneTargetsSizingMethod = Clamped;
+			int32 CaptureTargetSizeMethod = CVarSceneCaptureResizingMethod.GetValueOnRenderThread();
+			if (FPlatformProperties::SupportsWindowedMode() && CaptureTargetSizeMethod == 1)
+			{
+				SceneTargetsSizingMethod = Grow;
+			}
 		}
 	}
 
