added bSelfShadowOnly on PrimitiveComponent (for 1st person self shadowing weapon without casting a shadow onto the world)

diff --git a/Engine/Source/Runtime/Engine/Classes/Components/PrimitiveComponent.h b/Engine/Source/Runtime/Engine/Classes/Components/PrimitiveComponent.h
--- a/Engine/Source/Runtime/Engine/Classes/Components/PrimitiveComponent.h
+++ b/Engine/Source/Runtime/Engine/Classes/Components/PrimitiveComponent.h
@@ -289,9 +289,15 @@ public:
 	UPROPERTY(EditAnywhere, AdvancedDisplay, BlueprintReadOnly, Category=Lighting, meta=(EditCondition="CastShadow"))
 	uint32 bCastVolumetricTranslucentShadow:1;
 
+	/** 
+	 * When enabled, the component will only cast a shadow on itself and not other components in the world.  This is especially useful for first person weapons, and forces bCastInsetShadow to be enabled.
+	 */
+	UPROPERTY(EditAnywhere, AdvancedDisplay, BlueprintReadOnly, Category=Lighting, meta=(EditCondition="CastShadow"))
+	uint32 bSelfShadowOnly:1;
+
 	/** 
 	 * Whether this component should create a per-object shadow that gives higher effective shadow resolution. 
-	 * Useful for cinematic character shadowing.
+	 * Useful for cinematic character shadowing. Assumed to be enabled if bSelfShadowOnly is enabled.
 	 */
 	UPROPERTY(EditAnywhere, AdvancedDisplay, BlueprintReadOnly, Category=Lighting, meta=(EditCondition="CastShadow"))
 	uint32 bCastInsetShadow:1;
diff --git a/Engine/Source/Runtime/Engine/Private/PrimitiveComponent.cpp b/Engine/Source/Runtime/Engine/Private/PrimitiveComponent.cpp
--- a/Engine/Source/Runtime/Engine/Private/PrimitiveComponent.cpp
+++ b/Engine/Source/Runtime/Engine/Private/PrimitiveComponent.cpp
@@ -544,6 +544,11 @@ bool UPrimitiveComponent::CanEditChange(const UProperty* InProperty) const
 			return Mobility == EComponentMobility::Movable;
 		}
 
+		if (PropertyName == GET_MEMBER_NAME_CHECKED(UPrimitiveComponent, bCastInsetShadow))
+		{
+			return !bSelfShadowOnly;
+		}
+
 		if (PropertyName == GET_MEMBER_NAME_CHECKED(UPrimitiveComponent, CastShadow))
 		{
 			// Look for any lit materials
diff --git a/Engine/Source/Runtime/Engine/Private/PrimitiveSceneProxy.cpp b/Engine/Source/Runtime/Engine/Private/PrimitiveSceneProxy.cpp
--- a/Engine/Source/Runtime/Engine/Private/PrimitiveSceneProxy.cpp
+++ b/Engine/Source/Runtime/Engine/Private/PrimitiveSceneProxy.cpp
@@ -41,7 +41,8 @@ FPrimitiveSceneProxy::FPrimitiveSceneProxy(const UPrimitiveComponent* InComponen
 ,	bCastVolumetricTranslucentShadow(InComponent->bCastDynamicShadow && InComponent->CastShadow && InComponent->bCastVolumetricTranslucentShadow)
 ,	bCastHiddenShadow(InComponent->bCastHiddenShadow)
 ,	bCastShadowAsTwoSided(InComponent->bCastShadowAsTwoSided)
-,	bCastInsetShadow(InComponent->bCastInsetShadow)
+,	bSelfShadowOnly(InComponent->bSelfShadowOnly)
+,	bCastInsetShadow(InComponent->bSelfShadowOnly ? true : InComponent->bCastInsetShadow)	// Assumed to be enabled if bSelfShadowOnly is enabled.
 ,	bLightAttachmentsAsGroup(InComponent->bLightAttachmentsAsGroup)
 ,	bStaticElementsAlwaysUseProxyPrimitiveUniformBuffer(false)
 ,	bAlwaysHasVelocity(false)
diff --git a/Engine/Source/Runtime/Engine/Public/PrimitiveSceneProxy.h b/Engine/Source/Runtime/Engine/Public/PrimitiveSceneProxy.h
--- a/Engine/Source/Runtime/Engine/Public/PrimitiveSceneProxy.h
+++ b/Engine/Source/Runtime/Engine/Public/PrimitiveSceneProxy.h
@@ -358,6 +358,7 @@ public:
 	inline bool CastsVolumetricTranslucentShadow() const { return bCastVolumetricTranslucentShadow; }
 	inline bool CastsHiddenShadow() const { return bCastHiddenShadow; }
 	inline bool CastsShadowAsTwoSided() const { return bCastShadowAsTwoSided; }
+	inline bool CastsSelfShadowOnly() const { return bSelfShadowOnly; }
 	inline bool CastsInsetShadow() const { return bCastInsetShadow; }
 	inline bool LightAttachmentsAsGroup() const { return bLightAttachmentsAsGroup; }
 	inline bool StaticElementsAlwaysUseProxyPrimitiveUniformBuffer() const { return bStaticElementsAlwaysUseProxyPrimitiveUniformBuffer; }
@@ -528,7 +529,10 @@ protected:
 	/** Whether this primitive should cast dynamic shadows as if it were a two sided material. */
 	uint32 bCastShadowAsTwoSided : 1;
 
-	/** Whether this component should create a per-object shadow that gives higher effective shadow resolution.  */
+	/** When enabled, the component will only cast a shadow on itself and not other components in the world.  This is especially useful for first person weapons, and forces bCastInsetShadow to be enabled. */
+	uint32 bSelfShadowOnly : 1;
+
+	/** Whether this component should create a per-object shadow that gives higher effective shadow resolution. true if bSelfShadowOnly is true. */
 	uint32 bCastInsetShadow : 1;
 
 	/** 
diff --git a/Engine/Source/Runtime/Renderer/Private/ShadowRendering.cpp b/Engine/Source/Runtime/Renderer/Private/ShadowRendering.cpp
--- a/Engine/Source/Runtime/Renderer/Private/ShadowRendering.cpp
+++ b/Engine/Source/Runtime/Renderer/Private/ShadowRendering.cpp
@@ -1868,7 +1868,7 @@ void FProjectedShadowInfo::RenderProjection(FRHICommandListImmediate& RHICmdList
 	bool bDepthBoundsTestEnabled = false;
 
 	// If this is a preshadow, mask the projection by the receiver primitives.
-	if (bPreShadow)
+	if (bPreShadow || bSelfShadowOnly)
 	{
 		SCOPED_DRAW_EVENTF(RHICmdList, EventMaskSubjects, TEXT("Stencil Mask Subjects"));
 
@@ -1885,9 +1885,13 @@ void FProjectedShadowInfo::RenderProjection(FRHICommandListImmediate& RHICmdList
 		TDynamicPrimitiveDrawer<FDepthDrawingPolicyFactory> Drawer(RHICmdList, View, FDepthDrawingPolicyFactory::ContextType(DDM_AllOccluders), true);
 		const bool bUseGetMeshElements = ShouldUseGetDynamicMeshElements();
 
-		for (int32 PrimitiveIndex = 0; PrimitiveIndex < ReceiverPrimitives.Num(); PrimitiveIndex++)
+		// bPreShadow: only receive shadow on specific objects
+		// bSelfShadowOnly: only cast shadow on itself
+		const PrimitiveArrayType& MaskPrimitives = bSelfShadowOnly ? SubjectPrimitives : ReceiverPrimitives;
+
+		for (int32 PrimitiveIndex = 0, PrimitiveCount = MaskPrimitives.Num(); PrimitiveIndex < PrimitiveCount; PrimitiveIndex++)
 		{
-			const FPrimitiveSceneInfo* ReceiverPrimitiveSceneInfo = ReceiverPrimitives[PrimitiveIndex];
+			const FPrimitiveSceneInfo* ReceiverPrimitiveSceneInfo = MaskPrimitives[PrimitiveIndex];
 
 			if (View->PrimitiveVisibilityMap[ReceiverPrimitiveSceneInfo->GetIndex()])
 			{
diff --git a/Engine/Source/Runtime/Renderer/Private/ShadowRendering.h b/Engine/Source/Runtime/Renderer/Private/ShadowRendering.h
--- a/Engine/Source/Runtime/Renderer/Private/ShadowRendering.h
+++ b/Engine/Source/Runtime/Renderer/Private/ShadowRendering.h
@@ -608,6 +608,9 @@ public:
 	/** Whether the shadow is a preshadow or not.  A preshadow is a per object shadow that handles the static environment casting on a dynamic receiver. */
 	uint32 bPreShadow : 1;
 
+	/** To not cast a shadow on the ground outside the object and having higher quality (useful for first person weapon). */
+	uint32 bSelfShadowOnly : 1;
+
 	/** Whether the shadow is a directional light cascade that should be computed by ray tracing mesh distance fields. */
 	uint32 bRayTracedDistanceFieldShadow : 1;
 
@@ -643,7 +646,10 @@ private:
 
 public:
 
-	/** Initialization constructor for a per-object shadow. e.g. translucent particle system */
+	/**
+	 * Initialization constructor for a per-object shadow. e.g. translucent particle system
+	 * @param InParentSceneInfo must not be 0
+	 */
 	FProjectedShadowInfo(
 		FLightSceneInfo* InLightSceneInfo,
 		const FPrimitiveSceneInfo* InParentSceneInfo,
diff --git a/Engine/Source/Runtime/Renderer/Private/ShadowSetup.cpp b/Engine/Source/Runtime/Renderer/Private/ShadowSetup.cpp
--- a/Engine/Source/Runtime/Renderer/Private/ShadowSetup.cpp
+++ b/Engine/Source/Runtime/Renderer/Private/ShadowSetup.cpp
@@ -393,6 +393,7 @@ FProjectedShadowInfo::FProjectedShadowInfo(
 	bReflectiveShadowmap(false),
 	bTranslucentShadow(bInTranslucentShadow),
 	bPreShadow(bInPreShadow),
+	bSelfShadowOnly(InParentSceneInfo->Proxy->CastsSelfShadowOnly()),
 	bRayTracedDistanceFieldShadow(false)
 {
 	const FMatrix WorldToLightScaled = Initializer.WorldToLight * FScaleMatrix(Initializer.Scales);
@@ -517,6 +518,7 @@ FProjectedShadowInfo::FProjectedShadowInfo(
 ,	bReflectiveShadowmap(bInReflectiveShadowMap) 
 ,	bTranslucentShadow(false)
 ,	bPreShadow(false)
+,	bSelfShadowOnly(false)
 ,	bRayTracedDistanceFieldShadow(Initializer.bRayTracedDistanceFieldShadow)
 ,	bValidTransform(true)
 {	
@@ -1284,12 +1286,12 @@ void FDeferredShadingSceneRenderer::CreatePerObjectProjectedShadow(
 					LightSceneInfo,
 					PrimitiveSceneInfo,
 					ShadowInitializer,
-					false,
+					false,					// no preshadow
 					SizeX,
 					MaxShadowResolutionY,
 					MaxScreenPercent,
 					ResolutionFadeAlphas,
-					false
+					false					// no tranlucent shadow
 					);
 				VisibleLightInfo.MemStackProjectedShadows.Add(ProjectedShadowInfo);
 
@@ -1322,12 +1324,12 @@ void FDeferredShadingSceneRenderer::CreatePerObjectProjectedShadow(
 					LightSceneInfo,
 					PrimitiveSceneInfo,
 					ShadowInitializer,
-					false,
+					false,					// no preshadow
 					SizeX,
 					MaxShadowResolutionY,
 					MaxScreenPercent,
 					ResolutionFadeAlphas,
-					true
+					true					// tranlucent shadow
 					);
 				VisibleLightInfo.MemStackProjectedShadows.Add(ProjectedShadowInfo);
 
