Shadow passes are now implemented properly with GetDynamicMeshElements
* Previously was reusing meshes from the main view, but meshes in the shadow view might not be visible from the main view

diff --git a/Engine/Source/Runtime/Engine/Public/SceneManagement.h b/Engine/Source/Runtime/Engine/Public/SceneManagement.h
--- a/Engine/Source/Runtime/Engine/Public/SceneManagement.h
+++ b/Engine/Source/Runtime/Engine/Public/SceneManagement.h
@@ -1386,6 +1386,7 @@ private:
 	ERHIFeatureLevel::Type FeatureLevel;
 
 	friend class FSceneRenderer;
+	friend class FProjectedShadowInfo;
 };
 
 
diff --git a/Engine/Source/Runtime/Renderer/Private/SceneRendering.h b/Engine/Source/Runtime/Renderer/Private/SceneRendering.h
--- a/Engine/Source/Runtime/Renderer/Private/SceneRendering.h
+++ b/Engine/Source/Runtime/Renderer/Private/SceneRendering.h
@@ -642,6 +642,9 @@ protected:
 	/** Calculates projected shadow visibility. */
 	void InitProjectedShadowVisibility(FRHICommandListImmediate& RHICmdList);	
 
+	/** Gathers dynamic mesh elements for all shadows. */
+	void GatherShadowDynamicMeshElements();
+
 	/** Performs once per frame setup prior to visibility determination. */
 	void PreVisibilityFrameSetup(FRHICommandListImmediate& RHICmdList);
 
diff --git a/Engine/Source/Runtime/Renderer/Private/ShadowRendering.cpp b/Engine/Source/Runtime/Renderer/Private/ShadowRendering.cpp
--- a/Engine/Source/Runtime/Renderer/Private/ShadowRendering.cpp
+++ b/Engine/Source/Runtime/Renderer/Private/ShadowRendering.cpp
@@ -1296,51 +1296,47 @@ void FProjectedShadowInfo::RenderDepthDynamic(FRHICommandList& RHICmdList, FScen
 	// Draw the subject's dynamic elements.
 	SCOPE_CYCLE_COUNTER(STAT_WholeSceneDynamicShadowDepthsTime);
 
-	uint32 DrawPrimitiveFlags = 0;
-	if (bReflectiveShadowmap)
-	{
-		// force lowest LOD for RSMs
-		DrawPrimitiveFlags = EDrawDynamicFlags::ForceLowestLOD;
-	}
-
 	const bool bUseGetMeshElements = ShouldUseGetDynamicMeshElements();
-	TDynamicPrimitiveDrawer<FShadowDepthDrawingPolicyFactory> Drawer(RHICmdList, FoundView, FShadowDepthDrawingPolicyFactory::ContextType(this), true);
-	uint32 PrimitiveCount = SubjectPrimitives.Num();
 
-	for (uint32 PrimitiveIndex = 0; PrimitiveIndex < PrimitiveCount; ++PrimitiveIndex)
+	if (bUseGetMeshElements)
 	{
-		const FPrimitiveSceneInfo* PrimitiveSceneInfo = SubjectPrimitives[PrimitiveIndex];
-		const FPrimitiveSceneProxy* PrimitiveSceneProxy = PrimitiveSceneInfo->Proxy;
+		FShadowDepthDrawingPolicyFactory::ContextType Context(this);
 
-		// Lookup the primitive's cached view relevance
-		FPrimitiveViewRelevance ViewRelevance = FoundView->PrimitiveViewRelevanceMap[PrimitiveSceneInfo->GetIndex()];
-
-		if (!ViewRelevance.bInitializedThisFrame)
+		for (int32 MeshBatchIndex = 0; MeshBatchIndex < DynamicSubjectMeshElements.Num(); MeshBatchIndex++)
 		{
-			// Compute the subject primitive's view relevance since it wasn't cached
-			ViewRelevance = PrimitiveSceneInfo->Proxy->GetViewRelevance(FoundView);
+			const FMeshBatchAndRelevance& MeshBatchAndRelevance = DynamicSubjectMeshElements[MeshBatchIndex];
+			const FMeshBatch& MeshBatch = *MeshBatchAndRelevance.Mesh;
+			FShadowDepthDrawingPolicyFactory::DrawDynamicMesh(RHICmdList, *FoundView, Context, MeshBatch, false, true, MeshBatchAndRelevance.PrimitiveSceneProxy, MeshBatch.BatchHitProxyId);
 		}
+	}
+	else
+	{
+		uint32 DrawPrimitiveFlags = 0;
+		if (bReflectiveShadowmap)
+		{
+			// force lowest LOD for RSMs
+			DrawPrimitiveFlags = EDrawDynamicFlags::ForceLowestLOD;
+		}
+
+		TDynamicPrimitiveDrawer<FShadowDepthDrawingPolicyFactory> Drawer(RHICmdList, FoundView, FShadowDepthDrawingPolicyFactory::ContextType(this), true);
+		uint32 PrimitiveCount = SubjectPrimitives.Num();
 
-		// Only draw if the subject primitive is shadow relevant.
-		if (ViewRelevance.bShadowRelevance)
+		for (uint32 PrimitiveIndex = 0; PrimitiveIndex < PrimitiveCount; ++PrimitiveIndex)
 		{
-			if (bUseGetMeshElements)
-			{
-				FShadowDepthDrawingPolicyFactory::ContextType Context(this);
+			const FPrimitiveSceneInfo* PrimitiveSceneInfo = SubjectPrimitives[PrimitiveIndex];
+			const FPrimitiveSceneProxy* PrimitiveSceneProxy = PrimitiveSceneInfo->Proxy;
 
-				//@todo parallelrendering - come up with a better way to filter these by primitive
-				for (int32 MeshBatchIndex = 0; MeshBatchIndex < FoundView->DynamicMeshElements.Num(); MeshBatchIndex++)
-				{
-					const FMeshBatchAndRelevance& MeshBatchAndRelevance = FoundView->DynamicMeshElements[MeshBatchIndex];
+			// Lookup the primitive's cached view relevance
+			FPrimitiveViewRelevance ViewRelevance = FoundView->PrimitiveViewRelevanceMap[PrimitiveSceneInfo->GetIndex()];
 
-					if (MeshBatchAndRelevance.PrimitiveSceneProxy == PrimitiveSceneProxy)
-					{
-						const FMeshBatch& MeshBatch = *MeshBatchAndRelevance.Mesh;
-						FShadowDepthDrawingPolicyFactory::DrawDynamicMesh(RHICmdList, *FoundView, Context, MeshBatch, false, true, MeshBatchAndRelevance.PrimitiveSceneProxy, MeshBatch.BatchHitProxyId);
-					}
-				}
+			if (!ViewRelevance.bInitializedThisFrame)
+			{
+				// Compute the subject primitive's view relevance since it wasn't cached
+				ViewRelevance = PrimitiveSceneInfo->Proxy->GetViewRelevance(FoundView);
 			}
-			else 
+
+			// Only draw if the subject primitive is shadow relevant.
+			if (ViewRelevance.bShadowRelevance)
 			{
 				FScopeCycleCounter Context(PrimitiveSceneProxy->GetStatId());
 				Drawer.SetPrimitive(PrimitiveSceneProxy);
@@ -1885,6 +1881,20 @@ void FProjectedShadowInfo::RenderProjection(FRHICommandListImmediate& RHICmdList
 		TDynamicPrimitiveDrawer<FDepthDrawingPolicyFactory> Drawer(RHICmdList, View, FDepthDrawingPolicyFactory::ContextType(DDM_AllOccluders), true);
 		const bool bUseGetMeshElements = ShouldUseGetDynamicMeshElements();
 
+		if (bUseGetMeshElements)
+		{
+			TArray<FMeshBatchAndRelevance,SceneRenderingAllocator>& DynamicMeshElements = bSelfShadowOnly ? DynamicSubjectMeshElements : DynamicReceiverMeshElements;
+
+			FDepthDrawingPolicyFactory::ContextType Context(DDM_AllOccluders);
+
+			for (int32 MeshBatchIndex = 0; MeshBatchIndex < DynamicMeshElements.Num(); MeshBatchIndex++)
+			{
+				const FMeshBatchAndRelevance& MeshBatchAndRelevance = DynamicMeshElements[MeshBatchIndex];
+				const FMeshBatch& MeshBatch = *MeshBatchAndRelevance.Mesh;
+				FDepthDrawingPolicyFactory::DrawDynamicMesh(RHICmdList, *View, Context, MeshBatch, false, true, MeshBatchAndRelevance.PrimitiveSceneProxy, MeshBatch.BatchHitProxyId);
+			}
+		}
+
 		// bPreShadow: only receive shadow on specific objects
 		// bSelfShadowOnly: only cast shadow on itself
 		const PrimitiveArrayType& MaskPrimitives = bSelfShadowOnly ? SubjectPrimitives : ReceiverPrimitives;
@@ -1898,22 +1908,7 @@ void FProjectedShadowInfo::RenderProjection(FRHICommandListImmediate& RHICmdList
 				const FPrimitiveViewRelevance& ViewRelevance = View->PrimitiveViewRelevanceMap[ReceiverPrimitiveSceneInfo->GetIndex()];
 				if (ViewRelevance.bRenderInMainPass)
 				{
-					if (bUseGetMeshElements)
-					{
-						FDepthDrawingPolicyFactory::ContextType Context(DDM_AllOccluders);
-
-						for (int32 MeshBatchIndex = 0; MeshBatchIndex < View->DynamicMeshElements.Num(); MeshBatchIndex++)
-						{
-							const FMeshBatchAndRelevance& MeshBatchAndRelevance = View->DynamicMeshElements[MeshBatchIndex];
-
-							if (MeshBatchAndRelevance.PrimitiveSceneProxy == ReceiverPrimitiveSceneInfo->Proxy)
-							{
-								const FMeshBatch& MeshBatch = *MeshBatchAndRelevance.Mesh;
-								FDepthDrawingPolicyFactory::DrawDynamicMesh(RHICmdList, *View, Context, MeshBatch, false, true, MeshBatchAndRelevance.PrimitiveSceneProxy, MeshBatch.BatchHitProxyId);
-							}
-						}
-					}
-					else if (ViewRelevance.bDynamicRelevance)
+					if (!bUseGetMeshElements && ViewRelevance.bDynamicRelevance)
 					{
 						Drawer.SetPrimitive(ReceiverPrimitiveSceneInfo->Proxy);
 						ReceiverPrimitiveSceneInfo->Proxy->DrawDynamicElements(&Drawer, View);
diff --git a/Engine/Source/Runtime/Renderer/Private/ShadowRendering.h b/Engine/Source/Runtime/Renderer/Private/ShadowRendering.h
--- a/Engine/Source/Runtime/Renderer/Private/ShadowRendering.h
+++ b/Engine/Source/Runtime/Renderer/Private/ShadowRendering.h
@@ -513,7 +513,7 @@ public:
 	/** Parent primitive of the shadow group that created this shadow, if not a bWholeSceneShadow. */
 	const FPrimitiveSceneInfo* const ParentSceneInfo;
 
-	/** The view this shadow must be rendered in, or NULL for a view independent shadow. */
+	/** The main view this shadow must be rendered in, or NULL for a view independent shadow. */
 	FViewInfo* DependentView;
 
 	/** Index of the shadow into FVisibleLightInfo::AllProjectedShadows. */
@@ -631,12 +631,18 @@ private:
 	PrimitiveArrayType SubjectPrimitives;
 	/** For preshadows, this contains the receiver primitives to mask the projection to. */
 	PrimitiveArrayType ReceiverPrimitives;
+	/** Subject primitives with translucent relevance. */
+	PrimitiveArrayType SubjectTranslucentPrimitives;
 
 	/** Static shadow casting elements. */
 	TArray<FShadowStaticMeshElement,SceneRenderingAllocator> SubjectMeshElements;
 
-	/** Subject primitives with translucent relevance. */
-	PrimitiveArrayType SubjectTranslucentPrimitives;
+	/** Dynamic mesh elements for subject primitives. */
+	TArray<FMeshBatchAndRelevance,SceneRenderingAllocator> DynamicSubjectMeshElements;
+	/** Dynamic mesh elements for receiver primitives. */
+	TArray<FMeshBatchAndRelevance,SceneRenderingAllocator> DynamicReceiverMeshElements;
+	/** Dynamic mesh elements for translucent subject primitives. */
+	TArray<FMeshBatchAndRelevance,SceneRenderingAllocator> DynamicSubjectTranslucentMeshElements;
 
 	/**
 	 * Bias during in shadowmap rendering, stored redundantly for better performance 
@@ -718,6 +724,9 @@ public:
 	 */
 	void AddReceiverPrimitive(FPrimitiveSceneInfo* PrimitiveSceneInfo);
 
+	/** Gathers dynamic mesh elements for all the shadow's primitives arrays. */
+	void GatherDynamicMeshElements(FSceneRenderer& Renderer, class FVisibleLightInfo& VisibleLightInfo, TArray<const FSceneView*>& ReusedViewsArray);
+
 	/** 
 	 * @param View view to check visibility in
 	 * @return true if this shadow info has any subject prims visible in the view
@@ -779,6 +788,14 @@ private:
 
 	/** Updates object buffers needed by ray traced distance field shadows. */
 	int32 UpdateShadowCastingObjectBuffers() const;
+
+	/** Gathers dynamic mesh elements for the given primitive array. */
+	void GatherDynamicMeshElementsArray(
+		FViewInfo* FoundView,
+		FSceneRenderer& Renderer, 
+		PrimitiveArrayType& PrimitiveArray, 
+		TArray<FMeshBatchAndRelevance,SceneRenderingAllocator>& OutDynamicMeshElements, 
+		TArray<const FSceneView*>& ReusedViewsArray);
 };
 
 /** Shader parameters for rendering the depth of a mesh for shadowing. */
diff --git a/Engine/Source/Runtime/Renderer/Private/ShadowSetup.cpp b/Engine/Source/Runtime/Renderer/Private/ShadowSetup.cpp
--- a/Engine/Source/Runtime/Renderer/Private/ShadowSetup.cpp
+++ b/Engine/Source/Runtime/Renderer/Private/ShadowSetup.cpp
@@ -729,7 +729,9 @@ void FProjectedShadowInfo::AddSubjectPrimitive(FPrimitiveSceneInfo* PrimitiveSce
 				*(PrimitiveSceneInfo->ComponentLastRenderTime) = CurrentWorldTime;
 			}
 
-			if (bNeedsPreRenderView)
+			const bool bUseGetDynamicMeshElements = ShouldUseGetDynamicMeshElements();
+
+			if (!bUseGetDynamicMeshElements && bNeedsPreRenderView)
 			{
 				// Call PreRenderView on primitives that weren't visible in any of the main views, but need to be rendered in this shadow's depth pass
 				PrimitiveSceneInfo->Proxy->PreRenderView(Views[0]->Family, ViewMask, Views[0]->FrameNumber);
@@ -882,6 +884,84 @@ void FProjectedShadowInfo::AddReceiverPrimitive(FPrimitiveSceneInfo* PrimitiveSc
 	ReceiverPrimitives.Add(PrimitiveSceneInfo);
 }
 
+void FProjectedShadowInfo::GatherDynamicMeshElements(FSceneRenderer& Renderer, FVisibleLightInfo& VisibleLightInfo, TArray<const FSceneView*>& ReusedViewsArray)
+{
+	if (SubjectPrimitives.Num() > 0 || ReceiverPrimitives.Num() > 0 || SubjectTranslucentPrimitives.Num() > 0)
+	{
+		// Choose an arbitrary view where this shadow's subject is relevant.
+		FViewInfo* FoundView = NULL;
+		for (int32 ViewIndex = 0;ViewIndex < Renderer.Views.Num();ViewIndex++)
+		{
+			FViewInfo* CheckView = &Renderer.Views[ViewIndex];
+			const FVisibleLightViewInfo& VisibleLightViewInfo = CheckView->VisibleLightInfos[LightSceneInfo->Id];
+			FPrimitiveViewRelevance ViewRel = VisibleLightViewInfo.ProjectedShadowViewRelevanceMap[ShadowId];
+			if (ViewRel.bShadowRelevance)
+			{
+				FoundView = CheckView;
+				break;
+			}
+		}
+		check(FoundView && IsInRenderingThread());
+
+		// Backup properties of the view that we will override
+		FMatrix OriginalViewMatrix = FoundView->ViewMatrices.ViewMatrix;
+
+		// Override the view matrix so that billboarding primitives will be aligned to the light
+		FoundView->ViewMatrices.ViewMatrix = ShadowViewMatrix;
+
+		// Prevent materials from getting overridden during shadow casting in viewmodes like lighting only
+		// Lighting only should only affect the material used with direct lighting, not the indirect lighting
+		FoundView->bForceShowMaterials = true;
+
+		ReusedViewsArray[0] = FoundView;
+
+		GatherDynamicMeshElementsArray(FoundView, Renderer, SubjectPrimitives, DynamicSubjectMeshElements, ReusedViewsArray);
+		GatherDynamicMeshElementsArray(FoundView, Renderer, ReceiverPrimitives, DynamicReceiverMeshElements, ReusedViewsArray);
+		GatherDynamicMeshElementsArray(FoundView, Renderer, SubjectTranslucentPrimitives, DynamicSubjectTranslucentMeshElements, ReusedViewsArray);
+
+		FoundView->bForceShowMaterials = false;
+		FoundView->ViewMatrices.ViewMatrix = OriginalViewMatrix;
+	}
+}
+
+void FProjectedShadowInfo::GatherDynamicMeshElementsArray(
+	FViewInfo* FoundView,
+	FSceneRenderer& Renderer, 
+	PrimitiveArrayType& PrimitiveArray, 
+	TArray<FMeshBatchAndRelevance,SceneRenderingAllocator>& OutDynamicMeshElements, 
+	TArray<const FSceneView*>& ReusedViewsArray)
+{
+	// Simple elements not supported in shadow passes
+	FSimpleElementCollector DynamicSubjectSimpleElements;
+
+	Renderer.MeshCollector.ClearViewMeshArrays();
+	Renderer.MeshCollector.AddViewMeshArrays(FoundView, &OutDynamicMeshElements, &DynamicSubjectSimpleElements, Renderer.ViewFamily.GetFeatureLevel());
+
+	const uint32 PrimitiveCount = PrimitiveArray.Num();
+
+	for (uint32 PrimitiveIndex = 0; PrimitiveIndex < PrimitiveCount; ++PrimitiveIndex)
+	{
+		const FPrimitiveSceneInfo* PrimitiveSceneInfo = PrimitiveArray[PrimitiveIndex];
+		const FPrimitiveSceneProxy* PrimitiveSceneProxy = PrimitiveSceneInfo->Proxy;
+
+		// Lookup the primitive's cached view relevance
+		FPrimitiveViewRelevance ViewRelevance = FoundView->PrimitiveViewRelevanceMap[PrimitiveSceneInfo->GetIndex()];
+
+		if (!ViewRelevance.bInitializedThisFrame)
+		{
+			// Compute the subject primitive's view relevance since it wasn't cached
+			ViewRelevance = PrimitiveSceneInfo->Proxy->GetViewRelevance(FoundView);
+		}
+
+		// Only draw if the subject primitive is shadow relevant.
+		if (ViewRelevance.bShadowRelevance && ViewRelevance.bDynamicRelevance)
+		{
+			Renderer.MeshCollector.SetPrimitive(PrimitiveSceneInfo->Proxy, PrimitiveSceneInfo->DefaultDynamicHitProxyId);
+			PrimitiveSceneInfo->Proxy->GetDynamicMeshElements(ReusedViewsArray, Renderer.ViewFamily, 0x1, Renderer.MeshCollector);
+		}
+	}
+}
+
 /** 
  * @param View view to check visibility in
  * @return true if this shadow info has any subject prims visible in the view
@@ -900,13 +980,19 @@ bool FProjectedShadowInfo::SubjectsVisible(const FViewInfo& View) const
 	return false;
 }
 
-/** Clears arrays allocated with the scene rendering allocator. */
+/** 
+ * Clears arrays allocated with the scene rendering allocator.
+ * Cached preshadows are reused across frames so scene rendering allocations will be invalid.
+ */
 void FProjectedShadowInfo::ClearTransientArrays()
 {
 	SubjectTranslucentPrimitives.Empty();
 	SubjectPrimitives.Empty();
 	ReceiverPrimitives.Empty();
 	SubjectMeshElements.Empty();
+	DynamicSubjectMeshElements.Empty();
+	DynamicReceiverMeshElements.Empty();
+	DynamicSubjectTranslucentMeshElements.Empty();
 }
 
 /** Returns a cached preshadow matching the input criteria if one exists. */
@@ -1615,20 +1701,23 @@ void FSceneRenderer::InitProjectedShadowVisibility(FRHICommandListImmediate& RHI
 			for(int32 ViewIndex = 0;ViewIndex < Views.Num();ViewIndex++)
 			{
 				FViewInfo& View = Views[ViewIndex];
+
 				if (ProjectedShadowInfo.DependentView && ProjectedShadowInfo.DependentView != &View)
 				{
 					// The view dependent projected shadow is valid for this view if it's the
 					// right eye and the projected shadow is being rendered for the left eye.
-					bool bIsValidForView = View.StereoPass == eSSP_RIGHT_EYE
+					const bool bIsValidForView = View.StereoPass == eSSP_RIGHT_EYE
 						&& Views.IsValidIndex(ViewIndex - 1)
 						&& Views[ViewIndex - 1].StereoPass == eSSP_LEFT_EYE
 						&& ProjectedShadowInfo.FadeAlphas.IsValidIndex(ViewIndex)
 						&& ProjectedShadowInfo.FadeAlphas[ViewIndex] == 1.0f;
+
 					if (!bIsValidForView)
 					{
 						continue;
 					}
 				}
+
 				FVisibleLightViewInfo& VisibleLightViewInfo = View.VisibleLightInfos[LightIt.GetIndex()];
 
 				if(VisibleLightViewInfo.bInViewFrustum)
@@ -1724,6 +1813,24 @@ void FSceneRenderer::InitProjectedShadowVisibility(FRHICommandListImmediate& RHI
 	}
 }
 
+void FSceneRenderer::GatherShadowDynamicMeshElements()
+{
+	TArray<const FSceneView*> ReusedViewsArray;
+	ReusedViewsArray.AddZeroed(1);
+
+	for (TSparseArray<FLightSceneInfoCompact>::TConstIterator LightIt(Scene->Lights); LightIt; ++LightIt)
+	{
+		FVisibleLightInfo& VisibleLightInfo = VisibleLightInfos[LightIt.GetIndex()];
+
+		for (int32 ShadowIndex = 0; ShadowIndex<VisibleLightInfo.AllProjectedShadows.Num(); ShadowIndex++)
+		{
+			FProjectedShadowInfo& ProjectedShadowInfo = *VisibleLightInfo.AllProjectedShadows[ShadowIndex];
+
+			ProjectedShadowInfo.GatherDynamicMeshElements(*this, VisibleLightInfo, ReusedViewsArray);
+		}
+	}
+}
+
 inline void FSceneRenderer::GatherShadowsForPrimitiveInner(
 	const FPrimitiveSceneInfoCompact& PrimitiveSceneInfoCompact, 
 	const TArray<FProjectedShadowInfo*,SceneRenderingAllocator>& PreShadows,
@@ -2104,6 +2211,9 @@ void FForwardShadingSceneRenderer::InitDynamicShadows(FRHICommandListImmediate&
 
 	// Gathers the list of primitives used to draw various shadow types
 	GatherShadowPrimitives(PreShadows, ViewDependentWholeSceneShadows, false);
+
+	// Generate mesh element arrays from shadow primitive arrays
+	GatherShadowDynamicMeshElements();
 }
 
 void FDeferredShadingSceneRenderer::InitDynamicShadows(FRHICommandListImmediate& RHICmdList)
@@ -2207,4 +2317,7 @@ void FDeferredShadingSceneRenderer::InitDynamicShadows(FRHICommandListImmediate&
 
 	// Gathers the list of primitives used to draw various shadow types
 	GatherShadowPrimitives(PreShadows, ViewDependentWholeSceneShadows, bStaticSceneOnly);
+
+	// Generate mesh element arrays from shadow primitive arrays
+	GatherShadowDynamicMeshElements();
 }
diff --git a/Engine/Source/Runtime/Renderer/Private/TranslucentLighting.cpp b/Engine/Source/Runtime/Renderer/Private/TranslucentLighting.cpp
--- a/Engine/Source/Runtime/Renderer/Private/TranslucentLighting.cpp
+++ b/Engine/Source/Runtime/Renderer/Private/TranslucentLighting.cpp
@@ -570,6 +570,19 @@ void FProjectedShadowInfo::RenderTranslucencyDepths(FRHICommandList& RHICmdList,
 		const bool bUseGetMeshElements = ShouldUseGetDynamicMeshElements();
 
 		FTranslucencyShadowDepthDrawingPolicyFactory::ContextType DrawingContext(this,bDirectionalLight);
+
+		if (bUseGetMeshElements)
+		{
+			FShadowDepthDrawingPolicyFactory::ContextType Context(this);
+
+			for (int32 MeshBatchIndex = 0; MeshBatchIndex < DynamicSubjectTranslucentMeshElements.Num(); MeshBatchIndex++)
+			{
+				const FMeshBatchAndRelevance& MeshBatchAndRelevance = DynamicSubjectTranslucentMeshElements[MeshBatchIndex];
+				const FMeshBatch& MeshBatch = *MeshBatchAndRelevance.Mesh;
+				FTranslucencyShadowDepthDrawingPolicyFactory::DrawDynamicMesh(RHICmdList, *FoundView, DrawingContext, MeshBatch, false, true, MeshBatchAndRelevance.PrimitiveSceneProxy, MeshBatch.BatchHitProxyId);
+			}
+		}
+
 		TDynamicPrimitiveDrawer<FTranslucencyShadowDepthDrawingPolicyFactory> OpacityDrawer(RHICmdList, FoundView, DrawingContext, true);
 
 		for (int32 PrimitiveIndex = 0; PrimitiveIndex < SubjectTranslucentPrimitives.Num(); PrimitiveIndex++)
@@ -586,20 +599,7 @@ void FProjectedShadowInfo::RenderTranslucencyDepths(FRHICommandList& RHICmdList,
 
 			if(ViewRelevance.bDrawRelevance)
 			{
-				if (bUseGetMeshElements)
-				{
-					for (int32 MeshBatchIndex = 0; MeshBatchIndex < FoundView->DynamicMeshElements.Num(); MeshBatchIndex++)
-					{
-						const FMeshBatchAndRelevance& MeshBatchAndRelevance = FoundView->DynamicMeshElements[MeshBatchIndex];
-
-						if (MeshBatchAndRelevance.PrimitiveSceneProxy == PrimitiveSceneInfo->Proxy)
-						{
-							const FMeshBatch& MeshBatch = *MeshBatchAndRelevance.Mesh;
-							FTranslucencyShadowDepthDrawingPolicyFactory::DrawDynamicMesh(RHICmdList, *FoundView, DrawingContext, MeshBatch, false, true, MeshBatchAndRelevance.PrimitiveSceneProxy, MeshBatch.BatchHitProxyId);
-						}
-					}
-				}
-				else if (ViewRelevance.bDynamicRelevance)
+				if (!bUseGetMeshElements && ViewRelevance.bDynamicRelevance)
 				{
 					OpacityDrawer.SetPrimitive(PrimitiveSceneInfo->Proxy);
 					PrimitiveSceneInfo->Proxy->DrawDynamicElements(&OpacityDrawer, FoundView);
