Fixed compile error

diff --git a/Engine/Source/Runtime/Renderer/Private/ShadowRendering.cpp b/Engine/Source/Runtime/Renderer/Private/ShadowRendering.cpp
--- a/Engine/Source/Runtime/Renderer/Private/ShadowRendering.cpp
+++ b/Engine/Source/Runtime/Renderer/Private/ShadowRendering.cpp
@@ -1883,7 +1883,7 @@ void FProjectedShadowInfo::RenderProjection(FRHICommandListImmediate& RHICmdList
 
 		if (bUseGetMeshElements)
 		{
-			TArray<FMeshBatchAndRelevance,SceneRenderingAllocator>& DynamicMeshElements = bSelfShadowOnly ? DynamicSubjectMeshElements : DynamicReceiverMeshElements;
+			const TArray<FMeshBatchAndRelevance,SceneRenderingAllocator>& DynamicMeshElements = bSelfShadowOnly ? DynamicSubjectMeshElements : DynamicReceiverMeshElements;
 
 			FDepthDrawingPolicyFactory::ContextType Context(DDM_AllOccluders);
 
