Avoid assert with zero scale objects

diff --git a/Engine/Source/Runtime/Renderer/Private/DistanceFieldShadowing.cpp b/Engine/Source/Runtime/Renderer/Private/DistanceFieldShadowing.cpp
--- a/Engine/Source/Runtime/Renderer/Private/DistanceFieldShadowing.cpp
+++ b/Engine/Source/Runtime/Renderer/Private/DistanceFieldShadowing.cpp
@@ -573,9 +573,9 @@ int32 FProjectedShadowInfo::UpdateShadowCastingObjectBuffers() const
 
 					ObjectBoxBounds.Add(FVector4(MinViewSpacePosition, 0));
 					ObjectBoxBounds.Add(FVector4(MaxViewSpacePosition, 0));
-					ObjectBoxBounds.Add(FVector4(ObjectXAxis / ObjectXAxis.SizeSquared(), 0));
-					ObjectBoxBounds.Add(FVector4(ObjectYAxis / ObjectYAxis.SizeSquared(), 0));
-					ObjectBoxBounds.Add(FVector4(ObjectZAxis / ObjectZAxis.SizeSquared(), 0));
+					ObjectBoxBounds.Add(FVector4(ObjectXAxis / FMath::Max(ObjectXAxis.SizeSquared(), KINDA_SMALL_NUMBER), 0));
+					ObjectBoxBounds.Add(FVector4(ObjectYAxis / FMath::Max(ObjectYAxis.SizeSquared(), KINDA_SMALL_NUMBER), 0));
+					ObjectBoxBounds.Add(FVector4(ObjectZAxis / FMath::Max(ObjectZAxis.SizeSquared(), KINDA_SMALL_NUMBER), 0));
 				}
 
 				checkSlow(ObjectData.Num() % FDistanceFieldObjectBuffers::ObjectDataStride == 0);
