Resolve UE-3535 (was TTP 349592). This adds in temporal AA for near Gaussian DOF also to remove flicker.

diff --git a/Engine/Source/Runtime/Renderer/Private/PostProcess/PostProcessTemporalAA.cpp b/Engine/Source/Runtime/Renderer/Private/PostProcess/PostProcessTemporalAA.cpp
--- a/Engine/Source/Runtime/Renderer/Private/PostProcess/PostProcessTemporalAA.cpp
+++ b/Engine/Source/Runtime/Renderer/Private/PostProcess/PostProcessTemporalAA.cpp
@@ -358,6 +358,88 @@ FPooledRenderTargetDesc FRCPassPostProcessDOFTemporalAA::ComputeOutputDesc(EPass
 }
 
 
+void FRCPassPostProcessDOFTemporalAANear::Process(FRenderingCompositePassContext& Context)
+{
+	SCOPED_DRAW_EVENT(Context.RHICmdList, DOFTemporalAANear);
+
+	const FPooledRenderTargetDesc* InputDesc = GetInputDesc(ePId_Input0);
+
+	if(!InputDesc)
+	{
+		// input is not hooked up correctly
+		return;
+	}
+
+	const FViewInfo& View = Context.View;
+	FSceneViewState* ViewState = Context.ViewState;
+
+	FIntPoint TexSize = InputDesc->Extent;
+
+	// we assume the input and output is full resolution
+
+	FIntPoint SrcSize = InputDesc->Extent;
+	FIntPoint DestSize = PassOutputs[0].RenderTargetDesc.Extent;
+
+	// e.g. 4 means the input texture is 4x smaller than the buffer size
+	uint32 ScaleFactor = GSceneRenderTargets.GetBufferSizeXY().X / SrcSize.X;
+
+	FIntRect SrcRect = View.ViewRect / ScaleFactor;
+	FIntRect DestRect = SrcRect;
+
+	const FSceneRenderTargetItem& DestRenderTarget = PassOutputs[0].RequestSurface(Context);
+
+	SetRenderTarget(Context.RHICmdList, DestRenderTarget.TargetableTexture, FTextureRHIRef());
+
+	// is optimized away if possible (RT size=view size, )
+	Context.RHICmdList.Clear(true, FLinearColor::Black, false, 1.0f, false, 0, SrcRect);
+
+	Context.SetViewportAndCallRHI(SrcRect);
+
+	// set the state
+	Context.RHICmdList.SetBlendState(TStaticBlendState<>::GetRHI());
+	Context.RHICmdList.SetRasterizerState(TStaticRasterizerState<>::GetRHI());
+	Context.RHICmdList.SetDepthStencilState(TStaticDepthStencilState<false,CF_Always>::GetRHI());
+
+	TShaderMapRef< FPostProcessTonemapVS >			VertexShader(Context.GetShaderMap());
+	TShaderMapRef< FPostProcessTemporalAAPS<0, 0> >	PixelShader(Context.GetShaderMap());
+
+	static FGlobalBoundShaderState BoundShaderState;
+
+
+	SetGlobalBoundShaderState(Context.RHICmdList, Context.GetFeatureLevel(), BoundShaderState, GFilterVertexDeclaration.VertexDeclarationRHI, *VertexShader, *PixelShader);
+
+	VertexShader->SetVS(Context);
+	PixelShader->SetParameters(Context);
+
+	// Draw a quad mapping scene color to the view's render target
+	DrawRectangle(
+		Context.RHICmdList,
+		0, 0,
+		SrcRect.Width(), SrcRect.Height(),
+		SrcRect.Min.X, SrcRect.Min.Y, 
+		SrcRect.Width(), SrcRect.Height(),
+		SrcRect.Size(),
+		SrcSize,
+		*VertexShader,
+		EDRF_UseTriangleOptimization);
+
+	Context.RHICmdList.CopyToResolveTarget(DestRenderTarget.TargetableTexture, DestRenderTarget.ShaderResourceTexture, false, FResolveParams());
+
+	ViewState->DOFHistoryRT2 = PassOutputs[0].PooledRenderTarget;
+	check( ViewState->DOFHistoryRT2 );
+}
+
+FPooledRenderTargetDesc FRCPassPostProcessDOFTemporalAANear::ComputeOutputDesc(EPassOutputId InPassOutputId) const
+{
+	FPooledRenderTargetDesc Ret = PassInputs[0].GetOutput()->RenderTargetDesc;
+
+	Ret.DebugName = TEXT("BokehDOFTemporalAANear");
+
+	return Ret;
+}
+
+
+
 void FRCPassPostProcessLightShaftTemporalAA::Process(FRenderingCompositePassContext& Context)
 {
 	const FPooledRenderTargetDesc* InputDesc = GetInputDesc(ePId_Input0);
diff --git a/Engine/Source/Runtime/Renderer/Private/PostProcess/PostProcessTemporalAA.h b/Engine/Source/Runtime/Renderer/Private/PostProcess/PostProcessTemporalAA.h
--- a/Engine/Source/Runtime/Renderer/Private/PostProcess/PostProcessTemporalAA.h
+++ b/Engine/Source/Runtime/Renderer/Private/PostProcess/PostProcessTemporalAA.h
@@ -34,6 +34,20 @@ public:
 	virtual FPooledRenderTargetDesc ComputeOutputDesc(EPassOutputId InPassOutputId) const;
 };
 
+// ePId_Input0: Half Res DOF input (point)
+// ePId_Input1: Previous frame's output (bilinear)
+// ePId_Input2: Previous frame's output (point)
+// derives from TRenderingCompositePassBase<InputCount, OutputCount> 
+class FRCPassPostProcessDOFTemporalAANear : public TRenderingCompositePassBase<3, 1>
+{
+public:
+	// interface FRenderingCompositePass ---------
+	virtual void Process(FRenderingCompositePassContext& Context);
+	virtual void Release() override { delete this; }
+	virtual FPooledRenderTargetDesc ComputeOutputDesc(EPassOutputId InPassOutputId) const;
+};
+
+
 // ePId_Input0: Half Res light shaft input (point)
 // ePId_Input1: Previous frame's output (bilinear)
 // ePId_Input2: Previous frame's output (point)
diff --git a/Engine/Source/Runtime/Renderer/Private/PostProcess/PostProcessing.cpp b/Engine/Source/Runtime/Renderer/Private/PostProcess/PostProcessing.cpp
--- a/Engine/Source/Runtime/Renderer/Private/PostProcess/PostProcessing.cpp
+++ b/Engine/Source/Runtime/Renderer/Private/PostProcess/PostProcessing.cpp
@@ -325,6 +325,36 @@ static void AddPostProcessDepthOfFieldGaussian(FPostprocessContext& Context, FDe
 		ViewState->bBokehDOFHistory = false;
 	}
 
+	FRenderingCompositePass* DOFInputPass2 = DOFSetup;
+	EPassOutputId DOFInputPassId = ePId_Output1;
+	if( Context.View.FinalPostProcessSettings.AntiAliasingMethod == AAM_TemporalAA && ViewState )
+	{
+		FRenderingCompositePass* HistoryInput;
+		EPassOutputId DOFInputPassId2 = ePId_Output1;
+		if( ViewState->DOFHistoryRT2 && !ViewState->bBokehDOFHistory2 && !Context.View.bCameraCut )
+		{
+			HistoryInput = Context.Graph.RegisterPass( new(FMemStack::Get()) FRCPassPostProcessInput( ViewState->DOFHistoryRT2 ) );
+			DOFInputPassId2 = ePId_Output0;
+		}
+		else
+		{
+			// No history so use current as history
+			HistoryInput = DOFSetup;
+			DOFInputPassId2 = ePId_Output1;
+		}
+
+		FRenderingCompositePass* NodeTemporalAA = Context.Graph.RegisterPass( new(FMemStack::Get()) FRCPassPostProcessDOFTemporalAANear );
+		NodeTemporalAA->SetInput( ePId_Input0, FRenderingCompositeOutputRef( DOFSetup, ePId_Output1 ) );
+		NodeTemporalAA->SetInput( ePId_Input1, FRenderingCompositeOutputRef( HistoryInput, DOFInputPassId2 ) );
+		NodeTemporalAA->SetInput( ePId_Input2, FRenderingCompositeOutputRef( HistoryInput, DOFInputPassId2 ) );
+		//NodeTemporalAA->SetInput( ePId_Input3, VelocityInput );
+
+		DOFInputPass2 = NodeTemporalAA;
+		ViewState->bBokehDOFHistory2 = false;
+
+		DOFInputPassId = ePId_Output0;
+	}
+	
 	FRenderingCompositeOutputRef Far(DOFInputPass, ePId_Output0);
 	FRenderingCompositeOutputRef Near; // Don't need to bind a dummy here as we use a different permutation which doesn't read from this input when near DOF is disabled
 
@@ -337,7 +367,7 @@ static void AddPostProcessDepthOfFieldGaussian(FPostprocessContext& Context, FDe
 	// near
 	if (Out.bNear)
 	{
-		Near = RenderGaussianBlur(Context, TEXT("NearDOFBlurX"), TEXT("NearDOFBlurY"), FRenderingCompositeOutputRef(DOFSetup, ePId_Output1), NearSize);
+		Near = RenderGaussianBlur(Context, TEXT("NearDOFBlurX"), TEXT("NearDOFBlurY"), FRenderingCompositeOutputRef(DOFInputPass2, DOFInputPassId), NearSize);
 	}
 
 	FRenderingCompositePass* NodeRecombined = Context.Graph.RegisterPass(new(FMemStack::Get()) FRCPassPostProcessDOFRecombine(Out.bNear));
diff --git a/Engine/Source/Runtime/Renderer/Private/RendererScene.cpp b/Engine/Source/Runtime/Renderer/Private/RendererScene.cpp
--- a/Engine/Source/Runtime/Renderer/Private/RendererScene.cpp
+++ b/Engine/Source/Runtime/Renderer/Private/RendererScene.cpp
@@ -59,6 +59,7 @@ FSceneViewState::FSceneViewState()
 	TemporalAASampleCount = 1;
 	AOTileIntersectionResources = NULL;
 	bBokehDOFHistory = true;
+	bBokehDOFHistory2 = true;
 
 	LightPropagationVolume = NULL; 
 
diff --git a/Engine/Source/Runtime/Renderer/Private/ScenePrivate.h b/Engine/Source/Runtime/Renderer/Private/ScenePrivate.h
--- a/Engine/Source/Runtime/Renderer/Private/ScenePrivate.h
+++ b/Engine/Source/Runtime/Renderer/Private/ScenePrivate.h
@@ -527,6 +527,7 @@ public:
 	TRefCountPtr<IPooledRenderTarget> TemporalAAHistoryRT;
 	// Temporal AA result for DOF of last frame
 	TRefCountPtr<IPooledRenderTarget> DOFHistoryRT;
+	TRefCountPtr<IPooledRenderTarget> DOFHistoryRT2;
 	// Temporal AA result for SSR
 	TRefCountPtr<IPooledRenderTarget> SSRHistoryRT;
 	// Temporal AA result for light shafts of last frame
@@ -549,6 +550,7 @@ public:
 
 	// Is DOFHistoryRT set from Bokeh DOF?
 	bool bBokehDOFHistory;
+	bool bBokehDOFHistory2;
 
 	// call after SetupTemporalAA()
 	uint32 GetCurrentTemporalAASampleIndex() const
@@ -688,6 +690,7 @@ public:
 		SeparateTranslucencyRT.SafeRelease();
 		TemporalAAHistoryRT.SafeRelease();
 		DOFHistoryRT.SafeRelease();
+		DOFHistoryRT2.SafeRelease();
 		SSRHistoryRT.SafeRelease();
 		LightShaftOcclusionHistoryRT.SafeRelease();
 		LightShaftBloomHistoryRTs.Empty();
