Sky light capture will initialize scratch targets in GSceneRenderTargets if needed

diff --git a/Engine/Source/Runtime/Renderer/Private/ReflectionEnvironmentCapture.cpp b/Engine/Source/Runtime/Renderer/Private/ReflectionEnvironmentCapture.cpp
--- a/Engine/Source/Runtime/Renderer/Private/ReflectionEnvironmentCapture.cpp
+++ b/Engine/Source/Runtime/Renderer/Private/ReflectionEnvironmentCapture.cpp
@@ -648,6 +648,7 @@ int32 FindOrAllocateCubemapIndex(FScene* Scene, const UReflectionCaptureComponen
 
 void ClearScratchCubemaps(FRHICommandList& RHICmdList)
 {
+	GSceneRenderTargets.AllocateReflectionTargets();
 	// Clear scratch render targets to a consistent but noticeable value
 	// This makes debugging capture issues much easier, otherwise the random contents from previous captures is shown
 
@@ -749,8 +750,6 @@ void CopyCubemapToScratchCubemap(FRHICommandList& RHICmdList, ERHIFeatureLevel::
 {
 	check(SourceCubemap);
 	
-	ClearScratchCubemaps(RHICmdList);
-
 	const int32 EffectiveSize = GReflectionCaptureSize;
 	FSceneRenderTargetItem& EffectiveColorRT =  GSceneRenderTargets.ReflectionColorScratchCubemap[0]->GetRenderTargetItem();
 
@@ -1275,6 +1274,12 @@ void FScene::UpdateReflectionCaptureContents(UReflectionCaptureComponent* Captur
 		}
 		else
 		{
+			ENQUEUE_UNIQUE_RENDER_COMMAND( 
+				ClearCommand,
+			{
+				ClearScratchCubemaps(RHICmdList);
+			});
+
 			CaptureSceneIntoScratchCubemap(this, CaptureComponent->GetComponentLocation(), false, true, 0, false, false);
 
 			ENQUEUE_UNIQUE_RENDER_COMMAND_ONEPARAMETER( 
