New feature: Cylindrical projection/distortion
  r.Upscale.Cylinder 1
  r.ScreenPercentage.Editor 1

diff --git a/Engine/Config/BaseScalability.ini b/Engine/Config/BaseScalability.ini
--- a/Engine/Config/BaseScalability.ini
+++ b/Engine/Config/BaseScalability.ini
@@ -110,7 +110,7 @@ r.SceneColorFringeQuality=0
 r.EyeAdaptationQuality=0
 r.BloomQuality=4
 r.FastBlurThreshold=0
-r.UpsampleQuality=1
+r.Upscale.Quality=1
 r.TonemapperQuality=0
 r.LightShaftQuality=0
 
@@ -125,7 +125,7 @@ r.SceneColorFringeQuality=0
 r.EyeAdaptationQuality=0
 r.BloomQuality=4
 r.FastBlurThreshold=2
-r.UpsampleQuality=2
+r.Upscale.Quality=2
 r.TonemapperQuality=0
 r.LightShaftQuality=0
 
@@ -140,7 +140,7 @@ r.SceneColorFringeQuality=1
 r.EyeAdaptationQuality=2
 r.BloomQuality=5
 r.FastBlurThreshold=3
-r.UpsampleQuality=2
+r.Upscale.Quality=2
 r.TonemapperQuality=1
 r.LightShaftQuality=1
 
@@ -155,7 +155,7 @@ r.SceneColorFringeQuality=1
 r.EyeAdaptationQuality=2
 r.BloomQuality=5
 r.FastBlurThreshold=7
-r.UpsampleQuality=3
+r.Upscale.Quality=3
 r.TonemapperQuality=1
 r.LightShaftQuality=1
 
diff --git a/Engine/Shaders/PostProcessUpscale.usf b/Engine/Shaders/PostProcessUpscale.usf
--- a/Engine/Shaders/PostProcessUpscale.usf
+++ b/Engine/Shaders/PostProcessUpscale.usf
@@ -7,8 +7,22 @@
 #include "Common.usf"
 #include "PostProcessCommon.usf"
 
+// only for MainPS
 float UpscaleSoftness;
 
+// only for MainVS()
+// .x:HalfFOV .y:tan(HalfFov) .z:CylinderDistortion(0=none, 1:full) .w:Correction to scale Y the same as X is scaled in the center
+float4 DistortionParams;
+
+// to avoid compiler errors on the PS
+#ifndef TESS_RECT_X
+	#define TESS_RECT_X 1
+#endif
+#ifndef TESS_RECT_Y
+	#define TESS_RECT_Y 1
+#endif
+
+
 float Luma(float3 Color) 
 {
 	// Rec 709 function for luma.
@@ -20,9 +34,52 @@ float Gaussian(float Scale, float2 Offset)
 	return exp2(Scale * dot(Offset, Offset));
 }
 
-void MainPS(float4 UVAndScreenPos : TEXCOORD0, float4 SvPosition : SV_POSITION, out float4 OutColor : SV_Target0)
+// @param In -1 .. 1 range
+// @return In -1 .. 1 range
+float2 CylindricalDistortPos(float2 In)
+{
+	float HalfFov = DistortionParams.x;
+	float TanHalfFov = DistortionParams.y;
+	float InvHalfFov = 1.0f / HalfFov;
+	float CylinderDistortion = DistortionParams.z;
+	float Correction = DistortionParams.w;
+
+	float X = In.x * TanHalfFov;
+	float AngleX = atan(X);
+
+	float2 Ret = float2(AngleX * InvHalfFov, In.y * cos(AngleX) * Correction);
+
+	return lerp(In, Ret, CylinderDistortion);
+}
+
+
+// vertex shader entry point
+void MainVS(
+	in uint VertexId : SV_VertexID,
+	out float4 OutTexCoord : TEXCOORD0,
+	out float4 OutPosition : SV_POSITION
+	)
 {
+	// amount of Quads in X and Y, from GTesselatedScreenRectangleIndexBuffer
+	uint2 RectangleTessellation = uint2(TESS_RECT_X, TESS_RECT_Y);
 
+	float2 TexCoord = float2(VertexId % (RectangleTessellation.x + 1), VertexId / (RectangleTessellation.x + 1)) / float2(RectangleTessellation);
+
+	// still in 0..1 range
+	float4 Position = float4(TexCoord.x, TexCoord.y, 0, 1);
+
+	float2 PlanarScreenPos = TexCoord * 2 - 1;
+
+	// distort pos
+	Position.xy = CylindricalDistortPos(PlanarScreenPos) * 0.5f + 0.5f;
+
+	DrawRectangle(Position, TexCoord, OutPosition, OutTexCoord.xy);
+	OutTexCoord.zw =  OutPosition.xy;
+}
+
+
+void MainPS(float4 UVAndScreenPos : TEXCOORD0, float4 SvPosition : SV_POSITION, out float4 OutColor : SV_Target0)
+{
 #if METHOD == 0
 	// Nearest sampling (not blurry but blocky, more for testing)
 	OutColor = PostprocessInput1.SampleLevel(PostprocessInput1Sampler, UVAndScreenPos.xy, 0, int2(0, 0));
@@ -70,7 +127,15 @@ void MainPS(float4 UVAndScreenPos : TEXCOORD0, float4 SvPosition : SV_POSITION,
     float2 Dir;
     Dir.x = DirSWMinusNE + DirSEMinusNW;
     Dir.y = DirSWMinusNE - DirSEMinusNW;
-    Dir = normalize(Dir);
+
+	{
+		// minor rendering artifacts when using the distortion feature
+		//Dir = normalize(Dir);
+
+		// fixed these problems
+		Dir = Dir * rsqrt(dot(Dir, Dir) + 0.0001f);
+	}
+
 	float2 Dir90 = float2(-Dir.y, Dir.x);
 
 	UVShift = 1.0 - UVShift;
@@ -225,4 +290,4 @@ void MainPS(float4 UVAndScreenPos : TEXCOORD0, float4 SvPosition : SV_POSITION,
 	OutColor.b *= Base + T;
 #endif
 
-}
\ No newline at end of file
+}
diff --git a/Engine/Source/Editor/UnrealEd/Private/EditorViewportClient.cpp b/Engine/Source/Editor/UnrealEd/Private/EditorViewportClient.cpp
--- a/Engine/Source/Editor/UnrealEd/Private/EditorViewportClient.cpp
+++ b/Engine/Source/Editor/UnrealEd/Private/EditorViewportClient.cpp
@@ -662,7 +662,7 @@ FSceneView* FEditorViewportClient::CalcSceneView(FSceneViewFamily* ViewFamily)
 
 	OverridePostProcessSettings( *View );
 
-	View->EndFinalPostprocessSettings();
+	View->EndFinalPostprocessSettings(ViewInitOptions);
 
 	return View;
 }
diff --git a/Engine/Source/Editor/UnrealEd/Private/ThumbnailHelpers.cpp b/Engine/Source/Editor/UnrealEd/Private/ThumbnailHelpers.cpp
--- a/Engine/Source/Editor/UnrealEd/Private/ThumbnailHelpers.cpp
+++ b/Engine/Source/Editor/UnrealEd/Private/ThumbnailHelpers.cpp
@@ -105,7 +105,7 @@ void FThumbnailPreviewScene::GetView(FSceneViewFamily* ViewFamily, int32 X, int3
 		ViewFamily->Views.Add(NewView);
 
 		NewView->StartFinalPostprocessSettings( ViewMatrix.GetOrigin() );
-		NewView->EndFinalPostprocessSettings();
+		NewView->EndFinalPostprocessSettings(ViewInitOptions);
 
 		FFinalPostProcessSettings::FCubemapEntry& CubemapEntry = *new(NewView->FinalPostProcessSettings.ContributingCubemaps) FFinalPostProcessSettings::FCubemapEntry;
 		CubemapEntry.AmbientCubemap = GUnrealEd->GetThumbnailManager()->AmbientCubemap;
diff --git a/Engine/Source/Runtime/Engine/Private/LocalPlayer.cpp b/Engine/Source/Runtime/Engine/Private/LocalPlayer.cpp
--- a/Engine/Source/Runtime/Engine/Private/LocalPlayer.cpp
+++ b/Engine/Source/Runtime/Engine/Private/LocalPlayer.cpp
@@ -754,53 +754,7 @@ FSceneView* ULocalPlayer::CalcSceneView( class FSceneViewFamily* ViewFamily,
 		//	NOTE: Matinee works through this channel
 		View->OverridePostProcessSettings(ViewInfo.PostProcessSettings, ViewInfo.PostProcessBlendWeight);
 
-		View->EndFinalPostprocessSettings();
-	}
-
-	// Upscaling or Super sampling
-	{
-		float LocalScreenPercentage = View->FinalPostProcessSettings.ScreenPercentage;
-
-		float Fraction = 1.0f;
-
-		// apply ScreenPercentage
-		if (LocalScreenPercentage != 100.f)
-		{
-			Fraction = FMath::Clamp(LocalScreenPercentage / 100.0f, 0.1f, 4.0f);
-		}
-
-		// Window full screen mode with upscaling
-		bool bFullscreen = false;
-		if (GEngine && GEngine->GameViewport && GEngine->GameViewport->GetWindow().IsValid())
-		{
-			bFullscreen = GEngine->GameViewport->GetWindow()->GetWindowMode() != EWindowMode::Windowed;
-		}
-
-		if (bFullscreen)
-		{
-			int32 WindowModeType = GetBoundFullScreenModeCVar();
-
-			// CVar mode 2 is fullscreen with upscale
-			if(WindowModeType == 2)
-			{
-				FIntPoint WindowSize = Viewport->GetSizeXY();
-
-				// allow only upscaling
-				float FractionX = FMath::Clamp((float)GSystemResolution.ResX / WindowSize.X, 0.1f, 4.0f);
-				float FractionY = FMath::Clamp((float)GSystemResolution.ResY / WindowSize.Y, 0.1f, 4.0f);
-
-				// maintain a pixel aspect ratio of 1:1 for easier internal computations
-				Fraction *= FMath::Max(FractionX, FractionY);
-			}
-		}
-
-		// Upscale if needed
-		if (Fraction != 1.0f)
-		{
-			// compute the view rectangle with the ScreenPercentage applied
-			const FIntRect ScreenPercentageAffectedViewRect = ViewInitOptions.GetConstrainedViewRect().Scale(Fraction);
-			View->SetScaledViewRect(ScreenPercentageAffectedViewRect);
-		}
+		View->EndFinalPostprocessSettings(ViewInitOptions);
 	}
 
 	for (int ViewExt = 0; ViewExt < ViewFamily->ViewExtensions.Num(); ViewExt++)
diff --git a/Engine/Source/Runtime/Engine/Private/SceneView.cpp b/Engine/Source/Runtime/Engine/Private/SceneView.cpp
--- a/Engine/Source/Runtime/Engine/Private/SceneView.cpp
+++ b/Engine/Source/Runtime/Engine/Private/SceneView.cpp
@@ -52,6 +52,15 @@ static TAutoConsoleVariable<int32> CVarRenderTimeFrozen(
 	TEXT(" 0: off\n")
 	TEXT(" 1: on (Note: this also disables occlusion queries)"),
 	ECVF_Cheat);
+
+static TAutoConsoleVariable<int32> CVarScreenPercentageEditor(
+	TEXT("r.ScreenPercentage.Editor"),
+	0,
+	TEXT("To allow to have an effect of ScreenPercentage in the editor.\n")
+	TEXT("0: off (default)\n")
+	TEXT("1: allow upsample (blurry but faster) and downsample (cripser but slower)"),
+	ECVF_Default);
+
 #endif
 
 static TAutoConsoleVariable<float> CVarSSAOFadeRadiusScale(
@@ -945,7 +954,7 @@ void FSceneView::StartFinalPostprocessSettings(FVector InViewLocation)
 	}
 }
 
-void FSceneView::EndFinalPostprocessSettings()
+void FSceneView::EndFinalPostprocessSettings(FSceneViewInitOptions ViewInitOptions)
 {
 	{
 		static const auto CVarMobileMSAA = IConsoleManager::Get().FindTConsoleVariableDataInt(TEXT("r.MobileMSAA"));
@@ -1218,6 +1227,70 @@ void FSceneView::EndFinalPostprocessSettings()
 		}
 	}
 #endif // WITH_EDITOR
+
+
+	// Upscaling or Super sampling
+	{
+		float LocalScreenPercentage = FinalPostProcessSettings.ScreenPercentage;
+
+		float Fraction = 1.0f;
+
+		// apply ScreenPercentage
+		if (LocalScreenPercentage != 100.f)
+		{
+			Fraction = FMath::Clamp(LocalScreenPercentage / 100.0f, 0.1f, 4.0f);
+		}
+
+		// Window full screen mode with upscaling
+		bool bFullscreen = false;
+		if (GEngine && GEngine->GameViewport && GEngine->GameViewport->GetWindow().IsValid())
+		{
+			bFullscreen = GEngine->GameViewport->GetWindow()->GetWindowMode() != EWindowMode::Windowed;
+		}
+
+		check(Family->RenderTarget);
+
+		if (bFullscreen)
+		{
+			extern int32 GetBoundFullScreenModeCVar();
+			int32 WindowModeType = GetBoundFullScreenModeCVar();
+
+			// CVar mode 2 is fullscreen with upscale
+			if(WindowModeType == 2)
+			{
+//				FIntPoint WindowSize = Viewport->GetSizeXY();
+				FIntPoint WindowSize = Family->RenderTarget->GetSizeXY();
+
+				// allow only upscaling
+				float FractionX = FMath::Clamp((float)GSystemResolution.ResX / WindowSize.X, 0.1f, 4.0f);
+				float FractionY = FMath::Clamp((float)GSystemResolution.ResY / WindowSize.Y, 0.1f, 4.0f);
+
+				// maintain a pixel aspect ratio of 1:1 for easier internal computations
+				Fraction *= FMath::Max(FractionX, FractionY);
+			}
+		}
+
+#if !(UE_BUILD_SHIPPING || UE_BUILD_TEST)
+		if(CVarScreenPercentageEditor.GetValueOnAnyThread() == 0)
+		{
+			bool bNotInGame = GEngine && GEngine->GameViewport == 0;
+
+			if(bNotInGame)
+			{
+				Fraction = 1.0f;
+			}
+		}
+#endif
+
+
+		// Upscale if needed
+		if (Fraction != 1.0f)
+		{
+			// compute the view rectangle with the ScreenPercentage applied
+			const FIntRect ScreenPercentageAffectedViewRect = ViewInitOptions.GetConstrainedViewRect().Scale(Fraction);
+			SetScaledViewRect(ScreenPercentageAffectedViewRect);
+		}
+	}
 }
 
 void FSceneView::ConfigureBufferVisualizationSettings()
diff --git a/Engine/Source/Runtime/Engine/Public/SceneView.h b/Engine/Source/Runtime/Engine/Public/SceneView.h
--- a/Engine/Source/Runtime/Engine/Public/SceneView.h
+++ b/Engine/Source/Runtime/Engine/Public/SceneView.h
@@ -442,7 +442,7 @@ public:
 	void OverridePostProcessSettings(const FPostProcessSettings& Src, float Weight);
 
 	/** applied global restrictions from show flags */
-	void EndFinalPostprocessSettings();
+	void EndFinalPostprocessSettings(FSceneViewInitOptions ViewInitOptions);
 
 	/** Configure post process settings for the buffer visualization system */
 	void ConfigureBufferVisualizationSettings();
@@ -555,10 +555,10 @@ public:
 	/** The views which make up the family. */
 	TArray<const FSceneView*> Views;
 
-	/** The width in screen pixels of the view family being rendered. */
+	/** The width in screen pixels of the view family being rendered (maximum x of all viewports). */
 	uint32 FamilySizeX;
 
-	/** The height in screen pixels of the view family being rendered. */
+	/** The height in screen pixels of the view family being rendered (maximum y of all viewports). */
 	uint32 FamilySizeY;
 
 	/** The render target which the views are being rendered to. */
diff --git a/Engine/Source/Runtime/Renderer/Private/PostProcess/PostProcessUpscale.cpp b/Engine/Source/Runtime/Renderer/Private/PostProcess/PostProcessUpscale.cpp
--- a/Engine/Source/Runtime/Renderer/Private/PostProcess/PostProcessUpscale.cpp
+++ b/Engine/Source/Runtime/Renderer/Private/PostProcess/PostProcessUpscale.cpp
@@ -13,14 +13,71 @@
 #include "PostProcessEyeAdaptation.h"
 #include "SceneUtils.h"
 
-static TAutoConsoleVariable<float> CVarScreenPercentageSoftness(
-	TEXT("r.ScreenPercentageSoftness"),
+static TAutoConsoleVariable<float> CVarUpscaleSoftness(
+	TEXT("r.Upscale.Softness"),
 	0.3f,
 	TEXT("To scale up with higher quality loosing some sharpness\n")
 	TEXT(" 0..1 (0.3 is good for ScreenPercentage 90"),
 	ECVF_Scalability | ECVF_RenderThreadSafe);
 
-/** Encapsulates the post processing eye adaptation pixel shader. */
+/** Encapsulates the upscale vertex shader. */
+class FPostProcessUpscaleVS : public FPostProcessVS
+{
+	DECLARE_SHADER_TYPE(FPostProcessUpscaleVS,Global);
+
+	/** Default constructor. */
+	FPostProcessUpscaleVS() {}
+
+public:
+	FShaderParameter DistortionParams;
+
+	/** Initialization constructor. */
+	FPostProcessUpscaleVS(const ShaderMetaType::CompiledShaderInitializerType& Initializer)
+		: FPostProcessVS(Initializer)
+	{
+		DistortionParams.Bind(Initializer.ParameterMap,TEXT("DistortionParams"));
+	}
+
+	static void ModifyCompilationEnvironment(EShaderPlatform Platform, FShaderCompilerEnvironment& OutEnvironment)
+	{
+		FPostProcessVS::ModifyCompilationEnvironment(Platform,OutEnvironment);
+		OutEnvironment.SetDefine(TEXT("TESS_RECT_X"), FTesselatedScreenRectangleIndexBuffer::Width);
+		OutEnvironment.SetDefine(TEXT("TESS_RECT_Y"), FTesselatedScreenRectangleIndexBuffer::Height);
+	}
+
+	// @param InCylinderDistortion 0=none..1=full in percent, must be in that range
+	void SetParameters(const FRenderingCompositePassContext& Context, float InCylinderDistortion)
+	{
+		const FVertexShaderRHIParamRef ShaderRHI = GetVertexShader();
+
+		FGlobalShader::SetParameters(Context.RHICmdList, ShaderRHI, Context.View);
+
+		{
+			float HalfFOV = FMath::Atan(1.0f / Context.View.ViewMatrices.ProjMatrix.M[0][0]);
+			float TanHalfFov = 1.0f / Context.View.ViewMatrices.ProjMatrix.M[0][0];
+			float InvHalfFov = 1.0f / HalfFOV;
+
+			// compute Correction to scale Y the same as X is scaled in the center
+			float SmallX = 0.01f;
+			float Correction = atan(SmallX * TanHalfFov) * InvHalfFov / SmallX;
+
+			FVector4 Value(HalfFOV, TanHalfFov, InCylinderDistortion, Correction);
+
+			SetShaderValue(Context.RHICmdList, ShaderRHI, DistortionParams, Value);
+		}
+	}
+
+	virtual bool Serialize(FArchive& Ar)
+	{
+		bool bShaderHasOutdatedParameters = FPostProcessVS::Serialize(Ar);
+		Ar << DistortionParams;
+		return bShaderHasOutdatedParameters;
+	}
+};
+
+IMPLEMENT_SHADER_TYPE(,FPostProcessUpscaleVS,TEXT("PostProcessUpscale"),TEXT("MainVS"),SF_Vertex);
+
+/** Encapsulates the post processing upscale pixel shader. */
 template <uint32 Method>
 class FPostProcessUpscalePS : public FGlobalShader
 {
@@ -67,10 +124,8 @@ public:
 		PostprocessParameter.SetPS(ShaderRHI, Context, 0, false, FilterTable);
 		DeferredParameters.Set(Context.RHICmdList, ShaderRHI, Context.View);
 
-		// If the method needs softness value
-		if(Method == 2)
 		{
-			float UpscaleSoftnessValue = FMath::Clamp(CVarScreenPercentageSoftness.GetValueOnRenderThread(), 0.0f, 1.0f);
+			float UpscaleSoftnessValue = FMath::Clamp(CVarUpscaleSoftness.GetValueOnRenderThread(), 0.0f, 1.0f);
 
 			SetShaderValue(Context.RHICmdList, ShaderRHI, UpscaleSoftness, UpscaleSoftnessValue);
 		}
@@ -107,23 +162,45 @@ VARIATION1(3)
 
 #undef VARIATION1
 
-FRCPassPostProcessUpscale::FRCPassPostProcessUpscale(uint32 InUpscaleMethod)
-	: UpscaleMethod(InUpscaleMethod)
+FRCPassPostProcessUpscale::FRCPassPostProcessUpscale(uint32 InUpscaleQuality, float InCylinderDistortion)
+	: UpscaleQuality(InUpscaleQuality)
+	, CylinderDistortion(FMath::Clamp(InCylinderDistortion, 0.0f, 1.0f))
 {
 }
 
-template <uint32 Method>
-void FRCPassPostProcessUpscale::SetShader(const FRenderingCompositePassContext& Context)
+template <uint32 Method, uint32 bTesselatedQuad>
+FShader* FRCPassPostProcessUpscale::SetShader(const FRenderingCompositePassContext& Context, float InCylinderDistortion)
 {
-	TShaderMapRef<FPostProcessVS> VertexShader(Context.GetShaderMap());
-	TShaderMapRef<FPostProcessUpscalePS<Method> > PixelShader(Context.GetShaderMap());
+	if(bTesselatedQuad)
+	{
+		check(InCylinderDistortion > 0.0f);
 
-	static FGlobalBoundShaderState BoundShaderState;
-	
+		TShaderMapRef<FPostProcessUpscaleVS> VertexShader(Context.GetShaderMap());
+		TShaderMapRef<FPostProcessUpscalePS<Method> > PixelShader(Context.GetShaderMap());
+
+		static FGlobalBoundShaderState BoundShaderState;
+
+		SetGlobalBoundShaderState(Context.RHICmdList, Context.GetFeatureLevel(), BoundShaderState, GFilterVertexDeclaration.VertexDeclarationRHI, *VertexShader, *PixelShader);
+
+		PixelShader->SetPS(Context);
+		VertexShader->SetParameters(Context, InCylinderDistortion);
+		return *VertexShader;
+	}
+	else
+	{
+		check(InCylinderDistortion == 0.0f);
+
+		TShaderMapRef<FPostProcessVS> VertexShader(Context.GetShaderMap());
+		TShaderMapRef<FPostProcessUpscalePS<Method> > PixelShader(Context.GetShaderMap());
+
+		static FGlobalBoundShaderState BoundShaderState;
 
-	SetGlobalBoundShaderState(Context.RHICmdList, Context.GetFeatureLevel(), BoundShaderState, GFilterVertexDeclaration.VertexDeclarationRHI, *VertexShader, *PixelShader);
+		SetGlobalBoundShaderState(Context.RHICmdList, Context.GetFeatureLevel(), BoundShaderState, GFilterVertexDeclaration.VertexDeclarationRHI, *VertexShader, *PixelShader);
 
-	PixelShader->SetPS(Context);
+		PixelShader->SetPS(Context);
+		VertexShader->SetParameters(Context);
+		return *VertexShader;
+	}
 }
 
 void FRCPassPostProcessUpscale::Process(FRenderingCompositePassContext& Context)
@@ -150,32 +227,48 @@ void FRCPassPostProcessUpscale::Process(FRenderingCompositePassContext& Context)
 	SetRenderTarget(Context.RHICmdList, DestRenderTarget.TargetableTexture, FTextureRHIRef());
 	Context.SetViewportAndCallRHI(DestRect);
 
+	bool bTessellatedQuad = CylinderDistortion >= 0.01f;
+
+	// with distortion (bTessellatedQuad) we need to clear the background
+	FIntRect ExcludeRect = bTessellatedQuad ? FIntRect() : View.ViewRect;
+
+	Context.RHICmdList.Clear(true, FLinearColor::Black, false, 1.0f, false, 0, ExcludeRect);
+
 	// set the state
 	Context.RHICmdList.SetBlendState(TStaticBlendState<>::GetRHI());
 	Context.RHICmdList.SetRasterizerState(TStaticRasterizerState<>::GetRHI());
 	Context.RHICmdList.SetDepthStencilState(TStaticDepthStencilState<false, CF_Always>::GetRHI());
 
-	switch (UpscaleMethod)
-	{
-		case 0:
-			SetShader<0>(Context);
-		break;
-		case 1:
-			SetShader<1>(Context);
-		break;
-		case 2:
-			SetShader<2>(Context);
-		break;
-		case 3:
-			SetShader<3>(Context);
-		break;
-		default:
-			checkNoEntry();
-		break;
-	}
-
-	// Draw a quad mapping scene color to the view's render target
-	TShaderMapRef<FPostProcessVS> VertexShader(Context.GetShaderMap());
+	FShader* VertexShader = 0;
+
+	if(bTessellatedQuad)
+	{
+		switch (UpscaleQuality)
+		{
+			case 0:	VertexShader = SetShader<0, 1>(Context, CylinderDistortion); break;
+			case 1:	VertexShader = SetShader<1, 1>(Context, CylinderDistortion); break;
+			case 2:	VertexShader = SetShader<2, 1>(Context, CylinderDistortion); break;
+			case 3:	VertexShader = SetShader<3, 1>(Context, CylinderDistortion); break;
+			default:
+				checkNoEntry();
+				break;
+		}
+	}
+	else
+	{
+		switch (UpscaleQuality)
+		{
+			case 0:	VertexShader = SetShader<0, 0>(Context); break;
+			case 1:	VertexShader = SetShader<1, 0>(Context); break;
+			case 2:	VertexShader = SetShader<2, 0>(Context); break;
+			case 3:	VertexShader = SetShader<3, 0>(Context); break;
+			default:
+				checkNoEntry();
+				break;
+		}
+	}
+
+	// Draw a quad, a triangle or a tessellated quad
 	DrawRectangle(
 		Context.RHICmdList,
 		0, 0,
@@ -184,8 +277,8 @@ void FRCPassPostProcessUpscale::Process(FRenderingCompositePassContext& Context)
 		SrcRect.Width(), SrcRect.Height(),
 		DestRect.Size(),
 		SrcSize,
-		*VertexShader,
-		EDRF_UseTriangleOptimization);
+		VertexShader,
+		bTessellatedQuad ? EDRF_UseTesselatedIndexBuffer: EDRF_UseTriangleOptimization);
 
 	Context.RHICmdList.CopyToResolveTarget(DestRenderTarget.TargetableTexture, DestRenderTarget.ShaderResourceTexture, false, FResolveParams());
 }
diff --git a/Engine/Source/Runtime/Renderer/Private/PostProcess/PostProcessUpscale.h b/Engine/Source/Runtime/Renderer/Private/PostProcess/PostProcessUpscale.h
--- a/Engine/Source/Runtime/Renderer/Private/PostProcess/PostProcessUpscale.h
+++ b/Engine/Source/Runtime/Renderer/Private/PostProcess/PostProcessUpscale.h
@@ -15,19 +15,25 @@ class FRCPassPostProcessUpscale : public TRenderingCompositePassBase<2, 1>
 {
 public:
 	// constructor
-	// @param InUpscaleMethod - value denoting Upscale method to use:
+	// @param InUpscaleQuality - value denoting Upscale method to use:
 	//				0: Nearest
 	//				1: Bilinear
 	//				2: 4 tap Bilinear (with radius adjustment)
 	//				3: Directional blur with unsharp mask upsample.
-	FRCPassPostProcessUpscale(uint32 InUpscaleMethod);
+	// @param InCylinderDistortion 0=none..1=full in percent, can be outside of that range
+	FRCPassPostProcessUpscale(uint32 InUpscaleQuality, float InCylinderDistortion);
 
 	// interface FRenderingCompositePass ---------
 	virtual void Process(FRenderingCompositePassContext& Context);
 	virtual void Release() override { delete this; }
 	FPooledRenderTargetDesc ComputeOutputDesc(EPassOutputId InPassOutputId) const;
 private:
-	template <uint32 Method> static void SetShader(const FRenderingCompositePassContext& Context);
-	uint32 UpscaleMethod;
+	// @param InCylinderDistortion 0=none..1=full in percent, must be in that range
+	template <uint32 Quality, uint32 bTesselatedQuad> static FShader* SetShader(const FRenderingCompositePassContext& Context, float InCylinderDistortion = 0.0f);
+
+	// 0: Nearest, 1: Bilinear, 2: 4 tap Bilinear (with radius adjustment), 3: Directional blur with unsharp mask upsample.
+	uint32 UpscaleQuality;
+	// 0=none..1=full, must be in that range
+	float CylinderDistortion;
 };
 
diff --git a/Engine/Source/Runtime/Renderer/Private/PostProcess/PostProcessing.cpp b/Engine/Source/Runtime/Renderer/Private/PostProcess/PostProcessing.cpp
--- a/Engine/Source/Runtime/Renderer/Private/PostProcess/PostProcessing.cpp
+++ b/Engine/Source/Runtime/Renderer/Private/PostProcess/PostProcessing.cpp
@@ -71,8 +71,18 @@ static TAutoConsoleVariable<int32> CVarRenderTargetSwitchWorkaround(
 	TEXT("We want this enabled (1) on all 32 bit iOS devices (implemented through DeviceProfiles)."),
 	ECVF_RenderThreadSafe);
 
-static TAutoConsoleVariable<int32> CVarUpsampleQuality(
-	TEXT("r.UpsampleQuality"),
+static TAutoConsoleVariable<float> CVarUpscaleCylinder(
+	TEXT("r.Upscale.Cylinder"),
+	0,
+	TEXT("Allows to apply a cylindrical distortion to the rendered image. Values between 0 and 1 allow to fade the effect (lerp).\n")
+	TEXT("There is a quality loss that can be compensated by adjusting r.ScreenPercentage (>100).\n")
+	TEXT("0: off(default)\n")
+	TEXT(">0: enabled (requires an extra post processing pass if upsampling wasn't used - see r.ScreenPercentage)\n")
+	TEXT("1: full effect"),
+	ECVF_RenderThreadSafe);
+
+static TAutoConsoleVariable<int32> CVarUpscaleQuality(
+	TEXT("r.Upscale.Quality"),
 	3,
 	TEXT(" 0: Nearest filtering\n")
 	TEXT(" 1: Simple Bilinear\n")
@@ -1295,13 +1305,21 @@ void FPostProcessing::Process(FRHICommandListImmediate& RHICmdList, FViewInfo& V
 
 		AddHighResScreenshotMask(Context, SeparateTranslucency);
 
+		// 0=none..1=full
+		float Cylinder = 0.0f;
+
+		if(View.IsPerspectiveProjection() && !GEngine->StereoRenderingDevice.IsValid())
+		{
+			Cylinder = FMath::Clamp(CVarUpscaleCylinder.GetValueOnRenderThread(), 0.0f, 1.0f);
+		}
+
 		// Do not use upscale if SeparateRenderTarget is in use!
-		if (View.UnscaledViewRect != View.ViewRect && (!View.Family->EngineShowFlags.StereoRendering || (!View.Family->EngineShowFlags.HMDDistortion && !View.Family->bUseSeparateRenderTarget)))
+		if ((Cylinder > 0.01f || View.UnscaledViewRect != View.ViewRect)
+			&& (!View.Family->EngineShowFlags.StereoRendering || (!View.Family->EngineShowFlags.HMDDistortion && !View.Family->bUseSeparateRenderTarget)))
 		{
-			static const auto CVar = IConsoleManager::Get().FindTConsoleVariableDataInt(TEXT("r.UpsampleQuality")); 
-			int32 UpsampleMethod = CVar->GetValueOnRenderThread();
-			UpsampleMethod = FMath::Clamp(UpsampleMethod, 0, 3);
-			FRenderingCompositePass* Node = Context.Graph.RegisterPass(new(FMemStack::Get()) FRCPassPostProcessUpscale(UpsampleMethod));
+			int32 UpscaleQuality = CVarUpscaleQuality.GetValueOnRenderThread();
+			UpscaleQuality = FMath::Clamp(UpscaleQuality, 0, 3);
+			FRenderingCompositePass* Node = Context.Graph.RegisterPass(new(FMemStack::Get()) FRCPassPostProcessUpscale(UpscaleQuality, Cylinder));
 			Node->SetInput(ePId_Input0, FRenderingCompositeOutputRef(Context.FinalOutput)); // Bilinear sampling.
 			Node->SetInput(ePId_Input1, FRenderingCompositeOutputRef(Context.FinalOutput)); // Point sampling.
 			Context.FinalOutput = FRenderingCompositeOutputRef(Node);
diff --git a/Engine/Source/Runtime/Renderer/Private/PostProcess/SceneFilterRendering.cpp b/Engine/Source/Runtime/Renderer/Private/PostProcess/SceneFilterRendering.cpp
--- a/Engine/Source/Runtime/Renderer/Private/PostProcess/SceneFilterRendering.cpp
+++ b/Engine/Source/Runtime/Renderer/Private/PostProcess/SceneFilterRendering.cpp
@@ -70,6 +70,56 @@ public:
 static TGlobalResource<FScreenRectangleVertexBuffer> GScreenRectangleVertexBuffer;
 static TGlobalResource<FScreenRectangleIndexBuffer> GScreenRectangleIndexBuffer;
 
+void FTesselatedScreenRectangleIndexBuffer::InitRHI()
+{
+	TResourceArray<uint16, INDEXBUFFER_ALIGNMENT> IndexBuffer;
+
+	uint32 NumIndices = NumPrimitives() * 3;
+	IndexBuffer.AddUninitialized(NumIndices);
+		
+	uint16* Out = (uint16*)IndexBuffer.GetData();
+
+	for(uint32 y = 0; y < Height; ++y)
+	{
+		for(uint32 x = 0; x < Width; ++x)
+		{
+			// left top to bottom right in reading order
+			uint16 Index00 = x  + y * (Width + 1);
+			uint16 Index10 = Index00 + 1;
+			uint16 Index01 = Index00 + (Width + 1);
+			uint16 Index11 = Index01 + 1;
+
+			// todo: diagonal can be flipped on parts of the screen
+
+			// triangle A
+			*Out++ = Index00; *Out++ = Index01; *Out++ = Index10;
+
+			// triangle B
+			*Out++ = Index11; *Out++ = Index10; *Out++ = Index01;
+		}
+	}
+
+	// Create index buffer. Fill buffer with initial data upon creation
+	FRHIResourceCreateInfo CreateInfo(&IndexBuffer);
+	IndexBufferRHI = RHICreateIndexBuffer(sizeof(uint16), IndexBuffer.GetResourceDataSize(), BUF_Static, CreateInfo);
+}
+
+uint32 FTesselatedScreenRectangleIndexBuffer::NumVertices() const
+{
+	// 4 vertices per quad but shared
+	return (Width + 1) * (Height + 1);
+}
+
+uint32 FTesselatedScreenRectangleIndexBuffer::NumPrimitives() const
+{
+	// triangle has 3 corners, 2 triangle per quad
+	return 2 * Width * Height;
+}
+
+/** We don't need a vertex buffer as we can compute the vertex attributes in the VS */
+static TGlobalResource<FTesselatedScreenRectangleIndexBuffer> GTesselatedScreenRectangleIndexBuffer;
+
+
 /** Vertex declaration for the 2D screen rectangle. */
 TGlobalResource<FFilterVertexDeclaration> GFilterVertexDeclaration;
 /** Vertex declaration for vertex shaders that don't require any inputs (eg. generated via vertex ID). */
@@ -149,38 +199,59 @@ void DrawRectangle(
 		1.0f / TextureSize.X, 1.0f / TextureSize.Y);
 
 	SetUniformBufferParameterImmediate(RHICmdList, VertexShader->GetVertexShader(), VertexShader->GetUniformBufferParameter<FDrawRectangleParameters>(), Parameters);
-	RHICmdList.SetStreamSource(0, GScreenRectangleVertexBuffer.VertexBufferRHI, sizeof(FFilterVertex), 0);
 
-
-	if (Flags == EDRF_UseTriangleOptimization)
+	if(Flags == EDRF_UseTesselatedIndexBuffer)
 	{
-		// A single triangle spans the entire viewport this results in a quad that fill the viewport. This can increase rasterization efficiency
-		// as we do not have a diagonal edge (through the center) for the rasterizer/span-dispatch. Although the actual benefit of this technique is dependent upon hardware.
+		// no vertex buffer needed as we compute it in VS
+		RHICmdList.SetStreamSource(0, NULL, 0, 0);
 
-		// We offset into the index buffer when using the triangle optimization to access the correct vertices.
 		RHICmdList.DrawIndexedPrimitive(
-			GScreenRectangleIndexBuffer.IndexBufferRHI,
+			GTesselatedScreenRectangleIndexBuffer.IndexBufferRHI,
 			PT_TriangleList,
 			/*BaseVertexIndex=*/ 0,
 			/*MinIndex=*/ 0,
-			/*NumVertices=*/ 3,
-			/*StartIndex=*/ 6,
-			/*NumPrimitives=*/ 1,
+			/*NumVertices=*/ GTesselatedScreenRectangleIndexBuffer.NumVertices(),
+			/*StartIndex=*/ 0,
+			/*NumPrimitives=*/ GTesselatedScreenRectangleIndexBuffer.NumPrimitives(),
 			/*NumInstances=*/ 1
 			);
 	}
 	else
 	{
-		RHICmdList.DrawIndexedPrimitive(
-			GScreenRectangleIndexBuffer.IndexBufferRHI,
-			PT_TriangleList,
-			/*BaseVertexIndex=*/ 0,
-			/*MinIndex=*/ 0,
-			/*NumVertices=*/ 4,
-			/*StartIndex=*/ 0,
-			/*NumPrimitives=*/ 2,
-			/*NumInstances=*/ 1
-			);
+		RHICmdList.SetStreamSource(0, GScreenRectangleVertexBuffer.VertexBufferRHI, sizeof(FFilterVertex), 0);
+
+		if (Flags == EDRF_UseTriangleOptimization)
+		{
+			// A single triangle spans the entire viewport this results in a quad that fill the viewport. This can increase rasterization efficiency
+			// as we do not have a diagonal edge (through the center) for the rasterizer/span-dispatch. Although the actual benefit of this technique is dependent upon hardware.
+
+			// We offset into the index buffer when using the triangle optimization to access the correct vertices.
+			RHICmdList.DrawIndexedPrimitive(
+				GScreenRectangleIndexBuffer.IndexBufferRHI,
+				PT_TriangleList,
+				/*BaseVertexIndex=*/ 0,
+				/*MinIndex=*/ 0,
+				/*NumVertices=*/ 3,
+				/*StartIndex=*/ 6,
+				/*NumPrimitives=*/ 1,
+				/*NumInstances=*/ 1
+				);
+		}
+		else
+		{
+			RHICmdList.SetStreamSource(0, GScreenRectangleVertexBuffer.VertexBufferRHI, sizeof(FFilterVertex), 0);
+
+			RHICmdList.DrawIndexedPrimitive(
+				GScreenRectangleIndexBuffer.IndexBufferRHI,
+				PT_TriangleList,
+				/*BaseVertexIndex=*/ 0,
+				/*MinIndex=*/ 0,
+				/*NumVertices=*/ 4,
+				/*StartIndex=*/ 0,
+				/*NumPrimitives=*/ 2,
+				/*NumInstances=*/ 1
+				);
+		}
 	}
 }
 
diff --git a/Engine/Source/Runtime/Renderer/Private/PostProcess/SceneFilterRendering.h b/Engine/Source/Runtime/Renderer/Private/PostProcess/SceneFilterRendering.h
--- a/Engine/Source/Runtime/Renderer/Private/PostProcess/SceneFilterRendering.h
+++ b/Engine/Source/Runtime/Renderer/Private/PostProcess/SceneFilterRendering.h
@@ -154,4 +154,24 @@ private:
 	FShaderParameter				RenderTargetExtent;
 };
 
+
+class FTesselatedScreenRectangleIndexBuffer : public FIndexBuffer
+{
+public:
+
+	// if one of those constants change, UpscaleVS needs to be recompiled
+
+	// number of quads in x
+	static const uint32 Width = 32;		// used for CylindricalProjection (smaller FOV could do less tessellation)
+	// number of quads in y
+	static const uint32 Height = 20;	// to minimize distortion we also tessellate in Y but a perspective distortion could do that with fewer triangles.
+
+	/** Initialize the RHI for this rendering resource */
+	void InitRHI() override;
+
+	uint32 NumVertices() const;
+
+	uint32 NumPrimitives() const;
+};
+
 #endif
diff --git a/Engine/Source/Runtime/Renderer/Private/ReflectionEnvironmentCapture.cpp b/Engine/Source/Runtime/Renderer/Private/ReflectionEnvironmentCapture.cpp
--- a/Engine/Source/Runtime/Renderer/Private/ReflectionEnvironmentCapture.cpp
+++ b/Engine/Source/Runtime/Renderer/Private/ReflectionEnvironmentCapture.cpp
@@ -1185,7 +1185,7 @@ void CaptureSceneIntoScratchCubemap(FScene* Scene, FVector CapturePosition, bool
 		View->bIsReflectionCapture = true;
 		View->bStaticSceneOnly = bStaticSceneOnly;
 		View->StartFinalPostprocessSettings(CapturePosition);
-		View->EndFinalPostprocessSettings();
+		View->EndFinalPostprocessSettings(ViewInitOptions);
 
 		ViewFamily.Views.Add(View);
 
diff --git a/Engine/Source/Runtime/Renderer/Private/SceneCaptureRendering.cpp b/Engine/Source/Runtime/Renderer/Private/SceneCaptureRendering.cpp
--- a/Engine/Source/Runtime/Renderer/Private/SceneCaptureRendering.cpp
+++ b/Engine/Source/Runtime/Renderer/Private/SceneCaptureRendering.cpp
@@ -233,7 +233,7 @@ FSceneRenderer* FScene::CreateSceneRenderer( USceneCaptureComponent* SceneCaptur
 	{
 		View->OverridePostProcessSettings(*PostProcessSettings, PostProcessBlendWeight);
 	}
-	View->EndFinalPostprocessSettings();
+	View->EndFinalPostprocessSettings(ViewInitOptions);
 
 	return FSceneRenderer::CreateSceneRenderer(&ViewFamily, NULL);
 }
diff --git a/Engine/Source/Runtime/Renderer/Public/RendererInterface.h b/Engine/Source/Runtime/Renderer/Public/RendererInterface.h
--- a/Engine/Source/Runtime/Renderer/Public/RendererInterface.h
+++ b/Engine/Source/Runtime/Renderer/Public/RendererInterface.h
@@ -439,7 +439,9 @@ enum EDrawRectangleFlags
 	// Rectangle is created by 2 triangles (diagonal can cause some slightly less efficient shader execution), this is the default as it has no artifacts
 	EDRF_Default,
 	//
-	EDRF_UseTriangleOptimization
+	EDRF_UseTriangleOptimization,
+	//
+	EDRF_UseTesselatedIndexBuffer
 };
 
 class ICustomVisibilityQuery: public IRefCountedObject
diff --git a/Engine/Source/Runtime/UMG/Private/Components/Viewport.cpp b/Engine/Source/Runtime/UMG/Private/Components/Viewport.cpp
--- a/Engine/Source/Runtime/UMG/Private/Components/Viewport.cpp
+++ b/Engine/Source/Runtime/UMG/Private/Components/Viewport.cpp
@@ -434,7 +434,7 @@ FSceneView* FUMGViewportClient::CalcSceneView(FSceneViewFamily* ViewFamily)
 
 	//OverridePostProcessSettings(*View);
 
-	View->EndFinalPostprocessSettings();
+	View->EndFinalPostprocessSettings(ViewInitOptions);
 
 	return View;
 }
