GetDynamicMeshElements path on scene proxies is enabled by default
* PreRenderView / DrawDynamicElements is no longer used and will be removed

diff --git a/Engine/Source/Runtime/Renderer/Private/DeferredShadingRenderer.cpp b/Engine/Source/Runtime/Renderer/Private/DeferredShadingRenderer.cpp
--- a/Engine/Source/Runtime/Renderer/Private/DeferredShadingRenderer.cpp
+++ b/Engine/Source/Runtime/Renderer/Private/DeferredShadingRenderer.cpp
@@ -56,7 +56,7 @@ static TAutoConsoleVariable<int32> CVarVisualizeTexturePool(
 /**  */
 static TAutoConsoleVariable<int32> CVarUseGetMeshElements(
 	TEXT("r.UseGetMeshElements"),
-	0,
+	1,
 	TEXT("Whether to use the GetMeshElements path or the legacy DrawDynamicElements"),
 	ECVF_RenderThreadSafe
 	);
