Fix case where necessary StaticMeshElementUpdates could be lost

diff --git a/Engine/Source/Runtime/Renderer/Private/RendererScene.cpp b/Engine/Source/Runtime/Renderer/Private/RendererScene.cpp
--- a/Engine/Source/Runtime/Renderer/Private/RendererScene.cpp
+++ b/Engine/Source/Runtime/Renderer/Private/RendererScene.cpp
@@ -645,7 +645,7 @@ void FScene::AddLightSceneInfo_RenderThread(FLightSceneInfo* LightSceneInfo)
 		SimpleDirectionalLight = LightSceneInfo;
 
 		// if we are forward rendered and this light is a dynamic shadowcast then we need to update the static draw lists to pick a new lightingpolicy
-		bScenesPrimitivesNeedStaticMeshElementUpdate = !ShouldUseDeferredRenderer() && !SimpleDirectionalLight->Proxy->HasStaticShadowing();
+		bScenesPrimitivesNeedStaticMeshElementUpdate = bScenesPrimitivesNeedStaticMeshElementUpdate || (!ShouldUseDeferredRenderer() && !SimpleDirectionalLight->Proxy->HasStaticShadowing());		
 	}
 
 	if (LightSceneInfo->Proxy->IsUsedAsAtmosphereSunLight() &&
@@ -1095,7 +1095,7 @@ void FScene::RemoveLightSceneInfo_RenderThread(FLightSceneInfo* LightSceneInfo)
 		if (LightSceneInfo == SimpleDirectionalLight)
 		{
 			// if we are forward rendered and this light is a dynamic shadowcast then we need to update the static draw lists to pick a new lightingpolicy
-			bScenesPrimitivesNeedStaticMeshElementUpdate = !ShouldUseDeferredRenderer() && !SimpleDirectionalLight->Proxy->HasStaticShadowing();
+			bScenesPrimitivesNeedStaticMeshElementUpdate = bScenesPrimitivesNeedStaticMeshElementUpdate  || (!ShouldUseDeferredRenderer() && !SimpleDirectionalLight->Proxy->HasStaticShadowing());
 			SimpleDirectionalLight = NULL;
 		}
 
