fixed LDR SceneCaptureing after recent changes to CaptureEveryFrame (added State)

diff --git a/Engine/Source/Runtime/Renderer/Private/PostProcess/PostProcessing.cpp b/Engine/Source/Runtime/Renderer/Private/PostProcess/PostProcessing.cpp
--- a/Engine/Source/Runtime/Renderer/Private/PostProcess/PostProcessing.cpp
+++ b/Engine/Source/Runtime/Renderer/Private/PostProcess/PostProcessing.cpp
@@ -1050,7 +1050,7 @@ void FPostProcessing::Process(FRHICommandListImmediate& RHICmdList, FViewInfo& V
 			}
 
 			// some views don't have a state (thumbnail rendering)
-			if(!GIsHighResScreenshot && View.State && (StereoPass != eSSP_RIGHT_EYE) && !View.bIsSceneCapture)
+			if(!GIsHighResScreenshot && View.State && (StereoPass != eSSP_RIGHT_EYE))
 			{
 				// we always add eye adaptation, if the engine show flag is disabled we set the ExposureScale in the texture to a fixed value
 				EyeAdaptation = AddPostProcessEyeAdaptation(Context, Histogram);
diff --git a/Engine/Source/Runtime/Renderer/Private/SceneCaptureRendering.cpp b/Engine/Source/Runtime/Renderer/Private/SceneCaptureRendering.cpp
--- a/Engine/Source/Runtime/Renderer/Private/SceneCaptureRendering.cpp
+++ b/Engine/Source/Runtime/Renderer/Private/SceneCaptureRendering.cpp
@@ -229,10 +229,7 @@ FSceneRenderer* FScene::CreateSceneRenderer( USceneCaptureComponent* SceneCaptur
 	ViewFamily.Views.Add(View);
 
 	View->StartFinalPostprocessSettings(ViewLocation);
-	if (!bCaptureSceneColour)
-	{
-		View->OverridePostProcessSettings(*PostProcessSettings, PostProcessBlendWeight);
-	}
+	View->OverridePostProcessSettings(*PostProcessSettings, PostProcessBlendWeight);
 	View->EndFinalPostprocessSettings(ViewInitOptions);
 
 	return FSceneRenderer::CreateSceneRenderer(&ViewFamily, NULL);
