Raised default shadowmap resolution to 2048
* Translucent shadows are now half res compared to opaque shadows, to keep memory usage under control
* New memory usage is 8mb for ShadowDepthZ, 8mb for PreshadowCache, 8*2mb for translucent shadow targets (unchanged), 12mb for largest cubemap shadowmap

diff --git a/Engine/Config/BaseScalability.ini b/Engine/Config/BaseScalability.ini
--- a/Engine/Config/BaseScalability.ini
+++ b/Engine/Config/BaseScalability.ini
@@ -90,7 +90,7 @@ r.DistanceFieldAO=0
 r.LightFunctionQuality=1
 r.ShadowQuality=5
 r.Shadow.CSM.MaxCascades=4
-r.Shadow.MaxResolution=1024
+r.Shadow.MaxResolution=2048
 r.Shadow.RadiusThreshold=0.03
 r.Shadow.DistanceScale=1.0
 r.Shadow.CSM.TransitionScale=1.0
diff --git a/Engine/Source/Runtime/Core/Private/HAL/ConsoleManager.cpp b/Engine/Source/Runtime/Core/Private/HAL/ConsoleManager.cpp
--- a/Engine/Source/Runtime/Core/Private/HAL/ConsoleManager.cpp
+++ b/Engine/Source/Runtime/Core/Private/HAL/ConsoleManager.cpp
@@ -1873,7 +1873,7 @@ static TAutoConsoleVariable<int32> CVarMaxAnistropy(
 
 static TAutoConsoleVariable<int32> CVarShadowMaxResolution(
 	TEXT("r.Shadow.MaxResolution"),
-	1024,
+	2048,
 	TEXT("Max square dimensions (in texels) allowed for rendering shadow depths. Range 4 to hardware limit. Higher = better quality shadows but at a performance cost."),
 	ECVF_Scalability | ECVF_RenderThreadSafe);
 
diff --git a/Engine/Source/Runtime/Renderer/Private/PostProcess/SceneRenderTargets.cpp b/Engine/Source/Runtime/Renderer/Private/PostProcess/SceneRenderTargets.cpp
--- a/Engine/Source/Runtime/Renderer/Private/PostProcess/SceneRenderTargets.cpp
+++ b/Engine/Source/Runtime/Renderer/Private/PostProcess/SceneRenderTargets.cpp
@@ -1488,7 +1488,7 @@ void FSceneRenderTargets::ReleaseDynamicRHI()
 /** Returns the size of the shadow depth buffer, taking into account platform limitations and game specific resolution limits. */
 FIntPoint FSceneRenderTargets::GetShadowDepthTextureResolution() const
 {
-	int32 MaxShadowRes = GetCachedScalabilityCVars().MaxShadowResolution;
+	int32 MaxShadowRes = CurrentMaxShadowResolution;
 	const FIntPoint ShadowBufferResolution(
 			FMath::Clamp(MaxShadowRes,1,GMaxShadowDepthBufferSizeX),
 			FMath::Clamp(MaxShadowRes,1,GMaxShadowDepthBufferSizeY));
@@ -1513,7 +1513,9 @@ FIntPoint FSceneRenderTargets::GetPreShadowCacheTextureResolution() const
 
 FIntPoint FSceneRenderTargets::GetTranslucentShadowDepthTextureResolution() const
 {
-	const FIntPoint ShadowDepthResolution = GetShadowDepthTextureResolution();
+	FIntPoint ShadowDepthResolution = GetShadowDepthTextureResolution();
+	ShadowDepthResolution.X = FMath::Max<int32>(ShadowDepthResolution.X / GetTranslucentShadowDownsampleFactor(), 1);
+	ShadowDepthResolution.Y = FMath::Max<int32>(ShadowDepthResolution.Y / GetTranslucentShadowDownsampleFactor(), 1);
 	return ShadowDepthResolution;
 }
 
diff --git a/Engine/Source/Runtime/Renderer/Private/PostProcess/SceneRenderTargets.h b/Engine/Source/Runtime/Renderer/Private/PostProcess/SceneRenderTargets.h
--- a/Engine/Source/Runtime/Renderer/Private/PostProcess/SceneRenderTargets.h
+++ b/Engine/Source/Runtime/Renderer/Private/PostProcess/SceneRenderTargets.h
@@ -379,6 +379,7 @@ public:
 	FIntPoint GetShadowDepthTextureResolution() const;
 	FIntPoint GetPreShadowCacheTextureResolution() const;
 	FIntPoint GetTranslucentShadowDepthTextureResolution() const;
+	int32 GetTranslucentShadowDownsampleFactor() const { return 2; }
 
 	/** Returns the size of the RSM buffer, taking into account platform limitations and game specific resolution limits. */
 	FIntPoint GetReflectiveShadowMapTextureResolution() const;
diff --git a/Engine/Source/Runtime/Renderer/Private/ShadowRendering.cpp b/Engine/Source/Runtime/Renderer/Private/ShadowRendering.cpp
--- a/Engine/Source/Runtime/Renderer/Private/ShadowRendering.cpp
+++ b/Engine/Source/Runtime/Renderer/Private/ShadowRendering.cpp
@@ -2353,7 +2353,7 @@ FMatrix FProjectedShadowInfo::GetScreenToShadowMatrix(const FSceneView& View) co
 
 FMatrix FProjectedShadowInfo::GetWorldToShadowMatrix(FVector4& ShadowmapMinMax, const FIntPoint* ShadowBufferResolutionOverride, bool bHasShadowBorder ) const
 {
-	FIntPoint ShadowBufferResolution = ( ShadowBufferResolutionOverride ) ? *ShadowBufferResolutionOverride : GSceneRenderTargets.GetShadowDepthTextureResolution();
+	FIntPoint ShadowBufferResolution = ( ShadowBufferResolutionOverride ) ? *ShadowBufferResolutionOverride : GetShadowBufferResolution();
 	float ShadowBorder = (bHasShadowBorder) ? SHADOW_BORDER : 0.0f;
 
 	const float InvBufferResolutionX = 1.0f / (float)ShadowBufferResolution.X;
@@ -2393,6 +2393,11 @@ FMatrix FProjectedShadowInfo::GetWorldToShadowMatrix(FVector4& ShadowmapMinMax,
 /** Returns the resolution of the shadow buffer used for this shadow, based on the shadow's type. */
 FIntPoint FProjectedShadowInfo::GetShadowBufferResolution() const
 {
+	if (bTranslucentShadow)
+	{
+		return GSceneRenderTargets.GetTranslucentShadowDepthTextureResolution();
+	}
+
 	const FTexture2DRHIRef& ShadowTexture = GSceneRenderTargets.GetShadowDepthZTexture(bAllocatedInPreshadowCache);
 
 	//prefer to return the actual size of the allocated texture if possible.  It may be larger than the size of a single shadowmap due to atlasing (see forward renderer CSM handling in InitDynamicShadows).
diff --git a/Engine/Source/Runtime/Renderer/Private/ShadowSetup.cpp b/Engine/Source/Runtime/Renderer/Private/ShadowSetup.cpp
--- a/Engine/Source/Runtime/Renderer/Private/ShadowSetup.cpp
+++ b/Engine/Source/Runtime/Renderer/Private/ShadowSetup.cpp
@@ -1372,12 +1372,12 @@ void FDeferredShadingSceneRenderer::CreatePerObjectProjectedShadow(
 					LightSceneInfo,
 					PrimitiveSceneInfo,
 					ShadowInitializer,
-					false,					// no preshadow
+					false,					// not preshadow
 					SizeX,
 					MaxShadowResolutionY,
 					MaxScreenPercent,
 					ResolutionFadeAlphas,
-					false					// no tranlucent shadow
+					false					// not translucent shadow
 					);
 				VisibleLightInfo.MemStackProjectedShadows.Add(ProjectedShadowInfo);
 
@@ -1410,12 +1410,13 @@ void FDeferredShadingSceneRenderer::CreatePerObjectProjectedShadow(
 					LightSceneInfo,
 					PrimitiveSceneInfo,
 					ShadowInitializer,
-					false,					// no preshadow
-					SizeX,
-					MaxShadowResolutionY,
+					false,					// not preshadow
+					// Size was computed for the full res opaque shadow, convert to downsampled translucent shadow size with proper clamping
+					FMath::Clamp<int32>(SizeX / GSceneRenderTargets.GetTranslucentShadowDownsampleFactor(), 1, GSceneRenderTargets.GetTranslucentShadowDepthTextureResolution().X - SHADOW_BORDER * 2),
+					FMath::Clamp<int32>(MaxShadowResolutionY / GSceneRenderTargets.GetTranslucentShadowDownsampleFactor(), 1, GSceneRenderTargets.GetTranslucentShadowDepthTextureResolution().Y - SHADOW_BORDER * 2),
 					MaxScreenPercent,
 					ResolutionFadeAlphas,
-					true					// tranlucent shadow
+					true					// translucent shadow
 					);
 				VisibleLightInfo.MemStackProjectedShadows.Add(ProjectedShadowInfo);
 
@@ -1490,12 +1491,12 @@ void FDeferredShadingSceneRenderer::CreatePerObjectProjectedShadow(
 						LightSceneInfo,
 						PrimitiveSceneInfo,
 						ShadowInitializer,
-						true,
+						true,				// preshadow
 						PreshadowSizeX,
 						FMath::TruncToInt(MaxShadowResolutionY * CVarPreShadowResolutionFactor.GetValueOnRenderThread()),
 						MaxScreenPercent,
 						ResolutionPreShadowFadeAlphas,
-						false
+						false				// not translucent shadow
 						);
 				}
 
