Only gather dynamic mesh elements for shadows when cvar is enabled

diff --git a/Engine/Source/Runtime/Renderer/Private/ShadowSetup.cpp b/Engine/Source/Runtime/Renderer/Private/ShadowSetup.cpp
--- a/Engine/Source/Runtime/Renderer/Private/ShadowSetup.cpp
+++ b/Engine/Source/Runtime/Renderer/Private/ShadowSetup.cpp
@@ -1816,18 +1816,21 @@ void FSceneRenderer::InitProjectedShadowVisibility(FRHICommandListImmediate& RHI
 
 void FSceneRenderer::GatherShadowDynamicMeshElements()
 {
-	TArray<const FSceneView*> ReusedViewsArray;
-	ReusedViewsArray.AddZeroed(1);
-
-	for (TSparseArray<FLightSceneInfoCompact>::TConstIterator LightIt(Scene->Lights); LightIt; ++LightIt)
+	if (ShouldUseGetDynamicMeshElements())
 	{
-		FVisibleLightInfo& VisibleLightInfo = VisibleLightInfos[LightIt.GetIndex()];
+		TArray<const FSceneView*> ReusedViewsArray;
+		ReusedViewsArray.AddZeroed(1);
 
-		for (int32 ShadowIndex = 0; ShadowIndex<VisibleLightInfo.AllProjectedShadows.Num(); ShadowIndex++)
+		for (TSparseArray<FLightSceneInfoCompact>::TConstIterator LightIt(Scene->Lights); LightIt; ++LightIt)
 		{
-			FProjectedShadowInfo& ProjectedShadowInfo = *VisibleLightInfo.AllProjectedShadows[ShadowIndex];
+			FVisibleLightInfo& VisibleLightInfo = VisibleLightInfos[LightIt.GetIndex()];
 
-			ProjectedShadowInfo.GatherDynamicMeshElements(*this, VisibleLightInfo, ReusedViewsArray);
+			for (int32 ShadowIndex = 0; ShadowIndex<VisibleLightInfo.AllProjectedShadows.Num(); ShadowIndex++)
+			{
+				FProjectedShadowInfo& ProjectedShadowInfo = *VisibleLightInfo.AllProjectedShadows[ShadowIndex];
+
+				ProjectedShadowInfo.GatherDynamicMeshElements(*this, VisibleLightInfo, ReusedViewsArray);
+			}
 		}
 	}
 }
