Fixed shadow flickering artifacts caused by recent WORLD_MAX increase
* bCameraInsideShadowFrustum was being computed incorrectly, causing frontfaces of the shadow frustum to be drawn even when camera is inside

diff --git a/Engine/Source/Runtime/Core/Public/Math/Plane.h b/Engine/Source/Runtime/Core/Public/Math/Plane.h
--- a/Engine/Source/Runtime/Core/Public/Math/Plane.h
+++ b/Engine/Source/Runtime/Core/Public/Math/Plane.h
@@ -313,8 +313,9 @@ FORCEINLINE FPlane::FPlane( FVector InBase, const FVector &InNormal )
 
 FORCEINLINE FPlane::FPlane( FVector A, FVector B, FVector C )
 	:	FVector( ((B-A)^(C-A)).SafeNormal() )
-	,	W( A | ((B-A)^(C-A)).SafeNormal() )
-{}
+{
+	W = A | (FVector)(*this);
+}
 
 
 FORCEINLINE FPlane::FPlane(EForceInit)
diff --git a/Engine/Source/Runtime/Engine/Private/DirectionalLightComponent.cpp b/Engine/Source/Runtime/Engine/Private/DirectionalLightComponent.cpp
--- a/Engine/Source/Runtime/Engine/Private/DirectionalLightComponent.cpp
+++ b/Engine/Source/Runtime/Engine/Private/DirectionalLightComponent.cpp
@@ -232,7 +232,9 @@ public:
 		OutInitializer.SubjectBounds = FBoxSphereBounds(FVector::ZeroVector,SubjectBounds.BoxExtent,SubjectBounds.SphereRadius);
 		OutInitializer.WAxis = FVector4(0,0,0,1);
 		OutInitializer.MinLightW = -HALF_WORLD_MAX;
-		OutInitializer.MaxDistanceToCastInLightW = HALF_WORLD_MAX / 8.0f;
+		// Reduce casting distance on a directional light
+		// This is necessary to improve floating point precision in several places, especially when deriving frustum verts from InvReceiverMatrix
+		OutInitializer.MaxDistanceToCastInLightW = HALF_WORLD_MAX / 32.0f;
 		OutInitializer.SplitIndex = SplitIndex;
 		// Last cascade is the ray traced shadow, if enabled
 		OutInitializer.bRayTracedDistanceFieldShadow = ShouldCreateRayTracedCascade(View.GetFeatureLevel()) && SplitIndex == FDirectionalLightSceneProxy::GetNumViewDependentWholeSceneShadows(View) - 1;
@@ -257,7 +259,9 @@ public:
 		OutInitializer.SubjectBounds = FBoxSphereBounds( FVector::ZeroVector, LightPropagationVolumeBounds.GetExtent(), FMath::Sqrt( LpvExtent * LpvExtent * 3.0f ) );
 		OutInitializer.WAxis = FVector4(0,0,0,1);
 		OutInitializer.MinLightW = -HALF_WORLD_MAX;
-		OutInitializer.MaxDistanceToCastInLightW = HALF_WORLD_MAX / 8.0f;
+		// Reduce casting distance on a directional light
+		// This is necessary to improve floating point precision in several places, especially when deriving frustum verts from InvReceiverMatrix
+		OutInitializer.MaxDistanceToCastInLightW = HALF_WORLD_MAX / 32.0f;
 
 		// Compute the RSM bounds
 		{
@@ -301,7 +305,9 @@ public:
 		OutInitializer.SubjectBounds = FBoxSphereBounds(FVector::ZeroVector,SubjectBounds.BoxExtent,SubjectBounds.SphereRadius);
 		OutInitializer.WAxis = FVector4(0,0,0,1);
 		OutInitializer.MinLightW = -HALF_WORLD_MAX;
-		OutInitializer.MaxDistanceToCastInLightW = HALF_WORLD_MAX;
+		// Reduce casting distance on a directional light
+		// This is necessary to improve floating point precision in several places, especially when deriving frustum verts from InvReceiverMatrix
+		OutInitializer.MaxDistanceToCastInLightW = HALF_WORLD_MAX / 32.0f;
 		return true;
 	}
 
diff --git a/Engine/Source/Runtime/Renderer/Private/ShadowRendering.cpp b/Engine/Source/Runtime/Renderer/Private/ShadowRendering.cpp
--- a/Engine/Source/Runtime/Renderer/Private/ShadowRendering.cpp
+++ b/Engine/Source/Runtime/Renderer/Private/ShadowRendering.cpp
@@ -1784,15 +1784,16 @@ void FProjectedShadowInfo::RenderProjection(FRHICommandListImmediate& RHICmdList
 		return;
 	}
 
-	// The shadow transforms and view transforms are relative to different origins, so the world coordinates need to be translated.
-	const FVector4 PreShadowToPreViewTranslation(View->ViewMatrices.PreViewTranslation - PreShadowTranslation,0);
-
 	FVector4 FrustumVertices[8];
-	
+
 	// Calculate whether the camera is inside the shadow frustum, or the near plane is potentially intersecting the frustum.
 	bool bCameraInsideShadowFrustum = true;
+
 	if (!IsWholeSceneDirectionalShadow())
 	{
+		// The shadow transforms and view transforms are relative to different origins, so the world coordinates need to be translated.
+		const FVector PreShadowToPreViewTranslation(View->ViewMatrices.PreViewTranslation - PreShadowTranslation);
+
 		// fill out the frustum vertices (this is only needed in the non-whole scene case)
 		for(uint32 Z = 0;Z < 2;Z++)
 		{
@@ -1802,11 +1803,11 @@ void FProjectedShadowInfo::RenderProjection(FRHICommandListImmediate& RHICmdList
 				{
 					const FVector4 UnprojectedVertex = InvReceiverMatrix.TransformFVector4(
 						FVector4(
-						(X ? -1.0f : 1.0f),
-						(Y ? -1.0f : 1.0f),
-						(Z ?  0.0f : 1.0f),
-						1.0f
-						)
+							(X ? -1.0f : 1.0f),
+							(Y ? -1.0f : 1.0f),
+							(Z ?  0.0f : 1.0f),
+							1.0f
+							)
 						);
 					const FVector ProjectedVertex = UnprojectedVertex / UnprojectedVertex.W + PreShadowToPreViewTranslation;
 					FrustumVertices[GetCubeVertexIndex(X,Y,Z)] = FVector4(ProjectedVertex, 0);
@@ -1824,26 +1825,27 @@ void FProjectedShadowInfo::RenderProjection(FRHICommandListImmediate& RHICmdList
 		const FVector BackBottomRight = FrustumVertices[GetCubeVertexIndex(0,1,0)] - View->ViewMatrices.PreViewTranslation;
 
 		const FPlane Front(FrontTopRight, FrontTopLeft, FrontBottomLeft);
-		const float FrontDistance = Front.PlaneDot(View->ViewMatrices.ViewOrigin) / FVector(Front).Size();
+		const float FrontDistance = Front.PlaneDot(View->ViewMatrices.ViewOrigin);
 
 		const FPlane Right(BackTopRight, FrontTopRight, FrontBottomRight);
-		const float RightDistance = Right.PlaneDot(View->ViewMatrices.ViewOrigin) / FVector(Right).Size();
+		const float RightDistance = Right.PlaneDot(View->ViewMatrices.ViewOrigin);
 
 		const FPlane Back(BackTopLeft, BackTopRight, BackBottomRight);
-		const float BackDistance = Back.PlaneDot(View->ViewMatrices.ViewOrigin) / FVector(Back).Size();
+		const float BackDistance = Back.PlaneDot(View->ViewMatrices.ViewOrigin);
 
 		const FPlane Left(FrontTopLeft, BackTopLeft, BackBottomLeft);
-		const float LeftDistance = Left.PlaneDot(View->ViewMatrices.ViewOrigin) / FVector(Left).Size();
+		const float LeftDistance = Left.PlaneDot(View->ViewMatrices.ViewOrigin);
 
 		const FPlane Top(BackTopRight, BackTopLeft, FrontTopLeft);
-		const float TopDistance = Top.PlaneDot(View->ViewMatrices.ViewOrigin) / FVector(Top).Size();
+		const float TopDistance = Top.PlaneDot(View->ViewMatrices.ViewOrigin);
 
 		const FPlane Bottom(FrontBottomRight, FrontBottomLeft, BackBottomLeft);
-		const float BottomDistance = Bottom.PlaneDot(View->ViewMatrices.ViewOrigin) / FVector(Bottom).Size();
+		const float BottomDistance = Bottom.PlaneDot(View->ViewMatrices.ViewOrigin);
 
 		// Use a distance threshold to treat the case where the near plane is intersecting the frustum as the camera being inside
 		// The near plane handling is not exact since it just needs to be conservative about saying the camera is outside the frustum
 		const float DistanceThreshold = -View->NearClippingDistance * 3.0f;
+
 		bCameraInsideShadowFrustum = 
 			FrontDistance > DistanceThreshold && 
 			RightDistance > DistanceThreshold && 
