Fix crash when using high resolution screenshot capture regions (UE-5288)

diff --git a/Engine/Source/Runtime/Renderer/Private/PostProcess/PostProcessing.cpp b/Engine/Source/Runtime/Renderer/Private/PostProcess/PostProcessing.cpp
--- a/Engine/Source/Runtime/Renderer/Private/PostProcess/PostProcessing.cpp
+++ b/Engine/Source/Runtime/Renderer/Private/PostProcess/PostProcessing.cpp
@@ -685,9 +685,19 @@ static void AddHighResScreenshotMask(FPostprocessContext& Context, FRenderingCom
 	// Draw the capture region if a material was supplied
 	if (Context.View.FinalPostProcessSettings.HighResScreenshotCaptureRegionMaterial)
 	{
-		FRenderingCompositePass* CaptureRegionVisualizationPass = Context.Graph.RegisterPass(new(FMemStack::Get()) FRCPassPostProcessMaterial(Context.View.FinalPostProcessSettings.HighResScreenshotCaptureRegionMaterial));
+		auto Material = Context.View.FinalPostProcessSettings.HighResScreenshotCaptureRegionMaterial;
+
+		FRenderingCompositePass* CaptureRegionVisualizationPass = Context.Graph.RegisterPass(new(FMemStack::Get()) FRCPassPostProcessMaterial(Material));
 		CaptureRegionVisualizationPass->SetInput(ePId_Input0, FRenderingCompositeOutputRef(Context.FinalOutput));
 		Context.FinalOutput = FRenderingCompositeOutputRef(CaptureRegionVisualizationPass);
+
+		auto Proxy = Material->GetRenderProxy(false);
+		const FMaterial* RendererMaterial = Proxy->GetMaterial(Context.View.GetFeatureLevel());
+		if (RendererMaterial->NeedsGBuffer())
+		{
+			// AdjustGBufferRefCount(-1) call is done when the pass gets executed
+			GSceneRenderTargets.AdjustGBufferRefCount(1);
+		}
 	}
 }
 
