fixed UE-5588
Crash occurs enabling Show Visualize Subsurface Scattering (Screen Space)

diff --git a/Engine/Source/Runtime/Renderer/Private/PostProcess/PostProcessing.cpp b/Engine/Source/Runtime/Renderer/Private/PostProcess/PostProcessing.cpp
--- a/Engine/Source/Runtime/Renderer/Private/PostProcess/PostProcessing.cpp
+++ b/Engine/Source/Runtime/Renderer/Private/PostProcess/PostProcessing.cpp
@@ -1263,7 +1263,8 @@ void FPostProcessing::Process(FRHICommandListImmediate& RHICmdList, FViewInfo& V
 		
 		AddPostProcessMaterial(Context, BL_AfterTonemapping, SeparateTranslucency);
 
-		if (View.Family->EngineShowFlags.VisualizeSSS)
+		if (View.Family->EngineShowFlags.VisualizeSSS &&
+			!IsOpenGLPlatform(View.GetShaderPlatform()))				//@todo-rco: Remove this when we fix the cross-compiler
 		{
 			// the setup pass also does visualization, based on EngineShowFlags.VisualizeSSS
 			FRenderingCompositePass* PassVisualize = Context.Graph.RegisterPass(new(FMemStack::Get()) FRCPassPostProcessSubsurfaceVisualize());
