fixed UE-5634 Deferred decal materials do not render in game

diff --git a/Engine/Source/Runtime/Renderer/Private/CompositionLighting/PostProcessDeferredDecals.cpp b/Engine/Source/Runtime/Renderer/Private/CompositionLighting/PostProcessDeferredDecals.cpp
--- a/Engine/Source/Runtime/Renderer/Private/CompositionLighting/PostProcessDeferredDecals.cpp
+++ b/Engine/Source/Runtime/Renderer/Private/CompositionLighting/PostProcessDeferredDecals.cpp
@@ -772,14 +772,6 @@ void FRCPassPostProcessDeferredDecals::Process(FRenderingCompositePassContext& C
 		return;
 	}
 
-	const FPooledRenderTargetDesc* InputDesc = GetInputDesc(ePId_Input0);
-
-	if(!InputDesc)
-	{
-		// input is not hooked up correctly
-		return;
-	}
-
 	TArray<FTransientDecalRenderData, SceneRenderingAllocator> SortedDecals;
 	SortedDecals.Empty(Scene.Decals.Num());
 
