UE4 & 4.6 - Enable subsurface scattering for OpenGL

diff --git a/Engine/Source/Runtime/Renderer/Private/PostProcess/PostProcessSubsurface.cpp b/Engine/Source/Runtime/Renderer/Private/PostProcess/PostProcessSubsurface.cpp
--- a/Engine/Source/Runtime/Renderer/Private/PostProcess/PostProcessSubsurface.cpp
+++ b/Engine/Source/Runtime/Renderer/Private/PostProcess/PostProcessSubsurface.cpp
@@ -37,12 +37,6 @@ static TAutoConsoleVariable<int> CVarSubsurfaceQuality(
 	TEXT(" 1: higher quality as specular is separated before screenspace blurring (Only used if SceneColor has an alpha channel)"),
 	ECVF_Scalability | ECVF_RenderThreadSafe);
 
-static bool SubsurfaceShouldCache(EShaderPlatform Platform)
-{
-	//@todo-rco: Remove this when we fix the cross-compiler
-	return IsFeatureLevelSupported(Platform, ERHIFeatureLevel::SM4) && !IsOpenGLPlatform(Platform);
-}
-
 // -------------------------------------------------------------
 
 float GetSubsurfaceRadiusScale()
@@ -127,7 +121,7 @@ class FPostProcessSubsurfaceVisualizePS : public FGlobalShader
 
 	static bool ShouldCache(EShaderPlatform Platform)
 	{
-		return SubsurfaceShouldCache(Platform);
+		return IsFeatureLevelSupported(Platform, ERHIFeatureLevel::SM4);
 	}
 
 	static void ModifyCompilationEnvironment(EShaderPlatform Platform, FShaderCompilerEnvironment& OutEnvironment)
@@ -336,7 +330,7 @@ class FPostProcessSubsurfaceSetupPS : public FGlobalShader
 
 	static bool ShouldCache(EShaderPlatform Platform)
 	{
-		return SubsurfaceShouldCache(Platform);
+		return IsFeatureLevelSupported(Platform, ERHIFeatureLevel::SM4);
 	}
 
 	static void ModifyCompilationEnvironment(EShaderPlatform Platform, FShaderCompilerEnvironment& OutEnvironment)
@@ -445,7 +439,7 @@ class FPostProcessSubsurfaceExtractSpecularPS : public FGlobalShader
 
 	static bool ShouldCache(EShaderPlatform Platform)
 	{
-		return SubsurfaceShouldCache(Platform);
+		return IsFeatureLevelSupported(Platform, ERHIFeatureLevel::SM4);
 	}
 
 	static void ModifyCompilationEnvironment(EShaderPlatform Platform, FShaderCompilerEnvironment& OutEnvironment)
@@ -697,7 +691,7 @@ class TPostProcessSubsurfacePS : public FGlobalShader
 
 	static bool ShouldCache(EShaderPlatform Platform)
 	{
-		return SubsurfaceShouldCache(Platform);
+		return IsFeatureLevelSupported(Platform, ERHIFeatureLevel::SM4);
 	}
 
 	static void ModifyCompilationEnvironment(EShaderPlatform Platform, FShaderCompilerEnvironment& OutEnvironment)
@@ -905,7 +899,7 @@ class TPostProcessSubsurfaceRecombinePS : public FGlobalShader
 
 	static bool ShouldCache(EShaderPlatform Platform)
 	{
-		return SubsurfaceShouldCache(Platform);
+		return IsFeatureLevelSupported(Platform, ERHIFeatureLevel::SM4);
 	}
 
 	static void ModifyCompilationEnvironment(EShaderPlatform Platform, FShaderCompilerEnvironment& OutEnvironment)
diff --git a/Engine/Source/Runtime/Renderer/Private/PostProcess/PostProcessing.cpp b/Engine/Source/Runtime/Renderer/Private/PostProcess/PostProcessing.cpp
--- a/Engine/Source/Runtime/Renderer/Private/PostProcess/PostProcessing.cpp
+++ b/Engine/Source/Runtime/Renderer/Private/PostProcess/PostProcessing.cpp
@@ -1263,8 +1263,7 @@ void FPostProcessing::Process(FRHICommandListImmediate& RHICmdList, FViewInfo& V
 		
 		AddPostProcessMaterial(Context, BL_AfterTonemapping, SeparateTranslucency);
 
-		if (View.Family->EngineShowFlags.VisualizeSSS &&
-			!IsOpenGLPlatform(View.GetShaderPlatform()))				//@todo-rco: Remove this when we fix the cross-compiler
+		if (View.Family->EngineShowFlags.VisualizeSSS)
 		{
 			// the setup pass also does visualization, based on EngineShowFlags.VisualizeSSS
 			FRenderingCompositePass* PassVisualize = Context.Graph.RegisterPass(new(FMemStack::Get()) FRCPassPostProcessSubsurfaceVisualize());
