- Now require Geometry shader support for using a Volume Texture LUT in PostProcessCombineLUTs.cpp

diff --git a/Engine/Source/Runtime/Renderer/Private/PostProcess/PostProcessCombineLUTs.cpp b/Engine/Source/Runtime/Renderer/Private/PostProcess/PostProcessCombineLUTs.cpp
--- a/Engine/Source/Runtime/Renderer/Private/PostProcess/PostProcessCombineLUTs.cpp
+++ b/Engine/Source/Runtime/Renderer/Private/PostProcess/PostProcessCombineLUTs.cpp
@@ -40,7 +40,7 @@ static TAutoConsoleVariable<float> CVarColorMax(
 static bool UseVolumeTextureLUT(EShaderPlatform Platform) 
 {
 	// @todo Mac OS X: in order to share precompiled shaders between GL 3.3 & GL 4.1 devices we mustn't use volume-texture rendering as it isn't universally supported.
-	return (IsFeatureLevelSupported(Platform,ERHIFeatureLevel::SM4) && GSupportsVolumeTextureRendering && !PLATFORM_MAC);
+	return (IsFeatureLevelSupported(Platform,ERHIFeatureLevel::SM4) && GSupportsVolumeTextureRendering && !PLATFORM_MAC && RHISupportsGeometryShaders(Platform));
 }
 
 // including the neutral one at index 0
