fixed crash/check() when running some map outside in game (SkyLight/ReflectionEnvinronments)

diff --git a/Engine/Source/Runtime/Renderer/Private/PostProcess/SceneRenderTargets.cpp b/Engine/Source/Runtime/Renderer/Private/PostProcess/SceneRenderTargets.cpp
--- a/Engine/Source/Runtime/Renderer/Private/PostProcess/SceneRenderTargets.cpp
+++ b/Engine/Source/Runtime/Renderer/Private/PostProcess/SceneRenderTargets.cpp
@@ -116,18 +116,21 @@ inline const TCHAR* GetSceneColorTargetName(FSceneRenderTargets::EShadingPath Sh
 	return SceneColorNames[(uint32)ShadingPath];
 }
 
-FIntPoint FSceneRenderTargets::GetSceneRenderTargetSize(const FSceneViewFamily& ViewFamily) const
+FIntPoint FSceneRenderTargets::ComputeDesiredSize(const FSceneViewFamily& ViewFamily) const
 {
 	// Don't expose Clamped to the cvar since you need to at least grow to the initial state.
-	enum ESizingMethods { RequestedSize, ScreenRes, Grow, VisibleSizingMethodsCount, Clamped};
+	enum ESizingMethods { RequestedSize, ScreenRes, Grow, VisibleSizingMethodsCount, Clamped };
 	ESizingMethods SceneTargetsSizingMethod = Grow;
 
-	bool bCapture = false;
-	for (int32 ViewIndex = 0; ViewIndex < ViewFamily.Views.Num(); ++ViewIndex)
+	bool bIsSceneCapture = false;
+	bool bIsReflectionCapture = false;
+
+	for (int32 ViewIndex = 0, ViewCount = ViewFamily.Views.Num(); ViewIndex < ViewCount; ++ViewIndex)
 	{
 		const FSceneView* View = ViewFamily.Views[ViewIndex];
 
-		bCapture |= View->bIsSceneCapture || View->bIsReflectionCapture;
+		bIsSceneCapture |= View->bIsSceneCapture;
+		bIsReflectionCapture |= View->bIsReflectionCapture;
 	}
 
 	if(!FPlatformProperties::SupportsWindowedMode())
@@ -146,7 +149,7 @@ FIntPoint FSceneRenderTargets::GetSceneRenderTargetSize(const FSceneViewFamily&
 		SceneTargetsSizingMethod = (ESizingMethods) FMath::Clamp(CVarSceneTargetsResizingMethod.GetValueOnRenderThread(), 0, (int32)VisibleSizingMethodsCount);
 	}
 	
-	if (bCapture)
+	if (bIsSceneCapture)
 	{
 		// In general, we don't want scenecapture to grow our buffers, because depending on the cvar for our game, we may not recover that memory.  This can be changed if necessary.
 		// However, in the editor a user might have a small editor window, but be capturing cubemaps or other dynamic assets for data distribution, 
@@ -166,6 +169,17 @@ FIntPoint FSceneRenderTargets::GetSceneRenderTargetSize(const FSceneViewFamily&
 			}
 		}
 	}
+	else if(bIsReflectionCapture)
+	{
+		// reflection captures are low in resolution usually 128, they should not cause the BufferSize to be larger than needed (larger than ScreenSize)
+		check(ViewFamily.FamilySizeX <= 256);
+		check(ViewFamily.FamilySizeY <= 256);
+
+		// this can cause a reallocation during startup if the reflections are updated before the SceneRenderTarget have been allocated with the proper Screen size.
+		// This results in one log printout and some extra minor cost memory and performance cost.
+		// On XBoxOn it can cause different ESRam allocation (Rare, only if Reflection Capture wasn't cooked)
+		SceneTargetsSizingMethod = Grow;
+	}
 
 	switch (SceneTargetsSizingMethod)
 	{
@@ -196,7 +210,7 @@ void FSceneRenderTargets::Allocate(const FSceneViewFamily& ViewFamily)
 	const auto NewFeatureLevel = ViewFamily.Scene->GetFeatureLevel();
 	CurrentShadingPath = ViewFamily.Scene->ShouldUseDeferredRenderer() ? EShadingPath::Deferred : EShadingPath::Forward;
 
-	FIntPoint DesiredBufferSize = GetSceneRenderTargetSize(ViewFamily);
+	FIntPoint DesiredBufferSize = ComputeDesiredSize(ViewFamily);
 	check(DesiredBufferSize.X > 0 && DesiredBufferSize.Y > 0);
 	QuantizeBufferSize(DesiredBufferSize.X, DesiredBufferSize.Y);
 
diff --git a/Engine/Source/Runtime/Renderer/Private/PostProcess/SceneRenderTargets.h b/Engine/Source/Runtime/Renderer/Private/PostProcess/SceneRenderTargets.h
--- a/Engine/Source/Runtime/Renderer/Private/PostProcess/SceneRenderTargets.h
+++ b/Engine/Source/Runtime/Renderer/Private/PostProcess/SceneRenderTargets.h
@@ -543,7 +543,7 @@ private:
 	void AllocateCommonDepthTargets();
 
 	/** Determine the appropriate render target dimensions. */
-	FIntPoint GetSceneRenderTargetSize(const FSceneViewFamily& ViewFamily) const;
+	FIntPoint ComputeDesiredSize(const FSceneViewFamily& ViewFamily) const;
 
 	void AllocSceneColor();
 
