* fixed UE-2907 LIVE: RENDERING: Render Targets are rendering with stretched pixels in Standalone and Launched Game

diff --git a/Engine/Source/Runtime/Renderer/Private/PostProcess/SceneRenderTargets.cpp b/Engine/Source/Runtime/Renderer/Private/PostProcess/SceneRenderTargets.cpp
--- a/Engine/Source/Runtime/Renderer/Private/PostProcess/SceneRenderTargets.cpp
+++ b/Engine/Source/Runtime/Renderer/Private/PostProcess/SceneRenderTargets.cpp
@@ -41,17 +41,6 @@ static TAutoConsoleVariable<int32> CVarSceneTargetsResizingMethod(
 	ECVF_RenderThreadSafe
 	);
 
-static TAutoConsoleVariable<int32> CVarSceneCaptureResizingMethod(
-	TEXT("r.SceneCaptureResizeMethod"),
-	0,
-	TEXT("Control the scene render target resize method for scene captures:\n")
-	TEXT("(This value is only used in game mode and on windowing platforms.)\n")
-	TEXT("0: All scene capture renders are limited to screen resolution or smaller. (Default - prevents allocation when requested dimensions are too large.)\n")
-	TEXT("1: Allows scene capture targets to expand to encompass the dimensions requested.\n")
-	TEXT("   (large sizes could cause stalling without 'r.SceneRenderTargetResizeMethod 2'. Out of memory issues can occur if size is too large)"),
-	ECVF_RenderThreadSafe
-	);
-
 static TAutoConsoleVariable<int32> CVarOptimizeForUAVPerformance(
 	TEXT("r.OptimizeForUAVPerformance"),
 	0,
@@ -151,23 +140,11 @@ FIntPoint FSceneRenderTargets::ComputeDesiredSize(const FSceneViewFamily& ViewFa
 	
 	if (bIsSceneCapture)
 	{
-		// In general, we don't want scenecapture to grow our buffers, because depending on the cvar for our game, we may not recover that memory.  This can be changed if necessary.
-		// However, in the editor a user might have a small editor window, but be capturing cubemaps or other dynamic assets for data distribution, 
-		// in which case we need to grow for correctness.
-		// We also don't want to reallocate all our buffers for a temporary use case like a capture.  So we just clamp the biggest capture size to the currently available buffers.
-		if (GIsEditor)
-		{
-			SceneTargetsSizingMethod = Grow;
-		}
-		else
-		{
-			SceneTargetsSizingMethod = Clamped;
-			int32 CaptureTargetSizeMethod = CVarSceneCaptureResizingMethod.GetValueOnRenderThread();
-			if (FPlatformProperties::SupportsWindowedMode() && CaptureTargetSizeMethod == 1)
-			{
-				SceneTargetsSizingMethod = Grow;
-			}
-		}
+		// This is a bit dangerous as someone in the middle of the game can requests a high res RenderTarget (Larger than screen) and from that point on the SceneRenderTargets will never get smaller again.
+		// This can cause more pressure in the RenderTargetPool and cause more trashing there. It also could slowdown post processing (border clear / handling).
+		// Preventing that case can cause more trouble (When to shrink? Should we keep the former RenderTargets in the pool?).
+		// If this becomes an issue we suggest to shrink the SceneRenderTargets with some game code call (e.g. each level load).
+		SceneTargetsSizingMethod = Grow;
 	}
 	else if(bIsReflectionCapture)
 	{
