disable check as it seems the designer can create higher res versions.

diff --git a/Engine/Source/Runtime/Renderer/Private/PostProcess/SceneRenderTargets.cpp b/Engine/Source/Runtime/Renderer/Private/PostProcess/SceneRenderTargets.cpp
--- a/Engine/Source/Runtime/Renderer/Private/PostProcess/SceneRenderTargets.cpp
+++ b/Engine/Source/Runtime/Renderer/Private/PostProcess/SceneRenderTargets.cpp
@@ -149,8 +149,8 @@ FIntPoint FSceneRenderTargets::ComputeDesiredSize(const FSceneViewFamily& ViewFa
 	else if(bIsReflectionCapture)
 	{
 		// reflection captures are low in resolution usually 128, they should not cause the BufferSize to be larger than needed (larger than ScreenSize)
-		check(ViewFamily.FamilySizeX <= 256);
-		check(ViewFamily.FamilySizeY <= 256);
+//		check(ViewFamily.FamilySizeX <= 256);
+//		check(ViewFamily.FamilySizeY <= 256);
 
 		// this can cause a reallocation during startup if the reflections are updated before the SceneRenderTarget have been allocated with the proper Screen size.
 		// This results in one log printout and some extra minor cost memory and performance cost.
