Distance Field AO pipeline can also interpolate irradiance
* Shaders compiled for both paths, so no overhead when just computing AO
* Prototype with irradiance being gathered from scene color
* Fixed cache overflow in scenes with high geometry detail causing artifacts
* At least one cache record is always discarded per frame now, fixes artifacts from lowest resolution cache never updating
* Fixed pixels inside a distance field getting max occlusion radius due to signed float -> uint conversion

diff --git a/Engine/Shaders/Common.usf b/Engine/Shaders/Common.usf
--- a/Engine/Shaders/Common.usf
+++ b/Engine/Shaders/Common.usf
@@ -439,7 +439,7 @@ MaterialFloat4 EncodeLightAttenuation(MaterialFloat4 InColor)
 /** return the scene lighting texture */
 MaterialFloat3 CalcSceneColor(MaterialFloat2 ScreenUV)
 {
-	return Texture2DSample(SceneColorTexture, SceneColorTextureSampler,ScreenUV).rgb;
+	return Texture2DSampleLevel(SceneColorTexture, SceneColorTextureSampler, ScreenUV, 0).rgb;
 }
 
 /** return all channels of the scene lighting texture */
diff --git a/Engine/Shaders/DistanceFieldSurfaceCacheLighting.usf b/Engine/Shaders/DistanceFieldSurfaceCacheLighting.usf
--- a/Engine/Shaders/DistanceFieldSurfaceCacheLighting.usf
+++ b/Engine/Shaders/DistanceFieldSurfaceCacheLighting.usf
@@ -8,6 +8,7 @@
 #include "DeferredShadingCommon.usf"
 #include "DistanceFieldLightingShared.usf"
 #include "DistanceFieldAOShared.usf"
+#include "MonteCarlo.usf"
 
 struct FObjectCullVertexOutput
 {
@@ -260,9 +261,14 @@ struct FIrradianceCacheSplatVertexOutput
 	float4 Position : SV_POSITION;
 	nointerpolation float4 PositionRadius : TEXCOORD0;
 	nointerpolation float3 Normal : TEXCOORD1;
+#if FINAL_INTERPOLATION_PASS
 	nointerpolation float3 BentNormal : TEXCOORD2;
+	#if SUPPORT_IRRADIANCE
+		nointerpolation float3 Irradiance : TEXCOORD3;
+	#endif
+#endif
 #if USE_SCREENVECTOR_WORLD_POSITION
-	float3 ScreenVector : TEXCOORD4;
+	float3 ScreenVector : TEXCOORD5;
 #endif
 };
 
@@ -272,6 +278,7 @@ Buffer<float4> IrradianceCachePositionRadius;
 Buffer<float> IrradianceCacheOccluderRadius;
 Buffer<float4> IrradianceCacheNormal;
 Buffer<float4> IrradianceCacheBentNormal;
+Buffer<float4> IrradianceCacheIrradiance;
 float InterpolationRadiusScale;
 float2 NormalizedOffsetToPixelCenter;
 float HackExpand;
@@ -340,7 +347,13 @@ void IrradianceCacheSplatVS(
 
 	Output.PositionRadius = PositionAndRadius; 
 	Output.Normal = RecordNormal;
+
+#if FINAL_INTERPOLATION_PASS
 	Output.BentNormal = DecodeNormalForAO(IrradianceCacheBentNormal.Load(JobIndex).xyz);
+	#if SUPPORT_IRRADIANCE
+		Output.Irradiance = IrradianceCacheIrradiance.Load(JobIndex).xyz;
+	#endif
+#endif
 } 
 
 float InterpolationAngleNormalization; 
@@ -349,10 +362,17 @@ float InvMinCosPointBehindPlane;
 /** Computes surface cache weighting between the current pixel and the record being splatted. */
 void IrradianceCacheSplatPS(
 	FIrradianceCacheSplatVertexOutput Input, 
-	in float4 SVPos : SV_POSITION,
-	out float4 OutColor : COLOR0)
+	in float4 SVPos : SV_POSITION
+	, out float4 OutBentNormal : COLOR0
+#if FINAL_INTERPOLATION_PASS && SUPPORT_IRRADIANCE
+	, out float4 OutIrradiance : COLOR1
+#endif
+	)
 {
-	OutColor = 0;
+	OutBentNormal = 0;
+#if FINAL_INTERPOLATION_PASS && SUPPORT_IRRADIANCE
+	OutIrradiance = 0;
+#endif
 
 	float2 BaseLevelScreenUV = float2((((uint2)SVPos.xy) * DownsampleFactorToBaseLevel + .5f) * BaseLevelTexelSize);
 	float4 DistanceFieldNormalValue = Texture2DSampleLevel(DistanceFieldNormalTexture, DistanceFieldNormalSampler, BaseLevelScreenUV, 0);
@@ -399,21 +419,27 @@ void IrradianceCacheSplatPS(
 		//float Weight = saturate(PrecisionScale * (1 - NormalError));
 
 		float VisualizePlacement = Distance < .1f * Input.PositionRadius.w;
-		OutColor.rgb = Input.BentNormal * Weight;
-		OutColor.a = Weight;
+
+#if FINAL_INTERPOLATION_PASS
+		OutBentNormal.rgb = Input.BentNormal * Weight;
+		#if SUPPORT_IRRADIANCE
+			OutIrradiance.rgb = Input.Irradiance * Weight;
+		#endif
+#endif
+		OutBentNormal.a = Weight;
 	}
 	
 	#define VISUALIZE_SPLAT_OVERDRAW 0
 	#if VISUALIZE_SPLAT_OVERDRAW && FINAL_INTERPOLATION_PASS
-		OutColor.rgb = 0;
-		//OutColor.a = Input.Position.w <= SceneW ? .01f : 0.0f;
-		OutColor.a = .005f + (Weight > .001f ? .01f : 0);
+		OutBentNormal.rgb = 0;
+		//OutBentNormal.a = Input.Position.w <= SceneW ? .01f : 0.0f;
+		OutBentNormal.a = .005f + (Weight > .001f ? .01f : 0);
 	#endif
 
 	#define VISUALIZE_RECORD_POINTS 0
 	#if VISUALIZE_RECORD_POINTS && FINAL_INTERPOLATION_PASS
-		OutColor.rgb = 1;
-		OutColor.a = (DistanceError < .05f) * .1f;
+		OutBentNormal.rgb = 1;
+		OutBentNormal.a = (DistanceError < .05f) * .1f;
 	#endif
 }
 
@@ -712,79 +738,172 @@ void InterpolateCacheCS(
 	RWIrradianceCacheSplat[PixelIndex.xy] = LocalAccumulation;
 }
 
-Texture2D BentNormalAOTexture3;
-SamplerState BentNormalAOSampler3;
+Texture2D BentNormalAOTexture;
+SamplerState BentNormalAOSampler;
+
+Texture2D IrradianceTexture;
+SamplerState IrradianceSampler;
+
+Buffer<float4> DebugBuffer;
 
 /** Normalizes the splatted surface cache values, packs depth in alpha. */
-void AOCombinePS2(
-	in float4 UVAndScreenPos : TEXCOORD0, 
-	out float4 OutColor : SV_Target0)
+void AOCombinePS(
+	in float4 UVAndScreenPos : TEXCOORD0
+	,out float4 OutBentNormal : SV_Target0
+#if SUPPORT_IRRADIANCE
+	,out float4 OutIrradiance : SV_Target1
+#endif
+	)
 {
+#if SUPPORT_IRRADIANCE
+	OutIrradiance = 0;
+#endif
+
 	float4 DistanceFieldNormalValue = Texture2DSampleLevel(DistanceFieldNormalTexture, DistanceFieldNormalSampler, UVAndScreenPos.xy, 0);
 	float SceneDepth = abs(DistanceFieldNormalValue.w);
 
 #define VISUALIZE_ACCUMULATED_WEIGHTS 0
 #if VISUALIZE_ACCUMULATED_WEIGHTS
 
-	float3 BentNormalAO = Texture2DSampleLevel(BentNormalAOTexture3, BentNormalAOSampler3, UVAndScreenPos.xy, 0).aaa;
-	OutColor = float4(BentNormalAO, SceneDepth); 
+	float3 BentNormalAO = Texture2DSampleLevel(BentNormalAOTexture, BentNormalAOSampler, UVAndScreenPos.xy, 0).aaa;
+	OutBentNormal = float4(BentNormalAO, SceneDepth); 
 
 #else
 
-	float4 IrradianceCacheAccumulation = Texture2DSampleLevel(BentNormalAOTexture3, BentNormalAOSampler3, UVAndScreenPos.xy, 0);
+	float4 BentNormalAccumulation = Texture2DSampleLevel(BentNormalAOTexture, BentNormalAOSampler, UVAndScreenPos.xy, 0);
 
 	if (DistanceFieldNormalValue.w < 0)
 	{
 		// Add a low weight completely unshadowed contribution for pixels with no distance field representation
 		float Weight = .001f;
-		IrradianceCacheAccumulation += float4(DistanceFieldNormalValue.xyz * Weight, Weight);
+		BentNormalAccumulation += float4(DistanceFieldNormalValue.xyz * Weight, Weight);
 	}
 
-	if (IrradianceCacheAccumulation.w > 0)
+	if (BentNormalAccumulation.w > 0)
 	{
-		OutColor.rgb = IrradianceCacheAccumulation.xyz / IrradianceCacheAccumulation.w;
-		OutColor.a = SceneDepth;
+		OutBentNormal.rgb = BentNormalAccumulation.xyz / BentNormalAccumulation.w;
+		OutBentNormal.a = SceneDepth;
+
+		#if SUPPORT_IRRADIANCE
+			OutIrradiance.rgb = Texture2DSampleLevel(IrradianceTexture, IrradianceSampler, UVAndScreenPos.xy, 0).xyz / BentNormalAccumulation.w;
+		#endif
 	}
 	else
 	{
-		OutColor.rgb = float3(0, 0, .1f);
+		OutBentNormal.rgb = float3(0, 0, .1f);
 		// Sign bit stores whether texel is valid
-		OutColor.a = -SceneDepth;
+		OutBentNormal.a = -SceneDepth;
 	}
 
-	float BentNormalLength = length(OutColor.rgb);
+	float BentNormalLength = length(OutBentNormal.rgb);
 	// Fade to unoccluded in the distance
 	float FadeAlpha = saturate(.001f * (.95f * AOMaxViewDistance - SceneDepth));
-	OutColor.rgb = normalize(OutColor.rgb) * lerp(1, BentNormalLength, FadeAlpha);
+	OutBentNormal.rgb = normalize(OutBentNormal.rgb) * lerp(1, BentNormalLength, FadeAlpha);
 
 	// Clamp Nan's before they get to the filter
-	OutColor.rgb = clamp(OutColor.rgb, -1, 1);
+	OutBentNormal.rgb = clamp(OutBentNormal.rgb, -1, 1);
 
 	FLATTEN
 	if (SceneDepth > AOMaxViewDistance)
 	{
 		// Mark as valid for the gap filling pass, so we don't get overwritten
-		OutColor.a = abs(OutColor.a);
+		OutBentNormal.a = abs(OutBentNormal.a);
+	}
+
+#endif
+
+	/////////////////////////////////////////////////////////////////////////////////////////////////////////////////////////////////////
+#define VISUALIZE_SCENE_COLOR_GATHER 0
+#define VISUALIZE_DEBUG 0
+#if VISUALIZE_DEBUG && SUPPORT_IRRADIANCE
+
+	float2 ScreenPosition = (UVAndScreenPos.xy - View.ScreenPositionScaleBias.wz) / View.ScreenPositionScaleBias.xy;
+
+	float4 HomogeneousWorldPosition = mul(float4(ScreenPosition * SceneDepth, SceneDepth, 1), View.ScreenToWorld);
+	float3 OpaqueWorldPosition = HomogeneousWorldPosition.xyz / HomogeneousWorldPosition.w;
+	float3 ShadingPosition = DebugBuffer[0].xyz;
+
+	float Distance = length(OpaqueWorldPosition - ShadingPosition);
+	float ScaledDistance = pow(Distance / 20, .01f);
+	OutIrradiance.rgb = 2 * (1 - saturate(ScaledDistance)) * float3(.1f, 1, .1f);
+
+	for (int i = 0; i < 9; i++)
+	{
+		float4 GatheredSceneColorAndVisibility = DebugBuffer[1 + i * 3 + 2];
+
+		for (int j = 0; j < 1; j++)
+		{
+			float4 ClosestPositionAndRadius = DebugBuffer[1 + i * 3 + j];
+			float Distance2 = length(OpaqueWorldPosition - ClosestPositionAndRadius.xyz) / ClosestPositionAndRadius.w;
+			Distance2 = pow(Distance2, .5f);
+
+			if (Distance2 < 1)
+			{
+				//OutIrradiance.rgb += (1 - saturate(Distance2)) * 2 * float3(1, .1f, .1f) * (1 - GatheredSceneColorAndVisibility.w);
+				#if !VISUALIZE_SCENE_COLOR_GATHER
+					//OutIrradiance.rgb += GatheredSceneColorAndVisibility.xyz * (1 - GatheredSceneColorAndVisibility.w);
+				#endif
+			}
+
+			//if (GatheredSceneColorAndVisibility.w < .9f)
+			{
+				OutIrradiance.rgb += (1 - saturate(abs(Distance2 - 1) * 100)) * .4f * float3(1, .1f, .1f);
+			}
+		}
 	}
 
+	#if VISUALIZE_SCENE_COLOR_GATHER
+		
+		float SphereRadius = 50;
+
+		if (Distance < SphereRadius)
+		{
+			float3 ConePositionAtClosestToOccluder = ShadingPosition;
+			float4 ClipPosition = mul(float4(ConePositionAtClosestToOccluder, 1), View.WorldToClip);
+			float2 ScreenUV = ScreenAlignedPosition(ClipPosition).xy;
+			float UVDiskRadius = SphereRadius / ClipPosition.w * View.ViewToClip[0][0] * .5f;
+
+			uint NumSamplesPerCone = 81;
+
+			float3 GatheredColor = 0;
+			float GatheredWeight = .0001f;
+
+			for (uint SampleIndex = 0; SampleIndex < NumSamplesPerCone; SampleIndex++)
+			{
+				float2 Random = Hammersley(SampleIndex, NumSamplesPerCone, uint2(0, 0));
+				float2 UVOffset = UniformSampleDisk(Random);
+				float2 SampleUV = ScreenUV + UVOffset * UVDiskRadius;
+				float3 SampleSceneColor = CalcSceneColor(SampleUV);
+				float SampleSceneDepth = CalcSceneDepth(SampleUV);
+				float Weight = 1 - saturate(abs(SampleSceneDepth - ClipPosition.w) / SphereRadius);
+				GatheredColor += SampleSceneColor * Weight;
+				GatheredWeight += Weight;
+			} 
+		
+			OutIrradiance.rgb = GatheredColor / GatheredWeight;
+		}
+	#endif
 #endif
 }
 
 float2 BentNormalAOTexelSize;
-Texture2D BentNormalAOTexture;
-SamplerState BentNormalAOSampler;
 
 #define HALF_FILL_KERNEL_SIZE 2
 
 /** Fills in texels with no splatted weight from screenspace neighbors. */
 void FillGapsPS(
-	in float4 UVAndScreenPos : TEXCOORD0, 
-	out float4 OutColor : SV_Target0)
+	in float4 UVAndScreenPos : TEXCOORD0
+	,out float4 OutBentNormal : SV_Target0
+#if SUPPORT_IRRADIANCE
+	,out float4 OutIrradiance : SV_Target1
+#endif
+	)
 {
 	float4 CenterValue = Texture2DSampleLevel(BentNormalAOTexture, BentNormalAOSampler, UVAndScreenPos.xy, 0);
 	float SceneDepth = abs(CenterValue.w);
 
 	float4 Accumulation = 0;
+	float3 IrradianceAccumulation = 0;
 
 	for (float y = -HALF_FILL_KERNEL_SIZE; y <= HALF_FILL_KERNEL_SIZE; y++)
 	{
@@ -792,6 +911,7 @@ void FillGapsPS(
 		{
 			float2 UVOffset = BentNormalAOTexelSize * float2(x, y);
 			float4 TextureValue = Texture2DSampleLevel(BentNormalAOTexture, BentNormalAOSampler, UVAndScreenPos.xy + UVOffset, 0);
+
 			float SampleSceneDepth = abs(TextureValue.w);
 			float ValidMask = TextureValue.w > 0;
 
@@ -804,21 +924,37 @@ void FillGapsPS(
 
 			Accumulation.rgb += TextureValue.rgb * Weight;
 			Accumulation.a += Weight;
+
+			#if SUPPORT_IRRADIANCE
+				float3 Irradiance = Texture2DSampleLevel(IrradianceTexture, IrradianceSampler, UVAndScreenPos.xy + UVOffset, 0).xyz;
+				IrradianceAccumulation += Irradiance * Weight;
+			#endif
 		}
 	}
 	
-	OutColor = float4(Accumulation.rgb / max(Accumulation.a, .00001f), CenterValue.w);
+	float InvSafeWeight = 1.0f / max(Accumulation.a, .00001f);
+	OutBentNormal = float4(Accumulation.rgb * InvSafeWeight, CenterValue.w);
+
+#if SUPPORT_IRRADIANCE
+	OutIrradiance = float4(IrradianceAccumulation * InvSafeWeight, 0);
+#endif
 }
 
 float4 CameraMotion[5];
-Texture2D HistoryTexture;
-SamplerState HistorySampler;
+Texture2D BentNormalHistoryTexture;
+SamplerState BentNormalHistorySampler;
+Texture2D IrradianceHistoryTexture;
+SamplerState IrradianceHistorySampler;
 float HistoryWeight;
 
 /** Reprojection the occlusion history, only correct for pixels belonging to static objects (camera motion only). */
 void UpdateHistoryPS(
-	in float4 UVAndScreenPos : TEXCOORD0, 
-	out float4 OutColor : SV_Target0)
+	in float4 UVAndScreenPos : TEXCOORD0
+	,out float4 OutBentNormal : SV_Target0
+#if SUPPORT_IRRADIANCE
+	,out float4 OutIrradiance : SV_Target1
+#endif
+	)
 {
 	float3 xyd;
     xyd.xy = UVAndScreenPos.zw * float2(0.5, -0.5) + 0.5;
@@ -836,7 +972,7 @@ void UpdateHistoryPS(
 	float2 DistanceFieldUVs = UVAndScreenPos.xy - View.ViewRectMin.xy * View.ViewSizeAndSceneTexelSize.zw;
 	float4 NewValue = Texture2DSampleLevel(BentNormalAOTexture, BentNormalAOSampler, DistanceFieldUVs, 0);
 	float2 OldDistanceFieldUVs = OldUV.xy - View.ViewRectMin.xy * View.ViewSizeAndSceneTexelSize.zw;
-	float4 HistoryValue = Texture2DSampleLevel(HistoryTexture, HistorySampler, OldDistanceFieldUVs, 0);
+	float4 HistoryValue = Texture2DSampleLevel(BentNormalHistoryTexture, BentNormalHistorySampler, OldDistanceFieldUVs, 0);
 	
 	//@todo - this doesn't take camera movement into account
 	float DepthDeltaFraction = abs(abs(NewValue.a) - abs(HistoryValue.a)) / abs(NewValue.a);
@@ -852,15 +988,27 @@ void UpdateHistoryPS(
 		EffectiveHistoryWeight = 0;
 	}
 
-	OutColor.rgb = lerp(NewValue.rgb, HistoryValue.rgb, EffectiveHistoryWeight);
+	OutBentNormal.rgb = lerp(NewValue.rgb, HistoryValue.rgb, EffectiveHistoryWeight);
+
+#if SUPPORT_IRRADIANCE
+	float3 NewIrradiance = Texture2DSampleLevel(IrradianceTexture, IrradianceSampler, DistanceFieldUVs, 0).xyz;
+	float3 HistoryIrradianceValue = Texture2DSampleLevel(IrradianceHistoryTexture, IrradianceHistorySampler, OldDistanceFieldUVs, 0).xyz;
 
-	OutColor.rgb = isnan(OutColor.rgb) ? 0 : OutColor.rgb;
+	OutIrradiance = float4(lerp(NewIrradiance, HistoryIrradianceValue, EffectiveHistoryWeight), 0);
+#endif
 
-	OutColor.a = NewValue.a;
+	OutBentNormal.rgb = isnan(OutBentNormal.rgb) ? 0 : OutBentNormal.rgb;
+	OutBentNormal.a = NewValue.a;
 }
 
 /** Upsamples the AO results to full resolution using a bilateral filter. */
-void AOUpsamplePS(in float4 UVAndScreenPos : TEXCOORD0, out float4 OutColor : SV_Target0)
+void AOUpsamplePS(
+	in float4 UVAndScreenPos : TEXCOORD0
+	, out float4 OutBentNormal : SV_Target0
+#if OUTPUT_BENT_NORMAL && SUPPORT_IRRADIANCE
+	, out float4 OutIrradiance : SV_Target1
+#endif
+	)
 {
 	// Distance field AO was computed at 0,0 regardless of viewrect min
 	float2 DistanceFieldUVs = UVAndScreenPos.xy - View.ViewRectMin.xy * View.ViewSizeAndSceneTexelSize.zw;
@@ -890,23 +1038,43 @@ void AOUpsamplePS(in float4 UVAndScreenPos : TEXCOORD0, out float4 OutColor : SV
 	float4 DepthWeights = 1.0f / (abs(CornerDepths - SceneDepth.xxxx) + Epsilon);
 	float4 FinalWeights = CornerWeights * DepthWeights;
 
+	float InvWeight = 1.0f / dot(FinalWeights, 1);
+
 	float3 InterpolatedResult = 
 		(FinalWeights.x * TextureValues00.xyz 
 		+ FinalWeights.y * TextureValues10.xyz
 		+ FinalWeights.z * TextureValues01.xyz 
 		+ FinalWeights.w * TextureValues11.xyz)
-		/ dot(FinalWeights, 1);
+		* InvWeight;
 
 	float3 BentNormal = InterpolatedResult.xyz;
 
+	#if SUPPORT_IRRADIANCE
+		float4 IrradianceValues00 = Texture2DSampleLevel(IrradianceTexture, IrradianceSampler, Corner00UV, 0);
+		float4 IrradianceValues10 = Texture2DSampleLevel(IrradianceTexture, IrradianceSampler, Corner00UV + float2(LowResTexelSize.x, 0), 0);
+		float4 IrradianceValues01 = Texture2DSampleLevel(IrradianceTexture, IrradianceSampler, Corner00UV + float2(0, LowResTexelSize.y), 0);
+		float4 IrradianceValues11 = Texture2DSampleLevel(IrradianceTexture, IrradianceSampler, Corner00UV + LowResTexelSize, 0);
+
+		float3 Irradiance = 
+			(FinalWeights.x * IrradianceValues00.xyz 
+			+ FinalWeights.y * IrradianceValues10.xyz
+			+ FinalWeights.z * IrradianceValues01.xyz 
+			+ FinalWeights.w * IrradianceValues11.xyz)
+			* InvWeight;
+	#endif
+
 #else
 	float3 BentNormal = Texture2DSampleLevel(BentNormalAOTexture, BentNormalAOSampler, DistanceFieldUVs, 0).xyz;
+	float3 Irradiance = Texture2DSampleLevel(IrradianceTexture, IrradianceSampler, DistanceFieldUVs, 0).xyz;
 #endif
 
 #if OUTPUT_BENT_NORMAL
-	OutColor = float4(BentNormal, 1);
+	OutBentNormal = float4(BentNormal, 1);
+	#if SUPPORT_IRRADIANCE
+		OutIrradiance = float4(Irradiance, 1);
+	#endif
 #else
-	OutColor = float4(length(BentNormal).xxx, 1);
+	OutBentNormal = float4(length(BentNormal).xxx, 1);
 #endif
 }
 
diff --git a/Engine/Shaders/DistanceFieldSurfaceCacheLightingCompute.usf b/Engine/Shaders/DistanceFieldSurfaceCacheLightingCompute.usf
--- a/Engine/Shaders/DistanceFieldSurfaceCacheLightingCompute.usf
+++ b/Engine/Shaders/DistanceFieldSurfaceCacheLightingCompute.usf
@@ -9,6 +9,7 @@
 #include "ReflectionEnvironmentShared.usf"
 #include "DistanceFieldLightingShared.usf"
 #include "DistanceFieldAOShared.usf"
+#include "MonteCarlo.usf"
 
 /** Min and Max depth for this tile. */
 groupshared uint IntegerTileMinZ;
@@ -461,26 +462,25 @@ Buffer<uint> SavedStartIndex;
 
 RWBuffer<float> RWOccluderRadius;
 RWBuffer<float4> RWIrradianceCacheBentNormal;
-
-// Whether to use cone tracing or cone stepping
-#define USE_CONE_TRACING 1
+RWBuffer<float4> RWIrradianceCacheIrradiance;
+RWBuffer<float4> RWDebugBuffer;
 
 // Have to disable surface caching dependencies for this to work
 #define VISUALIZE_ONE_CONE 0
 
 #define FINAL_GATHER_THREADGROUP_SIZE 64
 
-#if USE_CONE_TRACING
-	// Enforce one thread per cone direction
-	#define SIMULTANEOUSLY_TRACED_OBJECTS (FINAL_GATHER_THREADGROUP_SIZE / NUM_CONE_DIRECTIONS)
-#else
-	#define SIMULTANEOUSLY_TRACED_OBJECTS 4
-#endif
+// Enforce one thread per cone direction
+#define SIMULTANEOUSLY_TRACED_OBJECTS (FINAL_GATHER_THREADGROUP_SIZE / NUM_CONE_DIRECTIONS)
+
 #define THREADS_PER_OBJECT (FINAL_GATHER_THREADGROUP_SIZE / SIMULTANEOUSLY_TRACED_OBJECTS)
 
 groupshared uint SharedConeVisibility[NUM_CONE_DIRECTIONS];
 groupshared uint SharedMinOcclusionDistance;
 
+groupshared float4 SharedConeDistanceAtClosestToOccluder[FINAL_GATHER_THREADGROUP_SIZE];
+groupshared float3 SharedGatheredLighting[FINAL_GATHER_THREADGROUP_SIZE];
+
 /** Computes ambient occlusion for a surface cache record by cone stepping through the nearby object distance fields. */
 [numthreads(FINAL_GATHER_THREADGROUP_SIZE, 1, 1)]
 void FinalGatherCS(
@@ -524,7 +524,6 @@ void FinalGatherCS(
 
 		float3 DebugConeDirection = normalize(float3(.3f, .5f, .4f));
 
-#if USE_CONE_TRACING
 		uint ConeIndex = ThreadIndex % THREADS_PER_OBJECT;
 
 		float3 ConeDirection = AOSamples2.SampleDirections[ConeIndex].xyz;
@@ -535,9 +534,12 @@ void FinalGatherCS(
 #endif
 
 		float MinVisibility = 1;
+		float SecondMinVisibility = 1;
 		float MinWorldDistanceToOccluder = AOMaxDistance;
+		float ConeDistanceAtClosestToOccluder = 0;
+		float ConeDistanceAtSecondClosestToOccluder = 0;
 		float MaxWorldStepOffset = GetStepOffset(NUM_CONE_STEPS);
-#endif
+
 		LOOP 
 		for (uint ListObjectIndex = 0; ListObjectIndex < NumObjectsAffectingTile; ListObjectIndex += SIMULTANEOUSLY_TRACED_OBJECTS)
 		{
@@ -594,7 +596,6 @@ void FinalGatherCS(
 						StartStepIndex = 5;
 					}
 
-#if USE_CONE_TRACING
 					float WorldStepOffset = GetStepOffset(StartStepIndex);
 
 					LOOP
@@ -624,7 +625,20 @@ void FinalGatherCS(
 							Visibility = max(Visibility, saturate((OccluderDistanceFraction - 0) / 4));
 							Visibility = max(Visibility, saturate((OccluderDistanceFraction - .7f) / .4f));
 
-							MinVisibility = min(MinVisibility, Visibility);
+							if (Visibility < MinVisibility)
+							{
+								ConeDistanceAtSecondClosestToOccluder = ConeDistanceAtClosestToOccluder;
+								ConeDistanceAtClosestToOccluder = WorldStepOffset;
+								SecondMinVisibility = MinVisibility;
+								MinVisibility = Visibility;
+							}
+							else if (Visibility < SecondMinVisibility)
+							{
+								ConeDistanceAtSecondClosestToOccluder = WorldStepOffset;
+								SecondMinVisibility = Visibility;
+							}
+
+							//MinVisibility = min(MinVisibility, Visibility);
 
 							if (DistanceToOccluder < .9f * SphereRadius)
 							{
@@ -637,79 +651,86 @@ void FinalGatherCS(
 						float MinStepSize = .6f * (GetStepOffset(StepIndex + 1) - GetStepOffset(StepIndex));
 						WorldStepOffset += max(DistanceToOccluder, MinStepSize);
 					}
+				}
+			}
+		}
 
-#else // Cone stepping
-
-					uint NumConeSteps = NUM_CONE_STEPS - StartStepIndex;
-
-					LOOP
-					for (uint SampleIndex = ThreadIndex % THREADS_PER_OBJECT; SampleIndex < NUM_CONE_DIRECTIONS * NumConeSteps; SampleIndex += THREADS_PER_OBJECT)
-					{
-						// Avoid uint divide
-						uint ConeIndex = (uint)(SampleIndex / (float)NumConeSteps);
-						uint StepIndex = StartStepIndex + SampleIndex - ConeIndex * NumConeSteps;
+		InterlockedMin(SharedConeVisibility[ConeIndex], asuint(MinVisibility));
+		InterlockedMin(SharedMinOcclusionDistance, asuint(max(MinWorldDistanceToOccluder, 0)));
 
-						float3 ConeDirection = AOSamples2.SampleDirections[ConeIndex].xyz;
-						float3 RotatedConeDirection = ConeDirection.x * TangentX + ConeDirection.y * TangentY + ConeDirection.z * WorldNormal;
+		SharedConeDistanceAtClosestToOccluder[ThreadIndex] = float4(MinVisibility, ConeDistanceAtClosestToOccluder, SecondMinVisibility, ConeDistanceAtSecondClosestToOccluder);
 
-						#if VISUALIZE_ONE_CONE
-							RotatedConeDirection = DebugConeDirection;
-						#endif
+		GroupMemoryBarrierWithGroupSync();
 
-						float WorldStepOffset = AOStepScale * exp2(AOStepExponentScale * StepIndex);
-						float3 WorldSamplePosition = WorldShadingPosition + RotatedConeDirection * WorldStepOffset;
-						float3 StepSamplePosition = mul(float4(WorldSamplePosition, 1), WorldToVolume).xyz;
-						float3 ClampedSamplePosition = clamp(StepSamplePosition, -LocalPositionExtent, LocalPositionExtent);
-						float DistanceToClamped = length(StepSamplePosition - ClampedSamplePosition);
+#if SUPPORT_IRRADIANCE
+		float ConeVisibility = asfloat(SharedConeVisibility[ConeIndex]);
 
-						float3 StepVolumeUV = DistanceFieldVolumePositionToUV(ClampedSamplePosition, UVScaleAndVolumeScale.xyz, UVAdd);
-						float DistanceToOccluder = (Texture3DSampleLevel(DistanceFieldTexture, DistanceFieldSampler, StepVolumeUV, 0).x + DistanceToClamped) * UVScaleAndVolumeScale.w;
+		float CompactedConeDistanceAtClosestToOccluder = 0;
+		float SecondConeMinVisibility = 1;
 
-						float SphereRadius = WorldStepOffset * TanConeHalfAngle;
-						float ShadingSphereRadius = SphereRadius * 1.0f;
+		for (uint CompactThreadIndex = ConeIndex; CompactThreadIndex < FINAL_GATHER_THREADGROUP_SIZE; CompactThreadIndex += THREADS_PER_OBJECT)
+		{
+			if (SharedConeDistanceAtClosestToOccluder[CompactThreadIndex].x == ConeVisibility)
+			{
+				CompactedConeDistanceAtClosestToOccluder = SharedConeDistanceAtClosestToOccluder[CompactThreadIndex].y;
+			}
+			else
+			{
+				SecondConeMinVisibility = min(SecondConeMinVisibility, min(SharedConeDistanceAtClosestToOccluder[CompactThreadIndex].x, SharedConeDistanceAtClosestToOccluder[CompactThreadIndex].z));
+			}
+		}
 
-						#define CIRCLE_LINE_INTERSECTION 0
+		float CompactedConeDistanceAtSecondClosestToOccluder = 0;
 
-						#if CIRCLE_LINE_INTERSECTION
+		for (uint CompactThreadIndex = ConeIndex; CompactThreadIndex < FINAL_GATHER_THREADGROUP_SIZE; CompactThreadIndex += THREADS_PER_OBJECT)
+		{
+			if (SharedConeDistanceAtClosestToOccluder[CompactThreadIndex].x == SecondConeMinVisibility)
+			{
+				CompactedConeDistanceAtSecondClosestToOccluder = SharedConeDistanceAtClosestToOccluder[CompactThreadIndex].y;
+			}
+			else if (SharedConeDistanceAtClosestToOccluder[CompactThreadIndex].z == SecondConeMinVisibility)
+			{
+				CompactedConeDistanceAtSecondClosestToOccluder = SharedConeDistanceAtClosestToOccluder[CompactThreadIndex].w;
+			}
+		}
 
-							float Visibility = 1;
+		float SphereRadius = CompactedConeDistanceAtClosestToOccluder * TanConeHalfAngle;
+		float3 ConePositionAtClosestToOccluder = WorldShadingPosition + RotatedConeDirection * CompactedConeDistanceAtClosestToOccluder;
+		float4 ClipPosition = mul(float4(ConePositionAtClosestToOccluder, 1), View.WorldToClip);
+		float2 ScreenUV = ScreenAlignedPosition(ClipPosition).xy;
+		float UVDiskRadius = SphereRadius / ClipPosition.w * View.ViewToClip[0][0] * .5f;
 
-							BRANCH
-							if (DistanceToOccluder < SphereRadius)
-							{
-								float SphereRadiusSquared = SphereRadius * SphereRadius;
-								// Area of the circle cut off by an intersecting line segment of distance DistanceToOccluder to the circle center
-								float CircleSegmentArea = SphereRadiusSquared * acos(DistanceToOccluder / SphereRadius) - DistanceToOccluder * sqrt(SphereRadiusSquared - DistanceToOccluder * DistanceToOccluder);
-								Visibility = 1 - saturate(CircleSegmentArea / (.5f * PI * SphereRadiusSquared));
-							}
+		uint NumSamplesPerThread = 9;
+		uint NumSamplesPerCone = NumSamplesPerThread * (FINAL_GATHER_THREADGROUP_SIZE / NUM_CONE_DIRECTIONS);
+		uint BaseSampleIndex = (ThreadIndex / NUM_CONE_DIRECTIONS) * NumSamplesPerThread;
 
-						#else
-							// Allows use of negative distances in the small steps to create full penumbra range, reduces over-occlusion
-							float MinOcclusionFraction = 0;//-.6f * (1 - saturate(4 * StepIndex / (float)NUM_CONE_STEPS));
-							// Derive visibility from 1d intersection
-							float Visibility = saturate((DistanceToOccluder / ShadingSphereRadius - MinOcclusionFraction) / (1 - MinOcclusionFraction));
-						#endif
+		float3 GatheredColor = 0;
+		float GatheredWeight = .0001f;
 
-						// Don't allow small objects to fully occlude a cone step
-						Visibility = max(Visibility, 1 - saturate(ObjectOccluderRadius / SphereRadius));
+		for (uint SampleIndex = 0; SampleIndex < NumSamplesPerThread; SampleIndex++)
+		{
+			float2 Random = Hammersley(SampleIndex + BaseSampleIndex, NumSamplesPerCone, uint2(0, 0));
+			float2 UVOffset = UniformSampleDisk(Random);
+			float2 SampleUV = ScreenUV + UVOffset * UVDiskRadius;
+			float3 SampleSceneColor = CalcSceneColor(SampleUV);
+			float SampleSceneDepth = CalcSceneDepth(SampleUV);
+			float Weight = 1 - saturate(abs(SampleSceneDepth - ClipPosition.w) / SphereRadius);
+			GatheredColor += SampleSceneColor * Weight;
+			GatheredWeight += Weight;
+		} 
+		
+		SharedGatheredLighting[ThreadIndex] = GatheredColor / GatheredWeight;
 
-						InterlockedMin(SharedConeVisibility[ConeIndex], asuint(Visibility));
+		if (ThreadIndex < THREADS_PER_OBJECT)
+		{
+			uint DebugIndex = ThreadIndex % THREADS_PER_OBJECT;
+			RWDebugBuffer[0] = float4(WorldShadingPosition, 0);
+			RWDebugBuffer[1 + DebugIndex * 3 + 0] = float4(ConePositionAtClosestToOccluder, SphereRadius);
 
-						if (DistanceToOccluder < .9f * SphereRadius)
-						{
-							// Assuming occluder is straight forward along the cone
-							float WorldDistanceToOccluder = WorldStepOffset + DistanceToOccluder;
-							InterlockedMin(SharedMinOcclusionDistance, asuint(WorldDistanceToOccluder));
-						}
-					}
-#endif
-				}
-			}
+			float SphereRadius2 = CompactedConeDistanceAtSecondClosestToOccluder * TanConeHalfAngle;
+			float3 ConePositionAtClosestToOccluder2 = WorldShadingPosition + RotatedConeDirection * CompactedConeDistanceAtSecondClosestToOccluder;
+			RWDebugBuffer[1 + DebugIndex * 3 + 1] = float4(ConePositionAtClosestToOccluder2, SphereRadius2);
 		}
-
-#if USE_CONE_TRACING
-		InterlockedMin(SharedConeVisibility[ConeIndex], asuint(MinVisibility));
-		InterlockedMin(SharedMinOcclusionDistance, asuint(MinWorldDistanceToOccluder));
 #endif
 	}
 
@@ -719,7 +740,8 @@ void FinalGatherCS(
 	{
 		float3 UnoccludedDirection = 0;
 		float Visibility = 0;
-		float NumSamplesToDisplay = VISUALIZE_ONE_CONE ? 1 : NUM_CONE_DIRECTIONS;
+		uint NumSamplesToDisplay = VISUALIZE_ONE_CONE ? 1 : NUM_CONE_DIRECTIONS;
+		float3 Irradiance = 0;
 
 		for (uint ConeIndex = 0; ConeIndex < NumSamplesToDisplay; ConeIndex++)
 		{
@@ -729,13 +751,34 @@ void FinalGatherCS(
 			float ConeVisibility = asfloat(SharedConeVisibility[ConeIndex]);
 			UnoccludedDirection += ConeVisibility * RotatedConeDirection;
 			Visibility += ConeVisibility;
+
+#if SUPPORT_IRRADIANCE
+			float3 AccumulatedIrradiance = 0;			
+
+			for (uint CompactThreadIndex = ConeIndex; CompactThreadIndex < FINAL_GATHER_THREADGROUP_SIZE; CompactThreadIndex += THREADS_PER_OBJECT)
+			{
+				AccumulatedIrradiance += SharedGatheredLighting[CompactThreadIndex];
+			}
+
+			AccumulatedIrradiance /= (float)SIMULTANEOUSLY_TRACED_OBJECTS;
+
+			RWDebugBuffer[1 + ConeIndex * 3 + 2] = float4(AccumulatedIrradiance, ConeVisibility);
+
+			Irradiance += AccumulatedIrradiance * (1 - ConeVisibility);
+#endif
 		}
 
-		UnoccludedDirection = UnoccludedDirection * BentNormalNormalizeFactor / (float)NumSamplesToDisplay;
-		Visibility = Visibility / (float)NumSamplesToDisplay;
+		float InvNumSamples = 1.0f  / (float)NumSamplesToDisplay;
+		UnoccludedDirection = UnoccludedDirection * (BentNormalNormalizeFactor * InvNumSamples);
+		Visibility = Visibility * InvNumSamples;
+		Irradiance = Irradiance * InvNumSamples;
 
 		RWOccluderRadius[RecordIndex] = RecordRadiusScale * asfloat(SharedMinOcclusionDistance);
 		RWIrradianceCacheBentNormal[RecordIndex] = float4(UnoccludedDirection, 0);
+
+#if SUPPORT_IRRADIANCE
+		RWIrradianceCacheIrradiance[RecordIndex] = float4(Irradiance, 0);
+#endif
 	}
 }
 
@@ -781,7 +824,7 @@ void SetupCopyIndirectArgumentsCS(
 {
 	uint NumRecords = (1 - TrimFraction) * DrawParameters[1];
 
-	// Make sure we spawn at least one group so RWScatterDrawParameters gets written to
+	// Make sure we spawn at least one group so RWDispatchParameters gets written to
 	RWDispatchParameters[0] = max((NumRecords + COPY_THREADGROUP_SIZE - 1) / COPY_THREADGROUP_SIZE, 1);
 	RWDispatchParameters[1] = 1;
 	RWDispatchParameters[2] = 1;
@@ -791,10 +834,12 @@ RWBuffer<float4> RWCopyIrradianceCachePositionRadius;
 RWBuffer<float4> RWCopyIrradianceCacheNormal;
 RWBuffer<float> RWCopyOccluderRadius;
 RWBuffer<float4> RWCopyIrradianceCacheBentNormal;
+RWBuffer<float4> RWCopyIrradianceCacheIrradiance;
 RWBuffer<uint2> RWCopyIrradianceCacheTileCoordinate;
 
 Buffer<float> OccluderRadius;
 Buffer<float4> IrradianceCacheBentNormal;
+Buffer<float4> IrradianceCacheIrradiance;
 
 [numthreads(COPY_THREADGROUP_SIZE, 1, 1)]
 void CopyIrradianceCacheSamplesCS(
@@ -803,7 +848,9 @@ void CopyIrradianceCacheSamplesCS(
     uint3 GroupThreadId : SV_GroupThreadID) 
 {
 	uint NumRecords = DrawParameters[1];
-	uint StartIndex = NumRecords * TrimFraction;
+	// Always trim at least one to handle dynamic scene changes not accounted for with irradiance cache misses
+	uint StartIndex = max(NumRecords * TrimFraction, 1);
+	StartIndex = min(StartIndex, NumRecords);
 	uint SourceIndex = StartIndex + DispatchThreadId.x;
 	uint DestIndex = DispatchThreadId.x;
 
@@ -813,6 +860,9 @@ void CopyIrradianceCacheSamplesCS(
 		RWCopyIrradianceCacheNormal[DestIndex] = IrradianceCacheNormal[SourceIndex];
 		RWCopyOccluderRadius[DestIndex] = OccluderRadius[SourceIndex];
 		RWCopyIrradianceCacheBentNormal[DestIndex] = IrradianceCacheBentNormal[SourceIndex];
+#if SUPPORT_IRRADIANCE
+		RWCopyIrradianceCacheIrradiance[DestIndex] = IrradianceCacheIrradiance[SourceIndex];
+#endif
 		RWCopyIrradianceCacheTileCoordinate[DestIndex] = IrradianceCacheTileCoordinate[SourceIndex];
 	}
 
diff --git a/Engine/Shaders/MonteCarlo.usf b/Engine/Shaders/MonteCarlo.usf
--- a/Engine/Shaders/MonteCarlo.usf
+++ b/Engine/Shaders/MonteCarlo.usf
@@ -65,6 +65,13 @@ float4 UniformSampleHemisphere( float2 E )
 	return float4( H, PDF );
 }
 
+float2 UniformSampleDisk(float2 Random)
+{
+	const float Theta = 2.0f * (float)PI * Random.x;
+	const float Radius = sqrt(Random.y);
+	return float2(Radius * cos(Theta), Radius * sin(Theta));
+}
+
 float4 CosineSampleHemisphere( float2 E )
 {
 	float Phi = 2 * PI * E.x;
diff --git a/Engine/Shaders/SkyLighting.usf b/Engine/Shaders/SkyLighting.usf
--- a/Engine/Shaders/SkyLighting.usf
+++ b/Engine/Shaders/SkyLighting.usf
@@ -11,6 +11,9 @@
 Texture2D BentNormalAOTexture;
 SamplerState BentNormalAOSampler;
 
+Texture2D IrradianceTexture;
+SamplerState IrradianceSampler;
+
 float3 ContrastAndNormalizeMulAdd;
 float4 OcclusionTintAndMinOcclusion;
 
@@ -32,7 +35,8 @@ void SkyLightDiffusePS(
 		float SkyVisibility = 1;
 		float DotProductFactor = 1;
 		float3 SkyLightingNormal = InGBufferData.WorldNormal;
-		
+		float3 DiffuseIrradiance = 0;
+
 #if APPLY_SHADOWING
 		float3 BentNormal = Texture2DSampleLevel(BentNormalAOTexture, BentNormalAOSampler, UV, 0).xyz;
 
@@ -65,6 +69,9 @@ void SkyLightDiffusePS(
 		float ContrastCurve = 1 / (1 + exp(-ContrastAndNormalizeMulAdd.x * (SkyVisibility * 10 - 5)));
 		SkyVisibility = saturate(ContrastCurve * ContrastAndNormalizeMulAdd.y + ContrastAndNormalizeMulAdd.z);
 		
+		#if SUPPORT_IRRADIANCE
+			DiffuseIrradiance = Texture2DSampleLevel(IrradianceTexture, IrradianceSampler, UV, 0).xyz * View.GeneralPurposeTweak;
+		#endif
 #endif
 
 		float ScalarFactors = lerp(SkyVisibility, 1, OcclusionTintAndMinOcclusion.w) * InGBufferData.GBufferAO * ScreenSpaceData.AmbientOcclusion;
@@ -78,11 +85,17 @@ void SkyLightDiffusePS(
 
 		// Compute the preconvolved incoming lighting with the bent normal direction
 		float3 DiffuseLookup = GetSkySHDiffuse(SkyLightingNormal) * View.SkyLightColor.rgb;
-
+		
 		// Apply AO to the sky diffuse and account for darkening due to the geometry term
 		Lighting += (ScalarFactors * DotProductFactor) * DiffuseLookup * InGBufferData.DiffuseColor + (1 - SkyVisibility) * OcclusionTintAndMinOcclusion.xyz;
 
+		Lighting += DiffuseIrradiance * InGBufferData.DiffuseColor;
+
 		// For debugging
+		#if SUPPORT_IRRADIANCE
+			//Lighting = max(Texture2DSampleLevel(IrradianceTexture, IrradianceSampler, UV, 0).xyz, 0) + .01f * GetSkySHDiffuse(InGBufferData.WorldNormal);
+		#endif
+
 		//Lighting = (Texture2DSampleLevel(BentNormalAOTexture, BentNormalAOSampler, UV, 0).xyz);
 	}
 
diff --git a/Engine/Source/Runtime/Renderer/Private/DeferredShadingRenderer.cpp b/Engine/Source/Runtime/Renderer/Private/DeferredShadingRenderer.cpp
--- a/Engine/Source/Runtime/Renderer/Private/DeferredShadingRenderer.cpp
+++ b/Engine/Source/Runtime/Renderer/Private/DeferredShadingRenderer.cpp
@@ -1016,7 +1016,7 @@ void FDeferredShadingSceneRenderer::Render(FRHICommandListImmediate& RHICmdList)
 	if (ViewFamily.EngineShowFlags.VisualizeDistanceFieldAO)
 	{
 		TRefCountPtr<IPooledRenderTarget> DummyOutput;
-		RenderDistanceFieldAOSurfaceCache(RHICmdList, FDistanceFieldAOParameters(), DummyOutput, true);
+		RenderDistanceFieldAOSurfaceCache(RHICmdList, FDistanceFieldAOParameters(), DummyOutput, DummyOutput, true);
 	}
 
 	if (ViewFamily.EngineShowFlags.VisualizeMeshDistanceFields)
diff --git a/Engine/Source/Runtime/Renderer/Private/DeferredShadingRenderer.h b/Engine/Source/Runtime/Renderer/Private/DeferredShadingRenderer.h
--- a/Engine/Source/Runtime/Renderer/Private/DeferredShadingRenderer.h
+++ b/Engine/Source/Runtime/Renderer/Private/DeferredShadingRenderer.h
@@ -195,7 +195,12 @@ private:
 	void RenderDynamicSkyLighting(FRHICommandListImmediate& RHICmdList, TRefCountPtr<IPooledRenderTarget>& DynamicBentNormalAO);
 
 	/** Render Ambient Occlusion using mesh distance fields and the surface cache, which supports dynamic rigid meshes. */
-	bool RenderDistanceFieldAOSurfaceCache(FRHICommandListImmediate& RHICmdList, const class FDistanceFieldAOParameters& Parameters, TRefCountPtr<IPooledRenderTarget>& OutDynamicBentNormalAO, bool bApplyToSceneColor);
+	bool RenderDistanceFieldAOSurfaceCache(
+		FRHICommandListImmediate& RHICmdList, 
+		const class FDistanceFieldAOParameters& Parameters, 
+		TRefCountPtr<IPooledRenderTarget>& OutDynamicBentNormalAO, 
+		TRefCountPtr<IPooledRenderTarget>& OutDynamicIrradiance,
+		bool bApplyToSceneColor);
 
 	void RenderMeshDistanceFieldVisualization(FRHICommandListImmediate& RHICmdList, const FDistanceFieldAOParameters& Parameters);
 
diff --git a/Engine/Source/Runtime/Renderer/Private/DistanceFieldSurfaceCacheLighting.cpp b/Engine/Source/Runtime/Renderer/Private/DistanceFieldSurfaceCacheLighting.cpp
--- a/Engine/Source/Runtime/Renderer/Private/DistanceFieldSurfaceCacheLighting.cpp
+++ b/Engine/Source/Runtime/Renderer/Private/DistanceFieldSurfaceCacheLighting.cpp
@@ -25,6 +25,14 @@ FAutoConsoleVariableRef CVarDistanceFieldAO(
 	ECVF_Cheat | ECVF_RenderThreadSafe
 	);
 
+int32 GDistanceFieldGI = 0;
+FAutoConsoleVariableRef CVarDistanceFieldGI(
+	TEXT("r.DistanceFieldGI"),
+	GDistanceFieldGI,
+	TEXT(""),
+	ECVF_Cheat | ECVF_RenderThreadSafe
+	);
+
 int32 GAOPowerOfTwoBetweenLevels = 2;
 FAutoConsoleVariableRef CVarAOPowerOfTwoBetweenLevels(
 	TEXT("r.AOPowerOfTwoBetweenLevels"),
@@ -211,6 +219,14 @@ FAutoConsoleVariableRef CVarAOInterpolationDepthTesting(
 	ECVF_Cheat | ECVF_RenderThreadSafe
 	);
 
+int32 GAOOverwriteSceneColor = 0;
+FAutoConsoleVariableRef CVarAOOverwriteSceneColor(
+	TEXT("r.AOOverwriteSceneColor"),
+	GAOOverwriteSceneColor,
+	TEXT(""),
+	ECVF_Cheat | ECVF_RenderThreadSafe
+	);
+
 IMPLEMENT_UNIFORM_BUFFER_STRUCT(FAOSampleData2,TEXT("AOSamples2"));
 
 FIntPoint GetBufferSizeForAO()
@@ -1077,10 +1093,10 @@ private:
 
 IMPLEMENT_SHADER_TYPE(,FSetupCopyIndirectArgumentsCS,TEXT("DistanceFieldSurfaceCacheLightingCompute"),TEXT("SetupCopyIndirectArgumentsCS"),SF_Compute);
 
-
-class FCopyIrradianceCacheSamplesCS : public FGlobalShader
+template<bool bSupportIrradiance>
+class TCopyIrradianceCacheSamplesCS : public FGlobalShader
 {
-	DECLARE_SHADER_TYPE(FCopyIrradianceCacheSamplesCS,Global)
+	DECLARE_SHADER_TYPE(TCopyIrradianceCacheSamplesCS,Global)
 public:
 
 	static bool ShouldCache(EShaderPlatform Platform)
@@ -1092,27 +1108,30 @@ public:
 	{
 		FGlobalShader::ModifyCompilationEnvironment(Platform,OutEnvironment);
 		OutEnvironment.CompilerFlags.Add(CFLAG_StandardOptimization);
+		OutEnvironment.SetDefine(TEXT("SUPPORT_IRRADIANCE"), bSupportIrradiance ? TEXT("1") : TEXT("0"));
 	}
 
-	FCopyIrradianceCacheSamplesCS(const ShaderMetaType::CompiledShaderInitializerType& Initializer)
+	TCopyIrradianceCacheSamplesCS(const ShaderMetaType::CompiledShaderInitializerType& Initializer)
 		: FGlobalShader(Initializer)
 	{
 		IrradianceCachePositionRadius.Bind(Initializer.ParameterMap, TEXT("IrradianceCachePositionRadius"));
 		IrradianceCacheNormal.Bind(Initializer.ParameterMap, TEXT("IrradianceCacheNormal"));
 		OccluderRadius.Bind(Initializer.ParameterMap, TEXT("OccluderRadius"));
 		IrradianceCacheBentNormal.Bind(Initializer.ParameterMap, TEXT("IrradianceCacheBentNormal"));
+		IrradianceCacheIrradiance.Bind(Initializer.ParameterMap, TEXT("IrradianceCacheIrradiance"));
 		IrradianceCacheTileCoordinate.Bind(Initializer.ParameterMap, TEXT("IrradianceCacheTileCoordinate"));
 		DrawParameters.Bind(Initializer.ParameterMap, TEXT("DrawParameters"));
 		CopyIrradianceCachePositionRadius.Bind(Initializer.ParameterMap, TEXT("CopyIrradianceCachePositionRadius"));
 		CopyIrradianceCacheNormal.Bind(Initializer.ParameterMap, TEXT("CopyIrradianceCacheNormal"));
 		CopyOccluderRadius.Bind(Initializer.ParameterMap, TEXT("CopyOccluderRadius"));
 		CopyIrradianceCacheBentNormal.Bind(Initializer.ParameterMap, TEXT("CopyIrradianceCacheBentNormal"));
+		CopyIrradianceCacheIrradiance.Bind(Initializer.ParameterMap, TEXT("CopyIrradianceCacheIrradiance"));
 		CopyIrradianceCacheTileCoordinate.Bind(Initializer.ParameterMap, TEXT("CopyIrradianceCacheTileCoordinate"));
 		ScatterDrawParameters.Bind(Initializer.ParameterMap, TEXT("ScatterDrawParameters"));
 		TrimFraction.Bind(Initializer.ParameterMap, TEXT("TrimFraction"));
 	}
 
-	FCopyIrradianceCacheSamplesCS()
+	TCopyIrradianceCacheSamplesCS()
 	{
 	}
 
@@ -1128,6 +1147,7 @@ public:
 		SetSRVParameter(RHICmdList, ShaderRHI, IrradianceCacheNormal, SurfaceCacheResources.Level[DepthLevel]->Normal.SRV);
 		SetSRVParameter(RHICmdList, ShaderRHI, OccluderRadius, SurfaceCacheResources.Level[DepthLevel]->OccluderRadius.SRV);
 		SetSRVParameter(RHICmdList, ShaderRHI, IrradianceCacheBentNormal, SurfaceCacheResources.Level[DepthLevel]->BentNormal.SRV);
+		SetSRVParameter(RHICmdList, ShaderRHI, IrradianceCacheIrradiance, SurfaceCacheResources.Level[DepthLevel]->Irradiance.SRV);
 		SetSRVParameter(RHICmdList, ShaderRHI, IrradianceCacheTileCoordinate, SurfaceCacheResources.Level[DepthLevel]->TileCoordinate.SRV);
 		SetSRVParameter(RHICmdList, ShaderRHI, DrawParameters, SurfaceCacheResources.Level[DepthLevel]->ScatterDrawParameters.SRV);
 
@@ -1135,6 +1155,7 @@ public:
 		CopyIrradianceCacheNormal.SetBuffer(RHICmdList, ShaderRHI, SurfaceCacheResources.TempResources->Normal);
 		CopyOccluderRadius.SetBuffer(RHICmdList, ShaderRHI, SurfaceCacheResources.TempResources->OccluderRadius);
 		CopyIrradianceCacheBentNormal.SetBuffer(RHICmdList, ShaderRHI, SurfaceCacheResources.TempResources->BentNormal);
+		CopyIrradianceCacheIrradiance.SetBuffer(RHICmdList, ShaderRHI, SurfaceCacheResources.TempResources->Irradiance);
 		CopyIrradianceCacheTileCoordinate.SetBuffer(RHICmdList, ShaderRHI, SurfaceCacheResources.TempResources->TileCoordinate);
 		ScatterDrawParameters.SetBuffer(RHICmdList, ShaderRHI, SurfaceCacheResources.TempResources->ScatterDrawParameters);
 
@@ -1158,12 +1179,14 @@ public:
 		Ar << IrradianceCacheNormal;
 		Ar << OccluderRadius;
 		Ar << IrradianceCacheBentNormal;
+		Ar << IrradianceCacheIrradiance;
 		Ar << IrradianceCacheTileCoordinate;
 		Ar << DrawParameters;
 		Ar << CopyIrradianceCachePositionRadius;
 		Ar << CopyIrradianceCacheNormal;
 		Ar << CopyOccluderRadius;
 		Ar << CopyIrradianceCacheBentNormal;
+		Ar << CopyIrradianceCacheIrradiance;
 		Ar << CopyIrradianceCacheTileCoordinate;
 		Ar << ScatterDrawParameters;
 		Ar << TrimFraction;
@@ -1176,19 +1199,21 @@ private:
 	FShaderResourceParameter IrradianceCacheNormal;
 	FShaderResourceParameter OccluderRadius;
 	FShaderResourceParameter IrradianceCacheBentNormal;
+	FShaderResourceParameter IrradianceCacheIrradiance;
 	FShaderResourceParameter IrradianceCacheTileCoordinate;
 	FShaderResourceParameter DrawParameters;
 	FRWShaderParameter CopyIrradianceCachePositionRadius;
 	FRWShaderParameter CopyIrradianceCacheNormal;
 	FRWShaderParameter CopyOccluderRadius;
 	FRWShaderParameter CopyIrradianceCacheBentNormal;
+	FRWShaderParameter CopyIrradianceCacheIrradiance;
 	FRWShaderParameter CopyIrradianceCacheTileCoordinate;
 	FRWShaderParameter ScatterDrawParameters;
 	FShaderParameter TrimFraction;
 };
 
-IMPLEMENT_SHADER_TYPE(,FCopyIrradianceCacheSamplesCS,TEXT("DistanceFieldSurfaceCacheLightingCompute"),TEXT("CopyIrradianceCacheSamplesCS"),SF_Compute);
-
+IMPLEMENT_SHADER_TYPE(template<>,TCopyIrradianceCacheSamplesCS<true>,TEXT("DistanceFieldSurfaceCacheLightingCompute"),TEXT("CopyIrradianceCacheSamplesCS"),SF_Compute);
+IMPLEMENT_SHADER_TYPE(template<>,TCopyIrradianceCacheSamplesCS<false>,TEXT("DistanceFieldSurfaceCacheLightingCompute"),TEXT("CopyIrradianceCacheSamplesCS"),SF_Compute);
 
 class FSaveStartIndexCS : public FGlobalShader
 {
@@ -1471,9 +1496,28 @@ private:
 
 IMPLEMENT_SHADER_TYPE(,FSetupFinalGatherIndirectArgumentsCS,TEXT("DistanceFieldSurfaceCacheLightingCompute"),TEXT("SetupFinalGatherIndirectArgumentsCS"),SF_Compute);
 
-class FFinalGatherCS : public FGlobalShader
+class FDebugBuffer : public FRenderResource
+{
+public:
+	virtual void InitDynamicRHI() override
+	{
+		DebugData.Initialize(sizeof(float)* 4, 50, PF_A32B32G32R32F, BUF_Static);
+	}
+
+	virtual void ReleaseDynamicRHI() override
+	{
+		DebugData.Release();
+	}
+
+	FRWBuffer DebugData;
+};
+
+TGlobalResource<FDebugBuffer> GDebugBuffer;
+
+template<bool bSupportIrradiance>
+class TFinalGatherCS : public FGlobalShader
 {
-	DECLARE_SHADER_TYPE(FFinalGatherCS,Global)
+	DECLARE_SHADER_TYPE(TFinalGatherCS,Global)
 public:
 
 	static bool ShouldCache(EShaderPlatform Platform)
@@ -1487,12 +1531,13 @@ public:
 		OutEnvironment.SetDefine(TEXT("NUM_CONE_DIRECTIONS"), NumConeSampleDirections);
 		OutEnvironment.SetDefine(TEXT("DOWNSAMPLE_FACTOR"), GAODownsampleFactor);
 		OutEnvironment.SetDefine(TEXT("MAX_OBJECTS_PER_TILE"), GMaxNumObjectsPerTile);
+		OutEnvironment.SetDefine(TEXT("SUPPORT_IRRADIANCE"), bSupportIrradiance ? TEXT("1") : TEXT("0"));
 
 		// To reduce shader compile time of compute shaders with shared memory, doesn't have an impact on generated code with current compiler (June 2010 DX SDK)
 		OutEnvironment.CompilerFlags.Add(CFLAG_StandardOptimization);
 	}
 
-	FFinalGatherCS(const ShaderMetaType::CompiledShaderInitializerType& Initializer)
+	TFinalGatherCS(const ShaderMetaType::CompiledShaderInitializerType& Initializer)
 		: FGlobalShader(Initializer)
 	{
 		DeferredParameters.Bind(Initializer.ParameterMap);
@@ -1502,6 +1547,8 @@ public:
 		IrradianceCacheNormal.Bind(Initializer.ParameterMap, TEXT("IrradianceCacheNormal"));
 		OccluderRadius.Bind(Initializer.ParameterMap, TEXT("OccluderRadius"));
 		IrradianceCacheBentNormal.Bind(Initializer.ParameterMap, TEXT("IrradianceCacheBentNormal"));
+		IrradianceCacheIrradiance.Bind(Initializer.ParameterMap, TEXT("IrradianceCacheIrradiance"));
+		DebugBuffer.Bind(Initializer.ParameterMap, TEXT("DebugBuffer"));
 		IrradianceCacheTileCoordinate.Bind(Initializer.ParameterMap, TEXT("IrradianceCacheTileCoordinate"));
 		ScatterDrawParameters.Bind(Initializer.ParameterMap, TEXT("ScatterDrawParameters"));
 		SavedStartIndex.Bind(Initializer.ParameterMap, TEXT("SavedStartIndex"));
@@ -1522,7 +1569,7 @@ public:
 		RecordRadiusScale.Bind(Initializer.ParameterMap, TEXT("RecordRadiusScale"));
 	}
 
-	FFinalGatherCS()
+	TFinalGatherCS()
 	{
 	}
 
@@ -1554,6 +1601,8 @@ public:
 
 		OccluderRadius.SetBuffer(RHICmdList, ShaderRHI, SurfaceCacheResources.Level[DepthLevel]->OccluderRadius);
 		IrradianceCacheBentNormal.SetBuffer(RHICmdList, ShaderRHI, SurfaceCacheResources.Level[DepthLevel]->BentNormal);
+		IrradianceCacheIrradiance.SetBuffer(RHICmdList, ShaderRHI, SurfaceCacheResources.Level[DepthLevel]->Irradiance);
+		DebugBuffer.SetBuffer(RHICmdList, ShaderRHI, GDebugBuffer.DebugData);
 
 		SetShaderValue(RHICmdList, ShaderRHI, ViewDimensionsParameter, View.ViewRect);
 
@@ -1624,6 +1673,8 @@ public:
 	{
 		OccluderRadius.UnsetUAV(RHICmdList, GetComputeShader());
 		IrradianceCacheBentNormal.UnsetUAV(RHICmdList, GetComputeShader());
+		IrradianceCacheIrradiance.UnsetUAV(RHICmdList, GetComputeShader());
+		DebugBuffer.UnsetUAV(RHICmdList, GetComputeShader());
 	}
 
 	virtual bool Serialize(FArchive& Ar)
@@ -1637,6 +1688,8 @@ public:
 		Ar << IrradianceCacheNormal;
 		Ar << OccluderRadius;
 		Ar << IrradianceCacheBentNormal;
+		Ar << IrradianceCacheIrradiance;
+		Ar << DebugBuffer;
 		Ar << IrradianceCacheTileCoordinate;
 		Ar << ScatterDrawParameters;
 		Ar << SavedStartIndex;
@@ -1667,6 +1720,8 @@ private:
 	FShaderResourceParameter IrradianceCachePositionRadius;
 	FRWShaderParameter OccluderRadius;
 	FRWShaderParameter IrradianceCacheBentNormal;
+	FRWShaderParameter IrradianceCacheIrradiance;
+	FRWShaderParameter DebugBuffer;
 	FShaderResourceParameter ScatterDrawParameters;
 	FShaderResourceParameter SavedStartIndex;
 	FShaderResourceParameter IrradianceCacheTileCoordinate;
@@ -1686,7 +1741,8 @@ private:
 	FShaderParameter RecordRadiusScale;
 };
 
-IMPLEMENT_SHADER_TYPE(,FFinalGatherCS,TEXT("DistanceFieldSurfaceCacheLightingCompute"),TEXT("FinalGatherCS"),SF_Compute);
+IMPLEMENT_SHADER_TYPE(template<>,TFinalGatherCS<true>,TEXT("DistanceFieldSurfaceCacheLightingCompute"),TEXT("FinalGatherCS"),SF_Compute);
+IMPLEMENT_SHADER_TYPE(template<>,TFinalGatherCS<false>,TEXT("DistanceFieldSurfaceCacheLightingCompute"),TEXT("FinalGatherCS"),SF_Compute);
 
 class FWriteDownsampledDepthPS : public FGlobalShader
 {
@@ -1745,7 +1801,7 @@ IMPLEMENT_SHADER_TYPE(,FWriteDownsampledDepthPS,TEXT("DistanceFieldSurfaceCacheL
 
 
 /**  */
-template<bool bFinalInterpolationPass>
+template<bool bFinalInterpolationPass, bool bSupportIrradiance>
 class TIrradianceCacheSplatVS : public FGlobalShader
 {
 	DECLARE_SHADER_TYPE(TIrradianceCacheSplatVS,Global);
@@ -1761,6 +1817,7 @@ public:
 		FGlobalShader::ModifyCompilationEnvironment(Platform,OutEnvironment);
 
 		OutEnvironment.SetDefine(TEXT("FINAL_INTERPOLATION_PASS"), (uint32)bFinalInterpolationPass);
+		OutEnvironment.SetDefine(TEXT("SUPPORT_IRRADIANCE"), bSupportIrradiance ? TEXT("1") : TEXT("0"));
 	}
 
 	TIrradianceCacheSplatVS(const ShaderMetaType::CompiledShaderInitializerType& Initializer):
@@ -1770,6 +1827,7 @@ public:
 		IrradianceCacheOccluderRadius.Bind(Initializer.ParameterMap, TEXT("IrradianceCacheOccluderRadius"));
 		IrradianceCacheNormal.Bind(Initializer.ParameterMap, TEXT("IrradianceCacheNormal"));
 		IrradianceCacheBentNormal.Bind(Initializer.ParameterMap, TEXT("IrradianceCacheBentNormal"));
+		IrradianceCacheIrradiance.Bind(Initializer.ParameterMap, TEXT("IrradianceCacheIrradiance"));
 		InterpolationRadiusScale.Bind(Initializer.ParameterMap, TEXT("InterpolationRadiusScale"));
 		NormalizedOffsetToPixelCenter.Bind(Initializer.ParameterMap, TEXT("NormalizedOffsetToPixelCenter"));
 		HackExpand.Bind(Initializer.ParameterMap, TEXT("HackExpand"));
@@ -1789,6 +1847,7 @@ public:
 		SetSRVParameter(RHICmdList, ShaderRHI, IrradianceCacheOccluderRadius, SurfaceCacheResources.Level[DepthLevel]->OccluderRadius.SRV);
 		SetSRVParameter(RHICmdList, ShaderRHI, IrradianceCacheNormal, SurfaceCacheResources.Level[DepthLevel]->Normal.SRV);
 		SetSRVParameter(RHICmdList, ShaderRHI, IrradianceCacheBentNormal, SurfaceCacheResources.Level[DepthLevel]->BentNormal.SRV);
+		SetSRVParameter(RHICmdList, ShaderRHI, IrradianceCacheIrradiance, SurfaceCacheResources.Level[DepthLevel]->Irradiance.SRV);
 
 		SetShaderValue(RHICmdList, ShaderRHI, InterpolationRadiusScale, (bFinalInterpolationPass ? GAOInterpolationRadiusScale : 1.0f));
 		SetShaderValue(RHICmdList, ShaderRHI, NormalizedOffsetToPixelCenter, NormalizedOffsetToPixelCenterValue);
@@ -1805,6 +1864,7 @@ public:
 		Ar << IrradianceCacheOccluderRadius;
 		Ar << IrradianceCacheNormal;
 		Ar << IrradianceCacheBentNormal;
+		Ar << IrradianceCacheIrradiance;
 		Ar << InterpolationRadiusScale;
 		Ar << NormalizedOffsetToPixelCenter;
 		Ar << HackExpand;
@@ -1816,15 +1876,23 @@ private:
 	FShaderResourceParameter IrradianceCacheOccluderRadius;
 	FShaderResourceParameter IrradianceCacheNormal;
 	FShaderResourceParameter IrradianceCacheBentNormal;
+	FShaderResourceParameter IrradianceCacheIrradiance;
 	FShaderParameter InterpolationRadiusScale;
 	FShaderParameter NormalizedOffsetToPixelCenter;
 	FShaderParameter HackExpand;
 };
 
-IMPLEMENT_SHADER_TYPE(template<>,TIrradianceCacheSplatVS<true>,TEXT("DistanceFieldSurfaceCacheLighting"),TEXT("IrradianceCacheSplatVS"),SF_Vertex);
-IMPLEMENT_SHADER_TYPE(template<>,TIrradianceCacheSplatVS<false>,TEXT("DistanceFieldSurfaceCacheLighting"),TEXT("IrradianceCacheSplatVS"),SF_Vertex);
+// typedef required to get around macro expansion failure due to commas in template argument list 
+#define IMPLEMENT_SPLAT_VS_TYPE(bFinalInterpolationPass, bSupportIrradiance) \
+	typedef TIrradianceCacheSplatVS<bFinalInterpolationPass, bSupportIrradiance> TIrradianceCacheSplatVS##bFinalInterpolationPass##bSupportIrradiance; \
+	IMPLEMENT_SHADER_TYPE(template<>,TIrradianceCacheSplatVS##bFinalInterpolationPass##bSupportIrradiance,TEXT("DistanceFieldSurfaceCacheLighting"),TEXT("IrradianceCacheSplatVS"),SF_Vertex);
+
+IMPLEMENT_SPLAT_VS_TYPE(true, true)
+IMPLEMENT_SPLAT_VS_TYPE(false, true)
+IMPLEMENT_SPLAT_VS_TYPE(true, false)
+IMPLEMENT_SPLAT_VS_TYPE(false, false)
 
-template<bool bFinalInterpolationPass>
+template<bool bFinalInterpolationPass, bool bSupportIrradiance>
 class TIrradianceCacheSplatPS : public FGlobalShader
 {
 	DECLARE_SHADER_TYPE(TIrradianceCacheSplatPS, Global);
@@ -1841,6 +1909,7 @@ public:
 		
 		OutEnvironment.SetDefine(TEXT("DOWNSAMPLE_FACTOR"), GAODownsampleFactor);
 		OutEnvironment.SetDefine(TEXT("FINAL_INTERPOLATION_PASS"), (uint32)bFinalInterpolationPass);
+		OutEnvironment.SetDefine(TEXT("SUPPORT_IRRADIANCE"), bSupportIrradiance ? TEXT("1") : TEXT("0"));
 	}
 
 	/** Default constructor. */
@@ -1908,11 +1977,17 @@ private:
 	FShaderParameter InvMinCosPointBehindPlane;
 };
 
-IMPLEMENT_SHADER_TYPE(template<>,TIrradianceCacheSplatPS<true>,TEXT("DistanceFieldSurfaceCacheLighting"),TEXT("IrradianceCacheSplatPS"),SF_Pixel);
-IMPLEMENT_SHADER_TYPE(template<>,TIrradianceCacheSplatPS<false>,TEXT("DistanceFieldSurfaceCacheLighting"),TEXT("IrradianceCacheSplatPS"),SF_Pixel);
+// typedef required to get around macro expansion failure due to commas in template argument list 
+#define IMPLEMENT_SPLAT_PS_TYPE(bFinalInterpolationPass, bSupportIrradiance) \
+	typedef TIrradianceCacheSplatPS<bFinalInterpolationPass, bSupportIrradiance> TIrradianceCacheSplatPS##bFinalInterpolationPass##bSupportIrradiance; \
+	IMPLEMENT_SHADER_TYPE(template<>,TIrradianceCacheSplatPS##bFinalInterpolationPass##bSupportIrradiance,TEXT("DistanceFieldSurfaceCacheLighting"),TEXT("IrradianceCacheSplatPS"),SF_Pixel);
 
+IMPLEMENT_SPLAT_PS_TYPE(true, true)
+IMPLEMENT_SPLAT_PS_TYPE(false, true)
+IMPLEMENT_SPLAT_PS_TYPE(true, false)
+IMPLEMENT_SPLAT_PS_TYPE(false, false)
 
-template<uint32 NumLevels, bool bFinalInterpolationPass>
+template<uint32 NumLevels, bool bFinalInterpolationPass, bool bSupportIrradiance>
 class TInterpolateCacheCS : public FGlobalShader
 {
 	DECLARE_SHADER_TYPE(TInterpolateCacheCS,Global)
@@ -1931,6 +2006,8 @@ public:
 		OutEnvironment.SetDefine(TEXT("NUM_INTERPOLATE_LEVELS"), NumLevels);
 		OutEnvironment.SetDefine(TEXT("THREADGROUP_SIZEX"), GDistanceFieldAOTileSizeX);
 		OutEnvironment.SetDefine(TEXT("THREADGROUP_SIZEY"), GDistanceFieldAOTileSizeY);
+		OutEnvironment.SetDefine(TEXT("SUPPORT_IRRADIANCE"), bSupportIrradiance ? TEXT("1") : TEXT("0"));
+
 		OutEnvironment.CompilerFlags.Add(CFLAG_StandardOptimization);
 	}
 
@@ -1949,6 +2026,7 @@ public:
 		IrradianceCacheOccluderRadius.Bind(Initializer.ParameterMap, TEXT("LevelIrradianceCacheOccluderRadius"));
 		IrradianceCacheNormal.Bind(Initializer.ParameterMap, TEXT("LevelIrradianceCacheNormal"));
 		IrradianceCacheBentNormal.Bind(Initializer.ParameterMap, TEXT("LevelIrradianceCacheBentNormal"));
+		IrradianceCacheIrradiance.Bind(Initializer.ParameterMap, TEXT("LevelIrradianceCacheIrradiance"));
 		InterpolationRadiusScale.Bind(Initializer.ParameterMap, TEXT("InterpolationRadiusScale"));
 		DrawParameters.Bind(Initializer.ParameterMap, TEXT("LevelScatterDrawParameters"));
 		NumGroups.Bind(Initializer.ParameterMap, TEXT("NumGroups"));
@@ -1958,7 +2036,7 @@ public:
 	{
 	}
 
-	void SetParameters(FRHICommandList& RHICmdList, const FSceneView& View, FSceneRenderTargetItem& DistanceFieldNormal, int32 DepthLevel, int32 DestLevelDownsampleFactor, IPooledRenderTarget* InterpolationTarget, const FDistanceFieldAOParameters& Parameters, FVector2D NumGroupsValue, FSurfaceCacheResources& SurfaceCacheResources)
+	void SetParameters(FRHICommandList& RHICmdList, const FSceneView& View, FSceneRenderTargetItem& DistanceFieldNormal, int32 DepthLevel, int32 DestLevelDownsampleFactor, IPooledRenderTarget* BentNormalInterpolationTarget, const FDistanceFieldAOParameters& Parameters, FVector2D NumGroupsValue, FSurfaceCacheResources& SurfaceCacheResources)
 	{
 		FComputeShaderRHIParamRef ShaderRHI = GetComputeShader();
 		FGlobalShader::SetParameters(RHICmdList, ShaderRHI, View);
@@ -1985,8 +2063,8 @@ public:
 		IrradianceCacheSplat.SetTexture(
 			RHICmdList,
 			ShaderRHI, 
-			InterpolationTarget->GetRenderTargetItem().TargetableTexture, 
-			InterpolationTarget->GetRenderTargetItem().UAV);
+			BentNormalInterpolationTarget->GetRenderTargetItem().TargetableTexture, 
+			BentNormalInterpolationTarget->GetRenderTargetItem().UAV);
 
 		const int32 MaxAllowedLevel = ARRAY_COUNT(SurfaceCacheResources.Level) - 1;
 		const int32 MinLevel = MaxAllowedLevel + 1 - NumLevels;
@@ -2037,6 +2115,17 @@ public:
 					);
 			}
 
+			if (IrradianceCacheIrradiance.IsBound())
+			{
+				check(MaxAllowedLevel - LevelIndex < (int32)IrradianceCacheIrradiance.GetNumResources());
+
+				RHICmdList.SetShaderResourceViewParameter(
+					ShaderRHI,
+					IrradianceCacheIrradiance.GetBaseIndex() + MaxAllowedLevel - LevelIndex,
+					SurfaceCacheResources.Level[LevelIndex]->Irradiance.SRV
+					);
+			}
+
 			if (DrawParameters.IsBound())
 			{
 				check(MaxAllowedLevel - LevelIndex < (int32)DrawParameters.GetNumResources());
@@ -2074,6 +2163,7 @@ public:
 		Ar << IrradianceCacheOccluderRadius;
 		Ar << IrradianceCacheNormal;
 		Ar << IrradianceCacheBentNormal;
+		Ar << IrradianceCacheIrradiance;
 		Ar << InterpolationRadiusScale;
 		Ar << DrawParameters;
 		Ar << NumGroups;
@@ -2094,27 +2184,35 @@ private:
 	FShaderResourceParameter IrradianceCacheOccluderRadius;
 	FShaderResourceParameter IrradianceCacheNormal;
 	FShaderResourceParameter IrradianceCacheBentNormal;
+	FShaderResourceParameter IrradianceCacheIrradiance;
 	FShaderParameter InterpolationRadiusScale;
 	FShaderResourceParameter DrawParameters;
 	FShaderParameter NumGroups;
 };
 
-#define IMPLEMENT_INTERPOLATION_TYPE(NumLevels,bFinalInterpolationPass) \
-	typedef TInterpolateCacheCS<NumLevels,bFinalInterpolationPass> TInterpolateCacheCS##NumLevels##bFinalInterpolationPass; \
-	IMPLEMENT_SHADER_TYPE(template<>,TInterpolateCacheCS##NumLevels##bFinalInterpolationPass,TEXT("DistanceFieldSurfaceCacheLighting"),TEXT("InterpolateCacheCS"),SF_Compute);
-
-IMPLEMENT_INTERPOLATION_TYPE(5, true);
-IMPLEMENT_INTERPOLATION_TYPE(4, true);
-IMPLEMENT_INTERPOLATION_TYPE(1, false);
-IMPLEMENT_INTERPOLATION_TYPE(2, false);
-IMPLEMENT_INTERPOLATION_TYPE(3, false);
-IMPLEMENT_INTERPOLATION_TYPE(4, false);
-IMPLEMENT_INTERPOLATION_TYPE(5, false);
-
-
-class FDistanceFieldAOCombinePS2 : public FGlobalShader
+#define IMPLEMENT_INTERPOLATION_TYPE(NumLevels,bFinalInterpolationPass,bSupportIrradiance) \
+	typedef TInterpolateCacheCS<NumLevels,bFinalInterpolationPass,bSupportIrradiance> TInterpolateCacheCS##NumLevels##bFinalInterpolationPass##bSupportIrradiance; \
+	IMPLEMENT_SHADER_TYPE(template<>,TInterpolateCacheCS##NumLevels##bFinalInterpolationPass##bSupportIrradiance,TEXT("DistanceFieldSurfaceCacheLighting"),TEXT("InterpolateCacheCS"),SF_Compute);
+
+IMPLEMENT_INTERPOLATION_TYPE(5, true, true);
+IMPLEMENT_INTERPOLATION_TYPE(4, true, true);
+IMPLEMENT_INTERPOLATION_TYPE(1, false, true);
+IMPLEMENT_INTERPOLATION_TYPE(2, false, true);
+IMPLEMENT_INTERPOLATION_TYPE(3, false, true);
+IMPLEMENT_INTERPOLATION_TYPE(4, false, true);
+IMPLEMENT_INTERPOLATION_TYPE(5, false, true);
+IMPLEMENT_INTERPOLATION_TYPE(5, true, false);
+IMPLEMENT_INTERPOLATION_TYPE(4, true, false);
+IMPLEMENT_INTERPOLATION_TYPE(1, false, false);
+IMPLEMENT_INTERPOLATION_TYPE(2, false, false);
+IMPLEMENT_INTERPOLATION_TYPE(3, false, false);
+IMPLEMENT_INTERPOLATION_TYPE(4, false, false);
+IMPLEMENT_INTERPOLATION_TYPE(5, false, false);
+
+template<bool bSupportIrradiance>
+class TDistanceFieldAOCombinePS : public FGlobalShader
 {
-	DECLARE_SHADER_TYPE(FDistanceFieldAOCombinePS2, Global);
+	DECLARE_SHADER_TYPE(TDistanceFieldAOCombinePS, Global);
 public:
 
 	static bool ShouldCache(EShaderPlatform Platform)
@@ -2125,25 +2223,29 @@ public:
 	static void ModifyCompilationEnvironment(EShaderPlatform Platform, FShaderCompilerEnvironment& OutEnvironment)
 	{
 		OutEnvironment.SetDefine(TEXT("DOWNSAMPLE_FACTOR"), GAODownsampleFactor);
+		OutEnvironment.SetDefine(TEXT("SUPPORT_IRRADIANCE"), bSupportIrradiance ? TEXT("1") : TEXT("0"));
 	}
 
 	/** Default constructor. */
-	FDistanceFieldAOCombinePS2() {}
+	TDistanceFieldAOCombinePS() {}
 
 	/** Initialization constructor. */
-	FDistanceFieldAOCombinePS2(const ShaderMetaType::CompiledShaderInitializerType& Initializer)
+	TDistanceFieldAOCombinePS(const ShaderMetaType::CompiledShaderInitializerType& Initializer)
 		: FGlobalShader(Initializer)
 	{
 		DeferredParameters.Bind(Initializer.ParameterMap);
 		AOParameters.Bind(Initializer.ParameterMap);
 		AOLevelParameters.Bind(Initializer.ParameterMap);
-		BentNormalAOTexture3.Bind(Initializer.ParameterMap,TEXT("BentNormalAOTexture3"));
-		BentNormalAOSampler3.Bind(Initializer.ParameterMap,TEXT("BentNormalAOSampler3"));
+		BentNormalAOTexture.Bind(Initializer.ParameterMap,TEXT("BentNormalAOTexture"));
+		BentNormalAOSampler.Bind(Initializer.ParameterMap,TEXT("BentNormalAOSampler"));
+		IrradianceTexture.Bind(Initializer.ParameterMap,TEXT("IrradianceTexture"));
+		IrradianceSampler.Bind(Initializer.ParameterMap,TEXT("IrradianceSampler"));
 		DistanceFieldNormalTexture.Bind(Initializer.ParameterMap, TEXT("DistanceFieldNormalTexture"));
 		DistanceFieldNormalSampler.Bind(Initializer.ParameterMap, TEXT("DistanceFieldNormalSampler"));
+		DebugBuffer.Bind(Initializer.ParameterMap, TEXT("DebugBuffer"));
 	}
 
-	void SetParameters(FRHICommandList& RHICmdList, const FSceneView& View, FSceneRenderTargetItem& InBentNormalTexture, FSceneRenderTargetItem& DistanceFieldNormal, const FDistanceFieldAOParameters& Parameters)
+	void SetParameters(FRHICommandList& RHICmdList, const FSceneView& View, FSceneRenderTargetItem& InBentNormalTexture, IPooledRenderTarget* InIrradianceTexture, FSceneRenderTargetItem& DistanceFieldNormal, const FDistanceFieldAOParameters& Parameters)
 	{
 		const FPixelShaderRHIParamRef ShaderRHI = GetPixelShader();
 
@@ -2152,7 +2254,12 @@ public:
 		AOParameters.Set(RHICmdList, ShaderRHI, Parameters);
 		AOLevelParameters.Set(RHICmdList, ShaderRHI, View, GAODownsampleFactor);
 
-		SetTextureParameter(RHICmdList, ShaderRHI, BentNormalAOTexture3, BentNormalAOSampler3, TStaticSamplerState<SF_Point>::GetRHI(), InBentNormalTexture.ShaderResourceTexture);
+		SetTextureParameter(RHICmdList, ShaderRHI, BentNormalAOTexture, BentNormalAOSampler, TStaticSamplerState<SF_Point>::GetRHI(), InBentNormalTexture.ShaderResourceTexture);
+
+		if (IrradianceTexture.IsBound())
+		{
+			SetTextureParameter(RHICmdList, ShaderRHI, IrradianceTexture, IrradianceSampler, TStaticSamplerState<SF_Point>::GetRHI(), InIrradianceTexture->GetRenderTargetItem().ShaderResourceTexture);
+		}
 
 		SetTextureParameter(
 			RHICmdList,
@@ -2162,6 +2269,8 @@ public:
 			TStaticSamplerState<SF_Point,AM_Clamp,AM_Clamp,AM_Clamp>::GetRHI(),
 			DistanceFieldNormal.ShaderResourceTexture
 			);
+
+		SetSRVParameter(RHICmdList, ShaderRHI, DebugBuffer, GDebugBuffer.DebugData.SRV);
 	}
 	// FShader interface.
 	virtual bool Serialize(FArchive& Ar)
@@ -2170,10 +2279,13 @@ public:
 		Ar << DeferredParameters;
 		Ar << AOParameters;
 		Ar << AOLevelParameters;
-		Ar << BentNormalAOTexture3;
-		Ar << BentNormalAOSampler3;
+		Ar << BentNormalAOTexture;
+		Ar << BentNormalAOSampler;
+		Ar << IrradianceTexture;
+		Ar << IrradianceSampler;
 		Ar << DistanceFieldNormalTexture;
 		Ar << DistanceFieldNormalSampler;
+		Ar << DebugBuffer;
 		return bShaderHasOutdatedParameters;
 	}
 
@@ -2182,17 +2294,22 @@ private:
 	FDeferredPixelShaderParameters DeferredParameters;
 	FAOParameters AOParameters;
 	FAOLevelParameters AOLevelParameters;
-	FShaderResourceParameter BentNormalAOTexture3;
-	FShaderResourceParameter BentNormalAOSampler3;
+	FShaderResourceParameter BentNormalAOTexture;
+	FShaderResourceParameter BentNormalAOSampler;
+	FShaderResourceParameter IrradianceTexture;
+	FShaderResourceParameter IrradianceSampler;
 	FShaderResourceParameter DistanceFieldNormalTexture;
 	FShaderResourceParameter DistanceFieldNormalSampler;
+	FShaderResourceParameter DebugBuffer;
 };
 
-IMPLEMENT_SHADER_TYPE(,FDistanceFieldAOCombinePS2,TEXT("DistanceFieldSurfaceCacheLighting"),TEXT("AOCombinePS2"),SF_Pixel);
+IMPLEMENT_SHADER_TYPE(template<>,TDistanceFieldAOCombinePS<true>,TEXT("DistanceFieldSurfaceCacheLighting"),TEXT("AOCombinePS"),SF_Pixel);
+IMPLEMENT_SHADER_TYPE(template<>,TDistanceFieldAOCombinePS<false>,TEXT("DistanceFieldSurfaceCacheLighting"),TEXT("AOCombinePS"),SF_Pixel);
 
-class FFillGapsPS : public FGlobalShader
+template<bool bSupportIrradiance>
+class TFillGapsPS : public FGlobalShader
 {
-	DECLARE_SHADER_TYPE(FFillGapsPS, Global);
+	DECLARE_SHADER_TYPE(TFillGapsPS, Global);
 public:
 
 	static bool ShouldCache(EShaderPlatform Platform)
@@ -2203,21 +2320,24 @@ public:
 	static void ModifyCompilationEnvironment(EShaderPlatform Platform, FShaderCompilerEnvironment& OutEnvironment)
 	{
 		OutEnvironment.SetDefine(TEXT("DOWNSAMPLE_FACTOR"), GAODownsampleFactor);
+		OutEnvironment.SetDefine(TEXT("SUPPORT_IRRADIANCE"), bSupportIrradiance ? TEXT("1") : TEXT("0"));
 	}
 
 	/** Default constructor. */
-	FFillGapsPS() {}
+	TFillGapsPS() {}
 
 	/** Initialization constructor. */
-	FFillGapsPS(const ShaderMetaType::CompiledShaderInitializerType& Initializer)
+	TFillGapsPS(const ShaderMetaType::CompiledShaderInitializerType& Initializer)
 		: FGlobalShader(Initializer)
 	{
 		BentNormalAOTexture.Bind(Initializer.ParameterMap, TEXT("BentNormalAOTexture"));
 		BentNormalAOSampler.Bind(Initializer.ParameterMap, TEXT("BentNormalAOSampler"));
+		IrradianceTexture.Bind(Initializer.ParameterMap, TEXT("IrradianceTexture"));
+		IrradianceSampler.Bind(Initializer.ParameterMap, TEXT("IrradianceSampler"));
 		BentNormalAOTexelSize.Bind(Initializer.ParameterMap, TEXT("BentNormalAOTexelSize"));
 	}
 
-	void SetParameters(FRHICommandList& RHICmdList, const FSceneView& View, FSceneRenderTargetItem& DistanceFieldAOBentNormal2)
+	void SetParameters(FRHICommandList& RHICmdList, const FSceneView& View, FSceneRenderTargetItem& DistanceFieldAOBentNormal2, IPooledRenderTarget* DistanceFieldIrradiance2)
 	{
 		const FPixelShaderRHIParamRef ShaderRHI = GetPixelShader();
 
@@ -2232,6 +2352,18 @@ public:
 			DistanceFieldAOBentNormal2.ShaderResourceTexture
 			);
 
+		if (IrradianceTexture.IsBound())
+		{
+			SetTextureParameter(
+				RHICmdList,
+				ShaderRHI,
+				IrradianceTexture,
+				IrradianceSampler,
+				TStaticSamplerState<SF_Bilinear,AM_Clamp,AM_Clamp,AM_Clamp>::GetRHI(),
+				DistanceFieldIrradiance2->GetRenderTargetItem().ShaderResourceTexture
+				);
+		}
+
 		const FIntPoint DownsampledBufferSize(GSceneRenderTargets.GetBufferSizeXY() / FIntPoint(GAODownsampleFactor, GAODownsampleFactor));
 		const FVector2D BaseLevelTexelSizeValue(1.0f / DownsampledBufferSize.X, 1.0f / DownsampledBufferSize.Y);
 		SetShaderValue(RHICmdList, ShaderRHI, BentNormalAOTexelSize, BaseLevelTexelSizeValue);
@@ -2242,6 +2374,8 @@ public:
 		bool bShaderHasOutdatedParameters = FGlobalShader::Serialize(Ar);
 		Ar << BentNormalAOTexture;
 		Ar << BentNormalAOSampler;
+		Ar << IrradianceTexture;
+		Ar << IrradianceSampler;
 		Ar << BentNormalAOTexelSize;
 		return bShaderHasOutdatedParameters;
 	}
@@ -2250,15 +2384,18 @@ private:
 
 	FShaderResourceParameter BentNormalAOTexture;
 	FShaderResourceParameter BentNormalAOSampler;
+	FShaderResourceParameter IrradianceTexture;
+	FShaderResourceParameter IrradianceSampler;
 	FShaderParameter BentNormalAOTexelSize;
 };
 
-IMPLEMENT_SHADER_TYPE(,FFillGapsPS,TEXT("DistanceFieldSurfaceCacheLighting"),TEXT("FillGapsPS"),SF_Pixel);
-
+IMPLEMENT_SHADER_TYPE(template<>,TFillGapsPS<true>,TEXT("DistanceFieldSurfaceCacheLighting"),TEXT("FillGapsPS"),SF_Pixel);
+IMPLEMENT_SHADER_TYPE(template<>,TFillGapsPS<false>,TEXT("DistanceFieldSurfaceCacheLighting"),TEXT("FillGapsPS"),SF_Pixel);
 
-class FUpdateHistoryPS : public FGlobalShader
+template<bool bSupportIrradiance>
+class TUpdateHistoryPS : public FGlobalShader
 {
-	DECLARE_SHADER_TYPE(FUpdateHistoryPS, Global);
+	DECLARE_SHADER_TYPE(TUpdateHistoryPS, Global);
 public:
 
 	static bool ShouldCache(EShaderPlatform Platform)
@@ -2269,25 +2406,36 @@ public:
 	static void ModifyCompilationEnvironment(EShaderPlatform Platform, FShaderCompilerEnvironment& OutEnvironment)
 	{
 		OutEnvironment.SetDefine(TEXT("DOWNSAMPLE_FACTOR"), GAODownsampleFactor);
+		OutEnvironment.SetDefine(TEXT("SUPPORT_IRRADIANCE"), bSupportIrradiance ? TEXT("1") : TEXT("0"));
 	}
 
 	/** Default constructor. */
-	FUpdateHistoryPS() {}
+	TUpdateHistoryPS() {}
 
 	/** Initialization constructor. */
-	FUpdateHistoryPS(const ShaderMetaType::CompiledShaderInitializerType& Initializer)
+	TUpdateHistoryPS(const ShaderMetaType::CompiledShaderInitializerType& Initializer)
 		: FGlobalShader(Initializer)
 	{
 		DeferredParameters.Bind(Initializer.ParameterMap);
 		BentNormalAOTexture.Bind(Initializer.ParameterMap, TEXT("BentNormalAOTexture"));
 		BentNormalAOSampler.Bind(Initializer.ParameterMap, TEXT("BentNormalAOSampler"));
-		HistoryTexture.Bind(Initializer.ParameterMap, TEXT("HistoryTexture"));
-		HistorySampler.Bind(Initializer.ParameterMap, TEXT("HistorySampler"));
+		BentNormalHistoryTexture.Bind(Initializer.ParameterMap, TEXT("BentNormalHistoryTexture"));
+		BentNormalHistorySampler.Bind(Initializer.ParameterMap, TEXT("BentNormalHistorySampler"));
+		IrradianceTexture.Bind(Initializer.ParameterMap, TEXT("IrradianceTexture"));
+		IrradianceSampler.Bind(Initializer.ParameterMap, TEXT("IrradianceSampler"));
+		IrradianceHistoryTexture.Bind(Initializer.ParameterMap, TEXT("IrradianceHistoryTexture"));
+		IrradianceHistorySampler.Bind(Initializer.ParameterMap, TEXT("IrradianceHistorySampler"));
 		CameraMotion.Bind(Initializer.ParameterMap);
 		HistoryWeight.Bind(Initializer.ParameterMap, TEXT("HistoryWeight"));
 	}
 
-	void SetParameters(FRHICommandList& RHICmdList, const FSceneView& View, FSceneRenderTargetItem& HistoryTextureValue, FSceneRenderTargetItem& DistanceFieldAOBentNormal)
+	void SetParameters(
+		FRHICommandList& RHICmdList, 
+		const FSceneView& View, 
+		FSceneRenderTargetItem& BentNormalHistoryTextureValue, 
+		TRefCountPtr<IPooledRenderTarget>* IrradianceHistoryRT, 
+		FSceneRenderTargetItem& DistanceFieldAOBentNormal, 
+		IPooledRenderTarget* DistanceFieldIrradiance)
 	{
 		const FPixelShaderRHIParamRef ShaderRHI = GetPixelShader();
 
@@ -2306,12 +2454,36 @@ public:
 		SetTextureParameter(
 			RHICmdList,
 			ShaderRHI,
-			HistoryTexture,
-			HistorySampler,
+			BentNormalHistoryTexture,
+			BentNormalHistorySampler,
 			TStaticSamplerState<SF_Bilinear,AM_Clamp,AM_Clamp,AM_Clamp>::GetRHI(),
-			HistoryTextureValue.ShaderResourceTexture
+			BentNormalHistoryTextureValue.ShaderResourceTexture
 			);
 
+		if (IrradianceTexture.IsBound())
+		{
+			SetTextureParameter(
+				RHICmdList,
+				ShaderRHI,
+				IrradianceTexture,
+				IrradianceSampler,
+				TStaticSamplerState<SF_Bilinear,AM_Clamp,AM_Clamp,AM_Clamp>::GetRHI(),
+				DistanceFieldIrradiance->GetRenderTargetItem().ShaderResourceTexture
+				);
+		}
+		
+		if (IrradianceHistoryTexture.IsBound())
+		{
+			SetTextureParameter(
+				RHICmdList,
+				ShaderRHI,
+				IrradianceHistoryTexture,
+				IrradianceHistorySampler,
+				TStaticSamplerState<SF_Bilinear,AM_Clamp,AM_Clamp,AM_Clamp>::GetRHI(),
+				(*IrradianceHistoryRT)->GetRenderTargetItem().ShaderResourceTexture
+				);
+		}
+
 		CameraMotion.Set(RHICmdList, View, ShaderRHI);
 
 		SetShaderValue(RHICmdList, ShaderRHI, HistoryWeight, GAOHistoryWeight);
@@ -2323,8 +2495,12 @@ public:
 		Ar << DeferredParameters;
 		Ar << BentNormalAOTexture;
 		Ar << BentNormalAOSampler;
-		Ar << HistoryTexture;
-		Ar << HistorySampler;
+		Ar << BentNormalHistoryTexture;
+		Ar << BentNormalHistorySampler;
+		Ar << IrradianceTexture;
+		Ar << IrradianceSampler;
+		Ar << IrradianceHistoryTexture;
+		Ar << IrradianceHistorySampler;
 		Ar << CameraMotion;
 		Ar << HistoryWeight;
 		return bShaderHasOutdatedParameters;
@@ -2335,21 +2511,26 @@ private:
 	FDeferredPixelShaderParameters DeferredParameters;
 	FShaderResourceParameter BentNormalAOTexture;
 	FShaderResourceParameter BentNormalAOSampler;
-	FShaderResourceParameter HistoryTexture;
-	FShaderResourceParameter HistorySampler;
+	FShaderResourceParameter BentNormalHistoryTexture;
+	FShaderResourceParameter BentNormalHistorySampler;
+	FShaderResourceParameter IrradianceTexture;
+	FShaderResourceParameter IrradianceSampler;
+	FShaderResourceParameter IrradianceHistoryTexture;
+	FShaderResourceParameter IrradianceHistorySampler;
 	FCameraMotionParameters CameraMotion;
 	FShaderParameter HistoryWeight;
 };
 
-IMPLEMENT_SHADER_TYPE(,FUpdateHistoryPS,TEXT("DistanceFieldSurfaceCacheLighting"),TEXT("UpdateHistoryPS"),SF_Pixel);
+IMPLEMENT_SHADER_TYPE(template<>,TUpdateHistoryPS<true>,TEXT("DistanceFieldSurfaceCacheLighting"),TEXT("UpdateHistoryPS"),SF_Pixel);
+IMPLEMENT_SHADER_TYPE(template<>,TUpdateHistoryPS<false>,TEXT("DistanceFieldSurfaceCacheLighting"),TEXT("UpdateHistoryPS"),SF_Pixel);
 
-void AllocateOrReuseAORenderTarget(TRefCountPtr<IPooledRenderTarget>& Target, const TCHAR* Name)
+void AllocateOrReuseAORenderTarget(TRefCountPtr<IPooledRenderTarget>& Target, const TCHAR* Name, bool bWithAlpha)
 {
 	if (!Target)
 	{
 		FIntPoint BufferSize = GetBufferSizeForAO();
 
-		FPooledRenderTargetDesc Desc(FPooledRenderTargetDesc::Create2DDesc(BufferSize, PF_FloatRGBA, TexCreate_None, TexCreate_RenderTargetable | TexCreate_UAV, false));
+		FPooledRenderTargetDesc Desc(FPooledRenderTargetDesc::Create2DDesc(BufferSize, bWithAlpha ? PF_FloatRGBA : PF_FloatRGB, TexCreate_None, TexCreate_RenderTargetable | TexCreate_UAV, false));
 		GRenderTargetPool.FindFreeElement(Desc, Target, Name);
 	}
 }
@@ -2357,25 +2538,46 @@ void AllocateOrReuseAORenderTarget(TRefCountPtr<IPooledRenderTarget>& Target, co
 void UpdateHistory(
 	FRHICommandList& RHICmdList,
 	FViewInfo& View, 
-	const TCHAR* HistoryRTName,
+	const TCHAR* BentNormalHistoryRTName,
+	const TCHAR* IrradianceHistoryRTName,
 	/** Contains last frame's history, if non-NULL.  This will be updated with the new frame's history. */
-	TRefCountPtr<IPooledRenderTarget>* HistoryState,
+	TRefCountPtr<IPooledRenderTarget>* BentNormalHistoryState,
+	TRefCountPtr<IPooledRenderTarget>* IrradianceHistoryState,
 	/** Source */
-	TRefCountPtr<IPooledRenderTarget>& AOSource, 
+	TRefCountPtr<IPooledRenderTarget>& BentNormalSource, 
+	TRefCountPtr<IPooledRenderTarget>& IrradianceSource, 
 	/** Output of Temporal Reprojection for the next step in the pipeline. */
-	TRefCountPtr<IPooledRenderTarget>& HistoryOutput)
+	TRefCountPtr<IPooledRenderTarget>& BentNormalHistoryOutput,
+	TRefCountPtr<IPooledRenderTarget>& IrradianceHistoryOutput)
 {
-	if (HistoryState)
+	if (BentNormalHistoryState)
 	{
-		if (*HistoryState && !View.bCameraCut && !View.bPrevTransformsReset)
+		if (*BentNormalHistoryState 
+			&& !View.bCameraCut 
+			&& !View.bPrevTransformsReset 
+			&& (!GDistanceFieldGI || (IrradianceHistoryState && *IrradianceHistoryState)))
 		{
 			// Reuse a render target from the pool with a consistent name, for vis purposes
-			TRefCountPtr<IPooledRenderTarget> NewHistory;
-			AllocateOrReuseAORenderTarget(NewHistory, HistoryRTName);
+			TRefCountPtr<IPooledRenderTarget> NewBentNormalHistory;
+			AllocateOrReuseAORenderTarget(NewBentNormalHistory, BentNormalHistoryRTName, true);
+
+			TRefCountPtr<IPooledRenderTarget> NewIrradianceHistory;
+
+			if (GDistanceFieldGI)
+			{
+				AllocateOrReuseAORenderTarget(NewIrradianceHistory, IrradianceHistoryRTName, false);
+			}
 
 			{
 				SCOPED_DRAW_EVENT(RHICmdList, UpdateHistory);
-				SetRenderTarget(RHICmdList, NewHistory->GetRenderTargetItem().TargetableTexture, NULL);
+
+				FTextureRHIParamRef RenderTargets[2] =
+				{
+					NewBentNormalHistory->GetRenderTargetItem().TargetableTexture,
+					GDistanceFieldGI ? NewIrradianceHistory->GetRenderTargetItem().TargetableTexture : NULL,
+				};
+
+				SetRenderTargets(RHICmdList, ARRAY_COUNT(RenderTargets) - (GDistanceFieldGI ? 0 : 1), RenderTargets, FTextureRHIParamRef(), 0, NULL);
 
 				RHICmdList.SetViewport(0, 0, 0.0f, View.ViewRect.Width() / GAODownsampleFactor, View.ViewRect.Height() / GAODownsampleFactor, 1.0f);
 				RHICmdList.SetRasterizerState(TStaticRasterizerState<FM_Solid, CM_None>::GetRHI());
@@ -2383,13 +2585,25 @@ void UpdateHistory(
 				RHICmdList.SetBlendState(TStaticBlendState<>::GetRHI());
 
 				TShaderMapRef<FPostProcessVS> VertexShader(View.ShaderMap);
-				TShaderMapRef<FUpdateHistoryPS> PixelShader(View.ShaderMap);
 
-				static FGlobalBoundShaderState BoundShaderState;
-				
-				SetGlobalBoundShaderState(RHICmdList, View.GetFeatureLevel(), BoundShaderState, GFilterVertexDeclaration.VertexDeclarationRHI, *VertexShader, *PixelShader);
+				if (GDistanceFieldGI)
+				{
+					TShaderMapRef<TUpdateHistoryPS<true> > PixelShader(View.ShaderMap);
+
+					static FGlobalBoundShaderState BoundShaderState;
+					SetGlobalBoundShaderState(RHICmdList, View.GetFeatureLevel(), BoundShaderState, GFilterVertexDeclaration.VertexDeclarationRHI, *VertexShader, *PixelShader);
+					PixelShader->SetParameters(RHICmdList, View, (*BentNormalHistoryState)->GetRenderTargetItem(), IrradianceHistoryState, BentNormalSource->GetRenderTargetItem(), IrradianceSource);
+				}
+				else
+				{
+					TShaderMapRef<TUpdateHistoryPS<false> > PixelShader(View.ShaderMap);
+
+					static FGlobalBoundShaderState BoundShaderState;
+					SetGlobalBoundShaderState(RHICmdList, View.GetFeatureLevel(), BoundShaderState, GFilterVertexDeclaration.VertexDeclarationRHI, *VertexShader, *PixelShader);
+					PixelShader->SetParameters(RHICmdList, View, (*BentNormalHistoryState)->GetRenderTargetItem(), IrradianceHistoryState, BentNormalSource->GetRenderTargetItem(), IrradianceSource);
+				}
 
-				PixelShader->SetParameters(RHICmdList, View, (*HistoryState)->GetRenderTargetItem(), AOSource->GetRenderTargetItem());
+				VertexShader->SetParameters(RHICmdList, View);
 
 				DrawRectangle( 
 					RHICmdList,
@@ -2403,21 +2617,29 @@ void UpdateHistory(
 			}
 
 			// Update the view state's render target reference with the new history
-			*HistoryState = NewHistory;
-			HistoryOutput = NewHistory;
+			*BentNormalHistoryState = NewBentNormalHistory;
+			BentNormalHistoryOutput = NewBentNormalHistory;
+
+			*IrradianceHistoryState = NewIrradianceHistory;
+			IrradianceHistoryOutput = NewIrradianceHistory;
 		}
 		else
 		{
 			// Use the current frame's mask for next frame's history
-			*HistoryState = AOSource;
-			HistoryOutput = AOSource;
-			AOSource = NULL;
+			*BentNormalHistoryState = BentNormalSource;
+			BentNormalHistoryOutput = BentNormalSource;
+			BentNormalSource = NULL;
+
+			*IrradianceHistoryState = IrradianceSource;
+			IrradianceHistoryOutput = IrradianceSource;
+			IrradianceSource = NULL;
 		}
 	}
 	else
 	{
 		// Temporal reprojection is disabled or there is no view state - pass through
-		HistoryOutput = AOSource;
+		BentNormalHistoryOutput = BentNormalSource;
+		IrradianceHistoryOutput = IrradianceSource;
 	}
 }
 
@@ -2425,27 +2647,51 @@ void PostProcessBentNormalAO(
 	FRHICommandList& RHICmdList, 
 	const FDistanceFieldAOParameters& Parameters, 
 	TArray<FViewInfo>& Views, 
-	FSceneRenderTargetItem& IrradianceCacheInterpolation, 
+	FSceneRenderTargetItem& BentNormalInterpolation, 
+	IPooledRenderTarget* IrradianceInterpolation,
 	FSceneRenderTargetItem& DistanceFieldNormal,
-	TRefCountPtr<IPooledRenderTarget>& AOOutput)
+	TRefCountPtr<IPooledRenderTarget>& BentNormalOutput,
+	TRefCountPtr<IPooledRenderTarget>& IrradianceOutput)
 {
 	TRefCountPtr<IPooledRenderTarget> DistanceFieldAOBentNormal;
-	AllocateOrReuseAORenderTarget(DistanceFieldAOBentNormal, TEXT("DistanceFieldBentNormalAO"));
+	TRefCountPtr<IPooledRenderTarget> DistanceFieldIrradiance;
+	AllocateOrReuseAORenderTarget(DistanceFieldAOBentNormal, TEXT("DistanceFieldBentNormalAO"), true);
+
+	if (GDistanceFieldGI)
+	{
+		AllocateOrReuseAORenderTarget(DistanceFieldIrradiance, TEXT("DistanceFieldIrradiance"), false);
+	}
 
 	TRefCountPtr<IPooledRenderTarget> DistanceFieldAOBentNormal2;
+	TRefCountPtr<IPooledRenderTarget> DistanceFieldIrradiance2;
 
 	if (GAOFillGaps)
 	{
-		AllocateOrReuseAORenderTarget(DistanceFieldAOBentNormal2, TEXT("DistanceFieldBentNormalAO2"));
+		AllocateOrReuseAORenderTarget(DistanceFieldAOBentNormal2, TEXT("DistanceFieldBentNormalAO2"), true);
+
+		if (GDistanceFieldGI)
+		{
+			AllocateOrReuseAORenderTarget(DistanceFieldIrradiance2, TEXT("DistanceFieldIrradiance2"), false);
+		}
 	}
 
 	{
 		SCOPED_DRAW_EVENT(RHICmdList, AOCombine);
 
-		SetRenderTarget(RHICmdList, GAOFillGaps
-			? DistanceFieldAOBentNormal2->GetRenderTargetItem().TargetableTexture
-			: DistanceFieldAOBentNormal->GetRenderTargetItem().TargetableTexture, 
-			NULL);
+		FTextureRHIParamRef DiffuseIrradianceTarget = NULL;
+
+		if (GDistanceFieldGI)
+		{
+			DiffuseIrradianceTarget = GAOFillGaps ? DistanceFieldIrradiance2->GetRenderTargetItem().TargetableTexture : DistanceFieldIrradiance->GetRenderTargetItem().TargetableTexture;
+		}
+
+		FTextureRHIParamRef RenderTargets[2] =
+		{
+			GAOFillGaps ? DistanceFieldAOBentNormal2->GetRenderTargetItem().TargetableTexture : DistanceFieldAOBentNormal->GetRenderTargetItem().TargetableTexture,
+			DiffuseIrradianceTarget,
+		};
+
+		SetRenderTargets(RHICmdList, ARRAY_COUNT(RenderTargets) - (GDistanceFieldGI ? 0 : 1), RenderTargets, FTextureRHIParamRef(), 0, NULL);
 
 		for (int32 ViewIndex = 0; ViewIndex < Views.Num(); ViewIndex++)
 		{
@@ -2457,13 +2703,21 @@ void PostProcessBentNormalAO(
 			RHICmdList.SetBlendState(TStaticBlendState<>::GetRHI());
 
 			TShaderMapRef<FPostProcessVS> VertexShader(View.ShaderMap);
-			TShaderMapRef<FDistanceFieldAOCombinePS2> PixelShader(View.ShaderMap);
 
-			static FGlobalBoundShaderState BoundShaderState;
-			
-			SetGlobalBoundShaderState(RHICmdList, View.GetFeatureLevel(), BoundShaderState, GFilterVertexDeclaration.VertexDeclarationRHI, *VertexShader, *PixelShader);
-
-			PixelShader->SetParameters(RHICmdList, View, IrradianceCacheInterpolation, DistanceFieldNormal, Parameters);
+			if (GDistanceFieldGI)
+			{
+				TShaderMapRef<TDistanceFieldAOCombinePS<true> > PixelShader(View.ShaderMap);
+				static FGlobalBoundShaderState BoundShaderState;
+				SetGlobalBoundShaderState(RHICmdList, View.GetFeatureLevel(), BoundShaderState, GFilterVertexDeclaration.VertexDeclarationRHI, *VertexShader, *PixelShader);
+				PixelShader->SetParameters(RHICmdList, View, BentNormalInterpolation, IrradianceInterpolation, DistanceFieldNormal, Parameters);
+			}
+			else
+			{
+				TShaderMapRef<TDistanceFieldAOCombinePS<false> > PixelShader(View.ShaderMap);
+				static FGlobalBoundShaderState BoundShaderState;
+				SetGlobalBoundShaderState(RHICmdList, View.GetFeatureLevel(), BoundShaderState, GFilterVertexDeclaration.VertexDeclarationRHI, *VertexShader, *PixelShader);
+				PixelShader->SetParameters(RHICmdList, View, BentNormalInterpolation, IrradianceInterpolation, DistanceFieldNormal, Parameters);
+			}
 
 			DrawRectangle( 
 				RHICmdList,
@@ -2480,7 +2734,14 @@ void PostProcessBentNormalAO(
 	if (GAOFillGaps)
 	{
 		SCOPED_DRAW_EVENT(RHICmdList, FillGaps);
-		SetRenderTarget(RHICmdList, DistanceFieldAOBentNormal->GetRenderTargetItem().TargetableTexture, NULL);
+
+		FTextureRHIParamRef RenderTargets[2] =
+		{
+			DistanceFieldAOBentNormal->GetRenderTargetItem().TargetableTexture,
+			GDistanceFieldGI ? DistanceFieldIrradiance->GetRenderTargetItem().TargetableTexture : NULL,
+		};
+
+		SetRenderTargets(RHICmdList, ARRAY_COUNT(RenderTargets) - (GDistanceFieldGI ? 0 : 1), RenderTargets, FTextureRHIParamRef(), 0, NULL);
 
 		for (int32 ViewIndex = 0; ViewIndex < Views.Num(); ViewIndex++)
 		{
@@ -2492,14 +2753,25 @@ void PostProcessBentNormalAO(
 			RHICmdList.SetBlendState(TStaticBlendState<>::GetRHI());
 
 			TShaderMapRef<FPostProcessVS> VertexShader(View.ShaderMap);
-			TShaderMapRef<FFillGapsPS> PixelShader(View.ShaderMap);
 
-			static FGlobalBoundShaderState BoundShaderState;
-			
-			SetGlobalBoundShaderState(RHICmdList, View.GetFeatureLevel(), BoundShaderState, GFilterVertexDeclaration.VertexDeclarationRHI, *VertexShader, *PixelShader);
+			if (GDistanceFieldGI)
+			{
+				TShaderMapRef<TFillGapsPS<true> > PixelShader(View.ShaderMap);
+
+				static FGlobalBoundShaderState BoundShaderState;
+				SetGlobalBoundShaderState(RHICmdList, View.GetFeatureLevel(), BoundShaderState, GFilterVertexDeclaration.VertexDeclarationRHI, *VertexShader, *PixelShader);
+				PixelShader->SetParameters(RHICmdList, View, DistanceFieldAOBentNormal2->GetRenderTargetItem(), DistanceFieldIrradiance2);
+			}
+			else
+			{
+				TShaderMapRef<TFillGapsPS<false> > PixelShader(View.ShaderMap);
+
+				static FGlobalBoundShaderState BoundShaderState;
+				SetGlobalBoundShaderState(RHICmdList, View.GetFeatureLevel(), BoundShaderState, GFilterVertexDeclaration.VertexDeclarationRHI, *VertexShader, *PixelShader);
+				PixelShader->SetParameters(RHICmdList, View, DistanceFieldAOBentNormal2->GetRenderTargetItem(), DistanceFieldIrradiance2);
+			}
 
 			VertexShader->SetParameters(RHICmdList, View);
-			PixelShader->SetParameters(RHICmdList, View, DistanceFieldAOBentNormal2->GetRenderTargetItem());
 
 			DrawRectangle( 
 				RHICmdList,
@@ -2514,8 +2786,10 @@ void PostProcessBentNormalAO(
 	}
 
 	FSceneViewState* ViewState = (FSceneViewState*)Views[0].State;
-	TRefCountPtr<IPooledRenderTarget>* HistoryState = ViewState ? &ViewState->DistanceFieldAOHistoryRT : NULL;
-	AOOutput = DistanceFieldAOBentNormal;
+	TRefCountPtr<IPooledRenderTarget>* BentNormalHistoryState = ViewState ? &ViewState->DistanceFieldAOHistoryRT : NULL;
+	TRefCountPtr<IPooledRenderTarget>* IrradianceHistoryState = ViewState ? &ViewState->DistanceFieldIrradianceHistoryRT : NULL;
+	BentNormalOutput = DistanceFieldAOBentNormal;
+	IrradianceOutput = DistanceFieldIrradiance;
 
 	if (GAOUseHistory)
 	{
@@ -2525,13 +2799,17 @@ void PostProcessBentNormalAO(
 			RHICmdList,
 			View, 
 			TEXT("DistanceFieldAOHistory"),
-			HistoryState,
+			TEXT("DistanceFieldIrradianceHistory"),
+			BentNormalHistoryState,
+			IrradianceHistoryState,
 			DistanceFieldAOBentNormal, 
-			AOOutput);
+			DistanceFieldIrradiance,
+			BentNormalOutput,
+			IrradianceOutput);
 	}
 }
 
-template<bool bOutputBentNormal>
+template<bool bOutputBentNormal, bool bSupportIrradiance>
 class TDistanceFieldAOUpsamplePS : public FGlobalShader
 {
 	DECLARE_SHADER_TYPE(TDistanceFieldAOUpsamplePS, Global);
@@ -2546,6 +2824,7 @@ public:
 	{
 		OutEnvironment.SetDefine(TEXT("DOWNSAMPLE_FACTOR"), GAODownsampleFactor);
 		OutEnvironment.SetDefine(TEXT("OUTPUT_BENT_NORMAL"), bOutputBentNormal ? TEXT("1") : TEXT("0"));
+		OutEnvironment.SetDefine(TEXT("SUPPORT_IRRADIANCE"), bSupportIrradiance ? TEXT("1") : TEXT("0"));
 	}
 
 	/** Default constructor. */
@@ -2558,9 +2837,11 @@ public:
 		DeferredParameters.Bind(Initializer.ParameterMap);
 		BentNormalAOTexture.Bind(Initializer.ParameterMap,TEXT("BentNormalAOTexture"));
 		BentNormalAOSampler.Bind(Initializer.ParameterMap,TEXT("BentNormalAOSampler"));
+		IrradianceTexture.Bind(Initializer.ParameterMap,TEXT("IrradianceTexture"));
+		IrradianceSampler.Bind(Initializer.ParameterMap,TEXT("IrradianceSampler"));
 	}
 
-	void SetParameters(FRHICommandList& RHICmdList, const FSceneView& View, TRefCountPtr<IPooledRenderTarget>& DistanceFieldAOBentNormal)
+	void SetParameters(FRHICommandList& RHICmdList, const FSceneView& View, TRefCountPtr<IPooledRenderTarget>& DistanceFieldAOBentNormal, IPooledRenderTarget* DistanceFieldIrradiance)
 	{
 		const FPixelShaderRHIParamRef ShaderRHI = GetPixelShader();
 
@@ -2568,6 +2849,11 @@ public:
 		DeferredParameters.Set(RHICmdList, ShaderRHI, View);
 
 		SetTextureParameter(RHICmdList, ShaderRHI, BentNormalAOTexture, BentNormalAOSampler, TStaticSamplerState<SF_Bilinear>::GetRHI(), DistanceFieldAOBentNormal->GetRenderTargetItem().ShaderResourceTexture);
+
+		if (IrradianceTexture.IsBound())
+		{
+			SetTextureParameter(RHICmdList, ShaderRHI, IrradianceTexture, IrradianceSampler, TStaticSamplerState<SF_Bilinear>::GetRHI(), DistanceFieldIrradiance->GetRenderTargetItem().ShaderResourceTexture);
+		}
 	}
 	// FShader interface.
 	virtual bool Serialize(FArchive& Ar)
@@ -2576,6 +2862,8 @@ public:
 		Ar << DeferredParameters;
 		Ar << BentNormalAOTexture;
 		Ar << BentNormalAOSampler;
+		Ar << IrradianceTexture;
+		Ar << IrradianceSampler;
 		return bShaderHasOutdatedParameters;
 	}
 
@@ -2584,12 +2872,19 @@ private:
 	FDeferredPixelShaderParameters DeferredParameters;
 	FShaderResourceParameter BentNormalAOTexture;
 	FShaderResourceParameter BentNormalAOSampler;
+	FShaderResourceParameter IrradianceTexture;
+	FShaderResourceParameter IrradianceSampler;
 };
 
-IMPLEMENT_SHADER_TYPE(template<>,TDistanceFieldAOUpsamplePS<true>,TEXT("DistanceFieldSurfaceCacheLighting"),TEXT("AOUpsamplePS"),SF_Pixel);
-IMPLEMENT_SHADER_TYPE(template<>,TDistanceFieldAOUpsamplePS<false>,TEXT("DistanceFieldSurfaceCacheLighting"),TEXT("AOUpsamplePS"),SF_Pixel);
+#define IMPLEMENT_UPSAMPLE_PS_TYPE(bOutputBentNormal, bSupportIrradiance) \
+	typedef TDistanceFieldAOUpsamplePS<bOutputBentNormal, bSupportIrradiance> TDistanceFieldAOUpsamplePS##bOutputBentNormal##bSupportIrradiance; \
+	IMPLEMENT_SHADER_TYPE(template<>,TDistanceFieldAOUpsamplePS##bOutputBentNormal##bSupportIrradiance,TEXT("DistanceFieldSurfaceCacheLighting"),TEXT("AOUpsamplePS"),SF_Pixel);
 
-void UpsampleBentNormalAO(FRHICommandList& RHICmdList, const TArray<FViewInfo>& Views, TRefCountPtr<IPooledRenderTarget>& DistanceFieldAOBentNormal, bool bOutputBentNormal)
+IMPLEMENT_UPSAMPLE_PS_TYPE(true, true)
+IMPLEMENT_UPSAMPLE_PS_TYPE(true, false)
+IMPLEMENT_UPSAMPLE_PS_TYPE(false, false)
+
+void UpsampleBentNormalAO(FRHICommandList& RHICmdList, const TArray<FViewInfo>& Views, TRefCountPtr<IPooledRenderTarget>& DistanceFieldAOBentNormal, TRefCountPtr<IPooledRenderTarget>& DistanceFieldIrradiance, bool bOutputBentNormal)
 {
 	for (int32 ViewIndex = 0; ViewIndex < Views.Num(); ViewIndex++)
 	{
@@ -2606,17 +2901,27 @@ void UpsampleBentNormalAO(FRHICommandList& RHICmdList, const TArray<FViewInfo>&
 
 		if (bOutputBentNormal)
 		{
-			TShaderMapRef<TDistanceFieldAOUpsamplePS<true> > PixelShader(View.ShaderMap);
-			static FGlobalBoundShaderState BoundShaderState;
-			SetGlobalBoundShaderState(RHICmdList, View.GetFeatureLevel(), BoundShaderState, GFilterVertexDeclaration.VertexDeclarationRHI, *VertexShader, *PixelShader);
-			PixelShader->SetParameters(RHICmdList, View, DistanceFieldAOBentNormal);
+			if (GDistanceFieldGI)
+			{
+				TShaderMapRef<TDistanceFieldAOUpsamplePS<true, true> > PixelShader(View.ShaderMap);
+				static FGlobalBoundShaderState BoundShaderState;
+				SetGlobalBoundShaderState(RHICmdList, View.GetFeatureLevel(), BoundShaderState, GFilterVertexDeclaration.VertexDeclarationRHI, *VertexShader, *PixelShader);
+				PixelShader->SetParameters(RHICmdList, View, DistanceFieldAOBentNormal, DistanceFieldIrradiance);
+			}
+			else
+			{
+				TShaderMapRef<TDistanceFieldAOUpsamplePS<true, false> > PixelShader(View.ShaderMap);
+				static FGlobalBoundShaderState BoundShaderState;
+				SetGlobalBoundShaderState(RHICmdList, View.GetFeatureLevel(), BoundShaderState, GFilterVertexDeclaration.VertexDeclarationRHI, *VertexShader, *PixelShader);
+				PixelShader->SetParameters(RHICmdList, View, DistanceFieldAOBentNormal, DistanceFieldIrradiance);
+			}
 		}
 		else
 		{
-			TShaderMapRef<TDistanceFieldAOUpsamplePS<false> > PixelShader(View.ShaderMap);
+			TShaderMapRef<TDistanceFieldAOUpsamplePS<false, false> > PixelShader(View.ShaderMap);
 			static FGlobalBoundShaderState BoundShaderState;
 			SetGlobalBoundShaderState(RHICmdList, View.GetFeatureLevel(), BoundShaderState, GFilterVertexDeclaration.VertexDeclarationRHI, *VertexShader, *PixelShader);
-			PixelShader->SetParameters(RHICmdList, View, DistanceFieldAOBentNormal);
+			PixelShader->SetParameters(RHICmdList, View, DistanceFieldAOBentNormal, DistanceFieldIrradiance);
 		}
 
 		DrawRectangle( 
@@ -2994,7 +3299,8 @@ void SetupDepthStencil(FRHICommandListImmediate& RHICmdList, FViewInfo& View, FS
 void RenderIrradianceCacheInterpolation(
 	FRHICommandListImmediate& RHICmdList, 
 	TArray<FViewInfo>& Views, 
-	IPooledRenderTarget* InterpolationTarget, 
+	IPooledRenderTarget* BentNormalInterpolationTarget, 
+	IPooledRenderTarget* IrradianceInterpolationTarget, 
 	FSceneRenderTargetItem& DistanceFieldNormal, 
 	int32 DepthLevel, 
 	int32 DestLevelDownsampleFactor, 
@@ -3016,13 +3322,26 @@ void RenderIrradianceCacheInterpolation(
 			uint32 GroupSizeX = FMath::DivideAndRoundUp(DownsampledViewSize.X, GDistanceFieldAOTileSizeX);
 			uint32 GroupSizeY = FMath::DivideAndRoundUp(DownsampledViewSize.Y, GDistanceFieldAOTileSizeY);
 
-			TShaderMapRef<TInterpolateCacheCS<4, true> > ComputeShader(View.ShaderMap);
+			if (GDistanceFieldGI)
+			{
+				TShaderMapRef<TInterpolateCacheCS<4, true, true> > ComputeShader(View.ShaderMap);
 
-			RHICmdList.SetComputeShader(ComputeShader->GetComputeShader());	
-			ComputeShader->SetParameters(RHICmdList, View, DistanceFieldNormal, DepthLevel, DestLevelDownsampleFactor, InterpolationTarget, Parameters, FVector2D(GroupSizeX, GroupSizeY), SurfaceCacheResources);
-			DispatchComputeShader(RHICmdList, *ComputeShader, GroupSizeX, GroupSizeY, 1);
+				RHICmdList.SetComputeShader(ComputeShader->GetComputeShader());	
+				ComputeShader->SetParameters(RHICmdList, View, DistanceFieldNormal, DepthLevel, DestLevelDownsampleFactor, BentNormalInterpolationTarget, Parameters, FVector2D(GroupSizeX, GroupSizeY), SurfaceCacheResources);
+				DispatchComputeShader(RHICmdList, *ComputeShader, GroupSizeX, GroupSizeY, 1);
 
-			ComputeShader->UnsetParameters(RHICmdList);
+				ComputeShader->UnsetParameters(RHICmdList);
+			}
+			else
+			{
+				TShaderMapRef<TInterpolateCacheCS<4, true, false> > ComputeShader(View.ShaderMap);
+
+				RHICmdList.SetComputeShader(ComputeShader->GetComputeShader());	
+				ComputeShader->SetParameters(RHICmdList, View, DistanceFieldNormal, DepthLevel, DestLevelDownsampleFactor, BentNormalInterpolationTarget, Parameters, FVector2D(GroupSizeX, GroupSizeY), SurfaceCacheResources);
+				DispatchComputeShader(RHICmdList, *ComputeShader, GroupSizeX, GroupSizeY, 1);
+
+				ComputeShader->UnsetParameters(RHICmdList);
+			}
 		}
 	}
 	else
@@ -3031,11 +3350,11 @@ void RenderIrradianceCacheInterpolation(
 
 		if (bFinalInterpolation && (GAOInterpolationDepthTesting || GAOInterpolationStencilTesting))
 		{
-			FPooledRenderTargetDesc Desc(FPooledRenderTargetDesc::Create2DDesc(InterpolationTarget->GetDesc().Extent, PF_DepthStencil, TexCreate_None, TexCreate_DepthStencilTargetable, false));
+			FPooledRenderTargetDesc Desc(FPooledRenderTargetDesc::Create2DDesc(BentNormalInterpolationTarget->GetDesc().Extent, PF_DepthStencil, TexCreate_None, TexCreate_DepthStencilTargetable, false));
 			GRenderTargetPool.FindFreeElement(Desc, SplatDepthStencilBuffer, TEXT("DistanceFieldAOSplatDepthBuffer"));
 
-			SetRenderTarget(RHICmdList, InterpolationTarget->GetRenderTargetItem().TargetableTexture, SplatDepthStencilBuffer->GetRenderTargetItem().TargetableTexture);
-			RHICmdList.Clear(true, FLinearColor(0, 0, 0, 0), true, 0, true, 0, FIntRect());
+			SetRenderTarget(RHICmdList, NULL, SplatDepthStencilBuffer->GetRenderTargetItem().TargetableTexture);
+			RHICmdList.Clear(false, FLinearColor(0, 0, 0, 0), true, 0, true, 0, FIntRect());
 
 			for (int32 ViewIndex = 0; ViewIndex < Views.Num(); ViewIndex++)
 			{
@@ -3045,12 +3364,19 @@ void RenderIrradianceCacheInterpolation(
 
 			GRenderTargetPool.VisualizeTexture.SetCheckPoint(RHICmdList, SplatDepthStencilBuffer);
 
-			// Restore render target after checkpoint
-			SetRenderTarget(RHICmdList, InterpolationTarget->GetRenderTargetItem().TargetableTexture, SplatDepthStencilBuffer ? SplatDepthStencilBuffer->GetRenderTargetItem().TargetableTexture : NULL);
+			FTextureRHIParamRef RenderTargets[2] =
+			{
+				BentNormalInterpolationTarget->GetRenderTargetItem().TargetableTexture,
+				GDistanceFieldGI ? IrradianceInterpolationTarget->GetRenderTargetItem().TargetableTexture : NULL
+			};
+			SetRenderTargets(RHICmdList, ARRAY_COUNT(RenderTargets) - (GDistanceFieldGI ? 0 : 1), RenderTargets, SplatDepthStencilBuffer->GetRenderTargetItem().TargetableTexture, 0, NULL);
+
+			FLinearColor ClearColors[2] = {FLinearColor(0, 0, 0, 0), FLinearColor(0, 0, 0, 0)};
+			RHICmdList.ClearMRT(true, ARRAY_COUNT(ClearColors) - (GDistanceFieldGI ? 0 : 1), ClearColors, false, 0, false, 0, FIntRect());
 		}
 		else
 		{
-			SetRenderTarget(RHICmdList, InterpolationTarget->GetRenderTargetItem().TargetableTexture, NULL);
+			SetRenderTarget(RHICmdList, BentNormalInterpolationTarget->GetRenderTargetItem().TargetableTexture, NULL);
 			RHICmdList.Clear(true, FLinearColor(0, 0, 0, 0), false, 0, false, 0, FIntRect());
 		}
 
@@ -3109,11 +3435,10 @@ void RenderIrradianceCacheInterpolation(
 			}
 			else
 			{
-				//@todo - render front faces with depth testing, requires depth buffer
 				RHICmdList.SetDepthStencilState(TStaticDepthStencilState<false, CF_Always>::GetRHI());
 			}
 
-			RHICmdList.SetBlendState(TStaticBlendState<CW_RGBA, BO_Add, BF_One, BF_One, BO_Add, BF_One, BF_One>::GetRHI());
+			RHICmdList.SetBlendState(TStaticBlendState<CW_RGBA, BO_Add, BF_One, BF_One, BO_Add, BF_One, BF_One, CW_RGBA, BO_Add, BF_One, BF_One, BO_Add, BF_One, BF_One>::GetRHI());
 
 			const FIntPoint BufferSize = GSceneRenderTargets.GetBufferSizeXY();
 			const float InvBufferSizeX = 1.0f / BufferSize.X;
@@ -3135,34 +3460,60 @@ void RenderIrradianceCacheInterpolation(
 
 			if (bFinalInterpolation)
 			{
-				TShaderMapRef<TIrradianceCacheSplatVS<true> > VertexShader(View.ShaderMap);
-				TShaderMapRef<TIrradianceCacheSplatPS<true> > PixelShader(View.ShaderMap);
-
 				for (int32 SplatSourceDepthLevel = GAOMaxLevel; SplatSourceDepthLevel >= FMath::Max(DepthLevel, GAOMinLevel); SplatSourceDepthLevel--)
 				{
-					static FGlobalBoundShaderState BoundShaderState;
-					
-					SetGlobalBoundShaderState(RHICmdList, View.GetFeatureLevel(), BoundShaderState, GScreenVertexDeclaration.VertexDeclarationRHI, *VertexShader, *PixelShader);
+					if (GDistanceFieldGI)
+					{
+						TShaderMapRef<TIrradianceCacheSplatVS<true, true> > VertexShader(View.ShaderMap);
+						TShaderMapRef<TIrradianceCacheSplatPS<true, true> > PixelShader(View.ShaderMap);
+
+						static FGlobalBoundShaderState BoundShaderState;
+						SetGlobalBoundShaderState(RHICmdList, View.GetFeatureLevel(), BoundShaderState, GScreenVertexDeclaration.VertexDeclarationRHI, *VertexShader, *PixelShader);
 
-					VertexShader->SetParameters(RHICmdList, View, SplatSourceDepthLevel, DestLevelDownsampleFactor, NormalizedOffsetToPixelCenter);
-					PixelShader->SetParameters(RHICmdList, View, DistanceFieldNormal, DestLevelDownsampleFactor, Parameters);
+						VertexShader->SetParameters(RHICmdList, View, SplatSourceDepthLevel, DestLevelDownsampleFactor, NormalizedOffsetToPixelCenter);
+						PixelShader->SetParameters(RHICmdList, View, DistanceFieldNormal, DestLevelDownsampleFactor, Parameters);
+					}
+					else
+					{
+						TShaderMapRef<TIrradianceCacheSplatVS<true, false> > VertexShader(View.ShaderMap);
+						TShaderMapRef<TIrradianceCacheSplatPS<true, false> > PixelShader(View.ShaderMap);
+
+						static FGlobalBoundShaderState BoundShaderState;
+						SetGlobalBoundShaderState(RHICmdList, View.GetFeatureLevel(), BoundShaderState, GScreenVertexDeclaration.VertexDeclarationRHI, *VertexShader, *PixelShader);
+
+						VertexShader->SetParameters(RHICmdList, View, SplatSourceDepthLevel, DestLevelDownsampleFactor, NormalizedOffsetToPixelCenter);
+						PixelShader->SetParameters(RHICmdList, View, DistanceFieldNormal, DestLevelDownsampleFactor, Parameters);
+					}
 
 					RHICmdList.DrawPrimitiveIndirect(PT_TriangleList, SurfaceCacheResources.Level[SplatSourceDepthLevel]->ScatterDrawParameters.Buffer, 0);
 				}
 			}
 			else
 			{
-				TShaderMapRef<TIrradianceCacheSplatVS<false> > VertexShader(View.ShaderMap);
-				TShaderMapRef<TIrradianceCacheSplatPS<false> > PixelShader(View.ShaderMap);
-
 				for (int32 SplatSourceDepthLevel = GAOMaxLevel; SplatSourceDepthLevel >= FMath::Max(DepthLevel, GAOMinLevel); SplatSourceDepthLevel--)
 				{
-					static FGlobalBoundShaderState BoundShaderState;
-					
-					SetGlobalBoundShaderState(RHICmdList, View.GetFeatureLevel(), BoundShaderState, GScreenVertexDeclaration.VertexDeclarationRHI, *VertexShader, *PixelShader);
+					if (GDistanceFieldGI)
+					{
+						TShaderMapRef<TIrradianceCacheSplatVS<false, true> > VertexShader(View.ShaderMap);
+						TShaderMapRef<TIrradianceCacheSplatPS<false, true> > PixelShader(View.ShaderMap);
+
+						static FGlobalBoundShaderState BoundShaderState;
+						SetGlobalBoundShaderState(RHICmdList, View.GetFeatureLevel(), BoundShaderState, GScreenVertexDeclaration.VertexDeclarationRHI, *VertexShader, *PixelShader);
+
+						VertexShader->SetParameters(RHICmdList, View, SplatSourceDepthLevel, DestLevelDownsampleFactor, NormalizedOffsetToPixelCenter);
+						PixelShader->SetParameters(RHICmdList, View, DistanceFieldNormal, DestLevelDownsampleFactor, Parameters);
+					}
+					else
+					{
+						TShaderMapRef<TIrradianceCacheSplatVS<false, false> > VertexShader(View.ShaderMap);
+						TShaderMapRef<TIrradianceCacheSplatPS<false, false> > PixelShader(View.ShaderMap);
+
+						static FGlobalBoundShaderState BoundShaderState;
+						SetGlobalBoundShaderState(RHICmdList, View.GetFeatureLevel(), BoundShaderState, GScreenVertexDeclaration.VertexDeclarationRHI, *VertexShader, *PixelShader);
 
-					VertexShader->SetParameters(RHICmdList, View, SplatSourceDepthLevel, DestLevelDownsampleFactor, NormalizedOffsetToPixelCenter);
-					PixelShader->SetParameters(RHICmdList, View, DistanceFieldNormal, DestLevelDownsampleFactor, Parameters);
+						VertexShader->SetParameters(RHICmdList, View, SplatSourceDepthLevel, DestLevelDownsampleFactor, NormalizedOffsetToPixelCenter);
+						PixelShader->SetParameters(RHICmdList, View, DistanceFieldNormal, DestLevelDownsampleFactor, Parameters);
+					}
 
 					RHICmdList.DrawPrimitiveIndirect(PT_TriangleList, SurfaceCacheResources.Level[SplatSourceDepthLevel]->ScatterDrawParameters.Buffer, 0);
 				}
@@ -3309,7 +3660,12 @@ void UpdateVisibleObjectBuffers(const FScene* Scene, const FDistanceFieldAOParam
 	}
 }
 
-bool FDeferredShadingSceneRenderer::RenderDistanceFieldAOSurfaceCache(FRHICommandListImmediate& RHICmdList, const FDistanceFieldAOParameters& Parameters, TRefCountPtr<IPooledRenderTarget>& OutDynamicBentNormalAO, bool bApplyToSceneColor)
+bool FDeferredShadingSceneRenderer::RenderDistanceFieldAOSurfaceCache(
+	FRHICommandListImmediate& RHICmdList, 
+	const FDistanceFieldAOParameters& Parameters, 
+	TRefCountPtr<IPooledRenderTarget>& OutDynamicBentNormalAO, 
+	TRefCountPtr<IPooledRenderTarget>& OutDynamicIrradiance, 
+	bool bApplyToSceneColor)
 {
 	//@todo - support multiple views
 	const FViewInfo& View = Views[0];
@@ -3339,7 +3695,7 @@ bool FDeferredShadingSceneRenderer::RenderDistanceFieldAOSurfaceCache(FRHIComman
 
 			FSurfaceCacheResources& SurfaceCacheResources = *Scene->SurfaceCacheResources;
 			FIntPoint DownsampledViewRect = FIntPoint::DivideAndRoundDown(View.ViewRect.Size(), GAODownsampleFactor);
-			int32 SlackAmount = 2;
+			int32 SlackAmount = 4;
 			int32 MaxSurfaceCacheSamples = SlackAmount * DownsampledViewRect.X * DownsampledViewRect.Y / (1 << (2 * GAOMinLevel * GAOPowerOfTwoBetweenLevels));
 			SurfaceCacheResources.AllocateFor(GAOMinLevel, GAOMaxLevel, MaxSurfaceCacheSamples);
 
@@ -3402,8 +3758,19 @@ bool FDeferredShadingSceneRenderer::RenderDistanceFieldAOSurfaceCache(FRHIComman
 								ComputeShader->UnsetParameters(RHICmdList);
 							}
 
+							if (GDistanceFieldGI)
 							{
-								TShaderMapRef<FCopyIrradianceCacheSamplesCS> ComputeShader(View.ShaderMap);
+								TShaderMapRef<TCopyIrradianceCacheSamplesCS<true> > ComputeShader(View.ShaderMap);
+
+								RHICmdList.SetComputeShader(ComputeShader->GetComputeShader());
+								ComputeShader->SetParameters(RHICmdList, View, DepthLevel);
+								DispatchIndirectComputeShader(RHICmdList, *ComputeShader, SurfaceCacheResources.DispatchParameters.Buffer, 0);
+
+								ComputeShader->UnsetParameters(RHICmdList);
+							}
+							else
+							{
+								TShaderMapRef<TCopyIrradianceCacheSamplesCS<false> > ComputeShader(View.ShaderMap);
 
 								RHICmdList.SetComputeShader(ComputeShader->GetComputeShader());
 								ComputeShader->SetParameters(RHICmdList, View, DepthLevel);
@@ -3426,16 +3793,16 @@ bool FDeferredShadingSceneRenderer::RenderDistanceFieldAOSurfaceCache(FRHIComman
 
 					SCOPED_DRAW_EVENTF(RHICmdList, Level, TEXT("Level_%d"), DepthLevel);
 
-					TRefCountPtr<IPooledRenderTarget> DistanceFieldAOIrradianceCacheSplat;
+					TRefCountPtr<IPooledRenderTarget> DistanceFieldAOBentNormalSplat;
 
 					{
 						FIntPoint AOBufferSize = FIntPoint::DivideAndRoundUp(GSceneRenderTargets.GetBufferSizeXY(), DestLevelDownsampleFactor);
 						FPooledRenderTargetDesc Desc(FPooledRenderTargetDesc::Create2DDesc(AOBufferSize, PF_FloatRGBA, TexCreate_None, TexCreate_RenderTargetable | TexCreate_UAV, false));
-						GRenderTargetPool.FindFreeElement(Desc, DistanceFieldAOIrradianceCacheSplat, TEXT("DistanceFieldAOIrradianceCacheSplat"));
+						GRenderTargetPool.FindFreeElement(Desc, DistanceFieldAOBentNormalSplat, TEXT("DistanceFieldAOBentNormalSplat"));
 					}
 
 					// Splat / interpolate the surface cache records onto the buffer sized for the current depth level
-					RenderIrradianceCacheInterpolation(RHICmdList, Views, DistanceFieldAOIrradianceCacheSplat, DistanceFieldNormal->GetRenderTargetItem(), DepthLevel, DestLevelDownsampleFactor, Parameters, false);
+					RenderIrradianceCacheInterpolation(RHICmdList, Views, DistanceFieldAOBentNormalSplat, NULL, DistanceFieldNormal->GetRenderTargetItem(), DepthLevel, DestLevelDownsampleFactor, Parameters, false);
 
 					{
 						SCOPED_DRAW_EVENT(RHICmdList, PopulateIrradianceCache);
@@ -3464,7 +3831,7 @@ bool FDeferredShadingSceneRenderer::RenderDistanceFieldAOSurfaceCache(FRHIComman
 							TShaderMapRef<FPopulateCacheCS> ComputeShader(View.ShaderMap);
 
 							RHICmdList.SetComputeShader(ComputeShader->GetComputeShader());
-							ComputeShader->SetParameters(RHICmdList, View, DistanceFieldAOIrradianceCacheSplat->GetRenderTargetItem(), DistanceFieldNormal->GetRenderTargetItem(), DestLevelDownsampleFactor, DepthLevel, TileListGroupSize, Parameters);
+							ComputeShader->SetParameters(RHICmdList, View, DistanceFieldAOBentNormalSplat->GetRenderTargetItem(), DistanceFieldNormal->GetRenderTargetItem(), DestLevelDownsampleFactor, DepthLevel, TileListGroupSize, Parameters);
 							DispatchComputeShader(RHICmdList, *ComputeShader, GroupSizeX, GroupSizeY, 1);
 
 							ComputeShader->UnsetParameters(RHICmdList);
@@ -3485,11 +3852,22 @@ bool FDeferredShadingSceneRenderer::RenderDistanceFieldAOSurfaceCache(FRHIComman
 						}
 
 						// Compute lighting for the new surface cache records by cone-stepping through the object distance fields
+						if (GDistanceFieldGI)
+						{
+							TShaderMapRef<TFinalGatherCS<true> > ComputeShader(View.ShaderMap);
+
+							RHICmdList.SetComputeShader(ComputeShader->GetComputeShader());
+							ComputeShader->SetParameters(RHICmdList, View, DistanceFieldAOBentNormalSplat->GetRenderTargetItem(), DistanceFieldNormal->GetRenderTargetItem(), NumObjects, DestLevelDownsampleFactor, DepthLevel, TileListGroupSize, Parameters);
+							DispatchIndirectComputeShader(RHICmdList, *ComputeShader, SurfaceCacheResources.DispatchParameters.Buffer, 0);
+
+							ComputeShader->UnsetParameters(RHICmdList);
+						}
+						else
 						{
-							TShaderMapRef<FFinalGatherCS> ComputeShader(View.ShaderMap);
+							TShaderMapRef<TFinalGatherCS<false> > ComputeShader(View.ShaderMap);
 
 							RHICmdList.SetComputeShader(ComputeShader->GetComputeShader());
-							ComputeShader->SetParameters(RHICmdList, View, DistanceFieldAOIrradianceCacheSplat->GetRenderTargetItem(), DistanceFieldNormal->GetRenderTargetItem(), NumObjects, DestLevelDownsampleFactor, DepthLevel, TileListGroupSize, Parameters);
+							ComputeShader->SetParameters(RHICmdList, View, DistanceFieldAOBentNormalSplat->GetRenderTargetItem(), DistanceFieldNormal->GetRenderTargetItem(), NumObjects, DestLevelDownsampleFactor, DepthLevel, TileListGroupSize, Parameters);
 							DispatchIndirectComputeShader(RHICmdList, *ComputeShader, SurfaceCacheResources.DispatchParameters.Buffer, 0);
 
 							ComputeShader->UnsetParameters(RHICmdList);
@@ -3497,27 +3875,34 @@ bool FDeferredShadingSceneRenderer::RenderDistanceFieldAOSurfaceCache(FRHIComman
 					}
 				}
 
-				TRefCountPtr<IPooledRenderTarget> IrradianceCacheAccumulation;
+				TRefCountPtr<IPooledRenderTarget> BentNormalAccumulation;
+				TRefCountPtr<IPooledRenderTarget> IrradianceAccumulation;
 
 				{
 					FIntPoint BufferSize = GetBufferSizeForAO();
 					FPooledRenderTargetDesc Desc(FPooledRenderTargetDesc::Create2DDesc(BufferSize, PF_FloatRGBA, TexCreate_None, TexCreate_RenderTargetable | TexCreate_UAV, false));
-					GRenderTargetPool.FindFreeElement(Desc, IrradianceCacheAccumulation, TEXT("DistanceFieldAOIrradianceCacheAccumulation"));
+					GRenderTargetPool.FindFreeElement(Desc, BentNormalAccumulation, TEXT("BentNormalAccumulation"));
+
+					if (GDistanceFieldGI)
+					{
+						GRenderTargetPool.FindFreeElement(Desc, IrradianceAccumulation, TEXT("IrradianceAccumulation"));
+					}
 				}
 
 				// Splat the surface cache records onto the opaque pixels, using less strict weighting so the lighting is smoothed in world space
 				{
 					SCOPED_DRAW_EVENT(RHICmdList, FinalIrradianceCacheSplat);
-					RenderIrradianceCacheInterpolation(RHICmdList, Views, IrradianceCacheAccumulation, DistanceFieldNormal->GetRenderTargetItem(), 0, GAODownsampleFactor, Parameters, true);
+					RenderIrradianceCacheInterpolation(RHICmdList, Views, BentNormalAccumulation, IrradianceAccumulation, DistanceFieldNormal->GetRenderTargetItem(), 0, GAODownsampleFactor, Parameters, true);
 				}
 
-				GRenderTargetPool.VisualizeTexture.SetCheckPoint(RHICmdList, IrradianceCacheAccumulation);
+				GRenderTargetPool.VisualizeTexture.SetCheckPoint(RHICmdList, BentNormalAccumulation);
 
 				// Post process the AO to cover over artifacts
-				TRefCountPtr<IPooledRenderTarget> AOOutput;
-				PostProcessBentNormalAO(RHICmdList, Parameters, Views, IrradianceCacheAccumulation->GetRenderTargetItem(), DistanceFieldNormal->GetRenderTargetItem(), AOOutput);
+				TRefCountPtr<IPooledRenderTarget> BentNormalOutput;
+				TRefCountPtr<IPooledRenderTarget> IrradianceOutput;
+				PostProcessBentNormalAO(RHICmdList, Parameters, Views, BentNormalAccumulation->GetRenderTargetItem(), IrradianceAccumulation, DistanceFieldNormal->GetRenderTargetItem(), BentNormalOutput, IrradianceOutput);
 
-				GRenderTargetPool.VisualizeTexture.SetCheckPoint(RHICmdList, AOOutput);
+				GRenderTargetPool.VisualizeTexture.SetCheckPoint(RHICmdList, BentNormalOutput);
 
 				if (bApplyToSceneColor)
 				{
@@ -3530,11 +3915,22 @@ bool FDeferredShadingSceneRenderer::RenderDistanceFieldAOSurfaceCache(FRHIComman
 					Desc.Format = PF_FloatRGBA;
 					GRenderTargetPool.FindFreeElement(Desc, OutDynamicBentNormalAO, TEXT("DynamicBentNormalAO"));
 
-					SetRenderTarget(RHICmdList, OutDynamicBentNormalAO->GetRenderTargetItem().TargetableTexture, NULL);
+					if (GDistanceFieldGI)
+					{
+						Desc.Format = PF_FloatRGB;
+						GRenderTargetPool.FindFreeElement(Desc, OutDynamicIrradiance, TEXT("DynamicIrradiance"));
+					}
+
+					FTextureRHIParamRef RenderTargets[2] =
+					{
+						OutDynamicBentNormalAO->GetRenderTargetItem().TargetableTexture,
+						GDistanceFieldGI ? OutDynamicIrradiance->GetRenderTargetItem().TargetableTexture : NULL
+					};
+					SetRenderTargets(RHICmdList, ARRAY_COUNT(RenderTargets) - (GDistanceFieldGI ? 0 : 1), RenderTargets, FTextureRHIParamRef(), 0, NULL);
 				}
 
 				// Upsample to full resolution, write to output
-				UpsampleBentNormalAO(RHICmdList, Views, AOOutput, !bApplyToSceneColor);
+				UpsampleBentNormalAO(RHICmdList, Views, BentNormalOutput, IrradianceOutput, !bApplyToSceneColor);
 
 				return true;
 			}
@@ -3778,7 +4174,7 @@ void FDeferredShadingSceneRenderer::RenderMeshDistanceFieldVisualization(FRHICom
 	}
 }
 
-template<bool bApplyShadowing>
+template<bool bApplyShadowing, bool bSupportIrradiance>
 class TDynamicSkyLightDiffusePS : public FGlobalShader
 {
 	DECLARE_SHADER_TYPE(TDynamicSkyLightDiffusePS, Global);
@@ -3794,6 +4190,7 @@ public:
 		FGlobalShader::ModifyCompilationEnvironment(Platform,OutEnvironment);
 		OutEnvironment.SetDefine(TEXT("DOWNSAMPLE_FACTOR"), GAODownsampleFactor);
 		OutEnvironment.SetDefine(TEXT("APPLY_SHADOWING"), (uint32)(bApplyShadowing ? 1 : 0));
+		OutEnvironment.SetDefine(TEXT("SUPPORT_IRRADIANCE"), bSupportIrradiance ? TEXT("1") : TEXT("0"));
 	}
 
 	/** Default constructor. */
@@ -3806,11 +4203,13 @@ public:
 		DeferredParameters.Bind(Initializer.ParameterMap);
 		DynamicBentNormalAOTexture.Bind(Initializer.ParameterMap, TEXT("BentNormalAOTexture"));
 		DynamicBentNormalAOSampler.Bind(Initializer.ParameterMap, TEXT("BentNormalAOSampler"));
+		DynamicIrradianceTexture.Bind(Initializer.ParameterMap, TEXT("IrradianceTexture"));
+		DynamicIrradianceSampler.Bind(Initializer.ParameterMap, TEXT("IrradianceSampler"));
 		ContrastAndNormalizeMulAdd.Bind(Initializer.ParameterMap, TEXT("ContrastAndNormalizeMulAdd"));
 		OcclusionTintAndMinOcclusion.Bind(Initializer.ParameterMap, TEXT("OcclusionTintAndMinOcclusion"));
 	}
 
-	void SetParameters(FRHICommandList& RHICmdList, const FSceneView& View, FTextureRHIParamRef DynamicBentNormalAO, const FDistanceFieldAOParameters& Parameters, const FSkyLightSceneProxy* SkyLight)
+	void SetParameters(FRHICommandList& RHICmdList, const FSceneView& View, FTextureRHIParamRef DynamicBentNormalAO, IPooledRenderTarget* DynamicIrradiance, const FDistanceFieldAOParameters& Parameters, const FSkyLightSceneProxy* SkyLight)
 	{
 		const FPixelShaderRHIParamRef ShaderRHI = GetPixelShader();
 		FGlobalShader::SetParameters(RHICmdList, ShaderRHI, View);
@@ -3818,6 +4217,11 @@ public:
 
 		SetTextureParameter(RHICmdList, ShaderRHI, DynamicBentNormalAOTexture, DynamicBentNormalAOSampler, TStaticSamplerState<SF_Point>::GetRHI(), DynamicBentNormalAO);
 
+		if (DynamicIrradianceTexture.IsBound())
+		{
+			SetTextureParameter(RHICmdList, ShaderRHI, DynamicIrradianceTexture, DynamicIrradianceSampler, TStaticSamplerState<SF_Point>::GetRHI(), DynamicIrradiance->GetRenderTargetItem().ShaderResourceTexture);
+		}
+
 		// Scale and bias to remap the contrast curve to [0,1]
 		const float Min = 1 / (1 + FMath::Exp(-Parameters.Contrast * (0 * 10 - 5)));
 		const float Max = 1 / (1 + FMath::Exp(-Parameters.Contrast * (1 * 10 - 5)));
@@ -3838,6 +4242,8 @@ public:
 		Ar << DeferredParameters;
 		Ar << DynamicBentNormalAOTexture;
 		Ar << DynamicBentNormalAOSampler;
+		Ar << DynamicIrradianceTexture;
+		Ar << DynamicIrradianceSampler;
 		Ar << ContrastAndNormalizeMulAdd;
 		Ar << OcclusionTintAndMinOcclusion;
 		return bShaderHasOutdatedParameters;
@@ -3848,12 +4254,20 @@ private:
 	FDeferredPixelShaderParameters DeferredParameters;
 	FShaderResourceParameter DynamicBentNormalAOTexture;
 	FShaderResourceParameter DynamicBentNormalAOSampler;
+	FShaderResourceParameter DynamicIrradianceTexture;
+	FShaderResourceParameter DynamicIrradianceSampler;
 	FShaderParameter ContrastAndNormalizeMulAdd;
 	FShaderParameter OcclusionTintAndMinOcclusion;
 };
 
-IMPLEMENT_SHADER_TYPE(template<>,TDynamicSkyLightDiffusePS<true>,TEXT("SkyLighting"),TEXT("SkyLightDiffusePS"),SF_Pixel);
-IMPLEMENT_SHADER_TYPE(template<>,TDynamicSkyLightDiffusePS<false>,TEXT("SkyLighting"),TEXT("SkyLightDiffusePS"),SF_Pixel);
+#define IMPLEMENT_SKYLIGHT_PS_TYPE(bApplyShadowing, bSupportIrradiance) \
+	typedef TDynamicSkyLightDiffusePS<bApplyShadowing, bSupportIrradiance> TDynamicSkyLightDiffusePS##bApplyShadowing##bSupportIrradiance; \
+	IMPLEMENT_SHADER_TYPE(template<>,TDynamicSkyLightDiffusePS##bApplyShadowing##bSupportIrradiance,TEXT("SkyLighting"),TEXT("SkyLightDiffusePS"),SF_Pixel);
+
+IMPLEMENT_SKYLIGHT_PS_TYPE(true, true)
+IMPLEMENT_SKYLIGHT_PS_TYPE(false, true)
+IMPLEMENT_SKYLIGHT_PS_TYPE(true, false)
+IMPLEMENT_SKYLIGHT_PS_TYPE(false, false)
 
 void FDeferredShadingSceneRenderer::RenderDynamicSkyLighting(FRHICommandListImmediate& RHICmdList, TRefCountPtr<IPooledRenderTarget>& DynamicBentNormalAO)
 {
@@ -3871,13 +4285,14 @@ void FDeferredShadingSceneRenderer::RenderDynamicSkyLighting(FRHICommandListImme
 		bool bApplyShadowing = false;
 
 		FDistanceFieldAOParameters Parameters(Scene->SkyLight->OcclusionMaxDistance, Scene->SkyLight->Contrast);
+		TRefCountPtr<IPooledRenderTarget> DynamicIrradiance;
 
 		if (Scene->SkyLight->bCastShadows
 			&& ViewFamily.EngineShowFlags.DistanceFieldAO 
 			&& !ViewFamily.EngineShowFlags.VisualizeDistanceFieldAO
 			&& !ViewFamily.EngineShowFlags.VisualizeMeshDistanceFields)
 		{
-			bApplyShadowing = RenderDistanceFieldAOSurfaceCache(RHICmdList, Parameters, DynamicBentNormalAO, false);
+			bApplyShadowing = RenderDistanceFieldAOSurfaceCache(RHICmdList, Parameters, DynamicBentNormalAO, DynamicIrradiance, false);
 		}
 
 		GSceneRenderTargets.BeginRenderingSceneColor(RHICmdList);
@@ -3887,27 +4302,57 @@ void FDeferredShadingSceneRenderer::RenderDynamicSkyLighting(FRHICommandListImme
 			const FViewInfo& View = Views[ViewIndex];
 
 			RHICmdList.SetViewport(View.ViewRect.Min.X, View.ViewRect.Min.Y, 0.0f, View.ViewRect.Max.X, View.ViewRect.Max.Y, 1.0f);
-			RHICmdList.SetRasterizerState(TStaticRasterizerState< FM_Solid, CM_None >::GetRHI());
-			RHICmdList.SetDepthStencilState(TStaticDepthStencilState< false, CF_Always >::GetRHI());
-			RHICmdList.SetBlendState(TStaticBlendState< CW_RGB, BO_Add, BF_One, BF_One >::GetRHI());
+			RHICmdList.SetRasterizerState(TStaticRasterizerState<FM_Solid, CM_None>::GetRHI());
+			RHICmdList.SetDepthStencilState(TStaticDepthStencilState<false, CF_Always>::GetRHI());
+
+			if (GAOOverwriteSceneColor)
+			{
+				RHICmdList.SetBlendState(TStaticBlendState<>::GetRHI());
+			}
+			else
+			{
+				RHICmdList.SetBlendState(TStaticBlendState<CW_RGB, BO_Add, BF_One, BF_One>::GetRHI());
+			}
 
 			TShaderMapRef< FPostProcessVS > VertexShader(View.ShaderMap);
 
 			if (bApplyShadowing)
 			{
-				TShaderMapRef<TDynamicSkyLightDiffusePS<true> > PixelShader(View.ShaderMap);
+				if (GDistanceFieldGI)
+				{
+					TShaderMapRef<TDynamicSkyLightDiffusePS<true, true> > PixelShader(View.ShaderMap);
 
-				static FGlobalBoundShaderState BoundShaderState;
-				SetGlobalBoundShaderState(RHICmdList, View.GetFeatureLevel(), BoundShaderState, GFilterVertexDeclaration.VertexDeclarationRHI, *VertexShader, *PixelShader);
-				PixelShader->SetParameters(RHICmdList, View, DynamicBentNormalAO->GetRenderTargetItem().ShaderResourceTexture, Parameters, Scene->SkyLight);
+					static FGlobalBoundShaderState BoundShaderState;
+					SetGlobalBoundShaderState(RHICmdList, View.GetFeatureLevel(), BoundShaderState, GFilterVertexDeclaration.VertexDeclarationRHI, *VertexShader, *PixelShader);
+					PixelShader->SetParameters(RHICmdList, View, DynamicBentNormalAO->GetRenderTargetItem().ShaderResourceTexture, DynamicIrradiance, Parameters, Scene->SkyLight);
+				}
+				else
+				{
+					TShaderMapRef<TDynamicSkyLightDiffusePS<true, false> > PixelShader(View.ShaderMap);
+
+					static FGlobalBoundShaderState BoundShaderState;
+					SetGlobalBoundShaderState(RHICmdList, View.GetFeatureLevel(), BoundShaderState, GFilterVertexDeclaration.VertexDeclarationRHI, *VertexShader, *PixelShader);
+					PixelShader->SetParameters(RHICmdList, View, DynamicBentNormalAO->GetRenderTargetItem().ShaderResourceTexture, DynamicIrradiance, Parameters, Scene->SkyLight);
+				}
 			}
 			else
 			{
-				TShaderMapRef<TDynamicSkyLightDiffusePS<false> > PixelShader(View.ShaderMap);
+				if (GDistanceFieldGI)
+				{
+					TShaderMapRef<TDynamicSkyLightDiffusePS<false, true> > PixelShader(View.ShaderMap);
 
-				static FGlobalBoundShaderState BoundShaderState;
-				SetGlobalBoundShaderState(RHICmdList, View.GetFeatureLevel(), BoundShaderState, GFilterVertexDeclaration.VertexDeclarationRHI, *VertexShader, *PixelShader);
-				PixelShader->SetParameters(RHICmdList, View, GWhiteTexture->TextureRHI, Parameters, Scene->SkyLight);
+					static FGlobalBoundShaderState BoundShaderState;
+					SetGlobalBoundShaderState(RHICmdList, View.GetFeatureLevel(), BoundShaderState, GFilterVertexDeclaration.VertexDeclarationRHI, *VertexShader, *PixelShader);
+					PixelShader->SetParameters(RHICmdList, View, GWhiteTexture->TextureRHI, NULL, Parameters, Scene->SkyLight);
+				}
+				else
+				{
+					TShaderMapRef<TDynamicSkyLightDiffusePS<false, false> > PixelShader(View.ShaderMap);
+
+					static FGlobalBoundShaderState BoundShaderState;
+					SetGlobalBoundShaderState(RHICmdList, View.GetFeatureLevel(), BoundShaderState, GFilterVertexDeclaration.VertexDeclarationRHI, *VertexShader, *PixelShader);
+					PixelShader->SetParameters(RHICmdList, View, GWhiteTexture->TextureRHI, NULL, Parameters, Scene->SkyLight);
+				}
 			}
 
 			DrawRectangle( 
diff --git a/Engine/Source/Runtime/Renderer/Private/DistanceFieldSurfaceCacheLighting.h b/Engine/Source/Runtime/Renderer/Private/DistanceFieldSurfaceCacheLighting.h
--- a/Engine/Source/Runtime/Renderer/Private/DistanceFieldSurfaceCacheLighting.h
+++ b/Engine/Source/Runtime/Renderer/Private/DistanceFieldSurfaceCacheLighting.h
@@ -44,6 +44,7 @@ public:
 			OccluderRadius.Initialize(sizeof(float), MaxIrradianceCacheSamples, PF_R32_FLOAT, BUF_Static);
 			Normal.Initialize(sizeof(FFloat16Color), MaxIrradianceCacheSamples, PF_FloatRGBA, BUF_Static);
 			BentNormal.Initialize(sizeof(FFloat16Color), MaxIrradianceCacheSamples, PF_FloatRGBA, BUF_Static);
+			Irradiance.Initialize(sizeof(FFloat16Color), MaxIrradianceCacheSamples, PF_FloatRGBA, BUF_Static);
 			ScatterDrawParameters.Initialize(sizeof(uint32), 4, PF_R32_UINT, BUF_Static | BUF_DrawIndirect);
 			SavedStartIndex.Initialize(sizeof(uint32), 1, PF_R32_UINT, BUF_Static);
 			TileCoordinate.Initialize(sizeof(uint16)* 2, MaxIrradianceCacheSamples, PF_R16G16_UINT, BUF_Static);
@@ -56,6 +57,7 @@ public:
 		OccluderRadius.Release();
 		Normal.Release();
 		BentNormal.Release();
+		Irradiance.Release();
 		ScatterDrawParameters.Release();
 		SavedStartIndex.Release();
 		TileCoordinate.Release();
@@ -67,6 +69,7 @@ public:
 	FRWBuffer OccluderRadius;
 	FRWBuffer Normal;
 	FRWBuffer BentNormal;
+	FRWBuffer Irradiance;
 	FRWBuffer ScatterDrawParameters;
 	FRWBuffer SavedStartIndex;
 	FRWBuffer TileCoordinate;
diff --git a/Engine/Source/Runtime/Renderer/Private/ScenePrivate.h b/Engine/Source/Runtime/Renderer/Private/ScenePrivate.h
--- a/Engine/Source/Runtime/Renderer/Private/ScenePrivate.h
+++ b/Engine/Source/Runtime/Renderer/Private/ScenePrivate.h
@@ -536,6 +536,7 @@ public:
 	// Temporal AA result for light shafts of last frame
 	TMap<const ULightComponent*, TRefCountPtr<IPooledRenderTarget> > LightShaftBloomHistoryRTs;
 	TRefCountPtr<IPooledRenderTarget> DistanceFieldAOHistoryRT;
+	TRefCountPtr<IPooledRenderTarget> DistanceFieldIrradianceHistoryRT;
 	// Mobile temporal AA surfaces.
 	TRefCountPtr<IPooledRenderTarget> MobileAaBloomSunVignette0;
 	TRefCountPtr<IPooledRenderTarget> MobileAaBloomSunVignette1;
@@ -696,6 +697,7 @@ public:
 		LightShaftOcclusionHistoryRT.SafeRelease();
 		LightShaftBloomHistoryRTs.Empty();
 		DistanceFieldAOHistoryRT.SafeRelease();
+		DistanceFieldIrradianceHistoryRT.SafeRelease();
 		MobileAaBloomSunVignette0.SafeRelease();
 		MobileAaBloomSunVignette1.SafeRelease();
 		MobileAaColor0.SafeRelease();
