GPU object management of distance field meshes, can handle 100x more meshes with only a small GPU cost for culling.  CPU uploads scene change deltas only, nothing per frame.
New primitive component setting bAffectDistanceFieldLighting, useful for grass whose instance count overwhelms DFAO
Distance field volume texture allocations are sorted by size for better packing, and atlas is repacked on texture reallocation

diff --git a/Engine/Shaders/DistanceFieldLightingShared.usf b/Engine/Shaders/DistanceFieldLightingShared.usf
--- a/Engine/Shaders/DistanceFieldLightingShared.usf
+++ b/Engine/Shaders/DistanceFieldLightingShared.usf
@@ -22,6 +22,10 @@
 #define DOWNSAMPLE_FACTOR 1
 #endif
 
+#ifndef UPDATEOBJECTS_THREADGROUP_SIZE
+#define UPDATEOBJECTS_THREADGROUP_SIZE 1
+#endif
+
 float3 DistanceFieldVolumePositionToUV(float3 VolumePosition, float3 UVScale, float3 UVAdd)
 {
 	float3 VolumeUV = VolumePosition * UVScale + UVAdd;
@@ -32,72 +36,75 @@ Texture3D DistanceFieldTexture;
 SamplerState DistanceFieldSampler;
 float3 DistanceFieldAtlasTexelSize;
 
-uint NumObjects;
+RWBuffer<uint> RWObjectIndirectArguments;
+Buffer<uint> ObjectIndirectArguments;
+
+uint GetCulledNumObjects()
+{
+	// IndexCount, NumInstances, StartIndex, BaseVertexIndex, FirstInstance
+	return ObjectIndirectArguments[1];
+}
+
+// In float4's.  Must match equivalent C++ variables.
+#define OBJECT_DATA_STRIDE 12
+
+uint NumSceneObjects;
+
+// Have to make these R32F with 4x the reads and writes because of the horrible D3D11 limitation
+// "error X3676: typed UAV loads are only allowed for single-component 32-bit element types"
+Buffer<float> ObjectBounds;
+Buffer<float> ObjectData;
+
+RWBuffer<float> RWObjectBounds;
+RWBuffer<float> RWObjectData;
 
 // In float4's.  Must match equivalent C++ variables.
-#define OBJECT_DATA_STRIDE 8
-#define OBJECT_DATA2_STRIDE 4
-#define OBJECT_BOX_BOUNDS_STRIDE 5
+#define CULLED_OBJECT_DATA_STRIDE 8
+#define CULLED_OBJECT_BOX_BOUNDS_STRIDE 5
 
-Buffer<float4> ObjectBounds;
-Buffer<float4> ObjectData;
-Buffer<float4> ObjectData2;
-Buffer<float4> ObjectBoxBounds;
+Buffer<float4> CulledObjectBounds;
+Buffer<float4> CulledObjectData;
+Buffer<float4> CulledObjectBoxBounds;
 
 float4 LoadObjectPositionAndRadius(uint ObjectIndex)
 {
-	return ObjectBounds.Load(ObjectIndex);
+	return CulledObjectBounds.Load(ObjectIndex);
 }
 
 float4x4 LoadObjectWorldToVolume(uint ObjectIndex)
 {
-	float4 M0 = ObjectData.Load(ObjectIndex * OBJECT_DATA_STRIDE + 0);
-	float4 M1 = ObjectData.Load(ObjectIndex * OBJECT_DATA_STRIDE + 1);
-	float4 M2 = ObjectData.Load(ObjectIndex * OBJECT_DATA_STRIDE + 2);
-	float4 M3 = ObjectData.Load(ObjectIndex * OBJECT_DATA_STRIDE + 3);
+	float4 M0 = CulledObjectData.Load(ObjectIndex * CULLED_OBJECT_DATA_STRIDE + 0);
+	float4 M1 = CulledObjectData.Load(ObjectIndex * CULLED_OBJECT_DATA_STRIDE + 1);
+	float4 M2 = CulledObjectData.Load(ObjectIndex * CULLED_OBJECT_DATA_STRIDE + 2);
+	float4 M3 = CulledObjectData.Load(ObjectIndex * CULLED_OBJECT_DATA_STRIDE + 3);
 
 	return float4x4(M0, M1, M2, M3);
 }
 
 float3 LoadObjectLocalPositionExtent(uint ObjectIndex)
 {
-	return ObjectData.Load(ObjectIndex * OBJECT_DATA_STRIDE + 4).xyz;
-}
-
-float4 LoadObjectLocalPositionExtentAndTwoSided(uint ObjectIndex)
-{
-	return ObjectData.Load(ObjectIndex * OBJECT_DATA_STRIDE + 4).xyzw;
+	return CulledObjectData.Load(ObjectIndex * CULLED_OBJECT_DATA_STRIDE + 4).xyz;
 }
 
 float4 LoadObjectUVScale(uint ObjectIndex)
 {
-	return ObjectData.Load(ObjectIndex * OBJECT_DATA_STRIDE + 5).xyzw;
+	return CulledObjectData.Load(ObjectIndex * CULLED_OBJECT_DATA_STRIDE + 5).xyzw;
 }
 
 float3 LoadObjectUVAdd(uint ObjectIndex)
 {
-	return ObjectData.Load(ObjectIndex * OBJECT_DATA_STRIDE + 6).xyz;
-}
-
-float4x4 LoadObjectVolumeToWorld(uint ObjectIndex)
-{
-	float4 M0 = ObjectData2.Load(ObjectIndex * OBJECT_DATA2_STRIDE + 0);
-	float4 M1 = ObjectData2.Load(ObjectIndex * OBJECT_DATA2_STRIDE + 1);
-	float4 M2 = ObjectData2.Load(ObjectIndex * OBJECT_DATA2_STRIDE + 2);
-	float4 M3 = ObjectData2.Load(ObjectIndex * OBJECT_DATA2_STRIDE + 3);
-
-	return float4x4(M0, M1, M2, M3);
+	return CulledObjectData.Load(ObjectIndex * CULLED_OBJECT_DATA_STRIDE + 6).xyz;
 }
 
 void LoadObjectViewSpaceBox(uint ObjectIndex, out float3 ObjectViewSpaceMin, out float3 ObjectViewSpaceMax)
 {
-	ObjectViewSpaceMin = ObjectBoxBounds.Load(ObjectIndex * OBJECT_BOX_BOUNDS_STRIDE + 0).xyz;
-	ObjectViewSpaceMax = ObjectBoxBounds.Load(ObjectIndex * OBJECT_BOX_BOUNDS_STRIDE + 1).xyz;
+	ObjectViewSpaceMin = CulledObjectBoxBounds.Load(ObjectIndex * CULLED_OBJECT_BOX_BOUNDS_STRIDE + 0).xyz;
+	ObjectViewSpaceMax = CulledObjectBoxBounds.Load(ObjectIndex * CULLED_OBJECT_BOX_BOUNDS_STRIDE + 1).xyz;
 }
 
 void LoadObjectAxes(uint ObjectIndex, out float3 ObjectAxisX, out float3 ObjectAxisY, out float3 ObjectAxisZ)
 {
-	ObjectAxisX = ObjectBoxBounds.Load(ObjectIndex * OBJECT_BOX_BOUNDS_STRIDE + 2).xyz;
-	ObjectAxisY = ObjectBoxBounds.Load(ObjectIndex * OBJECT_BOX_BOUNDS_STRIDE + 3).xyz;
-	ObjectAxisZ = ObjectBoxBounds.Load(ObjectIndex * OBJECT_BOX_BOUNDS_STRIDE + 4).xyz;
+	ObjectAxisX = CulledObjectBoxBounds.Load(ObjectIndex * CULLED_OBJECT_BOX_BOUNDS_STRIDE + 2).xyz;
+	ObjectAxisY = CulledObjectBoxBounds.Load(ObjectIndex * CULLED_OBJECT_BOX_BOUNDS_STRIDE + 3).xyz;
+	ObjectAxisZ = CulledObjectBoxBounds.Load(ObjectIndex * CULLED_OBJECT_BOX_BOUNDS_STRIDE + 4).xyz;
 }
\ No newline at end of file
diff --git a/Engine/Shaders/DistanceFieldShadowing.usf b/Engine/Shaders/DistanceFieldShadowing.usf
--- a/Engine/Shaders/DistanceFieldShadowing.usf
+++ b/Engine/Shaders/DistanceFieldShadowing.usf
@@ -8,6 +8,153 @@
 #include "DeferredShadingCommon.usf"
 #include "DistanceFieldLightingShared.usf"
 
+RWBuffer<float4> RWCulledObjectBounds;
+RWBuffer<float4> RWCulledObjectData;
+RWBuffer<float4> RWCulledObjectBoxBounds;
+uint ObjectBoundingGeometryIndexCount;
+
+float4 FetchObjectDataFloat4(uint SourceIndex)
+{
+	return float4(ObjectData[4 * SourceIndex + 0], ObjectData[4 * SourceIndex + 1], ObjectData[4 * SourceIndex + 2], ObjectData[4 * SourceIndex + 3]);
+}
+
+float4x4 WorldToShadow;
+
+void CopyCulledObjectData(uint DestIndex, uint SourceIndex)
+{
+	RWCulledObjectBounds[DestIndex] = float4(ObjectBounds[4 * SourceIndex + 0], ObjectBounds[4 * SourceIndex + 1], ObjectBounds[4 * SourceIndex + 2], ObjectBounds[4 * SourceIndex + 3]);
+
+	UNROLL
+	for (uint VectorIndex = 0; VectorIndex < CULLED_OBJECT_DATA_STRIDE; VectorIndex++)
+	{
+		RWCulledObjectData[DestIndex * CULLED_OBJECT_DATA_STRIDE + VectorIndex] = FetchObjectDataFloat4(SourceIndex * OBJECT_DATA_STRIDE + VectorIndex);
+	}
+
+	float3 LocalVolumeBoundsMin = float3(ObjectData[4 * (SourceIndex * OBJECT_DATA_STRIDE + 4) + 3], ObjectData[4 * (SourceIndex * OBJECT_DATA_STRIDE + 6) + 3], ObjectData[4 * (SourceIndex * OBJECT_DATA_STRIDE + 7) + 3]);
+	uint FinalDataVectorIndex = 4 * (SourceIndex * OBJECT_DATA_STRIDE + 7);
+	float3 LocalVolumeBoundsMax = float3(ObjectData[FinalDataVectorIndex + 0], ObjectData[FinalDataVectorIndex + 1], ObjectData[FinalDataVectorIndex + 2]);
+
+	float3 LocalBoundsVertices[8];
+	LocalBoundsVertices[0] = float3(LocalVolumeBoundsMin.x, LocalVolumeBoundsMin.y, LocalVolumeBoundsMin.z);
+	LocalBoundsVertices[1] = float3(LocalVolumeBoundsMax.x, LocalVolumeBoundsMin.y, LocalVolumeBoundsMin.z);
+	LocalBoundsVertices[2] = float3(LocalVolumeBoundsMin.x, LocalVolumeBoundsMax.y, LocalVolumeBoundsMin.z);
+	LocalBoundsVertices[3] = float3(LocalVolumeBoundsMax.x, LocalVolumeBoundsMax.y, LocalVolumeBoundsMin.z);
+	LocalBoundsVertices[4] = float3(LocalVolumeBoundsMin.x, LocalVolumeBoundsMin.y, LocalVolumeBoundsMax.z);
+	LocalBoundsVertices[5] = float3(LocalVolumeBoundsMax.x, LocalVolumeBoundsMin.y, LocalVolumeBoundsMax.z);
+	LocalBoundsVertices[6] = float3(LocalVolumeBoundsMin.x, LocalVolumeBoundsMax.y, LocalVolumeBoundsMax.z);
+	LocalBoundsVertices[7] = float3(LocalVolumeBoundsMax.x, LocalVolumeBoundsMax.y, LocalVolumeBoundsMax.z);
+
+	float3 MinViewSpacePosition = float3(2000000, 2000000, 2000000);
+	float3 MaxViewSpacePosition = float3(-2000000, -2000000, -2000000);
+
+	float4 M0 = FetchObjectDataFloat4(SourceIndex * OBJECT_DATA_STRIDE + 8);
+	float4 M1 = FetchObjectDataFloat4(SourceIndex * OBJECT_DATA_STRIDE + 9);
+	float4 M2 = FetchObjectDataFloat4(SourceIndex * OBJECT_DATA_STRIDE + 10);
+	float4 M3 = FetchObjectDataFloat4(SourceIndex * OBJECT_DATA_STRIDE + 11);
+
+	float4x4 LocalToWorld = float4x4(M0, M1, M2, M3);
+
+	float3 ViewSpaceBoundsVertices[8];
+
+	for (uint i = 0; i < 8; i++)
+	{
+		float3 WorldBoundsPosition = mul(float4(LocalBoundsVertices[i], 1), LocalToWorld).xyz;
+		float3 ViewSpacePosition = mul(float4(WorldBoundsPosition, 1), WorldToShadow).xyz;
+		MinViewSpacePosition = min(MinViewSpacePosition, ViewSpacePosition);
+		MaxViewSpacePosition = max(MaxViewSpacePosition, ViewSpacePosition);
+		ViewSpaceBoundsVertices[i] = ViewSpacePosition;
+	}
+
+	float3 ObjectXAxis = (ViewSpaceBoundsVertices[1] - ViewSpaceBoundsVertices[0]) / 2.0f;
+	float3 ObjectYAxis = (ViewSpaceBoundsVertices[2] - ViewSpaceBoundsVertices[0]) / 2.0f;
+	float3 ObjectZAxis = (ViewSpaceBoundsVertices[4] - ViewSpaceBoundsVertices[0]) / 2.0f;
+
+	RWCulledObjectBoxBounds[DestIndex * CULLED_OBJECT_BOX_BOUNDS_STRIDE + 0] = float4(MinViewSpacePosition, 0);
+	RWCulledObjectBoxBounds[DestIndex * CULLED_OBJECT_BOX_BOUNDS_STRIDE + 1] = float4(MaxViewSpacePosition, 0);
+	RWCulledObjectBoxBounds[DestIndex * CULLED_OBJECT_BOX_BOUNDS_STRIDE + 2] = float4(ObjectXAxis / max(dot(ObjectXAxis, ObjectXAxis), .0001f), 0);
+	RWCulledObjectBoxBounds[DestIndex * CULLED_OBJECT_BOX_BOUNDS_STRIDE + 3] = float4(ObjectYAxis / max(dot(ObjectYAxis, ObjectYAxis), .0001f), 0);
+	RWCulledObjectBoxBounds[DestIndex * CULLED_OBJECT_BOX_BOUNDS_STRIDE + 4] = float4(ObjectZAxis / max(dot(ObjectZAxis, ObjectZAxis), .0001f), 0);
+}
+
+float4 ShadowConvexHull[12];
+float4 ShadowBoundingSphere;
+uint NumShadowHullPlanes;
+
+bool ShadowConvexHullIntersectSphere(float3 SphereOrigin, float SphereRadius)
+{
+	float3 TranslatedSphereOrigin = SphereOrigin + ShadowBoundingSphere.xyz;
+
+	for (uint PlaneIndex = 0; PlaneIndex < NumShadowHullPlanes; PlaneIndex++)
+	{
+		float4 PlaneData = ShadowConvexHull[PlaneIndex];
+		float PlaneDistance = dot(PlaneData.xyz, TranslatedSphereOrigin) - PlaneData.w;
+
+		if (PlaneDistance > SphereRadius)
+		{
+			return false;
+		}
+	}
+	
+	return true;
+}
+
+[numthreads(UPDATEOBJECTS_THREADGROUP_SIZE, 1, 1)]
+void CullObjectsForShadowCS(
+	uint3 GroupId : SV_GroupID,
+	uint3 DispatchThreadId : SV_DispatchThreadID,
+    uint3 GroupThreadId : SV_GroupThreadID) 
+{
+	uint ObjectIndex = DispatchThreadId.x;
+
+#define USE_FRUSTUM_CULLING 1
+#if USE_FRUSTUM_CULLING
+
+	if (DispatchThreadId.x == 0)
+	{
+		// IndexCount, NumInstances, StartIndex, BaseVertexIndex, FirstInstance
+		RWObjectIndirectArguments[0] = ObjectBoundingGeometryIndexCount;
+		RWObjectIndirectArguments[1] = 0;
+		RWObjectIndirectArguments[2] = 0;
+		RWObjectIndirectArguments[3] = 0;
+		RWObjectIndirectArguments[4] = 0;
+	}
+
+	GroupMemoryBarrierWithGroupSync();
+
+	uint SourceIndex = ObjectIndex;
+
+	if (ObjectIndex < NumSceneObjects)
+	{
+		float4 ObjectBoundingSphere = float4(ObjectBounds[4 * SourceIndex + 0], ObjectBounds[4 * SourceIndex + 1], ObjectBounds[4 * SourceIndex + 2], ObjectBounds[4 * SourceIndex + 3]);
+	
+		if (ShadowBoundingSphere.w == 0 && ShadowConvexHullIntersectSphere(ObjectBoundingSphere.xyz, ObjectBoundingSphere.w)
+			|| ShadowBoundingSphere.w > 0 && dot(ShadowBoundingSphere.xyz - ObjectBoundingSphere.xyz, ShadowBoundingSphere.xyz - ObjectBoundingSphere.xyz) < Square(ShadowBoundingSphere.w + ObjectBoundingSphere.w))
+		{
+			uint DestIndex;
+			InterlockedAdd(RWObjectIndirectArguments[1], 1, DestIndex);
+			CopyCulledObjectData(DestIndex, SourceIndex);
+		}
+	}
+
+#else
+
+	if (DispatchThreadId.x == 0)
+	{
+		// IndexCount, NumInstances, StartIndex, BaseVertexIndex, FirstInstance
+		RWObjectIndirectArguments[0] = ObjectBoundingGeometryIndexCount;
+		RWObjectIndirectArguments[1] = NumSceneObjects;
+		RWObjectIndirectArguments[2] = 0;
+		RWObjectIndirectArguments[3] = 0;
+		RWObjectIndirectArguments[4] = 0;
+	}
+
+	uint SourceIndex = ObjectIndex;
+	uint DestIndex = ObjectIndex;
+
+	CopyCulledObjectData(DestIndex, SourceIndex);
+#endif
+}
+
 bool RayHitSphere(float3 RayOrigin, float3 UnitRayDirection, float RayLength, float3 SphereCenter, float SphereRadius)
 {
 	float DistanceAlongRay = dot(SphereCenter - RayOrigin, UnitRayDirection);
@@ -42,7 +189,6 @@ struct FShadowObjectCullVertexOutput
 
 float ConservativeRadiusScale;
 float MinRadius;
-float4x4 WorldToShadow;
 
 /** Used when culling objects into screenspace tile lists */
 void ShadowObjectCullVS(
@@ -218,7 +364,7 @@ void DistanceFieldShadowingCS(
 	//@todo - handle invalid tile depth ranges due to MinDepth
 	float MinDepth = 0;
 
-	uint NumIntersectingObjects = NumObjects;
+	uint NumIntersectingObjects = GetCulledNumObjects();
 	bool bTileShouldComputeShadowing = SceneDepth > MinDepth;
 
 #define USE_CULLING 1
@@ -343,10 +489,12 @@ void DistanceFieldShadowingCS(
 	BRANCH
 	if (bCullTile)
 	{
+		uint NumCulledObjects = GetCulledNumObjects();
+
 		// Compute per-tile lists of affecting objects through bounds culling
 		// Each thread now operates on a sample instead of a pixel
 		LOOP
-		for (uint ObjectIndex = ThreadIndex; ObjectIndex < NumObjects; ObjectIndex += THREADGROUP_TOTALSIZE)
+		for (uint ObjectIndex = ThreadIndex; ObjectIndex < NumCulledObjects; ObjectIndex += THREADGROUP_TOTALSIZE)
 		{
 			float4 SphereCenterAndRadius = LoadObjectPositionAndRadius(ObjectIndex);
 
@@ -500,6 +648,9 @@ void DistanceFieldShadowingCS(
 						MinConeVisibility = 0;
 						MinRayTime = min(MinRayTime, SampleRayTime * UVScaleAndVolumeScale.w);
 					}
+
+					// Force to shadowed as we approach max steps
+					MinConeVisibility = min(MinConeVisibility, (1 - StepIndex / (float)MaxSteps));
 				}
 			}
 
@@ -509,7 +660,7 @@ void DistanceFieldShadowingCS(
 
 #else
 	//Result = .0001f * NumIntersectingObjects;
-	Result = NumIntersectingObjects / 30.0f;
+	Result = GetCulledNumObjects() / 20.0f;
 #endif
 
 	RWShadowFactors[DispatchThreadId.xy] = float2(Result, SceneDepth);
diff --git a/Engine/Shaders/DistanceFieldSurfaceCacheLighting.usf b/Engine/Shaders/DistanceFieldSurfaceCacheLighting.usf
--- a/Engine/Shaders/DistanceFieldSurfaceCacheLighting.usf
+++ b/Engine/Shaders/DistanceFieldSurfaceCacheLighting.usf
@@ -102,94 +102,6 @@ void ObjectCullPS(
 	}
 }
 
-void ComputeDistanceFieldNormal(uint2 TileCoordinate, float2 ScreenUV, float SceneDepth, uint HasDistanceFieldRepresentation, inout float3 WorldNormal)
-{
-	uint4 TileHead = GetTileHead(TileCoordinate);
-
-	float2 ScreenPosition = (ScreenUV.xy - View.ScreenPositionScaleBias.wz) / View.ScreenPositionScaleBias.xy;
-
-	float4 HomogeneousWorldPosition = mul(float4(ScreenPosition * SceneDepth, SceneDepth, 1), View.ScreenToWorld);
-	float3 OpaqueWorldPosition = HomogeneousWorldPosition.xyz / HomogeneousWorldPosition.w;
-
-	int ClosestDistanceFieldObject = -1;
-	float ClosestDistanceFieldSurface = 10000;
-	
-	LOOP 
-	// Iterate over objects whose surface we might be on
-	for (uint ListObjectIndex = 0; ListObjectIndex < TileHead.y; ListObjectIndex += 1)
-	{
-#if COHERENT_OBJECT_LIST_WRITES
-		uint ListIndex = 0;
-		uint ObjectIndex = TileArrayData.Load((ListObjectIndex * TileListGroupSize.x * TileListGroupSize.y + TileHead.x) * NUM_CULLED_OBJECT_LISTS + ListIndex);
-#else
-		uint ObjectIndex = TileArrayData.Load(TileHead.x + ListObjectIndex);
-#endif
-		float3 LocalPositionExtent = LoadObjectLocalPositionExtent(ObjectIndex);
-		float4x4 WorldToVolume = LoadObjectWorldToVolume(ObjectIndex);
-		float4 UVScaleAndVolumeScale = LoadObjectUVScale(ObjectIndex);
-		float3 UVAdd = LoadObjectUVAdd(ObjectIndex);
-		float3 VolumeShadingPosition = mul(float4(OpaqueWorldPosition, 1), WorldToVolume).xyz;
-
-		BRANCH
-		if (all(VolumeShadingPosition > -LocalPositionExtent) && all(VolumeShadingPosition < LocalPositionExtent))
-		{
-			float3 LocalShadingUV = DistanceFieldVolumePositionToUV(VolumeShadingPosition, UVScaleAndVolumeScale.xyz, UVAdd);
-			float DistanceFromSurface = Texture3DSampleLevel(DistanceFieldTexture, DistanceFieldSampler, LocalShadingUV, 0).x;
-			float WorldSpaceDistanceFromSurface = DistanceFromSurface * UVScaleAndVolumeScale.w;
-
-			if (WorldSpaceDistanceFromSurface < ClosestDistanceFieldSurface)
-			{
-				ClosestDistanceFieldSurface = WorldSpaceDistanceFromSurface;
-				ClosestDistanceFieldObject = ObjectIndex;
-			}
-		}
-	}
-
-	BRANCH
-	if (HasDistanceFieldRepresentation > 0 && ClosestDistanceFieldObject >= 0)
-	{
-		uint ObjectIndex = ClosestDistanceFieldObject;
-
-		float4 LocalPositionExtentAndTwoSided = LoadObjectLocalPositionExtentAndTwoSided(ObjectIndex);
-
-		// Can't compute a distance field normal if all faces were treated as frontfaces, 
-		// Because that creates a non-linearity in the distance field at surface intersections, which does not interpolate properly
-		// Fall back to GBuffer normal, which may have normalmap information or just bogus imported normals that cause surface cache artifacts
-		BRANCH
-		if (LocalPositionExtentAndTwoSided.w == 0)
-		{
-			float4x4 WorldToVolume = LoadObjectWorldToVolume(ObjectIndex);
-			float4 UVScaleAndVolumeScale = LoadObjectUVScale(ObjectIndex);
-			float3 UVAdd = LoadObjectUVAdd(ObjectIndex);
-
-			float4x4 VolumeToWorld = LoadObjectVolumeToWorld(ObjectIndex);
-
-			float3 VolumeShadingPosition = mul(float4(OpaqueWorldPosition, 1), WorldToVolume).xyz;
-			float3 LocalShadingUV = DistanceFieldVolumePositionToUV(VolumeShadingPosition, UVScaleAndVolumeScale.xyz, UVAdd);
-
-			// Used to clamp UVs inside valid space of this object's distance field
-			float3 UVMin = DistanceFieldVolumePositionToUV(-LocalPositionExtentAndTwoSided.xyz, UVScaleAndVolumeScale.xyz, UVAdd);
-			float3 UVMax = DistanceFieldVolumePositionToUV(LocalPositionExtentAndTwoSided.xyz, UVScaleAndVolumeScale.xyz, UVAdd);
-		
-			float R = Texture3DSampleLevel(DistanceFieldTexture, DistanceFieldSampler, float3(min(LocalShadingUV.x + DistanceFieldAtlasTexelSize.x, UVMax.x), LocalShadingUV.y, LocalShadingUV.z), 0).x;
-			float L = Texture3DSampleLevel(DistanceFieldTexture, DistanceFieldSampler, float3(max(LocalShadingUV.x - DistanceFieldAtlasTexelSize.x, UVMin.x), LocalShadingUV.y, LocalShadingUV.z), 0).x;
-			float F = Texture3DSampleLevel(DistanceFieldTexture, DistanceFieldSampler, float3(LocalShadingUV.x, min(LocalShadingUV.y + DistanceFieldAtlasTexelSize.y, UVMax.y), LocalShadingUV.z), 0).x;
-			float B = Texture3DSampleLevel(DistanceFieldTexture, DistanceFieldSampler, float3(LocalShadingUV.x, max(LocalShadingUV.y - DistanceFieldAtlasTexelSize.y, UVMin.y), LocalShadingUV.z), 0).x;
-			float U = Texture3DSampleLevel(DistanceFieldTexture, DistanceFieldSampler, float3(LocalShadingUV.x, LocalShadingUV.y, min(LocalShadingUV.z + DistanceFieldAtlasTexelSize.z, UVMax.z)), 0).x;
-			float D = Texture3DSampleLevel(DistanceFieldTexture, DistanceFieldSampler, float3(LocalShadingUV.x, LocalShadingUV.y, max(LocalShadingUV.z - DistanceFieldAtlasTexelSize.z, UVMin.z)), 0).x;
-		
-			float3 Gradient = .5f * float3(R - L, F - B, U - D);
-			float3 LocalNormal = normalize(Gradient);
-			WorldNormal = normalize(mul(LocalNormal, (float3x3)VolumeToWorld));
-			float DistanceAlpha = saturate(abs(ClosestDistanceFieldSurface) / 2);
-			//WorldNormal = normalize(lerp(WorldNormal, GBufferData.WorldNormal, .5f * DistanceAlpha));
-		}
-	}
-}
-
-// Using the distance field normal prevents incorrect self-intersection but adds seams between modular objects and has a high cost
-#define USE_DISTANCE_FIELD_NORMAL 0
-
 /** Computes the distance field normal, using a search through all the nearby objects to find the closest one, whose normal is used. */
 void ComputeDistanceFieldNormalPS(
 	in float4 UVAndScreenPos : TEXCOORD0, 
@@ -202,14 +114,6 @@ void ComputeDistanceFieldNormalPS(
 	FGBufferData GBufferData = GetGBufferData(ScreenUV);
 
 	float3 WorldNormal = GBufferData.WorldNormal;
-
-#if USE_DISTANCE_FIELD_NORMAL
-
-	uint2 TileCoordinate = (uint2)(SVPos) / uint2(THREADGROUP_SIZEX, THREADGROUP_SIZEY);
-	ComputeDistanceFieldNormal(TileCoordinate, ScreenUV, SceneDepth, GBufferData.HasDistanceFieldRepresentation, WorldNormal);
-
-#endif
-
 	float DepthSign = GBufferData.HasDistanceFieldRepresentation > 0 ? 1 : -1;
 	// Note: Encoding whether the pixel has a distance field representation in the sign bit of the alpha
 	OutColor = float4(EncodeNormalForAO(WorldNormal), DepthSign * SceneDepth);
@@ -228,14 +132,6 @@ void ComputeDistanceFieldNormalCS(
 	FGBufferData GBufferData = GetGBufferData(ScreenUV);
 
 	float3 WorldNormal = GBufferData.WorldNormal;
-
-#if USE_DISTANCE_FIELD_NORMAL
-
-	uint2 TileCoordinate = DispatchThreadId.xy / uint2(THREADGROUP_SIZEX, THREADGROUP_SIZEY);
-	ComputeDistanceFieldNormal(TileCoordinate, ScreenUV, SceneDepth, GBufferData.HasDistanceFieldRepresentation, WorldNormal);
-
-#endif
-	
 	float DepthSign = GBufferData.HasDistanceFieldRepresentation > 0 ? 1 : -1;
 	float4 OutValue = float4(EncodeNormalForAO(WorldNormal), DepthSign * SceneDepth);
 	RWDistanceFieldNormal[DispatchThreadId.xy] = OutValue;
@@ -1087,7 +983,7 @@ bool RayHitSphere( float3 RayOrigin, float3 UnitRayDirection, float3 SphereCente
 	return dot( CenterToRay, CenterToRay ) <= Square( SphereRadius );
 }
 
-#define MAX_INTERSECTING_OBJECTS 256
+#define MAX_INTERSECTING_OBJECTS 512
 groupshared uint IntersectingObjectIndices[MAX_INTERSECTING_OBJECTS];
 
 groupshared uint NumIntersectingObjects;
@@ -1188,8 +1084,10 @@ void VisualizeMeshDistanceFieldCS(
 
 	GroupMemoryBarrierWithGroupSync();
 
+	uint NumCulledObjects = GetCulledNumObjects();
+
 	LOOP
-	for (uint ObjectIndex = ThreadIndex; ObjectIndex < NumObjects; ObjectIndex += THREADGROUP_TOTALSIZE)
+	for (uint ObjectIndex = ThreadIndex; ObjectIndex < NumCulledObjects; ObjectIndex += THREADGROUP_TOTALSIZE)
 	{
 		float4 SphereCenterAndRadius = LoadObjectPositionAndRadius(ObjectIndex);
 
diff --git a/Engine/Shaders/DistanceFieldSurfaceCacheLightingCompute.usf b/Engine/Shaders/DistanceFieldSurfaceCacheLightingCompute.usf
--- a/Engine/Shaders/DistanceFieldSurfaceCacheLightingCompute.usf
+++ b/Engine/Shaders/DistanceFieldSurfaceCacheLightingCompute.usf
@@ -11,6 +11,260 @@
 #include "DistanceFieldAOShared.usf"
 #include "MonteCarlo.usf"
 
+uint NumUploadOperations;
+Buffer<uint> UploadOperationIndices;
+Buffer<float4> UploadOperationData;
+
+// In float4's.  Must match equivalent C++ variables.
+#define UPLOAD_DATA_STRIDE (1 + OBJECT_DATA_STRIDE)
+
+void UploadDataFloat4(uint DestIndex, uint UploadIndex)
+{
+	float4 UploadVector = UploadOperationData[UploadIndex];
+	RWObjectData[4 * DestIndex + 0] = UploadVector.x;
+	RWObjectData[4 * DestIndex + 1] = UploadVector.y;
+	RWObjectData[4 * DestIndex + 2] = UploadVector.z;
+	RWObjectData[4 * DestIndex + 3] = UploadVector.w;
+}
+
+[numthreads(UPDATEOBJECTS_THREADGROUP_SIZE, 1, 1)]
+void UploadObjectsToBufferCS(
+	uint3 GroupId : SV_GroupID,
+	uint3 DispatchThreadId : SV_DispatchThreadID,
+    uint3 GroupThreadId : SV_GroupThreadID) 
+{
+	uint UploadOperationIndex = DispatchThreadId.x;
+
+	if (DispatchThreadId.x < NumUploadOperations)
+	{
+		uint DestIndex = UploadOperationIndices[UploadOperationIndex];
+		float4 UploadBounds = UploadOperationData[UploadOperationIndex * UPLOAD_DATA_STRIDE + 0];
+		RWObjectBounds[4 * DestIndex + 0] = UploadBounds.x;
+		RWObjectBounds[4 * DestIndex + 1] = UploadBounds.y;
+		RWObjectBounds[4 * DestIndex + 2] = UploadBounds.z;
+		RWObjectBounds[4 * DestIndex + 3] = UploadBounds.w;
+
+		UNROLL
+		for (uint VectorIndex = 0; VectorIndex < OBJECT_DATA_STRIDE; VectorIndex++)
+		{
+			UploadDataFloat4(DestIndex * OBJECT_DATA_STRIDE + VectorIndex, UploadOperationIndex * UPLOAD_DATA_STRIDE + VectorIndex + 1);
+		}
+	}
+}
+
+RWBuffer<float> RWCopyObjectBounds;
+RWBuffer<float> RWCopyObjectData;
+
+void CopyDataFloat4(uint DestIndex, uint SourceIndex)
+{
+	RWCopyObjectData[4 * DestIndex + 0] = ObjectData[4 * SourceIndex + 0];
+	RWCopyObjectData[4 * DestIndex + 1] = ObjectData[4 * SourceIndex + 1];
+	RWCopyObjectData[4 * DestIndex + 2] = ObjectData[4 * SourceIndex + 2];
+	RWCopyObjectData[4 * DestIndex + 3] = ObjectData[4 * SourceIndex + 3];
+}
+
+[numthreads(UPDATEOBJECTS_THREADGROUP_SIZE, 1, 1)]
+void CopyObjectBufferCS(
+	uint3 GroupId : SV_GroupID,
+	uint3 DispatchThreadId : SV_DispatchThreadID,
+    uint3 GroupThreadId : SV_GroupThreadID) 
+{
+	uint CopyOperationIndex = DispatchThreadId.x;
+
+	if (CopyOperationIndex < NumSceneObjects)
+	{
+		uint SourceIndex = CopyOperationIndex;
+		uint DestIndex = CopyOperationIndex;
+		RWCopyObjectBounds[4 * DestIndex + 0] = ObjectBounds[4 * SourceIndex + 0];
+		RWCopyObjectBounds[4 * DestIndex + 1] = ObjectBounds[4 * SourceIndex + 1];
+		RWCopyObjectBounds[4 * DestIndex + 2] = ObjectBounds[4 * SourceIndex + 2];
+		RWCopyObjectBounds[4 * DestIndex + 3] = ObjectBounds[4 * SourceIndex + 3];
+
+		UNROLL
+		for (uint VectorIndex = 0; VectorIndex < OBJECT_DATA_STRIDE; VectorIndex++)
+		{
+			CopyDataFloat4(DestIndex * OBJECT_DATA_STRIDE + VectorIndex, SourceIndex * OBJECT_DATA_STRIDE + VectorIndex);
+		}
+	}
+}
+
+uint NumRemoveOperations;
+Buffer<uint4> RemoveOperationIndices;
+
+#if REMOVE_FROM_SAME_BUFFER
+	#define RWBoundsRemoveSource RWObjectBounds
+	#define RWDataRemoveSource RWObjectData
+#else
+
+	Buffer<float> ObjectBounds2;
+	Buffer<float> ObjectData2;
+
+	#define RWBoundsRemoveSource ObjectBounds2
+	#define RWDataRemoveSource ObjectData2
+#endif
+
+void WriteDataFloat4(uint DestIndex, uint SourceIndex)
+{
+	RWObjectData[4 * DestIndex + 0] = RWDataRemoveSource[4 * SourceIndex + 0];
+	RWObjectData[4 * DestIndex + 1] = RWDataRemoveSource[4 * SourceIndex + 1];
+	RWObjectData[4 * DestIndex + 2] = RWDataRemoveSource[4 * SourceIndex + 2];
+	RWObjectData[4 * DestIndex + 3] = RWDataRemoveSource[4 * SourceIndex + 3];
+}
+
+[numthreads(UPDATEOBJECTS_THREADGROUP_SIZE, 1, 1)]
+void RemoveObjectsFromBufferCS(
+	uint3 GroupId : SV_GroupID,
+	uint3 DispatchThreadId : SV_DispatchThreadID,
+    uint3 GroupThreadId : SV_GroupThreadID) 
+{
+	uint RemoveOperationIndex = DispatchThreadId.x;
+
+	if (RemoveOperationIndex < NumRemoveOperations)
+	{
+		// RemoveAtSwap
+		uint SourceIndex = RemoveOperationIndices[RemoveOperationIndex].y;
+		uint DestIndex = RemoveOperationIndices[RemoveOperationIndex].x;
+
+		RWObjectBounds[4 * DestIndex + 0] = RWBoundsRemoveSource[4 * SourceIndex + 0];
+		RWObjectBounds[4 * DestIndex + 1] = RWBoundsRemoveSource[4 * SourceIndex + 1];
+		RWObjectBounds[4 * DestIndex + 2] = RWBoundsRemoveSource[4 * SourceIndex + 2];
+		RWObjectBounds[4 * DestIndex + 3] = RWBoundsRemoveSource[4 * SourceIndex + 3];
+
+		UNROLL
+		for (uint VectorIndex = 0; VectorIndex < OBJECT_DATA_STRIDE; VectorIndex++)
+		{
+			WriteDataFloat4(DestIndex * OBJECT_DATA_STRIDE + VectorIndex, SourceIndex * OBJECT_DATA_STRIDE + VectorIndex);
+		}
+	}
+}
+
+RWBuffer<float4> RWCulledObjectBounds;
+RWBuffer<float4> RWCulledObjectData;
+RWBuffer<float4> RWCulledObjectBoxBounds;
+float4 ViewFrustumConvexHull[5];
+uint ObjectBoundingGeometryIndexCount;
+
+bool ViewFrustumIntersectSphere(float3 SphereOrigin, float SphereRadius)
+{
+	UNROLL
+	for (uint PlaneIndex = 0; PlaneIndex < 5; PlaneIndex++)
+	{
+		float4 PlaneData = ViewFrustumConvexHull[PlaneIndex];
+		float PlaneDistance = dot(PlaneData.xyz, SphereOrigin) - PlaneData.w;
+
+		if (PlaneDistance > SphereRadius)
+		{
+			return false;
+		}
+	}
+	
+	return true;
+}
+
+groupshared uint NumGroupObjects;
+groupshared uint GroupBaseIndex;
+groupshared uint GroupObjectIndices[UPDATEOBJECTS_THREADGROUP_SIZE];
+
+float4 FetchObjectDataFloat4(uint SourceIndex)
+{
+	return float4(ObjectData[4 * SourceIndex + 0], ObjectData[4 * SourceIndex + 1], ObjectData[4 * SourceIndex + 2], ObjectData[4 * SourceIndex + 3]);
+}
+
+void CopyCulledObjectData(uint DestIndex, uint SourceIndex)
+{
+	RWCulledObjectBounds[DestIndex] = float4(ObjectBounds[4 * SourceIndex + 0], ObjectBounds[4 * SourceIndex + 1], ObjectBounds[4 * SourceIndex + 2], ObjectBounds[4 * SourceIndex + 3]);
+
+	UNROLL
+	for (uint VectorIndex = 0; VectorIndex < CULLED_OBJECT_DATA_STRIDE; VectorIndex++)
+	{
+		RWCulledObjectData[DestIndex * CULLED_OBJECT_DATA_STRIDE + VectorIndex] = FetchObjectDataFloat4(SourceIndex * OBJECT_DATA_STRIDE + VectorIndex);
+	}
+}
+
+[numthreads(UPDATEOBJECTS_THREADGROUP_SIZE, 1, 1)]
+void CullObjectsForViewCS(
+	uint3 GroupId : SV_GroupID,
+	uint3 DispatchThreadId : SV_DispatchThreadID,
+    uint3 GroupThreadId : SV_GroupThreadID) 
+{
+	uint ObjectIndex = DispatchThreadId.x;
+
+#define USE_FRUSTUM_CULLING 1
+#if USE_FRUSTUM_CULLING
+	if (DispatchThreadId.x == 0)
+	{
+		// IndexCount, NumInstances, StartIndex, BaseVertexIndex, FirstInstance
+		RWObjectIndirectArguments[0] = ObjectBoundingGeometryIndexCount;
+		RWObjectIndirectArguments[1] = 0;
+		RWObjectIndirectArguments[2] = 0;
+		RWObjectIndirectArguments[3] = 0;
+		RWObjectIndirectArguments[4] = 0;
+	}
+
+	if (GroupThreadId.x == 0)
+	{
+		NumGroupObjects = 0;
+	}
+
+	GroupMemoryBarrierWithGroupSync();
+
+	if (ObjectIndex < NumSceneObjects)
+	{
+		uint SourceIndex = ObjectIndex;
+
+		float4 ObjectBoundingSphere = float4(ObjectBounds[4 * SourceIndex + 0], ObjectBounds[4 * SourceIndex + 1], ObjectBounds[4 * SourceIndex + 2], ObjectBounds[4 * SourceIndex + 3]);
+		float DistanceToViewSq = dot(View.ViewOrigin.xyz - ObjectBoundingSphere.xyz, View.ViewOrigin.xyz - ObjectBoundingSphere.xyz);
+
+		if (DistanceToViewSq < Square(AOMaxViewDistance + ObjectBoundingSphere.w)
+			&& ViewFrustumIntersectSphere(ObjectBoundingSphere.xyz, ObjectBoundingSphere.w + AOMaxDistance))
+		{
+			uint DestIndex;
+			InterlockedAdd(NumGroupObjects, 1, DestIndex);
+			GroupObjectIndices[DestIndex] = SourceIndex;
+		}
+	}
+
+	GroupMemoryBarrierWithGroupSync();
+
+	if (GroupThreadId.x == 0)
+	{
+		InterlockedAdd(RWObjectIndirectArguments[1], NumGroupObjects, GroupBaseIndex);
+	}
+
+	GroupMemoryBarrierWithGroupSync();
+
+	if (GroupThreadId.x < NumGroupObjects)
+	{
+		uint SourceIndex = GroupObjectIndices[GroupThreadId.x];
+		uint DestIndex = GroupBaseIndex + GroupThreadId.x;
+		CopyCulledObjectData(DestIndex, SourceIndex);
+	}
+
+#else
+
+	if (DispatchThreadId.x == 0)
+	{
+		// IndexCount, NumInstances, StartIndex, BaseVertexIndex, FirstInstance
+		RWObjectIndirectArguments[0] = ObjectBoundingGeometryIndexCount;
+		RWObjectIndirectArguments[1] = NumSceneObjects;
+		RWObjectIndirectArguments[2] = 0;
+		RWObjectIndirectArguments[3] = 0;
+		RWObjectIndirectArguments[4] = 0;
+	}
+
+	GroupMemoryBarrierWithGroupSync();
+
+	if (ObjectIndex < NumSceneObjects)
+	{
+		uint SourceIndex = ObjectIndex;
+		uint DestIndex = ObjectIndex;
+
+		CopyCulledObjectData(DestIndex, SourceIndex);
+	}
+
+#endif
+}
+
 /** Min and Max depth for this tile. */
 groupshared uint IntegerTileMinZ;
 groupshared uint IntegerTileMaxZ;
@@ -265,10 +519,12 @@ void DistanceFieldAOBuildTileListMain(
 		LargeGroupMaxSampleRadius = GetStepOffset(EndIndex) * 2 * RadiusScale;
 	}
 
+	uint NumCulledObjects = GetCulledNumObjects();
+
 	// Compute per-tile lists of affecting objects through bounds culling
 	// Each thread now operates on a sample instead of a pixel
 	LOOP
-	for (uint ObjectIndex = ThreadIndex; ObjectIndex < NumObjects; ObjectIndex += THREADGROUP_TOTALSIZE)
+	for (uint ObjectIndex = ThreadIndex; ObjectIndex < NumCulledObjects; ObjectIndex += THREADGROUP_TOTALSIZE)
 	{
 		float4 SphereCenterAndRadius = LoadObjectPositionAndRadius(ObjectIndex);
 		float3 ViewSpaceSphereCenter = mul(float4(SphereCenterAndRadius.xyz + View.PreViewTranslation.xyz, 1), View.TranslatedWorldToView).xyz;
diff --git a/Engine/Source/Runtime/Engine/Classes/Components/PrimitiveComponent.h b/Engine/Source/Runtime/Engine/Classes/Components/PrimitiveComponent.h
--- a/Engine/Source/Runtime/Engine/Classes/Components/PrimitiveComponent.h
+++ b/Engine/Source/Runtime/Engine/Classes/Components/PrimitiveComponent.h
@@ -275,6 +275,10 @@ public:
 	UPROPERTY(EditAnywhere, BlueprintReadOnly, Category=Lighting, AdvancedDisplay, meta=(EditCondition="CastShadow"))
 	uint32 bAffectDynamicIndirectLighting:1;
 
+	/** Controls whether the primitive should affect dynamic distance field lighting methods.  This flag is only used if CastShadow is true. **/
+	UPROPERTY(EditAnywhere, BlueprintReadOnly, Category=Lighting, AdvancedDisplay, meta=(EditCondition="CastShadow"))
+	uint32 bAffectDistanceFieldLighting:1;
+
 	/** Controls whether the primitive should cast shadows in the case of non precomputed shadowing.  This flag is only used if CastShadow is true. **/
 	UPROPERTY(EditAnywhere, BlueprintReadOnly, Category=Lighting, AdvancedDisplay, meta=(EditCondition="CastShadow"))
 	uint32 bCastDynamicShadow:1;
diff --git a/Engine/Source/Runtime/Engine/Classes/Foliage/FoliageType.h b/Engine/Source/Runtime/Engine/Classes/Foliage/FoliageType.h
--- a/Engine/Source/Runtime/Engine/Classes/Foliage/FoliageType.h
+++ b/Engine/Source/Runtime/Engine/Classes/Foliage/FoliageType.h
@@ -198,6 +198,10 @@ class UFoliageType : public UObject
 	UPROPERTY(EditAnywhere, BlueprintReadOnly, Category=Lighting)
 	uint32 bAffectDynamicIndirectLighting:1;
 
+	/** Controls whether the primitive should affect dynamic distance field lighting methods.  This flag is only used if CastShadow is true. **/
+	UPROPERTY(EditAnywhere, BlueprintReadOnly, Category=Lighting, AdvancedDisplay, meta=(EditCondition="CastShadow"))
+	uint32 bAffectDistanceFieldLighting:1;
+
 	/** Controls whether the foliage should cast shadows in the case of non precomputed shadowing.  This flag is only used if CastShadow is true. **/
 	UPROPERTY(EditAnywhere, BlueprintReadOnly, Category=Lighting)
 	uint32 bCastDynamicShadow:1;
diff --git a/Engine/Source/Runtime/Engine/Private/DistanceFieldAtlas.cpp b/Engine/Source/Runtime/Engine/Private/DistanceFieldAtlas.cpp
--- a/Engine/Source/Runtime/Engine/Private/DistanceFieldAtlas.cpp
+++ b/Engine/Source/Runtime/Engine/Private/DistanceFieldAtlas.cpp
@@ -15,8 +15,10 @@
 
 TGlobalResource<FDistanceFieldVolumeTextureAtlas> GDistanceFieldVolumeTextureAtlas = TGlobalResource<FDistanceFieldVolumeTextureAtlas>(PF_R16F);
 
+const int32 MaxAtlasDimension = 512;
+
 FDistanceFieldVolumeTextureAtlas::FDistanceFieldVolumeTextureAtlas(EPixelFormat InFormat) :
-	BlockAllocator(0, 0, 0, 512, 512, 512, false, false)
+	BlockAllocator(0, 0, 0, MaxAtlasDimension, MaxAtlasDimension, MaxAtlasDimension, false, false)
 {
 	Format = InFormat;
 }
@@ -39,11 +41,23 @@ void FDistanceFieldVolumeTextureAtlas::RemoveAllocation(FDistanceFieldVolumeText
 	}
 }
 
-void FDistanceFieldVolumeTextureAtlas::UpdateAllocations()
+struct FCompareVolumeAllocation
 {
+	FORCEINLINE bool operator()(const FDistanceFieldVolumeTexture& A, const FDistanceFieldVolumeTexture& B) const
+	{
+		return A.GetAllocationVolume() > B.GetAllocationVolume();
+	}
+};
+
+bool FDistanceFieldVolumeTextureAtlas::UpdateAllocations()
+{
+	bool bRecomputedLayouts = false;
+
 	if (PendingAllocations.Num() > 0)
 	{
-		//@todo - sort largest to smallest for best packing
+		// Sort largest to smallest for best packing
+		PendingAllocations.Sort(FCompareVolumeAllocation());
+
 		for (int32 AllocationIndex = 0; AllocationIndex < PendingAllocations.Num(); AllocationIndex++)
 		{
 			FDistanceFieldVolumeTexture* Texture = PendingAllocations[AllocationIndex];
@@ -62,6 +76,35 @@ void FDistanceFieldVolumeTextureAtlas::UpdateAllocations()
 			|| BlockAllocator.GetSizeY() > VolumeTextureRHI->GetSizeY()
 			|| BlockAllocator.GetSizeZ() > VolumeTextureRHI->GetSizeZ())
 		{
+			if (CurrentAllocations.Num() > 0)
+			{
+				// Remove all allocations from the layout so we have a clean slate
+				BlockAllocator = FTextureLayout3d(0, 0, 0, MaxAtlasDimension, MaxAtlasDimension, MaxAtlasDimension, false, false);
+				
+				bRecomputedLayouts = true;
+
+				// Re-upload all textures since we had to reallocate
+				PendingAllocations.Append(CurrentAllocations);
+				CurrentAllocations.Empty();
+
+				// Sort largest to smallest for best packing
+				PendingAllocations.Sort(FCompareVolumeAllocation());
+
+				// Add all allocations back to the layout
+				for (int32 AllocationIndex = 0; AllocationIndex < PendingAllocations.Num(); AllocationIndex++)
+				{
+					FDistanceFieldVolumeTexture* Texture = PendingAllocations[AllocationIndex];
+					const FIntVector Size = Texture->VolumeData.Size;
+
+					if (!BlockAllocator.AddElement((uint32&)Texture->AtlasAllocationMin.X, (uint32&)Texture->AtlasAllocationMin.Y, (uint32&)Texture->AtlasAllocationMin.Z, Size.X, Size.Y, Size.Z))
+					{
+						UE_LOG(LogStaticMesh,Error,TEXT("Failed to allocate %ux%ux%u in distance field atlas"), Size.X, Size.Y, Size.Z);
+						PendingAllocations.RemoveAt(AllocationIndex);
+						AllocationIndex--;
+					}
+				}
+			}
+
 			FRHIResourceCreateInfo CreateInfo;
 
 			VolumeTextureRHI = RHICreateTexture3D(
@@ -76,10 +119,6 @@ void FDistanceFieldVolumeTextureAtlas::UpdateAllocations()
 			const int32 FormatSize = GPixelFormats[Format].BlockBytes;
 			float MemorySize = VolumeTextureRHI->GetSizeX() * VolumeTextureRHI->GetSizeY() * VolumeTextureRHI->GetSizeZ() * FormatSize / 1024.0f / 1024.0f;
 			UE_LOG(LogStaticMesh,Log,TEXT("Allocated %ux%ux%u distance field atlas = %.1fMb"), VolumeTextureRHI->GetSizeX(), VolumeTextureRHI->GetSizeY(), VolumeTextureRHI->GetSizeZ(), MemorySize);
-
-			// Re-upload all textures since we had to reallocate
-			PendingAllocations.Append(CurrentAllocations);
-			CurrentAllocations.Empty();
 		}
 
 		for (int32 AllocationIndex = 0; AllocationIndex < PendingAllocations.Num(); AllocationIndex++)
@@ -107,6 +146,8 @@ void FDistanceFieldVolumeTextureAtlas::UpdateAllocations()
 		CurrentAllocations.Append(PendingAllocations);
 		PendingAllocations.Empty();
 	}
+
+	return bRecomputedLayouts;
 }
 
 void FDistanceFieldVolumeTexture::Initialize()
diff --git a/Engine/Source/Runtime/Engine/Private/PrimitiveComponent.cpp b/Engine/Source/Runtime/Engine/Private/PrimitiveComponent.cpp
--- a/Engine/Source/Runtime/Engine/Private/PrimitiveComponent.cpp
+++ b/Engine/Source/Runtime/Engine/Private/PrimitiveComponent.cpp
@@ -85,6 +85,7 @@ UPrimitiveComponent::UPrimitiveComponent(const FObjectInitializer& ObjectInitial
 	CastShadow = false;
 	bCastDynamicShadow = true;
 	bAffectDynamicIndirectLighting = true;
+	bAffectDistanceFieldLighting = true;
 	LpvBiasMultiplier = 1.0f;
 	bCastStaticShadow = true;
 	bCastVolumetricTranslucentShadow = false;
diff --git a/Engine/Source/Runtime/Engine/Private/PrimitiveSceneProxy.cpp b/Engine/Source/Runtime/Engine/Private/PrimitiveSceneProxy.cpp
--- a/Engine/Source/Runtime/Engine/Private/PrimitiveSceneProxy.cpp
+++ b/Engine/Source/Runtime/Engine/Private/PrimitiveSceneProxy.cpp
@@ -38,6 +38,7 @@ FPrimitiveSceneProxy::FPrimitiveSceneProxy(const UPrimitiveComponent* InComponen
 	// Disable dynamic shadow casting if the primitive only casts indirect shadows, since dynamic shadows are always shadowing direct lighting
 ,	bCastDynamicShadow(InComponent->bCastDynamicShadow && InComponent->CastShadow && !InComponent->GetShadowIndirectOnly())
 ,   bAffectDynamicIndirectLighting(InComponent->bAffectDynamicIndirectLighting)
+,   bAffectDistanceFieldLighting(InComponent->bAffectDistanceFieldLighting)
 ,	bCastStaticShadow(InComponent->CastShadow && InComponent->bCastStaticShadow)
 ,	bCastVolumetricTranslucentShadow(InComponent->bCastDynamicShadow && InComponent->CastShadow && InComponent->bCastVolumetricTranslucentShadow)
 ,	bCastHiddenShadow(InComponent->bCastHiddenShadow)
@@ -48,6 +49,7 @@ FPrimitiveSceneProxy::FPrimitiveSceneProxy(const UPrimitiveComponent* InComponen
 ,	bStaticElementsAlwaysUseProxyPrimitiveUniformBuffer(false)
 ,	bAlwaysHasVelocity(false)
 ,	bUseEditorDepthTest(true)
+,	bSupportsDistanceFieldRepresentation(false)
 ,	bUseAsOccluder(InComponent->bUseAsOccluder)
 ,	bAllowApproximateOcclusion(InComponent->Mobility != EComponentMobility::Movable)
 ,	bSelectable(InComponent->bSelectable)
diff --git a/Engine/Source/Runtime/Engine/Private/StaticMeshRender.cpp b/Engine/Source/Runtime/Engine/Private/StaticMeshRender.cpp
--- a/Engine/Source/Runtime/Engine/Private/StaticMeshRender.cpp
+++ b/Engine/Source/Runtime/Engine/Private/StaticMeshRender.cpp
@@ -66,6 +66,7 @@ FStaticMeshSceneProxy::FStaticMeshSceneProxy(UStaticMeshComponent* InComponent):
 	WireframeColor = InComponent->GetWireframeColor();
 	LevelColor = FLinearColor(1,1,1);
 	PropertyColor = FLinearColor(1,1,1);
+	bSupportsDistanceFieldRepresentation = true;
 
 	const auto FeatureLevel = GetScene().GetFeatureLevel();
 
@@ -1235,7 +1236,7 @@ void FStaticMeshSceneProxy::GetLightRelevance(const FLightSceneProxy* LightScene
 	}
 }
 
-void FStaticMeshSceneProxy::GetDistancefieldAtlasData(FBox& LocalVolumeBounds, FIntVector& OutBlockMin, FIntVector& OutBlockSize, bool& bOutBuiltAsIfTwoSided, bool& bMeshWasPlane) const
+void FStaticMeshSceneProxy::GetDistancefieldAtlasData(FBox& LocalVolumeBounds, FIntVector& OutBlockMin, FIntVector& OutBlockSize, bool& bOutBuiltAsIfTwoSided, bool& bMeshWasPlane, TArray<FMatrix>& ObjectLocalToWorldTransforms) const
 {
 	if (DistanceFieldData)
 	{
@@ -1244,6 +1245,7 @@ void FStaticMeshSceneProxy::GetDistancefieldAtlasData(FBox& LocalVolumeBounds, F
 		OutBlockSize = DistanceFieldData->VolumeTexture.GetAllocationSize();
 		bOutBuiltAsIfTwoSided = DistanceFieldData->bBuiltAsIfTwoSided;
 		bMeshWasPlane = DistanceFieldData->bMeshWasPlane;
+		ObjectLocalToWorldTransforms.Add(GetLocalToWorld());
 	}
 	else
 	{
diff --git a/Engine/Source/Runtime/Engine/Public/DistanceFieldAtlas.h b/Engine/Source/Runtime/Engine/Public/DistanceFieldAtlas.h
--- a/Engine/Source/Runtime/Engine/Public/DistanceFieldAtlas.h
+++ b/Engine/Source/Runtime/Engine/Public/DistanceFieldAtlas.h
@@ -41,6 +41,11 @@ public:
 
 	FIntVector GetAllocationSize() const;
 
+	int32 GetAllocationVolume() const
+	{
+		return GetAllocationSize().X * GetAllocationSize().Y * GetAllocationSize().Z;
+	}
+
 	bool IsValidDistanceFieldVolume() const;
 
 private:
@@ -73,7 +78,7 @@ public:
 	void RemoveAllocation(FDistanceFieldVolumeTexture* Texture);
 
 	/** Reallocates the volume texture if necessary and uploads new allocations. */
-	ENGINE_API void UpdateAllocations();
+	ENGINE_API bool UpdateAllocations();
 
 	EPixelFormat Format;
 	FTexture3DRHIRef VolumeTextureRHI;
diff --git a/Engine/Source/Runtime/Engine/Public/PrimitiveSceneProxy.h b/Engine/Source/Runtime/Engine/Public/PrimitiveSceneProxy.h
--- a/Engine/Source/Runtime/Engine/Public/PrimitiveSceneProxy.h
+++ b/Engine/Source/Runtime/Engine/Public/PrimitiveSceneProxy.h
@@ -224,7 +224,7 @@ public:
 		bShadowMapped = false;
 	}
 
-	virtual void GetDistancefieldAtlasData(FBox& LocalVolumeBounds, FIntVector& OutBlockMin, FIntVector& OutBlockSize, bool& bOutBuiltAsIfTwoSided, bool& bMeshWasPlane) const 
+	virtual void GetDistancefieldAtlasData(FBox& LocalVolumeBounds, FIntVector& OutBlockMin, FIntVector& OutBlockSize, bool& bOutBuiltAsIfTwoSided, bool& bMeshWasPlane, TArray<FMatrix>& ObjectLocalToWorldTransforms) const 
 	{
 		LocalVolumeBounds = FBox(0);
 		OutBlockMin = FIntVector(-1, -1, -1);
@@ -353,6 +353,7 @@ public:
 	inline bool CastsStaticShadow() const { return bCastStaticShadow; }
 	inline bool CastsDynamicShadow() const { return bCastDynamicShadow; }
 	inline bool AffectsDynamicIndirectLighting() const { return bAffectDynamicIndirectLighting; }
+	inline bool AffectsDistanceFieldLighting() const { return bAffectDistanceFieldLighting; }
 	inline float GetLpvBiasMultiplier() const { return LpvBiasMultiplier; }
 	inline EIndirectLightingCacheQuality GetIndirectLightingCacheQuality() const { return IndirectLightingCacheQuality; }
 	inline bool CastsVolumetricTranslucentShadow() const { return bCastVolumetricTranslucentShadow; }
@@ -373,6 +374,7 @@ public:
 	inline bool HasValidSettingsForStaticLighting() const { return bHasValidSettingsForStaticLighting; }
 	inline bool AlwaysHasVelocity() const { return bAlwaysHasVelocity; }
 	inline bool UseEditorDepthTest() const { return bUseEditorDepthTest; }
+	inline bool SupportsDistanceFieldRepresentation() const { return bSupportsDistanceFieldRepresentation; }
 	inline bool TreatAsBackgroundForOcclusion() const { return bTreatAsBackgroundForOcclusion; }
 #if WITH_EDITOR
 	inline int32 GetNumUncachedStaticLightingInteractions() { return NumUncachedStaticLightingInteractions; }
@@ -513,6 +515,8 @@ protected:
 	/** True if the primitive casts Reflective Shadow Map shadows (meaning it affects Light Propagation Volumes). */
 	uint32 bAffectDynamicIndirectLighting : 1;
 
+	uint32 bAffectDistanceFieldLighting : 1;
+
 	/** True if the primitive casts static shadows. */
 	uint32 bCastStaticShadow : 1;
 
@@ -554,6 +558,9 @@ protected:
 	/** Whether editor compositing depth testing should be used for this primitive.  Only matters for primitives with bUseEditorCompositing. */
 	uint32 bUseEditorDepthTest : 1;
 
+	/** Whether the primitive type supports a distance field representation.  Does not mean the primitive has a valid representation. */
+	uint32 bSupportsDistanceFieldRepresentation : 1;
+
 private:
 
 	/** If this is True, this primitive will be used to occlusion cull other primitives. */
diff --git a/Engine/Source/Runtime/Engine/Public/StaticMeshResources.h b/Engine/Source/Runtime/Engine/Public/StaticMeshResources.h
--- a/Engine/Source/Runtime/Engine/Public/StaticMeshResources.h
+++ b/Engine/Source/Runtime/Engine/Public/StaticMeshResources.h
@@ -777,7 +777,7 @@ public:
 	virtual FPrimitiveViewRelevance GetViewRelevance(const FSceneView* View) override;
 	virtual bool CanBeOccluded() const override;
 	virtual void GetLightRelevance(const FLightSceneProxy* LightSceneProxy, bool& bDynamic, bool& bRelevant, bool& bLightMapped, bool& bShadowMapped) const override;
-	virtual void GetDistancefieldAtlasData(FBox& LocalVolumeBounds, FIntVector& OutBlockMin, FIntVector& OutBlockSize, bool& bOutBuiltAsIfTwoSided, bool& bMeshWasPlane) const override;
+	virtual void GetDistancefieldAtlasData(FBox& LocalVolumeBounds, FIntVector& OutBlockMin, FIntVector& OutBlockSize, bool& bOutBuiltAsIfTwoSided, bool& bMeshWasPlane, TArray<FMatrix>& ObjectLocalToWorldTransforms) const override;
 	virtual bool HasDistanceFieldRepresentation() const override;
 	virtual uint32 GetMemoryFootprint( void ) const override { return( sizeof( *this ) + GetAllocatedSize() ); }
 	uint32 GetAllocatedSize( void ) const { return( FPrimitiveSceneProxy::GetAllocatedSize() + LODs.GetAllocatedSize() ); }
diff --git a/Engine/Source/Runtime/Renderer/Private/DeferredShadingRenderer.cpp b/Engine/Source/Runtime/Renderer/Private/DeferredShadingRenderer.cpp
--- a/Engine/Source/Runtime/Renderer/Private/DeferredShadingRenderer.cpp
+++ b/Engine/Source/Runtime/Renderer/Private/DeferredShadingRenderer.cpp
@@ -19,6 +19,7 @@
 #include "SceneUtils.h"
 #include "DistanceFieldSurfaceCacheLighting.h"
 #include "PostProcess/PostProcessing.h"
+#include "DistanceFieldAtlas.h"
 
 TAutoConsoleVariable<int32> CVarEarlyZPass(
 	TEXT("r.EarlyZPass"),
@@ -689,6 +690,11 @@ void FDeferredShadingSceneRenderer::Render(FRHICommandListImmediate& RHICmdList)
 		FRHICommandListExecutor::GetImmediateCommandList().ImmediateFlush(EImmediateFlushType::FlushRHIThreadFlushResources);
 	}
 
+	if (ShouldPrepareForDistanceFieldAO() || ShouldPrepareForDistanceFieldShadows())
+	{
+		const bool bReallocatedAtlasLayouts = GDistanceFieldVolumeTextureAtlas.UpdateAllocations();
+		UpdateGlobalDistanceFieldObjectBuffers(RHICmdList, bReallocatedAtlasLayouts);
+	}
 
 	const bool bIsWireframe = ViewFamily.EngineShowFlags.Wireframe;
 	static const auto ClearMethodCVar = IConsoleManager::Get().FindTConsoleVariableDataInt(TEXT("r.ClearSceneMethod"));
diff --git a/Engine/Source/Runtime/Renderer/Private/DeferredShadingRenderer.h b/Engine/Source/Runtime/Renderer/Private/DeferredShadingRenderer.h
--- a/Engine/Source/Runtime/Renderer/Private/DeferredShadingRenderer.h
+++ b/Engine/Source/Runtime/Renderer/Private/DeferredShadingRenderer.h
@@ -345,6 +345,14 @@ private:
 	void RenderStandardDeferredImageBasedReflections(FRHICommandListImmediate& RHICmdList, bool bReflectionEnv, const TRefCountPtr<IPooledRenderTarget>& DynamicBentNormalAO);
 
 	bool ShouldDoReflectionEnvironment() const;
+	
+	bool ShouldRenderDynamicSkyLight() const;
+
+	/** Whether distance field global data structures should be prepared for features that use it. */
+	bool ShouldPrepareForDistanceFieldShadows() const;
+	bool ShouldPrepareForDistanceFieldAO() const;
+
+	void UpdateGlobalDistanceFieldObjectBuffers(FRHICommandListImmediate& RHICmdList, bool bReallocatedAtlasLayouts);
 
 	friend class FTranslucentPrimSet;
 };
diff --git a/Engine/Source/Runtime/Renderer/Private/DistanceFieldLightingShared.h b/Engine/Source/Runtime/Renderer/Private/DistanceFieldLightingShared.h
--- a/Engine/Source/Runtime/Renderer/Private/DistanceFieldLightingShared.h
+++ b/Engine/Source/Runtime/Renderer/Private/DistanceFieldLightingShared.h
@@ -14,123 +14,172 @@ const int32 GDistanceFieldAOTileSizeY = 16;
 /** Base downsample factor that all distance field AO operations are done at. */
 const int32 GAODownsampleFactor = 2;
 static const int32 GMaxNumObjectsPerTile = 512;
+extern uint32 UpdateObjectsGroupSize;
 
 extern FIntPoint GetBufferSizeForAO();
 
-class FDistanceFieldObjectData
+class FDistanceFieldObjectBuffers
 {
 public:
 
-	FVertexBufferRHIRef Bounds;
-	FVertexBufferRHIRef Data;
-	FVertexBufferRHIRef Data2;
-	FVertexBufferRHIRef BoxBounds;
-	FShaderResourceViewRHIRef BoundsSRV;
-	FShaderResourceViewRHIRef DataSRV;
-	FShaderResourceViewRHIRef Data2SRV;
-	FShaderResourceViewRHIRef BoxBoundsSRV;
+	// In float4's
+	static int32 ObjectDataStride;
+
+	int32 MaxObjects;
+
+	FRWBuffer Bounds;
+	FRWBuffer Data;
+
+	FDistanceFieldObjectBuffers()
+	{
+		MaxObjects = 0;
+	}
+
+	void Initialize()
+	{
+		if (MaxObjects > 0)
+		{
+			const uint32 BufferFlags = BUF_ShaderResource;
+
+			Bounds.Initialize(sizeof(float), 4 * MaxObjects, PF_R32_FLOAT);
+			Data.Initialize(sizeof(float), 4 * MaxObjects * ObjectDataStride, PF_R32_FLOAT);
+		}
+	}
+
+	void Release()
+	{
+		Bounds.Release();
+		Data.Release();
+	}
 };
 
-class FDistanceFieldObjectBuffers
+class FDistanceFieldObjectBufferParameters
+{
+public:
+	void Bind(const FShaderParameterMap& ParameterMap)
+	{
+		ObjectBounds.Bind(ParameterMap, TEXT("ObjectBounds"));
+		ObjectData.Bind(ParameterMap, TEXT("ObjectData"));
+		NumSceneObjects.Bind(ParameterMap, TEXT("NumSceneObjects"));
+	}
+
+	template<typename TParamRef>
+	void Set(FRHICommandList& RHICmdList, const TParamRef& ShaderRHI, const FDistanceFieldObjectBuffers& ObjectBuffers, int32 NumObjectsValue)
+	{
+		ObjectBounds.SetBuffer(RHICmdList, ShaderRHI, ObjectBuffers.Bounds);
+		ObjectData.SetBuffer(RHICmdList, ShaderRHI, ObjectBuffers.Data);
+		SetShaderValue(RHICmdList, ShaderRHI, NumSceneObjects, NumObjectsValue);
+	}
+
+	template<typename TParamRef>
+	void UnsetParameters(FRHICommandList& RHICmdList, const TParamRef& ShaderRHI)
+	{
+		ObjectBounds.UnsetUAV(RHICmdList, ShaderRHI);
+		ObjectData.UnsetUAV(RHICmdList, ShaderRHI);
+	}
+
+	friend FArchive& operator<<(FArchive& Ar, FDistanceFieldObjectBufferParameters& P)
+	{
+		Ar << P.ObjectBounds << P.ObjectData << P.NumSceneObjects;
+		return Ar;
+	}
+
+private:
+	FRWShaderParameter ObjectBounds;
+	FRWShaderParameter ObjectData;
+	FShaderParameter NumSceneObjects;
+};
+
+class FDistanceFieldCulledObjectBuffers
 {
 public:
 
-	// In float4's
 	static int32 ObjectDataStride;
-	static int32 ObjectData2Stride;
 	static int32 ObjectBoxBoundsStride;
 
 	bool bWantBoxBounds;
-	bool bWantVolumeToWorld;
 	int32 MaxObjects;
-	FDistanceFieldObjectData ObjectData;
 
-	FDistanceFieldObjectBuffers()
+	FRWBuffer ObjectIndirectArguments;
+	FRWBuffer Bounds;
+	FRWBuffer Data;
+	FRWBuffer BoxBounds;
+
+	FDistanceFieldCulledObjectBuffers()
 	{
-		bWantBoxBounds = false;
-		bWantVolumeToWorld = true;
 		MaxObjects = 0;
+		bWantBoxBounds = false;
 	}
 
 	void Initialize()
 	{
 		if (MaxObjects > 0)
 		{
-			FRHIResourceCreateInfo CreateInfo;
-			ObjectData.Bounds = RHICreateVertexBuffer(MaxObjects * sizeof(FVector4), BUF_Volatile | BUF_ShaderResource, CreateInfo);
-			ObjectData.Data = RHICreateVertexBuffer(MaxObjects * sizeof(FVector4) * ObjectDataStride, BUF_Volatile | BUF_ShaderResource, CreateInfo);
+			const uint32 BufferFlags = BUF_ShaderResource;
 
-			ObjectData.BoundsSRV = RHICreateShaderResourceView(ObjectData.Bounds, sizeof(FVector4), PF_A32B32G32R32F);
-			ObjectData.DataSRV = RHICreateShaderResourceView(ObjectData.Data, sizeof(FVector4), PF_A32B32G32R32F);
-
-			if (bWantVolumeToWorld)
-			{
-				ObjectData.Data2 = RHICreateVertexBuffer(MaxObjects * sizeof(FVector4) * ObjectData2Stride, BUF_Volatile | BUF_ShaderResource, CreateInfo);
-				ObjectData.Data2SRV = RHICreateShaderResourceView(ObjectData.Data2, sizeof(FVector4), PF_A32B32G32R32F);
-			}
+			ObjectIndirectArguments.Initialize(sizeof(uint32), 5, PF_R32_UINT, BUF_Static | BUF_DrawIndirect);
+			Bounds.Initialize(sizeof(FVector4), MaxObjects, PF_A32B32G32R32F);
+			Data.Initialize(sizeof(FVector4), MaxObjects * ObjectDataStride, PF_A32B32G32R32F);
 
 			if (bWantBoxBounds)
 			{
-				ObjectData.BoxBounds = RHICreateVertexBuffer(MaxObjects * sizeof(FVector4) * ObjectBoxBoundsStride, BUF_Volatile | BUF_ShaderResource, CreateInfo);
-				ObjectData.BoxBoundsSRV = RHICreateShaderResourceView(ObjectData.BoxBounds, sizeof(FVector4), PF_A32B32G32R32F);
+				BoxBounds.Initialize(sizeof(FVector4), MaxObjects * ObjectBoxBoundsStride, PF_A32B32G32R32F);
 			}
 		}
 	}
 
 	void Release()
 	{
-		ObjectData.Bounds.SafeRelease();
-		ObjectData.Data.SafeRelease();
-		ObjectData.Data2.SafeRelease();
-		ObjectData.BoxBounds.SafeRelease();
-		ObjectData.BoundsSRV.SafeRelease(); 
-		ObjectData.DataSRV.SafeRelease(); 
-		ObjectData.Data2SRV.SafeRelease(); 
-		ObjectData.BoxBoundsSRV.SafeRelease(); 
+		ObjectIndirectArguments.Release();
+		Bounds.Release();
+		Data.Release();
+		BoxBounds.Release();
 	}
 };
 
-class FDistanceFieldObjectBufferParameters
+class FDistanceFieldObjectBufferResource : public FRenderResource
+{
+public:
+	FDistanceFieldCulledObjectBuffers Buffers;
+
+	virtual void InitDynamicRHI()  override
+	{
+		Buffers.Initialize();
+	}
+
+	virtual void ReleaseDynamicRHI() override
+	{
+		Buffers.Release();
+	}
+};
+
+class FDistanceFieldCulledObjectBufferParameters
 {
 public:
 	void Bind(const FShaderParameterMap& ParameterMap)
 	{
-		ObjectBounds.Bind(ParameterMap, TEXT("ObjectBounds"));
-		ObjectData.Bind(ParameterMap, TEXT("ObjectData"));
-		ObjectData2.Bind(ParameterMap, TEXT("ObjectData2"));
-		ObjectBoxBounds.Bind(ParameterMap, TEXT("ObjectBoxBounds"));
-		NumObjects.Bind(ParameterMap, TEXT("NumObjects"));
+		ObjectIndirectArguments.Bind(ParameterMap, TEXT("ObjectIndirectArguments"));
+		CulledObjectBounds.Bind(ParameterMap, TEXT("CulledObjectBounds"));
+		CulledObjectData.Bind(ParameterMap, TEXT("CulledObjectData"));
+		CulledObjectBoxBounds.Bind(ParameterMap, TEXT("CulledObjectBoxBounds"));
 		DistanceFieldTexture.Bind(ParameterMap, TEXT("DistanceFieldTexture"));
 		DistanceFieldSampler.Bind(ParameterMap, TEXT("DistanceFieldSampler"));
 		DistanceFieldAtlasTexelSize.Bind(ParameterMap, TEXT("DistanceFieldAtlasTexelSize"));
 	}
 
 	template<typename TParamRef>
-	void Set(FRHICommandList& RHICmdList, const TParamRef& ShaderRHI, const FDistanceFieldObjectBuffers& ObjectBuffers, int32 NumObjectsValue)
+	void Set(FRHICommandList& RHICmdList, const TParamRef& ShaderRHI, const FDistanceFieldCulledObjectBuffers& ObjectBuffers)
 	{
-		SetSRVParameter(RHICmdList, ShaderRHI, ObjectBounds, ObjectBuffers.ObjectData.BoundsSRV);
-		SetSRVParameter(RHICmdList, ShaderRHI, ObjectData, ObjectBuffers.ObjectData.DataSRV);
+		ObjectIndirectArguments.SetBuffer(RHICmdList, ShaderRHI, ObjectBuffers.ObjectIndirectArguments);
+		CulledObjectBounds.SetBuffer(RHICmdList, ShaderRHI, ObjectBuffers.Bounds);
+		CulledObjectData.SetBuffer(RHICmdList, ShaderRHI, ObjectBuffers.Data);
 
-		if (ObjectBuffers.bWantVolumeToWorld)
-		{
-			SetSRVParameter(RHICmdList, ShaderRHI, ObjectData2, ObjectBuffers.ObjectData.Data2SRV);
-		}
-		else
-		{
-			check(!ObjectData2.IsBound());
-		}
-
-		if (ObjectBuffers.bWantBoxBounds)
-		{
-			SetSRVParameter(RHICmdList, ShaderRHI, ObjectBoxBounds, ObjectBuffers.ObjectData.BoxBoundsSRV);
-		}
-		else
+		if (CulledObjectBoxBounds.IsBound())
 		{
-			check(!ObjectBoxBounds.IsBound());
+			check(ObjectBuffers.bWantBoxBounds);
+			CulledObjectBoxBounds.SetBuffer(RHICmdList, ShaderRHI, ObjectBuffers.BoxBounds);
 		}
 
-		SetShaderValue(RHICmdList, ShaderRHI, NumObjects, NumObjectsValue);
-
 		SetTextureParameter(
 			RHICmdList,
 			ShaderRHI,
@@ -145,26 +194,68 @@ public:
 		const int32 NumTexelsOneDimZ = GDistanceFieldVolumeTextureAtlas.GetSizeZ();
 		const FVector InvTextureDim(1.0f / NumTexelsOneDimX, 1.0f / NumTexelsOneDimY, 1.0f / NumTexelsOneDimZ);
 		SetShaderValue(RHICmdList, ShaderRHI, DistanceFieldAtlasTexelSize, InvTextureDim);
+	}
 
+	template<typename TParamRef>
+	void UnsetParameters(FRHICommandList& RHICmdList, const TParamRef& ShaderRHI)
+	{
+		ObjectIndirectArguments.UnsetUAV(RHICmdList, ShaderRHI);
+		CulledObjectBounds.UnsetUAV(RHICmdList, ShaderRHI);
+		CulledObjectData.UnsetUAV(RHICmdList, ShaderRHI);
+		CulledObjectBoxBounds.UnsetUAV(RHICmdList, ShaderRHI);
 	}
 
-	friend FArchive& operator<<(FArchive& Ar, FDistanceFieldObjectBufferParameters& P)
+	friend FArchive& operator<<(FArchive& Ar, FDistanceFieldCulledObjectBufferParameters& P)
 	{
-		Ar << P.ObjectBounds << P.ObjectData << P.ObjectData2 << P.ObjectBoxBounds << P.NumObjects << P.DistanceFieldTexture << P.DistanceFieldSampler << P.DistanceFieldAtlasTexelSize;
+		Ar << P.ObjectIndirectArguments << P.CulledObjectBounds << P.CulledObjectData << P.CulledObjectBoxBounds << P.DistanceFieldTexture << P.DistanceFieldSampler << P.DistanceFieldAtlasTexelSize;
 		return Ar;
 	}
 
 private:
-	FShaderResourceParameter ObjectBounds;
-	FShaderResourceParameter ObjectData;
-	FShaderResourceParameter ObjectData2;
-	FShaderResourceParameter ObjectBoxBounds;
-	FShaderParameter NumObjects;
+	FRWShaderParameter ObjectIndirectArguments;
+	FRWShaderParameter CulledObjectBounds;
+	FRWShaderParameter CulledObjectData;
+	FRWShaderParameter CulledObjectBoxBounds;
 	FShaderResourceParameter DistanceFieldTexture;
 	FShaderResourceParameter DistanceFieldSampler;
 	FShaderParameter DistanceFieldAtlasTexelSize;
 };
 
+class FCPUUpdatedBuffer
+{
+public:
+
+	EPixelFormat Format;
+	int32 Stride;
+	int32 MaxElements;
+
+	FVertexBufferRHIRef Buffer;
+	FShaderResourceViewRHIRef BufferSRV;
+
+	FCPUUpdatedBuffer()
+	{
+		Format = PF_A32B32G32R32F;
+		Stride = 1;
+		MaxElements = 0;
+	}
+
+	void Initialize()
+	{
+		if (MaxElements > 0 && Stride > 0)
+		{
+			FRHIResourceCreateInfo CreateInfo;
+			Buffer = RHICreateVertexBuffer(MaxElements * Stride * GPixelFormats[Format].BlockBytes, BUF_Volatile | BUF_ShaderResource, CreateInfo);
+			BufferSRV = RHICreateShaderResourceView(Buffer, GPixelFormats[Format].BlockBytes, Format);
+		}
+	}
+
+	void Release()
+	{
+		Buffer.SafeRelease();
+		BufferSRV.SafeRelease(); 
+	}
+};
+
 /**  */
 class FLightTileIntersectionResources
 {
diff --git a/Engine/Source/Runtime/Renderer/Private/DistanceFieldShadowing.cpp b/Engine/Source/Runtime/Renderer/Private/DistanceFieldShadowing.cpp
--- a/Engine/Source/Runtime/Renderer/Private/DistanceFieldShadowing.cpp
+++ b/Engine/Source/Runtime/Renderer/Private/DistanceFieldShadowing.cpp
@@ -40,6 +40,126 @@ FAutoConsoleVariableRef CVarShadowWorldTileSize(
 	ECVF_Cheat | ECVF_RenderThreadSafe
 	);
 
+TGlobalResource<FDistanceFieldObjectBufferResource> GShadowCulledObjectBuffers;
+
+class FCullObjectsForShadowCS : public FGlobalShader
+{
+	DECLARE_SHADER_TYPE(FCullObjectsForShadowCS,Global)
+public:
+
+	static bool ShouldCache(EShaderPlatform Platform)
+	{
+		return IsFeatureLevelSupported(Platform, ERHIFeatureLevel::SM5) && DoesPlatformSupportDistanceFieldShadowing(Platform);
+	}
+
+	static void ModifyCompilationEnvironment(EShaderPlatform Platform, FShaderCompilerEnvironment& OutEnvironment)
+	{
+		FGlobalShader::ModifyCompilationEnvironment(Platform,OutEnvironment);
+		OutEnvironment.SetDefine(TEXT("UPDATEOBJECTS_THREADGROUP_SIZE"), UpdateObjectsGroupSize);
+	}
+
+	FCullObjectsForShadowCS(const ShaderMetaType::CompiledShaderInitializerType& Initializer)
+		: FGlobalShader(Initializer)
+	{
+		ObjectBufferParameters.Bind(Initializer.ParameterMap);
+		ObjectIndirectArguments.Bind(Initializer.ParameterMap, TEXT("ObjectIndirectArguments"));
+		CulledObjectBounds.Bind(Initializer.ParameterMap, TEXT("CulledObjectBounds"));
+		CulledObjectData.Bind(Initializer.ParameterMap, TEXT("CulledObjectData"));
+		CulledObjectBoxBounds.Bind(Initializer.ParameterMap, TEXT("CulledObjectBoxBounds"));
+		ObjectBoundingGeometryIndexCount.Bind(Initializer.ParameterMap, TEXT("ObjectBoundingGeometryIndexCount"));
+		WorldToShadow.Bind(Initializer.ParameterMap, TEXT("WorldToShadow"));
+		NumShadowHullPlanes.Bind(Initializer.ParameterMap, TEXT("NumShadowHullPlanes"));
+		ShadowBoundingSphere.Bind(Initializer.ParameterMap, TEXT("ShadowBoundingSphere"));
+		ShadowConvexHull.Bind(Initializer.ParameterMap, TEXT("ShadowConvexHull"));
+	}
+
+	FCullObjectsForShadowCS()
+	{
+	}
+
+	void SetParameters(FRHICommandList& RHICmdList, const FScene* Scene, const FSceneView& View, const FProjectedShadowInfo* ProjectedShadowInfo)
+	{
+		FComputeShaderRHIParamRef ShaderRHI = GetComputeShader();
+		FGlobalShader::SetParameters(RHICmdList, ShaderRHI, View);
+		ObjectBufferParameters.Set(RHICmdList, ShaderRHI, *(Scene->DistanceFieldSceneData.ObjectBuffers), Scene->DistanceFieldSceneData.NumObjectsInBuffer);
+
+		ObjectIndirectArguments.SetBuffer(RHICmdList, ShaderRHI, GShadowCulledObjectBuffers.Buffers.ObjectIndirectArguments);
+		CulledObjectBounds.SetBuffer(RHICmdList, ShaderRHI, GShadowCulledObjectBuffers.Buffers.Bounds);
+		CulledObjectData.SetBuffer(RHICmdList, ShaderRHI, GShadowCulledObjectBuffers.Buffers.Data);
+		CulledObjectBoxBounds.SetBuffer(RHICmdList, ShaderRHI, GShadowCulledObjectBuffers.Buffers.BoxBounds);
+
+		SetShaderValue(RHICmdList, ShaderRHI, ObjectBoundingGeometryIndexCount, StencilingGeometry::GLowPolyStencilSphereIndexBuffer.GetIndexCount());
+		const FMatrix WorldToShadowValue = FTranslationMatrix(ProjectedShadowInfo->PreShadowTranslation) * ProjectedShadowInfo->SubjectAndReceiverMatrix;
+		SetShaderValue(RHICmdList, ShaderRHI, WorldToShadow, WorldToShadowValue);
+
+		int32 NumPlanes = 0;
+		const FPlane* PlaneData = NULL;
+		FVector4 ShadowBoundingSphereValue(0, 0, 0, 0);
+
+		if (ProjectedShadowInfo->bDirectionalLight)
+		{
+			NumPlanes = ProjectedShadowInfo->CascadeSettings.ShadowBoundsAccurate.Planes.Num();
+			PlaneData = ProjectedShadowInfo->CascadeSettings.ShadowBoundsAccurate.Planes.GetData();
+		}
+		else if (ProjectedShadowInfo->bOnePassPointLightShadow)
+		{
+			ShadowBoundingSphereValue = FVector4(ProjectedShadowInfo->ShadowBounds.Center.X, ProjectedShadowInfo->ShadowBounds.Center.Y, ProjectedShadowInfo->ShadowBounds.Center.Z, ProjectedShadowInfo->ShadowBounds.W);
+		}
+		else
+		{
+			NumPlanes = ProjectedShadowInfo->CasterFrustum.Planes.Num();
+			PlaneData = ProjectedShadowInfo->CasterFrustum.Planes.GetData();
+			ShadowBoundingSphereValue = FVector4(ProjectedShadowInfo->PreShadowTranslation, 0);
+		}
+
+		check(NumPlanes < 12);
+		SetShaderValue(RHICmdList, ShaderRHI, NumShadowHullPlanes, NumPlanes);
+		SetShaderValue(RHICmdList, ShaderRHI, ShadowBoundingSphere, ShadowBoundingSphereValue);
+		SetShaderValueArray(RHICmdList, ShaderRHI, ShadowConvexHull, PlaneData, NumPlanes);
+	}
+
+	void UnsetParameters(FRHICommandList& RHICmdList)
+	{
+		ObjectBufferParameters.UnsetParameters(RHICmdList, GetComputeShader());
+		ObjectIndirectArguments.UnsetUAV(RHICmdList, GetComputeShader());
+		CulledObjectBounds.UnsetUAV(RHICmdList, GetComputeShader());
+		CulledObjectData.UnsetUAV(RHICmdList, GetComputeShader());
+		CulledObjectBoxBounds.UnsetUAV(RHICmdList, GetComputeShader());
+	}
+
+	virtual bool Serialize(FArchive& Ar)
+	{		
+		bool bShaderHasOutdatedParameters = FGlobalShader::Serialize(Ar);
+		Ar << ObjectBufferParameters;
+		Ar << ObjectIndirectArguments;
+		Ar << CulledObjectBounds;
+		Ar << CulledObjectData;
+		Ar << CulledObjectBoxBounds;
+		Ar << ObjectBoundingGeometryIndexCount;
+		Ar << WorldToShadow;
+		Ar << NumShadowHullPlanes;
+		Ar << ShadowBoundingSphere;
+		Ar << ShadowConvexHull;
+		return bShaderHasOutdatedParameters;
+	}
+
+private:
+
+	FDistanceFieldObjectBufferParameters ObjectBufferParameters;
+	FRWShaderParameter ObjectIndirectArguments;
+	FRWShaderParameter CulledObjectBounds;
+	FRWShaderParameter CulledObjectData;
+	FRWShaderParameter CulledObjectBoxBounds;
+	FShaderParameter ObjectBoundingGeometryIndexCount;
+	FShaderParameter WorldToShadow;
+	FShaderParameter NumShadowHullPlanes;
+	FShaderParameter ShadowBoundingSphere;
+	FShaderParameter ShadowConvexHull;
+};
+
+IMPLEMENT_SHADER_TYPE(,FCullObjectsForShadowCS,TEXT("DistanceFieldShadowing"),TEXT("CullObjectsForShadowCS"),SF_Compute);
+
+
 /**  */
 class FClearTilesCS : public FGlobalShader
 {
@@ -128,12 +248,12 @@ public:
 
 	FShadowObjectCullVS() {}
 
-	void SetParameters(FRHICommandList& RHICmdList, const FSceneView& View, FVector2D NumGroupsValue, const FProjectedShadowInfo* ProjectedShadowInfo, int32 NumObjects)
+	void SetParameters(FRHICommandList& RHICmdList, const FSceneView& View, FVector2D NumGroupsValue, const FProjectedShadowInfo* ProjectedShadowInfo)
 	{
 		const FVertexShaderRHIParamRef ShaderRHI = GetVertexShader();
 		FGlobalShader::SetParameters(RHICmdList, ShaderRHI, View);
 		
-		ObjectParameters.Set(RHICmdList, ShaderRHI, *ProjectedShadowInfo->LightSceneInfo->DistanceFieldObjectBuffers, NumObjects);
+		ObjectParameters.Set(RHICmdList, ShaderRHI, GShadowCulledObjectBuffers.Buffers);
 
 		const int32 NumRings = StencilingGeometry::GLowPolyStencilSphereVertexBuffer.GetNumRings();
 		const float RadiansPerRingSegment = PI / (float)NumRings;
@@ -161,7 +281,7 @@ public:
 	}
 
 private:
-	FDistanceFieldObjectBufferParameters ObjectParameters;
+	FDistanceFieldCulledObjectBufferParameters ObjectParameters;
 	FShaderParameter ConservativeRadiusScale;
 	FShaderParameter WorldToShadow;
 	FShaderParameter MinRadius;
@@ -204,14 +324,13 @@ public:
 		FRHICommandList& RHICmdList, 
 		const FSceneView& View, 
 		FVector2D NumGroupsValue, 
-		int32 NumObjects, 
 		const FProjectedShadowInfo* ProjectedShadowInfo,
 		FLightTileIntersectionResources* TileIntersectionResources)
 	{
 		const FPixelShaderRHIParamRef ShaderRHI = GetPixelShader();
 		FGlobalShader::SetParameters(RHICmdList, ShaderRHI, View);
 
-		ObjectParameters.Set(RHICmdList, ShaderRHI, *ProjectedShadowInfo->LightSceneInfo->DistanceFieldObjectBuffers, NumObjects);
+		ObjectParameters.Set(RHICmdList, ShaderRHI, GShadowCulledObjectBuffers.Buffers);
 		SetShaderValue(RHICmdList, ShaderRHI, NumGroups, NumGroupsValue);
 	}
 
@@ -245,7 +364,7 @@ public:
 	}
 
 private:
-	FDistanceFieldObjectBufferParameters ObjectParameters;
+	FDistanceFieldCulledObjectBufferParameters ObjectParameters;
 	FRWShaderParameter TileHeadDataUnpacked;
 	FRWShaderParameter TileArrayData;
 	FShaderParameter NumGroups;
@@ -310,7 +429,6 @@ public:
 		const FSceneView& View, 
 		const FProjectedShadowInfo* ProjectedShadowInfo,
 		FSceneRenderTargetItem& ShadowFactorsValue, 
-		int32 NumObjects, 
 		FVector2D NumGroupsValue,
 		const FIntRect& ScissorRect,
 		FLightTileIntersectionResources* TileIntersectionResources)
@@ -321,7 +439,7 @@ public:
 
 		ShadowFactors.SetTexture(RHICmdList, ShaderRHI, ShadowFactorsValue.ShaderResourceTexture, ShadowFactorsValue.UAV);
 
-		ObjectParameters.Set(RHICmdList, ShaderRHI, *ProjectedShadowInfo->LightSceneInfo->DistanceFieldObjectBuffers, NumObjects);
+		ObjectParameters.Set(RHICmdList, ShaderRHI, GShadowCulledObjectBuffers.Buffers);
 		DeferredParameters.Set(RHICmdList, ShaderRHI, View);
 
 		SetShaderValue(RHICmdList, ShaderRHI, NumGroups, NumGroupsValue);
@@ -394,7 +512,7 @@ private:
 	FShaderParameter LightSourceRadius;
 	FShaderParameter TanLightAngleAndNormalThreshold;
 	FShaderParameter ScissorRectMinAndSize;
-	FDistanceFieldObjectBufferParameters ObjectParameters;
+	FDistanceFieldCulledObjectBufferParameters ObjectParameters;
 	FDeferredPixelShaderParameters DeferredParameters;
 	FShaderResourceParameter TileHeadDataUnpacked;
 	FShaderResourceParameter TileArrayData;
@@ -467,357 +585,249 @@ private:
 
 IMPLEMENT_SHADER_TYPE(,FDistanceFieldShadowingUpsamplePS,TEXT("DistanceFieldShadowing"),TEXT("DistanceFieldShadowingUpsamplePS"),SF_Pixel);
 
-int32 FProjectedShadowInfo::UpdateShadowCastingObjectBuffers() const
+bool SupportsDistanceFieldShadows(ERHIFeatureLevel::Type FeatureLevel, EShaderPlatform ShaderPlatform)
 {
-	int32 NumObjects = 0;
+	return GDistanceFieldShadowing
+		&& FeatureLevel >= ERHIFeatureLevel::SM5
+		&& DoesPlatformSupportDistanceFieldShadowing(ShaderPlatform);
+}
 
-	//@todo - scene rendering allocator
-	TArray<FVector4> ObjectBoundsData;
-	TArray<FVector4> ObjectData;
-	TArray<FVector4> ObjectBoxBounds;
+bool FDeferredShadingSceneRenderer::ShouldPrepareForDistanceFieldShadows() const 
+{
+	bool bSceneHasRayTracedDFShadows = false;
 
+	for (TSparseArray<FLightSceneInfoCompact>::TConstIterator LightIt(Scene->Lights); LightIt; ++LightIt)
 	{
-		QUICK_SCOPE_CYCLE_COUNTER(STAT_GatherObjectData);
-
-		const int32 NumTexelsOneDimX = GDistanceFieldVolumeTextureAtlas.GetSizeX();
-		const int32 NumTexelsOneDimY = GDistanceFieldVolumeTextureAtlas.GetSizeY();
-		const int32 NumTexelsOneDimZ = GDistanceFieldVolumeTextureAtlas.GetSizeZ();
-		const FVector InvTextureDim(1.0f / NumTexelsOneDimX, 1.0f / NumTexelsOneDimY, 1.0f / NumTexelsOneDimZ);
-		const FMatrix WorldToShadow = FTranslationMatrix(PreShadowTranslation) * SubjectAndReceiverMatrix;
+		const FLightSceneInfoCompact& LightSceneInfoCompact = *LightIt;
+		const FLightSceneInfo* const LightSceneInfo = LightSceneInfoCompact.LightSceneInfo;
 
-		ObjectBoundsData.Empty(SubjectPrimitives.Num());
-		ObjectData.Empty(FDistanceFieldObjectBuffers::ObjectDataStride * SubjectPrimitives.Num());
-		ObjectBoxBounds.Empty(FDistanceFieldObjectBuffers::ObjectBoxBoundsStride * SubjectPrimitives.Num());
-
-		for (int32 PrimitiveIndex = 0; PrimitiveIndex < SubjectPrimitives.Num() && NumObjects < MAX_uint16; PrimitiveIndex++)
+		if (LightSceneInfo->ShouldRenderLightViewIndependent())
 		{
-			const FPrimitiveSceneInfo* PrimitiveSceneInfo = SubjectPrimitives[PrimitiveIndex];
-
-			FBox LocalVolumeBounds;
-			FIntVector BlockMin;
-			FIntVector BlockSize;
-			bool bBuiltAsIfTwoSided;
-			bool bMeshWasPlane;
-			PrimitiveSceneInfo->Proxy->GetDistancefieldAtlasData(LocalVolumeBounds, BlockMin, BlockSize, bBuiltAsIfTwoSided, bMeshWasPlane);
+			const FVisibleLightInfo& VisibleLightInfo = VisibleLightInfos[LightSceneInfo->Id];
 
-			if (BlockMin.X >= 0 && BlockMin.Y >= 0 && BlockMin.Z >= 0)
+			for (int32 ShadowIndex = 0; ShadowIndex < VisibleLightInfo.AllProjectedShadows.Num(); ShadowIndex++)
 			{
-				const FMatrix LocalToWorld = PrimitiveSceneInfo->Proxy->GetLocalToWorld();
-
-				const FMatrix VolumeToWorld = FScaleMatrix(LocalVolumeBounds.GetExtent()) 
-					* FTranslationMatrix(LocalVolumeBounds.GetCenter())
-					* LocalToWorld;
-
-				const FVector4 ObjectBoundingSphere(VolumeToWorld.GetOrigin(), VolumeToWorld.GetScaleVector().Size());
-
-				ObjectBoundsData.Add(ObjectBoundingSphere);
-
-				const float MaxExtent = LocalVolumeBounds.GetExtent().GetMax();
-
-				const FMatrix UniformScaleVolumeToWorld = FScaleMatrix(MaxExtent) 
-					* FTranslationMatrix(LocalVolumeBounds.GetCenter())
-					* LocalToWorld;
-
-				const FVector InvBlockSize(1.0f / BlockSize.X, 1.0f / BlockSize.Y, 1.0f / BlockSize.Z);
-
-				//float3 VolumeUV = (VolumePosition / LocalPositionExtent * .5f * UVScale + .5f * UVScale + UVAdd;
-				const FVector LocalPositionExtent = LocalVolumeBounds.GetExtent() / FVector(MaxExtent);
-				const FVector UVScale = FVector(BlockSize) * InvTextureDim;
-				const float VolumeScale = UniformScaleVolumeToWorld.GetMaximumAxisScale();
-				const FMatrix WorldToVolume = UniformScaleVolumeToWorld.InverseFast();
-				// WorldToVolume
-				ObjectData.Add(*(FVector4*)&WorldToVolume.M[0]);
-				ObjectData.Add(*(FVector4*)&WorldToVolume.M[1]);
-				ObjectData.Add(*(FVector4*)&WorldToVolume.M[2]);
-				ObjectData.Add(*(FVector4*)&WorldToVolume.M[3]);
-
-				// Clamp to texel center by subtracting a half texel in the [-1,1] position space
-				// LocalPositionExtent
-				ObjectData.Add(FVector4(LocalPositionExtent - InvBlockSize, bBuiltAsIfTwoSided ? 1.0f : 0.0f));
-
-				// UVScale
-				ObjectData.Add(FVector4(FVector(BlockSize) * InvTextureDim * .5f / LocalPositionExtent, VolumeScale));
-
-				// UVAdd
-				ObjectData.Add(FVector4(FVector(BlockMin) * InvTextureDim + .5f * UVScale, 0));
-
-				// Padding
-				ObjectData.Add(FVector4(0, 0, 0, 0));
+				const FProjectedShadowInfo* ProjectedShadowInfo = VisibleLightInfo.AllProjectedShadows[ShadowIndex];
 
+				if (ProjectedShadowInfo->bRayTracedDistanceFieldShadow)
 				{
-					FVector LocalBoundsVertices[8];
-					LocalBoundsVertices[0] = FVector(LocalVolumeBounds.Min.X, LocalVolumeBounds.Min.Y, LocalVolumeBounds.Min.Z);
-					LocalBoundsVertices[1] = FVector(LocalVolumeBounds.Max.X, LocalVolumeBounds.Min.Y, LocalVolumeBounds.Min.Z);
-					LocalBoundsVertices[2] = FVector(LocalVolumeBounds.Min.X, LocalVolumeBounds.Max.Y, LocalVolumeBounds.Min.Z);
-					LocalBoundsVertices[3] = FVector(LocalVolumeBounds.Max.X, LocalVolumeBounds.Max.Y, LocalVolumeBounds.Min.Z);
-					LocalBoundsVertices[4] = FVector(LocalVolumeBounds.Min.X, LocalVolumeBounds.Min.Y, LocalVolumeBounds.Max.Z);
-					LocalBoundsVertices[5] = FVector(LocalVolumeBounds.Max.X, LocalVolumeBounds.Min.Y, LocalVolumeBounds.Max.Z);
-					LocalBoundsVertices[6] = FVector(LocalVolumeBounds.Min.X, LocalVolumeBounds.Max.Y, LocalVolumeBounds.Max.Z);
-					LocalBoundsVertices[7] = FVector(LocalVolumeBounds.Max.X, LocalVolumeBounds.Max.Y, LocalVolumeBounds.Max.Z);
-
-					FVector ViewSpaceBoundsVertices[8];
-					FVector MinViewSpacePosition(WORLD_MAX);
-					FVector MaxViewSpacePosition(-WORLD_MAX);
-
-					for (int32 i = 0; i < ARRAY_COUNT(LocalBoundsVertices); i++)
-					{
-						const FVector ViewSpacePosition = WorldToShadow.TransformPosition(LocalToWorld.TransformPosition(LocalBoundsVertices[i]));
-						MinViewSpacePosition = MinViewSpacePosition.ComponentMin(ViewSpacePosition);
-						MaxViewSpacePosition = MaxViewSpacePosition.ComponentMax(ViewSpacePosition);
-						ViewSpaceBoundsVertices[i] = ViewSpacePosition;
-					}
-
-					FVector ObjectXAxis = (ViewSpaceBoundsVertices[1] - ViewSpaceBoundsVertices[0]) / 2.0f;
-					FVector ObjectYAxis = (ViewSpaceBoundsVertices[2] - ViewSpaceBoundsVertices[0]) / 2.0f;
-					FVector ObjectZAxis = (ViewSpaceBoundsVertices[4] - ViewSpaceBoundsVertices[0]) / 2.0f;
-
-					ObjectBoxBounds.Add(FVector4(MinViewSpacePosition, 0));
-					ObjectBoxBounds.Add(FVector4(MaxViewSpacePosition, 0));
-					ObjectBoxBounds.Add(FVector4(ObjectXAxis / FMath::Max(ObjectXAxis.SizeSquared(), KINDA_SMALL_NUMBER), 0));
-					ObjectBoxBounds.Add(FVector4(ObjectYAxis / FMath::Max(ObjectYAxis.SizeSquared(), KINDA_SMALL_NUMBER), 0));
-					ObjectBoxBounds.Add(FVector4(ObjectZAxis / FMath::Max(ObjectZAxis.SizeSquared(), KINDA_SMALL_NUMBER), 0));
+					bSceneHasRayTracedDFShadows = true;
+					break;
 				}
-
-				checkSlow(ObjectData.Num() % FDistanceFieldObjectBuffers::ObjectDataStride == 0);
-				checkSlow(ObjectBoxBounds.Num() % FDistanceFieldObjectBuffers::ObjectBoxBoundsStride == 0);
-
-				NumObjects++;
 			}
 		}
 	}
 
-	if (NumObjects > 0)
-	{
-		QUICK_SCOPE_CYCLE_COUNTER(STAT_UpdateObjectBuffers);
-
-		if (!LightSceneInfo->DistanceFieldObjectBuffers)
-		{
-			LightSceneInfo->DistanceFieldObjectBuffers = new FDistanceFieldObjectBuffers;
-		}
-
-		FDistanceFieldObjectBuffers& ObjectBuffers = *LightSceneInfo->DistanceFieldObjectBuffers;
-
-		if (NumObjects > ObjectBuffers.MaxObjects)
-		{
-			// Allocate with slack
-			ObjectBuffers.bWantBoxBounds = true;
-			ObjectBuffers.bWantVolumeToWorld = false;
-			ObjectBuffers.MaxObjects = NumObjects * 5 / 4;
-			ObjectBuffers.Release();
-			ObjectBuffers.Initialize();
-		}
-
-		checkSlow(NumObjects < MAX_uint16);
-
-		void* LockedBuffer = RHILockVertexBuffer(ObjectBuffers.ObjectData.Bounds, 0, ObjectBuffers.ObjectData.Bounds->GetSize(), RLM_WriteOnly);
-		FPlatformMemory::Memcpy(LockedBuffer, ObjectBoundsData.GetData(), ObjectBoundsData.GetTypeSize() * ObjectBoundsData.Num());
-		RHIUnlockVertexBuffer(ObjectBuffers.ObjectData.Bounds);
-
-		LockedBuffer = RHILockVertexBuffer(ObjectBuffers.ObjectData.Data, 0, ObjectBuffers.ObjectData.Data->GetSize(), RLM_WriteOnly);
-		FPlatformMemory::Memcpy(LockedBuffer, ObjectData.GetData(), ObjectData.GetTypeSize() * ObjectData.Num());
-		RHIUnlockVertexBuffer(ObjectBuffers.ObjectData.Data);
-
-		LockedBuffer = RHILockVertexBuffer(ObjectBuffers.ObjectData.BoxBounds, 0, ObjectBuffers.ObjectData.BoxBounds->GetSize(), RLM_WriteOnly);
-		FPlatformMemory::Memcpy(LockedBuffer, ObjectBoxBounds.GetData(), ObjectBoxBounds.GetTypeSize() * ObjectBoxBounds.Num());
-		RHIUnlockVertexBuffer(ObjectBuffers.ObjectData.BoxBounds);
-	}
-
-	return NumObjects;
+	return ViewFamily.EngineShowFlags.DynamicShadows 
+		&& bSceneHasRayTracedDFShadows
+		&& SupportsDistanceFieldShadows(Scene->GetFeatureLevel(), Scene->GetShaderPlatform());
 }
 
 void FProjectedShadowInfo::RenderRayTracedDistanceFieldProjection(FRHICommandListImmediate& RHICmdList, const FViewInfo& View) const
 {
-	if (GDistanceFieldShadowing
-		&& View.GetFeatureLevel() >= ERHIFeatureLevel::SM5
-		&& DoesPlatformSupportDistanceFieldShadowing(View.GetShaderPlatform()))
+	if (SupportsDistanceFieldShadows(View.GetFeatureLevel(), View.GetShaderPlatform()))
 	{
 		QUICK_SCOPE_CYCLE_COUNTER(STAT_RenderRayTracedDistanceFieldShadows);
 		SCOPED_DRAW_EVENT(RHICmdList, RayTracedDistanceFieldShadow);
 
-		// Update the global distance field atlas
-		GDistanceFieldVolumeTextureAtlas.UpdateAllocations();
+		const FScene* Scene = (const FScene*)(View.Family->Scene);
 
-		if (GDistanceFieldVolumeTextureAtlas.VolumeTextureRHI)
+		if (GDistanceFieldVolumeTextureAtlas.VolumeTextureRHI && Scene->DistanceFieldSceneData.NumObjectsInBuffer > 0)
 		{
-			const int32 NumObjects = UpdateShadowCastingObjectBuffers();
+			check(!Scene->DistanceFieldSceneData.HasPendingOperations());
 
-			if (NumObjects > 0)
 			{
-				// Allocate tile resolution based on world space size
-				//@todo - light space perspective shadow maps would make much better use of the resolution
-				const float LightTiles = FMath::Min(ShadowBounds.W / GShadowWorldTileSize + 1, 256.0f);
-				FIntPoint LightTileDimensions(LightTiles, LightTiles);
+				SCOPED_DRAW_EVENT(RHICmdList, ObjectFrustumCulling);
+
+				if (GShadowCulledObjectBuffers.Buffers.MaxObjects < Scene->DistanceFieldSceneData.NumObjectsInBuffer
+					|| GShadowCulledObjectBuffers.Buffers.MaxObjects > 3 * Scene->DistanceFieldSceneData.NumObjectsInBuffer)
+				{
+					GShadowCulledObjectBuffers.Buffers.bWantBoxBounds = true;
+					GShadowCulledObjectBuffers.Buffers.MaxObjects = Scene->DistanceFieldSceneData.NumObjectsInBuffer * 5 / 4;
+					GShadowCulledObjectBuffers.Buffers.Release();
+					GShadowCulledObjectBuffers.Buffers.Initialize();
+				}
 
-				FLightTileIntersectionResources* TileIntersectionResources = NULL;
+				{
+					TShaderMapRef<FCullObjectsForShadowCS> ComputeShader(GetGlobalShaderMap(Scene->GetFeatureLevel()));
+					RHICmdList.SetComputeShader(ComputeShader->GetComputeShader());
+					ComputeShader->SetParameters(RHICmdList, Scene, View, this);
 
-				GSceneRenderTargets.FinishRenderingLightAttenuation(RHICmdList);
-				SetRenderTarget(RHICmdList, NULL, NULL);
+					DispatchComputeShader(RHICmdList, *ComputeShader, FMath::DivideAndRoundUp<uint32>(Scene->DistanceFieldSceneData.NumObjectsInBuffer, UpdateObjectsGroupSize), 1, 1);
+					ComputeShader->UnsetParameters(RHICmdList);
+				}
+			}
 
-				if (bDirectionalLight && GShadowScatterTileCulling)
+			// Allocate tile resolution based on world space size
+			//@todo - light space perspective shadow maps would make much better use of the resolution
+			const float LightTiles = FMath::Min(ShadowBounds.W / GShadowWorldTileSize + 1, 256.0f);
+			FIntPoint LightTileDimensions(LightTiles, LightTiles);
+
+			FLightTileIntersectionResources* TileIntersectionResources = NULL;
+
+			GSceneRenderTargets.FinishRenderingLightAttenuation(RHICmdList);
+			SetRenderTarget(RHICmdList, NULL, NULL);
+
+			if (bDirectionalLight && GShadowScatterTileCulling)
+			{
+				if (!LightSceneInfo->TileIntersectionResources || LightSceneInfo->TileIntersectionResources->TileDimensions != LightTileDimensions)
 				{
-					if (!LightSceneInfo->TileIntersectionResources || LightSceneInfo->TileIntersectionResources->TileDimensions != LightTileDimensions)
+					if (LightSceneInfo->TileIntersectionResources)
+					{
+						LightSceneInfo->TileIntersectionResources->Release();
+					}
+					else
 					{
-						if (LightSceneInfo->TileIntersectionResources)
-						{
-							LightSceneInfo->TileIntersectionResources->Release();
-						}
-						else
-						{
-							LightSceneInfo->TileIntersectionResources = new FLightTileIntersectionResources();
-						}
-
-						LightSceneInfo->TileIntersectionResources->TileDimensions = LightTileDimensions;
-
-						LightSceneInfo->TileIntersectionResources->Initialize();
+						LightSceneInfo->TileIntersectionResources = new FLightTileIntersectionResources();
 					}
 
-					TileIntersectionResources = LightSceneInfo->TileIntersectionResources;
+					LightSceneInfo->TileIntersectionResources->TileDimensions = LightTileDimensions;
 
-					{
-						SCOPED_DRAW_EVENT(RHICmdList, ClearTiles);
-						TShaderMapRef<FClearTilesCS> ComputeShader(View.ShaderMap);
+					LightSceneInfo->TileIntersectionResources->Initialize();
+				}
 
-						uint32 GroupSizeX = FMath::DivideAndRoundUp(LightTileDimensions.X, GDistanceFieldAOTileSizeX);
-						uint32 GroupSizeY = FMath::DivideAndRoundUp(LightTileDimensions.Y, GDistanceFieldAOTileSizeY);
+				TileIntersectionResources = LightSceneInfo->TileIntersectionResources;
 
-						RHICmdList.SetComputeShader(ComputeShader->GetComputeShader());
-						ComputeShader->SetParameters(RHICmdList, View, FVector2D(LightTileDimensions.X, LightTileDimensions.Y), TileIntersectionResources);
-						DispatchComputeShader(RHICmdList, *ComputeShader, GroupSizeX, GroupSizeY, 1);
+				{
+					SCOPED_DRAW_EVENT(RHICmdList, ClearTiles);
+					TShaderMapRef<FClearTilesCS> ComputeShader(View.ShaderMap);
 
-						ComputeShader->UnsetParameters(RHICmdList);
-					}
+					uint32 GroupSizeX = FMath::DivideAndRoundUp(LightTileDimensions.X, GDistanceFieldAOTileSizeX);
+					uint32 GroupSizeY = FMath::DivideAndRoundUp(LightTileDimensions.Y, GDistanceFieldAOTileSizeY);
 
-					{
-						SCOPED_DRAW_EVENT(RHICmdList, CullObjects);
+					RHICmdList.SetComputeShader(ComputeShader->GetComputeShader());
+					ComputeShader->SetParameters(RHICmdList, View, FVector2D(LightTileDimensions.X, LightTileDimensions.Y), TileIntersectionResources);
+					DispatchComputeShader(RHICmdList, *ComputeShader, GroupSizeX, GroupSizeY, 1);
+
+					ComputeShader->UnsetParameters(RHICmdList);
+				}
+
+				{
+					SCOPED_DRAW_EVENT(RHICmdList, CullObjects);
 
-						TShaderMapRef<FShadowObjectCullVS> VertexShader(View.ShaderMap);
-						TShaderMapRef<FShadowObjectCullPS> PixelShader(View.ShaderMap);
+					TShaderMapRef<FShadowObjectCullVS> VertexShader(View.ShaderMap);
+					TShaderMapRef<FShadowObjectCullPS> PixelShader(View.ShaderMap);
 
-						TArray<FUnorderedAccessViewRHIParamRef> UAVs;
-						PixelShader->GetUAVs(View, TileIntersectionResources, UAVs);
-						RHICmdList.SetRenderTargets(0, (const FRHIRenderTargetView *)NULL, NULL, UAVs.Num(), UAVs.GetData());
+					TArray<FUnorderedAccessViewRHIParamRef> UAVs;
+					PixelShader->GetUAVs(View, TileIntersectionResources, UAVs);
+					RHICmdList.SetRenderTargets(0, (const FRHIRenderTargetView *)NULL, NULL, UAVs.Num(), UAVs.GetData());
 
-						RHICmdList.SetViewport(0, 0, 0.0f, LightTileDimensions.X, LightTileDimensions.Y, 1.0f);
+					RHICmdList.SetViewport(0, 0, 0.0f, LightTileDimensions.X, LightTileDimensions.Y, 1.0f);
 
-						// Render backfaces since camera may intersect
-						RHICmdList.SetRasterizerState(View.bReverseCulling ? TStaticRasterizerState<FM_Solid, CM_CW>::GetRHI() : TStaticRasterizerState<FM_Solid, CM_CCW>::GetRHI());
-						RHICmdList.SetDepthStencilState(TStaticDepthStencilState<false, CF_Always>::GetRHI());
-						RHICmdList.SetBlendState(TStaticBlendState<>::GetRHI());
+					// Render backfaces since camera may intersect
+					RHICmdList.SetRasterizerState(View.bReverseCulling ? TStaticRasterizerState<FM_Solid, CM_CW>::GetRHI() : TStaticRasterizerState<FM_Solid, CM_CCW>::GetRHI());
+					RHICmdList.SetDepthStencilState(TStaticDepthStencilState<false, CF_Always>::GetRHI());
+					RHICmdList.SetBlendState(TStaticBlendState<>::GetRHI());
 
-						static FGlobalBoundShaderState BoundShaderState;
+					static FGlobalBoundShaderState BoundShaderState;
 
-						SetGlobalBoundShaderState(RHICmdList, View.GetFeatureLevel(), BoundShaderState, GetVertexDeclarationFVector4(), *VertexShader, *PixelShader);
+					SetGlobalBoundShaderState(RHICmdList, View.GetFeatureLevel(), BoundShaderState, GetVertexDeclarationFVector4(), *VertexShader, *PixelShader);
 
-						VertexShader->SetParameters(RHICmdList, View, FVector2D(LightTileDimensions.X, LightTileDimensions.Y), this, NumObjects);
-						PixelShader->SetParameters(RHICmdList, View, FVector2D(LightTileDimensions.X, LightTileDimensions.Y), NumObjects, this, TileIntersectionResources);
+					VertexShader->SetParameters(RHICmdList, View, FVector2D(LightTileDimensions.X, LightTileDimensions.Y), this);
+					PixelShader->SetParameters(RHICmdList, View, FVector2D(LightTileDimensions.X, LightTileDimensions.Y), this, TileIntersectionResources);
 
-						RHICmdList.SetStreamSource(0, StencilingGeometry::GLowPolyStencilSphereVertexBuffer.VertexBufferRHI, sizeof(FVector4), 0);
+					RHICmdList.SetStreamSource(0, StencilingGeometry::GLowPolyStencilSphereVertexBuffer.VertexBufferRHI, sizeof(FVector4), 0);
 
-						RHICmdList.DrawIndexedPrimitive(
-							StencilingGeometry::GLowPolyStencilSphereIndexBuffer.IndexBufferRHI, 
-							PT_TriangleList, 
-							0,
-							0, 
-							StencilingGeometry::GLowPolyStencilSphereVertexBuffer.GetVertexCount(), 
-							0, 
-							StencilingGeometry::GLowPolyStencilSphereIndexBuffer.GetIndexCount() / 3, 
-							NumObjects);
-					}
+					RHICmdList.DrawIndexedPrimitiveIndirect(
+						PT_TriangleList,
+						StencilingGeometry::GLowPolyStencilSphereIndexBuffer.IndexBufferRHI, 
+						GShadowCulledObjectBuffers.Buffers.ObjectIndirectArguments.Buffer,
+						0);
 				}
+			}
+
+			TRefCountPtr<IPooledRenderTarget> RayTracedShadowsRT;
 
-				TRefCountPtr<IPooledRenderTarget> RayTracedShadowsRT;
+			{
+				const FIntPoint BufferSize = GetBufferSizeForAO();
+				FPooledRenderTargetDesc Desc(FPooledRenderTargetDesc::Create2DDesc(BufferSize, PF_G16R16F, TexCreate_None, TexCreate_RenderTargetable | TexCreate_UAV, false));
+				GRenderTargetPool.FindFreeElement(Desc, RayTracedShadowsRT, TEXT("RayTracedShadows"));
+			}
+
+			FIntRect ScissorRect;
 
+			{
+				if (!LightSceneInfo->Proxy->GetScissorRect(ScissorRect, View))
 				{
-					const FIntPoint BufferSize = GetBufferSizeForAO();
-					FPooledRenderTargetDesc Desc(FPooledRenderTargetDesc::Create2DDesc(BufferSize, PF_G16R16F, TexCreate_None, TexCreate_RenderTargetable | TexCreate_UAV, false));
-					GRenderTargetPool.FindFreeElement(Desc, RayTracedShadowsRT, TEXT("RayTracedShadows"));
+					ScissorRect = View.ViewRect;
 				}
 
-				FIntRect ScissorRect;
+				uint32 GroupSizeX = FMath::DivideAndRoundUp(ScissorRect.Size().X / GAODownsampleFactor, GDistanceFieldAOTileSizeX);
+				uint32 GroupSizeY = FMath::DivideAndRoundUp(ScissorRect.Size().Y / GAODownsampleFactor, GDistanceFieldAOTileSizeY);
 
 				{
-					if (!LightSceneInfo->Proxy->GetScissorRect(ScissorRect, View))
+					SCOPED_DRAW_EVENT(RHICmdList, RayTraceShadows);
+					SetRenderTarget(RHICmdList, NULL, NULL);
+
+					if (bDirectionalLight && GShadowScatterTileCulling)
 					{
-						ScissorRect = View.ViewRect;
+						TShaderMapRef<TDistanceFieldShadowingCS<DFS_DirectionalLightScatterTileCulling> > ComputeShader(View.ShaderMap);
+						RHICmdList.SetComputeShader(ComputeShader->GetComputeShader());
+						ComputeShader->SetParameters(RHICmdList, View, this, RayTracedShadowsRT->GetRenderTargetItem(), FVector2D(GroupSizeX, GroupSizeY), ScissorRect, TileIntersectionResources);
+						DispatchComputeShader(RHICmdList, *ComputeShader, GroupSizeX, GroupSizeY, 1);
+						ComputeShader->UnsetParameters(RHICmdList);
 					}
-
-					uint32 GroupSizeX = FMath::DivideAndRoundUp(ScissorRect.Size().X / GAODownsampleFactor, GDistanceFieldAOTileSizeX);
-					uint32 GroupSizeY = FMath::DivideAndRoundUp(ScissorRect.Size().Y / GAODownsampleFactor, GDistanceFieldAOTileSizeY);
-
+					else if (bDirectionalLight)
 					{
-						SCOPED_DRAW_EVENT(RHICmdList, RayTraceShadows);
-						SetRenderTarget(RHICmdList, NULL, NULL);
-
-						if (bDirectionalLight && GShadowScatterTileCulling)
-						{
-							TShaderMapRef<TDistanceFieldShadowingCS<DFS_DirectionalLightScatterTileCulling> > ComputeShader(View.ShaderMap);
-							RHICmdList.SetComputeShader(ComputeShader->GetComputeShader());
-							ComputeShader->SetParameters(RHICmdList, View, this, RayTracedShadowsRT->GetRenderTargetItem(), NumObjects, FVector2D(GroupSizeX, GroupSizeY), ScissorRect, TileIntersectionResources);
-							DispatchComputeShader(RHICmdList, *ComputeShader, GroupSizeX, GroupSizeY, 1);
-							ComputeShader->UnsetParameters(RHICmdList);
-						}
-						else if (bDirectionalLight)
-						{
-							TShaderMapRef<TDistanceFieldShadowingCS<DFS_DirectionalLightTiledCulling> > ComputeShader(View.ShaderMap);
-							RHICmdList.SetComputeShader(ComputeShader->GetComputeShader());
-							ComputeShader->SetParameters(RHICmdList, View, this, RayTracedShadowsRT->GetRenderTargetItem(), NumObjects, FVector2D(GroupSizeX, GroupSizeY), ScissorRect, TileIntersectionResources);
-							DispatchComputeShader(RHICmdList, *ComputeShader, GroupSizeX, GroupSizeY, 1);
-							ComputeShader->UnsetParameters(RHICmdList);
-						}
-						else
-						{
-							TShaderMapRef<TDistanceFieldShadowingCS<DFS_PointLightTiledCulling> > ComputeShader(View.ShaderMap);
-							RHICmdList.SetComputeShader(ComputeShader->GetComputeShader());
-							ComputeShader->SetParameters(RHICmdList, View, this, RayTracedShadowsRT->GetRenderTargetItem(), NumObjects, FVector2D(GroupSizeX, GroupSizeY), ScissorRect, TileIntersectionResources);
-							DispatchComputeShader(RHICmdList, *ComputeShader, GroupSizeX, GroupSizeY, 1);
-							ComputeShader->UnsetParameters(RHICmdList);
-						}
+						TShaderMapRef<TDistanceFieldShadowingCS<DFS_DirectionalLightTiledCulling> > ComputeShader(View.ShaderMap);
+						RHICmdList.SetComputeShader(ComputeShader->GetComputeShader());
+						ComputeShader->SetParameters(RHICmdList, View, this, RayTracedShadowsRT->GetRenderTargetItem(), FVector2D(GroupSizeX, GroupSizeY), ScissorRect, TileIntersectionResources);
+						DispatchComputeShader(RHICmdList, *ComputeShader, GroupSizeX, GroupSizeY, 1);
+						ComputeShader->UnsetParameters(RHICmdList);
+					}
+					else
+					{
+						TShaderMapRef<TDistanceFieldShadowingCS<DFS_PointLightTiledCulling> > ComputeShader(View.ShaderMap);
+						RHICmdList.SetComputeShader(ComputeShader->GetComputeShader());
+						ComputeShader->SetParameters(RHICmdList, View, this, RayTracedShadowsRT->GetRenderTargetItem(), FVector2D(GroupSizeX, GroupSizeY), ScissorRect, TileIntersectionResources);
+						DispatchComputeShader(RHICmdList, *ComputeShader, GroupSizeX, GroupSizeY, 1);
+						ComputeShader->UnsetParameters(RHICmdList);
 					}
 				}
+			}
 
-				{
-					GSceneRenderTargets.BeginRenderingLightAttenuation(RHICmdList);
+			{
+				GSceneRenderTargets.BeginRenderingLightAttenuation(RHICmdList);
 
-					SCOPED_DRAW_EVENT(RHICmdList, Upsample);
+				SCOPED_DRAW_EVENT(RHICmdList, Upsample);
 
-					RHICmdList.SetViewport(ScissorRect.Min.X, ScissorRect.Min.Y, 0.0f, ScissorRect.Max.X, ScissorRect.Max.Y, 1.0f);
-					RHICmdList.SetRasterizerState(TStaticRasterizerState<FM_Solid, CM_None>::GetRHI());
-					RHICmdList.SetDepthStencilState(TStaticDepthStencilState<false, CF_Always>::GetRHI());
+				RHICmdList.SetViewport(ScissorRect.Min.X, ScissorRect.Min.Y, 0.0f, ScissorRect.Max.X, ScissorRect.Max.Y, 1.0f);
+				RHICmdList.SetRasterizerState(TStaticRasterizerState<FM_Solid, CM_None>::GetRHI());
+				RHICmdList.SetDepthStencilState(TStaticDepthStencilState<false, CF_Always>::GetRHI());
 					
-					if (bDirectionalLight)
-					{
-						check(CascadeSettings.FadePlaneLength == 0);
-						// first cascade rendered or old method doesn't require fading (CO_Min is needed to combine multiple shadow passes)
-						// The ray traced cascade should always be first
-						RHICmdList.SetBlendState(TStaticBlendState<CW_RG, BO_Min, BF_One, BF_One>::GetRHI());
-					}
-					else
-					{
-						// use B and A in Light Attenuation
-						// (CO_Min is needed to combine multiple shadow passes)
-						RHICmdList.SetBlendState(TStaticBlendState<CW_BA, BO_Min, BF_One, BF_One, BO_Min, BF_One, BF_One>::GetRHI());
-					}
+				if (bDirectionalLight)
+				{
+					check(CascadeSettings.FadePlaneLength == 0);
+					// first cascade rendered or old method doesn't require fading (CO_Min is needed to combine multiple shadow passes)
+					// The ray traced cascade should always be first
+					RHICmdList.SetBlendState(TStaticBlendState<CW_RG, BO_Min, BF_One, BF_One>::GetRHI());
+				}
+				else
+				{
+					// use B and A in Light Attenuation
+					// (CO_Min is needed to combine multiple shadow passes)
+					RHICmdList.SetBlendState(TStaticBlendState<CW_BA, BO_Min, BF_One, BF_One, BO_Min, BF_One, BF_One>::GetRHI());
+				}
 
-					TShaderMapRef<FPostProcessVS> VertexShader(View.ShaderMap);
-					TShaderMapRef<FDistanceFieldShadowingUpsamplePS> PixelShader(View.ShaderMap);
+				TShaderMapRef<FPostProcessVS> VertexShader(View.ShaderMap);
+				TShaderMapRef<FDistanceFieldShadowingUpsamplePS> PixelShader(View.ShaderMap);
 
-					static FGlobalBoundShaderState BoundShaderState;
+				static FGlobalBoundShaderState BoundShaderState;
 
-					SetGlobalBoundShaderState(RHICmdList, View.GetFeatureLevel(), BoundShaderState, GFilterVertexDeclaration.VertexDeclarationRHI, *VertexShader, *PixelShader);
+				SetGlobalBoundShaderState(RHICmdList, View.GetFeatureLevel(), BoundShaderState, GFilterVertexDeclaration.VertexDeclarationRHI, *VertexShader, *PixelShader);
 
-					VertexShader->SetParameters(RHICmdList, View);
-					PixelShader->SetParameters(RHICmdList, View, this, ScissorRect, RayTracedShadowsRT);
+				VertexShader->SetParameters(RHICmdList, View);
+				PixelShader->SetParameters(RHICmdList, View, this, ScissorRect, RayTracedShadowsRT);
 
-					DrawRectangle( 
-						RHICmdList,
-						0, 0, 
-						ScissorRect.Width(), ScissorRect.Height(),
-						ScissorRect.Min.X / GAODownsampleFactor, ScissorRect.Min.Y / GAODownsampleFactor, 
-						ScissorRect.Width() / GAODownsampleFactor, ScissorRect.Height() / GAODownsampleFactor,
-						FIntPoint(ScissorRect.Width(), ScissorRect.Height()),
-						GetBufferSizeForAO(),
-						*VertexShader);
-				}
+				DrawRectangle( 
+					RHICmdList,
+					0, 0, 
+					ScissorRect.Width(), ScissorRect.Height(),
+					ScissorRect.Min.X / GAODownsampleFactor, ScissorRect.Min.Y / GAODownsampleFactor, 
+					ScissorRect.Width() / GAODownsampleFactor, ScissorRect.Height() / GAODownsampleFactor,
+					FIntPoint(ScissorRect.Width(), ScissorRect.Height()),
+					GetBufferSizeForAO(),
+					*VertexShader);
 			}
 		}
 	}
diff --git a/Engine/Source/Runtime/Renderer/Private/DistanceFieldSurfaceCacheLighting.cpp b/Engine/Source/Runtime/Renderer/Private/DistanceFieldSurfaceCacheLighting.cpp
--- a/Engine/Source/Runtime/Renderer/Private/DistanceFieldSurfaceCacheLighting.cpp
+++ b/Engine/Source/Runtime/Renderer/Private/DistanceFieldSurfaceCacheLighting.cpp
@@ -155,14 +155,6 @@
 	ECVF_Cheat | ECVF_RenderThreadSafe
 	);
 
-int32 GAOCPUFrustumCull = 1;
-FAutoConsoleVariableRef CVarAOCPUFrustumCull(
-	TEXT("r.AOCPUFrustumCull"),
-	GAOCPUFrustumCull,
-	TEXT("Whether to perform CPU frustum culling of occluder objects."),
-	ECVF_Cheat | ECVF_RenderThreadSafe
-	);
-
 float GAOMaxViewDistance = 10000;
 FAutoConsoleVariableRef CVarAOMaxViewDistance(
 	TEXT("r.AOMaxViewDistance"),
@@ -344,27 +336,89 @@
 
 // Only 7 needed atm, pad to cache line size
 // Must match equivalent shader defines
-int32 FDistanceFieldObjectBuffers::ObjectDataStride = 8;
-int32 FDistanceFieldObjectBuffers::ObjectData2Stride = 4;
-int32 FDistanceFieldObjectBuffers::ObjectBoxBoundsStride = 5;
+int32 FDistanceFieldObjectBuffers::ObjectDataStride = 12;
+int32 FDistanceFieldCulledObjectBuffers::ObjectDataStride = 8;
+int32 FDistanceFieldCulledObjectBuffers::ObjectBoxBoundsStride = 5;
 
-class FDistanceFieldObjectBufferResource : public FRenderResource
+TGlobalResource<FDistanceFieldObjectBufferResource> GAOCulledObjectBuffers;
+
+// In float4's.  Must match corresponding usf definition
+int32 UploadObjectDataStride = 1 + FDistanceFieldObjectBuffers::ObjectDataStride;
+
+class FDistanceFieldUploadDataResource : public FRenderResource
 {
 public:
-	FDistanceFieldObjectBuffers Buffers;
+
+	FCPUUpdatedBuffer UploadData;
+
+	FDistanceFieldUploadDataResource()
+	{
+		UploadData.Format = PF_A32B32G32R32F;
+		UploadData.Stride = UploadObjectDataStride;
+	}
 
 	virtual void InitDynamicRHI()  override
 	{
-		Buffers.Initialize();
+		UploadData.Initialize();
 	}
 
 	virtual void ReleaseDynamicRHI() override
 	{
-		Buffers.Release();
+		UploadData.Release();
 	}
 };
 
-TGlobalResource<FDistanceFieldObjectBufferResource> GAOObjectBuffers;
+TGlobalResource<FDistanceFieldUploadDataResource> GDistanceFieldUploadData;
 
+class FDistanceFieldUploadIndicesResource : public FRenderResource
+{
+public:
+
+	FCPUUpdatedBuffer UploadIndices;
+
+	FDistanceFieldUploadIndicesResource()
+	{
+		UploadIndices.Format = PF_R32_UINT;
+		UploadIndices.Stride = 1;
+	}
+
+	virtual void InitDynamicRHI()  override
+	{
+		UploadIndices.Initialize();
+	}
+
+	virtual void ReleaseDynamicRHI() override
+	{
+		UploadIndices.Release();
+	}
+};
+
+TGlobalResource<FDistanceFieldUploadIndicesResource> GDistanceFieldUploadIndices;
+
+class FDistanceFieldRemoveIndicesResource : public FRenderResource
+{
+public:
+
+	FCPUUpdatedBuffer RemoveIndices;
+
+	FDistanceFieldRemoveIndicesResource()
+	{
+		RemoveIndices.Format = PF_R32G32B32A32_UINT;
+		RemoveIndices.Stride = 1;
+	}
+
+	virtual void InitDynamicRHI()  override
+	{
+		RemoveIndices.Initialize();
+	}
+
+	virtual void ReleaseDynamicRHI() override
+	{
+		RemoveIndices.Release();
+	}
+};
+
+TGlobalResource<FDistanceFieldRemoveIndicesResource> GDistanceFieldRemoveIndices;
+
 void FTileIntersectionResources::InitDynamicRHI()
 {
@@ -375,7 +429,7 @@
 	TileHeadData.Initialize(sizeof(uint32)* 4, TileDimensions.X * TileDimensions.Y, PF_R32G32B32A32_UINT, BUF_Static);
 
 	//@todo - handle max exceeded
-	TileArrayData.Initialize(sizeof(uint16), GMaxNumObjectsPerTile * TileDimensions.X * TileDimensions.Y * 3, PF_R16_UINT, BUF_Static);
+	TileArrayData.Initialize(sizeof(uint32), GMaxNumObjectsPerTile * TileDimensions.X * TileDimensions.Y * 3, PF_R32_UINT, BUF_Static);
 	TileArrayNextAllocation.Initialize(sizeof(uint32), 1, PF_R32_UINT, BUF_Static);
 }
 
@@ -413,10 +467,11 @@
 	FConsoleCommandWithWorldDelegate::CreateStatic(OnClearSurfaceCache)
 	);
 
-/**  */
-class FBuildTileConesCS : public FGlobalShader
+uint32 UpdateObjectsGroupSize = 64;
+
+class FUploadObjectsToBufferCS : public FGlobalShader
 {
-	DECLARE_SHADER_TYPE(FBuildTileConesCS,Global)
+	DECLARE_SHADER_TYPE(FUploadObjectsToBufferCS,Global)
 public:
 
 	static bool ShouldCache(EShaderPlatform Platform)
@@ -427,147 +482,128 @@
 	static void ModifyCompilationEnvironment(EShaderPlatform Platform, FShaderCompilerEnvironment& OutEnvironment)
 	{
 		FGlobalShader::ModifyCompilationEnvironment(Platform,OutEnvironment);
-		OutEnvironment.SetDefine(TEXT("THREADGROUP_SIZEX"), GDistanceFieldAOTileSizeX);
-		OutEnvironment.SetDefine(TEXT("THREADGROUP_SIZEY"), GDistanceFieldAOTileSizeY);
-		OutEnvironment.SetDefine(TEXT("NUM_CONE_DIRECTIONS"), NumConeSampleDirections);
-		OutEnvironment.SetDefine(TEXT("DOWNSAMPLE_FACTOR"), GAODownsampleFactor);
-		OutEnvironment.SetDefine(TEXT("MAX_OBJECTS_PER_TILE"), GMaxNumObjectsPerTile);
-		
-		// To reduce shader compile time of compute shaders with shared memory, doesn't have an impact on generated code with current compiler (June 2010 DX SDK)
-		OutEnvironment.CompilerFlags.Add(CFLAG_StandardOptimization);
+		OutEnvironment.SetDefine(TEXT("UPDATEOBJECTS_THREADGROUP_SIZE"), UpdateObjectsGroupSize);
 	}
 
-	FBuildTileConesCS(const ShaderMetaType::CompiledShaderInitializerType& Initializer)
+	FUploadObjectsToBufferCS(const ShaderMetaType::CompiledShaderInitializerType& Initializer)
 		: FGlobalShader(Initializer)
 	{
-		DeferredParameters.Bind(Initializer.ParameterMap);
-		AOParameters.Bind(Initializer.ParameterMap);
-		TileConeAxisAndCos.Bind(Initializer.ParameterMap, TEXT("TileConeAxisAndCos"));
-		TileConeDepthRanges.Bind(Initializer.ParameterMap, TEXT("TileConeDepthRanges"));
-		TileHeadDataUnpacked.Bind(Initializer.ParameterMap, TEXT("TileHeadDataUnpacked"));
-		NumGroups.Bind(Initializer.ParameterMap, TEXT("NumGroups"));
-		ViewDimensionsParameter.Bind(Initializer.ParameterMap, TEXT("ViewDimensions"));
+		NumUploadOperations.Bind(Initializer.ParameterMap, TEXT("NumUploadOperations"));
+		UploadOperationIndices.Bind(Initializer.ParameterMap, TEXT("UploadOperationIndices"));
+		UploadOperationData.Bind(Initializer.ParameterMap, TEXT("UploadOperationData"));
+		ObjectBufferParameters.Bind(Initializer.ParameterMap);
 	}
 
-	FBuildTileConesCS()
+	FUploadObjectsToBufferCS()
 	{
 	}
-	void SetParameters(FRHICommandList& RHICmdList, const FSceneView& View, FScene* Scene, FVector2D NumGroupsValue, const FDistanceFieldAOParameters& Parameters)
+
+	void SetParameters(FRHICommandList& RHICmdList, const FScene* Scene, uint32 NumUploadOperationsValue, FShaderResourceViewRHIParamRef InUploadOperationIndices, FShaderResourceViewRHIParamRef InUploadOperationData)
 	{
 		FComputeShaderRHIParamRef ShaderRHI = GetComputeShader();
 
-		FGlobalShader::SetParameters(RHICmdList, ShaderRHI, View);
-		DeferredParameters.Set(RHICmdList, ShaderRHI, View);
-		AOParameters.Set(RHICmdList, ShaderRHI, Parameters);
+		SetShaderValue(RHICmdList, ShaderRHI, NumUploadOperations, NumUploadOperationsValue);
 
-		FTileIntersectionResources* TileIntersectionResources = ((FSceneViewState*)View.State)->AOTileIntersectionResources;
+		SetSRVParameter(RHICmdList, ShaderRHI, UploadOperationIndices, InUploadOperationIndices);
+		SetSRVParameter(RHICmdList, ShaderRHI, UploadOperationData, InUploadOperationData);
 
-		TileConeAxisAndCos.SetBuffer(RHICmdList, ShaderRHI, TileIntersectionResources->TileConeAxisAndCos);
-		TileConeDepthRanges.SetBuffer(RHICmdList, ShaderRHI, TileIntersectionResources->TileConeDepthRanges);
-		TileHeadDataUnpacked.SetBuffer(RHICmdList, ShaderRHI, TileIntersectionResources->TileHeadDataUnpacked);
-
-		SetShaderValue(RHICmdList, ShaderRHI, ViewDimensionsParameter, View.ViewRect);
-
-		SetShaderValue(RHICmdList, ShaderRHI, NumGroups, NumGroupsValue);
+		ObjectBufferParameters.Set(RHICmdList, ShaderRHI, *(Scene->DistanceFieldSceneData.ObjectBuffers), Scene->DistanceFieldSceneData.NumObjectsInBuffer);
 	}
 
 	void UnsetParameters(FRHICommandList& RHICmdList)
 	{
-		TileConeAxisAndCos.UnsetUAV(RHICmdList, GetComputeShader());
-		TileConeDepthRanges.UnsetUAV(RHICmdList, GetComputeShader());
-		TileHeadDataUnpacked.UnsetUAV(RHICmdList, GetComputeShader());
+		ObjectBufferParameters.UnsetParameters(RHICmdList, GetComputeShader());
 	}
 
 	virtual bool Serialize(FArchive& Ar)
 	{		
 		bool bShaderHasOutdatedParameters = FGlobalShader::Serialize(Ar);
-		Ar << DeferredParameters;
-		Ar << AOParameters;
-		Ar << TileConeAxisAndCos;
-		Ar << TileConeDepthRanges;
-		Ar << TileHeadDataUnpacked;
-		Ar << NumGroups;
-		Ar << ViewDimensionsParameter;
+		Ar << NumUploadOperations;
+		Ar << UploadOperationIndices;
+		Ar << UploadOperationData;
+		Ar << ObjectBufferParameters;
 		return bShaderHasOutdatedParameters;
 	}
 
 private:
 
-	FDeferredPixelShaderParameters DeferredParameters;
-	FAOParameters AOParameters;
-	FRWShaderParameter TileConeAxisAndCos;
-	FRWShaderParameter TileConeDepthRanges;
-	FRWShaderParameter TileHeadDataUnpacked;
-	FShaderParameter ViewDimensionsParameter;
-	FShaderParameter NumGroups;
+	FShaderParameter NumUploadOperations;
+	FShaderResourceParameter UploadOperationIndices;
+	FShaderResourceParameter UploadOperationData;
+	FDistanceFieldObjectBufferParameters ObjectBufferParameters;
 };
 
-IMPLEMENT_SHADER_TYPE(,FBuildTileConesCS,TEXT("DistanceFieldSurfaceCacheLightingCompute"),TEXT("BuildTileConesMain"),SF_Compute);
+IMPLEMENT_SHADER_TYPE(,FUploadObjectsToBufferCS,TEXT("DistanceFieldSurfaceCacheLightingCompute"),TEXT("UploadObjectsToBufferCS"),SF_Compute);
 
 
-/**  */
-class FObjectCullVS : public FGlobalShader
+
+class FCopyObjectBufferCS : public FGlobalShader
 {
-	DECLARE_SHADER_TYPE(FObjectCullVS,Global);
+	DECLARE_SHADER_TYPE(FCopyObjectBufferCS,Global)
 public:
 
-	static bool ShouldCache(EShaderPlatform Platform) 
-	{ 
-		return IsFeatureLevelSupported(Platform, ERHIFeatureLevel::SM5) && DoesPlatformSupportDistanceFieldAO(Platform); 
+	static bool ShouldCache(EShaderPlatform Platform)
+	{
+		return IsFeatureLevelSupported(Platform, ERHIFeatureLevel::SM5) && DoesPlatformSupportDistanceFieldAO(Platform);
 	}
 
 	static void ModifyCompilationEnvironment(EShaderPlatform Platform, FShaderCompilerEnvironment& OutEnvironment)
 	{
 		FGlobalShader::ModifyCompilationEnvironment(Platform,OutEnvironment);
+		OutEnvironment.SetDefine(TEXT("UPDATEOBJECTS_THREADGROUP_SIZE"), UpdateObjectsGroupSize);
 	}
 
-	FObjectCullVS(const ShaderMetaType::CompiledShaderInitializerType& Initializer):
-		FGlobalShader(Initializer)
+	FCopyObjectBufferCS(const ShaderMetaType::CompiledShaderInitializerType& Initializer)
+		: FGlobalShader(Initializer)
 	{
-		ObjectParameters.Bind(Initializer.ParameterMap);
-		AOParameters.Bind(Initializer.ParameterMap);
-		ConservativeRadiusScale.Bind(Initializer.ParameterMap, TEXT("ConservativeRadiusScale"));
+		CopyObjectBounds.Bind(Initializer.ParameterMap, TEXT("CopyObjectBounds"));
+		CopyObjectData.Bind(Initializer.ParameterMap, TEXT("CopyObjectData"));
+		ObjectBufferParameters.Bind(Initializer.ParameterMap);
 	}
 
-	FObjectCullVS() {}
+	FCopyObjectBufferCS()
+	{
+	}
 
-	void SetParameters(FRHICommandList& RHICmdList, const FSceneView& View, int32 NumObjects, const FDistanceFieldAOParameters& Parameters)
+	void SetParameters(FRHICommandList& RHICmdList, FDistanceFieldObjectBuffers& ObjectBuffersSource, FDistanceFieldObjectBuffers& ObjectBuffersDest, int32 NumObjectsValue)
 	{
-		const FVertexShaderRHIParamRef ShaderRHI = GetVertexShader();
-		FGlobalShader::SetParameters(RHICmdList, ShaderRHI, View);
-		
-		ObjectParameters.Set(RHICmdList, ShaderRHI, GAOObjectBuffers.Buffers, NumObjects);
-		AOParameters.Set(RHICmdList, ShaderRHI, Parameters);
+		FComputeShaderRHIParamRef ShaderRHI = GetComputeShader();
 
-		const int32 NumRings = StencilingGeometry::GLowPolyStencilSphereVertexBuffer.GetNumRings();
-		const float RadiansPerRingSegment = PI / (float)NumRings;
+		CopyObjectBounds.SetBuffer(RHICmdList, ShaderRHI, ObjectBuffersDest.Bounds);
+		CopyObjectData.SetBuffer(RHICmdList, ShaderRHI, ObjectBuffersDest.Data);
 
-		// Boost the effective radius so that the edges of the sphere approximation lie on the sphere, instead of the vertices
-		const float ConservativeRadiusScaleValue = 1.0f / FMath::Cos(RadiansPerRingSegment);
+		ObjectBufferParameters.Set(RHICmdList, ShaderRHI, ObjectBuffersSource, NumObjectsValue);
+	}
 
-		SetShaderValue(RHICmdList, ShaderRHI, ConservativeRadiusScale, ConservativeRadiusScaleValue);
+	void UnsetParameters(FRHICommandList& RHICmdList)
+	{
+		ObjectBufferParameters.UnsetParameters(RHICmdList, GetComputeShader());
+		CopyObjectBounds.UnsetUAV(RHICmdList, GetComputeShader());
+		CopyObjectData.UnsetUAV(RHICmdList, GetComputeShader());
 	}
 
 	virtual bool Serialize(FArchive& Ar)
-	{
+	{		
 		bool bShaderHasOutdatedParameters = FGlobalShader::Serialize(Ar);
-		Ar << ObjectParameters;
-		Ar << AOParameters;
-		Ar << ConservativeRadiusScale;
+		Ar << CopyObjectBounds;
+		Ar << CopyObjectData;
+		Ar << ObjectBufferParameters;
 		return bShaderHasOutdatedParameters;
 	}
 
 private:
-	FDistanceFieldObjectBufferParameters ObjectParameters;
-	FAOParameters AOParameters;
-	FShaderParameter ConservativeRadiusScale;
+
+	FRWShaderParameter CopyObjectBounds;
+	FRWShaderParameter CopyObjectData;
+	FDistanceFieldObjectBufferParameters ObjectBufferParameters;
 };
 
-IMPLEMENT_SHADER_TYPE(,FObjectCullVS,TEXT("DistanceFieldSurfaceCacheLighting"),TEXT("ObjectCullVS"),SF_Vertex);
+IMPLEMENT_SHADER_TYPE(,FCopyObjectBufferCS,TEXT("DistanceFieldSurfaceCacheLightingCompute"),TEXT("CopyObjectBufferCS"),SF_Compute);
 
-class FObjectCullPS : public FGlobalShader
+template<bool bRemoveFromSameBuffer>
+class TRemoveObjectsFromBufferCS : public FGlobalShader
 {
-	DECLARE_SHADER_TYPE(FObjectCullPS, Global);
+	DECLARE_SHADER_TYPE(TRemoveObjectsFromBufferCS,Global)
 public:
 
 	static bool ShouldCache(EShaderPlatform Platform)
@@ -578,89 +614,72 @@
 	static void ModifyCompilationEnvironment(EShaderPlatform Platform, FShaderCompilerEnvironment& OutEnvironment)
 	{
 		FGlobalShader::ModifyCompilationEnvironment(Platform,OutEnvironment);
-		
-		OutEnvironment.SetDefine(TEXT("DOWNSAMPLE_FACTOR"), GAODownsampleFactor);
-		OutEnvironment.SetDefine(TEXT("MAX_OBJECTS_PER_TILE"), GMaxNumObjectsPerTile);
+		OutEnvironment.SetDefine(TEXT("UPDATEOBJECTS_THREADGROUP_SIZE"), UpdateObjectsGroupSize);
+		OutEnvironment.SetDefine(TEXT("REMOVE_FROM_SAME_BUFFER"), bRemoveFromSameBuffer ? TEXT("1") : TEXT("0"));
 	}
 
-	/** Default constructor. */
-	FObjectCullPS() {}
-
-	/** Initialization constructor. */
-	FObjectCullPS(const ShaderMetaType::CompiledShaderInitializerType& Initializer)
+	TRemoveObjectsFromBufferCS(const ShaderMetaType::CompiledShaderInitializerType& Initializer)
 		: FGlobalShader(Initializer)
 	{
-		AOParameters.Bind(Initializer.ParameterMap);
-		TileHeadDataUnpacked.Bind(Initializer.ParameterMap, TEXT("TileHeadDataUnpacked"));
-		TileArrayData.Bind(Initializer.ParameterMap, TEXT("TileArrayData"));
-		TileConeAxisAndCos.Bind(Initializer.ParameterMap, TEXT("TileConeAxisAndCos"));
-		TileConeDepthRanges.Bind(Initializer.ParameterMap, TEXT("TileConeDepthRanges"));
-		NumGroups.Bind(Initializer.ParameterMap, TEXT("NumGroups"));
+		NumRemoveOperations.Bind(Initializer.ParameterMap, TEXT("NumRemoveOperations"));
+		RemoveOperationIndices.Bind(Initializer.ParameterMap, TEXT("RemoveOperationIndices"));
+		ObjectBufferParameters.Bind(Initializer.ParameterMap);
+		ObjectBounds2.Bind(Initializer.ParameterMap, TEXT("ObjectBounds2"));
+		ObjectData2.Bind(Initializer.ParameterMap, TEXT("ObjectData2"));
 	}
 
-	void SetParameters(FRHICommandList& RHICmdList, const FSceneView& View, FVector2D NumGroupsValue, const FDistanceFieldAOParameters& Parameters)
+	TRemoveObjectsFromBufferCS()
 	{
-		const FPixelShaderRHIParamRef ShaderRHI = GetPixelShader();
-		FGlobalShader::SetParameters(RHICmdList, ShaderRHI, View);
-		AOParameters.Set(RHICmdList, ShaderRHI, Parameters);
-
-		FTileIntersectionResources* TileIntersectionResources = ((FSceneViewState*)View.State)->AOTileIntersectionResources;
-
-		SetSRVParameter(RHICmdList, ShaderRHI, TileConeAxisAndCos, TileIntersectionResources->TileConeAxisAndCos.SRV);
-		SetSRVParameter(RHICmdList, ShaderRHI, TileConeDepthRanges, TileIntersectionResources->TileConeDepthRanges.SRV);
-
-		SetShaderValue(RHICmdList, ShaderRHI, NumGroups, NumGroupsValue);
 	}
 
-	void GetUAVs(const FSceneView& View, TArray<FUnorderedAccessViewRHIParamRef>& UAVs)
+	void SetParameters(
+		FRHICommandList& RHICmdList, 
+		const FScene* Scene, 
+		uint32 NumRemoveOperationsValue, 
+		FShaderResourceViewRHIParamRef InRemoveOperationIndices, 
+		FShaderResourceViewRHIParamRef InObjectBounds2, 
+		FShaderResourceViewRHIParamRef InObjectData2)
 	{
-		FTileIntersectionResources* TileIntersectionResources = ((FSceneViewState*)View.State)->AOTileIntersectionResources;
+		FComputeShaderRHIParamRef ShaderRHI = GetComputeShader();
 
-		int32 MaxIndex = FMath::Max(TileHeadDataUnpacked.GetUAVIndex(), TileArrayData.GetUAVIndex());
-		UAVs.AddZeroed(MaxIndex + 1);
-
-		if (TileHeadDataUnpacked.IsBound())
-		{
-			UAVs[TileHeadDataUnpacked.GetUAVIndex()] = TileIntersectionResources->TileHeadDataUnpacked.UAV;
-		}
-
-		if (TileArrayData.IsBound())
-		{
-			UAVs[TileArrayData.GetUAVIndex()] = TileIntersectionResources->TileArrayData.UAV;
-		}
-
-		check(UAVs.Num() > 0);
+		SetShaderValue(RHICmdList, ShaderRHI, NumRemoveOperations, NumRemoveOperationsValue);
+		SetSRVParameter(RHICmdList, ShaderRHI, RemoveOperationIndices, InRemoveOperationIndices);
+		ObjectBufferParameters.Set(RHICmdList, ShaderRHI, *(Scene->DistanceFieldSceneData.ObjectBuffers), Scene->DistanceFieldSceneData.NumObjectsInBuffer);
+		SetSRVParameter(RHICmdList, ShaderRHI, ObjectBounds2, InObjectBounds2);
+		SetSRVParameter(RHICmdList, ShaderRHI, ObjectData2, InObjectData2);
 	}
 
-	// FShader interface.
+	void UnsetParameters(FRHICommandList& RHICmdList)
+	{
+		ObjectBufferParameters.UnsetParameters(RHICmdList, GetComputeShader());
+	}
+
 	virtual bool Serialize(FArchive& Ar)
-	{
+	{		
 		bool bShaderHasOutdatedParameters = FGlobalShader::Serialize(Ar);
-		Ar << AOParameters;
-		Ar << TileHeadDataUnpacked;
-		Ar << TileArrayData;
-		Ar << TileConeAxisAndCos;
-		Ar << TileConeDepthRanges;
-		Ar << NumGroups;
+		Ar << NumRemoveOperations;
+		Ar << RemoveOperationIndices;
+		Ar << ObjectBufferParameters;
+		Ar << ObjectBounds2;
+		Ar << ObjectData2;
 		return bShaderHasOutdatedParameters;
 	}
 
 private:
-	FAOParameters AOParameters;
-	FRWShaderParameter TileHeadDataUnpacked;
-	FRWShaderParameter TileArrayData;
-	FShaderResourceParameter TileConeAxisAndCos;
-	FShaderResourceParameter TileConeDepthRanges;
-	FShaderParameter NumGroups;
+
+	FShaderParameter NumRemoveOperations;
+	FShaderResourceParameter RemoveOperationIndices;
+	FDistanceFieldObjectBufferParameters ObjectBufferParameters;
+	FShaderResourceParameter ObjectBounds2;
+	FShaderResourceParameter ObjectData2;
 };
 
-IMPLEMENT_SHADER_TYPE(,FObjectCullPS,TEXT("DistanceFieldSurfaceCacheLighting"),TEXT("ObjectCullPS"),SF_Pixel);
+IMPLEMENT_SHADER_TYPE(template<>,TRemoveObjectsFromBufferCS<true>,TEXT("DistanceFieldSurfaceCacheLightingCompute"),TEXT("RemoveObjectsFromBufferCS"),SF_Compute);
+IMPLEMENT_SHADER_TYPE(template<>,TRemoveObjectsFromBufferCS<false>,TEXT("DistanceFieldSurfaceCacheLightingCompute"),TEXT("RemoveObjectsFromBufferCS"),SF_Compute);
 
-
-/**  */
-class FDistanceFieldBuildTileListCS : public FGlobalShader
+class FCullObjectsForViewCS : public FGlobalShader
 {
-	DECLARE_SHADER_TYPE(FDistanceFieldBuildTileListCS,Global)
+	DECLARE_SHADER_TYPE(FCullObjectsForViewCS,Global)
 public:
 
 	static bool ShouldCache(EShaderPlatform Platform)
@@ -671,6 +690,703 @@
 	static void ModifyCompilationEnvironment(EShaderPlatform Platform, FShaderCompilerEnvironment& OutEnvironment)
 	{
 		FGlobalShader::ModifyCompilationEnvironment(Platform,OutEnvironment);
+		OutEnvironment.SetDefine(TEXT("UPDATEOBJECTS_THREADGROUP_SIZE"), UpdateObjectsGroupSize);
+	}
+
+	FCullObjectsForViewCS(const ShaderMetaType::CompiledShaderInitializerType& Initializer)
+		: FGlobalShader(Initializer)
+	{
+		ObjectBufferParameters.Bind(Initializer.ParameterMap);
+		ObjectIndirectArguments.Bind(Initializer.ParameterMap, TEXT("ObjectIndirectArguments"));
+		CulledObjectBounds.Bind(Initializer.ParameterMap, TEXT("CulledObjectBounds"));
+		CulledObjectData.Bind(Initializer.ParameterMap, TEXT("CulledObjectData"));
+		CulledObjectBoxBounds.Bind(Initializer.ParameterMap, TEXT("CulledObjectBoxBounds"));
+		AOParameters.Bind(Initializer.ParameterMap);
+		ViewFrustumConvexHull.Bind(Initializer.ParameterMap, TEXT("ViewFrustumConvexHull"));
+		ObjectBoundingGeometryIndexCount.Bind(Initializer.ParameterMap, TEXT("ObjectBoundingGeometryIndexCount"));
+	}
+
+	FCullObjectsForViewCS()
+	{
+	}
+
+	void SetParameters(FRHICommandList& RHICmdList, const FScene* Scene, const FSceneView& View, const FDistanceFieldAOParameters& Parameters)
+	{
+		FComputeShaderRHIParamRef ShaderRHI = GetComputeShader();
+		FGlobalShader::SetParameters(RHICmdList, ShaderRHI, View);
+		ObjectBufferParameters.Set(RHICmdList, ShaderRHI, *(Scene->DistanceFieldSceneData.ObjectBuffers), Scene->DistanceFieldSceneData.NumObjectsInBuffer);
+
+		ObjectIndirectArguments.SetBuffer(RHICmdList, ShaderRHI, GAOCulledObjectBuffers.Buffers.ObjectIndirectArguments);
+		CulledObjectBounds.SetBuffer(RHICmdList, ShaderRHI, GAOCulledObjectBuffers.Buffers.Bounds);
+		CulledObjectData.SetBuffer(RHICmdList, ShaderRHI, GAOCulledObjectBuffers.Buffers.Data);
+		CulledObjectBoxBounds.SetBuffer(RHICmdList, ShaderRHI, GAOCulledObjectBuffers.Buffers.BoxBounds);
+		AOParameters.Set(RHICmdList, ShaderRHI, Parameters);
+
+		// Shader assumes 5
+		check(View.ViewFrustum.Planes.Num() == 5);
+		SetShaderValueArray(RHICmdList, ShaderRHI, ViewFrustumConvexHull, View.ViewFrustum.Planes.GetData(), View.ViewFrustum.Planes.Num());
+		SetShaderValue(RHICmdList, ShaderRHI, ObjectBoundingGeometryIndexCount, StencilingGeometry::GLowPolyStencilSphereIndexBuffer.GetIndexCount());
+	}
+
+	void UnsetParameters(FRHICommandList& RHICmdList)
+	{
+		ObjectBufferParameters.UnsetParameters(RHICmdList, GetComputeShader());
+		ObjectIndirectArguments.UnsetUAV(RHICmdList, GetComputeShader());
+		CulledObjectBounds.UnsetUAV(RHICmdList, GetComputeShader());
+		CulledObjectData.UnsetUAV(RHICmdList, GetComputeShader());
+		CulledObjectBoxBounds.UnsetUAV(RHICmdList, GetComputeShader());
+	}
+
+	virtual bool Serialize(FArchive& Ar)
+	{		
+		bool bShaderHasOutdatedParameters = FGlobalShader::Serialize(Ar);
+		Ar << ObjectBufferParameters;
+		Ar << ObjectIndirectArguments;
+		Ar << CulledObjectBounds;
+		Ar << CulledObjectData;
+		Ar << CulledObjectBoxBounds;
+		Ar << AOParameters;
+		Ar << ViewFrustumConvexHull;
+		Ar << ObjectBoundingGeometryIndexCount;
+		return bShaderHasOutdatedParameters;
+	}
+
+private:
+
+	FDistanceFieldObjectBufferParameters ObjectBufferParameters;
+	FRWShaderParameter ObjectIndirectArguments;
+	FRWShaderParameter CulledObjectBounds;
+	FRWShaderParameter CulledObjectData;
+	FRWShaderParameter CulledObjectBoxBounds;
+	FAOParameters AOParameters;
+	FShaderParameter ViewFrustumConvexHull;
+	FShaderParameter ObjectBoundingGeometryIndexCount;
+};
+
+IMPLEMENT_SHADER_TYPE(,FCullObjectsForViewCS,TEXT("DistanceFieldSurfaceCacheLightingCompute"),TEXT("CullObjectsForViewCS"),SF_Compute);
+
+
+void FDeferredShadingSceneRenderer::UpdateGlobalDistanceFieldObjectBuffers(FRHICommandListImmediate& RHICmdList, bool bReallocatedAtlasLayouts)
+{
+	FDistanceFieldSceneData& DistanceFieldSceneData = Scene->DistanceFieldSceneData;
+
+	if (DistanceFieldSceneData.PendingAddOperations.Num() > 0 
+		|| DistanceFieldSceneData.PendingUpdateOperations.Num() > 0 
+		|| DistanceFieldSceneData.PendingRemoveOperations.Num() > 0
+		|| bReallocatedAtlasLayouts)
+	{
+		QUICK_SCOPE_CYCLE_COUNTER(STAT_UpdateObjectData);
+		SCOPED_DRAW_EVENT(RHICmdList, UpdateSceneObjectData);
+
+		if (!DistanceFieldSceneData.ObjectBuffers)
+		{
+			DistanceFieldSceneData.ObjectBuffers = new FDistanceFieldObjectBuffers();
+		}
+
+		if (bReallocatedAtlasLayouts)
+		{
+			for (int32 PrimitiveInstanceIndex = 0; PrimitiveInstanceIndex < DistanceFieldSceneData.PrimitiveInstanceMapping.Num(); PrimitiveInstanceIndex++)
+			{
+				FPrimitiveAndInstance& PrimitiveInstance = DistanceFieldSceneData.PrimitiveInstanceMapping[PrimitiveInstanceIndex];
+
+				// Queue an update of all primitives, since the atlas layout has changed
+				if (PrimitiveInstance.InstanceIndex == 0 
+					&& !DistanceFieldSceneData.PendingRemoveOperations.Contains(PrimitiveInstanceIndex)
+					&& !DistanceFieldSceneData.PendingAddOperations.Contains(PrimitiveInstance.Primitive)
+					&& !DistanceFieldSceneData.PendingUpdateOperations.Contains(PrimitiveInstance.Primitive))
+				{
+					DistanceFieldSceneData.PendingUpdateOperations.Add(PrimitiveInstance.Primitive);
+				}
+			}
+		}
+
+		TArray<uint32> UploadObjectIndices;
+		TArray<FVector4> UploadObjectData;
+
+		if (DistanceFieldSceneData.PendingAddOperations.Num() > 0 || DistanceFieldSceneData.PendingUpdateOperations.Num() > 0)
+		{
+			TArray<FMatrix> ObjectLocalToWorldTransforms;
+
+			const int32 NumUploadOperations = DistanceFieldSceneData.PendingAddOperations.Num() + DistanceFieldSceneData.PendingUpdateOperations.Num();
+			UploadObjectData.Empty(NumUploadOperations * UploadObjectDataStride);
+			UploadObjectIndices.Empty(NumUploadOperations);
+
+			const int32 NumTexelsOneDimX = GDistanceFieldVolumeTextureAtlas.GetSizeX();
+			const int32 NumTexelsOneDimY = GDistanceFieldVolumeTextureAtlas.GetSizeY();
+			const int32 NumTexelsOneDimZ = GDistanceFieldVolumeTextureAtlas.GetSizeZ();
+			const FVector InvTextureDim(1.0f / NumTexelsOneDimX, 1.0f / NumTexelsOneDimY, 1.0f / NumTexelsOneDimZ);
+
+			int32 OriginalNumObjects = DistanceFieldSceneData.NumObjectsInBuffer;
+
+			for (int32 UploadPrimitiveIndex = 0; UploadPrimitiveIndex < NumUploadOperations; UploadPrimitiveIndex++)
+			{
+				const bool bIsAddOperation = UploadPrimitiveIndex < DistanceFieldSceneData.PendingAddOperations.Num();
+
+				FPrimitiveSceneInfo* PrimitiveSceneInfo = bIsAddOperation 
+					? DistanceFieldSceneData.PendingAddOperations[UploadPrimitiveIndex] 
+				: DistanceFieldSceneData.PendingUpdateOperations[UploadPrimitiveIndex - DistanceFieldSceneData.PendingAddOperations.Num()];
+
+				ObjectLocalToWorldTransforms.Reset();
+
+				FBox LocalVolumeBounds;
+				FIntVector BlockMin;
+				FIntVector BlockSize;
+				bool bBuiltAsIfTwoSided;
+				bool bMeshWasPlane;
+				PrimitiveSceneInfo->Proxy->GetDistancefieldAtlasData(LocalVolumeBounds, BlockMin, BlockSize, bBuiltAsIfTwoSided, bMeshWasPlane, ObjectLocalToWorldTransforms);
+
+				if (BlockMin.X >= 0 && BlockMin.Y >= 0 && BlockMin.Z >= 0)
+				{
+					if (bIsAddOperation)
+					{
+						DistanceFieldSceneData.NumObjectsInBuffer += ObjectLocalToWorldTransforms.Num();
+						PrimitiveSceneInfo->DistanceFieldInstanceIndices.Empty(ObjectLocalToWorldTransforms.Num());
+					}
+
+					for (int32 TransformIndex = 0; TransformIndex < ObjectLocalToWorldTransforms.Num(); TransformIndex++)
+					{
+						const uint32 UploadIndex = bIsAddOperation ? OriginalNumObjects + UploadObjectIndices.Num() : PrimitiveSceneInfo->DistanceFieldInstanceIndices[TransformIndex];
+
+						if (bIsAddOperation)
+						{
+							const int32 AddIndex = OriginalNumObjects + UploadObjectIndices.Num();
+							DistanceFieldSceneData.PrimitiveInstanceMapping.Add(FPrimitiveAndInstance(PrimitiveSceneInfo, TransformIndex));
+							PrimitiveSceneInfo->DistanceFieldInstanceIndices.Add(AddIndex);
+						}
+
+						UploadObjectIndices.Add(UploadIndex);
+
+						FMatrix LocalToWorld = ObjectLocalToWorldTransforms[TransformIndex];
+
+						if (bMeshWasPlane)
+						{
+							FVector LocalScales = LocalToWorld.GetScaleVector();
+							FVector AbsLocalScales(FMath::Abs(LocalScales.X), FMath::Abs(LocalScales.Y), FMath::Abs(LocalScales.Z));
+							float MidScale = FMath::Min(AbsLocalScales.X, AbsLocalScales.Y);
+							float ScaleAdjust = FMath::Sign(LocalScales.Z) * MidScale / AbsLocalScales.Z;
+							// The mesh was determined to be a plane flat in Z during the build process, so we can change the Z scale
+							// Helps in cases with modular ground pieces with scales of (10, 10, 1) and some triangles just above Z=0
+							LocalToWorld.SetAxis(2, LocalToWorld.GetScaledAxis(EAxis::Z) * ScaleAdjust);
+						}
+
+						const FMatrix VolumeToWorld = FScaleMatrix(LocalVolumeBounds.GetExtent()) 
+							* FTranslationMatrix(LocalVolumeBounds.GetCenter())
+							* LocalToWorld;
+
+						const FVector4 ObjectBoundingSphere(VolumeToWorld.GetOrigin(), VolumeToWorld.GetScaleVector().Size());
+
+						UploadObjectData.Add(ObjectBoundingSphere);
+
+						const float MaxExtent = LocalVolumeBounds.GetExtent().GetMax();
+
+						const FMatrix UniformScaleVolumeToWorld = FScaleMatrix(MaxExtent) 
+							* FTranslationMatrix(LocalVolumeBounds.GetCenter())
+							* LocalToWorld;
+
+						const FVector InvBlockSize(1.0f / BlockSize.X, 1.0f / BlockSize.Y, 1.0f / BlockSize.Z);
+
+						//float3 VolumeUV = (VolumePosition / LocalPositionExtent * .5f * UVScale + .5f * UVScale + UVAdd;
+						const FVector LocalPositionExtent = LocalVolumeBounds.GetExtent() / FVector(MaxExtent);
+						const FVector UVScale = FVector(BlockSize) * InvTextureDim;
+						const float VolumeScale = UniformScaleVolumeToWorld.GetMaximumAxisScale();
+
+						const FMatrix WorldToVolume = UniformScaleVolumeToWorld.Inverse();
+						// WorldToVolume
+						UploadObjectData.Add(*(FVector4*)&WorldToVolume.M[0]);
+						UploadObjectData.Add(*(FVector4*)&WorldToVolume.M[1]);
+						UploadObjectData.Add(*(FVector4*)&WorldToVolume.M[2]);
+						UploadObjectData.Add(*(FVector4*)&WorldToVolume.M[3]);
+
+						// Clamp to texel center by subtracting a half texel in the [-1,1] position space
+						// LocalPositionExtent
+						UploadObjectData.Add(FVector4(LocalPositionExtent - InvBlockSize, LocalVolumeBounds.Min.X));
+
+						// UVScale
+						UploadObjectData.Add(FVector4(FVector(BlockSize) * InvTextureDim * .5f / LocalPositionExtent, VolumeScale));
+
+						// UVAdd
+						UploadObjectData.Add(FVector4(FVector(BlockMin) * InvTextureDim + .5f * UVScale, LocalVolumeBounds.Min.Y));
+
+						// Box bounds
+						UploadObjectData.Add(FVector4(LocalVolumeBounds.Max, LocalVolumeBounds.Min.Z));
+
+						UploadObjectData.Add(*(FVector4*)&LocalToWorld.M[0]);
+						UploadObjectData.Add(*(FVector4*)&LocalToWorld.M[1]);
+						UploadObjectData.Add(*(FVector4*)&LocalToWorld.M[2]);
+						UploadObjectData.Add(*(FVector4*)&LocalToWorld.M[3]);
+
+						checkSlow(UploadObjectData.Num() % UploadObjectDataStride == 0);
+					}
+				}
+			}
+
+			DistanceFieldSceneData.PendingAddOperations.Reset();
+			DistanceFieldSceneData.PendingUpdateOperations.Reset();
+
+			if (DistanceFieldSceneData.ObjectBuffers->MaxObjects < DistanceFieldSceneData.NumObjectsInBuffer)
+			{
+				if (DistanceFieldSceneData.ObjectBuffers->MaxObjects > 0)
+				{
+					// Realloc
+					FDistanceFieldObjectBuffers* NewObjectBuffers = new FDistanceFieldObjectBuffers();
+					NewObjectBuffers->MaxObjects = DistanceFieldSceneData.NumObjectsInBuffer * 5 / 4;
+					NewObjectBuffers->Initialize();
+
+					{
+						TShaderMapRef<FCopyObjectBufferCS> ComputeShader(GetGlobalShaderMap(Scene->GetFeatureLevel()));
+						RHICmdList.SetComputeShader(ComputeShader->GetComputeShader());
+						ComputeShader->SetParameters(RHICmdList, *(DistanceFieldSceneData.ObjectBuffers), *NewObjectBuffers, OriginalNumObjects);
+
+						DispatchComputeShader(RHICmdList, *ComputeShader, FMath::DivideAndRoundUp<uint32>(OriginalNumObjects, UpdateObjectsGroupSize), 1, 1);
+						ComputeShader->UnsetParameters(RHICmdList);
+					}
+
+					DistanceFieldSceneData.ObjectBuffers->Release();
+					delete DistanceFieldSceneData.ObjectBuffers;
+					DistanceFieldSceneData.ObjectBuffers = NewObjectBuffers;
+				}
+				else
+				{
+					// First time allocate
+					DistanceFieldSceneData.ObjectBuffers->MaxObjects = DistanceFieldSceneData.NumObjectsInBuffer * 5 / 4;
+					DistanceFieldSceneData.ObjectBuffers->Initialize();
+				}
+			}
+		}
+
+		if (UploadObjectIndices.Num() > 0)
+		{
+			if (UploadObjectIndices.Num() > GDistanceFieldUploadIndices.UploadIndices.MaxElements)
+			{
+				GDistanceFieldUploadIndices.UploadIndices.MaxElements = UploadObjectIndices.Num() * 5 / 4;
+				GDistanceFieldUploadIndices.UploadIndices.Release();
+				GDistanceFieldUploadIndices.UploadIndices.Initialize();
+
+				GDistanceFieldUploadData.UploadData.MaxElements = UploadObjectIndices.Num() * 5 / 4;
+				GDistanceFieldUploadData.UploadData.Release();
+				GDistanceFieldUploadData.UploadData.Initialize();
+			}
+
+			void* LockedBuffer = RHILockVertexBuffer(GDistanceFieldUploadIndices.UploadIndices.Buffer, 0, GDistanceFieldUploadIndices.UploadIndices.Buffer->GetSize(), RLM_WriteOnly);
+			FPlatformMemory::Memcpy(LockedBuffer, UploadObjectIndices.GetData(), UploadObjectIndices.GetTypeSize() * UploadObjectIndices.Num());
+			RHIUnlockVertexBuffer(GDistanceFieldUploadIndices.UploadIndices.Buffer);
+
+			LockedBuffer = RHILockVertexBuffer(GDistanceFieldUploadData.UploadData.Buffer, 0, GDistanceFieldUploadData.UploadData.Buffer->GetSize(), RLM_WriteOnly);
+			FPlatformMemory::Memcpy(LockedBuffer, UploadObjectData.GetData(), UploadObjectData.GetTypeSize() * UploadObjectData.Num());
+			RHIUnlockVertexBuffer(GDistanceFieldUploadData.UploadData.Buffer);
+
+			{
+				TShaderMapRef<FUploadObjectsToBufferCS> ComputeShader(GetGlobalShaderMap(Scene->GetFeatureLevel()));
+				RHICmdList.SetComputeShader(ComputeShader->GetComputeShader());
+				ComputeShader->SetParameters(RHICmdList, Scene, UploadObjectIndices.Num(), GDistanceFieldUploadIndices.UploadIndices.BufferSRV, GDistanceFieldUploadData.UploadData.BufferSRV);
+
+				DispatchComputeShader(RHICmdList, *ComputeShader, FMath::DivideAndRoundUp<uint32>(UploadObjectIndices.Num(), UpdateObjectsGroupSize), 1, 1);
+				ComputeShader->UnsetParameters(RHICmdList);
+			}
+		}
+
+		check(DistanceFieldSceneData.NumObjectsInBuffer == DistanceFieldSceneData.PrimitiveInstanceMapping.Num());
+
+		TArray<FIntRect> RemoveObjectIndices;
+		FDistanceFieldObjectBuffers* TemporaryCopySourceBuffers = NULL;
+
+		if (DistanceFieldSceneData.PendingRemoveOperations.Num() > 0)
+		{
+			check(DistanceFieldSceneData.NumObjectsInBuffer >= DistanceFieldSceneData.PendingRemoveOperations.Num());
+
+			// Sort from smallest to largest
+			DistanceFieldSceneData.PendingRemoveOperations.Sort();
+
+			// We have multiple remove requests enqueued in PendingRemoveOperations, can only use the RemoveAtSwap version when there won't be collisions
+			const bool bUseRemoveAtSwap = DistanceFieldSceneData.PendingRemoveOperations.Last() < DistanceFieldSceneData.NumObjectsInBuffer - DistanceFieldSceneData.PendingRemoveOperations.Num();
+
+			if (bUseRemoveAtSwap)
+			{
+				// Remove everything in parallel in the same buffer with a RemoveAtSwap algorithm
+				for (int32 RemovePrimitiveIndex = 0; RemovePrimitiveIndex < DistanceFieldSceneData.PendingRemoveOperations.Num(); RemovePrimitiveIndex++)
+				{
+					DistanceFieldSceneData.NumObjectsInBuffer--;
+					const int32 RemoveIndex = DistanceFieldSceneData.PendingRemoveOperations[RemovePrimitiveIndex];
+					const int32 MoveFromIndex = DistanceFieldSceneData.NumObjectsInBuffer;
+
+					check(RemoveIndex != MoveFromIndex);
+					// Queue a compute shader move
+					RemoveObjectIndices.Add(FIntRect(RemoveIndex, MoveFromIndex, 0, 0));
+
+					// Fixup indices of the primitive that is being moved
+					FPrimitiveAndInstance& PrimitiveAndInstanceBeingMoved = DistanceFieldSceneData.PrimitiveInstanceMapping[MoveFromIndex];
+					check(PrimitiveAndInstanceBeingMoved.Primitive && PrimitiveAndInstanceBeingMoved.Primitive->DistanceFieldInstanceIndices.Num() > 0);
+					PrimitiveAndInstanceBeingMoved.Primitive->DistanceFieldInstanceIndices[PrimitiveAndInstanceBeingMoved.InstanceIndex] = RemoveIndex;
+
+					DistanceFieldSceneData.PrimitiveInstanceMapping.RemoveAtSwap(RemoveIndex);
+				}
+			}
+			else
+			{
+				// Have to copy the object data to allow parallel removing
+				TemporaryCopySourceBuffers = DistanceFieldSceneData.ObjectBuffers;
+				DistanceFieldSceneData.ObjectBuffers = new FDistanceFieldObjectBuffers();
+				DistanceFieldSceneData.ObjectBuffers->MaxObjects = TemporaryCopySourceBuffers->MaxObjects;
+				DistanceFieldSceneData.ObjectBuffers->Initialize();
+
+				TArray<FPrimitiveAndInstance> OriginalPrimitiveInstanceMapping = DistanceFieldSceneData.PrimitiveInstanceMapping;
+				DistanceFieldSceneData.PrimitiveInstanceMapping.Reset();
+
+				const int32 NumDestObjects = DistanceFieldSceneData.NumObjectsInBuffer - DistanceFieldSceneData.PendingRemoveOperations.Num();
+				int32 SourceIndex = 0;
+				int32 NextPendingRemoveIndex = 0;
+
+				for (int32 DestinationIndex = 0; DestinationIndex < NumDestObjects; DestinationIndex++)
+				{
+					while (NextPendingRemoveIndex < DistanceFieldSceneData.PendingRemoveOperations.Num()
+						&& DistanceFieldSceneData.PendingRemoveOperations[NextPendingRemoveIndex] == SourceIndex)
+					{
+						NextPendingRemoveIndex++;
+						SourceIndex++;
+					}
+
+					// Queue a compute shader move
+					RemoveObjectIndices.Add(FIntRect(DestinationIndex, SourceIndex, 0, 0));
+
+					// Fixup indices of the primitive that is being moved
+					FPrimitiveAndInstance& PrimitiveAndInstanceBeingMoved = OriginalPrimitiveInstanceMapping[SourceIndex];
+					check(PrimitiveAndInstanceBeingMoved.Primitive && PrimitiveAndInstanceBeingMoved.Primitive->DistanceFieldInstanceIndices.Num() > 0);
+					PrimitiveAndInstanceBeingMoved.Primitive->DistanceFieldInstanceIndices[PrimitiveAndInstanceBeingMoved.InstanceIndex] = DestinationIndex;
+
+					check(DistanceFieldSceneData.PrimitiveInstanceMapping.Num() == DestinationIndex);
+					DistanceFieldSceneData.PrimitiveInstanceMapping.Add(PrimitiveAndInstanceBeingMoved);
+
+					SourceIndex++;
+				}
+
+				DistanceFieldSceneData.NumObjectsInBuffer = NumDestObjects;
+
+
+				/*
+				// Have to remove one at a time while any entries to remove are at the end of the buffer
+				DistanceFieldSceneData.NumObjectsInBuffer--;
+				const int32 RemoveIndex = DistanceFieldSceneData.PendingRemoveOperations[ParallelConflictIndex];
+				const int32 MoveFromIndex = DistanceFieldSceneData.NumObjectsInBuffer;
+
+				if (RemoveIndex != MoveFromIndex)
+				{
+					// Queue a compute shader move
+					RemoveObjectIndices.Add(FIntRect(RemoveIndex, MoveFromIndex, 0, 0));
+
+					// Fixup indices of the primitive that is being moved
+					FPrimitiveAndInstance& PrimitiveAndInstanceBeingMoved = DistanceFieldSceneData.PrimitiveInstanceMapping[MoveFromIndex];
+					check(PrimitiveAndInstanceBeingMoved.Primitive && PrimitiveAndInstanceBeingMoved.Primitive->DistanceFieldInstanceIndices.Num() > 0);
+					PrimitiveAndInstanceBeingMoved.Primitive->DistanceFieldInstanceIndices[PrimitiveAndInstanceBeingMoved.InstanceIndex] = RemoveIndex;
+				}
+
+				DistanceFieldSceneData.PrimitiveInstanceMapping.RemoveAtSwap(RemoveIndex);
+				DistanceFieldSceneData.PendingRemoveOperations.RemoveAtSwap(ParallelConflictIndex);
+				*/
+			}
+
+			DistanceFieldSceneData.PendingRemoveOperations.Reset();
+
+			if (RemoveObjectIndices.Num() > 0)
+			{
+				if (RemoveObjectIndices.Num() > GDistanceFieldRemoveIndices.RemoveIndices.MaxElements)
+				{
+					GDistanceFieldRemoveIndices.RemoveIndices.MaxElements = RemoveObjectIndices.Num() * 5 / 4;
+					GDistanceFieldRemoveIndices.RemoveIndices.Release();
+					GDistanceFieldRemoveIndices.RemoveIndices.Initialize();
+				}
+
+				void* LockedBuffer = RHILockVertexBuffer(GDistanceFieldRemoveIndices.RemoveIndices.Buffer, 0, GDistanceFieldRemoveIndices.RemoveIndices.Buffer->GetSize(), RLM_WriteOnly);
+				FPlatformMemory::Memcpy(LockedBuffer, RemoveObjectIndices.GetData(), RemoveObjectIndices.GetTypeSize() * RemoveObjectIndices.Num());
+				RHIUnlockVertexBuffer(GDistanceFieldRemoveIndices.RemoveIndices.Buffer);
+
+				if (bUseRemoveAtSwap)
+				{
+					check(!TemporaryCopySourceBuffers);
+					TShaderMapRef<TRemoveObjectsFromBufferCS<true> > ComputeShader(GetGlobalShaderMap(Scene->GetFeatureLevel()));
+					RHICmdList.SetComputeShader(ComputeShader->GetComputeShader());
+					ComputeShader->SetParameters(RHICmdList, Scene, RemoveObjectIndices.Num(), GDistanceFieldRemoveIndices.RemoveIndices.BufferSRV, NULL, NULL);
+
+					DispatchComputeShader(RHICmdList, *ComputeShader, FMath::DivideAndRoundUp<uint32>(RemoveObjectIndices.Num(), UpdateObjectsGroupSize), 1, 1);
+					ComputeShader->UnsetParameters(RHICmdList);
+				}
+				else
+				{
+					check(TemporaryCopySourceBuffers);
+					TShaderMapRef<TRemoveObjectsFromBufferCS<false> > ComputeShader(GetGlobalShaderMap(Scene->GetFeatureLevel()));
+					RHICmdList.SetComputeShader(ComputeShader->GetComputeShader());
+					ComputeShader->SetParameters(RHICmdList, Scene, RemoveObjectIndices.Num(), GDistanceFieldRemoveIndices.RemoveIndices.BufferSRV, TemporaryCopySourceBuffers->Bounds.SRV, TemporaryCopySourceBuffers->Data.SRV);
+
+					DispatchComputeShader(RHICmdList, *ComputeShader, FMath::DivideAndRoundUp<uint32>(RemoveObjectIndices.Num(), UpdateObjectsGroupSize), 1, 1);
+					ComputeShader->UnsetParameters(RHICmdList);
+
+					TemporaryCopySourceBuffers->Release();
+					delete TemporaryCopySourceBuffers;
+				}
+			}
+		}
+
+		DistanceFieldSceneData.VerifyIntegrity();
+	}
+}
+
+/**  */
+class FBuildTileConesCS : public FGlobalShader
+{
+	DECLARE_SHADER_TYPE(FBuildTileConesCS,Global)
+public:
+
+	static bool ShouldCache(EShaderPlatform Platform)
+	{
+		return IsFeatureLevelSupported(Platform, ERHIFeatureLevel::SM5) && DoesPlatformSupportDistanceFieldAO(Platform);
+	}
+
+	static void ModifyCompilationEnvironment(EShaderPlatform Platform, FShaderCompilerEnvironment& OutEnvironment)
+	{
+		FGlobalShader::ModifyCompilationEnvironment(Platform,OutEnvironment);
+		OutEnvironment.SetDefine(TEXT("THREADGROUP_SIZEX"), GDistanceFieldAOTileSizeX);
+		OutEnvironment.SetDefine(TEXT("THREADGROUP_SIZEY"), GDistanceFieldAOTileSizeY);
+		OutEnvironment.SetDefine(TEXT("NUM_CONE_DIRECTIONS"), NumConeSampleDirections);
+		OutEnvironment.SetDefine(TEXT("DOWNSAMPLE_FACTOR"), GAODownsampleFactor);
+		OutEnvironment.SetDefine(TEXT("MAX_OBJECTS_PER_TILE"), GMaxNumObjectsPerTile);
+		
+		// To reduce shader compile time of compute shaders with shared memory, doesn't have an impact on generated code with current compiler (June 2010 DX SDK)
+		OutEnvironment.CompilerFlags.Add(CFLAG_StandardOptimization);
+	}
+
+	FBuildTileConesCS(const ShaderMetaType::CompiledShaderInitializerType& Initializer)
+		: FGlobalShader(Initializer)
+	{
+		DeferredParameters.Bind(Initializer.ParameterMap);
+		AOParameters.Bind(Initializer.ParameterMap);
+		TileConeAxisAndCos.Bind(Initializer.ParameterMap, TEXT("TileConeAxisAndCos"));
+		TileConeDepthRanges.Bind(Initializer.ParameterMap, TEXT("TileConeDepthRanges"));
+		TileHeadDataUnpacked.Bind(Initializer.ParameterMap, TEXT("TileHeadDataUnpacked"));
+		NumGroups.Bind(Initializer.ParameterMap, TEXT("NumGroups"));
+		ViewDimensionsParameter.Bind(Initializer.ParameterMap, TEXT("ViewDimensions"));
+	}
+
+	FBuildTileConesCS()
+	{
+	}
+	void SetParameters(FRHICommandList& RHICmdList, const FSceneView& View, FScene* Scene, FVector2D NumGroupsValue, const FDistanceFieldAOParameters& Parameters)
+	{
+		FComputeShaderRHIParamRef ShaderRHI = GetComputeShader();
+
+		FGlobalShader::SetParameters(RHICmdList, ShaderRHI, View);
+		DeferredParameters.Set(RHICmdList, ShaderRHI, View);
+		AOParameters.Set(RHICmdList, ShaderRHI, Parameters);
+
+		FTileIntersectionResources* TileIntersectionResources = ((FSceneViewState*)View.State)->AOTileIntersectionResources;
+
+		TileConeAxisAndCos.SetBuffer(RHICmdList, ShaderRHI, TileIntersectionResources->TileConeAxisAndCos);
+		TileConeDepthRanges.SetBuffer(RHICmdList, ShaderRHI, TileIntersectionResources->TileConeDepthRanges);
+		TileHeadDataUnpacked.SetBuffer(RHICmdList, ShaderRHI, TileIntersectionResources->TileHeadDataUnpacked);
+
+		SetShaderValue(RHICmdList, ShaderRHI, ViewDimensionsParameter, View.ViewRect);
+
+		SetShaderValue(RHICmdList, ShaderRHI, NumGroups, NumGroupsValue);
+	}
+
+	void UnsetParameters(FRHICommandList& RHICmdList)
+	{
+		TileConeAxisAndCos.UnsetUAV(RHICmdList, GetComputeShader());
+		TileConeDepthRanges.UnsetUAV(RHICmdList, GetComputeShader());
+		TileHeadDataUnpacked.UnsetUAV(RHICmdList, GetComputeShader());
+	}
+
+	virtual bool Serialize(FArchive& Ar)
+	{		
+		bool bShaderHasOutdatedParameters = FGlobalShader::Serialize(Ar);
+		Ar << DeferredParameters;
+		Ar << AOParameters;
+		Ar << TileConeAxisAndCos;
+		Ar << TileConeDepthRanges;
+		Ar << TileHeadDataUnpacked;
+		Ar << NumGroups;
+		Ar << ViewDimensionsParameter;
+		return bShaderHasOutdatedParameters;
+	}
+
+private:
+
+	FDeferredPixelShaderParameters DeferredParameters;
+	FAOParameters AOParameters;
+	FRWShaderParameter TileConeAxisAndCos;
+	FRWShaderParameter TileConeDepthRanges;
+	FRWShaderParameter TileHeadDataUnpacked;
+	FShaderParameter ViewDimensionsParameter;
+	FShaderParameter NumGroups;
+};
+
+IMPLEMENT_SHADER_TYPE(,FBuildTileConesCS,TEXT("DistanceFieldSurfaceCacheLightingCompute"),TEXT("BuildTileConesMain"),SF_Compute);
+
+
+/**  */
+class FObjectCullVS : public FGlobalShader
+{
+	DECLARE_SHADER_TYPE(FObjectCullVS,Global);
+public:
+
+	static bool ShouldCache(EShaderPlatform Platform) 
+	{ 
+		return IsFeatureLevelSupported(Platform, ERHIFeatureLevel::SM5) && DoesPlatformSupportDistanceFieldAO(Platform); 
+	}
+
+	static void ModifyCompilationEnvironment(EShaderPlatform Platform, FShaderCompilerEnvironment& OutEnvironment)
+	{
+		FGlobalShader::ModifyCompilationEnvironment(Platform,OutEnvironment);
+	}
+
+	FObjectCullVS(const ShaderMetaType::CompiledShaderInitializerType& Initializer):
+		FGlobalShader(Initializer)
+	{
+		ObjectParameters.Bind(Initializer.ParameterMap);
+		AOParameters.Bind(Initializer.ParameterMap);
+		ConservativeRadiusScale.Bind(Initializer.ParameterMap, TEXT("ConservativeRadiusScale"));
+	}
+
+	FObjectCullVS() {}
+
+	void SetParameters(FRHICommandList& RHICmdList, const FSceneView& View, const FDistanceFieldAOParameters& Parameters)
+	{
+		const FVertexShaderRHIParamRef ShaderRHI = GetVertexShader();
+		FGlobalShader::SetParameters(RHICmdList, ShaderRHI, View);
+		
+		ObjectParameters.Set(RHICmdList, ShaderRHI, GAOCulledObjectBuffers.Buffers);
+		AOParameters.Set(RHICmdList, ShaderRHI, Parameters);
+
+		const int32 NumRings = StencilingGeometry::GLowPolyStencilSphereVertexBuffer.GetNumRings();
+		const float RadiansPerRingSegment = PI / (float)NumRings;
+
+		// Boost the effective radius so that the edges of the sphere approximation lie on the sphere, instead of the vertices
+		const float ConservativeRadiusScaleValue = 1.0f / FMath::Cos(RadiansPerRingSegment);
+
+		SetShaderValue(RHICmdList, ShaderRHI, ConservativeRadiusScale, ConservativeRadiusScaleValue);
+	}
+
+	virtual bool Serialize(FArchive& Ar)
+	{
+		bool bShaderHasOutdatedParameters = FGlobalShader::Serialize(Ar);
+		Ar << ObjectParameters;
+		Ar << AOParameters;
+		Ar << ConservativeRadiusScale;
+		return bShaderHasOutdatedParameters;
+	}
+
+private:
+	FDistanceFieldCulledObjectBufferParameters ObjectParameters;
+	FAOParameters AOParameters;
+	FShaderParameter ConservativeRadiusScale;
+};
+
+IMPLEMENT_SHADER_TYPE(,FObjectCullVS,TEXT("DistanceFieldSurfaceCacheLighting"),TEXT("ObjectCullVS"),SF_Vertex);
+
+class FObjectCullPS : public FGlobalShader
+{
+	DECLARE_SHADER_TYPE(FObjectCullPS, Global);
+public:
+
+	static bool ShouldCache(EShaderPlatform Platform)
+	{
+		return IsFeatureLevelSupported(Platform, ERHIFeatureLevel::SM5) && DoesPlatformSupportDistanceFieldAO(Platform);
+	}
+
+	static void ModifyCompilationEnvironment(EShaderPlatform Platform, FShaderCompilerEnvironment& OutEnvironment)
+	{
+		FGlobalShader::ModifyCompilationEnvironment(Platform,OutEnvironment);
+		
+		OutEnvironment.SetDefine(TEXT("DOWNSAMPLE_FACTOR"), GAODownsampleFactor);
+		OutEnvironment.SetDefine(TEXT("MAX_OBJECTS_PER_TILE"), GMaxNumObjectsPerTile);
+	}
+
+	/** Default constructor. */
+	FObjectCullPS() {}
+
+	/** Initialization constructor. */
+	FObjectCullPS(const ShaderMetaType::CompiledShaderInitializerType& Initializer)
+		: FGlobalShader(Initializer)
+	{
+		AOParameters.Bind(Initializer.ParameterMap);
+		TileHeadDataUnpacked.Bind(Initializer.ParameterMap, TEXT("TileHeadDataUnpacked"));
+		TileArrayData.Bind(Initializer.ParameterMap, TEXT("TileArrayData"));
+		TileConeAxisAndCos.Bind(Initializer.ParameterMap, TEXT("TileConeAxisAndCos"));
+		TileConeDepthRanges.Bind(Initializer.ParameterMap, TEXT("TileConeDepthRanges"));
+		NumGroups.Bind(Initializer.ParameterMap, TEXT("NumGroups"));
+	}
+
+	void SetParameters(FRHICommandList& RHICmdList, const FSceneView& View, FVector2D NumGroupsValue, const FDistanceFieldAOParameters& Parameters)
+	{
+		const FPixelShaderRHIParamRef ShaderRHI = GetPixelShader();
+		FGlobalShader::SetParameters(RHICmdList, ShaderRHI, View);
+		AOParameters.Set(RHICmdList, ShaderRHI, Parameters);
+
+		FTileIntersectionResources* TileIntersectionResources = ((FSceneViewState*)View.State)->AOTileIntersectionResources;
+
+		SetSRVParameter(RHICmdList, ShaderRHI, TileConeAxisAndCos, TileIntersectionResources->TileConeAxisAndCos.SRV);
+		SetSRVParameter(RHICmdList, ShaderRHI, TileConeDepthRanges, TileIntersectionResources->TileConeDepthRanges.SRV);
+
+		SetShaderValue(RHICmdList, ShaderRHI, NumGroups, NumGroupsValue);
+	}
+
+	void GetUAVs(const FSceneView& View, TArray<FUnorderedAccessViewRHIParamRef>& UAVs)
+	{
+		FTileIntersectionResources* TileIntersectionResources = ((FSceneViewState*)View.State)->AOTileIntersectionResources;
+
+		int32 MaxIndex = FMath::Max(TileHeadDataUnpacked.GetUAVIndex(), TileArrayData.GetUAVIndex());
+		UAVs.AddZeroed(MaxIndex + 1);
+
+		if (TileHeadDataUnpacked.IsBound())
+		{
+			UAVs[TileHeadDataUnpacked.GetUAVIndex()] = TileIntersectionResources->TileHeadDataUnpacked.UAV;
+		}
+
+		if (TileArrayData.IsBound())
+		{
+			UAVs[TileArrayData.GetUAVIndex()] = TileIntersectionResources->TileArrayData.UAV;
+		}
+
+		check(UAVs.Num() > 0);
+	}
+
+	// FShader interface.
+	virtual bool Serialize(FArchive& Ar)
+	{
+		bool bShaderHasOutdatedParameters = FGlobalShader::Serialize(Ar);
+		Ar << AOParameters;
+		Ar << TileHeadDataUnpacked;
+		Ar << TileArrayData;
+		Ar << TileConeAxisAndCos;
+		Ar << TileConeDepthRanges;
+		Ar << NumGroups;
+		return bShaderHasOutdatedParameters;
+	}
+
+private:
+	FAOParameters AOParameters;
+	FRWShaderParameter TileHeadDataUnpacked;
+	FRWShaderParameter TileArrayData;
+	FShaderResourceParameter TileConeAxisAndCos;
+	FShaderResourceParameter TileConeDepthRanges;
+	FShaderParameter NumGroups;
+};
+
+IMPLEMENT_SHADER_TYPE(,FObjectCullPS,TEXT("DistanceFieldSurfaceCacheLighting"),TEXT("ObjectCullPS"),SF_Pixel);
+
+
+/**  */
+class FDistanceFieldBuildTileListCS : public FGlobalShader
+{
+	DECLARE_SHADER_TYPE(FDistanceFieldBuildTileListCS,Global)
+public:
+
+	static bool ShouldCache(EShaderPlatform Platform)
+	{
+		return IsFeatureLevelSupported(Platform, ERHIFeatureLevel::SM5) && DoesPlatformSupportDistanceFieldAO(Platform);
+	}
+
+	static void ModifyCompilationEnvironment(EShaderPlatform Platform, FShaderCompilerEnvironment& OutEnvironment)
+	{
+		FGlobalShader::ModifyCompilationEnvironment(Platform,OutEnvironment);
 		OutEnvironment.SetDefine(TEXT("THREADGROUP_SIZEX"), GDistanceFieldAOTileSizeX);
 		OutEnvironment.SetDefine(TEXT("THREADGROUP_SIZEY"), GDistanceFieldAOTileSizeY);
 		OutEnvironment.SetDefine(TEXT("NUM_CONE_DIRECTIONS"), NumConeSampleDirections);
@@ -696,13 +1412,13 @@
 	FDistanceFieldBuildTileListCS()
 	{
 	}
-	void SetParameters(FRHICommandList& RHICmdList, const FSceneView& View, FScene* Scene, FVector2D NumGroupsValue, int32 NumObjects, const FDistanceFieldAOParameters& Parameters)
+	void SetParameters(FRHICommandList& RHICmdList, const FSceneView& View, FScene* Scene, FVector2D NumGroupsValue, const FDistanceFieldAOParameters& Parameters)
 	{
 		FComputeShaderRHIParamRef ShaderRHI = GetComputeShader();
 
 		FGlobalShader::SetParameters(RHICmdList, ShaderRHI, View);
 		DeferredParameters.Set(RHICmdList, ShaderRHI, View);
-		ObjectParameters.Set(RHICmdList, ShaderRHI, GAOObjectBuffers.Buffers, NumObjects);
+		ObjectParameters.Set(RHICmdList, ShaderRHI, GAOCulledObjectBuffers.Buffers);
 		AOParameters.Set(RHICmdList, ShaderRHI, Parameters);
 
 		FTileIntersectionResources* TileIntersectionResources = ((FSceneViewState*)View.State)->AOTileIntersectionResources;
@@ -740,7 +1456,7 @@
 private:
 
 	FDeferredPixelShaderParameters DeferredParameters;
-	FDistanceFieldObjectBufferParameters ObjectParameters;
+	FDistanceFieldCulledObjectBufferParameters ObjectParameters;
 	FAOParameters AOParameters;
 	FRWShaderParameter TileHeadData;
 	FRWShaderParameter TileArrayData;
@@ -825,12 +1541,12 @@
 		AOParameters.Bind(Initializer.ParameterMap);
 	}
 
-	void SetParameters(FRHICommandList& RHICmdList, const FSceneView& View, int32 NumObjects, const FDistanceFieldAOParameters& Parameters)
+	void SetParameters(FRHICommandList& RHICmdList, const FSceneView& View, const FDistanceFieldAOParameters& Parameters)
 	{
 		const FPixelShaderRHIParamRef ShaderRHI = GetPixelShader();
 
 		FGlobalShader::SetParameters(RHICmdList, ShaderRHI, View);
-		ObjectParameters.Set(RHICmdList, ShaderRHI, GAOObjectBuffers.Buffers, NumObjects);
+		ObjectParameters.Set(RHICmdList, ShaderRHI, GAOCulledObjectBuffers.Buffers);
 		AOParameters.Set(RHICmdList, ShaderRHI, Parameters);
 		DeferredParameters.Set(RHICmdList, ShaderRHI, View);
 
@@ -858,7 +1574,7 @@
 
 private:
 
-	FDistanceFieldObjectBufferParameters ObjectParameters;
+	FDistanceFieldCulledObjectBufferParameters ObjectParameters;
 	FDeferredPixelShaderParameters DeferredParameters;
 	FShaderResourceParameter TileHeadData;
 	FShaderResourceParameter TileHeadDataUnpacked;
@@ -906,7 +1622,7 @@
 		AOParameters.Bind(Initializer.ParameterMap);
 	}
 
-	void SetParameters(FRHICommandList& RHICmdList, const FSceneView& View, FSceneRenderTargetItem& DistanceFieldNormalValue, int32 NumObjects, const FDistanceFieldAOParameters& Parameters)
+	void SetParameters(FRHICommandList& RHICmdList, const FSceneView& View, FSceneRenderTargetItem& DistanceFieldNormalValue, const FDistanceFieldAOParameters& Parameters)
 	{
 		const FComputeShaderRHIParamRef ShaderRHI = GetComputeShader();
 
@@ -914,7 +1630,7 @@
 
 		DistanceFieldNormal.SetTexture(RHICmdList, ShaderRHI, DistanceFieldNormalValue.ShaderResourceTexture, DistanceFieldNormalValue.UAV);
 
-		ObjectParameters.Set(RHICmdList, ShaderRHI, GAOObjectBuffers.Buffers, NumObjects);
+		ObjectParameters.Set(RHICmdList, ShaderRHI, GAOCulledObjectBuffers.Buffers);
 		AOParameters.Set(RHICmdList, ShaderRHI, Parameters);
 		DeferredParameters.Set(RHICmdList, ShaderRHI, View);
 
@@ -953,7 +1669,7 @@
 private:
 
 	FRWShaderParameter DistanceFieldNormal;
-	FDistanceFieldObjectBufferParameters ObjectParameters;
+	FDistanceFieldCulledObjectBufferParameters ObjectParameters;
 	FShaderParameter UVToTileHead;
 	FDeferredPixelShaderParameters DeferredParameters;
 	FShaderResourceParameter TileHeadData;
@@ -965,7 +1681,7 @@
 
 IMPLEMENT_SHADER_TYPE(,FComputeDistanceFieldNormalCS,TEXT("DistanceFieldSurfaceCacheLighting"),TEXT("ComputeDistanceFieldNormalCS"),SF_Compute);
 
-void ComputeDistanceFieldNormal(FRHICommandListImmediate& RHICmdList, const TArray<FViewInfo>& Views, FSceneRenderTargetItem& DistanceFieldNormal, int32 NumObjects, const FDistanceFieldAOParameters& Parameters)
+void ComputeDistanceFieldNormal(FRHICommandListImmediate& RHICmdList, const TArray<FViewInfo>& Views, FSceneRenderTargetItem& DistanceFieldNormal, const FDistanceFieldAOParameters& Parameters)
 {
 	if (GAOComputeShaderNormalCalculation)
 	{
@@ -983,7 +1699,7 @@
 				TShaderMapRef<FComputeDistanceFieldNormalCS> ComputeShader(View.ShaderMap);
 
 				RHICmdList.SetComputeShader(ComputeShader->GetComputeShader());
-				ComputeShader->SetParameters(RHICmdList, View, DistanceFieldNormal, NumObjects, Parameters);
+				ComputeShader->SetParameters(RHICmdList, View, DistanceFieldNormal, Parameters);
 				DispatchComputeShader(RHICmdList, *ComputeShader, GroupSizeX, GroupSizeY, 1);
 
 				ComputeShader->UnsetParameters(RHICmdList);
@@ -1012,7 +1728,7 @@
 			
 			SetGlobalBoundShaderState(RHICmdList, View.FeatureLevel, BoundShaderState, GFilterVertexDeclaration.VertexDeclarationRHI, *VertexShader, *PixelShader);
 
-			PixelShader->SetParameters(RHICmdList, View, NumObjects, Parameters);
+			PixelShader->SetParameters(RHICmdList, View, Parameters);
 
 			DrawRectangle( 
 				RHICmdList,
@@ -1578,7 +2294,6 @@
 		const FSceneView& View, 
 		FSceneRenderTargetItem& DistanceFieldAOIrradianceCacheSplat, 
 		FSceneRenderTargetItem& DistanceFieldNormal, 
-		int32 NumObjects, 
 		int32 DownsampleFactorValue, 
 		int32 DepthLevel, 
 		FIntPoint TileListGroupSizeValue, 
@@ -1587,7 +2302,7 @@
 		FComputeShaderRHIParamRef ShaderRHI = GetComputeShader();
 		FGlobalShader::SetParameters(RHICmdList, ShaderRHI, View);
 		DeferredParameters.Set(RHICmdList, ShaderRHI, View);
-		ObjectParameters.Set(RHICmdList, ShaderRHI, GAOObjectBuffers.Buffers, NumObjects);
+		ObjectParameters.Set(RHICmdList, ShaderRHI, GAOCulledObjectBuffers.Buffers);
 		AOParameters.Set(RHICmdList, ShaderRHI, Parameters);
 		AOLevelParameters.Set(RHICmdList, ShaderRHI, View, DownsampleFactorValue);
 
@@ -1713,7 +2428,7 @@
 private:
 
 	FDeferredPixelShaderParameters DeferredParameters;
-	FDistanceFieldObjectBufferParameters ObjectParameters;
+	FDistanceFieldCulledObjectBufferParameters ObjectParameters;
 	FAOParameters AOParameters;
 	FAOLevelParameters AOLevelParameters;
 	FShaderResourceParameter IrradianceCacheNormal;
@@ -3093,7 +3808,7 @@
 	}
 }
 
-FIntPoint BuildTileObjectLists(FRHICommandListImmediate& RHICmdList, FScene* Scene, TArray<FViewInfo>& Views, int32 NumObjects, const FDistanceFieldAOParameters& Parameters)
+FIntPoint BuildTileObjectLists(FRHICommandListImmediate& RHICmdList, FScene* Scene, TArray<FViewInfo>& Views, const FDistanceFieldAOParameters& Parameters)
 {
 	SCOPED_DRAW_EVENT(RHICmdList, BuildTileList);
 	SetRenderTarget(RHICmdList, NULL, NULL);
@@ -3161,20 +3876,16 @@
 				
 				SetGlobalBoundShaderState(RHICmdList, View.GetFeatureLevel(), BoundShaderState, GetVertexDeclarationFVector4(), *VertexShader, *PixelShader);
 
-				VertexShader->SetParameters(RHICmdList, View, NumObjects, Parameters);
+				VertexShader->SetParameters(RHICmdList, View, Parameters);
 				PixelShader->SetParameters(RHICmdList, View, FVector2D(GroupSizeX, GroupSizeY), Parameters);
 
 				RHICmdList.SetStreamSource(0, StencilingGeometry::GLowPolyStencilSphereVertexBuffer.VertexBufferRHI, sizeof(FVector4), 0);
 
-				RHICmdList.DrawIndexedPrimitive(
+				RHICmdList.DrawIndexedPrimitiveIndirect(
+					PT_TriangleList,
 					StencilingGeometry::GLowPolyStencilSphereIndexBuffer.IndexBufferRHI, 
-					PT_TriangleList, 
-					0,
-					0, 
-					StencilingGeometry::GLowPolyStencilSphereVertexBuffer.GetVertexCount(), 
-					0, 
-					StencilingGeometry::GLowPolyStencilSphereIndexBuffer.GetIndexCount() / 3, 
-					NumObjects);
+					GAOCulledObjectBuffers.Buffers.ObjectIndirectArguments.Buffer,
+					0);
 			}
 		}
 	}
@@ -3214,7 +3925,7 @@
 			TShaderMapRef<FDistanceFieldBuildTileListCS > ComputeShader(View.ShaderMap);
 
 			RHICmdList.SetComputeShader(ComputeShader->GetComputeShader());
-			ComputeShader->SetParameters(RHICmdList, View, Scene, FVector2D(GroupSizeX, GroupSizeY), NumObjects, Parameters);
+			ComputeShader->SetParameters(RHICmdList, View, Scene, FVector2D(GroupSizeX, GroupSizeY), Parameters);
 			DispatchComputeShader(RHICmdList, *ComputeShader, GroupSizeX, GroupSizeY, 1);
 
 			ComputeShader->UnsetParameters(RHICmdList);
@@ -3522,142 +4233,20 @@
 	}
 }
 
-void UpdateVisibleObjectBuffers(const FScene* Scene, const FDistanceFieldAOParameters& Parameters, const FViewInfo& View, int32& NumObjects)
+bool SupportsDistanceFieldAO(ERHIFeatureLevel::Type FeatureLevel, EShaderPlatform ShaderPlatform)
 {
-	//@todo - scene rendering allocator
-	TArray<FVector4> ObjectBoundsData;
-	TArray<FVector4> ObjectData;
-	TArray<FVector4> ObjectData2;
-
-	{
-		QUICK_SCOPE_CYCLE_COUNTER(STAT_GatherObjectData);
-
-		const FVector ViewOriginForDistanceCulling = View.ViewMatrices.ViewOrigin;
-		const int32 NumTexelsOneDimX = GDistanceFieldVolumeTextureAtlas.GetSizeX();
-		const int32 NumTexelsOneDimY = GDistanceFieldVolumeTextureAtlas.GetSizeY();
-		const int32 NumTexelsOneDimZ = GDistanceFieldVolumeTextureAtlas.GetSizeZ();
-		const FVector InvTextureDim(1.0f / NumTexelsOneDimX, 1.0f / NumTexelsOneDimY, 1.0f / NumTexelsOneDimZ);
-
-		ObjectBoundsData.Empty(Scene->Primitives.Num() / 2);
-		ObjectData.Empty(FDistanceFieldObjectBuffers::ObjectDataStride * Scene->Primitives.Num() / 2);
-		ObjectData2.Empty(FDistanceFieldObjectBuffers::ObjectData2Stride * Scene->Primitives.Num() / 2);
-
-		//@todo - dense data-coherent pass like frustum culling in InitViews
-		// FTileIntersectionResources uses 16 bit object indices
-		for (int32 PrimitiveIndex = 0; PrimitiveIndex < Scene->Primitives.Num() && NumObjects < MAX_uint16; PrimitiveIndex++)
-		{
-			const FPrimitiveSceneInfo* PrimitiveSceneInfo = Scene->Primitives[PrimitiveIndex];
-
-			if (PrimitiveSceneInfo->Proxy->CastsDynamicShadow() || PrimitiveSceneInfo->Proxy->CastsStaticShadow())
-			{
-				FBox LocalVolumeBounds;
-				FIntVector BlockMin;
-				FIntVector BlockSize;
-				bool bBuiltAsIfTwoSided;
-				bool bMeshWasPlane;
-				PrimitiveSceneInfo->Proxy->GetDistancefieldAtlasData(LocalVolumeBounds, BlockMin, BlockSize, bBuiltAsIfTwoSided, bMeshWasPlane);
-
-				if (BlockMin.X >= 0 && BlockMin.Y >= 0 && BlockMin.Z >= 0)
-				{
-					FMatrix LocalToWorld = PrimitiveSceneInfo->Proxy->GetLocalToWorld();
-
-					if (bMeshWasPlane)
-					{
-						FVector LocalScales = LocalToWorld.GetScaleVector();
-						FVector AbsLocalScales(FMath::Abs(LocalScales.X), FMath::Abs(LocalScales.Y), FMath::Abs(LocalScales.Z));
-						float MidScale = FMath::Min(AbsLocalScales.X, AbsLocalScales.Y);
-						float ScaleAdjust = FMath::Sign(LocalScales.Z) * MidScale / AbsLocalScales.Z;
-						// The mesh was determined to be a plane flat in Z during the build process, so we can change the Z scale
-						// Helps in cases with modular ground pieces with scales of (10, 10, 1) and some triangles just above Z=0
-						LocalToWorld.SetAxis(2, LocalToWorld.GetScaledAxis(EAxis::Z) * ScaleAdjust);
-					}
-
-					const FMatrix VolumeToWorld = FScaleMatrix(LocalVolumeBounds.GetExtent()) 
-						* FTranslationMatrix(LocalVolumeBounds.GetCenter())
-						* LocalToWorld;
-
-					const FVector4 ObjectBoundingSphere(VolumeToWorld.GetOrigin(), VolumeToWorld.GetScaleVector().Size());
-					const float DistanceToViewSq = (ViewOriginForDistanceCulling - (FVector)ObjectBoundingSphere).SizeSquared();
-
-					if (!GAOCPUFrustumCull 
-						|| (DistanceToViewSq < FMath::Square(GAOMaxViewDistance + ObjectBoundingSphere.W)
-						&& View.ViewFrustum.IntersectSphere((FVector)ObjectBoundingSphere, ObjectBoundingSphere.W + Parameters.OcclusionMaxDistance)))
-					{
-						ObjectBoundsData.Add(ObjectBoundingSphere);
-
-						const float MaxExtent = LocalVolumeBounds.GetExtent().GetMax();
-
-						const FMatrix UniformScaleVolumeToWorld = FScaleMatrix(MaxExtent) 
-							* FTranslationMatrix(LocalVolumeBounds.GetCenter())
-							* LocalToWorld;
-
-						const FVector InvBlockSize(1.0f / BlockSize.X, 1.0f / BlockSize.Y, 1.0f / BlockSize.Z);
-
-						//float3 VolumeUV = (VolumePosition / LocalPositionExtent * .5f * UVScale + .5f * UVScale + UVAdd;
-						const FVector LocalPositionExtent = LocalVolumeBounds.GetExtent() / FVector(MaxExtent);
-						const FVector UVScale = FVector(BlockSize) * InvTextureDim;
-						const float VolumeScale = UniformScaleVolumeToWorld.GetMaximumAxisScale();
-						ObjectData2.Add(*(FVector4*)&UniformScaleVolumeToWorld.M[0]);
-						ObjectData2.Add(*(FVector4*)&UniformScaleVolumeToWorld.M[1]);
-						ObjectData2.Add(*(FVector4*)&UniformScaleVolumeToWorld.M[2]);
-						ObjectData2.Add(*(FVector4*)&UniformScaleVolumeToWorld.M[3]);
-
-						const FMatrix WorldToVolume = UniformScaleVolumeToWorld.Inverse();
-						// WorldToVolume
-						ObjectData.Add(*(FVector4*)&WorldToVolume.M[0]);
-						ObjectData.Add(*(FVector4*)&WorldToVolume.M[1]);
-						ObjectData.Add(*(FVector4*)&WorldToVolume.M[2]);
-						ObjectData.Add(*(FVector4*)&WorldToVolume.M[3]);
-
-						// Clamp to texel center by subtracting a half texel in the [-1,1] position space
-						// LocalPositionExtent
-						ObjectData.Add(FVector4(LocalPositionExtent - InvBlockSize, bBuiltAsIfTwoSided ? 1.0f : 0.0f));
-
-						// UVScale
-						ObjectData.Add(FVector4(FVector(BlockSize) * InvTextureDim * .5f / LocalPositionExtent, VolumeScale));
-
-						// UVAdd
-						ObjectData.Add(FVector4(FVector(BlockMin) * InvTextureDim + .5f * UVScale, 0));
-
-						// Padding
-						ObjectData.Add(FVector4(0, 0, 0, 0));
-
-						checkSlow(ObjectData.Num() % FDistanceFieldObjectBuffers::ObjectDataStride == 0);
-						checkSlow(ObjectData2.Num() % FDistanceFieldObjectBuffers::ObjectData2Stride == 0);
-
-						NumObjects++;
-					}
-				}
-			}
-		}
-	}
-
-	if (NumObjects > 0)
-	{
-		QUICK_SCOPE_CYCLE_COUNTER(STAT_UpdateObjectBuffers);
-
-		if (NumObjects > GAOObjectBuffers.Buffers.MaxObjects)
-		{
-			// Allocate with slack
-			GAOObjectBuffers.Buffers.MaxObjects = NumObjects * 5 / 4;
-			GAOObjectBuffers.Buffers.Release();
-			GAOObjectBuffers.Buffers.Initialize();
-		}
-
-		checkSlow(NumObjects < MAX_uint16);
-
-		void* LockedBuffer = RHILockVertexBuffer(GAOObjectBuffers.Buffers.ObjectData.Bounds, 0, GAOObjectBuffers.Buffers.ObjectData.Bounds->GetSize(), RLM_WriteOnly);
-		FPlatformMemory::Memcpy(LockedBuffer, ObjectBoundsData.GetData(), ObjectBoundsData.GetTypeSize() * ObjectBoundsData.Num());
-		RHIUnlockVertexBuffer(GAOObjectBuffers.Buffers.ObjectData.Bounds);
-
-		LockedBuffer = RHILockVertexBuffer(GAOObjectBuffers.Buffers.ObjectData.Data, 0, GAOObjectBuffers.Buffers.ObjectData.Data->GetSize(), RLM_WriteOnly);
-		FPlatformMemory::Memcpy(LockedBuffer, ObjectData.GetData(), ObjectData.GetTypeSize() * ObjectData.Num());
-		RHIUnlockVertexBuffer(GAOObjectBuffers.Buffers.ObjectData.Data);
-
-		LockedBuffer = RHILockVertexBuffer(GAOObjectBuffers.Buffers.ObjectData.Data2, 0, GAOObjectBuffers.Buffers.ObjectData.Data2->GetSize(), RLM_WriteOnly);
-		FPlatformMemory::Memcpy(LockedBuffer, ObjectData2.GetData(), ObjectData2.GetTypeSize() * ObjectData2.Num());
-		RHIUnlockVertexBuffer(GAOObjectBuffers.Buffers.ObjectData.Data2);
-	}
+	return GDistanceFieldAO 
+		&& FeatureLevel >= ERHIFeatureLevel::SM5
+		&& DoesPlatformSupportDistanceFieldAO(ShaderPlatform);
 }
 
+bool FDeferredShadingSceneRenderer::ShouldPrepareForDistanceFieldAO() const
+{
+	return SupportsDistanceFieldAO(Scene->GetFeatureLevel(), Scene->GetShaderPlatform())
+		&& ViewFamily.EngineShowFlags.DistanceFieldAO
+		&& ((ShouldRenderDynamicSkyLight() && Scene->SkyLight->bCastShadows)
+			|| ViewFamily.EngineShowFlags.VisualizeMeshDistanceFields
+			|| ViewFamily.EngineShowFlags.VisualizeDistanceFieldAO);
+}
+
 bool FDeferredShadingSceneRenderer::RenderDistanceFieldAOSurfaceCache(
@@ -3670,9 +4259,7 @@
 	//@todo - support multiple views
 	const FViewInfo& View = Views[0];
 
-	if (GDistanceFieldAO 
-		&& FeatureLevel >= ERHIFeatureLevel::SM5
-		&& DoesPlatformSupportDistanceFieldAO(View.GetShaderPlatform())
+	if (SupportsDistanceFieldAO(View.GetFeatureLevel(), View.GetShaderPlatform())
 		&& Views.Num() == 1
 		// ViewState is used to cache tile intersection resources which have to be sized based on the view
 		&& View.State)
@@ -3680,11 +4267,10 @@
 		QUICK_SCOPE_CYCLE_COUNTER(STAT_RenderDistanceFieldAOSurfaceCache);
 		SCOPED_DRAW_EVENT(RHICmdList, DistanceFieldAO);
 
-		// Update the global distance field atlas
-		GDistanceFieldVolumeTextureAtlas.UpdateAllocations();
-
-		if (GDistanceFieldVolumeTextureAtlas.VolumeTextureRHI)
+		if (GDistanceFieldVolumeTextureAtlas.VolumeTextureRHI && Scene->DistanceFieldSceneData.NumObjectsInBuffer)
 		{
+			check(!Scene->DistanceFieldSceneData.HasPendingOperations());
+
 			GenerateBestSpacedVectors();
 
 			// Create surface cache state that will persist across frames if needed
@@ -3714,135 +4300,56 @@
 				SurfaceCacheResources.bClearedResources = true;
 			}
 
-			int32 NumObjects = 0;
+			QUICK_SCOPE_CYCLE_COUNTER(STAT_AOIssueGPUWork);
 
-			// Update object buffers with objects that have a valid distance field representation and may affect pixels on the screen
-			UpdateVisibleObjectBuffers(Scene, Parameters, View, NumObjects);
+			{
+				SCOPED_DRAW_EVENT(RHICmdList, ObjectFrustumCulling);
 
-			if (NumObjects > 0)
-			{
-				QUICK_SCOPE_CYCLE_COUNTER(STAT_AOIssueGPUWork);
-
-				// Intersect objects with screen tiles, build lists
-				FIntPoint TileListGroupSize = BuildTileObjectLists(RHICmdList, Scene, Views, NumObjects, Parameters);
-
-				TRefCountPtr<IPooledRenderTarget> DistanceFieldNormal;
+				if (GAOCulledObjectBuffers.Buffers.MaxObjects < Scene->DistanceFieldSceneData.NumObjectsInBuffer
+					|| GAOCulledObjectBuffers.Buffers.MaxObjects > 3 * Scene->DistanceFieldSceneData.NumObjectsInBuffer)
+				{
+					GAOCulledObjectBuffers.Buffers.MaxObjects = Scene->DistanceFieldSceneData.NumObjectsInBuffer * 5 / 4;
+					GAOCulledObjectBuffers.Buffers.Release();
+					GAOCulledObjectBuffers.Buffers.Initialize();
+				}
 
 				{
-					const FIntPoint BufferSize = GetBufferSizeForAO();
-					FPooledRenderTargetDesc Desc(FPooledRenderTargetDesc::Create2DDesc(BufferSize, PF_FloatRGBA, TexCreate_None, TexCreate_RenderTargetable | TexCreate_UAV, false));
-					GRenderTargetPool.FindFreeElement(Desc, DistanceFieldNormal, TEXT("DistanceFieldNormal"));
+					TShaderMapRef<FCullObjectsForViewCS> ComputeShader(GetGlobalShaderMap(Scene->GetFeatureLevel()));
+					RHICmdList.SetComputeShader(ComputeShader->GetComputeShader());
+					ComputeShader->SetParameters(RHICmdList, Scene, View, Parameters);
+
+					DispatchComputeShader(RHICmdList, *ComputeShader, FMath::DivideAndRoundUp<uint32>(Scene->DistanceFieldSceneData.NumObjectsInBuffer, UpdateObjectsGroupSize), 1, 1);
+					ComputeShader->UnsetParameters(RHICmdList);
 				}
+			}
 
-				// Compute the distance field normal, this is used for surface caching instead of the GBuffer normal because it matches the occluding geometry
-				ComputeDistanceFieldNormal(RHICmdList, Views, DistanceFieldNormal->GetRenderTargetItem(), NumObjects, Parameters);
+			// Intersect objects with screen tiles, build lists
+			FIntPoint TileListGroupSize = BuildTileObjectLists(RHICmdList, Scene, Views, Parameters);
 
-				GRenderTargetPool.VisualizeTexture.SetCheckPoint(RHICmdList, DistanceFieldNormal);
+			TRefCountPtr<IPooledRenderTarget> DistanceFieldNormal;
 
-				if (GAOReuseAcrossFrames)
-				{
-					SCOPED_DRAW_EVENT(RHICmdList, TrimRecords);
+			{
+				const FIntPoint BufferSize = GetBufferSizeForAO();
+				FPooledRenderTargetDesc Desc(FPooledRenderTargetDesc::Create2DDesc(BufferSize, PF_FloatRGBA, TexCreate_None, TexCreate_RenderTargetable | TexCreate_UAV, false));
+				GRenderTargetPool.FindFreeElement(Desc, DistanceFieldNormal, TEXT("DistanceFieldNormal"));
+			}
 
-					// Copy and trim last frame's surface cache samples
-					for (int32 DepthLevel = GAOMaxLevel; DepthLevel >= GAOMinLevel; DepthLevel--)
-					{
-						if (GetTrimOldRecordsFraction(DepthLevel) > 0)
-						{
-							{	
-								TShaderMapRef<FSetupCopyIndirectArgumentsCS> ComputeShader(View.ShaderMap);
+			// Compute the distance field normal, this is used for surface caching instead of the GBuffer normal because it matches the occluding geometry
+			ComputeDistanceFieldNormal(RHICmdList, Views, DistanceFieldNormal->GetRenderTargetItem(), Parameters);
 
-								RHICmdList.SetComputeShader(ComputeShader->GetComputeShader());
-								ComputeShader->SetParameters(RHICmdList, View, DepthLevel);
-								DispatchComputeShader(RHICmdList, *ComputeShader, 1, 1, 1);
+			GRenderTargetPool.VisualizeTexture.SetCheckPoint(RHICmdList, DistanceFieldNormal);
 
-								ComputeShader->UnsetParameters(RHICmdList);
-							}
+			if (GAOReuseAcrossFrames)
+			{
+				SCOPED_DRAW_EVENT(RHICmdList, TrimRecords);
 
-							if (GDistanceFieldGI)
-							{
-								TShaderMapRef<TCopyIrradianceCacheSamplesCS<true> > ComputeShader(View.ShaderMap);
-
-								RHICmdList.SetComputeShader(ComputeShader->GetComputeShader());
-								ComputeShader->SetParameters(RHICmdList, View, DepthLevel);
-								DispatchIndirectComputeShader(RHICmdList, *ComputeShader, SurfaceCacheResources.DispatchParameters.Buffer, 0);
-
-								ComputeShader->UnsetParameters(RHICmdList);
-							}
-							else
-							{
-								TShaderMapRef<TCopyIrradianceCacheSamplesCS<false> > ComputeShader(View.ShaderMap);
-
-								RHICmdList.SetComputeShader(ComputeShader->GetComputeShader());
-								ComputeShader->SetParameters(RHICmdList, View, DepthLevel);
-								DispatchIndirectComputeShader(RHICmdList, *ComputeShader, SurfaceCacheResources.DispatchParameters.Buffer, 0);
-
-								ComputeShader->UnsetParameters(RHICmdList);
-							}
-
-							Swap(SurfaceCacheResources.Level[DepthLevel], SurfaceCacheResources.TempResources);
-						}
-					}
-				}
-
-				// Start from the highest depth, which will consider the fewest potential shading points
-				// Each level potentially prevents the next higher resolution level from doing expensive shading work
-				// This is the core of the surface cache algorithm
+				// Copy and trim last frame's surface cache samples
 				for (int32 DepthLevel = GAOMaxLevel; DepthLevel >= GAOMinLevel; DepthLevel--)
 				{
-					int32 DestLevelDownsampleFactor = GAODownsampleFactor * (1 << (DepthLevel * GAOPowerOfTwoBetweenLevels));
-
-					SCOPED_DRAW_EVENTF(RHICmdList, Level, TEXT("Level_%d"), DepthLevel);
-
-					TRefCountPtr<IPooledRenderTarget> DistanceFieldAOBentNormalSplat;
-
+					if (GetTrimOldRecordsFraction(DepthLevel) > 0)
 					{
-						FIntPoint AOBufferSize = FIntPoint::DivideAndRoundUp(GSceneRenderTargets.GetBufferSizeXY(), DestLevelDownsampleFactor);
-						FPooledRenderTargetDesc Desc(FPooledRenderTargetDesc::Create2DDesc(AOBufferSize, PF_FloatRGBA, TexCreate_None, TexCreate_RenderTargetable | TexCreate_UAV, false));
-						GRenderTargetPool.FindFreeElement(Desc, DistanceFieldAOBentNormalSplat, TEXT("DistanceFieldAOBentNormalSplat"));
-					}
-
-					// Splat / interpolate the surface cache records onto the buffer sized for the current depth level
-					RenderIrradianceCacheInterpolation(RHICmdList, Views, DistanceFieldAOBentNormalSplat, NULL, DistanceFieldNormal->GetRenderTargetItem(), DepthLevel, DestLevelDownsampleFactor, Parameters, false);
-
-					{
-						SCOPED_DRAW_EVENT(RHICmdList, PopulateIrradianceCache);
-						SetRenderTarget(RHICmdList, NULL, NULL);
-
-						// Save off the current record count before adding any more
 						{	
-							TShaderMapRef<FSaveStartIndexCS> ComputeShader(View.ShaderMap);
+							TShaderMapRef<FSetupCopyIndirectArgumentsCS> ComputeShader(View.ShaderMap);
 
 							RHICmdList.SetComputeShader(ComputeShader->GetComputeShader());
 							ComputeShader->SetParameters(RHICmdList, View, DepthLevel);
-							DispatchComputeShader(RHICmdList, *ComputeShader, 1, 1, 1);
-
-							ComputeShader->UnsetParameters(RHICmdList);
-						}
-
-						// Create new records which haven't been shaded yet for shading points which don't have a valid interpolation from existing records
-						for (int32 ViewIndex = 0; ViewIndex < Views.Num(); ViewIndex++)
-						{
-							const FViewInfo& View = Views[ViewIndex];
-
-							FIntPoint DownsampledViewSize = FIntPoint::DivideAndRoundUp(View.ViewRect.Size(), DestLevelDownsampleFactor);
-							uint32 GroupSizeX = FMath::DivideAndRoundUp(DownsampledViewSize.X, GDistanceFieldAOTileSizeX);
-							uint32 GroupSizeY = FMath::DivideAndRoundUp(DownsampledViewSize.Y, GDistanceFieldAOTileSizeY);
-
-							TShaderMapRef<FPopulateCacheCS> ComputeShader(View.ShaderMap);
-
-							RHICmdList.SetComputeShader(ComputeShader->GetComputeShader());
-							ComputeShader->SetParameters(RHICmdList, View, DistanceFieldAOBentNormalSplat->GetRenderTargetItem(), DistanceFieldNormal->GetRenderTargetItem(), DestLevelDownsampleFactor, DepthLevel, TileListGroupSize, Parameters);
-							DispatchComputeShader(RHICmdList, *ComputeShader, GroupSizeX, GroupSizeY, 1);
-
-							ComputeShader->UnsetParameters(RHICmdList);
-						}
-					}
-
-					{
-						SCOPED_DRAW_EVENT(RHICmdList, ShadeIrradianceCache);
-
-						{	
-							TShaderMapRef<FSetupFinalGatherIndirectArgumentsCS> ComputeShader(View.ShaderMap);
-
-							RHICmdList.SetComputeShader(ComputeShader->GetComputeShader());
-							ComputeShader->SetParameters(RHICmdList, View, DepthLevel);
@@ -3851,327 +4358,185 @@
 							ComputeShader->UnsetParameters(RHICmdList);
 						}
 
-						// Compute lighting for the new surface cache records by cone-stepping through the object distance fields
 						if (GDistanceFieldGI)
 						{
-							TShaderMapRef<TFinalGatherCS<true> > ComputeShader(View.ShaderMap);
+							TShaderMapRef<TCopyIrradianceCacheSamplesCS<true> > ComputeShader(View.ShaderMap);
 
 							RHICmdList.SetComputeShader(ComputeShader->GetComputeShader());
-							ComputeShader->SetParameters(RHICmdList, View, DistanceFieldAOBentNormalSplat->GetRenderTargetItem(), DistanceFieldNormal->GetRenderTargetItem(), NumObjects, DestLevelDownsampleFactor, DepthLevel, TileListGroupSize, Parameters);
+							ComputeShader->SetParameters(RHICmdList, View, DepthLevel);
 							DispatchIndirectComputeShader(RHICmdList, *ComputeShader, SurfaceCacheResources.DispatchParameters.Buffer, 0);
 
 							ComputeShader->UnsetParameters(RHICmdList);
 						}
 						else
 						{
-							TShaderMapRef<TFinalGatherCS<false> > ComputeShader(View.ShaderMap);
+							TShaderMapRef<TCopyIrradianceCacheSamplesCS<false> > ComputeShader(View.ShaderMap);
 
 							RHICmdList.SetComputeShader(ComputeShader->GetComputeShader());
-							ComputeShader->SetParameters(RHICmdList, View, DistanceFieldAOBentNormalSplat->GetRenderTargetItem(), DistanceFieldNormal->GetRenderTargetItem(), NumObjects, DestLevelDownsampleFactor, DepthLevel, TileListGroupSize, Parameters);
+							ComputeShader->SetParameters(RHICmdList, View, DepthLevel);
 							DispatchIndirectComputeShader(RHICmdList, *ComputeShader, SurfaceCacheResources.DispatchParameters.Buffer, 0);
 
 							ComputeShader->UnsetParameters(RHICmdList);
 						}
+
+						Swap(SurfaceCacheResources.Level[DepthLevel], SurfaceCacheResources.TempResources);
 					}
 				}
+			}
 
-				TRefCountPtr<IPooledRenderTarget> BentNormalAccumulation;
-				TRefCountPtr<IPooledRenderTarget> IrradianceAccumulation;
+			// Start from the highest depth, which will consider the fewest potential shading points
+			// Each level potentially prevents the next higher resolution level from doing expensive shading work
+			// This is the core of the surface cache algorithm
+			for (int32 DepthLevel = GAOMaxLevel; DepthLevel >= GAOMinLevel; DepthLevel--)
+			{
+				int32 DestLevelDownsampleFactor = GAODownsampleFactor * (1 << (DepthLevel * GAOPowerOfTwoBetweenLevels));
+
+				SCOPED_DRAW_EVENTF(RHICmdList, Level, TEXT("Level_%d"), DepthLevel);
+
+				TRefCountPtr<IPooledRenderTarget> DistanceFieldAOBentNormalSplat;
 
 				{
-					FIntPoint BufferSize = GetBufferSizeForAO();
-					FPooledRenderTargetDesc Desc(FPooledRenderTargetDesc::Create2DDesc(BufferSize, PF_FloatRGBA, TexCreate_None, TexCreate_RenderTargetable | TexCreate_UAV, false));
-					GRenderTargetPool.FindFreeElement(Desc, BentNormalAccumulation, TEXT("BentNormalAccumulation"));
-
-					if (GDistanceFieldGI)
-					{
-						GRenderTargetPool.FindFreeElement(Desc, IrradianceAccumulation, TEXT("IrradianceAccumulation"));
-					}
+					FIntPoint AOBufferSize = FIntPoint::DivideAndRoundUp(GSceneRenderTargets.GetBufferSizeXY(), DestLevelDownsampleFactor);
+					FPooledRenderTargetDesc Desc(FPooledRenderTargetDesc::Create2DDesc(AOBufferSize, PF_FloatRGBA, TexCreate_None, TexCreate_RenderTargetable | TexCreate_UAV, false));
+					GRenderTargetPool.FindFreeElement(Desc, DistanceFieldAOBentNormalSplat, TEXT("DistanceFieldAOBentNormalSplat"));
 				}
 
-				// Splat the surface cache records onto the opaque pixels, using less strict weighting so the lighting is smoothed in world space
-				{
-					SCOPED_DRAW_EVENT(RHICmdList, FinalIrradianceCacheSplat);
-					RenderIrradianceCacheInterpolation(RHICmdList, Views, BentNormalAccumulation, IrradianceAccumulation, DistanceFieldNormal->GetRenderTargetItem(), 0, GAODownsampleFactor, Parameters, true);
-				}
-
-				GRenderTargetPool.VisualizeTexture.SetCheckPoint(RHICmdList, BentNormalAccumulation);
-
-				// Post process the AO to cover over artifacts
-				TRefCountPtr<IPooledRenderTarget> BentNormalOutput;
-				TRefCountPtr<IPooledRenderTarget> IrradianceOutput;
-				PostProcessBentNormalAO(RHICmdList, Parameters, Views, BentNormalAccumulation->GetRenderTargetItem(), IrradianceAccumulation, DistanceFieldNormal->GetRenderTargetItem(), BentNormalOutput, IrradianceOutput);
-
-				GRenderTargetPool.VisualizeTexture.SetCheckPoint(RHICmdList, BentNormalOutput);
-
-				if (bApplyToSceneColor)
-				{
-					GSceneRenderTargets.BeginRenderingSceneColor(RHICmdList);
-				}
-				else
-				{
-					FPooledRenderTargetDesc Desc = GSceneRenderTargets.GetSceneColor()->GetDesc();
-					// Make sure we get a signed format
-					Desc.Format = PF_FloatRGBA;
-					GRenderTargetPool.FindFreeElement(Desc, OutDynamicBentNormalAO, TEXT("DynamicBentNormalAO"));
-
-					if (GDistanceFieldGI)
-					{
-						Desc.Format = PF_FloatRGB;
-						GRenderTargetPool.FindFreeElement(Desc, OutDynamicIrradiance, TEXT("DynamicIrradiance"));
-					}
-
-					FTextureRHIParamRef RenderTargets[2] =
-					{
-						OutDynamicBentNormalAO->GetRenderTargetItem().TargetableTexture,
-						GDistanceFieldGI ? OutDynamicIrradiance->GetRenderTargetItem().TargetableTexture : NULL
-					};
-					SetRenderTargets(RHICmdList, ARRAY_COUNT(RenderTargets) - (GDistanceFieldGI ? 0 : 1), RenderTargets, FTextureRHIParamRef(), 0, NULL);
-				}
-
-				// Upsample to full resolution, write to output
-				UpsampleBentNormalAO(RHICmdList, Views, BentNormalOutput, IrradianceOutput, !bApplyToSceneColor);
-
-				return true;
-			}
-		}
-	}
-
-	return false;
-}
-
-class FVisualizeMeshDistanceFieldCS : public FGlobalShader
-{
-	DECLARE_SHADER_TYPE(FVisualizeMeshDistanceFieldCS, Global);
-public:
-
-	static bool ShouldCache(EShaderPlatform Platform)
-	{
-		return IsFeatureLevelSupported(Platform, ERHIFeatureLevel::SM5) && DoesPlatformSupportDistanceFieldAO(Platform);
-	}
-
-	static void ModifyCompilationEnvironment(EShaderPlatform Platform, FShaderCompilerEnvironment& OutEnvironment)
-	{
-		OutEnvironment.SetDefine(TEXT("DOWNSAMPLE_FACTOR"), GAODownsampleFactor);
-		OutEnvironment.SetDefine(TEXT("THREADGROUP_SIZEX"), GDistanceFieldAOTileSizeX);
-		OutEnvironment.SetDefine(TEXT("THREADGROUP_SIZEY"), GDistanceFieldAOTileSizeY);
-		OutEnvironment.SetDefine(TEXT("MAX_OBJECTS_PER_TILE"), GMaxNumObjectsPerTile);
-	}
-
-	/** Default constructor. */
-	FVisualizeMeshDistanceFieldCS() {}
-
-	/** Initialization constructor. */
-	FVisualizeMeshDistanceFieldCS(const ShaderMetaType::CompiledShaderInitializerType& Initializer)
-		: FGlobalShader(Initializer)
-	{
-		VisualizeMeshDistanceFields.Bind(Initializer.ParameterMap, TEXT("VisualizeMeshDistanceFields"));
-		NumGroups.Bind(Initializer.ParameterMap, TEXT("NumGroups"));
-		ObjectParameters.Bind(Initializer.ParameterMap);
-		DeferredParameters.Bind(Initializer.ParameterMap);
-		AOParameters.Bind(Initializer.ParameterMap);
-	}
-
-	void SetParameters(
-		FRHICommandList& RHICmdList, 
-		const FSceneView& View, 
-		FSceneRenderTargetItem& VisualizeMeshDistanceFieldsValue, 
-		int32 NumObjects, 
-		FVector2D NumGroupsValue,
-		const FDistanceFieldAOParameters& Parameters)
-	{
-		const FComputeShaderRHIParamRef ShaderRHI = GetComputeShader();
-
-		FGlobalShader::SetParameters(RHICmdList, ShaderRHI, View);
-
-		VisualizeMeshDistanceFields.SetTexture(RHICmdList, ShaderRHI, VisualizeMeshDistanceFieldsValue.ShaderResourceTexture, VisualizeMeshDistanceFieldsValue.UAV);
-
-		ObjectParameters.Set(RHICmdList, ShaderRHI, GAOObjectBuffers.Buffers, NumObjects);
-		AOParameters.Set(RHICmdList, ShaderRHI, Parameters);
-		DeferredParameters.Set(RHICmdList, ShaderRHI, View);
-
-		SetShaderValue(RHICmdList, ShaderRHI, NumGroups, NumGroupsValue);
-	}
-
-	void UnsetParameters(FRHICommandList& RHICmdList)
-	{
-		VisualizeMeshDistanceFields.UnsetUAV(RHICmdList, GetComputeShader());
-	}
-
-	// FShader interface.
-	virtual bool Serialize(FArchive& Ar)
-	{
-		bool bShaderHasOutdatedParameters = FGlobalShader::Serialize(Ar);
-		Ar << VisualizeMeshDistanceFields;
-		Ar << NumGroups;
-		Ar << ObjectParameters;
-		Ar << DeferredParameters;
-		Ar << AOParameters;
-		return bShaderHasOutdatedParameters;
-	}
-
-private:
-
-	FRWShaderParameter VisualizeMeshDistanceFields;
-	FShaderParameter NumGroups;
-	FDistanceFieldObjectBufferParameters ObjectParameters;
-	FDeferredPixelShaderParameters DeferredParameters;
-	FAOParameters AOParameters;
-};
-
-IMPLEMENT_SHADER_TYPE(,FVisualizeMeshDistanceFieldCS,TEXT("DistanceFieldSurfaceCacheLighting"),TEXT("VisualizeMeshDistanceFieldCS"),SF_Compute);
-
-
-class FVisualizeDistanceFieldUpsamplePS : public FGlobalShader
-{
-	DECLARE_SHADER_TYPE(FVisualizeDistanceFieldUpsamplePS, Global);
-public:
-
-	static bool ShouldCache(EShaderPlatform Platform)
-	{
-		return IsFeatureLevelSupported(Platform, ERHIFeatureLevel::SM5) && DoesPlatformSupportDistanceFieldAO(Platform);
-	}
-
-	static void ModifyCompilationEnvironment(EShaderPlatform Platform, FShaderCompilerEnvironment& OutEnvironment)
-	{
-		OutEnvironment.SetDefine(TEXT("DOWNSAMPLE_FACTOR"), GAODownsampleFactor);
-	}
-
-	/** Default constructor. */
-	FVisualizeDistanceFieldUpsamplePS() {}
-
-	/** Initialization constructor. */
-	FVisualizeDistanceFieldUpsamplePS(const ShaderMetaType::CompiledShaderInitializerType& Initializer)
-		: FGlobalShader(Initializer)
-	{
-		DeferredParameters.Bind(Initializer.ParameterMap);
-		VisualizeDistanceFieldTexture.Bind(Initializer.ParameterMap,TEXT("VisualizeDistanceFieldTexture"));
-		VisualizeDistanceFieldSampler.Bind(Initializer.ParameterMap,TEXT("VisualizeDistanceFieldSampler"));
-	}
-
-	void SetParameters(FRHICommandList& RHICmdList, const FSceneView& View, TRefCountPtr<IPooledRenderTarget>& VisualizeDistanceField)
-	{
-		const FPixelShaderRHIParamRef ShaderRHI = GetPixelShader();
-
-		FGlobalShader::SetParameters(RHICmdList, ShaderRHI, View);
-		DeferredParameters.Set(RHICmdList, ShaderRHI, View);
-
-		SetTextureParameter(RHICmdList, ShaderRHI, VisualizeDistanceFieldTexture, VisualizeDistanceFieldSampler, TStaticSamplerState<SF_Bilinear>::GetRHI(), VisualizeDistanceField->GetRenderTargetItem().ShaderResourceTexture);
-	}
-	// FShader interface.
-	virtual bool Serialize(FArchive& Ar)
-	{
-		bool bShaderHasOutdatedParameters = FGlobalShader::Serialize(Ar);
-		Ar << DeferredParameters;
-		Ar << VisualizeDistanceFieldTexture;
-		Ar << VisualizeDistanceFieldSampler;
-		return bShaderHasOutdatedParameters;
-	}
-
-private:
-
-	FDeferredPixelShaderParameters DeferredParameters;
-	FShaderResourceParameter VisualizeDistanceFieldTexture;
-	FShaderResourceParameter VisualizeDistanceFieldSampler;
-};
-
-IMPLEMENT_SHADER_TYPE(,FVisualizeDistanceFieldUpsamplePS,TEXT("DistanceFieldSurfaceCacheLighting"),TEXT("VisualizeDistanceFieldUpsamplePS"),SF_Pixel);
-
-void FDeferredShadingSceneRenderer::RenderMeshDistanceFieldVisualization(FRHICommandListImmediate& RHICmdList, const FDistanceFieldAOParameters& Parameters)
-{
-	//@todo - support multiple views
-	const FViewInfo& View = Views[0];
-
-	if (GDistanceFieldAO 
-		&& FeatureLevel >= ERHIFeatureLevel::SM5
-		&& DoesPlatformSupportDistanceFieldAO(View.GetShaderPlatform())
-		&& Views.Num() == 1)
-	{
-		QUICK_SCOPE_CYCLE_COUNTER(STAT_RenderMeshDistanceFieldVis);
-		SCOPED_DRAW_EVENT(RHICmdList, VisualizeMeshDistanceFields);
-
-		// Update the global distance field atlas
-		GDistanceFieldVolumeTextureAtlas.UpdateAllocations();
-
-		if (GDistanceFieldVolumeTextureAtlas.VolumeTextureRHI)
-		{
-			int32 NumObjects = 0;
-
-			// Update object buffers with objects that have a valid distance field representation and may affect pixels on the screen
-			UpdateVisibleObjectBuffers(Scene, Parameters, View, NumObjects);
-
-			if (NumObjects > 0)
-			{
-				QUICK_SCOPE_CYCLE_COUNTER(STAT_AOIssueGPUWork);
-
-				TRefCountPtr<IPooledRenderTarget> VisualizeResultRT;
+				// Splat / interpolate the surface cache records onto the buffer sized for the current depth level
+				RenderIrradianceCacheInterpolation(RHICmdList, Views, DistanceFieldAOBentNormalSplat, NULL, DistanceFieldNormal->GetRenderTargetItem(), DepthLevel, DestLevelDownsampleFactor, Parameters, false);
 
 				{
-					const FIntPoint BufferSize = GetBufferSizeForAO();
-					FPooledRenderTargetDesc Desc(FPooledRenderTargetDesc::Create2DDesc(BufferSize, PF_FloatRGBA, TexCreate_None, TexCreate_RenderTargetable | TexCreate_UAV, false));
-					GRenderTargetPool.FindFreeElement(Desc, VisualizeResultRT, TEXT("VisualizeDistanceField"));
-				}
-
-				{
+					SCOPED_DRAW_EVENT(RHICmdList, PopulateIrradianceCache);
 					SetRenderTarget(RHICmdList, NULL, NULL);
 
+					// Save off the current record count before adding any more
+					{	
+						TShaderMapRef<FSaveStartIndexCS> ComputeShader(View.ShaderMap);
+
+						RHICmdList.SetComputeShader(ComputeShader->GetComputeShader());
+						ComputeShader->SetParameters(RHICmdList, View, DepthLevel);
+						DispatchComputeShader(RHICmdList, *ComputeShader, 1, 1, 1);
+
+						ComputeShader->UnsetParameters(RHICmdList);
+					}
+
+					// Create new records which haven't been shaded yet for shading points which don't have a valid interpolation from existing records
 					for (int32 ViewIndex = 0; ViewIndex < Views.Num(); ViewIndex++)
 					{
 						const FViewInfo& View = Views[ViewIndex];
 
-						uint32 GroupSizeX = FMath::DivideAndRoundUp(View.ViewRect.Size().X / GAODownsampleFactor, GDistanceFieldAOTileSizeX);
-						uint32 GroupSizeY = FMath::DivideAndRoundUp(View.ViewRect.Size().Y / GAODownsampleFactor, GDistanceFieldAOTileSizeY);
+						FIntPoint DownsampledViewSize = FIntPoint::DivideAndRoundUp(View.ViewRect.Size(), DestLevelDownsampleFactor);
+						uint32 GroupSizeX = FMath::DivideAndRoundUp(DownsampledViewSize.X, GDistanceFieldAOTileSizeX);
+						uint32 GroupSizeY = FMath::DivideAndRoundUp(DownsampledViewSize.Y, GDistanceFieldAOTileSizeY);
 
-						{
-							SCOPED_DRAW_EVENT(RHICmdList, VisualizeMeshDistanceFieldCS);
-							TShaderMapRef<FVisualizeMeshDistanceFieldCS> ComputeShader(View.ShaderMap);
+						TShaderMapRef<FPopulateCacheCS> ComputeShader(View.ShaderMap);
 
-							RHICmdList.SetComputeShader(ComputeShader->GetComputeShader());
-							ComputeShader->SetParameters(RHICmdList, View, VisualizeResultRT->GetRenderTargetItem(), NumObjects, FVector2D(GroupSizeX, GroupSizeY), Parameters);
-							DispatchComputeShader(RHICmdList, *ComputeShader, GroupSizeX, GroupSizeY, 1);
+						RHICmdList.SetComputeShader(ComputeShader->GetComputeShader());
+						ComputeShader->SetParameters(RHICmdList, View, DistanceFieldAOBentNormalSplat->GetRenderTargetItem(), DistanceFieldNormal->GetRenderTargetItem(), DestLevelDownsampleFactor, DepthLevel, TileListGroupSize, Parameters);
+						DispatchComputeShader(RHICmdList, *ComputeShader, GroupSizeX, GroupSizeY, 1);
 
-							ComputeShader->UnsetParameters(RHICmdList);
-						}
+						ComputeShader->UnsetParameters(RHICmdList);
 					}
 				}
 
 				{
-					GSceneRenderTargets.BeginRenderingSceneColor(RHICmdList);
+					SCOPED_DRAW_EVENT(RHICmdList, ShadeIrradianceCache);
 
-					for (int32 ViewIndex = 0; ViewIndex < Views.Num(); ViewIndex++)
+					{	
+						TShaderMapRef<FSetupFinalGatherIndirectArgumentsCS> ComputeShader(View.ShaderMap);
+
+						RHICmdList.SetComputeShader(ComputeShader->GetComputeShader());
+						ComputeShader->SetParameters(RHICmdList, View, DepthLevel);
+						DispatchComputeShader(RHICmdList, *ComputeShader, 1, 1, 1);
+
+						ComputeShader->UnsetParameters(RHICmdList);
+					}
+
+					// Compute lighting for the new surface cache records by cone-stepping through the object distance fields
+					if (GDistanceFieldGI)
 					{
-						const FViewInfo& View = Views[ViewIndex];
+						TShaderMapRef<TFinalGatherCS<true> > ComputeShader(View.ShaderMap);
 
-						SCOPED_DRAW_EVENT(RHICmdList, UpsampleAO);
+						RHICmdList.SetComputeShader(ComputeShader->GetComputeShader());
+						ComputeShader->SetParameters(RHICmdList, View, DistanceFieldAOBentNormalSplat->GetRenderTargetItem(), DistanceFieldNormal->GetRenderTargetItem(), DestLevelDownsampleFactor, DepthLevel, TileListGroupSize, Parameters);
+						DispatchIndirectComputeShader(RHICmdList, *ComputeShader, SurfaceCacheResources.DispatchParameters.Buffer, 0);
 
-						RHICmdList.SetViewport(View.ViewRect.Min.X, View.ViewRect.Min.Y, 0.0f, View.ViewRect.Max.X, View.ViewRect.Max.Y, 1.0f);
-						RHICmdList.SetRasterizerState(TStaticRasterizerState<FM_Solid, CM_None>::GetRHI());
-						RHICmdList.SetDepthStencilState(TStaticDepthStencilState<false, CF_Always>::GetRHI());
-						RHICmdList.SetBlendState(TStaticBlendState<>::GetRHI());
+						ComputeShader->UnsetParameters(RHICmdList);
+					}
+					else
+					{
+						TShaderMapRef<TFinalGatherCS<false> > ComputeShader(View.ShaderMap);
 
-						TShaderMapRef<FPostProcessVS> VertexShader( View.ShaderMap );
-						TShaderMapRef<FVisualizeDistanceFieldUpsamplePS> PixelShader( View.ShaderMap );
+						RHICmdList.SetComputeShader(ComputeShader->GetComputeShader());
+						ComputeShader->SetParameters(RHICmdList, View, DistanceFieldAOBentNormalSplat->GetRenderTargetItem(), DistanceFieldNormal->GetRenderTargetItem(), DestLevelDownsampleFactor, DepthLevel, TileListGroupSize, Parameters);
+						DispatchIndirectComputeShader(RHICmdList, *ComputeShader, SurfaceCacheResources.DispatchParameters.Buffer, 0);
 
-						static FGlobalBoundShaderState BoundShaderState;
-
-						SetGlobalBoundShaderState(RHICmdList, View.GetFeatureLevel(), BoundShaderState, GFilterVertexDeclaration.VertexDeclarationRHI, *VertexShader, *PixelShader);
-
-						PixelShader->SetParameters(RHICmdList, View, VisualizeResultRT);
-
-						DrawRectangle( 
-							RHICmdList,
-							0, 0, 
-							View.ViewRect.Width(), View.ViewRect.Height(),
-							View.ViewRect.Min.X / GAODownsampleFactor, View.ViewRect.Min.Y / GAODownsampleFactor, 
-							View.ViewRect.Width() / GAODownsampleFactor, View.ViewRect.Height() / GAODownsampleFactor,
-							FIntPoint(View.ViewRect.Width(), View.ViewRect.Height()),
-							GetBufferSizeForAO(),
-							*VertexShader);
+						ComputeShader->UnsetParameters(RHICmdList);
 					}
 				}
 			}
+
+			TRefCountPtr<IPooledRenderTarget> BentNormalAccumulation;
+			TRefCountPtr<IPooledRenderTarget> IrradianceAccumulation;
+
+			{
+				FIntPoint BufferSize = GetBufferSizeForAO();
+				FPooledRenderTargetDesc Desc(FPooledRenderTargetDesc::Create2DDesc(BufferSize, PF_FloatRGBA, TexCreate_None, TexCreate_RenderTargetable | TexCreate_UAV, false));
+				GRenderTargetPool.FindFreeElement(Desc, BentNormalAccumulation, TEXT("BentNormalAccumulation"));
+
+				if (GDistanceFieldGI)
+				{
+					GRenderTargetPool.FindFreeElement(Desc, IrradianceAccumulation, TEXT("IrradianceAccumulation"));
+				}
+			}
+
+			// Splat the surface cache records onto the opaque pixels, using less strict weighting so the lighting is smoothed in world space
+			{
+				SCOPED_DRAW_EVENT(RHICmdList, FinalIrradianceCacheSplat);
+				RenderIrradianceCacheInterpolation(RHICmdList, Views, BentNormalAccumulation, IrradianceAccumulation, DistanceFieldNormal->GetRenderTargetItem(), 0, GAODownsampleFactor, Parameters, true);
+			}
+
+			GRenderTargetPool.VisualizeTexture.SetCheckPoint(RHICmdList, BentNormalAccumulation);
+
+			// Post process the AO to cover over artifacts
+			TRefCountPtr<IPooledRenderTarget> BentNormalOutput;
+			TRefCountPtr<IPooledRenderTarget> IrradianceOutput;
+			PostProcessBentNormalAO(RHICmdList, Parameters, Views, BentNormalAccumulation->GetRenderTargetItem(), IrradianceAccumulation, DistanceFieldNormal->GetRenderTargetItem(), BentNormalOutput, IrradianceOutput);
+
+			GRenderTargetPool.VisualizeTexture.SetCheckPoint(RHICmdList, BentNormalOutput);
+
+			if (bApplyToSceneColor)
+			{
+				GSceneRenderTargets.BeginRenderingSceneColor(RHICmdList);
+			}
+			else
+			{
+				FPooledRenderTargetDesc Desc = GSceneRenderTargets.GetSceneColor()->GetDesc();
+				// Make sure we get a signed format
+				Desc.Format = PF_FloatRGBA;
+				GRenderTargetPool.FindFreeElement(Desc, OutDynamicBentNormalAO, TEXT("DynamicBentNormalAO"));
+
+				if (GDistanceFieldGI)
+				{
+					Desc.Format = PF_FloatRGB;
+					GRenderTargetPool.FindFreeElement(Desc, OutDynamicIrradiance, TEXT("DynamicIrradiance"));
+				}
+
+				FTextureRHIParamRef RenderTargets[2] =
+				{
+					OutDynamicBentNormalAO->GetRenderTargetItem().TargetableTexture,
+					GDistanceFieldGI ? OutDynamicIrradiance->GetRenderTargetItem().TargetableTexture : NULL
+				};
+				SetRenderTargets(RHICmdList, ARRAY_COUNT(RenderTargets) - (GDistanceFieldGI ? 0 : 1), RenderTargets, FTextureRHIParamRef(), 0, NULL);
+			}
+
+			// Upsample to full resolution, write to output
+			UpsampleBentNormalAO(RHICmdList, Views, BentNormalOutput, IrradianceOutput, !bApplyToSceneColor);
+
+			return true;
 		}
 	}
+
+	return false;
 }
 
 template<bool bApplyShadowing, bool bSupportIrradiance>
@@ -4269,16 +4634,21 @@
 IMPLEMENT_SKYLIGHT_PS_TYPE(true, false)
 IMPLEMENT_SKYLIGHT_PS_TYPE(false, false)
 
-void FDeferredShadingSceneRenderer::RenderDynamicSkyLighting(FRHICommandListImmediate& RHICmdList, TRefCountPtr<IPooledRenderTarget>& DynamicBentNormalAO)
+bool FDeferredShadingSceneRenderer::ShouldRenderDynamicSkyLight() const
 {
-	if (Scene->SkyLight
+	return Scene->SkyLight
 		&& Scene->SkyLight->ProcessedTexture
 		&& !Scene->SkyLight->bWantsStaticShadowing
 		&& !Scene->SkyLight->bHasStaticLighting
 		&& ViewFamily.EngineShowFlags.SkyLighting
 		&& FeatureLevel >= ERHIFeatureLevel::SM4
 		&& !IsSimpleDynamicLightingEnabled() 
-		&& !ViewFamily.EngineShowFlags.VisualizeLightCulling)
+		&& !ViewFamily.EngineShowFlags.VisualizeLightCulling;
+}
+
+void FDeferredShadingSceneRenderer::RenderDynamicSkyLighting(FRHICommandListImmediate& RHICmdList, TRefCountPtr<IPooledRenderTarget>& DynamicBentNormalAO)
+{
+	if (ShouldRenderDynamicSkyLight())
 	{
 		SCOPED_DRAW_EVENT(RHICmdList, SkyLightDiffuse);
 
@@ -4367,3 +4737,246 @@
 		}
 	}
 }
+
+class FVisualizeMeshDistanceFieldCS : public FGlobalShader
+{
+	DECLARE_SHADER_TYPE(FVisualizeMeshDistanceFieldCS, Global);
+public:
+
+	static bool ShouldCache(EShaderPlatform Platform)
+	{
+		return IsFeatureLevelSupported(Platform, ERHIFeatureLevel::SM5) && DoesPlatformSupportDistanceFieldAO(Platform);
+	}
+
+	static void ModifyCompilationEnvironment(EShaderPlatform Platform, FShaderCompilerEnvironment& OutEnvironment)
+	{
+		OutEnvironment.SetDefine(TEXT("DOWNSAMPLE_FACTOR"), GAODownsampleFactor);
+		OutEnvironment.SetDefine(TEXT("THREADGROUP_SIZEX"), GDistanceFieldAOTileSizeX);
+		OutEnvironment.SetDefine(TEXT("THREADGROUP_SIZEY"), GDistanceFieldAOTileSizeY);
+		OutEnvironment.SetDefine(TEXT("MAX_OBJECTS_PER_TILE"), GMaxNumObjectsPerTile);
+	}
+
+	/** Default constructor. */
+	FVisualizeMeshDistanceFieldCS() {}
+
+	/** Initialization constructor. */
+	FVisualizeMeshDistanceFieldCS(const ShaderMetaType::CompiledShaderInitializerType& Initializer)
+		: FGlobalShader(Initializer)
+	{
+		VisualizeMeshDistanceFields.Bind(Initializer.ParameterMap, TEXT("VisualizeMeshDistanceFields"));
+		NumGroups.Bind(Initializer.ParameterMap, TEXT("NumGroups"));
+		ObjectParameters.Bind(Initializer.ParameterMap);
+		DeferredParameters.Bind(Initializer.ParameterMap);
+		AOParameters.Bind(Initializer.ParameterMap);
+	}
+
+	void SetParameters(
+		FRHICommandList& RHICmdList, 
+		const FSceneView& View, 
+		FSceneRenderTargetItem& VisualizeMeshDistanceFieldsValue, 
+		FVector2D NumGroupsValue,
+		const FDistanceFieldAOParameters& Parameters)
+	{
+		const FComputeShaderRHIParamRef ShaderRHI = GetComputeShader();
+
+		FGlobalShader::SetParameters(RHICmdList, ShaderRHI, View);
+
+		VisualizeMeshDistanceFields.SetTexture(RHICmdList, ShaderRHI, VisualizeMeshDistanceFieldsValue.ShaderResourceTexture, VisualizeMeshDistanceFieldsValue.UAV);
+
+		ObjectParameters.Set(RHICmdList, ShaderRHI, GAOCulledObjectBuffers.Buffers);
+		AOParameters.Set(RHICmdList, ShaderRHI, Parameters);
+		DeferredParameters.Set(RHICmdList, ShaderRHI, View);
+
+		SetShaderValue(RHICmdList, ShaderRHI, NumGroups, NumGroupsValue);
+	}
+
+	void UnsetParameters(FRHICommandList& RHICmdList)
+	{
+		VisualizeMeshDistanceFields.UnsetUAV(RHICmdList, GetComputeShader());
+	}
+
+	// FShader interface.
+	virtual bool Serialize(FArchive& Ar)
+	{
+		bool bShaderHasOutdatedParameters = FGlobalShader::Serialize(Ar);
+		Ar << VisualizeMeshDistanceFields;
+		Ar << NumGroups;
+		Ar << ObjectParameters;
+		Ar << DeferredParameters;
+		Ar << AOParameters;
+		return bShaderHasOutdatedParameters;
+	}
+
+private:
+
+	FRWShaderParameter VisualizeMeshDistanceFields;
+	FShaderParameter NumGroups;
+	FDistanceFieldCulledObjectBufferParameters ObjectParameters;
+	FDeferredPixelShaderParameters DeferredParameters;
+	FAOParameters AOParameters;
+};
+
+IMPLEMENT_SHADER_TYPE(,FVisualizeMeshDistanceFieldCS,TEXT("DistanceFieldSurfaceCacheLighting"),TEXT("VisualizeMeshDistanceFieldCS"),SF_Compute);
+
+
+class FVisualizeDistanceFieldUpsamplePS : public FGlobalShader
+{
+	DECLARE_SHADER_TYPE(FVisualizeDistanceFieldUpsamplePS, Global);
+public:
+
+	static bool ShouldCache(EShaderPlatform Platform)
+	{
+		return IsFeatureLevelSupported(Platform, ERHIFeatureLevel::SM5) && DoesPlatformSupportDistanceFieldAO(Platform);
+	}
+
+	static void ModifyCompilationEnvironment(EShaderPlatform Platform, FShaderCompilerEnvironment& OutEnvironment)
+	{
+		OutEnvironment.SetDefine(TEXT("DOWNSAMPLE_FACTOR"), GAODownsampleFactor);
+	}
+
+	/** Default constructor. */
+	FVisualizeDistanceFieldUpsamplePS() {}
+
+	/** Initialization constructor. */
+	FVisualizeDistanceFieldUpsamplePS(const ShaderMetaType::CompiledShaderInitializerType& Initializer)
+		: FGlobalShader(Initializer)
+	{
+		DeferredParameters.Bind(Initializer.ParameterMap);
+		VisualizeDistanceFieldTexture.Bind(Initializer.ParameterMap,TEXT("VisualizeDistanceFieldTexture"));
+		VisualizeDistanceFieldSampler.Bind(Initializer.ParameterMap,TEXT("VisualizeDistanceFieldSampler"));
+	}
+
+	void SetParameters(FRHICommandList& RHICmdList, const FSceneView& View, TRefCountPtr<IPooledRenderTarget>& VisualizeDistanceField)
+	{
+		const FPixelShaderRHIParamRef ShaderRHI = GetPixelShader();
+
+		FGlobalShader::SetParameters(RHICmdList, ShaderRHI, View);
+		DeferredParameters.Set(RHICmdList, ShaderRHI, View);
+
+		SetTextureParameter(RHICmdList, ShaderRHI, VisualizeDistanceFieldTexture, VisualizeDistanceFieldSampler, TStaticSamplerState<SF_Bilinear>::GetRHI(), VisualizeDistanceField->GetRenderTargetItem().ShaderResourceTexture);
+	}
+	// FShader interface.
+	virtual bool Serialize(FArchive& Ar)
+	{
+		bool bShaderHasOutdatedParameters = FGlobalShader::Serialize(Ar);
+		Ar << DeferredParameters;
+		Ar << VisualizeDistanceFieldTexture;
+		Ar << VisualizeDistanceFieldSampler;
+		return bShaderHasOutdatedParameters;
+	}
+
+private:
+
+	FDeferredPixelShaderParameters DeferredParameters;
+	FShaderResourceParameter VisualizeDistanceFieldTexture;
+	FShaderResourceParameter VisualizeDistanceFieldSampler;
+};
+
+IMPLEMENT_SHADER_TYPE(,FVisualizeDistanceFieldUpsamplePS,TEXT("DistanceFieldSurfaceCacheLighting"),TEXT("VisualizeDistanceFieldUpsamplePS"),SF_Pixel);
+
+
+void FDeferredShadingSceneRenderer::RenderMeshDistanceFieldVisualization(FRHICommandListImmediate& RHICmdList, const FDistanceFieldAOParameters& Parameters)
+{
+	//@todo - support multiple views
+	const FViewInfo& View = Views[0];
+
+	if (GDistanceFieldAO 
+		&& FeatureLevel >= ERHIFeatureLevel::SM5
+		&& DoesPlatformSupportDistanceFieldAO(View.GetShaderPlatform())
+		&& Views.Num() == 1)
+	{
+		QUICK_SCOPE_CYCLE_COUNTER(STAT_RenderMeshDistanceFieldVis);
+		SCOPED_DRAW_EVENT(RHICmdList, VisualizeMeshDistanceFields);
+
+		if (GDistanceFieldVolumeTextureAtlas.VolumeTextureRHI && Scene->DistanceFieldSceneData.NumObjectsInBuffer > 0)
+		{
+			check(!Scene->DistanceFieldSceneData.HasPendingOperations());
+
+			QUICK_SCOPE_CYCLE_COUNTER(STAT_AOIssueGPUWork);
+
+			{
+				if (GAOCulledObjectBuffers.Buffers.MaxObjects < Scene->DistanceFieldSceneData.NumObjectsInBuffer)
+				{
+					GAOCulledObjectBuffers.Buffers.MaxObjects = Scene->DistanceFieldSceneData.NumObjectsInBuffer * 5 / 4;
+					GAOCulledObjectBuffers.Buffers.Release();
+					GAOCulledObjectBuffers.Buffers.Initialize();
+				}
+
+				{
+					TShaderMapRef<FCullObjectsForViewCS> ComputeShader(GetGlobalShaderMap(Scene->GetFeatureLevel()));
+					RHICmdList.SetComputeShader(ComputeShader->GetComputeShader());
+					ComputeShader->SetParameters(RHICmdList, Scene, View, Parameters);
+
+					DispatchComputeShader(RHICmdList, *ComputeShader, FMath::DivideAndRoundUp<uint32>(Scene->DistanceFieldSceneData.NumObjectsInBuffer, UpdateObjectsGroupSize), 1, 1);
+					ComputeShader->UnsetParameters(RHICmdList);
+				}
+			}
+
+			TRefCountPtr<IPooledRenderTarget> VisualizeResultRT;
+
+			{
+				const FIntPoint BufferSize = GetBufferSizeForAO();
+				FPooledRenderTargetDesc Desc(FPooledRenderTargetDesc::Create2DDesc(BufferSize, PF_FloatRGBA, TexCreate_None, TexCreate_RenderTargetable | TexCreate_UAV, false));
+				GRenderTargetPool.FindFreeElement(Desc, VisualizeResultRT, TEXT("VisualizeDistanceField"));
+			}
+
+			{
+				SetRenderTarget(RHICmdList, NULL, NULL);
+
+				for (int32 ViewIndex = 0; ViewIndex < Views.Num(); ViewIndex++)
+				{
+					const FViewInfo& View = Views[ViewIndex];
+
+					uint32 GroupSizeX = FMath::DivideAndRoundUp(View.ViewRect.Size().X / GAODownsampleFactor, GDistanceFieldAOTileSizeX);
+					uint32 GroupSizeY = FMath::DivideAndRoundUp(View.ViewRect.Size().Y / GAODownsampleFactor, GDistanceFieldAOTileSizeY);
+
+					{
+						SCOPED_DRAW_EVENT(RHICmdList, VisualizeMeshDistanceFieldCS);
+						TShaderMapRef<FVisualizeMeshDistanceFieldCS> ComputeShader(View.ShaderMap);
+
+						RHICmdList.SetComputeShader(ComputeShader->GetComputeShader());
+						ComputeShader->SetParameters(RHICmdList, View, VisualizeResultRT->GetRenderTargetItem(), FVector2D(GroupSizeX, GroupSizeY), Parameters);
+						DispatchComputeShader(RHICmdList, *ComputeShader, GroupSizeX, GroupSizeY, 1);
+
+						ComputeShader->UnsetParameters(RHICmdList);
+					}
+				}
+			}
+
+			{
+				GSceneRenderTargets.BeginRenderingSceneColor(RHICmdList);
+
+				for (int32 ViewIndex = 0; ViewIndex < Views.Num(); ViewIndex++)
+				{
+					const FViewInfo& View = Views[ViewIndex];
+
+					SCOPED_DRAW_EVENT(RHICmdList, UpsampleAO);
+
+					RHICmdList.SetViewport(View.ViewRect.Min.X, View.ViewRect.Min.Y, 0.0f, View.ViewRect.Max.X, View.ViewRect.Max.Y, 1.0f);
+					RHICmdList.SetRasterizerState(TStaticRasterizerState<FM_Solid, CM_None>::GetRHI());
+					RHICmdList.SetDepthStencilState(TStaticDepthStencilState<false, CF_Always>::GetRHI());
+					RHICmdList.SetBlendState(TStaticBlendState<>::GetRHI());
+
+					TShaderMapRef<FPostProcessVS> VertexShader( View.ShaderMap );
+					TShaderMapRef<FVisualizeDistanceFieldUpsamplePS> PixelShader( View.ShaderMap );
+
+					static FGlobalBoundShaderState BoundShaderState;
+
+					SetGlobalBoundShaderState(RHICmdList, View.GetFeatureLevel(), BoundShaderState, GFilterVertexDeclaration.VertexDeclarationRHI, *VertexShader, *PixelShader);
+
+					PixelShader->SetParameters(RHICmdList, View, VisualizeResultRT);
+
+					DrawRectangle( 
+						RHICmdList,
+						0, 0, 
+						View.ViewRect.Width(), View.ViewRect.Height(),
+						View.ViewRect.Min.X / GAODownsampleFactor, View.ViewRect.Min.Y / GAODownsampleFactor, 
+						View.ViewRect.Width() / GAODownsampleFactor, View.ViewRect.Height() / GAODownsampleFactor,
+						FIntPoint(View.ViewRect.Width(), View.ViewRect.Height()),
+						GetBufferSizeForAO(),
+						*VertexShader);
+				}
+			}
+		}
+	}
+}
diff --git a/Engine/Source/Runtime/Renderer/Private/LightSceneInfo.cpp b/Engine/Source/Runtime/Renderer/Private/LightSceneInfo.cpp
--- a/Engine/Source/Runtime/Renderer/Private/LightSceneInfo.cpp
+++ b/Engine/Source/Runtime/Renderer/Private/LightSceneInfo.cpp
@@ -31,7 +31,6 @@ FLightSceneInfo::FLightSceneInfo(FLightSceneProxy* InProxy, bool InbVisible)
 	, DynamicPrimitiveList(NULL)
 	, Id(INDEX_NONE)
 	, TileIntersectionResources(NULL)
-	, DistanceFieldObjectBuffers(NULL)
 	, bPrecomputedLightingIsValid(InProxy->GetLightComponent()->bPrecomputedLightingIsValid)
 	, bVisible(InbVisible)
 	, bEnableLightShaftBloom(InProxy->GetLightComponent()->bEnableLightShaftBloom)
@@ -198,11 +197,6 @@ void FLightSceneInfo::ReleaseRHI()
 	{
 		TileIntersectionResources->Release();
 	}
-
-	if (DistanceFieldObjectBuffers)
-	{
-		DistanceFieldObjectBuffers->Release();
-	}
 }
 
 /** Determines whether two bounding spheres intersect. */
diff --git a/Engine/Source/Runtime/Renderer/Private/LightSceneInfo.h b/Engine/Source/Runtime/Renderer/Private/LightSceneInfo.h
--- a/Engine/Source/Runtime/Renderer/Private/LightSceneInfo.h
+++ b/Engine/Source/Runtime/Renderer/Private/LightSceneInfo.h
@@ -120,9 +120,6 @@ public:
 	/** Tile intersection buffer for distance field shadowing, stored on the light to avoid reallocating each frame. */
 	mutable TScopedPointer<class FLightTileIntersectionResources> TileIntersectionResources;
 
-	/** Culled object buffers for distance field shadowing, stored on the light to avoid reallocating each frame. */
-	mutable TScopedPointer<class FDistanceFieldObjectBuffers> DistanceFieldObjectBuffers;
-
 	/** True if the light is built. */
 	uint32 bPrecomputedLightingIsValid : 1;
 
diff --git a/Engine/Source/Runtime/Renderer/Private/PrimitiveSceneInfo.h b/Engine/Source/Runtime/Renderer/Private/PrimitiveSceneInfo.h
--- a/Engine/Source/Runtime/Renderer/Private/PrimitiveSceneInfo.h
+++ b/Engine/Source/Runtime/Renderer/Private/PrimitiveSceneInfo.h
@@ -195,6 +195,9 @@ public:
 	 */
 	const FReflectionCaptureProxy* CachedReflectionCaptureProxy;
 
+	/** Mapping from instance index in this primitive to index in the global distance field object buffers. */
+	TArray<int32, TInlineAllocator<1>> DistanceFieldInstanceIndices;
+
 	/** Whether the primitive is newly registered or moved and CachedReflectionCaptureProxy needs to be updated on the next render. */
 	uint32 bNeedsCachedReflectionCaptureUpdate : 1;
 
diff --git a/Engine/Source/Runtime/Renderer/Private/RendererScene.cpp b/Engine/Source/Runtime/Renderer/Private/RendererScene.cpp
--- a/Engine/Source/Runtime/Renderer/Private/RendererScene.cpp
+++ b/Engine/Source/Runtime/Renderer/Private/RendererScene.cpp
@@ -13,6 +13,7 @@
 #include "EngineModule.h"
 #include "PrecomputedLightVolume.h"
 #include "FXSystem.h"
+#include "DistanceFieldLightingShared.h"
 
 // Enable this define to do slow checks for components being added to the wrong
 // world's scene, when using PIE. This can happen if a PIE component is reattached
@@ -75,6 +76,96 @@ FSceneViewState::FSceneViewState()
 #endif
 }
 
+FDistanceFieldSceneData::FDistanceFieldSceneData() 
+	: NumObjectsInBuffer(0)
+	, ObjectBuffers(NULL)
+{
+
+}
+
+FDistanceFieldSceneData::~FDistanceFieldSceneData() 
+{
+	delete ObjectBuffers;
+}
+
+void FDistanceFieldSceneData::AddPrimitive(FPrimitiveSceneInfo* InPrimitive)
+{
+	const FPrimitiveSceneProxy* Proxy = InPrimitive->Proxy;
+
+	if (Proxy->SupportsDistanceFieldRepresentation()
+		&& Proxy->CastsDynamicShadow()
+		&& Proxy->AffectsDistanceFieldLighting())
+	{
+		checkSlow(!PendingAddOperations.Contains(InPrimitive));
+		checkSlow(!PendingUpdateOperations.Contains(InPrimitive));
+		PendingAddOperations.Add(InPrimitive);
+	}
+}
+
+void FDistanceFieldSceneData::UpdatePrimitive(FPrimitiveSceneInfo* InPrimitive)
+{
+	const FPrimitiveSceneProxy* Proxy = InPrimitive->Proxy;
+
+	if (Proxy->SupportsDistanceFieldRepresentation() 
+		&& Proxy->CastsDynamicShadow() 
+		&& Proxy->AffectsDistanceFieldLighting()
+		&& !PendingAddOperations.Contains(InPrimitive))
+	{
+		check(InPrimitive->DistanceFieldInstanceIndices.Num() > 0);
+		PendingUpdateOperations.Add(InPrimitive);
+	}
+}
+
+void FDistanceFieldSceneData::RemovePrimitive(FPrimitiveSceneInfo* InPrimitive)
+{
+	const FPrimitiveSceneProxy* Proxy = InPrimitive->Proxy;
+
+	if (Proxy->SupportsDistanceFieldRepresentation() && Proxy->AffectsDistanceFieldLighting())
+	{
+		PendingAddOperations.Remove(InPrimitive);
+		PendingUpdateOperations.Remove(InPrimitive);
+
+		for (int32 InstanceIndex = 0; InstanceIndex < InPrimitive->DistanceFieldInstanceIndices.Num(); InstanceIndex++)
+		{
+			int32 RemoveIndex = InPrimitive->DistanceFieldInstanceIndices[InstanceIndex];
+
+			// Sanity check that scales poorly
+			if (PendingRemoveOperations.Num() < 1000)
+			{
+				checkSlow(!PendingRemoveOperations.Contains(RemoveIndex));
+			}
+			
+			PendingRemoveOperations.Add(RemoveIndex);
+		}
+
+		InPrimitive->DistanceFieldInstanceIndices.Empty();
+	}
+}
+
+void FDistanceFieldSceneData::Release()
+{
+	if (ObjectBuffers)
+	{
+		ObjectBuffers->Release();
+	}
+}
+
+void FDistanceFieldSceneData::VerifyIntegrity()
+{
+	check(NumObjectsInBuffer == PrimitiveInstanceMapping.Num());
+
+	for (int32 PrimitiveInstanceIndex = 0; PrimitiveInstanceIndex < PrimitiveInstanceMapping.Num(); PrimitiveInstanceIndex++)
+	{
+		const FPrimitiveAndInstance& PrimitiveAndInstance = PrimitiveInstanceMapping[PrimitiveInstanceIndex];
+
+		check(PrimitiveAndInstance.Primitive && PrimitiveAndInstance.Primitive->DistanceFieldInstanceIndices.Num() > 0);
+		check(PrimitiveAndInstance.Primitive->DistanceFieldInstanceIndices.IsValidIndex(PrimitiveAndInstance.InstanceIndex));
+
+		const int32 InstanceIndex = PrimitiveAndInstance.Primitive->DistanceFieldInstanceIndices[PrimitiveAndInstance.InstanceIndex];
+		check(InstanceIndex == PrimitiveInstanceIndex);
+	}
+}
+
 /**
  * Sets the FX system associated with the scene.
  */
@@ -172,6 +263,8 @@ void FScene::AddPrimitiveSceneInfo_RenderThread(FRHICommandListImmediate& RHICmd
 
 	// Add the primitive to the scene.
 	PrimitiveSceneInfo->AddToScene(RHICmdList, true);
+
+	DistanceFieldSceneData.AddPrimitive(PrimitiveSceneInfo);
 }
 
 /**
@@ -264,6 +357,7 @@ FScene::~FScene()
 
 	ReflectionSceneData.CubemapArray.ReleaseResource();
 	IndirectLightingCache.ReleaseResource();
+	DistanceFieldSceneData.Release();
 
 	if (SurfaceCacheResources)
 	{
@@ -390,6 +484,8 @@ void FScene::UpdatePrimitiveTransform_RenderThread(FRHICommandListImmediate& RHI
 	// Update the primitive transform.
 	PrimitiveSceneProxy->SetTransform(LocalToWorld, WorldBounds, LocalBounds, OwnerPosition);
 
+	DistanceFieldSceneData.UpdatePrimitive(PrimitiveSceneProxy->GetPrimitiveSceneInfo());
+
 	// If the primitive has static mesh elements, it should have returned true from ShouldRecreateProxyOnUpdateTransform!
 	check(!(bUpdateStaticDrawLists && PrimitiveSceneProxy->GetPrimitiveSceneInfo()->StaticMeshes.Num()));
 
@@ -579,6 +675,8 @@ void FScene::RemovePrimitiveSceneInfo_RenderThread(FPrimitiveSceneInfo* Primitiv
 	// Remove the primitive from the scene.
 	PrimitiveSceneInfo->RemoveFromScene(true);
 
+	DistanceFieldSceneData.RemovePrimitive(PrimitiveSceneInfo);
+
 	// free the primitive scene proxy.
 	delete PrimitiveSceneInfo->Proxy;
 }
diff --git a/Engine/Source/Runtime/Renderer/Private/ScenePrivate.h b/Engine/Source/Runtime/Renderer/Private/ScenePrivate.h
--- a/Engine/Source/Runtime/Renderer/Private/ScenePrivate.h
+++ b/Engine/Source/Runtime/Renderer/Private/ScenePrivate.h
@@ -879,6 +879,50 @@ public:
 	{}
 };
 
+class FPrimitiveAndInstance
+{
+public:
+
+	FPrimitiveAndInstance(FPrimitiveSceneInfo* InPrimitive, int32 InInstanceIndex) :
+		Primitive(InPrimitive),
+		InstanceIndex(InInstanceIndex)
+	{}
+
+	FPrimitiveSceneInfo* Primitive;
+	int32 InstanceIndex;
+};
+
+/** Scene data used to manage distance field object buffers on the GPU. */
+class FDistanceFieldSceneData
+{
+public:
+
+	FDistanceFieldSceneData();
+	~FDistanceFieldSceneData();
+
+	void AddPrimitive(FPrimitiveSceneInfo* InPrimitive);
+	void UpdatePrimitive(FPrimitiveSceneInfo* InPrimitive);
+	void RemovePrimitive(FPrimitiveSceneInfo* InPrimitive);
+	void Release();
+	void VerifyIntegrity();
+
+	bool HasPendingOperations() const
+	{
+		return PendingAddOperations.Num() > 0 || PendingUpdateOperations.Num() > 0 || PendingRemoveOperations.Num() > 0;
+	}
+
+	int32 NumObjectsInBuffer;
+	class FDistanceFieldObjectBuffers* ObjectBuffers;
+
+	/** Stores the primitive and instance index of every entry in the object buffer. */
+	TArray<FPrimitiveAndInstance> PrimitiveInstanceMapping;
+
+	/** Pending operations on the object buffers to be processed next frame. */
+	TArray<FPrimitiveSceneInfo*> PendingAddOperations;
+	TArray<FPrimitiveSceneInfo*> PendingUpdateOperations;
+	TArray<int32> PendingRemoveOperations;
+};
+
 /** Stores data for an allocation in the FIndirectLightingCache. */
 class FIndirectLightingCacheBlock
 {
@@ -1272,6 +1316,9 @@ public:
 	/** Scene state of distance field AO.  NULL if the scene doesn't use the feature. */
 	class FSurfaceCacheResources* SurfaceCacheResources;
 
+	/** Distance field object scene data. */
+	FDistanceFieldSceneData DistanceFieldSceneData;
+
 	/** Preshadows that are currently cached in the PreshadowCache render target. */
 	TArray<TRefCountPtr<FProjectedShadowInfo> > CachedPreshadows;
 
diff --git a/Engine/Source/Runtime/Renderer/Private/SceneRendering.h b/Engine/Source/Runtime/Renderer/Private/SceneRendering.h
--- a/Engine/Source/Runtime/Renderer/Private/SceneRendering.h
+++ b/Engine/Source/Runtime/Renderer/Private/SceneRendering.h
@@ -602,7 +602,11 @@ protected:
 	// Shared functionality between all scene renderers
 
 	/** Generates FProjectedShadowInfos for all wholesceneshadows on the given light.*/
-	void AddViewDependentWholeSceneShadowsForView(TArray<FProjectedShadowInfo*, SceneRenderingAllocator>& ShadowInfos, FVisibleLightInfo& VisibleLightInfo, FLightSceneInfo& LightSceneInfo);
+	void AddViewDependentWholeSceneShadowsForView(
+		TArray<FProjectedShadowInfo*, SceneRenderingAllocator>& ShadowInfos, 
+		TArray<FProjectedShadowInfo*, SceneRenderingAllocator>& ShadowInfosThatNeedCulling, 
+		FVisibleLightInfo& VisibleLightInfo, 
+		FLightSceneInfo& LightSceneInfo);
 
 	/**
 	* Used by RenderLights to figure out if projected shadows need to be rendered to the attenuation buffer.
diff --git a/Engine/Source/Runtime/Renderer/Private/ShadowSetup.cpp b/Engine/Source/Runtime/Renderer/Private/ShadowSetup.cpp
--- a/Engine/Source/Runtime/Renderer/Private/ShadowSetup.cpp
+++ b/Engine/Source/Runtime/Renderer/Private/ShadowSetup.cpp
@@ -17,15 +17,6 @@ static FAutoConsoleVariableRef CVarMinScreenRadiusForShadowCaster(
 	ECVF_Scalability | ECVF_RenderThreadSafe
 	);
 
-static float GMinScreenRadiusForRayTracedCasters = 0.01f;
-static FAutoConsoleVariableRef CVarMinScreenRadiusForRayTracedCasters(
-	TEXT("r.Shadow.RadiusThresholdRayTraced"),
-	GMinScreenRadiusForRayTracedCasters,
-	TEXT("Cull shadow casters in the ray traced shadows if they are too small, value is the minimal screen space bounding sphere radius\n")
-	TEXT("(default 0.01)"),
-	ECVF_Scalability | ECVF_RenderThreadSafe
-	);
-
 static float GMinScreenRadiusForShadowCasterRSM = 0.06f;
 static FAutoConsoleVariableRef CVarMinScreenRadiusForShadowCasterRSM(
 	TEXT("r.Shadow.RadiusThresholdRSM"),
@@ -650,226 +641,222 @@ void FProjectedShadowInfo::AddSubjectPrimitive(FPrimitiveSceneInfo* PrimitiveSce
 {
 	QUICK_SCOPE_CYCLE_COUNTER(STAT_AddSubjectPrimitive);
 
-	if (bRayTracedDistanceFieldShadow)
-	{
-		SubjectPrimitives.Add(PrimitiveSceneInfo);
-	}
-	else
+	// Ray traced shadows use the GPU managed distance field object buffers, no CPU culling should be used
+	check(!bRayTracedDistanceFieldShadow);
+
+	if (!ReceiverPrimitives.Contains(PrimitiveSceneInfo))
 	{
-		if (!ReceiverPrimitives.Contains(PrimitiveSceneInfo))
+		TArray<FViewInfo*, TInlineAllocator<1> > Views;
+		const bool bWholeSceneDirectionalShadow = IsWholeSceneDirectionalShadow();
+
+		if (bWholeSceneDirectionalShadow)
+		{
+			Views.Add(DependentView);
+		}
+		else
 		{
-			TArray<FViewInfo*, TInlineAllocator<1> > Views;
-			const bool bWholeSceneDirectionalShadow = IsWholeSceneDirectionalShadow();
+			check(ViewArray);
 
-			if (bWholeSceneDirectionalShadow)
+			for (int32 ViewIndex = 0; ViewIndex < ViewArray->Num(); ViewIndex++)
 			{
-				Views.Add(DependentView);
+				Views.Add(&(*ViewArray)[ViewIndex]);
 			}
-			else
-			{
-				check(ViewArray);
+		}
 
-				for (int32 ViewIndex = 0; ViewIndex < ViewArray->Num(); ViewIndex++)
-				{
-					Views.Add(&(*ViewArray)[ViewIndex]);
-				}
-			}
+		bool bOpaqueRelevance = false;
+		bool bTranslucentRelevance = false;
+		bool bShadowRelevance = false;
+		bool bNeedsPreRenderView = false;
+		uint32 ViewMask = 0;
+		int32 PrimitiveId = PrimitiveSceneInfo->GetIndex();
+		const auto FeatureLevel = PrimitiveSceneInfo->Scene->GetFeatureLevel();
 
-			bool bOpaqueRelevance = false;
-			bool bTranslucentRelevance = false;
-			bool bShadowRelevance = false;
-			bool bNeedsPreRenderView = false;
-			uint32 ViewMask = 0;
-			int32 PrimitiveId = PrimitiveSceneInfo->GetIndex();
-			const auto FeatureLevel = PrimitiveSceneInfo->Scene->GetFeatureLevel();
+		for (int32 ViewIndex = 0, Num = Views.Num(); ViewIndex < Num; ViewIndex++)
+		{
+			FViewInfo& CurrentView = *Views[ViewIndex];
+			FPrimitiveViewRelevance& ViewRelevance = CurrentView.PrimitiveViewRelevanceMap[PrimitiveId];
 
-			for (int32 ViewIndex = 0, Num = Views.Num(); ViewIndex < Num; ViewIndex++)
+			if (!ViewRelevance.bInitializedThisFrame)
 			{
-				FViewInfo& CurrentView = *Views[ViewIndex];
-				FPrimitiveViewRelevance& ViewRelevance = CurrentView.PrimitiveViewRelevanceMap[PrimitiveId];
-
-				if (!ViewRelevance.bInitializedThisFrame)
+				if( CurrentView.IsPerspectiveProjection() )
 				{
-					if( CurrentView.IsPerspectiveProjection() )
-					{
-						const FPrimitiveSceneProxy* Proxy = PrimitiveSceneInfo->Proxy;
+					const FPrimitiveSceneProxy* Proxy = PrimitiveSceneInfo->Proxy;
 
-						// Compute the distance between the view and the primitive.
-						float DistanceSquared = (Proxy->GetBounds().Origin - CurrentView.ShadowViewMatrices.ViewOrigin).SizeSquared();
+					// Compute the distance between the view and the primitive.
+					float DistanceSquared = (Proxy->GetBounds().Origin - CurrentView.ShadowViewMatrices.ViewOrigin).SizeSquared();
 
-						bool bIsDistanceCulled = CurrentView.IsDistanceCulled(
-							DistanceSquared,
-							Proxy->GetMinDrawDistance(),
-							Proxy->GetMaxDrawDistance(),
-							PrimitiveSceneInfo
-							);
-						if( bIsDistanceCulled )
-						{
-							continue;
-						}
+					bool bIsDistanceCulled = CurrentView.IsDistanceCulled(
+						DistanceSquared,
+						Proxy->GetMinDrawDistance(),
+						Proxy->GetMaxDrawDistance(),
+						PrimitiveSceneInfo
+						);
+					if( bIsDistanceCulled )
+					{
+						continue;
 					}
-
-					// Compute the subject primitive's view relevance since it wasn't cached
-					// Update the main view's PrimitiveViewRelevanceMap
-					ViewRelevance = PrimitiveSceneInfo->Proxy->GetViewRelevance(&CurrentView);
-
-					bNeedsPreRenderView |= ViewRelevance.bNeedsPreRenderView;
-					ViewMask |= (1 << ViewIndex);
 				}
 
-				bOpaqueRelevance |= ViewRelevance.bOpaqueRelevance;
-				bTranslucentRelevance |= ViewRelevance.HasTranslucency();
-				bShadowRelevance |= ViewRelevance.bShadowRelevance;
-			}
+				// Compute the subject primitive's view relevance since it wasn't cached
+				// Update the main view's PrimitiveViewRelevanceMap
+				ViewRelevance = PrimitiveSceneInfo->Proxy->GetViewRelevance(&CurrentView);
 
-			if (bShadowRelevance)
-			{
-				// Update the primitive component's last render time. Allows the component to update when using bCastWhenHidden.
-				const float CurrentWorldTime = Views[0]->Family->CurrentWorldTime;
-				*(PrimitiveSceneInfo->ComponentLastRenderTime) = CurrentWorldTime;
+				bNeedsPreRenderView |= ViewRelevance.bNeedsPreRenderView;
+				ViewMask |= (1 << ViewIndex);
 			}
 
-			const bool bUseGetDynamicMeshElements = ShouldUseGetDynamicMeshElements();
+			bOpaqueRelevance |= ViewRelevance.bOpaqueRelevance;
+			bTranslucentRelevance |= ViewRelevance.HasTranslucency();
+			bShadowRelevance |= ViewRelevance.bShadowRelevance;
+		}
 
-			if (!bUseGetDynamicMeshElements && bNeedsPreRenderView)
-			{
-				// Call PreRenderView on primitives that weren't visible in any of the main views, but need to be rendered in this shadow's depth pass
-				PrimitiveSceneInfo->Proxy->PreRenderView(Views[0]->Family, ViewMask, Views[0]->FrameNumber);
-			}
+		if (bShadowRelevance)
+		{
+			// Update the primitive component's last render time. Allows the component to update when using bCastWhenHidden.
+			const float CurrentWorldTime = Views[0]->Family->CurrentWorldTime;
+			*(PrimitiveSceneInfo->ComponentLastRenderTime) = CurrentWorldTime;
+		}
 
-			if (bOpaqueRelevance && bShadowRelevance)
-			{
-				const FPrimitiveSceneProxy* Proxy = PrimitiveSceneInfo->Proxy;
-				const FBoxSphereBounds& Bounds = Proxy->GetBounds();
-				bool bDrawingStaticMeshes = false;
+		const bool bUseGetDynamicMeshElements = ShouldUseGetDynamicMeshElements();
 
-				if (PrimitiveSceneInfo->StaticMeshes.Num() > 0)
+		if (!bUseGetDynamicMeshElements && bNeedsPreRenderView)
+		{
+			// Call PreRenderView on primitives that weren't visible in any of the main views, but need to be rendered in this shadow's depth pass
+			PrimitiveSceneInfo->Proxy->PreRenderView(Views[0]->Family, ViewMask, Views[0]->FrameNumber);
+		}
+
+		if (bOpaqueRelevance && bShadowRelevance)
+		{
+			const FPrimitiveSceneProxy* Proxy = PrimitiveSceneInfo->Proxy;
+			const FBoxSphereBounds& Bounds = Proxy->GetBounds();
+			bool bDrawingStaticMeshes = false;
+
+			if (PrimitiveSceneInfo->StaticMeshes.Num() > 0)
+			{
+				for (int32 ViewIndex = 0; ViewIndex < Views.Num(); ViewIndex++)
 				{
-					for (int32 ViewIndex = 0; ViewIndex < Views.Num(); ViewIndex++)
+					FViewInfo& CurrentView = *Views[ViewIndex];
+
+					const float DistanceSquared = ( Bounds.Origin - CurrentView.ShadowViewMatrices.ViewOrigin ).SizeSquared();
+					const bool bDrawShadowDepth = FMath::Square( Bounds.SphereRadius ) > FMath::Square( GMinScreenRadiusForShadowCaster ) * DistanceSquared;
+					if( !bDrawShadowDepth )
 					{
-						FViewInfo& CurrentView = *Views[ViewIndex];
+						// cull object if it's too small to be considered as shadow caster
+						continue;
+					}
 
-						const float DistanceSquared = ( Bounds.Origin - CurrentView.ShadowViewMatrices.ViewOrigin ).SizeSquared();
-						const bool bDrawShadowDepth = FMath::Square( Bounds.SphereRadius ) > FMath::Square( GMinScreenRadiusForShadowCaster ) * DistanceSquared;
-						if( !bDrawShadowDepth )
-						{
-							// cull object if it's too small to be considered as shadow caster
-							continue;
-						}
+					// Update visibility for meshes which weren't visible in the main views or were visible with static relevance
+					if (!CurrentView.PrimitiveVisibilityMap[PrimitiveId] || CurrentView.PrimitiveViewRelevanceMap[PrimitiveId].bStaticRelevance)
+					{
+						bool bUseExistingVisibility = false;
 
-						// Update visibility for meshes which weren't visible in the main views or were visible with static relevance
-						if (!CurrentView.PrimitiveVisibilityMap[PrimitiveId] || CurrentView.PrimitiveViewRelevanceMap[PrimitiveId].bStaticRelevance)
+						if(!bReflectiveShadowmap) // Don't use existing visibility for RSMs 
 						{
-							bool bUseExistingVisibility = false;
-
-							if(!bReflectiveShadowmap) // Don't use existing visibility for RSMs 
+							for (int32 MeshIndex = 0; MeshIndex < PrimitiveSceneInfo->StaticMeshes.Num(); MeshIndex++)
 							{
-								for (int32 MeshIndex = 0; MeshIndex < PrimitiveSceneInfo->StaticMeshes.Num(); MeshIndex++)
-								{
-									const FStaticMesh& StaticMesh = PrimitiveSceneInfo->StaticMeshes[MeshIndex];
-									bool bMeshIsVisible = CurrentView.StaticMeshShadowDepthMap[StaticMesh.Id] && StaticMesh.CastShadow;
-									bUseExistingVisibility = bUseExistingVisibility || bMeshIsVisible;
+								const FStaticMesh& StaticMesh = PrimitiveSceneInfo->StaticMeshes[MeshIndex];
+								bool bMeshIsVisible = CurrentView.StaticMeshShadowDepthMap[StaticMesh.Id] && StaticMesh.CastShadow;
+								bUseExistingVisibility = bUseExistingVisibility || bMeshIsVisible;
 
-									if (bMeshIsVisible && bWholeSceneDirectionalShadow)
-									{
-										StaticMeshWholeSceneShadowDepthMap[StaticMesh.Id] = true;
-										StaticMeshWholeSceneShadowBatchVisibility[StaticMesh.Id] = CurrentView.StaticMeshBatchVisibility[StaticMesh.Id];
-									}
+								if (bMeshIsVisible && bWholeSceneDirectionalShadow)
+								{
+									StaticMeshWholeSceneShadowDepthMap[StaticMesh.Id] = true;
+									StaticMeshWholeSceneShadowBatchVisibility[StaticMesh.Id] = CurrentView.StaticMeshBatchVisibility[StaticMesh.Id];
 								}
 							}
+						}
+
+						if (bUseExistingVisibility)
+						{
+							bDrawingStaticMeshes = true;
+						}
+						// Don't overwrite visibility set by the main views
+						// This is necessary to avoid popping when transitioning between LODs, because on the frame of the transition, 
+						// The old LOD will continue to be drawn even though a different LOD would be chosen by distance.
+						else
+						{
+							int8 LODToRender = 0;
+							int32 ForcedLODLevel = (CurrentView.Family->EngineShowFlags.LOD) ? GetCVarForceLOD() : 0;
 
-							if (bUseExistingVisibility)
+							// Add the primitive's static mesh elements to the draw lists.
+							if ( bReflectiveShadowmap) 
 							{
-								bDrawingStaticMeshes = true;
+								LODToRender = -CHAR_MAX;
+								// Force the lowest detail LOD Level in reflective shadow maps.
+								for (int32 Index = 0; Index < PrimitiveSceneInfo->StaticMeshes.Num(); Index++)
+								{
+									LODToRender = FMath::Max<int8>(PrimitiveSceneInfo->StaticMeshes[Index].LODIndex, LODToRender);
+								}
 							}
-							// Don't overwrite visibility set by the main views
-							// This is necessary to avoid popping when transitioning between LODs, because on the frame of the transition, 
-							// The old LOD will continue to be drawn even though a different LOD would be chosen by distance.
 							else
 							{
-								int8 LODToRender = 0;
-								int32 ForcedLODLevel = (CurrentView.Family->EngineShowFlags.LOD) ? GetCVarForceLOD() : 0;
+								FPrimitiveBounds PrimitiveBounds;
+								PrimitiveBounds.Origin = Bounds.Origin;
+								PrimitiveBounds.SphereRadius = Bounds.SphereRadius;
+								LODToRender = ComputeLODForMeshes(PrimitiveSceneInfo->StaticMeshes, CurrentView, PrimitiveBounds.Origin, PrimitiveBounds.SphereRadius, ForcedLODLevel);
+							}
 
-								// Add the primitive's static mesh elements to the draw lists.
-								if ( bReflectiveShadowmap) 
+							for (int32 MeshIndex = 0; MeshIndex < PrimitiveSceneInfo->StaticMeshes.Num(); MeshIndex++)
+							{
+								const FStaticMesh& StaticMesh = PrimitiveSceneInfo->StaticMeshes[MeshIndex];
+								if (StaticMesh.CastShadow && StaticMesh.LODIndex == LODToRender)
 								{
-									LODToRender = -CHAR_MAX;
-									// Force the lowest detail LOD Level in reflective shadow maps.
-									for (int32 Index = 0; Index < PrimitiveSceneInfo->StaticMeshes.Num(); Index++)
+									if (bWholeSceneDirectionalShadow)
 									{
-										LODToRender = FMath::Max<int8>(PrimitiveSceneInfo->StaticMeshes[Index].LODIndex, LODToRender);
+										StaticMeshWholeSceneShadowDepthMap[StaticMesh.Id] = true;
+										StaticMeshWholeSceneShadowBatchVisibility[StaticMesh.Id] = StaticMesh.Elements.Num() == 1 ? 1 : StaticMesh.VertexFactory->GetStaticBatchElementVisibility(*DependentView, &StaticMesh);
 									}
-								}
-								else
-								{
-									FPrimitiveBounds PrimitiveBounds;
-									PrimitiveBounds.Origin = Bounds.Origin;
-									PrimitiveBounds.SphereRadius = Bounds.SphereRadius;
-									LODToRender = ComputeLODForMeshes(PrimitiveSceneInfo->StaticMeshes, CurrentView, PrimitiveBounds.Origin, PrimitiveBounds.SphereRadius, ForcedLODLevel);
-								}
-
-								for (int32 MeshIndex = 0; MeshIndex < PrimitiveSceneInfo->StaticMeshes.Num(); MeshIndex++)
-								{
-									const FStaticMesh& StaticMesh = PrimitiveSceneInfo->StaticMeshes[MeshIndex];
-									if (StaticMesh.CastShadow && StaticMesh.LODIndex == LODToRender)
+									else
 									{
-										if (bWholeSceneDirectionalShadow)
-										{
-											StaticMeshWholeSceneShadowDepthMap[StaticMesh.Id] = true;
-											StaticMeshWholeSceneShadowBatchVisibility[StaticMesh.Id] = StaticMesh.Elements.Num() == 1 ? 1 : StaticMesh.VertexFactory->GetStaticBatchElementVisibility(*DependentView, &StaticMesh);
-										}
-										else
-										{
-											CurrentView.StaticMeshShadowDepthMap[StaticMesh.Id] = true;
-											CurrentView.StaticMeshBatchVisibility[StaticMesh.Id] = StaticMesh.Elements.Num() == 1 ? 1 : StaticMesh.VertexFactory->GetStaticBatchElementVisibility(CurrentView, &StaticMesh);
-										}
-
-										bDrawingStaticMeshes = true;
+										CurrentView.StaticMeshShadowDepthMap[StaticMesh.Id] = true;
+										CurrentView.StaticMeshBatchVisibility[StaticMesh.Id] = StaticMesh.Elements.Num() == 1 ? 1 : StaticMesh.VertexFactory->GetStaticBatchElementVisibility(CurrentView, &StaticMesh);
 									}
+
+									bDrawingStaticMeshes = true;
 								}
 							}
 						}
 					}
 				}
+			}
 
-				if (bDrawingStaticMeshes)
+			if (bDrawingStaticMeshes)
+			{
+				if (!bWholeSceneDirectionalShadow)
 				{
-					if (!bWholeSceneDirectionalShadow)
+					// Add the primitive's static mesh elements to the draw lists.
+					for (int32 MeshIndex = 0; MeshIndex < PrimitiveSceneInfo->StaticMeshes.Num(); MeshIndex++)
 					{
-						// Add the primitive's static mesh elements to the draw lists.
-						for (int32 MeshIndex = 0; MeshIndex < PrimitiveSceneInfo->StaticMeshes.Num(); MeshIndex++)
+						FStaticMesh& StaticMesh = PrimitiveSceneInfo->StaticMeshes[MeshIndex];
+						if (StaticMesh.CastShadow)
 						{
-							FStaticMesh& StaticMesh = PrimitiveSceneInfo->StaticMeshes[MeshIndex];
-							if (StaticMesh.CastShadow)
-							{
-								const FMaterialRenderProxy* MaterialRenderProxy = StaticMesh.MaterialRenderProxy;
-								const FMaterial* Material = MaterialRenderProxy->GetMaterial(FeatureLevel);
-								const EBlendMode BlendMode = Material->GetBlendMode();
-								const EMaterialShadingModel ShadingModel = Material->GetShadingModel();
+							const FMaterialRenderProxy* MaterialRenderProxy = StaticMesh.MaterialRenderProxy;
+							const FMaterial* Material = MaterialRenderProxy->GetMaterial(FeatureLevel);
+							const EBlendMode BlendMode = Material->GetBlendMode();
+							const EMaterialShadingModel ShadingModel = Material->GetShadingModel();
 
-								if(((!IsTranslucentBlendMode(BlendMode)) && ShadingModel != MSM_Unlit) || (bReflectiveShadowmap && Material->ShouldInjectEmissiveIntoLPV())) 
-								{
-									const bool bTwoSided = Material->IsTwoSided() || PrimitiveSceneInfo->Proxy->CastsShadowAsTwoSided();
-									OverrideWithDefaultMaterialForShadowDepth(MaterialRenderProxy, Material, bReflectiveShadowmap, FeatureLevel);
-									SubjectMeshElements.Add(FShadowStaticMeshElement(MaterialRenderProxy, Material, &StaticMesh,bTwoSided));
-								}
+							if(((!IsTranslucentBlendMode(BlendMode)) && ShadingModel != MSM_Unlit) || (bReflectiveShadowmap && Material->ShouldInjectEmissiveIntoLPV())) 
+							{
+								const bool bTwoSided = Material->IsTwoSided() || PrimitiveSceneInfo->Proxy->CastsShadowAsTwoSided();
+								OverrideWithDefaultMaterialForShadowDepth(MaterialRenderProxy, Material, bReflectiveShadowmap, FeatureLevel);
+								SubjectMeshElements.Add(FShadowStaticMeshElement(MaterialRenderProxy, Material, &StaticMesh,bTwoSided));
 							}
 						}
 					}
 				}
-				else
-				{
-					// Add the primitive to the subject primitive list.
-					SubjectPrimitives.Add(PrimitiveSceneInfo);
-				}
 			}
-
-			// Add translucent shadow casting primitives to SubjectTranslucentPrimitives
-			if (bTranslucentRelevance && bShadowRelevance && bTranslucentShadow)
+			else
 			{
-				SubjectTranslucentPrimitives.Add(PrimitiveSceneInfo);
+				// Add the primitive to the subject primitive list.
+				SubjectPrimitives.Add(PrimitiveSceneInfo);
 			}
 		}
+
+		// Add translucent shadow casting primitives to SubjectTranslucentPrimitives
+		if (bTranslucentRelevance && bShadowRelevance && bTranslucentShadow)
+		{
+			SubjectTranslucentPrimitives.Add(PrimitiveSceneInfo);
+		}
 	}
 }
 
@@ -1658,15 +1645,19 @@ void FDeferredShadingSceneRenderer::CreateWholeSceneProjectedShadow(FLightSceneI
 					}
 				}
 
-				// Add all the shadow casting primitives affected by the light to the shadow's subject primitive list.
-				for(FLightPrimitiveInteraction* Interaction = LightSceneInfo->DynamicPrimitiveList;
-					Interaction;
-					Interaction = Interaction->GetNextPrimitive())
+				// Ray traced shadows use the GPU managed distance field object buffers, no CPU culling should be used
+				if (!ProjectedShadowInfo->bRayTracedDistanceFieldShadow)
 				{
-					if (Interaction->HasShadow() 
-						&& (!bStaticSceneOnly || Interaction->GetPrimitiveSceneInfo()->Proxy->HasStaticLighting()))
+					// Add all the shadow casting primitives affected by the light to the shadow's subject primitive list.
+					for(FLightPrimitiveInteraction* Interaction = LightSceneInfo->DynamicPrimitiveList;
+						Interaction;
+						Interaction = Interaction->GetNextPrimitive())
 					{
-						ProjectedShadowInfo->AddSubjectPrimitive(Interaction->GetPrimitiveSceneInfo(), &Views);
+						if (Interaction->HasShadow() 
+							&& (!bStaticSceneOnly || Interaction->GetPrimitiveSceneInfo()->Proxy->HasStaticLighting()))
+						{
+							ProjectedShadowInfo->AddSubjectPrimitive(Interaction->GetPrimitiveSceneInfo(), &Views);
+						}
 					}
 				}
 			}
@@ -1913,10 +1904,6 @@ inline void FSceneRenderer::GatherShadowsForPrimitiveInner(
 						{
 							MinScreenRadiusForShadowCaster = GMinScreenRadiusForShadowCasterRSM;
 						}
-						else if (ProjectedShadowInfo->bRayTracedDistanceFieldShadow)
-						{
-							MinScreenRadiusForShadowCaster = GMinScreenRadiusForRayTracedCasters;
-						}
 
 						bool bScreenSpaceSizeCulled = false;
 						check( ProjectedShadowInfo->DependentView );
@@ -2048,7 +2035,11 @@ void FSceneRenderer::GatherShadowPrimitives(
 	}
 }
 
-void FSceneRenderer::AddViewDependentWholeSceneShadowsForView(TArray<FProjectedShadowInfo*, SceneRenderingAllocator>& ShadowInfos, FVisibleLightInfo& VisibleLightInfo, FLightSceneInfo& LightSceneInfo)
+void FSceneRenderer::AddViewDependentWholeSceneShadowsForView(
+	TArray<FProjectedShadowInfo*, SceneRenderingAllocator>& ShadowInfos, 
+	TArray<FProjectedShadowInfo*, SceneRenderingAllocator>& ShadowInfosThatNeedCulling,
+	FVisibleLightInfo& VisibleLightInfo, 
+	FLightSceneInfo& LightSceneInfo)
 {
 	SCOPE_CYCLE_COUNTER(STAT_AddViewDependentWholeSceneShadowsForView);
 
@@ -2096,6 +2087,12 @@ void FSceneRenderer::AddViewDependentWholeSceneShadowsForView(TArray<FProjectedS
 					VisibleLightInfo.MemStackProjectedShadows.Add(ProjectedShadowInfo);
 					VisibleLightInfo.AllProjectedShadows.Add(ProjectedShadowInfo);
 					ShadowInfos.Add(ProjectedShadowInfo);
+
+					// Ray traced shadows use the GPU managed distance field object buffers, no CPU culling needed
+					if (!ProjectedShadowInfo->bRayTracedDistanceFieldShadow)
+					{
+						ShadowInfosThatNeedCulling.Add(ProjectedShadowInfo);
+					}
 				}
 			}
 
@@ -2128,6 +2125,12 @@ void FSceneRenderer::AddViewDependentWholeSceneShadowsForView(TArray<FProjectedS
 						VisibleLightInfo.AllProjectedShadows.Add(ProjectedShadowInfo);
 						VisibleLightInfo.ReflectiveShadowMaps.Add(ProjectedShadowInfo);
 						ShadowInfos.Add(ProjectedShadowInfo); // or separate list?
+
+						// Ray traced shadows use the GPU managed distance field object buffers, no CPU culling needed
+						if (!ProjectedShadowInfo->bRayTracedDistanceFieldShadow)
+						{
+							ShadowInfosThatNeedCulling.Add(ProjectedShadowInfo);
+						}
 					}
 				}
 			}
@@ -2138,6 +2141,7 @@ void FSceneRenderer::AddViewDependentWholeSceneShadowsForView(TArray<FProjectedS
 void FForwardShadingSceneRenderer::InitDynamicShadows(FRHICommandListImmediate& RHICmdList)
 {	
 	TArray<FProjectedShadowInfo*, SceneRenderingAllocator> ViewDependentWholeSceneShadows;
+	TArray<FProjectedShadowInfo*, SceneRenderingAllocator> ViewDependentWholeSceneShadowsThatNeedCulling;
 	TArray<FProjectedShadowInfo*, SceneRenderingAllocator> PreShadows;
 	{
 		SCOPE_CYCLE_COUNTER(STAT_InitDynamicShadowsTime);
@@ -2168,7 +2172,7 @@ void FForwardShadingSceneRenderer::InitDynamicShadows(FRHICommandListImmediate&
 
 				if (bIsVisibleInAnyView)
 				{									
-					AddViewDependentWholeSceneShadowsForView(ViewDependentWholeSceneShadows, VisibleLightInfo, *LightSceneInfo);					
+					AddViewDependentWholeSceneShadowsForView(ViewDependentWholeSceneShadows, ViewDependentWholeSceneShadowsThatNeedCulling, VisibleLightInfo, *LightSceneInfo);					
 				}
 
 				const int32 NumShadows = VisibleLightInfo.AllProjectedShadows.Num();
@@ -2214,7 +2218,7 @@ void FForwardShadingSceneRenderer::InitDynamicShadows(FRHICommandListImmediate&
 	}
 
 	// Gathers the list of primitives used to draw various shadow types
-	GatherShadowPrimitives(PreShadows, ViewDependentWholeSceneShadows, false);
+	GatherShadowPrimitives(PreShadows, ViewDependentWholeSceneShadowsThatNeedCulling, false);
 
 	// Generate mesh element arrays from shadow primitive arrays
 	GatherShadowDynamicMeshElements();
@@ -2234,6 +2238,7 @@ void FDeferredShadingSceneRenderer::InitDynamicShadows(FRHICommandListImmediate&
 
 	TArray<FProjectedShadowInfo*,SceneRenderingAllocator> PreShadows;
 	TArray<FProjectedShadowInfo*,SceneRenderingAllocator> ViewDependentWholeSceneShadows;
+	TArray<FProjectedShadowInfo*,SceneRenderingAllocator> ViewDependentWholeSceneShadowsThatNeedCulling;
 	{
 		SCOPE_CYCLE_COUNTER(STAT_InitDynamicShadowsTime);
 
@@ -2297,7 +2302,7 @@ void FDeferredShadingSceneRenderer::InitDynamicShadows(FRHICommandListImmediate&
 					// Allow movable and stationary lights to create CSM, or static lights that are unbuilt
 					if (!LightSceneInfo->Proxy->HasStaticLighting() || bCreateShadowToPreviewStaticLight)
 					{
-						AddViewDependentWholeSceneShadowsForView(ViewDependentWholeSceneShadows, VisibleLightInfo, *LightSceneInfo);
+						AddViewDependentWholeSceneShadowsForView(ViewDependentWholeSceneShadows, ViewDependentWholeSceneShadowsThatNeedCulling, VisibleLightInfo, *LightSceneInfo);
 
 						// Look for individual primitives with a dynamic shadow.
 						for (FLightPrimitiveInteraction* Interaction = LightSceneInfo->DynamicPrimitiveList;
@@ -2320,7 +2325,7 @@ void FDeferredShadingSceneRenderer::InitDynamicShadows(FRHICommandListImmediate&
 	UpdatePreshadowCache();
 
 	// Gathers the list of primitives used to draw various shadow types
-	GatherShadowPrimitives(PreShadows, ViewDependentWholeSceneShadows, bStaticSceneOnly);
+	GatherShadowPrimitives(PreShadows, ViewDependentWholeSceneShadowsThatNeedCulling, bStaticSceneOnly);
 
 	// Generate mesh element arrays from shadow primitive arrays
 	GatherShadowDynamicMeshElements();
