OS X's OpenGL doesn't support SM5, so don't try and use SM5 RWBuffers unless SM5 is supported.

diff --git a/Engine/Source/Runtime/Renderer/Private/DistanceFieldSurfaceCacheLighting.cpp b/Engine/Source/Runtime/Renderer/Private/DistanceFieldSurfaceCacheLighting.cpp
--- a/Engine/Source/Runtime/Renderer/Private/DistanceFieldSurfaceCacheLighting.cpp
+++ b/Engine/Source/Runtime/Renderer/Private/DistanceFieldSurfaceCacheLighting.cpp
@@ -2217,12 +2217,18 @@ class FDebugBuffer : public FRenderResource
 public:
 	virtual void InitDynamicRHI() override
 	{
-		DebugData.Initialize(sizeof(float)* 4, 50, PF_A32B32G32R32F, BUF_Static);
+        if(GMaxRHIFeatureLevel >= ERHIFeatureLevel::SM5)
+        {
+            DebugData.Initialize(sizeof(float)* 4, 50, PF_A32B32G32R32F, BUF_Static);
+        }
 	}
 
 	virtual void ReleaseDynamicRHI() override
 	{
-		DebugData.Release();
+        if(GMaxRHIFeatureLevel >= ERHIFeatureLevel::SM5)
+        {
+            DebugData.Release();
+        }
 	}
 
 	FRWBuffer DebugData;
