UE4 - New lightmap policy, adds lightmap + directional csm for static meshes

diff --git a/Engine/Shaders/BasePassForForwardShadingPixelShader.usf b/Engine/Shaders/BasePassForForwardShadingPixelShader.usf
--- a/Engine/Shaders/BasePassForForwardShadingPixelShader.usf
+++ b/Engine/Shaders/BasePassForForwardShadingPixelShader.usf
@@ -187,7 +187,7 @@ void Main(
 			Settings.ProjectionDepthBiasParameters = 0;
 						
 			float4 ShadowPosition;	
-#if MAX_FORWARD_SHADOWCASCADES == 2			
+	#if MAX_FORWARD_SHADOWCASCADES == 2			
 			{				
 				if (MaterialParameters.ScreenPosition.w >= View.DirectionalLightShadowDistances.x)
 				{
@@ -198,23 +198,23 @@ void Main(
 					ShadowPosition = mul(float4(MaterialParameters.ScreenPosition.xyw, 1), View.DirectionalLightScreenToShadow[0]);
 				}
 			}
-#else			
+	#else			
 			{
 				ShadowPosition = mul(float4(MaterialParameters.ScreenPosition.xyw, 1), View.DirectionalLightScreenToShadow[0]);				
 			}
-#endif
+	#endif
 
 			// Clamp pixel depth in light space for shadowing opaque, because areas of the shadow depth buffer that weren't rendered to will have been cleared to 1
 			// We want to force the shadow comparison to result in 'unshadowed' in that case, regardless of whether the pixel being shaded is in front or behind that plane
 			float LightSpacePixelDepthForOpaque = min(ShadowPosition.z, 0.99999f);
 			Settings.SceneDepth = LightSpacePixelDepthForOpaque;
 
-#if 0
+	#if 0
 			half ShadowMap = Manual1x1PCF(ShadowPosition.xy, Settings);
-#else
+	#else
 			half ShadowMap = Manual2x2PCF(ShadowPosition.xy, Settings);
-#endif
-			Shadow = min(Shadow, ShadowMap);
+	#endif
+			Shadow = ShadowMap;
 		}
 #endif
 
diff --git a/Engine/Shaders/LightmapCommon.usf b/Engine/Shaders/LightmapCommon.usf
--- a/Engine/Shaders/LightmapCommon.usf
+++ b/Engine/Shaders/LightmapCommon.usf
@@ -146,21 +146,22 @@ half GetPrimaryPrecomputedShadowMask(FVertexFactoryInterpolantsVSToPS Interpolan
 		half ShadowFactor = saturate( DistanceField * DistanceFieldParameters.x + DistanceFieldParameters.y );
 		return ShadowFactor * ShadowFactor;
 
+	#elif MOVABLE_DIRECTIONAL_LIGHT
+		// Do this before checking for lightmaps as we might have a lightmap + directional light
+		//@mw todo
+		// TODO: Filter shadowmap
+		return 1;
+
 	#elif HQ_TEXTURE_LIGHTMAP || LQ_TEXTURE_LIGHTMAP
 
 		// Mark as shadowed for lightmapped objects with no shadowmap
 		// This is necessary because objects inside a light's influence that were determined to be completely shadowed won't be rendered with STATICLIGHTING_TEXTUREMASK==1
 		return 0;
 
-	#elif MOVABLE_DIRECTIONAL_LIGHT
-		//@mw todo
-		// TODO: Filter shadowmap
-		return 1;
-
 	#else
 
 		// Mark as unshadowed by default
 		return 1;
 
 	#endif
-}
\ No newline at end of file
+}
diff --git a/Engine/Source/Runtime/Renderer/Private/ForwardBasePassRendering.cpp b/Engine/Source/Runtime/Renderer/Private/ForwardBasePassRendering.cpp
--- a/Engine/Source/Runtime/Renderer/Private/ForwardBasePassRendering.cpp
+++ b/Engine/Source/Runtime/Renderer/Private/ForwardBasePassRendering.cpp
@@ -33,6 +33,8 @@ IMPLEMENT_FORWARD_SHADING_BASEPASS_LIGHTMAPPED_SHADER_TYPE( TDistanceFieldShadow
 IMPLEMENT_FORWARD_SHADING_BASEPASS_LIGHTMAPPED_SHADER_TYPE( FSimpleDirectionalLightAndSHIndirectPolicy, FSimpleDirectionalLightAndSHIndirectPolicy );
 IMPLEMENT_FORWARD_SHADING_BASEPASS_LIGHTMAPPED_SHADER_TYPE( FMovableDirectionalLightLightingPolicy, FMovableDirectionalLightLightingPolicy );
 IMPLEMENT_FORWARD_SHADING_BASEPASS_LIGHTMAPPED_SHADER_TYPE( FMovableDirectionalLightCSMLightingPolicy, FMovableDirectionalLightCSMLightingPolicy );
+IMPLEMENT_FORWARD_SHADING_BASEPASS_LIGHTMAPPED_SHADER_TYPE(FMovableDirectionalLightWithLightmapLightingPolicy, FMovableDirectionalLightWithLightmapLightingPolicy);
+IMPLEMENT_FORWARD_SHADING_BASEPASS_LIGHTMAPPED_SHADER_TYPE(FMovableDirectionalLightCSMWithLightmapLightingPolicy, FMovableDirectionalLightCSMWithLightmapLightingPolicy);
 
 /** The action used to draw a base pass static mesh element. */
 class FDrawBasePassForwardShadingStaticMeshAction
@@ -332,6 +334,8 @@ void FForwardShadingSceneRenderer::RenderForwardShadingBasePass(FRHICommandListI
 			Scene->BasePassForForwardShadingLowQualityLightMapDrawList[DrawType].SortFrontToBack(Views[0].ViewLocation);
 			Scene->BasePassForForwardShadingDistanceFieldShadowMapLightMapDrawList[DrawType].SortFrontToBack(Views[0].ViewLocation);
 			Scene->BasePassForForwardShadingDirectionalLightAndSHIndirectDrawList[DrawType].SortFrontToBack(Views[0].ViewLocation);
+			Scene->BasePassForForwardShadingMovableDirectionalLightCSMLightmapDrawList[DrawType].SortFrontToBack(Views[0].ViewLocation);
+			Scene->BasePassForForwardShadingMovableDirectionalLightLightmapDrawList[DrawType].SortFrontToBack(Views[0].ViewLocation);
 			Scene->BasePassForForwardShadingMovableDirectionalLightCSMDrawList[DrawType].SortFrontToBack(Views[0].ViewLocation);
 			Scene->BasePassForForwardShadingMovableDirectionalLightDrawList[DrawType].SortFrontToBack(Views[0].ViewLocation);
 			Scene->BasePassForForwardShadingNoLightMapDrawList[DrawType].SortFrontToBack(Views[0].ViewLocation);
@@ -358,6 +362,8 @@ void FForwardShadingSceneRenderer::RenderForwardShadingBasePass(FRHICommandListI
 				MaxDraws -= Scene->BasePassForForwardShadingLowQualityLightMapDrawList[FScene::EBasePass_Default].DrawVisibleFrontToBack(RHICmdList, View, View.StaticMeshVisibilityMap, View.StaticMeshBatchVisibility, MaxDraws);
 				MaxDraws -= Scene->BasePassForForwardShadingDistanceFieldShadowMapLightMapDrawList[FScene::EBasePass_Default].DrawVisibleFrontToBack(RHICmdList, View, View.StaticMeshVisibilityMap, View.StaticMeshBatchVisibility, MaxDraws);
 				MaxDraws -= Scene->BasePassForForwardShadingDirectionalLightAndSHIndirectDrawList[FScene::EBasePass_Default].DrawVisibleFrontToBack(RHICmdList, View, View.StaticMeshVisibilityMap, View.StaticMeshBatchVisibility, MaxDraws);
+				MaxDraws -= Scene->BasePassForForwardShadingMovableDirectionalLightCSMLightmapDrawList[FScene::EBasePass_Default].DrawVisibleFrontToBack(RHICmdList, View, View.StaticMeshVisibilityMap, View.StaticMeshBatchVisibility, MaxDraws);
+				MaxDraws -= Scene->BasePassForForwardShadingMovableDirectionalLightLightmapDrawList[FScene::EBasePass_Default].DrawVisibleFrontToBack(RHICmdList, View, View.StaticMeshVisibilityMap, View.StaticMeshBatchVisibility, MaxDraws);
 				MaxDraws -= Scene->BasePassForForwardShadingMovableDirectionalLightCSMDrawList[FScene::EBasePass_Default].DrawVisibleFrontToBack(RHICmdList, View, View.StaticMeshVisibilityMap, View.StaticMeshBatchVisibility, MaxDraws);
 				MaxDraws -= Scene->BasePassForForwardShadingMovableDirectionalLightDrawList[FScene::EBasePass_Default].DrawVisibleFrontToBack(RHICmdList, View, View.StaticMeshVisibilityMap, View.StaticMeshBatchVisibility, MaxDraws);
 				MaxDraws -= Scene->BasePassForForwardShadingNoLightMapDrawList[FScene::EBasePass_Default].DrawVisibleFrontToBack(RHICmdList, View, View.StaticMeshVisibilityMap, View.StaticMeshBatchVisibility, MaxDraws);
@@ -368,6 +374,8 @@ void FForwardShadingSceneRenderer::RenderForwardShadingBasePass(FRHICommandListI
 				Scene->BasePassForForwardShadingLowQualityLightMapDrawList[FScene::EBasePass_Default].DrawVisible(RHICmdList, View, View.StaticMeshVisibilityMap, View.StaticMeshBatchVisibility);
 				Scene->BasePassForForwardShadingDistanceFieldShadowMapLightMapDrawList[FScene::EBasePass_Default].DrawVisible(RHICmdList, View, View.StaticMeshVisibilityMap, View.StaticMeshBatchVisibility);
 				Scene->BasePassForForwardShadingDirectionalLightAndSHIndirectDrawList[FScene::EBasePass_Default].DrawVisible(RHICmdList, View, View.StaticMeshVisibilityMap, View.StaticMeshBatchVisibility);
+				Scene->BasePassForForwardShadingMovableDirectionalLightCSMLightmapDrawList[FScene::EBasePass_Default].DrawVisible(RHICmdList, View, View.StaticMeshVisibilityMap, View.StaticMeshBatchVisibility);
+				Scene->BasePassForForwardShadingMovableDirectionalLightLightmapDrawList[FScene::EBasePass_Default].DrawVisible(RHICmdList, View, View.StaticMeshVisibilityMap, View.StaticMeshBatchVisibility);
 				Scene->BasePassForForwardShadingMovableDirectionalLightCSMDrawList[FScene::EBasePass_Default].DrawVisible(RHICmdList, View, View.StaticMeshVisibilityMap, View.StaticMeshBatchVisibility);
 				Scene->BasePassForForwardShadingMovableDirectionalLightDrawList[FScene::EBasePass_Default].DrawVisible(RHICmdList, View, View.StaticMeshVisibilityMap, View.StaticMeshBatchVisibility);
 				Scene->BasePassForForwardShadingNoLightMapDrawList[FScene::EBasePass_Default].DrawVisible(RHICmdList, View, View.StaticMeshVisibilityMap, View.StaticMeshBatchVisibility);
@@ -445,6 +453,8 @@ void FForwardShadingSceneRenderer::RenderForwardShadingBasePass(FRHICommandListI
 			MaxDraws -= Scene->BasePassForForwardShadingLowQualityLightMapDrawList[FScene::EBasePass_Masked].DrawVisibleFrontToBack(RHICmdList, View,View.StaticMeshVisibilityMap,View.StaticMeshBatchVisibility,MaxDraws);
 			MaxDraws -= Scene->BasePassForForwardShadingDistanceFieldShadowMapLightMapDrawList[FScene::EBasePass_Masked].DrawVisibleFrontToBack(RHICmdList, View,View.StaticMeshVisibilityMap,View.StaticMeshBatchVisibility,MaxDraws);
 			MaxDraws -= Scene->BasePassForForwardShadingDirectionalLightAndSHIndirectDrawList[FScene::EBasePass_Masked].DrawVisibleFrontToBack(RHICmdList, View,View.StaticMeshVisibilityMap,View.StaticMeshBatchVisibility,MaxDraws);
+			MaxDraws -= Scene->BasePassForForwardShadingMovableDirectionalLightLightmapDrawList[FScene::EBasePass_Masked].DrawVisibleFrontToBack(RHICmdList, View, View.StaticMeshVisibilityMap, View.StaticMeshBatchVisibility, MaxDraws);
+			MaxDraws -= Scene->BasePassForForwardShadingMovableDirectionalLightCSMLightmapDrawList[FScene::EBasePass_Masked].DrawVisibleFrontToBack(RHICmdList, View, View.StaticMeshVisibilityMap, View.StaticMeshBatchVisibility, MaxDraws);
 			MaxDraws -= Scene->BasePassForForwardShadingMovableDirectionalLightCSMDrawList[FScene::EBasePass_Masked].DrawVisibleFrontToBack(RHICmdList, View, View.StaticMeshVisibilityMap, View.StaticMeshBatchVisibility, MaxDraws);
 			MaxDraws -= Scene->BasePassForForwardShadingMovableDirectionalLightDrawList[FScene::EBasePass_Masked].DrawVisibleFrontToBack(RHICmdList, View,View.StaticMeshVisibilityMap,View.StaticMeshBatchVisibility,MaxDraws);
 		}
@@ -455,6 +465,8 @@ void FForwardShadingSceneRenderer::RenderForwardShadingBasePass(FRHICommandListI
 			Scene->BasePassForForwardShadingLowQualityLightMapDrawList[FScene::EBasePass_Masked].DrawVisible(RHICmdList, View,View.StaticMeshVisibilityMap,View.StaticMeshBatchVisibility);
 			Scene->BasePassForForwardShadingDistanceFieldShadowMapLightMapDrawList[FScene::EBasePass_Masked].DrawVisible(RHICmdList, View,View.StaticMeshVisibilityMap,View.StaticMeshBatchVisibility);
 			Scene->BasePassForForwardShadingDirectionalLightAndSHIndirectDrawList[FScene::EBasePass_Masked].DrawVisible(RHICmdList, View,View.StaticMeshVisibilityMap,View.StaticMeshBatchVisibility);
+			Scene->BasePassForForwardShadingMovableDirectionalLightCSMLightmapDrawList[FScene::EBasePass_Masked].DrawVisible(RHICmdList, View, View.StaticMeshVisibilityMap, View.StaticMeshBatchVisibility);
+			Scene->BasePassForForwardShadingMovableDirectionalLightLightmapDrawList[FScene::EBasePass_Masked].DrawVisible(RHICmdList, View, View.StaticMeshVisibilityMap, View.StaticMeshBatchVisibility);
 			Scene->BasePassForForwardShadingMovableDirectionalLightCSMDrawList[FScene::EBasePass_Masked].DrawVisible(RHICmdList, View, View.StaticMeshVisibilityMap, View.StaticMeshBatchVisibility);
 			Scene->BasePassForForwardShadingMovableDirectionalLightDrawList[FScene::EBasePass_Masked].DrawVisible(RHICmdList, View,View.StaticMeshVisibilityMap,View.StaticMeshBatchVisibility);
 		}
diff --git a/Engine/Source/Runtime/Renderer/Private/ForwardBasePassRendering.h b/Engine/Source/Runtime/Renderer/Private/ForwardBasePassRendering.h
--- a/Engine/Source/Runtime/Renderer/Private/ForwardBasePassRendering.h
+++ b/Engine/Source/Runtime/Renderer/Private/ForwardBasePassRendering.h
@@ -524,7 +524,19 @@ void ProcessBasePassMeshForForwardShading(
 		const FLightSceneInfo* SimpleDirectionalLight = Action.GetSimpleDirectionalLight();
 		const bool bUseMovableLight = SimpleDirectionalLight && !SimpleDirectionalLight->Proxy->HasStaticShadowing();
 
-		if (bUseMovableLight)
+		if (bUseMovableLight && LightMapInteraction.GetType() == LMIT_Texture)
+		{
+			// final determination of whether CSMs are rendered can be view dependent, thus we always need to clear the CSMs even if we're not going to render to them based on the condition below.
+			if (SimpleDirectionalLight && SimpleDirectionalLight->ShouldRenderViewIndependentWholeSceneShadows())
+			{
+				Action.template Process<FMovableDirectionalLightCSMWithLightmapLightingPolicy>(RHICmdList, Parameters, FMovableDirectionalLightCSMWithLightmapLightingPolicy(), LightMapInteraction);
+			}
+			else
+			{
+				Action.template Process<FMovableDirectionalLightWithLightmapLightingPolicy>(RHICmdList, Parameters, FMovableDirectionalLightCSMWithLightmapLightingPolicy(), LightMapInteraction);
+			}
+		}
+		else if (bUseMovableLight)
 		{
 			// final determination of whether CSMs are rendered can be view dependent, thus we always need to clear the CSMs even if we're not going to render to them based on the condition below.
 			if (SimpleDirectionalLight && SimpleDirectionalLight->ShouldRenderViewIndependentWholeSceneShadows())
diff --git a/Engine/Source/Runtime/Renderer/Private/LightMapRendering.h b/Engine/Source/Runtime/Renderer/Private/LightMapRendering.h
--- a/Engine/Source/Runtime/Renderer/Private/LightMapRendering.h
+++ b/Engine/Source/Runtime/Renderer/Private/LightMapRendering.h
@@ -996,4 +996,34 @@ public:
 	static void ModifyCompilationEnvironment(EShaderPlatform Platform, const FMaterial* Material, FShaderCompilerEnvironment& OutEnvironment);
 };
 
+class FMovableDirectionalLightWithLightmapLightingPolicy : public TLightMapPolicy<LQ_LIGHTMAP>
+{
+	typedef TLightMapPolicy< LQ_LIGHTMAP >	Super;
+
+public:
+	static bool ShouldCache(EShaderPlatform Platform, const FMaterial* Material, const FVertexFactoryType* VertexFactoryType)
+	{
+		return (Material->GetShadingModel() != MSM_Unlit) && Super::ShouldCache(Platform, Material, VertexFactoryType);
+	}
+
+	static void ModifyCompilationEnvironment(EShaderPlatform Platform, const FMaterial* Material, FShaderCompilerEnvironment& OutEnvironment)
+	{
+		OutEnvironment.SetDefine(TEXT("MOVABLE_DIRECTIONAL_LIGHT"), TEXT("1"));
+		OutEnvironment.SetDefine(TEXT("MAX_FORWARD_SHADOWCASCADES"), TEXT(PREPROCESSOR_TO_STRING(MAX_FORWARD_SHADOWCASCADES)));
+
+		Super::ModifyCompilationEnvironment(Platform, Material, OutEnvironment);
+	}
+};
+
+class FMovableDirectionalLightCSMWithLightmapLightingPolicy : public FMovableDirectionalLightWithLightmapLightingPolicy
+{
+public:
+	static void ModifyCompilationEnvironment(EShaderPlatform Platform, const FMaterial* Material, FShaderCompilerEnvironment& OutEnvironment)
+	{
+		OutEnvironment.SetDefine(TEXT("MOVABLE_DIRECTIONAL_LIGHT_CSM"), TEXT("1"));
+
+		FMovableDirectionalLightWithLightmapLightingPolicy::ModifyCompilationEnvironment(Platform, Material, OutEnvironment);
+	}
+};
+
 #endif // __LIGHTMAPRENDERING_H__
diff --git a/Engine/Source/Runtime/Renderer/Private/RendererScene.cpp b/Engine/Source/Runtime/Renderer/Private/RendererScene.cpp
--- a/Engine/Source/Runtime/Renderer/Private/RendererScene.cpp
+++ b/Engine/Source/Runtime/Renderer/Private/RendererScene.cpp
@@ -1748,6 +1748,8 @@ void FScene::UpdateStaticDrawListsForMaterials_RenderThread(FRHICommandListImmed
 		BasePassForForwardShadingDirectionalLightAndSHIndirectDrawList[DrawType].GetUsedPrimitivesBasedOnMaterials(FeatureLevel, Materials, PrimitivesToUpdate);
 		BasePassForForwardShadingMovableDirectionalLightDrawList[DrawType].GetUsedPrimitivesBasedOnMaterials(FeatureLevel, Materials, PrimitivesToUpdate);
 		BasePassForForwardShadingMovableDirectionalLightCSMDrawList[DrawType].GetUsedPrimitivesBasedOnMaterials(FeatureLevel, Materials, PrimitivesToUpdate);
+		BasePassForForwardShadingMovableDirectionalLightLightmapDrawList[DrawType].GetUsedPrimitivesBasedOnMaterials(FeatureLevel, Materials, PrimitivesToUpdate);
+		BasePassForForwardShadingMovableDirectionalLightCSMLightmapDrawList[DrawType].GetUsedPrimitivesBasedOnMaterials(FeatureLevel, Materials, PrimitivesToUpdate);
 	}
 
 	PositionOnlyDepthDrawList.GetUsedPrimitivesBasedOnMaterials(FeatureLevel, Materials, PrimitivesToUpdate);
@@ -2130,6 +2132,8 @@ void FScene::ApplyWorldOffset_RenderThread(FVector InOffset)
 	StaticMeshDrawListApplyWorldOffset(BasePassForForwardShadingDirectionalLightAndSHIndirectDrawList, InOffset);
 	StaticMeshDrawListApplyWorldOffset(BasePassForForwardShadingMovableDirectionalLightDrawList, InOffset);
 	StaticMeshDrawListApplyWorldOffset(BasePassForForwardShadingMovableDirectionalLightCSMDrawList, InOffset);
+	StaticMeshDrawListApplyWorldOffset(BasePassForForwardShadingMovableDirectionalLightLightmapDrawList, InOffset);
+	StaticMeshDrawListApplyWorldOffset(BasePassForForwardShadingMovableDirectionalLightCSMLightmapDrawList, InOffset);
 
 	// Motion blur 
 	MotionBlurInfoData.ApplyOffset(InOffset);
@@ -2416,6 +2420,18 @@ TStaticMeshDrawList<TBasePassForForwardShadingDrawingPolicy<FMovableDirectionalL
 	return BasePassForForwardShadingMovableDirectionalLightCSMDrawList[DrawType];
 }
 
+template<>
+TStaticMeshDrawList<TBasePassForForwardShadingDrawingPolicy<FMovableDirectionalLightWithLightmapLightingPolicy> >& FScene::GetForwardShadingBasePassDrawList<FMovableDirectionalLightWithLightmapLightingPolicy>(EBasePassDrawListType DrawType)
+{
+	return BasePassForForwardShadingMovableDirectionalLightLightmapDrawList[DrawType];
+}
+
+template<>
+TStaticMeshDrawList<TBasePassForForwardShadingDrawingPolicy<FMovableDirectionalLightCSMWithLightmapLightingPolicy> >& FScene::GetForwardShadingBasePassDrawList<FMovableDirectionalLightCSMWithLightmapLightingPolicy>(EBasePassDrawListType DrawType)
+{
+	return BasePassForForwardShadingMovableDirectionalLightCSMLightmapDrawList[DrawType];
+}
+
 /*-----------------------------------------------------------------------------
 	MotionBlurInfoData
 -----------------------------------------------------------------------------*/
diff --git a/Engine/Source/Runtime/Renderer/Private/ScenePrivate.h b/Engine/Source/Runtime/Renderer/Private/ScenePrivate.h
--- a/Engine/Source/Runtime/Renderer/Private/ScenePrivate.h
+++ b/Engine/Source/Runtime/Renderer/Private/ScenePrivate.h
@@ -1246,6 +1246,8 @@ public:
 	TStaticMeshDrawList<TBasePassForForwardShadingDrawingPolicy< FSimpleDirectionalLightAndSHIndirectPolicy > >				BasePassForForwardShadingDirectionalLightAndSHIndirectDrawList[EBasePass_MAX];
 	TStaticMeshDrawList<TBasePassForForwardShadingDrawingPolicy< FMovableDirectionalLightLightingPolicy > >					BasePassForForwardShadingMovableDirectionalLightDrawList[EBasePass_MAX];
 	TStaticMeshDrawList<TBasePassForForwardShadingDrawingPolicy< FMovableDirectionalLightCSMLightingPolicy > >				BasePassForForwardShadingMovableDirectionalLightCSMDrawList[EBasePass_MAX];
+	TStaticMeshDrawList<TBasePassForForwardShadingDrawingPolicy< FMovableDirectionalLightWithLightmapLightingPolicy> >		BasePassForForwardShadingMovableDirectionalLightLightmapDrawList[EBasePass_MAX];
+	TStaticMeshDrawList<TBasePassForForwardShadingDrawingPolicy< FMovableDirectionalLightCSMWithLightmapLightingPolicy> >	BasePassForForwardShadingMovableDirectionalLightCSMLightmapDrawList[EBasePass_MAX];
 
 	/** Maps a light-map type to the appropriate base pass draw list. */
 	template<typename LightMapPolicyType>
