Fixed DFAO crash in ortho viewports

diff --git a/Engine/Shaders/DistanceFieldSurfaceCacheLightingCompute.usf b/Engine/Shaders/DistanceFieldSurfaceCacheLightingCompute.usf
--- a/Engine/Shaders/DistanceFieldSurfaceCacheLightingCompute.usf
+++ b/Engine/Shaders/DistanceFieldSurfaceCacheLightingCompute.usf
@@ -141,13 +141,13 @@ void RemoveObjectsFromBufferCS(
 RWBuffer<float4> RWCulledObjectBounds;
 RWBuffer<float4> RWCulledObjectData;
 RWBuffer<float4> RWCulledObjectBoxBounds;
-float4 ViewFrustumConvexHull[5];
+uint NumConvexHullPlanes;
+float4 ViewFrustumConvexHull[6];
 uint ObjectBoundingGeometryIndexCount;
 
 bool ViewFrustumIntersectSphere(float3 SphereOrigin, float SphereRadius)
 {
-	UNROLL
-	for (uint PlaneIndex = 0; PlaneIndex < 5; PlaneIndex++)
+	for (uint PlaneIndex = 0; PlaneIndex < NumConvexHullPlanes; PlaneIndex++)
 	{
 		float4 PlaneData = ViewFrustumConvexHull[PlaneIndex];
 		float PlaneDistance = dot(PlaneData.xyz, SphereOrigin) - PlaneData.w;
diff --git a/Engine/Source/Runtime/Renderer/Private/DistanceFieldSurfaceCacheLighting.cpp b/Engine/Source/Runtime/Renderer/Private/DistanceFieldSurfaceCacheLighting.cpp
--- a/Engine/Source/Runtime/Renderer/Private/DistanceFieldSurfaceCacheLighting.cpp
+++ b/Engine/Source/Runtime/Renderer/Private/DistanceFieldSurfaceCacheLighting.cpp
@@ -702,6 +702,7 @@ public:
 		CulledObjectData.Bind(Initializer.ParameterMap, TEXT("CulledObjectData"));
 		CulledObjectBoxBounds.Bind(Initializer.ParameterMap, TEXT("CulledObjectBoxBounds"));
 		AOParameters.Bind(Initializer.ParameterMap);
+		NumConvexHullPlanes.Bind(Initializer.ParameterMap, TEXT("NumConvexHullPlanes"));
 		ViewFrustumConvexHull.Bind(Initializer.ParameterMap, TEXT("ViewFrustumConvexHull"));
 		ObjectBoundingGeometryIndexCount.Bind(Initializer.ParameterMap, TEXT("ObjectBoundingGeometryIndexCount"));
 	}
@@ -722,8 +723,9 @@ public:
 		CulledObjectBoxBounds.SetBuffer(RHICmdList, ShaderRHI, GAOCulledObjectBuffers.Buffers.BoxBounds);
 		AOParameters.Set(RHICmdList, ShaderRHI, Parameters);
 
-		// Shader assumes 5
-		check(View.ViewFrustum.Planes.Num() == 5);
+		// Shader assumes max 6
+		check(View.ViewFrustum.Planes.Num() < 6);
+		SetShaderValue(RHICmdList, ShaderRHI, NumConvexHullPlanes, View.ViewFrustum.Planes.Num());
 		SetShaderValueArray(RHICmdList, ShaderRHI, ViewFrustumConvexHull, View.ViewFrustum.Planes.GetData(), View.ViewFrustum.Planes.Num());
 		SetShaderValue(RHICmdList, ShaderRHI, ObjectBoundingGeometryIndexCount, StencilingGeometry::GLowPolyStencilSphereIndexBuffer.GetIndexCount());
 	}
@@ -746,6 +748,7 @@ public:
 		Ar << CulledObjectData;
 		Ar << CulledObjectBoxBounds;
 		Ar << AOParameters;
+		Ar << NumConvexHullPlanes;
 		Ar << ViewFrustumConvexHull;
 		Ar << ObjectBoundingGeometryIndexCount;
 		return bShaderHasOutdatedParameters;
@@ -759,6 +762,7 @@ private:
 	FRWShaderParameter CulledObjectData;
 	FRWShaderParameter CulledObjectBoxBounds;
 	FAOParameters AOParameters;
+	FShaderParameter NumConvexHullPlanes;
 	FShaderParameter ViewFrustumConvexHull;
 	FShaderParameter ObjectBoundingGeometryIndexCount;
 };
@@ -4268,7 +4272,8 @@ bool FDeferredShadingSceneRenderer::RenderDistanceFieldAOSurfaceCache(
 	if (SupportsDistanceFieldAO(View.GetFeatureLevel(), View.GetShaderPlatform())
 		&& Views.Num() == 1
 		// ViewState is used to cache tile intersection resources which have to be sized based on the view
-		&& View.State)
+		&& View.State
+		&& View.IsPerspectiveProjection())
 	{
 		QUICK_SCOPE_CYCLE_COUNTER(STAT_RenderDistanceFieldAOSurfaceCache);
 		SCOPED_DRAW_EVENT(RHICmdList, DistanceFieldAO);
