Don't update global object buffers when no distance field texture

diff --git a/Engine/Source/Runtime/Renderer/Private/DistanceFieldSurfaceCacheLighting.cpp b/Engine/Source/Runtime/Renderer/Private/DistanceFieldSurfaceCacheLighting.cpp
--- a/Engine/Source/Runtime/Renderer/Private/DistanceFieldSurfaceCacheLighting.cpp
+++ b/Engine/Source/Runtime/Renderer/Private/DistanceFieldSurfaceCacheLighting.cpp
@@ -774,10 +774,11 @@ void FDeferredShadingSceneRenderer::UpdateGlobalDistanceFieldObjectBuffers(FRHIC
 {
 	FDistanceFieldSceneData& DistanceFieldSceneData = Scene->DistanceFieldSceneData;
 
-	if (DistanceFieldSceneData.PendingAddOperations.Num() > 0 
-		|| DistanceFieldSceneData.PendingUpdateOperations.Num() > 0 
-		|| DistanceFieldSceneData.PendingRemoveOperations.Num() > 0
-		|| bReallocatedAtlasLayouts)
+	if (GDistanceFieldVolumeTextureAtlas.VolumeTextureRHI
+		&& (DistanceFieldSceneData.PendingAddOperations.Num() > 0 
+			|| DistanceFieldSceneData.PendingUpdateOperations.Num() > 0 
+			|| DistanceFieldSceneData.PendingRemoveOperations.Num() > 0
+			|| bReallocatedAtlasLayouts))
 	{
 		QUICK_SCOPE_CYCLE_COUNTER(STAT_UpdateObjectData);
 		SCOPED_DRAW_EVENT(RHICmdList, UpdateSceneObjectData);
