Added r.SkySpecularOcclusionStrength to tweak DFAO influence on sky specular

diff --git a/Engine/Shaders/SkyLightingShared.usf b/Engine/Shaders/SkyLightingShared.usf
--- a/Engine/Shaders/SkyLightingShared.usf
+++ b/Engine/Shaders/SkyLightingShared.usf
@@ -20,6 +20,7 @@ Texture2D BentNormalAOTexture;
 SamplerState BentNormalAOSampler;
 
 float ApplyBentNormalAO;
+float InvSkySpecularOcclusionStrength;
 
 float GetDistanceFieldAOSpecularOcclusion(float2 ScreenUV, float3 ReflectionVector, float Roughness, bool bTwoSidedFoliage)
 {
@@ -29,17 +30,17 @@ float GetDistanceFieldAOSpecularOcclusion(float2 ScreenUV, float3 ReflectionVect
 	if (ApplyBentNormalAO > 0)
 	{
 		float3 BentNormalAO = Texture2DSampleLevel(BentNormalAOTexture, BentNormalAOSampler, ScreenUV, 0).xyz;
+		float BentNormalLength = length(BentNormalAO);
 
 		BRANCH
 		if (bTwoSidedFoliage)
 		{
-			Visibility = length(BentNormalAO);
+			Visibility = BentNormalLength;
 		}
 		else
 		{
 			float ReflectionConeAngle = Roughness * PI;
-			float BentNormalLength = length(BentNormalAO);
-			float UnoccludedAngle = BentNormalLength * PI;
+			float UnoccludedAngle = BentNormalLength * PI * InvSkySpecularOcclusionStrength;
 			float AngleBetween = acos(dot(BentNormalAO, ReflectionVector) / max(BentNormalLength, .001f));
 			Visibility = ApproximateConeConeIntersection(ReflectionConeAngle, UnoccludedAngle, AngleBetween);
 
diff --git a/Engine/Source/Runtime/Renderer/Private/ReflectionEnvironment.cpp b/Engine/Source/Runtime/Renderer/Private/ReflectionEnvironment.cpp
--- a/Engine/Source/Runtime/Renderer/Private/ReflectionEnvironment.cpp
+++ b/Engine/Source/Runtime/Renderer/Private/ReflectionEnvironment.cpp
@@ -46,6 +46,12 @@ static TAutoConsoleVariable<int32> CVarHalfResReflections(
 	TEXT(" 0 is off (default), 1 is on"),
 	ECVF_RenderThreadSafe);
 
+static TAutoConsoleVariable<float> CVarSkySpecularOcclusionStrength(
+	TEXT("r.SkySpecularOcclusionStrength"),
+	1,
+	TEXT("Strength of skylight specular occlusion from DFAO"),
+	ECVF_RenderThreadSafe);
+
 // to avoid having direct access from many places
 static int GetReflectionEnvironmentCVar()
 {
@@ -124,10 +130,11 @@ public:
 		BentNormalAOTexture.Bind(ParameterMap, TEXT("BentNormalAOTexture"));
 		BentNormalAOSampler.Bind(ParameterMap, TEXT("BentNormalAOSampler"));
 		ApplyBentNormalAO.Bind(ParameterMap, TEXT("ApplyBentNormalAO"));
+		InvSkySpecularOcclusionStrength.Bind(ParameterMap, TEXT("InvSkySpecularOcclusionStrength"));
 	}
 
 	template<typename ShaderRHIParamRef>
-	void SetParameters(FRHICommandList& RHICmdList, const ShaderRHIParamRef& ShaderRHI, const TRefCountPtr<IPooledRenderTarget>& DynamicBentNormalAO )
+	void SetParameters(FRHICommandList& RHICmdList, const ShaderRHIParamRef& ShaderRHI, const TRefCountPtr<IPooledRenderTarget>& DynamicBentNormalAO, float SkySpecularOcclusionStrength)
 	{
 		FTextureRHIParamRef BentNormalAO = GWhiteTexture->TextureRHI;
 		bool bApplyBentNormalAO = false;
@@ -140,11 +147,12 @@ public:
 
 		SetTextureParameter(RHICmdList, ShaderRHI, BentNormalAOTexture, BentNormalAOSampler, TStaticSamplerState<SF_Point>::GetRHI(), BentNormalAO);
 		SetShaderValue(RHICmdList, ShaderRHI, ApplyBentNormalAO, bApplyBentNormalAO ? 1.0f : 0.0f);
+		SetShaderValue(RHICmdList, ShaderRHI, InvSkySpecularOcclusionStrength, 1.0f / FMath::Max(SkySpecularOcclusionStrength, .1f));
 	}
 
 	friend FArchive& operator<<(FArchive& Ar,FDistanceFieldAOSpecularOcclusionParameters& P)
 	{
-		Ar << P.BentNormalAOTexture << P.BentNormalAOSampler << P.ApplyBentNormalAO;
+		Ar << P.BentNormalAOTexture << P.BentNormalAOSampler << P.ApplyBentNormalAO << P.InvSkySpecularOcclusionStrength;
 		return Ar;
 	}
 
@@ -152,6 +160,7 @@ private:
 	FShaderResourceParameter BentNormalAOTexture;
 	FShaderResourceParameter BentNormalAOSampler;
 	FShaderParameter ApplyBentNormalAO;
+	FShaderParameter InvSkySpecularOcclusionStrength;
 };
 
 struct FReflectionCaptureSortData
@@ -324,7 +333,7 @@ public:
 		SetTextureParameter(RHICmdList, ShaderRHI, PreIntegratedGF, PreIntegratedGFSampler, TStaticSamplerState<SF_Bilinear,AM_Clamp,AM_Clamp,AM_Clamp>::GetRHI(), GSystemTextures.PreintegratedGF->GetRenderTargetItem().ShaderResourceTexture);
 	
 		SkyLightParameters.SetParameters(RHICmdList, ShaderRHI, Scene, View.Family->EngineShowFlags.SkyLighting);
-		SpecularOcclusionParameters.SetParameters(RHICmdList, ShaderRHI, DynamicBentNormalAO);
+		SpecularOcclusionParameters.SetParameters(RHICmdList, ShaderRHI, DynamicBentNormalAO, CVarSkySpecularOcclusionStrength.GetValueOnRenderThread());
 	}
 
 	void UnsetParameters(FRHICommandList& RHICmdList)
@@ -446,7 +455,7 @@ public:
 		SetTextureParameter(RHICmdList, ShaderRHI, ScreenSpaceReflectionsTexture, ScreenSpaceReflectionsSampler, TStaticSamplerState<SF_Point>::GetRHI(), ScreenSpaceReflections );
 		SetTextureParameter(RHICmdList, ShaderRHI, PreIntegratedGF, PreIntegratedGFSampler, TStaticSamplerState<SF_Bilinear,AM_Clamp,AM_Clamp,AM_Clamp>::GetRHI(), GSystemTextures.PreintegratedGF->GetRenderTargetItem().ShaderResourceTexture );
 		
-		SpecularOcclusionParameters.SetParameters(RHICmdList, ShaderRHI, DynamicBentNormalAO);
+		SpecularOcclusionParameters.SetParameters(RHICmdList, ShaderRHI, DynamicBentNormalAO, CVarSkySpecularOcclusionStrength.GetValueOnRenderThread());
 	}
 
 	// FShader interface.
