UE4 - Rework forward draw policy selection code; add missing policies for indirect + csm

diff --git a/Engine/Source/Runtime/Renderer/Private/ForwardBasePassRendering.cpp b/Engine/Source/Runtime/Renderer/Private/ForwardBasePassRendering.cpp
--- a/Engine/Source/Runtime/Renderer/Private/ForwardBasePassRendering.cpp
+++ b/Engine/Source/Runtime/Renderer/Private/ForwardBasePassRendering.cpp
@@ -31,6 +31,8 @@ IMPLEMENT_FORWARD_SHADING_BASEPASS_LIGHTMAPPED_SHADER_TYPE( FNoLightMapPolicy, F
 IMPLEMENT_FORWARD_SHADING_BASEPASS_LIGHTMAPPED_SHADER_TYPE( TLightMapPolicy<LQ_LIGHTMAP>, TLightMapPolicyLQ );
 IMPLEMENT_FORWARD_SHADING_BASEPASS_LIGHTMAPPED_SHADER_TYPE( TDistanceFieldShadowsAndLightMapPolicy<LQ_LIGHTMAP>, TDistanceFieldShadowsAndLightMapPolicyLQ );
 IMPLEMENT_FORWARD_SHADING_BASEPASS_LIGHTMAPPED_SHADER_TYPE( FSimpleDirectionalLightAndSHIndirectPolicy, FSimpleDirectionalLightAndSHIndirectPolicy );
+IMPLEMENT_FORWARD_SHADING_BASEPASS_LIGHTMAPPED_SHADER_TYPE(FSimpleDirectionalLightAndSHDirectionalIndirectPolicy, FSimpleDirectionalLightAndSHDirectionalIndirectPolicy);
+IMPLEMENT_FORWARD_SHADING_BASEPASS_LIGHTMAPPED_SHADER_TYPE(FSimpleDirectionalLightAndSHDirectionalCSMIndirectPolicy, FSimpleDirectionalLightAndSHDirectionalCSMIndirectPolicy);
 IMPLEMENT_FORWARD_SHADING_BASEPASS_LIGHTMAPPED_SHADER_TYPE( FMovableDirectionalLightLightingPolicy, FMovableDirectionalLightLightingPolicy );
 IMPLEMENT_FORWARD_SHADING_BASEPASS_LIGHTMAPPED_SHADER_TYPE( FMovableDirectionalLightCSMLightingPolicy, FMovableDirectionalLightCSMLightingPolicy );
 IMPLEMENT_FORWARD_SHADING_BASEPASS_LIGHTMAPPED_SHADER_TYPE(FMovableDirectionalLightWithLightmapLightingPolicy, FMovableDirectionalLightWithLightmapLightingPolicy);
diff --git a/Engine/Source/Runtime/Renderer/Private/ForwardBasePassRendering.h b/Engine/Source/Runtime/Renderer/Private/ForwardBasePassRendering.h
--- a/Engine/Source/Runtime/Renderer/Private/ForwardBasePassRendering.h
+++ b/Engine/Source/Runtime/Renderer/Private/ForwardBasePassRendering.h
@@ -514,73 +514,94 @@ void ProcessBasePassMeshForForwardShading(
 {
 	// Check for a cached light-map.
 	const bool bIsLitMaterial = Parameters.ShadingModel != MSM_Unlit;
-
-	const FLightMapInteraction LightMapInteraction = (Parameters.Mesh.LCI && bIsLitMaterial) 
-		? Parameters.Mesh.LCI->GetLightMapInteraction(Parameters.FeatureLevel) 
-		: FLightMapInteraction();
-
 	if (bIsLitMaterial)
 	{
+		const FLightMapInteraction LightMapInteraction = (Parameters.Mesh.LCI && bIsLitMaterial) 
+			? Parameters.Mesh.LCI->GetLightMapInteraction(Parameters.FeatureLevel) 
+			: FLightMapInteraction();
+
 		const FLightSceneInfo* SimpleDirectionalLight = Action.GetSimpleDirectionalLight();
 		const bool bUseMovableLight = SimpleDirectionalLight && !SimpleDirectionalLight->Proxy->HasStaticShadowing();
 
-		if (bUseMovableLight && LightMapInteraction.GetType() == LMIT_Texture)
+		if (LightMapInteraction.GetType() == LMIT_Texture)
 		{
-			// final determination of whether CSMs are rendered can be view dependent, thus we always need to clear the CSMs even if we're not going to render to them based on the condition below.
-			if (SimpleDirectionalLight && SimpleDirectionalLight->ShouldRenderViewIndependentWholeSceneShadows())
+			// Lightmap path
+			if (bUseMovableLight)
 			{
-				Action.template Process<FMovableDirectionalLightCSMWithLightmapLightingPolicy>(RHICmdList, Parameters, FMovableDirectionalLightCSMWithLightmapLightingPolicy(), LightMapInteraction);
+				// final determination of whether CSMs are rendered can be view dependent, thus we always need to clear the CSMs even if we're not going to render to them based on the condition below.
+				if (SimpleDirectionalLight && SimpleDirectionalLight->ShouldRenderViewIndependentWholeSceneShadows())
+				{
+					Action.template Process<FMovableDirectionalLightCSMWithLightmapLightingPolicy>(RHICmdList, Parameters, FMovableDirectionalLightCSMWithLightmapLightingPolicy(), LightMapInteraction);
+				}
+				else
+				{
+					Action.template Process<FMovableDirectionalLightWithLightmapLightingPolicy>(RHICmdList, Parameters, FMovableDirectionalLightCSMWithLightmapLightingPolicy(), LightMapInteraction);
+				}
 			}
 			else
 			{
-				Action.template Process<FMovableDirectionalLightWithLightmapLightingPolicy>(RHICmdList, Parameters, FMovableDirectionalLightCSMWithLightmapLightingPolicy(), LightMapInteraction);
+				const FShadowMapInteraction ShadowMapInteraction = (Parameters.Mesh.LCI && bIsLitMaterial)
+					? Parameters.Mesh.LCI->GetShadowMapInteraction()
+					: FShadowMapInteraction();
+
+				if (ShadowMapInteraction.GetType() == SMIT_Texture)
+				{
+					Action.template Process< TDistanceFieldShadowsAndLightMapPolicy<LQ_LIGHTMAP> >(
+						RHICmdList,
+						Parameters,
+						TDistanceFieldShadowsAndLightMapPolicy<LQ_LIGHTMAP>(),
+						TDistanceFieldShadowsAndLightMapPolicy<LQ_LIGHTMAP>::ElementDataType(ShadowMapInteraction, LightMapInteraction));
+				}
+				else
+				{
+					Action.template Process< TLightMapPolicy<LQ_LIGHTMAP> >(RHICmdList, Parameters, TLightMapPolicy<LQ_LIGHTMAP>(), LightMapInteraction);
+				}
 			}
+
+			// Exit to avoid NoLightmapPolicy
+			return;
 		}
-		else if (bUseMovableLight)
+		else if (IsIndirectLightingCacheAllowed(Parameters.FeatureLevel)
+			&& Parameters.PrimitiveSceneProxy
+			// Movable objects need to get their GI from the indirect lighting cache
+			&& Parameters.PrimitiveSceneProxy->IsMovable())
 		{
-			// final determination of whether CSMs are rendered can be view dependent, thus we always need to clear the CSMs even if we're not going to render to them based on the condition below.
-			if (SimpleDirectionalLight && SimpleDirectionalLight->ShouldRenderViewIndependentWholeSceneShadows())
+			if (bUseMovableLight)
 			{
-				Action.template Process<FMovableDirectionalLightCSMLightingPolicy>(RHICmdList, Parameters, FMovableDirectionalLightCSMLightingPolicy(), FMovableDirectionalLightCSMLightingPolicy::ElementDataType());
+				if (SimpleDirectionalLight && SimpleDirectionalLight->ShouldRenderViewIndependentWholeSceneShadows())
+				{
+					Action.template Process<FSimpleDirectionalLightAndSHDirectionalCSMIndirectPolicy>(RHICmdList, Parameters, FSimpleDirectionalLightAndSHDirectionalCSMIndirectPolicy(), FSimpleDirectionalLightAndSHDirectionalCSMIndirectPolicy::ElementDataType(Action.ShouldPackAmbientSH()));
+				}
+				else
+				{
+					Action.template Process<FSimpleDirectionalLightAndSHDirectionalIndirectPolicy>(RHICmdList, Parameters, FSimpleDirectionalLightAndSHDirectionalIndirectPolicy(), FSimpleDirectionalLightAndSHDirectionalIndirectPolicy::ElementDataType(Action.ShouldPackAmbientSH()));
+				}
 			}
 			else
 			{
-				Action.template Process<FMovableDirectionalLightLightingPolicy>(RHICmdList, Parameters, FMovableDirectionalLightLightingPolicy(), FMovableDirectionalLightLightingPolicy::ElementDataType());
-			}		
+				Action.template Process<FSimpleDirectionalLightAndSHIndirectPolicy>(RHICmdList, Parameters, FSimpleDirectionalLightAndSHIndirectPolicy(), FSimpleDirectionalLightAndSHIndirectPolicy::ElementDataType(Action.ShouldPackAmbientSH()));
+			}
+
+			// Exit to avoid NoLightmapPolicy
+			return;
 		}
-		else if (LightMapInteraction.GetType() == LMIT_Texture)
+		else if (bUseMovableLight)
 		{
-			const FShadowMapInteraction ShadowMapInteraction = (Parameters.Mesh.LCI && bIsLitMaterial) 
-				? Parameters.Mesh.LCI->GetShadowMapInteraction() 
-				: FShadowMapInteraction();
-
-			if (ShadowMapInteraction.GetType() == SMIT_Texture)
+			// final determination of whether CSMs are rendered can be view dependent, thus we always need to clear the CSMs even if we're not going to render to them based on the condition below.
+			if (SimpleDirectionalLight && SimpleDirectionalLight->ShouldRenderViewIndependentWholeSceneShadows())
 			{
-				Action.template Process< TDistanceFieldShadowsAndLightMapPolicy<LQ_LIGHTMAP> >(
-					RHICmdList, 
-					Parameters,
-					TDistanceFieldShadowsAndLightMapPolicy<LQ_LIGHTMAP>(),
-					TDistanceFieldShadowsAndLightMapPolicy<LQ_LIGHTMAP>::ElementDataType(ShadowMapInteraction, LightMapInteraction));
+				Action.template Process<FMovableDirectionalLightCSMLightingPolicy>(RHICmdList, Parameters, FMovableDirectionalLightCSMLightingPolicy(), FMovableDirectionalLightCSMLightingPolicy::ElementDataType());
 			}
 			else
 			{
-				Action.template Process< TLightMapPolicy<LQ_LIGHTMAP> >(RHICmdList, Parameters, TLightMapPolicy<LQ_LIGHTMAP>(), LightMapInteraction );
+				Action.template Process<FMovableDirectionalLightLightingPolicy>(RHICmdList, Parameters, FMovableDirectionalLightLightingPolicy(), FMovableDirectionalLightLightingPolicy::ElementDataType());
 			}
+
+			// Exit to avoid NoLightmapPolicy
+			return;
 		}
-		else if (IsIndirectLightingCacheAllowed(Parameters.FeatureLevel)
-			&& Parameters.PrimitiveSceneProxy
-			// Movable objects need to get their GI from the indirect lighting cache
-			&& Parameters.PrimitiveSceneProxy->IsMovable())
-		{
-			Action.template Process<FSimpleDirectionalLightAndSHIndirectPolicy>(RHICmdList, Parameters,FSimpleDirectionalLightAndSHIndirectPolicy(),FSimpleDirectionalLightAndSHIndirectPolicy::ElementDataType(Action.ShouldPackAmbientSH())); 
-		}
-		else
-		{
-			Action.template Process<FNoLightMapPolicy>(RHICmdList, Parameters,FNoLightMapPolicy(),FNoLightMapPolicy::ElementDataType());
-		}
-	}
-	else
-	{
-		Action.template Process<FNoLightMapPolicy>(RHICmdList, Parameters,FNoLightMapPolicy(),FNoLightMapPolicy::ElementDataType());
 	}
+
+	// Default to NoLightmapPolicy
+	Action.template Process<FNoLightMapPolicy>(RHICmdList, Parameters, FNoLightMapPolicy(), FNoLightMapPolicy::ElementDataType());
 }
diff --git a/Engine/Source/Runtime/Renderer/Private/LightMapRendering.h b/Engine/Source/Runtime/Renderer/Private/LightMapRendering.h
--- a/Engine/Source/Runtime/Renderer/Private/LightMapRendering.h
+++ b/Engine/Source/Runtime/Renderer/Private/LightMapRendering.h
@@ -969,6 +969,49 @@ private:
 	FCachedPointIndirectLightingPolicy CachedPointIndirectLightingPolicy;
 };
 
+/** Combines a directional light with indirect lighting from a single SH sample. */
+class FSimpleDirectionalLightAndSHDirectionalIndirectPolicy : public FSimpleDirectionalLightAndSHIndirectPolicy
+{
+public:
+	static void ModifyCompilationEnvironment(EShaderPlatform Platform, const FMaterial* Material, FShaderCompilerEnvironment& OutEnvironment)
+	{
+		OutEnvironment.SetDefine(TEXT("MOVABLE_DIRECTIONAL_LIGHT"), TEXT("1"));
+		OutEnvironment.SetDefine(TEXT("MAX_FORWARD_SHADOWCASCADES"), TEXT(PREPROCESSOR_TO_STRING(MAX_FORWARD_SHADOWCASCADES)));
+		FSimpleDirectionalLightAndSHIndirectPolicy::ModifyCompilationEnvironment(Platform, Material, OutEnvironment);
+	}
+
+	friend bool operator==(const FSimpleDirectionalLightAndSHDirectionalIndirectPolicy A, const FSimpleDirectionalLightAndSHDirectionalIndirectPolicy B)
+	{
+		return true;
+	}
+
+	friend int32 CompareDrawingPolicy(const FSimpleDirectionalLightAndSHDirectionalIndirectPolicy&, const FSimpleDirectionalLightAndSHDirectionalIndirectPolicy&)
+	{
+		return 0;
+	}
+};
+
+/** Combines a directional light with CSM with indirect lighting from a single SH sample. */
+class FSimpleDirectionalLightAndSHDirectionalCSMIndirectPolicy : public FSimpleDirectionalLightAndSHDirectionalIndirectPolicy
+{
+public:
+	static void ModifyCompilationEnvironment(EShaderPlatform Platform, const FMaterial* Material, FShaderCompilerEnvironment& OutEnvironment)
+	{
+		OutEnvironment.SetDefine(TEXT("MOVABLE_DIRECTIONAL_LIGHT_CSM"), TEXT("1"));
+		FSimpleDirectionalLightAndSHDirectionalIndirectPolicy::ModifyCompilationEnvironment(Platform, Material, OutEnvironment);
+	}
+
+	friend bool operator==(const FSimpleDirectionalLightAndSHDirectionalCSMIndirectPolicy A, const FSimpleDirectionalLightAndSHDirectionalCSMIndirectPolicy B)
+	{
+		return true;
+	}
+
+	friend int32 CompareDrawingPolicy(const FSimpleDirectionalLightAndSHDirectionalCSMIndirectPolicy&, const FSimpleDirectionalLightAndSHDirectionalCSMIndirectPolicy&)
+	{
+		return 0;
+	}
+};
+
 class FMovableDirectionalLightLightingPolicy : public FNoLightMapPolicy
 {
 public:
diff --git a/Engine/Source/Runtime/Renderer/Private/RendererScene.cpp b/Engine/Source/Runtime/Renderer/Private/RendererScene.cpp
--- a/Engine/Source/Runtime/Renderer/Private/RendererScene.cpp
+++ b/Engine/Source/Runtime/Renderer/Private/RendererScene.cpp
@@ -2408,6 +2408,18 @@ TStaticMeshDrawList<TBasePassForForwardShadingDrawingPolicy<FSimpleDirectionalLi
 	return BasePassForForwardShadingDirectionalLightAndSHIndirectDrawList[DrawType];
 }
 
+template<>
+TStaticMeshDrawList<TBasePassForForwardShadingDrawingPolicy<FSimpleDirectionalLightAndSHDirectionalIndirectPolicy> >& FScene::GetForwardShadingBasePassDrawList<FSimpleDirectionalLightAndSHDirectionalIndirectPolicy>(EBasePassDrawListType DrawType)
+{
+	return BasePassForForwardShadingDirectionalLightAndSHDirectionalIndirectDrawList[DrawType];
+}
+
+template<>
+TStaticMeshDrawList<TBasePassForForwardShadingDrawingPolicy<FSimpleDirectionalLightAndSHDirectionalCSMIndirectPolicy> >& FScene::GetForwardShadingBasePassDrawList<FSimpleDirectionalLightAndSHDirectionalCSMIndirectPolicy>(EBasePassDrawListType DrawType)
+{
+	return BasePassForForwardShadingDirectionalLightAndSHDirectionalCSMIndirectDrawList[DrawType];
+}
+
 template<>
 TStaticMeshDrawList<TBasePassForForwardShadingDrawingPolicy<FMovableDirectionalLightLightingPolicy> >& FScene::GetForwardShadingBasePassDrawList<FMovableDirectionalLightLightingPolicy>(EBasePassDrawListType DrawType)
 {
diff --git a/Engine/Source/Runtime/Renderer/Private/ScenePrivate.h b/Engine/Source/Runtime/Renderer/Private/ScenePrivate.h
--- a/Engine/Source/Runtime/Renderer/Private/ScenePrivate.h
+++ b/Engine/Source/Runtime/Renderer/Private/ScenePrivate.h
@@ -1244,6 +1244,8 @@ public:
 	TStaticMeshDrawList< TBasePassForForwardShadingDrawingPolicy< TLightMapPolicy<LQ_LIGHTMAP> > >							BasePassForForwardShadingLowQualityLightMapDrawList[EBasePass_MAX];
 	TStaticMeshDrawList< TBasePassForForwardShadingDrawingPolicy< TDistanceFieldShadowsAndLightMapPolicy<LQ_LIGHTMAP> > >	BasePassForForwardShadingDistanceFieldShadowMapLightMapDrawList[EBasePass_MAX];
 	TStaticMeshDrawList<TBasePassForForwardShadingDrawingPolicy< FSimpleDirectionalLightAndSHIndirectPolicy > >				BasePassForForwardShadingDirectionalLightAndSHIndirectDrawList[EBasePass_MAX];
+	TStaticMeshDrawList<TBasePassForForwardShadingDrawingPolicy< FSimpleDirectionalLightAndSHDirectionalIndirectPolicy > >		BasePassForForwardShadingDirectionalLightAndSHDirectionalIndirectDrawList[EBasePass_MAX];
+	TStaticMeshDrawList<TBasePassForForwardShadingDrawingPolicy< FSimpleDirectionalLightAndSHDirectionalCSMIndirectPolicy > >	BasePassForForwardShadingDirectionalLightAndSHDirectionalCSMIndirectDrawList[EBasePass_MAX];
 	TStaticMeshDrawList<TBasePassForForwardShadingDrawingPolicy< FMovableDirectionalLightLightingPolicy > >					BasePassForForwardShadingMovableDirectionalLightDrawList[EBasePass_MAX];
 	TStaticMeshDrawList<TBasePassForForwardShadingDrawingPolicy< FMovableDirectionalLightCSMLightingPolicy > >				BasePassForForwardShadingMovableDirectionalLightCSMDrawList[EBasePass_MAX];
 	TStaticMeshDrawList<TBasePassForForwardShadingDrawingPolicy< FMovableDirectionalLightWithLightmapLightingPolicy> >		BasePassForForwardShadingMovableDirectionalLightLightmapDrawList[EBasePass_MAX];
