Removed DrawDynamicElements, PreRenderView and dependencies.  These have been disabled since 11/04, and the GetDynamicMeshElements path is now the only path for dynamic mesh rendering.

diff --git a/Engine/Source/Runtime/Landscape/Private/LandscapeSplines.cpp b/Engine/Source/Runtime/Landscape/Private/LandscapeSplines.cpp
--- a/Engine/Source/Runtime/Landscape/Private/LandscapeSplines.cpp
+++ b/Engine/Source/Runtime/Landscape/Private/LandscapeSplines.cpp
@@ -119,7 +119,7 @@ public:
 		const FMatrix& LocalToWorld = GetLocalToWorld();
 
 		const FLinearColor SelectedSplineColor = GEngine->GetSelectedMaterialColor();
-		const FLinearColor SelectedControlPointSpriteColor = FLinearColor::White + (GEngine->GetSelectedMaterialColor() * GEngine->SelectionHighlightIntensityBillboards * 10); // copied from FSpriteSceneProxy::DrawDynamicElements()
+		const FLinearColor SelectedControlPointSpriteColor = FLinearColor::White + (GEngine->GetSelectedMaterialColor() * GEngine->SelectionHighlightIntensityBillboards * 10); 
 
 		for (int32 ViewIndex = 0; ViewIndex < Views.Num(); ViewIndex++)
 		{
@@ -270,157 +270,6 @@ public:
 		}
 	}
 
-	virtual void DrawDynamicElements(FPrimitiveDrawInterface* PDI,const FSceneView* View)
-	{
-		// Slight Depth Bias so that the splines show up when they exactly match the target surface
-		// e.g. someone playing with splines on a newly-created perfectly-flat landscape
-		static const float DepthBias = -0.0001;
-
-		const FMatrix& LocalToWorld = GetLocalToWorld();
-
-		const FLinearColor SelectedSplineColor = GEngine->GetSelectedMaterialColor();
-		const FLinearColor SelectedControlPointSpriteColor = FLinearColor::White + (GEngine->GetSelectedMaterialColor() * GEngine->SelectionHighlightIntensityBillboards * 10); // copied from FSpriteSceneProxy::DrawDynamicElements()
-
-		for (const FSegmentProxy& Segment : Segments)
-		{
-			const FLinearColor SegmentColor = Segment.bSelected ? SelectedSplineColor : SplineColor;
-
-			FLandscapeSplineInterpPoint OldPoint = Segment.Points[0];
-			OldPoint.Center       = LocalToWorld.TransformPosition(OldPoint.Center);
-			OldPoint.Left         = LocalToWorld.TransformPosition(OldPoint.Left);
-			OldPoint.Right        = LocalToWorld.TransformPosition(OldPoint.Right);
-			OldPoint.FalloffLeft  = LocalToWorld.TransformPosition(OldPoint.FalloffLeft);
-			OldPoint.FalloffRight = LocalToWorld.TransformPosition(OldPoint.FalloffRight);
-			for (int32 i = 1; i < Segment.Points.Num(); i++)
-			{
-				FLandscapeSplineInterpPoint NewPoint = Segment.Points[i];
-				NewPoint.Center       = LocalToWorld.TransformPosition(NewPoint.Center);
-				NewPoint.Left         = LocalToWorld.TransformPosition(NewPoint.Left);
-				NewPoint.Right        = LocalToWorld.TransformPosition(NewPoint.Right);
-				NewPoint.FalloffLeft  = LocalToWorld.TransformPosition(NewPoint.FalloffLeft);
-				NewPoint.FalloffRight = LocalToWorld.TransformPosition(NewPoint.FalloffRight);
-
-				// Draw lines from the last keypoint.
-				if (PDI->IsHitTesting()) PDI->SetHitProxy(Segment.HitProxy);
-
-				// center line
-				PDI->DrawLine(OldPoint.Center, NewPoint.Center, SegmentColor, GetDepthPriorityGroup(View), 0.0f, DepthBias);
-
-				// draw sides
-				PDI->DrawLine(OldPoint.Left, NewPoint.Left, SegmentColor, GetDepthPriorityGroup(View), 0.0f, DepthBias);
-				PDI->DrawLine(OldPoint.Right, NewPoint.Right, SegmentColor, GetDepthPriorityGroup(View), 0.0f, DepthBias);
-
-				if (PDI->IsHitTesting()) PDI->SetHitProxy(NULL);
-
-				// draw falloff sides
-				if (bDrawFalloff)
-				{
-					DrawDashedLine(PDI, OldPoint.FalloffLeft, NewPoint.FalloffLeft, SegmentColor, 100, GetDepthPriorityGroup(View), DepthBias);
-					DrawDashedLine(PDI, OldPoint.FalloffRight, NewPoint.FalloffRight, SegmentColor, 100, GetDepthPriorityGroup(View), DepthBias);
-				}
-
-				OldPoint = NewPoint;
-			}
-		}
-
-		for (const FControlPointProxy& ControlPoint : ControlPoints)
-		{
-			const float ControlPointSpriteScale = LocalToWorld.GetScaleVector().X * ControlPoint.SpriteScale;
-			const FVector ControlPointLocation = LocalToWorld.TransformPosition(ControlPoint.Location) + FVector(0, 0, ControlPointSpriteScale * 0.75f);
-
-			// Draw Sprite
-
-			const FLinearColor ControlPointSpriteColor = ControlPoint.bSelected ? SelectedControlPointSpriteColor : FLinearColor::White;
-
-			if (PDI->IsHitTesting()) PDI->SetHitProxy(ControlPoint.HitProxy);
-
-			PDI->DrawSprite(
-				ControlPointLocation,
-				ControlPointSpriteScale,
-				ControlPointSpriteScale,
-				ControlPointSprite->Resource,
-				ControlPointSpriteColor,
-				GetDepthPriorityGroup(View),
-				0, ControlPointSprite->Resource->GetSizeX(),
-				0, ControlPointSprite->Resource->GetSizeY(),
-				SE_BLEND_Masked);
-
-
-			// Draw Lines
-			const FLinearColor ControlPointColor = ControlPoint.bSelected ? SelectedSplineColor : SplineColor;
-
-			if (ControlPoint.Points.Num() == 1)
-			{
-				FLandscapeSplineInterpPoint NewPoint = ControlPoint.Points[0];
-				NewPoint.Center = LocalToWorld.TransformPosition(NewPoint.Center);
-				NewPoint.Left   = LocalToWorld.TransformPosition(NewPoint.Left);
-				NewPoint.Right  = LocalToWorld.TransformPosition(NewPoint.Right);
-				NewPoint.FalloffLeft  = LocalToWorld.TransformPosition(NewPoint.FalloffLeft);
-				NewPoint.FalloffRight = LocalToWorld.TransformPosition(NewPoint.FalloffRight);
-
-				// draw end for spline connection
-				PDI->DrawPoint(NewPoint.Center, ControlPointColor, 6.0f, GetDepthPriorityGroup(View));
-				PDI->DrawLine(NewPoint.Left, NewPoint.Center, ControlPointColor, GetDepthPriorityGroup(View), 0.0f, DepthBias);
-				PDI->DrawLine(NewPoint.Right, NewPoint.Center, ControlPointColor, GetDepthPriorityGroup(View), 0.0f, DepthBias);
-				if (bDrawFalloff)
-				{
-					DrawDashedLine(PDI, NewPoint.FalloffLeft, NewPoint.Left, ControlPointColor, 100, GetDepthPriorityGroup(View), DepthBias);
-					DrawDashedLine(PDI, NewPoint.FalloffRight, NewPoint.Right, ControlPointColor, 100, GetDepthPriorityGroup(View), DepthBias);
-				}
-			}
-			else if (ControlPoint.Points.Num() >= 2)
-			{
-				FLandscapeSplineInterpPoint OldPoint = ControlPoint.Points.Last();
-				//OldPoint.Left   = LocalToWorld.TransformPosition(OldPoint.Left);
-				OldPoint.Right  = LocalToWorld.TransformPosition(OldPoint.Right);
-				//OldPoint.FalloffLeft  = LocalToWorld.TransformPosition(OldPoint.FalloffLeft);
-				OldPoint.FalloffRight = LocalToWorld.TransformPosition(OldPoint.FalloffRight);
-
-				for (const FLandscapeSplineInterpPoint& Point : ControlPoint.Points)
-				{
-					FLandscapeSplineInterpPoint NewPoint = Point;
-					NewPoint.Center = LocalToWorld.TransformPosition(NewPoint.Center);
-					NewPoint.Left   = LocalToWorld.TransformPosition(NewPoint.Left);
-					NewPoint.Right  = LocalToWorld.TransformPosition(NewPoint.Right);
-					NewPoint.FalloffLeft  = LocalToWorld.TransformPosition(NewPoint.FalloffLeft);
-					NewPoint.FalloffRight = LocalToWorld.TransformPosition(NewPoint.FalloffRight);
-
-					if (PDI->IsHitTesting()) PDI->SetHitProxy(ControlPoint.HitProxy);
-
-					// center line
-					PDI->DrawLine(ControlPointLocation, NewPoint.Center, ControlPointColor, GetDepthPriorityGroup(View), 0.0f, DepthBias);
-
-					// draw sides
-					PDI->DrawLine(OldPoint.Right, NewPoint.Left, ControlPointColor, GetDepthPriorityGroup(View), 0.0f, DepthBias);
-
-					if (PDI->IsHitTesting()) PDI->SetHitProxy(NULL);
-
-					// draw falloff sides
-					if (bDrawFalloff)
-					{
-						DrawDashedLine(PDI, OldPoint.FalloffRight, NewPoint.FalloffLeft, ControlPointColor, 100, GetDepthPriorityGroup(View), DepthBias);
-					}
-
-					// draw end for spline connection
-					PDI->DrawPoint(NewPoint.Center, ControlPointColor, 6.0f, GetDepthPriorityGroup(View));
-					PDI->DrawLine(NewPoint.Left, NewPoint.Center, ControlPointColor, GetDepthPriorityGroup(View), 0.0f, DepthBias);
-					PDI->DrawLine(NewPoint.Right, NewPoint.Center, ControlPointColor, GetDepthPriorityGroup(View), 0.0f, DepthBias);
-					if (bDrawFalloff)
-					{
-						DrawDashedLine(PDI, NewPoint.FalloffLeft, NewPoint.Left, ControlPointColor, 100, GetDepthPriorityGroup(View), DepthBias);
-						DrawDashedLine(PDI, NewPoint.FalloffRight, NewPoint.Right, ControlPointColor, 100, GetDepthPriorityGroup(View), DepthBias);
-					}
-
-					//OldPoint = NewPoint;
-					OldPoint.Right = NewPoint.Right;
-					OldPoint.FalloffRight = NewPoint.FalloffRight;
-				}
-			}
-		}
-
-		if (PDI->IsHitTesting()) PDI->SetHitProxy(NULL);
-	}
-
 	virtual FPrimitiveViewRelevance GetViewRelevance(const FSceneView* View)
 	{
 		FPrimitiveViewRelevance Result;
diff --git a/Engine/Source/Runtime/Landscape/Public/LandscapeRender.h b/Engine/Source/Runtime/Landscape/Public/LandscapeRender.h
--- a/Engine/Source/Runtime/Landscape/Public/LandscapeRender.h
+++ b/Engine/Source/Runtime/Landscape/Public/LandscapeRender.h
@@ -471,14 +471,6 @@ protected:
 
 	virtual ~FLandscapeComponentSceneProxy();
 
-	// Used for DrawDynamicElements
-	FMeshBatch DynamicMesh;			// Landscape Rendering using Dynamic path
-#if WITH_EDITOR
-	FMeshBatch DynamicMeshTools;	// Tool rendering, don't support tessellation for now
-#endif
-	TArray<FLandscapeBatchElementParams> DynamicMeshBatchParamArray;
-
-
 public:
 	// constructor
 	FLandscapeComponentSceneProxy(ULandscapeComponent* InComponent, FLandscapeEditToolRenderData* InEditToolRenderData);
@@ -486,7 +478,6 @@ public:
 	// FPrimitiveSceneProxy interface.
 	virtual void DrawStaticElements(FStaticPrimitiveDrawInterface* PDI) override;
 	virtual void GetDynamicMeshElements(const TArray<const FSceneView*>& Views, const FSceneViewFamily& ViewFamily, uint32 VisibilityMap, FMeshElementCollector& Collector) const override;
-	virtual void DrawDynamicElements(FPrimitiveDrawInterface* PDI,const FSceneView* View) override;
 	virtual uint32 GetMemoryFootprint( void ) const override { return( sizeof( *this ) + GetAllocatedSize() ); }
 	virtual FPrimitiveViewRelevance GetViewRelevance(const FSceneView* View) override;
 	virtual bool CanBeOccluded() const override;
diff --git a/Engine/Source/Runtime/Launch/Private/LaunchEngineLoop.cpp b/Engine/Source/Runtime/Launch/Private/LaunchEngineLoop.cpp
--- a/Engine/Source/Runtime/Launch/Private/LaunchEngineLoop.cpp
+++ b/Engine/Source/Runtime/Launch/Private/LaunchEngineLoop.cpp
@@ -2173,9 +2173,6 @@ void FEngineLoop::Tick()
 
 	{ 
 		SCOPE_CYCLE_COUNTER( STAT_FrameTime );
-		
-		STAT(extern ENGINE_API void BeginOneFrameParticleStats());
-		STAT(BeginOneFrameParticleStats());
 
 		// Calculates average FPS/MS (outside STATS on purpose)
 		CalculateFPSTimings();
@@ -2320,11 +2317,6 @@ void FEngineLoop::Tick()
 			GEngine->TickDeferredCommands();		
 		}
 
-		{
-			STAT(extern ENGINE_API void FinishOneFrameParticleStats());
-			STAT(FinishOneFrameParticleStats());
-		} 
-
 		ENQUEUE_UNIQUE_RENDER_COMMAND(
 			EndFrame,
 		{
diff --git a/Engine/Source/Runtime/Renderer/Private/BasePassRendering.h b/Engine/Source/Runtime/Renderer/Private/BasePassRendering.h
--- a/Engine/Source/Runtime/Renderer/Private/BasePassRendering.h
+++ b/Engine/Source/Runtime/Renderer/Private/BasePassRendering.h
@@ -767,12 +767,6 @@ public:
 		const FPrimitiveSceneProxy* PrimitiveSceneProxy,
 		FHitProxyId HitProxyId
 		);
-	static bool IsMaterialIgnored(const FMaterialRenderProxy* MaterialRenderProxy, ERHIFeatureLevel::Type InFeatureLevel)
-	{
-		// Ignore non-opaque materials in the opaque base pass.
-		// Note: blend mode does not depend on the feature level.
-		return MaterialRenderProxy && IsTranslucentBlendMode(MaterialRenderProxy->GetMaterial(InFeatureLevel)->GetBlendMode());
-	}
 };
 
 /** The parameters used to process a base pass mesh. */
diff --git a/Engine/Source/Runtime/Renderer/Private/CustomDepthRendering.cpp b/Engine/Source/Runtime/Renderer/Private/CustomDepthRendering.cpp
--- a/Engine/Source/Runtime/Renderer/Private/CustomDepthRendering.cpp
+++ b/Engine/Source/Runtime/Renderer/Private/CustomDepthRendering.cpp
@@ -25,10 +25,6 @@ bool FCustomDepthPrimSet::DrawPrims(FRHICommandListImmediate& RHICmdList, const
 	{
 		SCOPED_DRAW_EVENT(RHICmdList, CustomDepth);
 
-		const bool bUseGetMeshElements = ShouldUseGetDynamicMeshElements();
-
-		TDynamicPrimitiveDrawer<FDepthDrawingPolicyFactory> Drawer(RHICmdList, &View, FDepthDrawingPolicyFactory::ContextType(DDM_AllOccluders), true);
-
 		for (int32 PrimIdx = 0; PrimIdx < Prims.Num(); PrimIdx++)
 		{
 			FPrimitiveSceneProxy* PrimitiveSceneProxy = Prims[PrimIdx];
@@ -38,26 +34,18 @@ bool FCustomDepthPrimSet::DrawPrims(FRHICommandListImmediate& RHICmdList, const
 			{
 				const FPrimitiveViewRelevance& ViewRelevance = View.PrimitiveViewRelevanceMap[PrimitiveSceneInfo->GetIndex()];
 
-				if (bUseGetMeshElements)
+				FDepthDrawingPolicyFactory::ContextType Context(DDM_AllOccluders);
+
+				for (int32 MeshBatchIndex = 0; MeshBatchIndex < View.DynamicMeshElements.Num(); MeshBatchIndex++)
 				{
-					FDepthDrawingPolicyFactory::ContextType Context(DDM_AllOccluders);
+					const FMeshBatchAndRelevance& MeshBatchAndRelevance = View.DynamicMeshElements[MeshBatchIndex];
 
-					for (int32 MeshBatchIndex = 0; MeshBatchIndex < View.DynamicMeshElements.Num(); MeshBatchIndex++)
+					if (MeshBatchAndRelevance.PrimitiveSceneProxy == PrimitiveSceneProxy)
 					{
-						const FMeshBatchAndRelevance& MeshBatchAndRelevance = View.DynamicMeshElements[MeshBatchIndex];
-
-						if (MeshBatchAndRelevance.PrimitiveSceneProxy == PrimitiveSceneProxy)
-						{
-							const FMeshBatch& MeshBatch = *MeshBatchAndRelevance.Mesh;
-							FDepthDrawingPolicyFactory::DrawDynamicMesh(RHICmdList, View, Context, MeshBatch, false, true, MeshBatchAndRelevance.PrimitiveSceneProxy, MeshBatch.BatchHitProxyId);
-						}
+						const FMeshBatch& MeshBatch = *MeshBatchAndRelevance.Mesh;
+						FDepthDrawingPolicyFactory::DrawDynamicMesh(RHICmdList, View, Context, MeshBatch, false, true, MeshBatchAndRelevance.PrimitiveSceneProxy, MeshBatch.BatchHitProxyId);
 					}
 				}
-				else if (ViewRelevance.bDynamicRelevance)
-				{
-					Drawer.SetPrimitive(PrimitiveSceneProxy);
-					PrimitiveSceneProxy->DrawDynamicElements(&Drawer, &View);
-				}
 
 				if (ViewRelevance.bStaticRelevance)
 				{
@@ -82,9 +70,6 @@ bool FCustomDepthPrimSet::DrawPrims(FRHICommandListImmediate& RHICmdList, const
 				}
 			}
 		}
-
-		// Mark dirty if dynamic drawer rendered
-		bDirty |= Drawer.IsDirty();
 	}
 
 	return bDirty;
diff --git a/Engine/Source/Runtime/Renderer/Private/DeferredShadingRenderer.cpp b/Engine/Source/Runtime/Renderer/Private/DeferredShadingRenderer.cpp
--- a/Engine/Source/Runtime/Renderer/Private/DeferredShadingRenderer.cpp
+++ b/Engine/Source/Runtime/Renderer/Private/DeferredShadingRenderer.cpp
@@ -54,14 +54,6 @@ static TAutoConsoleVariable<int32> CVarVisualizeTexturePool(
 	ECVF_Cheat | ECVF_RenderThreadSafe);
 #endif
 
-/**  */
-static TAutoConsoleVariable<int32> CVarUseGetMeshElements(
-	TEXT("r.UseGetMeshElements"),
-	1,
-	TEXT("Whether to use the GetMeshElements path or the legacy DrawDynamicElements"),
-	ECVF_RenderThreadSafe
-	);
-
 /*-----------------------------------------------------------------------------
 	FDeferredShadingSceneRenderer
 -----------------------------------------------------------------------------*/
@@ -363,57 +355,22 @@ void FDeferredShadingSceneRenderer::RenderBasePassDynamicData(FRHICommandList& R
 	SCOPE_CYCLE_COUNTER(STAT_DynamicPrimitiveDrawTime);
 	SCOPED_DRAW_EVENT(RHICmdList, Dynamic);
 
-	const bool bUseGetMeshElements = ShouldUseGetDynamicMeshElements();
+	FBasePassOpaqueDrawingPolicyFactory::ContextType Context(false, ESceneRenderTargetsMode::DontSet);
 
-	if (bUseGetMeshElements)
+	for (int32 MeshBatchIndex = 0; MeshBatchIndex < View.DynamicMeshElements.Num(); MeshBatchIndex++)
 	{
-		FBasePassOpaqueDrawingPolicyFactory::ContextType Context(false, ESceneRenderTargetsMode::DontSet);
+		const FMeshBatchAndRelevance& MeshBatchAndRelevance = View.DynamicMeshElements[MeshBatchIndex];
 
-		for (int32 MeshBatchIndex = 0; MeshBatchIndex < View.DynamicMeshElements.Num(); MeshBatchIndex++)
+		if ((MeshBatchAndRelevance.bHasOpaqueOrMaskedMaterial || ViewFamily.EngineShowFlags.Wireframe)
+			&& MeshBatchAndRelevance.bRenderInMainPass)
 		{
-			const FMeshBatchAndRelevance& MeshBatchAndRelevance = View.DynamicMeshElements[MeshBatchIndex];
-
-			if ((MeshBatchAndRelevance.bHasOpaqueOrMaskedMaterial || ViewFamily.EngineShowFlags.Wireframe)
-				&& MeshBatchAndRelevance.bRenderInMainPass)
-			{
-				const FMeshBatch& MeshBatch = *MeshBatchAndRelevance.Mesh;
-				FBasePassOpaqueDrawingPolicyFactory::DrawDynamicMesh(RHICmdList, View, Context, MeshBatch, false, true, MeshBatchAndRelevance.PrimitiveSceneProxy, MeshBatch.BatchHitProxyId);
-			}
+			const FMeshBatch& MeshBatch = *MeshBatchAndRelevance.Mesh;
+			FBasePassOpaqueDrawingPolicyFactory::DrawDynamicMesh(RHICmdList, View, Context, MeshBatch, false, true, MeshBatchAndRelevance.PrimitiveSceneProxy, MeshBatch.BatchHitProxyId);
 		}
-
-		View.SimpleElementCollector.DrawBatchedElements(RHICmdList, View, FTexture2DRHIRef(), EBlendModeFilter::OpaqueAndMasked);
-	}
-	else if (View.VisibleDynamicPrimitives.Num() > 0)
-	{
-		// Draw the dynamic non-occluded primitives using a base pass drawing policy.
-		TDynamicPrimitiveDrawer<FBasePassOpaqueDrawingPolicyFactory> Drawer(
-			RHICmdList, &View, FBasePassOpaqueDrawingPolicyFactory::ContextType(false, ESceneRenderTargetsMode::DontSet), true);
-		for (int32 PrimitiveIndex = 0, Num = View.VisibleDynamicPrimitives.Num(); PrimitiveIndex < Num; PrimitiveIndex++)
-		{
-			const FPrimitiveSceneInfo* PrimitiveSceneInfo = View.VisibleDynamicPrimitives[PrimitiveIndex];
-			int32 PrimitiveId = PrimitiveSceneInfo->GetIndex();
-			const FPrimitiveViewRelevance& PrimitiveViewRelevance = View.PrimitiveViewRelevanceMap[PrimitiveId];
-
-			const bool bVisible = View.PrimitiveVisibilityMap[PrimitiveId];
-
-			// Only draw the primitive if it's visible
-			if (bVisible &&
-				// only draw opaque and masked primitives if wireframe is disabled
-				(PrimitiveViewRelevance.bOpaqueRelevance || ViewFamily.EngineShowFlags.Wireframe) &&
-				PrimitiveViewRelevance.bRenderInMainPass
-				)
-			{
-				FScopeCycleCounter Context(PrimitiveSceneInfo->Proxy->GetStatId());
-				Drawer.SetPrimitive(PrimitiveSceneInfo->Proxy);
-				PrimitiveSceneInfo->Proxy->DrawDynamicElements(
-					&Drawer,
-					&View
-					);
-			}
-		}
-		bDirty |= Drawer.IsDirty();
 	}
 
+	View.SimpleElementCollector.DrawBatchedElements(RHICmdList, View, FTexture2DRHIRef(), EBlendModeFilter::OpaqueAndMasked);
+
 	if( !View.Family->EngineShowFlags.CompositeEditorPrimitives )
 	{
 		const auto ShaderPlatform = View.GetShaderPlatform();
@@ -1071,78 +1028,34 @@ void FDeferredShadingSceneRenderer::Render(FRHICommandListImmediate& RHICmdList)
 
 bool FDeferredShadingSceneRenderer::RenderPrePassViewDynamic(FRHICommandList& RHICmdList, const FViewInfo& View)
 {
-	const bool bUseGetMeshElements = ShouldUseGetDynamicMeshElements();
+	FDepthDrawingPolicyFactory::ContextType Context(EarlyZPassMode);
 
-	if (bUseGetMeshElements)
+	for (int32 MeshBatchIndex = 0; MeshBatchIndex < View.DynamicMeshElements.Num(); MeshBatchIndex++)
 	{
-		FDepthDrawingPolicyFactory::ContextType Context(EarlyZPassMode);
+		const FMeshBatchAndRelevance& MeshBatchAndRelevance = View.DynamicMeshElements[MeshBatchIndex];
 
-		for (int32 MeshBatchIndex = 0; MeshBatchIndex < View.DynamicMeshElements.Num(); MeshBatchIndex++)
+		if (MeshBatchAndRelevance.bHasOpaqueOrMaskedMaterial && MeshBatchAndRelevance.bRenderInMainPass)
 		{
-			const FMeshBatchAndRelevance& MeshBatchAndRelevance = View.DynamicMeshElements[MeshBatchIndex];
+			const FMeshBatch& MeshBatch = *MeshBatchAndRelevance.Mesh;
+			const FPrimitiveSceneProxy* PrimitiveSceneProxy = MeshBatchAndRelevance.PrimitiveSceneProxy;
+			bool bShouldUseAsOccluder = true;
 
-			if (MeshBatchAndRelevance.bHasOpaqueOrMaskedMaterial && MeshBatchAndRelevance.bRenderInMainPass)
+			if (EarlyZPassMode != DDM_AllOccluders)
 			{
-				const FMeshBatch& MeshBatch = *MeshBatchAndRelevance.Mesh;
-				const FPrimitiveSceneProxy* PrimitiveSceneProxy = MeshBatchAndRelevance.PrimitiveSceneProxy;
-				bool bShouldUseAsOccluder = true;
-
-				if (EarlyZPassMode != DDM_AllOccluders)
-				{
-					extern float GMinScreenRadiusForDepthPrepass;
-					//@todo - move these proxy properties into FMeshBatchAndRelevance so we don't have to dereference the proxy in order to reject a mesh
-					const float LODFactorDistanceSquared = (PrimitiveSceneProxy->GetBounds().Origin - View.ViewMatrices.ViewOrigin).SizeSquared() * FMath::Square(View.LODDistanceFactor);
-
-					// Only render primitives marked as occluders
-					bShouldUseAsOccluder = PrimitiveSceneProxy->ShouldUseAsOccluder()
-						// Only render static objects unless movable are requested
-						&& (!PrimitiveSceneProxy->IsMovable() || GEarlyZPassMovable)
-						&& (FMath::Square(PrimitiveSceneProxy->GetBounds().SphereRadius) > GMinScreenRadiusForDepthPrepass * GMinScreenRadiusForDepthPrepass * LODFactorDistanceSquared);
-				}
-
-				if (bShouldUseAsOccluder)
-				{
-					FDepthDrawingPolicyFactory::DrawDynamicMesh(RHICmdList, View, Context, MeshBatch, false, true, PrimitiveSceneProxy, MeshBatch.BatchHitProxyId);
-				}
+				extern float GMinScreenRadiusForDepthPrepass;
+				//@todo - move these proxy properties into FMeshBatchAndRelevance so we don't have to dereference the proxy in order to reject a mesh
+				const float LODFactorDistanceSquared = (PrimitiveSceneProxy->GetBounds().Origin - View.ViewMatrices.ViewOrigin).SizeSquared() * FMath::Square(View.LODDistanceFactor);
+
+				// Only render primitives marked as occluders
+				bShouldUseAsOccluder = PrimitiveSceneProxy->ShouldUseAsOccluder()
+					// Only render static objects unless movable are requested
+					&& (!PrimitiveSceneProxy->IsMovable() || GEarlyZPassMovable)
+					&& (FMath::Square(PrimitiveSceneProxy->GetBounds().SphereRadius) > GMinScreenRadiusForDepthPrepass * GMinScreenRadiusForDepthPrepass * LODFactorDistanceSquared);
 			}
-		}
-	}
-	else if (View.VisibleDynamicPrimitives.Num() > 0)
-	{
-		// Draw the dynamic occluder primitives using a depth drawing policy.
-		TDynamicPrimitiveDrawer<FDepthDrawingPolicyFactory> Drawer(RHICmdList, &View, FDepthDrawingPolicyFactory::ContextType(EarlyZPassMode), true);
-		{
-			SCOPED_DRAW_EVENT(RHICmdList, Dynamic);
-			for(int32 PrimitiveIndex = 0;PrimitiveIndex < View.VisibleDynamicPrimitives.Num();PrimitiveIndex++)
+
+			if (bShouldUseAsOccluder)
 			{
-				const FPrimitiveSceneInfo* PrimitiveSceneInfo = View.VisibleDynamicPrimitives[PrimitiveIndex];
-				int32 PrimitiveId = PrimitiveSceneInfo->GetIndex();
-				const FPrimitiveViewRelevance& PrimitiveViewRelevance = View.PrimitiveViewRelevanceMap[PrimitiveId];
-
-				bool bShouldUseAsOccluder = true;
-
-				if(EarlyZPassMode != DDM_AllOccluders)
-				{
-					extern float GMinScreenRadiusForDepthPrepass;
-					const float LODFactorDistanceSquared = (PrimitiveSceneInfo->Proxy->GetBounds().Origin - View.ViewMatrices.ViewOrigin).SizeSquared() * FMath::Square(View.LODDistanceFactor);
-
-					// Only render primitives marked as occluders
-					bShouldUseAsOccluder = PrimitiveSceneInfo->Proxy->ShouldUseAsOccluder()
-						// Only render static objects unless movable are requested
-						&& (!PrimitiveSceneInfo->Proxy->IsMovable() || GEarlyZPassMovable)
-						&& (FMath::Square(PrimitiveSceneInfo->Proxy->GetBounds().SphereRadius) > GMinScreenRadiusForDepthPrepass * GMinScreenRadiusForDepthPrepass * LODFactorDistanceSquared);
-				}
-
-				// Only render opaque primitives marked as occluders
-				if (bShouldUseAsOccluder && PrimitiveViewRelevance.bOpaqueRelevance && PrimitiveViewRelevance.bRenderInMainPass)
-				{
-					FScopeCycleCounter Context(PrimitiveSceneInfo->Proxy->GetStatId());
-					Drawer.SetPrimitive(PrimitiveSceneInfo->Proxy);
-					PrimitiveSceneInfo->Proxy->DrawDynamicElements(
-						&Drawer,
-						&View
-						);
-				}
+				FDepthDrawingPolicyFactory::DrawDynamicMesh(RHICmdList, View, Context, MeshBatch, false, true, PrimitiveSceneProxy, MeshBatch.BatchHitProxyId);
 			}
 		}
 	}
diff --git a/Engine/Source/Runtime/Renderer/Private/DeferredShadingRenderer.h b/Engine/Source/Runtime/Renderer/Private/DeferredShadingRenderer.h
--- a/Engine/Source/Runtime/Renderer/Private/DeferredShadingRenderer.h
+++ b/Engine/Source/Runtime/Renderer/Private/DeferredShadingRenderer.h
@@ -245,7 +245,6 @@ private:
 
 	/** Renders the velocities for a subset of movable objects for the motion blur effect. */
 	friend class FRenderVelocityDynamicThreadTask;
-	void RenderDynamicVelocitiesInner(FRHICommandList& RHICmdList, const FViewInfo& View, int32 FirstIndex, int32 LastIndex);
 	void RenderDynamicVelocitiesMeshElementsInner(FRHICommandList& RHICmdList, const FViewInfo& View, int32 FirstIndex, int32 LastIndex);
 
 	/** Renders the velocities of movable objects for the motion blur effect. */
diff --git a/Engine/Source/Runtime/Renderer/Private/DepthRendering.cpp b/Engine/Source/Runtime/Renderer/Private/DepthRendering.cpp
--- a/Engine/Source/Runtime/Renderer/Private/DepthRendering.cpp
+++ b/Engine/Source/Runtime/Renderer/Private/DepthRendering.cpp
@@ -525,8 +525,3 @@ bool FDepthDrawingPolicyFactory::DrawStaticMesh(
 
 	return bDirty;
 }
-
-bool FDepthDrawingPolicyFactory::IsMaterialIgnored(const FMaterialRenderProxy* MaterialRenderProxy, ERHIFeatureLevel::Type InFeatureLevel)
-{
-	return IsTranslucentBlendMode(MaterialRenderProxy->GetMaterial(InFeatureLevel)->GetBlendMode());
-}
diff --git a/Engine/Source/Runtime/Renderer/Private/DepthRendering.h b/Engine/Source/Runtime/Renderer/Private/DepthRendering.h
--- a/Engine/Source/Runtime/Renderer/Private/DepthRendering.h
+++ b/Engine/Source/Runtime/Renderer/Private/DepthRendering.h
@@ -165,8 +165,6 @@ public:
 		FHitProxyId HitProxyId
 		);
 
-	static bool IsMaterialIgnored(const FMaterialRenderProxy* MaterialRenderProxy, ERHIFeatureLevel::Type InFeatureLevel);
-
 private:
 	/**
 	* Render a dynamic or static mesh using a depth draw policy
diff --git a/Engine/Source/Runtime/Renderer/Private/DistortionRendering.cpp b/Engine/Source/Runtime/Renderer/Private/DistortionRendering.cpp
--- a/Engine/Source/Runtime/Renderer/Private/DistortionRendering.cpp
+++ b/Engine/Source/Runtime/Renderer/Private/DistortionRendering.cpp
@@ -572,8 +572,6 @@ public:
 		const FPrimitiveSceneProxy* PrimitiveSceneProxy,
 		FHitProxyId HitProxyId
 		);
-
-	static bool IsMaterialIgnored(const FMaterialRenderProxy* Material, ERHIFeatureLevel::Type InFeatureLevel);
 };
 
 /**
@@ -677,13 +675,6 @@ bool TDistortionMeshDrawingPolicyFactory<DistortMeshPolicy>::DrawStaticMesh(
 	}
 }
 
-template<typename DistortMeshPolicy>
-bool TDistortionMeshDrawingPolicyFactory<DistortMeshPolicy>::IsMaterialIgnored(const FMaterialRenderProxy* MaterialRenderProxy, ERHIFeatureLevel::Type InFeatureLevel)
-{
-	// Non-distorted materials are ignored when rendering distortion.
-	return MaterialRenderProxy && !MaterialRenderProxy->GetMaterial(InFeatureLevel)->IsDistorted();
-}
-
 /*-----------------------------------------------------------------------------
 	FDistortionPrimSet
 -----------------------------------------------------------------------------*/
@@ -703,46 +694,24 @@ bool FDistortionPrimSet::DrawAccumulatedOffsets(FRHICommandListImmediate& RHICmd
 
 	if( Prims.Num() )
 	{
-		const bool bUseGetMeshElements = ShouldUseGetDynamicMeshElements();
-
-		// For drawing scene prims with dynamic relevance.
-		TDynamicPrimitiveDrawer<TDistortionMeshDrawingPolicyFactory<FDistortMeshAccumulatePolicy> > Drawer(
-			RHICmdList,
-			&View,
-			bInitializeOffsets,
-			false // Distortion is rendered post fog.
-			);
-
 		// Draw sorted scene prims
 		for( int32 PrimIdx=0; PrimIdx < Prims.Num(); PrimIdx++ )
 		{
 			FPrimitiveSceneProxy* PrimitiveSceneProxy = Prims[PrimIdx];
 			const FPrimitiveViewRelevance& ViewRelevance = View.PrimitiveViewRelevanceMap[PrimitiveSceneProxy->GetPrimitiveSceneInfo()->GetIndex()];
 
-			if (bUseGetMeshElements)
+			TDistortionMeshDrawingPolicyFactory<FDistortMeshAccumulatePolicy>::ContextType Context(bInitializeOffsets);
+
+			for (int32 MeshBatchIndex = 0; MeshBatchIndex < View.DynamicMeshElements.Num(); MeshBatchIndex++)
 			{
-				TDistortionMeshDrawingPolicyFactory<FDistortMeshAccumulatePolicy>::ContextType Context(bInitializeOffsets);
+				const FMeshBatchAndRelevance& MeshBatchAndRelevance = View.DynamicMeshElements[MeshBatchIndex];
 
-				for (int32 MeshBatchIndex = 0; MeshBatchIndex < View.DynamicMeshElements.Num(); MeshBatchIndex++)
+				if (MeshBatchAndRelevance.PrimitiveSceneProxy == PrimitiveSceneProxy)
 				{
-					const FMeshBatchAndRelevance& MeshBatchAndRelevance = View.DynamicMeshElements[MeshBatchIndex];
-
-					if (MeshBatchAndRelevance.PrimitiveSceneProxy == PrimitiveSceneProxy)
-					{
-						const FMeshBatch& MeshBatch = *MeshBatchAndRelevance.Mesh;
-						TDistortionMeshDrawingPolicyFactory<FDistortMeshAccumulatePolicy>::DrawDynamicMesh(RHICmdList, View, Context, MeshBatch, false, false, MeshBatchAndRelevance.PrimitiveSceneProxy, MeshBatch.BatchHitProxyId);
-					}
+					const FMeshBatch& MeshBatch = *MeshBatchAndRelevance.Mesh;
+					TDistortionMeshDrawingPolicyFactory<FDistortMeshAccumulatePolicy>::DrawDynamicMesh(RHICmdList, View, Context, MeshBatch, false, false, MeshBatchAndRelevance.PrimitiveSceneProxy, MeshBatch.BatchHitProxyId);
 				}
 			}
-			// Render dynamic scene prim
-			else if( ViewRelevance.bDynamicRelevance )
-			{				
-				Drawer.SetPrimitive(PrimitiveSceneProxy);
-				PrimitiveSceneProxy->DrawDynamicElements(
-					&Drawer,
-					&View
-					);
-			}
 
 			// Render static scene prim
 			if( ViewRelevance.bStaticRelevance )
@@ -769,8 +738,6 @@ bool FDistortionPrimSet::DrawAccumulatedOffsets(FRHICommandListImmediate& RHICmd
 				}
 			}
 		}
-		// Mark dirty if dynamic drawer rendered
-		bDirty |= Drawer.IsDirty();
 	}
 	return bDirty;
 }
diff --git a/Engine/Source/Runtime/Renderer/Private/DynamicPrimitiveDrawing.h b/Engine/Source/Runtime/Renderer/Private/DynamicPrimitiveDrawing.h
--- a/Engine/Source/Runtime/Renderer/Private/DynamicPrimitiveDrawing.h
+++ b/Engine/Source/Runtime/Renderer/Private/DynamicPrimitiveDrawing.h
@@ -7,134 +7,6 @@
 #ifndef __DYNAMICPRIMITIVEDRAWING_H__
 #define __DYNAMICPRIMITIVEDRAWING_H__
 
-/**
- * An implementation of the dynamic primitive definition interface to draw the elements passed to it on a given RHI command interface.
- */
-template<typename DrawingPolicyFactoryType>
-class TDynamicPrimitiveDrawer : public FPrimitiveDrawInterface
-{
-public:
-
-	/**
-	* Init constructor
-	*
-	* @param InView - view region being rendered
-	* @param InDrawingContext - contex for the given draw policy type
-	* @param InPreFog - rendering is occurring before fog
-	* @param bInIsHitTesting - rendering is occurring during hit testing
-	* @param bInIsVelocityRendering - rendering is occurring during velocity pass
-	*/
-	TDynamicPrimitiveDrawer(
-		FRHICommandList& InRHICmdList,
-		const FViewInfo* InView,
-		const typename DrawingPolicyFactoryType::ContextType& InDrawingContext,
-		bool InPreFog,
-		bool bInIsHitTesting = false,
-		bool bInEditorCompositeDepthTest = false,
-		bool bInIsSelectionOutlineRendering = false
-		):
-		FPrimitiveDrawInterface(InView),
-		RHICmdList(InRHICmdList),
-		View(InView),
-		DrawingContext(InDrawingContext),
-		PrimitiveSceneProxy(NULL),
-		bPreFog(InPreFog),
-		bDirty(false),
-		bIsHitTesting(bInIsHitTesting),
-		bEditorCompositeDepthTest(bInEditorCompositeDepthTest),
-		bIsSelectionOutlineRendering(bInIsSelectionOutlineRendering)
-	{}
-
-	~TDynamicPrimitiveDrawer();
-
-	void SetPrimitive(const FPrimitiveSceneProxy* NewPrimitiveSceneProxy);
-
-	// FPrimitiveDrawInterface interface.
-	virtual bool IsHitTesting() override;
-	virtual void SetHitProxy(HHitProxy* HitProxy) override;
-	virtual bool IsRenderingSelectionOutline() const override { return bIsSelectionOutlineRendering; }
-	virtual void RegisterDynamicResource(FDynamicPrimitiveResource* DynamicResource) override;
-	virtual bool IsMaterialIgnored(const FMaterialRenderProxy* MaterialRenderProxy, ERHIFeatureLevel::Type InFeatureLevel) const override;
-	virtual int32 DrawMesh(const FMeshBatch& Mesh) override;
-	virtual void AddReserveLines(uint8 DepthPriorityGroup, int32 NumLines, bool bDepthBiased = false, bool bThickLines = false) override;
-	virtual void DrawSprite(
-		const FVector& Position,
-		float SizeX,
-		float SizeY,
-		const FTexture* Sprite,
-		const FLinearColor& Color,
-		uint8 DepthPriorityGroup,
-		float U,
-		float UL,
-		float V,
-		float VL,
-		uint8 BlendMode = SE_BLEND_Masked
-		) override;
-	virtual void DrawLine(
-		const FVector& Start,
-		const FVector& End,
-		const FLinearColor& Color,
-		uint8 DepthPriorityGroup,
-		float Thickness = 0.0f,
-		float DepthBias = 0.0f,
-		bool bScreenSpace = false
-		) override;
-	virtual void DrawPoint(
-		const FVector& Position,
-		const FLinearColor& Color,
-		float PointSize,
-		uint8 DepthPriorityGroup
-		) override;
-
-	bool IsDirty() const
-	{
-		return bDirty;
-	}
-	void ClearDirtyFlag()
-	{
-		bDirty = false;
-	}
-
-private:
-
-	/** Command list to draw to. */
-	FRHICommandList& RHICmdList;
-
-	/** The view which is being rendered. */
-	const FViewInfo* const View;
-
-	/** The drawing context passed to the drawing policy for the mesh elements rendered by this drawer. */
-	typename DrawingPolicyFactoryType::ContextType DrawingContext;
-
-	/** The primitive being rendered. */
-	const FPrimitiveSceneProxy* PrimitiveSceneProxy;
-
-	/** The current hit proxy ID being rendered. */
-	FHitProxyId HitProxyId;
-
-	/** The batched simple elements. */
-	FBatchedElements BatchedElements;
-
-	/** The dynamic resources which have been registered with this drawer. */
-	TArray<FDynamicPrimitiveResource*,SceneRenderingAllocator> DynamicResources;
-
-	/** true if fog has not yet been rendered. */
-	uint32 bPreFog : 1;
-
-	/** Tracks whether any elements have been rendered by this drawer. */
-	uint32 bDirty : 1;
-
-	/** true if hit proxies are being drawn. */
-	uint32 bIsHitTesting : 1;
-
-	/** true if rendering is occuring during editor compositing */
-	uint32 bEditorCompositeDepthTest : 1;
-
-	/** true if we are currently rendering the selection outline */
-	uint32 bIsSelectionOutlineRendering : 1;
-
-};
-
 /**
 * Draws a range of view's elements with the specified drawing policy factory type.
 * @param View - The view to draw the meshes for.
@@ -187,22 +59,6 @@ void DrawViewElementsParallel(
 	FParallelCommandListSet& ParallelCommandListSet
 	);
 
-/**
- * Draws a given set of dynamic primitives to a RHI command interface, using the specified drawing policy type.
- * @param View - The view to draw the meshes for.
- * @param DrawingContext - The drawing policy type specific context for the drawing.
- * @param bPreFog - true if the draw call is occurring before fog has been rendered.
- */
-template<class DrawingPolicyFactoryType>
-bool DrawDynamicPrimitiveSet(
-	FRHICommandListImmediate& RHICmdList, 
-	const FViewInfo& View,
-	const TArray<const FPrimitiveSceneInfo*,SceneRenderingAllocator>& PrimitiveSet,
-	const typename DrawingPolicyFactoryType::ContextType& DrawingContext,
-	bool bPreFog,
-	bool bIsHitTesting = false
-	);
-
 /** A primitive draw interface which adds the drawn elements to the view's batched elements. */
 class FViewElementPDI : public FPrimitiveDrawInterface
 {
diff --git a/Engine/Source/Runtime/Renderer/Private/DynamicPrimitiveDrawing.inl b/Engine/Source/Runtime/Renderer/Private/DynamicPrimitiveDrawing.inl
--- a/Engine/Source/Runtime/Renderer/Private/DynamicPrimitiveDrawing.inl
+++ b/Engine/Source/Runtime/Renderer/Private/DynamicPrimitiveDrawing.inl
@@ -7,213 +7,6 @@
 #ifndef __DYNAMICPRIMITIVEDRAWING_INL__
 #define __DYNAMICPRIMITIVEDRAWING_INL__
 
-template<typename DrawingPolicyFactoryType>
-TDynamicPrimitiveDrawer<DrawingPolicyFactoryType>::~TDynamicPrimitiveDrawer()
-{
-	if(View)
-	{
-		const auto ShaderPlatform = View->GetShaderPlatform();
-		const bool bNeedToSwitchVerticalAxis = RHINeedsToSwitchVerticalAxis(ShaderPlatform) && !IsMobileHDR();
-
-		// Determine whether or not to depth test simple batched elements manually (msaa only)
-		// We need to read from the scene depth in the pixel shader to do manual depth testing
-		// Hardware depth testing will not work because there are two depth buffers (in the msaa case)
-		FTexture2DRHIRef DepthTexture;
-		if( bEditorCompositeDepthTest )
-		{
-			DepthTexture = GSceneRenderTargets.GetSceneDepthTexture();
-		}
-		
-		// Draw the batched elements.
-		BatchedElements.Draw(
-			RHICmdList,
-			View->GetFeatureLevel(),
-			bNeedToSwitchVerticalAxis,
-			View->ViewProjectionMatrix,
-			View->ViewRect.Width(),
-			View->ViewRect.Height(),
-			View->Family->EngineShowFlags.HitProxies,
-			1.0f,
-			View,
-			DepthTexture
-			);
-	}
-
-	// Cleanup the dynamic resources.
-	for(int32 ResourceIndex = 0;ResourceIndex < DynamicResources.Num();ResourceIndex++)
-	{
-		//release the resources before deleting, they will delete themselves
-		DynamicResources[ResourceIndex]->ReleasePrimitiveResource();
-	}
-}
-
-template<typename DrawingPolicyFactoryType>
-void TDynamicPrimitiveDrawer<DrawingPolicyFactoryType>::SetPrimitive(const FPrimitiveSceneProxy* NewPrimitiveSceneProxy)
-{
-	PrimitiveSceneProxy = NewPrimitiveSceneProxy;
-	if (NewPrimitiveSceneProxy)
-	{
-		HitProxyId = PrimitiveSceneProxy->GetPrimitiveSceneInfo()->DefaultDynamicHitProxyId;
-	}
-}
-
-template<typename DrawingPolicyFactoryType>
-bool TDynamicPrimitiveDrawer<DrawingPolicyFactoryType>::IsHitTesting()
-{
-	return bIsHitTesting;
-}
-
-template<typename DrawingPolicyFactoryType>
-void TDynamicPrimitiveDrawer<DrawingPolicyFactoryType>::SetHitProxy(HHitProxy* HitProxy)
-{
-	if(HitProxy)
-	{
-		// Only allow hit proxies from CreateHitProxies.
-		checkf(
-			PrimitiveSceneProxy->GetPrimitiveSceneInfo()->HitProxies.Find(HitProxy) != INDEX_NONE,
-			TEXT("Hit proxy used in DrawDynamicElements which wasn't created in CreateHitProxies")
-			);
-		HitProxyId = HitProxy->Id;
-	}
-	else
-	{
-		HitProxyId = FHitProxyId();
-	}
-}
-
-template<typename DrawingPolicyFactoryType>
-void TDynamicPrimitiveDrawer<DrawingPolicyFactoryType>::RegisterDynamicResource(FDynamicPrimitiveResource* DynamicResource)
-{
-	// Add the dynamic resource to the list of resources to cleanup on destruction.
-	DynamicResources.Add(DynamicResource);
-
-	// Initialize the dynamic resource immediately.
-	DynamicResource->InitPrimitiveResource();
-}
-
-template<typename DrawingPolicyFactoryType>
-bool TDynamicPrimitiveDrawer<DrawingPolicyFactoryType>::IsMaterialIgnored(const FMaterialRenderProxy* MaterialRenderProxy, ERHIFeatureLevel::Type InFeatureLevel) const
-{
-	return DrawingPolicyFactoryType::IsMaterialIgnored(MaterialRenderProxy, InFeatureLevel);
-}
-
-template<typename DrawingPolicyFactoryType>
-int32 TDynamicPrimitiveDrawer<DrawingPolicyFactoryType>::DrawMesh(const FMeshBatch& Mesh)
-{
-	int32 NumPassesRendered=0;
-	
-	check(Mesh.GetNumPrimitives() > 0);
-	check(Mesh.VertexFactory);
-#if DO_CHECK
-	Mesh.CheckUniformBuffers();
-#endif
-	
-	INC_DWORD_STAT_BY(STAT_DynamicPathMeshDrawCalls,Mesh.Elements.Num());
-	const bool DrawDirty = DrawingPolicyFactoryType::DrawDynamicMesh(
-		RHICmdList,
-		*View,
-		DrawingContext,
-		Mesh,
-		false,
-		bPreFog,
-		PrimitiveSceneProxy,
-		HitProxyId
-		);
-	bDirty |= DrawDirty;
-
-	NumPassesRendered += DrawDirty;
-	return NumPassesRendered;
-}
-
-template<typename DrawingPolicyFactoryType>
-void TDynamicPrimitiveDrawer<DrawingPolicyFactoryType>::DrawSprite(
-	const FVector& Position,
-	float SizeX,
-	float SizeY,
-	const FTexture* Sprite,
-	const FLinearColor& Color,
-	uint8 DepthPriorityGroup,
-	float U,
-	float UL,
-	float V,
-	float VL,
-	uint8 BlendMode
-	)
-{
-	if(DrawingPolicyFactoryType::bAllowSimpleElements)
-	{
-		BatchedElements.AddSprite(
-			Position,
-			SizeX,
-			SizeY,
-			Sprite,
-			Color,
-			HitProxyId,
-			U,
-			UL,
-			V,
-			VL,
-			BlendMode
-		);
-		bDirty = true;
-	}
-}
-
-template<typename DrawingPolicyFactoryType>
-void TDynamicPrimitiveDrawer<DrawingPolicyFactoryType>::AddReserveLines(uint8 DepthPriorityGroup, int32 NumLines, bool bDepthBiased, bool bThickLines)
-{
-	if (DrawingPolicyFactoryType::bAllowSimpleElements)
-	{
-		BatchedElements.AddReserveLines(NumLines, bDepthBiased, bThickLines);
-	}
-}
-
-template<typename DrawingPolicyFactoryType>
-void TDynamicPrimitiveDrawer<DrawingPolicyFactoryType>::DrawLine(
-	const FVector& Start,
-	const FVector& End,
-	const FLinearColor& Color,
-	uint8 DepthPriorityGroup,
-	float Thickness/* = 0.0f*/,
-	float DepthBias/* = 0.0f*/,
-	bool bScreenSpace/* = false*/
-	)
-{
-	if(DrawingPolicyFactoryType::bAllowSimpleElements)
-	{
-		BatchedElements.AddLine(
-			Start,
-			End,
-			Color,
-			HitProxyId,
-			Thickness,
-			DepthBias,
-			bScreenSpace
-		);
-		bDirty = true;
-	}
-}
-
-template<typename DrawingPolicyFactoryType>
-void TDynamicPrimitiveDrawer<DrawingPolicyFactoryType>::DrawPoint(
-	const FVector& Position,
-	const FLinearColor& Color,
-	float PointSize,
-	uint8 DepthPriorityGroup
-	)
-{
-	if(DrawingPolicyFactoryType::bAllowSimpleElements)
-	{
-		BatchedElements.AddPoint(
-			Position,
-			PointSize,
-			Color,
-			HitProxyId
-		);
-		bDirty = true;
-	}
-}
-
 template<class DrawingPolicyFactoryType>
 void DrawViewElementsInner(
 	FRHICommandList& RHICmdList,
@@ -368,40 +161,6 @@ bool DrawViewElements(
 	return false;
 }
 
-template<class DrawingPolicyFactoryType>
-bool DrawDynamicPrimitiveSet(
-	FRHICommandListImmediate& RHICmdList,
-	const FViewInfo& View,
-	const TArray<const FPrimitiveSceneInfo*,SceneRenderingAllocator>& PrimitiveSet,
-	const typename DrawingPolicyFactoryType::ContextType& DrawingContext,
-	bool bPreFog,
-	bool bIsHitTesting
-	)
-{
-	
-	// Draw the elements of each dynamic primitive.
-	TDynamicPrimitiveDrawer<DrawingPolicyFactoryType> Drawer(RHICmdList, &View, DrawingContext, bPreFog, bIsHitTesting);
-	for(int32 PrimitiveIndex = 0; PrimitiveIndex < PrimitiveSet.Num(); PrimitiveIndex++)
-	{
-		const FPrimitiveSceneInfo* PrimitiveSceneInfo = PrimitiveSet[PrimitiveIndex];
-		int32 PrimitiveId = PrimitiveSceneInfo->GetIndex();
-
-		if(!View.PrimitiveVisibilityMap[PrimitiveId])
-		{
-			continue;
-		}
-		FScopeCycleCounter Context(PrimitiveSceneInfo->Proxy->GetStatId());
-
-		Drawer.SetPrimitive(PrimitiveSceneInfo->Proxy);
-		PrimitiveSceneInfo->Proxy->DrawDynamicElements(
-			&Drawer,
-			&View
-			);
-	}
-
-	return Drawer.IsDirty();
-}
-
 inline FViewElementPDI::FViewElementPDI(FViewInfo* InViewInfo,FHitProxyConsumer* InHitProxyConsumer):
 	FPrimitiveDrawInterface(InViewInfo),
 	ViewInfo(InViewInfo),
diff --git a/Engine/Source/Runtime/Renderer/Private/ForwardBasePassRendering.cpp b/Engine/Source/Runtime/Renderer/Private/ForwardBasePassRendering.cpp
--- a/Engine/Source/Runtime/Renderer/Private/ForwardBasePassRendering.cpp
+++ b/Engine/Source/Runtime/Renderer/Private/ForwardBasePassRendering.cpp
@@ -388,49 +388,20 @@ void FForwardShadingSceneRenderer::RenderForwardShadingBasePass(FRHICommandListI
 			SCOPE_CYCLE_COUNTER(STAT_DynamicPrimitiveDrawTime);
 			SCOPED_DRAW_EVENT(RHICmdList, Dynamic);
 
-			const bool bUseGetMeshElements = ShouldUseGetDynamicMeshElements();
+			FBasePassForwardOpaqueDrawingPolicyFactory::ContextType Context(ESceneRenderTargetsMode::DontSet);
 
-			if (bUseGetMeshElements)
+			for (int32 MeshBatchIndex = 0; MeshBatchIndex < View.DynamicMeshElements.Num(); MeshBatchIndex++)
 			{
-				FBasePassForwardOpaqueDrawingPolicyFactory::ContextType Context(ESceneRenderTargetsMode::DontSet);
+				const FMeshBatchAndRelevance& MeshBatchAndRelevance = View.DynamicMeshElements[MeshBatchIndex];
 
-				for (int32 MeshBatchIndex = 0; MeshBatchIndex < View.DynamicMeshElements.Num(); MeshBatchIndex++)
+				if (MeshBatchAndRelevance.bHasOpaqueOrMaskedMaterial || ViewFamily.EngineShowFlags.Wireframe)
 				{
-					const FMeshBatchAndRelevance& MeshBatchAndRelevance = View.DynamicMeshElements[MeshBatchIndex];
-
-					if (MeshBatchAndRelevance.bHasOpaqueOrMaskedMaterial || ViewFamily.EngineShowFlags.Wireframe)
-					{
-						const FMeshBatch& MeshBatch = *MeshBatchAndRelevance.Mesh;
-						FBasePassForwardOpaqueDrawingPolicyFactory::DrawDynamicMesh(RHICmdList, View, Context, MeshBatch, false, true, MeshBatchAndRelevance.PrimitiveSceneProxy, MeshBatch.BatchHitProxyId);
-					}
+					const FMeshBatch& MeshBatch = *MeshBatchAndRelevance.Mesh;
+					FBasePassForwardOpaqueDrawingPolicyFactory::DrawDynamicMesh(RHICmdList, View, Context, MeshBatch, false, true, MeshBatchAndRelevance.PrimitiveSceneProxy, MeshBatch.BatchHitProxyId);
 				}
-
-				View.SimpleElementCollector.DrawBatchedElements(RHICmdList, View, NULL, EBlendModeFilter::OpaqueAndMasked);
 			}
-			else if (View.VisibleDynamicPrimitives.Num() > 0)
-			{
-				// Draw the dynamic non-occluded primitives using a base pass drawing policy.
-				TDynamicPrimitiveDrawer<FBasePassForwardOpaqueDrawingPolicyFactory> Drawer(RHICmdList, &View, FBasePassForwardOpaqueDrawingPolicyFactory::ContextType(ESceneRenderTargetsMode::DontSet), true);
 
-				for (int32 PrimitiveIndex = 0; PrimitiveIndex < View.VisibleDynamicPrimitives.Num(); PrimitiveIndex++)
-				{
-					const FPrimitiveSceneInfo* PrimitiveSceneInfo = View.VisibleDynamicPrimitives[PrimitiveIndex];
-					int32 PrimitiveId = PrimitiveSceneInfo->GetIndex();
-					const FPrimitiveViewRelevance& PrimitiveViewRelevance = View.PrimitiveViewRelevanceMap[PrimitiveId];
-
-					const bool bVisible = View.PrimitiveVisibilityMap[PrimitiveId];
-
-					// Only draw the primitive if it's visible
-					if (bVisible && 
-						// only draw opaque and masked primitives if wireframe is disabled
-						(PrimitiveViewRelevance.bOpaqueRelevance || ViewFamily.EngineShowFlags.Wireframe))
-					{
-						FScopeCycleCounter Context(PrimitiveSceneInfo->Proxy->GetStatId());
-						Drawer.SetPrimitive(PrimitiveSceneInfo->Proxy);
-						PrimitiveSceneInfo->Proxy->DrawDynamicElements(&Drawer, &View);
-					}
-				}
-			}
+			View.SimpleElementCollector.DrawBatchedElements(RHICmdList, View, NULL, EBlendModeFilter::OpaqueAndMasked);
 
 			const bool bNeedToSwitchVerticalAxis = RHINeedsToSwitchVerticalAxis(GShaderPlatformForFeatureLevel[FeatureLevel]);
 
diff --git a/Engine/Source/Runtime/Renderer/Private/ForwardBasePassRendering.h b/Engine/Source/Runtime/Renderer/Private/ForwardBasePassRendering.h
--- a/Engine/Source/Runtime/Renderer/Private/ForwardBasePassRendering.h
+++ b/Engine/Source/Runtime/Renderer/Private/ForwardBasePassRendering.h
@@ -496,12 +496,6 @@ public:
 		const FPrimitiveSceneProxy* PrimitiveSceneProxy,
 		FHitProxyId HitProxyId
 		);
-	static bool IsMaterialIgnored(const FMaterialRenderProxy* MaterialRenderProxy, ERHIFeatureLevel::Type InFeatureLevel)
-	{
-		// Ignore non-opaque materials in the opaque base pass.
-		// Note: blend mode does not depend on the feature level.
-		return MaterialRenderProxy && IsTranslucentBlendMode(MaterialRenderProxy->GetMaterial(InFeatureLevel)->GetBlendMode());
-	}
 };
 
 /** Processes a base pass mesh using an unknown light map policy, and unknown fog density policy. */
diff --git a/Engine/Source/Runtime/Renderer/Private/ForwardTranslucentRendering.cpp b/Engine/Source/Runtime/Renderer/Private/ForwardTranslucentRendering.cpp
--- a/Engine/Source/Runtime/Renderer/Private/ForwardTranslucentRendering.cpp
+++ b/Engine/Source/Runtime/Renderer/Private/ForwardTranslucentRendering.cpp
@@ -247,8 +247,6 @@ FTranslucentPrimSet
 
 void FTranslucentPrimSet::DrawPrimitivesForForwardShading(FRHICommandListImmediate& RHICmdList, const FViewInfo& View, FSceneRenderer& Renderer) const
 {
-	const bool bUseGetMeshElements = ShouldUseGetDynamicMeshElements();
-
 	// Draw sorted scene prims
 	for (int32 PrimIdx = 0; PrimIdx < SortedPrims.Num(); PrimIdx++)
 	{
@@ -260,38 +258,19 @@ void FTranslucentPrimSet::DrawPrimitivesForForwardShading(FRHICommandListImmedia
 
 		if(ViewRelevance.bDrawRelevance)
 		{
-			if (bUseGetMeshElements)
+			FTranslucencyForwardShadingDrawingPolicyFactory::ContextType Context;
+
+			//@todo parallelrendering - come up with a better way to filter these by primitive
+			for (int32 MeshBatchIndex = 0; MeshBatchIndex < View.DynamicMeshElements.Num(); MeshBatchIndex++)
 			{
-				FTranslucencyForwardShadingDrawingPolicyFactory::ContextType Context;
+				const FMeshBatchAndRelevance& MeshBatchAndRelevance = View.DynamicMeshElements[MeshBatchIndex];
 
-				//@todo parallelrendering - come up with a better way to filter these by primitive
-				for (int32 MeshBatchIndex = 0; MeshBatchIndex < View.DynamicMeshElements.Num(); MeshBatchIndex++)
+				if (MeshBatchAndRelevance.PrimitiveSceneProxy == PrimitiveSceneInfo->Proxy)
 				{
-					const FMeshBatchAndRelevance& MeshBatchAndRelevance = View.DynamicMeshElements[MeshBatchIndex];
-
-					if (MeshBatchAndRelevance.PrimitiveSceneProxy == PrimitiveSceneInfo->Proxy)
-					{
-						const FMeshBatch& MeshBatch = *MeshBatchAndRelevance.Mesh;
-						FTranslucencyForwardShadingDrawingPolicyFactory::DrawDynamicMesh(RHICmdList, View, Context, MeshBatch, false, false, MeshBatchAndRelevance.PrimitiveSceneProxy, MeshBatch.BatchHitProxyId);
-					}
+					const FMeshBatch& MeshBatch = *MeshBatchAndRelevance.Mesh;
+					FTranslucencyForwardShadingDrawingPolicyFactory::DrawDynamicMesh(RHICmdList, View, Context, MeshBatch, false, false, MeshBatchAndRelevance.PrimitiveSceneProxy, MeshBatch.BatchHitProxyId);
 				}
 			}
-			// Render dynamic scene prim
-			else if( ViewRelevance.bDynamicRelevance )
-			{
-				TDynamicPrimitiveDrawer<FTranslucencyForwardShadingDrawingPolicyFactory> TranslucencyDrawer(
-					RHICmdList,
-					&View,
-					FTranslucencyForwardShadingDrawingPolicyFactory::ContextType(),
-					false
-					);
-
-				TranslucencyDrawer.SetPrimitive(PrimitiveSceneInfo->Proxy);
-				PrimitiveSceneInfo->Proxy->DrawDynamicElements(
-					&TranslucencyDrawer,
-					&View
-					);
-			}
 
 			// Render static scene prim
 			if( ViewRelevance.bStaticRelevance )
@@ -319,10 +298,7 @@ void FTranslucentPrimSet::DrawPrimitivesForForwardShading(FRHICommandListImmedia
 		}
 	}
 
-	if (bUseGetMeshElements)
-	{
-		View.SimpleElementCollector.DrawBatchedElements(RHICmdList, View, FTexture2DRHIRef(), EBlendModeFilter::Translucent);
-	}
+	View.SimpleElementCollector.DrawBatchedElements(RHICmdList, View, FTexture2DRHIRef(), EBlendModeFilter::Translucent);
 }
 
 void FForwardShadingSceneRenderer::RenderTranslucency(FRHICommandListImmediate& RHICmdList)
diff --git a/Engine/Source/Runtime/Renderer/Private/LightMapDensityRendering.cpp b/Engine/Source/Runtime/Renderer/Private/LightMapDensityRendering.cpp
--- a/Engine/Source/Runtime/Renderer/Private/LightMapDensityRendering.cpp
+++ b/Engine/Source/Runtime/Renderer/Private/LightMapDensityRendering.cpp
@@ -58,48 +58,17 @@ bool FDeferredShadingSceneRenderer::RenderLightMapDensities(FRHICommandListImmed
 			{
 				SCOPED_DRAW_EVENT(RHICmdList, Dynamic);
 
-				const bool bUseGetMeshElements = ShouldUseGetDynamicMeshElements();
+				FLightMapDensityDrawingPolicyFactory::ContextType Context;
 
-				if (bUseGetMeshElements)
+				for (int32 MeshBatchIndex = 0; MeshBatchIndex < View.DynamicMeshElements.Num(); MeshBatchIndex++)
 				{
-					FLightMapDensityDrawingPolicyFactory::ContextType Context;
+					const FMeshBatchAndRelevance& MeshBatchAndRelevance = View.DynamicMeshElements[MeshBatchIndex];
 
-					for (int32 MeshBatchIndex = 0; MeshBatchIndex < View.DynamicMeshElements.Num(); MeshBatchIndex++)
+					if (MeshBatchAndRelevance.bHasOpaqueOrMaskedMaterial || ViewFamily.EngineShowFlags.Wireframe)
 					{
-						const FMeshBatchAndRelevance& MeshBatchAndRelevance = View.DynamicMeshElements[MeshBatchIndex];
-
-						if (MeshBatchAndRelevance.bHasOpaqueOrMaskedMaterial || ViewFamily.EngineShowFlags.Wireframe)
-						{
-							const FMeshBatch& MeshBatch = *MeshBatchAndRelevance.Mesh;
-							FLightMapDensityDrawingPolicyFactory::DrawDynamicMesh(RHICmdList, View, Context, MeshBatch, false, true, MeshBatchAndRelevance.PrimitiveSceneProxy, MeshBatch.BatchHitProxyId);
-						}
-					}
-				}
-				else if( View.VisibleDynamicPrimitives.Num() > 0 )
-				{
-					// Draw the dynamic non-occluded primitives using a base pass drawing policy.
-					TDynamicPrimitiveDrawer<FLightMapDensityDrawingPolicyFactory> Drawer(
-						RHICmdList, &View, FLightMapDensityDrawingPolicyFactory::ContextType(), true);
-					for (int32 PrimitiveIndex = 0;PrimitiveIndex < View.VisibleDynamicPrimitives.Num();PrimitiveIndex++)
-					{
-						const FPrimitiveSceneInfo* PrimitiveSceneInfo = View.VisibleDynamicPrimitives[PrimitiveIndex];
-						int32 PrimitiveId = PrimitiveSceneInfo->GetIndex();
-						const FPrimitiveViewRelevance& PrimitiveViewRelevance = View.PrimitiveViewRelevanceMap[PrimitiveId];
-						const bool bVisible = View.PrimitiveVisibilityMap[PrimitiveId];
-
-						// Only draw the primitive if it's visible
-						if( bVisible && 
-							// only draw opaque and masked primitives if wireframe is disabled
-							(PrimitiveViewRelevance.bOpaqueRelevance || ViewFamily.EngineShowFlags.Wireframe) )
-						{
-							Drawer.SetPrimitive(PrimitiveSceneInfo->Proxy);
-							PrimitiveSceneInfo->Proxy->DrawDynamicElements(
-								&Drawer,
-								&View		
-								);
-						}
+						const FMeshBatch& MeshBatch = *MeshBatchAndRelevance.Mesh;
+						FLightMapDensityDrawingPolicyFactory::DrawDynamicMesh(RHICmdList, View, Context, MeshBatch, false, true, MeshBatchAndRelevance.PrimitiveSceneProxy, MeshBatch.BatchHitProxyId);
 					}
-					bDirty |= Drawer.IsDirty(); 
 				}
 			}
 		}
diff --git a/Engine/Source/Runtime/Renderer/Private/LightMapDensityRendering.h b/Engine/Source/Runtime/Renderer/Private/LightMapDensityRendering.h
--- a/Engine/Source/Runtime/Renderer/Private/LightMapDensityRendering.h
+++ b/Engine/Source/Runtime/Renderer/Private/LightMapDensityRendering.h
@@ -439,14 +439,4 @@ public:
 		const FPrimitiveSceneProxy* PrimitiveSceneProxy,
 		FHitProxyId HitProxyId
 		);
-	static bool IsMaterialIgnored(const FMaterialRenderProxy* MaterialRenderProxy, ERHIFeatureLevel::Type InFeatureLevel)
-	{
-		// Note: MaterialModifiesMeshPosition may depend on feature level!
-		const FMaterial* Material = MaterialRenderProxy ? MaterialRenderProxy->GetMaterial(InFeatureLevel) : NULL;
-		return (Material && 
-				!(	Material->IsSpecialEngineMaterial() || 
-					Material->IsMasked() ||
-					Material->MaterialModifiesMeshPosition_RenderThread()
-				));
-	}
 };
diff --git a/Engine/Source/Runtime/Renderer/Private/PostProcess/PostProcessCompositeEditorPrimitives.cpp b/Engine/Source/Runtime/Renderer/Private/PostProcess/PostProcessCompositeEditorPrimitives.cpp
--- a/Engine/Source/Runtime/Renderer/Private/PostProcess/PostProcessCompositeEditorPrimitives.cpp
+++ b/Engine/Source/Runtime/Renderer/Private/PostProcess/PostProcessCompositeEditorPrimitives.cpp
@@ -311,68 +311,21 @@ void FRCPassPostProcessCompositeEditorPrimitives::RenderPrimitivesToComposite(FR
 	const bool bNeedToSwitchVerticalAxis = RHINeedsToSwitchVerticalAxis(ShaderPlatform);
 	FTexture2DRHIRef SceneDepth = GSceneRenderTargets.GetSceneDepthTexture();
 
-	const bool bUseGetMeshElements = ShouldUseGetDynamicMeshElements();
+	FBasePassOpaqueDrawingPolicyFactory::ContextType Context(true, ESceneRenderTargetsMode::SetTextures);
 
-	if (bUseGetMeshElements)
+	for (int32 MeshBatchIndex = 0; MeshBatchIndex < View.DynamicEditorMeshElements.Num(); MeshBatchIndex++)
 	{
-		FBasePassOpaqueDrawingPolicyFactory::ContextType Context(true, ESceneRenderTargetsMode::SetTextures);
+		const FMeshBatchAndRelevance& MeshBatchAndRelevance = View.DynamicEditorMeshElements[MeshBatchIndex];
 
-		for (int32 MeshBatchIndex = 0; MeshBatchIndex < View.DynamicEditorMeshElements.Num(); MeshBatchIndex++)
+		if (MeshBatchAndRelevance.bHasOpaqueOrMaskedMaterial || View.Family->EngineShowFlags.Wireframe)
 		{
-			const FMeshBatchAndRelevance& MeshBatchAndRelevance = View.DynamicEditorMeshElements[MeshBatchIndex];
-
-			if (MeshBatchAndRelevance.bHasOpaqueOrMaskedMaterial || View.Family->EngineShowFlags.Wireframe)
-			{
-				const FMeshBatch& MeshBatch = *MeshBatchAndRelevance.Mesh;
-				FBasePassOpaqueDrawingPolicyFactory::DrawDynamicMesh(RHICmdList, View, Context, MeshBatch, false, true, MeshBatchAndRelevance.PrimitiveSceneProxy, MeshBatch.BatchHitProxyId);
-			}
-		}
-
-		View.EditorSimpleElementCollector.DrawBatchedElements(RHICmdList, View, GSceneRenderTargets.GetSceneDepthTexture(), EBlendModeFilter::OpaqueAndMasked);
-	}
-	else if( View.VisibleEditorPrimitives.Num() > 0 )
-	{
-		// Draw the dynamic non-occluded primitives using a base pass drawing policy.
-		TDynamicPrimitiveDrawer<FBasePassOpaqueDrawingPolicyFactory> DepthTestDrawer(
-			RHICmdList, &View, FBasePassOpaqueDrawingPolicyFactory::ContextType(bDepthTest, ESceneRenderTargetsMode::SetTextures), true, false, bDepthTest);
-
-		TDynamicPrimitiveDrawer<FBasePassOpaqueDrawingPolicyFactory> NoDepthTestDrawer(
-			RHICmdList, &View, FBasePassOpaqueDrawingPolicyFactory::ContextType(!bDepthTest, ESceneRenderTargetsMode::SetTextures), true, false, !bDepthTest);
-
-		for(int32 PrimitiveIndex = 0;PrimitiveIndex < View.VisibleEditorPrimitives.Num();PrimitiveIndex++)
-		{
-			const FPrimitiveSceneInfo* PrimitiveSceneInfo = View.VisibleEditorPrimitives[PrimitiveIndex];
-			int32 PrimitiveId = PrimitiveSceneInfo->GetIndex();
-			const FPrimitiveViewRelevance& PrimitiveViewRelevance = View.PrimitiveViewRelevanceMap[PrimitiveId];
-
-			const bool bVisible = View.PrimitiveVisibilityMap[PrimitiveId];
-
-			// Only draw the primitive if it's visible
-			if( bVisible && 
-				// only draw opaque and masked primitives if wireframe is disabled
-				(PrimitiveViewRelevance.bOpaqueRelevance || View.Family->EngineShowFlags.Wireframe))
-			{
-				// Distinguish between depth tested & non-depth tested primitives
-				if(PrimitiveViewRelevance.bEditorNoDepthTestPrimitiveRelevance)
-				{
-					NoDepthTestDrawer.SetPrimitive(PrimitiveSceneInfo->Proxy);
-					PrimitiveSceneInfo->Proxy->DrawDynamicElements(
-						&NoDepthTestDrawer,
-						&View			
-						);
-				}
-				else
-				{
-					DepthTestDrawer.SetPrimitive(PrimitiveSceneInfo->Proxy);
-					PrimitiveSceneInfo->Proxy->DrawDynamicElements(
-						&DepthTestDrawer,
-						&View			
-						);
-				}
-			}
+			const FMeshBatch& MeshBatch = *MeshBatchAndRelevance.Mesh;
+			FBasePassOpaqueDrawingPolicyFactory::DrawDynamicMesh(RHICmdList, View, Context, MeshBatch, false, true, MeshBatchAndRelevance.PrimitiveSceneProxy, MeshBatch.BatchHitProxyId);
 		}
 	}
 
+	View.EditorSimpleElementCollector.DrawBatchedElements(RHICmdList, View, GSceneRenderTargets.GetSceneDepthTexture(), EBlendModeFilter::OpaqueAndMasked);
+	
 	// Draw the base pass for the view's batched mesh elements.
 	DrawViewElements<FBasePassOpaqueDrawingPolicyFactory>(RHICmdList, View, FBasePassOpaqueDrawingPolicyFactory::ContextType(bDepthTest, ESceneRenderTargetsMode::SetTextures), SDPG_World, false);
 
diff --git a/Engine/Source/Runtime/Renderer/Private/PostProcess/PostProcessSelectionOutline.cpp b/Engine/Source/Runtime/Renderer/Private/PostProcess/PostProcessSelectionOutline.cpp
--- a/Engine/Source/Runtime/Renderer/Private/PostProcess/PostProcessSelectionOutline.cpp
+++ b/Engine/Source/Runtime/Renderer/Private/PostProcess/PostProcessSelectionOutline.cpp
@@ -42,100 +42,38 @@ void FRCPassPostProcessSelectionOutlineColor::Process(FRenderingCompositePassCon
 
 	if (View.Family->EngineShowFlags.Selection)
 	{
-		const bool bUseGetMeshElements = ShouldUseGetDynamicMeshElements();
+		FHitProxyDrawingPolicyFactory::ContextType FactoryContext;
 
-		if (bUseGetMeshElements)
-		{
-			FHitProxyDrawingPolicyFactory::ContextType FactoryContext;
-
-			//@todo - use memstack
-			TMap<FName, int32> ActorNameToStencilIndex;
-			ActorNameToStencilIndex.Add(NAME_BSP, 1);
+		//@todo - use memstack
+		TMap<FName, int32> ActorNameToStencilIndex;
+		ActorNameToStencilIndex.Add(NAME_BSP, 1);
 
-			Context.RHICmdList.SetRasterizerState(TStaticRasterizerState<>::GetRHI());
-			Context.RHICmdList.SetBlendState(TStaticBlendStateWriteMask<CW_NONE, CW_NONE, CW_NONE, CW_NONE>::GetRHI());
+		Context.RHICmdList.SetRasterizerState(TStaticRasterizerState<>::GetRHI());
+		Context.RHICmdList.SetBlendState(TStaticBlendStateWriteMask<CW_NONE, CW_NONE, CW_NONE, CW_NONE>::GetRHI());
 
-			for (int32 MeshBatchIndex = 0; MeshBatchIndex < View.DynamicMeshElements.Num(); MeshBatchIndex++)
-			{
-				const FMeshBatchAndRelevance& MeshBatchAndRelevance = View.DynamicMeshElements[MeshBatchIndex];
-				const FPrimitiveSceneProxy* PrimitiveSceneProxy = MeshBatchAndRelevance.PrimitiveSceneProxy;
-
-				if (PrimitiveSceneProxy->IsSelected() && MeshBatchAndRelevance.Mesh->bUseSelectionOutline)
-				{
-					const int32* AssignedStencilIndexPtr = ActorNameToStencilIndex.Find(PrimitiveSceneProxy->GetOwnerName());
-
-					if (!AssignedStencilIndexPtr)
-					{
-						AssignedStencilIndexPtr = &ActorNameToStencilIndex.Add(PrimitiveSceneProxy->GetOwnerName(), ActorNameToStencilIndex.Num() + 1);
-					}
-
-					// This is a reversed Z depth surface, using CF_GreaterEqual.
-					// Note that the stencil value will overflow with enough selected objects
-					Context.RHICmdList.SetDepthStencilState(TStaticDepthStencilState<true, CF_GreaterEqual, true, CF_Always, SO_Keep, SO_Keep, SO_Replace>::GetRHI(), *AssignedStencilIndexPtr);
-
-					const FMeshBatch& MeshBatch = *MeshBatchAndRelevance.Mesh;
-					FHitProxyDrawingPolicyFactory::DrawDynamicMesh(Context.RHICmdList, View, FactoryContext, MeshBatch, false, true, MeshBatchAndRelevance.PrimitiveSceneProxy, MeshBatch.BatchHitProxyId);
-				}
-			}
-		}
-		else if (View.VisibleDynamicPrimitives.Num() > 0)
+		for (int32 MeshBatchIndex = 0; MeshBatchIndex < View.DynamicMeshElements.Num(); MeshBatchIndex++)
 		{
-			TDynamicPrimitiveDrawer<FHitProxyDrawingPolicyFactory> Drawer(Context.RHICmdList, &View, FHitProxyDrawingPolicyFactory::ContextType(), true, false, false, true);
-			TMultiMap<FName, const FPrimitiveSceneInfo*> PrimitivesByActor;
+			const FMeshBatchAndRelevance& MeshBatchAndRelevance = View.DynamicMeshElements[MeshBatchIndex];
+			const FPrimitiveSceneProxy* PrimitiveSceneProxy = MeshBatchAndRelevance.PrimitiveSceneProxy;
 
-			for (int32 PrimitiveIndex = 0; PrimitiveIndex < View.VisibleDynamicPrimitives.Num();PrimitiveIndex++)
+			if (PrimitiveSceneProxy->IsSelected() && MeshBatchAndRelevance.Mesh->bUseSelectionOutline)
 			{
-				const FPrimitiveSceneInfo* PrimitiveSceneInfo = View.VisibleDynamicPrimitives[PrimitiveIndex];
+				const int32* AssignedStencilIndexPtr = ActorNameToStencilIndex.Find(PrimitiveSceneProxy->GetOwnerName());
 
-				// Only draw the primitive if relevant
-				if(PrimitiveSceneInfo->Proxy->IsSelected())
+				if (!AssignedStencilIndexPtr)
 				{
-					PrimitivesByActor.Add(PrimitiveSceneInfo->Proxy->GetOwnerName(), PrimitiveSceneInfo);
+					AssignedStencilIndexPtr = &ActorNameToStencilIndex.Add(PrimitiveSceneProxy->GetOwnerName(), ActorNameToStencilIndex.Num() + 1);
 				}
-			}
-
-			if (PrimitivesByActor.Num())
-			{
-				// 0 means no object, 1 means BSP so we start with 2
-				uint32 StencilValue = 2;
 
-				Context.RHICmdList.SetRasterizerState(TStaticRasterizerState<>::GetRHI());
-				Context.RHICmdList.SetBlendState(TStaticBlendStateWriteMask<CW_NONE, CW_NONE, CW_NONE, CW_NONE>::GetRHI());
+				// This is a reversed Z depth surface, using CF_GreaterEqual.
+				// Note that the stencil value will overflow with enough selected objects
+				Context.RHICmdList.SetDepthStencilState(TStaticDepthStencilState<true, CF_GreaterEqual, true, CF_Always, SO_Keep, SO_Keep, SO_Replace>::GetRHI(), *AssignedStencilIndexPtr);
 
-				// Sort by actor
-				TArray<FName> Actors;
-				PrimitivesByActor.GetKeys(Actors);
-				for( TArray<FName>::TConstIterator ActorIt(Actors); ActorIt; ++ActorIt )
-				{
-					bool bBSP = *ActorIt == NAME_BSP;
-					if (bBSP)
-					{
-						// This is a reversed Z depth surface, using CF_GreaterEqual.
-						Context.RHICmdList.SetDepthStencilState(TStaticDepthStencilState<true, CF_GreaterEqual, true, CF_Always, SO_Keep, SO_Keep, SO_Replace>::GetRHI(), 1);
-					}
-					else
-					{
-						// This is a reversed Z depth surface, using CF_GreaterEqual.
-						Context.RHICmdList.SetDepthStencilState(TStaticDepthStencilState<true, CF_GreaterEqual, true, CF_Always, SO_Keep, SO_Keep, SO_Replace>::GetRHI(), StencilValue);
-
-						// we want to use 1..255 for all objects, not correct silhouettes after that is acceptable
-						StencilValue = (StencilValue == 255) ? 2 : (StencilValue + 1);
-					}
-
-					TArray<const FPrimitiveSceneInfo*> Primitives;
-					PrimitivesByActor.MultiFind(*ActorIt, Primitives);
-
-					for( TArray<const FPrimitiveSceneInfo*>::TConstIterator PrimIt(Primitives); PrimIt; ++PrimIt )
-					{
-						const FPrimitiveSceneInfo* PrimitiveSceneInfo = *PrimIt;
-						// Render the object to the stencil/depth buffer
-						Drawer.SetPrimitive(PrimitiveSceneInfo->Proxy);
-						PrimitiveSceneInfo->Proxy->DrawDynamicElements(&Drawer, &View);
-					}
-				}
+				const FMeshBatch& MeshBatch = *MeshBatchAndRelevance.Mesh;
+				FHitProxyDrawingPolicyFactory::DrawDynamicMesh(Context.RHICmdList, View, FactoryContext, MeshBatch, false, true, MeshBatchAndRelevance.PrimitiveSceneProxy, MeshBatch.BatchHitProxyId);
 			}
 		}
-
+		
 		// to get an outline around the objects if it's partly outside of the screen
 		{
 			FIntRect InnerRect = ViewRect;
diff --git a/Engine/Source/Runtime/Renderer/Private/SceneHitProxyRendering.cpp b/Engine/Source/Runtime/Renderer/Private/SceneHitProxyRendering.cpp
--- a/Engine/Source/Runtime/Renderer/Private/SceneHitProxyRendering.cpp
+++ b/Engine/Source/Runtime/Renderer/Private/SceneHitProxyRendering.cpp
@@ -407,63 +407,36 @@ void RenderHitProxies(FRHICommandListImmediate& RHICmdList, const FSceneRenderer
 		const bool bPreFog = true;
 		const bool bHitTesting = true;
 
-		const bool bUseGetMeshElements = ShouldUseGetDynamicMeshElements();
+		FHitProxyDrawingPolicyFactory::ContextType DrawingContext;
 
-		if (bUseGetMeshElements)
+		for (int32 MeshBatchIndex = 0; MeshBatchIndex < View.DynamicMeshElements.Num(); MeshBatchIndex++)
 		{
-			FHitProxyDrawingPolicyFactory::ContextType DrawingContext;
+			const FMeshBatchAndRelevance& MeshBatchAndRelevance = View.DynamicMeshElements[MeshBatchIndex];
+			const FMeshBatch& MeshBatch = *MeshBatchAndRelevance.Mesh;
 
-			for (int32 MeshBatchIndex = 0; MeshBatchIndex < View.DynamicMeshElements.Num(); MeshBatchIndex++)
+			if (MeshBatch.bSelectable)
 			{
-				const FMeshBatchAndRelevance& MeshBatchAndRelevance = View.DynamicMeshElements[MeshBatchIndex];
-				const FMeshBatch& MeshBatch = *MeshBatchAndRelevance.Mesh;
-				
-				if (MeshBatch.bSelectable)
-				{
-					const FHitProxyId EffectiveHitProxyId = MeshBatch.BatchHitProxyId == FHitProxyId() ? MeshBatchAndRelevance.PrimitiveSceneProxy->GetPrimitiveSceneInfo()->DefaultDynamicHitProxyId : MeshBatch.BatchHitProxyId;
-					FHitProxyDrawingPolicyFactory::DrawDynamicMesh(RHICmdList, View, DrawingContext, MeshBatch, false, bPreFog, MeshBatchAndRelevance.PrimitiveSceneProxy, EffectiveHitProxyId);
-				}
+				const FHitProxyId EffectiveHitProxyId = MeshBatch.BatchHitProxyId == FHitProxyId() ? MeshBatchAndRelevance.PrimitiveSceneProxy->GetPrimitiveSceneInfo()->DefaultDynamicHitProxyId : MeshBatch.BatchHitProxyId;
+				FHitProxyDrawingPolicyFactory::DrawDynamicMesh(RHICmdList, View, DrawingContext, MeshBatch, false, bPreFog, MeshBatchAndRelevance.PrimitiveSceneProxy, EffectiveHitProxyId);
 			}
+		}
 
-			View.SimpleElementCollector.DrawBatchedElements(RHICmdList, View, FTexture2DRHIRef(), EBlendModeFilter::All);
-
-			for (int32 MeshBatchIndex = 0; MeshBatchIndex < View.DynamicEditorMeshElements.Num(); MeshBatchIndex++)
-			{
-				const FMeshBatchAndRelevance& MeshBatchAndRelevance = View.DynamicEditorMeshElements[MeshBatchIndex];
-				const FMeshBatch& MeshBatch = *MeshBatchAndRelevance.Mesh;
-
-				if (MeshBatch.bSelectable)
-				{
-					const FHitProxyId EffectiveHitProxyId = MeshBatch.BatchHitProxyId == FHitProxyId() ? MeshBatchAndRelevance.PrimitiveSceneProxy->GetPrimitiveSceneInfo()->DefaultDynamicHitProxyId : MeshBatch.BatchHitProxyId;
-					FHitProxyDrawingPolicyFactory::DrawDynamicMesh(RHICmdList, View, DrawingContext, MeshBatch, false, bPreFog, MeshBatchAndRelevance.PrimitiveSceneProxy, EffectiveHitProxyId);
-				}
-			}
+		View.SimpleElementCollector.DrawBatchedElements(RHICmdList, View, FTexture2DRHIRef(), EBlendModeFilter::All);
 
-			View.EditorSimpleElementCollector.DrawBatchedElements(RHICmdList, View, FTexture2DRHIRef(), EBlendModeFilter::All);
-		}
-		else 
+		for (int32 MeshBatchIndex = 0; MeshBatchIndex < View.DynamicEditorMeshElements.Num(); MeshBatchIndex++)
 		{
-			// Draw all dynamic primitives using the hit proxy drawing policy.
-			DrawDynamicPrimitiveSet<FHitProxyDrawingPolicyFactory>(
-				RHICmdList,
-				View,
-				View.VisibleDynamicPrimitives,
-				FHitProxyDrawingPolicyFactory::ContextType(),
-				bPreFog,
-				bHitTesting
-				);
+			const FMeshBatchAndRelevance& MeshBatchAndRelevance = View.DynamicEditorMeshElements[MeshBatchIndex];
+			const FMeshBatch& MeshBatch = *MeshBatchAndRelevance.Mesh;
 
-			// Draw all visible editor primitives using the hit proxy drawing policy
-			DrawDynamicPrimitiveSet<FHitProxyDrawingPolicyFactory>(
-				RHICmdList,
-				View,
-				View.VisibleEditorPrimitives,
-				FHitProxyDrawingPolicyFactory::ContextType(),
-				bPreFog,
-				bHitTesting
-				);
+			if (MeshBatch.bSelectable)
+			{
+				const FHitProxyId EffectiveHitProxyId = MeshBatch.BatchHitProxyId == FHitProxyId() ? MeshBatchAndRelevance.PrimitiveSceneProxy->GetPrimitiveSceneInfo()->DefaultDynamicHitProxyId : MeshBatch.BatchHitProxyId;
+				FHitProxyDrawingPolicyFactory::DrawDynamicMesh(RHICmdList, View, DrawingContext, MeshBatch, false, bPreFog, MeshBatchAndRelevance.PrimitiveSceneProxy, EffectiveHitProxyId);
+			}
 		}
 
+		View.EditorSimpleElementCollector.DrawBatchedElements(RHICmdList, View, FTexture2DRHIRef(), EBlendModeFilter::All);
+		
 		// Draw the view's elements.
 		DrawViewElements<FHitProxyDrawingPolicyFactory>(RHICmdList, View, FHitProxyDrawingPolicyFactory::ContextType(), SDPG_World, bPreFog);
 
diff --git a/Engine/Source/Runtime/Renderer/Private/SceneHitProxyRendering.h b/Engine/Source/Runtime/Renderer/Private/SceneHitProxyRendering.h
--- a/Engine/Source/Runtime/Renderer/Private/SceneHitProxyRendering.h
+++ b/Engine/Source/Runtime/Renderer/Private/SceneHitProxyRendering.h
@@ -78,8 +78,4 @@ public:
 		const FPrimitiveSceneProxy* PrimitiveSceneProxy,
 		FHitProxyId HitProxyId
 		);
-	static bool IsMaterialIgnored(const FMaterialRenderProxy* MaterialRenderProxy, ERHIFeatureLevel::Type InFeatureLevel)
-	{
-		return false;
-	}
 };
diff --git a/Engine/Source/Runtime/Renderer/Private/ScenePrivate.h b/Engine/Source/Runtime/Renderer/Private/ScenePrivate.h
--- a/Engine/Source/Runtime/Renderer/Private/ScenePrivate.h
+++ b/Engine/Source/Runtime/Renderer/Private/ScenePrivate.h
@@ -40,8 +40,6 @@ bool IsMobileHDR();
 /** True if the mobile renderer is emulating HDR in a 32bpp render target. */
 bool IsMobileHDR32bpp();
 
-extern bool ShouldUseGetDynamicMeshElements();
-
 // Dependencies.
 #include "StaticBoundShaderState.h"
 #include "BatchedElements.h"
diff --git a/Engine/Source/Runtime/Renderer/Private/SceneRendering.cpp b/Engine/Source/Runtime/Renderer/Private/SceneRendering.cpp
--- a/Engine/Source/Runtime/Renderer/Private/SceneRendering.cpp
+++ b/Engine/Source/Runtime/Renderer/Private/SceneRendering.cpp
@@ -1200,10 +1200,3 @@ bool IsMobileHDR32bpp()
 	static auto* MobileHDR32bppCvar = IConsoleManager::Get().FindTConsoleVariableDataInt(TEXT("r.MobileHDR32bpp"));
 	return IsMobileHDR() && (GSupportsRenderTargetFormat_PF_FloatRGBA == false || MobileHDR32bppCvar->GetValueOnRenderThread() == 1);
 }
-
-bool ShouldUseGetDynamicMeshElements()
-{
-	static const auto CVarUseGetMeshElements = IConsoleManager::Get().FindTConsoleVariableDataInt(TEXT("r.UseGetMeshElements"));
-	const bool bUseGetMeshElements = !GRHICommandList.Bypass() ||  CVarUseGetMeshElements->GetValueOnRenderThread() != 0;
-	return bUseGetMeshElements;
-}
\ No newline at end of file
diff --git a/Engine/Source/Runtime/Renderer/Private/SceneVisibility.cpp b/Engine/Source/Runtime/Renderer/Private/SceneVisibility.cpp
--- a/Engine/Source/Runtime/Renderer/Private/SceneVisibility.cpp
+++ b/Engine/Source/Runtime/Renderer/Private/SceneVisibility.cpp
@@ -642,7 +642,6 @@ typedef TArray<int32, SceneRenderingAllocator> FRelevantStaticPrimitives;
  * @param ViewBit - Bit mask: 1 << ViewIndex where Views(ViewIndex) == View.
  * @param OutRelevantStaticPrimitives - Upon return contains a list of relevant
  *                                      static primitives.
- * @param OutPreRenderViewMasks - Primitives requiring the PreRenderView
  *                                callback for this view will have ViewBit set.
  */
 static void ComputeRelevanceForView(
@@ -651,7 +650,6 @@ static void ComputeRelevanceForView(
 	FViewInfo& View,
 	uint8 ViewBit,
 	FRelevantStaticPrimitives& OutRelevantStaticPrimitives,
-	FPrimitiveViewMasks& OutPreRenderViewMasks,
 	FPrimitiveViewMasks& OutHasDynamicMeshElementsMasks,
 	FPrimitiveViewMasks& OutHasDynamicEditorMeshElementsMasks
 	)
@@ -664,7 +662,6 @@ static void ComputeRelevanceForView(
 	// - Update last render time.
 	// But it's hard to say since the view relevance is going to be available right now.
 
-	check(OutPreRenderViewMasks.Num() == Scene->Primitives.Num());
 	check(OutHasDynamicMeshElementsMasks.Num() == Scene->Primitives.Num());
 
 	float CurrentWorldTime = View.Family->CurrentWorldTime;
@@ -695,11 +692,6 @@ static void ComputeRelevanceForView(
 			continue;
 		}
 
-		if (ViewRelevance.bNeedsPreRenderView)
-		{
-			OutPreRenderViewMasks[BitIt.GetIndex()] |= ViewBit;
-		}
-
 		if (bEditorRelevance)
 		{
 			// Editor primitives are rendered after post processing and composited onto the scene
@@ -857,7 +849,6 @@ struct FRelevancePacket
 	const FViewInfo& View;
 	const uint8 ViewBit;
 	const FMarkRelevantStaticMeshesForViewData& ViewData;
-	FPrimitiveViewMasks& OutPreRenderViewMasks;
 	FPrimitiveViewMasks& OutHasDynamicMeshElementsMasks;
 	FPrimitiveViewMasks& OutHasDynamicEditorMeshElementsMasks;
 	uint8* RESTRICT MarkMasks;
@@ -879,7 +870,6 @@ struct FRelevancePacket
 		const FViewInfo& InView, 
 		uint8 InViewBit,
 		const FMarkRelevantStaticMeshesForViewData& InViewData,
-		FPrimitiveViewMasks& InOutPreRenderViewMasks,
 		FPrimitiveViewMasks& InOutHasDynamicMeshElementsMasks,
 		FPrimitiveViewMasks& InOutHasDynamicEditorMeshElementsMasks,
 		uint8* InMarkMasks)
@@ -891,7 +881,6 @@ struct FRelevancePacket
 		, View(InView)
 		, ViewBit(InViewBit)
 		, ViewData(InViewData)
-		, OutPreRenderViewMasks(InOutPreRenderViewMasks)
 		, OutHasDynamicMeshElementsMasks(InOutHasDynamicMeshElementsMasks)
 		, OutHasDynamicEditorMeshElementsMasks(InOutHasDynamicEditorMeshElementsMasks)
 		, MarkMasks(InMarkMasks)
@@ -933,11 +922,6 @@ struct FRelevancePacket
 				continue;
 			}
 
-			if (ViewRelevance.bNeedsPreRenderView)
-			{
-				OutPreRenderViewMasks[BitIndex] |= ViewBit;
-			}
-
 			if (bEditorRelevance)
 			{
 				// Editor primitives are rendered after post processing and composited onto the scene
@@ -1189,7 +1173,6 @@ static TAutoConsoleVariable<int32> CVarParallelInitViews(
  * @param ViewBit - Bit mask: 1 << ViewIndex where Views(ViewIndex) == View.
  * @param OutRelevantStaticPrimitives - Upon return contains a list of relevant
  *                                      static primitives.
- * @param OutPreRenderViewMasks - Primitives requiring the PreRenderView
  *                                callback for this view will have ViewBit set.
  */
 static void ComputeAndMarkRelevanceForViewParallel(
@@ -1197,13 +1180,10 @@ static void ComputeAndMarkRelevanceForViewParallel(
 	const FScene* Scene,
 	FViewInfo& View,
 	uint8 ViewBit,
-	FPrimitiveViewMasks& OutPreRenderViewMasks,
 	FPrimitiveViewMasks& OutHasDynamicMeshElementsMasks,
 	FPrimitiveViewMasks& OutHasDynamicEditorMeshElementsMasks
 	)
 {
-
-	check(OutPreRenderViewMasks.Num() == Scene->Primitives.Num());
 	check(OutHasDynamicMeshElementsMasks.Num() == Scene->Primitives.Num());
 
 	const FMarkRelevantStaticMeshesForViewData ViewData(View);
@@ -1235,7 +1215,6 @@ static void ComputeAndMarkRelevanceForViewParallel(
 				View, 
 				ViewBit,
 				ViewData,
-				OutPreRenderViewMasks,
 				OutHasDynamicMeshElementsMasks,
 				OutHasDynamicEditorMeshElementsMasks,
 				MarkMasks);
@@ -1273,7 +1252,6 @@ static void ComputeAndMarkRelevanceForViewParallel(
 							View, 
 							ViewBit,
 							ViewData,
-							OutPreRenderViewMasks,
 							OutHasDynamicMeshElementsMasks,
 							OutHasDynamicEditorMeshElementsMasks,
 							MarkMasks);
@@ -1488,34 +1466,6 @@ static void MarkRelevantStaticMeshesForView(
 	View.NumVisibleStaticMeshElements += NumVisibleStaticMeshElements;
 }
 
-/**
- * Dispatches PreRenderView callbacks to all primitives with the provided view
- * masks.
- */
-static void DispatchPreRenderView(
-	const FScene* Scene,
-	const FSceneViewFamily* ViewFamily,
-	const FPrimitiveViewMasks& PreRenderViewMasks,
-	uint32 FrameNumber
-	)
-{
-	SCOPE_CYCLE_COUNTER(STAT_PreRenderView);
-
-	int32 NumPrimitives = Scene->Primitives.Num();
-	check(PreRenderViewMasks.Num() == NumPrimitives);
-
-	for (int32 PrimitiveIndex = 0; PrimitiveIndex < NumPrimitives; ++PrimitiveIndex)
-	{
-		uint8 ViewMask = PreRenderViewMasks[PrimitiveIndex];
-		if (ViewMask != 0)
-		{
-			FPrimitiveSceneInfo* PrimitiveSceneInfo = Scene->Primitives[PrimitiveIndex];
-			FScopeCycleCounter Context(PrimitiveSceneInfo->Proxy->GetStatId());
-			PrimitiveSceneInfo->Proxy->PreRenderView(ViewFamily, ViewMask, FrameNumber);
-		}
-	}
-}
-
 /**
  * Helper for InitViews to detect large camera movement, in both angle and position.
  */
@@ -1856,11 +1806,6 @@ void FSceneRenderer::ComputeViewVisibility(FRHICommandListImmediate& RHICmdList)
 	int32 NumPrimitives = Scene->Primitives.Num();
 	float CurrentRealTime = ViewFamily.CurrentRealTime;
 
-	// This array contains a per-primitive mask specifying for which views a
-	// primitive must have PreRenderView called.
-	FPrimitiveViewMasks PreRenderViewMasks;
-	PreRenderViewMasks.AddZeroed(NumPrimitives);
-
 	FPrimitiveViewMasks HasDynamicMeshElementsMasks;
 	HasDynamicMeshElementsMasks.AddZeroed(NumPrimitives);
 
@@ -2037,13 +1982,13 @@ void FSceneRenderer::ComputeViewVisibility(FRHICommandListImmediate& RHICmdList)
 			// Compute view relevance for all visible primitives.
 			if (!CVarLegacySingleThreadedRelevance.GetValueOnRenderThread())
 			{
-				ComputeAndMarkRelevanceForViewParallel(RHICmdList, Scene, View, ViewBit, PreRenderViewMasks, HasDynamicMeshElementsMasks, HasDynamicEditorMeshElementsMasks);
+				ComputeAndMarkRelevanceForViewParallel(RHICmdList, Scene, View, ViewBit, HasDynamicMeshElementsMasks, HasDynamicEditorMeshElementsMasks);
 			}
 			else
 			{
 				// This array contains a list of relevant static primities.
 				FRelevantStaticPrimitives RelevantStaticPrimitives;
-				ComputeRelevanceForView(RHICmdList, Scene, View, ViewBit, RelevantStaticPrimitives, PreRenderViewMasks, HasDynamicMeshElementsMasks, HasDynamicEditorMeshElementsMasks);
+				ComputeRelevanceForView(RHICmdList, Scene, View, ViewBit, RelevantStaticPrimitives, HasDynamicMeshElementsMasks, HasDynamicEditorMeshElementsMasks);
 				MarkRelevantStaticMeshesForView(Scene, View, RelevantStaticPrimitives);
 			}
 		}
@@ -2071,16 +2016,7 @@ void FSceneRenderer::ComputeViewVisibility(FRHICommandListImmediate& RHICmdList)
 #endif
 	}
 
-	const bool bUseGetMeshElements = ShouldUseGetDynamicMeshElements();
-
-	if (bUseGetMeshElements)
-	{
-		GatherDynamicMeshElements(Views, Scene, ViewFamily, HasDynamicMeshElementsMasks, HasDynamicEditorMeshElementsMasks, MeshCollector);
-	}
-	else
-	{
-		DispatchPreRenderView(Scene, &ViewFamily, PreRenderViewMasks, FrameNumber);
-	}
+	GatherDynamicMeshElements(Views, Scene, ViewFamily, HasDynamicMeshElementsMasks, HasDynamicEditorMeshElementsMasks, MeshCollector);
 
 	INC_DWORD_STAT_BY(STAT_ProcessedPrimitives,NumProcessedPrimitives);
 	INC_DWORD_STAT_BY(STAT_CulledPrimitives,NumCulledPrimitives);
diff --git a/Engine/Source/Runtime/Renderer/Private/ShadowRendering.cpp b/Engine/Source/Runtime/Renderer/Private/ShadowRendering.cpp
--- a/Engine/Source/Runtime/Renderer/Private/ShadowRendering.cpp
+++ b/Engine/Source/Runtime/Renderer/Private/ShadowRendering.cpp
@@ -1296,54 +1296,13 @@ void FProjectedShadowInfo::RenderDepthDynamic(FRHICommandList& RHICmdList, FScen
 	// Draw the subject's dynamic elements.
 	SCOPE_CYCLE_COUNTER(STAT_WholeSceneDynamicShadowDepthsTime);
 
-	const bool bUseGetMeshElements = ShouldUseGetDynamicMeshElements();
+	FShadowDepthDrawingPolicyFactory::ContextType Context(this);
 
-	if (bUseGetMeshElements)
+	for (int32 MeshBatchIndex = 0; MeshBatchIndex < DynamicSubjectMeshElements.Num(); MeshBatchIndex++)
 	{
-		FShadowDepthDrawingPolicyFactory::ContextType Context(this);
-
-		for (int32 MeshBatchIndex = 0; MeshBatchIndex < DynamicSubjectMeshElements.Num(); MeshBatchIndex++)
-		{
-			const FMeshBatchAndRelevance& MeshBatchAndRelevance = DynamicSubjectMeshElements[MeshBatchIndex];
-			const FMeshBatch& MeshBatch = *MeshBatchAndRelevance.Mesh;
-			FShadowDepthDrawingPolicyFactory::DrawDynamicMesh(RHICmdList, *FoundView, Context, MeshBatch, false, true, MeshBatchAndRelevance.PrimitiveSceneProxy, MeshBatch.BatchHitProxyId);
-		}
-	}
-	else
-	{
-		uint32 DrawPrimitiveFlags = 0;
-		if (bReflectiveShadowmap)
-		{
-			// force lowest LOD for RSMs
-			DrawPrimitiveFlags = EDrawDynamicFlags::ForceLowestLOD;
-		}
-
-		TDynamicPrimitiveDrawer<FShadowDepthDrawingPolicyFactory> Drawer(RHICmdList, FoundView, FShadowDepthDrawingPolicyFactory::ContextType(this), true);
-		uint32 PrimitiveCount = SubjectPrimitives.Num();
-
-		for (uint32 PrimitiveIndex = 0; PrimitiveIndex < PrimitiveCount; ++PrimitiveIndex)
-		{
-			const FPrimitiveSceneInfo* PrimitiveSceneInfo = SubjectPrimitives[PrimitiveIndex];
-			const FPrimitiveSceneProxy* PrimitiveSceneProxy = PrimitiveSceneInfo->Proxy;
-
-			// Lookup the primitive's cached view relevance
-			FPrimitiveViewRelevance ViewRelevance = FoundView->PrimitiveViewRelevanceMap[PrimitiveSceneInfo->GetIndex()];
-
-			if (!ViewRelevance.bInitializedThisFrame)
-			{
-				// Compute the subject primitive's view relevance since it wasn't cached
-				ViewRelevance = PrimitiveSceneInfo->Proxy->GetViewRelevance(FoundView);
-			}
-
-			// Only draw if the subject primitive is shadow relevant.
-			if (ViewRelevance.bShadowRelevance)
-			{
-				FScopeCycleCounter Context(PrimitiveSceneProxy->GetStatId());
-				Drawer.SetPrimitive(PrimitiveSceneProxy);
-
-				SubjectPrimitives[PrimitiveIndex]->Proxy->DrawDynamicElements(&Drawer, FoundView, DrawPrimitiveFlags);
-			}
-		}
+		const FMeshBatchAndRelevance& MeshBatchAndRelevance = DynamicSubjectMeshElements[MeshBatchIndex];
+		const FMeshBatch& MeshBatch = *MeshBatchAndRelevance.Mesh;
+		FShadowDepthDrawingPolicyFactory::DrawDynamicMesh(RHICmdList, *FoundView, Context, MeshBatch, false, true, MeshBatchAndRelevance.PrimitiveSceneProxy, MeshBatch.BatchHitProxyId);
 	}
 }
 
@@ -1879,22 +1838,16 @@ void FProjectedShadowInfo::RenderProjection(FRHICommandListImmediate& RHICmdList
 			0xff,0xff
 			>::GetRHI(), 1);
 
-		// Draw the receiver's dynamic elements.
-		TDynamicPrimitiveDrawer<FDepthDrawingPolicyFactory> Drawer(RHICmdList, View, FDepthDrawingPolicyFactory::ContextType(DDM_AllOccluders), true);
-		const bool bUseGetMeshElements = ShouldUseGetDynamicMeshElements();
 
-		if (bUseGetMeshElements)
-		{
-			const TArray<FMeshBatchAndRelevance,SceneRenderingAllocator>& DynamicMeshElements = bSelfShadowOnly ? DynamicSubjectMeshElements : DynamicReceiverMeshElements;
+		const TArray<FMeshBatchAndRelevance,SceneRenderingAllocator>& DynamicMeshElements = bSelfShadowOnly ? DynamicSubjectMeshElements : DynamicReceiverMeshElements;
 
-			FDepthDrawingPolicyFactory::ContextType Context(DDM_AllOccluders);
+		FDepthDrawingPolicyFactory::ContextType Context(DDM_AllOccluders);
 
-			for (int32 MeshBatchIndex = 0; MeshBatchIndex < DynamicMeshElements.Num(); MeshBatchIndex++)
-			{
-				const FMeshBatchAndRelevance& MeshBatchAndRelevance = DynamicMeshElements[MeshBatchIndex];
-				const FMeshBatch& MeshBatch = *MeshBatchAndRelevance.Mesh;
-				FDepthDrawingPolicyFactory::DrawDynamicMesh(RHICmdList, *View, Context, MeshBatch, false, true, MeshBatchAndRelevance.PrimitiveSceneProxy, MeshBatch.BatchHitProxyId);
-			}
+		for (int32 MeshBatchIndex = 0; MeshBatchIndex < DynamicMeshElements.Num(); MeshBatchIndex++)
+		{
+			const FMeshBatchAndRelevance& MeshBatchAndRelevance = DynamicMeshElements[MeshBatchIndex];
+			const FMeshBatch& MeshBatch = *MeshBatchAndRelevance.Mesh;
+			FDepthDrawingPolicyFactory::DrawDynamicMesh(RHICmdList, *View, Context, MeshBatch, false, true, MeshBatchAndRelevance.PrimitiveSceneProxy, MeshBatch.BatchHitProxyId);
 		}
 
 		// bPreShadow: only receive shadow on specific objects
@@ -1908,33 +1861,25 @@ void FProjectedShadowInfo::RenderProjection(FRHICommandListImmediate& RHICmdList
 			if (View->PrimitiveVisibilityMap[ReceiverPrimitiveSceneInfo->GetIndex()])
 			{
 				const FPrimitiveViewRelevance& ViewRelevance = View->PrimitiveViewRelevanceMap[ReceiverPrimitiveSceneInfo->GetIndex()];
-				if (ViewRelevance.bRenderInMainPass)
+
+				if (ViewRelevance.bRenderInMainPass && ViewRelevance.bStaticRelevance)
 				{
-					if (!bUseGetMeshElements && ViewRelevance.bDynamicRelevance)
+					for (int32 StaticMeshIdx = 0; StaticMeshIdx < ReceiverPrimitiveSceneInfo->StaticMeshes.Num(); StaticMeshIdx++)
 					{
-						Drawer.SetPrimitive(ReceiverPrimitiveSceneInfo->Proxy);
-						ReceiverPrimitiveSceneInfo->Proxy->DrawDynamicElements(&Drawer, View);
-					}
+						const FStaticMesh& StaticMesh = ReceiverPrimitiveSceneInfo->StaticMeshes[StaticMeshIdx];
 
-					if (ViewRelevance.bStaticRelevance)
-					{
-						for (int32 StaticMeshIdx = 0; StaticMeshIdx < ReceiverPrimitiveSceneInfo->StaticMeshes.Num(); StaticMeshIdx++)
+						if (View->StaticMeshVisibilityMap[StaticMesh.Id])
 						{
-							const FStaticMesh& StaticMesh = ReceiverPrimitiveSceneInfo->StaticMeshes[StaticMeshIdx];
-
-							if (View->StaticMeshVisibilityMap[StaticMesh.Id])
-							{
-								FDepthDrawingPolicyFactory::DrawStaticMesh(
-									RHICmdList, 
-									*View,
-									FDepthDrawingPolicyFactory::ContextType(DDM_AllOccluders),
-									StaticMesh,
-									StaticMesh.Elements.Num() == 1 ? 1 : View->StaticMeshBatchVisibility[StaticMesh.Id],
-									true,
-									ReceiverPrimitiveSceneInfo->Proxy,
-									StaticMesh.BatchHitProxyId
-									);
-							}
+							FDepthDrawingPolicyFactory::DrawStaticMesh(
+								RHICmdList, 
+								*View,
+								FDepthDrawingPolicyFactory::ContextType(DDM_AllOccluders),
+								StaticMesh,
+								StaticMesh.Elements.Num() == 1 ? 1 : View->StaticMeshBatchVisibility[StaticMesh.Id],
+								true,
+								ReceiverPrimitiveSceneInfo->Proxy,
+								StaticMesh.BatchHitProxyId
+								);
 						}
 					}
 				}
diff --git a/Engine/Source/Runtime/Renderer/Private/ShadowRendering.h b/Engine/Source/Runtime/Renderer/Private/ShadowRendering.h
--- a/Engine/Source/Runtime/Renderer/Private/ShadowRendering.h
+++ b/Engine/Source/Runtime/Renderer/Private/ShadowRendering.h
@@ -428,11 +428,6 @@ public:
 		const FPrimitiveSceneProxy* PrimitiveSceneProxy,
 		FHitProxyId HitProxyId
 		);
-
-	static bool IsMaterialIgnored(const FMaterialRenderProxy* MaterialRenderProxy, ERHIFeatureLevel::Type InFeatureLevel)
-	{
-		return false;
-	}
 };
 
 /** 
diff --git a/Engine/Source/Runtime/Renderer/Private/TranslucentLighting.cpp b/Engine/Source/Runtime/Renderer/Private/TranslucentLighting.cpp
--- a/Engine/Source/Runtime/Renderer/Private/TranslucentLighting.cpp
+++ b/Engine/Source/Runtime/Renderer/Private/TranslucentLighting.cpp
@@ -479,11 +479,6 @@ public:
 
 		return bDirty;
 	}
-
-	static bool IsMaterialIgnored(const FMaterialRenderProxy* MaterialRenderProxy, ERHIFeatureLevel::Type InFeatureLevel)
-	{ 
-		return !IsTranslucentBlendMode(MaterialRenderProxy->GetMaterial(InFeatureLevel)->GetBlendMode());
-	}
 };
 
 /** Renders shadow maps for translucent primitives. */
@@ -567,24 +562,15 @@ void FProjectedShadowInfo::RenderTranslucencyDepths(FRHICommandList& RHICmdList,
 			CW_RGBA, BO_Add, BF_One, BF_One, BO_Add, BF_One, BF_One,
 			CW_RGBA, BO_Add, BF_One, BF_One, BO_Add, BF_One, BF_One>::GetRHI());
 
-		const bool bUseGetMeshElements = ShouldUseGetDynamicMeshElements();
-
 		FTranslucencyShadowDepthDrawingPolicyFactory::ContextType DrawingContext(this,bDirectionalLight);
 
-		if (bUseGetMeshElements)
+		for (int32 MeshBatchIndex = 0; MeshBatchIndex < DynamicSubjectTranslucentMeshElements.Num(); MeshBatchIndex++)
 		{
-			FShadowDepthDrawingPolicyFactory::ContextType Context(this);
-
-			for (int32 MeshBatchIndex = 0; MeshBatchIndex < DynamicSubjectTranslucentMeshElements.Num(); MeshBatchIndex++)
-			{
-				const FMeshBatchAndRelevance& MeshBatchAndRelevance = DynamicSubjectTranslucentMeshElements[MeshBatchIndex];
-				const FMeshBatch& MeshBatch = *MeshBatchAndRelevance.Mesh;
-				FTranslucencyShadowDepthDrawingPolicyFactory::DrawDynamicMesh(RHICmdList, *FoundView, DrawingContext, MeshBatch, false, true, MeshBatchAndRelevance.PrimitiveSceneProxy, MeshBatch.BatchHitProxyId);
-			}
+			const FMeshBatchAndRelevance& MeshBatchAndRelevance = DynamicSubjectTranslucentMeshElements[MeshBatchIndex];
+			const FMeshBatch& MeshBatch = *MeshBatchAndRelevance.Mesh;
+			FTranslucencyShadowDepthDrawingPolicyFactory::DrawDynamicMesh(RHICmdList, *FoundView, DrawingContext, MeshBatch, false, true, MeshBatchAndRelevance.PrimitiveSceneProxy, MeshBatch.BatchHitProxyId);
 		}
 
-		TDynamicPrimitiveDrawer<FTranslucencyShadowDepthDrawingPolicyFactory> OpacityDrawer(RHICmdList, FoundView, DrawingContext, true);
-
 		for (int32 PrimitiveIndex = 0; PrimitiveIndex < SubjectTranslucentPrimitives.Num(); PrimitiveIndex++)
 		{
 			const FPrimitiveSceneInfo* PrimitiveSceneInfo = SubjectTranslucentPrimitives[PrimitiveIndex];
@@ -597,27 +583,18 @@ void FProjectedShadowInfo::RenderTranslucencyDepths(FRHICommandList& RHICmdList,
 				ViewRelevance = PrimitiveSceneInfo->Proxy->GetViewRelevance(FoundView);
 			}
 
-			if(ViewRelevance.bDrawRelevance)
+			if (ViewRelevance.bDrawRelevance && ViewRelevance.bStaticRelevance)
 			{
-				if (!bUseGetMeshElements && ViewRelevance.bDynamicRelevance)
-				{
-					OpacityDrawer.SetPrimitive(PrimitiveSceneInfo->Proxy);
-					PrimitiveSceneInfo->Proxy->DrawDynamicElements(&OpacityDrawer, FoundView);
-				}
-
-				if (ViewRelevance.bStaticRelevance)
+				for (int32 MeshIndex = 0; MeshIndex < PrimitiveSceneInfo->StaticMeshes.Num(); MeshIndex++)
 				{
-					for (int32 MeshIndex = 0; MeshIndex < PrimitiveSceneInfo->StaticMeshes.Num(); MeshIndex++)
-					{
-						FTranslucencyShadowDepthDrawingPolicyFactory::DrawStaticMesh(
-							RHICmdList, 
-							*FoundView,
-							DrawingContext,
-							PrimitiveSceneInfo->StaticMeshes[MeshIndex],
-							true,
-							PrimitiveSceneInfo->Proxy,
-							FHitProxyId());
-					}
+					FTranslucencyShadowDepthDrawingPolicyFactory::DrawStaticMesh(
+						RHICmdList, 
+						*FoundView,
+						DrawingContext,
+						PrimitiveSceneInfo->StaticMeshes[MeshIndex],
+						true,
+						PrimitiveSceneInfo->Proxy,
+						FHitProxyId());
 				}
 			}
 		}
diff --git a/Engine/Source/Runtime/Renderer/Private/TranslucentRendering.cpp b/Engine/Source/Runtime/Renderer/Private/TranslucentRendering.cpp
--- a/Engine/Source/Runtime/Renderer/Private/TranslucentRendering.cpp
+++ b/Engine/Source/Runtime/Renderer/Private/TranslucentRendering.cpp
@@ -642,12 +642,7 @@ void FTranslucentPrimSet::DrawPrimitives(
 		}
 	}
 
-	const bool bUseGetMeshElements = ShouldUseGetDynamicMeshElements();
-
-	if (bUseGetMeshElements)
-	{
-		View.SimpleElementCollector.DrawBatchedElements(RHICmdList, View, FTexture2DRHIRef(), EBlendModeFilter::Translucent);
-	}
+	View.SimpleElementCollector.DrawBatchedElements(RHICmdList, View, FTexture2DRHIRef(), EBlendModeFilter::Translucent);
 }
 
 void FTranslucentPrimSet::RenderPrimitive(
@@ -663,40 +658,19 @@ void FTranslucentPrimSet::RenderPrimitive(
 
 	if(ViewRelevance.bDrawRelevance)
 	{
-		const bool bUseGetMeshElements = ShouldUseGetDynamicMeshElements();
-		
-		if (bUseGetMeshElements)
+		FTranslucencyDrawingPolicyFactory::ContextType Context(TranslucentSelfShadow, bSeparateTranslucencyPass);
+
+		//@todo parallelrendering - come up with a better way to filter these by primitive
+		for (int32 MeshBatchIndex = 0; MeshBatchIndex < View.DynamicMeshElements.Num(); MeshBatchIndex++)
 		{
-			FTranslucencyDrawingPolicyFactory::ContextType Context(TranslucentSelfShadow, bSeparateTranslucencyPass);
+			const FMeshBatchAndRelevance& MeshBatchAndRelevance = View.DynamicMeshElements[MeshBatchIndex];
 
-			//@todo parallelrendering - come up with a better way to filter these by primitive
-			for (int32 MeshBatchIndex = 0; MeshBatchIndex < View.DynamicMeshElements.Num(); MeshBatchIndex++)
+			if (MeshBatchAndRelevance.PrimitiveSceneProxy == PrimitiveSceneInfo->Proxy)
 			{
-				const FMeshBatchAndRelevance& MeshBatchAndRelevance = View.DynamicMeshElements[MeshBatchIndex];
-
-				if (MeshBatchAndRelevance.PrimitiveSceneProxy == PrimitiveSceneInfo->Proxy)
-				{
-					const FMeshBatch& MeshBatch = *MeshBatchAndRelevance.Mesh;
-					FTranslucencyDrawingPolicyFactory::DrawDynamicMesh(RHICmdList, View, Context, MeshBatch, false, false, MeshBatchAndRelevance.PrimitiveSceneProxy, MeshBatch.BatchHitProxyId);
-				}
+				const FMeshBatch& MeshBatch = *MeshBatchAndRelevance.Mesh;
+				FTranslucencyDrawingPolicyFactory::DrawDynamicMesh(RHICmdList, View, Context, MeshBatch, false, false, MeshBatchAndRelevance.PrimitiveSceneProxy, MeshBatch.BatchHitProxyId);
 			}
 		}
-		// Render dynamic scene prim
-		else if( ViewRelevance.bDynamicRelevance )
-		{
-			TDynamicPrimitiveDrawer<FTranslucencyDrawingPolicyFactory> TranslucencyDrawer(
-				RHICmdList,
-				&View,
-				FTranslucencyDrawingPolicyFactory::ContextType(TranslucentSelfShadow, bSeparateTranslucencyPass),
-				false
-				);
-
-			TranslucencyDrawer.SetPrimitive(PrimitiveSceneInfo->Proxy);
-			PrimitiveSceneInfo->Proxy->DrawDynamicElements(
-				&TranslucencyDrawer,
-				&View
-				);
-		}
 
 		// Render static scene prim
 		if( ViewRelevance.bStaticRelevance )
diff --git a/Engine/Source/Runtime/Renderer/Private/TranslucentRendering.h b/Engine/Source/Runtime/Renderer/Private/TranslucentRendering.h
--- a/Engine/Source/Runtime/Renderer/Private/TranslucentRendering.h
+++ b/Engine/Source/Runtime/Renderer/Private/TranslucentRendering.h
@@ -57,12 +57,6 @@ public:
 		FHitProxyId HitProxyId
 		);
 
-	static bool IsMaterialIgnored(const FMaterialRenderProxy* MaterialRenderProxy, ERHIFeatureLevel::Type InFeatureLevel)
-	{
-		// Note: blend mode does not depend on the feature level.
-		return !IsTranslucentBlendMode(MaterialRenderProxy->GetMaterial(InFeatureLevel)->GetBlendMode());
-	}
-
 	/**
 	* Resolves the scene color target and copies it for use as a source texture.
 	*/
@@ -128,12 +122,6 @@ public:
 		const FPrimitiveSceneProxy* PrimitiveSceneProxy,
 		FHitProxyId HitProxyId
 		);
-
-	static bool IsMaterialIgnored(const FMaterialRenderProxy* MaterialRenderProxy, ERHIFeatureLevel::Type InFeatureLevel)
-	{
-		// Note: blend mode does not depend on the feature level.
-		return !IsTranslucentBlendMode(MaterialRenderProxy->GetMaterial(InFeatureLevel)->GetBlendMode());
-	}
 };
 
 /** Represents a subregion of a volume texture. */
diff --git a/Engine/Source/Runtime/Renderer/Private/VelocityRendering.cpp b/Engine/Source/Runtime/Renderer/Private/VelocityRendering.cpp
--- a/Engine/Source/Runtime/Renderer/Private/VelocityRendering.cpp
+++ b/Engine/Source/Runtime/Renderer/Private/VelocityRendering.cpp
@@ -451,28 +451,6 @@ bool IsMotionBlurEnabled(const FViewInfo& View)
 		&& !(View.Family->Views.Num() > 1);
 }
 
-void FDeferredShadingSceneRenderer::RenderDynamicVelocitiesInner(FRHICommandList& RHICmdList, const FViewInfo& View, int32 FirstIndex, int32 LastIndex)
-{
-	// Draw velocities for movable dynamic meshes.
-	TDynamicPrimitiveDrawer<FVelocityDrawingPolicyFactory> Drawer(
-		RHICmdList, &View, FVelocityDrawingPolicyFactory::ContextType(DDM_AllOccluders), true, false, true
-		);
-	for (int32 PrimitiveIndex = FirstIndex; PrimitiveIndex <= LastIndex; PrimitiveIndex++)
-	{
-		const FPrimitiveSceneInfo* PrimitiveSceneInfo = View.VisibleDynamicPrimitives[PrimitiveIndex];
-
-		if (!PrimitiveSceneInfo->ShouldRenderVelocity(View))
-		{
-			continue;
-		}
-
-		FScopeCycleCounter Context(PrimitiveSceneInfo->Proxy->GetStatId());
-
-		Drawer.SetPrimitive(PrimitiveSceneInfo->Proxy);
-		PrimitiveSceneInfo->Proxy->DrawDynamicElements(&Drawer, &View);
-	}
-}
-
 void FDeferredShadingSceneRenderer::RenderDynamicVelocitiesMeshElementsInner(FRHICommandList& RHICmdList, const FViewInfo& View, int32 FirstIndex, int32 LastIndex)
 {
 	FVelocityDrawingPolicyFactory::ContextType Context(DDM_AllOccluders);
@@ -619,8 +597,6 @@ void FDeferredShadingSceneRenderer::RenderVelocitiesInnerParallel(FRHICommandLis
 
 void FDeferredShadingSceneRenderer::RenderVelocitiesInner(FRHICommandListImmediate& RHICmdList, TRefCountPtr<IPooledRenderTarget>& VelocityRT)
 {
-	const bool bUseGetMeshElements = ShouldUseGetDynamicMeshElements();
-
 	for(int32 ViewIndex = 0; ViewIndex < Views.Num(); ViewIndex++)
 	{
 		const FViewInfo& View = Views[ViewIndex];
@@ -629,14 +605,7 @@ void FDeferredShadingSceneRenderer::RenderVelocitiesInner(FRHICommandListImmedia
 		// Draw velocities for movable static meshes.
 		Scene->VelocityDrawList.DrawVisible(RHICmdList, View, View.StaticMeshVelocityMap, View.StaticMeshBatchVisibility);
 
-		if (bUseGetMeshElements)
-		{
-			RenderDynamicVelocitiesMeshElementsInner(RHICmdList, View, 0, View.DynamicMeshElements.Num() - 1);
-		}
-		else
-		{
-			RenderDynamicVelocitiesInner(RHICmdList, View, 0, View.VisibleDynamicPrimitives.Num() - 1);
-		}
+		RenderDynamicVelocitiesMeshElementsInner(RHICmdList, View, 0, View.DynamicMeshElements.Num() - 1);
 	}
 }
 
