UE4 - Metal - Remove geometry shaders

diff --git a/Engine/Source/Runtime/Renderer/Private/AtmosphereRendering.cpp b/Engine/Source/Runtime/Renderer/Private/AtmosphereRendering.cpp
--- a/Engine/Source/Runtime/Renderer/Private/AtmosphereRendering.cpp
+++ b/Engine/Source/Runtime/Renderer/Private/AtmosphereRendering.cpp
@@ -605,7 +605,7 @@ class FAtmosphereGS : public FGlobalShader
 
 	static bool ShouldCache(EShaderPlatform Platform) 
 	{ 
-		return IsFeatureLevelSupported(Platform, ERHIFeatureLevel::SM4); 
+		return IsFeatureLevelSupported(Platform, ERHIFeatureLevel::SM4) && RHISupportsGeometryShaders(Platform);
 	}
 
 	FAtmosphereGS() {}
diff --git a/Engine/Source/Runtime/Renderer/Private/LightPropagationVolumeVisualisation.cpp b/Engine/Source/Runtime/Renderer/Private/LightPropagationVolumeVisualisation.cpp
--- a/Engine/Source/Runtime/Renderer/Private/LightPropagationVolumeVisualisation.cpp
+++ b/Engine/Source/Runtime/Renderer/Private/LightPropagationVolumeVisualisation.cpp
@@ -53,7 +53,10 @@ public:
 	virtual bool Serialize( FArchive& Ar ) override																					{ return FLpvVisualiseBase::Serialize( Ar ); }
 
 	//@todo-rco: Remove this when reenabling for OpenGL
-	static bool ShouldCache( EShaderPlatform Platform )		{ return IsFeatureLevelSupported(Platform, ERHIFeatureLevel::SM5) && !IsOpenGLPlatform(Platform); }
+	static bool ShouldCache( EShaderPlatform Platform )
+	{
+		return IsFeatureLevelSupported(Platform, ERHIFeatureLevel::SM5) && RHISupportsGeometryShaders(Platform) && !IsOpenGLPlatform(Platform);
+	}
 	static void ModifyCompilationEnvironment( EShaderPlatform Platform, FShaderCompilerEnvironment& OutEnvironment )				{ FLpvVisualiseBase::ModifyCompilationEnvironment( Platform, OutEnvironment ); }
 
 	void SetParameters(
diff --git a/Engine/Source/Runtime/Renderer/Private/ShadowRendering.cpp b/Engine/Source/Runtime/Renderer/Private/ShadowRendering.cpp
--- a/Engine/Source/Runtime/Renderer/Private/ShadowRendering.cpp
+++ b/Engine/Source/Runtime/Renderer/Private/ShadowRendering.cpp
@@ -297,7 +297,7 @@ public:
 
 	static bool ShouldCache(EShaderPlatform Platform,const FMaterial* Material,const FVertexFactoryType* VertexFactoryType)
 	{
-		return TShadowDepthVS<VertexShadowDepth_OnePassPointLight, false, false>::ShouldCache(Platform, Material, VertexFactoryType);
+		return TShadowDepthVS<VertexShadowDepth_OnePassPointLight, false, false>::ShouldCache(Platform, Material, VertexFactoryType) && RHISupportsGeometryShaders(Platform);
 	}
 
 	static void ModifyCompilationEnvironment(EShaderPlatform Platform, const FMaterial* Material, FShaderCompilerEnvironment& OutEnvironment)
diff --git a/Engine/Source/Runtime/Renderer/Private/TranslucentRendering.h b/Engine/Source/Runtime/Renderer/Private/TranslucentRendering.h
--- a/Engine/Source/Runtime/Renderer/Private/TranslucentRendering.h
+++ b/Engine/Source/Runtime/Renderer/Private/TranslucentRendering.h
@@ -202,7 +202,7 @@ public:
 
 	static bool ShouldCache(EShaderPlatform Platform) 
 	{ 
-		return IsFeatureLevelSupported(Platform, ERHIFeatureLevel::SM4); 
+		return IsFeatureLevelSupported(Platform, ERHIFeatureLevel::SM4) && RHISupportsGeometryShaders(Platform); 
 	}
 
 	FWriteToSliceGS(const ShaderMetaType::CompiledShaderInitializerType& Initializer):
