Allow ES2 post process chain to deal with view rectangle that doesn't encompass entire target.
Fixes issues caused by constrained aspect ratio cameras.
Fixes UE-5536

diff --git a/Engine/Source/Runtime/Renderer/Private/PostProcess/PostProcessMobile.cpp b/Engine/Source/Runtime/Renderer/Private/PostProcess/PostProcessMobile.cpp
--- a/Engine/Source/Runtime/Renderer/Private/PostProcess/PostProcessMobile.cpp
+++ b/Engine/Source/Runtime/Renderer/Private/PostProcess/PostProcessMobile.cpp
@@ -171,7 +171,7 @@ void FRCPassPostProcessBloomSetupES2::Process(FRenderingCompositePassContext& Co
 	SCOPED_DRAW_EVENT(Context.RHICmdList, PostProcessBloomSetup);
 
 	const FPooledRenderTargetDesc* InputDesc = GetInputDesc(ePId_Input0);
-
+	FIntPoint PrePostSourceViewportSize = PrePostSourceViewportRect.Size();
 	uint32 DstX = FMath::Max(1, PrePostSourceViewportSize.X/4);
 	uint32 DstY = FMath::Max(1, PrePostSourceViewportSize.Y/4);
 
@@ -193,7 +193,7 @@ void FRCPassPostProcessBloomSetupES2::Process(FRenderingCompositePassContext& Co
 		//	uint32 ScaleFactor = View.ViewRect.Width() / SrcSize.X;
 		//	SrcRect = View.ViewRect / ScaleFactor;
 		// TODO: This won't work with scaled views.
-		SrcRect = View.ViewRect;
+		SrcRect = PrePostSourceViewportRect;
 	}
 	else
 	{
@@ -243,8 +243,8 @@ FPooledRenderTargetDesc FRCPassPostProcessBloomSetupES2::ComputeOutputDesc(EPass
 	Ret.bForceSeparateTargetAndShaderResource = false;
 	Ret.Format = PF_FloatRGBA;
 	Ret.NumSamples = 1;
-	Ret.Extent.X = FMath::Max(1, PrePostSourceViewportSize.X/4);
-	Ret.Extent.Y = FMath::Max(1, PrePostSourceViewportSize.Y/4);
+	Ret.Extent.X = FMath::Max(1, PrePostSourceViewportRect.Width() / 4);
+	Ret.Extent.Y = FMath::Max(1, PrePostSourceViewportRect.Height() / 4);
 	Ret.DebugName = TEXT("BloomSetup");
 	return Ret;
 }
@@ -1955,7 +1955,7 @@ void FRCPassPostProcessDofDownES2::Process(FRenderingCompositePassContext& Conte
 	SCOPED_DRAW_EVENT(Context.RHICmdList, PostProcessDofDown);
 
 	const FPooledRenderTargetDesc* InputDesc = GetInputDesc(ePId_Input0);
-
+	FIntPoint PrePostSourceViewportSize = PrePostSourceViewportRect.Size();
 	uint32 DstX = FMath::Max(1, PrePostSourceViewportSize.X/2);
 	uint32 DstY = FMath::Max(1, PrePostSourceViewportSize.Y/2);
 
@@ -1977,7 +1977,7 @@ void FRCPassPostProcessDofDownES2::Process(FRenderingCompositePassContext& Conte
 		//	uint32 ScaleFactor = View.ViewRect.Width() / SrcSize.X;
 		//	SrcRect = View.ViewRect / ScaleFactor;
 		// TODO: This won't work with scaled views.
-		SrcRect = View.ViewRect;
+		SrcRect = PrePostSourceViewportRect;
 	}
 	else
 	{
@@ -2027,6 +2027,7 @@ FPooledRenderTargetDesc FRCPassPostProcessDofDownES2::ComputeOutputDesc(EPassOut
 	Ret.bForceSeparateTargetAndShaderResource = false;
 	Ret.Format = PF_FloatRGBA;
 	Ret.NumSamples = 1;
+	FIntPoint PrePostSourceViewportSize = PrePostSourceViewportRect.Size();
 	Ret.Extent.X = FMath::Max(1, PrePostSourceViewportSize.X/2);
 	Ret.Extent.Y = FMath::Max(1, PrePostSourceViewportSize.Y/2);
 	Ret.DebugName = TEXT("DofDown");
@@ -2739,24 +2740,20 @@ void FRCPassPostProcessAaES2::Process(FRenderingCompositePassContext& Context)
 	SCOPED_DRAW_EVENT(Context.RHICmdList, PostProcessAa);
 
 	const FPooledRenderTargetDesc* InputDesc = GetInputDesc(ePId_Input0);
-	FIntPoint SrcSize = InputDesc->Extent;
+	const FSceneRenderTargetItem& DestRenderTarget = PassOutputs[0].RequestSurface(Context);
+	const FPooledRenderTargetDesc& OutputDesc = PassOutputs[0].RenderTargetDesc;
 
-	uint32 DstX = FMath::Max(1, PrePostSourceViewportSize.X);
-	uint32 DstY = FMath::Max(1, PrePostSourceViewportSize.Y);
+	const FIntPoint& SrcSize = InputDesc->Extent;
+	const FIntPoint& DestSize = OutputDesc.Extent;
 
-	FIntRect DstRect;
-	DstRect.Min.X = 0;
-	DstRect.Min.Y = 0;
-	DstRect.Max.X = DstX;
-	DstRect.Max.Y = DstY;
+	check(SrcSize == DestSize);
 
-	const FSceneRenderTargetItem& DestRenderTarget = PassOutputs[0].RequestSurface(Context);
 	SetRenderTarget(Context.RHICmdList, DestRenderTarget.TargetableTexture, FTextureRHIRef());	
 
 	// is optimized away if possible (RT size=view size, )
 	Context.RHICmdList.Clear(true, FLinearColor::Black, false, 1.0f, false, 0, FIntRect());
 
-	Context.SetViewportAndCallRHI(0, 0, 0.0f, DstX, DstY, 1.0f );
+	Context.SetViewportAndCallRHI(0, 0, 0.0f, DestSize.X, DestSize.Y, 1.0f);
 
 	Context.RHICmdList.SetBlendState(TStaticBlendState<>::GetRHI());
 	Context.RHICmdList.SetRasterizerState(TStaticRasterizerState<>::GetRHI());
@@ -2764,16 +2761,21 @@ void FRCPassPostProcessAaES2::Process(FRenderingCompositePassContext& Context)
 
 	SetShader(Context);
 
-	FIntPoint SrcDstSize = PrePostSourceViewportSize;
+	const FIntRect& ViewRect = Context.View.ViewRect;
+	float XPos = ViewRect.Min.X;
+	float YPos = ViewRect.Min.Y;
+	float Width = ViewRect.Width();
+	float Height = ViewRect.Height();
+
 	TShaderMapRef<FPostProcessAaVS_ES2> VertexShader(Context.GetShaderMap());
 
 	DrawRectangle(
 		Context.RHICmdList,
-		0, 0,
-		DstX, DstY,
-		0, 0,
-		DstX, DstY,
-		SrcDstSize,
+		XPos, YPos,
+		Width, Height,
+		XPos, YPos,
+		Width, Height,
+		DestSize,
 		SrcSize,
 		*VertexShader,
 		EDRF_UseTriangleOptimization);
@@ -2792,9 +2794,8 @@ FPooledRenderTargetDesc FRCPassPostProcessAaES2::ComputeOutputDesc(EPassOutputId
 	Ret.bForceSeparateTargetAndShaderResource = false;
 	Ret.Format = PF_B8G8R8A8;
 	Ret.NumSamples = 1;
-	Ret.Extent.X = FMath::Max(1, PrePostSourceViewportSize.X);
-	Ret.Extent.Y = FMath::Max(1, PrePostSourceViewportSize.Y);
 	Ret.DebugName = TEXT("Aa");
+	Ret.Extent = PassInputs[0].GetOutput()->RenderTargetDesc.Extent;
 	return Ret;
 }
 
diff --git a/Engine/Source/Runtime/Renderer/Private/PostProcess/PostProcessMobile.h b/Engine/Source/Runtime/Renderer/Private/PostProcess/PostProcessMobile.h
--- a/Engine/Source/Runtime/Renderer/Private/PostProcess/PostProcessMobile.h
+++ b/Engine/Source/Runtime/Renderer/Private/PostProcess/PostProcessMobile.h
@@ -11,12 +11,12 @@
 class FRCPassPostProcessBloomSetupES2 : public TRenderingCompositePassBase<1, 1>
 {
 public:
-	FRCPassPostProcessBloomSetupES2(FIntPoint InPrePostSourceViewportSize, bool bInUseViewRectSource) : PrePostSourceViewportSize(InPrePostSourceViewportSize), bUseViewRectSource(bInUseViewRectSource) { }
+	FRCPassPostProcessBloomSetupES2(FIntRect InPrePostSourceViewportRect, bool bInUseViewRectSource) : PrePostSourceViewportRect(InPrePostSourceViewportRect), bUseViewRectSource(bInUseViewRectSource) { }
 	virtual void Process(FRenderingCompositePassContext& Context);
 	virtual FPooledRenderTargetDesc ComputeOutputDesc(EPassOutputId InPassOutputId) const;
 	virtual void Release() override { delete this; }
 private:
-	FIntPoint PrePostSourceViewportSize;
+	FIntRect PrePostSourceViewportRect;
 	bool bUseViewRectSource;
 	void SetShader(const FRenderingCompositePassContext& Context);
 };
@@ -49,12 +49,12 @@ private:
 class FRCPassPostProcessDofDownES2 : public TRenderingCompositePassBase<2, 1>
 {
 public:
-	FRCPassPostProcessDofDownES2(FIntPoint InPrePostSourceViewportSize, bool bInUseViewRectSource) : PrePostSourceViewportSize(InPrePostSourceViewportSize), bUseViewRectSource(bInUseViewRectSource) { }
+	FRCPassPostProcessDofDownES2(FIntRect InPrePostSourceViewportRect, bool bInUseViewRectSource) : PrePostSourceViewportRect(InPrePostSourceViewportRect), bUseViewRectSource(bInUseViewRectSource) { }
 	virtual void Process(FRenderingCompositePassContext& Context);
 	virtual FPooledRenderTargetDesc ComputeOutputDesc(EPassOutputId InPassOutputId) const;
 	virtual void Release() { delete this; }
 private:
-	FIntPoint PrePostSourceViewportSize;
+	FIntRect PrePostSourceViewportRect;
 	bool bUseViewRectSource;
 	void SetShader(const FRenderingCompositePassContext& Context);
 };
@@ -171,12 +171,10 @@ private:
 class FRCPassPostProcessAaES2 : public TRenderingCompositePassBase<2, 1>
 {
 public:
-	FRCPassPostProcessAaES2(FIntPoint InPrePostSourceViewportSize) : PrePostSourceViewportSize(InPrePostSourceViewportSize) { }
 	virtual void Process(FRenderingCompositePassContext& Context);
 	virtual FPooledRenderTargetDesc ComputeOutputDesc(EPassOutputId InPassOutputId) const;
 	virtual void Release() { delete this; }
 private:
-	FIntPoint PrePostSourceViewportSize;
 	void SetShader(const FRenderingCompositePassContext& Context);
 };
 
diff --git a/Engine/Source/Runtime/Renderer/Private/PostProcess/PostProcessTonemap.cpp b/Engine/Source/Runtime/Renderer/Private/PostProcess/PostProcessTonemap.cpp
--- a/Engine/Source/Runtime/Renderer/Private/PostProcess/PostProcessTonemap.cpp
+++ b/Engine/Source/Runtime/Renderer/Private/PostProcess/PostProcessTonemap.cpp
@@ -1582,8 +1582,11 @@ static void SetShaderTemplES2(const FRenderingCompositePassContext& Context, boo
 }
 
 
-FRCPassPostProcessTonemapES2::FRCPassPostProcessTonemapES2(const FViewInfo& View, bool bInUsedFramebufferFetch) 
-	: bUsedFramebufferFetch(bInUsedFramebufferFetch)
+FRCPassPostProcessTonemapES2::FRCPassPostProcessTonemapES2(const FViewInfo& View, FIntRect InViewRect, FIntPoint InDestSize, bool bInUsedFramebufferFetch) 
+	:
+	ViewRect(InViewRect),
+	DestSize(InDestSize),
+	bUsedFramebufferFetch(bInUsedFramebufferFetch)
 {
 	uint32 ConfigBitmask = TonemapperGenerateBitmaskMobile(&View, false);
 	ConfigIndexMobile = TonemapperFindLeastExpensive(TonemapperConfBitmaskMobile, sizeof(TonemapperConfBitmaskMobile)/4, TonemapperCostTab, ConfigBitmask);
@@ -1604,11 +1607,12 @@ void FRCPassPostProcessTonemapES2::Process(FRenderingCompositePassContext& Conte
 	const FSceneView& View = Context.View;
 	const FSceneViewFamily& ViewFamily = *(View.Family);
 	const FSceneRenderTargetItem& DestRenderTarget = PassOutputs[0].RequestSurface(Context);
+	const FPooledRenderTargetDesc& OutputDesc = PassOutputs[0].RenderTargetDesc;
 
-	FIntRect SrcRect = View.ViewRect;
+	FIntRect SrcRect = ViewRect;
 	FIntRect DestRect = View.ViewRect;
 	FIntPoint SrcSize = InputDesc->Extent;
-	FIntPoint DstSize = DestRect.Size();
+	FIntPoint DstSize = OutputDesc.Extent;
 
 	// Set the view family's render target/viewport.
 	SetRenderTarget(Context.RHICmdList, DestRenderTarget.TargetableTexture, FTextureRHIParamRef());
@@ -1616,7 +1620,7 @@ void FRCPassPostProcessTonemapES2::Process(FRenderingCompositePassContext& Conte
 	// Full clear to avoid restore
 	Context.RHICmdList.Clear(true, FLinearColor::Black, false, 1.0f, false, 0, FIntRect());
 
-	Context.SetViewportAndCallRHI(View.ViewRect);
+	Context.SetViewportAndCallRHI(DestRect);
 
 	// set the state
 	Context.RHICmdList.SetBlendState(TStaticBlendState<>::GetRHI());
@@ -1674,9 +1678,9 @@ void FRCPassPostProcessTonemapES2::Process(FRenderingCompositePassContext& Conte
 	DrawRectangle(
 		Context.RHICmdList,
 		0, 0,
-		View.ViewRect.Width(), View.ViewRect.Height(),
-		View.ViewRect.Min.X, View.ViewRect.Min.Y, 
-		View.ViewRect.Width(), View.ViewRect.Height(),
+		DstSize.X, DstSize.Y,
+		SrcRect.Min.X, SrcRect.Min.Y,
+		SrcRect.Width(), SrcRect.Height(),
 		DstSize,
 		SrcSize,
 		*VertexShader,
@@ -1702,6 +1706,6 @@ FPooledRenderTargetDesc FRCPassPostProcessTonemapES2::ComputeOutputDesc(EPassOut
 	Ret.Reset();
 	Ret.Format = PF_B8G8R8A8;
 	Ret.DebugName = TEXT("Tonemap");
-
+	Ret.Extent = DestSize;
 	return Ret;
 }
diff --git a/Engine/Source/Runtime/Renderer/Private/PostProcess/PostProcessTonemap.h b/Engine/Source/Runtime/Renderer/Private/PostProcess/PostProcessTonemap.h
--- a/Engine/Source/Runtime/Renderer/Private/PostProcess/PostProcessTonemap.h
+++ b/Engine/Source/Runtime/Renderer/Private/PostProcess/PostProcessTonemap.h
@@ -66,7 +66,7 @@ private:
 class FRCPassPostProcessTonemapES2 : public TRenderingCompositePassBase<3, 1>
 {
 public:
-	FRCPassPostProcessTonemapES2(const FViewInfo& View, bool bInUsedFramebufferFetch);
+	FRCPassPostProcessTonemapES2(const FViewInfo& View, FIntRect InViewRect, FIntPoint InDestSize, bool bInUsedFramebufferFetch);
 
 	// interface FRenderingCompositePass ---------
 
@@ -75,6 +75,8 @@ public:
 	virtual FPooledRenderTargetDesc ComputeOutputDesc(EPassOutputId InPassOutputId) const;
 
 private:
+	FIntRect ViewRect;
+	FIntPoint DestSize;
 	bool bUsedFramebufferFetch;
 	// set in constructor
 	uint32 ConfigIndexMobile;
diff --git a/Engine/Source/Runtime/Renderer/Private/PostProcess/PostProcessing.cpp b/Engine/Source/Runtime/Renderer/Private/PostProcess/PostProcessing.cpp
--- a/Engine/Source/Runtime/Renderer/Private/PostProcess/PostProcessing.cpp
+++ b/Engine/Source/Runtime/Renderer/Private/PostProcess/PostProcessing.cpp
@@ -1423,6 +1423,8 @@ void FPostProcessing::ProcessES2(FRHICommandListImmediate& RHICmdList, FViewInfo
 			}
 		}
 
+		const FIntPoint FinalTargetSize = View.Family->RenderTarget->GetSizeXY();
+		FIntRect FinalOutputViewRect = View.ViewRect;
 		FIntPoint PrePostSourceViewportSize = View.ViewRect.Size();
 		// ES2 preview uses a subsection of the scene RT, bUsedFramebufferFetch == true deals with this case.  
 		bool bViewRectSource = bUsedFramebufferFetch || GSceneRenderTargets.GetBufferSizeXY() != PrePostSourceViewportSize;
@@ -1456,6 +1458,9 @@ void FPostProcessing::ProcessES2(FRHICommandListImmediate& RHICmdList, FViewInfo
 					FRenderingCompositePass* PostProcessSunMask = Context.Graph.RegisterPass(new(FMemStack::Get()) FRCPassPostProcessSunMaskES2(PrePostSourceViewportSize, false));
 					PostProcessSunMask->SetInput(ePId_Input0, Context.FinalOutput);
 					Context.FinalOutput = FRenderingCompositeOutputRef(PostProcessSunMask);
+					// Context.FinalOutput now contains entire image.
+					// (potentially clipped to content of View.ViewRect.)
+					FinalOutputViewRect = FIntRect(FIntPoint(0, 0), PrePostSourceViewportSize);
 				}
 
 				FRenderingCompositeOutputRef PostProcessBloomSetup;
@@ -1469,7 +1474,7 @@ void FPostProcessing::ProcessES2(FRHICommandListImmediate& RHICmdList, FViewInfo
 					}
 					else
 					{
-						FRenderingCompositePass* Pass = Context.Graph.RegisterPass(new(FMemStack::Get()) FRCPassPostProcessBloomSetupES2(PrePostSourceViewportSize, bViewRectSource));
+						FRenderingCompositePass* Pass = Context.Graph.RegisterPass(new(FMemStack::Get()) FRCPassPostProcessBloomSetupES2(FinalOutputViewRect, bViewRectSource));
 						Pass->SetInput(ePId_Input0, Context.FinalOutput);
 						PostProcessBloomSetup = FRenderingCompositeOutputRef(Pass);
 					}
@@ -1481,7 +1486,7 @@ void FPostProcessing::ProcessES2(FRHICommandListImmediate& RHICmdList, FViewInfo
 					// Samples at 1/16 area, writes to 1/16 area.
 					FRenderingCompositeOutputRef PostProcessNear;
 					{
-						FRenderingCompositePass* Pass = Context.Graph.RegisterPass(new(FMemStack::Get()) FRCPassPostProcessDofNearES2(PrePostSourceViewportSize));
+						FRenderingCompositePass* Pass = Context.Graph.RegisterPass(new(FMemStack::Get()) FRCPassPostProcessDofNearES2(FinalOutputViewRect.Size()));
 						Pass->SetInput(ePId_Input0, PostProcessBloomSetup);
 						PostProcessNear = FRenderingCompositeOutputRef(Pass);
 					}
@@ -1490,7 +1495,7 @@ void FPostProcessing::ProcessES2(FRHICommandListImmediate& RHICmdList, FViewInfo
 					// Samples at full resolution, writes to 1/4 area.
 					FRenderingCompositeOutputRef PostProcessDofDown;
 					{
-						FRenderingCompositePass* Pass = Context.Graph.RegisterPass(new(FMemStack::Get()) FRCPassPostProcessDofDownES2(PrePostSourceViewportSize, bViewRectSource));
+						FRenderingCompositePass* Pass = Context.Graph.RegisterPass(new(FMemStack::Get()) FRCPassPostProcessDofDownES2(FinalOutputViewRect, bViewRectSource));
 						Pass->SetInput(ePId_Input0, Context.FinalOutput);
 						Pass->SetInput(ePId_Input1, PostProcessNear);
 						PostProcessDofDown = FRenderingCompositeOutputRef(Pass);
@@ -1500,7 +1505,7 @@ void FPostProcessing::ProcessES2(FRHICommandListImmediate& RHICmdList, FViewInfo
 					// Samples at 1/4 area, writes to 1/4 area.
 					FRenderingCompositeOutputRef PostProcessDofBlur;
 					{
-						FRenderingCompositePass* Pass = Context.Graph.RegisterPass(new(FMemStack::Get()) FRCPassPostProcessDofBlurES2(PrePostSourceViewportSize));
+						FRenderingCompositePass* Pass = Context.Graph.RegisterPass(new(FMemStack::Get()) FRCPassPostProcessDofBlurES2(FinalOutputViewRect.Size()));
 						Pass->SetInput(ePId_Input0, PostProcessDofDown);
 						Pass->SetInput(ePId_Input1, PostProcessNear);
 						PostProcessDofBlur = FRenderingCompositeOutputRef(Pass);
@@ -1685,11 +1690,13 @@ void FPostProcessing::ProcessES2(FRHICommandListImmediate& RHICmdList, FViewInfo
 		}
 
 		// Must run to blit to back buffer even if post processing is off.
-		FRenderingCompositePass* PostProcessTonemap = Context.Graph.RegisterPass(new(FMemStack::Get()) FRCPassPostProcessTonemapES2(Context.View, bViewRectSource));
+		FRenderingCompositePass* PostProcessTonemap = Context.Graph.RegisterPass(new(FMemStack::Get()) FRCPassPostProcessTonemapES2(Context.View, FinalOutputViewRect, FinalTargetSize, bViewRectSource));
 		PostProcessTonemap->SetInput(ePId_Input0, Context.FinalOutput);
 		PostProcessTonemap->SetInput(ePId_Input1, BloomOutput);
 		PostProcessTonemap->SetInput(ePId_Input2, DofOutput);
 		Context.FinalOutput = FRenderingCompositeOutputRef(PostProcessTonemap);
+		// if Context.FinalOutput was the clipped result of sunmask stage then this stage also restores Context.FinalOutput back original target size.
+		FinalOutputViewRect = View.ViewRect;
 
 		if(bUseAa && View.Family->EngineShowFlags.PostProcessing)
 		{
@@ -1706,7 +1713,7 @@ void FPostProcessing::ProcessES2(FRHICommandListImmediate& RHICmdList, FViewInfo
 
 			// Mobile temporal AA is done after tonemapping.
 			FIntPoint PrePostSourceViewportSize = View.ViewRect.Size();
-			FRenderingCompositePass* PostProcessAa = Context.Graph.RegisterPass(new(FMemStack::Get()) FRCPassPostProcessAaES2(PrePostSourceViewportSize));
+			FRenderingCompositePass* PostProcessAa = Context.Graph.RegisterPass(new(FMemStack::Get()) FRCPassPostProcessAaES2());
 			PostProcessAa->SetInput(ePId_Input0, Context.FinalOutput);
 			PostProcessAa->SetInput(ePId_Input1, PostProcessPrior);
 			Context.FinalOutput = FRenderingCompositeOutputRef(PostProcessAa);
