Pixels without a distance field representation can still receive occlusion, with reasonable performance (surface cache sample radii are clamped)
Cone visibility is output from the cone tracing step instead of bent normal, allowing non-distance field based methods to accumulate their occlusion contribution
Landscape causes occlusion on other parts of the scene, using a locally-planar heuristic

diff --git a/Engine/Shaders/Common.usf b/Engine/Shaders/Common.usf
--- a/Engine/Shaders/Common.usf
+++ b/Engine/Shaders/Common.usf
@@ -372,6 +372,8 @@ SamplerState	LightAttenuationTextureSampler;
 // We don't use a function so we can avoid type promotion/ coercion.
 #define TangentBias(X) (X * 2.0f - 1.0f)
 
+#define TERRAIN_ZSCALE				(1.0f/128.0f)
+
 float Square( float x )
 {
 	return x*x;
diff --git a/Engine/Shaders/DistanceFieldSurfaceCacheLighting.usf b/Engine/Shaders/DistanceFieldSurfaceCacheLighting.usf
--- a/Engine/Shaders/DistanceFieldSurfaceCacheLighting.usf
+++ b/Engine/Shaders/DistanceFieldSurfaceCacheLighting.usf
@@ -143,7 +143,7 @@ void WriteDownsampledDepthPS(
 	out float OutDepth : SV_Depth)
 {
 	float2 DistanceFieldUVs = UVAndScreenPos.xy;
-	float Depth = Texture2DSampleLevel(DistanceFieldNormalTexture, DistanceFieldNormalSampler, DistanceFieldUVs, 0).w;
+	float Depth = abs(Texture2DSampleLevel(DistanceFieldNormalTexture, DistanceFieldNormalSampler, DistanceFieldUVs, 0).w);
 
 	OutColor = 0;
 	OutDepth = ConvertToDeviceZ(Depth);
@@ -189,7 +189,7 @@ void IrradianceCacheSplatVS(
 {
 	float4 PositionAndRadius = IrradianceCachePositionRadius.Load(JobIndex);
 	float OccluderRadius = IrradianceCacheOccluderRadius.Load(JobIndex);
-	PositionAndRadius.w = min(PositionAndRadius.w, OccluderRadius);
+	PositionAndRadius.w = min(abs(PositionAndRadius.w), OccluderRadius);
 
 	float ViewToCenterLength = length(PositionAndRadius.xyz - View.ViewOrigin.xyz);
 	float3 NormalizedViewToCenter = (PositionAndRadius.xyz - View.ViewOrigin.xyz) / ViewToCenterLength;
@@ -310,14 +310,14 @@ void IrradianceCacheSplatPS(
 
 		float PrecisionScale = .1f;
 		Weight = saturate(PrecisionScale * (1 - max(DistanceError, max(NormalError, PointBehindPlaneError))));
-		//float Weight = saturate(PrecisionScale * (1 - max(DistanceError, NormalError)));
-		//float Weight = saturate(PrecisionScale * (1 - DistanceError));
-		//float Weight = saturate(PrecisionScale * (1 - NormalError));
+		//Weight = saturate(PrecisionScale * (1 - DistanceError));
 
 		float VisualizePlacement = Distance < .1f * Input.PositionRadius.w;
 
 #if FINAL_INTERPOLATION_PASS
-		OutBentNormal.rgb = Input.BentNormal * Weight;
+		// Pixels without a distance field representation interpolate AO instead of a bent normal
+		float3 InterpolatedValue = DistanceFieldNormalValue.w < 0 ? length(Input.BentNormal).xxx : Input.BentNormal;
+		OutBentNormal.rgb = InterpolatedValue * Weight;
 		#if SUPPORT_IRRADIANCE
 			OutIrradiance.rgb = Input.Irradiance * Weight;
 		#endif
@@ -329,7 +329,7 @@ void IrradianceCacheSplatPS(
 	#if VISUALIZE_SPLAT_OVERDRAW && FINAL_INTERPOLATION_PASS
 		OutBentNormal.rgb = 0;
 		//OutBentNormal.a = Input.Position.w <= SceneW ? .01f : 0.0f;
-		OutBentNormal.a = .005f + (Weight > .001f ? .01f : 0);
+		OutBentNormal.a = .005f;// + (Weight > .001f ? .01f : 0);
 	#endif
 
 	#define VISUALIZE_RECORD_POINTS 0
@@ -339,301 +339,6 @@ void IrradianceCacheSplatPS(
 	#endif
 }
 
-
-#ifndef NUM_INTERPOLATE_LEVELS
-#define NUM_INTERPOLATE_LEVELS 1
-#endif
-
-RWTexture2D<float4> RWIrradianceCacheSplat;
-
-Buffer<uint> LevelScatterDrawParameters[NUM_INTERPOLATE_LEVELS];
-Buffer<float4> LevelIrradianceCachePositionRadius[NUM_INTERPOLATE_LEVELS];
-Buffer<float> LevelIrradianceCacheOccluderRadius[NUM_INTERPOLATE_LEVELS];
-Buffer<float4> LevelIrradianceCacheNormal[NUM_INTERPOLATE_LEVELS];
-Buffer<float4> LevelIrradianceCacheBentNormal[NUM_INTERPOLATE_LEVELS];
-
-#define MAX_RECORDS_PER_TILE 512
-
-groupshared uint IntegerTileMinZ;
-groupshared uint IntegerTileMaxZ;
-groupshared uint IntegerTileMinZ2;
-groupshared uint IntegerTileMaxZ2;
-groupshared uint TileNumSamples;
-groupshared uint TileSampleIndices[MAX_RECORDS_PER_TILE];
-
-[numthreads(THREADGROUP_SIZEX, THREADGROUP_SIZEY, 1)]
-void InterpolateCacheCS(
-	uint3 GroupId : SV_GroupID,
-	uint3 DispatchThreadId : SV_DispatchThreadID,
-    uint3 GroupThreadId : SV_GroupThreadID) 
-{
-	uint ThreadIndex = GroupThreadId.y * THREADGROUP_SIZEX + GroupThreadId.x;
-	uint2 PixelIndex = DispatchThreadId.xy;
-
-	float2 ScreenUV = (PixelIndex * CurrentLevelDownsampleFactor + float2(.5f, .5f)) * View.ViewSizeAndSceneTexelSize.zw;
-	float2 ScreenPosition = (ScreenUV.xy - View.ScreenPositionScaleBias.wz) / View.ScreenPositionScaleBias.xy;
-	float SceneDepth = CalcSceneDepth(ScreenUV);
-
-	float4 HomogeneousWorldPosition = mul(float4(ScreenPosition * SceneDepth, SceneDepth, 1), View.ScreenToWorld);
-	float3 OpaqueWorldPosition = HomogeneousWorldPosition.xyz / HomogeneousWorldPosition.w;
-	float2 BaseLevelScreenUV = (PixelIndex * DownsampleFactorToBaseLevel + float2(.5f, .5f)) * BaseLevelTexelSize;
-	float3 WorldNormal = DecodeNormalForAO(Texture2DSampleLevel(DistanceFieldNormalTexture, DistanceFieldNormalSampler, BaseLevelScreenUV, 0).xyz);
-	
-	// Initialize per-tile variables
-    if (ThreadIndex == 0) 
-	{
-        IntegerTileMinZ = 0x7F7FFFFF;     
-        IntegerTileMaxZ = 0;
-		IntegerTileMinZ2 = 0x7F7FFFFF;  
-		IntegerTileMaxZ2 = 0;
-    }
-
-    GroupMemoryBarrierWithGroupSync();
-    
-    InterlockedMin(IntegerTileMinZ, asuint(SceneDepth));
-    InterlockedMax(IntegerTileMaxZ, asuint(SceneDepth));
-
-    GroupMemoryBarrierWithGroupSync();
-
-    float MinTileZ = asfloat(IntegerTileMinZ);
-    float MaxTileZ = asfloat(IntegerTileMaxZ);
-	
-#define USE_TILE_FRUSTUM 1
-
-#define USE_DUAL_DEPTH_RANGES (1 && !USE_TILE_FRUSTUM)
-#if USE_DUAL_DEPTH_RANGES
-	float HalfZ = .5f * (MinTileZ + MaxTileZ);
-
-	// Compute a second min and max Z, clipped by HalfZ, so that we get two depth bounds per tile
-	// This results in more conservative tile depth bounds and fewer intersections
-	if (SceneDepth >= HalfZ)
-	{
-		InterlockedMin(IntegerTileMinZ2, asuint(SceneDepth));
-	}
-
-	if (SceneDepth <= HalfZ)
-	{
-		InterlockedMax(IntegerTileMaxZ2, asuint(SceneDepth));
-	}
-
-	GroupMemoryBarrierWithGroupSync();
-	
-	float MinTileZ2 = asfloat(IntegerTileMinZ2);
-	float MaxTileZ2 = asfloat(IntegerTileMaxZ2);
-
-#endif
-
-#if USE_TILE_FRUSTUM
-	
-	// Setup tile frustum planes
-	//@todo - use view size not buffer size
-	float2 TileScale = AOBufferSize * rcp(2 * float2(THREADGROUP_SIZEX, THREADGROUP_SIZEY));
-    float2 TileBias = TileScale - GroupId.xy;
-
-    float4 C1 = float4(View.ViewToClip._11 * TileScale.x,	0.0f,								View.ViewToClip._31 * TileScale.x + TileBias.x,	0.0f);
-    float4 C2 = float4(0.0f,								-View.ViewToClip._22 * TileScale.y, View.ViewToClip._32 * TileScale.y + TileBias.y,	0.0f);
-    float4 C4 = float4(0.0f,								0.0f,								1.0f,											0.0f);
-
-    float4 frustumPlanes[6];
-    frustumPlanes[0] = C4 - C1;
-    frustumPlanes[1] = C4 + C1;
-    frustumPlanes[2] = C4 - C2;
-    frustumPlanes[3] = C4 + C2;
-    frustumPlanes[4] = float4(0.0f, 0.0f,  1.0f, -MinTileZ);
-    frustumPlanes[5] = float4(0.0f, 0.0f, -1.0f,  MaxTileZ);
-
-	// Normalize tile frustum planes
-    UNROLL 
-	for (uint i = 0; i < 4; ++i) 
-	{
-        frustumPlanes[i] *= rcp(length(frustumPlanes[i].xyz));
-    }
-
-#else
-
-	float3 ViewTileMin;
-	float3 ViewTileMax;
-
-	float3 ViewTileMin2;
-	float3 ViewTileMax2;
-
-	float2 TanViewFOV = float2(1 / View.ViewToClip[0][0], 1 / View.ViewToClip[1][1]);
-	float2 TileSize = TanViewFOV * 2 / NumGroups;
-	float2 UnitPlaneTileMin = (GroupId.xy * TileSize - TanViewFOV) * float2(1, -1);
-	float2 UnitPlaneTileMax = ((GroupId.xy + 1) * TileSize - TanViewFOV) * float2(1, -1);
-
-	float ExpandRadius = 0;
-
-#if USE_DUAL_DEPTH_RANGES
-	ViewTileMin.xy = min(MinTileZ * UnitPlaneTileMin, MaxTileZ2 * UnitPlaneTileMin) - ExpandRadius;
-	ViewTileMax.xy = max(MinTileZ * UnitPlaneTileMax, MaxTileZ2 * UnitPlaneTileMax) + ExpandRadius;
-	ViewTileMin.z = MinTileZ - ExpandRadius;
-	ViewTileMax.z = MaxTileZ2 + ExpandRadius;
-	ViewTileMin2.xy = min(MinTileZ2 * UnitPlaneTileMin, MaxTileZ * UnitPlaneTileMin) - ExpandRadius;
-	ViewTileMax2.xy = max(MinTileZ2 * UnitPlaneTileMax, MaxTileZ * UnitPlaneTileMax) + ExpandRadius;
-	ViewTileMin2.z = MinTileZ2 - ExpandRadius;
-	ViewTileMax2.z = MaxTileZ + ExpandRadius;
-#else
-	ViewTileMin.xy = min(MinTileZ * UnitPlaneTileMin, MaxTileZ * UnitPlaneTileMin) - ExpandRadius;
-	ViewTileMax.xy = max(MinTileZ * UnitPlaneTileMax, MaxTileZ * UnitPlaneTileMax) + ExpandRadius;
-	ViewTileMin.z = MinTileZ - ExpandRadius;
-	ViewTileMax.z = MaxTileZ + ExpandRadius;
-#endif
-
-#endif
-	
-	float4 LocalAccumulation = 0;
-
-	UNROLL
-	for (uint LevelIndex = 0; LevelIndex < NUM_INTERPOLATE_LEVELS; LevelIndex++)
-	{
-		GroupMemoryBarrierWithGroupSync();
-
-		if (ThreadIndex == 0) 
-		{
-			TileNumSamples = 0;
-		}
-
-		GroupMemoryBarrierWithGroupSync();
-
-		uint NumRecords = LevelScatterDrawParameters[LevelIndex][1];
-
-		LOOP
-		for (uint SampleIndex = ThreadIndex; SampleIndex < NumRecords; SampleIndex += THREADGROUP_TOTALSIZE)
-		{
-			float4 PositionAndRadius = LevelIrradianceCachePositionRadius[LevelIndex].Load(SampleIndex);
-
-			float OccluderRadius = LevelIrradianceCacheOccluderRadius[LevelIndex].Load(SampleIndex);
-			PositionAndRadius.w = min(PositionAndRadius.w, OccluderRadius);
-
-			float ViewToCenterLength = length(PositionAndRadius.xyz - View.ViewOrigin.xyz);
-			PositionAndRadius.w = max(PositionAndRadius.w, ViewToCenterLength * .01f) * InterpolationRadiusScale;
-
-			float3 ViewSpaceSamplePosition = mul(float4(PositionAndRadius.xyz + View.PreViewTranslation.xyz, 1), View.TranslatedWorldToView).xyz;
-
-#if USE_TILE_FRUSTUM
-			// Cull the light against the tile's frustum planes
-			// Note: this has some false positives, a light that is intersecting three different axis frustum planes yet not intersecting the volume of the tile will be treated as intersecting
-			bool bInTile = true;  
-				
-			// Test against the screen x and y oriented planes first
-			UNROLL
-			for (uint i = 0; i < 4; ++i) 
-			{
-				float PlaneDistance = dot(frustumPlanes[i], float4(ViewSpaceSamplePosition, 1.0f));
-				bInTile = bInTile && (PlaneDistance >= -PositionAndRadius.w);
-			}
-
-			BRANCH
-			if (bInTile)
-			{
-				bool bInDepthRange = true;  
-				
-				// Test against the depth range
-				UNROLL 
-				for (uint i = 4; i < 6; ++i) 
-				{
-					float PlaneDistance = dot(frustumPlanes[i], float4(ViewSpaceSamplePosition, 1.0f));
-					bInDepthRange = bInDepthRange && (PlaneDistance >= -PositionAndRadius.w);
-				}
-				 
-				// Add this light to the list of indices if it intersects
-				BRANCH
-				if (bInDepthRange) 
-				{
-					uint ListIndex;
-					InterlockedAdd(TileNumSamples, 1U, ListIndex);
-
-					if (ListIndex < MAX_RECORDS_PER_TILE)
-					{
-						TileSampleIndices[ListIndex] = SampleIndex; 
-					}
-				}
-			}
-#else
-			//@todo - why needed?
-			PositionAndRadius.w *= 2;
-
-			float BoxDistance = ComputeDistanceFromBoxToPoint(ViewTileMin, ViewTileMax, ViewSpaceSamplePosition);
-
-#if USE_DUAL_DEPTH_RANGES
-			float BoxDistance2 = ComputeDistanceFromBoxToPoint(ViewTileMin2, ViewTileMax2, ViewSpaceSamplePosition);
-#endif
-
-			BRANCH
-			if (BoxDistance < PositionAndRadius.w
-#if USE_DUAL_DEPTH_RANGES
-				|| BoxDistance2 < PositionAndRadius.w
-#endif
-				)
-			{
-				uint ListIndex;
-				InterlockedAdd(TileNumSamples, 1U, ListIndex);
-
-				if (ListIndex < MAX_RECORDS_PER_TILE)
-				{
-					TileSampleIndices[ListIndex] = SampleIndex; 
-				}
-			}
-#endif
-		}
-
-		GroupMemoryBarrierWithGroupSync();
-
-		uint LocalNumSamples = TileNumSamples;
-
-		#define VISUALIZE_OVERDRAW 0
-		#if VISUALIZE_OVERDRAW
-			LocalAccumulation += LocalNumSamples * .005f;
-		#else
-
-			LOOP
-			for (uint ListIndex = 0; ListIndex < LocalNumSamples; ListIndex++)
-			{
-				uint SampleIndex = TileSampleIndices[ListIndex];
-				float4 PositionAndRadius = LevelIrradianceCachePositionRadius[LevelIndex].Load(SampleIndex);
-				float OccluderRadius = LevelIrradianceCacheOccluderRadius[LevelIndex].Load(SampleIndex);
-
-				PositionAndRadius.w = min(PositionAndRadius.w, OccluderRadius);
-
-				float ViewToCenterLength = length(PositionAndRadius.xyz - View.ViewOrigin.xyz);
-				PositionAndRadius.w = max(PositionAndRadius.w, ViewToCenterLength * .01f) * InterpolationRadiusScale;
-
-				float Distance = length(OpaqueWorldPosition - PositionAndRadius.xyz);
-				float DistanceError = saturate(Distance / PositionAndRadius.w);
-
-				BRANCH
-				if (DistanceError < 1)
-				{ 
-					float3 RecordNormal = DecodeNormalForAO(LevelIrradianceCacheNormal[LevelIndex].Load(SampleIndex).xyz);
-					float3 RecordBentNormal = DecodeNormalForAO(LevelIrradianceCacheBentNormal[LevelIndex].Load(SampleIndex).xyz);
-
-					float NormalError = InterpolationAngleNormalization * sqrt(saturate(1 - dot(WorldNormal, RecordNormal)));
-
-					// Don't use a lighting record if it's in front of the query point.
-					// Query points behind the lighting record may have nearby occluders that the lighting record does not see.
-					float3 RecordToVertexVector = OpaqueWorldPosition - (PositionAndRadius.xyz - 1 * RecordNormal.xyz);
-					float DistanceToVertex = length(RecordToVertexVector);
-					float PlaneDistance = dot(RecordNormal, RecordToVertexVector) / DistanceToVertex;
-		
-					// Setup an error metric that goes from 0 if the points are coplanar, to 1 if the point being shaded is at the angle corresponding to MinCosPointBehindPlane behind the plane
-					float PointBehindPlaneError = min(max(PlaneDistance * InvMinCosPointBehindPlane, 0.0f), DistanceToVertex / 3.0f);
-
-					float PrecisionScale = .1f;
-					float Weight = saturate(PrecisionScale * (1 - max(DistanceError, max(NormalError, PointBehindPlaneError))));
-					//float Weight = saturate(PrecisionScale * (1 - DistanceError));
-
-
-					LocalAccumulation.rgb += RecordBentNormal * Weight;
-					LocalAccumulation.a += Weight;
-				}
-			}
-		#endif
-	}
-
-	RWIrradianceCacheSplat[PixelIndex.xy] = LocalAccumulation;
-}
-
 Texture2D BentNormalAOTexture;
 SamplerState BentNormalAOSampler;
 
@@ -668,18 +373,18 @@ void AOCombinePS(
 
 	float4 BentNormalAccumulation = Texture2DSampleLevel(BentNormalAOTexture, BentNormalAOSampler, UVAndScreenPos.xy, 0);
 
-	if (DistanceFieldNormalValue.w < 0)
-	{
-		// Add a low weight completely unshadowed contribution for pixels with no distance field representation
-		float Weight = .001f;
-		BentNormalAccumulation += float4(DistanceFieldNormalValue.xyz * Weight, Weight);
-	}
-
 	if (BentNormalAccumulation.w > 0)
 	{
 		OutBentNormal.rgb = BentNormalAccumulation.xyz / BentNormalAccumulation.w;
 		OutBentNormal.a = SceneDepth;
 
+		if (DistanceFieldNormalValue.w < 0)
+		{
+			// Pixels without a distance field representation interpolate AO instead of a bent normal
+			// Construct a bent normal from the interpolated AO
+			OutBentNormal.rgb = OutBentNormal.r * DistanceFieldNormalValue.xyz;
+		}
+
 		#if SUPPORT_IRRADIANCE
 			OutIrradiance.rgb = Texture2DSampleLevel(IrradianceTexture, IrradianceSampler, UVAndScreenPos.xy, 0).xyz / BentNormalAccumulation.w;
 		#endif
@@ -694,7 +399,8 @@ void AOCombinePS(
 	float BentNormalLength = length(OutBentNormal.rgb);
 	// Fade to unoccluded in the distance
 	float FadeAlpha = saturate(.001f * (.95f * AOMaxViewDistance - SceneDepth));
-	OutBentNormal.rgb = normalize(OutBentNormal.rgb) * lerp(1, BentNormalLength, FadeAlpha);
+	float3 NormalizedBentNormal = OutBentNormal.rgb / max(BentNormalLength, .0001f);
+	OutBentNormal.rgb = NormalizedBentNormal * lerp(1, BentNormalLength, FadeAlpha);
 
 	// Clamp Nan's before they get to the filter
 	OutBentNormal.rgb = clamp(OutBentNormal.rgb, -1, 1);
@@ -873,13 +579,11 @@ void UpdateHistoryPS(
 	//@todo - this doesn't take camera movement into account
 	float DepthDeltaFraction = abs(abs(NewValue.a) - abs(HistoryValue.a)) / abs(NewValue.a);
 	float EffectiveHistoryWeight = HistoryWeight;
-
+	
 	FLATTEN
 	if (any(OldScreenPosition > 1) 
 		|| any(OldScreenPosition < -1)
-		|| DepthDeltaFraction > .03f
-		// Toss history value from pixels that don't have a distance field representation
-		|| HistoryValue.a < 0)
+		|| DepthDeltaFraction > .03f)
 	{
 		EffectiveHistoryWeight = 0;
 	}
@@ -994,7 +698,7 @@ void RayTraceThroughDistanceFields(float3 WorldRayStart, float3 WorldRayEnd, flo
 	TotalStepsTaken = 0;
 
 	LOOP
-	for (uint ListObjectIndex = 0; ListObjectIndex < NumIntersectingObjects; ListObjectIndex++)
+	for (uint ListObjectIndex = 0; ListObjectIndex < min(NumIntersectingObjects, MAX_INTERSECTING_OBJECTS); ListObjectIndex++)
 	{
 		uint ObjectIndex = IntersectingObjectIndices[ListObjectIndex];
 		float4 SphereCenterAndRadius = LoadObjectPositionAndRadius(ObjectIndex);
@@ -1071,7 +775,7 @@ void VisualizeMeshDistanceFieldCS(
 	float4 HomogeneousWorldPosition = mul(float4(ScreenPosition * SceneDepth, SceneDepth, 1), View.ScreenToWorld);
 	float3 OpaqueWorldPosition = HomogeneousWorldPosition.xyz / HomogeneousWorldPosition.w;
 
-	float TraceDistance = 10000;
+	float TraceDistance = 20000;
 	float3 WorldRayStart = View.ViewOrigin.xyz;
 	float3 WorldRayEnd = WorldRayStart + normalize(OpaqueWorldPosition - View.ViewOrigin.xyz) * TraceDistance;
 	float3 WorldRayDirection = WorldRayEnd - WorldRayStart;
diff --git a/Engine/Shaders/DistanceFieldSurfaceCacheLightingCompute.usf b/Engine/Shaders/DistanceFieldSurfaceCacheLightingCompute.usf
--- a/Engine/Shaders/DistanceFieldSurfaceCacheLightingCompute.usf
+++ b/Engine/Shaders/DistanceFieldSurfaceCacheLightingCompute.usf
@@ -642,7 +642,6 @@ void PopulateCacheCS(
 
 	BRANCH
 	if (GBufferData.ShadingModelID > 0
-		&& GBufferData.HasDistanceFieldRepresentation > 0
 		&& IrradianceCacheSplat.w < .0001f 
 		&& all(DispatchThreadId.xy < AOBufferSize))
 	{
@@ -669,7 +668,10 @@ void PopulateCacheCS(
 				int NextSampleIndex;
 				InterlockedAdd(NumQueuedIrradianceCacheRecords, 1, NextSampleIndex);
 				// W stores max allowed radius, used to limit overdraw from nearby samples placed in the high resolution passes
-				CachedIrradianceCachePositionRadius[NextSampleIndex] = float4(WorldShadingPosition, SceneDepth * CurrentLevelDownsampleFactor * .005f);
+				float MaxRadiusScale = GBufferData.HasDistanceFieldRepresentation > 0 ? .005f : .0005f;
+				// Sign of W stores whether the sample was placed on geometry with a distance field representation or not
+				float RadiusSign = GBufferData.HasDistanceFieldRepresentation > 0 ? 1 : -1;
+				CachedIrradianceCachePositionRadius[NextSampleIndex] = float4(WorldShadingPosition, RadiusSign * SceneDepth * CurrentLevelDownsampleFactor * MaxRadiusScale);
 				CachedIrradianceCacheNormal[NextSampleIndex] = float4(WorldNormal, .01f * CurrentLevelDownsampleFactor / 2);
 				CachedIrradianceCacheTileCoordinate[NextSampleIndex] = TileCoordinate;
 			}
@@ -717,7 +719,7 @@ Buffer<uint> ScatterDrawParameters;
 Buffer<uint> SavedStartIndex;
 
 RWBuffer<float> RWOccluderRadius;
-RWBuffer<float4> RWIrradianceCacheBentNormal;
+RWBuffer<float> RWRecordConeVisibility;
 RWBuffer<float4> RWIrradianceCacheIrradiance;
 RWBuffer<float4> RWDebugBuffer;
 
@@ -877,9 +879,8 @@ void FinalGatherCS(
 							// Don't allow small objects to fully occlude a cone step
 							Visibility = max(Visibility, 1 - saturate(ObjectOccluderRadius / SphereRadius));
 							float OccluderDistanceFraction = (WorldStepOffset + DistanceToOccluder) / AOMaxDistance;
-							// Fade out occlusion based on distance to occluder
-							Visibility = max(Visibility, saturate((OccluderDistanceFraction - 0) / 4));
-							Visibility = max(Visibility, saturate((OccluderDistanceFraction - .7f) / .4f));
+							// Fade out occlusion based on distance to occluder to avoid a discontinuity at the max AO distance
+							Visibility = max(Visibility, saturate(OccluderDistanceFraction * OccluderDistanceFraction * .8f));
 
 							if (Visibility < MinVisibility)
 							{
@@ -996,15 +997,17 @@ void FinalGatherCS(
 	{
 		float3 UnoccludedDirection = 0;
 		float Visibility = 0;
-		uint NumSamplesToDisplay = VISUALIZE_ONE_CONE ? 1 : NUM_CONE_DIRECTIONS;
 		float3 Irradiance = 0;
 
-		for (uint ConeIndex = 0; ConeIndex < NumSamplesToDisplay; ConeIndex++)
+		for (uint ConeIndex = 0; ConeIndex < NUM_CONE_DIRECTIONS; ConeIndex++)
 		{
 			float3 ConeDirection = AOSamples2.SampleDirections[ConeIndex].xyz;
 			float3 RotatedConeDirection = ConeDirection.x * TangentX + ConeDirection.y * TangentY + ConeDirection.z * WorldNormal;
 
 			float ConeVisibility = asfloat(SharedConeVisibility[ConeIndex]);
+
+			RWRecordConeVisibility[GroupId.x * 9 + ConeIndex] = ConeVisibility;
+
 			UnoccludedDirection += ConeVisibility * RotatedConeDirection;
 			Visibility += ConeVisibility;
 
@@ -1024,13 +1027,12 @@ void FinalGatherCS(
 #endif
 		}
 
-		float InvNumSamples = 1.0f  / (float)NumSamplesToDisplay;
+		float InvNumSamples = 1.0f  / (float)NUM_CONE_DIRECTIONS;
 		UnoccludedDirection = UnoccludedDirection * (BentNormalNormalizeFactor * InvNumSamples);
 		Visibility = Visibility * InvNumSamples;
 		Irradiance = Irradiance * InvNumSamples;
 
 		RWOccluderRadius[RecordIndex] = RecordRadiusScale * asfloat(SharedMinOcclusionDistance);
-		RWIrradianceCacheBentNormal[RecordIndex] = float4(UnoccludedDirection, 0);
 
 #if SUPPORT_IRRADIANCE
 		RWIrradianceCacheIrradiance[RecordIndex] = float4(Irradiance, 0);
@@ -1038,6 +1040,166 @@ void FinalGatherCS(
 	}
 }
 
+Texture2D HeightfieldTexture;
+SamplerState HeightfieldSampler;
+uint NumHeightfields;
+Buffer<float4> HeightfieldDescriptions;
+
+// In float4's, must match C++
+#define HEIGHTFIELD_DATA_STRIDE 10
+
+float4 GetHeightfieldScaleBias(uint Index)
+{
+	return HeightfieldDescriptions[Index * HEIGHTFIELD_DATA_STRIDE + 0];
+}
+
+float4 GetHeightfieldMinMaxUV(uint Index)
+{
+	return HeightfieldDescriptions[Index * HEIGHTFIELD_DATA_STRIDE + 1];
+}
+
+float4x4 GetWorldToLocal(uint Index)
+{
+	return float4x4(
+		HeightfieldDescriptions[Index * HEIGHTFIELD_DATA_STRIDE + 2],
+		HeightfieldDescriptions[Index * HEIGHTFIELD_DATA_STRIDE + 3],
+		HeightfieldDescriptions[Index * HEIGHTFIELD_DATA_STRIDE + 4],
+		HeightfieldDescriptions[Index * HEIGHTFIELD_DATA_STRIDE + 5]);
+}
+
+float4x4 GetLocalToWorld(uint Index)
+{
+	return float4x4(
+		HeightfieldDescriptions[Index * HEIGHTFIELD_DATA_STRIDE + 6],
+		HeightfieldDescriptions[Index * HEIGHTFIELD_DATA_STRIDE + 7],
+		HeightfieldDescriptions[Index * HEIGHTFIELD_DATA_STRIDE + 8],
+		HeightfieldDescriptions[Index * HEIGHTFIELD_DATA_STRIDE + 9]);
+}
+
+/**  */
+[numthreads(FINAL_GATHER_THREADGROUP_SIZE, 1, 1)]
+void CalculateHeightfieldOcclusionCS(
+	uint3 GroupId : SV_GroupID,
+	uint3 DispatchThreadId : SV_DispatchThreadID,
+    uint3 GroupThreadId : SV_GroupThreadID) 
+{
+	uint StartIndex = SavedStartIndex[0];
+	uint NumRecords = ScatterDrawParameters[1];
+	uint RecordIndex = StartIndex + DispatchThreadId.x;
+
+	float4 PositionAndRadius = IrradianceCachePositionRadius[RecordIndex];
+
+	// Only compute heightfield occlusion contribution for samples created on a distance field, 
+	// Otherwise we are potentially self-shadowing, the heightfield is not accurate enough for that
+	BRANCH
+	if (RecordIndex < NumRecords && PositionAndRadius.w > 0)
+	{
+		float3 WorldNormal = DecodeNormalForAO(IrradianceCacheNormal[RecordIndex].xyz);
+		float3 WorldPosition = PositionAndRadius.xyz;
+
+		float3 TangentX;
+		float3 TangentY;
+		FindBestAxisVectors2(WorldNormal, TangentX, TangentY);
+
+		LOOP
+		for (uint HeightfieldIndex = 0; HeightfieldIndex < NumHeightfields; HeightfieldIndex++)
+		{
+			float3 LocalPosition = mul(float4(WorldPosition, 1), GetWorldToLocal(HeightfieldIndex)).xyz;
+			float4 HeightfieldScaleBias = GetHeightfieldScaleBias(HeightfieldIndex);
+			float2 HeightfieldUV = LocalPosition.xy * HeightfieldScaleBias.xy + HeightfieldScaleBias.zw;
+			float4 MinMaxHeightfieldUV = GetHeightfieldMinMaxUV(HeightfieldIndex);
+
+			BRANCH
+			if (all(HeightfieldUV > MinMaxHeightfieldUV.xy) 
+				&& all(HeightfieldUV < MinMaxHeightfieldUV.zw))
+			{
+				float4 SampleValue = Texture2DSampleLevel(HeightfieldTexture, HeightfieldSampler, HeightfieldUV + .5f * HeightfieldScaleBias.xy, 0);
+				float HeightfieldHeight = ((SampleValue.r * 255.0 * 256.0 + SampleValue.g * 255.0) - 32768.0) * TERRAIN_ZSCALE;
+
+				float2 SampleNormal = float2(SampleValue.b, SampleValue.a) * float2(2.0,2.0) - float2(1.0,1.0);
+				float3 HeightfieldNormal = float3(SampleNormal, sqrt(max(1.0 - dot(SampleNormal, SampleNormal), 0.0)));
+
+				float3 HeightfieldShadingPosition = float3(LocalPosition.xy, HeightfieldHeight);
+				float3 WorldHeightfieldShadingPosition = mul(float4(HeightfieldShadingPosition, 1), GetLocalToWorld(HeightfieldIndex)).xyz;
+
+				float MinDistanceToOccluder = AOMaxDistance;
+
+				for (uint ConeIndex = 0; ConeIndex < NUM_CONE_DIRECTIONS; ConeIndex++)
+				{
+					float3 ConeDirection = AOSamples2.SampleDirections[ConeIndex].xyz;
+					float3 RotatedConeDirection = ConeDirection.x * TangentX + ConeDirection.y * TangentY + ConeDirection.z * WorldNormal;
+
+					float ExistingConeVisibility = RWRecordConeVisibility[DispatchThreadId.x * 9 + ConeIndex];
+				
+					float ZBias = 10;
+					float DistanceAlongConeToOccluder = clamp((WorldPosition.z + ZBias - WorldHeightfieldShadingPosition.z) / max(-RotatedConeDirection.z, .0001f), .00001f, AOMaxDistance);
+					float DistanceFromConeEndToOccluder = max(WorldPosition.z + ZBias + RotatedConeDirection.z * DistanceAlongConeToOccluder - WorldHeightfieldShadingPosition.z, 0);
+
+					float SphereRadius = DistanceAlongConeToOccluder * TanConeHalfAngle;
+
+					// Derive visibility from 1d intersection
+					float Visibility = saturate(DistanceFromConeEndToOccluder / SphereRadius);
+
+					float OccluderDistanceFraction = DistanceAlongConeToOccluder / AOMaxDistance;
+					// Fade out occlusion based on distance to occluder to avoid a discontinuity at the max AO distance
+					Visibility = max(Visibility, saturate(OccluderDistanceFraction * OccluderDistanceFraction * .8f));
+
+					MinDistanceToOccluder = min(MinDistanceToOccluder, DistanceAlongConeToOccluder);
+
+					// Shouldn't have to use atomics here as only one heightfield can affect a given world space position
+					RWRecordConeVisibility[DispatchThreadId.x * 9 + ConeIndex] = min(Visibility, ExistingConeVisibility);
+				}
+
+				RWOccluderRadius[RecordIndex] = min(RWOccluderRadius[RecordIndex], RecordRadiusScale * MinDistanceToOccluder);
+			}
+		}
+	}
+}
+
+Buffer<float> RecordConeVisibility;
+RWBuffer<float4> RWIrradianceCacheBentNormal;
+
+/**  */
+[numthreads(FINAL_GATHER_THREADGROUP_SIZE, 1, 1)]
+void CombineConesCS(
+	uint3 GroupId : SV_GroupID,
+	uint3 DispatchThreadId : SV_DispatchThreadID,
+    uint3 GroupThreadId : SV_GroupThreadID) 
+{
+	uint StartIndex = SavedStartIndex[0];
+	uint NumRecords = ScatterDrawParameters[1];
+
+	uint RecordIndex = StartIndex + DispatchThreadId.x;
+
+	if (RecordIndex < NumRecords)
+	{
+		float3 WorldNormal = DecodeNormalForAO(IrradianceCacheNormal[RecordIndex].xyz);
+
+		float3 TangentX;
+		float3 TangentY;
+		FindBestAxisVectors2(WorldNormal, TangentX, TangentY);
+
+		float3 UnoccludedDirection = 0;
+		float Visibility = 0;
+		uint NumSamplesToDisplay = VISUALIZE_ONE_CONE ? 1 : NUM_CONE_DIRECTIONS;
+
+		for (uint ConeIndex = 0; ConeIndex < NumSamplesToDisplay; ConeIndex++)
+		{
+			float3 ConeDirection = AOSamples2.SampleDirections[ConeIndex].xyz;
+			float3 RotatedConeDirection = ConeDirection.x * TangentX + ConeDirection.y * TangentY + ConeDirection.z * WorldNormal;
+
+			float ConeVisibility = RecordConeVisibility[DispatchThreadId.x * 9 + ConeIndex];
+			UnoccludedDirection += ConeVisibility * RotatedConeDirection;
+			Visibility += ConeVisibility;
+		}
+
+		float InvNumSamples = 1.0f  / (float)NumSamplesToDisplay;
+		UnoccludedDirection = UnoccludedDirection * (BentNormalNormalizeFactor * InvNumSamples);
+
+		RWIrradianceCacheBentNormal[RecordIndex] = float4(UnoccludedDirection, 0);
+	}
+}
+
 #define COMPACT_THREADGROUP_SIZEX 64
 
 Buffer<uint> DrawParameters;
@@ -1064,7 +1226,13 @@ void SetupFinalGatherIndirectArgumentsCS(
 	uint NumRecords = DrawParameters[1];
 	uint WorkRange = NumRecords - StartIndex;
 
+#if ONE_GROUP_PER_RECORD
+	// One thread group per record
 	RWDispatchParameters[0] = WorkRange;
+#else
+	// One thread per record, divide and round up
+	RWDispatchParameters[0] = (WorkRange + FINAL_GATHER_THREADGROUP_SIZE - 1) / FINAL_GATHER_THREADGROUP_SIZE;
+#endif
 	RWDispatchParameters[1] = 1;
 	RWDispatchParameters[2] = 1;
 }
diff --git a/Engine/Shaders/LandscapeVertexFactory.usf b/Engine/Shaders/LandscapeVertexFactory.usf
--- a/Engine/Shaders/LandscapeVertexFactory.usf
+++ b/Engine/Shaders/LandscapeVertexFactory.usf
@@ -6,8 +6,6 @@
 
 #include "VertexFactoryCommon.usf"
 
-#define TERRAIN_ZSCALE				(1.0f/128.0f)
-
 #define ES2_MAXLOD	5
 
 // Set by FLightMapTexturePolicy
diff --git a/Engine/Source/Editor/FoliageEdit/Private/FoliageEdMode.cpp b/Engine/Source/Editor/FoliageEdit/Private/FoliageEdMode.cpp
--- a/Engine/Source/Editor/FoliageEdit/Private/FoliageEdMode.cpp
+++ b/Engine/Source/Editor/FoliageEdit/Private/FoliageEdMode.cpp
@@ -58,6 +58,7 @@ FEdModeFoliage::FEdModeFoliage()
 	SphereBrushComponent->StaticMesh = StaticMesh;
 	SphereBrushComponent->OverrideMaterials.Add(BrushMaterial);
 	SphereBrushComponent->SetAbsolute(true, true, true);
+	SphereBrushComponent->CastShadow = false;
 
 	bBrushTraceValid = false;
 	BrushLocation = FVector::ZeroVector;
diff --git a/Engine/Source/Runtime/Engine/Private/PrimitiveSceneProxy.cpp b/Engine/Source/Runtime/Engine/Private/PrimitiveSceneProxy.cpp
--- a/Engine/Source/Runtime/Engine/Private/PrimitiveSceneProxy.cpp
+++ b/Engine/Source/Runtime/Engine/Private/PrimitiveSceneProxy.cpp
@@ -50,6 +50,7 @@ FPrimitiveSceneProxy::FPrimitiveSceneProxy(const UPrimitiveComponent* InComponen
 ,	bAlwaysHasVelocity(false)
 ,	bUseEditorDepthTest(true)
 ,	bSupportsDistanceFieldRepresentation(false)
+,	bSupportsHeightfieldRepresentation(false)
 ,	bUseAsOccluder(InComponent->bUseAsOccluder)
 ,	bAllowApproximateOcclusion(InComponent->Mobility != EComponentMobility::Movable)
 ,	bSelectable(InComponent->bSelectable)
diff --git a/Engine/Source/Runtime/Engine/Private/StaticMeshRender.cpp b/Engine/Source/Runtime/Engine/Private/StaticMeshRender.cpp
--- a/Engine/Source/Runtime/Engine/Private/StaticMeshRender.cpp
+++ b/Engine/Source/Runtime/Engine/Private/StaticMeshRender.cpp
@@ -974,7 +974,7 @@ void FStaticMeshSceneProxy::GetDistancefieldAtlasData(FBox& LocalVolumeBounds, F
 
 bool FStaticMeshSceneProxy::HasDistanceFieldRepresentation() const
 {
-	return (CastsDynamicShadow() || CastsStaticShadow()) && DistanceFieldData && DistanceFieldData->VolumeTexture.IsValidDistanceFieldVolume();
+	return CastsDynamicShadow() && AffectsDistanceFieldLighting() && DistanceFieldData && DistanceFieldData->VolumeTexture.IsValidDistanceFieldVolume();
 }
 
 /** Initialization constructor. */
diff --git a/Engine/Source/Runtime/Engine/Public/PrimitiveSceneProxy.h b/Engine/Source/Runtime/Engine/Public/PrimitiveSceneProxy.h
--- a/Engine/Source/Runtime/Engine/Public/PrimitiveSceneProxy.h
+++ b/Engine/Source/Runtime/Engine/Public/PrimitiveSceneProxy.h
@@ -204,6 +204,13 @@ public:
 		bMeshWasPlane = false;
 	}
 
+	virtual void GetHeightfieldRepresentation(UTexture2D*& OutHeightmapTexture, FVector4& OutHeightfieldScaleBias, FVector4& OutMinMaxUV)
+	{
+		OutHeightmapTexture = NULL;
+		OutHeightfieldScaleBias = FVector4(0, 0, 0, 0);
+		OutMinMaxUV = FVector4(0, 0, 0, 0);
+	}
+
 	/**
 	 *	Called when the rendering thread adds the proxy to the scene.
 	 *	This function allows for generating renderer-side resources.
@@ -346,6 +353,7 @@ public:
 	inline bool AlwaysHasVelocity() const { return bAlwaysHasVelocity; }
 	inline bool UseEditorDepthTest() const { return bUseEditorDepthTest; }
 	inline bool SupportsDistanceFieldRepresentation() const { return bSupportsDistanceFieldRepresentation; }
+	inline bool SupportsHeightfieldRepresentation() const { return bSupportsHeightfieldRepresentation; }
 	inline bool TreatAsBackgroundForOcclusion() const { return bTreatAsBackgroundForOcclusion; }
 #if WITH_EDITOR
 	inline int32 GetNumUncachedStaticLightingInteractions() { return NumUncachedStaticLightingInteractions; }
@@ -532,6 +540,9 @@ protected:
 	/** Whether the primitive type supports a distance field representation.  Does not mean the primitive has a valid representation. */
 	uint32 bSupportsDistanceFieldRepresentation : 1;
 
+	/** Whether the primitive implements GetHeightfieldRepresentation() */
+	uint32 bSupportsHeightfieldRepresentation : 1;
+
 private:
 
 	/** If this is True, this primitive will be used to occlusion cull other primitives. */
diff --git a/Engine/Source/Runtime/Landscape/Private/LandscapeRender.cpp b/Engine/Source/Runtime/Landscape/Private/LandscapeRender.cpp
--- a/Engine/Source/Runtime/Landscape/Private/LandscapeRender.cpp
+++ b/Engine/Source/Runtime/Landscape/Private/LandscapeRender.cpp
@@ -662,6 +662,8 @@ FLandscapeComponentSceneProxy::FLandscapeComponentSceneProxy(ULandscapeComponent
 
 	bRequiresAdjacencyInformation = RequiresAdjacencyInformation(MaterialInterface, XYOffsetmapTexture == NULL ? &FLandscapeVertexFactory::StaticType : &FLandscapeXYOffsetVertexFactory::StaticType, InComponent->GetWorld()->FeatureLevel);
 	SharedBuffersKey = (SubsectionSizeQuads & 0xffff) | ((NumSubsections & 0xf) << 16) | (FeatureLevel <= ERHIFeatureLevel::ES3_1 ? 0 : 1<<20) | (XYOffsetmapTexture == NULL ? 0 : 1 << 31);
+
+	bSupportsHeightfieldRepresentation = true;
 }
 
 void FLandscapeComponentSceneProxy::CreateRenderThreadResources()
diff --git a/Engine/Source/Runtime/Landscape/Public/LandscapeRender.h b/Engine/Source/Runtime/Landscape/Public/LandscapeRender.h
--- a/Engine/Source/Runtime/Landscape/Public/LandscapeRender.h
+++ b/Engine/Source/Runtime/Landscape/Public/LandscapeRender.h
@@ -501,6 +501,16 @@ public:
 
 	// FLandcapeSceneProxy
 	void ChangeLODDistanceFactor_RenderThread(FVector2D InLODDistanceFactors);
+
+	virtual void GetHeightfieldRepresentation(UTexture2D*& OutHeightmapTexture, FVector4& OutHeightfieldScaleBias, FVector4& OutMinMaxUV) override
+	{
+		OutHeightmapTexture = HeightmapTexture;
+		OutHeightfieldScaleBias = HeightmapScaleBias;
+		// Section base is in terms of quads, convert into texels
+		int32 TexelsSectionBaseX = SectionBase.X / SubsectionSizeQuads * SubsectionSizeVerts;
+		int32 TexelsSectionBaseY = SectionBase.Y / SubsectionSizeQuads * SubsectionSizeVerts;
+		OutMinMaxUV = FVector4(TexelsSectionBaseX * HeightmapScaleBias.X, TexelsSectionBaseY * HeightmapScaleBias.Y, (TexelsSectionBaseX + SubsectionSizeVerts - 1) * HeightmapScaleBias.X, (TexelsSectionBaseY + SubsectionSizeVerts - 1) * HeightmapScaleBias.Y);
+	}
 };
 
 class FLandscapeDebugMaterialRenderProxy : public FMaterialRenderProxy
diff --git a/Engine/Source/Runtime/Renderer/Private/DistanceFieldSurfaceCacheLighting.cpp b/Engine/Source/Runtime/Renderer/Private/DistanceFieldSurfaceCacheLighting.cpp
--- a/Engine/Source/Runtime/Renderer/Private/DistanceFieldSurfaceCacheLighting.cpp
+++ b/Engine/Source/Runtime/Renderer/Private/DistanceFieldSurfaceCacheLighting.cpp
@@ -163,6 +163,13 @@ FAutoConsoleVariableRef CVarAOMaxViewDistance(
 	ECVF_Cheat | ECVF_RenderThreadSafe
 	);
 
+float GetMaxAOViewDistance()
+{
+	// Scene depth stored in fp16 alpha, must fade out before it runs out of range
+	// The fade extends past GAOMaxViewDistance a bit
+	return FMath::Min(GAOMaxViewDistance, 65000.0f);
+}
+
 int32 GAOScatterTileCulling = 1;
 FAutoConsoleVariableRef CVarAOScatterTileCulling(
 	TEXT("r.AOScatterTileCulling"),
@@ -171,14 +178,6 @@ FAutoConsoleVariableRef CVarAOScatterTileCulling(
 	ECVF_Cheat | ECVF_RenderThreadSafe
 	);
 
-int32 GAOScatterInterpolation = 1;
-FAutoConsoleVariableRef CVarAOScatterInterpolation(
-	TEXT("r.AOScatterInterpolation"),
-	GAOScatterInterpolation,
-	TEXT(""),
-	ECVF_Cheat | ECVF_RenderThreadSafe
-	);
-
 int32 GAOComputeShaderNormalCalculation = 1;
 FAutoConsoleVariableRef CVarAOComputeShaderNormalCalculation(
 	TEXT("r.AOComputeShaderNormalCalculation"),
@@ -420,6 +419,82 @@ public:
 
 TGlobalResource<FDistanceFieldRemoveIndicesResource> GDistanceFieldRemoveIndices;
 
+class FHeightfieldDescriptionsResource : public FRenderResource
+{
+public:
+
+	FCPUUpdatedBuffer Data;
+
+	FHeightfieldDescriptionsResource()
+	{
+		Data.Format = PF_A32B32G32R32F;
+		// In float4's, must match usf
+		Data.Stride = 10;
+	}
+
+	virtual void InitDynamicRHI()  override
+	{
+		Data.Initialize();
+	}
+
+	virtual void ReleaseDynamicRHI() override
+	{
+		Data.Release();
+	}
+};
+
+TGlobalResource<FHeightfieldDescriptionsResource> GHeightfieldDescriptions;
+
+/**  */
+class FTemporaryIrradianceCacheResources : public FRenderResource
+{
+public:
+
+	FTemporaryIrradianceCacheResources()
+	{
+		MaxIrradianceCacheSamples = 0;
+	}
+
+	virtual void InitDynamicRHI()
+	{
+		if (MaxIrradianceCacheSamples > 0)
+		{
+			ConeVisibility.Initialize(sizeof(float), MaxIrradianceCacheSamples * 9, PF_R32_FLOAT, BUF_Static);
+		}
+	}
+
+	virtual void ReleaseDynamicRHI()
+	{
+		ConeVisibility.Release();
+	}
+
+	void AllocateFor(int32 InMaxIrradianceCacheSamples)
+	{
+		bool bReallocate = false;
+
+		if (InMaxIrradianceCacheSamples > MaxIrradianceCacheSamples)
+		{
+			MaxIrradianceCacheSamples = InMaxIrradianceCacheSamples;
+			bReallocate = true;
+		}
+
+		if (!IsInitialized())
+		{
+			InitResource();
+		}
+		else if (bReallocate)
+		{
+			UpdateRHI();
+		}
+	}
+
+	int32 MaxIrradianceCacheSamples;
+
+	FRWBuffer ConeVisibility;
+};
+
+TGlobalResource<FTemporaryIrradianceCacheResources> GTemporaryIrradianceCacheResources;
+
 void FTileIntersectionResources::InitDynamicRHI()
 {
 	TileConeAxisAndCos.Initialize(sizeof(float)* 4, TileDimensions.X * TileDimensions.Y, PF_A32B32G32R32F, BUF_Static);
@@ -973,11 +1048,15 @@ void FDeferredShadingSceneRenderer::UpdateGlobalDistanceFieldObjectBuffers(FRHIC
 			}
 
 			void* LockedBuffer = RHILockVertexBuffer(GDistanceFieldUploadIndices.UploadIndices.Buffer, 0, GDistanceFieldUploadIndices.UploadIndices.Buffer->GetSize(), RLM_WriteOnly);
-			FPlatformMemory::Memcpy(LockedBuffer, UploadObjectIndices.GetData(), UploadObjectIndices.GetTypeSize() * UploadObjectIndices.Num());
+			const uint32 MemcpySize = UploadObjectIndices.GetTypeSize() * UploadObjectIndices.Num();
+			check(GDistanceFieldUploadIndices.UploadIndices.Buffer->GetSize() >= MemcpySize);
+			FPlatformMemory::Memcpy(LockedBuffer, UploadObjectIndices.GetData(), MemcpySize);
 			RHIUnlockVertexBuffer(GDistanceFieldUploadIndices.UploadIndices.Buffer);
 
 			LockedBuffer = RHILockVertexBuffer(GDistanceFieldUploadData.UploadData.Buffer, 0, GDistanceFieldUploadData.UploadData.Buffer->GetSize(), RLM_WriteOnly);
-			FPlatformMemory::Memcpy(LockedBuffer, UploadObjectData.GetData(), UploadObjectData.GetTypeSize() * UploadObjectData.Num());
+			const uint32 MemcpySize2 = UploadObjectData.GetTypeSize() * UploadObjectData.Num();
+			check(GDistanceFieldUploadData.UploadData.Buffer->GetSize() >= MemcpySize2);
+			FPlatformMemory::Memcpy(LockedBuffer, UploadObjectData.GetData(), MemcpySize2);
 			RHIUnlockVertexBuffer(GDistanceFieldUploadData.UploadData.Buffer);
 
 			{
@@ -1101,7 +1180,9 @@ void FDeferredShadingSceneRenderer::UpdateGlobalDistanceFieldObjectBuffers(FRHIC
 				}
 
 				void* LockedBuffer = RHILockVertexBuffer(GDistanceFieldRemoveIndices.RemoveIndices.Buffer, 0, GDistanceFieldRemoveIndices.RemoveIndices.Buffer->GetSize(), RLM_WriteOnly);
-				FPlatformMemory::Memcpy(LockedBuffer, RemoveObjectIndices.GetData(), RemoveObjectIndices.GetTypeSize() * RemoveObjectIndices.Num());
+				const uint32 MemcpySize = RemoveObjectIndices.GetTypeSize() * RemoveObjectIndices.Num();
+				check(GDistanceFieldRemoveIndices.RemoveIndices.Buffer->GetSize() >= MemcpySize);
+				FPlatformMemory::Memcpy(LockedBuffer, RemoveObjectIndices.GetData(), MemcpySize);
 				RHIUnlockVertexBuffer(GDistanceFieldRemoveIndices.RemoveIndices.Buffer);
 
 				if (bUseRemoveAtSwap)
@@ -2150,9 +2231,10 @@ private:
 
 IMPLEMENT_SHADER_TYPE(,FPopulateCacheCS,TEXT("DistanceFieldSurfaceCacheLightingCompute"),TEXT("PopulateCacheCS"),SF_Compute);
 
-class FSetupFinalGatherIndirectArgumentsCS : public FGlobalShader
+template<bool bOneGroupPerRecord>
+class TSetupFinalGatherIndirectArgumentsCS : public FGlobalShader
 {
-	DECLARE_SHADER_TYPE(FSetupFinalGatherIndirectArgumentsCS,Global)
+	DECLARE_SHADER_TYPE(TSetupFinalGatherIndirectArgumentsCS,Global)
 public:
 
 	static bool ShouldCache(EShaderPlatform Platform)
@@ -2163,12 +2245,13 @@ public:
 	static void ModifyCompilationEnvironment(EShaderPlatform Platform, FShaderCompilerEnvironment& OutEnvironment)
 	{
 		FGlobalShader::ModifyCompilationEnvironment(Platform,OutEnvironment);
+		OutEnvironment.SetDefine(TEXT("ONE_GROUP_PER_RECORD"), bOneGroupPerRecord ? TEXT("1") : TEXT("0"));
 
 		// To reduce shader compile time of compute shaders with shared memory, doesn't have an impact on generated code with current compiler (June 2010 DX SDK)
 		OutEnvironment.CompilerFlags.Add(CFLAG_StandardOptimization);
 	}
 
-	FSetupFinalGatherIndirectArgumentsCS(const ShaderMetaType::CompiledShaderInitializerType& Initializer)
+	TSetupFinalGatherIndirectArgumentsCS(const ShaderMetaType::CompiledShaderInitializerType& Initializer)
 		: FGlobalShader(Initializer)
 	{
 		DrawParameters.Bind(Initializer.ParameterMap, TEXT("DrawParameters"));
@@ -2176,7 +2259,7 @@ public:
 		SavedStartIndex.Bind(Initializer.ParameterMap, TEXT("SavedStartIndex"));
 	}
 
-	FSetupFinalGatherIndirectArgumentsCS()
+	TSetupFinalGatherIndirectArgumentsCS()
 	{
 	}
 
@@ -2215,7 +2298,8 @@ private:
 	FRWShaderParameter DispatchParameters;
 };
 
-IMPLEMENT_SHADER_TYPE(,FSetupFinalGatherIndirectArgumentsCS,TEXT("DistanceFieldSurfaceCacheLightingCompute"),TEXT("SetupFinalGatherIndirectArgumentsCS"),SF_Compute);
+IMPLEMENT_SHADER_TYPE(template<>,TSetupFinalGatherIndirectArgumentsCS<true>,TEXT("DistanceFieldSurfaceCacheLightingCompute"),TEXT("SetupFinalGatherIndirectArgumentsCS"),SF_Compute);
+IMPLEMENT_SHADER_TYPE(template<>,TSetupFinalGatherIndirectArgumentsCS<false>,TEXT("DistanceFieldSurfaceCacheLightingCompute"),TEXT("SetupFinalGatherIndirectArgumentsCS"),SF_Compute);
 
 class FDebugBuffer : public FRenderResource
 {
@@ -2273,7 +2357,7 @@ public:
 		IrradianceCachePositionRadius.Bind(Initializer.ParameterMap, TEXT("IrradianceCachePositionRadius"));
 		IrradianceCacheNormal.Bind(Initializer.ParameterMap, TEXT("IrradianceCacheNormal"));
 		OccluderRadius.Bind(Initializer.ParameterMap, TEXT("OccluderRadius"));
-		IrradianceCacheBentNormal.Bind(Initializer.ParameterMap, TEXT("IrradianceCacheBentNormal"));
+		RecordConeVisibility.Bind(Initializer.ParameterMap, TEXT("RecordConeVisibility"));
 		IrradianceCacheIrradiance.Bind(Initializer.ParameterMap, TEXT("IrradianceCacheIrradiance"));
 		DebugBuffer.Bind(Initializer.ParameterMap, TEXT("DebugBuffer"));
 		IrradianceCacheTileCoordinate.Bind(Initializer.ParameterMap, TEXT("IrradianceCacheTileCoordinate"));
@@ -2326,7 +2410,7 @@ public:
 		SetSRVParameter(RHICmdList, ShaderRHI, IrradianceCachePositionRadius, SurfaceCacheResources.Level[DepthLevel]->PositionAndRadius.SRV);
 
 		OccluderRadius.SetBuffer(RHICmdList, ShaderRHI, SurfaceCacheResources.Level[DepthLevel]->OccluderRadius);
-		IrradianceCacheBentNormal.SetBuffer(RHICmdList, ShaderRHI, SurfaceCacheResources.Level[DepthLevel]->BentNormal);
+		RecordConeVisibility.SetBuffer(RHICmdList, ShaderRHI, GTemporaryIrradianceCacheResources.ConeVisibility);
 		IrradianceCacheIrradiance.SetBuffer(RHICmdList, ShaderRHI, SurfaceCacheResources.Level[DepthLevel]->Irradiance);
 		DebugBuffer.SetBuffer(RHICmdList, ShaderRHI, GDebugBuffer.DebugData);
 
@@ -2398,7 +2482,7 @@ public:
 	void UnsetParameters(FRHICommandList& RHICmdList)
 	{
 		OccluderRadius.UnsetUAV(RHICmdList, GetComputeShader());
-		IrradianceCacheBentNormal.UnsetUAV(RHICmdList, GetComputeShader());
+		RecordConeVisibility.UnsetUAV(RHICmdList, GetComputeShader());
 		IrradianceCacheIrradiance.UnsetUAV(RHICmdList, GetComputeShader());
 		DebugBuffer.UnsetUAV(RHICmdList, GetComputeShader());
 	}
@@ -2413,7 +2497,7 @@ public:
 		Ar << IrradianceCachePositionRadius;
 		Ar << IrradianceCacheNormal;
 		Ar << OccluderRadius;
-		Ar << IrradianceCacheBentNormal;
+		Ar << RecordConeVisibility;
 		Ar << IrradianceCacheIrradiance;
 		Ar << DebugBuffer;
 		Ar << IrradianceCacheTileCoordinate;
@@ -2445,7 +2529,7 @@ private:
 	FShaderResourceParameter IrradianceCacheNormal;
 	FShaderResourceParameter IrradianceCachePositionRadius;
 	FRWShaderParameter OccluderRadius;
-	FRWShaderParameter IrradianceCacheBentNormal;
+	FRWShaderParameter RecordConeVisibility;
 	FRWShaderParameter IrradianceCacheIrradiance;
 	FRWShaderParameter DebugBuffer;
 	FShaderResourceParameter ScatterDrawParameters;
@@ -2470,9 +2554,10 @@ private:
 IMPLEMENT_SHADER_TYPE(template<>,TFinalGatherCS<true>,TEXT("DistanceFieldSurfaceCacheLightingCompute"),TEXT("FinalGatherCS"),SF_Compute);
 IMPLEMENT_SHADER_TYPE(template<>,TFinalGatherCS<false>,TEXT("DistanceFieldSurfaceCacheLightingCompute"),TEXT("FinalGatherCS"),SF_Compute);
 
-class FWriteDownsampledDepthPS : public FGlobalShader
+
+class FCalculateHeightfieldOcclusionCS : public FGlobalShader
 {
-	DECLARE_SHADER_TYPE(FWriteDownsampledDepthPS, Global);
+	DECLARE_SHADER_TYPE(FCalculateHeightfieldOcclusionCS,Global)
 public:
 
 	static bool ShouldCache(EShaderPlatform Platform)
@@ -2482,146 +2567,275 @@ public:
 
 	static void ModifyCompilationEnvironment(EShaderPlatform Platform, FShaderCompilerEnvironment& OutEnvironment)
 	{
-		OutEnvironment.SetDefine(TEXT("DOWNSAMPLE_FACTOR"), GAODownsampleFactor);
-	}
+		FGlobalShader::ModifyCompilationEnvironment(Platform,OutEnvironment);
+		OutEnvironment.SetDefine(TEXT("NUM_CONE_DIRECTIONS"), NumConeSampleDirections);
 
-	/** Default constructor. */
-	FWriteDownsampledDepthPS() {}
+		// To reduce shader compile time of compute shaders with shared memory, doesn't have an impact on generated code with current compiler (June 2010 DX SDK)
+		OutEnvironment.CompilerFlags.Add(CFLAG_StandardOptimization);
+	}
 
-	/** Initialization constructor. */
-	FWriteDownsampledDepthPS(const ShaderMetaType::CompiledShaderInitializerType& Initializer)
+	FCalculateHeightfieldOcclusionCS(const ShaderMetaType::CompiledShaderInitializerType& Initializer)
 		: FGlobalShader(Initializer)
 	{
-		DeferredParameters.Bind(Initializer.ParameterMap);
-		DistanceFieldNormalTexture.Bind(Initializer.ParameterMap,TEXT("DistanceFieldNormalTexture"));
-		DistanceFieldNormalSampler.Bind(Initializer.ParameterMap,TEXT("DistanceFieldNormalSampler"));
+		AOParameters.Bind(Initializer.ParameterMap);
+		OccluderRadius.Bind(Initializer.ParameterMap, TEXT("OccluderRadius"));
+		IrradianceCacheNormal.Bind(Initializer.ParameterMap, TEXT("IrradianceCacheNormal"));
+		IrradianceCachePositionRadius.Bind(Initializer.ParameterMap, TEXT("IrradianceCachePositionRadius"));
+		RecordConeVisibility.Bind(Initializer.ParameterMap, TEXT("RecordConeVisibility"));
+		ScatterDrawParameters.Bind(Initializer.ParameterMap, TEXT("ScatterDrawParameters"));
+		SavedStartIndex.Bind(Initializer.ParameterMap, TEXT("SavedStartIndex"));
+		BentNormalNormalizeFactor.Bind(Initializer.ParameterMap, TEXT("BentNormalNormalizeFactor"));
+		DistanceFieldNormalTexture.Bind(Initializer.ParameterMap, TEXT("DistanceFieldNormalTexture"));
+		DistanceFieldNormalSampler.Bind(Initializer.ParameterMap, TEXT("DistanceFieldNormalSampler"));
+		RecordRadiusScale.Bind(Initializer.ParameterMap, TEXT("RecordRadiusScale"));
+		HeightfieldTexture.Bind(Initializer.ParameterMap, TEXT("HeightfieldTexture"));
+		HeightfieldSampler.Bind(Initializer.ParameterMap, TEXT("HeightfieldSampler"));
+		HeightfieldDescriptions.Bind(Initializer.ParameterMap, TEXT("HeightfieldDescriptions"));
+		NumHeightfields.Bind(Initializer.ParameterMap, TEXT("NumHeightfields"));
+		TanConeHalfAngle.Bind(Initializer.ParameterMap, TEXT("TanConeHalfAngle"));
 	}
 
-	void SetParameters(FRHICommandList& RHICmdList, const FSceneView& View, FSceneRenderTargetItem& DistanceFieldNormal)
+	FCalculateHeightfieldOcclusionCS()
 	{
-		const FPixelShaderRHIParamRef ShaderRHI = GetPixelShader();
+	}
 
+	void SetParameters(
+		FRHICommandList& RHICmdList, 
+		const FSceneView& View, 
+		int32 DepthLevel, 
+		UTexture2D* HeightfieldTextureValue,
+		int32 NumHeightfieldsValue,
+		const FDistanceFieldAOParameters& Parameters)
+	{
+		FComputeShaderRHIParamRef ShaderRHI = GetComputeShader();
 		FGlobalShader::SetParameters(RHICmdList, ShaderRHI, View);
-		DeferredParameters.Set(RHICmdList, ShaderRHI, View);
 
-		SetTextureParameter(RHICmdList, ShaderRHI, DistanceFieldNormalTexture, DistanceFieldNormalSampler, TStaticSamplerState<SF_Point>::GetRHI(), DistanceFieldNormal.ShaderResourceTexture);
+		AOParameters.Set(RHICmdList, ShaderRHI, Parameters);
+
+		const FScene* Scene = (const FScene*)View.Family->Scene;
+		FSurfaceCacheResources& SurfaceCacheResources = *Scene->SurfaceCacheResources;
+
+		SetSRVParameter(RHICmdList, ShaderRHI, IrradianceCacheNormal, SurfaceCacheResources.Level[DepthLevel]->Normal.SRV);
+		SetSRVParameter(RHICmdList, ShaderRHI, IrradianceCachePositionRadius, SurfaceCacheResources.Level[DepthLevel]->PositionAndRadius.SRV);
+		SetSRVParameter(RHICmdList, ShaderRHI, ScatterDrawParameters, SurfaceCacheResources.Level[DepthLevel]->ScatterDrawParameters.SRV);
+		SetSRVParameter(RHICmdList, ShaderRHI, SavedStartIndex, SurfaceCacheResources.Level[DepthLevel]->SavedStartIndex.SRV);
+		SetSRVParameter(RHICmdList, ShaderRHI, HeightfieldDescriptions, GHeightfieldDescriptions.Data.BufferSRV);
+
+		SetTextureParameter(RHICmdList, ShaderRHI, HeightfieldTexture, HeightfieldSampler, HeightfieldTextureValue->Resource);
+		SetShaderValue(RHICmdList, ShaderRHI, NumHeightfields, NumHeightfieldsValue);
+
+		OccluderRadius.SetBuffer(RHICmdList, ShaderRHI, SurfaceCacheResources.Level[DepthLevel]->OccluderRadius);
+		RecordConeVisibility.SetBuffer(RHICmdList, ShaderRHI, GTemporaryIrradianceCacheResources.ConeVisibility);
+
+		FAOSampleData2 AOSampleData;
+
+		TArray<FVector, TInlineAllocator<9> > SampleDirections;
+		GetSpacedVectors(SampleDirections);
+
+		for (int32 SampleIndex = 0; SampleIndex < NumConeSampleDirections; SampleIndex++)
+		{
+			AOSampleData.SampleDirections[SampleIndex] = FVector4(SampleDirections[SampleIndex]);
+		}
+
+		SetUniformBufferParameterImmediate(RHICmdList, ShaderRHI, GetUniformBufferParameter<FAOSampleData2>(), AOSampleData);
+
+		FVector UnoccludedVector(0);
+
+		for (int32 SampleIndex = 0; SampleIndex < NumConeSampleDirections; SampleIndex++)
+		{
+			UnoccludedVector += SampleDirections[SampleIndex];
+		}
+
+		float BentNormalNormalizeFactorValue = 1.0f / (UnoccludedVector / NumConeSampleDirections).Size();
+		SetShaderValue(RHICmdList, ShaderRHI, BentNormalNormalizeFactor, BentNormalNormalizeFactorValue);
+
+		SetShaderValue(RHICmdList, ShaderRHI, TanConeHalfAngle, FMath::Tan(GAOConeHalfAngle));
+		SetShaderValue(RHICmdList, ShaderRHI, RecordRadiusScale, GAORecordRadiusScale);
 	}
-	// FShader interface.
-	virtual bool Serialize(FArchive& Ar)
+
+	void UnsetParameters(FRHICommandList& RHICmdList)
 	{
+		OccluderRadius.UnsetUAV(RHICmdList, GetComputeShader());
+		RecordConeVisibility.UnsetUAV(RHICmdList, GetComputeShader());
+	}
+
+	virtual bool Serialize(FArchive& Ar)
+	{		
 		bool bShaderHasOutdatedParameters = FGlobalShader::Serialize(Ar);
-		Ar << DeferredParameters;
+		Ar << AOParameters;
+		Ar << OccluderRadius;
+		Ar << IrradianceCacheNormal;
+		Ar << IrradianceCachePositionRadius;
+		Ar << RecordConeVisibility;
+		Ar << ScatterDrawParameters;
+		Ar << SavedStartIndex;
+		Ar << BentNormalNormalizeFactor;
 		Ar << DistanceFieldNormalTexture;
 		Ar << DistanceFieldNormalSampler;
+		Ar << RecordRadiusScale;
+		Ar << HeightfieldTexture;
+		Ar << HeightfieldSampler;
+		Ar << HeightfieldDescriptions;
+		Ar << NumHeightfields;
+		Ar << TanConeHalfAngle;
 		return bShaderHasOutdatedParameters;
 	}
 
 private:
 
-	FDeferredPixelShaderParameters DeferredParameters;
+	FAOParameters AOParameters;
+	FRWShaderParameter OccluderRadius;
+	FShaderResourceParameter IrradianceCacheNormal;
+	FShaderResourceParameter IrradianceCachePositionRadius;
+	FRWShaderParameter RecordConeVisibility;
+	FShaderResourceParameter ScatterDrawParameters;
+	FShaderResourceParameter SavedStartIndex;
+	FShaderParameter BentNormalNormalizeFactor;
 	FShaderResourceParameter DistanceFieldNormalTexture;
 	FShaderResourceParameter DistanceFieldNormalSampler;
+	FShaderParameter RecordRadiusScale;
+	FShaderResourceParameter HeightfieldTexture;
+	FShaderResourceParameter HeightfieldSampler;
+	FShaderResourceParameter HeightfieldDescriptions;
+	FShaderParameter NumHeightfields;
+	FShaderParameter TanConeHalfAngle;
 };
 
-IMPLEMENT_SHADER_TYPE(,FWriteDownsampledDepthPS,TEXT("DistanceFieldSurfaceCacheLighting"),TEXT("WriteDownsampledDepthPS"),SF_Pixel);
+IMPLEMENT_SHADER_TYPE(,FCalculateHeightfieldOcclusionCS,TEXT("DistanceFieldSurfaceCacheLightingCompute"),TEXT("CalculateHeightfieldOcclusionCS"),SF_Compute);
 
 
-/**  */
-template<bool bFinalInterpolationPass, bool bSupportIrradiance>
-class TIrradianceCacheSplatVS : public FGlobalShader
+template<bool bSupportIrradiance>
+class TCombineConesCS : public FGlobalShader
 {
-	DECLARE_SHADER_TYPE(TIrradianceCacheSplatVS,Global);
+	DECLARE_SHADER_TYPE(TCombineConesCS,Global)
 public:
 
-	static bool ShouldCache(EShaderPlatform Platform) 
-	{ 
+	static bool ShouldCache(EShaderPlatform Platform)
+	{
 		return IsFeatureLevelSupported(Platform, ERHIFeatureLevel::SM5) && DoesPlatformSupportDistanceFieldAO(Platform);
 	}
 
 	static void ModifyCompilationEnvironment(EShaderPlatform Platform, FShaderCompilerEnvironment& OutEnvironment)
 	{
 		FGlobalShader::ModifyCompilationEnvironment(Platform,OutEnvironment);
-
-		OutEnvironment.SetDefine(TEXT("FINAL_INTERPOLATION_PASS"), (uint32)bFinalInterpolationPass);
+		OutEnvironment.SetDefine(TEXT("NUM_CONE_DIRECTIONS"), NumConeSampleDirections);
 		OutEnvironment.SetDefine(TEXT("SUPPORT_IRRADIANCE"), bSupportIrradiance ? TEXT("1") : TEXT("0"));
+
+		// To reduce shader compile time of compute shaders with shared memory, doesn't have an impact on generated code with current compiler (June 2010 DX SDK)
+		OutEnvironment.CompilerFlags.Add(CFLAG_StandardOptimization);
 	}
 
-	TIrradianceCacheSplatVS(const ShaderMetaType::CompiledShaderInitializerType& Initializer):
-		FGlobalShader(Initializer)
+	TCombineConesCS(const ShaderMetaType::CompiledShaderInitializerType& Initializer)
+		: FGlobalShader(Initializer)
 	{
-		IrradianceCachePositionRadius.Bind(Initializer.ParameterMap, TEXT("IrradianceCachePositionRadius"));
-		IrradianceCacheOccluderRadius.Bind(Initializer.ParameterMap, TEXT("IrradianceCacheOccluderRadius"));
-		IrradianceCacheNormal.Bind(Initializer.ParameterMap, TEXT("IrradianceCacheNormal"));
+		AOParameters.Bind(Initializer.ParameterMap);
+		OccluderRadius.Bind(Initializer.ParameterMap, TEXT("OccluderRadius"));
 		IrradianceCacheBentNormal.Bind(Initializer.ParameterMap, TEXT("IrradianceCacheBentNormal"));
-		IrradianceCacheIrradiance.Bind(Initializer.ParameterMap, TEXT("IrradianceCacheIrradiance"));
-		InterpolationRadiusScale.Bind(Initializer.ParameterMap, TEXT("InterpolationRadiusScale"));
-		NormalizedOffsetToPixelCenter.Bind(Initializer.ParameterMap, TEXT("NormalizedOffsetToPixelCenter"));
-		HackExpand.Bind(Initializer.ParameterMap, TEXT("HackExpand"));
+		IrradianceCacheNormal.Bind(Initializer.ParameterMap, TEXT("IrradianceCacheNormal"));
+		RecordConeVisibility.Bind(Initializer.ParameterMap, TEXT("RecordConeVisibility"));
+		ScatterDrawParameters.Bind(Initializer.ParameterMap, TEXT("ScatterDrawParameters"));
+		SavedStartIndex.Bind(Initializer.ParameterMap, TEXT("SavedStartIndex"));
+		BentNormalNormalizeFactor.Bind(Initializer.ParameterMap, TEXT("BentNormalNormalizeFactor"));
+		DistanceFieldNormalTexture.Bind(Initializer.ParameterMap, TEXT("DistanceFieldNormalTexture"));
+		DistanceFieldNormalSampler.Bind(Initializer.ParameterMap, TEXT("DistanceFieldNormalSampler"));
+		RecordRadiusScale.Bind(Initializer.ParameterMap, TEXT("RecordRadiusScale"));
 	}
 
-	TIrradianceCacheSplatVS() {}
+	TCombineConesCS()
+	{
+	}
 
-	void SetParameters(FRHICommandList& RHICmdList, const FSceneView& View, int32 DepthLevel, int32 CurrentLevelDownsampleFactorValue, FVector2D NormalizedOffsetToPixelCenterValue)
+	void SetParameters(
+		FRHICommandList& RHICmdList, 
+		const FSceneView& View, 
+		int32 DepthLevel, 
+		const FDistanceFieldAOParameters& Parameters)
 	{
-		const FVertexShaderRHIParamRef ShaderRHI = GetVertexShader();
+		FComputeShaderRHIParamRef ShaderRHI = GetComputeShader();
 		FGlobalShader::SetParameters(RHICmdList, ShaderRHI, View);
 		
+		AOParameters.Set(RHICmdList, ShaderRHI, Parameters);
+
 		const FScene* Scene = (const FScene*)View.Family->Scene;
 		FSurfaceCacheResources& SurfaceCacheResources = *Scene->SurfaceCacheResources;
 
-		SetSRVParameter(RHICmdList, ShaderRHI, IrradianceCachePositionRadius, SurfaceCacheResources.Level[DepthLevel]->PositionAndRadius.SRV);
-		SetSRVParameter(RHICmdList, ShaderRHI, IrradianceCacheOccluderRadius, SurfaceCacheResources.Level[DepthLevel]->OccluderRadius.SRV);
 		SetSRVParameter(RHICmdList, ShaderRHI, IrradianceCacheNormal, SurfaceCacheResources.Level[DepthLevel]->Normal.SRV);
-		SetSRVParameter(RHICmdList, ShaderRHI, IrradianceCacheBentNormal, SurfaceCacheResources.Level[DepthLevel]->BentNormal.SRV);
-		SetSRVParameter(RHICmdList, ShaderRHI, IrradianceCacheIrradiance, SurfaceCacheResources.Level[DepthLevel]->Irradiance.SRV);
+		SetSRVParameter(RHICmdList, ShaderRHI, ScatterDrawParameters, SurfaceCacheResources.Level[DepthLevel]->ScatterDrawParameters.SRV);
+		SetSRVParameter(RHICmdList, ShaderRHI, SavedStartIndex, SurfaceCacheResources.Level[DepthLevel]->SavedStartIndex.SRV);
+		SetSRVParameter(RHICmdList, ShaderRHI, RecordConeVisibility, GTemporaryIrradianceCacheResources.ConeVisibility.SRV);
 
-		SetShaderValue(RHICmdList, ShaderRHI, InterpolationRadiusScale, (bFinalInterpolationPass ? GAOInterpolationRadiusScale : 1.0f));
-		SetShaderValue(RHICmdList, ShaderRHI, NormalizedOffsetToPixelCenter, NormalizedOffsetToPixelCenterValue);
+		OccluderRadius.SetBuffer(RHICmdList, ShaderRHI, SurfaceCacheResources.Level[DepthLevel]->OccluderRadius);
+		IrradianceCacheBentNormal.SetBuffer(RHICmdList, ShaderRHI, SurfaceCacheResources.Level[DepthLevel]->BentNormal);
+		
+		FAOSampleData2 AOSampleData;
 
-		const FIntPoint AOViewRectSize = FIntPoint::DivideAndRoundUp(View.ViewRect.Size(), CurrentLevelDownsampleFactorValue);
-		const FVector2D HackExpandValue(.5f / AOViewRectSize.X, .5f / AOViewRectSize.Y);
-		SetShaderValue(RHICmdList, ShaderRHI, HackExpand, HackExpandValue);
+		TArray<FVector, TInlineAllocator<9> > SampleDirections;
+		GetSpacedVectors(SampleDirections);
+
+		for (int32 SampleIndex = 0; SampleIndex < NumConeSampleDirections; SampleIndex++)
+		{
+			AOSampleData.SampleDirections[SampleIndex] = FVector4(SampleDirections[SampleIndex]);
+		}
+
+		SetUniformBufferParameterImmediate(RHICmdList, ShaderRHI, GetUniformBufferParameter<FAOSampleData2>(), AOSampleData);
+
+		FVector UnoccludedVector(0);
+
+		for (int32 SampleIndex = 0; SampleIndex < NumConeSampleDirections; SampleIndex++)
+		{
+			UnoccludedVector += SampleDirections[SampleIndex];
+		}
+
+		float BentNormalNormalizeFactorValue = 1.0f / (UnoccludedVector / NumConeSampleDirections).Size();
+		SetShaderValue(RHICmdList, ShaderRHI, BentNormalNormalizeFactor, BentNormalNormalizeFactorValue);
+
+		SetShaderValue(RHICmdList, ShaderRHI, RecordRadiusScale, GAORecordRadiusScale);
 	}
 
-	virtual bool Serialize(FArchive& Ar)
+	void UnsetParameters(FRHICommandList& RHICmdList)
 	{
-		bool bShaderHasOutdatedParameters = FGlobalShader::Serialize(Ar);
-		Ar << IrradianceCachePositionRadius;
-		Ar << IrradianceCacheOccluderRadius;
-		Ar << IrradianceCacheNormal;
-		Ar << IrradianceCacheBentNormal;
-		Ar << IrradianceCacheIrradiance;
-		Ar << InterpolationRadiusScale;
-		Ar << NormalizedOffsetToPixelCenter;
-		Ar << HackExpand;
+		OccluderRadius.UnsetUAV(RHICmdList, GetComputeShader());
+		IrradianceCacheBentNormal.UnsetUAV(RHICmdList, GetComputeShader());
+	}
+
+	virtual bool Serialize(FArchive& Ar)
+	{		
+		bool bShaderHasOutdatedParameters = FGlobalShader::Serialize(Ar);
+		Ar << AOParameters;
+		Ar << OccluderRadius;
+		Ar << IrradianceCacheBentNormal;
+		Ar << IrradianceCacheNormal;
+		Ar << RecordConeVisibility;
+		Ar << ScatterDrawParameters;
+		Ar << SavedStartIndex;
+		Ar << BentNormalNormalizeFactor;
+		Ar << DistanceFieldNormalTexture;
+		Ar << DistanceFieldNormalSampler;
+		Ar << RecordRadiusScale;
 		return bShaderHasOutdatedParameters;
 	}
 
 private:
-	FShaderResourceParameter IrradianceCachePositionRadius;
-	FShaderResourceParameter IrradianceCacheOccluderRadius;
+
+	FAOParameters AOParameters;
+	FRWShaderParameter OccluderRadius;
+	FRWShaderParameter IrradianceCacheBentNormal;
 	FShaderResourceParameter IrradianceCacheNormal;
-	FShaderResourceParameter IrradianceCacheBentNormal;
-	FShaderResourceParameter IrradianceCacheIrradiance;
-	FShaderParameter InterpolationRadiusScale;
-	FShaderParameter NormalizedOffsetToPixelCenter;
-	FShaderParameter HackExpand;
+	FShaderResourceParameter RecordConeVisibility;
+	FShaderResourceParameter ScatterDrawParameters;
+	FShaderResourceParameter SavedStartIndex;
+	FShaderParameter BentNormalNormalizeFactor;
+	FShaderResourceParameter DistanceFieldNormalTexture;
+	FShaderResourceParameter DistanceFieldNormalSampler;
+	FShaderParameter RecordRadiusScale;
 };
 
-// typedef required to get around macro expansion failure due to commas in template argument list 
-#define IMPLEMENT_SPLAT_VS_TYPE(bFinalInterpolationPass, bSupportIrradiance) \
-	typedef TIrradianceCacheSplatVS<bFinalInterpolationPass, bSupportIrradiance> TIrradianceCacheSplatVS##bFinalInterpolationPass##bSupportIrradiance; \
-	IMPLEMENT_SHADER_TYPE(template<>,TIrradianceCacheSplatVS##bFinalInterpolationPass##bSupportIrradiance,TEXT("DistanceFieldSurfaceCacheLighting"),TEXT("IrradianceCacheSplatVS"),SF_Vertex);
-
-IMPLEMENT_SPLAT_VS_TYPE(true, true)
-IMPLEMENT_SPLAT_VS_TYPE(false, true)
-IMPLEMENT_SPLAT_VS_TYPE(true, false)
-IMPLEMENT_SPLAT_VS_TYPE(false, false)
+IMPLEMENT_SHADER_TYPE(template<>,TCombineConesCS<true>,TEXT("DistanceFieldSurfaceCacheLightingCompute"),TEXT("CombineConesCS"),SF_Compute);
+IMPLEMENT_SHADER_TYPE(template<>,TCombineConesCS<false>,TEXT("DistanceFieldSurfaceCacheLightingCompute"),TEXT("CombineConesCS"),SF_Compute);
 
-template<bool bFinalInterpolationPass, bool bSupportIrradiance>
-class TIrradianceCacheSplatPS : public FGlobalShader
+class FWriteDownsampledDepthPS : public FGlobalShader
 {
-	DECLARE_SHADER_TYPE(TIrradianceCacheSplatPS, Global);
+	DECLARE_SHADER_TYPE(FWriteDownsampledDepthPS, Global);
 public:
 
 	static bool ShouldCache(EShaderPlatform Platform)
@@ -2631,92 +2845,146 @@ public:
 
 	static void ModifyCompilationEnvironment(EShaderPlatform Platform, FShaderCompilerEnvironment& OutEnvironment)
 	{
-		FGlobalShader::ModifyCompilationEnvironment(Platform,OutEnvironment);
-		
 		OutEnvironment.SetDefine(TEXT("DOWNSAMPLE_FACTOR"), GAODownsampleFactor);
-		OutEnvironment.SetDefine(TEXT("FINAL_INTERPOLATION_PASS"), (uint32)bFinalInterpolationPass);
-		OutEnvironment.SetDefine(TEXT("SUPPORT_IRRADIANCE"), bSupportIrradiance ? TEXT("1") : TEXT("0"));
 	}
 
 	/** Default constructor. */
-	TIrradianceCacheSplatPS() {}
+	FWriteDownsampledDepthPS() {}
 
 	/** Initialization constructor. */
-	TIrradianceCacheSplatPS(const ShaderMetaType::CompiledShaderInitializerType& Initializer)
+	FWriteDownsampledDepthPS(const ShaderMetaType::CompiledShaderInitializerType& Initializer)
 		: FGlobalShader(Initializer)
 	{
-		AOParameters.Bind(Initializer.ParameterMap);
 		DeferredParameters.Bind(Initializer.ParameterMap);
-		AOLevelParameters.Bind(Initializer.ParameterMap);
-		DistanceFieldNormalTexture.Bind(Initializer.ParameterMap, TEXT("DistanceFieldNormalTexture"));
-		DistanceFieldNormalSampler.Bind(Initializer.ParameterMap, TEXT("DistanceFieldNormalSampler"));
-		InterpolationAngleNormalization.Bind(Initializer.ParameterMap, TEXT("InterpolationAngleNormalization"));
-		InvMinCosPointBehindPlane.Bind(Initializer.ParameterMap, TEXT("InvMinCosPointBehindPlane"));
+		DistanceFieldNormalTexture.Bind(Initializer.ParameterMap,TEXT("DistanceFieldNormalTexture"));
+		DistanceFieldNormalSampler.Bind(Initializer.ParameterMap,TEXT("DistanceFieldNormalSampler"));
 	}
 
-	void SetParameters(FRHICommandList& RHICmdList, const FSceneView& View, FSceneRenderTargetItem& DistanceFieldNormal, int32 DestLevelDownsampleFactor, const FDistanceFieldAOParameters& Parameters)
+	void SetParameters(FRHICommandList& RHICmdList, const FSceneView& View, FSceneRenderTargetItem& DistanceFieldNormal)
 	{
 		const FPixelShaderRHIParamRef ShaderRHI = GetPixelShader();
+
 		FGlobalShader::SetParameters(RHICmdList, ShaderRHI, View);
-		AOParameters.Set(RHICmdList, ShaderRHI, Parameters);
 		DeferredParameters.Set(RHICmdList, ShaderRHI, View);
-		AOLevelParameters.Set(RHICmdList, ShaderRHI, View, DestLevelDownsampleFactor);
 
-		SetTextureParameter(
-			RHICmdList,
-			ShaderRHI,
-			DistanceFieldNormalTexture,
-			DistanceFieldNormalSampler,
-			TStaticSamplerState<SF_Point,AM_Clamp,AM_Clamp,AM_Clamp>::GetRHI(),
-			DistanceFieldNormal.ShaderResourceTexture
-			);
-
-		const float EffectiveMaxAngle = bFinalInterpolationPass ? GAOInterpolationMaxAngle * GAOInterpolationAngleScale : GAOInterpolationMaxAngle;
-		const float InterpolationAngleNormalizationValue = 1.0f / FMath::Sqrt(1.0f - FMath::Cos(EffectiveMaxAngle * PI / 180.0f));
-		SetShaderValue(RHICmdList, ShaderRHI, InterpolationAngleNormalization, InterpolationAngleNormalizationValue);
-
-		const float MinCosPointBehindPlaneValue = FMath::Cos((GAOMinPointBehindPlaneAngle + 90.0f) * PI / 180.0f);
-		SetShaderValue(RHICmdList, ShaderRHI, InvMinCosPointBehindPlane, 1.0f / MinCosPointBehindPlaneValue);
+		SetTextureParameter(RHICmdList, ShaderRHI, DistanceFieldNormalTexture, DistanceFieldNormalSampler, TStaticSamplerState<SF_Point>::GetRHI(), DistanceFieldNormal.ShaderResourceTexture);
 	}
-
 	// FShader interface.
 	virtual bool Serialize(FArchive& Ar)
 	{
 		bool bShaderHasOutdatedParameters = FGlobalShader::Serialize(Ar);
-		Ar << AOParameters;
-		Ar << AOLevelParameters;
 		Ar << DeferredParameters;
 		Ar << DistanceFieldNormalTexture;
 		Ar << DistanceFieldNormalSampler;
-		Ar << InterpolationAngleNormalization;
-		Ar << InvMinCosPointBehindPlane;
 		return bShaderHasOutdatedParameters;
 	}
 
 private:
-	FAOParameters AOParameters;
-	FAOLevelParameters AOLevelParameters;
+
 	FDeferredPixelShaderParameters DeferredParameters;
 	FShaderResourceParameter DistanceFieldNormalTexture;
 	FShaderResourceParameter DistanceFieldNormalSampler;
-	FShaderParameter InterpolationAngleNormalization;
-	FShaderParameter InvMinCosPointBehindPlane;
+};
+
+IMPLEMENT_SHADER_TYPE(,FWriteDownsampledDepthPS,TEXT("DistanceFieldSurfaceCacheLighting"),TEXT("WriteDownsampledDepthPS"),SF_Pixel);
+
+
+/**  */
+template<bool bFinalInterpolationPass, bool bSupportIrradiance>
+class TIrradianceCacheSplatVS : public FGlobalShader
+{
+	DECLARE_SHADER_TYPE(TIrradianceCacheSplatVS,Global);
+public:
+
+	static bool ShouldCache(EShaderPlatform Platform) 
+	{ 
+		return IsFeatureLevelSupported(Platform, ERHIFeatureLevel::SM5) && DoesPlatformSupportDistanceFieldAO(Platform);
+	}
+
+	static void ModifyCompilationEnvironment(EShaderPlatform Platform, FShaderCompilerEnvironment& OutEnvironment)
+	{
+		FGlobalShader::ModifyCompilationEnvironment(Platform,OutEnvironment);
+
+		OutEnvironment.SetDefine(TEXT("FINAL_INTERPOLATION_PASS"), (uint32)bFinalInterpolationPass);
+		OutEnvironment.SetDefine(TEXT("SUPPORT_IRRADIANCE"), bSupportIrradiance ? TEXT("1") : TEXT("0"));
+	}
+
+	TIrradianceCacheSplatVS(const ShaderMetaType::CompiledShaderInitializerType& Initializer):
+		FGlobalShader(Initializer)
+	{
+		IrradianceCachePositionRadius.Bind(Initializer.ParameterMap, TEXT("IrradianceCachePositionRadius"));
+		IrradianceCacheOccluderRadius.Bind(Initializer.ParameterMap, TEXT("IrradianceCacheOccluderRadius"));
+		IrradianceCacheNormal.Bind(Initializer.ParameterMap, TEXT("IrradianceCacheNormal"));
+		IrradianceCacheBentNormal.Bind(Initializer.ParameterMap, TEXT("IrradianceCacheBentNormal"));
+		IrradianceCacheIrradiance.Bind(Initializer.ParameterMap, TEXT("IrradianceCacheIrradiance"));
+		InterpolationRadiusScale.Bind(Initializer.ParameterMap, TEXT("InterpolationRadiusScale"));
+		NormalizedOffsetToPixelCenter.Bind(Initializer.ParameterMap, TEXT("NormalizedOffsetToPixelCenter"));
+		HackExpand.Bind(Initializer.ParameterMap, TEXT("HackExpand"));
+	}
+
+	TIrradianceCacheSplatVS() {}
+
+	void SetParameters(FRHICommandList& RHICmdList, const FSceneView& View, int32 DepthLevel, int32 CurrentLevelDownsampleFactorValue, FVector2D NormalizedOffsetToPixelCenterValue)
+	{
+		const FVertexShaderRHIParamRef ShaderRHI = GetVertexShader();
+		FGlobalShader::SetParameters(RHICmdList, ShaderRHI, View);
+		
+		const FScene* Scene = (const FScene*)View.Family->Scene;
+		FSurfaceCacheResources& SurfaceCacheResources = *Scene->SurfaceCacheResources;
+
+		SetSRVParameter(RHICmdList, ShaderRHI, IrradianceCachePositionRadius, SurfaceCacheResources.Level[DepthLevel]->PositionAndRadius.SRV);
+		SetSRVParameter(RHICmdList, ShaderRHI, IrradianceCacheOccluderRadius, SurfaceCacheResources.Level[DepthLevel]->OccluderRadius.SRV);
+		SetSRVParameter(RHICmdList, ShaderRHI, IrradianceCacheNormal, SurfaceCacheResources.Level[DepthLevel]->Normal.SRV);
+		SetSRVParameter(RHICmdList, ShaderRHI, IrradianceCacheBentNormal, SurfaceCacheResources.Level[DepthLevel]->BentNormal.SRV);
+		SetSRVParameter(RHICmdList, ShaderRHI, IrradianceCacheIrradiance, SurfaceCacheResources.Level[DepthLevel]->Irradiance.SRV);
+
+		SetShaderValue(RHICmdList, ShaderRHI, InterpolationRadiusScale, (bFinalInterpolationPass ? GAOInterpolationRadiusScale : 1.0f));
+		SetShaderValue(RHICmdList, ShaderRHI, NormalizedOffsetToPixelCenter, NormalizedOffsetToPixelCenterValue);
+
+		const FIntPoint AOViewRectSize = FIntPoint::DivideAndRoundUp(View.ViewRect.Size(), CurrentLevelDownsampleFactorValue);
+		const FVector2D HackExpandValue(.5f / AOViewRectSize.X, .5f / AOViewRectSize.Y);
+		SetShaderValue(RHICmdList, ShaderRHI, HackExpand, HackExpandValue);
+	}
+
+	virtual bool Serialize(FArchive& Ar)
+	{
+		bool bShaderHasOutdatedParameters = FGlobalShader::Serialize(Ar);
+		Ar << IrradianceCachePositionRadius;
+		Ar << IrradianceCacheOccluderRadius;
+		Ar << IrradianceCacheNormal;
+		Ar << IrradianceCacheBentNormal;
+		Ar << IrradianceCacheIrradiance;
+		Ar << InterpolationRadiusScale;
+		Ar << NormalizedOffsetToPixelCenter;
+		Ar << HackExpand;
+		return bShaderHasOutdatedParameters;
+	}
+
+private:
+	FShaderResourceParameter IrradianceCachePositionRadius;
+	FShaderResourceParameter IrradianceCacheOccluderRadius;
+	FShaderResourceParameter IrradianceCacheNormal;
+	FShaderResourceParameter IrradianceCacheBentNormal;
+	FShaderResourceParameter IrradianceCacheIrradiance;
+	FShaderParameter InterpolationRadiusScale;
+	FShaderParameter NormalizedOffsetToPixelCenter;
+	FShaderParameter HackExpand;
 };
 
 // typedef required to get around macro expansion failure due to commas in template argument list 
-#define IMPLEMENT_SPLAT_PS_TYPE(bFinalInterpolationPass, bSupportIrradiance) \
-	typedef TIrradianceCacheSplatPS<bFinalInterpolationPass, bSupportIrradiance> TIrradianceCacheSplatPS##bFinalInterpolationPass##bSupportIrradiance; \
-	IMPLEMENT_SHADER_TYPE(template<>,TIrradianceCacheSplatPS##bFinalInterpolationPass##bSupportIrradiance,TEXT("DistanceFieldSurfaceCacheLighting"),TEXT("IrradianceCacheSplatPS"),SF_Pixel);
+#define IMPLEMENT_SPLAT_VS_TYPE(bFinalInterpolationPass, bSupportIrradiance) \
+	typedef TIrradianceCacheSplatVS<bFinalInterpolationPass, bSupportIrradiance> TIrradianceCacheSplatVS##bFinalInterpolationPass##bSupportIrradiance; \
+	IMPLEMENT_SHADER_TYPE(template<>,TIrradianceCacheSplatVS##bFinalInterpolationPass##bSupportIrradiance,TEXT("DistanceFieldSurfaceCacheLighting"),TEXT("IrradianceCacheSplatVS"),SF_Vertex);
 
-IMPLEMENT_SPLAT_PS_TYPE(true, true)
-IMPLEMENT_SPLAT_PS_TYPE(false, true)
-IMPLEMENT_SPLAT_PS_TYPE(true, false)
-IMPLEMENT_SPLAT_PS_TYPE(false, false)
+IMPLEMENT_SPLAT_VS_TYPE(true, true)
+IMPLEMENT_SPLAT_VS_TYPE(false, true)
+IMPLEMENT_SPLAT_VS_TYPE(true, false)
+IMPLEMENT_SPLAT_VS_TYPE(false, false)
 
-template<uint32 NumLevels, bool bFinalInterpolationPass, bool bSupportIrradiance>
-class TInterpolateCacheCS : public FGlobalShader
+template<bool bFinalInterpolationPass, bool bSupportIrradiance>
+class TIrradianceCacheSplatPS : public FGlobalShader
 {
-	DECLARE_SHADER_TYPE(TInterpolateCacheCS,Global)
+	DECLARE_SHADER_TYPE(TIrradianceCacheSplatPS, Global);
 public:
 
 	static bool ShouldCache(EShaderPlatform Platform)
@@ -2727,17 +2995,17 @@ public:
 	static void ModifyCompilationEnvironment(EShaderPlatform Platform, FShaderCompilerEnvironment& OutEnvironment)
 	{
 		FGlobalShader::ModifyCompilationEnvironment(Platform,OutEnvironment);
-
+		
+		OutEnvironment.SetDefine(TEXT("DOWNSAMPLE_FACTOR"), GAODownsampleFactor);
 		OutEnvironment.SetDefine(TEXT("FINAL_INTERPOLATION_PASS"), (uint32)bFinalInterpolationPass);
-		OutEnvironment.SetDefine(TEXT("NUM_INTERPOLATE_LEVELS"), NumLevels);
-		OutEnvironment.SetDefine(TEXT("THREADGROUP_SIZEX"), GDistanceFieldAOTileSizeX);
-		OutEnvironment.SetDefine(TEXT("THREADGROUP_SIZEY"), GDistanceFieldAOTileSizeY);
 		OutEnvironment.SetDefine(TEXT("SUPPORT_IRRADIANCE"), bSupportIrradiance ? TEXT("1") : TEXT("0"));
-
-		OutEnvironment.CompilerFlags.Add(CFLAG_StandardOptimization);
 	}
 
-	TInterpolateCacheCS(const ShaderMetaType::CompiledShaderInitializerType& Initializer)
+	/** Default constructor. */
+	TIrradianceCacheSplatPS() {}
+
+	/** Initialization constructor. */
+	TIrradianceCacheSplatPS(const ShaderMetaType::CompiledShaderInitializerType& Initializer)
 		: FGlobalShader(Initializer)
 	{
 		AOParameters.Bind(Initializer.ParameterMap);
@@ -2747,24 +3015,11 @@ public:
 		DistanceFieldNormalSampler.Bind(Initializer.ParameterMap, TEXT("DistanceFieldNormalSampler"));
 		InterpolationAngleNormalization.Bind(Initializer.ParameterMap, TEXT("InterpolationAngleNormalization"));
 		InvMinCosPointBehindPlane.Bind(Initializer.ParameterMap, TEXT("InvMinCosPointBehindPlane"));
-		IrradianceCacheSplat.Bind(Initializer.ParameterMap, TEXT("IrradianceCacheSplat"));
-		IrradianceCachePositionRadius.Bind(Initializer.ParameterMap, TEXT("LevelIrradianceCachePositionRadius"));
-		IrradianceCacheOccluderRadius.Bind(Initializer.ParameterMap, TEXT("LevelIrradianceCacheOccluderRadius"));
-		IrradianceCacheNormal.Bind(Initializer.ParameterMap, TEXT("LevelIrradianceCacheNormal"));
-		IrradianceCacheBentNormal.Bind(Initializer.ParameterMap, TEXT("LevelIrradianceCacheBentNormal"));
-		IrradianceCacheIrradiance.Bind(Initializer.ParameterMap, TEXT("LevelIrradianceCacheIrradiance"));
-		InterpolationRadiusScale.Bind(Initializer.ParameterMap, TEXT("InterpolationRadiusScale"));
-		DrawParameters.Bind(Initializer.ParameterMap, TEXT("LevelScatterDrawParameters"));
-		NumGroups.Bind(Initializer.ParameterMap, TEXT("NumGroups"));
-	}
-
-	TInterpolateCacheCS()
-	{
 	}
 
-	void SetParameters(FRHICommandList& RHICmdList, const FSceneView& View, FSceneRenderTargetItem& DistanceFieldNormal, int32 DepthLevel, int32 DestLevelDownsampleFactor, IPooledRenderTarget* BentNormalInterpolationTarget, const FDistanceFieldAOParameters& Parameters, FVector2D NumGroupsValue, FSurfaceCacheResources& SurfaceCacheResources)
+	void SetParameters(FRHICommandList& RHICmdList, const FSceneView& View, FSceneRenderTargetItem& DistanceFieldNormal, int32 DestLevelDownsampleFactor, const FDistanceFieldAOParameters& Parameters)
 	{
-		FComputeShaderRHIParamRef ShaderRHI = GetComputeShader();
+		const FPixelShaderRHIParamRef ShaderRHI = GetPixelShader();
 		FGlobalShader::SetParameters(RHICmdList, ShaderRHI, View);
 		AOParameters.Set(RHICmdList, ShaderRHI, Parameters);
 		DeferredParameters.Set(RHICmdList, ShaderRHI, View);
@@ -2785,97 +3040,11 @@ public:
 
 		const float MinCosPointBehindPlaneValue = FMath::Cos((GAOMinPointBehindPlaneAngle + 90.0f) * PI / 180.0f);
 		SetShaderValue(RHICmdList, ShaderRHI, InvMinCosPointBehindPlane, 1.0f / MinCosPointBehindPlaneValue);
-
-		IrradianceCacheSplat.SetTexture(
-			RHICmdList,
-			ShaderRHI, 
-			BentNormalInterpolationTarget->GetRenderTargetItem().TargetableTexture, 
-			BentNormalInterpolationTarget->GetRenderTargetItem().UAV);
-
-		const int32 MaxAllowedLevel = ARRAY_COUNT(SurfaceCacheResources.Level) - 1;
-		const int32 MinLevel = MaxAllowedLevel + 1 - NumLevels;
-
-		for (int32 LevelIndex = MaxAllowedLevel; LevelIndex >= MinLevel; LevelIndex--)
-		{
-			if (IrradianceCachePositionRadius.IsBound())
-			{
-				check(MaxAllowedLevel - LevelIndex < (int32)IrradianceCachePositionRadius.GetNumResources());
-
-				RHICmdList.SetShaderResourceViewParameter(
-					ShaderRHI,
-					IrradianceCachePositionRadius.GetBaseIndex() + MaxAllowedLevel - LevelIndex,
-					SurfaceCacheResources.Level[LevelIndex]->PositionAndRadius.SRV
-					);
-			}
-
-			if (IrradianceCacheOccluderRadius.IsBound())
-			{
-				check(MaxAllowedLevel - LevelIndex < (int32)IrradianceCacheOccluderRadius.GetNumResources());
-
-				RHICmdList.SetShaderResourceViewParameter(
-					ShaderRHI,
-					IrradianceCacheOccluderRadius.GetBaseIndex() + MaxAllowedLevel - LevelIndex,
-					SurfaceCacheResources.Level[LevelIndex]->OccluderRadius.SRV
-					);
-			}
-
-			if (IrradianceCacheNormal.IsBound())
-			{
-				check(MaxAllowedLevel - LevelIndex < (int32)IrradianceCacheNormal.GetNumResources());
-
-				RHICmdList.SetShaderResourceViewParameter(
-					ShaderRHI,
-					IrradianceCacheNormal.GetBaseIndex() + MaxAllowedLevel - LevelIndex,
-					SurfaceCacheResources.Level[LevelIndex]->Normal.SRV
-					);
-			}
-
-			if (IrradianceCacheBentNormal.IsBound())
-			{
-				check(MaxAllowedLevel - LevelIndex < (int32)IrradianceCacheBentNormal.GetNumResources());
-
-				RHICmdList.SetShaderResourceViewParameter(
-					ShaderRHI,
-					IrradianceCacheBentNormal.GetBaseIndex() + MaxAllowedLevel - LevelIndex,
-					SurfaceCacheResources.Level[LevelIndex]->BentNormal.SRV
-					);
-			}
-
-			if (IrradianceCacheIrradiance.IsBound())
-			{
-				check(MaxAllowedLevel - LevelIndex < (int32)IrradianceCacheIrradiance.GetNumResources());
-
-				RHICmdList.SetShaderResourceViewParameter(
-					ShaderRHI,
-					IrradianceCacheIrradiance.GetBaseIndex() + MaxAllowedLevel - LevelIndex,
-					SurfaceCacheResources.Level[LevelIndex]->Irradiance.SRV
-					);
-			}
-
-			if (DrawParameters.IsBound())
-			{
-				check(MaxAllowedLevel - LevelIndex < (int32)DrawParameters.GetNumResources());
-
-				RHICmdList.SetShaderResourceViewParameter(
-					ShaderRHI,
-					DrawParameters.GetBaseIndex() + MaxAllowedLevel - LevelIndex,
-					SurfaceCacheResources.Level[LevelIndex]->ScatterDrawParameters.SRV
-					);
-			}
-		}
-
-		SetShaderValue(RHICmdList, ShaderRHI, InterpolationRadiusScale, (bFinalInterpolationPass ? GAOInterpolationRadiusScale : 1.0f));
-
-		SetShaderValue(RHICmdList, ShaderRHI, NumGroups, NumGroupsValue);
-	}
-
-	void UnsetParameters(FRHICommandListImmediate& RHICmdList)
-	{
-		IrradianceCacheSplat.UnsetUAV(RHICmdList, GetComputeShader());
 	}
 
+	// FShader interface.
 	virtual bool Serialize(FArchive& Ar)
-	{		
+	{
 		bool bShaderHasOutdatedParameters = FGlobalShader::Serialize(Ar);
 		Ar << AOParameters;
 		Ar << AOLevelParameters;
@@ -2884,20 +3053,10 @@ public:
 		Ar << DistanceFieldNormalSampler;
 		Ar << InterpolationAngleNormalization;
 		Ar << InvMinCosPointBehindPlane;
-		Ar << IrradianceCacheSplat;
-		Ar << IrradianceCachePositionRadius;
-		Ar << IrradianceCacheOccluderRadius;
-		Ar << IrradianceCacheNormal;
-		Ar << IrradianceCacheBentNormal;
-		Ar << IrradianceCacheIrradiance;
-		Ar << InterpolationRadiusScale;
-		Ar << DrawParameters;
-		Ar << NumGroups;
 		return bShaderHasOutdatedParameters;
 	}
 
 private:
-
 	FAOParameters AOParameters;
 	FAOLevelParameters AOLevelParameters;
 	FDeferredPixelShaderParameters DeferredParameters;
@@ -2905,35 +3064,17 @@ private:
 	FShaderResourceParameter DistanceFieldNormalSampler;
 	FShaderParameter InterpolationAngleNormalization;
 	FShaderParameter InvMinCosPointBehindPlane;
-	FRWShaderParameter IrradianceCacheSplat;
-	FShaderResourceParameter IrradianceCachePositionRadius;
-	FShaderResourceParameter IrradianceCacheOccluderRadius;
-	FShaderResourceParameter IrradianceCacheNormal;
-	FShaderResourceParameter IrradianceCacheBentNormal;
-	FShaderResourceParameter IrradianceCacheIrradiance;
-	FShaderParameter InterpolationRadiusScale;
-	FShaderResourceParameter DrawParameters;
-	FShaderParameter NumGroups;
 };
 
-#define IMPLEMENT_INTERPOLATION_TYPE(NumLevels,bFinalInterpolationPass,bSupportIrradiance) \
-	typedef TInterpolateCacheCS<NumLevels,bFinalInterpolationPass,bSupportIrradiance> TInterpolateCacheCS##NumLevels##bFinalInterpolationPass##bSupportIrradiance; \
-	IMPLEMENT_SHADER_TYPE(template<>,TInterpolateCacheCS##NumLevels##bFinalInterpolationPass##bSupportIrradiance,TEXT("DistanceFieldSurfaceCacheLighting"),TEXT("InterpolateCacheCS"),SF_Compute);
-
-IMPLEMENT_INTERPOLATION_TYPE(5, true, true);
-IMPLEMENT_INTERPOLATION_TYPE(4, true, true);
-IMPLEMENT_INTERPOLATION_TYPE(1, false, true);
-IMPLEMENT_INTERPOLATION_TYPE(2, false, true);
-IMPLEMENT_INTERPOLATION_TYPE(3, false, true);
-IMPLEMENT_INTERPOLATION_TYPE(4, false, true);
-IMPLEMENT_INTERPOLATION_TYPE(5, false, true);
-IMPLEMENT_INTERPOLATION_TYPE(5, true, false);
-IMPLEMENT_INTERPOLATION_TYPE(4, true, false);
-IMPLEMENT_INTERPOLATION_TYPE(1, false, false);
-IMPLEMENT_INTERPOLATION_TYPE(2, false, false);
-IMPLEMENT_INTERPOLATION_TYPE(3, false, false);
-IMPLEMENT_INTERPOLATION_TYPE(4, false, false);
-IMPLEMENT_INTERPOLATION_TYPE(5, false, false);
+// typedef required to get around macro expansion failure due to commas in template argument list 
+#define IMPLEMENT_SPLAT_PS_TYPE(bFinalInterpolationPass, bSupportIrradiance) \
+	typedef TIrradianceCacheSplatPS<bFinalInterpolationPass, bSupportIrradiance> TIrradianceCacheSplatPS##bFinalInterpolationPass##bSupportIrradiance; \
+	IMPLEMENT_SHADER_TYPE(template<>,TIrradianceCacheSplatPS##bFinalInterpolationPass##bSupportIrradiance,TEXT("DistanceFieldSurfaceCacheLighting"),TEXT("IrradianceCacheSplatPS"),SF_Pixel);
+
+IMPLEMENT_SPLAT_PS_TYPE(true, true)
+IMPLEMENT_SPLAT_PS_TYPE(false, true)
+IMPLEMENT_SPLAT_PS_TYPE(true, false)
+IMPLEMENT_SPLAT_PS_TYPE(false, false)
 
 template<bool bSupportIrradiance>
 class TDistanceFieldAOCombinePS : public FGlobalShader
@@ -2971,7 +3112,13 @@ public:
 		DebugBuffer.Bind(Initializer.ParameterMap, TEXT("DebugBuffer"));
 	}
 
-	void SetParameters(FRHICommandList& RHICmdList, const FSceneView& View, FSceneRenderTargetItem& InBentNormalTexture, IPooledRenderTarget* InIrradianceTexture, FSceneRenderTargetItem& DistanceFieldNormal, const FDistanceFieldAOParameters& Parameters)
+	void SetParameters(
+		FRHICommandList& RHICmdList, 
+		const FSceneView& View, 
+		FSceneRenderTargetItem& InBentNormalTexture, 
+		IPooledRenderTarget* InIrradianceTexture, 
+		FSceneRenderTargetItem& DistanceFieldNormal, 
+		const FDistanceFieldAOParameters& Parameters)
 	{
 		const FPixelShaderRHIParamRef ShaderRHI = GetPixelShader();
 
@@ -2998,6 +3145,7 @@ public:
 
 		SetSRVParameter(RHICmdList, ShaderRHI, DebugBuffer, GDebugBuffer.DebugData.SRV);
 	}
+
 	// FShader interface.
 	virtual bool Serialize(FArchive& Ar)
 	{
@@ -4000,7 +4148,7 @@ void SetupDepthStencil(FRHICommandListImmediate& RHICmdList, FViewInfo& View, FS
 
 		SetShaderValue(RHICmdList, PixelShader->GetPixelShader(), PixelShader->ColorParameter, FLinearColor::Black);
 
-		FVector ViewSpaceMaxDistance(0.0f, 0.0f, GAOMaxViewDistance);
+		FVector ViewSpaceMaxDistance(0.0f, 0.0f, GetMaxAOViewDistance());
 		FVector4 ClipSpaceMaxDistance = View.ViewMatrices.ProjMatrix.TransformPosition(ViewSpaceMaxDistance);
 		float ClipSpaceZ = ClipSpaceMaxDistance.Z / ClipSpaceMaxDistance.W;
 
@@ -4031,42 +4179,6 @@ void RenderIrradianceCacheInterpolation(
 {
 	check(!(bFinalInterpolation && DepthLevel != 0));
 
-	if (bFinalInterpolation && !GAOScatterInterpolation)
-	{
-		for (int32 ViewIndex = 0; ViewIndex < Views.Num(); ViewIndex++)
-		{
-			FViewInfo& View = Views[ViewIndex];
-
-			const FScene* Scene = (const FScene*)View.Family->Scene;
-			FSurfaceCacheResources& SurfaceCacheResources = *Scene->SurfaceCacheResources;
-
-			FIntPoint DownsampledViewSize = FIntPoint::DivideAndRoundUp(View.ViewRect.Size(), DestLevelDownsampleFactor);
-			uint32 GroupSizeX = FMath::DivideAndRoundUp(DownsampledViewSize.X, GDistanceFieldAOTileSizeX);
-			uint32 GroupSizeY = FMath::DivideAndRoundUp(DownsampledViewSize.Y, GDistanceFieldAOTileSizeY);
-
-			if (GDistanceFieldGI)
-			{
-				TShaderMapRef<TInterpolateCacheCS<4, true, true> > ComputeShader(View.ShaderMap);
-
-				RHICmdList.SetComputeShader(ComputeShader->GetComputeShader());	
-				ComputeShader->SetParameters(RHICmdList, View, DistanceFieldNormal, DepthLevel, DestLevelDownsampleFactor, BentNormalInterpolationTarget, Parameters, FVector2D(GroupSizeX, GroupSizeY), SurfaceCacheResources);
-				DispatchComputeShader(RHICmdList, *ComputeShader, GroupSizeX, GroupSizeY, 1);
-
-				ComputeShader->UnsetParameters(RHICmdList);
-			}
-			else
-			{
-				TShaderMapRef<TInterpolateCacheCS<4, true, false> > ComputeShader(View.ShaderMap);
-
-				RHICmdList.SetComputeShader(ComputeShader->GetComputeShader());	
-				ComputeShader->SetParameters(RHICmdList, View, DistanceFieldNormal, DepthLevel, DestLevelDownsampleFactor, BentNormalInterpolationTarget, Parameters, FVector2D(GroupSizeX, GroupSizeY), SurfaceCacheResources);
-				DispatchComputeShader(RHICmdList, *ComputeShader, GroupSizeX, GroupSizeY, 1);
-
-				ComputeShader->UnsetParameters(RHICmdList);
-			}
-		}
-	}
-	else
 	{
 		TRefCountPtr<IPooledRenderTarget> SplatDepthStencilBuffer;
 
@@ -4260,6 +4372,23 @@ bool FDeferredShadingSceneRenderer::ShouldPrepareForDistanceFieldAO() const
 			|| ViewFamily.EngineShowFlags.VisualizeDistanceFieldAO);
 }
 
+class FHeightfieldDescription
+{
+public:
+	FVector4 HeightfieldScaleBias;
+	FVector4 MinMaxUV;
+	FMatrix LocalToWorld;
+
+	FHeightfieldDescription(
+		const FVector4& InHeightfieldScaleBias,
+		const FVector4& InMinMaxUV,
+		const FMatrix& InLocalToWorld) :
+		HeightfieldScaleBias(InHeightfieldScaleBias),
+		MinMaxUV(InMinMaxUV),
+		LocalToWorld(InLocalToWorld)
+	{}
+};
+
 bool FDeferredShadingSceneRenderer::RenderDistanceFieldAOSurfaceCache(
 	FRHICommandListImmediate& RHICmdList, 
 	const FDistanceFieldAOParameters& Parameters, 
@@ -4293,9 +4422,12 @@ bool FDeferredShadingSceneRenderer::RenderDistanceFieldAOSurfaceCache(
 
 			FSurfaceCacheResources& SurfaceCacheResources = *Scene->SurfaceCacheResources;
 			FIntPoint DownsampledViewRect = FIntPoint::DivideAndRoundDown(View.ViewRect.Size(), GAODownsampleFactor);
+			// Needed because samples are kept across frames
 			int32 SlackAmount = 4;
-			int32 MaxSurfaceCacheSamples = SlackAmount * DownsampledViewRect.X * DownsampledViewRect.Y / (1 << (2 * GAOMinLevel * GAOPowerOfTwoBetweenLevels));
-			SurfaceCacheResources.AllocateFor(GAOMinLevel, GAOMaxLevel, MaxSurfaceCacheSamples);
+			int32 MaxSurfaceCacheSamples = DownsampledViewRect.X * DownsampledViewRect.Y / (1 << (2 * GAOMinLevel * GAOPowerOfTwoBetweenLevels));
+			SurfaceCacheResources.AllocateFor(GAOMinLevel, GAOMaxLevel, SlackAmount * MaxSurfaceCacheSamples);
+
+			GTemporaryIrradianceCacheResources.AllocateFor(MaxSurfaceCacheSamples);
 
 			uint32 ClearValues[4] = {0};
 
@@ -4417,7 +4549,7 @@ bool FDeferredShadingSceneRenderer::RenderDistanceFieldAOSurfaceCache(
 				RenderIrradianceCacheInterpolation(RHICmdList, Views, DistanceFieldAOBentNormalSplat, NULL, DistanceFieldNormal->GetRenderTargetItem(), DepthLevel, DestLevelDownsampleFactor, Parameters, false);
 
 				{
-					SCOPED_DRAW_EVENT(RHICmdList, PopulateIrradianceCache);
+					SCOPED_DRAW_EVENT(RHICmdList, PopulateAndShadeIrradianceCache);
 					SetRenderTarget(RHICmdList, NULL, NULL);
 
 					// Save off the current record count before adding any more
@@ -4448,13 +4580,9 @@ bool FDeferredShadingSceneRenderer::RenderDistanceFieldAOSurfaceCache(
 
 						ComputeShader->UnsetParameters(RHICmdList);
 					}
-				}
-
-				{
-					SCOPED_DRAW_EVENT(RHICmdList, ShadeIrradianceCache);
 
 					{	
-						TShaderMapRef<FSetupFinalGatherIndirectArgumentsCS> ComputeShader(View.ShaderMap);
+						TShaderMapRef<TSetupFinalGatherIndirectArgumentsCS<true> > ComputeShader(View.ShaderMap);
 
 						RHICmdList.SetComputeShader(ComputeShader->GetComputeShader());
 						ComputeShader->SetParameters(RHICmdList, View, DepthLevel);
@@ -4485,6 +4613,118 @@ bool FDeferredShadingSceneRenderer::RenderDistanceFieldAOSurfaceCache(
 						ComputeShader->UnsetParameters(RHICmdList);
 					}
 				}
+
+				{	
+					TShaderMapRef<TSetupFinalGatherIndirectArgumentsCS<false> > ComputeShader(View.ShaderMap);
+
+					RHICmdList.SetComputeShader(ComputeShader->GetComputeShader());
+					ComputeShader->SetParameters(RHICmdList, View, DepthLevel);
+					DispatchComputeShader(RHICmdList, *ComputeShader, 1, 1, 1);
+
+					ComputeShader->UnsetParameters(RHICmdList);
+				}
+				
+				SCOPED_DRAW_EVENT(RHICmdList, CombineCones);
+
+				if (Scene->DistanceFieldSceneData.HeightfieldPrimitives.Num() > 0)
+				{
+					int32 NumPrimitives = Scene->DistanceFieldSceneData.HeightfieldPrimitives.Num();
+
+					TMap<UTexture2D*, TArray<FHeightfieldDescription, SceneRenderingAllocator>, SceneRenderingSetAllocator> HeightfieldTextures; 
+
+					for (int32 HeightfieldIndex = 0; HeightfieldIndex < NumPrimitives; HeightfieldIndex++)
+					{
+						const FPrimitiveSceneInfo* HeightfieldPrimitive = Scene->DistanceFieldSceneData.HeightfieldPrimitives[HeightfieldIndex];
+
+						UTexture2D* HeightfieldTexture = NULL;
+						FVector4 HeightfieldScaleBias(0, 0, 0, 0);
+						FVector4 MinMaxUV(0, 0, 0, 0);
+						HeightfieldPrimitive->Proxy->GetHeightfieldRepresentation(HeightfieldTexture, HeightfieldScaleBias, MinMaxUV);
+
+						if (HeightfieldTexture && HeightfieldTexture->Resource->TextureRHI)
+						{
+							TArray<FHeightfieldDescription, SceneRenderingAllocator>& HeightfieldDescriptions = HeightfieldTextures.FindOrAdd(HeightfieldTexture);
+							HeightfieldDescriptions.Add(FHeightfieldDescription(HeightfieldScaleBias, MinMaxUV, HeightfieldPrimitive->Proxy->GetLocalToWorld()));
+						}
+					}
+
+					if (HeightfieldTextures.Num() > 0)
+					{
+						for (TMap<UTexture2D*, TArray<FHeightfieldDescription, SceneRenderingAllocator>, SceneRenderingSetAllocator>::TIterator It(HeightfieldTextures); It; ++It)
+						{
+							const TArray<FHeightfieldDescription, SceneRenderingAllocator>& HeightfieldDescriptions = It.Value();
+							TArray<FVector4, SceneRenderingAllocator> HeightfieldDescriptionData;
+							HeightfieldDescriptionData.Empty(HeightfieldDescriptions.Num() * GHeightfieldDescriptions.Data.Stride);
+
+							for (int32 DescriptionIndex = 0; DescriptionIndex < HeightfieldDescriptions.Num(); DescriptionIndex++)
+							{
+								const FHeightfieldDescription& Description = HeightfieldDescriptions[DescriptionIndex];
+
+								HeightfieldDescriptionData.Add(Description.HeightfieldScaleBias);
+								HeightfieldDescriptionData.Add(Description.MinMaxUV);
+								
+								const FMatrix WorldToLocal = Description.LocalToWorld.Inverse();
+
+								HeightfieldDescriptionData.Add(*(FVector4*)&WorldToLocal.M[0]);
+								HeightfieldDescriptionData.Add(*(FVector4*)&WorldToLocal.M[1]);
+								HeightfieldDescriptionData.Add(*(FVector4*)&WorldToLocal.M[2]);
+								HeightfieldDescriptionData.Add(*(FVector4*)&WorldToLocal.M[3]);
+
+								HeightfieldDescriptionData.Add(*(FVector4*)&Description.LocalToWorld.M[0]);
+								HeightfieldDescriptionData.Add(*(FVector4*)&Description.LocalToWorld.M[1]);
+								HeightfieldDescriptionData.Add(*(FVector4*)&Description.LocalToWorld.M[2]);
+								HeightfieldDescriptionData.Add(*(FVector4*)&Description.LocalToWorld.M[3]);
+							}
+
+							if (HeightfieldDescriptionData.Num() > GHeightfieldDescriptions.Data.MaxElements)
+							{
+								GHeightfieldDescriptions.Data.MaxElements = HeightfieldDescriptionData.Num() * 5 / 4;
+								GHeightfieldDescriptions.Data.Release();
+								GHeightfieldDescriptions.Data.Initialize();
+							}
+
+							void* LockedBuffer = RHILockVertexBuffer(GHeightfieldDescriptions.Data.Buffer, 0, GHeightfieldDescriptions.Data.Buffer->GetSize(), RLM_WriteOnly);
+							const uint32 MemcpySize = HeightfieldDescriptionData.GetTypeSize() * HeightfieldDescriptionData.Num();
+							check(GHeightfieldDescriptions.Data.Buffer->GetSize() >= MemcpySize);
+							FPlatformMemory::Memcpy(LockedBuffer, HeightfieldDescriptionData.GetData(), MemcpySize);
+							RHIUnlockVertexBuffer(GHeightfieldDescriptions.Data.Buffer);
+
+							UTexture2D* HeightfieldTexture = It.Key();
+
+							{
+								TShaderMapRef<FCalculateHeightfieldOcclusionCS> ComputeShader(View.ShaderMap);
+								RHICmdList.SetComputeShader(ComputeShader->GetComputeShader());
+								ComputeShader->SetParameters(RHICmdList, View, DepthLevel, HeightfieldTexture, HeightfieldDescriptions.Num(), Parameters);
+								DispatchIndirectComputeShader(RHICmdList, *ComputeShader, SurfaceCacheResources.DispatchParameters.Buffer, 0);
+
+								ComputeShader->UnsetParameters(RHICmdList);
+							}
+						}
+					}
+				}
+
+				{
+					if (GDistanceFieldGI)
+					{
+						TShaderMapRef<TCombineConesCS<true> > ComputeShader(View.ShaderMap);
+
+						RHICmdList.SetComputeShader(ComputeShader->GetComputeShader());
+						ComputeShader->SetParameters(RHICmdList, View, DepthLevel, Parameters);
+						DispatchIndirectComputeShader(RHICmdList, *ComputeShader, SurfaceCacheResources.DispatchParameters.Buffer, 0);
+
+						ComputeShader->UnsetParameters(RHICmdList);
+					}
+					else
+					{
+						TShaderMapRef<TCombineConesCS<false> > ComputeShader(View.ShaderMap);
+
+						RHICmdList.SetComputeShader(ComputeShader->GetComputeShader());
+						ComputeShader->SetParameters(RHICmdList, View, DepthLevel, Parameters);
+						DispatchIndirectComputeShader(RHICmdList, *ComputeShader, SurfaceCacheResources.DispatchParameters.Buffer, 0);
+
+						ComputeShader->UnsetParameters(RHICmdList);
+					}
+				}
 			}
 
 			TRefCountPtr<IPooledRenderTarget> BentNormalAccumulation;
diff --git a/Engine/Source/Runtime/Renderer/Private/DistanceFieldSurfaceCacheLighting.h b/Engine/Source/Runtime/Renderer/Private/DistanceFieldSurfaceCacheLighting.h
--- a/Engine/Source/Runtime/Renderer/Private/DistanceFieldSurfaceCacheLighting.h
+++ b/Engine/Source/Runtime/Renderer/Private/DistanceFieldSurfaceCacheLighting.h
@@ -21,7 +21,7 @@ public:
 	FDistanceFieldAOParameters(float InOcclusionMaxDistance = 600.0f, float InContrast = 0)
 	{
 		Contrast = FMath::Clamp(InContrast, .01f, 2.0f);
-		OcclusionMaxDistance = FMath::Clamp(InOcclusionMaxDistance, 200.0f, 1500.0f);
+		OcclusionMaxDistance = FMath::Clamp(InOcclusionMaxDistance, 200.0f, 3000.0f);
 	}
 };
 
@@ -259,8 +259,8 @@ public:
 
 		SetShaderValue(RHICmdList, ShaderRHI, AOStepExponentScale, GAOStepExponentScale);
 
-		extern float GAOMaxViewDistance;
-		SetShaderValue(RHICmdList, ShaderRHI, AOMaxViewDistance, GAOMaxViewDistance);
+		extern float GetMaxAOViewDistance();
+		SetShaderValue(RHICmdList, ShaderRHI, AOMaxViewDistance, GetMaxAOViewDistance());
 	}
 
 private:
diff --git a/Engine/Source/Runtime/Renderer/Private/RendererScene.cpp b/Engine/Source/Runtime/Renderer/Private/RendererScene.cpp
--- a/Engine/Source/Runtime/Renderer/Private/RendererScene.cpp
+++ b/Engine/Source/Runtime/Renderer/Private/RendererScene.cpp
@@ -92,13 +92,20 @@ void FDistanceFieldSceneData::AddPrimitive(FPrimitiveSceneInfo* InPrimitive)
 {
 	const FPrimitiveSceneProxy* Proxy = InPrimitive->Proxy;
 
-	if (Proxy->SupportsDistanceFieldRepresentation()
-		&& Proxy->CastsDynamicShadow()
+	if (Proxy->CastsDynamicShadow()
 		&& Proxy->AffectsDistanceFieldLighting())
 	{
-		checkSlow(!PendingAddOperations.Contains(InPrimitive));
-		checkSlow(!PendingUpdateOperations.Contains(InPrimitive));
-		PendingAddOperations.Add(InPrimitive);
+		if (Proxy->SupportsHeightfieldRepresentation())
+		{
+			HeightfieldPrimitives.Add(InPrimitive);
+		}
+
+		if (Proxy->SupportsDistanceFieldRepresentation())
+		{
+			checkSlow(!PendingAddOperations.Contains(InPrimitive));
+			checkSlow(!PendingUpdateOperations.Contains(InPrimitive));
+			PendingAddOperations.Add(InPrimitive);
+		}
 	}
 }
 
@@ -106,9 +113,9 @@ void FDistanceFieldSceneData::UpdatePrimitive(FPrimitiveSceneInfo* InPrimitive)
 {
 	const FPrimitiveSceneProxy* Proxy = InPrimitive->Proxy;
 
-	if (Proxy->SupportsDistanceFieldRepresentation() 
-		&& Proxy->CastsDynamicShadow() 
+	if (Proxy->CastsDynamicShadow() 
 		&& Proxy->AffectsDistanceFieldLighting()
+		&& Proxy->SupportsDistanceFieldRepresentation() 
 		&& !PendingAddOperations.Contains(InPrimitive))
 	{
 		check(InPrimitive->DistanceFieldInstanceIndices.Num() > 0);
@@ -140,6 +147,11 @@ void FDistanceFieldSceneData::RemovePrimitive(FPrimitiveSceneInfo* InPrimitive)
 
 		InPrimitive->DistanceFieldInstanceIndices.Empty();
 	}
+
+	if (Proxy->SupportsHeightfieldRepresentation() && Proxy->AffectsDistanceFieldLighting())
+	{
+		HeightfieldPrimitives.Remove(InPrimitive);
+	}
 }
 
 void FDistanceFieldSceneData::Release()
diff --git a/Engine/Source/Runtime/Renderer/Private/ScenePrivate.h b/Engine/Source/Runtime/Renderer/Private/ScenePrivate.h
--- a/Engine/Source/Runtime/Renderer/Private/ScenePrivate.h
+++ b/Engine/Source/Runtime/Renderer/Private/ScenePrivate.h
@@ -914,6 +914,7 @@ public:
 
 	/** Stores the primitive and instance index of every entry in the object buffer. */
 	TArray<FPrimitiveAndInstance> PrimitiveInstanceMapping;
+	TArray<const FPrimitiveSceneInfo*> HeightfieldPrimitives;
 
 	/** Pending operations on the object buffers to be processed next frame. */
 	TArray<FPrimitiveSceneInfo*> PendingAddOperations;
