refactor FrameTime (moved from View to ViewFamily)

diff --git a/Engine/Plugins/Experimental/SteamVR/Source/SteamVR/Private/SteamVRHMD.cpp b/Engine/Plugins/Experimental/SteamVR/Source/SteamVR/Private/SteamVRHMD.cpp
--- a/Engine/Plugins/Experimental/SteamVR/Source/SteamVR/Private/SteamVRHMD.cpp
+++ b/Engine/Plugins/Experimental/SteamVR/Source/SteamVR/Private/SteamVRHMD.cpp
@@ -513,7 +513,7 @@ void FSteamVRHMD::PreRenderView_RenderThread(FSceneView& View)
 //	UpdatePlayerViewPoint(CurrentOrientation, CurrentPosition, View);
 }
 
-void FSteamVRHMD::PreRenderViewFamily_RenderThread(FSceneViewFamily& ViewFamily, uint32 InFrameNumber)
+void FSteamVRHMD::PreRenderViewFamily_RenderThread(FSceneViewFamily& ViewFamily)
 {
 	check(IsInRenderingThread());
 }
diff --git a/Engine/Plugins/Experimental/SteamVR/Source/SteamVR/Private/SteamVRHMD.h b/Engine/Plugins/Experimental/SteamVR/Source/SteamVR/Private/SteamVRHMD.h
--- a/Engine/Plugins/Experimental/SteamVR/Source/SteamVR/Private/SteamVRHMD.h
+++ b/Engine/Plugins/Experimental/SteamVR/Source/SteamVR/Private/SteamVRHMD.h
@@ -100,7 +100,7 @@ public:
 	virtual void ModifyShowFlags(FEngineShowFlags& ShowFlags) override;
 	virtual void SetupView(FSceneViewFamily& InViewFamily, FSceneView& InView) override;
 	virtual void PreRenderView_RenderThread(FSceneView& InView) override;
-	virtual void PreRenderViewFamily_RenderThread(FSceneViewFamily& InViewFamily, uint32 InFrameNumber) override;
+	virtual void PreRenderViewFamily_RenderThread(FSceneViewFamily& InViewFamily) override;
 
 public:
 	/** Constructor */
diff --git a/Engine/Plugins/Runtime/OculusRift/Source/OculusRift/Private/OculusRiftHMD.h b/Engine/Plugins/Runtime/OculusRift/Source/OculusRift/Private/OculusRiftHMD.h
--- a/Engine/Plugins/Runtime/OculusRift/Source/OculusRift/Private/OculusRiftHMD.h
+++ b/Engine/Plugins/Runtime/OculusRift/Source/OculusRift/Private/OculusRiftHMD.h
@@ -131,7 +131,7 @@ public:
     virtual void ModifyShowFlags(FEngineShowFlags& ShowFlags) override;
     virtual void SetupView(FSceneViewFamily& InViewFamily, FSceneView& InView) override;
     virtual void PreRenderView_RenderThread(FSceneView& InView) override;
-	virtual void PreRenderViewFamily_RenderThread(FSceneViewFamily& InViewFamily, uint32 InFrameNumber) override;
+	virtual void PreRenderViewFamily_RenderThread(FSceneViewFamily& InViewFamily) override;
 
 	/** Positional tracking control methods */
 	virtual bool IsPositionalTrackingEnabled() const override;
diff --git a/Engine/Plugins/Runtime/OculusRift/Source/OculusRift/Private/OculusRiftRender.cpp b/Engine/Plugins/Runtime/OculusRift/Source/OculusRift/Private/OculusRiftRender.cpp
--- a/Engine/Plugins/Runtime/OculusRift/Source/OculusRift/Private/OculusRiftRender.cpp
+++ b/Engine/Plugins/Runtime/OculusRift/Source/OculusRift/Private/OculusRiftRender.cpp
@@ -149,14 +149,14 @@ void FOculusRiftHMD::GetTimewarpMatrices_RenderThread(EStereoscopicPass StereoPa
 #endif
 }
 
-void FOculusRiftHMD::PreRenderViewFamily_RenderThread(FSceneViewFamily& ViewFamily, uint32 InFrameNumber)
+void FOculusRiftHMD::PreRenderViewFamily_RenderThread(FSceneViewFamily& ViewFamily)
 {
 	check(IsInRenderingThread());
 
 	RenderParams_RenderThread.ShowFlags = ViewFamily.EngineShowFlags;
 	// FrameNumber comes from SceneRenderer.FrameNumber and contains already incremented
 	// GFrameNumber. Therefore, decrementing it to match with the frame number used on a game thread.
-	RenderParams_RenderThread.FrameNumber = InFrameNumber - 1;
+	RenderParams_RenderThread.FrameNumber = ViewFamily.FrameNumber - 1;
 
 #ifndef OVR_SDK_RENDERING 
 	{
diff --git a/Engine/Source/Runtime/Engine/Private/GPUSkinVertexFactory.cpp b/Engine/Source/Runtime/Engine/Private/GPUSkinVertexFactory.cpp
--- a/Engine/Source/Runtime/Engine/Private/GPUSkinVertexFactory.cpp
+++ b/Engine/Source/Runtime/Engine/Private/GPUSkinVertexFactory.cpp
@@ -375,6 +375,8 @@ public:
 	{
 		if(Shader->GetVertexShader())
 		{
+			uint32 FrameNumber = View.Family->FrameNumber;
+
 			const auto FeatureLevel = View.GetFeatureLevel();
 			const FGPUBaseSkinVertexFactory::ShaderDataType& ShaderData = ((const FGPUBaseSkinVertexFactory*)VertexFactory)->GetShaderData();
 
@@ -410,7 +412,7 @@ public:
 				// we are in the velocity rendering pass
 
 				// 0xffffffff or valid index
-				uint32 OldBoneDataIndex = ShaderData.GetOldBoneData(View.FrameNumber);
+				uint32 OldBoneDataIndex = ShaderData.GetOldBoneData(FrameNumber);
 
 				// Read old data if it was written last frame (normal data) or this frame (e.g. split screen)
 				bLocalPerBoneMotionBlur = (OldBoneDataIndex != 0xffffffff);
@@ -439,14 +441,14 @@ public:
 				}
 				FScopeLock Lock(&ShaderData.OldBoneDataLock);
 				// if we haven't copied the data yet we skip the update (e.g. split screen)
-				if(ShaderData.IsOldBoneDataUpdateNeeded(View.FrameNumber))
+				if(ShaderData.IsOldBoneDataUpdateNeeded(FrameNumber))
 				{
 					const FGPUBaseSkinVertexFactory* GPUVertexFactory = (const FGPUBaseSkinVertexFactory*)VertexFactory;
 
 					// copy the bone data and tell the instance where it can pick it up next frame
 					// append data to a buffer we bind next frame to read old matrix data for motion blur
 					uint32 OldBoneDataStartIndex = GPrevPerBoneMotionBlur.AppendData(ShaderData.BoneMatrices.GetData(), ShaderData.BoneMatrices.Num());
-					GPUVertexFactory->SetOldBoneDataStartIndex(View.FrameNumber, OldBoneDataStartIndex);
+					GPUVertexFactory->SetOldBoneDataStartIndex(FrameNumber, OldBoneDataStartIndex);
 				}
 			}
 
diff --git a/Engine/Source/Runtime/Engine/Private/SceneView.cpp b/Engine/Source/Runtime/Engine/Private/SceneView.cpp
--- a/Engine/Source/Runtime/Engine/Private/SceneView.cpp
+++ b/Engine/Source/Runtime/Engine/Private/SceneView.cpp
@@ -260,7 +260,6 @@ FSceneView::FSceneView(const FSceneViewInitOptions& InitOptions)
 	, ViewRect(InitOptions.GetConstrainedViewRect())
 	, UnscaledViewRect(InitOptions.GetConstrainedViewRect())
 	, UnconstrainedViewRect(InitOptions.GetViewRect())
-	, FrameNumber(UINT_MAX)
 	, MaxShadowCascades(10)
 	, WorldToMetersScale(InitOptions.WorldToMetersScale)
 	, ProjectionMatrixUnadjustedForRHI(InitOptions.ProjectionMatrix)
@@ -1374,6 +1373,7 @@ FSceneViewFamily::FSceneViewFamily( const ConstructionValues& CVS )
 	CurrentWorldTime(CVS.CurrentWorldTime),
 	DeltaWorldTime(CVS.DeltaWorldTime),
 	CurrentRealTime(CVS.CurrentRealTime),
+	FrameNumber(UINT_MAX),
 	bRealtimeUpdate(CVS.bRealtimeUpdate),
 	bDeferClear(CVS.bDeferClear),
 	bResolveScene(CVS.bResolveScene),
diff --git a/Engine/Source/Runtime/Engine/Public/SceneView.h b/Engine/Source/Runtime/Engine/Public/SceneView.h
--- a/Engine/Source/Runtime/Engine/Public/SceneView.h
+++ b/Engine/Source/Runtime/Engine/Public/SceneView.h
@@ -332,11 +332,6 @@ public:
 	/* Raw view size (in pixels), used for screen space calculations */
 	const FIntRect UnconstrainedViewRect;
 
-	/** 
-	 * Copy from GFrameNumber
-	 */
-	uint32 FrameNumber;
-
 	/** Maximum number of shadow cascades to render with. */
 	int32 MaxShadowCascades;
 
@@ -633,7 +628,7 @@ public:
 
 		/** Gamma correction used when rendering this family. Default is 1.0 */
 		float GammaCorrection;
-		
+
 		/** Indicates whether the view family is updated in real-time. */
 		uint32 bRealtimeUpdate:1;
 		
@@ -692,6 +687,9 @@ public:
 	/** The current real time. */
 	float CurrentRealTime;
 
+	/** Copy from main thread GFrameNumber to be accessible on render thread side. UINT_MAX before BeginRenderingViewFamily() was called */
+	uint32 FrameNumber;
+
 	/** Indicates whether the view family is updated in realtime. */
 	bool bRealtimeUpdate;
 
diff --git a/Engine/Source/Runtime/Engine/Public/SceneViewExtension.h b/Engine/Source/Runtime/Engine/Public/SceneViewExtension.h
--- a/Engine/Source/Runtime/Engine/Public/SceneViewExtension.h
+++ b/Engine/Source/Runtime/Engine/Public/SceneViewExtension.h
@@ -21,9 +21,8 @@ public:
 
     /**
      * Called on render thread at the start of rendering.
-	 * FrameNumber is a copy of SceneRenderer.FrameNumber (that is equal to already incremented GFrameNumber).
      */
-    virtual void PreRenderViewFamily_RenderThread(FSceneViewFamily& InViewFamily, uint32 InFrameNumber) = 0;
+    virtual void PreRenderViewFamily_RenderThread(FSceneViewFamily& InViewFamily) = 0;
 
 	/**
      * Called on render thread at the start of rendering, for each view, after PreRenderViewFamily_RenderThread call.
diff --git a/Engine/Source/Runtime/Renderer/Private/PostProcess/PostProcessTemporalAA.cpp b/Engine/Source/Runtime/Renderer/Private/PostProcess/PostProcessTemporalAA.cpp
--- a/Engine/Source/Runtime/Renderer/Private/PostProcess/PostProcessTemporalAA.cpp
+++ b/Engine/Source/Runtime/Renderer/Private/PostProcess/PostProcessTemporalAA.cpp
@@ -181,7 +181,7 @@ public:
 			}
 			
 			FVector2D RandomOffsetValue;
-			TemporalRandom(&RandomOffsetValue, Context.View.FrameNumber);
+			TemporalRandom(&RandomOffsetValue, Context.View.Family->FrameNumber);
 			SetShaderValue(Context.RHICmdList, ShaderRHI, RandomOffset, RandomOffsetValue);
 
 		}
diff --git a/Engine/Source/Runtime/Renderer/Private/PostProcess/PostProcessTestImage.cpp b/Engine/Source/Runtime/Renderer/Private/PostProcess/PostProcessTestImage.cpp
--- a/Engine/Source/Runtime/Renderer/Private/PostProcess/PostProcessTestImage.cpp
+++ b/Engine/Source/Runtime/Renderer/Private/PostProcess/PostProcessTestImage.cpp
@@ -60,7 +60,7 @@ public:
 		DeferredParameters.Set(Context.RHICmdList, ShaderRHI, Context.View);
 
 		{
-			uint32 FrameNumberValue = Context.View.FrameNumber;
+			uint32 FrameNumberValue = Context.View.Family->FrameNumber;
 			SetShaderValue(Context.RHICmdList, ShaderRHI, FrameNumber, FrameNumberValue);
 		}
 
diff --git a/Engine/Source/Runtime/Renderer/Private/PostProcess/PostProcessTonemap.cpp b/Engine/Source/Runtime/Renderer/Private/PostProcess/PostProcessTonemap.cpp
--- a/Engine/Source/Runtime/Renderer/Private/PostProcess/PostProcessTonemap.cpp
+++ b/Engine/Source/Runtime/Renderer/Private/PostProcess/PostProcessTonemap.cpp
@@ -1547,7 +1547,7 @@ public:
 		PostprocessParameter.SetVS(ShaderRHI, Context, TStaticSamplerState<SF_Bilinear,AM_Clamp,AM_Clamp,AM_Clamp>::GetRHI());
 
 		FVector GrainRandomFullValue;
-		GrainRandomFromFrame(&GrainRandomFullValue, Context.View.FrameNumber);
+		GrainRandomFromFrame(&GrainRandomFullValue, Context.View.Family->FrameNumber);
 		// TODO: Don't use full on mobile with framebuffer fetch.
 		GrainRandomFullValue.Z = bUsedFramebufferFetch ? 0.0f : 1.0f;
 		SetShaderValue(Context.RHICmdList, ShaderRHI, GrainRandomFull, GrainRandomFullValue);
diff --git a/Engine/Source/Runtime/Renderer/Private/PostProcess/PostProcessTonemap.h b/Engine/Source/Runtime/Renderer/Private/PostProcess/PostProcessTonemap.h
--- a/Engine/Source/Runtime/Renderer/Private/PostProcess/PostProcessTonemap.h
+++ b/Engine/Source/Runtime/Renderer/Private/PostProcess/PostProcessTonemap.h
@@ -124,7 +124,7 @@ public:
 		PostprocessParameter.SetVS(ShaderRHI, Context, TStaticSamplerState<SF_Bilinear,AM_Clamp,AM_Clamp,AM_Clamp>::GetRHI());
 
 		FVector GrainRandomFullValue;
-		GrainRandomFromFrame(&GrainRandomFullValue, Context.View.FrameNumber);
+		GrainRandomFromFrame(&GrainRandomFullValue, Context.View.Family->FrameNumber);
 		SetShaderValue(Context.RHICmdList, ShaderRHI, GrainRandomFull, GrainRandomFullValue);
 
 		if(EyeAdaptation.IsBound())
diff --git a/Engine/Source/Runtime/Renderer/Private/PostProcess/SceneRenderTargets.cpp b/Engine/Source/Runtime/Renderer/Private/PostProcess/SceneRenderTargets.cpp
--- a/Engine/Source/Runtime/Renderer/Private/PostProcess/SceneRenderTargets.cpp
+++ b/Engine/Source/Runtime/Renderer/Private/PostProcess/SceneRenderTargets.cpp
@@ -254,7 +254,7 @@ void FSceneRenderTargets::Allocate(const FSceneViewFamily& ViewFamily)
 		// Reinitialize the render targets for the given size.
 		SetBufferSize(DesiredBufferSize.X, DesiredBufferSize.Y);
 
-		UE_LOG(LogRenderer, Warning, TEXT("Reallocating scene render targets to support %ux%u."), BufferSize.X, BufferSize.Y);
+		UE_LOG(LogRenderer, Warning, TEXT("Reallocating scene render targets to support %ux%u (Frame:%u)."), BufferSize.X, BufferSize.Y, ViewFamily.FrameNumber);
 
 		UpdateRHI();
 	}
diff --git a/Engine/Source/Runtime/Renderer/Private/RendererPrivate.h b/Engine/Source/Runtime/Renderer/Private/RendererPrivate.h
--- a/Engine/Source/Runtime/Renderer/Private/RendererPrivate.h
+++ b/Engine/Source/Runtime/Renderer/Private/RendererPrivate.h
@@ -22,7 +22,7 @@ class FRendererModule : public IRendererModule
 public:
 	virtual bool SupportsDynamicReloading() override { return true; }
 
-	virtual void BeginRenderingViewFamily(FCanvas* Canvas,const FSceneViewFamily* ViewFamily) override;
+	virtual void BeginRenderingViewFamily(FCanvas* Canvas,FSceneViewFamily* ViewFamily) override;
 	virtual FSceneInterface* AllocateScene(UWorld* World, bool bInRequiresHitProxies, bool bCreateFXSystem, ERHIFeatureLevel::Type InFeatureLevel) override;
 	virtual void RemoveScene(FSceneInterface* Scene) override;
 	virtual void UpdateStaticDrawListsForMaterials(const TArray<const FMaterial*>& Materials) override;
diff --git a/Engine/Source/Runtime/Renderer/Private/SceneRendering.cpp b/Engine/Source/Runtime/Renderer/Private/SceneRendering.cpp
--- a/Engine/Source/Runtime/Renderer/Private/SceneRendering.cpp
+++ b/Engine/Source/Runtime/Renderer/Private/SceneRendering.cpp
@@ -488,7 +488,7 @@ TUniformBufferRef<FViewUniformShaderParameters> FViewInfo::CreateUniformBuffer(
 	ViewUniformShaderParameters.GameTime = Family->CurrentWorldTime;
 	ViewUniformShaderParameters.RealTime = Family->CurrentRealTime;
 	ViewUniformShaderParameters.Random = FMath::Rand();
-	ViewUniformShaderParameters.FrameNumber = FrameNumber;
+	ViewUniformShaderParameters.FrameNumber = Family->FrameNumber;
 
 	static const auto CVar = IConsoleManager::Get().FindConsoleVariable(TEXT("r.DiffuseFromCaptures"));
 	const bool bUseLightmaps = CVar->GetInt() == 0;
@@ -687,7 +687,6 @@ FSceneRenderer::FSceneRenderer(const FSceneViewFamily* InViewFamily,FHitProxyCon
 :	Scene(InViewFamily->Scene ? InViewFamily->Scene->GetRenderScene() : NULL)
 ,	ViewFamily(*InViewFamily)
 ,	bUsedPrecomputedVisibility(false)
-,	FrameNumber(GFrameNumber)
 {
 	check(Scene != NULL);
 
@@ -1060,7 +1059,7 @@ static void RenderViewFamily_RenderThread(FRHICommandListImmediate& RHICmdList,
 
 	for( int ViewExt = 0; ViewExt < SceneRenderer->ViewFamily.ViewExtensions.Num(); ViewExt++ )
 	{
-		SceneRenderer->ViewFamily.ViewExtensions[ViewExt]->PreRenderViewFamily_RenderThread(SceneRenderer->ViewFamily, SceneRenderer->FrameNumber);
+		SceneRenderer->ViewFamily.ViewExtensions[ViewExt]->PreRenderViewFamily_RenderThread(SceneRenderer->ViewFamily);
 		for( int ViewIndex = 0; ViewIndex < SceneRenderer->ViewFamily.Views.Num(); ViewIndex++ )
 		{
 			SceneRenderer->ViewFamily.ViewExtensions[ViewExt]->PreRenderView_RenderThread(SceneRenderer->Views[ViewIndex]);
@@ -1123,7 +1122,7 @@ static void RenderViewFamily_RenderThread(FRHICommandListImmediate& RHICmdList,
 #endif
 }
 
-void FRendererModule::BeginRenderingViewFamily(FCanvas* Canvas,const FSceneViewFamily* ViewFamily)
+void FRendererModule::BeginRenderingViewFamily(FCanvas* Canvas,FSceneViewFamily* ViewFamily)
 {
 	// Flush the canvas first.
 	Canvas->Flush_GameThread();
@@ -1132,6 +1131,8 @@ void FRendererModule::BeginRenderingViewFamily(FCanvas* Canvas,const FSceneViewF
 	// This is the only spot we change GFrameNumber, other places can only read.
 	++GFrameNumber;
 
+	ViewFamily->FrameNumber = GFrameNumber;
+
 	check(ViewFamily->Scene);
 	FScene* const Scene = ViewFamily->Scene->GetRenderScene();
 	if (Scene)
diff --git a/Engine/Source/Runtime/Renderer/Private/SceneRendering.h b/Engine/Source/Runtime/Renderer/Private/SceneRendering.h
--- a/Engine/Source/Runtime/Renderer/Private/SceneRendering.h
+++ b/Engine/Source/Runtime/Renderer/Private/SceneRendering.h
@@ -576,9 +576,6 @@ public:
 	/** True if precomputed visibility was used when rendering the scene. */
 	bool bUsedPrecomputedVisibility;
 
-	/** Copy from main thread GFrameNumber to be accessible on renderthread side */
-	uint32 FrameNumber;
-
 	/** Feature level being rendered */
 	ERHIFeatureLevel::Type FeatureLevel;
 
diff --git a/Engine/Source/Runtime/Renderer/Private/SceneVisibility.cpp b/Engine/Source/Runtime/Renderer/Private/SceneVisibility.cpp
--- a/Engine/Source/Runtime/Renderer/Private/SceneVisibility.cpp
+++ b/Engine/Source/Runtime/Renderer/Private/SceneVisibility.cpp
@@ -171,7 +171,7 @@ static void UpdatePrimitiveFadingState(FPrimitiveFadingState& FadingState, FView
 		}
 	}
 
-	FadingState.FrameNumber = View.FrameNumber;
+	FadingState.FrameNumber = View.Family->FrameNumber;
 	FadingState.bIsVisible = bVisible;
 	FadingState.bValid = true;
 }
@@ -1549,7 +1549,6 @@ void FSceneRenderer::PreVisibilityFrameSetup(FRHICommandListImmediate& RHICmdLis
 	for(int32 ViewIndex = 0;ViewIndex < Views.Num();ViewIndex++)
 	{
 		FViewInfo& View = Views[ViewIndex];
-		View.FrameNumber = FrameNumber;
 		FSceneViewState* ViewState = (FSceneViewState*) View.State;
 		static bool bEnableTimeScale = true;
 
@@ -1774,7 +1773,7 @@ void FSceneRenderer::PreVisibilityFrameSetup(FRHICommandListImmediate& RHICmdLis
 			}
 
 			ViewState->PrevFrameNumber = ViewState->PendingPrevFrameNumber;
-			ViewState->PendingPrevFrameNumber = View.FrameNumber;
+			ViewState->PendingPrevFrameNumber = View.Family->FrameNumber;
 
 			// This finishes the update of view state
 			ViewState->UpdateLastRenderTime(*View.Family);
diff --git a/Engine/Source/Runtime/Renderer/Public/RendererInterface.h b/Engine/Source/Runtime/Renderer/Public/RendererInterface.h
--- a/Engine/Source/Runtime/Renderer/Public/RendererInterface.h
+++ b/Engine/Source/Runtime/Renderer/Public/RendererInterface.h
@@ -470,7 +470,7 @@ class IRendererModule : public IModuleInterface
 public:
 
 	/** Call from the game thread to send a message to the rendering thread to being rendering this view family. */
-	virtual void BeginRenderingViewFamily(FCanvas* Canvas,const FSceneViewFamily* ViewFamily) = 0;
+	virtual void BeginRenderingViewFamily(FCanvas* Canvas, FSceneViewFamily* ViewFamily) = 0;
 	
 	/**
 	 * Allocates a new instance of the private FScene implementation for the given world.
diff --git a/Engine/Source/Runtime/SlateRHIRenderer/Private/SlateRHIRenderingPolicy.cpp b/Engine/Source/Runtime/SlateRHIRenderer/Private/SlateRHIRenderingPolicy.cpp
--- a/Engine/Source/Runtime/SlateRHIRenderer/Private/SlateRHIRenderingPolicy.cpp
+++ b/Engine/Source/Runtime/SlateRHIRenderer/Private/SlateRHIRenderingPolicy.cpp
@@ -232,7 +232,7 @@ static FSceneView& CreateSceneView( FSceneViewFamilyContext& ViewFamilyContext,
 	ViewUniformShaderParameters.GameTime = View->Family->CurrentWorldTime;
 	ViewUniformShaderParameters.RealTime = View->Family->CurrentRealTime;
 	ViewUniformShaderParameters.Random = FMath::Rand();
-	ViewUniformShaderParameters.FrameNumber = View->FrameNumber;
+	ViewUniformShaderParameters.FrameNumber = View->Family->FrameNumber;
 
 	ViewUniformShaderParameters.DirectionalLightShadowTexture = GWhiteTexture->TextureRHI;
 	ViewUniformShaderParameters.DirectionalLightShadowSampler = TStaticSamplerState<SF_Point, AM_Clamp, AM_Clamp, AM_Clamp>::GetRHI();
