changed BufferSize computations, system should adapt more efficiently to SceneCaptures (assuming they are a one off or come each frame or are smaller than the Screen/Window Size)

diff --git a/Engine/Source/Runtime/Engine/Public/SceneView.h b/Engine/Source/Runtime/Engine/Public/SceneView.h
--- a/Engine/Source/Runtime/Engine/Public/SceneView.h
+++ b/Engine/Source/Runtime/Engine/Public/SceneView.h
@@ -687,7 +687,7 @@ public:
 	/** The current real time. */
 	float CurrentRealTime;
 
-	/** Copy from main thread GFrameNumber to be accessible on render thread side. UINT_MAX before BeginRenderingViewFamily() was called */
+	/** Copy from main thread GFrameNumber to be accessible on render thread side. UINT_MAX before CreateSceneRenderer() or BeginRenderingViewFamily() was called */
 	uint32 FrameNumber;
 
 	/** Indicates whether the view family is updated in realtime. */
diff --git a/Engine/Source/Runtime/Renderer/Private/PostProcess/SceneRenderTargets.cpp b/Engine/Source/Runtime/Renderer/Private/PostProcess/SceneRenderTargets.cpp
--- a/Engine/Source/Runtime/Renderer/Private/PostProcess/SceneRenderTargets.cpp
+++ b/Engine/Source/Runtime/Renderer/Private/PostProcess/SceneRenderTargets.cpp
@@ -105,7 +105,7 @@ inline const TCHAR* GetSceneColorTargetName(FSceneRenderTargets::EShadingPath Sh
 	return SceneColorNames[(uint32)ShadingPath];
 }
 
-FIntPoint FSceneRenderTargets::ComputeDesiredSize(const FSceneViewFamily& ViewFamily) const
+FIntPoint FSceneRenderTargets::ComputeDesiredSize(const FSceneViewFamily& ViewFamily)
 {
 	// Don't expose Clamped to the cvar since you need to at least grow to the initial state.
 	enum ESizingMethods { RequestedSize, ScreenRes, Grow, VisibleSizingMethodsCount, Clamped };
@@ -137,51 +137,64 @@ FIntPoint FSceneRenderTargets::ComputeDesiredSize(const FSceneViewFamily& ViewFa
 		// Otherwise use the setting specified by the console variable.
 		SceneTargetsSizingMethod = (ESizingMethods) FMath::Clamp(CVarSceneTargetsResizingMethod.GetValueOnRenderThread(), 0, (int32)VisibleSizingMethodsCount);
 	}
-	
-	if (bIsSceneCapture)
-	{
-		// This is a bit dangerous as someone in the middle of the game can requests a high res RenderTarget (Larger than screen) and from that point on the SceneRenderTargets will never get smaller again.
-		// This can cause more pressure in the RenderTargetPool and cause more trashing there. It also could slowdown post processing (border clear / handling).
-		// Preventing that case can cause more trouble (When to shrink? Should we keep the former RenderTargets in the pool?).
-		// If this becomes an issue we suggest to shrink the SceneRenderTargets with some game code call (e.g. each level load).
-		SceneTargetsSizingMethod = Grow;
-	}
-	else if(bIsReflectionCapture)
-	{
-		// reflection captures are low in resolution usually 128, they should not cause the BufferSize to be larger than needed (larger than ScreenSize)
-//		check(ViewFamily.FamilySizeX <= 256);
-//		check(ViewFamily.FamilySizeY <= 256);
 
-		// this can cause a reallocation during startup if the reflections are updated before the SceneRenderTarget have been allocated with the proper Screen size.
-		// This results in one log printout and some extra minor cost memory and performance cost.
-		// On XBoxOn it can cause different ESRam allocation (Rare, only if Reflection Capture wasn't cooked)
-		SceneTargetsSizingMethod = Grow;
-	}
+	FIntPoint DesiredBufferSize = FIntPoint::ZeroValue;
 
 	switch (SceneTargetsSizingMethod)
 	{
 		case RequestedSize:
-			return FIntPoint(ViewFamily.FamilySizeX, ViewFamily.FamilySizeY);
+			DesiredBufferSize = FIntPoint(ViewFamily.FamilySizeX, ViewFamily.FamilySizeY);
+			break;
+
 		case ScreenRes:
-			return FIntPoint(GSystemResolution.ResX, GSystemResolution.ResY);
+			DesiredBufferSize = FIntPoint(GSystemResolution.ResX, GSystemResolution.ResY);
+			break;
+
 		case Grow:
-			return FIntPoint(FMath::Max((uint32)GetBufferSizeXY().X, ViewFamily.FamilySizeX),
+			DesiredBufferSize = FIntPoint(FMath::Max((uint32)GetBufferSizeXY().X, ViewFamily.FamilySizeX),
 					FMath::Max((uint32)GetBufferSizeXY().Y, ViewFamily.FamilySizeY));
+			break;
+
 		case Clamped:
 			if (((uint32)BufferSize.X < ViewFamily.FamilySizeX) || ((uint32)BufferSize.Y < ViewFamily.FamilySizeY))
 			{
 				UE_LOG(LogRenderer, Warning, TEXT("Capture target size: %ux%u clamped to %ux%u."), ViewFamily.FamilySizeX, ViewFamily.FamilySizeY, BufferSize.X, BufferSize.Y);
 			}
-			return FIntPoint(GetBufferSizeXY().X, GetBufferSizeXY().Y);
+			DesiredBufferSize = FIntPoint(GetBufferSizeXY().X, GetBufferSizeXY().Y);
+			break;
+
 		default:
 			checkNoEntry();
-			return FIntPoint::ZeroValue;
 	}
+
+
+	// we want to shrink the buffer but as we can have multiple scenecaptures per frame we have to delay that a frame to get all size requests
+	{
+		// this allows The BufferSize to not grow below the SceneCapture requests (happen before scene rendering, in the same frame with a Grow request)
+		LargestDesiredSizeThisFrame = LargestDesiredSizeThisFrame.ComponentMax(DesiredBufferSize);
+
+		uint32 FrameNumber = ViewFamily.FrameNumber;
+
+		// this could be refined to be some time or multiple frame if we have SceneCaptures not running each frame any more
+		if(ThisFrameNumber != FrameNumber)
+		{
+			// this allows the BufferSize to shrink each frame (in game)
+			ThisFrameNumber = FrameNumber;
+			LargestDesiredSizeLastFrame = LargestDesiredSizeThisFrame;
+			LargestDesiredSizeThisFrame = FIntPoint(0, 0);
+		}
+
+		DesiredBufferSize = DesiredBufferSize.ComponentMax(LargestDesiredSizeLastFrame);
+	}
+
+	return DesiredBufferSize;
 }
 
 void FSceneRenderTargets::Allocate(const FSceneViewFamily& ViewFamily)
 {
 	check(IsInRenderingThread());
+	// ViewFamily setup wasn't complete
+	check(ViewFamily.FrameNumber != UINT_MAX);
 
 	// If feature level has changed, release all previously allocated targets to the pool. If feature level has changed but
 	const auto NewFeatureLevel = ViewFamily.Scene->GetFeatureLevel();
diff --git a/Engine/Source/Runtime/Renderer/Private/PostProcess/SceneRenderTargets.h b/Engine/Source/Runtime/Renderer/Private/PostProcess/SceneRenderTargets.h
--- a/Engine/Source/Runtime/Renderer/Private/PostProcess/SceneRenderTargets.h
+++ b/Engine/Source/Runtime/Renderer/Private/PostProcess/SceneRenderTargets.h
@@ -116,6 +116,9 @@ protected:
 		bPreshadowCacheNewlyAllocated(false),
 		GBufferRefCount(0),
 		BufferSize(0, 0), 
+		LargestDesiredSizeThisFrame(0, 0),
+		LargestDesiredSizeLastFrame(0, 0),
+		ThisFrameNumber(0),
 		SmallColorDepthDownsampleFactor(2),
 		bLightAttenuationEnabled(true),
 		bUseDownsizedOcclusionQueries(true),
@@ -513,6 +516,12 @@ private:
 	/** used by AdjustGBufferRefCount */
 	int32 GBufferRefCount;
 
+	/** as we might get multiple BufferSize requests each frame for SceneCaptures and we want to avoid reallocations we can only go as low as the largest request */
+	FIntPoint LargestDesiredSizeThisFrame;
+	FIntPoint LargestDesiredSizeLastFrame;
+	/** to detect when LargestDesiredSizeThisFrame is outdated */
+	uint32 ThisFrameNumber;
+
 	/**
 	 * Takes the requested buffer size and quantizes it to an appropriate size for the rest of the
 	 * rendering pipeline. Currently ensures that sizes are multiples of 8 so that they can safely
@@ -543,7 +552,7 @@ private:
 	void AllocateCommonDepthTargets();
 
 	/** Determine the appropriate render target dimensions. */
-	FIntPoint ComputeDesiredSize(const FSceneViewFamily& ViewFamily) const;
+	FIntPoint ComputeDesiredSize(const FSceneViewFamily& ViewFamily);
 
 	void AllocSceneColor();
 
diff --git a/Engine/Source/Runtime/Renderer/Private/SceneRendering.cpp b/Engine/Source/Runtime/Renderer/Private/SceneRendering.cpp
--- a/Engine/Source/Runtime/Renderer/Private/SceneRendering.cpp
+++ b/Engine/Source/Runtime/Renderer/Private/SceneRendering.cpp
@@ -690,6 +690,9 @@ FSceneRenderer::FSceneRenderer(const FSceneViewFamily* InViewFamily,FHitProxyCon
 {
 	check(Scene != NULL);
 
+	check(IsInGameThread());
+	ViewFamily.FrameNumber = GFrameNumber;
+
 	// Copy the individual views.
 	bool bAnyViewIsLocked = false;
 	Views.Empty(InViewFamily->Views.Num());
@@ -1131,6 +1134,7 @@ void FRendererModule::BeginRenderingViewFamily(FCanvas* Canvas,FSceneViewFamily*
 	// This is the only spot we change GFrameNumber, other places can only read.
 	++GFrameNumber;
 
+	// this is passes to the render thread, better access that than GFrameNumberRenderThread
 	ViewFamily->FrameNumber = GFrameNumber;
 
 	check(ViewFamily->Scene);
