Handle primitives which couldn't be allocated in the SDF atlas

diff --git a/Engine/Source/Runtime/Renderer/Private/RendererScene.cpp b/Engine/Source/Runtime/Renderer/Private/RendererScene.cpp
--- a/Engine/Source/Runtime/Renderer/Private/RendererScene.cpp
+++ b/Engine/Source/Runtime/Renderer/Private/RendererScene.cpp
@@ -116,9 +116,10 @@ void FDistanceFieldSceneData::UpdatePrimitive(FPrimitiveSceneInfo* InPrimitive)
 	if (Proxy->CastsDynamicShadow() 
 		&& Proxy->AffectsDistanceFieldLighting()
 		&& Proxy->SupportsDistanceFieldRepresentation() 
-		&& !PendingAddOperations.Contains(InPrimitive))
+		&& !PendingAddOperations.Contains(InPrimitive)
+		// This can happen when the primitive fails to allocate from the SDF atlas
+		&& InPrimitive->DistanceFieldInstanceIndices.Num() > 0)
 	{
-		check(InPrimitive->DistanceFieldInstanceIndices.Num() > 0);
 		PendingUpdateOperations.Add(InPrimitive);
 	}
 }
