Fix editor rendering on OS X GL 3.3 GPUs which lack support for volume texture rendering .

diff --git a/Engine/Source/Runtime/Renderer/Private/TranslucentLighting.cpp b/Engine/Source/Runtime/Renderer/Private/TranslucentLighting.cpp
--- a/Engine/Source/Runtime/Renderer/Private/TranslucentLighting.cpp
+++ b/Engine/Source/Runtime/Renderer/Private/TranslucentLighting.cpp
@@ -1162,7 +1162,7 @@ void FDeferredShadingSceneRenderer::InjectAmbientCubemapTranslucentVolumeLightin
 	//@todo - support multiple views
 	const FViewInfo& View = Views[0];
 
-	if (GUseTranslucentLightingVolumes && View.FinalPostProcessSettings.ContributingCubemaps.Num() && !IsSimpleDynamicLightingEnabled())
+	if (GUseTranslucentLightingVolumes && GSupportsVolumeTextureRendering && View.FinalPostProcessSettings.ContributingCubemaps.Num() && !IsSimpleDynamicLightingEnabled())
 	{
 		SCOPED_DRAW_EVENT(RHICmdList, InjectAmbientCubemapTranslucentVolumeLighting);
 
