ES2 post process will upscale if UnscaledViewRect is larger than ViewRect (which only occurs when using ES2 emulation).

Fixes UE-5594

diff --git a/Engine/Source/Runtime/Renderer/Private/ForwardShadingRenderer.cpp b/Engine/Source/Runtime/Renderer/Private/ForwardShadingRenderer.cpp
--- a/Engine/Source/Runtime/Renderer/Private/ForwardShadingRenderer.cpp
+++ b/Engine/Source/Runtime/Renderer/Private/ForwardShadingRenderer.cpp
@@ -12,6 +12,7 @@
 #include "SceneFilterRendering.h"
 #include "PostProcessMobile.h"
 #include "SceneUtils.h"
+#include "PostProcessUpscale.h"
 
 uint32 GetShadowQuality();
 
@@ -106,7 +107,9 @@ void FForwardShadingSceneRenderer::Render(FRHICommandListImmediate& RHICmdList)
 	FGlobalDynamicIndexBuffer::Get().Commit();
 
 	const bool bGammaSpace = !IsMobileHDR();
-	if( bGammaSpace )
+	const bool bRequiresUpscale = ((uint32)ViewFamily.RenderTarget->GetSizeXY().X > ViewFamily.FamilySizeX || (uint32)ViewFamily.RenderTarget->GetSizeXY().Y > ViewFamily.FamilySizeY);
+
+	if (bGammaSpace && !bRequiresUpscale)
 	{
 		SetRenderTarget(RHICmdList, ViewFamily.RenderTarget->GetRenderTargetTexture(), GSceneRenderTargets.GetSceneDepthTexture(), ESimpleRenderTargetMode::EClearToDefault);
 	}
@@ -144,43 +147,45 @@ void FForwardShadingSceneRenderer::Render(FRHICommandListImmediate& RHICmdList)
 		RenderTranslucency(RHICmdList);
 	}
 
-	if( !bGammaSpace )
-	{
-		// This might eventually be a problem with multiple views.
-		// Using only view 0 to check to do on-chip transform of alpha.
-		FViewInfo& View = Views[0];
-
-		static const auto CVarMobileMSAA = IConsoleManager::Get().FindTConsoleVariableDataInt(TEXT("r.MobileMSAA"));
+	// This might eventually be a problem with multiple views.
+	// Using only view 0 to check to do on-chip transform of alpha.
+	FViewInfo& View = Views[0];
 
-		bool bOnChipSunMask = 
-			GSupportsRenderTargetFormat_PF_FloatRGBA && 
-			GSupportsShaderFramebufferFetch &&
-			ViewFamily.EngineShowFlags.PostProcessing &&
-			((View.bLightShaftUse) || (View.FinalPostProcessSettings.DepthOfFieldScale > 0.0) || 
-			((ViewFamily.GetShaderPlatform() == SP_METAL) && (CVarMobileMSAA ? CVarMobileMSAA->GetValueOnAnyThread() > 1 : false))
-			);
+	static const auto CVarMobileMSAA = IConsoleManager::Get().FindTConsoleVariableDataInt(TEXT("r.MobileMSAA"));
+	bool bOnChipSunMask =
+		GSupportsRenderTargetFormat_PF_FloatRGBA &&
+		GSupportsShaderFramebufferFetch &&
+		ViewFamily.EngineShowFlags.PostProcessing &&
+		((View.bLightShaftUse) || (View.FinalPostProcessSettings.DepthOfFieldScale > 0.0) ||
+		((ViewFamily.GetShaderPlatform() == SP_METAL) && (CVarMobileMSAA ? CVarMobileMSAA->GetValueOnAnyThread() > 1 : false))
+		);
 
+	if (!bGammaSpace && bOnChipSunMask)
+	{
 		// Convert alpha from depth to circle of confusion with sunshaft intensity.
 		// This is done before resolve on hardware with framebuffer fetch.
-		if(bOnChipSunMask)
-		{
-			// This will break when PrePostSourceViewportSize is not full size.
-			FIntPoint PrePostSourceViewportSize = GSceneRenderTargets.GetBufferSizeXY();
+		// This will break when PrePostSourceViewportSize is not full size.
+		FIntPoint PrePostSourceViewportSize = GSceneRenderTargets.GetBufferSizeXY();
 
-			FMemMark Mark(FMemStack::Get());
-			FRenderingCompositePassContext CompositeContext(RHICmdList, View);
+		FMemMark Mark(FMemStack::Get());
+		FRenderingCompositePassContext CompositeContext(RHICmdList, View);
 
-			FRenderingCompositePass* PostProcessSunMask = CompositeContext.Graph.RegisterPass(new(FMemStack::Get()) FRCPassPostProcessSunMaskES2(PrePostSourceViewportSize, true));
-			CompositeContext.Root->AddDependency(FRenderingCompositeOutputRef(PostProcessSunMask));
-			CompositeContext.Process(TEXT("OnChipAlphaTransform"));
-		}
+		FRenderingCompositePass* PostProcessSunMask = CompositeContext.Graph.RegisterPass(new(FMemStack::Get()) FRCPassPostProcessSunMaskES2(PrePostSourceViewportSize, true));
+		CompositeContext.Root->AddDependency(FRenderingCompositeOutputRef(PostProcessSunMask));
+		CompositeContext.Process(TEXT("OnChipAlphaTransform"));
+	}
 
+	if (!bGammaSpace || bRequiresUpscale)
+	{
 		// Resolve the scene color for post processing.
 		GSceneRenderTargets.ResolveSceneColor(RHICmdList, FResolveRect(0, 0, ViewFamily.FamilySizeX, ViewFamily.FamilySizeY));
 
 		// Drop depth and stencil before post processing to avoid export.
 		RHICmdList.DiscardRenderTargets(true, true, 0);
+	}
 
+	if (!bGammaSpace)
+	{
 		// Finish rendering for each view, or the full stereo buffer if enabled
 		if (ViewFamily.bResolveScene)
 		{
@@ -199,13 +204,53 @@ void FForwardShadingSceneRenderer::Render(FRHICommandListImmediate& RHICmdList)
 				SCOPED_DRAW_EVENT(RHICmdList, PostProcessing);
 				SCOPE_CYCLE_COUNTER(STAT_FinishRenderViewTargetTime);
 				for(int32 ViewIndex = 0;ViewIndex < Views.Num();ViewIndex++)
-				{	
+				{
 					SCOPED_CONDITIONAL_DRAW_EVENTF(RHICmdList, EventView, Views.Num() > 1, TEXT("View%d"), ViewIndex);
 					GPostProcessing.ProcessES2(RHICmdList, Views[ViewIndex], bOnChipSunMask);
 				}
 			}
 		}
 	}
-
+	else if (bRequiresUpscale)
+	{
+		SimpleUpscale(RHICmdList, View);
+	}
 	RenderFinish(RHICmdList);
 }
+
+// Perform simple upscale when required and post process is not in use.
+void FForwardShadingSceneRenderer::SimpleUpscale(FRHICommandListImmediate& RHICmdList, FViewInfo &View)
+{
+	FRenderingCompositePassContext CompositeContext(RHICmdList, View);
+	FPostprocessContext Context(CompositeContext.Graph, View);
+	// simple bilinear upscaling for ES2.
+	FRenderingCompositePass* Node = Context.Graph.RegisterPass(new(FMemStack::Get()) FRCPassPostProcessUpscale(1, 0.0f));
+
+	Node->SetInput(ePId_Input0, FRenderingCompositeOutputRef(Context.FinalOutput));
+	Node->SetInput(ePId_Input1, FRenderingCompositeOutputRef(Context.FinalOutput));
+
+	Context.FinalOutput = FRenderingCompositeOutputRef(Node);
+
+	{
+		// currently created on the heap each frame but View.Family->RenderTarget could keep this object and all would be cleaner
+		TRefCountPtr<IPooledRenderTarget> Temp;
+		FSceneRenderTargetItem Item;
+		Item.TargetableTexture = (FTextureRHIRef&)View.Family->RenderTarget->GetRenderTargetTexture();
+		Item.ShaderResourceTexture = (FTextureRHIRef&)View.Family->RenderTarget->GetRenderTargetTexture();
+
+		FPooledRenderTargetDesc Desc;
+
+		Desc.Extent = View.Family->RenderTarget->GetSizeXY();
+		// todo: this should come from View.Family->RenderTarget
+		Desc.Format = PF_B8G8R8A8;
+		Desc.NumMips = 1;
+
+		GRenderTargetPool.CreateUntrackedElement(Desc, Temp, Item);
+
+		Context.FinalOutput.GetOutput()->PooledRenderTarget = Temp;
+		Context.FinalOutput.GetOutput()->RenderTargetDesc = Desc;
+	}
+
+	CompositeContext.Root->AddDependency(Context.FinalOutput);
+	CompositeContext.Process(TEXT("ES2Upscale"));
+}
diff --git a/Engine/Source/Runtime/Renderer/Private/PostProcess/PostProcessMobile.cpp b/Engine/Source/Runtime/Renderer/Private/PostProcess/PostProcessMobile.cpp
--- a/Engine/Source/Runtime/Renderer/Private/PostProcess/PostProcessMobile.cpp
+++ b/Engine/Source/Runtime/Renderer/Private/PostProcess/PostProcessMobile.cpp
@@ -2761,7 +2761,7 @@ void FRCPassPostProcessAaES2::Process(FRenderingCompositePassContext& Context)
 
 	SetShader(Context);
 
-	const FIntRect& ViewRect = Context.View.ViewRect;
+	const FIntRect& ViewRect = Context.View.UnscaledViewRect; // Simple upscaling, ES2 post process does not currently have a specific upscaling pass.
 	float XPos = ViewRect.Min.X;
 	float YPos = ViewRect.Min.Y;
 	float Width = ViewRect.Width();
diff --git a/Engine/Source/Runtime/Renderer/Private/PostProcess/PostProcessTonemap.cpp b/Engine/Source/Runtime/Renderer/Private/PostProcess/PostProcessTonemap.cpp
--- a/Engine/Source/Runtime/Renderer/Private/PostProcess/PostProcessTonemap.cpp
+++ b/Engine/Source/Runtime/Renderer/Private/PostProcess/PostProcessTonemap.cpp
@@ -1609,7 +1609,7 @@ void FRCPassPostProcessTonemapES2::Process(FRenderingCompositePassContext& Conte
 	const FPooledRenderTargetDesc& OutputDesc = PassOutputs[0].RenderTargetDesc;
 
 	FIntRect SrcRect = ViewRect;
-	FIntRect DestRect = View.ViewRect;
+	FIntRect DestRect = View.UnscaledViewRect; // Simple upscaling, ES2 post process does not currently have a specific upscaling pass.
 	FIntPoint SrcSize = InputDesc->Extent;
 	FIntPoint DstSize = OutputDesc.Extent;
 
diff --git a/Engine/Source/Runtime/Renderer/Private/PostProcess/PostProcessUpscale.cpp b/Engine/Source/Runtime/Renderer/Private/PostProcess/PostProcessUpscale.cpp
--- a/Engine/Source/Runtime/Renderer/Private/PostProcess/PostProcessUpscale.cpp
+++ b/Engine/Source/Runtime/Renderer/Private/PostProcess/PostProcessUpscale.cpp
@@ -16,7 +16,7 @@
 static TAutoConsoleVariable<float> CVarUpscaleSoftness(
 	TEXT("r.Upscale.Softness"),
 	0.3f,
-	TEXT("To scale up with higher quality loosing some sharpness\n")
+	TEXT("To scale up with higher quality losing some sharpness\n")
 	TEXT(" 0..1 (0.3 is good for ScreenPercentage 90"),
 	ECVF_Scalability | ECVF_RenderThreadSafe);
 
@@ -40,7 +40,7 @@ public:
 
 	static bool ShouldCache(EShaderPlatform Platform)
 	{
-		return IsFeatureLevelSupported(Platform, ERHIFeatureLevel::SM4);
+		return true;
 	}
 
 	static void ModifyCompilationEnvironment(EShaderPlatform Platform, FShaderCompilerEnvironment& OutEnvironment)
@@ -90,6 +90,12 @@ class FPostProcessUpscalePS : public FGlobalShader
 
 	static bool ShouldCache(EShaderPlatform Platform)
 	{
+		// Always allow simple bilinear upscale. (Provides upscaling for ES2 emulation)
+		if (Method == 1)
+		{
+			return true;
+		}
+
 		return IsFeatureLevelSupported(Platform, ERHIFeatureLevel::SM4);
 	}
 
diff --git a/Engine/Source/Runtime/Renderer/Private/PostProcess/PostProcessing.cpp b/Engine/Source/Runtime/Renderer/Private/PostProcess/PostProcessing.cpp
--- a/Engine/Source/Runtime/Renderer/Private/PostProcess/PostProcessing.cpp
+++ b/Engine/Source/Runtime/Renderer/Private/PostProcess/PostProcessing.cpp
@@ -1696,7 +1696,7 @@ void FPostProcessing::ProcessES2(FRHICommandListImmediate& RHICmdList, FViewInfo
 		PostProcessTonemap->SetInput(ePId_Input2, DofOutput);
 		Context.FinalOutput = FRenderingCompositeOutputRef(PostProcessTonemap);
 		// if Context.FinalOutput was the clipped result of sunmask stage then this stage also restores Context.FinalOutput back original target size.
-		FinalOutputViewRect = View.ViewRect;
+		FinalOutputViewRect = View.UnscaledViewRect;
 
 		if(bUseAa && View.Family->EngineShowFlags.PostProcessing)
 		{
diff --git a/Engine/Source/Runtime/Renderer/Private/SceneRendering.h b/Engine/Source/Runtime/Renderer/Private/SceneRendering.h
--- a/Engine/Source/Runtime/Renderer/Private/SceneRendering.h
+++ b/Engine/Source/Runtime/Renderer/Private/SceneRendering.h
@@ -721,6 +721,9 @@ protected:
 	/** Renders any necessary shadowmaps. */
 	void RenderShadowDepthMaps(FRHICommandListImmediate& RHICmdList);
 
+	/** Perform upscaling when post process is not used. */
+	void SimpleUpscale(FRHICommandListImmediate& RHICmdList, FViewInfo &View);
+
 	/**
 	  * Used by RenderShadowDepthMaps to render shadowmap for the given light.
 	  *
