UE4 - fix UE-6438, missing shadow shaders for instanced static meshes that modify vertex pos. And UE-6439, a crash related to not having a cluster tree at all.

diff --git a/Engine/Source/Runtime/Engine/Private/HierarchicalInstancedStaticMesh.cpp b/Engine/Source/Runtime/Engine/Private/HierarchicalInstancedStaticMesh.cpp
--- a/Engine/Source/Runtime/Engine/Private/HierarchicalInstancedStaticMesh.cpp
+++ b/Engine/Source/Runtime/Engine/Private/HierarchicalInstancedStaticMesh.cpp
@@ -1034,6 +1034,7 @@ void FHierarchicalStaticMeshSceneProxy::GetDynamicMeshElements(const TArray<cons
 			ElementParams.View = View;
 
 			// Render built instances
+			if (ClusterTree.Num())
 			{
 
 				FFoliageCullInstanceParams InstanceParams(bSingleSections, bMultipleSections, ClusterTree);
diff --git a/Engine/Source/Runtime/Renderer/Private/ShadowRendering.cpp b/Engine/Source/Runtime/Renderer/Private/ShadowRendering.cpp
--- a/Engine/Source/Runtime/Renderer/Private/ShadowRendering.cpp
+++ b/Engine/Source/Runtime/Renderer/Private/ShadowRendering.cpp
@@ -546,6 +546,7 @@ public:
 			return (Material->IsSpecialEngineMaterial()
 					// Only compile for masked or lit translucent materials
 					|| Material->IsMasked()
+					|| (Material->MaterialMayModifyMeshPosition() && Material->IsUsedWithInstancedStaticMeshes())
 					// Perspective correct rendering needs a pixel shader and WPO materials can't be overridden with default material.
 					|| (ShaderMode == PixelShadowDepth_PerspectiveCorrect && Material->MaterialMayModifyMeshPosition()))
 				// Only compile one pass point light shaders for feature levels >= SM4
