UE4 - UE-6380 - Fix crash when warn on no precomputed vis on dragging a skylight

diff --git a/Engine/Source/Runtime/Renderer/Private/SceneRendering.cpp b/Engine/Source/Runtime/Renderer/Private/SceneRendering.cpp
--- a/Engine/Source/Runtime/Renderer/Private/SceneRendering.cpp
+++ b/Engine/Source/Runtime/Renderer/Private/SceneRendering.cpp
@@ -821,70 +821,72 @@ void FSceneRenderer::RenderFinish(FRHICommandListImmediate& RHICmdList)
 	
 #if !(UE_BUILD_SHIPPING || UE_BUILD_TEST)
 	{
+		bool bShowPrecomputedVisibilityWarning = false;
+		static const auto* CVarPrecomputedVisibilityWarning = IConsoleManager::Get().FindTConsoleVariableDataInt(TEXT("r.PrecomputedVisibilityWarning"));
+		if (CVarPrecomputedVisibilityWarning && CVarPrecomputedVisibilityWarning->GetValueOnRenderThread() == 1)
+		{
+			bShowPrecomputedVisibilityWarning = !bUsedPrecomputedVisibility;
+		}
+
 		for(int32 ViewIndex = 0;ViewIndex < Views.Num();ViewIndex++)
 		{	
-			SCOPED_CONDITIONAL_DRAW_EVENTF(RHICmdList, EventView, Views.Num() > 1, TEXT("View%d"), ViewIndex);
 			FViewInfo& View = Views[ViewIndex];
-
-			bool bShowPrecomputedVisibilityWarning = false;
-			static const auto* CVarPrecomputedVisibilityWarning = IConsoleManager::Get().FindTConsoleVariableDataInt(TEXT("r.PrecomputedVisibilityWarning"));
-			if (CVarPrecomputedVisibilityWarning && CVarPrecomputedVisibilityWarning->GetValueOnRenderThread() == 1)
+			if (!View.bIsReflectionCapture && !View.bIsSceneCapture )
 			{
-				bShowPrecomputedVisibilityWarning = !bUsedPrecomputedVisibility;
-			}
-
-			// display a message saying we're frozen
-			FSceneViewState* ViewState = (FSceneViewState*)View.State;
-			bool bViewParentOrFrozen = ViewState && (ViewState->HasViewParent() || ViewState->bIsFrozen);
-			bool bLocked = View.bIsLocked;
-			if (bViewParentOrFrozen || bShowPrecomputedVisibilityWarning || bLocked)
-			{
-				// this is a helper class for FCanvas to be able to get screen size
-				class FRenderTargetTemp : public FRenderTarget
+				SCOPED_CONDITIONAL_DRAW_EVENTF(RHICmdList, EventView, Views.Num() > 1, TEXT("View%d"), ViewIndex);
+				// display a message saying we're frozen
+				FSceneViewState* ViewState = (FSceneViewState*)View.State;
+				bool bViewParentOrFrozen = ViewState && (ViewState->HasViewParent() || ViewState->bIsFrozen);
+				bool bLocked = View.bIsLocked;
+				if (bViewParentOrFrozen || bShowPrecomputedVisibilityWarning || bLocked)
 				{
-				public:
-					FViewInfo& View;
-
-					FRenderTargetTemp(FViewInfo& InView) : View(InView)
+					// this is a helper class for FCanvas to be able to get screen size
+					class FRenderTargetTemp : public FRenderTarget
 					{
+					public:
+						FViewInfo& View;
+
+						FRenderTargetTemp(FViewInfo& InView) : View(InView)
+						{
+						}
+						virtual FIntPoint GetSizeXY() const
+						{
+							return View.ViewRect.Size();
+						};
+						virtual const FTexture2DRHIRef& GetRenderTargetTexture() const
+						{
+							return View.Family->RenderTarget->GetRenderTargetTexture();
+						}
+					} TempRenderTarget(View);
+
+					// create a temporary FCanvas object with the temp render target
+					// so it can get the screen size
+					int32 Y = 130;
+					FCanvas Canvas(&TempRenderTarget, NULL, View.Family->CurrentRealTime, View.Family->CurrentWorldTime, View.Family->DeltaWorldTime, FeatureLevel);
+					if (bViewParentOrFrozen)
+					{
+						const FText StateText =
+							ViewState->bIsFrozen ?
+							NSLOCTEXT("SceneRendering", "RenderingFrozen", "Rendering frozen...")
+							:
+							NSLOCTEXT("SceneRendering", "OcclusionChild", "Occlusion Child");
+						Canvas.DrawShadowedText(10, Y, StateText, GetStatsFont(), FLinearColor(0.8, 1.0, 0.2, 1.0));
+						Y += 14;
 					}
-					virtual FIntPoint GetSizeXY() const
+					if (bShowPrecomputedVisibilityWarning)
 					{
-						return View.ViewRect.Size();
-					};
-					virtual const FTexture2DRHIRef& GetRenderTargetTexture() const
+						const FText Message = NSLOCTEXT("Renderer", "NoPrecomputedVisibility", "NO PRECOMPUTED VISIBILITY");
+						Canvas.DrawShadowedText(10, Y, Message, GetStatsFont(), FLinearColor(1.0, 0.05, 0.05, 1.0));
+						Y += 14;
+					}
+					if (bLocked)
 					{
-						return View.Family->RenderTarget->GetRenderTargetTexture();
+						const FText Message = NSLOCTEXT("Renderer", "ViewLocked", "VIEW LOCKED");
+						Canvas.DrawShadowedText(10, Y, Message, GetStatsFont(), FLinearColor(0.8, 1.0, 0.2, 1.0));
+						Y += 14;
 					}
-				} TempRenderTarget(View);
-
-				// create a temporary FCanvas object with the temp render target
-				// so it can get the screen size
-				int32 Y = 130;
-				FCanvas Canvas(&TempRenderTarget, NULL, View.Family->CurrentRealTime, View.Family->CurrentWorldTime, View.Family->DeltaWorldTime, FeatureLevel);
-				if (bViewParentOrFrozen)
-				{
-					const FText StateText =
-						ViewState->bIsFrozen ?
-						NSLOCTEXT("SceneRendering", "RenderingFrozen", "Rendering frozen...")
-						:
-						NSLOCTEXT("SceneRendering", "OcclusionChild", "Occlusion Child");
-					Canvas.DrawShadowedText( 10, Y, StateText, GetStatsFont(), FLinearColor(0.8,1.0,0.2,1.0));
-					Y += 14;
-				}
-				if (bShowPrecomputedVisibilityWarning)
-				{
-					const FText Message = NSLOCTEXT("Renderer", "NoPrecomputedVisibility", "NO PRECOMPUTED VISIBILITY");
-					Canvas.DrawShadowedText( 10, Y, Message, GetStatsFont(), FLinearColor(1.0,0.05,0.05,1.0));
-					Y += 14;
-				}
-				if (bLocked)
-				{
-					const FText Message = NSLOCTEXT("Renderer", "ViewLocked", "VIEW LOCKED");
-					Canvas.DrawShadowedText( 10, Y, Message, GetStatsFont(), FLinearColor(0.8,1.0,0.2,1.0));
-					Y += 14;
+					Canvas.Flush_RenderThread(RHICmdList);
 				}
-				Canvas.Flush_RenderThread(RHICmdList);
 			}
 		}
 	}
