All FScene's detect distance field atlas reallocation properly now, fixes artifacts when the atlas was reallocated in PIE

diff --git a/Engine/Source/Runtime/Engine/Private/DistanceFieldAtlas.cpp b/Engine/Source/Runtime/Engine/Private/DistanceFieldAtlas.cpp
--- a/Engine/Source/Runtime/Engine/Private/DistanceFieldAtlas.cpp
+++ b/Engine/Source/Runtime/Engine/Private/DistanceFieldAtlas.cpp
@@ -20,6 +20,7 @@ const int32 MaxAtlasDimension = 512;
 FDistanceFieldVolumeTextureAtlas::FDistanceFieldVolumeTextureAtlas(EPixelFormat InFormat) :
 	BlockAllocator(0, 0, 0, MaxAtlasDimension, MaxAtlasDimension, MaxAtlasDimension, false, false)
 {
+	Generation = 0;
 	Format = InFormat;
 }
  
@@ -49,10 +50,8 @@ struct FCompareVolumeAllocation
 	}
 };
 
-bool FDistanceFieldVolumeTextureAtlas::UpdateAllocations()
+void FDistanceFieldVolumeTextureAtlas::UpdateAllocations()
 {
-	bool bRecomputedLayouts = false;
-
 	if (PendingAllocations.Num() > 0)
 	{
 		// Sort largest to smallest for best packing
@@ -81,7 +80,7 @@ bool FDistanceFieldVolumeTextureAtlas::UpdateAllocations()
 				// Remove all allocations from the layout so we have a clean slate
 				BlockAllocator = FTextureLayout3d(0, 0, 0, MaxAtlasDimension, MaxAtlasDimension, MaxAtlasDimension, false, false);
 				
-				bRecomputedLayouts = true;
+				Generation++;
 
 				// Re-upload all textures since we had to reallocate
 				PendingAllocations.Append(CurrentAllocations);
@@ -146,8 +145,6 @@ bool FDistanceFieldVolumeTextureAtlas::UpdateAllocations()
 		CurrentAllocations.Append(PendingAllocations);
 		PendingAllocations.Empty();
 	}
-
-	return bRecomputedLayouts;
 }
 
 void FDistanceFieldVolumeTexture::Initialize()
diff --git a/Engine/Source/Runtime/Engine/Public/DistanceFieldAtlas.h b/Engine/Source/Runtime/Engine/Public/DistanceFieldAtlas.h
--- a/Engine/Source/Runtime/Engine/Public/DistanceFieldAtlas.h
+++ b/Engine/Source/Runtime/Engine/Public/DistanceFieldAtlas.h
@@ -78,7 +78,9 @@ public:
 	void RemoveAllocation(FDistanceFieldVolumeTexture* Texture);
 
 	/** Reallocates the volume texture if necessary and uploads new allocations. */
-	ENGINE_API bool UpdateAllocations();
+	ENGINE_API void UpdateAllocations();
+
+	int32 GetGeneration() const { return Generation; }
 
 	EPixelFormat Format;
 	FTexture3DRHIRef VolumeTextureRHI;
@@ -92,6 +94,9 @@ private:
 
 	/** Allocations that have already been added, stored in case we need to realloc. */
 	TArray<FDistanceFieldVolumeTexture*> CurrentAllocations;
+
+	/** Incremented when the atlas is reallocated, so dependencies know to update. */
+	int32 Generation;
 };
 
 extern ENGINE_API TGlobalResource<FDistanceFieldVolumeTextureAtlas> GDistanceFieldVolumeTextureAtlas;
diff --git a/Engine/Source/Runtime/Renderer/Private/DeferredShadingRenderer.cpp b/Engine/Source/Runtime/Renderer/Private/DeferredShadingRenderer.cpp
--- a/Engine/Source/Runtime/Renderer/Private/DeferredShadingRenderer.cpp
+++ b/Engine/Source/Runtime/Renderer/Private/DeferredShadingRenderer.cpp
@@ -649,8 +649,8 @@ void FDeferredShadingSceneRenderer::Render(FRHICommandListImmediate& RHICmdList)
 
 	if (ShouldPrepareForDistanceFieldAO() || ShouldPrepareForDistanceFieldShadows())
 	{
-		const bool bReallocatedAtlasLayouts = GDistanceFieldVolumeTextureAtlas.UpdateAllocations();
-		UpdateGlobalDistanceFieldObjectBuffers(RHICmdList, bReallocatedAtlasLayouts);
+		GDistanceFieldVolumeTextureAtlas.UpdateAllocations();
+		UpdateGlobalDistanceFieldObjectBuffers(RHICmdList);
 	}
 
 	const bool bIsWireframe = ViewFamily.EngineShowFlags.Wireframe;
diff --git a/Engine/Source/Runtime/Renderer/Private/DeferredShadingRenderer.h b/Engine/Source/Runtime/Renderer/Private/DeferredShadingRenderer.h
--- a/Engine/Source/Runtime/Renderer/Private/DeferredShadingRenderer.h
+++ b/Engine/Source/Runtime/Renderer/Private/DeferredShadingRenderer.h
@@ -351,7 +351,7 @@ private:
 	bool ShouldPrepareForDistanceFieldShadows() const;
 	bool ShouldPrepareForDistanceFieldAO() const;
 
-	void UpdateGlobalDistanceFieldObjectBuffers(FRHICommandListImmediate& RHICmdList, bool bReallocatedAtlasLayouts);
+	void UpdateGlobalDistanceFieldObjectBuffers(FRHICommandListImmediate& RHICmdList);
 
 	friend class FTranslucentPrimSet;
 };
diff --git a/Engine/Source/Runtime/Renderer/Private/DistanceFieldSurfaceCacheLighting.cpp b/Engine/Source/Runtime/Renderer/Private/DistanceFieldSurfaceCacheLighting.cpp
--- a/Engine/Source/Runtime/Renderer/Private/DistanceFieldSurfaceCacheLighting.cpp
+++ b/Engine/Source/Runtime/Renderer/Private/DistanceFieldSurfaceCacheLighting.cpp
@@ -845,7 +845,7 @@ private:
 IMPLEMENT_SHADER_TYPE(,FCullObjectsForViewCS,TEXT("DistanceFieldSurfaceCacheLightingCompute"),TEXT("CullObjectsForViewCS"),SF_Compute);
 
 
-void FDeferredShadingSceneRenderer::UpdateGlobalDistanceFieldObjectBuffers(FRHICommandListImmediate& RHICmdList, bool bReallocatedAtlasLayouts)
+void FDeferredShadingSceneRenderer::UpdateGlobalDistanceFieldObjectBuffers(FRHICommandListImmediate& RHICmdList)
 {
 	FDistanceFieldSceneData& DistanceFieldSceneData = Scene->DistanceFieldSceneData;
 
@@ -853,7 +853,7 @@ void FDeferredShadingSceneRenderer::UpdateGlobalDistanceFieldObjectBuffers(FRHIC
 		&& (DistanceFieldSceneData.PendingAddOperations.Num() > 0 
 			|| DistanceFieldSceneData.PendingUpdateOperations.Num() > 0 
 			|| DistanceFieldSceneData.PendingRemoveOperations.Num() > 0
-			|| bReallocatedAtlasLayouts))
+			|| DistanceFieldSceneData.AtlasGeneration != GDistanceFieldVolumeTextureAtlas.GetGeneration()))
 	{
 		QUICK_SCOPE_CYCLE_COUNTER(STAT_UpdateObjectData);
 		SCOPED_DRAW_EVENT(RHICmdList, UpdateSceneObjectData);
@@ -863,8 +863,10 @@ void FDeferredShadingSceneRenderer::UpdateGlobalDistanceFieldObjectBuffers(FRHIC
 			DistanceFieldSceneData.ObjectBuffers = new FDistanceFieldObjectBuffers();
 		}
 
-		if (bReallocatedAtlasLayouts)
+		if (DistanceFieldSceneData.AtlasGeneration != GDistanceFieldVolumeTextureAtlas.GetGeneration())
 		{
+			DistanceFieldSceneData.AtlasGeneration = GDistanceFieldVolumeTextureAtlas.GetGeneration();
+
 			for (int32 PrimitiveInstanceIndex = 0; PrimitiveInstanceIndex < DistanceFieldSceneData.PrimitiveInstanceMapping.Num(); PrimitiveInstanceIndex++)
 			{
 				FPrimitiveAndInstance& PrimitiveInstance = DistanceFieldSceneData.PrimitiveInstanceMapping[PrimitiveInstanceIndex];
diff --git a/Engine/Source/Runtime/Renderer/Private/RendererScene.cpp b/Engine/Source/Runtime/Renderer/Private/RendererScene.cpp
--- a/Engine/Source/Runtime/Renderer/Private/RendererScene.cpp
+++ b/Engine/Source/Runtime/Renderer/Private/RendererScene.cpp
@@ -79,6 +79,7 @@ FSceneViewState::FSceneViewState()
 FDistanceFieldSceneData::FDistanceFieldSceneData() 
 	: NumObjectsInBuffer(0)
 	, ObjectBuffers(NULL)
+	, AtlasGeneration(0)
 {
 
 }
diff --git a/Engine/Source/Runtime/Renderer/Private/ScenePrivate.h b/Engine/Source/Runtime/Renderer/Private/ScenePrivate.h
--- a/Engine/Source/Runtime/Renderer/Private/ScenePrivate.h
+++ b/Engine/Source/Runtime/Renderer/Private/ScenePrivate.h
@@ -627,6 +627,8 @@ public:
 		}
 	}
 
+	void TrimHistoryRenderTargets(const FScene* Scene);
+
 	/** 
 	 * Called every frame after UpdateLastRenderTime, sets up the information for the lagged temporal LOD transition
 	 */
@@ -948,6 +950,9 @@ public:
 	TArray<FPrimitiveSceneInfo*> PendingAddOperations;
 	TArray<FPrimitiveSceneInfo*> PendingUpdateOperations;
 	TArray<int32> PendingRemoveOperations;
+
+	/** Used to detect atlas reallocations, since objects store UVs into the atlas and need to be updated when it changes. */
+	int32 AtlasGeneration;
 };
 
 /** Stores data for an allocation in the FIndirectLightingCache. */
