Light shaft render targets in the view state are now released when the corresponding light is deregistered

diff --git a/Engine/Source/Runtime/Renderer/Private/LightShaftRendering.cpp b/Engine/Source/Runtime/Renderer/Private/LightShaftRendering.cpp
--- a/Engine/Source/Runtime/Renderer/Private/LightShaftRendering.cpp
+++ b/Engine/Source/Runtime/Renderer/Private/LightShaftRendering.cpp
@@ -873,6 +873,32 @@ void ApplyLightShaftBloom(FRHICommandListImmediate& RHICmdList, const FViewInfo&
 	GSceneRenderTargets.FinishRenderingSceneColor(RHICmdList, false);
 }
 
+void FSceneViewState::TrimHistoryRenderTargets(const FScene* Scene)
+{
+	for (TMap<const ULightComponent*, TRefCountPtr<IPooledRenderTarget> >::TIterator It(LightShaftBloomHistoryRTs); It; ++It)
+	{
+		bool bLightIsUsed = false;
+
+		for (TSparseArray<FLightSceneInfoCompact>::TConstIterator LightIt(Scene->Lights); LightIt; ++LightIt)
+		{
+			const FLightSceneInfo* const LightSceneInfo = LightIt->LightSceneInfo;
+
+			if (LightSceneInfo->Proxy->GetLightComponent() == It.Key())
+			{
+				bLightIsUsed = true;
+				break;
+			}
+		}
+
+		if (!bLightIsUsed)
+		{
+			// Remove references to render targets for lights that are no longer in the scene
+			// This has to be done every frame instead of at light deregister time because the view states are not known by FScene
+			It.RemoveCurrent();
+		}
+	}
+}
+
 void FDeferredShadingSceneRenderer::RenderLightShaftBloom(FRHICommandListImmediate& RHICmdList)
 {
 	if (DoesViewFamilyAllowLightShafts(ViewFamily))
diff --git a/Engine/Source/Runtime/Renderer/Private/SceneVisibility.cpp b/Engine/Source/Runtime/Renderer/Private/SceneVisibility.cpp
--- a/Engine/Source/Runtime/Renderer/Private/SceneVisibility.cpp
+++ b/Engine/Source/Runtime/Renderer/Private/SceneVisibility.cpp
@@ -2032,6 +2032,11 @@ void FSceneRenderer::PostVisibilityFrameSetup()
 
 		// sort the translucent primitives
 		View.TranslucentPrimSet.SortPrimitives();
+
+		if (View.State)
+		{
+			((FSceneViewState*)View.State)->TrimHistoryRenderTargets(Scene);
+		}
 	}
 
 	bool bCheckLightShafts = false;
