Rendering: Add a project-configurable translucent sort policy (WIP)

r.TranslucentSortPolicy can take the following values:
  0: Sort based on distance from camera centerpoint to bounding sphere centerpoint. (default, best for 3D games)
  1: Sort based on the post-projection Z distance to the camera.
  2: Sort based on the projection onto a fixed axis. (best for 2D games)

diff --git a/Engine/Source/Runtime/Core/Private/HAL/ConsoleManager.cpp b/Engine/Source/Runtime/Core/Private/HAL/ConsoleManager.cpp
--- a/Engine/Source/Runtime/Core/Private/HAL/ConsoleManager.cpp
+++ b/Engine/Source/Runtime/Core/Private/HAL/ConsoleManager.cpp
@@ -1619,6 +1619,14 @@ static TAutoConsoleVariable<int32> CVarSimpleDynamicLighting(
 	TEXT("0:off, 1:on"),
 	ECVF_RenderThreadSafe);
 
+static TAutoConsoleVariable<int32> CVarTranslucentSortPolicy(
+	TEXT("r.TranslucentSortPolicy"),
+	0,
+	TEXT("0: Sort based on distance from camera centerpoint to bounding sphere centerpoint. (default, best for 3D games)\n")
+	TEXT("1: Sort based on projected distance to camera.")
+	TEXT("2: Sort based on the projection onto a fixed axis. (best for 2D games)"),
+	ECVF_RenderThreadSafe);
+
 static TAutoConsoleVariable<int32> CVarMobileHDR(
 	TEXT("r.MobileHDR"),
 	1,
diff --git a/Engine/Source/Runtime/Engine/Classes/Engine/EngineTypes.h b/Engine/Source/Runtime/Engine/Classes/Engine/EngineTypes.h
--- a/Engine/Source/Runtime/Engine/Classes/Engine/EngineTypes.h
+++ b/Engine/Source/Runtime/Engine/Classes/Engine/EngineTypes.h
@@ -123,6 +123,25 @@ enum ETranslucencyLightingMode
 	TLM_MAX,
 };
 
+/**
+ * Enumerates available options for the translucency sort policy.
+ */
+UENUM()
+namespace ETranslucentSortPolicy
+{
+	enum Type
+	{
+		// Sort based on distance from camera centerpoint to bounding sphere centerpoint. (default, best for 3D games)
+		SortByDistance = 0,
+
+		// Sort based on the post-projection Z distance to the camera.
+		SortByProjectedZ = 1,
+
+		// Sort based on the projection onto a fixed axis. (best for 2D games)
+		SortAlongAxis = 2,
+	};
+}
+
 /** Controls the way that the width scale property affects anim trails */
 UENUM()
 enum ETrailWidthMode
diff --git a/Engine/Source/Runtime/Engine/Classes/Engine/RendererSettings.h b/Engine/Source/Runtime/Engine/Classes/Engine/RendererSettings.h
--- a/Engine/Source/Runtime/Engine/Classes/Engine/RendererSettings.h
+++ b/Engine/Source/Runtime/Engine/Classes/Engine/RendererSettings.h
@@ -169,6 +169,11 @@ class ENGINE_API URendererSettings
 		ToolTip="Allow translucency to be rendered to a separate render targeted and composited after depth of field. Prevents translucency from appearing out of focus."))
 	uint32 bSeparateTranslucency:1;
 
+	UPROPERTY(config, EditAnywhere, Category=Translucency, meta=(
+		ConsoleVariable="r.TranslucentSortPolicy",DisplayName="Translucent Sort Policy",
+		ToolTip="The sort mode for translucent primitives, affecting how they are ordered and how they change order as the camera moves."))
+	TEnumAsByte<ETranslucentSortPolicy::Type> TranslucentSortPolicy;
+
 	UPROPERTY(config, EditAnywhere, Category=Postprocessing, meta=(
 		ConsoleVariable="r.CustomDepth",DisplayName="Custom Depth Pass",
 		ToolTip="Whether the custom depth pass for tagging primitives for postprocessing passes is enabled. Enabling it on demand can save memory but may cause a hitch the first time the feature is used."))
diff --git a/Engine/Source/Runtime/Engine/Private/SceneView.cpp b/Engine/Source/Runtime/Engine/Private/SceneView.cpp
--- a/Engine/Source/Runtime/Engine/Private/SceneView.cpp
+++ b/Engine/Source/Runtime/Engine/Private/SceneView.cpp
@@ -421,6 +421,12 @@ FSceneView::FSceneView(const FSceneViewInitOptions& InitOptions)
 	// into world oriented z.
 	InvDeviceZToWorldZTransform = CreateInvDeviceZToWorldZTransform(ProjectionMatrixUnadjustedForRHI);
 
+	static TConsoleVariableData<int32>* SortPolicyCvar = IConsoleManager::Get().FindTConsoleVariableDataInt(TEXT("r.TranslucentSortPolicy"));
+	TranslucentSortPolicy = static_cast<ETranslucentSortPolicy::Type>(SortPolicyCvar->GetValueOnAnyThread());
+
+	//@TODO: PAPER2D: Should come from the config value if/when axis switching is allowed
+	TranslucentSortAxis = FVector(0.0f, -1.0f, 0.0f);
+
 	// As the world is only accessable from the game thread, bIsGameView should be explicitly
 	// set on any other thread.
 	if(IsInGameThread())
diff --git a/Engine/Source/Runtime/Engine/Public/SceneView.h b/Engine/Source/Runtime/Engine/Public/SceneView.h
--- a/Engine/Source/Runtime/Engine/Public/SceneView.h
+++ b/Engine/Source/Runtime/Engine/Public/SceneView.h
@@ -450,6 +450,12 @@ public:
 	 **/
 	FIntRect CameraConstrainedViewRect;
 
+	/** Sort axis for when TranslucentSortPolicy is SortAlongAxis */
+	FVector TranslucentSortAxis;
+
+	/** Translucent sort mode */
+	TEnumAsByte<ETranslucentSortPolicy::Type> TranslucentSortPolicy;
+	
 #if WITH_EDITOR
 	/** The set of (the first 64) groups' visibility info for this view */
 	uint64 EditorViewBitflag;
diff --git a/Engine/Source/Runtime/Renderer/Private/TranslucentRendering.cpp b/Engine/Source/Runtime/Renderer/Private/TranslucentRendering.cpp
--- a/Engine/Source/Runtime/Renderer/Private/TranslucentRendering.cpp
+++ b/Engine/Source/Runtime/Renderer/Private/TranslucentRendering.cpp
@@ -700,6 +700,31 @@ void FTranslucentPrimSet::RenderPrimitive(
 	}
 }
 
+inline float CalculateTranslucentSortKey(FPrimitiveSceneInfo* PrimitiveSceneInfo, const FViewInfo& ViewInfo)
+{
+	float SortKey = 0.0f;
+	if (ViewInfo.TranslucentSortPolicy == ETranslucentSortPolicy::SortByDistance)
+	{
+		//sort based on distance to the view position, view rotation is not a factor
+		SortKey = (PrimitiveSceneInfo->Proxy->GetBounds().Origin - ViewInfo.ViewMatrices.ViewOrigin).Size();
+		// UE4_TODO: also account for DPG in the sort key.
+	}
+	else if (ViewInfo.TranslucentSortPolicy == ETranslucentSortPolicy::SortAlongAxis)
+	{
+		// Sort based on enforced orthogonal distance
+		const FVector CameraToObject = PrimitiveSceneInfo->Proxy->GetBounds().Origin - ViewInfo.ViewMatrices.ViewOrigin;
+		SortKey = FVector::DotProduct(CameraToObject, ViewInfo.TranslucentSortAxis);
+	}
+	else
+	{
+		// Sort based on projected Z distance
+		check(ViewInfo.TranslucentSortPolicy == ETranslucentSortPolicy::SortByProjectedZ);
+		SortKey = ViewInfo.ViewMatrices.ViewMatrix.TransformPosition(PrimitiveSceneInfo->Proxy->GetBounds().Origin).Z;
+	}
+
+	return SortKey;
+}
+
 /**
 * Add a new primitive to the list of sorted prims
 * @param PrimitiveSceneInfo - primitive info to add. Origin of bounds is used for sort.
@@ -707,10 +732,7 @@ void FTranslucentPrimSet::RenderPrimitive(
 */
 void FTranslucentPrimSet::AddScenePrimitive(FPrimitiveSceneInfo* PrimitiveSceneInfo, const FViewInfo& ViewInfo, bool bUseNormalTranslucency, bool bUseSeparateTranslucency)
 {
-	float SortKey=0.f;
-	//sort based on distance to the view position, view rotation is not a factor
-	SortKey = (PrimitiveSceneInfo->Proxy->GetBounds().Origin - ViewInfo.ViewMatrices.ViewOrigin).Size();
-	// UE4_TODO: also account for DPG in the sort key.
+	const float SortKey = CalculateTranslucentSortKey(PrimitiveSceneInfo, ViewInfo);
 
 	const auto FeatureLevel = ViewInfo.GetFeatureLevel();
 
@@ -738,10 +760,7 @@ void FTranslucentPrimSet::AppendScenePrimitives(FSortedPrim* Normal, int32 NumNo
 
 void FTranslucentPrimSet::PlaceScenePrimitive(FPrimitiveSceneInfo* PrimitiveSceneInfo, const FViewInfo& ViewInfo, bool bUseNormalTranslucency, bool bUseSeparateTranslucency, void *NormalPlace, int32& NormalNum, void* SeparatePlace, int32& SeparateNum)
 {
-	float SortKey=0.f;
-	//sort based on distance to the view position, view rotation is not a factor
-	SortKey = (PrimitiveSceneInfo->Proxy->GetBounds().Origin - ViewInfo.ViewMatrices.ViewOrigin).Size();
-	// UE4_TODO: also account for DPG in the sort key.
+	const float SortKey = CalculateTranslucentSortKey(PrimitiveSceneInfo, ViewInfo);
 
 	const auto FeatureLevel = ViewInfo.GetFeatureLevel();
 
