added new unfinished/experimental translucency lighting mode: TLM_SurfacePerPixelLighting
 Higher quality than Surface but more expensive (loops through point lights with some basic culling, only inverse square, expensive, no shadow support yet)
 Requires 'r.ForwardLighting' to be 1

diff --git a/Engine/Shaders/BasePassPixelShader.usf b/Engine/Shaders/BasePassPixelShader.usf
--- a/Engine/Shaders/BasePassPixelShader.usf
+++ b/Engine/Shaders/BasePassPixelShader.usf
@@ -50,7 +50,7 @@ float4 PointSkyBentNormal;
 #define EDITOR_ALPHA2COVERAGE (EDITOR_PRIMITIVE_MATERIAL && FEATURE_LEVEL >= FEATURE_LEVEL_SM5 && !COMPILER_GLSL)
 
 
-#if TRANSLUCENCY_LIGHTING_SURFACE
+#if TRANSLUCENCY_LIGHTING_SURFACE || TRANSLUCENCY_LIGHTING_SURFACE_PERPIXEL
 
 #if FEATURE_LEVEL >= FEATURE_LEVEL_SM5
 	/** Prenormalized capture of the scene that's closest to the object being rendered, used for reflection environment on translucency. */
@@ -106,6 +106,10 @@ half3 GetImageBasedReflectionLighting(FMaterialPixelParameters MaterialParameter
 
 #endif
 
+#if TRANSLUCENCY_LIGHTING_SURFACE_PERPIXEL
+	#include "ForEachLight.usf"
+#endif
+
 /** Calculates lighting for translucency. */
 float3 GetTranslucencyLighting(FMaterialPixelParameters MaterialParameters, float3 DiffuseColor, half Roughness, half3 SpecularColor, float IndirectIrradiance, float MaterialAO)
 {
@@ -175,6 +179,10 @@ float3 GetTranslucencyLighting(FMaterialPixelParameters MaterialParameters, floa
 		// Determine lighting if no self shadowing is applied
 		InterpolatedLighting = DiffuseColor * VolumeLighting.rgb;
 
+	#elif TRANSLUCENCY_LIGHTING_SURFACE_PERPIXEL
+
+		InterpolatedLighting = AccumulateForwardLights(MaterialParameters, DiffuseColor, Roughness, SpecularColor);
+
 	#endif
 
 	#if (TRANSLUCENCY_LIGHTING_VOLUMETRIC_DIRECTIONAL || TRANSLUCENCY_LIGHTING_VOLUMETRIC_NONDIRECTIONAL || TRANSLUCENCY_LIGHTING_SURFACE) && TRANSLUCENT_SELF_SHADOWING
@@ -236,7 +244,7 @@ float3 GetTranslucencyLighting(FMaterialPixelParameters MaterialParameters, floa
 
 	#endif
 
-	#if TRANSLUCENCY_LIGHTING_SURFACE
+	#if TRANSLUCENCY_LIGHTING_SURFACE || TRANSLUCENCY_LIGHTING_SURFACE_PERPIXEL
 		InterpolatedLighting += GetImageBasedReflectionLighting(MaterialParameters, Roughness, SpecularColor, IndirectIrradiance);
 	#endif
 
@@ -447,7 +455,7 @@ void GetPrecomputedIndirectLightingAndSkyLight(
 	}
 	else
 	{
-	#if TRANSLUCENCY_LIGHTING_SURFACE && FEATURE_LEVEL >= FEATURE_LEVEL_SM5
+	#if (TRANSLUCENCY_LIGHTING_SURFACE || TRANSLUCENCY_LIGHTING_SURFACE_PERPIXEL) && FEATURE_LEVEL >= FEATURE_LEVEL_SM5
 		float DiffuseMip = ComputeReflectionCaptureMipFromRoughness(1);
 		float4 DiffuseIBL = TextureCubeArraySampleLevel(ReflectionCubemap, ReflectionCubemapSampler, MaterialParameters.WorldNormal, CubemapArrayIndex, DiffuseMip);
 		
diff --git a/Engine/Shaders/ForEachLight.usf b/Engine/Shaders/ForEachLight.usf
new file mode 100644
--- /dev/null
+++ b/Engine/Shaders/ForEachLight.usf
@@ -0,0 +1,68 @@
+// Copyright 1998-2015 Epic Games, Inc. All Rights Reserved.
+
+/*=============================================================================
+	ForEachLight.usf: Common code for dynamic forward lighting
+=============================================================================*/
+
+#ifndef __FOR_EACH_LIGHT_COMMON__
+#define __FOR_EACH_LIGHT_COMMON__
+
+#include "DeferredLightingCommon.usf"	// FSimpleDeferredLightData and GetSimpleDynamicLighting() 
+Buffer<uint> LightGrid;
+
+// greatly reduces shadow mapping artifacts
+float BiasedNDotL(float NDotLWithoutSaturate )
+{
+	return saturate(NDotLWithoutSaturate * 1.08f - 0.08f);
+}
+
+float3 SimplePointLightDiffuse( FScreenSpaceData ScreenSpaceData )
+{
+	return Diffuse_Lambert(ScreenSpaceData.GBuffer.DiffuseColor);
+}
+
+float3 LightingForOneLight(FMaterialPixelParameters Parameters, FSimpleDeferredLightData LightData, FScreenSpaceData ScreenSpaceData)
+{
+	// At the moment we only support very simple lighting (no shadows, phong, no area lights, basic attenuation).
+	// Later we can extend this function.
+	return GetSimpleDynamicLighting(Parameters.WorldPosition, Parameters.CameraVector, ScreenSpaceData, LightData);
+}
+
+// @param WorldNormal is not assumed to be normalized
+float3 AccumulateForwardLights(FMaterialPixelParameters Parameters, float3 DiffuseColor, float Roughness, float3 SpecularColor)
+{
+	float3 Ret = 0;
+
+	// Which tile we are?
+	int2 Tile = int2((Parameters.SVPosition.xy - View.ViewRectMin.xy) * ForwardLightData.InvTileSize);
+
+	// Get bitmask from CPU for this tile
+	uint culling = LightGrid[Tile.x + Tile.y * ForwardLightData.TileCountX];
+
+	FScreenSpaceData ScreenSpaceData = (FScreenSpaceData)0;
+	ScreenSpaceData.GBuffer.WorldNormal = normalize(Parameters.WorldNormal);
+	ScreenSpaceData.GBuffer.Roughness = Roughness;
+	ScreenSpaceData.GBuffer.SpecularColor = SpecularColor;
+	ScreenSpaceData.GBuffer.DiffuseColor = DiffuseColor;
+	ScreenSpaceData.AmbientOcclusion = 1;
+
+	// todo: support more than 32 lights
+
+	// loop through all lights specified by this bitmask
+	[allow_uav_condition] while(culling != 0)
+	{
+		uint LightIndex = firstbitlow(culling);
+
+		culling &= ~(1 << LightIndex);
+
+		FSimpleDeferredLightData LightData = (FSimpleDeferredLightData)0;
+		LightData.LightPositionAndInvRadius = ForwardLightData.LightPositionAndInvRadius[LightIndex];
+		LightData.LightColorAndFalloffExponent = ForwardLightData.LightColorAndFalloffExponent[LightIndex];
+		LightData.bInverseSquared = true;
+
+		Ret += LightingForOneLight(Parameters, LightData, ScreenSpaceData);
+	}
+	return Ret;
+}
+
+#endif // __FOR_EACH_LIGHT_COMMON__
\ No newline at end of file
diff --git a/Engine/Source/Runtime/Core/Private/Math/UnrealMath.cpp b/Engine/Source/Runtime/Core/Private/Math/UnrealMath.cpp
--- a/Engine/Source/Runtime/Core/Private/Math/UnrealMath.cpp
+++ b/Engine/Source/Runtime/Core/Private/Math/UnrealMath.cpp
@@ -1487,6 +1487,117 @@ bool FMath::SegmentPlaneIntersection(const FVector& StartPoint, const FVector& E
 	return false;
 }
 
+
+/**
+ * Compute the screen bounds of a point light along one axis.
+ * Based on http://www.gamasutra.com/features/20021011/lengyel_06.htm
+ * and http://sourceforge.net/mailarchive/message.php?msg_id=10501105
+ */
+static bool ComputeProjectedSphereShaft(
+	float LightX,
+	float LightZ,
+	float Radius,
+	const FMatrix& ProjMatrix,
+	const FVector& Axis,
+	float AxisSign,
+	int32& InOutMinX,
+	int32& InOutMaxX
+	)
+{
+	float ViewX = InOutMinX;
+	float ViewSizeX = InOutMaxX - InOutMinX;
+
+	// Vertical planes: T = <Nx, 0, Nz, 0>
+	float Discriminant = (FMath::Square(LightX) - FMath::Square(Radius) + FMath::Square(LightZ)) * FMath::Square(LightZ);
+	if(Discriminant >= 0)
+	{
+		float SqrtDiscriminant = FMath::Sqrt(Discriminant);
+		float InvLightSquare = 1.0f / (FMath::Square(LightX) + FMath::Square(LightZ));
+
+		float Nxa = (Radius * LightX - SqrtDiscriminant) * InvLightSquare;
+		float Nxb = (Radius * LightX + SqrtDiscriminant) * InvLightSquare;
+		float Nza = (Radius - Nxa * LightX) / LightZ;
+		float Nzb = (Radius - Nxb * LightX) / LightZ;
+		float Pza = LightZ - Radius * Nza;
+		float Pzb = LightZ - Radius * Nzb;
+
+		// Tangent a
+		if(Pza > 0)
+		{
+			float Pxa = -Pza * Nza / Nxa;
+			FVector4 P = ProjMatrix.TransformFVector4(FVector4(Axis.X * Pxa,Axis.Y * Pxa,Pza,1));
+			float X = (Dot3(P,Axis) / P.W + 1.0f * AxisSign) / 2.0f * AxisSign;
+			if(FMath::IsNegativeFloat(Nxa) ^ FMath::IsNegativeFloat(AxisSign))
+			{
+				InOutMaxX = FMath::Min<int64>(FMath::CeilToInt(ViewSizeX * X + ViewX),InOutMaxX);
+			}
+			else
+			{
+				InOutMinX = FMath::Max<int64>(FMath::FloorToInt(ViewSizeX * X + ViewX),InOutMinX);
+			}
+		}
+
+		// Tangent b
+		if(Pzb > 0)
+		{
+			float Pxb = -Pzb * Nzb / Nxb;
+			FVector4 P = ProjMatrix.TransformFVector4(FVector4(Axis.X * Pxb,Axis.Y * Pxb,Pzb,1));
+			float X = (Dot3(P,Axis) / P.W + 1.0f * AxisSign) / 2.0f * AxisSign;
+			if(FMath::IsNegativeFloat(Nxb) ^ FMath::IsNegativeFloat(AxisSign))
+			{
+				InOutMaxX = FMath::Min<int64>(FMath::CeilToInt(ViewSizeX * X + ViewX),InOutMaxX);
+			}
+			else
+			{
+				InOutMinX = FMath::Max<int64>(FMath::FloorToInt(ViewSizeX * X + ViewX),InOutMinX);
+			}
+		}
+	}
+
+	return InOutMinX <= InOutMaxX;
+}
+
+uint32 FMath::ComputeProjectedSphereScissorRect(FIntRect& InOutScissorRect, FVector SphereOrigin, float Radius, FVector ViewOrigin, const FMatrix& ViewMatrix, const FMatrix& ProjMatrix)
+{
+	// Calculate a scissor rectangle for the light's radius.
+	if((SphereOrigin - ViewOrigin).Size() > Radius)
+	{
+		FVector LightVector = ViewMatrix.TransformPosition(SphereOrigin);
+
+		if(!ComputeProjectedSphereShaft(
+			LightVector.X,
+			LightVector.Z,
+			Radius,
+			ProjMatrix,
+			FVector(+1,0,0),
+			+1,
+			InOutScissorRect.Min.X,
+			InOutScissorRect.Max.X))
+		{
+			return 0;
+		}
+
+		if(!ComputeProjectedSphereShaft(
+			LightVector.Y,
+			LightVector.Z,
+			Radius,
+			ProjMatrix,
+			FVector(0,+1,0),
+			-1,
+			InOutScissorRect.Min.Y,
+			InOutScissorRect.Max.Y))
+		{
+			return 0;
+		}
+
+		return 1;
+	}
+	else
+	{
+		return 2;
+	}
+}
+
 bool FMath::PlaneAABBIntersection(const FPlane& P, const FBox& AABB)
 {
 	// find diagonal most closely aligned with normal of plane
diff --git a/Engine/Source/Runtime/Core/Public/Math/UnrealMathUtility.h b/Engine/Source/Runtime/Core/Public/Math/UnrealMathUtility.h
--- a/Engine/Source/Runtime/Core/Public/Math/UnrealMathUtility.h
+++ b/Engine/Source/Runtime/Core/Public/Math/UnrealMathUtility.h
@@ -773,7 +773,11 @@ struct FMath : public FPlatformMath
 	static FVector LinePlaneIntersection( const FVector &Point1, const FVector &Point2, const FVector &PlaneOrigin, const FVector &PlaneNormal);
 	static FVector LinePlaneIntersection( const FVector &Point1, const FVector &Point2, const FPlane  &Plane);
 
-	/**
+	// @parma InOutScissorRect should be set to View.ViewRect before the call
+	// @return 0: light is not visible, 1:use scissor rect, 2: no scissor rect needed
+	static CORE_API uint32 ComputeProjectedSphereScissorRect(struct FIntRect& InOutScissorRect, FVector SphereOrigin, float Radius, FVector ViewOrigin, const FMatrix& ViewMatrix, const FMatrix& ProjMatrix);
+
+	/** 
 	 * Determine if a plane and an AABB intersect
 	 * @param P - the plane to test
 	 * @param AABB - the axis aligned bounding box to test
diff --git a/Engine/Source/Runtime/Engine/Classes/Engine/EngineTypes.h b/Engine/Source/Runtime/Engine/Classes/Engine/EngineTypes.h
--- a/Engine/Source/Runtime/Engine/Classes/Engine/EngineTypes.h
+++ b/Engine/Source/Runtime/Engine/Classes/Engine/EngineTypes.h
@@ -107,19 +107,27 @@ enum ETranslucencyLightingMode
 	 * Lighting will be calculated for a volume, without directionality.  Use this on particle effects like smoke and dust.
 	 * This is the cheapest lighting method, however the material normal is not taken into account.
 	 */
-	TLM_VolumetricNonDirectional,
+	TLM_VolumetricNonDirectional UMETA(DisplayName="Volumetric Non Directional"),
 
 	 /** 
 	 * Lighting will be calculated for a volume, with directionality so that the normal of the material is taken into account. 
 	 * Note that the default particle tangent space is facing the camera, so enable bGenerateSphericalParticleNormals to get a more useful tangent space.
 	 */
-	TLM_VolumetricDirectional,
+	TLM_VolumetricDirectional UMETA(DisplayName="Volumetric Directional"),
 
 	/** 
-	 * Lighting will be calculated for a surface.
-	 * Use this on translucent surfaces like glass and water.
+	 * Lighting will be calculated for a surface. The light in accumulated in a volume so the result is blurry
+	 * (fixed resolution), limited distance but the per pixel cost is very low. Use this on translucent surfaces like glass and water.
 	 */
-	TLM_Surface,
+	TLM_Surface UMETA(DisplayName="Surface TranslucencyVolume"),
+
+	/** 
+	 * Lighting will be calculated for a surface. Use this on translucent surfaces like glass and water.
+	 * Higher quality than Surface but more expensive (loops through point lights with some basic culling, only inverse square, expensive, no shadow support yet)
+	 * Requires 'r.ForwardLighting' to be 1
+	 */
+	TLM_SurfacePerPixelLighting UMETA(DisplayName="Surface PerPixel (experimental, limited features)"),
+
 	TLM_MAX,
 };
 
diff --git a/Engine/Source/Runtime/Engine/Private/Materials/MaterialShared.cpp b/Engine/Source/Runtime/Engine/Private/Materials/MaterialShared.cpp
--- a/Engine/Source/Runtime/Engine/Private/Materials/MaterialShared.cpp
+++ b/Engine/Source/Runtime/Engine/Private/Materials/MaterialShared.cpp
@@ -1204,6 +1204,8 @@ void FMaterial::SetupMaterialEnvironment(
 		case TLM_VolumetricNonDirectional: OutEnvironment.SetDefine(TEXT("TRANSLUCENCY_LIGHTING_VOLUMETRIC_NONDIRECTIONAL"),TEXT("1")); break;
 		case TLM_VolumetricDirectional: OutEnvironment.SetDefine(TEXT("TRANSLUCENCY_LIGHTING_VOLUMETRIC_DIRECTIONAL"),TEXT("1")); break;
 		case TLM_Surface: OutEnvironment.SetDefine(TEXT("TRANSLUCENCY_LIGHTING_SURFACE"),TEXT("1")); break;
+		case TLM_SurfacePerPixelLighting: OutEnvironment.SetDefine(TEXT("TRANSLUCENCY_LIGHTING_SURFACE_PERPIXEL"),TEXT("1")); break;
+
 		default: 
 			UE_LOG(LogMaterial, Warning, TEXT("Unknown lighting mode: %u"),(int32)GetTranslucencyLightingMode());
 			OutEnvironment.SetDefine(TEXT("TRANSLUCENCY_LIGHTING_VOLUMETRIC_NONDIRECTIONAL"),TEXT("1")); break;
diff --git a/Engine/Source/Runtime/Engine/Private/PointLightSceneProxy.h b/Engine/Source/Runtime/Engine/Private/PointLightSceneProxy.h
--- a/Engine/Source/Runtime/Engine/Private/PointLightSceneProxy.h
+++ b/Engine/Source/Runtime/Engine/Private/PointLightSceneProxy.h
@@ -8,73 +8,6 @@
 #define __PointLightSceneProxy_H__
 #include "Components/PointLightComponent.h"
 
-/**
- * Compute the screen bounds of a point light along one axis.
- * Based on http://www.gamasutra.com/features/20021011/lengyel_06.htm
- * and http://sourceforge.net/mailarchive/message.php?msg_id=10501105
- */
-static bool GetPointLightBounds(
-	float LightX,
-	float LightZ,
-	float Radius,
-	const FVector& Axis,
-	float AxisSign,
-	const FSceneView& View,
-	float ViewX,
-	float ViewSizeX,
-	int32& OutMinX,
-	int32& OutMaxX
-	)
-{
-	// Vertical planes: T = <Nx, 0, Nz, 0>
-	float Discriminant = (FMath::Square(LightX) - FMath::Square(Radius) + FMath::Square(LightZ)) * FMath::Square(LightZ);
-	if(Discriminant >= 0)
-	{
-		float SqrtDiscriminant = FMath::Sqrt(Discriminant);
-		float InvLightSquare = 1.0f / (FMath::Square(LightX) + FMath::Square(LightZ));
-
-		float Nxa = (Radius * LightX - SqrtDiscriminant) * InvLightSquare;
-		float Nxb = (Radius * LightX + SqrtDiscriminant) * InvLightSquare;
-		float Nza = (Radius - Nxa * LightX) / LightZ;
-		float Nzb = (Radius - Nxb * LightX) / LightZ;
-		float Pza = LightZ - Radius * Nza;
-		float Pzb = LightZ - Radius * Nzb;
-
-		// Tangent a
-		if(Pza > 0)
-		{
-			float Pxa = -Pza * Nza / Nxa;
-			FVector4 P = View.ViewMatrices.ProjMatrix.TransformFVector4(FVector4(Axis.X * Pxa,Axis.Y * Pxa,Pza,1));
-			float X = (Dot3(P,Axis) / P.W + 1.0f * AxisSign) / 2.0f * AxisSign;
-			if(FMath::IsNegativeFloat(Nxa) ^ FMath::IsNegativeFloat(AxisSign))
-			{
-				OutMaxX = FMath::Min<int64>(FMath::CeilToInt(ViewSizeX * X + ViewX),OutMaxX);
-			}
-			else
-			{
-				OutMinX = FMath::Max<int64>(FMath::FloorToInt(ViewSizeX * X + ViewX),OutMinX);
-			}
-		}
-
-		// Tangent b
-		if(Pzb > 0)
-		{
-			float Pxb = -Pzb * Nzb / Nxb;
-			FVector4 P = View.ViewMatrices.ProjMatrix.TransformFVector4(FVector4(Axis.X * Pxb,Axis.Y * Pxb,Pzb,1));
-			float X = (Dot3(P,Axis) / P.W + 1.0f * AxisSign) / 2.0f * AxisSign;
-			if(FMath::IsNegativeFloat(Nxb) ^ FMath::IsNegativeFloat(AxisSign))
-			{
-				OutMaxX = FMath::Min<int64>(FMath::CeilToInt(ViewSizeX * X + ViewX),OutMaxX);
-			}
-			else
-			{
-				OutMinX = FMath::Max<int64>(FMath::FloorToInt(ViewSizeX * X + ViewX),OutMinX);
-			}
-		}
-	}
-
-	return OutMinX <= OutMaxX;
-}
 
 /** The parts of the point light scene info that aren't dependent on the light policy type. */
 class FPointLightSceneProxyBase : public FLightSceneProxy
@@ -152,50 +85,7 @@ public:
 	virtual bool GetScissorRect(FIntRect& ScissorRect, const FSceneView& View) const override
 	{
 		ScissorRect = View.ViewRect;
-
-		// Calculate a scissor rectangle for the light's radius.
-		if((GetLightToWorld().GetOrigin() - View.ViewMatrices.ViewOrigin).Size() > Radius)
-		{
-			FVector LightVector = View.ViewMatrices.ViewMatrix.TransformPosition(GetLightToWorld().GetOrigin());
-
-			if(!GetPointLightBounds(
-				LightVector.X,
-				LightVector.Z,
-				Radius,
-				FVector(+1,0,0),
-				+1,
-				View,
-				View.ViewRect.Min.X,
-				View.ViewRect.Width(),
-				ScissorRect.Min.X,
-				ScissorRect.Max.X))
-			{
-				return false;
-			}
-
-			int32 ScissorMinY = View.ViewRect.Min.Y;
-			int32 ScissorMaxY = View.ViewRect.Max.Y;
-			if(!GetPointLightBounds(
-				LightVector.Y,
-				LightVector.Z,
-				Radius,
-				FVector(0,+1,0),
-				-1,
-				View,
-				View.ViewRect.Min.Y,
-				View.ViewRect.Height(),
-				ScissorRect.Min.Y,
-				ScissorRect.Max.Y))
-			{
-				return false;
-			}
-
-			return true;
-		}
-		else
-		{
-			return false;
-		}
+		return FMath::ComputeProjectedSphereScissorRect(ScissorRect, GetLightToWorld().GetOrigin(), Radius, View.ViewMatrices.ViewOrigin, View.ViewMatrices.ViewMatrix, View.ViewMatrices.ProjMatrix) == 1;
 	}
 
 	virtual void SetScissorRect(FRHICommandList& RHICmdList, const FSceneView& View) const override
diff --git a/Engine/Source/Runtime/Engine/Private/SceneView.cpp b/Engine/Source/Runtime/Engine/Private/SceneView.cpp
--- a/Engine/Source/Runtime/Engine/Private/SceneView.cpp
+++ b/Engine/Source/Runtime/Engine/Private/SceneView.cpp
@@ -20,6 +20,7 @@ DEFINE_LOG_CATEGORY(LogBufferVisualization);
 
 IMPLEMENT_UNIFORM_BUFFER_STRUCT(FPrimitiveUniformShaderParameters,TEXT("Primitive"));
 IMPLEMENT_UNIFORM_BUFFER_STRUCT(FViewUniformShaderParameters,TEXT("View"));
+IMPLEMENT_UNIFORM_BUFFER_STRUCT(FForwardLightData,TEXT("ForwardLightData"));
 
 #if !(UE_BUILD_SHIPPING || UE_BUILD_TEST)
 static TAutoConsoleVariable<float> CVarSSRMaxRoughness(
diff --git a/Engine/Source/Runtime/Engine/Public/SceneView.h b/Engine/Source/Runtime/Engine/Public/SceneView.h
--- a/Engine/Source/Runtime/Engine/Public/SceneView.h
+++ b/Engine/Source/Runtime/Engine/Public/SceneView.h
@@ -10,6 +10,7 @@
 #include "SceneInterface.h"
 #include "SceneTypes.h"
 #include "ShaderParameters.h"
+#include "RendererInterface.h"
 
 class FSceneViewStateInterface;
 class FViewUniformShaderParameters;
@@ -199,6 +200,28 @@ struct FViewMatrices
 	}
 };
 
+#include "ShaderParameters.h"
+
+/** 
+ * Current limit of the dynamic branching forward lighting code path (see r.ForwardLighting)
+ */
+static const int32 GMaxNumForwardLights = 32;
+
+/** 
+ * Data used to pass light properties and some other parameters down to the dynamic forward lighting code path.
+ * todo: Fix static size (to be more efficient and to not limit the feature)
+ */
+BEGIN_UNIFORM_BUFFER_STRUCT(FForwardLightData,ENGINE_API)
+	DECLARE_UNIFORM_BUFFER_STRUCT_MEMBER(uint32,LightCount)
+	DECLARE_UNIFORM_BUFFER_STRUCT_MEMBER(uint32,TileSize)
+	DECLARE_UNIFORM_BUFFER_STRUCT_MEMBER(uint32,TileCountX)
+	DECLARE_UNIFORM_BUFFER_STRUCT_MEMBER(float,InvTileSize)
+	DECLARE_UNIFORM_BUFFER_STRUCT_MEMBER_ARRAY(FVector4,LightPositionAndInvRadius,[GMaxNumForwardLights])
+	DECLARE_UNIFORM_BUFFER_STRUCT_MEMBER_ARRAY(FVector4,LightColorAndFalloffExponent,[GMaxNumForwardLights])
+END_UNIFORM_BUFFER_STRUCT(FForwardLightData)
+
+//////////////////////////////////////////////////////////////////////////
+
 //////////////////////////////////////////////////////////////////////////
 
 static const int MAX_FORWARD_SHADOWCASCADES = 2;
@@ -323,6 +346,9 @@ public:
 	/** The uniform buffer for the view's parameters.  This is only initialized in the rendering thread's copies of the FSceneView. */
 	TUniformBufferRef<FViewUniformShaderParameters> UniformBuffer;
 
+	/** uniform buffer with the lights for forward lighting/shading */
+	TUniformBufferRef<FForwardLightData> ForwardLightData;
+
 	/** The actor which is being viewed from. */
 	const AActor* ViewActor;
 
diff --git a/Engine/Source/Runtime/Renderer/Private/BasePassRendering.h b/Engine/Source/Runtime/Renderer/Private/BasePassRendering.h
--- a/Engine/Source/Runtime/Renderer/Private/BasePassRendering.h
+++ b/Engine/Source/Runtime/Renderer/Private/BasePassRendering.h
@@ -8,6 +8,7 @@
 
 #include "LightMapRendering.h"
 #include "ShaderBaseClasses.h"
+#include "LightGrid.h"
 
 /** Whether to allow the indirect lighting cache to be applied to dynamic objects. */
 extern int32 GIndirectLightingCache;
@@ -321,6 +322,7 @@ public:
 		MSAASampleCount.Bind(Initializer.ParameterMap,TEXT("MSAASampleCount"));
 		FilteredSceneDepthTexture.Bind(Initializer.ParameterMap,TEXT("FilteredSceneDepthTexture"));
 		FilteredSceneDepthTextureSampler.Bind(Initializer.ParameterMap,TEXT("FilteredSceneDepthTextureSampler"));
+		LightGrid.Bind(Initializer.ParameterMap,TEXT("LightGrid"));
 	}
 	TBasePassPixelShaderBaseType() {}
 
@@ -342,6 +344,22 @@ public:
 			if (IsTranslucentBlendMode(BlendMode))
 			{
 				TranslucentLightingParameters.Set(RHICmdList, this, View);
+
+				// Experimental dynamic forward lighting for translucency. Can be the base for opaque forward lighting which will allow more lighting models or rendering without a GBuffer
+				if(BlendMode == BLEND_Translucent && MaterialResource.GetTranslucencyLightingMode() == TLM_SurfacePerPixelLighting)
+				{
+					check(GetUniformBufferParameter<FForwardLightData>().IsInitialized());
+
+					if(GetUniformBufferParameter<FForwardLightData>().IsBound())
+					{
+						SetUniformBufferParameter(RHICmdList, ShaderRHI,GetUniformBufferParameter<FForwardLightData>(),View->ForwardLightData);
+					}
+
+					if(LightGrid.IsBound())
+					{
+						RHICmdList.SetShaderResourceViewParameter(ShaderRHI, LightGrid.GetBaseIndex(), GLightGridVertexBuffer.VertexBufferSRV);
+					}
+				}
 			}
 		}
 
@@ -402,6 +420,7 @@ public:
 		Ar << MSAASampleCount;
 		Ar << FilteredSceneDepthTexture;
 		Ar << FilteredSceneDepthTextureSampler;
+		Ar << LightGrid;
 		return bShaderHasOutdatedParameters;
 	}
 
@@ -413,6 +432,7 @@ private:
 	/** Parameter for reading filtered depth values */
 	FShaderResourceParameter FilteredSceneDepthTexture; 
 	FShaderResourceParameter FilteredSceneDepthTextureSampler; 
+	FShaderResourceParameter LightGrid; 
 };
 
 /** The concrete base pass pixel shader type. */
diff --git a/Engine/Source/Runtime/Renderer/Private/LightGrid.h b/Engine/Source/Runtime/Renderer/Private/LightGrid.h
new file mode 100644
--- /dev/null
+++ b/Engine/Source/Runtime/Renderer/Private/LightGrid.h
@@ -0,0 +1,76 @@
+// Copyright 1998-2014 Epic Games, Inc. All Rights Reserved.
+
+/*=============================================================================
+	LightGrid.h: screen space light culling
+=============================================================================*/
+
+#ifndef __LIGHTGRID_H__
+#define __LIGHTGRID_H__
+
+#pragma once
+
+/**
+ * Vertex buffer used to hold light information for screen space light culling
+ */
+class FLightGridVertexBuffer : public FVertexBuffer
+{
+public:
+	const static uint32 MaxSize = 512 * 1024;
+	TArray<uint32> CPUData;
+
+	/**
+	 * Initializes the vertex buffer from the input
+	 */
+	void UpdateGPUFromCPUData()
+	{
+		check( IsInRenderingThread() );
+		InitResource();
+		if(CPUData.Num())
+		{
+			check( IsInRenderingThread() );
+
+			const uint32 Size = CPUData.Num();
+
+			check(Size <= MaxSize);
+
+			const uint32 BufferSize = Size * sizeof(uint32);
+
+			uint32* RESTRICT Src = &CPUData[0];
+			uint32* RESTRICT Dest = (uint32*)RHILockVertexBuffer( VertexBufferRHI, 0, BufferSize, RLM_WriteOnly );
+
+			for(uint32 i = 0; i < Size; ++i)
+			{
+				*Dest++ = *Src++;
+			}
+
+			RHIUnlockVertexBuffer(VertexBufferRHI);
+		}
+	}
+
+	/** Shader resource view of the vertex buffer. */
+	FShaderResourceViewRHIRef VertexBufferSRV;
+
+private:
+
+	/** Release RHI resources. */
+	virtual void ReleaseRHI() override
+	{
+		VertexBufferSRV.SafeRelease();
+		FVertexBuffer::ReleaseRHI();
+	}
+
+	/** Initialize RHI resources. */
+	virtual void InitRHI() override
+	{
+		const int32 BufferStride = sizeof(uint32);
+		const int32 BufferSize = MaxSize * BufferStride;
+		uint32 Flags = BUF_Volatile | /*BUF_KeepCPUAccessible | */BUF_ShaderResource;
+		FRHIResourceCreateInfo CreateInfo;
+		VertexBufferRHI = RHICreateVertexBuffer(BufferSize, Flags, CreateInfo);
+		VertexBufferSRV = RHICreateShaderResourceView(VertexBufferRHI, BufferStride, PF_R32_UINT);
+	}
+};
+
+extern TGlobalResource<FLightGridVertexBuffer> GLightGridVertexBuffer;
+
+#endif // __LIGHTGRID_H__
\ No newline at end of file
diff --git a/Engine/Source/Runtime/Renderer/Private/SceneRendering.cpp b/Engine/Source/Runtime/Renderer/Private/SceneRendering.cpp
--- a/Engine/Source/Runtime/Renderer/Private/SceneRendering.cpp
+++ b/Engine/Source/Runtime/Renderer/Private/SceneRendering.cpp
@@ -16,6 +16,7 @@
 #include "SceneViewExtension.h"
 #include "PostProcessBusyWait.h"
 #include "SceneUtils.h"
+#include "LightGrid.h"
 
 /*-----------------------------------------------------------------------------
 	Globals
@@ -24,6 +25,8 @@
 extern ENGINE_API FLightMap2D* GDebugSelectedLightmap;
 extern ENGINE_API UPrimitiveComponent* GDebugSelectedComponent;
 
+TGlobalResource<FLightGridVertexBuffer> GLightGridVertexBuffer;
+
 /**
  * Console variable controlling whether or not occlusion queries are allowed.
  */
@@ -40,6 +43,16 @@ static TAutoConsoleVariable<float> CVarDemosaicVposOffset(
 	TEXT("This offset is added to the rasterized position used for demosaic in the ES2 tonemapping shader. It exists to workaround driver bugs on some Android devices that have a half-pixel offset."),
 	ECVF_RenderThreadSafe);
 
+static TAutoConsoleVariable<int32> CVarForwardLighting(
+	TEXT("r.ForwardLighting"),
+	0,
+	TEXT("Experimental dynamic forward lighting for translucency. Can be the base for opaque forward lighting which will allow\n")
+	TEXT("more lighting models or rendering without a GBuffer.\n")
+	TEXT("The current implementation is limited to 32 lights, coarse 2d culling (on CPU), no shadows and simple shading (no area lights, phong).\n")
+	TEXT("Enabled with the new TranslucencyLightingMode SurfacePerPixelLighting\n")
+	TEXT("0: off (default)\n")
+	TEXT("1: on (some CPU setup cost on GPU and extra per pixel cost)"),
+	ECVF_RenderThreadSafe);
 
 #if !(UE_BUILD_SHIPPING || UE_BUILD_TEST)
 static TAutoConsoleVariable<float> CVarGeneralPurposeTweak(
@@ -619,6 +632,108 @@ TUniformBufferRef<FViewUniformShaderParameters> FViewInfo::CreateUniformBuffer(
 	return TUniformBufferRef<FViewUniformShaderParameters>::CreateUniformBufferImmediate(ViewUniformShaderParameters, UniformBuffer_SingleFrame);
 }
 
+void FViewInfo::CreateForwardLightDataUniformBuffer(FForwardLightData &OutForwardLightData) const
+{
+	uint32 LightIndex = 0;
+	// todo: we should do this only once before rendering, not late after GBuffer pass
+
+	FScene* Scene = (FScene*)(Family->Scene);
+
+	if(Scene)
+	{
+		// Build a list of visible lights.
+		for (TSparseArray<FLightSceneInfoCompact>::TConstIterator LightIt(Scene->Lights); LightIt; ++LightIt)
+		{
+			const FLightSceneInfoCompact& LightSceneInfoCompact = *LightIt;
+			const FLightSceneInfo* const LightSceneInfo = LightSceneInfoCompact.LightSceneInfo;
+
+			if(!LightSceneInfoCompact.Color.IsAlmostBlack()
+				// Only render lights with dynamic lighting or unbuilt static lights
+				&& (!LightSceneInfo->Proxy->HasStaticLighting() || !LightSceneInfo->bPrecomputedLightingIsValid)
+				// Reflection override skips direct specular because it tends to be blindingly bright with a perfectly smooth surface
+				&& !Family->EngineShowFlags.ReflectionOverride
+				&& LightSceneInfoCompact.LightType == LightType_Point)		// for now we only handle point lights
+			{
+				// Check if the light is visible in this view.
+				if(LightSceneInfo->ShouldRenderLight(*this))
+				{
+					FVector4 BoundingSphereVector = *(FVector4*)&LightSceneInfoCompact.BoundingSphereVector;
+
+					float InvRadius = 1.0f / BoundingSphereVector.W;
+
+					OutForwardLightData.LightPositionAndInvRadius[LightIndex] = FVector4(FVector(BoundingSphereVector), InvRadius);
+					OutForwardLightData.LightColorAndFalloffExponent[LightIndex] = LightSceneInfoCompact.Color;
+					++LightIndex;
+
+					if(LightIndex >= GMaxNumForwardLights)
+					{
+						// we cannot handle more lights this way
+						break;
+					}
+				}
+			}
+		}
+	}
+
+	OutForwardLightData.LightCount = LightIndex;
+}
+
+void FViewInfo::CreateLightGrid()
+{
+	// This constant affects performance, power of two makes sense, too small costs more CPU and upload data, too large culls less efficiently.
+	const int32 TileSize = 16;
+
+	uint32 TileCountX = FMath::DivideAndRoundUp(ViewRect.Width(), TileSize);
+	uint32 TileCountY = FMath::DivideAndRoundUp(ViewRect.Height(), TileSize);
+
+	//@todo - creating a new uniform buffer is expensive, only do this when the vertex factory needs an accurate view matrix (particle sprites)
+	FForwardLightData LocalForwardLightData;
+
+	CreateForwardLightDataUniformBuffer(LocalForwardLightData);
+
+	LocalForwardLightData.TileSize = TileSize;
+	LocalForwardLightData.TileCountX = TileCountX;
+	LocalForwardLightData.InvTileSize = 1.0f / TileSize;
+
+	ForwardLightData = TUniformBufferRef<FForwardLightData>::CreateUniformBufferImmediate(LocalForwardLightData, UniformBuffer_SingleFrame);
+	// todo: temp data should be released when possible (after lighting)
+
+	{
+		GLightGridVertexBuffer.CPUData.SetNumUninitialized(TileCountX * TileCountY);
+
+		FMemory::Memset(&GLightGridVertexBuffer.CPUData[0], 0, TileCountX * TileCountY * sizeof(uint32));
+
+		for(uint32 i = 0, Count = LocalForwardLightData.LightCount; i < Count; ++i)
+		{
+			if(Count > 31)
+			{
+				// todo: currently we don't support more
+				break;
+			}
+
+			const FVector4 PosAndRadius = LocalForwardLightData.LightPositionAndInvRadius[i];
+
+			FIntRect LocalPixelRect = ViewRect;
+			if(FMath::ComputeProjectedSphereScissorRect(LocalPixelRect, FVector(PosAndRadius), 1.0f / PosAndRadius.W, ViewMatrices.ViewOrigin, ViewMatrices.ViewMatrix, ViewMatrices.ProjMatrix) > 0)
+			{
+				FIntRect LocalTileRect = FIntRect::DivideAndRoundUp(LocalPixelRect, TileSize);
+
+				for(int32 y = LocalTileRect.Min.Y; y < LocalTileRect.Max.Y; ++y)
+				{
+					for(int32 x = LocalTileRect.Min.X; x < LocalTileRect.Max.X; ++x)
+					{
+						uint32& Tile = GLightGridVertexBuffer.CPUData[x + y * TileCountX];
+
+						Tile |= 1 << i;
+					}
+				}
+			}
+		}
+
+		GLightGridVertexBuffer.UpdateGPUFromCPUData();
+	}
+}
+
 void FViewInfo::InitRHIResources(const TArray<FProjectedShadowInfo*, SceneRenderingAllocator>* DirectionalLightShadowInfo)
 {
 	FBox VolumeBounds[TVC_MAX];
@@ -644,6 +759,11 @@ void FViewInfo::InitRHIResources(const TArray<FProjectedShadowInfo*, SceneRender
 	{
 		DynamicResources[ResourceIndex]->InitPrimitiveResource();
 	}
+
+	if(CVarForwardLighting.GetValueOnRenderThread())
+	{
+		CreateLightGrid();
+	}
 }
 
 IPooledRenderTarget* FViewInfo::GetEyeAdaptation() const
diff --git a/Engine/Source/Runtime/Renderer/Private/SceneRendering.h b/Engine/Source/Runtime/Renderer/Private/SceneRendering.h
--- a/Engine/Source/Runtime/Renderer/Private/SceneRendering.h
+++ b/Engine/Source/Runtime/Renderer/Private/SceneRendering.h
@@ -504,6 +504,9 @@ public:
 	/** Gets the eye adaptation render target for this view. */
 	IPooledRenderTarget* GetEyeAdaptation() const;
 
+	/** Create acceleration data structure and information to do forward lighting with dynamic branching. */
+	void CreateLightGrid();
+
 private:
 
 	/** Initialization that is common to the constructors. */
@@ -514,6 +517,9 @@ private:
 
 	/** Sets the sky SH irradiance map coefficients. */
 	void SetupSkyIrradianceEnvironmentMapConstants(FVector4* OutSkyIrradianceEnvironmentMap) const;
+
+	/** All light sources available for forward shading. Can be indexed in the shader.*/
+	void CreateForwardLightDataUniformBuffer(FForwardLightData& Out) const;
 };
 
 
