Fixed HZB occlusion bug where large vertical meshes would be incorrectly culled.

diff --git a/Engine/Shaders/HZBOcclusion.usf b/Engine/Shaders/HZBOcclusion.usf
--- a/Engine/Shaders/HZBOcclusion.usf
+++ b/Engine/Shaders/HZBOcclusion.usf
@@ -11,10 +11,10 @@ void HZBBuildPS(float2 InUV : TEXCOORD0, out float4 OutColor : SV_Target0)
 {
 	float2 UV[4];
 
-	UV[0] = InUV + float2(-0.5f, -0.5f) * InvSize;
-	UV[1] = InUV + float2( 0.5f, -0.5f) * InvSize;
-	UV[2] = InUV + float2(-0.5f,  0.5f) * InvSize;
-	UV[3] = InUV + float2( 0.5f,  0.5f) * InvSize;
+	UV[0] = InUV + float2(-0.25f, -0.25f) * InvSize;
+	UV[1] = InUV + float2( 0.25f, -0.25f) * InvSize;
+	UV[2] = InUV + float2(-0.25f,  0.25f) * InvSize;
+	UV[3] = InUV + float2( 0.25f,  0.25f) * InvSize;
 
 #if STAGE == 0
 	// Inverted Z buffer
@@ -99,7 +99,7 @@ void HZBTestPS(float2 InUV : TEXCOORD0, out float4 OutColor : SV_Target0)
 #endif
 
 	// FIXME assumes DX
-	float4 Rect = saturate( float4( RectMin.xy, RectMax.xy ) * float2( 0.5, -0.5 ).xyxy + 0.5 );
+	float4 Rect = saturate( float4( RectMin.xy, RectMax.xy ) * float2( 0.5, -0.5 ).xyxy + 0.5 ).xwzy;
 	float4 RectPixels = Rect * HZBSize.xyxy;
 	float2 RectSize = ( RectPixels.zw - RectPixels.xy ) * 0.5;	// 0.5 for 4x4
 	float Level = ceil( log2( max( RectSize.x, RectSize.y ) ) );
diff --git a/Engine/Source/Runtime/Renderer/Private/SceneOcclusion.cpp b/Engine/Source/Runtime/Renderer/Private/SceneOcclusion.cpp
--- a/Engine/Source/Runtime/Renderer/Private/SceneOcclusion.cpp
+++ b/Engine/Source/Runtime/Renderer/Private/SceneOcclusion.cpp
@@ -881,7 +881,6 @@ void BuildHZB(FRHICommandListImmediate& RHICmdList, const FViewInfo& View)
 	
 	FSceneRenderTargetItem& HZBRenderTarget = ViewState->HZB.Texture->GetRenderTargetItem();
 	
-	
 	RHICmdList.SetBlendState(TStaticBlendState<>::GetRHI());
 	RHICmdList.SetRasterizerState(TStaticRasterizerState<>::GetRHI());
 	RHICmdList.SetDepthStencilState(TStaticDepthStencilState< false, CF_Always >::GetRHI());
@@ -900,7 +899,7 @@ void BuildHZB(FRHICommandListImmediate& RHICmdList, const FViewInfo& View)
 		SetGlobalBoundShaderState(RHICmdList, View.GetFeatureLevel(), BoundShaderState, GFilterVertexDeclaration.VertexDeclarationRHI, *VertexShader, *PixelShader);
 
 		// Imperfect sampling, doesn't matter too much
-		PixelShader->SetParameters(RHICmdList, View, GSceneRenderTargets.GetBufferSizeXY() );
+		PixelShader->SetParameters( RHICmdList, View, HZBSize );
 
 		RHICmdList.SetViewport(0, 0, 0.0f, HZBSize.X, HZBSize.Y, 1.0f);
 
@@ -914,6 +913,8 @@ void BuildHZB(FRHICommandListImmediate& RHICmdList, const FViewInfo& View)
 			GSceneRenderTargets.GetBufferSizeXY(),
 			*VertexShader,
 			EDRF_UseTriangleOptimization);
+
+		RHICmdList.CopyToResolveTarget(HZBRenderTarget.TargetableTexture, HZBRenderTarget.ShaderResourceTexture, false, FResolveParams(FResolveRect(), CubeFace_PosX, 0));
 	}
 
 	FIntPoint SrcSize = HZBSize;
@@ -933,7 +934,7 @@ void BuildHZB(FRHICommandListImmediate& RHICmdList, const FViewInfo& View)
 		
 		SetGlobalBoundShaderState(RHICmdList, View.GetFeatureLevel(), BoundShaderState, GFilterVertexDeclaration.VertexDeclarationRHI, *VertexShader, *PixelShader);
 
-		PixelShader->SetParameters(RHICmdList, View, SrcSize, ViewState->HZB.MipSRVs[ MipIndex - 1 ] );
+		PixelShader->SetParameters( RHICmdList, View, DstSize, ViewState->HZB.MipSRVs[ MipIndex - 1 ] );
 
 		RHICmdList.SetViewport(0, 0, 0.0f, DstSize.X, DstSize.Y, 1.0f);
 
