fixed DBuffer decals in stereo rendering

diff --git a/Engine/Source/Runtime/Renderer/Private/CompositionLighting/PostProcessDeferredDecals.cpp b/Engine/Source/Runtime/Renderer/Private/CompositionLighting/PostProcessDeferredDecals.cpp
--- a/Engine/Source/Runtime/Renderer/Private/CompositionLighting/PostProcessDeferredDecals.cpp
+++ b/Engine/Source/Runtime/Renderer/Private/CompositionLighting/PostProcessDeferredDecals.cpp
@@ -746,9 +746,14 @@ void FRCPassPostProcessDeferredDecals::Process(FRenderingCompositePassContext& C
 			GRenderTargetPool.FindFreeElement(Desc, GSceneRenderTargets.DBufferC, TEXT("DBufferC"));
 		}
 
+		// we assume views are non overlapping, then we need to clear only once in the beginning, otherwise we would need to set scissor rects
+		// and don't get FastClear any more.
+		bool bFirstView = Context.View.Family->Views[0] == &Context.View;
 
-		SCOPED_DRAW_EVENT(RHICmdList, DBufferClear);
+		if(bFirstView)
 		{
+			SCOPED_DRAW_EVENT(RHICmdList, DBufferClear);
+
 			// could be optimized
 			SetRenderTarget(RHICmdList, GSceneRenderTargets.DBufferA->GetRenderTargetItem().TargetableTexture, FTextureRHIParamRef());
 			RHICmdList.Clear(true, FLinearColor(0, 0, 0, 1), false, 0, false, 0, FIntRect());
