Fixed SubsurfaceScatteringProfile, recent full resolution mix caused discoloration

diff --git a/Engine/Shaders/PostProcessSubsurface.usf b/Engine/Shaders/PostProcessSubsurface.usf
--- a/Engine/Shaders/PostProcessSubsurface.usf
+++ b/Engine/Shaders/PostProcessSubsurface.usf
@@ -11,9 +11,8 @@
 // for VisualizeSSS
 #include "MiniFontCommon.usf"
 
-/** RGBA8 linear texture containing random normals */
+/** one profile per line, sample set left to right: high/med/low */
 Texture2D SSProfilesTexture;
-SamplerState SSProfilesTextureSampler;
 // x:Radius*DistanceToProjectionWindow/KernelSize*0.5, y:DistanceToProjectionWindow, zw:xy scale correction for BufferSize!=ViewportSize
 float4 SSSParams;
 
@@ -427,7 +426,7 @@ void SubsurfaceRecombinePS(float4 UVAndScreenPos : TEXCOORD0, out float4 OutColo
 		SSSColor = SSSColorWithAlpha.rgb / max(SSSColorWithAlpha.a, 0.00001f);
 	}
 
-	// 0: don't fade in subsurface half res,  1:fully fade in half res
+	// 0: no SSS,  1: full SSS
 	float BlendFactor = (ScreenSpaceData.GBuffer.ShadingModelID == SHADINGMODELID_SUBSURFACE_PROFILE);
 	float LerpFactor = 1;
 
@@ -440,6 +439,9 @@ void SubsurfaceRecombinePS(float4 UVAndScreenPos : TEXCOORD0, out float4 OutColo
 //	SSSColor = 0;
 #endif
 
+	// only affect pixels that have SSS
+	LerpFactor *= BlendFactor;
+
 	// add back in non subsurface pixels and specular, alpha: 0=don't blend in, 1:fully blend in
 	OutColor = float4(SSSColor * BlendFactor + SceneColor4.rgb, LerpFactor);
 }
diff --git a/Engine/Source/Runtime/Engine/Private/Rendering/SubsurfaceProfile.cpp b/Engine/Source/Runtime/Engine/Private/Rendering/SubsurfaceProfile.cpp
--- a/Engine/Source/Runtime/Engine/Private/Rendering/SubsurfaceProfile.cpp
+++ b/Engine/Source/Runtime/Engine/Private/Rendering/SubsurfaceProfile.cpp
@@ -166,8 +166,8 @@ void FSubsurfaceProfileTexture::CreateTexture(FRHICommandListImmediate& RHICmdLi
 
 		// bias to avoid div by 0 and a jump to a different value
 		// this basically means we don't want subsurface scattering
-		// 0.0001f turned out to be too small to fix the issue
-		const float Bias = 0.001f;
+		// 0.0001f turned out to be too small to fix the issue (for a small KernelSize)
+		const float Bias = 0.009f;
 
 		Data.SubsurfaceColor = Data.SubsurfaceColor.GetClamped(Bias);
 		Data.FalloffColor = Data.FalloffColor.GetClamped(Bias);
@@ -176,9 +176,11 @@ void FSubsurfaceProfileTexture::CreateTexture(FRHICommandListImmediate& RHICmdLi
 		ComputeMirroredSSSKernel(&kernel[KernelSize0], KernelSize1, Data.SubsurfaceColor, Data.FalloffColor);
 		ComputeMirroredSSSKernel(&kernel[KernelSize0 + KernelSize1], KernelSize2, Data.SubsurfaceColor, Data.FalloffColor);
 
+		// could be lower than 1 (but higher than 0) to range compress for better quality (for 8 bit)
 		const float TableMaxRGB = 1.0f;
 		const float TableMaxA = 3.0f;
 
+		// each kernel is normalized to be 1 per channel (center + one_side_samples * 2)
 		for (int32 Pos = 0; Pos < KernelTotalSize; ++Pos)
 		{
 			FVector4 C = kernel[Pos] * FLinearColor(1.0f / TableMaxRGB, 1.0f / TableMaxRGB, 1.0f / TableMaxRGB, 1.0f / TableMaxA);
diff --git a/Engine/Source/Runtime/Renderer/Private/PostProcess/PostProcessSubsurface.cpp b/Engine/Source/Runtime/Renderer/Private/PostProcess/PostProcessSubsurface.cpp
--- a/Engine/Source/Runtime/Renderer/Private/PostProcess/PostProcessSubsurface.cpp
+++ b/Engine/Source/Runtime/Renderer/Private/PostProcess/PostProcessSubsurface.cpp
@@ -58,7 +58,6 @@ public:
 	{
 		SSSParams.Bind(ParameterMap, TEXT("SSSParams"));
 		SSProfilesTexture.Bind(ParameterMap, TEXT("SSProfilesTexture"));
-		SSProfilesTextureSampler.Bind(ParameterMap, TEXT("SSProfilesTextureSampler"));
 	}
 
 	void SetParameters(FRHICommandList& RHICmdList, const FPixelShaderRHIParamRef& ShaderRHI, const FRenderingCompositePassContext& Context) const
@@ -96,21 +95,19 @@ public:
 
 			const FSceneRenderTargetItem& Item = PooledRT->GetRenderTargetItem();
 
-			SetTextureParameter(Context.RHICmdList, ShaderRHI, SSProfilesTexture, SSProfilesTextureSampler, TStaticSamplerState<SF_Point, AM_Wrap, AM_Wrap, AM_Wrap>::GetRHI(), Item.ShaderResourceTexture);
+			SetTextureParameter(Context.RHICmdList, ShaderRHI, SSProfilesTexture, Item.ShaderResourceTexture);
 		}
 	}
 
 	friend FArchive& operator<<(FArchive& Ar,FSubsurfaceParameters& P)
 	{
-		Ar << P.SSSParams << P.SSProfilesTexture << P.SSProfilesTextureSampler;
-
+		Ar << P.SSSParams << P.SSProfilesTexture;
 		return Ar;
 	}
 
 private:
 	FShaderParameter SSSParams;
 	FShaderResourceParameter SSProfilesTexture;
-	FShaderResourceParameter SSProfilesTextureSampler;
 };
 
 // ---------------------------------------------
@@ -435,8 +432,9 @@ static bool DoSpecularCorrection()
 
 /**
  * Encapsulates the post processing subsurface scattering pixel shader.
+ * @param SampleSet 0:low, 1:med, 2:high
  */
-template <uint32 RecombineMethod>
+template <uint32 RecombineMethod, uint32 SampleSet>
 class FPostProcessSubsurfaceExtractSpecularPS : public FGlobalShader
 {
 	DECLARE_SHADER_TYPE(FPostProcessSubsurfaceExtractSpecularPS, Global);
@@ -450,6 +448,7 @@ class FPostProcessSubsurfaceExtractSpecularPS : public FGlobalShader
 	{
 		FGlobalShader::ModifyCompilationEnvironment(Platform, OutEnvironment);
 		OutEnvironment.SetDefine(TEXT("SSS_RECOMBINE_METHOD"), RecombineMethod);
+		OutEnvironment.SetDefine(TEXT("SSS_SAMPLESET"), SampleSet);
 	}
 
 	/** Default constructor. */
@@ -500,24 +499,53 @@ public:
 };
 
 // #define avoids a lot of code duplication
-#define VARIATION1(A) typedef FPostProcessSubsurfaceExtractSpecularPS<A> FPostProcessSubsurfaceExtractSpecularPS##A; \
-	IMPLEMENT_SHADER_TYPE2(FPostProcessSubsurfaceExtractSpecularPS##A, SF_Pixel);
-VARIATION1(0) VARIATION1(1)
+#define VARIATION1(A)		VARIATION2(A,0)			VARIATION2(A,1)			VARIATION2(A,2)
+#define VARIATION2(A, B) typedef FPostProcessSubsurfaceExtractSpecularPS<A, B> FPostProcessSubsurfaceExtractSpecularPS##A##B; \
+	IMPLEMENT_SHADER_TYPE2(FPostProcessSubsurfaceExtractSpecularPS##A##B, SF_Pixel);
+	VARIATION1(0) VARIATION1(1)
 #undef VARIATION1
+#undef VARIATION2
 
 
-
+// @param SampleSet 0:low, 1:med, 2:high
 template <uint32 RecombineMethod>
-void SetSubsurfaceExtractSpecular(const FRenderingCompositePassContext& Context)
+void SetSubsurfaceExtractSpecular(const FRenderingCompositePassContext& Context, uint32 SampleSet)
 {
 	TShaderMapRef<FPostProcessVS> VertexShader(Context.GetShaderMap());
-	TShaderMapRef<FPostProcessSubsurfaceExtractSpecularPS<RecombineMethod> > PixelShader(Context.GetShaderMap());
 
-	static FGlobalBoundShaderState BoundShaderState;
+	switch(SampleSet)
+	{
+		case 0:
+			{
+				TShaderMapRef<FPostProcessSubsurfaceExtractSpecularPS<RecombineMethod, 0> > PixelShader(Context.GetShaderMap());
+				static FGlobalBoundShaderState BoundShaderState;
 
-	SetGlobalBoundShaderState(Context.RHICmdList, Context.GetFeatureLevel(), BoundShaderState, GFilterVertexDeclaration.VertexDeclarationRHI, *VertexShader, *PixelShader);
+				SetGlobalBoundShaderState(Context.RHICmdList, Context.GetFeatureLevel(), BoundShaderState, GFilterVertexDeclaration.VertexDeclarationRHI, *VertexShader, *PixelShader);
+				PixelShader->SetParameters(Context);
+				break;
+			}
+		case 1:
+			{
+				TShaderMapRef<FPostProcessSubsurfaceExtractSpecularPS<RecombineMethod, 1> > PixelShader(Context.GetShaderMap());
+				static FGlobalBoundShaderState BoundShaderState;
+
+				SetGlobalBoundShaderState(Context.RHICmdList, Context.GetFeatureLevel(), BoundShaderState, GFilterVertexDeclaration.VertexDeclarationRHI, *VertexShader, *PixelShader);
+				PixelShader->SetParameters(Context);
+				break;
+			}
+		case 2:
+			{
+				TShaderMapRef<FPostProcessSubsurfaceExtractSpecularPS<RecombineMethod, 2> > PixelShader(Context.GetShaderMap());
+				static FGlobalBoundShaderState BoundShaderState;
+
+				SetGlobalBoundShaderState(Context.RHICmdList, Context.GetFeatureLevel(), BoundShaderState, GFilterVertexDeclaration.VertexDeclarationRHI, *VertexShader, *PixelShader);
+				PixelShader->SetParameters(Context);
+				break;
+			} 
+		default:
+			check(0);
+	}
 
-	PixelShader->SetParameters(Context);
 	VertexShader->SetParameters(Context);
 }
 
@@ -561,13 +589,15 @@ void FRCPassPostProcessSubsurfaceExtractSpecular::Process(FRenderingCompositePas
 
 	SCOPED_DRAW_EVENTF(Context.RHICmdList, SubsurfacePass, TEXT("SubsurfacePassExtractSpecular#%d"), (int32)bDoSpecularCorrection);
 
+	uint32 SampleSet = FMath::Clamp(CVarSSSSampleSet.GetValueOnRenderThread(), 0, 2);
+
 	if(bDoSpecularCorrection)
 	{
-		SetSubsurfaceExtractSpecular<1>(Context);
+		SetSubsurfaceExtractSpecular<1>(Context, SampleSet);
 	}
 	else
 	{
-		SetSubsurfaceExtractSpecular<0>(Context);
+		SetSubsurfaceExtractSpecular<0>(Context, SampleSet);
 	}
 
 	DrawRectangle(
@@ -762,7 +792,7 @@ public:
 	}
 };
 
-	// #define avoids a lot of code duplication
+// #define avoids a lot of code duplication
 #define VARIATION1(A)		VARIATION2(A,0)			VARIATION2(A,1)			VARIATION2(A,2)
 #define VARIATION2(A, B) typedef TPostProcessSubsurfacePS<A, B> TPostProcessSubsurfacePS##A##B; \
 	IMPLEMENT_SHADER_TYPE2(TPostProcessSubsurfacePS##A##B, SF_Pixel);
