Distance Field Global Illumination first working version
* Requires a Movable Skylight + 'r.DistanceFieldGI 1' + 'r.GenerateMeshDistanceFields' project setting enabled
* Provides a single bounce of Diffuse GI with bounce distance limited by OcclusionMaxDistance (default 6m)
* Virtual Point Lights are placed from the directional light by ray tracing the scene's distance fields, normal comes from SDF gradient
* Currently only placing 128*128 VPLs up to 40m from camera
* VPLs light irradiance cache samples using hemispherical disk light model, indirect shadowing is provided by the same cone traces used for Distance Field AO
* Irradiance is computed only at irradiance cache points (max every 8 pixels) and interpolated to the rest of the pixels
* Hardcoded half grey material color for now, and poor performance due to brute force VPL lighting

diff --git a/Engine/Shaders/DistanceFieldAOShared.usf b/Engine/Shaders/DistanceFieldAOShared.usf
--- a/Engine/Shaders/DistanceFieldAOShared.usf
+++ b/Engine/Shaders/DistanceFieldAOShared.usf
@@ -86,16 +86,6 @@ float2 AOBufferSize;
 int DownsampleFactorToBaseLevel;
 float2 BaseLevelTexelSize;
 
-float3 EncodeNormalForAO(float3 InNormal)
-{
-	return InNormal.xyz;
-}
-
-float3 DecodeNormalForAO(float3 InNormal)
-{
-	return InNormal.xyz;
-}
-
 bool SphereIntersectCone(float4 SphereCenterAndRadius, float3 ConeVertex, float3 ConeAxis, float ConeAngleCos, float ConeAngleSin)
 {
 	float3 U = ConeVertex - (SphereCenterAndRadius.w / ConeAngleSin) * ConeAxis;
diff --git a/Engine/Shaders/DistanceFieldGlobalIllumination.usf b/Engine/Shaders/DistanceFieldGlobalIllumination.usf
new file mode 100644
--- /dev/null
+++ b/Engine/Shaders/DistanceFieldGlobalIllumination.usf
@@ -0,0 +1,831 @@
+// Copyright 1998-2014 Epic Games, Inc. All Rights Reserved.
+
+/*=============================================================================
+	DistanceFieldGlobalIllumination.usf
+=============================================================================*/
+
+#include "Common.usf"
+#include "DeferredShadingCommon.usf"
+#include "DistanceFieldLightingShared.usf"
+#include "DistanceFieldAOShared.usf"
+#include "DynamicLightingCommon.usf"
+
+#define USE_SHADOW_CULLING_FOR_VPL_PLACEMENT 0
+
+void RayTraceThroughLightTileCulledObjects(
+	float3 RayStartPosition, 
+	float3 RayDirection, 
+	float RayLength, 
+	uint NumIntersectingObjects, 
+	uint CulledDataStart, 
+	out float OutMinRayTime, 
+	out float OutMinRayVisibility,
+	out uint OutObjectIndex)
+{
+	float MinSphereRadius = .1f;
+	float MaxSphereRadius = .1f;
+
+	float3 WorldRayStart = RayStartPosition;
+	float3 WorldRayEnd = RayStartPosition + RayDirection * RayLength;
+	float MaxRayTime = RayLength;
+	float TanLightAngle = 0;
+
+	float MinRayTime = MaxRayTime;
+	float MinVisibility = 1;
+
+	LOOP
+	for (uint ListObjectIndex = 0; ListObjectIndex < NumIntersectingObjects; ListObjectIndex++)
+	{
+		#if USE_SHADOW_CULLING_FOR_VPL_PLACEMENT
+			uint ObjectIndex = ShadowTileArrayData.Load(ListObjectIndex * ShadowTileListGroupSize.x * ShadowTileListGroupSize.y + CulledDataStart);
+		#else
+			uint ObjectIndex = ListObjectIndex;
+		#endif
+
+		float3 LocalPositionExtent = LoadObjectLocalPositionExtent(ObjectIndex);
+		float4x4 WorldToVolume = LoadObjectWorldToVolume(ObjectIndex);
+		float4 UVScaleAndVolumeScale = LoadObjectUVScale(ObjectIndex);
+		float3 UVAdd = LoadObjectUVAdd(ObjectIndex);
+
+		float3 VolumeRayStart = mul(float4(WorldRayStart, 1), WorldToVolume).xyz;
+		float3 VolumeRayEnd = mul(float4(WorldRayEnd, 1), WorldToVolume).xyz;
+		float3 VolumeRayDirection = VolumeRayEnd - VolumeRayStart;
+		float VolumeRayLength = length(VolumeRayDirection);
+		VolumeRayDirection /= VolumeRayLength;
+		float VolumeMinSphereRadius = MinSphereRadius / UVScaleAndVolumeScale.w;
+		float VolumeMaxSphereRadius = MaxSphereRadius / UVScaleAndVolumeScale.w;
+
+		float4 SphereCenterAndRadius = LoadObjectPositionAndRadius(ObjectIndex);
+		float ObjectCenterDistanceAlongRay = max(dot(SphereCenterAndRadius.xyz - WorldRayStart, WorldRayEnd - WorldRayStart), 0);
+		// Expand the intersection box by the radius of the cone at the distance of the object along the cone
+		float LocalConeRadiusAtObject = min(TanLightAngle * ObjectCenterDistanceAlongRay / UVScaleAndVolumeScale.w, VolumeMaxSphereRadius);
+
+		float2 IntersectionTimes = LineBoxIntersect(VolumeRayStart, VolumeRayEnd, -LocalPositionExtent - LocalConeRadiusAtObject, LocalPositionExtent + LocalConeRadiusAtObject);
+
+		BRANCH
+		if (IntersectionTimes.x < IntersectionTimes.y && IntersectionTimes.x < 1)
+		{
+			float SampleRayTime = IntersectionTimes.x * VolumeRayLength;
+			uint MaxSteps = 64;
+			float MinStepSize = 1.0f / (4 * MaxSteps);
+
+			float MinDistance = 1000000;
+			float3 IntersectionPosition = float3(0, 0, 0);
+
+			uint StepIndex = 0;
+
+			LOOP
+			for (; StepIndex < MaxSteps; StepIndex++)
+			{
+				float3 SampleVolumePosition = VolumeRayStart + VolumeRayDirection * SampleRayTime;
+				float3 ClampedSamplePosition = clamp(SampleVolumePosition, -LocalPositionExtent, LocalPositionExtent);
+				float DistanceToClamped = length(ClampedSamplePosition - SampleVolumePosition);
+				float3 VolumeUV = DistanceFieldVolumePositionToUV(ClampedSamplePosition, UVScaleAndVolumeScale.xyz, UVAdd);
+				float DistanceField = Texture3DSampleLevel(DistanceFieldTexture, DistanceFieldSampler, VolumeUV, 0).x + DistanceToClamped;
+
+				MinDistance = min(MinDistance, DistanceField);
+				float SphereRadius = clamp(TanLightAngle * SampleRayTime, VolumeMinSphereRadius, VolumeMaxSphereRadius);
+
+				MinVisibility = min(MinVisibility, saturate(DistanceField / SphereRadius));
+				IntersectionPosition = SampleVolumePosition;
+
+				float StepDistance = max(DistanceField, MinStepSize);
+				
+				// Terminate the trace if we reached a negative area or went past the end of the ray
+				if (DistanceField <= 0 
+					|| SampleRayTime + StepDistance > IntersectionTimes.y * VolumeRayLength)
+				{
+					// Step back to the intersection point if we went inside
+					SampleRayTime += min(DistanceField, 0);
+					break;
+				}
+
+				SampleRayTime += StepDistance;
+			}
+
+			if (MinDistance < 0 || StepIndex == MaxSteps)
+			{
+				MinVisibility = 0;
+				//MinRayTime = min(MinRayTime, SampleRayTime * UVScaleAndVolumeScale.w);
+
+				if (SampleRayTime * UVScaleAndVolumeScale.w < MinRayTime)
+				{
+					MinRayTime = UVScaleAndVolumeScale.w * SampleRayTime;
+					OutObjectIndex = ObjectIndex;
+				}
+			}
+		}
+	}
+
+	OutMinRayVisibility = MinVisibility;
+	OutMinRayTime = MinRayTime;
+}
+
+float3 ComputeDistanceFieldNormal(float3 WorldPosition, uint ObjectIndex)
+{
+	float4x4 WorldToVolume = LoadObjectWorldToVolume(ObjectIndex);
+	float4 UVScaleAndVolumeScale = LoadObjectUVScale(ObjectIndex);
+	float3 UVAdd = LoadObjectUVAdd(ObjectIndex);
+
+	float3x3 VolumeToWorld = LoadObjectVolumeToWorld(ObjectIndex);
+
+	float3 LocalPositionExtent = LoadObjectLocalPositionExtent(ObjectIndex);
+
+	float3 VolumeShadingPosition = mul(float4(WorldPosition, 1), WorldToVolume).xyz;
+	float3 ClampedSamplePosition = clamp(VolumeShadingPosition, -LocalPositionExtent, LocalPositionExtent);
+	float3 LocalShadingUV = DistanceFieldVolumePositionToUV(ClampedSamplePosition, UVScaleAndVolumeScale.xyz, UVAdd);
+
+	// Used to clamp UVs inside valid space of this object's distance field
+	float3 UVMin = DistanceFieldVolumePositionToUV(-LocalPositionExtent, UVScaleAndVolumeScale.xyz, UVAdd);
+	float3 UVMax = DistanceFieldVolumePositionToUV(LocalPositionExtent, UVScaleAndVolumeScale.xyz, UVAdd);
+		
+	float R = Texture3DSampleLevel(DistanceFieldTexture, DistanceFieldSampler, float3(min(LocalShadingUV.x + DistanceFieldAtlasTexelSize.x, UVMax.x), LocalShadingUV.y, LocalShadingUV.z), 0).x;
+	float L = Texture3DSampleLevel(DistanceFieldTexture, DistanceFieldSampler, float3(max(LocalShadingUV.x - DistanceFieldAtlasTexelSize.x, UVMin.x), LocalShadingUV.y, LocalShadingUV.z), 0).x;
+	float F = Texture3DSampleLevel(DistanceFieldTexture, DistanceFieldSampler, float3(LocalShadingUV.x, min(LocalShadingUV.y + DistanceFieldAtlasTexelSize.y, UVMax.y), LocalShadingUV.z), 0).x;
+	float B = Texture3DSampleLevel(DistanceFieldTexture, DistanceFieldSampler, float3(LocalShadingUV.x, max(LocalShadingUV.y - DistanceFieldAtlasTexelSize.y, UVMin.y), LocalShadingUV.z), 0).x;
+	float U = Texture3DSampleLevel(DistanceFieldTexture, DistanceFieldSampler, float3(LocalShadingUV.x, LocalShadingUV.y, min(LocalShadingUV.z + DistanceFieldAtlasTexelSize.z, UVMax.z)), 0).x;
+	float D = Texture3DSampleLevel(DistanceFieldTexture, DistanceFieldSampler, float3(LocalShadingUV.x, LocalShadingUV.y, max(LocalShadingUV.z - DistanceFieldAtlasTexelSize.z, UVMin.z)), 0).x;
+		
+	float3 Gradient = .5f * float3(R - L, F - B, U - D);
+	
+	if (dot(Gradient, Gradient) == 0)
+	{
+		Gradient = float3(0, 0, 1);
+	}
+
+	float3 LocalNormal = normalize(Gradient);
+	float3 WorldNormal = mul(LocalNormal, VolumeToWorld);
+	return normalize(WorldNormal);
+}
+
+/** From light source, into world. */
+float4 LightDirectionAndTraceDistance;
+float4 LightColor;
+float4x4 WorldToShadow;
+float4x4 ShadowToWorld;
+float2 InvPlacementGridSize;
+float VPLPlacementCameraRadius;
+
+// Must match C++
+#define VPL_CLUSTER_SIZE_ONE_DIM 4
+#define VPL_CLUSTER_SIZE (VPL_CLUSTER_SIZE_ONE_DIM * VPL_CLUSTER_SIZE_ONE_DIM)
+
+// In float4's, must match C++
+RWBuffer<uint> RWVPLParameterBuffer;
+RWBuffer<float4> RWVPLClusterData;
+RWBuffer<float4> RWVPLData;
+
+[numthreads(THREADGROUP_SIZEX, THREADGROUP_SIZEY, 1)]
+void VPLPlacementCS(
+	uint3 GroupId : SV_GroupID,
+	uint3 DispatchThreadId : SV_DispatchThreadID,
+    uint3 GroupThreadId : SV_GroupThreadID) 
+{
+	{
+		// Distance for directional lights to trace
+		float TraceDistance = LightDirectionAndTraceDistance.w;
+		float3 LightDirection = LightDirectionAndTraceDistance.xyz;
+		
+// Note: have to change the dispatch dimensions in C++ to enable this
+#define GENERATE_VPL_CLUSTERS 0
+#if GENERATE_VPL_CLUSTERS
+		uint NextVPLIndex = 0;
+		float3 IntersectionPositions[VPL_CLUSTER_SIZE];
+		float3 IntersectionNormals[VPL_CLUSTER_SIZE];
+		float3 AveragePosition = 0;
+		float3 AverageNormal = 0;
+
+		for (uint ChildY = 0; ChildY < VPL_CLUSTER_SIZE_ONE_DIM; ChildY++)
+		{
+			for (uint ChildX = 0; ChildX < VPL_CLUSTER_SIZE_ONE_DIM; ChildX++)
+			{
+				uint2 ChildCoordinate = DispatchThreadId.xy * VPL_CLUSTER_SIZE_ONE_DIM + uint2(ChildX, ChildY);
+
+				if (all(ChildCoordinate * InvPlacementGridSize < 1))
+				{
+					float2 NormalizedCellPosition = ChildCoordinate * InvPlacementGridSize * 2 - 1;
+					float3 CellStartWorldPosition = mul(float4(NormalizedCellPosition.x, NormalizedCellPosition.y, 0, 1), ShadowToWorld).xyz;
+
+					uint NumIntersectingObjects = GetCulledNumObjects();
+					uint CulledDataStart = 0;
+
+				#if USE_SHADOW_CULLING_FOR_VPL_PLACEMENT
+
+					// Transform into shadow space
+					float4 HomogeneousShadowPosition = mul(float4(CellStartWorldPosition, 1), WorldToShadow);
+					float2 NormalizedShadowPosition = HomogeneousShadowPosition.xy * .5f + .5f;
+					NormalizedShadowPosition.y = 1 - NormalizedShadowPosition.y;
+					// Quantize the shadow position to get our tile position
+					uint2 TilePosition = (uint2)(NormalizedShadowPosition * ShadowTileListGroupSize);
+					// Fetch the tile head information
+					uint2 TileHead = GetShadowTileHead(TilePosition);
+					NumIntersectingObjects = TileHead.y;
+					CulledDataStart = TileHead.x;
+
+				#endif
+
+					float MinRayTime = 0;
+					float MinRayVisibility = 1;
+					uint ObjectIndex = 0;
+					RayTraceThroughLightTileCulledObjects(CellStartWorldPosition, LightDirection, TraceDistance, NumIntersectingObjects, CulledDataStart, MinRayTime, MinRayVisibility, ObjectIndex);
+
+					if (MinRayVisibility < 1)
+					{
+						float3 IntersectionPosition = CellStartWorldPosition + LightDirection * MinRayTime;
+						float3 IntersectionNormal = ComputeDistanceFieldNormal(IntersectionPosition, ObjectIndex);
+						AveragePosition += IntersectionPosition;
+						AverageNormal += IntersectionNormal;
+						IntersectionPositions[NextVPLIndex] = IntersectionPosition;
+						IntersectionNormals[NextVPLIndex] = IntersectionNormal;
+						NextVPLIndex++;
+					}
+				}
+			}
+		}
+
+		uint NumVPLs = NextVPLIndex;
+
+		if (NumVPLs > 0)
+		{
+			uint VPLArrayStartIndex;
+			InterlockedAdd(RWVPLParameterBuffer[1], NumVPLs, VPLArrayStartIndex);
+
+			float3 DiffuseColor = .5f;
+			float CellExtent = VPLPlacementCameraRadius * InvPlacementGridSize.x;
+			float CellRadius = sqrt(2) * CellExtent;
+
+			float3 TotalFlux = 0;
+
+			for (uint ChildIndex = 0; ChildIndex < NumVPLs; ChildIndex++)
+			{
+				float3 Position = IntersectionPositions[ChildIndex];
+				float3 Normal = IntersectionNormals[ChildIndex];
+				float3 Flux = DiffuseColor * LightColor.rgb * max(dot(Normal, -LightDirection), 0) * PI * CellRadius * CellRadius;
+
+				uint VPLBaseIndex = (VPLArrayStartIndex + ChildIndex) * VPL_DATA_STRIDE;
+				RWVPLData[VPLBaseIndex + 0] = float4(Position, CellRadius);
+				RWVPLData[VPLBaseIndex + 1] = float4(Normal, 0);
+				RWVPLData[VPLBaseIndex + 2] = float4(Flux, 0);
+
+				TotalFlux += Flux;
+			}
+
+			AveragePosition /= NumVPLs;
+			float AverageNormalLength = length(AverageNormal);
+			AverageNormal = AverageNormal / max(AverageNormalLength, .000001f);
+
+			float TotalChildDiskArea = NumVPLs * PI * CellRadius * CellRadius;
+			float CombinedRadius = sqrt(TotalChildDiskArea / PI);
+
+			uint ClusterArrayStartIndex;
+			InterlockedAdd(RWVPLParameterBuffer[0], 1, ClusterArrayStartIndex);
+			RWVPLClusterData[ClusterArrayStartIndex * VPL_DATA_STRIDE + 0] = float4(AveragePosition, CombinedRadius);
+			RWVPLClusterData[ClusterArrayStartIndex * VPL_DATA_STRIDE + 1] = float4(AverageNormal, NumVPLs);
+			RWVPLClusterData[ClusterArrayStartIndex * VPL_DATA_STRIDE + 2] = float4(TotalFlux, VPLArrayStartIndex);
+		}
+#else
+		uint2 ChildCoordinate = DispatchThreadId.xy;
+
+		if (all(ChildCoordinate * InvPlacementGridSize < 1))
+		{
+			float2 NormalizedCellPosition = ChildCoordinate * InvPlacementGridSize * 2 - 1;
+			float3 CellStartWorldPosition = mul(float4(NormalizedCellPosition.x, NormalizedCellPosition.y, 0, 1), ShadowToWorld).xyz;
+
+			uint NumIntersectingObjects = GetCulledNumObjects();
+			uint CulledDataStart = 0;
+
+		#if USE_SHADOW_CULLING_FOR_VPL_PLACEMENT
+
+			// Transform into shadow space
+			float4 HomogeneousShadowPosition = mul(float4(CellStartWorldPosition, 1), WorldToShadow);
+			float2 NormalizedShadowPosition = HomogeneousShadowPosition.xy * .5f + .5f;
+			NormalizedShadowPosition.y = 1 - NormalizedShadowPosition.y;
+			// Quantize the shadow position to get our tile position
+			uint2 TilePosition = (uint2)(NormalizedShadowPosition * ShadowTileListGroupSize);
+			// Fetch the tile head information
+			uint2 TileHead = GetShadowTileHead(TilePosition);
+			NumIntersectingObjects = TileHead.y;
+			CulledDataStart = TileHead.x;
+
+		#endif
+
+			float MinRayTime = 0;
+			float MinRayVisibility = 1;
+			uint ObjectIndex = 0;
+			RayTraceThroughLightTileCulledObjects(CellStartWorldPosition, LightDirection, TraceDistance, NumIntersectingObjects, CulledDataStart, MinRayTime, MinRayVisibility, ObjectIndex);
+
+			if (MinRayVisibility < 1)
+			{
+				float3 IntersectionPosition = CellStartWorldPosition + LightDirection * MinRayTime;
+				float3 IntersectionNormal = ComputeDistanceFieldNormal(IntersectionPosition, ObjectIndex);
+
+				uint VPLArrayStartIndex;
+				InterlockedAdd(RWVPLParameterBuffer[1], 1, VPLArrayStartIndex);
+
+				float3 DiffuseColor = .5f;
+				float CellExtent = VPLPlacementCameraRadius * InvPlacementGridSize.x;
+				float CellRadius = sqrt(2) * CellExtent;
+
+				float3 Flux = DiffuseColor * LightColor.rgb * max(dot(IntersectionNormal, -LightDirection), 0) * PI * CellRadius * CellRadius;
+
+				uint VPLBaseIndex = (VPLArrayStartIndex + 0) * VPL_DATA_STRIDE;
+				RWVPLData[VPLBaseIndex + 0] = float4(IntersectionPosition, CellRadius);
+				RWVPLData[VPLBaseIndex + 1] = float4(IntersectionNormal, 0);
+				RWVPLData[VPLBaseIndex + 2] = float4(Flux, 0);
+			}
+		}
+#endif
+	}
+}
+
+Buffer<uint> VPLParameterBuffer;
+RWBuffer<uint> RWDispatchParameters;
+
+[numthreads(1, 1, 1)]
+void SetupVPLCullndirectArgumentsCS(
+	uint3 GroupId : SV_GroupID,
+	uint3 DispatchThreadId : SV_DispatchThreadID,
+    uint3 GroupThreadId : SV_GroupThreadID) 
+{
+#if GENERATE_VPL_CLUSTERS
+	uint NumClusterVPLs = VPLParameterBuffer[0];
+#else
+	uint NumClusterVPLs = VPLParameterBuffer[1];
+#endif
+
+	// One thread per record, divide and round up
+	RWDispatchParameters[0] = (NumClusterVPLs + THREADGROUP_TOTALSIZE - 1) / THREADGROUP_TOTALSIZE;
+	RWDispatchParameters[1] = 1;
+	RWDispatchParameters[2] = 1;
+}
+
+Buffer<float4> VPLClusterData;
+Buffer<float4> VPLData;
+
+RWBuffer<uint> RWCulledVPLParameterBuffer;
+RWBuffer<float4> RWCulledVPLClusterData;
+RWBuffer<float4> RWCulledVPLData;
+
+[numthreads(THREADGROUP_TOTALSIZE, 1, 1)]
+void CullVPLsForViewCS(
+	uint3 GroupId : SV_GroupID,
+	uint3 DispatchThreadId : SV_DispatchThreadID,
+    uint3 GroupThreadId : SV_GroupThreadID) 
+{
+#if GENERATE_VPL_CLUSTERS
+	uint VPLClusterIndex = DispatchThreadId.x;
+	uint NumClusterVPLs = VPLParameterBuffer[0];
+
+	if (VPLClusterIndex < NumClusterVPLs)
+	{
+		float4 BoundingSphere = VPLClusterData[VPLClusterIndex * VPL_DATA_STRIDE + 0];
+		float DistanceToViewSq = dot(View.ViewOrigin.xyz - BoundingSphere.xyz, View.ViewOrigin.xyz - BoundingSphere.xyz);
+
+		if (DistanceToViewSq < Square(AOMaxViewDistance + BoundingSphere.w)
+			&& ViewFrustumIntersectSphere(BoundingSphere.xyz, BoundingSphere.w + AOMaxDistance))
+		{
+			float4 ClusterNormalAndNumVPLs = VPLClusterData[VPLClusterIndex * VPL_DATA_STRIDE + 1];
+			float4 ClusterFluxAndVPLStartIndex = VPLClusterData[VPLClusterIndex * VPL_DATA_STRIDE + 2];
+
+			uint DestStartVPLIndex;
+			InterlockedAdd(RWCulledVPLParameterBuffer[1], (uint)ClusterNormalAndNumVPLs.w, DestStartVPLIndex);
+			uint SourceVPLIndex = ClusterFluxAndVPLStartIndex.w;
+			uint NumChildVPLs = ClusterNormalAndNumVPLs.w;
+
+			for (uint ChildIndex = 0; ChildIndex < NumChildVPLs; ChildIndex++)
+			{
+				uint SourceIndex = (SourceVPLIndex + ChildIndex) * VPL_DATA_STRIDE;
+				uint DestIndex = (DestStartVPLIndex + ChildIndex) * VPL_DATA_STRIDE;
+				RWCulledVPLData[DestIndex + 0] = VPLData[SourceIndex + 0];
+				RWCulledVPLData[DestIndex + 1] = VPLData[SourceIndex + 1];
+				RWCulledVPLData[DestIndex + 2] = VPLData[SourceIndex + 2];
+			}
+
+			uint DestClusterIndex;
+			InterlockedAdd(RWCulledVPLParameterBuffer[0], 1, DestClusterIndex);
+			
+			RWCulledVPLClusterData[DestClusterIndex * VPL_DATA_STRIDE + 0] = BoundingSphere;
+			RWCulledVPLClusterData[DestClusterIndex * VPL_DATA_STRIDE + 1] = ClusterNormalAndNumVPLs;
+			RWCulledVPLClusterData[DestClusterIndex * VPL_DATA_STRIDE + 2] = float4(ClusterFluxAndVPLStartIndex.xyz, DestStartVPLIndex);
+		}
+	}
+
+#else
+
+	uint VPLIndex = DispatchThreadId.x;
+	uint NumVPLs = VPLParameterBuffer[1];
+
+	if (VPLIndex < NumVPLs)
+	{
+		uint SourceIndex = VPLIndex * VPL_DATA_STRIDE;
+		float4 BoundingSphere = VPLData[SourceIndex + 0];
+		float DistanceToViewSq = dot(View.ViewOrigin.xyz - BoundingSphere.xyz, View.ViewOrigin.xyz - BoundingSphere.xyz);
+
+		if (DistanceToViewSq < Square(AOMaxViewDistance + BoundingSphere.w)
+			&& ViewFrustumIntersectSphere(BoundingSphere.xyz, BoundingSphere.w + AOMaxDistance))
+		{
+			uint DestStartVPLIndex;
+			InterlockedAdd(RWCulledVPLParameterBuffer[1], 1, DestStartVPLIndex);
+
+			uint DestIndex = DestStartVPLIndex * VPL_DATA_STRIDE;
+			RWCulledVPLData[DestIndex + 0] = BoundingSphere;
+			RWCulledVPLData[DestIndex + 1] = VPLData[SourceIndex + 1];
+			RWCulledVPLData[DestIndex + 2] = VPLData[SourceIndex + 2];
+		}
+	}
+
+#endif
+}
+
+RWBuffer<uint> RWTileArrayData;
+RWBuffer<uint> RWTileHeadDataUnpacked;
+
+Buffer<float4> TileConeAxisAndCos;
+Buffer<float4> TileConeDepthRanges;
+uint2 NumGroups;
+
+// Must match C++
+#define MAX_VPLS_PER_TILE 1024
+
+[numthreads(THREADGROUP_TOTALSIZE, 1, 1)]
+void VPLTileCullCS(
+	uint3 GroupId : SV_GroupID,
+	uint3 DispatchThreadId : SV_DispatchThreadID,
+    uint3 GroupThreadId : SV_GroupThreadID) 
+{
+    uint ThreadIndex = GroupThreadId.x;
+	uint NumClusterVPLs = VPLParameterBuffer[0];
+	uint TileIndex = GroupId.y * NumGroups.x + GroupId.x;
+
+	float4 ConeAxisAndCos = TileConeAxisAndCos.Load(TileIndex);
+	float4 ConeAxisDepthRanges = TileConeDepthRanges.Load(TileIndex);
+	float3 TileConeVertex = 0;
+	float3 TileConeAxis = ConeAxisAndCos.xyz;
+	float TileConeAngleCos = ConeAxisAndCos.w;
+	float TileConeAngleSin = sqrt(1 - TileConeAngleCos * TileConeAngleCos);
+
+	LOOP
+	for (uint VPLClusterIndex = ThreadIndex; VPLClusterIndex < NumClusterVPLs; VPLClusterIndex += THREADGROUP_TOTALSIZE)
+	{
+		float4 BoundingSphere = VPLClusterData[VPLClusterIndex * VPL_DATA_STRIDE + 0];
+		float3 ViewSpaceSphereCenter = mul(float4(BoundingSphere.xyz + View.PreViewTranslation.xyz, 1), View.TranslatedWorldToView).xyz;
+		
+		BRANCH
+		if (SphereIntersectConeWithDepthRanges(float4(ViewSpaceSphereCenter, AOMaxDistance), TileConeVertex, TileConeAxis, TileConeAngleCos, TileConeAngleSin, ConeAxisDepthRanges))
+		{
+			uint CulledVPLClusterIndex;
+			InterlockedAdd(RWTileHeadDataUnpacked[TileIndex * 2 + 1], 1, CulledVPLClusterIndex);
+
+			if (CulledVPLClusterIndex < MAX_VPLS_PER_TILE)
+			{
+				RWTileArrayData[TileIndex * MAX_VPLS_PER_TILE + CulledVPLClusterIndex] = VPLClusterIndex;
+			}
+		}
+	}
+}
+
+float RayTraceThroughGlobalObjects(
+	float3 RayStartPosition, 
+	float3 RayDirection, 
+	float RayLength,
+	float TanConeAngle,
+	float ConeEndRayTime)
+{
+	float MaxSphereRadius = 100;
+
+	float3 WorldRayStart = RayStartPosition;
+	float3 WorldRayEnd = RayStartPosition + RayDirection * RayLength;
+	float MaxRayTime = RayLength;
+
+	float MinRayTime = MaxRayTime;
+	float MinVisibility = 1;
+
+	LOOP
+	for (uint ObjectIndex = 0; ObjectIndex < NumSceneObjects && MinVisibility > 0; ObjectIndex++)
+	{
+		float3 LocalPositionExtent = LoadGlobalObjectLocalPositionExtent(ObjectIndex);
+		float4x4 WorldToVolume = LoadGlobalObjectWorldToVolume(ObjectIndex);
+		float4 UVScaleAndVolumeScale = LoadGlobalObjectUVScale(ObjectIndex);
+		float3 UVAdd = LoadGlobalObjectUVAdd(ObjectIndex);
+
+		float3 VolumeRayStart = mul(float4(WorldRayStart, 1), WorldToVolume).xyz;
+		float3 VolumeRayEnd = mul(float4(WorldRayEnd, 1), WorldToVolume).xyz;
+		float3 VolumeRayDirection = VolumeRayEnd - VolumeRayStart;
+		float VolumeRayLength = length(VolumeRayDirection);
+		VolumeRayDirection /= VolumeRayLength;
+		float VolumeMaxSphereRadius = MaxSphereRadius / UVScaleAndVolumeScale.w;
+		float VolumeConeEndRayTime = ConeEndRayTime / UVScaleAndVolumeScale.w;
+		float ConeEndNormalization = 1.0f / (VolumeRayLength - VolumeConeEndRayTime);
+
+		float4 SphereCenterAndRadius = LoadGlobalObjectPositionAndRadius(ObjectIndex);
+		float ObjectCenterDistanceAlongRay = max(dot(SphereCenterAndRadius.xyz - WorldRayStart, WorldRayEnd - WorldRayStart), 0);
+		// Expand the intersection box by the radius of the cone at the distance of the object along the cone
+		float LocalConeRadiusAtObject = min(TanConeAngle * ObjectCenterDistanceAlongRay / UVScaleAndVolumeScale.w, VolumeMaxSphereRadius);
+
+		float2 IntersectionTimes = LineBoxIntersect(VolumeRayStart, VolumeRayEnd, -LocalPositionExtent - LocalConeRadiusAtObject, LocalPositionExtent + LocalConeRadiusAtObject);
+
+		BRANCH
+		if (IntersectionTimes.x < IntersectionTimes.y && IntersectionTimes.x < 1)
+		{
+			float SampleRayTime = IntersectionTimes.x * VolumeRayLength;
+			uint MaxSteps = 32;
+			float MinStepSize = 1.0f / (4 * MaxSteps);
+
+			uint StepIndex = 0;
+
+			LOOP
+			for (; StepIndex < MaxSteps; StepIndex++)
+			{
+				float3 SampleVolumePosition = VolumeRayStart + VolumeRayDirection * SampleRayTime;
+				float3 ClampedSamplePosition = clamp(SampleVolumePosition, -LocalPositionExtent, LocalPositionExtent);
+				float DistanceToClamped = length(ClampedSamplePosition - SampleVolumePosition);
+				float3 VolumeUV = DistanceFieldVolumePositionToUV(ClampedSamplePosition, UVScaleAndVolumeScale.xyz, UVAdd);
+				float DistanceField = Texture3DSampleLevel(DistanceFieldTexture, DistanceFieldSampler, VolumeUV, 0).x + DistanceToClamped;
+
+				float SphereRadius = clamp(TanConeAngle * SampleRayTime, 0, VolumeMaxSphereRadius);
+
+				if (SampleRayTime > VolumeConeEndRayTime)
+				{
+					// 0 at VolumeRayLength, 1 at VolumeConeEndRayTime
+					float ConeEndAlpha = saturate((VolumeRayLength - SampleRayTime) * ConeEndNormalization);
+					// Reduce the intersection sphere radius to 0 at the end of the cone
+					SphereRadius = ConeEndAlpha * TanConeAngle * VolumeConeEndRayTime;
+				}
+
+				//SphereRadius = 0;
+
+				MinVisibility = min(MinVisibility, saturate(DistanceField / SphereRadius));
+
+				float StepDistance = max(DistanceField, MinStepSize);
+				SampleRayTime += StepDistance;
+
+				// Terminate the trace if we reached a negative area or went past the end of the ray
+				if (DistanceField <= 0 
+					|| SampleRayTime > IntersectionTimes.y * VolumeRayLength)
+				{
+					break;
+				}
+			}
+
+			if (StepIndex == MaxSteps)
+			{
+				MinVisibility = 0;
+			}
+		}
+	}
+
+	return MinVisibility;
+}
+
+float ApproximateConeConeIntersection(float ArcLength0, float ArcLength1, float AngleBetweenCones)
+{
+	float AngleDifference = abs(ArcLength0 - ArcLength1);
+
+	float Intersection = smoothstep(
+		0,
+		1.0,
+		1.0 - saturate((AngleBetweenCones - AngleDifference) / (ArcLength0 + ArcLength1 - AngleDifference)));
+
+	return Intersection;
+}
+
+float GetVPLOcclusion(float3 BentNormalAO, float3 NormalizedVectorToVPL)
+{
+	float BentNormalLength = length(BentNormalAO);
+	float VPLConeAngle = .5f;
+	float VPLOcclusionStrength = 1;
+	float UnoccludedAngle = BentNormalLength * PI / VPLOcclusionStrength;
+	float AngleBetween = acos(dot(BentNormalAO, NormalizedVectorToVPL) / max(BentNormalLength, .0001f));
+	float Visibility = ApproximateConeConeIntersection(VPLConeAngle, UnoccludedAngle, AngleBetween);
+
+	// Can't rely on the direction of the bent normal when close to fully occluded, lerp to shadowed
+	Visibility = lerp(0, Visibility, saturate((UnoccludedAngle - .1f) / .2f));
+
+	return Visibility;
+}
+
+Buffer<float> RecordConeData;
+
+float VPLGatherRadius;
+float ConeHalfAngle;
+float BentNormalNormalizeFactor;
+
+static float3 StepBentNormal[NUM_VISIBILITY_STEPS];
+
+float3 ComputeVirtualPointLighting(uniform bool bCluster, uint VPLIndex, float3 WorldPosition, float3 WorldNormal)
+{
+	float3 Irradiance = 0;
+
+	float VisibilityStepSize = NUM_VISIBILITY_STEPS / AOMaxDistance;
+	float MaxGatherDistanceSq = VPLGatherRadius * VPLGatherRadius;
+	float4 VPLPositionAndRadius = bCluster ? VPLClusterData[VPLIndex * VPL_DATA_STRIDE + 0] : VPLData[VPLIndex * VPL_DATA_STRIDE + 0];
+	float3 VPLToGatherPoint = WorldPosition - VPLPositionAndRadius.xyz;
+	float DistanceSq = dot(VPLToGatherPoint, VPLToGatherPoint);
+	float DirectionDot = dot(-VPLToGatherPoint, WorldNormal);
+
+	// Hack
+	//VPLPositionAndRadius.w = 100;
+
+#define VISUALIZE_VPL_PLACEMENT 0
+#if VISUALIZE_VPL_PLACEMENT
+	Irradiance.r += 3 * (DistanceSq < VPLPositionAndRadius.w * VPLPositionAndRadius.w);
+#endif
+
+	BRANCH
+	if (DistanceSq < MaxGatherDistanceSq && DirectionDot > 0)
+	{
+		float3 VPLNormal = bCluster ? VPLClusterData[VPLIndex * VPL_DATA_STRIDE + 1].xyz : VPLData[VPLIndex * VPL_DATA_STRIDE + 1].xyz;
+		float VPLNormalDot = dot(VPLNormal, VPLToGatherPoint);
+
+		BRANCH
+		if (VPLNormalDot > 0)
+		{
+			float Distance = sqrt(DistanceSq);
+			float3 VPLDirection = -VPLToGatherPoint / Distance;
+
+			#define USE_INVERSE_SQUARED_DISK_APPROX 1
+			#define USE_PROJECTED_ANGLE_ATTENUATION 0
+
+			#if USE_INVERSE_SQUARED_DISK_APPROX
+				float EffectiveDiskRadius = VPLPositionAndRadius.w * 1;
+				float DistanceAttenuation = 1.0f / (PI * DistanceSq + PI * EffectiveDiskRadius * EffectiveDiskRadius);
+				//@todo - this is scale dependent
+				DistanceAttenuation = min(DistanceAttenuation, .0001f * RadialAttenuation(VPLToGatherPoint / VPLGatherRadius, 4));
+			#elif USE_PROJECTED_ANGLE_ATTENUATION
+				float TangentProjectedAngle = VPLPositionAndRadius.w / Distance;
+				float CosProjectedAngle = 1 / sqrt(1 + TangentProjectedAngle * TangentProjectedAngle);
+				float ProjectedSolidAngle = 2 * PI * (1 - CosProjectedAngle);
+				float DistanceAttenuation = ProjectedSolidAngle / (2 * PI);
+				DistanceAttenuation = min(DistanceAttenuation, RadialAttenuation(VPLToGatherPoint / VPLGatherRadius, 4));
+				DistanceAttenuation *= .0001f;
+			#else
+				float DistanceAttenuation = RadialAttenuation(VPLToGatherPoint / VPLGatherRadius, 8);
+				DistanceAttenuation *= .0001f;
+			#endif
+
+			float CosTheta = DirectionDot / Distance;
+			float SinTheta = sqrt(1 - CosTheta * CosTheta);
+			
+			#if IRRADIANCE_FROM_AO_CONES
+
+				float ShadowDepthBias = 0;
+				float ShadowingDistance = Distance + ShadowDepthBias;
+				float NormalizedDistance = saturate(ShadowingDistance / AOMaxDistance);
+				uint LowerStepIndex = (uint)min(NormalizedDistance * NUM_VISIBILITY_STEPS, NUM_VISIBILITY_STEPS - 1);
+				float LerpAlpha = ShadowingDistance - LowerStepIndex * VisibilityStepSize;
+				float3 InterpolatedBentNormal = lerp(StepBentNormal[LowerStepIndex], StepBentNormal[LowerStepIndex + 1], saturate(LerpAlpha));
+
+				float Shadow = GetVPLOcclusion(InterpolatedBentNormal, VPLDirection);
+
+			#else
+
+				float StartOffset = 1;
+				float EndOffset = 10;
+				float RayLength = max(Distance - StartOffset - EndOffset, 0);
+				float ConeEndDistance = max(RayLength - SinTheta * VPLPositionAndRadius.w, 0);
+				float RadiusAtConeEnd = CosTheta * VPLPositionAndRadius.w;
+				// Clamp the cone angle so that it doesn't intersect the gather point surface
+				float TanConeAngle = min(RadiusAtConeEnd / ConeEndDistance, SinTheta / CosTheta);
+				float Shadow = RayTraceThroughGlobalObjects(WorldPosition + StartOffset * VPLDirection, VPLDirection, RayLength, TanConeAngle, ConeEndDistance);
+
+				if (ConeEndDistance == 0)
+				{
+					//Shadow = 10;
+				}
+
+			#endif
+
+			float VPLCosineLobe = saturate(VPLNormalDot / Distance);
+			float3 VPLFlux = bCluster ? VPLClusterData[VPLIndex * VPL_DATA_STRIDE + 2].xyz : VPLData[VPLIndex * VPL_DATA_STRIDE + 2].xyz;
+
+			Irradiance += (saturate(CosTheta) * VPLCosineLobe * DistanceAttenuation * Shadow) * VPLFlux;
+		}
+	}
+
+	return Irradiance;
+}
+
+float3 GatherIrradianceFromVPLs(float3 WorldPosition, float3 WorldNormal, uint RelativeRecordIndex, uint RecordIndex, uint ThreadIndex)
+{
+	float3 TangentX;
+	float3 TangentY;
+	FindBestAxisVectors2(WorldNormal, TangentX, TangentY);
+
+	for (uint StepIndex = 0; StepIndex < NUM_VISIBILITY_STEPS; StepIndex++)
+	{
+		float3 UnoccludedDirection = 0;
+
+		for (uint ConeIndex = 0; ConeIndex < NUM_CONE_DIRECTIONS; ConeIndex++)
+		{
+			float3 ConeDirection = AOSamples2.SampleDirections[ConeIndex].xyz;
+			float3 RotatedConeDirection = ConeDirection.x * TangentX + ConeDirection.y * TangentY + ConeDirection.z * WorldNormal;
+
+			uint RecordConeDataIndex = (RelativeRecordIndex * NUM_CONE_DIRECTIONS + ConeIndex) * RECORD_CONE_DATA_STRIDE;
+			float ConeVisibility = RecordConeData[RecordConeDataIndex + StepIndex];
+			UnoccludedDirection += ConeVisibility * RotatedConeDirection;
+		}
+
+		float InvNumSamples = 1.0f / (float)NUM_CONE_DIRECTIONS;
+		UnoccludedDirection = UnoccludedDirection * (BentNormalNormalizeFactor * InvNumSamples);
+		StepBentNormal[StepIndex] = UnoccludedDirection;
+	}
+
+	float3 Irradiance = 0;
+
+#define USE_VPL_CLUSTERS 0
+#if USE_VPL_CLUSTERS
+
+	#define USE_TILE_CULLING 0
+	#if USE_TILE_CULLING
+
+		//@todo - won't be correct across frames
+		uint2 TileCoordinate = IrradianceCacheTileCoordinate[RecordIndex];
+		uint TileIndex = TileCoordinate.y * TileListGroupSize.x + TileCoordinate.x;
+		uint NumClusterVPLs = min(TileHeadDataUnpacked[TileIndex * 2 + 1], MAX_VPLS_PER_TILE);
+
+		LOOP
+		for (uint VPLListIndex = ThreadIndex; VPLListIndex < NumClusterVPLs; VPLListIndex += FINAL_GATHER_THREADGROUP_SIZE)
+		{
+			uint VPLClusterIndex = TileArrayData[TileIndex * MAX_VPLS_PER_TILE + VPLListIndex];
+	#else
+
+		uint NumClusterVPLs = VPLParameterBuffer[0];
+
+		LOOP
+		for (uint VPLClusterIndex = ThreadIndex; VPLClusterIndex < NumClusterVPLs; VPLClusterIndex += FINAL_GATHER_THREADGROUP_SIZE)
+		{
+
+	#endif
+			float4 ClusterPositionAndRadius = VPLClusterData[VPLClusterIndex * VPL_DATA_STRIDE + 0];
+			float4 ClusterNormalAndNumVPLs = VPLClusterData[VPLClusterIndex * VPL_DATA_STRIDE + 1];
+			float4 ClusterFluxAndVPLStartIndex = VPLClusterData[VPLClusterIndex * VPL_DATA_STRIDE + 2];
+
+			Irradiance += ComputeVirtualPointLighting(true, VPLClusterIndex, WorldPosition, WorldNormal);
+		/*
+			uint NumChildVPLs = ClusterNormalAndNumVPLs.w;
+
+			LOOP
+			for (uint ChildIndex = 0; ChildIndex < NumChildVPLs; ChildIndex++)
+			{
+				Irradiance += ComputeVirtualPointLighting(false, (uint)ClusterFluxAndVPLStartIndex.w + ChildIndex, WorldPosition, WorldNormal);
+			}*/
+		}
+
+#else
+	uint NumVPLs = VPLParameterBuffer[1];
+
+	LOOP
+	for (uint VPLIndex = ThreadIndex; VPLIndex < NumVPLs; VPLIndex += FINAL_GATHER_THREADGROUP_SIZE)
+	{
+		Irradiance += ComputeVirtualPointLighting(false, VPLIndex, WorldPosition, WorldNormal);
+	}
+#endif
+
+	//Irradiance = NumClusterVPLs / (float)MAX_VPLS_PER_TILE;
+
+	return Irradiance;
+}
+
+RWBuffer<float4> RWIrradianceCacheIrradiance;
+
+groupshared float3 SharedThreadIrradiance[FINAL_GATHER_THREADGROUP_SIZE];
+
+/**  */
+[numthreads(FINAL_GATHER_THREADGROUP_SIZE, 1, 1)]
+void ComputeIrradianceCS(
+	uint3 GroupId : SV_GroupID,
+	uint3 DispatchThreadId : SV_DispatchThreadID,
+    uint3 GroupThreadId : SV_GroupThreadID) 
+{
+	uint StartIndex = SavedStartIndex[0];
+	uint NumRecords = ScatterDrawParameters[1];
+
+	uint RecordIndex = StartIndex + GroupId.x;
+
+	if (RecordIndex < NumRecords)
+	{
+		float3 RecordWorldNormal = IrradianceCacheNormal[RecordIndex].xyz;
+
+		float3 DiffuseColor = .5f;
+
+		float4 PositionAndRadius = IrradianceCachePositionRadius[RecordIndex];
+		float3 RecordWorldPosition = PositionAndRadius.xyz;
+		float3 Irradiance = GatherIrradianceFromVPLs(RecordWorldPosition, RecordWorldNormal, GroupId.x, RecordIndex, GroupThreadId.x) * DiffuseColor;
+
+		SharedThreadIrradiance[GroupThreadId.x] = Irradiance;
+	}
+
+	GroupMemoryBarrierWithGroupSync();
+
+	if (GroupThreadId.x == 0)
+	{
+		float3 Irradiance = 0;
+
+		for (uint i = 0; i < FINAL_GATHER_THREADGROUP_SIZE; i++)
+		{
+			Irradiance += SharedThreadIrradiance[i];
+		}
+
+		RWIrradianceCacheIrradiance[RecordIndex] = float4(Irradiance, 0);
+	}
+}
diff --git a/Engine/Shaders/DistanceFieldLightingShared.usf b/Engine/Shaders/DistanceFieldLightingShared.usf
--- a/Engine/Shaders/DistanceFieldLightingShared.usf
+++ b/Engine/Shaders/DistanceFieldLightingShared.usf
@@ -14,10 +14,6 @@
 
 #define THREADGROUP_TOTALSIZE (THREADGROUP_SIZEX * THREADGROUP_SIZEY)
 
-#ifndef MAX_OBJECTS_PER_TILE
-#define MAX_OBJECTS_PER_TILE 1
-#endif
-
 #ifndef DOWNSAMPLE_FACTOR
 #define DOWNSAMPLE_FACTOR 1
 #endif
@@ -26,6 +22,9 @@
 #define UPDATEOBJECTS_THREADGROUP_SIZE 1
 #endif
 
+// Must match C++
+#define MAX_OBJECTS_PER_TILE 512
+
 float3 DistanceFieldVolumePositionToUV(float3 VolumePosition, float3 UVScale, float3 UVAdd)
 {
 	float3 VolumeUV = VolumePosition * UVScale + UVAdd;
@@ -46,7 +45,7 @@ uint GetCulledNumObjects()
 }
 
 // In float4's.  Must match equivalent C++ variables.
-#define OBJECT_DATA_STRIDE 12
+#define OBJECT_DATA_STRIDE 15
 
 uint NumSceneObjects;
 
@@ -58,8 +57,46 @@ Buffer<float> ObjectData;
 RWBuffer<float> RWObjectBounds;
 RWBuffer<float> RWObjectData;
 
+float4 LoadGlobalObjectPositionAndRadius(uint ObjectIndex)
+{
+	uint VectorIndex = ObjectIndex * 4;
+	return float4(ObjectBounds[VectorIndex + 0], ObjectBounds[VectorIndex + 1], ObjectBounds[VectorIndex + 2], ObjectBounds[VectorIndex + 3]);
+}
+
+float4x4 LoadGlobalObjectWorldToVolume(uint ObjectIndex)
+{
+	uint VectorIndex = (ObjectIndex * OBJECT_DATA_STRIDE + 0) * 4;
+	float4 M0 = float4(ObjectData[VectorIndex + 0], ObjectData[VectorIndex + 1], ObjectData[VectorIndex + 2], ObjectData[VectorIndex + 3]);
+	VectorIndex = (ObjectIndex * OBJECT_DATA_STRIDE + 1) * 4;
+	float4 M1 = float4(ObjectData[VectorIndex + 0], ObjectData[VectorIndex + 1], ObjectData[VectorIndex + 2], ObjectData[VectorIndex + 3]);
+	VectorIndex = (ObjectIndex * OBJECT_DATA_STRIDE + 2) * 4;
+	float4 M2 = float4(ObjectData[VectorIndex + 0], ObjectData[VectorIndex + 1], ObjectData[VectorIndex + 2], ObjectData[VectorIndex + 3]);
+	VectorIndex = (ObjectIndex * OBJECT_DATA_STRIDE + 3) * 4;
+	float4 M3 = float4(ObjectData[VectorIndex + 0], ObjectData[VectorIndex + 1], ObjectData[VectorIndex + 2], ObjectData[VectorIndex + 3]);
+
+	return float4x4(M0, M1, M2, M3);
+}
+ 
+float3 LoadGlobalObjectLocalPositionExtent(uint ObjectIndex)
+{
+	uint VectorIndex = (ObjectIndex * OBJECT_DATA_STRIDE + 4) * 4;
+	return float3(ObjectData[VectorIndex + 0], ObjectData[VectorIndex + 1], ObjectData[VectorIndex + 2]);
+}
+
+float4 LoadGlobalObjectUVScale(uint ObjectIndex)
+{
+	uint VectorIndex = (ObjectIndex * OBJECT_DATA_STRIDE + 5) * 4;
+	return float4(ObjectData[VectorIndex + 0], ObjectData[VectorIndex + 1], ObjectData[VectorIndex + 2], ObjectData[VectorIndex + 3]);
+}
+
+float3 LoadGlobalObjectUVAdd(uint ObjectIndex)
+{
+	uint VectorIndex = (ObjectIndex * OBJECT_DATA_STRIDE + 6) * 4;
+	return float3(ObjectData[VectorIndex + 0], ObjectData[VectorIndex + 1], ObjectData[VectorIndex + 2]);
+}
+
 // In float4's.  Must match equivalent C++ variables.
-#define CULLED_OBJECT_DATA_STRIDE 8
+#define CULLED_OBJECT_DATA_STRIDE 11
 #define CULLED_OBJECT_BOX_BOUNDS_STRIDE 5
 
 Buffer<float4> CulledObjectBounds;
@@ -80,7 +117,7 @@ float4x4 LoadObjectWorldToVolume(uint ObjectIndex)
 
 	return float4x4(M0, M1, M2, M3);
 }
-
+ 
 float3 LoadObjectLocalPositionExtent(uint ObjectIndex)
 {
 	return CulledObjectData.Load(ObjectIndex * CULLED_OBJECT_DATA_STRIDE + 4).xyz;
@@ -96,6 +133,15 @@ float3 LoadObjectUVAdd(uint ObjectIndex)
 	return CulledObjectData.Load(ObjectIndex * CULLED_OBJECT_DATA_STRIDE + 6).xyz;
 }
 
+float3x3 LoadObjectVolumeToWorld(uint ObjectIndex)
+{
+	float3 M0 = CulledObjectData.Load(ObjectIndex * CULLED_OBJECT_DATA_STRIDE + 8).xyz;
+	float3 M1 = CulledObjectData.Load(ObjectIndex * CULLED_OBJECT_DATA_STRIDE + 9).xyz;
+	float3 M2 = CulledObjectData.Load(ObjectIndex * CULLED_OBJECT_DATA_STRIDE + 10).xyz;
+
+	return float3x3(M0, M1, M2);
+}
+
 void LoadObjectViewSpaceBox(uint ObjectIndex, out float3 ObjectViewSpaceMin, out float3 ObjectViewSpaceMax)
 {
 	ObjectViewSpaceMin = CulledObjectBoxBounds.Load(ObjectIndex * CULLED_OBJECT_BOX_BOUNDS_STRIDE + 0).xyz;
@@ -107,4 +153,62 @@ void LoadObjectAxes(uint ObjectIndex, out float3 ObjectAxisX, out float3 ObjectA
 	ObjectAxisX = CulledObjectBoxBounds.Load(ObjectIndex * CULLED_OBJECT_BOX_BOUNDS_STRIDE + 2).xyz;
 	ObjectAxisY = CulledObjectBoxBounds.Load(ObjectIndex * CULLED_OBJECT_BOX_BOUNDS_STRIDE + 3).xyz;
 	ObjectAxisZ = CulledObjectBoxBounds.Load(ObjectIndex * CULLED_OBJECT_BOX_BOUNDS_STRIDE + 4).xyz;
+}
+
+Buffer<uint> ShadowTileHeadDataUnpacked;
+Buffer<uint> ShadowTileArrayData;
+uint2 ShadowTileListGroupSize;
+
+uint2 GetShadowTileHead(uint2 TileCoordinate)
+{
+	uint TileIndex = TileCoordinate.y * ShadowTileListGroupSize.x + TileCoordinate.x;
+
+	return uint2(
+		ShadowTileHeadDataUnpacked[TileIndex * 2 + 0], 
+		min(ShadowTileHeadDataUnpacked[TileIndex * 2 + 1], (uint)MAX_OBJECTS_PER_TILE)); 
+}
+
+void FindBestAxisVectors2(float3 InZAxis, out float3 OutXAxis, out float3 OutYAxis )
+{
+	float3 UpVector = abs(InZAxis.z) < 0.999 ? float3(0,0,1) : float3(1,0,0);
+	OutXAxis = normalize( cross( UpVector, InZAxis ) );
+	OutYAxis = cross( InZAxis, OutXAxis );
+}
+
+#define VPL_DATA_STRIDE 3
+
+#define FINAL_GATHER_THREADGROUP_SIZE 64
+
+#define NUM_VISIBILITY_STEPS 10
+
+// Must match C++
+#define RECORD_CONE_DATA_STRIDE NUM_VISIBILITY_STEPS
+
+Buffer<float4> IrradianceCachePositionRadius;
+Buffer<float> IrradianceCacheOccluderRadius;
+Buffer<uint2> IrradianceCacheTileCoordinate;
+Buffer<float4> IrradianceCacheNormal;
+Buffer<float4> IrradianceCacheBentNormal;
+Buffer<float4> IrradianceCacheIrradiance;
+
+Buffer<uint> ScatterDrawParameters;
+Buffer<uint> SavedStartIndex;
+
+uint NumConvexHullPlanes;
+float4 ViewFrustumConvexHull[6];
+
+bool ViewFrustumIntersectSphere(float3 SphereOrigin, float SphereRadius)
+{
+	for (uint PlaneIndex = 0; PlaneIndex < NumConvexHullPlanes; PlaneIndex++)
+	{
+		float4 PlaneData = ViewFrustumConvexHull[PlaneIndex];
+		float PlaneDistance = dot(PlaneData.xyz, SphereOrigin) - PlaneData.w;
+
+		if (PlaneDistance > SphereRadius)
+		{
+			return false;
+		}
+	}
+	
+	return true;
 }
\ No newline at end of file
diff --git a/Engine/Shaders/DistanceFieldShadowing.usf b/Engine/Shaders/DistanceFieldShadowing.usf
--- a/Engine/Shaders/DistanceFieldShadowing.usf
+++ b/Engine/Shaders/DistanceFieldShadowing.usf
@@ -27,6 +27,7 @@ void CopyCulledObjectData(uint DestIndex, uint SourceIndex)
 	UNROLL
 	for (uint VectorIndex = 0; VectorIndex < CULLED_OBJECT_DATA_STRIDE; VectorIndex++)
 	{
+		// Note: only copying the first CULLED_OBJECT_DATA_STRIDE of the original object data
 		RWCulledObjectData[DestIndex * CULLED_OBJECT_DATA_STRIDE + VectorIndex] = FetchObjectDataFloat4(SourceIndex * OBJECT_DATA_STRIDE + VectorIndex);
 	}
 
@@ -47,10 +48,10 @@ void CopyCulledObjectData(uint DestIndex, uint SourceIndex)
 	float3 MinViewSpacePosition = float3(2000000, 2000000, 2000000);
 	float3 MaxViewSpacePosition = float3(-2000000, -2000000, -2000000);
 
-	float4 M0 = FetchObjectDataFloat4(SourceIndex * OBJECT_DATA_STRIDE + 8);
-	float4 M1 = FetchObjectDataFloat4(SourceIndex * OBJECT_DATA_STRIDE + 9);
-	float4 M2 = FetchObjectDataFloat4(SourceIndex * OBJECT_DATA_STRIDE + 10);
-	float4 M3 = FetchObjectDataFloat4(SourceIndex * OBJECT_DATA_STRIDE + 11);
+	float4 M0 = FetchObjectDataFloat4(SourceIndex * OBJECT_DATA_STRIDE + 11);
+	float4 M1 = FetchObjectDataFloat4(SourceIndex * OBJECT_DATA_STRIDE + 12);
+	float4 M2 = FetchObjectDataFloat4(SourceIndex * OBJECT_DATA_STRIDE + 13);
+	float4 M3 = FetchObjectDataFloat4(SourceIndex * OBJECT_DATA_STRIDE + 14);
 
 	float4x4 LocalToWorld = float4x4(M0, M1, M2, M3);
 
@@ -111,12 +112,9 @@ void CullObjectsForShadowCS(
 
 	if (DispatchThreadId.x == 0)
 	{
+		// RWObjectIndirectArguments is zeroed by a clear before this shader, only need to set things that are non-zero (and are not read by this shader as that would be a race condition)
 		// IndexCount, NumInstances, StartIndex, BaseVertexIndex, FirstInstance
 		RWObjectIndirectArguments[0] = ObjectBoundingGeometryIndexCount;
-		RWObjectIndirectArguments[1] = 0;
-		RWObjectIndirectArguments[2] = 0;
-		RWObjectIndirectArguments[3] = 0;
-		RWObjectIndirectArguments[4] = 0;
 	}
 
 	GroupMemoryBarrierWithGroupSync();
@@ -143,9 +141,6 @@ void CullObjectsForShadowCS(
 		// IndexCount, NumInstances, StartIndex, BaseVertexIndex, FirstInstance
 		RWObjectIndirectArguments[0] = ObjectBoundingGeometryIndexCount;
 		RWObjectIndirectArguments[1] = NumSceneObjects;
-		RWObjectIndirectArguments[2] = 0;
-		RWObjectIndirectArguments[3] = 0;
-		RWObjectIndirectArguments[4] = 0;
 	}
 
 	uint SourceIndex = ObjectIndex;
@@ -163,8 +158,8 @@ bool RayHitSphere(float3 RayOrigin, float3 UnitRayDirection, float RayLength, fl
 	return dot(CenterToRay, CenterToRay) <= Square(SphereRadius) && DistanceAlongRay > -SphereRadius && DistanceAlongRay - SphereRadius < RayLength;
 }
 
-RWBuffer<uint> RWTileHeadDataUnpacked;
-RWBuffer<uint> RWTileArrayData;
+RWBuffer<uint> RWShadowTileHeadDataUnpacked;
+RWBuffer<uint> RWShadowTileArrayData;
 
 float2 NumGroups;
 
@@ -176,8 +171,8 @@ void ClearTilesCS(
 {
     uint TileIndex = DispatchThreadId.y * NumGroups.x + DispatchThreadId.x;
 	
-	RWTileHeadDataUnpacked[TileIndex * 2 + 0] = TileIndex;
-	RWTileHeadDataUnpacked[TileIndex * 2 + 1] = 0;
+	RWShadowTileHeadDataUnpacked[TileIndex * 2 + 0] = TileIndex;
+	RWShadowTileHeadDataUnpacked[TileIndex * 2 + 1] = 0;
 }
 
 struct FShadowObjectCullVertexOutput
@@ -275,12 +270,12 @@ void ShadowObjectCullPS(
 		if (all(MinProjections < 1) && all(MaxProjections > -1))
 		{
 			uint ArrayIndex;
-			InterlockedAdd(RWTileHeadDataUnpacked[TileIndex * 2 + 1], 1U, ArrayIndex);
+			InterlockedAdd(RWShadowTileHeadDataUnpacked[TileIndex * 2 + 1], 1U, ArrayIndex);
 
 			if (ArrayIndex < MAX_OBJECTS_PER_TILE)
 			{
 				uint DataIndex = (ArrayIndex * (uint)(NumGroups.x * NumGroups.y + .5f) + TileIndex);
-				RWTileArrayData[DataIndex] = Input.ObjectIndex;
+				RWShadowTileArrayData[DataIndex] = Input.ObjectIndex;
 			}
 		}
 	}
@@ -288,12 +283,12 @@ void ShadowObjectCullPS(
 #else
 	{
 		uint ArrayIndex;
-		InterlockedAdd(RWTileHeadDataUnpacked[TileIndex * 2 + 1], 1U, ArrayIndex);
+		InterlockedAdd(RWShadowTileHeadDataUnpacked[TileIndex * 2 + 1], 1U, ArrayIndex);
 
 		if (ArrayIndex < MAX_OBJECTS_PER_TILE)
 		{
 			uint DataIndex = (ArrayIndex * (uint)(NumGroups.x * NumGroups.y + .5f) + TileIndex);
-			RWTileArrayData[DataIndex] = Input.ObjectIndex;
+			RWShadowTileArrayData[DataIndex] = Input.ObjectIndex;
 		}
 	}
 #endif
@@ -308,23 +303,6 @@ float LightSourceRadius;
 float2 TanLightAngleAndNormalThreshold;
 int4 ScissorRectMinAndSize;
 
-#if SCATTER_TILE_CULLING
-
-Buffer<uint> TileHeadDataUnpacked;
-Buffer<uint> TileArrayData;
-uint2 TileListGroupSize;
-
-uint2 GetTileHead(uint2 TileCoordinate)
-{
-	uint TileIndex = TileCoordinate.y * TileListGroupSize.x + TileCoordinate.x;
-
-	return uint2(
-		TileHeadDataUnpacked[TileIndex * 2 + 0], 
-		min(TileHeadDataUnpacked[TileIndex * 2 + 1], (uint)MAX_OBJECTS_PER_TILE)); 
-}
-
-#else // SCATTER_TILE_CULLING
-
 /** Min and Max depth for this tile. */
 groupshared uint IntegerTileMinZ;
 groupshared uint IntegerTileMaxZ;
@@ -340,8 +318,6 @@ groupshared uint TileNumObjects1;
 #define MAX_INTERSECTING_OBJECTS 256
 groupshared uint IntersectingObjectIndices[MAX_INTERSECTING_OBJECTS * 2];
 
-#endif
-
 [numthreads(THREADGROUP_SIZEX, THREADGROUP_SIZEY, 1)]
 void DistanceFieldShadowingCS(
 	uint3 GroupId : SV_GroupID,
@@ -377,9 +353,9 @@ void DistanceFieldShadowingCS(
 	float2 NormalizedShadowPosition = HomogeneousShadowPosition.xy * .5f + .5f;
 	NormalizedShadowPosition.y = 1 - NormalizedShadowPosition.y;
 	// Quantize the shadow position to get our tile position
-	uint2 TilePosition = (uint2)(NormalizedShadowPosition * TileListGroupSize);
+	uint2 TilePosition = (uint2)(NormalizedShadowPosition * ShadowTileListGroupSize);
 	// Fetch the tile head information
-	uint2 TileHead = GetTileHead(TilePosition);
+	uint2 TileHead = GetShadowTileHead(TilePosition);
 	NumIntersectingObjects = TileHead.y;
 
 #else
@@ -577,7 +553,7 @@ void DistanceFieldShadowingCS(
 			{
 				#if USE_CULLING
 					#if SCATTER_TILE_CULLING
-						uint ObjectIndex = TileArrayData.Load(ListObjectIndex * TileListGroupSize.x * TileListGroupSize.y + TileHead.x);
+						uint ObjectIndex = ShadowTileArrayData.Load(ListObjectIndex * ShadowTileListGroupSize.x * ShadowTileListGroupSize.y + TileHead.x);
 					#else
 						uint ObjectIndex = IntersectingObjectIndices[MAX_INTERSECTING_OBJECTS * GroupIndex + ListObjectIndex];
 					#endif
diff --git a/Engine/Shaders/DistanceFieldSurfaceCacheLighting.usf b/Engine/Shaders/DistanceFieldSurfaceCacheLighting.usf
--- a/Engine/Shaders/DistanceFieldSurfaceCacheLighting.usf
+++ b/Engine/Shaders/DistanceFieldSurfaceCacheLighting.usf
@@ -116,7 +116,7 @@ void ComputeDistanceFieldNormalPS(
 	float3 WorldNormal = GBufferData.WorldNormal;
 	float DepthSign = GBufferData.HasDistanceFieldRepresentation > 0 ? 1 : -1;
 	// Note: Encoding whether the pixel has a distance field representation in the sign bit of the alpha
-	OutColor = float4(EncodeNormalForAO(WorldNormal), DepthSign * SceneDepth);
+	OutColor = float4(WorldNormal, DepthSign * SceneDepth);
 }
 
 RWTexture2D<float4> RWDistanceFieldNormal;
@@ -133,7 +133,7 @@ void ComputeDistanceFieldNormalCS(
 
 	float3 WorldNormal = GBufferData.WorldNormal;
 	float DepthSign = GBufferData.HasDistanceFieldRepresentation > 0 ? 1 : -1;
-	float4 OutValue = float4(EncodeNormalForAO(WorldNormal), DepthSign * SceneDepth);
+	float4 OutValue = float4(WorldNormal, DepthSign * SceneDepth);
 	RWDistanceFieldNormal[DispatchThreadId.xy] = OutValue;
 }
 
@@ -170,11 +170,6 @@ struct FIrradianceCacheSplatVertexOutput
 
 Buffer<uint> CompactedRedirect;
 Buffer<uint> IrradianceCacheUsed;
-Buffer<float4> IrradianceCachePositionRadius;
-Buffer<float> IrradianceCacheOccluderRadius;
-Buffer<float4> IrradianceCacheNormal;
-Buffer<float4> IrradianceCacheBentNormal;
-Buffer<float4> IrradianceCacheIrradiance;
 float InterpolationRadiusScale;
 float2 NormalizedOffsetToPixelCenter;
 float HackExpand;
@@ -227,7 +222,7 @@ void IrradianceCacheSplatVS(
 	VirtualRadius += HackExpand * ViewToCenterLength * 4;
 #endif
 
-	float3 RecordNormal = DecodeNormalForAO(IrradianceCacheNormal.Load(JobIndex).xyz);
+	float3 RecordNormal = IrradianceCacheNormal.Load(JobIndex).xyz;
 
 	float2 CornerUVs = InUV;
 	float3 CornerPosition = VirtualPosition + (View.ViewRight * CornerUVs.x + View.ViewUp * CornerUVs.y) * VirtualRadius;
@@ -245,9 +240,11 @@ void IrradianceCacheSplatVS(
 	Output.Normal = RecordNormal;
 
 #if FINAL_INTERPOLATION_PASS
-	Output.BentNormal = DecodeNormalForAO(IrradianceCacheBentNormal.Load(JobIndex).xyz);
+	Output.BentNormal = IrradianceCacheBentNormal.Load(JobIndex).xyz;
 	#if SUPPORT_IRRADIANCE
-		Output.Irradiance = IrradianceCacheIrradiance.Load(JobIndex).xyz;
+		float ViewDistance = length(View.ViewOrigin.xyz - PositionAndRadius.xyz);
+		float FadeAlpha = saturate(.001f * (.95f * AOMaxViewDistance - ViewDistance));
+		Output.Irradiance = IrradianceCacheIrradiance.Load(JobIndex).xyz * FadeAlpha;
 	#endif
 #endif
 } 
@@ -295,7 +292,7 @@ void IrradianceCacheSplatPS(
 	BRANCH
 	if (DistanceError < 1 && SceneW < AOMaxViewDistance)
 	{ 
-		float3 WorldNormal = DecodeNormalForAO(DistanceFieldNormalValue.xyz);
+		float3 WorldNormal = DistanceFieldNormalValue.xyz;
 		float NormalError = InterpolationAngleNormalization * sqrt(saturate(1 - dot(WorldNormal, Input.Normal)));
 
 		// Don't use a lighting record if it's in front of the query point.
@@ -589,6 +586,7 @@ void UpdateHistoryPS(
 	}
 
 	OutBentNormal.rgb = lerp(NewValue.rgb, HistoryValue.rgb, EffectiveHistoryWeight);
+	//OutBentNormal.rgb = HistoryValue.rgb + (NewValue.rgb - HistoryValue.rgb) * (1 - exp(-DeltaTime * EffectiveHistoryWeight));
 
 #if SUPPORT_IRRADIANCE
 	float3 NewIrradiance = Texture2DSampleLevel(IrradianceTexture, IrradianceSampler, DistanceFieldUVs, 0).xyz;
@@ -612,6 +610,7 @@ void AOUpsamplePS(
 {
 	// Distance field AO was computed at 0,0 regardless of viewrect min
 	float2 DistanceFieldUVs = UVAndScreenPos.xy - View.ViewRectMin.xy * View.ViewSizeAndSceneTexelSize.zw;
+	float3 Irradiance = 0;
 
 #define BILATERAL_UPSAMPLE 1
 #if BILATERAL_UPSAMPLE
@@ -655,7 +654,7 @@ void AOUpsamplePS(
 		float4 IrradianceValues01 = Texture2DSampleLevel(IrradianceTexture, IrradianceSampler, Corner00UV + float2(0, LowResTexelSize.y), 0);
 		float4 IrradianceValues11 = Texture2DSampleLevel(IrradianceTexture, IrradianceSampler, Corner00UV + LowResTexelSize, 0);
 
-		float3 Irradiance = 
+		Irradiance = 
 			(FinalWeights.x * IrradianceValues00.xyz 
 			+ FinalWeights.y * IrradianceValues10.xyz
 			+ FinalWeights.z * IrradianceValues01.xyz 
@@ -665,7 +664,10 @@ void AOUpsamplePS(
 
 #else
 	float3 BentNormal = Texture2DSampleLevel(BentNormalAOTexture, BentNormalAOSampler, DistanceFieldUVs, 0).xyz;
-	float3 Irradiance = Texture2DSampleLevel(IrradianceTexture, IrradianceSampler, DistanceFieldUVs, 0).xyz;
+
+	#if SUPPORT_IRRADIANCE
+		Irradiance = Texture2DSampleLevel(IrradianceTexture, IrradianceSampler, DistanceFieldUVs, 0).xyz;
+	#endif
 #endif
 
 #if OUTPUT_BENT_NORMAL
@@ -674,7 +676,13 @@ void AOUpsamplePS(
 		OutIrradiance = float4(Irradiance, 1);
 	#endif
 #else
-	OutBentNormal = float4(length(BentNormal).xxx, 1);
+
+	// Output to RT0 to visualize Irradiance or BentNormal
+	#if SUPPORT_IRRADIANCE
+		OutBentNormal = float4(Irradiance, 1);
+	#else
+		OutBentNormal = float4(length(BentNormal).xxx, 1);
+	#endif
 #endif
 }
 
@@ -692,7 +700,7 @@ groupshared uint IntersectingObjectIndices[MAX_INTERSECTING_OBJECTS];
 
 groupshared uint NumIntersectingObjects;
 
-void RayTraceThroughDistanceFields(float3 WorldRayStart, float3 WorldRayEnd, float MaxRayTime, out float MinRayTime, out float TotalStepsTaken)
+void RayTraceThroughTileCulledDistanceFields(float3 WorldRayStart, float3 WorldRayEnd, float MaxRayTime, out float MinRayTime, out float TotalStepsTaken)
 {
 	MinRayTime = MaxRayTime;
 	TotalStepsTaken = 0;
@@ -801,7 +809,11 @@ void VisualizeMeshDistanceFieldCS(
 		{
 			uint ListIndex;
 			InterlockedAdd(NumIntersectingObjects, 1U, ListIndex);
-			IntersectingObjectIndices[ListIndex] = ObjectIndex; 
+
+			if (ListIndex < MAX_INTERSECTING_OBJECTS)
+			{
+				IntersectingObjectIndices[ListIndex] = ObjectIndex; 
+			}
 		}
 	}
 
@@ -811,13 +823,13 @@ void VisualizeMeshDistanceFieldCS(
 	float TotalStepsTaken;
 
 	// Trace once to find the distance to first intersection
-	RayTraceThroughDistanceFields(WorldRayStart, WorldRayEnd, TraceDistance, MinRayTime, TotalStepsTaken);
+	RayTraceThroughTileCulledDistanceFields(WorldRayStart, WorldRayEnd, TraceDistance, MinRayTime, TotalStepsTaken);
 
 	float TempMinRayTime;
 	// Recompute the ray end point
 	WorldRayEnd = WorldRayStart + UnitWorldRayDirection * MinRayTime;
 	// Trace a second time to only accumulate steps taken before the first intersection, improves visualization
-	RayTraceThroughDistanceFields(WorldRayStart, WorldRayEnd, MinRayTime, TempMinRayTime, TotalStepsTaken);
+	RayTraceThroughTileCulledDistanceFields(WorldRayStart, WorldRayEnd, MinRayTime, TempMinRayTime, TotalStepsTaken);
 
 	float3 Result = saturate(TotalStepsTaken / 200.0f);
 
diff --git a/Engine/Shaders/DistanceFieldSurfaceCacheLightingCompute.usf b/Engine/Shaders/DistanceFieldSurfaceCacheLightingCompute.usf
--- a/Engine/Shaders/DistanceFieldSurfaceCacheLightingCompute.usf
+++ b/Engine/Shaders/DistanceFieldSurfaceCacheLightingCompute.usf
@@ -141,26 +141,8 @@ void RemoveObjectsFromBufferCS(
 RWBuffer<float4> RWCulledObjectBounds;
 RWBuffer<float4> RWCulledObjectData;
 RWBuffer<float4> RWCulledObjectBoxBounds;
-uint NumConvexHullPlanes;
-float4 ViewFrustumConvexHull[6];
 uint ObjectBoundingGeometryIndexCount;
 
-bool ViewFrustumIntersectSphere(float3 SphereOrigin, float SphereRadius)
-{
-	for (uint PlaneIndex = 0; PlaneIndex < NumConvexHullPlanes; PlaneIndex++)
-	{
-		float4 PlaneData = ViewFrustumConvexHull[PlaneIndex];
-		float PlaneDistance = dot(PlaneData.xyz, SphereOrigin) - PlaneData.w;
-
-		if (PlaneDistance > SphereRadius)
-		{
-			return false;
-		}
-	}
-	
-	return true;
-}
-
 groupshared uint NumGroupObjects;
 groupshared uint GroupBaseIndex;
 groupshared uint GroupObjectIndices[UPDATEOBJECTS_THREADGROUP_SIZE];
@@ -177,6 +159,7 @@ void CopyCulledObjectData(uint DestIndex, uint SourceIndex)
 	UNROLL
 	for (uint VectorIndex = 0; VectorIndex < CULLED_OBJECT_DATA_STRIDE; VectorIndex++)
 	{
+		// Note: only copying the first CULLED_OBJECT_DATA_STRIDE of the original object data
 		RWCulledObjectData[DestIndex * CULLED_OBJECT_DATA_STRIDE + VectorIndex] = FetchObjectDataFloat4(SourceIndex * OBJECT_DATA_STRIDE + VectorIndex);
 	}
 }
@@ -193,12 +176,9 @@ void CullObjectsForViewCS(
 #if USE_FRUSTUM_CULLING
 	if (DispatchThreadId.x == 0)
 	{
+		// RWObjectIndirectArguments is zeroed by a clear before this shader, only need to set things that are non-zero (and are not read by this shader as that would be a race condition)
 		// IndexCount, NumInstances, StartIndex, BaseVertexIndex, FirstInstance
 		RWObjectIndirectArguments[0] = ObjectBoundingGeometryIndexCount;
-		RWObjectIndirectArguments[1] = 0;
-		RWObjectIndirectArguments[2] = 0;
-		RWObjectIndirectArguments[3] = 0;
-		RWObjectIndirectArguments[4] = 0;
 	}
 
 	if (GroupThreadId.x == 0)
@@ -247,9 +227,6 @@ void CullObjectsForViewCS(
 		// IndexCount, NumInstances, StartIndex, BaseVertexIndex, FirstInstance
 		RWObjectIndirectArguments[0] = ObjectBoundingGeometryIndexCount;
 		RWObjectIndirectArguments[1] = NumSceneObjects;
-		RWObjectIndirectArguments[2] = 0;
-		RWObjectIndirectArguments[3] = 0;
-		RWObjectIndirectArguments[4] = 0;
 	}
 
 	GroupMemoryBarrierWithGroupSync();
@@ -371,7 +348,6 @@ void BuildTileConesMain(
 			ConeAxisDepthRanges.w = ConeAxisDistanceMultiply * (MaxTileZ + ConeExpandDistance) + ConeAxisDistanceAdd;
 
 			// Pull back cone vertex to contain potential samples
-			//@todo - only expand in sky direction
 			TileConeVertex = float3(0, 0, 0) - TileConeAxis * VertexPullbackLength;
 		}
 
@@ -603,13 +579,6 @@ RWBuffer<uint2> RWIrradianceCacheTileCoordinate;
 Texture2D IrradianceCacheSplatTexture;
 SamplerState IrradianceCacheSplatSampler;
 
-void FindBestAxisVectors2(float3 InZAxis, out float3 OutXAxis, out float3 OutYAxis )
-{
-	float3 UpVector = abs(InZAxis.z) < 0.999 ? float3(0,0,1) : float3(1,0,0);
-	OutXAxis = normalize( cross( UpVector, InZAxis ) );
-	OutYAxis = cross( InZAxis, OutXAxis );
-}
-
 groupshared float4 CachedIrradianceCachePositionRadius[THREADGROUP_TOTALSIZE];
 groupshared float4 CachedIrradianceCacheNormal[THREADGROUP_TOTALSIZE];
 groupshared uint2 CachedIrradianceCacheTileCoordinate[THREADGROUP_TOTALSIZE];
@@ -656,7 +625,7 @@ void PopulateCacheCS(
 			float3 OpaqueWorldPosition = HomogeneousWorldPosition.xyz / HomogeneousWorldPosition.w;
 
 			float2 BaseLevelScreenUV = (DispatchThreadId.xy * DownsampleFactorToBaseLevel + float2(.5f, .5f)) * BaseLevelTexelSize;
-			float3 WorldNormal = DecodeNormalForAO(Texture2DSampleLevel(DistanceFieldNormalTexture, DistanceFieldNormalSampler, BaseLevelScreenUV, 0).xyz);
+			float3 WorldNormal = Texture2DSampleLevel(DistanceFieldNormalTexture, DistanceFieldNormalSampler, BaseLevelScreenUV, 0).xyz;
 
 			//@todo - offset shading position along normal to avoid incorrect self-occlusion?
 			float3 WorldShadingPosition = OpaqueWorldPosition;
@@ -692,7 +661,7 @@ void PopulateCacheCS(
 	{
 		int SampleIndex = BaseRecordIndex + LocalRecordIndex;
 		RWIrradianceCachePositionRadius[SampleIndex] = CachedIrradianceCachePositionRadius[LocalRecordIndex];
-		RWIrradianceCacheNormal[SampleIndex] = float4(EncodeNormalForAO(CachedIrradianceCacheNormal[LocalRecordIndex].xyz), 0);
+		RWIrradianceCacheNormal[SampleIndex] = float4(CachedIrradianceCacheNormal[LocalRecordIndex].xyz, 0);
 		uint2 TileCoordinate = CachedIrradianceCacheTileCoordinate[LocalRecordIndex];
 		RWIrradianceCacheTileCoordinate[SampleIndex] = TileCoordinate;
 	}
@@ -712,31 +681,23 @@ float BentNormalNormalizeFactor;
 float TanConeHalfAngle;
 float RecordRadiusScale;
 
-Buffer<float4> IrradianceCachePositionRadius;
-Buffer<uint2> IrradianceCacheTileCoordinate;
-Buffer<float4> IrradianceCacheNormal;
-Buffer<uint> ScatterDrawParameters;
-Buffer<uint> SavedStartIndex;
-
 RWBuffer<float> RWOccluderRadius;
 RWBuffer<float> RWRecordConeVisibility;
-RWBuffer<float4> RWIrradianceCacheIrradiance;
+RWBuffer<float> RWRecordConeData;
 RWBuffer<float4> RWDebugBuffer;
 
 // Have to disable surface caching dependencies for this to work
 #define VISUALIZE_ONE_CONE 0
 
-#define FINAL_GATHER_THREADGROUP_SIZE 64
-
 // Enforce one thread per cone direction
 #define SIMULTANEOUSLY_TRACED_OBJECTS (FINAL_GATHER_THREADGROUP_SIZE / NUM_CONE_DIRECTIONS)
 
 #define THREADS_PER_OBJECT (FINAL_GATHER_THREADGROUP_SIZE / SIMULTANEOUSLY_TRACED_OBJECTS)
 
 groupshared uint SharedConeVisibility[NUM_CONE_DIRECTIONS];
+groupshared uint SharedConeRawVisibility[NUM_CONE_DIRECTIONS][NUM_VISIBILITY_STEPS];
 groupshared uint SharedMinOcclusionDistance;
 
-groupshared float4 SharedConeDistanceAtClosestToOccluder[FINAL_GATHER_THREADGROUP_SIZE];
 groupshared float3 SharedGatheredLighting[FINAL_GATHER_THREADGROUP_SIZE];
 
 /** Computes ambient occlusion for a surface cache record by cone stepping through the nearby object distance fields. */
@@ -754,6 +715,14 @@ void FinalGatherCS(
 		for (uint ConeIndex = 0; ConeIndex < NUM_CONE_DIRECTIONS; ConeIndex++) 
 		{
 			SharedConeVisibility[ConeIndex] = asuint(1.0f);
+
+			#if SUPPORT_IRRADIANCE
+				UNROLL
+				for (uint i = 0; i < NUM_VISIBILITY_STEPS; i++)
+				{
+					SharedConeRawVisibility[ConeIndex][i] = asuint(1.0f);
+				}
+			#endif
 		}
 
 		SharedMinOcclusionDistance = asuint(AOMaxDistance);
@@ -771,7 +740,7 @@ void FinalGatherCS(
 	float3 WorldNormal;
 
 	{
-		WorldNormal = DecodeNormalForAO(IrradianceCacheNormal[RecordIndex].xyz);
+		WorldNormal = IrradianceCacheNormal[RecordIndex].xyz;
 		float3 WorldShadingPosition = IrradianceCachePositionRadius[RecordIndex].xyz;
 		uint2 TileCoordinate = IrradianceCacheTileCoordinate[RecordIndex];
 
@@ -792,12 +761,19 @@ void FinalGatherCS(
 #endif
 
 		float MinVisibility = 1;
-		float SecondMinVisibility = 1;
+		float MinRawVisibility[NUM_VISIBILITY_STEPS];
 		float MinWorldDistanceToOccluder = AOMaxDistance;
 		float ConeDistanceAtClosestToOccluder = 0;
-		float ConeDistanceAtSecondClosestToOccluder = 0;
 		float MaxWorldStepOffset = GetStepOffset(NUM_CONE_STEPS);
 
+#if SUPPORT_IRRADIANCE
+		UNROLL
+		for (uint i = 0; i < NUM_VISIBILITY_STEPS; i++)
+		{
+			MinRawVisibility[i] = 1;
+		}
+#endif
+
 		LOOP 
 		for (uint ListObjectIndex = 0; ListObjectIndex < NumObjectsAffectingTile; ListObjectIndex += SIMULTANEOUSLY_TRACED_OBJECTS)
 		{
@@ -873,36 +849,26 @@ void FinalGatherCS(
 						// Derive visibility from 1d intersection
 						float Visibility = saturate(DistanceToOccluder / ShadingSphereRadius);
 
-						BRANCH
-						if (Visibility < MinVisibility)
+						// Don't allow small objects to fully occlude a cone step
+						Visibility = max(Visibility, 1 - saturate(ObjectOccluderRadius / SphereRadius));
+						float OccluderDistanceFraction = (WorldStepOffset + DistanceToOccluder) / AOMaxDistance;
+
+						#if SUPPORT_IRRADIANCE
+							uint VisibilityIndex = NUM_VISIBILITY_STEPS * WorldStepOffset / AOMaxDistance;
+							// Track raw visibility before the distance fade for GI shadowing
+							MinRawVisibility[VisibilityIndex] = min(MinRawVisibility[VisibilityIndex], Visibility);
+						#endif
+
+						// Fade out occlusion based on distance to occluder to avoid a discontinuity at the max AO distance
+						Visibility = max(Visibility, saturate(OccluderDistanceFraction * OccluderDistanceFraction * .8f));
+
+						MinVisibility = min(MinVisibility, Visibility);
+
+						if (DistanceToOccluder < .9f * SphereRadius)
 						{
-							// Don't allow small objects to fully occlude a cone step
-							Visibility = max(Visibility, 1 - saturate(ObjectOccluderRadius / SphereRadius));
-							float OccluderDistanceFraction = (WorldStepOffset + DistanceToOccluder) / AOMaxDistance;
-							// Fade out occlusion based on distance to occluder to avoid a discontinuity at the max AO distance
-							Visibility = max(Visibility, saturate(OccluderDistanceFraction * OccluderDistanceFraction * .8f));
-
-							if (Visibility < MinVisibility)
-							{
-								ConeDistanceAtSecondClosestToOccluder = ConeDistanceAtClosestToOccluder;
-								ConeDistanceAtClosestToOccluder = WorldStepOffset;
-								SecondMinVisibility = MinVisibility;
-								MinVisibility = Visibility;
-							}
-							else if (Visibility < SecondMinVisibility)
-							{
-								ConeDistanceAtSecondClosestToOccluder = WorldStepOffset;
-								SecondMinVisibility = Visibility;
-							}
-
-							//MinVisibility = min(MinVisibility, Visibility);
-
-							if (DistanceToOccluder < .9f * SphereRadius)
-							{
-								// Assuming occluder is straight forward along the cone
-								float WorldDistanceToOccluder = WorldStepOffset + DistanceToOccluder;
-								MinWorldDistanceToOccluder = min(MinWorldDistanceToOccluder, WorldDistanceToOccluder);
-							}
+							// Assuming occluder is straight forward along the cone
+							float WorldDistanceToOccluder = WorldStepOffset + DistanceToOccluder;
+							MinWorldDistanceToOccluder = min(MinWorldDistanceToOccluder, WorldDistanceToOccluder);
 						}
 						
 						float MinStepSize = .6f * (GetStepOffset(StepIndex + 1) - GetStepOffset(StepIndex));
@@ -915,78 +881,31 @@ void FinalGatherCS(
 		InterlockedMin(SharedConeVisibility[ConeIndex], asuint(MinVisibility));
 		InterlockedMin(SharedMinOcclusionDistance, asuint(max(MinWorldDistanceToOccluder, 0)));
 
-		SharedConeDistanceAtClosestToOccluder[ThreadIndex] = float4(MinVisibility, ConeDistanceAtClosestToOccluder, SecondMinVisibility, ConeDistanceAtSecondClosestToOccluder);
-
-		GroupMemoryBarrierWithGroupSync();
-
 #if SUPPORT_IRRADIANCE
-		float ConeVisibility = asfloat(SharedConeVisibility[ConeIndex]);
-
-		float CompactedConeDistanceAtClosestToOccluder = 0;
-		float SecondConeMinVisibility = 1;
-
-		for (uint CompactThreadIndex = ConeIndex; CompactThreadIndex < FINAL_GATHER_THREADGROUP_SIZE; CompactThreadIndex += THREADS_PER_OBJECT)
-		{
-			if (SharedConeDistanceAtClosestToOccluder[CompactThreadIndex].x == ConeVisibility)
-			{
-				CompactedConeDistanceAtClosestToOccluder = SharedConeDistanceAtClosestToOccluder[CompactThreadIndex].y;
-			}
-			else
-			{
-				SecondConeMinVisibility = min(SecondConeMinVisibility, min(SharedConeDistanceAtClosestToOccluder[CompactThreadIndex].x, SharedConeDistanceAtClosestToOccluder[CompactThreadIndex].z));
-			}
-		}
 
-		float CompactedConeDistanceAtSecondClosestToOccluder = 0;
-
-		for (uint CompactThreadIndex = ConeIndex; CompactThreadIndex < FINAL_GATHER_THREADGROUP_SIZE; CompactThreadIndex += THREADS_PER_OBJECT)
+		UNROLL
+		for (uint i = 0; i < NUM_VISIBILITY_STEPS; i++)
 		{
-			if (SharedConeDistanceAtClosestToOccluder[CompactThreadIndex].x == SecondConeMinVisibility)
-			{
-				CompactedConeDistanceAtSecondClosestToOccluder = SharedConeDistanceAtClosestToOccluder[CompactThreadIndex].y;
-			}
-			else if (SharedConeDistanceAtClosestToOccluder[CompactThreadIndex].z == SecondConeMinVisibility)
-			{
-				CompactedConeDistanceAtSecondClosestToOccluder = SharedConeDistanceAtClosestToOccluder[CompactThreadIndex].w;
-			}
+			InterlockedMin(SharedConeRawVisibility[ConeIndex][i], asuint(MinRawVisibility[i]));
 		}
 
-		float SphereRadius = CompactedConeDistanceAtClosestToOccluder * TanConeHalfAngle;
-		float3 ConePositionAtClosestToOccluder = WorldShadingPosition + RotatedConeDirection * CompactedConeDistanceAtClosestToOccluder;
-		float4 ClipPosition = mul(float4(ConePositionAtClosestToOccluder, 1), View.WorldToClip);
-		float2 ScreenUV = ScreenAlignedPosition(ClipPosition).xy;
-		float UVDiskRadius = SphereRadius / ClipPosition.w * View.ViewToClip[0][0] * .5f;
-
-		uint NumSamplesPerThread = 9;
-		uint NumSamplesPerCone = NumSamplesPerThread * (FINAL_GATHER_THREADGROUP_SIZE / NUM_CONE_DIRECTIONS);
-		uint BaseSampleIndex = (ThreadIndex / NUM_CONE_DIRECTIONS) * NumSamplesPerThread;
-
-		float3 GatheredColor = 0;
-		float GatheredWeight = .0001f;
+		GroupMemoryBarrierWithGroupSync();
 
-		for (uint SampleIndex = 0; SampleIndex < NumSamplesPerThread; SampleIndex++)
+		// Only need one thread per cone direction to write
+		if (ThreadIndex < NUM_CONE_DIRECTIONS)
 		{
-			float2 Random = Hammersley(SampleIndex + BaseSampleIndex, NumSamplesPerCone, uint2(0, 0));
-			float2 UVOffset = UniformSampleDisk(Random);
-			float2 SampleUV = ScreenUV + UVOffset * UVDiskRadius;
-			float3 SampleSceneColor = CalcSceneColor(SampleUV);
-			float SampleSceneDepth = CalcSceneDepth(SampleUV);
-			float Weight = 1 - saturate(abs(SampleSceneDepth - ClipPosition.w) / SphereRadius);
-			GatheredColor += SampleSceneColor * Weight;
-			GatheredWeight += Weight;
-		} 
-		
-		SharedGatheredLighting[ThreadIndex] = GatheredColor / GatheredWeight;
+			uint RecordConeDataIndex = (GroupId.x * NUM_CONE_DIRECTIONS + ConeIndex) * RECORD_CONE_DATA_STRIDE;
 
-		if (ThreadIndex < THREADS_PER_OBJECT)
-		{
-			uint DebugIndex = ThreadIndex % THREADS_PER_OBJECT;
-			RWDebugBuffer[0] = float4(WorldShadingPosition, 0);
-			RWDebugBuffer[1 + DebugIndex * 3 + 0] = float4(ConePositionAtClosestToOccluder, SphereRadius);
+			float MinStepVisibility = 1;
 
-			float SphereRadius2 = CompactedConeDistanceAtSecondClosestToOccluder * TanConeHalfAngle;
-			float3 ConePositionAtClosestToOccluder2 = WorldShadingPosition + RotatedConeDirection * CompactedConeDistanceAtSecondClosestToOccluder;
-			RWDebugBuffer[1 + DebugIndex * 3 + 1] = float4(ConePositionAtClosestToOccluder2, SphereRadius2);
+			UNROLL
+			for (uint i = 0; i < NUM_VISIBILITY_STEPS; i++)
+			{
+				float StepVisibility = asfloat(SharedConeRawVisibility[ConeIndex][i]);
+				// Propagate min visibility down the cone
+				MinStepVisibility = min(MinStepVisibility, StepVisibility);
+				RWRecordConeData[RecordConeDataIndex + i] = MinStepVisibility;
+			}
 		}
 #endif
 	}
@@ -995,48 +914,14 @@ void FinalGatherCS(
 
 	if (ThreadIndex == 0)
 	{
-		float3 UnoccludedDirection = 0;
-		float Visibility = 0;
-		float3 Irradiance = 0;
-
 		for (uint ConeIndex = 0; ConeIndex < NUM_CONE_DIRECTIONS; ConeIndex++)
 		{
-			float3 ConeDirection = AOSamples2.SampleDirections[ConeIndex].xyz;
-			float3 RotatedConeDirection = ConeDirection.x * TangentX + ConeDirection.y * TangentY + ConeDirection.z * WorldNormal;
-
 			float ConeVisibility = asfloat(SharedConeVisibility[ConeIndex]);
 
-			RWRecordConeVisibility[GroupId.x * 9 + ConeIndex] = ConeVisibility;
-
-			UnoccludedDirection += ConeVisibility * RotatedConeDirection;
-			Visibility += ConeVisibility;
-
-#if SUPPORT_IRRADIANCE
-			float3 AccumulatedIrradiance = 0;			
-
-			for (uint CompactThreadIndex = ConeIndex; CompactThreadIndex < FINAL_GATHER_THREADGROUP_SIZE; CompactThreadIndex += THREADS_PER_OBJECT)
-			{
-				AccumulatedIrradiance += SharedGatheredLighting[CompactThreadIndex];
-			}
-
-			AccumulatedIrradiance /= (float)SIMULTANEOUSLY_TRACED_OBJECTS;
-
-			RWDebugBuffer[1 + ConeIndex * 3 + 2] = float4(AccumulatedIrradiance, ConeVisibility);
-
-			Irradiance += AccumulatedIrradiance * (1 - ConeVisibility);
-#endif
+			RWRecordConeVisibility[GroupId.x * NUM_CONE_DIRECTIONS + ConeIndex] = ConeVisibility;
 		}
 
-		float InvNumSamples = 1.0f  / (float)NUM_CONE_DIRECTIONS;
-		UnoccludedDirection = UnoccludedDirection * (BentNormalNormalizeFactor * InvNumSamples);
-		Visibility = Visibility * InvNumSamples;
-		Irradiance = Irradiance * InvNumSamples;
-
 		RWOccluderRadius[RecordIndex] = RecordRadiusScale * asfloat(SharedMinOcclusionDistance);
-
-#if SUPPORT_IRRADIANCE
-		RWIrradianceCacheIrradiance[RecordIndex] = float4(Irradiance, 0);
-#endif
 	}
 }
 
@@ -1094,7 +979,7 @@ void CalculateHeightfieldOcclusionCS(
 	BRANCH
 	if (RecordIndex < NumRecords && PositionAndRadius.w > 0)
 	{
-		float3 WorldNormal = DecodeNormalForAO(IrradianceCacheNormal[RecordIndex].xyz);
+		float3 WorldNormal = IrradianceCacheNormal[RecordIndex].xyz;
 		float3 WorldPosition = PositionAndRadius.xyz;
 
 		float3 TangentX;
@@ -1129,7 +1014,7 @@ void CalculateHeightfieldOcclusionCS(
 					float3 ConeDirection = AOSamples2.SampleDirections[ConeIndex].xyz;
 					float3 RotatedConeDirection = ConeDirection.x * TangentX + ConeDirection.y * TangentY + ConeDirection.z * WorldNormal;
 
-					float ExistingConeVisibility = RWRecordConeVisibility[DispatchThreadId.x * 9 + ConeIndex];
+					float ExistingConeVisibility = RWRecordConeVisibility[DispatchThreadId.x * NUM_CONE_DIRECTIONS + ConeIndex];
 				
 					float ZBias = 10;
 					float DistanceAlongConeToOccluder = clamp((WorldPosition.z + ZBias - WorldHeightfieldShadingPosition.z) / max(-RotatedConeDirection.z, .0001f), .00001f, AOMaxDistance);
@@ -1147,7 +1032,7 @@ void CalculateHeightfieldOcclusionCS(
 					MinDistanceToOccluder = min(MinDistanceToOccluder, DistanceAlongConeToOccluder);
 
 					// Shouldn't have to use atomics here as only one heightfield can affect a given world space position
-					RWRecordConeVisibility[DispatchThreadId.x * 9 + ConeIndex] = min(Visibility, ExistingConeVisibility);
+					RWRecordConeVisibility[DispatchThreadId.x * NUM_CONE_DIRECTIONS + ConeIndex] = min(Visibility, ExistingConeVisibility);
 				}
 
 				RWOccluderRadius[RecordIndex] = min(RWOccluderRadius[RecordIndex], RecordRadiusScale * MinDistanceToOccluder);
@@ -1156,8 +1041,8 @@ void CalculateHeightfieldOcclusionCS(
 	}
 }
 
-Buffer<float> RecordConeVisibility;
 RWBuffer<float4> RWIrradianceCacheBentNormal;
+Buffer<float> RecordConeVisibility;
 
 /**  */
 [numthreads(FINAL_GATHER_THREADGROUP_SIZE, 1, 1)]
@@ -1173,27 +1058,24 @@ void CombineConesCS(
 
 	if (RecordIndex < NumRecords)
 	{
-		float3 WorldNormal = DecodeNormalForAO(IrradianceCacheNormal[RecordIndex].xyz);
+		float3 RecordWorldNormal = IrradianceCacheNormal[RecordIndex].xyz;
 
 		float3 TangentX;
 		float3 TangentY;
-		FindBestAxisVectors2(WorldNormal, TangentX, TangentY);
+		FindBestAxisVectors2(RecordWorldNormal, TangentX, TangentY);
 
 		float3 UnoccludedDirection = 0;
-		float Visibility = 0;
-		uint NumSamplesToDisplay = VISUALIZE_ONE_CONE ? 1 : NUM_CONE_DIRECTIONS;
 
-		for (uint ConeIndex = 0; ConeIndex < NumSamplesToDisplay; ConeIndex++)
+		for (uint ConeIndex = 0; ConeIndex < NUM_CONE_DIRECTIONS; ConeIndex++)
 		{
 			float3 ConeDirection = AOSamples2.SampleDirections[ConeIndex].xyz;
-			float3 RotatedConeDirection = ConeDirection.x * TangentX + ConeDirection.y * TangentY + ConeDirection.z * WorldNormal;
+			float3 RotatedConeDirection = ConeDirection.x * TangentX + ConeDirection.y * TangentY + ConeDirection.z * RecordWorldNormal;
 
-			float ConeVisibility = RecordConeVisibility[DispatchThreadId.x * 9 + ConeIndex];
+			float ConeVisibility = RecordConeVisibility[DispatchThreadId.x * NUM_CONE_DIRECTIONS + ConeIndex];
 			UnoccludedDirection += ConeVisibility * RotatedConeDirection;
-			Visibility += ConeVisibility;
 		}
 
-		float InvNumSamples = 1.0f  / (float)NumSamplesToDisplay;
+		float InvNumSamples = 1.0f / (float)NUM_CONE_DIRECTIONS;
 		UnoccludedDirection = UnoccludedDirection * (BentNormalNormalizeFactor * InvNumSamples);
 
 		RWIrradianceCacheBentNormal[RecordIndex] = float4(UnoccludedDirection, 0);
@@ -1262,8 +1144,6 @@ RWBuffer<float4> RWCopyIrradianceCacheIrradiance;
 RWBuffer<uint2> RWCopyIrradianceCacheTileCoordinate;
 
 Buffer<float> OccluderRadius;
-Buffer<float4> IrradianceCacheBentNormal;
-Buffer<float4> IrradianceCacheIrradiance;
 
 [numthreads(COPY_THREADGROUP_SIZE, 1, 1)]
 void CopyIrradianceCacheSamplesCS(
diff --git a/Engine/Shaders/SkyLighting.usf b/Engine/Shaders/SkyLighting.usf
--- a/Engine/Shaders/SkyLighting.usf
+++ b/Engine/Shaders/SkyLighting.usf
@@ -70,7 +70,7 @@ void SkyLightDiffusePS(
 		SkyVisibility = saturate(ContrastCurve * ContrastAndNormalizeMulAdd.y + ContrastAndNormalizeMulAdd.z);
 		
 		#if SUPPORT_IRRADIANCE
-			DiffuseIrradiance = Texture2DSampleLevel(IrradianceTexture, IrradianceSampler, UV, 0).xyz * View.GeneralPurposeTweak;
+			DiffuseIrradiance = Texture2DSampleLevel(IrradianceTexture, IrradianceSampler, UV, 0).xyz;
 		#endif
 #endif
 
diff --git a/Engine/Source/Developer/MeshUtilities/Private/MeshUtilities.cpp b/Engine/Source/Developer/MeshUtilities/Private/MeshUtilities.cpp
--- a/Engine/Source/Developer/MeshUtilities/Private/MeshUtilities.cpp
+++ b/Engine/Source/Developer/MeshUtilities/Private/MeshUtilities.cpp
@@ -272,13 +272,14 @@ private:
 void FMeshDistanceFieldAsyncTask::DoWork()
 {
 	FMeshBuildDataProvider kDOPDataProvider(*kDopTree);
-	const FVector4 DistanceFieldVoxelSize(VolumeBounds.GetSize() / FVector(VolumeDimensions.X, VolumeDimensions.Y, VolumeDimensions.Z));
+	const FVector DistanceFieldVoxelSize(VolumeBounds.GetSize() / FVector(VolumeDimensions.X, VolumeDimensions.Y, VolumeDimensions.Z));
+	const float VoxelDiameter = DistanceFieldVoxelSize.Size();
 
 	for (int32 YIndex = 0; YIndex < VolumeDimensions.Y; YIndex++)
 	{
 		for (int32 XIndex = 0; XIndex < VolumeDimensions.X; XIndex++)
 		{
-			const FVector4 VoxelPosition = FVector4(XIndex + .5f, YIndex + .5f, ZIndex + .5f) * DistanceFieldVoxelSize + VolumeBounds.Min;
+			const FVector VoxelPosition = FVector(XIndex + .5f, YIndex + .5f, ZIndex + .5f) * DistanceFieldVoxelSize + VolumeBounds.Min;
 			const int32 Index = (ZIndex * VolumeDimensions.Y * VolumeDimensions.X + YIndex * VolumeDimensions.X + XIndex);
 
 			float MinDistance = VolumeMaxDistance;
@@ -288,15 +289,14 @@ void FMeshDistanceFieldAsyncTask::DoWork()
 			for (int32 SampleIndex = 0; SampleIndex < SampleDirections->Num(); SampleIndex++)
 			{
 				const FVector RayDirection = (*SampleDirections)[SampleIndex];
-				const FVector ReflectedRayDirection = RayDirection;
 
-				if (FMath::LineBoxIntersection(VolumeBounds, VoxelPosition, VoxelPosition + ReflectedRayDirection * VolumeMaxDistance, ReflectedRayDirection))
+				if (FMath::LineBoxIntersection(VolumeBounds, VoxelPosition, VoxelPosition + RayDirection * VolumeMaxDistance, RayDirection))
 				{
 					FkHitResult Result;
 
 					TkDOPLineCollisionCheck<const FMeshBuildDataProvider,uint32> kDOPCheck(
 						VoxelPosition,
-						VoxelPosition + ReflectedRayDirection * VolumeMaxDistance,
+						VoxelPosition + RayDirection * VolumeMaxDistance,
 						true,
 						kDOPDataProvider,
 						&Result);
@@ -309,7 +309,7 @@ void FMeshDistanceFieldAsyncTask::DoWork()
 
 						const FVector HitNormal = kDOPCheck.GetHitNormal();
 
-						if (FVector::DotProduct(ReflectedRayDirection, HitNormal) > 0 
+						if (FVector::DotProduct(RayDirection, HitNormal) > 0 
 							// MaterialIndex on the build triangles was set to 1 if two-sided, or 0 if one-sided
 							&& kDOPCheck.Result->Item == 0)
 						{
@@ -326,9 +326,18 @@ void FMeshDistanceFieldAsyncTask::DoWork()
 				}
 			}
 
+			const float UnsignedDistance = MinDistance;
+
 			// Consider this voxel 'inside' an object if more than 50% of the rays hit back faces
 			MinDistance *= (Hit == 0 || HitBack < SampleDirections->Num() * .5f) ? 1 : -1;
 
+			// If we are very close to a surface and nearly all of our rays hit backfaces, treat as inside
+			// This is important for one sided planes
+			if (UnsignedDistance < VoxelDiameter && HitBack > .95f * Hit)
+			{
+				MinDistance = -UnsignedDistance;
+			}
+
 			const float VolumeSpaceDistance = MinDistance / VolumeBounds.GetExtent().GetMax();
 
 			if (MinDistance < 0 && 
diff --git a/Engine/Source/Runtime/Engine/Private/DirectionalLightComponent.cpp b/Engine/Source/Runtime/Engine/Private/DirectionalLightComponent.cpp
--- a/Engine/Source/Runtime/Engine/Private/DirectionalLightComponent.cpp
+++ b/Engine/Source/Runtime/Engine/Private/DirectionalLightComponent.cpp
@@ -376,7 +376,7 @@
 	}
 
 	// Computes a shadow culling volume (convex hull) based on a set of 8 vertices and a light direction
-	void ComputeShadowCullingVolume( const FVector *CascadeFrustumVerts, const FVector& LightDirection, FConvexVolume& ConvexVolumeOut, FPlane& NearPlaneOut, FPlane& FarPlaneOut ) const
+	void ComputeShadowCullingVolume(const FVector* CascadeFrustumVerts, const FVector& LightDirection, FConvexVolume& ConvexVolumeOut, FPlane& NearPlaneOut, FPlane& FarPlaneOut) const
 	{
 		// Pairs of plane indices from SubFrustumPlanes whose intersections
 		// form the edges of the frustum.
@@ -500,14 +500,8 @@
 		return ShadowNear + ComputeAccumulatedScale(CascadeDistributionExponent, SplitIndex, NumCascades) * (FarDistance - ShadowNear);
 	}
 
-	virtual FSphere GetShadowSplitBounds(const FSceneView& View, int32 SplitIndex, FShadowCascadeSettings* OutCascadeSettings) const
+	virtual FSphere GetShadowSplitBoundsDepthRange(const FSceneView& View, float SplitNear, float SplitFar, FShadowCascadeSettings* OutCascadeSettings) const
 	{
-		// this checks for WholeSceneDynamicShadowRadius and DynamicShadowCascades
-		int32 NumCascades = GetNumViewDependentWholeSceneShadows(View);
-
-		const bool bHasRayTracedCascade = ShouldCreateRayTracedCascade(View.GetFeatureLevel());
-		const bool bIsRayTracedCascade = bHasRayTracedCascade && SplitIndex == NumCascades - 1;
-
 		const FMatrix ViewMatrix = View.ShadowViewMatrices.ViewMatrix;
 		const FMatrix ProjectionMatrix = View.ShadowViewMatrices.ProjMatrix;
 		const FVector4 ViewOrigin = View.ShadowViewMatrices.ViewOrigin;
@@ -515,37 +509,6 @@
 		const FVector CameraDirection = ViewMatrix.GetColumn(2);
 		const FVector LightDirection = -GetDirection();
 
-		// Determine start and end distances to the current cascade's split planes
-		// Presence of the ray traced cascade does not change depth ranges for the shadow-mapped cascades
-		float SplitNear = GetSplitDistance(View, SplitIndex, NumCascades, View.NearClippingDistance, bHasRayTracedCascade);
-		float SplitFar = GetSplitDistance(View, SplitIndex + 1, NumCascades, View.NearClippingDistance, bHasRayTracedCascade && !bIsRayTracedCascade);
-		float FadePlane = SplitFar;
-
-		float LocalCascadeTransitionFraction = CascadeTransitionFraction * GetShadowTransitionScale();
-
-		float FadeExtension = (SplitFar - SplitNear) * LocalCascadeTransitionFraction;
-
-		// Add the fade region to the end of the subfrustum, if this is not the last cascade.
-		if (SplitIndex < (NumCascades - 1))
-		{
-			SplitFar += FadeExtension;			
-		}
-
-		if(OutCascadeSettings)
-		{
-			OutCascadeSettings->SplitFarFadeRegion = FadeExtension;
-
-			OutCascadeSettings->SplitNearFadeRegion = 0.0f;
-			if(SplitIndex >= 1)
-			{
-				// todo: optimize
-				float BeforeSplitNear = GetSplitDistance(View, SplitIndex - 1, NumCascades, View.NearClippingDistance, bHasRayTracedCascade);
-				float BeforeSplitFar = GetSplitDistance(View, SplitIndex, NumCascades, View.NearClippingDistance, bHasRayTracedCascade);
-
-				OutCascadeSettings->SplitNearFadeRegion = (BeforeSplitFar - BeforeSplitNear) * LocalCascadeTransitionFraction;
-			}
-		}
-
 		// Get FOV and AspectRatio from the view's projection matrix.
 		float AspectRatio = ProjectionMatrix.M[1][1] / ProjectionMatrix.M[0][0];
 		float HalfFOV = (ViewOrigin.W > 0.0f) ? FMath::Atan(1.0f / ProjectionMatrix.M[0][0]) : PI/4.0f;
@@ -555,7 +518,7 @@
 		const FVector StartCameraRightOffset = ViewMatrix.GetColumn(0) * StartHorizontalLength;
 		const float StartVerticalLength = StartHorizontalLength / AspectRatio;
 		const FVector StartCameraUpOffset = ViewMatrix.GetColumn(1) * StartVerticalLength;
-		
+
 		const float EndHorizontalLength = SplitFar * FMath::Tan(HalfFOV);
 		const FVector EndCameraRightOffset = ViewMatrix.GetColumn(0) * EndHorizontalLength;
 		const float EndVerticalLength = EndHorizontalLength / AspectRatio;
@@ -603,15 +566,60 @@
 
 		if (OutCascadeSettings)
 		{
+			// this function is needed, since it's also called by the LPV code.
+			ComputeShadowCullingVolume(CascadeFrustumVerts, LightDirection, OutCascadeSettings->ShadowBoundsAccurate, OutCascadeSettings->NearFrustumPlane, OutCascadeSettings->FarFrustumPlane);
+		}
+
+		return CascadeSphere;
+	}
+
+	virtual FSphere GetShadowSplitBounds(const FSceneView& View, int32 SplitIndex, FShadowCascadeSettings* OutCascadeSettings) const
+	{
+		// this checks for WholeSceneDynamicShadowRadius and DynamicShadowCascades
+		int32 NumCascades = GetNumViewDependentWholeSceneShadows(View);
+
+		const bool bHasRayTracedCascade = ShouldCreateRayTracedCascade(View.GetFeatureLevel());
+		const bool bIsRayTracedCascade = bHasRayTracedCascade && SplitIndex == NumCascades - 1;
+
+		// Determine start and end distances to the current cascade's split planes
+		// Presence of the ray traced cascade does not change depth ranges for the shadow-mapped cascades
+		float SplitNear = GetSplitDistance(View, SplitIndex, NumCascades, View.NearClippingDistance, bHasRayTracedCascade);
+		float SplitFar = GetSplitDistance(View, SplitIndex + 1, NumCascades, View.NearClippingDistance, bHasRayTracedCascade && !bIsRayTracedCascade);
+		float FadePlane = SplitFar;
+
+		float LocalCascadeTransitionFraction = CascadeTransitionFraction * GetShadowTransitionScale();
+
+		float FadeExtension = (SplitFar - SplitNear) * LocalCascadeTransitionFraction;
+
+		// Add the fade region to the end of the subfrustum, if this is not the last cascade.
+		if (SplitIndex < (NumCascades - 1))
+		{
+			SplitFar += FadeExtension;			
+		}
+
+		if(OutCascadeSettings)
+		{
+			OutCascadeSettings->SplitFarFadeRegion = FadeExtension;
+
+			OutCascadeSettings->SplitNearFadeRegion = 0.0f;
+
+			if(SplitIndex >= 1)
+			{
+				// todo: optimize
+				float BeforeSplitNear = GetSplitDistance(View, SplitIndex - 1, NumCascades, View.NearClippingDistance, bHasRayTracedCascade);
+				float BeforeSplitFar = GetSplitDistance(View, SplitIndex, NumCascades, View.NearClippingDistance, bHasRayTracedCascade);
+
+				OutCascadeSettings->SplitNearFadeRegion = (BeforeSplitFar - BeforeSplitNear) * LocalCascadeTransitionFraction;
+			}
+
 			// Pass out the split settings
 			OutCascadeSettings->SplitFar = SplitFar;
 			OutCascadeSettings->SplitNear = SplitNear;
 			OutCascadeSettings->FadePlaneOffset = FadePlane;
 			OutCascadeSettings->FadePlaneLength = SplitFar - FadePlane;
+		}
 
-			// this function is needed, since it's also called by the LPV code.
- 			ComputeShadowCullingVolume( CascadeFrustumVerts, LightDirection, OutCascadeSettings->ShadowBoundsAccurate, OutCascadeSettings->NearFrustumPlane, OutCascadeSettings->FarFrustumPlane );
-		}
+		const FSphere CascadeSphere = FDirectionalLightSceneProxy::GetShadowSplitBoundsDepthRange(View, SplitNear, SplitFar, OutCascadeSettings);
 
 		return CascadeSphere;
 	}
diff --git a/Engine/Source/Runtime/Engine/Public/DistanceFieldAtlas.h b/Engine/Source/Runtime/Engine/Public/DistanceFieldAtlas.h
--- a/Engine/Source/Runtime/Engine/Public/DistanceFieldAtlas.h
+++ b/Engine/Source/Runtime/Engine/Public/DistanceFieldAtlas.h
@@ -10,7 +10,7 @@
 #include "GCObject.h"
 
 // DDC key for distance field data, must be changed when modifying the generation code or data format
-#define DISTANCEFIELD_DERIVEDDATA_VER TEXT("7768798763B445A9543C94442EA899D")
+#define DISTANCEFIELD_DERIVEDDATA_VER TEXT("7768798764B445A9543C94442EA899D")
 
 /** Represents a distance field volume texture for a single UStaticMesh. */
 class FDistanceFieldVolumeTexture
diff --git a/Engine/Source/Runtime/Engine/Public/SceneManagement.h b/Engine/Source/Runtime/Engine/Public/SceneManagement.h
--- a/Engine/Source/Runtime/Engine/Public/SceneManagement.h
+++ b/Engine/Source/Runtime/Engine/Public/SceneManagement.h
@@ -731,6 +731,7 @@ public:
 
 	// @param OutCascadeSettings can be 0
 	virtual FSphere GetShadowSplitBounds(const class FSceneView& View, int32 SplitIndex, FShadowCascadeSettings* OutCascadeSettings) const { return FSphere(FVector::ZeroVector, 0); }
+	virtual FSphere GetShadowSplitBoundsDepthRange(const FSceneView& View, float SplitNear, float SplitFar, FShadowCascadeSettings* OutCascadeSettings) const { return FSphere(FVector::ZeroVector, 0); }
 
 	virtual bool GetScissorRect(FIntRect& ScissorRect, const FSceneView& View) const
 	{
diff --git a/Engine/Source/Runtime/Engine/Public/ShowFlags.h b/Engine/Source/Runtime/Engine/Public/ShowFlags.h
--- a/Engine/Source/Runtime/Engine/Public/ShowFlags.h
+++ b/Engine/Source/Runtime/Engine/Public/ShowFlags.h
@@ -208,6 +208,7 @@ struct FEngineShowFlags
 		SetSnap(false);
 		SetVisualizeMeshDistanceFields(false);
 		SetVisualizeDistanceFieldAO(false);
+		SetVisualizeDistanceFieldGI(false);
 		SetVisualizeSSR(false);
 		SetForceGBuffer(false);
 		SetVisualizeSSS(false);
@@ -250,6 +251,7 @@ struct FEngineShowFlags
 		SetHMDDistortion(false);
 		SetStereoRendering(false);
 		SetDistanceFieldAO(false);
+		SetDistanceFieldGI(false);
 
 		SetForceGBuffer(true);
 	}
diff --git a/Engine/Source/Runtime/Engine/Public/ShowFlagsValues.inl b/Engine/Source/Runtime/Engine/Public/ShowFlagsValues.inl
--- a/Engine/Source/Runtime/Engine/Public/ShowFlagsValues.inl
+++ b/Engine/Source/Runtime/Engine/Public/ShowFlagsValues.inl
@@ -267,11 +267,15 @@ SHOWFLAG_ALWAYS_ACCESSIBLE(LpvLightingOnly, SFG_Hidden, LOCTEXT("VisualizeLPVSF_
 /** Render Paper2D sprites */
 SHOWFLAG_ALWAYS_ACCESSIBLE(Paper2DSprites, SFG_Advanced, LOCTEXT("Paper2DSpritesSF", "Paper 2D Sprites"))
 /** Visualization of distance field AO */
-SHOWFLAG_ALWAYS_ACCESSIBLE(VisualizeDistanceFieldAO, SFG_Visualize, LOCTEXT("VisualizeDistanceFieldAOSF", "DistanceField AO"))
+SHOWFLAG_ALWAYS_ACCESSIBLE(VisualizeDistanceFieldAO, SFG_Visualize, LOCTEXT("VisualizeDistanceFieldAOSF", "Distance Field Ambient Occlusion"))
+/** Visualization of distance field GI */
+SHOWFLAG_ALWAYS_ACCESSIBLE(VisualizeDistanceFieldGI, SFG_Visualize, LOCTEXT("VisualizeDistanceFieldGISF", "Distance Field Global Illumination"))
 /** Mesh Distance fields */
 SHOWFLAG_ALWAYS_ACCESSIBLE(VisualizeMeshDistanceFields, SFG_Visualize, LOCTEXT("MeshDistanceFieldsSF", "Mesh DistanceFields"))
 /** Distance field AO */
-SHOWFLAG_ALWAYS_ACCESSIBLE(DistanceFieldAO, SFG_LightingFeatures, LOCTEXT("DistanceFieldAOSF", "DistanceField AO"))
+SHOWFLAG_ALWAYS_ACCESSIBLE(DistanceFieldAO, SFG_LightingFeatures, LOCTEXT("DistanceFieldAOSF", "Distance Field Ambient Occlusion"))
+/** Distance field GI */
+SHOWFLAG_ALWAYS_ACCESSIBLE(DistanceFieldGI, SFG_LightingFeatures, LOCTEXT("DistanceFieldGISF", "Distance Field Global Illumination"))
 /** Visualize screen space reflections, for developer (by default off): */
 SHOWFLAG_FIXED_IN_SHIPPING(VisualizeSSR, 0, SFG_Visualize, LOCTEXT("VisualizeSSR", "Screen Space Reflections"))
 /** Force the use of the GBuffer. */
diff --git a/Engine/Source/Runtime/Renderer/Private/DeferredShadingRenderer.cpp b/Engine/Source/Runtime/Renderer/Private/DeferredShadingRenderer.cpp
--- a/Engine/Source/Runtime/Renderer/Private/DeferredShadingRenderer.cpp
+++ b/Engine/Source/Runtime/Renderer/Private/DeferredShadingRenderer.cpp
@@ -976,15 +976,17 @@ void FDeferredShadingSceneRenderer::Render(FRHICommandListImmediate& RHICmdList)
 		RenderLightShaftBloom(RHICmdList);
 	}
 
-	if (ViewFamily.EngineShowFlags.VisualizeDistanceFieldAO)
+	if (ViewFamily.EngineShowFlags.VisualizeDistanceFieldAO || ViewFamily.EngineShowFlags.VisualizeDistanceFieldGI)
 	{
+		// Use the skylight's max distance if there is one, to be consistent with DFAO shadowing on the skylight
+		const float OcclusionMaxDistance = Scene->SkyLight && !Scene->SkyLight->bWantsStaticShadowing ? Scene->SkyLight->OcclusionMaxDistance : 600;
 		TRefCountPtr<IPooledRenderTarget> DummyOutput;
-		RenderDistanceFieldAOSurfaceCache(RHICmdList, FDistanceFieldAOParameters(), DummyOutput, DummyOutput, true);
+		RenderDistanceFieldAOSurfaceCache(RHICmdList, FDistanceFieldAOParameters(OcclusionMaxDistance), DummyOutput, DummyOutput, true);
 	}
 
 	if (ViewFamily.EngineShowFlags.VisualizeMeshDistanceFields)
 	{
-		RenderMeshDistanceFieldVisualization(RHICmdList, FDistanceFieldAOParameters());
+		RenderMeshDistanceFieldVisualization(RHICmdList, FDistanceFieldAOParameters(600));
 	}
 
 	// Resolve the scene color for post processing.
diff --git a/Engine/Source/Runtime/Renderer/Private/DistanceFieldGlobalIllumination.cpp b/Engine/Source/Runtime/Renderer/Private/DistanceFieldGlobalIllumination.cpp
new file mode 100644
--- /dev/null
+++ b/Engine/Source/Runtime/Renderer/Private/DistanceFieldGlobalIllumination.cpp
@@ -0,0 +1,878 @@
+// Copyright 1998-2014 Epic Games, Inc. All Rights Reserved.
+
+/*=============================================================================
+	DistanceFieldGlobalIllumination.cpp
+=============================================================================*/
+
+#include "RendererPrivate.h"
+#include "ScenePrivate.h"
+#include "UniformBuffer.h"
+#include "ShaderParameters.h"
+#include "PostProcessing.h"
+#include "SceneFilterRendering.h"
+#include "DistanceFieldLightingShared.h"
+#include "DistanceFieldSurfaceCacheLighting.h"
+#include "DistanceFieldGlobalIllumination.h"
+#include "RHICommandList.h"
+#include "SceneUtils.h"
+#include "DistanceFieldAtlas.h"
+
+int32 GDistanceFieldGI = 0;
+FAutoConsoleVariableRef CVarDistanceFieldGI(
+	TEXT("r.DistanceFieldGI"),
+	GDistanceFieldGI,
+	TEXT(""),
+	ECVF_Cheat | ECVF_RenderThreadSafe
+	);
+
+int32 GVPLGridDimension = 128;
+FAutoConsoleVariableRef CVarVPLGridDimension(
+	TEXT("r.VPLGridDimension"),
+	GVPLGridDimension,
+	TEXT(""),
+	ECVF_Cheat | ECVF_RenderThreadSafe
+	);
+
+float GVPLDirectionalLightTraceDistance = 100000;
+FAutoConsoleVariableRef CVarVPLDirectionalLightTraceDistance(
+	TEXT("r.VPLDirectionalLightTraceDistance"),
+	GVPLDirectionalLightTraceDistance,
+	TEXT(""),
+	ECVF_Cheat | ECVF_RenderThreadSafe
+	);
+
+float GVPLPlacementCameraRadius = 4000;
+FAutoConsoleVariableRef CVarVPLPlacementCameraRadius(
+	TEXT("r.VPLPlacementCameraRadius"),
+	GVPLPlacementCameraRadius,
+	TEXT(""),
+	ECVF_Cheat | ECVF_RenderThreadSafe
+	);
+
+int32 GVPLViewCulling = 1;
+FAutoConsoleVariableRef CVarVPLViewCulling(
+	TEXT("r.VPLViewCulling"),
+	GVPLViewCulling,
+	TEXT(""),
+	ECVF_Cheat | ECVF_RenderThreadSafe
+	);
+
+int32 GAOUseConesForGI = 1;
+FAutoConsoleVariableRef CVarAOUseConesForGI(
+	TEXT("r.AOUseConesForGI"),
+	GAOUseConesForGI,
+	TEXT(""),
+	ECVF_Cheat | ECVF_RenderThreadSafe
+	);
+
+TGlobalResource<FVPLResources> GVPLResources;
+TGlobalResource<FVPLResources> GCulledVPLResources;
+
+class FVPLPlacementCS : public FGlobalShader
+{
+	DECLARE_SHADER_TYPE(FVPLPlacementCS, Global);
+public:
+
+	static bool ShouldCache(EShaderPlatform Platform)
+	{
+		return IsFeatureLevelSupported(Platform, ERHIFeatureLevel::SM5) && DoesPlatformSupportDistanceFieldAO(Platform);
+	}
+
+	static void ModifyCompilationEnvironment(EShaderPlatform Platform, FShaderCompilerEnvironment& OutEnvironment)
+	{
+		OutEnvironment.SetDefine(TEXT("THREADGROUP_SIZEX"), GDistanceFieldAOTileSizeX);
+		OutEnvironment.SetDefine(TEXT("THREADGROUP_SIZEY"), GDistanceFieldAOTileSizeY);
+	}
+
+	/** Default constructor. */
+	FVPLPlacementCS() {}
+
+	/** Initialization constructor. */
+	FVPLPlacementCS(const ShaderMetaType::CompiledShaderInitializerType& Initializer)
+		: FGlobalShader(Initializer)
+	{
+		VPLParameterBuffer.Bind(Initializer.ParameterMap, TEXT("VPLParameterBuffer"));
+		VPLClusterData.Bind(Initializer.ParameterMap, TEXT("VPLClusterData"));
+		VPLData.Bind(Initializer.ParameterMap, TEXT("VPLData"));
+		InvPlacementGridSize.Bind(Initializer.ParameterMap, TEXT("InvPlacementGridSize"));
+		WorldToShadow.Bind(Initializer.ParameterMap, TEXT("WorldToShadow"));
+		ShadowToWorld.Bind(Initializer.ParameterMap, TEXT("ShadowToWorld"));
+		LightDirectionAndTraceDistance.Bind(Initializer.ParameterMap, TEXT("LightDirectionAndTraceDistance"));
+		LightColor.Bind(Initializer.ParameterMap, TEXT("LightColor"));
+		ObjectParameters.Bind(Initializer.ParameterMap);
+		ShadowTileHeadDataUnpacked.Bind(Initializer.ParameterMap, TEXT("ShadowTileHeadDataUnpacked"));
+		ShadowTileArrayData.Bind(Initializer.ParameterMap, TEXT("ShadowTileArrayData"));
+		ShadowTileListGroupSize.Bind(Initializer.ParameterMap, TEXT("ShadowTileListGroupSize"));
+		VPLPlacementCameraRadius.Bind(Initializer.ParameterMap, TEXT("VPLPlacementCameraRadius"));
+	}
+
+	void SetParameters(
+		FRHICommandList& RHICmdList,
+		const FSceneView& View, 
+		const FLightSceneProxy* LightSceneProxy,
+		FVector2D InvPlacementGridSizeValue,
+		const FMatrix& WorldToShadowValue,
+		const FMatrix& ShadowToWorldValue,
+		FLightTileIntersectionResources* TileIntersectionResources)
+	{
+		const FComputeShaderRHIParamRef ShaderRHI = GetComputeShader();
+		FGlobalShader::SetParameters(RHICmdList, ShaderRHI, View);
+		extern TGlobalResource<FDistanceFieldObjectBufferResource> GShadowCulledObjectBuffers;
+		ObjectParameters.Set(RHICmdList, ShaderRHI, GShadowCulledObjectBuffers.Buffers);
+
+		SetShaderValue(RHICmdList, ShaderRHI, InvPlacementGridSize, InvPlacementGridSizeValue);
+		SetShaderValue(RHICmdList, ShaderRHI, WorldToShadow, WorldToShadowValue);
+		SetShaderValue(RHICmdList, ShaderRHI, ShadowToWorld, ShadowToWorldValue);
+		SetShaderValue(RHICmdList, ShaderRHI, LightDirectionAndTraceDistance, FVector4(LightSceneProxy->GetDirection(), GVPLDirectionalLightTraceDistance));
+		SetShaderValue(RHICmdList, ShaderRHI, LightColor, LightSceneProxy->GetColor() * LightSceneProxy->GetIndirectLightingScale());
+
+		VPLParameterBuffer.SetBuffer(RHICmdList, ShaderRHI, GVPLResources.VPLParameterBuffer);
+		VPLClusterData.SetBuffer(RHICmdList, ShaderRHI, GVPLResources.VPLClusterData);
+		VPLData.SetBuffer(RHICmdList, ShaderRHI, GVPLResources.VPLData);
+
+		check(TileIntersectionResources || !ShadowTileHeadDataUnpacked.IsBound());
+
+		if (TileIntersectionResources)
+		{
+			SetSRVParameter(RHICmdList, ShaderRHI, ShadowTileHeadDataUnpacked, TileIntersectionResources->TileHeadDataUnpacked.SRV);
+			SetSRVParameter(RHICmdList, ShaderRHI, ShadowTileArrayData, TileIntersectionResources->TileArrayData.SRV);
+			SetShaderValue(RHICmdList, ShaderRHI, ShadowTileListGroupSize, TileIntersectionResources->TileDimensions);
+		}
+
+		SetShaderValue(RHICmdList, ShaderRHI, VPLPlacementCameraRadius, GVPLPlacementCameraRadius);
+	}
+
+	void UnsetParameters(FRHICommandList& RHICmdList)
+	{
+		VPLParameterBuffer.UnsetUAV(RHICmdList, GetComputeShader());
+		VPLClusterData.UnsetUAV(RHICmdList, GetComputeShader());
+		VPLData.UnsetUAV(RHICmdList, GetComputeShader());
+	}
+
+	// FShader interface.
+	virtual bool Serialize(FArchive& Ar)
+	{
+		bool bShaderHasOutdatedParameters = FGlobalShader::Serialize(Ar);
+		Ar << VPLParameterBuffer;
+		Ar << VPLClusterData;
+		Ar << VPLData;
+		Ar << InvPlacementGridSize;
+		Ar << WorldToShadow;
+		Ar << ShadowToWorld;
+		Ar << LightDirectionAndTraceDistance;
+		Ar << LightColor;
+		Ar << ObjectParameters;
+		Ar << ShadowTileHeadDataUnpacked;
+		Ar << ShadowTileArrayData;
+		Ar << ShadowTileListGroupSize;
+		Ar << VPLPlacementCameraRadius;
+		return bShaderHasOutdatedParameters;
+	}
+
+private:
+
+	FRWShaderParameter VPLParameterBuffer;
+	FRWShaderParameter VPLClusterData;
+	FRWShaderParameter VPLData;
+	FShaderParameter InvPlacementGridSize;
+	FShaderParameter WorldToShadow;
+	FShaderParameter ShadowToWorld;
+	FShaderParameter LightDirectionAndTraceDistance;
+	FShaderParameter LightColor;
+	FDistanceFieldCulledObjectBufferParameters ObjectParameters;
+	FShaderResourceParameter ShadowTileHeadDataUnpacked;
+	FShaderResourceParameter ShadowTileArrayData;
+	FShaderParameter ShadowTileListGroupSize;
+	FShaderParameter VPLPlacementCameraRadius;
+};
+
+IMPLEMENT_SHADER_TYPE(,FVPLPlacementCS,TEXT("DistanceFieldGlobalIllumination"),TEXT("VPLPlacementCS"),SF_Compute);
+
+
+
+class FSetupVPLCullndirectArgumentsCS : public FGlobalShader
+{
+	DECLARE_SHADER_TYPE(FSetupVPLCullndirectArgumentsCS,Global)
+public:
+
+	static bool ShouldCache(EShaderPlatform Platform)
+	{
+		return IsFeatureLevelSupported(Platform, ERHIFeatureLevel::SM5) && DoesPlatformSupportDistanceFieldAO(Platform);
+	}
+
+	static void ModifyCompilationEnvironment(EShaderPlatform Platform, FShaderCompilerEnvironment& OutEnvironment)
+	{
+		FGlobalShader::ModifyCompilationEnvironment(Platform,OutEnvironment);
+		OutEnvironment.SetDefine(TEXT("THREADGROUP_SIZEX"), GDistanceFieldAOTileSizeX);
+		OutEnvironment.SetDefine(TEXT("THREADGROUP_SIZEY"), GDistanceFieldAOTileSizeY);
+	}
+
+	FSetupVPLCullndirectArgumentsCS(const ShaderMetaType::CompiledShaderInitializerType& Initializer)
+		: FGlobalShader(Initializer)
+	{
+		DispatchParameters.Bind(Initializer.ParameterMap, TEXT("DispatchParameters"));
+		VPLParameterBuffer.Bind(Initializer.ParameterMap, TEXT("VPLParameterBuffer"));
+	}
+
+	FSetupVPLCullndirectArgumentsCS()
+	{
+	}
+
+	void SetParameters(FRHICommandList& RHICmdList, const FSceneView& View)
+	{
+		FComputeShaderRHIParamRef ShaderRHI = GetComputeShader();
+		FGlobalShader::SetParameters(RHICmdList, ShaderRHI, View);
+
+		DispatchParameters.SetBuffer(RHICmdList, ShaderRHI, GVPLResources.VPLDispatchIndirectBuffer);
+		SetSRVParameter(RHICmdList, ShaderRHI, VPLParameterBuffer, GVPLResources.VPLParameterBuffer.SRV);
+	}
+
+	void UnsetParameters(FRHICommandList& RHICmdList)
+	{
+		DispatchParameters.UnsetUAV(RHICmdList, GetComputeShader());
+	}
+
+	virtual bool Serialize(FArchive& Ar)
+	{		
+		bool bShaderHasOutdatedParameters = FGlobalShader::Serialize(Ar);
+		Ar << DispatchParameters;
+		Ar << VPLParameterBuffer;
+		return bShaderHasOutdatedParameters;
+	}
+
+private:
+
+	FRWShaderParameter DispatchParameters;
+	FShaderResourceParameter VPLParameterBuffer;
+};
+
+IMPLEMENT_SHADER_TYPE(,FSetupVPLCullndirectArgumentsCS,TEXT("DistanceFieldGlobalIllumination"),TEXT("SetupVPLCullndirectArgumentsCS"),SF_Compute);
+
+
+class FCullVPLsForViewCS : public FGlobalShader
+{
+	DECLARE_SHADER_TYPE(FCullVPLsForViewCS,Global)
+public:
+
+	static bool ShouldCache(EShaderPlatform Platform)
+	{
+		return IsFeatureLevelSupported(Platform, ERHIFeatureLevel::SM5) && DoesPlatformSupportDistanceFieldAO(Platform);
+	}
+
+	static void ModifyCompilationEnvironment(EShaderPlatform Platform, FShaderCompilerEnvironment& OutEnvironment)
+	{
+		FGlobalShader::ModifyCompilationEnvironment(Platform,OutEnvironment);
+		OutEnvironment.SetDefine(TEXT("THREADGROUP_SIZEX"), GDistanceFieldAOTileSizeX);
+		OutEnvironment.SetDefine(TEXT("THREADGROUP_SIZEY"), GDistanceFieldAOTileSizeY);
+	}
+
+	FCullVPLsForViewCS(const ShaderMetaType::CompiledShaderInitializerType& Initializer)
+		: FGlobalShader(Initializer)
+	{
+		VPLParameterBuffer.Bind(Initializer.ParameterMap, TEXT("VPLParameterBuffer"));
+		VPLClusterData.Bind(Initializer.ParameterMap, TEXT("VPLClusterData"));
+		VPLData.Bind(Initializer.ParameterMap, TEXT("VPLData"));
+		CulledVPLParameterBuffer.Bind(Initializer.ParameterMap, TEXT("CulledVPLParameterBuffer"));
+		CulledVPLClusterData.Bind(Initializer.ParameterMap, TEXT("CulledVPLClusterData"));
+		CulledVPLData.Bind(Initializer.ParameterMap, TEXT("CulledVPLData"));
+		AOParameters.Bind(Initializer.ParameterMap);
+		NumConvexHullPlanes.Bind(Initializer.ParameterMap, TEXT("NumConvexHullPlanes"));
+		ViewFrustumConvexHull.Bind(Initializer.ParameterMap, TEXT("ViewFrustumConvexHull"));
+	}
+
+	FCullVPLsForViewCS()
+	{
+	}
+
+	void SetParameters(FRHICommandList& RHICmdList, const FScene* Scene, const FSceneView& View, const FDistanceFieldAOParameters& Parameters)
+	{
+		FComputeShaderRHIParamRef ShaderRHI = GetComputeShader();
+		FGlobalShader::SetParameters(RHICmdList, ShaderRHI, View);
+
+		CulledVPLParameterBuffer.SetBuffer(RHICmdList, ShaderRHI, GCulledVPLResources.VPLParameterBuffer);
+		CulledVPLClusterData.SetBuffer(RHICmdList, ShaderRHI, GCulledVPLResources.VPLClusterData);
+		CulledVPLData.SetBuffer(RHICmdList, ShaderRHI, GCulledVPLResources.VPLData);
+		AOParameters.Set(RHICmdList, ShaderRHI, Parameters);
+
+		SetSRVParameter(RHICmdList, ShaderRHI, VPLParameterBuffer, GVPLResources.VPLParameterBuffer.SRV);
+		SetSRVParameter(RHICmdList, ShaderRHI, VPLClusterData, GVPLResources.VPLClusterData.SRV);
+		SetSRVParameter(RHICmdList, ShaderRHI, VPLData, GVPLResources.VPLData.SRV);
+
+		// Shader assumes max 6
+		check(View.ViewFrustum.Planes.Num() <= 6);
+		SetShaderValue(RHICmdList, ShaderRHI, NumConvexHullPlanes, View.ViewFrustum.Planes.Num());
+		SetShaderValueArray(RHICmdList, ShaderRHI, ViewFrustumConvexHull, View.ViewFrustum.Planes.GetData(), View.ViewFrustum.Planes.Num());
+	}
+
+	void UnsetParameters(FRHICommandList& RHICmdList)
+	{
+		CulledVPLParameterBuffer.UnsetUAV(RHICmdList, GetComputeShader());
+		CulledVPLClusterData.UnsetUAV(RHICmdList, GetComputeShader());
+		CulledVPLData.UnsetUAV(RHICmdList, GetComputeShader());
+	}
+
+	virtual bool Serialize(FArchive& Ar)
+	{		
+		bool bShaderHasOutdatedParameters = FGlobalShader::Serialize(Ar);
+		Ar << VPLParameterBuffer;
+		Ar << VPLClusterData;
+		Ar << VPLData;
+		Ar << CulledVPLParameterBuffer;
+		Ar << CulledVPLClusterData;
+		Ar << CulledVPLData;
+		Ar << AOParameters;
+		Ar << NumConvexHullPlanes;
+		Ar << ViewFrustumConvexHull;
+		return bShaderHasOutdatedParameters;
+	}
+
+private:
+
+	FShaderResourceParameter VPLParameterBuffer;
+	FShaderResourceParameter VPLClusterData;
+	FShaderResourceParameter VPLData;
+	FRWShaderParameter CulledVPLParameterBuffer;
+	FRWShaderParameter CulledVPLClusterData;
+	FRWShaderParameter CulledVPLData;
+	FAOParameters AOParameters;
+	FShaderParameter NumConvexHullPlanes;
+	FShaderParameter ViewFrustumConvexHull;
+};
+
+IMPLEMENT_SHADER_TYPE(,FCullVPLsForViewCS,TEXT("DistanceFieldGlobalIllumination"),TEXT("CullVPLsForViewCS"),SF_Compute);
+
+
+/**  */
+class FVPLTileCullCS : public FGlobalShader
+{
+	DECLARE_SHADER_TYPE(FVPLTileCullCS,Global)
+public:
+
+	static bool ShouldCache(EShaderPlatform Platform)
+	{
+		return IsFeatureLevelSupported(Platform, ERHIFeatureLevel::SM5) && DoesPlatformSupportDistanceFieldAO(Platform);
+	}
+
+	static void ModifyCompilationEnvironment(EShaderPlatform Platform, FShaderCompilerEnvironment& OutEnvironment)
+	{
+		FGlobalShader::ModifyCompilationEnvironment(Platform,OutEnvironment);
+		OutEnvironment.SetDefine(TEXT("THREADGROUP_SIZEX"), GDistanceFieldAOTileSizeX);
+		OutEnvironment.SetDefine(TEXT("THREADGROUP_SIZEY"), GDistanceFieldAOTileSizeY);
+		OutEnvironment.CompilerFlags.Add(CFLAG_StandardOptimization);
+	}
+
+	FVPLTileCullCS(const ShaderMetaType::CompiledShaderInitializerType& Initializer)
+		: FGlobalShader(Initializer)
+	{
+		AOParameters.Bind(Initializer.ParameterMap);
+		TileHeadDataUnpacked.Bind(Initializer.ParameterMap, TEXT("TileHeadDataUnpacked"));
+		TileArrayData.Bind(Initializer.ParameterMap, TEXT("TileArrayData"));
+		VPLParameterBuffer.Bind(Initializer.ParameterMap, TEXT("VPLParameterBuffer"));
+		VPLClusterData.Bind(Initializer.ParameterMap, TEXT("VPLClusterData"));
+		TileConeAxisAndCos.Bind(Initializer.ParameterMap, TEXT("TileConeAxisAndCos"));
+		TileConeDepthRanges.Bind(Initializer.ParameterMap, TEXT("TileConeDepthRanges"));
+		NumGroups.Bind(Initializer.ParameterMap, TEXT("NumGroups"));
+	}
+
+	FVPLTileCullCS()
+	{
+	}
+	void SetParameters(FRHICommandList& RHICmdList, const FSceneView& View, const FScene* Scene, const FDistanceFieldAOParameters& Parameters, const FVPLResources& VPLResources)
+	{
+		FComputeShaderRHIParamRef ShaderRHI = GetComputeShader();
+
+		FGlobalShader::SetParameters(RHICmdList, ShaderRHI, View);
+		AOParameters.Set(RHICmdList, ShaderRHI, Parameters);
+
+		FTileIntersectionResources* TileIntersectionResources = ((FSceneViewState*)View.State)->AOTileIntersectionResources;
+		FVPLTileIntersectionResources* VPLTileIntersectionResources = ((FSceneViewState*)View.State)->VPLTileIntersectionResources;
+
+		TileHeadDataUnpacked.SetBuffer(RHICmdList, ShaderRHI, VPLTileIntersectionResources->TileHeadDataUnpacked);
+		TileArrayData.SetBuffer(RHICmdList, ShaderRHI, VPLTileIntersectionResources->TileArrayData);
+
+		SetSRVParameter(RHICmdList, ShaderRHI, VPLParameterBuffer, VPLResources.VPLParameterBuffer.SRV);
+		SetSRVParameter(RHICmdList, ShaderRHI, VPLClusterData, VPLResources.VPLClusterData.SRV);
+
+		SetSRVParameter(RHICmdList, ShaderRHI, TileConeAxisAndCos, TileIntersectionResources->TileConeAxisAndCos.SRV);
+		SetSRVParameter(RHICmdList, ShaderRHI, TileConeDepthRanges, TileIntersectionResources->TileConeDepthRanges.SRV);
+	
+		SetShaderValue(RHICmdList, ShaderRHI, NumGroups, VPLTileIntersectionResources->TileDimensions);
+	}
+
+	void UnsetParameters(FRHICommandList& RHICmdList)
+	{
+		TileHeadDataUnpacked.UnsetUAV(RHICmdList, GetComputeShader());
+		TileArrayData.UnsetUAV(RHICmdList, GetComputeShader());
+	}
+
+	virtual bool Serialize(FArchive& Ar)
+	{		
+		bool bShaderHasOutdatedParameters = FGlobalShader::Serialize(Ar);
+		Ar << AOParameters;
+		Ar << TileHeadDataUnpacked;
+		Ar << TileArrayData;
+		Ar << VPLParameterBuffer;
+		Ar << VPLClusterData;
+		Ar << TileConeAxisAndCos;
+		Ar << TileConeDepthRanges;
+		Ar << NumGroups;
+		return bShaderHasOutdatedParameters;
+	}
+
+private:
+
+	FAOParameters AOParameters;
+	FRWShaderParameter TileHeadDataUnpacked;
+	FRWShaderParameter TileArrayData;
+	FShaderResourceParameter VPLParameterBuffer;
+	FShaderResourceParameter VPLClusterData;
+	FShaderResourceParameter TileConeAxisAndCos;
+	FShaderResourceParameter TileConeDepthRanges;
+	FShaderParameter NumGroups;
+};
+
+IMPLEMENT_SHADER_TYPE(,FVPLTileCullCS,TEXT("DistanceFieldGlobalIllumination"),TEXT("VPLTileCullCS"),SF_Compute);
+
+
+enum EComputeIrradianceMode
+{
+	IrradianceMode_IrradianceFromAOCones,
+	IrradianceMode_IrradianceFinalGather
+};
+
+template<EComputeIrradianceMode CombineMode>
+class TComputeIrradianceCS : public FGlobalShader
+{
+	DECLARE_SHADER_TYPE(TComputeIrradianceCS,Global)
+public:
+
+	static bool ShouldCache(EShaderPlatform Platform)
+	{
+		return IsFeatureLevelSupported(Platform, ERHIFeatureLevel::SM5) && DoesPlatformSupportDistanceFieldAO(Platform);
+	}
+
+	static void ModifyCompilationEnvironment(EShaderPlatform Platform, FShaderCompilerEnvironment& OutEnvironment)
+	{
+		FGlobalShader::ModifyCompilationEnvironment(Platform,OutEnvironment);
+		OutEnvironment.SetDefine(TEXT("NUM_CONE_DIRECTIONS"), NumConeSampleDirections);
+		OutEnvironment.SetDefine(TEXT("IRRADIANCE_FROM_AO_CONES"), CombineMode == IrradianceMode_IrradianceFromAOCones ? TEXT("1") : TEXT("0"));
+
+		// To reduce shader compile time of compute shaders with shared memory, doesn't have an impact on generated code with current compiler (June 2010 DX SDK)
+		OutEnvironment.CompilerFlags.Add(CFLAG_StandardOptimization);
+	}
+
+	TComputeIrradianceCS(const ShaderMetaType::CompiledShaderInitializerType& Initializer)
+		: FGlobalShader(Initializer)
+	{
+		AOParameters.Bind(Initializer.ParameterMap);
+		IrradianceCacheIrradiance.Bind(Initializer.ParameterMap, TEXT("IrradianceCacheIrradiance"));
+		IrradianceCachePositionRadius.Bind(Initializer.ParameterMap, TEXT("IrradianceCachePositionRadius"));
+		IrradianceCacheNormal.Bind(Initializer.ParameterMap, TEXT("IrradianceCacheNormal"));
+		IrradianceCacheTileCoordinate.Bind(Initializer.ParameterMap, TEXT("IrradianceCacheTileCoordinate"));
+		RecordConeData.Bind(Initializer.ParameterMap, TEXT("RecordConeData"));
+		ScatterDrawParameters.Bind(Initializer.ParameterMap, TEXT("ScatterDrawParameters"));
+		SavedStartIndex.Bind(Initializer.ParameterMap, TEXT("SavedStartIndex"));
+		DistanceFieldNormalTexture.Bind(Initializer.ParameterMap, TEXT("DistanceFieldNormalTexture"));
+		DistanceFieldNormalSampler.Bind(Initializer.ParameterMap, TEXT("DistanceFieldNormalSampler"));
+		GlobalObjectParameters.Bind(Initializer.ParameterMap);
+		VPLParameterBuffer.Bind(Initializer.ParameterMap, TEXT("VPLParameterBuffer"));
+		VPLClusterData.Bind(Initializer.ParameterMap, TEXT("VPLClusterData"));
+		VPLData.Bind(Initializer.ParameterMap, TEXT("VPLData"));
+		VPLGatherRadius.Bind(Initializer.ParameterMap, TEXT("VPLGatherRadius"));
+		ConeHalfAngle.Bind(Initializer.ParameterMap, TEXT("ConeHalfAngle"));
+		BentNormalNormalizeFactor.Bind(Initializer.ParameterMap, TEXT("BentNormalNormalizeFactor"));
+		TileListGroupSize.Bind(Initializer.ParameterMap, TEXT("TileListGroupSize"));
+		TileHeadDataUnpacked.Bind(Initializer.ParameterMap, TEXT("TileHeadDataUnpacked"));
+		TileArrayData.Bind(Initializer.ParameterMap, TEXT("TileArrayData"));
+	}
+
+	TComputeIrradianceCS()
+	{
+	}
+
+	void SetParameters(
+		FRHICommandList& RHICmdList, 
+		const FSceneView& View, 
+		int32 DepthLevel, 
+		const FDistanceFieldAOParameters& Parameters,
+		FTemporaryIrradianceCacheResources* TemporaryIrradianceCacheResources)
+	{
+		FComputeShaderRHIParamRef ShaderRHI = GetComputeShader();
+		FGlobalShader::SetParameters(RHICmdList, ShaderRHI, View);
+
+		AOParameters.Set(RHICmdList, ShaderRHI, Parameters);
+
+		const FScene* Scene = (const FScene*)View.Family->Scene;
+		FSurfaceCacheResources& SurfaceCacheResources = *Scene->SurfaceCacheResources;
+
+		SetSRVParameter(RHICmdList, ShaderRHI, IrradianceCachePositionRadius, SurfaceCacheResources.Level[DepthLevel]->PositionAndRadius.SRV);
+		SetSRVParameter(RHICmdList, ShaderRHI, IrradianceCacheNormal, SurfaceCacheResources.Level[DepthLevel]->Normal.SRV);
+		SetSRVParameter(RHICmdList, ShaderRHI, IrradianceCacheTileCoordinate, SurfaceCacheResources.Level[DepthLevel]->TileCoordinate.SRV);
+		SetSRVParameter(RHICmdList, ShaderRHI, ScatterDrawParameters, SurfaceCacheResources.Level[DepthLevel]->ScatterDrawParameters.SRV);
+		SetSRVParameter(RHICmdList, ShaderRHI, SavedStartIndex, SurfaceCacheResources.Level[DepthLevel]->SavedStartIndex.SRV);
+		SetSRVParameter(RHICmdList, ShaderRHI, RecordConeData, TemporaryIrradianceCacheResources->ConeData.SRV);
+
+		IrradianceCacheIrradiance.SetBuffer(RHICmdList, ShaderRHI, SurfaceCacheResources.Level[DepthLevel]->Irradiance);
+
+		FAOSampleData2 AOSampleData;
+
+		TArray<FVector, TInlineAllocator<9> > SampleDirections;
+		GetSpacedVectors(SampleDirections);
+
+		for (int32 SampleIndex = 0; SampleIndex < NumConeSampleDirections; SampleIndex++)
+		{
+			AOSampleData.SampleDirections[SampleIndex] = FVector4(SampleDirections[SampleIndex]);
+		}
+
+		SetUniformBufferParameterImmediate(RHICmdList, ShaderRHI, GetUniformBufferParameter<FAOSampleData2>(), AOSampleData);
+
+		FVector UnoccludedVector(0);
+
+		for (int32 SampleIndex = 0; SampleIndex < NumConeSampleDirections; SampleIndex++)
+		{
+			UnoccludedVector += SampleDirections[SampleIndex];
+		}
+
+		float BentNormalNormalizeFactorValue = 1.0f / (UnoccludedVector / NumConeSampleDirections).Size();
+		SetShaderValue(RHICmdList, ShaderRHI, BentNormalNormalizeFactor, BentNormalNormalizeFactorValue);
+
+		GlobalObjectParameters.Set(RHICmdList, ShaderRHI, *(Scene->DistanceFieldSceneData.ObjectBuffers), Scene->DistanceFieldSceneData.NumObjectsInBuffer);
+
+		FVPLResources& SourceVPLResources = GVPLViewCulling ? GCulledVPLResources : GVPLResources;
+		SetSRVParameter(RHICmdList, ShaderRHI, VPLParameterBuffer, SourceVPLResources.VPLParameterBuffer.SRV);
+		SetSRVParameter(RHICmdList, ShaderRHI, VPLClusterData, SourceVPLResources.VPLClusterData.SRV);
+		SetSRVParameter(RHICmdList, ShaderRHI, VPLData, SourceVPLResources.VPLData.SRV);
+
+		SetShaderValue(RHICmdList, ShaderRHI, VPLGatherRadius, Parameters.OcclusionMaxDistance);
+		extern float GAOConeHalfAngle;
+		SetShaderValue(RHICmdList, ShaderRHI, ConeHalfAngle, GAOConeHalfAngle);
+
+		FVPLTileIntersectionResources* VPLTileIntersectionResources = ((FSceneViewState*)View.State)->VPLTileIntersectionResources;
+
+		SetShaderValue(RHICmdList, ShaderRHI, TileListGroupSize, VPLTileIntersectionResources->TileDimensions);
+		SetSRVParameter(RHICmdList, ShaderRHI, TileHeadDataUnpacked, VPLTileIntersectionResources->TileHeadDataUnpacked.SRV);
+		SetSRVParameter(RHICmdList, ShaderRHI, TileArrayData, VPLTileIntersectionResources->TileArrayData.SRV);
+	}
+
+	void UnsetParameters(FRHICommandList& RHICmdList)
+	{
+		IrradianceCacheIrradiance.UnsetUAV(RHICmdList, GetComputeShader());
+	}
+
+	virtual bool Serialize(FArchive& Ar)
+	{		
+		bool bShaderHasOutdatedParameters = FGlobalShader::Serialize(Ar);
+		Ar << AOParameters;
+		Ar << IrradianceCacheIrradiance;
+		Ar << IrradianceCachePositionRadius;
+		Ar << IrradianceCacheNormal;
+		Ar << IrradianceCacheTileCoordinate;
+		Ar << RecordConeData;
+		Ar << ScatterDrawParameters;
+		Ar << SavedStartIndex;
+		Ar << DistanceFieldNormalTexture;
+		Ar << DistanceFieldNormalSampler;
+		Ar << GlobalObjectParameters;
+		Ar << VPLParameterBuffer;
+		Ar << VPLClusterData;
+		Ar << VPLData;
+		Ar << VPLGatherRadius;
+		Ar << ConeHalfAngle;
+		Ar << BentNormalNormalizeFactor;
+		Ar << TileListGroupSize;
+		Ar << TileHeadDataUnpacked;
+		Ar << TileArrayData;
+		return bShaderHasOutdatedParameters;
+	}
+
+private:
+
+	FAOParameters AOParameters;
+	FRWShaderParameter IrradianceCacheIrradiance;
+	FShaderResourceParameter IrradianceCachePositionRadius;
+	FShaderResourceParameter IrradianceCacheNormal;
+	FShaderResourceParameter IrradianceCacheTileCoordinate;
+	FShaderResourceParameter RecordConeData;
+	FShaderResourceParameter ScatterDrawParameters;
+	FShaderResourceParameter SavedStartIndex;
+	FShaderResourceParameter DistanceFieldNormalTexture;
+	FShaderResourceParameter DistanceFieldNormalSampler;
+	FDistanceFieldObjectBufferParameters GlobalObjectParameters;
+	FShaderResourceParameter VPLParameterBuffer;
+	FShaderResourceParameter VPLClusterData;
+	FShaderResourceParameter VPLData;
+	FShaderParameter VPLGatherRadius;
+	FShaderParameter ConeHalfAngle;
+	FShaderParameter BentNormalNormalizeFactor;
+	FShaderParameter TileListGroupSize;
+	FShaderResourceParameter TileHeadDataUnpacked;
+	FShaderResourceParameter TileArrayData;
+};
+
+IMPLEMENT_SHADER_TYPE(template<>,TComputeIrradianceCS<IrradianceMode_IrradianceFromAOCones>,TEXT("DistanceFieldGlobalIllumination"),TEXT("ComputeIrradianceCS"),SF_Compute);
+IMPLEMENT_SHADER_TYPE(template<>,TComputeIrradianceCS<IrradianceMode_IrradianceFinalGather>,TEXT("DistanceFieldGlobalIllumination"),TEXT("ComputeIrradianceCS"),SF_Compute);
+
+// Must match usf
+const int32 MaxVPLsPerTile = 1024;
+
+void FVPLTileIntersectionResources::InitDynamicRHI()
+{
+	TileHeadDataUnpacked.Initialize(sizeof(uint32), TileDimensions.X * TileDimensions.Y * 2, PF_R32_UINT, BUF_Static);
+
+	//@todo - handle max exceeded
+	TileArrayData.Initialize(sizeof(uint32), MaxVPLsPerTile * TileDimensions.X * TileDimensions.Y, PF_R32_UINT, BUF_Static);
+	TileArrayNextAllocation.Initialize(sizeof(uint32), 1, PF_R32_UINT, BUF_Static);
+}
+
+TScopedPointer<FLightTileIntersectionResources> GVPLPlacementTileIntersectionResources;
+
+void PlaceVPLs(
+	FRHICommandListImmediate& RHICmdList,
+	const FViewInfo& View,
+	const FScene* Scene,
+	const FDistanceFieldAOParameters& Parameters)
+{
+	GVPLResources.AllocateFor(GVPLGridDimension * GVPLGridDimension);
+
+	uint32 ClearValues[4] = { 0 };
+	RHICmdList.ClearUAV(GVPLResources.VPLParameterBuffer.UAV, ClearValues);
+
+	FVPLTileIntersectionResources*& VPLTileIntersectionResources = ((FSceneViewState*)View.State)->VPLTileIntersectionResources;
+
+	{
+		uint32 GroupSizeX = (View.ViewRect.Size().X / GAODownsampleFactor + GDistanceFieldAOTileSizeX - 1) / GDistanceFieldAOTileSizeX;
+		uint32 GroupSizeY = (View.ViewRect.Size().Y / GAODownsampleFactor + GDistanceFieldAOTileSizeY - 1) / GDistanceFieldAOTileSizeY;
+
+		if (!VPLTileIntersectionResources || VPLTileIntersectionResources->TileDimensions != FIntPoint(GroupSizeX, GroupSizeY))
+		{
+			if (VPLTileIntersectionResources)
+			{
+				VPLTileIntersectionResources->ReleaseResource();
+			}
+			else
+			{
+				VPLTileIntersectionResources = new FVPLTileIntersectionResources();
+			}
+
+			VPLTileIntersectionResources->TileDimensions = FIntPoint(GroupSizeX, GroupSizeY);
+
+			VPLTileIntersectionResources->InitResource();
+		}
+
+		// Indicates the clear value for each channel of the UAV format
+		uint32 ClearValues[4] = { 0 };
+		RHICmdList.ClearUAV(VPLTileIntersectionResources->TileHeadDataUnpacked.UAV, ClearValues);
+	}
+
+	const FLightSceneProxy* DirectionalLightProxy = NULL;
+
+	for (TSparseArray<FLightSceneInfoCompact>::TConstIterator LightIt(Scene->Lights); LightIt; ++LightIt)
+	{
+		const FLightSceneInfoCompact& LightSceneInfoCompact = *LightIt;
+		const FLightSceneInfo* const LightSceneInfo = LightSceneInfoCompact.LightSceneInfo;
+
+		if (LightSceneInfo->ShouldRenderLightViewIndependent() 
+			&& LightSceneInfo->Proxy->GetLightType() == LightType_Directional
+			&& LightSceneInfo->Proxy->CastsDynamicShadow())
+		{
+			DirectionalLightProxy = LightSceneInfo->Proxy;
+			break;
+		}
+	}
+
+	if (DirectionalLightProxy)
+	{
+		SCOPED_DRAW_EVENT(RHICmdList, VPLPlacement);
+		FMatrix DirectionalLightShadowToWorld;
+
+		{
+			int32 NumPlanes = 0;
+			const FPlane* PlaneData = NULL;
+			FVector4 ShadowBoundingSphereValue(0, 0, 0, 0);
+			FShadowCascadeSettings CascadeSettings;
+			FSphere ShadowBounds;
+			FConvexVolume FrustumVolume;
+
+			static bool bUseShadowmapBounds = true;
+
+			if (bUseShadowmapBounds)
+			{
+				ShadowBounds = DirectionalLightProxy->GetShadowSplitBoundsDepthRange(
+					View, 
+					View.NearClippingDistance, 
+					GVPLPlacementCameraRadius, 
+					&CascadeSettings);
+
+				FSphere SubjectBounds(FVector::ZeroVector, ShadowBounds.W);
+
+				const FMatrix& WorldToLight = DirectionalLightProxy->GetWorldToLight();
+				const FMatrix InitializerWorldToLight = FInverseRotationMatrix(FVector(WorldToLight.M[0][0],WorldToLight.M[1][0],WorldToLight.M[2][0]).GetSafeNormal().Rotation());
+				const FVector InitializerFaceDirection = FVector(1,0,0);
+
+				FVector	XAxis, YAxis;
+				InitializerFaceDirection.FindBestAxisVectors(XAxis, YAxis);
+				const FMatrix WorldToLightScaled = InitializerWorldToLight * FScaleMatrix(FVector(1.0f,1.0f / SubjectBounds.W,1.0f / SubjectBounds.W));
+				const FMatrix WorldToFace = WorldToLightScaled * FBasisVectorMatrix(-XAxis,YAxis,InitializerFaceDirection.GetSafeNormal(),FVector::ZeroVector);
+
+				static bool bSnapPosition = true;
+
+				if (bSnapPosition)
+				{
+					// Transform the shadow's position into shadowmap space
+					const FVector TransformedPosition = WorldToFace.TransformPosition(ShadowBounds.Center);
+
+					// Determine the distance necessary to snap the shadow's position to the nearest texel
+					const float SnapX = FMath::Fmod(TransformedPosition.X, 2.0f / GVPLGridDimension);
+					const float SnapY = FMath::Fmod(TransformedPosition.Y, 2.0f / GVPLGridDimension);
+					// Snap the shadow's position and transform it back into world space
+					// This snapping prevents sub-texel camera movements which removes view dependent aliasing from the final shadow result
+					// This only maintains stable shadows under camera translation and rotation
+					const FVector SnappedWorldPosition = WorldToFace.InverseFast().TransformPosition(TransformedPosition - FVector(SnapX, SnapY, 0.0f));
+					ShadowBounds.Center = SnappedWorldPosition;
+				}
+
+				NumPlanes = CascadeSettings.ShadowBoundsAccurate.Planes.Num();
+				PlaneData = CascadeSettings.ShadowBoundsAccurate.Planes.GetData();
+
+				DirectionalLightShadowToWorld = FTranslationMatrix(-ShadowBounds.Center) * WorldToFace * FShadowProjectionMatrix(-GVPLDirectionalLightTraceDistance / 2, GVPLDirectionalLightTraceDistance / 2, FVector4(0,0,0,1));
+			}
+			else
+			{
+				ShadowBounds = FSphere(View.ViewMatrices.ViewOrigin, GVPLPlacementCameraRadius);
+
+				FSphere SubjectBounds(FVector::ZeroVector, ShadowBounds.W);
+
+				const FMatrix& WorldToLight = DirectionalLightProxy->GetWorldToLight();
+				const FMatrix InitializerWorldToLight = FInverseRotationMatrix(FVector(WorldToLight.M[0][0],WorldToLight.M[1][0],WorldToLight.M[2][0]).GetSafeNormal().Rotation());
+				const FVector InitializerFaceDirection = FVector(1,0,0);
+
+				FVector	XAxis, YAxis;
+				InitializerFaceDirection.FindBestAxisVectors(XAxis, YAxis);
+				const FMatrix WorldToLightScaled = InitializerWorldToLight * FScaleMatrix(FVector(1.0f,1.0f / GVPLPlacementCameraRadius,1.0f / GVPLPlacementCameraRadius));
+				const FMatrix WorldToFace = WorldToLightScaled * FBasisVectorMatrix(-XAxis,YAxis,InitializerFaceDirection.GetSafeNormal(),FVector::ZeroVector);
+
+				static bool bSnapPosition = true;
+
+				if (bSnapPosition)
+				{
+					// Transform the shadow's position into shadowmap space
+					const FVector TransformedPosition = WorldToFace.TransformPosition(ShadowBounds.Center);
+
+					// Determine the distance necessary to snap the shadow's position to the nearest texel
+					const float SnapX = FMath::Fmod(TransformedPosition.X, 2.0f / GVPLGridDimension);
+					const float SnapY = FMath::Fmod(TransformedPosition.Y, 2.0f / GVPLGridDimension);
+					// Snap the shadow's position and transform it back into world space
+					// This snapping prevents sub-texel camera movements which removes view dependent aliasing from the final shadow result
+					// This only maintains stable shadows under camera translation and rotation
+					const FVector SnappedWorldPosition = WorldToFace.InverseFast().TransformPosition(TransformedPosition - FVector(SnapX, SnapY, 0.0f));
+					ShadowBounds.Center = SnappedWorldPosition;
+				}
+
+				const float MaxSubjectZ = WorldToFace.TransformPosition(SubjectBounds.Center).Z + SubjectBounds.W;
+				const float MinSubjectZ = FMath::Max(MaxSubjectZ - SubjectBounds.W * 2, (float)-HALF_WORLD_MAX);
+
+				DirectionalLightShadowToWorld = FTranslationMatrix(-ShadowBounds.Center) * WorldToFace * FShadowProjectionMatrix(MinSubjectZ, MaxSubjectZ, FVector4(0,0,0,1));
+
+				GetViewFrustumBounds(FrustumVolume, DirectionalLightShadowToWorld, true);
+
+				NumPlanes = FrustumVolume.Planes.Num();
+				PlaneData = FrustumVolume.Planes.GetData();
+			}
+
+			CullDistanceFieldObjectsForLight(
+				RHICmdList,
+				View,
+				DirectionalLightProxy, 
+				DirectionalLightShadowToWorld,
+				NumPlanes,
+				PlaneData,
+				ShadowBoundingSphereValue,
+				ShadowBounds.W,
+				GVPLPlacementTileIntersectionResources);
+		}
+
+		{
+			SCOPED_DRAW_EVENT(RHICmdList, PlaceVPLs);
+
+			TShaderMapRef<FVPLPlacementCS> ComputeShader(View.ShaderMap);
+
+			RHICmdList.SetComputeShader(ComputeShader->GetComputeShader());
+			ComputeShader->SetParameters(RHICmdList, View, DirectionalLightProxy, FVector2D(1.0f / GVPLGridDimension, 1.0f / GVPLGridDimension), DirectionalLightShadowToWorld, DirectionalLightShadowToWorld.InverseFast(), GVPLPlacementTileIntersectionResources);
+			DispatchComputeShader(RHICmdList, *ComputeShader, FMath::DivideAndRoundUp<int32>(GVPLGridDimension, GDistanceFieldAOTileSizeX), FMath::DivideAndRoundUp<int32>(GVPLGridDimension, GDistanceFieldAOTileSizeY), 1);
+
+			ComputeShader->UnsetParameters(RHICmdList);
+		}
+		
+		if (GVPLViewCulling)
+		{
+			{
+				TShaderMapRef<FSetupVPLCullndirectArgumentsCS> ComputeShader(View.ShaderMap);
+				RHICmdList.SetComputeShader(ComputeShader->GetComputeShader());
+				ComputeShader->SetParameters(RHICmdList, View);
+
+				DispatchComputeShader(RHICmdList, *ComputeShader, 1, 1, 1);
+				ComputeShader->UnsetParameters(RHICmdList);
+			}
+
+			{
+				GCulledVPLResources.AllocateFor(GVPLGridDimension * GVPLGridDimension);
+
+				uint32 ClearValues[4] = { 0 };
+				RHICmdList.ClearUAV(GCulledVPLResources.VPLParameterBuffer.UAV, ClearValues);
+
+				TShaderMapRef<FCullVPLsForViewCS> ComputeShader(GetGlobalShaderMap(Scene->GetFeatureLevel()));
+				RHICmdList.SetComputeShader(ComputeShader->GetComputeShader());
+				ComputeShader->SetParameters(RHICmdList, Scene, View, Parameters);
+
+				DispatchIndirectComputeShader(RHICmdList, *ComputeShader, GVPLResources.VPLDispatchIndirectBuffer.Buffer, 0);
+				ComputeShader->UnsetParameters(RHICmdList);
+			}
+		}
+		/*
+		{
+			SCOPED_DRAW_EVENT(RHICmdList, TileCulling);
+
+			TShaderMapRef<FVPLTileCullCS> ComputeShader(View.ShaderMap);
+
+			RHICmdList.SetComputeShader(ComputeShader->GetComputeShader());
+			ComputeShader->SetParameters(RHICmdList, View, Scene, Parameters, GVPLViewCulling ? GCulledVPLResources : GVPLResources);
+			DispatchComputeShader(RHICmdList, *ComputeShader, VPLTileIntersectionResources->TileDimensions.X, VPLTileIntersectionResources->TileDimensions.Y, 1);
+
+			ComputeShader->UnsetParameters(RHICmdList);
+		}*/
+	}
+}
+
+void ComputeIrradianceForSamples(
+	int32 DepthLevel,
+	FRHICommandListImmediate& RHICmdList,
+	const FViewInfo& View,
+	const FScene* Scene,
+	const FDistanceFieldAOParameters& Parameters,
+	FTemporaryIrradianceCacheResources* TemporaryIrradianceCacheResources)
+{
+	SCOPED_DRAW_EVENT(RHICmdList, ComputeIrradiance);
+
+	FSurfaceCacheResources& SurfaceCacheResources = *Scene->SurfaceCacheResources;
+
+	if (GAOUseConesForGI)
+	{
+		TShaderMapRef<TComputeIrradianceCS<IrradianceMode_IrradianceFromAOCones> > ComputeShader(View.ShaderMap);
+
+		RHICmdList.SetComputeShader(ComputeShader->GetComputeShader());
+		ComputeShader->SetParameters(RHICmdList, View, DepthLevel, Parameters, TemporaryIrradianceCacheResources);
+		DispatchIndirectComputeShader(RHICmdList, *ComputeShader, SurfaceCacheResources.DispatchParameters.Buffer, 0);
+
+		ComputeShader->UnsetParameters(RHICmdList);
+	}
+	else
+	{
+		TShaderMapRef<TComputeIrradianceCS<IrradianceMode_IrradianceFinalGather> > ComputeShader(View.ShaderMap);
+
+		RHICmdList.SetComputeShader(ComputeShader->GetComputeShader());
+		ComputeShader->SetParameters(RHICmdList, View, DepthLevel, Parameters, TemporaryIrradianceCacheResources);
+		DispatchIndirectComputeShader(RHICmdList, *ComputeShader, SurfaceCacheResources.DispatchParameters.Buffer, 0);
+
+		ComputeShader->UnsetParameters(RHICmdList);
+	}
+}
\ No newline at end of file
diff --git a/Engine/Source/Runtime/Renderer/Private/DistanceFieldGlobalIllumination.h b/Engine/Source/Runtime/Renderer/Private/DistanceFieldGlobalIllumination.h
new file mode 100644
--- /dev/null
+++ b/Engine/Source/Runtime/Renderer/Private/DistanceFieldGlobalIllumination.h
@@ -0,0 +1,70 @@
+// Copyright 1998-2014 Epic Games, Inc. All Rights Reserved.
+
+/*=============================================================================
+	DistanceFieldGlobalIllumination.h
+=============================================================================*/
+
+#pragma once
+
+#include "RenderResource.h"
+
+// In float4's, must match usf
+const int32 VPLDataStride = 3;
+
+// Must match usf
+const int32 VPLClusterSizeOneDim = 4;
+
+/**  */
+class FVPLResources : public FMaxSizedRWBuffers
+{
+public:
+
+	virtual void InitDynamicRHI()
+	{
+		if (MaxSize > 0)
+		{
+			VPLParameterBuffer.Initialize(sizeof(uint32), 2, PF_R32_UINT, BUF_Static);
+			VPLDispatchIndirectBuffer.Initialize(sizeof(uint32), 3, PF_R32_UINT, BUF_Static | BUF_DrawIndirect);
+			VPLClusterData.Initialize(sizeof(FVector4), MaxSize * VPLDataStride / (VPLClusterSizeOneDim * VPLClusterSizeOneDim), PF_A32B32G32R32F, BUF_Static);
+			VPLData.Initialize(sizeof(FVector4), MaxSize * VPLDataStride, PF_A32B32G32R32F, BUF_Static);
+		}
+	}
+
+	virtual void ReleaseDynamicRHI()
+	{
+		VPLParameterBuffer.Release();
+		VPLDispatchIndirectBuffer.Release();
+		VPLClusterData.Release();
+		VPLData.Release();
+	}
+
+	FRWBuffer VPLParameterBuffer;
+	FRWBuffer VPLDispatchIndirectBuffer;
+	FRWBuffer VPLClusterData;
+	FRWBuffer VPLData;
+};
+
+/**  */
+class FVPLTileIntersectionResources : public FRenderResource
+{
+public:
+
+	FVPLTileIntersectionResources()
+		: TileDimensions(FIntPoint(0, 0))
+	{}
+
+	virtual void InitDynamicRHI() override;
+
+	virtual void ReleaseDynamicRHI() override
+	{
+		TileHeadDataUnpacked.Release();
+		TileArrayData.Release();
+		TileArrayNextAllocation.Release();
+	}
+
+	FIntPoint TileDimensions;
+
+	FRWBuffer TileHeadDataUnpacked;
+	FRWBuffer TileArrayData;
+	FRWBuffer TileArrayNextAllocation;
+};
\ No newline at end of file
diff --git a/Engine/Source/Runtime/Renderer/Private/DistanceFieldLightingShared.h b/Engine/Source/Runtime/Renderer/Private/DistanceFieldLightingShared.h
--- a/Engine/Source/Runtime/Renderer/Private/DistanceFieldLightingShared.h
+++ b/Engine/Source/Runtime/Renderer/Private/DistanceFieldLightingShared.h
@@ -13,6 +13,7 @@ const int32 GDistanceFieldAOTileSizeX = 16;
 const int32 GDistanceFieldAOTileSizeY = 16;
 /** Base downsample factor that all distance field AO operations are done at. */
 const int32 GAODownsampleFactor = 2;
+/** Must match usf */
 static const int32 GMaxNumObjectsPerTile = 512;
 extern uint32 UpdateObjectsGroupSize;
 
@@ -61,6 +62,9 @@ public:
 		ObjectBounds.Bind(ParameterMap, TEXT("ObjectBounds"));
 		ObjectData.Bind(ParameterMap, TEXT("ObjectData"));
 		NumSceneObjects.Bind(ParameterMap, TEXT("NumSceneObjects"));
+		DistanceFieldTexture.Bind(ParameterMap, TEXT("DistanceFieldTexture"));
+		DistanceFieldSampler.Bind(ParameterMap, TEXT("DistanceFieldSampler"));
+		DistanceFieldAtlasTexelSize.Bind(ParameterMap, TEXT("DistanceFieldAtlasTexelSize"));
 	}
 
 	template<typename TParamRef>
@@ -69,6 +73,21 @@ public:
 		ObjectBounds.SetBuffer(RHICmdList, ShaderRHI, ObjectBuffers.Bounds);
 		ObjectData.SetBuffer(RHICmdList, ShaderRHI, ObjectBuffers.Data);
 		SetShaderValue(RHICmdList, ShaderRHI, NumSceneObjects, NumObjectsValue);
+
+		SetTextureParameter(
+			RHICmdList,
+			ShaderRHI,
+			DistanceFieldTexture,
+			DistanceFieldSampler,
+			TStaticSamplerState<SF_Bilinear,AM_Clamp,AM_Clamp,AM_Clamp>::GetRHI(),
+			GDistanceFieldVolumeTextureAtlas.VolumeTextureRHI
+			);
+
+		const int32 NumTexelsOneDimX = GDistanceFieldVolumeTextureAtlas.GetSizeX();
+		const int32 NumTexelsOneDimY = GDistanceFieldVolumeTextureAtlas.GetSizeY();
+		const int32 NumTexelsOneDimZ = GDistanceFieldVolumeTextureAtlas.GetSizeZ();
+		const FVector InvTextureDim(1.0f / NumTexelsOneDimX, 1.0f / NumTexelsOneDimY, 1.0f / NumTexelsOneDimZ);
+		SetShaderValue(RHICmdList, ShaderRHI, DistanceFieldAtlasTexelSize, InvTextureDim);
 	}
 
 	template<typename TParamRef>
@@ -80,7 +99,7 @@ public:
 
 	friend FArchive& operator<<(FArchive& Ar, FDistanceFieldObjectBufferParameters& P)
 	{
-		Ar << P.ObjectBounds << P.ObjectData << P.NumSceneObjects;
+		Ar << P.ObjectBounds << P.ObjectData << P.NumSceneObjects << P.DistanceFieldTexture << P.DistanceFieldSampler << P.DistanceFieldAtlasTexelSize;
 		return Ar;
 	}
 
@@ -88,6 +107,9 @@ private:
 	FRWShaderParameter ObjectBounds;
 	FRWShaderParameter ObjectData;
 	FShaderParameter NumSceneObjects;
+	FShaderResourceParameter DistanceFieldTexture;
+	FShaderResourceParameter DistanceFieldSampler;
+	FShaderParameter DistanceFieldAtlasTexelSize;
 };
 
 class FDistanceFieldCulledObjectBuffers
@@ -265,7 +287,7 @@ public:
 		TileHeadDataUnpacked.Initialize(sizeof(uint32), TileDimensions.X * TileDimensions.Y * 2, PF_R32_UINT, BUF_Static);
 
 		//@todo - handle max exceeded
-		TileArrayData.Initialize(sizeof(uint16), GMaxNumObjectsPerTile * TileDimensions.X * TileDimensions.Y, PF_R16_UINT, BUF_Static);
+		TileArrayData.Initialize(sizeof(uint32), GMaxNumObjectsPerTile * TileDimensions.X * TileDimensions.Y, PF_R32_UINT, BUF_Static);
 	}
 
 	void Release()
@@ -279,3 +301,14 @@ public:
 	FRWBuffer TileHeadDataUnpacked;
 	FRWBuffer TileArrayData;
 };
+
+extern void CullDistanceFieldObjectsForLight(
+	FRHICommandListImmediate& RHICmdList,
+	const FViewInfo& View,
+	const FLightSceneProxy* LightSceneProxy, 
+	const FMatrix& WorldToShadowValue, 
+	int32 NumPlanes, 
+	const FPlane* PlaneData, 
+	const FVector4& ShadowBoundingSphereValue,
+	float ShadowBoundingRadius,
+	TScopedPointer<class FLightTileIntersectionResources>& TileIntersectionResources);
\ No newline at end of file
diff --git a/Engine/Source/Runtime/Renderer/Private/DistanceFieldShadowing.cpp b/Engine/Source/Runtime/Renderer/Private/DistanceFieldShadowing.cpp
--- a/Engine/Source/Runtime/Renderer/Private/DistanceFieldShadowing.cpp
+++ b/Engine/Source/Runtime/Renderer/Private/DistanceFieldShadowing.cpp
@@ -77,7 +77,7 @@ public:
 	{
 	}
 
-	void SetParameters(FRHICommandList& RHICmdList, const FScene* Scene, const FSceneView& View, const FProjectedShadowInfo* ProjectedShadowInfo)
+	void SetParameters(FRHICommandList& RHICmdList, const FScene* Scene, const FSceneView& View, const FMatrix& WorldToShadowValue, int32 NumPlanes, const FPlane* PlaneData, const FVector4& ShadowBoundingSphereValue)
 	{
 		FComputeShaderRHIParamRef ShaderRHI = GetComputeShader();
 		FGlobalShader::SetParameters(RHICmdList, ShaderRHI, View);
@@ -89,28 +89,8 @@ public:
 		CulledObjectBoxBounds.SetBuffer(RHICmdList, ShaderRHI, GShadowCulledObjectBuffers.Buffers.BoxBounds);
 
 		SetShaderValue(RHICmdList, ShaderRHI, ObjectBoundingGeometryIndexCount, StencilingGeometry::GLowPolyStencilSphereIndexBuffer.GetIndexCount());
-		const FMatrix WorldToShadowValue = FTranslationMatrix(ProjectedShadowInfo->PreShadowTranslation) * ProjectedShadowInfo->SubjectAndReceiverMatrix;
-		SetShaderValue(RHICmdList, ShaderRHI, WorldToShadow, WorldToShadowValue);
-
-		int32 NumPlanes = 0;
-		const FPlane* PlaneData = NULL;
-		FVector4 ShadowBoundingSphereValue(0, 0, 0, 0);
 
-		if (ProjectedShadowInfo->bDirectionalLight)
-		{
-			NumPlanes = ProjectedShadowInfo->CascadeSettings.ShadowBoundsAccurate.Planes.Num();
-			PlaneData = ProjectedShadowInfo->CascadeSettings.ShadowBoundsAccurate.Planes.GetData();
-		}
-		else if (ProjectedShadowInfo->bOnePassPointLightShadow)
-		{
-			ShadowBoundingSphereValue = FVector4(ProjectedShadowInfo->ShadowBounds.Center.X, ProjectedShadowInfo->ShadowBounds.Center.Y, ProjectedShadowInfo->ShadowBounds.Center.Z, ProjectedShadowInfo->ShadowBounds.W);
-		}
-		else
-		{
-			NumPlanes = ProjectedShadowInfo->CasterFrustum.Planes.Num();
-			PlaneData = ProjectedShadowInfo->CasterFrustum.Planes.GetData();
-			ShadowBoundingSphereValue = FVector4(ProjectedShadowInfo->PreShadowTranslation, 0);
-		}
+		SetShaderValue(RHICmdList, ShaderRHI, WorldToShadow, WorldToShadowValue);
 
 		check(NumPlanes < 12);
 		SetShaderValue(RHICmdList, ShaderRHI, NumShadowHullPlanes, NumPlanes);
@@ -181,7 +161,7 @@ public:
 	FClearTilesCS(const ShaderMetaType::CompiledShaderInitializerType& Initializer)
 		: FGlobalShader(Initializer)
 	{
-		TileHeadDataUnpacked.Bind(Initializer.ParameterMap, TEXT("TileHeadDataUnpacked"));
+		ShadowTileHeadDataUnpacked.Bind(Initializer.ParameterMap, TEXT("ShadowTileHeadDataUnpacked"));
 		NumGroups.Bind(Initializer.ParameterMap, TEXT("NumGroups"));
 	}
 
@@ -194,27 +174,27 @@ public:
 		FComputeShaderRHIParamRef ShaderRHI = GetComputeShader();
 
 		FGlobalShader::SetParameters(RHICmdList, ShaderRHI, View);
-		TileHeadDataUnpacked.SetBuffer(RHICmdList, ShaderRHI, TileIntersectionResources->TileHeadDataUnpacked);
+		ShadowTileHeadDataUnpacked.SetBuffer(RHICmdList, ShaderRHI, TileIntersectionResources->TileHeadDataUnpacked);
 
 		SetShaderValue(RHICmdList, ShaderRHI, NumGroups, NumGroupsValue);
 	}
 
 	void UnsetParameters(FRHICommandList& RHICmdList)
 	{
-		TileHeadDataUnpacked.UnsetUAV(RHICmdList, GetComputeShader());
+		ShadowTileHeadDataUnpacked.UnsetUAV(RHICmdList, GetComputeShader());
 	}
 
 	virtual bool Serialize(FArchive& Ar)
 	{		
 		bool bShaderHasOutdatedParameters = FGlobalShader::Serialize(Ar);
-		Ar << TileHeadDataUnpacked;
+		Ar << ShadowTileHeadDataUnpacked;
 		Ar << NumGroups;
 		return bShaderHasOutdatedParameters;
 	}
 
 private:
 
-	FRWShaderParameter TileHeadDataUnpacked;
+	FRWShaderParameter ShadowTileHeadDataUnpacked;
 	FShaderParameter NumGroups;
 };
 
@@ -248,7 +228,7 @@ public:
 
 	FShadowObjectCullVS() {}
 
-	void SetParameters(FRHICommandList& RHICmdList, const FSceneView& View, FVector2D NumGroupsValue, const FProjectedShadowInfo* ProjectedShadowInfo)
+	void SetParameters(FRHICommandList& RHICmdList, const FSceneView& View, FVector2D NumGroupsValue, const FMatrix& WorldToShadowMatrixValue, float ShadowRadius)
 	{
 		const FVertexShaderRHIParamRef ShaderRHI = GetVertexShader();
 		FGlobalShader::SetParameters(RHICmdList, ShaderRHI, View);
@@ -263,10 +243,9 @@ public:
 
 		SetShaderValue(RHICmdList, ShaderRHI, ConservativeRadiusScale, ConservativeRadiusScaleValue);
 
-		FMatrix WorldToShadowMatrixValue = FTranslationMatrix(ProjectedShadowInfo->PreShadowTranslation) * ProjectedShadowInfo->SubjectAndReceiverMatrix;
 		SetShaderValue(RHICmdList, ShaderRHI, WorldToShadow, WorldToShadowMatrixValue);
 
-		float MinRadiusValue = 2 * ProjectedShadowInfo->ShadowBounds.W / FMath::Min(NumGroupsValue.X, NumGroupsValue.Y);
+		float MinRadiusValue = 2 * ShadowRadius / FMath::Min(NumGroupsValue.X, NumGroupsValue.Y);
 		SetShaderValue(RHICmdList, ShaderRHI, MinRadius, MinRadiusValue);
 	}
 
@@ -304,7 +283,6 @@ public:
 		FGlobalShader::ModifyCompilationEnvironment(Platform,OutEnvironment);
 		
 		OutEnvironment.SetDefine(TEXT("DOWNSAMPLE_FACTOR"), GAODownsampleFactor);
-		OutEnvironment.SetDefine(TEXT("MAX_OBJECTS_PER_TILE"), GMaxNumObjectsPerTile);
 	}
 
 	/** Default constructor. */
@@ -315,17 +293,15 @@ public:
 		: FGlobalShader(Initializer)
 	{
 		ObjectParameters.Bind(Initializer.ParameterMap);
-		TileHeadDataUnpacked.Bind(Initializer.ParameterMap, TEXT("TileHeadDataUnpacked"));
-		TileArrayData.Bind(Initializer.ParameterMap, TEXT("TileArrayData"));
+		ShadowTileHeadDataUnpacked.Bind(Initializer.ParameterMap, TEXT("ShadowTileHeadDataUnpacked"));
+		ShadowTileArrayData.Bind(Initializer.ParameterMap, TEXT("ShadowTileArrayData"));
 		NumGroups.Bind(Initializer.ParameterMap, TEXT("NumGroups"));
 	}
 
 	void SetParameters(
 		FRHICommandList& RHICmdList, 
 		const FSceneView& View, 
-		FVector2D NumGroupsValue, 
-		const FProjectedShadowInfo* ProjectedShadowInfo,
-		FLightTileIntersectionResources* TileIntersectionResources)
+		FVector2D NumGroupsValue)
 	{
 		const FPixelShaderRHIParamRef ShaderRHI = GetPixelShader();
 		FGlobalShader::SetParameters(RHICmdList, ShaderRHI, View);
@@ -336,17 +312,17 @@ public:
 
 	void GetUAVs(const FSceneView& View, FLightTileIntersectionResources* TileIntersectionResources, TArray<FUnorderedAccessViewRHIParamRef>& UAVs)
 	{
-		int32 MaxIndex = FMath::Max(TileHeadDataUnpacked.GetUAVIndex(), TileArrayData.GetUAVIndex());
+		int32 MaxIndex = FMath::Max(ShadowTileHeadDataUnpacked.GetUAVIndex(), ShadowTileArrayData.GetUAVIndex());
 		UAVs.AddZeroed(MaxIndex + 1);
 
-		if (TileHeadDataUnpacked.IsBound())
+		if (ShadowTileHeadDataUnpacked.IsBound())
 		{
-			UAVs[TileHeadDataUnpacked.GetUAVIndex()] = TileIntersectionResources->TileHeadDataUnpacked.UAV;
+			UAVs[ShadowTileHeadDataUnpacked.GetUAVIndex()] = TileIntersectionResources->TileHeadDataUnpacked.UAV;
 		}
 
-		if (TileArrayData.IsBound())
+		if (ShadowTileArrayData.IsBound())
 		{
-			UAVs[TileArrayData.GetUAVIndex()] = TileIntersectionResources->TileArrayData.UAV;
+			UAVs[ShadowTileArrayData.GetUAVIndex()] = TileIntersectionResources->TileArrayData.UAV;
 		}
 
 		check(UAVs.Num() > 0);
@@ -357,16 +333,16 @@ public:
 	{
 		bool bShaderHasOutdatedParameters = FGlobalShader::Serialize(Ar);
 		Ar << ObjectParameters;
-		Ar << TileHeadDataUnpacked;
-		Ar << TileArrayData;
+		Ar << ShadowTileHeadDataUnpacked;
+		Ar << ShadowTileArrayData;
 		Ar << NumGroups;
 		return bShaderHasOutdatedParameters;
 	}
 
 private:
 	FDistanceFieldCulledObjectBufferParameters ObjectParameters;
-	FRWShaderParameter TileHeadDataUnpacked;
-	FRWShaderParameter TileArrayData;
+	FRWShaderParameter ShadowTileHeadDataUnpacked;
+	FRWShaderParameter ShadowTileArrayData;
 	FShaderParameter NumGroups;
 };
 
@@ -397,7 +373,6 @@ public:
 		OutEnvironment.SetDefine(TEXT("DOWNSAMPLE_FACTOR"), GAODownsampleFactor);
 		OutEnvironment.SetDefine(TEXT("THREADGROUP_SIZEX"), GDistanceFieldAOTileSizeX);
 		OutEnvironment.SetDefine(TEXT("THREADGROUP_SIZEY"), GDistanceFieldAOTileSizeY);
-		OutEnvironment.SetDefine(TEXT("MAX_OBJECTS_PER_TILE"), GMaxNumObjectsPerTile);
 		OutEnvironment.SetDefine(TEXT("SCATTER_TILE_CULLING"), ShadowingType == DFS_DirectionalLightScatterTileCulling ? TEXT("1") : TEXT("0"));
 		OutEnvironment.SetDefine(TEXT("POINT_LIGHT"), ShadowingType == DFS_PointLightTiledCulling ? TEXT("1") : TEXT("0"));
 	}
@@ -418,9 +393,9 @@ public:
 		ScissorRectMinAndSize.Bind(Initializer.ParameterMap, TEXT("ScissorRectMinAndSize"));
 		ObjectParameters.Bind(Initializer.ParameterMap);
 		DeferredParameters.Bind(Initializer.ParameterMap);
-		TileHeadDataUnpacked.Bind(Initializer.ParameterMap, TEXT("TileHeadDataUnpacked"));
-		TileArrayData.Bind(Initializer.ParameterMap, TEXT("TileArrayData"));
-		TileListGroupSize.Bind(Initializer.ParameterMap, TEXT("TileListGroupSize"));
+		ShadowTileHeadDataUnpacked.Bind(Initializer.ParameterMap, TEXT("ShadowTileHeadDataUnpacked"));
+		ShadowTileArrayData.Bind(Initializer.ParameterMap, TEXT("ShadowTileArrayData"));
+		ShadowTileListGroupSize.Bind(Initializer.ParameterMap, TEXT("ShadowTileListGroupSize"));
 		WorldToShadow.Bind(Initializer.ParameterMap, TEXT("WorldToShadow"));
 	}
 
@@ -465,13 +440,13 @@ public:
 
 		SetShaderValue(RHICmdList, ShaderRHI, ScissorRectMinAndSize, FIntRect(ScissorRect.Min, ScissorRect.Size()));
 
-		check(TileIntersectionResources || !TileHeadDataUnpacked.IsBound());
+		check(TileIntersectionResources || !ShadowTileHeadDataUnpacked.IsBound());
 
 		if (TileIntersectionResources)
 		{
-			SetSRVParameter(RHICmdList, ShaderRHI, TileHeadDataUnpacked, TileIntersectionResources->TileHeadDataUnpacked.SRV);
-			SetSRVParameter(RHICmdList, ShaderRHI, TileArrayData, TileIntersectionResources->TileArrayData.SRV);
-			SetShaderValue(RHICmdList, ShaderRHI, TileListGroupSize, TileIntersectionResources->TileDimensions);
+			SetSRVParameter(RHICmdList, ShaderRHI, ShadowTileHeadDataUnpacked, TileIntersectionResources->TileHeadDataUnpacked.SRV);
+			SetSRVParameter(RHICmdList, ShaderRHI, ShadowTileArrayData, TileIntersectionResources->TileArrayData.SRV);
+			SetShaderValue(RHICmdList, ShaderRHI, ShadowTileListGroupSize, TileIntersectionResources->TileDimensions);
 		}
 
 		FMatrix WorldToShadowMatrixValue = FTranslationMatrix(ProjectedShadowInfo->PreShadowTranslation) * ProjectedShadowInfo->SubjectAndReceiverMatrix;
@@ -496,9 +471,9 @@ public:
 		Ar << ScissorRectMinAndSize;
 		Ar << ObjectParameters;
 		Ar << DeferredParameters;
-		Ar << TileHeadDataUnpacked;
-		Ar << TileArrayData;
-		Ar << TileListGroupSize;
+		Ar << ShadowTileHeadDataUnpacked;
+		Ar << ShadowTileArrayData;
+		Ar << ShadowTileListGroupSize;
 		Ar << WorldToShadow;
 		return bShaderHasOutdatedParameters;
 	}
@@ -514,9 +489,9 @@ private:
 	FShaderParameter ScissorRectMinAndSize;
 	FDistanceFieldCulledObjectBufferParameters ObjectParameters;
 	FDeferredPixelShaderParameters DeferredParameters;
-	FShaderResourceParameter TileHeadDataUnpacked;
-	FShaderResourceParameter TileArrayData;
-	FShaderParameter TileListGroupSize;
+	FShaderResourceParameter ShadowTileHeadDataUnpacked;
+	FShaderResourceParameter ShadowTileArrayData;
+	FShaderParameter ShadowTileListGroupSize;
 	FShaderParameter WorldToShadow;
 };
 
@@ -585,6 +560,118 @@ private:
 
 IMPLEMENT_SHADER_TYPE(,FDistanceFieldShadowingUpsamplePS,TEXT("DistanceFieldShadowing"),TEXT("DistanceFieldShadowingUpsamplePS"),SF_Pixel);
 
+void CullDistanceFieldObjectsForLight(
+	FRHICommandListImmediate& RHICmdList,
+	const FViewInfo& View,
+	const FLightSceneProxy* LightSceneProxy, 
+	const FMatrix& WorldToShadowValue, 
+	int32 NumPlanes, 
+	const FPlane* PlaneData, 
+	const FVector4& ShadowBoundingSphereValue,
+	float ShadowBoundingRadius,
+	TScopedPointer<FLightTileIntersectionResources>& TileIntersectionResources)
+{
+	const FScene* Scene = (const FScene*)(View.Family->Scene);
+
+	{
+		SCOPED_DRAW_EVENT(RHICmdList, ObjectFrustumCulling);
+
+		if (GShadowCulledObjectBuffers.Buffers.MaxObjects < Scene->DistanceFieldSceneData.NumObjectsInBuffer
+			|| GShadowCulledObjectBuffers.Buffers.MaxObjects > 3 * Scene->DistanceFieldSceneData.NumObjectsInBuffer)
+		{
+			GShadowCulledObjectBuffers.Buffers.bWantBoxBounds = true;
+			GShadowCulledObjectBuffers.Buffers.MaxObjects = Scene->DistanceFieldSceneData.NumObjectsInBuffer * 5 / 4;
+			GShadowCulledObjectBuffers.Buffers.Release();
+			GShadowCulledObjectBuffers.Buffers.Initialize();
+		}
+
+		{
+			uint32 ClearValues[4] = { 0 };
+			RHICmdList.ClearUAV(GShadowCulledObjectBuffers.Buffers.ObjectIndirectArguments.UAV, ClearValues);
+
+			TShaderMapRef<FCullObjectsForShadowCS> ComputeShader(GetGlobalShaderMap(Scene->GetFeatureLevel()));
+			RHICmdList.SetComputeShader(ComputeShader->GetComputeShader());
+			ComputeShader->SetParameters(RHICmdList, Scene, View, WorldToShadowValue, NumPlanes, PlaneData, ShadowBoundingSphereValue);
+
+			DispatchComputeShader(RHICmdList, *ComputeShader, FMath::DivideAndRoundUp<uint32>(Scene->DistanceFieldSceneData.NumObjectsInBuffer, UpdateObjectsGroupSize), 1, 1);
+			ComputeShader->UnsetParameters(RHICmdList);
+		}
+	}
+
+	// Allocate tile resolution based on world space size
+	//@todo - light space perspective shadow maps would make much better use of the resolution
+	const float LightTiles = FMath::Min(ShadowBoundingRadius / GShadowWorldTileSize + 1, 256.0f);
+	FIntPoint LightTileDimensions(LightTiles, LightTiles);
+
+	if (LightSceneProxy->GetLightType() == LightType_Directional && GShadowScatterTileCulling)
+	{
+		if (!TileIntersectionResources || TileIntersectionResources->TileDimensions != LightTileDimensions)
+		{
+			if (TileIntersectionResources)
+			{
+				TileIntersectionResources->Release();
+			}
+			else
+			{
+				TileIntersectionResources = new FLightTileIntersectionResources();
+			}
+
+			TileIntersectionResources->TileDimensions = LightTileDimensions;
+
+			TileIntersectionResources->Initialize();
+		}
+
+		{
+			SCOPED_DRAW_EVENT(RHICmdList, ClearTiles);
+			TShaderMapRef<FClearTilesCS> ComputeShader(View.ShaderMap);
+
+			uint32 GroupSizeX = FMath::DivideAndRoundUp(LightTileDimensions.X, GDistanceFieldAOTileSizeX);
+			uint32 GroupSizeY = FMath::DivideAndRoundUp(LightTileDimensions.Y, GDistanceFieldAOTileSizeY);
+
+			RHICmdList.SetComputeShader(ComputeShader->GetComputeShader());
+			ComputeShader->SetParameters(RHICmdList, View, FVector2D(LightTileDimensions.X, LightTileDimensions.Y), TileIntersectionResources);
+			DispatchComputeShader(RHICmdList, *ComputeShader, GroupSizeX, GroupSizeY, 1);
+
+			ComputeShader->UnsetParameters(RHICmdList);
+		}
+		
+		{
+			SCOPED_DRAW_EVENT(RHICmdList, CullObjects);
+
+			TShaderMapRef<FShadowObjectCullVS> VertexShader(View.ShaderMap);
+			TShaderMapRef<FShadowObjectCullPS> PixelShader(View.ShaderMap);
+
+			TArray<FUnorderedAccessViewRHIParamRef> UAVs;
+			PixelShader->GetUAVs(View, TileIntersectionResources, UAVs);
+			RHICmdList.SetRenderTargets(0, (const FRHIRenderTargetView *)NULL, NULL, UAVs.Num(), UAVs.GetData());
+
+			RHICmdList.SetViewport(0, 0, 0.0f, LightTileDimensions.X, LightTileDimensions.Y, 1.0f);
+
+			// Render backfaces since camera may intersect
+			RHICmdList.SetRasterizerState(View.bReverseCulling ? TStaticRasterizerState<FM_Solid, CM_CW>::GetRHI() : TStaticRasterizerState<FM_Solid, CM_CCW>::GetRHI());
+			RHICmdList.SetDepthStencilState(TStaticDepthStencilState<false, CF_Always>::GetRHI());
+			RHICmdList.SetBlendState(TStaticBlendState<>::GetRHI());
+
+			static FGlobalBoundShaderState BoundShaderState;
+
+			SetGlobalBoundShaderState(RHICmdList, View.GetFeatureLevel(), BoundShaderState, GetVertexDeclarationFVector4(), *VertexShader, *PixelShader);
+
+			VertexShader->SetParameters(RHICmdList, View, FVector2D(LightTileDimensions.X, LightTileDimensions.Y), WorldToShadowValue, ShadowBoundingRadius);
+			PixelShader->SetParameters(RHICmdList, View, FVector2D(LightTileDimensions.X, LightTileDimensions.Y));
+
+			RHICmdList.SetStreamSource(0, StencilingGeometry::GLowPolyStencilSphereVertexBuffer.VertexBufferRHI, sizeof(FVector4), 0);
+
+			RHICmdList.DrawIndexedPrimitiveIndirect(
+				PT_TriangleList,
+				StencilingGeometry::GLowPolyStencilSphereIndexBuffer.IndexBufferRHI, 
+				GShadowCulledObjectBuffers.Buffers.ObjectIndirectArguments.Buffer,
+				0);
+
+			SetRenderTarget(RHICmdList, NULL, NULL);
+		}
+	}
+}
+
 bool SupportsDistanceFieldShadows(ERHIFeatureLevel::Type FeatureLevel, EShaderPlatform ShaderPlatform)
 {
 	return GDistanceFieldShadowing
@@ -636,105 +723,44 @@ void FProjectedShadowInfo::RenderRayTracedDistanceFieldProjection(FRHICommandLis
 		{
 			check(!Scene->DistanceFieldSceneData.HasPendingOperations());
 
-			{
-				SCOPED_DRAW_EVENT(RHICmdList, ObjectFrustumCulling);
-
-				if (GShadowCulledObjectBuffers.Buffers.MaxObjects < Scene->DistanceFieldSceneData.NumObjectsInBuffer
-					|| GShadowCulledObjectBuffers.Buffers.MaxObjects > 3 * Scene->DistanceFieldSceneData.NumObjectsInBuffer)
-				{
-					GShadowCulledObjectBuffers.Buffers.bWantBoxBounds = true;
-					GShadowCulledObjectBuffers.Buffers.MaxObjects = Scene->DistanceFieldSceneData.NumObjectsInBuffer * 5 / 4;
-					GShadowCulledObjectBuffers.Buffers.Release();
-					GShadowCulledObjectBuffers.Buffers.Initialize();
-				}
-
-				{
-					TShaderMapRef<FCullObjectsForShadowCS> ComputeShader(GetGlobalShaderMap(Scene->GetFeatureLevel()));
-					RHICmdList.SetComputeShader(ComputeShader->GetComputeShader());
-					ComputeShader->SetParameters(RHICmdList, Scene, View, this);
-
-					DispatchComputeShader(RHICmdList, *ComputeShader, FMath::DivideAndRoundUp<uint32>(Scene->DistanceFieldSceneData.NumObjectsInBuffer, UpdateObjectsGroupSize), 1, 1);
-					ComputeShader->UnsetParameters(RHICmdList);
-				}
-			}
-
-			// Allocate tile resolution based on world space size
-			//@todo - light space perspective shadow maps would make much better use of the resolution
-			const float LightTiles = FMath::Min(ShadowBounds.W / GShadowWorldTileSize + 1, 256.0f);
-			FIntPoint LightTileDimensions(LightTiles, LightTiles);
-
-			FLightTileIntersectionResources* TileIntersectionResources = NULL;
-
 			GSceneRenderTargets.FinishRenderingLightAttenuation(RHICmdList);
 			SetRenderTarget(RHICmdList, NULL, NULL);
 
-			if (bDirectionalLight && GShadowScatterTileCulling)
-			{
-				if (!LightSceneInfo->TileIntersectionResources || LightSceneInfo->TileIntersectionResources->TileDimensions != LightTileDimensions)
-				{
-					if (LightSceneInfo->TileIntersectionResources)
-					{
-						LightSceneInfo->TileIntersectionResources->Release();
-					}
-					else
-					{
-						LightSceneInfo->TileIntersectionResources = new FLightTileIntersectionResources();
-					}
+			int32 NumPlanes = 0;
+			const FPlane* PlaneData = NULL;
+			FVector4 ShadowBoundingSphereValue(0, 0, 0, 0);
 
-					LightSceneInfo->TileIntersectionResources->TileDimensions = LightTileDimensions;
-
-					LightSceneInfo->TileIntersectionResources->Initialize();
-				}
-
-				TileIntersectionResources = LightSceneInfo->TileIntersectionResources;
-
-				{
-					SCOPED_DRAW_EVENT(RHICmdList, ClearTiles);
-					TShaderMapRef<FClearTilesCS> ComputeShader(View.ShaderMap);
-
-					uint32 GroupSizeX = FMath::DivideAndRoundUp(LightTileDimensions.X, GDistanceFieldAOTileSizeX);
-					uint32 GroupSizeY = FMath::DivideAndRoundUp(LightTileDimensions.Y, GDistanceFieldAOTileSizeY);
-
-					RHICmdList.SetComputeShader(ComputeShader->GetComputeShader());
-					ComputeShader->SetParameters(RHICmdList, View, FVector2D(LightTileDimensions.X, LightTileDimensions.Y), TileIntersectionResources);
-					DispatchComputeShader(RHICmdList, *ComputeShader, GroupSizeX, GroupSizeY, 1);
-
-					ComputeShader->UnsetParameters(RHICmdList);
-				}
-
-				{
-					SCOPED_DRAW_EVENT(RHICmdList, CullObjects);
-
-					TShaderMapRef<FShadowObjectCullVS> VertexShader(View.ShaderMap);
-					TShaderMapRef<FShadowObjectCullPS> PixelShader(View.ShaderMap);
-
-					TArray<FUnorderedAccessViewRHIParamRef> UAVs;
-					PixelShader->GetUAVs(View, TileIntersectionResources, UAVs);
-					RHICmdList.SetRenderTargets(0, (const FRHIRenderTargetView *)NULL, NULL, UAVs.Num(), UAVs.GetData());
-
-					RHICmdList.SetViewport(0, 0, 0.0f, LightTileDimensions.X, LightTileDimensions.Y, 1.0f);
-
-					// Render backfaces since camera may intersect
-					RHICmdList.SetRasterizerState(View.bReverseCulling ? TStaticRasterizerState<FM_Solid, CM_CW>::GetRHI() : TStaticRasterizerState<FM_Solid, CM_CCW>::GetRHI());
-					RHICmdList.SetDepthStencilState(TStaticDepthStencilState<false, CF_Always>::GetRHI());
-					RHICmdList.SetBlendState(TStaticBlendState<>::GetRHI());
-
-					static FGlobalBoundShaderState BoundShaderState;
-
-					SetGlobalBoundShaderState(RHICmdList, View.GetFeatureLevel(), BoundShaderState, GetVertexDeclarationFVector4(), *VertexShader, *PixelShader);
+			if (bDirectionalLight)
+			{
+				NumPlanes = CascadeSettings.ShadowBoundsAccurate.Planes.Num();
+				PlaneData = CascadeSettings.ShadowBoundsAccurate.Planes.GetData();
+			}
+			else if (bOnePassPointLightShadow)
+			{
+				ShadowBoundingSphereValue = FVector4(ShadowBounds.Center.X, ShadowBounds.Center.Y, ShadowBounds.Center.Z, ShadowBounds.W);
+			}
+			else
+			{
+				NumPlanes = CasterFrustum.Planes.Num();
+				PlaneData = CasterFrustum.Planes.GetData();
+				ShadowBoundingSphereValue = FVector4(PreShadowTranslation, 0);
+			}
 
-					VertexShader->SetParameters(RHICmdList, View, FVector2D(LightTileDimensions.X, LightTileDimensions.Y), this);
-					PixelShader->SetParameters(RHICmdList, View, FVector2D(LightTileDimensions.X, LightTileDimensions.Y), this, TileIntersectionResources);
+			const FMatrix WorldToShadowValue = FTranslationMatrix(PreShadowTranslation) * SubjectAndReceiverMatrix;
 
-					RHICmdList.SetStreamSource(0, StencilingGeometry::GLowPolyStencilSphereVertexBuffer.VertexBufferRHI, sizeof(FVector4), 0);
+			CullDistanceFieldObjectsForLight(
+				RHICmdList,
+				View,
+				LightSceneInfo->Proxy, 
+				WorldToShadowValue,
+				NumPlanes,
+				PlaneData,
+				ShadowBoundingSphereValue,
+				ShadowBounds.W,
+				LightSceneInfo->TileIntersectionResources
+				);
 
-					RHICmdList.DrawIndexedPrimitiveIndirect(
-						PT_TriangleList,
-						StencilingGeometry::GLowPolyStencilSphereIndexBuffer.IndexBufferRHI, 
-						GShadowCulledObjectBuffers.Buffers.ObjectIndirectArguments.Buffer,
-						0);
-				}
-			}
+			FLightTileIntersectionResources* TileIntersectionResources = LightSceneInfo->TileIntersectionResources;
 
 			TRefCountPtr<IPooledRenderTarget> RayTracedShadowsRT;
 
@@ -757,7 +783,6 @@ void FProjectedShadowInfo::RenderRayTracedDistanceFieldProjection(FRHICommandLis
 
 				{
 					SCOPED_DRAW_EVENT(RHICmdList, RayTraceShadows);
-					SetRenderTarget(RHICmdList, NULL, NULL);
 
 					if (bDirectionalLight && GShadowScatterTileCulling)
 					{
diff --git a/Engine/Source/Runtime/Renderer/Private/DistanceFieldSurfaceCacheLighting.cpp b/Engine/Source/Runtime/Renderer/Private/DistanceFieldSurfaceCacheLighting.cpp
--- a/Engine/Source/Runtime/Renderer/Private/DistanceFieldSurfaceCacheLighting.cpp
+++ b/Engine/Source/Runtime/Renderer/Private/DistanceFieldSurfaceCacheLighting.cpp
@@ -12,10 +12,12 @@
 #include "SceneFilterRendering.h"
 #include "DistanceFieldLightingShared.h"
 #include "DistanceFieldSurfaceCacheLighting.h"
+#include "DistanceFieldGlobalIllumination.h"
 #include "PostProcessAmbientOcclusion.h"
 #include "RHICommandList.h"
 #include "SceneUtils.h"
 #include "OneColorShader.h"
+#include "BasePassRendering.h"
 
 int32 GDistanceFieldAO = 1;
 FAutoConsoleVariableRef CVarDistanceFieldAO(
@@ -25,13 +27,11 @@ FAutoConsoleVariableRef CVarDistanceFieldAO(
 	ECVF_Scalability | ECVF_RenderThreadSafe
 	);
 
-int32 GDistanceFieldGI = 0;
-FAutoConsoleVariableRef CVarDistanceFieldGI(
-	TEXT("r.DistanceFieldGI"),
-	GDistanceFieldGI,
-	TEXT(""),
-	ECVF_Cheat | ECVF_RenderThreadSafe
-	);
+bool IsDistanceFieldGIAllowed(const FViewInfo& View)
+{
+	extern int32 GDistanceFieldGI;
+	return View.Family->EngineShowFlags.VisualizeDistanceFieldGI || (View.Family->EngineShowFlags.DistanceFieldGI && GDistanceFieldGI && View.Family->EngineShowFlags.GlobalIllumination);
+}
 
 int32 GAOPowerOfTwoBetweenLevels = 2;
 FAutoConsoleVariableRef CVarAOPowerOfTwoBetweenLevels(
@@ -51,7 +51,7 @@ FAutoConsoleVariableRef CVarAOMinLevel(
 	ECVF_Cheat | ECVF_RenderThreadSafe
 	);
 
-int32 GAOMaxLevel = FMath::Min(4, GAOMaxSupportedLevel);
+int32 GAOMaxLevel = FMath::Min(3, GAOMaxSupportedLevel);
 FAutoConsoleVariableRef CVarAOMaxLevel(
 	TEXT("r.AOMaxLevel"),
 	GAOMaxLevel,
@@ -67,7 +67,7 @@ FAutoConsoleVariableRef CVarAOReuseAcrossFrames(
 	ECVF_Cheat | ECVF_RenderThreadSafe
 	);
 
-float GAOTrimOldRecordsFraction = .1f;
+float GAOTrimOldRecordsFraction = .2f;
 FAutoConsoleVariableRef CVarAOTrimOldRecordsFraction(
 	TEXT("r.AOTrimOldRecordsFraction"),
 	GAOTrimOldRecordsFraction,
@@ -163,13 +163,6 @@ FAutoConsoleVariableRef CVarAOMaxViewDistance(
 	ECVF_Cheat | ECVF_RenderThreadSafe
 	);
 
-float GetMaxAOViewDistance()
-{
-	// Scene depth stored in fp16 alpha, must fade out before it runs out of range
-	// The fade extends past GAOMaxViewDistance a bit
-	return FMath::Min(GAOMaxViewDistance, 65000.0f);
-}
-
 int32 GAOScatterTileCulling = 1;
 FAutoConsoleVariableRef CVarAOScatterTileCulling(
 	TEXT("r.AOScatterTileCulling"),
@@ -225,11 +218,6 @@ FIntPoint GetBufferSizeForAO()
 	return FIntPoint::DivideAndRoundDown(GSceneRenderTargets.GetBufferSizeXY(), GAODownsampleFactor);
 }
 
-bool DoesPlatformSupportDistanceFieldAO(EShaderPlatform Platform)
-{
-	return Platform == SP_PCD3D_SM5;
-}
-
 // Sample set restricted to not self-intersect a surface based on cone angle .475882232
 // Coverage of hemisphere = 0.755312979
 const FVector SpacedVectors9[] = 
@@ -335,8 +323,8 @@ TGlobalResource<FCircleVertexBuffer> GCircleVertexBuffer;
 
 // Only 7 needed atm, pad to cache line size
 // Must match equivalent shader defines
-int32 FDistanceFieldObjectBuffers::ObjectDataStride = 12;
-int32 FDistanceFieldCulledObjectBuffers::ObjectDataStride = 8;
+int32 FDistanceFieldObjectBuffers::ObjectDataStride = 15;
+int32 FDistanceFieldCulledObjectBuffers::ObjectDataStride = 11;
 int32 FDistanceFieldCulledObjectBuffers::ObjectBoxBoundsStride = 5;
 
 TGlobalResource<FDistanceFieldObjectBufferResource> GAOCulledObjectBuffers;
@@ -445,54 +433,6 @@ public:
 
 TGlobalResource<FHeightfieldDescriptionsResource> GHeightfieldDescriptions;
 
-/**  */
-class FTemporaryIrradianceCacheResources : public FRenderResource
-{
-public:
-
-	FTemporaryIrradianceCacheResources()
-	{
-		MaxIrradianceCacheSamples = 0;
-	}
-
-	virtual void InitDynamicRHI()
-	{
-		if (MaxIrradianceCacheSamples > 0)
-		{
-			ConeVisibility.Initialize(sizeof(float), MaxIrradianceCacheSamples * 9, PF_R32_FLOAT, BUF_Static);
-		}
-	}
-
-	virtual void ReleaseDynamicRHI()
-	{
-		ConeVisibility.Release();
-	}
-
-	void AllocateFor(int32 InMaxIrradianceCacheSamples)
-	{
-		bool bReallocate = false;
-
-		if (InMaxIrradianceCacheSamples > MaxIrradianceCacheSamples)
-		{
-			MaxIrradianceCacheSamples = InMaxIrradianceCacheSamples;
-			bReallocate = true;
-		}
-
-		if (!IsInitialized())
-		{
-			InitResource();
-		}
-		else if (bReallocate)
-		{
-			UpdateRHI();
-		}
-	}
-
-	int32 MaxIrradianceCacheSamples;
-
-	FRWBuffer ConeVisibility;
-};
-
 TGlobalResource<FTemporaryIrradianceCacheResources> GTemporaryIrradianceCacheResources;
 
 void FTileIntersectionResources::InitDynamicRHI()
@@ -510,18 +450,29 @@ void FTileIntersectionResources::InitDynamicRHI()
 
 void FSceneViewState::DestroyAOTileResources()
 {
-	if (AOTileIntersectionResources) 
+	if (AOTileIntersectionResources || VPLTileIntersectionResources) 
 	{
-		ENQUEUE_UNIQUE_RENDER_COMMAND_ONEPARAMETER(
+		ENQUEUE_UNIQUE_RENDER_COMMAND_TWOPARAMETER(
 			DeleteTileResources,
 			FTileIntersectionResources*, AOTileIntersectionResources, AOTileIntersectionResources,
+			FVPLTileIntersectionResources*, VPLTileIntersectionResources, VPLTileIntersectionResources,
 			{
-				AOTileIntersectionResources->ReleaseResource();
-				delete AOTileIntersectionResources;
+				if (AOTileIntersectionResources)
+				{
+					AOTileIntersectionResources->ReleaseResource();
+					delete AOTileIntersectionResources;
+				}
+				
+				if (VPLTileIntersectionResources)
+				{
+					VPLTileIntersectionResources->ReleaseResource();
+					delete VPLTileIntersectionResources;
+				}
 			}
 		);
 	}
 	AOTileIntersectionResources = NULL;
+	VPLTileIntersectionResources = NULL;
 }
 
 void OnClearSurfaceCache(UWorld* InWorld)
@@ -799,7 +750,7 @@ public:
 		AOParameters.Set(RHICmdList, ShaderRHI, Parameters);
 
 		// Shader assumes max 6
-		check(View.ViewFrustum.Planes.Num() < 6);
+		check(View.ViewFrustum.Planes.Num() <= 6);
 		SetShaderValue(RHICmdList, ShaderRHI, NumConvexHullPlanes, View.ViewFrustum.Planes.Num());
 		SetShaderValueArray(RHICmdList, ShaderRHI, ViewFrustumConvexHull, View.ViewFrustum.Planes.GetData(), View.ViewFrustum.Planes.Num());
 		SetShaderValue(RHICmdList, ShaderRHI, ObjectBoundingGeometryIndexCount, StencilingGeometry::GLowPolyStencilSphereIndexBuffer.GetIndexCount());
@@ -992,6 +943,10 @@ void FDeferredShadingSceneRenderer::UpdateGlobalDistanceFieldObjectBuffers(FRHIC
 						// Box bounds
 						UploadObjectData.Add(FVector4(LocalVolumeBounds.Max, LocalVolumeBounds.Min.Z));
 
+						UploadObjectData.Add(*(FVector4*)&UniformScaleVolumeToWorld.M[0]);
+						UploadObjectData.Add(*(FVector4*)&UniformScaleVolumeToWorld.M[1]);
+						UploadObjectData.Add(*(FVector4*)&UniformScaleVolumeToWorld.M[2]);
+
 						UploadObjectData.Add(*(FVector4*)&LocalToWorld.M[0]);
 						UploadObjectData.Add(*(FVector4*)&LocalToWorld.M[1]);
 						UploadObjectData.Add(*(FVector4*)&LocalToWorld.M[2]);
@@ -1235,7 +1190,6 @@ public:
 		OutEnvironment.SetDefine(TEXT("THREADGROUP_SIZEY"), GDistanceFieldAOTileSizeY);
 		OutEnvironment.SetDefine(TEXT("NUM_CONE_DIRECTIONS"), NumConeSampleDirections);
 		OutEnvironment.SetDefine(TEXT("DOWNSAMPLE_FACTOR"), GAODownsampleFactor);
-		OutEnvironment.SetDefine(TEXT("MAX_OBJECTS_PER_TILE"), GMaxNumObjectsPerTile);
 		
 		// To reduce shader compile time of compute shaders with shared memory, doesn't have an impact on generated code with current compiler (June 2010 DX SDK)
 		OutEnvironment.CompilerFlags.Add(CFLAG_StandardOptimization);
@@ -1384,7 +1338,6 @@ public:
 		FGlobalShader::ModifyCompilationEnvironment(Platform,OutEnvironment);
 		
 		OutEnvironment.SetDefine(TEXT("DOWNSAMPLE_FACTOR"), GAODownsampleFactor);
-		OutEnvironment.SetDefine(TEXT("MAX_OBJECTS_PER_TILE"), GMaxNumObjectsPerTile);
 	}
 
 	/** Default constructor. */
@@ -1479,7 +1432,6 @@ public:
 		OutEnvironment.SetDefine(TEXT("THREADGROUP_SIZEY"), GDistanceFieldAOTileSizeY);
 		OutEnvironment.SetDefine(TEXT("NUM_CONE_DIRECTIONS"), NumConeSampleDirections);
 		OutEnvironment.SetDefine(TEXT("DOWNSAMPLE_FACTOR"), GAODownsampleFactor);
-		OutEnvironment.SetDefine(TEXT("MAX_OBJECTS_PER_TILE"), GMaxNumObjectsPerTile);
 		// To reduce shader compile time of compute shaders with shared memory, doesn't have an impact on generated code with current compiler (June 2010 DX SDK)
 		OutEnvironment.CompilerFlags.Add(CFLAG_StandardOptimization);
 	}
@@ -1610,7 +1562,6 @@ public:
 		OutEnvironment.SetDefine(TEXT("DOWNSAMPLE_FACTOR"), GAODownsampleFactor);
 		OutEnvironment.SetDefine(TEXT("THREADGROUP_SIZEX"), GDistanceFieldAOTileSizeX);
 		OutEnvironment.SetDefine(TEXT("THREADGROUP_SIZEY"), GDistanceFieldAOTileSizeY);
-		OutEnvironment.SetDefine(TEXT("MAX_OBJECTS_PER_TILE"), GMaxNumObjectsPerTile);
 	}
 
 	/** Default constructor. */
@@ -1689,7 +1640,6 @@ public:
 		OutEnvironment.SetDefine(TEXT("DOWNSAMPLE_FACTOR"), GAODownsampleFactor);
 		OutEnvironment.SetDefine(TEXT("THREADGROUP_SIZEX"), GDistanceFieldAOTileSizeX);
 		OutEnvironment.SetDefine(TEXT("THREADGROUP_SIZEY"), GDistanceFieldAOTileSizeY);
-		OutEnvironment.SetDefine(TEXT("MAX_OBJECTS_PER_TILE"), GMaxNumObjectsPerTile);
 	}
 
 	/** Default constructor. */
@@ -2343,7 +2293,6 @@ public:
 		FGlobalShader::ModifyCompilationEnvironment(Platform,OutEnvironment);
 		OutEnvironment.SetDefine(TEXT("NUM_CONE_DIRECTIONS"), NumConeSampleDirections);
 		OutEnvironment.SetDefine(TEXT("DOWNSAMPLE_FACTOR"), GAODownsampleFactor);
-		OutEnvironment.SetDefine(TEXT("MAX_OBJECTS_PER_TILE"), GMaxNumObjectsPerTile);
 		OutEnvironment.SetDefine(TEXT("SUPPORT_IRRADIANCE"), bSupportIrradiance ? TEXT("1") : TEXT("0"));
 
 		// To reduce shader compile time of compute shaders with shared memory, doesn't have an impact on generated code with current compiler (June 2010 DX SDK)
@@ -2360,7 +2309,7 @@ public:
 		IrradianceCacheNormal.Bind(Initializer.ParameterMap, TEXT("IrradianceCacheNormal"));
 		OccluderRadius.Bind(Initializer.ParameterMap, TEXT("OccluderRadius"));
 		RecordConeVisibility.Bind(Initializer.ParameterMap, TEXT("RecordConeVisibility"));
-		IrradianceCacheIrradiance.Bind(Initializer.ParameterMap, TEXT("IrradianceCacheIrradiance"));
+		RecordConeData.Bind(Initializer.ParameterMap, TEXT("RecordConeData"));
 		DebugBuffer.Bind(Initializer.ParameterMap, TEXT("DebugBuffer"));
 		IrradianceCacheTileCoordinate.Bind(Initializer.ParameterMap, TEXT("IrradianceCacheTileCoordinate"));
 		ScatterDrawParameters.Bind(Initializer.ParameterMap, TEXT("ScatterDrawParameters"));
@@ -2413,7 +2362,7 @@ public:
 
 		OccluderRadius.SetBuffer(RHICmdList, ShaderRHI, SurfaceCacheResources.Level[DepthLevel]->OccluderRadius);
 		RecordConeVisibility.SetBuffer(RHICmdList, ShaderRHI, GTemporaryIrradianceCacheResources.ConeVisibility);
-		IrradianceCacheIrradiance.SetBuffer(RHICmdList, ShaderRHI, SurfaceCacheResources.Level[DepthLevel]->Irradiance);
+		RecordConeData.SetBuffer(RHICmdList, ShaderRHI, GTemporaryIrradianceCacheResources.ConeData);
 		DebugBuffer.SetBuffer(RHICmdList, ShaderRHI, GDebugBuffer.DebugData);
 
 		SetShaderValue(RHICmdList, ShaderRHI, ViewDimensionsParameter, View.ViewRect);
@@ -2485,7 +2434,7 @@ public:
 	{
 		OccluderRadius.UnsetUAV(RHICmdList, GetComputeShader());
 		RecordConeVisibility.UnsetUAV(RHICmdList, GetComputeShader());
-		IrradianceCacheIrradiance.UnsetUAV(RHICmdList, GetComputeShader());
+		RecordConeData.UnsetUAV(RHICmdList, GetComputeShader());
 		DebugBuffer.UnsetUAV(RHICmdList, GetComputeShader());
 	}
 
@@ -2500,7 +2449,7 @@ public:
 		Ar << IrradianceCacheNormal;
 		Ar << OccluderRadius;
 		Ar << RecordConeVisibility;
-		Ar << IrradianceCacheIrradiance;
+		Ar << RecordConeData;
 		Ar << DebugBuffer;
 		Ar << IrradianceCacheTileCoordinate;
 		Ar << ScatterDrawParameters;
@@ -2532,7 +2481,7 @@ private:
 	FShaderResourceParameter IrradianceCachePositionRadius;
 	FRWShaderParameter OccluderRadius;
 	FRWShaderParameter RecordConeVisibility;
-	FRWShaderParameter IrradianceCacheIrradiance;
+	FRWShaderParameter RecordConeData;
 	FRWShaderParameter DebugBuffer;
 	FShaderResourceParameter ScatterDrawParameters;
 	FShaderResourceParameter SavedStartIndex;
@@ -2705,11 +2654,9 @@ private:
 
 IMPLEMENT_SHADER_TYPE(,FCalculateHeightfieldOcclusionCS,TEXT("DistanceFieldSurfaceCacheLightingCompute"),TEXT("CalculateHeightfieldOcclusionCS"),SF_Compute);
 
-
-template<bool bSupportIrradiance>
-class TCombineConesCS : public FGlobalShader
+class FCombineConesCS : public FGlobalShader
 {
-	DECLARE_SHADER_TYPE(TCombineConesCS,Global)
+	DECLARE_SHADER_TYPE(FCombineConesCS,Global)
 public:
 
 	static bool ShouldCache(EShaderPlatform Platform)
@@ -2721,29 +2668,24 @@ public:
 	{
 		FGlobalShader::ModifyCompilationEnvironment(Platform,OutEnvironment);
 		OutEnvironment.SetDefine(TEXT("NUM_CONE_DIRECTIONS"), NumConeSampleDirections);
-		OutEnvironment.SetDefine(TEXT("SUPPORT_IRRADIANCE"), bSupportIrradiance ? TEXT("1") : TEXT("0"));
 
 		// To reduce shader compile time of compute shaders with shared memory, doesn't have an impact on generated code with current compiler (June 2010 DX SDK)
 		OutEnvironment.CompilerFlags.Add(CFLAG_StandardOptimization);
 	}
 
-	TCombineConesCS(const ShaderMetaType::CompiledShaderInitializerType& Initializer)
+	FCombineConesCS(const ShaderMetaType::CompiledShaderInitializerType& Initializer)
 		: FGlobalShader(Initializer)
 	{
 		AOParameters.Bind(Initializer.ParameterMap);
-		OccluderRadius.Bind(Initializer.ParameterMap, TEXT("OccluderRadius"));
 		IrradianceCacheBentNormal.Bind(Initializer.ParameterMap, TEXT("IrradianceCacheBentNormal"));
 		IrradianceCacheNormal.Bind(Initializer.ParameterMap, TEXT("IrradianceCacheNormal"));
 		RecordConeVisibility.Bind(Initializer.ParameterMap, TEXT("RecordConeVisibility"));
 		ScatterDrawParameters.Bind(Initializer.ParameterMap, TEXT("ScatterDrawParameters"));
 		SavedStartIndex.Bind(Initializer.ParameterMap, TEXT("SavedStartIndex"));
 		BentNormalNormalizeFactor.Bind(Initializer.ParameterMap, TEXT("BentNormalNormalizeFactor"));
-		DistanceFieldNormalTexture.Bind(Initializer.ParameterMap, TEXT("DistanceFieldNormalTexture"));
-		DistanceFieldNormalSampler.Bind(Initializer.ParameterMap, TEXT("DistanceFieldNormalSampler"));
-		RecordRadiusScale.Bind(Initializer.ParameterMap, TEXT("RecordRadiusScale"));
 	}
 
-	TCombineConesCS()
+	FCombineConesCS()
 	{
 	}
 
@@ -2751,11 +2693,15 @@ public:
 		FRHICommandList& RHICmdList, 
 		const FSceneView& View, 
 		int32 DepthLevel, 
-		const FDistanceFieldAOParameters& Parameters)
+		const FDistanceFieldAOParameters& Parameters,
+		const FLightSceneProxy* DirectionalLight,
+		const FMatrix& WorldToShadowMatrixValue,
+		FLightTileIntersectionResources* TileIntersectionResources,
+		FVPLResources* VPLResources)
 	{
 		FComputeShaderRHIParamRef ShaderRHI = GetComputeShader();
 		FGlobalShader::SetParameters(RHICmdList, ShaderRHI, View);
-		
+
 		AOParameters.Set(RHICmdList, ShaderRHI, Parameters);
 
 		const FScene* Scene = (const FScene*)View.Family->Scene;
@@ -2766,9 +2712,8 @@ public:
 		SetSRVParameter(RHICmdList, ShaderRHI, SavedStartIndex, SurfaceCacheResources.Level[DepthLevel]->SavedStartIndex.SRV);
 		SetSRVParameter(RHICmdList, ShaderRHI, RecordConeVisibility, GTemporaryIrradianceCacheResources.ConeVisibility.SRV);
 
-		OccluderRadius.SetBuffer(RHICmdList, ShaderRHI, SurfaceCacheResources.Level[DepthLevel]->OccluderRadius);
 		IrradianceCacheBentNormal.SetBuffer(RHICmdList, ShaderRHI, SurfaceCacheResources.Level[DepthLevel]->BentNormal);
-		
+
 		FAOSampleData2 AOSampleData;
 
 		TArray<FVector, TInlineAllocator<9> > SampleDirections;
@@ -2790,13 +2735,10 @@ public:
 
 		float BentNormalNormalizeFactorValue = 1.0f / (UnoccludedVector / NumConeSampleDirections).Size();
 		SetShaderValue(RHICmdList, ShaderRHI, BentNormalNormalizeFactor, BentNormalNormalizeFactorValue);
-
-		SetShaderValue(RHICmdList, ShaderRHI, RecordRadiusScale, GAORecordRadiusScale);
 	}
 
 	void UnsetParameters(FRHICommandList& RHICmdList)
 	{
-		OccluderRadius.UnsetUAV(RHICmdList, GetComputeShader());
 		IrradianceCacheBentNormal.UnsetUAV(RHICmdList, GetComputeShader());
 	}
 
@@ -2804,36 +2746,27 @@ public:
 	{		
 		bool bShaderHasOutdatedParameters = FGlobalShader::Serialize(Ar);
 		Ar << AOParameters;
-		Ar << OccluderRadius;
 		Ar << IrradianceCacheBentNormal;
 		Ar << IrradianceCacheNormal;
 		Ar << RecordConeVisibility;
 		Ar << ScatterDrawParameters;
 		Ar << SavedStartIndex;
 		Ar << BentNormalNormalizeFactor;
-		Ar << DistanceFieldNormalTexture;
-		Ar << DistanceFieldNormalSampler;
-		Ar << RecordRadiusScale;
 		return bShaderHasOutdatedParameters;
 	}
 
 private:
 
 	FAOParameters AOParameters;
-	FRWShaderParameter OccluderRadius;
 	FRWShaderParameter IrradianceCacheBentNormal;
 	FShaderResourceParameter IrradianceCacheNormal;
 	FShaderResourceParameter RecordConeVisibility;
 	FShaderResourceParameter ScatterDrawParameters;
 	FShaderResourceParameter SavedStartIndex;
 	FShaderParameter BentNormalNormalizeFactor;
-	FShaderResourceParameter DistanceFieldNormalTexture;
-	FShaderResourceParameter DistanceFieldNormalSampler;
-	FShaderParameter RecordRadiusScale;
 };
 
-IMPLEMENT_SHADER_TYPE(template<>,TCombineConesCS<true>,TEXT("DistanceFieldSurfaceCacheLightingCompute"),TEXT("CombineConesCS"),SF_Compute);
-IMPLEMENT_SHADER_TYPE(template<>,TCombineConesCS<false>,TEXT("DistanceFieldSurfaceCacheLightingCompute"),TEXT("CombineConesCS"),SF_Compute);
+IMPLEMENT_SHADER_TYPE(,FCombineConesCS,TEXT("DistanceFieldSurfaceCacheLightingCompute"),TEXT("CombineConesCS"),SF_Compute);
 
 class FWriteDownsampledDepthPS : public FGlobalShader
 {
@@ -2922,11 +2855,12 @@ public:
 		InterpolationRadiusScale.Bind(Initializer.ParameterMap, TEXT("InterpolationRadiusScale"));
 		NormalizedOffsetToPixelCenter.Bind(Initializer.ParameterMap, TEXT("NormalizedOffsetToPixelCenter"));
 		HackExpand.Bind(Initializer.ParameterMap, TEXT("HackExpand"));
+		AOParameters.Bind(Initializer.ParameterMap);
 	}
 
 	TIrradianceCacheSplatVS() {}
 
-	void SetParameters(FRHICommandList& RHICmdList, const FSceneView& View, int32 DepthLevel, int32 CurrentLevelDownsampleFactorValue, FVector2D NormalizedOffsetToPixelCenterValue)
+	void SetParameters(FRHICommandList& RHICmdList, const FSceneView& View, const FDistanceFieldAOParameters& Parameters, int32 DepthLevel, int32 CurrentLevelDownsampleFactorValue, FVector2D NormalizedOffsetToPixelCenterValue)
 	{
 		const FVertexShaderRHIParamRef ShaderRHI = GetVertexShader();
 		FGlobalShader::SetParameters(RHICmdList, ShaderRHI, View);
@@ -2946,6 +2880,8 @@ public:
 		const FIntPoint AOViewRectSize = FIntPoint::DivideAndRoundUp(View.ViewRect.Size(), CurrentLevelDownsampleFactorValue);
 		const FVector2D HackExpandValue(.5f / AOViewRectSize.X, .5f / AOViewRectSize.Y);
 		SetShaderValue(RHICmdList, ShaderRHI, HackExpand, HackExpandValue);
+
+		AOParameters.Set(RHICmdList, ShaderRHI, Parameters);
 	}
 
 	virtual bool Serialize(FArchive& Ar)
@@ -2959,6 +2895,7 @@ public:
 		Ar << InterpolationRadiusScale;
 		Ar << NormalizedOffsetToPixelCenter;
 		Ar << HackExpand;
+		Ar << AOParameters;
 		return bShaderHasOutdatedParameters;
 	}
 
@@ -2971,6 +2908,7 @@ private:
 	FShaderParameter InterpolationRadiusScale;
 	FShaderParameter NormalizedOffsetToPixelCenter;
 	FShaderParameter HackExpand;
+	FAOParameters AOParameters;
 };
 
 // typedef required to get around macro expansion failure due to commas in template argument list 
@@ -3413,7 +3351,7 @@ void AllocateOrReuseAORenderTarget(TRefCountPtr<IPooledRenderTarget>& Target, co
 
 void UpdateHistory(
 	FRHICommandList& RHICmdList,
-	FViewInfo& View, 
+	const FViewInfo& View, 
 	const TCHAR* BentNormalHistoryRTName,
 	const TCHAR* IrradianceHistoryRTName,
 	/** Contains last frame's history, if non-NULL.  This will be updated with the new frame's history. */
@@ -3428,10 +3366,12 @@ void UpdateHistory(
 {
 	if (BentNormalHistoryState)
 	{
+		const bool bUseDistanceFieldGI = IsDistanceFieldGIAllowed(View);
+
 		if (*BentNormalHistoryState 
 			&& !View.bCameraCut 
 			&& !View.bPrevTransformsReset 
-			&& (!GDistanceFieldGI || (IrradianceHistoryState && *IrradianceHistoryState)))
+			&& (!bUseDistanceFieldGI || (IrradianceHistoryState && *IrradianceHistoryState)))
 		{
 			// Reuse a render target from the pool with a consistent name, for vis purposes
 			TRefCountPtr<IPooledRenderTarget> NewBentNormalHistory;
@@ -3439,7 +3379,7 @@ void UpdateHistory(
 
 			TRefCountPtr<IPooledRenderTarget> NewIrradianceHistory;
 
-			if (GDistanceFieldGI)
+			if (bUseDistanceFieldGI)
 			{
 				AllocateOrReuseAORenderTarget(NewIrradianceHistory, IrradianceHistoryRTName, false);
 			}
@@ -3450,10 +3390,10 @@ void UpdateHistory(
 				FTextureRHIParamRef RenderTargets[2] =
 				{
 					NewBentNormalHistory->GetRenderTargetItem().TargetableTexture,
-					GDistanceFieldGI ? NewIrradianceHistory->GetRenderTargetItem().TargetableTexture : NULL,
+					bUseDistanceFieldGI ? NewIrradianceHistory->GetRenderTargetItem().TargetableTexture : NULL,
 				};
 
-				SetRenderTargets(RHICmdList, ARRAY_COUNT(RenderTargets) - (GDistanceFieldGI ? 0 : 1), RenderTargets, FTextureRHIParamRef(), 0, NULL);
+				SetRenderTargets(RHICmdList, ARRAY_COUNT(RenderTargets) - (bUseDistanceFieldGI ? 0 : 1), RenderTargets, FTextureRHIParamRef(), 0, NULL);
 
 				RHICmdList.SetViewport(0, 0, 0.0f, View.ViewRect.Width() / GAODownsampleFactor, View.ViewRect.Height() / GAODownsampleFactor, 1.0f);
 				RHICmdList.SetRasterizerState(TStaticRasterizerState<FM_Solid, CM_None>::GetRHI());
@@ -3462,7 +3402,7 @@ void UpdateHistory(
 
 				TShaderMapRef<FPostProcessVS> VertexShader(View.ShaderMap);
 
-				if (GDistanceFieldGI)
+				if (bUseDistanceFieldGI)
 				{
 					TShaderMapRef<TUpdateHistoryPS<true> > PixelShader(View.ShaderMap);
 
@@ -3522,18 +3462,20 @@ void UpdateHistory(
 void PostProcessBentNormalAO(
 	FRHICommandList& RHICmdList, 
 	const FDistanceFieldAOParameters& Parameters, 
-	TArray<FViewInfo>& Views, 
+	const FViewInfo& View, 
 	FSceneRenderTargetItem& BentNormalInterpolation, 
 	IPooledRenderTarget* IrradianceInterpolation,
 	FSceneRenderTargetItem& DistanceFieldNormal,
 	TRefCountPtr<IPooledRenderTarget>& BentNormalOutput,
 	TRefCountPtr<IPooledRenderTarget>& IrradianceOutput)
 {
+	const bool bUseDistanceFieldGI = IsDistanceFieldGIAllowed(View);
+
 	TRefCountPtr<IPooledRenderTarget> DistanceFieldAOBentNormal;
 	TRefCountPtr<IPooledRenderTarget> DistanceFieldIrradiance;
 	AllocateOrReuseAORenderTarget(DistanceFieldAOBentNormal, TEXT("DistanceFieldBentNormalAO"), true);
 
-	if (GDistanceFieldGI)
+	if (bUseDistanceFieldGI)
 	{
 		AllocateOrReuseAORenderTarget(DistanceFieldIrradiance, TEXT("DistanceFieldIrradiance"), false);
 	}
@@ -3545,7 +3487,7 @@ void PostProcessBentNormalAO(
 	{
 		AllocateOrReuseAORenderTarget(DistanceFieldAOBentNormal2, TEXT("DistanceFieldBentNormalAO2"), true);
 
-		if (GDistanceFieldGI)
+		if (bUseDistanceFieldGI)
 		{
 			AllocateOrReuseAORenderTarget(DistanceFieldIrradiance2, TEXT("DistanceFieldIrradiance2"), false);
 		}
@@ -3556,7 +3498,7 @@ void PostProcessBentNormalAO(
 
 		FTextureRHIParamRef DiffuseIrradianceTarget = NULL;
 
-		if (GDistanceFieldGI)
+		if (bUseDistanceFieldGI)
 		{
 			DiffuseIrradianceTarget = GAOFillGaps ? DistanceFieldIrradiance2->GetRenderTargetItem().TargetableTexture : DistanceFieldIrradiance->GetRenderTargetItem().TargetableTexture;
 		}
@@ -3567,12 +3509,9 @@ void PostProcessBentNormalAO(
 			DiffuseIrradianceTarget,
 		};
 
-		SetRenderTargets(RHICmdList, ARRAY_COUNT(RenderTargets) - (GDistanceFieldGI ? 0 : 1), RenderTargets, FTextureRHIParamRef(), 0, NULL);
+		SetRenderTargets(RHICmdList, ARRAY_COUNT(RenderTargets) - (bUseDistanceFieldGI ? 0 : 1), RenderTargets, FTextureRHIParamRef(), 0, NULL);
 
-		for (int32 ViewIndex = 0; ViewIndex < Views.Num(); ViewIndex++)
 		{
-			FViewInfo& View = Views[ViewIndex];
-
 			RHICmdList.SetViewport(0, 0, 0.0f, View.ViewRect.Width() / GAODownsampleFactor, View.ViewRect.Height() / GAODownsampleFactor, 1.0f);
 			RHICmdList.SetRasterizerState(TStaticRasterizerState<FM_Solid, CM_None>::GetRHI());
 			RHICmdList.SetDepthStencilState(TStaticDepthStencilState<false, CF_Always>::GetRHI());
@@ -3580,7 +3519,7 @@ void PostProcessBentNormalAO(
 
 			TShaderMapRef<FPostProcessVS> VertexShader(View.ShaderMap);
 
-			if (GDistanceFieldGI)
+			if (bUseDistanceFieldGI)
 			{
 				TShaderMapRef<TDistanceFieldAOCombinePS<true> > PixelShader(View.ShaderMap);
 				static FGlobalBoundShaderState BoundShaderState;
@@ -3614,15 +3553,12 @@ void PostProcessBentNormalAO(
 		FTextureRHIParamRef RenderTargets[2] =
 		{
 			DistanceFieldAOBentNormal->GetRenderTargetItem().TargetableTexture,
-			GDistanceFieldGI ? DistanceFieldIrradiance->GetRenderTargetItem().TargetableTexture : NULL,
+			bUseDistanceFieldGI ? DistanceFieldIrradiance->GetRenderTargetItem().TargetableTexture : NULL,
 		};
 
-		SetRenderTargets(RHICmdList, ARRAY_COUNT(RenderTargets) - (GDistanceFieldGI ? 0 : 1), RenderTargets, FTextureRHIParamRef(), 0, NULL);
+		SetRenderTargets(RHICmdList, ARRAY_COUNT(RenderTargets) - (bUseDistanceFieldGI ? 0 : 1), RenderTargets, FTextureRHIParamRef(), 0, NULL);
 
-		for (int32 ViewIndex = 0; ViewIndex < Views.Num(); ViewIndex++)
 		{
-			const FViewInfo& View = Views[ViewIndex];
-
 			RHICmdList.SetViewport(0, 0, 0.0f, View.ViewRect.Width() / GAODownsampleFactor, View.ViewRect.Height() / GAODownsampleFactor, 1.0f);
 			RHICmdList.SetRasterizerState(TStaticRasterizerState<FM_Solid, CM_None>::GetRHI());
 			RHICmdList.SetDepthStencilState(TStaticDepthStencilState<false, CF_Always>::GetRHI());
@@ -3630,7 +3566,7 @@ void PostProcessBentNormalAO(
 
 			TShaderMapRef<FPostProcessVS> VertexShader(View.ShaderMap);
 
-			if (GDistanceFieldGI)
+			if (bUseDistanceFieldGI)
 			{
 				TShaderMapRef<TFillGapsPS<true> > PixelShader(View.ShaderMap);
 
@@ -3661,7 +3597,7 @@ void PostProcessBentNormalAO(
 		}
 	}
 
-	FSceneViewState* ViewState = (FSceneViewState*)Views[0].State;
+	FSceneViewState* ViewState = (FSceneViewState*)View.State;
 	TRefCountPtr<IPooledRenderTarget>* BentNormalHistoryState = ViewState ? &ViewState->DistanceFieldAOHistoryRT : NULL;
 	TRefCountPtr<IPooledRenderTarget>* IrradianceHistoryState = ViewState ? &ViewState->DistanceFieldIrradianceHistoryRT : NULL;
 	BentNormalOutput = DistanceFieldAOBentNormal;
@@ -3669,8 +3605,6 @@ void PostProcessBentNormalAO(
 
 	if (GAOUseHistory)
 	{
-		FViewInfo& View = Views[0];
-
 		UpdateHistory(
 			RHICmdList,
 			View, 
@@ -3758,13 +3692,20 @@ private:
 
 IMPLEMENT_UPSAMPLE_PS_TYPE(true, true)
 IMPLEMENT_UPSAMPLE_PS_TYPE(true, false)
+IMPLEMENT_UPSAMPLE_PS_TYPE(false, true)
 IMPLEMENT_UPSAMPLE_PS_TYPE(false, false)
 
-void UpsampleBentNormalAO(FRHICommandList& RHICmdList, const TArray<FViewInfo>& Views, TRefCountPtr<IPooledRenderTarget>& DistanceFieldAOBentNormal, TRefCountPtr<IPooledRenderTarget>& DistanceFieldIrradiance, bool bOutputBentNormal)
+void UpsampleBentNormalAO(
+	FRHICommandList& RHICmdList, 
+	const TArray<FViewInfo>& Views, 
+	TRefCountPtr<IPooledRenderTarget>& DistanceFieldAOBentNormal, 
+	TRefCountPtr<IPooledRenderTarget>& DistanceFieldIrradiance, 
+	bool bOutputBentNormal)
 {
 	for (int32 ViewIndex = 0; ViewIndex < Views.Num(); ViewIndex++)
 	{
 		const FViewInfo& View = Views[ViewIndex];
+		const bool bUseDistanceFieldGI = IsDistanceFieldGIAllowed(View);
 
 		SCOPED_DRAW_EVENT(RHICmdList, UpsampleAO);
 
@@ -3777,7 +3718,7 @@ void UpsampleBentNormalAO(FRHICommandList& RHICmdList, const TArray<FViewInfo>&
 
 		if (bOutputBentNormal)
 		{
-			if (GDistanceFieldGI)
+			if (bUseDistanceFieldGI)
 			{
 				TShaderMapRef<TDistanceFieldAOUpsamplePS<true, true> > PixelShader(View.ShaderMap);
 				static FGlobalBoundShaderState BoundShaderState;
@@ -3794,10 +3735,20 @@ void UpsampleBentNormalAO(FRHICommandList& RHICmdList, const TArray<FViewInfo>&
 		}
 		else
 		{
-			TShaderMapRef<TDistanceFieldAOUpsamplePS<false, false> > PixelShader(View.ShaderMap);
-			static FGlobalBoundShaderState BoundShaderState;
-			SetGlobalBoundShaderState(RHICmdList, View.GetFeatureLevel(), BoundShaderState, GFilterVertexDeclaration.VertexDeclarationRHI, *VertexShader, *PixelShader);
-			PixelShader->SetParameters(RHICmdList, View, DistanceFieldAOBentNormal, DistanceFieldIrradiance);
+			if (bUseDistanceFieldGI)
+			{
+				TShaderMapRef<TDistanceFieldAOUpsamplePS<false, true> > PixelShader(View.ShaderMap);
+				static FGlobalBoundShaderState BoundShaderState;
+				SetGlobalBoundShaderState(RHICmdList, View.GetFeatureLevel(), BoundShaderState, GFilterVertexDeclaration.VertexDeclarationRHI, *VertexShader, *PixelShader);
+				PixelShader->SetParameters(RHICmdList, View, DistanceFieldAOBentNormal, DistanceFieldIrradiance);
+			}
+			else
+			{
+				TShaderMapRef<TDistanceFieldAOUpsamplePS<false, false> > PixelShader(View.ShaderMap);
+				static FGlobalBoundShaderState BoundShaderState;
+				SetGlobalBoundShaderState(RHICmdList, View.GetFeatureLevel(), BoundShaderState, GFilterVertexDeclaration.VertexDeclarationRHI, *VertexShader, *PixelShader);
+				PixelShader->SetParameters(RHICmdList, View, DistanceFieldAOBentNormal, DistanceFieldIrradiance);
+			}
 		}
 
 		DrawRectangle( 
@@ -3978,7 +3929,6 @@ FIntPoint BuildTileObjectLists(FRHICommandListImmediate& RHICmdList, FScene* Sce
 
 	if (GAOScatterTileCulling)
 	{
-
 		for (int32 ViewIndex = 0; ViewIndex < Views.Num(); ViewIndex++)
 		{
 			const FViewInfo& View = Views[ViewIndex];
@@ -4024,7 +3974,7 @@ FIntPoint BuildTileObjectLists(FRHICommandListImmediate& RHICmdList, FScene* Sce
 
 				TArray<FUnorderedAccessViewRHIParamRef> UAVs;
 				PixelShader->GetUAVs(Views[0], UAVs);
-				RHICmdList.SetRenderTargets(0, (const FRHIRenderTargetView *)NULL, NULL, UAVs.Num(), UAVs.GetData());
+				RHICmdList.SetRenderTargets(0, (const FRHIRenderTargetView*)NULL, NULL, UAVs.Num(), UAVs.GetData());
 
 				RHICmdList.SetViewport(0, 0, 0.0f, GroupSizeX, GroupSizeY, 1.0f);
 
@@ -4096,7 +4046,7 @@ FIntPoint BuildTileObjectLists(FRHICommandListImmediate& RHICmdList, FScene* Sce
 	return TileListGroupSize;
 }
 
-void SetupDepthStencil(FRHICommandListImmediate& RHICmdList, FViewInfo& View, FSceneRenderTargetItem& DistanceFieldNormal, int32 DepthLevel, int32 DestLevelDownsampleFactor)
+void SetupDepthStencil(FRHICommandListImmediate& RHICmdList, const FViewInfo& View, FSceneRenderTargetItem& DistanceFieldNormal, int32 DepthLevel, int32 DestLevelDownsampleFactor)
 {
 	SCOPED_DRAW_EVENT(RHICmdList, SetupDepthStencil);
 
@@ -4170,7 +4120,7 @@ void SetupDepthStencil(FRHICommandListImmediate& RHICmdList, FViewInfo& View, FS
 
 void RenderIrradianceCacheInterpolation(
 	FRHICommandListImmediate& RHICmdList, 
-	TArray<FViewInfo>& Views, 
+	const FViewInfo& View, 
 	IPooledRenderTarget* BentNormalInterpolationTarget, 
 	IPooledRenderTarget* IrradianceInterpolationTarget, 
 	FSceneRenderTargetItem& DistanceFieldNormal, 
@@ -4180,6 +4130,7 @@ void RenderIrradianceCacheInterpolation(
 	bool bFinalInterpolation)
 {
 	check(!(bFinalInterpolation && DepthLevel != 0));
+	const bool bUseDistanceFieldGI = IsDistanceFieldGIAllowed(View);
 
 	{
 		TRefCountPtr<IPooledRenderTarget> SplatDepthStencilBuffer;
@@ -4192,23 +4143,19 @@ void RenderIrradianceCacheInterpolation(
 			SetRenderTarget(RHICmdList, NULL, SplatDepthStencilBuffer->GetRenderTargetItem().TargetableTexture);
 			RHICmdList.Clear(false, FLinearColor(0, 0, 0, 0), true, 0, true, 0, FIntRect());
 
-			for (int32 ViewIndex = 0; ViewIndex < Views.Num(); ViewIndex++)
-			{
-				FViewInfo& View = Views[ViewIndex];
-				SetupDepthStencil(RHICmdList, View, DistanceFieldNormal, DepthLevel, DestLevelDownsampleFactor);
-			}
+			SetupDepthStencil(RHICmdList, View, DistanceFieldNormal, DepthLevel, DestLevelDownsampleFactor);
 
 			GRenderTargetPool.VisualizeTexture.SetCheckPoint(RHICmdList, SplatDepthStencilBuffer);
 
 			FTextureRHIParamRef RenderTargets[2] =
 			{
 				BentNormalInterpolationTarget->GetRenderTargetItem().TargetableTexture,
-				GDistanceFieldGI ? IrradianceInterpolationTarget->GetRenderTargetItem().TargetableTexture : NULL
+				bUseDistanceFieldGI ? IrradianceInterpolationTarget->GetRenderTargetItem().TargetableTexture : NULL
 			};
-			SetRenderTargets(RHICmdList, ARRAY_COUNT(RenderTargets) - (GDistanceFieldGI ? 0 : 1), RenderTargets, SplatDepthStencilBuffer->GetRenderTargetItem().TargetableTexture, 0, NULL);
+			SetRenderTargets(RHICmdList, ARRAY_COUNT(RenderTargets) - (bUseDistanceFieldGI ? 0 : 1), RenderTargets, SplatDepthStencilBuffer->GetRenderTargetItem().TargetableTexture, 0, NULL);
 
 			FLinearColor ClearColors[2] = {FLinearColor(0, 0, 0, 0), FLinearColor(0, 0, 0, 0)};
-			RHICmdList.ClearMRT(true, ARRAY_COUNT(ClearColors) - (GDistanceFieldGI ? 0 : 1), ClearColors, false, 0, false, 0, FIntRect());
+			RHICmdList.ClearMRT(true, ARRAY_COUNT(ClearColors) - (bUseDistanceFieldGI ? 0 : 1), ClearColors, false, 0, false, 0, FIntRect());
 		}
 		else
 		{
@@ -4225,10 +4172,7 @@ void RenderIrradianceCacheInterpolation(
 			GCircleVertexBuffer.UpdateRHI();
 		}
 
-		for (int32 ViewIndex = 0; ViewIndex < Views.Num(); ViewIndex++)
 		{
-			FViewInfo& View = Views[ViewIndex];
-
 			const FScene* Scene = (const FScene*)View.Family->Scene;
 			FSurfaceCacheResources& SurfaceCacheResources = *Scene->SurfaceCacheResources;
 
@@ -4298,7 +4242,7 @@ void RenderIrradianceCacheInterpolation(
 			{
 				for (int32 SplatSourceDepthLevel = GAOMaxLevel; SplatSourceDepthLevel >= FMath::Max(DepthLevel, GAOMinLevel); SplatSourceDepthLevel--)
 				{
-					if (GDistanceFieldGI)
+					if (bUseDistanceFieldGI)
 					{
 						TShaderMapRef<TIrradianceCacheSplatVS<true, true> > VertexShader(View.ShaderMap);
 						TShaderMapRef<TIrradianceCacheSplatPS<true, true> > PixelShader(View.ShaderMap);
@@ -4306,7 +4250,7 @@ void RenderIrradianceCacheInterpolation(
 						static FGlobalBoundShaderState BoundShaderState;
 						SetGlobalBoundShaderState(RHICmdList, View.GetFeatureLevel(), BoundShaderState, GScreenVertexDeclaration.VertexDeclarationRHI, *VertexShader, *PixelShader);
 
-						VertexShader->SetParameters(RHICmdList, View, SplatSourceDepthLevel, DestLevelDownsampleFactor, NormalizedOffsetToPixelCenter);
+						VertexShader->SetParameters(RHICmdList, View, Parameters, SplatSourceDepthLevel, DestLevelDownsampleFactor, NormalizedOffsetToPixelCenter);
 						PixelShader->SetParameters(RHICmdList, View, DistanceFieldNormal, DestLevelDownsampleFactor, Parameters);
 					}
 					else
@@ -4317,7 +4261,7 @@ void RenderIrradianceCacheInterpolation(
 						static FGlobalBoundShaderState BoundShaderState;
 						SetGlobalBoundShaderState(RHICmdList, View.GetFeatureLevel(), BoundShaderState, GScreenVertexDeclaration.VertexDeclarationRHI, *VertexShader, *PixelShader);
 
-						VertexShader->SetParameters(RHICmdList, View, SplatSourceDepthLevel, DestLevelDownsampleFactor, NormalizedOffsetToPixelCenter);
+						VertexShader->SetParameters(RHICmdList, View, Parameters, SplatSourceDepthLevel, DestLevelDownsampleFactor, NormalizedOffsetToPixelCenter);
 						PixelShader->SetParameters(RHICmdList, View, DistanceFieldNormal, DestLevelDownsampleFactor, Parameters);
 					}
 
@@ -4328,7 +4272,7 @@ void RenderIrradianceCacheInterpolation(
 			{
 				for (int32 SplatSourceDepthLevel = GAOMaxLevel; SplatSourceDepthLevel >= FMath::Max(DepthLevel, GAOMinLevel); SplatSourceDepthLevel--)
 				{
-					if (GDistanceFieldGI)
+					if (bUseDistanceFieldGI)
 					{
 						TShaderMapRef<TIrradianceCacheSplatVS<false, true> > VertexShader(View.ShaderMap);
 						TShaderMapRef<TIrradianceCacheSplatPS<false, true> > PixelShader(View.ShaderMap);
@@ -4336,7 +4280,7 @@ void RenderIrradianceCacheInterpolation(
 						static FGlobalBoundShaderState BoundShaderState;
 						SetGlobalBoundShaderState(RHICmdList, View.GetFeatureLevel(), BoundShaderState, GScreenVertexDeclaration.VertexDeclarationRHI, *VertexShader, *PixelShader);
 
-						VertexShader->SetParameters(RHICmdList, View, SplatSourceDepthLevel, DestLevelDownsampleFactor, NormalizedOffsetToPixelCenter);
+						VertexShader->SetParameters(RHICmdList, View, Parameters, SplatSourceDepthLevel, DestLevelDownsampleFactor, NormalizedOffsetToPixelCenter);
 						PixelShader->SetParameters(RHICmdList, View, DistanceFieldNormal, DestLevelDownsampleFactor, Parameters);
 					}
 					else
@@ -4347,7 +4291,7 @@ void RenderIrradianceCacheInterpolation(
 						static FGlobalBoundShaderState BoundShaderState;
 						SetGlobalBoundShaderState(RHICmdList, View.GetFeatureLevel(), BoundShaderState, GScreenVertexDeclaration.VertexDeclarationRHI, *VertexShader, *PixelShader);
 
-						VertexShader->SetParameters(RHICmdList, View, SplatSourceDepthLevel, DestLevelDownsampleFactor, NormalizedOffsetToPixelCenter);
+						VertexShader->SetParameters(RHICmdList, View, Parameters, SplatSourceDepthLevel, DestLevelDownsampleFactor, NormalizedOffsetToPixelCenter);
 						PixelShader->SetParameters(RHICmdList, View, DistanceFieldNormal, DestLevelDownsampleFactor, Parameters);
 					}
 
@@ -4371,7 +4315,8 @@ bool FDeferredShadingSceneRenderer::ShouldPrepareForDistanceFieldAO() const
 		&& ViewFamily.EngineShowFlags.DistanceFieldAO
 		&& ((ShouldRenderDynamicSkyLight() && Scene->SkyLight->bCastShadows)
 			|| ViewFamily.EngineShowFlags.VisualizeMeshDistanceFields
-			|| ViewFamily.EngineShowFlags.VisualizeDistanceFieldAO);
+			|| ViewFamily.EngineShowFlags.VisualizeDistanceFieldAO
+			|| ViewFamily.EngineShowFlags.VisualizeDistanceFieldGI);
 }
 
 class FHeightfieldDescription
@@ -4413,6 +4358,7 @@ bool FDeferredShadingSceneRenderer::RenderDistanceFieldAOSurfaceCache(
 		if (GDistanceFieldVolumeTextureAtlas.VolumeTextureRHI && Scene->DistanceFieldSceneData.NumObjectsInBuffer)
 		{
 			check(!Scene->DistanceFieldSceneData.HasPendingOperations());
+			const bool bUseDistanceFieldGI = IsDistanceFieldGIAllowed(View);
 
 			GenerateBestSpacedVectors();
 
@@ -4423,6 +4369,7 @@ bool FDeferredShadingSceneRenderer::RenderDistanceFieldAOSurfaceCache(
 			}
 
 			FSurfaceCacheResources& SurfaceCacheResources = *Scene->SurfaceCacheResources;
+
 			FIntPoint DownsampledViewRect = FIntPoint::DivideAndRoundDown(View.ViewRect.Size(), GAODownsampleFactor);
 			// Needed because samples are kept across frames
 			int32 SlackAmount = 4;
@@ -4433,7 +4380,10 @@ bool FDeferredShadingSceneRenderer::RenderDistanceFieldAOSurfaceCache(
 
 			uint32 ClearValues[4] = {0};
 
-			if (!SurfaceCacheResources.bClearedResources || !GAOReuseAcrossFrames)
+			if (!SurfaceCacheResources.bClearedResources 
+				|| !GAOReuseAcrossFrames 
+				// Drop records that will have uninitialized Irradiance if switching from AO only to AO + GI
+				|| bUseDistanceFieldGI && !SurfaceCacheResources.bHasIrradiance)
 			{
 				// Reset the number of active cache records to 0
 				for (int32 DepthLevel = GAOMaxLevel; DepthLevel >= GAOMinLevel; DepthLevel--)
@@ -4446,6 +4396,8 @@ bool FDeferredShadingSceneRenderer::RenderDistanceFieldAOSurfaceCache(
 				SurfaceCacheResources.bClearedResources = true;
 			}
 
+			SurfaceCacheResources.bHasIrradiance = bUseDistanceFieldGI;
+
 			QUICK_SCOPE_CYCLE_COUNTER(STAT_AOIssueGPUWork);
 
 			{
@@ -4460,6 +4412,9 @@ bool FDeferredShadingSceneRenderer::RenderDistanceFieldAOSurfaceCache(
 				}
 
 				{
+					uint32 ClearValues[4] = { 0 };
+					RHICmdList.ClearUAV(GAOCulledObjectBuffers.Buffers.ObjectIndirectArguments.UAV, ClearValues);
+
 					TShaderMapRef<FCullObjectsForViewCS> ComputeShader(GetGlobalShaderMap(Scene->GetFeatureLevel()));
 					RHICmdList.SetComputeShader(ComputeShader->GetComputeShader());
 					ComputeShader->SetParameters(RHICmdList, Scene, View, Parameters);
@@ -4504,7 +4459,7 @@ bool FDeferredShadingSceneRenderer::RenderDistanceFieldAOSurfaceCache(
 							ComputeShader->UnsetParameters(RHICmdList);
 						}
 
-						if (GDistanceFieldGI)
+						if (bUseDistanceFieldGI)
 						{
 							TShaderMapRef<TCopyIrradianceCacheSamplesCS<true> > ComputeShader(View.ShaderMap);
 
@@ -4530,6 +4485,17 @@ bool FDeferredShadingSceneRenderer::RenderDistanceFieldAOSurfaceCache(
 				}
 			}
 
+			if (bUseDistanceFieldGI)
+			{
+				extern void PlaceVPLs(
+					FRHICommandListImmediate& RHICmdList,
+					const FViewInfo& View,
+					const FScene* Scene,
+					const FDistanceFieldAOParameters& Parameters);
+
+				PlaceVPLs(RHICmdList, View, Scene, Parameters);
+			}
+
 			// Start from the highest depth, which will consider the fewest potential shading points
 			// Each level potentially prevents the next higher resolution level from doing expensive shading work
 			// This is the core of the surface cache algorithm
@@ -4548,10 +4514,10 @@ bool FDeferredShadingSceneRenderer::RenderDistanceFieldAOSurfaceCache(
 				}
 
 				// Splat / interpolate the surface cache records onto the buffer sized for the current depth level
-				RenderIrradianceCacheInterpolation(RHICmdList, Views, DistanceFieldAOBentNormalSplat, NULL, DistanceFieldNormal->GetRenderTargetItem(), DepthLevel, DestLevelDownsampleFactor, Parameters, false);
+				RenderIrradianceCacheInterpolation(RHICmdList, View, DistanceFieldAOBentNormalSplat, NULL, DistanceFieldNormal->GetRenderTargetItem(), DepthLevel, DestLevelDownsampleFactor, Parameters, false);
 
 				{
-					SCOPED_DRAW_EVENT(RHICmdList, PopulateAndShadeIrradianceCache);
+					SCOPED_DRAW_EVENT(RHICmdList, PopulateAndConeTrace);
 					SetRenderTarget(RHICmdList, NULL, NULL);
 
 					// Save off the current record count before adding any more
@@ -4594,7 +4560,7 @@ bool FDeferredShadingSceneRenderer::RenderDistanceFieldAOSurfaceCache(
 					}
 
 					// Compute lighting for the new surface cache records by cone-stepping through the object distance fields
-					if (GDistanceFieldGI)
+					if (bUseDistanceFieldGI)
 					{
 						TShaderMapRef<TFinalGatherCS<true> > ComputeShader(View.ShaderMap);
 
@@ -4625,8 +4591,6 @@ bool FDeferredShadingSceneRenderer::RenderDistanceFieldAOSurfaceCache(
 
 					ComputeShader->UnsetParameters(RHICmdList);
 				}
-				
-				SCOPED_DRAW_EVENT(RHICmdList, CombineCones);
 
 				if (Scene->DistanceFieldSceneData.HeightfieldPrimitives.Num() > 0)
 				{
@@ -4652,6 +4616,8 @@ bool FDeferredShadingSceneRenderer::RenderDistanceFieldAOSurfaceCache(
 
 					if (HeightfieldTextures.Num() > 0)
 					{
+						SCOPED_DRAW_EVENT(RHICmdList, HeightfieldOcclusion);
+
 						for (TMap<UTexture2D*, TArray<FHeightfieldDescription, SceneRenderingAllocator>, SceneRenderingSetAllocator>::TIterator It(HeightfieldTextures); It; ++It)
 						{
 							const TArray<FHeightfieldDescription, SceneRenderingAllocator>& HeightfieldDescriptions = It.Value();
@@ -4706,22 +4672,44 @@ bool FDeferredShadingSceneRenderer::RenderDistanceFieldAOSurfaceCache(
 				}
 
 				{
-					if (GDistanceFieldGI)
+					if (bUseDistanceFieldGI)
 					{
-						TShaderMapRef<TCombineConesCS<true> > ComputeShader(View.ShaderMap);
+						{	
+							TShaderMapRef<TSetupFinalGatherIndirectArgumentsCS<true> > ComputeShader(View.ShaderMap);
+
+							RHICmdList.SetComputeShader(ComputeShader->GetComputeShader());
+							ComputeShader->SetParameters(RHICmdList, View, DepthLevel);
+							DispatchComputeShader(RHICmdList, *ComputeShader, 1, 1, 1);
+
+							ComputeShader->UnsetParameters(RHICmdList);
+						}
+
+						extern void ComputeIrradianceForSamples(
+							int32 DepthLevel,
+							FRHICommandListImmediate& RHICmdList,
+							const FViewInfo& View,
+							const FScene* Scene,
+							const FDistanceFieldAOParameters& Parameters,
+							FTemporaryIrradianceCacheResources* TemporaryIrradianceCacheResources);
+
+						ComputeIrradianceForSamples(DepthLevel, RHICmdList, View, Scene, Parameters, &GTemporaryIrradianceCacheResources);
+					}
+
+					{	
+						TShaderMapRef<TSetupFinalGatherIndirectArgumentsCS<false> > ComputeShader(View.ShaderMap);
 
 						RHICmdList.SetComputeShader(ComputeShader->GetComputeShader());
-						ComputeShader->SetParameters(RHICmdList, View, DepthLevel, Parameters);
-						DispatchIndirectComputeShader(RHICmdList, *ComputeShader, SurfaceCacheResources.DispatchParameters.Buffer, 0);
+						ComputeShader->SetParameters(RHICmdList, View, DepthLevel);
+						DispatchComputeShader(RHICmdList, *ComputeShader, 1, 1, 1);
 
 						ComputeShader->UnsetParameters(RHICmdList);
 					}
-					else
+
 					{
-						TShaderMapRef<TCombineConesCS<false> > ComputeShader(View.ShaderMap);
+						TShaderMapRef<FCombineConesCS> ComputeShader(View.ShaderMap);
 
 						RHICmdList.SetComputeShader(ComputeShader->GetComputeShader());
-						ComputeShader->SetParameters(RHICmdList, View, DepthLevel, Parameters);
+						ComputeShader->SetParameters(RHICmdList, View, DepthLevel, Parameters, NULL, FMatrix::Identity, NULL, NULL);
 						DispatchIndirectComputeShader(RHICmdList, *ComputeShader, SurfaceCacheResources.DispatchParameters.Buffer, 0);
 
 						ComputeShader->UnsetParameters(RHICmdList);
@@ -4737,7 +4725,7 @@ bool FDeferredShadingSceneRenderer::RenderDistanceFieldAOSurfaceCache(
 				FPooledRenderTargetDesc Desc(FPooledRenderTargetDesc::Create2DDesc(BufferSize, PF_FloatRGBA, TexCreate_None, TexCreate_RenderTargetable | TexCreate_UAV, false));
 				GRenderTargetPool.FindFreeElement(Desc, BentNormalAccumulation, TEXT("BentNormalAccumulation"));
 
-				if (GDistanceFieldGI)
+				if (bUseDistanceFieldGI)
 				{
 					GRenderTargetPool.FindFreeElement(Desc, IrradianceAccumulation, TEXT("IrradianceAccumulation"));
 				}
@@ -4746,7 +4734,7 @@ bool FDeferredShadingSceneRenderer::RenderDistanceFieldAOSurfaceCache(
 			// Splat the surface cache records onto the opaque pixels, using less strict weighting so the lighting is smoothed in world space
 			{
 				SCOPED_DRAW_EVENT(RHICmdList, FinalIrradianceCacheSplat);
-				RenderIrradianceCacheInterpolation(RHICmdList, Views, BentNormalAccumulation, IrradianceAccumulation, DistanceFieldNormal->GetRenderTargetItem(), 0, GAODownsampleFactor, Parameters, true);
+				RenderIrradianceCacheInterpolation(RHICmdList, View, BentNormalAccumulation, IrradianceAccumulation, DistanceFieldNormal->GetRenderTargetItem(), 0, GAODownsampleFactor, Parameters, true);
 			}
 
 			GRenderTargetPool.VisualizeTexture.SetCheckPoint(RHICmdList, BentNormalAccumulation);
@@ -4754,7 +4742,7 @@ bool FDeferredShadingSceneRenderer::RenderDistanceFieldAOSurfaceCache(
 			// Post process the AO to cover over artifacts
 			TRefCountPtr<IPooledRenderTarget> BentNormalOutput;
 			TRefCountPtr<IPooledRenderTarget> IrradianceOutput;
-			PostProcessBentNormalAO(RHICmdList, Parameters, Views, BentNormalAccumulation->GetRenderTargetItem(), IrradianceAccumulation, DistanceFieldNormal->GetRenderTargetItem(), BentNormalOutput, IrradianceOutput);
+			PostProcessBentNormalAO(RHICmdList, Parameters, View, BentNormalAccumulation->GetRenderTargetItem(), IrradianceAccumulation, DistanceFieldNormal->GetRenderTargetItem(), BentNormalOutput, IrradianceOutput);
 
 			GRenderTargetPool.VisualizeTexture.SetCheckPoint(RHICmdList, BentNormalOutput);
 
@@ -4769,7 +4757,7 @@ bool FDeferredShadingSceneRenderer::RenderDistanceFieldAOSurfaceCache(
 				Desc.Format = PF_FloatRGBA;
 				GRenderTargetPool.FindFreeElement(Desc, OutDynamicBentNormalAO, TEXT("DynamicBentNormalAO"));
 
-				if (GDistanceFieldGI)
+				if (bUseDistanceFieldGI)
 				{
 					Desc.Format = PF_FloatRGB;
 					GRenderTargetPool.FindFreeElement(Desc, OutDynamicIrradiance, TEXT("DynamicIrradiance"));
@@ -4778,9 +4766,9 @@ bool FDeferredShadingSceneRenderer::RenderDistanceFieldAOSurfaceCache(
 				FTextureRHIParamRef RenderTargets[2] =
 				{
 					OutDynamicBentNormalAO->GetRenderTargetItem().TargetableTexture,
-					GDistanceFieldGI ? OutDynamicIrradiance->GetRenderTargetItem().TargetableTexture : NULL
+					bUseDistanceFieldGI ? OutDynamicIrradiance->GetRenderTargetItem().TargetableTexture : NULL
 				};
-				SetRenderTargets(RHICmdList, ARRAY_COUNT(RenderTargets) - (GDistanceFieldGI ? 0 : 1), RenderTargets, FTextureRHIParamRef(), 0, NULL);
+				SetRenderTargets(RHICmdList, ARRAY_COUNT(RenderTargets) - (bUseDistanceFieldGI ? 0 : 1), RenderTargets, FTextureRHIParamRef(), 0, NULL);
 			}
 
 			// Upsample to full resolution, write to output
@@ -4914,6 +4902,7 @@ void FDeferredShadingSceneRenderer::RenderDynamicSkyLighting(FRHICommandListImme
 		if (Scene->SkyLight->bCastShadows
 			&& ViewFamily.EngineShowFlags.DistanceFieldAO 
 			&& !ViewFamily.EngineShowFlags.VisualizeDistanceFieldAO
+			&& !ViewFamily.EngineShowFlags.VisualizeDistanceFieldGI
 			&& !ViewFamily.EngineShowFlags.VisualizeMeshDistanceFields)
 		{
 			bApplyShadowing = RenderDistanceFieldAOSurfaceCache(RHICmdList, Parameters, DynamicBentNormalAO, DynamicIrradiance, false);
@@ -4938,11 +4927,12 @@ void FDeferredShadingSceneRenderer::RenderDynamicSkyLighting(FRHICommandListImme
 				RHICmdList.SetBlendState(TStaticBlendState<CW_RGB, BO_Add, BF_One, BF_One>::GetRHI());
 			}
 
+			const bool bUseDistanceFieldGI = IsDistanceFieldGIAllowed(View);
 			TShaderMapRef< FPostProcessVS > VertexShader(View.ShaderMap);
 
 			if (bApplyShadowing)
 			{
-				if (GDistanceFieldGI)
+				if (bUseDistanceFieldGI)
 				{
 					TShaderMapRef<TDynamicSkyLightDiffusePS<true, true> > PixelShader(View.ShaderMap);
 
@@ -4961,7 +4951,7 @@ void FDeferredShadingSceneRenderer::RenderDynamicSkyLighting(FRHICommandListImme
 			}
 			else
 			{
-				if (GDistanceFieldGI)
+				if (bUseDistanceFieldGI)
 				{
 					TShaderMapRef<TDynamicSkyLightDiffusePS<false, true> > PixelShader(View.ShaderMap);
 
@@ -5007,7 +4997,6 @@ public:
 		OutEnvironment.SetDefine(TEXT("DOWNSAMPLE_FACTOR"), GAODownsampleFactor);
 		OutEnvironment.SetDefine(TEXT("THREADGROUP_SIZEX"), GDistanceFieldAOTileSizeX);
 		OutEnvironment.SetDefine(TEXT("THREADGROUP_SIZEY"), GDistanceFieldAOTileSizeY);
-		OutEnvironment.SetDefine(TEXT("MAX_OBJECTS_PER_TILE"), GMaxNumObjectsPerTile);
 	}
 
 	/** Default constructor. */
diff --git a/Engine/Source/Runtime/Renderer/Private/DistanceFieldSurfaceCacheLighting.h b/Engine/Source/Runtime/Renderer/Private/DistanceFieldSurfaceCacheLighting.h
--- a/Engine/Source/Runtime/Renderer/Private/DistanceFieldSurfaceCacheLighting.h
+++ b/Engine/Source/Runtime/Renderer/Private/DistanceFieldSurfaceCacheLighting.h
@@ -18,7 +18,7 @@ public:
 	float OcclusionMaxDistance;
 	float Contrast;
 
-	FDistanceFieldAOParameters(float InOcclusionMaxDistance = 600.0f, float InContrast = 0)
+	FDistanceFieldAOParameters(float InOcclusionMaxDistance, float InContrast = 0)
 	{
 		Contrast = FMath::Clamp(InContrast, .01f, 2.0f);
 		OcclusionMaxDistance = FMath::Clamp(InOcclusionMaxDistance, 200.0f, 3000.0f);
@@ -47,7 +47,7 @@ public:
 			Irradiance.Initialize(sizeof(FFloat16Color), MaxIrradianceCacheSamples, PF_FloatRGBA, BUF_Static);
 			ScatterDrawParameters.Initialize(sizeof(uint32), 4, PF_R32_UINT, BUF_Static | BUF_DrawIndirect);
 			SavedStartIndex.Initialize(sizeof(uint32), 1, PF_R32_UINT, BUF_Static);
-			TileCoordinate.Initialize(sizeof(uint16)* 2, MaxIrradianceCacheSamples, PF_R16G16_UINT, BUF_Static);
+			TileCoordinate.Initialize(sizeof(uint16) * 2, MaxIrradianceCacheSamples, PF_R16G16_UINT, BUF_Static);
 		}
 	}
 
@@ -91,6 +91,7 @@ public:
 
 		MinLevel = 100;
 		MaxLevel = 0;
+		bHasIrradiance = false;
 	}
 
 	~FSurfaceCacheResources()
@@ -181,6 +182,7 @@ public:
 	FRefinementLevelResources* TempResources;
 
 	bool bClearedResources;
+	bool bHasIrradiance;
 
 private:
 
@@ -270,3 +272,86 @@ private:
 	FShaderParameter AOMaxViewDistance;
 };
 
+inline bool DoesPlatformSupportDistanceFieldAO(EShaderPlatform Platform)
+{
+	return Platform == SP_PCD3D_SM5;
+}
+
+inline float GetMaxAOViewDistance()
+{
+	extern float GAOMaxViewDistance;
+	// Scene depth stored in fp16 alpha, must fade out before it runs out of range
+	// The fade extends past GAOMaxViewDistance a bit
+	return FMath::Min(GAOMaxViewDistance, 65000.0f);
+}
+
+class FMaxSizedRWBuffers : public FRenderResource
+{
+public:
+	FMaxSizedRWBuffers()
+	{
+		MaxSize = 0;
+	}
+
+	virtual void InitDynamicRHI()
+	{
+		check(0);
+	}
+
+	virtual void ReleaseDynamicRHI()
+	{
+		check(0);
+	}
+
+	void AllocateFor(int32 InMaxSize)
+	{
+		bool bReallocate = false;
+
+		if (InMaxSize > MaxSize)
+		{
+			MaxSize = InMaxSize;
+			bReallocate = true;
+		}
+
+		if (!IsInitialized())
+		{
+			InitResource();
+		}
+		else if (bReallocate)
+		{
+			UpdateRHI();
+		}
+	}
+
+	int32 GetMaxSize() const { return MaxSize; }
+
+protected:
+	int32 MaxSize;
+};
+
+// Must match usf
+const int32 RecordConeDataStride = 10;
+
+/**  */
+class FTemporaryIrradianceCacheResources : public FMaxSizedRWBuffers
+{
+public:
+
+	virtual void InitDynamicRHI()
+	{
+		if (MaxSize > 0)
+		{
+			ConeVisibility.Initialize(sizeof(float), MaxSize * NumConeSampleDirections, PF_R32_FLOAT, BUF_Static);
+			ConeData.Initialize(sizeof(float), MaxSize * NumConeSampleDirections * RecordConeDataStride, PF_R32_FLOAT, BUF_Static);
+		}
+	}
+
+	virtual void ReleaseDynamicRHI()
+	{
+		ConeVisibility.Release();
+		ConeData.Release();
+	}
+
+	FRWBuffer ConeVisibility;
+	FRWBuffer ConeData;
+};
diff --git a/Engine/Source/Runtime/Renderer/Private/PostProcess/PostProcessing.cpp b/Engine/Source/Runtime/Renderer/Private/PostProcess/PostProcessing.cpp
--- a/Engine/Source/Runtime/Renderer/Private/PostProcess/PostProcessing.cpp
+++ b/Engine/Source/Runtime/Renderer/Private/PostProcess/PostProcessing.cpp
@@ -809,6 +809,7 @@ bool FPostProcessing::AllowFullPostProcessing(const FViewInfo& View, ERHIFeature
 	return View.Family->EngineShowFlags.PostProcessing 
 		&& FeatureLevel >= ERHIFeatureLevel::SM4 
 		&& !View.Family->EngineShowFlags.VisualizeDistanceFieldAO
+		&& !View.Family->EngineShowFlags.VisualizeDistanceFieldGI
 		&& !View.Family->EngineShowFlags.VisualizeMeshDistanceFields;
 }
 
diff --git a/Engine/Source/Runtime/Renderer/Private/RendererScene.cpp b/Engine/Source/Runtime/Renderer/Private/RendererScene.cpp
--- a/Engine/Source/Runtime/Renderer/Private/RendererScene.cpp
+++ b/Engine/Source/Runtime/Renderer/Private/RendererScene.cpp
@@ -59,6 +59,7 @@ FSceneViewState::FSceneViewState()
 	TemporalAASampleIndex = 0;
 	TemporalAASampleCount = 1;
 	AOTileIntersectionResources = NULL;
+	VPLTileIntersectionResources = NULL;
 	bBokehDOFHistory = true;
 	bBokehDOFHistory2 = true;
 
diff --git a/Engine/Source/Runtime/Renderer/Private/ScenePrivate.h b/Engine/Source/Runtime/Renderer/Private/ScenePrivate.h
--- a/Engine/Source/Runtime/Renderer/Private/ScenePrivate.h
+++ b/Engine/Source/Runtime/Renderer/Private/ScenePrivate.h
@@ -548,6 +548,8 @@ public:
 	/** Distance field AO tile intersection GPU resources.  Last frame's state is not used, but they must be sized exactly to the view so stored here. */
 	class FTileIntersectionResources* AOTileIntersectionResources;
 
+	class FVPLTileIntersectionResources* VPLTileIntersectionResources;
+
 	// Is DOFHistoryRT set from Bokeh DOF?
 	bool bBokehDOFHistory;
 	bool bBokehDOFHistory2;
diff --git a/Engine/Source/Runtime/Renderer/Private/ShadowRendering.h b/Engine/Source/Runtime/Renderer/Private/ShadowRendering.h
--- a/Engine/Source/Runtime/Renderer/Private/ShadowRendering.h
+++ b/Engine/Source/Runtime/Renderer/Private/ShadowRendering.h
@@ -1443,4 +1443,17 @@ private:
 
 };
 
+/** A transform the remaps depth and potentially projects onto some plane. */
+struct FShadowProjectionMatrix: FMatrix
+{
+	FShadowProjectionMatrix(float MinZ,float MaxZ,const FVector4& WAxis):
+		FMatrix(
+		FPlane(1,	0,	0,													WAxis.X),
+		FPlane(0,	1,	0,													WAxis.Y),
+		FPlane(0,	0,	(WAxis.Z * MaxZ + WAxis.W) / (MaxZ - MinZ),			WAxis.Z),
+		FPlane(0,	0,	-MinZ * (WAxis.Z * MaxZ + WAxis.W) / (MaxZ - MinZ),	WAxis.W)
+		)
+	{}
+};
+
 #endif
diff --git a/Engine/Source/Runtime/Renderer/Private/ShadowSetup.cpp b/Engine/Source/Runtime/Renderer/Private/ShadowSetup.cpp
--- a/Engine/Source/Runtime/Renderer/Private/ShadowSetup.cpp
+++ b/Engine/Source/Runtime/Renderer/Private/ShadowSetup.cpp
@@ -334,19 +334,6 @@ static bool GetBestShadowTransform(const FVector& ZAxis,const FBoundingBoxVertex
 	}
 }
 
-/** A transform the remaps depth and potentially projects onto some plane. */
-struct FShadowProjectionMatrix: FMatrix
-{
-	FShadowProjectionMatrix(float MinZ,float MaxZ,const FVector4& WAxis):
-		FMatrix(
-			FPlane(1,	0,	0,													WAxis.X),
-			FPlane(0,	1,	0,													WAxis.Y),
-			FPlane(0,	0,	(WAxis.Z * MaxZ + WAxis.W) / (MaxZ - MinZ),			WAxis.Z),
-			FPlane(0,	0,	-MinZ * (WAxis.Z * MaxZ + WAxis.W) / (MaxZ - MinZ),	WAxis.W)
-			)
-	{}
-};
-
 /** Initialization constructor for a per-object shadow. e.g. translucent particle system */
 FProjectedShadowInfo::FProjectedShadowInfo(
 	FLightSceneInfo* InLightSceneInfo,
diff --git a/Engine/Source/Runtime/ShaderCore/Private/ShaderCore.cpp b/Engine/Source/Runtime/ShaderCore/Private/ShaderCore.cpp
--- a/Engine/Source/Runtime/ShaderCore/Private/ShaderCore.cpp
+++ b/Engine/Source/Runtime/ShaderCore/Private/ShaderCore.cpp
@@ -141,6 +141,12 @@ bool FShaderParameterMap::FindParameterAllocation(const TCHAR* ParameterName,uin
 		OutBufferIndex = Allocation->BufferIndex;
 		OutBaseIndex = Allocation->BaseIndex;
 		OutSize = Allocation->Size;
+
+		if (Allocation->bBound)
+		{
+			UE_LOG(LogShaders, Warning, TEXT("Parameter %s was bound multiple times. Code error?"), ParameterName);
+		}
+
 		Allocation->bBound = true;
 		return true;
 	}
