Fixed compile errors

diff --git a/Engine/Source/Programs/ShaderCompileWorker/Private/ShaderCompileWorker.cpp b/Engine/Source/Programs/ShaderCompileWorker/Private/ShaderCompileWorker.cpp
--- a/Engine/Source/Programs/ShaderCompileWorker/Private/ShaderCompileWorker.cpp
+++ b/Engine/Source/Programs/ShaderCompileWorker/Private/ShaderCompileWorker.cpp
@@ -633,6 +633,8 @@ IMPLEMENT_APPLICATION(ShaderCompileWorker, "ShaderCompileWorker")
 
 INT32_MAIN_INT32_ARGC_TCHAR_ARGV()
 {
+	// FPlatformProcess::OpenProcess only implemented for windows atm
+#if PLATFORM_WINDOWS
 	if (ArgC == 4 && FCString::Strcmp(ArgV[1], TEXT("-xgemonitor")) == 0)
 	{
 		// Open handles to the two processes
@@ -658,7 +660,7 @@ INT32_MAIN_INT32_ARGC_TCHAR_ARGV()
 		}
 		return 0;
 	}
-
+#endif
 	if(ArgC < 6)
 	{
 		printf("ShaderCompileWorker is called by UE4, it requires specific command like arguments.\n");
diff --git a/Engine/Source/Runtime/Renderer/Private/DistanceFieldSurfaceCacheLighting.cpp b/Engine/Source/Runtime/Renderer/Private/DistanceFieldSurfaceCacheLighting.cpp
--- a/Engine/Source/Runtime/Renderer/Private/DistanceFieldSurfaceCacheLighting.cpp
+++ b/Engine/Source/Runtime/Renderer/Private/DistanceFieldSurfaceCacheLighting.cpp
@@ -4383,7 +4383,7 @@ bool FDeferredShadingSceneRenderer::RenderDistanceFieldAOSurfaceCache(
 			if (!SurfaceCacheResources.bClearedResources 
 				|| !GAOReuseAcrossFrames 
 				// Drop records that will have uninitialized Irradiance if switching from AO only to AO + GI
-				|| bUseDistanceFieldGI && !SurfaceCacheResources.bHasIrradiance)
+				|| (bUseDistanceFieldGI && !SurfaceCacheResources.bHasIrradiance))
 			{
 				// Reset the number of active cache records to 0
 				for (int32 DepthLevel = GAOMaxLevel; DepthLevel >= GAOMinLevel; DepthLevel--)
