Fix shader complexity viewmode for foliage.
Added DDM_AllOpaque to ensure we get a prepass for all opaque objects including non-occluders when in shader complexity viewmode.
UE-6667

diff --git a/Engine/Source/Runtime/Renderer/Private/DeferredShadingRenderer.cpp b/Engine/Source/Runtime/Renderer/Private/DeferredShadingRenderer.cpp
--- a/Engine/Source/Runtime/Renderer/Private/DeferredShadingRenderer.cpp
+++ b/Engine/Source/Runtime/Renderer/Private/DeferredShadingRenderer.cpp
@@ -77,7 +77,7 @@ FDeferredShadingSceneRenderer::FDeferredShadingSceneRenderer(const FSceneViewFam
 	// Shader complexity requires depth only pass to display masked material cost correctly
 	if (ViewFamily.EngineShowFlags.ShaderComplexity)
 	{
-		EarlyZPassMode = DDM_AllOccluders;
+		EarlyZPassMode = DDM_AllOpaque;
 	}
 
 	// developer override, good for profiling, can be useful as project setting
@@ -90,6 +90,7 @@ FDeferredShadingSceneRenderer::FDeferredShadingSceneRenderer(const FSceneViewFam
 			case 0: EarlyZPassMode = DDM_None; break;
 			case 1: EarlyZPassMode = DDM_NonMaskedOnly; break;
 			case 2: EarlyZPassMode = DDM_AllOccluders; break;
+			case 3: EarlyZPassMode = DDM_AllOpaque; break;
 		}
 	}
 }
@@ -1042,7 +1043,7 @@ bool FDeferredShadingSceneRenderer::RenderPrePassViewDynamic(FRHICommandList& RH
 			const FPrimitiveSceneProxy* PrimitiveSceneProxy = MeshBatchAndRelevance.PrimitiveSceneProxy;
 			bool bShouldUseAsOccluder = true;
 
-			if (EarlyZPassMode != DDM_AllOccluders)
+			if (EarlyZPassMode < DDM_AllOccluders)
 			{
 				extern float GMinScreenRadiusForDepthPrepass;
 				//@todo - move these proxy properties into FMeshBatchAndRelevance so we don't have to dereference the proxy in order to reject a mesh
diff --git a/Engine/Source/Runtime/Renderer/Private/DepthRendering.cpp b/Engine/Source/Runtime/Renderer/Private/DepthRendering.cpp
--- a/Engine/Source/Runtime/Renderer/Private/DepthRendering.cpp
+++ b/Engine/Source/Runtime/Renderer/Private/DepthRendering.cpp
@@ -388,7 +388,7 @@ bool FDepthDrawingPolicyFactory::DrawMesh(
 
 	//Do a per-FMeshBatch check on top of the proxy check in RenderPrePass to handle the case where a proxy that is relevant 
 	//to the depth only pass has to submit multiple FMeshElements but only some of them should be used as occluders.
-	if (Mesh.bUseAsOccluder)
+	if (Mesh.bUseAsOccluder || DrawingContext.DepthDrawingMode == DDM_AllOpaque)
 	{
 		const FMaterialRenderProxy* MaterialRenderProxy = Mesh.MaterialRenderProxy;
 		const FMaterial* Material = MaterialRenderProxy->GetMaterial(View.GetFeatureLevel());
@@ -431,6 +431,8 @@ bool FDepthDrawingPolicyFactory::DrawMesh(
 
 			switch(DrawingContext.DepthDrawingMode)
 			{
+			case DDM_AllOpaque:
+				break;
 			case DDM_AllOccluders: 
 				break;
 			case DDM_NonMaskedOnly: 
diff --git a/Engine/Source/Runtime/Renderer/Private/DepthRendering.h b/Engine/Source/Runtime/Renderer/Private/DepthRendering.h
--- a/Engine/Source/Runtime/Renderer/Private/DepthRendering.h
+++ b/Engine/Source/Runtime/Renderer/Private/DepthRendering.h
@@ -14,6 +14,8 @@ enum EDepthDrawingMode
 	DDM_NonMaskedOnly,
 	//
 	DDM_AllOccluders,
+	//
+	DDM_AllOpaque,
 };
 
 template<bool>
