Moved Roughness in GBuffer to a better spot
Rendering Optimization: Lighting shader can read roughness from a GBuffer texture that is already needed. NonDBuffer decals need one less MRT channel to blend.

diff --git a/Engine/Shaders/DeferredDecal.usf b/Engine/Shaders/DeferredDecal.usf
--- a/Engine/Shaders/DeferredDecal.usf
+++ b/Engine/Shaders/DeferredDecal.usf
@@ -71,10 +71,6 @@ void MainPS(
 #if DECAL_RENDERTARGET_COUNT > 3
 	, out float4 OutTarget3 : SV_Target3
 #endif
-
-#if DECAL_RENDERTARGET_COUNT > 4
-	, out float4 OutTarget4 : SV_Target4
-#endif
 	)
 {
 	half SceneW = PreviousDepth(ScreenPosition);
@@ -191,6 +187,7 @@ void MainPS(
 		// after base pass, before lighting (GBuffer)
 
 		// static lighting isn't updated by decals so we don't need to update that render target
+		float4 OutTarget4 = 0;
 		float4 OutTarget5 = 0;
 
 		EncodeGBuffer(Data, OutTarget1, OutTarget2, OutTarget3, OutTarget4, OutTarget5);
@@ -202,9 +199,8 @@ void MainPS(
 
 		OutTarget0.a = Opacity;	// Emissive
 		OutTarget1.a = Opacity;	// Normal
-		OutTarget2.a = Opacity;	// Metallic, Specular
+		OutTarget2.a = Opacity;	// Metallic, Specular, Roughness
 		OutTarget3.a = Opacity;	// BaseColor
-		OutTarget4.a = Opacity;	// Subsurface
 	}
 	#endif // DECAL_RENDERSTAGE
 #endif //DECAL_RENDERTARGET_COUNT > 1
diff --git a/Engine/Shaders/DeferredShadingCommon.usf b/Engine/Shaders/DeferredShadingCommon.usf
--- a/Engine/Shaders/DeferredShadingCommon.usf
+++ b/Engine/Shaders/DeferredShadingCommon.usf
@@ -335,18 +335,17 @@ void EncodeGBuffer(
 	{
 		bool bSubsurface = IsSubsurfaceModel(GBuffer.ShadingModelID);
 
-		// NOTE OutGBufferB.b is currently unused!
 		OutGBufferB.r = bSubsurface ? GBuffer.Opacity : GBuffer.Metallic;
 		OutGBufferB.g = GBuffer.Specular;
-		OutGBufferB.b = 0;
+		OutGBufferB.b = GBuffer.Roughness;
 		OutGBufferC.rgb = EncodeBaseColor(GBuffer.BaseColor);
 
 		// No space for AO. Multiply IndirectIrradiance by AO instead of storing.
 		OutGBufferC.a = EncodeIndirectIrradiance(GBuffer.IndirectIrradiance * GBuffer.GBufferAO) + QuantizationBias * (1.0 / 255.0);
 		OutGBufferB.a = GBuffer.ShadingModelID * (1.0 / 255.0);
 
-		// Roughness in OutGBufferD.r instead of OutGBufferB.a so that deferred decals can blend in roughness while using MRT
-		OutGBufferD.r = GBuffer.Roughness;
+		// unused, can be used to make a 4th CustomData channel
+		OutGBufferD.r = 0;
 		OutGBufferD.gba = GBuffer.CustomData;
 	}
 
@@ -385,7 +384,7 @@ FGBufferData DecodeGBufferData(
 	GBuffer.Specular = InGBufferB.g;
 	GBuffer.Opacity = bSubsurface ? InGBufferB.r : 0;
 	
-	GBuffer.Roughness = InGBufferD.r;
+	GBuffer.Roughness = InGBufferB.b;
 	GBuffer.CustomData = InGBufferD.gba;
 	GBuffer.GBufferAO = 1;
 	GBuffer.IndirectIrradiance = DecodeIndirectIrradiance(InGBufferC.a);
diff --git a/Engine/Source/Runtime/Renderer/Private/CompositionLighting/PostProcessDeferredDecals.cpp b/Engine/Source/Runtime/Renderer/Private/CompositionLighting/PostProcessDeferredDecals.cpp
--- a/Engine/Source/Runtime/Renderer/Private/CompositionLighting/PostProcessDeferredDecals.cpp
+++ b/Engine/Source/Runtime/Renderer/Private/CompositionLighting/PostProcessDeferredDecals.cpp
@@ -35,7 +35,7 @@ uint32 ComputeRenderTargetCount(ERenderTargetMode RenderTargetMode)
 {
 	switch(RenderTargetMode)
 	{
-		case RTM_SceneColorAndGBuffer:	return 5;
+		case RTM_SceneColorAndGBuffer:	return 4;
 		case RTM_DBuffer:				return 3;
 		case RTM_GBufferNormal:			return 1;
 		case RTM_SceneColor:			return 1;
@@ -242,9 +242,8 @@ void SetDecalBlendState(FRHICommandList& RHICmdList, const ERHIFeatureLevel::Typ
 					RHICmdList.SetBlendState( TStaticBlendState<
 						CW_RGB, BO_Add, BF_SourceAlpha, BF_One,						BO_Add, BF_Zero, BF_One,	// Emissive
 						CW_RGB, BO_Add, BF_SourceAlpha, BF_InverseSourceAlpha,		BO_Add, BF_Zero, BF_One,	// Normal
-						CW_RGB, BO_Add, BF_SourceAlpha, BF_InverseSourceAlpha,		BO_Add, BF_Zero, BF_One,	// Metallic, Specular
-						CW_RGB, BO_Add, BF_SourceAlpha, BF_InverseSourceAlpha,		BO_Add, BF_Zero, BF_One,	// BaseColor
-						CW_RED, BO_Add, BF_SourceAlpha, BF_InverseSourceAlpha,		BO_Add, BF_Zero, BF_One		// Roughness in r
+						CW_RGB, BO_Add, BF_SourceAlpha, BF_InverseSourceAlpha,		BO_Add, BF_Zero, BF_One,	// Metallic, Specular, Roughness
+						CW_RGB, BO_Add, BF_SourceAlpha, BF_InverseSourceAlpha,		BO_Add, BF_Zero, BF_One		// BaseColor
 					>::GetRHI() );
 				}
 				else
@@ -252,9 +251,8 @@ void SetDecalBlendState(FRHICommandList& RHICmdList, const ERHIFeatureLevel::Typ
 					RHICmdList.SetBlendState( TStaticBlendState<
 						CW_RGB, BO_Add, BF_SourceAlpha, BF_One,						BO_Add, BF_Zero, BF_One,	// Emissive
 						CW_RGB, BO_Add, BF_Zero,		BF_One,						BO_Add, BF_Zero, BF_One,	// Normal
-						CW_RGB, BO_Add, BF_SourceAlpha, BF_InverseSourceAlpha,		BO_Add, BF_Zero, BF_One,	// Metallic, Specular
-						CW_RGB, BO_Add, BF_SourceAlpha, BF_InverseSourceAlpha,		BO_Add, BF_Zero, BF_One,	// BaseColor
-						CW_RED, BO_Add, BF_SourceAlpha, BF_InverseSourceAlpha,		BO_Add, BF_Zero, BF_One		// Roughness in r
+						CW_RGB, BO_Add, BF_SourceAlpha, BF_InverseSourceAlpha,		BO_Add, BF_Zero, BF_One,	// Metallic, Specular, Roughness
+						CW_RGB, BO_Add, BF_SourceAlpha, BF_InverseSourceAlpha,		BO_Add, BF_Zero, BF_One		// BaseColor
 					>::GetRHI() );
 				}
 			}
@@ -263,9 +261,8 @@ void SetDecalBlendState(FRHICommandList& RHICmdList, const ERHIFeatureLevel::Typ
 				RHICmdList.SetBlendState( TStaticBlendState<
 					CW_RGB, BO_Add, BF_SourceAlpha, BF_One,							BO_Add, BF_Zero, BF_One,	// Emissive
 					CW_RGB, BO_Add, BF_SourceAlpha, BF_One,							BO_Add, BF_Zero, BF_One,	// Normal
-					CW_RGB, BO_Add, BF_SourceAlpha, BF_One,							BO_Add, BF_Zero, BF_One,	// Metallic, Specular
-					CW_RGB, BO_Add, BF_SourceAlpha,	BF_One,							BO_Add, BF_Zero, BF_One,	// BaseColor
-					CW_RED, BO_Add, BF_SourceAlpha, BF_One,							BO_Add, BF_Zero, BF_One		// Roughness in r
+					CW_RGB, BO_Add, BF_SourceAlpha, BF_One,							BO_Add, BF_Zero, BF_One,	// Metallic, Specular, Roughness
+					CW_RGB, BO_Add, BF_SourceAlpha,	BF_One,							BO_Add, BF_Zero, BF_One		// BaseColor
 				>::GetRHI() );
 			}
 			break;
@@ -278,9 +275,8 @@ void SetDecalBlendState(FRHICommandList& RHICmdList, const ERHIFeatureLevel::Typ
 					RHICmdList.SetBlendState( TStaticBlendState<
 						CW_RGB, BO_Add, BF_SourceAlpha, BF_One,						BO_Add, BF_Zero, BF_One,	// Emissive
 						CW_RGB, BO_Add, BF_SourceAlpha, BF_InverseSourceAlpha,		BO_Add, BF_Zero, BF_One,	// Normal
-						CW_RGB, BO_Add, BF_SourceAlpha, BF_InverseSourceAlpha,		BO_Add, BF_Zero, BF_One,	// Metallic, Specular
-						CW_RGB, BO_Add, BF_DestColor,	BF_InverseSourceAlpha,		BO_Add, BF_Zero, BF_One,	// BaseColor
-						CW_RED, BO_Add, BF_SourceAlpha, BF_InverseSourceAlpha,		BO_Add, BF_Zero, BF_One		// Roughness in r
+						CW_RGB, BO_Add, BF_SourceAlpha, BF_InverseSourceAlpha,		BO_Add, BF_Zero, BF_One,	// Metallic, Specular, Roughness
+						CW_RGB, BO_Add, BF_DestColor,	BF_InverseSourceAlpha,		BO_Add, BF_Zero, BF_One		// BaseColor
 					>::GetRHI() );
 				}
 				else
@@ -288,9 +284,8 @@ void SetDecalBlendState(FRHICommandList& RHICmdList, const ERHIFeatureLevel::Typ
 					RHICmdList.SetBlendState( TStaticBlendState<
 						CW_RGB, BO_Add, BF_SourceAlpha, BF_One,						BO_Add, BF_Zero, BF_One,	// Emissive
 						CW_RGB, BO_Add, BF_Zero,		BF_One,						BO_Add, BF_Zero, BF_One,	// Normal
-						CW_RGB, BO_Add, BF_SourceAlpha, BF_InverseSourceAlpha,		BO_Add, BF_Zero, BF_One,	// Metallic, Specular
-						CW_RGB, BO_Add, BF_DestColor,	BF_InverseSourceAlpha,		BO_Add, BF_Zero, BF_One,	// BaseColor
-						CW_RED, BO_Add, BF_SourceAlpha, BF_InverseSourceAlpha,		BO_Add, BF_Zero, BF_One		// Roughness in r
+						CW_RGB, BO_Add, BF_SourceAlpha, BF_InverseSourceAlpha,		BO_Add, BF_Zero, BF_One,	// Metallic, Specular, Roughness
+						CW_RGB, BO_Add, BF_DestColor,	BF_InverseSourceAlpha,		BO_Add, BF_Zero, BF_One		// BaseColor
 					>::GetRHI() );
 				}
 			}
@@ -299,9 +294,8 @@ void SetDecalBlendState(FRHICommandList& RHICmdList, const ERHIFeatureLevel::Typ
 				RHICmdList.SetBlendState( TStaticBlendState<
 					CW_RGB, BO_Add, BF_SourceAlpha, BF_InverseSourceAlpha,			BO_Add, BF_Zero, BF_One,	// Emissive
 					CW_RGB, BO_Add, BF_SourceAlpha, BF_InverseSourceAlpha,			BO_Add, BF_Zero, BF_One,	// Normal
-					CW_RGB, BO_Add, BF_SourceAlpha, BF_InverseSourceAlpha,			BO_Add, BF_Zero, BF_One,	// Metallic, Specular
-					CW_RGB, BO_Add, BF_SourceAlpha,	BF_InverseSourceAlpha,			BO_Add, BF_Zero, BF_One,	// BaseColor
-					CW_RED, BO_Add, BF_SourceAlpha, BF_InverseSourceAlpha,			BO_Add, BF_Zero, BF_One		// Roughness in r
+					CW_RGB, BO_Add, BF_SourceAlpha, BF_InverseSourceAlpha,			BO_Add, BF_Zero, BF_One,	// Metallic, Specular, Roughness
+					CW_RGB, BO_Add, BF_SourceAlpha,	BF_InverseSourceAlpha,			BO_Add, BF_Zero, BF_One		// BaseColor
 				>::GetRHI() );
 			}
 			break;
@@ -674,9 +668,8 @@ bool RenderPreStencil(FRenderingCompositePassContext& Context, const FMaterialSh
 	Context.RHICmdList.SetBlendState(TStaticBlendState<
 		CW_NONE, BO_Add, BF_SourceAlpha, BF_InverseSourceAlpha,		BO_Add, BF_Zero, BF_One,	// Emissive
 		CW_NONE, BO_Add, BF_SourceAlpha, BF_InverseSourceAlpha,		BO_Add, BF_Zero, BF_One,	// Normal
-		CW_NONE, BO_Add, BF_SourceAlpha, BF_InverseSourceAlpha,		BO_Add, BF_Zero, BF_One,	// Metallic, Specular
-		CW_NONE, BO_Add, BF_SourceAlpha, BF_InverseSourceAlpha,		BO_Add, BF_Zero, BF_One,	// BaseColor
-		CW_NONE, BO_Add, BF_SourceAlpha, BF_InverseSourceAlpha,		BO_Add, BF_Zero, BF_One		// Roughness in r
+		CW_NONE, BO_Add, BF_SourceAlpha, BF_InverseSourceAlpha,		BO_Add, BF_Zero, BF_One,	// Metallic, Specular, Roughness
+		CW_NONE, BO_Add, BF_SourceAlpha, BF_InverseSourceAlpha,		BO_Add, BF_Zero, BF_One		// BaseColor
 	>::GetRHI() );
 
 	// Carmack's reverse on the bounds
@@ -926,12 +919,11 @@ void FRCPassPostProcessDeferredDecals::Process(FRenderingCompositePassContext& C
 				{
 					case RTM_SceneColorAndGBuffer:
 						{
-							FTextureRHIParamRef RenderTargets[5];
+							FTextureRHIParamRef RenderTargets[4];
 							RenderTargets[0] = GSceneRenderTargets.GetSceneColor()->GetRenderTargetItem().TargetableTexture;
 							RenderTargets[1] = GSceneRenderTargets.GBufferA->GetRenderTargetItem().TargetableTexture;
 							RenderTargets[2] = GSceneRenderTargets.GBufferB->GetRenderTargetItem().TargetableTexture;
 							RenderTargets[3] = GSceneRenderTargets.GBufferC->GetRenderTargetItem().TargetableTexture;
-							RenderTargets[4] = GSceneRenderTargets.GBufferD->GetRenderTargetItem().TargetableTexture;
 							SetRenderTargets(RHICmdList, ARRAY_COUNT(RenderTargets), RenderTargets, GSceneRenderTargets.GetSceneDepthSurface(), 0, NULL);
 						}
 						break;
