Fix for bug UE-6754. Distortion pass was never rendered if only dynamic distortion relevant primitives were on screen.

diff --git a/Engine/Source/Runtime/Renderer/Private/DistortionRendering.cpp b/Engine/Source/Runtime/Renderer/Private/DistortionRendering.cpp
--- a/Engine/Source/Runtime/Renderer/Private/DistortionRendering.cpp
+++ b/Engine/Source/Runtime/Renderer/Private/DistortionRendering.cpp
@@ -709,7 +709,7 @@ bool FDistortionPrimSet::DrawAccumulatedOffsets(FRHICommandListImmediate& RHICmd
 				if (MeshBatchAndRelevance.PrimitiveSceneProxy == PrimitiveSceneProxy)
 				{
 					const FMeshBatch& MeshBatch = *MeshBatchAndRelevance.Mesh;
-					TDistortionMeshDrawingPolicyFactory<FDistortMeshAccumulatePolicy>::DrawDynamicMesh(RHICmdList, View, Context, MeshBatch, false, false, MeshBatchAndRelevance.PrimitiveSceneProxy, MeshBatch.BatchHitProxyId);
+					bDirty |= TDistortionMeshDrawingPolicyFactory<FDistortMeshAccumulatePolicy>::DrawDynamicMesh(RHICmdList, View, Context, MeshBatch, false, false, MeshBatchAndRelevance.PrimitiveSceneProxy, MeshBatch.BatchHitProxyId);
 				}
 			}
 
