Fix for incorrect handling of r.EarlyZPass=3 (default behavior) breaking screendoor fading on foliage.

diff --git a/Engine/Source/Runtime/Renderer/Private/DeferredShadingRenderer.cpp b/Engine/Source/Runtime/Renderer/Private/DeferredShadingRenderer.cpp
--- a/Engine/Source/Runtime/Renderer/Private/DeferredShadingRenderer.cpp
+++ b/Engine/Source/Runtime/Renderer/Private/DeferredShadingRenderer.cpp
@@ -90,7 +90,7 @@ FDeferredShadingSceneRenderer::FDeferredShadingSceneRenderer(const FSceneViewFam
 			case 0: EarlyZPassMode = DDM_None; break;
 			case 1: EarlyZPassMode = DDM_NonMaskedOnly; break;
 			case 2: EarlyZPassMode = DDM_AllOccluders; break;
-			case 3: EarlyZPassMode = DDM_AllOpaque; break;
+			case 3: break;	// Note: 3 indicates "default behavior" and does not specify an override
 		}
 	}
 }
