reduced verbosity

diff --git a/Engine/Source/Runtime/Renderer/Private/PostProcess/SceneRenderTargets.cpp b/Engine/Source/Runtime/Renderer/Private/PostProcess/SceneRenderTargets.cpp
--- a/Engine/Source/Runtime/Renderer/Private/PostProcess/SceneRenderTargets.cpp
+++ b/Engine/Source/Runtime/Renderer/Private/PostProcess/SceneRenderTargets.cpp
@@ -267,7 +267,7 @@ void FSceneRenderTargets::Allocate(const FSceneViewFamily& ViewFamily)
 		// Reinitialize the render targets for the given size.
 		SetBufferSize(DesiredBufferSize.X, DesiredBufferSize.Y);
 
-		UE_LOG(LogRenderer, Warning, TEXT("Reallocating scene render targets to support %ux%u (Frame:%u)."), BufferSize.X, BufferSize.Y, ViewFamily.FrameNumber);
+		UE_LOG(LogRenderer, Log, TEXT("Reallocating scene render targets to support %ux%u (Frame:%u)."), BufferSize.X, BufferSize.Y, ViewFamily.FrameNumber);
 
 		UpdateRHI();
 	}
