fixed using SubsurfaceProfile with LPV

diff --git a/Engine/Source/Runtime/Renderer/Private/PostProcess/PostProcessSubsurface.cpp b/Engine/Source/Runtime/Renderer/Private/PostProcess/PostProcessSubsurface.cpp
--- a/Engine/Source/Runtime/Renderer/Private/PostProcess/PostProcessSubsurface.cpp
+++ b/Engine/Source/Runtime/Renderer/Private/PostProcess/PostProcessSubsurface.cpp
@@ -617,7 +617,7 @@ void FRCPassPostProcessSubsurfaceExtractSpecular::Process(FRenderingCompositePas
 
 FPooledRenderTargetDesc FRCPassPostProcessSubsurfaceExtractSpecular::ComputeOutputDesc(EPassOutputId InPassOutputId) const
 {
-	FPooledRenderTargetDesc Ret = PassInputs[0].GetOutput()->RenderTargetDesc;
+	FPooledRenderTargetDesc Ret = GSceneRenderTargets.GetSceneColor()->GetDesc();
 
 	Ret.Reset();
 	Ret.DebugName = TEXT("SubsurfaceExtractSpecular");
