added comment

diff --git a/Engine/Source/Runtime/Renderer/Private/SystemTextures.cpp b/Engine/Source/Runtime/Renderer/Private/SystemTextures.cpp
--- a/Engine/Source/Runtime/Renderer/Private/SystemTextures.cpp
+++ b/Engine/Source/Runtime/Renderer/Private/SystemTextures.cpp
@@ -219,7 +219,7 @@ void FSystemTextures::InitializeTextures(FRHICommandListImmediate& RHICmdList, E
 			int32 Reorder[16] = { 0, 11, 7, 3, 10, 4, 15, 12, 6, 8, 1, 14, 13, 2, 9, 5 };
 			int32 w = Reorder[Pos];
 
-			// ordered sampling of the rotation basis
+			// ordered sampling of the rotation basis (*2 is missing as we use mirrored samples)
 			float ww = w / 16.0f * PI;
 
 			// randomize base scale
