DF AO from heightfields view frustum culling
* 21ms in GDC level for landscape occlusion down to .12ms
* Also fixed heightfield occlusion from Landscapes with subsections

diff --git a/Engine/Shaders/DistanceFieldSurfaceCacheLightingCompute.usf b/Engine/Shaders/DistanceFieldSurfaceCacheLightingCompute.usf
--- a/Engine/Shaders/DistanceFieldSurfaceCacheLightingCompute.usf
+++ b/Engine/Shaders/DistanceFieldSurfaceCacheLightingCompute.usf
@@ -1002,6 +1002,7 @@ void CalculateHeightfieldOcclusionCS(
 				float HeightfieldHeight = ((SampleValue.r * 255.0 * 256.0 + SampleValue.g * 255.0) - 32768.0) * TERRAIN_ZSCALE;
 
 				float2 SampleNormal = float2(SampleValue.b, SampleValue.a) * float2(2.0,2.0) - float2(1.0,1.0);
+				// @todo - use the normal to compute occlusion from an arbitrary plane, instead of assuming upward facing plane
 				float3 HeightfieldNormal = float3(SampleNormal, sqrt(max(1.0 - dot(SampleNormal, SampleNormal), 0.0)));
 
 				float3 HeightfieldShadingPosition = float3(LocalPosition.xy, HeightfieldHeight);
diff --git a/Engine/Source/Runtime/Landscape/Private/LandscapeRender.cpp b/Engine/Source/Runtime/Landscape/Private/LandscapeRender.cpp
--- a/Engine/Source/Runtime/Landscape/Private/LandscapeRender.cpp
+++ b/Engine/Source/Runtime/Landscape/Private/LandscapeRender.cpp
@@ -2578,3 +2578,16 @@ void FLandscapeComponentSceneProxy::ChangeLODDistanceFactor_RenderThread(float I
 {
 	LODDistance = InLODDistanceFactor;
 }
+
+void FLandscapeComponentSceneProxy::GetHeightfieldRepresentation(UTexture2D*& OutHeightmapTexture, FVector4& OutHeightfieldScaleBias, FVector4& OutMinMaxUV)
+{
+	OutHeightmapTexture = HeightmapTexture;
+	OutHeightfieldScaleBias = HeightmapScaleBias;
+
+	OutMinMaxUV = FVector4(
+		HeightmapScaleBias.Z, 
+		HeightmapScaleBias.W, 
+		//@todo - subsections duplicate a line of border heights in the heightmap, take this into account in the UV mapping
+		HeightmapScaleBias.Z + SubsectionSizeVerts * NumSubsections * HeightmapScaleBias.X, 
+		HeightmapScaleBias.W + SubsectionSizeVerts * NumSubsections * HeightmapScaleBias.Y);
+}
\ No newline at end of file
diff --git a/Engine/Source/Runtime/Landscape/Public/LandscapeRender.h b/Engine/Source/Runtime/Landscape/Public/LandscapeRender.h
--- a/Engine/Source/Runtime/Landscape/Public/LandscapeRender.h
+++ b/Engine/Source/Runtime/Landscape/Public/LandscapeRender.h
@@ -515,15 +515,7 @@ public:
 	// FLandcapeSceneProxy
 	void ChangeLODDistanceFactor_RenderThread(float InLODDistanceFactor);
 
-	virtual void GetHeightfieldRepresentation(UTexture2D*& OutHeightmapTexture, FVector4& OutHeightfieldScaleBias, FVector4& OutMinMaxUV) override
-	{
-		OutHeightmapTexture = HeightmapTexture;
-		OutHeightfieldScaleBias = HeightmapScaleBias;
-		// Section base is in terms of quads, convert into texels
-		int32 TexelsSectionBaseX = SectionBase.X / SubsectionSizeQuads * SubsectionSizeVerts;
-		int32 TexelsSectionBaseY = SectionBase.Y / SubsectionSizeQuads * SubsectionSizeVerts;
-		OutMinMaxUV = FVector4(TexelsSectionBaseX * HeightmapScaleBias.X, TexelsSectionBaseY * HeightmapScaleBias.Y, (TexelsSectionBaseX + SubsectionSizeVerts - 1) * HeightmapScaleBias.X, (TexelsSectionBaseY + SubsectionSizeVerts - 1) * HeightmapScaleBias.Y);
-	}
+	virtual void GetHeightfieldRepresentation(UTexture2D*& OutHeightmapTexture, FVector4& OutHeightfieldScaleBias, FVector4& OutMinMaxUV) override;
 };
 
 class FLandscapeDebugMaterialRenderProxy : public FMaterialRenderProxy
diff --git a/Engine/Source/Runtime/Renderer/Private/DistanceFieldSurfaceCacheLighting.cpp b/Engine/Source/Runtime/Renderer/Private/DistanceFieldSurfaceCacheLighting.cpp
--- a/Engine/Source/Runtime/Renderer/Private/DistanceFieldSurfaceCacheLighting.cpp
+++ b/Engine/Source/Runtime/Renderer/Private/DistanceFieldSurfaceCacheLighting.cpp
@@ -211,6 +211,14 @@ FAutoConsoleVariableRef CVarAOOverwriteSceneColor(
 	ECVF_Cheat | ECVF_RenderThreadSafe
 	);
 
+int32 GAOHeightfieldOcclusion = 1;
+FAutoConsoleVariableRef CVarAOHeightfieldOcclusion(
+	TEXT("r.AOHeightfieldOcclusion"),
+	GAOHeightfieldOcclusion,
+	TEXT("Whether to compute AO from heightfields (landscape)"),
+	ECVF_Cheat | ECVF_RenderThreadSafe
+	);
+
 IMPLEMENT_UNIFORM_BUFFER_STRUCT(FAOSampleData2,TEXT("AOSamples2"));
 
 FIntPoint GetBufferSizeForAO()
@@ -4592,25 +4600,32 @@ bool FDeferredShadingSceneRenderer::RenderDistanceFieldAOSurfaceCache(
 					ComputeShader->UnsetParameters(RHICmdList);
 				}
 
-				if (Scene->DistanceFieldSceneData.HeightfieldPrimitives.Num() > 0)
+				if (GAOHeightfieldOcclusion && Scene->DistanceFieldSceneData.HeightfieldPrimitives.Num() > 0)
 				{
-					int32 NumPrimitives = Scene->DistanceFieldSceneData.HeightfieldPrimitives.Num();
+					const int32 NumPrimitives = Scene->DistanceFieldSceneData.HeightfieldPrimitives.Num();
+					const float MaxDistanceSquared = FMath::Square(GAOMaxViewDistance + Parameters.OcclusionMaxDistance);
 
 					TMap<UTexture2D*, TArray<FHeightfieldDescription, SceneRenderingAllocator>, SceneRenderingSetAllocator> HeightfieldTextures; 
 
-					for (int32 HeightfieldIndex = 0; HeightfieldIndex < NumPrimitives; HeightfieldIndex++)
+					for (int32 HeightfieldPrimitiveIndex = 0; HeightfieldPrimitiveIndex < NumPrimitives; HeightfieldPrimitiveIndex++)
 					{
-						const FPrimitiveSceneInfo* HeightfieldPrimitive = Scene->DistanceFieldSceneData.HeightfieldPrimitives[HeightfieldIndex];
-
-						UTexture2D* HeightfieldTexture = NULL;
-						FVector4 HeightfieldScaleBias(0, 0, 0, 0);
-						FVector4 MinMaxUV(0, 0, 0, 0);
-						HeightfieldPrimitive->Proxy->GetHeightfieldRepresentation(HeightfieldTexture, HeightfieldScaleBias, MinMaxUV);
-
-						if (HeightfieldTexture && HeightfieldTexture->Resource->TextureRHI)
+						const FPrimitiveSceneInfo* HeightfieldPrimitive = Scene->DistanceFieldSceneData.HeightfieldPrimitives[HeightfieldPrimitiveIndex];
+						const FBoxSphereBounds& PrimitiveBounds = HeightfieldPrimitive->Proxy->GetBounds();
+						const float DistanceToPrimitiveSq = (PrimitiveBounds.Origin - View.ViewMatrices.ViewOrigin).SizeSquared();
+						
+						if (View.ViewFrustum.IntersectSphere(PrimitiveBounds.Origin, PrimitiveBounds.SphereRadius + Parameters.OcclusionMaxDistance)
+							&& DistanceToPrimitiveSq < MaxDistanceSquared)
 						{
-							TArray<FHeightfieldDescription, SceneRenderingAllocator>& HeightfieldDescriptions = HeightfieldTextures.FindOrAdd(HeightfieldTexture);
-							HeightfieldDescriptions.Add(FHeightfieldDescription(HeightfieldScaleBias, MinMaxUV, HeightfieldPrimitive->Proxy->GetLocalToWorld()));
+							UTexture2D* HeightfieldTexture = NULL;
+							FVector4 HeightfieldScaleBias(0, 0, 0, 0);
+							FVector4 MinMaxUV(0, 0, 0, 0);
+							HeightfieldPrimitive->Proxy->GetHeightfieldRepresentation(HeightfieldTexture, HeightfieldScaleBias, MinMaxUV);
+
+							if (HeightfieldTexture && HeightfieldTexture->Resource->TextureRHI)
+							{
+								TArray<FHeightfieldDescription, SceneRenderingAllocator>& HeightfieldDescriptions = HeightfieldTextures.FindOrAdd(HeightfieldTexture);
+								HeightfieldDescriptions.Add(FHeightfieldDescription(HeightfieldScaleBias, MinMaxUV, HeightfieldPrimitive->Proxy->GetLocalToWorld()));
+							}
 						}
 					}
 
@@ -4621,51 +4636,55 @@ bool FDeferredShadingSceneRenderer::RenderDistanceFieldAOSurfaceCache(
 						for (TMap<UTexture2D*, TArray<FHeightfieldDescription, SceneRenderingAllocator>, SceneRenderingSetAllocator>::TIterator It(HeightfieldTextures); It; ++It)
 						{
 							const TArray<FHeightfieldDescription, SceneRenderingAllocator>& HeightfieldDescriptions = It.Value();
-							TArray<FVector4, SceneRenderingAllocator> HeightfieldDescriptionData;
-							HeightfieldDescriptionData.Empty(HeightfieldDescriptions.Num() * GHeightfieldDescriptions.Data.Stride);
-
-							for (int32 DescriptionIndex = 0; DescriptionIndex < HeightfieldDescriptions.Num(); DescriptionIndex++)
-							{
-								const FHeightfieldDescription& Description = HeightfieldDescriptions[DescriptionIndex];
 
-								HeightfieldDescriptionData.Add(Description.HeightfieldScaleBias);
-								HeightfieldDescriptionData.Add(Description.MinMaxUV);
-								
-								const FMatrix WorldToLocal = Description.LocalToWorld.Inverse();
-
-								HeightfieldDescriptionData.Add(*(FVector4*)&WorldToLocal.M[0]);
-								HeightfieldDescriptionData.Add(*(FVector4*)&WorldToLocal.M[1]);
-								HeightfieldDescriptionData.Add(*(FVector4*)&WorldToLocal.M[2]);
-								HeightfieldDescriptionData.Add(*(FVector4*)&WorldToLocal.M[3]);
-
-								HeightfieldDescriptionData.Add(*(FVector4*)&Description.LocalToWorld.M[0]);
-								HeightfieldDescriptionData.Add(*(FVector4*)&Description.LocalToWorld.M[1]);
-								HeightfieldDescriptionData.Add(*(FVector4*)&Description.LocalToWorld.M[2]);
-								HeightfieldDescriptionData.Add(*(FVector4*)&Description.LocalToWorld.M[3]);
-							}
-
-							if (HeightfieldDescriptionData.Num() > GHeightfieldDescriptions.Data.MaxElements)
+							if (HeightfieldDescriptions.Num() > 0)
 							{
-								GHeightfieldDescriptions.Data.MaxElements = HeightfieldDescriptionData.Num() * 5 / 4;
-								GHeightfieldDescriptions.Data.Release();
-								GHeightfieldDescriptions.Data.Initialize();
-							}
-
-							void* LockedBuffer = RHILockVertexBuffer(GHeightfieldDescriptions.Data.Buffer, 0, GHeightfieldDescriptions.Data.Buffer->GetSize(), RLM_WriteOnly);
-							const uint32 MemcpySize = HeightfieldDescriptionData.GetTypeSize() * HeightfieldDescriptionData.Num();
-							check(GHeightfieldDescriptions.Data.Buffer->GetSize() >= MemcpySize);
-							FPlatformMemory::Memcpy(LockedBuffer, HeightfieldDescriptionData.GetData(), MemcpySize);
-							RHIUnlockVertexBuffer(GHeightfieldDescriptions.Data.Buffer);
+								TArray<FVector4, SceneRenderingAllocator> HeightfieldDescriptionData;
+								HeightfieldDescriptionData.Empty(HeightfieldDescriptions.Num() * GHeightfieldDescriptions.Data.Stride);
 
-							UTexture2D* HeightfieldTexture = It.Key();
+								for (int32 DescriptionIndex = 0; DescriptionIndex < HeightfieldDescriptions.Num(); DescriptionIndex++)
+								{
+									const FHeightfieldDescription& Description = HeightfieldDescriptions[DescriptionIndex];
 
-							{
-								TShaderMapRef<FCalculateHeightfieldOcclusionCS> ComputeShader(View.ShaderMap);
-								RHICmdList.SetComputeShader(ComputeShader->GetComputeShader());
-								ComputeShader->SetParameters(RHICmdList, View, DepthLevel, HeightfieldTexture, HeightfieldDescriptions.Num(), Parameters);
-								DispatchIndirectComputeShader(RHICmdList, *ComputeShader, SurfaceCacheResources.DispatchParameters.Buffer, 0);
-
-								ComputeShader->UnsetParameters(RHICmdList);
+									HeightfieldDescriptionData.Add(Description.HeightfieldScaleBias);
+									HeightfieldDescriptionData.Add(Description.MinMaxUV);
+								
+									const FMatrix WorldToLocal = Description.LocalToWorld.Inverse();
+
+									HeightfieldDescriptionData.Add(*(FVector4*)&WorldToLocal.M[0]);
+									HeightfieldDescriptionData.Add(*(FVector4*)&WorldToLocal.M[1]);
+									HeightfieldDescriptionData.Add(*(FVector4*)&WorldToLocal.M[2]);
+									HeightfieldDescriptionData.Add(*(FVector4*)&WorldToLocal.M[3]);
+
+									HeightfieldDescriptionData.Add(*(FVector4*)&Description.LocalToWorld.M[0]);
+									HeightfieldDescriptionData.Add(*(FVector4*)&Description.LocalToWorld.M[1]);
+									HeightfieldDescriptionData.Add(*(FVector4*)&Description.LocalToWorld.M[2]);
+									HeightfieldDescriptionData.Add(*(FVector4*)&Description.LocalToWorld.M[3]);
+								}
+
+								if (HeightfieldDescriptionData.Num() > GHeightfieldDescriptions.Data.MaxElements)
+								{
+									GHeightfieldDescriptions.Data.MaxElements = HeightfieldDescriptionData.Num() * 5 / 4;
+									GHeightfieldDescriptions.Data.Release();
+									GHeightfieldDescriptions.Data.Initialize();
+								}
+
+								void* LockedBuffer = RHILockVertexBuffer(GHeightfieldDescriptions.Data.Buffer, 0, GHeightfieldDescriptions.Data.Buffer->GetSize(), RLM_WriteOnly);
+								const uint32 MemcpySize = HeightfieldDescriptionData.GetTypeSize() * HeightfieldDescriptionData.Num();
+								check(GHeightfieldDescriptions.Data.Buffer->GetSize() >= MemcpySize);
+								FPlatformMemory::Memcpy(LockedBuffer, HeightfieldDescriptionData.GetData(), MemcpySize);
+								RHIUnlockVertexBuffer(GHeightfieldDescriptions.Data.Buffer);
+
+								UTexture2D* HeightfieldTexture = It.Key();
+
+								{
+									TShaderMapRef<FCalculateHeightfieldOcclusionCS> ComputeShader(View.ShaderMap);
+									RHICmdList.SetComputeShader(ComputeShader->GetComputeShader());
+									ComputeShader->SetParameters(RHICmdList, View, DepthLevel, HeightfieldTexture, HeightfieldDescriptions.Num(), Parameters);
+									DispatchIndirectComputeShader(RHICmdList, *ComputeShader, SurfaceCacheResources.DispatchParameters.Buffer, 0);
+
+									ComputeShader->UnsetParameters(RHICmdList);
+								}
 							}
 						}
 					}
