fixed eyeadaptation if only VisualizeHDR is on

diff --git a/Engine/Source/Runtime/Renderer/Private/PostProcess/PostProcessing.cpp b/Engine/Source/Runtime/Renderer/Private/PostProcess/PostProcessing.cpp
--- a/Engine/Source/Runtime/Renderer/Private/PostProcess/PostProcessing.cpp
+++ b/Engine/Source/Runtime/Renderer/Private/PostProcess/PostProcessing.cpp
@@ -1304,6 +1304,8 @@ void FPostProcessing::Process(FRHICommandListImmediate& RHICmdList, FViewInfo& V
 			Node->SetInput(ePId_Input1, Histogram);
 			Node->SetInput(ePId_Input2, HDRColor);
 			Node->SetInput(ePId_Input3, HistogramOverScreen);
+			Node->AddDependency(EyeAdaptation);
+
 			Context.FinalOutput = FRenderingCompositeOutputRef(Node);
 		}
 
