Fixed visualize mesh distance fields

diff --git a/Engine/Source/Runtime/Renderer/Private/DistanceFieldSurfaceCacheLighting.cpp b/Engine/Source/Runtime/Renderer/Private/DistanceFieldSurfaceCacheLighting.cpp
--- a/Engine/Source/Runtime/Renderer/Private/DistanceFieldSurfaceCacheLighting.cpp
+++ b/Engine/Source/Runtime/Renderer/Private/DistanceFieldSurfaceCacheLighting.cpp
@@ -5165,6 +5165,9 @@ void FDeferredShadingSceneRenderer::RenderMeshDistanceFieldVisualization(FRHICom
 				}
 
 				{
+					uint32 ClearValues[4] = { 0 };
+					RHICmdList.ClearUAV(GAOCulledObjectBuffers.Buffers.ObjectIndirectArguments.UAV, ClearValues);
+
 					TShaderMapRef<FCullObjectsForViewCS> ComputeShader(GetGlobalShaderMap(Scene->GetFeatureLevel()));
 					RHICmdList.SetComputeShader(ComputeShader->GetComputeShader());
 					ComputeShader->SetParameters(RHICmdList, Scene, View, Parameters);
