Fixed crash with visualize mesh distance fields enabled, when the DFAO show flag is disabled under Lighting Features, and moving a mesh

diff --git a/Engine/Source/Runtime/Renderer/Private/DistanceFieldSurfaceCacheLighting.cpp b/Engine/Source/Runtime/Renderer/Private/DistanceFieldSurfaceCacheLighting.cpp
--- a/Engine/Source/Runtime/Renderer/Private/DistanceFieldSurfaceCacheLighting.cpp
+++ b/Engine/Source/Runtime/Renderer/Private/DistanceFieldSurfaceCacheLighting.cpp
@@ -4320,8 +4320,7 @@ bool SupportsDistanceFieldAO(ERHIFeatureLevel::Type FeatureLevel, EShaderPlatfor
 bool FDeferredShadingSceneRenderer::ShouldPrepareForDistanceFieldAO() const
 {
 	return SupportsDistanceFieldAO(Scene->GetFeatureLevel(), Scene->GetShaderPlatform())
-		&& ViewFamily.EngineShowFlags.DistanceFieldAO
-		&& ((ShouldRenderDynamicSkyLight() && Scene->SkyLight->bCastShadows)
+		&& ((ShouldRenderDynamicSkyLight() && Scene->SkyLight->bCastShadows && ViewFamily.EngineShowFlags.DistanceFieldAO)
 			|| ViewFamily.EngineShowFlags.VisualizeMeshDistanceFields
 			|| ViewFamily.EngineShowFlags.VisualizeDistanceFieldAO
 			|| ViewFamily.EngineShowFlags.VisualizeDistanceFieldGI);
