UE4 - Fix crash on mobile OpenGL (UE-7327)

diff --git a/Engine/Source/Runtime/Renderer/Private/PostProcess/SceneRenderTargets.cpp b/Engine/Source/Runtime/Renderer/Private/PostProcess/SceneRenderTargets.cpp
--- a/Engine/Source/Runtime/Renderer/Private/PostProcess/SceneRenderTargets.cpp
+++ b/Engine/Source/Runtime/Renderer/Private/PostProcess/SceneRenderTargets.cpp
@@ -1217,6 +1217,12 @@ void FSceneRenderTargets::AllocateReflectionTargets()
 		if (!ReflectionBrightnessTarget)
 		{
 			bool bSupportsR32Float = CurrentFeatureLevel > ERHIFeatureLevel::ES2;
+			//@todo: FScene::UpdateSkyCaptureContents() is called before FSceneRenderTargets::AllocateRenderTargets()
+			// so CurrentFeatureLevel == ERHIFeatureLevel::Num and that crashes OpenGL ES2 as R32F is not valid
+			if (CurrentFeatureLevel == ERHIFeatureLevel::Num)
+			{
+				bSupportsR32Float = GMaxRHIFeatureLevel > ERHIFeatureLevel::ES2;
+			}
 			int32 ReflectionBrightnessIndex = bSupportsR32Float ? 0 : 1;
 			EPixelFormat BrightnessFormat = bSupportsR32Float ? PF_R32_FLOAT : PF_FloatRGBA;
 			FPooledRenderTargetDesc Desc3(FPooledRenderTargetDesc::Create2DDesc(FIntPoint(1, 1), BrightnessFormat, TexCreate_None, TexCreate_RenderTargetable, false));
