UE4 - Fix for missing forward rendering draw lists

diff --git a/Engine/Source/Runtime/Renderer/Private/ForwardBasePassRendering.cpp b/Engine/Source/Runtime/Renderer/Private/ForwardBasePassRendering.cpp
--- a/Engine/Source/Runtime/Renderer/Private/ForwardBasePassRendering.cpp
+++ b/Engine/Source/Runtime/Renderer/Private/ForwardBasePassRendering.cpp
@@ -333,14 +333,16 @@ void FForwardShadingSceneRenderer::RenderForwardShadingBasePass(FRHICommandListI
 
 		for (int32 DrawType = 0; DrawType < FScene::EBasePass_MAX; DrawType++)
 		{
+			Scene->BasePassForForwardShadingNoLightMapDrawList[DrawType].SortFrontToBack(Views[0].ViewLocation);
 			Scene->BasePassForForwardShadingLowQualityLightMapDrawList[DrawType].SortFrontToBack(Views[0].ViewLocation);
 			Scene->BasePassForForwardShadingDistanceFieldShadowMapLightMapDrawList[DrawType].SortFrontToBack(Views[0].ViewLocation);
 			Scene->BasePassForForwardShadingDirectionalLightAndSHIndirectDrawList[DrawType].SortFrontToBack(Views[0].ViewLocation);
-			Scene->BasePassForForwardShadingMovableDirectionalLightCSMLightmapDrawList[DrawType].SortFrontToBack(Views[0].ViewLocation);
-			Scene->BasePassForForwardShadingMovableDirectionalLightLightmapDrawList[DrawType].SortFrontToBack(Views[0].ViewLocation);
-			Scene->BasePassForForwardShadingMovableDirectionalLightCSMDrawList[DrawType].SortFrontToBack(Views[0].ViewLocation);
+			Scene->BasePassForForwardShadingDirectionalLightAndSHDirectionalIndirectDrawList[DrawType].SortFrontToBack(Views[0].ViewLocation);
+			Scene->BasePassForForwardShadingDirectionalLightAndSHDirectionalCSMIndirectDrawList[DrawType].SortFrontToBack(Views[0].ViewLocation);
 			Scene->BasePassForForwardShadingMovableDirectionalLightDrawList[DrawType].SortFrontToBack(Views[0].ViewLocation);
-			Scene->BasePassForForwardShadingNoLightMapDrawList[DrawType].SortFrontToBack(Views[0].ViewLocation);
+			Scene->BasePassForForwardShadingMovableDirectionalLightCSMDrawList[DrawType].SortFrontToBack(Views[0].ViewLocation);
+			Scene->BasePassForForwardShadingMovableDirectionalLightLightmapDrawList[DrawType].SortFrontToBack(Views[0].ViewLocation);
+			Scene->BasePassForForwardShadingMovableDirectionalLightCSMLightmapDrawList[DrawType].SortFrontToBack(Views[0].ViewLocation);
 		}
 	}
 
@@ -361,26 +363,30 @@ void FForwardShadingSceneRenderer::RenderForwardShadingBasePass(FRHICommandListI
 			if (SortMode == EBasePassSort::SortPerMesh)
 			{
 				SCOPE_CYCLE_COUNTER(STAT_StaticDrawListDrawTime);
+				MaxDraws -= Scene->BasePassForForwardShadingNoLightMapDrawList[FScene::EBasePass_Default].DrawVisibleFrontToBack(RHICmdList, View, View.StaticMeshVisibilityMap, View.StaticMeshBatchVisibility, MaxDraws);
 				MaxDraws -= Scene->BasePassForForwardShadingLowQualityLightMapDrawList[FScene::EBasePass_Default].DrawVisibleFrontToBack(RHICmdList, View, View.StaticMeshVisibilityMap, View.StaticMeshBatchVisibility, MaxDraws);
 				MaxDraws -= Scene->BasePassForForwardShadingDistanceFieldShadowMapLightMapDrawList[FScene::EBasePass_Default].DrawVisibleFrontToBack(RHICmdList, View, View.StaticMeshVisibilityMap, View.StaticMeshBatchVisibility, MaxDraws);
 				MaxDraws -= Scene->BasePassForForwardShadingDirectionalLightAndSHIndirectDrawList[FScene::EBasePass_Default].DrawVisibleFrontToBack(RHICmdList, View, View.StaticMeshVisibilityMap, View.StaticMeshBatchVisibility, MaxDraws);
-				MaxDraws -= Scene->BasePassForForwardShadingMovableDirectionalLightCSMLightmapDrawList[FScene::EBasePass_Default].DrawVisibleFrontToBack(RHICmdList, View, View.StaticMeshVisibilityMap, View.StaticMeshBatchVisibility, MaxDraws);
-				MaxDraws -= Scene->BasePassForForwardShadingMovableDirectionalLightLightmapDrawList[FScene::EBasePass_Default].DrawVisibleFrontToBack(RHICmdList, View, View.StaticMeshVisibilityMap, View.StaticMeshBatchVisibility, MaxDraws);
-				MaxDraws -= Scene->BasePassForForwardShadingMovableDirectionalLightCSMDrawList[FScene::EBasePass_Default].DrawVisibleFrontToBack(RHICmdList, View, View.StaticMeshVisibilityMap, View.StaticMeshBatchVisibility, MaxDraws);
+				MaxDraws -= Scene->BasePassForForwardShadingDirectionalLightAndSHDirectionalIndirectDrawList[FScene::EBasePass_Default].DrawVisibleFrontToBack(RHICmdList, View, View.StaticMeshVisibilityMap, View.StaticMeshBatchVisibility, MaxDraws);
+				MaxDraws -= Scene->BasePassForForwardShadingDirectionalLightAndSHDirectionalCSMIndirectDrawList[FScene::EBasePass_Default].DrawVisibleFrontToBack(RHICmdList, View, View.StaticMeshVisibilityMap, View.StaticMeshBatchVisibility, MaxDraws);
 				MaxDraws -= Scene->BasePassForForwardShadingMovableDirectionalLightDrawList[FScene::EBasePass_Default].DrawVisibleFrontToBack(RHICmdList, View, View.StaticMeshVisibilityMap, View.StaticMeshBatchVisibility, MaxDraws);
-				MaxDraws -= Scene->BasePassForForwardShadingNoLightMapDrawList[FScene::EBasePass_Default].DrawVisibleFrontToBack(RHICmdList, View, View.StaticMeshVisibilityMap, View.StaticMeshBatchVisibility, MaxDraws);
+				MaxDraws -= Scene->BasePassForForwardShadingMovableDirectionalLightCSMDrawList[FScene::EBasePass_Default].DrawVisibleFrontToBack(RHICmdList, View, View.StaticMeshVisibilityMap, View.StaticMeshBatchVisibility, MaxDraws);
+				MaxDraws -= Scene->BasePassForForwardShadingMovableDirectionalLightLightmapDrawList[FScene::EBasePass_Default].DrawVisibleFrontToBack(RHICmdList, View, View.StaticMeshVisibilityMap, View.StaticMeshBatchVisibility, MaxDraws);
+				MaxDraws -= Scene->BasePassForForwardShadingMovableDirectionalLightCSMLightmapDrawList[FScene::EBasePass_Default].DrawVisibleFrontToBack(RHICmdList, View, View.StaticMeshVisibilityMap, View.StaticMeshBatchVisibility, MaxDraws);
 			}
 			else
 			{
 				SCOPE_CYCLE_COUNTER(STAT_StaticDrawListDrawTime);
+				Scene->BasePassForForwardShadingNoLightMapDrawList[FScene::EBasePass_Default].DrawVisible(RHICmdList, View, View.StaticMeshVisibilityMap, View.StaticMeshBatchVisibility);
 				Scene->BasePassForForwardShadingLowQualityLightMapDrawList[FScene::EBasePass_Default].DrawVisible(RHICmdList, View, View.StaticMeshVisibilityMap, View.StaticMeshBatchVisibility);
 				Scene->BasePassForForwardShadingDistanceFieldShadowMapLightMapDrawList[FScene::EBasePass_Default].DrawVisible(RHICmdList, View, View.StaticMeshVisibilityMap, View.StaticMeshBatchVisibility);
 				Scene->BasePassForForwardShadingDirectionalLightAndSHIndirectDrawList[FScene::EBasePass_Default].DrawVisible(RHICmdList, View, View.StaticMeshVisibilityMap, View.StaticMeshBatchVisibility);
-				Scene->BasePassForForwardShadingMovableDirectionalLightCSMLightmapDrawList[FScene::EBasePass_Default].DrawVisible(RHICmdList, View, View.StaticMeshVisibilityMap, View.StaticMeshBatchVisibility);
-				Scene->BasePassForForwardShadingMovableDirectionalLightLightmapDrawList[FScene::EBasePass_Default].DrawVisible(RHICmdList, View, View.StaticMeshVisibilityMap, View.StaticMeshBatchVisibility);
-				Scene->BasePassForForwardShadingMovableDirectionalLightCSMDrawList[FScene::EBasePass_Default].DrawVisible(RHICmdList, View, View.StaticMeshVisibilityMap, View.StaticMeshBatchVisibility);
+				Scene->BasePassForForwardShadingDirectionalLightAndSHDirectionalIndirectDrawList[FScene::EBasePass_Default].DrawVisible(RHICmdList, View, View.StaticMeshVisibilityMap, View.StaticMeshBatchVisibility);
+				Scene->BasePassForForwardShadingDirectionalLightAndSHDirectionalCSMIndirectDrawList[FScene::EBasePass_Default].DrawVisible(RHICmdList, View, View.StaticMeshVisibilityMap, View.StaticMeshBatchVisibility);
 				Scene->BasePassForForwardShadingMovableDirectionalLightDrawList[FScene::EBasePass_Default].DrawVisible(RHICmdList, View, View.StaticMeshVisibilityMap, View.StaticMeshBatchVisibility);
-				Scene->BasePassForForwardShadingNoLightMapDrawList[FScene::EBasePass_Default].DrawVisible(RHICmdList, View, View.StaticMeshVisibilityMap, View.StaticMeshBatchVisibility);
+				Scene->BasePassForForwardShadingMovableDirectionalLightCSMDrawList[FScene::EBasePass_Default].DrawVisible(RHICmdList, View, View.StaticMeshVisibilityMap, View.StaticMeshBatchVisibility);
+				Scene->BasePassForForwardShadingMovableDirectionalLightLightmapDrawList[FScene::EBasePass_Default].DrawVisible(RHICmdList, View, View.StaticMeshVisibilityMap, View.StaticMeshBatchVisibility);
+				Scene->BasePassForForwardShadingMovableDirectionalLightCSMLightmapDrawList[FScene::EBasePass_Default].DrawVisible(RHICmdList, View, View.StaticMeshVisibilityMap, View.StaticMeshBatchVisibility);
 			}
 		}
 
@@ -426,10 +432,12 @@ void FForwardShadingSceneRenderer::RenderForwardShadingBasePass(FRHICommandListI
 			MaxDraws -= Scene->BasePassForForwardShadingLowQualityLightMapDrawList[FScene::EBasePass_Masked].DrawVisibleFrontToBack(RHICmdList, View,View.StaticMeshVisibilityMap,View.StaticMeshBatchVisibility,MaxDraws);
 			MaxDraws -= Scene->BasePassForForwardShadingDistanceFieldShadowMapLightMapDrawList[FScene::EBasePass_Masked].DrawVisibleFrontToBack(RHICmdList, View,View.StaticMeshVisibilityMap,View.StaticMeshBatchVisibility,MaxDraws);
 			MaxDraws -= Scene->BasePassForForwardShadingDirectionalLightAndSHIndirectDrawList[FScene::EBasePass_Masked].DrawVisibleFrontToBack(RHICmdList, View,View.StaticMeshVisibilityMap,View.StaticMeshBatchVisibility,MaxDraws);
+			MaxDraws -= Scene->BasePassForForwardShadingDirectionalLightAndSHDirectionalIndirectDrawList[FScene::EBasePass_Masked].DrawVisibleFrontToBack(RHICmdList, View,View.StaticMeshVisibilityMap,View.StaticMeshBatchVisibility,MaxDraws);
+			MaxDraws -= Scene->BasePassForForwardShadingDirectionalLightAndSHDirectionalCSMIndirectDrawList[FScene::EBasePass_Masked].DrawVisibleFrontToBack(RHICmdList, View,View.StaticMeshVisibilityMap,View.StaticMeshBatchVisibility,MaxDraws);
+			MaxDraws -= Scene->BasePassForForwardShadingMovableDirectionalLightDrawList[FScene::EBasePass_Masked].DrawVisibleFrontToBack(RHICmdList, View,View.StaticMeshVisibilityMap,View.StaticMeshBatchVisibility,MaxDraws);
+			MaxDraws -= Scene->BasePassForForwardShadingMovableDirectionalLightCSMDrawList[FScene::EBasePass_Masked].DrawVisibleFrontToBack(RHICmdList, View, View.StaticMeshVisibilityMap, View.StaticMeshBatchVisibility, MaxDraws);
 			MaxDraws -= Scene->BasePassForForwardShadingMovableDirectionalLightLightmapDrawList[FScene::EBasePass_Masked].DrawVisibleFrontToBack(RHICmdList, View, View.StaticMeshVisibilityMap, View.StaticMeshBatchVisibility, MaxDraws);
 			MaxDraws -= Scene->BasePassForForwardShadingMovableDirectionalLightCSMLightmapDrawList[FScene::EBasePass_Masked].DrawVisibleFrontToBack(RHICmdList, View, View.StaticMeshVisibilityMap, View.StaticMeshBatchVisibility, MaxDraws);
-			MaxDraws -= Scene->BasePassForForwardShadingMovableDirectionalLightCSMDrawList[FScene::EBasePass_Masked].DrawVisibleFrontToBack(RHICmdList, View, View.StaticMeshVisibilityMap, View.StaticMeshBatchVisibility, MaxDraws);
-			MaxDraws -= Scene->BasePassForForwardShadingMovableDirectionalLightDrawList[FScene::EBasePass_Masked].DrawVisibleFrontToBack(RHICmdList, View,View.StaticMeshVisibilityMap,View.StaticMeshBatchVisibility,MaxDraws);
 		}
 		else
 		{
@@ -438,10 +446,12 @@ void FForwardShadingSceneRenderer::RenderForwardShadingBasePass(FRHICommandListI
 			Scene->BasePassForForwardShadingLowQualityLightMapDrawList[FScene::EBasePass_Masked].DrawVisible(RHICmdList, View,View.StaticMeshVisibilityMap,View.StaticMeshBatchVisibility);
 			Scene->BasePassForForwardShadingDistanceFieldShadowMapLightMapDrawList[FScene::EBasePass_Masked].DrawVisible(RHICmdList, View,View.StaticMeshVisibilityMap,View.StaticMeshBatchVisibility);
 			Scene->BasePassForForwardShadingDirectionalLightAndSHIndirectDrawList[FScene::EBasePass_Masked].DrawVisible(RHICmdList, View,View.StaticMeshVisibilityMap,View.StaticMeshBatchVisibility);
-			Scene->BasePassForForwardShadingMovableDirectionalLightCSMLightmapDrawList[FScene::EBasePass_Masked].DrawVisible(RHICmdList, View, View.StaticMeshVisibilityMap, View.StaticMeshBatchVisibility);
-			Scene->BasePassForForwardShadingMovableDirectionalLightLightmapDrawList[FScene::EBasePass_Masked].DrawVisible(RHICmdList, View, View.StaticMeshVisibilityMap, View.StaticMeshBatchVisibility);
-			Scene->BasePassForForwardShadingMovableDirectionalLightCSMDrawList[FScene::EBasePass_Masked].DrawVisible(RHICmdList, View, View.StaticMeshVisibilityMap, View.StaticMeshBatchVisibility);
+			Scene->BasePassForForwardShadingDirectionalLightAndSHDirectionalIndirectDrawList[FScene::EBasePass_Masked].DrawVisible(RHICmdList, View,View.StaticMeshVisibilityMap,View.StaticMeshBatchVisibility);
+			Scene->BasePassForForwardShadingDirectionalLightAndSHDirectionalCSMIndirectDrawList[FScene::EBasePass_Masked].DrawVisible(RHICmdList, View,View.StaticMeshVisibilityMap,View.StaticMeshBatchVisibility);
 			Scene->BasePassForForwardShadingMovableDirectionalLightDrawList[FScene::EBasePass_Masked].DrawVisible(RHICmdList, View,View.StaticMeshVisibilityMap,View.StaticMeshBatchVisibility);
+			Scene->BasePassForForwardShadingMovableDirectionalLightCSMDrawList[FScene::EBasePass_Masked].DrawVisible(RHICmdList, View, View.StaticMeshVisibilityMap, View.StaticMeshBatchVisibility);
+			Scene->BasePassForForwardShadingMovableDirectionalLightLightmapDrawList[FScene::EBasePass_Masked].DrawVisible(RHICmdList, View, View.StaticMeshVisibilityMap, View.StaticMeshBatchVisibility);
+			Scene->BasePassForForwardShadingMovableDirectionalLightCSMLightmapDrawList[FScene::EBasePass_Masked].DrawVisible(RHICmdList, View, View.StaticMeshVisibilityMap, View.StaticMeshBatchVisibility);
 		}
 	}
 }
diff --git a/Engine/Source/Runtime/Renderer/Private/RendererScene.cpp b/Engine/Source/Runtime/Renderer/Private/RendererScene.cpp
--- a/Engine/Source/Runtime/Renderer/Private/RendererScene.cpp
+++ b/Engine/Source/Runtime/Renderer/Private/RendererScene.cpp
@@ -1761,6 +1761,8 @@ void FScene::UpdateStaticDrawListsForMaterials_RenderThread(FRHICommandListImmed
 		BasePassForForwardShadingLowQualityLightMapDrawList[DrawType].GetUsedPrimitivesBasedOnMaterials(FeatureLevel, Materials, PrimitivesToUpdate);
 		BasePassForForwardShadingDistanceFieldShadowMapLightMapDrawList[DrawType].GetUsedPrimitivesBasedOnMaterials(FeatureLevel, Materials, PrimitivesToUpdate);
 		BasePassForForwardShadingDirectionalLightAndSHIndirectDrawList[DrawType].GetUsedPrimitivesBasedOnMaterials(FeatureLevel, Materials, PrimitivesToUpdate);
+		BasePassForForwardShadingDirectionalLightAndSHDirectionalIndirectDrawList[DrawType].GetUsedPrimitivesBasedOnMaterials(FeatureLevel, Materials, PrimitivesToUpdate);
+		BasePassForForwardShadingDirectionalLightAndSHDirectionalCSMIndirectDrawList[DrawType].GetUsedPrimitivesBasedOnMaterials(FeatureLevel, Materials, PrimitivesToUpdate);
 		BasePassForForwardShadingMovableDirectionalLightDrawList[DrawType].GetUsedPrimitivesBasedOnMaterials(FeatureLevel, Materials, PrimitivesToUpdate);
 		BasePassForForwardShadingMovableDirectionalLightCSMDrawList[DrawType].GetUsedPrimitivesBasedOnMaterials(FeatureLevel, Materials, PrimitivesToUpdate);
 		BasePassForForwardShadingMovableDirectionalLightLightmapDrawList[DrawType].GetUsedPrimitivesBasedOnMaterials(FeatureLevel, Materials, PrimitivesToUpdate);
@@ -1976,6 +1978,10 @@ void FScene::DumpStaticMeshDrawListStats() const
 	DUMP_DRAW_LIST(BasePassForForwardShadingDistanceFieldShadowMapLightMapDrawList[EBasePass_Masked]);
 	DUMP_DRAW_LIST(BasePassForForwardShadingDirectionalLightAndSHIndirectDrawList[EBasePass_Default]);
 	DUMP_DRAW_LIST(BasePassForForwardShadingDirectionalLightAndSHIndirectDrawList[EBasePass_Masked]);
+	DUMP_DRAW_LIST(BasePassForForwardShadingDirectionalLightAndSHDirectionalIndirectDrawList[EBasePass_Default]);
+	DUMP_DRAW_LIST(BasePassForForwardShadingDirectionalLightAndSHDirectionalIndirectDrawList[EBasePass_Masked]);
+	DUMP_DRAW_LIST(BasePassForForwardShadingDirectionalLightAndSHDirectionalCSMIndirectDrawList[EBasePass_Default]);
+	DUMP_DRAW_LIST(BasePassForForwardShadingDirectionalLightAndSHDirectionalCSMIndirectDrawList[EBasePass_Masked]);
 	DUMP_DRAW_LIST(BasePassForForwardShadingMovableDirectionalLightDrawList[EBasePass_Default]);
 	DUMP_DRAW_LIST(BasePassForForwardShadingMovableDirectionalLightDrawList[EBasePass_Masked]);
 	DUMP_DRAW_LIST(BasePassForForwardShadingMovableDirectionalLightCSMDrawList[EBasePass_Default]);
@@ -2144,7 +2150,10 @@ void FScene::ApplyWorldOffset_RenderThread(FVector InOffset)
 	StaticMeshDrawListApplyWorldOffset(WholeSceneShadowDepthDrawList, InOffset);
 	StaticMeshDrawListApplyWorldOffset(BasePassForForwardShadingNoLightMapDrawList, InOffset);
 	StaticMeshDrawListApplyWorldOffset(BasePassForForwardShadingLowQualityLightMapDrawList, InOffset);
+	StaticMeshDrawListApplyWorldOffset(BasePassForForwardShadingDistanceFieldShadowMapLightMapDrawList, InOffset);
 	StaticMeshDrawListApplyWorldOffset(BasePassForForwardShadingDirectionalLightAndSHIndirectDrawList, InOffset);
+	StaticMeshDrawListApplyWorldOffset(BasePassForForwardShadingDirectionalLightAndSHDirectionalIndirectDrawList, InOffset);
+	StaticMeshDrawListApplyWorldOffset(BasePassForForwardShadingDirectionalLightAndSHDirectionalCSMIndirectDrawList, InOffset);
 	StaticMeshDrawListApplyWorldOffset(BasePassForForwardShadingMovableDirectionalLightDrawList, InOffset);
 	StaticMeshDrawListApplyWorldOffset(BasePassForForwardShadingMovableDirectionalLightCSMDrawList, InOffset);
 	StaticMeshDrawListApplyWorldOffset(BasePassForForwardShadingMovableDirectionalLightLightmapDrawList, InOffset);
