UE4 - Fix for forward renderer ignoring translucent material's disable depth test enabled (fixes UE-7330)

diff --git a/Engine/Source/Runtime/Renderer/Private/ForwardTranslucentRendering.cpp b/Engine/Source/Runtime/Renderer/Private/ForwardTranslucentRendering.cpp
--- a/Engine/Source/Runtime/Renderer/Private/ForwardTranslucentRendering.cpp
+++ b/Engine/Source/Runtime/Renderer/Private/ForwardTranslucentRendering.cpp
@@ -188,6 +188,12 @@ bool FTranslucencyForwardShadingDrawingPolicyFactory::DrawDynamicMesh(
 	// Only render translucent materials.
 	if (IsTranslucentBlendMode(BlendMode))
 	{
+		const bool bDisableDepthTest = Material->ShouldDisableDepthTest();
+		if (bDisableDepthTest)
+		{
+			RHICmdList.SetDepthStencilState(TStaticDepthStencilState<false, CF_Always>::GetRHI());
+		}
+
 		ProcessBasePassMeshForForwardShading(
 			RHICmdList,
 			FProcessBasePassMeshParameters(
@@ -206,6 +212,13 @@ bool FTranslucencyForwardShadingDrawingPolicyFactory::DrawDynamicMesh(
 				)
 			);
 
+		if (bDisableDepthTest)
+		{
+			// Restore default depth state
+			// Note, this is a reversed Z depth surface, using CF_GreaterEqual.	
+			RHICmdList.SetDepthStencilState(TStaticDepthStencilState<false, CF_GreaterEqual>::GetRHI());
+		}
+
 		bDirty = true;
 	}
 	return bDirty;
