Faster 3D texture clear for translucent lighting on Xbox One.

diff --git a/Engine/Source/Runtime/Renderer/Private/TranslucentLighting.cpp b/Engine/Source/Runtime/Renderer/Private/TranslucentLighting.cpp
--- a/Engine/Source/Runtime/Renderer/Private/TranslucentLighting.cpp
+++ b/Engine/Source/Runtime/Renderer/Private/TranslucentLighting.cpp
@@ -1052,8 +1052,15 @@ void ClearVolumeTextures(FRHICommandList& RHICmdList, ERHIFeatureLevel::Type Fea
 	static FGlobalBoundShaderState VolumeClearBoundShaderState;
 	
 
+
+
+#ifdef PLATFORM_XBOXONE
+	// ClearMRT is faster on Xbox
+	if (true)
+#else
 	// Currently using a manual clear, which is ~10x faster than a hardware clear of the volume textures on AMD PC GPU's
 	if (false)
+#endif
 	{
 		RHICmdList.ClearMRT(true, NumRenderTargets, ClearColors, false, 0, false, 0, FIntRect());
 	}
