SSR on translucency.

HZB is now always built and stored on ViewInfo.
SSR uses constant depth intersect instead of segment intersect for better quality near contacts.

diff --git a/Engine/Shaders/BasePassPixelShader.usf b/Engine/Shaders/BasePassPixelShader.usf
--- a/Engine/Shaders/BasePassPixelShader.usf
+++ b/Engine/Shaders/BasePassPixelShader.usf
@@ -14,9 +14,11 @@
 #include "BRDF.usf"
 #include "Random.usf"
 #include "LightAccumulator.usf"
-
 #include "DeferredShadingCommon.usf"
 
+#define PREV_FRAME_COLOR	1
+#include "ScreenSpaceRayCast.usf"
+
 #if TRANSLUCENT_SELF_SHADOWING
 
 	#include "ShadowProjectionCommon.usf"
@@ -37,6 +39,11 @@ SamplerState TranslucencyLightingVolumeDirectionalInnerSampler;
 Texture3D TranslucencyLightingVolumeDirectionalOuter;
 SamplerState TranslucencyLightingVolumeDirectionalOuterSampler;
 
+Texture2D		HZBTexture;
+SamplerState	HZBSampler;
+Texture2D		PrevSceneColor;
+SamplerState	PrevSceneColorSampler;
+
 // 2nd order SH indirect lighting interpolated from the indirect lighting cache for the whole object
 // Used to apply precomputed lighting to translucency
 float4 IndirectLightingSHCoefficients[3];
@@ -62,6 +69,15 @@ float4 PointSkyBentNormal;
 	SamplerState ReflectionCubemapSampler;
 #endif
 
+uint MortonCode( uint x )
+{
+	//x = (x ^ (x <<  8)) & 0x00ff00ff;
+	//x = (x ^ (x <<  4)) & 0x0f0f0f0f;
+	x = (x ^ (x <<  2)) & 0x33333333;
+	x = (x ^ (x <<  1)) & 0x55555555;
+	return x;
+}
+
 half3 GetImageBasedReflectionLighting(FMaterialPixelParameters MaterialParameters, half Roughness, half3 SpecularColor, half IndirectIrradiance)
 {
 	half AbsoluteSpecularMip = ComputeReflectionCaptureMipFromRoughness(Roughness);
@@ -97,7 +113,32 @@ half3 GetImageBasedReflectionLighting(FMaterialPixelParameters MaterialParameter
 		SpecularIBL.rgb += (1 - SpecularIBL.a) * SkyLighting;
 	}
 #endif
-	
+
+#if MATERIAL_SSR
+	//uint FrameRandom = (uint)(View.TemporalAAParams.r * 1551);
+	uint FrameRandom = (View.FrameNumber & 15) * 1551;
+
+	uint Morton = MortonCode( (uint)MaterialParameters.SVPosition.x & 3 ) | ( MortonCode( (uint)MaterialParameters.SVPosition.y & 3 ) * 2 );
+	uint PixelIndex = reversebits( Morton ) >> 28;
+	//uint PixelIndex = ( (uint)MaterialParameters.SVPosition.x & 3 ) | ( ( (uint)MaterialParameters.SVPosition.y & 3 ) * 2 );
+	//PixelIndex = ( PixelIndex * 1551 ) & 15;
+
+	uint Offset = ( PixelIndex + FrameRandom ) & 15;
+	float StepOffset = Offset / 15.0;
+	StepOffset -= 0.5;
+		
+	float4 HitUVzTime = RayCastCheap( HZBTexture, HZBSampler, MaterialParameters.ReflectionVector, MaterialParameters.ScreenPosition.w, MaterialParameters.WorldPosition_CamRelative, StepOffset );
+	//float4 HitUVzTime = RayCast( HZBTexture, HZBSampler, MaterialParameters.ReflectionVector, 0, MaterialParameters.ScreenPosition.w, MaterialParameters.WorldPosition_CamRelative, 12, StepOffset );
+
+	// if there was a hit
+	BRANCH if( HitUVzTime.w < 1 )
+	{
+		float4 SSR = SampleScreenColor( PrevSceneColor, PrevSceneColorSampler, HitUVzTime.xyz );
+		SSR *= saturate( 2 - 6.6 * Roughness );
+		SpecularIBL.rgb = SpecularIBL.rgb * (1 - SSR.a) + SSR.rgb;
+	}
+#endif
+
 	half NoV = saturate(dot(MaterialParameters.WorldNormal, MaterialParameters.CameraVector));
 	SpecularColor = EnvBRDFApprox(SpecularColor, Roughness, NoV);
 
@@ -716,7 +757,7 @@ void Main(
 		OutColor = half4(Color * VertexFog.a + VertexFog.rgb, Opacity);
 		OutColor = RETURN_COLOR(OutColor);
 	#elif MATERIALBLENDING_ADDITIVE
-		OutColor = half4(Color * VertexFog.aaa * Opacity.xxx, 0.0f);
+		OutColor = half4(Color * VertexFog.a * Opacity, 0.0f);
 		OutColor = RETURN_COLOR(OutColor);
 	#elif MATERIALBLENDING_MODULATE
 		// RETURN_COLOR not needed with modulative blending
diff --git a/Engine/Shaders/ScreenSpaceReflections.usf b/Engine/Shaders/ScreenSpaceReflections.usf
--- a/Engine/Shaders/ScreenSpaceReflections.usf
+++ b/Engine/Shaders/ScreenSpaceReflections.usf
@@ -10,9 +10,7 @@
 #include "Random.usf"
 #include "BRDF.usf"
 #include "MonteCarlo.usf"
-
-// 0:reference (slower, potentially higher quality) 1:use mips of the HZB depth (better performance)
-#define USE_HZB		1
+#include "ScreenSpaceRayCast.usf"
 
 #define SCALAR_BRANCHLESS		0
 #define VECTORIZED_BRANCHLESS	0
@@ -40,254 +38,6 @@ uint ReverseUIntBits( uint bits )
 	return bits;
 }
 
-
-float4 SampleDepthTexture( Texture2D Texture, SamplerState Sampler, float Level, float4 SampleUV0, float4 SampleUV1 )
-{
-	float4 SampleDepth;
-#if USE_HZB
-	SampleDepth.x = Texture2DSampleLevel( Texture, Sampler, SampleUV0.xy, Level ).r;
-	SampleDepth.y = Texture2DSampleLevel( Texture, Sampler, SampleUV0.zw, Level ).r;
-	SampleDepth.z = Texture2DSampleLevel( Texture, Sampler, SampleUV1.xy, Level ).r;
-	SampleDepth.w = Texture2DSampleLevel( Texture, Sampler, SampleUV1.zw, Level ).r;
-#else
-	SampleDepth.x = Texture2DSampleLevel( SceneDepthTexture, SceneDepthTextureSampler, SampleUV0.xy, 0 ).r;
-	SampleDepth.y = Texture2DSampleLevel( SceneDepthTexture, SceneDepthTextureSampler, SampleUV0.zw, 0 ).r;
-	SampleDepth.z = Texture2DSampleLevel( SceneDepthTexture, SceneDepthTextureSampler, SampleUV1.xy, 0 ).r;
-	SampleDepth.w = Texture2DSampleLevel( SceneDepthTexture, SceneDepthTextureSampler, SampleUV1.zw, 0 ).r;
-#endif
-	return SampleDepth;
-}
-
-
-float4 RayCast( float3 R, float Roughness, float SceneDepth, float3 PositionTranslatedWorld, int NumSteps, float StepOffset )
-{
-	// TODO provide RayStartUVz
-
-	// NOTE could clip ray against frustum planes
-
-	// TODO use screen position and skip matrix mul
-	float4 RayStartClip	= mul( float4( PositionTranslatedWorld, 1 ), View.TranslatedWorldToClip );
-	float4 RayEndClip	= mul( float4( PositionTranslatedWorld + R * SceneDepth, 1 ), View.TranslatedWorldToClip );
-
-	float3 RayStartScreen = RayStartClip.xyz / RayStartClip.w;
-	float3 RayEndScreen = RayEndClip.xyz / RayEndClip.w;
-
-	// Normalize 2D length
-	float3 RayStepScreen = ( RayEndScreen - RayStartScreen ) / length( RayEndScreen.xy - RayStartScreen.xy );
-	RayStepScreen *= 1.5;
-
-#if USE_HZB
-	float3 RayStartUVz = float3( RayStartScreen.xy * float2( 0.5, -0.5 ) + 0.5, RayStartScreen.z );
-	float3 RayStepUVz = float3( RayStepScreen.xy * float2( 0.5, -0.5 ), RayStepScreen.z );
-#else
-	float3 RayStartUVz = float3( RayStartScreen.xy * View.ScreenPositionScaleBias.xy + View.ScreenPositionScaleBias.wz, RayStartScreen.z );
-	float3 RayStepUVz = float3( RayStepScreen.xy * View.ScreenPositionScaleBias.xy, RayStepScreen.z );
-#endif
-	
-	const float Step = 1.0 / (NumSteps + 1);
-
-	// *2 to get less morie pattern in extreme cases, larger values make object appear not grounded in reflections
-	const float CompareTolerance = abs( RayStepScreen.z ) * Step * 2;
-
-	// avoid bugs with early returns inside of loops on certain platform compilers.
-	float4 Result = float4( 0, 0, 0, 1 );
-
-#if SCALAR_BRANCHLESS
-
-	float MinHitTime = 1;
-	float LastDiff = 0;
-
-	float SampleTime = StepOffset * Step + Step;
-
-	UNROLL
-	for( int i = 0; i < NumSteps; i++ )
-	{
-		float3 SampleUVz = RayStartUVz + RayStepUVz * SampleTime;
-		
-		// Use lower res for farther samples
-		float Level = Roughness * (i * 4.0 / NumSteps);
-		float SampleDepth = PostprocessInput1.SampleLevel( PostprocessInput1Sampler, SampleUVz.xy, Level ).r;
-
-		float DepthDiff = SampleUVz.z - SampleDepth;
-		bool Hit = abs( DepthDiff + CompareTolerance ) < CompareTolerance;
-
-		// Find more accurate hit using line segment intersection
-		float TimeLerp = saturate( LastDiff / (LastDiff - DepthDiff) );
-		float IntersectTime = SampleTime + TimeLerp * Step - Step;
-		float HitTime = Hit ? IntersectTime : 1;
-		MinHitTime = min( MinHitTime, HitTime );
-
-		LastDiff = DepthDiff;
-
-		SampleTime += Step;
-	}
-
-	float3 HitUVz = RayStartUVz + RayStepUVz * MinHitTime;
-#if USE_HZB
-	HitUVz.xy = HitUVz.xy * float2( 2, -2 ) + float2( -1, 1 );
-	HitUVz.xy = HitUVz.xy * View.ScreenPositionScaleBias.xy + View.ScreenPositionScaleBias.wz;
-#endif
-	Result = float4( HitUVz, MinHitTime );
-
-#elif VECTORIZED_BRANCHLESS
-
-	float MinHitTime = 1;
-	float LastDiff = 0;
-	float Level = 0;
-	
-	// Vectorized to group fetches
-	float4 SampleTime = ( StepOffset + float4( 1, 2, 3, 4 ) ) / (NumSteps + 1);
-	float4 SampleUV0 = RayStartUVz.xyxy + RayStepUVz.xyxy * SampleTime.xxyy;
-	float4 SampleUV1 = RayStartUVz.xyxy + RayStepUVz.xyxy * SampleTime.zzww;
-	float4 SampleZ   = RayStartUVz.zzzz + RayStepUVz.zzzz * SampleTime;
-
-	LOOP
-	for( int i = 0; i < NumSteps; i += 4 )
-	{
-		// Use lower res for farther samples
-		float4 SampleDepth = SampleDepthTexture( PostprocessInput1, PostprocessInput1Sampler, Level, SampleUV0, SampleUV1 );
-
-		#if 1
-			float4 DepthDiff = SampleZ - SampleDepth;
-			bool4 Hit = abs( DepthDiff + CompareTolerance ) < CompareTolerance;
-	
-			// If hit set to intersect time. If missed set to 1, beyond end of ray
-			float4 HitTime = Hit ? SampleTime : 1;
-
-			// Take closest hit
-			HitTime.xy = min( HitTime.xy, HitTime.zw );
-			MinHitTime = min( MinHitTime, min( HitTime.x, HitTime.y ) );
-		#else
-			// Line segment intersection
-			float4 DepthDiff1 = SampleZ - SampleDepth;
-			float4 DepthDiff0 = float4( LastDiff, DepthDiff1.xyz );
-			float4 TimeLerp = saturate( DepthDiff0 / (DepthDiff0 - DepthDiff1) );
-			float4 IntersectTime = SampleTime + (TimeLerp - 1) / (NumSteps + 1);
-		
-			bool4 Hit = abs( DepthDiff1 + CompareTolerance ) < CompareTolerance;
-	
-			// If hit set to intersect time. If missed set to 1, beyond end of ray
-			float4 HitTime = Hit ? IntersectTime : 1;
-
-			// Take closest hit
-			HitTime.xy = min( HitTime.xy, HitTime.zw );
-			MinHitTime = min( MinHitTime, min( HitTime.x, HitTime.y ) );
-
-			LastDiff = DepthDiff1.w;
-		#endif
-		Level += Roughness * (16.0 / NumSteps);
-
-		SampleTime	+= 4.0 / (NumSteps + 1);
-		SampleUV0	+= RayStepUVz.xyxy * 4.0 / (NumSteps + 1);
-		SampleUV1	+= RayStepUVz.xyxy * 4.0 / (NumSteps + 1);
-		SampleZ		+= RayStepUVz.zzzz * 4.0 / (NumSteps + 1);
-	}
-
-	float3 HitUVz = RayStartUVz + RayStepUVz * MinHitTime;
-#if USE_HZB
-	HitUVz.xy = HitUVz.xy = HitUVz.xy * float2( 2, -2 ) + float2( -1, 1 );
-	HitUVz.xy = HitUVz.xy * View.ScreenPositionScaleBias.xy + View.ScreenPositionScaleBias.wz;
-#endif
-	Result = float4( HitUVz, MinHitTime );
-
-#else // VECTORIZED_EARLY_OUT
-
-	float LastDiff = 0;
-	float Level = 0;
-
-	float4 SampleTime = ( StepOffset + float4( 1, 2, 3, 4 ) ) / (NumSteps + 1);	
-
-	LOOP
-	for( int i = 0; i < NumSteps; i += 4 )
-	{
-		// Vectorized to group fetches
-		float4 SampleUV0 = RayStartUVz.xyxy + RayStepUVz.xyxy * SampleTime.xxyy;
-		float4 SampleUV1 = RayStartUVz.xyxy + RayStepUVz.xyxy * SampleTime.zzww;
-		float4 SampleZ   = RayStartUVz.zzzz + RayStepUVz.zzzz * SampleTime;
-		
-		// Use lower res for farther samples
-		float4 SampleDepth = SampleDepthTexture( PostprocessInput1, PostprocessInput1Sampler, Level, SampleUV0, SampleUV1 );
-
-		float4 DepthDiff1 = SampleZ - SampleDepth;
-		bool4 Hit = abs( -DepthDiff1 - CompareTolerance ) < CompareTolerance;
-
-		BRANCH if( any( Hit ) )
-		{
-			// Find more accurate hit using line segment intersection
-			float4 DepthDiff0 = float4( LastDiff, DepthDiff1.xyz );
-			float4 TimeLerp = saturate( DepthDiff0 / (DepthDiff0 - DepthDiff1) );
-			float4 IntersectTime = SampleTime + (TimeLerp - 1) / (NumSteps + 1);
-			float4 HitTime = Hit ? IntersectTime : 1;
-
-			// Take closest hit
-			HitTime.xy = min( HitTime.xy, HitTime.zw );
-			float MinHitTime = min( HitTime.x, HitTime.y );
-			
-			float3 HitUVz = RayStartUVz + RayStepUVz * MinHitTime;
-#if USE_HZB
-			HitUVz.xy = HitUVz.xy * float2( 2, -2 ) + float2( -1, 1 );
-			HitUVz.xy = HitUVz.xy * View.ScreenPositionScaleBias.xy + View.ScreenPositionScaleBias.wz;
-#endif
-			Result = float4( HitUVz, MinHitTime );
-			break;
-		}
-
-		LastDiff = DepthDiff1.w;
-		Level += Roughness * (16.0 / NumSteps);
-
-		SampleTime += 4.0 / (NumSteps + 1);
-	}
-
-#endif
-
-	return Result;
-}
-
-float4 SampleScreenColor( float3 HitUVz, float HitTime )
-{
-	float4 OutColor;
-
-#if PREV_FRAME_COLOR
-	// Find previous screen position for hit since color buffer is from last frame
-	// TODO combine to single matrix
-	float4 HitClip = float4( (HitUVz.xy - View.ScreenPositionScaleBias.wz) / View.ScreenPositionScaleBias.xy, HitUVz.z, 1 );
-	float4 HitTranslatedWorld = mul( HitClip, View.ClipToTranslatedWorld );
-	HitTranslatedWorld /= HitTranslatedWorld.w;
-		
-	float3 PrevTranslatedWorld = HitTranslatedWorld.xyz + (View.PrevPreViewTranslation - View.PreViewTranslation);
-	float4 PrevClip = mul( float4( PrevTranslatedWorld, 1 ), View.PrevTranslatedWorldToClip );
-	float2 PrevScreen = PrevClip.xy / PrevClip.w;
-	float2 PrevUV = PrevScreen.xy * View.ScreenPositionScaleBias.xy + View.ScreenPositionScaleBias.wz;
-
-	OutColor.rgb = PostprocessInput0.SampleLevel( PostprocessInput0Sampler, PrevUV, 0 ).rgb;
-	OutColor.a = 1;
-
-	// Off screen masking		
-	float2 Vignette = saturate( abs( PrevScreen ) * 5 - 4 );
-#else
-	OutColor.rgb = PostprocessInput0.SampleLevel( PostprocessInput0Sampler, HitUVz.xy, 0 ).rgb;
-	OutColor.a = 1;
-
-	// Off screen masking
-	float2 HitScreenPos = ( HitUVz.xy - View.ScreenPositionScaleBias.wz ) / View.ScreenPositionScaleBias.xy;
-	float2 Vignette = saturate( abs( HitScreenPos ) * 5 - 4 );
-#endif
-		
-	//PrevScreen sometimes has NaNs or Infs.  DX11 is protected because saturate turns NaNs -> 0.
-	//Do a SafeSaturate so other platforms get the same protection.
-	OutColor *= SafeSaturate( 1.0 - dot( Vignette, Vignette ) );
-
-	// Transform NaNs to black, transform negative colors to black.
-	OutColor.rgb = -min(-OutColor.rgb, 0.0);
-
-	// Fade end of trace
-	OutColor *= saturate( 4 - 4 * HitTime );
-
-	return OutColor;
-}
-
-
-
 void ScreenSpaceReflectionsPS(in float4 UVAndScreenPos : TEXCOORD0, out float4 OutColor : SV_Target0)
 {
 	float2 UV = UVAndScreenPos.xy;
@@ -390,12 +140,16 @@ void ScreenSpaceReflectionsPS(in float4 UVAndScreenPos : TEXCOORD0, out float4 O
 			float3 H = TangentToWorld( ImportanceSampleBlinn( E, Roughness ).xyz, N );
 			float3 R = 2 * dot( V, H ) * H - V;
 
-			float4 HitUVzTime = RayCast( R, Roughness, SceneDepth, PositionTranslatedWorld, NumSteps, StepOffset );
+			float4 HitUVzTime = RayCast( PostprocessInput1, PostprocessInput1Sampler, R, Roughness, SceneDepth, PositionTranslatedWorld, NumSteps, StepOffset );
 
 			// if there was a hit
 			BRANCH if( HitUVzTime.w < 1 )
 			{
-				float4 SampleColor = SampleScreenColor( HitUVzTime.xyz, HitUVzTime.w );
+				float4 SampleColor = SampleScreenColor( PostprocessInput0, PostprocessInput0Sampler, HitUVzTime.xyz );
+				
+				// Fade end of trace
+				SampleColor *= saturate( 4 - 4 * HitUVzTime.w );
+				
 				SampleColor.rgb /= 1 + Luminance(SampleColor.rgb);
 				OutColor += SampleColor;
 			}
@@ -412,12 +166,14 @@ void ScreenSpaceReflectionsPS(in float4 UVAndScreenPos : TEXCOORD0, out float4 O
 		
 		float3 R = reflect( -V, N );
 
-		float4 HitUVzTime = RayCast( R, Roughness, SceneDepth, PositionTranslatedWorld, NumSteps, StepOffset );
-
+		float4 HitUVzTime = RayCast( PostprocessInput1, PostprocessInput1Sampler, R, Roughness, SceneDepth, PositionTranslatedWorld, NumSteps, StepOffset );
 		// if there was a hit
 		BRANCH if( HitUVzTime.w < 1 )
 		{
-			OutColor = SampleScreenColor( HitUVzTime.xyz, HitUVzTime.w );
+			OutColor = SampleScreenColor( PostprocessInput0, PostprocessInput0Sampler, HitUVzTime.xyz );
+
+			// Fade end of trace
+			OutColor *= saturate( 4 - 4 * HitUVzTime.w );
 		}
 	}
 	
diff --git a/Engine/Source/Runtime/Engine/Classes/Materials/Material.h b/Engine/Source/Runtime/Engine/Classes/Materials/Material.h
--- a/Engine/Source/Runtime/Engine/Classes/Materials/Material.h
+++ b/Engine/Source/Runtime/Engine/Classes/Materials/Material.h
@@ -394,6 +394,10 @@ public:
 	UPROPERTY(EditAnywhere, Category=Translucency, meta=(DisplayName = "Responsive AA"))
 	uint32 bEnableResponsiveAA:1;
 
+	/** SSR on translucency */
+	UPROPERTY(EditAnywhere, Category=Translucency, meta=(DisplayName = "Screen Space Reflections"))
+	uint32 bScreenSpaceReflections:1;
+
 	/** Indicates that the material should be rendered without backface culling and the normal should be flipped for backfaces. */
 	UPROPERTY(EditAnywhere, Category=Material)
 	uint32 TwoSided:1;
diff --git a/Engine/Source/Runtime/Engine/Private/Materials/Material.cpp b/Engine/Source/Runtime/Engine/Private/Materials/Material.cpp
--- a/Engine/Source/Runtime/Engine/Private/Materials/Material.cpp
+++ b/Engine/Source/Runtime/Engine/Private/Materials/Material.cpp
@@ -2217,6 +2217,7 @@ bool UMaterial::CanEditChange(const UProperty* InProperty) const
 
 		if (PropertyName == GET_MEMBER_NAME_STRING_CHECKED(UMaterial, bEnableSeparateTranslucency)
 			|| PropertyName == GET_MEMBER_NAME_STRING_CHECKED(UMaterial, bEnableResponsiveAA)
+			|| PropertyName == GET_MEMBER_NAME_STRING_CHECKED(UMaterial, bScreenSpaceReflections)
 			|| PropertyName == GET_MEMBER_NAME_STRING_CHECKED(UMaterial, bDisableDepthTest)
 			|| PropertyName == GET_MEMBER_NAME_STRING_CHECKED(UMaterial, bUseTranslucencyVertexFog))
 		{
diff --git a/Engine/Source/Runtime/Engine/Private/Materials/MaterialShared.cpp b/Engine/Source/Runtime/Engine/Private/Materials/MaterialShared.cpp
--- a/Engine/Source/Runtime/Engine/Private/Materials/MaterialShared.cpp
+++ b/Engine/Source/Runtime/Engine/Private/Materials/MaterialShared.cpp
@@ -686,6 +686,7 @@ bool FMaterialResource::ShouldInjectEmissiveIntoLPV() const { return Material->b
 bool FMaterialResource::ShouldGenerateSphericalParticleNormals() const { return Material->bGenerateSphericalParticleNormals; }
 bool FMaterialResource::ShouldDisableDepthTest() const { return Material->bDisableDepthTest; }
 bool FMaterialResource::ShouldEnableResponsiveAA() const { return Material->bEnableResponsiveAA; }
+bool FMaterialResource::ShouldDoSSR() const { return Material->bScreenSpaceReflections; }
 bool FMaterialResource::IsWireframe() const { return Material->Wireframe; }
 bool FMaterialResource::IsLightFunction() const { return Material->MaterialDomain == MD_LightFunction; }
 bool FMaterialResource::IsUsedWithEditorCompositing() const { return Material->bUsedWithEditorCompositing; }
@@ -1172,6 +1173,7 @@ void FMaterial::SetupMaterialEnvironment(
 	OutEnvironment.SetDefine(TEXT("MATERIAL_NONMETAL"), IsNonmetal() ? TEXT("1") : TEXT("0"));
 	OutEnvironment.SetDefine(TEXT("MATERIAL_USE_LM_DIRECTIONALITY"), UseLmDirectionality() ? TEXT("1") : TEXT("0"));
 	OutEnvironment.SetDefine(TEXT("MATERIAL_INJECT_EMISSIVE_INTO_LPV"), ShouldInjectEmissiveIntoLPV() ? TEXT("1") : TEXT("0"));
+	OutEnvironment.SetDefine(TEXT("MATERIAL_SSR"), ShouldDoSSR() ? TEXT("1") : TEXT("0"));
 
 	{
 		auto DecalBlendMode = (EDecalBlendMode)GetDecalBlendMode();
diff --git a/Engine/Source/Runtime/Engine/Public/MaterialShared.h b/Engine/Source/Runtime/Engine/Public/MaterialShared.h
--- a/Engine/Source/Runtime/Engine/Public/MaterialShared.h
+++ b/Engine/Source/Runtime/Engine/Public/MaterialShared.h
@@ -895,6 +895,7 @@ public:
 	virtual bool ShouldGenerateSphericalParticleNormals() const { return false; }
 	virtual	bool ShouldDisableDepthTest() const { return false; }
 	virtual	bool ShouldEnableResponsiveAA() const { return false; }
+	virtual bool ShouldDoSSR() const { return false; }
 	virtual bool IsLightFunction() const = 0;
 	virtual bool IsUsedWithEditorCompositing() const { return false; }
 	virtual bool IsUsedWithDeferredDecal() const = 0;
@@ -1465,6 +1466,7 @@ public:
 	ENGINE_API virtual bool ShouldGenerateSphericalParticleNormals() const;
 	ENGINE_API virtual bool ShouldDisableDepthTest() const;
 	ENGINE_API virtual bool ShouldEnableResponsiveAA() const;
+	ENGINE_API virtual bool ShouldDoSSR() const;
 	ENGINE_API virtual bool IsLightFunction() const;
 	ENGINE_API virtual bool IsUsedWithEditorCompositing() const;
 	ENGINE_API virtual bool IsUsedWithDeferredDecal() const;
diff --git a/Engine/Source/Runtime/Renderer/Private/BasePassRendering.cpp b/Engine/Source/Runtime/Renderer/Private/BasePassRendering.cpp
--- a/Engine/Source/Runtime/Renderer/Private/BasePassRendering.cpp
+++ b/Engine/Source/Runtime/Renderer/Private/BasePassRendering.cpp
@@ -72,7 +72,7 @@ void FSkyLightReflectionParameters::GetSkyParametersFromScene(const FScene* Scen
 	}
 }
 
-void FTranslucentLightingParameters::Set(FRHICommandList& RHICmdList, FShader* Shader, const FSceneView* View)
+void FTranslucentLightingParameters::Set(FRHICommandList& RHICmdList, FShader* Shader, const FViewInfo* View)
 {
 	SetTextureParameter(
 		RHICmdList, 
@@ -107,6 +107,30 @@ void FTranslucentLightingParameters::Set(FRHICommandList& RHICmdList, FShader* S
 		GSceneRenderTargets.GetTranslucencyVolumeDirectional(TVC_Outer)->GetRenderTargetItem().ShaderResourceTexture);
 
 	SkyLightReflectionParameters.SetParameters(RHICmdList, Shader->GetPixelShader(), (const FScene*)(View->Family->Scene), true);
+
+	SetTextureParameter(
+		RHICmdList, 
+		Shader->GetPixelShader(), 
+		HZBTexture, 
+		HZBSampler, 
+		TStaticSamplerState<SF_Point, AM_Clamp, AM_Clamp, AM_Clamp>::GetRHI(), 
+		View->HZB->GetRenderTargetItem().ShaderResourceTexture );
+
+	TRefCountPtr<IPooledRenderTarget> PrevSceneColorRT = GSceneRenderTargets.GetSceneColor();
+	
+	FSceneViewState* ViewState = (FSceneViewState*)View->State;
+	if( ViewState && ViewState->TemporalAAHistoryRT && !View->bCameraCut )
+	{
+		PrevSceneColorRT = ViewState->TemporalAAHistoryRT;
+	}
+
+	SetTextureParameter(
+		RHICmdList, 
+		Shader->GetPixelShader(), 
+		PrevSceneColor, 
+		PrevSceneColorSampler, 
+		TStaticSamplerState<SF_Bilinear, AM_Clamp, AM_Clamp, AM_Clamp>::GetRHI(), 
+		PrevSceneColorRT->GetRenderTargetItem().ShaderResourceTexture );
 }
 
 void FTranslucentLightingParameters::SetMesh(FRHICommandList& RHICmdList, FShader* Shader, const FPrimitiveSceneProxy* Proxy, ERHIFeatureLevel::Type FeatureLevel)
diff --git a/Engine/Source/Runtime/Renderer/Private/BasePassRendering.h b/Engine/Source/Runtime/Renderer/Private/BasePassRendering.h
--- a/Engine/Source/Runtime/Renderer/Private/BasePassRendering.h
+++ b/Engine/Source/Runtime/Renderer/Private/BasePassRendering.h
@@ -252,9 +252,13 @@ public:
 		ReflectionCubemapSampler.Bind(ParameterMap, TEXT("ReflectionCubemapSampler"));
 		CubemapArrayIndex.Bind(ParameterMap, TEXT("CubemapArrayIndex"));
 		SkyLightReflectionParameters.Bind(ParameterMap);
+		HZBTexture.Bind(ParameterMap, TEXT("HZBTexture"));
+		HZBSampler.Bind(ParameterMap, TEXT("HZBSampler"));
+		PrevSceneColor.Bind(ParameterMap, TEXT("PrevSceneColor"));
+		PrevSceneColorSampler.Bind(ParameterMap, TEXT("PrevSceneColorSampler"));
 	}
 
-	void Set(FRHICommandList& RHICmdList, FShader* Shader, const FSceneView* View);
+	void Set(FRHICommandList& RHICmdList, FShader* Shader, const FViewInfo* View);
 
 	void SetMesh(FRHICommandList& RHICmdList, FShader* Shader, const FPrimitiveSceneProxy* Proxy, ERHIFeatureLevel::Type FeatureLevel);
 
@@ -273,6 +277,10 @@ public:
 		Ar << P.ReflectionCubemapSampler;
 		Ar << P.CubemapArrayIndex;
 		Ar << P.SkyLightReflectionParameters;
+		Ar << P.HZBTexture;
+		Ar << P.HZBSampler;
+		Ar << P.PrevSceneColor;
+		Ar << P.PrevSceneColorSampler;
 		return Ar;
 	}
 
@@ -290,6 +298,10 @@ private:
 	FShaderResourceParameter ReflectionCubemapSampler;
 	FShaderParameter CubemapArrayIndex;
 	FSkyLightReflectionParameters SkyLightReflectionParameters;
+	FShaderResourceParameter HZBTexture;
+	FShaderResourceParameter HZBSampler;
+	FShaderResourceParameter PrevSceneColor;
+	FShaderResourceParameter PrevSceneColorSampler;
 };
 
 /**
@@ -330,7 +342,7 @@ public:
 		FRHICommandList& RHICmdList, 
 		const FMaterialRenderProxy* MaterialRenderProxy, 
 		const FMaterial& MaterialResource, 
-		const FSceneView* View, 
+		const FViewInfo* View, 
 		EBlendMode BlendMode, 
 		bool bEnableEditorPrimitveDepthTest,
 		ESceneRenderTargetsMode::Type TextureMode)
diff --git a/Engine/Source/Runtime/Renderer/Private/CompositionLighting/CompositionLighting.cpp b/Engine/Source/Runtime/Renderer/Private/CompositionLighting/CompositionLighting.cpp
--- a/Engine/Source/Runtime/Renderer/Private/CompositionLighting/CompositionLighting.cpp
+++ b/Engine/Source/Runtime/Renderer/Private/CompositionLighting/CompositionLighting.cpp
@@ -163,18 +163,7 @@ static FRenderingCompositeOutputRef AddPostProcessingAmbientOcclusion(FRHIComman
 		AmbientOcclusionInMip3->SetInput(ePId_Input1, FRenderingCompositeOutputRef(AmbientOcclusionInMip2, ePId_Output0));
 	}
 
-	FRenderingCompositePass* HZBInput = 0;
-
-	FSceneViewState* ViewState = (FSceneViewState*)Context.View.State;
-
-	if(ViewState)
-	{
-		void BuildHZB(FRHICommandListImmediate& RHICmdList, const FViewInfo& View);
-
-		BuildHZB(RHICmdList, Context.View);
-
-		HZBInput = Context.Graph.RegisterPass( new FRCPassPostProcessInput( ViewState->HZB.Texture ) );
-	}
+	FRenderingCompositePass* HZBInput = Context.Graph.RegisterPass( new FRCPassPostProcessInput( const_cast< FViewInfo& >( Context.View ).HZB ) );
 
 	// upsample from lower resolution
 
diff --git a/Engine/Source/Runtime/Renderer/Private/DeferredShadingRenderer.cpp b/Engine/Source/Runtime/Renderer/Private/DeferredShadingRenderer.cpp
--- a/Engine/Source/Runtime/Renderer/Private/DeferredShadingRenderer.cpp
+++ b/Engine/Source/Runtime/Renderer/Private/DeferredShadingRenderer.cpp
@@ -571,6 +571,8 @@ void FDeferredShadingSceneRenderer::RenderFinish(FRHICommandListImmediate& RHICm
 	GSceneRenderTargets.SetLightAttenuation(0);
 }
 
+void BuildHZB( FRHICommandListImmediate& RHICmdList, FViewInfo& View );
+
 /** 
 * Renders the view family. 
 */
@@ -832,6 +834,11 @@ void FDeferredShadingSceneRenderer::Render(FRHICommandListImmediate& RHICmdList)
 	// This needs to happen before occlusion tests, which makes use of the small depth buffer.
 	UpdateDownsampledDepthSurface(RHICmdList);
 
+	for (int32 ViewIndex = 0; ViewIndex < Views.Num(); ViewIndex++)
+	{
+		BuildHZB( RHICmdList, Views[ ViewIndex ] );
+	}
+
 	// Issue occlusion queries
 	// This is done after the downsampled depth buffer is created so that it can be used for issuing queries
 	if ( bIsOcclusionTesting )
diff --git a/Engine/Source/Runtime/Renderer/Private/PostProcess/RenderTargetPool.cpp b/Engine/Source/Runtime/Renderer/Private/PostProcess/RenderTargetPool.cpp
--- a/Engine/Source/Runtime/Renderer/Private/PostProcess/RenderTargetPool.cpp
+++ b/Engine/Source/Runtime/Renderer/Private/PostProcess/RenderTargetPool.cpp
@@ -259,6 +259,15 @@ bool FRenderTargetPool::FindFreeElement(const FPooledRenderTargetDesc& Desc, TRe
 					(FTexture2DRHIRef&)Found->RenderTargetItem.ShaderResourceTexture,
 					Desc.NumSamples
 					);
+
+				if( Desc.NumMips > 1 )
+				{
+					Found->RenderTargetItem.MipSRVs.SetNum( Desc.NumMips );
+					for( uint16 i = 0; i < Desc.NumMips; i++ )
+					{
+						Found->RenderTargetItem.MipSRVs[i] = RHICreateShaderResourceView( (FTexture2DRHIRef&)Found->RenderTargetItem.ShaderResourceTexture, i );
+					}
+				}
 			}
 			else if(Desc.Is3DTexture())
 			{
diff --git a/Engine/Source/Runtime/Renderer/Private/PostProcess/ScreenSpaceReflections.cpp b/Engine/Source/Runtime/Renderer/Private/PostProcess/ScreenSpaceReflections.cpp
--- a/Engine/Source/Runtime/Renderer/Private/PostProcess/ScreenSpaceReflections.cpp
+++ b/Engine/Source/Runtime/Renderer/Private/PostProcess/ScreenSpaceReflections.cpp
@@ -256,19 +256,16 @@ FPooledRenderTargetDesc FRCPassPostProcessScreenSpaceReflections::ComputeOutputD
 
 	return Ret;
 }
-void BuildHZB(FRHICommandListImmediate& RHICmdList, const FViewInfo& View);
 
 void ScreenSpaceReflections(FRHICommandListImmediate& RHICmdList, FViewInfo& View, TRefCountPtr<IPooledRenderTarget>& SSROutput)
 {
-	BuildHZB(RHICmdList, View);
-
 	FRenderingCompositePassContext CompositeContext(RHICmdList, View);
 	FPostprocessContext Context( CompositeContext.Graph, View );
 
 	FSceneViewState* ViewState = (FSceneViewState*)Context.View.State;
 
 	FRenderingCompositePass* SceneColorInput = Context.Graph.RegisterPass( new FRCPassPostProcessInput( GSceneRenderTargets.GetSceneColor() ) );
-	FRenderingCompositePass* HZBInput = Context.Graph.RegisterPass( new FRCPassPostProcessInput( ViewState->HZB.Texture ) );
+	FRenderingCompositePass* HZBInput = Context.Graph.RegisterPass( new FRCPassPostProcessInput( View.HZB ) );
 
 	bool bPrevFrame = 0;
 	if( ViewState && ViewState->TemporalAAHistoryRT && !Context.View.bCameraCut )
diff --git a/Engine/Source/Runtime/Renderer/Private/SceneOcclusion.cpp b/Engine/Source/Runtime/Renderer/Private/SceneOcclusion.cpp
--- a/Engine/Source/Runtime/Renderer/Private/SceneOcclusion.cpp
+++ b/Engine/Source/Runtime/Renderer/Private/SceneOcclusion.cpp
@@ -552,15 +552,7 @@ void FHZBOcclusionTester::MapResults(FRHICommandListImmediate& RHICmdList)
 {
 	check( !ResultsBuffer );
 
-	// hacky (we point to some buffer that is not the right size but we prevent reads from it be having a invalid frame number)
-	// First frame
-	static uint8 FirstFrameBuffer[] = { 255 };
-
-	if( IsInvalidFrame() )
-	{
-		ResultsBuffer = FirstFrameBuffer;
-	}
-	else
+	if( !IsInvalidFrame() )
 	{
 		uint32 IdleStart = FPlatformTime::Cycles();
 
@@ -569,18 +561,19 @@ void FHZBOcclusionTester::MapResults(FRHICommandListImmediate& RHICmdList)
 
 		RHICmdList.MapStagingSurface(ResultsTextureCPU->GetRenderTargetItem().ShaderResourceTexture, (void*&)ResultsBuffer, Width, Height);
 
-		// Can happen because of device removed, we might crash later but this occlusion culling system can behave gracefully.
-		if(!ResultsBuffer)
-		{
-			ResultsBuffer = FirstFrameBuffer;
-			SetInvalidFrameNumber();
-		}
-
 		// RHIMapStagingSurface will block until the results are ready (from the previous frame) so we need to consider this RT idle time
 		GRenderThreadIdle[ERenderThreadIdleTypes::WaitingForGPUQuery] += FPlatformTime::Cycles() - IdleStart;
 		GRenderThreadNumIdle[ERenderThreadIdleTypes::WaitingForGPUQuery]++;
 	}
-	check( ResultsBuffer );
+	
+	// Can happen because of device removed, we might crash later but this occlusion culling system can behave gracefully.
+	if( ResultsBuffer == NULL )
+	{
+		// First frame
+		static uint8 FirstFrameBuffer[] = { 255 };
+		ResultsBuffer = FirstFrameBuffer;
+		SetInvalidFrameNumber();
+	}
 }
 
 void FHZBOcclusionTester::UnmapResults(FRHICommandListImmediate& RHICmdList)
@@ -599,10 +592,31 @@ bool FHZBOcclusionTester::IsVisible( uint32 Index ) const
 	checkSlow( Index < SizeX * SizeY );
 	
 	// TODO shader compress to bits
+
+#if 0
+	return ResultsBuffer[ 4 * Index ] != 0;
+#elif 0
 	uint32 x = FMath::ReverseMortonCode2( Index >> 0 );
 	uint32 y = FMath::ReverseMortonCode2( Index >> 1 );
 	uint32 m = x + y * SizeX;
 	return ResultsBuffer[ 4 * m ] != 0;
+#else
+	// TODO put block constants in class
+	// TODO optimize
+	const uint32 BlockSize = 8;
+	const uint32 SizeInBlocksX = SizeX / BlockSize;
+	const uint32 SizeInBlocksY = SizeY / BlockSize;
+
+	const int32 BlockIndex = Index / (BlockSize * BlockSize);
+	const int32 BlockX = BlockIndex % SizeInBlocksX;
+	const int32 BlockY = BlockIndex / SizeInBlocksY;
+
+	const int32 b = Index % (BlockSize * BlockSize);
+	const int32 x = BlockX * BlockSize + b % BlockSize;
+	const int32 y = BlockY * BlockSize + b / BlockSize;
+
+	return ResultsBuffer[ 4 * (x + y * SizeY) ] != 0;
+#endif
 }
 
 class FHZBTestPS : public FGlobalShader
@@ -642,16 +656,16 @@ public:
 		BoundsExtentSampler.Bind( Initializer.ParameterMap, TEXT("BoundsExtentSampler") );
 	}
 
-	void SetParameters(FRHICommandList& RHICmdList, const FSceneView& View, const FHZB& HZB, FTextureRHIParamRef BoundsCenter, FTextureRHIParamRef BoundsExtent )
+	void SetParameters(FRHICommandList& RHICmdList, const FViewInfo& View, FTextureRHIParamRef BoundsCenter, FTextureRHIParamRef BoundsExtent )
 	{
 		const FPixelShaderRHIParamRef ShaderRHI = GetPixelShader();
 
 		FGlobalShader::SetParameters(RHICmdList, ShaderRHI, View );
 
-		const FVector2D InvSize( 1.0f / HZB.Size.X, 1.0f / HZB.Size.Y );
+		const FVector2D InvSize( 1.0f / 512.0f, 1.0f / 256.0f );
 		SetShaderValue(RHICmdList, ShaderRHI, InvSizeParameter, InvSize );
 
-		SetTextureParameter(RHICmdList, ShaderRHI, HZBTexture, HZBSampler, TStaticSamplerState<SF_Point,AM_Clamp,AM_Clamp,AM_Clamp>::GetRHI(), HZB.Texture->GetRenderTargetItem().ShaderResourceTexture );
+		SetTextureParameter(RHICmdList, ShaderRHI, HZBTexture, HZBSampler, TStaticSamplerState<SF_Point,AM_Clamp,AM_Clamp,AM_Clamp>::GetRHI(), View.HZB->GetRenderTargetItem().ShaderResourceTexture );
 
 		SetTextureParameter(RHICmdList, ShaderRHI, BoundsCenterTexture, BoundsCenterSampler, TStaticSamplerState<SF_Point,AM_Clamp,AM_Clamp,AM_Clamp>::GetRHI(), BoundsCenter );
 		SetTextureParameter(RHICmdList, ShaderRHI, BoundsExtentTexture, BoundsExtentSampler, TStaticSamplerState<SF_Point,AM_Clamp,AM_Clamp,AM_Clamp>::GetRHI(), BoundsExtent );
@@ -697,7 +711,6 @@ void FHZBOcclusionTester::Submit(FRHICommandListImmediate& RHICmdList, const FVi
 	{
 #if PLATFORM_MAC // Workaround radr://16096028 Texture Readback via glReadPixels + PBOs stalls on Nvidia GPUs
 		FPooledRenderTargetDesc Desc( FPooledRenderTargetDesc::Create2DDesc( FIntPoint( SizeX, SizeY ), PF_R8G8B8A8, TexCreate_None, TexCreate_RenderTargetable, false ) );
-
 #else
 		FPooledRenderTargetDesc Desc( FPooledRenderTargetDesc::Create2DDesc( FIntPoint( SizeX, SizeY ), PF_B8G8R8A8, TexCreate_None, TexCreate_RenderTargetable, false ) );
 #endif
@@ -705,6 +718,33 @@ void FHZBOcclusionTester::Submit(FRHICommandListImmediate& RHICmdList, const FVi
 	}
 
 	{
+#if 0
+		static float CenterBuffer[ SizeX * SizeY ][4];
+		static float ExtentBuffer[ SizeX * SizeY ][4];
+
+		FMemory::Memset( CenterBuffer, 0, sizeof( CenterBuffer ) );
+		FMemory::Memset( ExtentBuffer, 0, sizeof( ExtentBuffer ) );
+
+		const uint32 NumPrimitives = Primitives.Num();
+		for( uint32 i = 0; i < NumPrimitives; i++ )
+		{
+			const FOcclusionPrimitive& Primitive = Primitives[i];
+
+			CenterBuffer[i][0] = Primitive.Center.X;
+			CenterBuffer[i][1] = Primitive.Center.Y;
+			CenterBuffer[i][2] = Primitive.Center.Z;
+			CenterBuffer[i][3] = 0.0f;
+
+			ExtentBuffer[i][0] = Primitive.Extent.X;
+			ExtentBuffer[i][1] = Primitive.Extent.Y;
+			ExtentBuffer[i][2] = Primitive.Extent.Z;
+			ExtentBuffer[i][3] = 1.0f;
+		}
+
+		FUpdateTextureRegion2D Region( 0, 0, 0, 0, SizeX, SizeY );
+		RHIUpdateTexture2D( (FTexture2DRHIRef&)BoundsCenterTexture->GetRenderTargetItem().ShaderResourceTexture, 0, Region, SizeX * 4 * sizeof( float ), (uint8*)CenterBuffer );
+		RHIUpdateTexture2D( (FTexture2DRHIRef&)BoundsExtentTexture->GetRenderTargetItem().ShaderResourceTexture, 0, Region, SizeX * 4 * sizeof( float ), (uint8*)ExtentBuffer );
+#elif 0
 		static float CenterBuffer[ SizeX * SizeY ][4];
 		static float ExtentBuffer[ SizeX * SizeY ][4];
 
@@ -739,6 +779,51 @@ void FHZBOcclusionTester::Submit(FRHICommandListImmediate& RHICmdList, const FVi
 		FUpdateTextureRegion2D Region( 0, 0, 0, 0, SizeX, SizeY );
 		RHIUpdateTexture2D( (FTexture2DRHIRef&)BoundsCenterTexture->GetRenderTargetItem().ShaderResourceTexture, 0, Region, SizeX * 4 * sizeof( float ), (uint8*)CenterBuffer );
 		RHIUpdateTexture2D( (FTexture2DRHIRef&)BoundsExtentTexture->GetRenderTargetItem().ShaderResourceTexture, 0, Region, SizeX * 4 * sizeof( float ), (uint8*)ExtentBuffer );
+#else
+		// Update in blocks to avoid large update
+		const uint32 BlockSize = 8;
+		const uint32 SizeInBlocksX = SizeX / BlockSize;
+		const uint32 SizeInBlocksY = SizeY / BlockSize;
+		const uint32 BlockStride = BlockSize * 4 * sizeof( float );
+
+		float CenterBuffer[ BlockSize * BlockSize ][4];
+		float ExtentBuffer[ BlockSize * BlockSize ][4];
+
+		const uint32 NumPrimitives = Primitives.Num();
+		for( uint32 i = 0; i < NumPrimitives; i += BlockSize * BlockSize )
+		{
+			const uint32 BlockEnd = FMath::Min( BlockSize * BlockSize, NumPrimitives - i );
+			for( uint32 b = 0; b < BlockEnd; b++ )
+			{
+				const FOcclusionPrimitive& Primitive = Primitives[ i + b ];
+
+				CenterBuffer[b][0] = Primitive.Center.X;
+				CenterBuffer[b][1] = Primitive.Center.Y;
+				CenterBuffer[b][2] = Primitive.Center.Z;
+				CenterBuffer[b][3] = 0.0f;
+
+				ExtentBuffer[b][0] = Primitive.Extent.X;
+				ExtentBuffer[b][1] = Primitive.Extent.Y;
+				ExtentBuffer[b][2] = Primitive.Extent.Z;
+				ExtentBuffer[b][3] = 1.0f;
+			}
+
+			// Clear rest of block
+			if( BlockEnd < BlockSize * BlockSize )
+			{
+				FMemory::Memset( (float*)CenterBuffer + BlockEnd * 4, 0, sizeof( CenterBuffer ) - BlockEnd * 4 * sizeof(float) );
+				FMemory::Memset( (float*)ExtentBuffer + BlockEnd * 4, 0, sizeof( ExtentBuffer ) - BlockEnd * 4 * sizeof(float) );
+			}
+
+			const int32 BlockIndex = i / (BlockSize * BlockSize);
+			const int32 BlockX = BlockIndex % SizeInBlocksX;
+			const int32 BlockY = BlockIndex / SizeInBlocksY;
+
+			FUpdateTextureRegion2D Region( BlockX * BlockSize, BlockY * BlockSize, 0, 0, BlockSize, BlockSize );
+			RHIUpdateTexture2D( (FTexture2DRHIRef&)BoundsCenterTexture->GetRenderTargetItem().ShaderResourceTexture, 0, Region, BlockStride, (uint8*)CenterBuffer );
+			RHIUpdateTexture2D( (FTexture2DRHIRef&)BoundsExtentTexture->GetRenderTargetItem().ShaderResourceTexture, 0, Region, BlockStride, (uint8*)ExtentBuffer );
+		}
+#endif
 		Primitives.Empty();
 	}
 
@@ -748,13 +833,17 @@ void FHZBOcclusionTester::Submit(FRHICommandListImmediate& RHICmdList, const FVi
 
 		SetRenderTarget(RHICmdList, ResultsTextureGPU->GetRenderTargetItem().TargetableTexture, NULL);
 
+		RHICmdList.SetBlendState(TStaticBlendState<>::GetRHI());
+		RHICmdList.SetRasterizerState(TStaticRasterizerState<>::GetRHI());
+		RHICmdList.SetDepthStencilState(TStaticDepthStencilState< false, CF_Always >::GetRHI());
+
 		TShaderMapRef< FScreenVS >	VertexShader(View.ShaderMap);
 		TShaderMapRef< FHZBTestPS >	PixelShader(View.ShaderMap);
 
 		static FGlobalBoundShaderState BoundShaderState;
 		SetGlobalBoundShaderState(RHICmdList, View.GetFeatureLevel(), BoundShaderState, GFilterVertexDeclaration.VertexDeclarationRHI, *VertexShader, *PixelShader);
 
-		PixelShader->SetParameters(RHICmdList, View,  ViewState->HZB, BoundsCenterTexture->GetRenderTargetItem().ShaderResourceTexture, BoundsExtentTexture->GetRenderTargetItem().ShaderResourceTexture );
+		PixelShader->SetParameters(RHICmdList, View, BoundsCenterTexture->GetRenderTargetItem().ShaderResourceTexture, BoundsExtentTexture->GetRenderTargetItem().ShaderResourceTexture );
 
 		RHICmdList.SetViewport(0, 0, 0.0f, SizeX, SizeY, 1.0f);
 
@@ -852,34 +941,20 @@ public:
 IMPLEMENT_SHADER_TYPE(template<>,THZBBuildPS<0>,TEXT("HZBOcclusion"),TEXT("HZBBuildPS"),SF_Pixel);
 IMPLEMENT_SHADER_TYPE(template<>,THZBBuildPS<1>,TEXT("HZBOcclusion"),TEXT("HZBBuildPS"),SF_Pixel);
 
-void BuildHZB(FRHICommandListImmediate& RHICmdList, const FViewInfo& View)
+void BuildHZB( FRHICommandListImmediate& RHICmdList, FViewInfo& View )
 {
 	QUICK_SCOPE_CYCLE_COUNTER(STAT_BuildHZB);
 	SCOPED_DRAW_EVENT(RHICmdList, BuildHZB);
 
-	FSceneViewState* ViewState = (FSceneViewState*)View.State;
-
-	if(!ViewState)
-	{
-		// not view state (e.g. thumbnail rendering?), no HZB (no screen space reflections or occlusion culling)
-		return;
-	}
-
-	ViewState->HZB.AllocHZB();
-
-	if(ViewState->HZB.bDataIsValid)
-	{
-		// data was already computed, no need to do it again
-		return;
-	}
-
-	ViewState->HZB.bDataIsValid = true;
-
 	// Must be power of 2
-	const FIntPoint HZBSize = ViewState->HZB.Size;
-	const uint32 NumMips = ViewState->HZB.NumMips;
+	const FIntPoint HZBSize( 512, 256 );
+	const uint32 NumMips = 8;
+
+	FPooledRenderTargetDesc Desc( FPooledRenderTargetDesc::Create2DDesc( HZBSize, PF_R16F, TexCreate_None, TexCreate_RenderTargetable | TexCreate_ShaderResource, false, NumMips ) );
+	Desc.Flags |= TexCreate_FastVRAM;
+	GRenderTargetPool.FindFreeElement( Desc, View.HZB, TEXT("HZB") );
 	
-	FSceneRenderTargetItem& HZBRenderTarget = ViewState->HZB.Texture->GetRenderTargetItem();
+	FSceneRenderTargetItem& HZBRenderTarget = View.HZB->GetRenderTargetItem();
 	
 	RHICmdList.SetBlendState(TStaticBlendState<>::GetRHI());
 	RHICmdList.SetRasterizerState(TStaticRasterizerState<>::GetRHI());
@@ -934,7 +1009,7 @@ void BuildHZB(FRHICommandListImmediate& RHICmdList, const FViewInfo& View)
 		
 		SetGlobalBoundShaderState(RHICmdList, View.GetFeatureLevel(), BoundShaderState, GFilterVertexDeclaration.VertexDeclarationRHI, *VertexShader, *PixelShader);
 
-		PixelShader->SetParameters( RHICmdList, View, DstSize, ViewState->HZB.MipSRVs[ MipIndex - 1 ] );
+		PixelShader->SetParameters(RHICmdList, View, SrcSize, HZBRenderTarget.MipSRVs[ MipIndex - 1 ] );
 
 		RHICmdList.SetViewport(0, 0, 0.0f, DstSize.X, DstSize.Y, 1.0f);
 
@@ -955,7 +1030,7 @@ void BuildHZB(FRHICommandListImmediate& RHICmdList, const FViewInfo& View)
 		DstSize /= 2;
 	}
 
-	GRenderTargetPool.VisualizeTexture.SetCheckPoint(RHICmdList, ViewState->HZB.Texture);
+	GRenderTargetPool.VisualizeTexture.SetCheckPoint( RHICmdList, View.HZB );
 }
 
 void FDeferredShadingSceneRenderer::BeginOcclusionTests(FRHICommandListImmediate& RHICmdList)
@@ -1122,8 +1197,6 @@ void FDeferredShadingSceneRenderer::BeginOcclusionTests(FRHICommandListImmediate
 			check( ViewState->HZBOcclusionTests.IsValidFrame(ViewState->OcclusionFrameCounter) );
 
 			SCOPED_DRAW_EVENT(RHICmdList, HZB);
-
-			BuildHZB(RHICmdList, View);
 			ViewState->HZBOcclusionTests.Submit(RHICmdList, View);
 		}
 	}
diff --git a/Engine/Source/Runtime/Renderer/Private/ScenePrivate.h b/Engine/Source/Runtime/Renderer/Private/ScenePrivate.h
--- a/Engine/Source/Runtime/Renderer/Private/ScenePrivate.h
+++ b/Engine/Source/Runtime/Renderer/Private/ScenePrivate.h
@@ -358,54 +358,6 @@ private:
 /** Random table for occlusion **/
 extern FOcclusionRandomStream GOcclusionRandomStream;
 
-// Hierarchical Z Buffer for occlusion culling / TemporalAA
-// Must be power of 2
-struct FHZB
-{
-	enum		{ NumMips = 8 };
-	FIntPoint	Size;
-
-	TRefCountPtr<IPooledRenderTarget>	Texture;
-	FShaderResourceViewRHIRef			MipSRVs[ NumMips ];
-	// set to false on InitViews(), true when BuildHZB() did run
-	bool bDataIsValid;
-
-	FHZB() 
-		: Size(512, 256)
-		, bDataIsValid(false)
-	{
-	}
-
-	void SafeRelease()
-	{
-		Texture.SafeRelease();
-
-		for(uint32 i = 0; i < NumMips; ++i)
-		{
-			MipSRVs[i].SafeRelease();
-		}
-	}
-
-	void AllocHZB()
-	{
-		if(Texture)
-		{
-			// avoid reallocation, would be no problem with pooled RT but the SRV for the mips might be wasteful
-			return;
-		}
-
-		FPooledRenderTargetDesc Desc( FPooledRenderTargetDesc::Create2DDesc( Size, PF_R32_FLOAT, TexCreate_None, TexCreate_RenderTargetable | TexCreate_ShaderResource, false, NumMips ) );
-		GRenderTargetPool.FindFreeElement( Desc, Texture, TEXT("HZB") );
-
-		FSceneRenderTargetItem& HZBRenderTarget = Texture->GetRenderTargetItem();
-		for(uint32 i = 0; i < NumMips; ++i)
-		{
-			// TODO implement for all RHI
-			MipSRVs[i] = RHICreateShaderResourceView( (FTexture2DRHIRef&)HZBRenderTarget.ShaderResourceTexture, i );
-		}
-	}
-};
-
 /**
  * The scene manager's private implementation of persistent view state.
  * This class is associated with a particular camera across multiple frames by the game thread.
@@ -517,9 +469,6 @@ private:
 	FLightPropagationVolume* LightPropagationVolume;
 
 public:
-	// Hierarchical Z Buffer for occlusion culling (per view to allow occlusion test to be done for each view and still allow screenspace reflections using the same data)
-	FHZB HZB;
-
 	// If Translucency should be rendered into a separate RT and composited without DepthOfField, can be disabled in the materials (affects sorting)
 	TRefCountPtr<IPooledRenderTarget> SeparateTranslucencyRT;
 	// Temporal AA result of last frame
@@ -714,7 +663,6 @@ public:
 		MobileAaBloomSunVignette1.SafeRelease();
 		MobileAaColor0.SafeRelease();
 		MobileAaColor1.SafeRelease();
-		HZB.SafeRelease();
 		SelectionOutlineCacheKey.SafeRelease();
 		SelectionOutlineCacheValue.SafeRelease();
 	}
@@ -739,8 +687,6 @@ public:
 	{
 		check(IsInRenderingThread());
 
-		HZB.bDataIsValid = false;
-
 		if(!(CurrentView.FinalPostProcessSettings.IndirectLightingColor * CurrentView.FinalPostProcessSettings.IndirectLightingIntensity).IsAlmostBlack())
 		{
 			CreateLightPropagationVolumeIfNeeded(CurrentView.GetFeatureLevel());
diff --git a/Engine/Source/Runtime/Renderer/Private/SceneRendering.h b/Engine/Source/Runtime/Renderer/Private/SceneRendering.h
--- a/Engine/Source/Runtime/Renderer/Private/SceneRendering.h
+++ b/Engine/Source/Runtime/Renderer/Private/SceneRendering.h
@@ -458,6 +458,9 @@ public:
 	FOcclusionQueryBatcher IndividualOcclusionQueries;
 	FOcclusionQueryBatcher GroupedOcclusionQueries;
 
+	// Hierarchical Z Buffer
+	TRefCountPtr<IPooledRenderTarget> HZB;
+
 	/** Used by occlusion for percent unoccluded calculations. */
 	float OneOverNumPossiblePixels;
 
diff --git a/Engine/Source/Runtime/Renderer/Public/RendererInterface.h b/Engine/Source/Runtime/Renderer/Public/RendererInterface.h
--- a/Engine/Source/Runtime/Renderer/Public/RendererInterface.h
+++ b/Engine/Source/Runtime/Renderer/Public/RendererInterface.h
@@ -330,6 +330,10 @@ struct FSceneRenderTargetItem
 		TargetableTexture.SafeRelease();
 		ShaderResourceTexture.SafeRelease();
 		UAV.SafeRelease();
+		for( int32 i = 0; i < MipSRVs.Num(); i++ )
+		{
+			MipSRVs[i].SafeRelease();
+		}
 	}
 
 	bool IsValid() const
@@ -345,6 +349,8 @@ struct FSceneRenderTargetItem
 	FTextureRHIRef ShaderResourceTexture;
 	/** only created if requested through the flag  */
 	FUnorderedAccessViewRHIRef UAV;
+	/** only created if requested through the flag  */
+	TArray< FShaderResourceViewRHIRef > MipSRVs;
 };
 
 /**
