UE4 - Fix crash UE-7413

diff --git a/Engine/Source/Runtime/Renderer/Private/BasePassRendering.cpp b/Engine/Source/Runtime/Renderer/Private/BasePassRendering.cpp
--- a/Engine/Source/Runtime/Renderer/Private/BasePassRendering.cpp
+++ b/Engine/Source/Runtime/Renderer/Private/BasePassRendering.cpp
@@ -108,13 +108,16 @@ void FTranslucentLightingParameters::Set(FRHICommandList& RHICmdList, FShader* S
 
 	SkyLightReflectionParameters.SetParameters(RHICmdList, Shader->GetPixelShader(), (const FScene*)(View->Family->Scene), true);
 
-	SetTextureParameter(
-		RHICmdList, 
-		Shader->GetPixelShader(), 
-		HZBTexture, 
-		HZBSampler, 
-		TStaticSamplerState<SF_Point, AM_Clamp, AM_Clamp, AM_Clamp>::GetRHI(), 
-		View->HZB->GetRenderTargetItem().ShaderResourceTexture );
+	if (View->HZB)
+	{
+		SetTextureParameter(
+			RHICmdList,
+			Shader->GetPixelShader(),
+			HZBTexture,
+			HZBSampler,
+			TStaticSamplerState<SF_Point, AM_Clamp, AM_Clamp, AM_Clamp>::GetRHI(),
+			View->HZB->GetRenderTargetItem().ShaderResourceTexture);
+	}
 
 	TRefCountPtr<IPooledRenderTarget> PrevSceneColorRT = GSceneRenderTargets.GetSceneColor();
 	
