UE4 - crash fix for consoles: bad ref counting on new translucency code

This was specifically approved by NickP for lockdown

diff --git a/Engine/Source/Runtime/Renderer/Private/BasePassRendering.cpp b/Engine/Source/Runtime/Renderer/Private/BasePassRendering.cpp
--- a/Engine/Source/Runtime/Renderer/Private/BasePassRendering.cpp
+++ b/Engine/Source/Runtime/Renderer/Private/BasePassRendering.cpp
@@ -111,29 +111,29 @@ void FTranslucentLightingParameters::Set(FRHICommandList& RHICmdList, FShader* S
 	if (View->HZB)
 	{
 		SetTextureParameter(
-			RHICmdList,
-			Shader->GetPixelShader(),
-			HZBTexture,
-			HZBSampler,
-			TStaticSamplerState<SF_Point, AM_Clamp, AM_Clamp, AM_Clamp>::GetRHI(),
-			View->HZB->GetRenderTargetItem().ShaderResourceTexture);
-	}
+			RHICmdList, 
+			Shader->GetPixelShader(), 
+			HZBTexture, 
+			HZBSampler, 
+			TStaticSamplerState<SF_Point, AM_Clamp, AM_Clamp, AM_Clamp>::GetRHI(), 
+			View->HZB->GetRenderTargetItem().ShaderResourceTexture );
 
-	TRefCountPtr<IPooledRenderTarget> PrevSceneColorRT = GSceneRenderTargets.GetSceneColor();
-	
-	FSceneViewState* ViewState = (FSceneViewState*)View->State;
-	if( ViewState && ViewState->TemporalAAHistoryRT && !View->bCameraCut )
-	{
-		PrevSceneColorRT = ViewState->TemporalAAHistoryRT;
-	}
+		TRefCountPtr<IPooledRenderTarget>* PrevSceneColorRT = &GSceneRenderTargets.GetSceneColor();
 
-	SetTextureParameter(
-		RHICmdList, 
-		Shader->GetPixelShader(), 
-		PrevSceneColor, 
-		PrevSceneColorSampler, 
-		TStaticSamplerState<SF_Bilinear, AM_Clamp, AM_Clamp, AM_Clamp>::GetRHI(), 
-		PrevSceneColorRT->GetRenderTargetItem().ShaderResourceTexture );
+		FSceneViewState* ViewState = (FSceneViewState*)View->State;
+		if( ViewState && ViewState->TemporalAAHistoryRT && !View->bCameraCut )
+		{
+			PrevSceneColorRT = &ViewState->TemporalAAHistoryRT;
+		}
+
+		SetTextureParameter(
+			RHICmdList, 
+			Shader->GetPixelShader(), 
+			PrevSceneColor, 
+			PrevSceneColorSampler, 
+			TStaticSamplerState<SF_Bilinear, AM_Clamp, AM_Clamp, AM_Clamp>::GetRHI(), 
+			(*PrevSceneColorRT)->GetRenderTargetItem().ShaderResourceTexture );
+	}
 }
 
 void FTranslucentLightingParameters::SetMesh(FRHICommandList& RHICmdList, FShader* Shader, const FPrimitiveSceneProxy* Proxy, ERHIFeatureLevel::Type FeatureLevel)
