Fix translucent lighting

diff --git a/Engine/Source/Runtime/Renderer/Private/TranslucentLighting.cpp b/Engine/Source/Runtime/Renderer/Private/TranslucentLighting.cpp
--- a/Engine/Source/Runtime/Renderer/Private/TranslucentLighting.cpp
+++ b/Engine/Source/Runtime/Renderer/Private/TranslucentLighting.cpp
@@ -1054,7 +1054,7 @@ void ClearVolumeTextures(FRHICommandList& RHICmdList, ERHIFeatureLevel::Type Fea
 
 
 
-#ifdef PLATFORM_XBOXONE
+#if PLATFORM_XBOXONE
 	// ClearMRT is faster on Xbox
 	if (true)
 #else
