added experimental r.Bloom/Cross feature to get more interestingly shaped bloom (for free, almost like anisomorphic lens flares)

diff --git a/Engine/Source/Runtime/Renderer/Private/PostProcess/PostProcessWeightedSampleSum.cpp b/Engine/Source/Runtime/Renderer/Private/PostProcess/PostProcessWeightedSampleSum.cpp
--- a/Engine/Source/Runtime/Renderer/Private/PostProcess/PostProcessWeightedSampleSum.cpp
+++ b/Engine/Source/Runtime/Renderer/Private/PostProcess/PostProcessWeightedSampleSum.cpp
@@ -253,9 +253,6 @@ void SetFilterShaders(
 #undef SET_FILTER_SHADER_TYPE
 }
 
-
-
-
 /**
  * Evaluates a normal distribution PDF at given X.
  * This function misses the math for scaling the result (faster, not needed if the resulting values are renormalized).
@@ -264,16 +261,25 @@ void SetFilterShaders(
  * @param Variance - The normal distribution's variance.
  * @return The value of the normal distribution at X. (unscaled)
  */
-static float NormalDistributionUnscaled(float X,float Mean,float Variance)
+static float NormalDistributionUnscaled(float X,float Mean,float Variance, EFilterShape FilterShape, float CrossCenterWeight)
 {
-	return FMath::Exp(-FMath::Square(X - Mean) / (2.0f * Variance));
+	float Ret = FMath::Exp(-FMath::Square(X - Mean) / (2.0f * Variance));
+
+	// if center sample
+	if(X == 0.0f)
+	{
+		// give it extra weight
+		Ret += CrossCenterWeight;
+	}
+
+	return Ret;
 }
 
 /**
  * @return NumSamples >0
  */
 
-static uint32 Compute1DGaussianFilterKernel(ERHIFeatureLevel::Type InFeatureLevel, float KernelRadius, FVector2D OutOffsetAndWeight[MAX_FILTER_SAMPLES], uint32 MaxFilterSamples)
+static uint32 Compute1DGaussianFilterKernel(ERHIFeatureLevel::Type InFeatureLevel, float KernelRadius, FVector2D OutOffsetAndWeight[MAX_FILTER_SAMPLES], uint32 MaxFilterSamples, EFilterShape FilterShape, float CrossCenterWeight)
 {
 	float ClampedKernelRadius = FRCPassPostProcessWeightedSampleSum::GetClampedKernelRadius( InFeatureLevel, KernelRadius );
 	int32 IntegerKernelRadius = FRCPassPostProcessWeightedSampleSum::GetIntegerKernelRadius( InFeatureLevel, KernelRadius );
@@ -284,7 +290,7 @@ static uint32 Compute1DGaussianFilterKernel(ERHIFeatureLevel::Type InFeatureLeve
 	float WeightSum = 0.0f;
 	for(int32 SampleIndex = -IntegerKernelRadius; SampleIndex <= IntegerKernelRadius; SampleIndex += 2)
 	{
-		float Weight0 = NormalDistributionUnscaled(SampleIndex, 0, ClampedKernelRadius);
+		float Weight0 = NormalDistributionUnscaled(SampleIndex, 0, ClampedKernelRadius, FilterShape, CrossCenterWeight);
 		float Weight1 = 0.0f;
 
 		// Because we use bilinear filtering we only require half the sample count.
@@ -294,7 +300,7 @@ static uint32 Compute1DGaussianFilterKernel(ERHIFeatureLevel::Type InFeatureLeve
 		//    c * .. but another texel to the right would accidentially leak into this computation.
 		if(SampleIndex != IntegerKernelRadius)
 		{
-			Weight1 = NormalDistributionUnscaled(SampleIndex + 1, 0, ClampedKernelRadius);
+			Weight1 = NormalDistributionUnscaled(SampleIndex + 1, 0, ClampedKernelRadius, FilterShape, CrossCenterWeight);
 		}
 
 		float TotalWeight = Weight0 + Weight1;
@@ -320,6 +326,7 @@ FRCPassPostProcessWeightedSampleSum::FRCPassPostProcessWeightedSampleSum(EFilter
 	, SizeScale(InSizeScale)
 	, TintValue(InTintValue)
 	, DebugName(InDebugName)
+	, CrossCenterWeight(0.0f)
 {
 }
 
@@ -368,7 +375,7 @@ void FRCPassPostProcessWeightedSampleSum::Process(FRenderingCompositePassContext
 	// compute 1D filtered samples
 	uint32 MaxNumSamples = GetMaxNumSamples(FeatureLevel);
 
-	uint32 NumSamples = Compute1DGaussianFilterKernel(FeatureLevel, EffectiveBlurRadius, OffsetAndWeight, MaxNumSamples);
+	uint32 NumSamples = Compute1DGaussianFilterKernel(FeatureLevel, EffectiveBlurRadius, OffsetAndWeight, MaxNumSamples, FilterShape, CrossCenterWeight);
 
 	SCOPED_DRAW_EVENTF(Context.RHICmdList, PostProcessWeightedSampleSum, TEXT("PostProcessWeightedSampleSum#%d"), NumSamples);
 
diff --git a/Engine/Source/Runtime/Renderer/Private/PostProcess/PostProcessWeightedSampleSum.h b/Engine/Source/Runtime/Renderer/Private/PostProcess/PostProcessWeightedSampleSum.h
--- a/Engine/Source/Runtime/Renderer/Private/PostProcess/PostProcessWeightedSampleSum.h
+++ b/Engine/Source/Runtime/Renderer/Private/PostProcess/PostProcessWeightedSampleSum.h
@@ -122,6 +122,9 @@ public:
 	static float GetClampedKernelRadius(ERHIFeatureLevel::Type InFeatureLevel, float KernelRadius);
 	static int GetIntegerKernelRadius(ERHIFeatureLevel::Type InFeatureLevel, float KernelRadius);
 
+	// @param InCrossCenterWeight >=0
+	void SetCrossCenterWeight(float InCrossCenterWeight) { check(InCrossCenterWeight >= 0.0f); CrossCenterWeight = InCrossCenterWeight; }
+
 private:
 	void DrawQuad(FRHICommandListImmediate& RHICmdList, bool bDoFastBlur, FIntRect SrcRect, FIntRect DestRect, bool bRequiresClear, FIntPoint DestSize, FIntPoint SrcSize, FShader* VertexShader) const;
 	static uint32 GetMaxNumSamples(ERHIFeatureLevel::Type InFeatureLevel);
@@ -132,6 +135,9 @@ private:
 	float SizeScale;
 	FLinearColor TintValue;
 	const TCHAR* DebugName;
+	// to give the center sample some special weight (see r.Bloom.Cross), >=0
+	float CrossCenterWeight;
+
 	// @return true: half x resolution for horizontal pass, vertical pass takes that as input, lower quality
 	bool DoFastBlur() const;
 };
diff --git a/Engine/Source/Runtime/Renderer/Private/PostProcess/PostProcessing.cpp b/Engine/Source/Runtime/Renderer/Private/PostProcess/PostProcessing.cpp
--- a/Engine/Source/Runtime/Renderer/Private/PostProcess/PostProcessing.cpp
+++ b/Engine/Source/Runtime/Renderer/Private/PostProcess/PostProcessing.cpp
@@ -100,6 +100,16 @@ static TAutoConsoleVariable<float> CVarMotionBlurSoftEdgeSize(
 	TEXT(" 0:off (not free and does never completely disable), >0, 1.0 (default)"),
 	ECVF_RenderThreadSafe);
 
+static TAutoConsoleVariable<float> CVarBloomCross(
+	TEXT("r.Bloom.Cross"),
+	0.0f,
+	TEXT("Experimental feature to give bloom kernel a more bright center sample (values between 1 and 3 work without causing aliasing)\n")
+	TEXT("Existing bloom get lowered to match the same brightness\n")
+	TEXT("<0 for a anisomorphic lens flare look (X only)\n")
+	TEXT(" 0 off (default)\n")
+	TEXT(">0 for a cross look (X and Y)"),
+	ECVF_RenderThreadSafe);
+
 IMPLEMENT_SHADER_TYPE(,FPostProcessVS,TEXT("PostProcessBloom"),TEXT("MainPostprocessCommonVS"),SF_Vertex);
 
 // -------------------------------------------------------
@@ -138,6 +148,7 @@ static FRenderingCompositeOutputRef RenderHalfResBloomThreshold(FPostprocessCont
 
 
 // 2 pass Gaussian blur using uni-linear filtering
+// @param CrossCenterWeight see r.Bloom.Cross (positive for X and Y, otherwise for X only)
 static FRenderingCompositeOutputRef RenderGaussianBlur(
 	FPostprocessContext& Context,
 	const TCHAR* DebugNameX,
@@ -145,16 +156,22 @@ static FRenderingCompositeOutputRef RenderGaussianBlur(
 	const FRenderingCompositeOutputRef& Input,
 	float SizeScale,
 	FLinearColor Tint = FLinearColor::White,
-	const FRenderingCompositeOutputRef Additive = FRenderingCompositeOutputRef())
+	const FRenderingCompositeOutputRef Additive = FRenderingCompositeOutputRef(),
+	float CrossCenterWeight = 0.0f)
 {
 	// Gaussian blur in x
 	FRCPassPostProcessWeightedSampleSum* PostProcessBlurX = Context.Graph.RegisterPass(new(FMemStack::Get()) FRCPassPostProcessWeightedSampleSum(EFS_Horiz, EFCM_Weighted, SizeScale, DebugNameX));
 	PostProcessBlurX->SetInput(ePId_Input0, Input);
+	if(CrossCenterWeight > 0)
+	{
+		PostProcessBlurX->SetCrossCenterWeight(CrossCenterWeight);
+	}
 
 	// Gaussian blur in y
 	FRCPassPostProcessWeightedSampleSum* PostProcessBlurY = Context.Graph.RegisterPass(new(FMemStack::Get()) FRCPassPostProcessWeightedSampleSum(EFS_Vert, EFCM_Weighted, SizeScale, DebugNameY, Tint));
 	PostProcessBlurY->SetInput(ePId_Input0, FRenderingCompositeOutputRef(PostProcessBlurX));
 	PostProcessBlurY->SetInput(ePId_Input1, Additive);
+	PostProcessBlurY->SetCrossCenterWeight(FMath::Abs(CrossCenterWeight));
 
 	return FRenderingCompositeOutputRef(PostProcessBlurY);
 }
@@ -167,7 +184,9 @@ static FRenderingCompositeOutputRef RenderBloom(
 	FLinearColor Tint = FLinearColor::White,
 	const FRenderingCompositeOutputRef Additive = FRenderingCompositeOutputRef())
 {
-	return RenderGaussianBlur(Context, TEXT("BloomBlurX"), TEXT("BloomBlurY"), PreviousBloom, Size, Tint, Additive);
+	const float CrossBloom = CVarBloomCross.GetValueOnRenderThread();
+
+	return RenderGaussianBlur(Context, TEXT("BloomBlurX"), TEXT("BloomBlurY"), PreviousBloom, Size, Tint, Additive,CrossBloom);
 }
 
 static void AddTonemapper(
