fix for UE-7290

Editor primitives use forward basepass when ES2 feature level is active.

diff --git a/Engine/Shaders/BasePassForForwardShadingPixelShader.usf b/Engine/Shaders/BasePassForForwardShadingPixelShader.usf
--- a/Engine/Shaders/BasePassForForwardShadingPixelShader.usf
+++ b/Engine/Shaders/BasePassForForwardShadingPixelShader.usf
@@ -56,6 +56,43 @@ half3 GetImageBasedReflectionLighting(FMaterialPixelParameters MaterialParameter
 }
 #endif
 
+#if EDITOR_PRIMITIVE_MATERIAL
+bool bEnableEditorPrimitiveDepthTest;
+// depth in the red channel in DeviceZ
+Texture2D		FilteredSceneDepthTexture;
+SamplerState	FilteredSceneDepthTextureSampler;
+#endif
+
+
+#if EDITOR_PRIMITIVE_MATERIAL && (ES2_EMULATION)
+// @return 0:translucent..1:opaque
+float ClipForEditorPrimitives(FMaterialPixelParameters MaterialParameters)
+{
+	float Ret = 1;
+
+	// Depth test manually if compositing editor primitives since the depth buffer is different (MSAA only)
+	BRANCH
+		if (bEnableEditorPrimitiveDepthTest)
+		{
+			bool bIsPerspective = (View.ViewToClip._m33 < 1.0f);
+
+			// dejitter the sample position and make a filtered lookup - for planes this allows to reconstruct a much less jittery depth comparison function, it however doesn't fix silhuetes
+			float DeviceZ = Texture2DSampleLevel(FilteredSceneDepthTexture, FilteredSceneDepthTextureSampler, (MaterialParameters.SVPosition.xy - View.TemporalAAParams.zw) * View.ViewSizeAndSceneTexelSize.zw, 0).r;
+
+			float PixelDeviceZ = MaterialParameters.SVPosition.z;
+
+			// Soft Bias with DeviceZ for best quality
+			const float DeviceDepthFade = 0.00005f;
+
+			// 0.5f is to bias around the actual value, 1 or 0 are another option
+			Ret = saturate(0.5f - (DeviceZ - PixelDeviceZ) / DeviceDepthFade);
+		}
+
+	// Note: multiple returns cause strange HLSL compiler error for CV_Coverage in later code
+	return Ret;
+}
+#endif
+
 half4 IndirectLightingSHCoefficients[3]; 
 half DirectionalLightShadowing;
 
@@ -239,6 +276,10 @@ void Main(
 #endif
 #endif
 
+#if !MATERIAL_SHADINGMODEL_UNLIT && ES2_EMULATION
+	Color = lerp(Color, DiffuseColor + SpecularColor, View.UnlitViewmodeMask);
+#endif
+
 	#if MATERIALBLENDING_TRANSLUCENT
 		OutColor = half4(Color * VertexFog.a + VertexFog.rgb, Opacity);
 	#elif MATERIALBLENDING_ADDITIVE
@@ -277,4 +318,11 @@ void Main(
 			OutColor.rgb = View.HdrMosaic ? HdrMosaic(OutColor.rgb, SvPosition.xy) : OutColor.rgb;
 		#endif
 	#endif
+
+	// Editor primitive depth testing
+	#if EDITOR_PRIMITIVE_MATERIAL && (ES2_EMULATION)
+		// we output premultiplied alpha to we have to darken all 4 channels
+		OutColor *= ClipForEditorPrimitives(MaterialParameters);
+		clip(OutColor.a - GetMaterialOpacityMaskClipValue());
+	#endif
 }
diff --git a/Engine/Source/Runtime/Renderer/Private/BasePassRendering.h b/Engine/Source/Runtime/Renderer/Private/BasePassRendering.h
--- a/Engine/Source/Runtime/Renderer/Private/BasePassRendering.h
+++ b/Engine/Source/Runtime/Renderer/Private/BasePassRendering.h
@@ -9,6 +9,7 @@
 #include "LightMapRendering.h"
 #include "ShaderBaseClasses.h"
 #include "LightGrid.h"
+#include "EditorCompositeParams.h"
 
 /** Whether to allow the indirect lighting cache to be applied to dynamic objects. */
 extern int32 GIndirectLightingCache;
@@ -330,10 +331,7 @@ public:
 	{
 		LightMapPolicyType::PixelParametersType::Bind(Initializer.ParameterMap);
 		TranslucentLightingParameters.Bind(Initializer.ParameterMap);
-		EditorCompositeDepthTestParameter.Bind(Initializer.ParameterMap,TEXT("bEnableEditorPrimitiveDepthTest"));
-		MSAASampleCount.Bind(Initializer.ParameterMap,TEXT("MSAASampleCount"));
-		FilteredSceneDepthTexture.Bind(Initializer.ParameterMap,TEXT("FilteredSceneDepthTexture"));
-		FilteredSceneDepthTextureSampler.Bind(Initializer.ParameterMap,TEXT("FilteredSceneDepthTextureSampler"));
+		EditorCompositeParams.Bind(Initializer.ParameterMap);
 		LightGrid.Bind(Initializer.ParameterMap,TEXT("LightGrid"));
 	}
 	TBasePassPixelShaderBaseType() {}
@@ -374,42 +372,8 @@ public:
 				}
 			}
 		}
-
-#if WITH_EDITOR
-		if ( GMaxRHIFeatureLevel >= ERHIFeatureLevel::SM4 )
-		{
-			if( MaterialResource.IsUsedWithEditorCompositing() )
-			{
-				// Compute parameters for converting from screen space to pixel 
-				FIntRect DestRect = View->ViewRect;
-				FIntPoint ViewportOffset = DestRect.Min;
-				FIntPoint ViewportExtent = DestRect.Size();
-
-				FVector4 ScreenPosToPixelValue(
-					ViewportExtent.X * 0.5f,
-					-ViewportExtent.Y * 0.5f,
-					ViewportExtent.X * 0.5f - 0.5f + ViewportOffset.X,
-					ViewportExtent.Y * 0.5f - 0.5f + ViewportOffset.Y);
-
-				if(FilteredSceneDepthTexture.IsBound())
-				{
-					const FTexture2DRHIRef* DepthTexture = GSceneRenderTargets.GetActualDepthTexture();
-					check(DepthTexture != NULL);
-					SetTextureParameter(
-						RHICmdList,
-						ShaderRHI,
-						FilteredSceneDepthTexture,
-						FilteredSceneDepthTextureSampler,
-						TStaticSamplerState<SF_Bilinear,AM_Clamp,AM_Clamp,AM_Clamp>::GetRHI(),
-						*DepthTexture
-						);
-				}
-
-				SetShaderValue(RHICmdList, ShaderRHI, EditorCompositeDepthTestParameter, bEnableEditorPrimitveDepthTest );
-				SetShaderValue(RHICmdList, ShaderRHI, MSAASampleCount, GSceneRenderTargets.EditorPrimitivesColor ? GSceneRenderTargets.EditorPrimitivesColor->GetDesc().NumSamples : 0 );
-			}
-		}
-#endif
+		
+		EditorCompositeParams.SetParameters(RHICmdList, MaterialResource, View, bEnableEditorPrimitveDepthTest, GetPixelShader());
 	}
 
 	void SetMesh(FRHICommandList& RHICmdList, const FVertexFactory* VertexFactory,const FSceneView& View,const FPrimitiveSceneProxy* Proxy,const FMeshBatchElement& BatchElement, EBlendMode BlendMode)
@@ -428,22 +392,14 @@ public:
 		bool bShaderHasOutdatedParameters = FMeshMaterialShader::Serialize(Ar);
 		LightMapPolicyType::PixelParametersType::Serialize(Ar);
 		Ar << TranslucentLightingParameters;
-		Ar << EditorCompositeDepthTestParameter;
-		Ar << MSAASampleCount;
-		Ar << FilteredSceneDepthTexture;
-		Ar << FilteredSceneDepthTextureSampler;
+ 		Ar << EditorCompositeParams;
 		Ar << LightGrid;
 		return bShaderHasOutdatedParameters;
 	}
 
 private:
 	FTranslucentLightingParameters TranslucentLightingParameters;
-	/** Parameter for whether or not to depth test in the pixel shader for editor primitives */
-	FShaderParameter EditorCompositeDepthTestParameter;
-	FShaderParameter MSAASampleCount;
-	/** Parameter for reading filtered depth values */
-	FShaderResourceParameter FilteredSceneDepthTexture; 
-	FShaderResourceParameter FilteredSceneDepthTextureSampler; 
+	FEditorCompositingParameters EditorCompositeParams;
 	FShaderResourceParameter LightGrid; 
 };
 
diff --git a/Engine/Source/Runtime/Renderer/Private/DeferredShadingRenderer.cpp b/Engine/Source/Runtime/Renderer/Private/DeferredShadingRenderer.cpp
--- a/Engine/Source/Runtime/Renderer/Private/DeferredShadingRenderer.cpp
+++ b/Engine/Source/Runtime/Renderer/Private/DeferredShadingRenderer.cpp
@@ -769,7 +769,7 @@ void FDeferredShadingSceneRenderer::Render(FRHICommandListImmediate& RHICmdList)
 		bIsGBufferCurrent = true;
 	}
 
-	if(bIsWireframe && FDeferredShadingSceneRenderer::ShouldCompositeEditorPrimitives(Views[0]))
+	if(bIsWireframe && FSceneRenderer::ShouldCompositeEditorPrimitives(Views[0]))
 	{
 		// In Editor we want wire frame view modes to be MSAA for better quality. Resolve will be done with EditorPrimitives
 		SetRenderTarget(RHICmdList, GSceneRenderTargets.GetEditorPrimitivesColor(), GSceneRenderTargets.GetEditorPrimitivesDepth());
@@ -1459,26 +1459,3 @@ void FDeferredShadingSceneRenderer::UpdateDownsampledDepthSurface(FRHICommandLis
 		}
 	}
 }
-
-bool FDeferredShadingSceneRenderer::ShouldCompositeEditorPrimitives(const FViewInfo& View)
-{
-	// If the show flag is enabled
-	if(!View.Family->EngineShowFlags.CompositeEditorPrimitives)
-	{
-		return false;
-	}
-
-	if(GIsEditor && View.Family->EngineShowFlags.Wireframe)
-	{
-		// In Editor we want wire frame view modes to be in MSAA
-		return true;
-	}
-
-	// Any elements that needed compositing were drawn then compositing should be done
-	if( View.ViewMeshElements.Num() || View.TopViewMeshElements.Num() || View.BatchedViewElements.HasPrimsToDraw() || View.TopBatchedViewElements.HasPrimsToDraw() || View.VisibleEditorPrimitives.Num() )
-	{
-		return true;
-	}
-
-	return false;
-}
diff --git a/Engine/Source/Runtime/Renderer/Private/DeferredShadingRenderer.h b/Engine/Source/Runtime/Renderer/Private/DeferredShadingRenderer.h
--- a/Engine/Source/Runtime/Renderer/Private/DeferredShadingRenderer.h
+++ b/Engine/Source/Runtime/Renderer/Private/DeferredShadingRenderer.h
@@ -111,14 +111,6 @@ public:
 	void RenderVisualizeTexturePool(FRHICommandListImmediate& RHICmdList);
 #endif
 
-	/**
-	 * Whether or not to composite editor objects onto the scene as a post processing step
-	 *
-	 * @param View The view to test against
-	 *
-	 * @return true if compositing is needed
-	 */
-	static bool ShouldCompositeEditorPrimitives(const FViewInfo& View);
 
 	/** bound shader state for occlusion test prims */
 	static FGlobalBoundShaderState OcclusionTestBoundShaderState;
diff --git a/Engine/Source/Runtime/Renderer/Private/EditorCompositeParams.cpp b/Engine/Source/Runtime/Renderer/Private/EditorCompositeParams.cpp
new file mode 100644
--- /dev/null
+++ b/Engine/Source/Runtime/Renderer/Private/EditorCompositeParams.cpp
@@ -0,0 +1,59 @@
+// Copyright 1998-2015 Epic Games, Inc. All Rights Reserved.
+#include "RendererPrivate.h"
+#include "EditorCompositeParams.h"
+
+FEditorCompositingParameters::FEditorCompositingParameters()
+{
+}
+
+void FEditorCompositingParameters::Bind(const FShaderParameterMap& ParameterMap)
+{
+	EditorCompositeDepthTestParameter.Bind(ParameterMap, TEXT("bEnableEditorPrimitiveDepthTest"));
+	MSAASampleCount.Bind(ParameterMap, TEXT("MSAASampleCount"));
+	FilteredSceneDepthTexture.Bind(ParameterMap, TEXT("FilteredSceneDepthTexture"));
+	FilteredSceneDepthTextureSampler.Bind(ParameterMap, TEXT("FilteredSceneDepthTextureSampler"));
+}
+
+void FEditorCompositingParameters::SetParameters(
+	FRHICommandList& RHICmdList,
+	const FMaterial& MaterialResource,
+	const FSceneView* View,
+	bool bEnableEditorPrimitveDepthTest,
+	const FPixelShaderRHIParamRef ShaderRHI)
+{
+#if WITH_EDITOR
+		if (GMaxRHIFeatureLevel >= ERHIFeatureLevel::SM4)
+		{
+			if (MaterialResource.IsUsedWithEditorCompositing())
+			{
+				// Compute parameters for converting from screen space to pixel 
+				FIntRect DestRect = View->ViewRect;
+				FIntPoint ViewportOffset = DestRect.Min;
+				FIntPoint ViewportExtent = DestRect.Size();
+
+				FVector4 ScreenPosToPixelValue(
+					ViewportExtent.X * 0.5f,
+					-ViewportExtent.Y * 0.5f,
+					ViewportExtent.X * 0.5f - 0.5f + ViewportOffset.X,
+					ViewportExtent.Y * 0.5f - 0.5f + ViewportOffset.Y);
+
+				if (FilteredSceneDepthTexture.IsBound())
+				{
+					const FTexture2DRHIRef* DepthTexture = GSceneRenderTargets.GetActualDepthTexture();
+					check(DepthTexture != NULL);
+					SetTextureParameter(
+						RHICmdList,
+						ShaderRHI,
+						FilteredSceneDepthTexture,
+						FilteredSceneDepthTextureSampler,
+						TStaticSamplerState<SF_Bilinear, AM_Clamp, AM_Clamp, AM_Clamp>::GetRHI(),
+						*DepthTexture
+						);
+				}
+
+				SetShaderValue(RHICmdList, ShaderRHI, EditorCompositeDepthTestParameter, bEnableEditorPrimitveDepthTest);
+				SetShaderValue(RHICmdList, ShaderRHI, MSAASampleCount, GSceneRenderTargets.EditorPrimitivesColor ? GSceneRenderTargets.EditorPrimitivesColor->GetDesc().NumSamples : 0);
+			}
+		}
+#endif
+}
diff --git a/Engine/Source/Runtime/Renderer/Private/EditorCompositeParams.h b/Engine/Source/Runtime/Renderer/Private/EditorCompositeParams.h
new file mode 100644
--- /dev/null
+++ b/Engine/Source/Runtime/Renderer/Private/EditorCompositeParams.h
@@ -0,0 +1,38 @@
+// Copyright 1998-2015 Epic Games, Inc. All Rights Reserved.
+
+/*=============================================================================
+EditorCompositeParams.h: Manages shader parameters required for editor's composited primitives.
+=============================================================================*/
+#pragma once
+
+class FEditorCompositingParameters
+{
+public:
+	FEditorCompositingParameters();
+
+	void Bind(const FShaderParameterMap& ParameterMap);
+	void SetParameters(
+		FRHICommandList& RHICmdList,
+		const FMaterial& MaterialResource,
+		const FSceneView* View,
+		bool bEnableEditorPrimitveDepthTest,
+		const FPixelShaderRHIParamRef ShaderRHI);
+
+	/** Serializer. */
+	friend FArchive& operator << (FArchive& Ar, FEditorCompositingParameters& P)
+	{
+		Ar << P.EditorCompositeDepthTestParameter;
+		Ar << P.MSAASampleCount;
+		Ar << P.FilteredSceneDepthTexture;
+		Ar << P.FilteredSceneDepthTextureSampler;
+		return Ar;
+	}
+
+private:
+	/** Parameter for whether or not to depth test in the pixel shader for editor primitives */
+	FShaderParameter EditorCompositeDepthTestParameter;
+	FShaderParameter MSAASampleCount;
+	/** Parameter for reading filtered depth values */
+	FShaderResourceParameter FilteredSceneDepthTexture;
+	FShaderResourceParameter FilteredSceneDepthTextureSampler;
+};
\ No newline at end of file
diff --git a/Engine/Source/Runtime/Renderer/Private/ForwardBasePassRendering.cpp b/Engine/Source/Runtime/Renderer/Private/ForwardBasePassRendering.cpp
--- a/Engine/Source/Runtime/Renderer/Private/ForwardBasePassRendering.cpp
+++ b/Engine/Source/Runtime/Renderer/Private/ForwardBasePassRendering.cpp
@@ -96,7 +96,8 @@ public:
 				Parameters.TextureMode,
 				Parameters.ShadingModel != MSM_Unlit && Scene && Scene->ShouldRenderSkylight(),
 				false,
-				FeatureLevel
+				FeatureLevel,
+				Parameters.bEditorCompositeDepthTest
 				),
 				FeatureLevel
 			);
@@ -110,6 +111,9 @@ void FBasePassForwardOpaqueDrawingPolicyFactory::AddStaticMesh(FRHICommandList&
 	const FMaterial* Material = StaticMesh->MaterialRenderProxy->GetMaterial(FeatureLevel);
 	const EBlendMode BlendMode = Material->GetBlendMode();
 
+	// Don't composite static meshes
+	const bool bEditorCompositeDepthTest = false;
+
 	// Only draw opaque materials.
 	if( !IsTranslucentBlendMode(BlendMode) )
 	{
@@ -123,7 +127,7 @@ void FBasePassForwardOpaqueDrawingPolicyFactory::AddStaticMesh(FRHICommandList&
 				Material,
 				StaticMesh->PrimitiveSceneInfo->Proxy,
 				true,
-				false,
+				bEditorCompositeDepthTest,
 				ESceneRenderTargetsMode::DontSet,
 				FeatureLevel
 				),
@@ -193,7 +197,8 @@ public:
 			Parameters.TextureMode,
 			Parameters.ShadingModel != MSM_Unlit && Scene && Scene->ShouldRenderSkylight(),
 			View.Family->EngineShowFlags.ShaderComplexity,
-			View.GetFeatureLevel()
+			View.GetFeatureLevel(),
+			Parameters.bEditorCompositeDepthTest
 			);
 		RHICmdList.BuildAndSetLocalBoundShaderState(DrawingPolicy.GetBoundShaderStateInput(View.GetFeatureLevel()));
 		DrawingPolicy.SetSharedState(RHICmdList, &View, typename TBasePassForForwardShadingDrawingPolicy<LightMapPolicyType>::ContextDataType());
@@ -251,7 +256,7 @@ bool FBasePassForwardOpaqueDrawingPolicyFactory::DrawDynamicMesh(
 				Material,
 				PrimitiveSceneProxy,
 				true,
-				false,
+				DrawingContext.bEditorCompositeDepthTest,
 				DrawingContext.TextureMode,
 				FeatureLevel
 				),
@@ -394,7 +399,7 @@ void FForwardShadingSceneRenderer::RenderForwardShadingBasePass(FRHICommandListI
 			SCOPE_CYCLE_COUNTER(STAT_DynamicPrimitiveDrawTime);
 			SCOPED_DRAW_EVENT(RHICmdList, Dynamic);
 
-			FBasePassForwardOpaqueDrawingPolicyFactory::ContextType Context(ESceneRenderTargetsMode::DontSet);
+			FBasePassForwardOpaqueDrawingPolicyFactory::ContextType Context(false, ESceneRenderTargetsMode::DontSet);
 
 			for (int32 MeshBatchIndex = 0; MeshBatchIndex < View.DynamicMeshElements.Num(); MeshBatchIndex++)
 			{
@@ -412,13 +417,13 @@ void FForwardShadingSceneRenderer::RenderForwardShadingBasePass(FRHICommandListI
 			const bool bNeedToSwitchVerticalAxis = RHINeedsToSwitchVerticalAxis(GShaderPlatformForFeatureLevel[FeatureLevel]);
 
 			// Draw the base pass for the view's batched mesh elements.
-			DrawViewElements<FBasePassForwardOpaqueDrawingPolicyFactory>(RHICmdList, View, FBasePassForwardOpaqueDrawingPolicyFactory::ContextType(ESceneRenderTargetsMode::DontSet), SDPG_World, true);
+			DrawViewElements<FBasePassForwardOpaqueDrawingPolicyFactory>(RHICmdList, View, FBasePassForwardOpaqueDrawingPolicyFactory::ContextType(false, ESceneRenderTargetsMode::DontSet), SDPG_World, true);
 			
 			// Draw the view's batched simple elements(lines, sprites, etc).
 			View.BatchedViewElements.Draw(RHICmdList, FeatureLevel, bNeedToSwitchVerticalAxis, View.ViewProjectionMatrix, View.ViewRect.Width(), View.ViewRect.Height(), false);
 
 			// Draw foreground objects last
-			DrawViewElements<FBasePassForwardOpaqueDrawingPolicyFactory>(RHICmdList, View, FBasePassForwardOpaqueDrawingPolicyFactory::ContextType(ESceneRenderTargetsMode::DontSet), SDPG_Foreground, true);
+			DrawViewElements<FBasePassForwardOpaqueDrawingPolicyFactory>(RHICmdList, View, FBasePassForwardOpaqueDrawingPolicyFactory::ContextType(false, ESceneRenderTargetsMode::DontSet), SDPG_Foreground, true);
 
 			// Draw the view's batched simple elements(lines, sprites, etc).
 			View.TopBatchedViewElements.Draw(RHICmdList, FeatureLevel, bNeedToSwitchVerticalAxis, View.ViewProjectionMatrix, View.ViewRect.Width(), View.ViewRect.Height(), false);
diff --git a/Engine/Source/Runtime/Renderer/Private/ForwardBasePassRendering.h b/Engine/Source/Runtime/Renderer/Private/ForwardBasePassRendering.h
--- a/Engine/Source/Runtime/Renderer/Private/ForwardBasePassRendering.h
+++ b/Engine/Source/Runtime/Renderer/Private/ForwardBasePassRendering.h
@@ -8,6 +8,7 @@
 
 #include "LightMapRendering.h"
 #include "ShaderBaseClasses.h"
+#include "EditorCompositeParams.h"
 
 enum EOutputFormat
 {
@@ -142,13 +143,15 @@ public:
 	{
 		LightMapPolicyType::PixelParametersType::Bind(Initializer.ParameterMap);
 		ReflectionCubemap.Bind(Initializer.ParameterMap, TEXT("ReflectionCubemap"));
-		ReflectionSampler.Bind(Initializer.ParameterMap, TEXT("ReflectionCubemapSampler"));		
+		ReflectionSampler.Bind(Initializer.ParameterMap, TEXT("ReflectionCubemapSampler"));
+		EditorCompositeParams.Bind(Initializer.ParameterMap);
 	}
 	TBasePassForForwardShadingPSBaseType() {}
 
-	void SetParameters(FRHICommandList& RHICmdList, const FMaterialRenderProxy* MaterialRenderProxy,const FMaterial& MaterialResource,const FSceneView* View,ESceneRenderTargetsMode::Type TextureMode)
+	void SetParameters(FRHICommandList& RHICmdList, const FMaterialRenderProxy* MaterialRenderProxy, const FMaterial& MaterialResource, const FSceneView* View, ESceneRenderTargetsMode::Type TextureMode, bool bEnableEditorPrimitveDepthTest)
 	{
 		FMeshMaterialShader::SetParameters(RHICmdList, GetPixelShader(),MaterialRenderProxy,MaterialResource,*View,TextureMode);
+		EditorCompositeParams.SetParameters(RHICmdList, MaterialResource, View, bEnableEditorPrimitveDepthTest, GetPixelShader());
 	}
 
 	void SetMesh(FRHICommandList& RHICmdList, const FVertexFactory* VertexFactory,const FSceneView& View,const FPrimitiveSceneProxy* Proxy,const FMeshBatchElement& BatchElement)
@@ -179,13 +182,15 @@ public:
 		bool bShaderHasOutdatedParameters = FMeshMaterialShader::Serialize(Ar);
 		LightMapPolicyType::PixelParametersType::Serialize(Ar);
 		Ar << ReflectionCubemap;
-		Ar << ReflectionSampler;		
+		Ar << ReflectionSampler;
+		Ar << EditorCompositeParams;
 		return bShaderHasOutdatedParameters;
 	}
 
 private:
 	FShaderResourceParameter ReflectionCubemap;
-	FShaderResourceParameter ReflectionSampler;	
+	FShaderResourceParameter ReflectionSampler;
+	FEditorCompositingParameters EditorCompositeParams;
 };
 
 
@@ -257,12 +262,14 @@ public:
 		ESceneRenderTargetsMode::Type InSceneTextureMode,
 		bool bInEnableSkyLight,
 		bool bOverrideWithShaderComplexity,
-		ERHIFeatureLevel::Type FeatureLevel
+		ERHIFeatureLevel::Type FeatureLevel,
+		bool bInEnableEditorPrimitiveDepthTest = false
 		):
 		FMeshDrawingPolicy(InVertexFactory,InMaterialRenderProxy,InMaterialResource,bOverrideWithShaderComplexity),
 		LightMapPolicy(InLightMapPolicy),
 		BlendMode(InBlendMode),
-		SceneTextureMode(InSceneTextureMode)
+		SceneTextureMode(InSceneTextureMode),
+		bEnableEditorPrimitiveDepthTest(bInEnableEditorPrimitiveDepthTest)
 	{
 		if (IsMobileHDR32bpp())
 		{
@@ -304,6 +311,14 @@ public:
 				PixelShader = InMaterialResource.GetShader< TBasePassForForwardShadingPS<LightMapPolicyType, LDR_GAMMA_32, false> >(InVertexFactory->GetType());
 			}			
 		}
+
+#if DO_GUARD_SLOW
+		// Somewhat hacky
+		if (SceneTextureMode == ESceneRenderTargetsMode::DontSet && !bEnableEditorPrimitiveDepthTest && InMaterialResource.IsUsedWithEditorCompositing())
+		{
+			SceneTextureMode = ESceneRenderTargetsMode::DontSetIgnoreBoundByEditorCompositing;
+		}
+#endif
 	}
 
 	// FMeshDrawingPolicy interface.
@@ -342,7 +357,7 @@ public:
 		else
 #endif
 		{
-			PixelShader->SetParameters(RHICmdList, MaterialRenderProxy, *MaterialResource, View, SceneTextureMode);
+			PixelShader->SetParameters(RHICmdList, MaterialRenderProxy, *MaterialResource, View, SceneTextureMode, bEnableEditorPrimitiveDepthTest);
 
 			switch(BlendMode)
 			{
@@ -466,6 +481,8 @@ protected:
 	LightMapPolicyType LightMapPolicy;
 	EBlendMode BlendMode;
 	ESceneRenderTargetsMode::Type SceneTextureMode;
+	/** Whether or not this policy is compositing editor primitives and needs to depth test against the scene geometry in the base pass pixel shader */
+	uint32 bEnableEditorPrimitiveDepthTest : 1;
 };
 
 /**
@@ -478,7 +495,11 @@ public:
 	enum { bAllowSimpleElements = true };
 	struct ContextType 
 	{
-		ContextType(ESceneRenderTargetsMode::Type InTextureMode) :
+		/** Whether or not to perform depth test in the pixel shader */
+		bool bEditorCompositeDepthTest;
+
+		ContextType(bool bInEditorCompositeDepthTest, ESceneRenderTargetsMode::Type InTextureMode) :
+			bEditorCompositeDepthTest(bInEditorCompositeDepthTest),
 			TextureMode(InTextureMode)
 		{}
 
diff --git a/Engine/Source/Runtime/Renderer/Private/ForwardShadingRenderer.cpp b/Engine/Source/Runtime/Renderer/Private/ForwardShadingRenderer.cpp
--- a/Engine/Source/Runtime/Renderer/Private/ForwardShadingRenderer.cpp
+++ b/Engine/Source/Runtime/Renderer/Private/ForwardShadingRenderer.cpp
@@ -13,6 +13,7 @@
 #include "PostProcessMobile.h"
 #include "SceneUtils.h"
 #include "PostProcessUpscale.h"
+#include "PostProcessCompositeEditorPrimitives.h"
 
 uint32 GetShadowQuality();
 
@@ -251,6 +252,15 @@ void FForwardShadingSceneRenderer::SimpleUpscale(FRHICommandListImmediate& RHICm
 		Context.FinalOutput.GetOutput()->RenderTargetDesc = Desc;
 	}
 
+	// Composite editor primitives if we had any to draw and compositing is enabled
+	if (FSceneRenderer::ShouldCompositeEditorPrimitives(View))
+	{
+		FRenderingCompositePass* EditorCompNode = Context.Graph.RegisterPass(new(FMemStack::Get()) FRCPassPostProcessCompositeEditorPrimitives(false));
+		EditorCompNode->SetInput(ePId_Input0, FRenderingCompositeOutputRef(Context.FinalOutput));
+		//Node->SetInput(ePId_Input1, FRenderingCompositeOutputRef(Context.SceneDepth));
+		Context.FinalOutput = FRenderingCompositeOutputRef(EditorCompNode);
+	}
+
 	CompositeContext.Root->AddDependency(Context.FinalOutput);
 	CompositeContext.Process(TEXT("ES2Upscale"));
 }
diff --git a/Engine/Source/Runtime/Renderer/Private/PostProcess/PostProcessCompositeEditorPrimitives.cpp b/Engine/Source/Runtime/Renderer/Private/PostProcess/PostProcessCompositeEditorPrimitives.cpp
--- a/Engine/Source/Runtime/Renderer/Private/PostProcess/PostProcessCompositeEditorPrimitives.cpp
+++ b/Engine/Source/Runtime/Renderer/Private/PostProcess/PostProcessCompositeEditorPrimitives.cpp
@@ -223,7 +223,14 @@ void FRCPassPostProcessCompositeEditorPrimitives::Process(FRenderingCompositePas
 
 		Context.RHICmdList.SetRasterizerState(View.bReverseCulling ? TStaticRasterizerState<FM_Solid, CM_CW>::GetRHI() : TStaticRasterizerState<FM_Solid, CM_CCW>::GetRHI());
 
-		RenderPrimitivesToComposite(Context.RHICmdList, View);
+		if (bDeferredBasePass)
+		{
+			RenderPrimitivesToComposite<FBasePassOpaqueDrawingPolicyFactory>(Context.RHICmdList, View);
+		}
+		else
+		{
+			RenderPrimitivesToComposite<FBasePassForwardOpaqueDrawingPolicyFactory>(Context.RHICmdList, View);
+		}
 
 		GRenderTargetPool.VisualizeTexture.SetCheckPoint(Context.RHICmdList, GSceneRenderTargets.EditorPrimitivesColor);
 	}
@@ -295,6 +302,7 @@ FPooledRenderTargetDesc FRCPassPostProcessCompositeEditorPrimitives::ComputeOutp
 	return Ret;
 }
 
+template<typename TBasePass>
 void FRCPassPostProcessCompositeEditorPrimitives::RenderPrimitivesToComposite(FRHICommandListImmediate& RHICmdList, const FViewInfo& View)
 {
 
@@ -311,7 +319,7 @@ void FRCPassPostProcessCompositeEditorPrimitives::RenderPrimitivesToComposite(FR
 	const bool bNeedToSwitchVerticalAxis = RHINeedsToSwitchVerticalAxis(ShaderPlatform);
 	FTexture2DRHIRef SceneDepth = GSceneRenderTargets.GetSceneDepthTexture();
 
-	FBasePassOpaqueDrawingPolicyFactory::ContextType Context(true, ESceneRenderTargetsMode::SetTextures);
+	typename TBasePass::ContextType Context(true, ESceneRenderTargetsMode::SetTextures);
 
 	for (int32 MeshBatchIndex = 0; MeshBatchIndex < View.DynamicEditorMeshElements.Num(); MeshBatchIndex++)
 	{
@@ -320,14 +328,14 @@ void FRCPassPostProcessCompositeEditorPrimitives::RenderPrimitivesToComposite(FR
 		if (MeshBatchAndRelevance.bHasOpaqueOrMaskedMaterial || View.Family->EngineShowFlags.Wireframe)
 		{
 			const FMeshBatch& MeshBatch = *MeshBatchAndRelevance.Mesh;
-			FBasePassOpaqueDrawingPolicyFactory::DrawDynamicMesh(RHICmdList, View, Context, MeshBatch, false, true, MeshBatchAndRelevance.PrimitiveSceneProxy, MeshBatch.BatchHitProxyId);
+			TBasePass::DrawDynamicMesh(RHICmdList, View, Context, MeshBatch, false, true, MeshBatchAndRelevance.PrimitiveSceneProxy, MeshBatch.BatchHitProxyId);
 		}
 	}
 
 	View.EditorSimpleElementCollector.DrawBatchedElements(RHICmdList, View, GSceneRenderTargets.GetSceneDepthTexture(), EBlendModeFilter::OpaqueAndMasked);
 	
 	// Draw the base pass for the view's batched mesh elements.
-	DrawViewElements<FBasePassOpaqueDrawingPolicyFactory>(RHICmdList, View, FBasePassOpaqueDrawingPolicyFactory::ContextType(bDepthTest, ESceneRenderTargetsMode::SetTextures), SDPG_World, false);
+	DrawViewElements<TBasePass>(RHICmdList, View, typename TBasePass::ContextType(bDepthTest, ESceneRenderTargetsMode::SetTextures), SDPG_World, false);
 
 	// Draw the view's batched simple elements(lines, sprites, etc).
 	View.BatchedViewElements.Draw(RHICmdList, FeatureLevel, bNeedToSwitchVerticalAxis, View.ViewProjectionMatrix, View.ViewRect.Width(), View.ViewRect.Height(), false, 1.0f, &View, SceneDepth);
@@ -339,14 +347,14 @@ void FRCPassPostProcessCompositeEditorPrimitives::RenderPrimitivesToComposite(FR
 
 		RHICmdList.SetDepthStencilState(TStaticDepthStencilState<false, CF_Always>::GetRHI());
 
-		DrawViewElements<FBasePassOpaqueDrawingPolicyFactory>(RHICmdList, View, FBasePassOpaqueDrawingPolicyFactory::ContextType(bDepthTest, ESceneRenderTargetsMode::SetTextures), SDPG_Foreground, false);
+		DrawViewElements<TBasePass>(RHICmdList, View, typename TBasePass::ContextType(bDepthTest, ESceneRenderTargetsMode::SetTextures), SDPG_Foreground, false);
 
 		View.TopBatchedViewElements.Draw(RHICmdList, FeatureLevel, bNeedToSwitchVerticalAxis, View.ViewProjectionMatrix, View.ViewRect.Width(), View.ViewRect.Height(), false);
 
 		// Note, this is a reversed Z depth surface, using CF_GreaterEqual.
 		RHICmdList.SetDepthStencilState(TStaticDepthStencilState<true, CF_GreaterEqual>::GetRHI());
 
-		DrawViewElements<FBasePassOpaqueDrawingPolicyFactory>(RHICmdList, View, FBasePassOpaqueDrawingPolicyFactory::ContextType(bDepthTest, ESceneRenderTargetsMode::SetTextures), SDPG_Foreground, false);
+		DrawViewElements<TBasePass>(RHICmdList, View, typename TBasePass::ContextType(bDepthTest, ESceneRenderTargetsMode::SetTextures), SDPG_Foreground, false);
 
 		View.TopBatchedViewElements.Draw(RHICmdList, FeatureLevel, bNeedToSwitchVerticalAxis, View.ViewProjectionMatrix, View.ViewRect.Width(), View.ViewRect.Height(), false);
 	}
diff --git a/Engine/Source/Runtime/Renderer/Private/PostProcess/PostProcessCompositeEditorPrimitives.h b/Engine/Source/Runtime/Renderer/Private/PostProcess/PostProcessCompositeEditorPrimitives.h
--- a/Engine/Source/Runtime/Renderer/Private/PostProcess/PostProcessCompositeEditorPrimitives.h
+++ b/Engine/Source/Runtime/Renderer/Private/PostProcess/PostProcessCompositeEditorPrimitives.h
@@ -10,6 +10,8 @@
 class FRCPassPostProcessCompositeEditorPrimitives : public TRenderingCompositePassBase<1, 1>
 {
 public:
+	FRCPassPostProcessCompositeEditorPrimitives(bool bInDeferredBasePass) : bDeferredBasePass(bInDeferredBasePass) {}
+
 	// interface FRenderingCompositePass ---------
 	virtual void Process(FRenderingCompositePassContext& Context);
 	virtual void Release() override { delete this; }
@@ -21,6 +23,9 @@ private:
 	 *
 	 * @param View	The view to draw in
 	 */
+	template <typename TBasePass>
 	void RenderPrimitivesToComposite(FRHICommandListImmediate& RHICmdList, const FViewInfo& View);
+
+	bool bDeferredBasePass;
 };
 
diff --git a/Engine/Source/Runtime/Renderer/Private/PostProcess/PostProcessing.cpp b/Engine/Source/Runtime/Renderer/Private/PostProcess/PostProcessing.cpp
--- a/Engine/Source/Runtime/Renderer/Private/PostProcess/PostProcessing.cpp
+++ b/Engine/Source/Runtime/Renderer/Private/PostProcess/PostProcessing.cpp
@@ -1264,9 +1264,9 @@ void FPostProcessing::Process(FRHICommandListImmediate& RHICmdList, FViewInfo& V
 		}
 
 		// Composite editor primitives if we had any to draw and compositing is enabled
-		if( FDeferredShadingSceneRenderer::ShouldCompositeEditorPrimitives(View) )
+		if (FSceneRenderer::ShouldCompositeEditorPrimitives(View))
 		{
-			FRenderingCompositePass* Node = Context.Graph.RegisterPass(new(FMemStack::Get()) FRCPassPostProcessCompositeEditorPrimitives());
+			FRenderingCompositePass* Node = Context.Graph.RegisterPass(new(FMemStack::Get()) FRCPassPostProcessCompositeEditorPrimitives(true));
 			Node->SetInput(ePId_Input0, FRenderingCompositeOutputRef(Context.FinalOutput));
 			//Node->SetInput(ePId_Input1, FRenderingCompositeOutputRef(Context.SceneDepth));
 			Context.FinalOutput = FRenderingCompositeOutputRef(Node);
@@ -1702,10 +1702,9 @@ void FPostProcessing::ProcessES2(FRHICommandListImmediate& RHICmdList, FViewInfo
 		}
 
 		// Composite editor primitives if we had any to draw and compositing is enabled
-		// TODO: Move FDeferredShadingSceneRenderer::ShouldCompositeEditorPrimitives somewhere more generic
-		if (FDeferredShadingSceneRenderer::ShouldCompositeEditorPrimitives(View) && !DofOutput.IsValid())
+		if (FSceneRenderer::ShouldCompositeEditorPrimitives(View) && !DofOutput.IsValid())
 		{
-			FRenderingCompositePass* EditorCompNode = Context.Graph.RegisterPass(new(FMemStack::Get()) FRCPassPostProcessCompositeEditorPrimitives());
+			FRenderingCompositePass* EditorCompNode = Context.Graph.RegisterPass(new(FMemStack::Get()) FRCPassPostProcessCompositeEditorPrimitives(false));
 			EditorCompNode->SetInput(ePId_Input0, FRenderingCompositeOutputRef(Context.FinalOutput));
 			//Node->SetInput(ePId_Input1, FRenderingCompositeOutputRef(Context.SceneDepth));
 			Context.FinalOutput = FRenderingCompositeOutputRef(EditorCompNode);
@@ -1743,10 +1742,10 @@ void FPostProcessing::ProcessES2(FRHICommandListImmediate& RHICmdList, FViewInfo
 
 		{
 			// temporary work around for EditorPrimitives / Dof alpha issues.
-			if (FDeferredShadingSceneRenderer::ShouldCompositeEditorPrimitives(View) && DofOutput.IsValid())
+			if (FSceneRenderer::ShouldCompositeEditorPrimitives(View) && DofOutput.IsValid())
 			{
 				// TODO: combine editor prims in tonemap pass or remove AA jitter
-				FRenderingCompositePass* EditorCompNode = Context.Graph.RegisterPass(new(FMemStack::Get()) FRCPassPostProcessCompositeEditorPrimitives());
+				FRenderingCompositePass* EditorCompNode = Context.Graph.RegisterPass(new(FMemStack::Get()) FRCPassPostProcessCompositeEditorPrimitives(false));
 				EditorCompNode->SetInput(ePId_Input0, FRenderingCompositeOutputRef(Context.FinalOutput));
 				//Node->SetInput(ePId_Input1, FRenderingCompositeOutputRef(Context.SceneDepth));
 				Context.FinalOutput = FRenderingCompositeOutputRef(EditorCompNode);
diff --git a/Engine/Source/Runtime/Renderer/Private/Renderer.cpp b/Engine/Source/Runtime/Renderer/Private/Renderer.cpp
--- a/Engine/Source/Runtime/Renderer/Private/Renderer.cpp
+++ b/Engine/Source/Runtime/Renderer/Private/Renderer.cpp
@@ -79,7 +79,7 @@ void FRendererModule::DrawTileMesh(FRHICommandListImmediate& RHICmdList, const F
 				}
 				else
 				{
-					FBasePassForwardOpaqueDrawingPolicyFactory::DrawDynamicMesh(RHICmdList, View, FBasePassForwardOpaqueDrawingPolicyFactory::ContextType(ESceneRenderTargetsMode::SetTextures), Mesh, false, false, NULL, HitProxyId);
+					FBasePassForwardOpaqueDrawingPolicyFactory::DrawDynamicMesh(RHICmdList, View, FBasePassForwardOpaqueDrawingPolicyFactory::ContextType(false, ESceneRenderTargetsMode::SetTextures), Mesh, false, false, NULL, HitProxyId);
 				}
 			}
 		}	
diff --git a/Engine/Source/Runtime/Renderer/Private/SceneRendering.cpp b/Engine/Source/Runtime/Renderer/Private/SceneRendering.cpp
--- a/Engine/Source/Runtime/Renderer/Private/SceneRendering.cpp
+++ b/Engine/Source/Runtime/Renderer/Private/SceneRendering.cpp
@@ -1168,6 +1168,29 @@ void FSceneRenderer::OnStartFrame()
 	}
 }
 
+bool FSceneRenderer::ShouldCompositeEditorPrimitives(const FViewInfo& View)
+{
+	// If the show flag is enabled
+	if (!View.Family->EngineShowFlags.CompositeEditorPrimitives)
+	{
+		return false;
+	}
+
+	if (GIsEditor && View.Family->EngineShowFlags.Wireframe)
+	{
+		// In Editor we want wire frame view modes to be in MSAA
+		return true;
+	}
+
+	// Any elements that needed compositing were drawn then compositing should be done
+	if (View.ViewMeshElements.Num() || View.TopViewMeshElements.Num() || View.BatchedViewElements.HasPrimsToDraw() || View.TopBatchedViewElements.HasPrimsToDraw() || View.VisibleEditorPrimitives.Num())
+	{
+		return true;
+	}
+
+	return false;
+}
+
 /*-----------------------------------------------------------------------------
 	FRendererModule::BeginRenderingViewFamily
 -----------------------------------------------------------------------------*/
diff --git a/Engine/Source/Runtime/Renderer/Private/SceneRendering.h b/Engine/Source/Runtime/Renderer/Private/SceneRendering.h
--- a/Engine/Source/Runtime/Renderer/Private/SceneRendering.h
+++ b/Engine/Source/Runtime/Renderer/Private/SceneRendering.h
@@ -603,6 +603,15 @@ public:
 
 	bool DoOcclusionQueries(ERHIFeatureLevel::Type InFeatureLevel) const;
 
+	/**
+	* Whether or not to composite editor objects onto the scene as a post processing step
+	*
+	* @param View The view to test against
+	*
+	* @return true if compositing is needed
+	*/
+	static bool ShouldCompositeEditorPrimitives(const FViewInfo& View);
+
 protected:
 
 	// Shared functionality between all scene renderers
