checks for RenderTargetPool reference counting.

diff --git a/Engine/Source/Runtime/Renderer/Private/PostProcess/RenderTargetPool.cpp b/Engine/Source/Runtime/Renderer/Private/PostProcess/RenderTargetPool.cpp
--- a/Engine/Source/Runtime/Renderer/Private/PostProcess/RenderTargetPool.cpp
+++ b/Engine/Source/Runtime/Renderer/Private/PostProcess/RenderTargetPool.cpp
@@ -1210,11 +1210,13 @@ void FRenderTargetPool::DumpMemoryUsage(FOutputDevice& OutputDevice)
 
 uint32 FPooledRenderTarget::AddRef() const
 {
+	check(IsInRenderingThread());
 	return uint32(++NumRefs);
 }
 
 uint32 FPooledRenderTarget::Release() const
 {
+	check(IsInRenderingThread());
 	uint32 Refs = uint32(--NumRefs);
 	if(Refs == 0)
 	{
