UE4 - Velocity rendering on non-movable static meshes with WPO

diff --git a/Engine/Source/Runtime/Engine/Classes/Materials/MaterialInterface.h b/Engine/Source/Runtime/Engine/Classes/Materials/MaterialInterface.h
--- a/Engine/Source/Runtime/Engine/Classes/Materials/MaterialInterface.h
+++ b/Engine/Source/Runtime/Engine/Classes/Materials/MaterialInterface.h
@@ -57,6 +57,9 @@ struct ENGINE_API FMaterialRelevance
 	UPROPERTY()
 		uint32 bSubsurfaceProfile : 1;
 
+	UPROPERTY()
+		uint32 bHasWorldPositionOffset : 1;
+
 	/** Default constructor. */
 	FMaterialRelevance()
 		: bOpaque(false)
@@ -66,6 +69,7 @@ struct ENGINE_API FMaterialRelevance
 		, bNormalTranslucency(false)
 		, bDisableDepthTest(false)
 		, bSubsurfaceProfile(false)
+		, bHasWorldPositionOffset(false)
 	{}
 
 	/** Bitwise OR operator.  Sets any relevance bits which are present in either FMaterialRelevance. */
@@ -78,6 +82,7 @@ struct ENGINE_API FMaterialRelevance
 		bNormalTranslucency |= B.bNormalTranslucency;
 		bDisableDepthTest |= B.bDisableDepthTest;
 		bSubsurfaceProfile |= B.bSubsurfaceProfile;
+		bHasWorldPositionOffset |= B.bHasWorldPositionOffset;
 		return *this;
 	}
 
diff --git a/Engine/Source/Runtime/Engine/Private/Materials/MaterialInterface.cpp b/Engine/Source/Runtime/Engine/Private/Materials/MaterialInterface.cpp
--- a/Engine/Source/Runtime/Engine/Private/Materials/MaterialInterface.cpp
+++ b/Engine/Source/Runtime/Engine/Private/Materials/MaterialInterface.cpp
@@ -18,6 +18,7 @@ void FMaterialRelevance::SetPrimitiveViewRelevance(FPrimitiveViewRelevance& OutV
 	OutViewRelevance.bSeparateTranslucencyRelevance = bSeparateTranslucency;
 	OutViewRelevance.bNormalTranslucencyRelevance = bNormalTranslucency;
 	OutViewRelevance.bSubsurfaceProfileRelevance = bSubsurfaceProfile;
+	OutViewRelevance.bHasWorldPositionOffset = bHasWorldPositionOffset;
 }
 
 //////////////////////////////////////////////////////////////////////////
@@ -57,6 +58,7 @@ FMaterialRelevance UMaterialInterface::GetRelevance_Internal(const UMaterial* Ma
 		MaterialRelevance.bNormalTranslucency = bIsTranslucent && !Material->bEnableSeparateTranslucency;
 		MaterialRelevance.bDisableDepthTest = bIsTranslucent && Material->bDisableDepthTest;
 		MaterialRelevance.bSubsurfaceProfile = (Material->MaterialDomain == MD_Surface) && !bIsTranslucent && (ShadingModel == MSM_SubsurfaceProfile);
+		MaterialRelevance.bHasWorldPositionOffset = MaterialResource->MaterialModifiesMeshPosition_GameThread();
 
 		return MaterialRelevance;
 	}
diff --git a/Engine/Source/Runtime/Engine/Public/PrimitiveViewRelevance.h b/Engine/Source/Runtime/Engine/Public/PrimitiveViewRelevance.h
--- a/Engine/Source/Runtime/Engine/Public/PrimitiveViewRelevance.h
+++ b/Engine/Source/Runtime/Engine/Public/PrimitiveViewRelevance.h
@@ -39,6 +39,8 @@ struct FPrimitiveViewRelevance
 	uint32 bNormalTranslucencyRelevance : 1;
 	/** The primitive has one or more elements that have the SubsurfaceProfile shading model. */
 	uint32 bSubsurfaceProfileRelevance : 1;
+	// The primitive has one or more elements that have World Position Offset.
+	uint32 bHasWorldPositionOffset : 1;
 
 	/** 
 	 * Whether this primitive view relevance has been initialized this frame.  
@@ -69,6 +71,7 @@ struct FPrimitiveViewRelevance
 		bSeparateTranslucencyRelevance(false),
 		bNormalTranslucencyRelevance(false),
 		bSubsurfaceProfileRelevance(false),
+		bHasWorldPositionOffset(false),
 		bInitializedThisFrame(false)
 	{}
 
@@ -91,6 +94,7 @@ struct FPrimitiveViewRelevance
 		bNormalTranslucencyRelevance |= B.bNormalTranslucencyRelevance != 0;
 		bInitializedThisFrame |= B.bInitializedThisFrame;
 		bSubsurfaceProfileRelevance |= B.bSubsurfaceProfileRelevance != 0;
+		bHasWorldPositionOffset |= B.bHasWorldPositionOffset != 0;
 		return *this;
 	}
 
diff --git a/Engine/Source/Runtime/Renderer/Private/PrimitiveSceneInfo.cpp b/Engine/Source/Runtime/Renderer/Private/PrimitiveSceneInfo.cpp
--- a/Engine/Source/Runtime/Renderer/Private/PrimitiveSceneInfo.cpp
+++ b/Engine/Source/Runtime/Renderer/Private/PrimitiveSceneInfo.cpp
@@ -420,21 +420,21 @@ bool FPrimitiveSceneInfo::ShouldRenderVelocity(const FViewInfo& View, bool bChec
 		const bool bVisible = View.PrimitiveVisibilityMap[PrimitiveId];
 
 		// Only render if visible.
-		if(!bVisible)
+		if (!bVisible)
 		{
 			return false;
 		}
 	}
-	// Used to determine whether object is movable or not.
-	if(!Proxy->IsMovable())
+
+	const FPrimitiveViewRelevance& PrimitiveViewRelevance = View.PrimitiveViewRelevanceMap[PrimitiveId];
+
+	if (!Proxy->IsMovable() && !PrimitiveViewRelevance.bHasWorldPositionOffset)
 	{
 		return false;
 	}
 
-	const FPrimitiveViewRelevance& PrimitiveViewRelevance = View.PrimitiveViewRelevanceMap[PrimitiveId];
-
 	// !Skip translucent objects as they don't support velocities and in the case of particles have a significant CPU overhead.
-	if(!PrimitiveViewRelevance.bOpaqueRelevance || !PrimitiveViewRelevance.bRenderInMainPass)
+	if (!PrimitiveViewRelevance.bOpaqueRelevance || !PrimitiveViewRelevance.bRenderInMainPass)
 	{
 		return false;
 	}
@@ -446,20 +446,21 @@ bool FPrimitiveSceneInfo::ShouldRenderVelocity(const FViewInfo& View, bool bChec
 	float MinScreenRadiusForVelocityPassSquared = FMath::Square(MinScreenRadiusForVelocityPass);
 
 	// Skip primitives that only cover a small amount of screenspace, motion blur on them won't be noticeable.
-	if(FMath::Square(Proxy->GetBounds().SphereRadius) <= MinScreenRadiusForVelocityPassSquared * LODFactorDistanceSquared)
+	if (FMath::Square(Proxy->GetBounds().SphereRadius) <= MinScreenRadiusForVelocityPassSquared * LODFactorDistanceSquared)
 	{
 		return false;
 	}
 
 	// Only render primitives with velocity.
-	if(!FVelocityDrawingPolicy::HasVelocity(View, this))
+	if (!FVelocityDrawingPolicy::HasVelocity(View, this, PrimitiveViewRelevance.bHasWorldPositionOffset))
 	{
 		return false;
 	}
 
 	return true;
 }
+
 void FPrimitiveSceneInfo::ApplyWorldOffset(FVector InOffset)
 {
 	Proxy->ApplyWorldOffset(InOffset);
-}
\ No newline at end of file
+}
diff --git a/Engine/Source/Runtime/Renderer/Private/VelocityRendering.cpp b/Engine/Source/Runtime/Renderer/Private/VelocityRendering.cpp
--- a/Engine/Source/Runtime/Renderer/Private/VelocityRendering.cpp
+++ b/Engine/Source/Runtime/Renderer/Private/VelocityRendering.cpp
@@ -273,41 +273,42 @@ void FVelocityDrawingPolicy::SetMeshRenderState(
 	FMeshDrawingPolicy::SetMeshRenderState(RHICmdList, View, PrimitiveSceneProxy, Mesh, BatchElementIndex, bBackFace, ElementData, PolicyContext);
 }
 
-bool FVelocityDrawingPolicy::HasVelocity(const FViewInfo& View, const FPrimitiveSceneInfo* PrimitiveSceneInfo)
+bool FVelocityDrawingPolicy::HasVelocity(const FViewInfo& View, const FPrimitiveSceneInfo* PrimitiveSceneInfo, bool bHasWorldPositionOffset)
 {
 	checkSlow(IsInParallelRenderingThread());
 
 	// No velocity if motionblur is off, or if it's a non-moving object (treat as background in that case)
-	if(View.bCameraCut || !PrimitiveSceneInfo->Proxy->IsMovable())
+	if (View.bCameraCut || (!PrimitiveSceneInfo->Proxy->IsMovable() && !bHasWorldPositionOffset))
 	{
 		return false;
 	}
 
-	if(PrimitiveSceneInfo->Proxy->AlwaysHasVelocity())
+	if (PrimitiveSceneInfo->Proxy->AlwaysHasVelocity())
 	{
 		return true;
 	}
 
-	if(PrimitiveSceneInfo->bVelocityIsSupressed)
+	if (PrimitiveSceneInfo->bVelocityIsSupressed)
 	{
 		return false;
 	}
 
 	// check if the primitive has moved
+	if (!bHasWorldPositionOffset)
 	{
 		FMatrix PreviousLocalToWorld;
 
 		// hack
 		FScene* Scene = PrimitiveSceneInfo->Scene;
 
-		if(Scene->MotionBlurInfoData.GetPrimitiveMotionBlurInfo(PrimitiveSceneInfo, PreviousLocalToWorld))
+		if (Scene->MotionBlurInfoData.GetPrimitiveMotionBlurInfo(PrimitiveSceneInfo, PreviousLocalToWorld))
 		{
 			check(PrimitiveSceneInfo->Proxy);
 
 			const FMatrix& LocalToWorld = PrimitiveSceneInfo->Proxy->GetLocalToWorld();
 
 			// Hasn't moved (treat as background by not rendering any special velocities)?
-			if(LocalToWorld.Equals(PreviousLocalToWorld, 0.0001f))
+			if (LocalToWorld.Equals(PreviousLocalToWorld, 0.0001f))
 			{
 				return false;
 			}
@@ -349,12 +350,13 @@ int32 Compare(const FVelocityDrawingPolicy& A,const FVelocityDrawingPolicy& B)
 
 void FVelocityDrawingPolicyFactory::AddStaticMesh(FScene* Scene, FStaticMesh* StaticMesh, ContextType)
 {
-	// Velocity only needs to be directly rendered for movable meshes.
-	if(StaticMesh->PrimitiveSceneInfo->Proxy->IsMovable())
+	const auto FeatureLevel = Scene->GetFeatureLevel();
+	const FMaterialRenderProxy* MaterialRenderProxy = StaticMesh->MaterialRenderProxy;
+	const FMaterial* Material = MaterialRenderProxy->GetMaterial(FeatureLevel);
+
+	// Velocity only needs to be directly rendered for movable meshes, or those that modify WPO
+	if (StaticMesh->PrimitiveSceneInfo->Proxy->IsMovable() || Material->MaterialModifiesMeshPosition_RenderThread())
 	{
-		const auto FeatureLevel = Scene->GetFeatureLevel();
-	    const FMaterialRenderProxy* MaterialRenderProxy = StaticMesh->MaterialRenderProxy;
-		const FMaterial* Material = MaterialRenderProxy->GetMaterial(FeatureLevel);
 	    EBlendMode BlendMode = Material->GetBlendMode();
 	    if(BlendMode == BLEND_Opaque || BlendMode == BLEND_Masked)
 	    {
@@ -368,7 +370,7 @@ void FVelocityDrawingPolicyFactory::AddStaticMesh(FScene* Scene, FStaticMesh* St
 			if (DrawingPolicy.SupportsVelocity())
 			{
 				// Add the static mesh to the depth-only draw list.
-				Scene->VelocityDrawList.AddMesh(StaticMesh, FVelocityDrawingPolicy::ElementDataType(), DrawingPolicy, Scene->GetFeatureLevel());
+				Scene->VelocityDrawList.AddMesh(StaticMesh, FVelocityDrawingPolicy::ElementDataType(), DrawingPolicy, FeatureLevel);
 			}
 	    }
 	}
@@ -404,7 +406,7 @@ bool FVelocityDrawingPolicyFactory::DrawDynamicMesh(
 		FVelocityDrawingPolicy DrawingPolicy(Mesh.VertexFactory, MaterialRenderProxy, *MaterialRenderProxy->GetMaterial(FeatureLevel), FeatureLevel);
 		if(DrawingPolicy.SupportsVelocity())
 		{			
-			RHICmdList.BuildAndSetLocalBoundShaderState(DrawingPolicy.GetBoundShaderStateInput(View.GetFeatureLevel()));
+			RHICmdList.BuildAndSetLocalBoundShaderState(DrawingPolicy.GetBoundShaderStateInput(FeatureLevel));
 			DrawingPolicy.SetSharedState(RHICmdList, &View, FVelocityDrawingPolicy::ContextDataType());
 			for (int32 BatchElementIndex = 0; BatchElementIndex < Mesh.Elements.Num(); ++BatchElementIndex)
 			{
@@ -413,12 +415,9 @@ bool FVelocityDrawingPolicyFactory::DrawDynamicMesh(
 			}
 			return true;
 		}
-		return false;
-	}
-	else
-	{
-		return false;
 	}
+
+	return false;
 }
 
 
diff --git a/Engine/Source/Runtime/Renderer/Private/VelocityRendering.h b/Engine/Source/Runtime/Renderer/Private/VelocityRendering.h
--- a/Engine/Source/Runtime/Renderer/Private/VelocityRendering.h
+++ b/Engine/Source/Runtime/Renderer/Private/VelocityRendering.h
@@ -47,7 +47,7 @@ public:
 	bool SupportsVelocity( ) const;
 
 	/** Determines whether this primitive has motionblur velocity to render */
-	static bool HasVelocity(const FViewInfo& View, const FPrimitiveSceneInfo* PrimitiveSceneInfo);
+	static bool HasVelocity(const FViewInfo& View, const FPrimitiveSceneInfo* PrimitiveSceneInfo, bool bHasWorldPositionOffset);
 
 private:
 	class FVelocityVS* VertexShader;
