added feature to tweak roughness with console commands (for testing, not in shipping)

r.Roughness.Min
  Allows quick material test by remapping the roughness at 0 to a new value (0..1), Only for non shipping built!
  0: (default)

r.Roughness.Max
Allows quick material test by remapping the roughness at 1 to a new value (0..1), Only for non shipping built!
  1: (default)

diff --git a/Engine/Source/Runtime/Renderer/Private/SceneRendering.cpp b/Engine/Source/Runtime/Renderer/Private/SceneRendering.cpp
--- a/Engine/Source/Runtime/Renderer/Private/SceneRendering.cpp
+++ b/Engine/Source/Runtime/Renderer/Private/SceneRendering.cpp
@@ -63,6 +63,21 @@ static TAutoConsoleVariable<float> CVarGeneralPurposeTweak(
 	TEXT("Example usage: Multiplier on some value to tweak, toggle to switch between different algorithms (Default: 1.0)\n")
 	TEXT("DON'T USE THIS FOR ANYTHING THAT IS CHECKED IN. Compiled out in SHIPPING to make cheating a bit harder."),
 	ECVF_RenderThreadSafe);
+
+static TAutoConsoleVariable<float> CVarRoughnessMin(
+	TEXT("r.Roughness.Min"),
+	0.0f,
+	TEXT("Allows quick material test by remapping the roughness at 0 to a new value (0..1), Only for non shipping built!\n")
+	TEXT("0: (default)"),
+	ECVF_Cheat | ECVF_RenderThreadSafe
+	);
+static TAutoConsoleVariable<float> CVarRoughnessMax(
+	TEXT("r.Roughness.Max"),
+	1.0,
+	TEXT("Allows quick material test by remapping the roughness at 1 to a new value (0..1), Only for non shipping built!\n")
+	TEXT("1: (default)"),
+	ECVF_Cheat | ECVF_RenderThreadSafe
+	);
 #endif
 
 /**
@@ -306,6 +321,19 @@ TUniformBufferRef<FViewUniformShaderParameters> FViewInfo::CreateUniformBuffer(
 		(ViewRect.Width() / 2.0f + ViewRect.Min.X) * InvBufferSizeX
 		);
 	
+	FVector2D LocalRoughnessOverrideParameter = RoughnessOverrideParameter;
+
+#if !(UE_BUILD_SHIPPING || UE_BUILD_TEST)
+	{
+		// RoughnessOverrideParameter is applied a * y + x
+		float MaxValue = LocalRoughnessOverrideParameter.X + LocalRoughnessOverrideParameter.Y;
+
+		LocalRoughnessOverrideParameter.X = FMath::Max(RoughnessOverrideParameter.X, CVarRoughnessMin.GetValueOnRenderThread());
+		MaxValue = FMath::Min(MaxValue, CVarRoughnessMax.GetValueOnRenderThread());
+		LocalRoughnessOverrideParameter.Y = MaxValue - LocalRoughnessOverrideParameter.X;
+	}
+#endif
+
 	const bool bIsUnlitView = !Family->EngineShowFlags.Lighting;
 
 	// Create the view's uniform buffer.
@@ -329,7 +357,7 @@ TUniformBufferRef<FViewUniformShaderParameters> FViewInfo::CreateUniformBuffer(
 	ViewUniformShaderParameters.DiffuseOverrideParameter = DiffuseOverrideParameter;
 	ViewUniformShaderParameters.SpecularOverrideParameter = SpecularOverrideParameter;
 	ViewUniformShaderParameters.NormalOverrideParameter = NormalOverrideParameter;
-	ViewUniformShaderParameters.RoughnessOverrideParameter = RoughnessOverrideParameter;
+	ViewUniformShaderParameters.RoughnessOverrideParameter = LocalRoughnessOverrideParameter;
 	ViewUniformShaderParameters.PrevFrameGameTime = Family->CurrentWorldTime - Family->DeltaWorldTime;
 	ViewUniformShaderParameters.PrevFrameRealTime = Family->CurrentRealTime - Family->DeltaWorldTime;
 	ViewUniformShaderParameters.PreViewTranslation = ViewMatrices.PreViewTranslation;
