added r.DiffuseColor.Min and r.DiffuseColor.Max for tweak experiments, will be changed to r.BaseColor and r.Metallic

diff --git a/Engine/Source/Runtime/Renderer/Private/SceneRendering.cpp b/Engine/Source/Runtime/Renderer/Private/SceneRendering.cpp
--- a/Engine/Source/Runtime/Renderer/Private/SceneRendering.cpp
+++ b/Engine/Source/Runtime/Renderer/Private/SceneRendering.cpp
@@ -64,6 +64,20 @@ static TAutoConsoleVariable<float> CVarGeneralPurposeTweak(
 	TEXT("DON'T USE THIS FOR ANYTHING THAT IS CHECKED IN. Compiled out in SHIPPING to make cheating a bit harder."),
 	ECVF_RenderThreadSafe);
 
+static TAutoConsoleVariable<float> CVarDiffuseColorMin(
+	TEXT("r.DiffuseColor.Min"),
+	0.0f,
+	TEXT("Allows quick material test by remapping the diffuse color at 1 to a new value (0..1), Only for non shipping built!\n")
+	TEXT("1: (default)"),
+	ECVF_Cheat | ECVF_RenderThreadSafe
+	);
+static TAutoConsoleVariable<float> CVarDiffuseColorMax(
+	TEXT("r.DiffuseColor.Max"),
+	1.0f,
+	TEXT("Allows quick material test by remapping the diffuse color at 1 to a new value (0..1), Only for non shipping built!\n")
+	TEXT("1: (default)"),
+	ECVF_Cheat | ECVF_RenderThreadSafe
+	);
 static TAutoConsoleVariable<float> CVarRoughnessMin(
 	TEXT("r.Roughness.Min"),
 	0.0f,
@@ -73,7 +87,7 @@ static TAutoConsoleVariable<float> CVarRoughnessMin(
 	);
 static TAutoConsoleVariable<float> CVarRoughnessMax(
 	TEXT("r.Roughness.Max"),
-	1.0,
+	1.0f,
 	TEXT("Allows quick material test by remapping the roughness at 1 to a new value (0..1), Only for non shipping built!\n")
 	TEXT("1: (default)"),
 	ECVF_Cheat | ECVF_RenderThreadSafe
@@ -321,17 +335,32 @@ TUniformBufferRef<FViewUniformShaderParameters> FViewInfo::CreateUniformBuffer(
 		(ViewRect.Width() / 2.0f + ViewRect.Min.X) * InvBufferSizeX
 		);
 	
+	FVector4 LocalDiffuseOverrideParameter = DiffuseOverrideParameter;
 	FVector2D LocalRoughnessOverrideParameter = RoughnessOverrideParameter;
 
 #if !(UE_BUILD_SHIPPING || UE_BUILD_TEST)
 	{
-		// RoughnessOverrideParameter is applied a * y + x
-		float MaxValue = LocalRoughnessOverrideParameter.X + LocalRoughnessOverrideParameter.Y;
+		// assuming we have no color in the multipliers
+		float MinValue = LocalDiffuseOverrideParameter.X;
+		float MaxValue = MinValue + LocalDiffuseOverrideParameter.W;
+
+		float NewMinValue = FMath::Max(MinValue, CVarDiffuseColorMin.GetValueOnRenderThread());
+		float NewMaxValue = FMath::Min(MaxValue, CVarDiffuseColorMax.GetValueOnRenderThread());
 
-		LocalRoughnessOverrideParameter.X = FMath::Max(RoughnessOverrideParameter.X, CVarRoughnessMin.GetValueOnRenderThread());
-		MaxValue = FMath::Min(MaxValue, CVarRoughnessMax.GetValueOnRenderThread());
-		LocalRoughnessOverrideParameter.Y = MaxValue - LocalRoughnessOverrideParameter.X;
+		LocalDiffuseOverrideParameter.X = LocalDiffuseOverrideParameter.Y = LocalDiffuseOverrideParameter.Z = NewMinValue;
+		LocalDiffuseOverrideParameter.W = NewMaxValue - NewMinValue;
 	}
+	{
+		float MinValue = LocalRoughnessOverrideParameter.X;
+		float MaxValue = MinValue + LocalRoughnessOverrideParameter.Y;
+
+		float NewMinValue = FMath::Max(MinValue, CVarRoughnessMin.GetValueOnRenderThread());
+		float NewMaxValue = FMath::Min(MaxValue, CVarRoughnessMax.GetValueOnRenderThread());
+
+		LocalRoughnessOverrideParameter.X = NewMinValue;
+		LocalRoughnessOverrideParameter.Y = NewMaxValue - NewMinValue;
+	}
+
 #endif
 
 	const bool bIsUnlitView = !Family->EngineShowFlags.Lighting;
@@ -354,7 +383,7 @@ TUniformBufferRef<FViewUniformShaderParameters> FViewInfo::CreateUniformBuffer(
 	ViewUniformShaderParameters.ViewSizeAndSceneTexelSize = FVector4(ViewRect.Width(), ViewRect.Height(), InvBufferSizeX, InvBufferSizeY);
 	ViewUniformShaderParameters.ViewOrigin = ViewMatrices.ViewOrigin;
 	ViewUniformShaderParameters.TranslatedViewOrigin = ViewMatrices.ViewOrigin + ViewMatrices.PreViewTranslation;
-	ViewUniformShaderParameters.DiffuseOverrideParameter = DiffuseOverrideParameter;
+	ViewUniformShaderParameters.DiffuseOverrideParameter = LocalDiffuseOverrideParameter;
 	ViewUniformShaderParameters.SpecularOverrideParameter = SpecularOverrideParameter;
 	ViewUniformShaderParameters.NormalOverrideParameter = NormalOverrideParameter;
 	ViewUniformShaderParameters.RoughnessOverrideParameter = LocalRoughnessOverrideParameter;
