added comment

diff --git a/Engine/Source/Runtime/Renderer/Private/SceneRendering.cpp b/Engine/Source/Runtime/Renderer/Private/SceneRendering.cpp
--- a/Engine/Source/Runtime/Renderer/Private/SceneRendering.cpp
+++ b/Engine/Source/Runtime/Renderer/Private/SceneRendering.cpp
@@ -64,6 +64,7 @@ static TAutoConsoleVariable<float> CVarGeneralPurposeTweak(
 	TEXT("DON'T USE THIS FOR ANYTHING THAT IS CHECKED IN. Compiled out in SHIPPING to make cheating a bit harder."),
 	ECVF_RenderThreadSafe);
 
+// should be changed to BaseColor and Metallic, since some time now UE4 is not using DiffuseColor and SpecularColor any more
 static TAutoConsoleVariable<float> CVarDiffuseColorMin(
 	TEXT("r.DiffuseColor.Min"),
 	0.0f,
