Use runtime check instead of PLATFORM_MAC define when generating post process shaders to solve issues with shaders cooked for Mac on PC

diff --git a/Engine/Source/Runtime/Renderer/Private/PostProcess/PostProcessCombineLUTs.cpp b/Engine/Source/Runtime/Renderer/Private/PostProcess/PostProcessCombineLUTs.cpp
--- a/Engine/Source/Runtime/Renderer/Private/PostProcess/PostProcessCombineLUTs.cpp
+++ b/Engine/Source/Runtime/Renderer/Private/PostProcess/PostProcessCombineLUTs.cpp
@@ -40,7 +40,7 @@ static TAutoConsoleVariable<float> CVarColorMax(
 static bool UseVolumeTextureLUT(EShaderPlatform Platform) 
 {
 	// @todo Mac OS X: in order to share precompiled shaders between GL 3.3 & GL 4.1 devices we mustn't use volume-texture rendering as it isn't universally supported.
-	return (IsFeatureLevelSupported(Platform,ERHIFeatureLevel::SM4) && GSupportsVolumeTextureRendering && !PLATFORM_MAC && RHISupportsGeometryShaders(Platform));
+	return (IsFeatureLevelSupported(Platform,ERHIFeatureLevel::SM4) && GSupportsVolumeTextureRendering && Platform != EShaderPlatform::SP_OPENGL_SM4_MAC && RHISupportsGeometryShaders(Platform));
 }
 
 // including the neutral one at index 0
diff --git a/Engine/Source/Runtime/Renderer/Private/PostProcess/PostProcessTonemap.cpp b/Engine/Source/Runtime/Renderer/Private/PostProcess/PostProcessTonemap.cpp
--- a/Engine/Source/Runtime/Renderer/Private/PostProcess/PostProcessTonemap.cpp
+++ b/Engine/Source/Runtime/Renderer/Private/PostProcess/PostProcessTonemap.cpp
@@ -964,7 +964,7 @@ class FPostProcessTonemapPS : public FGlobalShader
 		OutEnvironment.SetDefine(TEXT("USE_VIGNETTE_COLOR"),     TonemapperIsDefined(ConfigBitmask, TonemapperVignetteColor));
 		OutEnvironment.SetDefine(TEXT("USE_COLOR_FRINGE"),		 TonemapperIsDefined(ConfigBitmask, TonemapperColorFringe));
 		// @todo Mac OS X: in order to share precompiled shaders between GL 3.3 & GL 4.1 devices we mustn't use volume-texture rendering as it isn't universally supported.
-		OutEnvironment.SetDefine(TEXT("USE_VOLUME_LUT"),		 (IsFeatureLevelSupported(Platform,ERHIFeatureLevel::SM4) && GSupportsVolumeTextureRendering && !PLATFORM_MAC));
+		OutEnvironment.SetDefine(TEXT("USE_VOLUME_LUT"),		 (IsFeatureLevelSupported(Platform,ERHIFeatureLevel::SM4) && GSupportsVolumeTextureRendering && Platform != EShaderPlatform::SP_OPENGL_SM4_MAC));
 		OutEnvironment.SetDefine(TEXT("USE_COLOR_GRADING"),		 TonemapperIsDefined(ConfigBitmask, TonemapperColorGrading));
 
 		if( !IsFeatureLevelSupported(Platform,ERHIFeatureLevel::SM5) )
