Disable WPO motion vectors temporarily

diff --git a/Engine/Source/Runtime/Renderer/Private/VelocityRendering.cpp b/Engine/Source/Runtime/Renderer/Private/VelocityRendering.cpp
--- a/Engine/Source/Runtime/Renderer/Private/VelocityRendering.cpp
+++ b/Engine/Source/Runtime/Renderer/Private/VelocityRendering.cpp
@@ -294,7 +294,7 @@ bool FVelocityDrawingPolicy::HasVelocity(const FViewInfo& View, const FPrimitive
 	}
 
 	// check if the primitive has moved
-	if (!bHasWorldPositionOffset)
+	//if (!bHasWorldPositionOffset)
 	{
 		FMatrix PreviousLocalToWorld;
 
