UE-7290
Prevent platforms < SM4 from compiling deferred basepass shaders.

diff --git a/Engine/Source/Runtime/Renderer/Private/BasePassRendering.h b/Engine/Source/Runtime/Renderer/Private/BasePassRendering.h
--- a/Engine/Source/Runtime/Renderer/Private/BasePassRendering.h
+++ b/Engine/Source/Runtime/Renderer/Private/BasePassRendering.h
@@ -103,7 +103,7 @@ public:
 	{
 		bool bShouldCache = TBasePassVertexShaderBaseType<LightMapPolicyType>::ShouldCache(Platform, Material, VertexFactoryType);
 		return bShouldCache 
-			&& (IsFeatureLevelSupported(Platform, ERHIFeatureLevel::SM4) || (IsPCPlatform(Platform)))
+			&& (IsFeatureLevelSupported(Platform, ERHIFeatureLevel::SM4))
 			&& (!bEnableAtmosphericFog || IsFeatureLevelSupported(Platform, ERHIFeatureLevel::SM4));
 	}
 
@@ -416,7 +416,7 @@ public:
 		const bool bCacheShaders = !bEnableSkyLight || (Material->GetShadingModel() != MSM_Unlit);
 
 		return bCacheShaders
-			&& (IsFeatureLevelSupported(Platform, ERHIFeatureLevel::SM4) || (IsPCPlatform(Platform)))
+			&& (IsFeatureLevelSupported(Platform, ERHIFeatureLevel::SM4))
 			&& TBasePassPixelShaderBaseType<LightMapPolicyType>::ShouldCache(Platform, Material, VertexFactoryType);
 	}
 
