Fix parallel rendering crash due to TRefCountPtr temporary object in ternary operator.

diff --git a/Engine/Source/Runtime/Renderer/Private/PostProcess/SceneRenderTargets.cpp b/Engine/Source/Runtime/Renderer/Private/PostProcess/SceneRenderTargets.cpp
--- a/Engine/Source/Runtime/Renderer/Private/PostProcess/SceneRenderTargets.cpp
+++ b/Engine/Source/Runtime/Renderer/Private/PostProcess/SceneRenderTargets.cpp
@@ -1885,7 +1885,7 @@ void FDeferredPixelShaderParameters::Set(FRHICommandList& RHICmdList, const Shad
 		}
 
 		// if there is no custom depth it's better to have the far distance there
-		IPooledRenderTarget* CustomDepth = GSceneRenderTargets.bCustomDepthIsValid ? GSceneRenderTargets.CustomDepth : 0;
+		IPooledRenderTarget* CustomDepth = GSceneRenderTargets.bCustomDepthIsValid ? GSceneRenderTargets.CustomDepth.GetReference() : 0;
 		if(!CustomDepth)
 		{
 			CustomDepth = GSystemTextures.BlackDummy;
