refactor some shadow code

diff --git a/Engine/Source/Runtime/Engine/Private/DirectionalLightComponent.cpp b/Engine/Source/Runtime/Engine/Private/DirectionalLightComponent.cpp
--- a/Engine/Source/Runtime/Engine/Private/DirectionalLightComponent.cpp
+++ b/Engine/Source/Runtime/Engine/Private/DirectionalLightComponent.cpp
@@ -223,7 +223,6 @@ public:
 
 		const float ShadowExtent = Bounds.W / FMath::Sqrt(3.0f);
 		const FBoxSphereBounds SubjectBounds(Bounds.Center, FVector(ShadowExtent, ShadowExtent, ShadowExtent), Bounds.W);
-		OutInitializer.bDirectionalLight = true;
 		OutInitializer.bOnePassPointLightShadow = false;
 		OutInitializer.PreShadowTranslation = -Bounds.Center;
 		OutInitializer.WorldToLight = FInverseRotationMatrix(FVector(WorldToLight.M[0][0],WorldToLight.M[1][0],WorldToLight.M[2][0]).GetSafeNormal().Rotation());
@@ -251,7 +250,6 @@ public:
 
 		float LpvExtent = LightPropagationVolumeBounds.GetExtent().X; // LPV is a cube, so this should be valid
 
-		OutInitializer.bDirectionalLight = true;
 		OutInitializer.PreShadowTranslation = -LightPropagationVolumeBounds.GetCenter();
 		OutInitializer.WorldToLight = FInverseRotationMatrix(FVector(WorldToLight.M[0][0],WorldToLight.M[1][0],WorldToLight.M[2][0]).GetSafeNormal().Rotation());
 		OutInitializer.Scales = FVector(1.0f,1.0f / LpvExtent,1.0f / LpvExtent);
@@ -297,7 +295,6 @@ public:
 
 	virtual bool GetPerObjectProjectedShadowInitializer(const FBoxSphereBounds& SubjectBounds,FPerObjectProjectedShadowInitializer& OutInitializer) const
 	{
-		OutInitializer.bDirectionalLight = true;
 		OutInitializer.PreShadowTranslation = -SubjectBounds.Origin;
 		OutInitializer.WorldToLight = FInverseRotationMatrix(FVector(WorldToLight.M[0][0],WorldToLight.M[1][0],WorldToLight.M[2][0]).GetSafeNormal().Rotation());
 		OutInitializer.Scales = FVector(1.0f,1.0f / SubjectBounds.SphereRadius,1.0f / SubjectBounds.SphereRadius);
diff --git a/Engine/Source/Runtime/Engine/Private/PointLightComponent.cpp b/Engine/Source/Runtime/Engine/Private/PointLightComponent.cpp
--- a/Engine/Source/Runtime/Engine/Private/PointLightComponent.cpp
+++ b/Engine/Source/Runtime/Engine/Private/PointLightComponent.cpp
@@ -73,7 +73,6 @@ public:
 			&& GAllowPointLightCubemapShadows != 0)
 		{
 			FWholeSceneProjectedShadowInitializer& OutInitializer = *new(OutInitializers) FWholeSceneProjectedShadowInitializer;
-			OutInitializer.bDirectionalLight = false;
 			OutInitializer.bOnePassPointLightShadow = true;
 			OutInitializer.PreShadowTranslation = -GetLightToWorld().GetOrigin();
 			OutInitializer.WorldToLight = GetWorldToLight().RemoveTranslation();
diff --git a/Engine/Source/Runtime/Engine/Private/PointLightSceneProxy.h b/Engine/Source/Runtime/Engine/Private/PointLightSceneProxy.h
--- a/Engine/Source/Runtime/Engine/Private/PointLightSceneProxy.h
+++ b/Engine/Source/Runtime/Engine/Private/PointLightSceneProxy.h
@@ -126,7 +126,6 @@ public:
 			SilhouetteRadius = 1.0f;
 		}
 
-		OutInitializer.bDirectionalLight = false;
 		OutInitializer.PreShadowTranslation = -LightPosition;
 		OutInitializer.WorldToLight = FInverseRotationMatrix((LightVector / LightDistance).Rotation());
 		OutInitializer.Scales = FVector(1.0f,1.0f / SilhouetteRadius,1.0f / SilhouetteRadius);
diff --git a/Engine/Source/Runtime/Engine/Private/SpotLightComponent.cpp b/Engine/Source/Runtime/Engine/Private/SpotLightComponent.cpp
--- a/Engine/Source/Runtime/Engine/Private/SpotLightComponent.cpp
+++ b/Engine/Source/Runtime/Engine/Private/SpotLightComponent.cpp
@@ -122,7 +122,6 @@ public:
 	virtual bool GetWholeSceneProjectedShadowInitializer(const FSceneViewFamily& ViewFamily, TArray<FWholeSceneProjectedShadowInitializer, TInlineAllocator<6> >& OutInitializers) const
 	{
 		FWholeSceneProjectedShadowInitializer& OutInitializer = *new(OutInitializers) FWholeSceneProjectedShadowInitializer;
-		OutInitializer.bDirectionalLight = false;
 		OutInitializer.bOnePassPointLightShadow = false;
 		OutInitializer.PreShadowTranslation = -GetLightToWorld().GetOrigin();
 		OutInitializer.WorldToLight = GetWorldToLight().RemoveTranslation();
diff --git a/Engine/Source/Runtime/Engine/Public/SceneManagement.h b/Engine/Source/Runtime/Engine/Public/SceneManagement.h
--- a/Engine/Source/Runtime/Engine/Public/SceneManagement.h
+++ b/Engine/Source/Runtime/Engine/Public/SceneManagement.h
@@ -550,12 +550,8 @@ public:
 	float MinLightW;
 	float MaxDistanceToCastInLightW;
 
-	/** Whether the shadow is for a directional light. */
-	bool bDirectionalLight;
-
 	/** Default constructor. */
 	FProjectedShadowInitializer()
-	:	bDirectionalLight(false)
 	{}
 };
 
diff --git a/Engine/Source/Runtime/Renderer/Private/SceneOcclusion.cpp b/Engine/Source/Runtime/Renderer/Private/SceneOcclusion.cpp
--- a/Engine/Source/Runtime/Renderer/Private/SceneOcclusion.cpp
+++ b/Engine/Source/Runtime/Renderer/Private/SceneOcclusion.cpp
@@ -468,14 +468,8 @@ static void IssueProjectedShadowOcclusionQuery(FRHICommandListImmediate& RHICmdL
 		// we just copy the indices right in
 		FMemory::Memcpy(Indices, GCubeIndices, sizeof(GCubeIndices));
 
-		FSceneViewState::FProjectedShadowKey Key(
-			ProjectedShadowInfo.ParentSceneInfo ? 
-				ProjectedShadowInfo.ParentSceneInfo->PrimitiveComponentId :
-				FPrimitiveComponentId(),
-			ProjectedShadowInfo.LightSceneInfo->Proxy->GetLightComponent(),
-			ProjectedShadowInfo.SplitIndex,
-			ProjectedShadowInfo.bTranslucentShadow
-			);
+		FSceneViewState::FProjectedShadowKey Key(ProjectedShadowInfo);
+
 		checkSlow(ShadowOcclusionQueryMap.Find(Key) == NULL);
 		ShadowOcclusionQueryMap.Add(Key, ShadowOcclusionQuery);
 
@@ -1135,14 +1129,8 @@ void FDeferredShadingSceneRenderer::BeginOcclusionTests(FRHICommandListImmediate
 								const FRenderQueryRHIRef ShadowOcclusionQuery = ViewState->OcclusionQueryPool.AllocateQuery();
 								RHICmdList.BeginRenderQuery(ShadowOcclusionQuery);
 
-								FSceneViewState::FProjectedShadowKey Key(
-									ProjectedShadowInfo.ParentSceneInfo ? 
-										ProjectedShadowInfo.ParentSceneInfo->PrimitiveComponentId :
-										FPrimitiveComponentId(),
-									ProjectedShadowInfo.LightSceneInfo->Proxy->GetLightComponent(),
-									ProjectedShadowInfo.SplitIndex,
-									ProjectedShadowInfo.bTranslucentShadow
-									);
+								FSceneViewState::FProjectedShadowKey Key(ProjectedShadowInfo);
+
 								checkSlow(ShadowOcclusionQueryMap.Find(Key) == NULL);
 								ShadowOcclusionQueryMap.Add(Key, ShadowOcclusionQuery);
 
diff --git a/Engine/Source/Runtime/Renderer/Private/ScenePrivate.h b/Engine/Source/Runtime/Renderer/Private/ScenePrivate.h
--- a/Engine/Source/Runtime/Renderer/Private/ScenePrivate.h
+++ b/Engine/Source/Runtime/Renderer/Private/ScenePrivate.h
@@ -376,6 +376,14 @@ public:
 			return (PrimitiveId == Other.PrimitiveId && Light == Other.Light && SplitIndex == Other.SplitIndex && bTranslucentShadow == Other.bTranslucentShadow);
 		}
 
+		FProjectedShadowKey(const FProjectedShadowInfo& ProjectedShadowInfo)
+			: PrimitiveId(ProjectedShadowInfo.ParentSceneInfo ? ProjectedShadowInfo.ParentSceneInfo->PrimitiveComponentId : FPrimitiveComponentId())
+			, Light(ProjectedShadowInfo.LightSceneInfo->Proxy->GetLightComponent())
+			, SplitIndex(ProjectedShadowInfo.ShadowSplitIndex)
+			, bTranslucentShadow(ProjectedShadowInfo.bTranslucentShadow)
+		{
+		}
+
 		FProjectedShadowKey(FPrimitiveComponentId InPrimitiveId, const ULightComponent* InLight, int32 InSplitIndex, bool bInTranslucentShadow):
 			PrimitiveId(InPrimitiveId),
 			Light(InLight),
diff --git a/Engine/Source/Runtime/Renderer/Private/ShadowRendering.cpp b/Engine/Source/Runtime/Renderer/Private/ShadowRendering.cpp
--- a/Engine/Source/Runtime/Renderer/Private/ShadowRendering.cpp
+++ b/Engine/Source/Runtime/Renderer/Private/ShadowRendering.cpp
@@ -611,7 +611,7 @@ bool FShadowProjectionVS::ShouldCache(EShaderPlatform Platform)
 	return IsFeatureLevelSupported(Platform, ERHIFeatureLevel::SM4);
 }
 
-void FShadowProjectionVS::SetParameters(FRHICommandList& RHICmdList, const FSceneView& View, const FProjectedShadowInfo* ShadowInfo, const int32 ShadowSplitIndex)
+void FShadowProjectionVS::SetParameters(FRHICommandList& RHICmdList, const FSceneView& View, const FProjectedShadowInfo* ShadowInfo)
 {
 	FGlobalShader::SetParameters(RHICmdList, GetVertexShader(),View);
 	
@@ -1073,12 +1073,12 @@ bool FShadowDepthDrawingPolicyFactory::DrawDynamicMesh(
 		const EBlendMode BlendMode = Material->GetBlendMode();
 		const EMaterialShadingModel ShadingModel = Material->GetShadingModel();
 
-		const bool bOnePassPointLightShadow = Context.ShadowInfo->bOnePassPointLightShadow;
-		const bool bReflectiveShadowmap = Context.ShadowInfo->bReflectiveShadowmap && !bOnePassPointLightShadow && Material->ShouldInjectEmissiveIntoLPV();
+		const bool bLocalOnePassPointLightShadow = Context.ShadowInfo->bOnePassPointLightShadow;
+		const bool bReflectiveShadowmap = Context.ShadowInfo->bReflectiveShadowmap && !bLocalOnePassPointLightShadow && Material->ShouldInjectEmissiveIntoLPV();
 
 		if ( (!IsTranslucentBlendMode(BlendMode) || bReflectiveShadowmap ) && ShadingModel != MSM_Unlit )
 		{
-			const bool bDirectionalLight = Context.ShadowInfo->bDirectionalLight;
+			const bool bLocalDirectionalLight = Context.ShadowInfo->bDirectionalLight;
 			const bool bPreShadow = Context.ShadowInfo->bPreShadow;
 			const bool bTwoSided = Material->IsTwoSided() || PrimitiveSceneProxy->CastsShadowAsTwoSided();
 			const FShadowDepthDrawingPolicyContext PolicyContext(Context.ShadowInfo);
@@ -1089,8 +1089,8 @@ bool FShadowDepthDrawingPolicyFactory::DrawDynamicMesh(
 			{
 				FShadowDepthDrawingPolicy<true> DrawingPolicy(
 					MaterialRenderProxy->GetMaterial(FeatureLevel),
-					bDirectionalLight,
-					bOnePassPointLightShadow,
+					bLocalDirectionalLight,
+					bLocalOnePassPointLightShadow,
 					bPreShadow,
 					FeatureLevel,
 					Mesh.VertexFactory,
@@ -1111,8 +1111,8 @@ bool FShadowDepthDrawingPolicyFactory::DrawDynamicMesh(
 			{
 				FShadowDepthDrawingPolicy<false> DrawingPolicy(
 					MaterialRenderProxy->GetMaterial(FeatureLevel),
-					bDirectionalLight,
-					bOnePassPointLightShadow,
+					bLocalDirectionalLight,
+					bLocalOnePassPointLightShadow,
 					bPreShadow,
 					FeatureLevel,
 					Mesh.VertexFactory,
@@ -1966,7 +1966,7 @@ void FProjectedShadowInfo::RenderProjection(FRHICommandListImmediate& RHICmdList
 
 			RHICmdList.SetRasterizerState(TStaticRasterizerState<FM_Solid, CM_None>::GetRHI());
 
-			checkSlow(SplitIndex >= 0);
+			checkSlow(ShadowSplitIndex >= 0);
 			checkSlow(bDirectionalLight);
 
 			// Draw 2 fullscreen planes, front facing one at the near subfrustum plane, and back facing one at the far.
@@ -2001,7 +2001,7 @@ void FProjectedShadowInfo::RenderProjection(FRHICommandListImmediate& RHICmdList
 			};
 
 			// Only draw the near plane if this is not the nearest split
-			DrawPrimitiveUP(RHICmdList, PT_TriangleList, (SplitIndex > 0) ? 4 : 2, Verts, sizeof(FVector4));
+			DrawPrimitiveUP(RHICmdList, PT_TriangleList, (ShadowSplitIndex > 0) ? 4 : 2, Verts, sizeof(FVector4));
 		}
 	}
 	// Not a preshadow, mask the projection to any pixels inside the frustum.
@@ -2111,12 +2111,10 @@ void FProjectedShadowInfo::RenderProjection(FRHICommandListImmediate& RHICmdList
 
 		if (LocalQuality > 1)
 		{
-			if (IsWholeSceneDirectionalShadow() && SplitIndex > 0)
+			if (IsWholeSceneDirectionalShadow() && ShadowSplitIndex > 0)
 			{
 				// adjust kernel size so that the penumbra size of distant splits will better match up with the closer ones
-				const float SizeScale = SplitIndex / FMath::Max(0.001f, CVarCSMSplitPenumbraScale.GetValueOnRenderThread());
-
-		//test		LocalQuality = FMath::Clamp((int32)LocalQuality - SplitIndex, 2, 5);
+				const float SizeScale = ShadowSplitIndex / FMath::Max(0.001f, CVarCSMSplitPenumbraScale.GetValueOnRenderThread());
 			}
 			else if (LocalQuality > 2 && !bWholeSceneShadow)
 			{
@@ -2243,7 +2241,6 @@ void FProjectedShadowInfo::RenderOnePassPointLightProjection(FRHICommandListImme
 					++Reduce;
 				}
 			}
-	//test		LocalQuality = FMath::Clamp((int32)LocalQuality - Reduce, 2, 5);
 		}
 
 		switch(LocalQuality)
@@ -2279,7 +2276,7 @@ void FProjectedShadowInfo::RenderFrustumWireframe(FPrimitiveDrawInterface* PDI)
 	if(IsWholeSceneDirectionalShadow())
 	{
 		Color = FColor::White;
-		switch(SplitIndex)
+		switch(ShadowSplitIndex)
 		{
 			case 0: Color = FColor::Red; break;
 			case 1: Color = FColor::Yellow; break;
@@ -2415,7 +2412,7 @@ void FProjectedShadowInfo::UpdateShaderDepthBias()
 	}
 	else if (IsWholeSceneDirectionalShadow())
 	{
-		check(SplitIndex >= 0);
+		check(ShadowSplitIndex >= 0);
 
 		// the z range is adjusted to we need to adjust here as well
 		DepthBias = CVarCSMShadowDepthBias.GetValueOnRenderThread() / (MaxSubjectZ - MinSubjectZ);
@@ -2471,7 +2468,7 @@ float FProjectedShadowInfo::ComputeTransitionSize() const
 	}
 	else if (IsWholeSceneDirectionalShadow())
 	{
-		check(SplitIndex >= 0);
+		check(ShadowSplitIndex >= 0);
 
 		// todo: remove GetShadowTransitionScale()
 		// make 1/ ShadowTransitionScale, SpotLightShadowTransitionScale
@@ -2527,9 +2524,9 @@ void FProjectedShadowInfo::GetShadowTypeNameForDrawEvent(FString& TypeName) cons
 
 		if (bWholeSceneShadow)
 		{
-			if (SplitIndex >= 0)
+			if (ShadowSplitIndex >= 0)
 			{
-				TypeName = FString(TEXT("WholeScene split ")) + FString::FromInt(SplitIndex);
+				TypeName = FString(TEXT("WholeScene split ")) + FString::FromInt(ShadowSplitIndex);
 			}
 			else
 			{
@@ -2732,7 +2729,7 @@ struct FCompareFProjectedShadowInfoBySplitIndex
 			{
 				// Both A and B are CSMs
 				// Compare Split Indexes, to order them far to near.
-				return (B.SplitIndex < A.SplitIndex);
+				return (B.ShadowSplitIndex < A.ShadowSplitIndex);
 			}
 
 			// A is a CSM, B is per-object shadow etc.
diff --git a/Engine/Source/Runtime/Renderer/Private/ShadowRendering.h b/Engine/Source/Runtime/Renderer/Private/ShadowRendering.h
--- a/Engine/Source/Runtime/Renderer/Private/ShadowRendering.h
+++ b/Engine/Source/Runtime/Renderer/Private/ShadowRendering.h
@@ -566,7 +566,7 @@ public:
 	 * Index of the split if this is a whole scene shadow from a directional light, 
 	 * Or index of the direction if this is a whole scene shadow from a point light, otherwise INDEX_NONE. 
 	 */
-	int32 SplitIndex;
+	const int32 ShadowSplitIndex;
 
 	/** Whether the shadow has been allocated in the shadow depth buffer, and its X and Y properties have been initialized. */
 	uint32 bAllocated : 1;
@@ -586,13 +586,14 @@ public:
 	/** Whether the shadow is in the preshadow cache and its depths are up to date. */
 	uint32 bDepthsCached : 1;
 
-	uint32 bDirectionalLight : 1;
+	// redundant to LightSceneInfo->Proxy->GetLightType() == LightType_Directional, could be made ELightComponentType LightType
+	const uint32 bDirectionalLight : 1;
 
 	/** Whether this shadow affects the whole scene or only a group of objects. */
 	uint32 bWholeSceneShadow : 1;
 
 	/** Whether the shadow is a point light shadow that renders all faces of a cubemap in one pass. */
-	uint32 bOnePassPointLightShadow : 1;
+	const uint32 bOnePassPointLightShadow : 1;
 
 	/** Whether the shadow needs to render reflective shadow maps. */ 
 	uint32 bReflectiveShadowmap : 1; 
@@ -753,7 +754,8 @@ public:
 
 	inline bool IsWholeSceneDirectionalShadow() const 
 	{ 
-		return bWholeSceneShadow && SplitIndex >= 0 && LightSceneInfo->Proxy->GetLightType() == LightType_Directional; 
+		check(bDirectionalLight == LightSceneInfo->Proxy->GetLightType() == LightType_Directional);
+		return bWholeSceneShadow && ShadowSplitIndex >= 0 && bDirectionalLight; 
 	}
 
 	inline bool IsWholeScenePointLightShadow() const
@@ -932,7 +934,7 @@ public:
 		OutEnvironment.SetDefine(TEXT("USE_TRANSFORM"), (uint32)1);
 	}
 
-	void SetParameters(FRHICommandList& RHICmdList, const FSceneView& View, const FProjectedShadowInfo* ShadowInfo, int32 ShadowSplitIndex = INDEX_NONE);
+	void SetParameters(FRHICommandList& RHICmdList, const FSceneView& View, const FProjectedShadowInfo* ShadowInfo);
 
 	// Begin FShader interface
 	virtual bool Serialize(FArchive& Ar) override
diff --git a/Engine/Source/Runtime/Renderer/Private/ShadowSetup.cpp b/Engine/Source/Runtime/Renderer/Private/ShadowSetup.cpp
--- a/Engine/Source/Runtime/Renderer/Private/ShadowSetup.cpp
+++ b/Engine/Source/Runtime/Renderer/Private/ShadowSetup.cpp
@@ -359,13 +359,13 @@ FProjectedShadowInfo::FProjectedShadowInfo(
 	ResolutionY(0),
 	MaxScreenPercent(InMaxScreenPercent),
 	FadeAlphas(InFadeAlphas),
-	SplitIndex(INDEX_NONE),
+	ShadowSplitIndex(INDEX_NONE),
 	bAllocated(false),
 	bAllocatedInTranslucentLayout(false),
 	bRendered(false),
 	bAllocatedInPreshadowCache(false),
 	bDepthsCached(false),
-	bDirectionalLight(Initializer.bDirectionalLight),
+	bDirectionalLight(InLightSceneInfo->Proxy->GetLightType() == LightType_Directional),
 	bWholeSceneShadow(false),
 	bOnePassPointLightShadow(false),
 	bReflectiveShadowmap(false),
@@ -484,13 +484,13 @@ FProjectedShadowInfo::FProjectedShadowInfo(
 ,	ResolutionX(InResolutionX)
 ,	ResolutionY(InResolutionY)
 ,	MaxScreenPercent(1.0f)
-,	SplitIndex(Initializer.SplitIndex)
+,	ShadowSplitIndex(Initializer.SplitIndex)
 ,	bAllocated(false)
 ,	bAllocatedInTranslucentLayout(false)
 ,	bRendered(false)
 ,	bAllocatedInPreshadowCache(false)
 ,	bDepthsCached(false)
-,	bDirectionalLight(Initializer.bDirectionalLight)
+,	bDirectionalLight(InLightSceneInfo->Proxy->GetLightType() == LightType_Directional)
 ,	bWholeSceneShadow(true)
 ,	bOnePassPointLightShadow(Initializer.bOnePassPointLightShadow)
 ,	bReflectiveShadowmap(bInReflectiveShadowMap) 
@@ -511,8 +511,7 @@ FProjectedShadowInfo::FProjectedShadowInfo(
 	if(bInReflectiveShadowMap)
 	{
 		check(!bOnePassPointLightShadow);
-
-		SplitIndex = 0;
+		check(!ShadowSplitIndex);
 
 		// Quantise the RSM in shadow texel space
 		static bool bQuantize = true;
@@ -544,7 +543,7 @@ FProjectedShadowInfo::FProjectedShadowInfo(
 	}
 	else
 	{
-		if(Initializer.bDirectionalLight)
+		if(bDirectionalLight)
 		{
 			// Limit how small the depth range can be for smaller cascades
 			// This is needed for shadow modes like subsurface shadows which need depth information outside of the smaller cascade depth range
@@ -573,10 +572,10 @@ FProjectedShadowInfo::FProjectedShadowInfo(
 			PreShadowTranslation = -SnappedWorldPosition;
 		}
 
-		if (Initializer.SplitIndex >= 0 && Initializer.bDirectionalLight)
+		if (Initializer.SplitIndex >= 0 && bDirectionalLight)
 		{
 			checkSlow(InDependentView);
-			ShadowBounds = InLightSceneInfo->Proxy->GetShadowSplitBounds(*InDependentView, SplitIndex, 0);
+			ShadowBounds = InLightSceneInfo->Proxy->GetShadowSplitBounds(*InDependentView, ShadowSplitIndex, 0);
 		}
 		else
 		{
@@ -1543,7 +1542,7 @@ void FDeferredShadingSceneRenderer::CreateWholeSceneProjectedShadow(FLightSceneI
 		float MaxFadeAlpha = 0;
 		bool bStaticSceneOnly = false;
 
-		for(int32 ViewIndex = 0;ViewIndex < Views.Num();ViewIndex++)
+		for(int32 ViewIndex = 0, ViewCount = Views.Num(); ViewIndex < ViewCount; ++ViewIndex)
 		{
 			const FViewInfo& View = Views[ViewIndex];
 
@@ -1581,7 +1580,7 @@ void FDeferredShadingSceneRenderer::CreateWholeSceneProjectedShadow(FLightSceneI
 
 		if (MaxFadeAlpha > 1.0f / 256.0f)
 		{
-			for (int32 ShadowIndex = 0; ShadowIndex < ProjectedShadowInitializers.Num(); ShadowIndex++)
+			for (int32 ShadowIndex = 0, ShadowCount = ProjectedShadowInitializers.Num(); ShadowIndex < ShadowCount; ShadowIndex++)
 			{
 				const FWholeSceneProjectedShadowInitializer& ProjectedShadowInitializer = ProjectedShadowInitializers[ShadowIndex];
 
@@ -1745,7 +1744,7 @@ void FSceneRenderer::InitProjectedShadowVisibility(FRHICommandListImmediate& RHI
 								ProjectedShadowInfo.ParentSceneInfo->PrimitiveComponentId :
 								FPrimitiveComponentId(),
 							ProjectedShadowInfo.LightSceneInfo->Proxy->GetLightComponent(),
-							ProjectedShadowInfo.SplitIndex,
+							ProjectedShadowInfo.ShadowSplitIndex,
 							ProjectedShadowInfo.bTranslucentShadow
 							);
 
@@ -1769,7 +1768,7 @@ void FSceneRenderer::InitProjectedShadowVisibility(FRHICommandListImmediate& RHI
 							{
 								// Get split color
 								FColor Color = FColor::White;
-								switch(ProjectedShadowInfo.SplitIndex)
+								switch(ProjectedShadowInfo.ShadowSplitIndex)
 								{
 									case 0: Color = FColor::Red; break;
 									case 1: Color = FColor::Yellow; break;
@@ -2111,6 +2110,9 @@ void FSceneRenderer::AddViewDependentWholeSceneShadowsForView(
 
 					if (LightSceneInfo.Proxy->GetViewDependentRsmWholeSceneProjectedShadowInitializer(View, Lpv.GetBoundingBox(), ProjectedShadowInitializer))
 					{
+						// moved out from the FProjectedShadowInfo constructor
+						ProjectedShadowInitializer.SplitIndex = 0;
+
 						const FIntPoint ShadowBufferResolution = GSceneRenderTargets.GetReflectiveShadowMapTextureResolution();
 
 						// Create the projected shadow info.
diff --git a/Engine/Source/Runtime/Renderer/Private/TranslucentLighting.cpp b/Engine/Source/Runtime/Renderer/Private/TranslucentLighting.cpp
--- a/Engine/Source/Runtime/Renderer/Private/TranslucentLighting.cpp
+++ b/Engine/Source/Runtime/Renderer/Private/TranslucentLighting.cpp
@@ -1535,7 +1535,7 @@ static void InjectTranslucentLightArray(FRHICommandListImmediate& RHICmdList, co
 					if (InjectionData.ProjectedShadowInfo)
 					{
 						// shadows, restricting light contribution to the cascade bounds (except last cascade far to get light functions and no shadows there)
-						SetInjectionShader<LightType_Directional, true>(RHICmdList, View, InjectionData.LightFunctionMaterialProxy, LightSceneInfo, InjectionData.ProjectedShadowInfo, InjectionData.ProjectedShadowInfo->SplitIndex, VolumeCascadeIndex, *VertexShader, *GeometryShader, InjectionData.bApplyLightFunction, false);
+						SetInjectionShader<LightType_Directional, true>(RHICmdList, View, InjectionData.LightFunctionMaterialProxy, LightSceneInfo, InjectionData.ProjectedShadowInfo, InjectionData.ProjectedShadowInfo->ShadowSplitIndex, VolumeCascadeIndex, *VertexShader, *GeometryShader, InjectionData.bApplyLightFunction, false);
 					}
 					else
 					{
