added comment

diff --git a/Engine/Source/Runtime/Renderer/Private/SystemTextures.cpp b/Engine/Source/Runtime/Renderer/Private/SystemTextures.cpp
--- a/Engine/Source/Runtime/Renderer/Private/SystemTextures.cpp
+++ b/Engine/Source/Runtime/Renderer/Private/SystemTextures.cpp
@@ -227,6 +227,7 @@ void FSystemTextures::InitializeTextures(FRHICommandListImmediate& RHICmdList, E
 			float s = FMath::Sin(ww) * lenm;
 			float c = FMath::Cos(ww) * lenm;
 
+			// .zw is redundant and could be constructed in the shader but it adds 2 more instructions (likely to be unnoticed)
 			Bases[Pos] = FColor(FMath::Quantize8SignedByte(c), FMath::Quantize8SignedByte(s), FMath::Quantize8SignedByte(-s), FMath::Quantize8SignedByte(c));
 		}
 
