Added prototype of newer photo film post "r.TonemapperPhoto 1" option.

diff --git a/Engine/Shaders/PostProcessTonemap.usf b/Engine/Shaders/PostProcessTonemap.usf
--- a/Engine/Shaders/PostProcessTonemap.usf
+++ b/Engine/Shaders/PostProcessTonemap.usf
@@ -264,7 +264,11 @@ void MainPS(
 			LinearColor.rgb *= GrainMul;
 		#endif
 
-		half3 FilmColor = FilmPostProcess(LinearColor.rgb);
+		#if USE_PHOTO
+			half3 FilmColor = FilmPostProcessPhoto(LinearColor.rgb);
+		#else
+			half3 FilmColor = FilmPostProcess(LinearColor.rgb);
+		#endif
 		
 		// Apply "gamma" curve adjustment.
 		FilmColor.rgb = pow(FilmColor.rgb, InverseGamma.y);
diff --git a/Engine/Shaders/TonemapCommon.usf b/Engine/Shaders/TonemapCommon.usf
--- a/Engine/Shaders/TonemapCommon.usf
+++ b/Engine/Shaders/TonemapCommon.usf
@@ -62,7 +62,6 @@ half3 FilmPostProcess(half3 LinearColor)
 	return CurveColor;
 }
 
-
 half3 TonemapAndGammaCorrect(half3 LinearColor)
 {
 	// Clamp input to be positive
@@ -87,3 +86,36 @@ half3 TonemapAndGammaCorrect(half3 LinearColor)
 
 	return GammaColor;
 }
+
+//
+// Improved Film Post
+//
+
+half4 PhotoMatrixR_Lift;
+half4 PhotoMatrixG_LiftScale;
+half4 PhotoMatrixB_Gamma;
+half4 PhotoTint_Over;
+half4 PhotoOverTint;
+
+half3 FilmPostProcessPhoto(half3 LinearColor) 
+{
+	// Color and exposure control.
+	half3 PhotoColor = LinearColor;
+	// Apply color matrix.
+	PhotoColor.r=dot(LinearColor, PhotoMatrixR_Lift.rgb);
+	PhotoColor.g=dot(LinearColor, PhotoMatrixG_LiftScale.rgb);
+	PhotoColor.b=dot(LinearColor, PhotoMatrixB_Gamma.rgb);
+	// Lift, gamma, tint.
+	PhotoColor = pow(PhotoColor*PhotoMatrixG_LiftScale.a+PhotoMatrixR_Lift.a, PhotoMatrixB_Gamma.a)*PhotoTint_Over.rgb;
+	// Remove negatives.    
+	PhotoColor = max(PhotoColor, 1.0/32768.0);
+	// Move over pure hues towards "white" on over-exposure.
+	// This actually moves over to PhotoOverTint, 
+	// which can be a highly colored value, which later turns to white via highlight compression.
+	half OverLuma = dot(PhotoColor, half3(0.2126, 0.7152, 0.0722)); // BT.709 luma seems to work best.
+	half Over = (PhotoTint_Over.a*OverLuma) / ((1.0-0.18)+OverLuma);
+	PhotoColor = lerp(PhotoColor, PhotoOverTint.rgb*OverLuma, Over);
+	// Highlight compression without change of 18% mid-tone.
+	PhotoColor *= 1.0/((1.0-0.18)+PhotoColor);    
+	return PhotoColor;
+}
diff --git a/Engine/Source/Runtime/Core/Private/HAL/ConsoleManager.cpp b/Engine/Source/Runtime/Core/Private/HAL/ConsoleManager.cpp
--- a/Engine/Source/Runtime/Core/Private/HAL/ConsoleManager.cpp
+++ b/Engine/Source/Runtime/Core/Private/HAL/ConsoleManager.cpp
@@ -1985,6 +1985,13 @@ static TAutoConsoleVariable<int32> CVarTonemapperQuality(
 	TEXT("1: high (default, with high frequency pixel pattern to fight 8 bit color quantization)"),
 	ECVF_Scalability | ECVF_RenderThreadSafe);
 
+static TAutoConsoleVariable<int32> CVarTonemapperPhoto(
+	TEXT("r.TonemapperPhoto"),
+	0,
+	TEXT("0: use old film post\n")
+	TEXT("1: use new photo post"),
+	ECVF_Scalability | ECVF_RenderThreadSafe);
+
 static TAutoConsoleVariable<int32> CVarDetailMode(
 	TEXT("r.DetailMode"),
 	2,
diff --git a/Engine/Source/Runtime/Engine/Classes/Engine/Scene.h b/Engine/Source/Runtime/Engine/Classes/Engine/Scene.h
--- a/Engine/Source/Runtime/Engine/Classes/Engine/Scene.h
+++ b/Engine/Source/Runtime/Engine/Classes/Engine/Scene.h
@@ -73,6 +73,39 @@ struct FPostProcessSettings
 	UPROPERTY(BlueprintReadWrite, Category=Overrides, meta=(PinHiddenByDefault))
 	uint32 bOverride_FilmShadowTintAmount:1;
 
+	UPROPERTY(BlueprintReadWrite, Category=Overrides, meta=(PinHiddenByDefault))
+	uint32 bOverride_PhotoWhite:1;
+
+	UPROPERTY(BlueprintReadWrite, Category=Overrides, meta=(PinHiddenByDefault))
+	uint32 bOverride_PhotoSaturation:1;
+
+	UPROPERTY(BlueprintReadWrite, Category=Overrides, meta=(PinHiddenByDefault))
+	uint32 bOverride_PhotoChannelMixerRed:1;
+
+	UPROPERTY(BlueprintReadWrite, Category=Overrides, meta=(PinHiddenByDefault))
+	uint32 bOverride_PhotoChannelMixerGreen:1;
+
+	UPROPERTY(BlueprintReadWrite, Category=Overrides, meta=(PinHiddenByDefault))
+	uint32 bOverride_PhotoChannelMixerBlue:1;
+
+	UPROPERTY(BlueprintReadWrite, Category=Overrides, meta=(PinHiddenByDefault))
+	uint32 bOverride_PhotoBlackIn:1;
+
+	UPROPERTY(BlueprintReadWrite, Category=Overrides, meta=(PinHiddenByDefault))
+	uint32 bOverride_PhotoBlackOut:1;
+
+	UPROPERTY(BlueprintReadWrite, Category=Overrides, meta=(PinHiddenByDefault))
+	uint32 bOverride_PhotoContrast:1;
+
+	UPROPERTY(BlueprintReadWrite, Category=Overrides, meta=(PinHiddenByDefault))
+	uint32 bOverride_PhotoTint:1;
+
+	UPROPERTY(BlueprintReadWrite, Category=Overrides, meta=(PinHiddenByDefault))
+	uint32 bOverride_PhotoOver:1;
+
+	UPROPERTY(BlueprintReadWrite, Category=Overrides, meta=(PinHiddenByDefault))
+	uint32 bOverride_PhotoOverTint:1;
+
 	UPROPERTY(BlueprintReadWrite, Category=Overrides, meta=(PinHiddenByDefault))
 	uint32 bOverride_SceneColorTint:1;
 
@@ -367,8 +400,29 @@ struct FPostProcessSettings
 	UPROPERTY(interp, BlueprintReadWrite, Category=Film, AdvancedDisplay, meta=(UIMin = "1.0", UIMax = "4.0", editcondition = "bOverride_FilmDynamicRange", DisplayName = "Dynamic Range"))
 	float FilmDynamicRange;
 
-
-
+	UPROPERTY(interp, BlueprintReadWrite, Category=Photo, AdvancedDisplay, meta=(editcondition = "bOverride_PhotoWhite", DisplayName = "White", HideAlphaChannel))
+	FLinearColor PhotoWhite;
+	UPROPERTY(interp, BlueprintReadWrite, Category=Photo, AdvancedDisplay, meta=(UIMin = "0.0", UIMax = "2.0", editcondition = "bOverride_PhotoSaturation", DisplayName = "Saturation"))
+	float PhotoSaturation;
+	UPROPERTY(interp, BlueprintReadWrite, Category=Photo, AdvancedDisplay, meta=(editcondition = "bOverride_PhotoChannelMixerRed", DisplayName = "Channel Mixer Red", HideAlphaChannel))
+	FLinearColor PhotoChannelMixerRed;
+	UPROPERTY(interp, BlueprintReadWrite, Category=Photo, AdvancedDisplay, meta=(editcondition = "bOverride_PhotoChannelMixerGreen", DisplayName = "Channel Mixer Green", HideAlphaChannel))
+	FLinearColor PhotoChannelMixerGreen;
+	UPROPERTY(interp, BlueprintReadWrite, Category=Photo, AdvancedDisplay, meta=(editcondition = "bOverride_PhotoChannelMixerBlue", DisplayName = " Channel Mixer Blue", HideAlphaChannel))
+	FLinearColor PhotoChannelMixerBlue;
+	UPROPERTY(interp, BlueprintReadWrite, Category=Photo, AdvancedDisplay, meta=(UIMin = "0.0", UIMax = "1.0", editcondition = "bOverride_PhotoBlackIn", DisplayName = "Black In"))
+	float PhotoBlackIn;
+	UPROPERTY(interp, BlueprintReadWrite, Category=Photo, AdvancedDisplay, meta=(UIMin = "0.0", UIMax = "1.0", editcondition = "bOverride_PhotoBlackOut", DisplayName = "Black Out"))
+	float PhotoBlackOut;
+	UPROPERTY(interp, BlueprintReadWrite, Category=Photo, AdvancedDisplay, meta=(UIMin = "0.5", UIMax = "4.0", editcondition = "bOverride_PhotoContrast", DisplayName = "Contrast"))
+	float PhotoContrast;
+	UPROPERTY(interp, BlueprintReadWrite, Category=Photo, AdvancedDisplay, meta=(editcondition = "bOverride_PhotoTint", DisplayName = "Tint", HideAlphaChannel))
+	FLinearColor PhotoTint;
+	UPROPERTY(interp, BlueprintReadWrite, Category=Photo, AdvancedDisplay, meta=(UIMin = "0.0", UIMax = "1.0", editcondition = "bOverride_PhotoOver", DisplayName = "Over Exposure Mix"))
+	float PhotoOver;
+	UPROPERTY(interp, BlueprintReadWrite, Category=Photo, AdvancedDisplay, meta=(editcondition = "bOverride_PhotoOverTint", DisplayName = "Over Exposure Tint", HideAlphaChannel))
+	FLinearColor PhotoOverTint;
+	
 	/** Scene tint color */
 	UPROPERTY(interp, BlueprintReadWrite, Category=SceneColor, AdvancedDisplay, meta=(editcondition = "bOverride_SceneColorTint", HideAlphaChannel))
 	FLinearColor SceneColorTint;
@@ -800,6 +854,18 @@ struct FPostProcessSettings
 		FilmShadowTintBlend = 0.5;
 		FilmShadowTintAmount = 0.0;
 
+		PhotoWhite = FLinearColor(1.0f,1.0f,1.0f);
+		PhotoSaturation = 1.0f;
+		PhotoChannelMixerRed = FLinearColor(1.0f,0.0f,0.0f);
+		PhotoChannelMixerGreen = FLinearColor(0.0f,1.0f,0.0f);
+		PhotoChannelMixerBlue = FLinearColor(0.0f,0.0f,1.0f);
+		PhotoBlackIn = 1.0f/512.0f;
+		PhotoBlackOut = 1.0f/1024.0f;
+		PhotoContrast = 1.2f;
+		PhotoTint = FLinearColor(1.0f,1.0f,1.0f);
+		PhotoOver = 1.0f/32.0f;
+		PhotoOverTint = FLinearColor(16.0f,12.0f,1.0f);
+
 		SceneColorTint = FLinearColor(1, 1, 1);
 		SceneFringeIntensity = 0.0f;
 		SceneFringeSaturation = 0.5f;
diff --git a/Engine/Source/Runtime/Engine/Private/SceneView.cpp b/Engine/Source/Runtime/Engine/Private/SceneView.cpp
--- a/Engine/Source/Runtime/Engine/Private/SceneView.cpp
+++ b/Engine/Source/Runtime/Engine/Private/SceneView.cpp
@@ -715,6 +715,18 @@ void FSceneView::OverridePostProcessSettings(const FPostProcessSettings& Src, fl
 	LERP_PP(FilmShadowTintBlend);
 	LERP_PP(FilmShadowTintAmount);
 
+	LERP_PP(PhotoWhite);
+	LERP_PP(PhotoSaturation);
+	LERP_PP(PhotoChannelMixerRed);
+	LERP_PP(PhotoChannelMixerGreen);
+	LERP_PP(PhotoChannelMixerBlue);
+	LERP_PP(PhotoBlackIn);
+	LERP_PP(PhotoBlackOut);
+	LERP_PP(PhotoContrast);
+	LERP_PP(PhotoTint);
+	LERP_PP(PhotoOver);
+	LERP_PP(PhotoOverTint);
+
 	LERP_PP(SceneColorTint);
 	LERP_PP(SceneFringeIntensity);
 	LERP_PP(SceneFringeSaturation);
diff --git a/Engine/Source/Runtime/Renderer/Private/PostProcess/PostProcessTonemap.cpp b/Engine/Source/Runtime/Renderer/Private/PostProcess/PostProcessTonemap.cpp
--- a/Engine/Source/Runtime/Renderer/Private/PostProcess/PostProcessTonemap.cpp
+++ b/Engine/Source/Runtime/Renderer/Private/PostProcess/PostProcessTonemap.cpp
@@ -37,6 +37,7 @@ typedef enum {
 	TonemapperColorFringe       = (1<<13),
 	TonemapperColorGrading      = (1<<14),
 	TonemapperMsaa              = (1<<15),
+	TonemapperPhoto             = (1<<16),
 } TonemapperOption;
 
 // Tonemapper option cost (0 = no cost, 255 = max cost).
@@ -59,6 +60,7 @@ static uint8 TonemapperCostTab[] = {
 	1, //TonemapperColorFringe
 	1, //TonemapperColorGrading
 	1, //TonemapperMsaa
+	1, //TonemapperPhoto
 };
 
 // Edit the following to add and remove configurations.
@@ -159,83 +161,81 @@ static uint32 TonemapperConfBitmaskPC[22] = {
 	TonemapperColorGrading +
 	0,
 
-// SAME WITHOUT COLORGRADING
+// SAME WITHOUT COLORGRADING WITH NEW PHOTO POST
 
 	TonemapperBloom +
-	TonemapperContrast +
-	TonemapperColorMatrix +
-	TonemapperShadowTint +
 	TonemapperGrainJitter +
 	TonemapperGrainIntensity +
 	TonemapperGrainQuantization +
 	TonemapperVignette +
 	TonemapperVignetteColor +
 	TonemapperColorFringe +
+	TonemapperPhoto +
 	0,
 
 	TonemapperBloom +
-	TonemapperContrast +
-	TonemapperColorMatrix +
-	TonemapperShadowTint +
-	TonemapperVignette +
+	TonemapperGrainJitter +
+	TonemapperGrainIntensity +
 	TonemapperGrainQuantization +
-	TonemapperColorFringe +
+	TonemapperVignette +
+	TonemapperVignetteColor +
+	TonemapperPhoto +
 	0,
 
-	TonemapperBloom + 
-	TonemapperContrast + 
-	TonemapperColorMatrix + 
-	TonemapperShadowTint +
+	TonemapperBloom +
+	TonemapperVignette +
 	TonemapperGrainQuantization +
+	TonemapperColorFringe +
+	TonemapperPhoto +
 	0,
 
-	TonemapperBloom + 
-	TonemapperContrast + 
-	TonemapperColorMatrix +
+	TonemapperBloom +
+	TonemapperVignette +
 	TonemapperGrainQuantization +
+	TonemapperPhoto +
 	0,
 
-	TonemapperBloom	+ 
-	TonemapperContrast +
+	TonemapperBloom +
 	TonemapperGrainQuantization +
+	TonemapperPhoto +
 	0,
 
+
 	// same without TonemapperGrainQuantization
 
 	TonemapperBloom +
-	TonemapperContrast +
-	TonemapperColorMatrix +
-	TonemapperShadowTint +
 	TonemapperGrainJitter +
 	TonemapperGrainIntensity +
 	TonemapperVignette +
 	TonemapperVignetteColor +
 	TonemapperColorFringe +
+	TonemapperPhoto +
 	0,
 
 	TonemapperBloom +
-	TonemapperContrast +
-	TonemapperColorMatrix +
-	TonemapperShadowTint +
+	TonemapperGrainJitter +
+	TonemapperGrainIntensity +
 	TonemapperVignette +
-	TonemapperColorFringe +
+	TonemapperVignetteColor +
+	TonemapperPhoto +
 	0,
 
-	TonemapperBloom + 
-	TonemapperContrast + 
-	TonemapperColorMatrix + 
-	TonemapperShadowTint +
+	TonemapperBloom +
+	TonemapperVignette +
+	TonemapperColorFringe +
+	TonemapperPhoto +
 	0,
 
-	TonemapperBloom + 
-	TonemapperContrast + 
-	TonemapperColorMatrix +
+	TonemapperBloom +
+	TonemapperVignette +
+	TonemapperPhoto +
 	0,
 
-	TonemapperBloom	+ 
-	TonemapperContrast +
+	TonemapperBloom +
+	TonemapperPhoto +
 	0,
-	
+
+
 	//
 
 	TonemapperGammaOnly +
@@ -556,6 +556,11 @@ static uint32 TonemapperFindLeastExpensive(uint32* RESTRICT Table, uint32 TableE
 	MustNotHaveBitmask += ((RequiredOptionsBitmask & TonemapperMosaic) == 0) ? TonemapperMosaic : 0;
 	MustNotHaveBitmask += ((RequiredOptionsBitmask & TonemapperMsaa) == 0) ? TonemapperMsaa : 0;
 
+	MustNotHaveBitmask += ((RequiredOptionsBitmask & TonemapperPhoto) != 0) ? TonemapperColorMatrix : 0;
+	MustNotHaveBitmask += ((RequiredOptionsBitmask & TonemapperPhoto) != 0) ? TonemapperShadowTint : 0;
+	MustNotHaveBitmask += ((RequiredOptionsBitmask & TonemapperPhoto) != 0) ? TonemapperColorGrading : 0;
+	MustNotHaveBitmask += ((RequiredOptionsBitmask & TonemapperPhoto) == 0) ? TonemapperPhoto : 0;
+
 	// Search for exact match first.
 	uint32 Index;
 	for(Index = 0; Index < TableEntries; ++Index) 
@@ -685,7 +690,19 @@ static uint32 TonemapperGenerateBitmaskPC(const FViewInfo* RESTRICT View, bool b
 			Bitmask |= TonemapperGrainQuantization;
 		}
 	}
-	
+
+	{
+		static TConsoleVariableData<int32>* CVar = IConsoleManager::Get().FindTConsoleVariableDataInt(TEXT("r.TonemapperPhoto"));
+		int32 Value = CVar->GetValueOnRenderThread();
+
+		if(Value > 0)
+		{
+			// Remove settings not compatible.
+			Bitmask &= ~(TonemapperColorMatrix|TonemapperShadowTint|TonemapperContrast|TonemapperColorGrading|TonemapperMsaa);
+			Bitmask |= TonemapperPhoto;
+		}
+	}
+
 	static const auto CVar = IConsoleManager::Get().FindTConsoleVariableDataInt(TEXT("r.SceneColorFringeQuality")); 
 
 	int32 FringeQuality = CVar->GetValueOnRenderThread();
@@ -926,6 +943,90 @@ void FilmPostSetConstants(FVector4* RESTRICT const Constants, const uint32 Confi
 	Constants[7] = FVector4(OutColorShadow_Tint2, 0.0f);
 }
 
+// Improved film post: setting of constants.
+void FilmPostSetConstantsPhoto(FVector4* RESTRICT const Constants, const FPostProcessSettings* RESTRICT const FinalPostProcessSettings)
+{
+	FVector PhotoWhite(FinalPostProcessSettings->PhotoWhite);
+	float PhotoSat = FMath::Clamp(FinalPostProcessSettings->PhotoSaturation, 0.0f, 2.0f) - 1.0f;
+	FVector PhotoMixR(FinalPostProcessSettings->PhotoChannelMixerRed);
+	FVector PhotoMixG(FinalPostProcessSettings->PhotoChannelMixerGreen);
+	FVector PhotoMixB(FinalPostProcessSettings->PhotoChannelMixerBlue);
+	FVector PhotoOverTint(FinalPostProcessSettings->PhotoOverTint);
+	float PhotoGamma = FMath::Clamp(FinalPostProcessSettings->PhotoContrast, 0.0f, 4.0f);
+	float PhotoBlackOut = FMath::Clamp(FinalPostProcessSettings->PhotoBlackOut, 0.0f, 1.0f) * (1.0f/64.0f);
+	float PhotoBlackIn = FMath::Clamp(FinalPostProcessSettings->PhotoBlackIn, 0.0f, 1.0f) * (1.0f/64.0f);
+	FVector PhotoTint(FinalPostProcessSettings->PhotoTint);
+	float PhotoOver = FMath::Clamp(FinalPostProcessSettings->PhotoOver, 0.0f, 1.0f);
+
+	// Place holder for eventual exposure adjustment.
+	float PhotoExp = 1.0;
+
+	// Constant, don't change...
+	float PhotoMid = 0.18f;
+	// Conversion from linear rgb to luma (using HDTV coef).
+	FVector LumaWeights = FVector(0.2126f, 0.7152f, 0.0722f);
+
+	// Make sure not zero.
+	PhotoWhite += FVector(1.0f / (256.0f*256.0f*32.0f));
+	PhotoMixR += FVector(1.0f / (256.0f*256.0f*32.0f));
+	PhotoMixG += FVector(1.0f / (256.0f*256.0f*32.0f));
+	PhotoMixB += FVector(1.0f / (256.0f*256.0f*32.0f));
+
+	// Normalize inputs for luma.
+	FVector PW = PhotoWhite/FVector::DotProduct(PhotoWhite, LumaWeights);
+	FVector PR = PhotoMixR*(FVector::DotProduct(FVector(1.0f,0.0f,0.0f), LumaWeights)/FVector::DotProduct(PhotoMixR, LumaWeights));    
+	FVector PG = PhotoMixG*(FVector::DotProduct(FVector(0.0f,1.0f,0.0f), LumaWeights)/FVector::DotProduct(PhotoMixG, LumaWeights));    
+	FVector PB = PhotoMixB*(FVector::DotProduct(FVector(0.0f,0.0f,1.0f), LumaWeights)/FVector::DotProduct(PhotoMixB, LumaWeights));    
+
+	// Fold exposure into white point.
+	PW *= PhotoExp;
+
+	// Channel mixer matrix 
+	FVector MR = FVector(PR.X,PG.X,PB.X);
+	FVector MG = FVector(PR.Y,PG.Y,PB.Y);
+	FVector MB = FVector(PR.Z,PG.Z,PB.Z);
+
+	// Fold white point exposure into channel mixer matrix.
+	MR *= PW;
+	MG *= PW;
+	MB *= PW;
+
+	// Saturation matrix.
+	FVector SR = FVector(1.0+(1.0-LumaWeights.X)*PhotoSat,(-LumaWeights.Y)*PhotoSat,(-LumaWeights.Z)*PhotoSat);
+	FVector SG = FVector((-LumaWeights.X)*PhotoSat,1.0+(1.0-LumaWeights.Y)*PhotoSat,(-LumaWeights.Z)*PhotoSat);
+	FVector SB = FVector((-LumaWeights.X)*PhotoSat,(-LumaWeights.Y)*PhotoSat,1.0+(1.0-LumaWeights.Z)*PhotoSat);    
+
+	// Final matrix, fold white point exposure mixer into saturation.
+	FVector OutMatrixR = FVector(
+		FVector::DotProduct(FVector(MR.X,MG.X,MB.X),SR),
+		FVector::DotProduct(FVector(MR.Y,MG.Y,MB.Y),SR),
+		FVector::DotProduct(FVector(MR.Z,MG.Z,MB.Z),SR));
+	FVector OutMatrixG = FVector(
+		FVector::DotProduct(FVector(MR.X,MG.X,MB.X),SG),
+		FVector::DotProduct(FVector(MR.Y,MG.Y,MB.Y),SG),
+		FVector::DotProduct(FVector(MR.Z,MG.Z,MB.Z),SG));
+	FVector OutMatrixB = FVector(
+		FVector::DotProduct(FVector(MR.X,MG.X,MB.X),SB),
+		FVector::DotProduct(FVector(MR.Y,MG.Y,MB.Y),SB),
+		FVector::DotProduct(FVector(MR.Z,MG.Z,MB.Z),SB));    
+
+	float OutLift = (PhotoMid*(FMath::Pow(PhotoBlackOut*FMath::Pow(PhotoMid,PhotoGamma-1.0f),1.0f/PhotoGamma)-PhotoBlackIn))/(PhotoMid-PhotoBlackIn);
+	float OutLiftScale = (PhotoMid-OutLift)/PhotoMid;
+
+	FVector OutTint = PhotoTint;
+	OutTint += FVector(1.0f / (256.0f*256.0f*32.0f));
+	OutTint *= 1.0f / FVector::DotProduct(OutTint, LumaWeights);
+	// Factor in the scale to ensure gamma does not change the 18% midtone.
+	OutTint *= PhotoMid/FMath::Pow(PhotoMid, PhotoGamma);
+
+	Constants[0] = FVector4(OutMatrixR, OutLift);
+	Constants[1] = FVector4(OutMatrixG, OutLiftScale);
+	Constants[2] = FVector4(OutMatrixB, PhotoGamma); 
+	Constants[3] = FVector4(OutTint, PhotoOver*(1.0f/3.0f)); 
+	Constants[4] = FVector4(PhotoOverTint, 0.0f);
+}
+
+
 BEGIN_UNIFORM_BUFFER_STRUCT(FBloomDirtMaskParameters,)
 	DECLARE_UNIFORM_BUFFER_STRUCT_MEMBER(FVector4,Tint)
 	DECLARE_UNIFORM_BUFFER_STRUCT_MEMBER_TEXTURE(Texture2D,Mask)
@@ -966,6 +1067,7 @@ class FPostProcessTonemapPS : public FGlobalShader
 		// @todo Mac OS X: in order to share precompiled shaders between GL 3.3 & GL 4.1 devices we mustn't use volume-texture rendering as it isn't universally supported.
 		OutEnvironment.SetDefine(TEXT("USE_VOLUME_LUT"),		 (IsFeatureLevelSupported(Platform,ERHIFeatureLevel::SM4) && GSupportsVolumeTextureRendering && Platform != EShaderPlatform::SP_OPENGL_SM4_MAC));
 		OutEnvironment.SetDefine(TEXT("USE_COLOR_GRADING"),		 TonemapperIsDefined(ConfigBitmask, TonemapperColorGrading));
+		OutEnvironment.SetDefine(TEXT("USE_PHOTO"),              TonemapperIsDefined(ConfigBitmask, TonemapperPhoto));
 
 		if( !IsFeatureLevelSupported(Platform,ERHIFeatureLevel::SM5) )
 		{
@@ -999,6 +1101,12 @@ public:
 	FShaderParameter ColorShadow_Tint1;
 	FShaderParameter ColorShadow_Tint2;
 
+	FShaderParameter PhotoMatrixR_Lift;
+	FShaderParameter PhotoMatrixG_LiftScale;
+	FShaderParameter PhotoMatrixB_Gamma;
+	FShaderParameter PhotoTint_Over;
+	FShaderParameter PhotoOverTint;
+
 	//@HACK
 	FShaderParameter OverlayColor;
 
@@ -1027,6 +1135,12 @@ public:
 		ColorShadow_Tint1.Bind(Initializer.ParameterMap, TEXT("ColorShadow_Tint1"));
 		ColorShadow_Tint2.Bind(Initializer.ParameterMap, TEXT("ColorShadow_Tint2"));
 
+		PhotoMatrixR_Lift.Bind(Initializer.ParameterMap, TEXT("PhotoMatrixR_Lift"));
+		PhotoMatrixG_LiftScale.Bind(Initializer.ParameterMap, TEXT("PhotoMatrixG_LiftScale"));
+		PhotoMatrixB_Gamma.Bind(Initializer.ParameterMap, TEXT("PhotoMatrixB_Gamma"));
+		PhotoTint_Over.Bind(Initializer.ParameterMap, TEXT("PhotoTint_Over"));
+		PhotoOverTint.Bind(Initializer.ParameterMap, TEXT("PhotoOverTint"));
+		
 		OverlayColor.Bind(Initializer.ParameterMap, TEXT("OverlayColor"));
 	}
 	
@@ -1038,6 +1152,7 @@ public:
 			<< TexScale << VignetteColorIntensity << GrainScaleBiasJitter
 			<< ColorGradingLUT << ColorGradingLUTSampler
 			<< ColorMatrixR_ColorCurveCd1 << ColorMatrixG_ColorCurveCd3Cm3 << ColorMatrixB_ColorCurveCm2 << ColorCurve_Cm0Cd0_Cd2_Ch0Cm1_Ch3 << ColorCurve_Ch1_Ch2 << ColorShadow_Luma << ColorShadow_Tint1 << ColorShadow_Tint2
+			<< PhotoMatrixR_Lift << PhotoMatrixG_LiftScale << PhotoMatrixB_Gamma << PhotoTint_Over << PhotoOverTint
 			<< OverlayColor;
 		
 		return bShaderHasOutdatedParameters;
@@ -1157,6 +1272,16 @@ public:
 			SetShaderValue(Context.RHICmdList, ShaderRHI, ColorShadow_Tint1, Constants[6]);
 			SetShaderValue(Context.RHICmdList, ShaderRHI, ColorShadow_Tint2, Constants[7]);
 		}
+
+		{
+			FVector4 Constants[5];
+			FilmPostSetConstantsPhoto(Constants, &Context.View.FinalPostProcessSettings);
+			SetShaderValue(Context.RHICmdList, ShaderRHI, PhotoMatrixR_Lift, Constants[0]);
+			SetShaderValue(Context.RHICmdList, ShaderRHI, PhotoMatrixG_LiftScale, Constants[1]);
+			SetShaderValue(Context.RHICmdList, ShaderRHI, PhotoMatrixB_Gamma, Constants[2]); 
+			SetShaderValue(Context.RHICmdList, ShaderRHI, PhotoTint_Over, Constants[3]); 
+			SetShaderValue(Context.RHICmdList, ShaderRHI, PhotoOverTint, Constants[4]);
+		}
 	}
 	
 	static const TCHAR* GetSourceFilename()
