UE4 - Switch between outputting velocity as a separate render target or in the base pass as an extra GBuffer (r.BasePassOutputsVelocity defaults to 0)

diff --git a/Engine/Shaders/BasePassCommon.usf b/Engine/Shaders/BasePassCommon.usf
--- a/Engine/Shaders/BasePassCommon.usf
+++ b/Engine/Shaders/BasePassCommon.usf
@@ -26,7 +26,12 @@ struct FSharedBasePassInterpolants
 	#if USE_WORLD_POSITION_EXCLUDING_SHADER_OFFSETS
 	float4 PixelPositionExcludingWPO : TEXCOORD9;
 	#endif
+#endif
 
+#if OUTPUT_GBUFFER_VELOCITY
+	// .xy is clip position, pre divide by w; .w is clip W; .z is 0 or 1 to mask out the velocity output
+	float4 PrevScreenPosition : TEXCOORD16;
+	float4 ScreenPosition : TEXCOORD17;
 #endif
 };
 
diff --git a/Engine/Shaders/BasePassPixelShader.usf b/Engine/Shaders/BasePassPixelShader.usf
--- a/Engine/Shaders/BasePassPixelShader.usf
+++ b/Engine/Shaders/BasePassPixelShader.usf
@@ -15,6 +15,7 @@
 #include "Random.usf"
 #include "LightAccumulator.usf"
 #include "DeferredShadingCommon.usf"
+#include "VelocityCommon.usf"
 
 #define PREV_FRAME_COLOR	1
 #include "ScreenSpaceRayCast.usf"
@@ -586,6 +587,9 @@ void Main(
 	#if ALLOW_STATIC_LIGHTING
 		,out float4 OutGBufferE : SV_Target5
 	#endif
+	#if OUTPUT_GBUFFER_VELOCITY
+		, out float4 OutGBufferVelocity : SV_Target6
+	#endif
 #endif
 
 #if EDITOR_ALPHA2COVERAGE != 0
@@ -846,8 +850,24 @@ void Main(
 			float4 OutGBufferE = 0;
 		#endif
 
+		#if OUTPUT_GBUFFER_VELOCITY
+			{
+				float2 ScreenPosition = BasePassInterpolants.ScreenPosition.xy / BasePassInterpolants.ScreenPosition.w - View.PrevProjection[2].xy;
+				float2 PrevScreenPosition = BasePassInterpolants.PrevScreenPosition.xy / BasePassInterpolants.PrevScreenPosition.w - View.PrevProjection[2].xy;
+
+				// 2d velocity, includes camera an object motion
+				float2 Velocity = ScreenPosition - PrevScreenPosition;
+
+				// Make sure not to touch 0,0 which is clear color
+				GBuffer.Velocity = float4(EncodeVelocityToTexture(Velocity), 0, 0) * BasePassInterpolants.PrevScreenPosition.z;
+			}
+		#else
+			GBuffer.Velocity = 0;
+			float4 OutGBufferVelocity = 0;
+		#endif
+
 		float QuantizationBias = PseudoRandom( MaterialParameters.SVPosition.xy ) * 0.5 - 0.5;
-		EncodeGBuffer(GBuffer, OutGBufferA, OutGBufferB, OutGBufferC, OutGBufferD, OutGBufferE, QuantizationBias);
+		EncodeGBuffer(GBuffer, OutGBufferA, OutGBufferB, OutGBufferC, OutGBufferD, OutGBufferE, OutGBufferVelocity, QuantizationBias);
 	#endif 
 	
 
diff --git a/Engine/Shaders/BasePassVertexShader.usf b/Engine/Shaders/BasePassVertexShader.usf
--- a/Engine/Shaders/BasePassVertexShader.usf
+++ b/Engine/Shaders/BasePassVertexShader.usf
@@ -6,6 +6,11 @@
 
 #include "BasePassVertexCommon.usf"
 
+#if OUTPUT_GBUFFER_VELOCITY
+//@todo-rco: Move to pixel shader
+float SkipOutputVelocity;
+#endif
+
 /** Entry point for the base pass vertex shader. */
 void Main(
 	FVertexFactoryInput Input,
@@ -53,11 +58,24 @@ void Main(
 
 // Calculate the fog needed for translucency
 #if NEEDS_BASEPASS_FOGGING
-#if BASEPASS_ATMOSPHERIC_FOG
+	#if BASEPASS_ATMOSPHERIC_FOG
 	Output.BasePassInterpolants.VertexFog = CalculateVertexAtmosphericFog(WorldPosition.xyz, View.TranslatedViewOrigin.xyz);
-#else
+	#else
 	Output.BasePassInterpolants.VertexFog = CalculateVertexHeightFog(WorldPosition.xyz, View.TranslatedViewOrigin);
+	#endif
 #endif
+
+#if OUTPUT_GBUFFER_VELOCITY
+	{
+		float4 PrevTranslatedWorldPosition = VertexFactoryGetPreviousWorldPosition( Input, VFIntermediates );	
+		VertexParameters = GetMaterialVertexParameters(Input, VFIntermediates, PrevTranslatedWorldPosition.xyz, TangentToLocal);
+		PrevTranslatedWorldPosition.xyz += GetMaterialPreviousWorldPositionOffset(VertexParameters);
+
+		// compute the old screen pos with the old world position and the old camera matrix
+		Output.BasePassInterpolants.PrevScreenPosition = mul(float4(PrevTranslatedWorldPosition.xyz, 1), View.PrevTranslatedWorldToClip);		
+		Output.BasePassInterpolants.PrevScreenPosition.z = 1 - SkipOutputVelocity;
+		Output.BasePassInterpolants.ScreenPosition = mul(float4(WorldPosition.xyz, 1), View.TranslatedWorldToClip);
+	}
 #endif
 
 	OutputVertexID( Output );
diff --git a/Engine/Shaders/DeferredDecal.usf b/Engine/Shaders/DeferredDecal.usf
--- a/Engine/Shaders/DeferredDecal.usf
+++ b/Engine/Shaders/DeferredDecal.usf
@@ -189,8 +189,9 @@ void MainPS(
 		// static lighting isn't updated by decals so we don't need to update that render target
 		float4 OutTarget4 = 0;
 		float4 OutTarget5 = 0;
+		float4 OutTarget6 = 0;
 
-		EncodeGBuffer(Data, OutTarget1, OutTarget2, OutTarget3, OutTarget4, OutTarget5);
+		EncodeGBuffer(Data, OutTarget1, OutTarget2, OutTarget3, OutTarget4, OutTarget5, OutTarget6);
 
 		#if DECAL_BLEND_MODE == 1
 			// Stain
diff --git a/Engine/Shaders/DeferredShadingCommon.usf b/Engine/Shaders/DeferredShadingCommon.usf
--- a/Engine/Shaders/DeferredShadingCommon.usf
+++ b/Engine/Shaders/DeferredShadingCommon.usf
@@ -278,6 +278,8 @@ struct FGBufferData
 	// in unreal units (linear), can be used to reconstruct world position,
 	// only valid when decoding the GBuffer as the value gets reconstructed from the Z buffer
 	float Depth;
+	// Velocity for motion blur (only used when OUTPUT_GBUFFER_VELOCITY is enabled)
+	float4 Velocity;
 };
 
 
@@ -316,6 +318,7 @@ void EncodeGBuffer(
 	out float4 OutGBufferC,
 	out float4 OutGBufferD,
 	out float4 OutGBufferE,
+	out float4 OutGBufferVelocity,
 	float QuantizationBias = 0		// -0.5 to 0.5 random float. Used to bias quantization.
 	)
 {
@@ -347,6 +350,12 @@ void EncodeGBuffer(
 	}
 
 	OutGBufferE = GBuffer.PrecomputedShadowFactors;
+
+#if OUTPUT_GBUFFER_VELOCITY
+	OutGBufferVelocity = GBuffer.Velocity;
+#else
+	OutGBufferVelocity = 0;
+#endif
 }
 
 /** Populates FGBufferData */
@@ -356,6 +365,7 @@ FGBufferData DecodeGBufferData(
 	float4 InGBufferC,
 	float4 InGBufferD,
 	float4 InGBufferE,
+	float4 InGBufferVelocity,
 	float CustomNativeDepth,
 	float SceneDepth,
 	bool bGetNormalizedNormal)
@@ -402,6 +412,8 @@ FGBufferData DecodeGBufferData(
 		#endif //COMPILE_SHADERS_FOR_DEVELOPMENT == 1
 	}
 
+	GBuffer.Velocity = InGBufferVelocity;
+
 	return GBuffer;
 }
 
@@ -526,9 +538,15 @@ SamplerState CustomDepthTextureSampler;
 			float4 GBufferE = 1;
 		#endif
 
+		#if OUTPUT_GBUFFER_VELOCITY
+			float4 GBufferVelocity = GBuffers.GBufferVelocityTextureNonMS.Load(int3(PixelPos, 0));
+		#else
+			float4 GBufferVelocity = 0;
+		#endif
+
 		float SceneDepth = CalcSceneDepth(PixelPos);
 
-		return DecodeGBufferData(GBufferA, GBufferB, GBufferC, GBufferD, GBufferE, CustomNativeDepth, SceneDepth, bGetNormalizedNormal);
+		return DecodeGBufferData(GBufferA, GBufferB, GBufferC, GBufferD, GBufferE, GBufferVelocity, CustomNativeDepth, SceneDepth, bGetNormalizedNormal);
 	}
 
 	FDBufferData GetDBufferData(uint2 PixelPos)
@@ -594,11 +612,17 @@ FGBufferData GetGBufferData(float2 UV, bool bGetNormalizedNormal = true)
 	#else
 		float4 GBufferE = 1;
 	#endif
+
+	#if OUTPUT_GBUFFER_VELOCITY
+		float4 GBufferVelocity = Texture2DSampleLevel(GBuffers.GBufferVelocityTexture, GBuffers.GBufferVelocityTextureSampler, UV, 0);
+	#else
+		float4 GBufferVelocity = 0;
+	#endif
 #endif
 
 	float SceneDepth = CalcSceneDepth(UV);
 	
-	return DecodeGBufferData(GBufferA, GBufferB, GBufferC, GBufferD, GBufferE, CustomNativeDepth, SceneDepth, bGetNormalizedNormal);
+	return DecodeGBufferData(GBufferA, GBufferB, GBufferC, GBufferD, GBufferE, GBufferVelocity, CustomNativeDepth, SceneDepth, bGetNormalizedNormal);
 }
 
 // Minimal path for just the lighting model, used to branch around unlit pixels (skybox)
@@ -657,7 +681,13 @@ FScreenSpaceData GetScreenSpaceData(float2 UV, bool bGetNormalizedNormal = true)
 		float DeviceZ = SceneDepthSurface.Load(IntUV, SampleIndex);
 		float SceneDepth = ConvertFromDeviceZ(DeviceZ);
 
-		return DecodeGBufferData(GBufferA, GBufferB, GBufferC, GBufferD, GBufferE, CustomNativeDepth, SceneDepth, bGetNormalizedNormal);
+	#if OUTPUT_GBUFFER_VELOCITY
+		float4 GBufferVelocity = GBuffers.GBufferVelocityTextureMS.Load(IntUV, SampleIndex);
+	#else
+		float4 GBufferVelocity = 0;
+	#endif
+
+		return DecodeGBufferData(GBufferA, GBufferB, GBufferC, GBufferD, GBufferE, GBufferVelocity, CustomNativeDepth, SceneDepth, bGetNormalizedNormal);
 	}
 
 	FGBufferData GetGBufferDataMS(float2 UV, uint SampleIndex, bool bGetNormalizedNormal = true)
diff --git a/Engine/Shaders/Definitions.usf b/Engine/Shaders/Definitions.usf
--- a/Engine/Shaders/Definitions.usf
+++ b/Engine/Shaders/Definitions.usf
@@ -112,6 +112,10 @@
 #define USING_TESSELLATION								0
 #endif
 
+#ifndef OUTPUT_GBUFFER_VELOCITY
+#define OUTPUT_GBUFFER_VELOCITY									0
+#endif
+
 #define PC_D3D											SM5_PROFILE
 
 #ifndef NEEDS_LIGHTMAP_COORDINATE
diff --git a/Engine/Shaders/VelocityCommon.usf b/Engine/Shaders/VelocityCommon.usf
new file mode 100644
--- /dev/null
+++ b/Engine/Shaders/VelocityCommon.usf
@@ -0,0 +1,19 @@
+// Copyright 1998-2015 Epic Games, Inc. All Rights Reserved.
+
+/*=============================================================================
+	VelocityCommon.usf: Common functions for calculating velocity vectors.
+=============================================================================*/
+
+#pragma once
+
+float2 Calculate2DVelocity(float4 PackedVelocityA, float4 PackedVelocityC)
+{
+	float2 ScreenPosition = PackedVelocityA.xy / PackedVelocityA.w - View.ViewToClip[2].xy;
+	float2 PrevScreenPosition = PackedVelocityC.xy / PackedVelocityC.w - View.PrevProjection[2].xy;
+
+	// 2d velocity, includes camera an object motion
+	float2 Velocity = ScreenPosition - PrevScreenPosition;
+
+	// Make sure not to touch 0,0 which is clear color
+	return Velocity;
+}
diff --git a/Engine/Shaders/VelocityShader.usf b/Engine/Shaders/VelocityShader.usf
--- a/Engine/Shaders/VelocityShader.usf
+++ b/Engine/Shaders/VelocityShader.usf
@@ -7,6 +7,7 @@
 #include "Common.usf"
 #include "Material.usf"
 #include "VertexFactory.usf"
+#include "VelocityCommon.usf"
 
 #ifndef GPUSKIN_PASS_THROUGH
 #define GPUSKIN_PASS_THROUGH 0
@@ -198,14 +199,10 @@ void MainPixelShader(
 	CalcMaterialParameters(MaterialParameters, SvPosition, bIsFrontFace, SvPosition);
 	GetMaterialCoverageAndClipping(MaterialParameters);
 
-	float2 ScreenPosition = Input.PackedVelocityA.xy / Input.PackedVelocityA.w - View.ViewToClip[2].xy;
-	float2 PrevScreenPosition = Input.PackedVelocityC.xy / Input.PackedVelocityC.w - View.PrevProjection[2].xy;
-
 	// 2d velocity, includes camera an object motion
-	float2 Velocity = ScreenPosition - PrevScreenPosition;
+	float2 Velocity = Calculate2DVelocity(Input.PackedVelocityA, Input.PackedVelocityC);
 
 	// Make sure not to touch 0,0 which is clear color
 	OutColor.xy = EncodeVelocityToTexture(Velocity);
 	OutColor.zw = float2(0, 0);
 }
-
diff --git a/Engine/Source/Runtime/Engine/Classes/Engine/RendererSettings.h b/Engine/Source/Runtime/Engine/Classes/Engine/RendererSettings.h
--- a/Engine/Source/Runtime/Engine/Classes/Engine/RendererSettings.h
+++ b/Engine/Source/Runtime/Engine/Classes/Engine/RendererSettings.h
@@ -234,6 +234,11 @@ class ENGINE_API URendererSettings
 		ToolTip="Control how the scene is cleared before rendering"))
 	TEnumAsByte<EClearSceneOptions::Type> ClearSceneMethod;
 
+	UPROPERTY(config, EditAnywhere, Category=Optimizations, meta=(
+		ConsoleVariable="r.BasePassOutputsVelocity", DisplayName="WPO materials output velocity",
+		ToolTip="Enables materials with time-based World Position Offset attributes to render their velocities for motion blur in the base rendering pass."))
+	uint32 bBasePassOutputsVelocity:1;
+
 	UPROPERTY(config, EditAnywhere, Category=Editor, meta=(
 		ConsoleVariable="r.WireframeCullThreshold",DisplayName="Wireframe Cull Threshold",
 		ToolTip="Screen radius at which wireframe objects are culled. Larger values can improve performance when viewing a scene in wireframe."))
diff --git a/Engine/Source/Runtime/Engine/Public/SceneManagement.h b/Engine/Source/Runtime/Engine/Public/SceneManagement.h
--- a/Engine/Source/Runtime/Engine/Public/SceneManagement.h
+++ b/Engine/Source/Runtime/Engine/Public/SceneManagement.h
@@ -1485,6 +1485,7 @@ public:
 	 *	@return	bool				true if the primitive info was found and set
 	 */
 	bool GetPrimitiveMotionBlurInfo(const FPrimitiveSceneInfo* PrimitiveSceneInfo, FMatrix& OutPreviousLocalToWorld);
+	bool GetPrimitiveMotionBlurInfo(const FPrimitiveSceneInfo* PrimitiveSceneInfo, FMatrix& OutPreviousLocalToWorld) const;
 
 	/** Request to clear all stored motion blur data for this scene. */
 	void SetClearMotionBlurInfo();
@@ -1505,6 +1506,7 @@ private:
 	 * @return 0 if not found, otherwise pointer into MotionBlurInfos, don't store for longer
 	 */
 	FMotionBlurInfo* FindMBInfoIndex(FPrimitiveComponentId ComponentId);
+	const FMotionBlurInfo* FindMBInfoIndex(FPrimitiveComponentId ComponentId) const;
 };
 
 
diff --git a/Engine/Source/Runtime/Renderer/Private/BasePassRendering.h b/Engine/Source/Runtime/Renderer/Private/BasePassRendering.h
--- a/Engine/Source/Runtime/Renderer/Private/BasePassRendering.h
+++ b/Engine/Source/Runtime/Renderer/Private/BasePassRendering.h
@@ -7,6 +7,7 @@
 #pragma once
 
 #include "LightMapRendering.h"
+#include "VelocityRendering.h"
 #include "ShaderBaseClasses.h"
 #include "LightGrid.h"
 #include "EditorCompositeParams.h"
@@ -25,11 +26,18 @@ class TBasePassVertexShaderBaseType : public FMeshMaterialShader, public LightMa
 protected:
 	TBasePassVertexShaderBaseType() {}
 	TBasePassVertexShaderBaseType(const FMeshMaterialShaderType::CompiledShaderInitializerType& Initializer):
-	FMeshMaterialShader(Initializer)
+		FMeshMaterialShader(Initializer)
 	{
 		LightMapPolicyType::VertexParametersType::Bind(Initializer.ParameterMap);
 		HeightFogParameters.Bind(Initializer.ParameterMap);
 		AtmosphericFogTextureParameters.Bind(Initializer.ParameterMap);
+		const bool bOutputsVelocityToGBuffer = FVelocityRendering::OutputsToGBuffer();
+		if (bOutputsVelocityToGBuffer)
+		{
+			PreviousLocalToWorldParameter.Bind(Initializer.ParameterMap, TEXT("PreviousLocalToWorld"));
+//@todo-rco: Move to pixel shader
+			SkipOutputVelocityParameter.Bind(Initializer.ParameterMap, TEXT("SkipOutputVelocity"));
+		}
 	}
 
 public:
@@ -42,6 +50,7 @@ public:
 	{
 		FMeshMaterialShader::ModifyCompilationEnvironment(Platform, Material, OutEnvironment);
 		LightMapPolicyType::ModifyCompilationEnvironment(Platform, Material, OutEnvironment);
+		OutEnvironment.SetDefine(TEXT("OUTPUT_GBUFFER_VELOCITY"), (uint32)(FVelocityRendering::OutputsToGBuffer() ? 1 : 0));
 	}
 
 	virtual bool Serialize(FArchive& Ar)
@@ -50,6 +59,8 @@ public:
 		LightMapPolicyType::VertexParametersType::Serialize(Ar);
 		Ar << HeightFogParameters;
 		Ar << AtmosphericFogTextureParameters;
+		Ar << PreviousLocalToWorldParameter;
+		Ar << SkipOutputVelocityParameter;
 		return bShaderHasOutdatedParameters;
 	}
 
@@ -72,36 +83,36 @@ public:
 		}
 	}
 
-	void SetMesh(FRHICommandList& RHICmdList, const FVertexFactory* VertexFactory,const FSceneView& View,const FPrimitiveSceneProxy* Proxy,const FMeshBatchElement& BatchElement)
-	{
-		FMeshMaterialShader::SetMesh(RHICmdList, GetVertexShader(),VertexFactory,View,Proxy,BatchElement);
-	}
+	void SetMesh(FRHICommandList& RHICmdList, const FVertexFactory* VertexFactory,const FSceneView& View,const FPrimitiveSceneProxy* Proxy,const FMeshBatchElement& BatchElement);
 
 private:
 	
 	/** The parameters needed to calculate the fog contribution from height fog layers. */
 	FHeightFogShaderParameters HeightFogParameters;
 	FAtmosphereShaderTextureParameters AtmosphericFogTextureParameters;
+	// When outputting from base pass, the previous transform
+	FShaderParameter PreviousLocalToWorldParameter;
+	FShaderParameter SkipOutputVelocityParameter;
 };
 
 template<typename LightMapPolicyType, bool bEnableAtmosphericFog>
 class TBasePassVS : public TBasePassVertexShaderBaseType<LightMapPolicyType>
 {
 	DECLARE_SHADER_TYPE(TBasePassVS,MeshMaterial);
+	typedef TBasePassVertexShaderBaseType<LightMapPolicyType> Super;
 
 protected:
 
 	TBasePassVS() {}
 	TBasePassVS(const FMeshMaterialShaderType::CompiledShaderInitializerType& Initializer):
-		TBasePassVertexShaderBaseType<LightMapPolicyType>(Initializer)
+		Super(Initializer)
 	{
 	}
 
 public:
-
 	static bool ShouldCache(EShaderPlatform Platform,const FMaterial* Material,const FVertexFactoryType* VertexFactoryType)
 	{
-		bool bShouldCache = TBasePassVertexShaderBaseType<LightMapPolicyType>::ShouldCache(Platform, Material, VertexFactoryType);
+		bool bShouldCache = Super::ShouldCache(Platform, Material, VertexFactoryType);
 		return bShouldCache 
 			&& (IsFeatureLevelSupported(Platform, ERHIFeatureLevel::SM4))
 			&& (!bEnableAtmosphericFog || IsFeatureLevelSupported(Platform, ERHIFeatureLevel::SM4));
@@ -109,7 +120,7 @@ public:
 
 	static void ModifyCompilationEnvironment(EShaderPlatform Platform, const FMaterial* Material, FShaderCompilerEnvironment& OutEnvironment)
 	{
-		TBasePassVertexShaderBaseType<LightMapPolicyType>::ModifyCompilationEnvironment(Platform, Material, OutEnvironment);
+		Super::ModifyCompilationEnvironment(Platform, Material, OutEnvironment);
 		OutEnvironment.SetDefine(TEXT("BASEPASS_ATMOSPHERIC_FOG"),(uint32)(bEnableAtmosphericFog ? 1 : 0));
 	}
 };
@@ -323,6 +334,13 @@ public:
 	{
 		FMeshMaterialShader::ModifyCompilationEnvironment(Platform, Material, OutEnvironment);
 		LightMapPolicyType::ModifyCompilationEnvironment(Platform, Material, OutEnvironment);
+
+		const bool bOutputVelocity = FVelocityRendering::OutputsToGBuffer();
+		OutEnvironment.SetDefine(TEXT("OUTPUT_GBUFFER_VELOCITY"), (uint32)(bOutputVelocity ? 1 : 0));
+		if (bOutputVelocity)
+		{
+			OutEnvironment.SetRenderTargetOutputFormat(6, PF_G16R16);
+		}
 	}
 
 	/** Initialization constructor. */
@@ -376,16 +394,7 @@ public:
 		EditorCompositeParams.SetParameters(RHICmdList, MaterialResource, View, bEnableEditorPrimitveDepthTest, GetPixelShader());
 	}
 
-	void SetMesh(FRHICommandList& RHICmdList, const FVertexFactory* VertexFactory,const FSceneView& View,const FPrimitiveSceneProxy* Proxy,const FMeshBatchElement& BatchElement, EBlendMode BlendMode)
-	{
-		if (View.GetFeatureLevel() >= ERHIFeatureLevel::SM4
-			&& IsTranslucentBlendMode(BlendMode))
-		{
-			TranslucentLightingParameters.SetMesh(RHICmdList, this, Proxy, View.GetFeatureLevel());
-		}
-
-		FMeshMaterialShader::SetMesh(RHICmdList, GetPixelShader(),VertexFactory,View,Proxy,BatchElement);
-	}
+	void SetMesh(FRHICommandList& RHICmdList, const FVertexFactory* VertexFactory,const FSceneView& View,const FPrimitiveSceneProxy* Proxy,const FMeshBatchElement& BatchElement, EBlendMode BlendMode);
 
 	virtual bool Serialize(FArchive& Ar)
 	{
@@ -400,7 +409,7 @@ public:
 private:
 	FTranslucentLightingParameters TranslucentLightingParameters;
 	FEditorCompositingParameters EditorCompositeParams;
-	FShaderResourceParameter LightGrid; 
+	FShaderResourceParameter LightGrid;
 };
 
 /** The concrete base pass pixel shader type. */
diff --git a/Engine/Source/Runtime/Renderer/Private/BasePassRendering.inl b/Engine/Source/Runtime/Renderer/Private/BasePassRendering.inl
new file mode 100644
--- /dev/null
+++ b/Engine/Source/Runtime/Renderer/Private/BasePassRendering.inl
@@ -0,0 +1,48 @@
+// Copyright 1998-2015 Epic Games, Inc. All Rights Reserved.
+
+/*=============================================================================
+	BasePassRendering.inl: Base pass rendering implementations.
+		(Due to forward declaration issues)
+=============================================================================*/
+
+#pragma once
+
+template<typename LightMapPolicyType>
+inline void TBasePassVertexShaderBaseType<LightMapPolicyType>::SetMesh(FRHICommandList& RHICmdList, const FVertexFactory* VertexFactory,const FSceneView& View,const FPrimitiveSceneProxy* Proxy,const FMeshBatchElement& BatchElement)
+{
+	FVertexShaderRHIParamRef VertexShaderRHI = GetVertexShader();
+	FMeshMaterialShader::SetMesh(RHICmdList, VertexShaderRHI, VertexFactory, View, Proxy, BatchElement);
+
+	if (PreviousLocalToWorldParameter.IsBound() && Proxy)
+	{
+		check(SkipOutputVelocityParameter.IsBound());
+
+		// previous transform can be stored in the scene for each primitive
+		FMatrix PreviousLocalToWorld;
+		const auto& Scene = (const FScene&)Proxy->GetScene();
+		//@todo-rco: Move to InitViews
+		if (Scene.MotionBlurInfoData.GetPrimitiveMotionBlurInfo(Proxy->GetPrimitiveSceneInfo(), PreviousLocalToWorld))
+		{
+			const auto& ViewInfo = (const FViewInfo&)View;
+			SetShaderValue(RHICmdList, VertexShaderRHI, PreviousLocalToWorldParameter, PreviousLocalToWorld.ConcatTranslation(ViewInfo.PrevViewMatrices.PreViewTranslation));
+			SetShaderValue(RHICmdList, VertexShaderRHI, SkipOutputVelocityParameter, 0.0f);
+		}
+		else
+		{
+			SetShaderValue(RHICmdList, VertexShaderRHI, PreviousLocalToWorldParameter, Proxy->GetLocalToWorld());
+			SetShaderValue(RHICmdList, VertexShaderRHI, SkipOutputVelocityParameter, 1.0f);
+		}
+	}
+}
+
+template<typename LightMapPolicyType>
+void TBasePassPixelShaderBaseType<LightMapPolicyType>::SetMesh(FRHICommandList& RHICmdList, const FVertexFactory* VertexFactory, const FSceneView& View, const FPrimitiveSceneProxy* Proxy, const FMeshBatchElement& BatchElement, EBlendMode BlendMode)
+{
+	if (View.GetFeatureLevel() >= ERHIFeatureLevel::SM4
+		&& IsTranslucentBlendMode(BlendMode))
+	{
+		TranslucentLightingParameters.SetMesh(RHICmdList, this, Proxy, View.GetFeatureLevel());
+	}
+
+	FMeshMaterialShader::SetMesh(RHICmdList, GetPixelShader(), VertexFactory, View, Proxy, BatchElement);
+}
diff --git a/Engine/Source/Runtime/Renderer/Private/CompositionLighting/CompositionLighting.cpp b/Engine/Source/Runtime/Renderer/Private/CompositionLighting/CompositionLighting.cpp
--- a/Engine/Source/Runtime/Renderer/Private/CompositionLighting/CompositionLighting.cpp
+++ b/Engine/Source/Runtime/Renderer/Private/CompositionLighting/CompositionLighting.cpp
@@ -281,6 +281,7 @@ void FCompositionLighting::ProcessAfterBasePass(FRHICommandListImmediate& RHICmd
 	GRenderTargetPool.VisualizeTexture.SetCheckPoint(RHICmdList, GSceneRenderTargets.GBufferC);
 	GRenderTargetPool.VisualizeTexture.SetCheckPoint(RHICmdList, GSceneRenderTargets.GBufferD);
 	GRenderTargetPool.VisualizeTexture.SetCheckPoint(RHICmdList, GSceneRenderTargets.GBufferE);
+	GRenderTargetPool.VisualizeTexture.SetCheckPoint(RHICmdList, GSceneRenderTargets.GBufferVelocity);
 	GRenderTargetPool.VisualizeTexture.SetCheckPoint(RHICmdList, GSceneRenderTargets.ScreenSpaceAO);
 	
 	// so that the passes can register themselves to the graph
diff --git a/Engine/Source/Runtime/Renderer/Private/DeferredShadingRenderer.cpp b/Engine/Source/Runtime/Renderer/Private/DeferredShadingRenderer.cpp
--- a/Engine/Source/Runtime/Renderer/Private/DeferredShadingRenderer.cpp
+++ b/Engine/Source/Runtime/Renderer/Private/DeferredShadingRenderer.cpp
@@ -20,6 +20,7 @@
 #include "DistanceFieldSurfaceCacheLighting.h"
 #include "PostProcess/PostProcessing.h"
 #include "DistanceFieldAtlas.h"
+#include "../../Engine/Private/SkeletalRenderGPUSkin.h"		// GPrevPerBoneMotionBlur
 
 TAutoConsoleVariable<int32> CVarEarlyZPass(
 	TEXT("r.EarlyZPass"),
@@ -996,10 +997,17 @@ void FDeferredShadingSceneRenderer::Render(FRHICommandListImmediate& RHICmdList)
 
 	TRefCountPtr<IPooledRenderTarget> VelocityRT;
 
-	// Render the velocities of movable objects
-	RHICmdList.SetCurrentStat(GET_STATID(STAT_CLM_Velocity));
-	RenderVelocities(RHICmdList, VelocityRT);
-	RHICmdList.SetCurrentStat(GET_STATID(STAT_CLM_AfterVelocity));
+	if (FVelocityRendering::OutputsToGBuffer())
+	{
+		VelocityRT = GSceneRenderTargets.GetGBufferVelocityRT();
+	}
+	else
+	{
+		// Render the velocities of movable objects
+		RHICmdList.SetCurrentStat(GET_STATID(STAT_CLM_Velocity));
+		RenderVelocities(RHICmdList, VelocityRT);
+		RHICmdList.SetCurrentStat(GET_STATID(STAT_CLM_AfterVelocity));
+	}
 
 	// Finish rendering for each view.
 	if(ViewFamily.bResolveScene)
@@ -1025,7 +1033,7 @@ void FDeferredShadingSceneRenderer::Render(FRHICommandListImmediate& RHICmdList)
 	}
 
 	//grab the new transform out of the proxies for next frame
-	if(VelocityRT)
+	if (VelocityRT)
 	{
 		Scene->MotionBlurInfoData.UpdateMotionBlurCache(Scene);
 	}
@@ -1269,6 +1277,11 @@ bool FDeferredShadingSceneRenderer::RenderBasePass(FRHICommandListImmediate& RHI
 {
 	bool bDirty = false;
 
+	if (FVelocityRendering::OutputsToGBuffer())
+	{
+		GPrevPerBoneMotionBlur.LockData();
+	}
+
 	if(ViewFamily.EngineShowFlags.LightMapDensity && AllowDebugViewmodes())
 	{
 		// Override the base pass with the lightmap density pass if the viewmode is enabled.
@@ -1302,6 +1315,11 @@ bool FDeferredShadingSceneRenderer::RenderBasePass(FRHICommandListImmediate& RHI
 
 	}
 
+	if (FVelocityRendering::OutputsToGBuffer())
+	{
+		GPrevPerBoneMotionBlur.UnlockData();
+	}
+
 	return bDirty;
 }
 
diff --git a/Engine/Source/Runtime/Renderer/Private/PostProcess/SceneRenderTargets.cpp b/Engine/Source/Runtime/Renderer/Private/PostProcess/SceneRenderTargets.cpp
--- a/Engine/Source/Runtime/Renderer/Private/PostProcess/SceneRenderTargets.cpp
+++ b/Engine/Source/Runtime/Renderer/Private/PostProcess/SceneRenderTargets.cpp
@@ -302,24 +302,32 @@ void FSceneRenderTargets::BeginRenderingGBuffer(FRHICommandList& RHICmdList, ERe
 	// Set the scene color surface as the render target, and the scene depth surface as the depth-stencil target.
 	if (CurrentFeatureLevel >= ERHIFeatureLevel::SM4)
 	{
-		FRHIRenderTargetView RenderTargets[6];
+		const bool bUseVelocityGBuffer = FVelocityRendering::OutputsToGBuffer();
+		FRHIRenderTargetView RenderTargets[7];
 		RenderTargets[0] = FRHIRenderTargetView(GetSceneColorSurface(), 0, -1, ColorLoadAction, ERenderTargetStoreAction::EStore);
 		RenderTargets[1] = FRHIRenderTargetView(GSceneRenderTargets.GBufferA->GetRenderTargetItem().TargetableTexture, 0, -1, ColorLoadAction, ERenderTargetStoreAction::EStore);
 		RenderTargets[2] = FRHIRenderTargetView(GSceneRenderTargets.GBufferB->GetRenderTargetItem().TargetableTexture, 0, -1, ColorLoadAction, ERenderTargetStoreAction::EStore);
 		RenderTargets[3] = FRHIRenderTargetView(GSceneRenderTargets.GBufferC->GetRenderTargetItem().TargetableTexture, 0, -1, ColorLoadAction, ERenderTargetStoreAction::EStore);
 		RenderTargets[4] = FRHIRenderTargetView(GSceneRenderTargets.GBufferD->GetRenderTargetItem().TargetableTexture, 0, -1, ColorLoadAction, ERenderTargetStoreAction::EStore);
+		int32 MRTCount = 5;	// Derived from previously used RenderTargets[]; would have been nice to keep a pointer while filling it but it's more confusing to use
 
-		uint32 MRTCount = ARRAY_COUNT(RenderTargets);
+		int32 VelocityRTIndex = -1;
 
 		if (bAllowStaticLighting)
 		{
-			RenderTargets[5] = FRHIRenderTargetView(GSceneRenderTargets.GBufferE->GetRenderTargetItem().TargetableTexture, 0, -1, ColorLoadAction, ERenderTargetStoreAction::EStore);
+			RenderTargets[MRTCount] = FRHIRenderTargetView(GSceneRenderTargets.GBufferE->GetRenderTargetItem().TargetableTexture, 0, -1, ColorLoadAction, ERenderTargetStoreAction::EStore);
+			++MRTCount;
 		}
-		else
+
+		if (bUseVelocityGBuffer)
 		{
-			MRTCount--;
+			VelocityRTIndex = MRTCount;
+			++MRTCount;
+			RenderTargets[VelocityRTIndex] = FRHIRenderTargetView(GSceneRenderTargets.GBufferVelocity->GetRenderTargetItem().TargetableTexture, 0, -1, ColorLoadAction, ERenderTargetStoreAction::EStore);
 		}
 
+		check(MRTCount <= ARRAY_COUNT(RenderTargets));
+
 		FRHIDepthRenderTargetView DepthView(GetSceneDepthSurface(), DepthLoadAction, ERenderTargetStoreAction::EStore);
 		FRHISetRenderTargetsInfo Info(MRTCount, RenderTargets, DepthView);
 		if (ColorLoadAction == ERenderTargetLoadAction::EClear)
@@ -328,6 +336,11 @@ void FSceneRenderTargets::BeginRenderingGBuffer(FRHICommandList& RHICmdList, ERe
 			Info.ClearColors[1] = Info.ClearColors[2] = Info.ClearColors[3] = FLinearColor::Transparent;
 			Info.ClearColors[4] = FLinearColor(0, 1, 1, 1);
 			Info.ClearColors[5] = FLinearColor(1, 1, 1, 1);
+
+			if (bUseVelocityGBuffer)
+			{
+				Info.ClearColors[VelocityRTIndex] = FLinearColor(0, 0, 0, 0);
+			}
 		}
 		Info.DepthClearValue = 0.0f;
 
@@ -343,6 +356,11 @@ int32 FSceneRenderTargets::GetNumGBufferTargets() const
 	if (CurrentFeatureLevel >= ERHIFeatureLevel::SM4 && !IsSimpleDynamicLightingEnabled())
 	{
 		NumGBufferTargets = bAllowStaticLighting ? 6 : 5;
+
+		if (FVelocityRendering::OutputsToGBuffer())
+		{
+			++NumGBufferTargets;
+		}
 	}
 	return NumGBufferTargets;
 }
@@ -437,6 +455,7 @@ void FSceneRenderTargets::ReleaseGBufferTargets()
 	GBufferC.SafeRelease();
 	GBufferD.SafeRelease();
 	GBufferE.SafeRelease();
+	GBufferVelocity.SafeRelease();
 }
 
 void FSceneRenderTargets::AllocGBufferTargets()
@@ -504,6 +523,13 @@ void FSceneRenderTargets::AllocGBufferTargets()
 		GRenderTargetPool.FindFreeElement(Desc, GBufferE, TEXT("GBufferE"));
 	}
 
+	const bool bUseVelocityGBuffer = FVelocityRendering::OutputsToGBuffer();
+	if (bUseVelocityGBuffer)
+	{
+		FPooledRenderTargetDesc VelocityRTDesc = FVelocityRendering::GetRenderTargetDesc();
+		GRenderTargetPool.FindFreeElement(VelocityRTDesc, GBufferVelocity, TEXT("GBufferVelocity"));
+	}
+
 	// otherwise we have a severe problem
 	check(GBufferA);
 
@@ -517,6 +543,7 @@ void FSceneRenderTargets::AllocGBufferTargets()
 		const FSceneRenderTargetItem& GBufferC = this->GBufferC ? this->GBufferC->GetRenderTargetItem() : GSystemTextures.BlackDummy->GetRenderTargetItem();
 		const FSceneRenderTargetItem& GBufferD = this->GBufferD ? this->GBufferD->GetRenderTargetItem() : GSystemTextures.BlackDummy->GetRenderTargetItem();
 		const FSceneRenderTargetItem& GBufferE = this->GBufferE ? this->GBufferE->GetRenderTargetItem() : GSystemTextures.BlackDummy->GetRenderTargetItem();
+		const FSceneRenderTargetItem& GBufferVelocity = this->GBufferVelocity ? this->GBufferVelocity->GetRenderTargetItem() : GSystemTextures.BlackDummy->GetRenderTargetItem();
 
 		FGBufferResourceStruct GBufferResourceStruct;
 
@@ -525,24 +552,28 @@ void FSceneRenderTargets::AllocGBufferTargets()
 		GBufferResourceStruct.GBufferCTexture = GBufferC.ShaderResourceTexture;
 		GBufferResourceStruct.GBufferDTexture = GBufferD.ShaderResourceTexture;
 		GBufferResourceStruct.GBufferETexture = GBufferE.ShaderResourceTexture;
+		GBufferResourceStruct.GBufferVelocityTexture = GBufferVelocity.ShaderResourceTexture;
 
 		GBufferResourceStruct.GBufferATextureNonMS = GBufferA.ShaderResourceTexture;
 		GBufferResourceStruct.GBufferBTextureNonMS = GBufferB.ShaderResourceTexture;
 		GBufferResourceStruct.GBufferCTextureNonMS = GBufferC.ShaderResourceTexture;
 		GBufferResourceStruct.GBufferDTextureNonMS = GBufferD.ShaderResourceTexture;
 		GBufferResourceStruct.GBufferETextureNonMS = GBufferE.ShaderResourceTexture;
+		GBufferResourceStruct.GBufferVelocityTextureNonMS = GBufferVelocity.ShaderResourceTexture;
 
 		GBufferResourceStruct.GBufferATextureMS = GBufferA.TargetableTexture;
 		GBufferResourceStruct.GBufferBTextureMS = GBufferB.TargetableTexture;
 		GBufferResourceStruct.GBufferCTextureMS = GBufferC.TargetableTexture;
 		GBufferResourceStruct.GBufferDTextureMS = GBufferD.TargetableTexture;
 		GBufferResourceStruct.GBufferETextureMS = GBufferE.TargetableTexture;
+		GBufferResourceStruct.GBufferVelocityTextureMS = GBufferVelocity.TargetableTexture;
 
 		GBufferResourceStruct.GBufferATextureSampler = TStaticSamplerState<>::GetRHI();
 		GBufferResourceStruct.GBufferBTextureSampler = TStaticSamplerState<>::GetRHI();
 		GBufferResourceStruct.GBufferCTextureSampler = TStaticSamplerState<>::GetRHI();
 		GBufferResourceStruct.GBufferDTextureSampler = TStaticSamplerState<>::GetRHI();
 		GBufferResourceStruct.GBufferETextureSampler = TStaticSamplerState<>::GetRHI();
+		GBufferResourceStruct.GBufferVelocityTextureSampler = TStaticSamplerState<>::GetRHI();
 
 		GBufferResourcesUniformBuffer = FGBufferResourceStruct::CreateUniformBuffer(GBufferResourceStruct, UniformBuffer_SingleFrame);
 	}
@@ -781,6 +812,11 @@ void FSceneRenderTargets::ResolveGBufferSurfaces(FRHICommandList& RHICmdList, co
 		{
 			RHICmdList.CopyToResolveTarget(GSceneRenderTargets.GBufferE->GetRenderTargetItem().TargetableTexture, GSceneRenderTargets.GBufferE->GetRenderTargetItem().ShaderResourceTexture, true, FResolveParams(ResolveRect));
 		}
+
+		if (FVelocityRendering::OutputsToGBuffer())
+		{
+			RHICmdList.CopyToResolveTarget(GSceneRenderTargets.GBufferVelocity->GetRenderTargetItem().TargetableTexture, GSceneRenderTargets.GBufferVelocity->GetRenderTargetItem().ShaderResourceTexture, true, FResolveParams(ResolveRect));
+		}
 	}
 }
 
@@ -1369,6 +1405,12 @@ void FSceneRenderTargets::AllocateDeferredShadingPathRenderTargets()
 			}
 		}
 	}
+
+	if (FVelocityRendering::OutputsToGBuffer())
+	{
+		FPooledRenderTargetDesc VelocityRTDesc = FVelocityRendering::GetRenderTargetDesc();
+		GRenderTargetPool.FindFreeElement(VelocityRTDesc, GBufferVelocity, TEXT("GBufferVelocity"));
+	}
 }
 
 EPixelFormat FSceneRenderTargets::GetSceneColorFormat() const
@@ -1549,6 +1591,16 @@ const FTextureRHIRef& FSceneRenderTargets::GetSceneColorTexture() const
 	return (const FTextureRHIRef&)GetSceneColor()->GetRenderTargetItem().ShaderResourceTexture; 
 }
 
+IPooledRenderTarget* FSceneRenderTargets::GetGBufferVelocityRT()
+{
+	if (!FVelocityRendering::OutputsToGBuffer())
+	{
+		return nullptr;
+	}
+	
+	return GBufferVelocity;
+}
+
 IPooledRenderTarget* FSceneRenderTargets::RequestCustomDepth(bool bPrimitives)
 {
 	int Value = CVarCustomDepth.GetValueOnRenderThread();
diff --git a/Engine/Source/Runtime/Renderer/Private/PostProcess/SceneRenderTargets.h b/Engine/Source/Runtime/Renderer/Private/PostProcess/SceneRenderTargets.h
--- a/Engine/Source/Runtime/Renderer/Private/PostProcess/SceneRenderTargets.h
+++ b/Engine/Source/Runtime/Renderer/Private/PostProcess/SceneRenderTargets.h
@@ -82,21 +82,25 @@ BEGIN_UNIFORM_BUFFER_STRUCT(FGBufferResourceStruct, )
 	DECLARE_UNIFORM_BUFFER_STRUCT_MEMBER_TEXTURE( Texture2D,			GBufferCTexture )
 	DECLARE_UNIFORM_BUFFER_STRUCT_MEMBER_TEXTURE( Texture2D,			GBufferDTexture )
 	DECLARE_UNIFORM_BUFFER_STRUCT_MEMBER_TEXTURE( Texture2D,			GBufferETexture )
+	DECLARE_UNIFORM_BUFFER_STRUCT_MEMBER_TEXTURE( Texture2D,			GBufferVelocityTexture )
 	DECLARE_UNIFORM_BUFFER_STRUCT_MEMBER_TEXTURE( Texture2D<float4>,	GBufferATextureNonMS )
 	DECLARE_UNIFORM_BUFFER_STRUCT_MEMBER_TEXTURE( Texture2D<float4>,	GBufferBTextureNonMS )
 	DECLARE_UNIFORM_BUFFER_STRUCT_MEMBER_TEXTURE( Texture2D<float4>,	GBufferCTextureNonMS )
 	DECLARE_UNIFORM_BUFFER_STRUCT_MEMBER_TEXTURE( Texture2D<float4>,	GBufferDTextureNonMS )
 	DECLARE_UNIFORM_BUFFER_STRUCT_MEMBER_TEXTURE( Texture2D<float4>,	GBufferETextureNonMS )
+	DECLARE_UNIFORM_BUFFER_STRUCT_MEMBER_TEXTURE( Texture2D<float4>,	GBufferVelocityTextureNonMS )
 	DECLARE_UNIFORM_BUFFER_STRUCT_MEMBER_TEXTURE( Texture2DMS<float4>,	GBufferATextureMS )
 	DECLARE_UNIFORM_BUFFER_STRUCT_MEMBER_TEXTURE( Texture2DMS<float4>,	GBufferBTextureMS )
 	DECLARE_UNIFORM_BUFFER_STRUCT_MEMBER_TEXTURE( Texture2DMS<float4>,	GBufferCTextureMS )
 	DECLARE_UNIFORM_BUFFER_STRUCT_MEMBER_TEXTURE( Texture2DMS<float4>,	GBufferDTextureMS )
 	DECLARE_UNIFORM_BUFFER_STRUCT_MEMBER_TEXTURE( Texture2DMS<float4>,	GBufferETextureMS )
+	DECLARE_UNIFORM_BUFFER_STRUCT_MEMBER_TEXTURE( Texture2DMS<float4>,	GBufferVelocityTextureMS )
 	DECLARE_UNIFORM_BUFFER_STRUCT_MEMBER_SAMPLER( SamplerState,			GBufferATextureSampler )
 	DECLARE_UNIFORM_BUFFER_STRUCT_MEMBER_SAMPLER( SamplerState,			GBufferBTextureSampler )
 	DECLARE_UNIFORM_BUFFER_STRUCT_MEMBER_SAMPLER( SamplerState,			GBufferCTextureSampler )
 	DECLARE_UNIFORM_BUFFER_STRUCT_MEMBER_SAMPLER( SamplerState,			GBufferDTextureSampler )
 	DECLARE_UNIFORM_BUFFER_STRUCT_MEMBER_SAMPLER( SamplerState,			GBufferETextureSampler )
+	DECLARE_UNIFORM_BUFFER_STRUCT_MEMBER_SAMPLER( SamplerState,			GBufferVelocityTextureSampler)
 END_UNIFORM_BUFFER_STRUCT( FGBufferResourceStruct )
 
 /**
@@ -345,6 +349,8 @@ public:
 		return (const FTexture2DRHIRef&)AuxiliarySceneDepthZ->GetRenderTargetItem().TargetableTexture; 
 	}
 
+	IPooledRenderTarget* GetGBufferVelocityRT();
+
 	// @return can be 0 if the feature is disabled
 	IPooledRenderTarget* RequestCustomDepth(bool bPrimitives);
 
@@ -450,6 +456,8 @@ public:
 	TRefCountPtr<IPooledRenderTarget> GBufferD;
 	TRefCountPtr<IPooledRenderTarget> GBufferE;
 
+	TRefCountPtr<IPooledRenderTarget> GBufferVelocity;
+
 	// DBuffer: For decals before base pass (only temporarily available after early z pass and until base pass)
 	TRefCountPtr<IPooledRenderTarget> DBufferA;
 	TRefCountPtr<IPooledRenderTarget> DBufferB;
diff --git a/Engine/Source/Runtime/Renderer/Private/RendererScene.cpp b/Engine/Source/Runtime/Renderer/Private/RendererScene.cpp
--- a/Engine/Source/Runtime/Renderer/Private/RendererScene.cpp
+++ b/Engine/Source/Runtime/Renderer/Private/RendererScene.cpp
@@ -2606,6 +2606,11 @@ void FMotionBlurInfoData::ApplyOffset(FVector InOffset)
 	}
 }
 
+const FMotionBlurInfo* FMotionBlurInfoData::FindMBInfoIndex(FPrimitiveComponentId ComponentId) const
+{
+	return MotionBlurInfos.Find(ComponentId);
+}
+
 FMotionBlurInfo* FMotionBlurInfoData::FindMBInfoIndex(FPrimitiveComponentId ComponentId)
 {
 	return MotionBlurInfos.Find(ComponentId);
@@ -2628,4 +2633,21 @@ bool FMotionBlurInfoData::GetPrimitiveMotionBlurInfo(const FPrimitiveSceneInfo*
 	return false;
 }
 
+bool FMotionBlurInfoData::GetPrimitiveMotionBlurInfo(const FPrimitiveSceneInfo* PrimitiveSceneInfo, FMatrix& OutPreviousLocalToWorld) const
+{
+	check(IsInParallelRenderingThread());
+
+	if (PrimitiveSceneInfo && PrimitiveSceneInfo->PrimitiveComponentId.IsValid())
+	{
+		const FMotionBlurInfo* MotionBlurInfo = FindMBInfoIndex(PrimitiveSceneInfo->PrimitiveComponentId);
+
+		if (MotionBlurInfo)
+		{
+			OutPreviousLocalToWorld = MotionBlurInfo->GetPreviousLocalToWorld();
+			return true;
+		}
+	}
+	return false;
+}
+
 #endif
diff --git a/Engine/Source/Runtime/Renderer/Private/ScenePrivate.h b/Engine/Source/Runtime/Renderer/Private/ScenePrivate.h
--- a/Engine/Source/Runtime/Renderer/Private/ScenePrivate.h
+++ b/Engine/Source/Runtime/Renderer/Private/ScenePrivate.h
@@ -1635,4 +1635,6 @@ private:
 	ERHIFeatureLevel::Type FeatureLevel;
 };
 
+#include "BasePassRendering.inl"
+
 #endif // __SCENEPRIVATE_H__
diff --git a/Engine/Source/Runtime/Renderer/Private/VelocityRendering.cpp b/Engine/Source/Runtime/Renderer/Private/VelocityRendering.cpp
--- a/Engine/Source/Runtime/Renderer/Private/VelocityRendering.cpp
+++ b/Engine/Source/Runtime/Renderer/Private/VelocityRendering.cpp
@@ -8,6 +8,15 @@
 #include "../../Engine/Private/SkeletalRenderGPUSkin.h"		// GPrevPerBoneMotionBlur
 #include "SceneUtils.h"
 
+// Changing this causes a full shader recompile
+static TAutoConsoleVariable<int32> CVarBasePassOutputsVelocity(
+	TEXT("r.BasePassOutputsVelocity"),
+	0,
+	TEXT("Enables rendering WPO velocities on the base pass.\n") \
+	TEXT("0 - Renders in a separate pass/rendertarget, all movable static meshes + dynamic.\n") \
+	TEXT("1 - Renders during the regular base pass adding an extra GBuffer, but allowing motion blur on materials with Time-based WPO.\n"),
+	ECVF_ReadOnly | ECVF_RenderThreadSafe);
+
 //=============================================================================
 /** Encapsulates the Velocity vertex shader. */
 class FVelocityVS : public FMeshMaterialShader
@@ -354,8 +363,8 @@ void FVelocityDrawingPolicyFactory::AddStaticMesh(FScene* Scene, FStaticMesh* St
 	const FMaterialRenderProxy* MaterialRenderProxy = StaticMesh->MaterialRenderProxy;
 	const FMaterial* Material = MaterialRenderProxy->GetMaterial(FeatureLevel);
 
-	// Velocity only needs to be directly rendered for movable meshes, or those that modify WPO
-	if (StaticMesh->PrimitiveSceneInfo->Proxy->IsMovable() || Material->MaterialModifiesMeshPosition_RenderThread())
+	// Velocity only needs to be directly rendered for movable meshes
+	if (StaticMesh->PrimitiveSceneInfo->Proxy->IsMovable())
 	{
 	    EBlendMode BlendMode = Material->GetBlendMode();
 	    if(BlendMode == BLEND_Opaque || BlendMode == BLEND_Masked)
@@ -633,6 +642,7 @@ void FDeferredShadingSceneRenderer::RenderVelocities(FRHICommandListImmediate& R
 
 		bNeedsVelocity |= bMotionBlur || bTemporalAA;
 	}
+
 	if( !bNeedsVelocity || !GPixelFormats[PF_G16R16].Supported )
 	{
 		return;
@@ -643,10 +653,7 @@ void FDeferredShadingSceneRenderer::RenderVelocities(FRHICommandListImmediate& R
 
 	SCOPED_DRAW_EVENT(RHICmdList, RenderVelocities);
 
-	const FIntPoint BufferSize = GSceneRenderTargets.GetBufferSizeXY();
-	const FIntPoint VelocityBufferSize = BufferSize;		// full resolution so we can reuse the existing full res z buffer
-
-	FPooledRenderTargetDesc Desc(FPooledRenderTargetDesc::Create2DDesc(VelocityBufferSize, PF_G16R16, TexCreate_None, TexCreate_RenderTargetable, false));
+	FPooledRenderTargetDesc Desc = FVelocityRendering::GetRenderTargetDesc();
 	GRenderTargetPool.FindFreeElement(Desc, VelocityRT, TEXT("Velocity"));
 
 	GPrevPerBoneMotionBlur.LockData();
@@ -676,3 +683,15 @@ void FDeferredShadingSceneRenderer::RenderVelocities(FRHICommandListImmediate& R
 	// to be able to observe results with VisualizeTexture
 	GRenderTargetPool.VisualizeTexture.SetCheckPoint(RHICmdList, VelocityRT);
 }
+
+FPooledRenderTargetDesc FVelocityRendering::GetRenderTargetDesc()
+{
+	const FIntPoint BufferSize = GSceneRenderTargets.GetBufferSizeXY();
+	const FIntPoint VelocityBufferSize = BufferSize;		// full resolution so we can reuse the existing full res z buffer
+	return FPooledRenderTargetDesc(FPooledRenderTargetDesc::Create2DDesc(VelocityBufferSize, PF_G16R16, TexCreate_None, TexCreate_RenderTargetable, false));
+}
+
+bool FVelocityRendering::OutputsToGBuffer()
+{
+	return CVarBasePassOutputsVelocity.GetValueOnAnyThread() == 1;
+}
diff --git a/Engine/Source/Runtime/Renderer/Private/VelocityRendering.h b/Engine/Source/Runtime/Renderer/Private/VelocityRendering.h
--- a/Engine/Source/Runtime/Renderer/Private/VelocityRendering.h
+++ b/Engine/Source/Runtime/Renderer/Private/VelocityRendering.h
@@ -80,3 +80,11 @@ int32 GetMotionBlurQualityFromCVar();
 
 /** If this view need motion blur processing */
 bool IsMotionBlurEnabled(const FViewInfo& View);
+
+// Group Velocity Rendering accessors, types, etc.
+struct FVelocityRendering
+{
+	static FPooledRenderTargetDesc GetRenderTargetDesc();
+
+	static bool OutputsToGBuffer();
+};
