Fix for UE-7660 - Grid in orthographic view becomes distorted in Preview Rendering Level ES2
- Forward renderer clears background color in editor mode.
- Remove TAA jitter from proj matrices after ES2 TAA post process completes. (allows stable rendering of editor primitives)
- Moved ES2 post process editor composite primitive pass to after tonemap. (fixes incorrect grid colour)
- Add opacity mask for editor prim materials to forward base pass shader. (fixes perspective grid rendering)
- Add Basic ES2 post process pass for !mobileHDR scenario, enabling editor's primitives to still be rendered.
- DrawNewGrid only bails out if the material it's about to use hasnt compiled. (allows ES2 to render 'texture solution' grid implementation)
- Prevent some view elements drawing 2x when ShowFlags.CompositeEditorPrimitives == true.

diff --git a/Engine/Shaders/BasePassForForwardShadingPixelShader.usf b/Engine/Shaders/BasePassForForwardShadingPixelShader.usf
--- a/Engine/Shaders/BasePassForForwardShadingPixelShader.usf
+++ b/Engine/Shaders/BasePassForForwardShadingPixelShader.usf
@@ -321,6 +321,10 @@ void Main(
 
 	// Editor primitive depth testing
 	#if EDITOR_PRIMITIVE_MATERIAL && (ES2_EMULATION)
+		#if MATERIALBLENDING_MASKED
+			// some material might have a opacity value
+			OutColor.a = GetMaterialMaskInputRaw(MaterialParameters);
+		#endif
 		// we output premultiplied alpha to we have to darken all 4 channels
 		OutColor *= ClipForEditorPrimitives(MaterialParameters);
 		clip(OutColor.a - GetMaterialOpacityMaskClipValue());
diff --git a/Engine/Source/Editor/UnrealEd/Private/EditorComponents.cpp b/Engine/Source/Editor/UnrealEd/Private/EditorComponents.cpp
--- a/Engine/Source/Editor/UnrealEd/Private/EditorComponents.cpp
+++ b/Engine/Source/Editor/UnrealEd/Private/EditorComponents.cpp
@@ -113,17 +113,17 @@ static void GetAxisColors(FLinearColor Out[3], bool b3D)
 
 void FGridWidget::DrawNewGrid(const FSceneView* View, FPrimitiveDrawInterface* PDI)
 {
-	if (LevelGridMaterial->IsCompilingOrHadCompileError(View->GetFeatureLevel()) || LevelGridMaterial2->IsCompilingOrHadCompileError(View->GetFeatureLevel()))
+	bool bUseTextureSolution = CVarEditorNewLevelGrid.GetValueOnGameThread() > 1;
+	UMaterial* GridMaterial = bUseTextureSolution ? LevelGridMaterial2 : LevelGridMaterial;
+	UMaterialInstanceDynamic* MaterialInst = bUseTextureSolution ? LevelGridMaterialInst2 : LevelGridMaterialInst;
+
+	if (GridMaterial->IsCompilingOrHadCompileError(View->GetFeatureLevel()))
 	{
 		// The material would appear to be black (because we don't use a MaterialDomain but a UsageFlag - we should change that).
 		// Here we rather want to hide it.
 		return;
 	}
 
-	bool bUseTextureSolution = CVarEditorNewLevelGrid.GetValueOnGameThread() > 1;
-
-	UMaterialInstanceDynamic *MaterialInst = bUseTextureSolution ? LevelGridMaterialInst2 : LevelGridMaterialInst;
-
 	if(!MaterialInst)
 	{
 		return;
diff --git a/Engine/Source/Runtime/Renderer/Private/ForwardBasePassRendering.cpp b/Engine/Source/Runtime/Renderer/Private/ForwardBasePassRendering.cpp
--- a/Engine/Source/Runtime/Renderer/Private/ForwardBasePassRendering.cpp
+++ b/Engine/Source/Runtime/Renderer/Private/ForwardBasePassRendering.cpp
@@ -414,19 +414,22 @@ void FForwardShadingSceneRenderer::RenderForwardShadingBasePass(FRHICommandListI
 
 			View.SimpleElementCollector.DrawBatchedElements(RHICmdList, View, NULL, EBlendModeFilter::OpaqueAndMasked);
 
-			const bool bNeedToSwitchVerticalAxis = RHINeedsToSwitchVerticalAxis(GShaderPlatformForFeatureLevel[FeatureLevel]);
+			if (!View.Family->EngineShowFlags.CompositeEditorPrimitives)
+			{
+				const bool bNeedToSwitchVerticalAxis = RHINeedsToSwitchVerticalAxis(GShaderPlatformForFeatureLevel[FeatureLevel]);
+
+				// Draw the base pass for the view's batched mesh elements.
+				DrawViewElements<FBasePassForwardOpaqueDrawingPolicyFactory>(RHICmdList, View, FBasePassForwardOpaqueDrawingPolicyFactory::ContextType(false, ESceneRenderTargetsMode::DontSet), SDPG_World, true);
 
-			// Draw the base pass for the view's batched mesh elements.
-			DrawViewElements<FBasePassForwardOpaqueDrawingPolicyFactory>(RHICmdList, View, FBasePassForwardOpaqueDrawingPolicyFactory::ContextType(false, ESceneRenderTargetsMode::DontSet), SDPG_World, true);
-			
-			// Draw the view's batched simple elements(lines, sprites, etc).
-			View.BatchedViewElements.Draw(RHICmdList, FeatureLevel, bNeedToSwitchVerticalAxis, View.ViewProjectionMatrix, View.ViewRect.Width(), View.ViewRect.Height(), false);
+				// Draw the view's batched simple elements(lines, sprites, etc).
+				View.BatchedViewElements.Draw(RHICmdList, FeatureLevel, bNeedToSwitchVerticalAxis, View.ViewProjectionMatrix, View.ViewRect.Width(), View.ViewRect.Height(), false);
 
-			// Draw foreground objects last
-			DrawViewElements<FBasePassForwardOpaqueDrawingPolicyFactory>(RHICmdList, View, FBasePassForwardOpaqueDrawingPolicyFactory::ContextType(false, ESceneRenderTargetsMode::DontSet), SDPG_Foreground, true);
+				// Draw foreground objects last
+				DrawViewElements<FBasePassForwardOpaqueDrawingPolicyFactory>(RHICmdList, View, FBasePassForwardOpaqueDrawingPolicyFactory::ContextType(false, ESceneRenderTargetsMode::DontSet), SDPG_Foreground, true);
 
-			// Draw the view's batched simple elements(lines, sprites, etc).
-			View.TopBatchedViewElements.Draw(RHICmdList, FeatureLevel, bNeedToSwitchVerticalAxis, View.ViewProjectionMatrix, View.ViewRect.Width(), View.ViewRect.Height(), false);
+				// Draw the view's batched simple elements(lines, sprites, etc).
+				View.TopBatchedViewElements.Draw(RHICmdList, FeatureLevel, bNeedToSwitchVerticalAxis, View.ViewProjectionMatrix, View.ViewRect.Width(), View.ViewRect.Height(), false);
+			}
 		}
 
 		// Issue static draw list masked draw calls last, as PVR wants it
diff --git a/Engine/Source/Runtime/Renderer/Private/ForwardShadingRenderer.cpp b/Engine/Source/Runtime/Renderer/Private/ForwardShadingRenderer.cpp
--- a/Engine/Source/Runtime/Renderer/Private/ForwardShadingRenderer.cpp
+++ b/Engine/Source/Runtime/Renderer/Private/ForwardShadingRenderer.cpp
@@ -107,17 +107,27 @@ void FForwardShadingSceneRenderer::Render(FRHICommandListImmediate& RHICmdList)
 	FGlobalDynamicVertexBuffer::Get().Commit();
 	FGlobalDynamicIndexBuffer::Get().Commit();
 
+	// This might eventually be a problem with multiple views.
+	// Using only view 0 to check to do on-chip transform of alpha.
+	FViewInfo& View = Views[0];
+
 	const bool bGammaSpace = !IsMobileHDR();
 	const bool bRequiresUpscale = ((uint32)ViewFamily.RenderTarget->GetSizeXY().X > ViewFamily.FamilySizeX || (uint32)ViewFamily.RenderTarget->GetSizeXY().Y > ViewFamily.FamilySizeY);
+	const bool bRenderToScene = bRequiresUpscale || FSceneRenderer::ShouldCompositeEditorPrimitives(View);
 
-	if (bGammaSpace && !bRequiresUpscale)
+	if (bGammaSpace && !bRenderToScene)
 	{
 		SetRenderTarget(RHICmdList, ViewFamily.RenderTarget->GetRenderTargetTexture(), GSceneRenderTargets.GetSceneDepthTexture(), ESimpleRenderTargetMode::EClearToDefault);
 	}
 	else
 	{
 		// Begin rendering to scene color
-        GSceneRenderTargets.BeginRenderingSceneColor(RHICmdList, ESimpleRenderTargetMode::EClearToDefault);
+		GSceneRenderTargets.BeginRenderingSceneColor(RHICmdList, ESimpleRenderTargetMode::EClearToDefault);
+	}
+
+	if (GIsEditor)
+	{
+		RHICmdList.Clear(true, Views[0].BackgroundColor, false, 0, false, 0, FIntRect());
 	}
 
 	RenderForwardShadingBasePass(RHICmdList);
@@ -148,10 +158,6 @@ void FForwardShadingSceneRenderer::Render(FRHICommandListImmediate& RHICmdList)
 		RenderTranslucency(RHICmdList);
 	}
 
-	// This might eventually be a problem with multiple views.
-	// Using only view 0 to check to do on-chip transform of alpha.
-	FViewInfo& View = Views[0];
-
 	static const auto CVarMobileMSAA = IConsoleManager::Get().FindTConsoleVariableDataInt(TEXT("r.MobileMSAA"));
 	bool bOnChipSunMask =
 		GSupportsRenderTargetFormat_PF_FloatRGBA &&
@@ -176,7 +182,7 @@ void FForwardShadingSceneRenderer::Render(FRHICommandListImmediate& RHICmdList)
 		CompositeContext.Process(TEXT("OnChipAlphaTransform"));
 	}
 
-	if (!bGammaSpace || bRequiresUpscale)
+	if (!bGammaSpace || bRenderToScene)
 	{
 		// Resolve the scene color for post processing.
 		GSceneRenderTargets.ResolveSceneColor(RHICmdList, FResolveRect(0, 0, ViewFamily.FamilySizeX, ViewFamily.FamilySizeY));
@@ -212,48 +218,31 @@ void FForwardShadingSceneRenderer::Render(FRHICommandListImmediate& RHICmdList)
 			}
 		}
 	}
-	else if (bRequiresUpscale)
+	else if (bRenderToScene)
 	{
-		SimpleUpscale(RHICmdList, View);
+		BasicPostProcess(RHICmdList, View, bRequiresUpscale, FSceneRenderer::ShouldCompositeEditorPrimitives(View));
 	}
 	RenderFinish(RHICmdList);
 }
 
-// Perform simple upscale when required and post process is not in use.
-void FForwardShadingSceneRenderer::SimpleUpscale(FRHICommandListImmediate& RHICmdList, FViewInfo &View)
+// Perform simple upscale and/or editor primitive composite if the fully-featured post process is not in use.
+void FForwardShadingSceneRenderer::BasicPostProcess(FRHICommandListImmediate& RHICmdList, FViewInfo &View, bool bDoUpscale, bool bDoEditorPrimitives)
 {
 	FRenderingCompositePassContext CompositeContext(RHICmdList, View);
 	FPostprocessContext Context(CompositeContext.Graph, View);
-	// simple bilinear upscaling for ES2.
-	FRenderingCompositePass* Node = Context.Graph.RegisterPass(new(FMemStack::Get()) FRCPassPostProcessUpscale(1, 0.0f));
 
-	Node->SetInput(ePId_Input0, FRenderingCompositeOutputRef(Context.FinalOutput));
-	Node->SetInput(ePId_Input1, FRenderingCompositeOutputRef(Context.FinalOutput));
+	if (bDoUpscale)
+	{	// simple bilinear upscaling for ES2.
+		FRenderingCompositePass* Node = Context.Graph.RegisterPass(new(FMemStack::Get()) FRCPassPostProcessUpscale(1, 0.0f));
 
-	Context.FinalOutput = FRenderingCompositeOutputRef(Node);
+		Node->SetInput(ePId_Input0, FRenderingCompositeOutputRef(Context.FinalOutput));
+		Node->SetInput(ePId_Input1, FRenderingCompositeOutputRef(Context.FinalOutput));
 
-	{
-		// currently created on the heap each frame but View.Family->RenderTarget could keep this object and all would be cleaner
-		TRefCountPtr<IPooledRenderTarget> Temp;
-		FSceneRenderTargetItem Item;
-		Item.TargetableTexture = (FTextureRHIRef&)View.Family->RenderTarget->GetRenderTargetTexture();
-		Item.ShaderResourceTexture = (FTextureRHIRef&)View.Family->RenderTarget->GetRenderTargetTexture();
-
-		FPooledRenderTargetDesc Desc;
-
-		Desc.Extent = View.Family->RenderTarget->GetSizeXY();
-		// todo: this should come from View.Family->RenderTarget
-		Desc.Format = PF_B8G8R8A8;
-		Desc.NumMips = 1;
-
-		GRenderTargetPool.CreateUntrackedElement(Desc, Temp, Item);
-
-		Context.FinalOutput.GetOutput()->PooledRenderTarget = Temp;
-		Context.FinalOutput.GetOutput()->RenderTargetDesc = Desc;
+		Context.FinalOutput = FRenderingCompositeOutputRef(Node);
 	}
 
 	// Composite editor primitives if we had any to draw and compositing is enabled
-	if (FSceneRenderer::ShouldCompositeEditorPrimitives(View))
+	if (bDoEditorPrimitives)
 	{
 		FRenderingCompositePass* EditorCompNode = Context.Graph.RegisterPass(new(FMemStack::Get()) FRCPassPostProcessCompositeEditorPrimitives(false));
 		EditorCompNode->SetInput(ePId_Input0, FRenderingCompositeOutputRef(Context.FinalOutput));
@@ -261,6 +250,24 @@ void FForwardShadingSceneRenderer::SimpleUpscale(FRHICommandListImmediate& RHICm
 		Context.FinalOutput = FRenderingCompositeOutputRef(EditorCompNode);
 	}
 
+	// currently created on the heap each frame but View.Family->RenderTarget could keep this object and all would be cleaner
+	TRefCountPtr<IPooledRenderTarget> Temp;
+	FSceneRenderTargetItem Item;
+	Item.TargetableTexture = (FTextureRHIRef&)View.Family->RenderTarget->GetRenderTargetTexture();
+	Item.ShaderResourceTexture = (FTextureRHIRef&)View.Family->RenderTarget->GetRenderTargetTexture();
+
+	FPooledRenderTargetDesc Desc;
+
+	Desc.Extent = View.Family->RenderTarget->GetSizeXY();
+	// todo: this should come from View.Family->RenderTarget
+	Desc.Format = PF_B8G8R8A8;
+	Desc.NumMips = 1;
+
+	GRenderTargetPool.CreateUntrackedElement(Desc, Temp, Item);
+
+	Context.FinalOutput.GetOutput()->PooledRenderTarget = Temp;
+	Context.FinalOutput.GetOutput()->RenderTargetDesc = Desc;
+
 	CompositeContext.Root->AddDependency(Context.FinalOutput);
-	CompositeContext.Process(TEXT("ES2Upscale"));
+	CompositeContext.Process(TEXT("ES2BasicPostProcess"));
 }
diff --git a/Engine/Source/Runtime/Renderer/Private/PostProcess/PostProcessMobile.cpp b/Engine/Source/Runtime/Renderer/Private/PostProcess/PostProcessMobile.cpp
--- a/Engine/Source/Runtime/Renderer/Private/PostProcess/PostProcessMobile.cpp
+++ b/Engine/Source/Runtime/Renderer/Private/PostProcess/PostProcessMobile.cpp
@@ -2781,6 +2781,27 @@ void FRCPassPostProcessAaES2::Process(FRenderingCompositePassContext& Context)
 		EDRF_UseTriangleOptimization);
 
 	Context.RHICmdList.CopyToResolveTarget(DestRenderTarget.TargetableTexture, DestRenderTarget.ShaderResourceTexture, false, FResolveParams());
+
+	auto& View = Context.View;
+	if (FSceneRenderer::ShouldCompositeEditorPrimitives(View))
+	{
+		// Remove jitter (ensures editor prims are stable.)
+		View.ViewMatrices.ProjMatrix.M[2][0] = 0.0f;
+		View.ViewMatrices.ProjMatrix.M[2][1] = 0.0f;
+
+		// Compute the view projection matrix and its inverse.
+		View.ViewProjectionMatrix = View.ViewMatrices.ViewMatrix * View.ViewMatrices.ProjMatrix;
+		View.InvViewProjectionMatrix = View.ViewMatrices.GetInvProjMatrix() * View.InvViewMatrix;
+
+		/** The view transform, starting from world-space points translated by -ViewOrigin. */
+		FMatrix TranslatedViewMatrix = FTranslationMatrix(-View.ViewMatrices.PreViewTranslation) * View.ViewMatrices.ViewMatrix;
+
+		// Compute a transform from view origin centered world-space to clip space.
+		View.ViewMatrices.TranslatedViewProjectionMatrix = TranslatedViewMatrix * View.ViewMatrices.ProjMatrix;
+		View.ViewMatrices.InvTranslatedViewProjectionMatrix = View.ViewMatrices.TranslatedViewProjectionMatrix.Inverse();
+
+		View.InitRHIResources(nullptr);
+	}
 }
 
 FPooledRenderTargetDesc FRCPassPostProcessAaES2::ComputeOutputDesc(EPassOutputId InPassOutputId) const
diff --git a/Engine/Source/Runtime/Renderer/Private/PostProcess/PostProcessing.cpp b/Engine/Source/Runtime/Renderer/Private/PostProcess/PostProcessing.cpp
--- a/Engine/Source/Runtime/Renderer/Private/PostProcess/PostProcessing.cpp
+++ b/Engine/Source/Runtime/Renderer/Private/PostProcess/PostProcessing.cpp
@@ -1701,14 +1701,6 @@ void FPostProcessing::ProcessES2(FRHICommandListImmediate& RHICmdList, FViewInfo
 			}
 		}
 
-		// Composite editor primitives if we had any to draw and compositing is enabled
-		if (FSceneRenderer::ShouldCompositeEditorPrimitives(View) && !DofOutput.IsValid())
-		{
-			FRenderingCompositePass* EditorCompNode = Context.Graph.RegisterPass(new(FMemStack::Get()) FRCPassPostProcessCompositeEditorPrimitives(false));
-			EditorCompNode->SetInput(ePId_Input0, FRenderingCompositeOutputRef(Context.FinalOutput));
-			//Node->SetInput(ePId_Input1, FRenderingCompositeOutputRef(Context.SceneDepth));
-			Context.FinalOutput = FRenderingCompositeOutputRef(EditorCompNode);
-		}
 
 		// Must run to blit to back buffer even if post processing is off.
 		FRenderingCompositePass* PostProcessTonemap = Context.Graph.RegisterPass(new(FMemStack::Get()) FRCPassPostProcessTonemapES2(Context.View, FinalOutputViewRect, FinalTargetSize, bViewRectSource));
@@ -1740,16 +1732,12 @@ void FPostProcessing::ProcessES2(FRHICommandListImmediate& RHICmdList, FViewInfo
 			Context.FinalOutput = FRenderingCompositeOutputRef(PostProcessAa);
 		}
 
+		if (FSceneRenderer::ShouldCompositeEditorPrimitives(View) )
 		{
-			// temporary work around for EditorPrimitives / Dof alpha issues.
-			if (FSceneRenderer::ShouldCompositeEditorPrimitives(View) && DofOutput.IsValid())
-			{
-				// TODO: combine editor prims in tonemap pass or remove AA jitter
-				FRenderingCompositePass* EditorCompNode = Context.Graph.RegisterPass(new(FMemStack::Get()) FRCPassPostProcessCompositeEditorPrimitives(false));
-				EditorCompNode->SetInput(ePId_Input0, FRenderingCompositeOutputRef(Context.FinalOutput));
-				//Node->SetInput(ePId_Input1, FRenderingCompositeOutputRef(Context.SceneDepth));
-				Context.FinalOutput = FRenderingCompositeOutputRef(EditorCompNode);
-			}
+			FRenderingCompositePass* EditorCompNode = Context.Graph.RegisterPass(new(FMemStack::Get()) FRCPassPostProcessCompositeEditorPrimitives(false));
+			EditorCompNode->SetInput(ePId_Input0, FRenderingCompositeOutputRef(Context.FinalOutput));
+			//Node->SetInput(ePId_Input1, FRenderingCompositeOutputRef(Context.SceneDepth));
+			Context.FinalOutput = FRenderingCompositeOutputRef(EditorCompNode);
 		}
 
 		if(View.Family->EngineShowFlags.ShaderComplexity)
diff --git a/Engine/Source/Runtime/Renderer/Private/SceneRendering.h b/Engine/Source/Runtime/Renderer/Private/SceneRendering.h
--- a/Engine/Source/Runtime/Renderer/Private/SceneRendering.h
+++ b/Engine/Source/Runtime/Renderer/Private/SceneRendering.h
@@ -740,7 +740,7 @@ protected:
 	void RenderShadowDepthMaps(FRHICommandListImmediate& RHICmdList);
 
 	/** Perform upscaling when post process is not used. */
-	void SimpleUpscale(FRHICommandListImmediate& RHICmdList, FViewInfo &View);
+	void BasicPostProcess(FRHICommandListImmediate& RHICmdList, FViewInfo &View, bool bDoUpscale, bool bDoEditorPrimitives);
 
 	/**
 	  * Used by RenderShadowDepthMaps to render shadowmap for the given light.
