Disabled HZB by default as it tends be slower than occlusion queries in many scenes

diff --git a/Engine/Source/Runtime/Renderer/Private/SceneVisibility.cpp b/Engine/Source/Runtime/Renderer/Private/SceneVisibility.cpp
--- a/Engine/Source/Runtime/Renderer/Private/SceneVisibility.cpp
+++ b/Engine/Source/Runtime/Renderer/Private/SceneVisibility.cpp
@@ -57,7 +57,7 @@ static TAutoConsoleVariable<int32> CVarTemporalAASamples(
 #if PLATFORM_MAC // @todo: disabled until rendering problems with HZB occlusion in OpenGL are solved
 static int32 GHZBOcclusion = 0;
 #else
-static int32 GHZBOcclusion = 1;
+static int32 GHZBOcclusion = 0;
 #endif
 static FAutoConsoleVariableRef CVarHZBOcclusion(
 	TEXT("r.HZBOcclusion"),
