UE4 - Fix transform for motion blur; remove old path for static wpo materials on separate pass

diff --git a/Engine/Source/Runtime/Renderer/Private/BasePassRendering.inl b/Engine/Source/Runtime/Renderer/Private/BasePassRendering.inl
--- a/Engine/Source/Runtime/Renderer/Private/BasePassRendering.inl
+++ b/Engine/Source/Runtime/Renderer/Private/BasePassRendering.inl
@@ -17,19 +17,25 @@ inline void TBasePassVertexShaderBaseType<LightMapPolicyType>::SetMesh(FRHIComma
 	{
 		check(SkipOutputVelocityParameter.IsBound());
 
-		// previous transform can be stored in the scene for each primitive
 		FMatrix PreviousLocalToWorld;
-		const auto& Scene = (const FScene&)Proxy->GetScene();
-		//@todo-rco: Move to InitViews
-		if (Scene.MotionBlurInfoData.GetPrimitiveMotionBlurInfo(Proxy->GetPrimitiveSceneInfo(), PreviousLocalToWorld))
+		bool bHasPreviousLocalToWorld = false;
+		const auto& ViewInfo = (const FViewInfo&)View;
+		if (FVelocityDrawingPolicy::HasVelocityOnBasePass(ViewInfo, Proxy, Proxy->GetPrimitiveSceneInfo(), 
+			bHasPreviousLocalToWorld, PreviousLocalToWorld))
 		{
-			const auto& ViewInfo = (const FViewInfo&)View;
-			SetShaderValue(RHICmdList, VertexShaderRHI, PreviousLocalToWorldParameter, PreviousLocalToWorld.ConcatTranslation(ViewInfo.PrevViewMatrices.PreViewTranslation));
+			if (bHasPreviousLocalToWorld)
+			{
+				SetShaderValue(RHICmdList, VertexShaderRHI, PreviousLocalToWorldParameter, PreviousLocalToWorld.ConcatTranslation(ViewInfo.PrevViewMatrices.PreViewTranslation));
+			}
+			else
+			{
+				SetShaderValue(RHICmdList, VertexShaderRHI, PreviousLocalToWorldParameter, Proxy->GetLocalToWorld().ConcatTranslation(ViewInfo.PrevViewMatrices.PreViewTranslation));
+			}
+
 			SetShaderValue(RHICmdList, VertexShaderRHI, SkipOutputVelocityParameter, 0.0f);
 		}
 		else
 		{
-			SetShaderValue(RHICmdList, VertexShaderRHI, PreviousLocalToWorldParameter, Proxy->GetLocalToWorld());
 			SetShaderValue(RHICmdList, VertexShaderRHI, SkipOutputVelocityParameter, 1.0f);
 		}
 	}
diff --git a/Engine/Source/Runtime/Renderer/Private/PrimitiveSceneInfo.cpp b/Engine/Source/Runtime/Renderer/Private/PrimitiveSceneInfo.cpp
--- a/Engine/Source/Runtime/Renderer/Private/PrimitiveSceneInfo.cpp
+++ b/Engine/Source/Runtime/Renderer/Private/PrimitiveSceneInfo.cpp
@@ -478,7 +478,7 @@ bool FPrimitiveSceneInfo::ShouldRenderVelocity(const FViewInfo& View, bool bChec
 	}
 
 	// Only render primitives with velocity.
-	if (!FVelocityDrawingPolicy::HasVelocity(View, this, PrimitiveViewRelevance.bHasWorldPositionOffset))
+	if (!FVelocityDrawingPolicy::HasVelocity(View, this))
 	{
 		return false;
 	}
diff --git a/Engine/Source/Runtime/Renderer/Private/VelocityRendering.cpp b/Engine/Source/Runtime/Renderer/Private/VelocityRendering.cpp
--- a/Engine/Source/Runtime/Renderer/Private/VelocityRendering.cpp
+++ b/Engine/Source/Runtime/Renderer/Private/VelocityRendering.cpp
@@ -282,12 +282,12 @@ void FVelocityDrawingPolicy::SetMeshRenderState(
 	FMeshDrawingPolicy::SetMeshRenderState(RHICmdList, View, PrimitiveSceneProxy, Mesh, BatchElementIndex, bBackFace, ElementData, PolicyContext);
 }
 
-bool FVelocityDrawingPolicy::HasVelocity(const FViewInfo& View, const FPrimitiveSceneInfo* PrimitiveSceneInfo, bool bHasWorldPositionOffset)
+bool FVelocityDrawingPolicy::HasVelocity(const FViewInfo& View, const FPrimitiveSceneInfo* PrimitiveSceneInfo)
 {
 	checkSlow(IsInParallelRenderingThread());
 
 	// No velocity if motionblur is off, or if it's a non-moving object (treat as background in that case)
-	if (View.bCameraCut || (!PrimitiveSceneInfo->Proxy->IsMovable() && !bHasWorldPositionOffset))
+	if (View.bCameraCut || !PrimitiveSceneInfo->Proxy->IsMovable())
 	{
 		return false;
 	}
@@ -303,7 +303,6 @@ bool FVelocityDrawingPolicy::HasVelocity(const FViewInfo& View, const FPrimitive
 	}
 
 	// check if the primitive has moved
-	//if (!bHasWorldPositionOffset)
 	{
 		FMatrix PreviousLocalToWorld;
 
@@ -331,6 +330,47 @@ bool FVelocityDrawingPolicy::HasVelocity(const FViewInfo& View, const FPrimitive
 	return true;
 }
 
+bool FVelocityDrawingPolicy::HasVelocityOnBasePass(const FViewInfo& View,const FPrimitiveSceneProxy* PrimitiveSceneProxy, const FPrimitiveSceneInfo* PrimitiveSceneInfo, bool& bOutHasTransform, FMatrix& OutTransform)
+{
+	checkSlow(IsInParallelRenderingThread());
+	// No velocity if motionblur is off, or if it's a non-moving object (treat as background in that case)
+	if (View.bCameraCut)
+	{
+		return false;
+	}
+
+	//@todo-rco: Where is this set?
+	if (PrimitiveSceneInfo->bVelocityIsSupressed)
+	{
+		return false;
+	}
+
+	// hack
+	FScene* Scene = PrimitiveSceneInfo->Scene;
+	if (Scene->MotionBlurInfoData.GetPrimitiveMotionBlurInfo(PrimitiveSceneInfo, OutTransform))
+	{
+		const FMatrix& LocalToWorld = PrimitiveSceneProxy->GetLocalToWorld();
+/*
+		// Hasn't moved (treat as background by not rendering any special velocities)?
+		if (LocalToWorld.Equals(OutTransform, 0.0001f))
+		{
+			return false;
+		}
+*/
+		bOutHasTransform = true;
+		return true;
+	}
+
+	bOutHasTransform = false;
+	if (PrimitiveSceneProxy->IsMovable())
+	{
+		return true;
+	}
+
+	//@todo-rco: Optimize finding WPO!
+	return View.PrimitiveViewRelevanceMap[PrimitiveSceneInfo->GetIndex()].bHasWorldPositionOffset;
+}
+
 FBoundShaderStateInput FVelocityDrawingPolicy::GetBoundShaderStateInput(ERHIFeatureLevel::Type InFeatureLevel)
 {
 	return FBoundShaderStateInput(
diff --git a/Engine/Source/Runtime/Renderer/Private/VelocityRendering.h b/Engine/Source/Runtime/Renderer/Private/VelocityRendering.h
--- a/Engine/Source/Runtime/Renderer/Private/VelocityRendering.h
+++ b/Engine/Source/Runtime/Renderer/Private/VelocityRendering.h
@@ -47,7 +47,8 @@ public:
 	bool SupportsVelocity( ) const;
 
 	/** Determines whether this primitive has motionblur velocity to render */
-	static bool HasVelocity(const FViewInfo& View, const FPrimitiveSceneInfo* PrimitiveSceneInfo, bool bHasWorldPositionOffset);
+	static bool HasVelocity(const FViewInfo& View, const FPrimitiveSceneInfo* PrimitiveSceneInfo);
+	static bool HasVelocityOnBasePass(const FViewInfo& View,const FPrimitiveSceneProxy* Proxy, const FPrimitiveSceneInfo* PrimitiveSceneInfo, bool& bOutHasTransform, FMatrix& OutTransform);
 
 private:
 	class FVelocityVS* VertexShader;
