fixed long range shadows with large terrains

diff --git a/Engine/Source/Runtime/Renderer/Private/ShadowSetup.cpp b/Engine/Source/Runtime/Renderer/Private/ShadowSetup.cpp
--- a/Engine/Source/Runtime/Renderer/Private/ShadowSetup.cpp
+++ b/Engine/Source/Runtime/Renderer/Private/ShadowSetup.cpp
@@ -1950,7 +1950,8 @@ inline void FSceneRenderer::GatherShadowsForPrimitiveInner(
 					// Check if this primitive is in the shadow's cylinder
 					&& PrimitiveDistanceFromCylinderAxisSq < FMath::Square(ProjectedShadowInfo->ShadowBounds.W + PrimitiveBounds.SphereRadius)
 					// Check if the primitive is closer than the cylinder cap toward the light
-					&& ProjectedDistanceFromShadowOriginAlongLightDir - PrimitiveBounds.SphereRadius < -ProjectedShadowInfo->MinPreSubjectZ
+					// next line is commented as it breaks large world shadows, if this was meant to be an optimization we should think about a better solution
+					//// && ProjectedDistanceFromShadowOriginAlongLightDir - PrimitiveBounds.SphereRadius < -ProjectedShadowInfo->MinPreSubjectZ
 					// If the primitive is further along the cone axis than the shadow bounds origin, 
 					// Check if the primitive is inside the spherical cap of the cascade's bounds
 					&& !(ProjectedDistanceFromShadowOriginAlongLightDir < 0 
