Do not attempt to process drawlists for unused base pass types.

diff --git a/Engine/Source/Runtime/Renderer/Private/RendererScene.cpp b/Engine/Source/Runtime/Renderer/Private/RendererScene.cpp
--- a/Engine/Source/Runtime/Renderer/Private/RendererScene.cpp
+++ b/Engine/Source/Runtime/Renderer/Private/RendererScene.cpp
@@ -1762,28 +1762,37 @@ void FScene::UpdateStaticDrawListsForMaterials_RenderThread(FRHICommandListImmed
 	// Warning: if any static draw lists are missed here, there will be a crash when trying to render with shaders that have been deleted!
 	TArray<FPrimitiveSceneInfo*> PrimitivesToUpdate;
 	auto FeatureLevel = GetFeatureLevel();
-	for (int32 DrawType = 0; DrawType < EBasePass_MAX; DrawType++)
-	{
-		BasePassNoLightMapDrawList[DrawType].GetUsedPrimitivesBasedOnMaterials(FeatureLevel, Materials, PrimitivesToUpdate);
-		BasePassSimpleDynamicLightingDrawList[DrawType].GetUsedPrimitivesBasedOnMaterials(FeatureLevel, Materials, PrimitivesToUpdate);
-		BasePassCachedVolumeIndirectLightingDrawList[DrawType].GetUsedPrimitivesBasedOnMaterials(FeatureLevel, Materials, PrimitivesToUpdate);
-		BasePassCachedPointIndirectLightingDrawList[DrawType].GetUsedPrimitivesBasedOnMaterials(FeatureLevel, Materials, PrimitivesToUpdate);
-		BasePassHighQualityLightMapDrawList[DrawType].GetUsedPrimitivesBasedOnMaterials(FeatureLevel, Materials, PrimitivesToUpdate);
-		BasePassDistanceFieldShadowMapLightMapDrawList[DrawType].GetUsedPrimitivesBasedOnMaterials(FeatureLevel, Materials, PrimitivesToUpdate);
-		BasePassLowQualityLightMapDrawList[DrawType].GetUsedPrimitivesBasedOnMaterials(FeatureLevel, Materials, PrimitivesToUpdate);
-		BasePassSelfShadowedTranslucencyDrawList[DrawType].GetUsedPrimitivesBasedOnMaterials(FeatureLevel, Materials, PrimitivesToUpdate);
-		BasePassSelfShadowedCachedPointIndirectTranslucencyDrawList[DrawType].GetUsedPrimitivesBasedOnMaterials(FeatureLevel, Materials, PrimitivesToUpdate);
-
-		BasePassForForwardShadingNoLightMapDrawList[DrawType].GetUsedPrimitivesBasedOnMaterials(FeatureLevel, Materials, PrimitivesToUpdate);
-		BasePassForForwardShadingLowQualityLightMapDrawList[DrawType].GetUsedPrimitivesBasedOnMaterials(FeatureLevel, Materials, PrimitivesToUpdate);
-		BasePassForForwardShadingDistanceFieldShadowMapLightMapDrawList[DrawType].GetUsedPrimitivesBasedOnMaterials(FeatureLevel, Materials, PrimitivesToUpdate);
-		BasePassForForwardShadingDirectionalLightAndSHIndirectDrawList[DrawType].GetUsedPrimitivesBasedOnMaterials(FeatureLevel, Materials, PrimitivesToUpdate);
-		BasePassForForwardShadingDirectionalLightAndSHDirectionalIndirectDrawList[DrawType].GetUsedPrimitivesBasedOnMaterials(FeatureLevel, Materials, PrimitivesToUpdate);
-		BasePassForForwardShadingDirectionalLightAndSHDirectionalCSMIndirectDrawList[DrawType].GetUsedPrimitivesBasedOnMaterials(FeatureLevel, Materials, PrimitivesToUpdate);
-		BasePassForForwardShadingMovableDirectionalLightDrawList[DrawType].GetUsedPrimitivesBasedOnMaterials(FeatureLevel, Materials, PrimitivesToUpdate);
-		BasePassForForwardShadingMovableDirectionalLightCSMDrawList[DrawType].GetUsedPrimitivesBasedOnMaterials(FeatureLevel, Materials, PrimitivesToUpdate);
-		BasePassForForwardShadingMovableDirectionalLightLightmapDrawList[DrawType].GetUsedPrimitivesBasedOnMaterials(FeatureLevel, Materials, PrimitivesToUpdate);
-		BasePassForForwardShadingMovableDirectionalLightCSMLightmapDrawList[DrawType].GetUsedPrimitivesBasedOnMaterials(FeatureLevel, Materials, PrimitivesToUpdate);
+
+	if (ShouldUseDeferredRenderer())
+	{
+		for (int32 DrawType = 0; DrawType < EBasePass_MAX; DrawType++)
+		{
+			BasePassNoLightMapDrawList[DrawType].GetUsedPrimitivesBasedOnMaterials(FeatureLevel, Materials, PrimitivesToUpdate);
+			BasePassSimpleDynamicLightingDrawList[DrawType].GetUsedPrimitivesBasedOnMaterials(FeatureLevel, Materials, PrimitivesToUpdate);
+			BasePassCachedVolumeIndirectLightingDrawList[DrawType].GetUsedPrimitivesBasedOnMaterials(FeatureLevel, Materials, PrimitivesToUpdate);
+			BasePassCachedPointIndirectLightingDrawList[DrawType].GetUsedPrimitivesBasedOnMaterials(FeatureLevel, Materials, PrimitivesToUpdate);
+			BasePassHighQualityLightMapDrawList[DrawType].GetUsedPrimitivesBasedOnMaterials(FeatureLevel, Materials, PrimitivesToUpdate);
+			BasePassDistanceFieldShadowMapLightMapDrawList[DrawType].GetUsedPrimitivesBasedOnMaterials(FeatureLevel, Materials, PrimitivesToUpdate);
+			BasePassLowQualityLightMapDrawList[DrawType].GetUsedPrimitivesBasedOnMaterials(FeatureLevel, Materials, PrimitivesToUpdate);
+			BasePassSelfShadowedTranslucencyDrawList[DrawType].GetUsedPrimitivesBasedOnMaterials(FeatureLevel, Materials, PrimitivesToUpdate);
+			BasePassSelfShadowedCachedPointIndirectTranslucencyDrawList[DrawType].GetUsedPrimitivesBasedOnMaterials(FeatureLevel, Materials, PrimitivesToUpdate);
+		}
+	}
+	else
+	{
+		for (int32 DrawType = 0; DrawType < EBasePass_MAX; DrawType++)
+		{
+			BasePassForForwardShadingNoLightMapDrawList[DrawType].GetUsedPrimitivesBasedOnMaterials(FeatureLevel, Materials, PrimitivesToUpdate);
+			BasePassForForwardShadingLowQualityLightMapDrawList[DrawType].GetUsedPrimitivesBasedOnMaterials(FeatureLevel, Materials, PrimitivesToUpdate);
+			BasePassForForwardShadingDistanceFieldShadowMapLightMapDrawList[DrawType].GetUsedPrimitivesBasedOnMaterials(FeatureLevel, Materials, PrimitivesToUpdate);
+			BasePassForForwardShadingDirectionalLightAndSHIndirectDrawList[DrawType].GetUsedPrimitivesBasedOnMaterials(FeatureLevel, Materials, PrimitivesToUpdate);
+			BasePassForForwardShadingDirectionalLightAndSHDirectionalIndirectDrawList[DrawType].GetUsedPrimitivesBasedOnMaterials(FeatureLevel, Materials, PrimitivesToUpdate);
+			BasePassForForwardShadingDirectionalLightAndSHDirectionalCSMIndirectDrawList[DrawType].GetUsedPrimitivesBasedOnMaterials(FeatureLevel, Materials, PrimitivesToUpdate);
+			BasePassForForwardShadingMovableDirectionalLightDrawList[DrawType].GetUsedPrimitivesBasedOnMaterials(FeatureLevel, Materials, PrimitivesToUpdate);
+			BasePassForForwardShadingMovableDirectionalLightCSMDrawList[DrawType].GetUsedPrimitivesBasedOnMaterials(FeatureLevel, Materials, PrimitivesToUpdate);
+			BasePassForForwardShadingMovableDirectionalLightLightmapDrawList[DrawType].GetUsedPrimitivesBasedOnMaterials(FeatureLevel, Materials, PrimitivesToUpdate);
+			BasePassForForwardShadingMovableDirectionalLightCSMLightmapDrawList[DrawType].GetUsedPrimitivesBasedOnMaterials(FeatureLevel, Materials, PrimitivesToUpdate);
+		}
 	}
 
 	PositionOnlyDepthDrawList.GetUsedPrimitivesBasedOnMaterials(FeatureLevel, Materials, PrimitivesToUpdate);
