easier debugging of shadows

diff --git a/Engine/Source/Runtime/Renderer/Private/ShadowRendering.cpp b/Engine/Source/Runtime/Renderer/Private/ShadowRendering.cpp
--- a/Engine/Source/Runtime/Renderer/Private/ShadowRendering.cpp
+++ b/Engine/Source/Runtime/Renderer/Private/ShadowRendering.cpp
@@ -2609,7 +2609,7 @@ bool FDeferredShadingSceneRenderer::RenderOnePassPointLightShadows(FRHICommandLi
 
 			if (!ProjectedShadowInfo->bRayTracedDistanceFieldShadow)
 			{
-				SCOPED_DRAW_EVENT(RHICmdList, ShadowDepthsFromOpaque);
+				SCOPED_DRAW_EVENT(RHICmdList, ShadowDepthsFromOpaquePointLight);
 				auto SetShadowRenderTargets = [ProjectedShadowInfo](FRHICommandList& RHICmdList)
 				{
 					GSceneRenderTargets.BeginRenderingCubeShadowDepth(RHICmdList, ProjectedShadowInfo->ResolutionX);
@@ -2697,6 +2697,8 @@ void FDeferredShadingSceneRenderer::RenderProjections(
 				if (ProjectedShadowInfo->FadeAlphas[ViewIndex] > 1.0f / 256.0f)
 				{
 					ProjectedShadowInfo->RenderProjection(RHICmdList, ViewIndex, &View);
+
+					GRenderTargetPool.VisualizeTexture.SetCheckPoint(RHICmdList, GSceneRenderTargets.GetLightAttenuation());
 				}
 			}
 		}
@@ -3195,7 +3197,7 @@ bool FDeferredShadingSceneRenderer::RenderProjectedShadows(FRHICommandListImmedi
 
 		{
 			// Render the shadow depths.
-			SCOPED_DRAW_EVENT(RHICmdList, ShadowDepthsFromOpaque);
+			SCOPED_DRAW_EVENT(RHICmdList, ShadowDepthsFromOpaqueProjected);
 
 			bool bFirst = true;
 			auto SetShadowRenderTargets = [&bFirst](FRHICommandList& RHICmdList)
@@ -3456,7 +3458,7 @@ bool FForwardShadingSceneRenderer::RenderShadowDepthMap(FRHICommandListImmediate
 
 	{
 		// Render the shadow depths.
-		SCOPED_DRAW_EVENT(RHICmdList, ShadowDepthsFromOpaque);
+		SCOPED_DRAW_EVENT(RHICmdList, ShadowDepthsFromOpaqueForward);
 
 		bool bFirst = true;
 		auto SetShadowRenderTargets = [&bFirst](FRHICommandList& RHICmdList)
