Added Rec.709 output option for the new prototype tonemapper.

diff --git a/Engine/Shaders/PostProcessTonemap.usf b/Engine/Shaders/PostProcessTonemap.usf
--- a/Engine/Shaders/PostProcessTonemap.usf
+++ b/Engine/Shaders/PostProcessTonemap.usf
@@ -34,6 +34,12 @@ half GrainFromUV(float2 GrainUV)
 	return Grain;
 }
 
+half3 LinearTo709Branchless(half3 lin) 
+{
+	lin = max(6.10352e-5, lin); // minimum positive non-denormal (fixes black problem on DX11 AMD and NV)
+	return min(lin * 4.5, pow(max(lin, 0.018), 0.45) * 1.099 - 0.099);
+}
+
 half3 LinearToSrgbBranchless(half3 lin) 
 {
 	lin = max(6.10352e-5, lin); // minimum positive non-denormal (fixes black problem on DX11 AMD and NV)
@@ -277,9 +283,14 @@ void MainPS(
 			// Note, MacOSX native output is raw gamma 2.2 not sRGB!
 			half3 TonemappedColor = pow(FilmColor.rgb, 1.0/2.2);
 		#else
-			// Apply conversion to sRGB (this must be an exact sRGB conversion else darks are bad).
-			// Branching is faster than branchless on AMD on PC.
-			half3 TonemappedColor = LinearToSrgbBranching(FilmColor);
+			#if USE_709
+				// Didn't profile yet if the branching version would be faster (different linear segment).
+				half3 TonemappedColor = LinearTo709Branchless(FilmColor);
+			#else
+				// Apply conversion to sRGB (this must be an exact sRGB conversion else darks are bad).
+				// Branching is faster than branchless on AMD on PC.
+				half3 TonemappedColor = LinearToSrgbBranching(FilmColor);
+			#endif
 		#endif
 
 		//TonemappedColor = rsqrt( ( 0.8 + 0.8 * LinearColor ) / LinearColor );		// 1 vmad, 3 rcp, 1 vmul, 3 rsqrt
diff --git a/Engine/Source/Runtime/Core/Private/HAL/ConsoleManager.cpp b/Engine/Source/Runtime/Core/Private/HAL/ConsoleManager.cpp
--- a/Engine/Source/Runtime/Core/Private/HAL/ConsoleManager.cpp
+++ b/Engine/Source/Runtime/Core/Private/HAL/ConsoleManager.cpp
@@ -1992,6 +1992,14 @@ static TAutoConsoleVariable<int32> CVarTonemapperPhoto(
 	TEXT("1: use new photo post"),
 	ECVF_Scalability | ECVF_RenderThreadSafe);
 
+static TAutoConsoleVariable<int32> CVarTonemapper709(
+	TEXT("r.Tonemapper709"),
+	0,
+	TEXT("0: use sRGB on PC monitor output\n")
+	TEXT("1: use Rec.709 for HDTV/projector output (requires r.TonemapperPhoto 1)"),
+	ECVF_Scalability | ECVF_RenderThreadSafe);
+
+
 static TAutoConsoleVariable<int32> CVarDetailMode(
 	TEXT("r.DetailMode"),
 	2,
diff --git a/Engine/Source/Runtime/Renderer/Private/PostProcess/PostProcessTonemap.cpp b/Engine/Source/Runtime/Renderer/Private/PostProcess/PostProcessTonemap.cpp
--- a/Engine/Source/Runtime/Renderer/Private/PostProcess/PostProcessTonemap.cpp
+++ b/Engine/Source/Runtime/Renderer/Private/PostProcess/PostProcessTonemap.cpp
@@ -38,6 +38,7 @@ typedef enum {
 	TonemapperColorGrading      = (1<<14),
 	TonemapperMsaa              = (1<<15),
 	TonemapperPhoto             = (1<<16),
+	Tonemapper709               = (1<<17),
 } TonemapperOption;
 
 // Tonemapper option cost (0 = no cost, 255 = max cost).
@@ -61,6 +62,7 @@ static uint8 TonemapperCostTab[] = {
 	1, //TonemapperColorGrading
 	1, //TonemapperMsaa
 	1, //TonemapperPhoto
+	1, //Tonemapper709
 };
 
 // Edit the following to add and remove configurations.
@@ -183,16 +185,24 @@ static uint32 TonemapperConfBitmaskPC[22] = {
 	0,
 
 	TonemapperBloom +
-	TonemapperVignette +
+	TonemapperGrainJitter +
+	TonemapperGrainIntensity +
 	TonemapperGrainQuantization +
+	TonemapperVignette +
+	TonemapperVignetteColor +
 	TonemapperColorFringe +
 	TonemapperPhoto +
+	Tonemapper709 +
 	0,
 
 	TonemapperBloom +
-	TonemapperVignette +
+	TonemapperGrainJitter +
+	TonemapperGrainIntensity +
 	TonemapperGrainQuantization +
+	TonemapperVignette +
+	TonemapperVignetteColor +
 	TonemapperPhoto +
+	Tonemapper709 +
 	0,
 
 	TonemapperBloom +
@@ -201,11 +211,9 @@ static uint32 TonemapperConfBitmaskPC[22] = {
 	0,
 
 
-	// same without TonemapperGrainQuantization
+	// similar without TonemapperGrain
 
 	TonemapperBloom +
-	TonemapperGrainJitter +
-	TonemapperGrainIntensity +
 	TonemapperVignette +
 	TonemapperVignetteColor +
 	TonemapperColorFringe +
@@ -213,8 +221,6 @@ static uint32 TonemapperConfBitmaskPC[22] = {
 	0,
 
 	TonemapperBloom +
-	TonemapperGrainJitter +
-	TonemapperGrainIntensity +
 	TonemapperVignette +
 	TonemapperVignetteColor +
 	TonemapperPhoto +
@@ -222,13 +228,17 @@ static uint32 TonemapperConfBitmaskPC[22] = {
 
 	TonemapperBloom +
 	TonemapperVignette +
+	TonemapperVignetteColor +
 	TonemapperColorFringe +
 	TonemapperPhoto +
+	Tonemapper709 +
 	0,
 
 	TonemapperBloom +
 	TonemapperVignette +
+	TonemapperVignetteColor +
 	TonemapperPhoto +
+	Tonemapper709 +
 	0,
 
 	TonemapperBloom +
@@ -691,6 +701,18 @@ static uint32 TonemapperGenerateBitmaskPC(const FViewInfo* RESTRICT View, bool b
 		}
 	}
 
+	{
+		static TConsoleVariableData<int32>* CVar = IConsoleManager::Get().FindTConsoleVariableDataInt(TEXT("r.Tonemapper709"));
+		int32 Value = CVar->GetValueOnRenderThread();
+
+		if(Value > 0)
+		{
+			// Remove settings not compatible.
+			Bitmask &= ~(TonemapperColorMatrix|TonemapperShadowTint|TonemapperContrast|TonemapperColorGrading|TonemapperMsaa);
+			Bitmask |= TonemapperPhoto | Tonemapper709;
+		}
+	}
+
 	{
 		static TConsoleVariableData<int32>* CVar = IConsoleManager::Get().FindTConsoleVariableDataInt(TEXT("r.TonemapperPhoto"));
 		int32 Value = CVar->GetValueOnRenderThread();
@@ -1065,6 +1087,7 @@ class FPostProcessTonemapPS : public FGlobalShader
 		OutEnvironment.SetDefine(TEXT("USE_VOLUME_LUT"),		 (IsFeatureLevelSupported(Platform,ERHIFeatureLevel::SM4) && GSupportsVolumeTextureRendering && Platform != EShaderPlatform::SP_OPENGL_SM4_MAC));
 		OutEnvironment.SetDefine(TEXT("USE_COLOR_GRADING"),		 TonemapperIsDefined(ConfigBitmask, TonemapperColorGrading));
 		OutEnvironment.SetDefine(TEXT("USE_PHOTO"),              TonemapperIsDefined(ConfigBitmask, TonemapperPhoto));
+		OutEnvironment.SetDefine(TEXT("USE_709"),                TonemapperIsDefined(ConfigBitmask, Tonemapper709));
 
 		if( !IsFeatureLevelSupported(Platform,ERHIFeatureLevel::SM5) )
 		{
