fixed comment

diff --git a/Engine/Source/Runtime/Renderer/Private/PostProcess/PostProcessNoiseBlur.cpp b/Engine/Source/Runtime/Renderer/Private/PostProcess/PostProcessNoiseBlur.cpp
--- a/Engine/Source/Runtime/Renderer/Private/PostProcess/PostProcessNoiseBlur.cpp
+++ b/Engine/Source/Runtime/Renderer/Private/PostProcess/PostProcessNoiseBlur.cpp
@@ -11,7 +11,7 @@
 #include "PostProcessing.h"
 #include "SceneUtils.h"
 
-/** Encapsulates the post processing down sample pixel shader. */
+/** Encapsulates the post processing noise blur shader. */
 template <uint32 Method>
 class FPostProcessNoiseBlurPS : public FGlobalShader
 {
