refactor shadow projection code

diff --git a/Engine/Source/Runtime/Engine/Private/DirectionalLightComponent.cpp b/Engine/Source/Runtime/Engine/Private/DirectionalLightComponent.cpp
--- a/Engine/Source/Runtime/Engine/Private/DirectionalLightComponent.cpp
+++ b/Engine/Source/Runtime/Engine/Private/DirectionalLightComponent.cpp
@@ -220,19 +220,17 @@ public:
 	{
 		const FMatrix& WorldToLight = GetWorldToLight();
 
-		const bool bRaytracedCascade = (InCascadeIndex == INDEX_NONE);
-
-		FShadowCascadeSettings CascadeSettings;
+		const bool bRayTracedCascade = (InCascadeIndex == INDEX_NONE);
 
 		FSphere Bounds = FDirectionalLightSceneProxy::GetShadowSplitBounds(View, InCascadeIndex, &OutInitializer.CascadeSettings);
 
 		uint32 NumNearCascades = GetNumShadowMappedCascades(View.MaxShadowCascades);
 
-		const uint32 ShadowSplitIndex = bRaytracedCascade ? NumNearCascades : InCascadeIndex;
+		// Last cascade is the ray traced shadow, if enabled
+		OutInitializer.CascadeSettings.ShadowSplitIndex = bRayTracedCascade ? NumNearCascades : InCascadeIndex;
 
 		const float ShadowExtent = Bounds.W / FMath::Sqrt(3.0f);
 		const FBoxSphereBounds SubjectBounds(Bounds.Center, FVector(ShadowExtent, ShadowExtent, ShadowExtent), Bounds.W);
-		OutInitializer.bOnePassPointLightShadow = false;
 		OutInitializer.PreShadowTranslation = -Bounds.Center;
 		OutInitializer.WorldToLight = FInverseRotationMatrix(FVector(WorldToLight.M[0][0],WorldToLight.M[1][0],WorldToLight.M[2][0]).GetSafeNormal().Rotation());
 		OutInitializer.Scales = FVector(1.0f,1.0f / Bounds.W,1.0f / Bounds.W);
@@ -243,10 +241,8 @@ public:
 		// Reduce casting distance on a directional light
 		// This is necessary to improve floating point precision in several places, especially when deriving frustum verts from InvReceiverMatrix
 		OutInitializer.MaxDistanceToCastInLightW = HALF_WORLD_MAX / 32.0f;
-		// Last cascade is the ray traced shadow, if enabled
-		OutInitializer.InitShadowSplitIndex = ShadowSplitIndex;
-		OutInitializer.bRayTracedDistanceFieldShadow = bRaytracedCascade;
-		OutInitializer.CascadeSettings.bFarShadowCascade = !bRaytracedCascade && ShadowSplitIndex >= NumNearCascades;
+		OutInitializer.CascadeSettings.bRayTracedDistanceField = bRayTracedCascade;
+		OutInitializer.CascadeSettings.bFarShadowCascade = !bRayTracedCascade && OutInitializer.CascadeSettings.ShadowSplitIndex >= (int32)NumNearCascades;
 		return true;
 	}
 
diff --git a/Engine/Source/Runtime/Engine/Private/PointLightComponent.cpp b/Engine/Source/Runtime/Engine/Private/PointLightComponent.cpp
--- a/Engine/Source/Runtime/Engine/Private/PointLightComponent.cpp
+++ b/Engine/Source/Runtime/Engine/Private/PointLightComponent.cpp
@@ -73,7 +73,6 @@ public:
 			&& GAllowPointLightCubemapShadows != 0)
 		{
 			FWholeSceneProjectedShadowInitializer& OutInitializer = *new(OutInitializers) FWholeSceneProjectedShadowInitializer;
-			OutInitializer.bOnePassPointLightShadow = true;
 			OutInitializer.PreShadowTranslation = -GetLightToWorld().GetOrigin();
 			OutInitializer.WorldToLight = GetWorldToLight().RemoveTranslation();
 			OutInitializer.Scales = FVector(1, 1, 1);
@@ -82,9 +81,8 @@ public:
 			OutInitializer.WAxis = FVector4(0,0,1,0);
 			OutInitializer.MinLightW = 0.1f;
 			OutInitializer.MaxDistanceToCastInLightW = Radius;
-			OutInitializer.InitShadowSplitIndex = INDEX_NONE;
-			OutInitializer.bRayTracedDistanceFieldShadow = UseRayTracedDistanceFieldShadows() && DoesPlatformSupportDistanceFieldShadowing(ViewFamily.GetShaderPlatform());
-		
+			OutInitializer.CascadeSettings.bOnePassPointLightShadow = true;
+			OutInitializer.CascadeSettings.bRayTracedDistanceField = UseRayTracedDistanceFieldShadows() && DoesPlatformSupportDistanceFieldShadowing(ViewFamily.GetShaderPlatform());
 			return true;
 		}
 		
diff --git a/Engine/Source/Runtime/Engine/Private/SpotLightComponent.cpp b/Engine/Source/Runtime/Engine/Private/SpotLightComponent.cpp
--- a/Engine/Source/Runtime/Engine/Private/SpotLightComponent.cpp
+++ b/Engine/Source/Runtime/Engine/Private/SpotLightComponent.cpp
@@ -122,7 +122,6 @@ public:
 	virtual bool GetWholeSceneProjectedShadowInitializer(const FSceneViewFamily& ViewFamily, TArray<FWholeSceneProjectedShadowInitializer, TInlineAllocator<6> >& OutInitializers) const
 	{
 		FWholeSceneProjectedShadowInitializer& OutInitializer = *new(OutInitializers) FWholeSceneProjectedShadowInitializer;
-		OutInitializer.bOnePassPointLightShadow = false;
 		OutInitializer.PreShadowTranslation = -GetLightToWorld().GetOrigin();
 		OutInitializer.WorldToLight = GetWorldToLight().RemoveTranslation();
 		OutInitializer.Scales = FVector(1.0f,InvTanOuterCone,InvTanOuterCone);
@@ -135,8 +134,7 @@ public:
 		OutInitializer.WAxis = FVector4(0,0,1,0);
 		OutInitializer.MinLightW = 0.1f;
 		OutInitializer.MaxDistanceToCastInLightW = Radius;
-		OutInitializer.InitShadowSplitIndex = INDEX_NONE;
-		OutInitializer.bRayTracedDistanceFieldShadow = UseRayTracedDistanceFieldShadows() && DoesPlatformSupportDistanceFieldShadowing(ViewFamily.GetShaderPlatform());
+		OutInitializer.CascadeSettings.bRayTracedDistanceField = UseRayTracedDistanceFieldShadows() && DoesPlatformSupportDistanceFieldShadowing(ViewFamily.GetShaderPlatform());
 		return true;
 	}
 
diff --git a/Engine/Source/Runtime/Engine/Public/SceneManagement.h b/Engine/Source/Runtime/Engine/Public/SceneManagement.h
--- a/Engine/Source/Runtime/Engine/Public/SceneManagement.h
+++ b/Engine/Source/Runtime/Engine/Public/SceneManagement.h
@@ -524,6 +524,18 @@ public:
 	/** When enabled, the cascade only renders objects marked with bCastFarShadows enabled (e.g. Landscape). */
 	bool bFarShadowCascade;
 
+	/** Whether the shadow will be computed by ray tracing the distance field. */
+	bool bRayTracedDistanceField;
+
+	/** Whether the shadow is a point light shadow that renders all faces of a cubemap in one pass. */
+	bool bOnePassPointLightShadow;
+
+	/** 
+	 * Index of the split if this is a whole scene shadow from a directional light, 
+	 * Or index of the direction if this is a whole scene shadow from a point light, otherwise INDEX_NONE. 
+	 */
+	int32 ShadowSplitIndex;
+	
 	FShadowCascadeSettings()
 		: SplitNear(0.0f)
 		, SplitFar(WORLD_MAX)
@@ -532,6 +544,9 @@ public:
 		, FadePlaneOffset(SplitFar)
 		, FadePlaneLength(SplitFar - FadePlaneOffset)
 		, bFarShadowCascade(false)
+		, bRayTracedDistanceField(false)
+		, bOnePassPointLightShadow(false)
+		, ShadowSplitIndex(INDEX_NONE)
 	{
 	}
 };
@@ -570,21 +585,7 @@ public:
 class ENGINE_API FWholeSceneProjectedShadowInitializer : public FProjectedShadowInitializer
 {
 public:
-	int32 InitShadowSplitIndex;
-	
 	FShadowCascadeSettings CascadeSettings;
-
-	/** Whether the shadow is a point light shadow that renders all faces of a cubemap in one pass. */
-	bool bOnePassPointLightShadow;
-
-	/** Whether the shadow will be computed by ray tracing the distance field. */
-	bool bRayTracedDistanceFieldShadow;
-
-	FWholeSceneProjectedShadowInitializer()
-		: InitShadowSplitIndex(INDEX_NONE)
-		, bOnePassPointLightShadow(false)
-		, bRayTracedDistanceFieldShadow(false)
-	{}	
 };
 
 inline bool DoesPlatformSupportDistanceFieldShadowing(EShaderPlatform Platform)
diff --git a/Engine/Source/Runtime/Renderer/Private/DistanceFieldShadowing.cpp b/Engine/Source/Runtime/Renderer/Private/DistanceFieldShadowing.cpp
--- a/Engine/Source/Runtime/Renderer/Private/DistanceFieldShadowing.cpp
+++ b/Engine/Source/Runtime/Renderer/Private/DistanceFieldShadowing.cpp
@@ -696,7 +696,7 @@ bool FDeferredShadingSceneRenderer::ShouldPrepareForDistanceFieldShadows() const
 			{
 				const FProjectedShadowInfo* ProjectedShadowInfo = VisibleLightInfo.AllProjectedShadows[ShadowIndex];
 
-				if (ProjectedShadowInfo->bRayTracedDistanceFieldShadow)
+				if (ProjectedShadowInfo->CascadeSettings.bRayTracedDistanceField)
 				{
 					bSceneHasRayTracedDFShadows = true;
 					break;
@@ -735,7 +735,7 @@ void FProjectedShadowInfo::RenderRayTracedDistanceFieldProjection(FRHICommandLis
 				NumPlanes = CascadeSettings.ShadowBoundsAccurate.Planes.Num();
 				PlaneData = CascadeSettings.ShadowBoundsAccurate.Planes.GetData();
 			}
-			else if (bOnePassPointLightShadow)
+			else if (CascadeSettings.bOnePassPointLightShadow)
 			{
 				ShadowBoundingSphereValue = FVector4(ShadowBounds.Center.X, ShadowBounds.Center.Y, ShadowBounds.Center.Z, ShadowBounds.W);
 			}
diff --git a/Engine/Source/Runtime/Renderer/Private/SceneOcclusion.cpp b/Engine/Source/Runtime/Renderer/Private/SceneOcclusion.cpp
--- a/Engine/Source/Runtime/Renderer/Private/SceneOcclusion.cpp
+++ b/Engine/Source/Runtime/Renderer/Private/SceneOcclusion.cpp
@@ -1116,7 +1116,7 @@ void FDeferredShadingSceneRenderer::BeginOcclusionTests(FRHICommandListImmediate
 							continue;
 						}
 
-						if (ProjectedShadowInfo.bOnePassPointLightShadow)
+						if (ProjectedShadowInfo.CascadeSettings.bOnePassPointLightShadow)
 						{
 							FLightSceneProxy* Proxy = ProjectedShadowInfo.LightSceneInfo->Proxy;
 
diff --git a/Engine/Source/Runtime/Renderer/Private/ScenePrivate.h b/Engine/Source/Runtime/Renderer/Private/ScenePrivate.h
--- a/Engine/Source/Runtime/Renderer/Private/ScenePrivate.h
+++ b/Engine/Source/Runtime/Renderer/Private/ScenePrivate.h
@@ -373,22 +373,22 @@ public:
 
 		FORCEINLINE bool operator == (const FProjectedShadowKey &Other) const
 		{
-			return (PrimitiveId == Other.PrimitiveId && Light == Other.Light && SplitIndex == Other.SplitIndex && bTranslucentShadow == Other.bTranslucentShadow);
+			return (PrimitiveId == Other.PrimitiveId && Light == Other.Light && ShadowSplitIndex == Other.ShadowSplitIndex && bTranslucentShadow == Other.bTranslucentShadow);
 		}
 
 		FProjectedShadowKey(const FProjectedShadowInfo& ProjectedShadowInfo)
 			: PrimitiveId(ProjectedShadowInfo.ParentSceneInfo ? ProjectedShadowInfo.ParentSceneInfo->PrimitiveComponentId : FPrimitiveComponentId())
 			, Light(ProjectedShadowInfo.LightSceneInfo->Proxy->GetLightComponent())
-			, SplitIndex(ProjectedShadowInfo.ShadowSplitIndex)
+			, ShadowSplitIndex(ProjectedShadowInfo.CascadeSettings.ShadowSplitIndex)
 			, bTranslucentShadow(ProjectedShadowInfo.bTranslucentShadow)
 		{
 		}
 
-		FProjectedShadowKey(FPrimitiveComponentId InPrimitiveId, const ULightComponent* InLight, int32 InSplitIndex, bool bInTranslucentShadow):
-			PrimitiveId(InPrimitiveId),
-			Light(InLight),
-			SplitIndex(InSplitIndex),
-			bTranslucentShadow(bInTranslucentShadow)
+		FProjectedShadowKey(FPrimitiveComponentId InPrimitiveId, const ULightComponent* InLight, int32 InSplitIndex, bool bInTranslucentShadow)
+			: PrimitiveId(InPrimitiveId)
+			, Light(InLight)
+			, ShadowSplitIndex(InSplitIndex)
+			, bTranslucentShadow(bInTranslucentShadow)
 		{
 		}
 
@@ -400,7 +400,7 @@ public:
 	private:
 		FPrimitiveComponentId PrimitiveId;
 		const ULightComponent* Light;
-		int32 SplitIndex;
+		int32 ShadowSplitIndex;
 		bool bTranslucentShadow;
 	};
 
diff --git a/Engine/Source/Runtime/Renderer/Private/ShadowRendering.cpp b/Engine/Source/Runtime/Renderer/Private/ShadowRendering.cpp
--- a/Engine/Source/Runtime/Renderer/Private/ShadowRendering.cpp
+++ b/Engine/Source/Runtime/Renderer/Private/ShadowRendering.cpp
@@ -1073,7 +1073,7 @@ bool FShadowDepthDrawingPolicyFactory::DrawDynamicMesh(
 		const EBlendMode BlendMode = Material->GetBlendMode();
 		const EMaterialShadingModel ShadingModel = Material->GetShadingModel();
 
-		const bool bLocalOnePassPointLightShadow = Context.ShadowInfo->bOnePassPointLightShadow;
+		const bool bLocalOnePassPointLightShadow = Context.ShadowInfo->CascadeSettings.bOnePassPointLightShadow;
 		const bool bReflectiveShadowmap = Context.ShadowInfo->bReflectiveShadowmap && !bLocalOnePassPointLightShadow && Material->ShouldInjectEmissiveIntoLPV();
 
 		if ( (!IsTranslucentBlendMode(BlendMode) || bReflectiveShadowmap ) && ShadingModel != MSM_Unlit )
@@ -1153,7 +1153,7 @@ static void CheckShadowDepthMaterials(const FMaterialRenderProxy* InRenderProxy,
 
 void FProjectedShadowInfo::ClearDepth(FRHICommandList& RHICmdList, FDeferredShadingSceneRenderer* SceneRenderer)
 {
-	if (bOnePassPointLightShadow)
+	if (CascadeSettings.bOnePassPointLightShadow)
 	{
 		// Set the viewport to the whole render target since it's a cube map, don't leave any border space
 		RHICmdList.SetViewport(
@@ -1259,7 +1259,7 @@ void DrawShadowMeshElements(FRHICommandList& RHICmdList, const FViewInfo& View,
 	FShadowDepthDrawingPolicy<bReflectiveShadowmap> SharedDrawingPolicy(
 		FirstMaterialResource,
 		ShadowInfo.bDirectionalLight,
-		ShadowInfo.bOnePassPointLightShadow,
+		ShadowInfo.CascadeSettings.bOnePassPointLightShadow,
 		ShadowInfo.bPreShadow,
 		FeatureLevel);
 
@@ -1404,7 +1404,7 @@ void FProjectedShadowInfo::SetStateForDepth(FRHICommandList& RHICmdList)
 {
 	check(bAllocated);
 
-	if (bOnePassPointLightShadow)
+	if (CascadeSettings.bOnePassPointLightShadow)
 	{
 		// Set the viewport to the whole render target since it's a cube map, don't leave any border space
 		RHICmdList.SetViewport(
@@ -1443,7 +1443,7 @@ void FProjectedShadowInfo::SetStateForDepth(FRHICommandList& RHICmdList)
 		}
 	}
 
-	if (bReflectiveShadowmap && !bOnePassPointLightShadow)
+	if (bReflectiveShadowmap && !CascadeSettings.bOnePassPointLightShadow)
 	{
 		// Enable color writes to the reflective shadow map targets with opaque blending
 		RHICmdList.SetBlendState(TStaticBlendStateWriteMask<CW_RGBA, CW_RGBA>::GetRHI());
@@ -1523,7 +1523,7 @@ void FProjectedShadowInfo::RenderDepthInner(FRHICommandList& RHICmdList, FSceneR
 				FRHICommandList* CmdList = ParallelCommandListSet.NewParallelCommandList();
 
 				FGraphEventRef AnyThreadCompletionEvent = TGraphTask<FDrawShadowMeshElementsThreadTask>::CreateTask(nullptr, ENamedThreads::RenderThread)
-					.ConstructAndDispatchWhenReady(*this, *CmdList, *FoundView, bReflectiveShadowmap && !bOnePassPointLightShadow);
+					.ConstructAndDispatchWhenReady(*this, *CmdList, *FoundView, bReflectiveShadowmap && !CascadeSettings.bOnePassPointLightShadow);
 
 				ParallelCommandListSet.AddParallelCommandList(CmdList, AnyThreadCompletionEvent);
 			}
@@ -1562,7 +1562,7 @@ void FProjectedShadowInfo::RenderDepthInner(FRHICommandList& RHICmdList, FSceneR
 		{
 			SCOPE_CYCLE_COUNTER(STAT_WholeSceneStaticShadowDepthsTime);
 
-			if (bReflectiveShadowmap && !bOnePassPointLightShadow)
+			if (bReflectiveShadowmap && !CascadeSettings.bOnePassPointLightShadow)
 			{
 				// reflective shadow map
 				DrawShadowMeshElements<true>(RHICmdList, *FoundView, *this);
@@ -1782,7 +1782,7 @@ void FProjectedShadowInfo::RenderProjection(FRHICommandListImmediate& RHICmdList
 		return;
 	}
 
-	if (bRayTracedDistanceFieldShadow)
+	if (CascadeSettings.bRayTracedDistanceField)
 	{
 		RenderRayTracedDistanceFieldProjection(RHICmdList, *View);
 		return;
@@ -1966,7 +1966,7 @@ void FProjectedShadowInfo::RenderProjection(FRHICommandListImmediate& RHICmdList
 
 			RHICmdList.SetRasterizerState(TStaticRasterizerState<FM_Solid, CM_None>::GetRHI());
 
-			checkSlow(ShadowSplitIndex >= 0);
+			checkSlow(CascadeSettings.ShadowSplitIndex >= 0);
 			checkSlow(bDirectionalLight);
 
 			// Draw 2 fullscreen planes, front facing one at the near subfrustum plane, and back facing one at the far.
@@ -2001,7 +2001,7 @@ void FProjectedShadowInfo::RenderProjection(FRHICommandListImmediate& RHICmdList
 			};
 
 			// Only draw the near plane if this is not the nearest split
-			DrawPrimitiveUP(RHICmdList, PT_TriangleList, (ShadowSplitIndex > 0) ? 4 : 2, Verts, sizeof(FVector4));
+			DrawPrimitiveUP(RHICmdList, PT_TriangleList, (CascadeSettings.ShadowSplitIndex > 0) ? 4 : 2, Verts, sizeof(FVector4));
 		}
 	}
 	// Not a preshadow, mask the projection to any pixels inside the frustum.
@@ -2111,10 +2111,10 @@ void FProjectedShadowInfo::RenderProjection(FRHICommandListImmediate& RHICmdList
 
 		if (LocalQuality > 1)
 		{
-			if (IsWholeSceneDirectionalShadow() && ShadowSplitIndex > 0)
+			if (IsWholeSceneDirectionalShadow() && CascadeSettings.ShadowSplitIndex > 0)
 			{
 				// adjust kernel size so that the penumbra size of distant splits will better match up with the closer ones
-				const float SizeScale = ShadowSplitIndex / FMath::Max(0.001f, CVarCSMSplitPenumbraScale.GetValueOnRenderThread());
+				const float SizeScale = CascadeSettings.ShadowSplitIndex / FMath::Max(0.001f, CVarCSMSplitPenumbraScale.GetValueOnRenderThread());
 			}
 			else if (LocalQuality > 2 && !bWholeSceneShadow)
 			{
@@ -2200,7 +2200,7 @@ void FProjectedShadowInfo::RenderOnePassPointLightProjection(FRHICommandListImme
 {
 	SCOPE_CYCLE_COUNTER(STAT_RenderWholeSceneShadowProjectionsTime);
 
-	checkSlow(bOnePassPointLightShadow);
+	checkSlow(CascadeSettings.bOnePassPointLightShadow);
 	
 	const FSphere LightBounds = LightSceneInfo->Proxy->GetBoundingSphere();
 
@@ -2276,7 +2276,7 @@ void FProjectedShadowInfo::RenderFrustumWireframe(FPrimitiveDrawInterface* PDI)
 	if(IsWholeSceneDirectionalShadow())
 	{
 		Color = FColor::White;
-		switch(ShadowSplitIndex)
+		switch(CascadeSettings.ShadowSplitIndex)
 		{
 			case 0: Color = FColor::Red; break;
 			case 1: Color = FColor::Yellow; break;
@@ -2412,7 +2412,7 @@ void FProjectedShadowInfo::UpdateShaderDepthBias()
 	}
 	else if (IsWholeSceneDirectionalShadow())
 	{
-		check(ShadowSplitIndex >= 0);
+		check(CascadeSettings.ShadowSplitIndex >= 0);
 
 		// the z range is adjusted to we need to adjust here as well
 		DepthBias = CVarCSMShadowDepthBias.GetValueOnRenderThread() / (MaxSubjectZ - MinSubjectZ);
@@ -2468,7 +2468,7 @@ float FProjectedShadowInfo::ComputeTransitionSize() const
 	}
 	else if (IsWholeSceneDirectionalShadow())
 	{
-		check(ShadowSplitIndex >= 0);
+		check(CascadeSettings.ShadowSplitIndex >= 0);
 
 		// todo: remove GetShadowTransitionScale()
 		// make 1/ ShadowTransitionScale, SpotLightShadowTransitionScale
@@ -2524,9 +2524,9 @@ void FProjectedShadowInfo::GetShadowTypeNameForDrawEvent(FString& TypeName) cons
 
 		if (bWholeSceneShadow)
 		{
-			if (ShadowSplitIndex >= 0)
+			if (CascadeSettings.ShadowSplitIndex >= 0)
 			{
-				TypeName = FString(TEXT("WholeScene split ")) + FString::FromInt(ShadowSplitIndex);
+				TypeName = FString(TEXT("WholeScene split ")) + FString::FromInt(CascadeSettings.ShadowSplitIndex);
 			}
 			else
 			{
@@ -2603,11 +2603,11 @@ bool FDeferredShadingSceneRenderer::RenderOnePassPointLightShadows(FRHICommandLi
 			bShadowIsVisible |= (ViewRelevance.bOpaqueRelevance	&& VisibleLightViewInfo.ProjectedShadowVisibilityMap[ShadowIndex]);
 		}
 
-		if (bShadowIsVisible && ProjectedShadowInfo->bOnePassPointLightShadow)
+		if (bShadowIsVisible && ProjectedShadowInfo->CascadeSettings.bOnePassPointLightShadow)
 		{
 			INC_DWORD_STAT(STAT_WholeSceneShadows);
 
-			if (!ProjectedShadowInfo->bRayTracedDistanceFieldShadow)
+			if (!ProjectedShadowInfo->CascadeSettings.bRayTracedDistanceField)
 			{
 				SCOPED_DRAW_EVENT(RHICmdList, ShadowDepthsFromOpaquePointLight);
 				auto SetShadowRenderTargets = [ProjectedShadowInfo](FRHICommandList& RHICmdList)
@@ -2637,7 +2637,7 @@ bool FDeferredShadingSceneRenderer::RenderOnePassPointLightShadows(FRHICommandLi
 					// Set the device viewport for the view.
 					RHICmdList.SetViewport(View.ViewRect.Min.X, View.ViewRect.Min.Y, 0.0f, View.ViewRect.Max.X, View.ViewRect.Max.Y, 1.0f);
 					
-					if (ProjectedShadowInfo->bRayTracedDistanceFieldShadow)
+					if (ProjectedShadowInfo->CascadeSettings.bRayTracedDistanceField)
 					{
 						ProjectedShadowInfo->RenderRayTracedDistanceFieldProjection(RHICmdList, View);
 					}
@@ -2650,7 +2650,7 @@ bool FDeferredShadingSceneRenderer::RenderOnePassPointLightShadows(FRHICommandLi
 
 			// Don't inject shadowed lighting with whole scene shadows used for previewing a light with static shadows,
 			// Since that would cause a mismatch with the built lighting
-			if (!LightSceneInfo->Proxy->HasStaticShadowing() && !ProjectedShadowInfo->bRayTracedDistanceFieldShadow)
+			if (!LightSceneInfo->Proxy->HasStaticShadowing() && !ProjectedShadowInfo->CascadeSettings.bRayTracedDistanceField)
 			{
 				bInjectedTranslucentVolume = true;
 				SCOPED_DRAW_EVENT(RHICmdList, InjectTranslucentVolume);
@@ -2731,7 +2731,7 @@ struct FCompareFProjectedShadowInfoBySplitIndex
 			{
 				// Both A and B are CSMs
 				// Compare Split Indexes, to order them far to near.
-				return (B.ShadowSplitIndex < A.ShadowSplitIndex);
+				return (B.CascadeSettings.ShadowSplitIndex < A.CascadeSettings.ShadowSplitIndex);
 			}
 
 			// A is a CSM, B is per-object shadow etc.
@@ -2784,7 +2784,7 @@ bool FDeferredShadingSceneRenderer::RenderTranslucentProjectedShadows(FRHIComman
 		}
 
 		// Skip shadows which will be handled in RenderOnePassPointLightShadows
-		if (ProjectedShadowInfo->bOnePassPointLightShadow)
+		if (ProjectedShadowInfo->CascadeSettings.bOnePassPointLightShadow)
 		{
 			bShadowIsVisible = false;
 		}
@@ -3120,7 +3120,7 @@ bool FDeferredShadingSceneRenderer::RenderProjectedShadows(FRHICommandListImmedi
 		}
 
 		// Skip shadows which will be handled in RenderOnePassPointLightShadows or RenderReflectiveShadowmaps
-		if (ProjectedShadowInfo->bOnePassPointLightShadow || ProjectedShadowInfo->bReflectiveShadowmap)
+		if (ProjectedShadowInfo->CascadeSettings.bOnePassPointLightShadow || ProjectedShadowInfo->bReflectiveShadowmap)
 		{
 			bShadowIsVisible = false;
 		}
@@ -3212,7 +3212,7 @@ bool FDeferredShadingSceneRenderer::RenderProjectedShadows(FRHICommandListImmedi
 			for (int32 ShadowIndex = 0; ShadowIndex < Shadows.Num(); ShadowIndex++)
 			{
 				FProjectedShadowInfo* ProjectedShadowInfo = Shadows[ShadowIndex];
-				if (ProjectedShadowInfo->bAllocated && !ProjectedShadowInfo->bTranslucentShadow && !ProjectedShadowInfo->bRayTracedDistanceFieldShadow)
+				if (ProjectedShadowInfo->bAllocated && !ProjectedShadowInfo->bTranslucentShadow && !ProjectedShadowInfo->CascadeSettings.bRayTracedDistanceField)
 				{
 					ProjectedShadowInfo->RenderDepth(RHICmdList, this, SetShadowRenderTargets);
 				}
@@ -3235,7 +3235,7 @@ bool FDeferredShadingSceneRenderer::RenderProjectedShadows(FRHICommandListImmedi
 				if (ProjectedShadowInfo->bAllocated
 					&& ProjectedShadowInfo->bWholeSceneShadow
 					// Not supported on translucency yet
-					&& !ProjectedShadowInfo->bRayTracedDistanceFieldShadow
+					&& !ProjectedShadowInfo->CascadeSettings.bRayTracedDistanceField
 					// Don't inject shadowed lighting with whole scene shadows used for previewing a light with static shadows,
 					// Since that would cause a mismatch with the built lighting
 					// However, stationary directional lights allow whole scene shadows that blend with precomputed shadowing
@@ -3422,7 +3422,7 @@ bool FForwardShadingSceneRenderer::RenderShadowDepthMap(FRHICommandListImmediate
 		}
 
 		// Skip shadows which will be handled in RenderOnePassPointLightShadows or RenderReflectiveShadowmaps
-		if (ProjectedShadowInfo->bOnePassPointLightShadow || ProjectedShadowInfo->bReflectiveShadowmap)
+		if (ProjectedShadowInfo->CascadeSettings.bOnePassPointLightShadow || ProjectedShadowInfo->bReflectiveShadowmap)
 		{
 			bShadowIsVisible = false;
 		}
diff --git a/Engine/Source/Runtime/Renderer/Private/ShadowRendering.h b/Engine/Source/Runtime/Renderer/Private/ShadowRendering.h
--- a/Engine/Source/Runtime/Renderer/Private/ShadowRendering.h
+++ b/Engine/Source/Runtime/Renderer/Private/ShadowRendering.h
@@ -491,15 +491,6 @@
 class FProjectedShadowInfo : public FRefCountedObject
 {
 public:
-
-	friend class FShadowDepthVS;
-	//	friend class FShadowDepthPS;
-	template <bool bRenderingReflectiveShadowMaps> friend class TShadowDepthBasePS;
-
-	friend class FShadowProjectionVS;
-	friend class FShadowProjectionPS;
-	friend class FShadowDepthDrawingPolicyFactory;
-
 	typedef TArray<const FPrimitiveSceneInfo*,SceneRenderingAllocator> PrimitiveArrayType;
 
 	const FLightSceneInfo* const LightSceneInfo;
@@ -562,12 +553,6 @@
 	/** Fade Alpha per view. */
 	TArray<float, TInlineAllocator<2> > FadeAlphas;
 
-	/** 
-	 * Index of the split if this is a whole scene shadow from a directional light, 
-	 * Or index of the direction if this is a whole scene shadow from a point light, otherwise INDEX_NONE. 
-	 */
-	const int32 ShadowSplitIndex;
-
 	/** Whether the shadow has been allocated in the shadow depth buffer, and its X and Y properties have been initialized. */
 	uint32 bAllocated : 1;
 
@@ -592,9 +577,6 @@
 	/** Whether this shadow affects the whole scene or only a group of objects. */
 	uint32 bWholeSceneShadow : 1;
 
-	/** Whether the shadow is a point light shadow that renders all faces of a cubemap in one pass. */
-	const uint32 bOnePassPointLightShadow : 1;
-
 	/** Whether the shadow needs to render reflective shadow maps. */ 
 	uint32 bReflectiveShadowmap : 1; 
 
@@ -607,9 +589,6 @@
 	/** To not cast a shadow on the ground outside the object and having higher quality (useful for first person weapon). */
 	uint32 bSelfShadowOnly : 1;
 
-	/** Whether the shadow is a directional light cascade that should be computed by ray tracing mesh distance fields. */
-	uint32 bRayTracedDistanceFieldShadow : 1;
-
 	uint32 bValidTransform : 1;
 
 	TBitArray<SceneRenderingBitArrayAllocator> StaticMeshWholeSceneShadowDepthMap;
@@ -621,31 +600,6 @@
 	/** Frustums for each cubemap face, used for object culling one pass point light shadows. */
 	TArray<FConvexVolume> OnePassShadowFrustums;
 
-private:
-
-	/** dynamic shadow casting elements */
-	PrimitiveArrayType SubjectPrimitives;
-	/** For preshadows, this contains the receiver primitives to mask the projection to. */
-	PrimitiveArrayType ReceiverPrimitives;
-	/** Subject primitives with translucent relevance. */
-	PrimitiveArrayType SubjectTranslucentPrimitives;
-
-	/** Static shadow casting elements. */
-	TArray<FShadowStaticMeshElement,SceneRenderingAllocator> SubjectMeshElements;
-
-	/** Dynamic mesh elements for subject primitives. */
-	TArray<FMeshBatchAndRelevance,SceneRenderingAllocator> DynamicSubjectMeshElements;
-	/** Dynamic mesh elements for receiver primitives. */
-	TArray<FMeshBatchAndRelevance,SceneRenderingAllocator> DynamicReceiverMeshElements;
-	/** Dynamic mesh elements for translucent subject primitives. */
-	TArray<FMeshBatchAndRelevance,SceneRenderingAllocator> DynamicSubjectTranslucentMeshElements;
-
-	/**
-	 * Bias during in shadowmap rendering, stored redundantly for better performance 
-	 * Set by UpdateShaderDepthBias(), get with GetShaderDepthBias(), -1 if not set
-	 */
-	float ShaderDepthBias;
-
 public:
 
 	/**
@@ -754,8 +708,7 @@
 
 	inline bool IsWholeSceneDirectionalShadow() const 
 	{ 
-		check(bDirectionalLight == LightSceneInfo->Proxy->GetLightType() == LightType_Directional);
-		return bWholeSceneShadow && ShadowSplitIndex >= 0 && bDirectionalLight; 
+		return bWholeSceneShadow && CascadeSettings.ShadowSplitIndex >= 0 && bDirectionalLight; 
 	}
 
 	inline bool IsWholeScenePointLightShadow() const
@@ -767,6 +720,28 @@
 	void SortSubjectMeshElements();
 
 private:
+	/** dynamic shadow casting elements */
+	PrimitiveArrayType SubjectPrimitives;
+	/** For preshadows, this contains the receiver primitives to mask the projection to. */
+	PrimitiveArrayType ReceiverPrimitives;
+	/** Subject primitives with translucent relevance. */
+	PrimitiveArrayType SubjectTranslucentPrimitives;
+
+	/** Static shadow casting elements. */
+	TArray<FShadowStaticMeshElement,SceneRenderingAllocator> SubjectMeshElements;
+
+	/** Dynamic mesh elements for subject primitives. */
+	TArray<FMeshBatchAndRelevance,SceneRenderingAllocator> DynamicSubjectMeshElements;
+	/** Dynamic mesh elements for receiver primitives. */
+	TArray<FMeshBatchAndRelevance,SceneRenderingAllocator> DynamicReceiverMeshElements;
+	/** Dynamic mesh elements for translucent subject primitives. */
+	TArray<FMeshBatchAndRelevance,SceneRenderingAllocator> DynamicSubjectTranslucentMeshElements;
 
 	/**
+	 * Bias during in shadowmap rendering, stored redundantly for better performance 
+	 * Set by UpdateShaderDepthBias(), get with GetShaderDepthBias(), -1 if not set
+	 */
+	float ShaderDepthBias;
+
+	/**
 	* Renders the shadow subject depth, to a particular hacked view
@@ -793,6 +768,12 @@
 		PrimitiveArrayType& PrimitiveArray, 
 		TArray<FMeshBatchAndRelevance,SceneRenderingAllocator>& OutDynamicMeshElements, 
 		TArray<const FSceneView*>& ReusedViewsArray);
+
+	friend class FShadowDepthVS;
+	template <bool bRenderingReflectiveShadowMaps> friend class TShadowDepthBasePS;
+	friend class FShadowProjectionVS;
+	friend class FShadowProjectionPS;
+	friend class FShadowDepthDrawingPolicyFactory;
 };
 
 /** Shader parameters for rendering the depth of a mesh for shadowing. */
diff --git a/Engine/Source/Runtime/Renderer/Private/ShadowSetup.cpp b/Engine/Source/Runtime/Renderer/Private/ShadowSetup.cpp
--- a/Engine/Source/Runtime/Renderer/Private/ShadowSetup.cpp
+++ b/Engine/Source/Runtime/Renderer/Private/ShadowSetup.cpp
@@ -363,36 +363,34 @@ FProjectedShadowInfo::FProjectedShadowInfo(
 	uint32 MaxShadowResolutionY,
 	float InMaxScreenPercent,
 	const TArray<float, TInlineAllocator<2> >& InFadeAlphas,
-	bool bInTranslucentShadow
-	):
-	LightSceneInfo(InLightSceneInfo),
-	LightSceneInfoCompact(InLightSceneInfo),
-	ParentSceneInfo(InParentSceneInfo),
-	DependentView(NULL),
-	ShadowId(INDEX_NONE),
-	PreShadowTranslation(Initializer.PreShadowTranslation),	
-	ShadowBounds(Initializer.SubjectBounds.Origin - Initializer.PreShadowTranslation, Initializer.SubjectBounds.SphereRadius),
-	X(0),
-	Y(0),
-	ResolutionX(InResolutionX),
-	ResolutionY(0),
-	MaxScreenPercent(InMaxScreenPercent),
-	FadeAlphas(InFadeAlphas),
-	ShadowSplitIndex(INDEX_NONE),
-	bAllocated(false),
-	bAllocatedInTranslucentLayout(false),
-	bRendered(false),
-	bAllocatedInPreshadowCache(false),
-	bDepthsCached(false),
-	bDirectionalLight(InLightSceneInfo->Proxy->GetLightType() == LightType_Directional),
-	bWholeSceneShadow(false),
-	bOnePassPointLightShadow(false),
-	bReflectiveShadowmap(false),
-	bTranslucentShadow(bInTranslucentShadow),
-	bPreShadow(bInPreShadow),
-	bSelfShadowOnly(InParentSceneInfo->Proxy->CastsSelfShadowOnly()),
-	bRayTracedDistanceFieldShadow(false)
+	bool bInTranslucentShadow)
+	: LightSceneInfo(InLightSceneInfo)
+	, LightSceneInfoCompact(InLightSceneInfo)
+	, ParentSceneInfo(InParentSceneInfo)
+	, DependentView(NULL)
+	, ShadowId(INDEX_NONE)
+	, PreShadowTranslation(Initializer.PreShadowTranslation)	
+	, ShadowBounds(Initializer.SubjectBounds.Origin - Initializer.PreShadowTranslation, Initializer.SubjectBounds.SphereRadius)
+	, X(0)
+	, Y(0)
+	, ResolutionX(InResolutionX)
+	, ResolutionY(0)
+	, MaxScreenPercent(InMaxScreenPercent)
+	, FadeAlphas(InFadeAlphas)
+	, bAllocated(false)
+	, bAllocatedInTranslucentLayout(false)
+	, bRendered(false)
+	, bAllocatedInPreshadowCache(false)
+	, bDepthsCached(false)
+	, bDirectionalLight(InLightSceneInfo->Proxy->GetLightType() == LightType_Directional)
+	, bWholeSceneShadow(false)
+	, bReflectiveShadowmap(false)
+	, bTranslucentShadow(bInTranslucentShadow)
+	, bPreShadow(bInPreShadow)
+	, bSelfShadowOnly(InParentSceneInfo->Proxy->CastsSelfShadowOnly())
 {
+	check(!CascadeSettings.bRayTracedDistanceField);
+
 	const FMatrix WorldToLightScaled = Initializer.WorldToLight * FScaleMatrix(Initializer.Scales);
 	
 	// Create an array of the extreme vertices of the subject's bounds.
@@ -489,35 +487,31 @@ FProjectedShadowInfo::FProjectedShadowInfo(
 	const FWholeSceneProjectedShadowInitializer& Initializer,
 	uint32 InResolutionX,
 	uint32 InResolutionY,
-	bool bInReflectiveShadowMap
-	)
-:	LightSceneInfo(InLightSceneInfo)
-,	LightSceneInfoCompact(InLightSceneInfo)
-,	ParentSceneInfo(NULL)
-,	DependentView(InDependentView)
-,	ShadowId(INDEX_NONE)
-,	PreShadowTranslation(Initializer.PreShadowTranslation)
-,	CascadeSettings(Initializer.CascadeSettings)
-,	X(0)
-,	Y(0)
-,	ResolutionX(InResolutionX)
-,	ResolutionY(InResolutionY)
-,	MaxScreenPercent(1.0f)
-,	ShadowSplitIndex(Initializer.InitShadowSplitIndex)
-,	bAllocated(false)
-,	bAllocatedInTranslucentLayout(false)
-,	bRendered(false)
-,	bAllocatedInPreshadowCache(false)
-,	bDepthsCached(false)
-,	bDirectionalLight(InLightSceneInfo->Proxy->GetLightType() == LightType_Directional)
-,	bWholeSceneShadow(true)
-,	bOnePassPointLightShadow(Initializer.bOnePassPointLightShadow)
-,	bReflectiveShadowmap(bInReflectiveShadowMap) 
-,	bTranslucentShadow(false)
-,	bPreShadow(false)
-,	bSelfShadowOnly(false)
-,	bRayTracedDistanceFieldShadow(Initializer.bRayTracedDistanceFieldShadow)
-,	bValidTransform(true)
+	bool bInReflectiveShadowMap)
+	: LightSceneInfo(InLightSceneInfo)
+	, LightSceneInfoCompact(InLightSceneInfo)
+	, ParentSceneInfo(NULL)
+	, DependentView(InDependentView)
+	, ShadowId(INDEX_NONE)
+	, PreShadowTranslation(Initializer.PreShadowTranslation)
+	, CascadeSettings(Initializer.CascadeSettings)
+	, X(0)
+	, Y(0)
+	, ResolutionX(InResolutionX)
+	, ResolutionY(InResolutionY)
+	, MaxScreenPercent(1.0f)
+	, bAllocated(false)
+	, bAllocatedInTranslucentLayout(false)
+	, bRendered(false)
+	, bAllocatedInPreshadowCache(false)
+	, bDepthsCached(false)
+	, bDirectionalLight(InLightSceneInfo->Proxy->GetLightType() == LightType_Directional)
+	, bWholeSceneShadow(true)
+	, bReflectiveShadowmap(bInReflectiveShadowMap) 
+	, bTranslucentShadow(false)
+	, bPreShadow(false)
+	, bSelfShadowOnly(false)
+	, bValidTransform(true)
 {	
 	FVector	XAxis, YAxis;
 	Initializer.FaceDirection.FindBestAxisVectors(XAxis,YAxis);
@@ -529,8 +523,8 @@ FProjectedShadowInfo::FProjectedShadowInfo(
 
 	if(bInReflectiveShadowMap)
 	{
-		check(!bOnePassPointLightShadow);
-		check(!ShadowSplitIndex);
+		check(!CascadeSettings.bOnePassPointLightShadow);
+		check(!CascadeSettings.ShadowSplitIndex);
 
 		// Quantise the RSM in shadow texel space
 		static bool bQuantize = true;
@@ -591,10 +585,10 @@ FProjectedShadowInfo::FProjectedShadowInfo(
 			PreShadowTranslation = -SnappedWorldPosition;
 		}
 
-		if (ShadowSplitIndex >= 0 && bDirectionalLight)
+		if (CascadeSettings.ShadowSplitIndex >= 0 && bDirectionalLight)
 		{
 			checkSlow(InDependentView);
-			ShadowBounds = InLightSceneInfo->Proxy->GetShadowSplitBounds(*InDependentView, bRayTracedDistanceFieldShadow ? INDEX_NONE : ShadowSplitIndex, 0);
+			ShadowBounds = InLightSceneInfo->Proxy->GetShadowSplitBounds(*InDependentView, CascadeSettings.bRayTracedDistanceField ? INDEX_NONE : CascadeSettings.ShadowSplitIndex, 0);
 		}
 		else
 		{
@@ -619,7 +613,7 @@ FProjectedShadowInfo::FProjectedShadowInfo(
 		+ WorldToLightScaled.InverseFast().TransformVector(Initializer.FaceDirection) * Initializer.SubjectBounds.SphereRadius
 		).Z;
 
-	if (Initializer.bOnePassPointLightShadow)
+	if (CascadeSettings.bOnePassPointLightShadow)
 	{
 		MaxSubjectDepth = Initializer.SubjectBounds.SphereRadius;
 	}
@@ -647,7 +641,7 @@ void FProjectedShadowInfo::AddSubjectPrimitive(FPrimitiveSceneInfo* PrimitiveSce
 	QUICK_SCOPE_CYCLE_COUNTER(STAT_AddSubjectPrimitive);
 
 	// Ray traced shadows use the GPU managed distance field object buffers, no CPU culling should be used
-	check(!bRayTracedDistanceFieldShadow);
+	check(!CascadeSettings.bRayTracedDistanceField);
 
 	if (!ReceiverPrimitives.Contains(PrimitiveSceneInfo))
 	{
@@ -740,7 +734,7 @@ void FProjectedShadowInfo::AddSubjectPrimitive(FPrimitiveSceneInfo* PrimitiveSce
 							// cull objects that are not supposed to be in the far shadow cascades
 							continue;
 						}
-						if(bRayTracedDistanceFieldShadow && Proxy->AffectsDistanceFieldLighting())
+						if(CascadeSettings.bRayTracedDistanceField && Proxy->AffectsDistanceFieldLighting())
 						{
 							// cull objects that already have been handles by the ray traced distance field ones
 							continue;
@@ -1561,7 +1555,7 @@ void FDeferredShadingSceneRenderer::CreateWholeSceneProjectedShadow(FLightSceneI
 		checkSlow(ProjectedShadowInitializers.Num() > 0);
 
 		// Shadow resolution constants.
-		const uint32 EffectiveDoubleShadowBorder = ProjectedShadowInitializers[0].bOnePassPointLightShadow ? 0 : SHADOW_BORDER * 2;
+		const uint32 EffectiveDoubleShadowBorder = ProjectedShadowInitializers[0].CascadeSettings.bOnePassPointLightShadow ? 0 : SHADOW_BORDER * 2;
 		const int32 MinShadowResolution = CVarMinShadowResolution.GetValueOnRenderThread();
 		const int32 MaxShadowResolutionSetting = GetCachedScalabilityCVars().MaxShadowResolution;
 		const FIntPoint ShadowBufferResolution = GSceneRenderTargets.GetShadowDepthTextureResolution();
@@ -1624,7 +1618,7 @@ void FDeferredShadingSceneRenderer::CreateWholeSceneProjectedShadow(FLightSceneI
 				const uint32 DesiredSizeY = FMath::TruncToInt(MaxDesiredResolution);
 				int32 SizeY = DesiredSizeY >= MaxShadowResolutionY ? MaxShadowResolutionY : (1 << (FMath::CeilLogTwo(DesiredSizeY) - 1));
 
-				if (ProjectedShadowInitializer.bOnePassPointLightShadow)
+				if (ProjectedShadowInitializer.CascadeSettings.bOnePassPointLightShadow)
 				{
 					// Round to a resolution that is supported for one pass point light shadows
 					SizeX = SizeY = GSceneRenderTargets.GetCubeShadowDepthZResolution(GSceneRenderTargets.GetCubeShadowDepthZIndex(MaxDesiredResolution));
@@ -1644,7 +1638,7 @@ void FDeferredShadingSceneRenderer::CreateWholeSceneProjectedShadow(FLightSceneI
 				VisibleLightInfo.MemStackProjectedShadows.Add(ProjectedShadowInfo);
 				VisibleLightInfo.AllProjectedShadows.Add(ProjectedShadowInfo);
 
-				if (ProjectedShadowInitializer.bOnePassPointLightShadow)
+				if (ProjectedShadowInitializer.CascadeSettings.bOnePassPointLightShadow)
 				{
 					const static FVector CubeDirections[6] = 
 					{
@@ -1684,7 +1678,7 @@ void FDeferredShadingSceneRenderer::CreateWholeSceneProjectedShadow(FLightSceneI
 				}
 
 				// Ray traced shadows use the GPU managed distance field object buffers, no CPU culling should be used
-				if (!ProjectedShadowInfo->bRayTracedDistanceFieldShadow)
+				if (!ProjectedShadowInfo->CascadeSettings.bRayTracedDistanceField)
 				{
 					// Add all the shadow casting primitives affected by the light to the shadow's subject primitive list.
 					for(FLightPrimitiveInteraction* Interaction = LightSceneInfo->DynamicPrimitiveList;
@@ -1778,7 +1772,7 @@ void FSceneRenderer::InitProjectedShadowVisibility(FRHICommandListImmediate& RHI
 								ProjectedShadowInfo.ParentSceneInfo->PrimitiveComponentId :
 								FPrimitiveComponentId(),
 							ProjectedShadowInfo.LightSceneInfo->Proxy->GetLightComponent(),
-							ProjectedShadowInfo.ShadowSplitIndex,
+							ProjectedShadowInfo.CascadeSettings.ShadowSplitIndex,
 							ProjectedShadowInfo.bTranslucentShadow
 							);
 
@@ -1802,7 +1796,7 @@ void FSceneRenderer::InitProjectedShadowVisibility(FRHICommandListImmediate& RHI
 							{
 								// Get split color
 								FColor Color = FColor::White;
-								switch(ProjectedShadowInfo.ShadowSplitIndex)
+								switch(ProjectedShadowInfo.CascadeSettings.ShadowSplitIndex)
 								{
 									case 0: Color = FColor::Red; break;
 									case 1: Color = FColor::Yellow; break;
@@ -1872,11 +1866,11 @@ void FSceneRenderer::InitProjectedShadowVisibility(FRHICommandListImmediate& RHI
 						UE_LOG(LogRenderer, Display, TEXT("   Shadow %d/%d: ShadowId=%d"),  ShadowIndex, ShadowCount, ProjectedShadowInfo.ShadowId);
 						UE_LOG(LogRenderer, Display, TEXT("    WholeSceneDir=%d SplitIndex=%d near=%f far=%f"),
 							ProjectedShadowInfo.IsWholeSceneDirectionalShadow(),
-							ProjectedShadowInfo.ShadowSplitIndex,
+							ProjectedShadowInfo.CascadeSettings.ShadowSplitIndex,
 							ProjectedShadowInfo.CascadeSettings.SplitNear,
 							ProjectedShadowInfo.CascadeSettings.SplitFar);
 						UE_LOG(LogRenderer, Display, TEXT("    bDistField=%d bFarShadows=%d Bounds=%f,%f,%f,%f"),
-							ProjectedShadowInfo.bRayTracedDistanceFieldShadow,
+							ProjectedShadowInfo.CascadeSettings.bRayTracedDistanceField,
 							ProjectedShadowInfo.CascadeSettings.bFarShadowCascade,
 							ProjectedShadowInfo.ShadowBounds.Center.X,
 							ProjectedShadowInfo.ShadowBounds.Center.Y,
@@ -2191,7 +2185,7 @@ void FSceneRenderer::AddViewDependentWholeSceneShadowsForView(
 					ShadowInfos.Add(ProjectedShadowInfo);
 
 					// Ray traced shadows use the GPU managed distance field object buffers, no CPU culling needed
-					if (!ProjectedShadowInfo->bRayTracedDistanceFieldShadow)
+					if (!ProjectedShadowInfo->CascadeSettings.bRayTracedDistanceField)
 					{
 						ShadowInfosThatNeedCulling.Add(ProjectedShadowInfo);
 					}
@@ -2212,7 +2206,7 @@ void FSceneRenderer::AddViewDependentWholeSceneShadowsForView(
 					if (LightSceneInfo.Proxy->GetViewDependentRsmWholeSceneProjectedShadowInitializer(View, Lpv.GetBoundingBox(), ProjectedShadowInitializer))
 					{
 						// moved out from the FProjectedShadowInfo constructor
-						ProjectedShadowInitializer.InitShadowSplitIndex = 0;
+						ProjectedShadowInitializer.CascadeSettings.ShadowSplitIndex = 0;
 
 						const FIntPoint ShadowBufferResolution = GSceneRenderTargets.GetReflectiveShadowMapTextureResolution();
 
@@ -2232,7 +2226,7 @@ void FSceneRenderer::AddViewDependentWholeSceneShadowsForView(
 						ShadowInfos.Add(ProjectedShadowInfo); // or separate list?
 
 						// Ray traced shadows use the GPU managed distance field object buffers, no CPU culling needed
-						if (!ProjectedShadowInfo->bRayTracedDistanceFieldShadow)
+						if (!ProjectedShadowInfo->CascadeSettings.bRayTracedDistanceField)
 						{
 							ShadowInfosThatNeedCulling.Add(ProjectedShadowInfo);
 						}
diff --git a/Engine/Source/Runtime/Renderer/Private/TranslucentLighting.cpp b/Engine/Source/Runtime/Renderer/Private/TranslucentLighting.cpp
--- a/Engine/Source/Runtime/Renderer/Private/TranslucentLighting.cpp
+++ b/Engine/Source/Runtime/Renderer/Private/TranslucentLighting.cpp
@@ -1535,12 +1535,16 @@ static void InjectTranslucentLightArray(FRHICommandListImmediate& RHICmdList, co
 					if (InjectionData.ProjectedShadowInfo)
 					{
 						// shadows, restricting light contribution to the cascade bounds (except last cascade far to get light functions and no shadows there)
-						SetInjectionShader<LightType_Directional, true>(RHICmdList, View, InjectionData.LightFunctionMaterialProxy, LightSceneInfo, InjectionData.ProjectedShadowInfo, InjectionData.ProjectedShadowInfo->ShadowSplitIndex, VolumeCascadeIndex, *VertexShader, *GeometryShader, InjectionData.bApplyLightFunction, false);
+						SetInjectionShader<LightType_Directional, true>(RHICmdList, View, InjectionData.LightFunctionMaterialProxy, LightSceneInfo,
+							InjectionData.ProjectedShadowInfo, InjectionData.ProjectedShadowInfo->CascadeSettings.ShadowSplitIndex, VolumeCascadeIndex,
+							*VertexShader, *GeometryShader, InjectionData.bApplyLightFunction, false);
 					}
 					else
 					{
 						// no shadows
-						SetInjectionShader<LightType_Directional, false>(RHICmdList, View, InjectionData.LightFunctionMaterialProxy, LightSceneInfo, InjectionData.ProjectedShadowInfo, -1, VolumeCascadeIndex, *VertexShader, *GeometryShader, InjectionData.bApplyLightFunction, false);
+						SetInjectionShader<LightType_Directional, false>(RHICmdList, View, InjectionData.LightFunctionMaterialProxy, LightSceneInfo,
+							InjectionData.ProjectedShadowInfo, -1, VolumeCascadeIndex,
+							*VertexShader, *GeometryShader, InjectionData.bApplyLightFunction, false);
 					}
 				}
 				else
@@ -1552,11 +1556,15 @@ static void InjectTranslucentLightArray(FRHICommandListImmediate& RHICmdList, co
 
 					if (InjectionData.ProjectedShadowInfo)
 					{
-						SetInjectionShader<LightType_Point, true>(RHICmdList, View, InjectionData.LightFunctionMaterialProxy, LightSceneInfo, InjectionData.ProjectedShadowInfo, -1, VolumeCascadeIndex, *VertexShader, *GeometryShader, InjectionData.bApplyLightFunction, bInverseSquared);
+						SetInjectionShader<LightType_Point, true>(RHICmdList, View, InjectionData.LightFunctionMaterialProxy, LightSceneInfo,
+							InjectionData.ProjectedShadowInfo, -1, VolumeCascadeIndex,
+							*VertexShader, *GeometryShader, InjectionData.bApplyLightFunction, bInverseSquared);
 					}
 					else
 					{
-						SetInjectionShader<LightType_Point, false>(RHICmdList, View, InjectionData.LightFunctionMaterialProxy, LightSceneInfo, InjectionData.ProjectedShadowInfo, -1, VolumeCascadeIndex, *VertexShader, *GeometryShader, InjectionData.bApplyLightFunction, bInverseSquared);
+						SetInjectionShader<LightType_Point, false>(RHICmdList, View, InjectionData.LightFunctionMaterialProxy, LightSceneInfo,
+							InjectionData.ProjectedShadowInfo, -1, VolumeCascadeIndex,
+							*VertexShader, *GeometryShader, InjectionData.bApplyLightFunction, bInverseSquared);
 					}
 				}
 
