UE4 - Fix tessellation shaders not calculating velocity correctly

diff --git a/Engine/Source/Runtime/Engine/Private/ShaderCompiler/ShaderCompiler.cpp b/Engine/Source/Runtime/Engine/Private/ShaderCompiler/ShaderCompiler.cpp
--- a/Engine/Source/Runtime/Engine/Private/ShaderCompiler/ShaderCompiler.cpp
+++ b/Engine/Source/Runtime/Engine/Private/ShaderCompiler/ShaderCompiler.cpp
@@ -2174,6 +2174,11 @@ void GlobalBeginCompileShader(
 		Input.Environment.SetDefine(TEXT("ALLOW_STATIC_LIGHTING"), CVar ? (CVar->GetValueOnGameThread() != 0) : 1);
 	}
 
+	{
+		static const auto CVar = IConsoleManager::Get().FindTConsoleVariableDataInt(TEXT("r.BasePassOutputsVelocity"));
+		Input.Environment.SetDefine(TEXT("OUTPUT_GBUFFER_VELOCITY"), CVar ? (CVar->GetValueOnGameThread() != 0) : 0);
+	}
+
 	{
 		static const auto CVar = IConsoleManager::Get().FindTConsoleVariableDataInt(TEXT("r.GBuffer"));
 		Input.Environment.SetDefine(TEXT("NO_GBUFFER"), CVar ? (CVar->GetValueOnGameThread() == 0) : 0);
diff --git a/Engine/Source/Runtime/Renderer/Private/BasePassRendering.h b/Engine/Source/Runtime/Renderer/Private/BasePassRendering.h
--- a/Engine/Source/Runtime/Renderer/Private/BasePassRendering.h
+++ b/Engine/Source/Runtime/Renderer/Private/BasePassRendering.h
@@ -50,7 +50,6 @@ public:
 	{
 		FMeshMaterialShader::ModifyCompilationEnvironment(Platform, Material, OutEnvironment);
 		LightMapPolicyType::ModifyCompilationEnvironment(Platform, Material, OutEnvironment);
-		OutEnvironment.SetDefine(TEXT("OUTPUT_GBUFFER_VELOCITY"), (uint32)(FVelocityRendering::OutputsToGBuffer() ? 1 : 0));
 	}
 
 	virtual bool Serialize(FArchive& Ar)
@@ -336,7 +335,6 @@ public:
 		LightMapPolicyType::ModifyCompilationEnvironment(Platform, Material, OutEnvironment);
 
 		const bool bOutputVelocity = FVelocityRendering::OutputsToGBuffer();
-		OutEnvironment.SetDefine(TEXT("OUTPUT_GBUFFER_VELOCITY"), (uint32)(bOutputVelocity ? 1 : 0));
 		if (bOutputVelocity)
 		{
 			OutEnvironment.SetRenderTargetOutputFormat(6, PF_G16R16);
