UE4 - Add mesh batch to base pass SetMesh function

diff --git a/Engine/Source/Runtime/Renderer/Private/BasePassRendering.h b/Engine/Source/Runtime/Renderer/Private/BasePassRendering.h
--- a/Engine/Source/Runtime/Renderer/Private/BasePassRendering.h
+++ b/Engine/Source/Runtime/Renderer/Private/BasePassRendering.h
@@ -82,7 +82,7 @@ public:
 		}
 	}
 
-	void SetMesh(FRHICommandList& RHICmdList, const FVertexFactory* VertexFactory,const FSceneView& View,const FPrimitiveSceneProxy* Proxy,const FMeshBatchElement& BatchElement);
+	void SetMesh(FRHICommandList& RHICmdList, const FVertexFactory* VertexFactory,const FSceneView& View,const FPrimitiveSceneProxy* Proxy, const FMeshBatch& Mesh, const FMeshBatchElement& BatchElement);
 
 private:
 	
@@ -664,7 +664,7 @@ public:
 			ElementData.LightMapElementData);
 
 		const FMeshBatchElement& BatchElement = Mesh.Elements[BatchElementIndex];
-		VertexShader->SetMesh(RHICmdList, VertexFactory,View,PrimitiveSceneProxy,BatchElement);
+		VertexShader->SetMesh(RHICmdList, VertexFactory,View,PrimitiveSceneProxy, Mesh,BatchElement);
 		
 		if(HullShader && DomainShader)
 		{
diff --git a/Engine/Source/Runtime/Renderer/Private/BasePassRendering.inl b/Engine/Source/Runtime/Renderer/Private/BasePassRendering.inl
--- a/Engine/Source/Runtime/Renderer/Private/BasePassRendering.inl
+++ b/Engine/Source/Runtime/Renderer/Private/BasePassRendering.inl
@@ -8,7 +8,7 @@
 #pragma once
 
 template<typename LightMapPolicyType>
-inline void TBasePassVertexShaderBaseType<LightMapPolicyType>::SetMesh(FRHICommandList& RHICmdList, const FVertexFactory* VertexFactory,const FSceneView& View,const FPrimitiveSceneProxy* Proxy,const FMeshBatchElement& BatchElement)
+inline void TBasePassVertexShaderBaseType<LightMapPolicyType>::SetMesh(FRHICommandList& RHICmdList, const FVertexFactory* VertexFactory,const FSceneView& View,const FPrimitiveSceneProxy* Proxy, const FMeshBatch& Mesh, const FMeshBatchElement& BatchElement)
 {
 	FVertexShaderRHIParamRef VertexShaderRHI = GetVertexShader();
 	FMeshMaterialShader::SetMesh(RHICmdList, VertexShaderRHI, VertexFactory, View, Proxy, BatchElement);
@@ -20,17 +20,11 @@ inline void TBasePassVertexShaderBaseType<LightMapPolicyType>::SetMesh(FRHIComma
 		FMatrix PreviousLocalToWorld;
 		bool bHasPreviousLocalToWorld = false;
 		const auto& ViewInfo = (const FViewInfo&)View;
-		if (FVelocityDrawingPolicy::HasVelocityOnBasePass(ViewInfo, Proxy, Proxy->GetPrimitiveSceneInfo(), 
+		if (FVelocityDrawingPolicy::HasVelocityOnBasePass(ViewInfo, Proxy, Proxy->GetPrimitiveSceneInfo(), Mesh,
 			bHasPreviousLocalToWorld, PreviousLocalToWorld))
 		{
-			if (bHasPreviousLocalToWorld)
-			{
-				SetShaderValue(RHICmdList, VertexShaderRHI, PreviousLocalToWorldParameter, PreviousLocalToWorld.ConcatTranslation(ViewInfo.PrevViewMatrices.PreViewTranslation));
-			}
-			else
-			{
-				SetShaderValue(RHICmdList, VertexShaderRHI, PreviousLocalToWorldParameter, Proxy->GetLocalToWorld().ConcatTranslation(ViewInfo.PrevViewMatrices.PreViewTranslation));
-			}
+			const FMatrix& Transform = bHasPreviousLocalToWorld ? PreviousLocalToWorld : Proxy->GetLocalToWorld();
+			SetShaderValue(RHICmdList, VertexShaderRHI, PreviousLocalToWorldParameter, Transform.ConcatTranslation(ViewInfo.PrevViewMatrices.PreViewTranslation));
 
 			SetShaderValue(RHICmdList, VertexShaderRHI, SkipOutputVelocityParameter, 0.0f);
 		}
diff --git a/Engine/Source/Runtime/Renderer/Private/VelocityRendering.cpp b/Engine/Source/Runtime/Renderer/Private/VelocityRendering.cpp
--- a/Engine/Source/Runtime/Renderer/Private/VelocityRendering.cpp
+++ b/Engine/Source/Runtime/Renderer/Private/VelocityRendering.cpp
@@ -330,7 +330,7 @@ bool FVelocityDrawingPolicy::HasVelocity(const FViewInfo& View, const FPrimitive
 	return true;
 }
 
-bool FVelocityDrawingPolicy::HasVelocityOnBasePass(const FViewInfo& View,const FPrimitiveSceneProxy* PrimitiveSceneProxy, const FPrimitiveSceneInfo* PrimitiveSceneInfo, bool& bOutHasTransform, FMatrix& OutTransform)
+bool FVelocityDrawingPolicy::HasVelocityOnBasePass(const FViewInfo& View,const FPrimitiveSceneProxy* PrimitiveSceneProxy, const FPrimitiveSceneInfo* PrimitiveSceneInfo, const FMeshBatch& Mesh, bool& bOutHasTransform, FMatrix& OutTransform)
 {
 	checkSlow(IsInParallelRenderingThread());
 	// No velocity if motionblur is off, or if it's a non-moving object (treat as background in that case)
diff --git a/Engine/Source/Runtime/Renderer/Private/VelocityRendering.h b/Engine/Source/Runtime/Renderer/Private/VelocityRendering.h
--- a/Engine/Source/Runtime/Renderer/Private/VelocityRendering.h
+++ b/Engine/Source/Runtime/Renderer/Private/VelocityRendering.h
@@ -48,7 +48,7 @@ public:
 
 	/** Determines whether this primitive has motionblur velocity to render */
 	static bool HasVelocity(const FViewInfo& View, const FPrimitiveSceneInfo* PrimitiveSceneInfo);
-	static bool HasVelocityOnBasePass(const FViewInfo& View,const FPrimitiveSceneProxy* Proxy, const FPrimitiveSceneInfo* PrimitiveSceneInfo, bool& bOutHasTransform, FMatrix& OutTransform);
+	static bool HasVelocityOnBasePass(const FViewInfo& View,const FPrimitiveSceneProxy* Proxy, const FPrimitiveSceneInfo* PrimitiveSceneInfo, const FMeshBatch& Mesh, bool& bOutHasTransform, FMatrix& OutTransform);
 
 private:
 	class FVelocityVS* VertexShader;
