shadow projection code cleanup, single constructor, accessor functions for some members

diff --git a/Engine/Source/Runtime/Renderer/Private/DistanceFieldShadowing.cpp b/Engine/Source/Runtime/Renderer/Private/DistanceFieldShadowing.cpp
--- a/Engine/Source/Runtime/Renderer/Private/DistanceFieldShadowing.cpp
+++ b/Engine/Source/Runtime/Renderer/Private/DistanceFieldShadowing.cpp
@@ -428,16 +428,18 @@ public:
 		float LightSourceLength;
 		float LightMinRoughness;
 
-		ProjectedShadowInfo->LightSceneInfo->Proxy->GetParameters(LightPositionAndInvRadiusValue, LightColorAndFalloffExponent, NormalizedLightDirection, SpotAngles, LightSourceRadiusValue, LightSourceLength, LightMinRoughness);
+		const FLightSceneProxy& LightProxy = *(ProjectedShadowInfo->GetLightSceneInfo().Proxy);
+
+		LightProxy.GetParameters(LightPositionAndInvRadiusValue, LightColorAndFalloffExponent, NormalizedLightDirection, SpotAngles, LightSourceRadiusValue, LightSourceLength, LightMinRoughness);
 
 		SetShaderValue(RHICmdList, ShaderRHI, LightDirection, NormalizedLightDirection);
 		SetShaderValue(RHICmdList, ShaderRHI, LightPositionAndInvRadius, LightPositionAndInvRadiusValue);
 		// Default light source radius of 0 gives poor results
 		SetShaderValue(RHICmdList, ShaderRHI, LightSourceRadius, LightSourceRadiusValue == 0 ? 20 : FMath::Clamp(LightSourceRadiusValue, .001f, 1.0f / (4 * LightPositionAndInvRadiusValue.W)));
 
-		SetShaderValue(RHICmdList, ShaderRHI, RayStartOffsetDepthScale, ProjectedShadowInfo->LightSceneInfo->Proxy->GetRayStartOffsetDepthScale());
+		SetShaderValue(RHICmdList, ShaderRHI, RayStartOffsetDepthScale, LightProxy.GetRayStartOffsetDepthScale());
 
-		const float LightSourceAngle = FMath::Clamp(ProjectedShadowInfo->LightSceneInfo->Proxy->GetLightSourceAngle(), 0.001f, 5.0f) * PI / 180.0f;
+		const float LightSourceAngle = FMath::Clamp(LightProxy.GetLightSourceAngle(), 0.001f, 5.0f) * PI / 180.0f;
 		const FVector2D TanLightAngleAndNormalThresholdValue(FMath::Tan(LightSourceAngle), FMath::Cos(PI / 2 + LightSourceAngle));
 		SetShaderValue(RHICmdList, ShaderRHI, TanLightAngleAndNormalThreshold, TanLightAngleAndNormalThresholdValue);
 
diff --git a/Engine/Source/Runtime/Renderer/Private/LightMapRendering.h b/Engine/Source/Runtime/Renderer/Private/LightMapRendering.h
--- a/Engine/Source/Runtime/Renderer/Private/LightMapRendering.h
+++ b/Engine/Source/Runtime/Renderer/Private/LightMapRendering.h
@@ -179,8 +179,10 @@ public:
 
 		void SetLightMapScale(FRHICommandList& RHICmdList, FShader* PixelShader,const FLightMapInteraction& LightMapInteraction) const
 		{
-			SetShaderValueArray(RHICmdList, PixelShader->GetPixelShader(),LightMapScaleParameter,LightMapInteraction.GetScaleArray(),LightMapInteraction.GetNumLightmapCoefficients());
-			SetShaderValueArray(RHICmdList, PixelShader->GetPixelShader(),LightMapAddParameter,LightMapInteraction.GetAddArray(),LightMapInteraction.GetNumLightmapCoefficients());
+			const FPixelShaderRHIParamRef ShaderRHI = PixelShader->GetPixelShader();
+
+			SetShaderValueArray(RHICmdList,ShaderRHI,LightMapScaleParameter,LightMapInteraction.GetScaleArray(),LightMapInteraction.GetNumLightmapCoefficients());
+			SetShaderValueArray(RHICmdList,ShaderRHI,LightMapAddParameter,LightMapInteraction.GetAddArray(),LightMapInteraction.GetNumLightmapCoefficients());
 		}
 
 		void Serialize(FArchive& Ar)
@@ -357,9 +359,11 @@ public:
 
 		void SetMesh(FRHICommandList& RHICmdList, FShader* PixelShader, const FShadowMapInteraction& ShadowMapInteraction, const FVector4& DistanceFieldValues) const
 		{
-			SetShaderValue(RHICmdList, PixelShader->GetPixelShader(), DistanceFieldParameters, DistanceFieldValues);
+			const FPixelShaderRHIParamRef ShaderRHI = PixelShader->GetPixelShader();
+
+			SetShaderValue(RHICmdList, ShaderRHI, DistanceFieldParameters, DistanceFieldValues);
 			
-			SetShaderValue(RHICmdList, PixelShader->GetPixelShader(), StaticShadowMapMasks, FVector4(
+			SetShaderValue(RHICmdList, ShaderRHI, StaticShadowMapMasks, FVector4(
 				ShadowMapInteraction.GetChannelValid(0),
 				ShadowMapInteraction.GetChannelValid(1),
 				ShadowMapInteraction.GetChannelValid(2),
@@ -368,7 +372,7 @@ public:
 
 			SetTextureParameter(
 				RHICmdList, 
-				PixelShader->GetPixelShader(),
+				ShaderRHI,
 				StaticShadowTexture,
 				StaticShadowSampler,
 				ShadowMapInteraction.GetTexture() ? ShadowMapInteraction.GetTexture()->Resource : GWhiteTexture
@@ -567,6 +571,8 @@ public:
 	{
 		if (PixelShaderParameters)
 		{
+			const FPixelShaderRHIParamRef ShaderRHI = PixelShader->GetPixelShader();
+
 			// Set these even if ElementData.TranslucentSelfShadow is NULL to avoid a d3d debug error from the shader expecting texture SRV's when a different type are bound
 			PixelShaderParameters->TranslucencyShadowParameters.Set(RHICmdList, PixelShader);
 
@@ -575,18 +581,20 @@ public:
 				FVector4 ShadowmapMinMax;
 				FMatrix WorldToShadowMatrixValue = ElementData.TranslucentSelfShadow->GetWorldToShadowMatrix(ShadowmapMinMax);
 
-				SetShaderValue(RHICmdList, PixelShader->GetPixelShader(), PixelShaderParameters->WorldToShadowMatrix, WorldToShadowMatrixValue);
-				SetShaderValue(RHICmdList, PixelShader->GetPixelShader(), PixelShaderParameters->ShadowUVMinMax, ShadowmapMinMax);
-				SetShaderValue(RHICmdList, PixelShader->GetPixelShader(), PixelShaderParameters->DirectionalLightDirection, ElementData.TranslucentSelfShadow->LightSceneInfo->Proxy->GetDirection());
+				SetShaderValue(RHICmdList, ShaderRHI, PixelShaderParameters->WorldToShadowMatrix, WorldToShadowMatrixValue);
+				SetShaderValue(RHICmdList, ShaderRHI, PixelShaderParameters->ShadowUVMinMax, ShadowmapMinMax);
+
+				const FLightSceneProxy* const LightProxy = ElementData.TranslucentSelfShadow->GetLightSceneInfo().Proxy;
+				SetShaderValue(RHICmdList, ShaderRHI, PixelShaderParameters->DirectionalLightDirection, LightProxy->GetDirection());
 				//@todo - support fading from both views
 				const float FadeAlpha = ElementData.TranslucentSelfShadow->FadeAlphas[0];
 				// Incorporate the diffuse scale of 1 / PI into the light color
-				const FVector4 DirectionalLightColorValue(FVector(ElementData.TranslucentSelfShadow->LightSceneInfo->Proxy->GetColor() * FadeAlpha / PI), FadeAlpha);
-				SetShaderValue(RHICmdList, PixelShader->GetPixelShader(), PixelShaderParameters->DirectionalLightColor, DirectionalLightColorValue);
+				const FVector4 DirectionalLightColorValue(FVector(LightProxy->GetColor() * FadeAlpha / PI), FadeAlpha);
+				SetShaderValue(RHICmdList, ShaderRHI, PixelShaderParameters->DirectionalLightColor, DirectionalLightColorValue);
 			}
 			else
 			{
-				SetShaderValue(RHICmdList, PixelShader->GetPixelShader(), PixelShaderParameters->DirectionalLightColor, FVector4(0, 0, 0, 0));
+				SetShaderValue(RHICmdList, ShaderRHI, PixelShaderParameters->DirectionalLightColor, FVector4(0, 0, 0, 0));
 			}
 		}
 	}
diff --git a/Engine/Source/Runtime/Renderer/Private/LightPropagationVolume.cpp b/Engine/Source/Runtime/Renderer/Private/LightPropagationVolume.cpp
--- a/Engine/Source/Runtime/Renderer/Private/LightPropagationVolume.cpp
+++ b/Engine/Source/Runtime/Renderer/Private/LightPropagationVolume.cpp
@@ -1040,7 +1040,7 @@ void FLightPropagationVolume::InjectDirectionalLightRSM(
 	const FProjectedShadowInfo&	ProjectedShadowInfo,
 	const FLinearColor&			LightColour )
 {
-	const FLightSceneProxy* LightProxy = ProjectedShadowInfo.LightSceneInfo->Proxy;
+	const FLightSceneProxy* LightProxy = ProjectedShadowInfo.GetLightSceneInfo().Proxy;
 	{
 		SCOPED_DRAW_EVENT(RHICmdList, LpvInjectDirectionalLightRSM);
 
diff --git a/Engine/Source/Runtime/Renderer/Private/SceneOcclusion.cpp b/Engine/Source/Runtime/Renderer/Private/SceneOcclusion.cpp
--- a/Engine/Source/Runtime/Renderer/Private/SceneOcclusion.cpp
+++ b/Engine/Source/Runtime/Renderer/Private/SceneOcclusion.cpp
@@ -1118,10 +1118,10 @@ void FDeferredShadingSceneRenderer::BeginOcclusionTests(FRHICommandListImmediate
 
 						if (ProjectedShadowInfo.CascadeSettings.bOnePassPointLightShadow)
 						{
-							FLightSceneProxy* Proxy = ProjectedShadowInfo.LightSceneInfo->Proxy;
+							FLightSceneProxy& LightProxy = *(ProjectedShadowInfo.GetLightSceneInfo().Proxy);
 
 							// Query one pass point light shadows separately because they don't have a shadow frustum, they have a bounding sphere instead.
-							FSphere LightBounds = Proxy->GetBoundingSphere();
+							FSphere LightBounds = LightProxy.GetBoundingSphere();
 
 							const bool bCameraInsideLightGeometry = ((FVector)View.ViewMatrices.ViewOrigin - LightBounds.Center).SizeSquared() < FMath::Square(LightBounds.W * 1.05f + View.NearClippingDistance * 2.0f);
 							if (!bCameraInsideLightGeometry)
@@ -1135,7 +1135,7 @@ void FDeferredShadingSceneRenderer::BeginOcclusionTests(FRHICommandListImmediate
 								ShadowOcclusionQueryMap.Add(Key, ShadowOcclusionQuery);
 
 								// Draw bounding sphere
-								VertexShader->SetParametersWithBoundingSphere(RHICmdList, View, ProjectedShadowInfo.LightSceneInfo->Proxy->GetBoundingSphere());
+								VertexShader->SetParametersWithBoundingSphere(RHICmdList, View, LightBounds);
 								StencilingGeometry::DrawSphere(RHICmdList);
 
 								RHICmdList.EndRenderQuery(ShadowOcclusionQuery);
diff --git a/Engine/Source/Runtime/Renderer/Private/ScenePrivate.h b/Engine/Source/Runtime/Renderer/Private/ScenePrivate.h
--- a/Engine/Source/Runtime/Renderer/Private/ScenePrivate.h
+++ b/Engine/Source/Runtime/Renderer/Private/ScenePrivate.h
@@ -377,8 +377,8 @@ public:
 		}
 
 		FProjectedShadowKey(const FProjectedShadowInfo& ProjectedShadowInfo)
-			: PrimitiveId(ProjectedShadowInfo.ParentSceneInfo ? ProjectedShadowInfo.ParentSceneInfo->PrimitiveComponentId : FPrimitiveComponentId())
-			, Light(ProjectedShadowInfo.LightSceneInfo->Proxy->GetLightComponent())
+			: PrimitiveId(ProjectedShadowInfo.GetParentSceneInfo() ? ProjectedShadowInfo.GetParentSceneInfo()->PrimitiveComponentId : FPrimitiveComponentId())
+			, Light(ProjectedShadowInfo.GetLightSceneInfo().Proxy->GetLightComponent())
 			, ShadowSplitIndex(ProjectedShadowInfo.CascadeSettings.ShadowSplitIndex)
 			, bTranslucentShadow(ProjectedShadowInfo.bTranslucentShadow)
 		{
diff --git a/Engine/Source/Runtime/Renderer/Private/ShadowRendering.h b/Engine/Source/Runtime/Renderer/Private/ShadowRendering.h
--- a/Engine/Source/Runtime/Renderer/Private/ShadowRendering.h
+++ b/Engine/Source/Runtime/Renderer/Private/ShadowRendering.h
@@ -493,12 +493,6 @@ class FProjectedShadowInfo : public FRefCountedObject
 public:
 	typedef TArray<const FPrimitiveSceneInfo*,SceneRenderingAllocator> PrimitiveArrayType;
 
-	const FLightSceneInfo* const LightSceneInfo;
-	const FLightSceneInfoCompact LightSceneInfoCompact;
-
-	/** Parent primitive of the shadow group that created this shadow, if not a bWholeSceneShadow. */
-	const FPrimitiveSceneInfo* const ParentSceneInfo;
-
 	/** The main view this shadow must be rendered in, or NULL for a view independent shadow. */
 	FViewInfo* DependentView;
 
@@ -572,7 +566,7 @@ public:
 	uint32 bDepthsCached : 1;
 
 	// redundant to LightSceneInfo->Proxy->GetLightType() == LightType_Directional, could be made ELightComponentType LightType
-	const uint32 bDirectionalLight : 1;
+	uint32 bDirectionalLight : 1;
 
 	/** Whether this shadow affects the whole scene or only a group of objects. */
 	uint32 bWholeSceneShadow : 1;
@@ -602,11 +596,14 @@ public:
 
 public:
 
+	// default constructor
+	FProjectedShadowInfo();
+
 	/**
-	 * Initialization constructor for a per-object shadow. e.g. translucent particle system
+	 * for a per-object shadow. e.g. translucent particle system or a dynamic object in a precomputed shadow situation
 	 * @param InParentSceneInfo must not be 0
 	 */
-	FProjectedShadowInfo(
+	void SetupPerObjectProjection(
 		FLightSceneInfo* InLightSceneInfo,
 		const FPrimitiveSceneInfo* InParentSceneInfo,
 		const FPerObjectProjectedShadowInitializer& Initializer,
@@ -614,12 +611,11 @@ public:
 		uint32 InResolutionX,
 		uint32 MaxShadowResolutionY,
 		float InMaxScreenPercent,
-		const TArray<float, TInlineAllocator<2> >& InFadeAlphas,
 		bool bInTranslucentShadow
 		);
 
-	/** Initialization constructor for a whole-scene shadow. */
-	FProjectedShadowInfo(
+	/** for a whole-scene shadow. */
+	void SetupWholeSceneProjection(
 		FLightSceneInfo* InLightSceneInfo,
 		FViewInfo* InDependentView,
 		const FWholeSceneProjectedShadowInitializer& Initializer,
@@ -719,7 +715,26 @@ public:
 	/** Sorts SubjectMeshElements based on state so that rendering the static elements will set as little state as possible. */
 	void SortSubjectMeshElements();
 
+	// 0 if Setup...() wasn't called yet
+	const FLightSceneInfo & GetLightSceneInfo() const { return *LightSceneInfo; }
+	const FLightSceneInfoCompact& GetLightSceneInfoCompact() const { return LightSceneInfoCompact; }
+	/**
+	 * Parent primitive of the shadow group that created this shadow, if not a bWholeSceneShadow.
+	 * 0 if Setup...() wasn't called yet
+	 */	
+	const FPrimitiveSceneInfo* GetParentSceneInfo() const { return ParentSceneInfo; }
+
 private:
+	// 0 if Setup...() wasn't called yet
+	const FLightSceneInfo* LightSceneInfo;
+	FLightSceneInfoCompact LightSceneInfoCompact;
+
+	/**
+	 * Parent primitive of the shadow group that created this shadow, if not a bWholeSceneShadow.
+	 * 0 if Setup...() wasn't called yet or for whole scene shadows
+	 */	
+	const FPrimitiveSceneInfo* ParentSceneInfo;
+
 	/** dynamic shadow casting elements */
 	PrimitiveArrayType SubjectPrimitives;
 	/** For preshadows, this contains the receiver primitives to mask the projection to. */
@@ -1155,12 +1170,14 @@ public:
 		const FProjectedShadowInfo* ShadowInfo
 		) override
 	{
+		const FPixelShaderRHIParamRef ShaderRHI = GetPixelShader();
+
 		FShadowProjectionPixelShaderInterface::SetParameters(RHICmdList, ViewIndex,View,ShadowInfo);
 
 		ProjectionParameters.Set(RHICmdList, this, View, ShadowInfo);
 
-		SetShaderValue(RHICmdList, GetPixelShader(), ShadowFadeFraction, ShadowInfo->FadeAlphas[ViewIndex] );
-		SetShaderValue(RHICmdList, GetPixelShader(), ShadowSharpen, ShadowInfo->LightSceneInfo->Proxy->GetShadowSharpen() * 7.0f + 1.0f );
+		SetShaderValue(RHICmdList, ShaderRHI, ShadowFadeFraction, ShadowInfo->FadeAlphas[ViewIndex] );
+		SetShaderValue(RHICmdList, ShaderRHI, ShadowSharpen, ShadowInfo->GetLightSceneInfo().Proxy->GetShadowSharpen() * 7.0f + 1.0f );
 	}
 
 	/**
@@ -1397,10 +1414,12 @@ public:
 		DeferredParameters.Set(RHICmdList, ShaderRHI, View);
 		OnePassShadowParameters.Set(RHICmdList, ShaderRHI, ShadowInfo);
 
-		SetShaderValue(RHICmdList, ShaderRHI, LightPosition, FVector4(ShadowInfo->LightSceneInfo->Proxy->GetPosition(), 1.0f / ShadowInfo->LightSceneInfo->Proxy->GetRadius()));
+		const FLightSceneProxy& LightProxy = *(ShadowInfo->GetLightSceneInfo().Proxy);
+
+		SetShaderValue(RHICmdList, ShaderRHI, LightPosition, FVector4(LightProxy.GetPosition(), 1.0f / LightProxy.GetRadius()));
 
 		SetShaderValue(RHICmdList, ShaderRHI, ShadowFadeFraction, ShadowInfo->FadeAlphas[ViewIndex]);
-		SetShaderValue(RHICmdList, ShaderRHI, ShadowSharpen, ShadowInfo->LightSceneInfo->Proxy->GetShadowSharpen() * 7.0f + 1.0f);
+		SetShaderValue(RHICmdList, ShaderRHI, ShadowSharpen, LightProxy.GetShadowSharpen() * 7.0f + 1.0f);
 		SetShaderValue(RHICmdList, ShaderRHI, PointLightDepthBiasParameters, FVector2D(ShadowInfo->GetShaderDepthBias(), 0.0f));
 	}
 
diff --git a/Engine/Source/Runtime/Renderer/Private/ShadowSetup.cpp b/Engine/Source/Runtime/Renderer/Private/ShadowSetup.cpp
--- a/Engine/Source/Runtime/Renderer/Private/ShadowSetup.cpp
+++ b/Engine/Source/Runtime/Renderer/Private/ShadowSetup.cpp
@@ -353,42 +353,58 @@ static bool GetBestShadowTransform(const FVector& ZAxis,const FBoundingBoxVertex
 	}
 }
 
-/** Initialization constructor for a per-object shadow. e.g. translucent particle system */
-FProjectedShadowInfo::FProjectedShadowInfo(
-	FLightSceneInfo* InLightSceneInfo,
-	const FPrimitiveSceneInfo* InParentSceneInfo,
-	const FPerObjectProjectedShadowInitializer& Initializer,
-	bool bInPreShadow,
-	uint32 InResolutionX,
-	uint32 MaxShadowResolutionY,
-	float InMaxScreenPercent,
-	const TArray<float, TInlineAllocator<2> >& InFadeAlphas,
-	bool bInTranslucentShadow)
-	: LightSceneInfo(InLightSceneInfo)
-	, LightSceneInfoCompact(InLightSceneInfo)
-	, ParentSceneInfo(InParentSceneInfo)
-	, DependentView(NULL)
+FProjectedShadowInfo::FProjectedShadowInfo()
+	: LightSceneInfo(0)
+	, ParentSceneInfo(0)
+	, DependentView(0)
 	, ShadowId(INDEX_NONE)
-	, PreShadowTranslation(Initializer.PreShadowTranslation)	
-	, ShadowBounds(Initializer.SubjectBounds.Origin - Initializer.PreShadowTranslation, Initializer.SubjectBounds.SphereRadius)
+	, PreShadowTranslation(0, 0, 0)
+	, ShadowBounds(0)
 	, X(0)
 	, Y(0)
-	, ResolutionX(InResolutionX)
+	, ResolutionX(0)
 	, ResolutionY(0)
-	, MaxScreenPercent(InMaxScreenPercent)
-	, FadeAlphas(InFadeAlphas)
+	, MaxScreenPercent(1.0f)
 	, bAllocated(false)
 	, bAllocatedInTranslucentLayout(false)
 	, bRendered(false)
 	, bAllocatedInPreshadowCache(false)
 	, bDepthsCached(false)
-	, bDirectionalLight(InLightSceneInfo->Proxy->GetLightType() == LightType_Directional)
+	, bDirectionalLight(false)
 	, bWholeSceneShadow(false)
 	, bReflectiveShadowmap(false)
-	, bTranslucentShadow(bInTranslucentShadow)
-	, bPreShadow(bInPreShadow)
-	, bSelfShadowOnly(InParentSceneInfo->Proxy->CastsSelfShadowOnly())
+	, bTranslucentShadow(false)
+	, bPreShadow(false)
+	, bSelfShadowOnly(false)
+	, bValidTransform(false)
+	, ShaderDepthBias(0.0f)
+{
+}
+
+void FProjectedShadowInfo::SetupPerObjectProjection(
+	FLightSceneInfo* InLightSceneInfo,
+	const FPrimitiveSceneInfo* InParentSceneInfo,
+	const FPerObjectProjectedShadowInitializer& Initializer,
+	bool bInPreShadow,
+	uint32 InResolutionX,
+	uint32 MaxShadowResolutionY,
+	float InMaxScreenPercent,
+	bool bInTranslucentShadow)
 {
+	check(InParentSceneInfo);
+
+	LightSceneInfo = InLightSceneInfo;
+	LightSceneInfoCompact = InLightSceneInfo;
+	ParentSceneInfo = InParentSceneInfo;
+	PreShadowTranslation = Initializer.PreShadowTranslation;
+	ShadowBounds = FSphere(Initializer.SubjectBounds.Origin - Initializer.PreShadowTranslation, Initializer.SubjectBounds.SphereRadius);
+	ResolutionX = InResolutionX;
+	MaxScreenPercent = InMaxScreenPercent;
+	bDirectionalLight = InLightSceneInfo->Proxy->GetLightType() == LightType_Directional;
+	bTranslucentShadow = bInTranslucentShadow;
+	bPreShadow = bInPreShadow;
+	bSelfShadowOnly = InParentSceneInfo->Proxy->CastsSelfShadowOnly();
+		
 	check(!CascadeSettings.bRayTracedDistanceField);
 
 	const FMatrix WorldToLightScaled = Initializer.WorldToLight * FScaleMatrix(Initializer.Scales);
@@ -480,39 +496,26 @@ FProjectedShadowInfo::FProjectedShadowInfo(
 	UpdateShaderDepthBias();
 }
 
-/** Initialization constructor for a whole-scene shadow. e.g. reflective shadow map, directional light cascade or point light */
-FProjectedShadowInfo::FProjectedShadowInfo(
+void FProjectedShadowInfo::SetupWholeSceneProjection(
 	FLightSceneInfo* InLightSceneInfo,
 	FViewInfo* InDependentView,
 	const FWholeSceneProjectedShadowInitializer& Initializer,
 	uint32 InResolutionX,
 	uint32 InResolutionY,
 	bool bInReflectiveShadowMap)
-	: LightSceneInfo(InLightSceneInfo)
-	, LightSceneInfoCompact(InLightSceneInfo)
-	, ParentSceneInfo(NULL)
-	, DependentView(InDependentView)
-	, ShadowId(INDEX_NONE)
-	, PreShadowTranslation(Initializer.PreShadowTranslation)
-	, CascadeSettings(Initializer.CascadeSettings)
-	, X(0)
-	, Y(0)
-	, ResolutionX(InResolutionX)
-	, ResolutionY(InResolutionY)
-	, MaxScreenPercent(1.0f)
-	, bAllocated(false)
-	, bAllocatedInTranslucentLayout(false)
-	, bRendered(false)
-	, bAllocatedInPreshadowCache(false)
-	, bDepthsCached(false)
-	, bDirectionalLight(InLightSceneInfo->Proxy->GetLightType() == LightType_Directional)
-	, bWholeSceneShadow(true)
-	, bReflectiveShadowmap(bInReflectiveShadowMap) 
-	, bTranslucentShadow(false)
-	, bPreShadow(false)
-	, bSelfShadowOnly(false)
-	, bValidTransform(true)
 {	
+	LightSceneInfo = InLightSceneInfo;
+	LightSceneInfoCompact = InLightSceneInfo;
+	DependentView = InDependentView;
+	PreShadowTranslation = Initializer.PreShadowTranslation;
+	CascadeSettings = Initializer.CascadeSettings;
+	ResolutionX = InResolutionX;
+	ResolutionY = InResolutionY;
+	bDirectionalLight = InLightSceneInfo->Proxy->GetLightType() == LightType_Directional;
+	bWholeSceneShadow = true;
+	bReflectiveShadowmap = bInReflectiveShadowMap; 
+	bValidTransform = true;
+
 	FVector	XAxis, YAxis;
 	Initializer.FaceDirection.FindBestAxisVectors(XAxis,YAxis);
 	const FMatrix WorldToLightScaled = Initializer.WorldToLight * FScaleMatrix(Initializer.Scales);
@@ -1031,8 +1034,8 @@ TRefCountPtr<FProjectedShadowInfo> FDeferredShadingSceneRenderer::GetCachedPresh
 		{
 			TRefCountPtr<FProjectedShadowInfo> CachedShadow = Scene->CachedPreshadows[ShadowIndex];
 			// Only reuse a cached preshadow if it was created for the same primitive and light
-			if (CachedShadow->ParentSceneInfo == PrimitiveInfo
-				&& CachedShadow->LightSceneInfo == LightInfo
+			if (CachedShadow->GetParentSceneInfo() == PrimitiveInfo
+				&& &CachedShadow->GetLightSceneInfo() == LightInfo
 				// Only reuse if it contains the bounds being queried, with some tolerance
 				&& QueryBounds.IsInside(CachedShadow->ShadowBounds, CachedShadow->ShadowBounds.W * .04f)
 				// Only reuse if the resolution matches
@@ -1387,7 +1390,9 @@ void FDeferredShadingSceneRenderer::CreatePerObjectProjectedShadow(
 			if (bOpaqueRelevance && bCreateOpaqueObjectShadow && (bOpaqueShadowIsVisibleThisFrame || bShadowIsPotentiallyVisibleNextFrame))
 			{
 				// Create a projected shadow for this interaction's shadow.
-				FProjectedShadowInfo* ProjectedShadowInfo = new(FMemStack::Get(),1,16) FProjectedShadowInfo(
+				FProjectedShadowInfo* ProjectedShadowInfo = new(FMemStack::Get(),1,16) FProjectedShadowInfo;
+
+				ProjectedShadowInfo->SetupPerObjectProjection(
 					LightSceneInfo,
 					PrimitiveSceneInfo,
 					ShadowInitializer,
@@ -1395,9 +1400,9 @@ void FDeferredShadingSceneRenderer::CreatePerObjectProjectedShadow(
 					SizeX,
 					MaxShadowResolutionY,
 					MaxScreenPercent,
-					ResolutionFadeAlphas,
 					false					// not translucent shadow
 					);
+				ProjectedShadowInfo->FadeAlphas = ResolutionFadeAlphas,
 				VisibleLightInfo.MemStackProjectedShadows.Add(ProjectedShadowInfo);
 
 				if (ProjectedShadowInfo->bValidTransform)
@@ -1425,7 +1430,9 @@ void FDeferredShadingSceneRenderer::CreatePerObjectProjectedShadow(
 				&& (bTranslucentShadowIsVisibleThisFrame || bShadowIsPotentiallyVisibleNextFrame))
 			{
 				// Create a projected shadow for this interaction's shadow.
-				FProjectedShadowInfo* ProjectedShadowInfo = new(FMemStack::Get(),1,16) FProjectedShadowInfo(
+				FProjectedShadowInfo* ProjectedShadowInfo = new(FMemStack::Get(),1,16) FProjectedShadowInfo;
+
+				ProjectedShadowInfo->SetupPerObjectProjection(
 					LightSceneInfo,
 					PrimitiveSceneInfo,
 					ShadowInitializer,
@@ -1434,9 +1441,9 @@ void FDeferredShadingSceneRenderer::CreatePerObjectProjectedShadow(
 					FMath::Clamp<int32>(SizeX / GSceneRenderTargets.GetTranslucentShadowDownsampleFactor(), 1, GSceneRenderTargets.GetTranslucentShadowDepthTextureResolution().X - SHADOW_BORDER * 2),
 					FMath::Clamp<int32>(MaxShadowResolutionY / GSceneRenderTargets.GetTranslucentShadowDownsampleFactor(), 1, GSceneRenderTargets.GetTranslucentShadowDepthTextureResolution().Y - SHADOW_BORDER * 2),
 					MaxScreenPercent,
-					ResolutionFadeAlphas,
 					true					// translucent shadow
 					);
+				ProjectedShadowInfo->FadeAlphas = ResolutionFadeAlphas,
 				VisibleLightInfo.MemStackProjectedShadows.Add(ProjectedShadowInfo);
 
 				if (ProjectedShadowInfo->bValidTransform)
@@ -1476,7 +1483,7 @@ void FDeferredShadingSceneRenderer::CreatePerObjectProjectedShadow(
 			for (int32 i = 0; i < ViewDependentWholeSceneShadows.Num(); i++)
 			{
 				const FProjectedShadowInfo* WholeSceneShadow = ViewDependentWholeSceneShadows[i];
-				const FVector2D DistanceFadeValues = WholeSceneShadow->LightSceneInfo->Proxy->GetDirectionalLightDistanceFadeParameters(Scene->GetFeatureLevel());
+				const FVector2D DistanceFadeValues = WholeSceneShadow->GetLightSceneInfo().Proxy->GetDirectionalLightDistanceFadeParameters(Scene->GetFeatureLevel());
 				const float DistanceFromShadowCenterSquared = (WholeSceneShadow->ShadowBounds.Center - Bounds.Origin).SizeSquared();
 				//@todo - if view dependent whole scene shadows are ever supported in splitscreen, 
 				// We can only disable the preshadow at this point if it is inside a whole scene shadow for all views
@@ -1497,16 +1504,13 @@ void FDeferredShadingSceneRenderer::CreatePerObjectProjectedShadow(
 				// Try to reuse a preshadow from the cache
 				TRefCountPtr<FProjectedShadowInfo> ProjectedPreShadowInfo = GetCachedPreshadow(Interaction, ShadowInitializer, OriginalBounds, PreshadowSizeX);
 
-				if (ProjectedPreShadowInfo)
-				{
-					// Update fade alpha on the cached preshadow
-					ProjectedPreShadowInfo->FadeAlphas = ResolutionPreShadowFadeAlphas;
-				}
-				else
+				if(!ProjectedPreShadowInfo)
 				{
 					// Create a new projected shadow for this interaction's preshadow
 					// Not using the scene rendering mem stack because this shadow info may need to persist for multiple frames if it gets cached
-					ProjectedPreShadowInfo = new FProjectedShadowInfo(
+					ProjectedPreShadowInfo = new FProjectedShadowInfo;
+
+					ProjectedPreShadowInfo->SetupPerObjectProjection(
 						LightSceneInfo,
 						PrimitiveSceneInfo,
 						ShadowInitializer,
@@ -1514,10 +1518,11 @@ void FDeferredShadingSceneRenderer::CreatePerObjectProjectedShadow(
 						PreshadowSizeX,
 						FMath::TruncToInt(MaxShadowResolutionY * CVarPreShadowResolutionFactor.GetValueOnRenderThread()),
 						MaxScreenPercent,
-						ResolutionPreShadowFadeAlphas,
 						false				// not translucent shadow
 						);
 				}
+				// Update fade alpha on the cached preshadow
+				ProjectedPreShadowInfo->FadeAlphas = ResolutionPreShadowFadeAlphas;
 
 				VisibleLightInfo.AllProjectedShadows.Add(ProjectedPreShadowInfo);
 				VisibleLightInfo.ProjectedPreShadows.Add(ProjectedPreShadowInfo);
@@ -1625,7 +1630,9 @@ void FDeferredShadingSceneRenderer::CreateWholeSceneProjectedShadow(FLightSceneI
 				}
 
 				// Create the projected shadow info.
-				FProjectedShadowInfo* ProjectedShadowInfo = new(FMemStack::Get(),1,16) FProjectedShadowInfo(
+				FProjectedShadowInfo* ProjectedShadowInfo = new(FMemStack::Get(),1,16) FProjectedShadowInfo;
+
+				ProjectedShadowInfo->SetupWholeSceneProjection(
 					LightSceneInfo,
 					NULL,
 					ProjectedShadowInitializer,
@@ -1633,6 +1640,7 @@ void FDeferredShadingSceneRenderer::CreateWholeSceneProjectedShadow(FLightSceneI
 					SizeY,
 					false	// no RSM
 					);
+
 				ProjectedShadowInfo->FadeAlphas = FadeAlphas;
 
 				VisibleLightInfo.MemStackProjectedShadows.Add(ProjectedShadowInfo);
@@ -1660,8 +1668,10 @@ void FDeferredShadingSceneRenderer::CreateWholeSceneProjectedShadow(FLightSceneI
 						FVector(0, 1, 0)
 					};
 
-					const FMatrix FaceProjection = FPerspectiveMatrix(PI / 4.0f, 1, 1, 1, ProjectedShadowInfo->LightSceneInfo->Proxy->GetRadius());
-					const FVector LightPosition = ProjectedShadowInfo->LightSceneInfo->Proxy->GetPosition();
+					const FLightSceneProxy& LightProxy = *(ProjectedShadowInfo->GetLightSceneInfo().Proxy);
+
+					const FMatrix FaceProjection = FPerspectiveMatrix(PI / 4.0f, 1, 1, 1, LightProxy.GetRadius());
+					const FVector LightPosition = LightProxy.GetPosition();
 					
 					ProjectedShadowInfo->OnePassShadowViewProjectionMatrices.Empty(6);
 					ProjectedShadowInfo->OnePassShadowFrustums.Empty(6);
@@ -1749,9 +1759,9 @@ void FSceneRenderer::InitProjectedShadowVisibility(FRHICommandListImmediate& RHI
 					// Compute the subject primitive's view relevance.  Note that the view won't necessarily have it cached,
 					// since the primitive might not be visible.
 					FPrimitiveViewRelevance ViewRelevance;
-					if(ProjectedShadowInfo.ParentSceneInfo)
+					if(ProjectedShadowInfo.GetParentSceneInfo())
 					{
-						ViewRelevance = ProjectedShadowInfo.ParentSceneInfo->Proxy->GetViewRelevance(&View);
+						ViewRelevance = ProjectedShadowInfo.GetParentSceneInfo()->Proxy->GetViewRelevance(&View);
 					}
 					else
 					{
@@ -1768,10 +1778,10 @@ void FSceneRenderer::InitProjectedShadowVisibility(FRHICommandListImmediate& RHI
 						View.State &&
 						((FSceneViewState*)View.State)->IsShadowOccluded(
 							RHICmdList,
-							ProjectedShadowInfo.ParentSceneInfo ? 
-								ProjectedShadowInfo.ParentSceneInfo->PrimitiveComponentId :
+							ProjectedShadowInfo.GetParentSceneInfo() ? 
+								ProjectedShadowInfo.GetParentSceneInfo()->PrimitiveComponentId :
 								FPrimitiveComponentId(),
-							ProjectedShadowInfo.LightSceneInfo->Proxy->GetLightComponent(),
+							ProjectedShadowInfo.GetLightSceneInfo().Proxy->GetLightComponent(),
 							ProjectedShadowInfo.CascadeSettings.ShadowSplitIndex,
 							ProjectedShadowInfo.bTranslucentShadow
 							);
@@ -1930,7 +1940,7 @@ inline void FSceneRenderer::GatherShadowsForPrimitiveInner(
 				// Check if this primitive is in the shadow's frustum.
 				bool bInFrustum = ProjectedShadowInfo->CasterFrustum.IntersectBox( PrimitiveBounds.Origin, ProjectedShadowInfo->PreShadowTranslation, PrimitiveBounds.BoxExtent );
 
-				if( bInFrustum && ProjectedShadowInfo->LightSceneInfoCompact.AffectsPrimitive(PrimitiveSceneInfoCompact) )
+				if( bInFrustum && ProjectedShadowInfo->GetLightSceneInfoCompact().AffectsPrimitive(PrimitiveSceneInfoCompact) )
 				{
 					// Add this primitive to the shadow.
 					ProjectedShadowInfo->AddSubjectPrimitive(PrimitiveSceneInfo, &Views);
@@ -1954,7 +1964,7 @@ inline void FSceneRenderer::GatherShadowsForPrimitiveInner(
 				}
 
 				FPrimitiveSceneProxy* PrimitiveProxy = PrimitiveSceneInfoCompact.Proxy;
-				FLightSceneProxy* RESTRICT LightProxy = ProjectedShadowInfo->LightSceneInfo->Proxy;
+				FLightSceneProxy* RESTRICT LightProxy = ProjectedShadowInfo->GetLightSceneInfo().Proxy;
 
 				const FVector LightDirection = LightProxy->GetDirection();
 				const FVector PrimitiveToShadowCenter = ProjectedShadowInfo->ShadowBounds.Center - PrimitiveBounds.Origin;
@@ -1965,7 +1975,7 @@ inline void FSceneRenderer::GatherShadowsForPrimitiveInner(
 
 				// Include all primitives for movable lights, but only statically shadowed primitives from a light with static shadowing,
 				// Since lights with static shadowing still create per-object shadows for primitives without static shadowing.
-				if( (!LightProxy->HasStaticLighting() || !ProjectedShadowInfo->LightSceneInfo->bPrecomputedLightingIsValid)
+				if( (!LightProxy->HasStaticLighting() || !ProjectedShadowInfo->GetLightSceneInfo().bPrecomputedLightingIsValid)
 					// Check if this primitive is in the shadow's cylinder
 					&& PrimitiveDistanceFromCylinderAxisSq < FMath::Square(ProjectedShadowInfo->ShadowBounds.W + PrimitiveBounds.SphereRadius)
 					// Check if the primitive is closer than the cylinder cap toward the light
@@ -1995,11 +2005,11 @@ inline void FSceneRenderer::GatherShadowsForPrimitiveInner(
 							bScreenSpaceSizeCulled = FMath::Square( PrimitiveBounds.SphereRadius ) < FMath::Square( MinScreenRadiusForShadowCaster ) * DistanceSquared;
 						}
 
-						if (ProjectedShadowInfo->LightSceneInfoCompact.AffectsPrimitive(PrimitiveSceneInfoCompact)
+						if (ProjectedShadowInfo->GetLightSceneInfoCompact().AffectsPrimitive(PrimitiveSceneInfoCompact)
 							// Exclude primitives that will create their own per-object shadow, except when rendering RSMs
 							&& ( !PrimitiveProxy->CastsInsetShadow() || ProjectedShadowInfo->bReflectiveShadowmap )
 							// Exclude primitives that will create a per-object shadow from a stationary light
-							&& !ShouldCreateObjectShadowForStationaryLight(ProjectedShadowInfo->LightSceneInfo, PrimitiveSceneInfo->Proxy, true)
+							&& !ShouldCreateObjectShadowForStationaryLight(&ProjectedShadowInfo->GetLightSceneInfo(), PrimitiveSceneInfo->Proxy, true)
 							// Only render shadows from objects that use static lighting during a reflection capture, since the reflection capture doesn't update at runtime
 							&& (!bStaticSceneOnly || PrimitiveProxy->HasStaticLighting()) 
 							&& !bScreenSpaceSizeCulled )
@@ -2167,7 +2177,9 @@ void FSceneRenderer::AddViewDependentWholeSceneShadowsForView(
 				{
 					const FIntPoint ShadowBufferResolution = GSceneRenderTargets.GetShadowDepthTextureResolution();
 					// Create the projected shadow info.
-					FProjectedShadowInfo* ProjectedShadowInfo = new(FMemStack::Get(), 1, 16) FProjectedShadowInfo(
+					FProjectedShadowInfo* ProjectedShadowInfo = new(FMemStack::Get(), 1, 16) FProjectedShadowInfo;
+
+					ProjectedShadowInfo->SetupWholeSceneProjection(
 						&LightSceneInfo,
 						&View,
 						ProjectedShadowInitializer,
@@ -2211,13 +2223,15 @@ void FSceneRenderer::AddViewDependentWholeSceneShadowsForView(
 						const FIntPoint ShadowBufferResolution = GSceneRenderTargets.GetReflectiveShadowMapTextureResolution();
 
 						// Create the projected shadow info.
-						FProjectedShadowInfo* ProjectedShadowInfo = new(FMemStack::Get(), 1, 16) FProjectedShadowInfo(
+						FProjectedShadowInfo* ProjectedShadowInfo = new(FMemStack::Get(), 1, 16) FProjectedShadowInfo;
+
+						ProjectedShadowInfo->SetupWholeSceneProjection(
 							&LightSceneInfo,
 							&View,
 							ProjectedShadowInitializer,
 							ShadowBufferResolution.X,
 							ShadowBufferResolution.Y,
-							true);
+							true);		// RSM
 
 						FVisibleLightInfo& LightViewInfo = VisibleLightInfos[LightSceneInfo.Id];
 						VisibleLightInfo.MemStackProjectedShadows.Add(ProjectedShadowInfo);
diff --git a/Engine/Source/Runtime/Renderer/Private/TranslucentLighting.cpp b/Engine/Source/Runtime/Renderer/Private/TranslucentLighting.cpp
--- a/Engine/Source/Runtime/Renderer/Private/TranslucentLighting.cpp
+++ b/Engine/Source/Runtime/Renderer/Private/TranslucentLighting.cpp
@@ -262,7 +262,7 @@ public:
 
 		SetShaderValue(RHICmdList, ShaderRHI, TranslInvMaxSubjectDepth, ShadowInfo->InvMaxSubjectDepth);
 
-		const float LocalToWorldScale = ShadowInfo->ParentSceneInfo->Proxy->GetLocalToWorld().GetScaleVector().GetMax();
+		const float LocalToWorldScale = ShadowInfo->GetParentSceneInfo()->Proxy->GetLocalToWorld().GetScaleVector().GetMax();
 		const float TranslucentShadowStartOffsetValue = MaterialRenderProxy->GetMaterial(FeatureLevel)->GetTranslucentShadowStartOffset() * LocalToWorldScale;
 		SetShaderValue(RHICmdList, ShaderRHI,TranslucentShadowStartOffset, TranslucentShadowStartOffsetValue / (ShadowInfo->MaxSubjectZ - ShadowInfo->MinSubjectZ));
 		TranslucencyProjectionParameters.Set(RHICmdList, this);
@@ -1263,7 +1263,7 @@ FGlobalBoundShaderState ObjectShadowingBoundShaderState;
 
 void FDeferredShadingSceneRenderer::AccumulateTranslucentVolumeObjectShadowing(FRHICommandList& RHICmdList, const FProjectedShadowInfo* InProjectedShadowInfo, bool bClearVolume)
 {
-	const FLightSceneInfo* LightSceneInfo = InProjectedShadowInfo->LightSceneInfo;
+	const FLightSceneInfo* LightSceneInfo = &InProjectedShadowInfo->GetLightSceneInfo();
 
 	if (bClearVolume)
 	{
diff --git a/Engine/Source/Runtime/Renderer/Private/TranslucentRendering.cpp b/Engine/Source/Runtime/Renderer/Private/TranslucentRendering.cpp
--- a/Engine/Source/Runtime/Renderer/Private/TranslucentRendering.cpp
+++ b/Engine/Source/Runtime/Renderer/Private/TranslucentRendering.cpp
@@ -59,7 +59,7 @@ const FProjectedShadowInfo* FDeferredShadingSceneRenderer::PrepareTranslucentSha
 				{
 					FProjectedShadowInfo* CurrentShadowInfo = VisibleLightInfo->AllProjectedShadows[ShadowIndex];
 
-					if (CurrentShadowInfo && CurrentShadowInfo->bTranslucentShadow && CurrentShadowInfo->ParentSceneInfo == PrimitiveSceneInfo)
+					if (CurrentShadowInfo && CurrentShadowInfo->bTranslucentShadow && CurrentShadowInfo->GetParentSceneInfo() == PrimitiveSceneInfo)
 					{
 						TranslucentSelfShadow = CurrentShadowInfo;
 						break;
