cleanup/simplified shadow rendering code

diff --git a/Engine/Source/Runtime/Renderer/Private/ShadowRendering.cpp b/Engine/Source/Runtime/Renderer/Private/ShadowRendering.cpp
--- a/Engine/Source/Runtime/Renderer/Private/ShadowRendering.cpp
+++ b/Engine/Source/Runtime/Renderer/Private/ShadowRendering.cpp
@@ -2691,7 +2691,7 @@ void FDeferredShadingSceneRenderer::RenderProjections(
 		for (int32 ShadowIndex = 0; ShadowIndex < Shadows.Num(); ShadowIndex++)
 		{
 			FProjectedShadowInfo* ProjectedShadowInfo = Shadows[ShadowIndex];
-			if (ProjectedShadowInfo->bAllocated && ProjectedShadowInfo->bValidTransform)
+			if (ProjectedShadowInfo->bAllocated)
 			{
 				// Only project the shadow if it's large enough in this particular view (split screen, etc... may have shadows that are large in one view but irrelevantly small in others)
 				if (ProjectedShadowInfo->FadeAlphas[ViewIndex] > 1.0f / 256.0f)
diff --git a/Engine/Source/Runtime/Renderer/Private/ShadowRendering.h b/Engine/Source/Runtime/Renderer/Private/ShadowRendering.h
--- a/Engine/Source/Runtime/Renderer/Private/ShadowRendering.h
+++ b/Engine/Source/Runtime/Renderer/Private/ShadowRendering.h
@@ -583,8 +583,6 @@ public:
 	/** To not cast a shadow on the ground outside the object and having higher quality (useful for first person weapon). */
 	uint32 bSelfShadowOnly : 1;
 
-	uint32 bValidTransform : 1;
-
 	TBitArray<SceneRenderingBitArrayAllocator> StaticMeshWholeSceneShadowDepthMap;
 	TArray<uint64,SceneRenderingAllocator> StaticMeshWholeSceneShadowBatchVisibility;
 
@@ -602,8 +600,9 @@ public:
 	/**
 	 * for a per-object shadow. e.g. translucent particle system or a dynamic object in a precomputed shadow situation
 	 * @param InParentSceneInfo must not be 0
+	 * @return success, if false the shadow project is invalid and the projection should nto be created
 	 */
-	void SetupPerObjectProjection(
+	bool SetupPerObjectProjection(
 		FLightSceneInfo* InLightSceneInfo,
 		const FPrimitiveSceneInfo* InParentSceneInfo,
 		const FPerObjectProjectedShadowInitializer& Initializer,
diff --git a/Engine/Source/Runtime/Renderer/Private/ShadowSetup.cpp b/Engine/Source/Runtime/Renderer/Private/ShadowSetup.cpp
--- a/Engine/Source/Runtime/Renderer/Private/ShadowSetup.cpp
+++ b/Engine/Source/Runtime/Renderer/Private/ShadowSetup.cpp
@@ -376,12 +376,11 @@ FProjectedShadowInfo::FProjectedShadowInfo()
 	, bTranslucentShadow(false)
 	, bPreShadow(false)
 	, bSelfShadowOnly(false)
-	, bValidTransform(false)
 	, ShaderDepthBias(0.0f)
 {
 }
 
-void FProjectedShadowInfo::SetupPerObjectProjection(
+bool FProjectedShadowInfo::SetupPerObjectProjection(
 	FLightSceneInfo* InLightSceneInfo,
 	const FPrimitiveSceneInfo* InParentSceneInfo,
 	const FPerObjectProjectedShadowInitializer& Initializer,
@@ -433,10 +432,13 @@ void FProjectedShadowInfo::SetupPerObjectProjection(
 	// Compute the transform from light-space to shadow-space.
 	FMatrix LightToShadow;
 	float AspectRatio;
+	
+	// if this is a valid transform (can be false if the object is around the light)
+	bool bRet = false;
 
 	if (GetBestShadowTransform(Initializer.FaceDirection.GetSafeNormal(), ProjectedBoundsPoints, BoundsEdges, AspectRatio, LightToShadow))
 	{
-		bValidTransform = true;
+		bRet = true;
 		const FMatrix WorldToShadow = WorldToLightScaled * LightToShadow;
 
 		const FBox ShadowSubjectBounds = Initializer.SubjectBounds.GetBox().TransformBy(WorldToShadow);
@@ -487,13 +489,10 @@ void FProjectedShadowInfo::SetupPerObjectProjection(
 
 		InvReceiverMatrix = ReceiverMatrix.InverseFast();
 		GetViewFrustumBounds(ReceiverFrustum,ReceiverMatrix,true);
-	}
-	else
-	{
-		bValidTransform = false;
+		UpdateShaderDepthBias();
 	}
 
-	UpdateShaderDepthBias();
+	return bRet;
 }
 
 void FProjectedShadowInfo::SetupWholeSceneProjection(
@@ -514,7 +513,6 @@ void FProjectedShadowInfo::SetupWholeSceneProjection(
 	bDirectionalLight = InLightSceneInfo->Proxy->GetLightType() == LightType_Directional;
 	bWholeSceneShadow = true;
 	bReflectiveShadowmap = bInReflectiveShadowMap; 
-	bValidTransform = true;
 
 	FVector	XAxis, YAxis;
 	Initializer.FaceDirection.FindBestAxisVectors(XAxis,YAxis);
@@ -1392,26 +1390,24 @@ void FDeferredShadingSceneRenderer::CreatePerObjectProjectedShadow(
 				// Create a projected shadow for this interaction's shadow.
 				FProjectedShadowInfo* ProjectedShadowInfo = new(FMemStack::Get(),1,16) FProjectedShadowInfo;
 
-				ProjectedShadowInfo->SetupPerObjectProjection(
+				if(ProjectedShadowInfo->SetupPerObjectProjection(
 					LightSceneInfo,
 					PrimitiveSceneInfo,
 					ShadowInitializer,
-					false,					// not preshadow
+					false,					// no preshadow
 					SizeX,
 					MaxShadowResolutionY,
 					MaxScreenPercent,
-					false					// not translucent shadow
-					);
-				ProjectedShadowInfo->FadeAlphas = ResolutionFadeAlphas,
-				VisibleLightInfo.MemStackProjectedShadows.Add(ProjectedShadowInfo);
-
-				if (ProjectedShadowInfo->bValidTransform)
+					false))					// no translucent shadow
 				{
+					ProjectedShadowInfo->FadeAlphas = ResolutionFadeAlphas;
+					VisibleLightInfo.MemStackProjectedShadows.Add(ProjectedShadowInfo);
+
 					if (bOpaqueShadowIsVisibleThisFrame)
 					{
 						VisibleLightInfo.AllProjectedShadows.Add(ProjectedShadowInfo);
 
-						for (int32 ChildIndex = 0; ChildIndex < ShadowGroupPrimitives.Num(); ChildIndex++)
+						for (int32 ChildIndex = 0, ChildCount = ShadowGroupPrimitives.Num(); ChildIndex < ChildCount; ChildIndex++)
 						{
 							FPrimitiveSceneInfo* ShadowChild = ShadowGroupPrimitives[ChildIndex];
 							ProjectedShadowInfo->AddSubjectPrimitive(ShadowChild, &Views);
@@ -1432,27 +1428,25 @@ void FDeferredShadingSceneRenderer::CreatePerObjectProjectedShadow(
 				// Create a projected shadow for this interaction's shadow.
 				FProjectedShadowInfo* ProjectedShadowInfo = new(FMemStack::Get(),1,16) FProjectedShadowInfo;
 
-				ProjectedShadowInfo->SetupPerObjectProjection(
+				if(ProjectedShadowInfo->SetupPerObjectProjection(
 					LightSceneInfo,
 					PrimitiveSceneInfo,
 					ShadowInitializer,
-					false,					// not preshadow
+					false,					// no preshadow
 					// Size was computed for the full res opaque shadow, convert to downsampled translucent shadow size with proper clamping
 					FMath::Clamp<int32>(SizeX / GSceneRenderTargets.GetTranslucentShadowDownsampleFactor(), 1, GSceneRenderTargets.GetTranslucentShadowDepthTextureResolution().X - SHADOW_BORDER * 2),
 					FMath::Clamp<int32>(MaxShadowResolutionY / GSceneRenderTargets.GetTranslucentShadowDownsampleFactor(), 1, GSceneRenderTargets.GetTranslucentShadowDepthTextureResolution().Y - SHADOW_BORDER * 2),
 					MaxScreenPercent,
-					true					// translucent shadow
-					);
-				ProjectedShadowInfo->FadeAlphas = ResolutionFadeAlphas,
-				VisibleLightInfo.MemStackProjectedShadows.Add(ProjectedShadowInfo);
-
-				if (ProjectedShadowInfo->bValidTransform)
+					true))					// translucent shadow
 				{
+					ProjectedShadowInfo->FadeAlphas = ResolutionFadeAlphas,
+					VisibleLightInfo.MemStackProjectedShadows.Add(ProjectedShadowInfo);
+
 					if (bTranslucentShadowIsVisibleThisFrame)
 					{
 						VisibleLightInfo.AllProjectedShadows.Add(ProjectedShadowInfo);
 
-						for (int32 ChildIndex = 0; ChildIndex < ShadowGroupPrimitives.Num(); ChildIndex++)
+						for (int32 ChildIndex = 0, ChildCount = ShadowGroupPrimitives.Num(); ChildIndex < ChildCount; ChildIndex++)
 						{
 							FPrimitiveSceneInfo* ShadowChild = ShadowGroupPrimitives[ChildIndex];
 							ProjectedShadowInfo->AddSubjectPrimitive(ShadowChild, &Views);
