added visualize blom showflag and shader, will be extended further

diff --git a/Engine/Shaders/PostProcessBloom.usf b/Engine/Shaders/PostProcessBloom.usf
--- a/Engine/Shaders/PostProcessBloom.usf
+++ b/Engine/Shaders/PostProcessBloom.usf
@@ -19,6 +19,26 @@ void MainPostprocessCommonVS(
 	OutTexCoord.zw =  OutPosition.xy;
 }
 
+void VisualizeBloomSetupPS(
+	  float4 UVAndScreenPos : TEXCOORD0
+	, out float4 OutColor : SV_Target0)
+{
+	float2 UV = UVAndScreenPos.xy;
+
+	OutColor = 0;
+
+	int2 PixelPos = (int2)ComputePixelPosCenter(UVAndScreenPos.zw, true);
+
+	int2 ViewportCenter = (int2)(ViewportRect + ViewportSize / 2);
+
+	bool bCross = ViewportCenter.x == PixelPos.x || ViewportCenter.y == PixelPos.y;
+
+	if(bCross)
+	{
+		OutColor = 20;
+	}
+}
+
 // vertex shader entry point
 void MainVS(
 	in float4 InPosition : ATTRIBUTE0,
diff --git a/Engine/Source/Runtime/Engine/Public/ShowFlags.h b/Engine/Source/Runtime/Engine/Public/ShowFlags.h
--- a/Engine/Source/Runtime/Engine/Public/ShowFlags.h
+++ b/Engine/Source/Runtime/Engine/Public/ShowFlags.h
@@ -212,6 +212,7 @@ struct FEngineShowFlags
 		SetVisualizeSSR(false);
 		SetForceGBuffer(false);
 		SetVisualizeSSS(false);
+		SetVisualizeBloom(false);
 	}
 
 	static bool CanBeToggledInEditor(const TCHAR* Name)
diff --git a/Engine/Source/Runtime/Engine/Public/ShowFlagsValues.inl b/Engine/Source/Runtime/Engine/Public/ShowFlagsValues.inl
--- a/Engine/Source/Runtime/Engine/Public/ShowFlagsValues.inl
+++ b/Engine/Source/Runtime/Engine/Public/ShowFlagsValues.inl
@@ -282,6 +282,8 @@ SHOWFLAG_FIXED_IN_SHIPPING(VisualizeSSR, 0, SFG_Visualize, LOCTEXT("VisualizeSSR
 SHOWFLAG_ALWAYS_ACCESSIBLE(ForceGBuffer, SFG_Hidden, LOCTEXT("ForceGBuffer", "Force usage of GBuffer"))
 /** Visualize the senses configuration of AIs' PawnSensingComponent */
 SHOWFLAG_ALWAYS_ACCESSIBLE(VisualizeSenses, SFG_Advanced, LOCTEXT("VisualizeSenses", "Senses"))
+/** Visualize the bloom, for developer (by default off): */
+SHOWFLAG_FIXED_IN_SHIPPING(VisualizeBloom, 0, SFG_Visualize, LOCTEXT("VisualizeBloom", "Bloom"))
 
 #undef SHOWFLAG_ALWAYS_ACCESSIBLE
 #undef SHOWFLAG_FIXED_IN_SHIPPING
diff --git a/Engine/Source/Runtime/Renderer/Private/PostProcess/PostProcessBloomSetup.cpp b/Engine/Source/Runtime/Renderer/Private/PostProcess/PostProcessBloomSetup.cpp
--- a/Engine/Source/Runtime/Renderer/Private/PostProcess/PostProcessBloomSetup.cpp
+++ b/Engine/Source/Runtime/Renderer/Private/PostProcess/PostProcessBloomSetup.cpp
@@ -179,8 +179,7 @@ void FRCPassPostProcessBloomSetup::Process(FRenderingCompositePassContext& Conte
 	TShaderMapRef<FPostProcessBloomSetupVS> VertexShader(Context.GetShaderMap());
 	TShaderMapRef<FPostProcessBloomSetupPS> PixelShader(Context.GetShaderMap());
 
-	static FGlobalBoundShaderState BoundShaderState;
-	
+	static FGlobalBoundShaderState BoundShaderState;	
 
 	SetGlobalBoundShaderState(Context.RHICmdList, Context.GetFeatureLevel(), BoundShaderState, GFilterVertexDeclaration.VertexDeclarationRHI, *VertexShader, *PixelShader);
 
@@ -214,3 +213,128 @@ FPooledRenderTargetDesc FRCPassPostProcessBloomSetup::ComputeOutputDesc(EPassOut
 
 
 
+
+
+
+
+/** Encapsulates the visualize bloom setup pixel shader. */
+class FPostProcessVisualizeBloomSetupPS : public FGlobalShader
+{
+	DECLARE_SHADER_TYPE(FPostProcessVisualizeBloomSetupPS, Global);
+
+	static bool ShouldCache(EShaderPlatform Platform)
+	{
+		return IsFeatureLevelSupported(Platform, ERHIFeatureLevel::SM4);
+	}
+
+	static void ModifyCompilationEnvironment(EShaderPlatform Platform, FShaderCompilerEnvironment& OutEnvironment)
+	{
+		FGlobalShader::ModifyCompilationEnvironment(Platform,OutEnvironment);
+	}
+
+	/** Default constructor. */
+	FPostProcessVisualizeBloomSetupPS() {}
+
+public:
+	FPostProcessPassParameters PostprocessParameter;
+
+	/** Initialization constructor. */
+	FPostProcessVisualizeBloomSetupPS(const ShaderMetaType::CompiledShaderInitializerType& Initializer)
+		: FGlobalShader(Initializer)
+	{
+		PostprocessParameter.Bind(Initializer.ParameterMap);
+	}
+
+	// FShader interface.
+	virtual bool Serialize(FArchive& Ar)
+	{
+		bool bShaderHasOutdatedParameters = FGlobalShader::Serialize(Ar);
+		Ar << PostprocessParameter;
+		return bShaderHasOutdatedParameters;
+	}
+
+	void SetParameters(const FRenderingCompositePassContext& Context)
+	{
+		const FPixelShaderRHIParamRef ShaderRHI = GetPixelShader();
+
+		FGlobalShader::SetParameters(Context.RHICmdList, ShaderRHI, Context.View);
+		PostprocessParameter.SetPS(ShaderRHI, Context, TStaticSamplerState<SF_Bilinear,AM_Border,AM_Border,AM_Border>::GetRHI());
+	}
+};
+
+IMPLEMENT_SHADER_TYPE(,FPostProcessVisualizeBloomSetupPS,TEXT("PostProcessBloom"),TEXT("VisualizeBloomSetupPS"),SF_Pixel);
+
+
+
+
+void FRCPassPostProcessVisualizeBloomSetup::Process(FRenderingCompositePassContext& Context)
+{
+	SCOPED_DRAW_EVENT(Context.RHICmdList, VisualizeBloomSetup);
+
+	const FPooledRenderTargetDesc* InputDesc = GetInputDesc(ePId_Input0);
+
+	check(InputDesc && "Input is not hooked up correctly");
+
+	const FSceneView& View = Context.View;
+	const FSceneViewFamily& ViewFamily = *(View.Family);
+
+	FIntPoint SrcSize = InputDesc->Extent;
+	FIntPoint DestSize = PassOutputs[0].RenderTargetDesc.Extent;
+
+	// e.g. 4 means the input texture is 4x smaller than the buffer size
+	uint32 ScaleFactor = GSceneRenderTargets.GetBufferSizeXY().X / SrcSize.X;
+
+	FIntRect SrcRect = View.ViewRect / ScaleFactor;
+	FIntRect DestRect = SrcRect;
+
+	const FSceneRenderTargetItem& DestRenderTarget = PassOutputs[0].RequestSurface(Context);
+
+	SetRenderTarget(Context.RHICmdList, DestRenderTarget.TargetableTexture, FTextureRHIRef());
+
+	// is optimized away if possible (RT size=view size, )
+	Context.RHICmdList.Clear(true, FLinearColor(0, 0, 0, 0), false, 1.0f, false, 0, DestRect);
+
+	Context.SetViewportAndCallRHI(0, 0, 0.0f, DestSize.X, DestSize.Y, 1.0f );
+
+	// set the state
+	Context.RHICmdList.SetBlendState(TStaticBlendState<>::GetRHI());
+	Context.RHICmdList.SetRasterizerState(TStaticRasterizerState<>::GetRHI());
+	Context.RHICmdList.SetDepthStencilState(TStaticDepthStencilState<false, CF_Always>::GetRHI());
+
+	TShaderMapRef<FPostProcessVS> VertexShader(Context.GetShaderMap());
+	TShaderMapRef<FPostProcessVisualizeBloomSetupPS> PixelShader(Context.GetShaderMap());
+
+	static FGlobalBoundShaderState BoundShaderState;
+
+	SetGlobalBoundShaderState(Context.RHICmdList, Context.GetFeatureLevel(), BoundShaderState, GFilterVertexDeclaration.VertexDeclarationRHI, *VertexShader, *PixelShader);
+
+	PixelShader->SetParameters(Context);
+	VertexShader->SetParameters(Context);
+
+	// Draw a quad mapping scene color to the view's render target
+	DrawRectangle(
+		Context.RHICmdList,
+		DestRect.Min.X, DestRect.Min.Y,
+		DestRect.Width(), DestRect.Height(),
+		SrcRect.Min.X, SrcRect.Min.Y,
+		SrcRect.Width(), SrcRect.Height(),
+		DestSize,
+		SrcSize,
+		*VertexShader,
+		EDRF_UseTriangleOptimization);
+
+	Context.RHICmdList.CopyToResolveTarget(DestRenderTarget.TargetableTexture, DestRenderTarget.ShaderResourceTexture, false, FResolveParams());
+}
+
+
+FPooledRenderTargetDesc FRCPassPostProcessVisualizeBloomSetup::ComputeOutputDesc(EPassOutputId InPassOutputId) const
+{
+	FPooledRenderTargetDesc Ret = PassInputs[0].GetOutput()->RenderTargetDesc;
+
+	Ret.Reset();
+	Ret.TargetableFlags &= ~TexCreate_UAV;
+	Ret.TargetableFlags |= TexCreate_RenderTargetable;
+	Ret.DebugName = TEXT("VisualizeBloomSetup");
+
+	return Ret;
+}
diff --git a/Engine/Source/Runtime/Renderer/Private/PostProcess/PostProcessBloomSetup.h b/Engine/Source/Runtime/Renderer/Private/PostProcess/PostProcessBloomSetup.h
--- a/Engine/Source/Runtime/Renderer/Private/PostProcess/PostProcessBloomSetup.h
+++ b/Engine/Source/Runtime/Renderer/Private/PostProcess/PostProcessBloomSetup.h
@@ -19,3 +19,14 @@ public:
 	virtual FPooledRenderTargetDesc ComputeOutputDesc(EPassOutputId InPassOutputId) const;
 	virtual void Release() override { delete this; }
 };
+
+// ePId_Input0: HDR SceneColor
+// derives from TRenderingCompositePassBase<InputCount, OutputCount> 
+class FRCPassPostProcessVisualizeBloomSetup : public TRenderingCompositePassBase<1, 1>
+{
+public:
+	// interface FRenderingCompositePass ---------
+	virtual void Process(FRenderingCompositePassContext& Context);
+	virtual FPooledRenderTargetDesc ComputeOutputDesc(EPassOutputId InPassOutputId) const;
+	virtual void Release() override { delete this; }
+};
diff --git a/Engine/Source/Runtime/Renderer/Private/PostProcess/PostProcessing.cpp b/Engine/Source/Runtime/Renderer/Private/PostProcess/PostProcessing.cpp
--- a/Engine/Source/Runtime/Renderer/Private/PostProcess/PostProcessing.cpp
+++ b/Engine/Source/Runtime/Renderer/Private/PostProcess/PostProcessing.cpp
@@ -253,6 +253,15 @@ static FRenderingCompositeOutputRef AddPostProcessEyeAdaptation(FPostprocessCont
 	return FRenderingCompositeOutputRef(Node);
 }
 
+static void AddVisualizeBloomSetup(FPostprocessContext& Context)
+{
+	FRenderingCompositePass* Node = Context.Graph.RegisterPass(new(FMemStack::Get()) FRCPassPostProcessVisualizeBloomSetup());
+
+	Node->SetInput(ePId_Input0, Context.FinalOutput);
+
+	Context.FinalOutput = FRenderingCompositeOutputRef(Node);
+}
+
 static void AddPostProcessDepthOfFieldBokeh(FPostprocessContext& Context, FRenderingCompositeOutputRef& SeparateTranslucency)
 {
 	// downsample, mask out the in focus part, depth in alpha
@@ -887,6 +896,7 @@ void FPostProcessing::Process(FRHICommandListImmediate& RHICmdList, FViewInfo& V
 		}
 
 		bool bVisualizeHDR = View.Family->EngineShowFlags.VisualizeHDR && FeatureLevel >= ERHIFeatureLevel::SM5;
+		bool bVisualizeBloom = View.Family->EngineShowFlags.VisualizeBloom && FeatureLevel >= ERHIFeatureLevel::SM4;
 
 		if(bVisualizeHDR)
 		{
@@ -1033,6 +1043,11 @@ void FPostProcessing::Process(FRHICommandListImmediate& RHICmdList, FViewInfo& V
 				Context.FinalOutput = FRenderingCompositeOutputRef(MotionBlurRecombinePass);
 			}
 
+			if(bVisualizeBloom)
+			{
+				AddVisualizeBloomSetup(Context);
+			}
+
 			// down sample Scene color from full to half res
 			FRenderingCompositeOutputRef SceneColorHalfRes;
 			{
@@ -1257,14 +1272,18 @@ void FPostProcessing::Process(FRHICommandListImmediate& RHICmdList, FViewInfo& V
 
 		// Show the selection outline if it is in the editor and we arent in wireframe 
 		// If the engine is in demo mode and game view is on we also do not show the selection outline
-		if ( GIsEditor && View.Family->EngineShowFlags.SelectionOutline && !(View.Family->EngineShowFlags.Wireframe) && ( !GIsDemoMode || ( GIsDemoMode && !View.Family->EngineShowFlags.Game ) ) )
+		if ( GIsEditor
+			&& View.Family->EngineShowFlags.SelectionOutline
+			&& !(View.Family->EngineShowFlags.Wireframe)
+			&& ( !GIsDemoMode || ( GIsDemoMode && !View.Family->EngineShowFlags.Game ) ) 
+			&& !bVisualizeBloom)
 		{
 			// Selection outline is after bloom, but before AA
 			AddSelectionOutline(Context);
 		}
 
 		// Composite editor primitives if we had any to draw and compositing is enabled
-		if (FSceneRenderer::ShouldCompositeEditorPrimitives(View))
+		if (FSceneRenderer::ShouldCompositeEditorPrimitives(View) && !bVisualizeBloom)
 		{
 			FRenderingCompositePass* Node = Context.Graph.RegisterPass(new(FMemStack::Get()) FRCPassPostProcessCompositeEditorPrimitives(true));
 			Node->SetInput(ePId_Input0, FRenderingCompositeOutputRef(Context.FinalOutput));
