improved visualize bloom

diff --git a/Engine/Shaders/PostProcessBloom.usf b/Engine/Shaders/PostProcessBloom.usf
--- a/Engine/Shaders/PostProcessBloom.usf
+++ b/Engine/Shaders/PostProcessBloom.usf
@@ -7,6 +7,10 @@
 #include "Common.usf"
 #include "PostProcessCommon.usf"
 
+// for VisualizeBloomOverlayPS
+// xyz:Bloom1Tint.rgb, w:unused
+float4 ColorScale1;
+
 // vertex shader entry point
 void MainPostprocessCommonVS(
 	in float4 InPosition : ATTRIBUTE0,
@@ -29,16 +33,70 @@ void VisualizeBloomSetupPS(
 
 	int2 PixelPos = (int2)ComputePixelPosCenter(UVAndScreenPos.zw, true);
 
-	int2 ViewportCenter = (int2)(ViewportRect + ViewportSize / 2);
+	int2 ViewportCenter = (int2)(ViewportRect.xy + ViewportSize / 2);
 
-	bool bCross = ViewportCenter.x == PixelPos.x || ViewportCenter.y == PixelPos.y;
+	// cross in x and y
+//	bool bMask = ViewportCenter.x == PixelPos.x || ViewportCenter.y == PixelPos.y;
+	// vertical line at viewport center
+	bool bMask = ViewportCenter.x == PixelPos.x;
 
-	if(bCross)
+	if(bMask)
 	{
 		OutColor = 20;
 	}
 }
 
+void VisualizeBloomOverlayPS(
+	  float4 UVAndScreenPos : TEXCOORD0
+	, out float4 OutColor : SV_Target0)
+{
+	float2 UV = UVAndScreenPos.xy;
+	
+	float3 LDRColor = Texture2DSample(PostprocessInput0, PostprocessInput0Sampler, UV).rgb;
+	
+	int2 ViewportCenter = (int2)(ViewportRect + ViewportSize / 2);
+
+	float2 UVLinePos = float2(UV.x, ViewportCenter.y * PostprocessInput0Size.w);
+
+	float3 NonBloomColor = Texture2DSampleLevel(PostprocessInput1, PostprocessInput1Sampler, UVLinePos, 0).rgb;
+
+	float3 BloomColor;
+	{
+		float4 CombinedBloom = Texture2DSample(PostprocessInput2, PostprocessInput2Sampler, UVLinePos);
+		float3 BloomDirtMaskColor = 0;
+		BloomColor = CombinedBloom.rgb * (ColorScale1.rgb + BloomDirtMaskColor); 
+	}
+
+	float NonBloomValue = Luminance(NonBloomColor);
+	float BloomValue = Luminance(BloomColor);
+
+	float CompareY = lerp(1.1f, -0.1f, UV.y);
+
+	bool bNonBloom = NonBloomValue > CompareY;
+	bool bBloom = BloomValue > CompareY;
+	bool bOutside = CompareY < 0 || CompareY > 1;
+
+	OutColor = 0.1f;
+
+	FLATTEN if(bNonBloom)
+	{
+		OutColor.rgb = 0.3f;
+	}
+
+	FLATTEN if(bBloom)
+	{
+		OutColor.rgb = float3(0, 1, 0);
+	}
+
+	FLATTEN if(bOutside)
+	{
+		OutColor.rgb = float3(0, 0, 0);
+	}
+	
+	// for debugging: show LDR scene color
+//	OutColor.rgb = LDRColor;
+}
+
 // vertex shader entry point
 void MainVS(
 	in float4 InPosition : ATTRIBUTE0,
diff --git a/Engine/Shaders/PostProcessCommon.usf b/Engine/Shaders/PostProcessCommon.usf
--- a/Engine/Shaders/PostProcessCommon.usf
+++ b/Engine/Shaders/PostProcessCommon.usf
@@ -58,9 +58,9 @@ float EyeAdaptationLookup()
 }
 #endif
 
-// viewport width, height, 1/width, 1/height (scaled to the current rendertarget resolution)
+// viewport width, height, 1/width, 1/height (scaled to the current rendertarget resolution), depends on SetViewportAndCallRHI() call
 float4 ViewportSize;
-// in pixels (scaled to the current rendertarget resolution), float4(minx,miny,maxx,maxy)
+// in pixels (scaled to the current rendertarget resolution), float4(minx,miny,maxx,maxy), depends on SetViewportAndCallRHI() call
 float4 ViewportRect;
 
 SamplerState BilinearTextureSampler0;
diff --git a/Engine/Source/Runtime/Renderer/Private/PostProcess/PostProcessBloomSetup.cpp b/Engine/Source/Runtime/Renderer/Private/PostProcess/PostProcessBloomSetup.cpp
--- a/Engine/Source/Runtime/Renderer/Private/PostProcess/PostProcessBloomSetup.cpp
+++ b/Engine/Source/Runtime/Renderer/Private/PostProcess/PostProcessBloomSetup.cpp
@@ -294,7 +294,7 @@ void FRCPassPostProcessVisualizeBloomSetup::Process(FRenderingCompositePassConte
 	// is optimized away if possible (RT size=view size, )
 	Context.RHICmdList.Clear(true, FLinearColor(0, 0, 0, 0), false, 1.0f, false, 0, DestRect);
 
-	Context.SetViewportAndCallRHI(0, 0, 0.0f, DestSize.X, DestSize.Y, 1.0f );
+	Context.SetViewportAndCallRHI(0, 0, 0.0f, DestRect.Width(), DestRect.Height(), 1.0f );
 
 	// set the state
 	Context.RHICmdList.SetBlendState(TStaticBlendState<>::GetRHI());
@@ -318,7 +318,7 @@ void FRCPassPostProcessVisualizeBloomSetup::Process(FRenderingCompositePassConte
 		DestRect.Width(), DestRect.Height(),
 		SrcRect.Min.X, SrcRect.Min.Y,
 		SrcRect.Width(), SrcRect.Height(),
-		DestSize,
+		DestRect.Size(),
 		SrcSize,
 		*VertexShader,
 		EDRF_UseTriangleOptimization);
@@ -338,3 +338,142 @@ FPooledRenderTargetDesc FRCPassPostProcessVisualizeBloomSetup::ComputeOutputDesc
 
 	return Ret;
 }
+
+
+
+
+
+
+
+/** Encapsulates the visualize bloom overlay pixel shader. */
+class FPostProcessVisualizeBloomOverlayPS : public FGlobalShader
+{
+	DECLARE_SHADER_TYPE(FPostProcessVisualizeBloomOverlayPS, Global);
+
+	static bool ShouldCache(EShaderPlatform Platform)
+	{
+		return IsFeatureLevelSupported(Platform, ERHIFeatureLevel::SM4);
+	}
+
+	static void ModifyCompilationEnvironment(EShaderPlatform Platform, FShaderCompilerEnvironment& OutEnvironment)
+	{
+		FGlobalShader::ModifyCompilationEnvironment(Platform,OutEnvironment);
+	}
+
+	/** Default constructor. */
+	FPostProcessVisualizeBloomOverlayPS() {}
+
+public:
+	FPostProcessPassParameters PostprocessParameter;
+	FShaderParameter ColorScale1;
+
+	/** Initialization constructor. */
+	FPostProcessVisualizeBloomOverlayPS(const ShaderMetaType::CompiledShaderInitializerType& Initializer)
+		: FGlobalShader(Initializer)
+	{
+		PostprocessParameter.Bind(Initializer.ParameterMap);
+		ColorScale1.Bind(Initializer.ParameterMap, TEXT("ColorScale1"));
+	}
+
+	// FShader interface.
+	virtual bool Serialize(FArchive& Ar)
+	{
+		bool bShaderHasOutdatedParameters = FGlobalShader::Serialize(Ar);
+		Ar << PostprocessParameter << ColorScale1;
+		return bShaderHasOutdatedParameters;
+	}
+
+	void SetParameters(const FRenderingCompositePassContext& Context)
+	{
+		const FPixelShaderRHIParamRef ShaderRHI = GetPixelShader();
+
+		const FPostProcessSettings& Settings = Context.View.FinalPostProcessSettings;
+
+		FGlobalShader::SetParameters(Context.RHICmdList, ShaderRHI, Context.View);
+		PostprocessParameter.SetPS(ShaderRHI, Context, TStaticSamplerState<SF_Bilinear,AM_Border,AM_Border,AM_Border>::GetRHI());
+
+		{
+			FLinearColor Col = FLinearColor::White * Settings.BloomIntensity;
+			FVector4 ColorScale(Col.R, Col.G, Col.B, 0);
+			SetShaderValue(Context.RHICmdList, ShaderRHI, ColorScale1, ColorScale);
+		}
+	}
+};
+
+IMPLEMENT_SHADER_TYPE(,FPostProcessVisualizeBloomOverlayPS,TEXT("PostProcessBloom"),TEXT("VisualizeBloomOverlayPS"),SF_Pixel);
+
+
+
+
+void FRCPassPostProcessVisualizeBloomOverlay::Process(FRenderingCompositePassContext& Context)
+{
+	SCOPED_DRAW_EVENT(Context.RHICmdList, VisualizeBloomOverlay);
+
+	const FPooledRenderTargetDesc* InputDesc = GetInputDesc(ePId_Input0);
+
+	check(InputDesc && "Input is not hooked up correctly");
+
+	const FSceneView& View = Context.View;
+	const FSceneViewFamily& ViewFamily = *(View.Family);
+
+	FIntPoint SrcSize = InputDesc->Extent;
+	FIntPoint DestSize = PassOutputs[0].RenderTargetDesc.Extent;
+
+	// e.g. 4 means the input texture is 4x smaller than the buffer size
+	uint32 ScaleFactor = GSceneRenderTargets.GetBufferSizeXY().X / SrcSize.X;
+
+	FIntRect SrcRect = View.ViewRect / ScaleFactor;
+	FIntRect DestRect = SrcRect;
+
+	const FSceneRenderTargetItem& DestRenderTarget = PassOutputs[0].RequestSurface(Context);
+
+	SetRenderTarget(Context.RHICmdList, DestRenderTarget.TargetableTexture, FTextureRHIRef());
+
+	// is optimized away if possible (RT size=view size, )
+	Context.RHICmdList.Clear(true, FLinearColor(0, 0, 0, 0), false, 1.0f, false, 0, DestRect);
+
+	Context.SetViewportAndCallRHI(0, 0, 0.0f, DestSize.X, DestSize.Y, 1.0f );
+
+	// set the state
+	Context.RHICmdList.SetBlendState(TStaticBlendState<>::GetRHI());
+	Context.RHICmdList.SetRasterizerState(TStaticRasterizerState<>::GetRHI());
+	Context.RHICmdList.SetDepthStencilState(TStaticDepthStencilState<false, CF_Always>::GetRHI());
+
+	TShaderMapRef<FPostProcessVS> VertexShader(Context.GetShaderMap());
+	TShaderMapRef<FPostProcessVisualizeBloomOverlayPS> PixelShader(Context.GetShaderMap());
+
+	static FGlobalBoundShaderState BoundShaderState;
+
+	SetGlobalBoundShaderState(Context.RHICmdList, Context.GetFeatureLevel(), BoundShaderState, GFilterVertexDeclaration.VertexDeclarationRHI, *VertexShader, *PixelShader);
+
+	PixelShader->SetParameters(Context);
+	VertexShader->SetParameters(Context);
+
+	// Draw a quad mapping scene color to the view's render target
+	DrawRectangle(
+		Context.RHICmdList,
+		DestRect.Min.X, DestRect.Min.Y,
+		DestRect.Width(), DestRect.Height(),
+		SrcRect.Min.X, SrcRect.Min.Y,
+		SrcRect.Width(), SrcRect.Height(),
+		DestSize,
+		SrcSize,
+		*VertexShader,
+		EDRF_UseTriangleOptimization);
+
+	Context.RHICmdList.CopyToResolveTarget(DestRenderTarget.TargetableTexture, DestRenderTarget.ShaderResourceTexture, false, FResolveParams());
+}
+
+
+FPooledRenderTargetDesc FRCPassPostProcessVisualizeBloomOverlay::ComputeOutputDesc(EPassOutputId InPassOutputId) const
+{
+	FPooledRenderTargetDesc Ret = PassInputs[0].GetOutput()->RenderTargetDesc;
+
+	Ret.Reset();
+	Ret.TargetableFlags &= ~TexCreate_UAV;
+	Ret.TargetableFlags |= TexCreate_RenderTargetable;
+	Ret.DebugName = TEXT("VisualizeBloomOverlay");
+
+	return Ret;
+}
+
diff --git a/Engine/Source/Runtime/Renderer/Private/PostProcess/PostProcessBloomSetup.h b/Engine/Source/Runtime/Renderer/Private/PostProcess/PostProcessBloomSetup.h
--- a/Engine/Source/Runtime/Renderer/Private/PostProcess/PostProcessBloomSetup.h
+++ b/Engine/Source/Runtime/Renderer/Private/PostProcess/PostProcessBloomSetup.h
@@ -30,3 +30,16 @@ public:
 	virtual FPooledRenderTargetDesc ComputeOutputDesc(EPassOutputId InPassOutputId) const;
 	virtual void Release() override { delete this; }
 };
+
+// ePId_Input0: LDR SceneColor
+// ePId_Input1: HDR SceneColor
+// ePId_Input2: BloomOutputCombined
+// derives from TRenderingCompositePassBase<InputCount, OutputCount> 
+class FRCPassPostProcessVisualizeBloomOverlay : public TRenderingCompositePassBase<3, 1>
+{
+public:
+	// interface FRenderingCompositePass ---------
+	virtual void Process(FRenderingCompositePassContext& Context);
+	virtual FPooledRenderTargetDesc ComputeOutputDesc(EPassOutputId InPassOutputId) const;
+	virtual void Release() override { delete this; }
+};
diff --git a/Engine/Source/Runtime/Renderer/Private/PostProcess/PostProcessing.cpp b/Engine/Source/Runtime/Renderer/Private/PostProcess/PostProcessing.cpp
--- a/Engine/Source/Runtime/Renderer/Private/PostProcess/PostProcessing.cpp
+++ b/Engine/Source/Runtime/Renderer/Private/PostProcess/PostProcessing.cpp
@@ -255,13 +255,24 @@ static FRenderingCompositeOutputRef AddPostProcessEyeAdaptation(FPostprocessCont
 
 static void AddVisualizeBloomSetup(FPostprocessContext& Context)
 {
-	FRenderingCompositePass* Node = Context.Graph.RegisterPass(new(FMemStack::Get()) FRCPassPostProcessVisualizeBloomSetup());
+	auto Node = Context.Graph.RegisterPass(new(FMemStack::Get()) FRCPassPostProcessVisualizeBloomSetup());
 
 	Node->SetInput(ePId_Input0, Context.FinalOutput);
 
 	Context.FinalOutput = FRenderingCompositeOutputRef(Node);
 }
 
+static void AddVisualizeBloomOverlay(FPostprocessContext& Context, FRenderingCompositeOutputRef& HDRColor, FRenderingCompositeOutputRef& BloomOutputCombined)
+{
+	auto Node = Context.Graph.RegisterPass(new(FMemStack::Get()) FRCPassPostProcessVisualizeBloomOverlay());
+
+	Node->SetInput(ePId_Input0, Context.FinalOutput);
+	Node->SetInput(ePId_Input1, HDRColor);
+	Node->SetInput(ePId_Input2, BloomOutputCombined);
+
+	Context.FinalOutput = FRenderingCompositeOutputRef(Node);
+}
+
 static void AddPostProcessDepthOfFieldBokeh(FPostprocessContext& Context, FRenderingCompositeOutputRef& SeparateTranslucency)
 {
 	// downsample, mask out the in focus part, depth in alpha
@@ -876,6 +887,8 @@ void FPostProcessing::Process(FRHICommandListImmediate& RHICmdList, FViewInfo& V
 		FRenderingCompositeOutputRef EyeAdaptation;
 		// not always valid
 		FRenderingCompositeOutputRef SeparateTranslucency;
+		// optional
+		FRenderingCompositeOutputRef BloomOutputCombined;
 
 		bool bAllowTonemapper = true;
 		EStereoscopicPass StereoPass = View.StereoPass;
@@ -1101,9 +1114,6 @@ void FPostProcessing::Process(FRHICommandListImmediate& RHICmdList, FViewInfo& V
 				EyeAdaptation = AddPostProcessEyeAdaptation(Context, Histogram);
 			}
 
-			// optional
-			FRenderingCompositeOutputRef BloomOutputCombined;
-
 			if(View.Family->EngineShowFlags.Bloom)
 			{
 				if (CVarUseMobileBloom.GetValueOnRenderThread() == 0)
@@ -1302,6 +1312,11 @@ void FPostProcessing::Process(FRHICommandListImmediate& RHICmdList, FViewInfo& V
 		
 		AddPostProcessMaterial(Context, BL_AfterTonemapping, SeparateTranslucency);
 
+		if(bVisualizeBloom)
+		{
+			AddVisualizeBloomOverlay(Context, HDRColor, BloomOutputCombined);
+		}
+
 		if (View.Family->EngineShowFlags.VisualizeSSS)
 		{
 			// the setup pass also does visualization, based on EngineShowFlags.VisualizeSSS
