Fixed Clang compile error

diff --git a/Engine/Source/Runtime/Renderer/Private/ShadowSetup.cpp b/Engine/Source/Runtime/Renderer/Private/ShadowSetup.cpp
--- a/Engine/Source/Runtime/Renderer/Private/ShadowSetup.cpp
+++ b/Engine/Source/Runtime/Renderer/Private/ShadowSetup.cpp
@@ -354,9 +354,7 @@ static bool GetBestShadowTransform(const FVector& ZAxis,const FBoundingBoxVertex
 }
 
 FProjectedShadowInfo::FProjectedShadowInfo()
-	: LightSceneInfo(0)
-	, ParentSceneInfo(0)
-	, DependentView(0)
+	: DependentView(0)
 	, ShadowId(INDEX_NONE)
 	, PreShadowTranslation(0, 0, 0)
 	, ShadowBounds(0)
@@ -376,6 +374,8 @@ FProjectedShadowInfo::FProjectedShadowInfo()
 	, bTranslucentShadow(false)
 	, bPreShadow(false)
 	, bSelfShadowOnly(false)
+	, LightSceneInfo(0)
+	, ParentSceneInfo(0)
 	, ShaderDepthBias(0.0f)
 {
 }
