UE4 - Skip velocity VS/PS when using base pass velocity outputs

diff --git a/Engine/Source/Runtime/Renderer/Private/VelocityRendering.cpp b/Engine/Source/Runtime/Renderer/Private/VelocityRendering.cpp
--- a/Engine/Source/Runtime/Renderer/Private/VelocityRendering.cpp
+++ b/Engine/Source/Runtime/Renderer/Private/VelocityRendering.cpp
@@ -50,7 +50,7 @@ public:
 			|| (Material->IsTwoSided() && !IsTranslucentBlendMode(Material->GetBlendMode()))
 			// or if the material modifies meshes
 			|| Material->MaterialMayModifyMeshPosition()))
-			&& IsFeatureLevelSupported(Platform, ERHIFeatureLevel::SM4);
+			&& IsFeatureLevelSupported(Platform, ERHIFeatureLevel::SM4) && !FVelocityRendering::OutputsToGBuffer();
 	}
 
 protected:
@@ -138,7 +138,7 @@ public:
 			|| (Material->IsTwoSided() && !IsTranslucentBlendMode(Material->GetBlendMode()))
 			// or if the material modifies meshes
 			|| Material->MaterialMayModifyMeshPosition()))
-			&& IsFeatureLevelSupported(Platform, ERHIFeatureLevel::SM4);
+			&& IsFeatureLevelSupported(Platform, ERHIFeatureLevel::SM4) && !FVelocityRendering::OutputsToGBuffer();
 	}
 
 	static void ModifyCompilationEnvironment( EShaderPlatform Platform, const FMaterial* Material, FShaderCompilerEnvironment& OutEnvironment )
