UE4 - Fix tessellation velocity calculations

diff --git a/Engine/Shaders/BasePassTessellationShaders.usf b/Engine/Shaders/BasePassTessellationShaders.usf
--- a/Engine/Shaders/BasePassTessellationShaders.usf
+++ b/Engine/Shaders/BasePassTessellationShaders.usf
@@ -42,6 +42,13 @@ FBasePassVSToPS PassFinalizeTessellationOutput(FBasePassVSToDS Interpolants, flo
 	ISOLATE
 	{
 		O.Position = mul(WorldPosition, View.TranslatedWorldToClip);
+#if WRITES_VELOCITY_TO_GBUFFER
+		float SkipVelocity = O.BasePassInterpolants.PrevScreenPosition.w;
+		O.BasePassInterpolants.PrevScreenPosition = mul(float4(Interpolants.BasePassInterpolants.PrevScreenPosition.xyz, 1), View.PrevTranslatedWorldToClip);
+		// Move the mask to .z when as we don't need the clip space Z
+		O.BasePassInterpolants.PrevScreenPosition.z = SkipVelocity;
+		O.BasePassInterpolants.ScreenPosition = mul(float4(Interpolants.BasePassInterpolants.ScreenPosition.xyz, 1), View.TranslatedWorldToClip);
+#endif
 	}
 	
 	// Calc position-dependent interps
diff --git a/Engine/Shaders/BasePassVertexShader.usf b/Engine/Shaders/BasePassVertexShader.usf
--- a/Engine/Shaders/BasePassVertexShader.usf
+++ b/Engine/Shaders/BasePassVertexShader.usf
@@ -52,7 +52,7 @@ void Main(
 	#if USE_WORLD_POSITION_EXCLUDING_SHADER_OFFSETS
 		Output.BasePassInterpolants.PixelPositionExcludingWPO = WorldPositionExcludingWPO;
 	#endif
-#endif
+#endif	// USING_TESSELLATION
 
 	Output.FactoryInterpolants = VertexFactoryGetInterpolants(Input, VFIntermediates, VertexParameters);
 
@@ -71,12 +71,21 @@ void Main(
 		VertexParameters = GetMaterialVertexParameters(Input, VFIntermediates, PrevTranslatedWorldPosition.xyz, TangentToLocal);
 		PrevTranslatedWorldPosition.xyz += GetMaterialPreviousWorldPositionOffset(VertexParameters);
 
-		// compute the old screen pos with the old world position and the old camera matrix
-		Output.BasePassInterpolants.PrevScreenPosition = mul(float4(PrevTranslatedWorldPosition.xyz, 1), View.PrevTranslatedWorldToClip);		
-		Output.BasePassInterpolants.PrevScreenPosition.z = 1 - SkipOutputVelocity;
-		Output.BasePassInterpolants.ScreenPosition = mul(float4(WorldPosition.xyz, 1), View.TranslatedWorldToClip);
+		#if USING_TESSELLATION
+			// We let the Domain Shader convert to post projection when tessellating
+			Output.BasePassInterpolants.PrevScreenPosition = PrevTranslatedWorldPosition;
+			// Store the mask in .w when doing tessellation as we need the world space Z
+			Output.BasePassInterpolants.PrevScreenPosition.w = 1 - SkipOutputVelocity;
+			Output.BasePassInterpolants.ScreenPosition = WorldPosition;
+		#else
+			// compute the old screen pos with the old world position and the old camera matrix
+			Output.BasePassInterpolants.PrevScreenPosition = mul(float4(PrevTranslatedWorldPosition.xyz, 1), View.PrevTranslatedWorldToClip);		
+			// Store the mask in .z when as we don't need the clip space Z
+			Output.BasePassInterpolants.PrevScreenPosition.z = 1 - SkipOutputVelocity;
+			Output.BasePassInterpolants.ScreenPosition = Output.Position;
+		#endif	// USING_TESSELLATION
 	}
-#endif
+#endif	// WRITES_VELOCITY_TO_GBUFFER
 
 	OutputVertexID( Output );
 }
diff --git a/Engine/Source/Runtime/Renderer/Private/BasePassRendering.inl b/Engine/Source/Runtime/Renderer/Private/BasePassRendering.inl
--- a/Engine/Source/Runtime/Renderer/Private/BasePassRendering.inl
+++ b/Engine/Source/Runtime/Renderer/Private/BasePassRendering.inl
@@ -13,24 +13,35 @@ inline void TBasePassVertexShaderBaseType<LightMapPolicyType>::SetMesh(FRHIComma
 	FVertexShaderRHIParamRef VertexShaderRHI = GetVertexShader();
 	FMeshMaterialShader::SetMesh(RHICmdList, VertexShaderRHI, VertexFactory, View, Proxy, BatchElement);
 
-	if (PreviousLocalToWorldParameter.IsBound() && Proxy)
+	const bool bHasPreviousLocalToWorldParameter = PreviousLocalToWorldParameter.IsBound();
+	const bool bHasSkipOutputVelocityParameter = SkipOutputVelocityParameter.IsBound();
+	if ((bHasSkipOutputVelocityParameter || bHasPreviousLocalToWorldParameter) && Proxy)
 	{
-		check(SkipOutputVelocityParameter.IsBound());
+//		check(SkipOutputVelocityParameter.IsBound());
 
-		FMatrix PreviousLocalToWorld;
-		bool bHasPreviousLocalToWorld = false;
+		FMatrix PreviousLocalToWorldMatrix;
+		bool bHasPreviousLocalToWorldMatrix = false;
 		const auto& ViewInfo = (const FViewInfo&)View;
+		float SkipOutputVelocityValue;
 		if (FVelocityDrawingPolicy::HasVelocityOnBasePass(ViewInfo, Proxy, Proxy->GetPrimitiveSceneInfo(), Mesh,
-			bHasPreviousLocalToWorld, PreviousLocalToWorld))
+			bHasPreviousLocalToWorldMatrix, PreviousLocalToWorldMatrix))
 		{
-			const FMatrix& Transform = bHasPreviousLocalToWorld ? PreviousLocalToWorld : Proxy->GetLocalToWorld();
-			SetShaderValue(RHICmdList, VertexShaderRHI, PreviousLocalToWorldParameter, Transform.ConcatTranslation(ViewInfo.PrevViewMatrices.PreViewTranslation));
+			if (bHasPreviousLocalToWorldParameter)
+			{
+				const FMatrix& Transform = bHasPreviousLocalToWorldMatrix ? PreviousLocalToWorldMatrix : Proxy->GetLocalToWorld();
+				SetShaderValue(RHICmdList, VertexShaderRHI, PreviousLocalToWorldParameter, Transform.ConcatTranslation(ViewInfo.PrevViewMatrices.PreViewTranslation));
+			}
 
-			SetShaderValue(RHICmdList, VertexShaderRHI, SkipOutputVelocityParameter, 0.0f);
+			SkipOutputVelocityValue = 0.0f;
 		}
 		else
 		{
-			SetShaderValue(RHICmdList, VertexShaderRHI, SkipOutputVelocityParameter, 1.0f);
+			SkipOutputVelocityValue = 1.0f;
+		}
+
+		if (bHasSkipOutputVelocityParameter)
+		{
+			SetShaderValue(RHICmdList, VertexShaderRHI, SkipOutputVelocityParameter, SkipOutputVelocityValue);
 		}
 	}
 }
