disabled lens flares when using visualize bloom

diff --git a/Engine/Source/Runtime/Renderer/Private/PostProcess/PostProcessing.cpp b/Engine/Source/Runtime/Renderer/Private/PostProcess/PostProcessing.cpp
--- a/Engine/Source/Runtime/Renderer/Private/PostProcess/PostProcessing.cpp
+++ b/Engine/Source/Runtime/Renderer/Private/PostProcess/PostProcessing.cpp
@@ -459,6 +459,8 @@ static FRenderingCompositeOutputRef AddBloom(FPostprocessContext Context, FRende
 		PostProcessDownsamples[i] = FRenderingCompositeOutputRef(Pass);
 	}
 
+	const bool bVisualizeBloom = Context.View.Family->EngineShowFlags.VisualizeBloom;
+
 	FRenderingCompositeOutputRef BloomOutput;
 	if (BloomQuality == 0)
 	{
@@ -475,8 +477,6 @@ static FRenderingCompositeOutputRef AddBloom(FPostprocessContext Context, FRende
 		};
 		const FFinalPostProcessSettings& Settings = Context.View.FinalPostProcessSettings;
 
-		bool bVisualizeBloom = Context.View.Family->EngineShowFlags.VisualizeBloom;
-
 		FBloomStage BloomStages[] =
 		{
 			{ Settings.Bloom6Size, &Settings.Bloom6Tint},
@@ -521,7 +521,7 @@ static FRenderingCompositeOutputRef AddBloom(FPostprocessContext Context, FRende
 		LensFlareQuality = FMath::Clamp(CVar->GetValueOnRenderThread(), 0, MaxLensFlareQuality);
 	}
 
-	if (!LensFlareHDRColor.IsAlmostBlack() && LensFlareQuality > 0)
+	if (!LensFlareHDRColor.IsAlmostBlack() && LensFlareQuality > 0 && !bVisualizeBloom)
 	{
 		float PercentKernelSize = Context.View.FinalPostProcessSettings.LensFlareBokehSize;
 
