Add MRT clear value binding to the RHI.  This allows parallel rendering platforms to propagate clear color values across parallel execution boundaries for platforms which require low level control over hardware clears.
Also add many explicit calls to CopyToResolveTarget to indicate the transition of a render target to be able to be used as an SRV.  Soon to be replaced by a more specific RHI function.
Added r.CheckSRVTransitions for D3D11 on windows to check if any of these transition are missing.
Fixes a bug with RHIThreadFence when using PLATFORM_SUPPORTS_PARALLEL_RHI_EXECUTE
PLATFORM_SUPPORTS_PARALLEL_RHI_EXECUTE on PS4 now supports CMASK and HTILE enabled buffers.

diff --git a/Engine/Source/Runtime/Core/Private/HAL/ConsoleManager.cpp b/Engine/Source/Runtime/Core/Private/HAL/ConsoleManager.cpp
--- a/Engine/Source/Runtime/Core/Private/HAL/ConsoleManager.cpp
+++ b/Engine/Source/Runtime/Core/Private/HAL/ConsoleManager.cpp
@@ -2086,4 +2086,10 @@ static TAutoConsoleVariable<float> CVarEmitterSpawnRateScale(
 	TEXT("r.EmitterSpawnRateScale"),
 	1.0,
 	TEXT("A global scale upon the spawn rate of emitters. Emitters can choose to apply or ignore it via their bApplyGlobalSpawnRateScale property."),
-	ECVF_Scalability | ECVF_RenderThreadSafe);
\ No newline at end of file
+	ECVF_Scalability | ECVF_RenderThreadSafe);
+
+static TAutoConsoleVariable<int32> CVarCheckSRVTransitions(
+	TEXT("r.CheckSRVTransitions"),
+	0,
+	TEXT("Tests that render targets are properly transitioned to SRV when SRVs are set."),
+	ECVF_RenderThreadSafe);  
\ No newline at end of file
diff --git a/Engine/Source/Runtime/Engine/Private/Particles/ParticleGpuSimulation.cpp b/Engine/Source/Runtime/Engine/Private/Particles/ParticleGpuSimulation.cpp
--- a/Engine/Source/Runtime/Engine/Private/Particles/ParticleGpuSimulation.cpp
+++ b/Engine/Source/Runtime/Engine/Private/Particles/ParticleGpuSimulation.cpp
@@ -4049,7 +4049,9 @@ void FFXSystem::SimulateGPUParticles(
 	{
 		SetRenderTarget(RHICmdList, PrevStateTextures.PositionTextureTargetRHI, FTextureRHIRef(), ESimpleRenderTargetMode::EClearToDefault);
 		SetRenderTarget(RHICmdList, PrevStateTextures.VelocityTextureTargetRHI, FTextureRHIRef(), ESimpleRenderTargetMode::EClearToDefault);
-
+		RHICmdList.CopyToResolveTarget(PrevStateTextures.PositionTextureTargetRHI, PrevStateTextures.PositionTextureTargetRHI, true, FResolveParams());
+		RHICmdList.CopyToResolveTarget(PrevStateTextures.VelocityTextureTargetRHI, PrevStateTextures.VelocityTextureTargetRHI, true, FResolveParams());
+		
 		PrevStateTextures.bTexturesCleared = true;
 	}
 
diff --git a/Engine/Source/Runtime/IOS/MetalRHI/Private/MetalCommands.cpp b/Engine/Source/Runtime/IOS/MetalRHI/Private/MetalCommands.cpp
--- a/Engine/Source/Runtime/IOS/MetalRHI/Private/MetalCommands.cpp
+++ b/Engine/Source/Runtime/IOS/MetalRHI/Private/MetalCommands.cpp
@@ -834,7 +834,10 @@ void FMetalDynamicRHI::RHIClearMRT(bool bClearColor,int32 NumClearColors,const F
 	}
 }
 
-
+void FMetalDynamicRHI::RHIBindClearMRTValues(bool bClearColor, int32 NumClearColors, const FLinearColor* ClearColorArray, bool bClearDepth, float Depth, bool bClearStencil, uint32 Stencil, bool bCheckErrors)
+{
+	// Not necessary
+}
 
 void FMetalDynamicRHI::RHISuspendRendering()
 {
diff --git a/Engine/Source/Runtime/OpenGLDrv/Private/OpenGLCommands.cpp b/Engine/Source/Runtime/OpenGLDrv/Private/OpenGLCommands.cpp
--- a/Engine/Source/Runtime/OpenGLDrv/Private/OpenGLCommands.cpp
+++ b/Engine/Source/Runtime/OpenGLDrv/Private/OpenGLCommands.cpp
@@ -1581,7 +1581,7 @@ void FOpenGLDynamicRHI::RHISetBlendState(FBlendStateRHIParamRef NewStateRHI,cons
 void FOpenGLDynamicRHI::RHISetRenderTargets(
 	uint32 NumSimultaneousRenderTargets,
 	const FRHIRenderTargetView* NewRenderTargetsRHI,
-	FTextureRHIParamRef NewDepthStencilTargetRHI,
+	const FRHIDepthRenderTargetView* NewDepthStencilTargetRHI,
 	uint32 NumUAVs,
 	const FUnorderedAccessViewRHIParamRef* UAVs
 	)
@@ -1607,7 +1607,7 @@ void FOpenGLDynamicRHI::RHISetRenderTargets(
 		}
 	}
 
-	FOpenGLTextureBase* NewDepthStencilRT = GetOpenGLTextureFromRHITexture(NewDepthStencilTargetRHI);
+	FOpenGLTextureBase* NewDepthStencilRT = GetOpenGLTextureFromRHITexture(NewDepthStencilTargetRHI ? NewDepthStencilTargetRHI->Texture : nullptr);
 
 	if (IsES2Platform(GMaxRHIShaderPlatform) && !IsPCPlatform(GMaxRHIShaderPlatform))
 	{
@@ -1716,8 +1716,8 @@ void FOpenGLDynamicRHI::RHISetRenderTargets(
 		// Set viewport size to new depth target size.
 		PendingState.Viewport.Min.X = 0;
 		PendingState.Viewport.Min.Y = 0;
-		PendingState.Viewport.Max.X = GetOpenGLTextureSizeXFromRHITexture(NewDepthStencilTargetRHI);
-		PendingState.Viewport.Max.Y = GetOpenGLTextureSizeYFromRHITexture(NewDepthStencilTargetRHI);
+		PendingState.Viewport.Max.X = GetOpenGLTextureSizeXFromRHITexture(NewDepthStencilTargetRHI->Texture);
+		PendingState.Viewport.Max.Y = GetOpenGLTextureSizeYFromRHITexture(NewDepthStencilTargetRHI->Texture);
 	}
 }
 
@@ -1760,7 +1760,7 @@ void FOpenGLDynamicRHI::RHISetRenderTargetsAndClear(const FRHISetRenderTargetsIn
 {
 	this->RHISetRenderTargets(RenderTargetsInfo.NumColorRenderTargets,
 		RenderTargetsInfo.ColorRenderTarget,
-		RenderTargetsInfo.DepthStencilRenderTarget.Texture,
+		&RenderTargetsInfo.DepthStencilRenderTarget,
 		0,
 		nullptr);
 	if (RenderTargetsInfo.bClearColor || RenderTargetsInfo.bClearStencil || RenderTargetsInfo.bClearDepth)
@@ -3213,6 +3213,11 @@ void FOpenGLDynamicRHI::RHIClearMRT(bool bClearColor,int32 NumClearColors,const
 	}
 }
 
+void FOpenGLDynamicRHI::RHIBindClearMRTValues(bool bClearColor, int32 NumClearColors, const FLinearColor* ClearColorArray, bool bClearDepth, float Depth, bool bClearStencil, uint32 Stencil)
+{
+	// Not necessary for opengl.
+}
+
 // Functions to yield and regain rendering control from OpenGL
 
 void FOpenGLDynamicRHI::RHISuspendRendering()
diff --git a/Engine/Source/Runtime/OpenGLDrv/Private/OpenGLViewport.cpp b/Engine/Source/Runtime/OpenGLDrv/Private/OpenGLViewport.cpp
--- a/Engine/Source/Runtime/OpenGLDrv/Private/OpenGLViewport.cpp
+++ b/Engine/Source/Runtime/OpenGLDrv/Private/OpenGLViewport.cpp
@@ -124,13 +124,13 @@ void FOpenGLDynamicRHI::RHIBeginDrawingViewport(FViewportRHIParamRef ViewportRHI
 	// Set the render target and viewport.
 	if( RenderTarget )
 	{
-		FRHIRenderTargetView RTV(RenderTarget);
-		RHISetRenderTargets(1, &RTV, FTextureRHIRef(), 0, NULL);
+		FRHIRenderTargetView RTV(RenderTarget);	
+		RHISetRenderTargets(1, &RTV, nullptr, 0, NULL);
 	}
 	else
 	{
 		FRHIRenderTargetView RTV(DrawingViewport->GetBackBuffer());
-		RHISetRenderTargets(1, &RTV, FTextureRHIRef(), 0, NULL);
+		RHISetRenderTargets(1, &RTV, nullptr, 0, NULL);
 	}
 }
 
diff --git a/Engine/Source/Runtime/RHI/Private/RHICommandList.cpp b/Engine/Source/Runtime/RHI/Private/RHICommandList.cpp
--- a/Engine/Source/Runtime/RHI/Private/RHICommandList.cpp
+++ b/Engine/Source/Runtime/RHI/Private/RHICommandList.cpp
@@ -54,6 +54,11 @@ static TAutoConsoleVariable<int32> CVarRHIDeferredContextWidth(
 	TEXT("r.RHIDeferredContextWidth"),
 	99999,
 	TEXT("Number of pieces to break a parallel translate task into. A very large number is one task per original cmd list."));
+
+static TAutoConsoleVariable<int32> CVarRHICmdFlushOnQueueParallelSubmit(
+	TEXT("r.RHICmdFlushOnQueueParallelSubmit"),
+	0,
+	TEXT("Wait for completion of parallel commandlists immediately after submitting.  Debug feature."));
 #endif
 
 RHI_API FRHICommandListExecutor GRHICommandList;
@@ -62,7 +67,7 @@ static FGraphEventArray AllOutstandingTasks;
 static FGraphEventArray WaitOutstandingTasks;
 static FGraphEventRef RHIThreadTask;
 static FGraphEventRef RenderThreadSublistDispatchTask;
-
+static FGraphEventRef RenderThreadExecutionTask;
 
 DECLARE_CYCLE_STAT(TEXT("RHI Thread Execute"), STAT_RHIThreadExecute, STATGROUP_RHICMDLIST);
 
@@ -127,11 +132,19 @@ void FRHICommandListExecutor::ExecuteInner_DoExecute(FRHICommandListBase& CmdLis
 class FExecuteRHIThreadTask
 {
 	FRHICommandListBase* RHICmdList;
+	FGraphEventRef ExecuteFenceTask;
 
 public:
 
 	FExecuteRHIThreadTask(FRHICommandListBase* InRHICmdList)
 		: RHICmdList(InRHICmdList)
+		, ExecuteFenceTask(nullptr)
+	{
+	}
+
+	FExecuteRHIThreadTask(FRHICommandListBase* InRHICmdList, FGraphEventRef InExecuteFenceTask)
+		: RHICmdList(InRHICmdList)
+		, ExecuteFenceTask(InExecuteFenceTask)
 	{
 	}
 
@@ -157,18 +170,67 @@ public:
 		SCOPE_CYCLE_COUNTER(STAT_RHIThreadExecute);
 		FRHICommandListExecutor::ExecuteInner_DoExecute(*RHICmdList);
 		delete RHICmdList;
+
+		if (ExecuteFenceTask)
+		{			
+			ExecuteFenceTask->DispatchSubsequents();
+			ExecuteFenceTask = nullptr;
+		}
+	}
+};
+
+class FExecuteFenceRHITask
+{	
+	FGraphEventRef ExecuteFenceTask;
+
+public:
+
+	FExecuteFenceRHITask(FGraphEventRef InExecuteFenceTask)		
+		: ExecuteFenceTask(InExecuteFenceTask)
+	{
+	}
+
+	FORCEINLINE TStatId GetStatId() const
+	{
+		RETURN_QUICK_DECLARE_CYCLE_STAT(FExecuteFenceRHITask, STATGROUP_TaskGraphTasks);
+	}
+
+	ENamedThreads::Type GetDesiredThread()
+	{
+#if PLATFORM_SUPPORTS_RHI_THREAD
+		return ENamedThreads::RHIThread;
+#else
+		check(0); // this should never be used on a platform that doesn't support the RHI thread
+		return ENamedThreads::AnyThread;
+#endif
+	}
+
+	static ESubsequentsMode::Type GetSubsequentsMode() { return ESubsequentsMode::TrackSubsequents; }
+
+	void DoTask(ENamedThreads::Type CurrentThread, const FGraphEventRef& MyCompletionGraphEvent)
+	{
+		ExecuteFenceTask->DispatchSubsequents();
+		ExecuteFenceTask = nullptr;		
 	}
 };
 
 class FDispatchRHIThreadTask
 {
 	FRHICommandListBase* RHICmdList;
+	FGraphEventRef ExecuteFenceTask;
 
 public:
 
 	FDispatchRHIThreadTask(FRHICommandListBase* InRHICmdList)
 		: RHICmdList(InRHICmdList)
-	{
+		, ExecuteFenceTask(nullptr)
+	{		
+	}
+
+	FDispatchRHIThreadTask(FRHICommandListBase* InRHICmdList, FGraphEventRef InExecuteFenceTask)
+		: RHICmdList(InRHICmdList)
+		, ExecuteFenceTask(InExecuteFenceTask)
+	{		
 	}
 
 	FORCEINLINE TStatId GetStatId() const
@@ -194,7 +256,8 @@ public:
 		{
 			Prereq.Add(RHIThreadTask);
 		}
-		RHIThreadTask = TGraphTask<FExecuteRHIThreadTask>::CreateTask(&Prereq, ENamedThreads::RenderThread).ConstructAndDispatchWhenReady(RHICmdList);
+		RHIThreadTask = TGraphTask<FExecuteRHIThreadTask>::CreateTask(&Prereq, ENamedThreads::RenderThread).ConstructAndDispatchWhenReady(RHICmdList, ExecuteFenceTask);
+		ExecuteFenceTask = nullptr;
 	}
 };
 
@@ -214,10 +277,9 @@ void FRHICommandListExecutor::ExecuteInner(FRHICommandListBase& CmdList)
 		ECVF_Cheat);
 	check(CmdList.RTTasks.Num() == 0 && CmdList.HasCommands()); 
 
+	bool bIsInRenderingThread = IsInRenderingThread();
 	if (GRHIThread)
 	{
-		bool bIsInRenderingThread = IsInRenderingThread();
-
 		if (bIsInRenderingThread)
 		{
 			if (RHIThreadTask.GetReference() && RHIThreadTask->IsComplete())
@@ -228,6 +290,10 @@ void FRHICommandListExecutor::ExecuteInner(FRHICommandListBase& CmdList)
 			{
 				RenderThreadSublistDispatchTask = nullptr;
 			}
+			if (RenderThreadExecutionTask.GetReference() && RenderThreadExecutionTask->IsComplete())
+			{
+				RenderThreadExecutionTask = nullptr;
+			}
 		}
 		if (CVarRHICmdUseThread.GetValueOnRenderThread() > 0 && bIsInRenderingThread && !IsInGameThread())
 		{
@@ -238,6 +304,10 @@ void FRHICommandListExecutor::ExecuteInner(FRHICommandListBase& CmdList)
 			static_assert(sizeof(FRHICommandList) == sizeof(FRHICommandListImmediate), "We are memswapping FRHICommandList and FRHICommandListImmediate; they need to be swappable.");
 			SwapCmdList->ExchangeCmdList(CmdList);
 
+			//if we use a FDispatchRHIThreadTask, we must have it pass an event along to the FExecuteRHIThreadTask it will spawn so that fences can know which event to wait on for execution completion
+			//before the dispatch completes.
+			//if we use a FExecuteRHIThreadTask directly we pass the same event just to keep things consistent.
+			RenderThreadExecutionTask = FGraphEvent::CreateGraphEvent();
 			if (AllOutstandingTasks.Num() || RenderThreadSublistDispatchTask.GetReference())
 			{
 				FGraphEventArray Prereq(AllOutstandingTasks);
@@ -246,7 +316,7 @@ void FRHICommandListExecutor::ExecuteInner(FRHICommandListBase& CmdList)
 				{
 					Prereq.Add(RenderThreadSublistDispatchTask);
 				}
-				RenderThreadSublistDispatchTask = TGraphTask<FDispatchRHIThreadTask>::CreateTask(&Prereq, ENamedThreads::RenderThread).ConstructAndDispatchWhenReady(SwapCmdList);
+				RenderThreadSublistDispatchTask = TGraphTask<FDispatchRHIThreadTask>::CreateTask(&Prereq, ENamedThreads::RenderThread).ConstructAndDispatchWhenReady(SwapCmdList, RenderThreadExecutionTask);
 			}
 			else
 			{
@@ -255,7 +325,7 @@ void FRHICommandListExecutor::ExecuteInner(FRHICommandListBase& CmdList)
 				{
 					Prereq.Add(RHIThreadTask);
 				}
-				RHIThreadTask = TGraphTask<FExecuteRHIThreadTask>::CreateTask(&Prereq, ENamedThreads::RenderThread).ConstructAndDispatchWhenReady(SwapCmdList);
+				RHIThreadTask = TGraphTask<FExecuteRHIThreadTask>::CreateTask(&Prereq, ENamedThreads::RenderThread).ConstructAndDispatchWhenReady(SwapCmdList, RenderThreadExecutionTask);
 			}
 			if (CVarRHICmdForceRHIFlush.GetValueOnRenderThread() > 0 )
 			{
@@ -437,31 +507,17 @@ FGraphEventRef FRHICommandListExecutor::RHIThreadFence()
 		QUICK_SCOPE_CYCLE_COUNTER(STAT_RHIThreadFence_Dispatch);
 		FRHICommandListExecutor::GetImmediateCommandList().ImmediateFlush(EImmediateFlushType::DispatchToRHIThread);
 	}
+
 #if PLATFORM_SUPPORTS_RHI_THREAD
-	FGraphEventArray Finish;
-	if (RHIThreadTask.GetReference() && RHIThreadTask->IsComplete())
-	{
-		RHIThreadTask = nullptr;
-	}
-	if (RHIThreadTask.GetReference())
+	// every task dispatched or executed through the flush creates a new RenderThreadExecutionTask that represents the 'execution' of the last dispatched
+	// set of RHI commands.
+	if (RenderThreadExecutionTask.GetReference() && RenderThreadExecutionTask->IsComplete())
 	{
-		Finish.Add(RHIThreadTask);
-	}
-	if (RenderThreadSublistDispatchTask.GetReference() && RenderThreadSublistDispatchTask->IsComplete())
-	{
-		RenderThreadSublistDispatchTask = nullptr;
-	}
-	if (RenderThreadSublistDispatchTask.GetReference())
-	{
-		Finish.Add(RenderThreadSublistDispatchTask);
-	}
-	if (Finish.Num())
-	{
-		return TGraphTask<FNullGraphTask>::CreateTask( &Finish, ENamedThreads::RenderThread )
-			.ConstructAndDispatchWhenReady(GET_STATID(STAT_RHIThreadFence), ENamedThreads::RHIThread);
-	}
+		RenderThreadExecutionTask = nullptr;
+	}		
 #endif
-	return FGraphEventRef();
+	
+	return RenderThreadExecutionTask;
 }
 void FRHICommandListExecutor::WaitOnRHIThreadFence(FGraphEventRef& Fence)
 {
@@ -658,6 +714,15 @@ void FRHICommandListBase::QueueParallelAsyncCommandListSubmit(FGraphEventRef* An
 		FRHICommandListExecutor::GetImmediateCommandList().ImmediateFlush(EImmediateFlushType::DispatchToRHIThread); // we should start on the stuff before this async list
 	}
 #if PLATFORM_SUPPORTS_PARALLEL_RHI_EXECUTE
+
+#if !UE_BUILD_SHIPPING
+	// do a flush before hand so we can tell if it was this parallel set that broke something, or what came before.
+	if (CVarRHICmdFlushOnQueueParallelSubmit.GetValueOnRenderThread())
+	{
+		FRHICommandListExecutor::GetImmediateCommandList().ImmediateFlush(EImmediateFlushType::FlushRHIThread);
+	}
+#endif
+
 	if (Num && GRHIThread)
 	{
 		IRHICommandContextContainer* ContextContainer = nullptr;
@@ -713,6 +778,13 @@ void FRHICommandListBase::QueueParallelAsyncCommandListSubmit(FGraphEventRef* An
 					ContextContainer = nullptr;
 					Start = Last + 1;
 				}
+
+#if !UE_BUILD_SHIPPING
+				if (CVarRHICmdFlushOnQueueParallelSubmit.GetValueOnRenderThread())
+				{
+					FRHICommandListExecutor::GetImmediateCommandList().ImmediateFlush(EImmediateFlushType::FlushRHIThread);
+				}
+#endif
 				return;
 			}
 		}
diff --git a/Engine/Source/Runtime/RHI/Public/RHICommandList.h b/Engine/Source/Runtime/RHI/Public/RHICommandList.h
--- a/Engine/Source/Runtime/RHI/Public/RHICommandList.h
+++ b/Engine/Source/Runtime/RHI/Public/RHICommandList.h
@@ -508,19 +508,18 @@ struct FRHICommandSetRenderTargets : public FRHICommand<FRHICommandSetRenderTarg
 {
 	uint32 NewNumSimultaneousRenderTargets;
 	FRHIRenderTargetView NewRenderTargetsRHI[MaxSimultaneousRenderTargets];
-	FTextureRHIParamRef NewDepthStencilTargetRHI;
+	FRHIDepthRenderTargetView NewDepthStencilTarget;
 	uint32 NewNumUAVs;
 	FUnorderedAccessViewRHIParamRef UAVs[MaxSimultaneousUAVs];
 
 	FORCEINLINE_DEBUGGABLE FRHICommandSetRenderTargets(
 		uint32 InNewNumSimultaneousRenderTargets,
 		const FRHIRenderTargetView* InNewRenderTargetsRHI,
-		FTextureRHIParamRef InNewDepthStencilTargetRHI,
+		const FRHIDepthRenderTargetView* InNewDepthStencilTargetRHI,
 		uint32 InNewNumUAVs,
 		const FUnorderedAccessViewRHIParamRef* InUAVs
 		)
 		: NewNumSimultaneousRenderTargets(InNewNumSimultaneousRenderTargets)
-		, NewDepthStencilTargetRHI(InNewDepthStencilTargetRHI)
 		, NewNumUAVs(InNewNumUAVs)
 
 	{
@@ -533,6 +532,10 @@ struct FRHICommandSetRenderTargets : public FRHICommand<FRHICommandSetRenderTarg
 		{
 			UAVs[Index] = InUAVs[Index];
 		}
+		if (InNewDepthStencilTargetRHI)
+		{
+			NewDepthStencilTarget = *InNewDepthStencilTargetRHI;
+		}		
 	}
 	RHI_API void Execute(FRHICommandListBase& CmdList);
 };
@@ -549,6 +552,42 @@ struct FRHICommandSetRenderTargetsAndClear : public FRHICommand<FRHICommandSetRe
 	RHI_API void Execute(FRHICommandListBase& CmdList);
 };
 
+struct FRHICommandBindClearMRTValues : public FRHICommand<FRHICommandBindClearMRTValues>
+{
+	FLinearColor ColorArray[MaxSimultaneousRenderTargets];	
+	float Depth;
+	uint32 Stencil;
+	int32 NumClearColors;
+	bool bClearColor;
+	bool bClearDepth;
+	bool bClearStencil;
+
+	FORCEINLINE_DEBUGGABLE FRHICommandBindClearMRTValues(
+		bool InbClearColor,
+		int32 InNumClearColors,
+		const FLinearColor* InColorArray,
+		bool InbClearDepth,
+		float InDepth,
+		bool InbClearStencil,
+		uint32 InStencil		
+		) 
+		: Depth(InDepth)
+		, Stencil(InStencil)
+		, NumClearColors(InNumClearColors)
+		, bClearColor(InbClearColor)
+		, bClearDepth(InbClearDepth)
+		, bClearStencil(InbClearStencil)
+	{
+		check(InNumClearColors < MaxSimultaneousRenderTargets);
+		for (int32 Index = 0; Index < InNumClearColors; Index++)
+		{
+			ColorArray[Index] = InColorArray[Index];
+		}
+	}	
+
+	RHI_API void Execute(FRHICommandListBase& CmdList);
+};
+
 struct FRHICommandEndDrawPrimitiveUP : public FRHICommand<FRHICommandEndDrawPrimitiveUP>
 {
 	uint32 PrimitiveType;
@@ -1509,7 +1548,7 @@ public:
 	FORCEINLINE_DEBUGGABLE void SetRenderTargets(
 		uint32 NewNumSimultaneousRenderTargets,
 		const FRHIRenderTargetView* NewRenderTargetsRHI,
-		FTextureRHIParamRef NewDepthStencilTargetRHI,
+		const FRHIDepthRenderTargetView* NewDepthStencilTargetRHI,
 		uint32 NewNumUAVs,
 		const FUnorderedAccessViewRHIParamRef* UAVs
 		)
@@ -1540,7 +1579,17 @@ public:
 			return;
 		}
 		new (AllocCommand<FRHICommandSetRenderTargetsAndClear>()) FRHICommandSetRenderTargetsAndClear(RenderTargetsInfo);
-	}
+	}	
+
+	FORCEINLINE_DEBUGGABLE void BindClearMRTValues(bool bClearColor, int32 NumClearColors, const FLinearColor* ColorArray, bool bClearDepth, float Depth, bool bClearStencil, uint32 Stencil)
+	{
+		if (Bypass())
+		{
+			CMD_CONTEXT(BindClearMRTValues)(bClearColor, NumClearColors, ColorArray, bClearDepth, Depth, bClearStencil, Stencil);
+			return;
+		}
+		new (AllocCommand<FRHICommandBindClearMRTValues>()) FRHICommandBindClearMRTValues(bClearColor, NumClearColors, ColorArray, bClearDepth, Depth, bClearStencil, Stencil);
+	}	
 
 	FORCEINLINE_DEBUGGABLE void BeginDrawPrimitiveUP(uint32 PrimitiveType, uint32 NumPrimitives, uint32 NumVertices, uint32 VertexDataStride, void*& OutVertexData)
 	{
@@ -1900,11 +1949,10 @@ class RHI_API FRHICommandListImmediate : public FRHICommandList
 	{
 		check(!HasCommands());
 	}
-
-	static bool bFlushedGlobal;
-
 public:
 
+	static bool bFlushedGlobal;	
+	
 	inline void ImmediateFlush(EImmediateFlushType::Type FlushType);
 
 	void SetCurrentStat(TStatId Stat);
diff --git a/Engine/Source/Runtime/RHI/Public/RHICommandListCommandExecutes.inl b/Engine/Source/Runtime/RHI/Public/RHICommandListCommandExecutes.inl
--- a/Engine/Source/Runtime/RHI/Public/RHICommandListCommandExecutes.inl
+++ b/Engine/Source/Runtime/RHI/Public/RHICommandListCommandExecutes.inl
@@ -163,7 +163,7 @@ void FRHICommandSetRenderTargets::Execute(FRHICommandListBase& CmdList)
 	INTERNAL_DECORATOR(SetRenderTargets)(
 		NewNumSimultaneousRenderTargets,
 		NewRenderTargetsRHI,
-		NewDepthStencilTargetRHI,
+		&NewDepthStencilTarget,
 		NewNumUAVs,
 		UAVs);
 }
@@ -174,6 +174,20 @@ void FRHICommandSetRenderTargetsAndClear::Execute(FRHICommandListBase& CmdList)
 	INTERNAL_DECORATOR(SetRenderTargetsAndClear)(RenderTargetsInfo);
 }
 
+void FRHICommandBindClearMRTValues::Execute(FRHICommandListBase& CmdList)
+{
+	RHISTAT(BindClearMRTValues);
+	INTERNAL_DECORATOR(BindClearMRTValues)(
+		bClearColor,
+		NumClearColors,
+		ColorArray,
+		bClearDepth,
+		Depth,
+		bClearStencil,
+		Stencil	
+		);
+}
+
 void FRHICommandEndDrawPrimitiveUP::Execute(FRHICommandListBase& CmdList)
 {
 	RHISTAT(EndDrawPrimitiveUP);
diff --git a/Engine/Source/Runtime/RHI/Public/RHIDefinitions.h b/Engine/Source/Runtime/RHI/Public/RHIDefinitions.h
--- a/Engine/Source/Runtime/RHI/Public/RHIDefinitions.h
+++ b/Engine/Source/Runtime/RHI/Public/RHIDefinitions.h
@@ -546,16 +546,21 @@ enum class ERenderTargetStoreAction
 enum class ESimpleRenderTargetMode
 {
 	// These will all store out color and depth
-	EExistingColorAndDepth,			// Color = Existing, Depth = Existing
-	EUninitializedColorAndDepth,	// Color = ????, Depth = ????
-	EUninitializedColorExistingDepth,// Color = ????, Depth = Existing
-	EUninitializedColorClearDepth,	// Color = ????, Depth = Default
-	EClearToDefault,				// Default Color = (0,0,0,0), Default Depth = 0.0f
-	EClearColorToBlack,				// Color = (0,0,0,0), Depth = Existing
-	EClearColorToBlackWithFullAlpha,// Color = (0,0,0,1), Depth = Existing
-	EClearColorToWhite,				// Color = (1,1,1,1), Depth = Existing
-	EClearDepthToOne,				// Color = Existing, Depth = 1.0
-	EExistingContents_NoDepthStore,	// Load existing contents, but don't store depth out
+	EExistingColorAndDepth,						// Color = Existing, Depth = Existing
+	EExistingColorAndReadOnlyDepth,					// Color = Existing, Depth = Existing. depth cannot be written.
+	EUninitializedColorAndDepth,					// Color = ????, Depth = ????
+	EUninitializedColorExistingDepth,				// Color = ????, Depth = Existing
+	EUninitializedColorExistingReadOnlyDepth,		// Color = ????, Depth = Existing. depth cannot be written.
+	EUninitializedColorClearDepth,					// Color = ????, Depth = Default
+	EClearToDefault,								// Default Color = (0,0,0,0), Default Depth = 0.0f
+	EClearColorToBlack,								// Color = (0,0,0,0), Depth = Existing
+	EClearColorToBlackReadOnlyDepth,				// Color = (0,0,0,0), Depth = Existing. depth cannot be written
+	EClearColorToBlackWithFullAlpha,				// Color = (0,0,0,1), Depth = Existing
+	EClearColorToBlackWithFullAlphaReadOnlyDepth,	// Color = (0,0,0,1), Depth = Existing. depth cannot be written
+	EClearColorToWhite,								// Color = (1,1,1,1), Depth = Existing
+	EClearColorToWhiteReadOnlyDepth,				// Color = (1,1,1,1), Depth = Existing. depth cannot be written
+	EClearDepthToOne,								// Color = Existing, Depth = 1.0
+	EExistingContents_NoDepthStore,					// Load existing contents, but don't store depth out.  depth can be written.
 
 	// If you add an item here, make sure to add it to DecodeRenderTargetMode() as well!
 };
diff --git a/Engine/Source/Runtime/RHI/Public/RHIMethods.h b/Engine/Source/Runtime/RHI/Public/RHIMethods.h
--- a/Engine/Source/Runtime/RHI/Public/RHIMethods.h
+++ b/Engine/Source/Runtime/RHI/Public/RHIMethods.h
@@ -1897,11 +1897,12 @@ DEFINE_RHIMETHOD_CMDLIST_5(
 	void,SetRenderTargets,
 	uint32,NumSimultaneousRenderTargets,
 	const FRHIRenderTargetView*,NewRenderTargets,
-	FTextureRHIParamRef,NewDepthStencilTarget,
+	const FRHIDepthRenderTargetView*, NewDepthStencilTarget,
 	uint32,NumUAVs,
 	const FUnorderedAccessViewRHIParamRef*,UAVs,
 	,
 	);
+
 DEFINE_RHIMETHOD_3(
 	void,DiscardRenderTargets,
 	bool,Depth,
@@ -1915,6 +1916,21 @@ DEFINE_RHIMETHOD_CMDLIST_1(
 	,
 	);
 
+// Bind a clear values to currently bound rendertargets.  This is used by platforms which
+// need the color when transitioning a target that supports hardware clears from a rendertarget to a shader resource.
+// The explicit bind is needed to support parallel rendering.
+DEFINE_RHIMETHOD_CMDLIST_7(
+	void, BindClearMRTValues,
+	bool, bClearColor,
+	int32, NumClearColors,
+	const FLinearColor*, ColorArray,
+	bool, bClearDepth,
+	float, Depth,
+	bool, bClearStencil,
+	uint32, Stencil,	
+	,
+	);
+
 // Primitive drawing.
 DEFINE_RHIMETHOD_CMDLIST_4(
 	void,DrawPrimitive,
diff --git a/Engine/Source/Runtime/RHI/Public/RHIResources.h b/Engine/Source/Runtime/RHI/Public/RHIResources.h
--- a/Engine/Source/Runtime/RHI/Public/RHIResources.h
+++ b/Engine/Source/Runtime/RHI/Public/RHIResources.h
@@ -354,6 +354,16 @@ public:
 		return &LastRenderTime;
 	}
 
+	void SetName(FName& InName)
+	{
+		TextureName = InName;
+	}
+
+	FName GetName() const
+	{
+		return TextureName;
+	}
+
 private:
 	uint32 NumMips;
 	uint32 NumSamples;
@@ -361,6 +371,7 @@ private:
 	uint32 Flags;
 	FLastRenderTimeContainer& LastRenderTime;
 	FLastRenderTimeContainer DefaultLastRenderTime;
+	FName TextureName;
 };
 
 class RHI_API FRHITexture2D : public FRHITexture
@@ -636,17 +647,19 @@ class FRHIDepthRenderTargetView
 public:
 	FTextureRHIParamRef Texture;
 
-	ERenderTargetLoadAction DepthLoadAction;
-	ERenderTargetStoreAction DepthStoreAction;
-	ERenderTargetLoadAction StencilLoadAction;
-	ERenderTargetStoreAction StencilStoreAction;
+	ERenderTargetLoadAction		DepthLoadAction;
+	ERenderTargetStoreAction	DepthStoreAction;
+	ERenderTargetLoadAction		StencilLoadAction;
+	ERenderTargetStoreAction	StencilStoreAction;
+	bool						bReadOnly; //target will not be written to.
 
 	FRHIDepthRenderTargetView() :
 		Texture(nullptr),
 		DepthLoadAction(ERenderTargetLoadAction::EClear),
 		DepthStoreAction(ERenderTargetStoreAction::EStore),
 		StencilLoadAction(ERenderTargetLoadAction::EClear),
-		StencilStoreAction(ERenderTargetStoreAction::EStore)
+		StencilStoreAction(ERenderTargetStoreAction::EStore),
+		bReadOnly(false)
 	{}
 
 	FRHIDepthRenderTargetView(FTextureRHIParamRef InTexture) :
@@ -654,7 +667,8 @@ public:
 		DepthLoadAction(ERenderTargetLoadAction::EClear),
 		DepthStoreAction(ERenderTargetStoreAction::EStore),
 		StencilLoadAction(ERenderTargetLoadAction::EClear),
-		StencilStoreAction(ERenderTargetStoreAction::EStore)
+		StencilStoreAction(ERenderTargetStoreAction::EStore),
+		bReadOnly(false)
 	{}
 
 	FRHIDepthRenderTargetView(FTextureRHIParamRef InTexture, ERenderTargetLoadAction InLoadAction, ERenderTargetStoreAction InStoreAction) :
@@ -662,17 +676,41 @@ public:
 		DepthLoadAction(InLoadAction),
 		DepthStoreAction(InStoreAction),
 		StencilLoadAction(InLoadAction),
-		StencilStoreAction(InStoreAction)
+		StencilStoreAction(InStoreAction),
+		bReadOnly(false)
 	{}
 
+	FRHIDepthRenderTargetView(FTextureRHIParamRef InTexture, ERenderTargetLoadAction InLoadAction, ERenderTargetStoreAction InStoreAction, bool bInReadOnly) :
+		Texture(InTexture),
+		DepthLoadAction(InLoadAction),
+		DepthStoreAction(InStoreAction),
+		StencilLoadAction(InLoadAction),
+		StencilStoreAction(InStoreAction),
+		bReadOnly(bInReadOnly)
+	{
+		ensure(!bReadOnly || InStoreAction == ERenderTargetStoreAction::ENoAction);
+	}
+
 	FRHIDepthRenderTargetView(FTextureRHIParamRef InTexture, ERenderTargetLoadAction InDepthLoadAction, ERenderTargetStoreAction InDepthStoreAction, ERenderTargetLoadAction InStencilLoadAction, ERenderTargetStoreAction InStencilStoreAction) :
 		Texture(InTexture),
 		DepthLoadAction(InDepthLoadAction),
 		DepthStoreAction(InDepthStoreAction),
 		StencilLoadAction(InStencilLoadAction),
-		StencilStoreAction(InStencilStoreAction)
+		StencilStoreAction(InStencilStoreAction),
+		bReadOnly(false)
 	{}
 
+	FRHIDepthRenderTargetView(FTextureRHIParamRef InTexture, ERenderTargetLoadAction InDepthLoadAction, ERenderTargetStoreAction InDepthStoreAction, ERenderTargetLoadAction InStencilLoadAction, ERenderTargetStoreAction InStencilStoreAction, bool bInReadOnly) :
+		Texture(InTexture),
+		DepthLoadAction(InDepthLoadAction),
+		DepthStoreAction(InDepthStoreAction),
+		StencilLoadAction(InStencilLoadAction),
+		StencilStoreAction(InStencilStoreAction),
+		bReadOnly(bInReadOnly)
+	{
+		ensure(!bReadOnly || InDepthStoreAction == ERenderTargetStoreAction::ENoAction);
+	}
+
 	bool operator==(const FRHIDepthRenderTargetView& Other)
 	{
 		return
@@ -680,7 +718,8 @@ public:
 			DepthLoadAction == Other.DepthLoadAction &&
 			DepthStoreAction == Other.DepthStoreAction &&
 			StencilLoadAction == Other.StencilLoadAction &&
-			StencilStoreAction == Other.StencilStoreAction;
+			StencilStoreAction == Other.StencilStoreAction &&
+			bReadOnly == Other.bReadOnly;
 	}
 };
 
@@ -715,8 +754,8 @@ public:
 		DepthStencilRenderTarget(InDepthStencilRenderTarget),
 		DepthClearValue(0.0f),
 		StencilClearValue(0),
-		bClearDepth(InDepthStencilRenderTarget.DepthLoadAction == ERenderTargetLoadAction::EClear),
-		bClearStencil(InDepthStencilRenderTarget.StencilLoadAction == ERenderTargetLoadAction::EClear)
+		bClearDepth(InDepthStencilRenderTarget.Texture && InDepthStencilRenderTarget.DepthLoadAction == ERenderTargetLoadAction::EClear),
+		bClearStencil(InDepthStencilRenderTarget.Texture && InDepthStencilRenderTarget.StencilLoadAction == ERenderTargetLoadAction::EClear)
 	{
 		check(InNumColorRenderTargets <= 0 || InColorRenderTargets);
 		for (int32 Index = 0; Index < InNumColorRenderTargets; ++Index)
diff --git a/Engine/Source/Runtime/RHI/Public/RHIUtilities.h b/Engine/Source/Runtime/RHI/Public/RHIUtilities.h
--- a/Engine/Source/Runtime/RHI/Public/RHIUtilities.h
+++ b/Engine/Source/Runtime/RHI/Public/RHIUtilities.h
@@ -123,13 +123,14 @@ struct FRWBufferByteAddress
  * Convert the ESimpleRenderTargetMode into usable values 
  * @todo: Can we easily put this into a .cpp somewhere?
  */
-inline void DecodeRenderTargetMode(ESimpleRenderTargetMode Mode, ERenderTargetLoadAction& ColorLoadAction, ERenderTargetStoreAction& ColorStoreAction, ERenderTargetLoadAction& DepthLoadAction, ERenderTargetStoreAction& DepthStoreAction, FLinearColor& ClearColor, float& ClearDepth)
+inline void DecodeRenderTargetMode(ESimpleRenderTargetMode Mode, ERenderTargetLoadAction& ColorLoadAction, ERenderTargetStoreAction& ColorStoreAction, ERenderTargetLoadAction& DepthLoadAction, ERenderTargetStoreAction& DepthStoreAction, FLinearColor& ClearColor, float& ClearDepth, bool& bReadOnlyDepth)
 {
 	// set defaults
 	ColorStoreAction = ERenderTargetStoreAction::EStore;
 	DepthStoreAction = ERenderTargetStoreAction::EStore;
 	ClearColor = FLinearColor(0, 0, 0, 0);
 	ClearDepth = 0.0f;
+	bReadOnlyDepth = false;
 
 	switch (Mode)
 	{
@@ -137,6 +138,12 @@ inline void DecodeRenderTargetMode(ESimpleRenderTargetMode Mode, ERenderTargetLo
 		ColorLoadAction = ERenderTargetLoadAction::ELoad;
 		DepthLoadAction = ERenderTargetLoadAction::ELoad;
 		break;
+	case ESimpleRenderTargetMode::EExistingColorAndReadOnlyDepth:
+		ColorLoadAction = ERenderTargetLoadAction::ELoad;
+		DepthLoadAction = ERenderTargetLoadAction::ELoad;
+		DepthStoreAction = ERenderTargetStoreAction::ENoAction;
+		bReadOnlyDepth = true;
+		break;
 	case ESimpleRenderTargetMode::EUninitializedColorAndDepth:
 		ColorLoadAction = ERenderTargetLoadAction::ENoAction;
 		DepthLoadAction = ERenderTargetLoadAction::ENoAction;
@@ -145,6 +152,12 @@ inline void DecodeRenderTargetMode(ESimpleRenderTargetMode Mode, ERenderTargetLo
 		ColorLoadAction = ERenderTargetLoadAction::ENoAction;
 		DepthLoadAction = ERenderTargetLoadAction::ELoad;
 		break;
+	case ESimpleRenderTargetMode::EUninitializedColorExistingReadOnlyDepth:
+		ColorLoadAction = ERenderTargetLoadAction::ENoAction;
+		DepthLoadAction = ERenderTargetLoadAction::ELoad;
+		DepthStoreAction = ERenderTargetStoreAction::ENoAction;
+		bReadOnlyDepth = true;
+		break;
 	case ESimpleRenderTargetMode::EUninitializedColorClearDepth:
 		ColorLoadAction = ERenderTargetLoadAction::ENoAction;
 		DepthLoadAction = ERenderTargetLoadAction::EClear;
@@ -158,16 +171,37 @@ inline void DecodeRenderTargetMode(ESimpleRenderTargetMode Mode, ERenderTargetLo
 		DepthLoadAction = ERenderTargetLoadAction::ELoad;
 		ClearColor = FLinearColor::Black;
 		break;
+	case ESimpleRenderTargetMode::EClearColorToBlackReadOnlyDepth:
+		ColorLoadAction = ERenderTargetLoadAction::EClear;
+		DepthLoadAction = ERenderTargetLoadAction::ELoad;
+		ClearColor = FLinearColor::Black;
+		DepthStoreAction = ERenderTargetStoreAction::ENoAction;
+		bReadOnlyDepth = true;
+		break;
 	case ESimpleRenderTargetMode::EClearColorToBlackWithFullAlpha:
 		ColorLoadAction = ERenderTargetLoadAction::EClear;
 		DepthLoadAction = ERenderTargetLoadAction::ELoad;
 		ClearColor = FLinearColor(0, 0, 0, 1);
 		break;
+	case ESimpleRenderTargetMode::EClearColorToBlackWithFullAlphaReadOnlyDepth:
+		ColorLoadAction = ERenderTargetLoadAction::EClear;
+		DepthLoadAction = ERenderTargetLoadAction::ELoad;
+		ClearColor = FLinearColor(0, 0, 0, 1);
+		DepthStoreAction = ERenderTargetStoreAction::ENoAction;
+		bReadOnlyDepth = true;
+		break;
 	case ESimpleRenderTargetMode::EClearColorToWhite:
 		ColorLoadAction = ERenderTargetLoadAction::EClear;
 		DepthLoadAction = ERenderTargetLoadAction::ELoad;
 		ClearColor = FLinearColor::White;
 		break;
+	case ESimpleRenderTargetMode::EClearColorToWhiteReadOnlyDepth:
+		ColorLoadAction = ERenderTargetLoadAction::EClear;
+		DepthLoadAction = ERenderTargetLoadAction::ELoad;
+		ClearColor = FLinearColor::White;
+		DepthStoreAction = ERenderTargetStoreAction::ENoAction;
+		bReadOnlyDepth = true;
+		break;
 	case ESimpleRenderTargetMode::EClearDepthToOne:
 		ColorLoadAction = ERenderTargetLoadAction::ELoad;
 		DepthLoadAction = ERenderTargetLoadAction::EClear;
@@ -187,7 +221,8 @@ inline void DecodeRenderTargetMode(ESimpleRenderTargetMode Mode, ERenderTargetLo
 inline void SetRenderTarget(FRHICommandList& RHICmdList, FTextureRHIParamRef NewRenderTarget, FTextureRHIParamRef NewDepthStencilTarget)
 {
 	FRHIRenderTargetView RTV(NewRenderTarget);
-	RHICmdList.SetRenderTargets(1, &RTV, NewDepthStencilTarget, 0, NULL);
+	FRHIDepthRenderTargetView DepthRTV(NewDepthStencilTarget);
+	RHICmdList.SetRenderTargets(1, &RTV, &DepthRTV, 0, NULL);
 }
 
 /** Helper for the common case of using a single color and depth render target. */
@@ -197,11 +232,12 @@ inline void SetRenderTarget(FRHICommandList& RHICmdList, FTextureRHIParamRef New
 	ERenderTargetStoreAction ColorStoreAction, DepthStoreAction;
 	FLinearColor ClearColor;
 	float ClearDepth;
-	DecodeRenderTargetMode(Mode, ColorLoadAction, ColorStoreAction, DepthLoadAction, DepthStoreAction, ClearColor, ClearDepth);
+	bool bReadOnlyDepth;
+	DecodeRenderTargetMode(Mode, ColorLoadAction, ColorStoreAction, DepthLoadAction, DepthStoreAction, ClearColor, ClearDepth, bReadOnlyDepth);
 
 	// now make the FRHISetRenderTargetsInfo that encapsulates all of the info
 	FRHIRenderTargetView ColorView(NewRenderTarget, 0, -1, ColorLoadAction, ColorStoreAction);
-	FRHISetRenderTargetsInfo Info(1, &ColorView, FRHIDepthRenderTargetView(NewDepthStencilTarget, DepthLoadAction, DepthStoreAction));
+	FRHISetRenderTargetsInfo Info(1, &ColorView, FRHIDepthRenderTargetView(NewDepthStencilTarget, DepthLoadAction, DepthStoreAction, bReadOnlyDepth));
 	Info.ClearColors[0] = ClearColor;
 	Info.DepthClearValue = ClearDepth;
 	RHICmdList.SetRenderTargetsAndClear(Info);
@@ -211,14 +247,16 @@ inline void SetRenderTarget(FRHICommandList& RHICmdList, FTextureRHIParamRef New
 inline void SetRenderTarget(FRHICommandList& RHICmdList, FTextureRHIParamRef NewRenderTarget, int32 MipIndex, FTextureRHIParamRef NewDepthStencilTarget)
 {
 	FRHIRenderTargetView RTV(NewRenderTarget, MipIndex, -1);
-	RHICmdList.SetRenderTargets(1, &RTV, NewDepthStencilTarget, 0, NULL);
+	FRHIDepthRenderTargetView DepthRTV(NewDepthStencilTarget);
+	RHICmdList.SetRenderTargets(1, &RTV, &DepthRTV, 0, nullptr);
 }
 
 /** Helper for the common case of using a single color and depth render target, with a mip index for the color target. */
 inline void SetRenderTarget(FRHICommandList& RHICmdList, FTextureRHIParamRef NewRenderTarget, int32 MipIndex, int32 ArraySliceIndex, FTextureRHIParamRef NewDepthStencilTarget)
 {
 	FRHIRenderTargetView RTV(NewRenderTarget, MipIndex, ArraySliceIndex);
-	RHICmdList.SetRenderTargets(1, &RTV, NewDepthStencilTarget, 0, NULL);
+	FRHIDepthRenderTargetView DepthRTV(NewDepthStencilTarget);
+	RHICmdList.SetRenderTargets(1, &RTV, &DepthRTV, 0, nullptr);
 }
 
 /** Helper that converts FTextureRHIParamRef's into FRHIRenderTargetView's. */
@@ -238,7 +276,35 @@ inline void SetRenderTargets(
 		RTVs[Index] = FRHIRenderTargetView(NewRenderTargetsRHI[Index]);
 	}
 
-	RHICmdList.SetRenderTargets(NewNumSimultaneousRenderTargets, RTVs, NewDepthStencilTargetRHI, NewNumUAVs, UAVs);
+	FRHIDepthRenderTargetView DepthRTV(NewDepthStencilTargetRHI);
+	RHICmdList.SetRenderTargets(NewNumSimultaneousRenderTargets, RTVs, &DepthRTV, NewNumUAVs, UAVs);
+}
+
+/** Helper that converts FTextureRHIParamRef's into FRHIRenderTargetView's. */
+inline void SetRenderTargets(
+	FRHICommandList& RHICmdList,
+	uint32 NewNumSimultaneousRenderTargets,
+	const FTextureRHIParamRef* NewRenderTargetsRHI,
+	FTextureRHIParamRef NewDepthStencilTargetRHI,
+	ESimpleRenderTargetMode Mode
+	)
+{
+	ERenderTargetLoadAction ColorLoadAction, DepthLoadAction;
+	ERenderTargetStoreAction ColorStoreAction, DepthStoreAction;
+	FLinearColor ClearColor;
+	float ClearDepth;
+	bool bReadOnlyDepth;
+	DecodeRenderTargetMode(Mode, ColorLoadAction, ColorStoreAction, DepthLoadAction, DepthStoreAction, ClearColor, ClearDepth, bReadOnlyDepth);
+
+	FRHIRenderTargetView RTVs[MaxSimultaneousRenderTargets];
+
+	for (uint32 Index = 0; Index < NewNumSimultaneousRenderTargets; Index++)
+	{
+		RTVs[Index] = FRHIRenderTargetView(NewRenderTargetsRHI[Index], 0, -1, ColorLoadAction, ColorStoreAction);
+	}
+
+	FRHIDepthRenderTargetView DepthRTV(NewDepthStencilTargetRHI, DepthLoadAction, DepthStoreAction, bReadOnlyDepth);
+	RHICmdList.SetRenderTargets(NewNumSimultaneousRenderTargets, RTVs, &DepthRTV, 0, nullptr);
 }
 
 /**
diff --git a/Engine/Source/Runtime/Renderer/Private/AtmosphereRendering.cpp b/Engine/Source/Runtime/Renderer/Private/AtmosphereRendering.cpp
--- a/Engine/Source/Runtime/Renderer/Private/AtmosphereRendering.cpp
+++ b/Engine/Source/Runtime/Renderer/Private/AtmosphereRendering.cpp
@@ -414,7 +414,7 @@ void FDeferredShadingSceneRenderer::RenderAtmosphere(FRHICommandListImmediate& R
 		};
 
 		
-		GSceneRenderTargets.BeginRenderingSceneColor(RHICmdList);
+		GSceneRenderTargets.BeginRenderingSceneColor(RHICmdList, ESimpleRenderTargetMode::EUninitializedColorExistingReadOnlyDepth);
 		for(int32 ViewIndex = 0;ViewIndex < Views.Num();ViewIndex++)
 		{
 			const FViewInfo& View = Views[ViewIndex];
diff --git a/Engine/Source/Runtime/Renderer/Private/CompositionLighting/PostProcessDeferredDecals.cpp b/Engine/Source/Runtime/Renderer/Private/CompositionLighting/PostProcessDeferredDecals.cpp
--- a/Engine/Source/Runtime/Renderer/Private/CompositionLighting/PostProcessDeferredDecals.cpp
+++ b/Engine/Source/Runtime/Renderer/Private/CompositionLighting/PostProcessDeferredDecals.cpp
@@ -371,7 +371,7 @@ void StencilDecalMask(FRHICommandList& RHICmdList, const FViewInfo& View)
 	SCOPED_DRAW_EVENT(RHICmdList, StencilDecalMask);
 	RHICmdList.SetRasterizerState(TStaticRasterizerState<FM_Solid, CM_None>::GetRHI());
 	RHICmdList.SetBlendState(TStaticBlendState<CW_NONE>::GetRHI());
-	SetRenderTarget(RHICmdList, NULL, GSceneRenderTargets.GetSceneDepthSurface());
+	SetRenderTarget(RHICmdList, NULL, GSceneRenderTargets.GetSceneDepthSurface(), ESimpleRenderTargetMode::EUninitializedColorExistingReadOnlyDepth);
 	RHICmdList.SetViewport(View.ViewRect.Min.X, View.ViewRect.Min.Y, 0.0f, View.ViewRect.Max.X, View.ViewRect.Max.Y, 1.0f);
 
 	// Write 1 to highest bit of stencil to areas that should not receive decals
@@ -875,6 +875,20 @@ void FRCPassPostProcessDeferredDecals::Process(FRenderingCompositePassContext& C
 		
 		RHICmdList.SetStreamSource(0, GUnitCubeVertexBuffer.VertexBufferRHI, sizeof(FVector4), 0);
 
+		enum EDecalResolveBufferIndex
+		{
+			SceneColorIndex,
+			GBufferAIndex,
+			GBufferBIndex,
+			GBufferCIndex,
+			DBufferAIndex,
+			DBufferBIndex,
+			DBufferCIndex,
+			ResolveBufferMax,
+		};
+
+		FTextureRHIParamRef TargetsToResolve[ResolveBufferMax] = { nullptr };
+
 		for (int32 DecalIndex = 0; DecalIndex < SortedDecals.Num(); DecalIndex++)
 		{
 			const FTransientDecalRenderData& DecalData = SortedDecals[DecalIndex];
@@ -918,31 +932,32 @@ void FRCPassPostProcessDeferredDecals::Process(FRenderingCompositePassContext& C
 				switch(CurrentRenderTargetMode)
 				{
 					case RTM_SceneColorAndGBuffer:
-						{
-							FTextureRHIParamRef RenderTargets[4];
-							RenderTargets[0] = GSceneRenderTargets.GetSceneColor()->GetRenderTargetItem().TargetableTexture;
-							RenderTargets[1] = GSceneRenderTargets.GBufferA->GetRenderTargetItem().TargetableTexture;
-							RenderTargets[2] = GSceneRenderTargets.GBufferB->GetRenderTargetItem().TargetableTexture;
-							RenderTargets[3] = GSceneRenderTargets.GBufferC->GetRenderTargetItem().TargetableTexture;
-							SetRenderTargets(RHICmdList, ARRAY_COUNT(RenderTargets), RenderTargets, GSceneRenderTargets.GetSceneDepthSurface(), 0, NULL);
+						{							
+							TargetsToResolve[SceneColorIndex] = GSceneRenderTargets.GetSceneColor()->GetRenderTargetItem().TargetableTexture;
+							TargetsToResolve[GBufferAIndex] = GSceneRenderTargets.GBufferA->GetRenderTargetItem().TargetableTexture;
+							TargetsToResolve[GBufferBIndex] = GSceneRenderTargets.GBufferB->GetRenderTargetItem().TargetableTexture;
+							TargetsToResolve[GBufferCIndex] = GSceneRenderTargets.GBufferC->GetRenderTargetItem().TargetableTexture;
+							
+							SetRenderTargets(RHICmdList, 4, TargetsToResolve, GSceneRenderTargets.GetSceneDepthSurface(), ESimpleRenderTargetMode::EExistingColorAndReadOnlyDepth);
 						}
 						break;
 
 					case RTM_GBufferNormal:
-						SetRenderTarget(RHICmdList, GSceneRenderTargets.GBufferA->GetRenderTargetItem().TargetableTexture, GSceneRenderTargets.GetSceneDepthSurface());
+						TargetsToResolve[GBufferAIndex] = GSceneRenderTargets.GBufferA->GetRenderTargetItem().TargetableTexture;
+						SetRenderTarget(RHICmdList, TargetsToResolve[GBufferAIndex], GSceneRenderTargets.GetSceneDepthSurface(), ESimpleRenderTargetMode::EExistingColorAndReadOnlyDepth);
 						break;
 					
 					case RTM_SceneColor:
-						SetRenderTarget(RHICmdList, GSceneRenderTargets.GetSceneColor()->GetRenderTargetItem().TargetableTexture, GSceneRenderTargets.GetSceneDepthSurface());
+						TargetsToResolve[SceneColorIndex] = GSceneRenderTargets.GetSceneColor()->GetRenderTargetItem().TargetableTexture;
+						SetRenderTarget(RHICmdList, TargetsToResolve[SceneColorIndex], GSceneRenderTargets.GetSceneDepthSurface(), ESimpleRenderTargetMode::EExistingColorAndReadOnlyDepth);
 						break;
 
 					case RTM_DBuffer:
-						{
-							FTextureRHIParamRef RenderTargets[3];
-							RenderTargets[0] = GSceneRenderTargets.DBufferA->GetRenderTargetItem().TargetableTexture;
-							RenderTargets[1] = GSceneRenderTargets.DBufferB->GetRenderTargetItem().TargetableTexture;
-							RenderTargets[2] = GSceneRenderTargets.DBufferC->GetRenderTargetItem().TargetableTexture;
-							SetRenderTargets(RHICmdList, ARRAY_COUNT(RenderTargets), RenderTargets, GSceneRenderTargets.GetSceneDepthSurface(), 0, NULL);
+						{							
+							TargetsToResolve[DBufferAIndex] = GSceneRenderTargets.DBufferA->GetRenderTargetItem().TargetableTexture;
+							TargetsToResolve[DBufferBIndex] = GSceneRenderTargets.DBufferB->GetRenderTargetItem().TargetableTexture;
+							TargetsToResolve[DBufferCIndex] = GSceneRenderTargets.DBufferC->GetRenderTargetItem().TargetableTexture;
+							SetRenderTargets(RHICmdList, 3, &TargetsToResolve[DBufferAIndex], GSceneRenderTargets.GetSceneDepthSurface(), ESimpleRenderTargetMode::EExistingColorAndReadOnlyDepth);
 						}
 						break;
 
@@ -1049,12 +1064,19 @@ void FRCPassPostProcessDeferredDecals::Process(FRenderingCompositePassContext& C
 			}
 		}
 
-		// @todo resolve - need to remember the target(s) and resolve them
-
 		// we don't modify stencil but if out input was having stencil for us (after base pass - we need to clear)
-
 		// Clear stencil to 0, which is the assumed default by other passes
 		RHICmdList.Clear(false, FLinearColor::White, false, 0, true, 0, FIntRect());
+
+		// resolve the targets we wrote to.
+		FResolveParams ResolveParams;
+		for (int32 i = 0; i < ResolveBufferMax; ++i)
+		{
+			if (TargetsToResolve[i])
+			{
+				RHICmdList.CopyToResolveTarget(TargetsToResolve[i], TargetsToResolve[i], true, ResolveParams);
+			}
+		}		
 	}
 
 	if(RenderStage == 0)
diff --git a/Engine/Source/Runtime/Renderer/Private/CompositionLighting/PostProcessLpvIndirect.cpp b/Engine/Source/Runtime/Renderer/Private/CompositionLighting/PostProcessLpvIndirect.cpp
--- a/Engine/Source/Runtime/Renderer/Private/CompositionLighting/PostProcessLpvIndirect.cpp
+++ b/Engine/Source/Runtime/Renderer/Private/CompositionLighting/PostProcessLpvIndirect.cpp
@@ -162,7 +162,7 @@ void FRCPassPostProcessLpvIndirect::Process(FRenderingCompositePassContext& Cont
 	};
 
 	// Set the view family's render target/viewport.
-	SetRenderTargets(Context.RHICmdList, 1, RenderTargets, GSceneRenderTargets.GetSceneDepthSurface(), 0, NULL);
+	SetRenderTargets(Context.RHICmdList, 1, RenderTargets, GSceneRenderTargets.GetSceneDepthSurface(), ESimpleRenderTargetMode::EExistingColorAndReadOnlyDepth);
 
 	Context.SetViewportAndCallRHI(View.ViewRect);
 
diff --git a/Engine/Source/Runtime/Renderer/Private/DeferredShadingRenderer.cpp b/Engine/Source/Runtime/Renderer/Private/DeferredShadingRenderer.cpp
--- a/Engine/Source/Runtime/Renderer/Private/DeferredShadingRenderer.cpp
+++ b/Engine/Source/Runtime/Renderer/Private/DeferredShadingRenderer.cpp
@@ -482,6 +482,12 @@ public:
 	{
 		SetStateOnCommandList(ParentCmdList);
 	}
+
+	virtual ~FBasePassParallelCommandListSet()
+	{
+		Dispatch();
+	}
+
 	virtual void SetStateOnCommandList(FRHICommandList& CmdList) override
 	{
 		GSceneRenderTargets.BeginRenderingGBuffer(CmdList, ERenderTargetLoadAction::ELoad, ERenderTargetLoadAction::ELoad);
@@ -817,6 +823,8 @@ void FDeferredShadingSceneRenderer::Render(FRHICommandListImmediate& RHICmdList)
 	GSceneRenderTargets.ResolveSceneColor(RHICmdList, FResolveRect(0, 0, ViewFamily.FamilySizeX, ViewFamily.FamilySizeY));
 	GSceneRenderTargets.ResolveSceneDepthTexture(RHICmdList);
 	GSceneRenderTargets.ResolveSceneDepthToAuxiliaryTexture(RHICmdList);
+
+	GSceneRenderTargets.FinishRenderingGBuffer(RHICmdList);
 	
 	RenderCustomDepthPass(RHICmdList);
 
@@ -893,6 +901,9 @@ void FDeferredShadingSceneRenderer::Render(FRHICommandListImmediate& RHICmdList)
 		TRefCountPtr<IPooledRenderTarget> DynamicBentNormalAO;
 		RenderDynamicSkyLighting(RHICmdList, DynamicBentNormalAO);
 
+		//SSR and SSS need the SceneColor finalized as an SRV.
+		GSceneRenderTargets.FinishRenderingSceneColor(RHICmdList, true);
+
 		// Render reflections that only operate on opaque pixels
 		RenderDeferredReflections(RHICmdList, DynamicBentNormalAO);
 
@@ -1175,9 +1186,15 @@ public:
 	{
 		SetStateOnCommandList(ParentCmdList);
 	}
+
+	virtual ~FPrePassParallelCommandListSet()
+	{
+		Dispatch();
+	}
+
 	virtual void SetStateOnCommandList(FRHICommandList& CmdList) override
 	{
-		GSceneRenderTargets.BeginRenderingPrePass(CmdList);
+		GSceneRenderTargets.BeginRenderingPrePass(CmdList, false);
 		SetupPrePassView(CmdList, View.ViewRect);
 	}
 };
@@ -1225,6 +1242,13 @@ void FDeferredShadingSceneRenderer::RenderPrePassViewParallel(const FViewInfo& V
 	}
 }
 
+static TAutoConsoleVariable<int32> CVarParallelPrePass(
+	TEXT("r.ParallelPrePass"),
+	1,
+	TEXT("Toggles parallel zprepass rendering. Parallel rendering must be enabled for this to have an effect."),
+	ECVF_RenderThreadSafe
+	);
+
 /** Renders the scene's prepass and occlusion queries */
 bool FDeferredShadingSceneRenderer::RenderPrePass(FRHICommandListImmediate& RHICmdList)
 {
@@ -1233,16 +1257,12 @@ bool FDeferredShadingSceneRenderer::RenderPrePass(FRHICommandListImmediate& RHIC
 
 	bool bDirty = false;
 
-	GSceneRenderTargets.BeginRenderingPrePass(RHICmdList);
-
-	// Clear the depth buffer.
-	// Note, this is a reversed Z depth surface, so 0.0f is the far plane.
-	RHICmdList.Clear(false,FLinearColor::Black,true,0.0f,true,0, FIntRect());
+	GSceneRenderTargets.BeginRenderingPrePass(RHICmdList, true);
 
 	// Draw a depth pass to avoid overdraw in the other passes.
 	if(EarlyZPassMode != DDM_None)
 	{
-		if (GRHICommandList.UseParallelAlgorithms())
+		if (GRHICommandList.UseParallelAlgorithms() && CVarParallelPrePass.GetValueOnRenderThread())
 		{
 			FScopedCommandListWaitForTasks Flusher(RHICmdList);
 
@@ -1269,6 +1289,13 @@ bool FDeferredShadingSceneRenderer::RenderPrePass(FRHICommandListImmediate& RHIC
 	return bDirty;
 }
 
+static TAutoConsoleVariable<int32> CVarParallelBasePass(
+	TEXT("r.ParallelBasePass"),
+	1,
+	TEXT("Toggles parallel base pass rendering. Parallel rendering must be enabled for this to have an effect."),
+	ECVF_RenderThreadSafe
+	);
+
 /**
  * Renders the scene's base pass 
  * @return true if anything was rendered
@@ -1292,7 +1319,7 @@ bool FDeferredShadingSceneRenderer::RenderBasePass(FRHICommandListImmediate& RHI
 		SCOPED_DRAW_EVENT(RHICmdList, BasePass);
 		SCOPE_CYCLE_COUNTER(STAT_BasePassDrawTime);
 
-		if (GRHICommandList.UseParallelAlgorithms())
+		if (GRHICommandList.UseParallelAlgorithms() && CVarParallelBasePass.GetValueOnRenderThread())
 		{
 			FScopedCommandListWaitForTasks Flusher(RHICmdList);
 			for(int32 ViewIndex = 0;ViewIndex < Views.Num();ViewIndex++)
diff --git a/Engine/Source/Runtime/Renderer/Private/DistortionRendering.cpp b/Engine/Source/Runtime/Renderer/Private/DistortionRendering.cpp
--- a/Engine/Source/Runtime/Renderer/Private/DistortionRendering.cpp
+++ b/Engine/Source/Runtime/Renderer/Private/DistortionRendering.cpp
@@ -786,7 +786,7 @@ void FSceneRenderer::RenderDistortion(FRHICommandListImmediate& RHICmdList)
 		// DistortionRT==0 should never happen but better we don't crash
 		if(DistortionRT)
 		{
-			SetRenderTarget(RHICmdList, DistortionRT->GetRenderTargetItem().TargetableTexture, GSceneRenderTargets.GetSceneDepthSurface());
+			SetRenderTarget(RHICmdList, DistortionRT->GetRenderTargetItem().TargetableTexture, GSceneRenderTargets.GetSceneDepthSurface(), ESimpleRenderTargetMode::EUninitializedColorExistingReadOnlyDepth);
 
 			for(int32 ViewIndex = 0;ViewIndex < Views.Num();ViewIndex++)
 			{
@@ -814,7 +814,7 @@ void FSceneRenderer::RenderDistortion(FRHICommandListImmediate& RHICmdList)
 			{
 				// restore default stencil state
 				// Note, this is a reversed Z depth surface, using CF_GreaterEqual.
-				RHICmdList.SetDepthStencilState(TStaticDepthStencilState<true, CF_GreaterEqual>::GetRHI());
+				RHICmdList.SetDepthStencilState(TStaticDepthStencilState<false, CF_GreaterEqual>::GetRHI());
 
 				// resolve using the current ResolveParams 
 				RHICmdList.CopyToResolveTarget(DistortionRT->GetRenderTargetItem().TargetableTexture, DistortionRT->GetRenderTargetItem().ShaderResourceTexture, false, FResolveParams());
diff --git a/Engine/Source/Runtime/Renderer/Private/FogRendering.cpp b/Engine/Source/Runtime/Renderer/Private/FogRendering.cpp
--- a/Engine/Source/Runtime/Renderer/Private/FogRendering.cpp
+++ b/Engine/Source/Runtime/Renderer/Private/FogRendering.cpp
@@ -308,7 +308,7 @@ bool FDeferredShadingSceneRenderer::RenderFog(FRHICommandListImmediate& RHICmdLi
 			0, 2, 3
 		};
 
-		GSceneRenderTargets.BeginRenderingSceneColor(RHICmdList);
+		GSceneRenderTargets.BeginRenderingSceneColor(RHICmdList, ESimpleRenderTargetMode::EUninitializedColorExistingReadOnlyDepth);
 		for(int32 ViewIndex = 0;ViewIndex < Views.Num();ViewIndex++)
 		{
 			const FViewInfo& View = Views[ViewIndex];
diff --git a/Engine/Source/Runtime/Renderer/Private/LightRendering.cpp b/Engine/Source/Runtime/Renderer/Private/LightRendering.cpp
--- a/Engine/Source/Runtime/Renderer/Private/LightRendering.cpp
+++ b/Engine/Source/Runtime/Renderer/Private/LightRendering.cpp
@@ -425,14 +425,14 @@ void FDeferredShadingSceneRenderer::RenderLights(FRHICommandListImmediate& RHICm
 			}
 			else if (SimpleLights.InstanceData.Num() > 0)
 			{
-				GSceneRenderTargets.BeginRenderingSceneColor(RHICmdList);
+				GSceneRenderTargets.BeginRenderingSceneColor(RHICmdList, ESimpleRenderTargetMode::EUninitializedColorExistingReadOnlyDepth);
 				RenderSimpleLightsStandardDeferred(RHICmdList, SimpleLights);
 			}
 
 			{
 				SCOPED_DRAW_EVENT(RHICmdList, StandardDeferredLighting);
 
-				GSceneRenderTargets.BeginRenderingSceneColor(RHICmdList);
+				GSceneRenderTargets.BeginRenderingSceneColor(RHICmdList, ESimpleRenderTargetMode::EUninitializedColorExistingReadOnlyDepth);
 
 				// Draw non-shadowed non-light function lights without changing render targets between them
 				for (int32 LightIndex = StandardDeferredStart; LightIndex < AttenuationLightStart; LightIndex++)
@@ -489,8 +489,7 @@ void FDeferredShadingSceneRenderer::RenderLights(FRHICommandListImmediate& RHICm
 				INC_DWORD_STAT(STAT_NumShadowedLights);
 
 				// All shadows render with min blending
-				GSceneRenderTargets.BeginRenderingLightAttenuation(RHICmdList);
-				RHICmdList.Clear(true, FLinearColor::White, false, 0, false, 0, FIntRect());
+				GSceneRenderTargets.BeginRenderingLightAttenuation(RHICmdList, true);
 
 				bool bRenderedTranslucentObjectShadows = RenderTranslucentProjectedShadows(RHICmdList, &LightSceneInfo );
 				// Render non-modulated projected shadows to the attenuation buffer.
@@ -542,7 +541,7 @@ void FDeferredShadingSceneRenderer::RenderLights(FRHICommandListImmediate& RHICm
 			}
 
 			GSceneRenderTargets.SetLightAttenuationMode(bUsedLightAttenuation);
-			GSceneRenderTargets.BeginRenderingSceneColor(RHICmdList);
+			GSceneRenderTargets.BeginRenderingSceneColor(RHICmdList, ESimpleRenderTargetMode::EUninitializedColorExistingReadOnlyDepth);
 
 			// Render the light to the scene color buffer, conditionally using the attenuation buffer or a 1x1 white texture as input 
 			if(bDirectLighting)
@@ -630,7 +629,7 @@ void FDeferredShadingSceneRenderer::RenderStationaryLightOverlap(FRHICommandList
 {
 	if (Scene->bIsEditorScene)
 	{
-		GSceneRenderTargets.BeginRenderingSceneColor(RHICmdList);
+		GSceneRenderTargets.BeginRenderingSceneColor(RHICmdList, ESimpleRenderTargetMode::EUninitializedColorExistingReadOnlyDepth);
 
 		// Clear to discard base pass values in scene color since we didn't skip that, to have valid scene depths
 		RHICmdList.Clear(true, FLinearColor::Black, false, 0, false, 0, FIntRect());
diff --git a/Engine/Source/Runtime/Renderer/Private/LightShaftRendering.cpp b/Engine/Source/Runtime/Renderer/Private/LightShaftRendering.cpp
--- a/Engine/Source/Runtime/Renderer/Private/LightShaftRendering.cpp
+++ b/Engine/Source/Runtime/Renderer/Private/LightShaftRendering.cpp
@@ -838,7 +838,7 @@ void ApplyLightShaftBloom(FRHICommandListImmediate& RHICmdList, const FViewInfo&
 {
 	SCOPED_DRAW_EVENT(RHICmdList, Apply);
 
-	GSceneRenderTargets.BeginRenderingSceneColor(RHICmdList);
+	GSceneRenderTargets.BeginRenderingSceneColor(RHICmdList, ESimpleRenderTargetMode::EUninitializedColorExistingReadOnlyDepth);
 
 	RHICmdList.SetViewport(View.ViewRect.Min.X, View.ViewRect.Min.Y, 0.0f, View.ViewRect.Max.X, View.ViewRect.Max.Y, 1.0f);
 	RHICmdList.SetBlendState(TStaticBlendState<CW_RGB, BO_Add, BF_One, BF_One>::GetRHI());
diff --git a/Engine/Source/Runtime/Renderer/Private/PostProcess/PostProcessTemporalAA.cpp b/Engine/Source/Runtime/Renderer/Private/PostProcess/PostProcessTemporalAA.cpp
--- a/Engine/Source/Runtime/Renderer/Private/PostProcess/PostProcessTemporalAA.cpp
+++ b/Engine/Source/Runtime/Renderer/Private/PostProcess/PostProcessTemporalAA.cpp
@@ -549,7 +549,7 @@ void FRCPassPostProcessTemporalAA::Process(FRenderingCompositePassContext& Conte
 	const FSceneRenderTargetItem& DestRenderTarget = PassOutputs[0].RequestSurface(Context);
 
 	//Context.SetRenderTarget(RHICmdList, DestRenderTarget.TargetableTexture, FTextureRHIRef());
-	SetRenderTarget(Context.RHICmdList, DestRenderTarget.TargetableTexture, GSceneRenderTargets.GetSceneDepthTexture());
+	SetRenderTarget(Context.RHICmdList, DestRenderTarget.TargetableTexture, GSceneRenderTargets.GetSceneDepthTexture(), ESimpleRenderTargetMode::EUninitializedColorExistingReadOnlyDepth);
 
 	// is optimized away if possible (RT size=view size, )
 	Context.RHICmdList.Clear(true, FLinearColor::Black, false, 1.0f, false, 0, SrcRect);
diff --git a/Engine/Source/Runtime/Renderer/Private/PostProcess/SceneRenderTargets.cpp b/Engine/Source/Runtime/Renderer/Private/PostProcess/SceneRenderTargets.cpp
--- a/Engine/Source/Runtime/Renderer/Private/PostProcess/SceneRenderTargets.cpp
+++ b/Engine/Source/Runtime/Renderer/Private/PostProcess/SceneRenderTargets.cpp
@@ -286,6 +286,37 @@ void FSceneRenderTargets::BeginRenderingSceneColor(FRHICommandList& RHICmdList,
 	SetRenderTarget(RHICmdList, GetSceneColorSurface(), GetSceneDepthSurface(), RenderTargetMode);
 } 
 
+int32 FSceneRenderTargets::GetGBufferRenderTargets(FRHIRenderTargetView RenderTargets[MAX_CONCURRENT_GBUFFERS], ERenderTargetLoadAction ColorLoadAction, int32& VelocityRTIndex)
+{
+	const bool bUseVelocityGBuffer = FVelocityRendering::OutputsToGBuffer();
+
+	RenderTargets[0] = FRHIRenderTargetView(GetSceneColorSurface(), 0, -1, ColorLoadAction, ERenderTargetStoreAction::EStore);
+	RenderTargets[1] = FRHIRenderTargetView(GSceneRenderTargets.GBufferA->GetRenderTargetItem().TargetableTexture, 0, -1, ColorLoadAction, ERenderTargetStoreAction::EStore);
+	RenderTargets[2] = FRHIRenderTargetView(GSceneRenderTargets.GBufferB->GetRenderTargetItem().TargetableTexture, 0, -1, ColorLoadAction, ERenderTargetStoreAction::EStore);
+	RenderTargets[3] = FRHIRenderTargetView(GSceneRenderTargets.GBufferC->GetRenderTargetItem().TargetableTexture, 0, -1, ColorLoadAction, ERenderTargetStoreAction::EStore);
+	RenderTargets[4] = FRHIRenderTargetView(GSceneRenderTargets.GBufferD->GetRenderTargetItem().TargetableTexture, 0, -1, ColorLoadAction, ERenderTargetStoreAction::EStore);
+	int32 MRTCount = 5;	// Derived from previously used RenderTargets[]; would have been nice to keep a pointer while filling it but it's more confusing to use
+
+	VelocityRTIndex = -1;
+
+	if (bAllowStaticLighting)
+	{
+		RenderTargets[MRTCount] = FRHIRenderTargetView(GSceneRenderTargets.GBufferE->GetRenderTargetItem().TargetableTexture, 0, -1, ColorLoadAction, ERenderTargetStoreAction::EStore);
+		++MRTCount;
+	}
+
+	if (bUseVelocityGBuffer)
+	{
+		VelocityRTIndex = MRTCount;
+		++MRTCount;
+		RenderTargets[VelocityRTIndex] = FRHIRenderTargetView(GSceneRenderTargets.GBufferVelocity->GetRenderTargetItem().TargetableTexture, 0, -1, ColorLoadAction, ERenderTargetStoreAction::EStore);
+	}
+
+	check(MRTCount < MAX_CONCURRENT_GBUFFERS);
+
+	return MRTCount;
+}
+
 void FSceneRenderTargets::BeginRenderingGBuffer(FRHICommandList& RHICmdList, ERenderTargetLoadAction ColorLoadAction, ERenderTargetLoadAction DepthLoadAction, const FLinearColor& ClearColor/*=(0,0,0,1)*/)
 {
 	SCOPED_DRAW_EVENT(RHICmdList, BeginRenderingSceneColor);
@@ -302,53 +333,70 @@ void FSceneRenderTargets::BeginRenderingGBuffer(FRHICommandList& RHICmdList, ERe
 	// Set the scene color surface as the render target, and the scene depth surface as the depth-stencil target.
 	if (CurrentFeatureLevel >= ERHIFeatureLevel::SM4)
 	{
-		const bool bUseVelocityGBuffer = FVelocityRendering::OutputsToGBuffer();
+		int32 VelocityRTIndex;
 		FRHIRenderTargetView RenderTargets[7];
-		RenderTargets[0] = FRHIRenderTargetView(GetSceneColorSurface(), 0, -1, ColorLoadAction, ERenderTargetStoreAction::EStore);
-		RenderTargets[1] = FRHIRenderTargetView(GSceneRenderTargets.GBufferA->GetRenderTargetItem().TargetableTexture, 0, -1, ColorLoadAction, ERenderTargetStoreAction::EStore);
-		RenderTargets[2] = FRHIRenderTargetView(GSceneRenderTargets.GBufferB->GetRenderTargetItem().TargetableTexture, 0, -1, ColorLoadAction, ERenderTargetStoreAction::EStore);
-		RenderTargets[3] = FRHIRenderTargetView(GSceneRenderTargets.GBufferC->GetRenderTargetItem().TargetableTexture, 0, -1, ColorLoadAction, ERenderTargetStoreAction::EStore);
-		RenderTargets[4] = FRHIRenderTargetView(GSceneRenderTargets.GBufferD->GetRenderTargetItem().TargetableTexture, 0, -1, ColorLoadAction, ERenderTargetStoreAction::EStore);
-		int32 MRTCount = 5;	// Derived from previously used RenderTargets[]; would have been nice to keep a pointer while filling it but it's more confusing to use
-
-		int32 VelocityRTIndex = -1;
+		int32 MRTCount = GetGBufferRenderTargets(RenderTargets, ColorLoadAction, VelocityRTIndex);
 
-		if (bAllowStaticLighting)
+		const float DepthClearValue = 0.0f;
+		bool bClearColor = ColorLoadAction == ERenderTargetLoadAction::EClear;
+		bool bClearDepth = DepthLoadAction == ERenderTargetLoadAction::EClear;
+		FRHIDepthRenderTargetView DepthView(GetSceneDepthSurface(), DepthLoadAction, ERenderTargetStoreAction::EStore);
+		FRHISetRenderTargetsInfo Info(MRTCount, RenderTargets, DepthView);
+		if (bClearColor)
 		{
-			RenderTargets[MRTCount] = FRHIRenderTargetView(GSceneRenderTargets.GBufferE->GetRenderTargetItem().TargetableTexture, 0, -1, ColorLoadAction, ERenderTargetStoreAction::EStore);
-			++MRTCount;
+			bGBuffersCleared = true;
+			GBufferClearColor = ClearColor;
 		}
 
-		if (bUseVelocityGBuffer)
+		Info.ClearColors[0] = GBufferClearColor;
+		Info.ClearColors[1] = Info.ClearColors[2] = Info.ClearColors[3] = FLinearColor::Transparent;
+		Info.ClearColors[4] = FLinearColor(0, 1, 1, 1);
+		Info.ClearColors[5] = FLinearColor(1, 1, 1, 1);
+
+		if (VelocityRTIndex != -1)
 		{
-			VelocityRTIndex = MRTCount;
-			++MRTCount;
-			RenderTargets[VelocityRTIndex] = FRHIRenderTargetView(GSceneRenderTargets.GBufferVelocity->GetRenderTargetItem().TargetableTexture, 0, -1, ColorLoadAction, ERenderTargetStoreAction::EStore);
+			Info.ClearColors[VelocityRTIndex] = FLinearColor(0, 0, 0, 0);
 		}
 
-		check(MRTCount <= ARRAY_COUNT(RenderTargets));
-
-		FRHIDepthRenderTargetView DepthView(GetSceneDepthSurface(), DepthLoadAction, ERenderTargetStoreAction::EStore);
-		FRHISetRenderTargetsInfo Info(MRTCount, RenderTargets, DepthView);
-		if (ColorLoadAction == ERenderTargetLoadAction::EClear)
+		if (bClearDepth)
 		{
-			Info.ClearColors[0] = ClearColor;
-			Info.ClearColors[1] = Info.ClearColors[2] = Info.ClearColors[3] = FLinearColor::Transparent;
-			Info.ClearColors[4] = FLinearColor(0, 1, 1, 1);
-			Info.ClearColors[5] = FLinearColor(1, 1, 1, 1);
-
-			if (bUseVelocityGBuffer)
-			{
-				Info.ClearColors[VelocityRTIndex] = FLinearColor(0, 0, 0, 0);
-			}
+			bSceneDepthCleared = true;
+			SceneDepthClearValue = DepthClearValue;
 		}
-		Info.DepthClearValue = 0.0f;
+		Info.DepthClearValue = DepthClearValue;
 
 		// set the render target
 		RHICmdList.SetRenderTargetsAndClear(Info);
+
+		//bind any clear data that won't be bound automatically by the preceding SetRenderTargetsAndClear
+		bool bBindClearColor = !bClearColor && bGBuffersCleared;
+		bool bBindClearDepth = !bClearDepth && bSceneDepthCleared;
+		RHICmdList.BindClearMRTValues(bBindClearColor, MRTCount, Info.ClearColors, bBindClearDepth, SceneDepthClearValue, bBindClearDepth, Info.StencilClearValue);
 	}
 }
 
+void FSceneRenderTargets::FinishRenderingGBuffer(FRHICommandListImmediate& RHICmdList)
+{
+	if (IsSimpleDynamicLightingEnabled())
+	{
+		// in this non-standard case, just render to scene color with default mode
+		FinishRenderingSceneColor(RHICmdList, true);
+		return;
+	}
+
+	int32 VelocityRTIndex;
+	FRHIRenderTargetView RenderTargets[7];
+	int32 NumMRTs = GetGBufferRenderTargets(RenderTargets, ERenderTargetLoadAction::ELoad, VelocityRTIndex);
+
+	FResolveParams ResolveParams;
+	for (int32 i = 0; i < NumMRTs; ++i)
+	{
+		RHICmdList.CopyToResolveTarget(RenderTargets[i].Texture, RenderTargets[i].Texture, true, ResolveParams);
+	}
+	FTextureRHIParamRef DepthSurface = GetSceneDepthSurface();
+	RHICmdList.CopyToResolveTarget(DepthSurface, DepthSurface, true, ResolveParams);
+}
+
 int32 FSceneRenderTargets::GetNumGBufferTargets() const
 {
 	int32 NumGBufferTargets = 1;
@@ -820,14 +868,38 @@ void FSceneRenderTargets::ResolveGBufferSurfaces(FRHICommandList& RHICmdList, co
 	}
 }
 
-void FSceneRenderTargets::BeginRenderingPrePass(FRHICommandList& RHICmdList)
+void FSceneRenderTargets::BeginRenderingPrePass(FRHICommandList& RHICmdList, bool bPerformClear)
 {
 	SCOPED_DRAW_EVENT(RHICmdList, BeginRenderingPrePass);
 
-	// Set the scene depth surface and a DUMMY buffer as color buffer
-	// (as long as it's the same dimension as the depth buffer),
-	
-	SetRenderTarget(RHICmdList, FTextureRHIRef(), GetSceneDepthSurface());
+	FTextureRHIRef ColorTarget;
+	FTexture2DRHIRef DepthTarget = GetSceneDepthSurface();
+
+	const float DepthClearValue = 0.0f;
+	const uint32 StencilClearValue = 0;
+	if (bPerformClear)
+	{				
+		FRHIRenderTargetView ColorView(ColorTarget, 0, -1, ERenderTargetLoadAction::ENoAction, ERenderTargetStoreAction::ENoAction);
+		FRHIDepthRenderTargetView DepthView(DepthTarget, ERenderTargetLoadAction::EClear, ERenderTargetStoreAction::EStore);
+
+		// Clear the depth buffer.
+		// Note, this is a reversed Z depth surface, so 0.0f is the far plane.
+		FRHISetRenderTargetsInfo Info(1, &ColorView, DepthView);
+		Info.DepthClearValue = DepthClearValue;
+		Info.StencilClearValue = StencilClearValue;
+
+		RHICmdList.SetRenderTargetsAndClear(Info);
+
+		bSceneDepthCleared = true;
+		SceneDepthClearValue = DepthClearValue;
+	}
+	else
+	{
+		// Set the scene depth surface and a DUMMY buffer as color buffer
+		// (as long as it's the same dimension as the depth buffer),	
+		SetRenderTarget(RHICmdList, FTextureRHIRef(), GetSceneDepthSurface());
+		RHICmdList.BindClearMRTValues(false, 0, nullptr, true, DepthClearValue, true, StencilClearValue);
+	}
 }
 
 void FSceneRenderTargets::FinishRenderingPrePass(FRHICommandListImmediate& RHICmdList)
@@ -900,10 +972,10 @@ void FSceneRenderTargets::FinishRenderingReflectiveShadowMap(FRHICommandList& RH
 	SetRenderTargets(RHICmdList, ARRAY_COUNT(RenderTargets), RenderTargets, FTextureRHIParamRef(), 2, Uavs);
 }
 
-void FSceneRenderTargets::FinishRenderingCubeShadowDepth(FRHICommandList& RHICmdList, int32 ShadowResolution, const FResolveParams& ResolveParams)
+void FSceneRenderTargets::FinishRenderingCubeShadowDepth(FRHICommandList& RHICmdList, int32 ShadowResolution)
 {
 	SCOPED_DRAW_EVENT(RHICmdList, FinishRenderingCubeShadowDepth);
-	RHICmdList.CopyToResolveTarget(GetCubeShadowDepthZSurface(ShadowResolution), GetCubeShadowDepthZTexture(ShadowResolution), false, ResolveParams);
+	RHICmdList.CopyToResolveTarget(GetCubeShadowDepthZSurface(ShadowResolution), GetCubeShadowDepthZTexture(ShadowResolution), false, FResolveParams(FResolveRect(), CubeFace_PosX, -1, -1, 0));
 }
 
 void FSceneRenderTargets::BeginRenderingSceneAlphaCopy(FRHICommandListImmediate& RHICmdList)
@@ -930,7 +1002,7 @@ void FSceneRenderTargets::BeginRenderingLightAttenuation(FRHICommandList& RHICmd
 	GRenderTargetPool.VisualizeTexture.SetCheckPoint(RHICmdList, GSceneRenderTargets.GetLightAttenuation());
 
 	// Set the light attenuation surface as the render target, and the scene depth buffer as the depth-stencil surface.
-	SetRenderTarget(RHICmdList, GetLightAttenuationSurface(), GetSceneDepthSurface(), bClearToWhite ? ESimpleRenderTargetMode::EClearColorToWhite : ESimpleRenderTargetMode::EExistingColorAndDepth);
+	SetRenderTarget(RHICmdList, GetLightAttenuationSurface(), GetSceneDepthSurface(), bClearToWhite ? ESimpleRenderTargetMode::EClearColorToWhiteReadOnlyDepth : ESimpleRenderTargetMode::EExistingColorAndReadOnlyDepth);
 }
 
 void FSceneRenderTargets::FinishRenderingLightAttenuation(FRHICommandList& RHICmdList)
@@ -946,12 +1018,17 @@ void FSceneRenderTargets::FinishRenderingLightAttenuation(FRHICommandList& RHICm
 void FSceneRenderTargets::BeginRenderingTranslucency(FRHICommandList& RHICmdList, const FViewInfo& View)
 {
 	// Use the scene color buffer.
-	GSceneRenderTargets.BeginRenderingSceneColor(RHICmdList, ESimpleRenderTargetMode::EExistingColorAndDepth);
+	GSceneRenderTargets.BeginRenderingSceneColor(RHICmdList, ESimpleRenderTargetMode::EExistingColorAndReadOnlyDepth);
 	
 	// viewport to match view size
 	RHICmdList.SetViewport(View.ViewRect.Min.X, View.ViewRect.Min.Y, 0.0f, View.ViewRect.Max.X, View.ViewRect.Max.Y, 1.0f);
 }
 
+void FSceneRenderTargets::FinishRenderingTranslucency(FRHICommandListImmediate& RHICmdList, const class FViewInfo& View)
+{
+	GSceneRenderTargets.FinishRenderingSceneColor(RHICmdList, true);
+}
+
 
 bool FSceneRenderTargets::BeginRenderingSeparateTranslucency(FRHICommandList& RHICmdList, const FViewInfo& View, bool bFirstTimeThisFrame)
 {
@@ -971,7 +1048,13 @@ bool FSceneRenderTargets::BeginRenderingSeparateTranslucency(FRHICommandList& RH
 
 		// clear the render target the first time, re-use afterwards
 		SetRenderTarget(RHICmdList, SeparateTranslucency->GetRenderTargetItem().TargetableTexture, GetSceneDepthSurface(), 
-			bFirstTimeThisFrame ? ESimpleRenderTargetMode::EClearColorToBlackWithFullAlpha : ESimpleRenderTargetMode::EExistingColorAndDepth);
+			bFirstTimeThisFrame ? ESimpleRenderTargetMode::EClearColorToBlackWithFullAlphaReadOnlyDepth : ESimpleRenderTargetMode::EExistingColorAndReadOnlyDepth);
+
+		if (!bFirstTimeThisFrame)
+		{
+			FLinearColor ClearColor(0.0f, 0.0f, 0.0f, 1.0f);
+			RHICmdList.BindClearMRTValues(true, 1, &ClearColor, false, 0.0f, false, 0);
+		}
 
 		RHICmdList.SetViewport(View.ViewRect.Min.X, View.ViewRect.Min.Y, 0.0f, View.ViewRect.Max.X, View.ViewRect.Max.Y, 1.0f);
 
@@ -985,6 +1068,8 @@ void FSceneRenderTargets::FinishRenderingSeparateTranslucency(FRHICommandList& R
 {
 	if(IsSeparateTranslucencyActive(View))
 	{
+		SCOPED_DRAW_EVENT(RHICmdList, FinishSeparateTranslucency);
+
 		FSceneViewState* ViewState = (FSceneViewState*)View.State;
 		TRefCountPtr<IPooledRenderTarget>& SeparateTranslucency = ViewState->GetSeparateTranslucency(View);
 
diff --git a/Engine/Source/Runtime/Renderer/Private/PostProcess/SceneRenderTargets.h b/Engine/Source/Runtime/Renderer/Private/PostProcess/SceneRenderTargets.h
--- a/Engine/Source/Runtime/Renderer/Private/PostProcess/SceneRenderTargets.h
+++ b/Engine/Source/Runtime/Renderer/Private/PostProcess/SceneRenderTargets.h
@@ -162,11 +162,15 @@ public:
 	 *
 	 */
 	void SetBufferSize(int32 InBufferSizeX, int32 InBufferSizeY);
+
+	void BeginRenderingGBuffer(FRHICommandList& RHICmdList, ERenderTargetLoadAction ColorLoadAction, ERenderTargetLoadAction DepthLoadAction, const FLinearColor& ClearColor = FLinearColor(0, 0, 0, 1));
+	void FinishRenderingGBuffer(FRHICommandListImmediate& RHICmdList);
+
 	/**
 	 * Sets the scene color target and restores its contents if necessary
 	 */
 	void BeginRenderingSceneColor(FRHICommandList& RHICmdList, ESimpleRenderTargetMode RenderTargetMode=ESimpleRenderTargetMode::EUninitializedColorExistingDepth);
-	void BeginRenderingGBuffer(FRHICommandList& RHICmdList, ERenderTargetLoadAction ColorLoadAction, ERenderTargetLoadAction DepthLoadAction, const FLinearColor& ClearColor=FLinearColor(0,0,0,1));
+	
 	/**
 	 * Called when finished rendering to the scene color surface
 	 * @param bKeepChanges - if true then the SceneColorSurface is resolved to the SceneColorTexture
@@ -201,9 +205,10 @@ public:
 	void FinishRenderingReflectiveShadowMap(FRHICommandList& RHICmdList, const FResolveRect& ResolveRect = FResolveRect());
 
 	/** Resolves the appropriate shadow depth cube map and restores default state. */
-	void FinishRenderingCubeShadowDepth(FRHICommandList& RHICmdList, int32 ShadowResolution, const FResolveParams& ResolveParams = FResolveParams());
+	void FinishRenderingCubeShadowDepth(FRHICommandList& RHICmdList, int32 ShadowResolution);
 	
 	void BeginRenderingTranslucency(FRHICommandList& RHICmdList, const class FViewInfo& View);
+	void FinishRenderingTranslucency(FRHICommandListImmediate& RHICmdList, const class FViewInfo& View);
 
 	bool BeginRenderingSeparateTranslucency(FRHICommandList& RHICmdList, const FViewInfo& View, bool bFirstTimeThisFrame);
 	void FinishRenderingSeparateTranslucency(FRHICommandList& RHICmdList, const FViewInfo& View);
@@ -212,7 +217,7 @@ public:
 	void ResolveSceneDepthTexture(FRHICommandList& RHICmdList);
 	void ResolveSceneDepthToAuxiliaryTexture(FRHICommandList& RHICmdList);
 
-	void BeginRenderingPrePass(FRHICommandList& RHICmdList);
+	void BeginRenderingPrePass(FRHICommandList& RHICmdList, bool bPerformClear);
 	void FinishRenderingPrePass(FRHICommandListImmediate& RHICmdList);
 
 	void BeginRenderingSceneAlphaCopy(FRHICommandListImmediate& RHICmdList);
@@ -582,6 +587,10 @@ private:
 	/** Determine whether the render targets for any shading path have been allocated */
 	bool AreAnyShadingPathRenderTargetsAllocated() const { return AreShadingPathRenderTargetsAllocated(EShadingPath::Deferred) || AreShadingPathRenderTargetsAllocated(EShadingPath::Forward); }
 
+	static const int32  MAX_CONCURRENT_GBUFFERS = 7;
+	/** Gets all GBuffers to use.  Returns the number actually used. */
+	int32 GetGBufferRenderTargets(FRHIRenderTargetView RenderTargets[MAX_CONCURRENT_GBUFFERS], ERenderTargetLoadAction ColorLoadAction, int32 & VelocityRTIndex);
+
 private:
 	/** Uniform buffer containing GBuffer resources. */
 	FUniformBufferRHIRef GBufferResourcesUniformBuffer;
@@ -615,6 +624,14 @@ private:
 	ERHIFeatureLevel::Type CurrentFeatureLevel;
 	/** Shading path that we are currently drawing through. Set when calling Allocate at the start of a scene render. */
 	EShadingPath CurrentShadingPath;
+
+	/** Helpers to track gbuffer state on platforms that need to propagate clear information across parallel rendering boundaries. */
+	bool bGBuffersCleared;
+	FLinearColor GBufferClearColor;
+
+	/** Helpers to track scenedepth state on platforms that need to propagate clear information across parallel rendering boundaries. */
+	bool bSceneDepthCleared;
+	float SceneDepthClearValue;
 };
 
 /** The global render targets used for scene rendering. */
diff --git a/Engine/Source/Runtime/Renderer/Private/PostProcess/VisualizeTexture.cpp b/Engine/Source/Runtime/Renderer/Private/PostProcess/VisualizeTexture.cpp
--- a/Engine/Source/Runtime/Renderer/Private/PostProcess/VisualizeTexture.cpp
+++ b/Engine/Source/Runtime/Renderer/Private/PostProcess/VisualizeTexture.cpp
@@ -223,9 +223,10 @@ template<uint32 TextureType> void VisualizeTextureForTextureType(FRHICommandList
 
 void RenderVisualizeTexture(FRHICommandListImmediate& RHICmdList, ERHIFeatureLevel::Type FeatureLevel, const FVisualizeTextureData& Data)
 {
+	RHICmdList.CopyToResolveTarget(Data.RenderTargetItem.ShaderResourceTexture, Data.RenderTargetItem.ShaderResourceTexture, true, FResolveParams());
 	if(Data.Desc.Is2DTexture())
 	{
-		// 2D
+		// 2D		
 		if(Data.Desc.NumSamples > 1)
 		{
 			// MSAA
@@ -238,12 +239,11 @@ void RenderVisualizeTexture(FRHICommandListImmediate& RHICmdList, ERHIFeatureLev
 		}
 	}
 	else if(Data.Desc.Is3DTexture())
-	{
-		// Volume
+	{		
 		VisualizeTextureForTextureType<3>(RHICmdList, FeatureLevel, Data);
 	}
 	else if(Data.Desc.IsCubemap())
-	{
+	{		
 		if(Data.Desc.IsArray())
 		{
 			// Cube[]
diff --git a/Engine/Source/Runtime/Renderer/Private/ReflectionEnvironment.cpp b/Engine/Source/Runtime/Renderer/Private/ReflectionEnvironment.cpp
--- a/Engine/Source/Runtime/Renderer/Private/ReflectionEnvironment.cpp
+++ b/Engine/Source/Runtime/Renderer/Private/ReflectionEnvironment.cpp
@@ -648,7 +648,7 @@ IMPLEMENT_SHADER_TYPE(template<>,TStandardDeferredReflectionPS<false>,TEXT("Refl
 
 void FDeferredShadingSceneRenderer::RenderReflectionCaptureSpecularBounceForAllViews(FRHICommandListImmediate& RHICmdList)
 {
-	GSceneRenderTargets.BeginRenderingSceneColor(RHICmdList);
+	GSceneRenderTargets.BeginRenderingSceneColor(RHICmdList, ESimpleRenderTargetMode::EUninitializedColorExistingReadOnlyDepth);
 	RHICmdList.SetRasterizerState(TStaticRasterizerState< FM_Solid, CM_None >::GetRHI());
 	RHICmdList.SetDepthStencilState(TStaticDepthStencilState< false, CF_Always >::GetRHI());
 	RHICmdList.SetBlendState(TStaticBlendState< CW_RGB, BO_Add, BF_One, BF_One >::GetRHI());
@@ -681,6 +681,8 @@ void FDeferredShadingSceneRenderer::RenderReflectionCaptureSpecularBounceForAllV
 			*VertexShader,
 			EDRF_UseTriangleOptimization);
 	}
+
+	GSceneRenderTargets.FinishRenderingSceneColor(RHICmdList);
 }
 
 bool FDeferredShadingSceneRenderer::ShouldDoReflectionEnvironment() const
@@ -916,7 +918,7 @@ void FDeferredShadingSceneRenderer::RenderStandardDeferredImageBasedReflections(
 			// Apply reflections to screen
 			SCOPED_DRAW_EVENT(RHICmdList, ReflectionApply);
 
-			GSceneRenderTargets.BeginRenderingSceneColor(RHICmdList);
+			GSceneRenderTargets.BeginRenderingSceneColor(RHICmdList, ESimpleRenderTargetMode::EUninitializedColorExistingReadOnlyDepth);
 
 			RHICmdList.SetViewport(View.ViewRect.Min.X, View.ViewRect.Min.Y, 0.0f, View.ViewRect.Max.X, View.ViewRect.Max.Y, 1.0f);
 			RHICmdList.SetRasterizerState(TStaticRasterizerState<FM_Solid, CM_None>::GetRHI());
@@ -973,6 +975,8 @@ void FDeferredShadingSceneRenderer::RenderStandardDeferredImageBasedReflections(
 				FIntPoint(View.ViewRect.Width(), View.ViewRect.Height()),
 				GSceneRenderTargets.GetBufferSizeXY(),
 				*VertexShader);
+
+			GSceneRenderTargets.FinishRenderingSceneColor(RHICmdList);
 		}
 	}
 }
@@ -1001,7 +1005,7 @@ void FDeferredShadingSceneRenderer::RenderDeferredReflections(FRHICommandListImm
 	{
 		const bool bReflectionEnvironment = ShouldDoReflectionEnvironment();
 		const bool bReflectionsWithCompute = (FeatureLevel >= ERHIFeatureLevel::SM5) && bReflectionEnvironment && Scene->ReflectionSceneData.CubemapArray.IsValid();
-
+		
 		if (bReflectionsWithCompute)
 		{
 			RenderTiledDeferredImageBasedReflections(RHICmdList, DynamicBentNormalAO);
diff --git a/Engine/Source/Runtime/Renderer/Private/ReflectionEnvironmentCapture.cpp b/Engine/Source/Runtime/Renderer/Private/ReflectionEnvironmentCapture.cpp
--- a/Engine/Source/Runtime/Renderer/Private/ReflectionEnvironmentCapture.cpp
+++ b/Engine/Source/Runtime/Renderer/Private/ReflectionEnvironmentCapture.cpp
@@ -75,6 +75,15 @@ FSceneRenderTargetItem& GetEffectiveSourceTexture(bool bDownsamplePass, int32 Ta
 	return GSceneRenderTargets.ReflectionColorScratchCubemap[ScratchTextureIndex]->GetRenderTargetItem();
 }
 
+void FullyResolveReflectionScratchCubes(FRHICommandListImmediate& RHICmdList)
+{
+	FTextureRHIRef& Scratch0 = GSceneRenderTargets.ReflectionColorScratchCubemap[0]->GetRenderTargetItem().TargetableTexture;
+	FTextureRHIRef& Scratch1 = GSceneRenderTargets.ReflectionColorScratchCubemap[1]->GetRenderTargetItem().TargetableTexture;
+	FResolveParams ResolveParams(FResolveRect(), CubeFace_PosX, -1, -1, -1);
+	RHICmdList.CopyToResolveTarget(Scratch0, Scratch0, true, ResolveParams);
+	RHICmdList.CopyToResolveTarget(Scratch1, Scratch1, true, ResolveParams);  
+}
+
 class FDownsamplePS : public FGlobalShader
 {
 	DECLARE_SHADER_TYPE(FDownsamplePS,Global);
@@ -275,7 +284,8 @@ IMPLEMENT_SHADER_TYPE(,FComputeBrightnessPS,TEXT("ReflectionEnvironmentShaders")
 /** Computes the average brightness of the given reflection capture and stores it in the scene. */
 void ComputeAverageBrightness(FRHICommandList& RHICmdList, ERHIFeatureLevel::Type FeatureLevel)
 {
-	SetRenderTarget(RHICmdList, GSceneRenderTargets.GetReflectionBrightnessTarget()->GetRenderTargetItem().TargetableTexture, NULL);
+	FTextureRHIRef& BrightnessTarget = GSceneRenderTargets.GetReflectionBrightnessTarget()->GetRenderTargetItem().TargetableTexture;
+	SetRenderTarget(RHICmdList, BrightnessTarget, NULL);
 	RHICmdList.SetRasterizerState(TStaticRasterizerState<FM_Solid, CM_None>::GetRHI());
 	RHICmdList.SetDepthStencilState(TStaticDepthStencilState<false, CF_Always>::GetRHI());
 	RHICmdList.SetBlendState(TStaticBlendState<>::GetRHI());
@@ -299,6 +309,8 @@ void ComputeAverageBrightness(FRHICommandList& RHICmdList, ERHIFeatureLevel::Typ
 		FIntPoint(1, 1),
 		FIntPoint(1, 1),
 		*VertexShader);
+
+	RHICmdList.CopyToResolveTarget(BrightnessTarget, BrightnessTarget, true, FResolveParams());
 }
 
 /** Generates mips for glossiness and filters the cubemap for a given reflection. */
@@ -313,6 +325,9 @@ void FilterReflectionEnvironment(FRHICommandListImmediate& RHICmdList, ERHIFeatu
 	static const auto CVar = IConsoleManager::Get().FindConsoleVariable(TEXT("r.DiffuseFromCaptures"));
 	bNormalize = bNormalize && CVar->GetInt() == 0;
 
+	// necessary to resolve the clears which touched all the mips.  scene rendering only resolves mip 0.
+	FullyResolveReflectionScratchCubes(RHICmdList);	
+
 	auto ShaderMap = GetGlobalShaderMap(FeatureLevel);
 
 	{	
diff --git a/Engine/Source/Runtime/Renderer/Private/SceneHitProxyRendering.cpp b/Engine/Source/Runtime/Renderer/Private/SceneHitProxyRendering.cpp
--- a/Engine/Source/Runtime/Renderer/Private/SceneHitProxyRendering.cpp
+++ b/Engine/Source/Runtime/Renderer/Private/SceneHitProxyRendering.cpp
@@ -350,7 +350,7 @@ TRefCountPtr<IPooledRenderTarget> InitHitProxyRender(FRHICommandListImmediate& R
 		return HitProxyRT;
 	}
 
-	SetRenderTarget(RHICmdList, HitProxyRT->GetRenderTargetItem().TargetableTexture, GSceneRenderTargets.GetSceneDepthSurface());
+	SetRenderTarget(RHICmdList, HitProxyRT->GetRenderTargetItem().TargetableTexture, GSceneRenderTargets.GetSceneDepthSurface(), ESimpleRenderTargetMode::EExistingColorAndReadOnlyDepth);
 
 	// Clear color for each view.
 	auto& Views = SceneRenderer->Views;
diff --git a/Engine/Source/Runtime/Renderer/Private/SceneOcclusion.cpp b/Engine/Source/Runtime/Renderer/Private/SceneOcclusion.cpp
--- a/Engine/Source/Runtime/Renderer/Private/SceneOcclusion.cpp
+++ b/Engine/Source/Runtime/Renderer/Private/SceneOcclusion.cpp
@@ -1036,11 +1036,11 @@ void FDeferredShadingSceneRenderer::BeginOcclusionTests(FRHICommandListImmediate
 
 	if (bUseDownsampledDepth)
 	{
-		SetRenderTarget(RHICmdList, NULL, GSceneRenderTargets.GetSmallDepthSurface());
+		SetRenderTarget(RHICmdList, NULL, GSceneRenderTargets.GetSmallDepthSurface(), ESimpleRenderTargetMode::EExistingColorAndReadOnlyDepth);
 	}
 	else
 	{
-		SetRenderTarget(RHICmdList, NULL, GSceneRenderTargets.GetSceneDepthSurface());
+		SetRenderTarget(RHICmdList, NULL, GSceneRenderTargets.GetSceneDepthSurface(), ESimpleRenderTargetMode::EExistingColorAndReadOnlyDepth);
 	}
 
 	// Perform occlusion queries for each view
@@ -1196,6 +1196,6 @@ void FDeferredShadingSceneRenderer::BeginOcclusionTests(FRHICommandListImmediate
 	if (bUseDownsampledDepth)
 	{
 		// Restore default render target
-		GSceneRenderTargets.BeginRenderingSceneColor(RHICmdList);
+		GSceneRenderTargets.BeginRenderingSceneColor(RHICmdList, ESimpleRenderTargetMode::EUninitializedColorExistingReadOnlyDepth);
 	}
 }
diff --git a/Engine/Source/Runtime/Renderer/Private/SceneRendering.cpp b/Engine/Source/Runtime/Renderer/Private/SceneRendering.cpp
--- a/Engine/Source/Runtime/Renderer/Private/SceneRendering.cpp
+++ b/Engine/Source/Runtime/Renderer/Private/SceneRendering.cpp
@@ -135,7 +135,7 @@ FParallelCommandListSet::FParallelCommandListSet(const FViewInfo& InView, FRHICo
 
 }
 
-FParallelCommandListSet::~FParallelCommandListSet()
+void FParallelCommandListSet::Dispatch()
 {
 	check(CommandLists.Num() == Events.Num());
 #if PLATFORM_SUPPORTS_PARALLEL_RHI_EXECUTE
@@ -156,13 +156,15 @@ FParallelCommandListSet::~FParallelCommandListSet()
 	Events.Reset();
 }
 
+FParallelCommandListSet::~FParallelCommandListSet()
+{
+	checkf(CommandLists.Num() == 0, TEXT("Derived class of FParallelCommandListSet did not call Dispatch in virtual destructor"));
+}
+
 FRHICommandList* FParallelCommandListSet::NewParallelCommandList()
 {
 	FRHICommandList* Result = AllocCommandList();
-//	if (bParallelExecute)
-	{
-		SetStateOnCommandList(*Result); 
-	}
+	SetStateOnCommandList(*Result); 
 	return Result;
 }
 
diff --git a/Engine/Source/Runtime/Renderer/Private/SceneRendering.h b/Engine/Source/Runtime/Renderer/Private/SceneRendering.h
--- a/Engine/Source/Runtime/Renderer/Private/SceneRendering.h
+++ b/Engine/Source/Runtime/Renderer/Private/SceneRendering.h
@@ -319,13 +319,16 @@ public:
 	TArray<FRHICommandList*,SceneRenderingAllocator> CommandLists;
 	TArray<FGraphEventRef,SceneRenderingAllocator> Events;
 protected:
+	//this must be called by deriving classes virtual destructor because it calls the virtual SetStateOnCommandList.
+	//C++ will not do dynamic dispatch of virtual calls from destructors so we can't call it in the base class.
+	void Dispatch();
 	FRHICommandList* AllocCommandList();
 	bool bParallelExecute;
 public:
 	FParallelCommandListSet(const FViewInfo& InView, FRHICommandList& InParentCmdList, bool* InOutDirty, bool bInParallelExecute);
-	~FParallelCommandListSet();
+	virtual ~FParallelCommandListSet();
 	FRHICommandList* NewParallelCommandList();
-	void AddParallelCommandList(FRHICommandList* CmdList, FGraphEventRef& CompletionEvent);
+	void AddParallelCommandList(FRHICommandList* CmdList, FGraphEventRef& CompletionEvent);	
 
 	virtual void SetStateOnCommandList(FRHICommandList& CmdList)
 	{
diff --git a/Engine/Source/Runtime/Renderer/Private/ShadowRendering.cpp b/Engine/Source/Runtime/Renderer/Private/ShadowRendering.cpp
--- a/Engine/Source/Runtime/Renderer/Private/ShadowRendering.cpp
+++ b/Engine/Source/Runtime/Renderer/Private/ShadowRendering.cpp
@@ -1151,57 +1151,105 @@ static void CheckShadowDepthMaterials(const FMaterialRenderProxy* InRenderProxy,
 	check(Material == InMaterial);
 }
 
-void FProjectedShadowInfo::ClearDepth(FRHICommandList& RHICmdList, FDeferredShadingSceneRenderer* SceneRenderer)
+void FProjectedShadowInfo::ClearDepth(FRHICommandList& RHICmdList, FDeferredShadingSceneRenderer* SceneRenderer, bool bPerformClear)
 {
+	uint32 ViewportMinX;
+	uint32 ViewportMinY;
+	float ViewportMinZ;
+	uint32 ViewportMaxX;
+	uint32 ViewportMaxY;
+	float ViewportMaxZ;
+
+	int32 NumClearColors;
+	bool bClearColor;
+	FLinearColor Colors[2];
+	bool bClearDepth;
+	float Depth;
+	bool bClearStencil;
+	uint32 Stencil;
+	FIntRect ExcludeRect;
+
 	if (CascadeSettings.bOnePassPointLightShadow)
 	{
 		// Set the viewport to the whole render target since it's a cube map, don't leave any border space
-		RHICmdList.SetViewport(
-			0,
-			0,
-			0.0f,
-			ResolutionX,
-			ResolutionY,
-			1.0f
-			);
-
+		ViewportMinX = 0;
+		ViewportMinY = 0;
+		ViewportMinZ = 0.0f;
+		ViewportMaxX = ResolutionX;
+		ViewportMaxY = ResolutionY;
+		ViewportMaxZ = 1.0f;
+		
+		
 		// Clear depth only.
-		RHICmdList.Clear(false, FColor::White, true, 1.0f, false, 0, FIntRect());
+		bClearColor = false;
+		Colors[0] = FLinearColor::White;
+		bClearDepth = true;
+		Depth = 1.0f;
+		bClearStencil = false;
+		Stencil = 0;
+		NumClearColors = 1;
 	}
 	else
 	{
 		if (bReflectiveShadowmap)
 		{
 			// Set the viewport for the shadow.
-			RHICmdList.SetViewport(
-				X,
-				Y,
-				0.0f,
-				X + ResolutionX, 
-				Y + ResolutionY,
-				1.0f
-				);
-
-			// Clear color and depth targets
-			FLinearColor ClearColors[2] = {FLinearColor(0,0,1,0), FLinearColor(0,0,0,0)};
-			RHICmdList.ClearMRT(true, 2, ClearColors, true, 1.0f, false, 0, FIntRect());
+			ViewportMinX = X;
+			ViewportMinY = Y;
+			ViewportMinZ = 0.0f;
+			ViewportMaxX = X + ResolutionX;
+			ViewportMaxY = Y + ResolutionY;
+			ViewportMaxZ = 1.0f;
+			
+			// Clear color and depth targets			
+			bClearColor = true;
+			Colors[0] = FLinearColor(0, 0, 1, 0);
+			Colors[1] = FLinearColor(0, 0, 0, 0);
+		
+			bClearDepth = true;
+			Depth = 1.0f;
+			bClearStencil = false;
+			Stencil = 0;
+			NumClearColors = 2;			
 		}
 		else
 		{
 			// Set the viewport for the shadow.
-			RHICmdList.SetViewport(
-				X,
-				Y,
-				0.0f,
-				X + SHADOW_BORDER*2 + ResolutionX,
-				Y + SHADOW_BORDER*2 + ResolutionY,
-				1.0f
-				);
+			ViewportMinX = X;
+			ViewportMinY = Y;
+			ViewportMinZ = 0.0f;
+			ViewportMaxX = X + SHADOW_BORDER * 2 + ResolutionX;
+			ViewportMaxY = Y + SHADOW_BORDER * 2 + ResolutionY;
+			ViewportMaxZ = 1.0f;
 
 			// Clear depth only.
-			RHICmdList.Clear(false, FColor::White, true, 1.0f, false, 0, FIntRect());
+			bClearColor = false;
+			Colors[0] = FLinearColor::White;
+			bClearDepth = true;
+			Depth = 1.0f;
+			bClearStencil = false;
+			Stencil = 0;
+			NumClearColors = 1;			
 		}
 	}
+
+	if (bPerformClear)
+	{
+		RHICmdList.SetViewport(
+			ViewportMinX,
+			ViewportMinY,
+			ViewportMinZ,
+			ViewportMaxX,
+			ViewportMaxY,
+			ViewportMaxZ
+			);
+		
+		RHICmdList.ClearMRT(bClearColor, NumClearColors, Colors, bClearDepth, Depth, bClearStencil, Stencil, FIntRect());
+	}
+	else
+	{
+		RHICmdList.BindClearMRTValues(bClearColor, NumClearColors, Colors, bClearDepth, Depth, bClearStencil, Stencil);
+	}
 }
 
 
@@ -1482,8 +1530,14 @@ public:
 	{
 		SetStateOnCommandList(ParentCmdList);
 	}
-	virtual void SetStateOnCommandList(FRHICommandList& CmdList) override
+
+	virtual ~FShadowParallelCommandListSet()
 	{
+		Dispatch();
+	}
+
+	virtual void SetStateOnCommandList(FRHICommandList& CmdList) override
+	{	
 		SetShadowRenderTargets(CmdList);
 		ProjectedShadowInfo.SetStateForDepth(CmdList);
 	}
@@ -2610,14 +2664,17 @@ bool FDeferredShadingSceneRenderer::RenderOnePassPointLightShadows(FRHICommandLi
 			if (!ProjectedShadowInfo->CascadeSettings.bRayTracedDistanceField)
 			{
 				SCOPED_DRAW_EVENT(RHICmdList, ShadowDepthsFromOpaquePointLight);
-				auto SetShadowRenderTargets = [ProjectedShadowInfo](FRHICommandList& RHICmdList)
+
+				bool bPerformClear = true;
+				auto SetShadowRenderTargets = [this, &bPerformClear, ProjectedShadowInfo](FRHICommandList& RHICmdList)
 				{
 					GSceneRenderTargets.BeginRenderingCubeShadowDepth(RHICmdList, ProjectedShadowInfo->ResolutionX);
+					ProjectedShadowInfo->ClearDepth(RHICmdList, this, bPerformClear);
 				};
 
-				SetShadowRenderTargets(RHICmdList);  // run it now, maybe run it later for parallel command lists
 				ProjectedShadowInfo->bAllocated = true;
-				ProjectedShadowInfo->ClearDepth(RHICmdList, this);
+				SetShadowRenderTargets(RHICmdList);  // run it now, maybe run it later for parallel command lists								
+				bPerformClear = false;
 
 				ProjectedShadowInfo->RenderDepth(RHICmdList, this, SetShadowRenderTargets);
 				GSceneRenderTargets.FinishRenderingCubeShadowDepth(RHICmdList, ProjectedShadowInfo->ResolutionX);
@@ -3022,19 +3079,18 @@ bool FDeferredShadingSceneRenderer::RenderReflectiveShadowMaps(FRHICommandListIm
 				FLightPropagationVolume* LightPropagationVolume = ViewState->GetLightPropagationVolume();
 
 				check(LightPropagationVolume);
-				auto SetShadowRenderTargets = [LightPropagationVolume](FRHICommandList& RHICmdList)
+				auto SetShadowRenderTargets = [this, LightPropagationVolume, ProjectedShadowInfo](FRHICommandList& RHICmdList)
 				{
 					GSceneRenderTargets.BeginRenderingReflectiveShadowMap(RHICmdList, LightPropagationVolume);
+					ProjectedShadowInfo->ClearDepth(RHICmdList, this, false);
 				};
-
-				SetShadowRenderTargets(RHICmdList);  // run it now, maybe run it later for parallel command lists
-
+				SetShadowRenderTargets(RHICmdList);  // run it now, maybe run it later for parallel command lists				
 
 				LightPropagationVolume->SetVplInjectionConstants( *ProjectedShadowInfo, LightSceneInfo->Proxy );
 
 				if (ProjectedShadowInfo->bAllocated && !ProjectedShadowInfo->bTranslucentShadow)
 				{
-					ProjectedShadowInfo->ClearDepth(RHICmdList, this);
+					ProjectedShadowInfo->ClearDepth(RHICmdList, this, true);
 					ProjectedShadowInfo->RenderDepth(RHICmdList, this, SetShadowRenderTargets);
 				}
 
@@ -3199,14 +3255,14 @@ bool FDeferredShadingSceneRenderer::RenderProjectedShadows(FRHICommandListImmedi
 			// Render the shadow depths.
 			SCOPED_DRAW_EVENT(RHICmdList, ShadowDepthsFromOpaqueProjected);
 
-			bool bFirst = true;
-			auto SetShadowRenderTargets = [&bFirst](FRHICommandList& RHICmdList)
+			bool bPerformClear = true;
+			auto SetShadowRenderTargets = [&bPerformClear](FRHICommandList& RHICmdList)
 			{
-				GSceneRenderTargets.BeginRenderingShadowDepth(RHICmdList, bFirst);
+				GSceneRenderTargets.BeginRenderingShadowDepth(RHICmdList, bPerformClear);
 			};
 
 			SetShadowRenderTargets(RHICmdList);  // run it now, maybe run it later for parallel command lists
-			bFirst = false;
+			bPerformClear = false;
 
 			// render depth for each shadow
 			for (int32 ShadowIndex = 0; ShadowIndex < Shadows.Num(); ShadowIndex++)
@@ -3337,15 +3393,6 @@ bool FDeferredShadingSceneRenderer::RenderCachedPreshadows(FRHICommandListImmedi
 		{
 			SCOPED_DRAW_EVENTF(RHICmdList, EventShadowDepths, TEXT("Preshadow Cache Depths"));
 
-			auto SetShadowRenderTargets = [](FRHICommandList& RHICmdList)
-			{
-				SetRenderTarget(RHICmdList, FTextureRHIRef(), GSceneRenderTargets.PreShadowCacheDepthZ->GetRenderTargetItem().TargetableTexture);
-			};
-
-			SetShadowRenderTargets(RHICmdList);  // run it now, maybe run it later for parallel command lists
-
-
-
 			// Render depth for each shadow
 			for (int32 ShadowIndex = 0; ShadowIndex < CachedPreshadows.Num(); ShadowIndex++)
 			{
@@ -3354,7 +3401,16 @@ bool FDeferredShadingSceneRenderer::RenderCachedPreshadows(FRHICommandListImmedi
 				if (!ProjectedShadowInfo->bDepthsCached)
 				{
 					check(ProjectedShadowInfo->bAllocated);
-					ProjectedShadowInfo->ClearDepth(RHICmdList, this);
+					
+					//
+					bool bPerformClear = true;;
+					auto SetShadowRenderTargets = [this, &bPerformClear, ProjectedShadowInfo](FRHICommandList& RHICmdList)
+					{
+						SetRenderTarget(RHICmdList, FTextureRHIRef(), GSceneRenderTargets.PreShadowCacheDepthZ->GetRenderTargetItem().TargetableTexture);
+						ProjectedShadowInfo->ClearDepth(RHICmdList, this, bPerformClear);
+					};					
+					SetShadowRenderTargets(RHICmdList); // run it now, maybe run it later for parallel command lists
+					bPerformClear = false;
 
 					ProjectedShadowInfo->RenderDepth(RHICmdList, this, SetShadowRenderTargets);
 					ProjectedShadowInfo->bDepthsCached = true;
diff --git a/Engine/Source/Runtime/Renderer/Private/ShadowRendering.h b/Engine/Source/Runtime/Renderer/Private/ShadowRendering.h
--- a/Engine/Source/Runtime/Renderer/Private/ShadowRendering.h
+++ b/Engine/Source/Runtime/Renderer/Private/ShadowRendering.h
@@ -633,7 +633,7 @@ public:
 	/** Set state for depth rendering */
 	void SetStateForDepth(FRHICommandList& RHICmdList);
 
-	void ClearDepth(FRHICommandList& RHICmdList, class FDeferredShadingSceneRenderer* SceneRenderer);
+	void ClearDepth(FRHICommandList& RHICmdList, class FDeferredShadingSceneRenderer* SceneRenderer, bool bPerformClear);
 
 	/** Renders shadow maps for translucent primitives. */
 	void RenderTranslucencyDepths(FRHICommandList& RHICmdList, class FDeferredShadingSceneRenderer* SceneRenderer);
diff --git a/Engine/Source/Runtime/Renderer/Private/TranslucentLighting.cpp b/Engine/Source/Runtime/Renderer/Private/TranslucentLighting.cpp
--- a/Engine/Source/Runtime/Renderer/Private/TranslucentLighting.cpp
+++ b/Engine/Source/Runtime/Renderer/Private/TranslucentLighting.cpp
@@ -527,10 +527,12 @@ void FProjectedShadowInfo::RenderTranslucencyDepths(FRHICommandList& RHICmdList,
 		SCOPED_DRAW_EVENTF(RHICmdList, EventShadowDepthActor, *EventName);
 #endif
 
+		FTextureRHIParamRef ShadowTransmission0 = GSceneRenderTargets.TranslucencyShadowTransmission[0]->GetRenderTargetItem().TargetableTexture;
+		FTextureRHIParamRef ShadowTransmission1 = GSceneRenderTargets.TranslucencyShadowTransmission[1]->GetRenderTargetItem().TargetableTexture;
 		FTextureRHIParamRef RenderTargets[2] =
 		{
-			GSceneRenderTargets.TranslucencyShadowTransmission[0]->GetRenderTargetItem().TargetableTexture,
-			GSceneRenderTargets.TranslucencyShadowTransmission[1]->GetRenderTargetItem().TargetableTexture
+			ShadowTransmission0,
+			ShadowTransmission1
 		};
 		SetRenderTargets(RHICmdList, ARRAY_COUNT(RenderTargets), RenderTargets, FTextureRHIParamRef(), 0, NULL);
 
@@ -598,7 +600,11 @@ void FProjectedShadowInfo::RenderTranslucencyDepths(FRHICommandList& RHICmdList,
 				}
 			}
 		}
-	}
+
+		FResolveParams ResolveParams;
+		RHICmdList.CopyToResolveTarget(ShadowTransmission0, ShadowTransmission0, true, ResolveParams);
+		RHICmdList.CopyToResolveTarget(ShadowTransmission1, ShadowTransmission0, true, ResolveParams);
+	}	
 
 	// Restore overridden properties
 	FoundView->bForceShowMaterials = false;
@@ -1765,6 +1771,17 @@ void FDeferredShadingSceneRenderer::FilterTranslucentVolumeLighting(FRHICommandL
 {
 	if (GUseTranslucentLightingVolumes && GSupportsVolumeTextureRendering)
 	{
+		// textures have to be finalized before reading.
+		for (int32 VolumeCascadeIndex = 0; VolumeCascadeIndex < TVC_MAX; VolumeCascadeIndex++)
+		{
+			const IPooledRenderTarget* RT0 = GSceneRenderTargets.TranslucencyLightingVolumeAmbient[VolumeCascadeIndex];
+			const IPooledRenderTarget* RT1 = GSceneRenderTargets.TranslucencyLightingVolumeDirectional[VolumeCascadeIndex];
+			FTextureRHIRef TargetTexture0 = RT0->GetRenderTargetItem().TargetableTexture;
+			FTextureRHIRef TargetTexture1 = RT1->GetRenderTargetItem().TargetableTexture;
+			RHICmdList.CopyToResolveTarget(TargetTexture0, TargetTexture0, true, FResolveParams());
+			RHICmdList.CopyToResolveTarget(TargetTexture1, TargetTexture1, true, FResolveParams());
+		}
+
 		if (GUseTranslucencyVolumeBlur)
 		{
 			SCOPED_DRAW_EVENT(RHICmdList, FilterTranslucentVolume);
diff --git a/Engine/Source/Runtime/Renderer/Private/TranslucentRendering.cpp b/Engine/Source/Runtime/Renderer/Private/TranslucentRendering.cpp
--- a/Engine/Source/Runtime/Renderer/Private/TranslucentRendering.cpp
+++ b/Engine/Source/Runtime/Renderer/Private/TranslucentRendering.cpp
@@ -32,6 +32,18 @@ static void SetTranslucentRenderTargetAndState(FRHICommandList& RHICmdList, cons
 	}
 }
 
+static void FinishTranslucentRenderTarget(FRHICommandListImmediate& RHICmdList, const FViewInfo& View, bool bSeparateTranslucencyPass)
+{
+	if (bSeparateTranslucencyPass && GSceneRenderTargets.IsSeparateTranslucencyActive(View))
+	{
+		GSceneRenderTargets.FinishRenderingSeparateTranslucency(RHICmdList, View);
+	}
+	else
+	{
+		GSceneRenderTargets.FinishRenderingTranslucency(RHICmdList, View);
+	}
+}
+
 const FProjectedShadowInfo* FDeferredShadingSceneRenderer::PrepareTranslucentShadowMap(FRHICommandList& RHICmdList, const FViewInfo& View, FPrimitiveSceneInfo* PrimitiveSceneInfo, bool bSeparateTranslucencyPass)
 {
 	const FVisibleLightInfo* VisibleLightInfo = NULL;
@@ -870,6 +882,12 @@ public:
 	{
 		SetStateOnCommandList(ParentCmdList);
 	}
+
+	virtual ~FTranslucencyPassParallelCommandListSet()
+	{
+		Dispatch();
+	}
+
 	virtual void SetStateOnCommandList(FRHICommandList& CmdList) override
 	{
 		SetTranslucentRenderTargetAndState(CmdList, View, bSeparateTranslucency, bFirstTimeThisFrame);
@@ -898,7 +916,8 @@ void FDeferredShadingSceneRenderer::RenderTranslucencyParallel(FRHICommandListIm
 		const FViewInfo& View = Views[ViewIndex];
 
 		{
-			FTranslucencyPassParallelCommandListSet ParallelCommandListSet(View, RHICmdList, nullptr, CVarRHICmdTranslucencyPassDeferredContexts.GetValueOnRenderThread() > 0, false);
+			const bool bSeparateTranslucency = false;
+			FTranslucencyPassParallelCommandListSet ParallelCommandListSet(View, RHICmdList, nullptr, CVarRHICmdTranslucencyPassDeferredContexts.GetValueOnRenderThread() > 0, bSeparateTranslucency);
 
 			{
 				int32 NumPrims = View.TranslucentPrimSet.NumPrims() - View.TranslucentPrimSet.NumSeparateTranslucencyPrims();
@@ -934,6 +953,7 @@ void FDeferredShadingSceneRenderer::RenderTranslucencyParallel(FRHICommandListIm
 			// Draw the view's mesh elements with the translucent drawing policy.
 			DrawViewElementsParallel<FTranslucencyDrawingPolicyFactory>(ThisContext, SDPG_Foreground, false, ParallelCommandListSet);
 
+			FinishTranslucentRenderTarget(RHICmdList, View, bSeparateTranslucency);
 		}
 
 #if 0 // unsupported visualization in the parallel case
@@ -1035,6 +1055,8 @@ void FDeferredShadingSceneRenderer::RenderTranslucency(FRHICommandListImmediate&
 					LightPropagationVolume->Visualise(RHICmdList, View);
 				}
 			}
+
+			FinishTranslucentRenderTarget(RHICmdList, View, false);
 			
 			{
 				bool bRenderSeparateTranslucency = View.TranslucentPrimSet.NumSeparateTranslucencyPrims() > 0;
diff --git a/Engine/Source/Runtime/Renderer/Private/VelocityRendering.cpp b/Engine/Source/Runtime/Renderer/Private/VelocityRendering.cpp
--- a/Engine/Source/Runtime/Renderer/Private/VelocityRendering.cpp
+++ b/Engine/Source/Runtime/Renderer/Private/VelocityRendering.cpp
@@ -560,10 +560,35 @@ public:
 	}
 };
 
-static void SetVelocitiesState(FRHICommandList& RHICmdList, const FViewInfo& View, TRefCountPtr<IPooledRenderTarget>& VelocityRT)
+static void BeginVelocityRendering(FRHICommandList& RHICmdList, TRefCountPtr<IPooledRenderTarget>& VelocityRT, bool bPerformClear)
 {
-	SetRenderTarget(RHICmdList, VelocityRT->GetRenderTargetItem().TargetableTexture, GSceneRenderTargets.GetSceneDepthTexture());
+	FTextureRHIRef VelocityTexture = VelocityRT->GetRenderTargetItem().TargetableTexture;
+	FTexture2DRHIRef DepthTexture = GSceneRenderTargets.GetSceneDepthTexture();
+	FLinearColor VelocityClearColor = FLinearColor::Black;
+	if (bPerformClear)
+	{
+		// now make the FRHISetRenderTargetsInfo that encapsulates all of the info
+		FRHIRenderTargetView ColorView(VelocityTexture, 0, -1, ERenderTargetLoadAction::EClear, ERenderTargetStoreAction::EStore);
+		FRHIDepthRenderTargetView DepthView(DepthTexture, ERenderTargetLoadAction::ELoad, ERenderTargetStoreAction::ENoAction, true);
+
+		FRHISetRenderTargetsInfo Info(1, &ColorView, DepthView);
+		Info.ClearColors[0] = VelocityClearColor;
 
+		// Clear the velocity buffer (0.0f means "use static background velocity").
+		RHICmdList.SetRenderTargetsAndClear(Info);		
+	}
+	else
+	{
+		SetRenderTarget(RHICmdList, VelocityTexture, DepthTexture, ESimpleRenderTargetMode::EExistingColorAndReadOnlyDepth);
+
+		// some platforms need the clear color when rendertargets transition to SRVs.  We propagate here to allow parallel rendering to always
+		// have the proper mapping when the RT is transitioned.
+		RHICmdList.BindClearMRTValues(true, 1, &VelocityClearColor, false, 1.0f, false, 0);
+	}
+}
+
+static void SetVelocitiesState(FRHICommandList& RHICmdList, const FViewInfo& View, TRefCountPtr<IPooledRenderTarget>& VelocityRT)
+{
 	const FIntPoint BufferSize = GSceneRenderTargets.GetBufferSizeXY();
 	const FIntPoint VelocityBufferSize = BufferSize;		// full resolution so we can reuse the existing full res z buffer
 
@@ -590,8 +615,15 @@ public:
 	{
 		SetStateOnCommandList(ParentCmdList);
 	}
+
+	virtual ~FVelocityPassParallelCommandListSet()
+	{
+		Dispatch();
+	}	
+
 	virtual void SetStateOnCommandList(FRHICommandList& CmdList) override
 	{
+		BeginVelocityRendering(CmdList, VelocityRT, false);
 		SetVelocitiesState(CmdList, View, VelocityRT);
 	}
 };
@@ -699,11 +731,8 @@ void FDeferredShadingSceneRenderer::RenderVelocities(FRHICommandListImmediate& R
 
 	GPrevPerBoneMotionBlur.LockData();
 
-	SetRenderTarget(RHICmdList, VelocityRT->GetRenderTargetItem().TargetableTexture, GSceneRenderTargets.GetSceneDepthTexture());
+	BeginVelocityRendering(RHICmdList, VelocityRT, true);
 	
-	// Clear the velocity buffer (0.0f means "use static background velocity").
-	RHICmdList.Clear(true, FLinearColor::Black, false, 1.0f, false, 0, FIntRect());
-
 	if (GRHICommandList.UseParallelAlgorithms() && CVarParallelVelocity.GetValueOnRenderThread())
 	{
 		RenderVelocitiesInnerParallel(RHICmdList, VelocityRT);
diff --git a/Engine/Source/Runtime/Windows/D3D11RHI/Private/D3D11Commands.cpp b/Engine/Source/Runtime/Windows/D3D11RHI/Private/D3D11Commands.cpp
--- a/Engine/Source/Runtime/Windows/D3D11RHI/Private/D3D11Commands.cpp
+++ b/Engine/Source/Runtime/Windows/D3D11RHI/Private/D3D11Commands.cpp
@@ -822,12 +822,20 @@ FRTVDesc GetRenderTargetViewDesc(ID3D11RenderTargetView* RenderTargetView)
 void FD3D11DynamicRHI::RHISetRenderTargets(
 	uint32 NewNumSimultaneousRenderTargets,
 	const FRHIRenderTargetView* NewRenderTargetsRHI,
-	FTextureRHIParamRef NewDepthStencilTargetRHI,
+	const FRHIDepthRenderTargetView* NewDepthStencilTargetRHI,
 	uint32 NewNumUAVs,
 	const FUnorderedAccessViewRHIParamRef* UAVs
 	)
 {
-	FD3D11TextureBase* NewDepthStencilTarget = GetD3D11TextureFromRHITexture(NewDepthStencilTargetRHI);
+	FD3D11TextureBase* NewDepthStencilTarget = GetD3D11TextureFromRHITexture(NewDepthStencilTargetRHI ? NewDepthStencilTargetRHI->Texture : nullptr);
+
+#if CHECK_SRV_TRANSITIONS
+	// if the depth buffer is writable then it counts as unresolved.
+	if (NewDepthStencilTargetRHI && !NewDepthStencilTargetRHI->bReadOnly && NewDepthStencilTarget)
+	{		
+		UnresolvedTargets.Add(NewDepthStencilTarget->GetResource(), FUnresolvedRTInfo(NewDepthStencilTargetRHI->Texture->GetName(), 0, 1, -1, 1));
+	}
+#endif
 
 	check(NewNumSimultaneousRenderTargets + NewNumUAVs <= MaxSimultaneousRenderTargets);
 
@@ -858,8 +866,18 @@ void FD3D11DynamicRHI::RHISetRenderTargets(
 		ID3D11RenderTargetView* RenderTargetView = NULL;
 		if(RenderTargetIndex < NewNumSimultaneousRenderTargets && IsValidRef(NewRenderTargetsRHI[RenderTargetIndex].Texture))
 		{
+			int32 RTMipIndex = NewRenderTargetsRHI[RenderTargetIndex].MipIndex;
+			int32 RTSliceIndex = NewRenderTargetsRHI[RenderTargetIndex].ArraySliceIndex;
 			FD3D11TextureBase* NewRenderTarget = GetD3D11TextureFromRHITexture(NewRenderTargetsRHI[RenderTargetIndex].Texture);
-			RenderTargetView = NewRenderTarget->GetRenderTargetView(NewRenderTargetsRHI[RenderTargetIndex].MipIndex, NewRenderTargetsRHI[RenderTargetIndex].ArraySliceIndex);
+			RenderTargetView = NewRenderTarget->GetRenderTargetView(RTMipIndex, RTSliceIndex);
+
+#if CHECK_SRV_TRANSITIONS
+			// remember this target as having been bound for write.
+			ID3D11Resource* RTVResource;
+			RenderTargetView->GetResource(&RTVResource);			
+			UnresolvedTargets.Add(RTVResource, FUnresolvedRTInfo(NewRenderTargetsRHI[RenderTargetIndex].Texture->GetName(), RTMipIndex, 1, RTSliceIndex, 1));
+			RTVResource->Release();
+#endif
 			
 			// Unbind any shader views of the render target that are bound.
 			ConditionalClearShaderResource(NewRenderTarget);
@@ -965,7 +983,7 @@ void FD3D11DynamicRHI::RHISetRenderTargetsAndClear(const FRHISetRenderTargetsInf
 {
 	this->RHISetRenderTargets(RenderTargetsInfo.NumColorRenderTargets,
 		RenderTargetsInfo.ColorRenderTarget,
-		RenderTargetsInfo.DepthStencilRenderTarget.Texture,
+		&RenderTargetsInfo.DepthStencilRenderTarget,
 		0,
 		nullptr);
 	if (RenderTargetsInfo.bClearColor || RenderTargetsInfo.bClearStencil || RenderTargetsInfo.bClearDepth)
@@ -1896,6 +1914,12 @@ void FD3D11DynamicRHI::RHIClearMRT(bool bClearColor,int32 NumClearColors,const F
 	GPUProfilingData.RegisterGPUWork(0);
 }
 
+void FD3D11DynamicRHI::RHIBindClearMRTValues(bool bClearColor, int32 NumClearColors, const FLinearColor* ClearColorArray, bool bClearDepth, float Depth, bool bClearStencil, uint32 Stencil)
+{
+	// Not necessary for d3d.
+}
+
+
 // Functions to yield and regain rendering control from D3D
 
 void FD3D11DynamicRHI::RHISuspendRendering()
diff --git a/Engine/Source/Runtime/Windows/D3D11RHI/Private/D3D11RHI.cpp b/Engine/Source/Runtime/Windows/D3D11RHI/Private/D3D11RHI.cpp
--- a/Engine/Source/Runtime/Windows/D3D11RHI/Private/D3D11RHI.cpp
+++ b/Engine/Source/Runtime/Windows/D3D11RHI/Private/D3D11RHI.cpp
@@ -65,11 +65,122 @@ void FD3D11DynamicRHI::ClearState()
 	}
 }
 
+void GetMipAndSliceInfoFromSRV(ID3D11ShaderResourceView* SRV, int32& MipLevel, int32& NumMips, int32& ArraySlice, int32& NumSlices)
+{
+	MipLevel = -1;
+	NumMips = -1;
+	ArraySlice = -1;
+	NumSlices = -1;
+
+	if (SRV)
+	{
+		D3D11_SHADER_RESOURCE_VIEW_DESC SRVDesc;
+		SRV->GetDesc(&SRVDesc);
+		switch (SRVDesc.ViewDimension)
+		{			
+			case D3D11_SRV_DIMENSION_TEXTURE1D:
+				MipLevel	= SRVDesc.Texture1D.MostDetailedMip;
+				NumMips		= SRVDesc.Texture1D.MipLevels;
+				break;
+			case D3D11_SRV_DIMENSION_TEXTURE1DARRAY:
+				MipLevel	= SRVDesc.Texture1DArray.MostDetailedMip;
+				NumMips		= SRVDesc.Texture1DArray.MipLevels;
+				ArraySlice	= SRVDesc.Texture1DArray.FirstArraySlice;
+				NumSlices	= SRVDesc.Texture1DArray.ArraySize;
+				break;
+			case D3D11_SRV_DIMENSION_TEXTURE2D:
+				MipLevel	= SRVDesc.Texture2D.MostDetailedMip;
+				NumMips		= SRVDesc.Texture2D.MipLevels;
+				break;
+			case D3D11_SRV_DIMENSION_TEXTURE2DARRAY:
+				MipLevel	= SRVDesc.Texture2DArray.MostDetailedMip;
+				NumMips		= SRVDesc.Texture2DArray.MipLevels;
+				ArraySlice	= SRVDesc.Texture2DArray.FirstArraySlice;
+				NumSlices	= SRVDesc.Texture2DArray.ArraySize;
+				break;
+			case D3D11_SRV_DIMENSION_TEXTURE2DMS:
+				MipLevel	= 0;
+				NumMips		= 1;
+				break;
+			case D3D11_SRV_DIMENSION_TEXTURE2DMSARRAY:
+				MipLevel	= 0;
+				NumMips		= 1;
+				ArraySlice	= SRVDesc.Texture2DMSArray.FirstArraySlice;
+				NumSlices	= SRVDesc.Texture2DMSArray.ArraySize;
+				break;
+			case D3D11_SRV_DIMENSION_TEXTURE3D:
+				MipLevel	= SRVDesc.Texture3D.MostDetailedMip;
+				NumMips		= SRVDesc.Texture3D.MipLevels;
+				break;
+			case D3D11_SRV_DIMENSION_TEXTURECUBE:
+				MipLevel = SRVDesc.TextureCube.MostDetailedMip;
+				NumMips		= SRVDesc.TextureCube.MipLevels;
+				break;
+			case D3D11_SRV_DIMENSION_TEXTURECUBEARRAY:
+				MipLevel	= SRVDesc.TextureCubeArray.MostDetailedMip;
+				NumMips		= SRVDesc.TextureCubeArray.MipLevels;
+				ArraySlice	= SRVDesc.TextureCubeArray.First2DArrayFace;
+				NumSlices	= SRVDesc.TextureCubeArray.NumCubes;
+				break;
+			case D3D11_SRV_DIMENSION_BUFFER:
+			case D3D11_SRV_DIMENSION_BUFFEREX:
+			default:
+				break;
+		}
+	}
+}
+
+void FD3D11DynamicRHI::CheckIfSRVIsResolved(ID3D11ShaderResourceView* SRV)
+{
+#if CHECK_SRV_TRANSITIONS
+	if (GRHIThread || !SRV)
+	{
+		return;
+	}
+
+	static const auto CheckSRVCVar = IConsoleManager::Get().FindTConsoleVariableDataInt(TEXT("r.CheckSRVTransitions"));	
+	if (!CheckSRVCVar->GetValueOnRenderThread())
+	{
+		return;
+	}
+
+	ID3D11Resource* SRVResource = nullptr;
+	SRV->GetResource(&SRVResource);
+
+	int32 MipLevel;
+	int32 NumMips;
+	int32 ArraySlice;
+	int32 NumSlices;
+	GetMipAndSliceInfoFromSRV(SRV, MipLevel, NumMips, ArraySlice, NumSlices);
+
+	//d3d uses -1 to mean 'all mips'
+	int32 LastMip = MipLevel + NumMips - 1;
+	int32 LastSlice = ArraySlice + NumSlices - 1;
+
+	TArray<FUnresolvedRTInfo> RTInfoArray;
+	UnresolvedTargets.MultiFind(SRVResource, RTInfoArray);
+
+	for (int32 InfoIndex = 0; InfoIndex < RTInfoArray.Num(); ++InfoIndex)
+	{
+		const FUnresolvedRTInfo& RTInfo = RTInfoArray[InfoIndex];
+		int32 RTLastMip = RTInfo.MipLevel + RTInfo.NumMips - 1;		
+		ensureMsgf((MipLevel == -1 || NumMips == -1) || (LastMip < RTInfo.MipLevel || MipLevel > RTLastMip), TEXT("SRV is set to read mips in range %i to %i.  Target %s is unresolved for mip %i"), MipLevel, LastMip, *RTInfo.ResourceName.ToString(), RTInfo.MipLevel);
+		ensureMsgf(NumMips != -1, TEXT("SRV is set to read all mips.  Target %s is unresolved for mip %i"), *RTInfo.ResourceName.ToString(), RTInfo.MipLevel);
+
+		int32 RTLastSlice = RTInfo.ArraySlice + RTInfo.ArraySize - 1;
+		ensureMsgf((ArraySlice == -1 || ArraySlice == -1) || (LastSlice < RTInfo.ArraySlice || ArraySlice > RTLastSlice), TEXT("SRV is set to read slices in range %i to %i.  Target %s is unresolved for mip %i"), ArraySlice, LastSlice, *RTInfo.ResourceName.ToString(), RTInfo.ArraySlice);
+		ensureMsgf(ArraySlice == -1 || NumSlices != -1, TEXT("SRV is set to read all slices.  Target %s is unresolved for slice %i"));
+	}
+	SRVResource->Release();
+#endif
+}
+
 template <EShaderFrequency ShaderFrequency>
 void FD3D11DynamicRHI::InternalSetShaderResourceView(FD3D11BaseShaderResource* Resource, ID3D11ShaderResourceView* SRV, int32 ResourceIndex, FD3D11StateCache::ESRV_Type SrvType)
 {
 	// Check either both are set, or both are null.
 	check((Resource && SRV) || (!Resource && !SRV));
+	CheckIfSRVIsResolved(SRV);
 
 	FD3D11BaseShaderResource*& ResourceSlot = CurrentResourcesBoundAsSRVs[ShaderFrequency][ResourceIndex];
 	int32& MaxResourceIndex = MaxBoundShaderResourcesIndex[ShaderFrequency];
diff --git a/Engine/Source/Runtime/Windows/D3D11RHI/Private/D3D11RHIPrivate.h b/Engine/Source/Runtime/Windows/D3D11RHI/Private/D3D11RHIPrivate.h
--- a/Engine/Source/Runtime/Windows/D3D11RHI/Private/D3D11RHIPrivate.h
+++ b/Engine/Source/Runtime/Windows/D3D11RHI/Private/D3D11RHIPrivate.h
@@ -32,6 +32,12 @@ DECLARE_LOG_CATEGORY_EXTERN(LogD3D11RHI, Log, All);
 #include "D3D11ConstantBuffer.h"
 #include "D3D11StateCache.h"
 
+#if UE_BUILD_SHIPPING || UE_BUILD_TEST
+#define CHECK_SRV_TRANSITIONS 0
+#else
+#define CHECK_SRV_TRANSITIONS 1
+#endif
+
 // DX11 doesn't support higher MSAA count
 #define DX_MAX_MSAA_COUNT 8
 
@@ -385,6 +391,8 @@ private:
 	template <EShaderFrequency ShaderFrequency>
 	void ClearAllShaderResourcesForFrequency();
 
+	void CheckIfSRVIsResolved(ID3D11ShaderResourceView* SRV);
+
 	template <EShaderFrequency ShaderFrequency>
 	void InternalSetShaderResourceView(FD3D11BaseShaderResource* Resource, ID3D11ShaderResourceView* SRV, int32 ResourceIndex, FD3D11StateCache::ESRV_Type SrvType = FD3D11StateCache::SRV_Unknown);
 
@@ -517,6 +525,39 @@ protected:
 	TGlobalResource< TBoundShaderStateHistory<10000> > BoundShaderStateHistory;
 	FComputeShaderRHIRef CurrentComputeShader;
 
+#if CHECK_SRV_TRANSITIONS
+	/*
+	 * Rendertargets must be explicitly 'resolved' to manage their transition to an SRV on some platforms and DX12
+	 * We keep track of targets that need 'resolving' to provide safety asserts at SRV binding time.
+	 */
+	struct FUnresolvedRTInfo
+	{
+		FUnresolvedRTInfo(FName InResourceName, int32 InMipLevel, int32 InNumMips, int32 InArraySlice, int32 InArraySize)
+		: ResourceName(InResourceName)
+		, MipLevel(InMipLevel)
+		, NumMips(InNumMips)
+		, ArraySlice(InArraySlice)  
+		, ArraySize(InArraySize)
+		{
+		}
+
+		bool operator==(const FUnresolvedRTInfo& Other) const
+		{
+			return MipLevel == Other.MipLevel &&
+				NumMips == Other.NumMips &&
+				ArraySlice == Other.ArraySlice &&
+				ArraySize == Other.ArraySize;
+		}
+
+		FName ResourceName;
+		int32 MipLevel;
+		int32 NumMips;
+		int32 ArraySlice;
+		int32 ArraySize;
+	};
+	TMultiMap<ID3D11Resource*, FUnresolvedRTInfo> UnresolvedTargets;
+#endif
+
 	FD3DGPUProfiler GPUProfilingData;
 	// >= 0, was computed before, unless hardware was changed during engine init it should be the same
 	int32 ChosenAdapter;
diff --git a/Engine/Source/Runtime/Windows/D3D11RHI/Private/D3D11RenderTarget.cpp b/Engine/Source/Runtime/Windows/D3D11RHI/Private/D3D11RenderTarget.cpp
--- a/Engine/Source/Runtime/Windows/D3D11RHI/Private/D3D11RenderTarget.cpp
+++ b/Engine/Source/Runtime/Windows/D3D11RHI/Private/D3D11RenderTarget.cpp
@@ -322,10 +322,45 @@ void FD3D11DynamicRHI::RHICopyToResolveTarget(FTextureRHIParamRef SourceTextureR
 	FD3D11TextureCube* SourceTextureCube = static_cast<FD3D11TextureCube*>(SourceTextureRHI->GetTextureCube());
 	FD3D11TextureCube* DestTextureCube = static_cast<FD3D11TextureCube*>(DestTextureRHI->GetTextureCube());
 
+	FD3D11Texture3D* SourceTexture3D = static_cast<FD3D11Texture3D*>(SourceTextureRHI->GetTexture3D());
+	FD3D11Texture3D* DestTexture3D = static_cast<FD3D11Texture3D*>(DestTextureRHI->GetTexture3D());
+
+#if CHECK_SRV_TRANSITIONS
+	//check first, as some of the resolve shaders could trigger the resource we're currently resolving.
+	ID3D11Resource* SourceResource = nullptr;
+	int32 ResolveSlice = -1;
+	if (SourceTexture2D)
+	{
+		SourceResource = SourceTexture2D->GetResource();
+	}
+	if (SourceTextureCube)
+	{
+		SourceResource = SourceTextureCube->GetResource();
+		ResolveSlice = ResolveParams.SourceArrayIndex * 6 + ResolveParams.CubeFace;
+	}
+	if (SourceTexture3D)
+	{
+		SourceResource = SourceTexture3D->GetResource();
+	}
+
+	if (SourceResource)
+	{
+		// if we don't have specific resolve data remove all references.
+		// this is helpful when you clear many mips/slices but only want to do one resolve.
+		if (ResolveParams.MipIndex == -1 && ResolveParams.SourceArrayIndex == -1)
+		{
+			UnresolvedTargets.Remove(SourceResource);
+		}
+		else
+		{
+			UnresolvedTargets.Remove(SourceResource, FUnresolvedRTInfo(SourceTextureRHI->GetName(), ResolveParams.MipIndex, 1, ResolveSlice, 1));
+		}
+	}
+#endif
+		
 	if(SourceTexture2D && DestTexture2D)
 	{
 		check(!SourceTextureCube && !DestTextureCube);
-
 		if(SourceTexture2D != DestTexture2D)
 		{
 			GPUProfilingData.RegisterGPUWork();
@@ -415,12 +450,12 @@ void FD3D11DynamicRHI::RHICopyToResolveTarget(FTextureRHIParamRef SourceTextureR
 	}
 	else if(SourceTextureCube && DestTextureCube)
 	{
-		check(!SourceTexture2D && !DestTexture2D);
-		
+		check(!SourceTexture2D && !DestTexture2D);			
+
 		if(SourceTextureCube != DestTextureCube)
 		{
 			GPUProfilingData.RegisterGPUWork();
-		
+
 			// Determine the cubemap face being resolved.
 			const uint32 D3DFace = GetD3D11CubeFace(ResolveParams.CubeFace);
 			const uint32 SourceSubresource = D3D11CalcSubresource(ResolveParams.MipIndex, ResolveParams.SourceArrayIndex * 6 + D3DFace, SourceTextureCube->GetNumMips());
@@ -451,6 +486,12 @@ void FD3D11DynamicRHI::RHICopyToResolveTarget(FTextureRHIParamRef SourceTextureR
 		const uint32 Subresource = D3D11CalcSubresource(0, D3DFace, 1);
 		Direct3DDeviceIMContext->CopySubresourceRegion(DestTextureCube->GetResource(), Subresource, 0, 0, 0, SourceTexture2D->GetResource(), 0, NULL);
 	}
+	else if (SourceTexture3D && DestTexture3D)
+	{
+		// bit of a hack.  no one resolves slice by slice and 0 is the default value.  assume for the moment they are resolving the whole texture.
+		check(ResolveParams.SourceArrayIndex == 0);
+		check(SourceTexture3D == DestTexture3D);
+	}
 }
 
 /**
diff --git a/Engine/Source/Runtime/Windows/D3D11RHI/Private/D3D11Texture.cpp b/Engine/Source/Runtime/Windows/D3D11RHI/Private/D3D11Texture.cpp
--- a/Engine/Source/Runtime/Windows/D3D11RHI/Private/D3D11Texture.cpp
+++ b/Engine/Source/Runtime/Windows/D3D11RHI/Private/D3D11Texture.cpp
@@ -1492,6 +1492,9 @@ void FD3D11DynamicRHI::RHIUnlockTextureCubeFace(FTextureCubeRHIParamRef TextureC
 
 void FD3D11DynamicRHI::RHIBindDebugLabelName(FTextureRHIParamRef TextureRHI, const TCHAR* Name)
 {
+	//todo: require names at texture creation time.
+	FName DebugName(Name);
+	TextureRHI->SetName(DebugName);
 }
 
 void FD3D11DynamicRHI::RHIVirtualTextureSetFirstMipInMemory(FTexture2DRHIParamRef TextureRHI, uint32 FirstMip)
diff --git a/Engine/Source/Runtime/Windows/D3D11RHI/Private/D3D11Viewport.cpp b/Engine/Source/Runtime/Windows/D3D11RHI/Private/D3D11Viewport.cpp
--- a/Engine/Source/Runtime/Windows/D3D11RHI/Private/D3D11Viewport.cpp
+++ b/Engine/Source/Runtime/Windows/D3D11RHI/Private/D3D11Viewport.cpp
@@ -512,7 +512,7 @@ void FD3D11DynamicRHI::RHIBeginDrawingViewport(FViewportRHIParamRef ViewportRHI,
 		RenderTarget = Viewport->GetBackBuffer();
 	}
 	FRHIRenderTargetView View(RenderTarget);
-	RHISetRenderTargets(1,&View,FTextureRHIRef(),0,NULL);
+	RHISetRenderTargets(1,&View,nullptr,0,NULL);
 
 	// Set an initially disabled scissor rect.
 	RHISetScissorRect(false,0,0,0,0);
@@ -587,6 +587,10 @@ void FD3D11DynamicRHI::RHIEndDrawingViewport(FViewportRHIParamRef ViewportRHI,bo
 			GInputLatencyTimer.RenderThreadTrigger = false;
 		}
 	}
+
+#if CHECK_SRV_TRANSITIONS
+	UnresolvedTargets.Reset();
+#endif
 }
 
 /**
