Adding a new experimental DOF algorithm called "CircleDOF".
This is work-in-progress and requires more work to get quality/perf up to par.
The "FocalDistance" controls the focal point, the "Scale" controls the how big the circle Bokeh gets (which is limited to small radius now).
There is no near DOF currently.

diff --git a/Engine/Shaders/PostProcessDOF.usf b/Engine/Shaders/PostProcessDOF.usf
--- a/Engine/Shaders/PostProcessDOF.usf
+++ b/Engine/Shaders/PostProcessDOF.usf
@@ -191,3 +191,361 @@ void MainRecombinePS(
 	OutColor.a = 0;
 }
 
+
+
+
+//
+// PROTOTYPE CIRCLE DOF : WORK IN PROGRESS
+//
+
+// CIRCLE DOF: Compute circle of confusion size in pixels.
+float DepthToCoc(float Depth) 
+{
+	float Focus = View.DepthOfFieldFocalDistance;
+	float Radius = View.DepthOfFieldScale * 2.0;
+
+	// Returning only FAR COC for now.
+	return max(0.0,(Depth - Focus) / Depth) * Radius;
+}
+
+// pixel shader entry point
+void CircleSetupPS(
+	float4 UVAndScreenPos : TEXCOORD0
+	, out float4 OutColor0 : SV_Target0
+#if ENABLE_NEAR_BLUR 
+	, out float4 OutColor1 : SV_Target1
+#endif
+	)
+{
+#if ENABLE_NEAR_BLUR == 0
+	// We aren't writing out to the second render target, so we'll just make a dummy value here which 
+	// doesn't end up going anywhere. Then, the source code can stay neat and tidy, while the compiler can still 
+	// strip it out
+	float4 OutColor1;
+#endif
+
+	float2 UV = UVAndScreenPos.xy;
+
+	float4 DepthQuad = GatherSceneDepth(UV, PostprocessInput1Size.zw);
+
+	UV = UVAndScreenPos.xy - 0.5*PostprocessInput0Size.zw;
+
+	float4 CW = PostprocessInput0.SampleLevel(PostprocessInput0Sampler, UV.xy, 0);
+	float4 CZ = PostprocessInput0.SampleLevel(PostprocessInput0Sampler, UV.xy, 0, int2(1,0));
+	float4 CX = PostprocessInput0.SampleLevel(PostprocessInput0Sampler, UV.xy, 0, int2(0,1));
+	float4 CY = PostprocessInput0.SampleLevel(PostprocessInput0Sampler, UV.xy, 0, int2(1,1));
+
+	float4 CocQuad = float4(DepthToCoc(DepthQuad.x), DepthToCoc(DepthQuad.y), DepthToCoc(DepthQuad.z), DepthToCoc(DepthQuad.w));
+
+	// Doing a max depth reduction (erode the foreground). Less correct, but less artifacts.
+	// Perhaps need to re-open this in the future.
+
+	// Stuff max radius in alpha.
+	OutColor0.a = max(max(CocQuad.x,CocQuad.y),max(CocQuad.z,CocQuad.w));
+
+	// Remove samples which are outside the size.
+	// TODO: Tune the ScaleFactor.
+	float ScaleFactor = 1.0;
+	float4 W = float4(
+		1.0-saturate((OutColor0.a - CocQuad.x) * ScaleFactor),
+		1.0-saturate((OutColor0.a - CocQuad.y) * ScaleFactor),
+		1.0-saturate((OutColor0.a - CocQuad.z) * ScaleFactor),
+		1.0-saturate((OutColor0.a - CocQuad.w) * ScaleFactor));
+
+	OutColor0.rgb = (1.0/(W.x+W.y+W.z+W.w)) * (CX.rgb*W.x + CY.rgb*W.y + CZ.rgb*W.z + CW.rgb*W.w);
+
+	// TODO: Near.
+	OutColor1 = 0;
+}
+
+
+
+// {0 to 1} output.
+float NoizNorm(float2 N, float X)
+{
+	N+=X;
+	return frac(sin(dot(N.xy,float2(12.9898, 78.233)))*43758.5453);
+}
+
+// {-1 to 1} output.
+float NoizSnorm(float2 N, float X)
+{
+	return NoizNorm(N,X)*2.0-1.0;
+}
+
+float2 RotVec(float Radius, float Radians)
+{
+	return Radius * float2(cos(Radians), sin(Radians));
+}
+
+float2 RandomOffset;
+
+// pixel shader entry point
+void CirclePS(
+	float4 UVAndScreenPos : TEXCOORD0
+	, out float4 OutColor0 : SV_Target0
+#if ENABLE_NEAR_BLUR 
+	, out float4 OutColor1 : SV_Target1
+#endif
+	)
+{
+	float2 UV = UVAndScreenPos.xy;
+
+#if ENABLE_NEAR_BLUR == 0
+	// We aren't writing out to the second render target, so we'll just make a dummy value here which 
+	// doesn't end up going anywhere. Then, the source code can stay neat and tidy, while the compiler can still 
+	// strip it out
+	float4 OutColor1;
+#endif
+
+	// TODO: Near.
+	OutColor1 = 0.0;
+
+	//
+	// TODO: This part of the algorithm is a placeholder for something better.
+	//       Current placeholder is a stochastic algorithm.
+	//
+
+	// Going to grab sets of 4 samples.
+	// Each set of 4 samples can be a smaller circle of confusion
+	// (aka can be in-front of the larger background).
+
+	//
+	// Pass 0
+	//
+
+	float TwoPi = 2.0 * 3.14159;
+
+	// Grab circle of confusion for the pixel.
+	OutColor0 = Texture2DSampleLevel(PostprocessInput0, PostprocessInput0Sampler, UV, 0);
+	float Coc = OutColor0.a + (1.0/65536.0);
+
+	float RadianBase = NoizSnorm(UVAndScreenPos.xy, 0.010 * RandomOffset.x) * TwoPi;
+	float RadiusBase = NoizNorm(UVAndScreenPos.xy, 0.013 * RandomOffset.x) * (1.0/11.5);
+
+	// Bring out to the smaller radius of sample sets.
+	// This has the highest chance of seeing a smaller overlapping CoC.
+	float R1 = (RadiusBase+9.0/11.5) * Coc;
+	float R2 = (RadiusBase+3.0/11.5) * Coc;
+	float R3 = (RadiusBase+6.0/11.5) * Coc;
+	float R4 = (RadiusBase+0.0/11.5) * Coc;
+
+	// Ensure at least getting different than center pixel.
+	float R1a = max(1.0,R1);
+	float R2a = max(1.0,R2);
+	float R3a = max(1.0,R3);
+	float R4a = max(1.0,R4);
+
+	float2 UV1 = RotVec(R1a, RadianBase + TwoPi * 0.0/12.0);
+	float2 UV2 = RotVec(R2a, RadianBase + TwoPi * 3.0/12.0);
+	float2 UV3 = RotVec(R3a, RadianBase + TwoPi * 6.0/12.0);
+	float2 UV4 = RotVec(R4a, RadianBase + TwoPi * 9.0/12.0);
+
+	UV1 = UVAndScreenPos.xy + UV1 * PostprocessInput0Size.zw;
+	UV2 = UVAndScreenPos.xy + UV2 * PostprocessInput0Size.zw;
+	UV3 = UVAndScreenPos.xy + UV3 * PostprocessInput0Size.zw;
+	UV4 = UVAndScreenPos.xy + UV4 * PostprocessInput0Size.zw;
+
+	float4 C1 = PostprocessInput0.SampleLevel(PostprocessInput0Sampler, UV1, 0);
+	float4 C2 = PostprocessInput0.SampleLevel(PostprocessInput0Sampler, UV2, 0);
+	float4 C3 = PostprocessInput0.SampleLevel(PostprocessInput0Sampler, UV3, 0);
+	float4 C4 = PostprocessInput0.SampleLevel(PostprocessInput0Sampler, UV4, 0);
+
+	// Base weight works around the max(1.0,radius) constraint.
+	// Base weight also shapes to weight higher on the outside radius.
+	float W0 = 1.0 - saturate(Coc);
+	float W1 = R1;
+	float W2 = R2;
+	float W3 = R3;
+	float W4 = R4;
+
+	// Adjusting weight to remove occluded.
+	W1 *= saturate(0.0 + C1.a - R1);
+	W2 *= saturate(0.0 + C2.a - R2);
+	W3 *= saturate(0.0 + C3.a - R3);
+	W4 *= saturate(0.0 + C4.a - R4);
+
+	// Make sure at least something is not zero.
+	W0 += 1.0/65536.0;
+
+	float Coc2 = Coc;
+	if(W1 > 0.0) Coc2 = min(Coc2, C1.a);
+	if(W2 > 0.0) Coc2 = min(Coc2, C2.a);
+	if(W3 > 0.0) Coc2 = min(Coc2, C3.a);
+	if(W4 > 0.0) Coc2 = min(Coc2, C4.a);
+
+	OutColor0.rgb = OutColor0.rgb * W0 + C1.rgb * W1 + C2.rgb * W2 + C3.rgb * W3 + C4.rgb * W4;
+
+	float Weight = W0+W1+W2+W3+W4;
+
+	float3 Background = OutColor0.rgb * (1.0/Weight);
+
+
+	//
+	// Pass 1
+	//
+
+	// Drop weight of existing pass if Coc changes too much.
+	float Drop = (1.0/65536.0) + 1.0 - saturate(abs(Coc - Coc2));
+	OutColor0.rgb *= Drop;
+	Weight *= Drop;
+
+	R1 = (RadiusBase+10.0/11.5) * Coc2;
+	R2 = (RadiusBase+ 4.0/11.5) * Coc2;
+	R3 = (RadiusBase+ 7.0/11.5) * Coc2;
+	R4 = (RadiusBase+ 1.0/11.5) * Coc2;
+
+	R1a = max(1.0,R1);
+	R2a = max(1.0,R2);
+	R3a = max(1.0,R3);
+	R4a = max(1.0,R4);
+
+	UV1 = RotVec(R1a, RadianBase + TwoPi *  8.0/12.0);
+	UV2 = RotVec(R2a, RadianBase + TwoPi * 11.0/12.0);
+	UV3 = RotVec(R3a, RadianBase + TwoPi *  2.0/12.0);
+	UV4 = RotVec(R4a, RadianBase + TwoPi *  5.0/12.0);
+
+	UV1 = UVAndScreenPos.xy + UV1 * PostprocessInput0Size.zw;
+	UV2 = UVAndScreenPos.xy + UV2 * PostprocessInput0Size.zw;
+	UV3 = UVAndScreenPos.xy + UV3 * PostprocessInput0Size.zw;
+	UV4 = UVAndScreenPos.xy + UV4 * PostprocessInput0Size.zw;
+
+	C1 = PostprocessInput0.SampleLevel(PostprocessInput0Sampler, UV1, 0);
+	C2 = PostprocessInput0.SampleLevel(PostprocessInput0Sampler, UV2, 0);
+	C3 = PostprocessInput0.SampleLevel(PostprocessInput0Sampler, UV3, 0);
+	C4 = PostprocessInput0.SampleLevel(PostprocessInput0Sampler, UV4, 0);
+
+	C1.rgb = lerp(C1.rgb, Background.rgb, saturate(C1.a - Coc2));
+	C2.rgb = lerp(C2.rgb, Background.rgb, saturate(C2.a - Coc2));
+	C3.rgb = lerp(C3.rgb, Background.rgb, saturate(C3.a - Coc2));
+	C4.rgb = lerp(C4.rgb, Background.rgb, saturate(C4.a - Coc2));
+
+	W1 = R1;
+	W2 = R2;
+	W3 = R3;
+	W4 = R4;
+
+	// Adjusting weight to remove occluded.
+	W1 *= saturate(0.0 + C1.a - R1);
+	W2 *= saturate(0.0 + C2.a - R2);
+	W3 *= saturate(0.0 + C3.a - R3);
+	W4 *= saturate(0.0 + C4.a - R4);
+
+	float Coc3 = Coc2;
+	if(W1 > 0.0) Coc3 = min(Coc3, C1.a);
+	if(W2 > 0.0) Coc3 = min(Coc3, C2.a);
+	if(W3 > 0.0) Coc3 = min(Coc3, C3.a);
+	if(W4 > 0.0) Coc3 = min(Coc3, C4.a);
+
+	OutColor0.rgb += C1.rgb * W1 + C2.rgb * W2 + C3.rgb * W3 + C4.rgb * W4;
+	Weight += W1+W2+W3+W4;
+
+
+
+
+	//
+	// Pass 2
+	//
+
+	// Drop weight of existing pass if Coc changes too much.
+	Drop = (1.0/65536.0) + 1.0 - saturate(abs(Coc2 - Coc3));
+	OutColor0.rgb *= Drop;
+	Weight *= Drop;
+
+	R1 = (RadiusBase+11.0/11.5) * Coc3;
+	R2 = (RadiusBase+ 5.0/11.5) * Coc3;
+	R3 = (RadiusBase+ 8.0/11.5) * Coc3;
+	R4 = (RadiusBase+ 2.0/11.5) * Coc3;
+
+	R1a = max(1.0,R1);
+	R2a = max(1.0,R2);
+	R3a = max(1.0,R3);
+	R4a = max(1.0,R4);
+
+	UV1 = RotVec(R1a, RadianBase + TwoPi *  4.0/12.0);
+	UV2 = RotVec(R2a, RadianBase + TwoPi *  7.0/12.0);
+	UV3 = RotVec(R3a, RadianBase + TwoPi * 10.0/12.0);
+	UV4 = RotVec(R4a, RadianBase + TwoPi *  1.0/12.0);
+
+	UV1 = UVAndScreenPos.xy + UV1 * PostprocessInput0Size.zw;
+	UV2 = UVAndScreenPos.xy + UV2 * PostprocessInput0Size.zw;
+	UV3 = UVAndScreenPos.xy + UV3 * PostprocessInput0Size.zw;
+	UV4 = UVAndScreenPos.xy + UV4 * PostprocessInput0Size.zw;
+
+	C1 = PostprocessInput0.SampleLevel(PostprocessInput0Sampler, UV1, 0);
+	C2 = PostprocessInput0.SampleLevel(PostprocessInput0Sampler, UV2, 0);
+	C3 = PostprocessInput0.SampleLevel(PostprocessInput0Sampler, UV3, 0);
+	C4 = PostprocessInput0.SampleLevel(PostprocessInput0Sampler, UV4, 0);
+
+	C1.rgb = lerp(C1.rgb, Background.rgb, saturate(C1.a - Coc3));
+	C2.rgb = lerp(C2.rgb, Background.rgb, saturate(C2.a - Coc3));
+	C3.rgb = lerp(C3.rgb, Background.rgb, saturate(C3.a - Coc3));
+	C4.rgb = lerp(C4.rgb, Background.rgb, saturate(C4.a - Coc3));
+
+	W1 = R1;
+	W2 = R2;
+	W3 = R3;
+	W4 = R4;
+
+	// Adjusting weight to remove occluded.
+	W1 *= saturate(0.0 + C1.a - R1);
+	W2 *= saturate(0.0 + C2.a - R2);
+	W3 *= saturate(0.0 + C3.a - R3);
+	W4 *= saturate(0.0 + C4.a - R4);
+
+	OutColor0.rgb += C1.rgb * W1 + C2.rgb * W2 + C3.rgb * W3 + C4.rgb * W4;
+	Weight += W1+W2+W3+W4;
+
+	OutColor0.rgb *= (1.0/Weight);
+
+}
+
+
+// pixel shader to combine the full res scene and the blurred images behind and in front of the the focal plane
+void MainCircleRecombinePS(
+	in float4 UVAndScreenPos : TEXCOORD0,
+	out float4 OutColor : SV_Target0
+	)
+{
+	// Circle of confusion size for the pixel.
+	float PixDepth = CalcSceneDepth(UVAndScreenPos.xy);
+	float PixCoc = DepthToCoc(PixDepth) * 2.0; // 2x because full instead of half resolution.
+
+	// Fetch 2 samples mirrored around the pixel 
+	// which is stochastically distributed to fill out the circle of confusion.
+	// TODO: Fix the "random values".
+	// The sqrt(sqrt()) is to shape toward higher weight on the CoC edge.
+	float2 UV = UVAndScreenPos.xy * PostprocessInput0Size.xy;
+	float RadianBase = NoizNorm(UVAndScreenPos.xy, 0.010 * RandomOffset.x) * 3.14159 * 2.0;
+	float RadiusJitter = NoizNorm(UVAndScreenPos.xy, 0.013 * RandomOffset.x);
+	float2 VP = RotVec(float2(PixCoc*sqrt(sqrt(RadiusJitter)), 0.0), RadianBase) * PostprocessInput0Size.zw;
+
+	// These two samples will still have jitter induced artifacts (very limited utility).
+	// These two samples will also have bleeding artifacts.
+	float4 CA = PostprocessInput0.SampleLevel(PostprocessInput0Sampler, UVAndScreenPos.xy + VP, 0);
+	float4 CB = PostprocessInput0.SampleLevel(PostprocessInput0Sampler, UVAndScreenPos.xy - VP, 0);
+	OutColor = (CA + CB) * 0.5;
+
+	// Grab the half resolution neighborhood to remove the artifacts from the full resolution output.
+	// Nearest location.
+	float2 HUVBase = UVAndScreenPos.xy * PostprocessInput1Size.xy - 0.5;
+	float2 HUVFrac = frac(HUVBase);
+	float2 HUV = (trunc(HUVBase) + 0.5) * PostprocessInput1Size.zw;
+
+	// Load four nearest samples.
+	float4 H0 = PostprocessInput1.SampleLevel(PostprocessInput1Sampler, HUV, 0);
+	float4 H1 = PostprocessInput1.SampleLevel(PostprocessInput1Sampler, HUV, 0, int2(1,0));
+	float4 H2 = PostprocessInput1.SampleLevel(PostprocessInput1Sampler, HUV, 0, int2(0,1));
+	float4 H3 = PostprocessInput1.SampleLevel(PostprocessInput1Sampler, HUV, 0, int2(1,1));
+
+	// TODO: This would work a lot better in YUV style colorspace?
+	// Limit the full resolution to remove jitter artifacts.
+	float4 HMax = max(max(H0,H1),max(H2,H3));
+	float4 HMin = min(min(H0,H1),min(H2,H3));
+
+	// Blend in the limited version quickly to remove HDR jitter artifacts and noise.
+	float4 OutLimited = min(max(OutColor,HMin),HMax);
+	OutColor = lerp(OutColor, OutLimited, saturate(PixCoc*PixCoc*8.0));
+}
+
+
diff --git a/Engine/Source/Runtime/Engine/Classes/Engine/Scene.h b/Engine/Source/Runtime/Engine/Classes/Engine/Scene.h
--- a/Engine/Source/Runtime/Engine/Classes/Engine/Scene.h
+++ b/Engine/Source/Runtime/Engine/Classes/Engine/Scene.h
@@ -13,6 +13,7 @@ enum EDepthOfFieldMethod
 {
 	DOFM_BokehDOF UMETA(DisplayName="BokehDOF"),
 	DOFM_Gaussian UMETA(DisplayName="Gaussian"),
+	DOFM_CircleDOF UMETA(DisplayName="CircleDOF"),
 	DOFM_MAX,
 };
 
diff --git a/Engine/Source/Runtime/Renderer/Private/PostProcess/PostProcessDOF.cpp b/Engine/Source/Runtime/Renderer/Private/PostProcess/PostProcessDOF.cpp
--- a/Engine/Source/Runtime/Renderer/Private/PostProcess/PostProcessDOF.cpp
+++ b/Engine/Source/Runtime/Renderer/Private/PostProcess/PostProcessDOF.cpp
@@ -346,3 +346,555 @@ FPooledRenderTargetDesc FRCPassPostProcessDOFRecombine::ComputeOutputDesc(EPassO
 
 	return Ret;
 }
+
+
+
+
+
+/** Encapsulates the Circle DOF setup pixel shader. */
+template <uint32 NearBlurEnable>
+class FPostProcessCircleDOFSetupPS : public FGlobalShader
+{
+	DECLARE_SHADER_TYPE(FPostProcessCircleDOFSetupPS, Global);
+
+	static bool ShouldCache(EShaderPlatform Platform)
+	{
+		return IsFeatureLevelSupported(Platform, ERHIFeatureLevel::SM4);
+	}
+
+	static void ModifyCompilationEnvironment(EShaderPlatform Platform, FShaderCompilerEnvironment& OutEnvironment)
+	{
+		FGlobalShader::ModifyCompilationEnvironment(Platform,OutEnvironment);
+		OutEnvironment.SetDefine(TEXT("ENABLE_NEAR_BLUR"), NearBlurEnable);
+	}
+
+	/** Default constructor. */
+	FPostProcessCircleDOFSetupPS() {}
+
+public:
+	FPostProcessPassParameters PostprocessParameter;
+	FDeferredPixelShaderParameters DeferredParameters;
+	FShaderParameter DepthOfFieldParams;
+
+	/** Initialization constructor. */
+	FPostProcessCircleDOFSetupPS(const ShaderMetaType::CompiledShaderInitializerType& Initializer)
+		: FGlobalShader(Initializer)
+	{
+		PostprocessParameter.Bind(Initializer.ParameterMap);
+		DeferredParameters.Bind(Initializer.ParameterMap);
+		DepthOfFieldParams.Bind(Initializer.ParameterMap,TEXT("DepthOfFieldParams"));
+	}
+
+	// FShader interface.
+	virtual bool Serialize(FArchive& Ar)
+	{
+		bool bShaderHasOutdatedParameters = FGlobalShader::Serialize(Ar);
+		Ar << PostprocessParameter << DeferredParameters << DepthOfFieldParams;
+		return bShaderHasOutdatedParameters;
+	}
+
+	void SetParameters(const FRenderingCompositePassContext& Context)
+	{
+		const FPixelShaderRHIParamRef ShaderRHI = GetPixelShader();
+
+		FGlobalShader::SetParameters(Context.RHICmdList, ShaderRHI, Context.View);
+
+		PostprocessParameter.SetPS(ShaderRHI, Context, TStaticSamplerState<SF_Point,AM_Border,AM_Border,AM_Clamp>::GetRHI());
+
+		DeferredParameters.Set(Context.RHICmdList, ShaderRHI, Context.View);
+
+		{
+			FVector4 DepthOfFieldParamValues[2];
+
+			FRCPassPostProcessBokehDOF::ComputeDepthOfFieldParams(Context, DepthOfFieldParamValues);
+
+			SetShaderValueArray(Context.RHICmdList, ShaderRHI, DepthOfFieldParams, DepthOfFieldParamValues, 2);
+		}
+	}
+};
+
+IMPLEMENT_SHADER_TYPE(template<>,FPostProcessCircleDOFSetupPS<0>,TEXT("PostProcessDOF"),TEXT("CircleSetupPS"),SF_Pixel);
+IMPLEMENT_SHADER_TYPE(template<>,FPostProcessCircleDOFSetupPS<1>,TEXT("PostProcessDOF"),TEXT("CircleSetupPS"),SF_Pixel);
+
+void FRCPassPostProcessCircleDOFSetup::Process(FRenderingCompositePassContext& Context)
+{
+	SCOPED_DRAW_EVENT(Context.RHICmdList, DOFSetup);
+
+	const FPooledRenderTargetDesc* InputDesc = GetInputDesc(ePId_Input0);
+
+	if(!InputDesc)
+	{
+		// input is not hooked up correctly
+		return;
+	}
+
+	uint32 NumRenderTargets = bNearBlurEnabled ? 2 : 1;
+
+	const FSceneView& View = Context.View;
+	const FSceneViewFamily& ViewFamily = *(View.Family);
+
+	const auto FeatureLevel = Context.GetFeatureLevel();
+	auto ShaderMap = Context.GetShaderMap();
+
+	FIntPoint SrcSize = InputDesc->Extent;
+	FIntPoint DestSize = PassOutputs[0].RenderTargetDesc.Extent;
+
+	// e.g. 4 means the input texture is 4x smaller than the buffer size
+	uint32 ScaleFactor = GSceneRenderTargets.GetBufferSizeXY().X / SrcSize.X;
+
+	FIntRect SrcRect = View.ViewRect / ScaleFactor;
+	FIntRect DestRect = SrcRect / 2;
+
+	const FSceneRenderTargetItem& DestRenderTarget0 = PassOutputs[0].RequestSurface(Context);
+	const FSceneRenderTargetItem& DestRenderTarget1 = bNearBlurEnabled ? PassOutputs[1].RequestSurface(Context) : FSceneRenderTargetItem();
+
+	// Set the view family's render target/viewport.
+	FTextureRHIParamRef RenderTargets[2] =
+	{
+		DestRenderTarget0.TargetableTexture,
+		DestRenderTarget1.TargetableTexture
+	};
+	SetRenderTargets(Context.RHICmdList, NumRenderTargets, RenderTargets, FTextureRHIParamRef(), 0, NULL);
+
+	FLinearColor ClearColors[2] = 
+	{
+		FLinearColor(0, 0, 0, 0),
+		FLinearColor(0, 0, 0, 0)
+	};
+	// is optimized away if possible (RT size=view size, )
+	Context.RHICmdList.ClearMRT(true, NumRenderTargets, ClearColors, false, 1.0f, false, 0, DestRect);
+
+	Context.SetViewportAndCallRHI(0, 0, 0.0f, DestSize.X, DestSize.Y, 1.0f );
+
+	// set the state
+	Context.RHICmdList.SetBlendState(TStaticBlendState<>::GetRHI());
+	Context.RHICmdList.SetRasterizerState(TStaticRasterizerState<>::GetRHI());
+	Context.RHICmdList.SetDepthStencilState(TStaticDepthStencilState<false, CF_Always>::GetRHI());
+
+	TShaderMapRef<FPostProcessVS> VertexShader(ShaderMap);
+
+	if (bNearBlurEnabled)
+	{
+		static FGlobalBoundShaderState BoundShaderState;
+
+
+		TShaderMapRef< FPostProcessCircleDOFSetupPS<1> > PixelShader(ShaderMap);
+		SetGlobalBoundShaderState(Context.RHICmdList, FeatureLevel, BoundShaderState, GFilterVertexDeclaration.VertexDeclarationRHI, *VertexShader, *PixelShader);
+
+		PixelShader->SetParameters(Context);
+	}
+	else
+	{
+		static FGlobalBoundShaderState BoundShaderState;
+
+		TShaderMapRef< FPostProcessCircleDOFSetupPS<0> > PixelShader(ShaderMap);
+		SetGlobalBoundShaderState(Context.RHICmdList, FeatureLevel, BoundShaderState, GFilterVertexDeclaration.VertexDeclarationRHI, *VertexShader, *PixelShader);
+
+		PixelShader->SetParameters(Context);
+	}
+
+	VertexShader->SetParameters(Context);
+
+	// Draw a quad mapping scene color to the view's render target
+	DrawRectangle(
+		Context.RHICmdList,
+		DestRect.Min.X, DestRect.Min.Y,
+		DestRect.Width() + 1, DestRect.Height() + 1,
+		SrcRect.Min.X, SrcRect.Min.Y,
+		SrcRect.Width() + 1, SrcRect.Height() + 1,
+		DestSize,
+		SrcSize,
+		*VertexShader,
+		EDRF_UseTriangleOptimization);
+
+	Context.RHICmdList.CopyToResolveTarget(DestRenderTarget0.TargetableTexture, DestRenderTarget0.ShaderResourceTexture, false, FResolveParams());
+	Context.RHICmdList.CopyToResolveTarget(DestRenderTarget1.TargetableTexture, DestRenderTarget1.ShaderResourceTexture, false, FResolveParams());
+}
+
+FPooledRenderTargetDesc FRCPassPostProcessCircleDOFSetup::ComputeOutputDesc(EPassOutputId InPassOutputId) const
+{
+	FPooledRenderTargetDesc Ret = PassInputs[0].GetOutput()->RenderTargetDesc;
+
+	//	Ret.Extent = FIntPoint::DivideAndRoundUp(ret.Extent, 2);
+	Ret.Extent /= 2;
+	Ret.Extent.X = FMath::Max(1, Ret.Extent.X);
+	Ret.Extent.Y = FMath::Max(1, Ret.Extent.Y);
+
+	Ret.Reset();
+	Ret.TargetableFlags &= ~(uint32)TexCreate_UAV;
+	Ret.TargetableFlags |= TexCreate_RenderTargetable;
+
+	Ret.DebugName = (InPassOutputId == ePId_Output0) ? TEXT("CircleDOFSetup0") : TEXT("CircleDOFSetup1");
+
+	// more precision for additive blending and we need the alpha channel
+	Ret.Format = PF_FloatRGBA;
+
+	return Ret;
+}
+
+
+
+
+/** Encapsulates the Circle DOF pixel shader. */
+
+static float TemporalHalton2( int32 Index, int32 Base )
+{
+	float Result = 0.0f;
+	float InvBase = 1.0f / Base;
+	float Fraction = InvBase;
+	while( Index > 0 )
+	{
+		Result += ( Index % Base ) * Fraction;
+		Index /= Base;
+		Fraction *= InvBase;
+	}
+	return Result;
+}
+
+static void TemporalRandom2(FVector2D* RESTRICT const Constant, uint32 FrameNumber)
+{
+	Constant->X = TemporalHalton2(FrameNumber & 1023, 2);
+	Constant->Y = TemporalHalton2(FrameNumber & 1023, 3);
+}
+
+
+template <uint32 NearBlurEnable>
+class FPostProcessCircleDOFPS : public FGlobalShader
+{
+	DECLARE_SHADER_TYPE(FPostProcessCircleDOFPS, Global);
+
+	static bool ShouldCache(EShaderPlatform Platform)
+	{
+		return IsFeatureLevelSupported(Platform, ERHIFeatureLevel::SM4);
+	}
+
+	static void ModifyCompilationEnvironment(EShaderPlatform Platform, FShaderCompilerEnvironment& OutEnvironment)
+	{
+		FGlobalShader::ModifyCompilationEnvironment(Platform,OutEnvironment);
+		OutEnvironment.SetDefine(TEXT("ENABLE_NEAR_BLUR"), NearBlurEnable);
+	}
+
+	/** Default constructor. */
+	FPostProcessCircleDOFPS() {}
+
+public:
+	FPostProcessPassParameters PostprocessParameter;
+	FDeferredPixelShaderParameters DeferredParameters;
+	FShaderParameter DepthOfFieldParams;
+	FShaderParameter RandomOffset;
+
+	/** Initialization constructor. */
+	FPostProcessCircleDOFPS(const ShaderMetaType::CompiledShaderInitializerType& Initializer)
+		: FGlobalShader(Initializer)
+	{
+		PostprocessParameter.Bind(Initializer.ParameterMap);
+		DeferredParameters.Bind(Initializer.ParameterMap);
+		DepthOfFieldParams.Bind(Initializer.ParameterMap,TEXT("DepthOfFieldParams"));
+		RandomOffset.Bind(Initializer.ParameterMap, TEXT("RandomOffset"));
+	}
+
+	// FShader interface.
+	virtual bool Serialize(FArchive& Ar)
+	{
+		bool bShaderHasOutdatedParameters = FGlobalShader::Serialize(Ar);
+		Ar << PostprocessParameter << DeferredParameters << DepthOfFieldParams << RandomOffset;
+		return bShaderHasOutdatedParameters;
+	}
+
+	void SetParameters(const FRenderingCompositePassContext& Context)
+	{
+		const FPixelShaderRHIParamRef ShaderRHI = GetPixelShader();
+
+		FGlobalShader::SetParameters(Context.RHICmdList, ShaderRHI, Context.View);
+
+		PostprocessParameter.SetPS(ShaderRHI, Context, TStaticSamplerState<SF_Point,AM_Border,AM_Border,AM_Clamp>::GetRHI());
+
+		DeferredParameters.Set(Context.RHICmdList, ShaderRHI, Context.View);
+
+		{
+			FVector4 DepthOfFieldParamValues[2];
+
+			FRCPassPostProcessBokehDOF::ComputeDepthOfFieldParams(Context, DepthOfFieldParamValues);
+
+			SetShaderValueArray(Context.RHICmdList, ShaderRHI, DepthOfFieldParams, DepthOfFieldParamValues, 2);
+		}
+
+		FVector2D RandomOffsetValue;
+		TemporalRandom2(&RandomOffsetValue, Context.View.Family->FrameNumber);
+		SetShaderValue(Context.RHICmdList, ShaderRHI, RandomOffset, RandomOffsetValue);
+	}
+};
+
+IMPLEMENT_SHADER_TYPE(template<>,FPostProcessCircleDOFPS<0>,TEXT("PostProcessDOF"),TEXT("CirclePS"),SF_Pixel);
+IMPLEMENT_SHADER_TYPE(template<>,FPostProcessCircleDOFPS<1>,TEXT("PostProcessDOF"),TEXT("CirclePS"),SF_Pixel);
+
+void FRCPassPostProcessCircleDOF::Process(FRenderingCompositePassContext& Context)
+{
+	SCOPED_DRAW_EVENT(Context.RHICmdList, DOFSetup);
+
+	const FPooledRenderTargetDesc* InputDesc = GetInputDesc(ePId_Input0);
+
+	if(!InputDesc)
+	{
+		// input is not hooked up correctly
+		return;
+	}
+
+	uint32 NumRenderTargets = bNearBlurEnabled ? 2 : 1;
+
+	const FSceneView& View = Context.View;
+	const FSceneViewFamily& ViewFamily = *(View.Family);
+
+	const auto FeatureLevel = Context.GetFeatureLevel();
+	auto ShaderMap = Context.GetShaderMap();
+
+	FIntPoint SrcSize = InputDesc->Extent;
+	FIntPoint DestSize = PassOutputs[0].RenderTargetDesc.Extent;
+
+	// e.g. 4 means the input texture is 4x smaller than the buffer size
+	uint32 ScaleFactor = GSceneRenderTargets.GetBufferSizeXY().X / SrcSize.X;
+
+	FIntRect SrcRect = View.ViewRect / ScaleFactor;
+	FIntRect DestRect = SrcRect;
+
+	const FSceneRenderTargetItem& DestRenderTarget0 = PassOutputs[0].RequestSurface(Context);
+	const FSceneRenderTargetItem& DestRenderTarget1 = bNearBlurEnabled ? PassOutputs[1].RequestSurface(Context) : FSceneRenderTargetItem();
+
+	// Set the view family's render target/viewport.
+	FTextureRHIParamRef RenderTargets[2] =
+	{
+		DestRenderTarget0.TargetableTexture,
+		DestRenderTarget1.TargetableTexture
+	};
+	SetRenderTargets(Context.RHICmdList, NumRenderTargets, RenderTargets, FTextureRHIParamRef(), 0, NULL);
+
+	FLinearColor ClearColors[2] = 
+	{
+		FLinearColor(0, 0, 0, 0),
+		FLinearColor(0, 0, 0, 0)
+	};
+	// is optimized away if possible (RT size=view size, )
+	Context.RHICmdList.ClearMRT(true, NumRenderTargets, ClearColors, false, 1.0f, false, 0, DestRect);
+
+	Context.SetViewportAndCallRHI(0, 0, 0.0f, DestSize.X, DestSize.Y, 1.0f );
+
+	// set the state
+	Context.RHICmdList.SetBlendState(TStaticBlendState<>::GetRHI());
+	Context.RHICmdList.SetRasterizerState(TStaticRasterizerState<>::GetRHI());
+	Context.RHICmdList.SetDepthStencilState(TStaticDepthStencilState<false, CF_Always>::GetRHI());
+
+	TShaderMapRef<FPostProcessVS> VertexShader(ShaderMap);
+
+	if (bNearBlurEnabled)
+	{
+		static FGlobalBoundShaderState BoundShaderState;
+
+
+		TShaderMapRef< FPostProcessCircleDOFPS<1> > PixelShader(ShaderMap);
+		SetGlobalBoundShaderState(Context.RHICmdList, FeatureLevel, BoundShaderState, GFilterVertexDeclaration.VertexDeclarationRHI, *VertexShader, *PixelShader);
+
+		PixelShader->SetParameters(Context);
+	}
+	else
+	{
+		static FGlobalBoundShaderState BoundShaderState;
+
+		TShaderMapRef< FPostProcessCircleDOFPS<0> > PixelShader(ShaderMap);
+		SetGlobalBoundShaderState(Context.RHICmdList, FeatureLevel, BoundShaderState, GFilterVertexDeclaration.VertexDeclarationRHI, *VertexShader, *PixelShader);
+
+		PixelShader->SetParameters(Context);
+	}
+
+	VertexShader->SetParameters(Context);
+
+	// Draw a quad mapping scene color to the view's render target
+	DrawRectangle(
+		Context.RHICmdList,
+		DestRect.Min.X, DestRect.Min.Y,
+		DestRect.Width() + 1, DestRect.Height() + 1,
+		SrcRect.Min.X, SrcRect.Min.Y,
+		SrcRect.Width() + 1, SrcRect.Height() + 1,
+		DestSize,
+		SrcSize,
+		*VertexShader,
+		EDRF_UseTriangleOptimization);
+
+	Context.RHICmdList.CopyToResolveTarget(DestRenderTarget0.TargetableTexture, DestRenderTarget0.ShaderResourceTexture, false, FResolveParams());
+	Context.RHICmdList.CopyToResolveTarget(DestRenderTarget1.TargetableTexture, DestRenderTarget1.ShaderResourceTexture, false, FResolveParams());
+}
+
+FPooledRenderTargetDesc FRCPassPostProcessCircleDOF::ComputeOutputDesc(EPassOutputId InPassOutputId) const
+{
+	FPooledRenderTargetDesc Ret = PassInputs[0].GetOutput()->RenderTargetDesc;
+
+	Ret.Extent.X = FMath::Max(1, Ret.Extent.X);
+	Ret.Extent.Y = FMath::Max(1, Ret.Extent.Y);
+
+	Ret.Reset();
+	Ret.TargetableFlags &= ~(uint32)TexCreate_UAV;
+	Ret.TargetableFlags |= TexCreate_RenderTargetable;
+
+	Ret.DebugName = (InPassOutputId == ePId_Output0) ? TEXT("CircleDOF0") : TEXT("CircleDOF1");
+
+	// more precision for additive blending and we need the alpha channel
+	Ret.Format = PF_FloatRGBA;
+
+	return Ret;
+}
+
+
+/** Encapsulates  the Circle DOF recombine pixel shader. */
+template <uint32 NearBlurEnable>
+class FPostProcessCircleDOFRecombinePS : public FGlobalShader
+{
+	DECLARE_SHADER_TYPE(FPostProcessCircleDOFRecombinePS, Global);
+
+	static bool ShouldCache(EShaderPlatform Platform)
+	{
+		return IsFeatureLevelSupported(Platform, ERHIFeatureLevel::SM4);
+	}
+
+	static void ModifyCompilationEnvironment(EShaderPlatform Platform, FShaderCompilerEnvironment& OutEnvironment)
+	{
+		FGlobalShader::ModifyCompilationEnvironment(Platform,OutEnvironment);
+		OutEnvironment.SetDefine(TEXT("ENABLE_NEAR_BLUR"), NearBlurEnable);
+	}
+
+	/** Default constructor. */
+	FPostProcessCircleDOFRecombinePS() {}
+
+public:
+	FPostProcessPassParameters PostprocessParameter;
+	FDeferredPixelShaderParameters DeferredParameters;
+	FShaderParameter DepthOfFieldUVLimit;
+	FShaderParameter RandomOffset;
+
+	/** Initialization constructor. */
+	FPostProcessCircleDOFRecombinePS(const ShaderMetaType::CompiledShaderInitializerType& Initializer)
+		: FGlobalShader(Initializer)
+	{
+		PostprocessParameter.Bind(Initializer.ParameterMap);
+		DeferredParameters.Bind(Initializer.ParameterMap);
+		DepthOfFieldUVLimit.Bind(Initializer.ParameterMap,TEXT("DepthOfFieldUVLimit"));
+		RandomOffset.Bind(Initializer.ParameterMap, TEXT("RandomOffset"));
+	}
+
+	// FShader interface.
+	virtual bool Serialize(FArchive& Ar)
+	{
+		bool bShaderHasOutdatedParameters = FGlobalShader::Serialize(Ar);
+		Ar << PostprocessParameter << DeferredParameters << DepthOfFieldUVLimit << RandomOffset;
+		return bShaderHasOutdatedParameters;
+	}
+
+	void SetParameters(const FRenderingCompositePassContext& Context)
+	{
+		const FPixelShaderRHIParamRef ShaderRHI = GetPixelShader();
+
+		FGlobalShader::SetParameters(Context.RHICmdList, ShaderRHI, Context.View);
+
+		DeferredParameters.Set(Context.RHICmdList, ShaderRHI, Context.View);
+		PostprocessParameter.SetPS(ShaderRHI, Context, TStaticSamplerState<SF_Bilinear,AM_Clamp,AM_Clamp,AM_Clamp>::GetRHI());
+
+		// Compute out of bounds UVs in the source texture.
+		FVector4 Bounds;
+		Bounds.X = (((float)((Context.View.ViewRect.Min.X + 1) & (~1))) + 3.0f) / ((float)(GSceneRenderTargets.GetBufferSizeXY().X));
+		Bounds.Y = (((float)((Context.View.ViewRect.Min.Y + 1) & (~1))) + 3.0f) / ((float)(GSceneRenderTargets.GetBufferSizeXY().Y));
+		Bounds.Z = (((float)(Context.View.ViewRect.Max.X & (~1))) - 3.0f) / ((float)(GSceneRenderTargets.GetBufferSizeXY().X));
+		Bounds.W = (((float)(Context.View.ViewRect.Max.Y & (~1))) - 3.0f) / ((float)(GSceneRenderTargets.GetBufferSizeXY().Y));
+
+		SetShaderValue(Context.RHICmdList, ShaderRHI, DepthOfFieldUVLimit, Bounds);
+
+		FVector2D RandomOffsetValue;
+		TemporalRandom2(&RandomOffsetValue, Context.View.Family->FrameNumber);
+		SetShaderValue(Context.RHICmdList, ShaderRHI, RandomOffset, RandomOffsetValue);
+	}
+};
+
+IMPLEMENT_SHADER_TYPE(template<>,FPostProcessCircleDOFRecombinePS<0>,TEXT("PostProcessDOF"),TEXT("MainCircleRecombinePS"),SF_Pixel);
+IMPLEMENT_SHADER_TYPE(template<>,FPostProcessCircleDOFRecombinePS<1>,TEXT("PostProcessDOF"),TEXT("MainCircleRecombinePS"),SF_Pixel);
+
+void FRCPassPostProcessCircleDOFRecombine::Process(FRenderingCompositePassContext& Context)
+{
+	SCOPED_DRAW_EVENT(Context.RHICmdList, CircleDOFRecombine);
+
+	const FPooledRenderTargetDesc* InputDesc = GetInputDesc(ePId_Input0);
+
+	if(!InputDesc)
+	{
+		// input is not hooked up correctly
+		return;
+	}
+
+	const FSceneView& View = Context.View;
+
+	const auto FeatureLevel = Context.GetFeatureLevel();
+	auto ShaderMap = Context.GetShaderMap();
+
+	FIntPoint TexSize = InputDesc->Extent;
+
+	// usually 1, 2, 4 or 8
+	uint32 ScaleToFullRes = GSceneRenderTargets.GetBufferSizeXY().X / TexSize.X;
+
+	FIntRect HalfResViewRect = View.ViewRect / ScaleToFullRes;
+
+	const FSceneRenderTargetItem& DestRenderTarget = PassOutputs[0].RequestSurface(Context);
+
+	// Set the view family's render target/viewport.
+	SetRenderTarget(Context.RHICmdList, DestRenderTarget.TargetableTexture, FTextureRHIRef());
+
+	// is optimized away if possible (RT size=view size, )
+	Context.RHICmdList.Clear(true, FLinearColor::Black, false, 1.0f, false, 0, View.ViewRect);
+
+	Context.SetViewportAndCallRHI(View.ViewRect);
+
+	// set the state
+	Context.RHICmdList.SetBlendState(TStaticBlendState<>::GetRHI());
+	Context.RHICmdList.SetRasterizerState(TStaticRasterizerState<>::GetRHI());
+	Context.RHICmdList.SetDepthStencilState(TStaticDepthStencilState<false, CF_Always>::GetRHI());
+
+	TShaderMapRef<FPostProcessVS> VertexShader(ShaderMap);
+
+	if (bNearBlurEnabled)
+	{
+		static FGlobalBoundShaderState BoundShaderState;
+
+		TShaderMapRef< FPostProcessCircleDOFRecombinePS<1> > PixelShader(ShaderMap);
+		SetGlobalBoundShaderState(Context.RHICmdList, FeatureLevel, BoundShaderState, GFilterVertexDeclaration.VertexDeclarationRHI, *VertexShader, *PixelShader);
+		PixelShader->SetParameters(Context);
+	}
+	else
+	{
+		static FGlobalBoundShaderState BoundShaderState;
+
+		TShaderMapRef< FPostProcessCircleDOFRecombinePS<0> > PixelShader(ShaderMap);
+		SetGlobalBoundShaderState(Context.RHICmdList, FeatureLevel, BoundShaderState, GFilterVertexDeclaration.VertexDeclarationRHI, *VertexShader, *PixelShader);
+		PixelShader->SetParameters(Context);
+	}
+
+	VertexShader->SetParameters(Context);
+
+	// Draw a quad mapping scene color to the view's render target
+	DrawRectangle(
+		Context.RHICmdList,
+		0, 0,
+		View.ViewRect.Width(), View.ViewRect.Height(),
+		0, 0,
+		View.ViewRect.Width(), View.ViewRect.Height(),
+		View.ViewRect.Size(),
+		TexSize,
+		*VertexShader,
+		EDRF_UseTriangleOptimization);
+
+	Context.RHICmdList.CopyToResolveTarget(DestRenderTarget.TargetableTexture, DestRenderTarget.ShaderResourceTexture, false, FResolveParams());
+}
+
+FPooledRenderTargetDesc FRCPassPostProcessCircleDOFRecombine::ComputeOutputDesc(EPassOutputId InPassOutputId) const
+{
+	FPooledRenderTargetDesc Ret = PassInputs[0].GetOutput()->RenderTargetDesc;
+
+	Ret.Reset();
+	Ret.DebugName = TEXT("CircleDOFRecombine");
+
+	return Ret;
+}
diff --git a/Engine/Source/Runtime/Renderer/Private/PostProcess/PostProcessDOF.h b/Engine/Source/Runtime/Renderer/Private/PostProcess/PostProcessDOF.h
--- a/Engine/Source/Runtime/Renderer/Private/PostProcess/PostProcessDOF.h
+++ b/Engine/Source/Runtime/Renderer/Private/PostProcess/PostProcessDOF.h
@@ -49,3 +49,65 @@ private:
 	bool bNearBlurEnabled;
 };
 
+
+// down sample and setup DOF input
+// ePId_Input0: SceneColor
+// ePId_Input1: SceneDepth
+// derives from TRenderingCompositePassBase<InputCount, OutputCount> 
+class FRCPassPostProcessCircleDOFSetup : public TRenderingCompositePassBase<2, 2>
+{
+public:
+
+	FRCPassPostProcessCircleDOFSetup(bool bInNearBlurEnabled) : bNearBlurEnabled ( bInNearBlurEnabled ) {}
+
+	// interface FRenderingCompositePass ---------
+
+	virtual void Process(FRenderingCompositePassContext& Context);
+	virtual void Release() override { delete this; }
+	virtual FPooledRenderTargetDesc ComputeOutputDesc(EPassOutputId InPassOutputId) const;
+
+private:
+
+	bool bNearBlurEnabled;
+};
+
+// down sample and setup DOF input
+// ePId_Input0: setup results for far
+// ePId_Input1: setup results for near
+// derives from TRenderingCompositePassBase<InputCount, OutputCount> 
+class FRCPassPostProcessCircleDOF : public TRenderingCompositePassBase<2, 2>
+{
+public:
+
+	FRCPassPostProcessCircleDOF(bool bInNearBlurEnabled) : bNearBlurEnabled ( bInNearBlurEnabled ) {}
+
+	// interface FRenderingCompositePass ---------
+
+	virtual void Process(FRenderingCompositePassContext& Context);
+	virtual void Release() override { delete this; }
+	virtual FPooledRenderTargetDesc ComputeOutputDesc(EPassOutputId InPassOutputId) const;
+
+private:
+
+	bool bNearBlurEnabled;
+};
+
+// derives from TRenderingCompositePassBase<InputCount, OutputCount> 
+// ePId_Input0: Full res scene color
+// ePId_Input1: output 0 from the DOFSetup (possibly further blurred)
+// ePId_Input2: output 1 from the DOFSetup (possibly further blurred)
+class FRCPassPostProcessCircleDOFRecombine : public TRenderingCompositePassBase<3, 1>
+{
+public:
+
+	FRCPassPostProcessCircleDOFRecombine(bool bInNearBlurEnabled) : bNearBlurEnabled ( bInNearBlurEnabled ) {}
+
+	// interface FRenderingCompositePass ---------
+	virtual void Process(FRenderingCompositePassContext& Context);
+	virtual void Release() override { delete this; }
+	virtual FPooledRenderTargetDesc ComputeOutputDesc(EPassOutputId InPassOutputId) const;
+
+private:
+
+	bool bNearBlurEnabled;
+};
\ No newline at end of file
diff --git a/Engine/Source/Runtime/Renderer/Private/PostProcess/PostProcessing.cpp b/Engine/Source/Runtime/Renderer/Private/PostProcess/PostProcessing.cpp
--- a/Engine/Source/Runtime/Renderer/Private/PostProcess/PostProcessing.cpp
+++ b/Engine/Source/Runtime/Renderer/Private/PostProcess/PostProcessing.cpp
@@ -427,6 +427,111 @@ static void AddPostProcessDepthOfFieldGaussian(FPostprocessContext& Context, FDe
 	Context.FinalOutput = FRenderingCompositeOutputRef(NodeRecombined);
 }
 
+static void AddPostProcessDepthOfFieldCircle(FPostprocessContext& Context, FDepthOfFieldStats& Out)
+{
+	float FarSize = Context.View.FinalPostProcessSettings.DepthOfFieldFarBlurSize;
+	float NearSize = Context.View.FinalPostProcessSettings.DepthOfFieldNearBlurSize;
+
+	float MaxSize = CVarDepthOfFieldMaxSize.GetValueOnRenderThread();
+
+	FarSize = FMath::Min(FarSize, MaxSize);
+	NearSize = FMath::Min(NearSize, MaxSize);
+
+	Out.bFar = FarSize >= 0.01f;
+	Out.bNear = NearSize >= GetCachedScalabilityCVars().GaussianDOFNearThreshold;
+
+	if(!Out.bFar && !Out.bNear)
+	{
+		return;
+	}
+
+	if(Context.View.Family->EngineShowFlags.VisualizeDOF)
+	{
+		// no need for this pass
+		return;
+	}
+
+	FRenderingCompositePass* DOFSetup = Context.Graph.RegisterPass(new(FMemStack::Get()) FRCPassPostProcessCircleDOFSetup(Out.bNear));
+	DOFSetup->SetInput(ePId_Input0, FRenderingCompositeOutputRef(Context.FinalOutput));
+	// We need the depth to create the near and far mask
+	DOFSetup->SetInput(ePId_Input1, FRenderingCompositeOutputRef(Context.SceneDepth));
+
+	FSceneViewState* ViewState = (FSceneViewState*)Context.View.State;
+
+	FRenderingCompositePass* DOFInputPass = DOFSetup;
+	if( Context.View.FinalPostProcessSettings.AntiAliasingMethod == AAM_TemporalAA && ViewState )
+	{
+		FRenderingCompositePass* HistoryInput;
+		if( ViewState->DOFHistoryRT && !ViewState->bBokehDOFHistory && !Context.View.bCameraCut )
+		{
+			HistoryInput = Context.Graph.RegisterPass( new(FMemStack::Get()) FRCPassPostProcessInput( ViewState->DOFHistoryRT ) );
+		}
+		else
+		{
+			// No history so use current as history
+			HistoryInput = DOFSetup;
+		}
+
+		FRenderingCompositePass* NodeTemporalAA = Context.Graph.RegisterPass( new(FMemStack::Get()) FRCPassPostProcessDOFTemporalAA );
+		NodeTemporalAA->SetInput( ePId_Input0, DOFSetup );
+		NodeTemporalAA->SetInput( ePId_Input1, FRenderingCompositeOutputRef( HistoryInput ) );
+		NodeTemporalAA->SetInput( ePId_Input2, FRenderingCompositeOutputRef( HistoryInput ) );
+		//NodeTemporalAA->SetInput( ePId_Input3, VelocityInput );
+
+		DOFInputPass = NodeTemporalAA;
+		ViewState->bBokehDOFHistory = false;
+	}
+
+	FRenderingCompositePass* DOFInputPass2 = DOFSetup;
+	EPassOutputId DOFInputPassId = ePId_Output1;
+	if( Context.View.FinalPostProcessSettings.AntiAliasingMethod == AAM_TemporalAA && ViewState )
+	{
+		FRenderingCompositePass* HistoryInput;
+		EPassOutputId DOFInputPassId2 = ePId_Output1;
+		if( ViewState->DOFHistoryRT2 && !ViewState->bBokehDOFHistory2 && !Context.View.bCameraCut )
+		{
+			HistoryInput = Context.Graph.RegisterPass( new(FMemStack::Get()) FRCPassPostProcessInput( ViewState->DOFHistoryRT2 ) );
+			DOFInputPassId2 = ePId_Output0;
+		}
+		else
+		{
+			// No history so use current as history
+			HistoryInput = DOFSetup;
+			DOFInputPassId2 = ePId_Output1;
+		}
+
+		FRenderingCompositePass* NodeTemporalAA = Context.Graph.RegisterPass( new(FMemStack::Get()) FRCPassPostProcessDOFTemporalAANear );
+		NodeTemporalAA->SetInput( ePId_Input0, FRenderingCompositeOutputRef( DOFSetup, ePId_Output1 ) );
+		NodeTemporalAA->SetInput( ePId_Input1, FRenderingCompositeOutputRef( HistoryInput, DOFInputPassId2 ) );
+		NodeTemporalAA->SetInput( ePId_Input2, FRenderingCompositeOutputRef( HistoryInput, DOFInputPassId2 ) );
+		//NodeTemporalAA->SetInput( ePId_Input3, VelocityInput );
+
+		DOFInputPass2 = NodeTemporalAA;
+		ViewState->bBokehDOFHistory2 = false;
+
+		DOFInputPassId = ePId_Output0;
+	}
+
+	FRenderingCompositeOutputRef Far;
+	FRenderingCompositeOutputRef Near;
+
+	FRenderingCompositePass* DOFApply = Context.Graph.RegisterPass(new(FMemStack::Get()) FRCPassPostProcessCircleDOF(Out.bNear));
+	DOFApply->SetInput(ePId_Input0, FRenderingCompositeOutputRef(DOFInputPass, ePId_Output0));
+	Far = FRenderingCompositeOutputRef(DOFApply, ePId_Output0);
+	if(Out.bNear)
+	{
+		DOFApply->SetInput(ePId_Input1, FRenderingCompositeOutputRef(DOFInputPass2, DOFInputPassId));
+		Near = FRenderingCompositeOutputRef(DOFApply, ePId_Output1);
+	}
+
+	FRenderingCompositePass* NodeRecombined = Context.Graph.RegisterPass(new(FMemStack::Get()) FRCPassPostProcessCircleDOFRecombine(Out.bNear));
+	NodeRecombined->SetInput(ePId_Input0, Context.FinalOutput);
+	NodeRecombined->SetInput(ePId_Input1, Far);
+	NodeRecombined->SetInput(ePId_Input2, Near);
+
+	Context.FinalOutput = FRenderingCompositeOutputRef(NodeRecombined);
+}
+
 static FRenderingCompositeOutputRef AddBloom(FPostprocessContext Context, FRenderingCompositeOutputRef PostProcessDownsample0)
 {
 	// Quality level to bloom stages table. Note: 0 is omitted, ensure element count tallys with the range documented with 'r.BloomQuality' definition.
@@ -965,7 +1070,15 @@ void FPostProcessing::Process(FRHICommandListImmediate& RHICmdList, FViewInfo& V
 
 			if(bDepthOfField && View.FinalPostProcessSettings.DepthOfFieldMethod != DOFM_BokehDOF)
 			{
-				AddPostProcessDepthOfFieldGaussian(Context, DepthOfFieldStat);
+				bool bCircleDOF = View.FinalPostProcessSettings.DepthOfFieldMethod == DOFM_CircleDOF;
+				if(!bCircleDOF)
+				{
+					AddPostProcessDepthOfFieldGaussian(Context, DepthOfFieldStat);
+				}
+				else
+				{
+					AddPostProcessDepthOfFieldCircle(Context, DepthOfFieldStat);
+				}
 			}
 
 			bool bBokehDOF = bDepthOfField
