UE4 - Fix crash when using extra GBuffer

diff --git a/Engine/Source/Runtime/Renderer/Private/PostProcess/SceneRenderTargets.cpp b/Engine/Source/Runtime/Renderer/Private/PostProcess/SceneRenderTargets.cpp
--- a/Engine/Source/Runtime/Renderer/Private/PostProcess/SceneRenderTargets.cpp
+++ b/Engine/Source/Runtime/Renderer/Private/PostProcess/SceneRenderTargets.cpp
@@ -286,33 +286,33 @@ void FSceneRenderTargets::BeginRenderingSceneColor(FRHICommandList& RHICmdList,
 	SetRenderTarget(RHICmdList, GetSceneColorSurface(), GetSceneDepthSurface(), RenderTargetMode);
 } 
 
-int32 FSceneRenderTargets::GetGBufferRenderTargets(FRHIRenderTargetView RenderTargets[MAX_CONCURRENT_GBUFFERS], ERenderTargetLoadAction ColorLoadAction, int32& VelocityRTIndex)
+int32 FSceneRenderTargets::GetGBufferRenderTargets(ERenderTargetLoadAction ColorLoadAction, FRHIRenderTargetView OutRenderTargets[MaxSimultaneousRenderTargets], int32& OutVelocityRTIndex)
 {
 	const bool bUseVelocityGBuffer = FVelocityRendering::OutputsToGBuffer();
 
-	RenderTargets[0] = FRHIRenderTargetView(GetSceneColorSurface(), 0, -1, ColorLoadAction, ERenderTargetStoreAction::EStore);
-	RenderTargets[1] = FRHIRenderTargetView(GSceneRenderTargets.GBufferA->GetRenderTargetItem().TargetableTexture, 0, -1, ColorLoadAction, ERenderTargetStoreAction::EStore);
-	RenderTargets[2] = FRHIRenderTargetView(GSceneRenderTargets.GBufferB->GetRenderTargetItem().TargetableTexture, 0, -1, ColorLoadAction, ERenderTargetStoreAction::EStore);
-	RenderTargets[3] = FRHIRenderTargetView(GSceneRenderTargets.GBufferC->GetRenderTargetItem().TargetableTexture, 0, -1, ColorLoadAction, ERenderTargetStoreAction::EStore);
-	RenderTargets[4] = FRHIRenderTargetView(GSceneRenderTargets.GBufferD->GetRenderTargetItem().TargetableTexture, 0, -1, ColorLoadAction, ERenderTargetStoreAction::EStore);
+	OutRenderTargets[0] = FRHIRenderTargetView(GetSceneColorSurface(), 0, -1, ColorLoadAction, ERenderTargetStoreAction::EStore);
+	OutRenderTargets[1] = FRHIRenderTargetView(GSceneRenderTargets.GBufferA->GetRenderTargetItem().TargetableTexture, 0, -1, ColorLoadAction, ERenderTargetStoreAction::EStore);
+	OutRenderTargets[2] = FRHIRenderTargetView(GSceneRenderTargets.GBufferB->GetRenderTargetItem().TargetableTexture, 0, -1, ColorLoadAction, ERenderTargetStoreAction::EStore);
+	OutRenderTargets[3] = FRHIRenderTargetView(GSceneRenderTargets.GBufferC->GetRenderTargetItem().TargetableTexture, 0, -1, ColorLoadAction, ERenderTargetStoreAction::EStore);
+	OutRenderTargets[4] = FRHIRenderTargetView(GSceneRenderTargets.GBufferD->GetRenderTargetItem().TargetableTexture, 0, -1, ColorLoadAction, ERenderTargetStoreAction::EStore);
 	int32 MRTCount = 5;	// Derived from previously used RenderTargets[]; would have been nice to keep a pointer while filling it but it's more confusing to use
 
-	VelocityRTIndex = -1;
+	OutVelocityRTIndex = -1;
 
 	if (bAllowStaticLighting)
 	{
-		RenderTargets[MRTCount] = FRHIRenderTargetView(GSceneRenderTargets.GBufferE->GetRenderTargetItem().TargetableTexture, 0, -1, ColorLoadAction, ERenderTargetStoreAction::EStore);
+		OutRenderTargets[MRTCount] = FRHIRenderTargetView(GSceneRenderTargets.GBufferE->GetRenderTargetItem().TargetableTexture, 0, -1, ColorLoadAction, ERenderTargetStoreAction::EStore);
 		++MRTCount;
 	}
 
 	if (bUseVelocityGBuffer)
 	{
-		VelocityRTIndex = MRTCount;
+		OutVelocityRTIndex = MRTCount;
 		++MRTCount;
-		RenderTargets[VelocityRTIndex] = FRHIRenderTargetView(GSceneRenderTargets.GBufferVelocity->GetRenderTargetItem().TargetableTexture, 0, -1, ColorLoadAction, ERenderTargetStoreAction::EStore);
+		OutRenderTargets[OutVelocityRTIndex] = FRHIRenderTargetView(GSceneRenderTargets.GBufferVelocity->GetRenderTargetItem().TargetableTexture, 0, -1, ColorLoadAction, ERenderTargetStoreAction::EStore);
 	}
 
-	check(MRTCount < MAX_CONCURRENT_GBUFFERS);
+	check(MRTCount <= MaxSimultaneousRenderTargets);
 
 	return MRTCount;
 }
@@ -334,8 +334,8 @@ void FSceneRenderTargets::BeginRenderingGBuffer(FRHICommandList& RHICmdList, ERe
 	if (CurrentFeatureLevel >= ERHIFeatureLevel::SM4)
 	{
 		int32 VelocityRTIndex;
-		FRHIRenderTargetView RenderTargets[7];
-		int32 MRTCount = GetGBufferRenderTargets(RenderTargets, ColorLoadAction, VelocityRTIndex);
+		FRHIRenderTargetView RenderTargets[MaxSimultaneousRenderTargets];
+		int32 MRTCount = GetGBufferRenderTargets(ColorLoadAction, RenderTargets, VelocityRTIndex);
 
 		const float DepthClearValue = 0.0f;
 		bool bClearColor = ColorLoadAction == ERenderTargetLoadAction::EClear;
@@ -385,8 +385,8 @@ void FSceneRenderTargets::FinishRenderingGBuffer(FRHICommandListImmediate& RHICm
 	}
 
 	int32 VelocityRTIndex;
-	FRHIRenderTargetView RenderTargets[7];
-	int32 NumMRTs = GetGBufferRenderTargets(RenderTargets, ERenderTargetLoadAction::ELoad, VelocityRTIndex);
+	FRHIRenderTargetView RenderTargets[MaxSimultaneousRenderTargets];
+	int32 NumMRTs = GetGBufferRenderTargets(ERenderTargetLoadAction::ELoad, RenderTargets, VelocityRTIndex);
 
 	FResolveParams ResolveParams;
 	for (int32 i = 0; i < NumMRTs; ++i)
diff --git a/Engine/Source/Runtime/Renderer/Private/PostProcess/SceneRenderTargets.h b/Engine/Source/Runtime/Renderer/Private/PostProcess/SceneRenderTargets.h
--- a/Engine/Source/Runtime/Renderer/Private/PostProcess/SceneRenderTargets.h
+++ b/Engine/Source/Runtime/Renderer/Private/PostProcess/SceneRenderTargets.h
@@ -587,9 +587,8 @@ private:
 	/** Determine whether the render targets for any shading path have been allocated */
 	bool AreAnyShadingPathRenderTargetsAllocated() const { return AreShadingPathRenderTargetsAllocated(EShadingPath::Deferred) || AreShadingPathRenderTargetsAllocated(EShadingPath::Forward); }
 
-	static const int32  MAX_CONCURRENT_GBUFFERS = 7;
 	/** Gets all GBuffers to use.  Returns the number actually used. */
-	int32 GetGBufferRenderTargets(FRHIRenderTargetView RenderTargets[MAX_CONCURRENT_GBUFFERS], ERenderTargetLoadAction ColorLoadAction, int32 & VelocityRTIndex);
+	int32 GetGBufferRenderTargets(ERenderTargetLoadAction ColorLoadAction, FRHIRenderTargetView OutRenderTargets[MaxSimultaneousRenderTargets], int32& OutVelocityRTIndex);
 
 private:
 	/** Uniform buffer containing GBuffer resources. */
