minor fixed to get better visualize bloom

diff --git a/Engine/Source/Runtime/Renderer/Private/PostProcess/PostProcessing.cpp b/Engine/Source/Runtime/Renderer/Private/PostProcess/PostProcessing.cpp
--- a/Engine/Source/Runtime/Renderer/Private/PostProcess/PostProcessing.cpp
+++ b/Engine/Source/Runtime/Renderer/Private/PostProcess/PostProcessing.cpp
@@ -1417,7 +1417,8 @@ void FPostProcessing::Process(FRHICommandListImmediate& RHICmdList, FViewInfo& V
 			&& View.Family->EngineShowFlags.SelectionOutline
 			&& !(View.Family->EngineShowFlags.Wireframe)
 			&& ( !GIsDemoMode || ( GIsDemoMode && !View.Family->EngineShowFlags.Game ) ) 
-			&& !bVisualizeBloom)
+			&& !bVisualizeBloom
+			&& !View.Family->EngineShowFlags.VisualizeHDR)
 		{
 			// Selection outline is after bloom, but before AA
 			AddSelectionOutline(Context);
diff --git a/Engine/Source/Runtime/Renderer/Private/SceneRendering.cpp b/Engine/Source/Runtime/Renderer/Private/SceneRendering.cpp
--- a/Engine/Source/Runtime/Renderer/Private/SceneRendering.cpp
+++ b/Engine/Source/Runtime/Renderer/Private/SceneRendering.cpp
@@ -1236,6 +1236,11 @@ bool FSceneRenderer::ShouldCompositeEditorPrimitives(const FViewInfo& View)
 		return false;
 	}
 
+	if(View.Family->EngineShowFlags.VisualizeHDR)
+	{
+		return false;
+	}
+
 	if (GIsEditor && View.Family->EngineShowFlags.Wireframe)
 	{
 		// In Editor we want wire frame view modes to be in MSAA
