Volume lighting samples show up even if ILC show flag is disabled

diff --git a/Engine/Source/Runtime/Renderer/Private/IndirectLightingCache.cpp b/Engine/Source/Runtime/Renderer/Private/IndirectLightingCache.cpp
--- a/Engine/Source/Runtime/Renderer/Private/IndirectLightingCache.cpp
+++ b/Engine/Source/Runtime/Renderer/Private/IndirectLightingCache.cpp
@@ -400,21 +400,21 @@ void FIndirectLightingCache::UpdateCache(FScene* Scene, FSceneRenderer& Renderer
 			UpdateBlocks(Scene, Renderer.Views.GetData(), BlocksToUpdate);
 
 			UpdateTransitionsOverTime(TransitionsOverTimeToUpdate, Renderer.ViewFamily.DeltaWorldTime);
+		}
 
-			if (GCacheDrawLightingSamples || Renderer.ViewFamily.EngineShowFlags.VolumeLightingSamples || GCacheDrawDirectionalShadowing)
-			{
-				FViewElementPDI DebugPDI(Renderer.Views.GetData(), NULL);
+		bUpdateAllCacheEntries = false;
+	}
 
-				for (int32 VolumeIndex = 0; VolumeIndex < Scene->PrecomputedLightVolumes.Num(); VolumeIndex++)
-				{
-					const FPrecomputedLightVolume* PrecomputedLightVolume = Scene->PrecomputedLightVolumes[VolumeIndex];
+	if (GCacheDrawLightingSamples || Renderer.ViewFamily.EngineShowFlags.VolumeLightingSamples || GCacheDrawDirectionalShadowing)
+	{
+		FViewElementPDI DebugPDI(Renderer.Views.GetData(), NULL);
 
-					PrecomputedLightVolume->DebugDrawSamples(&DebugPDI, GCacheDrawDirectionalShadowing != 0);
-				}
-			}
-		}
+		for (int32 VolumeIndex = 0; VolumeIndex < Scene->PrecomputedLightVolumes.Num(); VolumeIndex++)
+		{
+			const FPrecomputedLightVolume* PrecomputedLightVolume = Scene->PrecomputedLightVolumes[VolumeIndex];
 
-		bUpdateAllCacheEntries = false;
+			PrecomputedLightVolume->DebugDrawSamples(&DebugPDI, GCacheDrawDirectionalShadowing != 0);
+		}
 	}
 }
 
