Stationary lights with CastDynamicShadows disabled no longer do dynamic shadow setup

diff --git a/Engine/Source/Runtime/Renderer/Private/ShadowSetup.cpp b/Engine/Source/Runtime/Renderer/Private/ShadowSetup.cpp
--- a/Engine/Source/Runtime/Renderer/Private/ShadowSetup.cpp
+++ b/Engine/Source/Runtime/Renderer/Private/ShadowSetup.cpp
@@ -644,8 +644,12 @@ void FProjectedShadowInfo::AddSubjectPrimitive(FPrimitiveSceneInfo* PrimitiveSce
 	// Ray traced shadows use the GPU managed distance field object buffers, no CPU culling should be used
 	check(!CascadeSettings.bRayTracedDistanceField);
 
-	if (!ReceiverPrimitives.Contains(PrimitiveSceneInfo))
+	if (!ReceiverPrimitives.Contains(PrimitiveSceneInfo)
+		// Far cascade only casts from primitives marked for it
+		&& (!CascadeSettings.bFarShadowCascade || PrimitiveSceneInfo->Proxy->CastsFarShadow()))
 	{
+		const FPrimitiveSceneProxy* Proxy = PrimitiveSceneInfo->Proxy;
+
 		TArray<FViewInfo*, TInlineAllocator<1> > Views;
 		const bool bWholeSceneDirectionalShadow = IsWholeSceneDirectionalShadow();
 
@@ -679,8 +683,6 @@ void FProjectedShadowInfo::AddSubjectPrimitive(FPrimitiveSceneInfo* PrimitiveSce
 			{
 				if( CurrentView.IsPerspectiveProjection() )
 				{
-					const FPrimitiveSceneProxy* Proxy = PrimitiveSceneInfo->Proxy;
-
 					// Compute the distance between the view and the primitive.
 					float DistanceSquared = (Proxy->GetBounds().Origin - CurrentView.ShadowViewMatrices.ViewOrigin).SizeSquared();
 
@@ -717,7 +719,6 @@ void FProjectedShadowInfo::AddSubjectPrimitive(FPrimitiveSceneInfo* PrimitiveSce
 
 		if (bOpaqueRelevance && bShadowRelevance)
 		{
-			const FPrimitiveSceneProxy* Proxy = PrimitiveSceneInfo->Proxy;
 			const FBoxSphereBounds& Bounds = Proxy->GetBounds();
 			bool bDrawingStaticMeshes = false;
 
@@ -727,21 +728,6 @@ void FProjectedShadowInfo::AddSubjectPrimitive(FPrimitiveSceneInfo* PrimitiveSce
 				{
 					FViewInfo& CurrentView = *Views[ViewIndex];
 
-					// 	in a far shadow cascade we only render objects marked with bCastFarShadows enabled (e.g. Landscape).
-					if(CascadeSettings.bFarShadowCascade)
-					{
-						if(!Proxy->CastsFarShadow())
-						{
-							// cull objects that are not supposed to be in the far shadow cascades
-							continue;
-						}
-						if(CascadeSettings.bRayTracedDistanceField && Proxy->AffectsDistanceFieldLighting())
-						{
-							// cull objects that already have been handles by the ray traced distance field ones
-							continue;
-						}
-					}
-
 					const float DistanceSquared = ( Bounds.Origin - CurrentView.ShadowViewMatrices.ViewOrigin ).SizeSquared();
 					const bool bDrawShadowDepth = FMath::Square( Bounds.SphereRadius ) > FMath::Square( GMinScreenRadiusForShadowCaster ) * DistanceSquared;
 					if( !bDrawShadowDepth )
@@ -2407,7 +2393,7 @@ void FDeferredShadingSceneRenderer::InitDynamicShadows(FRHICommandListImmediate&
 					}
 
 					// Allow movable and stationary lights to create CSM, or static lights that are unbuilt
-					if (!LightSceneInfo->Proxy->HasStaticLighting() || bCreateShadowToPreviewStaticLight)
+					if ((!LightSceneInfo->Proxy->HasStaticLighting() && LightSceneInfoCompact.bCastDynamicShadow) || bCreateShadowToPreviewStaticLight)
 					{
 						AddViewDependentWholeSceneShadowsForView(ViewDependentWholeSceneShadows, ViewDependentWholeSceneShadowsThatNeedCulling, VisibleLightInfo, *LightSceneInfo);
 
