Surfel Global Illumination with Distance Field shadowing

diff --git a/Engine/Source/Runtime/Engine/Private/InstancedStaticMesh.h b/Engine/Source/Runtime/Engine/Private/InstancedStaticMesh.h
--- a/Engine/Source/Runtime/Engine/Private/InstancedStaticMesh.h
+++ b/Engine/Source/Runtime/Engine/Private/InstancedStaticMesh.h
@@ -533,6 +533,8 @@ public:
 
 	virtual void GetDistancefieldAtlasData(FBox& LocalVolumeBounds, FIntVector& OutBlockMin, FIntVector& OutBlockSize, bool& bOutBuiltAsIfTwoSided, bool& bMeshWasPlane, TArray<FMatrix>& ObjectLocalToWorldTransforms) const override;
 
+	virtual void GetDistanceFieldInstanceInfo(int32& NumInstances, float& BoundsSurfaceArea) const override;
+
 	/**
 	 * Creates the hit proxies are used when DrawDynamicElements is called.
 	 * Called in the game thread.
diff --git a/Engine/Source/Runtime/Engine/Private/Materials/HLSLMaterialTranslator.h b/Engine/Source/Runtime/Engine/Private/Materials/HLSLMaterialTranslator.h
--- a/Engine/Source/Runtime/Engine/Private/Materials/HLSLMaterialTranslator.h
+++ b/Engine/Source/Runtime/Engine/Private/Materials/HLSLMaterialTranslator.h
@@ -3397,7 +3397,7 @@ protected:
 			return NonPixelShaderExpressionError();
 		}
 
-		return AddCodeChunk(GetParameterType(X),TEXT("ddx(%s)"),*GetParameterCode(X));
+		return AddCodeChunk(GetParameterType(X),TEXT("DDX(%s)"),*GetParameterCode(X));
 	}
 
 	virtual int32 DDY( int32 X ) override
@@ -3422,7 +3422,7 @@ protected:
 			return NonPixelShaderExpressionError();
 		}
 
-		return AddCodeChunk(GetParameterType(X),TEXT("ddy(%s)"),*GetParameterCode(X));
+		return AddCodeChunk(GetParameterType(X),TEXT("DDY(%s)"),*GetParameterCode(X));
 	}
 
 	virtual int32 AntialiasedTextureMask(int32 Tex, int32 UV, float Threshold, uint8 Channel) override
diff --git a/Engine/Source/Runtime/Engine/Private/SceneView.cpp b/Engine/Source/Runtime/Engine/Private/SceneView.cpp
--- a/Engine/Source/Runtime/Engine/Private/SceneView.cpp
+++ b/Engine/Source/Runtime/Engine/Private/SceneView.cpp
@@ -269,6 +269,7 @@ FSceneView::FSceneView(const FSceneViewInitOptions& InitOptions)
 	, OverlayColor(InitOptions.OverlayColor)
 	, ColorScale(InitOptions.ColorScale)
 	, StereoPass(InitOptions.StereoPass)
+	, bRenderFirstInstanceOnly(false)
 	, DiffuseOverrideParameter(FVector4(0,0,0,1))
 	, SpecularOverrideParameter(FVector4(0,0,0,1))
 	, NormalOverrideParameter(FVector4(0,0,0,1))
diff --git a/Engine/Source/Runtime/Engine/Private/SkyLightComponent.cpp b/Engine/Source/Runtime/Engine/Private/SkyLightComponent.cpp
--- a/Engine/Source/Runtime/Engine/Private/SkyLightComponent.cpp
+++ b/Engine/Source/Runtime/Engine/Private/SkyLightComponent.cpp
@@ -109,7 +109,7 @@ USkyLightComponent::USkyLightComponent(const FObjectInitializer& ObjectInitializ
 	bCaptureDirty = false;
 	bLowerHemisphereIsBlack = true;
 	bSavedConstructionScriptValuesValid = true;
-	OcclusionMaxDistance = 600;
+	OcclusionMaxDistance = 1000;
 	MinOcclusion = 0;
 	OcclusionTint = FColor::Black;
 }
@@ -250,7 +250,8 @@ bool USkyLightComponent::CanEditChange(const UProperty* InProperty) const
 		}
 
 		if (FCString::Strcmp(*PropertyName, TEXT("Contrast")) == 0
-			|| FCString::Strcmp(*PropertyName, TEXT("OcclusionMaxDistance")) == 0)
+			|| FCString::Strcmp(*PropertyName, TEXT("OcclusionMaxDistance")) == 0
+			|| FCString::Strcmp(*PropertyName, TEXT("MinOcclusion")) == 0)
 		{
 			static const auto CVar = IConsoleManager::Get().FindTConsoleVariableDataInt(TEXT("r.GenerateMeshDistanceFields"));
 			return Mobility == EComponentMobility::Movable && CastShadows && CVar->GetValueOnGameThread() != 0;
diff --git a/Engine/Source/Runtime/Engine/Private/SplineMeshSceneProxy.cpp b/Engine/Source/Runtime/Engine/Private/SplineMeshSceneProxy.cpp
--- a/Engine/Source/Runtime/Engine/Private/SplineMeshSceneProxy.cpp
+++ b/Engine/Source/Runtime/Engine/Private/SplineMeshSceneProxy.cpp
@@ -6,6 +6,8 @@
 FSplineMeshSceneProxy::FSplineMeshSceneProxy(USplineMeshComponent* InComponent) :
 	FStaticMeshSceneProxy(InComponent)
 {
+	bSupportsDistanceFieldRepresentation = false;
+
 	// make sure all the materials are okay to be rendered as a spline mesh
 	for (FStaticMeshSceneProxy::FLODInfo& LODInfo : LODs)
 	{
diff --git a/Engine/Source/Runtime/Engine/Private/StaticMeshRender.cpp b/Engine/Source/Runtime/Engine/Private/StaticMeshRender.cpp
--- a/Engine/Source/Runtime/Engine/Private/StaticMeshRender.cpp
+++ b/Engine/Source/Runtime/Engine/Private/StaticMeshRender.cpp
@@ -972,6 +972,17 @@ void FStaticMeshSceneProxy::GetDistancefieldAtlasData(FBox& LocalVolumeBounds, F
 	}
 }
 
+void FStaticMeshSceneProxy::GetDistanceFieldInstanceInfo(int32& NumInstances, float& BoundsSurfaceArea) const
+{
+	NumInstances = DistanceFieldData ? 1 : 0;
+	const FVector AxisScales = GetLocalToWorld().GetScaleVector();
+	const FVector BoxDimensions = RenderData->Bounds.BoxExtent * AxisScales * 2;
+
+	BoundsSurfaceArea = 2 * BoxDimensions.X * BoxDimensions.Y
+		+ 2 * BoxDimensions.Z * BoxDimensions.Y
+		+ 2 * BoxDimensions.X * BoxDimensions.Z;
+}
+
 bool FStaticMeshSceneProxy::HasDistanceFieldRepresentation() const
 {
 	return CastsDynamicShadow() && AffectsDistanceFieldLighting() && DistanceFieldData && DistanceFieldData->VolumeTexture.IsValidDistanceFieldVolume();
diff --git a/Engine/Source/Runtime/Engine/Public/DistanceFieldAtlas.h b/Engine/Source/Runtime/Engine/Public/DistanceFieldAtlas.h
--- a/Engine/Source/Runtime/Engine/Public/DistanceFieldAtlas.h
+++ b/Engine/Source/Runtime/Engine/Public/DistanceFieldAtlas.h
@@ -71,6 +71,8 @@ public:
 	int32 GetSizeY() const { return VolumeTextureRHI->GetSizeY(); }
 	int32 GetSizeZ() const { return VolumeTextureRHI->GetSizeZ(); }
 
+	ENGINE_API FString GetSizeString() const;
+
 	/** Add an allocation to the atlas. */
 	void AddAllocation(FDistanceFieldVolumeTexture* Texture);
 
diff --git a/Engine/Source/Runtime/Engine/Public/MeshBatch.h b/Engine/Source/Runtime/Engine/Public/MeshBatch.h
--- a/Engine/Source/Runtime/Engine/Public/MeshBatch.h
+++ b/Engine/Source/Runtime/Engine/Public/MeshBatch.h
@@ -18,6 +18,8 @@ struct FMeshBatchElement
 	TUniformBufferRef<FPrimitiveUniformShaderParameters> PrimitiveUniformBuffer;
 
 	const FIndexBuffer* IndexBuffer;
+
+	/** Instance runs, where number of runs is specified by NumInstances.  Run structure is [StartInstanceIndex, EndInstanceIndex]. */
 	uint32* InstanceRuns;
 	const void* UserData;
 	/** 
@@ -27,6 +29,8 @@ struct FMeshBatchElement
 	const void* DynamicIndexData;
 	uint32 FirstIndex;
 	uint32 NumPrimitives;
+
+	/** Number of instances to draw.  If InstanceRuns is valid, this is actually the number of runs in InstanceRuns. */
 	uint32 NumInstances;
 	uint32 MinVertexIndex;
 	uint32 MaxVertexIndex;
diff --git a/Engine/Source/Runtime/Engine/Public/PrimitiveSceneProxy.h b/Engine/Source/Runtime/Engine/Public/PrimitiveSceneProxy.h
--- a/Engine/Source/Runtime/Engine/Public/PrimitiveSceneProxy.h
+++ b/Engine/Source/Runtime/Engine/Public/PrimitiveSceneProxy.h
@@ -204,6 +204,12 @@ public:
 		bMeshWasPlane = false;
 	}
 
+	virtual void GetDistanceFieldInstanceInfo(int32& NumInstances, float& BoundsSurfaceArea) const
+	{
+		NumInstances = 0;
+		BoundsSurfaceArea = 0;
+	}
+
 	virtual void GetHeightfieldRepresentation(UTexture2D*& OutHeightmapTexture, FVector4& OutHeightfieldScaleBias, FVector4& OutMinMaxUV)
 	{
 		OutHeightmapTexture = NULL;
diff --git a/Engine/Source/Runtime/Engine/Public/SceneManagement.h b/Engine/Source/Runtime/Engine/Public/SceneManagement.h
--- a/Engine/Source/Runtime/Engine/Public/SceneManagement.h
+++ b/Engine/Source/Runtime/Engine/Public/SceneManagement.h
@@ -735,7 +735,7 @@ public:
 	// @param InCascadeIndex cascade index or INDEX_NONE for the distance field cascade
 	// @param OutCascadeSettings can be 0
 	virtual FSphere GetShadowSplitBounds(const class FSceneView& View, int32 InCascadeIndex, FShadowCascadeSettings* OutCascadeSettings) const { return FSphere(FVector::ZeroVector, 0); }
-	virtual FSphere GetShadowSplitBoundsDepthRange(const FSceneView& View, float SplitNear, float SplitFar, FShadowCascadeSettings* OutCascadeSettings) const { return FSphere(FVector::ZeroVector, 0); }
+	virtual FSphere GetShadowSplitBoundsDepthRange(const FSceneView& View, FVector ViewOrigin, float SplitNear, float SplitFar, FShadowCascadeSettings* OutCascadeSettings) const { return FSphere(FVector::ZeroVector, 0); }
 
 	virtual bool GetScissorRect(FIntRect& ScissorRect, const FSceneView& View) const
 	{
@@ -1394,6 +1394,7 @@ private:
 
 	friend class FSceneRenderer;
 	friend class FProjectedShadowInfo;
+	friend class FUniformMeshConverter;
 };
 
 
diff --git a/Engine/Source/Runtime/Engine/Public/SceneView.h b/Engine/Source/Runtime/Engine/Public/SceneView.h
--- a/Engine/Source/Runtime/Engine/Public/SceneView.h
+++ b/Engine/Source/Runtime/Engine/Public/SceneView.h
@@ -392,6 +392,9 @@ public:
 	/** For stereoscopic rendering, whether or not this is a full pass, or a left / right eye pass */
 	EStereoscopicPass StereoPass;
 
+	/** Whether this view should render the first instance only of any meshes using instancing. */
+	bool bRenderFirstInstanceOnly;
+
 	/** Current buffer visualization mode */
 	FName CurrentBufferVisualizationMode;
 
diff --git a/Engine/Source/Runtime/Engine/Public/StaticMeshResources.h b/Engine/Source/Runtime/Engine/Public/StaticMeshResources.h
--- a/Engine/Source/Runtime/Engine/Public/StaticMeshResources.h
+++ b/Engine/Source/Runtime/Engine/Public/StaticMeshResources.h
@@ -776,6 +776,7 @@ public:
 	virtual bool CanBeOccluded() const override;
 	virtual void GetLightRelevance(const FLightSceneProxy* LightSceneProxy, bool& bDynamic, bool& bRelevant, bool& bLightMapped, bool& bShadowMapped) const override;
 	virtual void GetDistancefieldAtlasData(FBox& LocalVolumeBounds, FIntVector& OutBlockMin, FIntVector& OutBlockSize, bool& bOutBuiltAsIfTwoSided, bool& bMeshWasPlane, TArray<FMatrix>& ObjectLocalToWorldTransforms) const override;
+	virtual void GetDistanceFieldInstanceInfo(int32& NumInstances, float& BoundsSurfaceArea) const override;
 	virtual bool HasDistanceFieldRepresentation() const override;
 	virtual uint32 GetMemoryFootprint( void ) const override { return( sizeof( *this ) + GetAllocatedSize() ); }
 	uint32 GetAllocatedSize( void ) const { return( FPrimitiveSceneProxy::GetAllocatedSize() + LODs.GetAllocatedSize() ); }
diff --git a/Engine/Source/Runtime/Landscape/Private/LandscapeRender.cpp b/Engine/Source/Runtime/Landscape/Private/LandscapeRender.cpp
--- a/Engine/Source/Runtime/Landscape/Private/LandscapeRender.cpp
+++ b/Engine/Source/Runtime/Landscape/Private/LandscapeRender.cpp
@@ -2606,10 +2606,17 @@ void FLandscapeComponentSceneProxy::GetHeightfieldRepresentation(UTexture2D*& Ou
 	OutHeightmapTexture = HeightmapTexture;
 	OutHeightfieldScaleBias = HeightmapScaleBias;
 
+	check(OutHeightfieldScaleBias.X > 0);
+
+	// CalculateHeightfieldOcclusionCS needs to be fixed up if other values are ever supported
+	check(NumSubsections == 1 || NumSubsections == 2);
+
+	// Store the presence of subsections in the sign bit
+	OutHeightfieldScaleBias.X *= NumSubsections > 1 ? -1 : 1;
+
 	OutMinMaxUV = FVector4(
 		HeightmapScaleBias.Z, 
 		HeightmapScaleBias.W, 
-		//@todo - subsections duplicate a line of border heights in the heightmap, take this into account in the UV mapping
 		HeightmapScaleBias.Z + SubsectionSizeVerts * NumSubsections * HeightmapScaleBias.X, 
 		HeightmapScaleBias.W + SubsectionSizeVerts * NumSubsections * HeightmapScaleBias.Y);
 }
diff --git a/Engine/Source/Runtime/Renderer/Private/ConvertToUniformMesh.cpp b/Engine/Source/Runtime/Renderer/Private/ConvertToUniformMesh.cpp
new file mode 100644
--- /dev/null
+++ b/Engine/Source/Runtime/Renderer/Private/ConvertToUniformMesh.cpp
@@ -0,0 +1,576 @@
+// Copyright 1998-2014 Epic Games, Inc. All Rights Reserved.
+
+/*=============================================================================
+	SurfelTree.cpp
+=============================================================================*/
+
+#include "RendererPrivate.h"
+#include "ScenePrivate.h"
+#include "UniformBuffer.h"
+#include "ShaderParameters.h"
+#include "PostProcessing.h"
+#include "SceneFilterRendering.h"
+#include "DistanceFieldLightingShared.h"
+#include "DistanceFieldSurfaceCacheLighting.h"
+#include "DistanceFieldGlobalIllumination.h"
+#include "RHICommandList.h"
+#include "SceneUtils.h"
+#include "DistanceFieldAtlas.h"
+
+
+class FConvertToUniformMeshVS : public FMeshMaterialShader
+{
+	DECLARE_SHADER_TYPE(FConvertToUniformMeshVS,MeshMaterial);
+
+protected:
+
+	FConvertToUniformMeshVS(const ShaderMetaType::CompiledShaderInitializerType& Initializer)
+		:	FMeshMaterialShader(Initializer)
+	{
+	}
+
+	FConvertToUniformMeshVS()
+	{
+	}
+
+	static bool ShouldCache(EShaderPlatform Platform,const FMaterial* Material,const FVertexFactoryType* VertexFactoryType)
+	{
+		return IsFeatureLevelSupported(Platform, ERHIFeatureLevel::SM5) 
+			&& DoesPlatformSupportDistanceFieldAO(Platform)
+			&& (FCString::Strstr(VertexFactoryType->GetName(), TEXT("LocalVertexFactory")) != NULL
+				|| FCString::Strstr(VertexFactoryType->GetName(), TEXT("InstancedStaticMeshVertexFactory")) != NULL);
+	}
+
+public:
+	
+	void SetParameters(FRHICommandList& RHICmdList, const FVertexFactory* VertexFactory,const FMaterialRenderProxy* MaterialRenderProxy,const FSceneView* View)
+	{
+		FMeshMaterialShader::SetParameters(RHICmdList, GetVertexShader(), MaterialRenderProxy, *MaterialRenderProxy->GetMaterial(View->GetFeatureLevel()), *View, ESceneRenderTargetsMode::SetTextures);
+	}
+
+	void SetMesh(FRHICommandList& RHICmdList, const FVertexFactory* VertexFactory,const FSceneView& View,const FPrimitiveSceneProxy* Proxy,const FMeshBatchElement& BatchElement)
+	{
+		FMeshMaterialShader::SetMesh(RHICmdList, GetVertexShader(),VertexFactory,View,Proxy,BatchElement);
+	}
+};
+
+IMPLEMENT_MATERIAL_SHADER_TYPE(,FConvertToUniformMeshVS,TEXT("ConvertToUniformMesh"),TEXT("ConvertToUniformMeshVS"),SF_Vertex); 
+
+FStreamOutElementList GUniformMeshVertexLayout;
+
+// In float4's, must match usf
+int32 FSurfelBuffers::InterpolatedVertexDataStride = 6;
+
+/** Returns number of float's in the uniform vertex. */
+int32 ComputeUniformVertexStride()
+{
+	if (GUniformMeshVertexLayout.Num() == 0)
+	{
+		GUniformMeshVertexLayout.Add(FStreamOutElement(0, "SV_Position", 0, 4, 0));
+		GUniformMeshVertexLayout.Add(FStreamOutElement(0, "Tangent", 0, 3, 0));
+		GUniformMeshVertexLayout.Add(FStreamOutElement(0, "Tangent", 1, 3, 0));
+		GUniformMeshVertexLayout.Add(FStreamOutElement(0, "Tangent", 2, 3, 0));
+		GUniformMeshVertexLayout.Add(FStreamOutElement(0, "UV", 0, 2, 0));
+		GUniformMeshVertexLayout.Add(FStreamOutElement(0, "UV", 1, 2, 0));
+		GUniformMeshVertexLayout.Add(FStreamOutElement(0, "VertexColor", 0, 4, 0));
+	}
+
+	int32 StreamStride = 0;
+
+	for (int32 ElementIndex = 0; ElementIndex < GUniformMeshVertexLayout.Num(); ElementIndex++)
+	{
+		StreamStride += GUniformMeshVertexLayout[ElementIndex].ComponentCount;
+	}
+
+	// D3D11 stream out buffer element stride must be a factor of 4
+	return FMath::DivideAndRoundUp(StreamStride, 4) * 4;
+}
+
+void FUniformMeshBuffers::Initialize()
+{
+	if (MaxElements > 0)
+	{
+		const int32 VertexStride = ComputeUniformVertexStride();
+		FRHIResourceCreateInfo CreateInfo;
+		TriangleData = RHICreateVertexBuffer(MaxElements * VertexStride * GPixelFormats[PF_R32_FLOAT].BlockBytes, BUF_ShaderResource | BUF_StreamOutput, CreateInfo);
+		TriangleDataSRV = RHICreateShaderResourceView(TriangleData, GPixelFormats[PF_R32_FLOAT].BlockBytes, PF_R32_FLOAT);
+
+		TriangleAreas.Initialize(sizeof(float), MaxElements, PF_R32_FLOAT);
+		TriangleCDFs.Initialize(sizeof(float), MaxElements, PF_R32_FLOAT);
+	}
+}
+
+class FConvertToUniformMeshGS : public FMeshMaterialShader
+{
+	DECLARE_SHADER_TYPE(FConvertToUniformMeshGS,MeshMaterial);
+
+protected:
+
+	FConvertToUniformMeshGS(const ShaderMetaType::CompiledShaderInitializerType& Initializer)
+		:	FMeshMaterialShader(Initializer)
+	{
+	}
+
+	FConvertToUniformMeshGS()
+	{
+	}
+
+	static bool ShouldCache(EShaderPlatform Platform,const FMaterial* Material,const FVertexFactoryType* VertexFactoryType)
+	{
+		return IsFeatureLevelSupported(Platform, ERHIFeatureLevel::SM5) 
+			&& DoesPlatformSupportDistanceFieldAO(Platform)
+			&& (FCString::Strstr(VertexFactoryType->GetName(), TEXT("LocalVertexFactory")) != NULL
+				|| FCString::Strstr(VertexFactoryType->GetName(), TEXT("InstancedStaticMeshVertexFactory")) != NULL);
+	}
+
+public:
+	
+	void SetParameters(FRHICommandList& RHICmdList, const FVertexFactory* VertexFactory,const FMaterialRenderProxy* MaterialRenderProxy,const FSceneView* View)
+	{
+		FMeshMaterialShader::SetParameters(RHICmdList, (FGeometryShaderRHIParamRef)GetGeometryShader(), MaterialRenderProxy, *MaterialRenderProxy->GetMaterial(View->GetFeatureLevel()), *View, ESceneRenderTargetsMode::SetTextures);
+	}
+
+	void SetMesh(FRHICommandList& RHICmdList, const FVertexFactory* VertexFactory,const FSceneView& View,const FPrimitiveSceneProxy* Proxy,const FMeshBatchElement& BatchElement)
+	{
+		FMeshMaterialShader::SetMesh(RHICmdList, (FGeometryShaderRHIParamRef)GetGeometryShader(),VertexFactory,View,Proxy,BatchElement);
+	}
+
+	virtual void GetStreamOutElements(FStreamOutElementList& ElementList, TArray<uint32>& StreamStrides, int32& RasterizedStream) override
+	{
+		StreamStrides.Add(ComputeUniformVertexStride() * 4);
+		ElementList = GUniformMeshVertexLayout;
+		RasterizedStream = -1;
+	}
+};
+
+IMPLEMENT_MATERIAL_SHADER_TYPE(,FConvertToUniformMeshGS,TEXT("ConvertToUniformMesh"),TEXT("ConvertToUniformMeshGS"),SF_Geometry); 
+
+
+class FConvertToUniformMeshDrawingPolicy : public FMeshDrawingPolicy
+{
+public:
+	/** context type */
+	typedef FMeshDrawingPolicy::ElementDataType ElementDataType;
+
+	/**
+	* Constructor
+	* @param InIndexBuffer - index buffer for rendering
+	* @param InVertexFactory - vertex factory for rendering
+	* @param InMaterialRenderProxy - material instance for rendering
+	* @param bInOverrideWithShaderComplexity - whether to override with shader complexity
+	*/
+	FConvertToUniformMeshDrawingPolicy(
+		const FVertexFactory* InVertexFactory,
+		const FMaterialRenderProxy* InMaterialRenderProxy,
+		const FMaterial& MaterialResouce,
+		ERHIFeatureLevel::Type InFeatureLevel
+		);
+
+	// FMeshDrawingPolicy interface.
+
+	/**
+	* Match two draw policies
+	* @param Other - draw policy to compare
+	* @return true if the draw policies are a match
+	*/
+	bool Matches(const FConvertToUniformMeshDrawingPolicy& Other) const;
+
+	/**
+	* Executes the draw commands which can be shared between any meshes using this drawer.
+	* @param CI - The command interface to execute the draw commands on.
+	* @param View - The view of the scene being drawn.
+	*/
+	void SetSharedState(FRHICommandList& RHICmdList, const FSceneView* View, const ContextDataType PolicyContext) const;
+	
+	/** 
+	* Create bound shader state using the vertex decl from the mesh draw policy
+	* as well as the shaders needed to draw the mesh
+	* @return new bound shader state object
+	*/
+	FBoundShaderStateInput GetBoundShaderStateInput(ERHIFeatureLevel::Type InFeatureLevel);
+
+	/**
+	* Sets the render states for drawing a mesh.
+	* @param PrimitiveSceneProxy - The primitive drawing the dynamic mesh.  If this is a view element, this will be NULL.
+	* @param Mesh - mesh element with data needed for rendering
+	* @param ElementData - context specific data for mesh rendering
+	*/
+	void SetMeshRenderState(
+		FRHICommandList& RHICmdList, 
+		const FSceneView& View,
+		const FPrimitiveSceneProxy* PrimitiveSceneProxy,
+		const FMeshBatch& Mesh,
+		int32 BatchElementIndex,
+		bool bBackFace,
+		const ElementDataType& ElementData,
+		const ContextDataType PolicyContext
+		) const;
+
+private:
+	FConvertToUniformMeshVS* VertexShader;
+	FConvertToUniformMeshGS* GeometryShader;
+};
+
+FConvertToUniformMeshDrawingPolicy::FConvertToUniformMeshDrawingPolicy(
+	const FVertexFactory* InVertexFactory,
+	const FMaterialRenderProxy* InMaterialRenderProxy,
+	const FMaterial& InMaterialResource,
+	ERHIFeatureLevel::Type InFeatureLevel
+	)
+:	FMeshDrawingPolicy(InVertexFactory,InMaterialRenderProxy,InMaterialResource,false)
+{
+	VertexShader = InMaterialResource.GetShader<FConvertToUniformMeshVS>(InVertexFactory->GetType());
+	GeometryShader = InMaterialResource.GetShader<FConvertToUniformMeshGS>(InVertexFactory->GetType());
+}
+
+bool FConvertToUniformMeshDrawingPolicy::Matches(
+	const FConvertToUniformMeshDrawingPolicy& Other
+	) const
+{
+	return FMeshDrawingPolicy::Matches(Other) &&
+		VertexShader == Other.VertexShader &&
+		GeometryShader == Other.GeometryShader;
+}
+
+void FConvertToUniformMeshDrawingPolicy::SetSharedState(
+	FRHICommandList& RHICmdList, 
+	const FSceneView* View,
+	const ContextDataType PolicyContext
+	) const
+{
+	// Set shared mesh resources
+	FMeshDrawingPolicy::SetSharedState(RHICmdList, View, PolicyContext);
+
+	VertexShader->SetParameters(RHICmdList, VertexFactory, MaterialRenderProxy, View);
+	GeometryShader->SetParameters(RHICmdList, VertexFactory, MaterialRenderProxy, View);
+}
+
+FBoundShaderStateInput FConvertToUniformMeshDrawingPolicy::GetBoundShaderStateInput(ERHIFeatureLevel::Type InFeatureLevel)
+{
+	return FBoundShaderStateInput(
+		FMeshDrawingPolicy::GetVertexDeclaration(), 
+		VertexShader->GetVertexShader(),
+		NULL, 
+		NULL,
+		NULL,
+		GeometryShader->GetGeometryShader());
+
+}
+
+void FConvertToUniformMeshDrawingPolicy::SetMeshRenderState(
+	FRHICommandList& RHICmdList, 
+	const FSceneView& View,
+	const FPrimitiveSceneProxy* PrimitiveSceneProxy,
+	const FMeshBatch& Mesh,
+	int32 BatchElementIndex,
+	bool bBackFace,
+	const ElementDataType& ElementData,
+	const ContextDataType PolicyContext
+	) const
+{
+	const FMeshBatchElement& BatchElement = Mesh.Elements[BatchElementIndex];
+
+	EmitMeshDrawEvents(RHICmdList, PrimitiveSceneProxy, Mesh);
+
+	// Set transforms
+	VertexShader->SetMesh(RHICmdList, VertexFactory,View,PrimitiveSceneProxy,BatchElement);
+	GeometryShader->SetMesh(RHICmdList, VertexFactory,View,PrimitiveSceneProxy,BatchElement);
+	
+	// Set rasterizer state.
+	RHICmdList.SetRasterizerState(GetStaticRasterizerState<true>(
+		(Mesh.bWireframe || IsWireframe()) ? FM_Wireframe : FM_Solid,
+		IsTwoSided() ? CM_None : (XOR( XOR(View.bReverseCulling,bBackFace), Mesh.ReverseCulling) ? CM_CCW : CM_CW)));
+}
+
+bool ShouldGenerateSurfelsOnMesh(const FMeshBatch& Mesh, ERHIFeatureLevel::Type FeatureLevel)
+{
+	//@todo - support for tessellated meshes
+	return Mesh.Type == PT_TriangleList 
+		&& !Mesh.IsTranslucent(FeatureLevel) 
+		&& Mesh.MaterialRenderProxy->GetMaterial(FeatureLevel)->GetShadingModel() != MSM_Unlit;
+}
+
+FUniformMeshBuffers GUniformMeshTemporaryBuffers;
+
+TArray<TArray<FMeshBatchAndRelevance,SceneRenderingAllocator>> GDynamicMeshElements;
+
+int32 FUniformMeshConverter::Convert(
+	FRHICommandListImmediate& RHICmdList, 
+	FSceneRenderer& Renderer,
+	FViewInfo& View, 
+	const FPrimitiveSceneInfo* PrimitiveSceneInfo, 
+	FUniformMeshBuffers*& OutUniformMeshBuffers,
+	const FMaterialRenderProxy*& OutMaterialRenderProxy,
+	FUniformBufferRHIParamRef& OutPrimitiveUniformBuffer)
+{
+	const FPrimitiveSceneProxy* PrimitiveSceneProxy = PrimitiveSceneInfo->Proxy;
+	const auto FeatureLevel = View.GetFeatureLevel();
+
+	if (GDynamicMeshElements.Num() == 0)
+	{
+		GDynamicMeshElements.Empty(10000);
+	}
+
+	FSimpleElementCollector DynamicSimpleElements;
+	GDynamicMeshElements.Add(TArray<FMeshBatchAndRelevance,SceneRenderingAllocator>());
+	TArray<FMeshBatchAndRelevance,SceneRenderingAllocator>& DynamicMeshElements = GDynamicMeshElements.Last();
+	TArray<const FSceneView*> ViewsArray;
+	ViewsArray.Add(&View);
+
+	Renderer.MeshCollector.ClearViewMeshArrays();
+	Renderer.MeshCollector.AddViewMeshArrays(&View, &DynamicMeshElements, &DynamicSimpleElements, FeatureLevel);
+	
+	View.bRenderFirstInstanceOnly = true;
+	Renderer.MeshCollector.SetPrimitive(PrimitiveSceneInfo->Proxy, PrimitiveSceneInfo->DefaultDynamicHitProxyId);
+	PrimitiveSceneInfo->Proxy->GetDynamicMeshElements(ViewsArray, *(View.Family), 0x1, Renderer.MeshCollector);
+	View.bRenderFirstInstanceOnly = false;
+
+	int32 NumTriangles = 0;
+
+	FPrimitiveViewRelevance ViewRelevance = PrimitiveSceneInfo->Proxy->GetViewRelevance(&View);
+
+	if (ViewRelevance.bDynamicRelevance)
+	{
+		for (int32 MeshBatchIndex = 0; MeshBatchIndex < DynamicMeshElements.Num(); MeshBatchIndex++)
+		{
+			const FMeshBatchAndRelevance& MeshBatchAndRelevance = DynamicMeshElements[MeshBatchIndex];
+
+			if (ShouldGenerateSurfelsOnMesh(*MeshBatchAndRelevance.Mesh, FeatureLevel))
+			{
+				NumTriangles += MeshBatchAndRelevance.Mesh->GetNumPrimitives();
+			}
+		}
+	}
+
+	if (ViewRelevance.bStaticRelevance)
+	{
+		float LargestScreenSize = 0;
+
+		for (int32 StaticMeshIdx=0; StaticMeshIdx < PrimitiveSceneInfo->StaticMeshes.Num(); StaticMeshIdx++)
+		{
+			const FStaticMesh& StaticMesh = PrimitiveSceneInfo->StaticMeshes[StaticMeshIdx];
+			LargestScreenSize = FMath::Max(LargestScreenSize, StaticMesh.ScreenSize);
+		}
+
+		for (int32 StaticMeshIdx=0; StaticMeshIdx < PrimitiveSceneInfo->StaticMeshes.Num(); StaticMeshIdx++)
+		{
+			const FStaticMesh& StaticMesh = PrimitiveSceneInfo->StaticMeshes[StaticMeshIdx];
+
+			if (ShouldGenerateSurfelsOnMesh(StaticMesh, FeatureLevel)
+				&& !StaticMesh.bShadowOnly
+				// LOD0 only
+				&& StaticMesh.ScreenSize == LargestScreenSize)
+			{
+				NumTriangles += StaticMesh.GetNumPrimitives();
+			}
+		}
+	}
+
+	if (NumTriangles > 0)
+	{
+		if (GUniformMeshTemporaryBuffers.MaxElements < NumTriangles * 3)
+		{
+			GUniformMeshTemporaryBuffers.MaxElements = NumTriangles * 3;
+			GUniformMeshTemporaryBuffers.Release();
+			GUniformMeshTemporaryBuffers.Initialize();
+		}
+
+		RHICmdList.SetRenderTargets(0, (const FRHIRenderTargetView*)NULL, NULL, 0, (const FUnorderedAccessViewRHIParamRef*)NULL);
+
+		uint32 Offsets[1] = {0};
+		const FVertexBufferRHIParamRef StreamOutTargets[1] = {GUniformMeshTemporaryBuffers.TriangleData.GetReference()};
+		RHICmdList.SetStreamOutTargets(1, StreamOutTargets, Offsets);
+
+		if (ViewRelevance.bDynamicRelevance)
+		{
+			for (int32 MeshBatchIndex = 0; MeshBatchIndex < DynamicMeshElements.Num(); MeshBatchIndex++)
+			{
+				const FMeshBatchAndRelevance& MeshBatchAndRelevance = DynamicMeshElements[MeshBatchIndex];
+
+				if (ShouldGenerateSurfelsOnMesh(*MeshBatchAndRelevance.Mesh, FeatureLevel))
+				{
+					const FMeshBatch& Mesh = *MeshBatchAndRelevance.Mesh;
+
+					FConvertToUniformMeshDrawingPolicy DrawingPolicy(
+						Mesh.VertexFactory,
+						Mesh.MaterialRenderProxy,
+						*Mesh.MaterialRenderProxy->GetMaterial(FeatureLevel),
+						FeatureLevel);
+
+					//@todo - fix
+					OutMaterialRenderProxy = Mesh.MaterialRenderProxy;
+
+					RHICmdList.BuildAndSetLocalBoundShaderState(DrawingPolicy.GetBoundShaderStateInput(FeatureLevel));
+					DrawingPolicy.SetSharedState(RHICmdList, &View, FConvertToUniformMeshDrawingPolicy::ContextDataType());
+
+					for (int32 BatchElementIndex = 0; BatchElementIndex < Mesh.Elements.Num(); BatchElementIndex++)
+					{
+						//@todo - fix
+						OutPrimitiveUniformBuffer = IsValidRef(Mesh.Elements[BatchElementIndex].PrimitiveUniformBuffer) 
+							? Mesh.Elements[BatchElementIndex].PrimitiveUniformBuffer
+							: *Mesh.Elements[BatchElementIndex].PrimitiveUniformBufferResource;
+
+						DrawingPolicy.SetMeshRenderState(RHICmdList, View,PrimitiveSceneProxy,Mesh,BatchElementIndex,false,FConvertToUniformMeshDrawingPolicy::ElementDataType(), FConvertToUniformMeshDrawingPolicy::ContextDataType());
+						DrawingPolicy.DrawMesh(RHICmdList, Mesh, BatchElementIndex);
+					}
+				}
+			}
+		}
+
+		if (ViewRelevance.bStaticRelevance)
+		{
+			float LargestScreenSize = 0;
+
+			for (int32 StaticMeshIdx=0; StaticMeshIdx < PrimitiveSceneInfo->StaticMeshes.Num(); StaticMeshIdx++)
+			{
+				const FStaticMesh& StaticMesh = PrimitiveSceneInfo->StaticMeshes[StaticMeshIdx];
+				LargestScreenSize = FMath::Max(LargestScreenSize, StaticMesh.ScreenSize);
+			}
+
+			for (int32 StaticMeshIdx = 0; StaticMeshIdx < PrimitiveSceneInfo->StaticMeshes.Num(); StaticMeshIdx++)
+			{
+				const FStaticMesh& StaticMesh = PrimitiveSceneInfo->StaticMeshes[StaticMeshIdx];
+
+				if (ShouldGenerateSurfelsOnMesh(StaticMesh, FeatureLevel) 
+					&& !StaticMesh.bShadowOnly
+					&& StaticMesh.ScreenSize == LargestScreenSize)
+				{
+					FConvertToUniformMeshDrawingPolicy DrawingPolicy(
+						StaticMesh.VertexFactory,
+						StaticMesh.MaterialRenderProxy,
+						*StaticMesh.MaterialRenderProxy->GetMaterial(FeatureLevel),
+						FeatureLevel
+						);
+					RHICmdList.BuildAndSetLocalBoundShaderState(DrawingPolicy.GetBoundShaderStateInput(FeatureLevel));
+					DrawingPolicy.SetSharedState(RHICmdList, &View, FConvertToUniformMeshDrawingPolicy::ContextDataType());
+
+					//@todo - fix
+					OutMaterialRenderProxy = StaticMesh.MaterialRenderProxy;
+
+					for (int32 BatchElementIndex = 0; BatchElementIndex < StaticMesh.Elements.Num(); BatchElementIndex++)
+					{
+						//@todo - fix
+						OutPrimitiveUniformBuffer = IsValidRef(StaticMesh.Elements[BatchElementIndex].PrimitiveUniformBuffer) 
+							? StaticMesh.Elements[BatchElementIndex].PrimitiveUniformBuffer
+							: *StaticMesh.Elements[BatchElementIndex].PrimitiveUniformBufferResource;
+
+						DrawingPolicy.SetMeshRenderState(RHICmdList, View,PrimitiveSceneProxy,StaticMesh,BatchElementIndex,false,
+							FConvertToUniformMeshDrawingPolicy::ElementDataType(),
+							FConvertToUniformMeshDrawingPolicy::ContextDataType()
+							);
+						DrawingPolicy.DrawMesh(RHICmdList, StaticMesh, BatchElementIndex);
+					}
+				}
+			}
+		}
+
+		RHICmdList.SetStreamOutTargets(1, (const FVertexBufferRHIParamRef*)NULL, Offsets);
+	}
+
+	OutUniformMeshBuffers = &GUniformMeshTemporaryBuffers;
+	return NumTriangles;
+}
+
+int32 GEvaluateSurfelMaterialGroupSize = 64;
+
+class FEvaluateSurfelMaterialCS : public FMaterialShader
+{
+	DECLARE_SHADER_TYPE(FEvaluateSurfelMaterialCS,Material)
+public:
+
+	static bool ShouldCache(EShaderPlatform Platform, const FMaterial* Material)
+	{
+		//@todo - lit materials only 
+		return IsFeatureLevelSupported(Platform, ERHIFeatureLevel::SM5) && DoesPlatformSupportDistanceFieldAO(Platform);
+	}
+
+	static void ModifyCompilationEnvironment(EShaderPlatform Platform, const FMaterial* Material, FShaderCompilerEnvironment& OutEnvironment)
+	{
+		FMaterialShader::ModifyCompilationEnvironment(Platform,OutEnvironment);
+		OutEnvironment.SetDefine(TEXT("EVALUATE_SURFEL_MATERIAL_GROUP_SIZE"), GEvaluateSurfelMaterialGroupSize);
+		OutEnvironment.SetDefine(TEXT("HAS_PRIMITIVE_UNIFORM_BUFFER"), 1);
+	}
+
+	FEvaluateSurfelMaterialCS(const ShaderMetaType::CompiledShaderInitializerType& Initializer)
+		: FMaterialShader(Initializer)
+	{
+		SurfelBufferParameters.Bind(Initializer.ParameterMap);
+		SurfelStartIndex.Bind(Initializer.ParameterMap, TEXT("SurfelStartIndex"));
+		NumSurfelsToGenerate.Bind(Initializer.ParameterMap, TEXT("NumSurfelsToGenerate"));
+		Instance0InverseTransform.Bind(Initializer.ParameterMap, TEXT("Instance0InverseTransform"));
+	}
+
+	FEvaluateSurfelMaterialCS()
+	{
+	}
+
+	void SetParameters(
+		FRHICommandList& RHICmdList, 
+		const FSceneView& View, 
+		int32 SurfelStartIndexValue,
+		int32 NumSurfelsToGenerateValue,
+		const FMaterialRenderProxy* MaterialProxy,
+		FUniformBufferRHIParamRef PrimitiveUniformBuffer,
+		const FMatrix& Instance0Transform
+		)
+	{
+		FComputeShaderRHIParamRef ShaderRHI = GetComputeShader();
+		FMaterialShader::SetParameters(RHICmdList, ShaderRHI, MaterialProxy, *MaterialProxy->GetMaterial(View.GetFeatureLevel()), View, true, ESceneRenderTargetsMode::SetTextures);
+
+		SetUniformBufferParameter(RHICmdList, ShaderRHI,GetUniformBufferParameter<FPrimitiveUniformShaderParameters>(),PrimitiveUniformBuffer);
+
+		const FScene* Scene = (const FScene*)View.Family->Scene;
+		SurfelBufferParameters.Set(RHICmdList, ShaderRHI, *Scene->DistanceFieldSceneData.SurfelBuffers, *Scene->DistanceFieldSceneData.InstancedSurfelBuffers);
+		
+		SetShaderValue(RHICmdList, ShaderRHI, SurfelStartIndex, SurfelStartIndexValue);
+		SetShaderValue(RHICmdList, ShaderRHI, NumSurfelsToGenerate, NumSurfelsToGenerateValue);
+
+		SetShaderValue(RHICmdList, ShaderRHI, Instance0InverseTransform, Instance0Transform.Inverse());
+	}
+
+	void UnsetParameters(FRHICommandList& RHICmdList)
+	{
+		FComputeShaderRHIParamRef ShaderRHI = GetComputeShader();
+		SurfelBufferParameters.UnsetParameters(RHICmdList, ShaderRHI);
+	}
+
+	virtual bool Serialize(FArchive& Ar)
+	{		
+		bool bShaderHasOutdatedParameters = FMaterialShader::Serialize(Ar);
+		Ar << SurfelBufferParameters;
+		Ar << SurfelStartIndex;
+		Ar << NumSurfelsToGenerate;
+		Ar << Instance0InverseTransform;
+		return bShaderHasOutdatedParameters;
+	}
+
+private:
+
+	FSurfelBufferParameters SurfelBufferParameters;
+	FShaderParameter SurfelStartIndex;
+	FShaderParameter NumSurfelsToGenerate;
+	FShaderParameter Instance0InverseTransform;
+};
+
+IMPLEMENT_MATERIAL_SHADER_TYPE(,FEvaluateSurfelMaterialCS,TEXT("EvaluateSurfelMaterial"),TEXT("EvaluateSurfelMaterialCS"),SF_Compute);
+
+void FUniformMeshConverter::GenerateSurfels(
+	FRHICommandListImmediate& RHICmdList, 
+	FViewInfo& View, 
+	const FPrimitiveSceneInfo* PrimitiveSceneInfo, 
+	const FMaterialRenderProxy* MaterialProxy,
+	FUniformBufferRHIParamRef PrimitiveUniformBuffer, 
+	const FMatrix& Instance0Transform,
+	int32 SurfelOffset,
+	int32 NumSurfels)
+{
+	const FMaterial* Material = MaterialProxy->GetMaterial(View.GetFeatureLevel());
+	const FMaterialShaderMap* MaterialShaderMap = Material->GetRenderingThreadShaderMap();
+	FEvaluateSurfelMaterialCS* ComputeShader = MaterialShaderMap->GetShader<FEvaluateSurfelMaterialCS>();
+
+	RHICmdList.SetComputeShader(ComputeShader->GetComputeShader());
+	ComputeShader->SetParameters(RHICmdList, View, SurfelOffset, NumSurfels, MaterialProxy, PrimitiveUniformBuffer, Instance0Transform);
+	DispatchComputeShader(RHICmdList, ComputeShader, FMath::DivideAndRoundUp(NumSurfels, GEvaluateSurfelMaterialGroupSize), 1, 1);
+
+	ComputeShader->UnsetParameters(RHICmdList);
+}
\ No newline at end of file
diff --git a/Engine/Source/Runtime/Renderer/Private/DistanceFieldGlobalIllumination.cpp b/Engine/Source/Runtime/Renderer/Private/DistanceFieldGlobalIllumination.cpp
--- a/Engine/Source/Runtime/Renderer/Private/DistanceFieldGlobalIllumination.cpp
+++ b/Engine/Source/Runtime/Renderer/Private/DistanceFieldGlobalIllumination.cpp
@@ -22,6 +22,14 @@
 	TEXT("r.DistanceFieldGI"),
 	GDistanceFieldGI,
 	TEXT(""),
+	ECVF_Cheat | ECVF_RenderThreadSafe | ECVF_ReadOnly
+	);
+
+int32 GVPLSurfelRepresentation = 1;
+FAutoConsoleVariableRef CVarVPLSurfelRepresentation(
+	TEXT("r.VPLSurfelRepresentation"),
+	GVPLSurfelRepresentation,
+	TEXT(""),
 	ECVF_Cheat | ECVF_RenderThreadSafe
 	);
 
@@ -65,9 +73,19 @@
 	ECVF_Cheat | ECVF_RenderThreadSafe
 	);
 
+int32 GVPLSpreadUpdateOver = 5;
+FAutoConsoleVariableRef CVarVPLSpreadUpdateOver(
+	TEXT("r.VPLSpreadUpdateOver"),
+	GVPLSpreadUpdateOver,
+	TEXT(""),
+	ECVF_Cheat | ECVF_RenderThreadSafe
+	);
+
 TGlobalResource<FVPLResources> GVPLResources;
 TGlobalResource<FVPLResources> GCulledVPLResources;
 
+extern TGlobalResource<FDistanceFieldObjectBufferResource> GShadowCulledObjectBuffers;
+
 class FVPLPlacementCS : public FGlobalShader
 {
 	DECLARE_SHADER_TYPE(FVPLPlacementCS, Global);
@@ -92,7 +110,6 @@
 		: FGlobalShader(Initializer)
 	{
 		VPLParameterBuffer.Bind(Initializer.ParameterMap, TEXT("VPLParameterBuffer"));
-		VPLClusterData.Bind(Initializer.ParameterMap, TEXT("VPLClusterData"));
 		VPLData.Bind(Initializer.ParameterMap, TEXT("VPLData"));
 		InvPlacementGridSize.Bind(Initializer.ParameterMap, TEXT("InvPlacementGridSize"));
 		WorldToShadow.Bind(Initializer.ParameterMap, TEXT("WorldToShadow"));
@@ -117,7 +134,6 @@
 	{
 		const FComputeShaderRHIParamRef ShaderRHI = GetComputeShader();
 		FGlobalShader::SetParameters(RHICmdList, ShaderRHI, View);
-		extern TGlobalResource<FDistanceFieldObjectBufferResource> GShadowCulledObjectBuffers;
 		ObjectParameters.Set(RHICmdList, ShaderRHI, GShadowCulledObjectBuffers.Buffers);
 
 		SetShaderValue(RHICmdList, ShaderRHI, InvPlacementGridSize, InvPlacementGridSizeValue);
@@ -127,7 +143,6 @@
 		SetShaderValue(RHICmdList, ShaderRHI, LightColor, LightSceneProxy->GetColor() * LightSceneProxy->GetIndirectLightingScale());
 
 		VPLParameterBuffer.SetBuffer(RHICmdList, ShaderRHI, GVPLResources.VPLParameterBuffer);
-		VPLClusterData.SetBuffer(RHICmdList, ShaderRHI, GVPLResources.VPLClusterData);
 		VPLData.SetBuffer(RHICmdList, ShaderRHI, GVPLResources.VPLData);
 
 		check(TileIntersectionResources || !ShadowTileHeadDataUnpacked.IsBound());
@@ -145,7 +160,6 @@
 	void UnsetParameters(FRHICommandList& RHICmdList)
 	{
 		VPLParameterBuffer.UnsetUAV(RHICmdList, GetComputeShader());
-		VPLClusterData.UnsetUAV(RHICmdList, GetComputeShader());
 		VPLData.UnsetUAV(RHICmdList, GetComputeShader());
 	}
 
@@ -154,7 +168,6 @@
 	{
 		bool bShaderHasOutdatedParameters = FGlobalShader::Serialize(Ar);
 		Ar << VPLParameterBuffer;
-		Ar << VPLClusterData;
 		Ar << VPLData;
 		Ar << InvPlacementGridSize;
 		Ar << WorldToShadow;
@@ -172,7 +185,6 @@
 private:
 
 	FRWShaderParameter VPLParameterBuffer;
-	FRWShaderParameter VPLClusterData;
 	FRWShaderParameter VPLData;
 	FShaderParameter InvPlacementGridSize;
 	FShaderParameter WorldToShadow;
@@ -248,7 +260,6 @@
 
 IMPLEMENT_SHADER_TYPE(,FSetupVPLCullndirectArgumentsCS,TEXT("DistanceFieldGlobalIllumination"),TEXT("SetupVPLCullndirectArgumentsCS"),SF_Compute);
 
-
 class FCullVPLsForViewCS : public FGlobalShader
 {
 	DECLARE_SHADER_TYPE(FCullVPLsForViewCS,Global)
@@ -270,10 +281,8 @@
 		: FGlobalShader(Initializer)
 	{
 		VPLParameterBuffer.Bind(Initializer.ParameterMap, TEXT("VPLParameterBuffer"));
-		VPLClusterData.Bind(Initializer.ParameterMap, TEXT("VPLClusterData"));
 		VPLData.Bind(Initializer.ParameterMap, TEXT("VPLData"));
 		CulledVPLParameterBuffer.Bind(Initializer.ParameterMap, TEXT("CulledVPLParameterBuffer"));
-		CulledVPLClusterData.Bind(Initializer.ParameterMap, TEXT("CulledVPLClusterData"));
 		CulledVPLData.Bind(Initializer.ParameterMap, TEXT("CulledVPLData"));
 		AOParameters.Bind(Initializer.ParameterMap);
 		NumConvexHullPlanes.Bind(Initializer.ParameterMap, TEXT("NumConvexHullPlanes"));
@@ -290,12 +299,10 @@
 		FGlobalShader::SetParameters(RHICmdList, ShaderRHI, View);
 
 		CulledVPLParameterBuffer.SetBuffer(RHICmdList, ShaderRHI, GCulledVPLResources.VPLParameterBuffer);
-		CulledVPLClusterData.SetBuffer(RHICmdList, ShaderRHI, GCulledVPLResources.VPLClusterData);
 		CulledVPLData.SetBuffer(RHICmdList, ShaderRHI, GCulledVPLResources.VPLData);
 		AOParameters.Set(RHICmdList, ShaderRHI, Parameters);
 
 		SetSRVParameter(RHICmdList, ShaderRHI, VPLParameterBuffer, GVPLResources.VPLParameterBuffer.SRV);
-		SetSRVParameter(RHICmdList, ShaderRHI, VPLClusterData, GVPLResources.VPLClusterData.SRV);
 		SetSRVParameter(RHICmdList, ShaderRHI, VPLData, GVPLResources.VPLData.SRV);
 
 		// Shader assumes max 6
@@ -307,7 +314,6 @@
 	void UnsetParameters(FRHICommandList& RHICmdList)
 	{
 		CulledVPLParameterBuffer.UnsetUAV(RHICmdList, GetComputeShader());
-		CulledVPLClusterData.UnsetUAV(RHICmdList, GetComputeShader());
 		CulledVPLData.UnsetUAV(RHICmdList, GetComputeShader());
 	}
 
@@ -315,10 +321,8 @@
 	{		
 		bool bShaderHasOutdatedParameters = FGlobalShader::Serialize(Ar);
 		Ar << VPLParameterBuffer;
-		Ar << VPLClusterData;
 		Ar << VPLData;
 		Ar << CulledVPLParameterBuffer;
-		Ar << CulledVPLClusterData;
 		Ar << CulledVPLData;
 		Ar << AOParameters;
 		Ar << NumConvexHullPlanes;
@@ -329,10 +333,8 @@
 private:
 
 	FShaderResourceParameter VPLParameterBuffer;
-	FShaderResourceParameter VPLClusterData;
 	FShaderResourceParameter VPLData;
 	FRWShaderParameter CulledVPLParameterBuffer;
-	FRWShaderParameter CulledVPLClusterData;
 	FRWShaderParameter CulledVPLData;
 	FAOParameters AOParameters;
 	FShaderParameter NumConvexHullPlanes;
@@ -341,289 +343,6 @@
 
 IMPLEMENT_SHADER_TYPE(,FCullVPLsForViewCS,TEXT("DistanceFieldGlobalIllumination"),TEXT("CullVPLsForViewCS"),SF_Compute);
 
-
-/**  */
-class FVPLTileCullCS : public FGlobalShader
-{
-	DECLARE_SHADER_TYPE(FVPLTileCullCS,Global)
-public:
-
-	static bool ShouldCache(EShaderPlatform Platform)
-	{
-		return IsFeatureLevelSupported(Platform, ERHIFeatureLevel::SM5) && DoesPlatformSupportDistanceFieldAO(Platform);
-	}
-
-	static void ModifyCompilationEnvironment(EShaderPlatform Platform, FShaderCompilerEnvironment& OutEnvironment)
-	{
-		FGlobalShader::ModifyCompilationEnvironment(Platform,OutEnvironment);
-		OutEnvironment.SetDefine(TEXT("THREADGROUP_SIZEX"), GDistanceFieldAOTileSizeX);
-		OutEnvironment.SetDefine(TEXT("THREADGROUP_SIZEY"), GDistanceFieldAOTileSizeY);
-		OutEnvironment.CompilerFlags.Add(CFLAG_StandardOptimization);
-	}
-
-	FVPLTileCullCS(const ShaderMetaType::CompiledShaderInitializerType& Initializer)
-		: FGlobalShader(Initializer)
-	{
-		AOParameters.Bind(Initializer.ParameterMap);
-		TileHeadDataUnpacked.Bind(Initializer.ParameterMap, TEXT("TileHeadDataUnpacked"));
-		TileArrayData.Bind(Initializer.ParameterMap, TEXT("TileArrayData"));
-		VPLParameterBuffer.Bind(Initializer.ParameterMap, TEXT("VPLParameterBuffer"));
-		VPLClusterData.Bind(Initializer.ParameterMap, TEXT("VPLClusterData"));
-		TileConeAxisAndCos.Bind(Initializer.ParameterMap, TEXT("TileConeAxisAndCos"));
-		TileConeDepthRanges.Bind(Initializer.ParameterMap, TEXT("TileConeDepthRanges"));
-		NumGroups.Bind(Initializer.ParameterMap, TEXT("NumGroups"));
-	}
-
-	FVPLTileCullCS()
-	{
-	}
-	void SetParameters(FRHICommandList& RHICmdList, const FSceneView& View, const FScene* Scene, const FDistanceFieldAOParameters& Parameters, const FVPLResources& VPLResources)
-	{
-		FComputeShaderRHIParamRef ShaderRHI = GetComputeShader();
-
-		FGlobalShader::SetParameters(RHICmdList, ShaderRHI, View);
-		AOParameters.Set(RHICmdList, ShaderRHI, Parameters);
-
-		FTileIntersectionResources* TileIntersectionResources = ((FSceneViewState*)View.State)->AOTileIntersectionResources;
-		FVPLTileIntersectionResources* VPLTileIntersectionResources = ((FSceneViewState*)View.State)->VPLTileIntersectionResources;
-
-		TileHeadDataUnpacked.SetBuffer(RHICmdList, ShaderRHI, VPLTileIntersectionResources->TileHeadDataUnpacked);
-		TileArrayData.SetBuffer(RHICmdList, ShaderRHI, VPLTileIntersectionResources->TileArrayData);
-
-		SetSRVParameter(RHICmdList, ShaderRHI, VPLParameterBuffer, VPLResources.VPLParameterBuffer.SRV);
-		SetSRVParameter(RHICmdList, ShaderRHI, VPLClusterData, VPLResources.VPLClusterData.SRV);
-
-		SetSRVParameter(RHICmdList, ShaderRHI, TileConeAxisAndCos, TileIntersectionResources->TileConeAxisAndCos.SRV);
-		SetSRVParameter(RHICmdList, ShaderRHI, TileConeDepthRanges, TileIntersectionResources->TileConeDepthRanges.SRV);
-	
-		SetShaderValue(RHICmdList, ShaderRHI, NumGroups, VPLTileIntersectionResources->TileDimensions);
-	}
-
-	void UnsetParameters(FRHICommandList& RHICmdList)
-	{
-		TileHeadDataUnpacked.UnsetUAV(RHICmdList, GetComputeShader());
-		TileArrayData.UnsetUAV(RHICmdList, GetComputeShader());
-	}
-
-	virtual bool Serialize(FArchive& Ar)
-	{		
-		bool bShaderHasOutdatedParameters = FGlobalShader::Serialize(Ar);
-		Ar << AOParameters;
-		Ar << TileHeadDataUnpacked;
-		Ar << TileArrayData;
-		Ar << VPLParameterBuffer;
-		Ar << VPLClusterData;
-		Ar << TileConeAxisAndCos;
-		Ar << TileConeDepthRanges;
-		Ar << NumGroups;
-		return bShaderHasOutdatedParameters;
-	}
-
-private:
-
-	FAOParameters AOParameters;
-	FRWShaderParameter TileHeadDataUnpacked;
-	FRWShaderParameter TileArrayData;
-	FShaderResourceParameter VPLParameterBuffer;
-	FShaderResourceParameter VPLClusterData;
-	FShaderResourceParameter TileConeAxisAndCos;
-	FShaderResourceParameter TileConeDepthRanges;
-	FShaderParameter NumGroups;
-};
-
-IMPLEMENT_SHADER_TYPE(,FVPLTileCullCS,TEXT("DistanceFieldGlobalIllumination"),TEXT("VPLTileCullCS"),SF_Compute);
-
-
-enum EComputeIrradianceMode
-{
-	IrradianceMode_IrradianceFromAOCones,
-	IrradianceMode_IrradianceFinalGather
-};
-
-template<EComputeIrradianceMode CombineMode>
-class TComputeIrradianceCS : public FGlobalShader
-{
-	DECLARE_SHADER_TYPE(TComputeIrradianceCS,Global)
-public:
-
-	static bool ShouldCache(EShaderPlatform Platform)
-	{
-		return IsFeatureLevelSupported(Platform, ERHIFeatureLevel::SM5) && DoesPlatformSupportDistanceFieldAO(Platform);
-	}
-
-	static void ModifyCompilationEnvironment(EShaderPlatform Platform, FShaderCompilerEnvironment& OutEnvironment)
-	{
-		FGlobalShader::ModifyCompilationEnvironment(Platform,OutEnvironment);
-		OutEnvironment.SetDefine(TEXT("NUM_CONE_DIRECTIONS"), NumConeSampleDirections);
-		OutEnvironment.SetDefine(TEXT("IRRADIANCE_FROM_AO_CONES"), CombineMode == IrradianceMode_IrradianceFromAOCones ? TEXT("1") : TEXT("0"));
-
-		// To reduce shader compile time of compute shaders with shared memory, doesn't have an impact on generated code with current compiler (June 2010 DX SDK)
-		OutEnvironment.CompilerFlags.Add(CFLAG_StandardOptimization);
-	}
-
-	TComputeIrradianceCS(const ShaderMetaType::CompiledShaderInitializerType& Initializer)
-		: FGlobalShader(Initializer)
-	{
-		AOParameters.Bind(Initializer.ParameterMap);
-		IrradianceCacheIrradiance.Bind(Initializer.ParameterMap, TEXT("IrradianceCacheIrradiance"));
-		IrradianceCachePositionRadius.Bind(Initializer.ParameterMap, TEXT("IrradianceCachePositionRadius"));
-		IrradianceCacheNormal.Bind(Initializer.ParameterMap, TEXT("IrradianceCacheNormal"));
-		IrradianceCacheTileCoordinate.Bind(Initializer.ParameterMap, TEXT("IrradianceCacheTileCoordinate"));
-		RecordConeData.Bind(Initializer.ParameterMap, TEXT("RecordConeData"));
-		ScatterDrawParameters.Bind(Initializer.ParameterMap, TEXT("ScatterDrawParameters"));
-		SavedStartIndex.Bind(Initializer.ParameterMap, TEXT("SavedStartIndex"));
-		DistanceFieldNormalTexture.Bind(Initializer.ParameterMap, TEXT("DistanceFieldNormalTexture"));
-		DistanceFieldNormalSampler.Bind(Initializer.ParameterMap, TEXT("DistanceFieldNormalSampler"));
-		GlobalObjectParameters.Bind(Initializer.ParameterMap);
-		VPLParameterBuffer.Bind(Initializer.ParameterMap, TEXT("VPLParameterBuffer"));
-		VPLClusterData.Bind(Initializer.ParameterMap, TEXT("VPLClusterData"));
-		VPLData.Bind(Initializer.ParameterMap, TEXT("VPLData"));
-		VPLGatherRadius.Bind(Initializer.ParameterMap, TEXT("VPLGatherRadius"));
-		ConeHalfAngle.Bind(Initializer.ParameterMap, TEXT("ConeHalfAngle"));
-		BentNormalNormalizeFactor.Bind(Initializer.ParameterMap, TEXT("BentNormalNormalizeFactor"));
-		TileListGroupSize.Bind(Initializer.ParameterMap, TEXT("TileListGroupSize"));
-		TileHeadDataUnpacked.Bind(Initializer.ParameterMap, TEXT("TileHeadDataUnpacked"));
-		TileArrayData.Bind(Initializer.ParameterMap, TEXT("TileArrayData"));
-	}
-
-	TComputeIrradianceCS()
-	{
-	}
-
-	void SetParameters(
-		FRHICommandList& RHICmdList, 
-		const FSceneView& View, 
-		int32 DepthLevel, 
-		const FDistanceFieldAOParameters& Parameters,
-		FTemporaryIrradianceCacheResources* TemporaryIrradianceCacheResources)
-	{
-		FComputeShaderRHIParamRef ShaderRHI = GetComputeShader();
-		FGlobalShader::SetParameters(RHICmdList, ShaderRHI, View);
-
-		AOParameters.Set(RHICmdList, ShaderRHI, Parameters);
-
-		const FScene* Scene = (const FScene*)View.Family->Scene;
-		FSurfaceCacheResources& SurfaceCacheResources = *Scene->SurfaceCacheResources;
-
-		SetSRVParameter(RHICmdList, ShaderRHI, IrradianceCachePositionRadius, SurfaceCacheResources.Level[DepthLevel]->PositionAndRadius.SRV);
-		SetSRVParameter(RHICmdList, ShaderRHI, IrradianceCacheNormal, SurfaceCacheResources.Level[DepthLevel]->Normal.SRV);
-		SetSRVParameter(RHICmdList, ShaderRHI, IrradianceCacheTileCoordinate, SurfaceCacheResources.Level[DepthLevel]->TileCoordinate.SRV);
-		SetSRVParameter(RHICmdList, ShaderRHI, ScatterDrawParameters, SurfaceCacheResources.Level[DepthLevel]->ScatterDrawParameters.SRV);
-		SetSRVParameter(RHICmdList, ShaderRHI, SavedStartIndex, SurfaceCacheResources.Level[DepthLevel]->SavedStartIndex.SRV);
-		SetSRVParameter(RHICmdList, ShaderRHI, RecordConeData, TemporaryIrradianceCacheResources->ConeData.SRV);
-
-		IrradianceCacheIrradiance.SetBuffer(RHICmdList, ShaderRHI, SurfaceCacheResources.Level[DepthLevel]->Irradiance);
-
-		FAOSampleData2 AOSampleData;
-
-		TArray<FVector, TInlineAllocator<9> > SampleDirections;
-		GetSpacedVectors(SampleDirections);
-
-		for (int32 SampleIndex = 0; SampleIndex < NumConeSampleDirections; SampleIndex++)
-		{
-			AOSampleData.SampleDirections[SampleIndex] = FVector4(SampleDirections[SampleIndex]);
-		}
-
-		SetUniformBufferParameterImmediate(RHICmdList, ShaderRHI, GetUniformBufferParameter<FAOSampleData2>(), AOSampleData);
-
-		FVector UnoccludedVector(0);
-
-		for (int32 SampleIndex = 0; SampleIndex < NumConeSampleDirections; SampleIndex++)
-		{
-			UnoccludedVector += SampleDirections[SampleIndex];
-		}
-
-		float BentNormalNormalizeFactorValue = 1.0f / (UnoccludedVector / NumConeSampleDirections).Size();
-		SetShaderValue(RHICmdList, ShaderRHI, BentNormalNormalizeFactor, BentNormalNormalizeFactorValue);
-
-		GlobalObjectParameters.Set(RHICmdList, ShaderRHI, *(Scene->DistanceFieldSceneData.ObjectBuffers), Scene->DistanceFieldSceneData.NumObjectsInBuffer);
-
-		FVPLResources& SourceVPLResources = GVPLViewCulling ? GCulledVPLResources : GVPLResources;
-		SetSRVParameter(RHICmdList, ShaderRHI, VPLParameterBuffer, SourceVPLResources.VPLParameterBuffer.SRV);
-		SetSRVParameter(RHICmdList, ShaderRHI, VPLClusterData, SourceVPLResources.VPLClusterData.SRV);
-		SetSRVParameter(RHICmdList, ShaderRHI, VPLData, SourceVPLResources.VPLData.SRV);
-
-		SetShaderValue(RHICmdList, ShaderRHI, VPLGatherRadius, Parameters.OcclusionMaxDistance);
-		extern float GAOConeHalfAngle;
-		SetShaderValue(RHICmdList, ShaderRHI, ConeHalfAngle, GAOConeHalfAngle);
-
-		FVPLTileIntersectionResources* VPLTileIntersectionResources = ((FSceneViewState*)View.State)->VPLTileIntersectionResources;
-
-		SetShaderValue(RHICmdList, ShaderRHI, TileListGroupSize, VPLTileIntersectionResources->TileDimensions);
-		SetSRVParameter(RHICmdList, ShaderRHI, TileHeadDataUnpacked, VPLTileIntersectionResources->TileHeadDataUnpacked.SRV);
-		SetSRVParameter(RHICmdList, ShaderRHI, TileArrayData, VPLTileIntersectionResources->TileArrayData.SRV);
-	}
-
-	void UnsetParameters(FRHICommandList& RHICmdList)
-	{
-		IrradianceCacheIrradiance.UnsetUAV(RHICmdList, GetComputeShader());
-	}
-
-	virtual bool Serialize(FArchive& Ar)
-	{		
-		bool bShaderHasOutdatedParameters = FGlobalShader::Serialize(Ar);
-		Ar << AOParameters;
-		Ar << IrradianceCacheIrradiance;
-		Ar << IrradianceCachePositionRadius;
-		Ar << IrradianceCacheNormal;
-		Ar << IrradianceCacheTileCoordinate;
-		Ar << RecordConeData;
-		Ar << ScatterDrawParameters;
-		Ar << SavedStartIndex;
-		Ar << DistanceFieldNormalTexture;
-		Ar << DistanceFieldNormalSampler;
-		Ar << GlobalObjectParameters;
-		Ar << VPLParameterBuffer;
-		Ar << VPLClusterData;
-		Ar << VPLData;
-		Ar << VPLGatherRadius;
-		Ar << ConeHalfAngle;
-		Ar << BentNormalNormalizeFactor;
-		Ar << TileListGroupSize;
-		Ar << TileHeadDataUnpacked;
-		Ar << TileArrayData;
-		return bShaderHasOutdatedParameters;
-	}
-
-private:
-
-	FAOParameters AOParameters;
-	FRWShaderParameter IrradianceCacheIrradiance;
-	FShaderResourceParameter IrradianceCachePositionRadius;
-	FShaderResourceParameter IrradianceCacheNormal;
-	FShaderResourceParameter IrradianceCacheTileCoordinate;
-	FShaderResourceParameter RecordConeData;
-	FShaderResourceParameter ScatterDrawParameters;
-	FShaderResourceParameter SavedStartIndex;
-	FShaderResourceParameter DistanceFieldNormalTexture;
-	FShaderResourceParameter DistanceFieldNormalSampler;
-	FDistanceFieldObjectBufferParameters GlobalObjectParameters;
-	FShaderResourceParameter VPLParameterBuffer;
-	FShaderResourceParameter VPLClusterData;
-	FShaderResourceParameter VPLData;
-	FShaderParameter VPLGatherRadius;
-	FShaderParameter ConeHalfAngle;
-	FShaderParameter BentNormalNormalizeFactor;
-	FShaderParameter TileListGroupSize;
-	FShaderResourceParameter TileHeadDataUnpacked;
-	FShaderResourceParameter TileArrayData;
-};
-
-IMPLEMENT_SHADER_TYPE(template<>,TComputeIrradianceCS<IrradianceMode_IrradianceFromAOCones>,TEXT("DistanceFieldGlobalIllumination"),TEXT("ComputeIrradianceCS"),SF_Compute);
-IMPLEMENT_SHADER_TYPE(template<>,TComputeIrradianceCS<IrradianceMode_IrradianceFinalGather>,TEXT("DistanceFieldGlobalIllumination"),TEXT("ComputeIrradianceCS"),SF_Compute);
-
-// Must match usf
-const int32 MaxVPLsPerTile = 1024;
-
-void FVPLTileIntersectionResources::InitDynamicRHI()
-{
-	TileHeadDataUnpacked.Initialize(sizeof(uint32), TileDimensions.X * TileDimensions.Y * 2, PF_R32_UINT, BUF_Static);
-
-	//@todo - handle max exceeded
-	TileArrayData.Initialize(sizeof(uint32), MaxVPLsPerTile * TileDimensions.X * TileDimensions.Y, PF_R32_UINT, BUF_Static);
-	TileArrayNextAllocation.Initialize(sizeof(uint32), 1, PF_R32_UINT, BUF_Static);
-}
-
 TScopedPointer<FLightTileIntersectionResources> GVPLPlacementTileIntersectionResources;
 
 void PlaceVPLs(
@@ -637,33 +356,6 @@
 	uint32 ClearValues[4] = { 0 };
 	RHICmdList.ClearUAV(GVPLResources.VPLParameterBuffer.UAV, ClearValues);
 
-	FVPLTileIntersectionResources*& VPLTileIntersectionResources = ((FSceneViewState*)View.State)->VPLTileIntersectionResources;
-
-	{
-		uint32 GroupSizeX = (View.ViewRect.Size().X / GAODownsampleFactor + GDistanceFieldAOTileSizeX - 1) / GDistanceFieldAOTileSizeX;
-		uint32 GroupSizeY = (View.ViewRect.Size().Y / GAODownsampleFactor + GDistanceFieldAOTileSizeY - 1) / GDistanceFieldAOTileSizeY;
-
-		if (!VPLTileIntersectionResources || VPLTileIntersectionResources->TileDimensions != FIntPoint(GroupSizeX, GroupSizeY))
-		{
-			if (VPLTileIntersectionResources)
-			{
-				VPLTileIntersectionResources->ReleaseResource();
-			}
-			else
-			{
-				VPLTileIntersectionResources = new FVPLTileIntersectionResources();
-			}
-
-			VPLTileIntersectionResources->TileDimensions = FIntPoint(GroupSizeX, GroupSizeY);
-
-			VPLTileIntersectionResources->InitResource();
-		}
-
-		// Indicates the clear value for each channel of the UAV format
-		uint32 ClearValues[4] = { 0 };
-		RHICmdList.ClearUAV(VPLTileIntersectionResources->TileHeadDataUnpacked.UAV, ClearValues);
-	}
-
 	const FLightSceneProxy* DirectionalLightProxy = NULL;
 
 	for (TSparseArray<FLightSceneInfoCompact>::TConstIterator LightIt(Scene->Lights); LightIt; ++LightIt)
@@ -699,6 +391,7 @@
 			{
 				ShadowBounds = DirectionalLightProxy->GetShadowSplitBoundsDepthRange(
 					View, 
+					View.ViewMatrices.ViewOrigin,
 					View.NearClippingDistance, 
 					GVPLPlacementCameraRadius, 
 					&CascadeSettings);
@@ -771,6 +464,7 @@
 				const float MaxSubjectZ = WorldToFace.TransformPosition(SubjectBounds.Center).Z + SubjectBounds.W;
 				const float MinSubjectZ = FMath::Max(MaxSubjectZ - SubjectBounds.W * 2, (float)-HALF_WORLD_MAX);
 
+				//@todo - naming is wrong and maybe derived matrices
 				DirectionalLightShadowToWorld = FTranslationMatrix(-ShadowBounds.Center) * WorldToFace * FShadowProjectionMatrix(MinSubjectZ, MaxSubjectZ, FVector4(0,0,0,1));
 
 				GetViewFrustumBounds(FrustumVolume, DirectionalLightShadowToWorld, true);
@@ -828,21 +522,592 @@
 				ComputeShader->UnsetParameters(RHICmdList);
 			}
 		}
-		/*
-		{
-			SCOPED_DRAW_EVENT(RHICmdList, TileCulling);
-
-			TShaderMapRef<FVPLTileCullCS> ComputeShader(View.ShaderMap);
-
-			RHICmdList.SetComputeShader(ComputeShader->GetComputeShader());
-			ComputeShader->SetParameters(RHICmdList, View, Scene, Parameters, GVPLViewCulling ? GCulledVPLResources : GVPLResources);
-			DispatchComputeShader(RHICmdList, *ComputeShader, VPLTileIntersectionResources->TileDimensions.X, VPLTileIntersectionResources->TileDimensions.Y, 1);
-
-			ComputeShader->UnsetParameters(RHICmdList);
-		}*/
 	}
 }
 
+const int32 LightVPLsThreadGroupSize = 64;
+
+class FSetupLightVPLsIndirectArgumentsCS : public FGlobalShader
+{
+	DECLARE_SHADER_TYPE(FSetupLightVPLsIndirectArgumentsCS,Global)
+public:
+
+	static bool ShouldCache(EShaderPlatform Platform)
+	{
+		return IsFeatureLevelSupported(Platform, ERHIFeatureLevel::SM5) && DoesPlatformSupportDistanceFieldAO(Platform);
+	}
+
+	static void ModifyCompilationEnvironment(EShaderPlatform Platform, FShaderCompilerEnvironment& OutEnvironment)
+	{
+		FGlobalShader::ModifyCompilationEnvironment(Platform,OutEnvironment);
+		OutEnvironment.SetDefine(TEXT("LIGHT_VPLS_THREADGROUP_SIZE"), LightVPLsThreadGroupSize);
+	}
+
+	FSetupLightVPLsIndirectArgumentsCS(const ShaderMetaType::CompiledShaderInitializerType& Initializer)
+		: FGlobalShader(Initializer)
+	{
+		DispatchParameters.Bind(Initializer.ParameterMap, TEXT("DispatchParameters"));
+		ObjectParameters.Bind(Initializer.ParameterMap);
+		ObjectProcessStride.Bind(Initializer.ParameterMap, TEXT("ObjectProcessStride"));
+	}
+
+	FSetupLightVPLsIndirectArgumentsCS()
+	{
+	}
+
+	void SetParameters(FRHICommandList& RHICmdList, const FSceneView& View)
+	{
+		FComputeShaderRHIParamRef ShaderRHI = GetComputeShader();
+		FGlobalShader::SetParameters(RHICmdList, ShaderRHI, View);
+
+		DispatchParameters.SetBuffer(RHICmdList, ShaderRHI, GAOCulledObjectBuffers.Buffers.ObjectIndirectDispatch);
+		ObjectParameters.Set(RHICmdList, ShaderRHI, GAOCulledObjectBuffers.Buffers);
+
+		SetShaderValue(RHICmdList, ShaderRHI, ObjectProcessStride, GVPLSpreadUpdateOver);
+	}
+
+	void UnsetParameters(FRHICommandList& RHICmdList)
+	{
+		DispatchParameters.UnsetUAV(RHICmdList, GetComputeShader());
+	}
+
+	virtual bool Serialize(FArchive& Ar)
+	{		
+		bool bShaderHasOutdatedParameters = FGlobalShader::Serialize(Ar);
+		Ar << DispatchParameters;
+		Ar << ObjectParameters;
+		Ar << ObjectProcessStride;
+		return bShaderHasOutdatedParameters;
+	}
+
+private:
+
+	FRWShaderParameter DispatchParameters;
+	FDistanceFieldCulledObjectBufferParameters ObjectParameters;
+	FShaderParameter ObjectProcessStride;
+};
+
+IMPLEMENT_SHADER_TYPE(,FSetupLightVPLsIndirectArgumentsCS,TEXT("DistanceFieldGlobalIllumination"),TEXT("SetupLightVPLsIndirectArgumentsCS"),SF_Compute);
+
+class FLightVPLsCS : public FGlobalShader
+{
+	DECLARE_SHADER_TYPE(FLightVPLsCS, Global);
+public:
+
+	static bool ShouldCache(EShaderPlatform Platform)
+	{
+		return IsFeatureLevelSupported(Platform, ERHIFeatureLevel::SM5) && DoesPlatformSupportDistanceFieldAO(Platform);
+	}
+
+	static void ModifyCompilationEnvironment(EShaderPlatform Platform, FShaderCompilerEnvironment& OutEnvironment)
+	{
+		OutEnvironment.SetDefine(TEXT("LIGHT_VPLS_THREADGROUP_SIZE"), LightVPLsThreadGroupSize);
+	}
+
+	/** Default constructor. */
+	FLightVPLsCS() {}
+
+	/** Initialization constructor. */
+	FLightVPLsCS(const ShaderMetaType::CompiledShaderInitializerType& Initializer)
+		: FGlobalShader(Initializer)
+	{
+		AOParameters.Bind(Initializer.ParameterMap);
+		LightDirection.Bind(Initializer.ParameterMap, TEXT("LightDirection"));
+		LightSourceRadius.Bind(Initializer.ParameterMap, TEXT("LightSourceRadius"));
+		LightPositionAndInvRadius.Bind(Initializer.ParameterMap, TEXT("LightPositionAndInvRadius"));
+		TanLightAngleAndNormalThreshold.Bind(Initializer.ParameterMap, TEXT("TanLightAngleAndNormalThreshold"));
+		LightColor.Bind(Initializer.ParameterMap, TEXT("LightColor"));
+		ObjectParameters.Bind(Initializer.ParameterMap);
+		SurfelParameters.Bind(Initializer.ParameterMap);
+		ShadowTileHeadDataUnpacked.Bind(Initializer.ParameterMap, TEXT("ShadowTileHeadDataUnpacked"));
+		ShadowTileArrayData.Bind(Initializer.ParameterMap, TEXT("ShadowTileArrayData"));
+		ShadowTileListGroupSize.Bind(Initializer.ParameterMap, TEXT("ShadowTileListGroupSize"));
+		WorldToShadow.Bind(Initializer.ParameterMap, TEXT("WorldToShadow"));
+		ShadowCulledObjectBounds.Bind(Initializer.ParameterMap, TEXT("ShadowCulledObjectBounds"));
+		ShadowCulledObjectData.Bind(Initializer.ParameterMap, TEXT("ShadowCulledObjectData"));
+		ObjectProcessStride.Bind(Initializer.ParameterMap, TEXT("ObjectProcessStride"));
+		ObjectProcessStartIndex.Bind(Initializer.ParameterMap, TEXT("ObjectProcessStartIndex"));
+	}
+
+	void SetParameters(
+		FRHICommandList& RHICmdList, 
+		const FSceneView& View, 
+		const FLightSceneProxy* LightSceneProxy,
+		const FMatrix& WorldToShadowMatrixValue,
+		const FDistanceFieldAOParameters& Parameters,
+		FLightTileIntersectionResources* TileIntersectionResources)
+	{
+		const FComputeShaderRHIParamRef ShaderRHI = GetComputeShader();
+
+		FGlobalShader::SetParameters(RHICmdList, ShaderRHI, View);
+		AOParameters.Set(RHICmdList, ShaderRHI, Parameters);
+
+		const FScene* Scene = (const FScene*)View.Family->Scene;
+
+		ObjectParameters.Set(RHICmdList, ShaderRHI, GAOCulledObjectBuffers.Buffers);
+		SurfelParameters.Set(RHICmdList, ShaderRHI, *Scene->DistanceFieldSceneData.SurfelBuffers, *Scene->DistanceFieldSceneData.InstancedSurfelBuffers);
+
+		FVector4 LightPositionAndInvRadiusValue;
+		FVector4 LightColorAndFalloffExponent;
+		FVector NormalizedLightDirection;
+		FVector2D SpotAngles;
+		float LightSourceRadiusValue;
+		float LightSourceLength;
+		float LightMinRoughness;
+
+		LightSceneProxy->GetParameters(LightPositionAndInvRadiusValue, LightColorAndFalloffExponent, NormalizedLightDirection, SpotAngles, LightSourceRadiusValue, LightSourceLength, LightMinRoughness);
+
+		SetShaderValue(RHICmdList, ShaderRHI, LightDirection, NormalizedLightDirection);
+		SetShaderValue(RHICmdList, ShaderRHI, LightPositionAndInvRadius, LightPositionAndInvRadiusValue);
+		// Default light source radius of 0 gives poor results
+		SetShaderValue(RHICmdList, ShaderRHI, LightSourceRadius, LightSourceRadiusValue == 0 ? 20 : FMath::Clamp(LightSourceRadiusValue, .001f, 1.0f / (4 * LightPositionAndInvRadiusValue.W)));
+
+		const float LightSourceAngle = FMath::Clamp(LightSceneProxy->GetLightSourceAngle(), 0.001f, 5.0f) * PI / 180.0f;
+		const FVector2D TanLightAngleAndNormalThresholdValue(FMath::Tan(LightSourceAngle), FMath::Cos(PI / 2 + LightSourceAngle));
+		SetShaderValue(RHICmdList, ShaderRHI, TanLightAngleAndNormalThreshold, TanLightAngleAndNormalThresholdValue);
+		SetShaderValue(RHICmdList, ShaderRHI, LightColor, LightSceneProxy->GetColor() * LightSceneProxy->GetIndirectLightingScale());
+
+		check(TileIntersectionResources || !ShadowTileHeadDataUnpacked.IsBound());
+
+		if (TileIntersectionResources)
+		{
+			SetSRVParameter(RHICmdList, ShaderRHI, ShadowTileHeadDataUnpacked, TileIntersectionResources->TileHeadDataUnpacked.SRV);
+			SetSRVParameter(RHICmdList, ShaderRHI, ShadowTileArrayData, TileIntersectionResources->TileArrayData.SRV);
+			SetShaderValue(RHICmdList, ShaderRHI, ShadowTileListGroupSize, TileIntersectionResources->TileDimensions);
+		}
+
+		SetShaderValue(RHICmdList, ShaderRHI, WorldToShadow, WorldToShadowMatrixValue);
+
+		SetSRVParameter(RHICmdList, ShaderRHI, ShadowCulledObjectBounds, GShadowCulledObjectBuffers.Buffers.Bounds.SRV);
+		SetSRVParameter(RHICmdList, ShaderRHI, ShadowCulledObjectData, GShadowCulledObjectBuffers.Buffers.Data.SRV);
+
+		SetShaderValue(RHICmdList, ShaderRHI, ObjectProcessStride, GVPLSpreadUpdateOver);
+		SetShaderValue(RHICmdList, ShaderRHI, ObjectProcessStartIndex, GFrameNumberRenderThread % GVPLSpreadUpdateOver);
+	}
+
+	void UnsetParameters(FRHICommandList& RHICmdList)
+	{
+		SurfelParameters.UnsetParameters(RHICmdList, GetComputeShader());
+	}
+
+	// FShader interface.
+	virtual bool Serialize(FArchive& Ar)
+	{
+		bool bShaderHasOutdatedParameters = FGlobalShader::Serialize(Ar);
+		Ar << AOParameters;
+		Ar << LightDirection;
+		Ar << LightPositionAndInvRadius;
+		Ar << LightSourceRadius;
+		Ar << TanLightAngleAndNormalThreshold;
+		Ar << LightColor;
+		Ar << ObjectParameters;
+		Ar << SurfelParameters;
+		Ar << ShadowTileHeadDataUnpacked;
+		Ar << ShadowTileArrayData;
+		Ar << ShadowTileListGroupSize;
+		Ar << WorldToShadow;
+		Ar << ShadowCulledObjectBounds;
+		Ar << ShadowCulledObjectData;
+		Ar << ObjectProcessStride;
+		Ar << ObjectProcessStartIndex;
+		return bShaderHasOutdatedParameters;
+	}
+
+private:
+
+	FAOParameters AOParameters;
+	FShaderParameter LightDirection;
+	FShaderParameter LightPositionAndInvRadius;
+	FShaderParameter LightSourceRadius;
+	FShaderParameter TanLightAngleAndNormalThreshold;
+	FShaderParameter LightColor;
+	FDistanceFieldCulledObjectBufferParameters ObjectParameters;
+	FSurfelBufferParameters SurfelParameters;
+	FShaderResourceParameter ShadowTileHeadDataUnpacked;
+	FShaderResourceParameter ShadowTileArrayData;
+	FShaderParameter ShadowTileListGroupSize;
+	FShaderParameter WorldToShadow;
+	FShaderResourceParameter ShadowCulledObjectBounds;
+	FShaderResourceParameter ShadowCulledObjectData;
+	FShaderParameter ObjectProcessStride;
+	FShaderParameter ObjectProcessStartIndex;
+};
+
+IMPLEMENT_SHADER_TYPE(,FLightVPLsCS,TEXT("DistanceFieldGlobalIllumination"),TEXT("LightVPLsCS"),SF_Compute);
+
+
+class FComputeStepBentNormalCS : public FGlobalShader
+{
+	DECLARE_SHADER_TYPE(FComputeStepBentNormalCS,Global)
+public:
+
+	static bool ShouldCache(EShaderPlatform Platform)
+	{
+		return IsFeatureLevelSupported(Platform, ERHIFeatureLevel::SM5) && DoesPlatformSupportDistanceFieldAO(Platform);
+	}
+
+	static void ModifyCompilationEnvironment(EShaderPlatform Platform, FShaderCompilerEnvironment& OutEnvironment)
+	{
+		FGlobalShader::ModifyCompilationEnvironment(Platform,OutEnvironment);
+		OutEnvironment.SetDefine(TEXT("NUM_CONE_DIRECTIONS"), NumConeSampleDirections);
+	}
+
+	FComputeStepBentNormalCS(const ShaderMetaType::CompiledShaderInitializerType& Initializer)
+		: FGlobalShader(Initializer)
+	{
+		StepBentNormal.Bind(Initializer.ParameterMap, TEXT("StepBentNormal"));
+		IrradianceCacheNormal.Bind(Initializer.ParameterMap, TEXT("IrradianceCacheNormal"));
+		RecordConeData.Bind(Initializer.ParameterMap, TEXT("RecordConeData"));
+		ScatterDrawParameters.Bind(Initializer.ParameterMap, TEXT("ScatterDrawParameters"));
+		SavedStartIndex.Bind(Initializer.ParameterMap, TEXT("SavedStartIndex"));
+		BentNormalNormalizeFactor.Bind(Initializer.ParameterMap, TEXT("BentNormalNormalizeFactor"));
+	}
+
+	FComputeStepBentNormalCS()
+	{
+	}
+
+	void SetParameters(
+		FRHICommandList& RHICmdList, 
+		const FSceneView& View, 
+		int32 DepthLevel, 
+		FTemporaryIrradianceCacheResources* TemporaryIrradianceCacheResources)
+	{
+		FComputeShaderRHIParamRef ShaderRHI = GetComputeShader();
+		FGlobalShader::SetParameters(RHICmdList, ShaderRHI, View);
+
+		const FScene* Scene = (const FScene*)View.Family->Scene;
+		FSurfaceCacheResources& SurfaceCacheResources = *Scene->SurfaceCacheResources;
+
+		SetSRVParameter(RHICmdList, ShaderRHI, IrradianceCacheNormal, SurfaceCacheResources.Level[DepthLevel]->Normal.SRV);
+		SetSRVParameter(RHICmdList, ShaderRHI, ScatterDrawParameters, SurfaceCacheResources.Level[DepthLevel]->ScatterDrawParameters.SRV);
+		SetSRVParameter(RHICmdList, ShaderRHI, SavedStartIndex, SurfaceCacheResources.Level[DepthLevel]->SavedStartIndex.SRV);
+		SetSRVParameter(RHICmdList, ShaderRHI, RecordConeData, TemporaryIrradianceCacheResources->ConeData.SRV);
+
+		StepBentNormal.SetBuffer(RHICmdList, ShaderRHI, TemporaryIrradianceCacheResources->StepBentNormal);
+
+		FAOSampleData2 AOSampleData;
+
+		TArray<FVector, TInlineAllocator<9> > SampleDirections;
+		GetSpacedVectors(SampleDirections);
+
+		for (int32 SampleIndex = 0; SampleIndex < NumConeSampleDirections; SampleIndex++)
+		{
+			AOSampleData.SampleDirections[SampleIndex] = FVector4(SampleDirections[SampleIndex]);
+		}
+
+		SetUniformBufferParameterImmediate(RHICmdList, ShaderRHI, GetUniformBufferParameter<FAOSampleData2>(), AOSampleData);
+
+		FVector UnoccludedVector(0);
+
+		for (int32 SampleIndex = 0; SampleIndex < NumConeSampleDirections; SampleIndex++)
+		{
+			UnoccludedVector += SampleDirections[SampleIndex];
+		}
+
+		float BentNormalNormalizeFactorValue = 1.0f / (UnoccludedVector / NumConeSampleDirections).Size();
+		SetShaderValue(RHICmdList, ShaderRHI, BentNormalNormalizeFactor, BentNormalNormalizeFactorValue);
+	}
+
+	void UnsetParameters(FRHICommandList& RHICmdList)
+	{
+		StepBentNormal.UnsetUAV(RHICmdList, GetComputeShader());
+	}
+
+	virtual bool Serialize(FArchive& Ar)
+	{		
+		bool bShaderHasOutdatedParameters = FGlobalShader::Serialize(Ar);
+		Ar << StepBentNormal;
+		Ar << IrradianceCacheNormal;
+		Ar << RecordConeData;
+		Ar << ScatterDrawParameters;
+		Ar << SavedStartIndex;
+		Ar << BentNormalNormalizeFactor;
+		return bShaderHasOutdatedParameters;
+	}
+
+private:
+
+	FRWShaderParameter StepBentNormal;
+	FShaderResourceParameter IrradianceCacheNormal;
+	FShaderResourceParameter RecordConeData;
+	FShaderResourceParameter ScatterDrawParameters;
+	FShaderResourceParameter SavedStartIndex;
+	FShaderParameter BentNormalNormalizeFactor;
+};
+
+IMPLEMENT_SHADER_TYPE(,FComputeStepBentNormalCS,TEXT("DistanceFieldGlobalIllumination"),TEXT("ComputeStepBentNormalCS"),SF_Compute);
+
+enum EVPLMode
+{
+	VPLMode_PlaceFromLight,
+	VPLMode_PlaceFromSurfels
+};
+
+template<EVPLMode VPLMode>
+class TComputeIrradianceCS : public FGlobalShader
+{
+	DECLARE_SHADER_TYPE(TComputeIrradianceCS,Global)
+public:
+
+	static bool ShouldCache(EShaderPlatform Platform)
+	{
+		return IsFeatureLevelSupported(Platform, ERHIFeatureLevel::SM5) && DoesPlatformSupportDistanceFieldAO(Platform);
+	}
+
+	static void ModifyCompilationEnvironment(EShaderPlatform Platform, FShaderCompilerEnvironment& OutEnvironment)
+	{
+		FGlobalShader::ModifyCompilationEnvironment(Platform,OutEnvironment);
+		OutEnvironment.SetDefine(TEXT("NUM_CONE_DIRECTIONS"), NumConeSampleDirections);
+		OutEnvironment.SetDefine(TEXT("IRRADIANCE_FROM_SURFELS"), VPLMode == VPLMode_PlaceFromSurfels ? TEXT("1") : TEXT("0"));
+
+		// To reduce shader compile time of compute shaders with shared memory, doesn't have an impact on generated code with current compiler (June 2010 DX SDK)
+		OutEnvironment.CompilerFlags.Add(CFLAG_StandardOptimization);
+	}
+
+	TComputeIrradianceCS(const ShaderMetaType::CompiledShaderInitializerType& Initializer)
+		: FGlobalShader(Initializer)
+	{
+		AOParameters.Bind(Initializer.ParameterMap);
+		IrradianceCacheIrradiance.Bind(Initializer.ParameterMap, TEXT("IrradianceCacheIrradiance"));
+		IrradianceCachePositionRadius.Bind(Initializer.ParameterMap, TEXT("IrradianceCachePositionRadius"));
+		IrradianceCacheNormal.Bind(Initializer.ParameterMap, TEXT("IrradianceCacheNormal"));
+		IrradianceCacheTileCoordinate.Bind(Initializer.ParameterMap, TEXT("IrradianceCacheTileCoordinate"));
+		ScatterDrawParameters.Bind(Initializer.ParameterMap, TEXT("ScatterDrawParameters"));
+		SavedStartIndex.Bind(Initializer.ParameterMap, TEXT("SavedStartIndex"));
+		DistanceFieldNormalTexture.Bind(Initializer.ParameterMap, TEXT("DistanceFieldNormalTexture"));
+		DistanceFieldNormalSampler.Bind(Initializer.ParameterMap, TEXT("DistanceFieldNormalSampler"));
+		GlobalObjectParameters.Bind(Initializer.ParameterMap);
+		ObjectParameters.Bind(Initializer.ParameterMap);
+		SurfelParameters.Bind(Initializer.ParameterMap);
+		VPLParameterBuffer.Bind(Initializer.ParameterMap, TEXT("VPLParameterBuffer"));
+		VPLData.Bind(Initializer.ParameterMap, TEXT("VPLData"));
+		VPLGatherRadius.Bind(Initializer.ParameterMap, TEXT("VPLGatherRadius"));
+		TileListGroupSize.Bind(Initializer.ParameterMap, TEXT("TileListGroupSize"));
+		TileHeadDataUnpacked.Bind(Initializer.ParameterMap, TEXT("TileHeadDataUnpacked"));
+		TileArrayData.Bind(Initializer.ParameterMap, TEXT("TileArrayData"));
+		StepBentNormalBuffer.Bind(Initializer.ParameterMap, TEXT("StepBentNormalBuffer"));
+	}
+
+	TComputeIrradianceCS()
+	{
+	}
+
+	void SetParameters(
+		FRHICommandList& RHICmdList, 
+		const FSceneView& View, 
+		int32 DepthLevel, 
+		const FDistanceFieldAOParameters& Parameters,
+		FTemporaryIrradianceCacheResources* TemporaryIrradianceCacheResources)
+	{
+		FComputeShaderRHIParamRef ShaderRHI = GetComputeShader();
+		FGlobalShader::SetParameters(RHICmdList, ShaderRHI, View);
+
+		AOParameters.Set(RHICmdList, ShaderRHI, Parameters);
+
+		const FScene* Scene = (const FScene*)View.Family->Scene;
+		FSurfaceCacheResources& SurfaceCacheResources = *Scene->SurfaceCacheResources;
+
+		SetSRVParameter(RHICmdList, ShaderRHI, IrradianceCachePositionRadius, SurfaceCacheResources.Level[DepthLevel]->PositionAndRadius.SRV);
+		SetSRVParameter(RHICmdList, ShaderRHI, IrradianceCacheNormal, SurfaceCacheResources.Level[DepthLevel]->Normal.SRV);
+		SetSRVParameter(RHICmdList, ShaderRHI, IrradianceCacheTileCoordinate, SurfaceCacheResources.Level[DepthLevel]->TileCoordinate.SRV);
+		SetSRVParameter(RHICmdList, ShaderRHI, ScatterDrawParameters, SurfaceCacheResources.Level[DepthLevel]->ScatterDrawParameters.SRV);
+		SetSRVParameter(RHICmdList, ShaderRHI, SavedStartIndex, SurfaceCacheResources.Level[DepthLevel]->SavedStartIndex.SRV);
+
+		IrradianceCacheIrradiance.SetBuffer(RHICmdList, ShaderRHI, SurfaceCacheResources.Level[DepthLevel]->Irradiance);
+
+		GlobalObjectParameters.Set(RHICmdList, ShaderRHI, *Scene->DistanceFieldSceneData.ObjectBuffers, Scene->DistanceFieldSceneData.NumObjectsInBuffer);
+		extern TGlobalResource<FDistanceFieldObjectBufferResource> GAOCulledObjectBuffers;
+		ObjectParameters.Set(RHICmdList, ShaderRHI, GAOCulledObjectBuffers.Buffers);
+		SurfelParameters.Set(RHICmdList, ShaderRHI, *Scene->DistanceFieldSceneData.SurfelBuffers, *Scene->DistanceFieldSceneData.InstancedSurfelBuffers);
+
+		FVPLResources& SourceVPLResources = GVPLViewCulling ? GCulledVPLResources : GVPLResources;
+		SetSRVParameter(RHICmdList, ShaderRHI, VPLParameterBuffer, SourceVPLResources.VPLParameterBuffer.SRV);
+		SetSRVParameter(RHICmdList, ShaderRHI, VPLData, SourceVPLResources.VPLData.SRV);
+
+		SetShaderValue(RHICmdList, ShaderRHI, VPLGatherRadius, Parameters.OcclusionMaxDistance);
+
+		FTileIntersectionResources* TileIntersectionResources = ((FSceneViewState*)View.State)->AOTileIntersectionResources;
+
+		SetSRVParameter(RHICmdList, ShaderRHI, TileHeadDataUnpacked, TileIntersectionResources->TileHeadDataUnpacked.SRV);
+		SetSRVParameter(RHICmdList, ShaderRHI, TileArrayData, TileIntersectionResources->TileArrayData.SRV);
+		SetShaderValue(RHICmdList, ShaderRHI, TileListGroupSize, TileIntersectionResources->TileDimensions);
+
+		SetSRVParameter(RHICmdList, ShaderRHI, StepBentNormalBuffer, TemporaryIrradianceCacheResources->StepBentNormal.SRV);
+	}
+
+	void UnsetParameters(FRHICommandList& RHICmdList)
+	{
+		IrradianceCacheIrradiance.UnsetUAV(RHICmdList, GetComputeShader());
+	}
+
+	virtual bool Serialize(FArchive& Ar)
+	{		
+		bool bShaderHasOutdatedParameters = FGlobalShader::Serialize(Ar);
+		Ar << AOParameters;
+		Ar << IrradianceCacheIrradiance;
+		Ar << IrradianceCachePositionRadius;
+		Ar << IrradianceCacheNormal;
+		Ar << IrradianceCacheTileCoordinate;
+		Ar << ScatterDrawParameters;
+		Ar << SavedStartIndex;
+		Ar << DistanceFieldNormalTexture;
+		Ar << DistanceFieldNormalSampler;
+		Ar << GlobalObjectParameters;
+		Ar << ObjectParameters;
+		Ar << SurfelParameters;
+		Ar << VPLParameterBuffer;
+		Ar << VPLData;
+		Ar << VPLGatherRadius;
+		Ar << TileListGroupSize;
+		Ar << TileHeadDataUnpacked;
+		Ar << TileArrayData;
+		Ar << StepBentNormalBuffer;
+		return bShaderHasOutdatedParameters;
+	}
+
+private:
+
+	FAOParameters AOParameters;
+	FRWShaderParameter IrradianceCacheIrradiance;
+	FShaderResourceParameter IrradianceCachePositionRadius;
+	FShaderResourceParameter IrradianceCacheNormal;
+	FShaderResourceParameter IrradianceCacheTileCoordinate;
+	FShaderResourceParameter ScatterDrawParameters;
+	FShaderResourceParameter SavedStartIndex;
+	FShaderResourceParameter DistanceFieldNormalTexture;
+	FShaderResourceParameter DistanceFieldNormalSampler;
+	FDistanceFieldObjectBufferParameters GlobalObjectParameters;
+	FDistanceFieldCulledObjectBufferParameters ObjectParameters;
+	FSurfelBufferParameters SurfelParameters;
+	FShaderResourceParameter VPLParameterBuffer;
+	FShaderResourceParameter VPLData;
+	FShaderParameter VPLGatherRadius;
+	FShaderParameter TileListGroupSize;
+	FShaderResourceParameter TileHeadDataUnpacked;
+	FShaderResourceParameter TileArrayData;
+	FShaderResourceParameter StepBentNormalBuffer;
+};
+
+IMPLEMENT_SHADER_TYPE(template<>,TComputeIrradianceCS<VPLMode_PlaceFromLight>,TEXT("DistanceFieldGlobalIllumination"),TEXT("ComputeIrradianceCS"),SF_Compute);
+IMPLEMENT_SHADER_TYPE(template<>,TComputeIrradianceCS<VPLMode_PlaceFromSurfels>,TEXT("DistanceFieldGlobalIllumination"),TEXT("ComputeIrradianceCS"),SF_Compute);
+
+void UpdateVPLs(
+	FRHICommandListImmediate& RHICmdList,
+	const FViewInfo& View,
+	const FScene* Scene,
+	const FDistanceFieldAOParameters& Parameters)
+{
+	if (GVPLSurfelRepresentation)
+	{
+		SCOPED_DRAW_EVENT(RHICmdList, UpdateVPLs);
+
+		const FLightSceneProxy* DirectionalLightProxy = NULL;
+
+		for (TSparseArray<FLightSceneInfoCompact>::TConstIterator LightIt(Scene->Lights); LightIt; ++LightIt)
+		{
+			const FLightSceneInfoCompact& LightSceneInfoCompact = *LightIt;
+			const FLightSceneInfo* const LightSceneInfo = LightSceneInfoCompact.LightSceneInfo;
+
+			if (LightSceneInfo->ShouldRenderLightViewIndependent() 
+				&& LightSceneInfo->Proxy->GetLightType() == LightType_Directional
+				&& LightSceneInfo->Proxy->CastsDynamicShadow())
+			{
+				DirectionalLightProxy = LightSceneInfo->Proxy;
+				break;
+			}
+		}
+
+		FMatrix DirectionalLightWorldToShadow = FMatrix::Identity;
+
+		if (DirectionalLightProxy)
+		{
+			{
+				int32 NumPlanes = 0;
+				const FPlane* PlaneData = NULL;
+				FVector4 ShadowBoundingSphereValue(0, 0, 0, 0);
+				FShadowCascadeSettings CascadeSettings;
+				FSphere ShadowBounds;
+				FConvexVolume FrustumVolume;
+
+				{
+					const float ConeExpandDistance = Parameters.OcclusionMaxDistance;
+					const float TanHalfFOV = 1.0f / View.ViewMatrices.ProjMatrix.M[0][0];
+					const float VertexPullbackLength = ConeExpandDistance / TanHalfFOV;
+
+					// Pull back cone vertex to contain VPLs outside of the view
+					const FVector ViewConeVertex = View.ViewMatrices.ViewOrigin - View.GetViewDirection() * VertexPullbackLength;
+
+					//@todo - expand by AOMaxDistance
+					ShadowBounds = DirectionalLightProxy->GetShadowSplitBoundsDepthRange(
+						View, 
+						ViewConeVertex,
+						View.NearClippingDistance, 
+						GetMaxAOViewDistance() + VertexPullbackLength + Parameters.OcclusionMaxDistance, 
+						&CascadeSettings); 
+
+					FSphere SubjectBounds(FVector::ZeroVector, ShadowBounds.W);
+
+					const FMatrix& WorldToLight = DirectionalLightProxy->GetWorldToLight();
+					const FMatrix InitializerWorldToLight = FInverseRotationMatrix(FVector(WorldToLight.M[0][0],WorldToLight.M[1][0],WorldToLight.M[2][0]).GetSafeNormal().Rotation());
+					const FVector InitializerFaceDirection = FVector(1,0,0);
+
+					FVector	XAxis, YAxis;
+					InitializerFaceDirection.FindBestAxisVectors(XAxis, YAxis);
+					const FMatrix WorldToLightScaled = InitializerWorldToLight * FScaleMatrix(FVector(1.0f,1.0f / SubjectBounds.W,1.0f / SubjectBounds.W));
+					const FMatrix WorldToFace = WorldToLightScaled * FBasisVectorMatrix(-XAxis,YAxis,InitializerFaceDirection.GetSafeNormal(),FVector::ZeroVector);
+
+					NumPlanes = CascadeSettings.ShadowBoundsAccurate.Planes.Num();
+					PlaneData = CascadeSettings.ShadowBoundsAccurate.Planes.GetData();
+
+					DirectionalLightWorldToShadow = FTranslationMatrix(-ShadowBounds.Center) 
+						* WorldToFace 
+						* FShadowProjectionMatrix(-GVPLDirectionalLightTraceDistance / 2, GVPLDirectionalLightTraceDistance / 2, FVector4(0,0,0,1));
+				}
+
+				CullDistanceFieldObjectsForLight(
+					RHICmdList,
+					View,
+					DirectionalLightProxy, 
+					DirectionalLightWorldToShadow,
+					NumPlanes,
+					PlaneData,
+					ShadowBoundingSphereValue,
+					ShadowBounds.W,
+					GVPLPlacementTileIntersectionResources);
+			}
+
+			SCOPED_DRAW_EVENT(RHICmdList, LightVPLs);
+
+			{
+				TShaderMapRef<FSetupLightVPLsIndirectArgumentsCS> ComputeShader(View.ShaderMap);
+				RHICmdList.SetComputeShader(ComputeShader->GetComputeShader());
+				ComputeShader->SetParameters(RHICmdList, View);
+
+				DispatchComputeShader(RHICmdList, *ComputeShader, 1, 1, 1);
+				ComputeShader->UnsetParameters(RHICmdList);
+			}
+
+			{
+				TShaderMapRef<FLightVPLsCS> ComputeShader(View.ShaderMap);
+				RHICmdList.SetComputeShader(ComputeShader->GetComputeShader());
+				ComputeShader->SetParameters(RHICmdList, View, DirectionalLightProxy, DirectionalLightWorldToShadow, Parameters, GVPLPlacementTileIntersectionResources);
+				DispatchIndirectComputeShader(RHICmdList, *ComputeShader, GAOCulledObjectBuffers.Buffers.ObjectIndirectDispatch.Buffer, 0);
+				ComputeShader->UnsetParameters(RHICmdList);
+			}
+		}
+		else
+		{
+			uint32 ClearValues[4] = { 0 };
+			RHICmdList.ClearUAV(Scene->DistanceFieldSceneData.InstancedSurfelBuffers->VPLFlux.UAV, ClearValues);
+		}
+	}
+	else
+	{
+		PlaceVPLs(RHICmdList, View, Scene, Parameters);
+	}
+}
+
 void ComputeIrradianceForSamples(
 	int32 DepthLevel,
 	FRHICommandListImmediate& RHICmdList,
@@ -851,13 +1116,41 @@
 	const FDistanceFieldAOParameters& Parameters,
 	FTemporaryIrradianceCacheResources* TemporaryIrradianceCacheResources)
 {
-	SCOPED_DRAW_EVENT(RHICmdList, ComputeIrradiance);
-
 	FSurfaceCacheResources& SurfaceCacheResources = *Scene->SurfaceCacheResources;
+	
+	{
+		SCOPED_DRAW_EVENT(RHICmdList, ComputeStepBentNormal);
 
-	if (GAOUseConesForGI)
+		{	
+			TShaderMapRef<TSetupFinalGatherIndirectArgumentsCS<false> > ComputeShader(View.ShaderMap);
+			RHICmdList.SetComputeShader(ComputeShader->GetComputeShader());
+			ComputeShader->SetParameters(RHICmdList, View, DepthLevel);
+			DispatchComputeShader(RHICmdList, *ComputeShader, 1, 1, 1);
+			ComputeShader->UnsetParameters(RHICmdList);
+		}
+
+		{
+			TShaderMapRef<FComputeStepBentNormalCS> ComputeShader(View.ShaderMap);
+			RHICmdList.SetComputeShader(ComputeShader->GetComputeShader());
+			ComputeShader->SetParameters(RHICmdList, View, DepthLevel, TemporaryIrradianceCacheResources);
+			DispatchIndirectComputeShader(RHICmdList, *ComputeShader, SurfaceCacheResources.DispatchParameters.Buffer, 0);
+			ComputeShader->UnsetParameters(RHICmdList);
+		}
+	}
+	
+	{	
+		TShaderMapRef<TSetupFinalGatherIndirectArgumentsCS<true> > ComputeShader(View.ShaderMap);
+		RHICmdList.SetComputeShader(ComputeShader->GetComputeShader());
+		ComputeShader->SetParameters(RHICmdList, View, DepthLevel);
+		DispatchComputeShader(RHICmdList, *ComputeShader, 1, 1, 1);
+		ComputeShader->UnsetParameters(RHICmdList);
+	}
+	
+	SCOPED_DRAW_EVENT(RHICmdList, ComputeIrradiance);
+	
+	if (GVPLSurfelRepresentation)
 	{
-		TShaderMapRef<TComputeIrradianceCS<IrradianceMode_IrradianceFromAOCones> > ComputeShader(View.ShaderMap);
+		TShaderMapRef<TComputeIrradianceCS<VPLMode_PlaceFromSurfels> > ComputeShader(View.ShaderMap);
 
 		RHICmdList.SetComputeShader(ComputeShader->GetComputeShader());
 		ComputeShader->SetParameters(RHICmdList, View, DepthLevel, Parameters, TemporaryIrradianceCacheResources);
@@ -867,7 +1160,7 @@
 	}
 	else
 	{
-		TShaderMapRef<TComputeIrradianceCS<IrradianceMode_IrradianceFinalGather> > ComputeShader(View.ShaderMap);
+		TShaderMapRef<TComputeIrradianceCS<VPLMode_PlaceFromLight> > ComputeShader(View.ShaderMap);
 
 		RHICmdList.SetComputeShader(ComputeShader->GetComputeShader());
 		ComputeShader->SetParameters(RHICmdList, View, DepthLevel, Parameters, TemporaryIrradianceCacheResources);
@@ -876,3 +1169,13 @@
 		ComputeShader->UnsetParameters(RHICmdList);
 	}
 }
+
+void ListDistanceFieldGIMemory(const FViewInfo& View)
+{
+	const FScene* Scene = (const FScene*)View.Family->Scene;
+
+	if (GVPLPlacementTileIntersectionResources)
+	{
+		UE_LOG(LogTemp, Log, TEXT("   Shadow tile culled objects %.3fMb"), GVPLPlacementTileIntersectionResources->GetSizeBytes() / 1024.0f / 1024.0f);
+	}
+}
diff --git a/Engine/Source/Runtime/Renderer/Private/DistanceFieldGlobalIllumination.h b/Engine/Source/Runtime/Renderer/Private/DistanceFieldGlobalIllumination.h
--- a/Engine/Source/Runtime/Renderer/Private/DistanceFieldGlobalIllumination.h
+++ b/Engine/Source/Runtime/Renderer/Private/DistanceFieldGlobalIllumination.h
@@ -43,28 +43,3 @@ public:
 	FRWBuffer VPLClusterData;
 	FRWBuffer VPLData;
 };
-
-/**  */
-class FVPLTileIntersectionResources : public FRenderResource
-{
-public:
-
-	FVPLTileIntersectionResources()
-		: TileDimensions(FIntPoint(0, 0))
-	{}
-
-	virtual void InitDynamicRHI() override;
-
-	virtual void ReleaseDynamicRHI() override
-	{
-		TileHeadDataUnpacked.Release();
-		TileArrayData.Release();
-		TileArrayNextAllocation.Release();
-	}
-
-	FIntPoint TileDimensions;
-
-	FRWBuffer TileHeadDataUnpacked;
-	FRWBuffer TileArrayData;
-	FRWBuffer TileArrayNextAllocation;
-};
\ No newline at end of file
diff --git a/Engine/Source/Runtime/Renderer/Private/DistanceFieldLightingShared.h b/Engine/Source/Runtime/Renderer/Private/DistanceFieldLightingShared.h
--- a/Engine/Source/Runtime/Renderer/Private/DistanceFieldLightingShared.h
+++ b/Engine/Source/Runtime/Renderer/Private/DistanceFieldLightingShared.h
@@ -8,6 +8,8 @@
 
 #include "DistanceFieldAtlas.h"
 
+DECLARE_LOG_CATEGORY_EXTERN(LogDistanceField, Warning, All);
+
 /** Tile sized used for most AO compute shaders. */
 const int32 GDistanceFieldAOTileSizeX = 16;
 const int32 GDistanceFieldAOTileSizeY = 16;
@@ -52,6 +54,72 @@ public:
 		Bounds.Release();
 		Data.Release();
 	}
+
+	size_t GetSizeBytes() const
+	{
+		return Bounds.NumBytes + Data.NumBytes;
+	}
+};
+
+class FSurfelBuffers
+{
+public:
+
+	// In float4's
+	static int32 SurfelDataStride;
+	static int32 InterpolatedVertexDataStride;
+
+	int32 MaxSurfels;
+
+	void Initialize()
+	{
+		if (MaxSurfels > 0)
+		{
+			InterpolatedVertexData.Initialize(sizeof(FVector4), MaxSurfels * InterpolatedVertexDataStride, PF_A32B32G32R32F, BUF_Static);
+			Surfels.Initialize(sizeof(FVector4), MaxSurfels * SurfelDataStride, PF_A32B32G32R32F, BUF_Static);
+		}
+	}
+
+	void Release()
+	{
+		InterpolatedVertexData.Release();
+		Surfels.Release();
+	}
+
+	size_t GetSizeBytes() const
+	{
+		return InterpolatedVertexData.NumBytes + Surfels.NumBytes;
+	}
+
+	FRWBuffer InterpolatedVertexData;
+	FRWBuffer Surfels;
+};
+
+class FInstancedSurfelBuffers
+{
+public:
+
+	int32 MaxSurfels;
+
+	void Initialize()
+	{
+		if (MaxSurfels > 0)
+		{
+			VPLFlux.Initialize(sizeof(FVector4), MaxSurfels, PF_A32B32G32R32F, BUF_Static);
+		}
+	}
+
+	void Release()
+	{
+		VPLFlux.Release();
+	}
+
+	size_t GetSizeBytes() const
+	{
+		return VPLFlux.NumBytes;
+	}
+
+	FRWBuffer VPLFlux;
 };
 
 class FDistanceFieldObjectBufferParameters
@@ -112,6 +180,44 @@ private:
 	FShaderParameter DistanceFieldAtlasTexelSize;
 };
 
+class FSurfelBufferParameters
+{
+public:
+	void Bind(const FShaderParameterMap& ParameterMap)
+	{
+		InterpolatedVertexData.Bind(ParameterMap, TEXT("InterpolatedVertexData"));
+		SurfelData.Bind(ParameterMap, TEXT("SurfelData"));
+		VPLFlux.Bind(ParameterMap, TEXT("VPLFlux"));
+	}
+
+	template<typename TParamRef>
+	void Set(FRHICommandList& RHICmdList, const TParamRef& ShaderRHI, const FSurfelBuffers& SurfelBuffers, const FInstancedSurfelBuffers& InstancedSurfelBuffers)
+	{
+		InterpolatedVertexData.SetBuffer(RHICmdList, ShaderRHI, SurfelBuffers.InterpolatedVertexData);
+		SurfelData.SetBuffer(RHICmdList, ShaderRHI, SurfelBuffers.Surfels);
+		VPLFlux.SetBuffer(RHICmdList, ShaderRHI, InstancedSurfelBuffers.VPLFlux);
+	}
+
+	template<typename TParamRef>
+	void UnsetParameters(FRHICommandList& RHICmdList, const TParamRef& ShaderRHI)
+	{
+		InterpolatedVertexData.UnsetUAV(RHICmdList, ShaderRHI);
+		SurfelData.UnsetUAV(RHICmdList, ShaderRHI);
+		VPLFlux.UnsetUAV(RHICmdList, ShaderRHI);
+	}
+
+	friend FArchive& operator<<(FArchive& Ar, FSurfelBufferParameters& P)
+	{
+		Ar << P.InterpolatedVertexData << P.SurfelData << P.VPLFlux;
+		return Ar;
+	}
+
+private:
+	FRWShaderParameter InterpolatedVertexData;
+	FRWShaderParameter SurfelData;
+	FRWShaderParameter VPLFlux;
+};
+
 class FDistanceFieldCulledObjectBuffers
 {
 public:
@@ -123,6 +229,7 @@ public:
 	int32 MaxObjects;
 
 	FRWBuffer ObjectIndirectArguments;
+	FRWBuffer ObjectIndirectDispatch;
 	FRWBuffer Bounds;
 	FRWBuffer Data;
 	FRWBuffer BoxBounds;
@@ -140,6 +247,7 @@ public:
 			const uint32 BufferFlags = BUF_ShaderResource;
 
 			ObjectIndirectArguments.Initialize(sizeof(uint32), 5, PF_R32_UINT, BUF_Static | BUF_DrawIndirect);
+			ObjectIndirectDispatch.Initialize(sizeof(uint32), 3, PF_R32_UINT, BUF_Static | BUF_DrawIndirect);
 			Bounds.Initialize(sizeof(FVector4), MaxObjects, PF_A32B32G32R32F);
 			Data.Initialize(sizeof(FVector4), MaxObjects * ObjectDataStride, PF_A32B32G32R32F);
 
@@ -153,10 +261,16 @@ public:
 	void Release()
 	{
 		ObjectIndirectArguments.Release();
+		ObjectIndirectDispatch.Release();
 		Bounds.Release();
 		Data.Release();
 		BoxBounds.Release();
 	}
+
+	size_t GetSizeBytes() const
+	{
+		return ObjectIndirectArguments.NumBytes + ObjectIndirectDispatch.NumBytes + Bounds.NumBytes + Data.NumBytes + BoxBounds.NumBytes;
+	}
 };
 
 class FDistanceFieldObjectBufferResource : public FRenderResource
@@ -276,6 +390,11 @@ public:
 		Buffer.SafeRelease();
 		BufferSRV.SafeRelease(); 
 	}
+
+	size_t GetSizeBytes() const
+	{
+		return MaxElements * Stride * GPixelFormats[Format].BlockBytes;
+	}
 };
 
 /**  */
@@ -296,6 +415,11 @@ public:
 		TileArrayData.Release();
 	}
 
+	size_t GetSizeBytes() const
+	{
+		return TileHeadDataUnpacked.NumBytes + TileArrayData.NumBytes;
+	}
+
 	FIntPoint TileDimensions;
 
 	FRWBuffer TileHeadDataUnpacked;
@@ -311,4 +435,57 @@ extern void CullDistanceFieldObjectsForLight(
 	const FPlane* PlaneData, 
 	const FVector4& ShadowBoundingSphereValue,
 	float ShadowBoundingRadius,
-	TScopedPointer<class FLightTileIntersectionResources>& TileIntersectionResources);
\ No newline at end of file
+	TScopedPointer<class FLightTileIntersectionResources>& TileIntersectionResources);
+
+class FUniformMeshBuffers
+{
+public:
+
+	int32 MaxElements;
+
+	FVertexBufferRHIRef TriangleData;
+	FShaderResourceViewRHIRef TriangleDataSRV;
+
+	FRWBuffer TriangleAreas;
+	FRWBuffer TriangleCDFs;
+
+	FUniformMeshBuffers()
+	{
+		MaxElements = 0;
+	}
+
+	void Initialize();
+
+	void Release()
+	{
+		TriangleData.SafeRelease();
+		TriangleDataSRV.SafeRelease(); 
+		TriangleAreas.Release();
+		TriangleCDFs.Release();
+	}
+};
+
+class FUniformMeshConverter
+{
+public:
+	static int32 Convert(
+		FRHICommandListImmediate& RHICmdList, 
+		FSceneRenderer& Renderer,
+		FViewInfo& View, 
+		const FPrimitiveSceneInfo* PrimitiveSceneInfo, 
+		FUniformMeshBuffers*& OutUniformMeshBuffers,
+		const FMaterialRenderProxy*& OutMaterialRenderProxy,
+		FUniformBufferRHIParamRef& OutPrimitiveUniformBuffer);
+
+	static void GenerateSurfels(
+		FRHICommandListImmediate& RHICmdList, 
+		FViewInfo& View, 
+		const FPrimitiveSceneInfo* PrimitiveSceneInfo, 
+		const FMaterialRenderProxy* MaterialProxy,
+		FUniformBufferRHIParamRef PrimitiveUniformBuffer,
+		const FMatrix& Instance0Transform,
+		int32 SurfelOffset,
+		int32 NumSurfels);
+};
+
+extern TGlobalResource<FDistanceFieldObjectBufferResource> GAOCulledObjectBuffers;
\ No newline at end of file
