Surfel Global Illumination with Distance Field shadowing

diff --git a/Engine/Source/Runtime/Renderer/Private/DistanceFieldObjectManagement.cpp b/Engine/Source/Runtime/Renderer/Private/DistanceFieldObjectManagement.cpp
new file mode 100644
--- /dev/null
+++ b/Engine/Source/Runtime/Renderer/Private/DistanceFieldObjectManagement.cpp
@@ -0,0 +1,1062 @@
+// Copyright 1998-2015 Epic Games, Inc. All Rights Reserved.
+
+/*=============================================================================
+	DistanceFieldSurfaceCacheLighting.cpp
+=============================================================================*/
+
+float GAOMaxObjectBoundingRadius = 50000;
+FAutoConsoleVariableRef CVarAOMaxObjectBoundingRadius(
+	TEXT("r.AOMaxObjectBoundingRadius"),
+	GAOMaxObjectBoundingRadius,
+	TEXT("Objects larger than this will not contribute to AO calculations, to improve performance."),
+	ECVF_Cheat | ECVF_RenderThreadSafe
+	);
+
+#include "RendererPrivate.h"
+#include "ScenePrivate.h"
+#include "UniformBuffer.h"
+#include "ShaderParameters.h"
+#include "DistanceFieldLightingShared.h"
+#include "DistanceFieldSurfaceCacheLighting.h"
+#include "DistanceFieldGlobalIllumination.h"
+#include "RHICommandList.h"
+#include "SceneUtils.h"
+
+// Must match equivalent shader defines
+int32 FDistanceFieldObjectBuffers::ObjectDataStride = 16;
+int32 FDistanceFieldCulledObjectBuffers::ObjectDataStride = 16;
+int32 FDistanceFieldCulledObjectBuffers::ObjectBoxBoundsStride = 5;
+
+// In float4's.  Must match corresponding usf definition
+int32 UploadObjectDataStride = 1 + FDistanceFieldObjectBuffers::ObjectDataStride;
+
+class FDistanceFieldUploadDataResource : public FRenderResource
+{
+public:
+
+	FCPUUpdatedBuffer UploadData;
+
+	FDistanceFieldUploadDataResource()
+	{
+		UploadData.Format = PF_A32B32G32R32F;
+		UploadData.Stride = UploadObjectDataStride;
+	}
+
+	virtual void InitDynamicRHI()  override
+	{
+		UploadData.Initialize();
+	}
+
+	virtual void ReleaseDynamicRHI() override
+	{
+		UploadData.Release();
+	}
+};
+
+TGlobalResource<FDistanceFieldUploadDataResource> GDistanceFieldUploadData;
+
+class FDistanceFieldUploadIndicesResource : public FRenderResource
+{
+public:
+
+	FCPUUpdatedBuffer UploadIndices;
+
+	FDistanceFieldUploadIndicesResource()
+	{
+		UploadIndices.Format = PF_R32_UINT;
+		UploadIndices.Stride = 1;
+	}
+
+	virtual void InitDynamicRHI()  override
+	{
+		UploadIndices.Initialize();
+	}
+
+	virtual void ReleaseDynamicRHI() override
+	{
+		UploadIndices.Release();
+	}
+};
+
+TGlobalResource<FDistanceFieldUploadIndicesResource> GDistanceFieldUploadIndices;
+
+class FDistanceFieldRemoveIndicesResource : public FRenderResource
+{
+public:
+
+	FCPUUpdatedBuffer RemoveIndices;
+
+	FDistanceFieldRemoveIndicesResource()
+	{
+		RemoveIndices.Format = PF_R32G32B32A32_UINT;
+		RemoveIndices.Stride = 1;
+	}
+
+	virtual void InitDynamicRHI()  override
+	{
+		RemoveIndices.Initialize();
+	}
+
+	virtual void ReleaseDynamicRHI() override
+	{
+		RemoveIndices.Release();
+	}
+};
+
+TGlobalResource<FDistanceFieldRemoveIndicesResource> GDistanceFieldRemoveIndices;
+
+uint32 UpdateObjectsGroupSize = 64;
+
+class FUploadObjectsToBufferCS : public FGlobalShader
+{
+	DECLARE_SHADER_TYPE(FUploadObjectsToBufferCS,Global)
+public:
+
+	static bool ShouldCache(EShaderPlatform Platform)
+	{
+		return IsFeatureLevelSupported(Platform, ERHIFeatureLevel::SM5) && DoesPlatformSupportDistanceFieldAO(Platform);
+	}
+
+	static void ModifyCompilationEnvironment(EShaderPlatform Platform, FShaderCompilerEnvironment& OutEnvironment)
+	{
+		FGlobalShader::ModifyCompilationEnvironment(Platform,OutEnvironment);
+		OutEnvironment.SetDefine(TEXT("UPDATEOBJECTS_THREADGROUP_SIZE"), UpdateObjectsGroupSize);
+	}
+
+	FUploadObjectsToBufferCS(const ShaderMetaType::CompiledShaderInitializerType& Initializer)
+		: FGlobalShader(Initializer)
+	{
+		NumUploadOperations.Bind(Initializer.ParameterMap, TEXT("NumUploadOperations"));
+		UploadOperationIndices.Bind(Initializer.ParameterMap, TEXT("UploadOperationIndices"));
+		UploadOperationData.Bind(Initializer.ParameterMap, TEXT("UploadOperationData"));
+		ObjectBufferParameters.Bind(Initializer.ParameterMap);
+	}
+
+	FUploadObjectsToBufferCS()
+	{
+	}
+
+	void SetParameters(FRHICommandList& RHICmdList, const FScene* Scene, uint32 NumUploadOperationsValue, FShaderResourceViewRHIParamRef InUploadOperationIndices, FShaderResourceViewRHIParamRef InUploadOperationData)
+	{
+		FComputeShaderRHIParamRef ShaderRHI = GetComputeShader();
+
+		SetShaderValue(RHICmdList, ShaderRHI, NumUploadOperations, NumUploadOperationsValue);
+
+		SetSRVParameter(RHICmdList, ShaderRHI, UploadOperationIndices, InUploadOperationIndices);
+		SetSRVParameter(RHICmdList, ShaderRHI, UploadOperationData, InUploadOperationData);
+
+		ObjectBufferParameters.Set(RHICmdList, ShaderRHI, *(Scene->DistanceFieldSceneData.ObjectBuffers), Scene->DistanceFieldSceneData.NumObjectsInBuffer);
+	}
+
+	void UnsetParameters(FRHICommandList& RHICmdList)
+	{
+		ObjectBufferParameters.UnsetParameters(RHICmdList, GetComputeShader());
+	}
+
+	virtual bool Serialize(FArchive& Ar)
+	{		
+		bool bShaderHasOutdatedParameters = FGlobalShader::Serialize(Ar);
+		Ar << NumUploadOperations;
+		Ar << UploadOperationIndices;
+		Ar << UploadOperationData;
+		Ar << ObjectBufferParameters;
+		return bShaderHasOutdatedParameters;
+	}
+
+private:
+
+	FShaderParameter NumUploadOperations;
+	FShaderResourceParameter UploadOperationIndices;
+	FShaderResourceParameter UploadOperationData;
+	FDistanceFieldObjectBufferParameters ObjectBufferParameters;
+};
+
+IMPLEMENT_SHADER_TYPE(,FUploadObjectsToBufferCS,TEXT("DistanceFieldSurfaceCacheLightingCompute"),TEXT("UploadObjectsToBufferCS"),SF_Compute);
+
+class FCopyObjectBufferCS : public FGlobalShader
+{
+	DECLARE_SHADER_TYPE(FCopyObjectBufferCS,Global)
+public:
+
+	static bool ShouldCache(EShaderPlatform Platform)
+	{
+		return IsFeatureLevelSupported(Platform, ERHIFeatureLevel::SM5) && DoesPlatformSupportDistanceFieldAO(Platform);
+	}
+
+	static void ModifyCompilationEnvironment(EShaderPlatform Platform, FShaderCompilerEnvironment& OutEnvironment)
+	{
+		FGlobalShader::ModifyCompilationEnvironment(Platform,OutEnvironment);
+		OutEnvironment.SetDefine(TEXT("UPDATEOBJECTS_THREADGROUP_SIZE"), UpdateObjectsGroupSize);
+	}
+
+	FCopyObjectBufferCS(const ShaderMetaType::CompiledShaderInitializerType& Initializer)
+		: FGlobalShader(Initializer)
+	{
+		CopyObjectBounds.Bind(Initializer.ParameterMap, TEXT("CopyObjectBounds"));
+		CopyObjectData.Bind(Initializer.ParameterMap, TEXT("CopyObjectData"));
+		ObjectBufferParameters.Bind(Initializer.ParameterMap);
+	}
+
+	FCopyObjectBufferCS()
+	{
+	}
+
+	void SetParameters(FRHICommandList& RHICmdList, FDistanceFieldObjectBuffers& ObjectBuffersSource, FDistanceFieldObjectBuffers& ObjectBuffersDest, int32 NumObjectsValue)
+	{
+		FComputeShaderRHIParamRef ShaderRHI = GetComputeShader();
+
+		CopyObjectBounds.SetBuffer(RHICmdList, ShaderRHI, ObjectBuffersDest.Bounds);
+		CopyObjectData.SetBuffer(RHICmdList, ShaderRHI, ObjectBuffersDest.Data);
+
+		ObjectBufferParameters.Set(RHICmdList, ShaderRHI, ObjectBuffersSource, NumObjectsValue);
+	}
+
+	void UnsetParameters(FRHICommandList& RHICmdList)
+	{
+		ObjectBufferParameters.UnsetParameters(RHICmdList, GetComputeShader());
+		CopyObjectBounds.UnsetUAV(RHICmdList, GetComputeShader());
+		CopyObjectData.UnsetUAV(RHICmdList, GetComputeShader());
+	}
+
+	virtual bool Serialize(FArchive& Ar)
+	{		
+		bool bShaderHasOutdatedParameters = FGlobalShader::Serialize(Ar);
+		Ar << CopyObjectBounds;
+		Ar << CopyObjectData;
+		Ar << ObjectBufferParameters;
+		return bShaderHasOutdatedParameters;
+	}
+
+private:
+
+	FRWShaderParameter CopyObjectBounds;
+	FRWShaderParameter CopyObjectData;
+	FDistanceFieldObjectBufferParameters ObjectBufferParameters;
+};
+
+IMPLEMENT_SHADER_TYPE(,FCopyObjectBufferCS,TEXT("DistanceFieldSurfaceCacheLightingCompute"),TEXT("CopyObjectBufferCS"),SF_Compute);
+
+class FCopySurfelBufferCS : public FGlobalShader
+{
+	DECLARE_SHADER_TYPE(FCopySurfelBufferCS,Global)
+public:
+
+	static bool ShouldCache(EShaderPlatform Platform)
+	{
+		return IsFeatureLevelSupported(Platform, ERHIFeatureLevel::SM5) && DoesPlatformSupportDistanceFieldAO(Platform);
+	}
+
+	static void ModifyCompilationEnvironment(EShaderPlatform Platform, FShaderCompilerEnvironment& OutEnvironment)
+	{
+		FGlobalShader::ModifyCompilationEnvironment(Platform,OutEnvironment);
+		OutEnvironment.SetDefine(TEXT("UPDATEOBJECTS_THREADGROUP_SIZE"), UpdateObjectsGroupSize);
+	}
+
+	FCopySurfelBufferCS(const ShaderMetaType::CompiledShaderInitializerType& Initializer)
+		: FGlobalShader(Initializer)
+	{
+		CopyInterpolatedVertexData.Bind(Initializer.ParameterMap, TEXT("CopyInterpolatedVertexData"));
+		CopySurfelData.Bind(Initializer.ParameterMap, TEXT("CopySurfelData"));
+		SurfelBufferParameters.Bind(Initializer.ParameterMap);
+		NumSurfels.Bind(Initializer.ParameterMap, TEXT("NumSurfels"));
+	}
+
+	FCopySurfelBufferCS()
+	{
+	}
+
+	void SetParameters(FRHICommandList& RHICmdList, const FSurfelBuffers& SurfelBuffersSource, const FInstancedSurfelBuffers& InstancedSurfelBuffersSource, FSurfelBuffers& SurfelBuffersDest, int32 NumSurfelsValue)
+	{
+		FComputeShaderRHIParamRef ShaderRHI = GetComputeShader();
+		CopyInterpolatedVertexData.SetBuffer(RHICmdList, ShaderRHI, SurfelBuffersDest.InterpolatedVertexData);
+		CopySurfelData.SetBuffer(RHICmdList, ShaderRHI, SurfelBuffersDest.Surfels);
+		SurfelBufferParameters.Set(RHICmdList, ShaderRHI, SurfelBuffersSource, InstancedSurfelBuffersSource);
+		SetShaderValue(RHICmdList, ShaderRHI, NumSurfels, NumSurfelsValue);
+	}
+
+	void UnsetParameters(FRHICommandList& RHICmdList)
+	{
+		SurfelBufferParameters.UnsetParameters(RHICmdList, GetComputeShader());
+		CopyInterpolatedVertexData.UnsetUAV(RHICmdList, GetComputeShader());
+		CopySurfelData.UnsetUAV(RHICmdList, GetComputeShader());
+	}
+
+	virtual bool Serialize(FArchive& Ar)
+	{		
+		bool bShaderHasOutdatedParameters = FGlobalShader::Serialize(Ar);
+		Ar << CopyInterpolatedVertexData;
+		Ar << CopySurfelData;
+		Ar << SurfelBufferParameters;
+		Ar << NumSurfels;
+		return bShaderHasOutdatedParameters;
+	}
+
+private:
+
+	FRWShaderParameter CopyInterpolatedVertexData;
+	FRWShaderParameter CopySurfelData;
+	FSurfelBufferParameters SurfelBufferParameters;
+	FShaderParameter NumSurfels;
+};
+
+IMPLEMENT_SHADER_TYPE(,FCopySurfelBufferCS,TEXT("SurfelTree"),TEXT("CopySurfelBufferCS"),SF_Compute);
+
+
+class FCopyVPLFluxBufferCS : public FGlobalShader
+{
+	DECLARE_SHADER_TYPE(FCopyVPLFluxBufferCS,Global)
+public:
+
+	static bool ShouldCache(EShaderPlatform Platform)
+	{
+		return IsFeatureLevelSupported(Platform, ERHIFeatureLevel::SM5) && DoesPlatformSupportDistanceFieldAO(Platform);
+	}
+
+	static void ModifyCompilationEnvironment(EShaderPlatform Platform, FShaderCompilerEnvironment& OutEnvironment)
+	{
+		FGlobalShader::ModifyCompilationEnvironment(Platform,OutEnvironment);
+		OutEnvironment.SetDefine(TEXT("UPDATEOBJECTS_THREADGROUP_SIZE"), UpdateObjectsGroupSize);
+	}
+
+	FCopyVPLFluxBufferCS(const ShaderMetaType::CompiledShaderInitializerType& Initializer)
+		: FGlobalShader(Initializer)
+	{
+		CopyVPLFlux.Bind(Initializer.ParameterMap, TEXT("CopyVPLFlux"));
+		SurfelBufferParameters.Bind(Initializer.ParameterMap);
+		NumSurfels.Bind(Initializer.ParameterMap, TEXT("NumSurfels"));
+	}
+
+	FCopyVPLFluxBufferCS()
+	{
+	}
+
+	void SetParameters(FRHICommandList& RHICmdList, const FSurfelBuffers& SurfelBuffersSource, const FInstancedSurfelBuffers& InstancedSurfelBuffersSource, FInstancedSurfelBuffers& InstancedSurfelBuffersDest, int32 NumSurfelsValue)
+	{
+		FComputeShaderRHIParamRef ShaderRHI = GetComputeShader();
+		CopyVPLFlux.SetBuffer(RHICmdList, ShaderRHI, InstancedSurfelBuffersDest.VPLFlux);
+		SurfelBufferParameters.Set(RHICmdList, ShaderRHI, SurfelBuffersSource, InstancedSurfelBuffersSource);
+		SetShaderValue(RHICmdList, ShaderRHI, NumSurfels, NumSurfelsValue);
+	}
+
+	void UnsetParameters(FRHICommandList& RHICmdList)
+	{
+		SurfelBufferParameters.UnsetParameters(RHICmdList, GetComputeShader());
+		CopyVPLFlux.UnsetUAV(RHICmdList, GetComputeShader());
+	}
+
+	virtual bool Serialize(FArchive& Ar)
+	{		
+		bool bShaderHasOutdatedParameters = FGlobalShader::Serialize(Ar);
+		Ar << CopyVPLFlux;
+		Ar << SurfelBufferParameters;
+		Ar << NumSurfels;
+		return bShaderHasOutdatedParameters;
+	}
+
+private:
+
+	FRWShaderParameter CopyVPLFlux;
+	FSurfelBufferParameters SurfelBufferParameters;
+	FShaderParameter NumSurfels;
+};
+
+IMPLEMENT_SHADER_TYPE(,FCopyVPLFluxBufferCS,TEXT("SurfelTree"),TEXT("CopyVPLFluxBufferCS"),SF_Compute);
+
+template<bool bRemoveFromSameBuffer>
+class TRemoveObjectsFromBufferCS : public FGlobalShader
+{
+	DECLARE_SHADER_TYPE(TRemoveObjectsFromBufferCS,Global)
+public:
+
+	static bool ShouldCache(EShaderPlatform Platform)
+	{
+		return IsFeatureLevelSupported(Platform, ERHIFeatureLevel::SM5) && DoesPlatformSupportDistanceFieldAO(Platform);
+	}
+
+	static void ModifyCompilationEnvironment(EShaderPlatform Platform, FShaderCompilerEnvironment& OutEnvironment)
+	{
+		FGlobalShader::ModifyCompilationEnvironment(Platform,OutEnvironment);
+		OutEnvironment.SetDefine(TEXT("UPDATEOBJECTS_THREADGROUP_SIZE"), UpdateObjectsGroupSize);
+		OutEnvironment.SetDefine(TEXT("REMOVE_FROM_SAME_BUFFER"), bRemoveFromSameBuffer ? TEXT("1") : TEXT("0"));
+	}
+
+	TRemoveObjectsFromBufferCS(const ShaderMetaType::CompiledShaderInitializerType& Initializer)
+		: FGlobalShader(Initializer)
+	{
+		NumRemoveOperations.Bind(Initializer.ParameterMap, TEXT("NumRemoveOperations"));
+		RemoveOperationIndices.Bind(Initializer.ParameterMap, TEXT("RemoveOperationIndices"));
+		ObjectBufferParameters.Bind(Initializer.ParameterMap);
+		ObjectBounds2.Bind(Initializer.ParameterMap, TEXT("ObjectBounds2"));
+		ObjectData2.Bind(Initializer.ParameterMap, TEXT("ObjectData2"));
+	}
+
+	TRemoveObjectsFromBufferCS()
+	{
+	}
+
+	void SetParameters(
+		FRHICommandList& RHICmdList, 
+		const FScene* Scene, 
+		uint32 NumRemoveOperationsValue, 
+		FShaderResourceViewRHIParamRef InRemoveOperationIndices, 
+		FShaderResourceViewRHIParamRef InObjectBounds2, 
+		FShaderResourceViewRHIParamRef InObjectData2)
+	{
+		FComputeShaderRHIParamRef ShaderRHI = GetComputeShader();
+
+		SetShaderValue(RHICmdList, ShaderRHI, NumRemoveOperations, NumRemoveOperationsValue);
+		SetSRVParameter(RHICmdList, ShaderRHI, RemoveOperationIndices, InRemoveOperationIndices);
+		ObjectBufferParameters.Set(RHICmdList, ShaderRHI, *(Scene->DistanceFieldSceneData.ObjectBuffers), Scene->DistanceFieldSceneData.NumObjectsInBuffer);
+		SetSRVParameter(RHICmdList, ShaderRHI, ObjectBounds2, InObjectBounds2);
+		SetSRVParameter(RHICmdList, ShaderRHI, ObjectData2, InObjectData2);
+	}
+
+	void UnsetParameters(FRHICommandList& RHICmdList)
+	{
+		ObjectBufferParameters.UnsetParameters(RHICmdList, GetComputeShader());
+	}
+
+	virtual bool Serialize(FArchive& Ar)
+	{		
+		bool bShaderHasOutdatedParameters = FGlobalShader::Serialize(Ar);
+		Ar << NumRemoveOperations;
+		Ar << RemoveOperationIndices;
+		Ar << ObjectBufferParameters;
+		Ar << ObjectBounds2;
+		Ar << ObjectData2;
+		return bShaderHasOutdatedParameters;
+	}
+
+private:
+
+	FShaderParameter NumRemoveOperations;
+	FShaderResourceParameter RemoveOperationIndices;
+	FDistanceFieldObjectBufferParameters ObjectBufferParameters;
+	FShaderResourceParameter ObjectBounds2;
+	FShaderResourceParameter ObjectData2;
+};
+
+IMPLEMENT_SHADER_TYPE(template<>,TRemoveObjectsFromBufferCS<true>,TEXT("DistanceFieldSurfaceCacheLightingCompute"),TEXT("RemoveObjectsFromBufferCS"),SF_Compute);
+IMPLEMENT_SHADER_TYPE(template<>,TRemoveObjectsFromBufferCS<false>,TEXT("DistanceFieldSurfaceCacheLightingCompute"),TEXT("RemoveObjectsFromBufferCS"),SF_Compute);
+
+void FSurfelBufferAllocator::RemovePrimitive(const FPrimitiveSceneInfo* Primitive)
+{
+	FPrimitiveSurfelAllocation Allocation;
+
+	if (Allocations.RemoveAndCopyValue(Primitive, Allocation))
+	{
+		bool bMergedWithExisting = false;
+
+		FPrimitiveSurfelFreeEntry FreeEntry(Allocation.Offset, Allocation.GetTotalNumSurfels());
+
+		// Note: only does one merge
+		//@todo - keep free list sorted then can binary search
+		for (int32 FreeIndex = 0; FreeIndex < FreeList.Num(); FreeIndex++)
+		{
+			if (FreeList[FreeIndex].Offset == FreeEntry.Offset + FreeEntry.NumSurfels)
+			{
+				FreeList[FreeIndex].Offset = FreeEntry.Offset;
+				FreeList[FreeIndex].NumSurfels += FreeEntry.NumSurfels;
+				bMergedWithExisting = true;
+				break;
+			}
+			else if (FreeList[FreeIndex].Offset + FreeList[FreeIndex].NumSurfels == FreeEntry.Offset)
+			{
+				FreeList[FreeIndex].NumSurfels += FreeEntry.NumSurfels;
+				bMergedWithExisting = true;
+				break;
+			}
+		}
+
+		if (!bMergedWithExisting)
+		{
+			FreeList.Add(FreeEntry);
+		}
+	}
+}
+
+void FSurfelBufferAllocator::AddPrimitive(const FPrimitiveSceneInfo* PrimitiveSceneInfo, int32 PrimitiveLOD0Surfels, int32 PrimitiveNumSurfels, int32 NumInstances)
+{
+	int32 BestFreeAllocationIndex = -1;
+
+	for (int32 FreeIndex = 0; FreeIndex < FreeList.Num(); FreeIndex++)
+	{
+		const FPrimitiveSurfelFreeEntry& CurrentFreeEntry = FreeList[FreeIndex];
+
+		if (CurrentFreeEntry.NumSurfels >= PrimitiveNumSurfels * NumInstances
+			&& (BestFreeAllocationIndex == -1 
+				|| CurrentFreeEntry.NumSurfels < FreeList[BestFreeAllocationIndex].NumSurfels))
+		{
+			BestFreeAllocationIndex = FreeIndex;
+		}
+	}
+
+	if (BestFreeAllocationIndex != -1)
+	{
+		FPrimitiveSurfelFreeEntry FreeEntry = FreeList[BestFreeAllocationIndex];
+						
+		if (FreeEntry.NumSurfels == PrimitiveNumSurfels * NumInstances)
+		{
+			// Existing allocation matches exactly, remove it from the free list
+			FreeList.RemoveAtSwap(BestFreeAllocationIndex);
+		}
+		else
+		{
+			// Replace with the remaining free range
+			FreeList[BestFreeAllocationIndex] = FPrimitiveSurfelFreeEntry(FreeEntry.Offset + PrimitiveNumSurfels * NumInstances, FreeEntry.NumSurfels - PrimitiveNumSurfels * NumInstances);
+		}
+
+		Allocations.Add(PrimitiveSceneInfo, FPrimitiveSurfelAllocation(FreeEntry.Offset, PrimitiveLOD0Surfels, PrimitiveNumSurfels, NumInstances));
+	}
+	else
+	{
+		// Add a new allocation to the end of the buffer
+		Allocations.Add(PrimitiveSceneInfo, FPrimitiveSurfelAllocation(NumSurfelsInBuffer, PrimitiveLOD0Surfels, PrimitiveNumSurfels, NumInstances));
+		NumSurfelsInBuffer += PrimitiveNumSurfels * NumInstances;
+	}
+}
+
+void UpdateGlobalDistanceFieldObjectRemoves(FRHICommandListImmediate& RHICmdList, FScene* Scene)
+{
+	FDistanceFieldSceneData& DistanceFieldSceneData = Scene->DistanceFieldSceneData;
+
+	TArray<FIntRect> RemoveObjectIndices;
+	FDistanceFieldObjectBuffers* TemporaryCopySourceBuffers = NULL;
+
+	if (DistanceFieldSceneData.PendingRemoveOperations.Num() > 0)
+	{
+		static TArray<int32, SceneRenderingAllocator> PendingRemoveOperations;
+
+		for (int32 RemoveIndex = 0; RemoveIndex < DistanceFieldSceneData.PendingRemoveOperations.Num(); RemoveIndex++)
+		{
+			// Can't dereference the primitive here, it has already been deleted
+			const FPrimitiveSceneInfo* Primitive = DistanceFieldSceneData.PendingRemoveOperations[RemoveIndex].Primitive;
+			DistanceFieldSceneData.SurfelAllocations.RemovePrimitive(Primitive);
+			DistanceFieldSceneData.InstancedSurfelAllocations.RemovePrimitive(Primitive);
+			PendingRemoveOperations.Append(DistanceFieldSceneData.PendingRemoveOperations[RemoveIndex].DistanceFieldInstanceIndices);
+		}
+
+		check(DistanceFieldSceneData.NumObjectsInBuffer >= PendingRemoveOperations.Num());
+
+		// Sort from smallest to largest
+		PendingRemoveOperations.Sort();
+
+		// We have multiple remove requests enqueued in PendingRemoveOperations, can only use the RemoveAtSwap version when there won't be collisions
+		const bool bUseRemoveAtSwap = PendingRemoveOperations.Last() < DistanceFieldSceneData.NumObjectsInBuffer - PendingRemoveOperations.Num();
+
+		if (bUseRemoveAtSwap)
+		{
+			// Remove everything in parallel in the same buffer with a RemoveAtSwap algorithm
+			for (int32 RemovePrimitiveIndex = 0; RemovePrimitiveIndex < PendingRemoveOperations.Num(); RemovePrimitiveIndex++)
+			{
+				DistanceFieldSceneData.NumObjectsInBuffer--;
+				const int32 RemoveIndex = PendingRemoveOperations[RemovePrimitiveIndex];
+				const int32 MoveFromIndex = DistanceFieldSceneData.NumObjectsInBuffer;
+
+				check(RemoveIndex != MoveFromIndex);
+				// Queue a compute shader move
+				RemoveObjectIndices.Add(FIntRect(RemoveIndex, MoveFromIndex, 0, 0));
+
+				// Fixup indices of the primitive that is being moved
+				FPrimitiveAndInstance& PrimitiveAndInstanceBeingMoved = DistanceFieldSceneData.PrimitiveInstanceMapping[MoveFromIndex];
+				check(PrimitiveAndInstanceBeingMoved.Primitive && PrimitiveAndInstanceBeingMoved.Primitive->DistanceFieldInstanceIndices.Num() > 0);
+				PrimitiveAndInstanceBeingMoved.Primitive->DistanceFieldInstanceIndices[PrimitiveAndInstanceBeingMoved.InstanceIndex] = RemoveIndex;
+
+				DistanceFieldSceneData.PrimitiveInstanceMapping.RemoveAtSwap(RemoveIndex);
+			}
+		}
+		else
+		{
+			// Have to copy the object data to allow parallel removing
+			TemporaryCopySourceBuffers = DistanceFieldSceneData.ObjectBuffers;
+			DistanceFieldSceneData.ObjectBuffers = new FDistanceFieldObjectBuffers();
+			DistanceFieldSceneData.ObjectBuffers->MaxObjects = TemporaryCopySourceBuffers->MaxObjects;
+			DistanceFieldSceneData.ObjectBuffers->Initialize();
+
+			TArray<FPrimitiveAndInstance> OriginalPrimitiveInstanceMapping = DistanceFieldSceneData.PrimitiveInstanceMapping;
+			DistanceFieldSceneData.PrimitiveInstanceMapping.Reset();
+
+			const int32 NumDestObjects = DistanceFieldSceneData.NumObjectsInBuffer - PendingRemoveOperations.Num();
+			int32 SourceIndex = 0;
+			int32 NextPendingRemoveIndex = 0;
+
+			for (int32 DestinationIndex = 0; DestinationIndex < NumDestObjects; DestinationIndex++)
+			{
+				while (NextPendingRemoveIndex < PendingRemoveOperations.Num()
+					&& PendingRemoveOperations[NextPendingRemoveIndex] == SourceIndex)
+				{
+					NextPendingRemoveIndex++;
+					SourceIndex++;
+				}
+
+				// Queue a compute shader move
+				RemoveObjectIndices.Add(FIntRect(DestinationIndex, SourceIndex, 0, 0));
+
+				// Fixup indices of the primitive that is being moved
+				FPrimitiveAndInstance& PrimitiveAndInstanceBeingMoved = OriginalPrimitiveInstanceMapping[SourceIndex];
+				check(PrimitiveAndInstanceBeingMoved.Primitive && PrimitiveAndInstanceBeingMoved.Primitive->DistanceFieldInstanceIndices.Num() > 0);
+				PrimitiveAndInstanceBeingMoved.Primitive->DistanceFieldInstanceIndices[PrimitiveAndInstanceBeingMoved.InstanceIndex] = DestinationIndex;
+
+				check(DistanceFieldSceneData.PrimitiveInstanceMapping.Num() == DestinationIndex);
+				DistanceFieldSceneData.PrimitiveInstanceMapping.Add(PrimitiveAndInstanceBeingMoved);
+
+				SourceIndex++;
+			}
+
+			DistanceFieldSceneData.NumObjectsInBuffer = NumDestObjects;
+
+
+			/*
+			// Have to remove one at a time while any entries to remove are at the end of the buffer
+			DistanceFieldSceneData.NumObjectsInBuffer--;
+			const int32 RemoveIndex = DistanceFieldSceneData.PendingRemoveOperations[ParallelConflictIndex];
+			const int32 MoveFromIndex = DistanceFieldSceneData.NumObjectsInBuffer;
+
+			if (RemoveIndex != MoveFromIndex)
+			{
+				// Queue a compute shader move
+				RemoveObjectIndices.Add(FIntRect(RemoveIndex, MoveFromIndex, 0, 0));
+
+				// Fixup indices of the primitive that is being moved
+				FPrimitiveAndInstance& PrimitiveAndInstanceBeingMoved = DistanceFieldSceneData.PrimitiveInstanceMapping[MoveFromIndex];
+				check(PrimitiveAndInstanceBeingMoved.Primitive && PrimitiveAndInstanceBeingMoved.Primitive->DistanceFieldInstanceIndices.Num() > 0);
+				PrimitiveAndInstanceBeingMoved.Primitive->DistanceFieldInstanceIndices[PrimitiveAndInstanceBeingMoved.InstanceIndex] = RemoveIndex;
+			}
+
+			DistanceFieldSceneData.PrimitiveInstanceMapping.RemoveAtSwap(RemoveIndex);
+			DistanceFieldSceneData.PendingRemoveOperations.RemoveAtSwap(ParallelConflictIndex);
+			*/
+		}
+
+		PendingRemoveOperations.Reset();
+		DistanceFieldSceneData.PendingRemoveOperations.Reset();
+
+		if (RemoveObjectIndices.Num() > 0)
+		{
+			if (RemoveObjectIndices.Num() > GDistanceFieldRemoveIndices.RemoveIndices.MaxElements)
+			{
+				GDistanceFieldRemoveIndices.RemoveIndices.MaxElements = RemoveObjectIndices.Num() * 5 / 4;
+				GDistanceFieldRemoveIndices.RemoveIndices.Release();
+				GDistanceFieldRemoveIndices.RemoveIndices.Initialize();
+			}
+
+			void* LockedBuffer = RHILockVertexBuffer(GDistanceFieldRemoveIndices.RemoveIndices.Buffer, 0, GDistanceFieldRemoveIndices.RemoveIndices.Buffer->GetSize(), RLM_WriteOnly);
+			const uint32 MemcpySize = RemoveObjectIndices.GetTypeSize() * RemoveObjectIndices.Num();
+			check(GDistanceFieldRemoveIndices.RemoveIndices.Buffer->GetSize() >= MemcpySize);
+			FPlatformMemory::Memcpy(LockedBuffer, RemoveObjectIndices.GetData(), MemcpySize);
+			RHIUnlockVertexBuffer(GDistanceFieldRemoveIndices.RemoveIndices.Buffer);
+
+			if (bUseRemoveAtSwap)
+			{
+				check(!TemporaryCopySourceBuffers);
+				TShaderMapRef<TRemoveObjectsFromBufferCS<true> > ComputeShader(GetGlobalShaderMap(Scene->GetFeatureLevel()));
+				RHICmdList.SetComputeShader(ComputeShader->GetComputeShader());
+				ComputeShader->SetParameters(RHICmdList, Scene, RemoveObjectIndices.Num(), GDistanceFieldRemoveIndices.RemoveIndices.BufferSRV, NULL, NULL);
+
+				DispatchComputeShader(RHICmdList, *ComputeShader, FMath::DivideAndRoundUp<uint32>(RemoveObjectIndices.Num(), UpdateObjectsGroupSize), 1, 1);
+				ComputeShader->UnsetParameters(RHICmdList);
+			}
+			else
+			{
+				check(TemporaryCopySourceBuffers);
+				TShaderMapRef<TRemoveObjectsFromBufferCS<false> > ComputeShader(GetGlobalShaderMap(Scene->GetFeatureLevel()));
+				RHICmdList.SetComputeShader(ComputeShader->GetComputeShader());
+				ComputeShader->SetParameters(RHICmdList, Scene, RemoveObjectIndices.Num(), GDistanceFieldRemoveIndices.RemoveIndices.BufferSRV, TemporaryCopySourceBuffers->Bounds.SRV, TemporaryCopySourceBuffers->Data.SRV);
+
+				DispatchComputeShader(RHICmdList, *ComputeShader, FMath::DivideAndRoundUp<uint32>(RemoveObjectIndices.Num(), UpdateObjectsGroupSize), 1, 1);
+				ComputeShader->UnsetParameters(RHICmdList);
+
+				TemporaryCopySourceBuffers->Release();
+				delete TemporaryCopySourceBuffers;
+			}
+		}
+	}
+}
+
+void FDeferredShadingSceneRenderer::UpdateGlobalDistanceFieldObjectBuffers(FRHICommandListImmediate& RHICmdList) 
+{
+	FDistanceFieldSceneData& DistanceFieldSceneData = Scene->DistanceFieldSceneData;
+
+	if (GDistanceFieldVolumeTextureAtlas.VolumeTextureRHI
+		&& (DistanceFieldSceneData.HasPendingOperations() || DistanceFieldSceneData.AtlasGeneration != GDistanceFieldVolumeTextureAtlas.GetGeneration()))
+	{
+		QUICK_SCOPE_CYCLE_COUNTER(STAT_UpdateObjectData);
+		SCOPED_DRAW_EVENT(RHICmdList, UpdateSceneObjectData);
+
+		if (!DistanceFieldSceneData.ObjectBuffers)
+		{
+			DistanceFieldSceneData.ObjectBuffers = new FDistanceFieldObjectBuffers();
+		}
+
+		if (!DistanceFieldSceneData.SurfelBuffers)
+		{
+			DistanceFieldSceneData.SurfelBuffers = new FSurfelBuffers();
+		}
+
+		if (!DistanceFieldSceneData.InstancedSurfelBuffers)
+		{
+			DistanceFieldSceneData.InstancedSurfelBuffers = new FInstancedSurfelBuffers();
+		}
+
+		if (DistanceFieldSceneData.AtlasGeneration != GDistanceFieldVolumeTextureAtlas.GetGeneration())
+		{
+			DistanceFieldSceneData.AtlasGeneration = GDistanceFieldVolumeTextureAtlas.GetGeneration();
+
+			for (int32 PrimitiveInstanceIndex = 0; PrimitiveInstanceIndex < DistanceFieldSceneData.PrimitiveInstanceMapping.Num(); PrimitiveInstanceIndex++)
+			{
+				FPrimitiveAndInstance& PrimitiveInstance = DistanceFieldSceneData.PrimitiveInstanceMapping[PrimitiveInstanceIndex];
+
+				// Queue an update of all primitives, since the atlas layout has changed
+				if (PrimitiveInstance.InstanceIndex == 0 
+					&& !DistanceFieldSceneData.HasPendingRemovePrimitive(PrimitiveInstance.Primitive)
+					&& !DistanceFieldSceneData.PendingAddOperations.Contains(PrimitiveInstance.Primitive)
+					&& !DistanceFieldSceneData.PendingUpdateOperations.Contains(PrimitiveInstance.Primitive))
+				{
+					DistanceFieldSceneData.PendingUpdateOperations.Add(PrimitiveInstance.Primitive);
+				}
+			}
+		}
+
+		// Process removes before adds, as the adds will overwrite primitive allocation info in DistanceFieldSceneData.SurfelAllocations
+		UpdateGlobalDistanceFieldObjectRemoves(RHICmdList, Scene);
+
+		TArray<uint32> UploadObjectIndices;
+		TArray<FVector4> UploadObjectData;
+
+		if (DistanceFieldSceneData.PendingAddOperations.Num() > 0 || DistanceFieldSceneData.PendingUpdateOperations.Num() > 0)
+		{
+			TArray<FMatrix> ObjectLocalToWorldTransforms;
+
+			const int32 NumUploadOperations = DistanceFieldSceneData.PendingAddOperations.Num() + DistanceFieldSceneData.PendingUpdateOperations.Num();
+			UploadObjectData.Empty(NumUploadOperations * UploadObjectDataStride);
+			UploadObjectIndices.Empty(NumUploadOperations);
+
+			const int32 NumTexelsOneDimX = GDistanceFieldVolumeTextureAtlas.GetSizeX();
+			const int32 NumTexelsOneDimY = GDistanceFieldVolumeTextureAtlas.GetSizeY();
+			const int32 NumTexelsOneDimZ = GDistanceFieldVolumeTextureAtlas.GetSizeZ();
+			const FVector InvTextureDim(1.0f / NumTexelsOneDimX, 1.0f / NumTexelsOneDimY, 1.0f / NumTexelsOneDimZ);
+
+			int32 OriginalNumObjects = DistanceFieldSceneData.NumObjectsInBuffer;
+			int32 OriginalNumSurfels = DistanceFieldSceneData.SurfelAllocations.GetNumSurfelsInBuffer();
+			int32 OriginalNumInstancedSurfels = DistanceFieldSceneData.InstancedSurfelAllocations.GetNumSurfelsInBuffer();
+
+			extern int32 GDistanceFieldGI;
+
+			if (GDistanceFieldGI)
+			{
+				for (int32 UploadPrimitiveIndex = 0; UploadPrimitiveIndex < DistanceFieldSceneData.PendingAddOperations.Num(); UploadPrimitiveIndex++)
+				{
+					FPrimitiveSceneInfo* PrimitiveSceneInfo = DistanceFieldSceneData.PendingAddOperations[UploadPrimitiveIndex];
+
+					int32 NumInstances = 0;
+					float BoundsSurfaceArea = 0;
+					PrimitiveSceneInfo->Proxy->GetDistanceFieldInstanceInfo(NumInstances, BoundsSurfaceArea);
+
+					extern void ComputeNumSurfels(float BoundsSurfaceArea, int32& PrimitiveNumSurfels, int32& PrimitiveLOD0Surfels);
+					int32 PrimitiveNumSurfels;
+					int32 PrimitiveLOD0Surfels;
+					ComputeNumSurfels(BoundsSurfaceArea, PrimitiveNumSurfels, PrimitiveLOD0Surfels);
+
+					if (PrimitiveNumSurfels > 0 && NumInstances > 0)
+					{
+						const int32 PrimitiveTotalNumSurfels = PrimitiveNumSurfels * NumInstances;
+
+						if (PrimitiveNumSurfels > 5000)
+						{
+							UE_LOG(LogDistanceField,Warning,TEXT("Primitive %s %s used %u Surfels"), *PrimitiveSceneInfo->Proxy->GetOwnerName().ToString(), *PrimitiveSceneInfo->Proxy->GetResourceName().ToString(), PrimitiveNumSurfels);
+						}
+
+						DistanceFieldSceneData.SurfelAllocations.AddPrimitive(PrimitiveSceneInfo, PrimitiveLOD0Surfels, PrimitiveNumSurfels, 1);
+						DistanceFieldSceneData.InstancedSurfelAllocations.AddPrimitive(PrimitiveSceneInfo, PrimitiveLOD0Surfels, PrimitiveNumSurfels, NumInstances);
+					}
+				}
+
+				if (DistanceFieldSceneData.SurfelBuffers->MaxSurfels < DistanceFieldSceneData.SurfelAllocations.GetNumSurfelsInBuffer())
+				{
+					if (DistanceFieldSceneData.SurfelBuffers->MaxSurfels > 0)
+					{
+						// Realloc
+						FSurfelBuffers* NewSurfelBuffers = new FSurfelBuffers();
+						NewSurfelBuffers->MaxSurfels = DistanceFieldSceneData.SurfelAllocations.GetNumSurfelsInBuffer() * 5 / 4;
+						NewSurfelBuffers->Initialize();
+
+						{
+							TShaderMapRef<FCopySurfelBufferCS> ComputeShader(GetGlobalShaderMap(Scene->GetFeatureLevel()));
+							RHICmdList.SetComputeShader(ComputeShader->GetComputeShader());
+							ComputeShader->SetParameters(RHICmdList, *(DistanceFieldSceneData.SurfelBuffers), *(DistanceFieldSceneData.InstancedSurfelBuffers), *NewSurfelBuffers, OriginalNumSurfels);
+
+							DispatchComputeShader(RHICmdList, *ComputeShader, FMath::DivideAndRoundUp<uint32>(OriginalNumSurfels, UpdateObjectsGroupSize), 1, 1);
+							ComputeShader->UnsetParameters(RHICmdList);
+						}
+
+						DistanceFieldSceneData.SurfelBuffers->Release();
+						delete DistanceFieldSceneData.SurfelBuffers;
+						DistanceFieldSceneData.SurfelBuffers = NewSurfelBuffers;
+					}
+					else
+					{
+						// First time allocate
+						DistanceFieldSceneData.SurfelBuffers->MaxSurfels = DistanceFieldSceneData.SurfelAllocations.GetNumSurfelsInBuffer() * 5 / 4;
+						DistanceFieldSceneData.SurfelBuffers->Initialize();
+					}
+				}
+				
+				if (DistanceFieldSceneData.InstancedSurfelBuffers->MaxSurfels < DistanceFieldSceneData.InstancedSurfelAllocations.GetNumSurfelsInBuffer())
+				{
+					if (DistanceFieldSceneData.InstancedSurfelBuffers->MaxSurfels > 0)
+					{
+						// Realloc
+						FInstancedSurfelBuffers* NewInstancedSurfelBuffers = new FInstancedSurfelBuffers();
+						NewInstancedSurfelBuffers->MaxSurfels = DistanceFieldSceneData.InstancedSurfelAllocations.GetNumSurfelsInBuffer() * 5 / 4;
+						NewInstancedSurfelBuffers->Initialize();
+						
+						{
+							TShaderMapRef<FCopyVPLFluxBufferCS> ComputeShader(GetGlobalShaderMap(Scene->GetFeatureLevel()));
+							RHICmdList.SetComputeShader(ComputeShader->GetComputeShader());
+							ComputeShader->SetParameters(RHICmdList, *(DistanceFieldSceneData.SurfelBuffers), *(DistanceFieldSceneData.InstancedSurfelBuffers), *NewInstancedSurfelBuffers, OriginalNumInstancedSurfels);
+
+							DispatchComputeShader(RHICmdList, *ComputeShader, FMath::DivideAndRoundUp<uint32>(OriginalNumInstancedSurfels, UpdateObjectsGroupSize), 1, 1);
+							ComputeShader->UnsetParameters(RHICmdList);
+						}
+
+						DistanceFieldSceneData.InstancedSurfelBuffers->Release();
+						delete DistanceFieldSceneData.InstancedSurfelBuffers;
+						DistanceFieldSceneData.InstancedSurfelBuffers = NewInstancedSurfelBuffers;
+					}
+					else
+					{
+						// First time allocate
+						DistanceFieldSceneData.InstancedSurfelBuffers->MaxSurfels = DistanceFieldSceneData.InstancedSurfelAllocations.GetNumSurfelsInBuffer() * 5 / 4;
+						DistanceFieldSceneData.InstancedSurfelBuffers->Initialize();
+					}
+				}
+			}
+
+			for (int32 UploadPrimitiveIndex = 0; UploadPrimitiveIndex < NumUploadOperations; UploadPrimitiveIndex++)
+			{
+				const bool bIsAddOperation = UploadPrimitiveIndex < DistanceFieldSceneData.PendingAddOperations.Num();
+
+				FPrimitiveSceneInfo* PrimitiveSceneInfo = bIsAddOperation 
+					? DistanceFieldSceneData.PendingAddOperations[UploadPrimitiveIndex] 
+					: DistanceFieldSceneData.PendingUpdateOperations[UploadPrimitiveIndex - DistanceFieldSceneData.PendingAddOperations.Num()];
+
+				ObjectLocalToWorldTransforms.Reset();
+
+				FBox LocalVolumeBounds;
+				FIntVector BlockMin;
+				FIntVector BlockSize;
+				bool bBuiltAsIfTwoSided;
+				bool bMeshWasPlane;
+				PrimitiveSceneInfo->Proxy->GetDistancefieldAtlasData(LocalVolumeBounds, BlockMin, BlockSize, bBuiltAsIfTwoSided, bMeshWasPlane, ObjectLocalToWorldTransforms);
+
+				if (BlockMin.X >= 0 
+					&& BlockMin.Y >= 0 
+					&& BlockMin.Z >= 0 
+					&& ObjectLocalToWorldTransforms.Num() > 0)
+				{
+					const float BoundingRadius = PrimitiveSceneInfo->Proxy->GetBounds().SphereRadius;
+
+					// Proxy bounds are only useful if single instance
+					if (ObjectLocalToWorldTransforms.Num() > 1 || BoundingRadius < GAOMaxObjectBoundingRadius)
+					{
+						FPrimitiveSurfelAllocation Allocation;
+						FPrimitiveSurfelAllocation InstancedAllocation;
+
+						if (GDistanceFieldGI)
+						{
+							const FPrimitiveSurfelAllocation* AllocationPtr = Scene->DistanceFieldSceneData.SurfelAllocations.FindAllocation(PrimitiveSceneInfo);
+							const FPrimitiveSurfelAllocation* InstancedAllocationPtr = Scene->DistanceFieldSceneData.InstancedSurfelAllocations.FindAllocation(PrimitiveSceneInfo);
+
+							if (AllocationPtr)
+							{
+								checkSlow(InstancedAllocationPtr && InstancedAllocationPtr->NumInstances == ObjectLocalToWorldTransforms.Num());
+								Allocation = *AllocationPtr;
+								InstancedAllocation = *InstancedAllocationPtr;
+
+								extern void GenerateSurfelRepresentation(FRHICommandListImmediate& RHICmdList, FSceneRenderer& Renderer, FViewInfo& View, FPrimitiveSceneInfo* PrimitiveSceneInfo, const FMatrix& Instance0Transform, FPrimitiveSurfelAllocation& Allocation);
+								GenerateSurfelRepresentation(RHICmdList, *this, Views[0], PrimitiveSceneInfo, ObjectLocalToWorldTransforms[0], Allocation);
+
+								if (Allocation.NumSurfels == 0)
+								{
+									InstancedAllocation.NumSurfels = 0;
+									InstancedAllocation.NumInstances = 0;
+									InstancedAllocation.NumLOD0 = 0;
+								}
+							}
+						}
+
+						if (bIsAddOperation)
+						{
+							DistanceFieldSceneData.NumObjectsInBuffer += ObjectLocalToWorldTransforms.Num();
+							PrimitiveSceneInfo->DistanceFieldInstanceIndices.Empty(ObjectLocalToWorldTransforms.Num());
+						}
+
+						for (int32 TransformIndex = 0; TransformIndex < ObjectLocalToWorldTransforms.Num(); TransformIndex++)
+						{
+							const uint32 UploadIndex = bIsAddOperation ? OriginalNumObjects + UploadObjectIndices.Num() : PrimitiveSceneInfo->DistanceFieldInstanceIndices[TransformIndex];
+
+							if (bIsAddOperation)
+							{
+								const int32 AddIndex = OriginalNumObjects + UploadObjectIndices.Num();
+								DistanceFieldSceneData.PrimitiveInstanceMapping.Add(FPrimitiveAndInstance(PrimitiveSceneInfo, TransformIndex));
+								PrimitiveSceneInfo->DistanceFieldInstanceIndices.Add(AddIndex);
+							}
+
+							UploadObjectIndices.Add(UploadIndex);
+
+							FMatrix LocalToWorld = ObjectLocalToWorldTransforms[TransformIndex];
+
+							if (bMeshWasPlane)
+							{
+								FVector LocalScales = LocalToWorld.GetScaleVector();
+								FVector AbsLocalScales(FMath::Abs(LocalScales.X), FMath::Abs(LocalScales.Y), FMath::Abs(LocalScales.Z));
+								float MidScale = FMath::Min(AbsLocalScales.X, AbsLocalScales.Y);
+								float ScaleAdjust = FMath::Sign(LocalScales.Z) * MidScale / AbsLocalScales.Z;
+								// The mesh was determined to be a plane flat in Z during the build process, so we can change the Z scale
+								// Helps in cases with modular ground pieces with scales of (10, 10, 1) and some triangles just above Z=0
+								LocalToWorld.SetAxis(2, LocalToWorld.GetScaledAxis(EAxis::Z) * ScaleAdjust);
+							}
+
+							const FMatrix VolumeToWorld = FScaleMatrix(LocalVolumeBounds.GetExtent()) 
+								* FTranslationMatrix(LocalVolumeBounds.GetCenter())
+								* LocalToWorld;
+
+							const FVector4 ObjectBoundingSphere(VolumeToWorld.GetOrigin(), VolumeToWorld.GetScaleVector().Size());
+
+							UploadObjectData.Add(ObjectBoundingSphere);
+
+							const float MaxExtent = LocalVolumeBounds.GetExtent().GetMax();
+
+							const FMatrix UniformScaleVolumeToWorld = FScaleMatrix(MaxExtent) 
+								* FTranslationMatrix(LocalVolumeBounds.GetCenter())
+								* LocalToWorld;
+
+							const FVector InvBlockSize(1.0f / BlockSize.X, 1.0f / BlockSize.Y, 1.0f / BlockSize.Z);
+
+							//float3 VolumeUV = (VolumePosition / LocalPositionExtent * .5f * UVScale + .5f * UVScale + UVAdd;
+							const FVector LocalPositionExtent = LocalVolumeBounds.GetExtent() / FVector(MaxExtent);
+							const FVector UVScale = FVector(BlockSize) * InvTextureDim;
+							const float VolumeScale = UniformScaleVolumeToWorld.GetMaximumAxisScale();
+
+							const FMatrix WorldToVolume = UniformScaleVolumeToWorld.Inverse();
+							// WorldToVolume
+							UploadObjectData.Add(*(FVector4*)&WorldToVolume.M[0]);
+							UploadObjectData.Add(*(FVector4*)&WorldToVolume.M[1]);
+							UploadObjectData.Add(*(FVector4*)&WorldToVolume.M[2]);
+							UploadObjectData.Add(*(FVector4*)&WorldToVolume.M[3]);
+
+							// Clamp to texel center by subtracting a half texel in the [-1,1] position space
+							// LocalPositionExtent
+							UploadObjectData.Add(FVector4(LocalPositionExtent - InvBlockSize, LocalVolumeBounds.Min.X));
+
+							// UVScale, VolumeScale and sign gives bGeneratedAsTwoSided
+							const float WSign = bBuiltAsIfTwoSided ? -1 : 1;
+							UploadObjectData.Add(FVector4(FVector(BlockSize) * InvTextureDim * .5f / LocalPositionExtent, WSign * VolumeScale));
+
+							// UVAdd
+							UploadObjectData.Add(FVector4(FVector(BlockMin) * InvTextureDim + .5f * UVScale, LocalVolumeBounds.Min.Y));
+
+							// Box bounds
+							UploadObjectData.Add(FVector4(LocalVolumeBounds.Max, LocalVolumeBounds.Min.Z));
+
+							UploadObjectData.Add(*(FVector4*)&UniformScaleVolumeToWorld.M[0]);
+							UploadObjectData.Add(*(FVector4*)&UniformScaleVolumeToWorld.M[1]);
+							UploadObjectData.Add(*(FVector4*)&UniformScaleVolumeToWorld.M[2]);
+
+							UploadObjectData.Add(*(FVector4*)&LocalToWorld.M[0]);
+							UploadObjectData.Add(*(FVector4*)&LocalToWorld.M[1]);
+							UploadObjectData.Add(*(FVector4*)&LocalToWorld.M[2]);
+							UploadObjectData.Add(*(FVector4*)&LocalToWorld.M[3]);
+
+							UploadObjectData.Add(FVector4(Allocation.Offset, Allocation.NumLOD0, Allocation.NumSurfels, InstancedAllocation.Offset + InstancedAllocation.NumSurfels * TransformIndex));
+
+							checkSlow(UploadObjectData.Num() % UploadObjectDataStride == 0);
+						}
+					}
+					else
+					{
+						UE_LOG(LogDistanceField,Log,TEXT("Primitive %s %s excluded due to bounding radius %f"), *PrimitiveSceneInfo->Proxy->GetOwnerName().ToString(), *PrimitiveSceneInfo->Proxy->GetResourceName().ToString(), BoundingRadius);
+					}
+				}
+			}
+
+			DistanceFieldSceneData.PendingAddOperations.Reset();
+			DistanceFieldSceneData.PendingUpdateOperations.Reset();
+
+			if (DistanceFieldSceneData.ObjectBuffers->MaxObjects < DistanceFieldSceneData.NumObjectsInBuffer)
+			{
+				if (DistanceFieldSceneData.ObjectBuffers->MaxObjects > 0)
+				{
+					// Realloc
+					FDistanceFieldObjectBuffers* NewObjectBuffers = new FDistanceFieldObjectBuffers();
+					NewObjectBuffers->MaxObjects = DistanceFieldSceneData.NumObjectsInBuffer * 5 / 4;
+					NewObjectBuffers->Initialize();
+
+					{
+						TShaderMapRef<FCopyObjectBufferCS> ComputeShader(GetGlobalShaderMap(Scene->GetFeatureLevel()));
+						RHICmdList.SetComputeShader(ComputeShader->GetComputeShader());
+						ComputeShader->SetParameters(RHICmdList, *(DistanceFieldSceneData.ObjectBuffers), *NewObjectBuffers, OriginalNumObjects);
+
+						DispatchComputeShader(RHICmdList, *ComputeShader, FMath::DivideAndRoundUp<uint32>(OriginalNumObjects, UpdateObjectsGroupSize), 1, 1);
+						ComputeShader->UnsetParameters(RHICmdList);
+					}
+
+					DistanceFieldSceneData.ObjectBuffers->Release();
+					delete DistanceFieldSceneData.ObjectBuffers;
+					DistanceFieldSceneData.ObjectBuffers = NewObjectBuffers;
+				}
+				else
+				{
+					// First time allocate
+					DistanceFieldSceneData.ObjectBuffers->MaxObjects = DistanceFieldSceneData.NumObjectsInBuffer * 5 / 4;
+					DistanceFieldSceneData.ObjectBuffers->Initialize();
+				}
+			}
+		}
+
+		if (UploadObjectIndices.Num() > 0)
+		{
+			if (UploadObjectIndices.Num() > GDistanceFieldUploadIndices.UploadIndices.MaxElements)
+			{
+				GDistanceFieldUploadIndices.UploadIndices.MaxElements = UploadObjectIndices.Num() * 5 / 4;
+				GDistanceFieldUploadIndices.UploadIndices.Release();
+				GDistanceFieldUploadIndices.UploadIndices.Initialize();
+
+				GDistanceFieldUploadData.UploadData.MaxElements = UploadObjectIndices.Num() * 5 / 4;
+				GDistanceFieldUploadData.UploadData.Release();
+				GDistanceFieldUploadData.UploadData.Initialize();
+			}
+
+			void* LockedBuffer = RHILockVertexBuffer(GDistanceFieldUploadIndices.UploadIndices.Buffer, 0, GDistanceFieldUploadIndices.UploadIndices.Buffer->GetSize(), RLM_WriteOnly);
+			const uint32 MemcpySize = UploadObjectIndices.GetTypeSize() * UploadObjectIndices.Num();
+			check(GDistanceFieldUploadIndices.UploadIndices.Buffer->GetSize() >= MemcpySize);
+			FPlatformMemory::Memcpy(LockedBuffer, UploadObjectIndices.GetData(), MemcpySize);
+			RHIUnlockVertexBuffer(GDistanceFieldUploadIndices.UploadIndices.Buffer);
+
+			LockedBuffer = RHILockVertexBuffer(GDistanceFieldUploadData.UploadData.Buffer, 0, GDistanceFieldUploadData.UploadData.Buffer->GetSize(), RLM_WriteOnly);
+			const uint32 MemcpySize2 = UploadObjectData.GetTypeSize() * UploadObjectData.Num();
+			check(GDistanceFieldUploadData.UploadData.Buffer->GetSize() >= MemcpySize2);
+			FPlatformMemory::Memcpy(LockedBuffer, UploadObjectData.GetData(), MemcpySize2);
+			RHIUnlockVertexBuffer(GDistanceFieldUploadData.UploadData.Buffer);
+
+			{
+				TShaderMapRef<FUploadObjectsToBufferCS> ComputeShader(GetGlobalShaderMap(Scene->GetFeatureLevel()));
+				RHICmdList.SetComputeShader(ComputeShader->GetComputeShader());
+				ComputeShader->SetParameters(RHICmdList, Scene, UploadObjectIndices.Num(), GDistanceFieldUploadIndices.UploadIndices.BufferSRV, GDistanceFieldUploadData.UploadData.BufferSRV);
+
+				DispatchComputeShader(RHICmdList, *ComputeShader, FMath::DivideAndRoundUp<uint32>(UploadObjectIndices.Num(), UpdateObjectsGroupSize), 1, 1);
+				ComputeShader->UnsetParameters(RHICmdList);
+			}
+		}
+
+		check(DistanceFieldSceneData.NumObjectsInBuffer == DistanceFieldSceneData.PrimitiveInstanceMapping.Num());
+
+		DistanceFieldSceneData.VerifyIntegrity();
+	}
+}
+
+FString GetObjectBufferMemoryString()
+{
+	return FString::Printf(TEXT("Temp object buffers %.3fMb"), 
+		(GDistanceFieldUploadIndices.UploadIndices.GetSizeBytes() + GDistanceFieldUploadData.UploadData.GetSizeBytes() + GDistanceFieldRemoveIndices.RemoveIndices.GetSizeBytes()) / 1024.0f / 1024.0f);
+}
\ No newline at end of file
diff --git a/Engine/Source/Runtime/Renderer/Private/DistanceFieldShadowing.cpp b/Engine/Source/Runtime/Renderer/Private/DistanceFieldShadowing.cpp
--- a/Engine/Source/Runtime/Renderer/Private/DistanceFieldShadowing.cpp
+++ b/Engine/Source/Runtime/Renderer/Private/DistanceFieldShadowing.cpp
@@ -89,13 +89,18 @@ public:
 		CulledObjectBoxBounds.SetBuffer(RHICmdList, ShaderRHI, GShadowCulledObjectBuffers.Buffers.BoxBounds);
 
 		SetShaderValue(RHICmdList, ShaderRHI, ObjectBoundingGeometryIndexCount, StencilingGeometry::GLowPolyStencilSphereIndexBuffer.GetIndexCount());
-
 		SetShaderValue(RHICmdList, ShaderRHI, WorldToShadow, WorldToShadowValue);
-
-		check(NumPlanes < 12);
-		SetShaderValue(RHICmdList, ShaderRHI, NumShadowHullPlanes, NumPlanes);
 		SetShaderValue(RHICmdList, ShaderRHI, ShadowBoundingSphere, ShadowBoundingSphereValue);
-		SetShaderValueArray(RHICmdList, ShaderRHI, ShadowConvexHull, PlaneData, NumPlanes);
+
+		if (NumPlanes < 12)
+		{
+			SetShaderValue(RHICmdList, ShaderRHI, NumShadowHullPlanes, NumPlanes);
+			SetShaderValueArray(RHICmdList, ShaderRHI, ShadowConvexHull, PlaneData, NumPlanes);
+		}
+		else
+		{
+			SetShaderValue(RHICmdList, ShaderRHI, NumShadowHullPlanes, 0);
+		}
 	}
 
 	void UnsetParameters(FRHICommandList& RHICmdList)
@@ -580,9 +585,9 @@ void CullDistanceFieldObjectsForLight(
 {
 	const FScene* Scene = (const FScene*)(View.Family->Scene);
 
-	{
-		SCOPED_DRAW_EVENT(RHICmdList, ObjectFrustumCulling);
+	SCOPED_DRAW_EVENT(RHICmdList, CullObjectsForLight);
 
+	{
 		if (GShadowCulledObjectBuffers.Buffers.MaxObjects < Scene->DistanceFieldSceneData.NumObjectsInBuffer
 			|| GShadowCulledObjectBuffers.Buffers.MaxObjects > 3 * Scene->DistanceFieldSceneData.NumObjectsInBuffer)
 		{
@@ -629,7 +634,6 @@ void CullDistanceFieldObjectsForLight(
 		}
 
 		{
-			SCOPED_DRAW_EVENT(RHICmdList, ClearTiles);
 			TShaderMapRef<FClearTilesCS> ComputeShader(View.ShaderMap);
 
 			uint32 GroupSizeX = FMath::DivideAndRoundUp(LightTileDimensions.X, GDistanceFieldAOTileSizeX);
@@ -643,8 +647,6 @@ void CullDistanceFieldObjectsForLight(
 		}
 		
 		{
-			SCOPED_DRAW_EVENT(RHICmdList, CullObjects);
-
 			TShaderMapRef<FShadowObjectCullVS> VertexShader(View.ShaderMap);
 			TShaderMapRef<FShadowObjectCullPS> PixelShader(View.ShaderMap);
 
diff --git a/Engine/Source/Runtime/Renderer/Private/DistanceFieldSurfaceCacheLighting.cpp b/Engine/Source/Runtime/Renderer/Private/DistanceFieldSurfaceCacheLighting.cpp
--- a/Engine/Source/Runtime/Renderer/Private/DistanceFieldSurfaceCacheLighting.cpp
+++ b/Engine/Source/Runtime/Renderer/Private/DistanceFieldSurfaceCacheLighting.cpp
@@ -51,7 +51,7 @@ FAutoConsoleVariableRef CVarAOMinLevel(
 	ECVF_Cheat | ECVF_RenderThreadSafe
 	);
 
-int32 GAOMaxLevel = FMath::Min(3, GAOMaxSupportedLevel);
+int32 GAOMaxLevel = FMath::Min(2, GAOMaxSupportedLevel);
 FAutoConsoleVariableRef CVarAOMaxLevel(
 	TEXT("r.AOMaxLevel"),
 	GAOMaxLevel,
@@ -76,11 +76,6 @@ FAutoConsoleVariableRef CVarAOTrimOldRecordsFraction(
 	ECVF_Cheat | ECVF_RenderThreadSafe
 	);
 
-float GetTrimOldRecordsFraction(int32 DepthLevel)
-{
-	return GAOTrimOldRecordsFraction;
-}
-
 int32 GAOFillGaps = 1;
 FAutoConsoleVariableRef CVarAOFillGaps(
 	TEXT("r.AOFillGaps"),
@@ -163,6 +158,14 @@ FAutoConsoleVariableRef CVarAOMaxViewDistance(
 	ECVF_Cheat | ECVF_RenderThreadSafe
 	);
 
+float GAOViewFadeDistanceScale = .2f;
+FAutoConsoleVariableRef CVarAOViewFadeDistanceScale(
+	TEXT("r.AOViewFadeDistanceScale"),
+	GAOViewFadeDistanceScale,
+	TEXT("Distance over which AO will fade out as it approaches r.AOMaxViewDistance, as a fraction of r.AOMaxViewDistance."),
+	ECVF_Cheat | ECVF_RenderThreadSafe
+	);
+
 int32 GAOScatterTileCulling = 1;
 FAutoConsoleVariableRef CVarAOScatterTileCulling(
 	TEXT("r.AOScatterTileCulling"),
@@ -219,6 +222,8 @@ FAutoConsoleVariableRef CVarAOHeightfieldOcclusion(
 	ECVF_Cheat | ECVF_RenderThreadSafe
 	);
 
+DEFINE_LOG_CATEGORY(LogDistanceField);
+
 IMPLEMENT_UNIFORM_BUFFER_STRUCT(FAOSampleData2,TEXT("AOSamples2"));
 
 FIntPoint GetBufferSizeForAO()
@@ -328,92 +333,8 @@ public:
 };
 
 TGlobalResource<FCircleVertexBuffer> GCircleVertexBuffer;
-
-// Only 7 needed atm, pad to cache line size
-// Must match equivalent shader defines
-int32 FDistanceFieldObjectBuffers::ObjectDataStride = 15;
-int32 FDistanceFieldCulledObjectBuffers::ObjectDataStride = 11;
-int32 FDistanceFieldCulledObjectBuffers::ObjectBoxBoundsStride = 5;
-
 TGlobalResource<FDistanceFieldObjectBufferResource> GAOCulledObjectBuffers;
-
-// In float4's.  Must match corresponding usf definition
-int32 UploadObjectDataStride = 1 + FDistanceFieldObjectBuffers::ObjectDataStride;
-
-class FDistanceFieldUploadDataResource : public FRenderResource
-{
-public:
-
-	FCPUUpdatedBuffer UploadData;
-
-	FDistanceFieldUploadDataResource()
-	{
-		UploadData.Format = PF_A32B32G32R32F;
-		UploadData.Stride = UploadObjectDataStride;
-	}
-
-	virtual void InitDynamicRHI()  override
-	{
-		UploadData.Initialize();
-	}
-
-	virtual void ReleaseDynamicRHI() override
-	{
-		UploadData.Release();
-	}
-};
-
-TGlobalResource<FDistanceFieldUploadDataResource> GDistanceFieldUploadData;
-
-class FDistanceFieldUploadIndicesResource : public FRenderResource
-{
-public:
-
-	FCPUUpdatedBuffer UploadIndices;
-
-	FDistanceFieldUploadIndicesResource()
-	{
-		UploadIndices.Format = PF_R32_UINT;
-		UploadIndices.Stride = 1;
-	}
-
-	virtual void InitDynamicRHI()  override
-	{
-		UploadIndices.Initialize();
-	}
-
-	virtual void ReleaseDynamicRHI() override
-	{
-		UploadIndices.Release();
-	}
-};
-
-TGlobalResource<FDistanceFieldUploadIndicesResource> GDistanceFieldUploadIndices;
-
-class FDistanceFieldRemoveIndicesResource : public FRenderResource
-{
-public:
-
-	FCPUUpdatedBuffer RemoveIndices;
-
-	FDistanceFieldRemoveIndicesResource()
-	{
-		RemoveIndices.Format = PF_R32G32B32A32_UINT;
-		RemoveIndices.Stride = 1;
-	}
-
-	virtual void InitDynamicRHI()  override
-	{
-		RemoveIndices.Initialize();
-	}
-
-	virtual void ReleaseDynamicRHI() override
-	{
-		RemoveIndices.Release();
-	}
-};
-
-TGlobalResource<FDistanceFieldRemoveIndicesResource> GDistanceFieldRemoveIndices;
+TGlobalResource<FTemporaryIrradianceCacheResources> GTemporaryIrradianceCacheResources;
 
 class FHeightfieldDescriptionsResource : public FRenderResource
 {
@@ -441,8 +362,6 @@ public:
 
 TGlobalResource<FHeightfieldDescriptionsResource> GHeightfieldDescriptions;
 
-TGlobalResource<FTemporaryIrradianceCacheResources> GTemporaryIrradianceCacheResources;
-
 void FTileIntersectionResources::InitDynamicRHI()
 {
 	TileConeAxisAndCos.Initialize(sizeof(float)* 4, TileDimensions.X * TileDimensions.Y, PF_A32B32G32R32F, BUF_Static);
@@ -458,29 +377,21 @@ void FTileIntersectionResources::InitDynamicRHI()
 
 void FSceneViewState::DestroyAOTileResources()
 {
-	if (AOTileIntersectionResources || VPLTileIntersectionResources) 
+	if (AOTileIntersectionResources) 
 	{
-		ENQUEUE_UNIQUE_RENDER_COMMAND_TWOPARAMETER(
+		ENQUEUE_UNIQUE_RENDER_COMMAND_ONEPARAMETER(
 			DeleteTileResources,
 			FTileIntersectionResources*, AOTileIntersectionResources, AOTileIntersectionResources,
-			FVPLTileIntersectionResources*, VPLTileIntersectionResources, VPLTileIntersectionResources,
 			{
 				if (AOTileIntersectionResources)
 				{
 					AOTileIntersectionResources->ReleaseResource();
 					delete AOTileIntersectionResources;
 				}
-				
-				if (VPLTileIntersectionResources)
-				{
-					VPLTileIntersectionResources->ReleaseResource();
-					delete VPLTileIntersectionResources;
-				}
 			}
 		);
 	}
 	AOTileIntersectionResources = NULL;
-	VPLTileIntersectionResources = NULL;
 }
 
 void OnClearSurfaceCache(UWorld* InWorld)
@@ -501,215 +412,18 @@ FAutoConsoleCommandWithWorld ClearCacheConsoleCommand(
 	FConsoleCommandWithWorldDelegate::CreateStatic(OnClearSurfaceCache)
 	);
 
-uint32 UpdateObjectsGroupSize = 64;
-
-class FUploadObjectsToBufferCS : public FGlobalShader
-{
-	DECLARE_SHADER_TYPE(FUploadObjectsToBufferCS,Global)
-public:
-
-	static bool ShouldCache(EShaderPlatform Platform)
-	{
-		return IsFeatureLevelSupported(Platform, ERHIFeatureLevel::SM5) && DoesPlatformSupportDistanceFieldAO(Platform);
-	}
-
-	static void ModifyCompilationEnvironment(EShaderPlatform Platform, FShaderCompilerEnvironment& OutEnvironment)
-	{
-		FGlobalShader::ModifyCompilationEnvironment(Platform,OutEnvironment);
-		OutEnvironment.SetDefine(TEXT("UPDATEOBJECTS_THREADGROUP_SIZE"), UpdateObjectsGroupSize);
-	}
-
-	FUploadObjectsToBufferCS(const ShaderMetaType::CompiledShaderInitializerType& Initializer)
-		: FGlobalShader(Initializer)
-	{
-		NumUploadOperations.Bind(Initializer.ParameterMap, TEXT("NumUploadOperations"));
-		UploadOperationIndices.Bind(Initializer.ParameterMap, TEXT("UploadOperationIndices"));
-		UploadOperationData.Bind(Initializer.ParameterMap, TEXT("UploadOperationData"));
-		ObjectBufferParameters.Bind(Initializer.ParameterMap);
-	}
-
-	FUploadObjectsToBufferCS()
-	{
-	}
-
-	void SetParameters(FRHICommandList& RHICmdList, const FScene* Scene, uint32 NumUploadOperationsValue, FShaderResourceViewRHIParamRef InUploadOperationIndices, FShaderResourceViewRHIParamRef InUploadOperationData)
-	{
-		FComputeShaderRHIParamRef ShaderRHI = GetComputeShader();
-
-		SetShaderValue(RHICmdList, ShaderRHI, NumUploadOperations, NumUploadOperationsValue);
-
-		SetSRVParameter(RHICmdList, ShaderRHI, UploadOperationIndices, InUploadOperationIndices);
-		SetSRVParameter(RHICmdList, ShaderRHI, UploadOperationData, InUploadOperationData);
-
-		ObjectBufferParameters.Set(RHICmdList, ShaderRHI, *(Scene->DistanceFieldSceneData.ObjectBuffers), Scene->DistanceFieldSceneData.NumObjectsInBuffer);
-	}
-
-	void UnsetParameters(FRHICommandList& RHICmdList)
-	{
-		ObjectBufferParameters.UnsetParameters(RHICmdList, GetComputeShader());
-	}
-
-	virtual bool Serialize(FArchive& Ar)
-	{		
-		bool bShaderHasOutdatedParameters = FGlobalShader::Serialize(Ar);
-		Ar << NumUploadOperations;
-		Ar << UploadOperationIndices;
-		Ar << UploadOperationData;
-		Ar << ObjectBufferParameters;
-		return bShaderHasOutdatedParameters;
-	}
-
-private:
-
-	FShaderParameter NumUploadOperations;
-	FShaderResourceParameter UploadOperationIndices;
-	FShaderResourceParameter UploadOperationData;
-	FDistanceFieldObjectBufferParameters ObjectBufferParameters;
-};
-
-IMPLEMENT_SHADER_TYPE(,FUploadObjectsToBufferCS,TEXT("DistanceFieldSurfaceCacheLightingCompute"),TEXT("UploadObjectsToBufferCS"),SF_Compute);
-
-
-
-class FCopyObjectBufferCS : public FGlobalShader
-{
-	DECLARE_SHADER_TYPE(FCopyObjectBufferCS,Global)
-public:
-
-	static bool ShouldCache(EShaderPlatform Platform)
-	{
-		return IsFeatureLevelSupported(Platform, ERHIFeatureLevel::SM5) && DoesPlatformSupportDistanceFieldAO(Platform);
-	}
-
-	static void ModifyCompilationEnvironment(EShaderPlatform Platform, FShaderCompilerEnvironment& OutEnvironment)
-	{
-		FGlobalShader::ModifyCompilationEnvironment(Platform,OutEnvironment);
-		OutEnvironment.SetDefine(TEXT("UPDATEOBJECTS_THREADGROUP_SIZE"), UpdateObjectsGroupSize);
-	}
-
-	FCopyObjectBufferCS(const ShaderMetaType::CompiledShaderInitializerType& Initializer)
-		: FGlobalShader(Initializer)
-	{
-		CopyObjectBounds.Bind(Initializer.ParameterMap, TEXT("CopyObjectBounds"));
-		CopyObjectData.Bind(Initializer.ParameterMap, TEXT("CopyObjectData"));
-		ObjectBufferParameters.Bind(Initializer.ParameterMap);
-	}
-
-	FCopyObjectBufferCS()
-	{
-	}
+bool bListMemoryNextFrame = false;
 
-	void SetParameters(FRHICommandList& RHICmdList, FDistanceFieldObjectBuffers& ObjectBuffersSource, FDistanceFieldObjectBuffers& ObjectBuffersDest, int32 NumObjectsValue)
-	{
-		FComputeShaderRHIParamRef ShaderRHI = GetComputeShader();
-
-		CopyObjectBounds.SetBuffer(RHICmdList, ShaderRHI, ObjectBuffersDest.Bounds);
-		CopyObjectData.SetBuffer(RHICmdList, ShaderRHI, ObjectBuffersDest.Data);
-
-		ObjectBufferParameters.Set(RHICmdList, ShaderRHI, ObjectBuffersSource, NumObjectsValue);
-	}
-
-	void UnsetParameters(FRHICommandList& RHICmdList)
-	{
-		ObjectBufferParameters.UnsetParameters(RHICmdList, GetComputeShader());
-		CopyObjectBounds.UnsetUAV(RHICmdList, GetComputeShader());
-		CopyObjectData.UnsetUAV(RHICmdList, GetComputeShader());
-	}
-
-	virtual bool Serialize(FArchive& Ar)
-	{		
-		bool bShaderHasOutdatedParameters = FGlobalShader::Serialize(Ar);
-		Ar << CopyObjectBounds;
-		Ar << CopyObjectData;
-		Ar << ObjectBufferParameters;
-		return bShaderHasOutdatedParameters;
-	}
-
-private:
-
-	FRWShaderParameter CopyObjectBounds;
-	FRWShaderParameter CopyObjectData;
-	FDistanceFieldObjectBufferParameters ObjectBufferParameters;
-};
-
-IMPLEMENT_SHADER_TYPE(,FCopyObjectBufferCS,TEXT("DistanceFieldSurfaceCacheLightingCompute"),TEXT("CopyObjectBufferCS"),SF_Compute);
-
-template<bool bRemoveFromSameBuffer>
-class TRemoveObjectsFromBufferCS : public FGlobalShader
+void OnListMemory(UWorld* InWorld)
 {
-	DECLARE_SHADER_TYPE(TRemoveObjectsFromBufferCS,Global)
-public:
-
-	static bool ShouldCache(EShaderPlatform Platform)
-	{
-		return IsFeatureLevelSupported(Platform, ERHIFeatureLevel::SM5) && DoesPlatformSupportDistanceFieldAO(Platform);
-	}
-
-	static void ModifyCompilationEnvironment(EShaderPlatform Platform, FShaderCompilerEnvironment& OutEnvironment)
-	{
-		FGlobalShader::ModifyCompilationEnvironment(Platform,OutEnvironment);
-		OutEnvironment.SetDefine(TEXT("UPDATEOBJECTS_THREADGROUP_SIZE"), UpdateObjectsGroupSize);
-		OutEnvironment.SetDefine(TEXT("REMOVE_FROM_SAME_BUFFER"), bRemoveFromSameBuffer ? TEXT("1") : TEXT("0"));
-	}
-
-	TRemoveObjectsFromBufferCS(const ShaderMetaType::CompiledShaderInitializerType& Initializer)
-		: FGlobalShader(Initializer)
-	{
-		NumRemoveOperations.Bind(Initializer.ParameterMap, TEXT("NumRemoveOperations"));
-		RemoveOperationIndices.Bind(Initializer.ParameterMap, TEXT("RemoveOperationIndices"));
-		ObjectBufferParameters.Bind(Initializer.ParameterMap);
-		ObjectBounds2.Bind(Initializer.ParameterMap, TEXT("ObjectBounds2"));
-		ObjectData2.Bind(Initializer.ParameterMap, TEXT("ObjectData2"));
-	}
-
-	TRemoveObjectsFromBufferCS()
-	{
-	}
-
-	void SetParameters(
-		FRHICommandList& RHICmdList, 
-		const FScene* Scene, 
-		uint32 NumRemoveOperationsValue, 
-		FShaderResourceViewRHIParamRef InRemoveOperationIndices, 
-		FShaderResourceViewRHIParamRef InObjectBounds2, 
-		FShaderResourceViewRHIParamRef InObjectData2)
-	{
-		FComputeShaderRHIParamRef ShaderRHI = GetComputeShader();
-
-		SetShaderValue(RHICmdList, ShaderRHI, NumRemoveOperations, NumRemoveOperationsValue);
-		SetSRVParameter(RHICmdList, ShaderRHI, RemoveOperationIndices, InRemoveOperationIndices);
-		ObjectBufferParameters.Set(RHICmdList, ShaderRHI, *(Scene->DistanceFieldSceneData.ObjectBuffers), Scene->DistanceFieldSceneData.NumObjectsInBuffer);
-		SetSRVParameter(RHICmdList, ShaderRHI, ObjectBounds2, InObjectBounds2);
-		SetSRVParameter(RHICmdList, ShaderRHI, ObjectData2, InObjectData2);
-	}
-
-	void UnsetParameters(FRHICommandList& RHICmdList)
-	{
-		ObjectBufferParameters.UnsetParameters(RHICmdList, GetComputeShader());
-	}
-
-	virtual bool Serialize(FArchive& Ar)
-	{		
-		bool bShaderHasOutdatedParameters = FGlobalShader::Serialize(Ar);
-		Ar << NumRemoveOperations;
-		Ar << RemoveOperationIndices;
-		Ar << ObjectBufferParameters;
-		Ar << ObjectBounds2;
-		Ar << ObjectData2;
-		return bShaderHasOutdatedParameters;
-	}
-
-private:
-
-	FShaderParameter NumRemoveOperations;
-	FShaderResourceParameter RemoveOperationIndices;
-	FDistanceFieldObjectBufferParameters ObjectBufferParameters;
-	FShaderResourceParameter ObjectBounds2;
-	FShaderResourceParameter ObjectData2;
-};
+	bListMemoryNextFrame = true;
+}
 
-IMPLEMENT_SHADER_TYPE(template<>,TRemoveObjectsFromBufferCS<true>,TEXT("DistanceFieldSurfaceCacheLightingCompute"),TEXT("RemoveObjectsFromBufferCS"),SF_Compute);
-IMPLEMENT_SHADER_TYPE(template<>,TRemoveObjectsFromBufferCS<false>,TEXT("DistanceFieldSurfaceCacheLightingCompute"),TEXT("RemoveObjectsFromBufferCS"),SF_Compute);
+FAutoConsoleCommandWithWorld ListMemoryConsoleCommand(
+	TEXT("r.AOListMemory"),
+	TEXT(""),
+	FConsoleCommandWithWorldDelegate::CreateStatic(OnListMemory)
+	);
 
 class FCullObjectsForViewCS : public FGlobalShader
 {
@@ -803,383 +517,6 @@ private:
 
 IMPLEMENT_SHADER_TYPE(,FCullObjectsForViewCS,TEXT("DistanceFieldSurfaceCacheLightingCompute"),TEXT("CullObjectsForViewCS"),SF_Compute);
 
-
-void FDeferredShadingSceneRenderer::UpdateGlobalDistanceFieldObjectBuffers(FRHICommandListImmediate& RHICmdList)
-{
-	FDistanceFieldSceneData& DistanceFieldSceneData = Scene->DistanceFieldSceneData;
-
-	if (GDistanceFieldVolumeTextureAtlas.VolumeTextureRHI
-		&& (DistanceFieldSceneData.PendingAddOperations.Num() > 0 
-			|| DistanceFieldSceneData.PendingUpdateOperations.Num() > 0 
-			|| DistanceFieldSceneData.PendingRemoveOperations.Num() > 0
-			|| DistanceFieldSceneData.AtlasGeneration != GDistanceFieldVolumeTextureAtlas.GetGeneration()))
-	{
-		QUICK_SCOPE_CYCLE_COUNTER(STAT_UpdateObjectData);
-		SCOPED_DRAW_EVENT(RHICmdList, UpdateSceneObjectData);
-
-		if (!DistanceFieldSceneData.ObjectBuffers)
-		{
-			DistanceFieldSceneData.ObjectBuffers = new FDistanceFieldObjectBuffers();
-		}
-
-		if (DistanceFieldSceneData.AtlasGeneration != GDistanceFieldVolumeTextureAtlas.GetGeneration())
-		{
-			DistanceFieldSceneData.AtlasGeneration = GDistanceFieldVolumeTextureAtlas.GetGeneration();
-
-			for (int32 PrimitiveInstanceIndex = 0; PrimitiveInstanceIndex < DistanceFieldSceneData.PrimitiveInstanceMapping.Num(); PrimitiveInstanceIndex++)
-			{
-				FPrimitiveAndInstance& PrimitiveInstance = DistanceFieldSceneData.PrimitiveInstanceMapping[PrimitiveInstanceIndex];
-
-				// Queue an update of all primitives, since the atlas layout has changed
-				if (PrimitiveInstance.InstanceIndex == 0 
-					&& !DistanceFieldSceneData.PendingRemoveOperations.Contains(PrimitiveInstanceIndex)
-					&& !DistanceFieldSceneData.PendingAddOperations.Contains(PrimitiveInstance.Primitive)
-					&& !DistanceFieldSceneData.PendingUpdateOperations.Contains(PrimitiveInstance.Primitive))
-				{
-					DistanceFieldSceneData.PendingUpdateOperations.Add(PrimitiveInstance.Primitive);
-				}
-			}
-		}
-
-		TArray<uint32> UploadObjectIndices;
-		TArray<FVector4> UploadObjectData;
-
-		if (DistanceFieldSceneData.PendingAddOperations.Num() > 0 || DistanceFieldSceneData.PendingUpdateOperations.Num() > 0)
-		{
-			TArray<FMatrix> ObjectLocalToWorldTransforms;
-
-			const int32 NumUploadOperations = DistanceFieldSceneData.PendingAddOperations.Num() + DistanceFieldSceneData.PendingUpdateOperations.Num();
-			UploadObjectData.Empty(NumUploadOperations * UploadObjectDataStride);
-			UploadObjectIndices.Empty(NumUploadOperations);
-
-			const int32 NumTexelsOneDimX = GDistanceFieldVolumeTextureAtlas.GetSizeX();
-			const int32 NumTexelsOneDimY = GDistanceFieldVolumeTextureAtlas.GetSizeY();
-			const int32 NumTexelsOneDimZ = GDistanceFieldVolumeTextureAtlas.GetSizeZ();
-			const FVector InvTextureDim(1.0f / NumTexelsOneDimX, 1.0f / NumTexelsOneDimY, 1.0f / NumTexelsOneDimZ);
-
-			int32 OriginalNumObjects = DistanceFieldSceneData.NumObjectsInBuffer;
-
-			for (int32 UploadPrimitiveIndex = 0; UploadPrimitiveIndex < NumUploadOperations; UploadPrimitiveIndex++)
-			{
-				const bool bIsAddOperation = UploadPrimitiveIndex < DistanceFieldSceneData.PendingAddOperations.Num();
-
-				FPrimitiveSceneInfo* PrimitiveSceneInfo = bIsAddOperation 
-					? DistanceFieldSceneData.PendingAddOperations[UploadPrimitiveIndex] 
-				: DistanceFieldSceneData.PendingUpdateOperations[UploadPrimitiveIndex - DistanceFieldSceneData.PendingAddOperations.Num()];
-
-				ObjectLocalToWorldTransforms.Reset();
-
-				FBox LocalVolumeBounds;
-				FIntVector BlockMin;
-				FIntVector BlockSize;
-				bool bBuiltAsIfTwoSided;
-				bool bMeshWasPlane;
-				PrimitiveSceneInfo->Proxy->GetDistancefieldAtlasData(LocalVolumeBounds, BlockMin, BlockSize, bBuiltAsIfTwoSided, bMeshWasPlane, ObjectLocalToWorldTransforms);
-
-				if (BlockMin.X >= 0 && BlockMin.Y >= 0 && BlockMin.Z >= 0)
-				{
-					if (bIsAddOperation)
-					{
-						DistanceFieldSceneData.NumObjectsInBuffer += ObjectLocalToWorldTransforms.Num();
-						PrimitiveSceneInfo->DistanceFieldInstanceIndices.Empty(ObjectLocalToWorldTransforms.Num());
-					}
-
-					for (int32 TransformIndex = 0; TransformIndex < ObjectLocalToWorldTransforms.Num(); TransformIndex++)
-					{
-						const uint32 UploadIndex = bIsAddOperation ? OriginalNumObjects + UploadObjectIndices.Num() : PrimitiveSceneInfo->DistanceFieldInstanceIndices[TransformIndex];
-
-						if (bIsAddOperation)
-						{
-							const int32 AddIndex = OriginalNumObjects + UploadObjectIndices.Num();
-							DistanceFieldSceneData.PrimitiveInstanceMapping.Add(FPrimitiveAndInstance(PrimitiveSceneInfo, TransformIndex));
-							PrimitiveSceneInfo->DistanceFieldInstanceIndices.Add(AddIndex);
-						}
-
-						UploadObjectIndices.Add(UploadIndex);
-
-						FMatrix LocalToWorld = ObjectLocalToWorldTransforms[TransformIndex];
-
-						if (bMeshWasPlane)
-						{
-							FVector LocalScales = LocalToWorld.GetScaleVector();
-							FVector AbsLocalScales(FMath::Abs(LocalScales.X), FMath::Abs(LocalScales.Y), FMath::Abs(LocalScales.Z));
-							float MidScale = FMath::Min(AbsLocalScales.X, AbsLocalScales.Y);
-							float ScaleAdjust = FMath::Sign(LocalScales.Z) * MidScale / AbsLocalScales.Z;
-							// The mesh was determined to be a plane flat in Z during the build process, so we can change the Z scale
-							// Helps in cases with modular ground pieces with scales of (10, 10, 1) and some triangles just above Z=0
-							LocalToWorld.SetAxis(2, LocalToWorld.GetScaledAxis(EAxis::Z) * ScaleAdjust);
-						}
-
-						const FMatrix VolumeToWorld = FScaleMatrix(LocalVolumeBounds.GetExtent()) 
-							* FTranslationMatrix(LocalVolumeBounds.GetCenter())
-							* LocalToWorld;
-
-						const FVector4 ObjectBoundingSphere(VolumeToWorld.GetOrigin(), VolumeToWorld.GetScaleVector().Size());
-
-						UploadObjectData.Add(ObjectBoundingSphere);
-
-						const float MaxExtent = LocalVolumeBounds.GetExtent().GetMax();
-
-						const FMatrix UniformScaleVolumeToWorld = FScaleMatrix(MaxExtent) 
-							* FTranslationMatrix(LocalVolumeBounds.GetCenter())
-							* LocalToWorld;
-
-						const FVector InvBlockSize(1.0f / BlockSize.X, 1.0f / BlockSize.Y, 1.0f / BlockSize.Z);
-
-						//float3 VolumeUV = (VolumePosition / LocalPositionExtent * .5f * UVScale + .5f * UVScale + UVAdd;
-						const FVector LocalPositionExtent = LocalVolumeBounds.GetExtent() / FVector(MaxExtent);
-						const FVector UVScale = FVector(BlockSize) * InvTextureDim;
-						const float VolumeScale = UniformScaleVolumeToWorld.GetMaximumAxisScale();
-
-						const FMatrix WorldToVolume = UniformScaleVolumeToWorld.Inverse();
-						// WorldToVolume
-						UploadObjectData.Add(*(FVector4*)&WorldToVolume.M[0]);
-						UploadObjectData.Add(*(FVector4*)&WorldToVolume.M[1]);
-						UploadObjectData.Add(*(FVector4*)&WorldToVolume.M[2]);
-						UploadObjectData.Add(*(FVector4*)&WorldToVolume.M[3]);
-
-						// Clamp to texel center by subtracting a half texel in the [-1,1] position space
-						// LocalPositionExtent
-						UploadObjectData.Add(FVector4(LocalPositionExtent - InvBlockSize, LocalVolumeBounds.Min.X));
-
-						// UVScale
-						UploadObjectData.Add(FVector4(FVector(BlockSize) * InvTextureDim * .5f / LocalPositionExtent, VolumeScale));
-
-						// UVAdd
-						UploadObjectData.Add(FVector4(FVector(BlockMin) * InvTextureDim + .5f * UVScale, LocalVolumeBounds.Min.Y));
-
-						// Box bounds
-						UploadObjectData.Add(FVector4(LocalVolumeBounds.Max, LocalVolumeBounds.Min.Z));
-
-						UploadObjectData.Add(*(FVector4*)&UniformScaleVolumeToWorld.M[0]);
-						UploadObjectData.Add(*(FVector4*)&UniformScaleVolumeToWorld.M[1]);
-						UploadObjectData.Add(*(FVector4*)&UniformScaleVolumeToWorld.M[2]);
-
-						UploadObjectData.Add(*(FVector4*)&LocalToWorld.M[0]);
-						UploadObjectData.Add(*(FVector4*)&LocalToWorld.M[1]);
-						UploadObjectData.Add(*(FVector4*)&LocalToWorld.M[2]);
-						UploadObjectData.Add(*(FVector4*)&LocalToWorld.M[3]);
-
-						checkSlow(UploadObjectData.Num() % UploadObjectDataStride == 0);
-					}
-				}
-			}
-
-			DistanceFieldSceneData.PendingAddOperations.Reset();
-			DistanceFieldSceneData.PendingUpdateOperations.Reset();
-
-			if (DistanceFieldSceneData.ObjectBuffers->MaxObjects < DistanceFieldSceneData.NumObjectsInBuffer)
-			{
-				if (DistanceFieldSceneData.ObjectBuffers->MaxObjects > 0)
-				{
-					// Realloc
-					FDistanceFieldObjectBuffers* NewObjectBuffers = new FDistanceFieldObjectBuffers();
-					NewObjectBuffers->MaxObjects = DistanceFieldSceneData.NumObjectsInBuffer * 5 / 4;
-					NewObjectBuffers->Initialize();
-
-					{
-						TShaderMapRef<FCopyObjectBufferCS> ComputeShader(GetGlobalShaderMap(Scene->GetFeatureLevel()));
-						RHICmdList.SetComputeShader(ComputeShader->GetComputeShader());
-						ComputeShader->SetParameters(RHICmdList, *(DistanceFieldSceneData.ObjectBuffers), *NewObjectBuffers, OriginalNumObjects);
-
-						DispatchComputeShader(RHICmdList, *ComputeShader, FMath::DivideAndRoundUp<uint32>(OriginalNumObjects, UpdateObjectsGroupSize), 1, 1);
-						ComputeShader->UnsetParameters(RHICmdList);
-					}
-
-					DistanceFieldSceneData.ObjectBuffers->Release();
-					delete DistanceFieldSceneData.ObjectBuffers;
-					DistanceFieldSceneData.ObjectBuffers = NewObjectBuffers;
-				}
-				else
-				{
-					// First time allocate
-					DistanceFieldSceneData.ObjectBuffers->MaxObjects = DistanceFieldSceneData.NumObjectsInBuffer * 5 / 4;
-					DistanceFieldSceneData.ObjectBuffers->Initialize();
-				}
-			}
-		}
-
-		if (UploadObjectIndices.Num() > 0)
-		{
-			if (UploadObjectIndices.Num() > GDistanceFieldUploadIndices.UploadIndices.MaxElements)
-			{
-				GDistanceFieldUploadIndices.UploadIndices.MaxElements = UploadObjectIndices.Num() * 5 / 4;
-				GDistanceFieldUploadIndices.UploadIndices.Release();
-				GDistanceFieldUploadIndices.UploadIndices.Initialize();
-
-				GDistanceFieldUploadData.UploadData.MaxElements = UploadObjectIndices.Num() * 5 / 4;
-				GDistanceFieldUploadData.UploadData.Release();
-				GDistanceFieldUploadData.UploadData.Initialize();
-			}
-
-			void* LockedBuffer = RHILockVertexBuffer(GDistanceFieldUploadIndices.UploadIndices.Buffer, 0, GDistanceFieldUploadIndices.UploadIndices.Buffer->GetSize(), RLM_WriteOnly);
-			const uint32 MemcpySize = UploadObjectIndices.GetTypeSize() * UploadObjectIndices.Num();
-			check(GDistanceFieldUploadIndices.UploadIndices.Buffer->GetSize() >= MemcpySize);
-			FPlatformMemory::Memcpy(LockedBuffer, UploadObjectIndices.GetData(), MemcpySize);
-			RHIUnlockVertexBuffer(GDistanceFieldUploadIndices.UploadIndices.Buffer);
-
-			LockedBuffer = RHILockVertexBuffer(GDistanceFieldUploadData.UploadData.Buffer, 0, GDistanceFieldUploadData.UploadData.Buffer->GetSize(), RLM_WriteOnly);
-			const uint32 MemcpySize2 = UploadObjectData.GetTypeSize() * UploadObjectData.Num();
-			check(GDistanceFieldUploadData.UploadData.Buffer->GetSize() >= MemcpySize2);
-			FPlatformMemory::Memcpy(LockedBuffer, UploadObjectData.GetData(), MemcpySize2);
-			RHIUnlockVertexBuffer(GDistanceFieldUploadData.UploadData.Buffer);
-
-			{
-				TShaderMapRef<FUploadObjectsToBufferCS> ComputeShader(GetGlobalShaderMap(Scene->GetFeatureLevel()));
-				RHICmdList.SetComputeShader(ComputeShader->GetComputeShader());
-				ComputeShader->SetParameters(RHICmdList, Scene, UploadObjectIndices.Num(), GDistanceFieldUploadIndices.UploadIndices.BufferSRV, GDistanceFieldUploadData.UploadData.BufferSRV);
-
-				DispatchComputeShader(RHICmdList, *ComputeShader, FMath::DivideAndRoundUp<uint32>(UploadObjectIndices.Num(), UpdateObjectsGroupSize), 1, 1);
-				ComputeShader->UnsetParameters(RHICmdList);
-			}
-		}
-
-		check(DistanceFieldSceneData.NumObjectsInBuffer == DistanceFieldSceneData.PrimitiveInstanceMapping.Num());
-
-		TArray<FIntRect> RemoveObjectIndices;
-		FDistanceFieldObjectBuffers* TemporaryCopySourceBuffers = NULL;
-
-		if (DistanceFieldSceneData.PendingRemoveOperations.Num() > 0)
-		{
-			check(DistanceFieldSceneData.NumObjectsInBuffer >= DistanceFieldSceneData.PendingRemoveOperations.Num());
-
-			// Sort from smallest to largest
-			DistanceFieldSceneData.PendingRemoveOperations.Sort();
-
-			// We have multiple remove requests enqueued in PendingRemoveOperations, can only use the RemoveAtSwap version when there won't be collisions
-			const bool bUseRemoveAtSwap = DistanceFieldSceneData.PendingRemoveOperations.Last() < DistanceFieldSceneData.NumObjectsInBuffer - DistanceFieldSceneData.PendingRemoveOperations.Num();
-
-			if (bUseRemoveAtSwap)
-			{
-				// Remove everything in parallel in the same buffer with a RemoveAtSwap algorithm
-				for (int32 RemovePrimitiveIndex = 0; RemovePrimitiveIndex < DistanceFieldSceneData.PendingRemoveOperations.Num(); RemovePrimitiveIndex++)
-				{
-					DistanceFieldSceneData.NumObjectsInBuffer--;
-					const int32 RemoveIndex = DistanceFieldSceneData.PendingRemoveOperations[RemovePrimitiveIndex];
-					const int32 MoveFromIndex = DistanceFieldSceneData.NumObjectsInBuffer;
-
-					check(RemoveIndex != MoveFromIndex);
-					// Queue a compute shader move
-					RemoveObjectIndices.Add(FIntRect(RemoveIndex, MoveFromIndex, 0, 0));
-
-					// Fixup indices of the primitive that is being moved
-					FPrimitiveAndInstance& PrimitiveAndInstanceBeingMoved = DistanceFieldSceneData.PrimitiveInstanceMapping[MoveFromIndex];
-					check(PrimitiveAndInstanceBeingMoved.Primitive && PrimitiveAndInstanceBeingMoved.Primitive->DistanceFieldInstanceIndices.Num() > 0);
-					PrimitiveAndInstanceBeingMoved.Primitive->DistanceFieldInstanceIndices[PrimitiveAndInstanceBeingMoved.InstanceIndex] = RemoveIndex;
-
-					DistanceFieldSceneData.PrimitiveInstanceMapping.RemoveAtSwap(RemoveIndex);
-				}
-			}
-			else
-			{
-				// Have to copy the object data to allow parallel removing
-				TemporaryCopySourceBuffers = DistanceFieldSceneData.ObjectBuffers;
-				DistanceFieldSceneData.ObjectBuffers = new FDistanceFieldObjectBuffers();
-				DistanceFieldSceneData.ObjectBuffers->MaxObjects = TemporaryCopySourceBuffers->MaxObjects;
-				DistanceFieldSceneData.ObjectBuffers->Initialize();
-
-				TArray<FPrimitiveAndInstance> OriginalPrimitiveInstanceMapping = DistanceFieldSceneData.PrimitiveInstanceMapping;
-				DistanceFieldSceneData.PrimitiveInstanceMapping.Reset();
-
-				const int32 NumDestObjects = DistanceFieldSceneData.NumObjectsInBuffer - DistanceFieldSceneData.PendingRemoveOperations.Num();
-				int32 SourceIndex = 0;
-				int32 NextPendingRemoveIndex = 0;
-
-				for (int32 DestinationIndex = 0; DestinationIndex < NumDestObjects; DestinationIndex++)
-				{
-					while (NextPendingRemoveIndex < DistanceFieldSceneData.PendingRemoveOperations.Num()
-						&& DistanceFieldSceneData.PendingRemoveOperations[NextPendingRemoveIndex] == SourceIndex)
-					{
-						NextPendingRemoveIndex++;
-						SourceIndex++;
-					}
-
-					// Queue a compute shader move
-					RemoveObjectIndices.Add(FIntRect(DestinationIndex, SourceIndex, 0, 0));
-
-					// Fixup indices of the primitive that is being moved
-					FPrimitiveAndInstance& PrimitiveAndInstanceBeingMoved = OriginalPrimitiveInstanceMapping[SourceIndex];
-					check(PrimitiveAndInstanceBeingMoved.Primitive && PrimitiveAndInstanceBeingMoved.Primitive->DistanceFieldInstanceIndices.Num() > 0);
-					PrimitiveAndInstanceBeingMoved.Primitive->DistanceFieldInstanceIndices[PrimitiveAndInstanceBeingMoved.InstanceIndex] = DestinationIndex;
-
-					check(DistanceFieldSceneData.PrimitiveInstanceMapping.Num() == DestinationIndex);
-					DistanceFieldSceneData.PrimitiveInstanceMapping.Add(PrimitiveAndInstanceBeingMoved);
-
-					SourceIndex++;
-				}
-
-				DistanceFieldSceneData.NumObjectsInBuffer = NumDestObjects;
-
-
-				/*
-				// Have to remove one at a time while any entries to remove are at the end of the buffer
-				DistanceFieldSceneData.NumObjectsInBuffer--;
-				const int32 RemoveIndex = DistanceFieldSceneData.PendingRemoveOperations[ParallelConflictIndex];
-				const int32 MoveFromIndex = DistanceFieldSceneData.NumObjectsInBuffer;
-
-				if (RemoveIndex != MoveFromIndex)
-				{
-					// Queue a compute shader move
-					RemoveObjectIndices.Add(FIntRect(RemoveIndex, MoveFromIndex, 0, 0));
-
-					// Fixup indices of the primitive that is being moved
-					FPrimitiveAndInstance& PrimitiveAndInstanceBeingMoved = DistanceFieldSceneData.PrimitiveInstanceMapping[MoveFromIndex];
-					check(PrimitiveAndInstanceBeingMoved.Primitive && PrimitiveAndInstanceBeingMoved.Primitive->DistanceFieldInstanceIndices.Num() > 0);
-					PrimitiveAndInstanceBeingMoved.Primitive->DistanceFieldInstanceIndices[PrimitiveAndInstanceBeingMoved.InstanceIndex] = RemoveIndex;
-				}
-
-				DistanceFieldSceneData.PrimitiveInstanceMapping.RemoveAtSwap(RemoveIndex);
-				DistanceFieldSceneData.PendingRemoveOperations.RemoveAtSwap(ParallelConflictIndex);
-				*/
-			}
-
-			DistanceFieldSceneData.PendingRemoveOperations.Reset();
-
-			if (RemoveObjectIndices.Num() > 0)
-			{
-				if (RemoveObjectIndices.Num() > GDistanceFieldRemoveIndices.RemoveIndices.MaxElements)
-				{
-					GDistanceFieldRemoveIndices.RemoveIndices.MaxElements = RemoveObjectIndices.Num() * 5 / 4;
-					GDistanceFieldRemoveIndices.RemoveIndices.Release();
-					GDistanceFieldRemoveIndices.RemoveIndices.Initialize();
-				}
-
-				void* LockedBuffer = RHILockVertexBuffer(GDistanceFieldRemoveIndices.RemoveIndices.Buffer, 0, GDistanceFieldRemoveIndices.RemoveIndices.Buffer->GetSize(), RLM_WriteOnly);
-				const uint32 MemcpySize = RemoveObjectIndices.GetTypeSize() * RemoveObjectIndices.Num();
-				check(GDistanceFieldRemoveIndices.RemoveIndices.Buffer->GetSize() >= MemcpySize);
-				FPlatformMemory::Memcpy(LockedBuffer, RemoveObjectIndices.GetData(), MemcpySize);
-				RHIUnlockVertexBuffer(GDistanceFieldRemoveIndices.RemoveIndices.Buffer);
-
-				if (bUseRemoveAtSwap)
-				{
-					check(!TemporaryCopySourceBuffers);
-					TShaderMapRef<TRemoveObjectsFromBufferCS<true> > ComputeShader(GetGlobalShaderMap(Scene->GetFeatureLevel()));
-					RHICmdList.SetComputeShader(ComputeShader->GetComputeShader());
-					ComputeShader->SetParameters(RHICmdList, Scene, RemoveObjectIndices.Num(), GDistanceFieldRemoveIndices.RemoveIndices.BufferSRV, NULL, NULL);
-
-					DispatchComputeShader(RHICmdList, *ComputeShader, FMath::DivideAndRoundUp<uint32>(RemoveObjectIndices.Num(), UpdateObjectsGroupSize), 1, 1);
-					ComputeShader->UnsetParameters(RHICmdList);
-				}
-				else
-				{
-					check(TemporaryCopySourceBuffers);
-					TShaderMapRef<TRemoveObjectsFromBufferCS<false> > ComputeShader(GetGlobalShaderMap(Scene->GetFeatureLevel()));
-					RHICmdList.SetComputeShader(ComputeShader->GetComputeShader());
-					ComputeShader->SetParameters(RHICmdList, Scene, RemoveObjectIndices.Num(), GDistanceFieldRemoveIndices.RemoveIndices.BufferSRV, TemporaryCopySourceBuffers->Bounds.SRV, TemporaryCopySourceBuffers->Data.SRV);
-
-					DispatchComputeShader(RHICmdList, *ComputeShader, FMath::DivideAndRoundUp<uint32>(RemoveObjectIndices.Num(), UpdateObjectsGroupSize), 1, 1);
-					ComputeShader->UnsetParameters(RHICmdList);
-
-					TemporaryCopySourceBuffers->Release();
-					delete TemporaryCopySourceBuffers;
-				}
-			}
-		}
-
-		DistanceFieldSceneData.VerifyIntegrity();
-	}
-}
-
 /**  */
 class FBuildTileConesCS : public FGlobalShader
 {
@@ -1828,8 +1165,7 @@ public:
 		SetSRVParameter(RHICmdList, ShaderRHI, DrawParameters, SurfaceCacheResources.Level[DepthLevel]->ScatterDrawParameters.SRV);
 
 		DispatchParameters.SetBuffer(RHICmdList, ShaderRHI, SurfaceCacheResources.DispatchParameters);
-
-		SetShaderValue(RHICmdList, ShaderRHI, TrimFraction, GetTrimOldRecordsFraction(DepthLevel));
+		SetShaderValue(RHICmdList, ShaderRHI, TrimFraction, GAOTrimOldRecordsFraction);
 	}
 
 	void UnsetParameters(FRHICommandList& RHICmdList)
@@ -1921,7 +1257,7 @@ public:
 		CopyIrradianceCacheTileCoordinate.SetBuffer(RHICmdList, ShaderRHI, SurfaceCacheResources.TempResources->TileCoordinate);
 		ScatterDrawParameters.SetBuffer(RHICmdList, ShaderRHI, SurfaceCacheResources.TempResources->ScatterDrawParameters);
 
-		SetShaderValue(RHICmdList, ShaderRHI, TrimFraction, GetTrimOldRecordsFraction(DepthLevel));
+		SetShaderValue(RHICmdList, ShaderRHI, TrimFraction, GAOTrimOldRecordsFraction);
 	}
 
 	void UnsetParameters(FRHICommandList& RHICmdList)
@@ -1930,6 +1266,7 @@ public:
 		CopyIrradianceCacheNormal.UnsetUAV(RHICmdList, GetComputeShader());
 		CopyOccluderRadius.UnsetUAV(RHICmdList, GetComputeShader());
 		CopyIrradianceCacheBentNormal.UnsetUAV(RHICmdList, GetComputeShader());
+		CopyIrradianceCacheIrradiance.UnsetUAV(RHICmdList, GetComputeShader());
 		CopyIrradianceCacheTileCoordinate.UnsetUAV(RHICmdList, GetComputeShader());
 		ScatterDrawParameters.UnsetUAV(RHICmdList, GetComputeShader());
 	}
@@ -2191,73 +1528,6 @@ private:
 
 IMPLEMENT_SHADER_TYPE(,FPopulateCacheCS,TEXT("DistanceFieldSurfaceCacheLightingCompute"),TEXT("PopulateCacheCS"),SF_Compute);
 
-template<bool bOneGroupPerRecord>
-class TSetupFinalGatherIndirectArgumentsCS : public FGlobalShader
-{
-	DECLARE_SHADER_TYPE(TSetupFinalGatherIndirectArgumentsCS,Global)
-public:
-
-	static bool ShouldCache(EShaderPlatform Platform)
-	{
-		return IsFeatureLevelSupported(Platform, ERHIFeatureLevel::SM5) && DoesPlatformSupportDistanceFieldAO(Platform);
-	}
-
-	static void ModifyCompilationEnvironment(EShaderPlatform Platform, FShaderCompilerEnvironment& OutEnvironment)
-	{
-		FGlobalShader::ModifyCompilationEnvironment(Platform,OutEnvironment);
-		OutEnvironment.SetDefine(TEXT("ONE_GROUP_PER_RECORD"), bOneGroupPerRecord ? TEXT("1") : TEXT("0"));
-
-		// To reduce shader compile time of compute shaders with shared memory, doesn't have an impact on generated code with current compiler (June 2010 DX SDK)
-		OutEnvironment.CompilerFlags.Add(CFLAG_StandardOptimization);
-	}
-
-	TSetupFinalGatherIndirectArgumentsCS(const ShaderMetaType::CompiledShaderInitializerType& Initializer)
-		: FGlobalShader(Initializer)
-	{
-		DrawParameters.Bind(Initializer.ParameterMap, TEXT("DrawParameters"));
-		DispatchParameters.Bind(Initializer.ParameterMap, TEXT("DispatchParameters"));
-		SavedStartIndex.Bind(Initializer.ParameterMap, TEXT("SavedStartIndex"));
-	}
-
-	TSetupFinalGatherIndirectArgumentsCS()
-	{
-	}
-
-	void SetParameters(FRHICommandList& RHICmdList, const FSceneView& View, int32 DepthLevel)
-	{
-		FComputeShaderRHIParamRef ShaderRHI = GetComputeShader();
-		FGlobalShader::SetParameters(RHICmdList, ShaderRHI, View);
-
-		const FScene* Scene = (const FScene*)View.Family->Scene;
-		FSurfaceCacheResources& SurfaceCacheResources = *Scene->SurfaceCacheResources;
-
-		SetSRVParameter(RHICmdList, ShaderRHI, DrawParameters, SurfaceCacheResources.Level[DepthLevel]->ScatterDrawParameters.SRV);
-		SetSRVParameter(RHICmdList, ShaderRHI, SavedStartIndex, SurfaceCacheResources.Level[DepthLevel]->SavedStartIndex.SRV);
-
-		DispatchParameters.SetBuffer(RHICmdList, ShaderRHI, SurfaceCacheResources.DispatchParameters);
-	}
-
-	void UnsetParameters(FRHICommandList& RHICmdList)
-	{
-		DispatchParameters.UnsetUAV(RHICmdList, GetComputeShader());
-	}
-
-	virtual bool Serialize(FArchive& Ar)
-	{		
-		bool bShaderHasOutdatedParameters = FGlobalShader::Serialize(Ar);
-		Ar << DrawParameters;
-		Ar << SavedStartIndex;
-		Ar << DispatchParameters;
-		return bShaderHasOutdatedParameters;
-	}
-
-private:
-
-	FShaderResourceParameter DrawParameters;
-	FShaderResourceParameter SavedStartIndex;
-	FRWShaderParameter DispatchParameters;
-};
-
 IMPLEMENT_SHADER_TYPE(template<>,TSetupFinalGatherIndirectArgumentsCS<true>,TEXT("DistanceFieldSurfaceCacheLightingCompute"),TEXT("SetupFinalGatherIndirectArgumentsCS"),SF_Compute);
 IMPLEMENT_SHADER_TYPE(template<>,TSetupFinalGatherIndirectArgumentsCS<false>,TEXT("DistanceFieldSurfaceCacheLightingCompute"),TEXT("SetupFinalGatherIndirectArgumentsCS"),SF_Compute);
 
@@ -2286,9 +1556,9 @@ public:
 TGlobalResource<FDebugBuffer> GDebugBuffer;
 
 template<bool bSupportIrradiance>
-class TFinalGatherCS : public FGlobalShader
+class TConeTraceOcclusionCS : public FGlobalShader
 {
-	DECLARE_SHADER_TYPE(TFinalGatherCS,Global)
+	DECLARE_SHADER_TYPE(TConeTraceOcclusionCS,Global)
 public:
 
 	static bool ShouldCache(EShaderPlatform Platform)
@@ -2307,7 +1577,7 @@ public:
 		OutEnvironment.CompilerFlags.Add(CFLAG_StandardOptimization);
 	}
 
-	TFinalGatherCS(const ShaderMetaType::CompiledShaderInitializerType& Initializer)
+	TConeTraceOcclusionCS(const ShaderMetaType::CompiledShaderInitializerType& Initializer)
 		: FGlobalShader(Initializer)
 	{
 		DeferredParameters.Bind(Initializer.ParameterMap);
@@ -2339,7 +1609,7 @@ public:
 		RecordRadiusScale.Bind(Initializer.ParameterMap, TEXT("RecordRadiusScale"));
 	}
 
-	TFinalGatherCS()
+	TConeTraceOcclusionCS()
 	{
 	}
 
@@ -2510,8 +1780,8 @@ private:
 	FShaderParameter RecordRadiusScale;
 };
 
-IMPLEMENT_SHADER_TYPE(template<>,TFinalGatherCS<true>,TEXT("DistanceFieldSurfaceCacheLightingCompute"),TEXT("FinalGatherCS"),SF_Compute);
-IMPLEMENT_SHADER_TYPE(template<>,TFinalGatherCS<false>,TEXT("DistanceFieldSurfaceCacheLightingCompute"),TEXT("FinalGatherCS"),SF_Compute);
+IMPLEMENT_SHADER_TYPE(template<>,TConeTraceOcclusionCS<true>,TEXT("DistanceFieldSurfaceCacheLightingCompute"),TEXT("ConeTraceOcclusionCS"),SF_Compute);
+IMPLEMENT_SHADER_TYPE(template<>,TConeTraceOcclusionCS<false>,TEXT("DistanceFieldSurfaceCacheLightingCompute"),TEXT("ConeTraceOcclusionCS"),SF_Compute);
 
 
 class FCalculateHeightfieldOcclusionCS : public FGlobalShader
@@ -3058,6 +2328,7 @@ public:
 		DistanceFieldNormalTexture.Bind(Initializer.ParameterMap, TEXT("DistanceFieldNormalTexture"));
 		DistanceFieldNormalSampler.Bind(Initializer.ParameterMap, TEXT("DistanceFieldNormalSampler"));
 		DebugBuffer.Bind(Initializer.ParameterMap, TEXT("DebugBuffer"));
+		DistanceFadeScale.Bind(Initializer.ParameterMap, TEXT("DistanceFadeScale"));
 	}
 
 	void SetParameters(
@@ -3092,6 +2363,9 @@ public:
 			);
 
 		SetSRVParameter(RHICmdList, ShaderRHI, DebugBuffer, GDebugBuffer.DebugData.SRV);
+
+		const float DistanceFadeScaleValue = 1.0f / ((1.0f - GAOViewFadeDistanceScale) * GAOMaxViewDistance);
+		SetShaderValue(RHICmdList, ShaderRHI, DistanceFadeScale, DistanceFadeScaleValue);
 	}
 
 	// FShader interface.
@@ -3108,6 +2382,7 @@ public:
 		Ar << DistanceFieldNormalTexture;
 		Ar << DistanceFieldNormalSampler;
 		Ar << DebugBuffer;
+		Ar << DistanceFadeScale;
 		return bShaderHasOutdatedParameters;
 	}
 
@@ -3123,6 +2398,7 @@ private:
 	FShaderResourceParameter DistanceFieldNormalTexture;
 	FShaderResourceParameter DistanceFieldNormalSampler;
 	FShaderResourceParameter DebugBuffer;
+	FShaderParameter DistanceFadeScale;
 };
 
 IMPLEMENT_SHADER_TYPE(template<>,TDistanceFieldAOCombinePS<true>,TEXT("DistanceFieldSurfaceCacheLighting"),TEXT("AOCombinePS"),SF_Pixel);
@@ -3844,7 +3120,7 @@ void GenerateBestSpacedVectors()
 
 		for (int32 i = 0; i < ARRAY_COUNT(OriginalSpacedVectors9); i++)
 		{
-			UE_LOG(LogTemp, Log, TEXT("FVector(%f, %f, %f),"), OriginalSpacedVectors9[i].X, OriginalSpacedVectors9[i].Y, OriginalSpacedVectors9[i].Z);
+			UE_LOG(LogDistanceField, Log, TEXT("FVector(%f, %f, %f),"), OriginalSpacedVectors9[i].X, OriginalSpacedVectors9[i].Y, OriginalSpacedVectors9[i].Z);
 		}
 
 		int32 temp = 0;
@@ -3975,7 +3251,7 @@ FIntPoint BuildTileObjectLists(FRHICommandListImmediate& RHICmdList, FScene* Sce
 			}
 
 			{
-				SCOPED_DRAW_EVENT(RHICmdList, CullObjects);
+				SCOPED_DRAW_EVENT(RHICmdList, CullObjectsToTiles);
 
 				TShaderMapRef<FObjectCullVS> VertexShader(View.ShaderMap);
 				TShaderMapRef<FObjectCullPS> PixelShader(View.ShaderMap);
@@ -4310,6 +3586,54 @@ void RenderIrradianceCacheInterpolation(
 	}
 }
 
+void ListDistanceFieldLightingMemory(const FViewInfo& View)
+{
+	const FScene* Scene = (const FScene*)View.Family->Scene;
+	UE_LOG(LogTemp, Log, TEXT("Shared GPU memory (excluding render targets)"));
+
+	if (Scene->DistanceFieldSceneData.NumObjectsInBuffer > 0)
+	{
+		UE_LOG(LogTemp, Log, TEXT("   Scene Object data %.3fMb"), Scene->DistanceFieldSceneData.ObjectBuffers->GetSizeBytes() / 1024.0f / 1024.0f);
+	}
+	
+	UE_LOG(LogTemp, Log, TEXT("   %s"), *GDistanceFieldVolumeTextureAtlas.GetSizeString());
+	extern FString GetObjectBufferMemoryString();
+	UE_LOG(LogTemp, Log, TEXT("   %s"), *GetObjectBufferMemoryString());
+	UE_LOG(LogTemp, Log, TEXT(""));
+	UE_LOG(LogTemp, Log, TEXT("Distance Field AO"));
+
+	if (Scene->SurfaceCacheResources)
+	{
+		UE_LOG(LogTemp, Log, TEXT("   Surface cache %.3fMb"), Scene->SurfaceCacheResources->GetSizeBytes() / 1024.0f / 1024.0f);
+	}
+	
+	UE_LOG(LogTemp, Log, TEXT("   Temporary cache %.3fMb"), GTemporaryIrradianceCacheResources.GetSizeBytes() / 1024.0f / 1024.0f);
+	UE_LOG(LogTemp, Log, TEXT("   Culled objects %.3fMb"), GAOCulledObjectBuffers.Buffers.GetSizeBytes() / 1024.0f / 1024.0f);
+
+	FTileIntersectionResources* TileIntersectionResources = ((FSceneViewState*)View.State)->AOTileIntersectionResources;
+
+	if (TileIntersectionResources)
+	{
+		UE_LOG(LogTemp, Log, TEXT("   Tile Culled objects %.3fMb"), TileIntersectionResources->GetSizeBytes() / 1024.0f / 1024.0f);
+	}
+	
+	UE_LOG(LogTemp, Log, TEXT(""));
+	UE_LOG(LogTemp, Log, TEXT("Distance Field GI"));
+
+	if (Scene->DistanceFieldSceneData.SurfelBuffers)
+	{
+		UE_LOG(LogTemp, Log, TEXT("   Scene surfel data %.3fMb"), Scene->DistanceFieldSceneData.SurfelBuffers->GetSizeBytes() / 1024.0f / 1024.0f);
+	}
+	
+	if (Scene->DistanceFieldSceneData.InstancedSurfelBuffers)
+	{
+		UE_LOG(LogTemp, Log, TEXT("   Instanced scene surfel data %.3fMb"), Scene->DistanceFieldSceneData.InstancedSurfelBuffers->GetSizeBytes() / 1024.0f / 1024.0f);
+	}
+	
+	extern void ListDistanceFieldGIMemory(const FViewInfo& View);
+	ListDistanceFieldGIMemory(View);
+}
+
 bool SupportsDistanceFieldAO(ERHIFeatureLevel::Type FeatureLevel, EShaderPlatform ShaderPlatform)
 {
 	return GDistanceFieldAO 
@@ -4343,6 +3667,32 @@ public:
 	{}
 };
 
+FRWBuffer* GDebugBuffer2 = NULL;
+
+IMPLEMENT_SHADER_TYPE(,FTrackGPUProgressCS,TEXT("DistanceFieldGlobalIllumination"),TEXT("TrackGPUProgressCS"),SF_Compute);
+
+void TrackGPUProgress(FRHICommandListImmediate& RHICmdList, uint32 DebugId)
+{
+	/*
+	TShaderMapRef<FTrackGPUProgressCS> ComputeShader(GetGlobalShaderMap(ERHIFeatureLevel::SM5));
+	RHICmdList.SetComputeShader(ComputeShader->GetComputeShader());
+	ComputeShader->SetParameters(RHICmdList, DebugId);
+	DispatchComputeShader(RHICmdList, *ComputeShader, 1, 1, 1);
+	ComputeShader->UnsetParameters(RHICmdList);
+
+	{
+		FRenderQueryRHIRef Query = RHICmdList.CreateRenderQuery(RQT_AbsoluteTime);
+		uint64 Temp = 0;
+		RHICmdList.EndRenderQuery(Query);
+		RHICmdList.GetRenderQueryResult(Query, Temp, true);
+
+		uint32* Ptr = (uint32*)RHICmdList.LockVertexBuffer(GDebugBuffer2->Buffer, 0, 4, RLM_ReadOnly);
+
+		UE_LOG(LogDistanceField, Warning, TEXT("DebugId %u,"), *Ptr);
+		RHICmdList.UnlockVertexBuffer(GDebugBuffer2->Buffer);
+	}*/
+}
+
 bool FDeferredShadingSceneRenderer::RenderDistanceFieldAOSurfaceCache(
 	FRHICommandListImmediate& RHICmdList, 
 	const FDistanceFieldAOParameters& Parameters, 
@@ -4367,6 +3717,12 @@ bool FDeferredShadingSceneRenderer::RenderDistanceFieldAOSurfaceCache(
 			check(!Scene->DistanceFieldSceneData.HasPendingOperations());
 			const bool bUseDistanceFieldGI = IsDistanceFieldGIAllowed(View);
 
+			if (!GDebugBuffer2)
+			{
+				GDebugBuffer2 = new FRWBuffer();
+				GDebugBuffer2->Initialize(4, 1, PF_R32_UINT);
+			}
+
 			GenerateBestSpacedVectors();
 
 			// Create surface cache state that will persist across frames if needed
@@ -4375,6 +3731,12 @@ bool FDeferredShadingSceneRenderer::RenderDistanceFieldAOSurfaceCache(
 				Scene->SurfaceCacheResources = new FSurfaceCacheResources();
 			}
 
+			if (bListMemoryNextFrame)
+			{
+				bListMemoryNextFrame = false;
+				ListDistanceFieldLightingMemory(View);
+			}
+
 			FSurfaceCacheResources& SurfaceCacheResources = *Scene->SurfaceCacheResources;
 
 			FIntPoint DownsampledViewRect = FIntPoint::DivideAndRoundDown(View.ViewRect.Size(), GAODownsampleFactor);
@@ -4454,7 +3816,7 @@ bool FDeferredShadingSceneRenderer::RenderDistanceFieldAOSurfaceCache(
 				// Copy and trim last frame's surface cache samples
 				for (int32 DepthLevel = GAOMaxLevel; DepthLevel >= GAOMinLevel; DepthLevel--)
 				{
-					if (GetTrimOldRecordsFraction(DepthLevel) > 0)
+					if (GAOTrimOldRecordsFraction > 0)
 					{
 						{	
 							TShaderMapRef<FSetupCopyIndirectArgumentsCS> ComputeShader(View.ShaderMap);
@@ -4494,13 +3856,13 @@ bool FDeferredShadingSceneRenderer::RenderDistanceFieldAOSurfaceCache(
 
 			if (bUseDistanceFieldGI)
 			{
-				extern void PlaceVPLs(
+				extern void UpdateVPLs(
 					FRHICommandListImmediate& RHICmdList,
 					const FViewInfo& View,
 					const FScene* Scene,
 					const FDistanceFieldAOParameters& Parameters);
 
-				PlaceVPLs(RHICmdList, View, Scene, Parameters);
+				UpdateVPLs(RHICmdList, View, Scene, Parameters);
 			}
 
 			// Start from the highest depth, which will consider the fewest potential shading points
@@ -4569,7 +3931,7 @@ bool FDeferredShadingSceneRenderer::RenderDistanceFieldAOSurfaceCache(
 					// Compute lighting for the new surface cache records by cone-stepping through the object distance fields
 					if (bUseDistanceFieldGI)
 					{
-						TShaderMapRef<TFinalGatherCS<true> > ComputeShader(View.ShaderMap);
+						TShaderMapRef<TConeTraceOcclusionCS<true> > ComputeShader(View.ShaderMap);
 
 						RHICmdList.SetComputeShader(ComputeShader->GetComputeShader());
 						ComputeShader->SetParameters(RHICmdList, View, DistanceFieldAOBentNormalSplat->GetRenderTargetItem(), DistanceFieldNormal->GetRenderTargetItem(), DestLevelDownsampleFactor, DepthLevel, TileListGroupSize, Parameters);
@@ -4579,7 +3941,7 @@ bool FDeferredShadingSceneRenderer::RenderDistanceFieldAOSurfaceCache(
 					}
 					else
 					{
-						TShaderMapRef<TFinalGatherCS<false> > ComputeShader(View.ShaderMap);
+						TShaderMapRef<TConeTraceOcclusionCS<false> > ComputeShader(View.ShaderMap);
 
 						RHICmdList.SetComputeShader(ComputeShader->GetComputeShader());
 						ComputeShader->SetParameters(RHICmdList, View, DistanceFieldAOBentNormalSplat->GetRenderTargetItem(), DistanceFieldNormal->GetRenderTargetItem(), DestLevelDownsampleFactor, DepthLevel, TileListGroupSize, Parameters);
@@ -4692,16 +4054,6 @@ bool FDeferredShadingSceneRenderer::RenderDistanceFieldAOSurfaceCache(
 				{
 					if (bUseDistanceFieldGI)
 					{
-						{	
-							TShaderMapRef<TSetupFinalGatherIndirectArgumentsCS<true> > ComputeShader(View.ShaderMap);
-
-							RHICmdList.SetComputeShader(ComputeShader->GetComputeShader());
-							ComputeShader->SetParameters(RHICmdList, View, DepthLevel);
-							DispatchComputeShader(RHICmdList, *ComputeShader, 1, 1, 1);
-
-							ComputeShader->UnsetParameters(RHICmdList);
-						}
-
 						extern void ComputeIrradianceForSamples(
 							int32 DepthLevel,
 							FRHICommandListImmediate& RHICmdList,
