Surfel Global Illumination with Distance Field shadowing

diff --git a/Engine/Source/Runtime/Renderer/Private/DistanceFieldSurfaceCacheLighting.h b/Engine/Source/Runtime/Renderer/Private/DistanceFieldSurfaceCacheLighting.h
--- a/Engine/Source/Runtime/Renderer/Private/DistanceFieldSurfaceCacheLighting.h
+++ b/Engine/Source/Runtime/Renderer/Private/DistanceFieldSurfaceCacheLighting.h
@@ -63,6 +63,12 @@ public:
 		TileCoordinate.Release();
 	}
 
+	size_t GetSizeBytes() const
+	{
+		return PositionAndRadius.NumBytes + OccluderRadius.NumBytes + Normal.NumBytes + BentNormal.NumBytes
+			+ Irradiance.NumBytes + ScatterDrawParameters.NumBytes + SavedStartIndex.NumBytes + TileCoordinate.NumBytes;
+	}
+
 	int32 MaxIrradianceCacheSamples;
 
 	FRWBuffer PositionAndRadius;
@@ -174,6 +180,20 @@ public:
 		TempResources->ReleaseDynamicRHI();
 	}
 
+	size_t GetSizeBytes() const
+	{
+		size_t TotalSize = 0;
+
+		for (int32 i = MinLevel; i <= MaxLevel; i++)
+		{
+			TotalSize += Level[i]->GetSizeBytes();
+		}
+
+		TotalSize += TempResources->GetSizeBytes();
+
+		return TotalSize;
+	}
+
 	FRWBuffer DispatchParameters;
 
 	FRefinementLevelResources* Level[GAOMaxSupportedLevel + 1];
@@ -190,6 +210,74 @@ private:
 	int32 MaxLevel;
 };
 
+template<bool bOneGroupPerRecord>
+class TSetupFinalGatherIndirectArgumentsCS : public FGlobalShader
+{
+	DECLARE_SHADER_TYPE(TSetupFinalGatherIndirectArgumentsCS,Global)
+public:
+
+	static bool ShouldCache(EShaderPlatform Platform)
+	{
+		return IsFeatureLevelSupported(Platform, ERHIFeatureLevel::SM5) && DoesPlatformSupportDistanceFieldAO(Platform);
+	}
+
+	static void ModifyCompilationEnvironment(EShaderPlatform Platform, FShaderCompilerEnvironment& OutEnvironment)
+	{
+		FGlobalShader::ModifyCompilationEnvironment(Platform,OutEnvironment);
+		OutEnvironment.SetDefine(TEXT("ONE_GROUP_PER_RECORD"), bOneGroupPerRecord ? TEXT("1") : TEXT("0"));
+
+		// To reduce shader compile time of compute shaders with shared memory, doesn't have an impact on generated code with current compiler (June 2010 DX SDK)
+		OutEnvironment.CompilerFlags.Add(CFLAG_StandardOptimization);
+	}
+
+	TSetupFinalGatherIndirectArgumentsCS(const ShaderMetaType::CompiledShaderInitializerType& Initializer)
+		: FGlobalShader(Initializer)
+	{
+		DrawParameters.Bind(Initializer.ParameterMap, TEXT("DrawParameters"));
+		DispatchParameters.Bind(Initializer.ParameterMap, TEXT("DispatchParameters"));
+		SavedStartIndex.Bind(Initializer.ParameterMap, TEXT("SavedStartIndex"));
+	}
+
+	TSetupFinalGatherIndirectArgumentsCS()
+	{
+	}
+
+	void SetParameters(FRHICommandList& RHICmdList, const FSceneView& View, int32 DepthLevel)
+	{
+		FComputeShaderRHIParamRef ShaderRHI = GetComputeShader();
+		FGlobalShader::SetParameters(RHICmdList, ShaderRHI, View);
+
+		const FScene* Scene = (const FScene*)View.Family->Scene;
+		FSurfaceCacheResources& SurfaceCacheResources = *Scene->SurfaceCacheResources;
+
+		SetSRVParameter(RHICmdList, ShaderRHI, DrawParameters, SurfaceCacheResources.Level[DepthLevel]->ScatterDrawParameters.SRV);
+		SetSRVParameter(RHICmdList, ShaderRHI, SavedStartIndex, SurfaceCacheResources.Level[DepthLevel]->SavedStartIndex.SRV);
+
+		DispatchParameters.SetBuffer(RHICmdList, ShaderRHI, SurfaceCacheResources.DispatchParameters);
+	}
+
+	void UnsetParameters(FRHICommandList& RHICmdList)
+	{
+		DispatchParameters.UnsetUAV(RHICmdList, GetComputeShader());
+	}
+
+	virtual bool Serialize(FArchive& Ar)
+	{		
+		bool bShaderHasOutdatedParameters = FGlobalShader::Serialize(Ar);
+		Ar << DrawParameters;
+		Ar << SavedStartIndex;
+		Ar << DispatchParameters;
+		return bShaderHasOutdatedParameters;
+	}
+
+private:
+
+	FShaderResourceParameter DrawParameters;
+	FShaderResourceParameter SavedStartIndex;
+	FRWShaderParameter DispatchParameters;
+};
+
+
 /**  */
 class FTileIntersectionResources : public FRenderResource
 {
@@ -214,6 +302,11 @@ public:
 	FRWBuffer TileHeadData;
 	FRWBuffer TileArrayData;
 	FRWBuffer TileArrayNextAllocation;
+
+	size_t GetSizeBytes() const
+	{
+		return TileConeAxisAndCos.NumBytes + TileConeDepthRanges.NumBytes + TileHeadDataUnpacked.NumBytes + TileHeadData.NumBytes + TileArrayData.NumBytes + TileArrayNextAllocation.NumBytes;
+	}
 };
 
 
@@ -331,6 +424,8 @@ protected:
 
 // Must match usf
 const int32 RecordConeDataStride = 10;
+// In float4s, must match usf
+const int32 NumVisibilitySteps = 10;
 
 /**  */
 class FTemporaryIrradianceCacheResources : public FMaxSizedRWBuffers
@@ -343,6 +438,7 @@ public:
 		{
 			ConeVisibility.Initialize(sizeof(float), MaxSize * NumConeSampleDirections, PF_R32_FLOAT, BUF_Static);
 			ConeData.Initialize(sizeof(float), MaxSize * NumConeSampleDirections * RecordConeDataStride, PF_R32_FLOAT, BUF_Static);
+			StepBentNormal.Initialize(sizeof(float) * 4, MaxSize * NumVisibilitySteps, PF_A32B32G32R32F, BUF_Static);
 		}
 	}
 
@@ -350,8 +446,67 @@ public:
 	{
 		ConeVisibility.Release();
 		ConeData.Release();
+		StepBentNormal.Release();
+	}
+
+	size_t GetSizeBytes() const
+	{
+		return ConeVisibility.NumBytes + ConeData.NumBytes + StepBentNormal.NumBytes;
 	}
 
 	FRWBuffer ConeVisibility;
 	FRWBuffer ConeData;
+	FRWBuffer StepBentNormal;
+};
+
+extern FRWBuffer* GDebugBuffer2;
+
+class FTrackGPUProgressCS : public FGlobalShader
+{
+	DECLARE_SHADER_TYPE(FTrackGPUProgressCS,Global)
+public:
+
+	static bool ShouldCache(EShaderPlatform Platform)
+	{
+		return IsFeatureLevelSupported(Platform, ERHIFeatureLevel::SM5) && DoesPlatformSupportDistanceFieldAO(Platform);
+	}
+
+	FTrackGPUProgressCS(const ShaderMetaType::CompiledShaderInitializerType& Initializer)
+		: FGlobalShader(Initializer)
+	{
+		DebugBuffer.Bind(Initializer.ParameterMap, TEXT("DebugBuffer"));
+		DebugId.Bind(Initializer.ParameterMap, TEXT("DebugId"));
+	}
+
+	FTrackGPUProgressCS()
+	{
+	}
+
+	void SetParameters(FRHICommandList& RHICmdList, uint32 DebugIdValue)
+	{
+		FComputeShaderRHIParamRef ShaderRHI = GetComputeShader();
+
+		DebugBuffer.SetBuffer(RHICmdList, ShaderRHI, *GDebugBuffer2);
+		SetShaderValue(RHICmdList, ShaderRHI, DebugId, DebugIdValue);
+	}
+
+	void UnsetParameters(FRHICommandListImmediate& RHICmdList)
+	{
+		DebugBuffer.UnsetUAV(RHICmdList, GetComputeShader());
+	}
+
+	virtual bool Serialize(FArchive& Ar)
+	{		
+		bool bShaderHasOutdatedParameters = FGlobalShader::Serialize(Ar);
+		Ar << DebugBuffer;
+		Ar << DebugId;
+		return bShaderHasOutdatedParameters;
+	}
+
+private:
+
+	FRWShaderParameter DebugBuffer;
+	FShaderParameter DebugId;
 };
+
+extern void TrackGPUProgress(FRHICommandListImmediate& RHICmdList, uint32 DebugId);
\ No newline at end of file
diff --git a/Engine/Source/Runtime/Renderer/Private/RendererScene.cpp b/Engine/Source/Runtime/Renderer/Private/RendererScene.cpp
--- a/Engine/Source/Runtime/Renderer/Private/RendererScene.cpp
+++ b/Engine/Source/Runtime/Renderer/Private/RendererScene.cpp
@@ -59,7 +59,6 @@ FSceneViewState::FSceneViewState()
 	TemporalAASampleIndex = 0;
 	TemporalAASampleCount = 1;
 	AOTileIntersectionResources = NULL;
-	VPLTileIntersectionResources = NULL;
 	bBokehDOFHistory = true;
 	bBokehDOFHistory2 = true;
 
@@ -80,6 +79,8 @@ FSceneViewState::FSceneViewState()
 FDistanceFieldSceneData::FDistanceFieldSceneData() 
 	: NumObjectsInBuffer(0)
 	, ObjectBuffers(NULL)
+	, SurfelBuffers(NULL)
+	, InstancedSurfelBuffers(NULL)
 	, AtlasGeneration(0)
 {
 
@@ -130,30 +131,25 @@ void FDistanceFieldSceneData::RemovePrimitive(FPrimitiveSceneInfo* InPrimitive)
 {
 	const FPrimitiveSceneProxy* Proxy = InPrimitive->Proxy;
 
-	if (Proxy->SupportsDistanceFieldRepresentation() && Proxy->AffectsDistanceFieldLighting())
+	if (Proxy->AffectsDistanceFieldLighting())
 	{
-		PendingAddOperations.Remove(InPrimitive);
-		PendingUpdateOperations.Remove(InPrimitive);
-
-		for (int32 InstanceIndex = 0; InstanceIndex < InPrimitive->DistanceFieldInstanceIndices.Num(); InstanceIndex++)
+		if (Proxy->SupportsDistanceFieldRepresentation())
 		{
-			int32 RemoveIndex = InPrimitive->DistanceFieldInstanceIndices[InstanceIndex];
+			PendingAddOperations.Remove(InPrimitive);
+			PendingUpdateOperations.Remove(InPrimitive);
 
-			// Sanity check that scales poorly
-			if (PendingRemoveOperations.Num() < 1000)
+			if (InPrimitive->DistanceFieldInstanceIndices.Num() > 0)
 			{
-				checkSlow(!PendingRemoveOperations.Contains(RemoveIndex));
+				PendingRemoveOperations.Add(FPrimitiveRemoveInfo(InPrimitive));
 			}
 			
-			PendingRemoveOperations.Add(RemoveIndex);
+			InPrimitive->DistanceFieldInstanceIndices.Empty();
 		}
 
-		InPrimitive->DistanceFieldInstanceIndices.Empty();
-	}
-
-	if (Proxy->SupportsHeightfieldRepresentation() && Proxy->AffectsDistanceFieldLighting())
-	{
-		HeightfieldPrimitives.Remove(InPrimitive);
+		if (Proxy->SupportsHeightfieldRepresentation())
+		{
+			HeightfieldPrimitives.Remove(InPrimitive);
+		}
 	}
 }
 
diff --git a/Engine/Source/Runtime/Renderer/Private/ScenePrivate.h b/Engine/Source/Runtime/Renderer/Private/ScenePrivate.h
--- a/Engine/Source/Runtime/Renderer/Private/ScenePrivate.h
+++ b/Engine/Source/Runtime/Renderer/Private/ScenePrivate.h
@@ -505,8 +505,6 @@ public:
 	/** Distance field AO tile intersection GPU resources.  Last frame's state is not used, but they must be sized exactly to the view so stored here. */
 	class FTileIntersectionResources* AOTileIntersectionResources;
 
-	class FVPLTileIntersectionResources* VPLTileIntersectionResources;
-
 	// Is DOFHistoryRT set from Bokeh DOF?
 	bool bBokehDOFHistory;
 	bool bBokehDOFHistory2;
@@ -876,6 +874,90 @@ public:
 	int32 InstanceIndex;
 };
 
+class FPrimitiveSurfelFreeEntry
+{
+public:
+	FPrimitiveSurfelFreeEntry(int32 InOffset, int32 InNumSurfels) :
+		Offset(InOffset),
+		NumSurfels(InNumSurfels)
+	{}
+
+	FPrimitiveSurfelFreeEntry() :
+		Offset(0),
+		NumSurfels(0)
+	{}
+
+	int32 Offset;
+	int32 NumSurfels;
+};
+
+class FPrimitiveSurfelAllocation
+{
+public:
+	FPrimitiveSurfelAllocation(int32 InOffset, int32 InNumLOD0, int32 InNumSurfels, int32 InNumInstances) :
+		Offset(InOffset),
+		NumLOD0(InNumLOD0),
+		NumSurfels(InNumSurfels),
+		NumInstances(InNumInstances)
+	{}
+
+	FPrimitiveSurfelAllocation() :
+		Offset(0),
+		NumLOD0(0),
+		NumSurfels(0),
+		NumInstances(1)
+	{}
+
+	int32 GetTotalNumSurfels() const
+	{
+		return NumSurfels * NumInstances;
+	}
+
+	int32 Offset;
+	int32 NumLOD0;
+	int32 NumSurfels;
+	int32 NumInstances;
+};
+
+class FPrimitiveRemoveInfo
+{
+public:
+	FPrimitiveRemoveInfo(const FPrimitiveSceneInfo* InPrimitive) :
+		Primitive(InPrimitive),
+		DistanceFieldInstanceIndices(Primitive->DistanceFieldInstanceIndices)
+	{}
+
+	/** 
+	 * Must not be dereferenced after creation, the primitive was removed from the scene and deleted
+	 * Value of the pointer is still useful for map lookups
+	 */
+	const FPrimitiveSceneInfo* Primitive;
+
+	TArray<int32, TInlineAllocator<1>> DistanceFieldInstanceIndices;
+};
+
+class FSurfelBufferAllocator
+{
+public:
+
+	FSurfelBufferAllocator() : NumSurfelsInBuffer(0) {}
+
+	int32 GetNumSurfelsInBuffer() const { return NumSurfelsInBuffer; }
+	void AddPrimitive(const FPrimitiveSceneInfo* PrimitiveSceneInfo, int32 PrimitiveLOD0Surfels, int32 PrimitiveNumSurfels, int32 NumInstances);
+	void RemovePrimitive(const FPrimitiveSceneInfo* Primitive);
+
+	const FPrimitiveSurfelAllocation* FindAllocation(const FPrimitiveSceneInfo* Primitive)
+	{
+		return Allocations.Find(Primitive);
+	}
+
+private:
+
+	int32 NumSurfelsInBuffer;
+	TMap<const FPrimitiveSceneInfo*, FPrimitiveSurfelAllocation> Allocations;
+	TArray<FPrimitiveSurfelFreeEntry> FreeList;
+};
+
 /** Scene data used to manage distance field object buffers on the GPU. */
 class FDistanceFieldSceneData
 {
@@ -895,6 +977,19 @@ public:
 		return PendingAddOperations.Num() > 0 || PendingUpdateOperations.Num() > 0 || PendingRemoveOperations.Num() > 0;
 	}
 
+	bool HasPendingRemovePrimitive(const FPrimitiveSceneInfo* Primitive) const
+	{
+		for (int32 RemoveIndex = 0; RemoveIndex < PendingRemoveOperations.Num(); RemoveIndex++)
+		{
+			if (PendingRemoveOperations[RemoveIndex].Primitive == Primitive)
+			{
+				return true;
+			}
+		}
+
+		return false;
+	}
+
 	int32 NumObjectsInBuffer;
 	class FDistanceFieldObjectBuffers* ObjectBuffers;
 
@@ -902,10 +997,16 @@ public:
 	TArray<FPrimitiveAndInstance> PrimitiveInstanceMapping;
 	TArray<const FPrimitiveSceneInfo*> HeightfieldPrimitives;
 
+	class FSurfelBuffers* SurfelBuffers;
+	FSurfelBufferAllocator SurfelAllocations;
+
+	class FInstancedSurfelBuffers* InstancedSurfelBuffers;
+	FSurfelBufferAllocator InstancedSurfelAllocations;
+
 	/** Pending operations on the object buffers to be processed next frame. */
 	TArray<FPrimitiveSceneInfo*> PendingAddOperations;
 	TArray<FPrimitiveSceneInfo*> PendingUpdateOperations;
-	TArray<int32> PendingRemoveOperations;
+	TArray<FPrimitiveRemoveInfo> PendingRemoveOperations;
 
 	/** Used to detect atlas reallocations, since objects store UVs into the atlas and need to be updated when it changes. */
 	int32 AtlasGeneration;
diff --git a/Engine/Source/Runtime/Renderer/Private/SurfelTree.cpp b/Engine/Source/Runtime/Renderer/Private/SurfelTree.cpp
new file mode 100644
--- /dev/null
+++ b/Engine/Source/Runtime/Renderer/Private/SurfelTree.cpp
@@ -0,0 +1,354 @@
+// Copyright 1998-2014 Epic Games, Inc. All Rights Reserved.
+
+/*=============================================================================
+	SurfelTree.cpp
+=============================================================================*/
+
+#include "RendererPrivate.h"
+#include "ScenePrivate.h"
+#include "UniformBuffer.h"
+#include "ShaderParameters.h"
+#include "PostProcessing.h"
+#include "SceneFilterRendering.h"
+#include "DistanceFieldLightingShared.h"
+#include "DistanceFieldSurfaceCacheLighting.h"
+#include "DistanceFieldGlobalIllumination.h"
+#include "RHICommandList.h"
+#include "SceneUtils.h"
+#include "DistanceFieldAtlas.h"
+
+float GSurfelDensity = .05f;
+FAutoConsoleVariableRef CVarSurfelDensity(
+	TEXT("r.SurfelDensity"),
+	GSurfelDensity,
+	TEXT(""),
+	ECVF_Cheat | ECVF_RenderThreadSafe
+	);
+
+int32 GMaxSurfelsPerObject = 10000;
+FAutoConsoleVariableRef CVarMaxSurfelsPerObject(
+	TEXT("r.SurfelMaxPerObject"),
+	GMaxSurfelsPerObject,
+	TEXT(""),
+	ECVF_Cheat | ECVF_RenderThreadSafe
+	);
+
+float GSurfelLODDensityFraction = .2f;
+FAutoConsoleVariableRef CVarSurfelLODDensityFraction(
+	TEXT("r.SurfelLODDensityFraction"),
+	GSurfelLODDensityFraction,
+	TEXT(""),
+	ECVF_Cheat | ECVF_RenderThreadSafe
+	);
+
+class FComputeTriangleAreasCS : public FGlobalShader
+{
+	DECLARE_SHADER_TYPE(FComputeTriangleAreasCS,Global)
+public:
+
+	static bool ShouldCache(EShaderPlatform Platform)
+	{
+		return IsFeatureLevelSupported(Platform, ERHIFeatureLevel::SM5) && DoesPlatformSupportDistanceFieldAO(Platform);
+	}
+
+	static void ModifyCompilationEnvironment(EShaderPlatform Platform, FShaderCompilerEnvironment& OutEnvironment)
+	{
+		FGlobalShader::ModifyCompilationEnvironment(Platform,OutEnvironment);
+		OutEnvironment.SetDefine(TEXT("THREADGROUP_SIZEX"), GDistanceFieldAOTileSizeX);
+		OutEnvironment.SetDefine(TEXT("THREADGROUP_SIZEY"), GDistanceFieldAOTileSizeY);
+	}
+
+	FComputeTriangleAreasCS(const ShaderMetaType::CompiledShaderInitializerType& Initializer)
+		: FGlobalShader(Initializer)
+	{
+		NumTriangles.Bind(Initializer.ParameterMap, TEXT("NumTriangles"));
+		TriangleVertexData.Bind(Initializer.ParameterMap, TEXT("TriangleVertexData"));
+		TriangleAreas.Bind(Initializer.ParameterMap, TEXT("TriangleAreas"));
+	}
+
+	FComputeTriangleAreasCS()
+	{
+	}
+
+	void SetParameters(
+		FRHICommandList& RHICmdList, 
+		const FSceneView& View, 
+		int32 NumTrianglesValue,
+		const FUniformMeshBuffers& UniformMeshBuffers
+		)
+	{
+		FComputeShaderRHIParamRef ShaderRHI = GetComputeShader();
+		FGlobalShader::SetParameters(RHICmdList, ShaderRHI, View);
+
+		SetShaderValue(RHICmdList, ShaderRHI, NumTriangles, NumTrianglesValue);
+		SetSRVParameter(RHICmdList, ShaderRHI, TriangleVertexData, UniformMeshBuffers.TriangleDataSRV);
+
+		TriangleAreas.SetBuffer(RHICmdList, ShaderRHI, UniformMeshBuffers.TriangleAreas);
+	}
+
+	void UnsetParameters(FRHICommandList& RHICmdList)
+	{
+		FComputeShaderRHIParamRef ShaderRHI = GetComputeShader();
+		TriangleAreas.UnsetUAV(RHICmdList, ShaderRHI);
+		// RHISetStreamOutTargets doesn't unbind existing uses like render targets do 
+		SetSRVParameter(RHICmdList, ShaderRHI, TriangleVertexData, NULL);
+	}
+
+	virtual bool Serialize(FArchive& Ar)
+	{		
+		bool bShaderHasOutdatedParameters = FGlobalShader::Serialize(Ar);
+		Ar << NumTriangles;
+		Ar << TriangleVertexData;
+		Ar << TriangleAreas;
+		return bShaderHasOutdatedParameters;
+	}
+
+private:
+
+	FShaderParameter NumTriangles;
+	FShaderResourceParameter TriangleVertexData;
+	FRWShaderParameter TriangleAreas;
+};
+
+IMPLEMENT_SHADER_TYPE(,FComputeTriangleAreasCS,TEXT("SurfelTree"),TEXT("ComputeTriangleAreasCS"),SF_Compute);
+
+
+class FComputeTriangleCDFsCS : public FGlobalShader
+{
+	DECLARE_SHADER_TYPE(FComputeTriangleCDFsCS,Global)
+public:
+
+	static bool ShouldCache(EShaderPlatform Platform)
+	{
+		return IsFeatureLevelSupported(Platform, ERHIFeatureLevel::SM5) && DoesPlatformSupportDistanceFieldAO(Platform);
+	}
+
+	static void ModifyCompilationEnvironment(EShaderPlatform Platform, FShaderCompilerEnvironment& OutEnvironment)
+	{
+		FGlobalShader::ModifyCompilationEnvironment(Platform,OutEnvironment);
+		OutEnvironment.SetDefine(TEXT("THREADGROUP_SIZEX"), GDistanceFieldAOTileSizeX);
+		OutEnvironment.SetDefine(TEXT("THREADGROUP_SIZEY"), GDistanceFieldAOTileSizeY);
+	}
+
+	FComputeTriangleCDFsCS(const ShaderMetaType::CompiledShaderInitializerType& Initializer)
+		: FGlobalShader(Initializer)
+	{
+		NumTriangles.Bind(Initializer.ParameterMap, TEXT("NumTriangles"));
+		TriangleAreas.Bind(Initializer.ParameterMap, TEXT("TriangleAreas"));
+		TriangleCDFs.Bind(Initializer.ParameterMap, TEXT("TriangleCDFs"));
+	}
+
+	FComputeTriangleCDFsCS()
+	{
+	}
+
+	void SetParameters(
+		FRHICommandList& RHICmdList, 
+		const FSceneView& View, 
+		int32 NumTrianglesValue,
+		const FUniformMeshBuffers& UniformMeshBuffers
+		)
+	{
+		FComputeShaderRHIParamRef ShaderRHI = GetComputeShader();
+		FGlobalShader::SetParameters(RHICmdList, ShaderRHI, View);
+
+		SetShaderValue(RHICmdList, ShaderRHI, NumTriangles, NumTrianglesValue);
+		SetSRVParameter(RHICmdList, ShaderRHI, TriangleAreas, UniformMeshBuffers.TriangleAreas.SRV);
+
+		TriangleCDFs.SetBuffer(RHICmdList, ShaderRHI, UniformMeshBuffers.TriangleCDFs);
+	}
+
+	void UnsetParameters(FRHICommandList& RHICmdList)
+	{
+		FComputeShaderRHIParamRef ShaderRHI = GetComputeShader();
+		TriangleCDFs.UnsetUAV(RHICmdList, ShaderRHI);
+	}
+
+	virtual bool Serialize(FArchive& Ar)
+	{		
+		bool bShaderHasOutdatedParameters = FGlobalShader::Serialize(Ar);
+		Ar << NumTriangles;
+		Ar << TriangleAreas;
+		Ar << TriangleCDFs;
+		return bShaderHasOutdatedParameters;
+	}
+
+private:
+
+	FShaderParameter NumTriangles;
+	FShaderResourceParameter TriangleAreas;
+	FRWShaderParameter TriangleCDFs;
+};
+
+IMPLEMENT_SHADER_TYPE(,FComputeTriangleCDFsCS,TEXT("SurfelTree"),TEXT("ComputeTriangleCDFsCS"),SF_Compute);
+
+
+class FSampleTrianglesCS : public FGlobalShader
+{
+	DECLARE_SHADER_TYPE(FSampleTrianglesCS,Global)
+public:
+
+	static bool ShouldCache(EShaderPlatform Platform)
+	{
+		return IsFeatureLevelSupported(Platform, ERHIFeatureLevel::SM5) && DoesPlatformSupportDistanceFieldAO(Platform);
+	}
+
+	static void ModifyCompilationEnvironment(EShaderPlatform Platform, FShaderCompilerEnvironment& OutEnvironment)
+	{
+		FGlobalShader::ModifyCompilationEnvironment(Platform,OutEnvironment);
+		OutEnvironment.SetDefine(TEXT("THREADGROUP_SIZEX"), GDistanceFieldAOTileSizeX);
+		OutEnvironment.SetDefine(TEXT("THREADGROUP_SIZEY"), GDistanceFieldAOTileSizeY);
+	}
+
+	FSampleTrianglesCS(const ShaderMetaType::CompiledShaderInitializerType& Initializer)
+		: FGlobalShader(Initializer)
+	{
+		SurfelBufferParameters.Bind(Initializer.ParameterMap);
+		SurfelStartIndex.Bind(Initializer.ParameterMap, TEXT("SurfelStartIndex"));
+		NumSurfelsToGenerate.Bind(Initializer.ParameterMap, TEXT("NumSurfelsToGenerate"));
+		NumTriangles.Bind(Initializer.ParameterMap, TEXT("NumTriangles"));
+		TriangleVertexData.Bind(Initializer.ParameterMap, TEXT("TriangleVertexData"));
+		TriangleCDFs.Bind(Initializer.ParameterMap, TEXT("TriangleCDFs"));
+	}
+
+	FSampleTrianglesCS()
+	{
+	}
+
+	void SetParameters(
+		FRHICommandList& RHICmdList, 
+		const FSceneView& View, 
+		int32 SurfelStartIndexValue,
+		int32 NumSurfelsToGenerateValue,
+		int32 NumTrianglesValue,
+		const FUniformMeshBuffers& UniformMeshBuffers
+		)
+	{
+		FComputeShaderRHIParamRef ShaderRHI = GetComputeShader();
+		FGlobalShader::SetParameters(RHICmdList, ShaderRHI, View);
+
+		const FScene* Scene = (const FScene*)View.Family->Scene;
+		SurfelBufferParameters.Set(RHICmdList, ShaderRHI, *Scene->DistanceFieldSceneData.SurfelBuffers, *Scene->DistanceFieldSceneData.InstancedSurfelBuffers);
+		
+		SetShaderValue(RHICmdList, ShaderRHI, SurfelStartIndex, SurfelStartIndexValue);
+		SetShaderValue(RHICmdList, ShaderRHI, NumSurfelsToGenerate, NumSurfelsToGenerateValue);
+		SetShaderValue(RHICmdList, ShaderRHI, NumTriangles, NumTrianglesValue);
+
+		SetSRVParameter(RHICmdList, ShaderRHI, TriangleVertexData, UniformMeshBuffers.TriangleDataSRV);
+		SetSRVParameter(RHICmdList, ShaderRHI, TriangleCDFs, UniformMeshBuffers.TriangleCDFs.SRV);
+	}
+
+	void UnsetParameters(FRHICommandList& RHICmdList)
+	{
+		FComputeShaderRHIParamRef ShaderRHI = GetComputeShader();
+		SurfelBufferParameters.UnsetParameters(RHICmdList, ShaderRHI);
+		// RHISetStreamOutTargets doesn't unbind existing uses like render targets do 
+		SetSRVParameter(RHICmdList, ShaderRHI, TriangleVertexData, NULL);
+	}
+
+	virtual bool Serialize(FArchive& Ar)
+	{		
+		bool bShaderHasOutdatedParameters = FGlobalShader::Serialize(Ar);
+		Ar << SurfelBufferParameters;
+		Ar << SurfelStartIndex;
+		Ar << NumSurfelsToGenerate;
+		Ar << NumTriangles;
+		Ar << TriangleVertexData;
+		Ar << TriangleCDFs;
+		return bShaderHasOutdatedParameters;
+	}
+
+private:
+
+	FSurfelBufferParameters SurfelBufferParameters;
+	FShaderParameter SurfelStartIndex;
+	FShaderParameter NumSurfelsToGenerate;
+	FShaderParameter NumTriangles;
+	FShaderResourceParameter TriangleVertexData;
+	FShaderResourceParameter TriangleCDFs;
+};
+
+IMPLEMENT_SHADER_TYPE(,FSampleTrianglesCS,TEXT("SurfelTree"),TEXT("SampleTrianglesCS"),SF_Compute);
+
+// In float4's, must match usf
+int32 FSurfelBuffers::SurfelDataStride = 4;
+
+void ComputeNumSurfels(float BoundsSurfaceArea, int32& PrimitiveNumSurfels, int32& PrimitiveLOD0Surfels)
+{
+	//@todo - allocate based on actual triangle surface area.  
+	// Unfortunately this would make surfel count only known on the GPU which means surfel allocation would have to happen on the GPU.
+	const int32 NumSurfels = FMath::Clamp(FMath::TruncToInt(BoundsSurfaceArea * GSurfelDensity / 1000.0f), 10, GMaxSurfelsPerObject);
+	// Don't attempt to represent huge meshes
+	PrimitiveLOD0Surfels = NumSurfels == GMaxSurfelsPerObject ? 0 : NumSurfels;
+	const int32 LOD1Surfels = FMath::Clamp(FMath::TruncToInt(PrimitiveLOD0Surfels * GSurfelLODDensityFraction), 10, GMaxSurfelsPerObject);
+	PrimitiveNumSurfels = PrimitiveLOD0Surfels + LOD1Surfels;
+}
+
+void GenerateSurfelRepresentation(FRHICommandListImmediate& RHICmdList, FSceneRenderer& Renderer, FViewInfo& View, FPrimitiveSceneInfo* PrimitiveSceneInfo, const FMatrix& Instance0Transform, FPrimitiveSurfelAllocation& Allocation)
+{
+	if (Allocation.NumSurfels > 0)
+	{
+		FUniformMeshBuffers* UniformMeshBuffers = NULL;
+		const FMaterialRenderProxy* MaterialRenderProxy = NULL;
+		FUniformBufferRHIParamRef PrimitiveUniformBuffer = NULL;
+		const int32 NumUniformTriangles = FUniformMeshConverter::Convert(RHICmdList, Renderer, View, PrimitiveSceneInfo, UniformMeshBuffers, MaterialRenderProxy, PrimitiveUniformBuffer);
+
+		if (NumUniformTriangles > 0 && MaterialRenderProxy && PrimitiveUniformBuffer)
+		{
+			{
+				TShaderMapRef<FComputeTriangleAreasCS> ComputeShader(GetGlobalShaderMap(View.GetFeatureLevel()));
+
+				RHICmdList.SetComputeShader(ComputeShader->GetComputeShader());
+				ComputeShader->SetParameters(RHICmdList, View, NumUniformTriangles, *UniformMeshBuffers);
+				DispatchComputeShader(RHICmdList, *ComputeShader, FMath::DivideAndRoundUp(NumUniformTriangles, GDistanceFieldAOTileSizeX * GDistanceFieldAOTileSizeY), 1, 1);
+
+				ComputeShader->UnsetParameters(RHICmdList);
+			}
+
+			{
+				TShaderMapRef<FComputeTriangleCDFsCS> ComputeShader(GetGlobalShaderMap(View.GetFeatureLevel()));
+
+				RHICmdList.SetComputeShader(ComputeShader->GetComputeShader());
+				ComputeShader->SetParameters(RHICmdList, View, NumUniformTriangles, *UniformMeshBuffers);
+				DispatchComputeShader(RHICmdList, *ComputeShader, FMath::DivideAndRoundUp(NumUniformTriangles, GDistanceFieldAOTileSizeX * GDistanceFieldAOTileSizeY), 1, 1);
+
+				ComputeShader->UnsetParameters(RHICmdList);
+			}
+
+			{
+				TShaderMapRef<FSampleTrianglesCS> ComputeShader(GetGlobalShaderMap(View.GetFeatureLevel()));
+
+				RHICmdList.SetComputeShader(ComputeShader->GetComputeShader());
+				ComputeShader->SetParameters(RHICmdList, View, Allocation.Offset, Allocation.NumLOD0, NumUniformTriangles, *UniformMeshBuffers);
+				DispatchComputeShader(RHICmdList, *ComputeShader, FMath::DivideAndRoundUp(Allocation.NumLOD0, GDistanceFieldAOTileSizeX * GDistanceFieldAOTileSizeY), 1, 1);
+
+				ComputeShader->UnsetParameters(RHICmdList);
+			}
+
+			FUniformMeshConverter::GenerateSurfels(RHICmdList, View, PrimitiveSceneInfo, MaterialRenderProxy, PrimitiveUniformBuffer, Instance0Transform, Allocation.Offset, Allocation.NumLOD0);
+
+			const int32 NumLOD1 = Allocation.NumSurfels - Allocation.NumLOD0;
+
+			if (NumLOD1 > 0)
+			{
+				{
+					TShaderMapRef<FSampleTrianglesCS> ComputeShader(GetGlobalShaderMap(View.GetFeatureLevel()));
+
+					RHICmdList.SetComputeShader(ComputeShader->GetComputeShader());
+					ComputeShader->SetParameters(RHICmdList, View, Allocation.Offset + Allocation.NumLOD0, NumLOD1, NumUniformTriangles, *UniformMeshBuffers);
+					DispatchComputeShader(RHICmdList, *ComputeShader, FMath::DivideAndRoundUp(NumLOD1, GDistanceFieldAOTileSizeX * GDistanceFieldAOTileSizeY), 1, 1);
+
+					ComputeShader->UnsetParameters(RHICmdList);
+				}
+
+				FUniformMeshConverter::GenerateSurfels(RHICmdList, View, PrimitiveSceneInfo, MaterialRenderProxy, PrimitiveUniformBuffer, Instance0Transform, Allocation.Offset + Allocation.NumLOD0, NumLOD1);
+			}
+		}
+		else
+		{
+			Allocation.NumSurfels = 0;
+			Allocation.NumLOD0 = 0;
+			Allocation.NumInstances = 0;
+		}
+	}
+}
\ No newline at end of file
diff --git a/Engine/Source/Runtime/ShaderCore/Private/Shader.cpp b/Engine/Source/Runtime/ShaderCore/Private/Shader.cpp
--- a/Engine/Source/Runtime/ShaderCore/Private/Shader.cpp
+++ b/Engine/Source/Runtime/ShaderCore/Private/Shader.cpp
@@ -761,6 +761,7 @@ FShader::FShader(const CompiledShaderInitializerType& Initializer):
 
 FShader::~FShader()
 {
+	check(!IsInitialized());
 	check(Canary == ShaderMagic_Uninitialized || Canary == ShaderMagic_Initialized);
 	check(NumRefs == 0);
 	Canary = 0;
@@ -897,6 +898,8 @@ void FShader::Release()
 		// Deregister the shader now to eliminate references to it by the type's ShaderIdMap
 		Deregister();
 
+		BeginReleaseResource(this);
+
 		BeginCleanup(this);
 	}
 }
@@ -945,6 +948,41 @@ void FShader::FinishCleanup()
 }
 
 
+void FShader::InitRHI()
+{
+	checkf(Resource->GetCode().Num() > 0, TEXT("FShader::InitRHI was called with empty bytecode, which can happen if the resource is initialized multiple times on platforms with no editor data."));
+
+	// we can't have this called on the wrong platform's shaders
+	if (!ArePlatformsCompatible(GMaxRHIShaderPlatform, (EShaderPlatform)Target.Platform))
+ 	{
+ 		if (FPlatformProperties::RequiresCookedData())
+ 		{
+ 			UE_LOG(LogShaders, Fatal, TEXT("FShader::InitRHI got platform %s but it is not compatible with %s"), 
+				*LegacyShaderPlatformToShaderFormat((EShaderPlatform)Target.Platform).ToString(), *LegacyShaderPlatformToShaderFormat(GMaxRHIShaderPlatform).ToString());
+ 		}
+ 		return;
+ 	}
+
+	if (Target.Frequency == SF_Geometry)
+	{
+		FStreamOutElementList ElementList;
+		TArray<uint32> StreamStrides;
+		int32 RasterizedStream = -1;
+		GetStreamOutElements(ElementList, StreamStrides, RasterizedStream);
+
+		if (ElementList.Num() > 0)
+		{
+			GeometryShaderWithStreamOutput = RHICreateGeometryShaderWithStreamOutput(Resource->GetCode(), ElementList, StreamStrides.Num(), StreamStrides.GetData(), RasterizedStream);
+		}
+	}
+}
+
+
+void FShader::ReleaseRHI()
+{
+	GeometryShaderWithStreamOutput.SafeRelease();
+}
+
 void FShader::VerifyBoundUniformBufferParameters()
 {
 	// Support being called on a NULL pointer
diff --git a/Engine/Source/Runtime/ShaderCore/Public/Shader.h b/Engine/Source/Runtime/ShaderCore/Public/Shader.h
--- a/Engine/Source/Runtime/ShaderCore/Public/Shader.h
+++ b/Engine/Source/Runtime/ShaderCore/Public/Shader.h
@@ -129,6 +129,11 @@ public:
 		return Code.GetAllocatedSize() + sizeof(FShaderResource);
 	}
 
+	const TArray<uint8>& GetCode() const
+	{
+		return Code;
+	}
+
 	// FRenderResource interface.
 	virtual void InitRHI();
 	virtual void ReleaseRHI();
@@ -381,7 +386,7 @@ public:
 class FMaterial;
 
 /** A compiled shader and its parameter bindings. */
-class SHADERCORE_API FShader : public FDeferredCleanupInterface
+class SHADERCORE_API FShader : public FRenderResource, public FDeferredCleanupInterface
 {
 	friend class FShaderType;
 public:
@@ -477,7 +482,12 @@ public:
 	/** @return the shader's geometry shader */
 	const FGeometryShaderRHIRef& GetGeometryShader()
 	{
-		return Resource->GetGeometryShader();
+		if (!IsInitialized())
+		{
+			InitResourceFromPossiblyParallelRendering();
+		}
+
+		return GeometryShaderWithStreamOutput ? GeometryShaderWithStreamOutput : Resource->GetGeometryShader();
 	}
 	/** @return the shader's compute shader */
 	const FComputeShaderRHIRef& GetComputeShader()
@@ -529,6 +539,13 @@ public:
 	// FDeferredCleanupInterface implementation.
 	virtual void FinishCleanup();
 
+	// FRenderResource interface.
+	virtual void InitRHI();
+	virtual void ReleaseRHI();
+
+	/** Implement for geometry shaders that want to use stream out. */
+	virtual void GetStreamOutElements(FStreamOutElementList& ElementList, TArray<uint32>& StreamStrides, int32& RasterizedStream) {}
+
 	void InitializeResource()
 	{
 		BeginInitResource(Resource);
@@ -624,6 +641,8 @@ private:
 	/** Reference to the shader resource, which stores the compiled bytecode and the RHI shader resource. */
 	TRefCountPtr<FShaderResource> Resource;
 
+	FGeometryShaderRHIRef GeometryShaderWithStreamOutput;
+
 	/** Hash of the material shader map this shader belongs to, stored so that an FShaderId can be constructed from this shader. */
 	FSHAHash MaterialShaderMapHash;
 
diff --git a/Engine/Source/Runtime/ShaderCore/Public/VertexFactory.h b/Engine/Source/Runtime/ShaderCore/Public/VertexFactory.h
--- a/Engine/Source/Runtime/ShaderCore/Public/VertexFactory.h
+++ b/Engine/Source/Runtime/ShaderCore/Public/VertexFactory.h
@@ -171,7 +171,7 @@ public:
 		FString VertexFactoryIncludeString = FString::Printf( TEXT("#include \"%s.usf\""), GetShaderFilename() );
 		OutEnvironment.IncludeFileNameToContentsMap.Add(TEXT("VertexFactory.usf"), VertexFactoryIncludeString);
 
-		OutEnvironment.SetDefine(TEXT("MESH_MATERIAL_SHADER"), 1);
+		OutEnvironment.SetDefine(TEXT("HAS_PRIMITIVE_UNIFORM_BUFFER"), 1);
 
 		(*ModifyCompilationEnvironmentRef)(Platform, Material, OutEnvironment);
 	}
