Meshes whose distance fields are generated as two sided are expanded during ray traced shadow casting, helps trees match up with CSM much better
Ray traced distance field shadows can be done at full or half res with bilateral upsample.  Half res has some unresolved flickering artifacts from TemporalAA.

diff --git a/Engine/Shaders/DistanceFieldShadowing.usf b/Engine/Shaders/DistanceFieldShadowing.usf
--- a/Engine/Shaders/DistanceFieldShadowing.usf
+++ b/Engine/Shaders/DistanceFieldShadowing.usf
@@ -303,6 +303,7 @@
 float RayStartOffsetDepthScale;
 float2 TanLightAngleAndNormalThreshold;
 int4 ScissorRectMinAndSize;
+float TwoSidedMeshDistanceBias;
 
 /** Min and Max depth for this tile. */
 groupshared uint IntegerTileMinZ;
@@ -319,48 +320,16 @@
 #define MAX_INTERSECTING_OBJECTS 256
 groupshared uint IntersectingObjectIndices[MAX_INTERSECTING_OBJECTS * 2];
 
-[numthreads(THREADGROUP_SIZEX, THREADGROUP_SIZEY, 1)]
-void DistanceFieldShadowingCS(
-	uint3 GroupId : SV_GroupID,
-	uint3 DispatchThreadId : SV_DispatchThreadID,
-    uint3 GroupThreadId : SV_GroupThreadID) 
+void CullObjectsToTileWithGather(
+	float SceneDepth, 
+	uint ThreadIndex, 
+	uint2 GroupId, 
+	float TraceDistance, 
+	float MinDepth, 
+	out uint NumIntersectingObjects, 
+	out bool bTileShouldComputeShadowing,
+	out uint GroupIndex)
 {
-	uint ThreadIndex = GroupThreadId.y * THREADGROUP_SIZEX + GroupThreadId.x;
-
-	float2 ScreenUV = float2((DispatchThreadId.xy * DOWNSAMPLE_FACTOR + ScissorRectMinAndSize.xy + .5f) * View.ViewSizeAndSceneTexelSize.zw);
-	float2 ScreenPosition = (ScreenUV.xy - View.ScreenPositionScaleBias.wz) / View.ScreenPositionScaleBias.xy;
-
-	float SceneDepth = CalcSceneDepth(ScreenUV);
-	float4 HomogeneousWorldPosition = mul(float4(ScreenPosition * SceneDepth, SceneDepth, 1), View.ScreenToWorld);
-	float3 OpaqueWorldPosition = HomogeneousWorldPosition.xyz / HomogeneousWorldPosition.w;
-
-	// Distance for directional lights to trace
-	float TraceDistance = 10000;
-
-	//@todo - take advantage of CSM covering near pixels
-	//@todo - handle invalid tile depth ranges due to MinDepth
-	float MinDepth = 0;
-
-	uint NumIntersectingObjects = GetCulledNumObjects();
-	bool bTileShouldComputeShadowing = SceneDepth > MinDepth;
-
-#define USE_CULLING 1
-#if USE_CULLING
-
-#if SCATTER_TILE_CULLING
-
-	// Transform into shadow space
-	float4 HomogeneousShadowPosition = mul(float4(OpaqueWorldPosition, 1), WorldToShadow);
-	float2 NormalizedShadowPosition = HomogeneousShadowPosition.xy * .5f + .5f;
-	NormalizedShadowPosition.y = 1 - NormalizedShadowPosition.y;
-	// Quantize the shadow position to get our tile position
-	uint2 TilePosition = (uint2)(NormalizedShadowPosition * ShadowTileListGroupSize);
-	// Fetch the tile head information
-	uint2 TileHead = GetShadowTileHead(TilePosition);
-	NumIntersectingObjects = TileHead.y;
-
-#else
-
 	// Initialize per-tile variables
     if (ThreadIndex == 0) 
 	{
@@ -440,8 +409,6 @@
 	float3 WorldGroup1Center = mul(float4(ViewGroup1Center, 1), View.ViewToTranslatedWorld).xyz - View.PreViewTranslation;
 	float Group1BoundingRadius = length(ViewGroup1Center - ViewTileMax2);
 
-	bool bCullTile = true;
-
 #if POINT_LIGHT
 	float3 LightVector0 = LightPositionAndInvRadius.xyz - WorldGroup0Center;
 	float LightVector0Length = length(LightVector0);
@@ -453,7 +420,7 @@
 	float RayLength1 = LightVector1Length;
 
 	// Don't operate on tiles completely outside of the light's influence
-	bCullTile = bTileShouldComputeShadowing = LightVector0Length < 1.0f / LightPositionAndInvRadius.w + Group0BoundingRadius
+	bTileShouldComputeShadowing = LightVector0Length < 1.0f / LightPositionAndInvRadius.w + Group0BoundingRadius
 		|| LightVector1Length < 1.0f / LightPositionAndInvRadius.w + Group1BoundingRadius;
 
 #else
@@ -461,10 +428,13 @@
 	float3 LightDirection1 = LightDirection;
 	float RayLength0 = TraceDistance;
 	float RayLength1 = TraceDistance;
+
+	bTileShouldComputeShadowing = SceneDepth > MinDepth;
+
 #endif
 
 	BRANCH
-	if (bCullTile)
+	if (bTileShouldComputeShadowing)
 	{
 		uint NumCulledObjects = GetCulledNumObjects();
 
@@ -499,8 +469,57 @@
 
 	GroupMemoryBarrierWithGroupSync();
 
-	uint GroupIndex = SceneDepth > MaxTileZ2 ? 1 : 0;
+	GroupIndex = SceneDepth > MaxTileZ2 ? 1 : 0;
 	NumIntersectingObjects = GroupIndex == 0 ? TileNumObjects0 : TileNumObjects1;
+}
+
+uint DownsampleFactor;
+
+[numthreads(THREADGROUP_SIZEX, THREADGROUP_SIZEY, 1)]
+void DistanceFieldShadowingCS(
+	uint3 GroupId : SV_GroupID,
+	uint3 DispatchThreadId : SV_DispatchThreadID,
+    uint3 GroupThreadId : SV_GroupThreadID) 
+{
+	uint ThreadIndex = GroupThreadId.y * THREADGROUP_SIZEX + GroupThreadId.x;
+
+	float2 ScreenUV = float2((DispatchThreadId.xy * DownsampleFactor + ScissorRectMinAndSize.xy + .5f) * View.ViewSizeAndSceneTexelSize.zw);
+	float2 ScreenPosition = (ScreenUV.xy - View.ScreenPositionScaleBias.wz) / View.ScreenPositionScaleBias.xy;
+
+	float SceneDepth = CalcSceneDepth(ScreenUV);
+	float4 HomogeneousWorldPosition = mul(float4(ScreenPosition * SceneDepth, SceneDepth, 1), View.ScreenToWorld);
+	float3 OpaqueWorldPosition = HomogeneousWorldPosition.xyz / HomogeneousWorldPosition.w;
+
+	// Distance for directional lights to trace
+	float TraceDistance = 10000;
+
+	//@todo - take advantage of CSM covering near pixels
+	//@todo - handle invalid tile depth ranges due to MinDepth
+	float MinDepth = 0;
+
+	uint NumIntersectingObjects = GetCulledNumObjects();
+	bool bTileShouldComputeShadowing = SceneDepth > MinDepth;
+
+#define USE_CULLING 1
+#if USE_CULLING
+
+#if SCATTER_TILE_CULLING
+
+	// Transform into shadow space
+	float4 HomogeneousShadowPosition = mul(float4(OpaqueWorldPosition, 1), WorldToShadow);
+	float2 NormalizedShadowPosition = HomogeneousShadowPosition.xy * .5f + .5f;
+	NormalizedShadowPosition.y = 1 - NormalizedShadowPosition.y;
+	// Quantize the shadow position to get our tile position
+	uint2 TilePosition = (uint2)(NormalizedShadowPosition * ShadowTileListGroupSize);
+	// Fetch the tile head information
+	uint2 TileHead = GetShadowTileHead(TilePosition);
+	NumIntersectingObjects = TileHead.y;
+
+#else
+
+	uint GroupIndex;
+	CullObjectsToTileWithGather(SceneDepth, ThreadIndex, GroupId.xy, TraceDistance, MinDepth, NumIntersectingObjects, bTileShouldComputeShadowing, GroupIndex);
+
 #endif
 #endif // USE_CULLING
 
@@ -566,7 +585,8 @@
 				{
 					float3 LocalPositionExtent = LoadObjectLocalPositionExtent(ObjectIndex);
 					float4x4 WorldToVolume = LoadObjectWorldToVolume(ObjectIndex);
-					float4 UVScaleAndVolumeScale = LoadObjectUVScale(ObjectIndex);
+					bool bGeneratedAsTwoSided;
+					float4 UVScaleAndVolumeScale = LoadObjectUVScale(ObjectIndex, bGeneratedAsTwoSided);
 					float3 UVAdd = LoadObjectUVAdd(ObjectIndex);
 
 					float3 VolumeRayStart = mul(float4(WorldRayStart, 1), WorldToVolume).xyz;
@@ -574,11 +594,14 @@
 					float3 VolumeRayDirection = VolumeRayEnd - VolumeRayStart;
 					float VolumeRayLength = length(VolumeRayDirection);
 					VolumeRayDirection /= VolumeRayLength;
-					float VolumeMinSphereRadius = MinSphereRadius / UVScaleAndVolumeScale.w;
-					float VolumeMaxSphereRadius = MaxSphereRadius / UVScaleAndVolumeScale.w;
 
+					float WorldToVolumeScale = 1.0f / UVScaleAndVolumeScale.w;
+					float VolumeMinSphereRadius = MinSphereRadius * WorldToVolumeScale;
+					float VolumeMaxSphereRadius = MaxSphereRadius * WorldToVolumeScale;
+					float VolumeTwoSidedMeshDistanceBias = TwoSidedMeshDistanceBias * WorldToVolumeScale;
+
 					// Expand the intersection box by the radius of the cone at the distance of the object along the cone
-					float LocalConeRadiusAtObject = min(TanLightAngle * max(ObjectCenterDistanceAlongRay, 0) / UVScaleAndVolumeScale.w, VolumeMaxSphereRadius);
+					float LocalConeRadiusAtObject = min(TanLightAngle * max(ObjectCenterDistanceAlongRay, 0) * WorldToVolumeScale, VolumeMaxSphereRadius);
 
 					float2 IntersectionTimes = LineBoxIntersect(VolumeRayStart, VolumeRayEnd, -LocalPositionExtent - LocalConeRadiusAtObject, LocalPositionExtent + LocalConeRadiusAtObject);
 
@@ -603,6 +626,12 @@
 							float3 VolumeUV = DistanceFieldVolumePositionToUV(ClampedSamplePosition, UVScaleAndVolumeScale.xyz, UVAdd);
 							float DistanceField = Texture3DSampleLevel(DistanceFieldTexture, DistanceFieldSampler, VolumeUV, 0).x + DistanceToClamped;
 
+							FLATTEN
+							if (bGeneratedAsTwoSided)
+							{
+								DistanceField -= VolumeTwoSidedMeshDistanceBias;
+							}
+
 							MinDistance = min(MinDistance, DistanceField);
 							float SphereRadius = clamp(TanLightAngle * SampleRayTime, VolumeMinSphereRadius, VolumeMaxSphereRadius);
 
@@ -656,7 +685,7 @@
 	float2 DistanceFieldUVs = UVAndScreenPos.xy - ScissorRectMinAndSize.xy * View.ViewSizeAndSceneTexelSize.zw;
 
 #define BILATERAL_UPSAMPLE 1
-#if BILATERAL_UPSAMPLE
+#if BILATERAL_UPSAMPLE && UPSAMPLE_REQUIRED
 	float2 LowResBufferSize = floor(View.RenderTargetSize / DOWNSAMPLE_FACTOR);
 	float2 LowResTexelSize = 1.0f / LowResBufferSize;
 	float2 Corner00UV = floor(DistanceFieldUVs * LowResBufferSize - .5f) / LowResBufferSize + .5f * LowResTexelSize;
diff --git a/Engine/Source/Runtime/Renderer/Private/DistanceFieldShadowing.cpp b/Engine/Source/Runtime/Renderer/Private/DistanceFieldShadowing.cpp
--- a/Engine/Source/Runtime/Renderer/Private/DistanceFieldShadowing.cpp
+++ b/Engine/Source/Runtime/Renderer/Private/DistanceFieldShadowing.cpp
@@ -24,6 +24,14 @@ FAutoConsoleVariableRef CVarDistanceFieldShadowing(
 	ECVF_Scalability | ECVF_RenderThreadSafe
 	);
 
+int32 GFullResolutionDFShadowing = 0;
+FAutoConsoleVariableRef CVarFullResolutionDFShadowing(
+	TEXT("r.DFFullResolution"),
+	GFullResolutionDFShadowing,
+	TEXT("1 = full resolution distance field shadowing, 0 = half resolution with bilateral upsample."),
+	ECVF_Scalability | ECVF_RenderThreadSafe
+	);
+
 int32 GShadowScatterTileCulling = 1;
 FAutoConsoleVariableRef CVarShadowScatterTileCulling(
 	TEXT("r.DFShadowScatterTileCulling"),
@@ -40,6 +48,24 @@ FAutoConsoleVariableRef CVarShadowWorldTileSize(
 	ECVF_Cheat | ECVF_RenderThreadSafe
 	);
 
+float GTwoSidedMeshDistanceBias = 4;
+FAutoConsoleVariableRef CVarTwoSidedMeshDistanceBias(
+	TEXT("r.DFTwoSidedMeshDistanceBias"),
+	GTwoSidedMeshDistanceBias,
+	TEXT("World space amount to expand distance field representations of two sided meshes.  This is useful to get tree shadows to match up with standard shadow mapping."),
+	ECVF_Cheat | ECVF_RenderThreadSafe
+	);
+
+int32 GetDFShadowDownsampleFactor()
+{
+	return GFullResolutionDFShadowing ? 1 : GAODownsampleFactor;
+}
+
+FIntPoint GetBufferSizeForDFShadows()
+{
+	return FIntPoint::DivideAndRoundDown(GSceneRenderTargets.GetBufferSizeXY(), GetDFShadowDownsampleFactor());
+}
+
 TGlobalResource<FDistanceFieldObjectBufferResource> GShadowCulledObjectBuffers;
 
 class FCullObjectsForShadowCS : public FGlobalShader
@@ -283,13 +309,6 @@ public:
 		return IsFeatureLevelSupported(Platform, ERHIFeatureLevel::SM5) && DoesPlatformSupportDistanceFieldShadowing(Platform);
 	}
 
-	static void ModifyCompilationEnvironment(EShaderPlatform Platform, FShaderCompilerEnvironment& OutEnvironment)
-	{
-		FGlobalShader::ModifyCompilationEnvironment(Platform,OutEnvironment);
-		
-		OutEnvironment.SetDefine(TEXT("DOWNSAMPLE_FACTOR"), GAODownsampleFactor);
-	}
-
 	/** Default constructor. */
 	FShadowObjectCullPS() {}
 
@@ -375,7 +394,6 @@ public:
 
 	static void ModifyCompilationEnvironment(EShaderPlatform Platform, FShaderCompilerEnvironment& OutEnvironment)
 	{
-		OutEnvironment.SetDefine(TEXT("DOWNSAMPLE_FACTOR"), GAODownsampleFactor);
 		OutEnvironment.SetDefine(TEXT("THREADGROUP_SIZEX"), GDistanceFieldAOTileSizeX);
 		OutEnvironment.SetDefine(TEXT("THREADGROUP_SIZEY"), GDistanceFieldAOTileSizeY);
 		OutEnvironment.SetDefine(TEXT("SCATTER_TILE_CULLING"), ShadowingType == DFS_DirectionalLightScatterTileCulling ? TEXT("1") : TEXT("0"));
@@ -403,6 +421,8 @@ public:
 		ShadowTileArrayData.Bind(Initializer.ParameterMap, TEXT("ShadowTileArrayData"));
 		ShadowTileListGroupSize.Bind(Initializer.ParameterMap, TEXT("ShadowTileListGroupSize"));
 		WorldToShadow.Bind(Initializer.ParameterMap, TEXT("WorldToShadow"));
+		TwoSidedMeshDistanceBias.Bind(Initializer.ParameterMap, TEXT("TwoSidedMeshDistanceBias"));
+		DownsampleFactor.Bind(Initializer.ParameterMap, TEXT("DownsampleFactor"));
 	}
 
 	void SetParameters(
@@ -461,6 +481,10 @@ public:
 
 		FMatrix WorldToShadowMatrixValue = FTranslationMatrix(ProjectedShadowInfo->PreShadowTranslation) * ProjectedShadowInfo->SubjectAndReceiverMatrix;
 		SetShaderValue(RHICmdList, ShaderRHI, WorldToShadow, WorldToShadowMatrixValue);
+
+		SetShaderValue(RHICmdList, ShaderRHI, TwoSidedMeshDistanceBias, GTwoSidedMeshDistanceBias);
+
+		SetShaderValue(RHICmdList, ShaderRHI, DownsampleFactor, GetDFShadowDownsampleFactor());
 	}
 
 	void UnsetParameters(FRHICommandList& RHICmdList)
@@ -486,6 +510,8 @@ public:
 		Ar << ShadowTileArrayData;
 		Ar << ShadowTileListGroupSize;
 		Ar << WorldToShadow;
+		Ar << TwoSidedMeshDistanceBias;
+		Ar << DownsampleFactor;
 		return bShaderHasOutdatedParameters;
 	}
 
@@ -505,15 +531,18 @@ private:
 	FShaderResourceParameter ShadowTileArrayData;
 	FShaderParameter ShadowTileListGroupSize;
 	FShaderParameter WorldToShadow;
+	FShaderParameter TwoSidedMeshDistanceBias;
+	FShaderParameter DownsampleFactor;
 };
 
 IMPLEMENT_SHADER_TYPE(template<>,TDistanceFieldShadowingCS<DFS_DirectionalLightScatterTileCulling>,TEXT("DistanceFieldShadowing"),TEXT("DistanceFieldShadowingCS"),SF_Compute);
 IMPLEMENT_SHADER_TYPE(template<>,TDistanceFieldShadowingCS<DFS_DirectionalLightTiledCulling>,TEXT("DistanceFieldShadowing"),TEXT("DistanceFieldShadowingCS"),SF_Compute);
 IMPLEMENT_SHADER_TYPE(template<>,TDistanceFieldShadowingCS<DFS_PointLightTiledCulling>,TEXT("DistanceFieldShadowing"),TEXT("DistanceFieldShadowingCS"),SF_Compute);
 
-class FDistanceFieldShadowingUpsamplePS : public FGlobalShader
+template<bool bUpsampleRequired>
+class TDistanceFieldShadowingUpsamplePS : public FGlobalShader
 {
-	DECLARE_SHADER_TYPE(FDistanceFieldShadowingUpsamplePS, Global);
+	DECLARE_SHADER_TYPE(TDistanceFieldShadowingUpsamplePS, Global);
 public:
 
 	static bool ShouldCache(EShaderPlatform Platform)
@@ -524,13 +553,14 @@ public:
 	static void ModifyCompilationEnvironment(EShaderPlatform Platform, FShaderCompilerEnvironment& OutEnvironment)
 	{
 		OutEnvironment.SetDefine(TEXT("DOWNSAMPLE_FACTOR"), GAODownsampleFactor);
+		OutEnvironment.SetDefine(TEXT("UPSAMPLE_REQUIRED"), bUpsampleRequired ? TEXT("1") : TEXT("0"));
 	}
 
 	/** Default constructor. */
-	FDistanceFieldShadowingUpsamplePS() {}
+	TDistanceFieldShadowingUpsamplePS() {}
 
 	/** Initialization constructor. */
-	FDistanceFieldShadowingUpsamplePS(const ShaderMetaType::CompiledShaderInitializerType& Initializer)
+	TDistanceFieldShadowingUpsamplePS(const ShaderMetaType::CompiledShaderInitializerType& Initializer)
 		: FGlobalShader(Initializer)
 	{
 		DeferredParameters.Bind(Initializer.ParameterMap);
@@ -570,7 +600,8 @@ private:
 	FShaderParameter ScissorRectMinAndSize;
 };
 
-IMPLEMENT_SHADER_TYPE(,FDistanceFieldShadowingUpsamplePS,TEXT("DistanceFieldShadowing"),TEXT("DistanceFieldShadowingUpsamplePS"),SF_Pixel);
+IMPLEMENT_SHADER_TYPE(template<>,TDistanceFieldShadowingUpsamplePS<true>,TEXT("DistanceFieldShadowing"),TEXT("DistanceFieldShadowingUpsamplePS"),SF_Pixel);
+IMPLEMENT_SHADER_TYPE(template<>,TDistanceFieldShadowingUpsamplePS<false>,TEXT("DistanceFieldShadowing"),TEXT("DistanceFieldShadowingUpsamplePS"),SF_Pixel);
 
 void CullDistanceFieldObjectsForLight(
 	FRHICommandListImmediate& RHICmdList,
@@ -774,7 +805,7 @@ void FProjectedShadowInfo::RenderRayTracedDistanceFieldProjection(FRHICommandLis
 			TRefCountPtr<IPooledRenderTarget> RayTracedShadowsRT;
 
 			{
-				const FIntPoint BufferSize = GetBufferSizeForAO();
+				const FIntPoint BufferSize = GetBufferSizeForDFShadows();
 				FPooledRenderTargetDesc Desc(FPooledRenderTargetDesc::Create2DDesc(BufferSize, PF_G16R16F, TexCreate_None, TexCreate_RenderTargetable | TexCreate_UAV, false));
 				GRenderTargetPool.FindFreeElement(Desc, RayTracedShadowsRT, TEXT("RayTracedShadows"));
 			}
@@ -787,8 +818,8 @@ void FProjectedShadowInfo::RenderRayTracedDistanceFieldProjection(FRHICommandLis
 					ScissorRect = View.ViewRect;
 				}
 
-				uint32 GroupSizeX = FMath::DivideAndRoundUp(ScissorRect.Size().X / GAODownsampleFactor, GDistanceFieldAOTileSizeX);
-				uint32 GroupSizeY = FMath::DivideAndRoundUp(ScissorRect.Size().Y / GAODownsampleFactor, GDistanceFieldAOTileSizeY);
+				uint32 GroupSizeX = FMath::DivideAndRoundUp(ScissorRect.Size().X / GetDFShadowDownsampleFactor(), GDistanceFieldAOTileSizeX);
+				uint32 GroupSizeY = FMath::DivideAndRoundUp(ScissorRect.Size().Y / GetDFShadowDownsampleFactor(), GDistanceFieldAOTileSizeY);
 
 				{
 					SCOPED_DRAW_EVENT(RHICmdList, RayTraceShadows);
@@ -846,23 +877,36 @@ void FProjectedShadowInfo::RenderRayTracedDistanceFieldProjection(FRHICommandLis
 				}
 
 				TShaderMapRef<FPostProcessVS> VertexShader(View.ShaderMap);
-				TShaderMapRef<FDistanceFieldShadowingUpsamplePS> PixelShader(View.ShaderMap);
 
-				static FGlobalBoundShaderState BoundShaderState;
+				if (GFullResolutionDFShadowing)
+				{
+					TShaderMapRef<TDistanceFieldShadowingUpsamplePS<false> > PixelShader(View.ShaderMap);
 
-				SetGlobalBoundShaderState(RHICmdList, View.GetFeatureLevel(), BoundShaderState, GFilterVertexDeclaration.VertexDeclarationRHI, *VertexShader, *PixelShader);
+					static FGlobalBoundShaderState BoundShaderState;
+					SetGlobalBoundShaderState(RHICmdList, View.GetFeatureLevel(), BoundShaderState, GFilterVertexDeclaration.VertexDeclarationRHI, *VertexShader, *PixelShader);
 
-				VertexShader->SetParameters(RHICmdList, View);
-				PixelShader->SetParameters(RHICmdList, View, this, ScissorRect, RayTracedShadowsRT);
+					VertexShader->SetParameters(RHICmdList, View);
+					PixelShader->SetParameters(RHICmdList, View, this, ScissorRect, RayTracedShadowsRT);
+				}
+				else
+				{
+					TShaderMapRef<TDistanceFieldShadowingUpsamplePS<true> > PixelShader(View.ShaderMap);
+
+					static FGlobalBoundShaderState BoundShaderState;
+					SetGlobalBoundShaderState(RHICmdList, View.GetFeatureLevel(), BoundShaderState, GFilterVertexDeclaration.VertexDeclarationRHI, *VertexShader, *PixelShader);
+
+					VertexShader->SetParameters(RHICmdList, View);
+					PixelShader->SetParameters(RHICmdList, View, this, ScissorRect, RayTracedShadowsRT);
+				}
 
 				DrawRectangle( 
 					RHICmdList,
 					0, 0, 
 					ScissorRect.Width(), ScissorRect.Height(),
-					ScissorRect.Min.X / GAODownsampleFactor, ScissorRect.Min.Y / GAODownsampleFactor, 
-					ScissorRect.Width() / GAODownsampleFactor, ScissorRect.Height() / GAODownsampleFactor,
+					ScissorRect.Min.X / GetDFShadowDownsampleFactor(), ScissorRect.Min.Y / GetDFShadowDownsampleFactor(), 
+					ScissorRect.Width() / GetDFShadowDownsampleFactor(), ScissorRect.Height() / GetDFShadowDownsampleFactor(),
 					FIntPoint(ScissorRect.Width(), ScissorRect.Height()),
-					GetBufferSizeForAO(),
+					GetBufferSizeForDFShadows(),
 					*VertexShader);
 			}
 		}
