Improve precision in WorldToView and WorldToClip matrices when very far from the origin.

diff --git a/Engine/Source/Editor/UnrealEd/Private/EditorExporters.cpp b/Engine/Source/Editor/UnrealEd/Private/EditorExporters.cpp
--- a/Engine/Source/Editor/UnrealEd/Private/EditorExporters.cpp
+++ b/Engine/Source/Editor/UnrealEd/Private/EditorExporters.cpp
@@ -2697,7 +2697,8 @@ namespace MaterialExportUtils
 	void RenderSceneToTexture(
 		FSceneInterface* Scene,
 		const FName& VisualizationMode, 
-		const FMatrix& ViewMatrix, 
+		const FVector& ViewOrigin,
+		const FMatrix& ViewRotationMatrix, 
 		const FMatrix& ProjectionMatrix,  
 		const TSet<FPrimitiveComponentId>& HiddenPrimitives, 
 		FIntPoint TargetSize,
@@ -2726,7 +2727,8 @@ namespace MaterialExportUtils
 		ViewInitOptions.SetViewRectangle(FIntRect(0, 0, TargetSize.X, TargetSize.Y));
 		ViewInitOptions.ViewFamily = &ViewFamily;
 		ViewInitOptions.HiddenPrimitives = HiddenPrimitives;
-		ViewInitOptions.ViewMatrix = ViewMatrix;
+		ViewInitOptions.ViewOrigin = ViewOrigin;
+		ViewInitOptions.ViewRotationMatrix = ViewRotationMatrix;
 		ViewInitOptions.ProjectionMatrix = ProjectionMatrix;
 		
 		FSceneView* NewView = new FSceneView(ViewInitOptions);
@@ -2758,9 +2760,9 @@ namespace MaterialExportUtils
 		FVector LandscapeCenter = InLandscape->GetTransform().TransformPosition(MidPoint);
 		FVector LandscapeExtent = FVector(LandscapeRect.Size(), 0.f)*InLandscape->GetActorScale()*0.5f; 
 
-		FMatrix ViewMatrix = FTranslationMatrix(-LandscapeCenter);
-		ViewMatrix*= FInverseRotationMatrix(InLandscape->GetActorRotation());
-		ViewMatrix*= FMatrix(FPlane(1,	0,	0,	0),
+		FVector ViewOrigin = LandscapeCenter;
+		FMatrix ViewRotationMatrix = FInverseRotationMatrix(InLandscape->GetActorRotation());
+		ViewRotationMatrix*= FMatrix(FPlane(1,	0,	0,	0),
 							 FPlane(0,	-1,	0,	0),
 							 FPlane(0,	0,	-1,	0),
 							 FPlane(0,	0,	0,	1));
@@ -2780,7 +2782,7 @@ namespace MaterialExportUtils
 		{
 			static const FName BaseColorName("BaseColor");
 			const float BaseColorGamma = 1.0f; // BaseColor material already do transformation to gamma space
-			RenderSceneToTexture(Scene, BaseColorName, ViewMatrix, ProjectionMatrix, HiddenPrimitives, 
+			RenderSceneToTexture(Scene, BaseColorName, ViewOrigin, ViewRotationMatrix, ProjectionMatrix, HiddenPrimitives, 
 				OutFlattenMaterial.DiffuseSize, BaseColorGamma, OutFlattenMaterial.DiffuseSamples);
 		}
 
@@ -2791,7 +2793,7 @@ namespace MaterialExportUtils
 		{
 			static const FName WorldNormalName("WorldNormal");
 			const float NormalColorGamma = 1.0f; // Dump normal texture in linear space
-			RenderSceneToTexture(Scene, WorldNormalName, ViewMatrix, ProjectionMatrix, HiddenPrimitives, 
+			RenderSceneToTexture(Scene, WorldNormalName, ViewOrigin, ViewRotationMatrix, ProjectionMatrix, HiddenPrimitives, 
 				OutFlattenMaterial.NormalSize, NormalColorGamma, OutFlattenMaterial.NormalSamples);
 		}
 
@@ -2801,7 +2803,7 @@ namespace MaterialExportUtils
 		{
 			static const FName MetallicName("Metallic");
 			const float MetallicColorGamma = 1.0f; // Dump metallic texture in linear space
-			RenderSceneToTexture(Scene, MetallicName, ViewMatrix, ProjectionMatrix, HiddenPrimitives, 
+			RenderSceneToTexture(Scene, MetallicName, ViewOrigin, ViewRotationMatrix, ProjectionMatrix, HiddenPrimitives, 
 				OutFlattenMaterial.MetallicSize, MetallicColorGamma, OutFlattenMaterial.MetallicSamples);
 		}
 
@@ -2811,7 +2813,7 @@ namespace MaterialExportUtils
 		{
 			static const FName RoughnessName("Roughness");
 			const float RoughnessColorGamma = 2.2f; // Roughness material powers color by 2.2, transform it back to linear
-			RenderSceneToTexture(Scene, RoughnessName, ViewMatrix, ProjectionMatrix, HiddenPrimitives, 
+			RenderSceneToTexture(Scene, RoughnessName, ViewOrigin, ViewRotationMatrix, ProjectionMatrix, HiddenPrimitives, 
 				OutFlattenMaterial.RoughnessSize, RoughnessColorGamma, OutFlattenMaterial.RoughnessSamples);
 		}
 
@@ -2821,7 +2823,7 @@ namespace MaterialExportUtils
 		{
 			static const FName SpecularName("Specular");
 			const float SpecularColorGamma = 1.0f; // Dump specular texture in linear space
-			RenderSceneToTexture(Scene, SpecularName, ViewMatrix, ProjectionMatrix, HiddenPrimitives, 
+			RenderSceneToTexture(Scene, SpecularName, ViewOrigin, ViewRotationMatrix, ProjectionMatrix, HiddenPrimitives, 
 				OutFlattenMaterial.SpecularSize, SpecularColorGamma, OutFlattenMaterial.SpecularSamples);
 		}
 				
diff --git a/Engine/Source/Editor/UnrealEd/Private/EditorViewportClient.cpp b/Engine/Source/Editor/UnrealEd/Private/EditorViewportClient.cpp
--- a/Engine/Source/Editor/UnrealEd/Private/EditorViewportClient.cpp
+++ b/Engine/Source/Editor/UnrealEd/Private/EditorViewportClient.cpp
@@ -550,19 +550,20 @@ FSceneView* FEditorViewportClient::CalcSceneView(FSceneViewFamily* ViewFamily)
 
 	const bool bConstrainAspectRatio = bUseControllingActorViewInfo && ControllingActorViewInfo.bConstrainAspectRatio;
 
-	ViewInitOptions.ViewMatrix = FTranslationMatrix(-ViewLocation);
+	ViewInitOptions.ViewOrigin = ViewLocation;
 	if (EffectiveViewportType == LVT_Perspective)
 	{
 		if (bUsingOrbitCamera)
 		{
-			ViewInitOptions.ViewMatrix = ViewTransform.ComputeOrbitMatrix();
+			ViewInitOptions.ViewRotationMatrix = ViewTransform.ComputeOrbitMatrix();
+			ViewInitOptions.ViewOrigin = FVector::ZeroVector;
 		}
 		else
 		{
-			ViewInitOptions.ViewMatrix = ViewInitOptions.ViewMatrix * FInverseRotationMatrix(ViewRotation);
+			ViewInitOptions.ViewRotationMatrix = FInverseRotationMatrix(ViewRotation);
 		}
 
-		ViewInitOptions.ViewMatrix = ViewInitOptions.ViewMatrix * FMatrix(
+		ViewInitOptions.ViewRotationMatrix = ViewInitOptions.ViewRotationMatrix * FMatrix(
 			FPlane(0,	0,	1,	0),
 			FPlane(1,	0,	0,	0),
 			FPlane(0,	1,	0,	0),
@@ -628,7 +629,7 @@ FSceneView* FEditorViewportClient::CalcSceneView(FSceneViewFamily* ViewFamily)
 
 		if (EffectiveViewportType == LVT_OrthoXY)
 		{
-			ViewInitOptions.ViewMatrix = ViewInitOptions.ViewMatrix * FMatrix(
+			ViewInitOptions.ViewRotationMatrix = FMatrix(
 				FPlane(1,	0,	0,					0),
 				FPlane(0,	-1,	0,					0),
 				FPlane(0,	0,	-1,					0),
@@ -636,7 +637,7 @@ FSceneView* FEditorViewportClient::CalcSceneView(FSceneViewFamily* ViewFamily)
 		}
 		else if (EffectiveViewportType == LVT_OrthoXZ)
 		{
-			ViewInitOptions.ViewMatrix = ViewInitOptions.ViewMatrix * FMatrix(
+			ViewInitOptions.ViewRotationMatrix = FMatrix(
 				FPlane(1,	0,	0,					0),
 				FPlane(0,	0,	-1,					0),
 				FPlane(0,	1,	0,					0),
@@ -644,7 +645,7 @@ FSceneView* FEditorViewportClient::CalcSceneView(FSceneViewFamily* ViewFamily)
 		}
 		else if (EffectiveViewportType == LVT_OrthoYZ)
 		{
-			ViewInitOptions.ViewMatrix = ViewInitOptions.ViewMatrix * FMatrix(
+			ViewInitOptions.ViewRotationMatrix = FMatrix(
 				FPlane(0,	0,	1,					0),
 				FPlane(1,	0,	0,					0),
 				FPlane(0,	1,	0,					0),
@@ -652,8 +653,8 @@ FSceneView* FEditorViewportClient::CalcSceneView(FSceneViewFamily* ViewFamily)
 		}
 		else if (EffectiveViewportType == LVT_OrthoFreelook)
 		{
-			ViewInitOptions.ViewMatrix = ViewInitOptions.ViewMatrix * FInverseRotationMatrix(ViewRotation);
-			ViewInitOptions.ViewMatrix = ViewInitOptions.ViewMatrix * FMatrix(
+			ViewInitOptions.ViewRotationMatrix = FInverseRotationMatrix(ViewRotation);
+			ViewInitOptions.ViewRotationMatrix = ViewInitOptions.ViewRotationMatrix * FMatrix(
 				FPlane(0,	0,	 1,	0),
 				FPlane(1,	0,	 0,	0),
 				FPlane(0,	1,	 0,	0),
@@ -664,7 +665,7 @@ FSceneView* FEditorViewportClient::CalcSceneView(FSceneViewFamily* ViewFamily)
 			OrthoWidth = ControllingActorViewInfo.OrthoWidth / 2.0f;
 			OrthoHeight = (ControllingActorViewInfo.OrthoWidth / 2.0f) * YScale;
 
-			const float NearViewPlane = ViewInitOptions.ViewMatrix.GetOrigin().Z;
+			const float NearViewPlane = -ViewInitOptions.ViewOrigin.X;
 			const float FarViewPlane = NearViewPlane - 2.0f*WORLD_MAX;
 
 			const float InverseRange = 1.0f / (FarViewPlane - NearViewPlane);
diff --git a/Engine/Source/Editor/UnrealEd/Private/ThumbnailHelpers.cpp b/Engine/Source/Editor/UnrealEd/Private/ThumbnailHelpers.cpp
--- a/Engine/Source/Editor/UnrealEd/Private/ThumbnailHelpers.cpp
+++ b/Engine/Source/Editor/UnrealEd/Private/ThumbnailHelpers.cpp
@@ -89,13 +89,12 @@ void FThumbnailPreviewScene::GetView(FSceneViewFamily* ViewFamily, int32 X, int3
 		OrbitZoom = FMath::Max<float>(MinCameraDistance, OrbitZoom);
 
 		const FRotator RotationOffsetToViewCenter(0.f, 90.f, 0.f);
-		FMatrix ViewMatrix = FTranslationMatrix( Origin ) *
-			FRotationMatrix( FRotator(0, OrbitYaw, 0) ) * 
+		FMatrix ViewRotationMatrix = FRotationMatrix( FRotator(0, OrbitYaw, 0) ) * 
 			FRotationMatrix( FRotator(0, 0, OrbitPitch) ) *
 			FTranslationMatrix( FVector(0, OrbitZoom, 0) ) *
 			FInverseRotationMatrix( RotationOffsetToViewCenter );
 
-		ViewMatrix = ViewMatrix * FMatrix(
+		ViewRotationMatrix = ViewRotationMatrix * FMatrix(
 			FPlane(0,	0,	1,	0),
 			FPlane(1,	0,	0,	0),
 			FPlane(0,	1,	0,	0),
@@ -104,7 +103,8 @@ void FThumbnailPreviewScene::GetView(FSceneViewFamily* ViewFamily, int32 X, int3
 		FSceneViewInitOptions ViewInitOptions;
 		ViewInitOptions.ViewFamily = ViewFamily;
 		ViewInitOptions.SetViewRectangle(ViewRect);
-		ViewInitOptions.ViewMatrix = ViewMatrix;
+		ViewInitOptions.ViewOrigin = -Origin;
+		ViewInitOptions.ViewRotationMatrix = ViewRotationMatrix;
 		ViewInitOptions.ProjectionMatrix = ProjectionMatrix;
 		ViewInitOptions.BackgroundColor = FLinearColor::Black;
 
@@ -112,7 +112,7 @@ void FThumbnailPreviewScene::GetView(FSceneViewFamily* ViewFamily, int32 X, int3
 
 		ViewFamily->Views.Add(NewView);
 
-		NewView->StartFinalPostprocessSettings( ViewMatrix.GetOrigin() );
+		NewView->StartFinalPostprocessSettings( ViewInitOptions.ViewOrigin );
 		NewView->EndFinalPostprocessSettings(ViewInitOptions);
 
 		FFinalPostProcessSettings::FCubemapEntry& CubemapEntry = *new(NewView->FinalPostProcessSettings.ContributingCubemaps) FFinalPostProcessSettings::FCubemapEntry;
diff --git a/Engine/Source/Editor/UnrealEd/Private/ThumbnailRendering/LevelThumbnailRenderer.cpp b/Engine/Source/Editor/UnrealEd/Private/ThumbnailRendering/LevelThumbnailRenderer.cpp
--- a/Engine/Source/Editor/UnrealEd/Private/ThumbnailRendering/LevelThumbnailRenderer.cpp
+++ b/Engine/Source/Editor/UnrealEd/Private/ThumbnailRendering/LevelThumbnailRenderer.cpp
@@ -61,11 +61,12 @@ void ULevelThumbnailRenderer::GetView(ULevel* Level, FSceneViewFamily* ViewFamil
 		ViewInitOptions.ViewFamily = ViewFamily;
 
 		const FVector ViewPoint = LevelBox.GetCenter();
-		ViewInitOptions.ViewMatrix = FMatrix(
+		ViewInitOptions.ViewOrigin = FVector(ViewPoint.X, ViewPoint.Y, 0);
+		ViewInitOptions.ViewRotationMatrix = FMatrix(
 			FPlane(1,				0,				0,		0),
 			FPlane(0,				-1,				0,		0),
 			FPlane(0,				0,				-1,		0),
-			FPlane(-ViewPoint.X,	ViewPoint.Y,	0,		1));
+			FPlane(0,				0,				0,		1));
 
 		const float ZOffset = WORLD_MAX;
 		ViewInitOptions.ProjectionMatrix =  FReversedZOrthoMatrix(
diff --git a/Engine/Source/Editor/UnrealEd/Private/ThumbnailRendering/WorldThumbnailRenderer.cpp b/Engine/Source/Editor/UnrealEd/Private/ThumbnailRendering/WorldThumbnailRenderer.cpp
--- a/Engine/Source/Editor/UnrealEd/Private/ThumbnailRendering/WorldThumbnailRenderer.cpp
+++ b/Engine/Source/Editor/UnrealEd/Private/ThumbnailRendering/WorldThumbnailRenderer.cpp
@@ -120,7 +120,7 @@ void UWorldThumbnailRenderer::GetView(UWorld* World, FSceneViewFamily* ViewFamil
 		}
 
 		const FVector Origin = WorldBox.GetCenter();
-		FMatrix ViewMatrix = FTranslationMatrix(-Origin);
+		FMatrix ViewRotationMatrix;
 		FMatrix ProjectionMatrix;
 		float FOVScreenSize = 0; // Screen size taking FOV into account
 		if (ThumbnailInfo->CameraMode == ECameraProjectionMode::Perspective)
@@ -147,13 +147,12 @@ void UWorldThumbnailRenderer::GetView(UWorld* World, FSceneViewFamily* ViewFamil
 			OrbitZoom = FMath::Max<float>(MinCameraDistance, OrbitZoom);
 
 			const FRotator RotationOffsetToViewCenter(0.f, 90.f, 0.f);
-			ViewMatrix = ViewMatrix *
-				FRotationMatrix(FRotator(0, OrbitYaw, 0)) *
+			ViewRotationMatrix = FRotationMatrix(FRotator(0, OrbitYaw, 0)) *
 				FRotationMatrix(FRotator(0, 0, OrbitPitch)) *
 				FTranslationMatrix(FVector(0, OrbitZoom, 0)) *
 				FInverseRotationMatrix(RotationOffsetToViewCenter);
 
-			ViewMatrix = ViewMatrix * FMatrix(
+			ViewRotationMatrix = ViewRotationMatrix * FMatrix(
 				FPlane(0, 0, 1, 0),
 				FPlane(1, 0, 0, 0),
 				FPlane(0, 1, 0, 0),
@@ -176,7 +175,7 @@ void UWorldThumbnailRenderer::GetView(UWorld* World, FSceneViewFamily* ViewFamil
 			switch (ThumbnailInfo->OrthoDirection)
 			{
 				case EOrthoThumbnailDirection::Top:
-					ViewMatrix = ViewMatrix * FMatrix(
+					ViewRotationMatrix = FMatrix(
 						FPlane(1, 0, 0, 0),
 						FPlane(0, -1, 0, 0),
 						FPlane(0, 0, -1, 0),
@@ -185,7 +184,7 @@ void UWorldThumbnailRenderer::GetView(UWorld* World, FSceneViewFamily* ViewFamil
 					WorldSizeMax2D = FVector2D(WorldBox.Max.X,WorldBox.Max.Y);
 					break;
 				case EOrthoThumbnailDirection::Bottom:
-					ViewMatrix = ViewMatrix * FMatrix(
+					ViewRotationMatrix = FMatrix(
 						FPlane(1, 0, 0, 0),
 						FPlane(0, -1, 0, 0),
 						FPlane(0, 0, 1, 0),
@@ -194,7 +193,7 @@ void UWorldThumbnailRenderer::GetView(UWorld* World, FSceneViewFamily* ViewFamil
 					WorldSizeMax2D = FVector2D(WorldBox.Max.X, WorldBox.Max.Y);
 					break;
 				case EOrthoThumbnailDirection::Front:
-					ViewMatrix = ViewMatrix * FMatrix(
+					ViewRotationMatrix = FMatrix(
 						FPlane(1, 0, 0, 0),
 						FPlane(0, 0, -1, 0),
 						FPlane(0, 1, 0, 0),
@@ -203,7 +202,7 @@ void UWorldThumbnailRenderer::GetView(UWorld* World, FSceneViewFamily* ViewFamil
 					WorldSizeMax2D = FVector2D(WorldBox.Max.X, WorldBox.Max.Z);
 					break;
 				case EOrthoThumbnailDirection::Back:
-					ViewMatrix = ViewMatrix * FMatrix(
+					ViewRotationMatrix = FMatrix(
 						FPlane(-1, 0, 0, 0),
 						FPlane(0, 0, 1, 0),
 						FPlane(0, 1, 0, 0),
@@ -212,7 +211,7 @@ void UWorldThumbnailRenderer::GetView(UWorld* World, FSceneViewFamily* ViewFamil
 					WorldSizeMax2D = FVector2D(WorldBox.Max.X, WorldBox.Max.Z);
 					break;
 				case EOrthoThumbnailDirection::Left:
-					ViewMatrix = ViewMatrix * FMatrix(
+					ViewRotationMatrix = FMatrix(
 						FPlane(0, 0, -1, 0),
 						FPlane(-1, 0, 0, 0),
 						FPlane(0, 1, 0, 0),
@@ -221,7 +220,7 @@ void UWorldThumbnailRenderer::GetView(UWorld* World, FSceneViewFamily* ViewFamil
 					WorldSizeMax2D = FVector2D(WorldBox.Max.Y, WorldBox.Max.Z);
 					break;
 				case EOrthoThumbnailDirection::Right:
-					ViewMatrix = ViewMatrix * FMatrix(
+					ViewRotationMatrix = FMatrix(
 						FPlane(0, 0, 1, 0),
 						FPlane(1, 0, 0, 0),
 						FPlane(0, 1, 0, 0),
@@ -265,7 +264,8 @@ void UWorldThumbnailRenderer::GetView(UWorld* World, FSceneViewFamily* ViewFamil
 		ViewInitOptions.ViewFamily = ViewFamily;
 		ViewInitOptions.SetViewRectangle(ViewRect);
 		ViewInitOptions.BackgroundColor = FLinearColor::Black;
-		ViewInitOptions.ViewMatrix = ViewMatrix;
+		ViewInitOptions.ViewOrigin = Origin;
+		ViewInitOptions.ViewRotationMatrix = ViewRotationMatrix;
 		ViewInitOptions.ProjectionMatrix = ProjectionMatrix;
 
 		FSceneView* NewView = new FSceneView(ViewInitOptions);
diff --git a/Engine/Source/Runtime/Engine/Private/LocalPlayer.cpp b/Engine/Source/Runtime/Engine/Private/LocalPlayer.cpp
--- a/Engine/Source/Runtime/Engine/Private/LocalPlayer.cpp
+++ b/Engine/Source/Runtime/Engine/Private/LocalPlayer.cpp
@@ -819,7 +819,7 @@ bool ULocalPlayer::GetPixelBoundingBox(const FBox& ActorBox, FVector2D& OutLower
 		};
 
 		// create the view projection matrix
-		const FMatrix ViewProjectionMatrix = ProjectionData.ViewMatrix * ProjectionData.ProjectionMatrix;
+		const FMatrix ViewProjectionMatrix = ProjectionData.ComputeViewProjectionMatrix();
 
 		int SuccessCount = 0;
 		OutLowerLeft = FVector2D(FLT_MAX, FLT_MAX);
@@ -879,7 +879,7 @@ bool ULocalPlayer::GetPixelPoint(const FVector& InPoint, FVector2D& OutPoint, co
 		}
 
 		// create the view projection matrix
-		const FMatrix ViewProjectionMatrix = ProjectionData.ViewMatrix * ProjectionData.ProjectionMatrix;
+		const FMatrix ViewProjectionMatrix = ProjectionData.ComputeViewProjectionMatrix();
 
 		//@TODO: CAMERA: Validate this code!
 		// grab the point in screen space
@@ -965,9 +965,8 @@ bool ULocalPlayer::GetProjectionData(FViewport* Viewport, EStereoscopicPass Ster
     }
 
 	// Create the view matrix
-	ProjectionData.ViewMatrix = FTranslationMatrix(-StereoViewLocation);
-	ProjectionData.ViewMatrix = ProjectionData.ViewMatrix * FInverseRotationMatrix(ViewInfo.Rotation);
-	ProjectionData.ViewMatrix = ProjectionData.ViewMatrix * FMatrix(
+	ProjectionData.ViewOrigin = StereoViewLocation;
+	ProjectionData.ViewRotationMatrix = FInverseRotationMatrix(ViewInfo.Rotation) * FMatrix(
 		FPlane(0,	0,	1,	0),
 		FPlane(1,	0,	0,	0),
 		FPlane(0,	1,	0,	0),
@@ -988,7 +987,7 @@ bool ULocalPlayer::GetProjectionData(FViewport* Viewport, EStereoscopicPass Ster
 	
 				const float OrthoWidth = ViewInfo.OrthoWidth / 2.0f;
 				const float OrthoHeight = ViewInfo.OrthoWidth / 2.0f * YScale;
-				const float NearPlane = ProjectionData.ViewMatrix.GetOrigin().Z;
+				const float NearPlane = -ProjectionData.ViewOrigin.X;
 				const float FarPlane = NearPlane - 2.0f*WORLD_MAX;
 				const float InverseRange = 1.0f / (FarPlane - NearPlane);
 				const float ZScale = -2.0f * InverseRange;
@@ -1034,7 +1033,7 @@ bool ULocalPlayer::GetProjectionData(FViewport* Viewport, EStereoscopicPass Ster
 	
 			if (ViewInfo.ProjectionMode == ECameraProjectionMode::Orthographic)
 			{
-				const float NearPlane = ProjectionData.ViewMatrix.GetOrigin().Z;
+				const float NearPlane = -ProjectionData.ViewOrigin.X;
 				const float FarPlane = NearPlane - 2.0f*WORLD_MAX;
 	
 				const float OrthoWidth = ViewInfo.OrthoWidth / 2.0f * XAxisMultiplier;
diff --git a/Engine/Source/Runtime/Engine/Private/SceneView.cpp b/Engine/Source/Runtime/Engine/Private/SceneView.cpp
--- a/Engine/Source/Runtime/Engine/Private/SceneView.cpp
+++ b/Engine/Source/Runtime/Engine/Private/SceneView.cpp
@@ -297,7 +297,7 @@ FSceneView::FSceneView(const FSceneViewInitOptions& InitOptions)
 	check(UnscaledViewRect.Width() > 0);
 	check(UnscaledViewRect.Height() > 0);
 
-	ViewMatrices.ViewMatrix = InitOptions.ViewMatrix;
+	ViewMatrices.ViewMatrix = FTranslationMatrix(-InitOptions.ViewOrigin) * InitOptions.ViewRotationMatrix;
 
 	// Adjust the projection matrix for the current RHI.
 	ViewMatrices.ProjMatrix = AdjustProjectionMatrixForRHI(ProjectionMatrixUnadjustedForRHI);
@@ -309,31 +309,33 @@ FSceneView::FSceneView(const FSceneViewInitOptions& InitOptions)
 	InvViewMatrix = ViewMatrices.ViewMatrix.Inverse();
 	InvViewProjectionMatrix = ViewMatrices.GetInvProjMatrix() * InvViewMatrix;
 
-	bool ApplyPreViewTranslation = true;
+	bool bApplyPreViewTranslation = true;
+	bool bViewOriginIsFudged = false;
 
 	// Calculate the view origin from the view/projection matrices.
 	if(IsPerspectiveProjection())
 	{
-		ViewMatrices.ViewOrigin = InvViewMatrix.GetOrigin();
+		ViewMatrices.ViewOrigin = InitOptions.ViewOrigin;
 	}
 #if WITH_EDITOR
 	else if (InitOptions.bUseFauxOrthoViewPos)
 	{
 		float DistanceToViewOrigin = WORLD_MAX;
-		ViewMatrices.ViewOrigin = FVector4(InvViewMatrix.TransformVector(FVector(0,0,-1)).GetSafeNormal()*DistanceToViewOrigin,1) + InvViewMatrix.GetOrigin();
+		ViewMatrices.ViewOrigin = FVector4(InvViewMatrix.TransformVector(FVector(0,0,-1)).GetSafeNormal()*DistanceToViewOrigin,1) + InitOptions.ViewOrigin;
+		bViewOriginIsFudged = true;
 	}
 #endif
 	else
 	{
 		ViewMatrices.ViewOrigin = FVector4(InvViewMatrix.TransformVector(FVector(0,0,-1)).GetSafeNormal(),0);
 		// to avoid issues with view dependent effect (e.g. Frensel)
-		ApplyPreViewTranslation = false;
+		bApplyPreViewTranslation = false;
 	}
 
 	// Translate world-space so its origin is at ViewOrigin for improved precision.
 	// Note that this isn't exactly right for orthogonal projections (See the above special case), but we still use ViewOrigin
 	// in that case so the same value may be used in shaders for both the world-space translation and the camera's world position.
-	if(ApplyPreViewTranslation)
+	if(bApplyPreViewTranslation)
 	{
 		ViewMatrices.PreViewTranslation = -FVector(ViewMatrices.ViewOrigin);
 
@@ -358,7 +360,14 @@ FSceneView::FSceneView(const FSceneViewInitOptions& InitOptions)
 	}
 
 	/** The view transform, starting from world-space points translated by -ViewOrigin. */
-	FMatrix TranslatedViewMatrix = FTranslationMatrix(-ViewMatrices.PreViewTranslation) * ViewMatrices.ViewMatrix;
+	FMatrix TranslatedViewMatrix = InitOptions.ViewRotationMatrix;
+
+	// When the view origin is fudged for faux ortho view position the translations don't cancel out.
+	if (bViewOriginIsFudged)
+	{
+		TranslatedViewMatrix = FTranslationMatrix(-ViewMatrices.PreViewTranslation)
+			* FTranslationMatrix(-InitOptions.ViewOrigin) * InitOptions.ViewRotationMatrix;
+	}
 	
 	// Compute a transform from view origin centered world-space to clip space.
 	ViewMatrices.TranslatedViewProjectionMatrix = TranslatedViewMatrix * ViewMatrices.ProjMatrix;
diff --git a/Engine/Source/Runtime/Engine/Private/UserInterface/Canvas.cpp b/Engine/Source/Runtime/Engine/Private/UserInterface/Canvas.cpp
--- a/Engine/Source/Runtime/Engine/Private/UserInterface/Canvas.cpp
+++ b/Engine/Source/Runtime/Engine/Private/UserInterface/Canvas.cpp
@@ -481,7 +481,8 @@ bool FCanvasTileRendererItem::Render_RenderThread(FRHICommandListImmediate& RHIC
 	FSceneViewInitOptions ViewInitOptions;
 	ViewInitOptions.ViewFamily = ViewFamily;
 	ViewInitOptions.SetViewRectangle(ViewRect);
-	ViewInitOptions.ViewMatrix = FMatrix::Identity;
+	ViewInitOptions.ViewOrigin = FVector::ZeroVector;
+	ViewInitOptions.ViewRotationMatrix = FMatrix::Identity;
 	ViewInitOptions.ProjectionMatrix = Data->Transform.GetMatrix();
 	ViewInitOptions.BackgroundColor = FLinearColor::Black;
 	ViewInitOptions.OverlayColor = FLinearColor::White;
@@ -547,7 +548,8 @@ bool FCanvasTileRendererItem::Render_GameThread(const FCanvas* Canvas)
 	FSceneViewInitOptions ViewInitOptions;
 	ViewInitOptions.ViewFamily = ViewFamily;
 	ViewInitOptions.SetViewRectangle(ViewRect);
-	ViewInitOptions.ViewMatrix = FMatrix::Identity;
+	ViewInitOptions.ViewOrigin = FVector::ZeroVector;
+	ViewInitOptions.ViewRotationMatrix = FMatrix::Identity;
 	ViewInitOptions.ProjectionMatrix = Data->Transform.GetMatrix();
 	ViewInitOptions.BackgroundColor = FLinearColor::Black;
 	ViewInitOptions.OverlayColor = FLinearColor::White;
diff --git a/Engine/Source/Runtime/Engine/Public/SceneView.h b/Engine/Source/Runtime/Engine/Public/SceneView.h
--- a/Engine/Source/Runtime/Engine/Public/SceneView.h
+++ b/Engine/Source/Runtime/Engine/Public/SceneView.h
@@ -20,7 +20,11 @@ class FSceneViewFamily;
 // Projection data for a FSceneView
 struct FSceneViewProjectionData
 {
-	FMatrix ViewMatrix;
+	/** The view origin. */
+	FVector ViewOrigin;
+
+	/** Rotation matrix transforming from world space to view space. */
+	FMatrix ViewRotationMatrix;
 
 	/** UE4 projection matrix projects such that clip space Z=1 is the near plane, and Z=0 is the infinite far plane. */
 	FMatrix ProjectionMatrix;
@@ -54,6 +58,11 @@ public:
 
 	const FIntRect& GetViewRect() const { return ViewRect; }
 	const FIntRect& GetConstrainedViewRect() const { return ConstrainedViewRect; }
+
+	FMatrix ComputeViewProjectionMatrix() const
+	{
+		return FTranslationMatrix(-ViewOrigin) * ViewRotationMatrix * ProjectionMatrix;
+	}
 };
 
 // Construction parameters for a FSceneView
diff --git a/Engine/Source/Runtime/Landscape/Private/LandscapeGrass.cpp b/Engine/Source/Runtime/Landscape/Private/LandscapeGrass.cpp
--- a/Engine/Source/Runtime/Landscape/Private/LandscapeGrass.cpp
+++ b/Engine/Source/Runtime/Landscape/Private/LandscapeGrass.cpp
@@ -226,7 +226,8 @@ public:
 	FIntPoint TargetSize;
 	int32 NumPasses;
 	float PassOffsetX;
-	FMatrix ViewMatrix;
+	FVector ViewOrigin;
+	FMatrix ViewRotationMatrix;
 	FMatrix ProjectionMatrix;
 
 	void RenderLandscapeComponentToTexture_RenderThread(FRHICommandListImmediate& RHICmdList)
@@ -239,7 +240,8 @@ public:
 
 		FSceneViewInitOptions ViewInitOptions;
 		ViewInitOptions.SetViewRectangle(FIntRect(0, 0, TargetSize.X, TargetSize.Y));
-		ViewInitOptions.ViewMatrix = ViewMatrix;
+		ViewInitOptions.ViewOrigin = ViewOrigin;
+		ViewInitOptions.ViewRotationMatrix = ViewRotationMatrix;
 		ViewInitOptions.ProjectionMatrix = ProjectionMatrix;
 		ViewInitOptions.ViewFamily = &ViewFamily;
 
@@ -320,9 +322,9 @@ public:
 		// extent of target in world space
 		FVector TargetExtent = FVector(TargetSize, 0.f)*LandscapeProxy->GetActorScale()*0.5f;
 
-		ViewMatrix = FTranslationMatrix(-TargetCenter);
-		ViewMatrix *= FInverseRotationMatrix(LandscapeProxy->GetActorRotation());
-		ViewMatrix *= FMatrix(FPlane(1, 0, 0, 0),
+		ViewOrigin = TargetCenter;
+		ViewRotationMatrix = FInverseRotationMatrix(LandscapeProxy->GetActorRotation());
+		ViewRotationMatrix *= FMatrix(FPlane(1, 0, 0, 0),
 			FPlane(0, -1, 0, 0),
 			FPlane(0, 0, -1, 0),
 			FPlane(0, 0, 0, 1));
diff --git a/Engine/Source/Runtime/Renderer/Private/ReflectionEnvironmentCapture.cpp b/Engine/Source/Runtime/Renderer/Private/ReflectionEnvironmentCapture.cpp
--- a/Engine/Source/Runtime/Renderer/Private/ReflectionEnvironmentCapture.cpp
+++ b/Engine/Source/Runtime/Renderer/Private/ReflectionEnvironmentCapture.cpp
@@ -1060,7 +1060,7 @@ void UploadReflectionCapture_RenderingThread(FScene* Scene, const FReflectionCap
 }
 
 /** Creates a transformation for a cubemap face, following the D3D cubemap layout. */
-FMatrix CalcCubeFaceViewMatrix(ECubeFace Face, FVector WorldLocation)
+FMatrix CalcCubeFaceViewRotationMatrix(ECubeFace Face)
 {
 	FMatrix Result(FMatrix::Identity);
 
@@ -1099,7 +1099,7 @@ FMatrix CalcCubeFaceViewMatrix(ECubeFace Face, FVector WorldLocation)
 	// derive right vector
 	FVector vRight( vUp ^ vDir );
 	// create matrix from the 3 axes
-	Result = FBasisVectorMatrix( vRight, vUp, vDir, -WorldLocation );	
+	Result = FBasisVectorMatrix( vRight, vUp, vDir, FVector::ZeroVector );	
 
 	return Result;
 }
@@ -1184,7 +1184,8 @@ void CaptureSceneIntoScratchCubemap(FScene* Scene, FVector CapturePosition, bool
 			NearPlane
 			);
 
-		ViewInitOptions.ViewMatrix = CalcCubeFaceViewMatrix((ECubeFace)CubeFace, CapturePosition);
+		ViewInitOptions.ViewOrigin = CapturePosition;
+		ViewInitOptions.ViewRotationMatrix = CalcCubeFaceViewRotationMatrix((ECubeFace)CubeFace);
 
 		FSceneView* View = new FSceneView(ViewInitOptions);
 
diff --git a/Engine/Source/Runtime/Renderer/Private/Renderer.cpp b/Engine/Source/Runtime/Renderer/Private/Renderer.cpp
--- a/Engine/Source/Runtime/Renderer/Private/Renderer.cpp
+++ b/Engine/Source/Runtime/Renderer/Private/Renderer.cpp
@@ -137,11 +137,12 @@ void FRendererModule::GPUBenchmark(FSynthBenchmarkResults& InOut, float WorkScal
 	// Initialize Projection Matrix and ViewMatrix since FSceneView initialization is doing some math on them.
 	// Otherwise it trips NaN checks.
 	const FVector ViewPoint = LevelBox.GetCenter();
-	ViewInitOptions.ViewMatrix = FMatrix(
+	ViewInitOptions.ViewOrigin = FVector(ViewPoint.X, ViewPoint.Y, 0.0f);
+	ViewInitOptions.ViewRotationMatrix = FMatrix(
 		FPlane(1, 0, 0, 0),
 		FPlane(0, -1, 0, 0),
 		FPlane(0, 0, -1, 0),
-		FPlane(-ViewPoint.X, ViewPoint.Y, 0, 1));
+		FPlane(0, 0, 0, 1));
 
 	const float ZOffset = WORLD_MAX;
 	ViewInitOptions.ProjectionMatrix = FReversedZOrthoMatrix(
diff --git a/Engine/Source/Runtime/Renderer/Private/SceneCaptureRendering.cpp b/Engine/Source/Runtime/Renderer/Private/SceneCaptureRendering.cpp
--- a/Engine/Source/Runtime/Renderer/Private/SceneCaptureRendering.cpp
+++ b/Engine/Source/Runtime/Renderer/Private/SceneCaptureRendering.cpp
@@ -158,7 +158,7 @@ static void UpdateSceneCaptureContent_RenderThread(FRHICommandListImmediate& RHI
 }
 
 
-FSceneRenderer* FScene::CreateSceneRenderer( USceneCaptureComponent* SceneCaptureComponent, UTextureRenderTarget* TextureTarget, const FMatrix& ViewMatrix, const FVector& ViewLocation, float FOV, float MaxViewDistance, bool bCaptureSceneColour, FPostProcessSettings* PostProcessSettings, float PostProcessBlendWeight )
+FSceneRenderer* FScene::CreateSceneRenderer( USceneCaptureComponent* SceneCaptureComponent, UTextureRenderTarget* TextureTarget, const FMatrix& ViewRotationMatrix, const FVector& ViewLocation, float FOV, float MaxViewDistance, bool bCaptureSceneColour, FPostProcessSettings* PostProcessSettings, float PostProcessBlendWeight )
 {
 	FIntPoint CaptureSize(TextureTarget->GetSurfaceWidth(), TextureTarget->GetSurfaceHeight());
 
@@ -172,7 +172,8 @@ FSceneRenderer* FScene::CreateSceneRenderer( USceneCaptureComponent* SceneCaptur
 	FSceneViewInitOptions ViewInitOptions;
 	ViewInitOptions.SetViewRectangle(FIntRect(0, 0, CaptureSize.X, CaptureSize.Y));
 	ViewInitOptions.ViewFamily = &ViewFamily;
-	ViewInitOptions.ViewMatrix = ViewMatrix;
+	ViewInitOptions.ViewOrigin = ViewLocation;
+	ViewInitOptions.ViewRotationMatrix = ViewRotationMatrix;
 	ViewInitOptions.BackgroundColor = FLinearColor::Black;
 	ViewInitOptions.OverrideFarClippingPlaneDistance = MaxViewDistance;
 	ViewInitOptions.SceneViewStateInterface = SceneCaptureComponent->GetViewState();
@@ -241,18 +242,22 @@ void FScene::UpdateSceneCaptureContents(USceneCaptureComponent2D* CaptureCompone
 
 	if (CaptureComponent->TextureTarget)
 	{
-		const FTransform& Transform = CaptureComponent->GetComponentToWorld();
-		FMatrix ViewMatrix = Transform.ToInverseMatrixWithScale();
+		FTransform Transform = CaptureComponent->GetComponentToWorld();
 		FVector ViewLocation = Transform.GetTranslation();
+
+		// Remove the translation from Transform because we only need rotation.
+		Transform.SetTranslation(FVector::ZeroVector);
+		FMatrix ViewRotationMatrix = Transform.ToInverseMatrixWithScale();
+
 		// swap axis st. x=z,y=x,z=y (unreal coord space) so that z is up
-		ViewMatrix = ViewMatrix * FMatrix(
+		ViewRotationMatrix = ViewRotationMatrix * FMatrix(
 			FPlane(0,	0,	1,	0),
 			FPlane(1,	0,	0,	0),
 			FPlane(0,	1,	0,	0),
 			FPlane(0,	0,	0,	1));
 		const float FOV = CaptureComponent->FOVAngle * (float)PI / 360.0f;
 		const bool bUseSceneColorTexture = CaptureComponent->CaptureSource == SCS_SceneColorHDR;
-		FSceneRenderer* SceneRenderer = CreateSceneRenderer(CaptureComponent, CaptureComponent->TextureTarget, ViewMatrix , ViewLocation, FOV, CaptureComponent->MaxViewDistanceOverride, bUseSceneColorTexture, &CaptureComponent->PostProcessSettings, CaptureComponent->PostProcessBlendWeight);
+		FSceneRenderer* SceneRenderer = CreateSceneRenderer(CaptureComponent, CaptureComponent->TextureTarget, ViewRotationMatrix , ViewLocation, FOV, CaptureComponent->MaxViewDistanceOverride, bUseSceneColorTexture, &CaptureComponent->PostProcessSettings, CaptureComponent->PostProcessBlendWeight);
 
 		FTextureRenderTargetResource* TextureRenderTarget = CaptureComponent->TextureTarget->GameThread_GetRenderTargetResource();
 		const FName OwnerName = CaptureComponent->GetOwner() ? CaptureComponent->GetOwner()->GetFName() : NAME_None;
@@ -274,7 +279,7 @@ void FScene::UpdateSceneCaptureContents(USceneCaptureComponentCube* CaptureCompo
 	struct FLocal
 	{
 		/** Creates a transformation for a cubemap face, following the D3D cubemap layout. */
-		static FMatrix CalcCubeFaceTransform(ECubeFace Face, FVector WorldLocation)
+		static FMatrix CalcCubeFaceTransform(ECubeFace Face)
 		{
 			static const FVector XAxis(1.f, 0.f, 0.f);
 			static const FVector YAxis(0.f, 1.f, 0.f);
@@ -309,7 +314,7 @@ void FScene::UpdateSceneCaptureContents(USceneCaptureComponentCube* CaptureCompo
 			// derive right vector
 			FVector vRight(vUp ^ vDir);
 			// create matrix from the 3 axes
-			return FBasisVectorMatrix(vRight, vUp, vDir, -WorldLocation);
+			return FBasisVectorMatrix(vRight, vUp, vDir, FVector::ZeroVector);
 		}
 	} ;
 
@@ -322,8 +327,8 @@ void FScene::UpdateSceneCaptureContents(USceneCaptureComponentCube* CaptureCompo
 		{
 			const ECubeFace TargetFace = (ECubeFace)faceidx;
 			const FVector Location = CaptureComponent->GetComponentToWorld().GetTranslation();
-			const FMatrix ViewMatrix = FLocal::CalcCubeFaceTransform(TargetFace, Location);
-			FSceneRenderer* SceneRenderer = CreateSceneRenderer(CaptureComponent, CaptureComponent->TextureTarget, ViewMatrix, Location, FOV, CaptureComponent->MaxViewDistanceOverride);
+			const FMatrix ViewRotationMatrix = FLocal::CalcCubeFaceTransform(TargetFace);
+			FSceneRenderer* SceneRenderer = CreateSceneRenderer(CaptureComponent, CaptureComponent->TextureTarget, ViewRotationMatrix, Location, FOV, CaptureComponent->MaxViewDistanceOverride);
 
 			FTextureRenderTargetCubeResource* TextureRenderTarget = static_cast<FTextureRenderTargetCubeResource*>(CaptureComponent->TextureTarget->GameThread_GetRenderTargetResource());
 			const FName OwnerName = CaptureComponent->GetOwner() ? CaptureComponent->GetOwner()->GetFName() : NAME_None;
diff --git a/Engine/Source/Runtime/Renderer/Private/ScenePrivate.h b/Engine/Source/Runtime/Renderer/Private/ScenePrivate.h
--- a/Engine/Source/Runtime/Renderer/Private/ScenePrivate.h
+++ b/Engine/Source/Runtime/Renderer/Private/ScenePrivate.h
@@ -1775,7 +1775,7 @@ private:
 	 *
 	 * @param	SceneCaptureComponent - The scene capture component for which to create a scene renderer.
 	 * @param	TextureTarget - render target to draw to.
-	 * @param	ViewMatrix - Camera view matrix.
+	 * @param	ViewRotationMatrix - Camera rotation matrix.
 	 * @param	ViewLocation - Camera location.
 	 * @param	FOV - Camera field of view.
 	 * @param	MaxViewDistance - Far draw distance.
@@ -1784,7 +1784,7 @@ private:
 	 * @param	PostProcessBlendWeight - Blendweight for PostProcessSettings.
 	 * @return	pointer to a configured SceneRenderer instance.
 	 */
-	FSceneRenderer* CreateSceneRenderer(USceneCaptureComponent* SceneCaptureComponent, UTextureRenderTarget* TextureTarget, const FMatrix& ViewMatrix, const FVector& ViewLocation, float FOV, float MaxViewDistance, bool bCaptureSceneColour = true, FPostProcessSettings* PostProcessSettings = NULL, float PostProcessBlendWeight = 0);
+	FSceneRenderer* CreateSceneRenderer(USceneCaptureComponent* SceneCaptureComponent, UTextureRenderTarget* TextureTarget, const FMatrix& ViewRotationMatrix, const FVector& ViewLocation, float FOV, float MaxViewDistance, bool bCaptureSceneColour = true, FPostProcessSettings* PostProcessSettings = NULL, float PostProcessBlendWeight = 0);
 
 private:
 	/** 
diff --git a/Engine/Source/Runtime/SlateRHIRenderer/Private/SlateRHIRenderingPolicy.cpp b/Engine/Source/Runtime/SlateRHIRenderer/Private/SlateRHIRenderingPolicy.cpp
--- a/Engine/Source/Runtime/SlateRHIRenderer/Private/SlateRHIRenderingPolicy.cpp
+++ b/Engine/Source/Runtime/SlateRHIRenderer/Private/SlateRHIRenderingPolicy.cpp
@@ -200,7 +200,8 @@ static FSceneView& CreateSceneView( FSceneViewFamilyContext& ViewFamilyContext,
 	FSceneViewInitOptions ViewInitOptions;
 	ViewInitOptions.ViewFamily = &ViewFamilyContext;
 	ViewInitOptions.SetViewRectangle(ViewRect);
-	ViewInitOptions.ViewMatrix = FMatrix::Identity;
+	ViewInitOptions.ViewOrigin = FVector::ZeroVector;
+	ViewInitOptions.ViewRotationMatrix = FMatrix::Identity;
 	ViewInitOptions.ProjectionMatrix = ViewProjectionMatrix;
 	ViewInitOptions.BackgroundColor = FLinearColor::Black;
 	ViewInitOptions.OverlayColor = FLinearColor::White;
diff --git a/Engine/Source/Runtime/UMG/Private/Components/Viewport.cpp b/Engine/Source/Runtime/UMG/Private/Components/Viewport.cpp
--- a/Engine/Source/Runtime/UMG/Private/Components/Viewport.cpp
+++ b/Engine/Source/Runtime/UMG/Private/Components/Viewport.cpp
@@ -310,12 +310,12 @@ FSceneView* FUMGViewportClient::CalcSceneView(FSceneViewFamily* ViewFamily)
 	const ECameraProjectionMode::Type ProjectionType = GetProjectionType();
 	const bool bConstrainAspectRatio = ViewInfo.bConstrainAspectRatio;
 
-	ViewInitOptions.ViewMatrix = FTranslationMatrix(-GetViewLocation());
+	ViewInitOptions.ViewOrigin = ViewLocation;
 	if ( ProjectionType == ECameraProjectionMode::Perspective )
 	{
-		ViewInitOptions.ViewMatrix = ViewInitOptions.ViewMatrix * FInverseRotationMatrix(ViewRotation);
+		ViewInitOptions.ViewRotationMatrix = FInverseRotationMatrix(ViewRotation);
 
-		ViewInitOptions.ViewMatrix = ViewInitOptions.ViewMatrix * FMatrix(
+		ViewInitOptions.ViewRotationMatrix = ViewInitOptions.ViewRotationMatrix * FMatrix(
 			FPlane(0, 0, 1, 0),
 			FPlane(1, 0, 0, 0),
 			FPlane(0, 1, 0, 0),
@@ -377,8 +377,8 @@ FSceneView* FUMGViewportClient::CalcSceneView(FSceneViewFamily* ViewFamily)
 		float OrthoWidth = Zoom * ViewportSizeXY.X / 2.0f;
 		float OrthoHeight = Zoom * ViewportSizeXY.Y / 2.0f;
 
-		ViewInitOptions.ViewMatrix = ViewInitOptions.ViewMatrix * FInverseRotationMatrix(ViewRotation);
-		ViewInitOptions.ViewMatrix = ViewInitOptions.ViewMatrix * FMatrix(
+		ViewInitOptions.ViewRotationMatrix = FInverseRotationMatrix(ViewRotation);
+		ViewInitOptions.ViewRotationMatrix = ViewInitOptions.ViewRotationMatrix * FMatrix(
 			FPlane(0, 0, 1, 0),
 			FPlane(1, 0, 0, 0),
 			FPlane(0, 1, 0, 0),
@@ -389,7 +389,7 @@ FSceneView* FUMGViewportClient::CalcSceneView(FSceneViewFamily* ViewFamily)
 		OrthoWidth = ViewInfo.OrthoWidth / 2.0f;
 		OrthoHeight = ( ViewInfo.OrthoWidth / 2.0f ) * YScale;
 
-		const float NearViewPlane = ViewInitOptions.ViewMatrix.GetOrigin().Z;
+		const float NearViewPlane = -ViewInitOptions.ViewOrigin.X;
 		const float FarViewPlane = NearViewPlane - 2.0f*WORLD_MAX;
 
 		const float InverseRange = 1.0f / ( FarViewPlane - NearViewPlane );
