Cleaned up draw events

diff --git a/Engine/Source/Runtime/Engine/Private/UserInterface/Canvas.cpp b/Engine/Source/Runtime/Engine/Private/UserInterface/Canvas.cpp
--- a/Engine/Source/Runtime/Engine/Private/UserInterface/Canvas.cpp
+++ b/Engine/Source/Runtime/Engine/Private/UserInterface/Canvas.cpp
@@ -425,8 +425,6 @@ bool FCanvasBatchedElementRenderItem::Render_GameThread(const FCanvas* Canvas)
 			BatchedDrawCommand,
 			FBatchedDrawParameters,Parameters,DrawParameters,
 		{
-			SCOPED_DRAW_EVENT(RHICmdList, CanvasBatchedElements);
-
 			// draw batched items
 			Parameters.RenderData->BatchedElements.Draw(
 				RHICmdList,
diff --git a/Engine/Source/Runtime/Renderer/Private/SceneOcclusion.cpp b/Engine/Source/Runtime/Renderer/Private/SceneOcclusion.cpp
--- a/Engine/Source/Runtime/Renderer/Private/SceneOcclusion.cpp
+++ b/Engine/Source/Runtime/Renderer/Private/SceneOcclusion.cpp
@@ -956,8 +956,6 @@ void BuildHZB( FRHICommandListImmediate& RHICmdList, FViewInfo& View )
 	
 	// Mip 0
 	{
-		SCOPED_DRAW_EVENTF(RHICmdList, BuildHZB, TEXT("HZB#%d"), 0);
-
 		SetRenderTarget(RHICmdList, HZBRenderTarget.TargetableTexture, 0, NULL);
 
 		TShaderMapRef< FScreenVS >		VertexShader(View.ShaderMap);
@@ -992,8 +990,6 @@ void BuildHZB( FRHICommandListImmediate& RHICmdList, FViewInfo& View )
 	// Mip 1-7
 	for( uint8 MipIndex = 1; MipIndex < NumMips; MipIndex++ )
 	{
-		SCOPED_DRAW_EVENTF(RHICmdList, BuildHZB, TEXT("HZB#%d"), MipIndex);
-
 		SetRenderTarget(RHICmdList, HZBRenderTarget.TargetableTexture, MipIndex, NULL);
 
 		TShaderMapRef< FScreenVS >		VertexShader(View.ShaderMap);
