Move parameters for computing screen space scales to the View and fix the case of negative scaling in the projection matrix

diff --git a/Engine/Source/Runtime/Engine/Private/SceneView.cpp b/Engine/Source/Runtime/Engine/Private/SceneView.cpp
--- a/Engine/Source/Runtime/Engine/Private/SceneView.cpp
+++ b/Engine/Source/Runtime/Engine/Private/SceneView.cpp
@@ -377,6 +377,17 @@ FSceneView::FSceneView(const FSceneViewInitOptions& InitOptions)
 	// Compute a transform from view origin centered world-space to clip space.
 	ViewMatrices.TranslatedViewProjectionMatrix = TranslatedViewMatrix * ViewMatrices.ProjMatrix;
 	ViewMatrices.InvTranslatedViewProjectionMatrix = ViewMatrices.TranslatedViewProjectionMatrix.Inverse();
+
+	// Compute screen scale factors.
+	// Stereo renders at half horizontal resolution, but compute shadow resolution based on full resolution.
+	const bool bStereo = StereoPass != eSSP_FULL;
+	const float ScreenXScale = bStereo ? 2.0f : 1.0f;
+	ViewMatrices.ProjectionScale.X = ScreenXScale * FMath::Abs(ViewMatrices.ProjMatrix.M[0][0]);
+	ViewMatrices.ProjectionScale.Y = FMath::Abs(ViewMatrices.ProjMatrix.M[1][1]);
+	ViewMatrices.ScreenScale = FMath::Max(
+		ViewRect.Size().X * 0.5f * ViewMatrices.ProjectionScale.X,
+		ViewRect.Size().Y * 0.5f * ViewMatrices.ProjectionScale.Y
+		);
 	
 	ShadowViewMatrices = ViewMatrices;
 
diff --git a/Engine/Source/Runtime/Engine/Public/SceneView.h b/Engine/Source/Runtime/Engine/Public/SceneView.h
--- a/Engine/Source/Runtime/Engine/Public/SceneView.h
+++ b/Engine/Source/Runtime/Engine/Public/SceneView.h
@@ -168,6 +168,17 @@ struct FViewMatrices
 	FVector		PreViewTranslation;
 	/** To support ortho and other modes this is redundant, in world space */
 	FVector		ViewOrigin;
+	/** Scale applied by the projection matrix in X and Y. */
+	FVector2D ProjectionScale;
+	/**
+	 * Scale factor to use when computing the size of a sphere in pixels.
+	 * 
+	 * A common calculation is to determine the size of a sphere in pixels when projected on the screen:
+	 *		ScreenRadius = max(0.5 * ViewSizeX * ProjMatrix[0][0], 0.5 * ViewSizeY * ProjMatrix[1][1]) * SphereRadius / ProjectedSpherePosition.W
+	 * Instead you can now simply use:
+	 *		ScreenRadius = ScreenScale * SphereRadius / ProjectedSpherePosition.W
+	 */
+	float ScreenScale;
 
 	//
 	// World = TranslatedWorld - PreViewTranslation
diff --git a/Engine/Source/Runtime/Renderer/Private/ShadowSetup.cpp b/Engine/Source/Runtime/Renderer/Private/ShadowSetup.cpp
--- a/Engine/Source/Runtime/Renderer/Private/ShadowSetup.cpp
+++ b/Engine/Source/Runtime/Renderer/Private/ShadowSetup.cpp
@@ -1292,22 +1292,15 @@ void FDeferredShadingSceneRenderer::CreatePerObjectProjectedShadow(
 	{
 		const FViewInfo& View = Views[ViewIndex];
 
-		// Stereo renders at half horizontal resolution, but compute shadow resolution based on full resolution.
-		const bool bStereo = View.StereoPass != eSSP_FULL;
-		const float ScreenXScale = bStereo ? 2.0f : 1.0f;
-
 		// Determine the size of the subject's bounding sphere in this view.
 		const FVector4 ScreenPosition = View.WorldToScreen(OriginalBounds.Origin);
-		const float ScreenRadius = FMath::Max(
-			ScreenXScale * View.ViewRect.Size().X / 2.0f * View.ShadowViewMatrices.ProjMatrix.M[0][0],
-			View.ViewRect.Size().Y / 2.0f * View.ShadowViewMatrices.ProjMatrix.M[1][1]
-			) *
+		const float ScreenRadius = View.ShadowViewMatrices.ScreenScale *
 			OriginalBounds.SphereRadius /
 			FMath::Max(ScreenPosition.W,1.0f);
 
 		const float ScreenPercent = FMath::Max(
-			1.0f / 2.0f * View.ShadowViewMatrices.ProjMatrix.M[0][0],
-			1.0f / 2.0f * View.ShadowViewMatrices.ProjMatrix.M[1][1]
+			1.0f / 2.0f * View.ShadowViewMatrices.ProjectionScale.X,
+			1.0f / 2.0f * View.ShadowViewMatrices.ProjectionScale.Y
 			) *
 			OriginalBounds.SphereRadius /
 			FMath::Max(ScreenPosition.W,1.0f);
@@ -1559,16 +1552,9 @@ void FDeferredShadingSceneRenderer::CreateWholeSceneProjectedShadow(FLightSceneI
 		{
 			const FViewInfo& View = Views[ViewIndex];
 
-			// Stereo renders at half horizontal resolution, but compute shadow resolution based on full resolution.
-			const bool bStereo = View.StereoPass != eSSP_FULL;
-			const float ScreenXScale = bStereo ? 2.0f : 1.0f;
-
 			// Determine the size of the light's bounding sphere in this view.
 			const FVector4 ScreenPosition = View.WorldToScreen(LightSceneInfo->Proxy->GetOrigin());
-			const float ScreenRadius = FMath::Max(
-				ScreenXScale * View.ViewRect.Width() / 2.0f * View.ShadowViewMatrices.ProjMatrix.M[0][0],
-				View.ViewRect.Height() / 2.0f * View.ShadowViewMatrices.ProjMatrix.M[1][1]
-				) *
+			const float ScreenRadius = View.ShadowViewMatrices.ScreenScale *
 				LightSceneInfo->Proxy->GetRadius() /
 				FMath::Max(ScreenPosition.W,1.0f);
 
