Non-unity compile fix

diff --git a/Engine/Source/Runtime/Renderer/Private/DistanceFieldObjectManagement.cpp b/Engine/Source/Runtime/Renderer/Private/DistanceFieldObjectManagement.cpp
--- a/Engine/Source/Runtime/Renderer/Private/DistanceFieldObjectManagement.cpp
+++ b/Engine/Source/Runtime/Renderer/Private/DistanceFieldObjectManagement.cpp
@@ -4,6 +4,8 @@
 	DistanceFieldSurfaceCacheLighting.cpp
 =============================================================================*/
 
+#include "RendererPrivate.h"
+
 float GAOMaxObjectBoundingRadius = 50000;
 FAutoConsoleVariableRef CVarAOMaxObjectBoundingRadius(
 	TEXT("r.AOMaxObjectBoundingRadius"),
