Fixed incorrect usage of mem stack

diff --git a/Engine/Source/Runtime/Renderer/Private/DistanceFieldObjectManagement.cpp b/Engine/Source/Runtime/Renderer/Private/DistanceFieldObjectManagement.cpp
--- a/Engine/Source/Runtime/Renderer/Private/DistanceFieldObjectManagement.cpp
+++ b/Engine/Source/Runtime/Renderer/Private/DistanceFieldObjectManagement.cpp
@@ -527,7 +527,7 @@ void UpdateGlobalDistanceFieldObjectRemoves(FRHICommandListImmediate& RHICmdList
 
 	if (DistanceFieldSceneData.PendingRemoveOperations.Num() > 0)
 	{
-		static TArray<int32, SceneRenderingAllocator> PendingRemoveOperations;
+		TArray<int32, SceneRenderingAllocator> PendingRemoveOperations;
 
 		for (int32 RemoveIndex = 0; RemoveIndex < DistanceFieldSceneData.PendingRemoveOperations.Num(); RemoveIndex++)
 		{
