UE4 - Skip outputting velocities during camera cuts when not needed

diff --git a/Engine/Source/Runtime/Renderer/Private/DeferredShadingRenderer.cpp b/Engine/Source/Runtime/Renderer/Private/DeferredShadingRenderer.cpp
--- a/Engine/Source/Runtime/Renderer/Private/DeferredShadingRenderer.cpp
+++ b/Engine/Source/Runtime/Renderer/Private/DeferredShadingRenderer.cpp
@@ -462,7 +462,16 @@ static void SetupBasePassView(FRHICommandList& RHICmdList, const FIntRect& ViewR
 	else
 	{
 		// Opaque blending for all G buffer targets, depth tests and writes.
-		RHICmdList.SetBlendState(TStaticBlendStateWriteMask<CW_RGBA, CW_RGBA, CW_RGBA, CW_RGBA>::GetRHI());
+		static const auto CVar = IConsoleManager::Get().FindTConsoleVariableDataInt(TEXT("r.BasePassOutputsVelocityDebug"));
+		if (CVar && CVar->GetValueOnRenderThread() == 2)
+		{
+			RHICmdList.SetBlendState(TStaticBlendStateWriteMask<CW_RGBA, CW_RGBA, CW_RGBA, CW_RGBA, CW_RGBA, CW_RGBA, CW_NONE>::GetRHI());
+		}
+		else
+		{
+			RHICmdList.SetBlendState(TStaticBlendStateWriteMask<CW_RGBA, CW_RGBA, CW_RGBA, CW_RGBA>::GetRHI());
+		}
+
 		// Note, this is a reversed Z depth surface, using CF_GreaterEqual.
 		RHICmdList.SetDepthStencilState(TStaticDepthStencilState<true,CF_GreaterEqual>::GetRHI());
 	}
@@ -738,6 +747,9 @@ void FDeferredShadingSceneRenderer::Render(FRHICommandListImmediate& RHICmdList)
 	}
 	RHICmdList.SetCurrentStat(GET_STATID(STAT_CLM_AfterPrePass));
 	
+	const bool bShouldRenderVelocities = ShouldRenderVelocities();
+	const bool bUseVelocityGBuffer = FVelocityRendering::OutputsToGBuffer();
+	GSceneRenderTargets.PreallocGBufferTargets(bShouldRenderVelocities && bUseVelocityGBuffer);
 	GSceneRenderTargets.AllocGBufferTargets();
 	
 	// Clear LPVs for all views
@@ -1008,11 +1020,11 @@ void FDeferredShadingSceneRenderer::Render(FRHICommandListImmediate& RHICmdList)
 
 	TRefCountPtr<IPooledRenderTarget> VelocityRT;
 
-	if (FVelocityRendering::OutputsToGBuffer())
+	if (bUseVelocityGBuffer)
 	{
 		VelocityRT = GSceneRenderTargets.GetGBufferVelocityRT();
 	}
-	else
+	else if (bShouldRenderVelocities)
 	{
 		// Render the velocities of movable objects
 		RHICmdList.SetCurrentStat(GET_STATID(STAT_CLM_Velocity));
diff --git a/Engine/Source/Runtime/Renderer/Private/DeferredShadingRenderer.h b/Engine/Source/Runtime/Renderer/Private/DeferredShadingRenderer.h
--- a/Engine/Source/Runtime/Renderer/Private/DeferredShadingRenderer.h
+++ b/Engine/Source/Runtime/Renderer/Private/DeferredShadingRenderer.h
@@ -232,6 +232,8 @@ private:
 	/** Reuses an existing translucent shadow map if possible or re-renders one if necessary. */
 	const FProjectedShadowInfo* PrepareTranslucentShadowMap(FRHICommandList& RHICmdList, const FViewInfo& View, FPrimitiveSceneInfo* PrimitiveSceneInfo, bool bSeparateTranslucencyPass);
 
+	bool ShouldRenderVelocities() const;
+
 	/** Renders the velocities of movable objects for the motion blur effect. */
 	void RenderVelocities(FRHICommandListImmediate& RHICmdList, TRefCountPtr<IPooledRenderTarget>& VelocityRT);
 
diff --git a/Engine/Source/Runtime/Renderer/Private/PostProcess/SceneRenderTargets.cpp b/Engine/Source/Runtime/Renderer/Private/PostProcess/SceneRenderTargets.cpp
--- a/Engine/Source/Runtime/Renderer/Private/PostProcess/SceneRenderTargets.cpp
+++ b/Engine/Source/Runtime/Renderer/Private/PostProcess/SceneRenderTargets.cpp
@@ -18,6 +18,17 @@ const int ReflectiveShadowMapResolution = 256;
 
 IMPLEMENT_UNIFORM_BUFFER_STRUCT(FGBufferResourceStruct,TEXT("GBuffers"));
 
+static TAutoConsoleVariable<int32> CVarBasePassOutputsVelocityDebug(
+	TEXT("r.BasePassOutputsVelocityDebug"),
+	0,
+	TEXT("Debug settings for Base Pass outputting velocity.\n") \
+	TEXT("0 - Regular rendering\n") \
+	TEXT("1 - Skip setting GBufferVelocity RT\n") \
+	TEXT("2 - Set Color Mask 0 for GBufferVelocity RT\n"),
+	ECVF_RenderThreadSafe);
+
+static int32 GBasePassOutputsVelocityDebug = 0;
+
 /*-----------------------------------------------------------------------------
 FSceneRenderTargets
 -----------------------------------------------------------------------------*/
@@ -288,8 +299,6 @@ void FSceneRenderTargets::BeginRenderingSceneColor(FRHICommandList& RHICmdList,
 
 int32 FSceneRenderTargets::GetGBufferRenderTargets(ERenderTargetLoadAction ColorLoadAction, FRHIRenderTargetView OutRenderTargets[MaxSimultaneousRenderTargets], int32& OutVelocityRTIndex)
 {
-	const bool bUseVelocityGBuffer = FVelocityRendering::OutputsToGBuffer();
-
 	OutRenderTargets[0] = FRHIRenderTargetView(GetSceneColorSurface(), 0, -1, ColorLoadAction, ERenderTargetStoreAction::EStore);
 	OutRenderTargets[1] = FRHIRenderTargetView(GSceneRenderTargets.GBufferA->GetRenderTargetItem().TargetableTexture, 0, -1, ColorLoadAction, ERenderTargetStoreAction::EStore);
 	OutRenderTargets[2] = FRHIRenderTargetView(GSceneRenderTargets.GBufferB->GetRenderTargetItem().TargetableTexture, 0, -1, ColorLoadAction, ERenderTargetStoreAction::EStore);
@@ -305,7 +314,7 @@ int32 FSceneRenderTargets::GetGBufferRenderTargets(ERenderTargetLoadAction Color
 		++MRTCount;
 	}
 
-	if (bUseVelocityGBuffer)
+	if (bAllocateVelocityGBuffer)
 	{
 		OutVelocityRTIndex = MRTCount;
 		++MRTCount;
@@ -405,7 +414,7 @@ int32 FSceneRenderTargets::GetNumGBufferTargets() const
 	{
 		NumGBufferTargets = bAllowStaticLighting ? 6 : 5;
 
-		if (FVelocityRendering::OutputsToGBuffer())
+		if (bAllocateVelocityGBuffer)
 		{
 			++NumGBufferTargets;
 		}
@@ -506,12 +515,24 @@ void FSceneRenderTargets::ReleaseGBufferTargets()
 	GBufferVelocity.SafeRelease();
 }
 
+void FSceneRenderTargets::PreallocGBufferTargets(bool bShouldRenderVelocities)
+{
+	if (GBasePassOutputsVelocityDebug == 1)
+	{
+		bAllocateVelocityGBuffer = false;
+	}
+	else
+	{
+		bAllocateVelocityGBuffer = bShouldRenderVelocities && FVelocityRendering::OutputsToGBuffer();
+	}
+}
+
 void FSceneRenderTargets::AllocGBufferTargets()
 {
 	// AdjustGBufferRefCount +1 doesn't match -1 (within the same frame)
 	ensure(GBufferRefCount == 0);
 
-	if(GBufferA)
+	if (GBufferA)
 	{
 		// no work needed
 		return;
@@ -571,8 +592,8 @@ void FSceneRenderTargets::AllocGBufferTargets()
 		GRenderTargetPool.FindFreeElement(Desc, GBufferE, TEXT("GBufferE"));
 	}
 
-	const bool bUseVelocityGBuffer = FVelocityRendering::OutputsToGBuffer();
-	if (bUseVelocityGBuffer)
+	GBasePassOutputsVelocityDebug = CVarBasePassOutputsVelocityDebug.GetValueOnRenderThread();
+	if (bAllocateVelocityGBuffer)
 	{
 		FPooledRenderTargetDesc VelocityRTDesc = FVelocityRendering::GetRenderTargetDesc();
 		GRenderTargetPool.FindFreeElement(VelocityRTDesc, GBufferVelocity, TEXT("GBufferVelocity"));
@@ -861,7 +882,7 @@ void FSceneRenderTargets::ResolveGBufferSurfaces(FRHICommandList& RHICmdList, co
 			RHICmdList.CopyToResolveTarget(GSceneRenderTargets.GBufferE->GetRenderTargetItem().TargetableTexture, GSceneRenderTargets.GBufferE->GetRenderTargetItem().ShaderResourceTexture, true, FResolveParams(ResolveRect));
 		}
 
-		if (FVelocityRendering::OutputsToGBuffer())
+		if (bAllocateVelocityGBuffer)
 		{
 			RHICmdList.CopyToResolveTarget(GSceneRenderTargets.GBufferVelocity->GetRenderTargetItem().TargetableTexture, GSceneRenderTargets.GBufferVelocity->GetRenderTargetItem().ShaderResourceTexture, true, FResolveParams(ResolveRect));
 		}
@@ -1491,7 +1512,7 @@ void FSceneRenderTargets::AllocateDeferredShadingPathRenderTargets()
 		}
 	}
 
-	if (FVelocityRendering::OutputsToGBuffer())
+	if (bAllocateVelocityGBuffer)
 	{
 		FPooledRenderTargetDesc VelocityRTDesc = FVelocityRendering::GetRenderTargetDesc();
 		GRenderTargetPool.FindFreeElement(VelocityRTDesc, GBufferVelocity, TEXT("GBufferVelocity"));
@@ -1678,7 +1699,7 @@ const FTextureRHIRef& FSceneRenderTargets::GetSceneColorTexture() const
 
 IPooledRenderTarget* FSceneRenderTargets::GetGBufferVelocityRT()
 {
-	if (!FVelocityRendering::OutputsToGBuffer())
+	if (!bAllocateVelocityGBuffer)
 	{
 		return nullptr;
 	}
diff --git a/Engine/Source/Runtime/Renderer/Private/PostProcess/SceneRenderTargets.h b/Engine/Source/Runtime/Renderer/Private/PostProcess/SceneRenderTargets.h
--- a/Engine/Source/Runtime/Renderer/Private/PostProcess/SceneRenderTargets.h
+++ b/Engine/Source/Runtime/Renderer/Private/PostProcess/SceneRenderTargets.h
@@ -134,7 +134,8 @@ protected:
 		CurrentMobile32bpp(0),
 		bCurrentLightPropagationVolume(false),
 		CurrentFeatureLevel(ERHIFeatureLevel::Num),
-		CurrentShadingPath(EShadingPath::Num)
+		CurrentShadingPath(EShadingPath::Num),
+		bAllocateVelocityGBuffer(false)
 		{
 		}
 public:
@@ -428,6 +429,7 @@ public:
 	void AdjustGBufferRefCount(int Delta);
 
 	//
+	void PreallocGBufferTargets(bool bShouldRenderVelocities);
 	void AllocGBufferTargets();
 
 	void AllocLightAttenuation();
@@ -624,6 +626,9 @@ private:
 	/** Shading path that we are currently drawing through. Set when calling Allocate at the start of a scene render. */
 	EShadingPath CurrentShadingPath;
 
+	// Set this per frame since there might be cases where we don't need an extra GBuffer
+	bool bAllocateVelocityGBuffer;
+
 	/** Helpers to track gbuffer state on platforms that need to propagate clear information across parallel rendering boundaries. */
 	bool bGBuffersCleared;
 	FLinearColor GBufferClearColor;
diff --git a/Engine/Source/Runtime/Renderer/Private/VelocityRendering.cpp b/Engine/Source/Runtime/Renderer/Private/VelocityRendering.cpp
--- a/Engine/Source/Runtime/Renderer/Private/VelocityRendering.cpp
+++ b/Engine/Source/Runtime/Renderer/Private/VelocityRendering.cpp
@@ -699,24 +699,33 @@ static TAutoConsoleVariable<int32> CVarParallelVelocity(
 	ECVF_RenderThreadSafe
 	);
 
-
-void FDeferredShadingSceneRenderer::RenderVelocities(FRHICommandListImmediate& RHICmdList, TRefCountPtr<IPooledRenderTarget>& VelocityRT)
+bool FDeferredShadingSceneRenderer::ShouldRenderVelocities() const
 {
-	check(FeatureLevel >= ERHIFeatureLevel::SM4);
-	SCOPE_CYCLE_COUNTER(STAT_RenderVelocities);
+	if (!GPixelFormats[PF_G16R16].Supported)
+	{
+		return false;
+	}
 
 	bool bNeedsVelocity = false;
-	for(int32 ViewIndex = 0; ViewIndex < Views.Num(); ViewIndex++)
+	for (int32 ViewIndex = 0; ViewIndex < Views.Num(); ViewIndex++)
 	{
 		const FViewInfo& View = Views[ViewIndex];
-	
+
 		bool bTemporalAA = (View.FinalPostProcessSettings.AntiAliasingMethod == AAM_TemporalAA) && !View.bCameraCut;
 		bool bMotionBlur = IsMotionBlurEnabled(View);
 
 		bNeedsVelocity |= bMotionBlur || bTemporalAA;
 	}
 
-	if( !bNeedsVelocity || !GPixelFormats[PF_G16R16].Supported )
+	return bNeedsVelocity;
+}
+
+void FDeferredShadingSceneRenderer::RenderVelocities(FRHICommandListImmediate& RHICmdList, TRefCountPtr<IPooledRenderTarget>& VelocityRT)
+{
+	check(FeatureLevel >= ERHIFeatureLevel::SM4);
+	SCOPE_CYCLE_COUNTER(STAT_RenderVelocities);
+
+	if (!ShouldRenderVelocities())
 	{
 		return;
 	}
