Fix splitscreen issues in ShooterGame on PS4/Xbox.

diff --git a/Engine/Source/Runtime/Renderer/Private/DeferredShadingRenderer.cpp b/Engine/Source/Runtime/Renderer/Private/DeferredShadingRenderer.cpp
--- a/Engine/Source/Runtime/Renderer/Private/DeferredShadingRenderer.cpp
+++ b/Engine/Source/Runtime/Renderer/Private/DeferredShadingRenderer.cpp
@@ -477,7 +477,7 @@ static void SetupBasePassView(FRHICommandList& RHICmdList, const FIntRect& ViewR
 	}
 	RHICmdList.SetViewport(ViewRect.Min.X, ViewRect.Min.Y, 0, ViewRect.Max.X, ViewRect.Max.Y, 1);
 	RHICmdList.SetRasterizerState(TStaticRasterizerState<FM_Solid, CM_None>::GetRHI());
-	RHICmdList.SetScissorRect(false, 0, 0, 0, 0);
+	RHICmdList.SetScissorRect(true, ViewRect.Min.X, ViewRect.Min.Y, ViewRect.Max.X, ViewRect.Max.Y);
 }
 
 class FBasePassParallelCommandListSet : public FParallelCommandListSet
@@ -1112,7 +1112,7 @@ static void SetupPrePassView(FRHICommandList& RHICmdList, const FIntRect& ViewRe
 	RHICmdList.SetDepthStencilState(TStaticDepthStencilState<true,CF_GreaterEqual>::GetRHI());
 	RHICmdList.SetViewport(ViewRect.Min.X, ViewRect.Min.Y, 0, ViewRect.Max.X, ViewRect.Max.Y, 1);
 	RHICmdList.SetRasterizerState(TStaticRasterizerState<FM_Solid, CM_None>::GetRHI());
-	RHICmdList.SetScissorRect(false, 0, 0, 0, 0);
+	RHICmdList.SetScissorRect(true, ViewRect.Min.X, ViewRect.Min.Y, ViewRect.Max.X, ViewRect.Max.Y);
 }
 
 bool FDeferredShadingSceneRenderer::RenderPrePassView(FRHICommandList& RHICmdList, const FViewInfo& View)
