Fix stereo rendering being broken by SetScissorRect

diff --git a/Engine/Source/Runtime/Renderer/Private/DeferredShadingRenderer.cpp b/Engine/Source/Runtime/Renderer/Private/DeferredShadingRenderer.cpp
--- a/Engine/Source/Runtime/Renderer/Private/DeferredShadingRenderer.cpp
+++ b/Engine/Source/Runtime/Renderer/Private/DeferredShadingRenderer.cpp
@@ -475,9 +475,9 @@ static void SetupBasePassView(FRHICommandList& RHICmdList, const FIntRect& ViewR
 		// Note, this is a reversed Z depth surface, using CF_GreaterEqual.
 		RHICmdList.SetDepthStencilState(TStaticDepthStencilState<true,CF_GreaterEqual>::GetRHI());
 	}
+	RHICmdList.SetScissorRect(false, 0, 0, 0, 0);
 	RHICmdList.SetViewport(ViewRect.Min.X, ViewRect.Min.Y, 0, ViewRect.Max.X, ViewRect.Max.Y, 1);
 	RHICmdList.SetRasterizerState(TStaticRasterizerState<FM_Solid, CM_None>::GetRHI());
-	RHICmdList.SetScissorRect(true, ViewRect.Min.X, ViewRect.Min.Y, ViewRect.Max.X, ViewRect.Max.Y);
 }
 
 class FBasePassParallelCommandListSet : public FParallelCommandListSet
