Heightfield dynamic Global Illumination
* Enable with 'r.DistanceFieldGI 1' on console
* Surfel GI is currently disabled as it doesn't interop yet
* A GBuffer atlas is created for Landscape from components that can affect indirect lighting based on the current view
* Shadowing from a directional light is computed by reusing shadowmaps and ray tracing distance fields
* Lighting is computed for the atlas using GBuffer properties
* First bounce GI light transfer is computed for a receiver by iterating over all the nearby lit texels of the heightfield and treating them as Virtual Point Lights.  Two gather regions are used, the larger one skips through the heightfield to achieve a larger bounce distance.  Distance field AO is used to provide local shadowing.
* Material diffuse color not yet supported, and regions outside of the view don't have valid shadowing.

diff --git a/Engine/Shaders/Common.usf b/Engine/Shaders/Common.usf
--- a/Engine/Shaders/Common.usf
+++ b/Engine/Shaders/Common.usf
@@ -481,8 +481,6 @@ SamplerState	LightAttenuationTextureSampler;
 // We don't use a function so we can avoid type promotion/ coercion.
 #define TangentBias(X) (X * 2.0f - 1.0f)
 
-#define TERRAIN_ZSCALE				(1.0f/128.0f)
-
 float Square( float x )
 {
 	return x*x;
@@ -1114,3 +1112,20 @@ float IsShadowDepthShader()
 	return 0;
 #endif
 }
+#define TERRAIN_ZSCALE (1.0f/128.0f)
+
+// Decodes a value which was packed into two 8 bit channels
+float DecodePackedTwoChannelValue(float2 PackedHeight)
+{
+	return PackedHeight.x * 255.0 * 256.0 + PackedHeight.y * 255.0;
+}
+
+float DecodeHeightValue(float InValue)
+{
+	return (InValue - 32768.0) * TERRAIN_ZSCALE;
+}
+
+float DecodePackedHeight(float2 PackedHeight)
+{
+	return DecodeHeightValue(DecodePackedTwoChannelValue(PackedHeight));
+}
diff --git a/Engine/Shaders/DistanceFieldGlobalIllumination.usf b/Engine/Shaders/DistanceFieldGlobalIllumination.usf
--- a/Engine/Shaders/DistanceFieldGlobalIllumination.usf
+++ b/Engine/Shaders/DistanceFieldGlobalIllumination.usf
@@ -9,6 +9,7 @@
 #include "DistanceFieldLightingShared.usf"
 #include "DistanceFieldAOShared.usf"
 #include "DynamicLightingCommon.usf"
+#include "DistanceFieldShadowingShared.usf"
 
 #define USE_SHADOW_CULLING_FOR_VPL_PLACEMENT 1
 
@@ -161,7 +162,6 @@ float3 ComputeDistanceFieldNormal(float3 WorldPosition, uint ObjectIndex)
 /** From light source, into world. */
 float4 LightDirectionAndTraceDistance;
 float4 LightColor;
-float4x4 WorldToShadow;
 float4x4 ShadowToWorld;
 float2 InvPlacementGridSize;
 float VPLPlacementCameraRadius;
@@ -192,16 +192,7 @@ void VPLPlacementCS(
 
 		#if USE_SHADOW_CULLING_FOR_VPL_PLACEMENT
 
-			// Transform into shadow space
-			float4 HomogeneousShadowPosition = mul(float4(CellStartWorldPosition, 1), WorldToShadow);
-			float2 NormalizedShadowPosition = HomogeneousShadowPosition.xy * .5f + .5f;
-			NormalizedShadowPosition.y = 1 - NormalizedShadowPosition.y;
-			// Quantize the shadow position to get our tile position
-			uint2 TilePosition = (uint2)(NormalizedShadowPosition * ShadowTileListGroupSize);
-			// Fetch the tile head information
-			uint2 TileHead = GetShadowTileHead(TilePosition);
-			NumIntersectingObjects = TileHead.y;
-			CulledDataStart = TileHead.x;
+			GetShadowTileCulledData(CellStartWorldPosition, CulledDataStart, NumIntersectingObjects);
 
 		#endif
 
@@ -484,16 +475,8 @@ float3 ComputeVPLFlux(uint SurfelIndex, float4x4 InstanceToWorld)
 
 					#if USE_SHADOW_CULLING_FOR_VPL_PLACEMENT
 
-						// Transform into shadow space
-						float4 HomogeneousShadowPosition = mul(float4(WorldRayStart, 1), WorldToShadow);
-						float2 NormalizedShadowPosition = HomogeneousShadowPosition.xy * .5f + .5f;
-						NormalizedShadowPosition.y = 1 - NormalizedShadowPosition.y;
-						// Quantize the shadow position to get our tile position
-						uint2 TilePosition = (uint2)(NormalizedShadowPosition * ShadowTileListGroupSize);
-						// Fetch the tile head information
-						uint2 TileHead = GetShadowTileHead(TilePosition);
-						NumIntersectingObjects = TileHead.y;
-						CulledDataStart = TileHead.x;
+						uint CulledDataStart;
+						GetShadowTileCulledData(WorldRayStart, CulledDataStart, NumIntersectingObjects);
 
 					#endif
 				}
@@ -511,9 +494,9 @@ float3 ComputeVPLFlux(uint SurfelIndex, float4x4 InstanceToWorld)
 	if (any(AccumulatedLighting > 0))
 	{
 		float3 DiffuseColor = LoadSurfelDiffuseColor(SurfelIndex);
-		float3 EmissiveColor = LoadSurfelDiffuseColor(SurfelIndex);
+		//float3 EmissiveColor = LoadSurfelEmissiveColor(SurfelIndex);
 
-		Flux = (DiffuseColor * AccumulatedLighting + EmissiveColor) * PI * PositionAndRadius.w * PositionAndRadius.w;
+		Flux = (DiffuseColor * AccumulatedLighting/* + EmissiveColor*/) * PI * PositionAndRadius.w * PositionAndRadius.w;
 	}
 
 	return Flux;
@@ -589,6 +572,22 @@ void LightVPLsCS(
 	}
 }
 
+RWBuffer<float4> RWIrradianceCacheIrradiance;
+
+/**  */
+[numthreads(FINAL_GATHER_THREADGROUP_SIZE, 1, 1)]
+void ClearIrradianceSamplesCS(
+	uint3 GroupId : SV_GroupID,
+	uint3 DispatchThreadId : SV_DispatchThreadID,
+    uint3 GroupThreadId : SV_GroupThreadID) 
+{
+	uint StartIndex = SavedStartIndex[0];
+	uint NumRecords = ScatterDrawParameters[1];
+	uint RecordIndex = StartIndex + DispatchThreadId.x;
+
+	RWIrradianceCacheIrradiance[RecordIndex] = 0;
+}
+
 float RayTraceThroughGlobalObjects(
 	float3 RayStartPosition, 
 	float3 RayDirection, 
@@ -681,33 +680,6 @@ float RayTraceThroughGlobalObjects(
 	return MinVisibility;
 }
 
-float ApproximateConeConeIntersection(float ArcLength0, float ArcLength1, float AngleBetweenCones)
-{
-	float AngleDifference = abs(ArcLength0 - ArcLength1);
-
-	float Intersection = smoothstep(
-		0,
-		1.0,
-		1.0 - saturate((AngleBetweenCones - AngleDifference) / (ArcLength0 + ArcLength1 - AngleDifference)));
-
-	return Intersection;
-}
-
-float GetVPLOcclusion(float3 BentNormalAO, float3 NormalizedVectorToVPL)
-{
-	float BentNormalLength = length(BentNormalAO);
-	float VPLConeAngle = .5f;
-	float VPLOcclusionStrength = 1;
-	float UnoccludedAngle = BentNormalLength * PI / VPLOcclusionStrength;
-	float AngleBetween = acos(dot(BentNormalAO, NormalizedVectorToVPL) / max(BentNormalLength, .0001f));
-	float Visibility = ApproximateConeConeIntersection(VPLConeAngle, UnoccludedAngle, AngleBetween);
-
-	// Can't rely on the direction of the bent normal when close to fully occluded, lerp to shadowed
-	Visibility = lerp(0, Visibility, saturate((UnoccludedAngle - .1f) / .2f));
-
-	return Visibility;
-}
-
 Buffer<float> RecordConeData;
 RWBuffer<float4> RWStepBentNormal;
 float BentNormalNormalizeFactor;
@@ -987,8 +959,6 @@ float3 GatherIrradianceFromVPLs(float3 WorldPosition, float3 WorldNormal, uint R
 	return Irradiance;
 }
 
-RWBuffer<float4> RWIrradianceCacheIrradiance;
-
 groupshared float3 SharedThreadIrradiance[FINAL_GATHER_THREADGROUP_SIZE];
 
 /**  */
diff --git a/Engine/Shaders/DistanceFieldLightingShared.usf b/Engine/Shaders/DistanceFieldLightingShared.usf
--- a/Engine/Shaders/DistanceFieldLightingShared.usf
+++ b/Engine/Shaders/DistanceFieldLightingShared.usf
@@ -216,19 +216,6 @@ void LoadObjectAxes(uint ObjectIndex, out float3 ObjectAxisX, out float3 ObjectA
 	ObjectAxisZ = CulledObjectBoxBounds.Load(ObjectIndex * CULLED_OBJECT_BOX_BOUNDS_STRIDE + 4).xyz;
 }
 
-Buffer<uint> ShadowTileHeadDataUnpacked;
-Buffer<uint> ShadowTileArrayData;
-uint2 ShadowTileListGroupSize;
-
-uint2 GetShadowTileHead(uint2 TileCoordinate)
-{
-	uint TileIndex = TileCoordinate.y * ShadowTileListGroupSize.x + TileCoordinate.x;
-
-	return uint2(
-		ShadowTileHeadDataUnpacked[TileIndex * 2 + 0], 
-		min(ShadowTileHeadDataUnpacked[TileIndex * 2 + 1], (uint)MAX_OBJECTS_PER_TILE)); 
-}
-
 void FindBestAxisVectors2(float3 InZAxis, out float3 OutXAxis, out float3 OutYAxis )
 {
 	float3 UpVector = abs(InZAxis.z) < 0.999 ? float3(0,0,1) : float3(1,0,0);
@@ -298,3 +285,29 @@ float3 LoadSurfelEmissiveColor(uint SurfelIndex)
 	return SurfelData[SurfelIndex * SURFEL_DATA_STRIDE + 3].xyz;
 }
 
+float ApproximateConeConeIntersection(float ArcLength0, float ArcLength1, float AngleBetweenCones)
+{
+	float AngleDifference = abs(ArcLength0 - ArcLength1);
+
+	float Intersection = smoothstep(
+		0,
+		1.0,
+		1.0 - saturate((AngleBetweenCones - AngleDifference) / (ArcLength0 + ArcLength1 - AngleDifference)));
+
+	return Intersection;
+}
+
+float GetVPLOcclusion(float3 BentNormalAO, float3 NormalizedVectorToVPL)
+{
+	float BentNormalLength = length(BentNormalAO);
+	float VPLConeAngle = .5f;
+	float VPLOcclusionStrength = 1;
+	float UnoccludedAngle = BentNormalLength * PI / VPLOcclusionStrength;
+	float AngleBetween = acos(dot(BentNormalAO, NormalizedVectorToVPL) / max(BentNormalLength, .0001f));
+	float Visibility = ApproximateConeConeIntersection(VPLConeAngle, UnoccludedAngle, AngleBetween);
+
+	// Can't rely on the direction of the bent normal when close to fully occluded, lerp to shadowed
+	Visibility = lerp(0, Visibility, saturate((UnoccludedAngle - .1f) / .2f));
+
+	return Visibility;
+}
\ No newline at end of file
diff --git a/Engine/Shaders/DistanceFieldShadowing.usf b/Engine/Shaders/DistanceFieldShadowing.usf
--- a/Engine/Shaders/DistanceFieldShadowing.usf
+++ b/Engine/Shaders/DistanceFieldShadowing.usf
@@ -7,6 +7,7 @@
 #include "Common.usf"
 #include "DeferredShadingCommon.usf"
 #include "DistanceFieldLightingShared.usf"
+#include "DistanceFieldShadowingShared.usf"
 
 RWBuffer<float4> RWCulledObjectBounds;
 RWBuffer<float4> RWCulledObjectData;
@@ -18,8 +19,6 @@ float4 FetchObjectDataFloat4(uint SourceIndex)
 	return float4(ObjectData[4 * SourceIndex + 0], ObjectData[4 * SourceIndex + 1], ObjectData[4 * SourceIndex + 2], ObjectData[4 * SourceIndex + 3]);
 }
 
-float4x4 WorldToShadow;
-
 void CopyCulledObjectData(uint DestIndex, uint SourceIndex)
 {
 	RWCulledObjectBounds[DestIndex] = float4(ObjectBounds[4 * SourceIndex + 0], ObjectBounds[4 * SourceIndex + 1], ObjectBounds[4 * SourceIndex + 2], ObjectBounds[4 * SourceIndex + 3]);
@@ -303,7 +302,6 @@ float LightSourceRadius;
 float RayStartOffsetDepthScale;
 float2 TanLightAngleAndNormalThreshold;
 int4 ScissorRectMinAndSize;
-float TwoSidedMeshDistanceBias;
 
 /** Min and Max depth for this tile. */
 groupshared uint IntegerTileMinZ;
@@ -317,9 +315,6 @@ groupshared uint IntegerTileMaxZ2;
 groupshared uint TileNumObjects0;
 groupshared uint TileNumObjects1;
 
-#define MAX_INTERSECTING_OBJECTS 256
-groupshared uint IntersectingObjectIndices[MAX_INTERSECTING_OBJECTS * 2];
-
 void CullObjectsToTileWithGather(
 	float SceneDepth, 
 	uint ThreadIndex, 
@@ -498,6 +493,7 @@ void DistanceFieldShadowingCS(
 	float MinDepth = 0;
 
 	uint NumIntersectingObjects = GetCulledNumObjects();
+	uint CulledDataParameter = 0;
 	bool bTileShouldComputeShadowing = SceneDepth > MinDepth;
 
 #define USE_CULLING 1
@@ -505,20 +501,11 @@ void DistanceFieldShadowingCS(
 
 #if SCATTER_TILE_CULLING
 
-	// Transform into shadow space
-	float4 HomogeneousShadowPosition = mul(float4(OpaqueWorldPosition, 1), WorldToShadow);
-	float2 NormalizedShadowPosition = HomogeneousShadowPosition.xy * .5f + .5f;
-	NormalizedShadowPosition.y = 1 - NormalizedShadowPosition.y;
-	// Quantize the shadow position to get our tile position
-	uint2 TilePosition = (uint2)(NormalizedShadowPosition * ShadowTileListGroupSize);
-	// Fetch the tile head information
-	uint2 TileHead = GetShadowTileHead(TilePosition);
-	NumIntersectingObjects = TileHead.y;
+	GetShadowTileCulledData(OpaqueWorldPosition, CulledDataParameter, NumIntersectingObjects);
 
 #else
 
-	uint GroupIndex;
-	CullObjectsToTileWithGather(SceneDepth, ThreadIndex, GroupId.xy, TraceDistance, MinDepth, NumIntersectingObjects, bTileShouldComputeShadowing, GroupIndex);
+	CullObjectsToTileWithGather(SceneDepth, ThreadIndex, GroupId.xy, TraceDistance, MinDepth, NumIntersectingObjects, bTileShouldComputeShadowing, CulledDataParameter);
 
 #endif
 #endif // USE_CULLING
@@ -560,108 +547,29 @@ void DistanceFieldShadowingCS(
 
 			#endif
 
-			float MinRayTime = MaxRayTime;
-			float MinConeVisibility = 1;
-			float3 WorldRayUnitDirection = normalize(WorldRayEnd - WorldRayStart);
-
-			LOOP
-			for (uint ListObjectIndex = 0; ListObjectIndex < NumIntersectingObjects && MinRayTime >= MaxRayTime; ListObjectIndex++)
-			{
-				#if USE_CULLING
-					#if SCATTER_TILE_CULLING
-						uint ObjectIndex = ShadowTileArrayData.Load(ListObjectIndex * ShadowTileListGroupSize.x * ShadowTileListGroupSize.y + TileHead.x);
-					#else
-						uint ObjectIndex = IntersectingObjectIndices[MAX_INTERSECTING_OBJECTS * GroupIndex + ListObjectIndex];
-					#endif
-				#else
-					uint ObjectIndex = ListObjectIndex;
-				#endif
-
-				float4 SphereCenterAndRadius = LoadObjectPositionAndRadius(ObjectIndex);
-				float ObjectCenterDistanceAlongRay = dot(SphereCenterAndRadius.xyz - WorldRayStart, WorldRayUnitDirection);
-
-				BRANCH
-				if (ObjectCenterDistanceAlongRay > -SphereCenterAndRadius.w)
-				{
-					float3 LocalPositionExtent = LoadObjectLocalPositionExtent(ObjectIndex);
-					float4x4 WorldToVolume = LoadObjectWorldToVolume(ObjectIndex);
-					bool bGeneratedAsTwoSided;
-					float4 UVScaleAndVolumeScale = LoadObjectUVScale(ObjectIndex, bGeneratedAsTwoSided);
-					float3 UVAdd = LoadObjectUVAdd(ObjectIndex);
-
-					float3 VolumeRayStart = mul(float4(WorldRayStart, 1), WorldToVolume).xyz;
-					float3 VolumeRayEnd = mul(float4(WorldRayEnd, 1), WorldToVolume).xyz;
-					float3 VolumeRayDirection = VolumeRayEnd - VolumeRayStart;
-					float VolumeRayLength = length(VolumeRayDirection);
-					VolumeRayDirection /= VolumeRayLength;
-
-					float WorldToVolumeScale = 1.0f / UVScaleAndVolumeScale.w;
-					float VolumeMinSphereRadius = MinSphereRadius * WorldToVolumeScale;
-					float VolumeMaxSphereRadius = MaxSphereRadius * WorldToVolumeScale;
-					float VolumeTwoSidedMeshDistanceBias = TwoSidedMeshDistanceBias * WorldToVolumeScale;
-
-					// Expand the intersection box by the radius of the cone at the distance of the object along the cone
-					float LocalConeRadiusAtObject = min(TanLightAngle * max(ObjectCenterDistanceAlongRay, 0) * WorldToVolumeScale, VolumeMaxSphereRadius);
-
-					float2 IntersectionTimes = LineBoxIntersect(VolumeRayStart, VolumeRayEnd, -LocalPositionExtent - LocalConeRadiusAtObject, LocalPositionExtent + LocalConeRadiusAtObject);
-
-					BRANCH
-					if (IntersectionTimes.x < IntersectionTimes.y && IntersectionTimes.x < 1)
-					{
-						float SampleRayTime = IntersectionTimes.x * VolumeRayLength;
-						uint MaxSteps = 64;
-						float MinStepSize = 1.0f / (4 * MaxSteps);
-
-						float MinDistance = 1000000;
-						float3 IntersectionPosition = float3(0, 0, 0);
-
-						uint StepIndex = 0;
-
-						LOOP
-						for (; StepIndex < MaxSteps; StepIndex++)
-						{
-							float3 SampleVolumePosition = VolumeRayStart + VolumeRayDirection * SampleRayTime;
-							float3 ClampedSamplePosition = clamp(SampleVolumePosition, -LocalPositionExtent, LocalPositionExtent);
-							float DistanceToClamped = length(ClampedSamplePosition - SampleVolumePosition);
-							float3 VolumeUV = DistanceFieldVolumePositionToUV(ClampedSamplePosition, UVScaleAndVolumeScale.xyz, UVAdd);
-							float DistanceField = Texture3DSampleLevel(DistanceFieldTexture, DistanceFieldSampler, VolumeUV, 0).x + DistanceToClamped;
-
-							FLATTEN
-							if (bGeneratedAsTwoSided)
-							{
-								DistanceField -= VolumeTwoSidedMeshDistanceBias;
-							}
-
-							MinDistance = min(MinDistance, DistanceField);
-							float SphereRadius = clamp(TanLightAngle * SampleRayTime, VolumeMinSphereRadius, VolumeMaxSphereRadius);
-
-							MinConeVisibility = min(MinConeVisibility, saturate(DistanceField / SphereRadius));
-							IntersectionPosition = SampleVolumePosition;
-
-							float StepDistance = max(DistanceField, MinStepSize);
-							SampleRayTime += StepDistance;
-
-							// Terminate the trace if we reached a negative area or went past the end of the ray
-							if (DistanceField < 0 
-								|| SampleRayTime > IntersectionTimes.y * VolumeRayLength)
-							{
-								break;
-							}
-						}
-
-						if (MinDistance < 0 || StepIndex == MaxSteps)
-						{
-							MinConeVisibility = 0;
-							MinRayTime = min(MinRayTime, SampleRayTime * UVScaleAndVolumeScale.w);
-						}
-
-						// Force to shadowed as we approach max steps
-						MinConeVisibility = min(MinConeVisibility, (1 - StepIndex / (float)MaxSteps));
-					}
-				}
-			}
+			#if SCATTER_TILE_CULLING
+				bool bUseScatterTileCulling = true;
+			#else
+				bool bUseScatterTileCulling = false;
+			#endif
+		
+			#if USE_CULLING
+				bool bUseCulling = true;
+			#else
+				bool bUseCulling = false;
+			#endif
 
-			Result = MinConeVisibility;
+			Result = ShadowRayTraceThroughCulledObjects(
+				WorldRayStart, 
+				WorldRayEnd, 
+				MaxRayTime, 
+				TanLightAngle, 
+				MinSphereRadius, 
+				MaxSphereRadius, 
+				CulledDataParameter, 
+				NumIntersectingObjects,
+				bUseCulling,
+				bUseScatterTileCulling);
 		}
 	}
 
diff --git a/Engine/Shaders/DistanceFieldShadowingShared.usf b/Engine/Shaders/DistanceFieldShadowingShared.usf
new file mode 100644
--- /dev/null
+++ b/Engine/Shaders/DistanceFieldShadowingShared.usf
@@ -0,0 +1,166 @@
+// Copyright 1998-2015 Epic Games, Inc. All Rights Reserved.
+
+/*=============================================================================
+	DistanceFieldShadowingShared.usf
+=============================================================================*/
+
+Buffer<uint> ShadowTileHeadDataUnpacked;
+Buffer<uint> ShadowTileArrayData;
+uint2 ShadowTileListGroupSize;
+
+uint2 GetShadowTileHead(uint2 TileCoordinate)
+{
+	uint TileIndex = TileCoordinate.y * ShadowTileListGroupSize.x + TileCoordinate.x;
+
+	return uint2(
+		ShadowTileHeadDataUnpacked[TileIndex * 2 + 0], 
+		min(ShadowTileHeadDataUnpacked[TileIndex * 2 + 1], (uint)MAX_OBJECTS_PER_TILE)); 
+}
+
+float4x4 WorldToShadow;
+
+void GetShadowTileCulledData(float3 WorldPosition, out uint CulledDataStart, out uint NumIntersectingObjects)
+{
+	// Transform into shadow space
+	float4 HomogeneousShadowPosition = mul(float4(WorldPosition, 1), WorldToShadow);
+	float2 NormalizedShadowPosition = HomogeneousShadowPosition.xy * .5f + .5f;
+	NormalizedShadowPosition.y = 1 - NormalizedShadowPosition.y;
+	// Quantize the shadow position to get our tile position
+	uint2 TilePosition = (uint2)(NormalizedShadowPosition * ShadowTileListGroupSize);
+	// Fetch the tile head information
+	uint2 TileHead = GetShadowTileHead(TilePosition);
+	CulledDataStart = TileHead.x;
+	NumIntersectingObjects = TileHead.y;
+}
+
+#define MAX_INTERSECTING_OBJECTS 256
+groupshared uint IntersectingObjectIndices[MAX_INTERSECTING_OBJECTS * 2];
+
+float TwoSidedMeshDistanceBias;
+
+float ShadowRayTraceThroughCulledObjects(
+	float3 WorldRayStart, 
+	float3 WorldRayEnd, 
+	float MaxRayTime, 
+	float TanLightAngle, 
+	float MinSphereRadius, 
+	float MaxSphereRadius, 
+	uint CulledDataParameter, 
+	uint NumIntersectingObjects,
+	uniform bool bUseCulling,
+	uniform bool bUseScatterTileCulling)
+{
+	float MinRayTime = MaxRayTime;
+	float MinConeVisibility = 1;
+	float3 WorldRayUnitDirection = normalize(WorldRayEnd - WorldRayStart);
+
+	LOOP
+	for (uint ListObjectIndex = 0; ListObjectIndex < NumIntersectingObjects && MinRayTime >= MaxRayTime; ListObjectIndex++)
+	{
+		uint ObjectIndex;
+
+		if (bUseCulling)
+		{
+			if (bUseScatterTileCulling)
+			{
+				uint CulledDataStart = CulledDataParameter;
+				ObjectIndex = ShadowTileArrayData.Load(ListObjectIndex * ShadowTileListGroupSize.x * ShadowTileListGroupSize.y + CulledDataStart);
+			}
+			else
+			{
+				uint GroupIndex = CulledDataParameter;
+				ObjectIndex = IntersectingObjectIndices[MAX_INTERSECTING_OBJECTS * GroupIndex + ListObjectIndex];
+			}
+		}
+		else
+		{
+			ObjectIndex = ListObjectIndex;
+		}
+
+		float4 SphereCenterAndRadius = LoadObjectPositionAndRadius(ObjectIndex);
+		float ObjectCenterDistanceAlongRay = dot(SphereCenterAndRadius.xyz - WorldRayStart, WorldRayUnitDirection);
+
+		BRANCH
+		if (ObjectCenterDistanceAlongRay > -SphereCenterAndRadius.w)
+		{
+			float3 LocalPositionExtent = LoadObjectLocalPositionExtent(ObjectIndex);
+			float4x4 WorldToVolume = LoadObjectWorldToVolume(ObjectIndex);
+			bool bGeneratedAsTwoSided;
+			float4 UVScaleAndVolumeScale = LoadObjectUVScale(ObjectIndex, bGeneratedAsTwoSided);
+			float3 UVAdd = LoadObjectUVAdd(ObjectIndex);
+
+			float3 VolumeRayStart = mul(float4(WorldRayStart, 1), WorldToVolume).xyz;
+			float3 VolumeRayEnd = mul(float4(WorldRayEnd, 1), WorldToVolume).xyz;
+			float3 VolumeRayDirection = VolumeRayEnd - VolumeRayStart;
+			float VolumeRayLength = length(VolumeRayDirection);
+			VolumeRayDirection /= VolumeRayLength;
+
+			float WorldToVolumeScale = 1.0f / UVScaleAndVolumeScale.w;
+			float VolumeMinSphereRadius = MinSphereRadius * WorldToVolumeScale;
+			float VolumeMaxSphereRadius = MaxSphereRadius * WorldToVolumeScale;
+			float VolumeTwoSidedMeshDistanceBias = TwoSidedMeshDistanceBias * WorldToVolumeScale;
+
+			// Expand the intersection box by the radius of the cone at the distance of the object along the cone
+			float LocalConeRadiusAtObject = min(TanLightAngle * max(ObjectCenterDistanceAlongRay, 0) * WorldToVolumeScale, VolumeMaxSphereRadius);
+
+			float2 IntersectionTimes = LineBoxIntersect(VolumeRayStart, VolumeRayEnd, -LocalPositionExtent - LocalConeRadiusAtObject, LocalPositionExtent + LocalConeRadiusAtObject);
+
+			BRANCH
+			if (IntersectionTimes.x < IntersectionTimes.y && IntersectionTimes.x < 1)
+			{
+				float SampleRayTime = IntersectionTimes.x * VolumeRayLength;
+				uint MaxSteps = 64;
+				float MinStepSize = 1.0f / (4 * MaxSteps);
+
+				float MinDistance = 1000000;
+				float3 IntersectionPosition = float3(0, 0, 0);
+
+				uint StepIndex = 0;
+
+				LOOP
+				for (; StepIndex < MaxSteps; StepIndex++)
+				{
+					float3 SampleVolumePosition = VolumeRayStart + VolumeRayDirection * SampleRayTime;
+					float3 ClampedSamplePosition = clamp(SampleVolumePosition, -LocalPositionExtent, LocalPositionExtent);
+					float DistanceToClamped = length(ClampedSamplePosition - SampleVolumePosition);
+					float3 VolumeUV = DistanceFieldVolumePositionToUV(ClampedSamplePosition, UVScaleAndVolumeScale.xyz, UVAdd);
+					float DistanceField = Texture3DSampleLevel(DistanceFieldTexture, DistanceFieldSampler, VolumeUV, 0).x + DistanceToClamped;
+
+					FLATTEN
+					if (bGeneratedAsTwoSided)
+					{
+						DistanceField -= VolumeTwoSidedMeshDistanceBias;
+					}
+
+					MinDistance = min(MinDistance, DistanceField);
+					float SphereRadius = clamp(TanLightAngle * SampleRayTime, VolumeMinSphereRadius, VolumeMaxSphereRadius);
+
+					MinConeVisibility = min(MinConeVisibility, saturate(DistanceField / SphereRadius));
+					IntersectionPosition = SampleVolumePosition;
+
+					float StepDistance = max(DistanceField, MinStepSize);
+					SampleRayTime += StepDistance;
+
+					// Terminate the trace if we reached a negative area or went past the end of the ray
+					if (DistanceField < 0 
+						|| SampleRayTime > IntersectionTimes.y * VolumeRayLength)
+					{
+						break;
+					}
+				}
+
+				if (MinDistance < 0 || StepIndex == MaxSteps)
+				{
+					MinConeVisibility = 0;
+					MinRayTime = min(MinRayTime, SampleRayTime * UVScaleAndVolumeScale.w);
+				}
+
+				// Force to shadowed as we approach max steps
+				MinConeVisibility = min(MinConeVisibility, (1 - StepIndex / (float)MaxSteps));
+			}
+		}
+	}
+
+	return MinConeVisibility;
+}
+	
\ No newline at end of file
diff --git a/Engine/Shaders/DistanceFieldSurfaceCacheLighting.usf b/Engine/Shaders/DistanceFieldSurfaceCacheLighting.usf
--- a/Engine/Shaders/DistanceFieldSurfaceCacheLighting.usf
+++ b/Engine/Shaders/DistanceFieldSurfaceCacheLighting.usf
@@ -164,8 +164,6 @@ struct FIrradianceCacheSplatVertexOutput
 #endif
 };
 
-Buffer<uint> CompactedRedirect;
-Buffer<uint> IrradianceCacheUsed;
 float InterpolationRadiusScale;
 float2 NormalizedOffsetToPixelCenter;
 float HackExpand;
@@ -687,18 +685,20 @@ void AOUpsamplePS(
 	#endif
 #endif
 
-#if OUTPUT_BENT_NORMAL
-	OutBentNormal = float4(BentNormal, 1);
-	#if SUPPORT_IRRADIANCE
-		OutIrradiance = float4(Irradiance, 1);
-	#endif
+#if OUTPUT_AO
+	OutBentNormal = float4(length(BentNormal).xxx, 1);
 #else
-
-	// Output to RT0 to visualize Irradiance or BentNormal
-	#if SUPPORT_IRRADIANCE
-		OutBentNormal = float4(Irradiance, 1);
+	#if OUTPUT_BENT_NORMAL
+		OutBentNormal = float4(BentNormal, 1);
+		#if SUPPORT_IRRADIANCE
+			OutIrradiance = float4(Irradiance, 1);
+		#endif
 	#else
-		OutBentNormal = float4(length(BentNormal).xxx, 1);
+
+		// Output to RT0 to visualize Irradiance or BentNormal
+		#if SUPPORT_IRRADIANCE
+			OutBentNormal = float4(Irradiance, 1);
+		#endif
 	#endif
 #endif
 }
diff --git a/Engine/Shaders/DistanceFieldSurfaceCacheLightingCompute.usf b/Engine/Shaders/DistanceFieldSurfaceCacheLightingCompute.usf
--- a/Engine/Shaders/DistanceFieldSurfaceCacheLightingCompute.usf
+++ b/Engine/Shaders/DistanceFieldSurfaceCacheLightingCompute.usf
@@ -932,138 +932,6 @@ void ConeTraceOcclusionCS(
 	}
 }
 
-Texture2D HeightfieldTexture;
-SamplerState HeightfieldSampler;
-uint NumHeightfields;
-Buffer<float4> HeightfieldDescriptions;
-
-// In float4's, must match C++
-#define HEIGHTFIELD_DATA_STRIDE 10
-
-float4 GetHeightfieldScaleBias(uint Index, out bool bHasSubsections)
-{
-	float4 Vector = HeightfieldDescriptions[Index * HEIGHTFIELD_DATA_STRIDE + 0];
-	// Sign bit stored presence of subsections
-	bHasSubsections = Vector < 0;
-	Vector.x = abs(Vector.x);
-	return Vector;
-}
-
-float4 GetHeightfieldMinMaxUV(uint Index)
-{
-	return HeightfieldDescriptions[Index * HEIGHTFIELD_DATA_STRIDE + 1];
-}
-
-float4x4 GetWorldToLocal(uint Index)
-{
-	return float4x4(
-		HeightfieldDescriptions[Index * HEIGHTFIELD_DATA_STRIDE + 2],
-		HeightfieldDescriptions[Index * HEIGHTFIELD_DATA_STRIDE + 3],
-		HeightfieldDescriptions[Index * HEIGHTFIELD_DATA_STRIDE + 4],
-		HeightfieldDescriptions[Index * HEIGHTFIELD_DATA_STRIDE + 5]);
-}
-
-float4x4 GetLocalToWorld(uint Index)
-{
-	return float4x4(
-		HeightfieldDescriptions[Index * HEIGHTFIELD_DATA_STRIDE + 6],
-		HeightfieldDescriptions[Index * HEIGHTFIELD_DATA_STRIDE + 7],
-		HeightfieldDescriptions[Index * HEIGHTFIELD_DATA_STRIDE + 8],
-		HeightfieldDescriptions[Index * HEIGHTFIELD_DATA_STRIDE + 9]);
-}
-
-/**  */
-[numthreads(FINAL_GATHER_THREADGROUP_SIZE, 1, 1)]
-void CalculateHeightfieldOcclusionCS(
-	uint3 GroupId : SV_GroupID,
-	uint3 DispatchThreadId : SV_DispatchThreadID,
-    uint3 GroupThreadId : SV_GroupThreadID) 
-{
-	uint StartIndex = SavedStartIndex[0];
-	uint NumRecords = ScatterDrawParameters[1];
-	uint RecordIndex = StartIndex + DispatchThreadId.x;
-
-	float4 PositionAndRadius = IrradianceCachePositionRadius[RecordIndex];
-
-	// Only compute heightfield occlusion contribution for samples created on a distance field, 
-	// Otherwise we are potentially self-shadowing, the heightfield is not accurate enough for that
-	BRANCH
-	if (RecordIndex < NumRecords && PositionAndRadius.w > 0)
-	{
-		float3 WorldNormal = IrradianceCacheNormal[RecordIndex].xyz;
-		float3 WorldPosition = PositionAndRadius.xyz;
-
-		float3 TangentX;
-		float3 TangentY;
-		FindBestAxisVectors2(WorldNormal, TangentX, TangentY);
-
-		LOOP
-		for (uint HeightfieldIndex = 0; HeightfieldIndex < NumHeightfields; HeightfieldIndex++)
-		{
-			float3 LocalPosition = mul(float4(WorldPosition, 1), GetWorldToLocal(HeightfieldIndex)).xyz;
-			bool bHasSubsections;
-			float4 HeightfieldScaleBias = GetHeightfieldScaleBias(HeightfieldIndex, bHasSubsections);
-			float2 HeightfieldUV = LocalPosition.xy * HeightfieldScaleBias.xy + HeightfieldScaleBias.zw;
-			float4 MinMaxHeightfieldUV = GetHeightfieldMinMaxUV(HeightfieldIndex);
-
-			if (bHasSubsections)
-			{
-				float2 CenterUV = .5f * (MinMaxHeightfieldUV.xy + MinMaxHeightfieldUV.zw);
-				// One row of heights in the heightfield are duplicated between subsections
-				// The only number of subsections supported is 2
-				// If we are in subsection 0, no adjustment is needed, but if we are in subsection 1 we need to shift up by one texel
-				HeightfieldUV += HeightfieldUV + .5f * HeightfieldScaleBias.xy > CenterUV ? HeightfieldScaleBias.xy : 0;
-			}
-
-			BRANCH
-			if (all(HeightfieldUV > MinMaxHeightfieldUV.xy) 
-				&& all(HeightfieldUV < MinMaxHeightfieldUV.zw))
-			{
-				float4 SampleValue = Texture2DSampleLevel(HeightfieldTexture, HeightfieldSampler, HeightfieldUV + .5f * HeightfieldScaleBias.xy, 0);
-				float HeightfieldHeight = ((SampleValue.r * 255.0 * 256.0 + SampleValue.g * 255.0) - 32768.0) * TERRAIN_ZSCALE;
-
-				float2 SampleNormal = float2(SampleValue.b, SampleValue.a) * float2(2.0,2.0) - float2(1.0,1.0);
-				// @todo - use the normal to compute occlusion from an arbitrary plane, instead of assuming upward facing plane
-				// This would require an averaged normal over the gather region
-				float3 HeightfieldNormal = float3(SampleNormal, sqrt(max(1.0 - dot(SampleNormal, SampleNormal), 0.0)));
-
-				float3 HeightfieldShadingPosition = float3(LocalPosition.xy, HeightfieldHeight);
-				float3 WorldHeightfieldShadingPosition = mul(float4(HeightfieldShadingPosition, 1), GetLocalToWorld(HeightfieldIndex)).xyz;
-
-				float MinDistanceToOccluder = AOMaxDistance;
-
-				for (uint ConeIndex = 0; ConeIndex < NUM_CONE_DIRECTIONS; ConeIndex++)
-				{
-					float3 ConeDirection = AOSamples2.SampleDirections[ConeIndex].xyz;
-					float3 RotatedConeDirection = ConeDirection.x * TangentX + ConeDirection.y * TangentY + ConeDirection.z * WorldNormal;
-
-					float ExistingConeVisibility = RWRecordConeVisibility[DispatchThreadId.x * NUM_CONE_DIRECTIONS + ConeIndex];
-				
-					float ZBias = 20;
-					float DistanceAlongConeToOccluder = clamp((WorldPosition.z + ZBias - WorldHeightfieldShadingPosition.z) / max(-RotatedConeDirection.z, .0001f), .00001f, AOMaxDistance);
-					float DistanceFromConeEndToOccluder = max(WorldPosition.z + ZBias + RotatedConeDirection.z * DistanceAlongConeToOccluder - WorldHeightfieldShadingPosition.z, 0);
-
-					float SphereRadius = DistanceAlongConeToOccluder * TanConeHalfAngle;
-
-					// Derive visibility from 1d intersection
-					float Visibility = saturate(DistanceFromConeEndToOccluder / SphereRadius);
-
-					float OccluderDistanceFraction = DistanceAlongConeToOccluder / AOMaxDistance;
-					// Fade out occlusion based on distance to occluder to avoid a discontinuity at the max AO distance
-					Visibility = max(Visibility, saturate(OccluderDistanceFraction * OccluderDistanceFraction * .8f));
-
-					MinDistanceToOccluder = min(MinDistanceToOccluder, DistanceAlongConeToOccluder);
-
-					// Shouldn't have to use atomics here as only one heightfield can affect a given world space position
-					RWRecordConeVisibility[DispatchThreadId.x * NUM_CONE_DIRECTIONS + ConeIndex] = min(Visibility, ExistingConeVisibility);
-				}
-
-				RWOccluderRadius[RecordIndex] = min(RWOccluderRadius[RecordIndex], RecordRadiusScale * MinDistanceToOccluder);
-			}
-		}
-	}
-}
-
 RWBuffer<float4> RWIrradianceCacheBentNormal;
 Buffer<float> RecordConeVisibility;
 
@@ -1198,4 +1066,4 @@ void CopyIrradianceCacheSamplesCS(
 	{
 		RWScatterDrawParameters[1] = NumRecords - StartIndex;
 	}
-}
\ No newline at end of file
+}
diff --git a/Engine/Shaders/HeightfieldLighting.usf b/Engine/Shaders/HeightfieldLighting.usf
new file mode 100644
--- /dev/null
+++ b/Engine/Shaders/HeightfieldLighting.usf
@@ -0,0 +1,682 @@
+// Copyright 1998-2014 Epic Games, Inc. All Rights Reserved.
+
+/*=============================================================================
+	HeightfieldLighting.usf
+=============================================================================*/
+
+#include "Common.usf"
+#include "DeferredShadingCommon.usf"
+#include "DistanceFieldLightingShared.usf"
+#include "DistanceFieldAOShared.usf"
+#include "DynamicLightingCommon.usf"
+#include "DistanceFieldShadowingShared.usf"
+
+Texture2D HeightfieldTexture;
+SamplerState HeightfieldSampler;
+uint NumHeightfields;
+Buffer<float4> HeightfieldDescriptions;
+
+// In float4's, must match C++
+#define HEIGHTFIELD_DATA_STRIDE 12
+
+float4 GetHeightfieldScaleBias(uint Index, out bool bHasSubsections)
+{
+	float4 Vector = HeightfieldDescriptions[Index * HEIGHTFIELD_DATA_STRIDE + 0];
+	// Sign bit stored presence of subsections
+	bHasSubsections = Vector.x < 0;
+	Vector.x = abs(Vector.x);
+	return Vector;
+}
+
+float4 GetHeightfieldMinMaxUV(uint Index)
+{
+	return HeightfieldDescriptions[Index * HEIGHTFIELD_DATA_STRIDE + 1];
+}
+
+float4 GetHeightfieldGlobalScaleBias(uint Index)
+{
+	return HeightfieldDescriptions[Index * HEIGHTFIELD_DATA_STRIDE + 2];
+}
+
+float2 GetHeightfieldSize(uint Index)
+{
+	return HeightfieldDescriptions[Index * HEIGHTFIELD_DATA_STRIDE + 3].xy;
+}
+
+float2 GetGlobalTexelSize(uint Index)
+{
+	return HeightfieldDescriptions[Index * HEIGHTFIELD_DATA_STRIDE + 3].zw;
+}
+
+float4x4 GetWorldToLocal(uint Index)
+{
+	return float4x4(
+		HeightfieldDescriptions[Index * HEIGHTFIELD_DATA_STRIDE + 4],
+		HeightfieldDescriptions[Index * HEIGHTFIELD_DATA_STRIDE + 5],
+		HeightfieldDescriptions[Index * HEIGHTFIELD_DATA_STRIDE + 6],
+		HeightfieldDescriptions[Index * HEIGHTFIELD_DATA_STRIDE + 7]);
+}
+
+float4x4 GetLocalToWorld(uint Index)
+{
+	return float4x4(
+		HeightfieldDescriptions[Index * HEIGHTFIELD_DATA_STRIDE + 8],
+		HeightfieldDescriptions[Index * HEIGHTFIELD_DATA_STRIDE + 9],
+		HeightfieldDescriptions[Index * HEIGHTFIELD_DATA_STRIDE + 10],
+		HeightfieldDescriptions[Index * HEIGHTFIELD_DATA_STRIDE + 11]);
+}
+
+float2 GetHeightfieldUV(uint HeightfieldIndex, float2 LocalPosition, out float4 MinMaxHeightfieldUV)
+{
+	bool bHasSubsections;
+	float4 HeightfieldScaleBias = GetHeightfieldScaleBias(HeightfieldIndex, bHasSubsections);
+	float2 HeightfieldUV = LocalPosition * HeightfieldScaleBias.xy + HeightfieldScaleBias.zw;
+	MinMaxHeightfieldUV = GetHeightfieldMinMaxUV(HeightfieldIndex);
+
+	if (bHasSubsections)
+	{
+		float2 CenterUV = .5f * (MinMaxHeightfieldUV.xy + MinMaxHeightfieldUV.zw);
+		// One row of heights in the heightfield are duplicated between subsections
+		// The only number of subsections supported is 2
+		// If we are in subsection 0, no adjustment is needed, but if we are in subsection 1 we need to shift up by one texel
+		HeightfieldUV += HeightfieldUV + .5f * HeightfieldScaleBias.xy > CenterUV ? HeightfieldScaleBias.xy : 0;
+	}
+
+	return HeightfieldUV;
+}
+
+float2 GetHeightfieldUV(uint HeightfieldIndex, float2 LocalPosition)
+{
+	float4 MinMaxHeightfieldUV;
+	return GetHeightfieldUV(HeightfieldIndex, LocalPosition, MinMaxHeightfieldUV);
+}
+
+float3 GetHeightfieldWorldPositionAndNormal(uint HeightfieldIndex, float2 LocalPosition, float2 HeightfieldUV, out float3 WorldNormal)
+{
+	//@todo - don't call GetHeightfieldScaleBias twice
+	bool bHasSubsections;
+	float4 HeightfieldScaleBias = GetHeightfieldScaleBias(HeightfieldIndex, bHasSubsections);
+
+	float4 SampleValue = Texture2DSampleLevel(HeightfieldTexture, HeightfieldSampler, HeightfieldUV + .5f * HeightfieldScaleBias.xy, 0);
+	float HeightfieldHeight = DecodePackedHeight(SampleValue.xy);
+
+	float2 SampleNormal = float2(SampleValue.b, SampleValue.a) * float2(2.0,2.0) - float2(1.0,1.0);
+	WorldNormal = float3(SampleNormal, sqrt(max(1.0 - dot(SampleNormal, SampleNormal), 0.0)));
+
+	float3 HeightfieldShadingPosition = float3(LocalPosition, HeightfieldHeight);
+	float3 WorldHeightfieldShadingPosition = mul(float4(HeightfieldShadingPosition, 1), GetLocalToWorld(HeightfieldIndex)).xyz;
+
+	return WorldHeightfieldShadingPosition;
+}
+
+float3 GetHeightfieldWorldPosition(uint HeightfieldIndex, float2 LocalPosition, float2 HeightfieldUV)
+{
+	float3 WorldNormal;
+	return GetHeightfieldWorldPositionAndNormal(HeightfieldIndex, LocalPosition, HeightfieldUV, WorldNormal);
+}
+
+Texture2D GlobalHeightfieldTexture;
+Texture2D GlobalNormalTexture;
+SamplerState GlobalHeightfieldSampler;
+
+float3 GetGlobalHeightfieldNormal(float2 GlobalUV)
+{
+	float2 PackedNormal = Texture2DSampleLevel(GlobalNormalTexture, GlobalHeightfieldSampler, GlobalUV, 0).xy;
+	float2 RescaledNormal = float2(PackedNormal.x, PackedNormal.y) * float2(2.0,2.0) - float2(1.0,1.0);
+	return float3(RescaledNormal, sqrt(max(1.0 - dot(RescaledNormal, RescaledNormal), 0.0)));
+}
+
+// Returns height in local space
+float GetGlobalHeightfieldHeight(float2 GlobalUV)
+{
+	float UnscaledHeight = Texture2DSampleLevel(GlobalHeightfieldTexture, GlobalHeightfieldSampler, GlobalUV, 0).x * 65535.0f;
+	return DecodeHeightValue(UnscaledHeight);
+}
+
+float3 GetGlobalHeightfieldWorldPosition(float2 LocalPosition, float2 GlobalUV, float4x4 LocalToWorld)
+{
+	float HeightfieldHeight = GetGlobalHeightfieldHeight(GlobalUV);
+	float3 HeightfieldShadingPosition = float3(LocalPosition, HeightfieldHeight);
+	float3 WorldHeightfieldShadingPosition = mul(float4(HeightfieldShadingPosition, 1), LocalToWorld).xyz;
+
+	return WorldHeightfieldShadingPosition;
+}
+
+// In float4's, must match C++
+#define SUBSECTION_HEIGHTFIELD_DATA_STRIDE 4
+
+Buffer<float4> SubsectionHeightfieldDescriptions;
+
+float2 GetSubsectionId(uint HeightfieldSubsectionIndex)
+{
+	return SubsectionHeightfieldDescriptions[HeightfieldSubsectionIndex * SUBSECTION_HEIGHTFIELD_DATA_STRIDE + 0].xy;
+}
+
+float4 GetSubsectionScaleAndBias(uint HeightfieldSubsectionIndex)
+{
+	return SubsectionHeightfieldDescriptions[HeightfieldSubsectionIndex * SUBSECTION_HEIGHTFIELD_DATA_STRIDE + 1];
+}
+
+float4 GetSubsectionComponentHeightmapUVScaleBias(uint HeightfieldSubsectionIndex)
+{
+	return SubsectionHeightfieldDescriptions[HeightfieldSubsectionIndex * SUBSECTION_HEIGHTFIELD_DATA_STRIDE + 2];
+}
+
+float4 GetSubsectionHeightfieldGlobalScaleBias(uint HeightfieldSubsectionIndex)
+{
+	return SubsectionHeightfieldDescriptions[HeightfieldSubsectionIndex * SUBSECTION_HEIGHTFIELD_DATA_STRIDE + 3];
+}
+
+struct FHeightfieldSubsectionQuadVertexOutput
+{
+	float4 Position : SV_POSITION;
+	float2 ComponentHeightfieldUV : TEXCOORD0;
+};
+
+void HeightfieldSubsectionQuadVS(
+	float2 InPosition : ATTRIBUTE0,
+	float2 InUV       : ATTRIBUTE1,
+	uint HeightfieldSubsectionIndex : SV_InstanceID,
+	out FHeightfieldSubsectionQuadVertexOutput Output)
+{
+	float2 SubsectionId = GetSubsectionId(HeightfieldSubsectionIndex);
+	// SubsectionScaleAndBias = SubsectionSizeQuads, SubSectionSizeVerts / HeightfieldTextureSize
+	float4 SubsectionScaleAndBias = GetSubsectionScaleAndBias(HeightfieldSubsectionIndex);
+	float2 ComponentLocalPosition = InPosition * SubsectionScaleAndBias.xy;
+	float4 ComponentHeightmapUVScaleBias = GetSubsectionComponentHeightmapUVScaleBias(HeightfieldSubsectionIndex);
+
+	Output.ComponentHeightfieldUV = ComponentLocalPosition * ComponentHeightmapUVScaleBias.xy + ComponentHeightmapUVScaleBias.zw + .5f * ComponentHeightmapUVScaleBias.xy + SubsectionId * SubsectionScaleAndBias.zw;
+	
+	// GlobalUVScaleBias = SubsectionSizeQuads / AtlasSize, (SubsectionBase + SubsectionId * SubsectionSizeQuads - AtlasMin) / AtlasSize
+	float4 GlobalUVScaleBias = GetSubsectionHeightfieldGlobalScaleBias(HeightfieldSubsectionIndex);
+	float2 GlobalUV = InPosition * GlobalUVScaleBias.xy + GlobalUVScaleBias.zw;
+
+	Output.Position = float4((GlobalUV - .5f) * float2(2, -2), 0, 1);
+}
+
+void InitializeHeightfieldsPS(
+	in FHeightfieldSubsectionQuadVertexOutput Input, 
+	out float4 OutHeight : SV_Target0,
+	out float4 OutNormal : SV_Target1)
+{
+	float4 PackedHeightAndNormal = Texture2DSampleLevel(HeightfieldTexture, HeightfieldSampler, Input.ComponentHeightfieldUV, 0);
+	OutHeight = DecodePackedTwoChannelValue(PackedHeightAndNormal.xy) / 65535.0f;
+	OutNormal = float4(PackedHeightAndNormal.zw, 0, 0);
+}
+
+struct FHeightfieldComponentQuadVertexOutput
+{
+	float4 Position : SV_POSITION;
+	float2 GlobalUV : TEXCOORD0;
+	float2 LocalPosition : TEXCOORD1;
+	uint HeightfieldIndex : TEXCOORD2;
+};
+
+void HeightfieldComponentQuadVS(
+	float2 InPosition : ATTRIBUTE0,
+	float2 InUV       : ATTRIBUTE1,
+	uint HeightfieldIndex : SV_InstanceID,
+	out FHeightfieldComponentQuadVertexOutput Output)
+{
+	float2 LocalPosition = InPosition * GetHeightfieldSize(HeightfieldIndex);
+	float4 GlobalUVScaleBias = GetHeightfieldGlobalScaleBias(HeightfieldIndex);
+	Output.GlobalUV = LocalPosition * GlobalUVScaleBias.xy + GlobalUVScaleBias.zw;
+	Output.Position = float4((Output.GlobalUV - .5f) * float2(2, -2), 0, 1);
+	Output.LocalPosition = LocalPosition;
+	Output.HeightfieldIndex = HeightfieldIndex;
+}
+
+Texture2D HeightfieldLighting;
+SamplerState HeightfieldLightingSampler;
+
+float4 ShadowmapMinMax;
+float ShadowDepthBias;
+float2 CascadeDepthMinMax;
+
+Texture2D ShadowDepthTexture;
+SamplerState ShadowDepthTextureSampler;
+
+void ShadowHeightfieldsPS(
+	in FHeightfieldComponentQuadVertexOutput Input, 
+	out float4 OutColor : SV_Target0)
+{
+	float3 WorldHeightfieldShadingPosition = GetGlobalHeightfieldWorldPosition(Input.LocalPosition, Input.GlobalUV, GetLocalToWorld(Input.HeightfieldIndex));
+
+	float ViewDistance = mul(float4(WorldHeightfieldShadingPosition + View.PreViewTranslation.xyz, 1), View.TranslatedWorldToView).z;
+
+	if (ViewDistance < CascadeDepthMinMax.x || ViewDistance > CascadeDepthMinMax.y)
+	{
+		clip(-1);
+	}
+
+	float4 HomogeneousShadowPosition = mul(float4(WorldHeightfieldShadingPosition, 1), WorldToShadow);
+	float2 ShadowUVs = HomogeneousShadowPosition.xy / HomogeneousShadowPosition.w;
+
+	float Shadow = 1;
+
+	// Treat as unshadowed if the texel is outside of the shadow map
+	if (all(ShadowUVs >= ShadowmapMinMax.xy && ShadowUVs <= ShadowmapMinMax.zw))
+	{
+		// Sample the shadowmap depth and determine if this texel is shadowed
+		float ShadowDepth = Texture2DSampleLevel(ShadowDepthTexture, ShadowDepthTextureSampler, ShadowUVs, 0).x;
+		Shadow = HomogeneousShadowPosition.z < ShadowDepth + ShadowDepthBias;
+	}
+	else
+	{
+		clip(-1);
+	}
+
+	OutColor = float4(Shadow, 0, 0, 0);
+}
+
+float3 LightDirection;
+float TanLightAngle;
+
+void RayTracedShadowHeightfieldsPS(
+	in FHeightfieldComponentQuadVertexOutput Input, 
+	out float4 OutColor : SV_Target0)
+{
+	float3 OpaqueWorldPosition = GetGlobalHeightfieldWorldPosition(Input.LocalPosition, Input.GlobalUV, GetLocalToWorld(Input.HeightfieldIndex));
+	float ViewDistance = mul(float4(OpaqueWorldPosition + View.PreViewTranslation.xyz, 1), View.TranslatedWorldToView).z;
+	float Shadow = 1;
+
+	if (ViewDistance > CascadeDepthMinMax.x && ViewDistance < CascadeDepthMinMax.y)
+	{
+		uint NumIntersectingObjects = GetCulledNumObjects();
+		uint CulledDataStart = 0;
+
+		GetShadowTileCulledData(OpaqueWorldPosition, CulledDataStart, NumIntersectingObjects);
+
+		// Keeps result from going all the way sharp
+		float MinSphereRadius = .4f;
+		// Maintain reasonable culling bounds
+		float MaxSphereRadius = 100;
+		// Distance for directional lights to trace
+		float TraceDistance = 10000;
+
+		float3 WorldRayStart = OpaqueWorldPosition;
+		float3 WorldRayEnd = OpaqueWorldPosition + -LightDirection * TraceDistance;
+		float MaxRayTime = TraceDistance;
+
+		Shadow = ShadowRayTraceThroughCulledObjects(
+			WorldRayStart, 
+			WorldRayEnd, 
+			MaxRayTime, 
+			TanLightAngle, 
+			MinSphereRadius, 
+			MaxSphereRadius, 
+			CulledDataStart, 
+			NumIntersectingObjects,
+			true,
+			true);
+	}
+	else
+	{
+		clip(-1);
+	}
+
+	OutColor = float4(Shadow, 0, 0, 0);
+}
+
+float3 LightColor;
+
+Texture2D HeightfieldShadowing;
+SamplerState HeightfieldShadowingSampler;
+
+void LightHeightfieldsPS(
+	in FHeightfieldComponentQuadVertexOutput Input, 
+	out float4 OutColor : SV_Target0)
+{
+	float3 WorldHeightfieldNormal = GetGlobalHeightfieldNormal(Input.GlobalUV);
+	float3 WorldHeightfieldShadingPosition = GetGlobalHeightfieldWorldPosition(Input.LocalPosition, Input.GlobalUV, GetLocalToWorld(Input.HeightfieldIndex));
+
+	float Shadowing = Texture2DSampleLevel(HeightfieldShadowing, HeightfieldShadowingSampler, Input.GlobalUV, 0).x;
+
+	float3 DiffuseColor = .5f;
+	float3 Lighting = (Shadowing * saturate(dot(-LightDirection, WorldHeightfieldNormal))) * DiffuseColor * LightColor;
+
+	OutColor = float4(Lighting, 0);
+}
+
+RWBuffer<float> RWOccluderRadius;
+RWBuffer<float> RWRecordConeVisibility;
+
+float TanConeHalfAngle;
+float RecordRadiusScale;
+
+/**  */
+[numthreads(FINAL_GATHER_THREADGROUP_SIZE, 1, 1)]
+void CalculateHeightfieldOcclusionCS(
+	uint3 GroupId : SV_GroupID,
+	uint3 DispatchThreadId : SV_DispatchThreadID,
+    uint3 GroupThreadId : SV_GroupThreadID) 
+{
+	uint StartIndex = SavedStartIndex[0];
+	uint NumRecords = ScatterDrawParameters[1];
+	uint RecordIndex = StartIndex + DispatchThreadId.x;
+
+	float4 PositionAndRadius = IrradianceCachePositionRadius[RecordIndex];
+
+	// Only compute heightfield occlusion contribution for samples created on a distance field, 
+	// Otherwise we are potentially self-shadowing, the heightfield is not accurate enough for that
+	BRANCH
+	if (RecordIndex < NumRecords && PositionAndRadius.w > 0)
+	{
+		float3 WorldNormal = IrradianceCacheNormal[RecordIndex].xyz;
+		float3 WorldPosition = PositionAndRadius.xyz;
+
+		float3 TangentX;
+		float3 TangentY;
+		FindBestAxisVectors2(WorldNormal, TangentX, TangentY);
+
+		LOOP
+		for (uint HeightfieldIndex = 0; HeightfieldIndex < NumHeightfields; HeightfieldIndex++)
+		{
+			float3 LocalPosition = mul(float4(WorldPosition, 1), GetWorldToLocal(HeightfieldIndex)).xyz;
+			float4 MinMaxHeightfieldUV;
+			float2 HeightfieldUV = GetHeightfieldUV(HeightfieldIndex, LocalPosition.xy, MinMaxHeightfieldUV);
+
+			BRANCH
+			if (all(HeightfieldUV > MinMaxHeightfieldUV.xy) 
+				&& all(HeightfieldUV < MinMaxHeightfieldUV.zw))
+			{
+				float3 WorldHeightfieldShadingPosition = GetHeightfieldWorldPosition(HeightfieldIndex, LocalPosition.xy, HeightfieldUV);
+
+				float MinDistanceToOccluder = AOMaxDistance;
+
+				for (uint ConeIndex = 0; ConeIndex < NUM_CONE_DIRECTIONS; ConeIndex++)
+				{
+					float3 ConeDirection = AOSamples2.SampleDirections[ConeIndex].xyz;
+					float3 RotatedConeDirection = ConeDirection.x * TangentX + ConeDirection.y * TangentY + ConeDirection.z * WorldNormal;
+
+					float ExistingConeVisibility = RWRecordConeVisibility[DispatchThreadId.x * NUM_CONE_DIRECTIONS + ConeIndex];
+				
+					float ZBias = 20;
+					float DistanceAlongConeToOccluder = clamp((WorldPosition.z + ZBias - WorldHeightfieldShadingPosition.z) / max(-RotatedConeDirection.z, .0001f), .00001f, AOMaxDistance);
+					float DistanceFromConeEndToOccluder = max(WorldPosition.z + ZBias + RotatedConeDirection.z * DistanceAlongConeToOccluder - WorldHeightfieldShadingPosition.z, 0);
+
+					float SphereRadius = DistanceAlongConeToOccluder * TanConeHalfAngle;
+
+					// Derive visibility from 1d intersection
+					float Visibility = saturate(DistanceFromConeEndToOccluder / SphereRadius);
+
+					float OccluderDistanceFraction = DistanceAlongConeToOccluder / AOMaxDistance;
+					// Fade out occlusion based on distance to occluder to avoid a discontinuity at the max AO distance
+					Visibility = max(Visibility, saturate(OccluderDistanceFraction * OccluderDistanceFraction * .8f));
+
+					MinDistanceToOccluder = min(MinDistanceToOccluder, DistanceAlongConeToOccluder);
+
+					// Shouldn't have to use atomics here as only one heightfield can affect a given world space position
+					RWRecordConeVisibility[DispatchThreadId.x * NUM_CONE_DIRECTIONS + ConeIndex] = min(Visibility, ExistingConeVisibility);
+				}
+
+				RWOccluderRadius[RecordIndex] = min(RWOccluderRadius[RecordIndex], RecordRadiusScale * MinDistanceToOccluder);
+			}
+		}
+	}
+}
+
+RWBuffer<float4> RWIrradianceCacheIrradiance;
+
+float3 SampleHeightfieldLighting(float2 GlobalUV)
+{
+	return Texture2DSampleLevel(HeightfieldLighting, HeightfieldLightingSampler, GlobalUV, 0).xyz;
+}
+
+uint NumHemisphereSamples;
+
+bool RayTraceHeightfield(float3 LocalRayStart, float3 LocalRayEnd, float4 GlobalUVScaleBias, out float IntersectionTime, out float2 IntersectionUV)
+{
+	IntersectionUV = 0;
+	IntersectionTime = 1;
+
+	float MaxSteps = 64;
+	float3 LocalRayDirection = (LocalRayEnd - LocalRayStart) / MaxSteps;
+
+	for (float StepIndex = 0; StepIndex < MaxSteps; StepIndex++)
+	{
+		float3 StepPosition = LocalRayStart + StepIndex * LocalRayDirection;
+		float2 StepUV = StepPosition.xy * GlobalUVScaleBias.xy + GlobalUVScaleBias.zw;
+		float StepHeight = GetGlobalHeightfieldHeight(StepUV);
+
+		if (StepHeight > StepPosition.z)
+		{
+			IntersectionUV = StepUV;
+			IntersectionTime = StepIndex / MaxSteps;
+			return true;
+		}
+	}
+
+	return false;
+}
+
+float3 ComputeVirtualPointLighting(float3 WorldPosition, float3 WorldNormal, float3 BentNormal, float4 VPLPositionAndRadius, float2 VPLGlobalUV, float MaxGatherDistance)
+{ 
+	float3 Irradiance = 0;
+
+	float MaxGatherDistanceSq = MaxGatherDistance * MaxGatherDistance;
+
+	float3 VPLToGatherPoint = WorldPosition - VPLPositionAndRadius.xyz;
+	float DistanceSq = dot(VPLToGatherPoint, VPLToGatherPoint);
+	float DirectionDot = dot(-VPLToGatherPoint, WorldNormal);
+
+#define VISUALIZE_VPL_PLACEMENT 0
+#if VISUALIZE_VPL_PLACEMENT
+	float3 DebugValue = LoadSurfelDiffuseColor(VPLIndex);
+	Irradiance += DebugValue * .1f * (DistanceSq < VPLPositionAndRadius.w * VPLPositionAndRadius.w);
+#endif
+
+#define COMPUTE_VPL_LIGHTING 1
+#if COMPUTE_VPL_LIGHTING
+	BRANCH
+	if (DistanceSq < MaxGatherDistanceSq && DirectionDot > 0)
+	{
+		float3 VPLNormal = GetGlobalHeightfieldNormal(VPLGlobalUV);
+		float VPLNormalDot = dot(VPLNormal, VPLToGatherPoint);
+
+		BRANCH
+		if (VPLNormalDot > 0)
+		{
+			float3 VPLFlux = SampleHeightfieldLighting(VPLGlobalUV);
+
+			BRANCH
+			if (any(VPLFlux > .01f))
+			{
+				float Distance = sqrt(DistanceSq);
+				float3 VPLDirection = -VPLToGatherPoint / Distance;
+
+				#define USE_INVERSE_SQUARED_DISK_APPROX 1
+
+				#if USE_INVERSE_SQUARED_DISK_APPROX
+					float DiskRadiusSq = VPLPositionAndRadius.w * VPLPositionAndRadius.w;
+					float DistanceAttenuation = DiskRadiusSq / (DistanceSq + DiskRadiusSq);
+					//@todo - this is scale dependent
+					//DistanceAttenuation = min(DistanceAttenuation, .000001f * RadialAttenuation(VPLToGatherPoint / MaxGatherDistance, 4));
+				#else
+					float DistanceAttenuation = RadialAttenuation(VPLToGatherPoint / MaxGatherDistance, 8);
+					DistanceAttenuation *= .0001f;
+				#endif
+
+				float CosTheta = DirectionDot / Distance;
+				float Shadow = GetVPLOcclusion(BentNormal, VPLDirection);
+				float VPLCosineLobe = saturate(VPLNormalDot / Distance);
+
+				Irradiance += (saturate(CosTheta) * VPLCosineLobe * DistanceAttenuation * Shadow) * VPLFlux;
+			}
+		}
+	}
+#endif
+
+	return Irradiance;
+}
+
+float3 GatherHeightfieldLighting(
+	float3 WorldPosition,
+	float3 BentNormal,
+	float3 WorldNormal,
+	float2 ShadingGlobalUV, 
+	float2 GlobalUVExtents,
+	float2 GlobalTexelSize, 
+	float VPLRadius, 
+	float GatherRadiusExpand,
+	float4 GlobalUVScaleBias, 
+	float4x4 LocalToWorld, 
+	float LightTransferDistance,
+	uint ThreadIndex,
+	uniform bool bUseExcludeRegion)
+{
+	float2 MinGatherUV = ShadingGlobalUV - GlobalUVExtents * GatherRadiusExpand;
+	float2 GatherUVSize = 2 * GlobalUVExtents * GatherRadiusExpand;
+
+	float2 ExcludeMinGatherUV = ShadingGlobalUV - GlobalUVExtents;
+	float2 ExcludeGatherUVSize = 2 * GlobalUVExtents;
+
+	float2 ExpandedTexelSize = GlobalTexelSize * GatherRadiusExpand;
+	float3 Lighting = 0;
+
+	for (float YOffset = ThreadIndex * ExpandedTexelSize.y; YOffset < GatherUVSize.y; YOffset += FINAL_GATHER_THREADGROUP_SIZE * ExpandedTexelSize.y)
+	{
+		for (float XOffset = 0; XOffset < GatherUVSize.x; XOffset += ExpandedTexelSize.x)
+		{
+			float2 SampleUV = MinGatherUV + float2(XOffset, YOffset);
+
+			//@todo - reformulate loop to avoid center region
+			//@todo - try alternate access pattern with all threads reading from the same x row at once
+			if (!bUseExcludeRegion || any(SampleUV < ExcludeMinGatherUV || SampleUV > ExcludeMinGatherUV + ExcludeGatherUVSize))
+			{
+				float2 SampleLocalPosition = (SampleUV - GlobalUVScaleBias.zw) / GlobalUVScaleBias.xy;
+				float3 SampleShadingPosition = GetGlobalHeightfieldWorldPosition(SampleLocalPosition.xy, SampleUV, LocalToWorld);
+
+				Lighting += ComputeVirtualPointLighting(WorldPosition, WorldNormal, BentNormal, float4(SampleShadingPosition, VPLRadius * GatherRadiusExpand), SampleUV, LightTransferDistance * GatherRadiusExpand);
+			}
+		}
+	}
+
+	return Lighting;
+}
+
+groupshared float3 SharedThreadIrradiance[FINAL_GATHER_THREADGROUP_SIZE];
+
+float InnerLightTransferDistance;
+float OuterLightTransferDistanceScale;
+
+/**  */
+[numthreads(FINAL_GATHER_THREADGROUP_SIZE, 1, 1)]
+void CalculateHeightfieldIrradianceCS(
+	uint3 GroupId : SV_GroupID,
+	uint3 DispatchThreadId : SV_DispatchThreadID,
+    uint3 GroupThreadId : SV_GroupThreadID) 
+{
+	uint StartIndex = SavedStartIndex[0];
+	uint NumRecords = ScatterDrawParameters[1];
+	uint RecordIndex = StartIndex + GroupId.x;
+	uint ThreadIndex = GroupThreadId.x;
+
+	float4 PositionAndRadius = IrradianceCachePositionRadius[RecordIndex];
+	float3 Lighting = 0;
+	
+	BRANCH
+	if (RecordIndex < NumRecords)
+	{
+		float3 WorldNormal = IrradianceCacheNormal[RecordIndex].xyz;
+		float3 WorldPosition = PositionAndRadius.xyz;
+		float3 BentNormal = IrradianceCacheBentNormal[RecordIndex].xyz;
+
+		float3 TangentX;
+		float3 TangentY;
+		FindBestAxisVectors2(WorldNormal, TangentX, TangentY);
+
+		//@todo - constant time grid lookup
+		LOOP
+		for (uint HeightfieldIndex = 0; HeightfieldIndex < NumHeightfields; HeightfieldIndex++)
+		{
+			float4x4 WorldToLocal = GetWorldToLocal(HeightfieldIndex);
+			float3 LocalPosition = mul(float4(WorldPosition, 1), WorldToLocal).xyz;
+			float4 MinMaxHeightfieldUV;
+			float2 HeightfieldUV = GetHeightfieldUV(HeightfieldIndex, LocalPosition.xy, MinMaxHeightfieldUV);
+
+			BRANCH
+			if (all(HeightfieldUV > MinMaxHeightfieldUV.xy) 
+				&& all(HeightfieldUV < MinMaxHeightfieldUV.zw))
+			{
+				float4 GlobalUVScaleBias = GetHeightfieldGlobalScaleBias(HeightfieldIndex);
+				float2 GlobalUV = LocalPosition.xy * GlobalUVScaleBias.xy + GlobalUVScaleBias.zw;
+				float4x4 LocalToWorld = GetLocalToWorld(HeightfieldIndex);
+
+				float3 WorldHeightfieldNormal = GetGlobalHeightfieldNormal(GlobalUV);
+				float3 WorldHeightfieldShadingPosition = GetGlobalHeightfieldWorldPosition(LocalPosition.xy, GlobalUV, LocalToWorld);
+
+				#define VISUALIZE_HEIGHTFIELD_LIGHTING 0
+				#if VISUALIZE_HEIGHTFIELD_LIGHTING
+
+					Lighting = SampleHeightfieldLighting(GlobalUV) / FINAL_GATHER_THREADGROUP_SIZE;
+					//Lighting = frac(WorldHeightfieldShadingPosition / 1000);
+					//Lighting = abs(WorldHeightfieldShadingPosition.z - WorldPosition.z) / 100.0f;
+					//Lighting = GetGlobalHeightfieldHeight(GlobalUV) / 100000.0f;
+				#else
+
+					float3 WorldShadingPosition = WorldPosition + 40 * WorldHeightfieldNormal;
+					float3 LocalShadingPosition = mul(float4(WorldShadingPosition, 1), WorldToLocal).xyz;
+
+					float2 GlobalTexelSize = GetGlobalTexelSize(HeightfieldIndex);
+					float2 ShadingGlobalUV = LocalShadingPosition.xy * GlobalUVScaleBias.xy + GlobalUVScaleBias.zw;
+
+					float WorldToLocalScale = length(WorldToLocal[0].xyz);
+					float HeightfieldTexelSize = 1.0f / WorldToLocalScale;
+					float VPLRadius = sqrt(2) * .5f * HeightfieldTexelSize * GlobalTexelSize.x / GlobalUVScaleBias.x;
+
+					float2 GlobalUVExtents = InnerLightTransferDistance * WorldToLocalScale * GlobalUVScaleBias.xy;
+					
+					Lighting = GatherHeightfieldLighting(
+						WorldPosition, 
+						WorldNormal, 
+						BentNormal, 
+						ShadingGlobalUV, 
+						GlobalUVExtents, 
+						GlobalTexelSize, 
+						VPLRadius, 
+						1,
+						GlobalUVScaleBias, 
+						LocalToWorld, 
+						InnerLightTransferDistance, 
+						ThreadIndex,
+						false);
+					
+					Lighting += GatherHeightfieldLighting(
+						WorldPosition, 
+						WorldNormal, 
+						BentNormal, 
+						ShadingGlobalUV, 
+						GlobalUVExtents, 
+						GlobalTexelSize, 
+						VPLRadius, 
+						OuterLightTransferDistanceScale,
+						GlobalUVScaleBias, 
+						LocalToWorld, 
+						InnerLightTransferDistance, 
+						ThreadIndex,
+						true);
+						
+				#endif
+
+				break;
+			}
+		}
+	}
+	
+	SharedThreadIrradiance[ThreadIndex] = Lighting;
+
+	GroupMemoryBarrierWithGroupSync();
+
+	if (ThreadIndex == 0)
+	{
+		float3 Irradiance = 0;
+
+		for (uint i = 0; i < FINAL_GATHER_THREADGROUP_SIZE; i++)
+		{
+			Irradiance += SharedThreadIrradiance[i];
+		}
+
+		RWIrradianceCacheIrradiance[RecordIndex] = float4(Irradiance, 0);
+	}
+}
\ No newline at end of file
diff --git a/Engine/Shaders/LandscapeVertexFactory.usf b/Engine/Shaders/LandscapeVertexFactory.usf
--- a/Engine/Shaders/LandscapeVertexFactory.usf
+++ b/Engine/Shaders/LandscapeVertexFactory.usf
@@ -547,11 +547,12 @@ FVertexFactoryIntermediates GetVertexFactoryIntermediates(FVertexFactoryInput In
 #if FEATURE_LEVEL >= FEATURE_LEVEL_SM4
 	float2 SampleCoords = InputPositionLODAdjusted * LandscapeParameters.HeightmapUVScaleBias.xy + LandscapeParameters.HeightmapUVScaleBias.zw + 0.5*LandscapeParameters.HeightmapUVScaleBias.xy + Intermediates.InputPosition.zw * LandscapeParameters.SubsectionOffsetParams.xy;
 	float4 SampleValue = Texture2DSampleLevel(HeightmapTexture, HeightmapTextureSampler, SampleCoords, LodValue-LodBias.z);
-	float Height = ((SampleValue.r * 255.0 * 256.0 + SampleValue.g * 255.0) - 32768.0) * TERRAIN_ZSCALE;
+	float Height = DecodePackedHeight(SampleValue.xy);
 
 	float2 SampleCoordsNextLOD = InputPositionNextLOD * LandscapeParameters.HeightmapUVScaleBias.xy + LandscapeParameters.HeightmapUVScaleBias.zw + 0.5*LandscapeParameters.HeightmapUVScaleBias.xy + Intermediates.InputPosition.zw * LandscapeParameters.SubsectionOffsetParams.xy;
 	float4 SampleValueNextLOD = Texture2DSampleLevel(HeightmapTexture, HeightmapTextureSampler, SampleCoordsNextLOD, LodValue+1-LodBias.z );
-	float HeightNextLOD = ((SampleValueNextLOD.r * 255.0 * 256.0 + SampleValueNextLOD.g * 255.0) - 32768.0) * TERRAIN_ZSCALE;
+	float HeightNextLOD = DecodePackedHeight(SampleValueNextLOD.xy);
+
 #else
 	// select the height values from the vertex input corresponding to the current and next LOD
 	float SampleLOD = floor(LodValue);
@@ -592,9 +593,10 @@ FVertexFactoryIntermediates GetVertexFactoryIntermediates(FVertexFactoryInput In
 		InputHeightNextLOD	= LOD45[1];	// clamp to last LOD.
 	}
 
-	float MinHeight = (Input.LODHeights[0][0] * 256.0 * 255.0 + Input.LODHeights[0][1] * 255.0 - 32768.0) * TERRAIN_ZSCALE;
-	float MaxHeight = (Input.LODHeights[0][2] * 256.0 * 255.0 + Input.LODHeights[0][3] * 255.0 - 32768.0) * TERRAIN_ZSCALE;
-	
+
+	float MinHeight = DecodePackedHeight(Input.LODHeights[0].xy);
+	float MaxHeight = DecodePackedHeight(Input.LODHeights[0].zw);
+
 	float Height = lerp(MinHeight, MaxHeight, InputHeight);
 	float HeightNextLOD = lerp(MinHeight, MaxHeight, InputHeightNextLOD);
 
diff --git a/Engine/Source/Runtime/Core/Public/Math/IntRect.h b/Engine/Source/Runtime/Core/Public/Math/IntRect.h
--- a/Engine/Source/Runtime/Core/Public/Math/IntRect.h
+++ b/Engine/Source/Runtime/Core/Public/Math/IntRect.h
@@ -177,6 +177,9 @@ public:
 	 */
 	void Clip( const FIntRect& Other );
 
+	/** Combines the two rectanges. */
+	void Union( const FIntRect& Other);
+
 	/**
 	 * Test whether this rectangle contains a point.
 	 *
@@ -448,6 +451,13 @@ FORCEINLINE void FIntRect::Clip( const FIntRect& R )
 	Max.Y = FMath::Max<int32>(Min.Y, Max.Y);
 }
 
+FORCEINLINE void FIntRect::Union( const FIntRect& R )
+{
+	Min.X = FMath::Min<int32>(Min.X, R.Min.X);
+	Min.Y = FMath::Min<int32>(Min.Y, R.Min.Y);
+	Max.X = FMath::Max<int32>(Max.X, R.Max.X);
+	Max.Y = FMath::Max<int32>(Max.Y, R.Max.Y);
+}
 
 FORCEINLINE bool FIntRect::Contains( FIntPoint P ) const
 {
diff --git a/Engine/Source/Runtime/Engine/Public/PrimitiveSceneProxy.h b/Engine/Source/Runtime/Engine/Public/PrimitiveSceneProxy.h
--- a/Engine/Source/Runtime/Engine/Public/PrimitiveSceneProxy.h
+++ b/Engine/Source/Runtime/Engine/Public/PrimitiveSceneProxy.h
@@ -86,6 +86,24 @@ public:
 	}
 };
 
+/** Information about a heightfield gathered by the renderer for heightfield lighting. */
+class FHeightfieldComponentDescription
+{
+public:
+	FVector4 HeightfieldScaleBias;
+	FVector4 MinMaxUV;
+	FMatrix LocalToWorld;
+	FVector2D LightingAtlasLocation;
+	FIntRect HeightfieldRect;
+
+	int32 NumSubsections;
+	FVector4 SubsectionScaleAndBias;
+
+	FHeightfieldComponentDescription(const FMatrix& InLocalToWorld) :
+		LocalToWorld(InLocalToWorld)
+	{}
+};
+
 namespace EDrawDynamicFlags
 {
 	enum Type
@@ -210,11 +228,9 @@ public:
 		BoundsSurfaceArea = 0;
 	}
 
-	virtual void GetHeightfieldRepresentation(UTexture2D*& OutHeightmapTexture, FVector4& OutHeightfieldScaleBias, FVector4& OutMinMaxUV)
+	virtual void GetHeightfieldRepresentation(UTexture2D*& OutHeightmapTexture, FHeightfieldComponentDescription& OutDescription)
 	{
 		OutHeightmapTexture = NULL;
-		OutHeightfieldScaleBias = FVector4(0, 0, 0, 0);
-		OutMinMaxUV = FVector4(0, 0, 0, 0);
 	}
 
 	/**
diff --git a/Engine/Source/Runtime/Landscape/Private/LandscapeRender.cpp b/Engine/Source/Runtime/Landscape/Private/LandscapeRender.cpp
--- a/Engine/Source/Runtime/Landscape/Private/LandscapeRender.cpp
+++ b/Engine/Source/Runtime/Landscape/Private/LandscapeRender.cpp
@@ -2594,24 +2594,22 @@ void FLandscapeComponentSceneProxy::ChangeLODDistanceFactor_RenderThread(float I
 	LODDistance = InLODDistanceFactor;
 }
 
-void FLandscapeComponentSceneProxy::GetHeightfieldRepresentation(UTexture2D*& OutHeightmapTexture, FVector4& OutHeightfieldScaleBias, FVector4& OutMinMaxUV)
+void FLandscapeComponentSceneProxy::GetHeightfieldRepresentation(UTexture2D*& OutHeightmapTexture, FHeightfieldComponentDescription& OutDescription)
 {
 	OutHeightmapTexture = HeightmapTexture;
-	OutHeightfieldScaleBias = HeightmapScaleBias;
+	OutDescription.HeightfieldScaleBias = HeightmapScaleBias;
 
-	check(OutHeightfieldScaleBias.X > 0);
-
-	// CalculateHeightfieldOcclusionCS needs to be fixed up if other values are ever supported
-	check(NumSubsections == 1 || NumSubsections == 2);
-
-	// Store the presence of subsections in the sign bit
-	OutHeightfieldScaleBias.X *= NumSubsections > 1 ? -1 : 1;
-
-	OutMinMaxUV = FVector4(
+	OutDescription.MinMaxUV = FVector4(
 		HeightmapScaleBias.Z, 
 		HeightmapScaleBias.W, 
 		HeightmapScaleBias.Z + SubsectionSizeVerts * NumSubsections * HeightmapScaleBias.X, 
 		HeightmapScaleBias.W + SubsectionSizeVerts * NumSubsections * HeightmapScaleBias.Y);
+
+	OutDescription.HeightfieldRect = FIntRect(SectionBase.X, SectionBase.Y, SectionBase.X + NumSubsections * SubsectionSizeQuads, SectionBase.Y + NumSubsections * SubsectionSizeQuads);
+
+	OutDescription.NumSubsections = NumSubsections;
+
+	OutDescription.SubsectionScaleAndBias = FVector4(SubsectionSizeQuads, SubsectionSizeQuads, HeightmapSubsectionOffsetU, HeightmapSubsectionOffsetV);
 }
 
 //
diff --git a/Engine/Source/Runtime/Landscape/Public/LandscapeRender.h b/Engine/Source/Runtime/Landscape/Public/LandscapeRender.h
--- a/Engine/Source/Runtime/Landscape/Public/LandscapeRender.h
+++ b/Engine/Source/Runtime/Landscape/Public/LandscapeRender.h
@@ -487,13 +487,25 @@ class FLandscapeComponentSceneProxy : public FPrimitiveSceneProxy, public FLands
 
 protected:
 	int8						MaxLOD;
+	/** 
+	 * Number of subsections within the component in each dimension, this can be 1 or 2.  
+	 * Subsections exist to improve the speed at which LOD transitions can take place over distance. 
+	 */
 	int32						NumSubsections;
+	/** Number of unique heights in the subsection. */
 	int32						SubsectionSizeQuads;
+	/** Number of heightmap heights in the subsection.  This includes the duplicate row at the end. */
 	int32						SubsectionSizeVerts;
-	int32						ComponentSizeQuads;	// Size of component in quads
-	int32						ComponentSizeVerts;
+	/** Size of the component in unique heights. */
+	int32						ComponentSizeQuads;	
+	/** 
+	 * ComponentSizeQuads + 1.  
+	 * Note: in the case of multiple subsections, this is not very useful, as there will be an internal duplicate row of heights in addition to the row at the end.
+	 */
+	int32						ComponentSizeVerts; 
 	uint8						StaticLightingLOD;
 	float						StaticLightingResolution;
+	/** Address of the component within the parent Landscape in unique height texels. */
 	FIntPoint					SectionBase;
 	FMatrix						LocalToWorldNoScaling;
 
@@ -576,7 +588,7 @@ public:
 	// FLandcapeSceneProxy
 	void ChangeLODDistanceFactor_RenderThread(float InLODDistanceFactor);
 
-	virtual void GetHeightfieldRepresentation(UTexture2D*& OutHeightmapTexture, FVector4& OutHeightfieldScaleBias, FVector4& OutMinMaxUV) override;
+	virtual void GetHeightfieldRepresentation(UTexture2D*& OutHeightmapTexture, FHeightfieldComponentDescription& OutDescription) override;
 };
 
 class FLandscapeDebugMaterialRenderProxy : public FMaterialRenderProxy
diff --git a/Engine/Source/Runtime/Renderer/Private/DeferredShadingRenderer.cpp b/Engine/Source/Runtime/Renderer/Private/DeferredShadingRenderer.cpp
--- a/Engine/Source/Runtime/Renderer/Private/DeferredShadingRenderer.cpp
+++ b/Engine/Source/Runtime/Renderer/Private/DeferredShadingRenderer.cpp
@@ -664,6 +664,11 @@ void FDeferredShadingSceneRenderer::Render(FRHICommandListImmediate& RHICmdList)
 	{
 		GDistanceFieldVolumeTextureAtlas.UpdateAllocations();
 		UpdateGlobalDistanceFieldObjectBuffers(RHICmdList);
+
+		for (int32 ViewIndex = 0; ViewIndex < Views.Num(); ViewIndex++)
+		{
+			Views[ViewIndex].HeightfieldLightingViewInfo.SetupVisibleHeightfields(Views[ViewIndex], RHICmdList);
+		}	
 	}
 
 	const bool bIsWireframe = ViewFamily.EngineShowFlags.Wireframe;
@@ -1007,7 +1012,7 @@ void FDeferredShadingSceneRenderer::Render(FRHICommandListImmediate& RHICmdList)
 		// Use the skylight's max distance if there is one, to be consistent with DFAO shadowing on the skylight
 		const float OcclusionMaxDistance = Scene->SkyLight && !Scene->SkyLight->bWantsStaticShadowing ? Scene->SkyLight->OcclusionMaxDistance : 600;
 		TRefCountPtr<IPooledRenderTarget> DummyOutput;
-		RenderDistanceFieldAOSurfaceCache(RHICmdList, FDistanceFieldAOParameters(OcclusionMaxDistance), DummyOutput, DummyOutput, true);
+		RenderDistanceFieldAOSurfaceCache(RHICmdList, FDistanceFieldAOParameters(OcclusionMaxDistance), DummyOutput, DummyOutput, ViewFamily.EngineShowFlags.VisualizeDistanceFieldAO, ViewFamily.EngineShowFlags.VisualizeDistanceFieldGI);
 	}
 
 	if (ViewFamily.EngineShowFlags.VisualizeMeshDistanceFields)
diff --git a/Engine/Source/Runtime/Renderer/Private/DeferredShadingRenderer.h b/Engine/Source/Runtime/Renderer/Private/DeferredShadingRenderer.h
--- a/Engine/Source/Runtime/Renderer/Private/DeferredShadingRenderer.h
+++ b/Engine/Source/Runtime/Renderer/Private/DeferredShadingRenderer.h
@@ -192,7 +192,8 @@ private:
 		const class FDistanceFieldAOParameters& Parameters, 
 		TRefCountPtr<IPooledRenderTarget>& OutDynamicBentNormalAO, 
 		TRefCountPtr<IPooledRenderTarget>& OutDynamicIrradiance,
-		bool bApplyToSceneColor);
+		bool bVisualizeAmbientOcclusion,
+		bool bVisualizeGlobalIllumination);
 
 	void RenderMeshDistanceFieldVisualization(FRHICommandListImmediate& RHICmdList, const FDistanceFieldAOParameters& Parameters);
 
@@ -340,6 +341,7 @@ private:
 	bool ShouldDoReflectionEnvironment() const;
 	
 	bool ShouldRenderDynamicSkyLight() const;
+	bool ShouldRenderDistanceFieldAO() const;
 
 	/** Whether distance field global data structures should be prepared for features that use it. */
 	bool ShouldPrepareForDistanceFieldShadows() const;
diff --git a/Engine/Source/Runtime/Renderer/Private/DistanceFieldGlobalIllumination.cpp b/Engine/Source/Runtime/Renderer/Private/DistanceFieldGlobalIllumination.cpp
--- a/Engine/Source/Runtime/Renderer/Private/DistanceFieldGlobalIllumination.cpp
+++ b/Engine/Source/Runtime/Renderer/Private/DistanceFieldGlobalIllumination.cpp
@@ -25,6 +25,14 @@ FAutoConsoleVariableRef CVarDistanceFieldGI(
 	ECVF_Cheat | ECVF_RenderThreadSafe | ECVF_ReadOnly
 	);
 
+int32 GVPLMeshGlobalIllumination = 0;
+FAutoConsoleVariableRef CVarVPLMeshGlobalIllumination(
+	TEXT("r.VPLMeshGlobalIllumination"),
+	GVPLMeshGlobalIllumination,
+	TEXT(""),
+	ECVF_Cheat | ECVF_RenderThreadSafe
+	);
+
 int32 GVPLSurfelRepresentation = 1;
 FAutoConsoleVariableRef CVarVPLSurfelRepresentation(
 	TEXT("r.VPLSurfelRepresentation"),
@@ -735,6 +743,194 @@ private:
 
 IMPLEMENT_SHADER_TYPE(,FLightVPLsCS,TEXT("DistanceFieldGlobalIllumination"),TEXT("LightVPLsCS"),SF_Compute);
 
+void UpdateVPLs(
+	FRHICommandListImmediate& RHICmdList,
+	const FViewInfo& View,
+	const FScene* Scene,
+	const FDistanceFieldAOParameters& Parameters)
+{
+	if (GVPLMeshGlobalIllumination)
+	{
+		if (GVPLSurfelRepresentation)
+		{
+			SCOPED_DRAW_EVENT(RHICmdList, UpdateVPLs);
+
+			const FLightSceneProxy* DirectionalLightProxy = NULL;
+
+			for (TSparseArray<FLightSceneInfoCompact>::TConstIterator LightIt(Scene->Lights); LightIt; ++LightIt)
+			{
+				const FLightSceneInfoCompact& LightSceneInfoCompact = *LightIt;
+				const FLightSceneInfo* const LightSceneInfo = LightSceneInfoCompact.LightSceneInfo;
+
+				if (LightSceneInfo->ShouldRenderLightViewIndependent() 
+					&& LightSceneInfo->Proxy->GetLightType() == LightType_Directional
+					&& LightSceneInfo->Proxy->CastsDynamicShadow())
+				{
+					DirectionalLightProxy = LightSceneInfo->Proxy;
+					break;
+				}
+			}
+
+			FMatrix DirectionalLightWorldToShadow = FMatrix::Identity;
+
+			if (DirectionalLightProxy)
+			{
+				{
+					int32 NumPlanes = 0;
+					const FPlane* PlaneData = NULL;
+					FVector4 ShadowBoundingSphereValue(0, 0, 0, 0);
+					FShadowCascadeSettings CascadeSettings;
+					FSphere ShadowBounds;
+					FConvexVolume FrustumVolume;
+
+					{
+						const float ConeExpandDistance = Parameters.OcclusionMaxDistance;
+						const float TanHalfFOV = 1.0f / View.ViewMatrices.ProjMatrix.M[0][0];
+						const float VertexPullbackLength = ConeExpandDistance / TanHalfFOV;
+
+						// Pull back cone vertex to contain VPLs outside of the view
+						const FVector ViewConeVertex = View.ViewMatrices.ViewOrigin - View.GetViewDirection() * VertexPullbackLength;
+
+						//@todo - expand by AOMaxDistance
+						ShadowBounds = DirectionalLightProxy->GetShadowSplitBoundsDepthRange(
+							View, 
+							ViewConeVertex,
+							View.NearClippingDistance, 
+							GetMaxAOViewDistance() + VertexPullbackLength + Parameters.OcclusionMaxDistance, 
+							&CascadeSettings); 
+
+						FSphere SubjectBounds(FVector::ZeroVector, ShadowBounds.W);
+
+						const FMatrix& WorldToLight = DirectionalLightProxy->GetWorldToLight();
+						const FMatrix InitializerWorldToLight = FInverseRotationMatrix(FVector(WorldToLight.M[0][0],WorldToLight.M[1][0],WorldToLight.M[2][0]).GetSafeNormal().Rotation());
+						const FVector InitializerFaceDirection = FVector(1,0,0);
+
+						FVector	XAxis, YAxis;
+						InitializerFaceDirection.FindBestAxisVectors(XAxis, YAxis);
+						const FMatrix WorldToLightScaled = InitializerWorldToLight * FScaleMatrix(FVector(1.0f,1.0f / SubjectBounds.W,1.0f / SubjectBounds.W));
+						const FMatrix WorldToFace = WorldToLightScaled * FBasisVectorMatrix(-XAxis,YAxis,InitializerFaceDirection.GetSafeNormal(),FVector::ZeroVector);
+
+						NumPlanes = CascadeSettings.ShadowBoundsAccurate.Planes.Num();
+						PlaneData = CascadeSettings.ShadowBoundsAccurate.Planes.GetData();
+
+						DirectionalLightWorldToShadow = FTranslationMatrix(-ShadowBounds.Center) 
+							* WorldToFace 
+							* FShadowProjectionMatrix(-GVPLDirectionalLightTraceDistance / 2, GVPLDirectionalLightTraceDistance / 2, FVector4(0,0,0,1));
+					}
+
+					CullDistanceFieldObjectsForLight(
+						RHICmdList,
+						View,
+						DirectionalLightProxy, 
+						DirectionalLightWorldToShadow,
+						NumPlanes,
+						PlaneData,
+						ShadowBoundingSphereValue,
+						ShadowBounds.W,
+						GVPLPlacementTileIntersectionResources);
+				}
+
+				SCOPED_DRAW_EVENT(RHICmdList, LightVPLs);
+
+				{
+					TShaderMapRef<FSetupLightVPLsIndirectArgumentsCS> ComputeShader(View.ShaderMap);
+					RHICmdList.SetComputeShader(ComputeShader->GetComputeShader());
+					ComputeShader->SetParameters(RHICmdList, View);
+
+					DispatchComputeShader(RHICmdList, *ComputeShader, 1, 1, 1);
+					ComputeShader->UnsetParameters(RHICmdList);
+				}
+
+				{
+					TShaderMapRef<FLightVPLsCS> ComputeShader(View.ShaderMap);
+					RHICmdList.SetComputeShader(ComputeShader->GetComputeShader());
+					ComputeShader->SetParameters(RHICmdList, View, DirectionalLightProxy, DirectionalLightWorldToShadow, Parameters, GVPLPlacementTileIntersectionResources);
+					DispatchIndirectComputeShader(RHICmdList, *ComputeShader, GAOCulledObjectBuffers.Buffers.ObjectIndirectDispatch.Buffer, 0);
+					ComputeShader->UnsetParameters(RHICmdList);
+				}
+			}
+			else
+			{
+				uint32 ClearValues[4] = { 0 };
+				RHICmdList.ClearUAV(Scene->DistanceFieldSceneData.InstancedSurfelBuffers->VPLFlux.UAV, ClearValues);
+			}
+		}
+		else
+		{
+			PlaceVPLs(RHICmdList, View, Scene, Parameters);
+		}
+	}
+}
+
+
+class FClearIrradianceSamplesCS : public FGlobalShader
+{
+	DECLARE_SHADER_TYPE(FClearIrradianceSamplesCS,Global)
+public:
+
+	static bool ShouldCache(EShaderPlatform Platform)
+	{
+		return IsFeatureLevelSupported(Platform, ERHIFeatureLevel::SM5) && DoesPlatformSupportDistanceFieldAO(Platform);
+	}
+
+	static void ModifyCompilationEnvironment(EShaderPlatform Platform, FShaderCompilerEnvironment& OutEnvironment)
+	{
+		FGlobalShader::ModifyCompilationEnvironment(Platform,OutEnvironment);
+	}
+
+	FClearIrradianceSamplesCS(const ShaderMetaType::CompiledShaderInitializerType& Initializer)
+		: FGlobalShader(Initializer)
+	{
+		IrradianceCacheIrradiance.Bind(Initializer.ParameterMap, TEXT("IrradianceCacheIrradiance"));
+		ScatterDrawParameters.Bind(Initializer.ParameterMap, TEXT("ScatterDrawParameters"));
+		SavedStartIndex.Bind(Initializer.ParameterMap, TEXT("SavedStartIndex"));
+	}
+
+	FClearIrradianceSamplesCS()
+	{
+	}
+
+	void SetParameters(
+		FRHICommandList& RHICmdList, 
+		const FSceneView& View, 
+		int32 DepthLevel, 
+		FTemporaryIrradianceCacheResources* TemporaryIrradianceCacheResources)
+	{
+		FComputeShaderRHIParamRef ShaderRHI = GetComputeShader();
+		FGlobalShader::SetParameters(RHICmdList, ShaderRHI, View);
+
+		const FScene* Scene = (const FScene*)View.Family->Scene;
+		FSurfaceCacheResources& SurfaceCacheResources = *Scene->SurfaceCacheResources;
+
+		SetSRVParameter(RHICmdList, ShaderRHI, ScatterDrawParameters, SurfaceCacheResources.Level[DepthLevel]->ScatterDrawParameters.SRV);
+		SetSRVParameter(RHICmdList, ShaderRHI, SavedStartIndex, SurfaceCacheResources.Level[DepthLevel]->SavedStartIndex.SRV);
+
+		IrradianceCacheIrradiance.SetBuffer(RHICmdList, ShaderRHI, SurfaceCacheResources.Level[DepthLevel]->Irradiance);
+	}
+
+	void UnsetParameters(FRHICommandList& RHICmdList)
+	{
+		IrradianceCacheIrradiance.UnsetUAV(RHICmdList, GetComputeShader());
+	}
+
+	virtual bool Serialize(FArchive& Ar)
+	{		
+		bool bShaderHasOutdatedParameters = FGlobalShader::Serialize(Ar);
+		Ar << IrradianceCacheIrradiance;
+		Ar << ScatterDrawParameters;
+		Ar << SavedStartIndex;
+		return bShaderHasOutdatedParameters;
+	}
+
+private:
+
+	FRWShaderParameter IrradianceCacheIrradiance;
+	FShaderResourceParameter ScatterDrawParameters;
+	FShaderResourceParameter SavedStartIndex;
+};
+
+IMPLEMENT_SHADER_TYPE(,FClearIrradianceSamplesCS,TEXT("DistanceFieldGlobalIllumination"),TEXT("ClearIrradianceSamplesCS"),SF_Compute);
+
 
 class FComputeStepBentNormalCS : public FGlobalShader
 {
@@ -838,6 +1034,7 @@ private:
 
 IMPLEMENT_SHADER_TYPE(,FComputeStepBentNormalCS,TEXT("DistanceFieldGlobalIllumination"),TEXT("ComputeStepBentNormalCS"),SF_Compute);
 
+
 enum EVPLMode
 {
 	VPLMode_PlaceFromLight,
@@ -992,121 +1189,6 @@ private:
 IMPLEMENT_SHADER_TYPE(template<>,TComputeIrradianceCS<VPLMode_PlaceFromLight>,TEXT("DistanceFieldGlobalIllumination"),TEXT("ComputeIrradianceCS"),SF_Compute);
 IMPLEMENT_SHADER_TYPE(template<>,TComputeIrradianceCS<VPLMode_PlaceFromSurfels>,TEXT("DistanceFieldGlobalIllumination"),TEXT("ComputeIrradianceCS"),SF_Compute);
 
-void UpdateVPLs(
-	FRHICommandListImmediate& RHICmdList,
-	const FViewInfo& View,
-	const FScene* Scene,
-	const FDistanceFieldAOParameters& Parameters)
-{
-	if (GVPLSurfelRepresentation)
-	{
-		SCOPED_DRAW_EVENT(RHICmdList, UpdateVPLs);
-
-		const FLightSceneProxy* DirectionalLightProxy = NULL;
-
-		for (TSparseArray<FLightSceneInfoCompact>::TConstIterator LightIt(Scene->Lights); LightIt; ++LightIt)
-		{
-			const FLightSceneInfoCompact& LightSceneInfoCompact = *LightIt;
-			const FLightSceneInfo* const LightSceneInfo = LightSceneInfoCompact.LightSceneInfo;
-
-			if (LightSceneInfo->ShouldRenderLightViewIndependent() 
-				&& LightSceneInfo->Proxy->GetLightType() == LightType_Directional
-				&& LightSceneInfo->Proxy->CastsDynamicShadow())
-			{
-				DirectionalLightProxy = LightSceneInfo->Proxy;
-				break;
-			}
-		}
-
-		FMatrix DirectionalLightWorldToShadow = FMatrix::Identity;
-
-		if (DirectionalLightProxy)
-		{
-			{
-				int32 NumPlanes = 0;
-				const FPlane* PlaneData = NULL;
-				FVector4 ShadowBoundingSphereValue(0, 0, 0, 0);
-				FShadowCascadeSettings CascadeSettings;
-				FSphere ShadowBounds;
-				FConvexVolume FrustumVolume;
-
-				{
-					const float ConeExpandDistance = Parameters.OcclusionMaxDistance;
-					const float TanHalfFOV = 1.0f / View.ViewMatrices.ProjMatrix.M[0][0];
-					const float VertexPullbackLength = ConeExpandDistance / TanHalfFOV;
-
-					// Pull back cone vertex to contain VPLs outside of the view
-					const FVector ViewConeVertex = View.ViewMatrices.ViewOrigin - View.GetViewDirection() * VertexPullbackLength;
-
-					//@todo - expand by AOMaxDistance
-					ShadowBounds = DirectionalLightProxy->GetShadowSplitBoundsDepthRange(
-						View, 
-						ViewConeVertex,
-						View.NearClippingDistance, 
-						GetMaxAOViewDistance() + VertexPullbackLength + Parameters.OcclusionMaxDistance, 
-						&CascadeSettings); 
-
-					FSphere SubjectBounds(FVector::ZeroVector, ShadowBounds.W);
-
-					const FMatrix& WorldToLight = DirectionalLightProxy->GetWorldToLight();
-					const FMatrix InitializerWorldToLight = FInverseRotationMatrix(FVector(WorldToLight.M[0][0],WorldToLight.M[1][0],WorldToLight.M[2][0]).GetSafeNormal().Rotation());
-					const FVector InitializerFaceDirection = FVector(1,0,0);
-
-					FVector	XAxis, YAxis;
-					InitializerFaceDirection.FindBestAxisVectors(XAxis, YAxis);
-					const FMatrix WorldToLightScaled = InitializerWorldToLight * FScaleMatrix(FVector(1.0f,1.0f / SubjectBounds.W,1.0f / SubjectBounds.W));
-					const FMatrix WorldToFace = WorldToLightScaled * FBasisVectorMatrix(-XAxis,YAxis,InitializerFaceDirection.GetSafeNormal(),FVector::ZeroVector);
-
-					NumPlanes = CascadeSettings.ShadowBoundsAccurate.Planes.Num();
-					PlaneData = CascadeSettings.ShadowBoundsAccurate.Planes.GetData();
-
-					DirectionalLightWorldToShadow = FTranslationMatrix(-ShadowBounds.Center) 
-						* WorldToFace 
-						* FShadowProjectionMatrix(-GVPLDirectionalLightTraceDistance / 2, GVPLDirectionalLightTraceDistance / 2, FVector4(0,0,0,1));
-				}
-
-				CullDistanceFieldObjectsForLight(
-					RHICmdList,
-					View,
-					DirectionalLightProxy, 
-					DirectionalLightWorldToShadow,
-					NumPlanes,
-					PlaneData,
-					ShadowBoundingSphereValue,
-					ShadowBounds.W,
-					GVPLPlacementTileIntersectionResources);
-			}
-
-			SCOPED_DRAW_EVENT(RHICmdList, LightVPLs);
-
-			{
-				TShaderMapRef<FSetupLightVPLsIndirectArgumentsCS> ComputeShader(View.ShaderMap);
-				RHICmdList.SetComputeShader(ComputeShader->GetComputeShader());
-				ComputeShader->SetParameters(RHICmdList, View);
-
-				DispatchComputeShader(RHICmdList, *ComputeShader, 1, 1, 1);
-				ComputeShader->UnsetParameters(RHICmdList);
-			}
-
-			{
-				TShaderMapRef<FLightVPLsCS> ComputeShader(View.ShaderMap);
-				RHICmdList.SetComputeShader(ComputeShader->GetComputeShader());
-				ComputeShader->SetParameters(RHICmdList, View, DirectionalLightProxy, DirectionalLightWorldToShadow, Parameters, GVPLPlacementTileIntersectionResources);
-				DispatchIndirectComputeShader(RHICmdList, *ComputeShader, GAOCulledObjectBuffers.Buffers.ObjectIndirectDispatch.Buffer, 0);
-				ComputeShader->UnsetParameters(RHICmdList);
-			}
-		}
-		else
-		{
-			uint32 ClearValues[4] = { 0 };
-			RHICmdList.ClearUAV(Scene->DistanceFieldSceneData.InstancedSurfelBuffers->VPLFlux.UAV, ClearValues);
-		}
-	}
-	else
-	{
-		PlaceVPLs(RHICmdList, View, Scene, Parameters);
-	}
-}
 
 void ComputeIrradianceForSamples(
 	int32 DepthLevel,
@@ -1117,7 +1199,7 @@ void ComputeIrradianceForSamples(
 	FTemporaryIrradianceCacheResources* TemporaryIrradianceCacheResources)
 {
 	FSurfaceCacheResources& SurfaceCacheResources = *Scene->SurfaceCacheResources;
-	
+
 	{
 		SCOPED_DRAW_EVENT(RHICmdList, ComputeStepBentNormal);
 
@@ -1136,37 +1218,50 @@ void ComputeIrradianceForSamples(
 			DispatchIndirectComputeShader(RHICmdList, *ComputeShader, SurfaceCacheResources.DispatchParameters.Buffer, 0);
 			ComputeShader->UnsetParameters(RHICmdList);
 		}
+		
+		{
+			TShaderMapRef<FClearIrradianceSamplesCS> ComputeShader(View.ShaderMap);
+			RHICmdList.SetComputeShader(ComputeShader->GetComputeShader());
+			ComputeShader->SetParameters(RHICmdList, View, DepthLevel, TemporaryIrradianceCacheResources);
+			DispatchIndirectComputeShader(RHICmdList, *ComputeShader, SurfaceCacheResources.DispatchParameters.Buffer, 0);
+			ComputeShader->UnsetParameters(RHICmdList);
+		}
 	}
-	
-	{	
-		TShaderMapRef<TSetupFinalGatherIndirectArgumentsCS<true> > ComputeShader(View.ShaderMap);
-		RHICmdList.SetComputeShader(ComputeShader->GetComputeShader());
-		ComputeShader->SetParameters(RHICmdList, View, DepthLevel);
-		DispatchComputeShader(RHICmdList, *ComputeShader, 1, 1, 1);
-		ComputeShader->UnsetParameters(RHICmdList);
-	}
-	
-	SCOPED_DRAW_EVENT(RHICmdList, ComputeIrradiance);
-	
-	if (GVPLSurfelRepresentation)
-	{
-		TShaderMapRef<TComputeIrradianceCS<VPLMode_PlaceFromSurfels> > ComputeShader(View.ShaderMap);
 
-		RHICmdList.SetComputeShader(ComputeShader->GetComputeShader());
-		ComputeShader->SetParameters(RHICmdList, View, DepthLevel, Parameters, TemporaryIrradianceCacheResources);
-		DispatchIndirectComputeShader(RHICmdList, *ComputeShader, SurfaceCacheResources.DispatchParameters.Buffer, 0);
+	View.HeightfieldLightingViewInfo.ComputeIrradianceForSamples(View, RHICmdList, *TemporaryIrradianceCacheResources, DepthLevel, Parameters);
 
-		ComputeShader->UnsetParameters(RHICmdList);
-	}
-	else
+	if (GVPLMeshGlobalIllumination)
 	{
-		TShaderMapRef<TComputeIrradianceCS<VPLMode_PlaceFromLight> > ComputeShader(View.ShaderMap);
+		SCOPED_DRAW_EVENT(RHICmdList, MeshIrradiance);
 
-		RHICmdList.SetComputeShader(ComputeShader->GetComputeShader());
-		ComputeShader->SetParameters(RHICmdList, View, DepthLevel, Parameters, TemporaryIrradianceCacheResources);
-		DispatchIndirectComputeShader(RHICmdList, *ComputeShader, SurfaceCacheResources.DispatchParameters.Buffer, 0);
+		{	
+			TShaderMapRef<TSetupFinalGatherIndirectArgumentsCS<true> > ComputeShader(View.ShaderMap);
+			RHICmdList.SetComputeShader(ComputeShader->GetComputeShader());
+			ComputeShader->SetParameters(RHICmdList, View, DepthLevel);
+			DispatchComputeShader(RHICmdList, *ComputeShader, 1, 1, 1);
+			ComputeShader->UnsetParameters(RHICmdList);
+		}
+	
+		if (GVPLSurfelRepresentation)
+		{
+			TShaderMapRef<TComputeIrradianceCS<VPLMode_PlaceFromSurfels> > ComputeShader(View.ShaderMap);
 
-		ComputeShader->UnsetParameters(RHICmdList);
+			RHICmdList.SetComputeShader(ComputeShader->GetComputeShader());
+			ComputeShader->SetParameters(RHICmdList, View, DepthLevel, Parameters, TemporaryIrradianceCacheResources);
+			DispatchIndirectComputeShader(RHICmdList, *ComputeShader, SurfaceCacheResources.DispatchParameters.Buffer, 0);
+
+			ComputeShader->UnsetParameters(RHICmdList);
+		}
+		else
+		{
+			TShaderMapRef<TComputeIrradianceCS<VPLMode_PlaceFromLight> > ComputeShader(View.ShaderMap);
+
+			RHICmdList.SetComputeShader(ComputeShader->GetComputeShader());
+			ComputeShader->SetParameters(RHICmdList, View, DepthLevel, Parameters, TemporaryIrradianceCacheResources);
+			DispatchIndirectComputeShader(RHICmdList, *ComputeShader, SurfaceCacheResources.DispatchParameters.Buffer, 0);
+
+			ComputeShader->UnsetParameters(RHICmdList);
+		}
 	}
 }
 
diff --git a/Engine/Source/Runtime/Renderer/Private/DistanceFieldLightingShared.h b/Engine/Source/Runtime/Renderer/Private/DistanceFieldLightingShared.h
--- a/Engine/Source/Runtime/Renderer/Private/DistanceFieldLightingShared.h
+++ b/Engine/Source/Runtime/Renderer/Private/DistanceFieldLightingShared.h
@@ -19,6 +19,8 @@ const int32 GAODownsampleFactor = 2;
 static const int32 GMaxNumObjectsPerTile = 512;
 extern uint32 UpdateObjectsGroupSize;
 
+extern int32 GDistanceFieldGI;
+
 extern FIntPoint GetBufferSizeForAO();
 
 class FDistanceFieldObjectBuffers
diff --git a/Engine/Source/Runtime/Renderer/Private/DistanceFieldObjectManagement.cpp b/Engine/Source/Runtime/Renderer/Private/DistanceFieldObjectManagement.cpp
--- a/Engine/Source/Runtime/Renderer/Private/DistanceFieldObjectManagement.cpp
+++ b/Engine/Source/Runtime/Renderer/Private/DistanceFieldObjectManagement.cpp
@@ -742,9 +742,7 @@ void FDeferredShadingSceneRenderer::UpdateGlobalDistanceFieldObjectBuffers(FRHIC
 			int32 OriginalNumSurfels = DistanceFieldSceneData.SurfelAllocations.GetNumSurfelsInBuffer();
 			int32 OriginalNumInstancedSurfels = DistanceFieldSceneData.InstancedSurfelAllocations.GetNumSurfelsInBuffer();
 
-			extern int32 GDistanceFieldGI;
-
-			if (GDistanceFieldGI)
+			if (GDistanceFieldGI && GVPLMeshGlobalIllumination)
 			{
 				for (int32 UploadPrimitiveIndex = 0; UploadPrimitiveIndex < DistanceFieldSceneData.PendingAddOperations.Num(); UploadPrimitiveIndex++)
 				{
@@ -864,7 +862,7 @@ void FDeferredShadingSceneRenderer::UpdateGlobalDistanceFieldObjectBuffers(FRHIC
 						FPrimitiveSurfelAllocation Allocation;
 						FPrimitiveSurfelAllocation InstancedAllocation;
 
-						if (GDistanceFieldGI)
+						if (GDistanceFieldGI && GVPLMeshGlobalIllumination)
 						{
 							const FPrimitiveSurfelAllocation* AllocationPtr = Scene->DistanceFieldSceneData.SurfelAllocations.FindAllocation(PrimitiveSceneInfo);
 							const FPrimitiveSurfelAllocation* InstancedAllocationPtr = Scene->DistanceFieldSceneData.InstancedSurfelAllocations.FindAllocation(PrimitiveSceneInfo);
diff --git a/Engine/Source/Runtime/Renderer/Private/DistanceFieldShadowing.cpp b/Engine/Source/Runtime/Renderer/Private/DistanceFieldShadowing.cpp
--- a/Engine/Source/Runtime/Renderer/Private/DistanceFieldShadowing.cpp
+++ b/Engine/Source/Runtime/Renderer/Private/DistanceFieldShadowing.cpp
@@ -802,6 +802,8 @@ void FProjectedShadowInfo::RenderRayTracedDistanceFieldProjection(FRHICommandLis
 
 			FLightTileIntersectionResources* TileIntersectionResources = LightSceneInfo->TileIntersectionResources;
 
+			View.HeightfieldLightingViewInfo.ComputeRayTracedShadowing(View, RHICmdList, this, TileIntersectionResources, GShadowCulledObjectBuffers);
+
 			TRefCountPtr<IPooledRenderTarget> RayTracedShadowsRT;
 
 			{
@@ -824,6 +826,8 @@ void FProjectedShadowInfo::RenderRayTracedDistanceFieldProjection(FRHICommandLis
 				{
 					SCOPED_DRAW_EVENT(RHICmdList, RayTraceShadows);
 
+					SetRenderTarget(RHICmdList, NULL, NULL);
+
 					if (bDirectionalLight && GShadowScatterTileCulling)
 					{
 						TShaderMapRef<TDistanceFieldShadowingCS<DFS_DirectionalLightScatterTileCulling> > ComputeShader(View.ShaderMap);
diff --git a/Engine/Source/Runtime/Renderer/Private/DistanceFieldSurfaceCacheLighting.cpp b/Engine/Source/Runtime/Renderer/Private/DistanceFieldSurfaceCacheLighting.cpp
--- a/Engine/Source/Runtime/Renderer/Private/DistanceFieldSurfaceCacheLighting.cpp
+++ b/Engine/Source/Runtime/Renderer/Private/DistanceFieldSurfaceCacheLighting.cpp
@@ -18,6 +18,7 @@
 #include "SceneUtils.h"
 #include "OneColorShader.h"
 #include "BasePassRendering.h"
+#include "HeightfieldLighting.h"
 
 int32 GDistanceFieldAO = 1;
 FAutoConsoleVariableRef CVarDistanceFieldAO(
@@ -29,7 +30,6 @@ FAutoConsoleVariableRef CVarDistanceFieldAO(
 
 bool IsDistanceFieldGIAllowed(const FViewInfo& View)
 {
-	extern int32 GDistanceFieldGI;
 	return View.Family->EngineShowFlags.VisualizeDistanceFieldGI || (View.Family->EngineShowFlags.DistanceFieldGI && GDistanceFieldGI && View.Family->EngineShowFlags.GlobalIllumination);
 }
 
@@ -214,14 +214,6 @@ FAutoConsoleVariableRef CVarAOOverwriteSceneColor(
 	ECVF_Cheat | ECVF_RenderThreadSafe
 	);
 
-int32 GAOHeightfieldOcclusion = 1;
-FAutoConsoleVariableRef CVarAOHeightfieldOcclusion(
-	TEXT("r.AOHeightfieldOcclusion"),
-	GAOHeightfieldOcclusion,
-	TEXT("Whether to compute AO from heightfields (landscape)"),
-	ECVF_Cheat | ECVF_RenderThreadSafe
-	);
-
 DEFINE_LOG_CATEGORY(LogDistanceField);
 
 IMPLEMENT_UNIFORM_BUFFER_STRUCT(FAOSampleData2,TEXT("AOSamples2"));
@@ -336,32 +328,6 @@ TGlobalResource<FCircleVertexBuffer> GCircleVertexBuffer;
 TGlobalResource<FDistanceFieldObjectBufferResource> GAOCulledObjectBuffers;
 TGlobalResource<FTemporaryIrradianceCacheResources> GTemporaryIrradianceCacheResources;
 
-class FHeightfieldDescriptionsResource : public FRenderResource
-{
-public:
-
-	FCPUUpdatedBuffer Data;
-
-	FHeightfieldDescriptionsResource()
-	{
-		Data.Format = PF_A32B32G32R32F;
-		// In float4's, must match usf
-		Data.Stride = 10;
-	}
-
-	virtual void InitDynamicRHI()  override
-	{
-		Data.Initialize();
-	}
-
-	virtual void ReleaseDynamicRHI() override
-	{
-		Data.Release();
-	}
-};
-
-TGlobalResource<FHeightfieldDescriptionsResource> GHeightfieldDescriptions;
-
 void FTileIntersectionResources::InitDynamicRHI()
 {
 	TileConeAxisAndCos.Initialize(sizeof(float)* 4, TileDimensions.X * TileDimensions.Y, PF_A32B32G32R32F, BUF_Static);
@@ -375,25 +341,6 @@ void FTileIntersectionResources::InitDynamicRHI()
 	TileArrayNextAllocation.Initialize(sizeof(uint32), 1, PF_R32_UINT, BUF_Static);
 }
 
-void FSceneViewState::DestroyAOTileResources()
-{
-	if (AOTileIntersectionResources) 
-	{
-		ENQUEUE_UNIQUE_RENDER_COMMAND_ONEPARAMETER(
-			DeleteTileResources,
-			FTileIntersectionResources*, AOTileIntersectionResources, AOTileIntersectionResources,
-			{
-				if (AOTileIntersectionResources)
-				{
-					AOTileIntersectionResources->ReleaseResource();
-					delete AOTileIntersectionResources;
-				}
-			}
-		);
-	}
-	AOTileIntersectionResources = NULL;
-}
-
 void OnClearSurfaceCache(UWorld* InWorld)
 {
 	FlushRenderingCommands();
@@ -1783,155 +1730,6 @@ private:
 IMPLEMENT_SHADER_TYPE(template<>,TConeTraceOcclusionCS<true>,TEXT("DistanceFieldSurfaceCacheLightingCompute"),TEXT("ConeTraceOcclusionCS"),SF_Compute);
 IMPLEMENT_SHADER_TYPE(template<>,TConeTraceOcclusionCS<false>,TEXT("DistanceFieldSurfaceCacheLightingCompute"),TEXT("ConeTraceOcclusionCS"),SF_Compute);
 
-
-class FCalculateHeightfieldOcclusionCS : public FGlobalShader
-{
-	DECLARE_SHADER_TYPE(FCalculateHeightfieldOcclusionCS,Global)
-public:
-
-	static bool ShouldCache(EShaderPlatform Platform)
-	{
-		return IsFeatureLevelSupported(Platform, ERHIFeatureLevel::SM5) && DoesPlatformSupportDistanceFieldAO(Platform);
-	}
-
-	static void ModifyCompilationEnvironment(EShaderPlatform Platform, FShaderCompilerEnvironment& OutEnvironment)
-	{
-		FGlobalShader::ModifyCompilationEnvironment(Platform,OutEnvironment);
-		OutEnvironment.SetDefine(TEXT("NUM_CONE_DIRECTIONS"), NumConeSampleDirections);
-
-		// To reduce shader compile time of compute shaders with shared memory, doesn't have an impact on generated code with current compiler (June 2010 DX SDK)
-		OutEnvironment.CompilerFlags.Add(CFLAG_StandardOptimization);
-	}
-
-	FCalculateHeightfieldOcclusionCS(const ShaderMetaType::CompiledShaderInitializerType& Initializer)
-		: FGlobalShader(Initializer)
-	{
-		AOParameters.Bind(Initializer.ParameterMap);
-		OccluderRadius.Bind(Initializer.ParameterMap, TEXT("OccluderRadius"));
-		IrradianceCacheNormal.Bind(Initializer.ParameterMap, TEXT("IrradianceCacheNormal"));
-		IrradianceCachePositionRadius.Bind(Initializer.ParameterMap, TEXT("IrradianceCachePositionRadius"));
-		RecordConeVisibility.Bind(Initializer.ParameterMap, TEXT("RecordConeVisibility"));
-		ScatterDrawParameters.Bind(Initializer.ParameterMap, TEXT("ScatterDrawParameters"));
-		SavedStartIndex.Bind(Initializer.ParameterMap, TEXT("SavedStartIndex"));
-		BentNormalNormalizeFactor.Bind(Initializer.ParameterMap, TEXT("BentNormalNormalizeFactor"));
-		DistanceFieldNormalTexture.Bind(Initializer.ParameterMap, TEXT("DistanceFieldNormalTexture"));
-		DistanceFieldNormalSampler.Bind(Initializer.ParameterMap, TEXT("DistanceFieldNormalSampler"));
-		RecordRadiusScale.Bind(Initializer.ParameterMap, TEXT("RecordRadiusScale"));
-		HeightfieldTexture.Bind(Initializer.ParameterMap, TEXT("HeightfieldTexture"));
-		HeightfieldSampler.Bind(Initializer.ParameterMap, TEXT("HeightfieldSampler"));
-		HeightfieldDescriptions.Bind(Initializer.ParameterMap, TEXT("HeightfieldDescriptions"));
-		NumHeightfields.Bind(Initializer.ParameterMap, TEXT("NumHeightfields"));
-		TanConeHalfAngle.Bind(Initializer.ParameterMap, TEXT("TanConeHalfAngle"));
-	}
-
-	FCalculateHeightfieldOcclusionCS()
-	{
-	}
-
-	void SetParameters(
-		FRHICommandList& RHICmdList, 
-		const FSceneView& View, 
-		int32 DepthLevel, 
-		UTexture2D* HeightfieldTextureValue,
-		int32 NumHeightfieldsValue,
-		const FDistanceFieldAOParameters& Parameters)
-	{
-		FComputeShaderRHIParamRef ShaderRHI = GetComputeShader();
-		FGlobalShader::SetParameters(RHICmdList, ShaderRHI, View);
-
-		AOParameters.Set(RHICmdList, ShaderRHI, Parameters);
-
-		const FScene* Scene = (const FScene*)View.Family->Scene;
-		FSurfaceCacheResources& SurfaceCacheResources = *Scene->SurfaceCacheResources;
-
-		SetSRVParameter(RHICmdList, ShaderRHI, IrradianceCacheNormal, SurfaceCacheResources.Level[DepthLevel]->Normal.SRV);
-		SetSRVParameter(RHICmdList, ShaderRHI, IrradianceCachePositionRadius, SurfaceCacheResources.Level[DepthLevel]->PositionAndRadius.SRV);
-		SetSRVParameter(RHICmdList, ShaderRHI, ScatterDrawParameters, SurfaceCacheResources.Level[DepthLevel]->ScatterDrawParameters.SRV);
-		SetSRVParameter(RHICmdList, ShaderRHI, SavedStartIndex, SurfaceCacheResources.Level[DepthLevel]->SavedStartIndex.SRV);
-		SetSRVParameter(RHICmdList, ShaderRHI, HeightfieldDescriptions, GHeightfieldDescriptions.Data.BufferSRV);
-
-		SetTextureParameter(RHICmdList, ShaderRHI, HeightfieldTexture, HeightfieldSampler, HeightfieldTextureValue->Resource);
-		SetShaderValue(RHICmdList, ShaderRHI, NumHeightfields, NumHeightfieldsValue);
-
-		OccluderRadius.SetBuffer(RHICmdList, ShaderRHI, SurfaceCacheResources.Level[DepthLevel]->OccluderRadius);
-		RecordConeVisibility.SetBuffer(RHICmdList, ShaderRHI, GTemporaryIrradianceCacheResources.ConeVisibility);
-
-		FAOSampleData2 AOSampleData;
-
-		TArray<FVector, TInlineAllocator<9> > SampleDirections;
-		GetSpacedVectors(SampleDirections);
-
-		for (int32 SampleIndex = 0; SampleIndex < NumConeSampleDirections; SampleIndex++)
-		{
-			AOSampleData.SampleDirections[SampleIndex] = FVector4(SampleDirections[SampleIndex]);
-		}
-
-		SetUniformBufferParameterImmediate(RHICmdList, ShaderRHI, GetUniformBufferParameter<FAOSampleData2>(), AOSampleData);
-
-		FVector UnoccludedVector(0);
-
-		for (int32 SampleIndex = 0; SampleIndex < NumConeSampleDirections; SampleIndex++)
-		{
-			UnoccludedVector += SampleDirections[SampleIndex];
-		}
-
-		float BentNormalNormalizeFactorValue = 1.0f / (UnoccludedVector / NumConeSampleDirections).Size();
-		SetShaderValue(RHICmdList, ShaderRHI, BentNormalNormalizeFactor, BentNormalNormalizeFactorValue);
-
-		SetShaderValue(RHICmdList, ShaderRHI, TanConeHalfAngle, FMath::Tan(GAOConeHalfAngle));
-		SetShaderValue(RHICmdList, ShaderRHI, RecordRadiusScale, GAORecordRadiusScale);
-	}
-
-	void UnsetParameters(FRHICommandList& RHICmdList)
-	{
-		OccluderRadius.UnsetUAV(RHICmdList, GetComputeShader());
-		RecordConeVisibility.UnsetUAV(RHICmdList, GetComputeShader());
-	}
-
-	virtual bool Serialize(FArchive& Ar)
-	{		
-		bool bShaderHasOutdatedParameters = FGlobalShader::Serialize(Ar);
-		Ar << AOParameters;
-		Ar << OccluderRadius;
-		Ar << IrradianceCacheNormal;
-		Ar << IrradianceCachePositionRadius;
-		Ar << RecordConeVisibility;
-		Ar << ScatterDrawParameters;
-		Ar << SavedStartIndex;
-		Ar << BentNormalNormalizeFactor;
-		Ar << DistanceFieldNormalTexture;
-		Ar << DistanceFieldNormalSampler;
-		Ar << RecordRadiusScale;
-		Ar << HeightfieldTexture;
-		Ar << HeightfieldSampler;
-		Ar << HeightfieldDescriptions;
-		Ar << NumHeightfields;
-		Ar << TanConeHalfAngle;
-		return bShaderHasOutdatedParameters;
-	}
-
-private:
-
-	FAOParameters AOParameters;
-	FRWShaderParameter OccluderRadius;
-	FShaderResourceParameter IrradianceCacheNormal;
-	FShaderResourceParameter IrradianceCachePositionRadius;
-	FRWShaderParameter RecordConeVisibility;
-	FShaderResourceParameter ScatterDrawParameters;
-	FShaderResourceParameter SavedStartIndex;
-	FShaderParameter BentNormalNormalizeFactor;
-	FShaderResourceParameter DistanceFieldNormalTexture;
-	FShaderResourceParameter DistanceFieldNormalSampler;
-	FShaderParameter RecordRadiusScale;
-	FShaderResourceParameter HeightfieldTexture;
-	FShaderResourceParameter HeightfieldSampler;
-	FShaderResourceParameter HeightfieldDescriptions;
-	FShaderParameter NumHeightfields;
-	FShaderParameter TanConeHalfAngle;
-};
-
-IMPLEMENT_SHADER_TYPE(,FCalculateHeightfieldOcclusionCS,TEXT("DistanceFieldSurfaceCacheLightingCompute"),TEXT("CalculateHeightfieldOcclusionCS"),SF_Compute);
-
 class FCombineConesCS : public FGlobalShader
 {
 	DECLARE_SHADER_TYPE(FCombineConesCS,Global)
@@ -2903,7 +2701,15 @@ void PostProcessBentNormalAO(
 	}
 }
 
-template<bool bOutputBentNormal, bool bSupportIrradiance>
+enum EAOUpsampleType
+{
+	AOUpsample_OutputBentNormal,
+	AOUpsample_OutputAO,
+	AOUpsample_OutputBentNormalAndIrradiance,
+	AOUpsample_OutputIrradiance
+};
+
+template<EAOUpsampleType UpsampleType>
 class TDistanceFieldAOUpsamplePS : public FGlobalShader
 {
 	DECLARE_SHADER_TYPE(TDistanceFieldAOUpsamplePS, Global);
@@ -2917,8 +2723,9 @@ public:
 	static void ModifyCompilationEnvironment(EShaderPlatform Platform, FShaderCompilerEnvironment& OutEnvironment)
 	{
 		OutEnvironment.SetDefine(TEXT("DOWNSAMPLE_FACTOR"), GAODownsampleFactor);
-		OutEnvironment.SetDefine(TEXT("OUTPUT_BENT_NORMAL"), bOutputBentNormal ? TEXT("1") : TEXT("0"));
-		OutEnvironment.SetDefine(TEXT("SUPPORT_IRRADIANCE"), bSupportIrradiance ? TEXT("1") : TEXT("0"));
+		OutEnvironment.SetDefine(TEXT("OUTPUT_BENT_NORMAL"), (UpsampleType == AOUpsample_OutputBentNormal || UpsampleType == AOUpsample_OutputBentNormalAndIrradiance) ? TEXT("1") : TEXT("0"));
+		OutEnvironment.SetDefine(TEXT("SUPPORT_IRRADIANCE"), (UpsampleType == AOUpsample_OutputIrradiance || UpsampleType == AOUpsample_OutputBentNormalAndIrradiance) ? TEXT("1") : TEXT("0"));
+		OutEnvironment.SetDefine(TEXT("OUTPUT_AO"), (UpsampleType == AOUpsample_OutputAO) ? TEXT("1") : TEXT("0"));
 	}
 
 	/** Default constructor. */
@@ -2970,21 +2777,22 @@ private:
 	FShaderResourceParameter IrradianceSampler;
 };
 
-#define IMPLEMENT_UPSAMPLE_PS_TYPE(bOutputBentNormal, bSupportIrradiance) \
-	typedef TDistanceFieldAOUpsamplePS<bOutputBentNormal, bSupportIrradiance> TDistanceFieldAOUpsamplePS##bOutputBentNormal##bSupportIrradiance; \
-	IMPLEMENT_SHADER_TYPE(template<>,TDistanceFieldAOUpsamplePS##bOutputBentNormal##bSupportIrradiance,TEXT("DistanceFieldSurfaceCacheLighting"),TEXT("AOUpsamplePS"),SF_Pixel);
+#define IMPLEMENT_UPSAMPLE_PS_TYPE(UpsampleType) \
+	typedef TDistanceFieldAOUpsamplePS<UpsampleType> TDistanceFieldAOUpsamplePS##UpsampleType; \
+	IMPLEMENT_SHADER_TYPE(template<>,TDistanceFieldAOUpsamplePS##UpsampleType,TEXT("DistanceFieldSurfaceCacheLighting"),TEXT("AOUpsamplePS"),SF_Pixel);
 
-IMPLEMENT_UPSAMPLE_PS_TYPE(true, true)
-IMPLEMENT_UPSAMPLE_PS_TYPE(true, false)
-IMPLEMENT_UPSAMPLE_PS_TYPE(false, true)
-IMPLEMENT_UPSAMPLE_PS_TYPE(false, false)
+IMPLEMENT_UPSAMPLE_PS_TYPE(AOUpsample_OutputBentNormal)
+IMPLEMENT_UPSAMPLE_PS_TYPE(AOUpsample_OutputAO)
+IMPLEMENT_UPSAMPLE_PS_TYPE(AOUpsample_OutputBentNormalAndIrradiance)
+IMPLEMENT_UPSAMPLE_PS_TYPE(AOUpsample_OutputIrradiance)
 
 void UpsampleBentNormalAO(
 	FRHICommandList& RHICmdList, 
 	const TArray<FViewInfo>& Views, 
 	TRefCountPtr<IPooledRenderTarget>& DistanceFieldAOBentNormal, 
-	TRefCountPtr<IPooledRenderTarget>& DistanceFieldIrradiance, 
-	bool bOutputBentNormal)
+	TRefCountPtr<IPooledRenderTarget>& DistanceFieldIrradiance,
+	bool bVisualizeAmbientOcclusion,
+	bool bVisualizeGlobalIllumination)
 {
 	for (int32 ViewIndex = 0; ViewIndex < Views.Num(); ViewIndex++)
 	{
@@ -3000,35 +2808,32 @@ void UpsampleBentNormalAO(
 
 		TShaderMapRef<FPostProcessVS> VertexShader(View.ShaderMap);
 
-		if (bOutputBentNormal)
+		if (bVisualizeAmbientOcclusion)
 		{
-			if (bUseDistanceFieldGI)
-			{
-				TShaderMapRef<TDistanceFieldAOUpsamplePS<true, true> > PixelShader(View.ShaderMap);
-				static FGlobalBoundShaderState BoundShaderState;
-				SetGlobalBoundShaderState(RHICmdList, View.GetFeatureLevel(), BoundShaderState, GFilterVertexDeclaration.VertexDeclarationRHI, *VertexShader, *PixelShader);
-				PixelShader->SetParameters(RHICmdList, View, DistanceFieldAOBentNormal, DistanceFieldIrradiance);
-			}
-			else
-			{
-				TShaderMapRef<TDistanceFieldAOUpsamplePS<true, false> > PixelShader(View.ShaderMap);
-				static FGlobalBoundShaderState BoundShaderState;
-				SetGlobalBoundShaderState(RHICmdList, View.GetFeatureLevel(), BoundShaderState, GFilterVertexDeclaration.VertexDeclarationRHI, *VertexShader, *PixelShader);
-				PixelShader->SetParameters(RHICmdList, View, DistanceFieldAOBentNormal, DistanceFieldIrradiance);
-			}
+			TShaderMapRef<TDistanceFieldAOUpsamplePS<AOUpsample_OutputAO> > PixelShader(View.ShaderMap);
+			static FGlobalBoundShaderState BoundShaderState;
+			SetGlobalBoundShaderState(RHICmdList, View.GetFeatureLevel(), BoundShaderState, GFilterVertexDeclaration.VertexDeclarationRHI, *VertexShader, *PixelShader);
+			PixelShader->SetParameters(RHICmdList, View, DistanceFieldAOBentNormal, DistanceFieldIrradiance);
+		}
+		else if (bVisualizeGlobalIllumination && bUseDistanceFieldGI)
+		{
+			TShaderMapRef<TDistanceFieldAOUpsamplePS<AOUpsample_OutputIrradiance> > PixelShader(View.ShaderMap);
+			static FGlobalBoundShaderState BoundShaderState;
+			SetGlobalBoundShaderState(RHICmdList, View.GetFeatureLevel(), BoundShaderState, GFilterVertexDeclaration.VertexDeclarationRHI, *VertexShader, *PixelShader);
+			PixelShader->SetParameters(RHICmdList, View, DistanceFieldAOBentNormal, DistanceFieldIrradiance);
 		}
 		else
 		{
 			if (bUseDistanceFieldGI)
 			{
-				TShaderMapRef<TDistanceFieldAOUpsamplePS<false, true> > PixelShader(View.ShaderMap);
+				TShaderMapRef<TDistanceFieldAOUpsamplePS<AOUpsample_OutputBentNormalAndIrradiance> > PixelShader(View.ShaderMap);
 				static FGlobalBoundShaderState BoundShaderState;
 				SetGlobalBoundShaderState(RHICmdList, View.GetFeatureLevel(), BoundShaderState, GFilterVertexDeclaration.VertexDeclarationRHI, *VertexShader, *PixelShader);
 				PixelShader->SetParameters(RHICmdList, View, DistanceFieldAOBentNormal, DistanceFieldIrradiance);
 			}
 			else
 			{
-				TShaderMapRef<TDistanceFieldAOUpsamplePS<false, false> > PixelShader(View.ShaderMap);
+				TShaderMapRef<TDistanceFieldAOUpsamplePS<AOUpsample_OutputBentNormal> > PixelShader(View.ShaderMap);
 				static FGlobalBoundShaderState BoundShaderState;
 				SetGlobalBoundShaderState(RHICmdList, View.GetFeatureLevel(), BoundShaderState, GFilterVertexDeclaration.VertexDeclarationRHI, *VertexShader, *PixelShader);
 				PixelShader->SetParameters(RHICmdList, View, DistanceFieldAOBentNormal, DistanceFieldIrradiance);
@@ -3650,23 +3455,6 @@ bool FDeferredShadingSceneRenderer::ShouldPrepareForDistanceFieldAO() const
 			|| ViewFamily.EngineShowFlags.VisualizeDistanceFieldGI);
 }
 
-class FHeightfieldDescription
-{
-public:
-	FVector4 HeightfieldScaleBias;
-	FVector4 MinMaxUV;
-	FMatrix LocalToWorld;
-
-	FHeightfieldDescription(
-		const FVector4& InHeightfieldScaleBias,
-		const FVector4& InMinMaxUV,
-		const FMatrix& InLocalToWorld) :
-		HeightfieldScaleBias(InHeightfieldScaleBias),
-		MinMaxUV(InMinMaxUV),
-		LocalToWorld(InLocalToWorld)
-	{}
-};
-
 FRWBuffer* GDebugBuffer2 = NULL;
 
 IMPLEMENT_SHADER_TYPE(,FTrackGPUProgressCS,TEXT("DistanceFieldGlobalIllumination"),TEXT("TrackGPUProgressCS"),SF_Compute);
@@ -3697,8 +3485,9 @@ bool FDeferredShadingSceneRenderer::RenderDistanceFieldAOSurfaceCache(
 	FRHICommandListImmediate& RHICmdList, 
 	const FDistanceFieldAOParameters& Parameters, 
 	TRefCountPtr<IPooledRenderTarget>& OutDynamicBentNormalAO, 
-	TRefCountPtr<IPooledRenderTarget>& OutDynamicIrradiance, 
-	bool bApplyToSceneColor)
+	TRefCountPtr<IPooledRenderTarget>& OutDynamicIrradiance,
+	bool bVisualizeAmbientOcclusion,
+	bool bVisualizeGlobalIllumination)
 {
 	//@todo - support multiple views
 	const FViewInfo& View = Views[0];
@@ -3951,6 +3740,8 @@ bool FDeferredShadingSceneRenderer::RenderDistanceFieldAOSurfaceCache(
 					}
 				}
 
+				View.HeightfieldLightingViewInfo.ComputeOcclusionForSamples(View, RHICmdList, GTemporaryIrradianceCacheResources, DepthLevel, Parameters);
+
 				{	
 					TShaderMapRef<TSetupFinalGatherIndirectArgumentsCS<false> > ComputeShader(View.ShaderMap);
 
@@ -3961,129 +3752,27 @@ bool FDeferredShadingSceneRenderer::RenderDistanceFieldAOSurfaceCache(
 					ComputeShader->UnsetParameters(RHICmdList);
 				}
 
-				if (GAOHeightfieldOcclusion && Scene->DistanceFieldSceneData.HeightfieldPrimitives.Num() > 0)
 				{
-					const int32 NumPrimitives = Scene->DistanceFieldSceneData.HeightfieldPrimitives.Num();
-					const float MaxDistanceSquared = FMath::Square(GAOMaxViewDistance + Parameters.OcclusionMaxDistance);
+					TShaderMapRef<FCombineConesCS> ComputeShader(View.ShaderMap);
 
-					TMap<UTexture2D*, TArray<FHeightfieldDescription, SceneRenderingAllocator>, SceneRenderingSetAllocator> HeightfieldTextures; 
-
-					for (int32 HeightfieldPrimitiveIndex = 0; HeightfieldPrimitiveIndex < NumPrimitives; HeightfieldPrimitiveIndex++)
-					{
-						const FPrimitiveSceneInfo* HeightfieldPrimitive = Scene->DistanceFieldSceneData.HeightfieldPrimitives[HeightfieldPrimitiveIndex];
-						const FBoxSphereBounds& PrimitiveBounds = HeightfieldPrimitive->Proxy->GetBounds();
-						const float DistanceToPrimitiveSq = (PrimitiveBounds.Origin - View.ViewMatrices.ViewOrigin).SizeSquared();
-						
-						if (View.ViewFrustum.IntersectSphere(PrimitiveBounds.Origin, PrimitiveBounds.SphereRadius + Parameters.OcclusionMaxDistance)
-							&& DistanceToPrimitiveSq < MaxDistanceSquared)
-						{
-							UTexture2D* HeightfieldTexture = NULL;
-							FVector4 HeightfieldScaleBias(0, 0, 0, 0);
-							FVector4 MinMaxUV(0, 0, 0, 0);
-							HeightfieldPrimitive->Proxy->GetHeightfieldRepresentation(HeightfieldTexture, HeightfieldScaleBias, MinMaxUV);
-
-							if (HeightfieldTexture && HeightfieldTexture->Resource->TextureRHI)
-							{
-								TArray<FHeightfieldDescription, SceneRenderingAllocator>& HeightfieldDescriptions = HeightfieldTextures.FindOrAdd(HeightfieldTexture);
-								HeightfieldDescriptions.Add(FHeightfieldDescription(HeightfieldScaleBias, MinMaxUV, HeightfieldPrimitive->Proxy->GetLocalToWorld()));
-							}
-						}
-					}
-
-					if (HeightfieldTextures.Num() > 0)
-					{
-						SCOPED_DRAW_EVENT(RHICmdList, HeightfieldOcclusion);
+					RHICmdList.SetComputeShader(ComputeShader->GetComputeShader());
+					ComputeShader->SetParameters(RHICmdList, View, DepthLevel, Parameters, NULL, FMatrix::Identity, NULL, NULL);
+					DispatchIndirectComputeShader(RHICmdList, *ComputeShader, SurfaceCacheResources.DispatchParameters.Buffer, 0);
 
-						for (TMap<UTexture2D*, TArray<FHeightfieldDescription, SceneRenderingAllocator>, SceneRenderingSetAllocator>::TIterator It(HeightfieldTextures); It; ++It)
-						{
-							const TArray<FHeightfieldDescription, SceneRenderingAllocator>& HeightfieldDescriptions = It.Value();
-
-							if (HeightfieldDescriptions.Num() > 0)
-							{
-								TArray<FVector4, SceneRenderingAllocator> HeightfieldDescriptionData;
-								HeightfieldDescriptionData.Empty(HeightfieldDescriptions.Num() * GHeightfieldDescriptions.Data.Stride);
-
-								for (int32 DescriptionIndex = 0; DescriptionIndex < HeightfieldDescriptions.Num(); DescriptionIndex++)
-								{
-									const FHeightfieldDescription& Description = HeightfieldDescriptions[DescriptionIndex];
-
-									HeightfieldDescriptionData.Add(Description.HeightfieldScaleBias);
-									HeightfieldDescriptionData.Add(Description.MinMaxUV);
-								
-									const FMatrix WorldToLocal = Description.LocalToWorld.Inverse();
-
-									HeightfieldDescriptionData.Add(*(FVector4*)&WorldToLocal.M[0]);
-									HeightfieldDescriptionData.Add(*(FVector4*)&WorldToLocal.M[1]);
-									HeightfieldDescriptionData.Add(*(FVector4*)&WorldToLocal.M[2]);
-									HeightfieldDescriptionData.Add(*(FVector4*)&WorldToLocal.M[3]);
-
-									HeightfieldDescriptionData.Add(*(FVector4*)&Description.LocalToWorld.M[0]);
-									HeightfieldDescriptionData.Add(*(FVector4*)&Description.LocalToWorld.M[1]);
-									HeightfieldDescriptionData.Add(*(FVector4*)&Description.LocalToWorld.M[2]);
-									HeightfieldDescriptionData.Add(*(FVector4*)&Description.LocalToWorld.M[3]);
-								}
-
-								if (HeightfieldDescriptionData.Num() > GHeightfieldDescriptions.Data.MaxElements)
-								{
-									GHeightfieldDescriptions.Data.MaxElements = HeightfieldDescriptionData.Num() * 5 / 4;
-									GHeightfieldDescriptions.Data.Release();
-									GHeightfieldDescriptions.Data.Initialize();
-								}
-
-								void* LockedBuffer = RHILockVertexBuffer(GHeightfieldDescriptions.Data.Buffer, 0, GHeightfieldDescriptions.Data.Buffer->GetSize(), RLM_WriteOnly);
-								const uint32 MemcpySize = HeightfieldDescriptionData.GetTypeSize() * HeightfieldDescriptionData.Num();
-								check(GHeightfieldDescriptions.Data.Buffer->GetSize() >= MemcpySize);
-								FPlatformMemory::Memcpy(LockedBuffer, HeightfieldDescriptionData.GetData(), MemcpySize);
-								RHIUnlockVertexBuffer(GHeightfieldDescriptions.Data.Buffer);
-
-								UTexture2D* HeightfieldTexture = It.Key();
-
-								{
-									TShaderMapRef<FCalculateHeightfieldOcclusionCS> ComputeShader(View.ShaderMap);
-									RHICmdList.SetComputeShader(ComputeShader->GetComputeShader());
-									ComputeShader->SetParameters(RHICmdList, View, DepthLevel, HeightfieldTexture, HeightfieldDescriptions.Num(), Parameters);
-									DispatchIndirectComputeShader(RHICmdList, *ComputeShader, SurfaceCacheResources.DispatchParameters.Buffer, 0);
-
-									ComputeShader->UnsetParameters(RHICmdList);
-								}
-							}
-						}
-					}
+					ComputeShader->UnsetParameters(RHICmdList);
 				}
 
+				if (bUseDistanceFieldGI)
 				{
-					if (bUseDistanceFieldGI)
-					{
-						extern void ComputeIrradianceForSamples(
-							int32 DepthLevel,
-							FRHICommandListImmediate& RHICmdList,
-							const FViewInfo& View,
-							const FScene* Scene,
-							const FDistanceFieldAOParameters& Parameters,
-							FTemporaryIrradianceCacheResources* TemporaryIrradianceCacheResources);
-
-						ComputeIrradianceForSamples(DepthLevel, RHICmdList, View, Scene, Parameters, &GTemporaryIrradianceCacheResources);
-					}
-
-					{	
-						TShaderMapRef<TSetupFinalGatherIndirectArgumentsCS<false> > ComputeShader(View.ShaderMap);
-
-						RHICmdList.SetComputeShader(ComputeShader->GetComputeShader());
-						ComputeShader->SetParameters(RHICmdList, View, DepthLevel);
-						DispatchComputeShader(RHICmdList, *ComputeShader, 1, 1, 1);
-
-						ComputeShader->UnsetParameters(RHICmdList);
-					}
-
-					{
-						TShaderMapRef<FCombineConesCS> ComputeShader(View.ShaderMap);
-
-						RHICmdList.SetComputeShader(ComputeShader->GetComputeShader());
-						ComputeShader->SetParameters(RHICmdList, View, DepthLevel, Parameters, NULL, FMatrix::Identity, NULL, NULL);
-						DispatchIndirectComputeShader(RHICmdList, *ComputeShader, SurfaceCacheResources.DispatchParameters.Buffer, 0);
-
-						ComputeShader->UnsetParameters(RHICmdList);
-					}
+					extern void ComputeIrradianceForSamples(
+						int32 DepthLevel,
+						FRHICommandListImmediate& RHICmdList,
+						const FViewInfo& View,
+						const FScene* Scene,
+						const FDistanceFieldAOParameters& Parameters,
+						FTemporaryIrradianceCacheResources* TemporaryIrradianceCacheResources);
+
+					ComputeIrradianceForSamples(DepthLevel, RHICmdList, View, Scene, Parameters, &GTemporaryIrradianceCacheResources);
 				}
 			}
 
@@ -4116,7 +3805,7 @@ bool FDeferredShadingSceneRenderer::RenderDistanceFieldAOSurfaceCache(
 
 			GRenderTargetPool.VisualizeTexture.SetCheckPoint(RHICmdList, BentNormalOutput);
 
-			if (bApplyToSceneColor)
+			if (bVisualizeAmbientOcclusion || bVisualizeGlobalIllumination)
 			{
 				GSceneRenderTargets.BeginRenderingSceneColor(RHICmdList);
 			}
@@ -4142,7 +3831,7 @@ bool FDeferredShadingSceneRenderer::RenderDistanceFieldAOSurfaceCache(
 			}
 
 			// Upsample to full resolution, write to output
-			UpsampleBentNormalAO(RHICmdList, Views, BentNormalOutput, IrradianceOutput, !bApplyToSceneColor);
+			UpsampleBentNormalAO(RHICmdList, Views, BentNormalOutput, IrradianceOutput, bVisualizeAmbientOcclusion, bVisualizeGlobalIllumination);
 
 			return true;
 		}
@@ -4265,6 +3954,15 @@ bool FDeferredShadingSceneRenderer::ShouldRenderDynamicSkyLight() const
 		&& !ViewFamily.EngineShowFlags.VisualizeLightCulling;
 }
 
+bool FDeferredShadingSceneRenderer::ShouldRenderDistanceFieldAO() const
+{
+	return Scene->SkyLight->bCastShadows
+		&& ViewFamily.EngineShowFlags.DistanceFieldAO 
+		&& !ViewFamily.EngineShowFlags.VisualizeDistanceFieldAO
+		&& !ViewFamily.EngineShowFlags.VisualizeDistanceFieldGI
+		&& !ViewFamily.EngineShowFlags.VisualizeMeshDistanceFields;
+}
+
 void FDeferredShadingSceneRenderer::RenderDynamicSkyLighting(FRHICommandListImmediate& RHICmdList, TRefCountPtr<IPooledRenderTarget>& DynamicBentNormalAO)
 {
 	if (ShouldRenderDynamicSkyLight())
@@ -4276,13 +3974,9 @@ void FDeferredShadingSceneRenderer::RenderDynamicSkyLighting(FRHICommandListImme
 		FDistanceFieldAOParameters Parameters(Scene->SkyLight->OcclusionMaxDistance, Scene->SkyLight->Contrast);
 		TRefCountPtr<IPooledRenderTarget> DynamicIrradiance;
 
-		if (Scene->SkyLight->bCastShadows
-			&& ViewFamily.EngineShowFlags.DistanceFieldAO 
-			&& !ViewFamily.EngineShowFlags.VisualizeDistanceFieldAO
-			&& !ViewFamily.EngineShowFlags.VisualizeDistanceFieldGI
-			&& !ViewFamily.EngineShowFlags.VisualizeMeshDistanceFields)
+		if (ShouldRenderDistanceFieldAO())
 		{
-			bApplyShadowing = RenderDistanceFieldAOSurfaceCache(RHICmdList, Parameters, DynamicBentNormalAO, DynamicIrradiance, false);
+			bApplyShadowing = RenderDistanceFieldAOSurfaceCache(RHICmdList, Parameters, DynamicBentNormalAO, DynamicIrradiance, false, false);
 		}
 
 		GSceneRenderTargets.BeginRenderingSceneColor(RHICmdList);
diff --git a/Engine/Source/Runtime/Renderer/Private/DistanceFieldSurfaceCacheLighting.h b/Engine/Source/Runtime/Renderer/Private/DistanceFieldSurfaceCacheLighting.h
--- a/Engine/Source/Runtime/Renderer/Private/DistanceFieldSurfaceCacheLighting.h
+++ b/Engine/Source/Runtime/Renderer/Private/DistanceFieldSurfaceCacheLighting.h
@@ -314,10 +314,6 @@ public:
 	}
 };
 
-
-/** Generates unit length, stratified and uniformly distributed direction samples in a hemisphere. */
-extern void GenerateStratifiedUniformHemisphereSamples2(int32 NumThetaSteps, int32 NumPhiSteps, FRandomStream& RandomStream, TArray<FVector4>& Samples);
-
 extern void GetSpacedVectors(TArray<FVector, TInlineAllocator<9> >& OutVectors);
 
 BEGIN_UNIFORM_BUFFER_STRUCT(FAOSampleData2,)
diff --git a/Engine/Source/Runtime/Renderer/Private/HeightfieldLighting.cpp b/Engine/Source/Runtime/Renderer/Private/HeightfieldLighting.cpp
new file mode 100644
--- /dev/null
+++ b/Engine/Source/Runtime/Renderer/Private/HeightfieldLighting.cpp
@@ -0,0 +1,1626 @@
+// Copyright 1998-2014 Epic Games, Inc. All Rights Reserved.
+
+/*=============================================================================
+	HeightfieldLighting.cpp
+=============================================================================*/
+
+#include "RendererPrivate.h"
+#include "ScenePrivate.h"
+#include "UniformBuffer.h"
+#include "ShaderParameters.h"
+#include "PostProcessing.h"
+#include "SceneFilterRendering.h"
+#include "DistanceFieldLightingShared.h"
+#include "DistanceFieldSurfaceCacheLighting.h"
+#include "DistanceFieldGlobalIllumination.h"
+#include "RHICommandList.h"
+#include "SceneUtils.h"
+#include "DistanceFieldAtlas.h"
+
+int32 GAOHeightfieldOcclusion = 1;
+FAutoConsoleVariableRef CVarAOHeightfieldOcclusion(
+	TEXT("r.AOHeightfieldOcclusion"),
+	GAOHeightfieldOcclusion,
+	TEXT("Whether to compute AO from heightfields (landscape)"),
+	ECVF_Cheat | ECVF_RenderThreadSafe
+	);
+
+int32 GHeightfieldGlobalIllumination = 1;
+FAutoConsoleVariableRef CVarHeightfieldGlobalIllumination(
+	TEXT("r.HeightfieldGlobalIllumination"),
+	GHeightfieldGlobalIllumination,
+	TEXT(""),
+	ECVF_Cheat | ECVF_RenderThreadSafe
+	);
+
+int32 GHeightfieldHemisphereSamples = 32;
+FAutoConsoleVariableRef CVarHeightfieldHemisphereSamples(
+	TEXT("r.HeightfieldHemisphereSamples"),
+	GHeightfieldHemisphereSamples,
+	TEXT(""),
+	ECVF_Cheat | ECVF_RenderThreadSafe
+	);
+
+float GHeightfieldInnerBounceDistance = 3000;
+FAutoConsoleVariableRef CVarHeightfieldInnerBounceDistancer(
+	TEXT("r.HeightfieldInnerBounceDistance"),
+	GHeightfieldInnerBounceDistance,
+	TEXT(""),
+	ECVF_Cheat | ECVF_RenderThreadSafe
+	);
+
+float GHeightfieldOuterBounceDistanceScale = 3;
+FAutoConsoleVariableRef CVarHeightfieldOuterBounceDistanceScale(
+	TEXT("r.HeightfieldOuterBounceDistanceScale"),
+	GHeightfieldOuterBounceDistanceScale,
+	TEXT(""),
+	ECVF_Cheat | ECVF_RenderThreadSafe
+	);
+
+float GetGHeightfieldBounceDistance()
+{
+	return GHeightfieldInnerBounceDistance * GHeightfieldOuterBounceDistanceScale;
+}
+
+float GHeightfieldTargetUnitsPerTexel = 200;
+FAutoConsoleVariableRef CVarHeightfieldTargetUnitsPerTexel(
+	TEXT("r.HeightfieldTargetUnitsPerTexel"),
+	GHeightfieldTargetUnitsPerTexel,
+	TEXT(""),
+	ECVF_Cheat | ECVF_RenderThreadSafe
+	);
+
+void FHeightfieldLightingAtlas::InitDynamicRHI()
+{
+	if (AtlasSize.GetMin() > 0)
+	{
+		FPooledRenderTargetDesc Desc(FPooledRenderTargetDesc::Create2DDesc(
+			AtlasSize,
+			PF_G16, 
+			TexCreate_None,
+			TexCreate_RenderTargetable,
+			false));
+
+		GRenderTargetPool.FindFreeElement(Desc, Height, TEXT("HeightAtlas"));
+
+		FPooledRenderTargetDesc Desc2(FPooledRenderTargetDesc::Create2DDesc(
+			AtlasSize,
+			PF_R8G8, 
+			TexCreate_None,
+			TexCreate_RenderTargetable,
+			false));
+
+		GRenderTargetPool.FindFreeElement(Desc2, Normal, TEXT("NormalAtlas"));
+
+		FPooledRenderTargetDesc Desc3(FPooledRenderTargetDesc::Create2DDesc(
+			AtlasSize,
+			PF_G8, 
+			TexCreate_None,
+			TexCreate_RenderTargetable,
+			false));
+
+		GRenderTargetPool.FindFreeElement(Desc3, DirectionalLightShadowing, TEXT("HeightfieldShadowingAtlas"));
+
+		FPooledRenderTargetDesc Desc4(FPooledRenderTargetDesc::Create2DDesc(
+			AtlasSize,
+			PF_FloatR11G11B10, 
+			TexCreate_None,
+			TexCreate_RenderTargetable,
+			false));
+
+		GRenderTargetPool.FindFreeElement(Desc4, Lighting, TEXT("HeightfieldLightingAtlas"));
+	}
+}
+
+void FHeightfieldLightingAtlas::ReleaseDynamicRHI()
+{
+	GRenderTargetPool.FreeUnusedResource(Height);
+	GRenderTargetPool.FreeUnusedResource(Normal);
+	GRenderTargetPool.FreeUnusedResource(DirectionalLightShadowing);
+	GRenderTargetPool.FreeUnusedResource(Lighting);
+}
+
+void FHeightfieldLightingAtlas::InitializeForSize(FIntPoint InAtlasSize)
+{
+	if (InAtlasSize.X > AtlasSize.X || InAtlasSize.Y > AtlasSize.Y)
+	{
+		AtlasSize.X = FMath::Max(InAtlasSize.X, AtlasSize.X);
+		AtlasSize.Y = FMath::Max(InAtlasSize.Y, AtlasSize.Y);
+
+		if (IsInitialized())
+		{
+			UpdateRHI();
+		}
+		else
+		{
+			InitResource();
+		}
+	}
+}
+
+class FSubsectionHeightfieldDescriptionsResource : public FRenderResource
+{
+public:
+
+	FCPUUpdatedBuffer Data;
+
+	FSubsectionHeightfieldDescriptionsResource()
+	{
+		Data.Format = PF_A32B32G32R32F;
+		// In float4's, must match usf
+		Data.Stride = 4;
+	}
+
+	virtual void InitDynamicRHI()  override
+	{
+		Data.Initialize();
+	}
+
+	virtual void ReleaseDynamicRHI() override
+	{
+		Data.Release();
+	}
+};
+
+TGlobalResource<FSubsectionHeightfieldDescriptionsResource> GSubsectionHeightfieldDescriptions;
+
+class FHeightfieldTextureParameters
+{
+public:
+	void Bind(const FShaderParameterMap& ParameterMap)
+	{
+		HeightfieldTexture.Bind(ParameterMap,TEXT("HeightfieldTexture"));
+		HeightfieldSampler.Bind(ParameterMap,TEXT("HeightfieldSampler"));
+	}
+
+	friend FArchive& operator<<(FArchive& Ar,FHeightfieldTextureParameters& Parameters)
+	{
+		Ar << Parameters.HeightfieldTexture;
+		Ar << Parameters.HeightfieldSampler;
+		return Ar;
+	}
+
+	template<typename ShaderRHIParamRef>
+	void Set(FRHICommandList& RHICmdList, const ShaderRHIParamRef ShaderRHI, UTexture2D* HeightfieldTextureValue)
+	{
+		//@todo - shouldn't filter the heightfield, it's packed
+		SetTextureParameter(RHICmdList, ShaderRHI, HeightfieldTexture, HeightfieldTexture, TStaticSamplerState<SF_Bilinear>::GetRHI(), HeightfieldTextureValue->Resource->TextureRHI);
+	}
+
+private:
+	FShaderResourceParameter HeightfieldTexture;
+	FShaderResourceParameter HeightfieldSampler;
+};
+
+class FSubsectionHeightfieldDescriptionParameters
+{
+public:
+	void Bind(const FShaderParameterMap& ParameterMap)
+	{
+		SubsectionHeightfieldDescriptions.Bind(ParameterMap,TEXT("SubsectionHeightfieldDescriptions"));
+	}
+
+	friend FArchive& operator<<(FArchive& Ar,FSubsectionHeightfieldDescriptionParameters& Parameters)
+	{
+		Ar << Parameters.SubsectionHeightfieldDescriptions;
+		return Ar;
+	}
+
+	template<typename ShaderRHIParamRef>
+	void Set(FRHICommandList& RHICmdList, const ShaderRHIParamRef ShaderRHI)
+	{
+		SetSRVParameter(RHICmdList, ShaderRHI, SubsectionHeightfieldDescriptions, GSubsectionHeightfieldDescriptions.Data.BufferSRV);
+	}
+
+private:
+	FShaderResourceParameter SubsectionHeightfieldDescriptions;
+};
+
+int32 UploadSubsectionHeightfieldDescriptions(const TArray<FHeightfieldComponentDescription>& HeightfieldDescriptions, FVector2D InvLightingAtlasSize, float InvDownsampleFactor)
+{
+	TArray<FVector4, SceneRenderingAllocator> HeightfieldDescriptionData;
+	HeightfieldDescriptionData.Empty(HeightfieldDescriptions.Num() * GSubsectionHeightfieldDescriptions.Data.Stride);
+
+	for (int32 DescriptionIndex = 0; DescriptionIndex < HeightfieldDescriptions.Num(); DescriptionIndex++)
+	{
+		const FHeightfieldComponentDescription& Description = HeightfieldDescriptions[DescriptionIndex];
+
+		for (int32 SubsectionY = 0; SubsectionY < Description.NumSubsections; SubsectionY++)
+		{
+			for (int32 SubsectionX = 0; SubsectionX < Description.NumSubsections; SubsectionX++)
+			{
+				HeightfieldDescriptionData.Add(FVector4(SubsectionX, SubsectionY, 0, 0));
+				HeightfieldDescriptionData.Add(Description.SubsectionScaleAndBias);
+				HeightfieldDescriptionData.Add(Description.HeightfieldScaleBias);
+
+				// GlobalUVScaleBias = SubsectionSizeQuads / AtlasSize, (SubsectionBase + SubsectionId * SubsectionSizeQuads - AtlasMin) / AtlasSize
+				const FVector4 GlobalUVScaleBias(
+					Description.SubsectionScaleAndBias.X * InvLightingAtlasSize.X * InvDownsampleFactor, 
+					Description.SubsectionScaleAndBias.Y * InvLightingAtlasSize.Y * InvDownsampleFactor, 
+					(Description.LightingAtlasLocation.X + SubsectionX * Description.SubsectionScaleAndBias.X * InvDownsampleFactor) * InvLightingAtlasSize.X, 
+					(Description.LightingAtlasLocation.Y + SubsectionY * Description.SubsectionScaleAndBias.Y * InvDownsampleFactor) * InvLightingAtlasSize.Y);
+
+				HeightfieldDescriptionData.Add(GlobalUVScaleBias);
+			}
+		}
+	}
+
+	check(HeightfieldDescriptionData.Num() % GSubsectionHeightfieldDescriptions.Data.Stride == 0);
+
+	if (HeightfieldDescriptionData.Num() > GSubsectionHeightfieldDescriptions.Data.MaxElements)
+	{
+		GSubsectionHeightfieldDescriptions.Data.MaxElements = HeightfieldDescriptionData.Num() * 5 / 4;
+		GSubsectionHeightfieldDescriptions.Data.Release();
+		GSubsectionHeightfieldDescriptions.Data.Initialize();
+	}
+
+	void* LockedBuffer = RHILockVertexBuffer(GSubsectionHeightfieldDescriptions.Data.Buffer, 0, GSubsectionHeightfieldDescriptions.Data.Buffer->GetSize(), RLM_WriteOnly);
+	const uint32 MemcpySize = HeightfieldDescriptionData.GetTypeSize() * HeightfieldDescriptionData.Num();
+	check(GSubsectionHeightfieldDescriptions.Data.Buffer->GetSize() >= MemcpySize);
+	FPlatformMemory::Memcpy(LockedBuffer, HeightfieldDescriptionData.GetData(), MemcpySize);
+	RHIUnlockVertexBuffer(GSubsectionHeightfieldDescriptions.Data.Buffer);
+
+	return HeightfieldDescriptionData.Num() / GSubsectionHeightfieldDescriptions.Data.Stride;
+}
+
+class FHeightfieldSubsectionQuadVS : public FGlobalShader
+{
+	DECLARE_SHADER_TYPE(FHeightfieldSubsectionQuadVS, Global);
+public:
+
+	static bool ShouldCache(EShaderPlatform Platform)
+	{
+		return IsFeatureLevelSupported(Platform, ERHIFeatureLevel::SM5) && DoesPlatformSupportDistanceFieldAO(Platform);
+	}
+
+	/** Default constructor. */
+	FHeightfieldSubsectionQuadVS() {}
+
+	/** Initialization constructor. */
+	FHeightfieldSubsectionQuadVS(const ShaderMetaType::CompiledShaderInitializerType& Initializer)
+		: FGlobalShader(Initializer)
+	{
+		SubsectionHeightfieldParameters.Bind(Initializer.ParameterMap);
+	}
+
+	void SetParameters(FRHICommandList& RHICmdList, const FSceneView& View)
+	{
+		const FVertexShaderRHIParamRef ShaderRHI = GetVertexShader();
+
+		FGlobalShader::SetParameters(RHICmdList, ShaderRHI, View);
+		SubsectionHeightfieldParameters.Set(RHICmdList, ShaderRHI);
+	}
+
+	// FShader interface.
+	virtual bool Serialize(FArchive& Ar)
+	{
+		bool bShaderHasOutdatedParameters = FGlobalShader::Serialize(Ar);
+		Ar << SubsectionHeightfieldParameters;
+		return bShaderHasOutdatedParameters;
+	}
+
+private:
+
+	FSubsectionHeightfieldDescriptionParameters SubsectionHeightfieldParameters;
+};
+
+IMPLEMENT_SHADER_TYPE(,FHeightfieldSubsectionQuadVS,TEXT("HeightfieldLighting"),TEXT("HeightfieldSubsectionQuadVS"),SF_Vertex);
+
+class FInitializeHeightfieldsPS : public FGlobalShader
+{
+	DECLARE_SHADER_TYPE(FInitializeHeightfieldsPS, Global);
+public:
+
+	static bool ShouldCache(EShaderPlatform Platform)
+	{
+		return IsFeatureLevelSupported(Platform, ERHIFeatureLevel::SM5) && DoesPlatformSupportDistanceFieldAO(Platform);
+	}
+
+	/** Default constructor. */
+	FInitializeHeightfieldsPS() {}
+
+	/** Initialization constructor. */
+	FInitializeHeightfieldsPS(const ShaderMetaType::CompiledShaderInitializerType& Initializer)
+		: FGlobalShader(Initializer)
+	{
+		HeightfieldTextureParameters.Bind(Initializer.ParameterMap);
+	}
+
+	void SetParameters(FRHICommandList& RHICmdList, const FSceneView& View, UTexture2D* HeightfieldTextureValue)
+	{
+		const FPixelShaderRHIParamRef ShaderRHI = GetPixelShader();
+
+		FGlobalShader::SetParameters(RHICmdList, ShaderRHI, View);
+		HeightfieldTextureParameters.Set(RHICmdList, ShaderRHI, HeightfieldTextureValue);
+	}
+	// FShader interface.
+	virtual bool Serialize(FArchive& Ar)
+	{
+		bool bShaderHasOutdatedParameters = FGlobalShader::Serialize(Ar);
+		Ar << HeightfieldTextureParameters;
+		return bShaderHasOutdatedParameters;
+	}
+
+private:
+
+	FHeightfieldTextureParameters HeightfieldTextureParameters;
+};
+
+IMPLEMENT_SHADER_TYPE(,FInitializeHeightfieldsPS,TEXT("HeightfieldLighting"),TEXT("InitializeHeightfieldsPS"),SF_Pixel);
+
+class FQuadVertexBuffer : public FVertexBuffer
+{
+public:
+
+	FQuadVertexBuffer()
+	{
+	}
+
+	virtual void InitRHI() override
+	{
+		const uint32 Size = 6 * sizeof(FScreenVertex);
+		FRHIResourceCreateInfo CreateInfo;
+		VertexBufferRHI = RHICreateVertexBuffer(Size, BUF_Static, CreateInfo);
+		void* Buffer = RHILockVertexBuffer(VertexBufferRHI, 0, Size, RLM_WriteOnly);
+		FScreenVertex* DestVertex = (FScreenVertex*)Buffer;
+
+		DestVertex[0].Position = FVector2D(1,  1);
+		DestVertex[0].UV = FVector2D(1,	1);
+
+		DestVertex[1].Position = FVector2D(0,  1);
+		DestVertex[1].UV = FVector2D(0,	1);
+
+		DestVertex[2].Position = FVector2D(1,	0);
+		DestVertex[2].UV = FVector2D(1,	0);
+
+		DestVertex[3].Position = FVector2D(1,	0);
+		DestVertex[3].UV = FVector2D(1,	0);
+
+		DestVertex[4].Position = FVector2D(0,  1);
+		DestVertex[4].UV = FVector2D(0,	1);
+
+		DestVertex[5].Position = FVector2D(0,	0);
+		DestVertex[5].UV = FVector2D(0,	0);
+
+		RHIUnlockVertexBuffer(VertexBufferRHI);      
+	}
+};
+
+TGlobalResource<FQuadVertexBuffer> GQuadVertexBuffer;
+
+bool SupportsHeightfieldLighting(ERHIFeatureLevel::Type FeatureLevel, EShaderPlatform ShaderPlatform)
+{
+	return FeatureLevel >= ERHIFeatureLevel::SM5
+		&& DoesPlatformSupportDistanceFieldAO(ShaderPlatform);
+}
+
+bool AllowHeightfieldGI(const FViewInfo& View)
+{
+	return GHeightfieldGlobalIllumination 
+		&& View.State 
+		&& GDistanceFieldGI 
+		&& View.Family->EngineShowFlags.DistanceFieldGI;
+}
+
+void FHeightfieldLightingViewInfo::SetupVisibleHeightfields(const FViewInfo& View, FRHICommandListImmediate& RHICmdList)
+{
+	const FScene* Scene = (const FScene*)View.Family->Scene;
+	const int32 NumPrimitives = Scene->DistanceFieldSceneData.HeightfieldPrimitives.Num();
+
+	if ((AllowHeightfieldGI(View) || GAOHeightfieldOcclusion) 
+		&& NumPrimitives > 0
+		&& SupportsHeightfieldLighting(View.GetFeatureLevel(), View.GetShaderPlatform()))
+	{
+		const int32 NumPrimitives = Scene->DistanceFieldSceneData.HeightfieldPrimitives.Num();
+		extern float GAOMaxViewDistance;
+		const float MaxDistanceSquared = FMath::Square(GAOMaxViewDistance + GetGHeightfieldBounceDistance());
+		float LocalToWorldScale = 1;
+
+		for (int32 HeightfieldPrimitiveIndex = 0; HeightfieldPrimitiveIndex < NumPrimitives; HeightfieldPrimitiveIndex++)
+		{
+			const FPrimitiveSceneInfo* HeightfieldPrimitive = Scene->DistanceFieldSceneData.HeightfieldPrimitives[HeightfieldPrimitiveIndex];
+			const FBoxSphereBounds& PrimitiveBounds = HeightfieldPrimitive->Proxy->GetBounds();
+			const float DistanceToPrimitiveSq = (PrimitiveBounds.Origin - View.ViewMatrices.ViewOrigin).SizeSquared();
+						
+			if (View.ViewFrustum.IntersectSphere(PrimitiveBounds.Origin, PrimitiveBounds.SphereRadius + GetGHeightfieldBounceDistance())
+				&& DistanceToPrimitiveSq < MaxDistanceSquared)
+			{
+				UTexture2D* HeightfieldTexture = NULL;
+				FHeightfieldComponentDescription NewComponentDescription(HeightfieldPrimitive->Proxy->GetLocalToWorld());
+				HeightfieldPrimitive->Proxy->GetHeightfieldRepresentation(HeightfieldTexture, NewComponentDescription);
+
+				if (HeightfieldTexture && HeightfieldTexture->Resource->TextureRHI)
+				{
+					const FIntPoint HeightfieldSize = NewComponentDescription.HeightfieldRect.Size();
+
+					if (Heightfield.Rect.Area() == 0)
+					{
+						Heightfield.Rect = NewComponentDescription.HeightfieldRect;
+						LocalToWorldScale = NewComponentDescription.LocalToWorld.GetScaleVector().X;
+					}
+					else
+					{
+						Heightfield.Rect.Union(NewComponentDescription.HeightfieldRect);
+					}
+
+					TArray<FHeightfieldComponentDescription>& ComponentDescriptions = Heightfield.ComponentDescriptions.FindOrAdd(HeightfieldTexture);
+					ComponentDescriptions.Add(NewComponentDescription);
+				}
+			}
+		}
+
+		if (AllowHeightfieldGI(View) && Heightfield.ComponentDescriptions.Num() > 0)
+		{
+			FSceneViewState* ViewState = (FSceneViewState*)View.State;
+
+			{
+				FHeightfieldLightingAtlas*& ExistingAtlas = ViewState->HeightfieldLightingAtlas;
+
+				if (!ExistingAtlas)
+				{
+					ExistingAtlas = new FHeightfieldLightingAtlas();
+				}
+
+				Heightfield.DownsampleFactor = FMath::Max(FMath::TruncToInt(GHeightfieldTargetUnitsPerTexel / LocalToWorldScale), 1);
+				Heightfield.DownsampledRect = Heightfield.Rect / Heightfield.DownsampleFactor;
+				Heightfield.Rect.Min = FIntPoint::DivideAndRoundDown(Heightfield.Rect.Min, Heightfield.DownsampleFactor) * Heightfield.DownsampleFactor;
+				ExistingAtlas->InitializeForSize(Heightfield.DownsampledRect.Size());
+
+				for (TMap<UTexture2D*, TArray<FHeightfieldComponentDescription>>::TIterator It(Heightfield.ComponentDescriptions); It; ++It)
+				{
+					TArray<FHeightfieldComponentDescription>& HeightfieldDescriptions = It.Value();
+
+					for (int32 ComponentIndex = 0; ComponentIndex < HeightfieldDescriptions.Num(); ComponentIndex++)
+					{
+						FIntPoint RelativeAtlasOffset = HeightfieldDescriptions[ComponentIndex].HeightfieldRect.Min - Heightfield.Rect.Min;
+						HeightfieldDescriptions[ComponentIndex].LightingAtlasLocation = FVector2D(RelativeAtlasOffset.X, RelativeAtlasOffset.Y) / Heightfield.DownsampleFactor;
+					}
+				}
+
+				{
+					SCOPED_DRAW_EVENT(RHICmdList, InitializeHeightfield);
+					const FIntPoint LightingAtlasSize = ExistingAtlas->GetAtlasSize();
+					const FVector2D InvLightingAtlasSize(1.0f / LightingAtlasSize.X, 1.0f / LightingAtlasSize.Y);
+
+					FTextureRHIParamRef RenderTargets[2] =
+					{
+						ExistingAtlas->Height->GetRenderTargetItem().TargetableTexture,
+						ExistingAtlas->Normal->GetRenderTargetItem().TargetableTexture
+					};
+
+					SetRenderTargets(RHICmdList, ARRAY_COUNT(RenderTargets), RenderTargets, FTextureRHIParamRef(), 0, NULL);
+
+					FLinearColor ClearColors[2] = {FLinearColor(0, 0, 0, 0), FLinearColor(0, 0, 0, 0)};
+					RHICmdList.ClearMRT(true, ARRAY_COUNT(ClearColors), ClearColors, false, 0, false, 0, FIntRect());
+
+					RHICmdList.SetViewport(0, 0, 0.0f, LightingAtlasSize.X, LightingAtlasSize.Y, 1.0f);
+					RHICmdList.SetRasterizerState(TStaticRasterizerState<FM_Solid, CM_None>::GetRHI());
+					RHICmdList.SetDepthStencilState(TStaticDepthStencilState<false, CF_Always>::GetRHI());
+					RHICmdList.SetBlendState(TStaticBlendState<>::GetRHI());
+					RHICmdList.SetStreamSource(0, GQuadVertexBuffer.VertexBufferRHI, sizeof(FScreenVertex), 0);
+
+					TShaderMapRef<FHeightfieldSubsectionQuadVS> VertexShader(View.ShaderMap);
+					TShaderMapRef<FInitializeHeightfieldsPS> PixelShader(View.ShaderMap);
+
+					static FGlobalBoundShaderState BoundShaderState;
+					SetGlobalBoundShaderState(RHICmdList, View.GetFeatureLevel(), BoundShaderState, GScreenVertexDeclaration.VertexDeclarationRHI, *VertexShader, *PixelShader);
+
+					for (TMap<UTexture2D*, TArray<FHeightfieldComponentDescription>>::TConstIterator It(Heightfield.ComponentDescriptions); It; ++It)
+					{
+						const TArray<FHeightfieldComponentDescription>& HeightfieldDescriptions = It.Value();
+
+						if (HeightfieldDescriptions.Num() > 0)
+						{
+							const int32 NumQuads = UploadSubsectionHeightfieldDescriptions(HeightfieldDescriptions, InvLightingAtlasSize, 1.0f / Heightfield.DownsampleFactor);
+
+							UTexture2D* HeightfieldTexture = It.Key();
+
+							VertexShader->SetParameters(RHICmdList, View);
+							PixelShader->SetParameters(RHICmdList, View, HeightfieldTexture);
+
+							RHICmdList.DrawPrimitive(PT_TriangleList, 0, 2, NumQuads);
+						}
+					}
+				}
+			}
+		}
+	}
+}
+
+
+class FHeightfieldDescriptionsResource : public FRenderResource
+{
+public:
+
+	FCPUUpdatedBuffer Data;
+
+	FHeightfieldDescriptionsResource()
+	{
+		Data.Format = PF_A32B32G32R32F;
+		// In float4's, must match usf
+		Data.Stride = 12;
+	}
+
+	virtual void InitDynamicRHI()  override
+	{
+		Data.Initialize();
+	}
+
+	virtual void ReleaseDynamicRHI() override
+	{
+		Data.Release();
+	}
+};
+
+TGlobalResource<FHeightfieldDescriptionsResource> GHeightfieldDescriptions;
+
+void UploadHeightfieldDescriptions(const TArray<FHeightfieldComponentDescription>& HeightfieldDescriptions, FVector2D InvLightingAtlasSize, float InvDownsampleFactor)
+{
+	TArray<FVector4, SceneRenderingAllocator> HeightfieldDescriptionData;
+	HeightfieldDescriptionData.Empty(HeightfieldDescriptions.Num() * GHeightfieldDescriptions.Data.Stride);
+
+	for (int32 DescriptionIndex = 0; DescriptionIndex < HeightfieldDescriptions.Num(); DescriptionIndex++)
+	{
+		const FHeightfieldComponentDescription& Description = HeightfieldDescriptions[DescriptionIndex];
+
+		FVector4 HeightfieldScaleBias = Description.HeightfieldScaleBias;
+		check(HeightfieldScaleBias.X > 0);
+
+		// CalculateHeightfieldOcclusionCS needs to be fixed up if other values are ever supported
+		check(Description.NumSubsections == 1 || Description.NumSubsections == 2);
+
+		// Store the presence of subsections in the sign bit
+		HeightfieldScaleBias.X *= Description.NumSubsections > 1 ? -1 : 1;
+
+		HeightfieldDescriptionData.Add(HeightfieldScaleBias);
+		HeightfieldDescriptionData.Add(Description.MinMaxUV);
+
+		const FVector4 LightingUVScaleBias(
+			InvLightingAtlasSize.X * InvDownsampleFactor, 
+			InvLightingAtlasSize.Y * InvDownsampleFactor, 
+			Description.LightingAtlasLocation.X * InvLightingAtlasSize.X, 
+			Description.LightingAtlasLocation.Y * InvLightingAtlasSize.Y);
+
+		HeightfieldDescriptionData.Add(LightingUVScaleBias);
+					
+		HeightfieldDescriptionData.Add(FVector4(Description.HeightfieldRect.Size().X, Description.HeightfieldRect.Size().Y, InvLightingAtlasSize.X, InvLightingAtlasSize.Y));
+
+		const FMatrix WorldToLocal = Description.LocalToWorld.Inverse();
+
+		HeightfieldDescriptionData.Add(*(FVector4*)&WorldToLocal.M[0]);
+		HeightfieldDescriptionData.Add(*(FVector4*)&WorldToLocal.M[1]);
+		HeightfieldDescriptionData.Add(*(FVector4*)&WorldToLocal.M[2]);
+		HeightfieldDescriptionData.Add(*(FVector4*)&WorldToLocal.M[3]);
+
+		HeightfieldDescriptionData.Add(*(FVector4*)&Description.LocalToWorld.M[0]);
+		HeightfieldDescriptionData.Add(*(FVector4*)&Description.LocalToWorld.M[1]);
+		HeightfieldDescriptionData.Add(*(FVector4*)&Description.LocalToWorld.M[2]);
+		HeightfieldDescriptionData.Add(*(FVector4*)&Description.LocalToWorld.M[3]);
+	}
+
+	check(HeightfieldDescriptionData.Num() % GHeightfieldDescriptions.Data.Stride == 0);
+
+	if (HeightfieldDescriptionData.Num() > GHeightfieldDescriptions.Data.MaxElements)
+	{
+		GHeightfieldDescriptions.Data.MaxElements = HeightfieldDescriptionData.Num() * 5 / 4;
+		GHeightfieldDescriptions.Data.Release();
+		GHeightfieldDescriptions.Data.Initialize();
+	}
+
+	void* LockedBuffer = RHILockVertexBuffer(GHeightfieldDescriptions.Data.Buffer, 0, GHeightfieldDescriptions.Data.Buffer->GetSize(), RLM_WriteOnly);
+	const uint32 MemcpySize = HeightfieldDescriptionData.GetTypeSize() * HeightfieldDescriptionData.Num();
+	check(GHeightfieldDescriptions.Data.Buffer->GetSize() >= MemcpySize);
+	FPlatformMemory::Memcpy(LockedBuffer, HeightfieldDescriptionData.GetData(), MemcpySize);
+	RHIUnlockVertexBuffer(GHeightfieldDescriptions.Data.Buffer);
+}
+
+class FHeightfieldDescriptionParameters
+{
+public:
+	void Bind(const FShaderParameterMap& ParameterMap)
+	{
+		HeightfieldDescriptions.Bind(ParameterMap,TEXT("HeightfieldDescriptions"));
+		NumHeightfields.Bind(ParameterMap,TEXT("NumHeightfields"));
+	}
+
+	friend FArchive& operator<<(FArchive& Ar,FHeightfieldDescriptionParameters& Parameters)
+	{
+		Ar << Parameters.HeightfieldDescriptions;
+		Ar << Parameters.NumHeightfields;
+		return Ar;
+	}
+
+	template<typename ShaderRHIParamRef>
+	void Set(FRHICommandList& RHICmdList, const ShaderRHIParamRef ShaderRHI, int32 NumHeightfieldsValue)
+	{
+		SetSRVParameter(RHICmdList, ShaderRHI, HeightfieldDescriptions, GHeightfieldDescriptions.Data.BufferSRV);
+		SetShaderValue(RHICmdList, ShaderRHI, NumHeightfields, NumHeightfieldsValue);
+	}
+
+private:
+	FShaderResourceParameter HeightfieldDescriptions;
+	FShaderParameter NumHeightfields;
+};
+
+class FGlobalHeightfieldParameters
+{
+public:
+	void Bind(const FShaderParameterMap& ParameterMap)
+	{
+		GlobalHeightfieldTexture.Bind(ParameterMap,TEXT("GlobalHeightfieldTexture"));
+		GlobalNormalTexture.Bind(ParameterMap,TEXT("GlobalNormalTexture"));
+		GlobalHeightfieldSampler.Bind(ParameterMap,TEXT("GlobalHeightfieldSampler"));
+	}
+
+	friend FArchive& operator<<(FArchive& Ar,FGlobalHeightfieldParameters& Parameters)
+	{
+		Ar << Parameters.GlobalHeightfieldTexture;
+		Ar << Parameters.GlobalNormalTexture;
+		Ar << Parameters.GlobalHeightfieldSampler;
+		return Ar;
+	}
+
+	template<typename ShaderRHIParamRef>
+	void Set(FRHICommandList& RHICmdList, const ShaderRHIParamRef ShaderRHI, const FHeightfieldLightingAtlas& Atlas)
+	{
+		SetTextureParameter(RHICmdList, ShaderRHI, GlobalHeightfieldTexture, GlobalHeightfieldSampler, TStaticSamplerState<SF_Bilinear>::GetRHI(), Atlas.Height->GetRenderTargetItem().ShaderResourceTexture);
+		SetTextureParameter(RHICmdList, ShaderRHI, GlobalNormalTexture, GlobalHeightfieldSampler, TStaticSamplerState<SF_Bilinear>::GetRHI(), Atlas.Normal->GetRenderTargetItem().ShaderResourceTexture);
+	}
+
+private:
+	FShaderResourceParameter GlobalHeightfieldTexture;
+	FShaderResourceParameter GlobalNormalTexture;
+	FShaderResourceParameter GlobalHeightfieldSampler;
+};
+
+class FHeightfieldComponentQuadVS : public FGlobalShader
+{
+	DECLARE_SHADER_TYPE(FHeightfieldComponentQuadVS, Global);
+public:
+
+	static bool ShouldCache(EShaderPlatform Platform)
+	{
+		return IsFeatureLevelSupported(Platform, ERHIFeatureLevel::SM5) && DoesPlatformSupportDistanceFieldAO(Platform);
+	}
+
+	/** Default constructor. */
+	FHeightfieldComponentQuadVS() {}
+
+	/** Initialization constructor. */
+	FHeightfieldComponentQuadVS(const ShaderMetaType::CompiledShaderInitializerType& Initializer)
+		: FGlobalShader(Initializer)
+	{
+		HeightfieldDescriptionParameters.Bind(Initializer.ParameterMap);
+	}
+
+	void SetParameters(FRHICommandList& RHICmdList, const FSceneView& View, int32 NumHeightfieldsValue)
+	{
+		const FVertexShaderRHIParamRef ShaderRHI = GetVertexShader();
+
+		FGlobalShader::SetParameters(RHICmdList, ShaderRHI, View);
+		HeightfieldDescriptionParameters.Set(RHICmdList, ShaderRHI, NumHeightfieldsValue);
+	}
+
+	// FShader interface.
+	virtual bool Serialize(FArchive& Ar)
+	{
+		bool bShaderHasOutdatedParameters = FGlobalShader::Serialize(Ar);
+		Ar << HeightfieldDescriptionParameters;
+		return bShaderHasOutdatedParameters;
+	}
+
+private:
+
+	FHeightfieldDescriptionParameters HeightfieldDescriptionParameters;
+};
+
+IMPLEMENT_SHADER_TYPE(,FHeightfieldComponentQuadVS,TEXT("HeightfieldLighting"),TEXT("HeightfieldComponentQuadVS"),SF_Vertex);
+
+class FShadowHeightfieldsPS : public FGlobalShader
+{
+	DECLARE_SHADER_TYPE(FShadowHeightfieldsPS, Global);
+public:
+
+	static bool ShouldCache(EShaderPlatform Platform)
+	{
+		return IsFeatureLevelSupported(Platform, ERHIFeatureLevel::SM5) && DoesPlatformSupportDistanceFieldAO(Platform);
+	}
+
+	/** Default constructor. */
+	FShadowHeightfieldsPS() {}
+
+	/** Initialization constructor. */
+	FShadowHeightfieldsPS(const ShaderMetaType::CompiledShaderInitializerType& Initializer)
+		: FGlobalShader(Initializer)
+	{
+		HeightfieldDescriptionParameters.Bind(Initializer.ParameterMap);
+		HeightfieldTextureParameters.Bind(Initializer.ParameterMap);
+		GlobalHeightfieldParameters.Bind(Initializer.ParameterMap);
+		WorldToShadow.Bind(Initializer.ParameterMap, TEXT("WorldToShadow"));
+		ShadowmapMinMax.Bind(Initializer.ParameterMap, TEXT("ShadowmapMinMax"));
+		ShadowDepthBias.Bind(Initializer.ParameterMap, TEXT("ShadowDepthBias"));
+		CascadeDepthMinMax.Bind(Initializer.ParameterMap, TEXT("CascadeDepthMinMax"));
+		ShadowDepthTexture.Bind(Initializer.ParameterMap, TEXT("ShadowDepthTexture"));
+		ShadowDepthTextureSampler.Bind(Initializer.ParameterMap, TEXT("ShadowDepthTextureSampler"));
+	}
+
+	void SetParameters(FRHICommandList& RHICmdList, const FSceneView& View, const FProjectedShadowInfo* ProjectedShadowInfo, UTexture2D* HeightfieldTextureValue, int32 NumHeightfieldsValue, const FHeightfieldLightingAtlas& Atlas)
+	{
+		const FPixelShaderRHIParamRef ShaderRHI = GetPixelShader();
+
+		FGlobalShader::SetParameters(RHICmdList, ShaderRHI, View);
+		HeightfieldDescriptionParameters.Set(RHICmdList, ShaderRHI, NumHeightfieldsValue);
+		HeightfieldTextureParameters.Set(RHICmdList, ShaderRHI, HeightfieldTextureValue);
+		GlobalHeightfieldParameters.Set(RHICmdList, ShaderRHI, Atlas);
+
+		FVector4 ShadowmapMinMaxValue;
+		FMatrix WorldToShadowMatrixValue = ProjectedShadowInfo->GetWorldToShadowMatrix(ShadowmapMinMaxValue);
+
+		SetShaderValue(RHICmdList, ShaderRHI, WorldToShadow, WorldToShadowMatrixValue);
+		SetShaderValue(RHICmdList, ShaderRHI, ShadowmapMinMax, ShadowmapMinMaxValue);
+		SetShaderValue(RHICmdList, ShaderRHI, ShadowDepthBias, ProjectedShadowInfo->GetShaderDepthBias());
+		SetShaderValue(RHICmdList, ShaderRHI, CascadeDepthMinMax, FVector2D(ProjectedShadowInfo->CascadeSettings.SplitNear, ProjectedShadowInfo->CascadeSettings.SplitFar));
+
+		FTexture2DRHIRef ShadowDepthTextureValue = GSceneRenderTargets.GetShadowDepthZTexture(ProjectedShadowInfo->bAllocatedInPreshadowCache);
+		FSamplerStateRHIParamRef DepthSamplerState = TStaticSamplerState<SF_Point,AM_Clamp,AM_Clamp,AM_Clamp>::GetRHI();
+
+		SetTextureParameter(RHICmdList, ShaderRHI, ShadowDepthTexture, ShadowDepthTextureSampler, DepthSamplerState, ShadowDepthTextureValue);	
+	}
+	// FShader interface.
+	virtual bool Serialize(FArchive& Ar)
+	{
+		bool bShaderHasOutdatedParameters = FGlobalShader::Serialize(Ar);
+		Ar << HeightfieldDescriptionParameters;
+		Ar << HeightfieldTextureParameters;
+		Ar << GlobalHeightfieldParameters;
+		Ar << WorldToShadow;
+		Ar << ShadowmapMinMax;
+		Ar << ShadowDepthBias;
+		Ar << CascadeDepthMinMax;
+		Ar << ShadowDepthTexture;
+		Ar << ShadowDepthTextureSampler;
+		return bShaderHasOutdatedParameters;
+	}
+
+private:
+
+	FHeightfieldDescriptionParameters HeightfieldDescriptionParameters;
+	FHeightfieldTextureParameters HeightfieldTextureParameters;
+	FGlobalHeightfieldParameters GlobalHeightfieldParameters;
+	FShaderParameter WorldToShadow;
+	FShaderParameter ShadowmapMinMax;
+	FShaderParameter ShadowDepthBias;
+	FShaderParameter CascadeDepthMinMax;
+	FShaderResourceParameter ShadowDepthTexture;
+	FShaderResourceParameter ShadowDepthTextureSampler;
+};
+
+IMPLEMENT_SHADER_TYPE(,FShadowHeightfieldsPS,TEXT("HeightfieldLighting"),TEXT("ShadowHeightfieldsPS"),SF_Pixel);
+
+void FHeightfieldLightingViewInfo::ClearShadowing(const FViewInfo& View, FRHICommandListImmediate& RHICmdList, const FLightSceneInfo& LightSceneInfo) const
+{
+	if (AllowHeightfieldGI(View)
+		&& SupportsHeightfieldLighting(View.GetFeatureLevel(), View.GetShaderPlatform())
+		&& Heightfield.ComponentDescriptions.Num() > 0
+		&& LightSceneInfo.Proxy->GetLightType() == LightType_Directional
+		&& LightSceneInfo.Proxy->CastsDynamicShadow())
+	{
+		FSceneViewState* ViewState = (FSceneViewState*)View.State;
+		const FHeightfieldLightingAtlas& Atlas = *ViewState->HeightfieldLightingAtlas;
+		SetRenderTarget(RHICmdList, Atlas.DirectionalLightShadowing->GetRenderTargetItem().TargetableTexture, NULL);
+		RHICmdList.Clear(true, FLinearColor(1, 1, 1), false, 0, false, 0, FIntRect());
+	}
+}
+
+void FHeightfieldLightingViewInfo::ComputeShadowMapShadowing(const FViewInfo& View, FRHICommandListImmediate& RHICmdList, const FProjectedShadowInfo* ProjectedShadowInfo) const
+{
+	if (AllowHeightfieldGI(View)
+		&& SupportsHeightfieldLighting(View.GetFeatureLevel(), View.GetShaderPlatform())
+		&& Heightfield.ComponentDescriptions.Num() > 0
+		&& ProjectedShadowInfo->IsWholeSceneDirectionalShadow()
+		&& ProjectedShadowInfo->DependentView == &View
+		&& !ProjectedShadowInfo->CascadeSettings.bRayTracedDistanceField)
+	{
+		SCOPED_DRAW_EVENT(RHICmdList, HeightfieldShadowMapShadowingForGI);
+
+		FSceneViewState* ViewState = (FSceneViewState*)View.State;
+
+		{
+			const FHeightfieldLightingAtlas& Atlas = *ViewState->HeightfieldLightingAtlas;
+
+			const FIntPoint LightingAtlasSize = Atlas.GetAtlasSize();
+			const FVector2D InvLightingAtlasSize(1.0f / LightingAtlasSize.X, 1.0f / LightingAtlasSize.Y);
+			SetRenderTarget(RHICmdList, Atlas.DirectionalLightShadowing->GetRenderTargetItem().TargetableTexture, NULL);
+
+			RHICmdList.SetViewport(0, 0, 0.0f, LightingAtlasSize.X, LightingAtlasSize.Y, 1.0f);
+			RHICmdList.SetRasterizerState(TStaticRasterizerState<FM_Solid, CM_None>::GetRHI());
+			RHICmdList.SetDepthStencilState(TStaticDepthStencilState<false, CF_Always>::GetRHI());
+			// Combine with other shadow types with min (ray traced)
+			RHICmdList.SetBlendState(TStaticBlendState<CW_RED, BO_Min, BF_One, BF_One>::GetRHI());
+			RHICmdList.SetStreamSource(0, GQuadVertexBuffer.VertexBufferRHI, sizeof(FScreenVertex), 0);
+
+			TShaderMapRef<FHeightfieldComponentQuadVS> VertexShader(View.ShaderMap);
+			TShaderMapRef<FShadowHeightfieldsPS> PixelShader(View.ShaderMap);
+
+			static FGlobalBoundShaderState BoundShaderState;
+			SetGlobalBoundShaderState(RHICmdList, View.GetFeatureLevel(), BoundShaderState, GScreenVertexDeclaration.VertexDeclarationRHI, *VertexShader, *PixelShader);
+
+			for (TMap<UTexture2D*, TArray<FHeightfieldComponentDescription>>::TConstIterator It(Heightfield.ComponentDescriptions); It; ++It)
+			{
+				const TArray<FHeightfieldComponentDescription>& HeightfieldDescriptions = It.Value();
+
+				if (HeightfieldDescriptions.Num() > 0)
+				{
+					//@todo - cull heightfield tiles with shadow bounds
+					UploadHeightfieldDescriptions(HeightfieldDescriptions, InvLightingAtlasSize, 1.0f / Heightfield.DownsampleFactor);
+
+					UTexture2D* HeightfieldTexture = It.Key();
+
+					VertexShader->SetParameters(RHICmdList, View, HeightfieldDescriptions.Num());
+					PixelShader->SetParameters(RHICmdList, View, ProjectedShadowInfo, HeightfieldTexture, HeightfieldDescriptions.Num(), Atlas);
+
+					RHICmdList.DrawPrimitive(PT_TriangleList, 0, 2, HeightfieldDescriptions.Num());
+				}
+			}
+		}
+	}
+}
+
+
+class FRayTracedShadowHeightfieldsPS : public FGlobalShader
+{
+	DECLARE_SHADER_TYPE(FRayTracedShadowHeightfieldsPS, Global);
+public:
+
+	static bool ShouldCache(EShaderPlatform Platform)
+	{
+		return IsFeatureLevelSupported(Platform, ERHIFeatureLevel::SM5) && DoesPlatformSupportDistanceFieldAO(Platform);
+	}
+
+	/** Default constructor. */
+	FRayTracedShadowHeightfieldsPS() {}
+
+	/** Initialization constructor. */
+	FRayTracedShadowHeightfieldsPS(const ShaderMetaType::CompiledShaderInitializerType& Initializer)
+		: FGlobalShader(Initializer)
+	{
+		HeightfieldDescriptionParameters.Bind(Initializer.ParameterMap);
+		GlobalHeightfieldParameters.Bind(Initializer.ParameterMap);
+		ObjectParameters.Bind(Initializer.ParameterMap);
+		LightDirection.Bind(Initializer.ParameterMap, TEXT("LightDirection"));
+		TanLightAngle.Bind(Initializer.ParameterMap, TEXT("TanLightAngle"));
+		CascadeDepthMinMax.Bind(Initializer.ParameterMap, TEXT("CascadeDepthMinMax"));
+		ShadowTileHeadDataUnpacked.Bind(Initializer.ParameterMap, TEXT("ShadowTileHeadDataUnpacked"));
+		ShadowTileArrayData.Bind(Initializer.ParameterMap, TEXT("ShadowTileArrayData"));
+		ShadowTileListGroupSize.Bind(Initializer.ParameterMap, TEXT("ShadowTileListGroupSize"));
+		WorldToShadow.Bind(Initializer.ParameterMap, TEXT("WorldToShadow"));
+		TwoSidedMeshDistanceBias.Bind(Initializer.ParameterMap, TEXT("TwoSidedMeshDistanceBias"));
+	}
+
+	void SetParameters(
+		FRHICommandList& RHICmdList, 
+		const FSceneView& View, 
+		const FProjectedShadowInfo* ProjectedShadowInfo, 
+		int32 NumHeightfieldsValue, 
+		const FHeightfieldLightingAtlas& Atlas, 
+		FLightTileIntersectionResources* TileIntersectionResources, 
+		FDistanceFieldObjectBufferResource& CulledObjectBuffers)
+	{
+		const FPixelShaderRHIParamRef ShaderRHI = GetPixelShader();
+
+		FGlobalShader::SetParameters(RHICmdList, ShaderRHI, View);
+		HeightfieldDescriptionParameters.Set(RHICmdList, ShaderRHI, NumHeightfieldsValue);
+		GlobalHeightfieldParameters.Set(RHICmdList, ShaderRHI, Atlas);
+		ObjectParameters.Set(RHICmdList, ShaderRHI, CulledObjectBuffers.Buffers);
+
+		SetShaderValue(RHICmdList, ShaderRHI, LightDirection, ProjectedShadowInfo->GetLightSceneInfo().Proxy->GetDirection());
+
+		const float LightSourceAngle = FMath::Clamp(ProjectedShadowInfo->GetLightSceneInfo().Proxy->GetLightSourceAngle(), 0.001f, 5.0f) * PI / 180.0f;
+		SetShaderValue(RHICmdList, ShaderRHI, TanLightAngle, FMath::Tan(LightSourceAngle));
+
+		SetShaderValue(RHICmdList, ShaderRHI, CascadeDepthMinMax, FVector2D(ProjectedShadowInfo->CascadeSettings.SplitNear, ProjectedShadowInfo->CascadeSettings.SplitFar));
+
+		check(TileIntersectionResources || !ShadowTileHeadDataUnpacked.IsBound());
+
+		if (TileIntersectionResources)
+		{
+			SetSRVParameter(RHICmdList, ShaderRHI, ShadowTileHeadDataUnpacked, TileIntersectionResources->TileHeadDataUnpacked.SRV);
+			SetSRVParameter(RHICmdList, ShaderRHI, ShadowTileArrayData, TileIntersectionResources->TileArrayData.SRV);
+			SetShaderValue(RHICmdList, ShaderRHI, ShadowTileListGroupSize, TileIntersectionResources->TileDimensions);
+		}
+
+		FMatrix WorldToShadowMatrixValue = FTranslationMatrix(ProjectedShadowInfo->PreShadowTranslation) * ProjectedShadowInfo->SubjectAndReceiverMatrix;
+		SetShaderValue(RHICmdList, ShaderRHI, WorldToShadow, WorldToShadowMatrixValue);
+
+		extern float GTwoSidedMeshDistanceBias;
+		SetShaderValue(RHICmdList, ShaderRHI, TwoSidedMeshDistanceBias, GTwoSidedMeshDistanceBias);
+	}
+	// FShader interface.
+	virtual bool Serialize(FArchive& Ar)
+	{
+		bool bShaderHasOutdatedParameters = FGlobalShader::Serialize(Ar);
+		Ar << HeightfieldDescriptionParameters;
+		Ar << GlobalHeightfieldParameters;
+		Ar << ObjectParameters;
+		Ar << LightDirection;
+		Ar << TanLightAngle;
+		Ar << CascadeDepthMinMax;
+		Ar << ShadowTileHeadDataUnpacked;
+		Ar << ShadowTileArrayData;
+		Ar << ShadowTileListGroupSize;
+		Ar << WorldToShadow;
+		Ar << TwoSidedMeshDistanceBias;
+		return bShaderHasOutdatedParameters;
+	}
+
+private:
+
+	FHeightfieldDescriptionParameters HeightfieldDescriptionParameters;
+	FGlobalHeightfieldParameters GlobalHeightfieldParameters;
+	FDistanceFieldCulledObjectBufferParameters ObjectParameters;
+	FShaderParameter LightDirection;
+	FShaderParameter TanLightAngle;
+	FShaderParameter CascadeDepthMinMax;
+	FShaderResourceParameter ShadowTileHeadDataUnpacked;
+	FShaderResourceParameter ShadowTileArrayData;
+	FShaderParameter ShadowTileListGroupSize;
+	FShaderParameter WorldToShadow;
+	FShaderParameter TwoSidedMeshDistanceBias;
+};
+
+IMPLEMENT_SHADER_TYPE(,FRayTracedShadowHeightfieldsPS,TEXT("HeightfieldLighting"),TEXT("RayTracedShadowHeightfieldsPS"),SF_Pixel);
+
+
+void FHeightfieldLightingViewInfo::ComputeRayTracedShadowing(
+	const FViewInfo& View, 
+	FRHICommandListImmediate& RHICmdList, 
+	const FProjectedShadowInfo* ProjectedShadowInfo, 
+	FLightTileIntersectionResources* TileIntersectionResources,
+	FDistanceFieldObjectBufferResource& CulledObjectBuffers) const
+{
+	if (AllowHeightfieldGI(View)
+		&& SupportsHeightfieldLighting(View.GetFeatureLevel(), View.GetShaderPlatform())
+		&& Heightfield.ComponentDescriptions.Num() > 0
+		&& ProjectedShadowInfo->IsWholeSceneDirectionalShadow()
+		&& ProjectedShadowInfo->DependentView == &View
+		&& ProjectedShadowInfo->CascadeSettings.bRayTracedDistanceField)
+	{
+		SCOPED_DRAW_EVENT(RHICmdList, HeightfieldRayTracedShadowingForGI);
+
+		FSceneViewState* ViewState = (FSceneViewState*)View.State;
+
+		{
+			const FHeightfieldLightingAtlas& Atlas = *ViewState->HeightfieldLightingAtlas;
+
+			const FIntPoint LightingAtlasSize = Atlas.GetAtlasSize();
+			const FVector2D InvLightingAtlasSize(1.0f / LightingAtlasSize.X, 1.0f / LightingAtlasSize.Y);
+			SetRenderTarget(RHICmdList, Atlas.DirectionalLightShadowing->GetRenderTargetItem().TargetableTexture, NULL);
+
+			RHICmdList.SetViewport(0, 0, 0.0f, LightingAtlasSize.X, LightingAtlasSize.Y, 1.0f);
+			RHICmdList.SetRasterizerState(TStaticRasterizerState<FM_Solid, CM_None>::GetRHI());
+			RHICmdList.SetDepthStencilState(TStaticDepthStencilState<false, CF_Always>::GetRHI());
+			// Combine with other shadow types with min (CSM)
+			RHICmdList.SetBlendState(TStaticBlendState<CW_RED, BO_Min, BF_One, BF_One>::GetRHI());
+			RHICmdList.SetStreamSource(0, GQuadVertexBuffer.VertexBufferRHI, sizeof(FScreenVertex), 0);
+
+			TShaderMapRef<FHeightfieldComponentQuadVS> VertexShader(View.ShaderMap);
+			TShaderMapRef<FRayTracedShadowHeightfieldsPS> PixelShader(View.ShaderMap);
+
+			static FGlobalBoundShaderState BoundShaderState;
+			SetGlobalBoundShaderState(RHICmdList, View.GetFeatureLevel(), BoundShaderState, GScreenVertexDeclaration.VertexDeclarationRHI, *VertexShader, *PixelShader);
+
+			for (TMap<UTexture2D*, TArray<FHeightfieldComponentDescription>>::TConstIterator It(Heightfield.ComponentDescriptions); It; ++It)
+			{
+				const TArray<FHeightfieldComponentDescription>& HeightfieldDescriptions = It.Value();
+
+				if (HeightfieldDescriptions.Num() > 0)
+				{
+					UploadHeightfieldDescriptions(HeightfieldDescriptions, InvLightingAtlasSize, 1.0f / Heightfield.DownsampleFactor);
+
+					UTexture2D* HeightfieldTexture = It.Key();
+
+					VertexShader->SetParameters(RHICmdList, View, HeightfieldDescriptions.Num());
+					PixelShader->SetParameters(RHICmdList, View, ProjectedShadowInfo, HeightfieldDescriptions.Num(), Atlas, TileIntersectionResources, CulledObjectBuffers);
+
+					RHICmdList.DrawPrimitive(PT_TriangleList, 0, 2, HeightfieldDescriptions.Num());
+				}
+			}
+		}
+	}
+}
+
+class FLightHeightfieldsPS : public FGlobalShader
+{
+	DECLARE_SHADER_TYPE(FLightHeightfieldsPS, Global);
+public:
+
+	static bool ShouldCache(EShaderPlatform Platform)
+	{
+		return IsFeatureLevelSupported(Platform, ERHIFeatureLevel::SM5) && DoesPlatformSupportDistanceFieldAO(Platform);
+	}
+
+	/** Default constructor. */
+	FLightHeightfieldsPS() {}
+
+	/** Initialization constructor. */
+	FLightHeightfieldsPS(const ShaderMetaType::CompiledShaderInitializerType& Initializer)
+		: FGlobalShader(Initializer)
+	{
+		HeightfieldDescriptionParameters.Bind(Initializer.ParameterMap);
+		HeightfieldTextureParameters.Bind(Initializer.ParameterMap);
+		GlobalHeightfieldParameters.Bind(Initializer.ParameterMap);
+		LightDirection.Bind(Initializer.ParameterMap, TEXT("LightDirection"));
+		LightColor.Bind(Initializer.ParameterMap, TEXT("LightColor"));
+		HeightfieldShadowing.Bind(Initializer.ParameterMap, TEXT("HeightfieldShadowing"));
+		HeightfieldShadowingSampler.Bind(Initializer.ParameterMap, TEXT("HeightfieldShadowingSampler"));
+	}
+
+	void SetParameters(FRHICommandList& RHICmdList, const FSceneView& View, const FLightSceneProxy* LightSceneProxy, UTexture2D* HeightfieldTextureValue, int32 NumHeightfieldsValue, const FHeightfieldLightingAtlas& Atlas)
+	{
+		const FPixelShaderRHIParamRef ShaderRHI = GetPixelShader();
+
+		FGlobalShader::SetParameters(RHICmdList, ShaderRHI, View);
+		HeightfieldDescriptionParameters.Set(RHICmdList, ShaderRHI, NumHeightfieldsValue);
+		HeightfieldTextureParameters.Set(RHICmdList, ShaderRHI, HeightfieldTextureValue);
+		GlobalHeightfieldParameters.Set(RHICmdList, ShaderRHI, Atlas);
+
+		SetShaderValue(RHICmdList, ShaderRHI, LightDirection, LightSceneProxy->GetDirection());
+		SetShaderValue(RHICmdList, ShaderRHI, LightColor, LightSceneProxy->GetColor() * LightSceneProxy->GetIndirectLightingScale());
+
+		const FSceneViewState* ViewState = (const FSceneViewState*)View.State;
+		SetTextureParameter(RHICmdList, ShaderRHI, HeightfieldShadowing, HeightfieldShadowingSampler, TStaticSamplerState<SF_Bilinear>::GetRHI(), Atlas.DirectionalLightShadowing->GetRenderTargetItem().ShaderResourceTexture);
+	}
+	// FShader interface.
+	virtual bool Serialize(FArchive& Ar)
+	{
+		bool bShaderHasOutdatedParameters = FGlobalShader::Serialize(Ar);
+		Ar << HeightfieldDescriptionParameters;
+		Ar << HeightfieldTextureParameters;
+		Ar << GlobalHeightfieldParameters;
+		Ar << LightDirection;
+		Ar << LightColor;
+		Ar << HeightfieldShadowing;
+		Ar << HeightfieldShadowingSampler;
+		return bShaderHasOutdatedParameters;
+	}
+
+private:
+
+	FHeightfieldDescriptionParameters HeightfieldDescriptionParameters;
+	FHeightfieldTextureParameters HeightfieldTextureParameters;
+	FGlobalHeightfieldParameters GlobalHeightfieldParameters;
+	FShaderParameter LightDirection;
+	FShaderParameter LightColor;
+	FShaderResourceParameter HeightfieldShadowing;
+	FShaderResourceParameter HeightfieldShadowingSampler;
+};
+
+IMPLEMENT_SHADER_TYPE(,FLightHeightfieldsPS,TEXT("HeightfieldLighting"),TEXT("LightHeightfieldsPS"),SF_Pixel);
+
+void FHeightfieldLightingViewInfo::ComputeLighting(const FViewInfo& View, FRHICommandListImmediate& RHICmdList, const FLightSceneInfo& LightSceneInfo) const
+{
+	if (AllowHeightfieldGI(View) 
+		&& SupportsHeightfieldLighting(View.GetFeatureLevel(), View.GetShaderPlatform())
+		//@todo - handle local lights
+		&& LightSceneInfo.Proxy->GetLightType() == LightType_Directional
+		&& LightSceneInfo.Proxy->CastsDynamicShadow()
+		&& Heightfield.ComponentDescriptions.Num() > 0)
+	{
+		SCOPED_DRAW_EVENT(RHICmdList, HeightfieldLightingForGI);
+
+		FSceneViewState* ViewState = (FSceneViewState*)View.State;
+
+		{
+			const FHeightfieldLightingAtlas& Atlas = *ViewState->HeightfieldLightingAtlas;
+
+			const FIntPoint LightingAtlasSize = Atlas.GetAtlasSize();
+			const FVector2D InvLightingAtlasSize(1.0f / LightingAtlasSize.X, 1.0f / LightingAtlasSize.Y);
+
+			SetRenderTarget(RHICmdList, Atlas.Lighting->GetRenderTargetItem().TargetableTexture, NULL);
+
+			// We will read outside the valid area later during light transfer
+			RHICmdList.Clear(true, FLinearColor(0, 0, 0), false, 0, false, 0, FIntRect());
+
+			RHICmdList.SetViewport(0, 0, 0.0f, LightingAtlasSize.X, LightingAtlasSize.Y, 1.0f);
+			RHICmdList.SetRasterizerState(TStaticRasterizerState<FM_Solid, CM_None>::GetRHI());
+			RHICmdList.SetDepthStencilState(TStaticDepthStencilState<false, CF_Always>::GetRHI());
+			RHICmdList.SetBlendState(TStaticBlendState<>::GetRHI());
+			RHICmdList.SetStreamSource(0, GQuadVertexBuffer.VertexBufferRHI, sizeof(FScreenVertex), 0);
+
+			TShaderMapRef<FHeightfieldComponentQuadVS> VertexShader(View.ShaderMap);
+			TShaderMapRef<FLightHeightfieldsPS> PixelShader(View.ShaderMap);
+
+			static FGlobalBoundShaderState BoundShaderState;
+			SetGlobalBoundShaderState(RHICmdList, View.GetFeatureLevel(), BoundShaderState, GScreenVertexDeclaration.VertexDeclarationRHI, *VertexShader, *PixelShader);
+
+			for (TMap<UTexture2D*, TArray<FHeightfieldComponentDescription>>::TConstIterator It(Heightfield.ComponentDescriptions); It; ++It)
+			{
+				const TArray<FHeightfieldComponentDescription>& HeightfieldDescriptions = It.Value();
+
+				if (HeightfieldDescriptions.Num() > 0)
+				{
+					UploadHeightfieldDescriptions(HeightfieldDescriptions, InvLightingAtlasSize, 1.0f / Heightfield.DownsampleFactor);
+
+					UTexture2D* HeightfieldTexture = It.Key();
+
+					VertexShader->SetParameters(RHICmdList, View, HeightfieldDescriptions.Num());
+					PixelShader->SetParameters(RHICmdList, View, LightSceneInfo.Proxy, HeightfieldTexture, HeightfieldDescriptions.Num(), Atlas);
+
+					RHICmdList.DrawPrimitive(PT_TriangleList, 0, 2, HeightfieldDescriptions.Num());
+				}
+			}
+		}
+	}
+}
+
+class FCalculateHeightfieldOcclusionCS : public FGlobalShader
+{
+	DECLARE_SHADER_TYPE(FCalculateHeightfieldOcclusionCS,Global)
+public:
+
+	static bool ShouldCache(EShaderPlatform Platform)
+	{
+		return IsFeatureLevelSupported(Platform, ERHIFeatureLevel::SM5) && DoesPlatformSupportDistanceFieldAO(Platform);
+	}
+
+	static void ModifyCompilationEnvironment(EShaderPlatform Platform, FShaderCompilerEnvironment& OutEnvironment)
+	{
+		FGlobalShader::ModifyCompilationEnvironment(Platform,OutEnvironment);
+		OutEnvironment.SetDefine(TEXT("NUM_CONE_DIRECTIONS"), NumConeSampleDirections);
+
+		// To reduce shader compile time of compute shaders with shared memory, doesn't have an impact on generated code with current compiler (June 2010 DX SDK)
+		OutEnvironment.CompilerFlags.Add(CFLAG_StandardOptimization);
+	}
+
+	FCalculateHeightfieldOcclusionCS(const ShaderMetaType::CompiledShaderInitializerType& Initializer)
+		: FGlobalShader(Initializer)
+	{
+		AOParameters.Bind(Initializer.ParameterMap);
+		HeightfieldDescriptionParameters.Bind(Initializer.ParameterMap);
+		HeightfieldTextureParameters.Bind(Initializer.ParameterMap);
+		OccluderRadius.Bind(Initializer.ParameterMap, TEXT("OccluderRadius"));
+		IrradianceCacheNormal.Bind(Initializer.ParameterMap, TEXT("IrradianceCacheNormal"));
+		IrradianceCachePositionRadius.Bind(Initializer.ParameterMap, TEXT("IrradianceCachePositionRadius"));
+		RecordConeVisibility.Bind(Initializer.ParameterMap, TEXT("RecordConeVisibility"));
+		ScatterDrawParameters.Bind(Initializer.ParameterMap, TEXT("ScatterDrawParameters"));
+		SavedStartIndex.Bind(Initializer.ParameterMap, TEXT("SavedStartIndex"));
+		BentNormalNormalizeFactor.Bind(Initializer.ParameterMap, TEXT("BentNormalNormalizeFactor"));
+		DistanceFieldNormalTexture.Bind(Initializer.ParameterMap, TEXT("DistanceFieldNormalTexture"));
+		DistanceFieldNormalSampler.Bind(Initializer.ParameterMap, TEXT("DistanceFieldNormalSampler"));
+		RecordRadiusScale.Bind(Initializer.ParameterMap, TEXT("RecordRadiusScale"));
+		TanConeHalfAngle.Bind(Initializer.ParameterMap, TEXT("TanConeHalfAngle"));
+	}
+
+	FCalculateHeightfieldOcclusionCS()
+	{
+	}
+
+	void SetParameters(
+		FRHICommandList& RHICmdList, 
+		const FSceneView& View, 
+		int32 DepthLevel, 
+		UTexture2D* HeightfieldTextureValue, 
+		int32 NumHeightfieldsValue,
+		const FTemporaryIrradianceCacheResources& TemporaryIrradianceCacheResources,
+		const FDistanceFieldAOParameters& Parameters)
+	{
+		FComputeShaderRHIParamRef ShaderRHI = GetComputeShader();
+		FGlobalShader::SetParameters(RHICmdList, ShaderRHI, View);
+
+		AOParameters.Set(RHICmdList, ShaderRHI, Parameters);
+		HeightfieldDescriptionParameters.Set(RHICmdList, ShaderRHI, NumHeightfieldsValue);
+		HeightfieldTextureParameters.Set(RHICmdList, ShaderRHI, HeightfieldTextureValue);
+
+		const FScene* Scene = (const FScene*)View.Family->Scene;
+		FSurfaceCacheResources& SurfaceCacheResources = *Scene->SurfaceCacheResources;
+
+		SetSRVParameter(RHICmdList, ShaderRHI, IrradianceCacheNormal, SurfaceCacheResources.Level[DepthLevel]->Normal.SRV);
+		SetSRVParameter(RHICmdList, ShaderRHI, IrradianceCachePositionRadius, SurfaceCacheResources.Level[DepthLevel]->PositionAndRadius.SRV);
+		SetSRVParameter(RHICmdList, ShaderRHI, ScatterDrawParameters, SurfaceCacheResources.Level[DepthLevel]->ScatterDrawParameters.SRV);
+		SetSRVParameter(RHICmdList, ShaderRHI, SavedStartIndex, SurfaceCacheResources.Level[DepthLevel]->SavedStartIndex.SRV);
+
+		OccluderRadius.SetBuffer(RHICmdList, ShaderRHI, SurfaceCacheResources.Level[DepthLevel]->OccluderRadius);
+		RecordConeVisibility.SetBuffer(RHICmdList, ShaderRHI, TemporaryIrradianceCacheResources.ConeVisibility);
+
+		FAOSampleData2 AOSampleData;
+
+		TArray<FVector, TInlineAllocator<9> > SampleDirections;
+		GetSpacedVectors(SampleDirections);
+
+		for (int32 SampleIndex = 0; SampleIndex < NumConeSampleDirections; SampleIndex++)
+		{
+			AOSampleData.SampleDirections[SampleIndex] = FVector4(SampleDirections[SampleIndex]);
+		}
+
+		SetUniformBufferParameterImmediate(RHICmdList, ShaderRHI, GetUniformBufferParameter<FAOSampleData2>(), AOSampleData);
+
+		FVector UnoccludedVector(0);
+
+		for (int32 SampleIndex = 0; SampleIndex < NumConeSampleDirections; SampleIndex++)
+		{
+			UnoccludedVector += SampleDirections[SampleIndex];
+		}
+
+		float BentNormalNormalizeFactorValue = 1.0f / (UnoccludedVector / NumConeSampleDirections).Size();
+		SetShaderValue(RHICmdList, ShaderRHI, BentNormalNormalizeFactor, BentNormalNormalizeFactorValue);
+
+		extern float GAOConeHalfAngle;
+		SetShaderValue(RHICmdList, ShaderRHI, TanConeHalfAngle, FMath::Tan(GAOConeHalfAngle));
+		extern float GAORecordRadiusScale;
+		SetShaderValue(RHICmdList, ShaderRHI, RecordRadiusScale, GAORecordRadiusScale);
+	}
+
+	void UnsetParameters(FRHICommandList& RHICmdList)
+	{
+		OccluderRadius.UnsetUAV(RHICmdList, GetComputeShader());
+		RecordConeVisibility.UnsetUAV(RHICmdList, GetComputeShader());
+	}
+
+	virtual bool Serialize(FArchive& Ar)
+	{		
+		bool bShaderHasOutdatedParameters = FGlobalShader::Serialize(Ar);
+		Ar << AOParameters;
+		Ar << HeightfieldDescriptionParameters;
+		Ar << HeightfieldTextureParameters;
+		Ar << OccluderRadius;
+		Ar << IrradianceCacheNormal;
+		Ar << IrradianceCachePositionRadius;
+		Ar << RecordConeVisibility;
+		Ar << ScatterDrawParameters;
+		Ar << SavedStartIndex;
+		Ar << BentNormalNormalizeFactor;
+		Ar << DistanceFieldNormalTexture;
+		Ar << DistanceFieldNormalSampler;
+		Ar << RecordRadiusScale;
+		Ar << TanConeHalfAngle;
+		return bShaderHasOutdatedParameters;
+	}
+
+private:
+
+	FAOParameters AOParameters;
+	FHeightfieldDescriptionParameters HeightfieldDescriptionParameters;
+	FHeightfieldTextureParameters HeightfieldTextureParameters;
+	FRWShaderParameter OccluderRadius;
+	FShaderResourceParameter IrradianceCacheNormal;
+	FShaderResourceParameter IrradianceCachePositionRadius;
+	FRWShaderParameter RecordConeVisibility;
+	FShaderResourceParameter ScatterDrawParameters;
+	FShaderResourceParameter SavedStartIndex;
+	FShaderParameter BentNormalNormalizeFactor;
+	FShaderResourceParameter DistanceFieldNormalTexture;
+	FShaderResourceParameter DistanceFieldNormalSampler;
+	FShaderParameter RecordRadiusScale;
+	FShaderParameter TanConeHalfAngle;
+};
+
+IMPLEMENT_SHADER_TYPE(,FCalculateHeightfieldOcclusionCS,TEXT("HeightfieldLighting"),TEXT("CalculateHeightfieldOcclusionCS"),SF_Compute);
+
+
+void FHeightfieldLightingViewInfo::ComputeOcclusionForSamples(
+	const FViewInfo& View, 
+	FRHICommandListImmediate& RHICmdList, 
+	const FTemporaryIrradianceCacheResources& TemporaryIrradianceCacheResources,
+	int32 DepthLevel,
+	const FDistanceFieldAOParameters& Parameters) const
+{
+	const FScene* Scene = (const FScene*)View.Family->Scene;
+
+	if (Heightfield.ComponentDescriptions.Num() > 0 && GAOHeightfieldOcclusion)
+	{
+		{	
+			TShaderMapRef<TSetupFinalGatherIndirectArgumentsCS<false> > ComputeShader(View.ShaderMap);
+
+			RHICmdList.SetComputeShader(ComputeShader->GetComputeShader());
+			ComputeShader->SetParameters(RHICmdList, View, DepthLevel);
+			DispatchComputeShader(RHICmdList, *ComputeShader, 1, 1, 1);
+
+			ComputeShader->UnsetParameters(RHICmdList);
+		}
+
+		SCOPED_DRAW_EVENT(RHICmdList, HeightfieldOcclusion);
+
+		FSceneViewState* ViewState = (FSceneViewState*)View.State;
+
+		{
+			for (TMap<UTexture2D*, TArray<FHeightfieldComponentDescription>>::TConstIterator It(Heightfield.ComponentDescriptions); It; ++It)
+			{
+				const TArray<FHeightfieldComponentDescription>& HeightfieldDescriptions = It.Value();
+
+				if (HeightfieldDescriptions.Num() > 0)
+				{
+					UploadHeightfieldDescriptions(HeightfieldDescriptions, FVector2D(1, 1), 1.0f / Heightfield.DownsampleFactor);
+
+					UTexture2D* HeightfieldTexture = It.Key();
+
+					{
+						TShaderMapRef<FCalculateHeightfieldOcclusionCS> ComputeShader(View.ShaderMap);
+						RHICmdList.SetComputeShader(ComputeShader->GetComputeShader());
+						ComputeShader->SetParameters(RHICmdList, View, DepthLevel, HeightfieldTexture, HeightfieldDescriptions.Num(), TemporaryIrradianceCacheResources, Parameters);
+
+						FSurfaceCacheResources& SurfaceCacheResources = *Scene->SurfaceCacheResources;
+						DispatchIndirectComputeShader(RHICmdList, *ComputeShader, SurfaceCacheResources.DispatchParameters.Buffer, 0);
+
+						ComputeShader->UnsetParameters(RHICmdList);
+					}
+				}
+			}
+		}
+	}
+}
+
+BEGIN_UNIFORM_BUFFER_STRUCT(FHemisphereSamples,)
+	DECLARE_UNIFORM_BUFFER_STRUCT_MEMBER_ARRAY(FVector4,SampleDirections,[1024])
+END_UNIFORM_BUFFER_STRUCT(FHemisphereSamples)
+
+IMPLEMENT_UNIFORM_BUFFER_STRUCT(FHemisphereSamples,TEXT("HemisphereSamples"));
+
+/** Generates unit length, stratified and uniformly distributed direction samples in a hemisphere. */
+void GenerateStratifiedUniformHemisphereSamples(int32 NumThetaSteps, int32 NumPhiSteps, FRandomStream& RandomStream, TArray<FVector4>& Samples)
+{
+	Samples.Empty(NumThetaSteps * NumPhiSteps);
+	for (int32 ThetaIndex = 0; ThetaIndex < NumThetaSteps; ThetaIndex++)
+	{
+		for (int32 PhiIndex = 0; PhiIndex < NumPhiSteps; PhiIndex++)
+		{
+			const float U1 = RandomStream.GetFraction();
+			const float U2 = RandomStream.GetFraction();
+
+			const float Fraction1 = (ThetaIndex + U1) / (float)NumThetaSteps;
+			const float Fraction2 = (PhiIndex + U2) / (float)NumPhiSteps;
+
+			const float R = FMath::Sqrt(1.0f - Fraction1 * Fraction1);
+
+			const float Phi = 2.0f * (float)PI * Fraction2;
+			// Convert to Cartesian
+			Samples.Add(FVector4(FMath::Cos(Phi) * R, FMath::Sin(Phi) * R, Fraction1));
+		}
+	}
+}
+
+class FCalculateHeightfieldIrradianceCS : public FGlobalShader
+{
+	DECLARE_SHADER_TYPE(FCalculateHeightfieldIrradianceCS,Global)
+public:
+
+	static bool ShouldCache(EShaderPlatform Platform)
+	{
+		return IsFeatureLevelSupported(Platform, ERHIFeatureLevel::SM5) && DoesPlatformSupportDistanceFieldAO(Platform);
+	}
+
+	static void ModifyCompilationEnvironment(EShaderPlatform Platform, FShaderCompilerEnvironment& OutEnvironment)
+	{
+		FGlobalShader::ModifyCompilationEnvironment(Platform,OutEnvironment);
+		OutEnvironment.SetDefine(TEXT("NUM_CONE_DIRECTIONS"), NumConeSampleDirections);
+
+		// To reduce shader compile time of compute shaders with shared memory, doesn't have an impact on generated code with current compiler (June 2010 DX SDK)
+		OutEnvironment.CompilerFlags.Add(CFLAG_StandardOptimization);
+	}
+
+	FCalculateHeightfieldIrradianceCS(const ShaderMetaType::CompiledShaderInitializerType& Initializer)
+		: FGlobalShader(Initializer)
+	{
+		AOParameters.Bind(Initializer.ParameterMap);
+		HeightfieldDescriptionParameters.Bind(Initializer.ParameterMap);
+		GlobalHeightfieldParameters.Bind(Initializer.ParameterMap);
+		IrradianceCacheIrradiance.Bind(Initializer.ParameterMap, TEXT("IrradianceCacheIrradiance"));
+		IrradianceCacheNormal.Bind(Initializer.ParameterMap, TEXT("IrradianceCacheNormal"));
+		IrradianceCachePositionRadius.Bind(Initializer.ParameterMap, TEXT("IrradianceCachePositionRadius"));
+		IrradianceCacheBentNormal.Bind(Initializer.ParameterMap, TEXT("IrradianceCacheBentNormal"));
+		ScatterDrawParameters.Bind(Initializer.ParameterMap, TEXT("ScatterDrawParameters"));
+		SavedStartIndex.Bind(Initializer.ParameterMap, TEXT("SavedStartIndex"));
+		BentNormalNormalizeFactor.Bind(Initializer.ParameterMap, TEXT("BentNormalNormalizeFactor"));
+		DistanceFieldNormalTexture.Bind(Initializer.ParameterMap, TEXT("DistanceFieldNormalTexture"));
+		DistanceFieldNormalSampler.Bind(Initializer.ParameterMap, TEXT("DistanceFieldNormalSampler"));
+		RecordRadiusScale.Bind(Initializer.ParameterMap, TEXT("RecordRadiusScale"));
+		TanConeHalfAngle.Bind(Initializer.ParameterMap, TEXT("TanConeHalfAngle"));
+		HeightfieldLighting.Bind(Initializer.ParameterMap, TEXT("HeightfieldLighting"));
+		HeightfieldLightingSampler.Bind(Initializer.ParameterMap, TEXT("HeightfieldLightingSampler"));
+		NumHemisphereSamples.Bind(Initializer.ParameterMap, TEXT("NumHemisphereSamples"));
+		InnerLightTransferDistance.Bind(Initializer.ParameterMap, TEXT("InnerLightTransferDistance"));
+		OuterLightTransferDistanceScale.Bind(Initializer.ParameterMap, TEXT("OuterLightTransferDistanceScale"));
+	}
+
+	FCalculateHeightfieldIrradianceCS()
+	{
+	}
+
+	void SetParameters(
+		FRHICommandList& RHICmdList, 
+		const FSceneView& View, 
+		int32 DepthLevel, 
+		int32 NumHeightfieldsValue,
+		const FTemporaryIrradianceCacheResources& TemporaryIrradianceCacheResources,
+		const FDistanceFieldAOParameters& Parameters,
+		const FHeightfieldLightingAtlas& Atlas)
+	{
+		FComputeShaderRHIParamRef ShaderRHI = GetComputeShader();
+		FGlobalShader::SetParameters(RHICmdList, ShaderRHI, View);
+
+		AOParameters.Set(RHICmdList, ShaderRHI, Parameters);
+		HeightfieldDescriptionParameters.Set(RHICmdList, ShaderRHI, NumHeightfieldsValue);
+		GlobalHeightfieldParameters.Set(RHICmdList, ShaderRHI, Atlas);
+
+		const FScene* Scene = (const FScene*)View.Family->Scene;
+		FSurfaceCacheResources& SurfaceCacheResources = *Scene->SurfaceCacheResources;
+
+		SetSRVParameter(RHICmdList, ShaderRHI, IrradianceCacheNormal, SurfaceCacheResources.Level[DepthLevel]->Normal.SRV);
+		SetSRVParameter(RHICmdList, ShaderRHI, IrradianceCachePositionRadius, SurfaceCacheResources.Level[DepthLevel]->PositionAndRadius.SRV);
+		SetSRVParameter(RHICmdList, ShaderRHI, IrradianceCacheBentNormal, SurfaceCacheResources.Level[DepthLevel]->BentNormal.SRV);
+		SetSRVParameter(RHICmdList, ShaderRHI, ScatterDrawParameters, SurfaceCacheResources.Level[DepthLevel]->ScatterDrawParameters.SRV);
+		SetSRVParameter(RHICmdList, ShaderRHI, SavedStartIndex, SurfaceCacheResources.Level[DepthLevel]->SavedStartIndex.SRV);
+
+		IrradianceCacheIrradiance.SetBuffer(RHICmdList, ShaderRHI, SurfaceCacheResources.Level[DepthLevel]->Irradiance);
+
+		{
+			FAOSampleData2 AOSampleData;
+
+			TArray<FVector, TInlineAllocator<9> > SampleDirections;
+			GetSpacedVectors(SampleDirections);
+
+			for (int32 SampleIndex = 0; SampleIndex < NumConeSampleDirections; SampleIndex++)
+			{
+				AOSampleData.SampleDirections[SampleIndex] = FVector4(SampleDirections[SampleIndex]);
+			}
+
+			SetUniformBufferParameterImmediate(RHICmdList, ShaderRHI, GetUniformBufferParameter<FAOSampleData2>(), AOSampleData);
+
+			FVector UnoccludedVector(0);
+
+			for (int32 SampleIndex = 0; SampleIndex < NumConeSampleDirections; SampleIndex++)
+			{
+				UnoccludedVector += SampleDirections[SampleIndex];
+			}
+
+			float BentNormalNormalizeFactorValue = 1.0f / (UnoccludedVector / NumConeSampleDirections).Size();
+			SetShaderValue(RHICmdList, ShaderRHI, BentNormalNormalizeFactor, BentNormalNormalizeFactorValue);
+		}
+
+		{
+			//@todo - cache this buffer
+			FHemisphereSamples HemisphereSamples;
+
+			const float NumThetaStepsFloat = FMath::Sqrt(GHeightfieldHemisphereSamples / (float)PI);
+			const int32 NumThetaSteps = FMath::TruncToInt(NumThetaStepsFloat);
+			const int32 NumPhiSteps = FMath::TruncToInt(NumThetaStepsFloat * (float)PI);
+
+			TArray<FVector4> SampleDirections;
+			FRandomStream RandomStream(12345);
+			GenerateStratifiedUniformHemisphereSamples(NumThetaSteps, NumPhiSteps, RandomStream, SampleDirections);
+
+			for (int32 SampleIndex = 0; SampleIndex < NumThetaSteps * NumPhiSteps; SampleIndex++)
+			{
+				HemisphereSamples.SampleDirections[SampleIndex] = FVector4(SampleDirections[SampleIndex]);
+			}
+
+			SetUniformBufferParameterImmediate(RHICmdList, ShaderRHI, GetUniformBufferParameter<FHemisphereSamples>(), HemisphereSamples);
+
+			SetShaderValue(RHICmdList, ShaderRHI, NumHemisphereSamples, NumThetaSteps * NumPhiSteps);
+		}
+
+		extern float GAOConeHalfAngle;
+		SetShaderValue(RHICmdList, ShaderRHI, TanConeHalfAngle, FMath::Tan(GAOConeHalfAngle));
+		extern float GAORecordRadiusScale;
+		SetShaderValue(RHICmdList, ShaderRHI, RecordRadiusScale, GAORecordRadiusScale);
+
+		const FSceneViewState* ViewState = (const FSceneViewState*)View.State;
+		SetTextureParameter(RHICmdList, ShaderRHI, HeightfieldLighting, HeightfieldLightingSampler, TStaticSamplerState<SF_Bilinear>::GetRHI(), Atlas.Lighting->GetRenderTargetItem().ShaderResourceTexture);
+	
+		SetShaderValue(RHICmdList, ShaderRHI, InnerLightTransferDistance, GHeightfieldInnerBounceDistance);
+		SetShaderValue(RHICmdList, ShaderRHI, OuterLightTransferDistanceScale, GHeightfieldOuterBounceDistanceScale);
+	}
+
+	void UnsetParameters(FRHICommandList& RHICmdList)
+	{
+		IrradianceCacheIrradiance.UnsetUAV(RHICmdList, GetComputeShader());
+	}
+
+	virtual bool Serialize(FArchive& Ar)
+	{		
+		bool bShaderHasOutdatedParameters = FGlobalShader::Serialize(Ar);
+		Ar << AOParameters;
+		Ar << HeightfieldDescriptionParameters;
+		Ar << GlobalHeightfieldParameters;
+		Ar << IrradianceCacheIrradiance;
+		Ar << IrradianceCacheNormal;
+		Ar << IrradianceCachePositionRadius;
+		Ar << IrradianceCacheBentNormal;
+		Ar << ScatterDrawParameters;
+		Ar << SavedStartIndex;
+		Ar << BentNormalNormalizeFactor;
+		Ar << DistanceFieldNormalTexture;
+		Ar << DistanceFieldNormalSampler;
+		Ar << RecordRadiusScale;
+		Ar << TanConeHalfAngle;
+		Ar << HeightfieldLighting;
+		Ar << HeightfieldLightingSampler;
+		Ar << NumHemisphereSamples;
+		Ar << InnerLightTransferDistance;
+		Ar << OuterLightTransferDistanceScale;
+		return bShaderHasOutdatedParameters;
+	}
+
+private:
+
+	FAOParameters AOParameters;
+	FHeightfieldDescriptionParameters HeightfieldDescriptionParameters;
+	FGlobalHeightfieldParameters GlobalHeightfieldParameters;
+	FRWShaderParameter IrradianceCacheIrradiance;
+	FShaderResourceParameter IrradianceCacheNormal;
+	FShaderResourceParameter IrradianceCachePositionRadius;
+	FShaderResourceParameter IrradianceCacheBentNormal;
+	FShaderResourceParameter ScatterDrawParameters;
+	FShaderResourceParameter SavedStartIndex;
+	FShaderParameter BentNormalNormalizeFactor;
+	FShaderResourceParameter DistanceFieldNormalTexture;
+	FShaderResourceParameter DistanceFieldNormalSampler;
+	FShaderParameter RecordRadiusScale;
+	FShaderParameter TanConeHalfAngle;
+	FShaderResourceParameter HeightfieldLighting;
+	FShaderResourceParameter HeightfieldLightingSampler;
+	FShaderParameter NumHemisphereSamples;
+	FShaderParameter InnerLightTransferDistance;
+	FShaderParameter OuterLightTransferDistanceScale;
+};
+
+IMPLEMENT_SHADER_TYPE(,FCalculateHeightfieldIrradianceCS,TEXT("HeightfieldLighting"),TEXT("CalculateHeightfieldIrradianceCS"),SF_Compute);
+
+void FHeightfieldLightingViewInfo::ComputeIrradianceForSamples(
+	const FViewInfo& View, 
+	FRHICommandListImmediate& RHICmdList, 
+	const FTemporaryIrradianceCacheResources& TemporaryIrradianceCacheResources,
+	int32 DepthLevel,
+	const FDistanceFieldAOParameters& Parameters) const
+{
+	const FScene* Scene = (const FScene*)View.Family->Scene;
+
+	if (Heightfield.ComponentDescriptions.Num() > 0 
+		&& AllowHeightfieldGI(View)
+		&& SupportsHeightfieldLighting(View.GetFeatureLevel(), View.GetShaderPlatform()))
+	{
+		{	
+			TShaderMapRef<TSetupFinalGatherIndirectArgumentsCS<true> > ComputeShader(View.ShaderMap);
+
+			RHICmdList.SetComputeShader(ComputeShader->GetComputeShader());
+			ComputeShader->SetParameters(RHICmdList, View, DepthLevel);
+			DispatchComputeShader(RHICmdList, *ComputeShader, 1, 1, 1);
+
+			ComputeShader->UnsetParameters(RHICmdList);
+		}
+
+		SCOPED_DRAW_EVENT(RHICmdList, HeightfieldIrradiance);
+
+		FSceneViewState* ViewState = (FSceneViewState*)View.State;
+
+		{
+			const FHeightfieldLightingAtlas& Atlas = *ViewState->HeightfieldLightingAtlas;
+
+			const FIntPoint LightingAtlasSize = Atlas.GetAtlasSize();
+			const FVector2D InvLightingAtlasSize(1.0f / LightingAtlasSize.X, 1.0f / LightingAtlasSize.Y);
+
+			TArray<FHeightfieldComponentDescription> CombinedHeightfieldDescriptions;
+
+			for (TMap<UTexture2D*, TArray<FHeightfieldComponentDescription>>::TConstIterator It(Heightfield.ComponentDescriptions); It; ++It)
+			{
+				const TArray<FHeightfieldComponentDescription>& HeightfieldDescriptions = It.Value();
+				CombinedHeightfieldDescriptions.Append(HeightfieldDescriptions);
+			}
+
+			if (CombinedHeightfieldDescriptions.Num() > 0)
+			{
+				UploadHeightfieldDescriptions(CombinedHeightfieldDescriptions, InvLightingAtlasSize, 1.0f / Heightfield.DownsampleFactor);
+
+				{
+					TShaderMapRef<FCalculateHeightfieldIrradianceCS> ComputeShader(View.ShaderMap);
+					RHICmdList.SetComputeShader(ComputeShader->GetComputeShader());
+					ComputeShader->SetParameters(RHICmdList, View, DepthLevel, CombinedHeightfieldDescriptions.Num(), TemporaryIrradianceCacheResources, Parameters, Atlas);
+
+					FSurfaceCacheResources& SurfaceCacheResources = *Scene->SurfaceCacheResources;
+					DispatchIndirectComputeShader(RHICmdList, *ComputeShader, SurfaceCacheResources.DispatchParameters.Buffer, 0);
+
+					ComputeShader->UnsetParameters(RHICmdList);
+				}
+			}
+		}
+	}
+}
diff --git a/Engine/Source/Runtime/Renderer/Private/HeightfieldLighting.h b/Engine/Source/Runtime/Renderer/Private/HeightfieldLighting.h
new file mode 100644
--- /dev/null
+++ b/Engine/Source/Runtime/Renderer/Private/HeightfieldLighting.h
@@ -0,0 +1,90 @@
+// Copyright 1998-2015 Epic Games, Inc. All Rights Reserved.
+
+/*=============================================================================
+	HeightfieldLighting.h
+=============================================================================*/
+
+#pragma once
+
+class FHeightfieldLightingAtlas : public FRenderResource
+{
+public:
+
+	TRefCountPtr<IPooledRenderTarget> Height;
+	TRefCountPtr<IPooledRenderTarget> Normal;
+	TRefCountPtr<IPooledRenderTarget> DirectionalLightShadowing;
+	TRefCountPtr<IPooledRenderTarget> Lighting;
+
+	FHeightfieldLightingAtlas() :
+		AtlasSize(FIntPoint(0, 0))
+	{}
+
+	virtual void InitDynamicRHI();
+	virtual void ReleaseDynamicRHI();
+
+	void InitializeForSize(FIntPoint InAtlasSize);
+
+	FIntPoint GetAtlasSize() const { return AtlasSize; }
+
+private:
+
+	FIntPoint AtlasSize;
+};
+
+class FHeightfieldDescription
+{
+public:
+	FIntRect Rect;
+	int32 DownsampleFactor;
+	FIntRect DownsampledRect;
+
+	TMap<UTexture2D*, TArray<FHeightfieldComponentDescription>> ComponentDescriptions;
+
+	FHeightfieldDescription() :
+		Rect(FIntRect(0, 0, 0, 0)),
+		DownsampleFactor(1),
+		DownsampledRect(FIntRect(0, 0, 0, 0))
+	{}
+};
+
+class FHeightfieldLightingViewInfo
+{
+public:
+
+	FHeightfieldLightingViewInfo()
+	{}
+
+	void SetupVisibleHeightfields(const FViewInfo& View, FRHICommandListImmediate& RHICmdList);
+
+	void ClearShadowing(const FViewInfo& View, FRHICommandListImmediate& RHICmdList, const FLightSceneInfo& LightSceneInfo) const;
+
+	void ComputeShadowMapShadowing(const FViewInfo& View, FRHICommandListImmediate& RHICmdList, const FProjectedShadowInfo* ProjectedShadowInfo) const;
+
+	void ComputeRayTracedShadowing(
+		const FViewInfo& View, 
+		FRHICommandListImmediate& RHICmdList, 
+		const FProjectedShadowInfo* ProjectedShadowInfo, 
+		FLightTileIntersectionResources* TileIntersectionResources,
+		class FDistanceFieldObjectBufferResource& CulledObjectBuffers) const;
+
+	void ComputeLighting(const FViewInfo& View, FRHICommandListImmediate& RHICmdList, const FLightSceneInfo& LightSceneInfo) const;
+
+	void ComputeOcclusionForSamples(
+		const FViewInfo& View, 
+		FRHICommandListImmediate& RHICmdList, 
+		const class FTemporaryIrradianceCacheResources& TemporaryIrradianceCacheResources,
+		int32 DepthLevel,
+		const class FDistanceFieldAOParameters& Parameters) const;
+
+	void ComputeIrradianceForSamples(
+		const FViewInfo& View, 
+		FRHICommandListImmediate& RHICmdList, 
+		const class FTemporaryIrradianceCacheResources& TemporaryIrradianceCacheResources,
+		int32 DepthLevel,
+		const class FDistanceFieldAOParameters& Parameters) const;
+
+private:
+
+	FHeightfieldDescription Heightfield;
+};
+
diff --git a/Engine/Source/Runtime/Renderer/Private/LightRendering.cpp b/Engine/Source/Runtime/Renderer/Private/LightRendering.cpp
--- a/Engine/Source/Runtime/Renderer/Private/LightRendering.cpp
+++ b/Engine/Source/Runtime/Renderer/Private/LightRendering.cpp
@@ -488,6 +488,12 @@ void FDeferredShadingSceneRenderer::RenderLights(FRHICommandListImmediate& RHICm
 			{
 				INC_DWORD_STAT(STAT_NumShadowedLights);
 
+				for (int32 ViewIndex = 0; ViewIndex < Views.Num(); ViewIndex++)
+				{
+					const FViewInfo& View = Views[ViewIndex];
+					View.HeightfieldLightingViewInfo.ClearShadowing(View, RHICmdList, LightSceneInfo);
+				}
+
 				// All shadows render with min blending
 				GSceneRenderTargets.BeginRenderingLightAttenuation(RHICmdList, true);
 
@@ -507,6 +513,12 @@ void FDeferredShadingSceneRenderer::RenderLights(FRHICommandListImmediate& RHICm
 					RenderReflectiveShadowMaps(RHICmdList, &LightSceneInfo );
 				}
 			}
+
+			for (int32 ViewIndex = 0; ViewIndex < Views.Num(); ViewIndex++)
+			{
+				const FViewInfo& View = Views[ViewIndex];
+				View.HeightfieldLightingViewInfo.ComputeLighting(View, RHICmdList, LightSceneInfo);
+			}
 			
 			// Render light function to the attenuation buffer.
 			if (bDirectLighting)
diff --git a/Engine/Source/Runtime/Renderer/Private/RendererScene.cpp b/Engine/Source/Runtime/Renderer/Private/RendererScene.cpp
--- a/Engine/Source/Runtime/Renderer/Private/RendererScene.cpp
+++ b/Engine/Source/Runtime/Renderer/Private/RendererScene.cpp
@@ -15,6 +15,7 @@
 #include "FXSystem.h"
 #include "DistanceFieldLightingShared.h"
 #include "SpeedTreeWind.h"
+#include "HeightfieldLighting.h"
 
 // Enable this define to do slow checks for components being added to the wrong
 // world's scene, when using PIE. This can happen if a PIE component is reattached
@@ -84,6 +85,7 @@ FSceneViewState::FSceneViewState()
 	bBokehDOFHistory2 = true;
 
 	LightPropagationVolume = NULL; 
+	HeightfieldLightingAtlas = NULL;
 
 	for (int32 CascadeIndex = 0; CascadeIndex < ARRAY_COUNT(TranslucencyLightingCacheAllocations); CascadeIndex++)
 	{
@@ -97,6 +99,36 @@ FSceneViewState::FSceneViewState()
 #endif
 }
 
+void DestroyRenderResource(FRenderResource* RenderResource)
+{
+	if (RenderResource) 
+	{
+		ENQUEUE_UNIQUE_RENDER_COMMAND_ONEPARAMETER(
+			DestroySceneViewStateRenderResource,
+			FRenderResource*, RenderResourceRT, RenderResource,
+			{
+				RenderResourceRT->ReleaseResource();
+				delete RenderResourceRT;
+			}
+		);
+	}
+}
+
+FSceneViewState::~FSceneViewState()
+{
+	CachedVisibilityChunk = NULL;
+
+	for (int32 CascadeIndex = 0; CascadeIndex < ARRAY_COUNT(TranslucencyLightingCacheAllocations); CascadeIndex++)
+	{
+		delete TranslucencyLightingCacheAllocations[CascadeIndex];
+	}
+
+	DestroyRenderResource(HeightfieldLightingAtlas);
+	DestroyRenderResource(AOTileIntersectionResources);
+	AOTileIntersectionResources = NULL;
+	DestroyLightPropagationVolume();
+}
+
 FDistanceFieldSceneData::FDistanceFieldSceneData() 
 	: NumObjectsInBuffer(0)
 	, ObjectBuffers(NULL)
diff --git a/Engine/Source/Runtime/Renderer/Private/ScenePrivate.h b/Engine/Source/Runtime/Renderer/Private/ScenePrivate.h
--- a/Engine/Source/Runtime/Renderer/Private/ScenePrivate.h
+++ b/Engine/Source/Runtime/Renderer/Private/ScenePrivate.h
@@ -476,6 +476,9 @@ private:
 	FLightPropagationVolume* LightPropagationVolume;
 
 public:
+
+	FHeightfieldLightingAtlas* HeightfieldLightingAtlas;
+
 	// If Translucency should be rendered into a separate RT and composited without DepthOfField, can be disabled in the materials (affects sorting)
 	TRefCountPtr<IPooledRenderTarget> SeparateTranslucencyRT;
 	// Temporal AA result of last frame
@@ -556,21 +559,9 @@ public:
 	/** Default constructor. */
 	FSceneViewState();
 
-	void DestroyAOTileResources();
 	void DestroyLightPropagationVolume();
 
-	virtual ~FSceneViewState()
-	{
-		CachedVisibilityChunk = NULL;
-
-		for (int32 CascadeIndex = 0; CascadeIndex < ARRAY_COUNT(TranslucencyLightingCacheAllocations); CascadeIndex++)
-		{
-			delete TranslucencyLightingCacheAllocations[CascadeIndex];
-		}
-
-		DestroyAOTileResources();
-		DestroyLightPropagationVolume();
-	}
+	virtual ~FSceneViewState();
 
 	// called every frame after the view state was updated
 	void UpdateLastRenderTime(const FSceneViewFamily& Family)
diff --git a/Engine/Source/Runtime/Renderer/Private/SceneRendering.h b/Engine/Source/Runtime/Renderer/Private/SceneRendering.h
--- a/Engine/Source/Runtime/Renderer/Private/SceneRendering.h
+++ b/Engine/Source/Runtime/Renderer/Private/SceneRendering.h
@@ -9,6 +9,7 @@
 #include "TextureLayout.h"
 #include "DistortionRendering.h"
 #include "CustomDepthRendering.h"
+#include "HeightfieldLighting.h"
 
 // Forward declarations.
 class FPostprocessContext;
@@ -473,6 +474,8 @@ public:
 	FLinearColor LightShaftColorApply;
 	bool bLightShaftUse;
 
+	FHeightfieldLightingViewInfo HeightfieldLightingViewInfo;
+
 	TShaderMap<FGlobalShaderType>* ShaderMap;
 
 	/** Custom visibility query for view */
diff --git a/Engine/Source/Runtime/Renderer/Private/ShadowRendering.cpp b/Engine/Source/Runtime/Renderer/Private/ShadowRendering.cpp
--- a/Engine/Source/Runtime/Renderer/Private/ShadowRendering.cpp
+++ b/Engine/Source/Runtime/Renderer/Private/ShadowRendering.cpp
@@ -3304,6 +3304,23 @@ bool FDeferredShadingSceneRenderer::RenderProjectedShadows(FRHICommandListImmedi
 				}
 			}
 
+			
+			for (int32 ViewIndex = 0; ViewIndex < Views.Num(); ViewIndex++)
+			{
+				const FViewInfo& View = Views[ViewIndex];
+
+				for (int32 ShadowIndex = 0; ShadowIndex < Shadows.Num(); ShadowIndex++)
+				{
+					FProjectedShadowInfo* ProjectedShadowInfo = Shadows[ShadowIndex];
+
+					if (ProjectedShadowInfo->bAllocated
+						&& ProjectedShadowInfo->bWholeSceneShadow)
+					{
+						View.HeightfieldLightingViewInfo.ComputeShadowMapShadowing(View, RHICmdList, ProjectedShadowInfo);
+					}
+				}
+			}
+
 			// Mark and count the rendered shadows.
 			for (int32 ShadowIndex = 0; ShadowIndex < Shadows.Num(); ShadowIndex++)
 			{
