Fix a few issues with self-shadow only.

1) Preshadows for self-shadow only primitives now correctly mask by the receiver and not the subject.
2) Self-shadow only objects are no longer included in whole scene shadows (view dependent whole scene shadows already worked).

diff --git a/Engine/Source/Runtime/Renderer/Private/SceneCore.cpp b/Engine/Source/Runtime/Renderer/Private/SceneCore.cpp
--- a/Engine/Source/Runtime/Renderer/Private/SceneCore.cpp
+++ b/Engine/Source/Runtime/Renderer/Private/SceneCore.cpp
@@ -162,7 +162,8 @@ FLightPrimitiveInteraction::FLightPrimitiveInteraction(
 	bIsShadowMapped(bInIsShadowMapped),
 	bUncachedStaticLighting(false),
 	bHasTranslucentObjectShadow(bInHasTranslucentObjectShadow),
-	bHasInsetObjectShadow(bInHasInsetObjectShadow)
+	bHasInsetObjectShadow(bInHasInsetObjectShadow),
+	bSelfShadowOnly(false)
 {
 	// Determine whether this light-primitive interaction produces a shadow.
 	if(PrimitiveSceneInfo->Proxy->HasStaticLighting())
@@ -203,6 +204,8 @@ FLightPrimitiveInteraction::FLightPrimitiveInteraction(
 	}
 #endif
 
+	bSelfShadowOnly = PrimitiveSceneInfo->Proxy->CastsSelfShadowOnly();
+
 	if (bIsDynamic)
 	{
 		// Add the interaction to the light's interaction list.
diff --git a/Engine/Source/Runtime/Renderer/Private/SceneCore.h b/Engine/Source/Runtime/Renderer/Private/SceneCore.h
--- a/Engine/Source/Runtime/Renderer/Private/SceneCore.h
+++ b/Engine/Source/Runtime/Renderer/Private/SceneCore.h
@@ -35,6 +35,7 @@ public:
 	bool IsUncachedStaticLighting() const { return bUncachedStaticLighting; }
 	bool HasTranslucentObjectShadow() const { return bHasTranslucentObjectShadow; }
 	bool HasInsetObjectShadow() const { return bHasInsetObjectShadow; }
+	bool CastsSelfShadowOnly() const { return bSelfShadowOnly; }
 	FLightSceneInfo* GetLight() const { return LightSceneInfo; }
 	int32 GetLightId() const { return LightId; }
 	FPrimitiveSceneInfo* GetPrimitiveSceneInfo() const { return PrimitiveSceneInfo; }
@@ -82,6 +83,9 @@ private:
 	/** True if the interaction has an inset per-object shadow. */
 	uint32 bHasInsetObjectShadow : 1;
 
+	/** True if the primitive only shadows itself. */
+	uint32 bSelfShadowOnly : 1;
+
 	/** A pointer to the NextPrimitive member of the previous interaction in the light's interaction list. */
 	FLightPrimitiveInteraction** PrevPrimitiveLink;
 
diff --git a/Engine/Source/Runtime/Renderer/Private/ShadowRendering.cpp b/Engine/Source/Runtime/Renderer/Private/ShadowRendering.cpp
--- a/Engine/Source/Runtime/Renderer/Private/ShadowRendering.cpp
+++ b/Engine/Source/Runtime/Renderer/Private/ShadowRendering.cpp
@@ -1937,8 +1937,9 @@ void FProjectedShadowInfo::RenderProjection(FRHICommandListImmediate& RHICmdList
 			0xff,0xff
 			>::GetRHI(), 1);
 
-
-		const TArray<FMeshBatchAndRelevance,SceneRenderingAllocator>& DynamicMeshElements = bSelfShadowOnly ? DynamicSubjectMeshElements : DynamicReceiverMeshElements;
+		// Pre-shadows mask by receiver elements, self-shadow mask by subject elements.
+		// Note that self-shadow pre-shadows still mask by receiver elements.
+		const TArray<FMeshBatchAndRelevance,SceneRenderingAllocator>& DynamicMeshElements = bPreShadow ? DynamicReceiverMeshElements : DynamicSubjectMeshElements;
 
 		FDepthDrawingPolicyFactory::ContextType Context(DDM_AllOccluders);
 
@@ -1949,9 +1950,9 @@ void FProjectedShadowInfo::RenderProjection(FRHICommandListImmediate& RHICmdList
 			FDepthDrawingPolicyFactory::DrawDynamicMesh(RHICmdList, *View, Context, MeshBatch, false, true, MeshBatchAndRelevance.PrimitiveSceneProxy, MeshBatch.BatchHitProxyId);
 		}
 
-		// bPreShadow: only receive shadow on specific objects
-		// bSelfShadowOnly: only cast shadow on itself
-		const PrimitiveArrayType& MaskPrimitives = bSelfShadowOnly ? SubjectPrimitives : ReceiverPrimitives;
+		// Pre-shadows mask by receiver elements, self-shadow mask by subject elements.
+		// Note that self-shadow pre-shadows still mask by receiver elements.
+		const PrimitiveArrayType& MaskPrimitives = bPreShadow ? ReceiverPrimitives : SubjectPrimitives;
 
 		for (int32 PrimitiveIndex = 0, PrimitiveCount = MaskPrimitives.Num(); PrimitiveIndex < PrimitiveCount; PrimitiveIndex++)
 		{
diff --git a/Engine/Source/Runtime/Renderer/Private/ShadowSetup.cpp b/Engine/Source/Runtime/Renderer/Private/ShadowSetup.cpp
--- a/Engine/Source/Runtime/Renderer/Private/ShadowSetup.cpp
+++ b/Engine/Source/Runtime/Renderer/Private/ShadowSetup.cpp
@@ -1178,6 +1178,7 @@ void FDeferredShadingSceneRenderer::SetupInteractionShadows(
 		const bool bCreateObjectShadowForStationaryLight = ShouldCreateObjectShadowForStationaryLight(Interaction->GetLight(), PrimitiveSceneInfo->Proxy, Interaction->IsShadowMapped());
 
 		if (Interaction->HasShadow() 
+			// TODO: Handle inset shadows, especially when an object is only casting a self-shadow.
 			// Only render shadows from objects that use static lighting during a reflection capture, since the reflection capture doesn't update at runtime
 			&& (!bStaticSceneOnly || PrimitiveSceneInfo->Proxy->HasStaticLighting())
 			&& (bCreateTranslucentObjectShadow || bCreateInsetObjectShadow || bCreateObjectShadowForStationaryLight))
@@ -1662,6 +1663,8 @@ void FDeferredShadingSceneRenderer::CreateWholeSceneProjectedShadow(FLightSceneI
 						Interaction = Interaction->GetNextPrimitive())
 					{
 						if (Interaction->HasShadow() 
+							// If the primitive only wants to cast a self shadow don't include it in whole scene shadows.
+							&& !Interaction->CastsSelfShadowOnly()
 							&& (!bStaticSceneOnly || Interaction->GetPrimitiveSceneInfo()->Proxy->HasStaticLighting()))
 						{
 							ProjectedShadowInfo->AddSubjectPrimitive(Interaction->GetPrimitiveSceneInfo(), &Views);
