Fix non-unity error

diff --git a/Engine/Source/Runtime/Renderer/Private/DistanceFieldObjectManagement.cpp b/Engine/Source/Runtime/Renderer/Private/DistanceFieldObjectManagement.cpp
--- a/Engine/Source/Runtime/Renderer/Private/DistanceFieldObjectManagement.cpp
+++ b/Engine/Source/Runtime/Renderer/Private/DistanceFieldObjectManagement.cpp
@@ -741,6 +741,7 @@ void FDeferredShadingSceneRenderer::UpdateGlobalDistanceFieldObjectBuffers(FRHIC
 			int32 OriginalNumObjects = DistanceFieldSceneData.NumObjectsInBuffer;
 			int32 OriginalNumSurfels = DistanceFieldSceneData.SurfelAllocations.GetNumSurfelsInBuffer();
 			int32 OriginalNumInstancedSurfels = DistanceFieldSceneData.InstancedSurfelAllocations.GetNumSurfelsInBuffer();
+			extern int32 GVPLMeshGlobalIllumination;
 
 			if (GDistanceFieldGI && GVPLMeshGlobalIllumination)
 			{
@@ -861,6 +862,7 @@ void FDeferredShadingSceneRenderer::UpdateGlobalDistanceFieldObjectBuffers(FRHIC
 					{
 						FPrimitiveSurfelAllocation Allocation;
 						FPrimitiveSurfelAllocation InstancedAllocation;
+						extern int32 GVPLMeshGlobalIllumination;
 
 						if (GDistanceFieldGI && GVPLMeshGlobalIllumination)
 						{
