[INTEGRATE] Change 2422246 by [REDACTED_PERSON][REDACTED_HANDLE] on 2015/01/28 10:50:36

	Fixes for hitproxy rendering crashing the editor in OpenGL

	I'll look at this again after GDC, and may change it to a different solution.

diff --git a/Engine/Source/Runtime/Renderer/Private/PostProcess/SceneRenderTargets.cpp b/Engine/Source/Runtime/Renderer/Private/PostProcess/SceneRenderTargets.cpp
--- a/Engine/Source/Runtime/Renderer/Private/PostProcess/SceneRenderTargets.cpp
+++ b/Engine/Source/Runtime/Renderer/Private/PostProcess/SceneRenderTargets.cpp
@@ -749,15 +749,19 @@ TRefCountPtr<IPooledRenderTarget>& FSceneRenderTargets::GetLightAttenuation()
 
 void FSceneRenderTargets::AdjustGBufferRefCount(int Delta)
 { 
-	GBufferRefCount += Delta; 
 	
-	if(Delta > 0 && GBufferRefCount == 1)
+	if(Delta > 0 && GBufferRefCount == 0)
 	{
 		AllocGBufferTargets();
 	}
-	else if(GBufferRefCount == 0)
+	else
 	{
-		ReleaseGBufferTargets();
+		GBufferRefCount += Delta;
+
+		if (GBufferRefCount == 0)
+		{
+			ReleaseGBufferTargets();
+		}
 	}
 }
 
diff --git a/Engine/Source/Runtime/Renderer/Private/SceneHitProxyRendering.cpp b/Engine/Source/Runtime/Renderer/Private/SceneHitProxyRendering.cpp
--- a/Engine/Source/Runtime/Renderer/Private/SceneHitProxyRendering.cpp
+++ b/Engine/Source/Runtime/Renderer/Private/SceneHitProxyRendering.cpp
@@ -381,6 +381,8 @@ void RenderHitProxies(FRHICommandListImmediate& RHICmdList, const FSceneRenderer
 
 	const bool bNeedToSwitchVerticalAxis = RHINeedsToSwitchVerticalAxis(GShaderPlatformForFeatureLevel[SceneRenderer->FeatureLevel]);
 
+	GSceneRenderTargets.AdjustGBufferRefCount(1);
+
 	for (int32 ViewIndex = 0; ViewIndex < Views.Num(); ViewIndex++)
 	{
 		const FViewInfo& View = Views[ViewIndex];
@@ -460,6 +462,8 @@ void RenderHitProxies(FRHICommandListImmediate& RHICmdList, const FSceneRenderer
 		View.TopBatchedViewElements.Draw(RHICmdList, FeatureLevel, bNeedToSwitchVerticalAxis, View.ViewProjectionMatrix, View.ViewRect.Width(), View.ViewRect.Height(), true);
 	}
 
+	GSceneRenderTargets.AdjustGBufferRefCount(-1);
+
 	// Finish drawing to the hit proxy render target.
 	RHICmdList.CopyToResolveTarget(HitProxyRT->GetRenderTargetItem().TargetableTexture, HitProxyRT->GetRenderTargetItem().ShaderResourceTexture, false, FResolveParams());
 	// to be able to observe results with VisualizeTexture
