[INTEGRATE] Change 2432514 by [REDACTED_PERSON][REDACTED_HANDLE] on 2015/02/04 17:22:50

	Editor selection outline now uniquely identifies selected components.  When a component is selected directly the rest of the actor gets a subdued selection color and the selected compoent(s) get the normal selection color.  When the actor is selected directly all components get the normal selection color.

diff --git a/Engine/Shaders/PostProcessSelectionOutline.usf b/Engine/Shaders/PostProcessSelectionOutline.usf
--- a/Engine/Shaders/PostProcessSelectionOutline.usf
+++ b/Engine/Shaders/PostProcessSelectionOutline.usf
@@ -16,8 +16,11 @@ Texture2D PostprocessInput1MS;
 Texture2D<uint2> EditorPrimitivesStencil;
 #endif
 
-// RGB:selection color, A:SelectionHighlightIntensity
+// RGB: selection color, A:SelectionHighlightIntensity
 float4 OutlineColor;
+// RGB:subdued selection color (for inactive selection), A:unused
+float4 SubduedOutlineColor;
+
 float BSPSelectionIntensity;
 
 // r. value of r.Editor.OpaqueGizmo .g:r.Editor.MovingPattern ba: unused
@@ -159,17 +162,26 @@ float4 PerSample(int2 PixelPos, int SampleID, float3 DeviceZPlane, float2 Device
 	// the contants express the two opacity values we see when the primitive is hidden
 	float LowContrastPatternMask = lerp(0.2, 1.0f, PatternMask);
 	
+	float3  SelectionColor;
+	FLATTEN if(VisibleBorderMinMax.x >= 128 )
+	{
+		SelectionColor = SubduedOutlineColor.rgb;
+	}
+	else
+	{
+		SelectionColor = OutlineColor.rgb;
+	}
 	FLATTEN if(bBorder)
 	{
 		FLATTEN if(bVisibleBorder)
 		{
 			// unoccluded border
-			Ret = float4(OutlineColor.rgb, 1);
+			Ret = float4(SelectionColor, 1);
 		}
 		else
 		{
 			// occluded border
-			Ret = lerp(float4(0, 0, 0, 0), float4(OutlineColor.rgb, 1), LowContrastPatternMask);
+			Ret = lerp(float4(0, 0, 0, 0), float4(SelectionColor, 1), LowContrastPatternMask);
 		}
 	}
 
@@ -180,9 +192,9 @@ float4 PerSample(int2 PixelPos, int SampleID, float3 DeviceZPlane, float2 Device
 		float VisibleMask = lerp(PatternMask, 1.0f, PixelInfo[0].fVisible);
 
 		// occluded parts are more tranparent
-//		Ret = lerp(Ret, float4(OutlineColor.rgb, 1), OutlineColor.a * VisibleMask);
+//		Ret = lerp(Ret, float4(SelectionColor.rgb, 1), OutlineColor.a * VisibleMask);
 		// darken occluded parts
-		Ret = lerp(Ret, float4(OutlineColor.rgb * VisibleMask, 1), OutlineColor.a );
+		Ret = lerp(Ret, float4(SelectionColor * VisibleMask, 1), OutlineColor.a );
 	}
 
 	// highlight inner part of the object
@@ -190,7 +202,7 @@ float4 PerSample(int2 PixelPos, int SampleID, float3 DeviceZPlane, float2 Device
 	{
 		float VisibleMask = lerp(PatternMask, 1.0f, PixelInfo[0].fVisible);
 		float InnerHighlightAmount = PixelInfo[0].ObjectId == 1 ? BSPSelectionIntensity : OutlineColor.a;
-		Ret = lerp(Ret, float4(OutlineColor.rgb, 1), VisibleMask * InnerHighlightAmount);
+		Ret = lerp(Ret, float4(SelectionColor, 1), VisibleMask * InnerHighlightAmount);
 	}
 
 	return Ret;
diff --git a/Engine/Source/Editor/EditorStyle/Private/EditorStyleClasses.cpp b/Engine/Source/Editor/EditorStyle/Private/EditorStyleClasses.cpp
--- a/Engine/Source/Editor/EditorStyle/Private/EditorStyleClasses.cpp
+++ b/Engine/Source/Editor/EditorStyle/Private/EditorStyleClasses.cpp
@@ -18,6 +18,14 @@ UEditorStyleSettings::UEditorStyleSettings( const FObjectInitializer& ObjectInit
 	LogTimestampMode = ELogTimes::None;
 }
 
+FLinearColor UEditorStyleSettings::GetSubduedSelectionColor() const
+{
+	FLinearColor SubduedSelectionColor = SelectionColor.LinearRGBToHSV();
+	SubduedSelectionColor.G *= 0.55f;		// take the saturation 
+	SubduedSelectionColor.B *= 0.8f;		// and brightness down
+
+	return SubduedSelectionColor.HSVToLinearRGB();
+}
 
 #if WITH_EDITOR
 
diff --git a/Engine/Source/Editor/EditorStyle/Private/SlateEditorStyle.cpp b/Engine/Source/Editor/EditorStyle/Private/SlateEditorStyle.cpp
--- a/Engine/Source/Editor/EditorStyle/Private/SlateEditorStyle.cpp
+++ b/Engine/Source/Editor/EditorStyle/Private/SlateEditorStyle.cpp
@@ -101,10 +101,8 @@ void FSlateEditorStyle::FStyle::SyncSettings()
 		SetColor( SelectionColor_Pressed_LinearRef, Settings->PressedSelectionColor );
 
 		// The subdued selection color is derived from the selection color
-		auto SubduedSelectionColor = Settings->SelectionColor.LinearRGBToHSV();
-		SubduedSelectionColor.G *= 0.55f;		// take the saturation 
-		SubduedSelectionColor.B *= 0.8f;		// and brightness down
-		SetColor( SelectionColor_Subdued_LinearRef, SubduedSelectionColor.HSVToLinearRGB() );
+		auto SubduedSelectionColor = Settings->GetSubduedSelectionColor();
+		SetColor( SelectionColor_Subdued_LinearRef, SubduedSelectionColor );
 
 		// Also sync the colors used by FCoreStyle, as FEditorStyle isn't yet being used as an override everywhere
 		FCoreStyle::SetSelectorColor( Settings->KeyboardFocusColor );
diff --git a/Engine/Source/Editor/EditorStyle/Public/Classes/EditorStyleSettings.h b/Engine/Source/Editor/EditorStyle/Public/Classes/EditorStyleSettings.h
--- a/Engine/Source/Editor/EditorStyle/Public/Classes/EditorStyleSettings.h
+++ b/Engine/Source/Editor/EditorStyle/Public/Classes/EditorStyleSettings.h
@@ -96,6 +96,8 @@ public:
 	DECLARE_EVENT_OneParam(UEditorStyleSettings, FSettingChangedEvent, FName /*PropertyName*/);
 	FSettingChangedEvent& OnSettingChanged( ) { return SettingChangedEvent; }
 
+	/** @return A subdued version of the users selection color (for use with inactive selection)*/
+	FLinearColor GetSubduedSelectionColor() const;
 protected:
 
 	// UObject overrides
diff --git a/Engine/Source/Editor/LevelEditor/Private/SLevelViewport.cpp b/Engine/Source/Editor/LevelEditor/Private/SLevelViewport.cpp
--- a/Engine/Source/Editor/LevelEditor/Private/SLevelViewport.cpp
+++ b/Engine/Source/Editor/LevelEditor/Private/SLevelViewport.cpp
@@ -831,7 +831,7 @@ void SLevelViewport::Tick( const FGeometry& AllottedGeometry, const double InCur
 
 	// When we have focus we update the 'Allow Throttling' option in slate to be disabled so that interactions in the
 	// viewport with Slate widgets that are part of the game, don't throttle.
-	if ( bPIEContainsFocus != bContainsFocus )
+	if ( GEditor->PlayWorld != nullptr && bPIEContainsFocus != bContainsFocus )
 	{
 		// We can arrive at this point before creating throttling manager (which registers the cvar), so create it explicitly.
 		static const FSlateThrottleManager & ThrottleManager = FSlateThrottleManager::Get();
diff --git a/Engine/Source/Editor/UnrealEd/Private/EditorSelectUtils.cpp b/Engine/Source/Editor/UnrealEd/Private/EditorSelectUtils.cpp
--- a/Engine/Source/Editor/UnrealEd/Private/EditorSelectUtils.cpp
+++ b/Engine/Source/Editor/UnrealEd/Private/EditorSelectUtils.cpp
@@ -652,12 +652,13 @@ void UUnrealEdEngine::SelectComponent(UActorComponent* Component, bool bInSelect
 
 bool UUnrealEdEngine::IsComponentSelected(const UPrimitiveComponent* PrimComponent)
 {
+	bool bIsSelected = false;
 	if (GetSelectedComponentCount() > 0)
 	{
-		return GetSelectedComponents()->IsSelected(PrimComponent->IsEditorOnly() ? PrimComponent->AttachParent : PrimComponent);
+		bIsSelected = GetSelectedComponents()->IsSelected(PrimComponent->IsEditorOnly() ? PrimComponent->AttachParent : PrimComponent);
 	}
 
-	return GetSelectedActors()->IsSelected(PrimComponent->GetOwner());
+	return bIsSelected;
 }
 
 void UUnrealEdEngine::SelectBSPSurf(UModel* InModel, int32 iSurf, bool bSelected, bool bNoteSelectionChange)
diff --git a/Engine/Source/Editor/UnrealEd/Private/EditorViewportClient.cpp b/Engine/Source/Editor/UnrealEd/Private/EditorViewportClient.cpp
--- a/Engine/Source/Editor/UnrealEd/Private/EditorViewportClient.cpp
+++ b/Engine/Source/Editor/UnrealEd/Private/EditorViewportClient.cpp
@@ -708,6 +708,8 @@ FSceneView* FEditorViewportClient::CalcSceneView(FSceneViewFamily* ViewFamily)
 
 	FSceneView* View = new FSceneView(ViewInitOptions);
 
+	View->SubduedSelectionOutlineColor = GetDefault<UEditorStyleSettings>()->GetSubduedSelectionColor();
+
 	ViewFamily->Views.Add(View);
 
 	View->StartFinalPostprocessSettings(ViewLocation);
diff --git a/Engine/Source/Editor/UnrealEd/Private/LevelEditorViewport.cpp b/Engine/Source/Editor/UnrealEd/Private/LevelEditorViewport.cpp
--- a/Engine/Source/Editor/UnrealEd/Private/LevelEditorViewport.cpp
+++ b/Engine/Source/Editor/UnrealEd/Private/LevelEditorViewport.cpp
@@ -1614,7 +1614,7 @@ FSceneView* FLevelEditorViewportClient::CalcSceneView(FSceneViewFamily* ViewFami
 
 	View->SpriteCategoryVisibility = SpriteCategoryVisibility;
 	View->bCameraCut = bEditorCameraCut;
-
+	View->bHasSelectedComponents = GEditor->GetSelectedComponentCount() > 0;
 	return View;
 
 }
diff --git a/Engine/Source/Runtime/Engine/Classes/Components/PrimitiveComponent.h b/Engine/Source/Runtime/Engine/Classes/Components/PrimitiveComponent.h
--- a/Engine/Source/Runtime/Engine/Classes/Components/PrimitiveComponent.h
+++ b/Engine/Source/Runtime/Engine/Classes/Components/PrimitiveComponent.h
@@ -959,6 +959,9 @@ public:
 	/** @return true if the owner is selected and this component is selectable */
 	virtual bool ShouldRenderSelected() const;
 
+	/** Component is directly selected in the editor separate from its parent actor */
+	bool IsComponentIndividuallySelected() const;
+
 	/**  @return True if a primitive's parameters as well as its position is static during gameplay, and can thus use static lighting. */
 	bool HasStaticLighting() const;
 
diff --git a/Engine/Source/Runtime/Engine/Private/ModelRender.cpp b/Engine/Source/Runtime/Engine/Private/ModelRender.cpp
--- a/Engine/Source/Runtime/Engine/Private/ModelRender.cpp
+++ b/Engine/Source/Runtime/Engine/Private/ModelRender.cpp
@@ -400,7 +400,7 @@ public:
 
 		//@todo parallelrendering - remove this legacy state modification
 		// Poly selected state is modified in many places, so it's hard to push the selection state to the proxy
-		const_cast<FModelSceneProxy*>(this)->SetSelection_RenderThread(bAnySelectedSurfs);
+		const_cast<FModelSceneProxy*>(this)->SetSelection_RenderThread(bAnySelectedSurfs,false);
 	}
 
 	virtual void DrawStaticElements(FStaticPrimitiveDrawInterface* PDI)
diff --git a/Engine/Source/Runtime/Engine/Private/PrimitiveComponent.cpp b/Engine/Source/Runtime/Engine/Private/PrimitiveComponent.cpp
--- a/Engine/Source/Runtime/Engine/Private/PrimitiveComponent.cpp
+++ b/Engine/Source/Runtime/Engine/Private/PrimitiveComponent.cpp
@@ -833,25 +833,28 @@ bool UPrimitiveComponent::HasValidPhysicsState() const
 	return BodyInstance.IsValidBodyInstance();
 }
 
-bool UPrimitiveComponent::ShouldRenderSelected() const
+bool UPrimitiveComponent::IsComponentIndividuallySelected() const
 {
+	bool bIndividuallySelected = false;
 #if WITH_EDITOR
-	if (SelectionOverrideDelegate.IsBound())
+	if(SelectionOverrideDelegate.IsBound())
 	{
-		return SelectionOverrideDelegate.Execute(this);
+		bIndividuallySelected = SelectionOverrideDelegate.Execute(this);
 	}
-	else
 #endif
-	{
-		const AActor* Owner = GetOwner();
-		return(	bSelectable && 
-				Owner != NULL && 
+	return bIndividuallySelected;
+}
+
+bool UPrimitiveComponent::ShouldRenderSelected() const
+{
+	const AActor* Owner = GetOwner();
+	return(	bSelectable && 
+			Owner != NULL && 
 #if WITH_EDITOR
-				(Owner->IsSelected() || (Owner->ParentComponentActor != NULL && Owner->ParentComponentActor->IsSelected())) );
+			(Owner->IsSelected() || (Owner->ParentComponentActor != NULL && Owner->ParentComponentActor->IsSelected())) );
 #else
-				Owner->IsSelected() );
+			Owner->IsSelected() );
 #endif
-	}
 }
 
 void UPrimitiveComponent::SetCastShadow(bool NewCastShadow)
@@ -877,7 +880,7 @@ void UPrimitiveComponent::PushSelectionToProxy()
 	//although this should only be called for attached components, some billboard components can get in without valid proxies
 	if (SceneProxy)
 	{
-		SceneProxy->SetSelection_GameThread(ShouldRenderSelected());
+		SceneProxy->SetSelection_GameThread(ShouldRenderSelected(),IsComponentIndividuallySelected());
 	}
 }
 
diff --git a/Engine/Source/Runtime/Engine/Private/PrimitiveSceneProxy.cpp b/Engine/Source/Runtime/Engine/Private/PrimitiveSceneProxy.cpp
--- a/Engine/Source/Runtime/Engine/Private/PrimitiveSceneProxy.cpp
+++ b/Engine/Source/Runtime/Engine/Private/PrimitiveSceneProxy.cpp
@@ -19,7 +19,8 @@ FPrimitiveSceneProxy::FPrimitiveSceneProxy(const UPrimitiveComponent* InComponen
 ,	bOwnerNoSee(InComponent->bOwnerNoSee)
 ,	bStatic(false)
 ,	bOftenMoving(false)
-,	bSelected(InComponent->ShouldRenderSelected())
+,	bParentSelected(InComponent->ShouldRenderSelected())
+,	bIndividuallySelected(InComponent->IsComponentIndividuallySelected())
 ,	bHovered(false)
 ,	bUseViewOwnerDepthPriorityGroup(InComponent->bUseViewOwnerDepthPriorityGroup)
 ,	bHasMotionBlurVelocityMeshes(InComponent->bHasMotionBlurVelocityMeshes)
@@ -234,10 +235,11 @@ void FPrimitiveSceneProxy::ApplyWorldOffset(FVector InOffset)
  * Updates selection for the primitive proxy. This is called in the rendering thread by SetSelection_GameThread.
  * @param bInSelected - true if the parent actor is selected in the editor
  */
-void FPrimitiveSceneProxy::SetSelection_RenderThread(const bool bInSelected)
+void FPrimitiveSceneProxy::SetSelection_RenderThread(const bool bInParentSelected, const bool bInIndividuallySelected)
 {
 	check(IsInRenderingThread());
-	bSelected = bInSelected;
+	bParentSelected = bInParentSelected;
+	bIndividuallySelected = bInIndividuallySelected;
 }
 
 /**
@@ -245,17 +247,18 @@ void FPrimitiveSceneProxy::SetSelection_RenderThread(const bool bInSelected)
  * This is called in the game thread as selection is toggled.
  * @param bInSelected - true if the parent actor is selected in the editor
  */
-void FPrimitiveSceneProxy::SetSelection_GameThread(const bool bInSelected)
+void FPrimitiveSceneProxy::SetSelection_GameThread(const bool bInParentSelected, const bool bInIndividuallySelected)
 {
 	check(IsInGameThread());
 
 	// Enqueue a message to the rendering thread containing the interaction to add.
-	ENQUEUE_UNIQUE_RENDER_COMMAND_TWOPARAMETER(
+	ENQUEUE_UNIQUE_RENDER_COMMAND_THREEPARAMETER(
 		SetNewSelection,
 		FPrimitiveSceneProxy*,PrimitiveSceneProxy,this,
-		const bool,bNewSelection,bInSelected,
+		const bool,bNewParentSelection,bInParentSelected,
+		const bool,bNewIndividuallySelected,bInIndividuallySelected,
 	{
-		PrimitiveSceneProxy->SetSelection_RenderThread(bNewSelection);
+		PrimitiveSceneProxy->SetSelection_RenderThread(bNewParentSelection,bNewIndividuallySelected);
 	});
 }
 
diff --git a/Engine/Source/Runtime/Engine/Private/SceneView.cpp b/Engine/Source/Runtime/Engine/Private/SceneView.cpp
--- a/Engine/Source/Runtime/Engine/Private/SceneView.cpp
+++ b/Engine/Source/Runtime/Engine/Private/SceneView.cpp
@@ -290,6 +290,7 @@ FSceneView::FSceneView(const FSceneViewInitOptions& InitOptions)
 #if WITH_EDITOR
 	, OverrideLODViewOrigin(InitOptions.OverrideLODViewOrigin)
 	, bAllowTranslucentPrimitivesInHitProxy( true )
+	, bHasSelectedComponents( false )
 #endif
 	, FeatureLevel(InitOptions.ViewFamily ? InitOptions.ViewFamily->GetFeatureLevel() : GMaxRHIFeatureLevel)
 {
diff --git a/Engine/Source/Runtime/Engine/Public/PrimitiveSceneProxy.h b/Engine/Source/Runtime/Engine/Public/PrimitiveSceneProxy.h
--- a/Engine/Source/Runtime/Engine/Public/PrimitiveSceneProxy.h
+++ b/Engine/Source/Runtime/Engine/Public/PrimitiveSceneProxy.h
@@ -129,9 +129,10 @@ public:
 	/**
 	 * Updates selection for the primitive proxy. This simply sends a message to the rendering thread to call SetSelection_RenderThread.
 	 * This is called in the game thread as selection is toggled.
-	 * @param bInSelected - true if the parent actor is selected in the editor
+	 * @param bInParentSelected - true if the parent actor is selected in the editor
+ 	 * @param bInIndividuallySelected - true if the component is selected in the editor directly
 	 */
-	void SetSelection_GameThread(const bool bInSelected);
+	void SetSelection_GameThread(const bool bInParentSelected, const bool bInIndividuallySelected=false);
 
 	/**
 	 * Updates hover state for the primitive proxy. This simply sends a message to the rendering thread to call SetHovered_RenderThread.
@@ -348,7 +349,9 @@ public:
 	inline bool IsOftenMoving() const { return bOftenMoving; }
 	inline bool IsStatic() const { return bStatic; }
 	inline bool IsSelectable() const { return bSelectable; }
-	inline bool IsSelected() const { return bSelected; }
+	inline bool IsParentSelected() const { return bParentSelected; }
+	inline bool IsIndividuallySelected() const { return bIndividuallySelected; }
+	inline bool IsSelected() const { return IsParentSelected() || IsIndividuallySelected(); }
 	inline bool ShouldRenderCustomDepth() const { return bRenderCustomDepth; }
 	inline bool ShouldRenderInMainPass() const { return bRenderInMainPass; }
 	inline bool IsCollisionEnabled() const { return bCollisionEnabled; }
@@ -475,7 +478,10 @@ private:
 	uint32 bOwnerNoSee : 1;
 	uint32 bStatic : 1;
 	uint32 bOftenMoving : 1;
-	uint32 bSelected : 1;
+	/** Parent Actor is selected */
+	uint32 bParentSelected : 1;
+	/** Component is selected directly */
+	uint32 bIndividuallySelected : 1;
 	
 	/** true if the mouse is currently hovered over this primitive in a level viewport */
 	uint32 bHovered : 1;
@@ -704,7 +710,7 @@ private:
 
 protected:
 	/** Updates selection for the primitive proxy. This is called in the rendering thread by SetSelection_GameThread. */
-	void SetSelection_RenderThread(const bool bInSelected);
+	void SetSelection_RenderThread(const bool bInParentSelected, const bool bInIndividuallySelected);
 
 	/** Updates hover state for the primitive proxy. This is called in the rendering thread by SetHovered_GameThread. */
 	void SetHovered_RenderThread(const bool bInHovered);
diff --git a/Engine/Source/Runtime/Engine/Public/SceneView.h b/Engine/Source/Runtime/Engine/Public/SceneView.h
--- a/Engine/Source/Runtime/Engine/Public/SceneView.h
+++ b/Engine/Source/Runtime/Engine/Public/SceneView.h
@@ -521,6 +521,10 @@ public:
 	TBitArray<> SpriteCategoryVisibility;
 	/** Selection color for the editor (used by post processing) */
 	FLinearColor SelectionOutlineColor;
+	/** Selection color for use in the editor with inactive primitives */
+	FLinearColor SubduedSelectionOutlineColor;
+	/** True if any components are selected in isolation (independent of actor selection) */
+	bool bHasSelectedComponents;
 #endif
 
 	/**
diff --git a/Engine/Source/Runtime/Renderer/Private/PostProcess/PostProcessSelectionOutline.cpp b/Engine/Source/Runtime/Renderer/Private/PostProcess/PostProcessSelectionOutline.cpp
--- a/Engine/Source/Runtime/Renderer/Private/PostProcess/PostProcessSelectionOutline.cpp
+++ b/Engine/Source/Runtime/Renderer/Private/PostProcess/PostProcessSelectionOutline.cpp
@@ -46,6 +46,7 @@ void FRCPassPostProcessSelectionOutlineColor::Process(FRenderingCompositePassCon
 
 		//@todo - use memstack
 		TMap<FName, int32> ActorNameToStencilIndex;
+		TMap<const FPrimitiveSceneProxy*, int32> IndividuallySelectedProxies;
 		ActorNameToStencilIndex.Add(NAME_BSP, 1);
 
 		Context.RHICmdList.SetRasterizerState(TStaticRasterizerState<>::GetRHI());
@@ -56,13 +57,30 @@ void FRCPassPostProcessSelectionOutlineColor::Process(FRenderingCompositePassCon
 			const FMeshBatchAndRelevance& MeshBatchAndRelevance = View.DynamicMeshElements[MeshBatchIndex];
 			const FPrimitiveSceneProxy* PrimitiveSceneProxy = MeshBatchAndRelevance.PrimitiveSceneProxy;
 
+#if WITH_EDITOR
+			// Selected actors should be subdued if any component is individually selected
+			bool bActorSelectionColorIsSubdued = View.bHasSelectedComponents;
+#else
+			bool bActorSelectionColorIsSubdued = false;
+#endif
 			if (PrimitiveSceneProxy->IsSelected() && MeshBatchAndRelevance.Mesh->bUseSelectionOutline)
 			{
-				const int32* AssignedStencilIndexPtr = ActorNameToStencilIndex.Find(PrimitiveSceneProxy->GetOwnerName());
+				const int32* AssignedStencilIndexPtr = PrimitiveSceneProxy->IsIndividuallySelected() ? IndividuallySelectedProxies.Find( PrimitiveSceneProxy ) : ActorNameToStencilIndex.Find(PrimitiveSceneProxy->GetOwnerName());
 
 				if (!AssignedStencilIndexPtr)
 				{
-					AssignedStencilIndexPtr = &ActorNameToStencilIndex.Add(PrimitiveSceneProxy->GetOwnerName(), ActorNameToStencilIndex.Num() + 1);
+					if( PrimitiveSceneProxy->IsIndividuallySelected() )
+					{
+						// Any component that is individually selected should have a stencil value of < 128 so that it can have a unique color.  We offset the value by 2 because 0 means no selection and 1 is for bsp
+						int32 StencilValue = IndividuallySelectedProxies.Num() % 126 + 2;
+						AssignedStencilIndexPtr = &IndividuallySelectedProxies.Add(PrimitiveSceneProxy, StencilValue);
+					}
+					else
+					{
+						// If we are subduing actor color highlight then use the top level bits to indicate that to the shader.  
+						int32 StencilValue = bActorSelectionColorIsSubdued ? ActorNameToStencilIndex.Num() % 128 + 128 : ActorNameToStencilIndex.Num() % 126 + 2;
+						AssignedStencilIndexPtr = &ActorNameToStencilIndex.Add(PrimitiveSceneProxy->GetOwnerName(), StencilValue);
+					}
 				}
 
 				// This is a reversed Z depth surface, using CF_GreaterEqual.
@@ -160,6 +178,7 @@ public:
 	FPostProcessPassParameters PostprocessParameter;
 	FDeferredPixelShaderParameters DeferredParameters;
 	FShaderParameter OutlineColor;
+	FShaderParameter SubduedOutlineColor;
 	FShaderParameter BSPSelectionIntensity;
 	FShaderResourceParameter PostprocessInput1MS;
 	FShaderResourceParameter EditorPrimitivesStencil;
@@ -172,6 +191,7 @@ public:
 		PostprocessParameter.Bind(Initializer.ParameterMap);
 		DeferredParameters.Bind(Initializer.ParameterMap);
 		OutlineColor.Bind(Initializer.ParameterMap, TEXT("OutlineColor"));
+		SubduedOutlineColor.Bind(Initializer.ParameterMap, TEXT("SubduedOutlineColor"));
 		BSPSelectionIntensity.Bind(Initializer.ParameterMap, TEXT("BSPSelectionIntensity"));
 		PostprocessInput1MS.Bind(Initializer.ParameterMap, TEXT("PostprocessInput1MS"));
 		EditorRenderParams.Bind(Initializer.ParameterMap, TEXT("EditorRenderParams"));
@@ -233,11 +253,12 @@ public:
 
 #if WITH_EDITOR
 		{
-			FLinearColor Value = Context.View.SelectionOutlineColor;
-
-			Value.A = GEngine->SelectionHighlightIntensity;
+			FLinearColor OutlineColorValue = Context.View.SelectionOutlineColor;
+			FLinearColor SubduedOutlineColorValue = Context.View.SubduedSelectionOutlineColor;
+			OutlineColorValue.A = GEngine->SelectionHighlightIntensity;
 
-			SetShaderValue(Context.RHICmdList, ShaderRHI, OutlineColor, Value);
+			SetShaderValue(Context.RHICmdList, ShaderRHI, OutlineColor, OutlineColorValue);
+			SetShaderValue(Context.RHICmdList, ShaderRHI, SubduedOutlineColor, SubduedOutlineColorValue);
 			SetShaderValue(Context.RHICmdList, ShaderRHI, BSPSelectionIntensity, GEngine->BSPSelectionHighlightIntensity);
 		}
 #else
@@ -263,7 +284,7 @@ public:
 	virtual bool Serialize(FArchive& Ar)
 	{
 		bool bShaderHasOutdatedParameters = FGlobalShader::Serialize(Ar);
-		Ar << PostprocessParameter << OutlineColor << BSPSelectionIntensity << DeferredParameters << PostprocessInput1MS << EditorPrimitivesStencil << EditorRenderParams;
+		Ar << PostprocessParameter << OutlineColor << SubduedOutlineColor << BSPSelectionIntensity << DeferredParameters << PostprocessInput1MS << EditorPrimitivesStencil << EditorRenderParams;
 		return bShaderHasOutdatedParameters;
 	}
 
