UE4 - Fix GBufferVelocity render target index when disabling static lighting

diff --git a/Engine/Shaders/BasePassPixelShader.usf b/Engine/Shaders/BasePassPixelShader.usf
--- a/Engine/Shaders/BasePassPixelShader.usf
+++ b/Engine/Shaders/BasePassPixelShader.usf
@@ -587,9 +587,13 @@ void Main(
 	,out float4 OutGBufferD : SV_Target4
 	#if ALLOW_STATIC_LIGHTING
 		,out float4 OutGBufferE : SV_Target5
-	#endif
-	#if WRITES_VELOCITY_TO_GBUFFER
-		, out float4 OutGBufferVelocity : SV_Target6
+		#if WRITES_VELOCITY_TO_GBUFFER
+			, out float4 OutGBufferVelocity : SV_Target6
+		#endif
+	#else
+		#if WRITES_VELOCITY_TO_GBUFFER
+			, out float4 OutGBufferVelocity : SV_Target5
+		#endif
 	#endif
 #endif
 
diff --git a/Engine/Source/Runtime/Renderer/Private/BasePassRendering.h b/Engine/Source/Runtime/Renderer/Private/BasePassRendering.h
--- a/Engine/Source/Runtime/Renderer/Private/BasePassRendering.h
+++ b/Engine/Source/Runtime/Renderer/Private/BasePassRendering.h
@@ -337,7 +337,10 @@ public:
 		const bool bOutputVelocity = FVelocityRendering::OutputsToGBuffer();
 		if (bOutputVelocity)
 		{
-			OutEnvironment.SetRenderTargetOutputFormat(6, PF_G16R16);
+			static const auto CVar = IConsoleManager::Get().FindTConsoleVariableDataInt(TEXT("r.AllowStaticLighting"));
+			// This needs to match FSceneRenderTargets::GetNumGBufferTargets()
+			const int32 VelocityIndex = (CVar && CVar->GetValueOnAnyThread() != 0) ? 6 : 5;
+			OutEnvironment.SetRenderTargetOutputFormat(VelocityIndex, PF_G16R16);
 		}
 	}
 
diff --git a/Engine/Source/Runtime/Renderer/Private/PostProcess/SceneRenderTargets.cpp b/Engine/Source/Runtime/Renderer/Private/PostProcess/SceneRenderTargets.cpp
--- a/Engine/Source/Runtime/Renderer/Private/PostProcess/SceneRenderTargets.cpp
+++ b/Engine/Source/Runtime/Renderer/Private/PostProcess/SceneRenderTargets.cpp
@@ -412,6 +412,7 @@ int32 FSceneRenderTargets::GetNumGBufferTargets() const
 
 	if (CurrentFeatureLevel >= ERHIFeatureLevel::SM4 && !IsSimpleDynamicLightingEnabled())
 	{
+		// This needs to match TBasePassPixelShaderBaseType::ModifyCompilationEnvironment()
 		NumGBufferTargets = bAllowStaticLighting ? 6 : 5;
 
 		if (bAllocateVelocityGBuffer)
