UE4 - Some asserts to track down incorrect GBuffer reordering

diff --git a/Engine/Source/Runtime/Renderer/Private/PostProcess/SceneRenderTargets.cpp b/Engine/Source/Runtime/Renderer/Private/PostProcess/SceneRenderTargets.cpp
--- a/Engine/Source/Runtime/Renderer/Private/PostProcess/SceneRenderTargets.cpp
+++ b/Engine/Source/Runtime/Renderer/Private/PostProcess/SceneRenderTargets.cpp
@@ -274,7 +274,7 @@ void FSceneRenderTargets::Allocate(const FSceneViewFamily& ViewFamily)
 		CurrentMobileMSAA = MobileMSAA;
 		CurrentMinShadowResolution = MinShadowResolution;
 		bCurrentLightPropagationVolume = bLightPropagationVolume;
-		
+
 		// Reinitialize the render targets for the given size.
 		SetBufferSize(DesiredBufferSize.X, DesiredBufferSize.Y);
 
@@ -310,6 +310,7 @@ int32 FSceneRenderTargets::GetGBufferRenderTargets(ERenderTargetLoadAction Color
 
 	if (bAllowStaticLighting)
 	{
+		check(MRTCount == GetGBufferEIndex());
 		OutRenderTargets[MRTCount] = FRHIRenderTargetView(GSceneRenderTargets.GBufferE->GetRenderTargetItem().TargetableTexture, 0, -1, ColorLoadAction, ERenderTargetStoreAction::EStore);
 		++MRTCount;
 	}
@@ -318,6 +319,7 @@ int32 FSceneRenderTargets::GetGBufferRenderTargets(ERenderTargetLoadAction Color
 	{
 		OutVelocityRTIndex = MRTCount;
 		++MRTCount;
+		check(OutVelocityRTIndex == GetGBufferVelocityIndex());
 		OutRenderTargets[OutVelocityRTIndex] = FRHIRenderTargetView(GSceneRenderTargets.GBufferVelocity->GetRenderTargetItem().TargetableTexture, 0, -1, ColorLoadAction, ERenderTargetStoreAction::EStore);
 	}
 
diff --git a/Engine/Source/Runtime/Renderer/Private/PostProcess/SceneRenderTargets.h b/Engine/Source/Runtime/Renderer/Private/PostProcess/SceneRenderTargets.h
--- a/Engine/Source/Runtime/Renderer/Private/PostProcess/SceneRenderTargets.h
+++ b/Engine/Source/Runtime/Renderer/Private/PostProcess/SceneRenderTargets.h
@@ -357,6 +357,21 @@ public:
 
 	IPooledRenderTarget* GetGBufferVelocityRT();
 
+	int32 GetGBufferEIndex() const
+	{
+		return bAllowStaticLighting ? 5 : -1;
+	}
+
+	int32 GetGBufferVelocityIndex() const
+	{
+		if (bAllocateVelocityGBuffer)
+		{
+			return bAllowStaticLighting ? 6 : 5;
+		}
+
+		return -1;
+	}
+
 	// @return can be 0 if the feature is disabled
 	IPooledRenderTarget* RequestCustomDepth(bool bPrimitives);
 
