Integrate - Sky light and ight function support for heightfield lighting

diff --git a/Engine/Shaders/HeightfieldLighting.usf b/Engine/Shaders/HeightfieldLighting.usf
--- a/Engine/Shaders/HeightfieldLighting.usf
+++ b/Engine/Shaders/HeightfieldLighting.usf
@@ -10,6 +10,12 @@
 #include "DistanceFieldAOShared.usf"
 #include "DynamicLightingCommon.usf"
 #include "DistanceFieldShadowingShared.usf"
+#include "ReflectionEnvironmentShared.usf"
+
+#if APPLY_LIGHT_FUNCTION
+#include "Material.usf"
+#include "LightFunctionCommon.usf"
+#endif
 
 Texture2D HeightfieldTexture;
 SamplerState HeightfieldSampler;
@@ -320,6 +326,10 @@ void RayTracedShadowHeightfieldsPS(
 
 float3 LightColor;
 
+#if APPLY_LIGHT_FUNCTION
+
+float4x4 WorldToLight;
+
 Texture2D HeightfieldShadowing;
 SamplerState HeightfieldShadowingSampler;
 
@@ -329,15 +339,26 @@ void LightHeightfieldsPS(
 {
 	float3 WorldHeightfieldNormal = GetGlobalHeightfieldNormal(Input.GlobalUV);
 	float3 WorldHeightfieldShadingPosition = GetGlobalHeightfieldWorldPosition(Input.LocalPosition, Input.GlobalUV, GetLocalToWorld(Input.HeightfieldIndex));
+	
+	float4 LightVector = mul(WorldHeightfieldShadingPosition, WorldToLight);
+	LightVector.xyz /= LightVector.w;
+
+	float3 LightFunction = GetLightFunctionColor(LightVector.xyz, WorldHeightfieldShadingPosition);
+	float GreyScaleLightFunction = dot(LightFunction, .3333f);
 
 	float Shadowing = Texture2DSampleLevel(HeightfieldShadowing, HeightfieldShadowingSampler, Input.GlobalUV, 0).x;
 
 	float3 DiffuseColor = .5f;
-	float3 Lighting = (Shadowing * saturate(dot(-LightDirection, WorldHeightfieldNormal))) * DiffuseColor * LightColor;
+	float3 Lighting = (Shadowing * GreyScaleLightFunction * saturate(dot(-LightDirection, WorldHeightfieldNormal))) * LightColor;
+
+	float3 SkylightDiffuse = GetSkySHDiffuse(WorldHeightfieldNormal) * View.SkyLightColor.rgb;
+	Lighting += SkylightDiffuse;
 
-	OutColor = float4(Lighting, 0);
+	OutColor = float4(DiffuseColor * Lighting, 0);
 }
 
+#endif
+
 RWBuffer<float> RWOccluderRadius;
 RWBuffer<float> RWRecordConeVisibility;
 
diff --git a/Engine/Source/Runtime/Renderer/Private/HeightfieldLighting.cpp b/Engine/Source/Runtime/Renderer/Private/HeightfieldLighting.cpp
--- a/Engine/Source/Runtime/Renderer/Private/HeightfieldLighting.cpp
+++ b/Engine/Source/Runtime/Renderer/Private/HeightfieldLighting.cpp
@@ -16,6 +16,7 @@
 #include "RHICommandList.h"
 #include "SceneUtils.h"
 #include "DistanceFieldAtlas.h"
+#include "LightRendering.h"
 
 int32 GAOHeightfieldOcclusion = 1;
 FAutoConsoleVariableRef CVarAOHeightfieldOcclusion(
@@ -1029,14 +1030,19 @@ void FHeightfieldLightingViewInfo::ComputeRayTracedShadowing(
 	}
 }
 
-class FLightHeightfieldsPS : public FGlobalShader
+class FLightHeightfieldsPS : public FMaterialShader
 {
-	DECLARE_SHADER_TYPE(FLightHeightfieldsPS, Global);
+	DECLARE_SHADER_TYPE(FLightHeightfieldsPS, Material);
 public:
 
-	static bool ShouldCache(EShaderPlatform Platform)
+	static bool ShouldCache(EShaderPlatform Platform, const FMaterial* Material)
 	{
-		return IsFeatureLevelSupported(Platform, ERHIFeatureLevel::SM5) && DoesPlatformSupportDistanceFieldAO(Platform);
+		return Material->IsLightFunction() && IsFeatureLevelSupported(Platform, ERHIFeatureLevel::SM5) && DoesPlatformSupportDistanceFieldAO(Platform);
+	}
+
+	static void ModifyCompilationEnvironment(EShaderPlatform Platform, const FMaterial* Material, FShaderCompilerEnvironment& OutEnvironment)
+	{
+		OutEnvironment.SetDefine(TEXT("APPLY_LIGHT_FUNCTION"), 1);
 	}
 
 	/** Default constructor. */
@@ -1044,63 +1050,84 @@ public:
 
 	/** Initialization constructor. */
 	FLightHeightfieldsPS(const ShaderMetaType::CompiledShaderInitializerType& Initializer)
-		: FGlobalShader(Initializer)
+		: FMaterialShader(Initializer)
 	{
 		HeightfieldDescriptionParameters.Bind(Initializer.ParameterMap);
-		HeightfieldTextureParameters.Bind(Initializer.ParameterMap);
 		GlobalHeightfieldParameters.Bind(Initializer.ParameterMap);
 		LightDirection.Bind(Initializer.ParameterMap, TEXT("LightDirection"));
 		LightColor.Bind(Initializer.ParameterMap, TEXT("LightColor"));
 		HeightfieldShadowing.Bind(Initializer.ParameterMap, TEXT("HeightfieldShadowing"));
 		HeightfieldShadowingSampler.Bind(Initializer.ParameterMap, TEXT("HeightfieldShadowingSampler"));
+		WorldToLight.Bind(Initializer.ParameterMap, TEXT("WorldToLight"));
+		LightFunctionParameters.Bind(Initializer.ParameterMap);
 	}
 
-	void SetParameters(FRHICommandList& RHICmdList, const FSceneView& View, const FLightSceneProxy* LightSceneProxy, UTexture2D* HeightfieldTextureValue, int32 NumHeightfieldsValue, const FHeightfieldLightingAtlas& Atlas)
+	void SetParameters(
+		FRHICommandList& RHICmdList, 
+		const FSceneView& View, 
+		const FLightSceneInfo& LightSceneInfo, 
+		const FMaterialRenderProxy* MaterialProxy,
+		int32 NumHeightfieldsValue, 
+		const FHeightfieldLightingAtlas& Atlas)
 	{
 		const FPixelShaderRHIParamRef ShaderRHI = GetPixelShader();
 
-		FGlobalShader::SetParameters(RHICmdList, ShaderRHI, View);
+		FMaterialShader::SetParameters(RHICmdList, ShaderRHI, MaterialProxy, *MaterialProxy->GetMaterial(View.GetFeatureLevel()), View, true, ESceneRenderTargetsMode::SetTextures);
+
 		HeightfieldDescriptionParameters.Set(RHICmdList, ShaderRHI, NumHeightfieldsValue);
-		HeightfieldTextureParameters.Set(RHICmdList, ShaderRHI, HeightfieldTextureValue);
 		GlobalHeightfieldParameters.Set(RHICmdList, ShaderRHI, Atlas);
 
-		SetShaderValue(RHICmdList, ShaderRHI, LightDirection, LightSceneProxy->GetDirection());
-		SetShaderValue(RHICmdList, ShaderRHI, LightColor, LightSceneProxy->GetColor() * LightSceneProxy->GetIndirectLightingScale());
+		SetShaderValue(RHICmdList, ShaderRHI, LightDirection, LightSceneInfo.Proxy->GetDirection());
+		SetShaderValue(RHICmdList, ShaderRHI, LightColor, LightSceneInfo.Proxy->GetColor() * LightSceneInfo.Proxy->GetIndirectLightingScale());
 
 		const FSceneViewState* ViewState = (const FSceneViewState*)View.State;
 		SetTextureParameter(RHICmdList, ShaderRHI, HeightfieldShadowing, HeightfieldShadowingSampler, TStaticSamplerState<SF_Bilinear>::GetRHI(), Atlas.DirectionalLightShadowing->GetRenderTargetItem().ShaderResourceTexture);
+	
+		const FVector Scale = LightSceneInfo.Proxy->GetLightFunctionScale();
+		// Switch x and z so that z of the user specified scale affects the distance along the light direction
+		const FVector InverseScale = FVector( 1.f / Scale.Z, 1.f / Scale.Y, 1.f / Scale.X );
+		const FMatrix WorldToLightValue = LightSceneInfo.Proxy->GetWorldToLight() * FScaleMatrix(FVector(InverseScale));	
+
+		SetShaderValue(RHICmdList, ShaderRHI, WorldToLight, WorldToLightValue);
+
+		LightFunctionParameters.Set(RHICmdList, ShaderRHI, &LightSceneInfo, 1);
 	}
+
 	// FShader interface.
 	virtual bool Serialize(FArchive& Ar)
 	{
-		bool bShaderHasOutdatedParameters = FGlobalShader::Serialize(Ar);
+		bool bShaderHasOutdatedParameters = FMaterialShader::Serialize(Ar);
 		Ar << HeightfieldDescriptionParameters;
-		Ar << HeightfieldTextureParameters;
 		Ar << GlobalHeightfieldParameters;
 		Ar << LightDirection;
 		Ar << LightColor;
 		Ar << HeightfieldShadowing;
 		Ar << HeightfieldShadowingSampler;
+		Ar << WorldToLight;
+		Ar << LightFunctionParameters;
 		return bShaderHasOutdatedParameters;
 	}
 
 private:
 
 	FHeightfieldDescriptionParameters HeightfieldDescriptionParameters;
-	FHeightfieldTextureParameters HeightfieldTextureParameters;
 	FGlobalHeightfieldParameters GlobalHeightfieldParameters;
 	FShaderParameter LightDirection;
 	FShaderParameter LightColor;
 	FShaderResourceParameter HeightfieldShadowing;
 	FShaderResourceParameter HeightfieldShadowingSampler;
+	FShaderParameter WorldToLight;
+	FLightFunctionSharedParameters LightFunctionParameters;
 };
 
-IMPLEMENT_SHADER_TYPE(,FLightHeightfieldsPS,TEXT("HeightfieldLighting"),TEXT("LightHeightfieldsPS"),SF_Pixel);
+IMPLEMENT_MATERIAL_SHADER_TYPE(,FLightHeightfieldsPS,TEXT("HeightfieldLighting"),TEXT("LightHeightfieldsPS"),SF_Pixel);
 
 void FHeightfieldLightingViewInfo::ComputeLighting(const FViewInfo& View, FRHICommandListImmediate& RHICmdList, const FLightSceneInfo& LightSceneInfo) const
 {
+	const ERHIFeatureLevel::Type FeatureLevel = View.GetFeatureLevel();
+
 	if (AllowHeightfieldGI(View) 
-		&& SupportsHeightfieldLighting(View.GetFeatureLevel(), View.GetShaderPlatform())
+		&& SupportsHeightfieldLighting(FeatureLevel, View.GetShaderPlatform())
 		//@todo - handle local lights
 		&& LightSceneInfo.Proxy->GetLightType() == LightType_Directional
 		&& LightSceneInfo.Proxy->CastsDynamicShadow()
@@ -1109,18 +1136,27 @@ void FHeightfieldLightingViewInfo::ComputeLighting(const FViewInfo& View, FRHICo
 		SCOPED_DRAW_EVENT(RHICmdList, HeightfieldLightingForGI);
 
 		FSceneViewState* ViewState = (FSceneViewState*)View.State;
+		const FHeightfieldLightingAtlas& Atlas = *ViewState->HeightfieldLightingAtlas;
 
-		{
-			const FHeightfieldLightingAtlas& Atlas = *ViewState->HeightfieldLightingAtlas;
+		const FIntPoint LightingAtlasSize = Atlas.GetAtlasSize();
+		const FVector2D InvLightingAtlasSize(1.0f / LightingAtlasSize.X, 1.0f / LightingAtlasSize.Y);
 
-			const FIntPoint LightingAtlasSize = Atlas.GetAtlasSize();
-			const FVector2D InvLightingAtlasSize(1.0f / LightingAtlasSize.X, 1.0f / LightingAtlasSize.Y);
+		SetRenderTarget(RHICmdList, Atlas.Lighting->GetRenderTargetItem().TargetableTexture, NULL);
 
-			SetRenderTarget(RHICmdList, Atlas.Lighting->GetRenderTargetItem().TargetableTexture, NULL);
+		// We will read outside the valid area later during light transfer
+		RHICmdList.Clear(true, FLinearColor(0, 0, 0), false, 0, false, 0, FIntRect());
 
-			// We will read outside the valid area later during light transfer
-			RHICmdList.Clear(true, FLinearColor(0, 0, 0), false, 0, false, 0, FIntRect());
+		const bool bApplyLightFunction = (View.Family->EngineShowFlags.LightFunctions &&
+			LightSceneInfo.Proxy->GetLightFunctionMaterial() && 
+			LightSceneInfo.Proxy->GetLightFunctionMaterial()->GetMaterial(FeatureLevel)->IsLightFunction());
 
+		const FMaterialRenderProxy* MaterialProxy = bApplyLightFunction ? 
+			LightSceneInfo.Proxy->GetLightFunctionMaterial() : 
+			UMaterial::GetDefaultMaterial(MD_LightFunction)->GetRenderProxy(false);
+
+		// Skip rendering if the DefaultLightFunctionMaterial isn't compiled yet
+		if (MaterialProxy->GetMaterial(FeatureLevel)->IsLightFunction())
+		{
 			RHICmdList.SetViewport(0, 0, 0.0f, LightingAtlasSize.X, LightingAtlasSize.Y, 1.0f);
 			RHICmdList.SetRasterizerState(TStaticRasterizerState<FM_Solid, CM_None>::GetRHI());
 			RHICmdList.SetDepthStencilState(TStaticDepthStencilState<false, CF_Always>::GetRHI());
@@ -1128,10 +1164,12 @@ void FHeightfieldLightingViewInfo::ComputeLighting(const FViewInfo& View, FRHICo
 			RHICmdList.SetStreamSource(0, GQuadVertexBuffer.VertexBufferRHI, sizeof(FScreenVertex), 0);
 
 			TShaderMapRef<FHeightfieldComponentQuadVS> VertexShader(View.ShaderMap);
-			TShaderMapRef<FLightHeightfieldsPS> PixelShader(View.ShaderMap);
 
-			static FGlobalBoundShaderState BoundShaderState;
-			SetGlobalBoundShaderState(RHICmdList, View.GetFeatureLevel(), BoundShaderState, GScreenVertexDeclaration.VertexDeclarationRHI, *VertexShader, *PixelShader);
+			const FMaterial* Material = MaterialProxy->GetMaterial(FeatureLevel);
+			const FMaterialShaderMap* MaterialShaderMap = Material->GetRenderingThreadShaderMap();
+			FLightHeightfieldsPS* PixelShader = MaterialShaderMap->GetShader<FLightHeightfieldsPS>();
+
+			FLocalBoundShaderState BoundShaderState = RHICmdList.BuildLocalBoundShaderState(GScreenVertexDeclaration.VertexDeclarationRHI, VertexShader->GetVertexShader(), FHullShaderRHIRef(), FDomainShaderRHIRef(), PixelShader->GetPixelShader(), FGeometryShaderRHIRef());
 
 			for (TMap<UTexture2D*, TArray<FHeightfieldComponentDescription>>::TConstIterator It(Heightfield.ComponentDescriptions); It; ++It)
 			{
@@ -1143,8 +1181,10 @@ void FHeightfieldLightingViewInfo::ComputeLighting(const FViewInfo& View, FRHICo
 
 					UTexture2D* HeightfieldTexture = It.Key();
 
+					RHICmdList.SetLocalBoundShaderState(BoundShaderState);
+
 					VertexShader->SetParameters(RHICmdList, View, HeightfieldDescriptions.Num());
-					PixelShader->SetParameters(RHICmdList, View, LightSceneInfo.Proxy, HeightfieldTexture, HeightfieldDescriptions.Num(), Atlas);
+					PixelShader->SetParameters(RHICmdList, View, LightSceneInfo, MaterialProxy, HeightfieldDescriptions.Num(), Atlas);
 
 					RHICmdList.DrawPrimitive(PT_TriangleList, 0, 2, HeightfieldDescriptions.Num());
 				}
