Rendering light shape into reflection captures for static lights.

diff --git a/Engine/Source/Runtime/Renderer/Private/SceneVisibility.cpp b/Engine/Source/Runtime/Renderer/Private/SceneVisibility.cpp
--- a/Engine/Source/Runtime/Renderer/Private/SceneVisibility.cpp
+++ b/Engine/Source/Runtime/Renderer/Private/SceneVisibility.cpp
@@ -2150,6 +2150,77 @@ void FSceneRenderer::PostVisibilityFrameSetup()
 					}
 				}
 			}
+
+			// Draw shapes for reflection captures
+			if( View.bIsReflectionCapture && VisibleLightViewInfo.bInViewFrustum && Proxy->HasStaticLighting() && Proxy->GetLightType() != LightType_Directional )
+			{
+				FVector Origin = Proxy->GetOrigin();
+				FVector ToLight = Origin - View.ViewMatrices.ViewOrigin;
+				float DistanceSqr = ToLight | ToLight;
+				float Radius = Proxy->GetRadius();
+
+				if( DistanceSqr < Radius * Radius )
+				{
+					FVector4	PositionAndInvRadius;
+					FVector4	ColorAndFalloffExponent;
+					FVector		Direction;
+					FVector2D	SpotAngles;
+					float		SourceRadius;
+					float		SourceLength;
+					float		MinRoughness;
+					Proxy->GetParameters( PositionAndInvRadius, ColorAndFalloffExponent, Direction, SpotAngles, SourceRadius, SourceLength, MinRoughness );
+
+					// Force to be at least 0.75 pixels
+					float CubemapSize = 128.0f;
+					float Distance = FMath::Sqrt( DistanceSqr );
+					float MinRadius = Distance * 0.75f / CubemapSize;
+					SourceRadius = FMath::Max( MinRadius, SourceRadius );
+
+					// Snap to cubemap pixel center to reduce aliasing
+					FVector Scale = ToLight.GetAbs();
+					int32 MaxComponent = Scale.X > Scale.Y ? ( Scale.X > Scale.Z ? 0 : 2 ) : ( Scale.Y > Scale.Z ? 1 : 2 );
+					for( int32 k = 1; k < 3; k++ )
+					{
+						float Projected = ToLight[ (MaxComponent + k) % 3 ] / Scale[ MaxComponent ];
+						float Quantized = ( FMath::RoundToFloat( Projected * (0.5f * CubemapSize) - 0.5f ) + 0.5f ) / (0.5f * CubemapSize);
+						ToLight[ (MaxComponent + k) % 3 ] = Quantized * Scale[ MaxComponent ];
+					}
+					Origin = ToLight + View.ViewMatrices.ViewOrigin;
+				
+					FLinearColor Color( ColorAndFalloffExponent );
+
+					if( Proxy->IsInverseSquared() )
+					{
+						float LightRadiusMask = FMath::Square( 1.0f - FMath::Square( DistanceSqr * FMath::Square( PositionAndInvRadius.W ) ) );
+						Color *= LightRadiusMask;
+						
+						// Correction for lumen units
+						Color *= 16.0f;
+					}
+					else
+					{
+						// Remove inverse square falloff
+						Color *= DistanceSqr + 1.0f;
+
+						// Apply falloff
+						Color *= FMath::Pow( 1.0f - DistanceSqr * FMath::Square( PositionAndInvRadius.W ), ColorAndFalloffExponent.W ); 
+					}
+
+					// Spot falloff
+					FVector L = ToLight.GetSafeNormal();
+					Color *= FMath::Square( FMath::Clamp( ( (L | Direction) - SpotAngles.X ) * SpotAngles.Y, 0.0f, 1.0f ) );
+
+					// Scale by visible area
+					Color /= PI * FMath::Square( SourceRadius );
+				
+					// Always opaque
+					Color.A = 1.0f;
+				
+					FViewElementPDI LightPDI( &View, NULL );
+					FMaterialRenderProxy* const ColoredMeshInstance = new(FMemStack::Get()) FColoredMaterialRenderProxy( GEngine->DebugMeshMaterial->GetRenderProxy(false), Color );
+					DrawSphere( &LightPDI, Origin, FVector( SourceRadius, SourceRadius, SourceRadius ), 8, 6, ColoredMeshInstance, SDPG_World );
+				}
+			}
 		}
 	}
 
