Heightfield GI uses baked diffuse color textures from Landscape

diff --git a/Engine/Shaders/HeightfieldLighting.usf b/Engine/Shaders/HeightfieldLighting.usf
--- a/Engine/Shaders/HeightfieldLighting.usf
+++ b/Engine/Shaders/HeightfieldLighting.usf
@@ -19,6 +19,8 @@
 
 Texture2D HeightfieldTexture;
 SamplerState HeightfieldSampler;
+Texture2D DiffuseColorTexture;
+SamplerState DiffuseColorSampler;
 uint NumHeightfields;
 Buffer<float4> HeightfieldDescriptions;
 
@@ -123,6 +125,7 @@ float3 GetHeightfieldWorldPosition(uint HeightfieldIndex, float2 LocalPosition,
 
 Texture2D GlobalHeightfieldTexture;
 Texture2D GlobalNormalTexture;
+Texture2D GlobalDiffuseColorTexture;
 SamplerState GlobalHeightfieldSampler;
 
 float3 GetGlobalHeightfieldNormal(float2 GlobalUV)
@@ -132,6 +135,12 @@ float3 GetGlobalHeightfieldNormal(float2 GlobalUV)
 	return float3(RescaledNormal, sqrt(max(1.0 - dot(RescaledNormal, RescaledNormal), 0.0)));
 }
 
+float3 GetGlobalDiffuseColor(float2 GlobalUV)
+{
+	float3 EncodedDiffuseColor = Texture2DSampleLevel(GlobalDiffuseColorTexture, GlobalHeightfieldSampler, GlobalUV, 0).xyz;
+	return EncodedDiffuseColor * EncodedDiffuseColor;
+}
+
 // Returns height in local space
 float GetGlobalHeightfieldHeight(float2 GlobalUV)
 {
@@ -203,11 +212,15 @@ void HeightfieldSubsectionQuadVS(
 void InitializeHeightfieldsPS(
 	in FHeightfieldSubsectionQuadVertexOutput Input, 
 	out float4 OutHeight : SV_Target0,
-	out float4 OutNormal : SV_Target1)
+	out float4 OutNormal : SV_Target1,
+	out float4 OutDiffuseColor : SV_Target2)
 {
 	float4 PackedHeightAndNormal = Texture2DSampleLevel(HeightfieldTexture, HeightfieldSampler, Input.ComponentHeightfieldUV, 0);
 	OutHeight = DecodePackedTwoChannelValue(PackedHeightAndNormal.xy) / 65535.0f;
 	OutNormal = float4(PackedHeightAndNormal.zw, 0, 0);
+
+	float3 DiffuseColor = Texture2DSampleLevel(DiffuseColorTexture, DiffuseColorSampler, Input.ComponentHeightfieldUV, 0).xyz;
+	OutDiffuseColor = float4(sqrt(DiffuseColor), 0);
 }
 
 struct FHeightfieldComponentQuadVertexOutput
@@ -348,9 +361,9 @@ void LightHeightfieldsPS(
 
 	float Shadowing = Texture2DSampleLevel(HeightfieldShadowing, HeightfieldShadowingSampler, Input.GlobalUV, 0).x;
 
-	float3 DiffuseColor = .5f;
+	float3 DiffuseColor = GetGlobalDiffuseColor(Input.GlobalUV);
 	float3 Lighting = (Shadowing * GreyScaleLightFunction * saturate(dot(-LightDirection, WorldHeightfieldNormal))) * LightColor;
-
+	 
 	float3 SkylightDiffuse = GetSkySHDiffuse(WorldHeightfieldNormal) * View.SkyLightColor.rgb;
 	Lighting += SkylightDiffuse;
 
@@ -562,7 +575,7 @@ float3 GatherHeightfieldLighting(
 			float2 SampleUV = MinGatherUV + float2(XOffset, YOffset);
 
 			//@todo - reformulate loop to avoid center region
-			//@todo - try alternate access pattern with all threads reading from the same x row at once
+			//@todo - square access pattern to better match 2d texture cache tiles
 			if (!bUseExcludeRegion || any(SampleUV < ExcludeMinGatherUV || SampleUV > ExcludeMinGatherUV + ExcludeGatherUVSize))
 			{
 				float2 SampleLocalPosition = (SampleUV - GlobalUVScaleBias.zw) / GlobalUVScaleBias.xy;
@@ -630,7 +643,8 @@ void CalculateHeightfieldIrradianceCS(
 				#define VISUALIZE_HEIGHTFIELD_LIGHTING 0
 				#if VISUALIZE_HEIGHTFIELD_LIGHTING
 
-					Lighting = SampleHeightfieldLighting(GlobalUV) / FINAL_GATHER_THREADGROUP_SIZE;
+					//Lighting = SampleHeightfieldLighting(GlobalUV) / FINAL_GATHER_THREADGROUP_SIZE;
+					Lighting = GetGlobalDiffuseColor(GlobalUV) / FINAL_GATHER_THREADGROUP_SIZE;
 					//Lighting = frac(WorldHeightfieldShadingPosition / 1000);
 					//Lighting = abs(WorldHeightfieldShadingPosition.z - WorldPosition.z) / 100.0f;
 					//Lighting = GetGlobalHeightfieldHeight(GlobalUV) / 100000.0f;
@@ -662,7 +676,7 @@ void CalculateHeightfieldIrradianceCS(
 						InnerLightTransferDistance, 
 						ThreadIndex,
 						false);
-					
+
 					Lighting += GatherHeightfieldLighting(
 						WorldPosition, 
 						WorldNormal, 
diff --git a/Engine/Source/Runtime/Engine/Public/PrimitiveSceneProxy.h b/Engine/Source/Runtime/Engine/Public/PrimitiveSceneProxy.h
--- a/Engine/Source/Runtime/Engine/Public/PrimitiveSceneProxy.h
+++ b/Engine/Source/Runtime/Engine/Public/PrimitiveSceneProxy.h
@@ -229,7 +229,7 @@ public:
 		BoundsSurfaceArea = 0;
 	}
 
-	virtual void GetHeightfieldRepresentation(UTexture2D*& OutHeightmapTexture, FHeightfieldComponentDescription& OutDescription)
+	virtual void GetHeightfieldRepresentation(UTexture2D*& OutHeightmapTexture, UTexture2D*& OutDiffuseColorTexture, FHeightfieldComponentDescription& OutDescription)
 	{
 		OutHeightmapTexture = NULL;
 	}
diff --git a/Engine/Source/Runtime/Landscape/Private/LandscapeRender.cpp b/Engine/Source/Runtime/Landscape/Private/LandscapeRender.cpp
--- a/Engine/Source/Runtime/Landscape/Private/LandscapeRender.cpp
+++ b/Engine/Source/Runtime/Landscape/Private/LandscapeRender.cpp
@@ -547,6 +547,7 @@ FLandscapeComponentSceneProxy::FLandscapeComponentSceneProxy(ULandscapeComponent
 	, WeightmapTextures(InComponent->WeightmapTextures)
 	, NumWeightmapLayerAllocations(InComponent->WeightmapLayerAllocations.Num())
 	, NormalmapTexture(InComponent->HeightmapTexture)
+	, BaseColorForGITexture(InComponent->GIBakedBaseColorTexture)
 	, HeightmapScaleBias(InComponent->HeightmapScaleBias)
 	, XYOffsetmapTexture(InComponent->XYOffsetmapTexture)
 	, SharedBuffersKey(0)
@@ -2594,9 +2595,10 @@ void FLandscapeComponentSceneProxy::ChangeLODDistanceFactor_RenderThread(float I
 	LODDistance = InLODDistanceFactor;
 }
 
-void FLandscapeComponentSceneProxy::GetHeightfieldRepresentation(UTexture2D*& OutHeightmapTexture, FHeightfieldComponentDescription& OutDescription)
+void FLandscapeComponentSceneProxy::GetHeightfieldRepresentation(UTexture2D*& OutHeightmapTexture, UTexture2D*& OutDiffuseColorTexture, FHeightfieldComponentDescription& OutDescription)
 {
 	OutHeightmapTexture = HeightmapTexture;
+	OutDiffuseColorTexture = BaseColorForGITexture;
 	OutDescription.HeightfieldScaleBias = HeightmapScaleBias;
 
 	OutDescription.MinMaxUV = FVector4(
diff --git a/Engine/Source/Runtime/Landscape/Public/LandscapeRender.h b/Engine/Source/Runtime/Landscape/Public/LandscapeRender.h
--- a/Engine/Source/Runtime/Landscape/Public/LandscapeRender.h
+++ b/Engine/Source/Runtime/Landscape/Public/LandscapeRender.h
@@ -521,6 +521,7 @@ protected:
 	TArray<UTexture2D*> WeightmapTextures;
 	int8 NumWeightmapLayerAllocations;
 	UTexture2D* NormalmapTexture; // PC : Heightmap, Mobile : Weightmap
+	UTexture2D* BaseColorForGITexture;
 	FVector4 HeightmapScaleBias;
 	float HeightmapSubsectionOffsetU;
 	float HeightmapSubsectionOffsetV;
@@ -588,7 +589,7 @@ public:
 	// FLandcapeSceneProxy
 	void ChangeLODDistanceFactor_RenderThread(float InLODDistanceFactor);
 
-	virtual void GetHeightfieldRepresentation(UTexture2D*& OutHeightmapTexture, FHeightfieldComponentDescription& OutDescription) override;
+	virtual void GetHeightfieldRepresentation(UTexture2D*& OutHeightmapTexture, UTexture2D*& OutDiffuseColorTexture, FHeightfieldComponentDescription& OutDescription) override;
 };
 
 class FLandscapeDebugMaterialRenderProxy : public FMaterialRenderProxy
diff --git a/Engine/Source/Runtime/Renderer/Private/HeightfieldLighting.cpp b/Engine/Source/Runtime/Renderer/Private/HeightfieldLighting.cpp
--- a/Engine/Source/Runtime/Renderer/Private/HeightfieldLighting.cpp
+++ b/Engine/Source/Runtime/Renderer/Private/HeightfieldLighting.cpp
@@ -95,21 +95,30 @@ void FHeightfieldLightingAtlas::InitDynamicRHI()
 
 		FPooledRenderTargetDesc Desc3(FPooledRenderTargetDesc::Create2DDesc(
 			AtlasSize,
-			PF_G8, 
+			PF_R8G8B8A8, 
 			TexCreate_None,
 			TexCreate_RenderTargetable,
 			false));
 
-		GRenderTargetPool.FindFreeElement(Desc3, DirectionalLightShadowing, TEXT("HeightfieldShadowingAtlas"));
+		GRenderTargetPool.FindFreeElement(Desc3, DiffuseColor, TEXT("DiffuseColorAtlas"));
 
 		FPooledRenderTargetDesc Desc4(FPooledRenderTargetDesc::Create2DDesc(
+			AtlasSize,
+			PF_G8, 
+			TexCreate_None,
+			TexCreate_RenderTargetable,
+			false));
+
+		GRenderTargetPool.FindFreeElement(Desc4, DirectionalLightShadowing, TEXT("HeightfieldShadowingAtlas"));
+
+		FPooledRenderTargetDesc Desc5(FPooledRenderTargetDesc::Create2DDesc(
 			AtlasSize,
 			PF_FloatR11G11B10, 
 			TexCreate_None,
 			TexCreate_RenderTargetable,
 			false));
 
-		GRenderTargetPool.FindFreeElement(Desc4, Lighting, TEXT("HeightfieldLightingAtlas"));
+		GRenderTargetPool.FindFreeElement(Desc5, Lighting, TEXT("HeightfieldLightingAtlas"));
 	}
 }
 
@@ -117,6 +126,7 @@ void FHeightfieldLightingAtlas::ReleaseDynamicRHI()
 {
 	GRenderTargetPool.FreeUnusedResource(Height);
 	GRenderTargetPool.FreeUnusedResource(Normal);
+	GRenderTargetPool.FreeUnusedResource(DiffuseColor);
 	GRenderTargetPool.FreeUnusedResource(DirectionalLightShadowing);
 	GRenderTargetPool.FreeUnusedResource(Lighting);
 }
@@ -172,25 +182,40 @@ public:
 	{
 		HeightfieldTexture.Bind(ParameterMap,TEXT("HeightfieldTexture"));
 		HeightfieldSampler.Bind(ParameterMap,TEXT("HeightfieldSampler"));
+		DiffuseColorTexture.Bind(ParameterMap,TEXT("DiffuseColorTexture"));
+		DiffuseColorSampler.Bind(ParameterMap,TEXT("DiffuseColorSampler"));
 	}
 
 	friend FArchive& operator<<(FArchive& Ar,FHeightfieldTextureParameters& Parameters)
 	{
 		Ar << Parameters.HeightfieldTexture;
 		Ar << Parameters.HeightfieldSampler;
+		Ar << Parameters.DiffuseColorTexture;
+		Ar << Parameters.DiffuseColorSampler;
 		return Ar;
 	}
 
 	template<typename ShaderRHIParamRef>
-	void Set(FRHICommandList& RHICmdList, const ShaderRHIParamRef ShaderRHI, UTexture2D* HeightfieldTextureValue)
+	void Set(FRHICommandList& RHICmdList, const ShaderRHIParamRef ShaderRHI, UTexture2D* HeightfieldTextureValue, UTexture2D* DiffuseColorTextureValue)
 	{
 		//@todo - shouldn't filter the heightfield, it's packed
-		SetTextureParameter(RHICmdList, ShaderRHI, HeightfieldTexture, HeightfieldTexture, TStaticSamplerState<SF_Bilinear>::GetRHI(), HeightfieldTextureValue->Resource->TextureRHI);
+		SetTextureParameter(RHICmdList, ShaderRHI, HeightfieldTexture, HeightfieldSampler, TStaticSamplerState<SF_Bilinear>::GetRHI(), HeightfieldTextureValue->Resource->TextureRHI);
+
+		if (DiffuseColorTextureValue)
+		{
+			SetTextureParameter(RHICmdList, ShaderRHI, DiffuseColorTexture, DiffuseColorSampler, TStaticSamplerState<SF_Bilinear>::GetRHI(), DiffuseColorTextureValue->Resource->TextureRHI);
+		}
+		else
+		{
+			SetTextureParameter(RHICmdList, ShaderRHI, DiffuseColorTexture, DiffuseColorSampler, TStaticSamplerState<SF_Bilinear>::GetRHI(), GBlackTexture->TextureRHI);
+		}
 	}
 
 private:
 	FShaderResourceParameter HeightfieldTexture;
 	FShaderResourceParameter HeightfieldSampler;
+	FShaderResourceParameter DiffuseColorTexture;
+	FShaderResourceParameter DiffuseColorSampler;
 };
 
 class FSubsectionHeightfieldDescriptionParameters
@@ -327,12 +352,12 @@ public:
 		HeightfieldTextureParameters.Bind(Initializer.ParameterMap);
 	}
 
-	void SetParameters(FRHICommandList& RHICmdList, const FSceneView& View, UTexture2D* HeightfieldTextureValue)
+	void SetParameters(FRHICommandList& RHICmdList, const FSceneView& View, UTexture2D* HeightfieldTextureValue, UTexture2D* DiffuseColorTextureValue)
 	{
 		const FPixelShaderRHIParamRef ShaderRHI = GetPixelShader();
 
 		FGlobalShader::SetParameters(RHICmdList, ShaderRHI, View);
-		HeightfieldTextureParameters.Set(RHICmdList, ShaderRHI, HeightfieldTextureValue);
+		HeightfieldTextureParameters.Set(RHICmdList, ShaderRHI, HeightfieldTextureValue, DiffuseColorTextureValue);
 	}
 	// FShader interface.
 	virtual bool Serialize(FArchive& Ar)
@@ -427,8 +452,9 @@ void FHeightfieldLightingViewInfo::SetupVisibleHeightfields(const FViewInfo& Vie
 				&& DistanceToPrimitiveSq < MaxDistanceSquared)
 			{
 				UTexture2D* HeightfieldTexture = NULL;
+				UTexture2D* DiffuseColorTexture = NULL;
 				FHeightfieldComponentDescription NewComponentDescription(HeightfieldPrimitive->Proxy->GetLocalToWorld());
-				HeightfieldPrimitive->Proxy->GetHeightfieldRepresentation(HeightfieldTexture, NewComponentDescription);
+				HeightfieldPrimitive->Proxy->GetHeightfieldRepresentation(HeightfieldTexture, DiffuseColorTexture, NewComponentDescription);
 
 				if (HeightfieldTexture && HeightfieldTexture->Resource->TextureRHI)
 				{
@@ -444,7 +470,7 @@ void FHeightfieldLightingViewInfo::SetupVisibleHeightfields(const FViewInfo& Vie
 						Heightfield.Rect.Union(NewComponentDescription.HeightfieldRect);
 					}
 
-					TArray<FHeightfieldComponentDescription>& ComponentDescriptions = Heightfield.ComponentDescriptions.FindOrAdd(HeightfieldTexture);
+					TArray<FHeightfieldComponentDescription>& ComponentDescriptions = Heightfield.ComponentDescriptions.FindOrAdd(FHeightfieldComponentTextures(HeightfieldTexture, DiffuseColorTexture));
 					ComponentDescriptions.Add(NewComponentDescription);
 				}
 			}
@@ -467,7 +493,7 @@ void FHeightfieldLightingViewInfo::SetupVisibleHeightfields(const FViewInfo& Vie
 				Heightfield.Rect.Min = FIntPoint::DivideAndRoundDown(Heightfield.Rect.Min, Heightfield.DownsampleFactor) * Heightfield.DownsampleFactor;
 				ExistingAtlas->InitializeForSize(Heightfield.DownsampledRect.Size());
 
-				for (TMap<UTexture2D*, TArray<FHeightfieldComponentDescription>>::TIterator It(Heightfield.ComponentDescriptions); It; ++It)
+				for (TMap<FHeightfieldComponentTextures, TArray<FHeightfieldComponentDescription>>::TIterator It(Heightfield.ComponentDescriptions); It; ++It)
 				{
 					TArray<FHeightfieldComponentDescription>& HeightfieldDescriptions = It.Value();
 
@@ -483,15 +509,16 @@ void FHeightfieldLightingViewInfo::SetupVisibleHeightfields(const FViewInfo& Vie
 					const FIntPoint LightingAtlasSize = ExistingAtlas->GetAtlasSize();
 					const FVector2D InvLightingAtlasSize(1.0f / LightingAtlasSize.X, 1.0f / LightingAtlasSize.Y);
 
-					FTextureRHIParamRef RenderTargets[2] =
+					FTextureRHIParamRef RenderTargets[3] =
 					{
 						ExistingAtlas->Height->GetRenderTargetItem().TargetableTexture,
-						ExistingAtlas->Normal->GetRenderTargetItem().TargetableTexture
+						ExistingAtlas->Normal->GetRenderTargetItem().TargetableTexture,
+						ExistingAtlas->DiffuseColor->GetRenderTargetItem().TargetableTexture
 					};
 
 					SetRenderTargets(RHICmdList, ARRAY_COUNT(RenderTargets), RenderTargets, FTextureRHIParamRef(), 0, NULL);
 
-					FLinearColor ClearColors[2] = {FLinearColor(0, 0, 0, 0), FLinearColor(0, 0, 0, 0)};
+					FLinearColor ClearColors[3] = {FLinearColor(0, 0, 0, 0), FLinearColor(0, 0, 0, 0), FLinearColor(0, 0, 0, 0)};
 					RHICmdList.ClearMRT(true, ARRAY_COUNT(ClearColors), ClearColors, false, 0, false, 0, FIntRect());
 
 					RHICmdList.SetViewport(0, 0, 0.0f, LightingAtlasSize.X, LightingAtlasSize.Y, 1.0f);
@@ -506,7 +533,7 @@ void FHeightfieldLightingViewInfo::SetupVisibleHeightfields(const FViewInfo& Vie
 					static FGlobalBoundShaderState BoundShaderState;
 					SetGlobalBoundShaderState(RHICmdList, View.GetFeatureLevel(), BoundShaderState, GScreenVertexDeclaration.VertexDeclarationRHI, *VertexShader, *PixelShader);
 
-					for (TMap<UTexture2D*, TArray<FHeightfieldComponentDescription>>::TConstIterator It(Heightfield.ComponentDescriptions); It; ++It)
+					for (TMap<FHeightfieldComponentTextures, TArray<FHeightfieldComponentDescription>>::TConstIterator It(Heightfield.ComponentDescriptions); It; ++It)
 					{
 						const TArray<FHeightfieldComponentDescription>& HeightfieldDescriptions = It.Value();
 
@@ -514,10 +541,8 @@ void FHeightfieldLightingViewInfo::SetupVisibleHeightfields(const FViewInfo& Vie
 						{
 							const int32 NumQuads = UploadSubsectionHeightfieldDescriptions(HeightfieldDescriptions, InvLightingAtlasSize, 1.0f / Heightfield.DownsampleFactor);
 
-							UTexture2D* HeightfieldTexture = It.Key();
-
 							VertexShader->SetParameters(RHICmdList, View);
-							PixelShader->SetParameters(RHICmdList, View, HeightfieldTexture);
+							PixelShader->SetParameters(RHICmdList, View, It.Key().HeightAndNormal, It.Key().DiffuseColor);
 
 							RHICmdList.DrawPrimitive(PT_TriangleList, 0, 2, NumQuads);
 						}
@@ -650,6 +675,7 @@ public:
 	{
 		GlobalHeightfieldTexture.Bind(ParameterMap,TEXT("GlobalHeightfieldTexture"));
 		GlobalNormalTexture.Bind(ParameterMap,TEXT("GlobalNormalTexture"));
+		GlobalDiffuseColorTexture.Bind(ParameterMap,TEXT("GlobalDiffuseColorTexture"));
 		GlobalHeightfieldSampler.Bind(ParameterMap,TEXT("GlobalHeightfieldSampler"));
 	}
 
@@ -657,6 +683,7 @@ public:
 	{
 		Ar << Parameters.GlobalHeightfieldTexture;
 		Ar << Parameters.GlobalNormalTexture;
+		Ar << Parameters.GlobalDiffuseColorTexture;
 		Ar << Parameters.GlobalHeightfieldSampler;
 		return Ar;
 	}
@@ -666,11 +693,13 @@ public:
 	{
 		SetTextureParameter(RHICmdList, ShaderRHI, GlobalHeightfieldTexture, GlobalHeightfieldSampler, TStaticSamplerState<SF_Bilinear>::GetRHI(), Atlas.Height->GetRenderTargetItem().ShaderResourceTexture);
 		SetTextureParameter(RHICmdList, ShaderRHI, GlobalNormalTexture, GlobalHeightfieldSampler, TStaticSamplerState<SF_Bilinear>::GetRHI(), Atlas.Normal->GetRenderTargetItem().ShaderResourceTexture);
+		SetTextureParameter(RHICmdList, ShaderRHI, GlobalDiffuseColorTexture, GlobalHeightfieldSampler, TStaticSamplerState<SF_Bilinear>::GetRHI(), Atlas.DiffuseColor->GetRenderTargetItem().ShaderResourceTexture);
 	}
 
 private:
 	FShaderResourceParameter GlobalHeightfieldTexture;
 	FShaderResourceParameter GlobalNormalTexture;
+	FShaderResourceParameter GlobalDiffuseColorTexture;
 	FShaderResourceParameter GlobalHeightfieldSampler;
 };
 
@@ -735,7 +764,6 @@ public:
 		: FGlobalShader(Initializer)
 	{
 		HeightfieldDescriptionParameters.Bind(Initializer.ParameterMap);
-		HeightfieldTextureParameters.Bind(Initializer.ParameterMap);
 		GlobalHeightfieldParameters.Bind(Initializer.ParameterMap);
 		WorldToShadow.Bind(Initializer.ParameterMap, TEXT("WorldToShadow"));
 		ShadowmapMinMax.Bind(Initializer.ParameterMap, TEXT("ShadowmapMinMax"));
@@ -745,13 +773,12 @@ public:
 		ShadowDepthTextureSampler.Bind(Initializer.ParameterMap, TEXT("ShadowDepthTextureSampler"));
 	}
 
-	void SetParameters(FRHICommandList& RHICmdList, const FSceneView& View, const FProjectedShadowInfo* ProjectedShadowInfo, UTexture2D* HeightfieldTextureValue, int32 NumHeightfieldsValue, const FHeightfieldLightingAtlas& Atlas)
+	void SetParameters(FRHICommandList& RHICmdList, const FSceneView& View, const FProjectedShadowInfo* ProjectedShadowInfo, int32 NumHeightfieldsValue, const FHeightfieldLightingAtlas& Atlas)
 	{
 		const FPixelShaderRHIParamRef ShaderRHI = GetPixelShader();
 
 		FGlobalShader::SetParameters(RHICmdList, ShaderRHI, View);
 		HeightfieldDescriptionParameters.Set(RHICmdList, ShaderRHI, NumHeightfieldsValue);
-		HeightfieldTextureParameters.Set(RHICmdList, ShaderRHI, HeightfieldTextureValue);
 		GlobalHeightfieldParameters.Set(RHICmdList, ShaderRHI, Atlas);
 
 		FVector4 ShadowmapMinMaxValue;
@@ -772,7 +799,6 @@ public:
 	{
 		bool bShaderHasOutdatedParameters = FGlobalShader::Serialize(Ar);
 		Ar << HeightfieldDescriptionParameters;
-		Ar << HeightfieldTextureParameters;
 		Ar << GlobalHeightfieldParameters;
 		Ar << WorldToShadow;
 		Ar << ShadowmapMinMax;
@@ -786,7 +812,6 @@ public:
 private:
 
 	FHeightfieldDescriptionParameters HeightfieldDescriptionParameters;
-	FHeightfieldTextureParameters HeightfieldTextureParameters;
 	FGlobalHeightfieldParameters GlobalHeightfieldParameters;
 	FShaderParameter WorldToShadow;
 	FShaderParameter ShadowmapMinMax;
@@ -846,7 +871,7 @@ void FHeightfieldLightingViewInfo::ComputeShadowMapShadowing(const FViewInfo& Vi
 			static FGlobalBoundShaderState BoundShaderState;
 			SetGlobalBoundShaderState(RHICmdList, View.GetFeatureLevel(), BoundShaderState, GScreenVertexDeclaration.VertexDeclarationRHI, *VertexShader, *PixelShader);
 
-			for (TMap<UTexture2D*, TArray<FHeightfieldComponentDescription>>::TConstIterator It(Heightfield.ComponentDescriptions); It; ++It)
+			for (TMap<FHeightfieldComponentTextures, TArray<FHeightfieldComponentDescription>>::TConstIterator It(Heightfield.ComponentDescriptions); It; ++It)
 			{
 				const TArray<FHeightfieldComponentDescription>& HeightfieldDescriptions = It.Value();
 
@@ -855,10 +880,8 @@ void FHeightfieldLightingViewInfo::ComputeShadowMapShadowing(const FViewInfo& Vi
 					//@todo - cull heightfield tiles with shadow bounds
 					UploadHeightfieldDescriptions(HeightfieldDescriptions, InvLightingAtlasSize, 1.0f / Heightfield.DownsampleFactor);
 
-					UTexture2D* HeightfieldTexture = It.Key();
-
 					VertexShader->SetParameters(RHICmdList, View, HeightfieldDescriptions.Num());
-					PixelShader->SetParameters(RHICmdList, View, ProjectedShadowInfo, HeightfieldTexture, HeightfieldDescriptions.Num(), Atlas);
+					PixelShader->SetParameters(RHICmdList, View, ProjectedShadowInfo, HeightfieldDescriptions.Num(), Atlas);
 
 					RHICmdList.DrawPrimitive(PT_TriangleList, 0, 2, HeightfieldDescriptions.Num());
 				}
@@ -1010,7 +1033,7 @@ void FHeightfieldLightingViewInfo::ComputeRayTracedShadowing(
 			static FGlobalBoundShaderState BoundShaderState;
 			SetGlobalBoundShaderState(RHICmdList, View.GetFeatureLevel(), BoundShaderState, GScreenVertexDeclaration.VertexDeclarationRHI, *VertexShader, *PixelShader);
 
-			for (TMap<UTexture2D*, TArray<FHeightfieldComponentDescription>>::TConstIterator It(Heightfield.ComponentDescriptions); It; ++It)
+			for (TMap<FHeightfieldComponentTextures, TArray<FHeightfieldComponentDescription>>::TConstIterator It(Heightfield.ComponentDescriptions); It; ++It)
 			{
 				const TArray<FHeightfieldComponentDescription>& HeightfieldDescriptions = It.Value();
 
@@ -1018,8 +1041,6 @@ void FHeightfieldLightingViewInfo::ComputeRayTracedShadowing(
 				{
 					UploadHeightfieldDescriptions(HeightfieldDescriptions, InvLightingAtlasSize, 1.0f / Heightfield.DownsampleFactor);
 
-					UTexture2D* HeightfieldTexture = It.Key();
-
 					VertexShader->SetParameters(RHICmdList, View, HeightfieldDescriptions.Num());
 					PixelShader->SetParameters(RHICmdList, View, ProjectedShadowInfo, HeightfieldDescriptions.Num(), Atlas, TileIntersectionResources, CulledObjectBuffers);
 
@@ -1171,7 +1192,7 @@ void FHeightfieldLightingViewInfo::ComputeLighting(const FViewInfo& View, FRHICo
 
 			FLocalBoundShaderState BoundShaderState = RHICmdList.BuildLocalBoundShaderState(GScreenVertexDeclaration.VertexDeclarationRHI, VertexShader->GetVertexShader(), FHullShaderRHIRef(), FDomainShaderRHIRef(), PixelShader->GetPixelShader(), FGeometryShaderRHIRef());
 
-			for (TMap<UTexture2D*, TArray<FHeightfieldComponentDescription>>::TConstIterator It(Heightfield.ComponentDescriptions); It; ++It)
+			for (TMap<FHeightfieldComponentTextures, TArray<FHeightfieldComponentDescription>>::TConstIterator It(Heightfield.ComponentDescriptions); It; ++It)
 			{
 				const TArray<FHeightfieldComponentDescription>& HeightfieldDescriptions = It.Value();
 
@@ -1179,8 +1200,6 @@ void FHeightfieldLightingViewInfo::ComputeLighting(const FViewInfo& View, FRHICo
 				{
 					UploadHeightfieldDescriptions(HeightfieldDescriptions, InvLightingAtlasSize, 1.0f / Heightfield.DownsampleFactor);
 
-					UTexture2D* HeightfieldTexture = It.Key();
-
 					RHICmdList.SetLocalBoundShaderState(BoundShaderState);
 
 					VertexShader->SetParameters(RHICmdList, View, HeightfieldDescriptions.Num());
@@ -1249,7 +1268,7 @@ public:
 
 		AOParameters.Set(RHICmdList, ShaderRHI, Parameters);
 		HeightfieldDescriptionParameters.Set(RHICmdList, ShaderRHI, NumHeightfieldsValue);
-		HeightfieldTextureParameters.Set(RHICmdList, ShaderRHI, HeightfieldTextureValue);
+		HeightfieldTextureParameters.Set(RHICmdList, ShaderRHI, HeightfieldTextureValue, NULL);
 
 		const FScene* Scene = (const FScene*)View.Family->Scene;
 		FSurfaceCacheResources& SurfaceCacheResources = *Scene->SurfaceCacheResources;
@@ -1363,7 +1382,7 @@ void FHeightfieldLightingViewInfo::ComputeOcclusionForSamples(
 		FSceneViewState* ViewState = (FSceneViewState*)View.State;
 
 		{
-			for (TMap<UTexture2D*, TArray<FHeightfieldComponentDescription>>::TConstIterator It(Heightfield.ComponentDescriptions); It; ++It)
+			for (TMap<FHeightfieldComponentTextures, TArray<FHeightfieldComponentDescription>>::TConstIterator It(Heightfield.ComponentDescriptions); It; ++It)
 			{
 				const TArray<FHeightfieldComponentDescription>& HeightfieldDescriptions = It.Value();
 
@@ -1371,7 +1390,7 @@ void FHeightfieldLightingViewInfo::ComputeOcclusionForSamples(
 				{
 					UploadHeightfieldDescriptions(HeightfieldDescriptions, FVector2D(1, 1), 1.0f / Heightfield.DownsampleFactor);
 
-					UTexture2D* HeightfieldTexture = It.Key();
+					UTexture2D* HeightfieldTexture = It.Key().HeightAndNormal;
 
 					{
 						TShaderMapRef<FCalculateHeightfieldOcclusionCS> ComputeShader(View.ShaderMap);
@@ -1640,7 +1659,7 @@ void FHeightfieldLightingViewInfo::ComputeIrradianceForSamples(
 
 			TArray<FHeightfieldComponentDescription> CombinedHeightfieldDescriptions;
 
-			for (TMap<UTexture2D*, TArray<FHeightfieldComponentDescription>>::TConstIterator It(Heightfield.ComponentDescriptions); It; ++It)
+			for (TMap<FHeightfieldComponentTextures, TArray<FHeightfieldComponentDescription>>::TConstIterator It(Heightfield.ComponentDescriptions); It; ++It)
 			{
 				const TArray<FHeightfieldComponentDescription>& HeightfieldDescriptions = It.Value();
 				CombinedHeightfieldDescriptions.Append(HeightfieldDescriptions);
diff --git a/Engine/Source/Runtime/Renderer/Private/HeightfieldLighting.h b/Engine/Source/Runtime/Renderer/Private/HeightfieldLighting.h
--- a/Engine/Source/Runtime/Renderer/Private/HeightfieldLighting.h
+++ b/Engine/Source/Runtime/Renderer/Private/HeightfieldLighting.h
@@ -12,6 +12,7 @@ public:
 
 	TRefCountPtr<IPooledRenderTarget> Height;
 	TRefCountPtr<IPooledRenderTarget> Normal;
+	TRefCountPtr<IPooledRenderTarget> DiffuseColor;
 	TRefCountPtr<IPooledRenderTarget> DirectionalLightShadowing;
 	TRefCountPtr<IPooledRenderTarget> Lighting;
 
@@ -31,6 +32,29 @@ private:
 	FIntPoint AtlasSize;
 };
 
+class FHeightfieldComponentTextures
+{
+public:
+
+	FHeightfieldComponentTextures(UTexture2D* InHeightAndNormal, UTexture2D* InDiffuseColor) :
+		HeightAndNormal(InHeightAndNormal),
+		DiffuseColor(InDiffuseColor)
+	{}
+
+	FORCEINLINE bool operator==(FHeightfieldComponentTextures Other) const
+	{
+		return HeightAndNormal == Other.HeightAndNormal && DiffuseColor == Other.DiffuseColor;
+	}
+
+	FORCEINLINE friend uint32 GetTypeHash(FHeightfieldComponentTextures ComponentTextures)
+	{
+		return GetTypeHash(ComponentTextures.HeightAndNormal);
+	}
+
+	UTexture2D* HeightAndNormal;
+	UTexture2D* DiffuseColor;
+};
+
 class FHeightfieldDescription
 {
 public:
@@ -38,7 +62,7 @@ public:
 	int32 DownsampleFactor;
 	FIntRect DownsampledRect;
 
-	TMap<UTexture2D*, TArray<FHeightfieldComponentDescription>> ComponentDescriptions;
+	TMap<FHeightfieldComponentTextures, TArray<FHeightfieldComponentDescription>> ComponentDescriptions;
 
 	FHeightfieldDescription() :
 		Rect(FIntRect(0, 0, 0, 0)),
