Tonemapper for mobile HDR mosaic point sampled instead of bilinear to correct precision issues
#android
#ue4

diff --git a/Engine/Source/Runtime/Renderer/Private/PostProcess/PostProcessTonemap.cpp b/Engine/Source/Runtime/Renderer/Private/PostProcess/PostProcessTonemap.cpp
--- a/Engine/Source/Runtime/Renderer/Private/PostProcess/PostProcessTonemap.cpp
+++ b/Engine/Source/Runtime/Renderer/Private/PostProcess/PostProcessTonemap.cpp
@@ -1578,7 +1578,16 @@ public:
 		
 		FGlobalShader::SetParameters(Context.RHICmdList, ShaderRHI, Context.View);
 
-		PostprocessParameter.SetPS(ShaderRHI, Context, TStaticSamplerState<SF_Bilinear,AM_Clamp,AM_Clamp,AM_Clamp>::GetRHI());
+		const uint32 ConfigBitmask = TonemapperConfBitmaskMobile[ConfigIndex];
+			
+		if (TonemapperIsDefined(ConfigBitmask, TonemapperMosaic))
+		{
+			PostprocessParameter.SetPS(ShaderRHI, Context, TStaticSamplerState<SF_Point, AM_Clamp, AM_Clamp, AM_Clamp>::GetRHI());
+		}
+		else
+		{
+			PostprocessParameter.SetPS(ShaderRHI, Context, TStaticSamplerState<SF_Bilinear, AM_Clamp, AM_Clamp, AM_Clamp>::GetRHI());
+		}
 			
 		SetShaderValue(Context.RHICmdList, ShaderRHI, OverlayColor, Context.View.OverlayColor);
 
