Update Morpheus to SDK 2.0

diff --git a/Engine/Source/Programs/UnrealBuildTool/Windows/UEBuildWindows.cs b/Engine/Source/Programs/UnrealBuildTool/Windows/UEBuildWindows.cs
--- a/Engine/Source/Programs/UnrealBuildTool/Windows/UEBuildWindows.cs
+++ b/Engine/Source/Programs/UnrealBuildTool/Windows/UEBuildWindows.cs
@@ -460,6 +460,10 @@ namespace UnrealBuildTool
             if (File.Exists(MorpheusShaderPath))
             {
                 InBuildTarget.GlobalCompileEnvironment.Config.Definitions.Add("HAS_MORPHEUS=1");
+
+				//on PS4 the SDK now handles distortion correction.  On PC we will still have to handle it manually,
+				//but we require SDK changes before we can get the required data.  For the moment, no platform does in-engine morpheus distortion
+				InBuildTarget.GlobalCompileEnvironment.Config.Definitions.Add("MORPHEUS_ENGINE_DISTORTION=0");				
             }
 
             if (InBuildTarget.Rules != null)
diff --git a/Engine/Source/Runtime/Renderer/Private/PostProcess/PostProcessMorpheus.cpp b/Engine/Source/Runtime/Renderer/Private/PostProcess/PostProcessMorpheus.cpp
--- a/Engine/Source/Runtime/Renderer/Private/PostProcess/PostProcessMorpheus.cpp
+++ b/Engine/Source/Runtime/Renderer/Private/PostProcess/PostProcessMorpheus.cpp
@@ -16,7 +16,7 @@
 
 DEFINE_LOG_CATEGORY_STATIC(LogMorpheusHMDPostProcess, All, All);
 
-#if HAS_MORPHEUS
+#if MORPHEUS_ENGINE_DISTORTION
 
 /** Encapsulates the post processing HMD distortion and correction pixel shader. */
 class FPostProcessMorpheusPS : public FGlobalShader
@@ -99,7 +99,10 @@ public:
 			if (StereoPass == eSSP_LEFT_EYE)
 			{
 				FTexture* TextureLeft = HMDDevice->GetDistortionTextureLeft();
-				SetTextureParameter(Context.RHICmdList, ShaderRHI, DistortionTextureParam, DistortionTextureSampler, TextureLeft->SamplerStateRHI, TextureLeft->TextureRHI);
+				if (TextureLeft)
+				{
+					SetTextureParameter(Context.RHICmdList, ShaderRHI, DistortionTextureParam, DistortionTextureSampler, TextureLeft->SamplerStateRHI, TextureLeft->TextureRHI);
+				}
 				SetShaderValue(Context.RHICmdList, ShaderRHI, TextureScale, HMDDevice->GetTextureScaleLeft());
 				SetShaderValue(Context.RHICmdList, ShaderRHI, TextureOffset, HMDDevice->GetTextureOffsetLeft());
 				SetShaderValue(Context.RHICmdList, ShaderRHI, TextureUVOffset, 0.0f);
@@ -107,7 +110,10 @@ public:
 			else
 			{
 				FTexture* TextureRight = HMDDevice->GetDistortionTextureRight();
-				SetTextureParameter(Context.RHICmdList, ShaderRHI, DistortionTextureParam, DistortionTextureSampler, TextureRight->SamplerStateRHI, TextureRight->TextureRHI);
+				if (TextureRight)
+				{
+					SetTextureParameter(Context.RHICmdList, ShaderRHI, DistortionTextureParam, DistortionTextureSampler, TextureRight->SamplerStateRHI, TextureRight->TextureRHI);
+				}
 				SetShaderValue(Context.RHICmdList, ShaderRHI, TextureScale, HMDDevice->GetTextureScaleRight());
 				SetShaderValue(Context.RHICmdList, ShaderRHI, TextureOffset, HMDDevice->GetTextureOffsetRight());
 				SetShaderValue(Context.RHICmdList, ShaderRHI, TextureUVOffset, -0.5f);
@@ -244,4 +250,4 @@ FPooledRenderTargetDesc FRCPassPostProcessMorpheus::ComputeOutputDesc(EPassOutpu
 	return Ret;
 }
 
-#endif // HAS_MORPHEUS
\ No newline at end of file
+#endif // MORPHEUS_ENGINE_DISTORTION
\ No newline at end of file
diff --git a/Engine/Source/Runtime/Renderer/Private/PostProcess/PostProcessMorpheus.h b/Engine/Source/Runtime/Renderer/Private/PostProcess/PostProcessMorpheus.h
--- a/Engine/Source/Runtime/Renderer/Private/PostProcess/PostProcessMorpheus.h
+++ b/Engine/Source/Runtime/Renderer/Private/PostProcess/PostProcessMorpheus.h
@@ -8,7 +8,7 @@
 
 #include "RenderingCompositionGraph.h"
 
-#if HAS_MORPHEUS
+#if MORPHEUS_ENGINE_DISTORTION
 
 // derives from TRenderingCompositePassBase<InputCount, OutputCount>
 // ePId_Input0: SceneColor
diff --git a/Engine/Source/Runtime/Renderer/Private/PostProcess/PostProcessing.cpp b/Engine/Source/Runtime/Renderer/Private/PostProcess/PostProcessing.cpp
--- a/Engine/Source/Runtime/Renderer/Private/PostProcess/PostProcessing.cpp
+++ b/Engine/Source/Runtime/Renderer/Private/PostProcess/PostProcessing.cpp
@@ -1420,7 +1420,7 @@ void FPostProcessing::Process(FRHICommandListImmediate& RHICmdList, FViewInfo& V
 			{
 				Node = Context.Graph.RegisterPass(new FRCPassPostProcessHMD());
 			}
-#if HAS_MORPHEUS
+#if MORPHEUS_ENGINE_DISTORTION
 			else if(DeviceType == EHMDDeviceType::DT_Morpheus)
 			{
 				Node = Context.Graph.RegisterPass(new FRCPassPostProcessMorpheus());
