Fix a factor of 2 off in radius in CircleDOF (now validated by real life shots).

diff --git a/Engine/Source/Runtime/Renderer/Private/PostProcess/PostProcessDOF.cpp b/Engine/Source/Runtime/Renderer/Private/PostProcess/PostProcessDOF.cpp
--- a/Engine/Source/Runtime/Renderer/Private/PostProcess/PostProcessDOF.cpp
+++ b/Engine/Source/Runtime/Renderer/Private/PostProcess/PostProcessDOF.cpp
@@ -387,8 +387,8 @@ static float CircleDofCoc(const FRenderingCompositePassContext& Context)
 	float Width = (float)Context.GetViewport().Size().X;
 	Radius = Radius * Width * (1.0f/24.576f);
 
-	// Convert to half resolution (algorithm radius is at half resolution).
-	Radius *= 0.5;
+	// Convert diameter to radius at half resolution (algorithm radius is at half resolution).
+	Radius *= 0.25f;
 
 	// Limit to algorithm max size.
 	if(Radius > 6.0f) 
