CircleDOF: Ok, removing limiters. Allowing artists to use settings which won't look good currently.

diff --git a/Engine/Source/Runtime/Renderer/Private/PostProcess/PostProcessDOF.cpp b/Engine/Source/Runtime/Renderer/Private/PostProcess/PostProcessDOF.cpp
--- a/Engine/Source/Runtime/Renderer/Private/PostProcess/PostProcessDOF.cpp
+++ b/Engine/Source/Runtime/Renderer/Private/PostProcess/PostProcessDOF.cpp
@@ -390,11 +390,14 @@ static float CircleDofCoc(const FRenderingCompositePassContext& Context)
 	// Convert diameter to radius at half resolution (algorithm radius is at half resolution).
 	Radius *= 0.25f;
 
-	// Limit to algorithm max size.
-	if(Radius > 6.0f) 
-	{
-		Radius = 6.0f; 
-	}
+	// Comment out for now, allowing settings which the algorithm cannot cleanly do.
+	#if 0
+		// Limit to algorithm max size.
+		if(Radius > 6.0f) 
+		{
+			Radius = 6.0f; 
+		}
+	#endif
 	return Radius;
 }
 
