Prototype support for HDR displays in the prototype tonemapper.

diff --git a/Engine/Source/Runtime/Core/Private/HAL/ConsoleManager.cpp b/Engine/Source/Runtime/Core/Private/HAL/ConsoleManager.cpp
--- a/Engine/Source/Runtime/Core/Private/HAL/ConsoleManager.cpp
+++ b/Engine/Source/Runtime/Core/Private/HAL/ConsoleManager.cpp
@@ -1998,6 +1998,16 @@ static TAutoConsoleVariable<int32> CVarTonemapperPhoto(
 	TEXT("1: use new photo post"),
 	ECVF_Scalability | ECVF_RenderThreadSafe);
 
+static TAutoConsoleVariable<int32> CVarTonemapperHDR(
+	TEXT("r.TonemapperHDR"),
+	1,
+	TEXT("Make tonemapper work with HDR display.\n")
+	TEXT("Requires 'r.TonemapperPhoto 1'.\n")
+	TEXT("1: standard dynamic range\n")
+	TEXT("#: high dynamic range (#=2 for 1 stop more, #=4 for 2 stops more, #=8 for 3 stops more and so on"),
+	ECVF_Scalability | ECVF_RenderThreadSafe);
+
+
 static TAutoConsoleVariable<int32> CVarTonemapper709(
 	TEXT("r.Tonemapper709"),
 	0,
diff --git a/Engine/Source/Runtime/Renderer/Private/PostProcess/PostProcessTonemap.cpp b/Engine/Source/Runtime/Renderer/Private/PostProcess/PostProcessTonemap.cpp
--- a/Engine/Source/Runtime/Renderer/Private/PostProcess/PostProcessTonemap.cpp
+++ b/Engine/Source/Runtime/Renderer/Private/PostProcess/PostProcessTonemap.cpp
@@ -1038,6 +1038,22 @@ void FilmPostSetConstantsPhoto(FVector4* RESTRICT const Constants, const FPostPr
 	// Factor in the scale to ensure gamma does not change the 18% midtone.
 	OutTint *= PhotoMid/FMath::Pow(PhotoMid, PhotoGamma);
 
+	// HDR display support provides a secondary exposure adjustment 
+	// which is factored into the tint mulitply in the {Lift, gamma, tint} processing.
+	// This is prior to over-exposure crosstalk (desaturate overexposed color),
+	// and prior to highlight compression.
+	// If an LDR display goes from 0-1 this multiply will be 1.
+	// If an HDR display goes from 0-16 this multiply will be 1/16 (the display is 16x brigher).
+	{
+		static TConsoleVariableData<int32>* CVar = IConsoleManager::Get().FindTConsoleVariableDataInt(TEXT("r.TonemapperHDR"));
+		float Value = CVar->GetValueOnRenderThread();
+		if(Value < 1.0f)
+		{
+			Value = 1.0f;
+		}
+		OutTint *= 1.0f/Value;
+	}
+
 	Constants[0] = FVector4(OutMatrixR, OutLift);
 	Constants[1] = FVector4(OutMatrixG, OutLiftScale);
 	Constants[2] = FVector4(OutMatrixB, PhotoGamma); 
