Adding a r.TonemapperGamma option for displays with fixed gamma colorspaces.

diff --git a/Engine/Shaders/PostProcessTonemap.usf b/Engine/Shaders/PostProcessTonemap.usf
--- a/Engine/Shaders/PostProcessTonemap.usf
+++ b/Engine/Shaders/PostProcessTonemap.usf
@@ -287,9 +287,15 @@ void MainPS(
 				// Didn't profile yet if the branching version would be faster (different linear segment).
 				half3 TonemappedColor = LinearTo709Branchless(FilmColor);
 			#else
-				// Apply conversion to sRGB (this must be an exact sRGB conversion else darks are bad).
-				// Branching is faster than branchless on AMD on PC.
-				half3 TonemappedColor = LinearToSrgbBranching(FilmColor);
+				#if USE_GAMMA
+					// This is different than the prior "gamma" curve adjustment (but reusing the variable).
+					// For displays set to a gamma colorspace.
+					half3 TonemappedColor = pow(FilmColor.rgb, InverseGamma.z);
+				#else
+					// Apply conversion to sRGB (this must be an exact sRGB conversion else darks are bad).
+					// Branching is faster than branchless on AMD on PC.
+					half3 TonemappedColor = LinearToSrgbBranching(FilmColor);
+				#endif
 			#endif
 		#endif
 
diff --git a/Engine/Shaders/TonemapCommon.usf b/Engine/Shaders/TonemapCommon.usf
--- a/Engine/Shaders/TonemapCommon.usf
+++ b/Engine/Shaders/TonemapCommon.usf
@@ -5,7 +5,7 @@
 =============================================================================*/
 
 // usually 1/2.2, the .y is used for inverse gamma when "gamma only" mode is not used
-half2 InverseGamma;
+half3 InverseGamma;
 
 // Film tonal and color control.
 half4 ColorMatrixR_ColorCurveCd1;
diff --git a/Engine/Source/Runtime/Core/Private/HAL/ConsoleManager.cpp b/Engine/Source/Runtime/Core/Private/HAL/ConsoleManager.cpp
--- a/Engine/Source/Runtime/Core/Private/HAL/ConsoleManager.cpp
+++ b/Engine/Source/Runtime/Core/Private/HAL/ConsoleManager.cpp
@@ -1998,7 +1998,7 @@ static TAutoConsoleVariable<int32> CVarTonemapperPhoto(
 	TEXT("1: use new photo post"),
 	ECVF_Scalability | ECVF_RenderThreadSafe);
 
-static TAutoConsoleVariable<int32> CVarTonemapperHDR(
+static TAutoConsoleVariable<float> CVarTonemapperHDR(
 	TEXT("r.TonemapperHDR"),
 	1,
 	TEXT("Make tonemapper work with HDR display.\n")
@@ -2007,6 +2007,12 @@ static TAutoConsoleVariable<int32> CVarTonemapperHDR(
 	TEXT("#: high dynamic range (#=2 for 1 stop more, #=4 for 2 stops more, #=8 for 3 stops more and so on"),
 	ECVF_Scalability | ECVF_RenderThreadSafe);
 
+static TAutoConsoleVariable<float> CVarTonemapperGamma(
+	TEXT("r.TonemapperGamma"),
+	0,
+	TEXT("0: don't use\n")
+	TEXT("#: used fixed gamma # instead of sRGB or Rec709 transform"),
+	ECVF_Scalability | ECVF_RenderThreadSafe);
 
 static TAutoConsoleVariable<int32> CVarTonemapper709(
 	TEXT("r.Tonemapper709"),
diff --git a/Engine/Source/Runtime/Renderer/Private/PostProcess/PostProcessTonemap.cpp b/Engine/Source/Runtime/Renderer/Private/PostProcess/PostProcessTonemap.cpp
--- a/Engine/Source/Runtime/Renderer/Private/PostProcess/PostProcessTonemap.cpp
+++ b/Engine/Source/Runtime/Renderer/Private/PostProcess/PostProcessTonemap.cpp
@@ -39,6 +39,7 @@ typedef enum {
 	TonemapperMsaa              = (1<<15),
 	TonemapperPhoto             = (1<<16),
 	Tonemapper709               = (1<<17),
+	TonemapperGamma             = (1<<18),
 } TonemapperOption;
 
 // Tonemapper option cost (0 = no cost, 255 = max cost).
@@ -63,6 +64,7 @@ static uint8 TonemapperCostTab[] = {
 	1, //TonemapperMsaa
 	1, //TonemapperPhoto
 	1, //Tonemapper709
+	1, //TonemapperGamma
 };
 
 // Edit the following to add and remove configurations.
@@ -201,8 +203,9 @@ static uint32 TonemapperConfBitmaskPC[22] = {
 	TonemapperGrainQuantization +
 	TonemapperVignette +
 	TonemapperVignetteColor +
+	TonemapperColorFringe +
 	TonemapperPhoto +
-	Tonemapper709 +
+	TonemapperGamma +
 	0,
 
 	TonemapperBloom +
@@ -237,8 +240,9 @@ static uint32 TonemapperConfBitmaskPC[22] = {
 	TonemapperBloom +
 	TonemapperVignette +
 	TonemapperVignetteColor +
+	TonemapperColorFringe +
 	TonemapperPhoto +
-	Tonemapper709 +
+	TonemapperGamma +
 	0,
 
 	TonemapperBloom +
@@ -713,6 +717,17 @@ static uint32 TonemapperGenerateBitmaskPC(const FViewInfo* RESTRICT View, bool b
 		}
 	}
 
+	{
+		static TConsoleVariableData<float>* CVar = IConsoleManager::Get().FindTConsoleVariableDataFloat(TEXT("r.TonemapperGamma"));
+		float Value = CVar->GetValueOnRenderThread();
+		if(Value > 0.0f)
+		{
+			// Remove settings not compatible.
+			Bitmask &= ~(TonemapperColorMatrix|TonemapperShadowTint|TonemapperContrast|TonemapperColorGrading|TonemapperMsaa);
+			Bitmask |= TonemapperPhoto | TonemapperGamma;
+		}
+	}
+
 	{
 		static TConsoleVariableData<int32>* CVar = IConsoleManager::Get().FindTConsoleVariableDataInt(TEXT("r.TonemapperPhoto"));
 		int32 Value = CVar->GetValueOnRenderThread();
@@ -1045,7 +1060,7 @@ void FilmPostSetConstantsPhoto(FVector4* RESTRICT const Constants, const FPostPr
 	// If an LDR display goes from 0-1 this multiply will be 1.
 	// If an HDR display goes from 0-16 this multiply will be 1/16 (the display is 16x brigher).
 	{
-		static TConsoleVariableData<int32>* CVar = IConsoleManager::Get().FindTConsoleVariableDataInt(TEXT("r.TonemapperHDR"));
+		static TConsoleVariableData<float>* CVar = IConsoleManager::Get().FindTConsoleVariableDataFloat(TEXT("r.TonemapperHDR"));
 		float Value = CVar->GetValueOnRenderThread();
 		if(Value < 1.0f)
 		{
@@ -1104,6 +1119,7 @@ class FPostProcessTonemapPS : public FGlobalShader
 		OutEnvironment.SetDefine(TEXT("USE_COLOR_GRADING"),		 TonemapperIsDefined(ConfigBitmask, TonemapperColorGrading));
 		OutEnvironment.SetDefine(TEXT("USE_PHOTO"),              TonemapperIsDefined(ConfigBitmask, TonemapperPhoto));
 		OutEnvironment.SetDefine(TEXT("USE_709"),                TonemapperIsDefined(ConfigBitmask, Tonemapper709));
+		OutEnvironment.SetDefine(TEXT("USE_GAMMA"),              TonemapperIsDefined(ConfigBitmask, TonemapperGamma));
 
 		if( !IsFeatureLevelSupported(Platform,ERHIFeatureLevel::SM5) )
 		{
@@ -1290,9 +1306,18 @@ public:
 		}
 
 		{
-			FVector2D InvDisplayGammaValue;
+			FVector InvDisplayGammaValue;
 			InvDisplayGammaValue.X = 1.0f / ViewFamily.RenderTarget->GetDisplayGamma();
 			InvDisplayGammaValue.Y = 2.2f / ViewFamily.RenderTarget->GetDisplayGamma();
+			{
+				static TConsoleVariableData<float>* CVar = IConsoleManager::Get().FindTConsoleVariableDataFloat(TEXT("r.TonemapperGamma"));
+				float Value = CVar->GetValueOnRenderThread();
+				if(Value < 1.0f)
+				{
+					Value = 1.0f;
+				}
+				InvDisplayGammaValue.Z = 1.0f / Value;
+			}
 			SetShaderValue(Context.RHICmdList, ShaderRHI, InverseGamma, InvDisplayGammaValue);
 		}
 
@@ -1640,9 +1665,10 @@ public:
 		SetShaderValue(Context.RHICmdList, ShaderRHI, GrainScaleBiasJitter, GrainValue);
 
 		{
-			FVector2D InvDisplayGammaValue;
+			FVector InvDisplayGammaValue;
 			InvDisplayGammaValue.X = 1.0f / ViewFamily.RenderTarget->GetDisplayGamma();
 			InvDisplayGammaValue.Y = 2.2f / ViewFamily.RenderTarget->GetDisplayGamma();
+			InvDisplayGammaValue.Z = 1.0; // Unused on mobile.
 			SetShaderValue(Context.RHICmdList, ShaderRHI, InverseGamma, InvDisplayGammaValue);
 		}
 
