Get motion vector inputs back into DOF temporal AA.

diff --git a/Engine/Shaders/PostProcessTemporalAA.usf b/Engine/Shaders/PostProcessTemporalAA.usf
--- a/Engine/Shaders/PostProcessTemporalAA.usf
+++ b/Engine/Shaders/PostProcessTemporalAA.usf
@@ -207,12 +207,17 @@ void DOFTemporalAAPS( float4 UVAndScreenPos : TEXCOORD0, float3 InExposureScaleV
 {
 #if ENABLE_TEMPORAL_AA
 	float InExposureScale = InExposureScaleVignette.x;
+	#define AA_FILTERED 0
+	#define AA_VELOCITY_WEIGHTING 0
 	#define AA_ALPHA 0
 	#define AA_CROSS 4
-	#define AA_DYNAMIC 0
+	#define AA_DYNAMIC 1
 	#define AA_NAN 1
 	#define AA_BORDER 1
 	#define AA_FORCE_ALPHA_CLAMP 1
+	#define AA_GREEN_AS_LUMA 1
+	#define AA_YCOCG 1
+	#define AA_DOF 1
 	#include "PostProcessTemporalCommon.usf"
 #else 
 	OutColor = PostprocessInput0.SampleLevel(PostprocessInput0Sampler, UVAndScreenPos.xy, 0);
diff --git a/Engine/Shaders/PostProcessTemporalCommon.usf b/Engine/Shaders/PostProcessTemporalCommon.usf
--- a/Engine/Shaders/PostProcessTemporalCommon.usf
+++ b/Engine/Shaders/PostProcessTemporalCommon.usf
@@ -120,6 +120,10 @@ which can be configured via defines for various temporal AA passes.
 	#define AA_GREEN_AS_LUMA AA_YCOCG
 #endif
 
+// Special adjustments for DOF.
+#ifndef AA_DOF
+	#define AA_DOF 0
+#endif
 
 
 
@@ -559,3 +563,4 @@ which can be configured via defines for various temporal AA passes.
 #undef AA_GREEN_AS_LUMA
 #undef AA_VELOCITY_WEIGHTING
 #undef AA_YCOCG
+#undef AA_DOF
\ No newline at end of file
diff --git a/Engine/Source/Runtime/Renderer/Private/PostProcess/PostProcessTemporalAA.h b/Engine/Source/Runtime/Renderer/Private/PostProcess/PostProcessTemporalAA.h
--- a/Engine/Source/Runtime/Renderer/Private/PostProcess/PostProcessTemporalAA.h
+++ b/Engine/Source/Runtime/Renderer/Private/PostProcess/PostProcessTemporalAA.h
@@ -25,7 +25,7 @@ public:
 // ePId_Input1: Previous frame's output (bilinear)
 // ePId_Input2: Previous frame's output (point)
 // derives from TRenderingCompositePassBase<InputCount, OutputCount> 
-class FRCPassPostProcessDOFTemporalAA : public TRenderingCompositePassBase<3, 1>
+class FRCPassPostProcessDOFTemporalAA : public TRenderingCompositePassBase<4, 1>
 {
 public:
 	// interface FRenderingCompositePass ---------
@@ -38,7 +38,7 @@ public:
 // ePId_Input1: Previous frame's output (bilinear)
 // ePId_Input2: Previous frame's output (point)
 // derives from TRenderingCompositePassBase<InputCount, OutputCount> 
-class FRCPassPostProcessDOFTemporalAANear : public TRenderingCompositePassBase<3, 1>
+class FRCPassPostProcessDOFTemporalAANear : public TRenderingCompositePassBase<4, 1>
 {
 public:
 	// interface FRenderingCompositePass ---------
diff --git a/Engine/Source/Runtime/Renderer/Private/PostProcess/PostProcessing.cpp b/Engine/Source/Runtime/Renderer/Private/PostProcess/PostProcessing.cpp
--- a/Engine/Source/Runtime/Renderer/Private/PostProcess/PostProcessing.cpp
+++ b/Engine/Source/Runtime/Renderer/Private/PostProcess/PostProcessing.cpp
@@ -273,7 +273,7 @@ static void AddVisualizeBloomOverlay(FPostprocessContext& Context, FRenderingCom
 	Context.FinalOutput = FRenderingCompositeOutputRef(Node);
 }
 
-static void AddPostProcessDepthOfFieldBokeh(FPostprocessContext& Context, FRenderingCompositeOutputRef& SeparateTranslucency)
+static void AddPostProcessDepthOfFieldBokeh(FPostprocessContext& Context, FRenderingCompositeOutputRef& SeparateTranslucency, FRenderingCompositeOutputRef& VelocityInput)
 {
 	// downsample, mask out the in focus part, depth in alpha
 	FRenderingCompositePass* DOFSetup = Context.Graph.RegisterPass(new(FMemStack::Get()) FRCPassPostProcessBokehDOFSetup());
@@ -300,7 +300,7 @@ static void AddPostProcessDepthOfFieldBokeh(FPostprocessContext& Context, FRende
 		NodeTemporalAA->SetInput( ePId_Input0, DOFSetup );
 		NodeTemporalAA->SetInput( ePId_Input1, FRenderingCompositeOutputRef( HistoryInput ) );
 		NodeTemporalAA->SetInput( ePId_Input2, FRenderingCompositeOutputRef( HistoryInput ) );
-		//NodeTemporalAA->SetInput( ePId_Input3, VelocityInput );
+		NodeTemporalAA->SetInput( ePId_Input3, VelocityInput );
 
 		DOFInputPass = NodeTemporalAA;
 		ViewState->bBokehDOFHistory = true;
@@ -319,7 +319,7 @@ static void AddPostProcessDepthOfFieldBokeh(FPostprocessContext& Context, FRende
 	Context.FinalOutput = FRenderingCompositeOutputRef(NodeRecombined);
 }
 
-static void AddPostProcessDepthOfFieldGaussian(FPostprocessContext& Context, FDepthOfFieldStats& Out)
+static void AddPostProcessDepthOfFieldGaussian(FPostprocessContext& Context, FDepthOfFieldStats& Out, FRenderingCompositeOutputRef& VelocityInput)
 {
 	float FarSize = Context.View.FinalPostProcessSettings.DepthOfFieldFarBlurSize;
 	float NearSize = Context.View.FinalPostProcessSettings.DepthOfFieldNearBlurSize;
@@ -368,7 +368,7 @@ static void AddPostProcessDepthOfFieldGaussian(FPostprocessContext& Context, FDe
 		NodeTemporalAA->SetInput( ePId_Input0, DOFSetup );
 		NodeTemporalAA->SetInput( ePId_Input1, FRenderingCompositeOutputRef( HistoryInput ) );
 		NodeTemporalAA->SetInput( ePId_Input2, FRenderingCompositeOutputRef( HistoryInput ) );
-		//NodeTemporalAA->SetInput( ePId_Input3, VelocityInput );
+		NodeTemporalAA->SetInput( ePId_Input3, VelocityInput );
 
 		DOFInputPass = NodeTemporalAA;
 		ViewState->bBokehDOFHistory = false;
@@ -396,7 +396,7 @@ static void AddPostProcessDepthOfFieldGaussian(FPostprocessContext& Context, FDe
 		NodeTemporalAA->SetInput( ePId_Input0, FRenderingCompositeOutputRef( DOFSetup, ePId_Output1 ) );
 		NodeTemporalAA->SetInput( ePId_Input1, FRenderingCompositeOutputRef( HistoryInput, DOFInputPassId2 ) );
 		NodeTemporalAA->SetInput( ePId_Input2, FRenderingCompositeOutputRef( HistoryInput, DOFInputPassId2 ) );
-		//NodeTemporalAA->SetInput( ePId_Input3, VelocityInput );
+		NodeTemporalAA->SetInput( ePId_Input3, VelocityInput );
 
 		DOFInputPass2 = NodeTemporalAA;
 		ViewState->bBokehDOFHistory2 = false;
@@ -427,7 +427,7 @@ static void AddPostProcessDepthOfFieldGaussian(FPostprocessContext& Context, FDe
 	Context.FinalOutput = FRenderingCompositeOutputRef(NodeRecombined);
 }
 
-static void AddPostProcessDepthOfFieldCircle(FPostprocessContext& Context, FDepthOfFieldStats& Out)
+static void AddPostProcessDepthOfFieldCircle(FPostprocessContext& Context, FDepthOfFieldStats& Out, FRenderingCompositeOutputRef& VelocityInput)
 {
 	if(Context.View.Family->EngineShowFlags.VisualizeDOF)
 	{
@@ -459,7 +459,7 @@ static void AddPostProcessDepthOfFieldCircle(FPostprocessContext& Context, FDept
 		NodeTemporalAA->SetInput( ePId_Input0, DOFSetup );
 		NodeTemporalAA->SetInput( ePId_Input1, FRenderingCompositeOutputRef( HistoryInput ) );
 		NodeTemporalAA->SetInput( ePId_Input2, FRenderingCompositeOutputRef( HistoryInput ) );
-		//NodeTemporalAA->SetInput( ePId_Input3, VelocityInput );
+		NodeTemporalAA->SetInput( ePId_Input3, VelocityInput );
 
 		DOFInputPass = NodeTemporalAA;
 		ViewState->bBokehDOFHistory = false;
@@ -1025,11 +1025,31 @@ void FPostProcessing::Process(FRHICommandListImmediate& RHICmdList, FViewInfo& V
 				bool bCircleDOF = View.FinalPostProcessSettings.DepthOfFieldMethod == DOFM_CircleDOF;
 				if(!bCircleDOF)
 				{
-					AddPostProcessDepthOfFieldGaussian(Context, DepthOfFieldStat);
+					if(VelocityInput.IsValid())
+					{
+						AddPostProcessDepthOfFieldGaussian(Context, DepthOfFieldStat, VelocityInput);
+					}
+					else
+					{
+						// black is how we clear the velocity buffer so this means no velocity
+						FRenderingCompositePass* NoVelocity = Context.Graph.RegisterPass(new(FMemStack::Get()) FRCPassPostProcessInput(GSystemTextures.BlackDummy));
+						FRenderingCompositeOutputRef NoVelocityRef(NoVelocity);
+						AddPostProcessDepthOfFieldGaussian(Context, DepthOfFieldStat, NoVelocityRef);
+					}
 				}
 				else
 				{
-					AddPostProcessDepthOfFieldCircle(Context, DepthOfFieldStat);
+					if(VelocityInput.IsValid())
+					{
+						AddPostProcessDepthOfFieldCircle(Context, DepthOfFieldStat, VelocityInput);
+					}
+					else
+					{
+						// black is how we clear the velocity buffer so this means no velocity
+						FRenderingCompositePass* NoVelocity = Context.Graph.RegisterPass(new(FMemStack::Get()) FRCPassPostProcessInput(GSystemTextures.BlackDummy));
+						FRenderingCompositeOutputRef NoVelocityRef(NoVelocity);
+						AddPostProcessDepthOfFieldCircle(Context, DepthOfFieldStat, NoVelocityRef);
+					}
 				}
 			}
 
@@ -1040,7 +1060,17 @@ void FPostProcessing::Process(FRHICommandListImmediate& RHICmdList, FViewInfo& V
 
 			if(bBokehDOF)
 			{
-				AddPostProcessDepthOfFieldBokeh(Context, SeparateTranslucency);
+				if(VelocityInput.IsValid())
+				{
+					AddPostProcessDepthOfFieldBokeh(Context, SeparateTranslucency, VelocityInput);
+				}
+				else
+				{
+					// black is how we clear the velocity buffer so this means no velocity
+					FRenderingCompositePass* NoVelocity = Context.Graph.RegisterPass(new(FMemStack::Get()) FRCPassPostProcessInput(GSystemTextures.BlackDummy));
+					FRenderingCompositeOutputRef NoVelocityRef(NoVelocity);
+					AddPostProcessDepthOfFieldBokeh(Context, SeparateTranslucency, NoVelocityRef);
+				}
 			}
 			else 
 			{
