Integrate - DFAO history now supports per pixel velocities
* Velocity pass has been moved to right after the base pass

diff --git a/Engine/Shaders/DistanceFieldSurfaceCacheLighting.usf b/Engine/Shaders/DistanceFieldSurfaceCacheLighting.usf
--- a/Engine/Shaders/DistanceFieldSurfaceCacheLighting.usf
+++ b/Engine/Shaders/DistanceFieldSurfaceCacheLighting.usf
@@ -561,7 +561,10 @@ Texture2D IrradianceHistoryTexture;
 SamplerState IrradianceHistorySampler;
 float HistoryWeight;
 
-/** Reprojection the occlusion history, only correct for pixels belonging to static objects (camera motion only). */
+Texture2D VelocityTexture;
+SamplerState VelocityTextureSampler;
+
+/** Reproject the occlusion history. */
 void UpdateHistoryPS(
 	in float4 UVAndScreenPos : TEXCOORD0
 	,out float4 OutBentNormal : SV_Target0
@@ -570,16 +573,27 @@ void UpdateHistoryPS(
 #endif
 	)
 {
+	float2 FullResTexel = UVAndScreenPos.xy - .5f * View.ViewSizeAndSceneTexelSize.zw;
+
 	float3 xyd;
     xyd.xy = UVAndScreenPos.zw * float2(0.5, -0.5) + 0.5;
-    xyd.z = Texture2DSampleLevel(SceneDepthTexture, SceneDepthTextureSampler, UVAndScreenPos.xy, 0).r;
+    xyd.z = Texture2DSampleLevel(SceneDepthTexture, SceneDepthTextureSampler, FullResTexel, 0).r;
 
     float scaleM = 1.0 / (dot(xyd, CameraMotion[0].xyz) + CameraMotion[0].w);
-    float2 mv;
-    mv.x = ((xyd.x * ((CameraMotion[1].x * xyd.y) + (CameraMotion[1].y * xyd.z) + CameraMotion[1].z)) + (CameraMotion[1].w * xyd.y) + (CameraMotion[2].x * xyd.x * xyd.x) + (CameraMotion[2].y * xyd.z) + CameraMotion[2].z) * scaleM;
-    mv.y = ((xyd.y * ((CameraMotion[3].x * xyd.x) + (CameraMotion[3].y * xyd.z) + CameraMotion[3].z)) + (CameraMotion[3].w * xyd.x) + (CameraMotion[4].x * xyd.y * xyd.y) + (CameraMotion[4].y * xyd.z) + CameraMotion[4].z) * scaleM;
+    float2 MotionVector;
+    MotionVector.x = ((xyd.x * ((CameraMotion[1].x * xyd.y) + (CameraMotion[1].y * xyd.z) + CameraMotion[1].z)) + (CameraMotion[1].w * xyd.y) + (CameraMotion[2].x * xyd.x * xyd.x) + (CameraMotion[2].y * xyd.z) + CameraMotion[2].z) * scaleM;
+    MotionVector.y = ((xyd.y * ((CameraMotion[3].x * xyd.x) + (CameraMotion[3].y * xyd.z) + CameraMotion[3].z)) + (CameraMotion[3].w * xyd.x) + (CameraMotion[4].x * xyd.y * xyd.y) + (CameraMotion[4].y * xyd.z) + CameraMotion[4].z) * scaleM;
+
+	// Note: have to sample from one of the high res texels exactly to avoid filtering of the velocity buffer
+	// 0 is stored where camera velocity should be used, which must not be filtered with valid object velocities
+	float2 VelocityN = Texture2DSampleLevel(VelocityTexture, VelocityTextureSampler, FullResTexel, 0).xy;
+
+	if (VelocityN.x > 0 && View.GeneralPurposeTweak > 0)
+	{
+		MotionVector = float2(-.5f, .5f) * DecodeVelocityFromTexture(VelocityN);
+	}
 
-    float2 OldUV = UVAndScreenPos.xy + mv / float2(0.5, -0.5) * View.ScreenPositionScaleBias.xy;
+    float2 OldUV = UVAndScreenPos.xy + MotionVector / float2(0.5, -0.5) * View.ScreenPositionScaleBias.xy;
 	float2 OldScreenPosition = (OldUV.xy - View.ScreenPositionScaleBias.wz) / View.ScreenPositionScaleBias.xy;
 
 	// Distance field AO was computed at 0,0 regardless of viewrect min
diff --git a/Engine/Source/Runtime/Renderer/Private/DeferredShadingRenderer.cpp b/Engine/Source/Runtime/Renderer/Private/DeferredShadingRenderer.cpp
--- a/Engine/Source/Runtime/Renderer/Private/DeferredShadingRenderer.cpp
+++ b/Engine/Source/Runtime/Renderer/Private/DeferredShadingRenderer.cpp
@@ -878,6 +878,20 @@ void FDeferredShadingSceneRenderer::Render(FRHICommandListImmediate& RHICmdList)
 		OcclusionSubmittedFence = FRHICommandListExecutor::RHIThreadFence();
 	}
 
+	TRefCountPtr<IPooledRenderTarget> VelocityRT;
+
+	if (bUseVelocityGBuffer)
+	{
+		VelocityRT = GSceneRenderTargets.GetGBufferVelocityRT();
+	}
+	else if (bShouldRenderVelocities)
+	{
+		// Render the velocities of movable objects
+		RHICmdList.SetCurrentStat(GET_STATID(STAT_CLM_Velocity));
+		RenderVelocities(RHICmdList, VelocityRT);
+		RHICmdList.SetCurrentStat(GET_STATID(STAT_CLM_AfterVelocity));
+	}
+
 	// Render lighting.
 	if (ViewFamily.EngineShowFlags.Lighting
 		&& FeatureLevel >= ERHIFeatureLevel::SM4
@@ -916,7 +930,7 @@ void FDeferredShadingSceneRenderer::Render(FRHICommandListImmediate& RHICmdList)
 		}
 
 		TRefCountPtr<IPooledRenderTarget> DynamicBentNormalAO;
-		RenderDynamicSkyLighting(RHICmdList, DynamicBentNormalAO);
+		RenderDynamicSkyLighting(RHICmdList, VelocityRT, DynamicBentNormalAO);
 
 		//SSR and SSS need the SceneColor finalized as an SRV.
 		GSceneRenderTargets.FinishRenderingSceneColor(RHICmdList, true);
@@ -1012,7 +1026,7 @@ void FDeferredShadingSceneRenderer::Render(FRHICommandListImmediate& RHICmdList)
 		// Use the skylight's max distance if there is one, to be consistent with DFAO shadowing on the skylight
 		const float OcclusionMaxDistance = Scene->SkyLight && !Scene->SkyLight->bWantsStaticShadowing ? Scene->SkyLight->OcclusionMaxDistance : 600;
 		TRefCountPtr<IPooledRenderTarget> DummyOutput;
-		RenderDistanceFieldAOSurfaceCache(RHICmdList, FDistanceFieldAOParameters(OcclusionMaxDistance), DummyOutput, DummyOutput, ViewFamily.EngineShowFlags.VisualizeDistanceFieldAO, ViewFamily.EngineShowFlags.VisualizeDistanceFieldGI);
+		RenderDistanceFieldAOSurfaceCache(RHICmdList, FDistanceFieldAOParameters(OcclusionMaxDistance), VelocityRT, DummyOutput, DummyOutput, ViewFamily.EngineShowFlags.VisualizeDistanceFieldAO, ViewFamily.EngineShowFlags.VisualizeDistanceFieldGI);
 	}
 
 	if (ViewFamily.EngineShowFlags.VisualizeMeshDistanceFields)
@@ -1023,20 +1037,6 @@ void FDeferredShadingSceneRenderer::Render(FRHICommandListImmediate& RHICmdList)
 	// Resolve the scene color for post processing.
 	GSceneRenderTargets.ResolveSceneColor(RHICmdList, FResolveRect(0, 0, ViewFamily.FamilySizeX, ViewFamily.FamilySizeY));
 
-	TRefCountPtr<IPooledRenderTarget> VelocityRT;
-
-	if (bUseVelocityGBuffer)
-	{
-		VelocityRT = GSceneRenderTargets.GetGBufferVelocityRT();
-	}
-	else if (bShouldRenderVelocities)
-	{
-		// Render the velocities of movable objects
-		RHICmdList.SetCurrentStat(GET_STATID(STAT_CLM_Velocity));
-		RenderVelocities(RHICmdList, VelocityRT);
-		RHICmdList.SetCurrentStat(GET_STATID(STAT_CLM_AfterVelocity));
-	}
-
 	// Finish rendering for each view.
 	if(ViewFamily.bResolveScene)
 	{
diff --git a/Engine/Source/Runtime/Renderer/Private/DeferredShadingRenderer.h b/Engine/Source/Runtime/Renderer/Private/DeferredShadingRenderer.h
--- a/Engine/Source/Runtime/Renderer/Private/DeferredShadingRenderer.h
+++ b/Engine/Source/Runtime/Renderer/Private/DeferredShadingRenderer.h
@@ -184,12 +184,13 @@ private:
 	void RenderDeferredReflections(FRHICommandListImmediate& RHICmdList, const TRefCountPtr<IPooledRenderTarget>& DynamicBentNormalAO);
 
 	/** Render dynamic sky lighting from Movable sky lights. */
-	void RenderDynamicSkyLighting(FRHICommandListImmediate& RHICmdList, TRefCountPtr<IPooledRenderTarget>& DynamicBentNormalAO);
+	void RenderDynamicSkyLighting(FRHICommandListImmediate& RHICmdList, const TRefCountPtr<IPooledRenderTarget>& VelocityTexture, TRefCountPtr<IPooledRenderTarget>& DynamicBentNormalAO);
 
 	/** Render Ambient Occlusion using mesh distance fields and the surface cache, which supports dynamic rigid meshes. */
 	bool RenderDistanceFieldAOSurfaceCache(
 		FRHICommandListImmediate& RHICmdList, 
 		const class FDistanceFieldAOParameters& Parameters, 
+		const TRefCountPtr<IPooledRenderTarget>& VelocityTexture,
 		TRefCountPtr<IPooledRenderTarget>& OutDynamicBentNormalAO, 
 		TRefCountPtr<IPooledRenderTarget>& OutDynamicIrradiance,
 		bool bVisualizeAmbientOcclusion,
diff --git a/Engine/Source/Runtime/Renderer/Private/DistanceFieldSurfaceCacheLighting.cpp b/Engine/Source/Runtime/Renderer/Private/DistanceFieldSurfaceCacheLighting.cpp
--- a/Engine/Source/Runtime/Renderer/Private/DistanceFieldSurfaceCacheLighting.cpp
+++ b/Engine/Source/Runtime/Renderer/Private/DistanceFieldSurfaceCacheLighting.cpp
@@ -2323,6 +2323,8 @@ public:
 		IrradianceHistorySampler.Bind(Initializer.ParameterMap, TEXT("IrradianceHistorySampler"));
 		CameraMotion.Bind(Initializer.ParameterMap);
 		HistoryWeight.Bind(Initializer.ParameterMap, TEXT("HistoryWeight"));
+		VelocityTexture.Bind(Initializer.ParameterMap, TEXT("VelocityTexture"));
+		VelocityTextureSampler.Bind(Initializer.ParameterMap, TEXT("VelocityTextureSampler"));
 	}
 
 	void SetParameters(
@@ -2331,7 +2333,8 @@ public:
 		FSceneRenderTargetItem& BentNormalHistoryTextureValue, 
 		TRefCountPtr<IPooledRenderTarget>* IrradianceHistoryRT, 
 		FSceneRenderTargetItem& DistanceFieldAOBentNormal, 
-		IPooledRenderTarget* DistanceFieldIrradiance)
+		IPooledRenderTarget* DistanceFieldIrradiance,
+		FSceneRenderTargetItem& VelocityTextureValue)
 	{
 		const FPixelShaderRHIParamRef ShaderRHI = GetPixelShader();
 
@@ -2383,6 +2386,15 @@ public:
 		CameraMotion.Set(RHICmdList, View, ShaderRHI);
 
 		SetShaderValue(RHICmdList, ShaderRHI, HistoryWeight, GAOHistoryWeight);
+
+		SetTextureParameter(
+			RHICmdList,
+			ShaderRHI,
+			VelocityTexture,
+			VelocityTextureSampler,
+			TStaticSamplerState<SF_Bilinear,AM_Clamp,AM_Clamp,AM_Clamp>::GetRHI(),
+			VelocityTextureValue.ShaderResourceTexture
+			);
 	}
 	// FShader interface.
 	virtual bool Serialize(FArchive& Ar)
@@ -2399,6 +2411,8 @@ public:
 		Ar << IrradianceHistorySampler;
 		Ar << CameraMotion;
 		Ar << HistoryWeight;
+		Ar << VelocityTexture;
+		Ar << VelocityTextureSampler;
 		return bShaderHasOutdatedParameters;
 	}
 
@@ -2415,6 +2429,8 @@ private:
 	FShaderResourceParameter IrradianceHistorySampler;
 	FCameraMotionParameters CameraMotion;
 	FShaderParameter HistoryWeight;
+	FShaderResourceParameter VelocityTexture;
+	FShaderResourceParameter VelocityTextureSampler;
 };
 
 IMPLEMENT_SHADER_TYPE(template<>,TUpdateHistoryPS<true>,TEXT("DistanceFieldSurfaceCacheLighting"),TEXT("UpdateHistoryPS"),SF_Pixel);
@@ -2436,6 +2452,7 @@ void UpdateHistory(
 	const FViewInfo& View, 
 	const TCHAR* BentNormalHistoryRTName,
 	const TCHAR* IrradianceHistoryRTName,
+	FSceneRenderTargetItem& VelocityTexture,
 	/** Contains last frame's history, if non-NULL.  This will be updated with the new frame's history. */
 	TRefCountPtr<IPooledRenderTarget>* BentNormalHistoryState,
 	TRefCountPtr<IPooledRenderTarget>* IrradianceHistoryState,
@@ -2490,7 +2507,7 @@ void UpdateHistory(
 
 					static FGlobalBoundShaderState BoundShaderState;
 					SetGlobalBoundShaderState(RHICmdList, View.GetFeatureLevel(), BoundShaderState, GFilterVertexDeclaration.VertexDeclarationRHI, *VertexShader, *PixelShader);
-					PixelShader->SetParameters(RHICmdList, View, (*BentNormalHistoryState)->GetRenderTargetItem(), IrradianceHistoryState, BentNormalSource->GetRenderTargetItem(), IrradianceSource);
+					PixelShader->SetParameters(RHICmdList, View, (*BentNormalHistoryState)->GetRenderTargetItem(), IrradianceHistoryState, BentNormalSource->GetRenderTargetItem(), IrradianceSource, VelocityTexture);
 				}
 				else
 				{
@@ -2498,7 +2515,7 @@ void UpdateHistory(
 
 					static FGlobalBoundShaderState BoundShaderState;
 					SetGlobalBoundShaderState(RHICmdList, View.GetFeatureLevel(), BoundShaderState, GFilterVertexDeclaration.VertexDeclarationRHI, *VertexShader, *PixelShader);
-					PixelShader->SetParameters(RHICmdList, View, (*BentNormalHistoryState)->GetRenderTargetItem(), IrradianceHistoryState, BentNormalSource->GetRenderTargetItem(), IrradianceSource);
+					PixelShader->SetParameters(RHICmdList, View, (*BentNormalHistoryState)->GetRenderTargetItem(), IrradianceHistoryState, BentNormalSource->GetRenderTargetItem(), IrradianceSource, VelocityTexture);
 				}
 
 				VertexShader->SetParameters(RHICmdList, View);
@@ -2545,6 +2562,7 @@ void PostProcessBentNormalAO(
 	FRHICommandList& RHICmdList, 
 	const FDistanceFieldAOParameters& Parameters, 
 	const FViewInfo& View, 
+	FSceneRenderTargetItem& VelocityTexture,
 	FSceneRenderTargetItem& BentNormalInterpolation, 
 	IPooledRenderTarget* IrradianceInterpolation,
 	FSceneRenderTargetItem& DistanceFieldNormal,
@@ -2692,6 +2710,7 @@ void PostProcessBentNormalAO(
 			View, 
 			TEXT("DistanceFieldAOHistory"),
 			TEXT("DistanceFieldIrradianceHistory"),
+			VelocityTexture,
 			BentNormalHistoryState,
 			IrradianceHistoryState,
 			DistanceFieldAOBentNormal, 
@@ -3484,6 +3503,7 @@ void TrackGPUProgress(FRHICommandListImmediate& RHICmdList, uint32 DebugId)
 bool FDeferredShadingSceneRenderer::RenderDistanceFieldAOSurfaceCache(
 	FRHICommandListImmediate& RHICmdList, 
 	const FDistanceFieldAOParameters& Parameters, 
+	const TRefCountPtr<IPooledRenderTarget>& VelocityTexture,
 	TRefCountPtr<IPooledRenderTarget>& OutDynamicBentNormalAO, 
 	TRefCountPtr<IPooledRenderTarget>& OutDynamicIrradiance,
 	bool bVisualizeAmbientOcclusion,
@@ -3801,7 +3821,17 @@ bool FDeferredShadingSceneRenderer::RenderDistanceFieldAOSurfaceCache(
 			// Post process the AO to cover over artifacts
 			TRefCountPtr<IPooledRenderTarget> BentNormalOutput;
 			TRefCountPtr<IPooledRenderTarget> IrradianceOutput;
-			PostProcessBentNormalAO(RHICmdList, Parameters, View, BentNormalAccumulation->GetRenderTargetItem(), IrradianceAccumulation, DistanceFieldNormal->GetRenderTargetItem(), BentNormalOutput, IrradianceOutput);
+
+			PostProcessBentNormalAO(
+				RHICmdList, 
+				Parameters, 
+				View, 
+				VelocityTexture->GetRenderTargetItem(), 
+				BentNormalAccumulation->GetRenderTargetItem(), 
+				IrradianceAccumulation, 
+				DistanceFieldNormal->GetRenderTargetItem(), 
+				BentNormalOutput, 
+				IrradianceOutput);
 
 			GRenderTargetPool.VisualizeTexture.SetCheckPoint(RHICmdList, BentNormalOutput);
 
@@ -3963,7 +3993,7 @@ bool FDeferredShadingSceneRenderer::ShouldRenderDistanceFieldAO() const
 		&& !ViewFamily.EngineShowFlags.VisualizeMeshDistanceFields;
 }
 
-void FDeferredShadingSceneRenderer::RenderDynamicSkyLighting(FRHICommandListImmediate& RHICmdList, TRefCountPtr<IPooledRenderTarget>& DynamicBentNormalAO)
+void FDeferredShadingSceneRenderer::RenderDynamicSkyLighting(FRHICommandListImmediate& RHICmdList, const TRefCountPtr<IPooledRenderTarget>& VelocityTexture, TRefCountPtr<IPooledRenderTarget>& DynamicBentNormalAO)
 {
 	if (ShouldRenderDynamicSkyLight())
 	{
@@ -3976,7 +4006,7 @@ void FDeferredShadingSceneRenderer::RenderDynamicSkyLighting(FRHICommandListImme
 
 		if (ShouldRenderDistanceFieldAO())
 		{
-			bApplyShadowing = RenderDistanceFieldAOSurfaceCache(RHICmdList, Parameters, DynamicBentNormalAO, DynamicIrradiance, false, false);
+			bApplyShadowing = RenderDistanceFieldAOSurfaceCache(RHICmdList, Parameters, VelocityTexture, DynamicBentNormalAO, DynamicIrradiance, false, false);
 		}
 
 		GSceneRenderTargets.BeginRenderingSceneColor(RHICmdList);
diff --git a/Engine/Source/Runtime/Renderer/Private/VelocityRendering.cpp b/Engine/Source/Runtime/Renderer/Private/VelocityRendering.cpp
--- a/Engine/Source/Runtime/Renderer/Private/VelocityRendering.cpp
+++ b/Engine/Source/Runtime/Renderer/Private/VelocityRendering.cpp
@@ -713,8 +713,9 @@ bool FDeferredShadingSceneRenderer::ShouldRenderVelocities() const
 
 		bool bTemporalAA = (View.FinalPostProcessSettings.AntiAliasingMethod == AAM_TemporalAA) && !View.bCameraCut;
 		bool bMotionBlur = IsMotionBlurEnabled(View);
+		bool bDistanceFieldAO = ShouldPrepareForDistanceFieldAO();
 
-		bNeedsVelocity |= bMotionBlur || bTemporalAA;
+		bNeedsVelocity |= bMotionBlur || bTemporalAA || bDistanceFieldAO;
 	}
 
 	return bNeedsVelocity;
