Integrate - Skylight IndirectLightingIntensity affects heightfield GI

diff --git a/Engine/Shaders/HeightfieldLighting.usf b/Engine/Shaders/HeightfieldLighting.usf
--- a/Engine/Shaders/HeightfieldLighting.usf
+++ b/Engine/Shaders/HeightfieldLighting.usf
@@ -340,6 +340,7 @@ void RayTracedShadowHeightfieldsPS(
 }
 
 float3 LightColor;
+float SkyLightIndirectScale;
 
 #if APPLY_LIGHT_FUNCTION
 
@@ -366,7 +367,7 @@ void LightHeightfieldsPS(
 	float3 DiffuseColor = GetGlobalDiffuseColor(Input.GlobalUV);
 	float3 Lighting = (Shadowing * GreyScaleLightFunction * saturate(dot(-LightDirection, WorldHeightfieldNormal))) * LightColor;
 	 
-	float3 SkylightDiffuse = GetSkySHDiffuse(WorldHeightfieldNormal) * View.SkyLightColor.rgb;
+	float3 SkylightDiffuse = GetSkySHDiffuse(WorldHeightfieldNormal) * View.SkyLightColor.rgb * SkyLightIndirectScale;
 	Lighting += SkylightDiffuse;
 
 	OutColor = float4(DiffuseColor * Lighting, 0);
@@ -645,8 +646,8 @@ void CalculateHeightfieldIrradianceCS(
 				#define VISUALIZE_HEIGHTFIELD_LIGHTING 0
 				#if VISUALIZE_HEIGHTFIELD_LIGHTING
 
-					//Lighting = SampleHeightfieldLighting(GlobalUV) / FINAL_GATHER_THREADGROUP_SIZE;
-					Lighting = GetGlobalDiffuseColor(GlobalUV) / FINAL_GATHER_THREADGROUP_SIZE;
+					Lighting = SampleHeightfieldLighting(GlobalUV) / FINAL_GATHER_THREADGROUP_SIZE;
+					//Lighting = GetGlobalDiffuseColor(GlobalUV) / FINAL_GATHER_THREADGROUP_SIZE;
 					//Lighting = frac(WorldHeightfieldShadingPosition / 1000);
 					//Lighting = abs(WorldHeightfieldShadingPosition.z - WorldPosition.z) / 100.0f;
 					//Lighting = GetGlobalHeightfieldHeight(GlobalUV) / 100000.0f;
diff --git a/Engine/Source/Runtime/Engine/Private/SkyLightComponent.cpp b/Engine/Source/Runtime/Engine/Private/SkyLightComponent.cpp
--- a/Engine/Source/Runtime/Engine/Private/SkyLightComponent.cpp
+++ b/Engine/Source/Runtime/Engine/Private/SkyLightComponent.cpp
@@ -79,6 +79,7 @@ FSkyLightSceneProxy::FSkyLightSceneProxy(const USkyLightComponent* InLightCompon
 	, bPrecomputedLightingIsValid(InLightComponent->bPrecomputedLightingIsValid)
 	, bHasStaticLighting(InLightComponent->HasStaticLighting())
 	, LightColor(FLinearColor(InLightComponent->LightColor) * InLightComponent->Intensity)
+	, IndirectLightingIntensity(InLightComponent->IndirectLightingIntensity)
 	, IrradianceEnvironmentMap(InLightComponent->IrradianceEnvironmentMap)
 	, OcclusionMaxDistance(InLightComponent->OcclusionMaxDistance)
 	, Contrast(InLightComponent->Contrast)
diff --git a/Engine/Source/Runtime/Engine/Public/SceneManagement.h b/Engine/Source/Runtime/Engine/Public/SceneManagement.h
--- a/Engine/Source/Runtime/Engine/Public/SceneManagement.h
+++ b/Engine/Source/Runtime/Engine/Public/SceneManagement.h
@@ -611,6 +611,7 @@ public:
 	bool bHasStaticLighting;
 	FLinearColor LightColor;
 	FSHVectorRGB3 IrradianceEnvironmentMap;
+	float IndirectLightingIntensity;
 	float OcclusionMaxDistance;
 	float Contrast;
 	float MinOcclusion;
diff --git a/Engine/Source/Runtime/Renderer/Private/DeferredShadingRenderer.h b/Engine/Source/Runtime/Renderer/Private/DeferredShadingRenderer.h
--- a/Engine/Source/Runtime/Renderer/Private/DeferredShadingRenderer.h
+++ b/Engine/Source/Runtime/Renderer/Private/DeferredShadingRenderer.h
@@ -341,7 +341,6 @@ private:
 
 	bool ShouldDoReflectionEnvironment() const;
 	
-	bool ShouldRenderDynamicSkyLight() const;
 	bool ShouldRenderDistanceFieldAO() const;
 
 	/** Whether distance field global data structures should be prepared for features that use it. */
diff --git a/Engine/Source/Runtime/Renderer/Private/DistanceFieldSurfaceCacheLighting.cpp b/Engine/Source/Runtime/Renderer/Private/DistanceFieldSurfaceCacheLighting.cpp
--- a/Engine/Source/Runtime/Renderer/Private/DistanceFieldSurfaceCacheLighting.cpp
+++ b/Engine/Source/Runtime/Renderer/Private/DistanceFieldSurfaceCacheLighting.cpp
@@ -3465,10 +3465,22 @@ bool SupportsDistanceFieldAO(ERHIFeatureLevel::Type FeatureLevel, EShaderPlatfor
 		&& DoesPlatformSupportDistanceFieldAO(ShaderPlatform);
 }
 
+bool ShouldRenderDynamicSkyLight(const FScene* Scene, const FSceneViewFamily& ViewFamily)
+{
+	return Scene->SkyLight
+		&& Scene->SkyLight->ProcessedTexture
+		&& !Scene->SkyLight->bWantsStaticShadowing
+		&& !Scene->SkyLight->bHasStaticLighting
+		&& ViewFamily.EngineShowFlags.SkyLighting
+		&& Scene->GetFeatureLevel() >= ERHIFeatureLevel::SM4
+		&& !IsSimpleDynamicLightingEnabled() 
+		&& !ViewFamily.EngineShowFlags.VisualizeLightCulling;
+}
+
 bool FDeferredShadingSceneRenderer::ShouldPrepareForDistanceFieldAO() const
 {
 	return SupportsDistanceFieldAO(Scene->GetFeatureLevel(), Scene->GetShaderPlatform())
-		&& ((ShouldRenderDynamicSkyLight() && Scene->SkyLight->bCastShadows && ViewFamily.EngineShowFlags.DistanceFieldAO)
+		&& ((ShouldRenderDynamicSkyLight(Scene, ViewFamily) && Scene->SkyLight->bCastShadows && ViewFamily.EngineShowFlags.DistanceFieldAO)
 			|| ViewFamily.EngineShowFlags.VisualizeMeshDistanceFields
 			|| ViewFamily.EngineShowFlags.VisualizeDistanceFieldAO
 			|| ViewFamily.EngineShowFlags.VisualizeDistanceFieldGI);
@@ -3972,18 +3984,6 @@ IMPLEMENT_SKYLIGHT_PS_TYPE(false, true)
 IMPLEMENT_SKYLIGHT_PS_TYPE(true, false)
 IMPLEMENT_SKYLIGHT_PS_TYPE(false, false)
 
-bool FDeferredShadingSceneRenderer::ShouldRenderDynamicSkyLight() const
-{
-	return Scene->SkyLight
-		&& Scene->SkyLight->ProcessedTexture
-		&& !Scene->SkyLight->bWantsStaticShadowing
-		&& !Scene->SkyLight->bHasStaticLighting
-		&& ViewFamily.EngineShowFlags.SkyLighting
-		&& FeatureLevel >= ERHIFeatureLevel::SM4
-		&& !IsSimpleDynamicLightingEnabled() 
-		&& !ViewFamily.EngineShowFlags.VisualizeLightCulling;
-}
-
 bool FDeferredShadingSceneRenderer::ShouldRenderDistanceFieldAO() const
 {
 	return Scene->SkyLight->bCastShadows
@@ -3995,7 +3995,7 @@ bool FDeferredShadingSceneRenderer::ShouldRenderDistanceFieldAO() const
 
 void FDeferredShadingSceneRenderer::RenderDynamicSkyLighting(FRHICommandListImmediate& RHICmdList, const TRefCountPtr<IPooledRenderTarget>& VelocityTexture, TRefCountPtr<IPooledRenderTarget>& DynamicBentNormalAO)
 {
-	if (ShouldRenderDynamicSkyLight())
+	if (ShouldRenderDynamicSkyLight(Scene, ViewFamily))
 	{
 		SCOPED_DRAW_EVENT(RHICmdList, SkyLightDiffuse);
 
diff --git a/Engine/Source/Runtime/Renderer/Private/DistanceFieldSurfaceCacheLighting.h b/Engine/Source/Runtime/Renderer/Private/DistanceFieldSurfaceCacheLighting.h
--- a/Engine/Source/Runtime/Renderer/Private/DistanceFieldSurfaceCacheLighting.h
+++ b/Engine/Source/Runtime/Renderer/Private/DistanceFieldSurfaceCacheLighting.h
@@ -505,4 +505,6 @@ private:
 	FShaderParameter DebugId;
 };
 
-extern void TrackGPUProgress(FRHICommandListImmediate& RHICmdList, uint32 DebugId);
\ No newline at end of file
+extern void TrackGPUProgress(FRHICommandListImmediate& RHICmdList, uint32 DebugId);
+
+extern bool ShouldRenderDynamicSkyLight(const FScene* Scene, const FSceneViewFamily& ViewFamily);
\ No newline at end of file
diff --git a/Engine/Source/Runtime/Renderer/Private/HeightfieldLighting.cpp b/Engine/Source/Runtime/Renderer/Private/HeightfieldLighting.cpp
--- a/Engine/Source/Runtime/Renderer/Private/HeightfieldLighting.cpp
+++ b/Engine/Source/Runtime/Renderer/Private/HeightfieldLighting.cpp
@@ -1077,6 +1077,7 @@ public:
 		GlobalHeightfieldParameters.Bind(Initializer.ParameterMap);
 		LightDirection.Bind(Initializer.ParameterMap, TEXT("LightDirection"));
 		LightColor.Bind(Initializer.ParameterMap, TEXT("LightColor"));
+		SkyLightIndirectScale.Bind(Initializer.ParameterMap, TEXT("SkyLightIndirectScale"));
 		HeightfieldShadowing.Bind(Initializer.ParameterMap, TEXT("HeightfieldShadowing"));
 		HeightfieldShadowingSampler.Bind(Initializer.ParameterMap, TEXT("HeightfieldShadowingSampler"));
 		WorldToLight.Bind(Initializer.ParameterMap, TEXT("WorldToLight"));
@@ -1089,7 +1090,8 @@ public:
 		const FLightSceneInfo& LightSceneInfo, 
 		const FMaterialRenderProxy* MaterialProxy,
 		int32 NumHeightfieldsValue, 
-		const FHeightfieldLightingAtlas& Atlas)
+		const FHeightfieldLightingAtlas& Atlas,
+		float SkyLightIndirectScaleValue)
 	{
 		const FPixelShaderRHIParamRef ShaderRHI = GetPixelShader();
 
@@ -1101,6 +1103,8 @@ public:
 		SetShaderValue(RHICmdList, ShaderRHI, LightDirection, LightSceneInfo.Proxy->GetDirection());
 		SetShaderValue(RHICmdList, ShaderRHI, LightColor, LightSceneInfo.Proxy->GetColor() * LightSceneInfo.Proxy->GetIndirectLightingScale());
 
+		SetShaderValue(RHICmdList, ShaderRHI, SkyLightIndirectScale, SkyLightIndirectScaleValue);
+
 		const FSceneViewState* ViewState = (const FSceneViewState*)View.State;
 		SetTextureParameter(RHICmdList, ShaderRHI, HeightfieldShadowing, HeightfieldShadowingSampler, TStaticSamplerState<SF_Bilinear>::GetRHI(), Atlas.DirectionalLightShadowing->GetRenderTargetItem().ShaderResourceTexture);
 	
@@ -1122,6 +1126,7 @@ public:
 		Ar << GlobalHeightfieldParameters;
 		Ar << LightDirection;
 		Ar << LightColor;
+		Ar << SkyLightIndirectScale;
 		Ar << HeightfieldShadowing;
 		Ar << HeightfieldShadowingSampler;
 		Ar << WorldToLight;
@@ -1135,6 +1140,7 @@ private:
 	FGlobalHeightfieldParameters GlobalHeightfieldParameters;
 	FShaderParameter LightDirection;
 	FShaderParameter LightColor;
+	FShaderParameter SkyLightIndirectScale;
 	FShaderResourceParameter HeightfieldShadowing;
 	FShaderResourceParameter HeightfieldShadowingSampler;
 	FShaderParameter WorldToLight;
@@ -1175,6 +1181,10 @@ void FHeightfieldLightingViewInfo::ComputeLighting(const FViewInfo& View, FRHICo
 			LightSceneInfo.Proxy->GetLightFunctionMaterial() : 
 			UMaterial::GetDefaultMaterial(MD_LightFunction)->GetRenderProxy(false);
 
+		const FScene* Scene = (const FScene*)View.Family->Scene;
+
+		const float SkyLightIndirectScale = ShouldRenderDynamicSkyLight(Scene, *View.Family) ? Scene->SkyLight->IndirectLightingIntensity : 0;
+
 		// Skip rendering if the DefaultLightFunctionMaterial isn't compiled yet
 		if (MaterialProxy->GetMaterial(FeatureLevel)->IsLightFunction())
 		{
@@ -1203,7 +1213,7 @@ void FHeightfieldLightingViewInfo::ComputeLighting(const FViewInfo& View, FRHICo
 					RHICmdList.SetLocalBoundShaderState(BoundShaderState);
 
 					VertexShader->SetParameters(RHICmdList, View, HeightfieldDescriptions.Num());
-					PixelShader->SetParameters(RHICmdList, View, LightSceneInfo, MaterialProxy, HeightfieldDescriptions.Num(), Atlas);
+					PixelShader->SetParameters(RHICmdList, View, LightSceneInfo, MaterialProxy, HeightfieldDescriptions.Num(), Atlas, SkyLightIndirectScale);
 
 					RHICmdList.DrawPrimitive(PT_TriangleList, 0, 2, HeightfieldDescriptions.Num());
 				}
