Integrate - Fixed disabled skylights affecting the scene

diff --git a/Engine/Source/Runtime/Engine/Private/SkyLightComponent.cpp b/Engine/Source/Runtime/Engine/Private/SkyLightComponent.cpp
--- a/Engine/Source/Runtime/Engine/Private/SkyLightComponent.cpp
+++ b/Engine/Source/Runtime/Engine/Private/SkyLightComponent.cpp
@@ -219,7 +219,7 @@ void USkyLightComponent::PostInterpChange(UProperty* PropertyThatChanged)
 void USkyLightComponent::DestroyRenderState_Concurrent()
 {
 	Super::DestroyRenderState_Concurrent();
-	World->Scene->SetSkyLight(NULL);
+	World->Scene->DisableSkyLight(SceneProxy);
 
 	ENQUEUE_UNIQUE_RENDER_COMMAND_ONEPARAMETER(
 		FDestroySkyLightCommand,
diff --git a/Engine/Source/Runtime/Engine/Public/SceneInterface.h b/Engine/Source/Runtime/Engine/Public/SceneInterface.h
--- a/Engine/Source/Runtime/Engine/Public/SceneInterface.h
+++ b/Engine/Source/Runtime/Engine/Public/SceneInterface.h
@@ -61,6 +61,7 @@ public:
 	 */
 	virtual void AddInvisibleLight(ULightComponent* Light) = 0;
 	virtual void SetSkyLight(FSkyLightSceneProxy* Light) = 0;
+	virtual void DisableSkyLight(FSkyLightSceneProxy* Light) = 0;
 	/** 
 	 * Adds a new decal component to the scene
 	 * 
diff --git a/Engine/Source/Runtime/Renderer/Private/RendererScene.cpp b/Engine/Source/Runtime/Renderer/Private/RendererScene.cpp
--- a/Engine/Source/Runtime/Renderer/Private/RendererScene.cpp
+++ b/Engine/Source/Runtime/Renderer/Private/RendererScene.cpp
@@ -919,6 +919,21 @@ void FScene::SetSkyLight(FSkyLightSceneProxy* LightProxy)
 	});
 }
 
+void FScene::DisableSkyLight(FSkyLightSceneProxy* LightProxy)
+{
+	ENQUEUE_UNIQUE_RENDER_COMMAND_TWOPARAMETER(
+		FDisableSkyLightCommand,
+		FScene*,Scene,this,
+		FSkyLightSceneProxy*,LightProxy,LightProxy,
+	{
+		if (Scene->SkyLight == LightProxy)
+		{
+			Scene->SkyLight = NULL;
+			Scene->bScenesPrimitivesNeedStaticMeshElementUpdate = true;
+		}
+	});
+}
+
 void FScene::AddOrRemoveDecal_RenderThread(FDeferredDecalProxy* Proxy, bool bAdd)
 {
 	if(bAdd)
@@ -2301,6 +2316,7 @@ public:
 	virtual void RemoveLight(ULightComponent* Light){}
 	virtual void AddInvisibleLight(ULightComponent* Light){}
 	virtual void SetSkyLight(FSkyLightSceneProxy* Light) {}
+	virtual void DisableSkyLight(FSkyLightSceneProxy* Light) {}
 
 	virtual void AddDecal(UDecalComponent*){}
 	virtual void RemoveDecal(UDecalComponent*){}
diff --git a/Engine/Source/Runtime/Renderer/Private/ScenePrivate.h b/Engine/Source/Runtime/Renderer/Private/ScenePrivate.h
--- a/Engine/Source/Runtime/Renderer/Private/ScenePrivate.h
+++ b/Engine/Source/Runtime/Renderer/Private/ScenePrivate.h
@@ -1516,6 +1516,7 @@ public:
 	virtual void RemoveLight(ULightComponent* Light);
 	virtual void AddInvisibleLight(ULightComponent* Light);
 	virtual void SetSkyLight(FSkyLightSceneProxy* Light);
+	virtual void DisableSkyLight(FSkyLightSceneProxy* Light);
 	virtual void AddDecal(UDecalComponent* Component);
 	virtual void RemoveDecal(UDecalComponent* Component);
 	virtual void UpdateDecalTransform(UDecalComponent* Decal) override;
@@ -1647,7 +1648,11 @@ public:
 
 	virtual void OnLevelAddedToWorld(FName InLevelName) override;
 
-	virtual bool HasAnyLights() const override { return NumVisibleLights > 0 || bHasSkyLight; }
+	virtual bool HasAnyLights() const override 
+	{ 
+		check(IsInGameThread());
+		return NumVisibleLights > 0 || bHasSkyLight; 
+	}
 
 	virtual bool IsEditorScene() const override { return bIsEditorScene; }
 
