Integrate - Logs to help track down DF hitching

diff --git a/Engine/Source/Runtime/Renderer/Private/DistanceFieldObjectManagement.cpp b/Engine/Source/Runtime/Renderer/Private/DistanceFieldObjectManagement.cpp
--- a/Engine/Source/Runtime/Renderer/Private/DistanceFieldObjectManagement.cpp
+++ b/Engine/Source/Runtime/Renderer/Private/DistanceFieldObjectManagement.cpp
@@ -14,6 +14,14 @@ FAutoConsoleVariableRef CVarAOMaxObjectBoundingRadius(
 	ECVF_Cheat | ECVF_RenderThreadSafe
 	);
 
+int32 GAOLogObjectBufferReallocation = 0;
+FAutoConsoleVariableRef CVarAOLogObjectBufferReallocation(
+	TEXT("r.AOLogObjectBufferReallocation"),
+	GAOLogObjectBufferReallocation,
+	TEXT(""),
+	ECVF_Cheat | ECVF_RenderThreadSafe
+	);
+
 #include "RendererPrivate.h"
 #include "ScenePrivate.h"
 #include "UniformBuffer.h"
@@ -569,6 +577,8 @@ void UpdateGlobalDistanceFieldObjectRemoves(FRHICommandListImmediate& RHICmdList
 		}
 		else
 		{
+			const double StartTime = FPlatformTime::Seconds();
+
 			// Have to copy the object data to allow parallel removing
 			TemporaryCopySourceBuffers = DistanceFieldSceneData.ObjectBuffers;
 			DistanceFieldSceneData.ObjectBuffers = new FDistanceFieldObjectBuffers();
@@ -607,6 +617,11 @@ void UpdateGlobalDistanceFieldObjectRemoves(FRHICommandListImmediate& RHICmdList
 
 			DistanceFieldSceneData.NumObjectsInBuffer = NumDestObjects;
 
+			if (GAOLogObjectBufferReallocation)
+			{
+				const float ElapsedTime = (float)(FPlatformTime::Seconds() - StartTime);
+				UE_LOG(LogDistanceField,Warning,TEXT("Global object buffer realloc %.3fs"), ElapsedTime);
+			}
 
 			/*
 			// Have to remove one at a time while any entries to remove are at the end of the buffer
