Integrate - Enabling a skylight from a blueprint SetVisibility call will now cause it to capture and generate its processed texture
The scene tracks enabled skylights in a stack and the most recently enabled one wins

diff --git a/Engine/Source/Runtime/Engine/Classes/Components/SkyLightComponent.h b/Engine/Source/Runtime/Engine/Classes/Components/SkyLightComponent.h
--- a/Engine/Source/Runtime/Engine/Classes/Components/SkyLightComponent.h
+++ b/Engine/Source/Runtime/Engine/Classes/Components/SkyLightComponent.h
@@ -175,6 +175,8 @@ class ENGINE_API USkyLightComponent : public ULightComponentBase
 	UFUNCTION(BlueprintCallable, Category="Rendering|Components|SkyLight")
 	void SetMinOcclusion(float InMinOcclusion);
 
+	virtual void SetVisibility(bool bNewVisibility, bool bPropagateToChildren=false) override;
+
 	/** Indicates that the capture needs to recapture the scene, adds it to the recapture queue. */
 	void SetCaptureIsDirty();
 
diff --git a/Engine/Source/Runtime/Engine/Private/SkyLightComponent.cpp b/Engine/Source/Runtime/Engine/Private/SkyLightComponent.cpp
--- a/Engine/Source/Runtime/Engine/Private/SkyLightComponent.cpp
+++ b/Engine/Source/Runtime/Engine/Private/SkyLightComponent.cpp
@@ -125,7 +125,7 @@ FSkyLightSceneProxy* USkyLightComponent::CreateSceneProxy() const
 	return NULL;
 }
 
-void USkyLightComponent::SetCaptureIsDirty() 
+void USkyLightComponent::SetCaptureIsDirty()
 { 
 	if (bVisible && bAffectsWorld)
 	{
@@ -160,8 +160,11 @@ void USkyLightComponent::CreateRenderState_Concurrent()
 		// Create the light's scene proxy.
 		SceneProxy = CreateSceneProxy();
 
-		// Add the light to the scene.
-		World->Scene->SetSkyLight(SceneProxy);
+		if (SceneProxy)
+		{
+			// Add the light to the scene.
+			World->Scene->SetSkyLight(SceneProxy);
+		}
 	}
 }
 
@@ -219,16 +222,20 @@ void USkyLightComponent::PostInterpChange(UProperty* PropertyThatChanged)
 void USkyLightComponent::DestroyRenderState_Concurrent()
 {
 	Super::DestroyRenderState_Concurrent();
-	World->Scene->DisableSkyLight(SceneProxy);
 
-	ENQUEUE_UNIQUE_RENDER_COMMAND_ONEPARAMETER(
-		FDestroySkyLightCommand,
-		FSkyLightSceneProxy*,LightSceneProxy,SceneProxy,
+	if (SceneProxy)
 	{
-		delete LightSceneProxy;
-	});
+		World->Scene->DisableSkyLight(SceneProxy);
+
+		ENQUEUE_UNIQUE_RENDER_COMMAND_ONEPARAMETER(
+			FDestroySkyLightCommand,
+			FSkyLightSceneProxy*,LightSceneProxy,SceneProxy,
+		{
+			delete LightSceneProxy;
+		});
 
-	SceneProxy = NULL;
+		SceneProxy = NULL;
+	}
 }
 
 #if WITH_EDITOR
@@ -411,7 +418,6 @@ void USkyLightComponent::UpdateSkyCaptureContents(UWorld* WorldToUpdate)
 						CaptureComponent->MarkRenderStateDirty();
 					}
 
-					//@todo - defer this until shader compilation is finished, just like reflection captures, otherwise we will capture an incomplete scene.
 					WorldToUpdate->Scene->UpdateSkyCaptureContents(CaptureComponent, false, CaptureComponent->ProcessedSkyTexture, CaptureComponent->IrradianceEnvironmentMap);
 
 					CaptureComponent->MarkRenderStateDirty();
@@ -496,6 +502,18 @@ void USkyLightComponent::SetMinOcclusion(float InMinOcclusion)
 	}
 }
 
+void USkyLightComponent::SetVisibility(bool bNewVisibility, bool bPropagateToChildren)
+{
+	const bool bOldWasVisible = bVisible;
+
+	Super::SetVisibility(bNewVisibility, bPropagateToChildren);
+
+	if (bVisible && !bOldWasVisible)
+	{
+		SetCaptureIsDirty();
+	}
+}
+
 void USkyLightComponent::RecaptureSky()
 {
 	SetCaptureIsDirty();
diff --git a/Engine/Source/Runtime/Renderer/Private/RendererScene.cpp b/Engine/Source/Runtime/Renderer/Private/RendererScene.cpp
--- a/Engine/Source/Runtime/Renderer/Private/RendererScene.cpp
+++ b/Engine/Source/Runtime/Renderer/Private/RendererScene.cpp
@@ -383,8 +383,8 @@ FScene::FScene(UWorld* InWorld, bool bInRequiresHitProxies, bool bInIsEditorScen
 ,	UpperDynamicSkylightColor(FLinearColor::Black)
 ,	LowerDynamicSkylightColor(FLinearColor::Black)
 ,	SceneLODHierarchy(this)
-,	NumVisibleLights(0)
-,	bHasSkyLight(false)
+,	NumVisibleLights_GameThread(0)
+,	NumEnabledSkylights_GameThread(0)
 {
 	check(World);
 	World->Scene = this;
@@ -855,7 +855,7 @@ void FScene::AddLight(ULightComponent* Light)
 		INC_DWORD_STAT(STAT_SceneLights);
 
 		// Adding a new light
-		++NumVisibleLights;
+		++NumVisibleLights_GameThread;
 
 		// Send a command to the rendering thread to add the light to the scene.
 		ENQUEUE_UNIQUE_RENDER_COMMAND_TWOPARAMETER(
@@ -901,7 +901,8 @@ void FScene::AddInvisibleLight(ULightComponent* Light)
 
 void FScene::SetSkyLight(FSkyLightSceneProxy* LightProxy)
 {
-	bHasSkyLight = LightProxy != NULL;
+	check(LightProxy);
+	NumEnabledSkylights_GameThread++;
 
 	// Send a command to the rendering thread to add the light to the scene.
 	ENQUEUE_UNIQUE_RENDER_COMMAND_TWOPARAMETER(
@@ -909,26 +910,49 @@ void FScene::SetSkyLight(FSkyLightSceneProxy* LightProxy)
 		FScene*,Scene,this,
 		FSkyLightSceneProxy*,LightProxy,LightProxy,
 	{
-		// Mark the scene as needing static draw lists to be recreated if needed
-		// The base pass chooses shaders based on whether there's a skylight in the scene, and that is cached in static draw lists
-		if ((Scene->SkyLight == NULL) != (LightProxy == NULL))
+		check(!Scene->SkyLightStack.Contains(LightProxy));
+		Scene->SkyLightStack.Push(LightProxy);
+		const bool bHadSkylight = Scene->SkyLight != NULL;
+
+		// Use the most recently enabled skylight
+		Scene->SkyLight = LightProxy;
+
+		if (!bHadSkylight)
 		{
+			// Mark the scene as needing static draw lists to be recreated if needed
+			// The base pass chooses shaders based on whether there's a skylight in the scene, and that is cached in static draw lists
 			Scene->bScenesPrimitivesNeedStaticMeshElementUpdate = true;
 		}
-		Scene->SkyLight = LightProxy;
 	});
 }
 
 void FScene::DisableSkyLight(FSkyLightSceneProxy* LightProxy)
 {
+	check(LightProxy);
+	NumEnabledSkylights_GameThread--;
+
 	ENQUEUE_UNIQUE_RENDER_COMMAND_TWOPARAMETER(
 		FDisableSkyLightCommand,
 		FScene*,Scene,this,
 		FSkyLightSceneProxy*,LightProxy,LightProxy,
 	{
-		if (Scene->SkyLight == LightProxy)
+		const bool bHadSkylight = Scene->SkyLight != NULL;
+
+		Scene->SkyLightStack.RemoveSingle(LightProxy);
+
+		if (Scene->SkyLightStack.Num() > 0)
+		{
+			// Use the most recently enabled skylight
+			Scene->SkyLight = Scene->SkyLightStack.Last();
+		}
+		else
 		{
 			Scene->SkyLight = NULL;
+		}
+
+		// Update the scene if we switched skylight enabled states
+		if ((Scene->SkyLight != NULL) != bHadSkylight)
+		{
 			Scene->bScenesPrimitivesNeedStaticMeshElementUpdate = true;
 		}
 	});
@@ -1334,7 +1358,7 @@ void FScene::RemoveLight(ULightComponent* Light)
 		DEC_DWORD_STAT(STAT_SceneLights);
 
 		// Removing one visible light
-		--NumVisibleLights;
+		--NumVisibleLights_GameThread;
 
 		// Disassociate the primitive's render info.
 		Light->SceneProxy = NULL;
diff --git a/Engine/Source/Runtime/Renderer/Private/ScenePrivate.h b/Engine/Source/Runtime/Renderer/Private/ScenePrivate.h
--- a/Engine/Source/Runtime/Renderer/Private/ScenePrivate.h
+++ b/Engine/Source/Runtime/Renderer/Private/ScenePrivate.h
@@ -1303,6 +1303,11 @@ private:
 
 typedef TMap<FMaterial*, FMaterialShaderMap*> FMaterialsToUpdateMap;
 
+/** 
+ * Renderer scene which is private to the renderer module.
+ * Ordinarily this is the renderer version of a UWorld, but an FScene can be created for previewing in editors which don't have a UWorld as well.
+ * The scene stores renderer state that is independent of any view or frame, with the primary actions being adding and removing of primitives and lights.
+ */
 class FScene : public FSceneInterface
 {
 public:
@@ -1420,6 +1425,9 @@ public:
 	/** The scene's sky light, if any. */
 	FSkyLightSceneProxy* SkyLight;
 
+	/** Used to track the order that skylights were enabled in. */
+	TArray<FSkyLightSceneProxy*> SkyLightStack;
+
 	/** The directional light to use for simple dynamic lighting, if any. */
 	FLightSceneInfo* SimpleDirectionalLight;
 
@@ -1651,7 +1659,7 @@ public:
 	virtual bool HasAnyLights() const override 
 	{ 
 		check(IsInGameThread());
-		return NumVisibleLights > 0 || bHasSkyLight; 
+		return NumVisibleLights_GameThread > 0 || NumEnabledSkylights_GameThread > 0; 
 	}
 
 	virtual bool IsEditorScene() const override { return bIsEditorScene; }
@@ -1772,13 +1780,13 @@ private:
 	 * The number of visible lights in the scene
 	 * Note: This is tracked on the game thread!
 	 */
-	int32 NumVisibleLights;
+	int32 NumVisibleLights_GameThread;
 
 	/** 
 	 * Whether the scene has a valid sky light.
 	 * Note: This is tracked on the game thread!
 	 */
-	bool bHasSkyLight;
+	int32 NumEnabledSkylights_GameThread;
 
 	/** This scene's feature level */
 	ERHIFeatureLevel::Type FeatureLevel;
