Integrate - Skylight MinOcclusion works on sky specular occlusion as well

diff --git a/Engine/Shaders/SkyLightingShared.usf b/Engine/Shaders/SkyLightingShared.usf
--- a/Engine/Shaders/SkyLightingShared.usf
+++ b/Engine/Shaders/SkyLightingShared.usf
@@ -21,6 +21,7 @@ SamplerState BentNormalAOSampler;
 
 float ApplyBentNormalAO;
 float InvSkySpecularOcclusionStrength;
+float MinSkySpecularOcclusion;
 
 float GetDistanceFieldAOSpecularOcclusion(float2 ScreenUV, float3 ReflectionVector, float Roughness, bool bTwoSidedFoliage)
 {
@@ -46,6 +47,7 @@ float GetDistanceFieldAOSpecularOcclusion(float2 ScreenUV, float3 ReflectionVect
 
 			// Can't rely on the direction of the bent normal when close to fully occluded, lerp to shadowed
 			Visibility = lerp(0, Visibility, saturate((UnoccludedAngle - .1f) / .2f));
+			Visibility = lerp(Visibility, 1, MinSkySpecularOcclusion);
 		}
 	}
 
diff --git a/Engine/Source/Runtime/Renderer/Private/ReflectionEnvironment.cpp b/Engine/Source/Runtime/Renderer/Private/ReflectionEnvironment.cpp
--- a/Engine/Source/Runtime/Renderer/Private/ReflectionEnvironment.cpp
+++ b/Engine/Source/Runtime/Renderer/Private/ReflectionEnvironment.cpp
@@ -131,10 +131,11 @@ public:
 		BentNormalAOSampler.Bind(ParameterMap, TEXT("BentNormalAOSampler"));
 		ApplyBentNormalAO.Bind(ParameterMap, TEXT("ApplyBentNormalAO"));
 		InvSkySpecularOcclusionStrength.Bind(ParameterMap, TEXT("InvSkySpecularOcclusionStrength"));
+		MinSkySpecularOcclusion.Bind(ParameterMap, TEXT("MinSkySpecularOcclusion"));
 	}
 
 	template<typename ShaderRHIParamRef>
-	void SetParameters(FRHICommandList& RHICmdList, const ShaderRHIParamRef& ShaderRHI, const TRefCountPtr<IPooledRenderTarget>& DynamicBentNormalAO, float SkySpecularOcclusionStrength)
+	void SetParameters(FRHICommandList& RHICmdList, const ShaderRHIParamRef& ShaderRHI, const TRefCountPtr<IPooledRenderTarget>& DynamicBentNormalAO, float SkySpecularOcclusionStrength, float MinOcclusionValue)
 	{
 		FTextureRHIParamRef BentNormalAO = GWhiteTexture->TextureRHI;
 		bool bApplyBentNormalAO = false;
@@ -148,11 +149,12 @@ public:
 		SetTextureParameter(RHICmdList, ShaderRHI, BentNormalAOTexture, BentNormalAOSampler, TStaticSamplerState<SF_Point>::GetRHI(), BentNormalAO);
 		SetShaderValue(RHICmdList, ShaderRHI, ApplyBentNormalAO, bApplyBentNormalAO ? 1.0f : 0.0f);
 		SetShaderValue(RHICmdList, ShaderRHI, InvSkySpecularOcclusionStrength, 1.0f / FMath::Max(SkySpecularOcclusionStrength, .1f));
+		SetShaderValue(RHICmdList, ShaderRHI, MinSkySpecularOcclusion, MinOcclusionValue);
 	}
 
 	friend FArchive& operator<<(FArchive& Ar,FDistanceFieldAOSpecularOcclusionParameters& P)
 	{
-		Ar << P.BentNormalAOTexture << P.BentNormalAOSampler << P.ApplyBentNormalAO << P.InvSkySpecularOcclusionStrength;
+		Ar << P.BentNormalAOTexture << P.BentNormalAOSampler << P.ApplyBentNormalAO << P.InvSkySpecularOcclusionStrength << P.MinSkySpecularOcclusion;
 		return Ar;
 	}
 
@@ -161,6 +163,7 @@ private:
 	FShaderResourceParameter BentNormalAOSampler;
 	FShaderParameter ApplyBentNormalAO;
 	FShaderParameter InvSkySpecularOcclusionStrength;
+	FShaderParameter MinSkySpecularOcclusion;
 };
 
 struct FReflectionCaptureSortData
@@ -333,7 +336,9 @@ public:
 		SetTextureParameter(RHICmdList, ShaderRHI, PreIntegratedGF, PreIntegratedGFSampler, TStaticSamplerState<SF_Bilinear,AM_Clamp,AM_Clamp,AM_Clamp>::GetRHI(), GSystemTextures.PreintegratedGF->GetRenderTargetItem().ShaderResourceTexture);
 	
 		SkyLightParameters.SetParameters(RHICmdList, ShaderRHI, Scene, View.Family->EngineShowFlags.SkyLighting);
-		SpecularOcclusionParameters.SetParameters(RHICmdList, ShaderRHI, DynamicBentNormalAO, CVarSkySpecularOcclusionStrength.GetValueOnRenderThread());
+
+		const float MinOcclusion = Scene->SkyLight ? Scene->SkyLight->MinOcclusion : 0;
+		SpecularOcclusionParameters.SetParameters(RHICmdList, ShaderRHI, DynamicBentNormalAO, CVarSkySpecularOcclusionStrength.GetValueOnRenderThread(), MinOcclusion);
 	}
 
 	void UnsetParameters(FRHICommandList& RHICmdList)
@@ -455,7 +460,9 @@ public:
 		SetTextureParameter(RHICmdList, ShaderRHI, ScreenSpaceReflectionsTexture, ScreenSpaceReflectionsSampler, TStaticSamplerState<SF_Point>::GetRHI(), ScreenSpaceReflections );
 		SetTextureParameter(RHICmdList, ShaderRHI, PreIntegratedGF, PreIntegratedGFSampler, TStaticSamplerState<SF_Bilinear,AM_Clamp,AM_Clamp,AM_Clamp>::GetRHI(), GSystemTextures.PreintegratedGF->GetRenderTargetItem().ShaderResourceTexture );
 		
-		SpecularOcclusionParameters.SetParameters(RHICmdList, ShaderRHI, DynamicBentNormalAO, CVarSkySpecularOcclusionStrength.GetValueOnRenderThread());
+		FScene* Scene = (FScene*)View.Family->Scene;
+		const float MinOcclusion = Scene->SkyLight ? Scene->SkyLight->MinOcclusion : 0;
+		SpecularOcclusionParameters.SetParameters(RHICmdList, ShaderRHI, DynamicBentNormalAO, CVarSkySpecularOcclusionStrength.GetValueOnRenderThread(), MinOcclusion);
 	}
 
 	// FShader interface.
