Integrate - Surfel GI and Heightfield GI can now be enabled at the same time

diff --git a/Engine/Shaders/DistanceFieldGlobalIllumination.usf b/Engine/Shaders/DistanceFieldGlobalIllumination.usf
--- a/Engine/Shaders/DistanceFieldGlobalIllumination.usf
+++ b/Engine/Shaders/DistanceFieldGlobalIllumination.usf
@@ -572,7 +572,8 @@ void LightVPLsCS(
 	}
 }
 
-RWBuffer<float4> RWIrradianceCacheIrradiance;
+RWBuffer<float4> RWSurfelIrradiance;
+RWBuffer<float4> RWHeightfieldIrradiance;
 
 /**  */
 [numthreads(FINAL_GATHER_THREADGROUP_SIZE, 1, 1)]
@@ -581,11 +582,10 @@ void ClearIrradianceSamplesCS(
 	uint3 DispatchThreadId : SV_DispatchThreadID,
     uint3 GroupThreadId : SV_GroupThreadID) 
 {
-	uint StartIndex = SavedStartIndex[0];
-	uint NumRecords = ScatterDrawParameters[1];
-	uint RecordIndex = StartIndex + DispatchThreadId.x;
+	uint RelativeRecordIndex = DispatchThreadId.x;
 
-	RWIrradianceCacheIrradiance[RecordIndex] = 0;
+	RWSurfelIrradiance[RelativeRecordIndex] = 0;
+	RWHeightfieldIrradiance[RelativeRecordIndex] = 0;
 }
 
 float RayTraceThroughGlobalObjects(
@@ -972,6 +972,7 @@ void ComputeIrradianceCS(
 	uint NumRecords = ScatterDrawParameters[1];
 
 	uint RecordIndex = StartIndex + GroupId.x;
+	uint RelativeRecordIndex = GroupId.x;
 	uint ThreadIndex = GroupThreadId.x;
 
 	float3 Irradiance = 0;
@@ -982,7 +983,7 @@ void ComputeIrradianceCS(
 
 		float4 PositionAndRadius = IrradianceCachePositionRadius[RecordIndex];
 		float3 RecordWorldPosition = PositionAndRadius.xyz;
-		Irradiance = GatherIrradianceFromVPLs(RecordWorldPosition, RecordWorldNormal, GroupId.x, RecordIndex, ThreadIndex);
+		Irradiance = GatherIrradianceFromVPLs(RecordWorldPosition, RecordWorldNormal, RelativeRecordIndex, RecordIndex, ThreadIndex);
 	}
 
 	SharedThreadIrradiance[ThreadIndex] = Irradiance;
@@ -998,6 +999,29 @@ void ComputeIrradianceCS(
 			Irradiance += SharedThreadIrradiance[i];
 		}
 
-		RWIrradianceCacheIrradiance[RecordIndex] = float4(Irradiance, 0);
+		RWSurfelIrradiance[RelativeRecordIndex] = float4(Irradiance, 0);
 	}
 }
+
+RWBuffer<float4> RWIrradianceCacheIrradiance;
+Buffer<float4> SurfelIrradiance;
+Buffer<float4> HeightfieldIrradiance;
+
+/**  */
+[numthreads(FINAL_GATHER_THREADGROUP_SIZE, 1, 1)]
+void CombineIrradianceSamplesCS(
+	uint3 GroupId : SV_GroupID,
+	uint3 DispatchThreadId : SV_DispatchThreadID,
+    uint3 GroupThreadId : SV_GroupThreadID) 
+{
+	uint StartIndex = SavedStartIndex[0];
+	uint NumRecords = ScatterDrawParameters[1];
+
+	uint RecordIndex = StartIndex + DispatchThreadId.x;
+	uint RelativeRecordIndex = DispatchThreadId.x;
+
+	if (RecordIndex < NumRecords)
+	{
+		RWIrradianceCacheIrradiance[RecordIndex] = SurfelIrradiance[RelativeRecordIndex] + HeightfieldIrradiance[RelativeRecordIndex];
+	}
+}
\ No newline at end of file
diff --git a/Engine/Shaders/HeightfieldLighting.usf b/Engine/Shaders/HeightfieldLighting.usf
--- a/Engine/Shaders/HeightfieldLighting.usf
+++ b/Engine/Shaders/HeightfieldLighting.usf
@@ -454,8 +454,6 @@ void CalculateHeightfieldOcclusionCS(
 	}
 }
 
-RWBuffer<float4> RWIrradianceCacheIrradiance;
-
 float3 SampleHeightfieldLighting(float2 GlobalUV)
 {
 	return Texture2DSampleLevel(HeightfieldLighting, HeightfieldLightingSampler, GlobalUV, 0).xyz;
@@ -593,6 +591,8 @@ float3 GatherHeightfieldLighting(
 	return Lighting;
 }
 
+RWBuffer<float4> RWHeightfieldIrradiance;
+
 groupshared float3 SharedThreadIrradiance[FINAL_GATHER_THREADGROUP_SIZE];
 
 float InnerLightTransferDistance;
@@ -608,6 +608,7 @@ void CalculateHeightfieldIrradianceCS(
 	uint StartIndex = SavedStartIndex[0];
 	uint NumRecords = ScatterDrawParameters[1];
 	uint RecordIndex = StartIndex + GroupId.x;
+	uint RelativeRecordIndex = GroupId.x;
 	uint ThreadIndex = GroupThreadId.x;
 
 	float4 PositionAndRadius = IrradianceCachePositionRadius[RecordIndex];
@@ -716,6 +717,6 @@ void CalculateHeightfieldIrradianceCS(
 			Irradiance += SharedThreadIrradiance[i];
 		}
 
-		RWIrradianceCacheIrradiance[RecordIndex] = float4(Irradiance, 0);
+		RWHeightfieldIrradiance[RelativeRecordIndex] = float4(Irradiance, 0);
 	}
 }
\ No newline at end of file
diff --git a/Engine/Source/Runtime/Renderer/Private/DistanceFieldGlobalIllumination.cpp b/Engine/Source/Runtime/Renderer/Private/DistanceFieldGlobalIllumination.cpp
--- a/Engine/Source/Runtime/Renderer/Private/DistanceFieldGlobalIllumination.cpp
+++ b/Engine/Source/Runtime/Renderer/Private/DistanceFieldGlobalIllumination.cpp
@@ -25,7 +25,7 @@ FAutoConsoleVariableRef CVarDistanceFieldGI(
 	ECVF_Cheat | ECVF_RenderThreadSafe | ECVF_ReadOnly
 	);
 
-int32 GVPLMeshGlobalIllumination = 0;
+int32 GVPLMeshGlobalIllumination = 1;
 FAutoConsoleVariableRef CVarVPLMeshGlobalIllumination(
 	TEXT("r.VPLMeshGlobalIllumination"),
 	GVPLMeshGlobalIllumination,
@@ -889,7 +889,8 @@ public:
 	FClearIrradianceSamplesCS(const ShaderMetaType::CompiledShaderInitializerType& Initializer)
 		: FGlobalShader(Initializer)
 	{
-		IrradianceCacheIrradiance.Bind(Initializer.ParameterMap, TEXT("IrradianceCacheIrradiance"));
+		SurfelIrradiance.Bind(Initializer.ParameterMap, TEXT("SurfelIrradiance"));
+		HeightfieldIrradiance.Bind(Initializer.ParameterMap, TEXT("HeightfieldIrradiance"));
 		ScatterDrawParameters.Bind(Initializer.ParameterMap, TEXT("ScatterDrawParameters"));
 		SavedStartIndex.Bind(Initializer.ParameterMap, TEXT("SavedStartIndex"));
 	}
@@ -913,18 +914,21 @@ public:
 		SetSRVParameter(RHICmdList, ShaderRHI, ScatterDrawParameters, SurfaceCacheResources.Level[DepthLevel]->ScatterDrawParameters.SRV);
 		SetSRVParameter(RHICmdList, ShaderRHI, SavedStartIndex, SurfaceCacheResources.Level[DepthLevel]->SavedStartIndex.SRV);
 
-		IrradianceCacheIrradiance.SetBuffer(RHICmdList, ShaderRHI, SurfaceCacheResources.Level[DepthLevel]->Irradiance);
+		SurfelIrradiance.SetBuffer(RHICmdList, ShaderRHI, TemporaryIrradianceCacheResources->SurfelIrradiance);
+		HeightfieldIrradiance.SetBuffer(RHICmdList, ShaderRHI, TemporaryIrradianceCacheResources->HeightfieldIrradiance);
 	}
 
 	void UnsetParameters(FRHICommandList& RHICmdList)
 	{
-		IrradianceCacheIrradiance.UnsetUAV(RHICmdList, GetComputeShader());
+		SurfelIrradiance.UnsetUAV(RHICmdList, GetComputeShader());
+		HeightfieldIrradiance.UnsetUAV(RHICmdList, GetComputeShader());
 	}
 
 	virtual bool Serialize(FArchive& Ar)
 	{		
 		bool bShaderHasOutdatedParameters = FGlobalShader::Serialize(Ar);
-		Ar << IrradianceCacheIrradiance;
+		Ar << SurfelIrradiance;
+		Ar << HeightfieldIrradiance;
 		Ar << ScatterDrawParameters;
 		Ar << SavedStartIndex;
 		return bShaderHasOutdatedParameters;
@@ -932,7 +936,8 @@ public:
 
 private:
 
-	FRWShaderParameter IrradianceCacheIrradiance;
+	FRWShaderParameter SurfelIrradiance;
+	FRWShaderParameter HeightfieldIrradiance;
 	FShaderResourceParameter ScatterDrawParameters;
 	FShaderResourceParameter SavedStartIndex;
 };
@@ -1074,7 +1079,7 @@ public:
 		: FGlobalShader(Initializer)
 	{
 		AOParameters.Bind(Initializer.ParameterMap);
-		IrradianceCacheIrradiance.Bind(Initializer.ParameterMap, TEXT("IrradianceCacheIrradiance"));
+		SurfelIrradiance.Bind(Initializer.ParameterMap, TEXT("SurfelIrradiance"));
 		IrradianceCachePositionRadius.Bind(Initializer.ParameterMap, TEXT("IrradianceCachePositionRadius"));
 		IrradianceCacheNormal.Bind(Initializer.ParameterMap, TEXT("IrradianceCacheNormal"));
 		IrradianceCacheTileCoordinate.Bind(Initializer.ParameterMap, TEXT("IrradianceCacheTileCoordinate"));
@@ -1119,7 +1124,7 @@ public:
 		SetSRVParameter(RHICmdList, ShaderRHI, ScatterDrawParameters, SurfaceCacheResources.Level[DepthLevel]->ScatterDrawParameters.SRV);
 		SetSRVParameter(RHICmdList, ShaderRHI, SavedStartIndex, SurfaceCacheResources.Level[DepthLevel]->SavedStartIndex.SRV);
 
-		IrradianceCacheIrradiance.SetBuffer(RHICmdList, ShaderRHI, SurfaceCacheResources.Level[DepthLevel]->Irradiance);
+		SurfelIrradiance.SetBuffer(RHICmdList, ShaderRHI, TemporaryIrradianceCacheResources->SurfelIrradiance);
 
 		GlobalObjectParameters.Set(RHICmdList, ShaderRHI, *Scene->DistanceFieldSceneData.ObjectBuffers, Scene->DistanceFieldSceneData.NumObjectsInBuffer);
 		extern TGlobalResource<FDistanceFieldObjectBufferResource> GAOCulledObjectBuffers;
@@ -1143,14 +1148,14 @@ public:
 
 	void UnsetParameters(FRHICommandList& RHICmdList)
 	{
-		IrradianceCacheIrradiance.UnsetUAV(RHICmdList, GetComputeShader());
+		SurfelIrradiance.UnsetUAV(RHICmdList, GetComputeShader());
 	}
 
 	virtual bool Serialize(FArchive& Ar)
 	{		
 		bool bShaderHasOutdatedParameters = FGlobalShader::Serialize(Ar);
 		Ar << AOParameters;
-		Ar << IrradianceCacheIrradiance;
+		Ar << SurfelIrradiance;
 		Ar << IrradianceCachePositionRadius;
 		Ar << IrradianceCacheNormal;
 		Ar << IrradianceCacheTileCoordinate;
@@ -1174,7 +1179,7 @@ public:
 private:
 
 	FAOParameters AOParameters;
-	FRWShaderParameter IrradianceCacheIrradiance;
+	FRWShaderParameter SurfelIrradiance;
 	FShaderResourceParameter IrradianceCachePositionRadius;
 	FShaderResourceParameter IrradianceCacheNormal;
 	FShaderResourceParameter IrradianceCacheTileCoordinate;
@@ -1198,6 +1203,84 @@ IMPLEMENT_SHADER_TYPE(template<>,TComputeIrradianceCS<VPLMode_PlaceFromLight>,TE
 IMPLEMENT_SHADER_TYPE(template<>,TComputeIrradianceCS<VPLMode_PlaceFromSurfels>,TEXT("DistanceFieldGlobalIllumination"),TEXT("ComputeIrradianceCS"),SF_Compute);
 
 
+class FCombineIrradianceSamplesCS : public FGlobalShader
+{
+	DECLARE_SHADER_TYPE(FCombineIrradianceSamplesCS,Global)
+public:
+
+	static bool ShouldCache(EShaderPlatform Platform)
+	{
+		return IsFeatureLevelSupported(Platform, ERHIFeatureLevel::SM5) && DoesPlatformSupportDistanceFieldAO(Platform);
+	}
+
+	static void ModifyCompilationEnvironment(EShaderPlatform Platform, FShaderCompilerEnvironment& OutEnvironment)
+	{
+		FGlobalShader::ModifyCompilationEnvironment(Platform,OutEnvironment);
+	}
+
+	FCombineIrradianceSamplesCS(const ShaderMetaType::CompiledShaderInitializerType& Initializer)
+		: FGlobalShader(Initializer)
+	{
+		IrradianceCacheIrradiance.Bind(Initializer.ParameterMap, TEXT("IrradianceCacheIrradiance"));
+		SurfelIrradiance.Bind(Initializer.ParameterMap, TEXT("SurfelIrradiance"));
+		HeightfieldIrradiance.Bind(Initializer.ParameterMap, TEXT("HeightfieldIrradiance"));
+		ScatterDrawParameters.Bind(Initializer.ParameterMap, TEXT("ScatterDrawParameters"));
+		SavedStartIndex.Bind(Initializer.ParameterMap, TEXT("SavedStartIndex"));
+	}
+
+	FCombineIrradianceSamplesCS()
+	{
+	}
+
+	void SetParameters(
+		FRHICommandList& RHICmdList, 
+		const FSceneView& View, 
+		int32 DepthLevel, 
+		FTemporaryIrradianceCacheResources* TemporaryIrradianceCacheResources)
+	{
+		FComputeShaderRHIParamRef ShaderRHI = GetComputeShader();
+		FGlobalShader::SetParameters(RHICmdList, ShaderRHI, View);
+
+		const FScene* Scene = (const FScene*)View.Family->Scene;
+		FSurfaceCacheResources& SurfaceCacheResources = *Scene->SurfaceCacheResources;
+
+		SetSRVParameter(RHICmdList, ShaderRHI, SurfelIrradiance, TemporaryIrradianceCacheResources->SurfelIrradiance.SRV);
+		SetSRVParameter(RHICmdList, ShaderRHI, HeightfieldIrradiance,TemporaryIrradianceCacheResources->HeightfieldIrradiance.SRV);
+
+		SetSRVParameter(RHICmdList, ShaderRHI, ScatterDrawParameters, SurfaceCacheResources.Level[DepthLevel]->ScatterDrawParameters.SRV);
+		SetSRVParameter(RHICmdList, ShaderRHI, SavedStartIndex, SurfaceCacheResources.Level[DepthLevel]->SavedStartIndex.SRV);
+
+		IrradianceCacheIrradiance.SetBuffer(RHICmdList, ShaderRHI, SurfaceCacheResources.Level[DepthLevel]->Irradiance);
+	}
+
+	void UnsetParameters(FRHICommandList& RHICmdList)
+	{
+		IrradianceCacheIrradiance.UnsetUAV(RHICmdList, GetComputeShader());
+	}
+
+	virtual bool Serialize(FArchive& Ar)
+	{		
+		bool bShaderHasOutdatedParameters = FGlobalShader::Serialize(Ar);
+		Ar << IrradianceCacheIrradiance;
+		Ar << SurfelIrradiance;
+		Ar << HeightfieldIrradiance;
+		Ar << ScatterDrawParameters;
+		Ar << SavedStartIndex;
+		return bShaderHasOutdatedParameters;
+	}
+
+private:
+
+	FRWShaderParameter IrradianceCacheIrradiance;
+	FShaderResourceParameter SurfelIrradiance;
+	FShaderResourceParameter HeightfieldIrradiance;
+	FShaderResourceParameter ScatterDrawParameters;
+	FShaderResourceParameter SavedStartIndex;
+};
+
+IMPLEMENT_SHADER_TYPE(,FCombineIrradianceSamplesCS,TEXT("DistanceFieldGlobalIllumination"),TEXT("CombineIrradianceSamplesCS"),SF_Compute);
+
+
 void ComputeIrradianceForSamples(
 	int32 DepthLevel,
 	FRHICommandListImmediate& RHICmdList,
@@ -1271,6 +1354,22 @@ void ComputeIrradianceForSamples(
 			ComputeShader->UnsetParameters(RHICmdList);
 		}
 	}
+
+	{	
+		TShaderMapRef<TSetupFinalGatherIndirectArgumentsCS<false> > ComputeShader(View.ShaderMap);
+		RHICmdList.SetComputeShader(ComputeShader->GetComputeShader());
+		ComputeShader->SetParameters(RHICmdList, View, DepthLevel);
+		DispatchComputeShader(RHICmdList, *ComputeShader, 1, 1, 1);
+		ComputeShader->UnsetParameters(RHICmdList);
+	}
+		
+	{
+		TShaderMapRef<FCombineIrradianceSamplesCS> ComputeShader(View.ShaderMap);
+		RHICmdList.SetComputeShader(ComputeShader->GetComputeShader());
+		ComputeShader->SetParameters(RHICmdList, View, DepthLevel, TemporaryIrradianceCacheResources);
+		DispatchIndirectComputeShader(RHICmdList, *ComputeShader, SurfaceCacheResources.DispatchParameters.Buffer, 0);
+		ComputeShader->UnsetParameters(RHICmdList);
+	}
 }
 
 void ListDistanceFieldGIMemory(const FViewInfo& View)
diff --git a/Engine/Source/Runtime/Renderer/Private/DistanceFieldSurfaceCacheLighting.h b/Engine/Source/Runtime/Renderer/Private/DistanceFieldSurfaceCacheLighting.h
--- a/Engine/Source/Runtime/Renderer/Private/DistanceFieldSurfaceCacheLighting.h
+++ b/Engine/Source/Runtime/Renderer/Private/DistanceFieldSurfaceCacheLighting.h
@@ -435,6 +435,8 @@ public:
 			ConeVisibility.Initialize(sizeof(float), MaxSize * NumConeSampleDirections, PF_R32_FLOAT, BUF_Static);
 			ConeData.Initialize(sizeof(float), MaxSize * NumConeSampleDirections * RecordConeDataStride, PF_R32_FLOAT, BUF_Static);
 			StepBentNormal.Initialize(sizeof(float) * 4, MaxSize * NumVisibilitySteps, PF_A32B32G32R32F, BUF_Static);
+			SurfelIrradiance.Initialize(sizeof(FFloat16Color), MaxSize, PF_FloatRGBA, BUF_Static);
+			HeightfieldIrradiance.Initialize(sizeof(FFloat16Color), MaxSize, PF_FloatRGBA, BUF_Static);
 		}
 	}
 
@@ -443,6 +445,8 @@ public:
 		ConeVisibility.Release();
 		ConeData.Release();
 		StepBentNormal.Release();
+		SurfelIrradiance.Release();
+		HeightfieldIrradiance.Release();
 	}
 
 	size_t GetSizeBytes() const
@@ -453,6 +457,8 @@ public:
 	FRWBuffer ConeVisibility;
 	FRWBuffer ConeData;
 	FRWBuffer StepBentNormal;
+	FRWBuffer SurfelIrradiance;
+	FRWBuffer HeightfieldIrradiance;
 };
 
 extern FRWBuffer* GDebugBuffer2;
diff --git a/Engine/Source/Runtime/Renderer/Private/HeightfieldLighting.cpp b/Engine/Source/Runtime/Renderer/Private/HeightfieldLighting.cpp
--- a/Engine/Source/Runtime/Renderer/Private/HeightfieldLighting.cpp
+++ b/Engine/Source/Runtime/Renderer/Private/HeightfieldLighting.cpp
@@ -1472,7 +1472,7 @@ public:
 		AOParameters.Bind(Initializer.ParameterMap);
 		HeightfieldDescriptionParameters.Bind(Initializer.ParameterMap);
 		GlobalHeightfieldParameters.Bind(Initializer.ParameterMap);
-		IrradianceCacheIrradiance.Bind(Initializer.ParameterMap, TEXT("IrradianceCacheIrradiance"));
+		HeightfieldIrradiance.Bind(Initializer.ParameterMap, TEXT("HeightfieldIrradiance"));
 		IrradianceCacheNormal.Bind(Initializer.ParameterMap, TEXT("IrradianceCacheNormal"));
 		IrradianceCachePositionRadius.Bind(Initializer.ParameterMap, TEXT("IrradianceCachePositionRadius"));
 		IrradianceCacheBentNormal.Bind(Initializer.ParameterMap, TEXT("IrradianceCacheBentNormal"));
@@ -1519,7 +1519,7 @@ public:
 		SetSRVParameter(RHICmdList, ShaderRHI, ScatterDrawParameters, SurfaceCacheResources.Level[DepthLevel]->ScatterDrawParameters.SRV);
 		SetSRVParameter(RHICmdList, ShaderRHI, SavedStartIndex, SurfaceCacheResources.Level[DepthLevel]->SavedStartIndex.SRV);
 
-		IrradianceCacheIrradiance.SetBuffer(RHICmdList, ShaderRHI, SurfaceCacheResources.Level[DepthLevel]->Irradiance);
+		HeightfieldIrradiance.SetBuffer(RHICmdList, ShaderRHI, TemporaryIrradianceCacheResources.HeightfieldIrradiance);
 
 		{
 			FAOSampleData2 AOSampleData;
@@ -1581,7 +1581,7 @@ public:
 
 	void UnsetParameters(FRHICommandList& RHICmdList)
 	{
-		IrradianceCacheIrradiance.UnsetUAV(RHICmdList, GetComputeShader());
+		HeightfieldIrradiance.UnsetUAV(RHICmdList, GetComputeShader());
 	}
 
 	virtual bool Serialize(FArchive& Ar)
@@ -1590,7 +1590,7 @@ public:
 		Ar << AOParameters;
 		Ar << HeightfieldDescriptionParameters;
 		Ar << GlobalHeightfieldParameters;
-		Ar << IrradianceCacheIrradiance;
+		Ar << HeightfieldIrradiance;
 		Ar << IrradianceCacheNormal;
 		Ar << IrradianceCachePositionRadius;
 		Ar << IrradianceCacheBentNormal;
@@ -1614,7 +1614,7 @@ private:
 	FAOParameters AOParameters;
 	FHeightfieldDescriptionParameters HeightfieldDescriptionParameters;
 	FGlobalHeightfieldParameters GlobalHeightfieldParameters;
-	FRWShaderParameter IrradianceCacheIrradiance;
+	FRWShaderParameter HeightfieldIrradiance;
 	FShaderResourceParameter IrradianceCacheNormal;
 	FShaderResourceParameter IrradianceCachePositionRadius;
 	FShaderResourceParameter IrradianceCacheBentNormal;
