Profiling tweaks

diff --git a/Engine/Source/Runtime/RHI/Private/GPUProfiler.cpp b/Engine/Source/Runtime/RHI/Private/GPUProfiler.cpp
--- a/Engine/Source/Runtime/RHI/Private/GPUProfiler.cpp
+++ b/Engine/Source/Runtime/RHI/Private/GPUProfiler.cpp
@@ -122,10 +122,10 @@ static void DumpStatsEventNode(FGPUProfilerEventNode* Node, float RootResult, in
 		const float UnaccountedTime = FMath::Max(Node->TimingResult - TotalChildTime, 0.0f);
 		const float UnaccountedPercent = UnaccountedTime * 100.0f / (RootResult * 1000.0f);
 
-		// Add an 'Unaccounted' node if necessary to show time spent in the current node that is not in any of its children
+		// Add an 'Other Children' node if necessary to show time spent in the current node that is not in any of its children
 		if (bMatchesFilter && Node->Children.Num() > 0 && TotalChildDraws > 0 && (UnaccountedPercent > 2.0f || UnaccountedTime > .2f))
 		{
-			UE_LOG(LogRHI, Warning, TEXT("%s%4.1f%%%5.2fms Unaccounted"), 
+			UE_LOG(LogRHI, Warning, TEXT("%s%4.1f%%%5.2fms Other Children"), 
 				*FString(TEXT("")).LeftPad((EffectiveDepth + 1) * 3), 
 				UnaccountedPercent,
 				UnaccountedTime);
diff --git a/Engine/Source/Runtime/Renderer/Private/CompositionLighting/PostProcessDeferredDecals.cpp b/Engine/Source/Runtime/Renderer/Private/CompositionLighting/PostProcessDeferredDecals.cpp
--- a/Engine/Source/Runtime/Renderer/Private/CompositionLighting/PostProcessDeferredDecals.cpp
+++ b/Engine/Source/Runtime/Renderer/Private/CompositionLighting/PostProcessDeferredDecals.cpp
@@ -634,8 +634,6 @@ void SetShader(const FRenderingCompositePassContext& Context, const FTransientDe
 
 bool RenderPreStencil(FRenderingCompositePassContext& Context, const FMaterialShaderMap* MaterialShaderMap, const FMatrix& ComponentToWorldMatrix, const FMatrix& FrustumComponentToClip)
 {
-	SCOPED_DRAW_EVENT(Context.RHICmdList, RenderPreStencil);
-
 	const FSceneView& View = Context.View;
 
 	float Distance = (View.ViewMatrices.ViewOrigin - ComponentToWorldMatrix.GetOrigin()).Size();
diff --git a/Engine/Source/Runtime/Renderer/Private/DistanceFieldSurfaceCacheLighting.cpp b/Engine/Source/Runtime/Renderer/Private/DistanceFieldSurfaceCacheLighting.cpp
--- a/Engine/Source/Runtime/Renderer/Private/DistanceFieldSurfaceCacheLighting.cpp
+++ b/Engine/Source/Runtime/Renderer/Private/DistanceFieldSurfaceCacheLighting.cpp
@@ -3695,7 +3695,7 @@ bool FDeferredShadingSceneRenderer::RenderDistanceFieldAOSurfaceCache(
 		&& View.IsPerspectiveProjection())
 	{
 		QUICK_SCOPE_CYCLE_COUNTER(STAT_RenderDistanceFieldAOSurfaceCache);
-		SCOPED_DRAW_EVENT(RHICmdList, DistanceFieldAO);
+		SCOPED_DRAW_EVENT(RHICmdList, DistanceFieldLighting);
 
 		if (GDistanceFieldVolumeTextureAtlas.VolumeTextureRHI && Scene->DistanceFieldSceneData.NumObjectsInBuffer)
 		{
diff --git a/Engine/Source/Runtime/Renderer/Private/LightShaftRendering.cpp b/Engine/Source/Runtime/Renderer/Private/LightShaftRendering.cpp
--- a/Engine/Source/Runtime/Renderer/Private/LightShaftRendering.cpp
+++ b/Engine/Source/Runtime/Renderer/Private/LightShaftRendering.cpp
@@ -451,8 +451,6 @@ void AllocateOrReuseLightShaftRenderTarget(FRHICommandListImmediate& RHICmdList,
 template<bool bDownsampleOcclusion>
 void DownsamplePass(FRHICommandListImmediate& RHICmdList, const FViewInfo& View, const FLightSceneInfo* LightSceneInfo, TRefCountPtr<IPooledRenderTarget>& LightShaftsSource, TRefCountPtr<IPooledRenderTarget>& LightShaftsDest)
 {
-	SCOPED_DRAW_EVENT(RHICmdList, Downsample);
-
 	const FIntPoint BufferSize = GSceneRenderTargets.GetBufferSizeXY();
 	const uint32 DownsampleFactor	= GetLightShaftDownsampleFactor();
 	const FIntPoint FilterBufferSize = GSceneRenderTargets.GetBufferSizeXY() / DownsampleFactor;
@@ -624,7 +622,6 @@ void ApplyRadialBlurPasses(
 		BlurLightShaftsPixelShader->SetParameters(RHICmdList, LightSceneInfo, View, PassIndex, EffectiveSource);
 
 		{
-			SCOPED_DRAW_EVENT(RHICmdList, RadialBlur);
 			// Apply a radial blur to the bloom and occlusion mask
 			DrawRectangle( 
 				RHICmdList,
@@ -668,7 +665,6 @@ void FinishOcclusionTerm(FRHICommandList& RHICmdList, const FViewInfo& View, con
 	MaskOcclusionTermPixelShader->SetParameters(RHICmdList, LightSceneInfo, View, LightShaftsSource);
 
 	{
-		SCOPED_DRAW_EVENT(RHICmdList, FinishOcclusion);
 		// Apply a radial blur to the bloom and occlusion mask
 		DrawRectangle( 
 			RHICmdList,
@@ -836,8 +832,6 @@ IMPLEMENT_SHADER_TYPE(,FApplyLightShaftsPixelShader,TEXT("LightShaftShader"),TEX
 
 void ApplyLightShaftBloom(FRHICommandListImmediate& RHICmdList, const FViewInfo& View, const FLightSceneInfo* const LightSceneInfo, TRefCountPtr<IPooledRenderTarget>& LightShaftsSource)
 {
-	SCOPED_DRAW_EVENT(RHICmdList, Apply);
-
 	GSceneRenderTargets.BeginRenderingSceneColor(RHICmdList, ESimpleRenderTargetMode::EUninitializedColorExistingReadOnlyDepth);
 
 	RHICmdList.SetViewport(View.ViewRect.Min.X, View.ViewRect.Min.Y, 0.0f, View.ViewRect.Max.X, View.ViewRect.Max.Y, 1.0f);
