Integrate - Heightfield GI is now computed before Heightfield AO, avoids incorrect self-shadowing
* Added a bias along the normal for heightfield shadowing used by GI, avoids aliasy grass shadows
* Also implemented a 3x3 filter for heightfield shadowmap shadows to reduce instability from camera movement and vertex animation

diff --git a/Engine/Shaders/DistanceFieldAOShared.usf b/Engine/Shaders/DistanceFieldAOShared.usf
--- a/Engine/Shaders/DistanceFieldAOShared.usf
+++ b/Engine/Shaders/DistanceFieldAOShared.usf
@@ -146,3 +146,28 @@ bool SphereIntersectConeWithDepthRanges(float4 SphereCenterAndRadius, float3 Con
 	return false;
 }
 
+Buffer<float> RecordConeVisibility;
+float BentNormalNormalizeFactor;
+
+float3 ComputeBentNormal(float3 RecordWorldNormal, uint RelativeRecordIndex)
+{
+	float3 TangentX;
+	float3 TangentY;
+	FindBestAxisVectors2(RecordWorldNormal, TangentX, TangentY);
+
+	float3 UnoccludedDirection = 0;
+
+	for (uint ConeIndex = 0; ConeIndex < NUM_CONE_DIRECTIONS; ConeIndex++)
+	{
+		float3 ConeDirection = AOSamples2.SampleDirections[ConeIndex].xyz;
+		float3 RotatedConeDirection = ConeDirection.x * TangentX + ConeDirection.y * TangentY + ConeDirection.z * RecordWorldNormal;
+
+		float ConeVisibility = RecordConeVisibility[RelativeRecordIndex * NUM_CONE_DIRECTIONS + ConeIndex];
+		UnoccludedDirection += ConeVisibility * RotatedConeDirection;
+	}
+
+	float InvNumSamples = 1.0f / (float)NUM_CONE_DIRECTIONS;
+	UnoccludedDirection = UnoccludedDirection * (BentNormalNormalizeFactor * InvNumSamples);
+
+	return UnoccludedDirection;
+}
diff --git a/Engine/Shaders/DistanceFieldGlobalIllumination.usf b/Engine/Shaders/DistanceFieldGlobalIllumination.usf
--- a/Engine/Shaders/DistanceFieldGlobalIllumination.usf
+++ b/Engine/Shaders/DistanceFieldGlobalIllumination.usf
@@ -682,7 +682,6 @@ float RayTraceThroughGlobalObjects(
 
 Buffer<float> RecordConeData;
 RWBuffer<float4> RWStepBentNormal;
-float BentNormalNormalizeFactor;
 
 /**  */
 [numthreads(FINAL_GATHER_THREADGROUP_SIZE, 1, 1)]
diff --git a/Engine/Shaders/DistanceFieldSurfaceCacheLightingCompute.usf b/Engine/Shaders/DistanceFieldSurfaceCacheLightingCompute.usf
--- a/Engine/Shaders/DistanceFieldSurfaceCacheLightingCompute.usf
+++ b/Engine/Shaders/DistanceFieldSurfaceCacheLightingCompute.usf
@@ -681,7 +681,6 @@ void PopulateCacheCS(
 	}
 }
 
-float BentNormalNormalizeFactor;
 float TanConeHalfAngle;
 float RecordRadiusScale;
 
@@ -938,7 +937,6 @@ void ConeTraceOcclusionCS(
 }
 
 RWBuffer<float4> RWIrradianceCacheBentNormal;
-Buffer<float> RecordConeVisibility;
 
 /**  */
 [numthreads(FINAL_GATHER_THREADGROUP_SIZE, 1, 1)]
@@ -955,24 +953,7 @@ void CombineConesCS(
 	if (RecordIndex < NumRecords)
 	{
 		float3 RecordWorldNormal = IrradianceCacheNormal[RecordIndex].xyz;
-
-		float3 TangentX;
-		float3 TangentY;
-		FindBestAxisVectors2(RecordWorldNormal, TangentX, TangentY);
-
-		float3 UnoccludedDirection = 0;
-
-		for (uint ConeIndex = 0; ConeIndex < NUM_CONE_DIRECTIONS; ConeIndex++)
-		{
-			float3 ConeDirection = AOSamples2.SampleDirections[ConeIndex].xyz;
-			float3 RotatedConeDirection = ConeDirection.x * TangentX + ConeDirection.y * TangentY + ConeDirection.z * RecordWorldNormal;
-
-			float ConeVisibility = RecordConeVisibility[DispatchThreadId.x * NUM_CONE_DIRECTIONS + ConeIndex];
-			UnoccludedDirection += ConeVisibility * RotatedConeDirection;
-		}
-
-		float InvNumSamples = 1.0f / (float)NUM_CONE_DIRECTIONS;
-		UnoccludedDirection = UnoccludedDirection * (BentNormalNormalizeFactor * InvNumSamples);
+		float3 UnoccludedDirection = ComputeBentNormal(RecordWorldNormal, DispatchThreadId.x);
 
 		RWIrradianceCacheBentNormal[RecordIndex] = float4(UnoccludedDirection, 0);
 	}
diff --git a/Engine/Shaders/HeightfieldLighting.usf b/Engine/Shaders/HeightfieldLighting.usf
--- a/Engine/Shaders/HeightfieldLighting.usf
+++ b/Engine/Shaders/HeightfieldLighting.usf
@@ -268,24 +268,76 @@ void ShadowHeightfieldsPS(
 		clip(-1);
 	}
 
+	float4x4 WorldToLocal = GetWorldToLocal(Input.HeightfieldIndex);
+	float WorldToLocalScale = length(WorldToLocal[0].xyz);
+	float HeightfieldTexelSize = 1.0f / WorldToLocalScale;
+
+	float3 WorldHeightfieldNormal = GetGlobalHeightfieldNormal(Input.GlobalUV);
+	// Bias away from the surface, mostly to avoid grass shadowing
+	WorldHeightfieldShadingPosition += .3f * HeightfieldTexelSize * WorldHeightfieldNormal;
+
+#define FILTER_SHADOWMAP_LOOKUP 1
+#if FILTER_SHADOWMAP_LOOKUP
+
+	float3 TangentX;
+	float3 TangentY;
+	FindBestAxisVectors2(WorldHeightfieldNormal, TangentX, TangentY);
+
+	float HalfNumSamples = 1;
+	float Visibility = 0;
+	float TotalWeight = 0;
+
+	for (float YOffset = -HalfNumSamples; YOffset <= HalfNumSamples; YOffset++)
+	{
+		for (float XOffset = -HalfNumSamples; XOffset <= HalfNumSamples; XOffset++)
+		{
+			float2 SampleOffset = float2(XOffset, YOffset) * (.5f * HeightfieldTexelSize / HalfNumSamples);
+			float3 WorldSamplePosition = WorldHeightfieldShadingPosition + TangentX * SampleOffset.x + TangentY * SampleOffset.y;
+
+			float4 HomogeneousShadowPosition = mul(float4(WorldSamplePosition, 1), WorldToShadow);
+			float2 ShadowUVs = HomogeneousShadowPosition.xy / HomogeneousShadowPosition.w;
+
+			if (all(ShadowUVs >= ShadowmapMinMax.xy && ShadowUVs <= ShadowmapMinMax.zw))
+			{
+				// Sample the shadowmap depth and determine if this texel is shadowed
+				float ShadowDepth = Texture2DSampleLevel(ShadowDepthTexture, ShadowDepthTextureSampler, ShadowUVs, 0).x;
+				Visibility += HomogeneousShadowPosition.z < ShadowDepth + ShadowDepthBias;
+				TotalWeight++;
+			}
+		}
+	}
+
+	if (TotalWeight > 0)
+	{
+		Visibility /= TotalWeight;
+	}
+	else
+	{
+		clip(-1);
+	}
+
+#else
+
 	float4 HomogeneousShadowPosition = mul(float4(WorldHeightfieldShadingPosition, 1), WorldToShadow);
 	float2 ShadowUVs = HomogeneousShadowPosition.xy / HomogeneousShadowPosition.w;
 
-	float Shadow = 1;
+	float Visibility = 1;
 
 	// Treat as unshadowed if the texel is outside of the shadow map
 	if (all(ShadowUVs >= ShadowmapMinMax.xy && ShadowUVs <= ShadowmapMinMax.zw))
 	{
 		// Sample the shadowmap depth and determine if this texel is shadowed
 		float ShadowDepth = Texture2DSampleLevel(ShadowDepthTexture, ShadowDepthTextureSampler, ShadowUVs, 0).x;
-		Shadow = HomogeneousShadowPosition.z < ShadowDepth + ShadowDepthBias;
+		Visibility = HomogeneousShadowPosition.z < ShadowDepth + ShadowDepthBias;
 	}
 	else
 	{
 		clip(-1);
 	}
 
-	OutColor = float4(Shadow, 0, 0, 0);
+#endif
+
+	OutColor = float4(Visibility, 0, 0, 0);
 }
 
 float3 LightDirection;
@@ -301,6 +353,14 @@ void RayTracedShadowHeightfieldsPS(
 
 	if (ViewDistance > CascadeDepthMinMax.x && ViewDistance < CascadeDepthMinMax.y)
 	{
+		float4x4 WorldToLocal = GetWorldToLocal(Input.HeightfieldIndex);
+		float WorldToLocalScale = length(WorldToLocal[0].xyz);
+		float HeightfieldTexelSize = 1.0f / WorldToLocalScale;
+
+		float3 WorldHeightfieldNormal = GetGlobalHeightfieldNormal(Input.GlobalUV);
+		// Bias away from the surface, mostly to avoid grass shadowing
+		OpaqueWorldPosition += .3f * HeightfieldTexelSize * WorldHeightfieldNormal;
+
 		uint NumIntersectingObjects = GetCulledNumObjects();
 		uint CulledDataStart = 0;
 
@@ -459,8 +519,6 @@ float3 SampleHeightfieldLighting(float2 GlobalUV)
 	return Texture2DSampleLevel(HeightfieldLighting, HeightfieldLightingSampler, GlobalUV, 0).xyz;
 }
 
-uint NumHemisphereSamples;
-
 bool RayTraceHeightfield(float3 LocalRayStart, float3 LocalRayEnd, float4 GlobalUVScaleBias, out float IntersectionTime, out float2 IntersectionUV)
 {
 	IntersectionUV = 0;
@@ -535,6 +593,7 @@ float3 ComputeVirtualPointLighting(float3 WorldPosition, float3 WorldNormal, flo
 
 				float CosTheta = DirectionDot / Distance;
 				float Shadow = GetVPLOcclusion(BentNormal, VPLDirection, .5f, .5f);
+
 				float VPLCosineLobe = saturate(VPLNormalDot / Distance);
 
 				Irradiance += (saturate(CosTheta) * VPLCosineLobe * DistanceAttenuation * Shadow) * VPLFlux;
@@ -617,13 +676,7 @@ void CalculateHeightfieldIrradianceCS(
 	BRANCH
 	if (RecordIndex < NumRecords)
 	{
-		float3 WorldNormal = IrradianceCacheNormal[RecordIndex].xyz;
 		float3 WorldPosition = PositionAndRadius.xyz;
-		float3 BentNormal = IrradianceCacheBentNormal[RecordIndex].xyz;
-
-		float3 TangentX;
-		float3 TangentY;
-		FindBestAxisVectors2(WorldNormal, TangentX, TangentY);
 
 		//@todo - constant time grid lookup
 		LOOP
@@ -638,6 +691,10 @@ void CalculateHeightfieldIrradianceCS(
 			if (all(HeightfieldUV > MinMaxHeightfieldUV.xy) 
 				&& all(HeightfieldUV < MinMaxHeightfieldUV.zw))
 			{
+				float3 WorldNormal = IrradianceCacheNormal[RecordIndex].xyz;
+				// Note: can't get this from IrradianceCacheBentNormal as that hasn't been computed yet to avoid heightfield self-occlusion
+				float3 BentNormal = ComputeBentNormal(WorldNormal, RelativeRecordIndex);
+
 				float4 GlobalUVScaleBias = GetHeightfieldGlobalScaleBias(HeightfieldIndex);
 				float2 GlobalUV = LocalPosition.xy * GlobalUVScaleBias.xy + GlobalUVScaleBias.zw;
 				float4x4 LocalToWorld = GetLocalToWorld(HeightfieldIndex);
diff --git a/Engine/Source/Runtime/Renderer/Private/DistanceFieldSurfaceCacheLighting.cpp b/Engine/Source/Runtime/Renderer/Private/DistanceFieldSurfaceCacheLighting.cpp
--- a/Engine/Source/Runtime/Renderer/Private/DistanceFieldSurfaceCacheLighting.cpp
+++ b/Engine/Source/Runtime/Renderer/Private/DistanceFieldSurfaceCacheLighting.cpp
@@ -3936,6 +3936,20 @@ bool FDeferredShadingSceneRenderer::RenderDistanceFieldAOSurfaceCache(
 					}
 				}
 
+				if (bUseDistanceFieldGI)
+				{
+					extern void ComputeIrradianceForSamples(
+						int32 DepthLevel,
+						FRHICommandListImmediate& RHICmdList,
+						const FViewInfo& View,
+						const FScene* Scene,
+						const FDistanceFieldAOParameters& Parameters,
+						FTemporaryIrradianceCacheResources* TemporaryIrradianceCacheResources);
+
+					ComputeIrradianceForSamples(DepthLevel, RHICmdList, View, Scene, Parameters, &GTemporaryIrradianceCacheResources);
+				}
+
+				// Compute heightfield occlusion after heightfield GI, otherwise it self-shadows incorrectly
 				View.HeightfieldLightingViewInfo.ComputeOcclusionForSamples(View, RHICmdList, GTemporaryIrradianceCacheResources, DepthLevel, Parameters);
 
 				{	
@@ -3948,6 +3962,7 @@ bool FDeferredShadingSceneRenderer::RenderDistanceFieldAOSurfaceCache(
 					ComputeShader->UnsetParameters(RHICmdList);
 				}
 
+				// Compute and store the final bent normal now that all occlusion sources have been computed (distance fields, heightfields)
 				{
 					TShaderMapRef<FCombineConesCS> ComputeShader(View.ShaderMap);
 
@@ -3957,19 +3972,6 @@ bool FDeferredShadingSceneRenderer::RenderDistanceFieldAOSurfaceCache(
 
 					ComputeShader->UnsetParameters(RHICmdList);
 				}
-
-				if (bUseDistanceFieldGI)
-				{
-					extern void ComputeIrradianceForSamples(
-						int32 DepthLevel,
-						FRHICommandListImmediate& RHICmdList,
-						const FViewInfo& View,
-						const FScene* Scene,
-						const FDistanceFieldAOParameters& Parameters,
-						FTemporaryIrradianceCacheResources* TemporaryIrradianceCacheResources);
-
-					ComputeIrradianceForSamples(DepthLevel, RHICmdList, View, Scene, Parameters, &GTemporaryIrradianceCacheResources);
-				}
 			}
 
 			TRefCountPtr<IPooledRenderTarget> BentNormalAccumulation;
diff --git a/Engine/Source/Runtime/Renderer/Private/HeightfieldLighting.cpp b/Engine/Source/Runtime/Renderer/Private/HeightfieldLighting.cpp
--- a/Engine/Source/Runtime/Renderer/Private/HeightfieldLighting.cpp
+++ b/Engine/Source/Runtime/Renderer/Private/HeightfieldLighting.cpp
@@ -34,14 +34,6 @@ FAutoConsoleVariableRef CVarHeightfieldGlobalIllumination(
 	ECVF_Cheat | ECVF_RenderThreadSafe
 	);
 
-int32 GHeightfieldHemisphereSamples = 32;
-FAutoConsoleVariableRef CVarHeightfieldHemisphereSamples(
-	TEXT("r.HeightfieldHemisphereSamples"),
-	GHeightfieldHemisphereSamples,
-	TEXT(""),
-	ECVF_Cheat | ECVF_RenderThreadSafe
-	);
-
 float GHeightfieldInnerBounceDistance = 3000;
 FAutoConsoleVariableRef CVarHeightfieldInnerBounceDistancer(
 	TEXT("r.HeightfieldInnerBounceDistance"),
@@ -1418,35 +1410,6 @@ void FHeightfieldLightingViewInfo::ComputeOcclusionForSamples(
 	}
 }
 
-BEGIN_UNIFORM_BUFFER_STRUCT(FHemisphereSamples,)
-	DECLARE_UNIFORM_BUFFER_STRUCT_MEMBER_ARRAY(FVector4,SampleDirections,[1024])
-END_UNIFORM_BUFFER_STRUCT(FHemisphereSamples)
-
-IMPLEMENT_UNIFORM_BUFFER_STRUCT(FHemisphereSamples,TEXT("HemisphereSamples"));
-
-/** Generates unit length, stratified and uniformly distributed direction samples in a hemisphere. */
-void GenerateStratifiedUniformHemisphereSamples(int32 NumThetaSteps, int32 NumPhiSteps, FRandomStream& RandomStream, TArray<FVector4>& Samples)
-{
-	Samples.Empty(NumThetaSteps * NumPhiSteps);
-	for (int32 ThetaIndex = 0; ThetaIndex < NumThetaSteps; ThetaIndex++)
-	{
-		for (int32 PhiIndex = 0; PhiIndex < NumPhiSteps; PhiIndex++)
-		{
-			const float U1 = RandomStream.GetFraction();
-			const float U2 = RandomStream.GetFraction();
-
-			const float Fraction1 = (ThetaIndex + U1) / (float)NumThetaSteps;
-			const float Fraction2 = (PhiIndex + U2) / (float)NumPhiSteps;
-
-			const float R = FMath::Sqrt(1.0f - Fraction1 * Fraction1);
-
-			const float Phi = 2.0f * (float)PI * Fraction2;
-			// Convert to Cartesian
-			Samples.Add(FVector4(FMath::Cos(Phi) * R, FMath::Sin(Phi) * R, Fraction1));
-		}
-	}
-}
-
 class FCalculateHeightfieldIrradianceCS : public FGlobalShader
 {
 	DECLARE_SHADER_TYPE(FCalculateHeightfieldIrradianceCS,Global)
@@ -1475,7 +1438,6 @@ public:
 		HeightfieldIrradiance.Bind(Initializer.ParameterMap, TEXT("HeightfieldIrradiance"));
 		IrradianceCacheNormal.Bind(Initializer.ParameterMap, TEXT("IrradianceCacheNormal"));
 		IrradianceCachePositionRadius.Bind(Initializer.ParameterMap, TEXT("IrradianceCachePositionRadius"));
-		IrradianceCacheBentNormal.Bind(Initializer.ParameterMap, TEXT("IrradianceCacheBentNormal"));
 		ScatterDrawParameters.Bind(Initializer.ParameterMap, TEXT("ScatterDrawParameters"));
 		SavedStartIndex.Bind(Initializer.ParameterMap, TEXT("SavedStartIndex"));
 		BentNormalNormalizeFactor.Bind(Initializer.ParameterMap, TEXT("BentNormalNormalizeFactor"));
@@ -1485,9 +1447,9 @@ public:
 		TanConeHalfAngle.Bind(Initializer.ParameterMap, TEXT("TanConeHalfAngle"));
 		HeightfieldLighting.Bind(Initializer.ParameterMap, TEXT("HeightfieldLighting"));
 		HeightfieldLightingSampler.Bind(Initializer.ParameterMap, TEXT("HeightfieldLightingSampler"));
-		NumHemisphereSamples.Bind(Initializer.ParameterMap, TEXT("NumHemisphereSamples"));
 		InnerLightTransferDistance.Bind(Initializer.ParameterMap, TEXT("InnerLightTransferDistance"));
 		OuterLightTransferDistanceScale.Bind(Initializer.ParameterMap, TEXT("OuterLightTransferDistanceScale"));
+		RecordConeVisibility.Bind(Initializer.ParameterMap, TEXT("RecordConeVisibility"));
 	}
 
 	FCalculateHeightfieldIrradianceCS()
@@ -1515,9 +1477,9 @@ public:
 
 		SetSRVParameter(RHICmdList, ShaderRHI, IrradianceCacheNormal, SurfaceCacheResources.Level[DepthLevel]->Normal.SRV);
 		SetSRVParameter(RHICmdList, ShaderRHI, IrradianceCachePositionRadius, SurfaceCacheResources.Level[DepthLevel]->PositionAndRadius.SRV);
-		SetSRVParameter(RHICmdList, ShaderRHI, IrradianceCacheBentNormal, SurfaceCacheResources.Level[DepthLevel]->BentNormal.SRV);
 		SetSRVParameter(RHICmdList, ShaderRHI, ScatterDrawParameters, SurfaceCacheResources.Level[DepthLevel]->ScatterDrawParameters.SRV);
 		SetSRVParameter(RHICmdList, ShaderRHI, SavedStartIndex, SurfaceCacheResources.Level[DepthLevel]->SavedStartIndex.SRV);
+		SetSRVParameter(RHICmdList, ShaderRHI, RecordConeVisibility, TemporaryIrradianceCacheResources.ConeVisibility.SRV);
 
 		HeightfieldIrradiance.SetBuffer(RHICmdList, ShaderRHI, TemporaryIrradianceCacheResources.HeightfieldIrradiance);
 
@@ -1545,28 +1507,6 @@ public:
 			SetShaderValue(RHICmdList, ShaderRHI, BentNormalNormalizeFactor, BentNormalNormalizeFactorValue);
 		}
 
-		{
-			//@todo - cache this buffer
-			FHemisphereSamples HemisphereSamples;
-
-			const float NumThetaStepsFloat = FMath::Sqrt(GHeightfieldHemisphereSamples / (float)PI);
-			const int32 NumThetaSteps = FMath::TruncToInt(NumThetaStepsFloat);
-			const int32 NumPhiSteps = FMath::TruncToInt(NumThetaStepsFloat * (float)PI);
-
-			TArray<FVector4> SampleDirections;
-			FRandomStream RandomStream(12345);
-			GenerateStratifiedUniformHemisphereSamples(NumThetaSteps, NumPhiSteps, RandomStream, SampleDirections);
-
-			for (int32 SampleIndex = 0; SampleIndex < NumThetaSteps * NumPhiSteps; SampleIndex++)
-			{
-				HemisphereSamples.SampleDirections[SampleIndex] = FVector4(SampleDirections[SampleIndex]);
-			}
-
-			SetUniformBufferParameterImmediate(RHICmdList, ShaderRHI, GetUniformBufferParameter<FHemisphereSamples>(), HemisphereSamples);
-
-			SetShaderValue(RHICmdList, ShaderRHI, NumHemisphereSamples, NumThetaSteps * NumPhiSteps);
-		}
-
 		extern float GAOConeHalfAngle;
 		SetShaderValue(RHICmdList, ShaderRHI, TanConeHalfAngle, FMath::Tan(GAOConeHalfAngle));
 		extern float GAORecordRadiusScale;
@@ -1593,7 +1533,6 @@ public:
 		Ar << HeightfieldIrradiance;
 		Ar << IrradianceCacheNormal;
 		Ar << IrradianceCachePositionRadius;
-		Ar << IrradianceCacheBentNormal;
 		Ar << ScatterDrawParameters;
 		Ar << SavedStartIndex;
 		Ar << BentNormalNormalizeFactor;
@@ -1603,9 +1542,9 @@ public:
 		Ar << TanConeHalfAngle;
 		Ar << HeightfieldLighting;
 		Ar << HeightfieldLightingSampler;
-		Ar << NumHemisphereSamples;
 		Ar << InnerLightTransferDistance;
 		Ar << OuterLightTransferDistanceScale;
+		Ar << RecordConeVisibility;
 		return bShaderHasOutdatedParameters;
 	}
 
@@ -1617,7 +1556,6 @@ private:
 	FRWShaderParameter HeightfieldIrradiance;
 	FShaderResourceParameter IrradianceCacheNormal;
 	FShaderResourceParameter IrradianceCachePositionRadius;
-	FShaderResourceParameter IrradianceCacheBentNormal;
 	FShaderResourceParameter ScatterDrawParameters;
 	FShaderResourceParameter SavedStartIndex;
 	FShaderParameter BentNormalNormalizeFactor;
@@ -1627,9 +1565,9 @@ private:
 	FShaderParameter TanConeHalfAngle;
 	FShaderResourceParameter HeightfieldLighting;
 	FShaderResourceParameter HeightfieldLightingSampler;
-	FShaderParameter NumHemisphereSamples;
 	FShaderParameter InnerLightTransferDistance;
 	FShaderParameter OuterLightTransferDistanceScale;
+	FShaderResourceParameter RecordConeVisibility;
 };
 
 IMPLEMENT_SHADER_TYPE(,FCalculateHeightfieldIrradianceCS,TEXT("HeightfieldLighting"),TEXT("CalculateHeightfieldIrradianceCS"),SF_Compute);
