Merge performance regression fix for IsShadowOccluded()

diff --git a/Engine/Source/Runtime/Renderer/Private/SceneOcclusion.cpp b/Engine/Source/Runtime/Renderer/Private/SceneOcclusion.cpp
--- a/Engine/Source/Runtime/Renderer/Private/SceneOcclusion.cpp
+++ b/Engine/Source/Runtime/Renderer/Private/SceneOcclusion.cpp
@@ -235,7 +235,7 @@ bool FSceneViewState::IsShadowOccluded(FRHICommandListImmediate& RHICmdList, FPr
 #if BUFFERED_OCCLUSION_QUERIES
 	// Get the oldest occlusion query	
 	const uint32 QueryIndex = FOcclusionQueryHelpers::GetQueryLookupIndex(PendingPrevFrameNumber, NumBufferedFrames);
-	FSceneViewState::ShadowKeyOcclusionQueryMap ShadowOcclusionQueryMap = ShadowOcclusionQueryMaps[QueryIndex];	
+	const FSceneViewState::ShadowKeyOcclusionQueryMap& ShadowOcclusionQueryMap = ShadowOcclusionQueryMaps[QueryIndex];	
 #endif
 
 
