New tile based motion blur with smooth gradients. WIP

diff --git a/Engine/Shaders/PostProcessMotionBlur.usf b/Engine/Shaders/PostProcessMotionBlur.usf
--- a/Engine/Shaders/PostProcessMotionBlur.usf
+++ b/Engine/Shaders/PostProcessMotionBlur.usf
@@ -18,7 +18,7 @@
 // but adapted for half resolution and not using randomization/noise
 
 // from the paper: We use SOFT Z EXTENT = 1mm to 10cm for our results
-#define SOFT_Z_EXTENT 1
+#define SOFT_Z_EXTENT 0.1
 
 // This,
 //  - Gets the MainPS pass to low cost under no motion.
@@ -622,4 +622,134 @@ void MainRecombinePS(float4 UVAndScreenPos : TEXCOORD0, out float4 OutColor : SV
 
 	OutColor.rgb = FullResSceneColor * (1 - Mask) + MotionBlurOutput.rgb;
 	OutColor.a = 0;
+}
+
+
+
+
+
+float2 DepthCmp( float CenterDepth, float SampleDepth, float DepthScale )
+{
+	return saturate( 0.5 + float2( DepthScale, -DepthScale ) * (SampleDepth - CenterDepth) );
+}
+
+float2 SpreadCmp( float OffsetLength, float2 SpreadLength, float PixelToSampleScale )
+{
+	return saturate( PixelToSampleScale * SpreadLength - max( OffsetLength - 1, 0 ) );
+}
+
+float SampleWeight( float CenterDepth, float SampleDepth, float OffsetLength, float CenterSpreadLength, float SampleSpreadLength, float PixelToSampleScale, float DepthScale )
+{
+	float2 DepthWeights = DepthCmp( CenterDepth, SampleDepth, DepthScale );
+	float2 SpreadWeights = SpreadCmp( OffsetLength, float2( CenterSpreadLength, SampleSpreadLength ), PixelToSampleScale );
+	return dot( DepthWeights, SpreadWeights );
+}
+
+// [0, 1]
+float RandFast( uint2 PixelPos, float Magic = 3571.0 )
+{
+	float2 Random = ( 1.0 / 4320.0 ) * PixelPos + float2( 0.25, 0.0 );
+	Random = frac( dot( Random * Random, Magic ) );
+	Random = frac( dot( Random * Random, Magic ) );
+	return Random.x;
+}
+
+
+void MainNewPS( float4 UVAndScreenPos : TEXCOORD0, out float4 OutColor : SV_Target0 )
+{
+	// Screen Quad UV 0..1
+	float2 UV = UVAndScreenPos.xy;
+	// screen position in [-1, 1] screen space
+	float2 ScreenSpacePos = UVAndScreenPos.zw;
+
+	OutColor = 0;
+
+	float VelocityScale = MotionBlurParameters.x;
+
+	const uint StepCount = GetStepCountFromQuality() / 2;
+
+#if 1
+	uint2 PixelPos = uint2(UVAndScreenPos.zw * ScreenPosToPixel.xy + ScreenPosToPixel.zw + 0.5f);
+	float2 PosMod = float2( PixelPos & 1 );
+	float Dither = ( PosMod.x * 2 - 1 ) * ( PosMod.y * 2 - 1 );
+	float Random = RandFast( PixelPos );
+	float Random2 = RandFast( PixelPos, 5521 );
+#endif
+	
+	float2 TileJitter = ( float2( Random, Random2 ) - 0.5 ) * 0.5;
+	
+	float2 MaxVelocity = PostprocessInput3.SampleLevel( PostprocessInput3Sampler, UV + TileJitter * PostprocessInput3Size.zw, 0 ).xy;
+	//float2 MaxVelocity = PostprocessInput3[ PixelPos / 16 ].xy;
+	float2 MaxVelocityScreen = MaxVelocity.xy * VelocityScale;
+	float2 MaxVelocityPixels = MaxVelocityScreen * ScreenPosToPixel.xy;
+
+	float3 CenterColor = PostprocessInput0.SampleLevel( PostprocessInput0Sampler, UV, 0 ).rgb;
+	float  CenterDepth = PostprocessInput1.SampleLevel( PostprocessInput1Sampler, UV, 0 ).r;
+
+	CenterDepth = ConvertFromDeviceZ( CenterDepth );
+	
+	BRANCH
+	if( length( MaxVelocityPixels ) <= 0.5 )
+	{
+		OutColor.rgb = CenterColor;
+		return;
+	}
+	
+	// in pixels
+	float2 CenterVelocity = PostprocessInput2.SampleLevel( PostprocessInput2Sampler, UV, 0 ).xy;
+	CenterVelocity = DecodeVelocityFromTexture( CenterVelocity ) * VelocityScale * ScreenPosToPixel.xy;
+	float CenterVelocityLength = length( CenterVelocity );
+
+	float4 ColorAccum = 0;
+
+	// converts pixel length to sample steps
+	float PixelToSampleScale = ( StepCount - 0.5 ) / length( MaxVelocityPixels );
+	//float PixelToSampleScale = ( StepCount - 0.5 + (Random - 0.5) ) / length( MaxVelocityPixels );
+
+	// WHY?
+	PixelToSampleScale *= 2;
+
+	UNROLL
+	for( uint i = 0; i < StepCount; i++ )
+	{
+		float2 SampleUV[2];
+		float3 SampleColor[2];
+		float  SampleDepth[2];
+		float  SampleVelocityLength[2];
+		float  Weight[2];
+
+		float OffsetLength = (float)i + 0.5 + (Random - 0.5);
+		float OffsetFraction = OffsetLength / ( StepCount - 0.5 );
+
+		// can be optimized
+		SampleUV[0] = UV + OffsetFraction * MaxVelocityPixels * PostprocessInput0Size.zw;
+		SampleUV[1] = UV - OffsetFraction * MaxVelocityPixels * PostprocessInput0Size.zw;
+
+		UNROLL
+		for( uint j = 0; j < 2; j++ )
+		{
+			SampleColor[j] = PostprocessInput0.SampleLevel( PostprocessInput0Sampler, SampleUV[j], 0 ).rgb;
+			SampleDepth[j] = PostprocessInput1.SampleLevel( PostprocessInput1Sampler, SampleUV[j], 0 ).r;
+
+			SampleDepth[j] = ConvertFromDeviceZ( SampleDepth[j] );
+
+			// in pixels
+			float2 SampleVelocity = PostprocessInput2.SampleLevel( PostprocessInput2Sampler, SampleUV[j], 0 ).xy;
+			SampleVelocity = DecodeVelocityFromTexture( SampleVelocity ) * VelocityScale * ScreenPosToPixel.xy;
+			SampleVelocityLength[j] = length( SampleVelocity );
+
+			Weight[j] = SampleWeight( CenterDepth, SampleDepth[j], OffsetLength, CenterVelocityLength, SampleVelocityLength[j], PixelToSampleScale, SOFT_Z_EXTENT );
+		}
+		
+		bool2 Mirror = bool2( SampleDepth[0] > SampleDepth[1], SampleVelocityLength[1] > SampleVelocityLength[0] );
+		Weight[0] = all( Mirror ) ? Weight[1] : Weight[0];
+		Weight[1] = any( Mirror ) ? Weight[1] : Weight[0];
+
+		ColorAccum += Weight[0] * float4( SampleColor[0].rgb, 1 );
+		ColorAccum += Weight[1] * float4( SampleColor[1].rgb, 1 );
+	}
+
+	ColorAccum *= 0.5 / StepCount;
+	OutColor.rgb = ColorAccum.rgb + ( 1 - ColorAccum.a ) * CenterColor;
+	OutColor.a = 1;
 }
\ No newline at end of file
diff --git a/Engine/Shaders/PostProcessVelocityFlatten.usf b/Engine/Shaders/PostProcessVelocityFlatten.usf
new file mode 100644
--- /dev/null
+++ b/Engine/Shaders/PostProcessVelocityFlatten.usf
@@ -0,0 +1,128 @@
+// Copyright 1998-2015 Epic Games, Inc. All Rights Reserved.
+
+#include "Common.usf"
+#include "PostProcessCommon.usf"
+
+#define THREADGROUP_TOTALSIZE (THREADGROUP_SIZEX * THREADGROUP_SIZEY)
+
+groupshared float2 Shared[ THREADGROUP_TOTALSIZE ];
+
+float4 CameraMotion[5];
+
+// View rect min in xy, max in zw.
+uint4 ViewDimensions;
+
+RWTexture2D<float4>	OutVelocityFlat;
+RWTexture2D<float4>	OutMaxTileVelocity;
+
+float2 MaxLength( float2 v0, float2 v1 )
+{
+	float LengthSqr0 = dot( v0, v0 );
+	float LengthSqr1 = dot( v1, v1 );
+	return LengthSqr0 > LengthSqr1 ? v0 : v1;
+}
+
+[numthreads(THREADGROUP_SIZEX, THREADGROUP_SIZEY, 1)]
+void VelocityFlattenMain(
+	uint3 GroupId : SV_GroupID,
+	uint3 DispatchThreadId : SV_DispatchThreadID,
+	uint3 GroupThreadId : SV_GroupThreadID,
+	uint GroupIndex : SV_GroupIndex ) 
+{
+	uint ThreadIndex = GroupThreadId.y * THREADGROUP_SIZEX + GroupThreadId.x;
+	uint2 PixelPos = DispatchThreadId.xy * 2 + ViewDimensions.xy;
+	float2 ScreenUV = ( (float2)( DispatchThreadId.xy * 2 ) + 0.5 ) / ( ViewDimensions.zw - ViewDimensions.xy );
+	float2 ScreenPosition = float2(2.0f, -2.0f) * ScreenUV + float2(-1.0f, 1.0f);
+
+	// FIXME for sub viewrect
+	float2 UV = ScreenUV;
+
+	uint2 Offsets[] =
+	{
+		{ 0, 0 },
+		{ 0, 1 },
+		{ 1, 0 },
+		{ 1, 1 },
+	};
+	
+	float4 DepthSamples;
+	float2 VelocitySamples[4];
+#if 0
+	// using Gather: xyzw in counter clockwise order starting with the sample to the lower left of the queried location
+	float4 X = PostprocessInput0.GatherRed( PostprocessInput0Sampler, UV ).wzxy;
+	float4 Y = PostprocessInput0.GatherGreen( PostprocessInput0Sampler, UV ).wzxy;
+
+	VelocitySamples[0] = float2(X.x, Y.x);
+	VelocitySamples[1] = float2(X.y, Y.y);
+	VelocitySamples[2] = float2(X.z, Y.z);
+	VelocitySamples[3] = float2(X.w, Y.w);
+	
+	DepthSamples = PostprocessInput1.GatherRed( PostprocessInput1Sampler, UV, 0 ).wzxy;
+#else
+	VelocitySamples[0] = PostprocessInput0[ PixelPos + Offsets[0] ].xy;
+	VelocitySamples[1] = PostprocessInput0[ PixelPos + Offsets[1] ].xy;
+	VelocitySamples[2] = PostprocessInput0[ PixelPos + Offsets[2] ].xy;
+	VelocitySamples[3] = PostprocessInput0[ PixelPos + Offsets[3] ].xy;
+	
+	DepthSamples[0] = PostprocessInput1[ PixelPos + Offsets[0] ].x;
+	DepthSamples[1] = PostprocessInput1[ PixelPos + Offsets[1] ].x;
+	DepthSamples[2] = PostprocessInput1[ PixelPos + Offsets[2] ].x;
+	DepthSamples[3] = PostprocessInput1[ PixelPos + Offsets[3] ].x;
+#endif
+
+	float2 Velocity[4];
+	UNROLL
+	for( uint i = 0; i < 4; i++ )
+	{
+		if( VelocitySamples[i].x > 0.0 )
+		{
+			Velocity[i] = DecodeVelocityFromTexture( VelocitySamples[i] );
+		}
+		else
+		{
+			float3 PosN = float3( ScreenUV, DepthSamples[i] );
+			float ScaleM = rcp( dot( PosN, CameraMotion[0].xyz ) + CameraMotion[0].w );
+			Velocity[i].x = -2.0 * ((PosN.x * ((CameraMotion[1].x * PosN.y) + (CameraMotion[1].y * PosN.z) + CameraMotion[1].z)) + (CameraMotion[1].w * PosN.y) + (CameraMotion[2].x * PosN.x * PosN.x) + (CameraMotion[2].y * PosN.z) + CameraMotion[2].z) * ScaleM;
+			Velocity[i].y =  2.0 * ((PosN.y * ((CameraMotion[3].x * PosN.x) + (CameraMotion[3].y * PosN.z) + CameraMotion[3].z)) + (CameraMotion[3].w * PosN.x) + (CameraMotion[4].x * PosN.y * PosN.y) + (CameraMotion[4].y * PosN.z) + CameraMotion[4].z) * ScaleM;
+		}
+
+		BRANCH
+		if( all( PixelPos.xy + Offsets[i] < ViewDimensions.zw ) )
+		{
+			OutVelocityFlat[ PixelPos + Offsets[i] ] = EncodeVelocityToTexture( Velocity[i] ).xyxy;
+			//OutPackedVelocityDepth[ PixelPos + Offsets[i] ] = float2( Velocity[i], DepthSamples[i] ).xyxy;
+		}
+	}
+
+	Shared[ GroupIndex ] = MaxLength(
+		MaxLength( Velocity[0], Velocity[1] ),
+		MaxLength( Velocity[2], Velocity[3] )
+	);
+
+	GroupMemoryBarrierWithGroupSync();
+
+#if 1
+	if( GroupIndex < 512 )	Shared[ GroupIndex ] = MaxLength( Shared[ GroupIndex ], Shared[ GroupIndex + 512 ] );
+	GroupMemoryBarrierWithGroupSync();
+	if( GroupIndex < 256 )	Shared[ GroupIndex ] = MaxLength( Shared[ GroupIndex ], Shared[ GroupIndex + 256 ] );
+	GroupMemoryBarrierWithGroupSync();
+	if( GroupIndex < 128 )	Shared[ GroupIndex ] = MaxLength( Shared[ GroupIndex ], Shared[ GroupIndex + 128 ] );
+	GroupMemoryBarrierWithGroupSync();
+	if( GroupIndex <  64 )	Shared[ GroupIndex ] = MaxLength( Shared[ GroupIndex ], Shared[ GroupIndex +  64 ] );
+	GroupMemoryBarrierWithGroupSync();
+#endif
+
+	// Safe for vector sizes 32 or larger, AMD and NV
+	// TODO Intel variable size vector
+	if( GroupIndex < 32 )	Shared[ GroupIndex ] = MaxLength( Shared[ GroupIndex ], Shared[ GroupIndex + 32 ] );
+	if( GroupIndex < 16 )	Shared[ GroupIndex ] = MaxLength( Shared[ GroupIndex ], Shared[ GroupIndex + 16 ] );
+	if( GroupIndex <  8 )	Shared[ GroupIndex ] = MaxLength( Shared[ GroupIndex ], Shared[ GroupIndex +  8 ] );
+	if( GroupIndex <  4 )	Shared[ GroupIndex ] = MaxLength( Shared[ GroupIndex ], Shared[ GroupIndex +  4 ] );
+	if( GroupIndex <  2 )	Shared[ GroupIndex ] = MaxLength( Shared[ GroupIndex ], Shared[ GroupIndex +  2 ] );
+	if( GroupIndex <  1 )	Shared[ GroupIndex ] = MaxLength( Shared[ GroupIndex ], Shared[ GroupIndex +  1 ] );
+
+	if( GroupIndex == 0 )
+	{
+		OutMaxTileVelocity[ GroupId.xy ] = Shared[0].xyxy;
+	}
+}
\ No newline at end of file
diff --git a/Engine/Source/Runtime/Engine/Private/GameEngine.cpp b/Engine/Source/Runtime/Engine/Private/GameEngine.cpp
--- a/Engine/Source/Runtime/Engine/Private/GameEngine.cpp
+++ b/Engine/Source/Runtime/Engine/Private/GameEngine.cpp
@@ -795,8 +795,6 @@ void UGameEngine::Tick( float DeltaSeconds, bool bIdleMode )
 		}
 	}
 
-	bool WorldWasPaused = false;
-
 	for (int32 WorldIdx = 0; WorldIdx < WorldList.Num(); ++WorldIdx)
 	{
 		FWorldContext &Context = WorldList[WorldIdx];
@@ -805,8 +803,6 @@ void UGameEngine::Tick( float DeltaSeconds, bool bIdleMode )
 			continue;
 		}
 
-		WorldWasPaused |= Context.World()->IsPaused();
-
 		GWorld = Context.World();
 
 		// Tick all travel and Pending NetGames (Seamless, server, client)
@@ -875,6 +871,16 @@ void UGameEngine::Tick( float DeltaSeconds, bool bIdleMode )
 			}
 		}
 
+		// tell renderer about GWorld->IsPaused(), before rendering
+		{
+			ENQUEUE_UNIQUE_RENDER_COMMAND_ONEPARAMETER(
+				SetPaused,
+				bool, bGamePaused, Context.World()->IsPaused(),
+			{
+				GRenderingRealtimeClock.SetGamePaused(bGamePaused);
+			});
+		}
+
 		if (!bIdleMode && !IsRunningDedicatedServer() && !IsRunningCommandlet())
 		{
 			// Render everything.
@@ -937,17 +943,6 @@ void UGameEngine::Tick( float DeltaSeconds, bool bIdleMode )
 		GWorld = GetWorldContextFromHandleChecked(OriginalGWorldContext).World();
 	}
 
-	// tell renderer about GWorld->IsPaused(), before rendering
-	{
-		ENQUEUE_UNIQUE_RENDER_COMMAND_ONEPARAMETER(
-			SetPaused,
-			bool, bGamePaused, WorldWasPaused,
-		{
-			GRenderingRealtimeClock.SetGamePaused(bGamePaused);
-		});
-	}
-
-
 	// rendering thread commands
 	{
 		ENQUEUE_UNIQUE_RENDER_COMMAND_TWOPARAMETER(
diff --git a/Engine/Source/Runtime/Renderer/Private/CompositionLighting/PostProcessAmbientOcclusion.cpp b/Engine/Source/Runtime/Renderer/Private/CompositionLighting/PostProcessAmbientOcclusion.cpp
--- a/Engine/Source/Runtime/Renderer/Private/CompositionLighting/PostProcessAmbientOcclusion.cpp
+++ b/Engine/Source/Runtime/Renderer/Private/CompositionLighting/PostProcessAmbientOcclusion.cpp
@@ -117,71 +117,6 @@ void FCameraMotionParameters::Bind(const FShaderParameterMap& ParameterMap)
 	CameraMotion.Bind(ParameterMap, TEXT("CameraMotion"));
 }
 
-void FCameraMotionParameters::Set(FRHICommandList& RHICmdList, const FSceneView& View, const FPixelShaderRHIParamRef ShaderRHI) const
-{
-	FSceneViewState* ViewState = (FSceneViewState*)View.State;
-
-	FMatrix Proj = View.ViewMatrices.ProjMatrix;
-	FMatrix PrevProj = ViewState->PrevViewMatrices.ProjMatrix;
-
-	// Remove jitter
-	Proj.M[2][0] = 0.0f;
-	Proj.M[2][1] = 0.0f;
-	PrevProj.M[2][0] = 0.0f;
-	PrevProj.M[2][1] = 0.0f;
-
-	FMatrix ViewProj = ( View.ViewMatrices.ViewMatrix * Proj ).GetTransposed();
-	FMatrix PrevViewProj = ( ViewState->PrevViewMatrices.ViewMatrix * PrevProj ).GetTransposed();
-
-	double InvViewProj[16];
-	Inverse4x4( InvViewProj, (float*)ViewProj.M );
-
-	const float* p = (float*)PrevViewProj.M;
-
-	const double cxx = InvViewProj[ 0]; const double cxy = InvViewProj[ 1]; const double cxz = InvViewProj[ 2]; const double cxw = InvViewProj[ 3];
-	const double cyx = InvViewProj[ 4]; const double cyy = InvViewProj[ 5]; const double cyz = InvViewProj[ 6]; const double cyw = InvViewProj[ 7];
-	const double czx = InvViewProj[ 8]; const double czy = InvViewProj[ 9]; const double czz = InvViewProj[10]; const double czw = InvViewProj[11];
-	const double cwx = InvViewProj[12]; const double cwy = InvViewProj[13]; const double cwz = InvViewProj[14]; const double cww = InvViewProj[15];
-
-	const double pxx = (double)(p[ 0]); const double pxy = (double)(p[ 1]); const double pxz = (double)(p[ 2]); const double pxw = (double)(p[ 3]);
-	const double pyx = (double)(p[ 4]); const double pyy = (double)(p[ 5]); const double pyz = (double)(p[ 6]); const double pyw = (double)(p[ 7]);
-	const double pwx = (double)(p[12]); const double pwy = (double)(p[13]); const double pwz = (double)(p[14]); const double pww = (double)(p[15]);
-
-	FVector4 CameraMotionValue[5];
-
-	CameraMotionValue[0] = FVector4(
-		(float)(4.0*(cwx*pww + cxx*pwx + cyx*pwy + czx*pwz)),
-		(float)((-4.0)*(cwy*pww + cxy*pwx + cyy*pwy + czy*pwz)),
-		(float)(2.0*(cwz*pww + cxz*pwx + cyz*pwy + czz*pwz)),
-		(float)(2.0*(cww*pww - cwx*pww + cwy*pww + (cxw - cxx + cxy)*pwx + (cyw - cyx + cyy)*pwy + (czw - czx + czy)*pwz)));
-
-	CameraMotionValue[1] = FVector4(
-		(float)(( 4.0)*(cwy*pww + cxy*pwx + cyy*pwy + czy*pwz)),
-		(float)((-2.0)*(cwz*pww + cxz*pwx + cyz*pwy + czz*pwz)),
-		(float)((-2.0)*(cww*pww + cwy*pww + cxw*pwx - 2.0*cxx*pwx + cxy*pwx + cyw*pwy - 2.0*cyx*pwy + cyy*pwy + czw*pwz - 2.0*czx*pwz + czy*pwz - cwx*(2.0*pww + pxw) - cxx*pxx - cyx*pxy - czx*pxz)),
-		(float)(-2.0*(cyy*pwy + czy*pwz + cwy*(pww + pxw) + cxy*(pwx + pxx) + cyy*pxy + czy*pxz)));
-
-	CameraMotionValue[2] = FVector4(
-		(float)((-4.0)*(cwx*pww + cxx*pwx + cyx*pwy + czx*pwz)),
-		(float)(cyz*pwy + czz*pwz + cwz*(pww + pxw) + cxz*(pwx + pxx) + cyz*pxy + czz*pxz),
-		(float)(cwy*pww + cwy*pxw + cww*(pww + pxw) - cwx*(pww + pxw) + (cxw - cxx + cxy)*(pwx + pxx) + (cyw - cyx + cyy)*(pwy + pxy) + (czw - czx + czy)*(pwz + pxz)),
-		(float)(0));
-
-	CameraMotionValue[3] = FVector4(
-		(float)((-4.0)*(cwx*pww + cxx*pwx + cyx*pwy + czx*pwz)),
-		(float)((-2.0)*(cwz*pww + cxz*pwx + cyz*pwy + czz*pwz)),
-		(float)(2.0*((-cww)*pww + cwx*pww - 2.0*cwy*pww - cxw*pwx + cxx*pwx - 2.0*cxy*pwx - cyw*pwy + cyx*pwy - 2.0*cyy*pwy - czw*pwz + czx*pwz - 2.0*czy*pwz + cwy*pyw + cxy*pyx + cyy*pyy + czy*pyz)),
-		(float)(2.0*(cyx*pwy + czx*pwz + cwx*(pww - pyw) + cxx*(pwx - pyx) - cyx*pyy - czx*pyz)));
-
-	CameraMotionValue[4] = FVector4(
-		(float)(4.0*(cwy*pww + cxy*pwx + cyy*pwy + czy*pwz)),
-		(float)(cyz*pwy + czz*pwz + cwz*(pww - pyw) + cxz*(pwx - pyx) - cyz*pyy - czz*pyz),
-		(float)(cwy*pww + cww*(pww - pyw) - cwy*pyw + cwx*((-pww) + pyw) + (cxw - cxx + cxy)*(pwx - pyx) + (cyw - cyx + cyy)*(pwy - pyy) + (czw - czx + czy)*(pwz - pyz)),
-		(float)(0));
-
-	SetShaderValueArray(RHICmdList, ShaderRHI, CameraMotion, CameraMotionValue, 5);
-}
-
 FArchive& operator<<(FArchive& Ar, FCameraMotionParameters& This)
 {
 	Ar << This.CameraMotion;
diff --git a/Engine/Source/Runtime/Renderer/Private/CompositionLighting/PostProcessAmbientOcclusion.h b/Engine/Source/Runtime/Renderer/Private/CompositionLighting/PostProcessAmbientOcclusion.h
--- a/Engine/Source/Runtime/Renderer/Private/CompositionLighting/PostProcessAmbientOcclusion.h
+++ b/Engine/Source/Runtime/Renderer/Private/CompositionLighting/PostProcessAmbientOcclusion.h
@@ -156,10 +156,78 @@ public:
 
 	void Bind(const FShaderParameterMap& ParameterMap);
 
-	void Set(FRHICommandList& RHICmdList, const FSceneView& View, const FPixelShaderRHIParamRef ShaderRHI) const;
+	template< typename ShaderRHIParamRef >
+	void Set(FRHICommandList& RHICmdList, const FSceneView& View, const ShaderRHIParamRef ShaderRHI) const;
 
 	friend FArchive& operator<<(FArchive& Ar, FCameraMotionParameters& This);
 
 private:
 	FShaderParameter CameraMotion;
 };
+
+
+template< typename ShaderRHIParamRef >
+void FCameraMotionParameters::Set(FRHICommandList& RHICmdList, const FSceneView& View, const ShaderRHIParamRef ShaderRHI) const
+{
+	FSceneViewState* ViewState = (FSceneViewState*)View.State;
+
+	FMatrix Proj = View.ViewMatrices.ProjMatrix;
+	FMatrix PrevProj = ViewState->PrevViewMatrices.ProjMatrix;
+
+	// Remove jitter
+	Proj.M[2][0] = 0.0f;
+	Proj.M[2][1] = 0.0f;
+	PrevProj.M[2][0] = 0.0f;
+	PrevProj.M[2][1] = 0.0f;
+
+	FMatrix ViewProj = ( View.ViewMatrices.ViewMatrix * Proj ).GetTransposed();
+	FMatrix PrevViewProj = ( ViewState->PrevViewMatrices.ViewMatrix * PrevProj ).GetTransposed();
+
+	double InvViewProj[16];
+	Inverse4x4( InvViewProj, (float*)ViewProj.M );
+
+	const float* p = (float*)PrevViewProj.M;
+
+	const double cxx = InvViewProj[ 0]; const double cxy = InvViewProj[ 1]; const double cxz = InvViewProj[ 2]; const double cxw = InvViewProj[ 3];
+	const double cyx = InvViewProj[ 4]; const double cyy = InvViewProj[ 5]; const double cyz = InvViewProj[ 6]; const double cyw = InvViewProj[ 7];
+	const double czx = InvViewProj[ 8]; const double czy = InvViewProj[ 9]; const double czz = InvViewProj[10]; const double czw = InvViewProj[11];
+	const double cwx = InvViewProj[12]; const double cwy = InvViewProj[13]; const double cwz = InvViewProj[14]; const double cww = InvViewProj[15];
+
+	const double pxx = (double)(p[ 0]); const double pxy = (double)(p[ 1]); const double pxz = (double)(p[ 2]); const double pxw = (double)(p[ 3]);
+	const double pyx = (double)(p[ 4]); const double pyy = (double)(p[ 5]); const double pyz = (double)(p[ 6]); const double pyw = (double)(p[ 7]);
+	const double pwx = (double)(p[12]); const double pwy = (double)(p[13]); const double pwz = (double)(p[14]); const double pww = (double)(p[15]);
+
+	FVector4 CameraMotionValue[5];
+
+	CameraMotionValue[0] = FVector4(
+		(float)(4.0*(cwx*pww + cxx*pwx + cyx*pwy + czx*pwz)),
+		(float)((-4.0)*(cwy*pww + cxy*pwx + cyy*pwy + czy*pwz)),
+		(float)(2.0*(cwz*pww + cxz*pwx + cyz*pwy + czz*pwz)),
+		(float)(2.0*(cww*pww - cwx*pww + cwy*pww + (cxw - cxx + cxy)*pwx + (cyw - cyx + cyy)*pwy + (czw - czx + czy)*pwz)));
+
+	CameraMotionValue[1] = FVector4(
+		(float)(( 4.0)*(cwy*pww + cxy*pwx + cyy*pwy + czy*pwz)),
+		(float)((-2.0)*(cwz*pww + cxz*pwx + cyz*pwy + czz*pwz)),
+		(float)((-2.0)*(cww*pww + cwy*pww + cxw*pwx - 2.0*cxx*pwx + cxy*pwx + cyw*pwy - 2.0*cyx*pwy + cyy*pwy + czw*pwz - 2.0*czx*pwz + czy*pwz - cwx*(2.0*pww + pxw) - cxx*pxx - cyx*pxy - czx*pxz)),
+		(float)(-2.0*(cyy*pwy + czy*pwz + cwy*(pww + pxw) + cxy*(pwx + pxx) + cyy*pxy + czy*pxz)));
+
+	CameraMotionValue[2] = FVector4(
+		(float)((-4.0)*(cwx*pww + cxx*pwx + cyx*pwy + czx*pwz)),
+		(float)(cyz*pwy + czz*pwz + cwz*(pww + pxw) + cxz*(pwx + pxx) + cyz*pxy + czz*pxz),
+		(float)(cwy*pww + cwy*pxw + cww*(pww + pxw) - cwx*(pww + pxw) + (cxw - cxx + cxy)*(pwx + pxx) + (cyw - cyx + cyy)*(pwy + pxy) + (czw - czx + czy)*(pwz + pxz)),
+		(float)(0));
+
+	CameraMotionValue[3] = FVector4(
+		(float)((-4.0)*(cwx*pww + cxx*pwx + cyx*pwy + czx*pwz)),
+		(float)((-2.0)*(cwz*pww + cxz*pwx + cyz*pwy + czz*pwz)),
+		(float)(2.0*((-cww)*pww + cwx*pww - 2.0*cwy*pww - cxw*pwx + cxx*pwx - 2.0*cxy*pwx - cyw*pwy + cyx*pwy - 2.0*cyy*pwy - czw*pwz + czx*pwz - 2.0*czy*pwz + cwy*pyw + cxy*pyx + cyy*pyy + czy*pyz)),
+		(float)(2.0*(cyx*pwy + czx*pwz + cwx*(pww - pyw) + cxx*(pwx - pyx) - cyx*pyy - czx*pyz)));
+
+	CameraMotionValue[4] = FVector4(
+		(float)(4.0*(cwy*pww + cxy*pwx + cyy*pwy + czy*pwz)),
+		(float)(cyz*pwy + czz*pwz + cwz*(pww - pyw) + cxz*(pwx - pyx) - cyz*pyy - czz*pyz),
+		(float)(cwy*pww + cww*(pww - pyw) - cwy*pyw + cwx*((-pww) + pyw) + (cxw - cxx + cxy)*(pwx - pyx) + (cyw - cyx + cyy)*(pwy - pyy) + (czw - czx + czy)*(pwz - pyz)),
+		(float)(0));
+
+	SetShaderValueArray(RHICmdList, ShaderRHI, CameraMotion, CameraMotionValue, 5);
+}
\ No newline at end of file
diff --git a/Engine/Source/Runtime/Renderer/Private/PostProcess/PostProcessMotionBlur.cpp b/Engine/Source/Runtime/Renderer/Private/PostProcess/PostProcessMotionBlur.cpp
--- a/Engine/Source/Runtime/Renderer/Private/PostProcess/PostProcessMotionBlur.cpp
+++ b/Engine/Source/Runtime/Renderer/Private/PostProcess/PostProcessMotionBlur.cpp
@@ -652,6 +652,412 @@ FPooledRenderTargetDesc FRCPassPostProcessMotionBlurRecombine::ComputeOutputDesc
 }
 
 
+
+class FPostProcessVelocityFlattenCS : public FGlobalShader
+{
+	DECLARE_SHADER_TYPE(FPostProcessVelocityFlattenCS, Global);
+
+	static bool ShouldCache(EShaderPlatform Platform)
+	{
+		return IsFeatureLevelSupported(Platform, ERHIFeatureLevel::SM5);
+	}
+
+	static void ModifyCompilationEnvironment(EShaderPlatform Platform, FShaderCompilerEnvironment& OutEnvironment)
+	{
+		FGlobalShader::ModifyCompilationEnvironment(Platform, OutEnvironment);
+		OutEnvironment.SetDefine(TEXT("THREADGROUP_SIZEX"), FRCPassPostProcessVelocityFlatten::ThreadGroupSizeX);
+		OutEnvironment.SetDefine(TEXT("THREADGROUP_SIZEY"), FRCPassPostProcessVelocityFlatten::ThreadGroupSizeY);
+		OutEnvironment.CompilerFlags.Add( CFLAG_StandardOptimization );
+	}
+
+	FPostProcessVelocityFlattenCS() {}
+
+public:
+	FShaderParameter		OutVelocityFlat;
+	FShaderParameter		OutPackedVelocityDepth;
+	FShaderParameter		OutMaxTileVelocity;
+	
+	FPostProcessVelocityFlattenCS(const ShaderMetaType::CompiledShaderInitializerType& Initializer)
+		: FGlobalShader(Initializer)
+	{
+		PostprocessParameter.Bind(Initializer.ParameterMap);
+		CameraMotionParams.Bind(Initializer.ParameterMap);
+		OutVelocityFlat.Bind(Initializer.ParameterMap, TEXT("OutVelocityFlat"));
+		OutPackedVelocityDepth.Bind(Initializer.ParameterMap, TEXT("OutPackedVelocityDepth"));
+		OutMaxTileVelocity.Bind(Initializer.ParameterMap, TEXT("OutMaxTileVelocity"));
+		ViewDimensions.Bind(Initializer.ParameterMap, TEXT("ViewDimensions"));
+	}
+
+	void SetCS( FRHICommandList& RHICmdList, const FRenderingCompositePassContext& Context, const FSceneView& View )
+	{
+		const FComputeShaderRHIParamRef ShaderRHI = GetComputeShader();
+
+		FGlobalShader::SetParameters(RHICmdList, ShaderRHI, Context.View);
+		PostprocessParameter.SetCS(ShaderRHI, Context, TStaticSamplerState<SF_Point,AM_Clamp,AM_Clamp,AM_Clamp>::GetRHI());
+		CameraMotionParams.Set(RHICmdList, Context.View, ShaderRHI);
+
+		SetShaderValue(RHICmdList, ShaderRHI, ViewDimensions, View.ViewRect);
+	}
+	
+	virtual bool Serialize(FArchive& Ar)
+	{
+		bool bShaderHasOutdatedParameters = FGlobalShader::Serialize(Ar);
+		Ar << PostprocessParameter;
+		Ar << CameraMotionParams;
+		Ar << OutVelocityFlat;
+		Ar << OutPackedVelocityDepth;
+		Ar << OutMaxTileVelocity;
+		Ar << ViewDimensions;
+		return bShaderHasOutdatedParameters;
+	}
+
+private:
+	FPostProcessPassParameters	PostprocessParameter;
+	FCameraMotionParameters		CameraMotionParams;
+	FShaderParameter			ViewDimensions;
+};
+
+IMPLEMENT_SHADER_TYPE(,FPostProcessVelocityFlattenCS,TEXT("PostProcessVelocityFlatten"),TEXT("VelocityFlattenMain"),SF_Compute);
+
+
+void FRCPassPostProcessVelocityFlatten::Process(FRenderingCompositePassContext& Context)
+{
+	SCOPED_DRAW_EVENT(Context.RHICmdList, PostProcessVelocityFlatten);
+	const FPooledRenderTargetDesc* InputDesc = GetInputDesc(ePId_Input0);
+
+	if(!InputDesc)
+	{
+		// input is not hooked up correctly
+		return;
+	}
+
+	const FSceneView& View = Context.View;
+	const FSceneViewFamily& ViewFamily = *(View.Family);
+	
+	FIntPoint SrcSize = InputDesc->Extent;
+	FIntRect DestRect = View.ViewRect;
+
+	const FSceneRenderTargetItem& DestRenderTarget0 = PassOutputs[0].RequestSurface(Context);
+	const FSceneRenderTargetItem& DestRenderTarget1 = PassOutputs[1].RequestSurface(Context);
+	//const FSceneRenderTargetItem& DestRenderTarget2 = PassOutputs[2].RequestSurface(Context);
+
+	TShaderMapRef< FPostProcessVelocityFlattenCS > ComputeShader( Context.GetShaderMap() );
+
+	Context.RHICmdList.SetComputeShader(ComputeShader->GetComputeShader());
+	SetRenderTarget(Context.RHICmdList, FTextureRHIRef(), FTextureRHIRef());
+
+	// set destination
+	Context.RHICmdList.SetUAVParameter( ComputeShader->GetComputeShader(), ComputeShader->OutVelocityFlat.GetBaseIndex(), DestRenderTarget0.UAV );
+	//Context.RHICmdList.SetUAVParameter( ComputeShader->GetComputeShader(), ComputeShader->OutPackedVelocityDepth.GetBaseIndex(), DestRenderTarget1.UAV );
+	Context.RHICmdList.SetUAVParameter( ComputeShader->GetComputeShader(), ComputeShader->OutMaxTileVelocity.GetBaseIndex(), DestRenderTarget1.UAV );
+
+	ComputeShader->SetCS(Context.RHICmdList, Context, View );
+	
+	FIntPoint ThreadGroupCountValue = ComputeThreadGroupCount( View.ViewRect.Size() );
+	DispatchComputeShader(Context.RHICmdList, *ComputeShader, ThreadGroupCountValue.X, ThreadGroupCountValue.Y, 1);
+
+	// un-set destination
+	Context.RHICmdList.SetUAVParameter( ComputeShader->GetComputeShader(), ComputeShader->OutVelocityFlat.GetBaseIndex(), NULL );
+	//Context.RHICmdList.SetUAVParameter( ComputeShader->GetComputeShader(), ComputeShader->OutPackedVelocityDepth.GetBaseIndex(), NULL );
+	Context.RHICmdList.SetUAVParameter( ComputeShader->GetComputeShader(), ComputeShader->OutMaxTileVelocity.GetBaseIndex(), NULL );
+
+	Context.RHICmdList.CopyToResolveTarget(DestRenderTarget0.TargetableTexture, DestRenderTarget0.ShaderResourceTexture, false, FResolveParams());
+	Context.RHICmdList.CopyToResolveTarget(DestRenderTarget1.TargetableTexture, DestRenderTarget1.ShaderResourceTexture, false, FResolveParams());
+	//Context.RHICmdList.CopyToResolveTarget(DestRenderTarget1.TargetableTexture, DestRenderTarget2.ShaderResourceTexture, false, FResolveParams());
+}
+
+FIntPoint FRCPassPostProcessVelocityFlatten::ComputeThreadGroupCount(FIntPoint PixelExtent)
+{
+	uint32 ThreadGroupCountX = (PixelExtent.X + TileSizeX - 1) / TileSizeX;
+	uint32 ThreadGroupCountY = (PixelExtent.Y + TileSizeY - 1) / TileSizeY;
+
+	return FIntPoint(ThreadGroupCountX, ThreadGroupCountY);
+}
+
+FPooledRenderTargetDesc FRCPassPostProcessVelocityFlatten::ComputeOutputDesc(EPassOutputId InPassOutputId) const
+{
+	if( InPassOutputId == ePId_Output0 )
+	{
+		// Flattened velocity
+		FPooledRenderTargetDesc Ret = PassInputs[0].GetOutput()->RenderTargetDesc;
+		Ret.Reset();
+
+		Ret.TargetableFlags |= TexCreate_UAV;
+		Ret.TargetableFlags |= TexCreate_RenderTargetable;
+		Ret.DebugName = TEXT("VelocityFlat");
+
+		return Ret;
+	}
+	/*else if( InPassOutputId == ePId_Output1 )
+	{
+		// Packed VelocityLength, Depth
+		FPooledRenderTargetDesc Ret = PassInputs[0].GetOutput()->RenderTargetDesc;
+		Ret.Reset();
+
+		Ret.Format = PF_R8G8;
+		Ret.TargetableFlags |= TexCreate_UAV;
+		Ret.TargetableFlags |= TexCreate_RenderTargetable;
+		Ret.DebugName = TEXT("PackedVelocityDepth");
+
+		return Ret;
+	}*/
+	else
+	{
+		// Max tile velocity
+		FPooledRenderTargetDesc UnmodifiedRet = PassInputs[0].GetOutput()->RenderTargetDesc;
+		UnmodifiedRet.Reset();
+
+		FIntPoint PixelExtent = UnmodifiedRet.Extent;
+
+		FIntPoint ThreadGroupCount = ComputeThreadGroupCount(PixelExtent);
+		FIntPoint NewSize = ThreadGroupCount;
+
+		// format can be optimized later
+		FPooledRenderTargetDesc Ret(FPooledRenderTargetDesc::Create2DDesc(NewSize, PF_G16R16F, TexCreate_None, TexCreate_RenderTargetable | TexCreate_UAV, false));
+
+		Ret.DebugName = TEXT("MaxVelocity");
+
+		return Ret;
+	}
+}
+
+/**
+ * @param Quality 0: visualize, 1:low, 2:medium, 3:high, 4:very high
+ */
+template< uint32 Quality >
+class FPostProcessMotionBlurNewPS : public FGlobalShader
+{
+	DECLARE_SHADER_TYPE(FPostProcessMotionBlurNewPS, Global);
+
+	static bool ShouldCache(EShaderPlatform Platform)
+	{
+		return IsFeatureLevelSupported(Platform, ERHIFeatureLevel::SM4);
+	}
+
+	static void ModifyCompilationEnvironment(EShaderPlatform Platform, FShaderCompilerEnvironment& OutEnvironment)
+	{
+		FGlobalShader::ModifyCompilationEnvironment(Platform,OutEnvironment);
+		OutEnvironment.SetDefine(TEXT("MOTION_BLUR_QUALITY"), Quality);
+	}
+
+	/** Default constructor. */
+	FPostProcessMotionBlurNewPS() {}
+
+public:
+	FPostProcessPassParameters PostprocessParameter;
+	FDeferredPixelShaderParameters DeferredParameters;
+	FShaderParameter MotionBlurParameters;
+
+	/** Initialization constructor. */
+	FPostProcessMotionBlurNewPS(const ShaderMetaType::CompiledShaderInitializerType& Initializer)
+		: FGlobalShader(Initializer)
+	{
+		PostprocessParameter.Bind(Initializer.ParameterMap);
+		DeferredParameters.Bind(Initializer.ParameterMap);
+		MotionBlurParameters.Bind(Initializer.ParameterMap, TEXT("MotionBlurParameters"));
+	}
+
+	// FShader interface.
+	virtual bool Serialize(FArchive& Ar)
+	{
+		bool bShaderHasOutdatedParameters = FGlobalShader::Serialize(Ar);
+		Ar << PostprocessParameter << DeferredParameters << MotionBlurParameters;
+		return bShaderHasOutdatedParameters;
+	}
+
+	void SetParameters(const FRenderingCompositePassContext& Context)
+	{
+		const FPixelShaderRHIParamRef ShaderRHI = GetPixelShader();
+
+		FGlobalShader::SetParameters(Context.RHICmdList, ShaderRHI, Context.View);
+
+		DeferredParameters.Set(Context.RHICmdList, ShaderRHI, Context.View);
+
+		{
+			bool bFiltered = false;
+
+#if !(UE_BUILD_SHIPPING || UE_BUILD_TEST)
+			bFiltered = CVarMotionBlurFiltering.GetValueOnRenderThread() != 0;
+#endif // !(UE_BUILD_SHIPPING || UE_BUILD_TEST)
+
+			if(bFiltered)
+			{
+				//PostprocessParameter.SetPS(ShaderRHI, Context, TStaticSamplerState<SF_Bilinear,AM_Border,AM_Border,AM_Clamp>::GetRHI());
+
+				FSamplerStateRHIParamRef Filters[] =
+				{
+					TStaticSamplerState<SF_Bilinear,AM_Border,AM_Border,AM_Clamp>::GetRHI(),
+					TStaticSamplerState<SF_Point,AM_Border,AM_Border,AM_Clamp>::GetRHI(),
+					TStaticSamplerState<SF_Bilinear,AM_Border,AM_Border,AM_Clamp>::GetRHI(),
+					TStaticSamplerState<SF_Point,AM_Border,AM_Border,AM_Clamp>::GetRHI(),
+				};
+
+				PostprocessParameter.SetPS( ShaderRHI, Context, 0, false, Filters );
+			}
+			else
+			{
+				PostprocessParameter.SetPS(ShaderRHI, Context, TStaticSamplerState<SF_Point,AM_Border,AM_Border,AM_Clamp>::GetRHI());
+			}
+		}
+
+		TRefCountPtr<IPooledRenderTarget> InputPooledElement = Context.Pass->GetInput(ePId_Input0)->GetOutput()->RequestInput();
+
+		{
+			const float SizeX = Context.View.ViewRect.Width();
+			const float SizeY = Context.View.ViewRect.Height();
+		
+			const float AspectRatio = SizeX / SizeY;
+			const float InvAspectRatio = SizeY / SizeX;
+
+			const FSceneViewState* ViewState = (FSceneViewState*) Context.View.State;
+			const float MotionBlurTimeScale = ViewState ? ViewState->MotionBlurTimeScale : 1.0f;
+
+			const float ViewMotionBlurScale = 0.5f * MotionBlurTimeScale * Context.View.FinalPostProcessSettings.MotionBlurAmount;
+
+			// 0:no 1:full screen width
+			float MaxVelocity = Context.View.FinalPostProcessSettings.MotionBlurMax / 100.0f;
+			float InvMaxVelocity = 1.0f / MaxVelocity;
+
+			// *2 to convert to -1..1 -1..1 screen space
+			// / MaxFraction to map screenpos to -1..1 normalized MaxFraction
+			FVector4 MotionBlurParametersValue(
+				ViewMotionBlurScale,
+				AspectRatio,
+				MaxVelocity,
+				InvMaxVelocity);
+			SetShaderValue(Context.RHICmdList, ShaderRHI, MotionBlurParameters, MotionBlurParametersValue);
+		}
+	}
+
+	static const TCHAR* GetSourceFilename()
+	{
+		return TEXT("PostProcessMotionBlur");
+	}
+
+	static const TCHAR* GetFunctionName()
+	{
+		return TEXT("MainNewPS");
+	}
+};
+
+// #define avoids a lot of code duplication
+#define VARIATION1(A) typedef FPostProcessMotionBlurNewPS<A> FPostProcessMotionBlurNewPS##A; \
+	IMPLEMENT_SHADER_TYPE2(FPostProcessMotionBlurNewPS##A, SF_Pixel);
+
+VARIATION1(0)			VARIATION1(1)			VARIATION1(2)			VARIATION1(3)			VARIATION1(4)
+#undef VARIATION1
+
+
+
+// @param Quality 0: visualize, 1:low, 2:medium, 3:high, 4:very high
+template< uint32 Quality >
+static void SetMotionBlurShaderNewTempl(const FRenderingCompositePassContext& Context)
+{
+	TShaderMapRef< FPostProcessVS >							VertexShader( Context.GetShaderMap() );
+	TShaderMapRef< FPostProcessMotionBlurNewPS< Quality > >	PixelShader( Context.GetShaderMap() );
+
+	static FGlobalBoundShaderState BoundShaderState;
+	SetGlobalBoundShaderState(Context.RHICmdList, Context.GetFeatureLevel(), BoundShaderState, GFilterVertexDeclaration.VertexDeclarationRHI, *VertexShader, *PixelShader);
+
+	VertexShader->SetParameters(Context);
+	PixelShader->SetParameters(Context);
+}
+
+FRCPassPostProcessMotionBlurNew::FRCPassPostProcessMotionBlurNew(uint32 InQuality) 
+	: Quality(InQuality)
+{
+	// internal error
+	check(Quality >= 1 && Quality <= 4);
+}
+
+void FRCPassPostProcessMotionBlurNew::Process(FRenderingCompositePassContext& Context)
+{
+	SCOPED_DRAW_EVENT(Context.RHICmdList, MotionBlur);
+
+	const FPooledRenderTargetDesc* InputDesc = GetInputDesc(ePId_Input0);
+
+	if(!InputDesc)
+	{
+		// input is not hooked up correctly
+		return;
+	}
+
+	const FSceneView& View = Context.View;
+
+	// we assume the input and output is full resolution
+
+	FIntPoint SrcSize = InputDesc->Extent;
+	FIntPoint DestSize = PassOutputs[0].RenderTargetDesc.Extent;
+
+	// e.g. 4 means the input texture is 4x smaller than the buffer size
+	uint32 ScaleFactor = GSceneRenderTargets.GetBufferSizeXY().X / SrcSize.X;
+
+	FIntRect SrcRect = FIntRect::DivideAndRoundUp(View.ViewRect, ScaleFactor);
+	FIntRect DestRect = SrcRect;
+
+	const FSceneRenderTargetItem& DestRenderTarget = PassOutputs[0].RequestSurface(Context);
+
+	// Set the view family's render target/viewport.
+	SetRenderTarget(Context.RHICmdList, DestRenderTarget.TargetableTexture, FTextureRHIRef());
+
+	// is optimized away if possible (RT size=view size, )
+	//Context.RHICmdList.Clear(true, FLinearColor::Black, false, 1.0f, false, 0, SrcRect);
+
+	Context.SetViewportAndCallRHI(SrcRect);
+
+	// set the state
+	Context.RHICmdList.SetBlendState(TStaticBlendState<>::GetRHI());
+	Context.RHICmdList.SetRasterizerState(TStaticRasterizerState<>::GetRHI());
+	Context.RHICmdList.SetDepthStencilState(TStaticDepthStencilState<false, CF_Always>::GetRHI());
+
+	if(Quality == 1)
+	{
+		SetMotionBlurShaderNewTempl<1>(Context);
+	}
+	else if(Quality == 2)
+	{
+		SetMotionBlurShaderNewTempl<2>(Context);
+	}
+	else if(Quality == 3)
+	{
+		SetMotionBlurShaderNewTempl<3>(Context);
+	}
+	else
+	{
+		check(Quality == 4);
+		SetMotionBlurShaderNewTempl<4>(Context);
+	}
+
+	TShaderMapRef<FPostProcessVS> VertexShader(Context.GetShaderMap());
+
+	// Draw a quad mapping scene color to the view's render target
+	DrawRectangle(
+		Context.RHICmdList,
+		0, 0,
+		SrcRect.Width(), SrcRect.Height(),
+		SrcRect.Min.X, SrcRect.Min.Y, 
+		SrcRect.Width(), SrcRect.Height(),
+		SrcRect.Size(),
+		SrcSize,
+		*VertexShader,
+		EDRF_UseTriangleOptimization);
+
+	Context.RHICmdList.CopyToResolveTarget(DestRenderTarget.TargetableTexture, DestRenderTarget.ShaderResourceTexture, false, FResolveParams());
+}
+
+FPooledRenderTargetDesc FRCPassPostProcessMotionBlurNew::ComputeOutputDesc(EPassOutputId InPassOutputId) const
+{
+	FPooledRenderTargetDesc Ret = PassInputs[0].GetOutput()->RenderTargetDesc;
+
+	Ret.Reset();
+	// we don't need the alpha channel and 32bit is faster and costs less memory
+	Ret.Format = PF_FloatRGB;
+	Ret.DebugName = TEXT("MotionBlur");
+
+	return Ret;
+}
+
+
 void FRCPassPostProcessVisualizeMotionBlur::Process(FRenderingCompositePassContext& Context)
 {
 	SCOPED_DRAW_EVENT(Context.RHICmdList, VisualizeMotionBlur);
diff --git a/Engine/Source/Runtime/Renderer/Private/PostProcess/PostProcessMotionBlur.h b/Engine/Source/Runtime/Renderer/Private/PostProcess/PostProcessMotionBlur.h
--- a/Engine/Source/Runtime/Renderer/Private/PostProcess/PostProcessMotionBlur.h
+++ b/Engine/Source/Runtime/Renderer/Private/PostProcess/PostProcessMotionBlur.h
@@ -71,3 +71,47 @@ public:
 	virtual void Release() override { delete this; }
 	virtual FPooledRenderTargetDesc ComputeOutputDesc(EPassOutputId InPassOutputId) const;
 };
+
+
+// Find max velocity per 16x16 tile
+// derives from TRenderingCompositePassBase<InputCount, OutputCount>
+class FRCPassPostProcessVelocityFlatten : public TRenderingCompositePassBase<2, 2>
+{
+public:
+	// interface FRenderingCompositePass ---------
+	virtual void Process(FRenderingCompositePassContext& Context);
+	virtual void Release() override { delete this; }
+	virtual FPooledRenderTargetDesc ComputeOutputDesc(EPassOutputId InPassOutputId) const;
+
+	// -------------------------------------------
+
+	static const uint32 TileSizeX = 64;
+	static const uint32 TileSizeY = 64;
+	
+	static const uint32 ThreadGroupSizeX = TileSizeX / 2;
+	static const uint32 ThreadGroupSizeY = TileSizeY / 2;
+
+	static FIntPoint ComputeThreadGroupCount(FIntPoint PixelExtent);
+};
+
+
+// input half res scene with depth, velocity, output half res motion blurred version
+// ePId_Input0: Full Res Scene Color
+// ePId_Input1: Full Res Scene Depth
+// ePId_Input2: Velocity
+// ePId_Input3: Max tile velocity
+// derives from TRenderingCompositePassBase<InputCount, OutputCount> 
+class FRCPassPostProcessMotionBlurNew : public TRenderingCompositePassBase<4, 1>
+{
+public:
+	// @param InQuality 0xffffffff to visualize, 0:off(no shader is used), 1:low, 2:medium, 3:high, 4:very high
+	FRCPassPostProcessMotionBlurNew(uint32 InQuality);
+
+	// interface FRenderingCompositePass ---------
+	virtual void Process(FRenderingCompositePassContext& Context);
+	virtual void Release() override { delete this; }
+	virtual FPooledRenderTargetDesc ComputeOutputDesc(EPassOutputId InPassOutputId) const;
+
+	// 1:low, 2:medium, 3:high, 4: very high
+	uint32 Quality;
+};
\ No newline at end of file
diff --git a/Engine/Source/Runtime/Renderer/Private/PostProcess/PostProcessTemporalAA.cpp b/Engine/Source/Runtime/Renderer/Private/PostProcess/PostProcessTemporalAA.cpp
--- a/Engine/Source/Runtime/Renderer/Private/PostProcess/PostProcessTemporalAA.cpp
+++ b/Engine/Source/Runtime/Renderer/Private/PostProcess/PostProcessTemporalAA.cpp
@@ -701,9 +701,11 @@ void FRCPassPostProcessTemporalAA::Process(FRenderingCompositePassContext& Conte
 
 	Context.RHICmdList.CopyToResolveTarget(DestRenderTarget.TargetableTexture, DestRenderTarget.ShaderResourceTexture, false, FResolveParams());
 
-	ViewState->TemporalAAHistoryRT = PassOutputs[0].PooledRenderTarget;
-	check( ViewState->TemporalAAHistoryRT );
-
+	if( !GRenderingRealtimeClock.GetGamePaused() )
+	{
+		ViewState->TemporalAAHistoryRT = PassOutputs[0].PooledRenderTarget;
+		check( ViewState->TemporalAAHistoryRT );
+	}
 
 	// TODO draw separate translucency after jitter has been removed
 
diff --git a/Engine/Source/Runtime/Renderer/Private/PostProcess/PostProcessing.cpp b/Engine/Source/Runtime/Renderer/Private/PostProcess/PostProcessing.cpp
--- a/Engine/Source/Runtime/Renderer/Private/PostProcess/PostProcessing.cpp
+++ b/Engine/Source/Runtime/Renderer/Private/PostProcess/PostProcessing.cpp
@@ -926,6 +926,13 @@ bool FPostProcessing::AllowFullPostProcessing(const FViewInfo& View, ERHIFeature
 		&& !View.Family->EngineShowFlags.VisualizeMeshDistanceFields;
 }
 
+static TAutoConsoleVariable<int32> CVarMotionBlurNew(
+	TEXT("r.MotionBlurNew"),
+	0,
+	TEXT(""),
+	ECVF_RenderThreadSafe
+);
+
 void FPostProcessing::Process(FRHICommandListImmediate& RHICmdList, FViewInfo& View, TRefCountPtr<IPooledRenderTarget>& VelocityRT)
 {
 	QUICK_SCOPE_CYCLE_COUNTER( STAT_PostProcessing_Process );
@@ -1106,53 +1113,80 @@ void FPostProcessing::Process(FRHICommandListImmediate& RHICmdList, FViewInfo& V
 
 			if(IsMotionBlurEnabled(View) && VelocityInput.IsValid())
 			{
-				// Motionblur is enabled
-				FRenderingCompositeOutputRef SoftEdgeVelocity;
+				// Motion blur
+
+				if( CVarMotionBlurNew.GetValueOnRenderThread() )
+				{
+					FRenderingCompositeOutputRef MaxTileVelocity;
+					FRenderingCompositeOutputRef SceneDepth( Context.SceneDepth );
 
-				FRenderingCompositeOutputRef MotionBlurHalfVelocity;
-				FRenderingCompositeOutputRef MotionBlurColorDepth;
+					{
+						FRenderingCompositePass* VelocityFlattenPass = Context.Graph.RegisterPass( new(FMemStack::Get()) FRCPassPostProcessVelocityFlatten() );
+						VelocityFlattenPass->SetInput( ePId_Input0, VelocityInput );
+						VelocityFlattenPass->SetInput( ePId_Input1, SceneDepth );
+
+						VelocityInput		= FRenderingCompositeOutputRef( VelocityFlattenPass, ePId_Output0 );
+						MaxTileVelocity		= FRenderingCompositeOutputRef( VelocityFlattenPass, ePId_Output1 );
+					}
 
-				// down sample screen space velocity and extend outside of borders to allows soft edge motion blur
+					{
+						FRenderingCompositePass* MotionBlurPass = Context.Graph.RegisterPass( new(FMemStack::Get()) FRCPassPostProcessMotionBlurNew( GetMotionBlurQualityFromCVar() ) );
+						MotionBlurPass->SetInput( ePId_Input0, Context.FinalOutput );
+						MotionBlurPass->SetInput( ePId_Input1, SceneDepth );
+						MotionBlurPass->SetInput( ePId_Input2, VelocityInput );
+						MotionBlurPass->SetInput( ePId_Input3, MaxTileVelocity );
+					
+						Context.FinalOutput = FRenderingCompositeOutputRef( MotionBlurPass );
+					}
+				}
+				else
 				{
-					// Down sample and prepare for soft masked blurring
-					FRenderingCompositePass* HalfResVelocity = Context.Graph.RegisterPass(new(FMemStack::Get()) FRCPassPostProcessMotionBlurSetup());
-					HalfResVelocity->SetInput(ePId_Input0, VelocityInput);
-					HalfResVelocity->SetInput(ePId_Input1, FRenderingCompositeOutputRef(Context.FinalOutput));
-					// only for Depth of Field we need the depth in the alpha channel
-					HalfResVelocity->SetInput(ePId_Input2, FRenderingCompositeOutputRef(Context.SceneDepth));
+					FRenderingCompositeOutputRef SoftEdgeVelocity;
+
+					FRenderingCompositeOutputRef MotionBlurHalfVelocity;
+					FRenderingCompositeOutputRef MotionBlurColorDepth;
+
+					// down sample screen space velocity and extend outside of borders to allows soft edge motion blur
+					{
+						// Down sample and prepare for soft masked blurring
+						FRenderingCompositePass* HalfResVelocity = Context.Graph.RegisterPass(new(FMemStack::Get()) FRCPassPostProcessMotionBlurSetup());
+						HalfResVelocity->SetInput(ePId_Input0, VelocityInput);
+						HalfResVelocity->SetInput(ePId_Input1, FRenderingCompositeOutputRef(Context.FinalOutput));
+						// only for Depth of Field we need the depth in the alpha channel
+						HalfResVelocity->SetInput(ePId_Input2, FRenderingCompositeOutputRef(Context.SceneDepth));
 					
-					MotionBlurHalfVelocity = FRenderingCompositeOutputRef(HalfResVelocity, ePId_Output0);
-					MotionBlurColorDepth = FRenderingCompositeOutputRef(HalfResVelocity, ePId_Output1);
+						MotionBlurHalfVelocity = FRenderingCompositeOutputRef(HalfResVelocity, ePId_Output0);
+						MotionBlurColorDepth = FRenderingCompositeOutputRef(HalfResVelocity, ePId_Output1);
 
-					FRenderingCompositePass* QuarterResVelocity = Context.Graph.RegisterPass(new(FMemStack::Get()) FRCPassPostProcessDownsample(PF_Unknown, 0, TEXT("QuarterResVelocity")));
-					QuarterResVelocity->SetInput(ePId_Input0, HalfResVelocity);
+						FRenderingCompositePass* QuarterResVelocity = Context.Graph.RegisterPass(new(FMemStack::Get()) FRCPassPostProcessDownsample(PF_Unknown, 0, TEXT("QuarterResVelocity")));
+						QuarterResVelocity->SetInput(ePId_Input0, HalfResVelocity);
 
-					SoftEdgeVelocity = FRenderingCompositeOutputRef(QuarterResVelocity);
+						SoftEdgeVelocity = FRenderingCompositeOutputRef(QuarterResVelocity);
 
-					float MotionBlurSoftEdgeSize = CVarMotionBlurSoftEdgeSize.GetValueOnRenderThread();
+						float MotionBlurSoftEdgeSize = CVarMotionBlurSoftEdgeSize.GetValueOnRenderThread();
 
-					if(MotionBlurSoftEdgeSize > 0.01f)
-					{
-						SoftEdgeVelocity = RenderGaussianBlur(Context, TEXT("VelocityBlurX"), TEXT("VelocityBlurY"), QuarterResVelocity, MotionBlurSoftEdgeSize);
+						if(MotionBlurSoftEdgeSize > 0.01f)
+						{
+							SoftEdgeVelocity = RenderGaussianBlur(Context, TEXT("VelocityBlurX"), TEXT("VelocityBlurY"), QuarterResVelocity, MotionBlurSoftEdgeSize);
+						}
 					}
-				}
 
-				// doing the actual motion blur sampling, alpha to mask out the blurred areas
-				FRenderingCompositePass* MotionBlurPass;
+					// doing the actual motion blur sampling, alpha to mask out the blurred areas
+					FRenderingCompositePass* MotionBlurPass;
 
-				if(View.Family->EngineShowFlags.VisualizeMotionBlur)
-				{
-					MotionBlurPass = Context.Graph.RegisterPass(new(FMemStack::Get()) FRCPassPostProcessVisualizeMotionBlur());
-					bAllowTonemapper = false;
-				}
-				else
-				{
-					int32 MotionBlurQuality = GetMotionBlurQualityFromCVar();
+					if(View.Family->EngineShowFlags.VisualizeMotionBlur)
+					{
+						MotionBlurPass = Context.Graph.RegisterPass(new(FMemStack::Get()) FRCPassPostProcessVisualizeMotionBlur());
+						bAllowTonemapper = false;
+					}
+					else
+					{
+						int32 MotionBlurQuality = GetMotionBlurQualityFromCVar();
 
-					MotionBlurPass = Context.Graph.RegisterPass(new(FMemStack::Get()) FRCPassPostProcessMotionBlur(MotionBlurQuality));
-				}
+						MotionBlurPass = Context.Graph.RegisterPass(new(FMemStack::Get()) FRCPassPostProcessMotionBlur(MotionBlurQuality));
+					}
 
-				MotionBlurPass->SetInput(ePId_Input0, MotionBlurColorDepth);
+					MotionBlurPass->SetInput(ePId_Input0, MotionBlurColorDepth);
 
 				if(VelocityInput.IsValid())
 				{
@@ -1162,10 +1196,11 @@ void FPostProcessing::Process(FRHICommandListImmediate& RHICmdList, FViewInfo& V
 					MotionBlurPass->SetInput(ePId_Input2, MotionBlurHalfVelocity);
 				}
 
-				FRenderingCompositePass* MotionBlurRecombinePass = Context.Graph.RegisterPass(new(FMemStack::Get()) FRCPassPostProcessMotionBlurRecombine());
-				MotionBlurRecombinePass->SetInput(ePId_Input0, Context.FinalOutput);
-				MotionBlurRecombinePass->SetInput(ePId_Input1, FRenderingCompositeOutputRef(MotionBlurPass));
-				Context.FinalOutput = FRenderingCompositeOutputRef(MotionBlurRecombinePass);
+					FRenderingCompositePass* MotionBlurRecombinePass = Context.Graph.RegisterPass(new(FMemStack::Get()) FRCPassPostProcessMotionBlurRecombine());
+					MotionBlurRecombinePass->SetInput(ePId_Input0, Context.FinalOutput);
+					MotionBlurRecombinePass->SetInput(ePId_Input1, FRenderingCompositeOutputRef(MotionBlurPass));
+					Context.FinalOutput = FRenderingCompositeOutputRef(MotionBlurRecombinePass);
+				}
 			}
 
 			if(bVisualizeBloom)
diff --git a/Engine/Source/Runtime/Renderer/Private/ScenePrivate.h b/Engine/Source/Runtime/Renderer/Private/ScenePrivate.h
--- a/Engine/Source/Runtime/Renderer/Private/ScenePrivate.h
+++ b/Engine/Source/Runtime/Renderer/Private/ScenePrivate.h
@@ -535,7 +535,10 @@ public:
 
 		TemporalAASampleCount = FMath::Min(SampleCount, (uint32)255);
 		
-		++TemporalAASampleIndex;
+		if( !GRenderingRealtimeClock.GetGamePaused() )
+		{
+			TemporalAASampleIndex++;
+		}
 
 		if(TemporalAASampleIndex >= TemporalAASampleCount)
 		{
