Fix Linux and Mac builds.

diff --git a/Engine/Source/Runtime/Renderer/Private/PostProcess/PostProcessMotionBlur.cpp b/Engine/Source/Runtime/Renderer/Private/PostProcess/PostProcessMotionBlur.cpp
--- a/Engine/Source/Runtime/Renderer/Private/PostProcess/PostProcessMotionBlur.cpp
+++ b/Engine/Source/Runtime/Renderer/Private/PostProcess/PostProcessMotionBlur.cpp
@@ -7,6 +7,7 @@
 #include "RendererPrivate.h"
 #include "ScenePrivate.h"
 #include "SceneFilterRendering.h"
+#include "PostProcessAmbientOcclusion.h"
 #include "PostProcessMotionBlur.h"
 #include "PostProcessing.h"
 #include "SceneUtils.h"
