Motion blur: added max velocity scatter and dilation.

diff --git a/Engine/Shaders/PostProcessMotionBlur.usf b/Engine/Shaders/PostProcessMotionBlur.usf
--- a/Engine/Shaders/PostProcessMotionBlur.usf
+++ b/Engine/Shaders/PostProcessMotionBlur.usf
@@ -18,7 +18,7 @@
 // but adapted for half resolution and not using randomization/noise
 
 // from the paper: We use SOFT Z EXTENT = 1mm to 10cm for our results
-#define SOFT_Z_EXTENT 0.1
+#define SOFT_Z_EXTENT 1
 
 // This,
 //  - Gets the MainPS pass to low cost under no motion.
@@ -627,6 +627,122 @@ void MainRecombinePS(float4 UVAndScreenPos : TEXCOORD0, out float4 OutColor : SV
 
 
 
+uint2 TileCount;
+
+void VelocityScatterVS(
+	uint VId : SV_VertexID,
+	uint IId : SV_InstanceID,
+	out nointerpolation float2 OutColor : TEXCOORD0,
+	out float4 OutPosition : SV_POSITION
+	)
+{
+	OutPosition =  float4(0, 0, 0, 1);
+
+	// needs to be the same on C++ side (faster on NVIDIA and AMD)
+	uint QuadsPerInstance = 8;
+#if MAC
+	// remap the indices to get vertexid to VId and quadid into IId
+	IId = IId * QuadsPerInstance + (VId / 6);
+	VId = VId % 6;
+
+	// triangle A: 0:left top, 1:right top, 2: left bottom
+	// triangle B: 3:right bottom, 4:left bottom, 5: right top
+	float2 CornerOffset = float2(VId % 2, VId > 1 && VId < 5) * 2 - 1;
+#else
+	// remap the indices to get vertexid to VId and quadid into IId
+	IId = IId * QuadsPerInstance + (VId / 4);
+	VId = VId % 4;
+
+	// triangle A: 0:left top, 1:right top, 2: left bottom
+	// triangle B: 3:right bottom, 4:left bottom, 5: right top
+	float2 CornerOffset = float2(VId % 2, VId / 2) * 2 - 1;
+#endif
+
+	uint2 TilePos = uint2( IId % TileCount.x, IId / TileCount.x );
+
+	float2 ScreenPos = ( TilePos - ScreenPosToPixel.zw ) / ScreenPosToPixel.xy;
+
+	float2 Velocity = PostprocessInput0[ TilePos ].xy;
+	OutColor = Velocity.xyxy;
+	
+	BRANCH
+	if( length( Velocity * 0.5 * float2(1920,1080) ) <= 0.5 )
+	{
+		return;
+	}
+
+	float  VelocityLengthSqr = dot( Velocity, Velocity );
+	float  VelocityLengthInv = rsqrt( VelocityLengthSqr );
+	float  VelocityLength = VelocityLengthSqr * VelocityLengthInv;
+	float2 VelocityDir = Velocity * VelocityLengthInv;
+
+	// +1 pixel radius because shape is swept tile, approx as capsule
+	// +1 pixel radius for conservative rasterization
+	const float TwoPixelRadius = sqrt( 2.0 );
+	CornerOffset *= float2( VelocityLength, 0 ) + TwoPixelRadius.xx / ScreenPosToPixel.xy;
+
+	// Orient along velocity direction
+	float2 AxisX = VelocityDir;
+	float2 AxisY = float2( -VelocityDir.y, VelocityDir.x );
+	CornerOffset = AxisX * CornerOffset.x + AxisY * CornerOffset.y;
+		
+	OutPosition.xy = ScreenPos + CornerOffset;
+
+	// Depth ordered by velocity length
+	OutPosition.z = saturate( VelocityLength * 0.5 );
+}
+
+void VelocityScatterPS(
+	nointerpolation float4 InColor : TEXCOORD0,
+	out float4 OutColor : SV_Target0
+	)
+{
+	OutColor = InColor;
+}
+
+
+
+
+float2 MaxLength( float2 v0, float2 v1 )
+{
+	float LengthSqr0 = dot( v0, v0 );
+	float LengthSqr1 = dot( v1, v1 );
+	return LengthSqr0 > LengthSqr1 ? v0 : v1;
+}
+
+float2 MaxLength( float2 v0, float2 v1, float2 v2 )
+{
+	float LengthSqr0 = dot( v0, v0 );
+	float LengthSqr1 = dot( v1, v1 );
+	float LengthSqr2 = dot( v2, v2 );
+	return LengthSqr0 > LengthSqr1 ? ( LengthSqr0 > LengthSqr2 ? v0 : v2 ) : ( LengthSqr1 > LengthSqr2 ? v1 : v2 );
+}
+
+void VelocityDilatePS(float4 UVAndScreenPos : TEXCOORD0, out float4 OutColor : SV_Target0)
+{
+	float2 UV = UVAndScreenPos.xy;
+
+	float2 Neighbor0 = PostprocessInput0.SampleLevel( PostprocessInput0Sampler, UV, 0, int2(-1, -1) ).xy;
+	float2 Neighbor1 = PostprocessInput0.SampleLevel( PostprocessInput0Sampler, UV, 0, int2( 0, -1) ).xy;
+	float2 Neighbor2 = PostprocessInput0.SampleLevel( PostprocessInput0Sampler, UV, 0, int2( 1, -1) ).xy;
+	float2 Neighbor3 = PostprocessInput0.SampleLevel( PostprocessInput0Sampler, UV, 0, int2(-1,  0) ).xy;
+	float2 Neighbor4 = PostprocessInput0.SampleLevel( PostprocessInput0Sampler, UV, 0				).xy;
+	float2 Neighbor5 = PostprocessInput0.SampleLevel( PostprocessInput0Sampler, UV, 0, int2( 1,  0) ).xy;
+	float2 Neighbor6 = PostprocessInput0.SampleLevel( PostprocessInput0Sampler, UV, 0, int2(-1,  1) ).xy;
+	float2 Neighbor7 = PostprocessInput0.SampleLevel( PostprocessInput0Sampler, UV, 0, int2( 0,  1) ).xy;
+	float2 Neighbor8 = PostprocessInput0.SampleLevel( PostprocessInput0Sampler, UV, 0, int2( 1,  1) ).xy;
+
+	float2 Max012 = MaxLength( Neighbor0, Neighbor1, Neighbor2 );
+	float2 Max345 = MaxLength( Neighbor3, Neighbor4, Neighbor5 );
+	float2 Max678 = MaxLength( Neighbor6, Neighbor7, Neighbor8 );
+
+	OutColor.xy = MaxLength( Max012, Max345, Max678 );
+	OutColor.zw = 0;
+}
+
+
+
+
 
 float2 DepthCmp( float CenterDepth, float SampleDepth, float DepthScale )
 {
diff --git a/Engine/Shaders/PostProcessVelocityFlatten.usf b/Engine/Shaders/PostProcessVelocityFlatten.usf
--- a/Engine/Shaders/PostProcessVelocityFlatten.usf
+++ b/Engine/Shaders/PostProcessVelocityFlatten.usf
@@ -101,7 +101,7 @@ void VelocityFlattenMain(
 
 	GroupMemoryBarrierWithGroupSync();
 
-#if 1
+#if 0
 	if( GroupIndex < 512 )	Shared[ GroupIndex ] = MaxLength( Shared[ GroupIndex ], Shared[ GroupIndex + 512 ] );
 	GroupMemoryBarrierWithGroupSync();
 	if( GroupIndex < 256 )	Shared[ GroupIndex ] = MaxLength( Shared[ GroupIndex ], Shared[ GroupIndex + 256 ] );
diff --git a/Engine/Source/Runtime/Renderer/Private/PostProcess/PostProcessMotionBlur.cpp b/Engine/Source/Runtime/Renderer/Private/PostProcess/PostProcessMotionBlur.cpp
--- a/Engine/Source/Runtime/Renderer/Private/PostProcess/PostProcessMotionBlur.cpp
+++ b/Engine/Source/Runtime/Renderer/Private/PostProcess/PostProcessMotionBlur.cpp
@@ -22,6 +22,14 @@ static TAutoConsoleVariable<int32> CVarMotionBlurFiltering(
 	ECVF_Cheat | ECVF_RenderThreadSafe);
 #endif
 
+static TAutoConsoleVariable<int32> CVarMotionBlurSmoothMax(
+	TEXT("r.MotionBlurSmoothMax"),
+	0,
+	TEXT("Useful developer variable\n")
+	TEXT("0: off (default, expected by the shader for better quality)\n")
+	TEXT("1: on"),
+	ECVF_Cheat | ECVF_RenderThreadSafe);
+
 
 /** Encapsulates the post processing motion blur vertex shader. */
 class FPostProcessMotionBlurSetupVS : public FGlobalShader
@@ -735,9 +743,6 @@ void FRCPassPostProcessVelocityFlatten::Process(FRenderingCompositePassContext&
 	const FSceneView& View = Context.View;
 	const FSceneViewFamily& ViewFamily = *(View.Family);
 	
-	FIntPoint SrcSize = InputDesc->Extent;
-	FIntRect DestRect = View.ViewRect;
-
 	const FSceneRenderTargetItem& DestRenderTarget0 = PassOutputs[0].RequestSurface(Context);
 	const FSceneRenderTargetItem& DestRenderTarget1 = PassOutputs[1].RequestSurface(Context);
 	//const FSceneRenderTargetItem& DestRenderTarget2 = PassOutputs[2].RequestSurface(Context);
@@ -822,6 +827,330 @@ FPooledRenderTargetDesc FRCPassPostProcessVelocityFlatten::ComputeOutputDesc(EPa
 	}
 }
 
+
+
+
+
+
+
+class FScatterQuadIndexBuffer : public FIndexBuffer
+{
+public:
+	virtual void InitRHI() override
+	{
+		const uint32 Size = sizeof(uint16) * 6 * 8;
+		const uint32 Stride = sizeof(uint16);
+		FRHIResourceCreateInfo CreateInfo;
+		IndexBufferRHI = RHICreateIndexBuffer( Stride, Size, BUF_Static, CreateInfo );
+		uint16* Indices = (uint16*)RHILockIndexBuffer( IndexBufferRHI, 0, Size, RLM_WriteOnly );
+		for (uint32 SpriteIndex = 0; SpriteIndex < 8; ++SpriteIndex)
+		{
+#if PLATFORM_MAC // Avoid a driver bug on OSX/NV cards that causes driver to generate an unwound index buffer
+			Indices[SpriteIndex*6 + 0] = SpriteIndex*6 + 0;
+			Indices[SpriteIndex*6 + 1] = SpriteIndex*6 + 1;
+			Indices[SpriteIndex*6 + 2] = SpriteIndex*6 + 2;
+			Indices[SpriteIndex*6 + 3] = SpriteIndex*6 + 3;
+			Indices[SpriteIndex*6 + 4] = SpriteIndex*6 + 4;
+			Indices[SpriteIndex*6 + 5] = SpriteIndex*6 + 5;
+#else
+			Indices[SpriteIndex*6 + 0] = SpriteIndex*4 + 0;
+			Indices[SpriteIndex*6 + 1] = SpriteIndex*4 + 3;
+			Indices[SpriteIndex*6 + 2] = SpriteIndex*4 + 2;
+			Indices[SpriteIndex*6 + 3] = SpriteIndex*4 + 0;
+			Indices[SpriteIndex*6 + 4] = SpriteIndex*4 + 1;
+			Indices[SpriteIndex*6 + 5] = SpriteIndex*4 + 3;
+#endif
+		}
+		RHIUnlockIndexBuffer( IndexBufferRHI );
+	}
+};
+
+TGlobalResource< FScatterQuadIndexBuffer > GScatterQuadIndexBuffer;
+
+
+class FPostProcessVelocityScatterVS : public FGlobalShader
+{
+	DECLARE_SHADER_TYPE(FPostProcessVelocityScatterVS,Global);
+
+	static bool ShouldCache(EShaderPlatform Platform)
+	{
+		return IsFeatureLevelSupported(Platform, ERHIFeatureLevel::SM4);
+	}
+
+	/** Default constructor. */
+	FPostProcessVelocityScatterVS() {}
+	
+public:
+	FPostProcessPassParameters PostprocessParameter;
+	FShaderParameter TileCount;
+
+	/** Initialization constructor. */
+	FPostProcessVelocityScatterVS(const ShaderMetaType::CompiledShaderInitializerType& Initializer)
+		: FGlobalShader(Initializer)
+	{
+		PostprocessParameter.Bind(Initializer.ParameterMap);
+		TileCount.Bind(Initializer.ParameterMap, TEXT("TileCount"));
+	}
+
+	// FShader interface.
+	virtual bool Serialize(FArchive& Ar)
+	{
+		bool bShaderHasOutdatedParameters = FGlobalShader::Serialize(Ar);
+		Ar << PostprocessParameter << TileCount;
+		return bShaderHasOutdatedParameters;
+	}
+
+	/** to have a similar interface as all other shaders */
+	void SetParameters(const FRenderingCompositePassContext& Context, FIntPoint TileCountValue)
+	{
+		const FVertexShaderRHIParamRef ShaderRHI = GetVertexShader();
+
+		FGlobalShader::SetParameters(Context.RHICmdList, ShaderRHI, Context.View);
+		PostprocessParameter.SetVS(ShaderRHI, Context, TStaticSamplerState<SF_Point,AM_Clamp,AM_Clamp,AM_Clamp>::GetRHI());
+
+		SetShaderValue(Context.RHICmdList, ShaderRHI, TileCount, TileCountValue);
+	}
+
+	static const TCHAR* GetSourceFilename()
+	{
+		return TEXT("PostProcessMotionBlur");
+	}
+
+	static const TCHAR* GetFunctionName()
+	{
+		return TEXT("VelocityScatterVS");
+	}
+};
+
+class FPostProcessVelocityScatterPS : public FGlobalShader
+{
+	DECLARE_SHADER_TYPE(FPostProcessVelocityScatterPS, Global);
+
+	static bool ShouldCache(EShaderPlatform Platform)
+	{
+		return IsFeatureLevelSupported(Platform, ERHIFeatureLevel::SM4);
+	}
+
+	/** Default constructor. */
+	FPostProcessVelocityScatterPS() {}
+
+public:
+	/** Initialization constructor. */
+	FPostProcessVelocityScatterPS(const ShaderMetaType::CompiledShaderInitializerType& Initializer)
+		: FGlobalShader(Initializer)
+	{}
+
+	void SetParameters(const FRenderingCompositePassContext& Context)
+	{
+		const FPixelShaderRHIParamRef ShaderRHI = GetPixelShader();
+
+		FGlobalShader::SetParameters(Context.RHICmdList, ShaderRHI, Context.View);
+	}
+};
+
+IMPLEMENT_SHADER_TYPE(,FPostProcessVelocityScatterVS,TEXT("PostProcessMotionBlur"),TEXT("VelocityScatterVS"),SF_Vertex);
+IMPLEMENT_SHADER_TYPE(,FPostProcessVelocityScatterPS,TEXT("PostProcessMotionBlur"),TEXT("VelocityScatterPS"),SF_Pixel);
+
+
+void FRCPassPostProcessVelocityScatter::Process(FRenderingCompositePassContext& Context)
+{
+	SCOPED_DRAW_EVENT(Context.RHICmdList, PassPostProcessVelocityScatter);
+
+	const FPooledRenderTargetDesc* InputDesc = GetInputDesc(ePId_Input0);
+	
+	if(!InputDesc)
+	{
+		// input is not hooked up correctly
+		return;
+	}
+
+	const FSceneView& View = Context.View;
+
+	FIntPoint SrcSize = InputDesc->Extent;
+	FIntPoint DestSize = PassOutputs[0].RenderTargetDesc.Extent;
+
+	// e.g. 4 means the input texture is 4x smaller than the buffer size
+	uint32 ScaleFactor = GSceneRenderTargets.GetBufferSizeXY().X / SrcSize.X;
+
+	FIntRect SrcRect = FIntRect::DivideAndRoundUp(View.ViewRect, ScaleFactor);
+	FIntRect DestRect = SrcRect;
+
+	const FSceneRenderTargetItem& DestRenderTarget = PassOutputs[0].RequestSurface(Context);
+
+	TRefCountPtr<IPooledRenderTarget> DepthTarget;
+	
+	FPooledRenderTargetDesc Desc( FPooledRenderTargetDesc::Create2DDesc( DestRect.Size(), PF_ShadowDepth, TexCreate_None, TexCreate_DepthStencilTargetable, false ) );
+	GRenderTargetPool.FindFreeElement( Desc, DepthTarget, TEXT("VelocityScatterDepth") );
+
+	// Set the view family's render target/viewport.
+	SetRenderTarget( Context.RHICmdList, DestRenderTarget.TargetableTexture, DepthTarget->GetRenderTargetItem().TargetableTexture );
+
+	Context.RHICmdList.Clear( true, FLinearColor::Black, true, 0.0f, false, 0, FIntRect() );
+
+	Context.SetViewportAndCallRHI(SrcRect);
+
+	// set the state
+	Context.RHICmdList.SetBlendState(TStaticBlendState<>::GetRHI());
+	Context.RHICmdList.SetRasterizerState(TStaticRasterizerState<>::GetRHI());
+	Context.RHICmdList.SetDepthStencilState(TStaticDepthStencilState<true, CF_Greater>::GetRHI());
+
+	TShaderMapRef< FPostProcessVelocityScatterVS > VertexShader(Context.GetShaderMap());
+	TShaderMapRef< FPostProcessVelocityScatterPS > PixelShader(Context.GetShaderMap());
+
+	static FGlobalBoundShaderState BoundShaderState;
+	SetGlobalBoundShaderState(Context.RHICmdList, Context.GetFeatureLevel(), BoundShaderState, GFilterVertexDeclaration.VertexDeclarationRHI, *VertexShader, *PixelShader);
+
+	FIntPoint TileCount = SrcRect.Size();
+
+	VertexShader->SetParameters( Context, TileCount );
+	PixelShader->SetParameters( Context );
+
+	// needs to be the same on shader side (faster on NVIDIA and AMD)
+	int32 QuadsPerInstance = 8;
+
+	Context.RHICmdList.SetStreamSource(0, NULL, 0, 0);
+	Context.RHICmdList.DrawIndexedPrimitive(GScatterQuadIndexBuffer.IndexBufferRHI, PT_TriangleList, 0, 0, 32, 0, 2 * QuadsPerInstance, FMath::DivideAndRoundUp(TileCount.X * TileCount.Y, QuadsPerInstance));
+
+	Context.RHICmdList.CopyToResolveTarget(DestRenderTarget.TargetableTexture, DestRenderTarget.ShaderResourceTexture, false, FResolveParams());
+}
+
+FPooledRenderTargetDesc FRCPassPostProcessVelocityScatter::ComputeOutputDesc(EPassOutputId InPassOutputId) const
+{
+	FPooledRenderTargetDesc Ret = PassInputs[0].GetOutput()->RenderTargetDesc;
+	Ret.Reset();
+
+	Ret.DebugName = TEXT("ScatteredMaxVelocity");
+
+	return Ret;
+}
+
+
+
+
+
+
+
+
+class FPostProcessVelocityDilatePS : public FGlobalShader
+{
+	DECLARE_SHADER_TYPE(FPostProcessVelocityDilatePS, Global);
+
+	static bool ShouldCache(EShaderPlatform Platform)
+	{
+		return IsFeatureLevelSupported(Platform, ERHIFeatureLevel::SM4);
+	}
+
+	static void ModifyCompilationEnvironment(EShaderPlatform Platform, FShaderCompilerEnvironment& OutEnvironment)
+	{
+		FGlobalShader::ModifyCompilationEnvironment(Platform, OutEnvironment);
+	}
+
+	/** Default constructor. */
+	FPostProcessVelocityDilatePS() {}
+
+public:
+	FPostProcessPassParameters PostprocessParameter;
+
+	/** Initialization constructor. */
+	FPostProcessVelocityDilatePS(const ShaderMetaType::CompiledShaderInitializerType& Initializer)
+		: FGlobalShader(Initializer)
+	{
+		PostprocessParameter.Bind(Initializer.ParameterMap);
+	}
+
+	// FShader interface.
+	virtual bool Serialize(FArchive& Ar)
+	{
+		bool bShaderHasOutdatedParameters = FGlobalShader::Serialize(Ar);
+		Ar << PostprocessParameter;
+		return bShaderHasOutdatedParameters;
+	}
+
+	void SetParameters(const FRenderingCompositePassContext& Context)
+	{
+		const FPixelShaderRHIParamRef ShaderRHI = GetPixelShader();
+
+		FGlobalShader::SetParameters(Context.RHICmdList, ShaderRHI, Context.View);
+
+		PostprocessParameter.SetPS(ShaderRHI, Context, TStaticSamplerState<SF_Point,AM_Clamp,AM_Clamp,AM_Clamp>::GetRHI());
+	}
+};
+
+IMPLEMENT_SHADER_TYPE(,FPostProcessVelocityDilatePS,TEXT("PostProcessMotionBlur"),TEXT("VelocityDilatePS"),SF_Pixel);
+
+void FRCPassPostProcessVelocityDilate::Process(FRenderingCompositePassContext& Context)
+{
+	SCOPED_DRAW_EVENT(Context.RHICmdList, VelocityDilate);
+
+	const FPooledRenderTargetDesc* InputDesc = GetInputDesc(ePId_Input0);
+
+	if(!InputDesc)
+	{
+		// input is not hooked up correctly
+		return;
+	}
+
+	const FSceneView& View = Context.View;
+
+	FIntPoint SrcSize = InputDesc->Extent;
+	FIntPoint DestSize = PassOutputs[0].RenderTargetDesc.Extent;
+
+	// e.g. 4 means the input texture is 4x smaller than the buffer size
+	uint32 ScaleFactor = GSceneRenderTargets.GetBufferSizeXY().X / SrcSize.X;
+
+	FIntRect SrcRect = FIntRect::DivideAndRoundUp(View.ViewRect, ScaleFactor);
+	FIntRect DestRect = SrcRect;
+
+	const FSceneRenderTargetItem& DestRenderTarget = PassOutputs[0].RequestSurface(Context);
+
+	// Set the view family's render target/viewport.
+	SetRenderTarget(Context.RHICmdList, DestRenderTarget.TargetableTexture, FTextureRHIRef());
+
+	// is optimized away if possible (RT size=view size, )
+	//Context.RHICmdList.Clear(true, FLinearColor::Black, false, 1.0f, false, 0, SrcRect);
+
+	Context.SetViewportAndCallRHI(SrcRect);
+
+	// set the state
+	Context.RHICmdList.SetBlendState(TStaticBlendState<>::GetRHI());
+	Context.RHICmdList.SetRasterizerState(TStaticRasterizerState<>::GetRHI());
+	Context.RHICmdList.SetDepthStencilState(TStaticDepthStencilState<false, CF_Always>::GetRHI());
+
+	TShaderMapRef< FPostProcessVS>					VertexShader( Context.GetShaderMap() );
+	TShaderMapRef< FPostProcessVelocityDilatePS >	PixelShader( Context.GetShaderMap() );
+
+	static FGlobalBoundShaderState BoundShaderState;
+	SetGlobalBoundShaderState(Context.RHICmdList, Context.GetFeatureLevel(), BoundShaderState, GFilterVertexDeclaration.VertexDeclarationRHI, *VertexShader, *PixelShader);
+
+	VertexShader->SetParameters(Context);
+	PixelShader->SetParameters(Context);
+
+	// Draw a quad mapping scene color to the view's render target
+	DrawRectangle(
+		Context.RHICmdList,
+		0, 0,
+		SrcRect.Width(), SrcRect.Height(),
+		SrcRect.Min.X, SrcRect.Min.Y, 
+		SrcRect.Width(), SrcRect.Height(),
+		SrcRect.Size(),
+		SrcSize,
+		*VertexShader,
+		EDRF_UseTriangleOptimization);
+
+	Context.RHICmdList.CopyToResolveTarget(DestRenderTarget.TargetableTexture, DestRenderTarget.ShaderResourceTexture, false, FResolveParams());
+}
+
+FPooledRenderTargetDesc FRCPassPostProcessVelocityDilate::ComputeOutputDesc(EPassOutputId InPassOutputId) const
+{
+	FPooledRenderTargetDesc Ret = PassInputs[0].GetOutput()->RenderTargetDesc;
+	Ret.Reset();
+
+	Ret.DebugName = TEXT("DilatedMaxVelocity");
+
+	return Ret;
+}
+
+
 /**
  * @param Quality 0: visualize, 1:low, 2:medium, 3:high, 4:very high
  */
@@ -889,12 +1218,24 @@ public:
 				{
 					TStaticSamplerState<SF_Bilinear,AM_Clamp,AM_Clamp,AM_Clamp>::GetRHI(),
 					TStaticSamplerState<SF_Point,AM_Clamp,AM_Clamp,AM_Clamp>::GetRHI(),
-					TStaticSamplerState<SF_Bilinear,AM_Clamp,AM_Clamp,AM_Clamp>::GetRHI(),
+					TStaticSamplerState<SF_Point,AM_Clamp,AM_Clamp,AM_Clamp>::GetRHI(),
 					TStaticSamplerState<SF_Point,AM_Clamp,AM_Clamp,AM_Clamp>::GetRHI(),
 				};
 
 				PostprocessParameter.SetPS( ShaderRHI, Context, 0, false, Filters );
 			}
+			else if( CVarMotionBlurSmoothMax.GetValueOnRenderThread() )
+			{
+				FSamplerStateRHIParamRef Filters[] =
+				{
+					TStaticSamplerState<SF_Point,AM_Clamp,AM_Clamp,AM_Clamp>::GetRHI(),
+					TStaticSamplerState<SF_Point,AM_Clamp,AM_Clamp,AM_Clamp>::GetRHI(),
+					TStaticSamplerState<SF_Point,AM_Clamp,AM_Clamp,AM_Clamp>::GetRHI(),
+					TStaticSamplerState<SF_Bilinear,AM_Clamp,AM_Clamp,AM_Clamp>::GetRHI(),
+				};
+
+				PostprocessParameter.SetPS( ShaderRHI, Context, 0, false, Filters );
+			}
 			else
 			{
 				PostprocessParameter.SetPS(ShaderRHI, Context, TStaticSamplerState<SF_Point,AM_Clamp,AM_Clamp,AM_Clamp>::GetRHI());
@@ -993,7 +1334,7 @@ void FRCPassPostProcessMotionBlurNew::Process(FRenderingCompositePassContext& Co
 	// e.g. 4 means the input texture is 4x smaller than the buffer size
 	uint32 ScaleFactor = GSceneRenderTargets.GetBufferSizeXY().X / SrcSize.X;
 
-	FIntRect SrcRect = FIntRect::DivideAndRoundUp(View.ViewRect, ScaleFactor);
+	FIntRect SrcRect = View.ViewRect / ScaleFactor;
 	FIntRect DestRect = SrcRect;
 
 	const FSceneRenderTargetItem& DestRenderTarget = PassOutputs[0].RequestSurface(Context);
diff --git a/Engine/Source/Runtime/Renderer/Private/PostProcess/PostProcessMotionBlur.h b/Engine/Source/Runtime/Renderer/Private/PostProcess/PostProcessMotionBlur.h
--- a/Engine/Source/Runtime/Renderer/Private/PostProcess/PostProcessMotionBlur.h
+++ b/Engine/Source/Runtime/Renderer/Private/PostProcess/PostProcessMotionBlur.h
@@ -85,8 +85,8 @@ public:
 
 	// -------------------------------------------
 
-	static const uint32 TileSizeX = 64;
-	static const uint32 TileSizeY = 64;
+	static const uint32 TileSizeX = 16;
+	static const uint32 TileSizeY = 16;
 	
 	static const uint32 ThreadGroupSizeX = TileSizeX / 2;
 	static const uint32 ThreadGroupSizeY = TileSizeY / 2;
@@ -95,6 +95,27 @@ public:
 };
 
 
+// derives from TRenderingCompositePassBase<InputCount, OutputCount>
+class FRCPassPostProcessVelocityScatter : public TRenderingCompositePassBase<1, 1>
+{
+public:
+	// interface FRenderingCompositePass ---------
+	virtual void Process(FRenderingCompositePassContext& Context);
+	virtual void Release() override { delete this; }
+	virtual FPooledRenderTargetDesc ComputeOutputDesc(EPassOutputId InPassOutputId) const;
+};
+
+// derives from TRenderingCompositePassBase<InputCount, OutputCount>
+class FRCPassPostProcessVelocityDilate : public TRenderingCompositePassBase<1, 1>
+{
+public:
+	// interface FRenderingCompositePass ---------
+	virtual void Process(FRenderingCompositePassContext& Context);
+	virtual void Release() override { delete this; }
+	virtual FPooledRenderTargetDesc ComputeOutputDesc(EPassOutputId InPassOutputId) const;
+};
+
+
 // input half res scene with depth, velocity, output half res motion blurred version
 // ePId_Input0: Full Res Scene Color
 // ePId_Input1: Full Res Scene Depth
diff --git a/Engine/Source/Runtime/Renderer/Private/PostProcess/PostProcessing.cpp b/Engine/Source/Runtime/Renderer/Private/PostProcess/PostProcessing.cpp
--- a/Engine/Source/Runtime/Renderer/Private/PostProcess/PostProcessing.cpp
+++ b/Engine/Source/Runtime/Renderer/Private/PostProcess/PostProcessing.cpp
@@ -933,6 +933,20 @@ static TAutoConsoleVariable<int32> CVarMotionBlurNew(
 	ECVF_RenderThreadSafe
 );
 
+static TAutoConsoleVariable<int32> CVarMotionBlurScatter(
+	TEXT("r.MotionBlurScatter"),
+	1,
+	TEXT(""),
+	ECVF_RenderThreadSafe
+);
+
+static TAutoConsoleVariable<int32> CVarMotionBlurDilate(
+	TEXT("r.MotionBlurDilate"),
+	0,
+	TEXT(""),
+	ECVF_RenderThreadSafe
+);
+
 void FPostProcessing::Process(FRHICommandListImmediate& RHICmdList, FViewInfo& View, TRefCountPtr<IPooledRenderTarget>& VelocityRT)
 {
 	QUICK_SCOPE_CYCLE_COUNTER( STAT_PostProcessing_Process );
@@ -1125,8 +1139,24 @@ void FPostProcessing::Process(FRHICommandListImmediate& RHICmdList, FViewInfo& V
 						VelocityFlattenPass->SetInput( ePId_Input0, VelocityInput );
 						VelocityFlattenPass->SetInput( ePId_Input1, SceneDepth );
 
-						VelocityInput		= FRenderingCompositeOutputRef( VelocityFlattenPass, ePId_Output0 );
-						MaxTileVelocity		= FRenderingCompositeOutputRef( VelocityFlattenPass, ePId_Output1 );
+						VelocityInput	= FRenderingCompositeOutputRef( VelocityFlattenPass, ePId_Output0 );
+						MaxTileVelocity	= FRenderingCompositeOutputRef( VelocityFlattenPass, ePId_Output1 );
+					}
+
+					if( CVarMotionBlurScatter.GetValueOnRenderThread() )
+					{
+						FRenderingCompositePass* VelocityScatterPass = Context.Graph.RegisterPass( new(FMemStack::Get()) FRCPassPostProcessVelocityScatter() );
+						VelocityScatterPass->SetInput( ePId_Input0, MaxTileVelocity );
+
+						MaxTileVelocity	= FRenderingCompositeOutputRef( VelocityScatterPass );
+					}
+
+					if( CVarMotionBlurDilate.GetValueOnRenderThread() )
+					{
+						FRenderingCompositePass* VelocityDilatePass = Context.Graph.RegisterPass( new(FMemStack::Get()) FRCPassPostProcessVelocityDilate() );
+						VelocityDilatePass->SetInput( ePId_Input0, MaxTileVelocity );
+
+						MaxTileVelocity	= FRenderingCompositeOutputRef( VelocityDilatePass );
 					}
 
 					{
