fixed typo

diff --git a/Engine/Source/Runtime/Renderer/Private/PostProcess/SceneRenderTargets.cpp b/Engine/Source/Runtime/Renderer/Private/PostProcess/SceneRenderTargets.cpp
--- a/Engine/Source/Runtime/Renderer/Private/PostProcess/SceneRenderTargets.cpp
+++ b/Engine/Source/Runtime/Renderer/Private/PostProcess/SceneRenderTargets.cpp
@@ -1207,7 +1207,7 @@ static TAutoConsoleVariable<int32> CVarSetSeperateTranslucencyEnabled(
 	TEXT("Allows to disable the separate translucency feature (all translucency is rendered in separate RT and composited\n")
 	TEXT("after DOF, if not specified otherwise in the material).\n")
 	TEXT(" 0: off (translucency is affected by depth of field)\n")
-	TEXT(" 1: on costs GPU performance and memory but keeps translucency unaffected by Depth of Fieled. (default)"),
+	TEXT(" 1: on costs GPU performance and memory but keeps translucency unaffected by Depth of Field. (default)"),
 	ECVF_RenderThreadSafe);
 
 bool FSceneRenderTargets::IsSeparateTranslucencyActive(const FViewInfo& View) const
