Integrate - Added r.Shadow.TexelsPerPixelSpotlight so spotlight shadow resolution can be controlled separately from per-object shadows

diff --git a/Engine/Source/Runtime/Renderer/Private/ShadowSetup.cpp b/Engine/Source/Runtime/Renderer/Private/ShadowSetup.cpp
--- a/Engine/Source/Runtime/Renderer/Private/ShadowSetup.cpp
+++ b/Engine/Source/Runtime/Renderer/Private/ShadowSetup.cpp
@@ -92,7 +92,13 @@ static TAutoConsoleVariable<float> CVarPreShadowResolutionFactor(
 static TAutoConsoleVariable<float> CVarShadowTexelsPerPixel(
 	TEXT("r.Shadow.TexelsPerPixel"),
 	1.27324f,
-	TEXT("The ratio of subject pixels to shadow texels"),
+	TEXT("The ratio of subject pixels to shadow texels for per-object shadows"),
+	ECVF_RenderThreadSafe);
+
+static TAutoConsoleVariable<float> CVarShadowTexelsPerPixelSpotlight(
+	TEXT("r.Shadow.TexelsPerPixelSpotlight"),
+	1.27324f,
+	TEXT("The ratio of subject pixels to shadow texels for spotlights"),
 	ECVF_RenderThreadSafe);
 
 static TAutoConsoleVariable<int32> CVarPreShadowFadeResolution(
@@ -1560,7 +1566,7 @@ void FDeferredShadingSceneRenderer::CreateWholeSceneProjectedShadow(FLightSceneI
 				FMath::Max(ScreenPosition.W,1.0f);
 
 			// Determine the amount of shadow buffer resolution needed for this view.
-			const float UnclampedResolution = ScreenRadius * CVarShadowTexelsPerPixel.GetValueOnRenderThread();
+			const float UnclampedResolution = ScreenRadius * CVarShadowTexelsPerPixelSpotlight.GetValueOnRenderThread();
 			MaxUnclampedResolution = FMath::Max( MaxUnclampedResolution, UnclampedResolution );
 			MaxDesiredResolution = FMath::Max(
 				MaxDesiredResolution,
