Integrate - Support for changing r.AOInterpolationDepthTesting

diff --git a/Engine/Shaders/DistanceFieldSurfaceCacheLighting.usf b/Engine/Shaders/DistanceFieldSurfaceCacheLighting.usf
--- a/Engine/Shaders/DistanceFieldSurfaceCacheLighting.usf
+++ b/Engine/Shaders/DistanceFieldSurfaceCacheLighting.usf
@@ -168,6 +168,11 @@ float InterpolationRadiusScale;
 float2 NormalizedOffsetToPixelCenter;
 float HackExpand;
 
+// +1 = behind volume (safe from near plane clipping without depth testing), 
+// -1 = in front of volume (required for conservative results with depth testing)
+//@todo - using -1 currently causes artifacts with architectural scenes
+float InterpolationBoundingDirection;
+
 /** Expands a screen-facing polygon to cover a surface cache record for splatting. */
 void IrradianceCacheSplatVS(
 	float2 InPosition : ATTRIBUTE0,
@@ -187,11 +192,7 @@ void IrradianceCacheSplatVS(
 	float EffectiveInterpolationRadiusScale = lerp(InterpolationRadiusScale, 1, max(saturate(PositionAndRadius.w / 80), .2f));
 	PositionAndRadius.w = max(PositionAndRadius.w, ViewToCenterLength * .01f) * EffectiveInterpolationRadiusScale;
 
-	// +1 = behind volume (safe from near plane clipping without depth testing), 
-	// -1 = in front of volume (required for conservative results with depth testing)
-	//@todo - using -1 currently causes artifacts
-	float BoundingDirection = +1;
-	float OffsetFromCenter = BoundingDirection * PositionAndRadius.w;
+	float OffsetFromCenter = InterpolationBoundingDirection * PositionAndRadius.w;
 	// Distance from the camera position that NormalizedViewToCenter will intersect the near plane
 	float NearPlaneDistance = View.NearPlane / dot(View.ViewForward, NormalizedViewToCenter);
 
diff --git a/Engine/Source/Runtime/Renderer/Private/DistanceFieldSurfaceCacheLighting.cpp b/Engine/Source/Runtime/Renderer/Private/DistanceFieldSurfaceCacheLighting.cpp
--- a/Engine/Source/Runtime/Renderer/Private/DistanceFieldSurfaceCacheLighting.cpp
+++ b/Engine/Source/Runtime/Renderer/Private/DistanceFieldSurfaceCacheLighting.cpp
@@ -1931,6 +1931,7 @@ public:
 		InterpolationRadiusScale.Bind(Initializer.ParameterMap, TEXT("InterpolationRadiusScale"));
 		NormalizedOffsetToPixelCenter.Bind(Initializer.ParameterMap, TEXT("NormalizedOffsetToPixelCenter"));
 		HackExpand.Bind(Initializer.ParameterMap, TEXT("HackExpand"));
+		InterpolationBoundingDirection.Bind(Initializer.ParameterMap, TEXT("InterpolationBoundingDirection"));
 		AOParameters.Bind(Initializer.ParameterMap);
 	}
 
@@ -1957,6 +1958,9 @@ public:
 		const FVector2D HackExpandValue(.5f / AOViewRectSize.X, .5f / AOViewRectSize.Y);
 		SetShaderValue(RHICmdList, ShaderRHI, HackExpand, HackExpandValue);
 
+		// Must push the bounding geometry toward the camera with depth testing, to be in front of any potential receivers
+		SetShaderValue(RHICmdList, ShaderRHI, InterpolationBoundingDirection, GAOInterpolationDepthTesting ? -1.0f : 1.0f);
+
 		AOParameters.Set(RHICmdList, ShaderRHI, Parameters);
 	}
 
@@ -1971,6 +1975,7 @@ public:
 		Ar << InterpolationRadiusScale;
 		Ar << NormalizedOffsetToPixelCenter;
 		Ar << HackExpand;
+		Ar << InterpolationBoundingDirection;
 		Ar << AOParameters;
 		return bShaderHasOutdatedParameters;
 	}
@@ -1984,6 +1989,7 @@ private:
 	FShaderParameter InterpolationRadiusScale;
 	FShaderParameter NormalizedOffsetToPixelCenter;
 	FShaderParameter HackExpand;
+	FShaderParameter InterpolationBoundingDirection;
 	FAOParameters AOParameters;
 };
 
