Integrate - Fixed primitives always updating their static draw list meshes on load, even when they are not dirty.

diff --git a/Engine/Source/Runtime/Renderer/Private/PrimitiveSceneInfo.cpp b/Engine/Source/Runtime/Renderer/Private/PrimitiveSceneInfo.cpp
--- a/Engine/Source/Runtime/Renderer/Private/PrimitiveSceneInfo.cpp
+++ b/Engine/Source/Runtime/Renderer/Private/PrimitiveSceneInfo.cpp
@@ -96,7 +96,7 @@ FPrimitiveSceneInfo::FPrimitiveSceneInfo(UPrimitiveComponent* InComponent,FScene
 	Scene(InScene),
 	PackedIndex(INDEX_NONE),
 	ComponentForDebuggingOnly(InComponent),
-	bNeedsStaticMeshUpdate(true)
+	bNeedsStaticMeshUpdate(false)
 {
 	check(ComponentForDebuggingOnly);
 	check(PrimitiveComponentId.IsValid());
