fixed directional lights for TLM_SurfacePerPixelLighting, remaining: shadow support, light functions, IES profiles

diff --git a/Engine/Shaders/ForEachLight.usf b/Engine/Shaders/ForEachLight.usf
--- a/Engine/Shaders/ForEachLight.usf
+++ b/Engine/Shaders/ForEachLight.usf
@@ -28,12 +28,12 @@ float3 LightingForOneLight(FMaterialPixelParameters Parameters, uint LightIndex,
 		LightData.bSpotLight = Value.w > 0;
 		LightData.MinRoughness = abs(Value.w);
 	}
-	LightData.SpotAnglesAndSourceRadius = float4( ForwardLightData.SpotAnglesAndSourceRadiusAndSimpleLighting[LightIndex].xyz, 0 );
+	LightData.SpotAnglesAndSourceRadius = float4( ForwardLightData.SpotAnglesAndSourceRadiusAndDir[LightIndex].xyz, 0 );
 	// currently we don't support shadows
 	LightData.ShadowMapChannelMask = 0;
 	LightData.bInverseSquared = ForwardLightData.LightColorAndFalloffExponent[LightIndex].w == 0;
 	// Only radial lights supported with tiled deferred
-	LightData.bRadialLight = true;
+	LightData.bRadialLight = ForwardLightData.SpotAnglesAndSourceRadiusAndDir[LightIndex].w == 0;
 	// The only type of shadowing supported for lights using tiled is static shadowing, so the light should only compute shadowing if it has static shadowing
 	LightData.bShadowed = dot(LightData.ShadowMapChannelMask, float4(1, 1, 1, 1));
 	
diff --git a/Engine/Source/Runtime/Engine/Public/SceneView.h b/Engine/Source/Runtime/Engine/Public/SceneView.h
--- a/Engine/Source/Runtime/Engine/Public/SceneView.h
+++ b/Engine/Source/Runtime/Engine/Public/SceneView.h
@@ -239,7 +239,7 @@ BEGIN_UNIFORM_BUFFER_STRUCT(FForwardLightData,ENGINE_API)
 	DECLARE_UNIFORM_BUFFER_STRUCT_MEMBER_ARRAY(FVector4,LightPositionAndInvRadius,[GMaxNumForwardLights])
 	DECLARE_UNIFORM_BUFFER_STRUCT_MEMBER_ARRAY(FVector4,LightColorAndFalloffExponent,[GMaxNumForwardLights])
 	DECLARE_UNIFORM_BUFFER_STRUCT_MEMBER_ARRAY(FVector4,LightDirectionAndSpotlightMaskAndMinRoughness,[GMaxNumForwardLights])
-	DECLARE_UNIFORM_BUFFER_STRUCT_MEMBER_ARRAY(FVector4,SpotAnglesAndSourceRadiusAndSimpleLighting,[GMaxNumForwardLights])
+	DECLARE_UNIFORM_BUFFER_STRUCT_MEMBER_ARRAY(FVector4,SpotAnglesAndSourceRadiusAndDir,[GMaxNumForwardLights])
 END_UNIFORM_BUFFER_STRUCT(FForwardLightData)
 
 //////////////////////////////////////////////////////////////////////////
diff --git a/Engine/Source/Runtime/Renderer/Private/SceneRendering.cpp b/Engine/Source/Runtime/Renderer/Private/SceneRendering.cpp
--- a/Engine/Source/Runtime/Renderer/Private/SceneRendering.cpp
+++ b/Engine/Source/Runtime/Renderer/Private/SceneRendering.cpp
@@ -735,41 +735,31 @@ void FViewInfo::CreateForwardLightDataUniformBuffer(FForwardLightData &OutForwar
 					SourceLength,
 					MinRoughness);
 
-				if(LightType == LightType_Point || LightType == LightType_Spot)
+				// Check if the light is visible in this view.
+				if(!LightSceneInfo->ShouldRenderLight(*this))
 				{
-					// Check if the light is visible in this view.
-					if(!LightSceneInfo->ShouldRenderLight(*this))
-					{
-						continue;
-					}
+					continue;
+				}
 
-					FVector4 BoundingSphereVector = *(FVector4*)&LightSceneInfoCompact.BoundingSphereVector;
+				FVector4 BoundingSphereVector = *(FVector4*)&LightSceneInfoCompact.BoundingSphereVector;
 
-					float InvRadius = 1.0f / BoundingSphereVector.W;
+				float InvRadius = 1.0f / BoundingSphereVector.W;
 
-					OutForwardLightData.LightPositionAndInvRadius[LightIndex] = FVector4(FVector(BoundingSphereVector), InvRadius);
+				OutForwardLightData.LightPositionAndInvRadius[LightIndex] = FVector4(FVector(BoundingSphereVector), InvRadius);
 
-					// SpotlightMaskAndMinRoughness, >0:Spotlight, MinRoughness = abs();
-					float W = FMath::Max(0.0001f, MinRoughness) * ((LightType == LightType_Spot) ? 1 : -1);
+				// SpotlightMaskAndMinRoughness, >0:Spotlight, MinRoughness = abs();
+				float W = FMath::Max(0.0001f, MinRoughness) * ((LightType == LightType_Spot) ? 1 : -1);
 
-					OutForwardLightData.LightDirectionAndSpotlightMaskAndMinRoughness[LightIndex] = FVector4(NormalizedLightDirection, W);
-					OutForwardLightData.SpotAnglesAndSourceRadiusAndSimpleLighting[LightIndex] = FVector4(SpotAngles.X, SpotAngles.Y, SourceRadius, 0);
+				OutForwardLightData.LightDirectionAndSpotlightMaskAndMinRoughness[LightIndex] = FVector4(NormalizedLightDirection, W);
+				OutForwardLightData.SpotAnglesAndSourceRadiusAndDir[LightIndex] = FVector4(SpotAngles.X, SpotAngles.Y, SourceRadius, LightType == LightType_Directional);
 
-					if (LightSceneInfo->Proxy->IsInverseSquared())
-					{
-						// Correction for lumen units
-						OutForwardLightData.LightColorAndFalloffExponent[LightIndex].X *= 16.0f;
-						OutForwardLightData.LightColorAndFalloffExponent[LightIndex].Y *= 16.0f;
-						OutForwardLightData.LightColorAndFalloffExponent[LightIndex].Z *= 16.0f;
-						OutForwardLightData.LightColorAndFalloffExponent[LightIndex].W = 0;
-					}
-				}
-				else
+				if(LightSceneInfo->Proxy->IsInverseSquared())
 				{
-//					FVector Direction = -LightSceneInfo->Proxy->GetDirection();
-
-					// we don't handle directional light types yet
-					continue;
+					// Correction for lumen units
+					OutForwardLightData.LightColorAndFalloffExponent[LightIndex].X *= 16.0f;
+					OutForwardLightData.LightColorAndFalloffExponent[LightIndex].Y *= 16.0f;
+					OutForwardLightData.LightColorAndFalloffExponent[LightIndex].Z *= 16.0f;
+					OutForwardLightData.LightColorAndFalloffExponent[LightIndex].W = 0;
 				}
 
 				{
@@ -779,8 +769,6 @@ void FViewInfo::CreateForwardLightDataUniformBuffer(FForwardLightData &OutForwar
 					OutForwardLightData.LightDirectionAndSpotlightMaskAndMinRoughness[LightIndex] = FVector4(NormalizedLightDirection, W);
 				}
 
-				OutForwardLightData.SpotAnglesAndSourceRadiusAndSimpleLighting[LightIndex] = FVector4(SpotAngles.X, SpotAngles.Y, SourceRadius, 0);
-
 				// we want to add one light
 				++LightIndex;
 
