Integrate - Ray traced distance field shadows now implement the cascade fade

diff --git a/Engine/Shaders/DistanceFieldShadowing.usf b/Engine/Shaders/DistanceFieldShadowing.usf
--- a/Engine/Shaders/DistanceFieldShadowing.usf
+++ b/Engine/Shaders/DistanceFieldShadowing.usf
@@ -601,6 +601,9 @@ void DistanceFieldShadowingCS(
 Texture2D ShadowFactorsTexture;
 SamplerState ShadowFactorsSampler;
 
+float FadePlaneOffset;
+float InvFadePlaneLength;
+
 void DistanceFieldShadowingUpsamplePS(
 	in float4 UVAndScreenPos : TEXCOORD0,
 	in float4 SVPos : SV_POSITION,
@@ -608,6 +611,7 @@ void DistanceFieldShadowingUpsamplePS(
 {
 	// Distance field shadowing was computed at 0,0 regardless of viewrect min
 	float2 DistanceFieldUVs = UVAndScreenPos.xy - ScissorRectMinAndSize.xy * View.ViewSizeAndSceneTexelSize.zw;
+	float SceneDepth = CalcSceneDepth(UVAndScreenPos.xy);
 
 #define BILATERAL_UPSAMPLE 1
 #if BILATERAL_UPSAMPLE && UPSAMPLE_REQUIRED
@@ -630,7 +634,6 @@ void DistanceFieldShadowingUpsamplePS(
 	float Epsilon = .0001f;
 
 	float4 CornerDepths = abs(float4(TextureValues00.y, TextureValues10.y, TextureValues01.y, TextureValues11.y));
-	float SceneDepth = CalcSceneDepth(UVAndScreenPos.xy);
 	float4 DepthWeights = 1.0f / (abs(CornerDepths - SceneDepth.xxxx) + Epsilon);
 	float4 FinalWeights = CornerWeights * DepthWeights;
 
@@ -647,5 +650,7 @@ void DistanceFieldShadowingUpsamplePS(
 	float Output = Texture2DSampleLevel(ShadowFactorsTexture, ShadowFactorsSampler, DistanceFieldUVs, 0).x;
 #endif
 
-	OutColor = EncodeLightAttenuation(half4(Output, Output, Output, Output));
+	float BlendFactor = 1.0f - saturate((SceneDepth - FadePlaneOffset) * InvFadePlaneLength);
+	Output = lerp(1, Output, BlendFactor);
+	OutColor = EncodeLightAttenuation(half4(Output, Output, Output, BlendFactor));
 }
\ No newline at end of file
diff --git a/Engine/Source/Runtime/Renderer/Private/DistanceFieldShadowing.cpp b/Engine/Source/Runtime/Renderer/Private/DistanceFieldShadowing.cpp
--- a/Engine/Source/Runtime/Renderer/Private/DistanceFieldShadowing.cpp
+++ b/Engine/Source/Runtime/Renderer/Private/DistanceFieldShadowing.cpp
@@ -567,6 +567,8 @@ public:
 		ShadowFactorsTexture.Bind(Initializer.ParameterMap,TEXT("ShadowFactorsTexture"));
 		ShadowFactorsSampler.Bind(Initializer.ParameterMap,TEXT("ShadowFactorsSampler"));
 		ScissorRectMinAndSize.Bind(Initializer.ParameterMap,TEXT("ScissorRectMinAndSize"));
+		FadePlaneOffset.Bind(Initializer.ParameterMap,TEXT("FadePlaneOffset"));
+		InvFadePlaneLength.Bind(Initializer.ParameterMap,TEXT("InvFadePlaneLength"));
 	}
 
 	void SetParameters(FRHICommandList& RHICmdList, const FSceneView& View, const FProjectedShadowInfo* ShadowInfo, const FIntRect& ScissorRect, TRefCountPtr<IPooledRenderTarget>& ShadowFactorsTextureValue)
@@ -579,6 +581,17 @@ public:
 		SetTextureParameter(RHICmdList, ShaderRHI, ShadowFactorsTexture, ShadowFactorsSampler, TStaticSamplerState<SF_Bilinear>::GetRHI(), ShadowFactorsTextureValue->GetRenderTargetItem().ShaderResourceTexture);
 	
 		SetShaderValue(RHICmdList, ShaderRHI, ScissorRectMinAndSize, FIntRect(ScissorRect.Min, ScissorRect.Size()));
+
+		if (ShadowInfo->bDirectionalLight && ShadowInfo->CascadeSettings.FadePlaneLength > 0)
+		{
+			SetShaderValue(RHICmdList, ShaderRHI, FadePlaneOffset, ShadowInfo->CascadeSettings.FadePlaneOffset);
+			SetShaderValue(RHICmdList, ShaderRHI, InvFadePlaneLength, 1.0f / FMath::Max(ShadowInfo->CascadeSettings.FadePlaneLength, .00001f));
+		}
+		else
+		{
+			SetShaderValue(RHICmdList, ShaderRHI, FadePlaneOffset, 0.0f);
+			SetShaderValue(RHICmdList, ShaderRHI, InvFadePlaneLength, 0.0f);
+		}
 	}
 
 	// FShader interface.
@@ -589,6 +602,8 @@ public:
 		Ar << ShadowFactorsTexture;
 		Ar << ShadowFactorsSampler;
 		Ar << ScissorRectMinAndSize;
+		Ar << FadePlaneOffset;
+		Ar << InvFadePlaneLength;
 		return bShaderHasOutdatedParameters;
 	}
 
@@ -598,6 +613,8 @@ private:
 	FShaderResourceParameter ShadowFactorsTexture;
 	FShaderResourceParameter ShadowFactorsSampler;
 	FShaderParameter ScissorRectMinAndSize;
+	FShaderParameter FadePlaneOffset;
+	FShaderParameter InvFadePlaneLength;
 };
 
 IMPLEMENT_SHADER_TYPE(template<>,TDistanceFieldShadowingUpsamplePS<true>,TEXT("DistanceFieldShadowing"),TEXT("DistanceFieldShadowingUpsamplePS"),SF_Pixel);
@@ -863,20 +880,17 @@ void FProjectedShadowInfo::RenderRayTracedDistanceFieldProjection(FRHICommandLis
 				RHICmdList.SetViewport(ScissorRect.Min.X, ScissorRect.Min.Y, 0.0f, ScissorRect.Max.X, ScissorRect.Max.Y, 1.0f);
 				RHICmdList.SetRasterizerState(TStaticRasterizerState<FM_Solid, CM_None>::GetRHI());
 				RHICmdList.SetDepthStencilState(TStaticDepthStencilState<false, CF_Always>::GetRHI());
-					
+				
 				if (bDirectionalLight)
 				{
-					// todo: when using far shadows we need to fix the shader to get a soft transition
-//					check(CascadeSettings.FadePlaneLength == 0);
-
-					// first cascade rendered or old method doesn't require fading (CO_Min is needed to combine multiple shadow passes)
-					// The ray traced cascade should always be first
+					// use R and G in Light Attenuation for directional lights
+					// CO_Min is needed to combine with far shadows which overlap the same depth range
 					RHICmdList.SetBlendState(TStaticBlendState<CW_RG, BO_Min, BF_One, BF_One>::GetRHI());
 				}
 				else
 				{
 					// use B and A in Light Attenuation
-					// (CO_Min is needed to combine multiple shadow passes)
+					// CO_Min is needed to combine multiple shadow passes
 					RHICmdList.SetBlendState(TStaticBlendState<CW_BA, BO_Min, BF_One, BF_One, BO_Min, BF_One, BF_One>::GetRHI());
 				}
 
