Integrate - View.ViewToTranslatedWorld is now correct for shadow passes, allows billboards to face the light

diff --git a/Engine/Source/Runtime/Renderer/Private/SceneRendering.cpp b/Engine/Source/Runtime/Renderer/Private/SceneRendering.cpp
--- a/Engine/Source/Runtime/Renderer/Private/SceneRendering.cpp
+++ b/Engine/Source/Runtime/Renderer/Private/SceneRendering.cpp
@@ -321,6 +321,7 @@ void FViewInfo::SetupSkyIrradianceEnvironmentMapConstants(FVector4* OutSkyIrradi
 TUniformBufferRef<FViewUniformShaderParameters> FViewInfo::CreateUniformBuffer(
 	const TArray<FProjectedShadowInfo*, SceneRenderingAllocator>* DirectionalLightShadowInfo,	
 	const FMatrix& EffectiveTranslatedViewMatrix, 
+	const FMatrix& EffectiveViewToTranslatedWorld, 
 	FBox* OutTranslucentCascadeBoundsArray, 
 	int32 NumTranslucentCascades) const
 {
@@ -373,7 +374,7 @@ TUniformBufferRef<FViewUniformShaderParameters> FViewInfo::CreateUniformBuffer(
 	ViewUniformShaderParameters.TranslatedWorldToClip = ViewMatrices.TranslatedViewProjectionMatrix;
 	ViewUniformShaderParameters.WorldToClip = ViewProjectionMatrix;
 	ViewUniformShaderParameters.TranslatedWorldToView = EffectiveTranslatedViewMatrix;
-	ViewUniformShaderParameters.ViewToTranslatedWorld = InvViewMatrix * FTranslationMatrix(ViewMatrices.PreViewTranslation);
+	ViewUniformShaderParameters.ViewToTranslatedWorld = EffectiveViewToTranslatedWorld;
 	ViewUniformShaderParameters.ViewToClip = ViewMatrices.ProjMatrix;
 	ViewUniformShaderParameters.ClipToView = ViewMatrices.GetInvProjMatrix();
 	ViewUniformShaderParameters.ClipToTranslatedWorld = ViewMatrices.InvTranslatedViewProjectionMatrix;
@@ -850,6 +851,7 @@ void FViewInfo::InitRHIResources(const TArray<FProjectedShadowInfo*, SceneRender
 	UniformBuffer = CreateUniformBuffer(
 		DirectionalLightShadowInfo,
 		TranslatedViewMatrix,
+		InvViewMatrix * FTranslationMatrix(ViewMatrices.PreViewTranslation),
 		VolumeBounds,
 		TVC_MAX);
 
diff --git a/Engine/Source/Runtime/Renderer/Private/SceneRendering.h b/Engine/Source/Runtime/Renderer/Private/SceneRendering.h
--- a/Engine/Source/Runtime/Renderer/Private/SceneRendering.h
+++ b/Engine/Source/Runtime/Renderer/Private/SceneRendering.h
@@ -501,6 +501,7 @@ public:
 	TUniformBufferRef<FViewUniformShaderParameters> CreateUniformBuffer(
 		const TArray<FProjectedShadowInfo*, SceneRenderingAllocator>* DirectionalLightShadowInfo,
 		const FMatrix& EffectiveTranslatedViewMatrix, 
+		const FMatrix& EffectiveViewToTranslatedWorld, 
 		FBox* OutTranslucentCascadeBoundsArray, 
 		int32 NumTranslucentCascades) const;
 
diff --git a/Engine/Source/Runtime/Renderer/Private/ShadowRendering.cpp b/Engine/Source/Runtime/Renderer/Private/ShadowRendering.cpp
--- a/Engine/Source/Runtime/Renderer/Private/ShadowRendering.cpp
+++ b/Engine/Source/Runtime/Renderer/Private/ShadowRendering.cpp
@@ -1694,6 +1694,7 @@ void FProjectedShadowInfo::RenderDepth(FRHICommandList& RHICmdList, FSceneRender
 	FoundView->UniformBuffer = FoundView->CreateUniformBuffer(
 		nullptr,
 		ShadowViewMatrix, 
+		ShadowViewMatrix.Inverse(),
 		VolumeBounds,
 		TVC_MAX);
 
diff --git a/Engine/Source/Runtime/Renderer/Private/TranslucentLighting.cpp b/Engine/Source/Runtime/Renderer/Private/TranslucentLighting.cpp
--- a/Engine/Source/Runtime/Renderer/Private/TranslucentLighting.cpp
+++ b/Engine/Source/Runtime/Renderer/Private/TranslucentLighting.cpp
@@ -513,6 +513,7 @@ void FProjectedShadowInfo::RenderTranslucencyDepths(FRHICommandList& RHICmdList,
 	FoundView->UniformBuffer = FoundView->CreateUniformBuffer(
 		nullptr,
 		ShadowViewMatrix, 
+		ShadowViewMatrix.Inverse(),
 		VolumeBounds,
 		TVC_MAX);
 		
