Fix shadowed variables

diff --git a/Engine/Source/Developer/GameplayDebugger/Private/GameplayDebuggingComponent.cpp b/Engine/Source/Developer/GameplayDebugger/Private/GameplayDebuggingComponent.cpp
--- a/Engine/Source/Developer/GameplayDebugger/Private/GameplayDebuggingComponent.cpp
+++ b/Engine/Source/Developer/GameplayDebugger/Private/GameplayDebuggingComponent.cpp
@@ -330,9 +330,9 @@ void UGameplayDebuggingComponent::SelectTargetToDebug()
 
 		if (PossibleTarget != NULL && PossibleTarget != GetSelectedActor())
 		{
-			if (AGameplayDebuggingReplicator* Replicator = Cast<AGameplayDebuggingReplicator>(GetOwner()))
+			if (AGameplayDebuggingReplicator* DebuggingReplicator = Cast<AGameplayDebuggingReplicator>(GetOwner()))
 			{
-				Replicator->ServerSetActorToDebug(Cast<AActor>(PossibleTarget));
+				DebuggingReplicator->ServerSetActorToDebug(Cast<AActor>(PossibleTarget));
 			}
 
 			ServerReplicateData(EDebugComponentMessage::ActivateReplication, EAIDebugDrawDataView::Empty);
diff --git a/Engine/Source/Developer/SourceControl/Private/SourceControlHelpers.cpp b/Engine/Source/Developer/SourceControl/Private/SourceControlHelpers.cpp
--- a/Engine/Source/Developer/SourceControl/Private/SourceControlHelpers.cpp
+++ b/Engine/Source/Developer/SourceControl/Private/SourceControlHelpers.cpp
@@ -318,14 +318,14 @@ bool CopyFileUnderSourceControl( const FString& InDestFile, const FString& InSou
 {
 	struct Local
 	{
-		static bool CopyFile(const FString& InDestFile, const FText& InFileDescription, FText& OutFailReason, FString InSourceFile)
+		static bool CopyFile(const FString& InDestinationFile, const FText& InFileDesc, FText& OutFailureReason, FString InFileToCopy)
 		{
 			const bool bReplace = true;
 			const bool bEvenIfReadOnly = true;
-			bool bSucceeded = (IFileManager::Get().Copy(*InDestFile, *InSourceFile, bReplace, bEvenIfReadOnly) == COPY_OK);
+			bool bSucceeded = (IFileManager::Get().Copy(*InDestinationFile, *InFileToCopy, bReplace, bEvenIfReadOnly) == COPY_OK);
 			if (!bSucceeded)
 			{
-				OutFailReason = FText::Format(LOCTEXT("ExternalImageCopyError", "Could not overwrite {0} file."), InFileDescription);
+				OutFailureReason = FText::Format(LOCTEXT("ExternalImageCopyError", "Could not overwrite {0} file."), InFileDesc);
 			}
 
 			return bSucceeded;
diff --git a/Engine/Source/Developer/iOS/IOSTargetPlatform/Private/Mac/IOSDeviceHelperMac.cpp b/Engine/Source/Developer/iOS/IOSTargetPlatform/Private/Mac/IOSDeviceHelperMac.cpp
--- a/Engine/Source/Developer/iOS/IOSTargetPlatform/Private/Mac/IOSDeviceHelperMac.cpp
+++ b/Engine/Source/Developer/iOS/IOSTargetPlatform/Private/Mac/IOSDeviceHelperMac.cpp
@@ -113,8 +113,8 @@ public:
     bool Connect()
     {
         bool Result = false;
-		int LogLevel = 7;
-		CFPreferencesSetAppValue(CFSTR("LogLevel"), CFNumberCreate(kCFAllocatorDefault, kCFNumberSInt32Type, &LogLevel), CFSTR("com.apple.MobileDevice"));
+		int AppLogLevel = 7;
+		CFPreferencesSetAppValue(CFSTR("LogLevel"), CFNumberCreate(kCFAllocatorDefault, kCFNumberSInt32Type, &AppLogLevel), CFSTR("com.apple.MobileDevice"));
 
 		// connect to the device
         int32 rc = AMDeviceConnect(DeviceHandle);
diff --git a/Engine/Source/Editor/ContentBrowser/Private/ContentBrowserUtils.cpp b/Engine/Source/Editor/ContentBrowser/Private/ContentBrowserUtils.cpp
--- a/Engine/Source/Editor/ContentBrowser/Private/ContentBrowserUtils.cpp
+++ b/Engine/Source/Editor/ContentBrowser/Private/ContentBrowserUtils.cpp
@@ -1373,16 +1373,16 @@ const TSharedPtr<FLinearColor> ContentBrowserUtils::LoadColor(const FString& Fol
 
 void ContentBrowserUtils::SaveColor(const FString& FolderPath, const TSharedPtr<FLinearColor>& FolderColor, bool bForceAdd)
 {
-	auto SaveColorInternal = [](const FString& InPath, const TSharedPtr<FLinearColor>& FolderColor)
+	auto SaveColorInternal = [](const FString& InPath, const TSharedPtr<FLinearColor>& InFolderColor)
 	{
 		// Saves the color of the folder to the config
 		if(FPaths::FileExists(GEditorUserSettingsIni))
 		{
-			GConfig->SetString(TEXT("PathColor"), *InPath, *FolderColor->ToString(), GEditorUserSettingsIni);
+			GConfig->SetString(TEXT("PathColor"), *InPath, *InFolderColor->ToString(), GEditorUserSettingsIni);
 		}
 
 		// Update the map too
-		PathColors.Add(InPath, FolderColor);
+		PathColors.Add(InPath, InFolderColor);
 	};
 
 	auto RemoveColorInternal = [](const FString& InPath)
diff --git a/Engine/Source/Editor/FoliageEdit/Private/FoliageEdMode.cpp b/Engine/Source/Editor/FoliageEdit/Private/FoliageEdMode.cpp
--- a/Engine/Source/Editor/FoliageEdit/Private/FoliageEdMode.cpp
+++ b/Engine/Source/Editor/FoliageEdit/Private/FoliageEdMode.cpp
@@ -1087,10 +1087,12 @@ void FEdModeFoliage::TransformSelectedInstances(UWorld* InWorld, const FVector&
 AInstancedFoliageActor* FEdModeFoliage::GetSelectionLocation(UWorld* InWorld, FVector& OutLocation) const
 {
 	// Prefer current level
-	AInstancedFoliageActor* IFA = AInstancedFoliageActor::GetInstancedFoliageActorForCurrentLevel(InWorld);
-	if (IFA && IFA->GetSelectionLocation(OutLocation))
 	{
-		return IFA;
+		AInstancedFoliageActor* IFA = AInstancedFoliageActor::GetInstancedFoliageActorForCurrentLevel(InWorld);
+		if (IFA && IFA->GetSelectionLocation(OutLocation))
+		{
+			return IFA;
+		}
 	}
 	
 	// Go through all sub-levels
diff --git a/Engine/Source/Editor/HardwareTargeting/Private/HardwareTargetingModule.cpp b/Engine/Source/Editor/HardwareTargeting/Private/HardwareTargetingModule.cpp
--- a/Engine/Source/Editor/HardwareTargeting/Private/HardwareTargetingModule.cpp
+++ b/Engine/Source/Editor/HardwareTargeting/Private/HardwareTargetingModule.cpp
@@ -40,7 +40,7 @@ struct FMetaSettingGatherer
 	{
 		if (bModified || bIncludeUnmodifiedProperties)
 		{
-			FTextBuilder& DescriptionBuffer = DescriptionBuffers.FindOrAdd(SettingsObject);
+			FTextBuilder& SettingsDescriptionBuffer = DescriptionBuffers.FindOrAdd(SettingsObject);
 
 			if (!bReadOnly)
 			{
@@ -59,7 +59,7 @@ struct FMetaSettingGatherer
 					LOCTEXT("MetaSettingDisplayStringModified", "{SettingName} is {SettingValue} <HardwareTargets.Strong>(modified)</>") :
 					LOCTEXT("MetaSettingDisplayStringUnmodified", "{SettingName} is {SettingValue}");
 
-				DescriptionBuffer.AppendLine(FText::Format(FormatString, Args));
+				SettingsDescriptionBuffer.AppendLine(FText::Format(FormatString, Args));
 			}
 		}
 	}
diff --git a/Engine/Source/Editor/LocalizationDashboard/Private/SLocalizationTargetEditorCultureRow.cpp b/Engine/Source/Editor/LocalizationDashboard/Private/SLocalizationTargetEditorCultureRow.cpp
--- a/Engine/Source/Editor/LocalizationDashboard/Private/SLocalizationTargetEditorCultureRow.cpp
+++ b/Engine/Source/Editor/LocalizationDashboard/Private/SLocalizationTargetEditorCultureRow.cpp
@@ -308,10 +308,12 @@ FReply SLocalizationTargetEditorCultureRow::Import()
 	if (Culture.IsValid() && TargetSettings && DesktopPlatform)
 	{
 		void* ParentWindowWindowHandle = NULL;
-		const TSharedPtr<SWindow> ParentWindow = FSlateApplication::Get().FindWidgetWindow(AsShared());
-		if (ParentWindow.IsValid() && ParentWindow->GetNativeWindow().IsValid())
 		{
-			ParentWindowWindowHandle = ParentWindow->GetNativeWindow()->GetOSWindowHandle();
+			const TSharedPtr<SWindow> ParentWindow = FSlateApplication::Get().FindWidgetWindow(AsShared());
+			if (ParentWindow.IsValid() && ParentWindow->GetNativeWindow().IsValid())
+			{
+				ParentWindowWindowHandle = ParentWindow->GetNativeWindow()->GetOSWindowHandle();
+			}
 		}
 
 		const FString POFileName = LocalizationConfigurationScript::GetDefaultPOFileName(*TargetSettings);
@@ -327,7 +329,7 @@ FReply SLocalizationTargetEditorCultureRow::Import()
 			FFormatNamedArguments FormatArguments;
 			FormatArguments.Add(TEXT("TargetName"), FText::FromString(TargetSettings->Name));
 			FormatArguments.Add(TEXT("CultureName"), FText::FromString(Culture->GetDisplayName()));
-			FText DialogTitle = FText::Format(LOCTEXT("ImportSpecificTranslationsForTargetDialogTitleFormat", "Import {CultureName} Translations for {TargetName} from Directory"), FormatArguments);
+			DialogTitle = FText::Format(LOCTEXT("ImportSpecificTranslationsForTargetDialogTitleFormat", "Import {CultureName} Translations for {TargetName} from Directory"), FormatArguments);
 		}
 
 		// Prompt the user for the directory
diff --git a/Engine/Source/Editor/NewLevelDialog/Private/NewLevelDialogModule.cpp b/Engine/Source/Editor/NewLevelDialog/Private/NewLevelDialogModule.cpp
--- a/Engine/Source/Editor/NewLevelDialog/Private/NewLevelDialogModule.cpp
+++ b/Engine/Source/Editor/NewLevelDialog/Private/NewLevelDialogModule.cpp
@@ -16,14 +16,14 @@ public:
 	SLATE_BEGIN_ARGS(STexture2DView) {}
 	SLATE_END_ARGS()
 
-	void Construct( const FArguments& InArgs, UTexture2D* InTexture )
+	void Construct( const FArguments& InArgs, UTexture2D* InTexture2D )
 	{
-		Size = FIntPoint(InTexture->GetSizeX(), InTexture->GetSizeY());
+		Size = FIntPoint(InTexture2D->GetSizeX(), InTexture2D->GetSizeY());
 
 		ENQUEUE_UNIQUE_RENDER_COMMAND_TWOPARAMETER(
 			UpdateSTexture2DView,
 			STexture2DView*,TextureView,this,
-			UTexture2D*,Texture,InTexture,
+			UTexture2D*,Texture,InTexture2D,
 			{
 				TextureView->ShaderResource = ((FTexture2DResource*)(Texture->Resource))->GetTexture2DRHI();
 			});
diff --git a/Engine/Source/Editor/PropertyEditor/Private/DetailItemNode.cpp b/Engine/Source/Editor/PropertyEditor/Private/DetailItemNode.cpp
--- a/Engine/Source/Editor/PropertyEditor/Private/DetailItemNode.cpp
+++ b/Engine/Source/Editor/PropertyEditor/Private/DetailItemNode.cpp
@@ -279,14 +279,14 @@ static bool PassesAllFilters( const FDetailLayoutCustomization& InCustomization,
 {	
 	struct Local
 	{
-		static bool StringPassesFilter(const FDetailFilter& InFilter, const FString& InString)
+		static bool StringPassesFilter(const FDetailFilter& InDetailFilter, const FString& InString)
 		{
 			// Make sure the passed string matches all filter strings
 			if( InString.Len() > 0 )
 			{
-				for (int32 TestNameIndex = 0; TestNameIndex < InFilter.FilterStrings.Num(); ++TestNameIndex)
+				for (int32 TestNameIndex = 0; TestNameIndex < InDetailFilter.FilterStrings.Num(); ++TestNameIndex)
 				{
-					const FString& TestName = InFilter.FilterStrings[TestNameIndex];
+					const FString& TestName = InDetailFilter.FilterStrings[TestNameIndex];
 					if ( !InString.Contains(TestName) ) 
 					{
 						return false;
diff --git a/Engine/Source/Editor/PropertyEditor/Private/PropertyCustomizationHelpers.cpp b/Engine/Source/Editor/PropertyEditor/Private/PropertyCustomizationHelpers.cpp
--- a/Engine/Source/Editor/PropertyEditor/Private/PropertyCustomizationHelpers.cpp
+++ b/Engine/Source/Editor/PropertyEditor/Private/PropertyCustomizationHelpers.cpp
@@ -1056,21 +1056,21 @@ TSharedRef<SWidget> PropertyCustomizationHelpers::MakeTextLocalizationButton(TSh
 {
 	struct Local
 	{
-		static void DirtyPropertyPackage(TSharedRef<IPropertyHandle> InPropertyHandle)
+		static void DirtyPropertyPackage(TSharedRef<IPropertyHandle> InPropHandle)
 		{
 			// dirty packages as they will need to be re-serialized with the correct key
 			TArray<UObject*> Objects;
-			InPropertyHandle->GetOuterObjects(Objects);
+			InPropHandle->GetOuterObjects(Objects);
 			if (Objects.Num() > 0 && Objects[0]->GetOutermost() != nullptr)
 			{
 				Objects[0]->GetOutermost()->MarkPackageDirty();
 			}
 		}
 
-		static FReply HandleGenerateKeyClicked(TSharedRef<IPropertyHandle> InPropertyHandle)
+		static FReply HandleGenerateKeyClicked(TSharedRef<IPropertyHandle> InPropHandle)
 		{
 			FText DisplayText;
-			if (InPropertyHandle->GetValueAsDisplayText(DisplayText) == FPropertyAccess::Success)
+			if (InPropHandle->GetValueAsDisplayText(DisplayText) == FPropertyAccess::Success)
 			{
 				TSharedRef< FString, ESPMode::ThreadSafe > DisplayString = FTextInspector::GetSharedDisplayString(DisplayText);
 				TSharedPtr< FString, ESPMode::ThreadSafe > Namespace;
@@ -1080,17 +1080,17 @@ TSharedRef<SWidget> PropertyCustomizationHelpers::MakeTextLocalizationButton(TSh
 				if (Key.IsValid())
 				{
 					*Key.Get() = FGuid::NewGuid().ToString();
-					DirtyPropertyPackage(InPropertyHandle);
+					DirtyPropertyPackage(InPropHandle);
 				}
 			}
 
 			return FReply::Unhandled();
 		}
 
-		static void HandleKeyTextCommitted(const FText& InText, ETextCommit::Type InCommitType, TSharedRef<IPropertyHandle> InPropertyHandle)
+		static void HandleKeyTextCommitted(const FText& InText, ETextCommit::Type InCommitType, TSharedRef<IPropertyHandle> InPropHandle)
 		{
 			FText DisplayText;
-			if (InPropertyHandle->GetValueAsDisplayText(DisplayText) == FPropertyAccess::Success)
+			if (InPropHandle->GetValueAsDisplayText(DisplayText) == FPropertyAccess::Success)
 			{
 				TSharedRef< FString, ESPMode::ThreadSafe > DisplayString = FTextInspector::GetSharedDisplayString(DisplayText);
 				TSharedPtr< FString, ESPMode::ThreadSafe > Namespace;
@@ -1100,15 +1100,15 @@ TSharedRef<SWidget> PropertyCustomizationHelpers::MakeTextLocalizationButton(TSh
 				if (Key.IsValid())
 				{
 					*Key.Get() = InText.ToString();
-					DirtyPropertyPackage(InPropertyHandle);
+					DirtyPropertyPackage(InPropHandle);
 				}
 			}
 		}
 
-		static void HandleNamespaceTextCommitted(const FText& InText, ETextCommit::Type InCommitType, TSharedRef<IPropertyHandle> InPropertyHandle)
+		static void HandleNamespaceTextCommitted(const FText& InText, ETextCommit::Type InCommitType, TSharedRef<IPropertyHandle> InPropHandle)
 		{
 			FText DisplayText;
-			if (InPropertyHandle->GetValueAsDisplayText(DisplayText) == FPropertyAccess::Success)
+			if (InPropHandle->GetValueAsDisplayText(DisplayText) == FPropertyAccess::Success)
 			{
 				TSharedRef< FString, ESPMode::ThreadSafe > DisplayString = FTextInspector::GetSharedDisplayString(DisplayText);
 				TSharedPtr< FString, ESPMode::ThreadSafe > Namespace;
@@ -1118,17 +1118,17 @@ TSharedRef<SWidget> PropertyCustomizationHelpers::MakeTextLocalizationButton(TSh
 				if (Namespace.IsValid())
 				{
 					*Namespace.Get() = InText.ToString();
-					DirtyPropertyPackage(InPropertyHandle);
+					DirtyPropertyPackage(InPropHandle);
 				}
 			}
 		}
 
-		static TSharedRef<SWidget> GetLocalizationMenuContent(TSharedRef<IPropertyHandle> InPropertyHandle)
+		static TSharedRef<SWidget> GetLocalizationMenuContent(TSharedRef<IPropertyHandle> InPropHandle)
 		{
 			FMenuBuilder MenuContentBuilder(true, NULL);
 			{
 				FText DisplayText;
-				if (InPropertyHandle->GetValueAsDisplayText(DisplayText) == FPropertyAccess::Success)
+				if (InPropHandle->GetValueAsDisplayText(DisplayText) == FPropertyAccess::Success)
 				{
 					TSharedRef< FString, ESPMode::ThreadSafe > DisplayString = FTextInspector::GetSharedDisplayString(DisplayText);
 					TSharedPtr< FString, ESPMode::ThreadSafe > Namespace;
@@ -1168,7 +1168,7 @@ TSharedRef<SWidget> PropertyCustomizationHelpers::MakeTextLocalizationButton(TSh
 										[
 											SNew(SEditableTextBox)
 											.Text_Lambda(KeyLambda)
-											.OnTextCommitted(FOnTextCommitted::CreateStatic(&Local::HandleKeyTextCommitted, InPropertyHandle))
+											.OnTextCommitted(FOnTextCommitted::CreateStatic(&Local::HandleKeyTextCommitted, InPropHandle))
 										]
 										+ SHorizontalBox::Slot()
 										.AutoWidth()
@@ -1178,7 +1178,7 @@ TSharedRef<SWidget> PropertyCustomizationHelpers::MakeTextLocalizationButton(TSh
 											SNew(SButton)
 											.ToolTipText(LOCTEXT("RefreshKeyTooltip", "Generate a new random key"))
 											.ButtonStyle(FEditorStyle::Get(), "HoverHintOnly")
-											.OnClicked(FOnClicked::CreateStatic(&Local::HandleGenerateKeyClicked, InPropertyHandle))
+											.OnClicked(FOnClicked::CreateStatic(&Local::HandleGenerateKeyClicked, InPropHandle))
 											.Content()
 											[
 												SNew(SImage)
@@ -1207,7 +1207,7 @@ TSharedRef<SWidget> PropertyCustomizationHelpers::MakeTextLocalizationButton(TSh
 									[
 										SNew(SEditableTextBox)
 										.Text_Lambda(NamespaceLambda)
-										.OnTextCommitted(FOnTextCommitted::CreateStatic(&Local::HandleNamespaceTextCommitted, InPropertyHandle))
+										.OnTextCommitted(FOnTextCommitted::CreateStatic(&Local::HandleNamespaceTextCommitted, InPropHandle))
 									];
 							}
 
diff --git a/Engine/Source/Editor/UnrealEd/Private/Commandlets/GenerateBlueprintAPICommandlet.cpp b/Engine/Source/Editor/UnrealEd/Private/Commandlets/GenerateBlueprintAPICommandlet.cpp
--- a/Engine/Source/Editor/UnrealEd/Private/Commandlets/GenerateBlueprintAPICommandlet.cpp
+++ b/Engine/Source/Editor/UnrealEd/Private/Commandlets/GenerateBlueprintAPICommandlet.cpp
@@ -981,14 +981,14 @@ int32 UGenerateBlueprintAPICommandlet::Main(FString const& Params)
 	// closing out the writer (and diffing the resultant file if the user deigns us to do so)
 	auto CloseFileStream = [](FArchive** FileOutPtr)
 	{
-		FArchive*& FileOut = (*FileOutPtr);
-		if (FileOut != nullptr)
+		FArchive*& FileOutAr = (*FileOutPtr);
+		if (FileOutAr != nullptr)
 		{
-			FileOut->Serialize(TCHAR_TO_ANSI(TEXT("\n}")), 2);
-			FileOut->Close();
+			FileOutAr->Serialize(TCHAR_TO_ANSI(TEXT("\n}")), 2);
+			FileOutAr->Close();
 
-			delete FileOut;
-			FileOut = nullptr;
+			delete FileOutAr;
+			FileOutAr = nullptr;
 		}
 	};
 
diff --git a/Engine/Source/Runtime/Engine/Private/AI/Navigation/RecastNavMeshGenerator.cpp b/Engine/Source/Runtime/Engine/Private/AI/Navigation/RecastNavMeshGenerator.cpp
--- a/Engine/Source/Runtime/Engine/Private/AI/Navigation/RecastNavMeshGenerator.cpp
+++ b/Engine/Source/Runtime/Engine/Private/AI/Navigation/RecastNavMeshGenerator.cpp
@@ -1784,16 +1784,16 @@ void FRecastTileGenerator::GatherGeometry(const FRecastNavMeshGenerator& ParentG
 	QUICK_SCOPE_CYCLE_COUNTER(STAT_RecastNavMeshGenerator_GatherGeometry);
 
 	UNavigationSystem*	NavSys = UNavigationSystem::GetCurrent(ParentGenerator.GetWorld());
-	FNavigationOctree*	NavOctree = NavSys ? NavSys->GetMutableNavOctree() : nullptr;
-	check(NavOctree);
-	const FNavDataConfig&		NavDataConfig = ParentGenerator.GetOwner()->GetConfig();
+	FNavigationOctree*	NavigationOctree = NavSys ? NavSys->GetMutableNavOctree() : nullptr;
+	check(NavigationOctree);
+	const FNavDataConfig& OwnerNavDataConfig = ParentGenerator.GetOwner()->GetConfig();
 
-	for (FNavigationOctree::TConstElementBoxIterator<FNavigationOctree::DefaultStackAllocator> It(*NavOctree, ParentGenerator.GrowBoundingBox(TileBB, /*bIncludeAgentHeight*/ false));
+	for (FNavigationOctree::TConstElementBoxIterator<FNavigationOctree::DefaultStackAllocator> It(*NavigationOctree, ParentGenerator.GrowBoundingBox(TileBB, /*bIncludeAgentHeight*/ false));
 		It.HasPendingElements();
 		It.Advance())
 	{
 		const FNavigationOctreeElement& Element = It.GetCurrentElement();
-		const bool bShouldUse = Element.ShouldUseGeometry(NavDataConfig);
+		const bool bShouldUse = Element.ShouldUseGeometry(OwnerNavDataConfig);
 		if (bShouldUse)
 		{
 			bool bDumpGeometryData = false;
@@ -1803,7 +1803,7 @@ void FRecastTileGenerator::GatherGeometry(const FRecastNavMeshGenerator& ParentG
 				if (bSupportsSlices == false || Element.Data->IsPendingLazyModifiersGathering() == true)
 				{
 					QUICK_SCOPE_CYCLE_COUNTER(STAT_RecastNavMeshGenerator_LazyGeometryExport);
-					NavOctree->DemandLazyDataGathering(Element);
+					NavigationOctree->DemandLazyDataGathering(Element);
 				}
 				
 				if (bSupportsSlices == true)
@@ -1816,7 +1816,7 @@ void FRecastTileGenerator::GatherGeometry(const FRecastNavMeshGenerator& ParentG
 					if (NavRelevant)
 					{
 						NavRelevant->PrepareGeometryExportSync();
-						NavRelevant->GatherGeometrySlice(GeomExport, TileBB.ExpandBy(NavDataConfig.AgentRadius * 2));
+						NavRelevant->GatherGeometrySlice(GeomExport, TileBB.ExpandBy(OwnerNavDataConfig.AgentRadius * 2));
 
 						RecastGeometryExport::CovertCoordDataToRecast(GeomExport.VertexBuffer);
 						RecastGeometryExport::StoreCollisionCache(GeomExport);
@@ -1868,7 +1868,7 @@ void FRecastTileGenerator::GatherGeometry(const FRecastNavMeshGenerator& ParentG
 				}
 			}
 						
-			const FCompositeNavModifier ModifierInstance = Element.GetModifierForAgent(&NavDataConfig);
+			const FCompositeNavModifier ModifierInstance = Element.GetModifierForAgent(&OwnerNavDataConfig);
 			if (ModifierInstance.IsEmpty() == false)
 			{
 				AppendModifier(ModifierInstance, Element.Data->NavDataPerInstanceTransformDelegate);
diff --git a/Engine/Source/Runtime/Engine/Private/Animation/AnimInstance.cpp b/Engine/Source/Runtime/Engine/Private/Animation/AnimInstance.cpp
--- a/Engine/Source/Runtime/Engine/Private/Animation/AnimInstance.cpp
+++ b/Engine/Source/Runtime/Engine/Private/Animation/AnimInstance.cpp
@@ -2185,14 +2185,14 @@ FName UAnimInstance::Montage_GetCurrentSection(UAnimMontage* Montage)
 	return NAME_None;
 }
 
-void UAnimInstance::Montage_SetEndDelegate(FOnMontageEnded & OnMontageEnded, UAnimMontage* Montage)
+void UAnimInstance::Montage_SetEndDelegate(FOnMontageEnded & InOnMontageEnded, UAnimMontage* Montage)
 {
 	if (Montage)
 	{
 		FAnimMontageInstance * MontageInstance = GetActiveInstanceForMontage(*Montage);
 		if (MontageInstance)
 		{
-			MontageInstance->OnMontageEnded = OnMontageEnded;
+			MontageInstance->OnMontageEnded = InOnMontageEnded;
 		}
 	}
 	else
@@ -2203,20 +2203,20 @@ void UAnimInstance::Montage_SetEndDelegate(FOnMontageEnded & OnMontageEnded, UAn
 			FAnimMontageInstance * MontageInstance = MontageInstances[InstanceIndex];
 			if (MontageInstance && MontageInstance->IsActive())
 			{
-				MontageInstance->OnMontageEnded = OnMontageEnded;
+				MontageInstance->OnMontageEnded = InOnMontageEnded;
 			}
 		}
 	}
 }
 
-void UAnimInstance::Montage_SetBlendingOutDelegate(FOnMontageBlendingOutStarted & OnMontageBlendingOut, UAnimMontage* Montage)
+void UAnimInstance::Montage_SetBlendingOutDelegate(FOnMontageBlendingOutStarted & InOnMontageBlendingOut, UAnimMontage* Montage)
 {
 	if (Montage)
 	{
 		FAnimMontageInstance * MontageInstance = GetActiveInstanceForMontage(*Montage);
 		if (MontageInstance)
 		{
-			MontageInstance->OnMontageBlendingOutStarted = OnMontageBlendingOut;
+			MontageInstance->OnMontageBlendingOutStarted = InOnMontageBlendingOut;
 		}
 	}
 	else
@@ -2227,7 +2227,7 @@ void UAnimInstance::Montage_SetBlendingOutDelegate(FOnMontageBlendingOutStarted
 			FAnimMontageInstance * MontageInstance = MontageInstances[InstanceIndex];
 			if (MontageInstance && MontageInstance->IsActive())
 			{
-				MontageInstance->OnMontageBlendingOutStarted = OnMontageBlendingOut;
+				MontageInstance->OnMontageBlendingOutStarted = InOnMontageBlendingOut;
 			}
 		}
 	}
diff --git a/Engine/Source/Runtime/Engine/Private/Animation/AnimNodeRotationMultiplier.cpp b/Engine/Source/Runtime/Engine/Private/Animation/AnimNodeRotationMultiplier.cpp
--- a/Engine/Source/Runtime/Engine/Private/Animation/AnimNodeRotationMultiplier.cpp
+++ b/Engine/Source/Runtime/Engine/Private/Animation/AnimNodeRotationMultiplier.cpp
@@ -86,7 +86,7 @@ FQuat FAnimNode_RotationMultiplier::ExtractAngle(const TArray<FTransform> & RefP
 	return DeltaQuat;
 }
 
-void FAnimNode_RotationMultiplier::MultiplyQuatBasedOnSourceIndex(const TArray<FTransform> & RefPoseTransforms, FA2CSPose& MeshBases, const EBoneAxis Axis, int32 SourceBoneIndex, float Multiplier, const FQuat& ReferenceQuat, FQuat& OutQuat)
+void FAnimNode_RotationMultiplier::MultiplyQuatBasedOnSourceIndex(const TArray<FTransform> & RefPoseTransforms, FA2CSPose& MeshBases, const EBoneAxis Axis, int32 SourceBoneIndex, float InMultiplier, const FQuat& ReferenceQuat, FQuat& OutQuat)
 {
 	// Find delta angle for source bone.
 	FQuat DeltaQuat = ExtractAngle(RefPoseTransforms, MeshBases, Axis, SourceBoneIndex);
@@ -109,7 +109,7 @@ void FAnimNode_RotationMultiplier::MultiplyQuatBasedOnSourceIndex(const TArray<F
 	RotationAngle = FMath::UnwindRadians(RotationAngle);
 
 	// New bone rotation
-	OutQuat = ReferenceQuat * FQuat(RotationAxis, RotationAngle* Multiplier);
+	OutQuat = ReferenceQuat * FQuat(RotationAxis, RotationAngle* InMultiplier);
 	// Normalize resulting quaternion.
 	OutQuat.Normalize();
 
diff --git a/Engine/Source/Runtime/Engine/Private/Animation/AnimNode_BlendListBase.cpp b/Engine/Source/Runtime/Engine/Private/Animation/AnimNode_BlendListBase.cpp
--- a/Engine/Source/Runtime/Engine/Private/Animation/AnimNode_BlendListBase.cpp
+++ b/Engine/Source/Runtime/Engine/Private/Animation/AnimNode_BlendListBase.cpp
@@ -76,12 +76,12 @@ void FAnimNode_BlendListBase::Update(const FAnimationUpdateContext& Context)
 		float SumWeight = 0.0f;
 		for (int32 i = 0; i < RemainingBlendTimes.Num(); ++i)
 		{
-			float& BlendTime = RemainingBlendTimes[i];
+			float& RemainingBlendTime = RemainingBlendTimes[i];
 			float& BlendWeight = BlendWeights[i];
 
 			const float DesiredWeight = (i == ChildIndex) ? 1.0f : 0.0f;
 
-			FAnimationRuntime::TickBlendWeight(Context.GetDeltaTime(), DesiredWeight, BlendWeight, BlendTime);
+			FAnimationRuntime::TickBlendWeight(Context.GetDeltaTime(), DesiredWeight, BlendWeight, RemainingBlendTime);
 
 			SumWeight += BlendWeight;
 		}
diff --git a/Engine/Source/Runtime/Engine/Private/Animation/AnimNode_Fabrik.cpp b/Engine/Source/Runtime/Engine/Private/Animation/AnimNode_Fabrik.cpp
--- a/Engine/Source/Runtime/Engine/Private/Animation/AnimNode_Fabrik.cpp
+++ b/Engine/Source/Runtime/Engine/Private/Animation/AnimNode_Fabrik.cpp
@@ -45,13 +45,15 @@ void FAnimNode_Fabrik::EvaluateBoneTransforms(USkeletalMeshComponent* SkelComp,
 
 	// Gather all bone indices between root and tip.
 	TArray<int32> BoneIndices;
-	int32 BoneIndex = TipBone.BoneIndex;
-	do 
 	{
+		int32 BoneIndex = TipBone.BoneIndex;
+		do 
+		{
+			BoneIndices.Insert(BoneIndex, 0);
+			BoneIndex = RequiredBones.GetParentBoneIndex(BoneIndex);
+		} while (BoneIndex != RootIndex);
 		BoneIndices.Insert(BoneIndex, 0);
-		BoneIndex = RequiredBones.GetParentBoneIndex(BoneIndex);
-	} while (BoneIndex != RootIndex);
-	BoneIndices.Insert(BoneIndex, 0);
+	}
 
 	// Maximum length of skeleton segment at full extension
 	float MaximumReach = 0;
diff --git a/Engine/Source/Runtime/Engine/Private/GameViewportClient.cpp b/Engine/Source/Runtime/Engine/Private/GameViewportClient.cpp
--- a/Engine/Source/Runtime/Engine/Private/GameViewportClient.cpp
+++ b/Engine/Source/Runtime/Engine/Private/GameViewportClient.cpp
@@ -576,7 +576,7 @@ void UGameViewportClient::AddCursor(EMouseCursor::Type Cursor, const FStringClas
 	}
 }
 
-TOptional<TSharedRef<SWidget>> UGameViewportClient::MapCursor(FViewport* Viewport, const FCursorReply& CursorReply)
+TOptional<TSharedRef<SWidget>> UGameViewportClient::MapCursor(FViewport* InViewport, const FCursorReply& CursorReply)
 {
 	const TSharedRef<SWidget>* CursorWidgetPtr = CursorWidgets.Find(CursorReply.GetCursorType());
 	if (CursorWidgetPtr != nullptr)
@@ -1317,10 +1317,10 @@ void UGameViewportClient::LostFocus(FViewport* InViewport)
 {
 	// We need to reset some key inputs, since keyup events will sometimes not be processed (such as going into immersive/maximized mode).  
 	// Resetting them will prevent them from "sticking"
-	UWorld* const World = GetWorld();
-	if (World)
+	UWorld* const ViewportWorld = GetWorld();
+	if (ViewportWorld)
 	{
-		for (FConstPlayerControllerIterator Iterator = World->GetPlayerControllerIterator(); Iterator; ++Iterator)
+		for (FConstPlayerControllerIterator Iterator = ViewportWorld->GetPlayerControllerIterator(); Iterator; ++Iterator)
 		{
 			APlayerController* const PlayerController = *Iterator;
 			if (PlayerController)
@@ -1448,27 +1448,27 @@ ULocalPlayer* UGameViewportClient::SetupInitialLocalPlayer(FString& OutError)
 	GEngine->OnTravelFailure().AddUObject(this, &UGameViewportClient::PeekTravelFailureMessages);
 	GEngine->OnNetworkFailure().AddUObject(this, &UGameViewportClient::PeekNetworkFailureMessages);
 
-	UGameInstance * GameInstance = GEngine->GetWorldContextFromGameViewportChecked(this).OwningGameInstance;
+	UGameInstance * ViewportGameInstance = GEngine->GetWorldContextFromGameViewportChecked(this).OwningGameInstance;
 
-	if ( !ensure( GameInstance != NULL ) )
+	if ( !ensure( ViewportGameInstance != NULL ) )
 	{
 		return NULL;
 	}
 
 	// Create the initial player - this is necessary or we can't render anything in-game.
-	return GameInstance->CreateInitialPlayer(OutError);
+	return ViewportGameInstance->CreateInitialPlayer(OutError);
 }
 
 ULocalPlayer* UGameViewportClient::CreatePlayer(int32 ControllerId, FString& OutError, bool bSpawnActor)
 {
-	UGameInstance * GameInstance = GEngine->GetWorldContextFromGameViewportChecked(this).OwningGameInstance;
-	return (GameInstance != NULL) ? GameInstance->CreateLocalPlayer(ControllerId, OutError, bSpawnActor) : NULL;
+	UGameInstance * ViewportGameInstance = GEngine->GetWorldContextFromGameViewportChecked(this).OwningGameInstance;
+	return (ViewportGameInstance != NULL) ? ViewportGameInstance->CreateLocalPlayer(ControllerId, OutError, bSpawnActor) : NULL;
 }
 
 bool UGameViewportClient::RemovePlayer(class ULocalPlayer* ExPlayer)
 {
-	UGameInstance * GameInstance = GEngine->GetWorldContextFromGameViewportChecked(this).OwningGameInstance;
-	return (GameInstance != NULL) ? GameInstance->RemoveLocalPlayer(ExPlayer) : false;
+	UGameInstance * ViewportGameInstance = GEngine->GetWorldContextFromGameViewportChecked(this).OwningGameInstance;
+	return (ViewportGameInstance != NULL) ? ViewportGameInstance->RemoveLocalPlayer(ExPlayer) : false;
 }
 
 void UGameViewportClient::UpdateActiveSplitscreenType()
@@ -1894,11 +1894,11 @@ void UGameViewportClient::VerifyPathRenderingComponents()
 {
 	const bool bShowPaths = !!EngineShowFlags.Navigation;
 
-	UWorld* const World = GetWorld();
+	UWorld* const ViewportWorld = GetWorld();
 
 	// make sure nav mesh has a rendering component
-	ANavigationData* const NavData = (World && World->GetNavigationSystem() != nullptr)
-		? World->GetNavigationSystem()->GetMainNavData(FNavigationSystem::DontCreate)
+	ANavigationData* const NavData = (ViewportWorld && ViewportWorld->GetNavigationSystem() != nullptr)
+		? ViewportWorld->GetNavigationSystem()->GetMainNavData(FNavigationSystem::DontCreate)
 		: NULL;
 
 	if(NavData && NavData->RenderingComp == NULL)
@@ -2967,10 +2967,10 @@ bool UGameViewportClient::RequestBugScreenShot(const TCHAR* Cmd, bool bDisplayHU
 			FCString::Sprintf( File, TEXT("%s"), *(OutputDir + SSFilename) );
 			if( IFileManager::Get().FileSize(File) == INDEX_NONE )
 			{
-				UWorld* const World = GetWorld();
-				if ( bDisplayHUDInfo && (World != nullptr) )
+				UWorld* const ViewportWorld = GetWorld();
+				if ( bDisplayHUDInfo && (ViewportWorld != nullptr) )
 				{
-					for( FConstPlayerControllerIterator Iterator = World->GetPlayerControllerIterator(); Iterator; ++Iterator )
+					for( FConstPlayerControllerIterator Iterator = ViewportWorld->GetPlayerControllerIterator(); Iterator; ++Iterator )
 					{
 						APlayerController* PlayerController = *Iterator;
 						if (PlayerController && PlayerController->GetHUD() )
diff --git a/Engine/Source/Runtime/Engine/Private/HierarchicalInstancedStaticMesh.cpp b/Engine/Source/Runtime/Engine/Private/HierarchicalInstancedStaticMesh.cpp
--- a/Engine/Source/Runtime/Engine/Private/HierarchicalInstancedStaticMesh.cpp
+++ b/Engine/Source/Runtime/Engine/Private/HierarchicalInstancedStaticMesh.cpp
@@ -149,11 +149,11 @@ class FClusterBuilder
 	};
 	TArray<FSortPair> SortPairs;
 
-	void Split(int32 Num)
+	void Split(int32 InNum)
 	{
-		check(Num);
+		check(InNum);
 		Clusters.Reset();
-		Split(0, Num - 1);
+		Split(0, InNum - 1);
 		Clusters.Sort();
 		check(Clusters.Num());
 		int32 At = 0;
@@ -162,7 +162,7 @@ class FClusterBuilder
 			check(At == Cluster.Start);
 			At += Cluster.Num;
 		}
-		check(At == Num);
+		check(At == InNum);
 	}
 
 	void Split(int32 Start, int32 End)
diff --git a/Engine/Source/Runtime/Engine/Private/KismetSystemLibrary.cpp b/Engine/Source/Runtime/Engine/Private/KismetSystemLibrary.cpp
--- a/Engine/Source/Runtime/Engine/Private/KismetSystemLibrary.cpp
+++ b/Engine/Source/Runtime/Engine/Private/KismetSystemLibrary.cpp
@@ -1186,7 +1186,7 @@ bool UKismetSystemLibrary::LineTraceSingle_DEPRECATED(UObject* WorldContextObjec
 			while (CurrentObject)
 			{
 				CurrentObject = CurrentObject->GetOuter();
-				AActor* IgnoreActor = Cast<AActor>(CurrentObject);
+				IgnoreActor = Cast<AActor>(CurrentObject);
 				if (IgnoreActor)
 				{
 					Params.AddIgnoredActor(IgnoreActor);
@@ -1249,7 +1249,7 @@ bool UKismetSystemLibrary::LineTraceMulti_DEPRECATED(UObject* WorldContextObject
 			while (CurrentObject)
 			{
 				CurrentObject = CurrentObject->GetOuter();
-				AActor* IgnoreActor = Cast<AActor>(CurrentObject);
+				IgnoreActor = Cast<AActor>(CurrentObject);
 				if (IgnoreActor)
 				{
 					Params.AddIgnoredActor(IgnoreActor);
@@ -1343,7 +1343,8 @@ bool UKismetSystemLibrary::BoxTraceSingle(UObject* WorldContextObject, const FVe
 	Params.AddIgnoredActors(ActorsToIgnore);
 	if (bIgnoreSelf)
 	{
-		if (AActor* IgnoreActor = Cast<AActor>(WorldContextObject))
+		AActor* IgnoreActor = Cast<AActor>(WorldContextObject);
+		if (IgnoreActor)
 		{
 			Params.AddIgnoredActor(IgnoreActor);
 		}
@@ -1354,7 +1355,8 @@ bool UKismetSystemLibrary::BoxTraceSingle(UObject* WorldContextObject, const FVe
 			while (CurrentObject)
 			{
 				CurrentObject = CurrentObject->GetOuter();
-				if (AActor* IgnoreActor = Cast<AActor>(CurrentObject))
+				IgnoreActor = Cast<AActor>(CurrentObject);
+				if (IgnoreActor)
 				{
 					Params.AddIgnoredActor(IgnoreActor);
 					break;
@@ -1398,7 +1400,8 @@ bool UKismetSystemLibrary::BoxTraceMulti(UObject* WorldContextObject, const FVec
 	Params.AddIgnoredActors(ActorsToIgnore);
 	if (bIgnoreSelf)
 	{
-		if (AActor* IgnoreActor = Cast<AActor>(WorldContextObject))
+		AActor* IgnoreActor = Cast<AActor>(WorldContextObject);
+		if (IgnoreActor)
 		{
 			Params.AddIgnoredActor(IgnoreActor);
 		}
@@ -1409,7 +1412,8 @@ bool UKismetSystemLibrary::BoxTraceMulti(UObject* WorldContextObject, const FVec
 			while (CurrentObject)
 			{
 				CurrentObject = CurrentObject->GetOuter();
-				if (AActor* IgnoreActor = Cast<AActor>(CurrentObject))
+				IgnoreActor = Cast<AActor>(CurrentObject);
+				if (IgnoreActor)
 				{
 					Params.AddIgnoredActor(IgnoreActor);
 					break;
@@ -1477,7 +1481,7 @@ bool UKismetSystemLibrary::SphereTraceSingle_DEPRECATED(UObject* WorldContextObj
 			while (CurrentObject)
 			{
 				CurrentObject = CurrentObject->GetOuter();
-				AActor* IgnoreActor = Cast<AActor>(CurrentObject);
+				IgnoreActor = Cast<AActor>(CurrentObject);
 				if (IgnoreActor)
 				{
 					Params.AddIgnoredActor(IgnoreActor);
@@ -1539,7 +1543,7 @@ bool UKismetSystemLibrary::SphereTraceMulti_DEPRECATED(UObject* WorldContextObje
 			while (CurrentObject)
 			{
 				CurrentObject = CurrentObject->GetOuter();
-				AActor* IgnoreActor = Cast<AActor>(CurrentObject);
+				IgnoreActor = Cast<AActor>(CurrentObject);
 				if (IgnoreActor)
 				{
 					Params.AddIgnoredActor(IgnoreActor);
@@ -1608,7 +1612,7 @@ bool UKismetSystemLibrary::CapsuleTraceSingle_DEPRECATED(UObject* WorldContextOb
 			while (CurrentObject)
 			{
 				CurrentObject = CurrentObject->GetOuter();
-				AActor* IgnoreActor = Cast<AActor>(CurrentObject);
+				IgnoreActor = Cast<AActor>(CurrentObject);
 				if (IgnoreActor)
 				{
 					Params.AddIgnoredActor(IgnoreActor);
@@ -1676,7 +1680,7 @@ bool UKismetSystemLibrary::CapsuleTraceMulti_DEPRECATED(UObject* WorldContextObj
 			while (CurrentObject)
 			{
 				CurrentObject = CurrentObject->GetOuter();
-				AActor* IgnoreActor = Cast<AActor>(CurrentObject);
+				IgnoreActor = Cast<AActor>(CurrentObject);
 				if (IgnoreActor)
 				{
 					Params.AddIgnoredActor(IgnoreActor);
@@ -1761,7 +1765,7 @@ bool UKismetSystemLibrary::LineTraceSingleByObject_DEPRECATED(UObject* WorldCont
 			while (CurrentObject)
 			{
 				CurrentObject = CurrentObject->GetOuter();
-				AActor* IgnoreActor = Cast<AActor>(CurrentObject);
+				IgnoreActor = Cast<AActor>(CurrentObject);
 				if (IgnoreActor)
 				{
 					Params.AddIgnoredActor(IgnoreActor);
@@ -1852,7 +1856,7 @@ bool UKismetSystemLibrary::LineTraceMultiByObject_DEPRECATED(UObject* WorldConte
 			while (CurrentObject)
 			{
 				CurrentObject = CurrentObject->GetOuter();
-				AActor* IgnoreActor = Cast<AActor>(CurrentObject);
+				IgnoreActor = Cast<AActor>(CurrentObject);
 				if (IgnoreActor)
 				{
 					Params.AddIgnoredActor(IgnoreActor);
@@ -1954,7 +1958,7 @@ bool UKismetSystemLibrary::SphereTraceSingleByObject_DEPRECATED(UObject* WorldCo
 			while (CurrentObject)
 			{
 				CurrentObject = CurrentObject->GetOuter();
-				AActor* IgnoreActor = Cast<AActor>(CurrentObject);
+				IgnoreActor = Cast<AActor>(CurrentObject);
 				if (IgnoreActor)
 				{
 					Params.AddIgnoredActor(IgnoreActor);
@@ -2044,7 +2048,7 @@ bool UKismetSystemLibrary::SphereTraceMultiByObject_DEPRECATED(UObject* WorldCon
 			while (CurrentObject)
 			{
 				CurrentObject = CurrentObject->GetOuter();
-				AActor* IgnoreActor = Cast<AActor>(CurrentObject);
+				IgnoreActor = Cast<AActor>(CurrentObject);
 				if (IgnoreActor)
 				{
 					Params.AddIgnoredActor(IgnoreActor);
@@ -2124,7 +2128,8 @@ bool UKismetSystemLibrary::BoxTraceSingleForObjects(UObject* WorldContextObject,
 	Params.AddIgnoredActors(ActorsToIgnore);
 	if (bIgnoreSelf)
 	{
-		if (AActor* IgnoreActor = Cast<AActor>(WorldContextObject))
+		AActor* IgnoreActor = Cast<AActor>(WorldContextObject);
+		if (IgnoreActor)
 		{
 			Params.AddIgnoredActor(IgnoreActor);
 		}
@@ -2135,7 +2140,8 @@ bool UKismetSystemLibrary::BoxTraceSingleForObjects(UObject* WorldContextObject,
 			while (CurrentObject)
 			{
 				CurrentObject = CurrentObject->GetOuter();
-				if (AActor* IgnoreActor = Cast<AActor>(CurrentObject))
+				IgnoreActor = Cast<AActor>(CurrentObject);
+				if (IgnoreActor)
 				{
 					Params.AddIgnoredActor(IgnoreActor);
 					break;
@@ -2207,7 +2213,8 @@ bool UKismetSystemLibrary::BoxTraceMultiForObjects(UObject* WorldContextObject,
 	Params.AddIgnoredActors(ActorsToIgnore);
 	if (bIgnoreSelf)
 	{
-		if (AActor* IgnoreActor = Cast<AActor>(WorldContextObject))
+		AActor* IgnoreActor = Cast<AActor>(WorldContextObject);
+		if (IgnoreActor)
 		{
 			Params.AddIgnoredActor(IgnoreActor);
 		}
@@ -2218,7 +2225,8 @@ bool UKismetSystemLibrary::BoxTraceMultiForObjects(UObject* WorldContextObject,
 			while (CurrentObject)
 			{
 				CurrentObject = CurrentObject->GetOuter();
-				if (AActor* IgnoreActor = Cast<AActor>(CurrentObject))
+				IgnoreActor = Cast<AActor>(CurrentObject);
+				if (IgnoreActor)
 				{
 					Params.AddIgnoredActor(IgnoreActor);
 					break;
@@ -2314,7 +2322,7 @@ bool UKismetSystemLibrary::CapsuleTraceSingleByObject_DEPRECATED(UObject* WorldC
 			while (CurrentObject)
 			{
 				CurrentObject = CurrentObject->GetOuter();
-				AActor* IgnoreActor = Cast<AActor>(CurrentObject);
+				IgnoreActor = Cast<AActor>(CurrentObject);
 				if (IgnoreActor)
 				{
 					Params.AddIgnoredActor(IgnoreActor);
@@ -2409,7 +2417,7 @@ bool UKismetSystemLibrary::CapsuleTraceMultiByObject_DEPRECATED(UObject* WorldCo
 			while (CurrentObject)
 			{
 				CurrentObject = CurrentObject->GetOuter();
-				AActor* IgnoreActor = Cast<AActor>(CurrentObject);
+				IgnoreActor = Cast<AActor>(CurrentObject);
 				if (IgnoreActor)
 				{
 					Params.AddIgnoredActor(IgnoreActor);
diff --git a/Engine/Source/Runtime/Engine/Private/Materials/ParameterCollection.cpp b/Engine/Source/Runtime/Engine/Private/Materials/ParameterCollection.cpp
--- a/Engine/Source/Runtime/Engine/Private/Materials/ParameterCollection.cpp
+++ b/Engine/Source/Runtime/Engine/Private/Materials/ParameterCollection.cpp
@@ -516,11 +516,11 @@ void UMaterialParameterCollectionInstance::FinishDestroy()
 	Super::FinishDestroy();
 }
 
-void FMaterialParameterCollectionInstanceResource::GameThread_UpdateContents(const FGuid& InId, const TArray<FVector4>& Data)
+void FMaterialParameterCollectionInstanceResource::GameThread_UpdateContents(const FGuid& InGuid, const TArray<FVector4>& Data)
 {
 	ENQUEUE_UNIQUE_RENDER_COMMAND_THREEPARAMETER(
 		FUpdateCollectionCommand,
-		FGuid,Id,InId,
+		FGuid,Id,InGuid,
 		TArray<FVector4>,Data,Data,
 		FMaterialParameterCollectionInstanceResource*,Resource,this,
 	{
diff --git a/Engine/Source/Runtime/Engine/Private/Particles/ParticleModules.cpp b/Engine/Source/Runtime/Engine/Private/Particles/ParticleModules.cpp
--- a/Engine/Source/Runtime/Engine/Private/Particles/ParticleModules.cpp
+++ b/Engine/Source/Runtime/Engine/Private/Particles/ParticleModules.cpp
@@ -763,12 +763,12 @@ void UParticleModule::SetTransactionFlag()
 		}
 		else if( UArrayProperty* ArrayProp = Cast<UArrayProperty>(Property) )
 		{
-			if( UStructProperty* StructProp = Cast<UStructProperty>(ArrayProp->Inner) )
+			if( UStructProperty* InnerStructProp = Cast<UStructProperty>(ArrayProp->Inner) )
 			{
 				FScriptArrayHelper ArrayHelper( ArrayProp, Property->ContainerPtrToValuePtr<void>(this) );
 				for( int32 Idx = 0; Idx < ArrayHelper.Num(); ++Idx )
 				{
-					for( UProperty* ArrayProperty = StructProp->Struct->PropertyLink; ArrayProperty; ArrayProperty = ArrayProperty->PropertyLinkNext )
+					for( UProperty* ArrayProperty = InnerStructProp->Struct->PropertyLink; ArrayProperty; ArrayProperty = ArrayProperty->PropertyLinkNext )
 					{
 						if( UStructProperty* ArrayStructProp = Cast<UStructProperty>(ArrayProperty) )
 						{
diff --git a/Engine/Source/Runtime/Engine/Private/PhysicsEngine/BodyInstance.cpp b/Engine/Source/Runtime/Engine/Private/PhysicsEngine/BodyInstance.cpp
--- a/Engine/Source/Runtime/Engine/Private/PhysicsEngine/BodyInstance.cpp
+++ b/Engine/Source/Runtime/Engine/Private/PhysicsEngine/BodyInstance.cpp
@@ -679,8 +679,8 @@ void FBodyInstance::UpdatePhysicsShapeFilterData(uint32 SkelMeshCompID, bool bUs
 			bool bUseNotify = bNotifyOverride ? *bNotifyOverride : BI->bNotifyRigidBodyCollision;
 
 
-			UPrimitiveComponent* OwnerComponent = BI->OwnerComponent.IsValid() ? BI->OwnerComponent.Get() : nullptr;
-			int32 CompID = (OwnerComponent != nullptr) ? OwnerComponent->GetUniqueID() : 0;
+			UPrimitiveComponent* OwnerPrimitiveComponent = BI->OwnerComponent.IsValid() ? BI->OwnerComponent.Get() : nullptr;
+			int32 CompID = (OwnerPrimitiveComponent != nullptr) ? OwnerPrimitiveComponent->GetUniqueID() : 0;
 
 			// Create the filterdata structs
 			PxFilterData PSimFilterData;
@@ -740,7 +740,7 @@ void FBodyInstance::UpdatePhysicsShapeFilterData(uint32 SkelMeshCompID, bool bUs
 						PShape->setFlag(PxShapeFlag::eSIMULATION_SHAPE, true);
 					}
 
-					if (OwnerComponent == NULL || !OwnerComponent->IsA(UModelComponent::StaticClass()))
+					if (OwnerPrimitiveComponent == NULL || !OwnerPrimitiveComponent->IsA(UModelComponent::StaticClass()))
 					{
 						PShape->setFlag(PxShapeFlag::eVISUALIZATION, false); // dont draw the tri mesh, we can see it anyway, and its slow
 					}
@@ -1937,8 +1937,8 @@ bool FBodyInstance::UpdateBodyScale(const FVector& InScale3D)
 
 	FVector RelativeScale3D;
 	FVector RelativeScale3DAbs;
-	FVector NewScale3D;
-	ComputeScalingVectors(ScaleMode, InScale3DAdjusted, OldScale3D, RelativeScale3D, RelativeScale3DAbs, NewScale3D);
+	FVector UpdatedScale3D;
+	ComputeScalingVectors(ScaleMode, InScale3DAdjusted, OldScale3D, RelativeScale3D, RelativeScale3DAbs, UpdatedScale3D);
 
 	// Apply scaling
 #if WITH_PHYSX
@@ -2157,7 +2157,7 @@ bool FBodyInstance::UpdateBodyScale(const FVector& InScale3D)
 	// if success, overwrite old Scale3D, otherwise, just don't do it. It will have invalid scale next time
 	if (bSuccess)
 	{
-		Scale3D = NewScale3D;
+		Scale3D = UpdatedScale3D;
 
 		// update mass if required
 		if (bUpdateMassWhenScaleChanges)
diff --git a/Engine/Source/Runtime/Engine/Private/PhysicsEngine/BodySetup.cpp b/Engine/Source/Runtime/Engine/Private/PhysicsEngine/BodySetup.cpp
--- a/Engine/Source/Runtime/Engine/Private/PhysicsEngine/BodySetup.cpp
+++ b/Engine/Source/Runtime/Engine/Private/PhysicsEngine/BodySetup.cpp
@@ -1119,9 +1119,9 @@ int32 FKAggregateGeom::GetElementCount(int32 Type) const
 
 void FKConvexElem::ScaleElem(FVector DeltaSize, float MinSize)
 {
-	FTransform Transform = GetTransform();
-	Transform.SetScale3D(Transform.GetScale3D() + DeltaSize);
-	SetTransform(Transform);
+	FTransform ScaledTransform = GetTransform();
+	ScaledTransform.SetScale3D(ScaledTransform.GetScale3D() + DeltaSize);
+	SetTransform(ScaledTransform);
 }
 
 // References: 
diff --git a/Engine/Source/Runtime/Engine/Private/PrimitiveSceneProxy.cpp b/Engine/Source/Runtime/Engine/Private/PrimitiveSceneProxy.cpp
--- a/Engine/Source/Runtime/Engine/Private/PrimitiveSceneProxy.cpp
+++ b/Engine/Source/Runtime/Engine/Private/PrimitiveSceneProxy.cpp
@@ -483,17 +483,17 @@ bool FPrimitiveSceneProxy::IsShadowCast(const FSceneView* View) const
 void FPrimitiveSceneProxy::RenderBounds(
 	FPrimitiveDrawInterface* PDI, 
 	const FEngineShowFlags& EngineShowFlags, 
-	const FBoxSphereBounds& Bounds, 
+	const FBoxSphereBounds& InBounds, 
 	bool bRenderInEditor) const
 {
 	const ESceneDepthPriorityGroup DrawBoundsDPG = EngineShowFlags.Game ? SDPG_World : SDPG_Foreground;
 	if (EngineShowFlags.Bounds && (EngineShowFlags.Game || bRenderInEditor))
 	{
 		// Draw the static mesh's bounding box and sphere.
-		DrawWireBox(PDI,Bounds.GetBox(), FColor(72,72,255),DrawBoundsDPG);
-		DrawCircle(PDI, Bounds.Origin, FVector(1, 0, 0), FVector(0, 1, 0), FColor::Yellow, Bounds.SphereRadius, 32, DrawBoundsDPG);
-		DrawCircle(PDI, Bounds.Origin, FVector(1, 0, 0), FVector(0, 0, 1), FColor::Yellow, Bounds.SphereRadius, 32, DrawBoundsDPG);
-		DrawCircle(PDI, Bounds.Origin, FVector(0, 1, 0), FVector(0, 0, 1), FColor::Yellow, Bounds.SphereRadius, 32, DrawBoundsDPG);
+		DrawWireBox(PDI,InBounds.GetBox(), FColor(72,72,255),DrawBoundsDPG);
+		DrawCircle(PDI, InBounds.Origin, FVector(1, 0, 0), FVector(0, 1, 0), FColor::Yellow, InBounds.SphereRadius, 32, DrawBoundsDPG);
+		DrawCircle(PDI, InBounds.Origin, FVector(1, 0, 0), FVector(0, 0, 1), FColor::Yellow, InBounds.SphereRadius, 32, DrawBoundsDPG);
+		DrawCircle(PDI, InBounds.Origin, FVector(0, 1, 0), FVector(0, 0, 1), FColor::Yellow, InBounds.SphereRadius, 32, DrawBoundsDPG);
 	}
 }
 
diff --git a/Engine/Source/Runtime/Engine/Private/SceneView.cpp b/Engine/Source/Runtime/Engine/Private/SceneView.cpp
--- a/Engine/Source/Runtime/Engine/Private/SceneView.cpp
+++ b/Engine/Source/Runtime/Engine/Private/SceneView.cpp
@@ -1255,19 +1255,19 @@ void FSceneView::EndFinalPostprocessSettings(const FSceneViewInitOptions& ViewIn
 
 	// Anti-Aliasing
 	{
-		const auto FeatureLevel = GetFeatureLevel();
+		const auto SceneViewFeatureLevel = GetFeatureLevel();
 
 		static const auto CVar = IConsoleManager::Get().FindTConsoleVariableDataInt(TEXT("r.PostProcessAAQuality")); 
 		static auto* MobileHDRCvar = IConsoleManager::Get().FindTConsoleVariableDataInt(TEXT("r.MobileHDR"));
 		static auto* MobileMSAACvar = IConsoleManager::Get().FindTConsoleVariableDataInt(TEXT("r.MobileMSAA"));
-		static uint32 MSAAValue = GShaderPlatformForFeatureLevel[FeatureLevel] == SP_OPENGL_ES2_IOS ? 1 : MobileMSAACvar->GetValueOnGameThread();
+		static uint32 MSAAValue = GShaderPlatformForFeatureLevel[SceneViewFeatureLevel] == SP_OPENGL_ES2_IOS ? 1 : MobileMSAACvar->GetValueOnGameThread();
 
 		int32 Quality = FMath::Clamp(CVar->GetValueOnGameThread(), 0, 6);
 
 		if( !Family->EngineShowFlags.PostProcessing || !Family->EngineShowFlags.AntiAliasing || Quality <= 0
 			// Disable antialiasing in GammaLDR mode to avoid jittering.
-			|| (FeatureLevel == ERHIFeatureLevel::ES2 && MobileHDRCvar->GetValueOnGameThread() == 0)
-			|| (FeatureLevel <= ERHIFeatureLevel::ES3_1 && (MSAAValue > 1))
+			|| (SceneViewFeatureLevel == ERHIFeatureLevel::ES2 && MobileHDRCvar->GetValueOnGameThread() == 0)
+			|| (SceneViewFeatureLevel <= ERHIFeatureLevel::ES3_1 && (MSAAValue > 1))
 			|| Family->EngineShowFlags.VisualizeBloom)
 		{
 			FinalPostProcessSettings.AntiAliasingMethod = AAM_None;
diff --git a/Engine/Source/Runtime/Engine/Private/SkeletalMesh.cpp b/Engine/Source/Runtime/Engine/Private/SkeletalMesh.cpp
--- a/Engine/Source/Runtime/Engine/Private/SkeletalMesh.cpp
+++ b/Engine/Source/Runtime/Engine/Private/SkeletalMesh.cpp
@@ -547,9 +547,9 @@ FSkeletalMeshVertexColorBuffer& FSkeletalMeshVertexColorBuffer::operator=(const
 	return *this;
 }
 
-void FSkeletalMeshVertexColorBuffer::ResizeData(int32 NumVertices)
+void FSkeletalMeshVertexColorBuffer::ResizeData(int32 InNumVertices)
 {
-	VertexData->ResizeBuffer(NumVertices);
+	VertexData->ResizeBuffer(InNumVertices);
 	UpdateCachedInfo();
 }
 
diff --git a/Engine/Source/Runtime/Engine/Private/SkeletalMeshComponentPhysics.cpp b/Engine/Source/Runtime/Engine/Private/SkeletalMeshComponentPhysics.cpp
--- a/Engine/Source/Runtime/Engine/Private/SkeletalMeshComponentPhysics.cpp
+++ b/Engine/Source/Runtime/Engine/Private/SkeletalMeshComponentPhysics.cpp
@@ -442,9 +442,9 @@ void USkeletalMeshComponent::SetSimulatePhysics(bool bSimulate)
 		{
 			if (FBodyInstance* BodyInstance = Bodies[BodyIdx])
 			{
-				if (UBodySetup * BodySetup = PhysAsset->BodySetup[BodyIdx])
+				if (UBodySetup * PhysAssetBodySetup = PhysAsset->BodySetup[BodyIdx])
 				{
-					if (BodySetup->PhysicsType == EPhysicsType::PhysType_Default)
+					if (PhysAssetBodySetup->PhysicsType == EPhysicsType::PhysType_Default)
 					{
 						BodyInstance->SetInstanceSimulatePhysics(bSimulate);
 					}
@@ -760,17 +760,17 @@ void USkeletalMeshComponent::InitArticulated(FPhysScene* PhysScene)
 	Bodies.AddZeroed(NumBodies);
 	for(int32 i=0; i<NumBodies; i++)
 	{
-		UBodySetup* BodySetup = PhysicsAsset->BodySetup[i];
+		UBodySetup* PhysicsAssetBodySetup = PhysicsAsset->BodySetup[i];
 		Bodies[i] = new FBodyInstance;
 		FBodyInstance* BodyInst = Bodies[i];
 		check(BodyInst);
 
 		// Get transform of bone by name.
-		int32 BoneIndex = GetBoneIndex( BodySetup->BoneName );
+		int32 BoneIndex = GetBoneIndex( PhysicsAssetBodySetup->BoneName );
 		if(BoneIndex != INDEX_NONE)
 		{
 			// Copy body setup default instance properties
-			BodyInst->CopyBodyInstancePropertiesFrom(&BodySetup->DefaultInstance);
+			BodyInst->CopyBodyInstancePropertiesFrom(&PhysicsAssetBodySetup->DefaultInstance);
 			// we don't allow them to use this in editor. For physics asset, this set up is overriden by Physics Type. 
 			// but before we hide in the detail customization, we saved with this being true, causing the simulate always happens for some bodies
 			// so adding initialization here to disable this. 
@@ -804,7 +804,7 @@ void USkeletalMeshComponent::InitArticulated(FPhysScene* PhysScene)
 #if WITH_PHYSX
 			// Create physics body instance.
 			FTransform BoneTransform = GetBoneTransform( BoneIndex );
-			BodyInst->InitBody( BodySetup, BoneTransform, this, PhysScene, Aggregate);
+			BodyInst->InitBody( PhysicsAssetBodySetup, BoneTransform, this, PhysScene, Aggregate);
 #endif //WITH_PHYSX
 
 			// Remember if we have bodies in sync/async scene, so we know which scene(s) to lock when moving bodies
@@ -1191,12 +1191,12 @@ void USkeletalMeshComponent::ResetAllBodiesSimulatePhysics()
 	for(int32 i=0; i<Bodies.Num(); i++)
 	{
 		FBodyInstance*	BodyInst	= Bodies[i];
-		UBodySetup*	BodySetup	= BodyInst->BodySetup.Get();
+		UBodySetup*	BodyInstSetup	= BodyInst->BodySetup.Get();
 
 		// Set fixed on any bodies with bAlwaysFullAnimWeight set to true
-		if(BodySetup && BodySetup->PhysicsType != PhysType_Default)
+		if(BodyInstSetup && BodyInstSetup->PhysicsType != PhysType_Default)
 		{
-			if (BodySetup->PhysicsType == PhysType_Simulated)
+			if (BodyInstSetup->PhysicsType == PhysType_Simulated)
 			{
 				BodyInst->SetInstanceSimulatePhysics(true);
 			}
@@ -1241,10 +1241,10 @@ void USkeletalMeshComponent::SetAllBodiesPhysicsBlendWeight(float PhysicsBlendWe
 	for(int32 i=0; i<Bodies.Num(); i++)
 	{
 		FBodyInstance*	BodyInst	= Bodies[i];
-		UBodySetup*	BodySetup	= BodyInst->BodySetup.Get();
+		UBodySetup*	BodyInstSetup	= BodyInst->BodySetup.Get();
 
 		// Set fixed on any bodies with bAlwaysFullAnimWeight set to true
-		if(BodySetup && (!bSkipCustomPhysicsType || BodySetup->PhysicsType == PhysType_Default) )
+		if(BodyInstSetup && (!bSkipCustomPhysicsType || BodyInstSetup->PhysicsType == PhysType_Default) )
 		{
 			BodyInst->PhysicsBlendWeight = PhysicsBlendWeight;
 		}
@@ -1405,12 +1405,12 @@ void USkeletalMeshComponent::UpdateMeshForBrokenConstraints()
 			// Get child bodies of this joint
 			for(int32 BodySetupIndex = 0; BodySetupIndex < PhysicsAsset->BodySetup.Num(); BodySetupIndex++)
 			{
-				UBodySetup* BodySetup = PhysicsAsset->BodySetup[BodySetupIndex];
-				int32 BoneIndex = GetBoneIndex(BodySetup->BoneName);
+				UBodySetup* PhysicsAssetBodySetup = PhysicsAsset->BodySetup[BodySetupIndex];
+				int32 BoneIndex = GetBoneIndex(PhysicsAssetBodySetup->BoneName);
 				if( BoneIndex != INDEX_NONE && 
 					(BoneIndex == JointBoneIndex || SkeletalMesh->RefSkeleton.BoneIsChildOf(BoneIndex, JointBoneIndex)) )
 				{
-					DEBUGBROKENCONSTRAINTUPDATE(UE_LOG(LogSkeletalMesh, Log, TEXT("    Found Child Bone: (%d) %s"), BoneIndex, *BodySetup->BoneName.ToString());)
+					DEBUGBROKENCONSTRAINTUPDATE(UE_LOG(LogSkeletalMesh, Log, TEXT("    Found Child Bone: (%d) %s"), BoneIndex, *PhysicsAssetBodySetup->BoneName.ToString());)
 
 					FBodyInstance* ChildBodyInst = Bodies[BodySetupIndex];
 					if( ChildBodyInst )
@@ -1423,7 +1423,7 @@ void USkeletalMeshComponent::UpdateMeshForBrokenConstraints()
 						}
 					}
 
-					FConstraintInstance* ChildConstraintInst = FindConstraintInstance(BodySetup->BoneName);
+					FConstraintInstance* ChildConstraintInst = FindConstraintInstance(PhysicsAssetBodySetup->BoneName);
 					if( ChildConstraintInst )
 					{
 						if( ChildConstraintInst->bLinearPositionDrive )
@@ -1510,9 +1510,9 @@ void USkeletalMeshComponent::GetWeldedBodies(TArray<FBodyInstance*> & OutWeldedB
 			OutWeldedBodies.Add(&BodyInstance);
 			if (PhysicsAsset)
 			{
-				if (UBodySetup * BodySetup = PhysicsAsset->BodySetup[BodyIdx])
+				if (UBodySetup * PhysicsAssetBodySetup = PhysicsAsset->BodySetup[BodyIdx])
 				{
-					OutLabels.Add(BodySetup->BoneName);
+					OutLabels.Add(PhysicsAssetBodySetup->BoneName);
 				}
 				else
 				{
@@ -4853,27 +4853,29 @@ void USkeletalMeshComponent::DrawClothingPhysicalMeshWire(FPrimitiveDrawInterfac
 		for(uint32 IndexIdx=0; IndexIdx < NumIndices; IndexIdx+=3)
 		{
 			// draw a triangle
-			uint32 Index0 = VisualInfo.ClothPhysicalMeshIndices[IndexIdx];
-			uint32 Index1 = VisualInfo.ClothPhysicalMeshIndices[IndexIdx + 1];
-			uint32 Index2 = VisualInfo.ClothPhysicalMeshIndices[IndexIdx + 2];
-
-			// If index is greater than Num of vertices, then skip 
-			if(Index0 >= NumPhysicalMeshVerts 
-			|| Index1 >= NumPhysicalMeshVerts 
-			|| Index2 >= NumPhysicalMeshVerts)
+			FVector V[3];
+			float MaxDists[3];
 			{
-				continue;
-			}
+				uint32 Index0 = VisualInfo.ClothPhysicalMeshIndices[IndexIdx];
+				uint32 Index1 = VisualInfo.ClothPhysicalMeshIndices[IndexIdx + 1];
+				uint32 Index2 = VisualInfo.ClothPhysicalMeshIndices[IndexIdx + 2];
+
+				// If index is greater than Num of vertices, then skip 
+				if(Index0 >= NumPhysicalMeshVerts 
+				|| Index1 >= NumPhysicalMeshVerts 
+				|| Index2 >= NumPhysicalMeshVerts)
+				{
+					continue;
+				}
 
-			FVector V[3];
-			V[0] = (*PhysicalMeshVertices)[Index0];
-			V[1] = (*PhysicalMeshVertices)[Index1];
-			V[2] = (*PhysicalMeshVertices)[Index2];
+				V[0] = (*PhysicalMeshVertices)[Index0];
+				V[1] = (*PhysicalMeshVertices)[Index1];
+				V[2] = (*PhysicalMeshVertices)[Index2];
 
-			float MaxDists[3];
-			MaxDists[0] = VisualInfo.ClothConstrainCoeffs[Index0].ClothMaxDistance;
-			MaxDists[1] = VisualInfo.ClothConstrainCoeffs[Index1].ClothMaxDistance;
-			MaxDists[2] = VisualInfo.ClothConstrainCoeffs[Index2].ClothMaxDistance;
+				MaxDists[0] = VisualInfo.ClothConstrainCoeffs[Index0].ClothMaxDistance;
+				MaxDists[1] = VisualInfo.ClothConstrainCoeffs[Index1].ClothMaxDistance;
+				MaxDists[2] = VisualInfo.ClothConstrainCoeffs[Index2].ClothMaxDistance;
+			}
 
 			for(int32 i=0; i<3; i++)
 			{
diff --git a/Engine/Source/Runtime/Engine/Private/StaticMesh.cpp b/Engine/Source/Runtime/Engine/Private/StaticMesh.cpp
--- a/Engine/Source/Runtime/Engine/Private/StaticMesh.cpp
+++ b/Engine/Source/Runtime/Engine/Private/StaticMesh.cpp
@@ -386,17 +386,17 @@ void FStaticMeshLODResources::InitVertexFactory(
 
 			// Use the "override" color vertex buffer if one was supplied.  Otherwise, the color vertex stream
 			// associated with the static mesh is used.
-			FColorVertexBuffer* ColorVertexBuffer = &Params.LODResources->ColorVertexBuffer;
+			FColorVertexBuffer* LODColorVertexBuffer = &Params.LODResources->ColorVertexBuffer;
 			if( Params.OverrideColorVertexBuffer != NULL )
 			{
-				ColorVertexBuffer = Params.OverrideColorVertexBuffer;
+				LODColorVertexBuffer = Params.OverrideColorVertexBuffer;
 			}
-			if( ColorVertexBuffer->GetNumVertices() > 0 )
+			if( LODColorVertexBuffer->GetNumVertices() > 0 )
 			{
 				Data.ColorComponent = FVertexStreamComponent(
-					ColorVertexBuffer,
+					LODColorVertexBuffer,
 					0,	// Struct offset to color
-					ColorVertexBuffer->GetStride(),
+					LODColorVertexBuffer->GetStride(),
 					VET_Color
 					);
 			}
diff --git a/Engine/Source/Runtime/Engine/Private/StaticMeshBuild.cpp b/Engine/Source/Runtime/Engine/Private/StaticMeshBuild.cpp
--- a/Engine/Source/Runtime/Engine/Private/StaticMeshBuild.cpp
+++ b/Engine/Source/Runtime/Engine/Private/StaticMeshBuild.cpp
@@ -290,11 +290,11 @@ void RemapPaintedVertexColors(
 				// iterator for future consideration.
 				else
 				{
-					FOREACH_OCTREE_CHILD_NODE( ChildRef )
+					FOREACH_OCTREE_CHILD_NODE( OctreeChildRef )
 					{
-						if( CurNode.HasChild( ChildRef ) && CurNode.GetChild( ChildRef )->GetInclusiveElementCount() > 0 )
+						if( CurNode.HasChild( OctreeChildRef ) && CurNode.GetChild( OctreeChildRef )->GetInclusiveElementCount() > 0 )
 						{
-							OctreeIter.PushChild( ChildRef );
+							OctreeIter.PushChild( OctreeChildRef );
 						}
 					}
 				}
diff --git a/Engine/Source/Runtime/Engine/Private/StaticMeshComponent.cpp b/Engine/Source/Runtime/Engine/Private/StaticMeshComponent.cpp
--- a/Engine/Source/Runtime/Engine/Private/StaticMeshComponent.cpp
+++ b/Engine/Source/Runtime/Engine/Private/StaticMeshComponent.cpp
@@ -318,10 +318,10 @@ void UStaticMeshComponent::CheckForErrors()
 		// Check for element material index/material mismatches
 		for (int32 LODIndex = 0; LODIndex < StaticMesh->RenderData->LODResources.Num(); ++LODIndex)
 		{
-			FStaticMeshLODResources& LODData = StaticMesh->RenderData->LODResources[LODIndex];
-			for (int32 SectionIndex = 0; SectionIndex < LODData.Sections.Num(); SectionIndex++)
+			FStaticMeshLODResources& MeshLODData = StaticMesh->RenderData->LODResources[LODIndex];
+			for (int32 SectionIndex = 0; SectionIndex < MeshLODData.Sections.Num(); SectionIndex++)
 			{
-				FStaticMeshSection& Element = LODData.Sections[SectionIndex];
+				FStaticMeshSection& Element = MeshLODData.Sections[SectionIndex];
 				if (Element.NumTriangles == 0)
 				{
 					ZeroTriangleElements++;
diff --git a/Engine/Source/Runtime/Engine/Private/StaticMeshRender.cpp b/Engine/Source/Runtime/Engine/Private/StaticMeshRender.cpp
--- a/Engine/Source/Runtime/Engine/Private/StaticMeshRender.cpp
+++ b/Engine/Source/Runtime/Engine/Private/StaticMeshRender.cpp
@@ -1000,7 +1000,7 @@ FStaticMeshSceneProxy::FLODInfo::FLODInfo(const UStaticMeshComponent* InComponen
 {
 	const auto FeatureLevel = InComponent->GetWorld()->FeatureLevel;
 
-	FStaticMeshRenderData* RenderData = InComponent->StaticMesh->RenderData;
+	FStaticMeshRenderData* MeshRenderData = InComponent->StaticMesh->RenderData;
 	if(LODIndex < InComponent->LODData.Num())
 	{
 		const FStaticMeshComponentLODInfo& ComponentLODInfo = InComponent->LODData[LODIndex];
@@ -1013,7 +1013,7 @@ FStaticMeshSceneProxy::FLODInfo::FLODInfo(const UStaticMeshComponent* InComponen
 		// Initialize this LOD's overridden vertex colors, if it has any
 		if( ComponentLODInfo.OverrideVertexColors )
 		{
-			FStaticMeshLODResources& LODRenderData = RenderData->LODResources[LODIndex];
+			FStaticMeshLODResources& LODRenderData = MeshRenderData->LODResources[LODIndex];
 			
 			// the instance should point to the loaded data to avoid copy and memory waste
 			OverrideColorVertexBuffer = ComponentLODInfo.OverrideVertexColors;
@@ -1028,7 +1028,7 @@ FStaticMeshSceneProxy::FLODInfo::FLODInfo(const UStaticMeshComponent* InComponen
 		}
 	}
 
-	if (RenderData->bLODsShareStaticLighting && InComponent->LODData.IsValidIndex(0))
+	if (MeshRenderData->bLODsShareStaticLighting && InComponent->LODData.IsValidIndex(0))
 	{
 		const FStaticMeshComponentLODInfo& ComponentLODInfo = InComponent->LODData[0];
 		LightMap = ComponentLODInfo.LightMap;
@@ -1038,8 +1038,8 @@ FStaticMeshSceneProxy::FLODInfo::FLODInfo(const UStaticMeshComponent* InComponen
 	bool bHasStaticLighting = LightMap != NULL || ShadowMap != NULL;
 
 	// Gather the materials applied to the LOD.
-	Sections.Empty(RenderData->LODResources[LODIndex].Sections.Num());
-	FStaticMeshLODResources& LODModel = RenderData->LODResources[LODIndex];
+	Sections.Empty(MeshRenderData->LODResources[LODIndex].Sections.Num());
+	FStaticMeshLODResources& LODModel = MeshRenderData->LODResources[LODIndex];
 	for(int32 SectionIndex = 0;SectionIndex < LODModel.Sections.Num();SectionIndex++)
 	{
 		const FStaticMeshSection& Section = LODModel.Sections[SectionIndex];
diff --git a/Engine/Source/Runtime/Engine/Private/Texture.cpp b/Engine/Source/Runtime/Engine/Private/Texture.cpp
--- a/Engine/Source/Runtime/Engine/Private/Texture.cpp
+++ b/Engine/Source/Runtime/Engine/Private/Texture.cpp
@@ -524,8 +524,8 @@ void FTextureSource::Init2DWithMipChain(
 	ETextureSourceFormat NewFormat
 	)
 {
-	int32 NumMips = FMath::Max(FMath::CeilLogTwo(NewSizeX),FMath::CeilLogTwo(NewSizeY)) + 1;
-	Init(NewSizeX, NewSizeY, 1, NumMips, NewFormat);
+	int32 NewMipCount = FMath::Max(FMath::CeilLogTwo(NewSizeX),FMath::CeilLogTwo(NewSizeY)) + 1;
+	Init(NewSizeX, NewSizeY, 1, NewMipCount, NewFormat);
 }
 
 void FTextureSource::InitCubeWithMipChain(
@@ -534,8 +534,8 @@ void FTextureSource::InitCubeWithMipChain(
 	ETextureSourceFormat NewFormat
 	)
 {
-	int32 NumMips = FMath::Max(FMath::CeilLogTwo(NewSizeX),FMath::CeilLogTwo(NewSizeY)) + 1;
-	Init(NewSizeX, NewSizeY, 6, NumMips, NewFormat);
+	int32 NewMipCount = FMath::Max(FMath::CeilLogTwo(NewSizeX),FMath::CeilLogTwo(NewSizeY)) + 1;
+	Init(NewSizeX, NewSizeY, 6, NewMipCount, NewFormat);
 }
 
 void FTextureSource::Compress()
diff --git a/Engine/Source/Runtime/Engine/Private/TextureDerivedData.cpp b/Engine/Source/Runtime/Engine/Private/TextureDerivedData.cpp
--- a/Engine/Source/Runtime/Engine/Private/TextureDerivedData.cpp
+++ b/Engine/Source/Runtime/Engine/Private/TextureDerivedData.cpp
@@ -1365,7 +1365,7 @@ void FAsyncStreamDerivedMipWorker::DoWork()
 void UTexture2D::GetMipData(int32 FirstMipToLoad, void** OutMipData)
 {
 #if WITH_EDITOR
-	TextureDerivedDataTimings::FScopedMeasurement Timer(TextureDerivedDataTimings::GetMipDataTime);
+	TextureDerivedDataTimings::FScopedMeasurement MipDataTimer(TextureDerivedDataTimings::GetMipDataTime);
 #endif // #if WITH_EDITOR
 	if (PlatformData->TryLoadMips(FirstMipToLoad, OutMipData) == false)
 	{
@@ -1761,7 +1761,6 @@ void UTexture::MarkPlatformDataTransient()
 	TMap<FString, FTexturePlatformData*> *CookedPlatformData = GetCookedPlatformData();
 	if ( CookedPlatformData )
 	{
-		FDerivedDataCacheInterface& DDC = GetDerivedDataCacheRef();
 		for ( auto It : *CookedPlatformData )
 		{
 			FTexturePlatformData* PlatformData = It.Value;
diff --git a/Engine/Source/Runtime/Engine/Private/TimerManagerTests.cpp b/Engine/Source/Runtime/Engine/Private/TimerManagerTests.cpp
--- a/Engine/Source/Runtime/Engine/Private/TimerManagerTests.cpp
+++ b/Engine/Source/Runtime/Engine/Private/TimerManagerTests.cpp
@@ -228,7 +228,7 @@ bool TimerManagerTest_LoopingTimers_DifferentHandles(UWorld* World, FAutomationT
 	FTimerHandle HandleOne, HandleTwo;
 
 	int32 CallCount = 0;
-	auto Func = [](int* CallCount){ (*CallCount)++; };
+	auto Func = [](int* InCallCount){ (*InCallCount)++; };
 
 	FTimerDelegate Delegate = FTimerDelegate::CreateStatic(Func, &CallCount);
 
diff --git a/Engine/Source/Runtime/Engine/Private/UnrealClient.cpp b/Engine/Source/Runtime/Engine/Private/UnrealClient.cpp
--- a/Engine/Source/Runtime/Engine/Private/UnrealClient.cpp
+++ b/Engine/Source/Runtime/Engine/Private/UnrealClient.cpp
@@ -414,10 +414,10 @@ int32 FStatUnitData::DrawStat(FViewport* InViewport, FCanvas* InCanvas, int32 In
 	// Draw simple unit time graph
 	if (bShowUnitTimeGraph)
 	{
-		UFont* Font = GEngine->GetSmallFont();
-		check(Font);
-		int32 AlertPrintWidth = Font->GetStringSize(TEXT("000.0"));
-		int32 AlertPrintHeight = Font->GetStringHeightSize(TEXT("000.0"));
+		UFont* SmallFont = GEngine->GetSmallFont();
+		check(SmallFont);
+		int32 AlertPrintWidth = SmallFont->GetStringSize(TEXT("000.0"));
+		int32 AlertPrintHeight = SmallFont->GetStringHeightSize(TEXT("000.0"));
 
 		// The vertical axis is time in milliseconds
 		// The horizontal axis is the frame number (NOT time!!!)
@@ -586,7 +586,7 @@ int32 FStatUnitData::DrawStat(FViewport* InViewport, FCanvas* InCanvas, int32 In
 					if (StartX > LastPrintX)
 					{
 
-						InCanvas->DrawShadowedString(StartX, PrintY, *FString::Printf(TEXT("%3.1f"), TimeValues[CurUnitIndex]), Font, StatColor);
+						InCanvas->DrawShadowedString(StartX, PrintY, *FString::Printf(TEXT("%3.1f"), TimeValues[CurUnitIndex]), SmallFont, StatColor);
 						LastPrintX = StartX + AlertPrintWidth + AlertPadding;
 					}
 				}
@@ -985,8 +985,8 @@ void FViewport::Draw( bool bShouldPresent /*= true */)
 					GameThread.Waits = 0;
 				}
 
-				auto World = ViewportClient->GetWorld();
-				FCanvas Canvas(this, NULL, World, World ? World->FeatureLevel : GMaxRHIFeatureLevel);
+				UWorld* ViewportWorld = ViewportClient->GetWorld();
+				FCanvas Canvas(this, NULL, ViewportWorld, ViewportWorld ? ViewportWorld->FeatureLevel : GMaxRHIFeatureLevel);
 				{
 					// Make sure the Canvas is not rendered upside down
 					Canvas.SetAllowSwitchVerticalAxis(false);
diff --git a/Engine/Source/Runtime/Engine/Private/UnrealEngine.cpp b/Engine/Source/Runtime/Engine/Private/UnrealEngine.cpp
--- a/Engine/Source/Runtime/Engine/Private/UnrealEngine.cpp
+++ b/Engine/Source/Runtime/Engine/Private/UnrealEngine.cpp
@@ -8074,17 +8074,17 @@ void UEngine::HandleDisconnect( UWorld *InWorld, UNetDriver *NetDriver )
 
 			// In this case, the world is null and something went wrong, so we should travel back to the default world so that we
 			// can get back to a good state.
-			for (FWorldContext& WorldContext : WorldList)
+			for (FWorldContext& PotentialWorldContext : WorldList)
 			{
-				if (WorldContext.WorldType == EWorldType::Game)
+				if (PotentialWorldContext.WorldType == EWorldType::Game)
 				{
 					FURL DefaultURL;
 					DefaultURL.LoadURLConfig(TEXT("DefaultPlayer"), GGameIni);
 					const UGameMapsSettings* GameMapsSettings = GetDefault<UGameMapsSettings>();
 					if (GameMapsSettings)
 					{
-						WorldContext.TravelURL = FURL(&DefaultURL, *(GameMapsSettings->GetGameDefaultMap() + GameMapsSettings->LocalMapOptions), TRAVEL_Partial).ToString();
-						WorldContext.TravelType = TRAVEL_Partial;
+						PotentialWorldContext.TravelURL = FURL(&DefaultURL, *(GameMapsSettings->GetGameDefaultMap() + GameMapsSettings->LocalMapOptions), TRAVEL_Partial).ToString();
+						PotentialWorldContext.TravelType = TRAVEL_Partial;
 					}
 				}
 			}
@@ -11396,7 +11396,6 @@ bool UEngine::ToggleStatRaw(UWorld* World, FCommonViewportClient* ViewportClient
 #if !(UE_BUILD_SHIPPING || UE_BUILD_TEST)
 int32 UEngine::RenderStatReverb(UWorld* World, FViewport* Viewport, FCanvas* Canvas, int32 X, int32 Y, const FVector* ViewLocation, const FRotator* ViewRotation)
 {
-	FAudioDevice* AudioDevice = GetAudioDevice();
 	if (AudioDevice)
 	{
 		UReverbEffect* ReverbEffect = (AudioDevice->Effects ? AudioDevice->Effects->GetCurrentReverbEffect() : NULL);
@@ -11465,7 +11464,6 @@ int32 UEngine::RenderStatReverb(UWorld* World, FViewport* Viewport, FCanvas* Can
 // SOUNDMIXES
 int32 UEngine::RenderStatSoundMixes(UWorld* World, FViewport* Viewport, FCanvas* Canvas, int32 X, int32 Y, const FVector* ViewLocation, const FRotator* ViewRotation)
 {
-	FAudioDevice* AudioDevice = GetAudioDevice();
 	if (AudioDevice)
 	{
 		Canvas->DrawShadowedString(X, Y, TEXT("Active Sound Mixes:"), GetSmallFont(), FColor::Green);
@@ -11508,7 +11506,6 @@ int32 UEngine::RenderStatSoundWaves(UWorld* World, FViewport* Viewport, FCanvas*
 
 	TSet<FActiveSound*> ActiveSounds;
 
-	FAudioDevice* AudioDevice = GetAudioDevice();
 	if (AudioDevice)
 	{
 		TArray<FWaveInstance*> WaveInstances;
@@ -11561,7 +11558,6 @@ int32 UEngine::RenderStatSoundCues(UWorld* World, FViewport* Viewport, FCanvas*
 {
 	TSet<FActiveSound*> ActiveSounds;
 
-	FAudioDevice* AudioDevice = GetAudioDevice();
 	if (AudioDevice)
 	{
 		TArray<FWaveInstance*> WaveInstances;
@@ -11680,7 +11676,6 @@ int32 UEngine::RenderStatSounds(UWorld* World, FViewport* Viewport, FCanvas* Can
 #endif // UE_BUILD_DEBUG
 
 	TMapSounds SoundInfos;
-	FAudioDevice* AudioDevice = GetAudioDevice();
 	const FViewportClient::ESoundShowFlags::Type ShowSounds = Viewport->GetClient() ? Viewport->GetClient()->GetSoundShowFlags() : FViewportClient::ESoundShowFlags::Disabled;
 	const bool bDebug = ShowSounds & FViewportClient::ESoundShowFlags::Debug;
 
diff --git a/Engine/Source/Runtime/Engine/Private/UserDefinedStruct.cpp b/Engine/Source/Runtime/Engine/Private/UserDefinedStruct.cpp
--- a/Engine/Source/Runtime/Engine/Private/UserDefinedStruct.cpp
+++ b/Engine/Source/Runtime/Engine/Private/UserDefinedStruct.cpp
@@ -59,9 +59,9 @@ void UUserDefinedStruct::GetAssetRegistryTags(TArray<FAssetRegistryTag>& OutTags
 
 UProperty* UUserDefinedStruct::CustomFindProperty(const FName Name) const
 {
-	const FGuid Guid = FStructureEditorUtils::GetGuidFromPropertyName(Name);
-	UProperty* Property = Guid.IsValid() ? FStructureEditorUtils::GetPropertyByGuid(this, Guid) : NULL;
-	ensure(!Property || Guid == FStructureEditorUtils::GetGuidForProperty(Property));
+	const FGuid PropertyGuid = FStructureEditorUtils::GetGuidFromPropertyName(Name);
+	UProperty* Property = PropertyGuid.IsValid() ? FStructureEditorUtils::GetPropertyByGuid(this, PropertyGuid) : NULL;
+	ensure(!Property || PropertyGuid == FStructureEditorUtils::GetGuidForProperty(Property));
 	return Property;
 }
 
@@ -91,8 +91,8 @@ void UUserDefinedStruct::ValidateGuid()
 FString UUserDefinedStruct::PropertyNameToDisplayName(FName Name) const
 {
 #if WITH_EDITOR
-	auto Guid = FStructureEditorUtils::GetGuidFromPropertyName(Name);
-	return FStructureEditorUtils::GetVariableDisplayName(this, Guid);
+	FGuid PropertyGuid = FStructureEditorUtils::GetGuidFromPropertyName(Name);
+	return FStructureEditorUtils::GetVariableDisplayName(this, PropertyGuid);
 #endif	// WITH_EDITOR
 
 	const int32 GuidStrLen = 32;
diff --git a/Engine/Source/Runtime/Engine/Private/UserInterface/CanvasItem.cpp b/Engine/Source/Runtime/Engine/Private/UserInterface/CanvasItem.cpp
--- a/Engine/Source/Runtime/Engine/Private/UserInterface/CanvasItem.cpp
+++ b/Engine/Source/Runtime/Engine/Private/UserInterface/CanvasItem.cpp
@@ -169,7 +169,7 @@ void FCanvasItemTestbed::Draw( class FViewport* Viewport, class FCanvas* Canvas
 	if (bShowTris == true)
 	{
 		// Triangle
-		FCanvasTriangleItem TriItem(FVector2D(32.0f, 32.0f), FVector2D(64.0f, 32.0f), FVector2D(64.0f, 64.0f), GWhiteTexture);
+		//FCanvasTriangleItem TriItem(FVector2D(32.0f, 32.0f), FVector2D(64.0f, 32.0f), FVector2D(64.0f, 64.0f), GWhiteTexture);
 		//TriItem.Draw(Canvas);
 
 		// Triangle list
@@ -184,8 +184,8 @@ void FCanvasItemTestbed::Draw( class FViewport* Viewport, class FCanvas* Canvas
 		SingleTri.V2_Pos = FVector2D(438.0f, 308.0f);
 		TriangleList.Add(SingleTri);
 
-		FCanvasTriangleItem TriItemList(TriangleList, GWhiteTexture);
-		TriItemList.SetColor(FLinearColor::Red);
+		//FCanvasTriangleItem TriItemList(TriangleList, GWhiteTexture);
+		//TriItemList.SetColor(FLinearColor::Red);
 		//TriItemList.Draw( Canvas );
 
 		if (TestMaterial)
diff --git a/Engine/Source/Runtime/Landscape/Private/LandscapeGrass.cpp b/Engine/Source/Runtime/Landscape/Private/LandscapeGrass.cpp
--- a/Engine/Source/Runtime/Landscape/Private/LandscapeGrass.cpp
+++ b/Engine/Source/Runtime/Landscape/Private/LandscapeGrass.cpp
@@ -614,9 +614,9 @@ void ULandscapeComponent::RemoveGrassMap()
 	GrassData = MakeShareable(new FLandscapeComponentGrassData());
 }
 
-void ALandscapeProxy::RenderGrassMaps(const TArray<ULandscapeComponent*>& LandscapeComponents, const TArray<ULandscapeGrassType*>& GrassTypes)
+void ALandscapeProxy::RenderGrassMaps(const TArray<ULandscapeComponent*>& InLandscapeComponents, const TArray<ULandscapeGrassType*>& GrassTypes)
 {
-	FLandscapeGrassWeightExporter Exporter(this, LandscapeComponents, GrassTypes);
+	FLandscapeGrassWeightExporter Exporter(this, InLandscapeComponents, GrassTypes);
 	Exporter.ApplyResults();
 }
 
@@ -1173,7 +1173,6 @@ struct FAsyncGrassBuilder : public FGrassBuilderBase
 
 			// in-place sort the instances
 			FInstanceStream SwapBuffer;
-			int32 FirstUnfixedIndex = 0;
 			for (int32 FirstUnfixedIndex = 0; FirstUnfixedIndex < NumInstances; FirstUnfixedIndex++)
 			{
 				int32 LoadFrom = SortedInstances[FirstUnfixedIndex];
diff --git a/Engine/Source/Runtime/Landscape/Private/LandscapeLight.cpp b/Engine/Source/Runtime/Landscape/Private/LandscapeLight.cpp
--- a/Engine/Source/Runtime/Landscape/Private/LandscapeLight.cpp
+++ b/Engine/Source/Runtime/Landscape/Private/LandscapeLight.cpp
@@ -248,10 +248,10 @@ namespace
 								continue;
 							}
 
-							ULandscapeComponent* Neighbor = Info->XYtoComponentMap.FindRef(ComponentBase + FIntPoint(x+xx, y+yy));
-							if (Neighbor)
+							ULandscapeComponent* ComponentNeighbor = Info->XYtoComponentMap.FindRef(ComponentBase + FIntPoint(x+xx, y+yy));
+							if (ComponentNeighbor)
 							{
-								NeighborLOD = FMath::Max(::GetLightingLOD(Neighbor), NeighborLOD);
+								NeighborLOD = FMath::Max(::GetLightingLOD(ComponentNeighbor), NeighborLOD);
 							}
 						}
 					}
diff --git a/Engine/Source/Runtime/Landscape/Private/LandscapeRender.cpp b/Engine/Source/Runtime/Landscape/Private/LandscapeRender.cpp
--- a/Engine/Source/Runtime/Landscape/Private/LandscapeRender.cpp
+++ b/Engine/Source/Runtime/Landscape/Private/LandscapeRender.cpp
@@ -448,23 +448,23 @@ UMaterialInstanceConstant* GMaskRegionMaterial = nullptr;
 UTexture2D* GLandscapeBlackTexture = nullptr;
 
 // Game thread update
-void FLandscapeEditToolRenderData::Update(UMaterialInterface* InNewToolMaterial)
+void FLandscapeEditToolRenderData::Update(UMaterialInterface* InToolMaterial)
 {
 	ENQUEUE_UNIQUE_RENDER_COMMAND_TWOPARAMETER(
 		UpdateEditToolRenderData,
 		FLandscapeEditToolRenderData*, LandscapeEditToolRenderData, this,
-		UMaterialInterface*, NewToolMaterial, InNewToolMaterial,
+		UMaterialInterface*, NewToolMaterial, InToolMaterial,
 		{
 			LandscapeEditToolRenderData->ToolMaterial = NewToolMaterial;
 		});
 }
 
-void FLandscapeEditToolRenderData::UpdateGizmo(UMaterialInterface* InNewGizmoMaterial)
+void FLandscapeEditToolRenderData::UpdateGizmo(UMaterialInterface* InGizmoMaterial)
 {
 	ENQUEUE_UNIQUE_RENDER_COMMAND_TWOPARAMETER(
 		UpdateEditToolRenderData,
 		FLandscapeEditToolRenderData*, LandscapeEditToolRenderData, this,
-		UMaterialInterface*, NewGizmoMaterial, InNewGizmoMaterial,
+		UMaterialInterface*, NewGizmoMaterial, InGizmoMaterial,
 		{
 			LandscapeEditToolRenderData->GizmoMaterial = NewGizmoMaterial;
 		});
diff --git a/Engine/Source/Runtime/OpenGLDrv/Private/OpenGLCommands.cpp b/Engine/Source/Runtime/OpenGLDrv/Private/OpenGLCommands.cpp
--- a/Engine/Source/Runtime/OpenGLDrv/Private/OpenGLCommands.cpp
+++ b/Engine/Source/Runtime/OpenGLDrv/Private/OpenGLCommands.cpp
@@ -695,11 +695,13 @@ void FOpenGLDynamicRHI::SetupTexturesForDraw( FOpenGLContextState& ContextState,
 			// which should be preferred.
 			if(!FOpenGL::SupportsTextureView())
 			{
-				FTextureStage& CurrenTextureStage = ContextState.Textures[TextureStageIndex];
-				const bool bSameTarget = (TextureStage.Target == CurrenTextureStage.Target);
-				const bool bSameResource = (TextureStage.Resource == CurrenTextureStage.Resource);
-				const bool bSameLimitMip = bSameTarget && bSameResource && CurrenTextureStage.LimitMip == TextureStage.LimitMip;
-				const bool bSameNumMips = bSameTarget && bSameResource && CurrenTextureStage.NumMips == TextureStage.NumMips;
+				{
+					FTextureStage& CurrentTextureStage = ContextState.Textures[TextureStageIndex];
+					const bool bSameTarget = (TextureStage.Target == CurrentTextureStage.Target);
+					const bool bSameResource = (TextureStage.Resource == CurrentTextureStage.Resource);
+					const bool bSameLimitMip = bSameTarget && bSameResource && CurrentTextureStage.LimitMip == TextureStage.LimitMip;
+					const bool bSameNumMips = bSameTarget && bSameResource && CurrentTextureStage.NumMips == TextureStage.NumMips;
+				}
 				
 				// When trying to limit the mip available for sampling (as part of texture SRV)
 				// ensure that the texture is bound to only one sampler, or that all samplers
diff --git a/Engine/Source/Runtime/Renderer/Private/DeferredShadingRenderer.cpp b/Engine/Source/Runtime/Renderer/Private/DeferredShadingRenderer.cpp
--- a/Engine/Source/Runtime/Renderer/Private/DeferredShadingRenderer.cpp
+++ b/Engine/Source/Runtime/Renderer/Private/DeferredShadingRenderer.cpp
@@ -96,17 +96,6 @@ FDeferredShadingSceneRenderer::FDeferredShadingSceneRenderer(const FSceneViewFam
 	}
 }
 
-namespace
-{
-	FVector4 ClearQuadVertices[4] = 
-	{
-		FVector4( -1.0f,  1.0f, 0.0f, 1.0f ),
-		FVector4(  1.0f,  1.0f, 0.0f, 1.0f ),
-		FVector4( -1.0f, -1.0f, 0.0f, 1.0f ),
-		FVector4(  1.0f, -1.0f, 0.0f, 1.0f )
-	};
-}
-
 extern FGlobalBoundShaderState GClearMRTBoundShaderState[8];
 
 /** 
@@ -176,6 +165,13 @@ void FDeferredShadingSceneRenderer::ClearGBufferAtMaxZ(FRHICommandList& RHICmdLi
 		SetShaderValueArray(RHICmdList, PixelShader->GetPixelShader(),PixelShader->ColorParameter, ClearColors, NumActiveRenderTargets);
 
 		// Render quad
+		static const FVector4 ClearQuadVertices[4] = 
+		{
+			FVector4( -1.0f,  1.0f, 0.0f, 1.0f ),
+			FVector4(  1.0f,  1.0f, 0.0f, 1.0f ),
+			FVector4( -1.0f, -1.0f, 0.0f, 1.0f ),
+			FVector4(  1.0f, -1.0f, 0.0f, 1.0f )
+		};
 		DrawPrimitiveUP(RHICmdList, PT_TriangleStrip, 2, ClearQuadVertices, sizeof(ClearQuadVertices[0]));
 	}
 }
diff --git a/Engine/Source/Runtime/Renderer/Private/DistanceFieldGlobalIllumination.cpp b/Engine/Source/Runtime/Renderer/Private/DistanceFieldGlobalIllumination.cpp
--- a/Engine/Source/Runtime/Renderer/Private/DistanceFieldGlobalIllumination.cpp
+++ b/Engine/Source/Runtime/Renderer/Private/DistanceFieldGlobalIllumination.cpp
@@ -369,8 +369,10 @@ void PlaceVPLs(
 {
 	GVPLResources.AllocateFor(GVPLGridDimension * GVPLGridDimension);
 
-	uint32 ClearValues[4] = { 0 };
-	RHICmdList.ClearUAV(GVPLResources.VPLParameterBuffer.UAV, ClearValues);
+	{
+		uint32 ClearValues[4] = { 0 };
+		RHICmdList.ClearUAV(GVPLResources.VPLParameterBuffer.UAV, ClearValues);
+	}
 
 	const FLightSceneProxy* DirectionalLightProxy = NULL;
 
diff --git a/Engine/Source/Runtime/Renderer/Private/DistanceFieldSurfaceCacheLighting.cpp b/Engine/Source/Runtime/Renderer/Private/DistanceFieldSurfaceCacheLighting.cpp
--- a/Engine/Source/Runtime/Renderer/Private/DistanceFieldSurfaceCacheLighting.cpp
+++ b/Engine/Source/Runtime/Renderer/Private/DistanceFieldSurfaceCacheLighting.cpp
@@ -3772,7 +3772,6 @@ bool FDeferredShadingSceneRenderer::RenderDistanceFieldAOSurfaceCache(
 				}
 
 				{
-					uint32 ClearValues[4] = { 0 };
 					RHICmdList.ClearUAV(GAOCulledObjectBuffers.Buffers.ObjectIndirectArguments.UAV, ClearValues);
 
 					TShaderMapRef<FCullObjectsForViewCS> ComputeShader(GetGlobalShaderMap(Scene->GetFeatureLevel()));
@@ -3894,16 +3893,16 @@ bool FDeferredShadingSceneRenderer::RenderDistanceFieldAOSurfaceCache(
 					// Create new records which haven't been shaded yet for shading points which don't have a valid interpolation from existing records
 					for (int32 ViewIndex = 0; ViewIndex < Views.Num(); ViewIndex++)
 					{
-						const FViewInfo& View = Views[ViewIndex];
+						const FViewInfo& ViewInfo = Views[ViewIndex];
 
-						FIntPoint DownsampledViewSize = FIntPoint::DivideAndRoundUp(View.ViewRect.Size(), DestLevelDownsampleFactor);
+						FIntPoint DownsampledViewSize = FIntPoint::DivideAndRoundUp(ViewInfo.ViewRect.Size(), DestLevelDownsampleFactor);
 						uint32 GroupSizeX = FMath::DivideAndRoundUp(DownsampledViewSize.X, GDistanceFieldAOTileSizeX);
 						uint32 GroupSizeY = FMath::DivideAndRoundUp(DownsampledViewSize.Y, GDistanceFieldAOTileSizeY);
 
-						TShaderMapRef<FPopulateCacheCS> ComputeShader(View.ShaderMap);
+						TShaderMapRef<FPopulateCacheCS> ComputeShader(ViewInfo.ShaderMap);
 
 						RHICmdList.SetComputeShader(ComputeShader->GetComputeShader());
-						ComputeShader->SetParameters(RHICmdList, View, DistanceFieldAOBentNormalSplat->GetRenderTargetItem(), DistanceFieldNormal->GetRenderTargetItem(), DestLevelDownsampleFactor, DepthLevel, TileListGroupSize, Parameters);
+						ComputeShader->SetParameters(RHICmdList, ViewInfo, DistanceFieldAOBentNormalSplat->GetRenderTargetItem(), DistanceFieldNormal->GetRenderTargetItem(), DestLevelDownsampleFactor, DepthLevel, TileListGroupSize, Parameters);
 						DispatchComputeShader(RHICmdList, *ComputeShader, GroupSizeX, GroupSizeY, 1);
 
 						ComputeShader->UnsetParameters(RHICmdList);
@@ -4444,17 +4443,17 @@ void FDeferredShadingSceneRenderer::RenderMeshDistanceFieldVisualization(FRHICom
 
 				for (int32 ViewIndex = 0; ViewIndex < Views.Num(); ViewIndex++)
 				{
-					const FViewInfo& View = Views[ViewIndex];
+					const FViewInfo& ViewInfo = Views[ViewIndex];
 
-					uint32 GroupSizeX = FMath::DivideAndRoundUp(View.ViewRect.Size().X / GAODownsampleFactor, GDistanceFieldAOTileSizeX);
-					uint32 GroupSizeY = FMath::DivideAndRoundUp(View.ViewRect.Size().Y / GAODownsampleFactor, GDistanceFieldAOTileSizeY);
+					uint32 GroupSizeX = FMath::DivideAndRoundUp(ViewInfo.ViewRect.Size().X / GAODownsampleFactor, GDistanceFieldAOTileSizeX);
+					uint32 GroupSizeY = FMath::DivideAndRoundUp(ViewInfo.ViewRect.Size().Y / GAODownsampleFactor, GDistanceFieldAOTileSizeY);
 
 					{
 						SCOPED_DRAW_EVENT(RHICmdList, VisualizeMeshDistanceFieldCS);
-						TShaderMapRef<FVisualizeMeshDistanceFieldCS> ComputeShader(View.ShaderMap);
+						TShaderMapRef<FVisualizeMeshDistanceFieldCS> ComputeShader(ViewInfo.ShaderMap);
 
 						RHICmdList.SetComputeShader(ComputeShader->GetComputeShader());
-						ComputeShader->SetParameters(RHICmdList, View, VisualizeResultRT->GetRenderTargetItem(), FVector2D(GroupSizeX, GroupSizeY), Parameters);
+						ComputeShader->SetParameters(RHICmdList, ViewInfo, VisualizeResultRT->GetRenderTargetItem(), FVector2D(GroupSizeX, GroupSizeY), Parameters);
 						DispatchComputeShader(RHICmdList, *ComputeShader, GroupSizeX, GroupSizeY, 1);
 
 						ComputeShader->UnsetParameters(RHICmdList);
@@ -4467,31 +4466,31 @@ void FDeferredShadingSceneRenderer::RenderMeshDistanceFieldVisualization(FRHICom
 
 				for (int32 ViewIndex = 0; ViewIndex < Views.Num(); ViewIndex++)
 				{
-					const FViewInfo& View = Views[ViewIndex];
+					const FViewInfo& ViewInfo = Views[ViewIndex];
 
 					SCOPED_DRAW_EVENT(RHICmdList, UpsampleAO);
 
-					RHICmdList.SetViewport(View.ViewRect.Min.X, View.ViewRect.Min.Y, 0.0f, View.ViewRect.Max.X, View.ViewRect.Max.Y, 1.0f);
+					RHICmdList.SetViewport(ViewInfo.ViewRect.Min.X, ViewInfo.ViewRect.Min.Y, 0.0f, ViewInfo.ViewRect.Max.X, ViewInfo.ViewRect.Max.Y, 1.0f);
 					RHICmdList.SetRasterizerState(TStaticRasterizerState<FM_Solid, CM_None>::GetRHI());
 					RHICmdList.SetDepthStencilState(TStaticDepthStencilState<false, CF_Always>::GetRHI());
 					RHICmdList.SetBlendState(TStaticBlendState<>::GetRHI());
 
-					TShaderMapRef<FPostProcessVS> VertexShader( View.ShaderMap );
-					TShaderMapRef<FVisualizeDistanceFieldUpsamplePS> PixelShader( View.ShaderMap );
+					TShaderMapRef<FPostProcessVS> VertexShader( ViewInfo.ShaderMap );
+					TShaderMapRef<FVisualizeDistanceFieldUpsamplePS> PixelShader( ViewInfo.ShaderMap );
 
 					static FGlobalBoundShaderState BoundShaderState;
 
-					SetGlobalBoundShaderState(RHICmdList, View.GetFeatureLevel(), BoundShaderState, GFilterVertexDeclaration.VertexDeclarationRHI, *VertexShader, *PixelShader);
+					SetGlobalBoundShaderState(RHICmdList, ViewInfo.GetFeatureLevel(), BoundShaderState, GFilterVertexDeclaration.VertexDeclarationRHI, *VertexShader, *PixelShader);
 
-					PixelShader->SetParameters(RHICmdList, View, VisualizeResultRT);
+					PixelShader->SetParameters(RHICmdList, ViewInfo, VisualizeResultRT);
 
 					DrawRectangle( 
 						RHICmdList,
 						0, 0, 
-						View.ViewRect.Width(), View.ViewRect.Height(),
-						View.ViewRect.Min.X / GAODownsampleFactor, View.ViewRect.Min.Y / GAODownsampleFactor, 
-						View.ViewRect.Width() / GAODownsampleFactor, View.ViewRect.Height() / GAODownsampleFactor,
-						FIntPoint(View.ViewRect.Width(), View.ViewRect.Height()),
+						ViewInfo.ViewRect.Width(), ViewInfo.ViewRect.Height(),
+						ViewInfo.ViewRect.Min.X / GAODownsampleFactor, ViewInfo.ViewRect.Min.Y / GAODownsampleFactor, 
+						ViewInfo.ViewRect.Width() / GAODownsampleFactor, ViewInfo.ViewRect.Height() / GAODownsampleFactor,
+						FIntPoint(ViewInfo.ViewRect.Width(), ViewInfo.ViewRect.Height()),
 						GetBufferSizeForAO(),
 						*VertexShader);
 				}
diff --git a/Engine/Source/Runtime/Renderer/Private/HeightfieldLighting.cpp b/Engine/Source/Runtime/Renderer/Private/HeightfieldLighting.cpp
--- a/Engine/Source/Runtime/Renderer/Private/HeightfieldLighting.cpp
+++ b/Engine/Source/Runtime/Renderer/Private/HeightfieldLighting.cpp
@@ -429,7 +429,6 @@ void FHeightfieldLightingViewInfo::SetupVisibleHeightfields(const FViewInfo& Vie
 		&& NumPrimitives > 0
 		&& SupportsHeightfieldLighting(View.GetFeatureLevel(), View.GetShaderPlatform()))
 	{
-		const int32 NumPrimitives = Scene->DistanceFieldSceneData.HeightfieldPrimitives.Num();
 		extern float GAOMaxViewDistance;
 		const float MaxDistanceSquared = FMath::Square(GAOMaxViewDistance + GetGHeightfieldBounceDistance());
 		float LocalToWorldScale = 1;
diff --git a/Engine/Source/Runtime/Renderer/Private/PostProcess/RenderingCompositionGraph.cpp b/Engine/Source/Runtime/Renderer/Private/PostProcess/RenderingCompositionGraph.cpp
--- a/Engine/Source/Runtime/Renderer/Private/PostProcess/RenderingCompositionGraph.cpp
+++ b/Engine/Source/Runtime/Renderer/Private/PostProcess/RenderingCompositionGraph.cpp
@@ -405,8 +405,6 @@ void FRenderingCompositionGraph::RecursivelyProcess(const FRenderingCompositeOut
 					GGMLFileWriter.WriteLine("\tedge");
 					GGMLFileWriter.WriteLine("\t[");
 					{
-						ANSICHAR Line[MAX_SPRINTF];
-
 						FCStringAnsi::Sprintf(Line, "\t\tsource\t%d", ComputeUniqueOutputId(OutputRefIt->Source, OutputRefIt->PassOutputId));
 						GGMLFileWriter.WriteLine(Line);
 						FCStringAnsi::Sprintf(Line, "\t\ttarget\t%d", PassId);
@@ -448,8 +446,6 @@ void FRenderingCompositionGraph::RecursivelyProcess(const FRenderingCompositeOut
 				GGMLFileWriter.WriteLine("\tedge");
 				GGMLFileWriter.WriteLine("\t[");
 				{
-					ANSICHAR Line[MAX_SPRINTF];
-
 					FCStringAnsi::Sprintf(Line, "\t\tsource\t%d", ComputeUniqueOutputId(OutputRefIt->Source, OutputRefIt->PassOutputId));
 					GGMLFileWriter.WriteLine(Line);
 					FCStringAnsi::Sprintf(Line, "\t\ttarget\t%d", PassId);
@@ -478,9 +474,9 @@ void FRenderingCompositionGraph::RecursivelyProcess(const FRenderingCompositeOut
 		}
 
 		uint32 OutputId = 0;
-		while(FRenderingCompositeOutput* Output = Pass->GetOutput((EPassOutputId)(OutputId)))
+		while(FRenderingCompositeOutput* PassOutput = Pass->GetOutput((EPassOutputId)(OutputId)))
 		{
-			UE_LOG(LogConsoleResponse,Log, TEXT("  ePId_Output%d %s %s Dep: %d"), OutputId, *Output->RenderTargetDesc.GenerateInfoString(), Output->RenderTargetDesc.DebugName, Output->GetDependencyCount());
+			UE_LOG(LogConsoleResponse,Log, TEXT("  ePId_Output%d %s %s Dep: %d"), OutputId, *PassOutput->RenderTargetDesc.GenerateInfoString(), PassOutput->RenderTargetDesc.DebugName, PassOutput->GetDependencyCount());
 
 			GGMLFileWriter.WriteLine("\tnode");
 			GGMLFileWriter.WriteLine("\t[");
@@ -504,8 +500,8 @@ void FRenderingCompositionGraph::RecursivelyProcess(const FRenderingCompositeOut
 				GGMLFileWriter.WriteLine("\t\t[");
 				FCStringAnsi::Sprintf(Line, "\t\t\ttext\t\"ePId_Output%d '%s'\r%s\"", 
 					OutputId, 
-					(const char *)TCHAR_TO_ANSI(Output->RenderTargetDesc.DebugName),
-					(const char *)TCHAR_TO_ANSI(*Output->RenderTargetDesc.GenerateInfoString()));
+					(const char *)TCHAR_TO_ANSI(PassOutput->RenderTargetDesc.DebugName),
+					(const char *)TCHAR_TO_ANSI(*PassOutput->RenderTargetDesc.GenerateInfoString()));
 				GGMLFileWriter.WriteLine(Line);
 				GGMLFileWriter.WriteLine("\t\t]");
 			}
diff --git a/Engine/Source/Runtime/Renderer/Private/RendererScene.cpp b/Engine/Source/Runtime/Renderer/Private/RendererScene.cpp
--- a/Engine/Source/Runtime/Renderer/Private/RendererScene.cpp
+++ b/Engine/Source/Runtime/Renderer/Private/RendererScene.cpp
@@ -511,12 +511,12 @@ void FScene::AddPrimitive(UPrimitiveComponent* Primitive)
 		FCreateRenderThreadResourcesCommand,
 		FCreateRenderThreadParameters, Params, Params,
 	{
-		FPrimitiveSceneProxy* PrimitiveSceneProxy = Params.PrimitiveSceneProxy;
-		FScopeCycleCounter Context(PrimitiveSceneProxy->GetStatId());
-		PrimitiveSceneProxy->SetTransform(Params.RenderMatrix, Params.WorldBounds, Params.LocalBounds, Params.OwnerPosition);
+		FPrimitiveSceneProxy* SceneProxy = Params.PrimitiveSceneProxy;
+		FScopeCycleCounter Context(SceneProxy->GetStatId());
+		SceneProxy->SetTransform(Params.RenderMatrix, Params.WorldBounds, Params.LocalBounds, Params.OwnerPosition);
 
 		// Create any RenderThreadResources required.
-		PrimitiveSceneProxy->CreateRenderThreadResources();
+		SceneProxy->CreateRenderThreadResources();
 	});
 
 	// Create the primitive scene info.
@@ -1778,12 +1778,12 @@ void FScene::GetRelevantLights( UPrimitiveComponent* Primitive, TArray<const ULi
 }
 
 /** Sets the precomputed visibility handler for the scene, or NULL to clear the current one. */
-void FScene::SetPrecomputedVisibility(const FPrecomputedVisibilityHandler* PrecomputedVisibilityHandler)
+void FScene::SetPrecomputedVisibility(const FPrecomputedVisibilityHandler* NewPrecomputedVisibilityHandler)
 {
 	ENQUEUE_UNIQUE_RENDER_COMMAND_TWOPARAMETER(
 		UpdatePrecomputedVisibility,
 		FScene*,Scene,this,
-		const FPrecomputedVisibilityHandler*,PrecomputedVisibilityHandler,PrecomputedVisibilityHandler,
+		const FPrecomputedVisibilityHandler*,PrecomputedVisibilityHandler,NewPrecomputedVisibilityHandler,
 	{
 		Scene->PrecomputedVisibilityHandler = PrecomputedVisibilityHandler;
 	});
@@ -1804,7 +1804,7 @@ void FScene::SetShaderMapsOnMaterialResources_RenderThread(FRHICommandListImmedi
 		MaterialArray.Add(Material);
 	}
 
-	const auto FeatureLevel = GetFeatureLevel();
+	const auto SceneFeatureLevel = GetFeatureLevel();
 	bool bFoundAnyInitializedMaterials = false;
 
 	// Iterate through all loaded material render proxies and recache their uniform expressions if needed
@@ -1812,7 +1812,7 @@ void FScene::SetShaderMapsOnMaterialResources_RenderThread(FRHICommandListImmedi
 	for (TSet<FMaterialRenderProxy*>::TConstIterator It(FMaterialRenderProxy::GetMaterialRenderProxyMap()); It; ++It)
 	{
 		FMaterialRenderProxy* MaterialProxy = *It;
-		FMaterial* Material = MaterialProxy->GetMaterialNoFallback(FeatureLevel);
+		FMaterial* Material = MaterialProxy->GetMaterialNoFallback(SceneFeatureLevel);
 
 		if (Material && MaterialsToUpdate.Contains(Material))
 		{
@@ -1821,11 +1821,11 @@ void FScene::SetShaderMapsOnMaterialResources_RenderThread(FRHICommandListImmedi
 			MaterialProxy->CacheUniformExpressions();
 			bFoundAnyInitializedMaterials = true;
 
-			const FMaterial& MaterialForRendering = *MaterialProxy->GetMaterial(FeatureLevel);
+			const FMaterial& MaterialForRendering = *MaterialProxy->GetMaterial(SceneFeatureLevel);
 			check(MaterialForRendering.GetRenderingThreadShaderMap());
 
-			check(!MaterialProxy->UniformExpressionCache[FeatureLevel].bUpToDate
-				|| MaterialProxy->UniformExpressionCache[FeatureLevel].CachedUniformExpressionShaderMap == MaterialForRendering.GetRenderingThreadShaderMap());
+			check(!MaterialProxy->UniformExpressionCache[SceneFeatureLevel].bUpToDate
+				|| MaterialProxy->UniformExpressionCache[SceneFeatureLevel].CachedUniformExpressionShaderMap == MaterialForRendering.GetRenderingThreadShaderMap());
 
 			check(MaterialForRendering.GetRenderingThreadShaderMap()->IsValidForRendering());
 		}
@@ -1861,48 +1861,48 @@ void FScene::UpdateStaticDrawListsForMaterials_RenderThread(FRHICommandListImmed
 
 	// Warning: if any static draw lists are missed here, there will be a crash when trying to render with shaders that have been deleted!
 	TArray<FPrimitiveSceneInfo*> PrimitivesToUpdate;
-	auto FeatureLevel = GetFeatureLevel();
+	auto SceneFeatureLevel = GetFeatureLevel();
 
 	if (ShouldUseDeferredRenderer())
 	{
 		for (int32 DrawType = 0; DrawType < EBasePass_MAX; DrawType++)
 		{
-			BasePassNoLightMapDrawList[DrawType].GetUsedPrimitivesBasedOnMaterials(FeatureLevel, Materials, PrimitivesToUpdate);
-			BasePassSimpleDynamicLightingDrawList[DrawType].GetUsedPrimitivesBasedOnMaterials(FeatureLevel, Materials, PrimitivesToUpdate);
-			BasePassCachedVolumeIndirectLightingDrawList[DrawType].GetUsedPrimitivesBasedOnMaterials(FeatureLevel, Materials, PrimitivesToUpdate);
-			BasePassCachedPointIndirectLightingDrawList[DrawType].GetUsedPrimitivesBasedOnMaterials(FeatureLevel, Materials, PrimitivesToUpdate);
-			BasePassHighQualityLightMapDrawList[DrawType].GetUsedPrimitivesBasedOnMaterials(FeatureLevel, Materials, PrimitivesToUpdate);
-			BasePassDistanceFieldShadowMapLightMapDrawList[DrawType].GetUsedPrimitivesBasedOnMaterials(FeatureLevel, Materials, PrimitivesToUpdate);
-			BasePassLowQualityLightMapDrawList[DrawType].GetUsedPrimitivesBasedOnMaterials(FeatureLevel, Materials, PrimitivesToUpdate);
-			BasePassSelfShadowedTranslucencyDrawList[DrawType].GetUsedPrimitivesBasedOnMaterials(FeatureLevel, Materials, PrimitivesToUpdate);
-			BasePassSelfShadowedCachedPointIndirectTranslucencyDrawList[DrawType].GetUsedPrimitivesBasedOnMaterials(FeatureLevel, Materials, PrimitivesToUpdate);
+			BasePassNoLightMapDrawList[DrawType].GetUsedPrimitivesBasedOnMaterials(SceneFeatureLevel, Materials, PrimitivesToUpdate);
+			BasePassSimpleDynamicLightingDrawList[DrawType].GetUsedPrimitivesBasedOnMaterials(SceneFeatureLevel, Materials, PrimitivesToUpdate);
+			BasePassCachedVolumeIndirectLightingDrawList[DrawType].GetUsedPrimitivesBasedOnMaterials(SceneFeatureLevel, Materials, PrimitivesToUpdate);
+			BasePassCachedPointIndirectLightingDrawList[DrawType].GetUsedPrimitivesBasedOnMaterials(SceneFeatureLevel, Materials, PrimitivesToUpdate);
+			BasePassHighQualityLightMapDrawList[DrawType].GetUsedPrimitivesBasedOnMaterials(SceneFeatureLevel, Materials, PrimitivesToUpdate);
+			BasePassDistanceFieldShadowMapLightMapDrawList[DrawType].GetUsedPrimitivesBasedOnMaterials(SceneFeatureLevel, Materials, PrimitivesToUpdate);
+			BasePassLowQualityLightMapDrawList[DrawType].GetUsedPrimitivesBasedOnMaterials(SceneFeatureLevel, Materials, PrimitivesToUpdate);
+			BasePassSelfShadowedTranslucencyDrawList[DrawType].GetUsedPrimitivesBasedOnMaterials(SceneFeatureLevel, Materials, PrimitivesToUpdate);
+			BasePassSelfShadowedCachedPointIndirectTranslucencyDrawList[DrawType].GetUsedPrimitivesBasedOnMaterials(SceneFeatureLevel, Materials, PrimitivesToUpdate);
 		}
 	}
 	else
 	{
 		for (int32 DrawType = 0; DrawType < EBasePass_MAX; DrawType++)
 		{
-			BasePassForForwardShadingNoLightMapDrawList[DrawType].GetUsedPrimitivesBasedOnMaterials(FeatureLevel, Materials, PrimitivesToUpdate);
-			BasePassForForwardShadingLowQualityLightMapDrawList[DrawType].GetUsedPrimitivesBasedOnMaterials(FeatureLevel, Materials, PrimitivesToUpdate);
-			BasePassForForwardShadingDistanceFieldShadowMapLightMapDrawList[DrawType].GetUsedPrimitivesBasedOnMaterials(FeatureLevel, Materials, PrimitivesToUpdate);
-			BasePassForForwardShadingDirectionalLightAndSHIndirectDrawList[DrawType].GetUsedPrimitivesBasedOnMaterials(FeatureLevel, Materials, PrimitivesToUpdate);
-			BasePassForForwardShadingDirectionalLightAndSHDirectionalIndirectDrawList[DrawType].GetUsedPrimitivesBasedOnMaterials(FeatureLevel, Materials, PrimitivesToUpdate);
-			BasePassForForwardShadingDirectionalLightAndSHDirectionalCSMIndirectDrawList[DrawType].GetUsedPrimitivesBasedOnMaterials(FeatureLevel, Materials, PrimitivesToUpdate);
-			BasePassForForwardShadingMovableDirectionalLightDrawList[DrawType].GetUsedPrimitivesBasedOnMaterials(FeatureLevel, Materials, PrimitivesToUpdate);
-			BasePassForForwardShadingMovableDirectionalLightCSMDrawList[DrawType].GetUsedPrimitivesBasedOnMaterials(FeatureLevel, Materials, PrimitivesToUpdate);
-			BasePassForForwardShadingMovableDirectionalLightLightmapDrawList[DrawType].GetUsedPrimitivesBasedOnMaterials(FeatureLevel, Materials, PrimitivesToUpdate);
-			BasePassForForwardShadingMovableDirectionalLightCSMLightmapDrawList[DrawType].GetUsedPrimitivesBasedOnMaterials(FeatureLevel, Materials, PrimitivesToUpdate);
+			BasePassForForwardShadingNoLightMapDrawList[DrawType].GetUsedPrimitivesBasedOnMaterials(SceneFeatureLevel, Materials, PrimitivesToUpdate);
+			BasePassForForwardShadingLowQualityLightMapDrawList[DrawType].GetUsedPrimitivesBasedOnMaterials(SceneFeatureLevel, Materials, PrimitivesToUpdate);
+			BasePassForForwardShadingDistanceFieldShadowMapLightMapDrawList[DrawType].GetUsedPrimitivesBasedOnMaterials(SceneFeatureLevel, Materials, PrimitivesToUpdate);
+			BasePassForForwardShadingDirectionalLightAndSHIndirectDrawList[DrawType].GetUsedPrimitivesBasedOnMaterials(SceneFeatureLevel, Materials, PrimitivesToUpdate);
+			BasePassForForwardShadingDirectionalLightAndSHDirectionalIndirectDrawList[DrawType].GetUsedPrimitivesBasedOnMaterials(SceneFeatureLevel, Materials, PrimitivesToUpdate);
+			BasePassForForwardShadingDirectionalLightAndSHDirectionalCSMIndirectDrawList[DrawType].GetUsedPrimitivesBasedOnMaterials(SceneFeatureLevel, Materials, PrimitivesToUpdate);
+			BasePassForForwardShadingMovableDirectionalLightDrawList[DrawType].GetUsedPrimitivesBasedOnMaterials(SceneFeatureLevel, Materials, PrimitivesToUpdate);
+			BasePassForForwardShadingMovableDirectionalLightCSMDrawList[DrawType].GetUsedPrimitivesBasedOnMaterials(SceneFeatureLevel, Materials, PrimitivesToUpdate);
+			BasePassForForwardShadingMovableDirectionalLightLightmapDrawList[DrawType].GetUsedPrimitivesBasedOnMaterials(SceneFeatureLevel, Materials, PrimitivesToUpdate);
+			BasePassForForwardShadingMovableDirectionalLightCSMLightmapDrawList[DrawType].GetUsedPrimitivesBasedOnMaterials(SceneFeatureLevel, Materials, PrimitivesToUpdate);
 		}
 	}
 
-	PositionOnlyDepthDrawList.GetUsedPrimitivesBasedOnMaterials(FeatureLevel, Materials, PrimitivesToUpdate);
-	DepthDrawList.GetUsedPrimitivesBasedOnMaterials(FeatureLevel, Materials, PrimitivesToUpdate);
-	MaskedDepthDrawList.GetUsedPrimitivesBasedOnMaterials(FeatureLevel, Materials, PrimitivesToUpdate);
-	HitProxyDrawList.GetUsedPrimitivesBasedOnMaterials(FeatureLevel, Materials, PrimitivesToUpdate);
-	HitProxyDrawList_OpaqueOnly.GetUsedPrimitivesBasedOnMaterials(FeatureLevel, Materials, PrimitivesToUpdate);
-	VelocityDrawList.GetUsedPrimitivesBasedOnMaterials(FeatureLevel, Materials, PrimitivesToUpdate);
-	WholeSceneShadowDepthDrawList.GetUsedPrimitivesBasedOnMaterials(FeatureLevel, Materials, PrimitivesToUpdate);
-	WholeSceneReflectiveShadowMapDrawList.GetUsedPrimitivesBasedOnMaterials(FeatureLevel, Materials, PrimitivesToUpdate);
+	PositionOnlyDepthDrawList.GetUsedPrimitivesBasedOnMaterials(SceneFeatureLevel, Materials, PrimitivesToUpdate);
+	DepthDrawList.GetUsedPrimitivesBasedOnMaterials(SceneFeatureLevel, Materials, PrimitivesToUpdate);
+	MaskedDepthDrawList.GetUsedPrimitivesBasedOnMaterials(SceneFeatureLevel, Materials, PrimitivesToUpdate);
+	HitProxyDrawList.GetUsedPrimitivesBasedOnMaterials(SceneFeatureLevel, Materials, PrimitivesToUpdate);
+	HitProxyDrawList_OpaqueOnly.GetUsedPrimitivesBasedOnMaterials(SceneFeatureLevel, Materials, PrimitivesToUpdate);
+	VelocityDrawList.GetUsedPrimitivesBasedOnMaterials(SceneFeatureLevel, Materials, PrimitivesToUpdate);
+	WholeSceneShadowDepthDrawList.GetUsedPrimitivesBasedOnMaterials(SceneFeatureLevel, Materials, PrimitivesToUpdate);
+	WholeSceneReflectiveShadowMapDrawList.GetUsedPrimitivesBasedOnMaterials(SceneFeatureLevel, Materials, PrimitivesToUpdate);
 
 	for (int32 PrimitiveIndex = 0; PrimitiveIndex < PrimitivesToUpdate.Num(); PrimitiveIndex++)
 	{
@@ -2289,12 +2289,12 @@ void FScene::ApplyWorldOffset_RenderThread(FVector InOffset)
 	MotionBlurInfoData.ApplyOffset(InOffset);
 }
 
-void FScene::OnLevelAddedToWorld(FName InLevelName)
+void FScene::OnLevelAddedToWorld(FName LevelAddedName)
 {
 	ENQUEUE_UNIQUE_RENDER_COMMAND_TWOPARAMETER(
 		FLevelAddedToWorld,
 		class FScene*, Scene, this,
-		FName, LevelName, InLevelName,
+		FName, LevelName, LevelAddedName,
 	{
 		Scene->OnLevelAddedToWorld_RenderThread(LevelName);
 	});
diff --git a/Engine/Source/Runtime/Renderer/Private/SceneRendering.cpp b/Engine/Source/Runtime/Renderer/Private/SceneRendering.cpp
--- a/Engine/Source/Runtime/Renderer/Private/SceneRendering.cpp
+++ b/Engine/Source/Runtime/Renderer/Private/SceneRendering.cpp
@@ -453,13 +453,15 @@ TUniformBufferRef<FViewUniformShaderParameters> FViewInfo::CreateUniformBuffer(
 			static_assert(MAX_FORWARD_SHADOWCASCADES <= 4, "more than 4 cascades not supported by the shader and uniform buffer");
 			if (DirectionalLightShadowInfo)
 			{
-				FProjectedShadowInfo& ShadowInfo = *((*DirectionalLightShadowInfo)[0]);
-				FIntPoint ShadowBufferResolution = ShadowInfo.GetShadowBufferResolution();
-				FVector4 ShadowBufferSizeValue((float)ShadowBufferResolution.X, (float)ShadowBufferResolution.Y, 1.0f / (float)ShadowBufferResolution.X, 1.0f / (float)ShadowBufferResolution.Y);
+				{
+					FProjectedShadowInfo& ShadowInfo = *((*DirectionalLightShadowInfo)[0]);
+					FIntPoint ShadowBufferResolution = ShadowInfo.GetShadowBufferResolution();
+					FVector4 ShadowBufferSizeValue((float)ShadowBufferResolution.X, (float)ShadowBufferResolution.Y, 1.0f / (float)ShadowBufferResolution.X, 1.0f / (float)ShadowBufferResolution.Y);
 
-				ViewUniformShaderParameters.DirectionalLightShadowTexture = GSceneRenderTargets.GetShadowDepthZTexture();				
-				ViewUniformShaderParameters.DirectionalLightShadowTransition = 1.0f / ShadowInfo.ComputeTransitionSize();
-				ViewUniformShaderParameters.DirectionalLightShadowSize = ShadowBufferSizeValue;
+					ViewUniformShaderParameters.DirectionalLightShadowTexture = GSceneRenderTargets.GetShadowDepthZTexture();				
+					ViewUniformShaderParameters.DirectionalLightShadowTransition = 1.0f / ShadowInfo.ComputeTransitionSize();
+					ViewUniformShaderParameters.DirectionalLightShadowSize = ShadowBufferSizeValue;
+				}
 
 				int32 NumShadowsToCopy = FMath::Min(DirectionalLightShadowInfo->Num(), MAX_FORWARD_SHADOWCASCADES);				
 				for (int32 i = 0; i < NumShadowsToCopy; ++i)
@@ -750,10 +752,12 @@ void FViewInfo::CreateForwardLightDataUniformBuffer(FForwardLightData &OutForwar
 				OutForwardLightData.LightPositionAndInvRadius[LightIndex] = FVector4(FVector(BoundingSphereVector), InvRadius);
 
 				// SpotlightMaskAndMinRoughness, >0:Spotlight, MinRoughness = abs();
-				float W = FMath::Max(0.0001f, MinRoughness) * ((LightType == LightType_Spot) ? 1 : -1);
+				{
+					float W = FMath::Max(0.0001f, MinRoughness) * ((LightType == LightType_Spot) ? 1 : -1);
 
-				OutForwardLightData.LightDirectionAndSpotlightMaskAndMinRoughness[LightIndex] = FVector4(NormalizedLightDirection, W);
-				OutForwardLightData.SpotAnglesAndSourceRadiusAndDir[LightIndex] = FVector4(SpotAngles.X, SpotAngles.Y, SourceRadius, LightType == LightType_Directional);
+					OutForwardLightData.LightDirectionAndSpotlightMaskAndMinRoughness[LightIndex] = FVector4(NormalizedLightDirection, W);
+					OutForwardLightData.SpotAnglesAndSourceRadiusAndDir[LightIndex] = FVector4(SpotAngles.X, SpotAngles.Y, SourceRadius, LightType == LightType_Directional);
+				}
 
 				if(LightSceneInfo->Proxy->IsInverseSquared())
 				{
diff --git a/Engine/Source/Runtime/Slate/Private/Framework/Application/SlateApplication.cpp b/Engine/Source/Runtime/Slate/Private/Framework/Application/SlateApplication.cpp
--- a/Engine/Source/Runtime/Slate/Private/Framework/Application/SlateApplication.cpp
+++ b/Engine/Source/Runtime/Slate/Private/Framework/Application/SlateApplication.cpp
@@ -1390,8 +1390,6 @@ void FSlateApplication::Tick()
 			SCOPE_CYCLE_COUNTER(STAT_SlateTickWindowAndChildren);
 			SLATE_CYCLE_COUNTER_SCOPE(GTickWindowAndChildrenTime);
 
-			TSharedPtr<SWindow> ActiveModalWindow = GetActiveModalWindow();
-
 			if ( ActiveModalWindow.IsValid() )
 			{
 				// There is a modal window, and we just need to tick it.
@@ -3916,24 +3914,24 @@ bool FSlateApplication::ProcessMouseButtonDownEvent( const TSharedPtr< FGenericW
 			FScopedSwitchWorldHack SwitchWorld( MouseCaptorPath );
 			bInGame = FApp::IsGame();
 
-			Reply = FEventRouter::Route<FReply>( this, FEventRouter::FToLeafmostPolicy(MouseCaptorPath), MouseEvent, []( const FArrangedWidget& MouseCaptorWidget, const FPointerEvent& Event )
+			Reply = FEventRouter::Route<FReply>( this, FEventRouter::FToLeafmostPolicy(MouseCaptorPath), MouseEvent, []( const FArrangedWidget& InMouseCaptorWidget, const FPointerEvent& Event )
 			{
-				return MouseCaptorWidget.Widget->OnPreviewMouseButtonDown( MouseCaptorWidget.Geometry, Event );
+				return InMouseCaptorWidget.Widget->OnPreviewMouseButtonDown( InMouseCaptorWidget.Geometry, Event );
 			});
 			
 			if (!Reply.IsEventHandled())
 			{
 				Reply = FEventRouter::Route<FReply>( this, FEventRouter::FToLeafmostPolicy(MouseCaptorPath), MouseEvent,
-					[this]( const FArrangedWidget& MouseCaptorWidget, const FPointerEvent& Event )
+					[this]( const FArrangedWidget& InMouseCaptorWidget, const FPointerEvent& Event )
 					{
 						FReply TempReply = FReply::Unhandled();
 						if ( Event.IsTouchEvent() )
 						{
-							TempReply = MouseCaptorWidget.Widget->OnTouchStarted( MouseCaptorWidget.Geometry, Event );
+							TempReply = InMouseCaptorWidget.Widget->OnTouchStarted( InMouseCaptorWidget.Geometry, Event );
 						}
 						if ( !Event.IsTouchEvent() || ( !TempReply.IsEventHandled() && this->bTouchFallbackToMouse ) )
 						{
-							TempReply = MouseCaptorWidget.Widget->OnMouseButtonDown( MouseCaptorWidget.Geometry, Event );
+							TempReply = InMouseCaptorWidget.Widget->OnMouseButtonDown( InMouseCaptorWidget.Geometry, Event );
 						}
 						return TempReply;
 					} );
@@ -4391,9 +4389,9 @@ bool FSlateApplication::ProcessMouseMoveEvent( FPointerEvent& MouseEvent, bool b
 				// Switch worlds widgets in the current path
 				FScopedSwitchWorldHack SwitchWorld( DragDetectPath );
 
-				FReply Reply = FEventRouter::Route<FReply>(this, FEventRouter::FDirectPolicy(DetectDragForMe,DragDetectPath), MouseEvent, []( const FArrangedWidget& DetectDragForMe, const FPointerEvent& TranslatedMouseEvent )
+				FReply Reply = FEventRouter::Route<FReply>(this, FEventRouter::FDirectPolicy(DetectDragForMe,DragDetectPath), MouseEvent, []( const FArrangedWidget& InDetectDragForMe, const FPointerEvent& TranslatedMouseEvent )
 				{
-					return DetectDragForMe.Widget->OnDragDetected(DetectDragForMe.Geometry, TranslatedMouseEvent );
+					return InDetectDragForMe.Widget->OnDragDetected(InDetectDragForMe.Geometry, TranslatedMouseEvent );
 				});
 				
 				LOG_EVENT( EEventLog::DragDetected, Reply );
@@ -4505,11 +4503,11 @@ bool FSlateApplication::ProcessMouseMoveEvent( FPointerEvent& MouseEvent, bool b
 		if (IsDragDropping())
 		{
 			FDragDropEvent DragDropEvent( MouseEvent, DragDropContent );
-			FEventRouter::Route<FNoReply>( this, FEventRouter::FBubblePolicy(WidgetsUnderCursor), DragDropEvent, [&LastWidgetsUnderCursor](const FArrangedWidget& WidgetUnderCursor, const FDragDropEvent& DragDropEvent)
+			FEventRouter::Route<FNoReply>( this, FEventRouter::FBubblePolicy(WidgetsUnderCursor), DragDropEvent, [&LastWidgetsUnderCursor](const FArrangedWidget& WidgetUnderCursor, const FDragDropEvent& InDragDropEvent)
 			{
 				if ( !LastWidgetsUnderCursor.ContainsWidget(WidgetUnderCursor.Widget) )
 				{
-					WidgetUnderCursor.Widget->OnDragEnter( WidgetUnderCursor.Geometry, DragDropEvent );
+					WidgetUnderCursor.Widget->OnDragEnter( WidgetUnderCursor.Geometry, InDragDropEvent );
 				}
 				return FNoReply();
 			} );
@@ -4861,9 +4859,9 @@ void FSlateApplication::ProcessMotionDetectedEvent( FMotionEvent& MotionEvent )
 		
 		FScopedSwitchWorldHack SwitchWorld(PathToWidget);
 
-		FReply Reply = FEventRouter::Route<FReply>(this, FEventRouter::FBubblePolicy(PathToWidget), MotionEvent, [] (const FArrangedWidget& SomeWidget, const FMotionEvent& MotionEvent)
+		FReply Reply = FEventRouter::Route<FReply>(this, FEventRouter::FBubblePolicy(PathToWidget), MotionEvent, [] (const FArrangedWidget& SomeWidget, const FMotionEvent& InMotionEvent)
 		{
-			return SomeWidget.Widget->OnMotionDetected(SomeWidget.Geometry, MotionEvent);
+			return SomeWidget.Widget->OnMotionDetected(SomeWidget.Geometry, InMotionEvent);
 		});
 	}
 }
diff --git a/Engine/Source/Runtime/SlateRHIRenderer/Private/SlateRHIRenderingPolicy.h b/Engine/Source/Runtime/SlateRHIRenderer/Private/SlateRHIRenderingPolicy.h
--- a/Engine/Source/Runtime/SlateRHIRenderer/Private/SlateRHIRenderingPolicy.h
+++ b/Engine/Source/Runtime/SlateRHIRenderer/Private/SlateRHIRenderingPolicy.h
@@ -20,8 +20,8 @@ public:
 		, BufferUsageSize(0)
 	{}
 
-	TSlateElementVertexBuffer( uint32 BufferSize )
-		: BufferSize(BufferSize)
+	TSlateElementVertexBuffer( uint32 InBufferSize )
+		: BufferSize(InBufferSize)
 		, BufferUsageSize(0)
 	{}
 
diff --git a/Engine/Source/Runtime/SlateRHIRenderer/Private/SlateRHIResourceManager.cpp b/Engine/Source/Runtime/SlateRHIRenderer/Private/SlateRHIResourceManager.cpp
--- a/Engine/Source/Runtime/SlateRHIRenderer/Private/SlateRHIResourceManager.cpp
+++ b/Engine/Source/Runtime/SlateRHIRenderer/Private/SlateRHIResourceManager.cpp
@@ -515,8 +515,6 @@ FSlateShaderResourceProxy* FSlateRHIResourceManager::FindOrCreateDynamicTextureR
 	const FName ResourceName = InBrush.GetResourceName();
 	if ( ResourceName.IsValid() && ResourceName != NAME_None )
 	{
-		TSharedPtr<FSlateUTextureResource> TextureResource ;
-
 		if( InBrush.GetResourceObject() != nullptr )
 		{
 			UTexture2D* TextureObject = CastChecked<UTexture2D>(InBrush.GetResourceObject());
@@ -534,8 +532,8 @@ FSlateShaderResourceProxy* FSlateRHIResourceManager::FindOrCreateDynamicTextureR
 
 			if(TextureResource.IsValid())
 			{
-				UTexture2D* TextureObject = TextureResource->TextureObject;
-				if(TextureObject && !AccessedUTextures.Contains(TextureObject) && TextureObject->Resource)
+				UTexture2D* Texture2DObject = TextureResource->TextureObject;
+				if(Texture2DObject && !AccessedUTextures.Contains(Texture2DObject) && Texture2DObject->Resource)
 				{
 					// Set the texture rendering resource that should be used.  The UTexture resource could change at any time so we must do this each frame
 					ENQUEUE_UNIQUE_RENDER_COMMAND_ONEPARAMETER(UpdateSlateUTextureResource,
@@ -552,7 +550,7 @@ FSlateShaderResourceProxy* FSlateRHIResourceManager::FindOrCreateDynamicTextureR
 						});
 
 				
-					AccessedUTextures.Add(TextureObject);
+					AccessedUTextures.Add(Texture2DObject);
 				}
 
 				return TextureResource->Proxy;
diff --git a/Engine/Source/Runtime/UMG/Private/WidgetBlueprintLibrary.cpp b/Engine/Source/Runtime/UMG/Private/WidgetBlueprintLibrary.cpp
--- a/Engine/Source/Runtime/UMG/Private/WidgetBlueprintLibrary.cpp
+++ b/Engine/Source/Runtime/UMG/Private/WidgetBlueprintLibrary.cpp
@@ -334,7 +334,7 @@ UMaterialInstanceDynamic* UWidgetBlueprintLibrary::GetDynamicMaterial(FSlateBrus
 	// If the resource has a material interface we'll just update the brush to have a dynamic material.
 	else if ( UMaterialInterface* Material = Cast<UMaterialInterface>(Resource) )
 	{
-		UMaterialInstanceDynamic* DynamicMaterial = UMaterialInstanceDynamic::Create(Material, nullptr);
+		DynamicMaterial = UMaterialInstanceDynamic::Create(Material, nullptr);
 		Brush.SetResourceObject(DynamicMaterial);
 
 		return DynamicMaterial;
