UE4 - Ignore rendering zero meshes
- Fixes UE-11505 and UE-10900

diff --git a/Engine/Source/Runtime/Renderer/Private/PrimitiveSceneInfo.cpp b/Engine/Source/Runtime/Renderer/Private/PrimitiveSceneInfo.cpp
--- a/Engine/Source/Runtime/Renderer/Private/PrimitiveSceneInfo.cpp
+++ b/Engine/Source/Runtime/Renderer/Private/PrimitiveSceneInfo.cpp
@@ -42,19 +42,21 @@ public:
 		bool bShadowOnly
 		)
 	{
-		check(Mesh.GetNumPrimitives() > 0);
-		check(Mesh.VertexFactory);
-		check(Mesh.VertexFactory->IsInitialized());
+		if (Mesh.GetNumPrimitives() > 0)
+		{
+			check(Mesh.VertexFactory);
+			check(Mesh.VertexFactory->IsInitialized());
 #if DO_CHECK
-		Mesh.CheckUniformBuffers();
+			Mesh.CheckUniformBuffers();
 #endif
-		FStaticMesh* StaticMesh = new(PrimitiveSceneInfo->StaticMeshes) FStaticMesh(
-			PrimitiveSceneInfo,
-			Mesh,
-			ScreenSize,
-			bShadowOnly,
-			CurrentHitProxy ? CurrentHitProxy->Id : FHitProxyId()
-			);
+			FStaticMesh* StaticMesh = new(PrimitiveSceneInfo->StaticMeshes) FStaticMesh(
+				PrimitiveSceneInfo,
+				Mesh,
+				ScreenSize,
+				bShadowOnly,
+				CurrentHitProxy ? CurrentHitProxy->Id : FHitProxyId()
+				);
+		}
 	}
 
 private:
