UE4 - fix threading mega-stall when parallel rendering on consoles

diff --git a/Engine/Source/Runtime/RHI/Public/RHICommandList.h b/Engine/Source/Runtime/RHI/Public/RHICommandList.h
--- a/Engine/Source/Runtime/RHI/Public/RHICommandList.h
+++ b/Engine/Source/Runtime/RHI/Public/RHICommandList.h
@@ -1934,6 +1934,7 @@ namespace EImmediateFlushType
 		WaitForOutstandingTasksOnly = 0, 
 		DispatchToRHIThread, 
 		WaitForRHIThread, 
+		WaitForDispatchToRHIThread,
 		FlushRHIThread,
 		FlushRHIThreadFlushResources
 	};
diff --git a/Engine/Source/Runtime/RHI/Public/RHICommandList.inl b/Engine/Source/Runtime/RHI/Public/RHICommandList.inl
--- a/Engine/Source/Runtime/RHI/Public/RHICommandList.inl
+++ b/Engine/Source/Runtime/RHI/Public/RHICommandList.inl
@@ -46,6 +46,15 @@ FORCEINLINE_DEBUGGABLE void FRHICommandListImmediate::ImmediateFlush(EImmediateF
 			}
 		}
 		break;
+	case EImmediateFlushType::WaitForDispatchToRHIThread:
+		{
+			if (HasCommands())
+			{
+				GRHICommandList.ExecuteList(*this);
+			}
+			WaitForDispatch();
+		}
+		break;
 	case EImmediateFlushType::WaitForRHIThread:
 		{
 			check(GRHIThread);
diff --git a/Engine/Source/Runtime/Renderer/Private/SceneRendering.cpp b/Engine/Source/Runtime/Renderer/Private/SceneRendering.cpp
--- a/Engine/Source/Runtime/Renderer/Private/SceneRendering.cpp
+++ b/Engine/Source/Runtime/Renderer/Private/SceneRendering.cpp
@@ -1382,7 +1382,7 @@ static void RenderViewFamily_RenderThread(FRHICommandListImmediate& RHICmdList,
 		}
 		{
 			QUICK_SCOPE_CYCLE_COUNTER(STAT_DeleteSceneRenderer_Dispatch);
-			FRHICommandListExecutor::GetImmediateCommandList().ImmediateFlush(EImmediateFlushType::DispatchToRHIThread); // we want to make sure this all gets to the GPU this frame and doesn't hang around
+			FRHICommandListExecutor::GetImmediateCommandList().ImmediateFlush(EImmediateFlushType::WaitForDispatchToRHIThread); // we want to make sure this all gets to the rhi thread this frame and doesn't hang around
 		}
 	}
 
