Fix shadowed variables

diff --git a/Engine/Source/Developer/Android/AndroidPlatformEditor/Private/AndroidSDKSettings.cpp b/Engine/Source/Developer/Android/AndroidPlatformEditor/Private/AndroidSDKSettings.cpp
--- a/Engine/Source/Developer/Android/AndroidPlatformEditor/Private/AndroidSDKSettings.cpp
+++ b/Engine/Source/Developer/Android/AndroidPlatformEditor/Private/AndroidSDKSettings.cpp
@@ -20,14 +20,14 @@ void UAndroidSDKSettings::PostEditChangeProperty(struct FPropertyChangedEvent& P
 	UpdateTargetModulePaths();	
 }
 
-void UAndroidSDKSettings::SetTargetModule(ITargetPlatformManagerModule * TargetManagerModule)
+void UAndroidSDKSettings::SetTargetModule(ITargetPlatformManagerModule * InTargetManagerModule)
 {
-	this->TargetManagerModule = TargetManagerModule;
+	TargetManagerModule = InTargetManagerModule;
 }
 
-void UAndroidSDKSettings::SetDeviceDetection(IAndroidDeviceDetection * AndroidDeviceDetection)
+void UAndroidSDKSettings::SetDeviceDetection(IAndroidDeviceDetection * InAndroidDeviceDetection)
 {
-	this->AndroidDeviceDetection = AndroidDeviceDetection;
+	AndroidDeviceDetection = InAndroidDeviceDetection;
 }
 
 void UAndroidSDKSettings::UpdateTargetModulePaths()
diff --git a/Engine/Source/Runtime/Renderer/Private/PostProcess/RenderTargetPool.cpp b/Engine/Source/Runtime/Renderer/Private/PostProcess/RenderTargetPool.cpp
--- a/Engine/Source/Runtime/Renderer/Private/PostProcess/RenderTargetPool.cpp
+++ b/Engine/Source/Runtime/Renderer/Private/PostProcess/RenderTargetPool.cpp
@@ -695,7 +695,7 @@ FRenderTargetPool::SMemoryStats FRenderTargetPool::ComputeView()
 		{
 			uint32 ColumnX = 0;
 
-			for(int32 ColumnIndex = 0, Num = Colums.Num(); ColumnIndex < Num; ++ColumnIndex)
+			for(int32 ColumnIndex = 0, ColumnsNum = Colums.Num(); ColumnIndex < ColumnsNum; ++ColumnIndex)
 			{
 				const FRTPColumn& RTPColumn = Colums[ColumnIndex];
 
@@ -717,7 +717,7 @@ FRenderTargetPool::SMemoryStats FRenderTargetPool::ComputeView()
 				MemoryStats.TotalColumnSize += ColumnSize;
 				MemoryStats.TotalUsageInBytes += RTPColumn.SizeInBytes;
 				
-				for(int32 EventIndex = 0, Num = RenderTargetPoolEvents.Num(); EventIndex < Num; EventIndex++)
+				for(int32 EventIndex = 0, PoolEventsNum = RenderTargetPoolEvents.Num(); EventIndex < PoolEventsNum; EventIndex++)
 				{
 					FRenderTargetPoolEvent* Event = &RenderTargetPoolEvents[EventIndex];
 
diff --git a/Engine/Source/Runtime/Renderer/Private/PostProcess/RenderingCompositionGraph.cpp b/Engine/Source/Runtime/Renderer/Private/PostProcess/RenderingCompositionGraph.cpp
--- a/Engine/Source/Runtime/Renderer/Private/PostProcess/RenderingCompositionGraph.cpp
+++ b/Engine/Source/Runtime/Renderer/Private/PostProcess/RenderingCompositionGraph.cpp
@@ -481,8 +481,6 @@ void FRenderingCompositionGraph::RecursivelyProcess(const FRenderingCompositeOut
 			GGMLFileWriter.WriteLine("\tnode");
 			GGMLFileWriter.WriteLine("\t[");
 
-			ANSICHAR Line[MAX_SPRINTF];
-
 			{
 				GGMLFileWriter.WriteLine("\t\tgraphics");
 				GGMLFileWriter.WriteLine("\t\t[");
