- Minor high level changes to support Metal (breaking apart SM4 == Geometry shader assumptions, setting Load actions, etc)

diff --git a/Engine/Source/Runtime/Renderer/Private/CompositionLighting/PostProcessAmbientOcclusion.cpp b/Engine/Source/Runtime/Renderer/Private/CompositionLighting/PostProcessAmbientOcclusion.cpp
--- a/Engine/Source/Runtime/Renderer/Private/CompositionLighting/PostProcessAmbientOcclusion.cpp
+++ b/Engine/Source/Runtime/Renderer/Private/CompositionLighting/PostProcessAmbientOcclusion.cpp
@@ -537,7 +537,7 @@ void FRCPassPostProcessBasePassAO::Process(FRenderingCompositePassContext& Conte
 	const FSceneRenderTargetItem& DestRenderTarget = GSceneRenderTargets.GetSceneColor()->GetRenderTargetItem();
 
 	// Set the view family's render target/viewport.
-	SetRenderTarget(Context.RHICmdList, DestRenderTarget.TargetableTexture, FTextureRHIRef());
+	SetRenderTarget(Context.RHICmdList, DestRenderTarget.TargetableTexture,	FTextureRHIParamRef(), ESimpleRenderTargetMode::EExistingColorAndDepth);
 	Context.SetViewportAndCallRHI(View.ViewRect);
 
 	// set the state
diff --git a/Engine/Source/Runtime/Renderer/Private/CompositionLighting/PostProcessDeferredDecals.cpp b/Engine/Source/Runtime/Renderer/Private/CompositionLighting/PostProcessDeferredDecals.cpp
--- a/Engine/Source/Runtime/Renderer/Private/CompositionLighting/PostProcessDeferredDecals.cpp
+++ b/Engine/Source/Runtime/Renderer/Private/CompositionLighting/PostProcessDeferredDecals.cpp
@@ -870,8 +870,6 @@ void FRCPassPostProcessDeferredDecals::Process(FRenderingCompositePassContext& C
 
 		SCOPED_DRAW_EVENT(RHICmdList, Decals);
 		INC_DWORD_STAT_BY(STAT_Decals, SortedDecals.Num());
-		
-		RHICmdList.SetStreamSource(0, GUnitCubeVertexBuffer.VertexBufferRHI, sizeof(FVector4), 0);
 
 		enum EDecalResolveBufferIndex
 		{
@@ -980,6 +978,9 @@ void FRCPassPostProcessDeferredDecals::Process(FRenderingCompositePassContext& C
 				SetDecalBlendState(RHICmdList, SMFeatureLevel, RenderStage, (EDecalBlendMode)LastDecalBlendMode, DecalData.bHasNormal);
 			}
 
+			// we need to reset the stream source after any call to SetRenderTarget (at least for Metal, which doesn't queue up VB assignments)
+			RHICmdList.SetStreamSource(0, GUnitCubeVertexBuffer.VertexBufferRHI, sizeof(FVector4), 0);
+
 			{
 				FDeferredDecalVS* VertexShader = MaterialShaderMap->GetShader<FDeferredDecalVS>();
 				SetShader(Context, DecalData, VertexShader);
@@ -1058,7 +1059,7 @@ void FRCPassPostProcessDeferredDecals::Process(FRenderingCompositePassContext& C
 
 				SetShader(Context, DecalData, VertexShader);
 
-				RHICmdList.DrawIndexedPrimitive(GUnitCubeIndexBuffer.IndexBufferRHI, PT_TriangleList, 0, 0, 8, 0, GUnitCubeIndexBuffer.GetIndexCount() / 3, 0);
+				RHICmdList.DrawIndexedPrimitive(GUnitCubeIndexBuffer.IndexBufferRHI, PT_TriangleList, 0, 0, 8, 0, GUnitCubeIndexBuffer.GetIndexCount() / 3, 1);
 			}
 		}
 
diff --git a/Engine/Source/Runtime/Renderer/Private/FogRendering.cpp b/Engine/Source/Runtime/Renderer/Private/FogRendering.cpp
--- a/Engine/Source/Runtime/Renderer/Private/FogRendering.cpp
+++ b/Engine/Source/Runtime/Renderer/Private/FogRendering.cpp
@@ -308,7 +308,7 @@ bool FDeferredShadingSceneRenderer::RenderFog(FRHICommandListImmediate& RHICmdLi
 			0, 2, 3
 		};
 
-		GSceneRenderTargets.BeginRenderingSceneColor(RHICmdList, ESimpleRenderTargetMode::EUninitializedColorExistingReadOnlyDepth);
+		GSceneRenderTargets.BeginRenderingSceneColor(RHICmdList, ESimpleRenderTargetMode::EExistingColorAndReadOnlyDepth);
 		for(int32 ViewIndex = 0;ViewIndex < Views.Num();ViewIndex++)
 		{
 			const FViewInfo& View = Views[ViewIndex];
diff --git a/Engine/Source/Runtime/Renderer/Private/LightRendering.cpp b/Engine/Source/Runtime/Renderer/Private/LightRendering.cpp
--- a/Engine/Source/Runtime/Renderer/Private/LightRendering.cpp
+++ b/Engine/Source/Runtime/Renderer/Private/LightRendering.cpp
@@ -425,14 +425,15 @@ void FDeferredShadingSceneRenderer::RenderLights(FRHICommandListImmediate& RHICm
 			}
 			else if (SimpleLights.InstanceData.Num() > 0)
 			{
-				GSceneRenderTargets.BeginRenderingSceneColor(RHICmdList, ESimpleRenderTargetMode::EUninitializedColorExistingReadOnlyDepth);
+				GSceneRenderTargets.BeginRenderingSceneColor(RHICmdList, ESimpleRenderTargetMode::EExistingColorAndReadOnlyDepth);
 				RenderSimpleLightsStandardDeferred(RHICmdList, SimpleLights);
 			}
 
 			{
 				SCOPED_DRAW_EVENT(RHICmdList, StandardDeferredLighting);
 
-				GSceneRenderTargets.BeginRenderingSceneColor(RHICmdList, ESimpleRenderTargetMode::EUninitializedColorExistingReadOnlyDepth);
+				// make sure we don't clear the depth
+				GSceneRenderTargets.BeginRenderingSceneColor(RHICmdList, ESimpleRenderTargetMode::EExistingColorAndReadOnlyDepth);
 
 				// Draw non-shadowed non-light function lights without changing render targets between them
 				for (int32 LightIndex = StandardDeferredStart; LightIndex < AttenuationLightStart; LightIndex++)
@@ -495,7 +496,8 @@ void FDeferredShadingSceneRenderer::RenderLights(FRHICommandListImmediate& RHICm
 				}
 
 				// All shadows render with min blending
-				GSceneRenderTargets.BeginRenderingLightAttenuation(RHICmdList, true);
+				bool bClearToWhite = true;
+				GSceneRenderTargets.BeginRenderingLightAttenuation(RHICmdList, bClearToWhite);
 
 				bool bRenderedTranslucentObjectShadows = RenderTranslucentProjectedShadows(RHICmdList, &LightSceneInfo );
 				// Render non-modulated projected shadows to the attenuation buffer.
@@ -553,7 +555,7 @@ void FDeferredShadingSceneRenderer::RenderLights(FRHICommandListImmediate& RHICm
 			}
 
 			GSceneRenderTargets.SetLightAttenuationMode(bUsedLightAttenuation);
-			GSceneRenderTargets.BeginRenderingSceneColor(RHICmdList, ESimpleRenderTargetMode::EUninitializedColorExistingReadOnlyDepth);
+			GSceneRenderTargets.BeginRenderingSceneColor(RHICmdList, ESimpleRenderTargetMode::EExistingColorAndReadOnlyDepth);
 
 			// Render the light to the scene color buffer, conditionally using the attenuation buffer or a 1x1 white texture as input 
 			if(bDirectLighting)
@@ -593,7 +595,6 @@ void FDeferredShadingSceneRenderer::RenderLights(FRHICommandListImmediate& RHICm
 								FLightPropagationVolume* Lpv = ViewState->GetLightPropagationVolume();
 								if (Lpv && LightSceneInfo->Proxy)
 								{
-									
 									Lpv->InjectLightDirect(RHICmdList, *LightSceneInfo->Proxy, View);
 								}
 							}
diff --git a/Engine/Source/Runtime/Renderer/Private/ReflectionEnvironment.cpp b/Engine/Source/Runtime/Renderer/Private/ReflectionEnvironment.cpp
--- a/Engine/Source/Runtime/Renderer/Private/ReflectionEnvironment.cpp
+++ b/Engine/Source/Runtime/Renderer/Private/ReflectionEnvironment.cpp
@@ -869,10 +869,8 @@ void FDeferredShadingSceneRenderer::RenderStandardDeferredImageBasedReflections(
 				GRenderTargetPool.FindFreeElement(Desc, LightAccumulation, TEXT("LightAccumulation"));
 			}
 
-			SetRenderTarget(RHICmdList, LightAccumulation->GetRenderTargetItem().TargetableTexture, NULL);
-
 			// Clear to no reflection contribution, alpha of 1 indicates full background contribution
-			RHICmdList.Clear(true, FLinearColor(0, 0, 0, 1), false, 0, false, 0, FIntRect());
+			SetRenderTarget(RHICmdList, LightAccumulation->GetRenderTargetItem().TargetableTexture, NULL, ESimpleRenderTargetMode::EClearColorToBlackWithFullAlpha);
 
 			RHICmdList.SetViewport(View.ViewRect.Min.X, View.ViewRect.Min.Y, 0.0f, View.ViewRect.Max.X, View.ViewRect.Max.Y, 1.0f);
 
diff --git a/Engine/Source/Runtime/Renderer/Private/ShadowRendering.cpp b/Engine/Source/Runtime/Renderer/Private/ShadowRendering.cpp
--- a/Engine/Source/Runtime/Renderer/Private/ShadowRendering.cpp
+++ b/Engine/Source/Runtime/Renderer/Private/ShadowRendering.cpp
@@ -149,7 +149,7 @@ public:
 
 	static bool ShouldCache(EShaderPlatform Platform,const FMaterial* Material,const FVertexFactoryType* VertexFactoryType)
 	{
-		if (bIsForGeometryShader && !IsFeatureLevelSupported(Platform, ERHIFeatureLevel::SM4))
+		if (bIsForGeometryShader && !RHISupportsGeometryShaders(Platform))
 		{
 			return false;
 		}
@@ -297,13 +297,13 @@ public:
 
 	static bool ShouldCache(EShaderPlatform Platform,const FMaterial* Material,const FVertexFactoryType* VertexFactoryType)
 	{
-		return TShadowDepthVS<VertexShadowDepth_OnePassPointLight, false, false>::ShouldCache(Platform, Material, VertexFactoryType) && RHISupportsGeometryShaders(Platform);
+		return TShadowDepthVS<VertexShadowDepth_OnePassPointLight, false, false, true>::ShouldCache(Platform, Material, VertexFactoryType);
 	}
 
 	static void ModifyCompilationEnvironment(EShaderPlatform Platform, const FMaterial* Material, FShaderCompilerEnvironment& OutEnvironment)
 	{
 		FMeshMaterialShader::ModifyCompilationEnvironment(Platform, Material, OutEnvironment);
-		TShadowDepthVS<VertexShadowDepth_OnePassPointLight, false, false>::ModifyCompilationEnvironment(Platform, Material, OutEnvironment);
+		TShadowDepthVS<VertexShadowDepth_OnePassPointLight, false, false, true>::ModifyCompilationEnvironment(Platform, Material, OutEnvironment);
 	}
 
 	FOnePassPointShadowProjectionGS(const ShaderMetaType::CompiledShaderInitializerType& Initializer)
@@ -1733,7 +1733,7 @@ void StencilingGeometry::DrawSphere(FRHICommandList& RHICmdList)
 	RHICmdList.SetStreamSource(0, StencilingGeometry::GStencilSphereVertexBuffer.VertexBufferRHI, sizeof(FVector4), 0);
 	RHICmdList.DrawIndexedPrimitive(StencilingGeometry::GStencilSphereIndexBuffer.IndexBufferRHI, PT_TriangleList, 0, 0,
 		StencilingGeometry::GStencilSphereVertexBuffer.GetVertexCount(), 0, 
-		StencilingGeometry::GStencilSphereIndexBuffer.GetIndexCount() / 3, 0);
+		StencilingGeometry::GStencilSphereIndexBuffer.GetIndexCount() / 3, 1);
 }
 
 void StencilingGeometry::DrawCone(FRHICommandList& RHICmdList)
@@ -1742,7 +1742,7 @@ void StencilingGeometry::DrawCone(FRHICommandList& RHICmdList)
 	RHICmdList.SetStreamSource(0, StencilingGeometry::GStencilConeVertexBuffer.VertexBufferRHI, sizeof(FVector4), 0);
 
 	RHICmdList.DrawIndexedPrimitive(StencilingGeometry::GStencilConeIndexBuffer.IndexBufferRHI, PT_TriangleList, 0, 0,
-		FStencilConeIndexBuffer::NumVerts, 0, StencilingGeometry::GStencilConeIndexBuffer.GetIndexCount() / 3, 0);
+		FStencilConeIndexBuffer::NumVerts, 0, StencilingGeometry::GStencilConeIndexBuffer.GetIndexCount() / 3, 1);
 }
 
 /** bound shader state for stencil masking the shadow projection [0]:FShadowProjectionNoTransformVS [1]:FShadowProjectionVS */
