Fix for NumPlanes threshold

diff --git a/Engine/Source/Runtime/Renderer/Private/DistanceFieldShadowing.cpp b/Engine/Source/Runtime/Renderer/Private/DistanceFieldShadowing.cpp
--- a/Engine/Source/Runtime/Renderer/Private/DistanceFieldShadowing.cpp
+++ b/Engine/Source/Runtime/Renderer/Private/DistanceFieldShadowing.cpp
@@ -118,7 +118,7 @@ public:
 		SetShaderValue(RHICmdList, ShaderRHI, WorldToShadow, WorldToShadowValue);
 		SetShaderValue(RHICmdList, ShaderRHI, ShadowBoundingSphere, ShadowBoundingSphereValue);
 
-		if (NumPlanes < 12)
+		if (NumPlanes <= 12)
 		{
 			SetShaderValue(RHICmdList, ShaderRHI, NumShadowHullPlanes, NumPlanes);
 			SetShaderValueArray(RHICmdList, ShaderRHI, ShadowConvexHull, PlaneData, NumPlanes);
