Replace manual stencil clear with stencil writing zero during the shadow projection. Saves approx 1ms of GPU in MatineeFightScene on GTX 780.

(https://udn.unrealengine.com/questions/240409/stencil-clear-improvement-in-shadowprojectiononopa.html)

diff --git a/Engine/Source/Runtime/Renderer/Private/ShadowRendering.cpp b/Engine/Source/Runtime/Renderer/Private/ShadowRendering.cpp
--- a/Engine/Source/Runtime/Renderer/Private/ShadowRendering.cpp
+++ b/Engine/Source/Runtime/Renderer/Private/ShadowRendering.cpp
@@ -2120,8 +2120,8 @@ void FProjectedShadowInfo::RenderProjection(FRHICommandListImmediate& RHICmdList
 		// no depth test or writes, Test stencil for non-zero.
 		RHICmdList.SetDepthStencilState(TStaticDepthStencilState<
 			false,CF_Always,
-			true,CF_NotEqual,SO_Keep,SO_Keep,SO_Keep,
-			false,CF_Always,SO_Keep,SO_Keep,SO_Keep,
+			true,CF_NotEqual,SO_Zero,SO_Zero,SO_Zero,
+			false,CF_Always,SO_Zero,SO_Zero,SO_Zero,
 			0xff,0xff
 			>::GetRHI());
 	}
@@ -2229,11 +2229,6 @@ void FProjectedShadowInfo::RenderProjection(FRHICommandListImmediate& RHICmdList
 		// Disable depth bounds testing
 		RHICmdList.EnableDepthBoundsTest(false, 0, 1);
 	}
-	else
-	{
-		// Clear the stencil buffer to 0.
-		RHICmdList.Clear(false, FColor(0, 0, 0), false, 0, true, 0, FIntRect());
-	}
 }
 
 
