CVAR to optionally move OcclusionQuery rendering before basepass.

diff --git a/Engine/Source/Runtime/Renderer/Private/DeferredShadingRenderer.cpp b/Engine/Source/Runtime/Renderer/Private/DeferredShadingRenderer.cpp
--- a/Engine/Source/Runtime/Renderer/Private/DeferredShadingRenderer.cpp
+++ b/Engine/Source/Runtime/Renderer/Private/DeferredShadingRenderer.cpp
@@ -603,6 +603,8 @@ DECLARE_CYCLE_STAT(TEXT("Velocity"), STAT_CLM_Velocity, STATGROUP_CommandListMar
 DECLARE_CYCLE_STAT(TEXT("AfterVelocity"), STAT_CLM_AfterVelocity, STATGROUP_CommandListMarkers);
 DECLARE_CYCLE_STAT(TEXT("AfterFrame"), STAT_CLM_AfterFrame, STATGROUP_CommandListMarkers);
 
+FGraphEventRef FDeferredShadingSceneRenderer::OcclusionSubmittedFence;
+
 /**
  * Returns true if the depth Prepass needs to run
  */
@@ -624,10 +626,48 @@ static void SetAndClearViewGBuffer(FRHICommandListImmediate& RHICmdList, FViewIn
 	GSceneRenderTargets.BeginRenderingGBuffer(RHICmdList, ERenderTargetLoadAction::EClear, DepthLoadAction, ClearColor);
 }
 
+static TAutoConsoleVariable<int32> CVarOcclusionQueryLocation(
+	TEXT("r.OcclusionQueryLocation"),
+	0,
+	TEXT("Controls when occlusion queries are rendered.  Rendering before the base pass may give worse occlusion (because not all occluders generally render in the earlyzpass).  ")
+	TEXT("However, it may reduce CPU waiting for query result stalls on some platforms and increase overall performance.")
+	TEXT("0: After BasePass.")
+	TEXT("1: After EarlyZPass, but before BasePass."));
+
+void FDeferredShadingSceneRenderer::RenderOcclusion(FRHICommandListImmediate& RHICmdList, bool bRenderQueries, bool bRenderHZB)
+{		
+	if (bRenderQueries || bRenderHZB)
+	{
+		{
+			// Update the quarter-sized depth buffer with the current contents of the scene depth texture.
+			// This needs to happen before occlusion tests, which makes use of the small depth buffer.
+			QUICK_SCOPE_CYCLE_COUNTER(STAT_FDeferredShadingSceneRenderer_UpdateDownsampledDepthSurface);
+			UpdateDownsampledDepthSurface(RHICmdList);
+		}
+
+		if (bRenderHZB)
+		{
+			for (int32 ViewIndex = 0; ViewIndex < Views.Num(); ViewIndex++)
+			{
+				BuildHZB(RHICmdList, Views[ViewIndex]);
+			}
+		}
+
+		// Issue occlusion queries
+		// This is done after the downsampled depth buffer is created so that it can be used for issuing queries
+		BeginOcclusionTests(RHICmdList, bRenderQueries, bRenderHZB);
+
+		if (bRenderQueries && GRHIThread)
+		{
+			QUICK_SCOPE_CYCLE_COUNTER(STAT_OcclusionSubmittedFence_Dispatch);
+			OcclusionSubmittedFence = FRHICommandListExecutor::RHIThreadFence();
+		}
+	}
+}
+
 void FDeferredShadingSceneRenderer::Render(FRHICommandListImmediate& RHICmdList)
 {
-	bool bDBuffer = IsDBufferEnabled();
-	static FGraphEventRef OcclusionSubmittedFence;
+	bool bDBuffer = IsDBufferEnabled();	
 	if (GRHIThread)
 	{
 		QUICK_SCOPE_CYCLE_COUNTER(STAT_OcclusionSubmittedFence_Wait);
@@ -746,7 +786,8 @@ void FDeferredShadingSceneRenderer::Render(FRHICommandListImmediate& RHICmdList)
 
 	// Draw the scene pre-pass / early z pass, populating the scene depth buffer and HiZ
 	bool bDepthWasCleared = false;
-	if (NeedsPrePass(this))
+	const bool bNeedsPrePass = NeedsPrePass(this);
+	if (bNeedsPrePass)
 	{
 		RHICmdList.SetCurrentStat(GET_STATID(STAT_CLM_PrePass));
 		RenderPrePass(RHICmdList);
@@ -754,6 +795,11 @@ void FDeferredShadingSceneRenderer::Render(FRHICommandListImmediate& RHICmdList)
 		bDepthWasCleared = true;
 	}
 	RHICmdList.SetCurrentStat(GET_STATID(STAT_CLM_AfterPrePass));
+
+	//occlusion can't run before basepass if there's no prepass to fill in some depth to occlude against.
+	bool bOcclusionBeforeBasePass = (CVarOcclusionQueryLocation.GetValueOnRenderThread() == 1) && bNeedsPrePass;	
+	bool bHZBBeforeBasePass = false;
+	RenderOcclusion(RHICmdList, bOcclusionBeforeBasePass, bHZBBeforeBasePass);
 	
 	const bool bShouldRenderVelocities = ShouldRenderVelocities();
 	const bool bUseVelocityGBuffer = FVelocityRendering::OutputsToGBuffer();
@@ -870,30 +916,9 @@ void FDeferredShadingSceneRenderer::Render(FRHICommandListImmediate& RHICmdList)
 			);
 	}
 
-	// Update the quarter-sized depth buffer with the current contents of the scene depth texture.
-	// This needs to happen before occlusion tests, which makes use of the small depth buffer.
-	{
-		QUICK_SCOPE_CYCLE_COUNTER(STAT_FDeferredShadingSceneRenderer_UpdateDownsampledDepthSurface);
-		UpdateDownsampledDepthSurface(RHICmdList);
-	}
-
-	for (int32 ViewIndex = 0; ViewIndex < Views.Num(); ViewIndex++)
-	{
-		BuildHZB( RHICmdList, Views[ ViewIndex ] );
-	}
-
-	// Issue occlusion queries
-	// This is done after the downsampled depth buffer is created so that it can be used for issuing queries
-	if ( bIsOcclusionTesting )
-	{
-		BeginOcclusionTests(RHICmdList);
-	}
-
-	if (GRHIThread)
-	{
-		QUICK_SCOPE_CYCLE_COUNTER(STAT_OcclusionSubmittedFence_Dispatch);
-		OcclusionSubmittedFence = FRHICommandListExecutor::RHIThreadFence();
-	}
+	bool bOcclusionAfterBasePass = bIsOcclusionTesting && !bOcclusionBeforeBasePass;
+	bool bHZBAfterBasePass = true;
+	RenderOcclusion(RHICmdList, bOcclusionAfterBasePass, bHZBAfterBasePass);
 
 	TRefCountPtr<IPooledRenderTarget> VelocityRT;
 
diff --git a/Engine/Source/Runtime/Renderer/Private/DeferredShadingRenderer.h b/Engine/Source/Runtime/Renderer/Private/DeferredShadingRenderer.h
--- a/Engine/Source/Runtime/Renderer/Private/DeferredShadingRenderer.h
+++ b/Engine/Source/Runtime/Renderer/Private/DeferredShadingRenderer.h
@@ -101,6 +101,8 @@ public:
 	/** Finishes the view family rendering. */
 	void RenderFinish(FRHICommandListImmediate& RHICmdList);
 
+	void RenderOcclusion(FRHICommandListImmediate& RHICmdList, bool bRenderQueries, bool bRenderHZB);
+
 	/** Renders the view family. */
 	virtual void Render(FRHICommandListImmediate& RHICmdList) override;
 
@@ -117,6 +119,8 @@ public:
 
 private:
 
+	static FGraphEventRef OcclusionSubmittedFence;
+
 	/** Creates a per object projected shadow for the given interaction. */
 	void CreatePerObjectProjectedShadow(
 		FRHICommandListImmediate& RHICmdList,
@@ -172,7 +176,7 @@ private:
 	bool RenderPrePass(FRHICommandListImmediate& RHICmdList);
 
 	/** Issues occlusion queries. */
-	void BeginOcclusionTests(FRHICommandListImmediate& RHICmdList);
+	void BeginOcclusionTests(FRHICommandListImmediate& RHICmdList, bool bRenderQueries, bool bRenderHZB);
 
 	/** Renders the scene's fogging. */
 	bool RenderFog(FRHICommandListImmediate& RHICmdList, FLightShaftsOutput LightShaftsOutput);
diff --git a/Engine/Source/Runtime/Renderer/Private/SceneOcclusion.cpp b/Engine/Source/Runtime/Renderer/Private/SceneOcclusion.cpp
--- a/Engine/Source/Runtime/Renderer/Private/SceneOcclusion.cpp
+++ b/Engine/Source/Runtime/Renderer/Private/SceneOcclusion.cpp
@@ -1018,12 +1018,10 @@ void BuildHZB( FRHICommandListImmediate& RHICmdList, FViewInfo& View )
 	GRenderTargetPool.VisualizeTexture.SetCheckPoint( RHICmdList, View.HZB );
 }
 
-void FDeferredShadingSceneRenderer::BeginOcclusionTests(FRHICommandListImmediate& RHICmdList)
+void FDeferredShadingSceneRenderer::BeginOcclusionTests(FRHICommandListImmediate& RHICmdList, bool bRenderQueries, bool bRenderHZB)
 {
 	SCOPE_CYCLE_COUNTER(STAT_BeginOcclusionTestsTime);
-	const bool bUseDownsampledDepth = IsValidRef(GSceneRenderTargets.GetSmallDepthSurface()) && GSceneRenderTargets.UseDownsizedOcclusionQueries();
-
-	RHICmdList.BeginOcclusionQueryBatch();
+	const bool bUseDownsampledDepth = IsValidRef(GSceneRenderTargets.GetSmallDepthSurface()) && GSceneRenderTargets.UseDownsizedOcclusionQueries();	
 
 	if (bUseDownsampledDepth)
 	{
@@ -1034,153 +1032,161 @@ void FDeferredShadingSceneRenderer::BeginOcclusionTests(FRHICommandListImmediate
 		SetRenderTarget(RHICmdList, NULL, GSceneRenderTargets.GetSceneDepthSurface(), ESimpleRenderTargetMode::EExistingColorAndReadOnlyDepth);
 	}
 
-	// Perform occlusion queries for each view
-	for(int32 ViewIndex = 0;ViewIndex < Views.Num();ViewIndex++)
+	if (bRenderQueries)
 	{
-		SCOPED_DRAW_EVENT(RHICmdList, BeginOcclusionTests);
-		FViewInfo& View = Views[ViewIndex];
-
-		if (bUseDownsampledDepth)
-		{
-			const uint32 DownsampledX = FMath::TruncToInt(View.ViewRect.Min.X / GSceneRenderTargets.GetSmallColorDepthDownsampleFactor());
-			const uint32 DownsampledY = FMath::TruncToInt(View.ViewRect.Min.Y / GSceneRenderTargets.GetSmallColorDepthDownsampleFactor());
-			const uint32 DownsampledSizeX = FMath::TruncToInt(View.ViewRect.Width() / GSceneRenderTargets.GetSmallColorDepthDownsampleFactor());
-			const uint32 DownsampledSizeY = FMath::TruncToInt(View.ViewRect.Height() / GSceneRenderTargets.GetSmallColorDepthDownsampleFactor());
+		RHICmdList.BeginOcclusionQueryBatch();
 
-			// Setup the viewport for rendering to the downsampled depth buffer
-			RHICmdList.SetViewport(DownsampledX, DownsampledY, 0.0f, DownsampledX + DownsampledSizeX, DownsampledY + DownsampledSizeY, 1.0f);
-		}
-		else
-		{
-			RHICmdList.SetViewport(View.ViewRect.Min.X, View.ViewRect.Min.Y, 0.0f, View.ViewRect.Max.X, View.ViewRect.Max.Y, 1.0f);
-		}
-	
-		FSceneViewState* ViewState = (FSceneViewState*)View.State;
-	
-		if(ViewState && !View.bDisableQuerySubmissions)
+		// Perform occlusion queries for each view
+		for (int32 ViewIndex = 0; ViewIndex < Views.Num(); ViewIndex++)
 		{
-			// Depth tests, no depth writes, no color writes, opaque, solid rasterization wo/ backface culling.
-			// Note, this is a reversed Z depth surface, using CF_GreaterEqual.
-			RHICmdList.SetDepthStencilState(TStaticDepthStencilState<false, CF_GreaterEqual>::GetRHI());
-			// We only need to render the front-faces of the culling geometry (this halves the amount of pixels we touch)
-			RHICmdList.SetRasterizerState(View.bReverseCulling ? TStaticRasterizerState<FM_Solid, CM_CCW>::GetRHI() : TStaticRasterizerState<FM_Solid, CM_CW>::GetRHI());
-			RHICmdList.SetBlendState(TStaticBlendState<CW_NONE>::GetRHI());
-
-			// Lookup the vertex shader.
-			TShaderMapRef<FOcclusionQueryVS> VertexShader(View.ShaderMap);
-			SetGlobalBoundShaderState(RHICmdList, View.GetFeatureLevel(), OcclusionTestBoundShaderState, GetVertexDeclarationFVector3(), *VertexShader, NULL);
-			VertexShader->SetParameters(RHICmdList, View);
-
-			// Issue this frame's occlusion queries (occlusion queries from last frame may still be in flight)
-#if BUFFERED_OCCLUSION_QUERIES				
-			const uint32 QueryIndex = FOcclusionQueryHelpers::GetQueryIssueIndex(ViewState->PendingPrevFrameNumber, ViewState->NumBufferedFrames);
-			FSceneViewState::ShadowKeyOcclusionQueryMap& ShadowOcclusionQueryMap = ViewState->ShadowOcclusionQueryMaps[QueryIndex];
-#else
-			FSceneViewState::ShadowKeyOcclusionQueryMap& ShadowOcclusionQueryMap = ViewState->ShadowOcclusionQueryMap;
-#endif			
+			SCOPED_DRAW_EVENT(RHICmdList, BeginOcclusionTests);
+			FViewInfo& View = Views[ViewIndex];
 
-			// Clear primitives which haven't been visible recently out of the occlusion history, and reset old pending occlusion queries.
-			ViewState->TrimOcclusionHistory(RHICmdList, ViewFamily.CurrentRealTime - GEngine->PrimitiveProbablyVisibleTime, ViewFamily.CurrentRealTime, ViewState->OcclusionFrameCounter);
+			if (bUseDownsampledDepth)
+			{
+				const uint32 DownsampledX = FMath::TruncToInt(View.ViewRect.Min.X / GSceneRenderTargets.GetSmallColorDepthDownsampleFactor());
+				const uint32 DownsampledY = FMath::TruncToInt(View.ViewRect.Min.Y / GSceneRenderTargets.GetSmallColorDepthDownsampleFactor());
+				const uint32 DownsampledSizeX = FMath::TruncToInt(View.ViewRect.Width() / GSceneRenderTargets.GetSmallColorDepthDownsampleFactor());
+				const uint32 DownsampledSizeY = FMath::TruncToInt(View.ViewRect.Height() / GSceneRenderTargets.GetSmallColorDepthDownsampleFactor());
 
-			// Give back all these occlusion queries to the pool.
-			for ( TMap<FSceneViewState::FProjectedShadowKey, FRenderQueryRHIRef>::TIterator QueryIt(ShadowOcclusionQueryMap); QueryIt; ++QueryIt )
+				// Setup the viewport for rendering to the downsampled depth buffer
+				RHICmdList.SetViewport(DownsampledX, DownsampledY, 0.0f, DownsampledX + DownsampledSizeX, DownsampledY + DownsampledSizeY, 1.0f);
+			}
+			else
 			{
-				//FRenderQueryRHIParamRef Query = QueryIt.Value();
-				//check( Query.GetRefCount() == 1 );
-				ViewState->OcclusionQueryPool.ReleaseQuery(RHICmdList, QueryIt.Value());
+				RHICmdList.SetViewport(View.ViewRect.Min.X, View.ViewRect.Min.Y, 0.0f, View.ViewRect.Max.X, View.ViewRect.Max.Y, 1.0f);
 			}
-			ShadowOcclusionQueryMap.Reset();
 
+			FSceneViewState* ViewState = (FSceneViewState*)View.State;
+
+			if (ViewState && !View.bDisableQuerySubmissions)
 			{
-				SCOPED_DRAW_EVENT(RHICmdList, ShadowFrustumQueries);
+				// Depth tests, no depth writes, no color writes, opaque, solid rasterization wo/ backface culling.
+				// Note, this is a reversed Z depth surface, using CF_GreaterEqual.
+				RHICmdList.SetDepthStencilState(TStaticDepthStencilState<false, CF_GreaterEqual>::GetRHI());
+				// We only need to render the front-faces of the culling geometry (this halves the amount of pixels we touch)
+				RHICmdList.SetRasterizerState(View.bReverseCulling ? TStaticRasterizerState<FM_Solid, CM_CCW>::GetRHI() : TStaticRasterizerState<FM_Solid, CM_CW>::GetRHI());
+				RHICmdList.SetBlendState(TStaticBlendState<CW_NONE>::GetRHI());
+
+				// Lookup the vertex shader.
+				TShaderMapRef<FOcclusionQueryVS> VertexShader(View.ShaderMap);
+				SetGlobalBoundShaderState(RHICmdList, View.GetFeatureLevel(), OcclusionTestBoundShaderState, GetVertexDeclarationFVector3(), *VertexShader, NULL);
+				VertexShader->SetParameters(RHICmdList, View);
 
-				for(TSparseArray<FLightSceneInfoCompact>::TConstIterator LightIt(Scene->Lights);LightIt;++LightIt)
+				// Issue this frame's occlusion queries (occlusion queries from last frame may still be in flight)
+#if BUFFERED_OCCLUSION_QUERIES				
+				const uint32 QueryIndex = FOcclusionQueryHelpers::GetQueryIssueIndex(ViewState->PendingPrevFrameNumber, ViewState->NumBufferedFrames);
+				FSceneViewState::ShadowKeyOcclusionQueryMap& ShadowOcclusionQueryMap = ViewState->ShadowOcclusionQueryMaps[QueryIndex];
+#else
+				FSceneViewState::ShadowKeyOcclusionQueryMap& ShadowOcclusionQueryMap = ViewState->ShadowOcclusionQueryMap;
+#endif			
+
+				// Clear primitives which haven't been visible recently out of the occlusion history, and reset old pending occlusion queries.
+				ViewState->TrimOcclusionHistory(RHICmdList, ViewFamily.CurrentRealTime - GEngine->PrimitiveProbablyVisibleTime, ViewFamily.CurrentRealTime, ViewState->OcclusionFrameCounter);
+
+				// Give back all these occlusion queries to the pool.
+				for (TMap<FSceneViewState::FProjectedShadowKey, FRenderQueryRHIRef>::TIterator QueryIt(ShadowOcclusionQueryMap); QueryIt; ++QueryIt)
 				{
-					const FVisibleLightInfo& VisibleLightInfo = VisibleLightInfos[LightIt.GetIndex()];
+					//FRenderQueryRHIParamRef Query = QueryIt.Value();
+					//check( Query.GetRefCount() == 1 );
+					ViewState->OcclusionQueryPool.ReleaseQuery(RHICmdList, QueryIt.Value());
+				}
+				ShadowOcclusionQueryMap.Reset();
 
-					for(int32 ShadowIndex = 0;ShadowIndex < VisibleLightInfo.AllProjectedShadows.Num();ShadowIndex++)
-					{
-						const FProjectedShadowInfo& ProjectedShadowInfo = *VisibleLightInfo.AllProjectedShadows[ShadowIndex];
+				{
+					SCOPED_DRAW_EVENT(RHICmdList, ShadowFrustumQueries);
 
-						if (ProjectedShadowInfo.DependentView && ProjectedShadowInfo.DependentView != &View)
-						{
-							continue;
-						}
+					for (TSparseArray<FLightSceneInfoCompact>::TConstIterator LightIt(Scene->Lights); LightIt; ++LightIt)
+					{
+						const FVisibleLightInfo& VisibleLightInfo = VisibleLightInfos[LightIt.GetIndex()];
 
-						if (ProjectedShadowInfo.CascadeSettings.bOnePassPointLightShadow)
+						for (int32 ShadowIndex = 0; ShadowIndex < VisibleLightInfo.AllProjectedShadows.Num(); ShadowIndex++)
 						{
-							FLightSceneProxy& LightProxy = *(ProjectedShadowInfo.GetLightSceneInfo().Proxy);
+							const FProjectedShadowInfo& ProjectedShadowInfo = *VisibleLightInfo.AllProjectedShadows[ShadowIndex];
 
-							// Query one pass point light shadows separately because they don't have a shadow frustum, they have a bounding sphere instead.
-							FSphere LightBounds = LightProxy.GetBoundingSphere();
+							if (ProjectedShadowInfo.DependentView && ProjectedShadowInfo.DependentView != &View)
+							{
+								continue;
+							}
 
-							const bool bCameraInsideLightGeometry = ((FVector)View.ViewMatrices.ViewOrigin - LightBounds.Center).SizeSquared() < FMath::Square(LightBounds.W * 1.05f + View.NearClippingDistance * 2.0f);
-							if (!bCameraInsideLightGeometry)
+							if (ProjectedShadowInfo.CascadeSettings.bOnePassPointLightShadow)
 							{
-								const FRenderQueryRHIRef ShadowOcclusionQuery = ViewState->OcclusionQueryPool.AllocateQuery();
-								RHICmdList.BeginRenderQuery(ShadowOcclusionQuery);
+								FLightSceneProxy& LightProxy = *(ProjectedShadowInfo.GetLightSceneInfo().Proxy);
+
+								// Query one pass point light shadows separately because they don't have a shadow frustum, they have a bounding sphere instead.
+								FSphere LightBounds = LightProxy.GetBoundingSphere();
 
-								FSceneViewState::FProjectedShadowKey Key(ProjectedShadowInfo);
+								const bool bCameraInsideLightGeometry = ((FVector)View.ViewMatrices.ViewOrigin - LightBounds.Center).SizeSquared() < FMath::Square(LightBounds.W * 1.05f + View.NearClippingDistance * 2.0f);
+								if (!bCameraInsideLightGeometry)
+								{
+									const FRenderQueryRHIRef ShadowOcclusionQuery = ViewState->OcclusionQueryPool.AllocateQuery();
+									RHICmdList.BeginRenderQuery(ShadowOcclusionQuery);
 
-								checkSlow(ShadowOcclusionQueryMap.Find(Key) == NULL);
-								ShadowOcclusionQueryMap.Add(Key, ShadowOcclusionQuery);
+									FSceneViewState::FProjectedShadowKey Key(ProjectedShadowInfo);
 
-								// Draw bounding sphere
-								VertexShader->SetParametersWithBoundingSphere(RHICmdList, View, LightBounds);
-								StencilingGeometry::DrawSphere(RHICmdList);
+									checkSlow(ShadowOcclusionQueryMap.Find(Key) == NULL);
+									ShadowOcclusionQueryMap.Add(Key, ShadowOcclusionQuery);
 
-								RHICmdList.EndRenderQuery(ShadowOcclusionQuery);
+									// Draw bounding sphere
+									VertexShader->SetParametersWithBoundingSphere(RHICmdList, View, LightBounds);
+									StencilingGeometry::DrawSphere(RHICmdList);
+
+									RHICmdList.EndRenderQuery(ShadowOcclusionQuery);
+								}
+							}
+							// Don't query preshadows, since they are culled if their subject is occluded.
+							// Don't query if any subjects are visible because the shadow frustum will be definitely unoccluded
+							else if (!ProjectedShadowInfo.IsWholeSceneDirectionalShadow() && !ProjectedShadowInfo.bPreShadow && !ProjectedShadowInfo.SubjectsVisible(View))
+							{
+								IssueProjectedShadowOcclusionQuery(RHICmdList, View, ProjectedShadowInfo, *VertexShader);
 							}
 						}
-						// Don't query preshadows, since they are culled if their subject is occluded.
-						// Don't query if any subjects are visible because the shadow frustum will be definitely unoccluded
-						else if (!ProjectedShadowInfo.IsWholeSceneDirectionalShadow() && !ProjectedShadowInfo.bPreShadow && !ProjectedShadowInfo.SubjectsVisible(View))
+
+						// Issue occlusion queries for all per-object projected shadows that we would have rendered but were occluded last frame.
+						for (int32 ShadowIndex = 0; ShadowIndex < VisibleLightInfo.OccludedPerObjectShadows.Num(); ShadowIndex++)
 						{
+							const FProjectedShadowInfo& ProjectedShadowInfo = *VisibleLightInfo.OccludedPerObjectShadows[ShadowIndex];
 							IssueProjectedShadowOcclusionQuery(RHICmdList, View, ProjectedShadowInfo, *VertexShader);
 						}
 					}
-
-					// Issue occlusion queries for all per-object projected shadows that we would have rendered but were occluded last frame.
-					for(int32 ShadowIndex = 0;ShadowIndex < VisibleLightInfo.OccludedPerObjectShadows.Num();ShadowIndex++)
-					{
-						const FProjectedShadowInfo& ProjectedShadowInfo = *VisibleLightInfo.OccludedPerObjectShadows[ShadowIndex];
-						IssueProjectedShadowOcclusionQuery(RHICmdList, View, ProjectedShadowInfo, *VertexShader);
-					}
 				}
-			}
-	
-			// Don't do primitive occlusion if we have a view parent or are frozen.
-	#if !(UE_BUILD_SHIPPING || UE_BUILD_TEST)
-			if ( !ViewState->HasViewParent() && !ViewState->bIsFrozen )
-	#endif
-			{
-				VertexShader->SetParameters(RHICmdList, View);
 
+				// Don't do primitive occlusion if we have a view parent or are frozen.
+#if !(UE_BUILD_SHIPPING || UE_BUILD_TEST)
+				if (!ViewState->HasViewParent() && !ViewState->bIsFrozen)
+#endif
 				{
-					SCOPED_DRAW_EVENT(RHICmdList, IndividualQueries);
-					View.IndividualOcclusionQueries.Flush(RHICmdList);
-				}
-				{
-					SCOPED_DRAW_EVENT(RHICmdList, GroupedQueries);
-					View.GroupedOcclusionQueries.Flush(RHICmdList);
+					VertexShader->SetParameters(RHICmdList, View);
+
+					{
+						SCOPED_DRAW_EVENT(RHICmdList, IndividualQueries);
+						View.IndividualOcclusionQueries.Flush(RHICmdList);
+					}
+					{
+						SCOPED_DRAW_EVENT(RHICmdList, GroupedQueries);
+						View.GroupedOcclusionQueries.Flush(RHICmdList);
+					}
 				}
 			}
 		}
-	}
 
-	RHICmdList.EndOcclusionQueryBatch();
+		RHICmdList.EndOcclusionQueryBatch();
+	}
 
-	for( int32 ViewIndex = 0; ViewIndex < Views.Num(); ViewIndex++ )
+	if (bRenderHZB)
 	{
-		FViewInfo& View = Views[ViewIndex];
-		FSceneViewState* ViewState = (FSceneViewState*)View.State;
-		
-		if( ViewState && ViewState->HZBOcclusionTests.GetNum() != 0 )
+		for (int32 ViewIndex = 0; ViewIndex < Views.Num(); ViewIndex++)
 		{
-			check( ViewState->HZBOcclusionTests.IsValidFrame(ViewState->OcclusionFrameCounter) );
+			FViewInfo& View = Views[ViewIndex];
+			FSceneViewState* ViewState = (FSceneViewState*)View.State;
 
-			SCOPED_DRAW_EVENT(RHICmdList, HZB);
-			ViewState->HZBOcclusionTests.Submit(RHICmdList, View);
+			if (ViewState && ViewState->HZBOcclusionTests.GetNum() != 0)
+			{
+				check(ViewState->HZBOcclusionTests.IsValidFrame(ViewState->OcclusionFrameCounter));
+
+				SCOPED_DRAW_EVENT(RHICmdList, HZB);
+				ViewState->HZBOcclusionTests.Submit(RHICmdList, View);
+			}
 		}
 	}
 
