Ensure the postprocess material node always falls back to the default postprocess material.

With existing code if you pass a material instance with no parent in it will always default to the default surface material.

diff --git a/Engine/Source/Runtime/Renderer/Private/PostProcess/PostProcessMaterial.cpp b/Engine/Source/Runtime/Renderer/Private/PostProcess/PostProcessMaterial.cpp
--- a/Engine/Source/Runtime/Renderer/Private/PostProcess/PostProcessMaterial.cpp
+++ b/Engine/Source/Runtime/Renderer/Private/PostProcess/PostProcessMaterial.cpp
@@ -92,9 +92,18 @@ private:
 
 IMPLEMENT_MATERIAL_SHADER_TYPE(,FPostProcessMaterialPS,TEXT("PostProcessMaterialShaders"),TEXT("MainPS"),SF_Pixel);
 
-FRCPassPostProcessMaterial::FRCPassPostProcessMaterial(UMaterialInterface* InMaterialInterface, EPixelFormat OutputFormatIN)
+FRCPassPostProcessMaterial::FRCPassPostProcessMaterial(UMaterialInterface* InMaterialInterface, ERHIFeatureLevel::Type InFeatureLevel, EPixelFormat OutputFormatIN)
 : MaterialInterface(InMaterialInterface), OutputFormat(OutputFormatIN)
 {
+	FMaterialRenderProxy* Proxy = MaterialInterface->GetRenderProxy(false);
+	check(Proxy);
+
+	const FMaterial* Material = Proxy->GetMaterialNoFallback(InFeatureLevel);
+	
+	if (!Material || Material->GetMaterialDomain() != MD_PostProcess)
+	{
+		MaterialInterface = UMaterial::GetDefaultMaterial(MD_PostProcess);
+	}
 }
 
 void FRCPassPostProcessMaterial::Process(FRenderingCompositePassContext& Context)
diff --git a/Engine/Source/Runtime/Renderer/Private/PostProcess/PostProcessMaterial.h b/Engine/Source/Runtime/Renderer/Private/PostProcess/PostProcessMaterial.h
--- a/Engine/Source/Runtime/Renderer/Private/PostProcess/PostProcessMaterial.h
+++ b/Engine/Source/Runtime/Renderer/Private/PostProcess/PostProcessMaterial.h
@@ -16,7 +16,7 @@ class FRCPassPostProcessMaterial : public TRenderingCompositePassBase<3,1>
 {
 public:
 	// constructor
-	FRCPassPostProcessMaterial(UMaterialInterface* InMaterialInterface, EPixelFormat OutputFormatIN = PF_Unknown);
+	FRCPassPostProcessMaterial(UMaterialInterface* InMaterialInterface, ERHIFeatureLevel::Type InFeatureLevel, EPixelFormat OutputFormatIN = PF_Unknown);
 
 	// interface FRenderingCompositePass ---------
 
diff --git a/Engine/Source/Runtime/Renderer/Private/PostProcess/PostProcessing.cpp b/Engine/Source/Runtime/Renderer/Private/PostProcess/PostProcessing.cpp
--- a/Engine/Source/Runtime/Renderer/Private/PostProcess/PostProcessing.cpp
+++ b/Engine/Source/Runtime/Renderer/Private/PostProcess/PostProcessing.cpp
@@ -693,7 +693,7 @@ static FRenderingCompositePass* AddSinglePostProcessMaterial(FPostprocessContext
 			GSceneRenderTargets.AdjustGBufferRefCount(1);
 		}
 
-		FRenderingCompositePass* Node = Context.Graph.RegisterPass(new(FMemStack::Get()) FRCPassPostProcessMaterial(MaterialInterface));
+		FRenderingCompositePass* Node = Context.Graph.RegisterPass(new(FMemStack::Get()) FRCPassPostProcessMaterial(MaterialInterface, Context.View.GetFeatureLevel()));
 
 		return Node;
 	}
@@ -741,6 +741,8 @@ static void AddPostProcessMaterial(FPostprocessContext& Context, EBlendableLocat
 	
 	::Sort(PPNodes, PPNodeCount, FPostProcessMaterialNode::FCompare());
 
+	ERHIFeatureLevel::Type FeatureLevel = Context.View.GetFeatureLevel();
+
 	for(uint32 i = 0; i < PPNodeCount; ++i)
 	{
 		UMaterialInterface* MaterialInterface = PPNodes[i].GetMaterialInterface();
@@ -759,7 +761,7 @@ static void AddPostProcessMaterial(FPostprocessContext& Context, EBlendableLocat
 			GSceneRenderTargets.AdjustGBufferRefCount(1);
 		}
 
-		FRenderingCompositePass* Node = Context.Graph.RegisterPass(new(FMemStack::Get()) FRCPassPostProcessMaterial(MaterialInterface));
+		FRenderingCompositePass* Node = Context.Graph.RegisterPass(new(FMemStack::Get()) FRCPassPostProcessMaterial(MaterialInterface,FeatureLevel));
 		Node->SetInput(ePId_Input0, FRenderingCompositeOutputRef(Context.FinalOutput));
 
 		// We are binding separate translucency here because the post process SceneTexture node can reference 
@@ -778,7 +780,7 @@ static void AddHighResScreenshotMask(FPostprocessContext& Context, FRenderingCom
 
 		FRenderingCompositeOutputRef Input = Context.FinalOutput;
 
-		FRenderingCompositePass* CompositePass = Context.Graph.RegisterPass(new(FMemStack::Get()) FRCPassPostProcessMaterial(Context.View.FinalPostProcessSettings.HighResScreenshotMaterial));
+		FRenderingCompositePass* CompositePass = Context.Graph.RegisterPass(new(FMemStack::Get()) FRCPassPostProcessMaterial(Context.View.FinalPostProcessSettings.HighResScreenshotMaterial, Context.View.GetFeatureLevel()));
 		CompositePass->SetInput(ePId_Input0, FRenderingCompositeOutputRef(Input));
 		Context.FinalOutput = FRenderingCompositeOutputRef(CompositePass);
 
@@ -786,7 +788,7 @@ static void AddHighResScreenshotMask(FPostprocessContext& Context, FRenderingCom
 		{
 			check(Context.View.FinalPostProcessSettings.HighResScreenshotMaskMaterial); 
 
-			FRenderingCompositePass* MaskPass = Context.Graph.RegisterPass(new(FMemStack::Get()) FRCPassPostProcessMaterial(Context.View.FinalPostProcessSettings.HighResScreenshotMaskMaterial));
+			FRenderingCompositePass* MaskPass = Context.Graph.RegisterPass(new(FMemStack::Get()) FRCPassPostProcessMaterial(Context.View.FinalPostProcessSettings.HighResScreenshotMaskMaterial, Context.View.GetFeatureLevel()));
 			MaskPass->SetInput(ePId_Input0, FRenderingCompositeOutputRef(Input));
 			CompositePass->AddDependency(MaskPass);
 
@@ -800,7 +802,7 @@ static void AddHighResScreenshotMask(FPostprocessContext& Context, FRenderingCom
 	{
 		auto Material = Context.View.FinalPostProcessSettings.HighResScreenshotCaptureRegionMaterial;
 
-		FRenderingCompositePass* CaptureRegionVisualizationPass = Context.Graph.RegisterPass(new(FMemStack::Get()) FRCPassPostProcessMaterial(Material));
+		FRenderingCompositePass* CaptureRegionVisualizationPass = Context.Graph.RegisterPass(new(FMemStack::Get()) FRCPassPostProcessMaterial(Material, Context.View.GetFeatureLevel()));
 		CaptureRegionVisualizationPass->SetInput(ePId_Input0, FRenderingCompositeOutputRef(Context.FinalOutput));
 		Context.FinalOutput = FRenderingCompositeOutputRef(CaptureRegionVisualizationPass);
 
@@ -849,7 +851,7 @@ static void AddGBufferVisualizationOverview(FPostprocessContext& Context, FRende
 				if (MaterialInterface)
 				{
 					// Apply requested material
-					FRenderingCompositePass* MaterialPass = Context.Graph.RegisterPass(new(FMemStack::Get()) FRCPassPostProcessMaterial(*It, OutputFormat));
+					FRenderingCompositePass* MaterialPass = Context.Graph.RegisterPass(new(FMemStack::Get()) FRCPassPostProcessMaterial(*It, Context.View.GetFeatureLevel(), OutputFormat));
 					MaterialPass->SetInput(ePId_Input0, FRenderingCompositeOutputRef(IncomingStage));
 					MaterialPass->SetInput(ePId_Input1, FRenderingCompositeOutputRef(SeparateTranslucencyInput));
 
