Fixed compile error in DistanceFieldObjectManagement.cpp

diff --git a/Engine/Source/Runtime/Renderer/Private/DistanceFieldObjectManagement.cpp b/Engine/Source/Runtime/Renderer/Private/DistanceFieldObjectManagement.cpp
--- a/Engine/Source/Runtime/Renderer/Private/DistanceFieldObjectManagement.cpp
+++ b/Engine/Source/Runtime/Renderer/Private/DistanceFieldObjectManagement.cpp
@@ -735,7 +735,7 @@ void FDeferredShadingSceneRenderer::UpdateGlobalDistanceFieldObjectBuffers(FRHIC
 
 		// Process removes before adds, as the adds will overwrite primitive allocation info in DistanceFieldSceneData.SurfelAllocations
 		UpdateGlobalDistanceFieldObjectRemoves(RHICmdList, Scene);
-
+		extern int32 GVPLMeshGlobalIllumination;
 		TArray<uint32> UploadObjectIndices;
 		TArray<FVector4> UploadObjectData;
 		const bool bPrepareForDistanceFieldGI = GVPLMeshGlobalIllumination && SupportsDistanceFieldGI(Scene->GetFeatureLevel(), Scene->GetShaderPlatform());
@@ -756,7 +756,7 @@ void FDeferredShadingSceneRenderer::UpdateGlobalDistanceFieldObjectBuffers(FRHIC
 			int32 OriginalNumObjects = DistanceFieldSceneData.NumObjectsInBuffer;
 			int32 OriginalNumSurfels = DistanceFieldSceneData.SurfelAllocations.GetNumSurfelsInBuffer();
 			int32 OriginalNumInstancedSurfels = DistanceFieldSceneData.InstancedSurfelAllocations.GetNumSurfelsInBuffer();
-			extern int32 GVPLMeshGlobalIllumination;
+			
 
 			if (bPrepareForDistanceFieldGI)
 			{
@@ -877,8 +877,7 @@ void FDeferredShadingSceneRenderer::UpdateGlobalDistanceFieldObjectBuffers(FRHIC
 					{
 						FPrimitiveSurfelAllocation Allocation;
 						FPrimitiveSurfelAllocation InstancedAllocation;
-						extern int32 GVPLMeshGlobalIllumination;
-
+						
 						if (bPrepareForDistanceFieldGI)
 						{
 							const FPrimitiveSurfelAllocation* AllocationPtr = Scene->DistanceFieldSceneData.SurfelAllocations.FindAllocation(PrimitiveSceneInfo);
