Added decals to ShaderComplexity rendering (decalmask is not yet considered - it would require the GBuffer which isn't setup for that view mode)

diff --git a/Engine/Shaders/DeferredDecal.usf b/Engine/Shaders/DeferredDecal.usf
--- a/Engine/Shaders/DeferredDecal.usf
+++ b/Engine/Shaders/DeferredDecal.usf
@@ -22,7 +22,8 @@ float4x4 ScreenToDecal;
 float4x4 DecalToWorld;
 // from component to clip space (for decal frustum)
 float4x4 FrustumComponentToClip;
-
+// only for StencilDecalMaskMain, to get ShaderComplexity for Decals (doesn't look at the DecalMask yet)
+float MaskComparison;
 
 void StencilDecalMaskMain(
 	FScreenVertexOutput Input,
@@ -32,7 +33,7 @@ void StencilDecalMaskMain(
 	FGBufferData InGBufferData = GetGBufferData(Input.UV);
     
 	// Write 1 to the stencil buffer for pixels that are masked out of receiving decals
-	clip(.5f - InGBufferData.DecalMask);
+	clip(MaskComparison - InGBufferData.DecalMask);
 
 	OutColor = 0;
 }
diff --git a/Engine/Shaders/ShaderComplexityAccumulatePixelShader.usf b/Engine/Shaders/ShaderComplexityAccumulatePixelShader.usf
--- a/Engine/Shaders/ShaderComplexityAccumulatePixelShader.usf
+++ b/Engine/Shaders/ShaderComplexityAccumulatePixelShader.usf
@@ -10,6 +10,7 @@ float NormalizedComplexity;
 
 void Main(out float4 OutColor : SV_Target0)
 {
-	//use the maximum range allowed for scene color
-	OutColor = RETURN_COLOR(float4(NormalizedComplexity.xxx, 0));
+	// use the maximum range allowed for scene color
+	// alpha channel needs to be 1 to have decals working where the framebuffer blend is setup to multiply with alpha
+	OutColor = RETURN_COLOR(float4(NormalizedComplexity.xxx, 1));
 }
diff --git a/Engine/Shaders/ShaderVersion.usf b/Engine/Shaders/ShaderVersion.usf
--- a/Engine/Shaders/ShaderVersion.usf
+++ b/Engine/Shaders/ShaderVersion.usf
@@ -2,4 +2,4 @@
 // Each time the console command is executed it generates a new GUID. As this file is included
 // in common.usf (which should be included in any shader) it allows to invalidate the shader DDC.
 // 
-// GUID = {2490AEE3-37E5-4684-93D8-F170F910238D}
+// GUID = 7D700E574A85DE267C385C9E997BCC23
\ No newline at end of file
diff --git a/Engine/Source/Runtime/Renderer/Private/CompositionLighting/PostProcessDeferredDecals.cpp b/Engine/Source/Runtime/Renderer/Private/CompositionLighting/PostProcessDeferredDecals.cpp
--- a/Engine/Source/Runtime/Renderer/Private/CompositionLighting/PostProcessDeferredDecals.cpp
+++ b/Engine/Source/Runtime/Renderer/Private/CompositionLighting/PostProcessDeferredDecals.cpp
@@ -339,6 +339,7 @@ public:
 		FGlobalShader(Initializer)
 	{
 		DeferredParameters.Bind(Initializer.ParameterMap);
+		MaskComparison.Bind(Initializer.ParameterMap,TEXT("MaskComparison"));
 	}
 	FStencilDecalMaskPS() {}
 
@@ -348,17 +349,19 @@ public:
 
 		FGlobalShader::SetParameters(RHICmdList, ShaderRHI, View);
 		DeferredParameters.Set(RHICmdList, ShaderRHI, View);
+		SetShaderValue(RHICmdList, ShaderRHI, MaskComparison, View.Family->EngineShowFlags.ShaderComplexity ? -1.0f : 0.5f);
 	}
 
 	virtual bool Serialize(FArchive& Ar)
 	{
 		bool bShaderHasOutdatedParameters = FGlobalShader::Serialize(Ar);
-		Ar << DeferredParameters;
+		Ar << DeferredParameters << MaskComparison;
 		return bShaderHasOutdatedParameters;
 	}
 
 private:
 	FDeferredPixelShaderParameters DeferredParameters;
+	FShaderParameter MaskComparison;
 };
 
 IMPLEMENT_SHADER_TYPE(,FStencilDecalMaskPS,TEXT("DeferredDecal"),TEXT("StencilDecalMaskMain"),SF_Pixel);
@@ -401,37 +404,34 @@ void StencilDecalMask(FRHICommandList& RHICmdList, const FViewInfo& View)
 /**
  * A vertex shader for projecting a deferred decal onto the scene.
  */
-class FDeferredDecalVS : public FMaterialShader
+class FDeferredDecalVS : public FGlobalShader
 {
-	DECLARE_SHADER_TYPE(FDeferredDecalVS,Material);
+	DECLARE_SHADER_TYPE(FDeferredDecalVS,Global);
 public:
 
-	/**
-	  * Makes sure only shaders for materials that are explicitly flagged
-	  * as 'UsedAsDeferredDecal' in the Material Editor gets compiled into
-	  * the shader cache.
-	  */
-	static bool ShouldCache(EShaderPlatform Platform, const FMaterial* Material)
+	static bool ShouldCache(EShaderPlatform Platform)
 	{
-		return Material->IsUsedWithDeferredDecal() && IsFeatureLevelSupported(Platform, ERHIFeatureLevel::SM4);
+		return IsFeatureLevelSupported(Platform, ERHIFeatureLevel::SM4);
 	}
 
 	FDeferredDecalVS( )	{ }
 	FDeferredDecalVS(const ShaderMetaType::CompiledShaderInitializerType& Initializer)
-		: FMaterialShader(Initializer)
+		: FGlobalShader(Initializer)
 	{
 		FrustumComponentToClip.Bind(Initializer.ParameterMap, TEXT("FrustumComponentToClip"));
 	}
 
 	void SetParameters(FRHICommandList& RHICmdList, const FSceneView& View, const FMatrix& InFrustumComponentToClip)
 	{
-		FMaterialShader::SetParameters(RHICmdList, GetVertexShader(), View);
-		SetShaderValue(RHICmdList, GetVertexShader(), FrustumComponentToClip, InFrustumComponentToClip);
+		const FVertexShaderRHIParamRef ShaderRHI = GetVertexShader();
+
+		FGlobalShader::SetParameters(RHICmdList, ShaderRHI, View);
+		SetShaderValue(RHICmdList, ShaderRHI, FrustumComponentToClip, InFrustumComponentToClip);
 	}
 
 	virtual bool Serialize(FArchive& Ar)
 	{
-		bool bShaderHasOutdatedParameters = FMaterialShader::Serialize(Ar);
+		bool bShaderHasOutdatedParameters = FGlobalShader::Serialize(Ar);
 		Ar << FrustumComponentToClip;
 		return bShaderHasOutdatedParameters;
 	}
@@ -440,7 +440,7 @@ private:
 	FShaderParameter FrustumComponentToClip;
 };
 
-IMPLEMENT_MATERIAL_SHADER_TYPE(,FDeferredDecalVS,TEXT("DeferredDecal"),TEXT("MainVS"),SF_Vertex);
+IMPLEMENT_SHADER_TYPE(,FDeferredDecalVS,TEXT("DeferredDecal"),TEXT("MainVS"),SF_Vertex);
 
 /**
  * A pixel shader for projecting a deferred decal onto the scene.
@@ -620,16 +620,34 @@ private:
 static TGlobalResource<FUnitCubeVertexBuffer> GUnitCubeVertexBuffer;
 static TGlobalResource<FUnitCubeIndexBuffer> GUnitCubeIndexBuffer;
 
-void SetShader(const FRenderingCompositePassContext& Context, const FTransientDecalRenderData& DecalData, FShader* VertexShader)
+void SetShader(const FRenderingCompositePassContext& Context, bool bShaderComplexity, const FTransientDecalRenderData& DecalData, FShader* VertexShader)
 {
 	const FSceneView& View = Context.View;
-	
+
 	const FMaterialShaderMap* MaterialShaderMap = DecalData.MaterialResource->GetRenderingThreadShaderMap();
-	FDeferredDecalPS* PixelShader = MaterialShaderMap->GetShader<FDeferredDecalPS>();
 
-	Context.RHICmdList.SetLocalBoundShaderState(Context.RHICmdList.BuildLocalBoundShaderState(GetVertexDeclarationFVector3(), VertexShader->GetVertexShader(), FHullShaderRHIRef(), FDomainShaderRHIRef(), PixelShader->GetPixelShader(), FGeometryShaderRHIRef()));
+	auto PixelShader = MaterialShaderMap->GetShader<FDeferredDecalPS>();
 
-	PixelShader->SetParameters(Context.RHICmdList, View, DecalData.MaterialProxy, *DecalData.DecalProxy);
+	if(bShaderComplexity)
+	{
+		const FViewInfo &ViewInfo = (const FViewInfo &)View;
+
+		TShaderMapRef<FShaderComplexityAccumulatePS> VisualizePixelShader(ViewInfo.ShaderMap);
+		const uint32 NumPixelShaderInstructions = PixelShader->GetNumInstructions();
+		const uint32 NumVertexShaderInstructions = VertexShader->GetNumInstructions();
+
+		static FGlobalBoundShaderState BoundShaderState;
+		SetGlobalBoundShaderState(Context.RHICmdList, Context.GetFeatureLevel(), BoundShaderState, GetVertexDeclarationFVector3(), VertexShader, *VisualizePixelShader);
+
+		VisualizePixelShader->SetParameters(Context.RHICmdList, NumVertexShaderInstructions, NumPixelShaderInstructions, View.GetFeatureLevel());
+	}
+	else
+	{
+		// first Bind, then SetParameters()
+		Context.RHICmdList.SetLocalBoundShaderState(Context.RHICmdList.BuildLocalBoundShaderState(GetVertexDeclarationFVector3(), VertexShader->GetVertexShader(), FHullShaderRHIRef(), FDomainShaderRHIRef(), PixelShader->GetPixelShader(), FGeometryShaderRHIRef()));
+
+		PixelShader->SetParameters(Context.RHICmdList, View, DecalData.MaterialProxy, *DecalData.DecalProxy);
+ 	}
 }
 
 bool RenderPreStencil(FRenderingCompositePassContext& Context, const FMaterialShaderMap* MaterialShaderMap, const FMatrix& ComponentToWorldMatrix, const FMatrix& FrustumComponentToClip)
@@ -653,7 +671,7 @@ bool RenderPreStencil(FRenderingCompositePassContext& Context, const FMaterialSh
 		}
 	}
 
-	FDeferredDecalVS* VertexShader = MaterialShaderMap->GetShader<FDeferredDecalVS>();
+	TShaderMapRef<FDeferredDecalVS> VertexShader(Context.GetShaderMap());
 	
 	Context.RHICmdList.SetLocalBoundShaderState(Context.RHICmdList.BuildLocalBoundShaderState(GetVertexDeclarationFVector3(), VertexShader->GetVertexShader(), FHullShaderRHIRef(), FDomainShaderRHIRef(), NULL, FGeometryShaderRHIRef()));
 
@@ -699,11 +717,11 @@ FRCPassPostProcessDeferredDecals::FRCPassPostProcessDeferredDecals(uint32 InRend
 
 void FRCPassPostProcessDeferredDecals::Process(FRenderingCompositePassContext& Context)
 {
-	
 	FRHICommandListImmediate& RHICmdList = Context.RHICmdList;
 
-	bool bDBuffer = IsDBufferEnabled();
-	float bDecalPreStencil = CVarStencilSizeThreshold.GetValueOnRenderThread() >= 0;
+	const bool bShaderComplexity = Context.View.Family->EngineShowFlags.ShaderComplexity;
+	const bool bDBuffer = IsDBufferEnabled();
+	const bool bDecalPreStencil = CVarStencilSizeThreshold.GetValueOnRenderThread() >= 0;
 
 	SCOPED_DRAW_EVENT(RHICmdList, PostProcessDeferredDecals);
 
@@ -920,6 +938,13 @@ void FRCPassPostProcessDeferredDecals::Process(FRenderingCompositePassContext& C
 
 			ERenderTargetMode CurrentRenderTargetMode = ComputeRenderTargetMode(DecalBlendMode);
 
+			if(bShaderComplexity)
+			{
+				CurrentRenderTargetMode = RTM_SceneColor;
+				// we want additive blending for the ShaderComplexity mode
+				DecalBlendMode = DBM_Emissive;
+			}
+
 			// fewer rendertarget switches if possible
 			if(CurrentRenderTargetMode != LastRenderTargetMode)
 			{
@@ -982,8 +1007,8 @@ void FRCPassPostProcessDeferredDecals::Process(FRenderingCompositePassContext& C
 			RHICmdList.SetStreamSource(0, GUnitCubeVertexBuffer.VertexBufferRHI, sizeof(FVector4), 0);
 
 			{
-				FDeferredDecalVS* VertexShader = MaterialShaderMap->GetShader<FDeferredDecalVS>();
-				SetShader(Context, DecalData, VertexShader);
+				TShaderMapRef<FDeferredDecalVS> VertexShader(Context.GetShaderMap());
+				SetShader(Context, bShaderComplexity, DecalData, *VertexShader);
 
 				VertexShader->SetParameters(RHICmdList, View, FrustumComponentToClip);
 
@@ -1057,8 +1082,6 @@ void FRCPassPostProcessDeferredDecals::Process(FRenderingCompositePassContext& C
 					}
 				}
 
-				SetShader(Context, DecalData, VertexShader);
-
 				RHICmdList.DrawIndexedPrimitive(GUnitCubeIndexBuffer.IndexBufferRHI, PT_TriangleList, 0, 0, 8, 0, GUnitCubeIndexBuffer.GetIndexCount() / 3, 1);
 			}
 		}
