UE4 - Added sanity checks for Inverted/Normal Depth Buffers; use ERHIZBuffer::NearPlane & FarPlane to switch from 0-1 to 1-0 on Depth
- Added shader define HAS_INVERTED_Z_BUFFER
- Added RHIHasInvertedZBuffer()

diff --git a/Engine/Shaders/BasePassPixelShader.usf b/Engine/Shaders/BasePassPixelShader.usf
--- a/Engine/Shaders/BasePassPixelShader.usf
+++ b/Engine/Shaders/BasePassPixelShader.usf
@@ -534,7 +534,10 @@ float ClipForEditorPrimitives(FMaterialPixelParameters MaterialParameters)
 	BRANCH
 	if (bEnableEditorPrimitiveDepthTest)
 	{
+#if HAS_INVERTED_Z_BUFFER
+		//@todo-briank
 		bool bIsPerspective = (View.ViewToClip._m33 < 1.0f);
+#endif	// HAS_INVERTED_Z_BUFFER
  		
 		// dejitter the sample position and make a filtered lookup - for planes this allows to reconstruct a much less jittery depth comparison function, it however doesn't fix silhuetes
 		float DeviceZ = Texture2DSampleLevel(FilteredSceneDepthTexture, FilteredSceneDepthTextureSampler, (MaterialParameters.SVPosition.xy - View.TemporalAAParams.zw) * View.ViewSizeAndSceneTexelSize.zw, 0).r;
@@ -547,7 +550,7 @@ float ClipForEditorPrimitives(FMaterialPixelParameters MaterialParameters)
 		// 0.5f is to bias around the actual value, 1 or 0 are another option
 		Ret = saturate(0.5f - (DeviceZ - PixelDeviceZ) / DeviceDepthFade);
 	}
-#endif
+#endif	// USE_EDITOR_COMPOSITING && (FEATURE_LEVEL >= FEATURE_LEVEL_SM4 || ES2_EMULATION)
 
 	// Note: multiple returns cause strange HLSL compiler error for CV_Coverage in later code
 	return Ret;
diff --git a/Engine/Shaders/Definitions.usf b/Engine/Shaders/Definitions.usf
--- a/Engine/Shaders/Definitions.usf
+++ b/Engine/Shaders/Definitions.usf
@@ -194,3 +194,7 @@
 #else
 #define EARLYDEPTHSTENCIL
 #endif
+
+#ifndef HAS_INVERTED_Z_BUFFER
+#define HAS_INVERTED_Z_BUFFER	1
+#endif
diff --git a/Engine/Shaders/PostProcessTemporalCommon.usf b/Engine/Shaders/PostProcessTemporalCommon.usf
--- a/Engine/Shaders/PostProcessTemporalCommon.usf
+++ b/Engine/Shaders/PostProcessTemporalCommon.usf
@@ -153,6 +153,8 @@ which can be configured via defines for various temporal AA passes.
 		#endif
 		float2 DepthOffset = float2(AA_CROSS, AA_CROSS);
 		float DepthOffsetXx = float(AA_CROSS);
+#if HAS_INVERTED_Z_BUFFER
+		//@todo-briank
 		// Nearest depth is the largest depth (depth surface 0=far, 1=near).
 		if(Depths.x > Depths.y) 
 		{
@@ -180,7 +182,10 @@ which can be configured via defines for various temporal AA passes.
 			PosN.xy = (UVAndScreenPos.zw + DepthOffset * ViewportSize.zw * 2.0) * float2(0.5, -0.5) + 0.5;
 			PosN.z = DepthsXYZW;
 		}
-	#endif
+#else
+#error Fix me!
+#endif // HAS_INVERTED_Z_BUFFER
+	#endif	// AA_CROSS
 	// Camera motion for pixel or nearest pixel (in ScreenPos space).
 	float ScaleM = 1.0 / (dot(PosN, CameraMotion[0].xyz) + CameraMotion[0].w);
 	float2 BackN;
diff --git a/Engine/Shaders/SimpleElementPixelShader.usf b/Engine/Shaders/SimpleElementPixelShader.usf
--- a/Engine/Shaders/SimpleElementPixelShader.usf
+++ b/Engine/Shaders/SimpleElementPixelShader.usf
@@ -50,7 +50,11 @@ void PerformEditorCompositingDepthTest(float4 SvPosition)
 		float InterpolatedDeviceZ = SvPosition.z;
 
 		// This is for a reversed Z depth buffer.
+#if HAS_INVERTED_Z_BUFFER
 		clip(InterpolatedDeviceZ - SceneDepthDeviceZ);
+#else
+		clip(SceneDepthDeviceZ - InterpolatedDeviceZ);
+#endif
 	}
 }
 
diff --git a/Engine/Shaders/VisualizeTexture.usf b/Engine/Shaders/VisualizeTexture.usf
--- a/Engine/Shaders/VisualizeTexture.usf
+++ b/Engine/Shaders/VisualizeTexture.usf
@@ -93,10 +93,13 @@ void VisualizeTexturePS(
 	int InputValueMapping = (int)VisualizeParam[2].r;
 	if(InputValueMapping == 1)
 	{
-#if FEATURE_LEVEL >= FEATURE_LEVEL_SM4
 		// current z buffer is negated (1 is near, 0 is far)
-		float Depth = 1 - TexLookup.r;
+		float Depth = TexLookup.r;
+#if HAS_INVERTED_Z_BUFFER
+		Depth = 1 - Depth;
+#endif
 
+#if FEATURE_LEVEL >= FEATURE_LEVEL_SM4
 		// map depth to nice colors
 		float Scale = 64.0; // This 64 looks like a hardware limit.
 		Depth = (exp2(Depth * Scale) - 1.0) * (1.0 / (exp2(Scale) - 1.0));
@@ -111,7 +114,7 @@ void VisualizeTexturePS(
 
 		TexLookup.a = fStencil;
 #else
-		TexLookup.r = 1 - TexLookup.r;
+		TexLookup.r = Depth;
 #endif
 	}	
 
diff --git a/Engine/Source/Editor/UnrealEd/Private/EditorViewportClient.cpp b/Engine/Source/Editor/UnrealEd/Private/EditorViewportClient.cpp
--- a/Engine/Source/Editor/UnrealEd/Private/EditorViewportClient.cpp
+++ b/Engine/Source/Editor/UnrealEd/Private/EditorViewportClient.cpp
@@ -627,14 +627,28 @@ FSceneView* FEditorViewportClient::CalcSceneView(FSceneViewFamily* ViewFamily)
 
 			if (bConstrainAspectRatio)
 			{
-				ViewInitOptions.ProjectionMatrix = FReversedZPerspectiveMatrix(
-					MatrixFOV,
-					MatrixFOV,
-					1.0f,
-					AspectRatio,
-					MinZ,
-					MaxZ
-					);
+				if (RHIHasInvertedZBuffer())
+				{
+					ViewInitOptions.ProjectionMatrix = FReversedZPerspectiveMatrix(
+						MatrixFOV,
+						MatrixFOV,
+						1.0f,
+						AspectRatio,
+						MinZ,
+						MaxZ
+						);
+				}
+				else
+				{
+					ViewInitOptions.ProjectionMatrix = FPerspectiveMatrix(
+						MatrixFOV,
+						MatrixFOV,
+						1.0f,
+						AspectRatio,
+						MinZ,
+						MaxZ
+						);
+				}
 			}
 			else
 			{
@@ -654,18 +668,33 @@ FSceneView* FEditorViewportClient::CalcSceneView(FSceneViewFamily* ViewFamily)
 					YAxisMultiplier = 1.0f;
 				}
 
-				ViewInitOptions.ProjectionMatrix = FReversedZPerspectiveMatrix(
-					MatrixFOV,
-					MatrixFOV,
-					XAxisMultiplier,
-					YAxisMultiplier,
-					MinZ,
-					MaxZ
-					);
+				if (RHIHasInvertedZBuffer())
+				{
+					ViewInitOptions.ProjectionMatrix = FReversedZPerspectiveMatrix(
+						MatrixFOV,
+						MatrixFOV,
+						XAxisMultiplier,
+						YAxisMultiplier,
+						MinZ,
+						MaxZ
+						);
+				}
+				else
+				{
+					ViewInitOptions.ProjectionMatrix = FPerspectiveMatrix(
+						MatrixFOV,
+						MatrixFOV,
+						XAxisMultiplier,
+						YAxisMultiplier,
+						MinZ,
+						MaxZ
+						);
+				}
 			}
 		}
 		else
 		{
+			checkf(RHIHasInvertedZBuffer(), TEXT("Check all the Rotation Matrix transformations!"));
 			float ZScale = 0.5f / HALF_WORLD_MAX;
 			float ZOffset = HALF_WORLD_MAX;
 
diff --git a/Engine/Source/Editor/UnrealEd/Private/ThumbnailHelpers.cpp b/Engine/Source/Editor/UnrealEd/Private/ThumbnailHelpers.cpp
--- a/Engine/Source/Editor/UnrealEd/Private/ThumbnailHelpers.cpp
+++ b/Engine/Source/Editor/UnrealEd/Private/ThumbnailHelpers.cpp
@@ -70,6 +70,7 @@ void FThumbnailPreviewScene::GetView(FSceneViewFamily* ViewFamily, int32 X, int3
 	{
 		const float FOVDegrees = 30.f;
 		const float HalfFOVRadians = FMath::DegreesToRadians<float>(FOVDegrees) * 0.5f;
+		checkf(RHIHasInvertedZBuffer(), TEXT("Check NearPlane and Projection Matrix"));
 		const float NearPlane = 1.0f;
 		FMatrix ProjectionMatrix = FReversedZPerspectiveMatrix(
 			HalfFOVRadians,
diff --git a/Engine/Source/Runtime/Engine/Private/SceneView.cpp b/Engine/Source/Runtime/Engine/Private/SceneView.cpp
--- a/Engine/Source/Runtime/Engine/Private/SceneView.cpp
+++ b/Engine/Source/Runtime/Engine/Private/SceneView.cpp
@@ -426,6 +426,7 @@ FSceneView::FSceneView(const FSceneViewInitOptions& InitOptions)
 
 	// Derive the view's near clipping distance and plane.
 	// The GetFrustumFarPlane() is the near plane because of reverse Z projection.
+	checkf(RHIHasInvertedZBuffer(), TEXT("Fix Near Clip distance!"));
 	bHasNearClippingPlane = ViewProjectionMatrix.GetFrustumFarPlane(NearClippingPlane);
 	if(ViewMatrices.ProjMatrix.M[2][3] > DELTA)
 	{
diff --git a/Engine/Source/Runtime/Engine/Private/ShaderCompiler/ShaderCompiler.cpp b/Engine/Source/Runtime/Engine/Private/ShaderCompiler/ShaderCompiler.cpp
--- a/Engine/Source/Runtime/Engine/Private/ShaderCompiler/ShaderCompiler.cpp
+++ b/Engine/Source/Runtime/Engine/Private/ShaderCompiler/ShaderCompiler.cpp
@@ -2195,6 +2195,8 @@ void GlobalBeginCompileShader(
 		}
 	}
 
+	Input.Environment.SetDefine(TEXT("HAS_INVERTED_Z_BUFFER"), RHIHasInvertedZBuffer() ? 1 : 0);
+
 	{
 		FString ShaderPDBRoot;
 		GConfig->GetString(TEXT("DevOptions.Shaders"), TEXT("ShaderPDBRoot"), ShaderPDBRoot, GEngineIni);
diff --git a/Engine/Source/Runtime/Engine/Private/UserInterface/Canvas.cpp b/Engine/Source/Runtime/Engine/Private/UserInterface/Canvas.cpp
--- a/Engine/Source/Runtime/Engine/Private/UserInterface/Canvas.cpp
+++ b/Engine/Source/Runtime/Engine/Private/UserInterface/Canvas.cpp
@@ -331,14 +331,28 @@ FMatrix FCanvas::CalcProjectionMatrix(uint32 ViewSizeX, uint32 ViewSizeY, float
 	// convert FOV to randians
 	float FOVRad = fFOV * (float)PI / 360.0f;
 	// project based on the FOV and near plane given
-	return AdjustProjectionMatrixForRHI(
-		FReversedZPerspectiveMatrix(
-			FOVRad,
-			ViewSizeX,
-			ViewSizeY,
-			NearPlane
-			)
-		);
+	if (RHIHasInvertedZBuffer())
+	{
+		return AdjustProjectionMatrixForRHI(
+			FReversedZPerspectiveMatrix(
+				FOVRad,
+				ViewSizeX,
+				ViewSizeY,
+				NearPlane
+				)
+			);
+	}
+	else
+	{
+		return AdjustProjectionMatrixForRHI(
+			FPerspectiveMatrix(
+				FOVRad,
+				ViewSizeX,
+				ViewSizeY,
+				NearPlane
+				)
+			);
+	}
 }
 
 bool FCanvasBatchedElementRenderItem::Render_RenderThread(FRHICommandListImmediate& RHICmdList, const FCanvas* Canvas)
diff --git a/Engine/Source/Runtime/RHI/Private/RHI.cpp b/Engine/Source/Runtime/RHI/Private/RHI.cpp
--- a/Engine/Source/Runtime/RHI/Private/RHI.cpp
+++ b/Engine/Source/Runtime/RHI/Private/RHI.cpp
@@ -107,6 +107,10 @@ void FRHIResource::FlushPendingDeletes()
 }
 
 
+static_assert(ERHIZBuffer::FarPlane != ERHIZBuffer::NearPlane, "Near and Far planes must be different!");
+static_assert((int32)ERHIZBuffer::NearPlane == 0 || (int32)ERHIZBuffer::NearPlane == 1, "Invalid Values for Near Plane, can only be 0 or 1!");
+static_assert((int32)ERHIZBuffer::FarPlane == 0 || (int32)ERHIZBuffer::FarPlane == 1, "Invalid Values for Far Plane, can only be 0 or 1");
+
 
 /**
  * RHI configuration settings.
diff --git a/Engine/Source/Runtime/RHI/Public/RHIDefinitions.h b/Engine/Source/Runtime/RHI/Public/RHIDefinitions.h
--- a/Engine/Source/Runtime/RHI/Public/RHIDefinitions.h
+++ b/Engine/Source/Runtime/RHI/Public/RHIDefinitions.h
@@ -27,22 +27,22 @@ static_assert(SF_NumFrequencies <= (1 << SF_NumBits), "SF_NumFrequencies will no
 /** @warning: update *LegacyShaderPlatform* when the below changes */
 enum EShaderPlatform
 {
-	SP_PCD3D_SM5		= 0,
+	SP_PCD3D_SM5			= 0,
 	SP_OPENGL_SM4		= 1,
 	SP_PS4				= 2,
 	/** Used when running in Feature Level ES2 in OpenGL. */
 	SP_OPENGL_PCES2		= 3,
 	SP_XBOXONE			= 4,
-	SP_PCD3D_SM4		= 5,
+	SP_PCD3D_SM4			= 5,
 	SP_OPENGL_SM5		= 6,
 	/** Used when running in Feature Level ES2 in D3D11. */
-	SP_PCD3D_ES2		= 7,
+	SP_PCD3D_ES2			= 7,
 	SP_OPENGL_ES2		= 8,
 	SP_OPENGL_ES2_WEBGL = 9, 
 	SP_OPENGL_ES2_IOS	= 10,
-	SP_METAL			= 11,
+	SP_METAL				= 11,
 	SP_OPENGL_SM4_MAC	= 12,
-	SP_METAL_MRT		= 13,
+	SP_METAL_MRT			= 13,
 	SP_OPENGL_ES31_EXT	= 14,
 
 	SP_NumPlatforms		= 15,
@@ -75,6 +75,15 @@ enum { MaxSimultaneousRenderTargets = 8 };
 /** The number of UAVs that may be simultaneously bound to a shader. */
 enum { MaxSimultaneousUAVs = 8 };
 
+enum class ERHIZBuffer
+{
+	// Before changing this, make sure all math & shader assumptions are correct! Also wrap your C++ assumptions with
+	//		static_assert(FViewMatrices::IsInvertedZBuffer(), ...);
+	// Shader-wise, make sure to update Definitions.usf, HAS_INVERTED_Z_BUFFER
+	FarPlane = 0,
+	NearPlane = 1,
+};
+
 /**
  * The RHI's feature level indicates what level of support can be relied upon.
  * Note: these are named after graphics API's like ES2 but a feature level can be used with a different API (eg ERHIFeatureLevel::ES2 on D3D11)
@@ -157,7 +166,11 @@ enum ECompareFunction
 	CF_Equal,
 	CF_NotEqual,
 	CF_Never,
-	CF_Always
+	CF_Always,
+
+	CF_DepthFunction = ((int32)ERHIZBuffer::NearPlane < (int32)ERHIZBuffer::FarPlane) ? CF_LessEqual : CF_GreaterEqual,
+	CF_DepthFunctionNonInclusive = ((int32)ERHIZBuffer::NearPlane < (int32)ERHIZBuffer::FarPlane) ? CF_Less : CF_Greater,
+	CF_DepthFunctionReject = ((int32)ERHIZBuffer::NearPlane < (int32)ERHIZBuffer::FarPlane) ? CF_Greater : CF_Less,
 };
 
 enum EStencilOp
@@ -709,6 +722,11 @@ inline bool RHIHasTiledGPU(const EShaderPlatform Platform)
 	return Platform == SP_METAL_MRT|| Platform == SP_METAL || Platform == SP_OPENGL_ES2_IOS || Platform == SP_OPENGL_ES2;
 }
 
+static inline CONSTEXPR bool RHIHasInvertedZBuffer()
+{
+	return (int32)ERHIZBuffer::FarPlane < (int32)ERHIZBuffer::NearPlane;
+}
+
 inline uint32 GetFeatureLevelMaxTextureSamplers(ERHIFeatureLevel::Type FeatureLevel)
 {
 	if (FeatureLevel == ERHIFeatureLevel::ES2)
diff --git a/Engine/Source/Runtime/RHI/Public/RHIResources.h b/Engine/Source/Runtime/RHI/Public/RHIResources.h
--- a/Engine/Source/Runtime/RHI/Public/RHIResources.h
+++ b/Engine/Source/Runtime/RHI/Public/RHIResources.h
@@ -755,7 +755,7 @@ public:
 	FRHISetRenderTargetsInfo() :
 		NumColorRenderTargets(0),
 		bClearColor(false),
-		DepthClearValue(0.0f),
+		DepthClearValue((float)ERHIZBuffer::FarPlane),
 		StencilClearValue(0),
 		bClearDepth(false),
 		bClearStencil(false)
@@ -765,7 +765,7 @@ public:
 		NumColorRenderTargets(InNumColorRenderTargets),
 		bClearColor(InNumColorRenderTargets > 0 && InColorRenderTargets[0].LoadAction == ERenderTargetLoadAction::EClear),
 		DepthStencilRenderTarget(InDepthStencilRenderTarget),
-		DepthClearValue(0.0f),
+		DepthClearValue((float)ERHIZBuffer::FarPlane),
 		StencilClearValue(0),
 		bClearDepth(InDepthStencilRenderTarget.Texture && InDepthStencilRenderTarget.DepthLoadAction == ERenderTargetLoadAction::EClear),
 		bClearStencil(InDepthStencilRenderTarget.Texture && InDepthStencilRenderTarget.StencilLoadAction == ERenderTargetLoadAction::EClear)
diff --git a/Engine/Source/Runtime/RHI/Public/RHIStaticStates.h b/Engine/Source/Runtime/RHI/Public/RHIStaticStates.h
--- a/Engine/Source/Runtime/RHI/Public/RHIStaticStates.h
+++ b/Engine/Source/Runtime/RHI/Public/RHIStaticStates.h
@@ -220,7 +220,7 @@ FRasterizerStateRHIParamRef GetStaticRasterizerState(ERasterizerFillMode FillMod
  */
 template<
 	bool bEnableDepthWrite = true,
-	ECompareFunction DepthTest = CF_GreaterEqual, // Default for reversed Z depth surface (0=far, 1=near).
+	ECompareFunction DepthTest = CF_DepthFunction,
 	bool bEnableFrontFaceStencil = false,
 	ECompareFunction FrontFaceStencilTest = CF_Always,
 	EStencilOp FrontFaceStencilFailStencilOp = SO_Keep,
diff --git a/Engine/Source/Runtime/Renderer/Private/BasePassRendering.cpp b/Engine/Source/Runtime/Renderer/Private/BasePassRendering.cpp
--- a/Engine/Source/Runtime/Renderer/Private/BasePassRendering.cpp
+++ b/Engine/Source/Runtime/Renderer/Private/BasePassRendering.cpp
@@ -306,8 +306,7 @@ public:
 		// Then rendering in the base pass with additive complexity blending, depth tests on, and depth writes off.
 		if(View.Family->EngineShowFlags.ShaderComplexity)
 		{
-			// Note, this is a reversed Z depth surface, using CF_GreaterEqual.
-			RHICmdList.SetDepthStencilState(TStaticDepthStencilState<false,CF_GreaterEqual>::GetRHI());
+			RHICmdList.SetDepthStencilState(TStaticDepthStencilState<false,CF_DepthFunction>::GetRHI());
 		}
 #endif
 		const FScene* Scene = Parameters.PrimitiveSceneProxy ? Parameters.PrimitiveSceneProxy->GetPrimitiveSceneInfo()->Scene : NULL;
@@ -347,8 +346,7 @@ public:
 #if !(UE_BUILD_SHIPPING || UE_BUILD_TEST)
 		if(View.Family->EngineShowFlags.ShaderComplexity)
 		{
-			// Note, this is a reversed Z depth surface, using CF_GreaterEqual.
-			RHICmdList.SetDepthStencilState(TStaticDepthStencilState<true,CF_GreaterEqual>::GetRHI());
+			RHICmdList.SetDepthStencilState(TStaticDepthStencilState<true,CF_DepthFunction>::GetRHI());
 		}
 #endif
 	}
diff --git a/Engine/Source/Runtime/Renderer/Private/CompositionLighting/PostProcessDeferredDecals.cpp b/Engine/Source/Runtime/Renderer/Private/CompositionLighting/PostProcessDeferredDecals.cpp
--- a/Engine/Source/Runtime/Renderer/Private/CompositionLighting/PostProcessDeferredDecals.cpp
+++ b/Engine/Source/Runtime/Renderer/Private/CompositionLighting/PostProcessDeferredDecals.cpp
@@ -689,6 +689,7 @@ bool RenderPreStencil(FRenderingCompositePassContext& Context, const FMaterialSh
 	>::GetRHI() );
 
 	// Carmack's reverse on the bounds
+	//@todo-martinm
 	Context.RHICmdList.SetDepthStencilState(TStaticDepthStencilState<
 		false,CF_LessEqual,
 		true,CF_Equal,SO_Keep,SO_Keep,SO_Increment,
@@ -765,13 +766,13 @@ void FRCPassPostProcessDeferredDecals::Process(FRenderingCompositePassContext& C
 
 			// could be optimized
 			SetRenderTarget(RHICmdList, GSceneRenderTargets.DBufferA->GetRenderTargetItem().TargetableTexture, FTextureRHIParamRef());
-			RHICmdList.Clear(true, FLinearColor(0, 0, 0, 1), false, 0, false, 0, FIntRect());
+			RHICmdList.Clear(true, FLinearColor(0, 0, 0, 1), false, (float)ERHIZBuffer::FarPlane, false, 0, FIntRect());
 			SetRenderTarget(RHICmdList, GSceneRenderTargets.DBufferB->GetRenderTargetItem().TargetableTexture, FTextureRHIParamRef());
 			// todo: some hardware would like to have 0 or 1 for faster clear, we chose 128/255 to represent 0 (8 bit cannot represent 0.5f)
-			RHICmdList.Clear(true, FLinearColor(128.0f / 255.0f, 128.0f / 255.0f, 128.0f / 255.0f, 1), false, 0, false, 0, FIntRect());
+			RHICmdList.Clear(true, FLinearColor(128.0f / 255.0f, 128.0f / 255.0f, 128.0f / 255.0f, 1), false, (float)ERHIZBuffer::FarPlane, false, 0, FIntRect());
 			SetRenderTarget(RHICmdList, GSceneRenderTargets.DBufferC->GetRenderTargetItem().TargetableTexture, FTextureRHIParamRef());
 			// R:roughness, G:roughness opacity
-			RHICmdList.Clear(true, FLinearColor(0, 1, 0, 1), false, 0, false, 0, FIntRect());
+			RHICmdList.Clear(true, FLinearColor(0, 1, 0, 1), false, (float)ERHIZBuffer::FarPlane, false, 0, FIntRect());
 		}
 	}
 
@@ -1049,16 +1050,14 @@ void FRCPassPostProcessDeferredDecals::Process(FRenderingCompositePassContext& C
 					else
 					{
 						// Render frontfaces with depth tests on to get the speedup from HiZ since the camera is outside the light function geometry
-						// Note, this is a reversed Z depth surface, using CF_GreaterEqual.
 						if(bStencilDecals)
 						{
 							// Render frontfaces with depth tests on to get the speedup from HiZ since the camera is outside the light function geometry
 							// Enable stencil testing, only write to pixels with stencil of 0
-							// Note, this is a reversed Z depth surface, using CF_GreaterEqual.
 							if ( bThisDecalUsesStencil )
 							{
 								RHICmdList.SetDepthStencilState(TStaticDepthStencilState<
-									false,CF_GreaterEqual,
+									false,CF_DepthFunction,
 									true,CF_Equal,SO_Zero,SO_Zero,SO_Zero,
 									true,CF_Equal,SO_Zero,SO_Zero,SO_Zero,
 									0xff, 0x7f
@@ -1067,7 +1066,7 @@ void FRCPassPostProcessDeferredDecals::Process(FRenderingCompositePassContext& C
 							else
 							{
 								RHICmdList.SetDepthStencilState(TStaticDepthStencilState<
-									false,CF_GreaterEqual,
+									false,CF_DepthFunction,
 									true,CF_Equal,SO_Keep,SO_Keep,SO_Keep,
 									false,CF_Always,SO_Keep,SO_Keep,SO_Keep,
 									0x80,0x00>::GetRHI(), 0);
@@ -1076,7 +1075,7 @@ void FRCPassPostProcessDeferredDecals::Process(FRenderingCompositePassContext& C
 						}
 						else
 						{
-							RHICmdList.SetDepthStencilState(TStaticDepthStencilState<false, CF_GreaterEqual>::GetRHI(), 0);
+							RHICmdList.SetDepthStencilState(TStaticDepthStencilState<false, CF_DepthFunction>::GetRHI(), 0);
 						}
 						RHICmdList.SetRasterizerState(View.bReverseCulling ? TStaticRasterizerState<FM_Solid, CM_CW>::GetRHI() : TStaticRasterizerState<FM_Solid, CM_CCW>::GetRHI());
 					}
@@ -1088,7 +1087,7 @@ void FRCPassPostProcessDeferredDecals::Process(FRenderingCompositePassContext& C
 
 		// we don't modify stencil but if out input was having stencil for us (after base pass - we need to clear)
 		// Clear stencil to 0, which is the assumed default by other passes
-		RHICmdList.Clear(false, FLinearColor::White, false, 0, true, 0, FIntRect());
+		RHICmdList.Clear(false, FLinearColor::White, false, (float)ERHIZBuffer::FarPlane, true, 0, FIntRect());
 
 		// resolve the targets we wrote to.
 		FResolveParams ResolveParams;
diff --git a/Engine/Source/Runtime/Renderer/Private/DeferredShadingRenderer.cpp b/Engine/Source/Runtime/Renderer/Private/DeferredShadingRenderer.cpp
--- a/Engine/Source/Runtime/Renderer/Private/DeferredShadingRenderer.cpp
+++ b/Engine/Source/Runtime/Renderer/Private/DeferredShadingRenderer.cpp
@@ -148,8 +148,7 @@ void FDeferredShadingSceneRenderer::ClearGBufferAtMaxZ(FRHICommandList& RHICmdLi
 	// Opaque rendering, depth test but no depth writes
 	RHICmdList.SetRasterizerState(TStaticRasterizerState<FM_Solid, CM_None>::GetRHI());
 	RHICmdList.SetBlendState(TStaticBlendStateWriteMask<>::GetRHI());
-	// Note, this is a reversed Z depth surface, using CF_GreaterEqual.
-	RHICmdList.SetDepthStencilState(TStaticDepthStencilState<false, CF_GreaterEqual>::GetRHI());
+	RHICmdList.SetDepthStencilState(TStaticDepthStencilState<false, CF_DepthFunction>::GetRHI());
 
 	// Clear each viewport by drawing background color at MaxZ depth
 	for(int32 ViewIndex = 0;ViewIndex < Views.Num();ViewIndex++)
@@ -167,10 +166,10 @@ void FDeferredShadingSceneRenderer::ClearGBufferAtMaxZ(FRHICommandList& RHICmdLi
 		// Render quad
 		static const FVector4 ClearQuadVertices[4] = 
 		{
-			FVector4( -1.0f,  1.0f, 0.0f, 1.0f ),
-			FVector4(  1.0f,  1.0f, 0.0f, 1.0f ),
-			FVector4( -1.0f, -1.0f, 0.0f, 1.0f ),
-			FVector4(  1.0f, -1.0f, 0.0f, 1.0f )
+			FVector4( -1.0f,	  1.0f, (float)ERHIZBuffer::FarPlane, 1.0f),
+			FVector4(  1.0f,  1.0f, (float)ERHIZBuffer::FarPlane, 1.0f ),
+			FVector4( -1.0f, -1.0f, (float)ERHIZBuffer::FarPlane, 1.0f ),
+			FVector4(  1.0f, -1.0f, (float)ERHIZBuffer::FarPlane, 1.0f )
 		};
 		DrawPrimitiveUP(RHICmdList, PT_TriangleStrip, 2, ClearQuadVertices, sizeof(ClearQuadVertices[0]));
 	}
@@ -453,7 +452,7 @@ static void SetupBasePassView(FRHICommandList& RHICmdList, const FIntRect& ViewR
 		// Additive blending when shader complexity viewmode is enabled.
 		RHICmdList.SetBlendState(TStaticBlendState<CW_RGBA,BO_Add,BF_One,BF_One,BO_Add,BF_Zero,BF_One>::GetRHI());
 		// Disable depth writes as we have a full depth prepass.
-		RHICmdList.SetDepthStencilState(TStaticDepthStencilState<false,CF_GreaterEqual>::GetRHI());
+		RHICmdList.SetDepthStencilState(TStaticDepthStencilState<false,CF_DepthFunction>::GetRHI());
 	}
 	else
 	{
@@ -468,8 +467,7 @@ static void SetupBasePassView(FRHICommandList& RHICmdList, const FIntRect& ViewR
 			RHICmdList.SetBlendState(TStaticBlendStateWriteMask<CW_RGBA, CW_RGBA, CW_RGBA, CW_RGBA>::GetRHI());
 		}
 
-		// Note, this is a reversed Z depth surface, using CF_GreaterEqual.
-		RHICmdList.SetDepthStencilState(TStaticDepthStencilState<true,CF_GreaterEqual>::GetRHI());
+		RHICmdList.SetDepthStencilState(TStaticDepthStencilState<true,CF_DepthFunction>::GetRHI());
 	}
 	RHICmdList.SetScissorRect(false, 0, 0, 0, 0);
 	RHICmdList.SetViewport(ViewRect.Min.X, ViewRect.Min.Y, 0, ViewRect.Max.X, ViewRect.Max.Y, 1);
@@ -853,7 +851,7 @@ void FDeferredShadingSceneRenderer::Render(FRHICommandListImmediate& RHICmdList)
 	{
 		// In Editor we want wire frame view modes to be MSAA for better quality. Resolve will be done with EditorPrimitives
 		SetRenderTarget(RHICmdList, GSceneRenderTargets.GetEditorPrimitivesColor(), GSceneRenderTargets.GetEditorPrimitivesDepth());
-		RHICmdList.Clear(true, FLinearColor(0, 0, 0, 0), true, 0.0f, false, 0, FIntRect());
+		RHICmdList.Clear(true, FLinearColor(0, 0, 0, 0), true, (float)ERHIZBuffer::FarPlane, false, 0, FIntRect());
 	}
 	else if (!bIsGBufferCurrent)
 	{
@@ -872,7 +870,7 @@ void FDeferredShadingSceneRenderer::Render(FRHICommandListImmediate& RHICmdList)
 	  {
 		  // clear out emissive and baked lighting (not too efficient but simple and only needed for this debug view)
 		  GSceneRenderTargets.BeginRenderingSceneColor(RHICmdList);
-		  RHICmdList.Clear(true, FLinearColor(0, 0, 0, 0), false, 0, false, 0, FIntRect());
+		  RHICmdList.Clear(true, FLinearColor(0, 0, 0, 0), false, (float)ERHIZBuffer::FarPlane, false, 0, FIntRect());
 	  }
 
 	  GSceneRenderTargets.DBufferA.SafeRelease();
@@ -1165,8 +1163,7 @@ static void SetupPrePassView(FRHICommandList& RHICmdList, const FIntRect& ViewRe
 {
 	// Disable color writes, enable depth tests and writes.
 	RHICmdList.SetBlendState(TStaticBlendState<CW_NONE>::GetRHI());
-	// Note, this is a reversed Z depth surface, using CF_GreaterEqual.
-	RHICmdList.SetDepthStencilState(TStaticDepthStencilState<true,CF_GreaterEqual>::GetRHI());
+	RHICmdList.SetDepthStencilState(TStaticDepthStencilState<true,CF_DepthFunction>::GetRHI());
 	RHICmdList.SetViewport(ViewRect.Min.X, ViewRect.Min.Y, 0, ViewRect.Max.X, ViewRect.Max.Y, 1);
 	RHICmdList.SetRasterizerState(TStaticRasterizerState<FM_Solid, CM_None>::GetRHI());
 	RHICmdList.SetScissorRect(true, ViewRect.Min.X, ViewRect.Min.Y, ViewRect.Max.X, ViewRect.Max.Y);
diff --git a/Engine/Source/Runtime/Renderer/Private/DistanceFieldSurfaceCacheLighting.cpp b/Engine/Source/Runtime/Renderer/Private/DistanceFieldSurfaceCacheLighting.cpp
--- a/Engine/Source/Runtime/Renderer/Private/DistanceFieldSurfaceCacheLighting.cpp
+++ b/Engine/Source/Runtime/Renderer/Private/DistanceFieldSurfaceCacheLighting.cpp
@@ -3382,8 +3382,7 @@ void SetupDepthStencil(FRHICommandListImmediate& RHICmdList, const FViewInfo& Vi
 
 		// Depth tests on, write 1 to stencil if depth test passed
 		RHICmdList.SetDepthStencilState(TStaticDepthStencilState<
-			// Note, this is a reversed Z depth surface, using CF_GreaterEqual.
-			false,CF_GreaterEqual,
+			false,CF_DepthFunction,
 			true,CF_Always,SO_Keep,SO_Keep,SO_Replace,
 			false,CF_Always,SO_Keep,SO_Keep,SO_Keep,
 			0xff,0xff
@@ -3487,7 +3486,7 @@ void RenderIrradianceCacheInterpolation(
 					{
 						// Depth tests enabled, pass stencil test if stencil is zero
 						RHICmdList.SetDepthStencilState(TStaticDepthStencilState<
-							false,CF_GreaterEqual,
+							false,CF_DepthFunction,
 							true,CF_Equal,SO_Keep,SO_Keep,SO_Keep,
 							false,CF_Always,SO_Keep,SO_Keep,SO_Keep,
 							0xff,0xff
@@ -3496,7 +3495,7 @@ void RenderIrradianceCacheInterpolation(
 					else
 					{
 						// Depth tests enabled
-						RHICmdList.SetDepthStencilState(TStaticDepthStencilState<false, CF_GreaterEqual>::GetRHI());
+						RHICmdList.SetDepthStencilState(TStaticDepthStencilState<false, CF_DepthFunction>::GetRHI());
 					}
 				}
 				else if (GAOInterpolationStencilTesting)
diff --git a/Engine/Source/Runtime/Renderer/Private/DistortionRendering.cpp b/Engine/Source/Runtime/Renderer/Private/DistortionRendering.cpp
--- a/Engine/Source/Runtime/Renderer/Private/DistortionRendering.cpp
+++ b/Engine/Source/Runtime/Renderer/Private/DistortionRendering.cpp
@@ -797,11 +797,10 @@ void FSceneRenderer::RenderDistortion(FRHICommandListImmediate& RHICmdList)
 				RHICmdList.SetViewport(View.ViewRect.Min.X, View.ViewRect.Min.Y, 0.0f, View.ViewRect.Max.X, View.ViewRect.Max.Y, 1.0f);
 
 				// clear offsets to 0
-				RHICmdList.Clear(true, FLinearColor(0, 0, 0, 0), false, 0, false, 0, FIntRect());
+				RHICmdList.Clear(true, FLinearColor(0, 0, 0, 0), false, (float)ERHIZBuffer::FarPlane, false, 0, FIntRect());
 
 				// enable depth test but disable depth writes
-				// Note, this is a reversed Z depth surface, using CF_GreaterEqual.
-				RHICmdList.SetDepthStencilState(TStaticDepthStencilState<false, CF_GreaterEqual>::GetRHI());
+				RHICmdList.SetDepthStencilState(TStaticDepthStencilState<false, CF_DepthFunction>::GetRHI());
 
 				// additive blending of offsets (or complexity if the shader complexity viewmode is enabled)
 				RHICmdList.SetBlendState(TStaticBlendState<CW_RGBA, BO_Add, BF_One, BF_One, BO_Add, BF_One, BF_One>::GetRHI());
@@ -813,8 +812,7 @@ void FSceneRenderer::RenderDistortion(FRHICommandListImmediate& RHICmdList)
 			if (bDirty)
 			{
 				// restore default stencil state
-				// Note, this is a reversed Z depth surface, using CF_GreaterEqual.
-				RHICmdList.SetDepthStencilState(TStaticDepthStencilState<false, CF_GreaterEqual>::GetRHI());
+				RHICmdList.SetDepthStencilState(TStaticDepthStencilState<false, CF_DepthFunction>::GetRHI());
 
 				// resolve using the current ResolveParams 
 				RHICmdList.CopyToResolveTarget(DistortionRT->GetRenderTargetItem().TargetableTexture, DistortionRT->GetRenderTargetItem().ShaderResourceTexture, false, FResolveParams());
diff --git a/Engine/Source/Runtime/Renderer/Private/ForwardBasePassRendering.cpp b/Engine/Source/Runtime/Renderer/Private/ForwardBasePassRendering.cpp
--- a/Engine/Source/Runtime/Renderer/Private/ForwardBasePassRendering.cpp
+++ b/Engine/Source/Runtime/Renderer/Private/ForwardBasePassRendering.cpp
@@ -180,8 +180,7 @@ public:
 		// Treat masked materials as if they don't occlude in shader complexity, which is PVR behavior
 		if(Parameters.BlendMode == BLEND_Masked && View.Family->EngineShowFlags.ShaderComplexity)
 		{
-			// Note, this is a reversed Z depth surface, using CF_GreaterEqual.
-			RHICmdList.SetDepthStencilState(TStaticDepthStencilState<false,CF_GreaterEqual>::GetRHI());
+			RHICmdList.SetDepthStencilState(TStaticDepthStencilState<false,CF_DepthFunction>::GetRHI());
 		}
 #endif
 
@@ -222,8 +221,7 @@ public:
 #if !(UE_BUILD_SHIPPING || UE_BUILD_TEST)
 		if(Parameters.BlendMode == BLEND_Masked && View.Family->EngineShowFlags.ShaderComplexity)
 		{
-			// Note, this is a reversed Z depth surface, using CF_GreaterEqual.
-			RHICmdList.SetDepthStencilState(TStaticDepthStencilState<true,CF_GreaterEqual>::GetRHI());
+			RHICmdList.SetDepthStencilState(TStaticDepthStencilState<true,CF_DepthFunction>::GetRHI());
 		}
 #endif
 	}
@@ -359,8 +357,7 @@ void FForwardShadingSceneRenderer::RenderForwardShadingBasePass(FRHICommandListI
 
 		// Opaque blending
 		RHICmdList.SetBlendState(TStaticBlendStateWriteMask<CW_RGBA>::GetRHI());
-		// Note, this is a reversed Z depth surface, using CF_GreaterEqual.
-		RHICmdList.SetDepthStencilState(TStaticDepthStencilState<true, CF_GreaterEqual>::GetRHI());
+		RHICmdList.SetDepthStencilState(TStaticDepthStencilState<true, CF_DepthFunction>::GetRHI());
 		RHICmdList.SetViewport(View.ViewRect.Min.X, View.ViewRect.Min.Y, 0, View.ViewRect.Max.X, View.ViewRect.Max.Y, 1);
 
 		{
diff --git a/Engine/Source/Runtime/Renderer/Private/ForwardShadingRenderer.cpp b/Engine/Source/Runtime/Renderer/Private/ForwardShadingRenderer.cpp
--- a/Engine/Source/Runtime/Renderer/Private/ForwardShadingRenderer.cpp
+++ b/Engine/Source/Runtime/Renderer/Private/ForwardShadingRenderer.cpp
@@ -129,7 +129,7 @@ void FForwardShadingSceneRenderer::Render(FRHICommandListImmediate& RHICmdList)
 
 	if (GIsEditor)
 	{
-		RHICmdList.Clear(true, Views[0].BackgroundColor, false, 0, false, 0, FIntRect());
+		RHICmdList.Clear(true, Views[0].BackgroundColor, false, (float)ERHIZBuffer::FarPlane, false, 0, FIntRect());
 	}
 
 	RenderForwardShadingBasePass(RHICmdList);
diff --git a/Engine/Source/Runtime/Renderer/Private/ForwardTranslucentRendering.cpp b/Engine/Source/Runtime/Renderer/Private/ForwardTranslucentRendering.cpp
--- a/Engine/Source/Runtime/Renderer/Private/ForwardTranslucentRendering.cpp
+++ b/Engine/Source/Runtime/Renderer/Private/ForwardTranslucentRendering.cpp
@@ -215,8 +215,7 @@ bool FTranslucencyForwardShadingDrawingPolicyFactory::DrawDynamicMesh(
 		if (bDisableDepthTest)
 		{
 			// Restore default depth state
-			// Note, this is a reversed Z depth surface, using CF_GreaterEqual.	
-			RHICmdList.SetDepthStencilState(TStaticDepthStencilState<false, CF_GreaterEqual>::GetRHI());
+			RHICmdList.SetDepthStencilState(TStaticDepthStencilState<false, CF_DepthFunction>::GetRHI());
 		}
 
 		bDirty = true;
@@ -348,8 +347,7 @@ void FForwardShadingSceneRenderer::RenderTranslucency(FRHICommandListImmediate&
 			}
 
 			// Enable depth test, disable depth writes.
-			// Note, this is a reversed Z depth surface, using CF_GreaterEqual.
-			RHICmdList.SetDepthStencilState(TStaticDepthStencilState<false,CF_GreaterEqual>::GetRHI());
+			RHICmdList.SetDepthStencilState(TStaticDepthStencilState<false,CF_DepthFunction>::GetRHI());
 
 			// Draw only translucent prims that don't read from scene color
 			View.TranslucentPrimSet.DrawPrimitivesForForwardShading(RHICmdList, View, *this);
diff --git a/Engine/Source/Runtime/Renderer/Private/LightFunctionRendering.cpp b/Engine/Source/Runtime/Renderer/Private/LightFunctionRendering.cpp
--- a/Engine/Source/Runtime/Renderer/Private/LightFunctionRendering.cpp
+++ b/Engine/Source/Runtime/Renderer/Private/LightFunctionRendering.cpp
@@ -276,7 +276,7 @@ bool FDeferredShadingSceneRenderer::RenderLightFunctionForMaterial(FRHICommandLi
 					if( !bLightAttenuationCleared )
 					{
 						LightSceneInfo->Proxy->SetScissorRect(RHICmdList, View);
-						RHICmdList.Clear(true, FLinearColor::White, false, 0, false, 0, FIntRect());
+						RHICmdList.Clear(true, FLinearColor::White, false, (float)ERHIZBuffer::FarPlane, false, 0, FIntRect());
 					}
 				}
 				else
@@ -315,8 +315,7 @@ bool FDeferredShadingSceneRenderer::RenderLightFunctionForMaterial(FRHICommandLi
 					else
 					{
 						// Render frontfaces with depth tests on to get the speedup from HiZ since the camera is outside the light function geometry
-						// Note, this is a reversed Z depth surface, using CF_GreaterEqual.
-						RHICmdList.SetDepthStencilState(TStaticDepthStencilState<false,CF_GreaterEqual>::GetRHI());
+						RHICmdList.SetDepthStencilState(TStaticDepthStencilState<false,CF_DepthFunction>::GetRHI());
 						RHICmdList.SetRasterizerState(View.bReverseCulling ? TStaticRasterizerState<FM_Solid, CM_CCW>::GetRHI() : TStaticRasterizerState<FM_Solid, CM_CW>::GetRHI());
 					}
 
@@ -341,7 +340,7 @@ bool FDeferredShadingSceneRenderer::RenderLightFunctionForMaterial(FRHICommandLi
 		if (bRenderedLightFunction)
 		{
 			// Restore stencil buffer to all 0's which is the assumed default state
-			RHICmdList.Clear(false,FColor(0,0,0),false,0,true,0, FIntRect());
+			RHICmdList.Clear(false,FColor(0,0,0),false,(float)ERHIZBuffer::FarPlane,true,0, FIntRect());
 		}
 	}
 	return bRenderedLightFunction;
diff --git a/Engine/Source/Runtime/Renderer/Private/LightMapDensityRendering.cpp b/Engine/Source/Runtime/Renderer/Private/LightMapDensityRendering.cpp
--- a/Engine/Source/Runtime/Renderer/Private/LightMapDensityRendering.cpp
+++ b/Engine/Source/Runtime/Renderer/Private/LightMapDensityRendering.cpp
@@ -51,8 +51,7 @@ bool FDeferredShadingSceneRenderer::RenderLightMapDensities(FRHICommandListImmed
 
 			// Opaque blending, depth tests and writes.
 			RHICmdList.SetBlendState(TStaticBlendState<>::GetRHI());
-			// Note, this is a reversed Z depth surface, using CF_GreaterEqual.
-			RHICmdList.SetDepthStencilState(TStaticDepthStencilState<true,CF_GreaterEqual>::GetRHI());
+			RHICmdList.SetDepthStencilState(TStaticDepthStencilState<true,CF_DepthFunction>::GetRHI());
 			RHICmdList.SetViewport(View.ViewRect.Min.X, View.ViewRect.Min.Y, 0, View.ViewRect.Max.X, View.ViewRect.Max.Y, 1);
 
 			{
diff --git a/Engine/Source/Runtime/Renderer/Private/LightRendering.cpp b/Engine/Source/Runtime/Renderer/Private/LightRendering.cpp
--- a/Engine/Source/Runtime/Renderer/Private/LightRendering.cpp
+++ b/Engine/Source/Runtime/Renderer/Private/LightRendering.cpp
@@ -645,7 +645,7 @@ void FDeferredShadingSceneRenderer::RenderStationaryLightOverlap(FRHICommandList
 		GSceneRenderTargets.BeginRenderingSceneColor(RHICmdList, ESimpleRenderTargetMode::EUninitializedColorExistingReadOnlyDepth);
 
 		// Clear to discard base pass values in scene color since we didn't skip that, to have valid scene depths
-		RHICmdList.Clear(true, FLinearColor::Black, false, 0, false, 0, FIntRect());
+		RHICmdList.Clear(true, FLinearColor::Black, false, (float)ERHIZBuffer::FarPlane, false, 0, FIntRect());
 
 		RenderLightArrayForOverlapViewmode(RHICmdList, Scene->Lights);
 
@@ -670,11 +670,10 @@ void SetBoundingGeometryRasterizerAndDepthState(FRHICommandList& RHICmdList, con
 		RHICmdList.SetRasterizerState(View.bReverseCulling ? TStaticRasterizerState<FM_Solid, CM_CCW>::GetRHI() : TStaticRasterizerState<FM_Solid, CM_CW>::GetRHI());
 	}
 
-	// Note, this is a reversed Z depth surface, using CF_GreaterEqual.
 	RHICmdList.SetDepthStencilState(
 		bCameraInsideLightGeometry
 		? TStaticDepthStencilState<false,CF_Always>::GetRHI()
-		: TStaticDepthStencilState<false,CF_GreaterEqual>::GetRHI()
+		: TStaticDepthStencilState<false,CF_DepthFunction>::GetRHI()
 		);
 }
 
@@ -905,7 +904,7 @@ void FDeferredShadingSceneRenderer::RenderLight(FRHICommandList& RHICmdList, con
 	if (bStencilDirty)
 	{
 		// Clear the stencil buffer to 0.
-		RHICmdList.Clear(false, FColor(0, 0, 0), false, 0, true, 0, FIntRect());
+		RHICmdList.Clear(false, FColor(0, 0, 0), false, (float)ERHIZBuffer::FarPlane, true, 0, FIntRect());
 	}
 }
 
diff --git a/Engine/Source/Runtime/Renderer/Private/PostProcess/PostProcessCompositeEditorPrimitives.cpp b/Engine/Source/Runtime/Renderer/Private/PostProcess/PostProcessCompositeEditorPrimitives.cpp
--- a/Engine/Source/Runtime/Renderer/Private/PostProcess/PostProcessCompositeEditorPrimitives.cpp
+++ b/Engine/Source/Runtime/Renderer/Private/PostProcess/PostProcessCompositeEditorPrimitives.cpp
@@ -215,8 +215,7 @@ void FRCPassPostProcessCompositeEditorPrimitives::Process(FRenderingCompositePas
 		{
 			SCOPED_DRAW_EVENT(Context.RHICmdList, ClearViewEditorPrimitives);
 			// Clear color and depth
-			// Note, this is a reversed Z depth surface, so 0.0f is the far plane.
-			Context.RHICmdList.Clear(true, FLinearColor(0, 0, 0, 0), true, 0.0f, false, 0, FIntRect());
+			Context.RHICmdList.Clear(true, FLinearColor(0, 0, 0, 0), true, (float)ERHIZBuffer::FarPlane, false, 0, FIntRect());
 		}
 
 		SCOPED_DRAW_EVENT(Context.RHICmdList, RenderEditorPrimitives);
@@ -307,8 +306,7 @@ void FRCPassPostProcessCompositeEditorPrimitives::RenderPrimitivesToComposite(FR
 {
 
 	// Always depth test against other editor primitives
-	// Note, this is a reversed Z depth surface, using CF_GreaterEqual.
-	RHICmdList.SetDepthStencilState(TStaticDepthStencilState<true, CF_GreaterEqual>::GetRHI());
+	RHICmdList.SetDepthStencilState(TStaticDepthStencilState<true, CF_DepthFunction>::GetRHI());
 	RHICmdList.SetBlendState(TStaticBlendStateWriteMask<CW_RGBA>::GetRHI());
 
 	// most objects should be occluded by the existing scene so we do a manual depth test in the shader
@@ -351,8 +349,7 @@ void FRCPassPostProcessCompositeEditorPrimitives::RenderPrimitivesToComposite(FR
 
 		View.TopBatchedViewElements.Draw(RHICmdList, FeatureLevel, bNeedToSwitchVerticalAxis, View.ViewProjectionMatrix, View.ViewRect.Width(), View.ViewRect.Height(), false);
 
-		// Note, this is a reversed Z depth surface, using CF_GreaterEqual.
-		RHICmdList.SetDepthStencilState(TStaticDepthStencilState<true, CF_GreaterEqual>::GetRHI());
+		RHICmdList.SetDepthStencilState(TStaticDepthStencilState<true, CF_DepthFunction>::GetRHI());
 
 		DrawViewElements<TBasePass>(RHICmdList, View, typename TBasePass::ContextType(bDepthTest, ESceneRenderTargetsMode::SetTextures), SDPG_Foreground, false);
 
diff --git a/Engine/Source/Runtime/Renderer/Private/PostProcess/PostProcessMotionBlur.cpp b/Engine/Source/Runtime/Renderer/Private/PostProcess/PostProcessMotionBlur.cpp
--- a/Engine/Source/Runtime/Renderer/Private/PostProcess/PostProcessMotionBlur.cpp
+++ b/Engine/Source/Runtime/Renderer/Private/PostProcess/PostProcessMotionBlur.cpp
@@ -993,6 +993,8 @@ void FRCPassPostProcessVelocityScatter::Process(FRenderingCompositePassContext&
 	// set the state
 	Context.RHICmdList.SetBlendState(TStaticBlendState<>::GetRHI());
 	Context.RHICmdList.SetRasterizerState(TStaticRasterizerState<>::GetRHI());
+	//@todo-briank: Should this be CF_DepthFunctionNonInclusive?
+	check(RHIHasInvertedZBuffer());
 	Context.RHICmdList.SetDepthStencilState(TStaticDepthStencilState<true, CF_Greater>::GetRHI());
 
 	TShaderMapRef< FPostProcessVelocityScatterVS > VertexShader(Context.GetShaderMap());
diff --git a/Engine/Source/Runtime/Renderer/Private/PostProcess/PostProcessSelectionOutline.cpp b/Engine/Source/Runtime/Renderer/Private/PostProcess/PostProcessSelectionOutline.cpp
--- a/Engine/Source/Runtime/Renderer/Private/PostProcess/PostProcessSelectionOutline.cpp
+++ b/Engine/Source/Runtime/Renderer/Private/PostProcess/PostProcessSelectionOutline.cpp
@@ -37,7 +37,7 @@ void FRCPassPostProcessSelectionOutlineColor::Process(FRenderingCompositePassCon
 	// Set the render target/viewport.
 	SetRenderTarget(Context.RHICmdList, FTextureRHIParamRef(), DestRenderTarget.TargetableTexture);
 	// This is a reversed Z depth surface, so 0.0f is the far plane.
-	Context.RHICmdList.Clear(false, FLinearColor(0, 0, 0, 0), true, 0.0f, true, 0, FIntRect());
+	Context.RHICmdList.Clear(false, FLinearColor(0, 0, 0, 0), true, (float)ERHIZBuffer::FarPlane, true, 0, FIntRect());
 	Context.SetViewportAndCallRHI(ViewRect);
 
 	if (View.Family->EngineShowFlags.Selection)
@@ -83,9 +83,8 @@ void FRCPassPostProcessSelectionOutlineColor::Process(FRenderingCompositePassCon
 					}
 				}
 
-				// This is a reversed Z depth surface, using CF_GreaterEqual.
 				// Note that the stencil value will overflow with enough selected objects
-				Context.RHICmdList.SetDepthStencilState(TStaticDepthStencilState<true, CF_GreaterEqual, true, CF_Always, SO_Keep, SO_Keep, SO_Replace>::GetRHI(), *AssignedStencilIndexPtr);
+				Context.RHICmdList.SetDepthStencilState(TStaticDepthStencilState<true, CF_DepthFunction, true, CF_Always, SO_Keep, SO_Keep, SO_Replace>::GetRHI(), *AssignedStencilIndexPtr);
 
 				const FMeshBatch& MeshBatch = *MeshBatchAndRelevance.Mesh;
 				FHitProxyDrawingPolicyFactory::DrawDynamicMesh(Context.RHICmdList, View, FactoryContext, MeshBatch, false, true, MeshBatchAndRelevance.PrimitiveSceneProxy, MeshBatch.BatchHitProxyId);
@@ -105,16 +104,16 @@ void FRCPassPostProcessSelectionOutlineColor::Process(FRenderingCompositePassCon
 
 			// top
 			Context.RHICmdList.SetScissorRect(true, ViewRect.Min.X, ViewRect.Min.Y, ViewRect.Max.X, InnerRect.Min.Y);
-			Context.RHICmdList.Clear(false, FLinearColor(0, 0, 0, 0), true, 0.0f, true, 0, FIntRect());
+			Context.RHICmdList.Clear(false, FLinearColor(0, 0, 0, 0), true, (float)ERHIZBuffer::FarPlane, true, 0, FIntRect());
 			// bottom
 			Context.RHICmdList.SetScissorRect(true, ViewRect.Min.X, InnerRect.Max.Y, ViewRect.Max.X, ViewRect.Max.Y);
-			Context.RHICmdList.Clear(false, FLinearColor(0, 0, 0, 0), true, 0.0f, true, 0, FIntRect());
+			Context.RHICmdList.Clear(false, FLinearColor(0, 0, 0, 0), true, (float)ERHIZBuffer::FarPlane, true, 0, FIntRect());
 			// left
 			Context.RHICmdList.SetScissorRect(true, ViewRect.Min.X, ViewRect.Min.Y, InnerRect.Min.X, ViewRect.Max.Y);
-			Context.RHICmdList.Clear(false, FLinearColor(0, 0, 0, 0), true, 0.0f, true, 0, FIntRect());
+			Context.RHICmdList.Clear(false, FLinearColor(0, 0, 0, 0), true, (float)ERHIZBuffer::FarPlane, true, 0, FIntRect());
 			// right
 			Context.RHICmdList.SetScissorRect(true, InnerRect.Max.X, ViewRect.Min.Y, ViewRect.Max.X, ViewRect.Max.Y);
-			Context.RHICmdList.Clear(false, FLinearColor(0, 0, 0, 0), true, 0.0f, true, 0, FIntRect());
+			Context.RHICmdList.Clear(false, FLinearColor(0, 0, 0, 0), true, (float)ERHIZBuffer::FarPlane, true, 0, FIntRect());
 
 			Context.RHICmdList.SetScissorRect(false, 0, 0, 0, 0);
 		}
@@ -341,7 +340,7 @@ void FRCPassPostProcessSelectionOutline::Process(FRenderingCompositePassContext&
 
 	// Set the view family's render target/viewport.
 	SetRenderTarget(Context.RHICmdList, DestRenderTarget.TargetableTexture, FTextureRHIRef());
-	Context.RHICmdList.Clear(true, FLinearColor::Black, false, 1.0f, false, 0, ViewRect);
+	Context.RHICmdList.Clear(true, FLinearColor::Black, false, (float)ERHIZBuffer::FarPlane, false, 0, ViewRect);
 	Context.SetViewportAndCallRHI(ViewRect);
 
 	// set the state
diff --git a/Engine/Source/Runtime/Renderer/Private/PostProcess/PostProcessTemporalAA.cpp b/Engine/Source/Runtime/Renderer/Private/PostProcess/PostProcessTemporalAA.cpp
--- a/Engine/Source/Runtime/Renderer/Private/PostProcess/PostProcessTemporalAA.cpp
+++ b/Engine/Source/Runtime/Renderer/Private/PostProcess/PostProcessTemporalAA.cpp
@@ -233,7 +233,7 @@ void FRCPassPostProcessSSRTemporalAA::Process(FRenderingCompositePassContext& Co
 	SetRenderTarget(Context.RHICmdList, DestRenderTarget.TargetableTexture, FTextureRHIRef());
 
 	// is optimized away if possible (RT size=view size, )
-	Context.RHICmdList.Clear(true, FLinearColor::Black, false, 1.0f, false, 0, SrcRect);
+	Context.RHICmdList.Clear(true, FLinearColor::Black, false, (float)ERHIZBuffer::FarPlane, false, 0, SrcRect);
 
 	Context.SetViewportAndCallRHI(SrcRect);
 
@@ -310,7 +310,7 @@ void FRCPassPostProcessDOFTemporalAA::Process(FRenderingCompositePassContext& Co
 	SetRenderTarget(Context.RHICmdList, DestRenderTarget.TargetableTexture, FTextureRHIRef());
 
 	// is optimized away if possible (RT size=view size, )
-	Context.RHICmdList.Clear(true, FLinearColor::Black, false, 1.0f, false, 0, SrcRect);
+	Context.RHICmdList.Clear(true, FLinearColor::Black, false, (float)ERHIZBuffer::FarPlane, false, 0, SrcRect);
 
 	Context.SetViewportAndCallRHI(SrcRect);
 
@@ -391,7 +391,7 @@ void FRCPassPostProcessDOFTemporalAANear::Process(FRenderingCompositePassContext
 	SetRenderTarget(Context.RHICmdList, DestRenderTarget.TargetableTexture, FTextureRHIRef());
 
 	// is optimized away if possible (RT size=view size, )
-	Context.RHICmdList.Clear(true, FLinearColor::Black, false, 1.0f, false, 0, SrcRect);
+	Context.RHICmdList.Clear(true, FLinearColor::Black, false, (float)ERHIZBuffer::FarPlane, false, 0, SrcRect);
 
 	Context.SetViewportAndCallRHI(SrcRect);
 
@@ -471,7 +471,7 @@ void FRCPassPostProcessLightShaftTemporalAA::Process(FRenderingCompositePassCont
 	SetRenderTarget(Context.RHICmdList, DestRenderTarget.TargetableTexture, FTextureRHIRef());
 
 	// is optimized away if possible (RT size=view size, )
-	Context.RHICmdList.Clear(true, FLinearColor::Black, false, 1.0f, false, 0, SrcRect);
+	Context.RHICmdList.Clear(true, FLinearColor::Black, false, (float)ERHIZBuffer::FarPlane, false, 0, SrcRect);
 
 	Context.SetViewportAndCallRHI(SrcRect);
 
@@ -550,7 +550,7 @@ void FRCPassPostProcessTemporalAA::Process(FRenderingCompositePassContext& Conte
 	SetRenderTarget(Context.RHICmdList, DestRenderTarget.TargetableTexture, GSceneRenderTargets.GetSceneDepthTexture(), ESimpleRenderTargetMode::EUninitializedColorExistingReadOnlyDepth);
 
 	// is optimized away if possible (RT size=view size, )
-	Context.RHICmdList.Clear(true, FLinearColor::Black, false, 1.0f, false, 0, SrcRect);
+	Context.RHICmdList.Clear(true, FLinearColor::Black, false, (float)ERHIZBuffer::FarPlane, false, 0, SrcRect);
 
 	Context.SetViewportAndCallRHI(SrcRect);
 
diff --git a/Engine/Source/Runtime/Renderer/Private/PostProcess/SceneRenderTargets.cpp b/Engine/Source/Runtime/Renderer/Private/PostProcess/SceneRenderTargets.cpp
--- a/Engine/Source/Runtime/Renderer/Private/PostProcess/SceneRenderTargets.cpp
+++ b/Engine/Source/Runtime/Renderer/Private/PostProcess/SceneRenderTargets.cpp
@@ -359,7 +359,7 @@ void FSceneRenderTargets::BeginRenderingGBuffer(FRHICommandList& RHICmdList, ERe
 		FRHIRenderTargetView RenderTargets[MaxSimultaneousRenderTargets];
 		int32 MRTCount = GetGBufferRenderTargets(ColorLoadAction, RenderTargets, VelocityRTIndex);
 
-		const float DepthClearValue = 0.0f;
+		const float DepthClearValue = (float)ERHIZBuffer::FarPlane;
 		bool bClearColor = ColorLoadAction == ERenderTargetLoadAction::EClear;
 		bool bClearDepth = DepthLoadAction == ERenderTargetLoadAction::EClear;
 		FRHIDepthRenderTargetView DepthView(GetSceneDepthSurface(), DepthLoadAction, ERenderTargetStoreAction::EStore);
@@ -914,7 +914,7 @@ void FSceneRenderTargets::BeginRenderingPrePass(FRHICommandList& RHICmdList, boo
 	FTextureRHIRef ColorTarget;
 	FTexture2DRHIRef DepthTarget = GetSceneDepthSurface();
 
-	const float DepthClearValue = 0.0f;
+	const float DepthClearValue = (float)ERHIZBuffer::FarPlane;
 	const uint32 StencilClearValue = 0;
 	if (bPerformClear)
 	{				
@@ -1092,7 +1092,7 @@ bool FSceneRenderTargets::BeginRenderingSeparateTranslucency(FRHICommandList& RH
 		if (!bFirstTimeThisFrame)
 		{
 			FLinearColor ClearColor(0.0f, 0.0f, 0.0f, 1.0f);
-			RHICmdList.BindClearMRTValues(true, 1, &ClearColor, false, 0.0f, false, 0);
+			RHICmdList.BindClearMRTValues(true, 1, &ClearColor, false, (float)ERHIZBuffer::FarPlane, false, 0);
 		}
 
 		RHICmdList.SetViewport(View.ViewRect.Min.X, View.ViewRect.Min.Y, 0.0f, View.ViewRect.Max.X, View.ViewRect.Max.Y, 1.0f);
diff --git a/Engine/Source/Runtime/Renderer/Private/ReflectionEnvironmentCapture.cpp b/Engine/Source/Runtime/Renderer/Private/ReflectionEnvironmentCapture.cpp
--- a/Engine/Source/Runtime/Renderer/Private/ReflectionEnvironmentCapture.cpp
+++ b/Engine/Source/Runtime/Renderer/Private/ReflectionEnvironmentCapture.cpp
@@ -675,7 +675,7 @@ void ClearScratchCubemaps(FRHICommandList& RHICmdList)
 		for (int32 CubeFace = 0; CubeFace < CubeFace_MAX; CubeFace++)
 		{
 			SetRenderTarget(RHICmdList, RT0.TargetableTexture, MipIndex, CubeFace, NULL);
-			RHICmdList.Clear(true, FLinearColor(0, 10000, 0, 0), false, 0, false, 0, FIntRect());
+			RHICmdList.Clear(true, FLinearColor(0, 10000, 0, 0), false, (float)ERHIZBuffer::FarPlane, false, 0, FIntRect());
 		}
 	}
 
@@ -686,7 +686,7 @@ void ClearScratchCubemaps(FRHICommandList& RHICmdList)
 		for (int32 CubeFace = 0; CubeFace < CubeFace_MAX; CubeFace++)
 		{
 			SetRenderTarget(RHICmdList, RT1.TargetableTexture, MipIndex, CubeFace, NULL);
-			RHICmdList.Clear(true, FLinearColor(0, 10000, 0, 0), false, 0, false, 0, FIntRect());
+			RHICmdList.Clear(true, FLinearColor(0, 10000, 0, 0), false, (float)ERHIZBuffer::FarPlane, false, 0, FIntRect());
 		}
 	}
 }
@@ -1177,12 +1177,24 @@ void CaptureSceneIntoScratchCubemap(FScene* Scene, FVector CapturePosition, bool
 
 		// Projection matrix based on the fov, near / far clip settings
 		// Each face always uses a 90 degree field of view
-		ViewInitOptions.ProjectionMatrix = FReversedZPerspectiveMatrix(
-			90.0f * (float)PI / 360.0f,
-			(float)GReflectionCaptureSize * GSupersampleCaptureFactor,
-			(float)GReflectionCaptureSize * GSupersampleCaptureFactor,
-			NearPlane
-			);
+		if (RHIHasInvertedZBuffer())
+		{
+			ViewInitOptions.ProjectionMatrix = FReversedZPerspectiveMatrix(
+				90.0f * (float)PI / 360.0f,
+				(float)GReflectionCaptureSize * GSupersampleCaptureFactor,
+				(float)GReflectionCaptureSize * GSupersampleCaptureFactor,
+				NearPlane
+				);
+		}
+		else
+		{
+			ViewInitOptions.ProjectionMatrix = FPerspectiveMatrix(
+				90.0f * (float)PI / 360.0f,
+				(float)GReflectionCaptureSize * GSupersampleCaptureFactor,
+				(float)GReflectionCaptureSize * GSupersampleCaptureFactor,
+				NearPlane
+				);
+		}
 
 		ViewInitOptions.ViewOrigin = CapturePosition;
 		ViewInitOptions.ViewRotationMatrix = CalcCubeFaceViewRotationMatrix((ECubeFace)CubeFace);
diff --git a/Engine/Source/Runtime/Renderer/Private/SceneCaptureRendering.cpp b/Engine/Source/Runtime/Renderer/Private/SceneCaptureRendering.cpp
--- a/Engine/Source/Runtime/Renderer/Private/SceneCaptureRendering.cpp
+++ b/Engine/Source/Runtime/Renderer/Private/SceneCaptureRendering.cpp
@@ -116,7 +116,7 @@ static void UpdateSceneCaptureContent_RenderThread(FRHICommandListImmediate& RHI
 		FIntRect ViewRect = View.ViewRect;
 		FIntRect UnconstrainedViewRect = View.UnconstrainedViewRect;
 		SetRenderTarget(RHICmdList, Target->GetRenderTargetTexture(), NULL);
-		RHICmdList.Clear(true, FLinearColor::Black, false, 1.0f, false, 0, ViewRect);
+		RHICmdList.Clear(true, FLinearColor::Black, false, (float)ERHIZBuffer::FarPlane, false, 0, ViewRect);
 
 		// Render the scene normally
 		{
@@ -202,14 +202,28 @@ FSceneRenderer* FScene::CreateSceneRenderer( USceneCaptureComponent* SceneCaptur
 			YAxisMultiplier = 1.0f;
 		}
 
-		ViewInitOptions.ProjectionMatrix = FReversedZPerspectiveMatrix (
-			FOV,
-			FOV,
-			XAxisMultiplier,
-			YAxisMultiplier,
-			GNearClippingPlane,
-			GNearClippingPlane
-			);
+		if (RHIHasInvertedZBuffer())
+		{
+			ViewInitOptions.ProjectionMatrix = FReversedZPerspectiveMatrix(
+				FOV,
+				FOV,
+				XAxisMultiplier,
+				YAxisMultiplier,
+				GNearClippingPlane,
+				GNearClippingPlane
+				);
+		}
+		else
+		{
+			ViewInitOptions.ProjectionMatrix = FPerspectiveMatrix(
+				FOV,
+				FOV,
+				XAxisMultiplier,
+				YAxisMultiplier,
+				GNearClippingPlane,
+				GNearClippingPlane
+				);
+		}
 	}
 
 	FSceneView* View = new FSceneView(ViewInitOptions);
diff --git a/Engine/Source/Runtime/Renderer/Private/SceneHitProxyRendering.cpp b/Engine/Source/Runtime/Renderer/Private/SceneHitProxyRendering.cpp
--- a/Engine/Source/Runtime/Renderer/Private/SceneHitProxyRendering.cpp
+++ b/Engine/Source/Runtime/Renderer/Private/SceneHitProxyRendering.cpp
@@ -358,7 +358,7 @@ TRefCountPtr<IPooledRenderTarget> InitHitProxyRender(FRHICommandListImmediate& R
 	{
 		const FViewInfo& View = Views[ViewIndex];
 		RHICmdList.SetViewport(View.ViewRect.Min.X, View.ViewRect.Min.Y, 0.0f, View.ViewRect.Max.X, View.ViewRect.Max.Y, 1.0f);
-		RHICmdList.Clear(true, FLinearColor::White, false, 1, false, 0, FIntRect());
+		RHICmdList.Clear(true, FLinearColor::White, false, (float)ERHIZBuffer::FarPlane, false, 0, FIntRect());
 	}
 	return HitProxyRT;
 }
@@ -375,8 +375,7 @@ void RenderHitProxies(FRHICommandListImmediate& RHICmdList, const FSceneRenderer
 	FGlobalDynamicIndexBuffer::Get().Commit();
 
 	// Depth tests + writes, no alpha blending.
-	// Note, this is a reversed Z depth surface, using CF_GreaterEqual.
-	RHICmdList.SetDepthStencilState(TStaticDepthStencilState<true, CF_GreaterEqual>::GetRHI());
+	RHICmdList.SetDepthStencilState(TStaticDepthStencilState<true, CF_DepthFunction>::GetRHI());
 	RHICmdList.SetBlendState(TStaticBlendState<>::GetRHI());
 
 	const bool bNeedToSwitchVerticalAxis = RHINeedsToSwitchVerticalAxis(GShaderPlatformForFeatureLevel[SceneRenderer->FeatureLevel]);
@@ -391,8 +390,7 @@ void RenderHitProxies(FRHICommandListImmediate& RHICmdList, const FSceneRenderer
 		RHICmdList.SetViewport(View.ViewRect.Min.X, View.ViewRect.Min.Y, 0.0f, View.ViewRect.Max.X, View.ViewRect.Max.Y, 1.0f);
 
 		// Clear the depth buffer for each DPG.
-		// Note, this is a reversed Z depth surface, so 0.0f is the far plane.
-		RHICmdList.Clear(false, FLinearColor::Black, true, 0.0f, true, 0, FIntRect());
+		RHICmdList.Clear(false, FLinearColor::Black, true, (float)ERHIZBuffer::FarPlane, true, 0, FIntRect());
 
 		// Adjust the visibility map for this view
 		if (!View.bAllowTranslucentPrimitivesInHitProxy)
@@ -454,8 +452,7 @@ void RenderHitProxies(FRHICommandListImmediate& RHICmdList, const FSceneRenderer
 
 		View.TopBatchedViewElements.Draw(RHICmdList, FeatureLevel, bNeedToSwitchVerticalAxis, View.ViewProjectionMatrix, View.ViewRect.Width(), View.ViewRect.Height(), true);
 
-		// Note, this is a reversed Z depth surface, using CF_GreaterEqual.
-		RHICmdList.SetDepthStencilState(TStaticDepthStencilState<true, CF_GreaterEqual>::GetRHI());
+		RHICmdList.SetDepthStencilState(TStaticDepthStencilState<true, CF_DepthFunction>::GetRHI());
 		// Draw the view's foreground elements last.
 		DrawViewElements<FHitProxyDrawingPolicyFactory>(RHICmdList, View, FHitProxyDrawingPolicyFactory::ContextType(), SDPG_Foreground, bPreFog);
 
diff --git a/Engine/Source/Runtime/Renderer/Private/SceneOcclusion.cpp b/Engine/Source/Runtime/Renderer/Private/SceneOcclusion.cpp
--- a/Engine/Source/Runtime/Renderer/Private/SceneOcclusion.cpp
+++ b/Engine/Source/Runtime/Renderer/Private/SceneOcclusion.cpp
@@ -1062,8 +1062,7 @@ void FDeferredShadingSceneRenderer::BeginOcclusionTests(FRHICommandListImmediate
 			if (ViewState && !View.bDisableQuerySubmissions)
 			{
 				// Depth tests, no depth writes, no color writes, opaque, solid rasterization wo/ backface culling.
-				// Note, this is a reversed Z depth surface, using CF_GreaterEqual.
-				RHICmdList.SetDepthStencilState(TStaticDepthStencilState<false, CF_GreaterEqual>::GetRHI());
+				RHICmdList.SetDepthStencilState(TStaticDepthStencilState<false, CF_DepthFunction>::GetRHI());
 				// We only need to render the front-faces of the culling geometry (this halves the amount of pixels we touch)
 				RHICmdList.SetRasterizerState(View.bReverseCulling ? TStaticRasterizerState<FM_Solid, CM_CCW>::GetRHI() : TStaticRasterizerState<FM_Solid, CM_CW>::GetRHI());
 				RHICmdList.SetBlendState(TStaticBlendState<CW_NONE>::GetRHI());
diff --git a/Engine/Source/Runtime/Renderer/Private/SceneRendering.cpp b/Engine/Source/Runtime/Renderer/Private/SceneRendering.cpp
--- a/Engine/Source/Runtime/Renderer/Private/SceneRendering.cpp
+++ b/Engine/Source/Runtime/Renderer/Private/SceneRendering.cpp
@@ -1241,14 +1241,12 @@ void FSceneRenderer::RenderCustomDepthPass(FRHICommandListImmediate& RHICmdList)
 
 			RHICmdList.SetViewport(View.ViewRect.Min.X, View.ViewRect.Min.Y, 0.0f, View.ViewRect.Max.X, View.ViewRect.Max.Y, 1.0f);
 
-			// Note, this is a reversed Z depth surface, so 0.0f is the far plane.
-			RHICmdList.Clear(false, FLinearColor(0, 0, 0, 0), true, 0.0f, false, 0, FIntRect());
+			RHICmdList.Clear(false, FLinearColor(0, 0, 0, 0), true, (float)ERHIZBuffer::FarPlane, false, 0, FIntRect());
 			
 			// seems this is set each draw call anyway
 			RHICmdList.SetRasterizerState(TStaticRasterizerState<>::GetRHI());
 
-			// Note, this is a reversed Z depth surface, so CF_GreaterEqual.
-			RHICmdList.SetDepthStencilState(TStaticDepthStencilState<true, CF_GreaterEqual>::GetRHI());
+			RHICmdList.SetDepthStencilState(TStaticDepthStencilState<true, CF_DepthFunction>::GetRHI());
 			RHICmdList.SetBlendState(TStaticBlendState<>::GetRHI());
 
 			View.CustomDepthSet.DrawPrims(RHICmdList, View);
diff --git a/Engine/Source/Runtime/Renderer/Private/ShadowRendering.cpp b/Engine/Source/Runtime/Renderer/Private/ShadowRendering.cpp
--- a/Engine/Source/Runtime/Renderer/Private/ShadowRendering.cpp
+++ b/Engine/Source/Runtime/Renderer/Private/ShadowRendering.cpp
@@ -1918,8 +1918,7 @@ void FProjectedShadowInfo::RenderProjection(FRHICommandListImmediate& RHICmdList
 	TShaderMapRef<FShadowProjectionVS> VertexShader(View->ShaderMap);
 
 	// Depth test wo/ writes, no color writing.
-	// Note, this is a reversed Z depth surface, using CF_GreaterEqual.
-	RHICmdList.SetDepthStencilState(TStaticDepthStencilState<false, CF_GreaterEqual>::GetRHI());
+	RHICmdList.SetDepthStencilState(TStaticDepthStencilState<false, CF_DepthFunction>::GetRHI());
 	RHICmdList.SetBlendState(TStaticBlendState<CW_NONE>::GetRHI());
 	
 	bool bDepthBoundsTestEnabled = false;
@@ -1930,9 +1929,8 @@ void FProjectedShadowInfo::RenderProjection(FRHICommandListImmediate& RHICmdList
 		SCOPED_DRAW_EVENTF(RHICmdList, EventMaskSubjects, TEXT("Stencil Mask Subjects"));
 
 		// Set stencil to one.
-		// Note, this is a reversed Z depth surface, using CF_GreaterEqual.
 		RHICmdList.SetDepthStencilState(TStaticDepthStencilState<
-			false,CF_GreaterEqual,
+			false,CF_DepthFunction,
 			true,CF_Always,SO_Keep,SO_Keep,SO_Replace,
 			false,CF_Always,SO_Keep,SO_Keep,SO_Keep,
 			0xff,0xff
@@ -2012,9 +2010,8 @@ void FProjectedShadowInfo::RenderProjection(FRHICommandListImmediate& RHICmdList
 		else
 		{
 		// Increment stencil on front-facing zfail, decrement on back-facing zfail.
-		// Note, this is a reversed Z depth surface, using CF_GreaterEqual.
 			RHICmdList.SetDepthStencilState(TStaticDepthStencilState<
-			false,CF_GreaterEqual,
+			false,CF_DepthFunction,
 			true,CF_Always,SO_Keep,SO_Increment,SO_Keep,
 			true,CF_Always,SO_Keep,SO_Decrement,SO_Keep,
 			0xff,0xff
@@ -2072,9 +2069,8 @@ void FProjectedShadowInfo::RenderProjection(FRHICommandListImmediate& RHICmdList
 			// Because zfail handles these cases while zpass does not.
 			// zfail stenciling is somewhat slower than zpass because on xbox 360 HiZ will be disabled when setting up stencil.
 			// Increment stencil on front-facing zfail, decrement on back-facing zfail.
-			// Note, this is a reversed Z depth surface, using CF_GreaterEqual.
 			RHICmdList.SetDepthStencilState(TStaticDepthStencilState<
-				false,CF_GreaterEqual,
+				false,CF_DepthFunction,
 				true,CF_Always,SO_Keep,SO_Increment,SO_Keep,
 				true,CF_Always,SO_Keep,SO_Decrement,SO_Keep,
 				0xff,0xff
@@ -2084,9 +2080,8 @@ void FProjectedShadowInfo::RenderProjection(FRHICommandListImmediate& RHICmdList
 		{
 			// Increment stencil on front-facing zpass, decrement on back-facing zpass.
 			// HiZ will be enabled on xbox 360 which will save a little GPU time.
-			// Note, this is a reversed Z depth surface, using CF_GreaterEqual.
 			RHICmdList.SetDepthStencilState(TStaticDepthStencilState<
-				false,CF_GreaterEqual,
+				false,CF_DepthFunction,
 				true,CF_Always,SO_Keep,SO_Keep,SO_Increment,
 				true,CF_Always,SO_Keep,SO_Keep,SO_Decrement,
 				0xff,0xff
@@ -2274,8 +2269,7 @@ void FProjectedShadowInfo::RenderOnePassPointLightProjection(FRHICommandListImme
 	else
 	{
 		// Render frontfaces with depth tests on to get the speedup from HiZ since the camera is outside the light geometry
-		// Note, this is a reversed Z depth surface, using CF_GreaterEqual.
-		RHICmdList.SetDepthStencilState(TStaticDepthStencilState<false, CF_GreaterEqual>::GetRHI());
+		RHICmdList.SetDepthStencilState(TStaticDepthStencilState<false, CF_DepthFunction>::GetRHI());
 		RHICmdList.SetRasterizerState(View.bReverseCulling ? TStaticRasterizerState<FM_Solid, CM_CCW>::GetRHI() : TStaticRasterizerState<FM_Solid, CM_CW>::GetRHI());
 	}
 
diff --git a/Engine/Source/Runtime/Renderer/Private/TranslucentRendering.cpp b/Engine/Source/Runtime/Renderer/Private/TranslucentRendering.cpp
--- a/Engine/Source/Runtime/Renderer/Private/TranslucentRendering.cpp
+++ b/Engine/Source/Runtime/Renderer/Private/TranslucentRendering.cpp
@@ -27,8 +27,7 @@ static void SetTranslucentRenderTargetAndState(FRHICommandList& RHICmdList, cons
 	if (bSetupTranslucentState)
 	{
 		// Enable depth test, disable depth writes.
-		// Note, this is a reversed Z depth surface, using CF_GreaterEqual.
-		RHICmdList.SetDepthStencilState(TStaticDepthStencilState<false, CF_GreaterEqual>::GetRHI());
+		RHICmdList.SetDepthStencilState(TStaticDepthStencilState<false, CF_DepthFunction>::GetRHI());
 	}
 }
 
@@ -419,8 +418,7 @@ bool FTranslucencyDrawingPolicyFactory::DrawMesh(
 			}
 			else
 			{
-				// Note, this is a reversed Z depth surface, using CF_GreaterEqual.	
-				RHICmdList.SetDepthStencilState(TStaticDepthStencilState<false,CF_GreaterEqual,true,CF_Always,SO_Keep,SO_Keep,SO_Replace>::GetRHI(), 1);
+				RHICmdList.SetDepthStencilState(TStaticDepthStencilState<false,CF_DepthFunction,true,CF_Always,SO_Keep,SO_Keep,SO_Replace>::GetRHI(), 1);
 			}
 		}
 		else if( bDisableDepthTest )
@@ -452,8 +450,7 @@ bool FTranslucencyDrawingPolicyFactory::DrawMesh(
 		if (bDisableDepthTest || bEnableResponsiveAA)
 		{
 			// Restore default depth state
-			// Note, this is a reversed Z depth surface, using CF_GreaterEqual.	
-			RHICmdList.SetDepthStencilState(TStaticDepthStencilState<false,CF_GreaterEqual>::GetRHI());
+			RHICmdList.SetDepthStencilState(TStaticDepthStencilState<false,CF_DepthFunction>::GetRHI());
 		}
 
 		bDirty = true;
@@ -1069,7 +1066,7 @@ void FDeferredShadingSceneRenderer::RenderTranslucency(FRHICommandListImmediate&
 				{
 					if (bSetupTranslucency)
 					{
-						RHICmdList.SetDepthStencilState(TStaticDepthStencilState<false, CF_GreaterEqual>::GetRHI());
+						RHICmdList.SetDepthStencilState(TStaticDepthStencilState<false, CF_DepthFunction>::GetRHI());
 					}
 
 					View.TranslucentPrimSet.DrawPrimitives(RHICmdList, View, *this, true);
diff --git a/Engine/Source/Runtime/Renderer/Private/VelocityRendering.cpp b/Engine/Source/Runtime/Renderer/Private/VelocityRendering.cpp
--- a/Engine/Source/Runtime/Renderer/Private/VelocityRendering.cpp
+++ b/Engine/Source/Runtime/Renderer/Private/VelocityRendering.cpp
@@ -583,7 +583,7 @@ static void BeginVelocityRendering(FRHICommandList& RHICmdList, TRefCountPtr<IPo
 
 		// some platforms need the clear color when rendertargets transition to SRVs.  We propagate here to allow parallel rendering to always
 		// have the proper mapping when the RT is transitioned.
-		RHICmdList.BindClearMRTValues(true, 1, &VelocityClearColor, false, 1.0f, false, 0);
+		RHICmdList.BindClearMRTValues(true, 1, &VelocityClearColor, false, (float)ERHIZBuffer::FarPlane, false, 0);
 	}
 }
 
@@ -600,8 +600,7 @@ static void SetVelocitiesState(FRHICommandList& RHICmdList, const FViewInfo& Vie
 	RHICmdList.SetViewport(MinX, MinY, 0.0f, MaxX, MaxY, 1.0f);
 
 	RHICmdList.SetBlendState(TStaticBlendState<CW_RGBA>::GetRHI());
-	// Note, this is a reversed Z depth surface, using CF_GreaterEqual.
-	RHICmdList.SetDepthStencilState(TStaticDepthStencilState<false,CF_GreaterEqual>::GetRHI());
+	RHICmdList.SetDepthStencilState(TStaticDepthStencilState<false,CF_DepthFunction>::GetRHI());
 	RHICmdList.SetRasterizerState(GetStaticRasterizerState<true>(FM_Solid, CM_CW));
 }
 
