UE4 - Fix for incorrect culling on ortographic cameras (Fixes UE-12445)

diff --git a/Engine/Source/Runtime/Renderer/Private/SceneVisibility.cpp b/Engine/Source/Runtime/Renderer/Private/SceneVisibility.cpp
--- a/Engine/Source/Runtime/Renderer/Private/SceneVisibility.cpp
+++ b/Engine/Source/Runtime/Renderer/Private/SceneVisibility.cpp
@@ -561,6 +561,12 @@ static int32 OcclusionCull(FRHICommandListImmediate& RHICmdList, const FScene* S
 							bAllowBoundsTest = View.NearClippingPlane.PlaneDot(OcclusionBounds.Origin) < 
 								-(FVector::BoxPushOut(View.NearClippingPlane,OcclusionBounds.BoxExtent));
 						}
+						else if (!View.IsPerspectiveProjection())
+						{
+							// Transform parallel near plane
+							checkf(RHIHasInvertedZBuffer(), TEXT("Check equation for culling!"));
+							bAllowBoundsTest = View.WorldToScreen(OcclusionBounds.Origin).Z - View.ViewMatrices.ProjMatrix.M[2][2] * OcclusionBounds.SphereRadius < 1;
+						}
 						else
 						{
 							bAllowBoundsTest = OcclusionBounds.SphereRadius < HALF_WORLD_MAX;
