Fixed issue where one decal was often not rendering (VB was not set for first decal)

diff --git a/Engine/Source/Runtime/Renderer/Private/CompositionLighting/PostProcessDeferredDecals.cpp b/Engine/Source/Runtime/Renderer/Private/CompositionLighting/PostProcessDeferredDecals.cpp
--- a/Engine/Source/Runtime/Renderer/Private/CompositionLighting/PostProcessDeferredDecals.cpp
+++ b/Engine/Source/Runtime/Renderer/Private/CompositionLighting/PostProcessDeferredDecals.cpp
@@ -698,7 +698,7 @@ bool RenderPreStencil(FRenderingCompositePassContext& Context, const FMaterialSh
 	>::GetRHI());
 
 	// Render decal mask
-	Context.RHICmdList.DrawIndexedPrimitive(GUnitCubeIndexBuffer.IndexBufferRHI, PT_TriangleList, 0, 0, 8, 0, GUnitCubeIndexBuffer.GetIndexCount() / 3, 0);
+	Context.RHICmdList.DrawIndexedPrimitive(GUnitCubeIndexBuffer.IndexBufferRHI, PT_TriangleList, 0, 0, 8, 0, GUnitCubeIndexBuffer.GetIndexCount() / 3, 1);
 
 	return true;
 }
@@ -901,7 +901,7 @@ void FRCPassPostProcessDeferredDecals::Process(FRenderingCompositePassContext& C
 			DBufferCIndex,
 			ResolveBufferMax,
 		};
-
+	
 		FTextureRHIParamRef TargetsToResolve[ResolveBufferMax] = { nullptr };
 
 		for (int32 DecalIndex = 0; DecalIndex < SortedDecals.Num(); DecalIndex++)
@@ -917,18 +917,6 @@ void FRCPassPostProcessDeferredDecals::Process(FRenderingCompositePassContext& C
 			FTranslationMatrix PreViewTranslation(View.ViewMatrices.PreViewTranslation);
 			FMatrix FrustumComponentToClip = DecalScaleTransform * ComponentToWorldMatrix * PreViewTranslation * View.ViewMatrices.TranslatedViewProjectionMatrix;
 
-			bool bThisDecalUsesStencil = false;
-
-			if(bStencilDecals)
-			{
-				if(bDecalPreStencil)
-				{
-					bThisDecalUsesStencil = RenderPreStencil(Context, MaterialShaderMap, ComponentToWorldMatrix, FrustumComponentToClip);
-					WasInsideDecal = -1;
-					LastDecalBlendMode = -1;
-				}
-			}
-
 			// can be optimized as we test against a sphere around the box instead of the box itself
 			const float ConservativeRadius = FMath::Sqrt(
 				ComponentToWorldMatrix.GetScaledAxis( EAxis::X ).SizeSquared() * FMath::Square(GDefaultDecalSize.X) +
@@ -988,6 +976,23 @@ void FRCPassPostProcessDeferredDecals::Process(FRenderingCompositePassContext& C
 						break;
 				}
 				Context.SetViewportAndCallRHI(DestRect);
+
+				// we need to reset the stream source after any call to SetRenderTarget (at least for Metal, which doesn't queue up VB assignments)
+				RHICmdList.SetStreamSource(0, GUnitCubeVertexBuffer.VertexBufferRHI, sizeof(FVector4), 0);
+			}
+
+			bool bThisDecalUsesStencil = false;
+
+			if (bStencilDecals)
+			{
+				if (bDecalPreStencil)
+				{
+					// note this is after a SetStreamSource (in if CurrentRenderTargetMode != LastRenderTargetMode) call as it needs to get the VB input
+					bThisDecalUsesStencil = RenderPreStencil(Context, MaterialShaderMap, ComponentToWorldMatrix, FrustumComponentToClip);
+
+					WasInsideDecal = -1;
+					LastDecalBlendMode = -1;
+				}
 			}
 
 			const bool bBlendStateChange = DecalData.DecalBlendMode != LastDecalBlendMode;// Has decal mode changed.
@@ -1004,9 +1009,6 @@ void FRCPassPostProcessDeferredDecals::Process(FRenderingCompositePassContext& C
 				SetDecalBlendState(RHICmdList, SMFeatureLevel, RenderStage, (EDecalBlendMode)LastDecalBlendMode, DecalData.bHasNormal);
 			}
 
-			// we need to reset the stream source after any call to SetRenderTarget (at least for Metal, which doesn't queue up VB assignments)
-			RHICmdList.SetStreamSource(0, GUnitCubeVertexBuffer.VertexBufferRHI, sizeof(FVector4), 0);
-
 			{
 				TShaderMapRef<FDeferredDecalVS> VertexShader(Context.GetShaderMap());
 				SetShader(Context, bShaderComplexity, DecalData, *VertexShader);
