UE4 - Renamed depth compare enums, changed some checks to static_asserts

diff --git a/Engine/Source/Editor/UnrealEd/Private/EditorViewportClient.cpp b/Engine/Source/Editor/UnrealEd/Private/EditorViewportClient.cpp
--- a/Engine/Source/Editor/UnrealEd/Private/EditorViewportClient.cpp
+++ b/Engine/Source/Editor/UnrealEd/Private/EditorViewportClient.cpp
@@ -627,7 +627,7 @@ FSceneView* FEditorViewportClient::CalcSceneView(FSceneViewFamily* ViewFamily)
 
 			if (bConstrainAspectRatio)
 			{
-				if (RHIHasInvertedZBuffer())
+				if ((int32)ERHIZBuffer::IsInverted != 0)
 				{
 					ViewInitOptions.ProjectionMatrix = FReversedZPerspectiveMatrix(
 						MatrixFOV,
@@ -668,7 +668,7 @@ FSceneView* FEditorViewportClient::CalcSceneView(FSceneViewFamily* ViewFamily)
 					YAxisMultiplier = 1.0f;
 				}
 
-				if (RHIHasInvertedZBuffer())
+				if ((int32)ERHIZBuffer::IsInverted != 0)
 				{
 					ViewInitOptions.ProjectionMatrix = FReversedZPerspectiveMatrix(
 						MatrixFOV,
@@ -694,7 +694,7 @@ FSceneView* FEditorViewportClient::CalcSceneView(FSceneViewFamily* ViewFamily)
 		}
 		else
 		{
-			checkf(RHIHasInvertedZBuffer(), TEXT("Check all the Rotation Matrix transformations!"));
+			static_assert((int32)ERHIZBuffer::IsInverted != 0, "Check all the Rotation Matrix transformations!");
 			float ZScale = 0.5f / HALF_WORLD_MAX;
 			float ZOffset = HALF_WORLD_MAX;
 
diff --git a/Engine/Source/Editor/UnrealEd/Private/ThumbnailHelpers.cpp b/Engine/Source/Editor/UnrealEd/Private/ThumbnailHelpers.cpp
--- a/Engine/Source/Editor/UnrealEd/Private/ThumbnailHelpers.cpp
+++ b/Engine/Source/Editor/UnrealEd/Private/ThumbnailHelpers.cpp
@@ -70,7 +70,7 @@ void FThumbnailPreviewScene::GetView(FSceneViewFamily* ViewFamily, int32 X, int3
 	{
 		const float FOVDegrees = 30.f;
 		const float HalfFOVRadians = FMath::DegreesToRadians<float>(FOVDegrees) * 0.5f;
-		checkf(RHIHasInvertedZBuffer(), TEXT("Check NearPlane and Projection Matrix"));
+		static_assert((int32)ERHIZBuffer::IsInverted != 0, "Check NearPlane and Projection Matrix");
 		const float NearPlane = 1.0f;
 		FMatrix ProjectionMatrix = FReversedZPerspectiveMatrix(
 			HalfFOVRadians,
diff --git a/Engine/Source/Runtime/Engine/Private/SceneView.cpp b/Engine/Source/Runtime/Engine/Private/SceneView.cpp
--- a/Engine/Source/Runtime/Engine/Private/SceneView.cpp
+++ b/Engine/Source/Runtime/Engine/Private/SceneView.cpp
@@ -426,7 +426,7 @@ FSceneView::FSceneView(const FSceneViewInitOptions& InitOptions)
 
 	// Derive the view's near clipping distance and plane.
 	// The GetFrustumFarPlane() is the near plane because of reverse Z projection.
-	checkf(RHIHasInvertedZBuffer(), TEXT("Fix Near Clip distance!"));
+	static_assert((int32)ERHIZBuffer::IsInverted != 0, "Fix Near Clip distance!");
 	bHasNearClippingPlane = ViewProjectionMatrix.GetFrustumFarPlane(NearClippingPlane);
 	if(ViewMatrices.ProjMatrix.M[2][3] > DELTA)
 	{
diff --git a/Engine/Source/Runtime/Engine/Private/ShaderCompiler/ShaderCompiler.cpp b/Engine/Source/Runtime/Engine/Private/ShaderCompiler/ShaderCompiler.cpp
--- a/Engine/Source/Runtime/Engine/Private/ShaderCompiler/ShaderCompiler.cpp
+++ b/Engine/Source/Runtime/Engine/Private/ShaderCompiler/ShaderCompiler.cpp
@@ -2195,7 +2195,7 @@ void GlobalBeginCompileShader(
 		}
 	}
 
-	Input.Environment.SetDefine(TEXT("HAS_INVERTED_Z_BUFFER"), RHIHasInvertedZBuffer() ? 1 : 0);
+	Input.Environment.SetDefine(TEXT("HAS_INVERTED_Z_BUFFER"), ((int32)ERHIZBuffer::IsInverted != 0) ? 1 : 0);
 
 	{
 		FString ShaderPDBRoot;
diff --git a/Engine/Source/Runtime/Engine/Private/UserInterface/Canvas.cpp b/Engine/Source/Runtime/Engine/Private/UserInterface/Canvas.cpp
--- a/Engine/Source/Runtime/Engine/Private/UserInterface/Canvas.cpp
+++ b/Engine/Source/Runtime/Engine/Private/UserInterface/Canvas.cpp
@@ -331,7 +331,7 @@ FMatrix FCanvas::CalcProjectionMatrix(uint32 ViewSizeX, uint32 ViewSizeY, float
 	// convert FOV to randians
 	float FOVRad = fFOV * (float)PI / 360.0f;
 	// project based on the FOV and near plane given
-	if (RHIHasInvertedZBuffer())
+	if ((int32)ERHIZBuffer::IsInverted != 0)
 	{
 		return AdjustProjectionMatrixForRHI(
 			FReversedZPerspectiveMatrix(
diff --git a/Engine/Source/Runtime/RHI/Public/RHIDefinitions.h b/Engine/Source/Runtime/RHI/Public/RHIDefinitions.h
--- a/Engine/Source/Runtime/RHI/Public/RHIDefinitions.h
+++ b/Engine/Source/Runtime/RHI/Public/RHIDefinitions.h
@@ -78,10 +78,13 @@ enum { MaxSimultaneousUAVs = 8 };
 enum class ERHIZBuffer
 {
 	// Before changing this, make sure all math & shader assumptions are correct! Also wrap your C++ assumptions with
-	//		static_assert(FViewMatrices::IsInvertedZBuffer(), ...);
+	//		static_assert(ERHIZBuffer::IsInvertedZBuffer(), ...);
 	// Shader-wise, make sure to update Definitions.usf, HAS_INVERTED_Z_BUFFER
 	FarPlane = 0,
 	NearPlane = 1,
+
+	// 'bool' for knowing if the API is using Inverted Z buffer
+	IsInverted = (int32)((int32)ERHIZBuffer::FarPlane < (int32)ERHIZBuffer::NearPlane),
 };
 
 /**
@@ -168,9 +171,11 @@ enum ECompareFunction
 	CF_Never,
 	CF_Always,
 
-	CF_DepthFunction = ((int32)ERHIZBuffer::NearPlane < (int32)ERHIZBuffer::FarPlane) ? CF_LessEqual : CF_GreaterEqual,
-	CF_DepthFunctionNonInclusive = ((int32)ERHIZBuffer::NearPlane < (int32)ERHIZBuffer::FarPlane) ? CF_Less : CF_Greater,
-	CF_DepthFunctionReject = ((int32)ERHIZBuffer::NearPlane < (int32)ERHIZBuffer::FarPlane) ? CF_Greater : CF_Less,
+	// Utility enumerations
+	CF_DepthNearOrEqual		= (((int32)ERHIZBuffer::IsInverted != 0) ? CF_GreaterEqual : CF_LessEqual),
+	CF_DepthNear				= (((int32)ERHIZBuffer::IsInverted != 0) ? CF_Greater : CF_Less),
+	CF_DepthFartherOrEqual	= (((int32)ERHIZBuffer::IsInverted != 0) ? CF_LessEqual : CF_GreaterEqual),
+	CF_DepthFarther			= (((int32)ERHIZBuffer::IsInverted != 0) ? CF_Less : CF_Greater),
 };
 
 enum EStencilOp
@@ -722,11 +727,6 @@ inline bool RHIHasTiledGPU(const EShaderPlatform Platform)
 	return Platform == SP_METAL_MRT|| Platform == SP_METAL || Platform == SP_OPENGL_ES2_IOS || Platform == SP_OPENGL_ES2;
 }
 
-static inline CONSTEXPR bool RHIHasInvertedZBuffer()
-{
-	return (int32)ERHIZBuffer::FarPlane < (int32)ERHIZBuffer::NearPlane;
-}
-
 inline uint32 GetFeatureLevelMaxTextureSamplers(ERHIFeatureLevel::Type FeatureLevel)
 {
 	if (FeatureLevel == ERHIFeatureLevel::ES2)
diff --git a/Engine/Source/Runtime/RHI/Public/RHIStaticStates.h b/Engine/Source/Runtime/RHI/Public/RHIStaticStates.h
--- a/Engine/Source/Runtime/RHI/Public/RHIStaticStates.h
+++ b/Engine/Source/Runtime/RHI/Public/RHIStaticStates.h
@@ -220,7 +220,7 @@ FRasterizerStateRHIParamRef GetStaticRasterizerState(ERasterizerFillMode FillMod
  */
 template<
 	bool bEnableDepthWrite = true,
-	ECompareFunction DepthTest = CF_DepthFunction,
+	ECompareFunction DepthTest = CF_DepthNearOrEqual,
 	bool bEnableFrontFaceStencil = false,
 	ECompareFunction FrontFaceStencilTest = CF_Always,
 	EStencilOp FrontFaceStencilFailStencilOp = SO_Keep,
diff --git a/Engine/Source/Runtime/Renderer/Private/BasePassRendering.cpp b/Engine/Source/Runtime/Renderer/Private/BasePassRendering.cpp
--- a/Engine/Source/Runtime/Renderer/Private/BasePassRendering.cpp
+++ b/Engine/Source/Runtime/Renderer/Private/BasePassRendering.cpp
@@ -310,7 +310,7 @@ public:
 		// Then rendering in the base pass with additive complexity blending, depth tests on, and depth writes off.
 		if(View.Family->EngineShowFlags.ShaderComplexity)
 		{
-			RHICmdList.SetDepthStencilState(TStaticDepthStencilState<false,CF_DepthFunction>::GetRHI());
+			RHICmdList.SetDepthStencilState(TStaticDepthStencilState<false,CF_DepthNearOrEqual>::GetRHI());
 		}
 #endif
 		const FScene* Scene = Parameters.PrimitiveSceneProxy ? Parameters.PrimitiveSceneProxy->GetPrimitiveSceneInfo()->Scene : NULL;
@@ -350,7 +350,7 @@ public:
 #if !(UE_BUILD_SHIPPING || UE_BUILD_TEST)
 		if(View.Family->EngineShowFlags.ShaderComplexity)
 		{
-			RHICmdList.SetDepthStencilState(TStaticDepthStencilState<true,CF_DepthFunction>::GetRHI());
+			RHICmdList.SetDepthStencilState(TStaticDepthStencilState<true,CF_DepthNearOrEqual>::GetRHI());
 		}
 #endif
 	}
diff --git a/Engine/Source/Runtime/Renderer/Private/CompositionLighting/PostProcessDeferredDecals.cpp b/Engine/Source/Runtime/Renderer/Private/CompositionLighting/PostProcessDeferredDecals.cpp
--- a/Engine/Source/Runtime/Renderer/Private/CompositionLighting/PostProcessDeferredDecals.cpp
+++ b/Engine/Source/Runtime/Renderer/Private/CompositionLighting/PostProcessDeferredDecals.cpp
@@ -1059,7 +1059,7 @@ void FRCPassPostProcessDeferredDecals::Process(FRenderingCompositePassContext& C
 							if ( bThisDecalUsesStencil )
 							{
 								RHICmdList.SetDepthStencilState(TStaticDepthStencilState<
-									false,CF_DepthFunction,
+									false,CF_DepthNearOrEqual,
 									true,CF_Equal,SO_Zero,SO_Zero,SO_Zero,
 									true,CF_Equal,SO_Zero,SO_Zero,SO_Zero,
 									0xff, 0x7f
@@ -1068,7 +1068,7 @@ void FRCPassPostProcessDeferredDecals::Process(FRenderingCompositePassContext& C
 							else
 							{
 								RHICmdList.SetDepthStencilState(TStaticDepthStencilState<
-									false,CF_DepthFunction,
+									false,CF_DepthNearOrEqual,
 									true,CF_Equal,SO_Keep,SO_Keep,SO_Keep,
 									false,CF_Always,SO_Keep,SO_Keep,SO_Keep,
 									0x80,0x00>::GetRHI(), 0);
@@ -1077,7 +1077,7 @@ void FRCPassPostProcessDeferredDecals::Process(FRenderingCompositePassContext& C
 						}
 						else
 						{
-							RHICmdList.SetDepthStencilState(TStaticDepthStencilState<false, CF_DepthFunction>::GetRHI(), 0);
+							RHICmdList.SetDepthStencilState(TStaticDepthStencilState<false, CF_DepthNearOrEqual>::GetRHI(), 0);
 						}
 						RHICmdList.SetRasterizerState(View.bReverseCulling ? TStaticRasterizerState<FM_Solid, CM_CW>::GetRHI() : TStaticRasterizerState<FM_Solid, CM_CCW>::GetRHI());
 					}
diff --git a/Engine/Source/Runtime/Renderer/Private/DeferredShadingRenderer.cpp b/Engine/Source/Runtime/Renderer/Private/DeferredShadingRenderer.cpp
--- a/Engine/Source/Runtime/Renderer/Private/DeferredShadingRenderer.cpp
+++ b/Engine/Source/Runtime/Renderer/Private/DeferredShadingRenderer.cpp
@@ -148,7 +148,7 @@ void FDeferredShadingSceneRenderer::ClearGBufferAtMaxZ(FRHICommandList& RHICmdLi
 	// Opaque rendering, depth test but no depth writes
 	RHICmdList.SetRasterizerState(TStaticRasterizerState<FM_Solid, CM_None>::GetRHI());
 	RHICmdList.SetBlendState(TStaticBlendStateWriteMask<>::GetRHI());
-	RHICmdList.SetDepthStencilState(TStaticDepthStencilState<false, CF_DepthFunction>::GetRHI());
+	RHICmdList.SetDepthStencilState(TStaticDepthStencilState<false, CF_DepthNearOrEqual>::GetRHI());
 
 	// Clear each viewport by drawing background color at MaxZ depth
 	for(int32 ViewIndex = 0;ViewIndex < Views.Num();ViewIndex++)
@@ -452,7 +452,7 @@ static void SetupBasePassView(FRHICommandList& RHICmdList, const FIntRect& ViewR
 		// Additive blending when shader complexity viewmode is enabled.
 		RHICmdList.SetBlendState(TStaticBlendState<CW_RGBA,BO_Add,BF_One,BF_One,BO_Add,BF_Zero,BF_One>::GetRHI());
 		// Disable depth writes as we have a full depth prepass.
-		RHICmdList.SetDepthStencilState(TStaticDepthStencilState<false,CF_DepthFunction>::GetRHI());
+		RHICmdList.SetDepthStencilState(TStaticDepthStencilState<false,CF_DepthNearOrEqual>::GetRHI());
 	}
 	else
 	{
@@ -467,7 +467,7 @@ static void SetupBasePassView(FRHICommandList& RHICmdList, const FIntRect& ViewR
 			RHICmdList.SetBlendState(TStaticBlendStateWriteMask<CW_RGBA, CW_RGBA, CW_RGBA, CW_RGBA>::GetRHI());
 		}
 
-		RHICmdList.SetDepthStencilState(TStaticDepthStencilState<true,CF_DepthFunction>::GetRHI());
+		RHICmdList.SetDepthStencilState(TStaticDepthStencilState<true,CF_DepthNearOrEqual>::GetRHI());
 	}
 	RHICmdList.SetScissorRect(false, 0, 0, 0, 0);
 	RHICmdList.SetViewport(ViewRect.Min.X, ViewRect.Min.Y, 0, ViewRect.Max.X, ViewRect.Max.Y, 1);
@@ -1163,7 +1163,7 @@ static void SetupPrePassView(FRHICommandList& RHICmdList, const FIntRect& ViewRe
 {
 	// Disable color writes, enable depth tests and writes.
 	RHICmdList.SetBlendState(TStaticBlendState<CW_NONE>::GetRHI());
-	RHICmdList.SetDepthStencilState(TStaticDepthStencilState<true,CF_DepthFunction>::GetRHI());
+	RHICmdList.SetDepthStencilState(TStaticDepthStencilState<true,CF_DepthNearOrEqual>::GetRHI());
 	RHICmdList.SetViewport(ViewRect.Min.X, ViewRect.Min.Y, 0, ViewRect.Max.X, ViewRect.Max.Y, 1);
 	RHICmdList.SetRasterizerState(TStaticRasterizerState<FM_Solid, CM_None>::GetRHI());
 	RHICmdList.SetScissorRect(true, ViewRect.Min.X, ViewRect.Min.Y, ViewRect.Max.X, ViewRect.Max.Y);
diff --git a/Engine/Source/Runtime/Renderer/Private/DistanceFieldSurfaceCacheLighting.cpp b/Engine/Source/Runtime/Renderer/Private/DistanceFieldSurfaceCacheLighting.cpp
--- a/Engine/Source/Runtime/Renderer/Private/DistanceFieldSurfaceCacheLighting.cpp
+++ b/Engine/Source/Runtime/Renderer/Private/DistanceFieldSurfaceCacheLighting.cpp
@@ -3382,7 +3382,7 @@ void SetupDepthStencil(FRHICommandListImmediate& RHICmdList, const FViewInfo& Vi
 
 		// Depth tests on, write 1 to stencil if depth test passed
 		RHICmdList.SetDepthStencilState(TStaticDepthStencilState<
-			false,CF_DepthFunction,
+			false,CF_DepthNearOrEqual,
 			true,CF_Always,SO_Keep,SO_Keep,SO_Replace,
 			false,CF_Always,SO_Keep,SO_Keep,SO_Keep,
 			0xff,0xff
@@ -3486,7 +3486,7 @@ void RenderIrradianceCacheInterpolation(
 					{
 						// Depth tests enabled, pass stencil test if stencil is zero
 						RHICmdList.SetDepthStencilState(TStaticDepthStencilState<
-							false,CF_DepthFunction,
+							false,CF_DepthNearOrEqual,
 							true,CF_Equal,SO_Keep,SO_Keep,SO_Keep,
 							false,CF_Always,SO_Keep,SO_Keep,SO_Keep,
 							0xff,0xff
@@ -3495,7 +3495,7 @@ void RenderIrradianceCacheInterpolation(
 					else
 					{
 						// Depth tests enabled
-						RHICmdList.SetDepthStencilState(TStaticDepthStencilState<false, CF_DepthFunction>::GetRHI());
+						RHICmdList.SetDepthStencilState(TStaticDepthStencilState<false, CF_DepthNearOrEqual>::GetRHI());
 					}
 				}
 				else if (GAOInterpolationStencilTesting)
diff --git a/Engine/Source/Runtime/Renderer/Private/DistortionRendering.cpp b/Engine/Source/Runtime/Renderer/Private/DistortionRendering.cpp
--- a/Engine/Source/Runtime/Renderer/Private/DistortionRendering.cpp
+++ b/Engine/Source/Runtime/Renderer/Private/DistortionRendering.cpp
@@ -800,7 +800,7 @@ void FSceneRenderer::RenderDistortion(FRHICommandListImmediate& RHICmdList)
 				RHICmdList.Clear(true, FLinearColor(0, 0, 0, 0), false, (float)ERHIZBuffer::FarPlane, false, 0, FIntRect());
 
 				// enable depth test but disable depth writes
-				RHICmdList.SetDepthStencilState(TStaticDepthStencilState<false, CF_DepthFunction>::GetRHI());
+				RHICmdList.SetDepthStencilState(TStaticDepthStencilState<false, CF_DepthNearOrEqual>::GetRHI());
 
 				// additive blending of offsets (or complexity if the shader complexity viewmode is enabled)
 				RHICmdList.SetBlendState(TStaticBlendState<CW_RGBA, BO_Add, BF_One, BF_One, BO_Add, BF_One, BF_One>::GetRHI());
@@ -812,7 +812,7 @@ void FSceneRenderer::RenderDistortion(FRHICommandListImmediate& RHICmdList)
 			if (bDirty)
 			{
 				// restore default stencil state
-				RHICmdList.SetDepthStencilState(TStaticDepthStencilState<false, CF_DepthFunction>::GetRHI());
+				RHICmdList.SetDepthStencilState(TStaticDepthStencilState<false, CF_DepthNearOrEqual>::GetRHI());
 
 				// resolve using the current ResolveParams 
 				RHICmdList.CopyToResolveTarget(DistortionRT->GetRenderTargetItem().TargetableTexture, DistortionRT->GetRenderTargetItem().ShaderResourceTexture, false, FResolveParams());
diff --git a/Engine/Source/Runtime/Renderer/Private/ForwardBasePassRendering.cpp b/Engine/Source/Runtime/Renderer/Private/ForwardBasePassRendering.cpp
--- a/Engine/Source/Runtime/Renderer/Private/ForwardBasePassRendering.cpp
+++ b/Engine/Source/Runtime/Renderer/Private/ForwardBasePassRendering.cpp
@@ -183,7 +183,7 @@ public:
 		// Treat masked materials as if they don't occlude in shader complexity, which is PVR behavior
 		if(Parameters.BlendMode == BLEND_Masked && View.Family->EngineShowFlags.ShaderComplexity)
 		{
-			RHICmdList.SetDepthStencilState(TStaticDepthStencilState<false,CF_DepthFunction>::GetRHI());
+			RHICmdList.SetDepthStencilState(TStaticDepthStencilState<false,CF_DepthNearOrEqual>::GetRHI());
 		}
 #endif
 
@@ -224,7 +224,7 @@ public:
 #if !(UE_BUILD_SHIPPING || UE_BUILD_TEST)
 		if(Parameters.BlendMode == BLEND_Masked && View.Family->EngineShowFlags.ShaderComplexity)
 		{
-			RHICmdList.SetDepthStencilState(TStaticDepthStencilState<true,CF_DepthFunction>::GetRHI());
+			RHICmdList.SetDepthStencilState(TStaticDepthStencilState<true,CF_DepthNearOrEqual>::GetRHI());
 		}
 #endif
 	}
@@ -360,7 +360,7 @@ void FForwardShadingSceneRenderer::RenderForwardShadingBasePass(FRHICommandListI
 
 		// Opaque blending
 		RHICmdList.SetBlendState(TStaticBlendStateWriteMask<CW_RGBA>::GetRHI());
-		RHICmdList.SetDepthStencilState(TStaticDepthStencilState<true, CF_DepthFunction>::GetRHI());
+		RHICmdList.SetDepthStencilState(TStaticDepthStencilState<true, CF_DepthNearOrEqual>::GetRHI());
 		RHICmdList.SetViewport(View.ViewRect.Min.X, View.ViewRect.Min.Y, 0, View.ViewRect.Max.X, View.ViewRect.Max.Y, 1);
 
 		{
diff --git a/Engine/Source/Runtime/Renderer/Private/ForwardTranslucentRendering.cpp b/Engine/Source/Runtime/Renderer/Private/ForwardTranslucentRendering.cpp
--- a/Engine/Source/Runtime/Renderer/Private/ForwardTranslucentRendering.cpp
+++ b/Engine/Source/Runtime/Renderer/Private/ForwardTranslucentRendering.cpp
@@ -215,7 +215,7 @@ bool FTranslucencyForwardShadingDrawingPolicyFactory::DrawDynamicMesh(
 		if (bDisableDepthTest)
 		{
 			// Restore default depth state
-			RHICmdList.SetDepthStencilState(TStaticDepthStencilState<false, CF_DepthFunction>::GetRHI());
+			RHICmdList.SetDepthStencilState(TStaticDepthStencilState<false, CF_DepthNearOrEqual>::GetRHI());
 		}
 
 		bDirty = true;
@@ -347,7 +347,7 @@ void FForwardShadingSceneRenderer::RenderTranslucency(FRHICommandListImmediate&
 			}
 
 			// Enable depth test, disable depth writes.
-			RHICmdList.SetDepthStencilState(TStaticDepthStencilState<false,CF_DepthFunction>::GetRHI());
+			RHICmdList.SetDepthStencilState(TStaticDepthStencilState<false,CF_DepthNearOrEqual>::GetRHI());
 
 			// Draw only translucent prims that don't read from scene color
 			View.TranslucentPrimSet.DrawPrimitivesForForwardShading(RHICmdList, View, *this);
diff --git a/Engine/Source/Runtime/Renderer/Private/LightFunctionRendering.cpp b/Engine/Source/Runtime/Renderer/Private/LightFunctionRendering.cpp
--- a/Engine/Source/Runtime/Renderer/Private/LightFunctionRendering.cpp
+++ b/Engine/Source/Runtime/Renderer/Private/LightFunctionRendering.cpp
@@ -315,7 +315,7 @@ bool FDeferredShadingSceneRenderer::RenderLightFunctionForMaterial(FRHICommandLi
 					else
 					{
 						// Render frontfaces with depth tests on to get the speedup from HiZ since the camera is outside the light function geometry
-						RHICmdList.SetDepthStencilState(TStaticDepthStencilState<false,CF_DepthFunction>::GetRHI());
+						RHICmdList.SetDepthStencilState(TStaticDepthStencilState<false,CF_DepthNearOrEqual>::GetRHI());
 						RHICmdList.SetRasterizerState(View.bReverseCulling ? TStaticRasterizerState<FM_Solid, CM_CCW>::GetRHI() : TStaticRasterizerState<FM_Solid, CM_CW>::GetRHI());
 					}
 
diff --git a/Engine/Source/Runtime/Renderer/Private/LightMapDensityRendering.cpp b/Engine/Source/Runtime/Renderer/Private/LightMapDensityRendering.cpp
--- a/Engine/Source/Runtime/Renderer/Private/LightMapDensityRendering.cpp
+++ b/Engine/Source/Runtime/Renderer/Private/LightMapDensityRendering.cpp
@@ -51,7 +51,7 @@ bool FDeferredShadingSceneRenderer::RenderLightMapDensities(FRHICommandListImmed
 
 			// Opaque blending, depth tests and writes.
 			RHICmdList.SetBlendState(TStaticBlendState<>::GetRHI());
-			RHICmdList.SetDepthStencilState(TStaticDepthStencilState<true,CF_DepthFunction>::GetRHI());
+			RHICmdList.SetDepthStencilState(TStaticDepthStencilState<true,CF_DepthNearOrEqual>::GetRHI());
 			RHICmdList.SetViewport(View.ViewRect.Min.X, View.ViewRect.Min.Y, 0, View.ViewRect.Max.X, View.ViewRect.Max.Y, 1);
 
 			{
diff --git a/Engine/Source/Runtime/Renderer/Private/LightRendering.cpp b/Engine/Source/Runtime/Renderer/Private/LightRendering.cpp
--- a/Engine/Source/Runtime/Renderer/Private/LightRendering.cpp
+++ b/Engine/Source/Runtime/Renderer/Private/LightRendering.cpp
@@ -673,7 +673,7 @@ void SetBoundingGeometryRasterizerAndDepthState(FRHICommandList& RHICmdList, con
 	RHICmdList.SetDepthStencilState(
 		bCameraInsideLightGeometry
 		? TStaticDepthStencilState<false,CF_Always>::GetRHI()
-		: TStaticDepthStencilState<false,CF_DepthFunction>::GetRHI()
+		: TStaticDepthStencilState<false,CF_DepthNearOrEqual>::GetRHI()
 		);
 }
 
diff --git a/Engine/Source/Runtime/Renderer/Private/PostProcess/PostProcessCompositeEditorPrimitives.cpp b/Engine/Source/Runtime/Renderer/Private/PostProcess/PostProcessCompositeEditorPrimitives.cpp
--- a/Engine/Source/Runtime/Renderer/Private/PostProcess/PostProcessCompositeEditorPrimitives.cpp
+++ b/Engine/Source/Runtime/Renderer/Private/PostProcess/PostProcessCompositeEditorPrimitives.cpp
@@ -306,7 +306,7 @@ void FRCPassPostProcessCompositeEditorPrimitives::RenderPrimitivesToComposite(FR
 {
 
 	// Always depth test against other editor primitives
-	RHICmdList.SetDepthStencilState(TStaticDepthStencilState<true, CF_DepthFunction>::GetRHI());
+	RHICmdList.SetDepthStencilState(TStaticDepthStencilState<true, CF_DepthNearOrEqual>::GetRHI());
 	RHICmdList.SetBlendState(TStaticBlendStateWriteMask<CW_RGBA>::GetRHI());
 
 	// most objects should be occluded by the existing scene so we do a manual depth test in the shader
@@ -349,7 +349,7 @@ void FRCPassPostProcessCompositeEditorPrimitives::RenderPrimitivesToComposite(FR
 
 		View.TopBatchedViewElements.Draw(RHICmdList, FeatureLevel, bNeedToSwitchVerticalAxis, View.ViewProjectionMatrix, View.ViewRect.Width(), View.ViewRect.Height(), false);
 
-		RHICmdList.SetDepthStencilState(TStaticDepthStencilState<true, CF_DepthFunction>::GetRHI());
+		RHICmdList.SetDepthStencilState(TStaticDepthStencilState<true, CF_DepthNearOrEqual>::GetRHI());
 
 		DrawViewElements<TBasePass>(RHICmdList, View, typename TBasePass::ContextType(bDepthTest, ESceneRenderTargetsMode::SetTextures), SDPG_Foreground, false);
 
diff --git a/Engine/Source/Runtime/Renderer/Private/PostProcess/PostProcessMotionBlur.cpp b/Engine/Source/Runtime/Renderer/Private/PostProcess/PostProcessMotionBlur.cpp
--- a/Engine/Source/Runtime/Renderer/Private/PostProcess/PostProcessMotionBlur.cpp
+++ b/Engine/Source/Runtime/Renderer/Private/PostProcess/PostProcessMotionBlur.cpp
@@ -993,8 +993,8 @@ void FRCPassPostProcessVelocityScatter::Process(FRenderingCompositePassContext&
 	// set the state
 	Context.RHICmdList.SetBlendState(TStaticBlendState<>::GetRHI());
 	Context.RHICmdList.SetRasterizerState(TStaticRasterizerState<>::GetRHI());
-	//@todo-briank: Should this be CF_DepthFunctionNonInclusive?
-	check(RHIHasInvertedZBuffer());
+	//@todo-briank: Should this be CF_DepthNear?
+	static_assert((int32)ERHIZBuffer::IsInverted != 0, "Should this be CF_DepthNear?");
 	Context.RHICmdList.SetDepthStencilState(TStaticDepthStencilState<true, CF_Greater>::GetRHI());
 
 	TShaderMapRef< FPostProcessVelocityScatterVS > VertexShader(Context.GetShaderMap());
diff --git a/Engine/Source/Runtime/Renderer/Private/PostProcess/PostProcessSelectionOutline.cpp b/Engine/Source/Runtime/Renderer/Private/PostProcess/PostProcessSelectionOutline.cpp
--- a/Engine/Source/Runtime/Renderer/Private/PostProcess/PostProcessSelectionOutline.cpp
+++ b/Engine/Source/Runtime/Renderer/Private/PostProcess/PostProcessSelectionOutline.cpp
@@ -84,7 +84,7 @@ void FRCPassPostProcessSelectionOutlineColor::Process(FRenderingCompositePassCon
 				}
 
 				// Note that the stencil value will overflow with enough selected objects
-				Context.RHICmdList.SetDepthStencilState(TStaticDepthStencilState<true, CF_DepthFunction, true, CF_Always, SO_Keep, SO_Keep, SO_Replace>::GetRHI(), *AssignedStencilIndexPtr);
+				Context.RHICmdList.SetDepthStencilState(TStaticDepthStencilState<true, CF_DepthNearOrEqual, true, CF_Always, SO_Keep, SO_Keep, SO_Replace>::GetRHI(), *AssignedStencilIndexPtr);
 
 				const FMeshBatch& MeshBatch = *MeshBatchAndRelevance.Mesh;
 				FHitProxyDrawingPolicyFactory::DrawDynamicMesh(Context.RHICmdList, View, FactoryContext, MeshBatch, false, true, MeshBatchAndRelevance.PrimitiveSceneProxy, MeshBatch.BatchHitProxyId);
diff --git a/Engine/Source/Runtime/Renderer/Private/ReflectionEnvironmentCapture.cpp b/Engine/Source/Runtime/Renderer/Private/ReflectionEnvironmentCapture.cpp
--- a/Engine/Source/Runtime/Renderer/Private/ReflectionEnvironmentCapture.cpp
+++ b/Engine/Source/Runtime/Renderer/Private/ReflectionEnvironmentCapture.cpp
@@ -1177,7 +1177,7 @@ void CaptureSceneIntoScratchCubemap(FScene* Scene, FVector CapturePosition, bool
 
 		// Projection matrix based on the fov, near / far clip settings
 		// Each face always uses a 90 degree field of view
-		if (RHIHasInvertedZBuffer())
+		if ((int32)ERHIZBuffer::IsInverted != 0)
 		{
 			ViewInitOptions.ProjectionMatrix = FReversedZPerspectiveMatrix(
 				90.0f * (float)PI / 360.0f,
diff --git a/Engine/Source/Runtime/Renderer/Private/SceneCaptureRendering.cpp b/Engine/Source/Runtime/Renderer/Private/SceneCaptureRendering.cpp
--- a/Engine/Source/Runtime/Renderer/Private/SceneCaptureRendering.cpp
+++ b/Engine/Source/Runtime/Renderer/Private/SceneCaptureRendering.cpp
@@ -202,7 +202,7 @@ FSceneRenderer* FScene::CreateSceneRenderer( USceneCaptureComponent* SceneCaptur
 			YAxisMultiplier = 1.0f;
 		}
 
-		if (RHIHasInvertedZBuffer())
+		if ((int32)ERHIZBuffer::IsInverted != 0)
 		{
 			ViewInitOptions.ProjectionMatrix = FReversedZPerspectiveMatrix(
 				FOV,
diff --git a/Engine/Source/Runtime/Renderer/Private/SceneHitProxyRendering.cpp b/Engine/Source/Runtime/Renderer/Private/SceneHitProxyRendering.cpp
--- a/Engine/Source/Runtime/Renderer/Private/SceneHitProxyRendering.cpp
+++ b/Engine/Source/Runtime/Renderer/Private/SceneHitProxyRendering.cpp
@@ -375,7 +375,7 @@ void RenderHitProxies(FRHICommandListImmediate& RHICmdList, const FSceneRenderer
 	FGlobalDynamicIndexBuffer::Get().Commit();
 
 	// Depth tests + writes, no alpha blending.
-	RHICmdList.SetDepthStencilState(TStaticDepthStencilState<true, CF_DepthFunction>::GetRHI());
+	RHICmdList.SetDepthStencilState(TStaticDepthStencilState<true, CF_DepthNearOrEqual>::GetRHI());
 	RHICmdList.SetBlendState(TStaticBlendState<>::GetRHI());
 
 	const bool bNeedToSwitchVerticalAxis = RHINeedsToSwitchVerticalAxis(GShaderPlatformForFeatureLevel[SceneRenderer->FeatureLevel]);
@@ -452,7 +452,7 @@ void RenderHitProxies(FRHICommandListImmediate& RHICmdList, const FSceneRenderer
 
 		View.TopBatchedViewElements.Draw(RHICmdList, FeatureLevel, bNeedToSwitchVerticalAxis, View.ViewProjectionMatrix, View.ViewRect.Width(), View.ViewRect.Height(), true);
 
-		RHICmdList.SetDepthStencilState(TStaticDepthStencilState<true, CF_DepthFunction>::GetRHI());
+		RHICmdList.SetDepthStencilState(TStaticDepthStencilState<true, CF_DepthNearOrEqual>::GetRHI());
 		// Draw the view's foreground elements last.
 		DrawViewElements<FHitProxyDrawingPolicyFactory>(RHICmdList, View, FHitProxyDrawingPolicyFactory::ContextType(), SDPG_Foreground, bPreFog);
 
diff --git a/Engine/Source/Runtime/Renderer/Private/SceneOcclusion.cpp b/Engine/Source/Runtime/Renderer/Private/SceneOcclusion.cpp
--- a/Engine/Source/Runtime/Renderer/Private/SceneOcclusion.cpp
+++ b/Engine/Source/Runtime/Renderer/Private/SceneOcclusion.cpp
@@ -1062,7 +1062,7 @@ void FDeferredShadingSceneRenderer::BeginOcclusionTests(FRHICommandListImmediate
 			if (ViewState && !View.bDisableQuerySubmissions)
 			{
 				// Depth tests, no depth writes, no color writes, opaque, solid rasterization wo/ backface culling.
-				RHICmdList.SetDepthStencilState(TStaticDepthStencilState<false, CF_DepthFunction>::GetRHI());
+				RHICmdList.SetDepthStencilState(TStaticDepthStencilState<false, CF_DepthNearOrEqual>::GetRHI());
 				// We only need to render the front-faces of the culling geometry (this halves the amount of pixels we touch)
 				RHICmdList.SetRasterizerState(View.bReverseCulling ? TStaticRasterizerState<FM_Solid, CM_CCW>::GetRHI() : TStaticRasterizerState<FM_Solid, CM_CW>::GetRHI());
 				RHICmdList.SetBlendState(TStaticBlendState<CW_NONE>::GetRHI());
diff --git a/Engine/Source/Runtime/Renderer/Private/SceneRendering.cpp b/Engine/Source/Runtime/Renderer/Private/SceneRendering.cpp
--- a/Engine/Source/Runtime/Renderer/Private/SceneRendering.cpp
+++ b/Engine/Source/Runtime/Renderer/Private/SceneRendering.cpp
@@ -1246,7 +1246,7 @@ void FSceneRenderer::RenderCustomDepthPass(FRHICommandListImmediate& RHICmdList)
 			// seems this is set each draw call anyway
 			RHICmdList.SetRasterizerState(TStaticRasterizerState<>::GetRHI());
 
-			RHICmdList.SetDepthStencilState(TStaticDepthStencilState<true, CF_DepthFunction>::GetRHI());
+			RHICmdList.SetDepthStencilState(TStaticDepthStencilState<true, CF_DepthNearOrEqual>::GetRHI());
 			RHICmdList.SetBlendState(TStaticBlendState<>::GetRHI());
 
 			View.CustomDepthSet.DrawPrims(RHICmdList, View);
diff --git a/Engine/Source/Runtime/Renderer/Private/SceneVisibility.cpp b/Engine/Source/Runtime/Renderer/Private/SceneVisibility.cpp
--- a/Engine/Source/Runtime/Renderer/Private/SceneVisibility.cpp
+++ b/Engine/Source/Runtime/Renderer/Private/SceneVisibility.cpp
@@ -564,7 +564,7 @@ static int32 OcclusionCull(FRHICommandListImmediate& RHICmdList, const FScene* S
 						else if (!View.IsPerspectiveProjection())
 						{
 							// Transform parallel near plane
-							checkf(RHIHasInvertedZBuffer(), TEXT("Check equation for culling!"));
+							static_assert((int32)ERHIZBuffer::IsInverted != 0, "Check equation for culling!");
 							bAllowBoundsTest = View.WorldToScreen(OcclusionBounds.Origin).Z - View.ViewMatrices.ProjMatrix.M[2][2] * OcclusionBounds.SphereRadius < 1;
 						}
 						else
diff --git a/Engine/Source/Runtime/Renderer/Private/ShadowRendering.cpp b/Engine/Source/Runtime/Renderer/Private/ShadowRendering.cpp
--- a/Engine/Source/Runtime/Renderer/Private/ShadowRendering.cpp
+++ b/Engine/Source/Runtime/Renderer/Private/ShadowRendering.cpp
@@ -1918,7 +1918,7 @@ void FProjectedShadowInfo::RenderProjection(FRHICommandListImmediate& RHICmdList
 	TShaderMapRef<FShadowProjectionVS> VertexShader(View->ShaderMap);
 
 	// Depth test wo/ writes, no color writing.
-	RHICmdList.SetDepthStencilState(TStaticDepthStencilState<false, CF_DepthFunction>::GetRHI());
+	RHICmdList.SetDepthStencilState(TStaticDepthStencilState<false, CF_DepthNearOrEqual>::GetRHI());
 	RHICmdList.SetBlendState(TStaticBlendState<CW_NONE>::GetRHI());
 	
 	bool bDepthBoundsTestEnabled = false;
@@ -1930,7 +1930,7 @@ void FProjectedShadowInfo::RenderProjection(FRHICommandListImmediate& RHICmdList
 
 		// Set stencil to one.
 		RHICmdList.SetDepthStencilState(TStaticDepthStencilState<
-			false,CF_DepthFunction,
+			false,CF_DepthNearOrEqual,
 			true,CF_Always,SO_Keep,SO_Keep,SO_Replace,
 			false,CF_Always,SO_Keep,SO_Keep,SO_Keep,
 			0xff,0xff
@@ -2011,7 +2011,7 @@ void FProjectedShadowInfo::RenderProjection(FRHICommandListImmediate& RHICmdList
 		{
 		// Increment stencil on front-facing zfail, decrement on back-facing zfail.
 			RHICmdList.SetDepthStencilState(TStaticDepthStencilState<
-			false,CF_DepthFunction,
+			false,CF_DepthNearOrEqual,
 			true,CF_Always,SO_Keep,SO_Increment,SO_Keep,
 			true,CF_Always,SO_Keep,SO_Decrement,SO_Keep,
 			0xff,0xff
@@ -2070,7 +2070,7 @@ void FProjectedShadowInfo::RenderProjection(FRHICommandListImmediate& RHICmdList
 			// zfail stenciling is somewhat slower than zpass because on xbox 360 HiZ will be disabled when setting up stencil.
 			// Increment stencil on front-facing zfail, decrement on back-facing zfail.
 			RHICmdList.SetDepthStencilState(TStaticDepthStencilState<
-				false,CF_DepthFunction,
+				false,CF_DepthNearOrEqual,
 				true,CF_Always,SO_Keep,SO_Increment,SO_Keep,
 				true,CF_Always,SO_Keep,SO_Decrement,SO_Keep,
 				0xff,0xff
@@ -2081,7 +2081,7 @@ void FProjectedShadowInfo::RenderProjection(FRHICommandListImmediate& RHICmdList
 			// Increment stencil on front-facing zpass, decrement on back-facing zpass.
 			// HiZ will be enabled on xbox 360 which will save a little GPU time.
 			RHICmdList.SetDepthStencilState(TStaticDepthStencilState<
-				false,CF_DepthFunction,
+				false,CF_DepthNearOrEqual,
 				true,CF_Always,SO_Keep,SO_Keep,SO_Increment,
 				true,CF_Always,SO_Keep,SO_Keep,SO_Decrement,
 				0xff,0xff
@@ -2269,7 +2269,7 @@ void FProjectedShadowInfo::RenderOnePassPointLightProjection(FRHICommandListImme
 	else
 	{
 		// Render frontfaces with depth tests on to get the speedup from HiZ since the camera is outside the light geometry
-		RHICmdList.SetDepthStencilState(TStaticDepthStencilState<false, CF_DepthFunction>::GetRHI());
+		RHICmdList.SetDepthStencilState(TStaticDepthStencilState<false, CF_DepthNearOrEqual>::GetRHI());
 		RHICmdList.SetRasterizerState(View.bReverseCulling ? TStaticRasterizerState<FM_Solid, CM_CCW>::GetRHI() : TStaticRasterizerState<FM_Solid, CM_CW>::GetRHI());
 	}
 
diff --git a/Engine/Source/Runtime/Renderer/Private/TranslucentRendering.cpp b/Engine/Source/Runtime/Renderer/Private/TranslucentRendering.cpp
--- a/Engine/Source/Runtime/Renderer/Private/TranslucentRendering.cpp
+++ b/Engine/Source/Runtime/Renderer/Private/TranslucentRendering.cpp
@@ -27,7 +27,7 @@ static void SetTranslucentRenderTargetAndState(FRHICommandList& RHICmdList, cons
 	if (bSetupTranslucentState)
 	{
 		// Enable depth test, disable depth writes.
-		RHICmdList.SetDepthStencilState(TStaticDepthStencilState<false, CF_DepthFunction>::GetRHI());
+		RHICmdList.SetDepthStencilState(TStaticDepthStencilState<false, CF_DepthNearOrEqual>::GetRHI());
 	}
 }
 
@@ -418,7 +418,7 @@ bool FTranslucencyDrawingPolicyFactory::DrawMesh(
 			}
 			else
 			{
-				RHICmdList.SetDepthStencilState(TStaticDepthStencilState<false,CF_DepthFunction,true,CF_Always,SO_Keep,SO_Keep,SO_Replace>::GetRHI(), 1);
+				RHICmdList.SetDepthStencilState(TStaticDepthStencilState<false,CF_DepthNearOrEqual,true,CF_Always,SO_Keep,SO_Keep,SO_Replace>::GetRHI(), 1);
 			}
 		}
 		else if( bDisableDepthTest )
@@ -450,7 +450,7 @@ bool FTranslucencyDrawingPolicyFactory::DrawMesh(
 		if (bDisableDepthTest || bEnableResponsiveAA)
 		{
 			// Restore default depth state
-			RHICmdList.SetDepthStencilState(TStaticDepthStencilState<false,CF_DepthFunction>::GetRHI());
+			RHICmdList.SetDepthStencilState(TStaticDepthStencilState<false,CF_DepthNearOrEqual>::GetRHI());
 		}
 
 		bDirty = true;
@@ -1066,7 +1066,7 @@ void FDeferredShadingSceneRenderer::RenderTranslucency(FRHICommandListImmediate&
 				{
 					if (bSetupTranslucency)
 					{
-						RHICmdList.SetDepthStencilState(TStaticDepthStencilState<false, CF_DepthFunction>::GetRHI());
+						RHICmdList.SetDepthStencilState(TStaticDepthStencilState<false, CF_DepthNearOrEqual>::GetRHI());
 					}
 
 					View.TranslucentPrimSet.DrawPrimitives(RHICmdList, View, *this, true);
diff --git a/Engine/Source/Runtime/Renderer/Private/VelocityRendering.cpp b/Engine/Source/Runtime/Renderer/Private/VelocityRendering.cpp
--- a/Engine/Source/Runtime/Renderer/Private/VelocityRendering.cpp
+++ b/Engine/Source/Runtime/Renderer/Private/VelocityRendering.cpp
@@ -600,7 +600,7 @@ static void SetVelocitiesState(FRHICommandList& RHICmdList, const FViewInfo& Vie
 	RHICmdList.SetViewport(MinX, MinY, 0.0f, MaxX, MaxY, 1.0f);
 
 	RHICmdList.SetBlendState(TStaticBlendState<CW_RGBA>::GetRHI());
-	RHICmdList.SetDepthStencilState(TStaticDepthStencilState<false,CF_DepthFunction>::GetRHI());
+	RHICmdList.SetDepthStencilState(TStaticDepthStencilState<false,CF_DepthNearOrEqual>::GetRHI());
 	RHICmdList.SetRasterizerState(GetStaticRasterizerState<true>(FM_Solid, CM_CW));
 }
 
