DFAO optimizations for PS4
* DF lighting GBuffer dependencies packed into 2 textures (saves .1ms)
* Reverted CULL_OBJECTS_TO_RECORD as performance savings are inconsistent
* Added low quality gap fill pass that just does a bilateral filter (saves .4ms)
* Enabled interpolation depth testing (saves .2ms)
* Disabled the temporal stability filter (saves .2ms)
* Overall empty scene cost on PS4 3.77ms -> 2.83ms

diff --git a/Engine/Shaders/DeferredShadingCommon.usf b/Engine/Shaders/DeferredShadingCommon.usf
--- a/Engine/Shaders/DeferredShadingCommon.usf
+++ b/Engine/Shaders/DeferredShadingCommon.usf
@@ -332,7 +332,8 @@ void EncodeGBuffer(
 	OutGBufferA.rgb = EncodeNormal( GBuffer.WorldNormal );
 
 	// compress in 2 bits
-	uint PackedAlpha = GBuffer.DecalMask * 2 + (GBuffer.HasDistanceFieldRepresentation);
+	// Note: distance field lighting dependencies (normal, representation bits) packed for minimal texture fetches
+	uint PackedAlpha = GBuffer.HasHeightfieldRepresentation * 2 + (GBuffer.HasDistanceFieldRepresentation);
 	OutGBufferA.a = PackedAlpha * (1.0 / 3.0);
 
 	if (GBuffer.ShadingModelID == SHADINGMODELID_UNLIT)
@@ -354,7 +355,7 @@ void EncodeGBuffer(
 #else
 		OutGBufferC.a = GBuffer.GBufferAO;
 #endif
-		OutGBufferB.a = Encode71(EncodeShadingModelId(GBuffer.ShadingModelID), GBuffer.HasHeightfieldRepresentation);
+		OutGBufferB.a = Encode71(EncodeShadingModelId(GBuffer.ShadingModelID), GBuffer.DecalMask);
 
 		OutGBufferD.r = GBuffer.Opacity;
 		OutGBufferD.gba = GBuffer.CustomData;
@@ -391,10 +392,10 @@ FGBufferData DecodeGBufferData(
 	}
 
 	uint PackedAlpha = (uint)( InGBufferA.a * 3.999 );
-	GBuffer.DecalMask = (PackedAlpha & 2) != 0;
+	GBuffer.HasHeightfieldRepresentation = (PackedAlpha & 2) != 0;
 	GBuffer.HasDistanceFieldRepresentation = PackedAlpha & 1;
-	GBuffer.HasHeightfieldRepresentation = 0;
-	GBuffer.ShadingModelID = DecodeShadingModelId(Decode71(InGBufferB.a, GBuffer.HasHeightfieldRepresentation));
+	GBuffer.DecalMask = 0;
+	GBuffer.ShadingModelID = DecodeShadingModelId(Decode71(InGBufferB.a, GBuffer.DecalMask));
 	
 	GBuffer.BaseColor = DecodeBaseColor(InGBufferC.rgb);
 	GBuffer.Metallic = InGBufferB.r;
diff --git a/Engine/Shaders/DistanceFieldAOShared.usf b/Engine/Shaders/DistanceFieldAOShared.usf
--- a/Engine/Shaders/DistanceFieldAOShared.usf
+++ b/Engine/Shaders/DistanceFieldAOShared.usf
@@ -99,12 +99,23 @@ void GetDownsampledGBuffer(float2 ScreenUV, out float3 OutNormal, out float OutD
 	OutNormal.z *= abs(TextureValue.z) > .75f ? -1 : 1;
 }
 
+float GetDownsampledDepth(float2 ScreenUV)
+{
+	return abs(Texture2DSampleLevel(DistanceFieldNormalTexture, DistanceFieldNormalSampler, ScreenUV, 0).w);
+}
+
 uint CurrentLevelDownsampleFactor;
 float2 AOBufferSize;
 
 uint DownsampleFactorToBaseLevel;
 float2 BaseLevelTexelSize;
 
+Texture2D BentNormalAOTexture;
+SamplerState BentNormalAOSampler;
+
+Texture2D IrradianceTexture;
+SamplerState IrradianceSampler;
+
 bool SphereIntersectCone(float4 SphereCenterAndRadius, float3 ConeVertex, float3 ConeAxis, float ConeAngleCos, float ConeAngleSin)
 {
 	float3 U = ConeVertex - (SphereCenterAndRadius.w / ConeAngleSin) * ConeAxis;
diff --git a/Engine/Shaders/DistanceFieldLightingPost.usf b/Engine/Shaders/DistanceFieldLightingPost.usf
new file mode 100644
--- /dev/null
+++ b/Engine/Shaders/DistanceFieldLightingPost.usf
@@ -0,0 +1,423 @@
+// Copyright 1998-2015 Epic Games, Inc. All Rights Reserved.
+
+/*=============================================================================
+	DistanceFieldSurfaceCacheLighting.usf
+=============================================================================*/
+
+#include "Common.usf"
+#include "DeferredShadingCommon.usf"
+#include "DistanceFieldLightingShared.usf"
+#include "DistanceFieldAOShared.usf"
+
+float2 BentNormalAOTexelSize;
+float MinDownsampleFactorToBaseLevel;
+
+#define HALF_FILL_KERNEL_SIZE 2
+
+/** Fills in texels with no splatted weight from screenspace neighbors. */
+void FillGapsPS(
+	in float4 UVAndScreenPos : TEXCOORD0
+	,out float4 OutBentNormal : SV_Target0
+#if SUPPORT_IRRADIANCE
+	,out float4 OutIrradiance : SV_Target1
+#endif
+	)
+{
+	float4 CenterValue = Texture2DSampleLevel(BentNormalAOTexture, BentNormalAOSampler, UVAndScreenPos.xy, 0);
+	float SceneDepth = abs(CenterValue.w);
+
+#if HIGH_QUALITY_FILL_GAPS
+
+	float4 Accumulation = 0;
+	float3 IrradianceAccumulation = 0;
+	float Length = 0;
+	// Pixels without a valid interpolation need to find valid neighbors
+	// Pixels with a valid interpolation just get a slight spatial filter so don't rely on a large kernel
+	float HalfKernelSize = CenterValue.w < 0 ? HALF_FILL_KERNEL_SIZE : 1;
+
+	for (float y = -HalfKernelSize; y <= HalfKernelSize; y++)
+	{
+		float WeightY = exp2(-abs(y * 10.0f / HALF_FILL_KERNEL_SIZE));
+
+		for (float x = -HalfKernelSize; x <= HalfKernelSize; x++)
+		{
+			float2 UVOffset = BentNormalAOTexelSize * float2(x, y);
+			float4 TextureValue = Texture2DSampleLevel(BentNormalAOTexture, BentNormalAOSampler, UVAndScreenPos.xy + UVOffset, 0);
+
+			float SampleSceneDepth = abs(TextureValue.w);
+			float ValidMask = TextureValue.w > 0;
+
+			// Don't let depth weight go to 0 with huge depth differences
+			float DepthWeight = max(exp2(-abs(SceneDepth - SampleSceneDepth) * .01f), .001f);
+
+			float2 Weight2D = float2(exp2(-abs(x * 10.0f / HALF_FILL_KERNEL_SIZE)), WeightY);
+			float ScreenSpaceSpatialWeight = max(Weight2D.x, Weight2D.y);
+
+			float Weight = ValidMask * ScreenSpaceSpatialWeight * DepthWeight;
+
+			// Track length separately 
+			Length += length(TextureValue.rgb) * Weight;
+			Accumulation.rgb += TextureValue.rgb * Weight;
+			Accumulation.a += Weight;
+
+			#if SUPPORT_IRRADIANCE
+				float3 Irradiance = Texture2DSampleLevel(IrradianceTexture, IrradianceSampler, UVAndScreenPos.xy + UVOffset, 0).xyz;
+				IrradianceAccumulation += Irradiance * Weight;
+			#endif
+		}
+	}
+	
+	float InvSafeWeight = 1.0f / max(Accumulation.a, .00001f);
+	float AverageLength = Length * InvSafeWeight;
+	float3 AverageBentNormal = Accumulation.rgb * InvSafeWeight;
+	float BentNormalLength = length(AverageBentNormal);
+	float3 AverageIrradiance = IrradianceAccumulation * InvSafeWeight;
+
+	if (BentNormalLength < AverageLength && BentNormalLength > 0)
+	{
+		// Fixup normal shortening due to weighted average of vectors
+		AverageBentNormal = AverageBentNormal / BentNormalLength * AverageLength;
+	}
+
+#else
+
+	float3 AverageBentNormal;
+	float3 AverageIrradiance;
+
+	if (CenterValue.w < 0)
+	{
+		float2 LowResBufferSize = 1.0f / (MinDownsampleFactorToBaseLevel * BentNormalAOTexelSize);
+		float2 LowResTexelSize = BentNormalAOTexelSize * MinDownsampleFactorToBaseLevel;
+		float2 Corner00UV = floor(UVAndScreenPos.xy * LowResBufferSize) / LowResBufferSize + .5f * BentNormalAOTexelSize;
+		float2 BilinearWeights = (UVAndScreenPos.xy - Corner00UV) * LowResBufferSize;
+
+		float4 TextureValues00 = Texture2DSampleLevel(BentNormalAOTexture, BentNormalAOSampler, Corner00UV, 0);
+		float4 TextureValues10 = Texture2DSampleLevel(BentNormalAOTexture, BentNormalAOSampler, Corner00UV + float2(LowResTexelSize.x, 0), 0);
+		float4 TextureValues01 = Texture2DSampleLevel(BentNormalAOTexture, BentNormalAOSampler, Corner00UV + float2(0, LowResTexelSize.y), 0);
+		float4 TextureValues11 = Texture2DSampleLevel(BentNormalAOTexture, BentNormalAOSampler, Corner00UV + LowResTexelSize, 0);
+
+		float4 CornerWeights = float4(
+			(1 - BilinearWeights.y) * (1 - BilinearWeights.x), 
+			(1 - BilinearWeights.y) * BilinearWeights.x,
+			BilinearWeights.y * (1 - BilinearWeights.x),
+			BilinearWeights.y * BilinearWeights.x);
+
+		float Epsilon = .0001f;
+
+		float4 CornerDepths = abs(float4(TextureValues00.w, TextureValues10.w, TextureValues01.w, TextureValues11.w));
+		float4 DepthWeights = min(10.0f / (abs(CornerDepths - SceneDepth.xxxx) + Epsilon), 1);
+		float4 FinalWeights = CornerWeights * DepthWeights;
+
+		float InvSafeWeight = 1.0f / max(dot(FinalWeights, 1), .00001f);
+
+		AverageBentNormal = 
+			(FinalWeights.x * TextureValues00.xyz 
+			+ FinalWeights.y * TextureValues10.xyz
+			+ FinalWeights.z * TextureValues01.xyz 
+			+ FinalWeights.w * TextureValues11.xyz)
+			* InvSafeWeight;
+		
+		#if SUPPORT_IRRADIANCE
+			float4 IrradianceValues00 = Texture2DSampleLevel(IrradianceTexture, IrradianceSampler, Corner00UV, 0);
+			float4 IrradianceValues10 = Texture2DSampleLevel(IrradianceTexture, IrradianceSampler, Corner00UV + float2(LowResTexelSize.x, 0), 0);
+			float4 IrradianceValues01 = Texture2DSampleLevel(IrradianceTexture, IrradianceSampler, Corner00UV + float2(0, LowResTexelSize.y), 0);
+			float4 IrradianceValues11 = Texture2DSampleLevel(IrradianceTexture, IrradianceSampler, Corner00UV + LowResTexelSize, 0);
+
+			AverageIrradiance = 
+				(FinalWeights.x * IrradianceValues00.xyz 
+				+ FinalWeights.y * IrradianceValues10.xyz
+				+ FinalWeights.z * IrradianceValues01.xyz 
+				+ FinalWeights.w * IrradianceValues11.xyz)
+				* InvSafeWeight;
+		#endif
+	}
+	else
+	{
+		AverageBentNormal = CenterValue.xyz;
+
+		#if SUPPORT_IRRADIANCE
+			AverageIrradiance = Texture2DSampleLevel(IrradianceTexture, IrradianceSampler, UVAndScreenPos.xy, 0).xyz;
+		#endif
+	}
+
+#endif
+	
+	OutBentNormal = float4(AverageBentNormal, CenterValue.w);
+
+#if SUPPORT_IRRADIANCE
+	OutIrradiance = float4(AverageIrradiance, 0);
+#endif
+}
+
+float4 CameraMotion[5];
+Texture2D BentNormalHistoryTexture;
+SamplerState BentNormalHistorySampler;
+Texture2D IrradianceHistoryTexture;
+SamplerState IrradianceHistorySampler;
+float HistoryWeight;
+float HistoryDistanceThreshold;
+float UseHistoryFilter;
+
+Texture2D VelocityTexture;
+SamplerState VelocityTextureSampler;
+
+/** Reproject the occlusion history. */
+void UpdateHistoryPS(
+	in float4 UVAndScreenPos : TEXCOORD0
+	,out float4 OutBentNormal : SV_Target0
+#if SUPPORT_IRRADIANCE
+	,out float4 OutIrradiance : SV_Target1
+#endif
+	)
+{
+	// Distance field AO was computed at 0,0 regardless of viewrect min
+	float2 DistanceFieldUVs = UVAndScreenPos.xy - View.ViewRectMin.xy * View.ViewSizeAndSceneTexelSize.zw;
+	float4 NewValue = Texture2DSampleLevel(BentNormalAOTexture, BentNormalAOSampler, DistanceFieldUVs, 0);
+	float SceneDepth = abs(NewValue.w);
+
+	float2 OldUV;
+
+	{
+		float3 xyd;
+		xyd.xy = UVAndScreenPos.zw * float2(0.5, -0.5) + 0.5;
+		xyd.z = ConvertToDeviceZ(SceneDepth);
+
+		float scaleM = 1.0 / (dot(xyd, CameraMotion[0].xyz) + CameraMotion[0].w);
+		float2 MotionVector;
+		// Unwound vector-matrix transform with special care for precision taken from Temporal AA
+		MotionVector.x = ((xyd.x * ((CameraMotion[1].x * xyd.y) + (CameraMotion[1].y * xyd.z) + CameraMotion[1].z)) + (CameraMotion[1].w * xyd.y) + (CameraMotion[2].x * xyd.x * xyd.x) + (CameraMotion[2].y * xyd.z) + CameraMotion[2].z) * scaleM;
+		MotionVector.y = ((xyd.y * ((CameraMotion[3].x * xyd.x) + (CameraMotion[3].y * xyd.z) + CameraMotion[3].z)) + (CameraMotion[3].w * xyd.x) + (CameraMotion[4].x * xyd.y * xyd.y) + (CameraMotion[4].y * xyd.z) + CameraMotion[4].z) * scaleM;
+
+		// Note: have to sample from one of the high res texels exactly to avoid filtering of the velocity buffer
+		// 0 is stored where camera velocity should be used, which must not be filtered with valid object velocities
+		float2 FullResTexel = UVAndScreenPos.xy - .5f * View.ViewSizeAndSceneTexelSize.zw;
+		float2 VelocityN = Texture2DSampleLevel(VelocityTexture, VelocityTextureSampler, FullResTexel, 0).xy;
+
+		if (VelocityN.x > 0)
+		{
+			// Use the per-pixel velocity vector where valid, this handles object, bone and WPO movement
+			MotionVector = float2(-.5f, .5f) * DecodeVelocityFromTexture(VelocityN);
+		}
+
+		OldUV = UVAndScreenPos.xy + MotionVector / float2(0.5, -0.5) * View.ScreenPositionScaleBias.xy;
+	}
+	
+	float2 OldScreenPosition = (OldUV.xy - View.ScreenPositionScaleBias.wz) / View.ScreenPositionScaleBias.xy;
+
+	float2 OldDistanceFieldUVs = OldUV.xy - View.ViewRectMin.xy * View.ViewSizeAndSceneTexelSize.zw;
+	float4 HistoryValue = Texture2DSampleLevel(BentNormalHistoryTexture, BentNormalHistorySampler, OldDistanceFieldUVs, 0);
+	
+	float2 MinUV = View.ViewRectMin.xy * View.ViewSizeAndSceneTexelSize.zw;
+	// Pull in the max UV to exclude the region which will read outside the viewport due to bilinear filtering
+	float2 MaxUV = MinUV + (View.ViewSizeAndSceneTexelSize.xy - 1 * DOWNSAMPLE_FACTOR) * View.ViewSizeAndSceneTexelSize.zw;
+
+	float3 WorldPosition;
+	float3 PrevWorldPosition;
+
+	{
+		float4 HomogeneousWorldPosition = mul(float4(UVAndScreenPos.zw * SceneDepth, SceneDepth, 1), View.ScreenToWorld);
+		WorldPosition = HomogeneousWorldPosition.xyz / HomogeneousWorldPosition.w;
+
+		float HistoryDepth = abs(HistoryValue.w);
+		float4 PrevPositionTranslatedWorld = mul(float4(OldScreenPosition * HistoryDepth, HistoryDepth, 1), View.PrevScreenToTranslatedWorld);
+		PrevPositionTranslatedWorld.xyz /= PrevPositionTranslatedWorld.w;	
+		PrevWorldPosition = PrevPositionTranslatedWorld.xyz - View.PrevPreViewTranslation;
+	}
+
+	float EffectiveHistoryWeight = HistoryWeight;
+
+	FLATTEN
+	if (any(OldUV > MaxUV) 
+		|| any(OldUV < MinUV)
+		// Discard history if we are shading a new position (newly revealed)
+		|| length(PrevWorldPosition - WorldPosition) > HistoryDistanceThreshold)
+	{
+		EffectiveHistoryWeight = 0;
+	}
+
+	OutBentNormal.rgb = lerp(NewValue.rgb, HistoryValue.rgb, EffectiveHistoryWeight);
+
+#if SUPPORT_IRRADIANCE
+	float3 NewIrradiance = Texture2DSampleLevel(IrradianceTexture, IrradianceSampler, DistanceFieldUVs, 0).xyz;
+	float3 HistoryIrradianceValue = Texture2DSampleLevel(IrradianceHistoryTexture, IrradianceHistorySampler, OldDistanceFieldUVs, 0).xyz;
+
+	OutIrradiance = float4(lerp(NewIrradiance, HistoryIrradianceValue, EffectiveHistoryWeight), 0);
+#endif
+
+	OutBentNormal.rgb = isnan(OutBentNormal.rgb) ? 0 : OutBentNormal.rgb;
+	OutBentNormal.a = abs(NewValue.a);
+
+	FLATTEN
+	if (UseHistoryFilter > 0)
+	{
+		// Sign bit of alpha stores whether the history was rejected or not, to be read by the history filter pass
+		OutBentNormal.a *= EffectiveHistoryWeight > 0 ? 1 : -1;
+	}
+}
+
+#define HALF_HISTORY_FILL_KERNEL_SIZE 2
+
+/** Seeds newly rejected history values (which are sources of temporal instability) with the results of a spatial search from stable history values */
+void FilterHistoryPS(
+	in float4 UVAndScreenPos : TEXCOORD0
+	,out float4 OutBentNormal : SV_Target0
+#if SUPPORT_IRRADIANCE
+	,out float4 OutIrradiance : SV_Target1
+#endif
+	)
+{
+	float4 HistoryValue = Texture2DSampleLevel(BentNormalAOTexture, BentNormalAOSampler, UVAndScreenPos.xy, 0);
+
+#if SUPPORT_IRRADIANCE
+	float3 IrradianceValue = Texture2DSampleLevel(IrradianceTexture, IrradianceSampler, UVAndScreenPos.xy, 0).xyz;
+#endif
+
+	// Only do the spatial search for pixels who discarded their history value
+	if (HistoryValue.w < 0)
+	{
+		float SceneDepth = abs(HistoryValue.w);
+
+		float4 Accumulation = 0;
+		float3 IrradianceAccumulation = 0;
+
+		for (float y = -HALF_HISTORY_FILL_KERNEL_SIZE; y <= HALF_HISTORY_FILL_KERNEL_SIZE; y++)
+		{
+			for (float x = -HALF_HISTORY_FILL_KERNEL_SIZE; x <= HALF_HISTORY_FILL_KERNEL_SIZE; x++)
+			{
+				float2 UVOffset = BentNormalAOTexelSize * float2(x, y);
+				float4 TextureValue = Texture2DSampleLevel(BentNormalAOTexture, BentNormalAOSampler, UVAndScreenPos.xy + UVOffset, 0);
+
+				float SampleSceneDepth = abs(TextureValue.w);
+				float ValidMask = TextureValue.w > 0;
+
+				// Weight by depth to avoid pulling in values of a foreground object
+				// This is a careful tradeoff between ghosting behind panning foreground objects and successful spatial searches to reduce flickering
+				float DepthWeight = exp2(-1000 * abs(SceneDepth - SampleSceneDepth) / SceneDepth);
+
+				float2 Weight2D = exp2(-abs(float2(x, y) * 10.0f / HALF_HISTORY_FILL_KERNEL_SIZE));
+				float ScreenSpaceSpatialWeight = max(Weight2D.x, Weight2D.y);
+
+				float Weight = ValidMask * ScreenSpaceSpatialWeight * DepthWeight;
+
+				Accumulation.rgb += TextureValue.rgb * Weight;
+				Accumulation.a += Weight;
+
+				#if SUPPORT_IRRADIANCE
+					float3 Irradiance = Texture2DSampleLevel(IrradianceTexture, IrradianceSampler, UVAndScreenPos.xy + UVOffset, 0).xyz;
+					IrradianceAccumulation += Irradiance * Weight;
+				#endif
+			}
+		}
+
+		// Only change the history value if the spatial search turned up something applicable
+		if (Accumulation.a > 0)
+		{
+			float InvWeight = 1.0f / Accumulation.a;
+			// Construct the history value as if the spatial search result was the previous history, 
+			// And the AO we just computed this frame was the new value
+			HistoryValue.xyz = lerp(HistoryValue.xyz, Accumulation.xyz * InvWeight, HistoryWeight);
+
+			#if SUPPORT_IRRADIANCE
+				IrradianceValue = lerp(IrradianceValue, IrradianceAccumulation * InvWeight, HistoryWeight);
+			#endif
+		}
+	}
+
+	OutBentNormal = HistoryValue;
+	// Remove sign bit so future reprojection interpolation isn't affected
+	OutBentNormal.w = abs(OutBentNormal.w);
+
+#if SUPPORT_IRRADIANCE
+	OutIrradiance = float4(IrradianceValue, 0);
+#endif
+}
+
+/** Upsamples the AO results to full resolution using a bilateral filter. */
+void AOUpsamplePS(
+	in float4 UVAndScreenPos : TEXCOORD0
+	, out float4 OutBentNormal : SV_Target0
+#if OUTPUT_BENT_NORMAL && SUPPORT_IRRADIANCE
+	, out float4 OutIrradiance : SV_Target1
+#endif
+	)
+{
+	// Distance field AO was computed at 0,0 regardless of viewrect min
+	float2 DistanceFieldUVs = UVAndScreenPos.xy - View.ViewRectMin.xy * View.ViewSizeAndSceneTexelSize.zw;
+	float3 Irradiance = 0;
+
+#define BILATERAL_UPSAMPLE 1
+#if BILATERAL_UPSAMPLE
+	float2 LowResBufferSize = floor(View.RenderTargetSize / DOWNSAMPLE_FACTOR);
+	float2 LowResTexelSize = 1.0f / LowResBufferSize;
+	float2 Corner00UV = floor(DistanceFieldUVs * LowResBufferSize - .5f) / LowResBufferSize + .5f * LowResTexelSize;
+	float2 BilinearWeights = (DistanceFieldUVs - Corner00UV) * LowResBufferSize;
+
+	float4 TextureValues00 = Texture2DSampleLevel(BentNormalAOTexture, BentNormalAOSampler, Corner00UV, 0);
+	float4 TextureValues10 = Texture2DSampleLevel(BentNormalAOTexture, BentNormalAOSampler, Corner00UV + float2(LowResTexelSize.x, 0), 0);
+	float4 TextureValues01 = Texture2DSampleLevel(BentNormalAOTexture, BentNormalAOSampler, Corner00UV + float2(0, LowResTexelSize.y), 0);
+	float4 TextureValues11 = Texture2DSampleLevel(BentNormalAOTexture, BentNormalAOSampler, Corner00UV + LowResTexelSize, 0);
+
+	float4 CornerWeights = float4(
+		(1 - BilinearWeights.y) * (1 - BilinearWeights.x), 
+		(1 - BilinearWeights.y) * BilinearWeights.x,
+		BilinearWeights.y * (1 - BilinearWeights.x),
+		BilinearWeights.y * BilinearWeights.x);
+
+	float Epsilon = .0001f;
+
+	float4 CornerDepths = abs(float4(TextureValues00.w, TextureValues10.w, TextureValues01.w, TextureValues11.w));
+	float SceneDepth = CalcSceneDepth(UVAndScreenPos.xy);
+	float4 DepthWeights = 1.0f / (abs(CornerDepths - SceneDepth.xxxx) + Epsilon);
+	float4 FinalWeights = CornerWeights * DepthWeights;
+
+	float InvWeight = 1.0f / dot(FinalWeights, 1);
+
+	float3 InterpolatedResult = 
+		(FinalWeights.x * TextureValues00.xyz 
+		+ FinalWeights.y * TextureValues10.xyz
+		+ FinalWeights.z * TextureValues01.xyz 
+		+ FinalWeights.w * TextureValues11.xyz)
+		* InvWeight;
+
+	float3 BentNormal = InterpolatedResult.xyz;
+
+	#if SUPPORT_IRRADIANCE
+		float4 IrradianceValues00 = Texture2DSampleLevel(IrradianceTexture, IrradianceSampler, Corner00UV, 0);
+		float4 IrradianceValues10 = Texture2DSampleLevel(IrradianceTexture, IrradianceSampler, Corner00UV + float2(LowResTexelSize.x, 0), 0);
+		float4 IrradianceValues01 = Texture2DSampleLevel(IrradianceTexture, IrradianceSampler, Corner00UV + float2(0, LowResTexelSize.y), 0);
+		float4 IrradianceValues11 = Texture2DSampleLevel(IrradianceTexture, IrradianceSampler, Corner00UV + LowResTexelSize, 0);
+
+		Irradiance = 
+			(FinalWeights.x * IrradianceValues00.xyz 
+			+ FinalWeights.y * IrradianceValues10.xyz
+			+ FinalWeights.z * IrradianceValues01.xyz 
+			+ FinalWeights.w * IrradianceValues11.xyz)
+			* InvWeight;
+	#endif
+
+#else
+	float3 BentNormal = Texture2DSampleLevel(BentNormalAOTexture, BentNormalAOSampler, DistanceFieldUVs, 0).xyz;
+
+	#if SUPPORT_IRRADIANCE
+		Irradiance = Texture2DSampleLevel(IrradianceTexture, IrradianceSampler, DistanceFieldUVs, 0).xyz;
+	#endif
+#endif
+
+#if OUTPUT_AO
+	OutBentNormal = float4(length(BentNormal).xxx, 1);
+	//OutBentNormal = float4(BentNormal, 1);
+#else
+	#if OUTPUT_BENT_NORMAL
+		OutBentNormal = float4(BentNormal, 1);
+		#if SUPPORT_IRRADIANCE
+			OutIrradiance = float4(Irradiance, 1);
+		#endif
+	#else
+
+		// Output to RT0 to visualize Irradiance or BentNormal
+		#if SUPPORT_IRRADIANCE
+			OutBentNormal = float4(Irradiance, 1);
+		#endif
+	#endif
+#endif
+}
diff --git a/Engine/Shaders/DistanceFieldSurfaceCacheLighting.usf b/Engine/Shaders/DistanceFieldSurfaceCacheLighting.usf
--- a/Engine/Shaders/DistanceFieldSurfaceCacheLighting.usf
+++ b/Engine/Shaders/DistanceFieldSurfaceCacheLighting.usf
@@ -8,7 +8,6 @@
 #include "DeferredShadingCommon.usf"
 #include "DistanceFieldLightingShared.usf"
 #include "DistanceFieldAOShared.usf"
-#include "MonteCarlo.usf"
 
 struct FObjectCullVertexOutput
 {
@@ -104,8 +103,8 @@ void ComputeDistanceFieldNormalPS(
 	in float4 SVPos : SV_POSITION,
 	out float4 OutColor : SV_Target0)
 {
-	float2 ScreenUV = float2((((uint2)SVPos.xy) * DOWNSAMPLE_FACTOR + View.ViewRectMin.xy + .5f) * View.ViewSizeAndSceneTexelSize.zw);
-	float2 ScreenPosition = (ScreenUV.xy - View.ScreenPositionScaleBias.wz) / View.ScreenPositionScaleBias.xy;
+	// Sample from the center of the top left full resolution texel
+	float2 ScreenUV = float2((floor(SVPos.xy) * DOWNSAMPLE_FACTOR + View.ViewRectMin.xy + .5f) * View.ViewSizeAndSceneTexelSize.zw);
 	float SceneDepth = CalcSceneDepth(ScreenUV);
 	FGBufferData GBufferData = GetGBufferData(ScreenUV);
 
@@ -135,14 +134,8 @@ void WriteDownsampledDepthPS(
 {
 	float2 DistanceFieldUVs = UVAndScreenPos.xy;
 
-	float3 WorldNormal;
-	float SceneDepth;
-	bool bHasDistanceFieldRepresentation;
-	bool bHasHeightfieldRepresentation;
-	GetDownsampledGBuffer(DistanceFieldUVs, WorldNormal, SceneDepth, bHasDistanceFieldRepresentation, bHasHeightfieldRepresentation);
-
 	OutColor = 0;
-	OutDepth = ConvertToDeviceZ(SceneDepth);
+	OutDepth = ConvertToDeviceZ(GetDownsampledDepth(DistanceFieldUVs));
 }
 
 // For some reason gives innaccurate results at lower resolutions
@@ -285,7 +278,7 @@ void IrradianceCacheSplatPS(
 	OutIrradiance = 0;
 #endif
 
-	float2 BaseLevelScreenUV = float2((((uint2)SVPos.xy) * DownsampleFactorToBaseLevel + .5f) * BaseLevelTexelSize);
+	float2 BaseLevelScreenUV = float2(((floor(SVPos.xy)) * DownsampleFactorToBaseLevel + .5f) * BaseLevelTexelSize);
 
 	float3 WorldNormal;
 	float SceneDepth;
@@ -299,7 +292,7 @@ void IrradianceCacheSplatPS(
 
 #else
 
-	float2 ScreenUV = float2((((uint2)SVPos.xy) * CurrentLevelDownsampleFactor + View.ViewRectMin.xy + float2(.5f, .5f)) * View.ViewSizeAndSceneTexelSize.zw);
+	float2 ScreenUV = float2(((floor(SVPos.xy)) * CurrentLevelDownsampleFactor + View.ViewRectMin.xy + float2(.5f, .5f)) * View.ViewSizeAndSceneTexelSize.zw);
 	SceneDepth = CalcSceneDepth(ScreenUV); 
 	float2 ScreenPosition = (ScreenUV.xy - View.ScreenPositionScaleBias.wz) / View.ScreenPositionScaleBias.xy;
 	float4 HomogeneousWorldPosition = mul(float4(ScreenPosition.xy * SceneDepth, SceneDepth, 1), View.ScreenToWorld);
@@ -359,12 +352,6 @@ void IrradianceCacheSplatPS(
 	#endif
 }
 
-Texture2D BentNormalAOTexture;
-SamplerState BentNormalAOSampler;
-
-Texture2D IrradianceTexture;
-SamplerState IrradianceSampler;
-
 Buffer<float4> DebugBuffer;
 
 // 1 / ((1 - FadeDistanceFraction) * AOMaxViewDistance)
@@ -447,341 +434,6 @@ void AOCombinePS(
 #endif
 }
 
-float2 BentNormalAOTexelSize;
-
-#define HALF_FILL_KERNEL_SIZE 2
-
-/** Fills in texels with no splatted weight from screenspace neighbors. */
-void FillGapsPS(
-	in float4 UVAndScreenPos : TEXCOORD0
-	,out float4 OutBentNormal : SV_Target0
-#if SUPPORT_IRRADIANCE
-	,out float4 OutIrradiance : SV_Target1
-#endif
-	)
-{
-	float4 CenterValue = Texture2DSampleLevel(BentNormalAOTexture, BentNormalAOSampler, UVAndScreenPos.xy, 0);
-	float SceneDepth = abs(CenterValue.w);
-
-	float4 Accumulation = 0;
-	float3 IrradianceAccumulation = 0;
-	float Length = 0;
-
-	for (float y = -HALF_FILL_KERNEL_SIZE; y <= HALF_FILL_KERNEL_SIZE; y++)
-	{
-		for (float x = -HALF_FILL_KERNEL_SIZE; x <= HALF_FILL_KERNEL_SIZE; x++)
-		{
-			float2 UVOffset = BentNormalAOTexelSize * float2(x, y);
-			float4 TextureValue = Texture2DSampleLevel(BentNormalAOTexture, BentNormalAOSampler, UVAndScreenPos.xy + UVOffset, 0);
-
-			float SampleSceneDepth = abs(TextureValue.w);
-			float ValidMask = TextureValue.w > 0;
-
-			// Don't let depth weight go to 0 with huge depth differences
-			float DepthWeight = max(exp2(-abs(SceneDepth - SampleSceneDepth) * .01f), .001f);
-
-			float2 Weight2D = exp2(-abs(float2(x, y) * 10.0f / HALF_FILL_KERNEL_SIZE));
-			float ScreenSpaceSpatialWeight = max(Weight2D.x, Weight2D.y);
-
-			float Weight = ValidMask * ScreenSpaceSpatialWeight * DepthWeight;
-
-			// Track length separately 
-			Length += length(TextureValue.rgb) * Weight;
-			Accumulation.rgb += TextureValue.rgb * Weight;
-			Accumulation.a += Weight;
-
-			#if SUPPORT_IRRADIANCE
-				float3 Irradiance = Texture2DSampleLevel(IrradianceTexture, IrradianceSampler, UVAndScreenPos.xy + UVOffset, 0).xyz;
-				IrradianceAccumulation += Irradiance * Weight;
-			#endif
-		}
-	}
-	
-	float InvSafeWeight = 1.0f / max(Accumulation.a, .00001f);
-	float AverageLength = Length * InvSafeWeight;
-	float3 AverageBentNormal = Accumulation.rgb * InvSafeWeight;
-	float BentNormalLength = length(AverageBentNormal);
-
-	if (BentNormalLength < AverageLength && BentNormalLength > 0)
-	{
-		// Fixup normal shortening due to weighted average of vectors
-		AverageBentNormal = AverageBentNormal / BentNormalLength * AverageLength;
-	}
-	
-	OutBentNormal = float4(AverageBentNormal, CenterValue.w);
-
-#if SUPPORT_IRRADIANCE
-	OutIrradiance = float4(IrradianceAccumulation * InvSafeWeight, 0);
-#endif
-}
-
-float4 CameraMotion[5];
-Texture2D BentNormalHistoryTexture;
-SamplerState BentNormalHistorySampler;
-Texture2D IrradianceHistoryTexture;
-SamplerState IrradianceHistorySampler;
-float HistoryWeight;
-float HistoryDistanceThreshold;
-
-Texture2D VelocityTexture;
-SamplerState VelocityTextureSampler;
-
-/** Reproject the occlusion history. */
-void UpdateHistoryPS(
-	in float4 UVAndScreenPos : TEXCOORD0
-	,out float4 OutBentNormal : SV_Target0
-#if SUPPORT_IRRADIANCE
-	,out float4 OutIrradiance : SV_Target1
-#endif
-	)
-{
-	// Distance field AO was computed at 0,0 regardless of viewrect min
-	float2 DistanceFieldUVs = UVAndScreenPos.xy - View.ViewRectMin.xy * View.ViewSizeAndSceneTexelSize.zw;
-	float4 NewValue = Texture2DSampleLevel(BentNormalAOTexture, BentNormalAOSampler, DistanceFieldUVs, 0);
-	float SceneDepth = abs(NewValue.w);
-
-	float2 OldUV;
-
-	{
-		float3 xyd;
-		xyd.xy = UVAndScreenPos.zw * float2(0.5, -0.5) + 0.5;
-		xyd.z = ConvertToDeviceZ(SceneDepth);
-
-		float scaleM = 1.0 / (dot(xyd, CameraMotion[0].xyz) + CameraMotion[0].w);
-		float2 MotionVector;
-		// Unwound vector-matrix transform with special care for precision taken from Temporal AA
-		MotionVector.x = ((xyd.x * ((CameraMotion[1].x * xyd.y) + (CameraMotion[1].y * xyd.z) + CameraMotion[1].z)) + (CameraMotion[1].w * xyd.y) + (CameraMotion[2].x * xyd.x * xyd.x) + (CameraMotion[2].y * xyd.z) + CameraMotion[2].z) * scaleM;
-		MotionVector.y = ((xyd.y * ((CameraMotion[3].x * xyd.x) + (CameraMotion[3].y * xyd.z) + CameraMotion[3].z)) + (CameraMotion[3].w * xyd.x) + (CameraMotion[4].x * xyd.y * xyd.y) + (CameraMotion[4].y * xyd.z) + CameraMotion[4].z) * scaleM;
-
-		// Note: have to sample from one of the high res texels exactly to avoid filtering of the velocity buffer
-		// 0 is stored where camera velocity should be used, which must not be filtered with valid object velocities
-		float2 FullResTexel = UVAndScreenPos.xy - .5f * View.ViewSizeAndSceneTexelSize.zw;
-		float2 VelocityN = Texture2DSampleLevel(VelocityTexture, VelocityTextureSampler, FullResTexel, 0).xy;
-
-		if (VelocityN.x > 0)
-		{
-			// Use the per-pixel velocity vector where valid, this handles object, bone and WPO movement
-			MotionVector = float2(-.5f, .5f) * DecodeVelocityFromTexture(VelocityN);
-		}
-
-		OldUV = UVAndScreenPos.xy + MotionVector / float2(0.5, -0.5) * View.ScreenPositionScaleBias.xy;
-	}
-	
-	float2 OldScreenPosition = (OldUV.xy - View.ScreenPositionScaleBias.wz) / View.ScreenPositionScaleBias.xy;
-
-	float2 OldDistanceFieldUVs = OldUV.xy - View.ViewRectMin.xy * View.ViewSizeAndSceneTexelSize.zw;
-	float4 HistoryValue = Texture2DSampleLevel(BentNormalHistoryTexture, BentNormalHistorySampler, OldDistanceFieldUVs, 0);
-	
-	float2 MinUV = View.ViewRectMin.xy * View.ViewSizeAndSceneTexelSize.zw;
-	// Pull in the max UV to exclude the region which will read outside the viewport due to bilinear filtering
-	float2 MaxUV = MinUV + (View.ViewSizeAndSceneTexelSize.xy - 1 * DOWNSAMPLE_FACTOR) * View.ViewSizeAndSceneTexelSize.zw;
-
-	float3 WorldPosition;
-	float3 PrevWorldPosition;
-
-	{
-		float4 HomogeneousWorldPosition = mul(float4(UVAndScreenPos.zw * SceneDepth, SceneDepth, 1), View.ScreenToWorld);
-		WorldPosition = HomogeneousWorldPosition.xyz / HomogeneousWorldPosition.w;
-
-		float HistoryDepth = abs(HistoryValue.w);
-		float4 PrevPositionTranslatedWorld = mul(float4(OldScreenPosition * HistoryDepth, HistoryDepth, 1), View.PrevScreenToTranslatedWorld);
-		PrevPositionTranslatedWorld.xyz /= PrevPositionTranslatedWorld.w;	
-		PrevWorldPosition = PrevPositionTranslatedWorld.xyz - View.PrevPreViewTranslation;
-	}
-
-	float EffectiveHistoryWeight = HistoryWeight;
-
-	FLATTEN
-	if (any(OldUV > MaxUV) 
-		|| any(OldUV < MinUV)
-		// Discard history if we are shading a new position (newly revealed)
-		|| length(PrevWorldPosition - WorldPosition) > HistoryDistanceThreshold)
-	{
-		EffectiveHistoryWeight = 0;
-	}
-
-	OutBentNormal.rgb = lerp(NewValue.rgb, HistoryValue.rgb, EffectiveHistoryWeight);
-
-#if SUPPORT_IRRADIANCE
-	float3 NewIrradiance = Texture2DSampleLevel(IrradianceTexture, IrradianceSampler, DistanceFieldUVs, 0).xyz;
-	float3 HistoryIrradianceValue = Texture2DSampleLevel(IrradianceHistoryTexture, IrradianceHistorySampler, OldDistanceFieldUVs, 0).xyz;
-
-	OutIrradiance = float4(lerp(NewIrradiance, HistoryIrradianceValue, EffectiveHistoryWeight), 0);
-#endif
-
-	OutBentNormal.rgb = isnan(OutBentNormal.rgb) ? 0 : OutBentNormal.rgb;
-
-	// Sign bit of alpha stores whether the history was rejected or not
-	OutBentNormal.a = abs(NewValue.a) * (EffectiveHistoryWeight > 0 ? 1 : -1);
-}
-
-#define HALF_HISTORY_FILL_KERNEL_SIZE 2
-
-/** Seeds newly rejected history values (which are sources of temporal instability) with the results of a spatial search from stable history values */
-void FilterHistoryPS(
-	in float4 UVAndScreenPos : TEXCOORD0
-	,out float4 OutBentNormal : SV_Target0
-#if SUPPORT_IRRADIANCE
-	,out float4 OutIrradiance : SV_Target1
-#endif
-	)
-{
-	float4 HistoryValue = Texture2DSampleLevel(BentNormalAOTexture, BentNormalAOSampler, UVAndScreenPos.xy, 0);
-
-#if SUPPORT_IRRADIANCE
-	float3 IrradianceValue = Texture2DSampleLevel(IrradianceTexture, IrradianceSampler, UVAndScreenPos.xy, 0).xyz;
-#endif
-
-	// Only do the spatial search for pixels who discarded their history value
-	if (HistoryValue.w < 0)
-	{
-		float SceneDepth = abs(HistoryValue.w);
-
-		float4 Accumulation = 0;
-		float3 IrradianceAccumulation = 0;
-
-		for (float y = -HALF_HISTORY_FILL_KERNEL_SIZE; y <= HALF_HISTORY_FILL_KERNEL_SIZE; y++)
-		{
-			for (float x = -HALF_HISTORY_FILL_KERNEL_SIZE; x <= HALF_HISTORY_FILL_KERNEL_SIZE; x++)
-			{
-				float2 UVOffset = BentNormalAOTexelSize * float2(x, y);
-				float4 TextureValue = Texture2DSampleLevel(BentNormalAOTexture, BentNormalAOSampler, UVAndScreenPos.xy + UVOffset, 0);
-
-				float SampleSceneDepth = abs(TextureValue.w);
-				float ValidMask = TextureValue.w > 0;
-
-				// Weight by depth to avoid pulling in values of a foreground object
-				// This is a careful tradeoff between ghosting behind panning foreground objects and successful spatial searches to reduce flickering
-				float DepthWeight = exp2(-1000 * abs(SceneDepth - SampleSceneDepth) / SceneDepth);
-
-				float2 Weight2D = exp2(-abs(float2(x, y) * 10.0f / HALF_HISTORY_FILL_KERNEL_SIZE));
-				float ScreenSpaceSpatialWeight = max(Weight2D.x, Weight2D.y);
-
-				float Weight = ValidMask * ScreenSpaceSpatialWeight * DepthWeight;
-
-				Accumulation.rgb += TextureValue.rgb * Weight;
-				Accumulation.a += Weight;
-
-				#if SUPPORT_IRRADIANCE
-					float3 Irradiance = Texture2DSampleLevel(IrradianceTexture, IrradianceSampler, UVAndScreenPos.xy + UVOffset, 0).xyz;
-					IrradianceAccumulation += Irradiance * Weight;
-				#endif
-			}
-		}
-
-		// Only change the history value if the spatial search turned up something applicable
-		if (Accumulation.a > 0)
-		{
-			float InvWeight = 1.0f / Accumulation.a;
-			// Construct the history value as if the spatial search result was the previous history, 
-			// And the AO we just computed this frame was the new value
-			HistoryValue.xyz = lerp(HistoryValue.xyz, Accumulation.xyz * InvWeight, HistoryWeight);
-
-			#if SUPPORT_IRRADIANCE
-				IrradianceValue = lerp(IrradianceValue, IrradianceAccumulation * InvWeight, HistoryWeight);
-			#endif
-		}
-	}
-
-	OutBentNormal = HistoryValue;
-	// Remove sign bit so future reprojection interpolation isn't affected
-	OutBentNormal.w = abs(OutBentNormal.w);
-
-#if SUPPORT_IRRADIANCE
-	OutIrradiance = float4(IrradianceValue, 0);
-#endif
-}
-
-/** Upsamples the AO results to full resolution using a bilateral filter. */
-void AOUpsamplePS(
-	in float4 UVAndScreenPos : TEXCOORD0
-	, out float4 OutBentNormal : SV_Target0
-#if OUTPUT_BENT_NORMAL && SUPPORT_IRRADIANCE
-	, out float4 OutIrradiance : SV_Target1
-#endif
-	)
-{
-	// Distance field AO was computed at 0,0 regardless of viewrect min
-	float2 DistanceFieldUVs = UVAndScreenPos.xy - View.ViewRectMin.xy * View.ViewSizeAndSceneTexelSize.zw;
-	float3 Irradiance = 0;
-
-#define BILATERAL_UPSAMPLE 1
-#if BILATERAL_UPSAMPLE
-	float2 LowResBufferSize = floor(View.RenderTargetSize / DOWNSAMPLE_FACTOR);
-	float2 LowResTexelSize = 1.0f / LowResBufferSize;
-	float2 Corner00UV = floor(DistanceFieldUVs * LowResBufferSize - .5f) / LowResBufferSize + .5f * LowResTexelSize;
-	float2 BilinearWeights = (DistanceFieldUVs - Corner00UV) * LowResBufferSize;
-
-	float4 TextureValues00 = Texture2DSampleLevel(BentNormalAOTexture, BentNormalAOSampler, Corner00UV, 0);
-	float4 TextureValues10 = Texture2DSampleLevel(BentNormalAOTexture, BentNormalAOSampler, Corner00UV + float2(LowResTexelSize.x, 0), 0);
-	float4 TextureValues01 = Texture2DSampleLevel(BentNormalAOTexture, BentNormalAOSampler, Corner00UV + float2(0, LowResTexelSize.y), 0);
-	float4 TextureValues11 = Texture2DSampleLevel(BentNormalAOTexture, BentNormalAOSampler, Corner00UV + LowResTexelSize, 0);
-
-	float4 CornerWeights = float4(
-		(1 - BilinearWeights.y) * (1 - BilinearWeights.x), 
-		(1 - BilinearWeights.y) * BilinearWeights.x,
-		BilinearWeights.y * (1 - BilinearWeights.x),
-		BilinearWeights.y * BilinearWeights.x);
-
-	float Epsilon = .0001f;
-
-	float4 CornerDepths = abs(float4(TextureValues00.w, TextureValues10.w, TextureValues01.w, TextureValues11.w));
-	float SceneDepth = CalcSceneDepth(UVAndScreenPos.xy);
-	float4 DepthWeights = 1.0f / (abs(CornerDepths - SceneDepth.xxxx) + Epsilon);
-	float4 FinalWeights = CornerWeights * DepthWeights;
-
-	float InvWeight = 1.0f / dot(FinalWeights, 1);
-
-	float3 InterpolatedResult = 
-		(FinalWeights.x * TextureValues00.xyz 
-		+ FinalWeights.y * TextureValues10.xyz
-		+ FinalWeights.z * TextureValues01.xyz 
-		+ FinalWeights.w * TextureValues11.xyz)
-		* InvWeight;
-
-	float3 BentNormal = InterpolatedResult.xyz;
-
-	#if SUPPORT_IRRADIANCE
-		float4 IrradianceValues00 = Texture2DSampleLevel(IrradianceTexture, IrradianceSampler, Corner00UV, 0);
-		float4 IrradianceValues10 = Texture2DSampleLevel(IrradianceTexture, IrradianceSampler, Corner00UV + float2(LowResTexelSize.x, 0), 0);
-		float4 IrradianceValues01 = Texture2DSampleLevel(IrradianceTexture, IrradianceSampler, Corner00UV + float2(0, LowResTexelSize.y), 0);
-		float4 IrradianceValues11 = Texture2DSampleLevel(IrradianceTexture, IrradianceSampler, Corner00UV + LowResTexelSize, 0);
-
-		Irradiance = 
-			(FinalWeights.x * IrradianceValues00.xyz 
-			+ FinalWeights.y * IrradianceValues10.xyz
-			+ FinalWeights.z * IrradianceValues01.xyz 
-			+ FinalWeights.w * IrradianceValues11.xyz)
-			* InvWeight;
-	#endif
-
-#else
-	float3 BentNormal = Texture2DSampleLevel(BentNormalAOTexture, BentNormalAOSampler, DistanceFieldUVs, 0).xyz;
-
-	#if SUPPORT_IRRADIANCE
-		Irradiance = Texture2DSampleLevel(IrradianceTexture, IrradianceSampler, DistanceFieldUVs, 0).xyz;
-	#endif
-#endif
-
-#if OUTPUT_AO
-	OutBentNormal = float4(length(BentNormal).xxx, 1);
-	//OutBentNormal = float4(BentNormal, 1);
-#else
-	#if OUTPUT_BENT_NORMAL
-		OutBentNormal = float4(BentNormal, 1);
-		#if SUPPORT_IRRADIANCE
-			OutIrradiance = float4(Irradiance, 1);
-		#endif
-	#else
-
-		// Output to RT0 to visualize Irradiance or BentNormal
-		#if SUPPORT_IRRADIANCE
-			OutBentNormal = float4(Irradiance, 1);
-		#endif
-	#endif
-#endif
-}
-
 RWTexture2D<float4> RWVisualizeMeshDistanceFields;
 
 bool RayHitSphere( float3 RayOrigin, float3 UnitRayDirection, float3 SphereCenter, float SphereRadius )
diff --git a/Engine/Shaders/DistanceFieldSurfaceCacheLightingCompute.usf b/Engine/Shaders/DistanceFieldSurfaceCacheLightingCompute.usf
--- a/Engine/Shaders/DistanceFieldSurfaceCacheLightingCompute.usf
+++ b/Engine/Shaders/DistanceFieldSurfaceCacheLightingCompute.usf
@@ -273,8 +273,12 @@ void BuildTileConesMain(
     uint3 GroupThreadId : SV_GroupThreadID) 
 {
     uint ThreadIndex = GroupThreadId.y * THREADGROUP_SIZEX + GroupThreadId.x;
-	float2 ScreenUV = (DispatchThreadId.xy * DOWNSAMPLE_FACTOR + View.ViewRectMin.xy + float2(.5f, .5f)) * View.ViewSizeAndSceneTexelSize.zw;
-	float SceneDepth = CalcSceneDepth(ScreenUV);
+
+	float2 BaseLevelScreenUV = (DispatchThreadId.xy + float2(.5f, .5f)) * DOWNSAMPLE_FACTOR * View.ViewSizeAndSceneTexelSize.zw;
+	float SceneDepth = GetDownsampledDepth(BaseLevelScreenUV);
+
+	//float2 ScreenUV = (DispatchThreadId.xy * DOWNSAMPLE_FACTOR + View.ViewRectMin.xy + float2(.5f, .5f)) * View.ViewSizeAndSceneTexelSize.zw;
+	//float SceneDepth = CalcSceneDepth(ScreenUV);
 
 	// Initialize per-tile variables
     if (ThreadIndex == 0) 
@@ -709,7 +713,8 @@ groupshared float3 SharedGatheredLighting[FINAL_GATHER_THREADGROUP_SIZE];
 groupshared uint SharedCulledObjectList[MAX_RECORD_CULLED_OBJECTS];
 groupshared uint NumRecordCulledObjects;
 
-#define CULL_OBJECTS_TO_RECORD 1
+// Inconsistent performance savings + overflow of shared memory array
+#define CULL_OBJECTS_TO_RECORD 0
 
 /** Computes ambient occlusion for a surface cache record by cone stepping through the nearby object distance fields. */
 [numthreads(FINAL_GATHER_THREADGROUP_SIZE, 1, 1)]
@@ -843,27 +848,28 @@ void ConeTraceOcclusionCS(
 				// Ignore extra threads
 				&& ObjectOffsetIndex < SIMULTANEOUSLY_TRACED_OBJECTS)
 			{
-				#if CULL_OBJECTS_TO_RECORD
-					uint ObjectIndex = SharedCulledObjectList[EffectiveListObjectIndex];
-					{
-				#else
 
-					uint ListIndex = 0;
-					uint ArrayIndex = EffectiveListObjectIndex;
+			#if CULL_OBJECTS_TO_RECORD
+				uint ObjectIndex = SharedCulledObjectList[EffectiveListObjectIndex];
+				{
+			#else
 
-					FLATTEN
-					if (EffectiveListObjectIndex >= TileHead.y + TileHead.z)
-					{
-						ListIndex = 2;
-						ArrayIndex = EffectiveListObjectIndex - TileHead.y - TileHead.z;
-					}
-					else if (EffectiveListObjectIndex >= TileHead.y)
-					{
-						ListIndex = 1;
-						ArrayIndex = EffectiveListObjectIndex - TileHead.y;
-					}
+				uint ListIndex = 0;
+				uint ArrayIndex = EffectiveListObjectIndex;
 
-					uint ObjectIndex = TileArrayData.Load((ArrayIndex * TileListGroupSize.x * TileListGroupSize.y + TileHead.x) * NUM_CULLED_OBJECT_LISTS + ListIndex);
+				FLATTEN
+				if (EffectiveListObjectIndex >= TileHead.y + TileHead.z)
+				{
+					ListIndex = 2;
+					ArrayIndex = EffectiveListObjectIndex - TileHead.y - TileHead.z;
+				}
+				else if (EffectiveListObjectIndex >= TileHead.y)
+				{
+					ListIndex = 1;
+					ArrayIndex = EffectiveListObjectIndex - TileHead.y;
+				}
+
+				uint ObjectIndex = TileArrayData.Load((ArrayIndex * TileListGroupSize.x * TileListGroupSize.y + TileHead.x) * NUM_CULLED_OBJECT_LISTS + ListIndex);
 
 				float4 ObjectPositionAndRadius = LoadObjectPositionAndRadius(ObjectIndex);
 				float ObjectDistanceSq = dot(ObjectPositionAndRadius.xyz - WorldShadingPosition, ObjectPositionAndRadius.xyz - WorldShadingPosition);
@@ -873,7 +879,7 @@ void ConeTraceOcclusionCS(
 				if (ObjectPositionAndRadius.w * ObjectPositionAndRadius.w / ObjectDistanceSq > Square(.25f))
 				{
 
-				#endif
+			#endif
 					float3 LocalPositionExtent = LoadObjectLocalPositionExtent(ObjectIndex);
 					float4x4 WorldToVolume = LoadObjectWorldToVolume(ObjectIndex);
 					bool bGeneratedAsTwoSided;
diff --git a/Engine/Source/Runtime/Engine/Private/DistanceFieldAtlas.cpp b/Engine/Source/Runtime/Engine/Private/DistanceFieldAtlas.cpp
--- a/Engine/Source/Runtime/Engine/Private/DistanceFieldAtlas.cpp
+++ b/Engine/Source/Runtime/Engine/Private/DistanceFieldAtlas.cpp
@@ -120,16 +120,13 @@ void FDistanceFieldVolumeTextureAtlas::UpdateAllocations()
 
 			FRHIResourceCreateInfo CreateInfo;
 
-			int32 SizeX = BlockAllocator.GetSizeX();
-			int32 SizeY = BlockAllocator.GetSizeY();
-
 			VolumeTextureRHI = RHICreateTexture3D(
-				SizeX, 
-				SizeY, 
+				BlockAllocator.GetSizeX(), 
+				BlockAllocator.GetSizeY(), 
 				BlockAllocator.GetSizeZ(), 
 				Format,
 				1,
-				TexCreate_ShaderResource | TexCreate_NoTiling,
+				TexCreate_ShaderResource,
 				CreateInfo);
 
 			UE_LOG(LogStaticMesh,Log,TEXT("Allocated %s"), *GetSizeString());
diff --git a/Engine/Source/Runtime/Renderer/Private/DistanceFieldSurfaceCacheLighting.cpp b/Engine/Source/Runtime/Renderer/Private/DistanceFieldSurfaceCacheLighting.cpp
--- a/Engine/Source/Runtime/Renderer/Private/DistanceFieldSurfaceCacheLighting.cpp
+++ b/Engine/Source/Runtime/Renderer/Private/DistanceFieldSurfaceCacheLighting.cpp
@@ -86,6 +86,14 @@ FAutoConsoleVariableRef CVarAOFillGaps(
 	ECVF_Cheat | ECVF_RenderThreadSafe
 	);
 
+int32 GAOFillGapsHighQuality = 1;
+FAutoConsoleVariableRef CVarAOFillGapsHighQuality(
+	TEXT("r.AOFillGapsHighQuality"),
+	GAOFillGapsHighQuality,
+	TEXT("Whether to use the higher quality gap filling method that does a 5x5 gather, or the cheaper method that just looks up the 4 grid samples."),
+	ECVF_Cheat | ECVF_RenderThreadSafe
+	);
+
 int32 GAOUseHistory = 1;
 FAutoConsoleVariableRef CVarAOUseHistory(
 	TEXT("r.AOUseHistory"),
@@ -94,6 +102,14 @@ FAutoConsoleVariableRef CVarAOUseHistory(
 	ECVF_Cheat | ECVF_RenderThreadSafe
 	);
 
+int32 GAOHistoryStabilityPass = 1;
+FAutoConsoleVariableRef CVarAOHistoryStabilityPass(
+	TEXT("r.AOHistoryStabilityPass"),
+	GAOHistoryStabilityPass,
+	TEXT("Whether to gather stable results to fill in holes in the temporal reprojection.  Adds some GPU cost but improves temporal stability with foliage."),
+	ECVF_Cheat | ECVF_RenderThreadSafe
+	);
+
 float GAOHistoryWeight = .7f;
 FAutoConsoleVariableRef CVarAOHistoryWeight(
 	TEXT("r.AOHistoryWeight"),
@@ -183,7 +199,7 @@ FAutoConsoleVariableRef CVarAOScatterTileCulling(
 	ECVF_Cheat | ECVF_RenderThreadSafe
 	);
 
-int32 GAOComputeShaderNormalCalculation = 1;
+int32 GAOComputeShaderNormalCalculation = 0;
 FAutoConsoleVariableRef CVarAOComputeShaderNormalCalculation(
 	TEXT("r.AOComputeShaderNormalCalculation"),
 	GAOComputeShaderNormalCalculation,
@@ -207,7 +223,7 @@ FAutoConsoleVariableRef CVarAOInterpolationStencilTesting(
 	ECVF_Cheat | ECVF_RenderThreadSafe
 	);
 
-int32 GAOInterpolationDepthTesting = 0;
+int32 GAOInterpolationDepthTesting = 1;
 FAutoConsoleVariableRef CVarAOInterpolationDepthTesting(
 	TEXT("r.AOInterpolationDepthTesting"),
 	GAOInterpolationDepthTesting,
@@ -506,12 +522,14 @@ public:
 		TileHeadDataUnpacked.Bind(Initializer.ParameterMap, TEXT("TileHeadDataUnpacked"));
 		NumGroups.Bind(Initializer.ParameterMap, TEXT("NumGroups"));
 		ViewDimensionsParameter.Bind(Initializer.ParameterMap, TEXT("ViewDimensions"));
+		DistanceFieldNormalTexture.Bind(Initializer.ParameterMap, TEXT("DistanceFieldNormalTexture"));
+		DistanceFieldNormalSampler.Bind(Initializer.ParameterMap, TEXT("DistanceFieldNormalSampler"));
 	}
 
 	FBuildTileConesCS()
 	{
 	}
-	void SetParameters(FRHICommandList& RHICmdList, const FSceneView& View, FScene* Scene, FVector2D NumGroupsValue, const FDistanceFieldAOParameters& Parameters)
+	void SetParameters(FRHICommandList& RHICmdList, const FSceneView& View, FSceneRenderTargetItem& DistanceFieldNormal, FScene* Scene, FVector2D NumGroupsValue, const FDistanceFieldAOParameters& Parameters)
 	{
 		FComputeShaderRHIParamRef ShaderRHI = GetComputeShader();
 
@@ -528,6 +546,15 @@ public:
 		SetShaderValue(RHICmdList, ShaderRHI, ViewDimensionsParameter, View.ViewRect);
 
 		SetShaderValue(RHICmdList, ShaderRHI, NumGroups, NumGroupsValue);
+
+		SetTextureParameter(
+			RHICmdList,
+			ShaderRHI,
+			DistanceFieldNormalTexture,
+			DistanceFieldNormalSampler,
+			TStaticSamplerState<SF_Point,AM_Clamp,AM_Clamp,AM_Clamp>::GetRHI(),
+			DistanceFieldNormal.ShaderResourceTexture
+			);
 	}
 
 	void UnsetParameters(FRHICommandList& RHICmdList)
@@ -547,6 +574,8 @@ public:
 		Ar << TileHeadDataUnpacked;
 		Ar << NumGroups;
 		Ar << ViewDimensionsParameter;
+		Ar << DistanceFieldNormalTexture;
+		Ar << DistanceFieldNormalSampler;
 		return bShaderHasOutdatedParameters;
 	}
 
@@ -559,6 +588,8 @@ private:
 	FRWShaderParameter TileHeadDataUnpacked;
 	FShaderParameter ViewDimensionsParameter;
 	FShaderParameter NumGroups;
+	FShaderResourceParameter DistanceFieldNormalTexture;
+	FShaderResourceParameter DistanceFieldNormalSampler;
 };
 
 IMPLEMENT_SHADER_TYPE(,FBuildTileConesCS,TEXT("DistanceFieldSurfaceCacheLightingCompute"),TEXT("BuildTileConesMain"),SF_Compute);
@@ -872,12 +903,7 @@ public:
 	FComputeDistanceFieldNormalPS(const ShaderMetaType::CompiledShaderInitializerType& Initializer)
 		: FGlobalShader(Initializer)
 	{
-		ObjectParameters.Bind(Initializer.ParameterMap);
 		DeferredParameters.Bind(Initializer.ParameterMap);
-		TileHeadData.Bind(Initializer.ParameterMap, TEXT("TileHeadData"));
-		TileHeadDataUnpacked.Bind(Initializer.ParameterMap, TEXT("TileHeadDataUnpacked"));
-		TileArrayData.Bind(Initializer.ParameterMap, TEXT("TileArrayData"));
-		TileListGroupSize.Bind(Initializer.ParameterMap, TEXT("TileListGroupSize"));
 		AOParameters.Bind(Initializer.ParameterMap);
 	}
 
@@ -886,40 +912,21 @@ public:
 		const FPixelShaderRHIParamRef ShaderRHI = GetPixelShader();
 
 		FGlobalShader::SetParameters(RHICmdList, ShaderRHI, View);
-		ObjectParameters.Set(RHICmdList, ShaderRHI, GAOCulledObjectBuffers.Buffers);
 		AOParameters.Set(RHICmdList, ShaderRHI, Parameters);
 		DeferredParameters.Set(RHICmdList, ShaderRHI, View);
-
-		FTileIntersectionResources* TileIntersectionResources = ((FSceneViewState*)View.State)->AOTileIntersectionResources;
-
-		SetSRVParameter(RHICmdList, ShaderRHI, TileHeadData, TileIntersectionResources->TileHeadData.SRV);
-		SetSRVParameter(RHICmdList, ShaderRHI, TileHeadDataUnpacked, TileIntersectionResources->TileHeadDataUnpacked.SRV);
-		SetSRVParameter(RHICmdList, ShaderRHI, TileArrayData, TileIntersectionResources->TileArrayData.SRV);
-
-		SetShaderValue(RHICmdList, ShaderRHI, TileListGroupSize, TileIntersectionResources->TileDimensions);
 	}
 	// FShader interface.
 	virtual bool Serialize(FArchive& Ar) override
 	{
 		bool bShaderHasOutdatedParameters = FGlobalShader::Serialize(Ar);
-		Ar << ObjectParameters;
 		Ar << DeferredParameters;
-		Ar << TileHeadData;
-		Ar << TileHeadDataUnpacked;
-		Ar << TileArrayData;
-		Ar << TileListGroupSize;
 		Ar << AOParameters;
 		return bShaderHasOutdatedParameters;
 	}
 
 private:
 
-	FDistanceFieldCulledObjectBufferParameters ObjectParameters;
 	FDeferredPixelShaderParameters DeferredParameters;
-	FShaderResourceParameter TileHeadData;
-	FShaderResourceParameter TileHeadDataUnpacked;
-	FShaderResourceParameter TileArrayData;
-	FShaderParameter TileListGroupSize;
 	FAOParameters AOParameters;
 };
 
@@ -951,13 +958,7 @@ public:
 		: FGlobalShader(Initializer)
 	{
 		DistanceFieldNormal.Bind(Initializer.ParameterMap, TEXT("DistanceFieldNormal"));
-		ObjectParameters.Bind(Initializer.ParameterMap);
 		DeferredParameters.Bind(Initializer.ParameterMap);
-		UVToTileHead.Bind(Initializer.ParameterMap, TEXT("UVToTileHead"));
-		TileHeadData.Bind(Initializer.ParameterMap, TEXT("TileHeadData"));
-		TileHeadDataUnpacked.Bind(Initializer.ParameterMap, TEXT("TileHeadDataUnpacked"));
-		TileArrayData.Bind(Initializer.ParameterMap, TEXT("TileArrayData"));
-		TileListGroupSize.Bind(Initializer.ParameterMap, TEXT("TileListGroupSize"));
 		AOParameters.Bind(Initializer.ParameterMap);
 	}
 
@@ -968,20 +969,8 @@ public:
 		FGlobalShader::SetParameters(RHICmdList, ShaderRHI, View);
 
 		DistanceFieldNormal.SetTexture(RHICmdList, ShaderRHI, DistanceFieldNormalValue.ShaderResourceTexture, DistanceFieldNormalValue.UAV);
-
-		ObjectParameters.Set(RHICmdList, ShaderRHI, GAOCulledObjectBuffers.Buffers);
 		AOParameters.Set(RHICmdList, ShaderRHI, Parameters);
 		DeferredParameters.Set(RHICmdList, ShaderRHI, View);
-
-		SetShaderValue(RHICmdList, ShaderRHI, UVToTileHead, FVector2D(View.ViewRect.Size().X / (float)(GDistanceFieldAOTileSizeX * GAODownsampleFactor), View.ViewRect.Size().Y / (float)(GDistanceFieldAOTileSizeY * GAODownsampleFactor)));
-	
-		FTileIntersectionResources* TileIntersectionResources = ((FSceneViewState*)View.State)->AOTileIntersectionResources;
-
-		SetSRVParameter(RHICmdList, ShaderRHI, TileHeadData, TileIntersectionResources->TileHeadData.SRV);
-		SetSRVParameter(RHICmdList, ShaderRHI, TileHeadDataUnpacked, TileIntersectionResources->TileHeadDataUnpacked.SRV);
-		SetSRVParameter(RHICmdList, ShaderRHI, TileArrayData, TileIntersectionResources->TileArrayData.SRV);
-
-		SetShaderValue(RHICmdList, ShaderRHI, TileListGroupSize, TileIntersectionResources->TileDimensions);
 	}
 
 	void UnsetParameters(FRHICommandList& RHICmdList)
@@ -994,13 +983,7 @@ public:
 	{
 		bool bShaderHasOutdatedParameters = FGlobalShader::Serialize(Ar);
 		Ar << DistanceFieldNormal;
-		Ar << ObjectParameters;
-		Ar << UVToTileHead;
 		Ar << DeferredParameters;
-		Ar << TileHeadData;
-		Ar << TileHeadDataUnpacked;
-		Ar << TileArrayData;
-		Ar << TileListGroupSize;
 		Ar << AOParameters;
 		return bShaderHasOutdatedParameters;
 	}
@@ -1011,10 +994,6 @@ private:
 	FDistanceFieldCulledObjectBufferParameters ObjectParameters;
 	FShaderParameter UVToTileHead;
 	FDeferredPixelShaderParameters DeferredParameters;
-	FShaderResourceParameter TileHeadData;
-	FShaderResourceParameter TileHeadDataUnpacked;
-	FShaderResourceParameter TileArrayData;
-	FShaderParameter TileListGroupSize;
 	FAOParameters AOParameters;
 };
 
@@ -2229,7 +2208,7 @@ private:
 IMPLEMENT_SHADER_TYPE(template<>,TDistanceFieldAOCombinePS<true>,TEXT("DistanceFieldSurfaceCacheLighting"),TEXT("AOCombinePS"),SF_Pixel);
 IMPLEMENT_SHADER_TYPE(template<>,TDistanceFieldAOCombinePS<false>,TEXT("DistanceFieldSurfaceCacheLighting"),TEXT("AOCombinePS"),SF_Pixel);
 
-template<bool bSupportIrradiance>
+template<bool bSupportIrradiance, bool bHighQuality>
 class TFillGapsPS : public FGlobalShader
 {
 	DECLARE_SHADER_TYPE(TFillGapsPS, Global);
@@ -2244,6 +2223,7 @@ public:
 	{
 		OutEnvironment.SetDefine(TEXT("DOWNSAMPLE_FACTOR"), GAODownsampleFactor);
 		OutEnvironment.SetDefine(TEXT("SUPPORT_IRRADIANCE"), bSupportIrradiance ? TEXT("1") : TEXT("0"));
+		OutEnvironment.SetDefine(TEXT("HIGH_QUALITY_FILL_GAPS"), bHighQuality ? TEXT("1") : TEXT("0"));
 	}
 
 	/** Default constructor. */
@@ -2258,6 +2238,7 @@ public:
 		IrradianceTexture.Bind(Initializer.ParameterMap, TEXT("IrradianceTexture"));
 		IrradianceSampler.Bind(Initializer.ParameterMap, TEXT("IrradianceSampler"));
 		BentNormalAOTexelSize.Bind(Initializer.ParameterMap, TEXT("BentNormalAOTexelSize"));
+		MinDownsampleFactorToBaseLevel.Bind(Initializer.ParameterMap, TEXT("MinDownsampleFactorToBaseLevel"));
 	}
 
 	void SetParameters(FRHICommandList& RHICmdList, const FSceneView& View, FSceneRenderTargetItem& DistanceFieldAOBentNormal2, IPooledRenderTarget* DistanceFieldIrradiance2)
@@ -2290,6 +2271,9 @@ public:
 		const FIntPoint DownsampledBufferSize(GSceneRenderTargets.GetBufferSizeXY() / FIntPoint(GAODownsampleFactor, GAODownsampleFactor));
 		const FVector2D BaseLevelTexelSizeValue(1.0f / DownsampledBufferSize.X, 1.0f / DownsampledBufferSize.Y);
 		SetShaderValue(RHICmdList, ShaderRHI, BentNormalAOTexelSize, BaseLevelTexelSizeValue);
+
+		const float MinDownsampleFactor = 1 << (GAOMinLevel * GAOPowerOfTwoBetweenLevels); 
+		SetShaderValue(RHICmdList, ShaderRHI, MinDownsampleFactorToBaseLevel, MinDownsampleFactor);
 	}
 	// FShader interface.
 	virtual bool Serialize(FArchive& Ar) override
@@ -2300,6 +2284,7 @@ public:
 		Ar << IrradianceTexture;
 		Ar << IrradianceSampler;
 		Ar << BentNormalAOTexelSize;
+		Ar << MinDownsampleFactorToBaseLevel;
 		return bShaderHasOutdatedParameters;
 	}
 
@@ -2310,10 +2295,17 @@ private:
 	FShaderResourceParameter IrradianceTexture;
 	FShaderResourceParameter IrradianceSampler;
 	FShaderParameter BentNormalAOTexelSize;
+	FShaderParameter MinDownsampleFactorToBaseLevel;
 };
 
-IMPLEMENT_SHADER_TYPE(template<>,TFillGapsPS<true>,TEXT("DistanceFieldSurfaceCacheLighting"),TEXT("FillGapsPS"),SF_Pixel);
-IMPLEMENT_SHADER_TYPE(template<>,TFillGapsPS<false>,TEXT("DistanceFieldSurfaceCacheLighting"),TEXT("FillGapsPS"),SF_Pixel);
+#define IMPLEMENT_FILLGAPS_PS_TYPE(bSupportIrradiance,bHighQuality) \
+	typedef TFillGapsPS<bSupportIrradiance, bHighQuality> TFillGapsPS##bSupportIrradiance##bHighQuality; \
+	IMPLEMENT_SHADER_TYPE(template<>,TFillGapsPS##bSupportIrradiance##bHighQuality,TEXT("DistanceFieldLightingPost"),TEXT("FillGapsPS"),SF_Pixel); 
+
+IMPLEMENT_FILLGAPS_PS_TYPE(true, true);
+IMPLEMENT_FILLGAPS_PS_TYPE(true, false);
+IMPLEMENT_FILLGAPS_PS_TYPE(false, true);
+IMPLEMENT_FILLGAPS_PS_TYPE(false, false);
 
 template<bool bSupportIrradiance>
 class TUpdateHistoryPS : public FGlobalShader
@@ -2351,6 +2343,7 @@ public:
 		CameraMotion.Bind(Initializer.ParameterMap);
 		HistoryWeight.Bind(Initializer.ParameterMap, TEXT("HistoryWeight"));
 		HistoryDistanceThreshold.Bind(Initializer.ParameterMap, TEXT("HistoryDistanceThreshold"));
+		UseHistoryFilter.Bind(Initializer.ParameterMap, TEXT("UseHistoryFilter"));
 		VelocityTexture.Bind(Initializer.ParameterMap, TEXT("VelocityTexture"));
 		VelocityTextureSampler.Bind(Initializer.ParameterMap, TEXT("VelocityTextureSampler"));
 	}
@@ -2415,6 +2408,7 @@ public:
 
 		SetShaderValue(RHICmdList, ShaderRHI, HistoryWeight, GAOHistoryWeight);
 		SetShaderValue(RHICmdList, ShaderRHI, HistoryDistanceThreshold, GAOHistoryDistanceThreshold);
+		SetShaderValue(RHICmdList, ShaderRHI, UseHistoryFilter, (GAOHistoryStabilityPass ? 1.0f : 0.0f));
 
 		SetTextureParameter(
 			RHICmdList,
@@ -2441,6 +2435,7 @@ public:
 		Ar << CameraMotion;
 		Ar << HistoryWeight;
 		Ar << HistoryDistanceThreshold;
+		Ar << UseHistoryFilter;
 		Ar << VelocityTexture;
 		Ar << VelocityTextureSampler;
 		return bShaderHasOutdatedParameters;
@@ -2460,12 +2455,13 @@ private:
 	FCameraMotionParameters CameraMotion;
 	FShaderParameter HistoryWeight;
 	FShaderParameter HistoryDistanceThreshold;
+	FShaderParameter UseHistoryFilter;
 	FShaderResourceParameter VelocityTexture;
 	FShaderResourceParameter VelocityTextureSampler;
 };
 
-IMPLEMENT_SHADER_TYPE(template<>,TUpdateHistoryPS<true>,TEXT("DistanceFieldSurfaceCacheLighting"),TEXT("UpdateHistoryPS"),SF_Pixel);
-IMPLEMENT_SHADER_TYPE(template<>,TUpdateHistoryPS<false>,TEXT("DistanceFieldSurfaceCacheLighting"),TEXT("UpdateHistoryPS"),SF_Pixel);
+IMPLEMENT_SHADER_TYPE(template<>,TUpdateHistoryPS<true>,TEXT("DistanceFieldLightingPost"),TEXT("UpdateHistoryPS"),SF_Pixel);
+IMPLEMENT_SHADER_TYPE(template<>,TUpdateHistoryPS<false>,TEXT("DistanceFieldLightingPost"),TEXT("UpdateHistoryPS"),SF_Pixel);
 
 
 template<bool bSupportIrradiance>
@@ -2561,8 +2557,8 @@ private:
 	FShaderParameter BentNormalAOTexelSize;
 };
 
-IMPLEMENT_SHADER_TYPE(template<>,TFilterHistoryPS<true>,TEXT("DistanceFieldSurfaceCacheLighting"),TEXT("FilterHistoryPS"),SF_Pixel);
-IMPLEMENT_SHADER_TYPE(template<>,TFilterHistoryPS<false>,TEXT("DistanceFieldSurfaceCacheLighting"),TEXT("FilterHistoryPS"),SF_Pixel);
+IMPLEMENT_SHADER_TYPE(template<>,TFilterHistoryPS<true>,TEXT("DistanceFieldLightingPost"),TEXT("FilterHistoryPS"),SF_Pixel);
+IMPLEMENT_SHADER_TYPE(template<>,TFilterHistoryPS<false>,TEXT("DistanceFieldLightingPost"),TEXT("FilterHistoryPS"),SF_Pixel);
 
 
 void AllocateOrReuseAORenderTarget(TRefCountPtr<IPooledRenderTarget>& Target, const TCHAR* Name, bool bWithAlpha)
@@ -2658,74 +2654,100 @@ void UpdateHistory(
 					FIntPoint(View.ViewRect.Width() / GAODownsampleFactor, View.ViewRect.Height() / GAODownsampleFactor),
 					GSceneRenderTargets.GetBufferSizeXY() / FIntPoint(GAODownsampleFactor, GAODownsampleFactor),
 					*VertexShader);
-			}
-
-			const FPooledRenderTargetDesc& HistoryDesc = (*BentNormalHistoryState)->GetDesc();
-
-			// Reallocate history if buffer sizes have changed
-			if (HistoryDesc.Extent != GSceneRenderTargets.GetBufferSizeXY() / FIntPoint(GAODownsampleFactor, GAODownsampleFactor))
-			{
-				GRenderTargetPool.FreeUnusedResource(*BentNormalHistoryState);
-				*BentNormalHistoryState = NULL;
-				// Update the view state's render target reference with the new history
-				AllocateOrReuseAORenderTarget(*BentNormalHistoryState, BentNormalHistoryRTName, true);
 
+				RHICmdList.CopyToResolveTarget(NewBentNormalHistory->GetRenderTargetItem().TargetableTexture, NewBentNormalHistory->GetRenderTargetItem().ShaderResourceTexture, false, FResolveParams());
+			
 				if (bUseDistanceFieldGI)
 				{
-					GRenderTargetPool.FreeUnusedResource(*IrradianceHistoryState);
-					*IrradianceHistoryState = NULL;
-					AllocateOrReuseAORenderTarget(*IrradianceHistoryState, IrradianceHistoryRTName, false);
+					RHICmdList.CopyToResolveTarget(NewIrradianceHistory->GetRenderTargetItem().TargetableTexture, NewIrradianceHistory->GetRenderTargetItem().ShaderResourceTexture, false, FResolveParams());
 				}
 			}
 
+			if (GAOHistoryStabilityPass)
 			{
-				FTextureRHIParamRef RenderTargets[2] =
+				const FPooledRenderTargetDesc& HistoryDesc = (*BentNormalHistoryState)->GetDesc();
+
+				// Reallocate history if buffer sizes have changed
+				if (HistoryDesc.Extent != GSceneRenderTargets.GetBufferSizeXY() / FIntPoint(GAODownsampleFactor, GAODownsampleFactor))
 				{
-					(*BentNormalHistoryState)->GetRenderTargetItem().TargetableTexture,
-					bUseDistanceFieldGI ? (*IrradianceHistoryState)->GetRenderTargetItem().TargetableTexture : NULL,
-				};
+					GRenderTargetPool.FreeUnusedResource(*BentNormalHistoryState);
+					*BentNormalHistoryState = NULL;
+					// Update the view state's render target reference with the new history
+					AllocateOrReuseAORenderTarget(*BentNormalHistoryState, BentNormalHistoryRTName, true);
 
-				SetRenderTargets(RHICmdList, ARRAY_COUNT(RenderTargets) - (bUseDistanceFieldGI ? 0 : 1), RenderTargets, FTextureRHIParamRef(), 0, NULL);
+					if (bUseDistanceFieldGI)
+					{
+						GRenderTargetPool.FreeUnusedResource(*IrradianceHistoryState);
+						*IrradianceHistoryState = NULL;
+						AllocateOrReuseAORenderTarget(*IrradianceHistoryState, IrradianceHistoryRTName, false);
+					}
+				}
 
-				RHICmdList.SetViewport(0, 0, 0.0f, View.ViewRect.Width() / GAODownsampleFactor, View.ViewRect.Height() / GAODownsampleFactor, 1.0f);
-				RHICmdList.SetRasterizerState(TStaticRasterizerState<FM_Solid, CM_None>::GetRHI());
-				RHICmdList.SetDepthStencilState(TStaticDepthStencilState<false, CF_Always>::GetRHI());
-				RHICmdList.SetBlendState(TStaticBlendState<>::GetRHI());
+				{
+					FTextureRHIParamRef RenderTargets[2] =
+					{
+						(*BentNormalHistoryState)->GetRenderTargetItem().TargetableTexture,
+						bUseDistanceFieldGI ? (*IrradianceHistoryState)->GetRenderTargetItem().TargetableTexture : NULL,
+					};
 
-				TShaderMapRef<FPostProcessVS> VertexShader(View.ShaderMap);
+					SetRenderTargets(RHICmdList, ARRAY_COUNT(RenderTargets) - (bUseDistanceFieldGI ? 0 : 1), RenderTargets, FTextureRHIParamRef(), 0, NULL);
 
-				if (bUseDistanceFieldGI)
-				{
-					TShaderMapRef<TFilterHistoryPS<true> > PixelShader(View.ShaderMap);
+					RHICmdList.SetViewport(0, 0, 0.0f, View.ViewRect.Width() / GAODownsampleFactor, View.ViewRect.Height() / GAODownsampleFactor, 1.0f);
+					RHICmdList.SetRasterizerState(TStaticRasterizerState<FM_Solid, CM_None>::GetRHI());
+					RHICmdList.SetDepthStencilState(TStaticDepthStencilState<false, CF_Always>::GetRHI());
+					RHICmdList.SetBlendState(TStaticBlendState<>::GetRHI());
 
-					static FGlobalBoundShaderState BoundShaderState;
-					SetGlobalBoundShaderState(RHICmdList, View.GetFeatureLevel(), BoundShaderState, GFilterVertexDeclaration.VertexDeclarationRHI, *VertexShader, *PixelShader);
-					PixelShader->SetParameters(RHICmdList, View, NewBentNormalHistory->GetRenderTargetItem(), NewIrradianceHistory);
-				}
-				else
-				{
-					TShaderMapRef<TFilterHistoryPS<false> > PixelShader(View.ShaderMap);
+					TShaderMapRef<FPostProcessVS> VertexShader(View.ShaderMap);
 
-					static FGlobalBoundShaderState BoundShaderState;
-					SetGlobalBoundShaderState(RHICmdList, View.GetFeatureLevel(), BoundShaderState, GFilterVertexDeclaration.VertexDeclarationRHI, *VertexShader, *PixelShader);
-					PixelShader->SetParameters(RHICmdList, View, NewBentNormalHistory->GetRenderTargetItem(), NewIrradianceHistory);
-				}
+					if (bUseDistanceFieldGI)
+					{
+						TShaderMapRef<TFilterHistoryPS<true> > PixelShader(View.ShaderMap);
 
-				VertexShader->SetParameters(RHICmdList, View);
+						static FGlobalBoundShaderState BoundShaderState;
+						SetGlobalBoundShaderState(RHICmdList, View.GetFeatureLevel(), BoundShaderState, GFilterVertexDeclaration.VertexDeclarationRHI, *VertexShader, *PixelShader);
+						PixelShader->SetParameters(RHICmdList, View, NewBentNormalHistory->GetRenderTargetItem(), NewIrradianceHistory);
+					}
+					else
+					{
+						TShaderMapRef<TFilterHistoryPS<false> > PixelShader(View.ShaderMap);
 
-				DrawRectangle( 
-					RHICmdList,
-					0, 0, 
-					View.ViewRect.Width() / GAODownsampleFactor, View.ViewRect.Height() / GAODownsampleFactor,
-					0, 0,
-					View.ViewRect.Width() / GAODownsampleFactor, View.ViewRect.Height() / GAODownsampleFactor,
-					FIntPoint(View.ViewRect.Width() / GAODownsampleFactor, View.ViewRect.Height() / GAODownsampleFactor),
-					GSceneRenderTargets.GetBufferSizeXY() / FIntPoint(GAODownsampleFactor, GAODownsampleFactor),
-					*VertexShader);
+						static FGlobalBoundShaderState BoundShaderState;
+						SetGlobalBoundShaderState(RHICmdList, View.GetFeatureLevel(), BoundShaderState, GFilterVertexDeclaration.VertexDeclarationRHI, *VertexShader, *PixelShader);
+						PixelShader->SetParameters(RHICmdList, View, NewBentNormalHistory->GetRenderTargetItem(), NewIrradianceHistory);
+					}
+
+					VertexShader->SetParameters(RHICmdList, View);
+
+					DrawRectangle( 
+						RHICmdList,
+						0, 0, 
+						View.ViewRect.Width() / GAODownsampleFactor, View.ViewRect.Height() / GAODownsampleFactor,
+						0, 0,
+						View.ViewRect.Width() / GAODownsampleFactor, View.ViewRect.Height() / GAODownsampleFactor,
+						FIntPoint(View.ViewRect.Width() / GAODownsampleFactor, View.ViewRect.Height() / GAODownsampleFactor),
+						GSceneRenderTargets.GetBufferSizeXY() / FIntPoint(GAODownsampleFactor, GAODownsampleFactor),
+						*VertexShader);
+
+					RHICmdList.CopyToResolveTarget((*BentNormalHistoryState)->GetRenderTargetItem().TargetableTexture, (*BentNormalHistoryState)->GetRenderTargetItem().ShaderResourceTexture, false, FResolveParams());
+			
+					if (bUseDistanceFieldGI)
+					{
+						RHICmdList.CopyToResolveTarget((*IrradianceHistoryState)->GetRenderTargetItem().TargetableTexture, (*IrradianceHistoryState)->GetRenderTargetItem().ShaderResourceTexture, false, FResolveParams());
+					}
+				}
+
+				BentNormalHistoryOutput = *BentNormalHistoryState;
+				IrradianceHistoryOutput = *IrradianceHistoryState;
 			}
+			else
+			{
+				// Update the view state's render target reference with the new history
+				*BentNormalHistoryState = NewBentNormalHistory;
+				BentNormalHistoryOutput = NewBentNormalHistory;
 
-			BentNormalHistoryOutput = *BentNormalHistoryState;
-			IrradianceHistoryOutput = *IrradianceHistoryState;
+				*IrradianceHistoryState = NewIrradianceHistory;
+				IrradianceHistoryOutput = NewIrradianceHistory;
+			}
 		}
 		else
 		{
@@ -2785,17 +2807,17 @@ void PostProcessBentNormalAO(
 	{
 		SCOPED_DRAW_EVENT(RHICmdList, AOCombine);
 
-		FTextureRHIParamRef DiffuseIrradianceTarget = NULL;
+		TRefCountPtr<IPooledRenderTarget> EffectiveDiffuseIrradianceTarget;
 
 		if (bUseDistanceFieldGI)
 		{
-			DiffuseIrradianceTarget = GAOFillGaps ? DistanceFieldIrradiance2->GetRenderTargetItem().TargetableTexture : DistanceFieldIrradiance->GetRenderTargetItem().TargetableTexture;
+			EffectiveDiffuseIrradianceTarget = GAOFillGaps ? DistanceFieldIrradiance2 : DistanceFieldIrradiance;
 		}
 
 		FTextureRHIParamRef RenderTargets[2] =
 		{
 			GAOFillGaps ? DistanceFieldAOBentNormal2->GetRenderTargetItem().TargetableTexture : DistanceFieldAOBentNormal->GetRenderTargetItem().TargetableTexture,
-			DiffuseIrradianceTarget,
+			EffectiveDiffuseIrradianceTarget ? EffectiveDiffuseIrradianceTarget->GetRenderTargetItem().TargetableTexture : NULL,
 		};
 
 		SetRenderTargets(RHICmdList, ARRAY_COUNT(RenderTargets) - (bUseDistanceFieldGI ? 0 : 1), RenderTargets, FTextureRHIParamRef(), 0, NULL);
@@ -2833,6 +2855,13 @@ void PostProcessBentNormalAO(
 				GetBufferSizeForAO(),
 				*VertexShader);
 		}
+
+		RHICmdList.CopyToResolveTarget(DistanceFieldAOBentNormal2->GetRenderTargetItem().TargetableTexture, DistanceFieldAOBentNormal2->GetRenderTargetItem().ShaderResourceTexture, false, FResolveParams());
+			
+		if (bUseDistanceFieldGI)
+		{
+			RHICmdList.CopyToResolveTarget(EffectiveDiffuseIrradianceTarget->GetRenderTargetItem().TargetableTexture, EffectiveDiffuseIrradianceTarget->GetRenderTargetItem().ShaderResourceTexture, false, FResolveParams());
+		}
 	}
 
 	if (GAOFillGaps)
@@ -2855,21 +2884,43 @@ void PostProcessBentNormalAO(
 
 			TShaderMapRef<FPostProcessVS> VertexShader(View.ShaderMap);
 
-			if (bUseDistanceFieldGI)
+			if (GAOFillGapsHighQuality)
 			{
-				TShaderMapRef<TFillGapsPS<true> > PixelShader(View.ShaderMap);
+				if (bUseDistanceFieldGI)
+				{
+					TShaderMapRef<TFillGapsPS<true, true> > PixelShader(View.ShaderMap);
 
-				static FGlobalBoundShaderState BoundShaderState;
-				SetGlobalBoundShaderState(RHICmdList, View.GetFeatureLevel(), BoundShaderState, GFilterVertexDeclaration.VertexDeclarationRHI, *VertexShader, *PixelShader);
-				PixelShader->SetParameters(RHICmdList, View, DistanceFieldAOBentNormal2->GetRenderTargetItem(), DistanceFieldIrradiance2);
+					static FGlobalBoundShaderState BoundShaderState;
+					SetGlobalBoundShaderState(RHICmdList, View.GetFeatureLevel(), BoundShaderState, GFilterVertexDeclaration.VertexDeclarationRHI, *VertexShader, *PixelShader);
+					PixelShader->SetParameters(RHICmdList, View, DistanceFieldAOBentNormal2->GetRenderTargetItem(), DistanceFieldIrradiance2);
+				}
+				else
+				{
+					TShaderMapRef<TFillGapsPS<false, true> > PixelShader(View.ShaderMap);
+
+					static FGlobalBoundShaderState BoundShaderState;
+					SetGlobalBoundShaderState(RHICmdList, View.GetFeatureLevel(), BoundShaderState, GFilterVertexDeclaration.VertexDeclarationRHI, *VertexShader, *PixelShader);
+					PixelShader->SetParameters(RHICmdList, View, DistanceFieldAOBentNormal2->GetRenderTargetItem(), DistanceFieldIrradiance2);
+				}
 			}
 			else
 			{
-				TShaderMapRef<TFillGapsPS<false> > PixelShader(View.ShaderMap);
+				if (bUseDistanceFieldGI)
+				{
+					TShaderMapRef<TFillGapsPS<true, false> > PixelShader(View.ShaderMap);
 
-				static FGlobalBoundShaderState BoundShaderState;
-				SetGlobalBoundShaderState(RHICmdList, View.GetFeatureLevel(), BoundShaderState, GFilterVertexDeclaration.VertexDeclarationRHI, *VertexShader, *PixelShader);
-				PixelShader->SetParameters(RHICmdList, View, DistanceFieldAOBentNormal2->GetRenderTargetItem(), DistanceFieldIrradiance2);
+					static FGlobalBoundShaderState BoundShaderState;
+					SetGlobalBoundShaderState(RHICmdList, View.GetFeatureLevel(), BoundShaderState, GFilterVertexDeclaration.VertexDeclarationRHI, *VertexShader, *PixelShader);
+					PixelShader->SetParameters(RHICmdList, View, DistanceFieldAOBentNormal2->GetRenderTargetItem(), DistanceFieldIrradiance2);
+				}
+				else
+				{
+					TShaderMapRef<TFillGapsPS<false, false> > PixelShader(View.ShaderMap);
+
+					static FGlobalBoundShaderState BoundShaderState;
+					SetGlobalBoundShaderState(RHICmdList, View.GetFeatureLevel(), BoundShaderState, GFilterVertexDeclaration.VertexDeclarationRHI, *VertexShader, *PixelShader);
+					PixelShader->SetParameters(RHICmdList, View, DistanceFieldAOBentNormal2->GetRenderTargetItem(), DistanceFieldIrradiance2);
+				}
 			}
 
 			VertexShader->SetParameters(RHICmdList, View);
@@ -2884,6 +2935,13 @@ void PostProcessBentNormalAO(
 				GSceneRenderTargets.GetBufferSizeXY() / FIntPoint(GAODownsampleFactor, GAODownsampleFactor),
 				*VertexShader);
 		}
+
+		RHICmdList.CopyToResolveTarget(DistanceFieldAOBentNormal->GetRenderTargetItem().TargetableTexture, DistanceFieldAOBentNormal->GetRenderTargetItem().ShaderResourceTexture, false, FResolveParams());
+			
+		if (bUseDistanceFieldGI)
+		{
+			RHICmdList.CopyToResolveTarget(DistanceFieldIrradiance->GetRenderTargetItem().TargetableTexture, DistanceFieldIrradiance->GetRenderTargetItem().ShaderResourceTexture, false, FResolveParams());
+		}
 	}
 
 	FSceneViewState* ViewState = (FSceneViewState*)View.State;
@@ -2987,7 +3045,7 @@ private:
 
 #define IMPLEMENT_UPSAMPLE_PS_TYPE(UpsampleType) \
 	typedef TDistanceFieldAOUpsamplePS<UpsampleType> TDistanceFieldAOUpsamplePS##UpsampleType; \
-	IMPLEMENT_SHADER_TYPE(template<>,TDistanceFieldAOUpsamplePS##UpsampleType,TEXT("DistanceFieldSurfaceCacheLighting"),TEXT("AOUpsamplePS"),SF_Pixel);
+	IMPLEMENT_SHADER_TYPE(template<>,TDistanceFieldAOUpsamplePS##UpsampleType,TEXT("DistanceFieldLightingPost"),TEXT("AOUpsamplePS"),SF_Pixel);
 
 IMPLEMENT_UPSAMPLE_PS_TYPE(AOUpsample_OutputBentNormal)
 IMPLEMENT_UPSAMPLE_PS_TYPE(AOUpsample_OutputAO)
@@ -3217,7 +3275,7 @@ void GenerateBestSpacedVectors()
 	}
 }
 
-FIntPoint BuildTileObjectLists(FRHICommandListImmediate& RHICmdList, FScene* Scene, TArray<FViewInfo>& Views, const FDistanceFieldAOParameters& Parameters)
+FIntPoint BuildTileObjectLists(FRHICommandListImmediate& RHICmdList, FScene* Scene, TArray<FViewInfo>& Views, FSceneRenderTargetItem& DistanceFieldNormal, const FDistanceFieldAOParameters& Parameters)
 {
 	SCOPED_DRAW_EVENT(RHICmdList, BuildTileList);
 	SetRenderTarget(RHICmdList, NULL, NULL);
@@ -3257,7 +3315,7 @@ FIntPoint BuildTileObjectLists(FRHICommandListImmediate& RHICmdList, FScene* Sce
 				TShaderMapRef<FBuildTileConesCS> ComputeShader(View.ShaderMap);
 
 				RHICmdList.SetComputeShader(ComputeShader->GetComputeShader());
-				ComputeShader->SetParameters(RHICmdList, View, Scene, FVector2D(GroupSizeX, GroupSizeY), Parameters);
+				ComputeShader->SetParameters(RHICmdList, View, DistanceFieldNormal, Scene, FVector2D(GroupSizeX, GroupSizeY), Parameters);
 				DispatchComputeShader(RHICmdList, *ComputeShader, GroupSizeX, GroupSizeY, 1);
 
 				ComputeShader->UnsetParameters(RHICmdList);
@@ -3343,10 +3401,19 @@ FIntPoint BuildTileObjectLists(FRHICommandListImmediate& RHICmdList, FScene* Sce
 	return TileListGroupSize;
 }
 
-void SetupDepthStencil(FRHICommandListImmediate& RHICmdList, const FViewInfo& View, FSceneRenderTargetItem& DistanceFieldNormal, int32 DepthLevel, int32 DestLevelDownsampleFactor)
+void SetupDepthStencil(
+	FRHICommandListImmediate& RHICmdList, 
+	const FViewInfo& View, 
+	FSceneRenderTargetItem& SplatDepthStencilBuffer,
+	FSceneRenderTargetItem& DistanceFieldNormal, 
+	int32 DepthLevel, 
+	int32 DestLevelDownsampleFactor)
 {
 	SCOPED_DRAW_EVENT(RHICmdList, SetupDepthStencil);
 
+	SetRenderTarget(RHICmdList, NULL, SplatDepthStencilBuffer.TargetableTexture);
+	RHICmdList.Clear(false, FLinearColor(0, 0, 0, 0), true, 0, true, 0, FIntRect());
+
 	{		
 		RHICmdList.SetViewport(0, 0, 0.0f, View.ViewRect.Width() / GAODownsampleFactor, View.ViewRect.Height() / GAODownsampleFactor, 1.0f);
 		RHICmdList.SetRasterizerState(TStaticRasterizerState<FM_Solid, CM_None>::GetRHI());
@@ -3412,6 +3479,8 @@ void SetupDepthStencil(FRHICommandListImmediate& RHICmdList, const FViewInfo& Vi
 
 		DrawPrimitiveUP(RHICmdList, PT_TriangleStrip, 2, ClearQuadVertices, sizeof(ClearQuadVertices[0]));
 	}
+
+	RHICmdList.CopyToResolveTarget(SplatDepthStencilBuffer.TargetableTexture, SplatDepthStencilBuffer.ShaderResourceTexture, false, FResolveParams());
 }
 
 void RenderIrradianceCacheInterpolation(
@@ -3436,30 +3505,51 @@ void RenderIrradianceCacheInterpolation(
 			FPooledRenderTargetDesc Desc(FPooledRenderTargetDesc::Create2DDesc(BentNormalInterpolationTarget->GetDesc().Extent, PF_DepthStencil, TexCreate_None, TexCreate_DepthStencilTargetable, false));
 			GRenderTargetPool.FindFreeElement(Desc, SplatDepthStencilBuffer, TEXT("DistanceFieldAOSplatDepthBuffer"));
 
-			SetRenderTarget(RHICmdList, NULL, SplatDepthStencilBuffer->GetRenderTargetItem().TargetableTexture);
-			RHICmdList.Clear(false, FLinearColor(0, 0, 0, 0), true, 0, true, 0, FIntRect());
-
-			SetupDepthStencil(RHICmdList, View, DistanceFieldNormal, DepthLevel, DestLevelDownsampleFactor);
+			SetupDepthStencil(RHICmdList, View, SplatDepthStencilBuffer->GetRenderTargetItem(), DistanceFieldNormal, DepthLevel, DestLevelDownsampleFactor);
 
 			GRenderTargetPool.VisualizeTexture.SetCheckPoint(RHICmdList, SplatDepthStencilBuffer);
+		}
 
-			FTextureRHIParamRef RenderTargets[2] =
-			{
-				BentNormalInterpolationTarget->GetRenderTargetItem().TargetableTexture,
-				bUseDistanceFieldGI ? IrradianceInterpolationTarget->GetRenderTargetItem().TargetableTexture : NULL
-			};
-			SetRenderTargets(RHICmdList, ARRAY_COUNT(RenderTargets) - (bUseDistanceFieldGI ? 0 : 1), RenderTargets, SplatDepthStencilBuffer->GetRenderTargetItem().TargetableTexture, 0, NULL);
+		SCOPED_DRAW_EVENT(RHICmdList, IrradianceCacheSplat);
 
-			FLinearColor ClearColors[2] = {FLinearColor(0, 0, 0, 0), FLinearColor(0, 0, 0, 0)};
-			RHICmdList.ClearMRT(true, ARRAY_COUNT(ClearColors) - (bUseDistanceFieldGI ? 0 : 1), ClearColors, false, 0, false, 0, FIntRect());
+		const bool bBindIrradiance = bUseDistanceFieldGI && bFinalInterpolation;
+
+		FTextureRHIParamRef RenderTargets[2] =
+		{
+			BentNormalInterpolationTarget->GetRenderTargetItem().TargetableTexture,
+			bBindIrradiance ? IrradianceInterpolationTarget->GetRenderTargetItem().TargetableTexture : NULL
+		};
+
+		FTextureRHIParamRef DepthBuffer = SplatDepthStencilBuffer ? SplatDepthStencilBuffer->GetRenderTargetItem().TargetableTexture : NULL;
+
+		// PS4 fast clear which is only triggered through SetRenderTargetsAndClear adds 3ms of GPU idle time currently
+		static bool bUseSetAndClear = false;
+
+		if (bUseSetAndClear)
+		{
+			FRHIRenderTargetView RenderTargets[MaxSimultaneousRenderTargets];
+			RenderTargets[0] = FRHIRenderTargetView(BentNormalInterpolationTarget->GetRenderTargetItem().TargetableTexture);
+			RenderTargets[1] = FRHIRenderTargetView(bBindIrradiance ? IrradianceInterpolationTarget->GetRenderTargetItem().TargetableTexture : NULL);
+
+			const int32 MRTCount = bBindIrradiance ? 2 : 1;
+
+			FRHIDepthRenderTargetView DepthView(DepthBuffer);
+			FRHISetRenderTargetsInfo Info(MRTCount, RenderTargets, DepthView);
+
+			Info.bClearColor = true;
+			Info.ClearColors[0] = FLinearColor(0, 0, 0, 0);
+			Info.ClearColors[1] = FLinearColor(0, 0, 0, 0);
+
+			// set the render target
+			RHICmdList.SetRenderTargetsAndClear(Info);
 		}
 		else
 		{
-			SetRenderTarget(RHICmdList, BentNormalInterpolationTarget->GetRenderTargetItem().TargetableTexture, NULL);
-			RHICmdList.Clear(true, FLinearColor(0, 0, 0, 0), false, 0, false, 0, FIntRect());
-		}
+			SetRenderTargets(RHICmdList, ARRAY_COUNT(RenderTargets) - (bBindIrradiance ? 0 : 1), RenderTargets, DepthBuffer, 0, NULL);
 
-		SCOPED_DRAW_EVENT(RHICmdList, IrradianceCacheSplat);
+			FLinearColor ClearColors[2] = {FLinearColor(0, 0, 0, 0), FLinearColor(0, 0, 0, 0)};
+			RHICmdList.ClearMRT(true, ARRAY_COUNT(ClearColors) - (bBindIrradiance ? 0 : 1), ClearColors, false, 0, false, 0, FIntRect());
+		}
 
 		static int32 NumInterpolationSections = 8;
 		if (GCircleVertexBuffer.NumSections != NumInterpolationSections)
@@ -3595,6 +3685,13 @@ void RenderIrradianceCacheInterpolation(
 				}
 			}
 		}
+
+		RHICmdList.CopyToResolveTarget(BentNormalInterpolationTarget->GetRenderTargetItem().TargetableTexture, BentNormalInterpolationTarget->GetRenderTargetItem().ShaderResourceTexture, false, FResolveParams());
+
+		if (bFinalInterpolation && bUseDistanceFieldGI && (GAOInterpolationDepthTesting || GAOInterpolationStencilTesting))
+		{
+			RHICmdList.CopyToResolveTarget(IrradianceInterpolationTarget->GetRenderTargetItem().TargetableTexture, IrradianceInterpolationTarget->GetRenderTargetItem().ShaderResourceTexture, false, FResolveParams());
+		}
 	}
 }
 
@@ -3800,10 +3897,7 @@ bool FDeferredShadingSceneRenderer::RenderDistanceFieldAOSurfaceCache(
 					ComputeShader->UnsetParameters(RHICmdList);
 				}
 			}
-
-			// Intersect objects with screen tiles, build lists
-			FIntPoint TileListGroupSize = BuildTileObjectLists(RHICmdList, Scene, Views, Parameters);
-
+			
 			TRefCountPtr<IPooledRenderTarget> DistanceFieldNormal;
 
 			{
@@ -3815,6 +3909,9 @@ bool FDeferredShadingSceneRenderer::RenderDistanceFieldAOSurfaceCache(
 			// Compute the distance field normal, this is used for surface caching instead of the GBuffer normal because it matches the occluding geometry
 			ComputeDistanceFieldNormal(RHICmdList, Views, DistanceFieldNormal->GetRenderTargetItem(), Parameters);
 
+			// Intersect objects with screen tiles, build lists
+			FIntPoint TileListGroupSize = BuildTileObjectLists(RHICmdList, Scene, Views, DistanceFieldNormal->GetRenderTargetItem(), Parameters);
+
 			GRenderTargetPool.VisualizeTexture.SetCheckPoint(RHICmdList, DistanceFieldNormal);
 
 			if (GAOReuseAcrossFrames)
@@ -4064,6 +4161,16 @@ bool FDeferredShadingSceneRenderer::RenderDistanceFieldAOSurfaceCache(
 			// Upsample to full resolution, write to output
 			UpsampleBentNormalAO(RHICmdList, Views, BentNormalOutput, IrradianceOutput, bVisualizeAmbientOcclusion, bVisualizeGlobalIllumination);
 
+			if (!bVisualizeAmbientOcclusion && !bVisualizeGlobalIllumination)
+			{
+				RHICmdList.CopyToResolveTarget(OutDynamicBentNormalAO->GetRenderTargetItem().TargetableTexture, OutDynamicBentNormalAO->GetRenderTargetItem().ShaderResourceTexture, false, FResolveParams());
+			
+				if (bUseDistanceFieldGI)
+				{
+					RHICmdList.CopyToResolveTarget(OutDynamicIrradiance->GetRenderTargetItem().TargetableTexture, OutDynamicIrradiance->GetRenderTargetItem().ShaderResourceTexture, false, FResolveParams());
+				}
+			}
+
 			return true;
 		}
 	}
