Add permutations to various global shaders to compile out ClearCoat handling if there are no clearcoat materials in the given view.
Also FViewInfo now accumulates a mask of all lightingprofiles used in a given view during the relevancy calculation.

diff --git a/Engine/Shaders/DeferredLightingCommon.usf b/Engine/Shaders/DeferredLightingCommon.usf
--- a/Engine/Shaders/DeferredLightingCommon.usf
+++ b/Engine/Shaders/DeferredLightingCommon.usf
@@ -121,6 +121,7 @@ float3 SimpleShading( FGBufferData GBuffer, float Roughness, float3 L, float3 V,
 
 float3 ClearCoatShading( FGBufferData GBuffer, float3 LobeRoughness, float3 LobeEnergy, float3 L, float3 V, half3 N )
 {
+#if USE_CLEARCOAT
 	const float ClearCoat			= GBuffer.CustomData.x;
 	const float ClearCoatRoughness	= GBuffer.CustomData.y;
 	const float Film = 1 * ClearCoat;
@@ -227,6 +228,9 @@ float3 ClearCoatShading( FGBufferData GBuffer, float3 LobeRoughness, float3 Lobe
 	
 	return Fr1 + Fr2 * LayerAttenuation;
 #endif
+#else //USE_CLEARCOAT
+	return float3(0.0f, 0.0f, 0.0f);
+#endif //USE_CLEARCOAT
 }
 
 float3 SubsurfaceShadingSubsurface( FGBufferData GBuffer, float3 L, float3 V, half3 N )
diff --git a/Engine/Shaders/Definitions.usf b/Engine/Shaders/Definitions.usf
--- a/Engine/Shaders/Definitions.usf
+++ b/Engine/Shaders/Definitions.usf
@@ -150,6 +150,10 @@
 #define NUM_VF_PACKED_INTERPOLANTS						0
 #endif
 
+#ifndef USE_CLEARCOAT
+#define USE_CLEARCOAT									0
+#endif
+
 #if SM5_PROFILE || COMPILER_SUPPORTS_ATTRIBUTES
 /** Avoids flow control constructs. */
 #define UNROLL											[unroll]
diff --git a/Engine/Shaders/PostProcessAmbient.usf b/Engine/Shaders/PostProcessAmbient.usf
--- a/Engine/Shaders/PostProcessAmbient.usf
+++ b/Engine/Shaders/PostProcessAmbient.usf
@@ -301,11 +301,13 @@ float3 ImageBasedLightingMIS( FGBufferData GBuffer, float3 V, float3 N, uint2 Ra
 				float Weight = rcp( InvWeight );
 
 				float3 Shading = 0;
+#if USE_CLEARCOAT
 				BRANCH if( GBuffer.ShadingModelID == SHADINGMODELID_CLEAR_COAT )
 				{
 					Shading = ClearCoatShading( GBuffer, GBuffer.Roughness, L, V, N );
 				}
 				else
+#endif
 				{
 					Shading = StandardShading( GBuffer, GBuffer.Roughness, L, V, N );
 				}
@@ -497,6 +499,7 @@ void MainPS(in float4 UVAndScreenPos : TEXCOORD0, out float4 OutColor : SV_Targe
 		float Mip = ComputeCubemapMipFromRoughness( GBuffer.Roughness, AmbientCubemapMipAdjust.w );
 		float3 SampleColor = TextureCubeSampleLevel( AmbientCubemap, AmbientCubemapSampler, R, Mip ).rgb;
 
+#if USE_CLEARCOAT
 		BRANCH if( GBuffer.ShadingModelID == SHADINGMODELID_CLEAR_COAT )
 		{
 			float ClearCoat = GBuffer.CustomData.x;
@@ -504,12 +507,14 @@ void MainPS(in float4 UVAndScreenPos : TEXCOORD0, out float4 OutColor : SV_Targe
 			SpecularContribution += SampleColor * ( GBuffer.SpecularColor * AB.x + AB.y * (1 - ClearCoat) );
 		}
 		else
+#endif
 		{
 			SpecularContribution += SampleColor * EnvBRDF( GBuffer.SpecularColor, GBuffer.Roughness, NoV );
 			//SpecularContribution += ApproximateSpecularIBL( Random, GBuffer.SpecularColor, GBuffer.Roughness, GBuffer.WorldNormal, -ScreenVector );
 		}
 	}
 
+#if USE_CLEARCOAT
 	BRANCH if( GBuffer.ShadingModelID == SHADINGMODELID_CLEAR_COAT )
 	{
 		const float ClearCoat			= GBuffer.CustomData.x;
@@ -530,6 +535,7 @@ void MainPS(in float4 UVAndScreenPos : TEXCOORD0, out float4 OutColor : SV_Targe
 		SpecularContribution *= LayerAttenuation;
 		SpecularContribution += SampleColor * F;
 	}
+#endif
 
 #if IMPORTANCE_SAMPLE
 	if( GBuffer.ShadingModelID > 0 )
diff --git a/Engine/Shaders/ReflectionEnvironmentComputeShaders.usf b/Engine/Shaders/ReflectionEnvironmentComputeShaders.usf
--- a/Engine/Shaders/ReflectionEnvironmentComputeShaders.usf
+++ b/Engine/Shaders/ReflectionEnvironmentComputeShaders.usf
@@ -17,7 +17,7 @@ TextureCubeArray ReflectionEnvironmentColorTexture;
 SamplerState ReflectionEnvironmentColorSampler;
 
 #define THREADGROUP_TOTALSIZE (THREADGROUP_SIZEX * THREADGROUP_SIZEY)
-
+  
 // Workaround performance issue with shared memory bank collisions in GLSL
 #if GL4_PROFILE
 #define ATOMIC_REDUCTION 0
@@ -443,11 +443,16 @@ void GatherRadiance(inout float4 Color, float3 WorldPosition, float3 RayDirectio
 
 			// Fade out based on distance to capture
 			float DistanceAlpha = 0;
-
+			
 			#define PROJECT_ONTO_SHAPE 1
 			#if PROJECT_ONTO_SHAPE
+
+#if HAS_BOX_CAPTURES
+
+#	if HAS_SPHERE_CAPTURES
 				// Box
 				BRANCH if (CaptureProperties.b > 0)
+#	endif //HAS_SPHERE_CAPTURES
 				{
 					// Transform the ray into the local space of the box, where it is an AABB with mins at -1 and maxs at 1
 					float3 LocalRayStart		= mul(float4(WorldPosition, 1), ReflectionCapture.BoxTransform[CaptureIndex]).xyz;
@@ -478,8 +483,13 @@ void GatherRadiance(inout float4 Color, float3 WorldPosition, float3 RayDirectio
 					// Setup a fade based on receiver distance to the box, hides the box influence shape
 					DistanceAlpha = 1.0 - smoothstep(0, .7f * BoxScales.w, BoxDistance);
 				}
+#endif //HAS_BOX_CAPTURES
+
+#if HAS_SPHERE_CAPTURES
 				// Sphere
+#	if HAS_BOX_CAPTURES
 				else
+#	endif //HAS_BOX_CAPTURES
 				{
 					float ProjectionSphereRadius = CapturePositionAndRadius.w;// * 1.2f;
 					float SphereRadiusSquared = ProjectionSphereRadius * ProjectionSphereRadius;
@@ -508,9 +518,10 @@ void GatherRadiance(inout float4 Color, float3 WorldPosition, float3 RayDirectio
 						DistanceAlpha = 1 - x*x*(3 - 2*x);
 					}
 				}
+#endif //HAS_SPHERE_CAPTURES
 			#else 
 				DistanceAlpha = 1.0;
-			#endif
+			#endif //PROJECT_ONTO_SHAPE
 
 			float CaptureArrayIndex = CaptureProperties.g;
 
@@ -566,7 +577,7 @@ Texture2D InSceneColor;
 
 /** Output HDR target. */
 RWTexture2D<float4> RWOutSceneColor;
-
+  
 [numthreads(THREADGROUP_SIZEX, THREADGROUP_SIZEY, 1)]
 void ReflectionEnvironmentTiledDeferredMain(
 	uint3 GroupId : SV_GroupID,
@@ -627,17 +638,20 @@ void ReflectionEnvironmentTiledDeferredMain(
 		float a = Square( GBuffer.Roughness );
 		R = lerp( N, R, (1 - a) * ( sqrt(1 - a) + a ) );
 
+
 	#if 1//APPLY_SSR
 		float4 SSR = ScreenSpaceReflections.Load( int3(PixelPos, 0) );
 		Color.rgb = SSR.rgb;
 		Color.a = 1 - SSR.a;
 	#endif
 
+#if USE_CLEARCOAT
 		if( GBuffer.ShadingModelID == SHADINGMODELID_CLEAR_COAT )
 		{
 			const float ClearCoat = GBuffer.CustomData.x;
 			Color = lerp( Color, float4(0,0,0,1), ClearCoat );
 		}
+#endif
 	
 		float AO = ScreenSpaceData.AmbientOcclusion;
 		float SpecularOcclusion = saturate( pow( NoV + AO, a ) - 1 + AO );
@@ -646,6 +660,7 @@ void ReflectionEnvironmentTiledDeferredMain(
 
 		GatherRadiance(Color, WorldPosition, R, GBuffer.Roughness, ScreenPosition, IndirectIrradiance, GBuffer.ShadingModelID);
 
+#if USE_CLEARCOAT
 		BRANCH
 		if( GBuffer.ShadingModelID == SHADINGMODELID_CLEAR_COAT )
 		{
@@ -671,9 +686,10 @@ void ReflectionEnvironmentTiledDeferredMain(
 			Color.a *= 1 - SSR.a;
 			Color.a *= SpecularOcclusion;
 
-			GatherRadiance(Color, WorldPosition, R, ClearCoatRoughness, ScreenPosition, IndirectIrradiance, GBuffer.ShadingModelID);
+			GatherRadiance(Color, WorldPosition, R, ClearCoatRoughness, ScreenPosition, IndirectIrradiance, GBuffer.ShadingModelID);			
 		}
 		else
+#endif
 		{
 			Color.rgb *= EnvBRDF( GBuffer.SpecularColor, GBuffer.Roughness, NoV );
 		}
@@ -691,7 +707,7 @@ void ReflectionEnvironmentTiledDeferredMain(
 #else
 		OutColor.rgb = Color.rgb;
 #endif
-
+		
 		// Transform NaNs to black, transform negative colors to black.
 		OutColor.rgb = -min(-OutColor.rgb, 0.0);
 
diff --git a/Engine/Shaders/ScreenSpaceReflections.usf b/Engine/Shaders/ScreenSpaceReflections.usf
--- a/Engine/Shaders/ScreenSpaceReflections.usf
+++ b/Engine/Shaders/ScreenSpaceReflections.usf
@@ -51,6 +51,7 @@ void ScreenSpaceReflectionsPS(in float4 UVAndScreenPos : TEXCOORD0, out float4 O
 	
 	float Roughness = GBuffer.Roughness;
 
+#if USE_CLEARCOAT
 	BRANCH
 	if( GBuffer.ShadingModelID == SHADINGMODELID_CLEAR_COAT )
 	{
@@ -59,6 +60,7 @@ void ScreenSpaceReflectionsPS(in float4 UVAndScreenPos : TEXCOORD0, out float4 O
 
 		Roughness = lerp( Roughness, ClearCoatRoughness, ClearCoat );
 	}
+#endif
 
 	// mask SSR to reduce noise and for better performance, roughness of 0 should have SSR, at MaxRoughness we fade to 0
 	float RoughnessFade = saturate(Roughness * SSRParams.y + 2);
diff --git a/Engine/Source/Runtime/Engine/Classes/Materials/MaterialInterface.h b/Engine/Source/Runtime/Engine/Classes/Materials/MaterialInterface.h
--- a/Engine/Source/Runtime/Engine/Classes/Materials/MaterialInterface.h
+++ b/Engine/Source/Runtime/Engine/Classes/Materials/MaterialInterface.h
@@ -53,12 +53,12 @@ struct ENGINE_API FMaterialRelevance
 
 	UPROPERTY()
 	uint32 bDisableDepthTest : 1;
-
+	
 	UPROPERTY()
-	uint32 bSubsurfaceProfile : 1;
+	uint32 bOutputsVelocityInBasePass : 1;
 
 	UPROPERTY()
-	uint32 bOutputsVelocityInBasePass : 1;
+	uint16 LightingProfileMask;
 
 	/** Default constructor. */
 	FMaterialRelevance()
@@ -67,9 +67,9 @@ struct ENGINE_API FMaterialRelevance
 		, bDistortion(false)
 		, bSeparateTranslucency(false)
 		, bNormalTranslucency(false)
-		, bDisableDepthTest(false)
-		, bSubsurfaceProfile(false)
+		, bDisableDepthTest(false)		
 		, bOutputsVelocityInBasePass(true)
+		, LightingProfileMask(0)
 	{}
 
 	/** Bitwise OR operator.  Sets any relevance bits which are present in either FMaterialRelevance. */
@@ -81,7 +81,7 @@ struct ENGINE_API FMaterialRelevance
 		bSeparateTranslucency |= B.bSeparateTranslucency;
 		bNormalTranslucency |= B.bNormalTranslucency;
 		bDisableDepthTest |= B.bDisableDepthTest;
-		bSubsurfaceProfile |= B.bSubsurfaceProfile;
+		LightingProfileMask |= B.LightingProfileMask;
 		bOutputsVelocityInBasePass |= B.bOutputsVelocityInBasePass;
 		return *this;
 	}
diff --git a/Engine/Source/Runtime/Engine/Private/Materials/MaterialInterface.cpp b/Engine/Source/Runtime/Engine/Private/Materials/MaterialInterface.cpp
--- a/Engine/Source/Runtime/Engine/Private/Materials/MaterialInterface.cpp
+++ b/Engine/Source/Runtime/Engine/Private/Materials/MaterialInterface.cpp
@@ -17,7 +17,7 @@ void FMaterialRelevance::SetPrimitiveViewRelevance(FPrimitiveViewRelevance& OutV
 	OutViewRelevance.bDistortionRelevance = bDistortion;
 	OutViewRelevance.bSeparateTranslucencyRelevance = bSeparateTranslucency;
 	OutViewRelevance.bNormalTranslucencyRelevance = bNormalTranslucency;
-	OutViewRelevance.bSubsurfaceProfileRelevance = bSubsurfaceProfile;
+	OutViewRelevance.LightingProfileRelevanceMask = LightingProfileMask;
 }
 
 //////////////////////////////////////////////////////////////////////////
@@ -50,15 +50,14 @@ FMaterialRelevance UMaterialInterface::GetRelevance_Internal(const UMaterial* Ma
 		const bool bIsLit =  ShadingModel != MSM_Unlit;
 		// Determine the material's view relevance.
 		FMaterialRelevance MaterialRelevance;
+		MaterialRelevance.LightingProfileMask = 1 << ShadingModel;
 		MaterialRelevance.bOpaque = !bIsTranslucent;
 		MaterialRelevance.bMasked = IsMasked();
 		MaterialRelevance.bDistortion = Material->bUsesDistortion && bIsTranslucent;
 		MaterialRelevance.bSeparateTranslucency = bIsTranslucent && Material->bEnableSeparateTranslucency;
 		MaterialRelevance.bNormalTranslucency = bIsTranslucent && !Material->bEnableSeparateTranslucency;
-		MaterialRelevance.bDisableDepthTest = bIsTranslucent && Material->bDisableDepthTest;
-		MaterialRelevance.bSubsurfaceProfile = (Material->MaterialDomain == MD_Surface) && !bIsTranslucent && (ShadingModel == MSM_SubsurfaceProfile);
-		MaterialRelevance.bOutputsVelocityInBasePass = Material->bOutputVelocityOnBasePass;
-
+		MaterialRelevance.bDisableDepthTest = bIsTranslucent && Material->bDisableDepthTest;		
+		MaterialRelevance.bOutputsVelocityInBasePass = Material->bOutputVelocityOnBasePass;		
 		return MaterialRelevance;
 	}
 	else
diff --git a/Engine/Source/Runtime/Engine/Public/PrimitiveViewRelevance.h b/Engine/Source/Runtime/Engine/Public/PrimitiveViewRelevance.h
--- a/Engine/Source/Runtime/Engine/Public/PrimitiveViewRelevance.h
+++ b/Engine/Source/Runtime/Engine/Public/PrimitiveViewRelevance.h
@@ -8,6 +8,8 @@
  */
 struct FPrimitiveViewRelevance
 {
+	/* The LightingProfile supported by this primitive, as a bitmask. */
+	uint16 LightingProfileRelevanceMask;
 	/** The primitive's static elements are rendered for the view. */
 	uint32 bStaticRelevance : 1; 
 	/** The primitive's dynamic elements are rendered for the view. */
@@ -37,8 +39,6 @@ struct FPrimitiveViewRelevance
 	uint32 bSeparateTranslucencyRelevance : 1;
 	/** The primitive has one or more elements that have normal translucency. */
 	uint32 bNormalTranslucencyRelevance : 1;
-	/** The primitive has one or more elements that have the SubsurfaceProfile shading model. */
-	uint32 bSubsurfaceProfileRelevance : 1;
 	// The primitive has one or more elements that have World Position Offset.
 	uint32 bHasWorldPositionOffset : 1;
 
@@ -56,6 +56,7 @@ struct FPrimitiveViewRelevance
 
 	/** Initialization constructor. */
 	FPrimitiveViewRelevance():
+		LightingProfileRelevanceMask(0),
 		bStaticRelevance(false),
 		bDynamicRelevance(false),
 		bDrawRelevance(false),
@@ -69,8 +70,7 @@ struct FPrimitiveViewRelevance
 		bMaskedRelevance(false),
 		bDistortionRelevance(false),
 		bSeparateTranslucencyRelevance(false),
-		bNormalTranslucencyRelevance(false),
-		bSubsurfaceProfileRelevance(false),
+		bNormalTranslucencyRelevance(false),		
 		bHasWorldPositionOffset(false),
 		bInitializedThisFrame(false)
 	{}
@@ -78,6 +78,7 @@ struct FPrimitiveViewRelevance
 	/** Bitwise OR operator.  Sets any relevance bits which are present in either FPrimitiveViewRelevance. */
 	FPrimitiveViewRelevance& operator|=(const FPrimitiveViewRelevance& B)
 	{
+		LightingProfileRelevanceMask |= B.LightingProfileRelevanceMask;
 		bStaticRelevance |= B.bStaticRelevance != 0;
 		bDynamicRelevance |= B.bDynamicRelevance != 0;
 		bDrawRelevance |= B.bDrawRelevance != 0;
@@ -92,8 +93,7 @@ struct FPrimitiveViewRelevance
 		bHasSimpleLights |= B.bHasSimpleLights != 0;
 		bSeparateTranslucencyRelevance |= B.bSeparateTranslucencyRelevance != 0;
 		bNormalTranslucencyRelevance |= B.bNormalTranslucencyRelevance != 0;
-		bInitializedThisFrame |= B.bInitializedThisFrame;
-		bSubsurfaceProfileRelevance |= B.bSubsurfaceProfileRelevance != 0;
+		bInitializedThisFrame |= B.bInitializedThisFrame;		
 		bHasWorldPositionOffset |= B.bHasWorldPositionOffset != 0;
 		return *this;
 	}
diff --git a/Engine/Source/Runtime/Renderer/Private/CompositionLighting/CompositionLighting.cpp b/Engine/Source/Runtime/Renderer/Private/CompositionLighting/CompositionLighting.cpp
--- a/Engine/Source/Runtime/Renderer/Private/CompositionLighting/CompositionLighting.cpp
+++ b/Engine/Source/Runtime/Renderer/Private/CompositionLighting/CompositionLighting.cpp
@@ -357,7 +357,8 @@ void FCompositionLighting::ProcessLighting(FRHICommandListImmediate& RHICmdList,
 
 			bool bSimpleDynamicLighting = IsSimpleDynamicLightingEnabled();
 
-			if (View.bScreenSpaceSubsurfacePassNeeded && Radius > 0 && !bSimpleDynamicLighting && View.Family->EngineShowFlags.SubsurfaceScattering && 
+			bool bScreenSpaceSubsurfacePassNeeded = (View.LightingProfilesActiveInView & (1 << MSM_SubsurfaceProfile)) != 0;
+			if (bScreenSpaceSubsurfacePassNeeded && Radius > 0 && !bSimpleDynamicLighting && View.Family->EngineShowFlags.SubsurfaceScattering &&
 				//@todo-rco: Remove this when we fix the cross-compiler
 				!IsOpenGLPlatform(View.GetShaderPlatform()))
 			{
diff --git a/Engine/Source/Runtime/Renderer/Private/CompositionLighting/PostProcessAmbient.cpp b/Engine/Source/Runtime/Renderer/Private/CompositionLighting/PostProcessAmbient.cpp
--- a/Engine/Source/Runtime/Renderer/Private/CompositionLighting/PostProcessAmbient.cpp
+++ b/Engine/Source/Runtime/Renderer/Private/CompositionLighting/PostProcessAmbient.cpp
@@ -12,9 +12,10 @@
 #include "SceneUtils.h"
 
 /** Encapsulates the post processing ambient pixel shader. */
-class FPostProcessAmbientPS : public FGlobalShader
+template<bool bUseClearCoat>
+class TPostProcessAmbientPS : public FGlobalShader
 {
-	DECLARE_SHADER_TYPE(FPostProcessAmbientPS, Global);
+	DECLARE_SHADER_TYPE(TPostProcessAmbientPS, Global);
 
 	static bool ShouldCache(EShaderPlatform Platform)
 	{
@@ -24,10 +25,11 @@ class FPostProcessAmbientPS : public FGlobalShader
 	static void ModifyCompilationEnvironment(EShaderPlatform Platform, FShaderCompilerEnvironment& OutEnvironment)
 	{
 		FGlobalShader::ModifyCompilationEnvironment(Platform,OutEnvironment);
-	}
+		OutEnvironment.SetDefine(TEXT("USE_CLEARCOAT"), (uint32)bUseClearCoat);
+	}	
 
 	/** Default constructor. */
-	FPostProcessAmbientPS() {}
+	TPostProcessAmbientPS() {}
 
 public:
 	FPostProcessPassParameters PostprocessParameter;
@@ -37,7 +39,7 @@ public:
 	FShaderResourceParameter PreIntegratedGFSampler;
 
 	/** Initialization constructor. */
-	FPostProcessAmbientPS(const ShaderMetaType::CompiledShaderInitializerType& Initializer)
+	TPostProcessAmbientPS(const ShaderMetaType::CompiledShaderInitializerType& Initializer)
 		: FGlobalShader(Initializer)
 	{
 		PostprocessParameter.Bind(Initializer.ParameterMap);
@@ -67,13 +69,48 @@ public:
 	}
 };
 
-IMPLEMENT_SHADER_TYPE(,FPostProcessAmbientPS,TEXT("PostProcessAmbient"),TEXT("MainPS"),SF_Pixel);
+IMPLEMENT_SHADER_TYPE(template<>, TPostProcessAmbientPS<false>, TEXT("PostProcessAmbient"), TEXT("MainPS"), SF_Pixel);
+IMPLEMENT_SHADER_TYPE(template<>, TPostProcessAmbientPS<true>, TEXT("PostProcessAmbient"), TEXT("MainPS"), SF_Pixel);
+
+template<bool bUseClearCoat>
+void FRCPassPostProcessAmbient::Render(FRenderingCompositePassContext& Context)
+{
+	TShaderMapRef<TPostProcessAmbientPS<bUseClearCoat>> PixelShader(Context.GetShaderMap());
+	TShaderMapRef<FPostProcessVS> VertexShader(Context.GetShaderMap());
+	const FSceneView& View = Context.View;
+
+	static FGlobalBoundShaderState BoundShaderState;
+
+	uint32 Count = Context.View.FinalPostProcessSettings.ContributingCubemaps.Num();
+	for (uint32 i = 0; i < Count; ++i)
+	{
+		if (i == 0)
+		{
+			// call it once after setting up the shader data to avoid the warnings in the function
+			SetGlobalBoundShaderState(Context.RHICmdList, Context.GetFeatureLevel(), BoundShaderState, GFilterVertexDeclaration.VertexDeclarationRHI, *VertexShader, *PixelShader);
+		}
+
+		PixelShader->SetParameters(Context.RHICmdList, Context, Context.View.FinalPostProcessSettings.ContributingCubemaps[i]);
+
+		// Draw a quad mapping scene color to the view's render target
+		DrawRectangle(
+			Context.RHICmdList,
+			0, 0,
+			View.ViewRect.Width(), View.ViewRect.Height(),
+			View.ViewRect.Min.X, View.ViewRect.Min.Y,
+			View.ViewRect.Width(), View.ViewRect.Height(),
+			View.ViewRect.Size(),
+			GSceneRenderTargets.GetBufferSizeXY(),
+			*VertexShader,
+			EDRF_UseTriangleOptimization);
+	}
+}
 
 void FRCPassPostProcessAmbient::Process(FRenderingCompositePassContext& Context)
 {
 	SCOPED_DRAW_EVENT(Context.RHICmdList, PostProcessAmbient);
 
-	const FSceneView& View = Context.View;
+	const FViewInfo& View = Context.View;
 	const FSceneViewFamily& ViewFamily = *(View.Family);
 
 	FIntRect SrcRect = View.ViewRect;
@@ -92,43 +129,23 @@ void FRCPassPostProcessAmbient::Process(FRenderingCompositePassContext& Context)
 	Context.RHICmdList.SetRasterizerState(TStaticRasterizerState<>::GetRHI());
 	Context.RHICmdList.SetDepthStencilState(TStaticDepthStencilState<false, CF_Always>::GetRHI());
 
-	TShaderMapRef<FPostProcessVS> VertexShader(Context.GetShaderMap());
-
 	FScene* Scene = ViewFamily.Scene->GetRenderScene();
 	check(Scene);
 	uint32 NumReflectionCaptures = Scene->ReflectionSceneData.RegisteredReflectionCaptures.Num();
 	
 	// Ambient cubemap specular will be applied in the reflection environment pass if it is enabled
 	const bool bApplySpecular = View.Family->EngineShowFlags.ReflectionEnvironment == 0 || NumReflectionCaptures == 0;
-	TShaderMapRef<FPostProcessAmbientPS> PixelShader(Context.GetShaderMap());
-
-	static FGlobalBoundShaderState BoundShaderState;
-	
 
-	uint32 Count = Context.View.FinalPostProcessSettings.ContributingCubemaps.Num();
-	for(uint32 i = 0; i < Count; ++i)
+	bool bClearCoatNeeded = (View.LightingProfilesActiveInView & (1 << MSM_ClearCoat)) != 0;
+	if (bClearCoatNeeded)
 	{
-		if(i == 0)
-		{
-			// call it once after setting up the shader data to avoid the warnings in the function
-			SetGlobalBoundShaderState(Context.RHICmdList, Context.GetFeatureLevel(), BoundShaderState, GFilterVertexDeclaration.VertexDeclarationRHI, *VertexShader, *PixelShader);
-		}
-
-		
-		PixelShader->SetParameters(Context.RHICmdList, Context, Context.View.FinalPostProcessSettings.ContributingCubemaps[i]);
-
-		// Draw a quad mapping scene color to the view's render target
-		DrawRectangle(
-			Context.RHICmdList,
-			0, 0,
-			View.ViewRect.Width(), View.ViewRect.Height(),
-			View.ViewRect.Min.X, View.ViewRect.Min.Y, 
-			View.ViewRect.Width(), View.ViewRect.Height(),
-			View.ViewRect.Size(),
-			GSceneRenderTargets.GetBufferSizeXY(),
-			*VertexShader,
-			EDRF_UseTriangleOptimization);
+		Render<true>(Context);
 	}
+	else
+	{
+		Render<false>(Context);
+	}
+	
 
 	Context.RHICmdList.CopyToResolveTarget(DestRenderTarget.TargetableTexture, DestRenderTarget.ShaderResourceTexture, false, FResolveParams());
 }
diff --git a/Engine/Source/Runtime/Renderer/Private/CompositionLighting/PostProcessAmbient.h b/Engine/Source/Runtime/Renderer/Private/CompositionLighting/PostProcessAmbient.h
--- a/Engine/Source/Runtime/Renderer/Private/CompositionLighting/PostProcessAmbient.h
+++ b/Engine/Source/Runtime/Renderer/Private/CompositionLighting/PostProcessAmbient.h
@@ -19,6 +19,11 @@ public:
 	virtual const TCHAR* GetDebugName() { return TEXT("FRCPassPostProcessAmbient"); }
 	virtual void Process(FRenderingCompositePassContext& Context) override;
 	virtual void Release() override { delete this; }
-	virtual bool FrameBufferBlendingWithInput0() const override { return true; }
-	virtual FPooledRenderTargetDesc ComputeOutputDesc(EPassOutputId InPassOutputId) const override;
+	virtual bool FrameBufferBlendingWithInput0() const { return true; }
+	virtual FPooledRenderTargetDesc ComputeOutputDesc(EPassOutputId InPassOutputId) const;
+
+private:
+
+	template<bool bUseClearCoat>
+	void Render(FRenderingCompositePassContext& Context);
 };
diff --git a/Engine/Source/Runtime/Renderer/Private/LightRendering.cpp b/Engine/Source/Runtime/Renderer/Private/LightRendering.cpp
--- a/Engine/Source/Runtime/Renderer/Private/LightRendering.cpp
+++ b/Engine/Source/Runtime/Renderer/Private/LightRendering.cpp
@@ -32,7 +32,7 @@ IMPLEMENT_SHADER_TYPE(template<>,TDeferredLightVS<false>,TEXT("DeferredLightVert
 IMPLEMENT_SHADER_TYPE(template<>,TDeferredLightVS<true>,TEXT("DeferredLightVertexShaders"),TEXT("RadialVertexMain"),SF_Vertex);
 
 /** A pixel shader for rendering the light in a deferred pass. */
-template<bool bUseIESProfile, bool bRadialAttenuation, bool bInverseSquaredFalloff, bool bVisualizeLightCulling>
+template<bool bUseIESProfile, bool bRadialAttenuation, bool bInverseSquaredFalloff, bool bVisualizeLightCulling, bool bUseClearCoat>
 class TDeferredLightPS : public FGlobalShader
 {
 	DECLARE_SHADER_TYPE(TDeferredLightPS,Global)
@@ -51,6 +51,7 @@ public:
 		OutEnvironment.SetDefine(TEXT("INVERSE_SQUARED_FALLOFF"), (uint32)bInverseSquaredFalloff);
 		OutEnvironment.SetDefine(TEXT("LIGHT_SOURCE_SHAPE"), 1);
 		OutEnvironment.SetDefine(TEXT("VISUALIZE_LIGHT_CULLING"), (uint32)bVisualizeLightCulling);
+		OutEnvironment.SetDefine(TEXT("USE_CLEARCOAT"), (uint32)bUseClearCoat);
 	}
 
 	TDeferredLightPS(const ShaderMetaType::CompiledShaderInitializerType& Initializer)
@@ -155,20 +156,28 @@ private:
 };
 
 // Typedef is necessary because the C preprocessor thinks the comma in the template parameter list is a comma in the macro parameter list.
-#define IMPLEMENT_DEFERREDLIGHT_PIXELSHADER_TYPE(A, B, C, D, EntryName) \
-	typedef TDeferredLightPS<A,B,C,D> TDeferredLightPS##A##B##C##D; \
-	IMPLEMENT_SHADER_TYPE(template<>,TDeferredLightPS##A##B##C##D,TEXT("DeferredLightPixelShaders"),EntryName,SF_Pixel);
+#define IMPLEMENT_DEFERREDLIGHT_PIXELSHADER_TYPE(A, B, C, D, E, EntryName) \
+	typedef TDeferredLightPS<A,B,C,D,E> TDeferredLightPS##A##B##C##D##E; \
+	IMPLEMENT_SHADER_TYPE(template<>,TDeferredLightPS##A##B##C##D##E,TEXT("DeferredLightPixelShaders"),EntryName,SF_Pixel);
 
 // Implement a version for each light type, and it's shader permutations
-IMPLEMENT_DEFERREDLIGHT_PIXELSHADER_TYPE(true, true, true, false, TEXT("RadialPixelMain"));
-IMPLEMENT_DEFERREDLIGHT_PIXELSHADER_TYPE(true, true, false, false, TEXT("RadialPixelMain"));
-IMPLEMENT_DEFERREDLIGHT_PIXELSHADER_TYPE(true, false, false, false, TEXT("DirectionalPixelMain"));
-IMPLEMENT_DEFERREDLIGHT_PIXELSHADER_TYPE(false, true, true, false, TEXT("RadialPixelMain"));
-IMPLEMENT_DEFERREDLIGHT_PIXELSHADER_TYPE(false, true, false, false, TEXT("RadialPixelMain"));
-IMPLEMENT_DEFERREDLIGHT_PIXELSHADER_TYPE(false, false, false, false, TEXT("DirectionalPixelMain"));
-
-IMPLEMENT_DEFERREDLIGHT_PIXELSHADER_TYPE(false, false, false, true, TEXT("DirectionalPixelMain"));
-IMPLEMENT_DEFERREDLIGHT_PIXELSHADER_TYPE(false, true, false, true, TEXT("RadialPixelMain"));
+IMPLEMENT_DEFERREDLIGHT_PIXELSHADER_TYPE(true, true, true, false, false, TEXT("RadialPixelMain"));
+IMPLEMENT_DEFERREDLIGHT_PIXELSHADER_TYPE(true, true, false, false, false, TEXT("RadialPixelMain"));
+IMPLEMENT_DEFERREDLIGHT_PIXELSHADER_TYPE(true, false, false, false, false, TEXT("DirectionalPixelMain"));
+IMPLEMENT_DEFERREDLIGHT_PIXELSHADER_TYPE(false, true, true, false, false, TEXT("RadialPixelMain"));
+IMPLEMENT_DEFERREDLIGHT_PIXELSHADER_TYPE(false, true, false, false, false, TEXT("RadialPixelMain"));
+IMPLEMENT_DEFERREDLIGHT_PIXELSHADER_TYPE(false, false, false, false, false, TEXT("DirectionalPixelMain"));
+IMPLEMENT_DEFERREDLIGHT_PIXELSHADER_TYPE(false, false, false, true, false, TEXT("DirectionalPixelMain"));
+IMPLEMENT_DEFERREDLIGHT_PIXELSHADER_TYPE(false, true, false, true, false, TEXT("RadialPixelMain"));
+
+IMPLEMENT_DEFERREDLIGHT_PIXELSHADER_TYPE(true, true, true, false, true, TEXT("RadialPixelMain"));
+IMPLEMENT_DEFERREDLIGHT_PIXELSHADER_TYPE(true, true, false, false, true, TEXT("RadialPixelMain"));
+IMPLEMENT_DEFERREDLIGHT_PIXELSHADER_TYPE(true, false, false, false, true, TEXT("DirectionalPixelMain"));
+IMPLEMENT_DEFERREDLIGHT_PIXELSHADER_TYPE(false, true, true, false, true, TEXT("RadialPixelMain"));
+IMPLEMENT_DEFERREDLIGHT_PIXELSHADER_TYPE(false, true, false, false, true, TEXT("RadialPixelMain"));
+IMPLEMENT_DEFERREDLIGHT_PIXELSHADER_TYPE(false, false, false, false, true, TEXT("DirectionalPixelMain"));
+IMPLEMENT_DEFERREDLIGHT_PIXELSHADER_TYPE(false, false, false, true, true, TEXT("DirectionalPixelMain"));
+IMPLEMENT_DEFERREDLIGHT_PIXELSHADER_TYPE(false, true, false, true, true, TEXT("RadialPixelMain"));
 
 /** Shader used to visualize stationary light overlap. */
 template<bool bRadialAttenuation>
@@ -683,7 +692,7 @@ static FVertexDeclarationRHIParamRef GetDeferredLightingVertexDeclaration()
 	return bRadialAttenuation ? GetVertexDeclarationFVector4() : GFilterVertexDeclaration.VertexDeclarationRHI;
 }
 
-template<bool bUseIESProfile, bool bRadialAttenuation, bool bInverseSquaredFalloff>
+template<bool bUseIESProfile, bool bRadialAttenuation, bool bInverseSquaredFalloff, bool bHandleClearCoat>
 static void SetShaderTemplLighting(
 	FRHICommandList& RHICmdList,
 	const FViewInfo& View,
@@ -692,19 +701,19 @@ static void SetShaderTemplLighting(
 {
 	if(View.Family->EngineShowFlags.VisualizeLightCulling)
 	{
-		TShaderMapRef<TDeferredLightPS<false, bRadialAttenuation, false, true> > PixelShader(View.ShaderMap);
+		TShaderMapRef<TDeferredLightPS<false, bRadialAttenuation, false, true, bHandleClearCoat> > PixelShader(View.ShaderMap);
 		SetGlobalBoundShaderState(RHICmdList, View.GetFeatureLevel(), PixelShader->GetBoundShaderState(), GetDeferredLightingVertexDeclaration<bRadialAttenuation>(), VertexShader, *PixelShader);
 		PixelShader->SetParameters(RHICmdList, View, LightSceneInfo);
 	}
 	else
 	{
-		TShaderMapRef<TDeferredLightPS<bUseIESProfile, bRadialAttenuation, bInverseSquaredFalloff, false> > PixelShader(View.ShaderMap);
+		TShaderMapRef<TDeferredLightPS<bUseIESProfile, bRadialAttenuation, bInverseSquaredFalloff, false, bHandleClearCoat> > PixelShader(View.ShaderMap);
 		SetGlobalBoundShaderState(RHICmdList, View.GetFeatureLevel(), PixelShader->GetBoundShaderState(), GetDeferredLightingVertexDeclaration<bRadialAttenuation>(), VertexShader, *PixelShader);
 		PixelShader->SetParameters(RHICmdList, View, LightSceneInfo);
 	}
 }
 
-template<bool bUseIESProfile, bool bRadialAttenuation, bool bInverseSquaredFalloff>
+template<bool bUseIESProfile, bool bRadialAttenuation, bool bInverseSquaredFalloff, bool bHandleClearCoat>
 static void SetShaderTemplLightingSimple(
 	FRHICommandList& RHICmdList,
 	const FViewInfo& View,
@@ -714,13 +723,13 @@ static void SetShaderTemplLightingSimple(
 {
 	if(View.Family->EngineShowFlags.VisualizeLightCulling)
 	{
-		TShaderMapRef<TDeferredLightPS<false, bRadialAttenuation, false, true> > PixelShader(View.ShaderMap);
+		TShaderMapRef<TDeferredLightPS<false, bRadialAttenuation, false, true, bHandleClearCoat> > PixelShader(View.ShaderMap);
 		SetGlobalBoundShaderState(RHICmdList, View.GetFeatureLevel(), PixelShader->GetBoundShaderState(), GetDeferredLightingVertexDeclaration<bRadialAttenuation>(), VertexShader, *PixelShader);
 		PixelShader->SetParametersSimpleLight(RHICmdList, View, SimpleLight, SimpleLightPerViewData);
 	}
 	else
 	{
-		TShaderMapRef<TDeferredLightPS<bUseIESProfile, bRadialAttenuation, bInverseSquaredFalloff, false> > PixelShader(View.ShaderMap);
+		TShaderMapRef<TDeferredLightPS<bUseIESProfile, bRadialAttenuation, bInverseSquaredFalloff, false, bHandleClearCoat> > PixelShader(View.ShaderMap);
 		SetGlobalBoundShaderState(RHICmdList, View.GetFeatureLevel(), PixelShader->GetBoundShaderState(), GetDeferredLightingVertexDeclaration<bRadialAttenuation>(), VertexShader, *PixelShader);
 		PixelShader->SetParametersSimpleLight(RHICmdList, View, SimpleLight, SimpleLightPerViewData);
 	}
@@ -790,6 +799,7 @@ void FDeferredShadingSceneRenderer::RenderLight(FRHICommandList& RHICmdList, con
 		// Set the device viewport for the view.
 		RHICmdList.SetViewport(View.ViewRect.Min.X, View.ViewRect.Min.Y, 0.0f, View.ViewRect.Max.X, View.ViewRect.Max.Y, 1.0f);
 
+		bool bClearCoatNeeded = (View.LightingProfilesActiveInView & (1 << MSM_ClearCoat)) != 0;
 		if (LightSceneInfo->Proxy->GetLightType() == LightType_Directional)
 		{
 			TShaderMapRef<TDeferredLightVS<false> > VertexShader(View.ShaderMap);
@@ -807,11 +817,25 @@ void FDeferredShadingSceneRenderer::RenderLight(FRHICommandList& RHICmdList, con
 			{
 				if(bUseIESTexture)
 				{
-					SetShaderTemplLighting<true, false, false>(RHICmdList, View, *VertexShader, LightSceneInfo);
+					if (bClearCoatNeeded)
+					{
+						SetShaderTemplLighting<true, false, false, true>(RHICmdList, View, *VertexShader, LightSceneInfo);
+					}
+					else
+					{
+						SetShaderTemplLighting<true, false, false, false>(RHICmdList, View, *VertexShader, LightSceneInfo);
+					}
 				}
 				else
 				{
-					SetShaderTemplLighting<false, false, false>(RHICmdList, View, *VertexShader, LightSceneInfo);
+					if (bClearCoatNeeded)
+					{
+						SetShaderTemplLighting<false, false, false, true>(RHICmdList, View, *VertexShader, LightSceneInfo);
+					}
+					else
+					{
+						SetShaderTemplLighting<false, false, false, true>(RHICmdList, View, *VertexShader, LightSceneInfo);
+					}
 				}
 			}
 
@@ -847,22 +871,50 @@ void FDeferredShadingSceneRenderer::RenderLight(FRHICommandList& RHICmdList, con
 				{
 					if(bUseIESTexture)
 					{
-						SetShaderTemplLighting<true, true, true>(RHICmdList, View, *VertexShader, LightSceneInfo);
+						if (bClearCoatNeeded)
+						{
+							SetShaderTemplLighting<true, true, true, true>(RHICmdList, View, *VertexShader, LightSceneInfo);
+						}						
+						else
+						{
+							SetShaderTemplLighting<true, true, true, false>(RHICmdList, View, *VertexShader, LightSceneInfo);
+						}
 					}
 					else
 					{
-						SetShaderTemplLighting<false, true, true>(RHICmdList, View, *VertexShader, LightSceneInfo);
+						if (bClearCoatNeeded)
+						{
+							SetShaderTemplLighting<false, true, true, true>(RHICmdList, View, *VertexShader, LightSceneInfo);
+						}
+						else
+						{
+							SetShaderTemplLighting<false, true, true, false>(RHICmdList, View, *VertexShader, LightSceneInfo);
+						}
 					}
 				}
 				else
 				{
 					if(bUseIESTexture)
 					{
-						SetShaderTemplLighting<true, true, false>(RHICmdList, View, *VertexShader, LightSceneInfo);
+						if (bClearCoatNeeded)
+						{
+							SetShaderTemplLighting<true, true, false, true>(RHICmdList, View, *VertexShader, LightSceneInfo);
+						}
+						else
+						{
+							SetShaderTemplLighting<true, true, false, false>(RHICmdList, View, *VertexShader, LightSceneInfo);
+						}
 					}
 					else
 					{
-						SetShaderTemplLighting<false, true, false>(RHICmdList, View, *VertexShader, LightSceneInfo);
+						if (bClearCoatNeeded)
+						{
+							SetShaderTemplLighting<false, true, false, true>(RHICmdList, View, *VertexShader, LightSceneInfo);
+						}
+						else
+						{
+							SetShaderTemplLighting<false, true, false, false>(RHICmdList, View, *VertexShader, LightSceneInfo);
+						}
 					}
 				}
 			}
@@ -917,17 +969,20 @@ void FDeferredShadingSceneRenderer::RenderSimpleLightsStandardDeferred(FRHIComma
 	// Use additive blending for color
 	RHICmdList.SetBlendState(TStaticBlendState<CW_RGBA, BO_Add, BF_One, BF_One, BO_Add, BF_One, BF_One>::GetRHI());
 
+	
+
 	const int32 NumViews = Views.Num();
 	for (int32 LightIndex = 0; LightIndex < SimpleLights.InstanceData.Num(); LightIndex++)
 	{
 		const FSimpleLightEntry& SimpleLight = SimpleLights.InstanceData[LightIndex];
 
 		for (int32 ViewIndex = 0; ViewIndex < NumViews; ViewIndex++)
-	{
+		{
 			const FSimpleLightPerViewEntry& SimpleLightPerViewData = SimpleLights.GetViewDependentData(LightIndex, ViewIndex, NumViews);
 			const FSphere LightBounds(SimpleLightPerViewData.Position, SimpleLight.Radius);
 
 			FViewInfo& View = Views[ViewIndex];
+
 			// Set the device viewport for the view.
 			RHICmdList.SetViewport(View.ViewRect.Min.X, View.ViewRect.Min.Y, 0.0f, View.ViewRect.Max.X, View.ViewRect.Max.Y, 1.0f);
 
@@ -935,15 +990,30 @@ void FDeferredShadingSceneRenderer::RenderSimpleLightsStandardDeferred(FRHIComma
 
 			SetBoundingGeometryRasterizerAndDepthState(RHICmdList, View, LightBounds);
 
+			bool bClearCoatNeeded = (View.LightingProfilesActiveInView & (1 << MSM_ClearCoat)) != 0;
 			if (SimpleLight.Exponent == 0)
 			{
 				// inverse squared
-				SetShaderTemplLightingSimple<false, true, true>(RHICmdList, View, *VertexShader, SimpleLight, SimpleLightPerViewData);
+				if (bClearCoatNeeded)
+				{
+					SetShaderTemplLightingSimple<false, true, true, true>(RHICmdList, View, *VertexShader, SimpleLight, SimpleLightPerViewData);
+				}
+				else
+				{
+					SetShaderTemplLightingSimple<false, true, true, false>(RHICmdList, View, *VertexShader, SimpleLight, SimpleLightPerViewData);
+				}
 			}
 			else
 			{
 				// light's exponent, not inverse squared
-				SetShaderTemplLightingSimple<false, true, false>(RHICmdList, View, *VertexShader, SimpleLight, SimpleLightPerViewData);
+				if (bClearCoatNeeded)
+				{
+					SetShaderTemplLightingSimple<false, true, false, true>(RHICmdList, View, *VertexShader, SimpleLight, SimpleLightPerViewData);
+				}
+				else
+				{
+					SetShaderTemplLightingSimple<false, true, false, false>(RHICmdList, View, *VertexShader, SimpleLight, SimpleLightPerViewData);
+				}
 			}
 
 			VertexShader->SetSimpleLightParameters(RHICmdList, View, LightBounds);
diff --git a/Engine/Source/Runtime/Renderer/Private/ReflectionEnvironment.cpp b/Engine/Source/Runtime/Renderer/Private/ReflectionEnvironment.cpp
--- a/Engine/Source/Runtime/Renderer/Private/ReflectionEnvironment.cpp
+++ b/Engine/Source/Runtime/Renderer/Private/ReflectionEnvironment.cpp
@@ -46,6 +46,13 @@ static TAutoConsoleVariable<int32> CVarHalfResReflections(
 	TEXT(" 0 is off (default), 1 is on"),
 	ECVF_RenderThreadSafe);
 
+static TAutoConsoleVariable<int32> CVarDoTiledReflections(
+	TEXT("r.DoTiledReflections"),
+	1,
+	TEXT("Compute Reflection Environment with Tiled compute shader..\n")
+	TEXT(" 1 is on (default), 1 is on"),
+	ECVF_RenderThreadSafe);
+
 static TAutoConsoleVariable<float> CVarSkySpecularOcclusionStrength(
 	TEXT("r.SkySpecularOcclusionStrength"),
 	1,
@@ -243,7 +250,7 @@ public:
 	{
 	}
 
-	void SetParameters(FRHICommandList& RHICmdList, const FSceneView& View, FTextureRHIParamRef SSRTexture, FUnorderedAccessViewRHIParamRef OutSceneColorUAV, const TRefCountPtr<IPooledRenderTarget>& DynamicBentNormalAO)
+	void SetParameters(FRHICommandList& RHICmdList, const FSceneView& View, FTextureRHIParamRef SSRTexture, TArray<FReflectionCaptureSortData>& SortData, FUnorderedAccessViewRHIParamRef OutSceneColorUAV, const TRefCountPtr<IPooledRenderTarget>& DynamicBentNormalAO)
 	{
 		const FComputeShaderRHIParamRef ShaderRHI = GetComputeShader();
 
@@ -272,54 +279,6 @@ public:
 
 		SetShaderValue(RHICmdList, ShaderRHI, ViewDimensionsParameter, View.ViewRect);
 
-		static TArray<FReflectionCaptureSortData> SortData;
-		SortData.Reset(Scene->ReflectionSceneData.RegisteredReflectionCaptures.Num());
-
-		const int32 MaxCubemaps = Scene->ReflectionSceneData.CubemapArray.GetMaxCubemaps();
-
-		// Pack only visible reflection captures into the uniform buffer, each with an index to its cubemap array entry
-		for (int32 ReflectionProxyIndex = 0; ReflectionProxyIndex < Scene->ReflectionSceneData.RegisteredReflectionCaptures.Num() && SortData.Num() < GMaxNumReflectionCaptures; ReflectionProxyIndex++)
-		{
-			FReflectionCaptureProxy* CurrentCapture = Scene->ReflectionSceneData.RegisteredReflectionCaptures[ReflectionProxyIndex];
-			// Find the cubemap index this component was allocated with
-			const FCaptureComponentSceneState* ComponentStatePtr = Scene->ReflectionSceneData.AllocatedReflectionCaptureState.Find(CurrentCapture->Component);
-
-			if (ComponentStatePtr)
-			{
-				int32 CubemapIndex = ComponentStatePtr->CaptureIndex;
-				check(CubemapIndex < MaxCubemaps);
-
-				FReflectionCaptureSortData NewSortEntry;
-
-				NewSortEntry.CaptureIndex = CubemapIndex;
-				NewSortEntry.SM4FullHDRCubemap = NULL;
-				NewSortEntry.Guid = CurrentCapture->Guid;
-				NewSortEntry.PositionAndRadius = FVector4(CurrentCapture->Position, CurrentCapture->InfluenceRadius);
-				float ShapeTypeValue = (float)CurrentCapture->Shape;
-				NewSortEntry.CaptureProperties = FVector4(CurrentCapture->Brightness, CubemapIndex, ShapeTypeValue, 0);
-
-				if (CurrentCapture->Shape == EReflectionCaptureShape::Plane)
-				{
-					NewSortEntry.BoxTransform = FMatrix(
-						FPlane(CurrentCapture->ReflectionPlane), 
-						FPlane(CurrentCapture->ReflectionXAxisAndYScale), 
-						FPlane(0, 0, 0, 0), 
-						FPlane(0, 0, 0, 0));
-
-					NewSortEntry.BoxScales = FVector4(0);
-				}
-				else
-				{
-					NewSortEntry.BoxTransform = CurrentCapture->BoxTransform;
-
-					NewSortEntry.BoxScales = FVector4(CurrentCapture->BoxScales, CurrentCapture->BoxTransitionDistance);
-				}
-
-				SortData.Add(NewSortEntry);
-			}
-		}
-
-		SortData.Sort();
 		FReflectionCaptureData SamplePositionsBuffer;
 
 		for (int32 CaptureIndex = 0; CaptureIndex < SortData.Num(); CaptureIndex++)
@@ -381,7 +340,7 @@ private:
 	FDistanceFieldAOSpecularOcclusionParameters SpecularOcclusionParameters;
 };
 
-template< uint32 bUseLightmaps, uint32 bHalfRes >
+template< uint32 bUseLightmaps, uint32 bUseClearCoat, uint32 bBoxCapturesOnly, uint32 bSphereCapturesOnly >
 class TReflectionEnvironmentTiledDeferredCS : public FReflectionEnvironmentTiledDeferredCS
 {
 	DECLARE_SHADER_TYPE(TReflectionEnvironmentTiledDeferredCS, Global);
@@ -397,19 +356,37 @@ public:
 	{
 		FReflectionEnvironmentTiledDeferredCS::ModifyCompilationEnvironment(Platform, OutEnvironment);
 		OutEnvironment.SetDefine(TEXT("USE_LIGHTMAPS"), bUseLightmaps);
-		OutEnvironment.SetDefine(TEXT("HALF_RESOLUTION"), bHalfRes);
+		OutEnvironment.SetDefine(TEXT("USE_CLEARCOAT"), bUseClearCoat);
+		OutEnvironment.SetDefine(TEXT("HAS_BOX_CAPTURES"), bBoxCapturesOnly);
+		OutEnvironment.SetDefine(TEXT("HAS_SPHERE_CAPTURES"), bSphereCapturesOnly);
 	}
 };
 
 // Typedef is necessary because the C preprocessor thinks the comma in the template parameter list is a comma in the macro parameter list.
-#define IMPLEMENT_REFLECTION_COMPUTESHADER_TYPE(A, B) \
-	typedef TReflectionEnvironmentTiledDeferredCS<A,B> TReflectionEnvironmentTiledDeferredCS##A##B; \
-	IMPLEMENT_SHADER_TYPE(template<>,TReflectionEnvironmentTiledDeferredCS##A##B,TEXT("ReflectionEnvironmentComputeShaders"),TEXT("ReflectionEnvironmentTiledDeferredMain"),SF_Compute)
+#define IMPLEMENT_REFLECTION_COMPUTESHADER_TYPE(A, B, C, D) \
+	typedef TReflectionEnvironmentTiledDeferredCS<A,B,C,D> TReflectionEnvironmentTiledDeferredCS##A##B##C##D; \
+	IMPLEMENT_SHADER_TYPE(template<>,TReflectionEnvironmentTiledDeferredCS##A##B##C##D,TEXT("ReflectionEnvironmentComputeShaders"),TEXT("ReflectionEnvironmentTiledDeferredMain"),SF_Compute)
+
+IMPLEMENT_REFLECTION_COMPUTESHADER_TYPE(0, 0, 0, 0);
+IMPLEMENT_REFLECTION_COMPUTESHADER_TYPE(0, 0, 0, 1);
+IMPLEMENT_REFLECTION_COMPUTESHADER_TYPE(0, 0, 1, 0);
+IMPLEMENT_REFLECTION_COMPUTESHADER_TYPE(0, 0, 1, 1);
+
+IMPLEMENT_REFLECTION_COMPUTESHADER_TYPE(0, 1, 0, 0);
+IMPLEMENT_REFLECTION_COMPUTESHADER_TYPE(0, 1, 0, 1);
+IMPLEMENT_REFLECTION_COMPUTESHADER_TYPE(0, 1, 1, 0);
+IMPLEMENT_REFLECTION_COMPUTESHADER_TYPE(0, 1, 1, 1);
+
+IMPLEMENT_REFLECTION_COMPUTESHADER_TYPE(1, 0, 0, 0);
+IMPLEMENT_REFLECTION_COMPUTESHADER_TYPE(1, 0, 0, 1);
+IMPLEMENT_REFLECTION_COMPUTESHADER_TYPE(1, 0, 1, 0);
+IMPLEMENT_REFLECTION_COMPUTESHADER_TYPE(1, 0, 1, 1);
+
+IMPLEMENT_REFLECTION_COMPUTESHADER_TYPE(1, 1, 0, 0);
+IMPLEMENT_REFLECTION_COMPUTESHADER_TYPE(1, 1, 0, 1);
+IMPLEMENT_REFLECTION_COMPUTESHADER_TYPE(1, 1, 1, 0);
+IMPLEMENT_REFLECTION_COMPUTESHADER_TYPE(1, 1, 1, 1);
 
-IMPLEMENT_REFLECTION_COMPUTESHADER_TYPE(0,0);
-IMPLEMENT_REFLECTION_COMPUTESHADER_TYPE(0,1);
-IMPLEMENT_REFLECTION_COMPUTESHADER_TYPE(1,0);
-IMPLEMENT_REFLECTION_COMPUTESHADER_TYPE(1,1);
 
 template< uint32 bSSR, uint32 bReflectionEnv, uint32 bSkylight >
 class FReflectionApplyPS : public FGlobalShader
@@ -702,9 +679,157 @@ bool FDeferredShadingSceneRenderer::ShouldDoReflectionEnvironment() const
 		&& (FeatureLevel == ERHIFeatureLevel::SM4 || Scene->ReflectionSceneData.CubemapArray.IsValid());
 }
 
+void GatherAndSortReflectionCaptures(const FScene* Scene, TArray<FReflectionCaptureSortData>& OutSortData, int32& OutNumBoxCaptures, int32& OutNumSphereCaptures)
+{	
+	OutSortData.Reset(Scene->ReflectionSceneData.RegisteredReflectionCaptures.Num());
+	OutNumBoxCaptures = 0;
+	OutNumSphereCaptures = 0;
+
+	const int32 MaxCubemaps = Scene->ReflectionSceneData.CubemapArray.GetMaxCubemaps();
+
+	// Pack only visible reflection captures into the uniform buffer, each with an index to its cubemap array entry
+	for (int32 ReflectionProxyIndex = 0; ReflectionProxyIndex < Scene->ReflectionSceneData.RegisteredReflectionCaptures.Num() && OutSortData.Num() < GMaxNumReflectionCaptures; ReflectionProxyIndex++)
+	{
+		FReflectionCaptureProxy* CurrentCapture = Scene->ReflectionSceneData.RegisteredReflectionCaptures[ReflectionProxyIndex];
+		// Find the cubemap index this component was allocated with
+		const FCaptureComponentSceneState* ComponentStatePtr = Scene->ReflectionSceneData.AllocatedReflectionCaptureState.Find(CurrentCapture->Component);
+
+		if (ComponentStatePtr)
+		{
+			int32 CubemapIndex = ComponentStatePtr->CaptureIndex;
+			check(CubemapIndex < MaxCubemaps);
+
+			FReflectionCaptureSortData NewSortEntry;
+
+			NewSortEntry.CaptureIndex = CubemapIndex;
+			NewSortEntry.SM4FullHDRCubemap = NULL;
+			NewSortEntry.Guid = CurrentCapture->Guid;
+			NewSortEntry.PositionAndRadius = FVector4(CurrentCapture->Position, CurrentCapture->InfluenceRadius);
+			float ShapeTypeValue = (float)CurrentCapture->Shape;
+			NewSortEntry.CaptureProperties = FVector4(CurrentCapture->Brightness, CubemapIndex, ShapeTypeValue, 0);
+
+			if (CurrentCapture->Shape == EReflectionCaptureShape::Plane)
+			{
+				//planes count as boxes in the compute shader.
+				++OutNumBoxCaptures;
+				NewSortEntry.BoxTransform = FMatrix(
+					FPlane(CurrentCapture->ReflectionPlane),
+					FPlane(CurrentCapture->ReflectionXAxisAndYScale),
+					FPlane(0, 0, 0, 0),
+					FPlane(0, 0, 0, 0));
+
+				NewSortEntry.BoxScales = FVector4(0);
+			}
+			else if (CurrentCapture->Shape == EReflectionCaptureShape::Sphere)
+			{
+				++OutNumSphereCaptures;
+			}
+			else
+			{
+				++OutNumBoxCaptures;
+				NewSortEntry.BoxTransform = CurrentCapture->BoxTransform;
+				NewSortEntry.BoxScales = FVector4(CurrentCapture->BoxScales, CurrentCapture->BoxTransitionDistance);
+			}
+
+			OutSortData.Add(NewSortEntry);
+		}
+	}
+
+	OutSortData.Sort();	
+}
+
+FReflectionEnvironmentTiledDeferredCS* SelectReflectionEnvironmentTiledDeferredCS(TShaderMap<FGlobalShaderType>* ShaderMap, bool bUseLightmaps, bool bUseClearCoat, bool bHasBoxCaptures, bool bHasSphereCaptures)
+{
+	FReflectionEnvironmentTiledDeferredCS* ComputeShader = nullptr;
+	if (bUseLightmaps)
+	{
+		if (bUseClearCoat)
+		{
+			if (bHasBoxCaptures && bHasSphereCaptures)
+			{
+				ComputeShader = *TShaderMapRef< TReflectionEnvironmentTiledDeferredCS<1, 1, 1, 1> >(ShaderMap);
+			}
+			else if (bHasBoxCaptures)
+			{
+				ComputeShader = *TShaderMapRef< TReflectionEnvironmentTiledDeferredCS<1, 1, 1, 0> >(ShaderMap);
+			}
+			else if (bHasSphereCaptures)
+			{
+				ComputeShader = *TShaderMapRef< TReflectionEnvironmentTiledDeferredCS<1, 1, 0, 1> >(ShaderMap);
+			}
+			else
+			{
+				ComputeShader = *TShaderMapRef< TReflectionEnvironmentTiledDeferredCS<1, 1, 0, 0> >(ShaderMap);
+			}
+		}
+		else
+		{
+			if (bHasBoxCaptures && bHasSphereCaptures)
+			{
+				ComputeShader = *TShaderMapRef< TReflectionEnvironmentTiledDeferredCS<1, 0, 1, 1> >(ShaderMap);
+			}
+			else if (bHasBoxCaptures)
+			{
+				ComputeShader = *TShaderMapRef< TReflectionEnvironmentTiledDeferredCS<1, 0, 1, 0> >(ShaderMap);
+			}
+			else if (bHasSphereCaptures)
+			{
+				ComputeShader = *TShaderMapRef< TReflectionEnvironmentTiledDeferredCS<1, 0, 0, 1> >(ShaderMap);
+			}
+			else
+			{
+				ComputeShader = *TShaderMapRef< TReflectionEnvironmentTiledDeferredCS<1, 0, 0, 0> >(ShaderMap);
+			}
+		}		
+	}
+	else
+	{
+		if (bUseClearCoat)
+		{
+			if (bHasBoxCaptures && bHasSphereCaptures)
+			{
+				ComputeShader = *TShaderMapRef< TReflectionEnvironmentTiledDeferredCS<0, 1, 1, 1> >(ShaderMap);
+			}
+			else if (bHasBoxCaptures)
+			{
+				ComputeShader = *TShaderMapRef< TReflectionEnvironmentTiledDeferredCS<0, 1, 1, 0> >(ShaderMap);
+			}
+			else if (bHasSphereCaptures)
+			{
+				ComputeShader = *TShaderMapRef< TReflectionEnvironmentTiledDeferredCS<0, 1, 0, 1> >(ShaderMap);
+			}
+			else
+			{
+				ComputeShader = *TShaderMapRef< TReflectionEnvironmentTiledDeferredCS<0, 1, 0, 0> >(ShaderMap);
+			}
+		}
+		else
+		{
+			if (bHasBoxCaptures && bHasSphereCaptures)
+			{
+				ComputeShader = *TShaderMapRef< TReflectionEnvironmentTiledDeferredCS<0, 0, 1, 1> >(ShaderMap);
+			}
+			else if (bHasBoxCaptures)
+			{
+				ComputeShader = *TShaderMapRef< TReflectionEnvironmentTiledDeferredCS<0, 0, 1, 0> >(ShaderMap);
+			}
+			else if (bHasSphereCaptures)
+			{
+				ComputeShader = *TShaderMapRef< TReflectionEnvironmentTiledDeferredCS<0, 0, 0, 1> >(ShaderMap);
+			}
+			else
+			{
+				ComputeShader = *TShaderMapRef< TReflectionEnvironmentTiledDeferredCS<0, 0, 0, 0> >(ShaderMap);
+			}
+		}		
+	}
+	check(ComputeShader);
+	return ComputeShader;
+}
+
 void FDeferredShadingSceneRenderer::RenderTiledDeferredImageBasedReflections(FRHICommandListImmediate& RHICmdList, const TRefCountPtr<IPooledRenderTarget>& DynamicBentNormalAO)
 {
-	const uint32 bUseLightmaps = CVarDiffuseFromCaptures.GetValueOnRenderThread() == 0;
+	const bool bUseLightmaps = CVarDiffuseFromCaptures.GetValueOnRenderThread() == 0;
 	const uint32 bHalfRes = CVarHalfResReflections.GetValueOnRenderThread() != 0;
 
 	TRefCountPtr<IPooledRenderTarget> NewSceneColor;
@@ -740,33 +865,24 @@ void FDeferredShadingSceneRenderer::RenderTiledDeferredImageBasedReflections(FRH
 
 			FReflectionEnvironmentTiledDeferredCS* ComputeShader = NULL;
 			uint32 AdjustedReflectionTileSizeX = GReflectionEnvironmentTileSizeX;
-			if ( bHalfRes)
+			if (bHalfRes)
 			{
-				if( bUseLightmaps )
-				{
-					ComputeShader = *TShaderMapRef< TReflectionEnvironmentTiledDeferredCS<1,1> >( View.ShaderMap );
-				}
-				else
-				{
-					ComputeShader = *TShaderMapRef< TReflectionEnvironmentTiledDeferredCS<0,1> >( View.ShaderMap );
-				}
 				AdjustedReflectionTileSizeX *= 2;
 			}
-			else
-			{
-			if( bUseLightmaps )
-			{
-					ComputeShader = *TShaderMapRef< TReflectionEnvironmentTiledDeferredCS<1,0> >( View.ShaderMap );
-			}
-			else
-			{
-					ComputeShader = *TShaderMapRef< TReflectionEnvironmentTiledDeferredCS<0,0> >( View.ShaderMap );
-				}
-			}
+			
+			TArray<FReflectionCaptureSortData> SortData;
+			int32 NumBoxCaptures = 0;
+			int32 NumSphereCaptures = 0;
+			GatherAndSortReflectionCaptures(Scene, SortData, NumBoxCaptures, NumSphereCaptures);
+
+			bool bHasBoxCaptures = (NumBoxCaptures > 0);
+			bool bHasSphereCaptures = (NumSphereCaptures > 0);
+			bool bNeedsClearCoat = (View.LightingProfilesActiveInView & (1 << MSM_ClearCoat)) != 0;
+			ComputeShader = SelectReflectionEnvironmentTiledDeferredCS(View.ShaderMap, bUseLightmaps, bNeedsClearCoat, bHasBoxCaptures, bHasSphereCaptures);
 
 			RHICmdList.SetComputeShader(ComputeShader->GetComputeShader());
 
-			ComputeShader->SetParameters(RHICmdList, View, SSROutput->GetRenderTargetItem().ShaderResourceTexture, NewSceneColor->GetRenderTargetItem().UAV, DynamicBentNormalAO);
+			ComputeShader->SetParameters(RHICmdList, View, SSROutput->GetRenderTargetItem().ShaderResourceTexture, SortData, NewSceneColor->GetRenderTargetItem().UAV, DynamicBentNormalAO);
 
 			uint32 GroupSizeX = (View.ViewRect.Size().X + AdjustedReflectionTileSizeX - 1) / AdjustedReflectionTileSizeX;
 			uint32 GroupSizeY = (View.ViewRect.Size().Y + GReflectionEnvironmentTileSizeY - 1) / GReflectionEnvironmentTileSizeY;
@@ -1008,8 +1124,9 @@ void FDeferredShadingSceneRenderer::RenderDeferredReflections(FRHICommandListImm
 	}
 	else
 	{
+		const uint32 bDoTiledReflections = CVarDoTiledReflections.GetValueOnRenderThread() != 0;
 		const bool bReflectionEnvironment = ShouldDoReflectionEnvironment();
-		const bool bReflectionsWithCompute = (FeatureLevel >= ERHIFeatureLevel::SM5) && bReflectionEnvironment && Scene->ReflectionSceneData.CubemapArray.IsValid();
+		const bool bReflectionsWithCompute = bDoTiledReflections && (FeatureLevel >= ERHIFeatureLevel::SM5) && bReflectionEnvironment && Scene->ReflectionSceneData.CubemapArray.IsValid();
 		
 		if (bReflectionsWithCompute)
 		{
diff --git a/Engine/Source/Runtime/Renderer/Private/SceneRendering.cpp b/Engine/Source/Runtime/Renderer/Private/SceneRendering.cpp
--- a/Engine/Source/Runtime/Renderer/Private/SceneRendering.cpp
+++ b/Engine/Source/Runtime/Renderer/Private/SceneRendering.cpp
@@ -210,9 +210,9 @@ void FViewInfo::Init()
 	bPrevTransformsReset = false;
 	bIgnoreExistingQueries = false;
 	bDisableQuerySubmissions = false;
-	bDisableDistanceBasedFadeTransitions = false;
-	bScreenSpaceSubsurfacePassNeeded = false;
-	NumVisibleStaticMeshElements = false;
+	bDisableDistanceBasedFadeTransitions = false;	
+	LightingProfilesActiveInView = 0;
+	NumVisibleStaticMeshElements = 0;
 	PrecomputedVisibilityData = 0;
 
 	bIsViewInfo = true;
diff --git a/Engine/Source/Runtime/Renderer/Private/SceneRendering.h b/Engine/Source/Runtime/Renderer/Private/SceneRendering.h
--- a/Engine/Source/Runtime/Renderer/Private/SceneRendering.h
+++ b/Engine/Source/Runtime/Renderer/Private/SceneRendering.h
@@ -442,8 +442,8 @@ public:
 	uint32 bDisableQuerySubmissions : 1;
 	/** Whether we should disable distance-based fade transitions for this frame (usually after a large camera movement.) */
 	uint32 bDisableDistanceBasedFadeTransitions : 1;
-	/** Whether we render any objects with a material using SubsurfaceProfile as shading model. */
-	uint32 bScreenSpaceSubsurfacePassNeeded : 1;
+	/** Bitmask of all lighting profiles used by primitives in this view */
+	uint16 LightingProfilesActiveInView;
 
 	FViewMatrices PrevViewMatrices;
 
diff --git a/Engine/Source/Runtime/Renderer/Private/SceneVisibility.cpp b/Engine/Source/Runtime/Renderer/Private/SceneVisibility.cpp
--- a/Engine/Source/Runtime/Renderer/Private/SceneVisibility.cpp
+++ b/Engine/Source/Runtime/Renderer/Private/SceneVisibility.cpp
@@ -753,7 +753,7 @@ static void ComputeRelevanceForView(
 		const bool bDynamicRelevance = ViewRelevance.bDynamicRelevance;
 		const bool bShadowRelevance = ViewRelevance.bShadowRelevance;
 		const bool bEditorRelevance = ViewRelevance.bEditorPrimitiveRelevance;
-		const bool bTranslucentRelevance = ViewRelevance.HasTranslucency();
+		const bool bTranslucentRelevance = ViewRelevance.HasTranslucency();		
 
 		if (bStaticRelevance && (bDrawRelevance || bShadowRelevance))
 		{
@@ -795,10 +795,7 @@ static void ComputeRelevanceForView(
 			}
 		}
 		
-		if (ViewRelevance.bSubsurfaceProfileRelevance)
-		{
-			View.bScreenSpaceSubsurfacePassNeeded = true;
-		}
+		View.LightingProfilesActiveInView |= ViewRelevance.LightingProfileRelevanceMask;
 		
 		if (ViewRelevance.bRenderCustomDepth)
 		{
@@ -937,6 +934,7 @@ struct FRelevancePacket
 	FRelevancePrimSet<FPrimitiveSceneProxy*> CustomDepthSet;
 	FRelevancePrimSet<FPrimitiveSceneInfo*> UpdateStaticMeshes;
 	FRelevancePrimSet<FPrimitiveSceneInfo*> VisibleEditorPrimitives;
+	uint16 CombinedLightingProfileMask;
 
 	FRelevancePacket(
 		FRHICommandListImmediate& InRHICmdList,
@@ -958,6 +956,7 @@ struct FRelevancePacket
 		, OutHasDynamicMeshElementsMasks(InOutHasDynamicMeshElementsMasks)
 		, OutHasDynamicEditorMeshElementsMasks(InOutHasDynamicEditorMeshElementsMasks)
 		, MarkMasks(InMarkMasks)
+		, CombinedLightingProfileMask(0)
 	{
 	}
 
@@ -969,6 +968,7 @@ struct FRelevancePacket
 
 	void ComputeRelevance()
 	{
+		CombinedLightingProfileMask = 0;
 		SCOPE_CYCLE_COUNTER(STAT_ComputeViewRelevance);
 		for (int32 Index = 0; Index < Input.NumPrims; Index++)
 		{
@@ -1028,11 +1028,7 @@ struct FRelevancePacket
 				}
 			}
 
-			if (ViewRelevance.bSubsurfaceProfileRelevance)
-			{
-				// this is ok because it is write only
-				const_cast<FViewInfo&>(View).bScreenSpaceSubsurfacePassNeeded = true;
-			}
+			CombinedLightingProfileMask |= ViewRelevance.LightingProfileRelevanceMask;			
 
 			if (ViewRelevance.bRenderCustomDepth)
 			{
@@ -1151,11 +1147,12 @@ struct FRelevancePacket
 	void RenderThreadFinalize()
 	{
 		FViewInfo& WriteView = const_cast<FViewInfo&>(View);
-
+		
 		for (int32 Index = 0; Index < NotDrawRelevant.NumPrims; Index++)
 		{
 			WriteView.PrimitiveVisibilityMap[NotDrawRelevant.Prims[Index]] = false;
 		}
+		WriteView.LightingProfilesActiveInView |= CombinedLightingProfileMask;
 		VisibleEditorPrimitives.AppendTo(WriteView.VisibleEditorPrimitives);
 		VisibleDynamicPrimitives.AppendTo(WriteView.VisibleDynamicPrimitives);
 		WriteView.TranslucentPrimSet.AppendScenePrimitives(SortedTranslucencyPrims.Prims, SortedTranslucencyPrims.NumPrims, SortedSeparateTranslucencyPrims.Prims, SortedSeparateTranslucencyPrims.NumPrims);
