fixed UE-7950 DepthOfFieldFunction doesn't work

diff --git a/Engine/Shaders/DepthOfFieldCommon.usf b/Engine/Shaders/DepthOfFieldCommon.usf
--- a/Engine/Shaders/DepthOfFieldCommon.usf
+++ b/Engine/Shaders/DepthOfFieldCommon.usf
@@ -47,6 +47,23 @@ float ComputeCircleOfConfusion(float SceneDepth)
 	return saturate(abs(CoCRadius));
 }
 
+
+// @param SceneDepth
+// @return 0..1 0=in focus, 1:max blurry
+float ComputeCircleOfConfusionNorm(float SceneDepth)
+{
+	// artificial area where all content is in focus (starting at FocalLength, ending at FocalLength+FocalRegion)
+	FLATTEN if(SceneDepth > View.DepthOfFieldFocalDistance)
+	{
+		SceneDepth = View.DepthOfFieldFocalDistance + max(0, SceneDepth - View.DepthOfFieldFocalDistance - View.DepthOfFieldFocalRegion);
+	}
+
+	//	assumed to be >0 by c++ setup code
+	half TransitionRegion = (SceneDepth < View.DepthOfFieldFocalDistance) ? View.DepthOfFieldNearTransitionRegion : View.DepthOfFieldFarTransitionRegion;
+
+	return saturate(abs(SceneDepth - View.DepthOfFieldFocalDistance) / TransitionRegion);
+}
+
 /**
  * Computes the unfocused percent for a scene depth.
  * @param SceneDepth - The scene depth.
@@ -55,7 +72,8 @@ float ComputeCircleOfConfusion(float SceneDepth)
 half CalcUnfocusedPercentCustomBound(float SceneDepth, float MaxBlurNear, float MaxBlurFar)
 {
 	half MaxUnfocusedPercent = (SceneDepth < View.DepthOfFieldFocalDistance) ? MaxBlurNear : MaxBlurFar;
-	half Unbound = ComputeCircleOfConfusion(SceneDepth);
+
+	half Unbound = ComputeCircleOfConfusionNorm(SceneDepth);
 
 	return min(MaxUnfocusedPercent, Unbound);
 }
diff --git a/Engine/Source/Runtime/Renderer/Private/SceneRendering.cpp b/Engine/Source/Runtime/Renderer/Private/SceneRendering.cpp
--- a/Engine/Source/Runtime/Renderer/Private/SceneRendering.cpp
+++ b/Engine/Source/Runtime/Renderer/Private/SceneRendering.cpp
@@ -621,8 +621,10 @@ TUniformBufferRef<FViewUniformShaderParameters> FViewInfo::CreateUniformBuffer(
 	ViewUniformShaderParameters.ExposureScale = FVector4(ExposureScale,ExposureScale,ExposureScale,1.0f);
 	ViewUniformShaderParameters.DepthOfFieldFocalDistance = FinalPostProcessSettings.DepthOfFieldFocalDistance;
 	ViewUniformShaderParameters.DepthOfFieldFocalRegion = FinalPostProcessSettings.DepthOfFieldFocalRegion;
-	ViewUniformShaderParameters.DepthOfFieldNearTransitionRegion = FinalPostProcessSettings.DepthOfFieldNearTransitionRegion;
-	ViewUniformShaderParameters.DepthOfFieldFarTransitionRegion = FinalPostProcessSettings.DepthOfFieldFarTransitionRegion;
+	// clamped to avoid div by 0 in shader
+	ViewUniformShaderParameters.DepthOfFieldNearTransitionRegion = FMath::Max(0.01f, FinalPostProcessSettings.DepthOfFieldNearTransitionRegion);
+	// clamped to avoid div by 0 in shader
+	ViewUniformShaderParameters.DepthOfFieldFarTransitionRegion = FMath::Max(0.01f, FinalPostProcessSettings.DepthOfFieldFarTransitionRegion);
 	ViewUniformShaderParameters.DepthOfFieldScale = FinalPostProcessSettings.DepthOfFieldScale;
 	ViewUniformShaderParameters.DepthOfFieldFocalLength = 50.0f;
 	ViewUniformShaderParameters.MotionBlurNormalizedToPixel = FinalPostProcessSettings.MotionBlurMax * ViewRect.Width() / 100.0f;
