Support for new Morpheus HMD Server.  Morpheus PC preview now works.  TargetSettings must set Demo mode to 'Demo'.
Modifies postprocessing pipeline so that distortion is done on buffers scaled by ScreenPercentage.  Scaling up ScreenPercentage will now have a bigger effect on aliasing in the HMD.
Modifies headtracking to be driven by VBLank events rather than the render thread.  Reduces tracking loss.

diff --git a/Engine/Source/Programs/UnrealBuildTool/UnrealBuildTool.csproj b/Engine/Source/Programs/UnrealBuildTool/UnrealBuildTool.csproj
--- a/Engine/Source/Programs/UnrealBuildTool/UnrealBuildTool.csproj
+++ b/Engine/Source/Programs/UnrealBuildTool/UnrealBuildTool.csproj
@@ -197,6 +197,7 @@
     <Compile Include="Properties\AssemblyInfo.cs" />
     <Compile Include="Configuration\UEBuildConfiguration.cs" />
     <Compile Include="Configuration\UEBuildTarget.cs" />
+    <Compile Include="PS4\UEDeployPS4.cs" />
     <Compile Include="System\ActionGraph.cs" />
     <Compile Include="System\ActionHistory.cs" />
     <Compile Include="System\BuildException.cs" />
diff --git a/Engine/Source/Programs/UnrealBuildTool/Windows/UEBuildWindows.cs b/Engine/Source/Programs/UnrealBuildTool/Windows/UEBuildWindows.cs
--- a/Engine/Source/Programs/UnrealBuildTool/Windows/UEBuildWindows.cs
+++ b/Engine/Source/Programs/UnrealBuildTool/Windows/UEBuildWindows.cs
@@ -456,7 +456,7 @@ namespace UnrealBuildTool
 
 				//on PS4 the SDK now handles distortion correction.  On PC we will still have to handle it manually,
 				//but we require SDK changes before we can get the required data.  For the moment, no platform does in-engine morpheus distortion
-				InBuildTarget.GlobalCompileEnvironment.Config.Definitions.Add("MORPHEUS_ENGINE_DISTORTION=0");				
+				InBuildTarget.GlobalCompileEnvironment.Config.Definitions.Add("MORPHEUS_ENGINE_DISTORTION=1");				
             }
 
             if (InBuildTarget.Rules != null)
diff --git a/Engine/Source/Runtime/HeadMountedDisplay/Public/IHeadMountedDisplay.h b/Engine/Source/Runtime/HeadMountedDisplay/Public/IHeadMountedDisplay.h
--- a/Engine/Source/Runtime/HeadMountedDisplay/Public/IHeadMountedDisplay.h
+++ b/Engine/Source/Runtime/HeadMountedDisplay/Public/IHeadMountedDisplay.h
@@ -273,6 +273,8 @@ public:
 
 	virtual void DrawDistortionMesh_RenderThread(struct FRenderingCompositePassContext& Context, const FSceneView& View, const FIntPoint& TextureSize) {}
 
+	virtual void DrawSingleEyeDistortionSDK_RenderThread(FRHICommandListImmediate& RHICmdList, const FSceneView& View, const FTexture2DRHIRef& InputTexture) {}
+
 	/**
 	 * This method is able to change screen settings right before any drawing occurs. 
 	 * It is called at the beginning of UGameViewportClient::Draw() method.
diff --git a/Engine/Source/Runtime/Renderer/Private/PostProcess/PostProcessMorpheus.cpp b/Engine/Source/Runtime/Renderer/Private/PostProcess/PostProcessMorpheus.cpp
--- a/Engine/Source/Runtime/Renderer/Private/PostProcess/PostProcessMorpheus.cpp
+++ b/Engine/Source/Runtime/Renderer/Private/PostProcess/PostProcessMorpheus.cpp
@@ -16,7 +16,7 @@
 
 DEFINE_LOG_CATEGORY_STATIC(LogMorpheusHMDPostProcess, All, All);
 
-#if MORPHEUS_ENGINE_DISTORTION
+#if HAS_MORPHEUS
 
 /** Encapsulates the post processing HMD distortion and correction pixel shader. */
 class FPostProcessMorpheusPS : public FGlobalShader
@@ -28,7 +28,7 @@ class FPostProcessMorpheusPS : public FGlobalShader
 		// we must use a run time check for this because the builds the build machines create will have Morpheus defined,
 		// but a user will not necessarily have the Morpheus files
 		bool bEnableMorpheus = false;
-		if (GConfig->GetBool(TEXT("Morpheus.Settings"), TEXT("EnableMorpheus"), bEnableMorpheus, GEngineIni))
+		if (GConfig->GetBool(TEXT("/Script/MorpheusEditor.MorpheusRuntimeSettings"), TEXT("bEnableMorpheus"), bEnableMorpheus, GEngineIni))
 		{
 			return bEnableMorpheus;
 		}
@@ -88,7 +88,7 @@ public:
 		
 		FGlobalShader::SetParameters(Context.RHICmdList, ShaderRHI, Context.View);
 
-		PostprocessParameter.SetPS(ShaderRHI, Context, TStaticSamplerState<SF_Bilinear,AM_Clamp,AM_Clamp,AM_Clamp>::GetRHI());
+		PostprocessParameter.SetPS(ShaderRHI, Context, TStaticSamplerState<SF_Bilinear, AM_Border, AM_Border, AM_Border>::GetRHI());
 		DeferredParameters.Set(Context.RHICmdList, ShaderRHI, Context.View);
 
 		{
@@ -142,7 +142,7 @@ class FPostProcessMorpheusVS : public FGlobalShader
 		// we must use a run time check for this because the builds the build machines create will have Morpheus defined,
 		// but a user will not necessarily have the Morpheus files
 		bool bEnableMorpheus = false;
-		if (GConfig->GetBool(TEXT("Morpheus.Settings"), TEXT("EnableMorpheus"), bEnableMorpheus, GEngineIni))
+		if (GConfig->GetBool(TEXT("/Script/MorpheusEditor.MorpheusRuntimeSettings"), TEXT("bEnableMorpheus"), bEnableMorpheus, GEngineIni))
 		{
 			return bEnableMorpheus;
 		}
@@ -196,7 +196,10 @@ void FRCPassPostProcessMorpheus::Process(FRenderingCompositePassContext& Context
 	const FSceneViewFamily& ViewFamily = *(View.Family);
 	
 	FIntRect SrcRect = View.ViewRect;
-	FIntRect DestRect = View.UnscaledViewRect;
+
+	//we should be the last node in the graph, so use the 'unscaled' view rect.  aka the one not affected by screenpercentage as we should be 
+	//targetting the final final up/downsampled backbuffer.
+	FIntRect DestRect = View.UnscaledViewRect; //View.ViewRect; // View.UnscaledViewRect;
 	FIntPoint SrcSize = InputDesc->Extent;
 
 	const FSceneRenderTargetItem& DestRenderTarget = PassOutputs[0].RequestSurface(Context);
@@ -211,6 +214,7 @@ void FRCPassPostProcessMorpheus::Process(FRenderingCompositePassContext& Context
 	Context.RHICmdList.SetRasterizerState(TStaticRasterizerState<>::GetRHI());
 	Context.RHICmdList.SetDepthStencilState(TStaticDepthStencilState<false, CF_Always>::GetRHI());
 
+#if MORPHEUS_ENGINE_DISTORTION
 	TShaderMapRef<FPostProcessMorpheusVS> VertexShader(Context.GetShaderMap());
 	TShaderMapRef<FPostProcessMorpheusPS> PixelShader(Context.GetShaderMap());
 
@@ -236,15 +240,38 @@ void FRCPassPostProcessMorpheus::Process(FRenderingCompositePassContext& Context
 		DestRect.Size(),
 		SrcSize
 		);
+#elif PLATFORM_PS4
+
+		FRenderingCompositeOutputRef* OutputRef = Context.Pass->GetInput(ePId_Input0);
+		check(OutputRef);
+		
+		FRenderingCompositeOutput* Input = OutputRef->GetOutput();
+		TRefCountPtr<IPooledRenderTarget> InputPooledElement;
+
+		if (Input)
+		{
+			InputPooledElement = Input->RequestInput();
+		}
+
+		check(InputPooledElement);
+		
+		check(!InputPooledElement->IsFree());
+
+		const FTextureRHIRef& SrcTexture = InputPooledElement->GetRenderTargetItem().ShaderResourceTexture;		
+		GEngine->HMDDevice->DrawSingleEyeDistortionSDK_RenderThread(Context.RHICmdList, View, SrcTexture->GetTexture2D());		
+#else
+	checkf(false, TEXT("Unsupported path.  Morpheus should be disabled."));
+#endif
 
 	Context.RHICmdList.CopyToResolveTarget(DestRenderTarget.TargetableTexture, DestRenderTarget.ShaderResourceTexture, false, FResolveParams());
 }
 
 FPooledRenderTargetDesc FRCPassPostProcessMorpheus::ComputeOutputDesc(EPassOutputId InPassOutputId) const
 {
-	FPooledRenderTargetDesc Ret = PassOutputs[0].RenderTargetDesc;
+	FPooledRenderTargetDesc Ret = PassInputs[0].GetOutput()->RenderTargetDesc;
 
 	Ret.NumSamples = 1;	// no MSAA
+	Ret.Reset();
 	Ret.DebugName = TEXT("Morpheus");
 
 	return Ret;
diff --git a/Engine/Source/Runtime/Renderer/Private/PostProcess/PostProcessMorpheus.h b/Engine/Source/Runtime/Renderer/Private/PostProcess/PostProcessMorpheus.h
--- a/Engine/Source/Runtime/Renderer/Private/PostProcess/PostProcessMorpheus.h
+++ b/Engine/Source/Runtime/Renderer/Private/PostProcess/PostProcessMorpheus.h
@@ -8,7 +8,7 @@
 
 #include "RenderingCompositionGraph.h"
 
-#if MORPHEUS_ENGINE_DISTORTION
+#if HAS_MORPHEUS
 
 // derives from TRenderingCompositePassBase<InputCount, OutputCount>
 // ePId_Input0: SceneColor
diff --git a/Engine/Source/Runtime/Renderer/Private/PostProcess/PostProcessing.cpp b/Engine/Source/Runtime/Renderer/Private/PostProcess/PostProcessing.cpp
--- a/Engine/Source/Runtime/Renderer/Private/PostProcess/PostProcessing.cpp
+++ b/Engine/Source/Runtime/Renderer/Private/PostProcess/PostProcessing.cpp
@@ -1521,10 +1521,9 @@ void FPostProcessing::Process(FRHICommandListImmediate& RHICmdList, FViewInfo& V
 
 			else if(DeviceType == EHMDDeviceType::DT_Morpheus)
 			{
-#if MORPHEUS_ENGINE_DISTORTION
-				Node = Context.Graph.RegisterPass(new FRCPassPostProcessMorpheus());
-#endif
-				bHMDWantsUpscale = true;
+				FRCPassPostProcessMorpheus* MorpheusPass = new FRCPassPostProcessMorpheus();
+				MorpheusPass->SetInput(ePId_Input0, FRenderingCompositeOutputRef(Context.FinalOutput));
+				Node = Context.Graph.RegisterPass(MorpheusPass);
 			}
 
 			
