CIS fix.

diff --git a/Engine/Source/Runtime/Renderer/Private/PostProcess/PostProcessing.cpp b/Engine/Source/Runtime/Renderer/Private/PostProcess/PostProcessing.cpp
--- a/Engine/Source/Runtime/Renderer/Private/PostProcess/PostProcessing.cpp
+++ b/Engine/Source/Runtime/Renderer/Private/PostProcess/PostProcessing.cpp
@@ -1518,14 +1518,14 @@ void FPostProcessing::Process(FRHICommandListImmediate& RHICmdList, FViewInfo& V
 			{
 				Node = Context.Graph.RegisterPass(new FRCPassPostProcessHMD());
 			}
-
+#if HAS_MORPHEUS
 			else if(DeviceType == EHMDDeviceType::DT_Morpheus)
 			{
 				FRCPassPostProcessMorpheus* MorpheusPass = new FRCPassPostProcessMorpheus();
 				MorpheusPass->SetInput(ePId_Input0, FRenderingCompositeOutputRef(Context.FinalOutput));
 				Node = Context.Graph.RegisterPass(MorpheusPass);
 			}
-
+#endif
 			
 			if(Node)
 			{
