Only build HZB as necessary

diff --git a/Engine/Source/Runtime/Renderer/Private/DeferredShadingRenderer.cpp b/Engine/Source/Runtime/Renderer/Private/DeferredShadingRenderer.cpp
--- a/Engine/Source/Runtime/Renderer/Private/DeferredShadingRenderer.cpp
+++ b/Engine/Source/Runtime/Renderer/Private/DeferredShadingRenderer.cpp
@@ -9,6 +9,7 @@
 #include "ScenePrivate.h"
 #include "ScreenRendering.h"
 #include "SceneFilterRendering.h"
+#include "ScreenSpaceReflections.h"
 #include "VisualizeTexture.h"
 #include "CompositionLighting.h"
 #include "FXSystem.h"
@@ -643,11 +644,22 @@ void FDeferredShadingSceneRenderer::RenderOcclusion(FRHICommandListImmediate& RH
 			UpdateDownsampledDepthSurface(RHICmdList);
 		}
 
+		
 		if (bRenderHZB)
-		{
+		{			
+			static const auto ICVarAO		= IConsoleManager::Get().FindTConsoleVariableDataInt(TEXT("r.AmbientOcclusionLevels"));
+			static const auto ICVarHZBOcc	= IConsoleManager::Get().FindConsoleVariable(TEXT("r.HZBOcclusion"));
+			bool bSSAO						= ICVarAO->GetValueOnRenderThread() != 0;
+			bool bHzbOcclusion				= ICVarHZBOcc->GetInt() != 0;
+
 			for (int32 ViewIndex = 0; ViewIndex < Views.Num(); ViewIndex++)
 			{
-				BuildHZB(RHICmdList, Views[ViewIndex]);
+				const uint32 bSSR = DoScreenSpaceReflections(Views[ViewIndex]);
+				
+				if (bSSAO || bHzbOcclusion || bSSR)
+				{
+					BuildHZB(RHICmdList, Views[ViewIndex]);
+				}
 			}
 		}
 
