Changed vertex declaration for FRCPassPostProcessBokehDOF to empty

diff --git a/Engine/Source/Runtime/Renderer/Private/PostProcess/PostProcessBokehDOF.cpp b/Engine/Source/Runtime/Renderer/Private/PostProcess/PostProcessBokehDOF.cpp
--- a/Engine/Source/Runtime/Renderer/Private/PostProcess/PostProcessBokehDOF.cpp
+++ b/Engine/Source/Runtime/Renderer/Private/PostProcess/PostProcessBokehDOF.cpp
@@ -646,7 +646,7 @@ void FRCPassPostProcessBokehDOF::SetShaderTempl(const FRenderingCompositePassCon
 	static FGlobalBoundShaderState BoundShaderState;
 	
 
-	SetGlobalBoundShaderState(Context.RHICmdList, Context.GetFeatureLevel(), BoundShaderState, GFilterVertexDeclaration.VertexDeclarationRHI, *VertexShader, *PixelShader);
+	SetGlobalBoundShaderState(Context.RHICmdList, Context.GetFeatureLevel(), BoundShaderState, GEmptyVertexDeclaration.VertexDeclarationRHI, *VertexShader, *PixelShader);
 
 	VertexShader->SetParameters(Context, TileCount, TileSize, PixelKernelSize, LeftTop);
 	PixelShader->SetParameters(Context, PixelKernelSize);
