Fix shadow variables

diff --git a/Engine/Plugins/Experimental/CodeEditor/Source/CodeEditor/Private/CodeProjectEditor.cpp b/Engine/Plugins/Experimental/CodeEditor/Source/CodeEditor/Private/CodeProjectEditor.cpp
--- a/Engine/Plugins/Experimental/CodeEditor/Source/CodeEditor/Private/CodeProjectEditor.cpp
+++ b/Engine/Plugins/Experimental/CodeEditor/Source/CodeEditor/Private/CodeProjectEditor.cpp
@@ -297,8 +297,8 @@ bool FCodeProjectEditor::Save()
 {
 	if(DocumentManager.IsValid() && DocumentManager->GetActiveTab().IsValid())
 	{
-		TSharedRef<SCodeEditor> CodeEditor = StaticCastSharedRef<SCodeEditor>(DocumentManager->GetActiveTab()->GetContent());
-		return CodeEditor->Save();
+		TSharedRef<SCodeEditor> CodeEditorRef = StaticCastSharedRef<SCodeEditor>(DocumentManager->GetActiveTab()->GetContent());
+		return CodeEditorRef->Save();
 	}
 
 	return false;
@@ -309,8 +309,8 @@ bool FCodeProjectEditor::CanSave() const
 	if(DocumentManager.IsValid() && DocumentManager->GetActiveTab().IsValid())
 	{
 		TSharedRef<SWidget> Content = DocumentManager->GetActiveTab()->GetContent();
-		TSharedRef<SCodeEditor> CodeEditor = StaticCastSharedRef<SCodeEditor>(Content);
-		return CodeEditor->CanSave();
+		TSharedRef<SCodeEditor> CodeEditorRef = StaticCastSharedRef<SCodeEditor>(Content);
+		return CodeEditorRef->CanSave();
 	}
 
 	return false;
@@ -332,8 +332,8 @@ bool FCodeProjectEditor::SaveAll()
 		{
 			if(Tab.IsValid())
 			{
-				TSharedRef<SCodeEditor> CodeEditor = StaticCastSharedRef<SCodeEditor>(Tab->GetContent());
-				if(!CodeEditor->Save())
+				TSharedRef<SCodeEditor> CodeEditorRef = StaticCastSharedRef<SCodeEditor>(Tab->GetContent());
+				if(!CodeEditorRef->Save())
 				{
 					bResult = false;
 				}
diff --git a/Engine/Plugins/Experimental/CodeEditor/Source/CodeEditor/Private/CodeProjectEditorToolbar.cpp b/Engine/Plugins/Experimental/CodeEditor/Source/CodeEditor/Private/CodeProjectEditorToolbar.cpp
--- a/Engine/Plugins/Experimental/CodeEditor/Source/CodeEditor/Private/CodeProjectEditorToolbar.cpp
+++ b/Engine/Plugins/Experimental/CodeEditor/Source/CodeEditor/Private/CodeProjectEditorToolbar.cpp
@@ -37,9 +37,9 @@ void FCodeProjectEditorToolbar::FillEditorToolbar(FToolBarBuilder& ToolbarBuilde
 		{
 			struct Local
 			{
-				static void ExecuteCompile(TSharedPtr<FCodeProjectEditor> CodeProjectEditorPtr)
+				static void ExecuteCompile(TSharedPtr<FCodeProjectEditor> InCodeProjectEditorPtr)
 				{
-					if(CodeProjectEditorPtr->SaveAll())
+					if(InCodeProjectEditorPtr->SaveAll())
 					{
 						FLevelEditorActionCallbacks::RecompileGameCode_Clicked();
 					}
diff --git a/Engine/Source/Editor/KismetCompiler/Private/KismetCompilerMisc.cpp b/Engine/Source/Editor/KismetCompiler/Private/KismetCompilerMisc.cpp
--- a/Engine/Source/Editor/KismetCompiler/Private/KismetCompilerMisc.cpp
+++ b/Engine/Source/Editor/KismetCompiler/Private/KismetCompilerMisc.cpp
@@ -932,7 +932,7 @@ void FKismetCompilerUtilities::ValidateProperEndExecutionPath(FKismetFunctionCon
 	struct FRecrursiveHelper
 	{
 		// returns if the path is properly ended
-		static void CheckPath(UK2Node* StartingNode, TSet<UK2Node*>& VisitedNodes, FKismetFunctionContext& Context)
+		static void CheckPath(UK2Node* StartingNode, TSet<UK2Node*>& VisitedNodes, FKismetFunctionContext& InContext)
 		{
 			UK2Node* CurrentNode = StartingNode;
 			while (CurrentNode)
@@ -949,14 +949,14 @@ void FKismetCompilerUtilities::ValidateProperEndExecutionPath(FKismetFunctionCon
 					{
 						if (CurrentPin
 							&& (CurrentPin->Direction == EEdGraphPinDirection::EGPD_Output)
-							&& (CurrentPin->PinType.PinCategory == Context.Schema->PC_Exec)
+							&& (CurrentPin->PinType.PinCategory == InContext.Schema->PC_Exec)
 							&& (CurrentPin->LinkedTo.Num() > 0))
 						{
 							auto LinkedPin = CurrentPin->LinkedTo[0];
 							auto NextNode = ensure(LinkedPin) ? Cast<UK2Node>(LinkedPin->GetOwningNodeUnchecked()) : nullptr;
 							if (CurrentNode && !bIsExecutionSequence && ensure(NextNode)) // for sequence node, we want only the last output
 							{
-								FRecrursiveHelper::CheckPath(NextNode, VisitedNodes, Context);
+								FRecrursiveHelper::CheckPath(NextNode, VisitedNodes, InContext);
 							}
 							else
 							{
@@ -967,7 +967,7 @@ void FKismetCompilerUtilities::ValidateProperEndExecutionPath(FKismetFunctionCon
 
 					if (!CurrentNode)
 					{
-						Context.MessageLog.Warning(*LOCTEXT("ExecutionEnd_Warning", "The execution path doesn't end with a return node. @@").ToString(), SourceNode);
+						InContext.MessageLog.Warning(*LOCTEXT("ExecutionEnd_Warning", "The execution path doesn't end with a return node. @@").ToString(), SourceNode);
 					}
 				}
 			}
diff --git a/Engine/Source/Editor/LandscapeEditor/Private/LandscapeEdModeSplineTools.cpp b/Engine/Source/Editor/LandscapeEditor/Private/LandscapeEdModeSplineTools.cpp
--- a/Engine/Source/Editor/LandscapeEditor/Private/LandscapeEdModeSplineTools.cpp
+++ b/Engine/Source/Editor/LandscapeEditor/Private/LandscapeEdModeSplineTools.cpp
@@ -610,9 +610,9 @@ public:
 			{
 				if (!ToLandscape)
 				{
-					ULandscapeInfo* LandscapeInfo = FromProxy->GetLandscapeInfo(false);
-					check(LandscapeInfo);
-					ToLandscape = LandscapeInfo->GetCurrentLevelLandscapeProxy(true);
+					ULandscapeInfo* ProxyLandscapeInfo = FromProxy->GetLandscapeInfo(false);
+					check(ProxyLandscapeInfo);
+					ToLandscape = ProxyLandscapeInfo->GetCurrentLevelLandscapeProxy(true);
 					if (!ToLandscape)
 					{
 						// No Landscape Proxy, don't support for creating only for Spline now
@@ -699,9 +699,9 @@ public:
 			{
 				if (!ToLandscape)
 				{
-					ULandscapeInfo* LandscapeInfo = FromProxy->GetLandscapeInfo(false);
-					check(LandscapeInfo);
-					ToLandscape = LandscapeInfo->GetCurrentLevelLandscapeProxy(true);
+					ULandscapeInfo* ProxyLandscapeInfo = FromProxy->GetLandscapeInfo(false);
+					check(ProxyLandscapeInfo);
+					ToLandscape = ProxyLandscapeInfo->GetCurrentLevelLandscapeProxy(true);
 					if (!ToLandscape)
 					{
 						// No Landscape Proxy, don't support for creating only for Spline now
diff --git a/Engine/Source/Editor/LocalizationDashboard/Private/LocalizationCommandletExecution.cpp b/Engine/Source/Editor/LocalizationDashboard/Private/LocalizationCommandletExecution.cpp
--- a/Engine/Source/Editor/LocalizationDashboard/Private/LocalizationCommandletExecution.cpp
+++ b/Engine/Source/Editor/LocalizationDashboard/Private/LocalizationCommandletExecution.cpp
@@ -641,10 +641,10 @@ namespace
 		{
 			void* ParentWindowWindowHandle = NULL;
 
-			const TSharedPtr<SWindow>& ParentWindow = FSlateApplication::Get().FindWidgetWindow(AsShared());
-			if (ParentWindow.IsValid() && ParentWindow->GetNativeWindow().IsValid())
+			const TSharedPtr<SWindow>& ParentWindowPtr = FSlateApplication::Get().FindWidgetWindow(AsShared());
+			if (ParentWindowPtr.IsValid() && ParentWindowPtr->GetNativeWindow().IsValid())
 			{
-				ParentWindowWindowHandle = ParentWindow->GetNativeWindow()->GetOSWindowHandle();
+				ParentWindowWindowHandle = ParentWindowPtr->GetNativeWindow()->GetOSWindowHandle();
 			}
 
 			if (DesktopPlatform->SaveFileDialog(
diff --git a/Engine/Source/Editor/LocalizationDashboard/Private/LocalizationConfigurationScript.cpp b/Engine/Source/Editor/LocalizationDashboard/Private/LocalizationConfigurationScript.cpp
--- a/Engine/Source/Editor/LocalizationDashboard/Private/LocalizationConfigurationScript.cpp
+++ b/Engine/Source/Editor/LocalizationDashboard/Private/LocalizationConfigurationScript.cpp
@@ -105,7 +105,7 @@ namespace LocalizationConfigurationScript
 				}
 				AllLocalizationTargets.Append(LocalizationTargetSet->TargetObjects);
 
-				ULocalizationTarget* const * OtherTarget = AllLocalizationTargets.FindByPredicate([&TargetDependencyGuid](ULocalizationTarget* const OtherTarget)->bool{return OtherTarget->Settings.Guid == TargetDependencyGuid;});
+				ULocalizationTarget* const * OtherTarget = AllLocalizationTargets.FindByPredicate([&TargetDependencyGuid](ULocalizationTarget* const InOtherTarget)->bool{return InOtherTarget->Settings.Guid == TargetDependencyGuid;});
 				if (OtherTarget)
 				{
 					ConfigSection.Add( TEXT("ManifestDependencies"), MakePathRelativeForCommandletProcess(GetManifestPath(*OtherTarget), !Target->IsMemberOfEngineTargetSet()) );
diff --git a/Engine/Source/Editor/LocalizationDashboard/Private/LocalizationTargetDetailCustomization.cpp b/Engine/Source/Editor/LocalizationDashboard/Private/LocalizationTargetDetailCustomization.cpp
--- a/Engine/Source/Editor/LocalizationDashboard/Private/LocalizationTargetDetailCustomization.cpp
+++ b/Engine/Source/Editor/LocalizationDashboard/Private/LocalizationTargetDetailCustomization.cpp
@@ -856,27 +856,27 @@ void FLocalizationTargetDetailCustomization::UpdateTargetFromReports()
 
 		if (TargetSettingsPropertyHandle.IsValid() && TargetSettingsPropertyHandle->IsValidHandle())
 		{
-			const TSharedPtr<IPropertyHandle> NativeCultureStatisticsPropertyHandle = TargetSettingsPropertyHandle->GetChildHandle(GET_MEMBER_NAME_CHECKED(FLocalizationTargetSettings, NativeCultureStatistics));
-			if (NativeCultureStatisticsPropertyHandle.IsValid() && NativeCultureStatisticsPropertyHandle->IsValidHandle())
+			const TSharedPtr<IPropertyHandle> NativeCultureStatisticsPropHandle = TargetSettingsPropertyHandle->GetChildHandle(GET_MEMBER_NAME_CHECKED(FLocalizationTargetSettings, NativeCultureStatistics));
+			if (NativeCultureStatisticsPropHandle.IsValid() && NativeCultureStatisticsPropHandle->IsValidHandle())
 			{
-				const TSharedPtr<IPropertyHandle> WordCountPropertyHandle = NativeCultureStatisticsPropertyHandle->GetChildHandle(GET_MEMBER_NAME_CHECKED(FCultureStatistics, WordCount));
+				const TSharedPtr<IPropertyHandle> WordCountPropertyHandle = NativeCultureStatisticsPropHandle->GetChildHandle(GET_MEMBER_NAME_CHECKED(FCultureStatistics, WordCount));
 				if (WordCountPropertyHandle.IsValid() && WordCountPropertyHandle->IsValidHandle())
 				{
 					WordCountPropertyHandles.Add(WordCountPropertyHandle);
 				}
 			}
 
-			const TSharedPtr<IPropertyHandle> SupportedCulturesStatisticsPropertyHandle = TargetSettingsPropertyHandle->GetChildHandle(GET_MEMBER_NAME_CHECKED(FLocalizationTargetSettings, SupportedCulturesStatistics));
-			if (SupportedCulturesStatisticsPropertyHandle.IsValid() && SupportedCulturesStatisticsPropertyHandle->IsValidHandle())
+			const TSharedPtr<IPropertyHandle> SupportedCulturesStatisticsPropHandle = TargetSettingsPropertyHandle->GetChildHandle(GET_MEMBER_NAME_CHECKED(FLocalizationTargetSettings, SupportedCulturesStatistics));
+			if (SupportedCulturesStatisticsPropHandle.IsValid() && SupportedCulturesStatisticsPropHandle->IsValidHandle())
 			{
 				uint32 SupportedCultureCount = 0;
-				SupportedCulturesStatisticsPropertyHandle->GetNumChildren(SupportedCultureCount);
+				SupportedCulturesStatisticsPropHandle->GetNumChildren(SupportedCultureCount);
 				for (uint32 i = 0; i < SupportedCultureCount; ++i)
 				{
-					const TSharedPtr<IPropertyHandle> ElementPropertyHandle = SupportedCulturesStatisticsPropertyHandle->GetChildHandle(i);
+					const TSharedPtr<IPropertyHandle> ElementPropertyHandle = SupportedCulturesStatisticsPropHandle->GetChildHandle(i);
 					if (ElementPropertyHandle.IsValid() && ElementPropertyHandle->IsValidHandle())
 					{
-						const TSharedPtr<IPropertyHandle> WordCountPropertyHandle = SupportedCulturesStatisticsPropertyHandle->GetChildHandle(GET_MEMBER_NAME_CHECKED(FCultureStatistics, WordCount));
+						const TSharedPtr<IPropertyHandle> WordCountPropertyHandle = SupportedCulturesStatisticsPropHandle->GetChildHandle(GET_MEMBER_NAME_CHECKED(FCultureStatistics, WordCount));
 						if (WordCountPropertyHandle.IsValid() && WordCountPropertyHandle->IsValidHandle())
 						{
 							WordCountPropertyHandles.Add(WordCountPropertyHandle);
diff --git a/Engine/Source/Editor/LocalizationDashboard/Private/LocalizationTargetTypes.cpp b/Engine/Source/Editor/LocalizationDashboard/Private/LocalizationTargetTypes.cpp
--- a/Engine/Source/Editor/LocalizationDashboard/Private/LocalizationTargetTypes.cpp
+++ b/Engine/Source/Editor/LocalizationDashboard/Private/LocalizationTargetTypes.cpp
@@ -367,7 +367,6 @@ bool ULocalizationTarget::DeleteFiles(const FString* const Culture) const
 		TArray<FString> FilesInDataDirectory;
 		IFileManager::Get().FindFilesRecursive(FilesInDataDirectory, *DataDirectory, TEXT("*"), true, false);
 
-		ISourceControlProvider& SourceControlProvider = ISourceControlModule::Get().GetProvider();
 		TArray<FSourceControlStateRef> SourceControlStates;
 		ECommandResult::Type Result = SourceControlProvider.GetState( FilesInDataDirectory, SourceControlStates, EStateCacheUsage::ForceUpdate );
 		if (Result == ECommandResult::Succeeded)
@@ -408,7 +407,6 @@ bool ULocalizationTarget::DeleteFiles(const FString* const Culture) const
 
 		// Remove script files from source control.
 		{
-			ISourceControlProvider& SourceControlProvider = ISourceControlModule::Get().GetProvider();
 			TArray<FSourceControlStateRef> SourceControlStates;
 			ECommandResult::Type Result = SourceControlProvider.GetState( ScriptPaths, SourceControlStates, EStateCacheUsage::ForceUpdate );
 			if (Result == ECommandResult::Succeeded)
diff --git a/Engine/Source/Editor/PropertyEditor/Private/PropertyNode.cpp b/Engine/Source/Editor/PropertyEditor/Private/PropertyNode.cpp
--- a/Engine/Source/Editor/PropertyEditor/Private/PropertyNode.cpp
+++ b/Engine/Source/Editor/PropertyEditor/Private/PropertyNode.cpp
@@ -1541,22 +1541,20 @@ void FPropertyNode::ResetToDefault( FNotifyHook* InNotifyHook )
 
 						// dynamic arrays are the only property type that do not support CopySingleValue correctly due to the fact that they cannot
 						// be used in a static array
-						UArrayProperty* ArrayProp = Cast<UArrayProperty>(TheProperty);
-
 
-					
-						FPropertyNode* ParentNode = GetParentNode();
-						if(ParentNode != NULL && ParentNode->GetProperty() && ParentNode->GetProperty()->IsA(UArrayProperty::StaticClass()))
+						FPropertyNode* ParentPropertyNode = GetParentNode();
+						if(ParentPropertyNode != NULL && ParentPropertyNode->GetProperty() && ParentPropertyNode->GetProperty()->IsA(UArrayProperty::StaticClass()))
 						{
-							UArrayProperty* ArrayProp = Cast<UArrayProperty>(ParentNode->GetProperty());
+							UArrayProperty* ArrayProp = Cast<UArrayProperty>(ParentPropertyNode->GetProperty());
 							if(ArrayProp->Inner == TheProperty)
 							{
-								uint8* Addr = ParentNode->GetValueBaseAddress((uint8*)Object);
+								uint8* Addr = ParentPropertyNode->GetValueBaseAddress((uint8*)Object);
 
 								ArrayProp->ExportText_Direct(PreviousArrayValue, Addr, Addr, NULL, 0);
 							}
 						}
 
+						UArrayProperty* ArrayProp = Cast<UArrayProperty>(TheProperty);
 						if( ArrayProp != NULL )
 						{
 							TheProperty->CopyCompleteValue(ValueTracker.GetPropertyValueAddress(), ValueTracker.GetPropertyDefaultAddress());
diff --git a/Engine/Source/Editor/UnrealEd/Private/Kismet2/KismetDebugUtilities.cpp b/Engine/Source/Editor/UnrealEd/Private/Kismet2/KismetDebugUtilities.cpp
--- a/Engine/Source/Editor/UnrealEd/Private/Kismet2/KismetDebugUtilities.cpp
+++ b/Engine/Source/Editor/UnrealEd/Private/Kismet2/KismetDebugUtilities.cpp
@@ -322,8 +322,8 @@ void FKismetDebugUtilities::OnScriptException(const UObject* ActiveObject, const
 
 						// Display a UObject link to the UFunction that is crashing. Will open the Blueprint if able and focus on the function's graph
 						Message->AddToken(FTextToken::Create(LOCTEXT( "InfiniteLoopWarning_Function", ", asserted during ")));
-						const int32 BreakpointOffset = StackFrame.Code - StackFrame.Node->Script.GetData() - 1; //@TODO: Might want to make this a parameter of Info
-						UEdGraphNode* SourceNode = FindSourceNodeForCodeLocation(ActiveObject, StackFrame.Node, BreakpointOffset, /*bAllowImpreciseHit=*/ true);
+						const int32 BreakpointOpCodeOffset = StackFrame.Code - StackFrame.Node->Script.GetData() - 1; //@TODO: Might want to make this a parameter of Info
+						UEdGraphNode* SourceNode = FindSourceNodeForCodeLocation(ActiveObject, StackFrame.Node, BreakpointOpCodeOffset, /*bAllowImpreciseHit=*/ true);
 
 						// If a source node is found, that's the token we want to link, otherwise settle with the UFunction
 						if(SourceNode)
diff --git a/Engine/Source/Runtime/Engine/Private/Camera/CameraShake.cpp b/Engine/Source/Runtime/Engine/Private/Camera/CameraShake.cpp
--- a/Engine/Source/Runtime/Engine/Private/Camera/CameraShake.cpp
+++ b/Engine/Source/Runtime/Engine/Private/Camera/CameraShake.cpp
@@ -52,7 +52,7 @@ void UCameraShake::StopShake()
 	ReceiveStopShake();
 }
 
-void UCameraShake::PlayShake(APlayerCameraManager* Camera, float Scale, ECameraAnimPlaySpace::Type PlaySpace, FRotator UserPlaySpaceRot)
+void UCameraShake::PlayShake(APlayerCameraManager* Camera, float Scale, ECameraAnimPlaySpace::Type InPlaySpace, FRotator UserPlaySpaceRot)
 {
 	ShakeScale = Scale;
 	CameraOwner = Camera;
@@ -130,12 +130,12 @@ void UCameraShake::PlayShake(APlayerCameraManager* Camera, float Scale, ECameraA
 			float const FinalAnimScale = Scale * AnimScale;
 			if (FinalAnimScale > 0.f)
 			{
-				AnimInst = CameraOwner->PlayCameraAnim(Anim, AnimPlayRate, FinalAnimScale, AnimBlendInTime, AnimBlendOutTime, bLoop, bRandomStart, Duration, PlaySpace, UserPlaySpaceRot);
+				AnimInst = CameraOwner->PlayCameraAnim(Anim, AnimPlayRate, FinalAnimScale, AnimBlendInTime, AnimBlendOutTime, bLoop, bRandomStart, Duration, InPlaySpace, UserPlaySpaceRot);
 			}
 		}
 	}
 
-	if (PlaySpace == ECameraAnimPlaySpace::UserDefined)
+	if (InPlaySpace == ECameraAnimPlaySpace::UserDefined)
 	{
 		UserPlaySpaceMatrix = FRotationMatrix(UserPlaySpaceRot);
 	}
diff --git a/Engine/Source/Runtime/Engine/Private/Components/PrimitiveComponent.cpp b/Engine/Source/Runtime/Engine/Private/Components/PrimitiveComponent.cpp
--- a/Engine/Source/Runtime/Engine/Private/Components/PrimitiveComponent.cpp
+++ b/Engine/Source/Runtime/Engine/Private/Components/PrimitiveComponent.cpp
@@ -1768,9 +1768,9 @@ bool UPrimitiveComponent::OverlapComponent(const FVector& Pos, const FQuat& Rot,
 
 bool UPrimitiveComponent::ComputePenetration(FMTDResult& OutMTD, const FCollisionShape & CollisionShape, const FVector& Pos, const FQuat& Rot)
 {
-	if(FBodyInstance* BodyInstance = GetBodyInstance())
+	if(FBodyInstance* ComponentBodyInstance = GetBodyInstance())
 	{
-		return BodyInstance->OverlapTest(Pos, Rot, CollisionShape, &OutMTD);
+		return ComponentBodyInstance->OverlapTest(Pos, Rot, CollisionShape, &OutMTD);
 	}
 
 	return false;
diff --git a/Engine/Source/Runtime/Engine/Private/PrimitiveComponentPhysics.cpp b/Engine/Source/Runtime/Engine/Private/PrimitiveComponentPhysics.cpp
--- a/Engine/Source/Runtime/Engine/Private/PrimitiveComponentPhysics.cpp
+++ b/Engine/Source/Runtime/Engine/Private/PrimitiveComponentPhysics.cpp
@@ -332,9 +332,9 @@ FVector UPrimitiveComponent::GetPhysicsAngularVelocity(FName BoneName)
 
 FVector UPrimitiveComponent::GetCenterOfMass(FName BoneName)
 {
-	if (FBodyInstance* BodyInstance = GetBodyInstance(BoneName))
+	if (FBodyInstance* ComponentBodyInstance = GetBodyInstance(BoneName))
 	{
-		return BodyInstance->GetCOMPosition();
+		return ComponentBodyInstance->GetCOMPosition();
 	}
 
 	return FVector::ZeroVector;
diff --git a/Engine/Source/Runtime/Foliage/Private/FoliageInstanceBase.cpp b/Engine/Source/Runtime/Foliage/Private/FoliageInstanceBase.cpp
--- a/Engine/Source/Runtime/Foliage/Private/FoliageInstanceBase.cpp
+++ b/Engine/Source/Runtime/Foliage/Private/FoliageInstanceBase.cpp
@@ -166,9 +166,9 @@ void FFoliageInstanceBaseCache::CompactInstanceBaseCache(AInstancedFoliageActor*
 	FFoliageInstanceBaseCache& Cache = IFA->InstanceBaseCache;
 	
 	TSet<FFoliageInstanceBaseId> BasesInUse;
-	for (auto& Pair : IFA->FoliageMeshes)
+	for (auto& FoliageMeshPair : IFA->FoliageMeshes)
 	{
-		for (const auto& Pair : Pair.Value->ComponentHash)
+		for (const auto& Pair : FoliageMeshPair.Value->ComponentHash)
 		{
 			if (Pair.Key != FFoliageInstanceBaseCache::InvalidBaseId)
 			{
diff --git a/Engine/Source/Runtime/Foliage/Private/InstancedFoliage.cpp b/Engine/Source/Runtime/Foliage/Private/InstancedFoliage.cpp
--- a/Engine/Source/Runtime/Foliage/Private/InstancedFoliage.cpp
+++ b/Engine/Source/Runtime/Foliage/Private/InstancedFoliage.cpp
@@ -455,10 +455,10 @@ float UFoliageType::GetInitAge(FRandomStream& RandomStream) const
 	return RandomStream.FRandRange(0, InitialMaxAge);
 }
 
-float UFoliageType::GetNextAge(const float CurrentAge, const int32 NumSteps) const
+float UFoliageType::GetNextAge(const float CurrentAge, const int32 InNumSteps) const
 {
 	float NewAge = CurrentAge;
-	for (int32 Count = 0; Count < NumSteps; ++Count)
+	for (int32 Count = 0; Count < InNumSteps; ++Count)
 	{
 		const float GrowAge = NewAge + 1;
 		if (GrowAge <= MaxAge)
diff --git a/Engine/Source/Runtime/Foliage/Private/ProceduralFoliageTile.cpp b/Engine/Source/Runtime/Foliage/Private/ProceduralFoliageTile.cpp
--- a/Engine/Source/Runtime/Foliage/Private/ProceduralFoliageTile.cpp
+++ b/Engine/Source/Runtime/Foliage/Private/ProceduralFoliageTile.cpp
@@ -436,10 +436,10 @@ void UProceduralFoliageTile::RunSimulation(const int32 MaxNumSteps, bool bShadeG
 	InstancesToArray();
 }
 
-void UProceduralFoliageTile::Simulate(const UProceduralFoliageSpawner* InFoliageSpawner, const int32 RandomSeed, const int32 MaxNumSteps, const int32 InLastCancel)
+void UProceduralFoliageTile::Simulate(const UProceduralFoliageSpawner* InFoliageSpawner, const int32 InRandomSeed, const int32 MaxNumSteps, const int32 InLastCancel)
 {
 	LastCancel = InLastCancel;
-	InitSimulation(InFoliageSpawner, RandomSeed);
+	InitSimulation(InFoliageSpawner, InRandomSeed);
 
 	RunSimulation(MaxNumSteps, false);
 	RunSimulation(MaxNumSteps, true);
diff --git a/Engine/Source/Runtime/Online/OnlineSubsystemSteam/Private/OnlineLeaderboardInterfaceSteam.cpp b/Engine/Source/Runtime/Online/OnlineSubsystemSteam/Private/OnlineLeaderboardInterfaceSteam.cpp
--- a/Engine/Source/Runtime/Online/OnlineSubsystemSteam/Private/OnlineLeaderboardInterfaceSteam.cpp
+++ b/Engine/Source/Runtime/Online/OnlineSubsystemSteam/Private/OnlineLeaderboardInterfaceSteam.cpp
@@ -648,17 +648,17 @@ private:
 	 * @param SortMethod method the leaderboard scores will be sorted, ignored if leaderboard exists
 	 * @param DisplayFormat type of data the leaderboard represents, ignored if leaderboard exists
 	 */
-	void CreateOrFindLeaderboard(const FName& LeaderboardName, ELeaderboardSort::Type InSortMethod, ELeaderboardFormat::Type InDisplayFormat)
+	void CreateOrFindLeaderboard(const FName& InLeaderboardName, ELeaderboardSort::Type InSortMethod, ELeaderboardFormat::Type InDisplayFormat)
 	{
 		if (bFindOnly)
 		{
-			CallbackHandle = SteamUserStats()->FindLeaderboard(TCHAR_TO_UTF8(*LeaderboardName.ToString()));
+			CallbackHandle = SteamUserStats()->FindLeaderboard(TCHAR_TO_UTF8(*InLeaderboardName.ToString()));
 		}
 		else
 		{
 			ELeaderboardSortMethod SortMethodSteam = ToSteamLeaderboardSortMethod(InSortMethod);
 			ELeaderboardDisplayType DisplayTypeSteam = ToSteamLeaderboardDisplayType(InDisplayFormat);
-			CallbackHandle = SteamUserStats()->FindOrCreateLeaderboard(TCHAR_TO_UTF8(*LeaderboardName.ToString()), SortMethodSteam, DisplayTypeSteam);
+			CallbackHandle = SteamUserStats()->FindOrCreateLeaderboard(TCHAR_TO_UTF8(*InLeaderboardName.ToString()), SortMethodSteam, DisplayTypeSteam);
 		}
 	}
 
diff --git a/Engine/Source/Runtime/Online/OnlineSubsystemSteam/Private/OnlineSessionAsyncLobbySteam.cpp b/Engine/Source/Runtime/Online/OnlineSubsystemSteam/Private/OnlineSessionAsyncLobbySteam.cpp
--- a/Engine/Source/Runtime/Online/OnlineSubsystemSteam/Private/OnlineSessionAsyncLobbySteam.cpp
+++ b/Engine/Source/Runtime/Online/OnlineSubsystemSteam/Private/OnlineSessionAsyncLobbySteam.cpp
@@ -1058,10 +1058,10 @@ void FOnlineAsyncTaskSteamFindLobbies::TriggerDelegates()
  *
  * @param LobbyId lobby to create the search result for
  */
-void FOnlineAsyncTaskSteamFindLobby::ParseSearchResult(FUniqueNetIdSteam& LobbyId)
+void FOnlineAsyncTaskSteamFindLobby::ParseSearchResult(FUniqueNetIdSteam& InLobbyId)
 {
 	FOnlineSessionSearchResult* NewSearchResult = new (SearchSettings->SearchResults) FOnlineSessionSearchResult();
-	if (!FillSessionFromLobbyData(LobbyId, NewSearchResult->Session))
+	if (!FillSessionFromLobbyData(InLobbyId, NewSearchResult->Session))
 	{
 		// Remove the failed element
 		SearchSettings->SearchResults.RemoveAtSwap(SearchSettings->SearchResults.Num() - 1);
@@ -1113,10 +1113,10 @@ void FOnlineAsyncTaskSteamFindLobby::Finalize()
 			// Parse any ready search results
 			for (int32 LobbyIdx=0; LobbyIdx < SessionInt->PendingSearchLobbyIds.Num(); LobbyIdx++)
 			{
-				FUniqueNetIdSteam& LobbyId = SessionInt->PendingSearchLobbyIds[LobbyIdx];
-				if (LobbyId.IsValid() && CSteamID(LobbyId).IsLobby())
+				FUniqueNetIdSteam& PendingLobbyId = SessionInt->PendingSearchLobbyIds[LobbyIdx];
+				if (PendingLobbyId.IsValid() && CSteamID(PendingLobbyId).IsLobby())
 				{
-					ParseSearchResult(LobbyId);
+					ParseSearchResult(PendingLobbyId);
 				}
 			}
 
diff --git a/Engine/Source/Runtime/Online/OnlineSubsystemSteam/Private/OnlineUserCloudInterfaceSteam.cpp b/Engine/Source/Runtime/Online/OnlineSubsystemSteam/Private/OnlineUserCloudInterfaceSteam.cpp
--- a/Engine/Source/Runtime/Online/OnlineSubsystemSteam/Private/OnlineUserCloudInterfaceSteam.cpp
+++ b/Engine/Source/Runtime/Online/OnlineSubsystemSteam/Private/OnlineUserCloudInterfaceSteam.cpp
@@ -140,48 +140,48 @@ void FOnlineAsyncTaskSteamReadUserFile::TriggerDelegates()
 	UserCloudInterface->TriggerOnReadUserFileCompleteDelegates(bWasSuccessful, UserId, FileName);
 }
 
-bool FOnlineAsyncTaskSteamWriteUserFile::WriteUserFile(const FUniqueNetId& UserId, const FString& FileToWrite, const TArray<uint8>& Contents)
+bool FOnlineAsyncTaskSteamWriteUserFile::WriteUserFile(const FUniqueNetId& InUserId, const FString& InFileToWrite, const TArray<uint8>& InContents)
 {
 	bool bSuccess = false;
-	if (FileToWrite.Len() > 0 && Contents.Num() > 0)
+	if (InFileToWrite.Len() > 0 && InContents.Num() > 0)
 	{
 		if (SteamRemoteStorage() && FileName.Len() > 0)
 		{
-			CSteamID SteamId(*(uint64*)UserId.GetBytes());
+			CSteamID SteamId(*(uint64*)InUserId.GetBytes());
 			if (SteamUser()->BLoggedOn() && SteamUser()->GetSteamID() == SteamId)
 			{
 				// Currently don't support greater than 1 chunk
-				if (Contents.Num() < k_unMaxCloudFileChunkSize)
+				if (InContents.Num() < k_unMaxCloudFileChunkSize)
 				{
-					if (SteamRemoteStorage()->FileWrite(TCHAR_TO_UTF8(*FileToWrite), Contents.GetData(), Contents.Num()))
+					if (SteamRemoteStorage()->FileWrite(TCHAR_TO_UTF8(*InFileToWrite), InContents.GetData(), InContents.Num()))
 					{
 						FScopeLock ScopeLock(&Subsystem->UserCloudDataLock);
-						FSteamUserCloudData* UserCloud = Subsystem->GetUserCloudEntry(UserId);
+						FSteamUserCloudData* UserCloud = Subsystem->GetUserCloudEntry(InUserId);
 						if (UserCloud)
 						{
 							// Update the metadata table to reflect this write (might be new entry)
-							FCloudFileHeader* UserCloudFileMetadata = UserCloud->GetFileMetadata(FileToWrite, true);
+							FCloudFileHeader* UserCloudFileMetadata = UserCloud->GetFileMetadata(InFileToWrite, true);
 							check(UserCloudFileMetadata);
 
-							UserCloudFileMetadata->FileSize = SteamRemoteStorage()->GetFileSize(TCHAR_TO_UTF8(*FileToWrite));
+							UserCloudFileMetadata->FileSize = SteamRemoteStorage()->GetFileSize(TCHAR_TO_UTF8(*InFileToWrite));
 							UserCloudFileMetadata->Hash = FString(TEXT("0"));
 
 							// Update the file table to reflect this write
-							FCloudFile* UserCloudFileData = UserCloud->GetFileData(FileToWrite, true);
+							FCloudFile* UserCloudFileData = UserCloud->GetFileData(InFileToWrite, true);
 							check(UserCloudFileData);
 
-							UserCloudFileData->Data = Contents;
+							UserCloudFileData->Data = InContents;
 							bSuccess = true;
 						}
 					}
 					else
 					{
-						UE_LOG_ONLINE(Warning, TEXT("Failed to write file to Steam cloud \"%s\"."), *FileToWrite);
+						UE_LOG_ONLINE(Warning, TEXT("Failed to write file to Steam cloud \"%s\"."), *InFileToWrite);
 					}
 				}
 				else
 				{
-					UE_LOG_ONLINE(Warning, TEXT("File too large %d to write to Steam cloud."), Contents.Num());
+					UE_LOG_ONLINE(Warning, TEXT("File too large %d to write to Steam cloud."), InContents.Num());
 				}
 			}
 			else
@@ -197,16 +197,16 @@ bool FOnlineAsyncTaskSteamWriteUserFile::WriteUserFile(const FUniqueNetId& UserI
 
 	{
 		FScopeLock ScopeLock(&Subsystem->UserCloudDataLock);
-		FSteamUserCloudData* UserCloud = Subsystem->GetUserCloudEntry(UserId);
+		FSteamUserCloudData* UserCloud = Subsystem->GetUserCloudEntry(InUserId);
 		if (UserCloud)
 		{
-			FCloudFile* UserCloudFileData = UserCloud->GetFileData(FileToWrite, true);
+			FCloudFile* UserCloudFileData = UserCloud->GetFileData(InFileToWrite, true);
 			check(UserCloudFileData);
 			UserCloudFileData->AsyncState = bSuccess ? EOnlineAsyncTaskState::Done : EOnlineAsyncTaskState::Failed;
 		}
 	}
 
-	//SteamSubsystem->GetUserCloudInterface()->DumpCloudFileState(UserId, FileToWrite);
+	//SteamSubsystem->GetUserCloudInterface()->DumpCloudFileState(InUserId, InFileToWrite);
 	return bSuccess;
 }
 
diff --git a/Engine/Source/Runtime/Online/OnlineSubsystemUtils/Private/EndMatchCallbackProxy.cpp b/Engine/Source/Runtime/Online/OnlineSubsystemUtils/Private/EndMatchCallbackProxy.cpp
--- a/Engine/Source/Runtime/Online/OnlineSubsystemUtils/Private/EndMatchCallbackProxy.cpp
+++ b/Engine/Source/Runtime/Online/OnlineSubsystemUtils/Private/EndMatchCallbackProxy.cpp
@@ -58,7 +58,7 @@ void UEndMatchCallbackProxy::Activate()
 	OnFailure.Broadcast();
 }
 
-void UEndMatchCallbackProxy::EndMatchDelegate(FString MatchID, bool Succeeded)
+void UEndMatchCallbackProxy::EndMatchDelegate(FString InMatchID, bool Succeeded)
 {
 	if (Succeeded)
 	{ 
diff --git a/Engine/Source/Runtime/Online/OnlineSubsystemUtils/Private/EndTurnCallbackProxy.cpp b/Engine/Source/Runtime/Online/OnlineSubsystemUtils/Private/EndTurnCallbackProxy.cpp
--- a/Engine/Source/Runtime/Online/OnlineSubsystemUtils/Private/EndTurnCallbackProxy.cpp
+++ b/Engine/Source/Runtime/Online/OnlineSubsystemUtils/Private/EndTurnCallbackProxy.cpp
@@ -70,7 +70,7 @@ void UEndTurnCallbackProxy::Activate()
 	OnFailure.Broadcast();
 }
 
-void UEndTurnCallbackProxy::UploadMatchDataDelegate(FString MatchID, bool Succeeded)
+void UEndTurnCallbackProxy::UploadMatchDataDelegate(FString InMatchID, bool Succeeded)
 {
 	if (Succeeded)
 	{
diff --git a/Engine/Source/Runtime/Online/OnlineSubsystemUtils/Private/OnlineSessionClient.cpp b/Engine/Source/Runtime/Online/OnlineSubsystemUtils/Private/OnlineSessionClient.cpp
--- a/Engine/Source/Runtime/Online/OnlineSubsystemUtils/Private/OnlineSessionClient.cpp
+++ b/Engine/Source/Runtime/Online/OnlineSubsystemUtils/Private/OnlineSessionClient.cpp
@@ -318,15 +318,15 @@ void UOnlineSessionClient::StartOnlineSession(FName SessionName)
 	if (PC)
 	{
 		UWorld* World = PC->GetWorld();
-		IOnlineSessionPtr SessionInt = Online::GetSessionInterface(World);
-		if (SessionInt.IsValid())
+		IOnlineSessionPtr SessionInterface = Online::GetSessionInterface(World);
+		if (SessionInterface.IsValid())
 		{
-			FNamedOnlineSession* Session = SessionInt->GetNamedSession(SessionName);
+			FNamedOnlineSession* Session = SessionInterface->GetNamedSession(SessionName);
 			if (Session &&
 				(Session->SessionState == EOnlineSessionState::Pending || Session->SessionState == EOnlineSessionState::Ended))
 			{
-				StartSessionCompleteHandle = SessionInt->AddOnStartSessionCompleteDelegate_Handle(FOnStartSessionCompleteDelegate::CreateUObject(this, &UOnlineSessionClient::OnStartSessionComplete));
-				SessionInt->StartSession(SessionName);
+				StartSessionCompleteHandle = SessionInterface->AddOnStartSessionCompleteDelegate_Handle(FOnStartSessionCompleteDelegate::CreateUObject(this, &UOnlineSessionClient::OnStartSessionComplete));
+				SessionInterface->StartSession(SessionName);
 			}
 		}
 	}
@@ -339,10 +339,10 @@ void UOnlineSessionClient::OnStartSessionComplete(FName SessionName, bool bWasSu
 	if (PC)
 	{
 		UWorld* World = PC->GetWorld();
-		IOnlineSessionPtr SessionInt = Online::GetSessionInterface(World);
-		if (SessionInt.IsValid())
+		IOnlineSessionPtr SessionInterface = Online::GetSessionInterface(World);
+		if (SessionInterface.IsValid())
 		{
-			SessionInt->ClearOnStartSessionCompleteDelegate_Handle(StartSessionCompleteHandle);
+			SessionInterface->ClearOnStartSessionCompleteDelegate_Handle(StartSessionCompleteHandle);
 		}
 	}
 }
@@ -353,15 +353,15 @@ void UOnlineSessionClient::EndOnlineSession(FName SessionName)
 	if (PC)
 	{
 		UWorld* World = PC->GetWorld();
-		IOnlineSessionPtr SessionInt = Online::GetSessionInterface(World);
-		if (SessionInt.IsValid())
+		IOnlineSessionPtr SessionInterface = Online::GetSessionInterface(World);
+		if (SessionInterface.IsValid())
 		{
-			FNamedOnlineSession* Session = SessionInt->GetNamedSession(SessionName);
+			FNamedOnlineSession* Session = SessionInterface->GetNamedSession(SessionName);
 			if (Session &&
 				Session->SessionState == EOnlineSessionState::InProgress)
 			{
-				EndSessionCompleteHandle = SessionInt->AddOnEndSessionCompleteDelegate_Handle(FOnStartSessionCompleteDelegate::CreateUObject(this, &UOnlineSessionClient::OnEndSessionComplete));
-				SessionInt->EndSession(SessionName);
+				EndSessionCompleteHandle = SessionInterface->AddOnEndSessionCompleteDelegate_Handle(FOnStartSessionCompleteDelegate::CreateUObject(this, &UOnlineSessionClient::OnEndSessionComplete));
+				SessionInterface->EndSession(SessionName);
 			}
 		}
 	}
@@ -374,10 +374,10 @@ void UOnlineSessionClient::OnEndSessionComplete(FName SessionName, bool bWasSucc
 	if (PC)
 	{
 		UWorld* World = PC->GetWorld();
-		IOnlineSessionPtr SessionInt = Online::GetSessionInterface(World);
-		if (SessionInt.IsValid())
+		IOnlineSessionPtr SessionInterface = Online::GetSessionInterface(World);
+		if (SessionInterface.IsValid())
 		{
-			SessionInt->ClearOnEndSessionCompleteDelegate_Handle(EndSessionCompleteHandle);
+			SessionInterface->ClearOnEndSessionCompleteDelegate_Handle(EndSessionCompleteHandle);
 		}
 	}
 }
\ No newline at end of file
diff --git a/Engine/Source/Runtime/Online/OnlineSubsystemUtils/Private/QuitMatchCallbackProxy.cpp b/Engine/Source/Runtime/Online/OnlineSubsystemUtils/Private/QuitMatchCallbackProxy.cpp
--- a/Engine/Source/Runtime/Online/OnlineSubsystemUtils/Private/QuitMatchCallbackProxy.cpp
+++ b/Engine/Source/Runtime/Online/OnlineSubsystemUtils/Private/QuitMatchCallbackProxy.cpp
@@ -57,7 +57,7 @@ void UQuitMatchCallbackProxy::Activate()
 	OnFailure.Broadcast();
 }
 
-void UQuitMatchCallbackProxy::QuitMatchDelegate(FString MatchID, bool Succeeded)
+void UQuitMatchCallbackProxy::QuitMatchDelegate(FString InMatchID, bool Succeeded)
 {
 	if (Succeeded)
 	{
diff --git a/Engine/Source/Runtime/Online/OnlineSubsystemUtils/Private/Tests/TestEntitlementsInterface.cpp b/Engine/Source/Runtime/Online/OnlineSubsystemUtils/Private/Tests/TestEntitlementsInterface.cpp
--- a/Engine/Source/Runtime/Online/OnlineSubsystemUtils/Private/Tests/TestEntitlementsInterface.cpp
+++ b/Engine/Source/Runtime/Online/OnlineSubsystemUtils/Private/Tests/TestEntitlementsInterface.cpp
@@ -67,14 +67,14 @@ void FTestEntitlementsInterface::FinishTest()
 	delete this;
 }
 
-void FTestEntitlementsInterface::OnQueryEntitlementsComplete(bool bWasSuccessful, const FUniqueNetId& UserId, const FString& Namespace, const FString& Error)
+void FTestEntitlementsInterface::OnQueryEntitlementsComplete(bool bWasSuccessful, const FUniqueNetId& InUserId, const FString& Namespace, const FString& Error)
 {
 	UE_LOG(LogOnline, Log, TEXT("enumerated entitlements. UserId=%s Result=%s Error=[%s]"), 
-		*UserId.ToDebugString(), bWasSuccessful ? TEXT("true") : TEXT("false"), *Error);
+		*InUserId.ToDebugString(), bWasSuccessful ? TEXT("true") : TEXT("false"), *Error);
 
 	// Now check em out
 	TArray<TSharedRef<FOnlineEntitlement>> Entitlements;
-	EntitlementsOSS->GetAllEntitlements(UserId, FString(), Entitlements);
+	EntitlementsOSS->GetAllEntitlements(InUserId, FString(), Entitlements);
 
 	for (auto It = Entitlements.CreateConstIterator(); It; ++It)
 	{
diff --git a/Engine/Source/Runtime/Online/OnlineSubsystemUtils/Private/Tests/TestSessionInterface.cpp b/Engine/Source/Runtime/Online/OnlineSubsystemUtils/Private/Tests/TestSessionInterface.cpp
--- a/Engine/Source/Runtime/Online/OnlineSubsystemUtils/Private/Tests/TestSessionInterface.cpp
+++ b/Engine/Source/Runtime/Online/OnlineSubsystemUtils/Private/Tests/TestSessionInterface.cpp
@@ -513,10 +513,10 @@ bool FTestSessionInterface::Exec(UWorld* InWorld, const TCHAR* Cmd, FOutputDevic
 					if (!bFound)
 					{
 						// use friend str as id instead of display name
-						TSharedPtr<FUniqueNetId> UserId = Identity->CreateUniquePlayerId(FriendStr);
-						if (UserId.IsValid())
+						TSharedPtr<FUniqueNetId> FriendUserId = Identity->CreateUniquePlayerId(FriendStr);
+						if (FriendUserId.IsValid())
 						{	
-							SessionInt->SendSessionInviteToFriend(LocalUserNum, SessionName, *UserId);
+							SessionInt->SendSessionInviteToFriend(LocalUserNum, SessionName, *FriendUserId);
 						}
 					}
 				}
diff --git a/Engine/Source/Runtime/OpenGLDrv/Private/Mac/MacOpenGLContext.cpp b/Engine/Source/Runtime/OpenGLDrv/Private/Mac/MacOpenGLContext.cpp
--- a/Engine/Source/Runtime/OpenGLDrv/Private/Mac/MacOpenGLContext.cpp
+++ b/Engine/Source/Runtime/OpenGLDrv/Private/Mac/MacOpenGLContext.cpp
@@ -403,14 +403,14 @@ void FPlatformOpenGLContext::VerifyCurrentContext()
 		bool const bGPUChange = (GMacUseAutomaticGraphicsSwitching || GMacUseMultipleGPUs && GMacExplicitRendererID);
 		if(bGPUChange && RendererID != RendererToMatch)
 		{
-			CGLPixelFormatObj PixelFormat = CGLGetPixelFormat(Current);
+			CGLPixelFormatObj CurrentPixelFormat = CGLGetPixelFormat(Current);
 			GLint NumVirtualScreens = 0;
-			CGLDescribePixelFormat(PixelFormat, 0, kCGLPFAVirtualScreenCount, &NumVirtualScreens);
+			CGLDescribePixelFormat(CurrentPixelFormat, 0, kCGLPFAVirtualScreenCount, &NumVirtualScreens);
 			if(NumVirtualScreens > 1)
 			{
 				for(GLint i = 0; i < NumVirtualScreens; i++)
 				{
-					CGLDescribePixelFormat(PixelFormat, i, kCGLPFARendererID, &RendererID);
+					CGLDescribePixelFormat(CurrentPixelFormat, i, kCGLPFARendererID, &RendererID);
 					if(RendererID == RendererToMatch)
 					{
 						CGLSetVirtualScreen(Current, i);
diff --git a/Engine/Source/Runtime/OpenGLDrv/Private/Mac/OpenGLMac.cpp b/Engine/Source/Runtime/OpenGLDrv/Private/Mac/OpenGLMac.cpp
--- a/Engine/Source/Runtime/OpenGLDrv/Private/Mac/OpenGLMac.cpp
+++ b/Engine/Source/Runtime/OpenGLDrv/Private/Mac/OpenGLMac.cpp
@@ -961,9 +961,9 @@ uint64 FMacOpenGL::GetVideoMemorySize()
 		GLint RendererID = 0;
 		GLint DisplayMask = 0;
 		
-		CGLPixelFormatObj PixelFormat = CGLGetPixelFormat(Context);
+		CGLPixelFormatObj ContextPixelFormat = CGLGetPixelFormat(Context);
 		
-		if(PixelFormat && CGLDescribePixelFormat(PixelFormat, VirtualScreen, kCGLPFADisplayMask, &DisplayMask) == kCGLNoError
+		if(ContextPixelFormat && CGLDescribePixelFormat(ContextPixelFormat, VirtualScreen, kCGLPFADisplayMask, &DisplayMask) == kCGLNoError
 		   && CGLGetParameter(Context, kCGLCPCurrentRendererID, &RendererID) == kCGLNoError)
 		{
 			// Get renderer info for all renderers that match the display mask.
diff --git a/Engine/Source/Runtime/OpenGLDrv/Private/OpenGLRenderTarget.cpp b/Engine/Source/Runtime/OpenGLDrv/Private/OpenGLRenderTarget.cpp
--- a/Engine/Source/Runtime/OpenGLDrv/Private/OpenGLRenderTarget.cpp
+++ b/Engine/Source/Runtime/OpenGLDrv/Private/OpenGLRenderTarget.cpp
@@ -796,9 +796,9 @@ void FOpenGLDynamicRHI::RHIReadSurfaceFloatData(FTextureRHIParamRef TextureRHI,F
 		FMemory::Memzero(FloatData.GetData(),SizeX * SizeY*sizeof(FLinearColor));
 		glReadPixels(Rect.Min.X, Rect.Min.Y, SizeX, SizeY, GL_RGBA, GL_FLOAT, FloatData.GetData());
 		FLinearColor* FloatDataPtr = FloatData.GetData();
-		for (uint32 Index = 0; Index < SizeX * SizeY; ++Index, ++FloatDataPtr)
+		for (uint32 DataIndex = 0; DataIndex < SizeX * SizeY; ++DataIndex, ++FloatDataPtr)
 		{
-			OutData[Index] = FFloat16Color(*FloatDataPtr);
+			OutData[DataIndex] = FFloat16Color(*FloatDataPtr);
 		}
 	}
 	else
diff --git a/Engine/Source/Runtime/OpenGLDrv/Private/OpenGLShaders.cpp b/Engine/Source/Runtime/OpenGLDrv/Private/OpenGLShaders.cpp
--- a/Engine/Source/Runtime/OpenGLDrv/Private/OpenGLShaders.cpp
+++ b/Engine/Source/Runtime/OpenGLDrv/Private/OpenGLShaders.cpp
@@ -1915,8 +1915,6 @@ static void BindShaderStage(FOpenGLLinkedProgramConfiguration::ShaderInfo& Shade
 		
 			if(!bInterpolatorMatches)
 			{
-				FString PrevSource = GetShaderStageSource<TOpenGLStage1>(PrevStage);
-				
 				if(InputErrors.Num() == 0)
 				{
 					FOpenGLCodeHeader Header;
diff --git a/Engine/Source/Runtime/OpenGLDrv/Private/OpenGLTexture.cpp b/Engine/Source/Runtime/OpenGLDrv/Private/OpenGLTexture.cpp
--- a/Engine/Source/Runtime/OpenGLDrv/Private/OpenGLTexture.cpp
+++ b/Engine/Source/Runtime/OpenGLDrv/Private/OpenGLTexture.cpp
@@ -709,7 +709,7 @@ void TOpenGLTexture<RHIResourceType>::Resolve(uint32 MipIndex,uint32 ArrayIndex)
 }
 
 template<typename RHIResourceType>
-void* TOpenGLTexture<RHIResourceType>::Lock(uint32 MipIndex,uint32 ArrayIndex,EResourceLockMode LockMode,uint32& DestStride)
+void* TOpenGLTexture<RHIResourceType>::Lock(uint32 InMipIndex,uint32 ArrayIndex,EResourceLockMode LockMode,uint32& DestStride)
 {
 	VERIFY_GL_SCOPE();
 
@@ -729,8 +729,8 @@ void* TOpenGLTexture<RHIResourceType>::Lock(uint32 MipIndex,uint32 ArrayIndex,ER
 	const uint32 BlockSizeX = GPixelFormats[PixelFormat].BlockSizeX;
 	const uint32 BlockSizeY = GPixelFormats[PixelFormat].BlockSizeY;
 	const uint32 BlockBytes = GPixelFormats[PixelFormat].BlockBytes;
-	const uint32 MipSizeX = FMath::Max(this->GetSizeX() >> MipIndex,BlockSizeX);
-	const uint32 MipSizeY = FMath::Max(this->GetSizeY() >> MipIndex,BlockSizeY);
+	const uint32 MipSizeX = FMath::Max(this->GetSizeX() >> InMipIndex,BlockSizeX);
+	const uint32 MipSizeY = FMath::Max(this->GetSizeY() >> InMipIndex,BlockSizeY);
 	uint32 NumBlocksX = (MipSizeX + BlockSizeX - 1) / BlockSizeX;
 	uint32 NumBlocksY = (MipSizeY + BlockSizeY - 1) / BlockSizeY;
 	if ( PixelFormat == PF_PVRTC2 || PixelFormat == PF_PVRTC4 )
@@ -743,7 +743,7 @@ void* TOpenGLTexture<RHIResourceType>::Lock(uint32 MipIndex,uint32 ArrayIndex,ER
 
 	DestStride = NumBlocksX * BlockBytes;
 
-	const int32 BufferIndex = MipIndex * (bCubemap ? 6 : 1) * this->GetEffectiveSizeZ() + ArrayIndex;
+	const int32 BufferIndex = InMipIndex * (bCubemap ? 6 : 1) * this->GetEffectiveSizeZ() + ArrayIndex;
 
 	// Should we use client-storage to improve update time on platforms that require it
 	const FOpenGLTextureFormat& GLFormat = GOpenGLTextureFormats[PixelFormat];
@@ -768,7 +768,7 @@ void* TOpenGLTexture<RHIResourceType>::Lock(uint32 MipIndex,uint32 ArrayIndex,ER
 
 		if( LockMode != RLM_WriteOnly && !bCPUTexResolved && FOpenGL::SupportsPixelBuffers() )
 		{
-			Resolve(MipIndex, ArrayIndex);
+			Resolve(InMipIndex, ArrayIndex);
 		}
 
 		result = PixelBuffer->Lock(0, PixelBuffer->GetSize(), LockMode == RLM_ReadOnly, LockMode != RLM_ReadOnly);
@@ -778,7 +778,7 @@ void* TOpenGLTexture<RHIResourceType>::Lock(uint32 MipIndex,uint32 ArrayIndex,ER
 		// Use APPLE_client_storage to reduce memory usage and improve performance
 		// GL's which support this extension only need copy a pointer, not the memory contents
 		check( FOpenGL::SupportsClientStorage() && !FOpenGL::SupportsTextureView() );
-		if(GetAllocatedStorageForMip(MipIndex,ArrayIndex))
+		if(GetAllocatedStorageForMip(InMipIndex,ArrayIndex))
 		{
 			result = ClientStorageBuffers[BufferIndex].Data;
 		}
diff --git a/Engine/Source/Runtime/Renderer/Private/PostProcess/SceneRenderTargets.cpp b/Engine/Source/Runtime/Renderer/Private/PostProcess/SceneRenderTargets.cpp
--- a/Engine/Source/Runtime/Renderer/Private/PostProcess/SceneRenderTargets.cpp
+++ b/Engine/Source/Runtime/Renderer/Private/PostProcess/SceneRenderTargets.cpp
@@ -621,35 +621,35 @@ void FSceneRenderTargets::AllocGBufferTargets()
 	{
 
 		// Allocate the Gbuffer resource uniform buffer.
-		const FSceneRenderTargetItem& GBufferA = this->GBufferA ? this->GBufferA->GetRenderTargetItem() : GSystemTextures.BlackDummy->GetRenderTargetItem();
-		const FSceneRenderTargetItem& GBufferB = this->GBufferB ? this->GBufferB->GetRenderTargetItem() : GSystemTextures.BlackDummy->GetRenderTargetItem();
-		const FSceneRenderTargetItem& GBufferC = this->GBufferC ? this->GBufferC->GetRenderTargetItem() : GSystemTextures.BlackDummy->GetRenderTargetItem();
-		const FSceneRenderTargetItem& GBufferD = this->GBufferD ? this->GBufferD->GetRenderTargetItem() : GSystemTextures.BlackDummy->GetRenderTargetItem();
-		const FSceneRenderTargetItem& GBufferE = this->GBufferE ? this->GBufferE->GetRenderTargetItem() : GSystemTextures.BlackDummy->GetRenderTargetItem();
-		const FSceneRenderTargetItem& GBufferVelocity = this->GBufferVelocity ? this->GBufferVelocity->GetRenderTargetItem() : GSystemTextures.BlackDummy->GetRenderTargetItem();
+		const FSceneRenderTargetItem& GBufferAToUse = GBufferA ? GBufferA->GetRenderTargetItem() : GSystemTextures.BlackDummy->GetRenderTargetItem();
+		const FSceneRenderTargetItem& GBufferBToUse = GBufferB ? GBufferB->GetRenderTargetItem() : GSystemTextures.BlackDummy->GetRenderTargetItem();
+		const FSceneRenderTargetItem& GBufferCToUse = GBufferC ? GBufferC->GetRenderTargetItem() : GSystemTextures.BlackDummy->GetRenderTargetItem();
+		const FSceneRenderTargetItem& GBufferDToUse = GBufferD ? GBufferD->GetRenderTargetItem() : GSystemTextures.BlackDummy->GetRenderTargetItem();
+		const FSceneRenderTargetItem& GBufferEToUse = GBufferE ? GBufferE->GetRenderTargetItem() : GSystemTextures.BlackDummy->GetRenderTargetItem();
+		const FSceneRenderTargetItem& GBufferVelocityToUse = GBufferVelocity ? GBufferVelocity->GetRenderTargetItem() : GSystemTextures.BlackDummy->GetRenderTargetItem();
 
 		FGBufferResourceStruct GBufferResourceStruct;
 
-		GBufferResourceStruct.GBufferATexture = GBufferA.ShaderResourceTexture;
-		GBufferResourceStruct.GBufferBTexture = GBufferB.ShaderResourceTexture;
-		GBufferResourceStruct.GBufferCTexture = GBufferC.ShaderResourceTexture;
-		GBufferResourceStruct.GBufferDTexture = GBufferD.ShaderResourceTexture;
-		GBufferResourceStruct.GBufferETexture = GBufferE.ShaderResourceTexture;
-		GBufferResourceStruct.GBufferVelocityTexture = GBufferVelocity.ShaderResourceTexture;
-
-		GBufferResourceStruct.GBufferATextureNonMS = GBufferA.ShaderResourceTexture;
-		GBufferResourceStruct.GBufferBTextureNonMS = GBufferB.ShaderResourceTexture;
-		GBufferResourceStruct.GBufferCTextureNonMS = GBufferC.ShaderResourceTexture;
-		GBufferResourceStruct.GBufferDTextureNonMS = GBufferD.ShaderResourceTexture;
-		GBufferResourceStruct.GBufferETextureNonMS = GBufferE.ShaderResourceTexture;
-		GBufferResourceStruct.GBufferVelocityTextureNonMS = GBufferVelocity.ShaderResourceTexture;
-
-		GBufferResourceStruct.GBufferATextureMS = GBufferA.TargetableTexture;
-		GBufferResourceStruct.GBufferBTextureMS = GBufferB.TargetableTexture;
-		GBufferResourceStruct.GBufferCTextureMS = GBufferC.TargetableTexture;
-		GBufferResourceStruct.GBufferDTextureMS = GBufferD.TargetableTexture;
-		GBufferResourceStruct.GBufferETextureMS = GBufferE.TargetableTexture;
-		GBufferResourceStruct.GBufferVelocityTextureMS = GBufferVelocity.TargetableTexture;
+		GBufferResourceStruct.GBufferATexture = GBufferAToUse.ShaderResourceTexture;
+		GBufferResourceStruct.GBufferBTexture = GBufferBToUse.ShaderResourceTexture;
+		GBufferResourceStruct.GBufferCTexture = GBufferCToUse.ShaderResourceTexture;
+		GBufferResourceStruct.GBufferDTexture = GBufferDToUse.ShaderResourceTexture;
+		GBufferResourceStruct.GBufferETexture = GBufferEToUse.ShaderResourceTexture;
+		GBufferResourceStruct.GBufferVelocityTexture = GBufferVelocityToUse.ShaderResourceTexture;
+
+		GBufferResourceStruct.GBufferATextureNonMS = GBufferAToUse.ShaderResourceTexture;
+		GBufferResourceStruct.GBufferBTextureNonMS = GBufferBToUse.ShaderResourceTexture;
+		GBufferResourceStruct.GBufferCTextureNonMS = GBufferCToUse.ShaderResourceTexture;
+		GBufferResourceStruct.GBufferDTextureNonMS = GBufferDToUse.ShaderResourceTexture;
+		GBufferResourceStruct.GBufferETextureNonMS = GBufferEToUse.ShaderResourceTexture;
+		GBufferResourceStruct.GBufferVelocityTextureNonMS = GBufferVelocityToUse.ShaderResourceTexture;
+
+		GBufferResourceStruct.GBufferATextureMS = GBufferAToUse.TargetableTexture;
+		GBufferResourceStruct.GBufferBTextureMS = GBufferBToUse.TargetableTexture;
+		GBufferResourceStruct.GBufferCTextureMS = GBufferCToUse.TargetableTexture;
+		GBufferResourceStruct.GBufferDTextureMS = GBufferDToUse.TargetableTexture;
+		GBufferResourceStruct.GBufferETextureMS = GBufferEToUse.TargetableTexture;
+		GBufferResourceStruct.GBufferVelocityTextureMS = GBufferVelocityToUse.TargetableTexture;
 
 		GBufferResourceStruct.GBufferATextureSampler = TStaticSamplerState<>::GetRHI();
 		GBufferResourceStruct.GBufferBTextureSampler = TStaticSamplerState<>::GetRHI();
@@ -791,13 +791,13 @@ void FSceneRenderTargets::FinishRenderingSceneColor(FRHICommandListImmediate& RH
 
 bool FSceneRenderTargets::BeginRenderingCustomDepth(FRHICommandListImmediate& RHICmdList, bool bPrimitives)
 {
-	IPooledRenderTarget* CustomDepth = RequestCustomDepth(bPrimitives);
+	IPooledRenderTarget* CustomDepthRenderTarget = RequestCustomDepth(bPrimitives);
 
-	if(CustomDepth)
+	if(CustomDepthRenderTarget)
 	{
 		SCOPED_DRAW_EVENT(RHICmdList, BeginRenderingCustomDepth);
 		
-		SetRenderTarget(RHICmdList, FTextureRHIRef(), CustomDepth->GetRenderTargetItem().ShaderResourceTexture);
+		SetRenderTarget(RHICmdList, FTextureRHIRef(), CustomDepthRenderTarget->GetRenderTargetItem().ShaderResourceTexture);
 
 		return true;
 	}
@@ -1293,8 +1293,10 @@ void FSceneRenderTargets::AllocateForwardShadingPathRenderTargets()
 
 void FSceneRenderTargets::AllocateForwardShadingShadowDepthTarget(const FIntPoint& ShadowBufferResolution)
 {
-	FPooledRenderTargetDesc Desc(FPooledRenderTargetDesc::Create2DDesc(ShadowBufferResolution, PF_ShadowDepth, TexCreate_None, TexCreate_DepthStencilTargetable, false));
-	GRenderTargetPool.FindFreeElement(Desc, ShadowDepthZ, TEXT("ShadowDepthZ"));
+	{
+		FPooledRenderTargetDesc Desc(FPooledRenderTargetDesc::Create2DDesc(ShadowBufferResolution, PF_ShadowDepth, TexCreate_None, TexCreate_DepthStencilTargetable, false));
+		GRenderTargetPool.FindFreeElement(Desc, ShadowDepthZ, TEXT("ShadowDepthZ"));
+	}
 
 	if (!GSupportsDepthRenderTargetWithoutColorRenderTarget)
 	{
diff --git a/Engine/Source/Runtime/Renderer/Private/SceneVisibility.cpp b/Engine/Source/Runtime/Renderer/Private/SceneVisibility.cpp
--- a/Engine/Source/Runtime/Renderer/Private/SceneVisibility.cpp
+++ b/Engine/Source/Runtime/Renderer/Private/SceneVisibility.cpp
@@ -1388,24 +1388,24 @@ static void ComputeAndMarkRelevanceForViewParallel(
 }
 
 void FSceneRenderer::GatherDynamicMeshElements(
-	TArray<FViewInfo>& Views, 
-	const FScene* Scene, 
-	const FSceneViewFamily& ViewFamily, 
+	TArray<FViewInfo>& InViews, 
+	const FScene* InScene, 
+	const FSceneViewFamily& InViewFamily, 
 	const FPrimitiveViewMasks& HasDynamicMeshElementsMasks, 
 	const FPrimitiveViewMasks& HasDynamicEditorMeshElementsMasks, 
 	FMeshElementCollector& Collector)
 {
 	SCOPE_CYCLE_COUNTER(STAT_GetDynamicMeshElements);
 
-	int32 NumPrimitives = Scene->Primitives.Num();
+	int32 NumPrimitives = InScene->Primitives.Num();
 	check(HasDynamicMeshElementsMasks.Num() == NumPrimitives);
 
 	{
 		Collector.ClearViewMeshArrays();
 
-		for (int32 ViewIndex = 0; ViewIndex < Views.Num(); ViewIndex++)
+		for (int32 ViewIndex = 0; ViewIndex < InViews.Num(); ViewIndex++)
 		{
-			Collector.AddViewMeshArrays(&Views[ViewIndex], &Views[ViewIndex].DynamicMeshElements, &Views[ViewIndex].SimpleElementCollector, ViewFamily.GetFeatureLevel());
+			Collector.AddViewMeshArrays(&InViews[ViewIndex], &InViews[ViewIndex].DynamicMeshElements, &InViews[ViewIndex].SimpleElementCollector, InViewFamily.GetFeatureLevel());
 		}
 
 		for (int32 PrimitiveIndex = 0; PrimitiveIndex < NumPrimitives; ++PrimitiveIndex)
@@ -1414,9 +1414,9 @@ void FSceneRenderer::GatherDynamicMeshElements(
 
 			if (ViewMask != 0)
 			{
-				FPrimitiveSceneInfo* PrimitiveSceneInfo = Scene->Primitives[PrimitiveIndex];
+				FPrimitiveSceneInfo* PrimitiveSceneInfo = InScene->Primitives[PrimitiveIndex];
 				Collector.SetPrimitive(PrimitiveSceneInfo->Proxy, PrimitiveSceneInfo->DefaultDynamicHitProxyId);
-				PrimitiveSceneInfo->Proxy->GetDynamicMeshElements(ViewFamily.Views, ViewFamily, ViewMask, Collector);
+				PrimitiveSceneInfo->Proxy->GetDynamicMeshElements(InViewFamily.Views, InViewFamily, ViewMask, Collector);
 			}
 		}
 	}
@@ -1425,9 +1425,9 @@ void FSceneRenderer::GatherDynamicMeshElements(
 	{
 		Collector.ClearViewMeshArrays();
 
-		for (int32 ViewIndex = 0; ViewIndex < Views.Num(); ViewIndex++)
+		for (int32 ViewIndex = 0; ViewIndex < InViews.Num(); ViewIndex++)
 		{
-			Collector.AddViewMeshArrays(&Views[ViewIndex], &Views[ViewIndex].DynamicEditorMeshElements, &Views[ViewIndex].EditorSimpleElementCollector, ViewFamily.GetFeatureLevel());
+			Collector.AddViewMeshArrays(&InViews[ViewIndex], &InViews[ViewIndex].DynamicEditorMeshElements, &InViews[ViewIndex].EditorSimpleElementCollector, InViewFamily.GetFeatureLevel());
 		}
 
 		for (int32 PrimitiveIndex = 0; PrimitiveIndex < NumPrimitives; ++PrimitiveIndex)
@@ -1436,9 +1436,9 @@ void FSceneRenderer::GatherDynamicMeshElements(
 
 			if (ViewMask != 0)
 			{
-				FPrimitiveSceneInfo* PrimitiveSceneInfo = Scene->Primitives[PrimitiveIndex];
+				FPrimitiveSceneInfo* PrimitiveSceneInfo = InScene->Primitives[PrimitiveIndex];
 				Collector.SetPrimitive(PrimitiveSceneInfo->Proxy, PrimitiveSceneInfo->DefaultDynamicHitProxyId);
-				PrimitiveSceneInfo->Proxy->GetDynamicMeshElements(ViewFamily.Views, ViewFamily, ViewMask, Collector);
+				PrimitiveSceneInfo->Proxy->GetDynamicMeshElements(InViewFamily.Views, InViewFamily, ViewMask, Collector);
 			}
 		}
 	}
diff --git a/Engine/Source/Runtime/UMG/Private/Components/RichTextBlockDecorator.cpp b/Engine/Source/Runtime/UMG/Private/Components/RichTextBlockDecorator.cpp
--- a/Engine/Source/Runtime/UMG/Private/Components/RichTextBlockDecorator.cpp
+++ b/Engine/Source/Runtime/UMG/Private/Components/RichTextBlockDecorator.cpp
@@ -9,10 +9,10 @@
 class FDefaultRichTextRun : public FSlateTextRun
 {
 public:
-	FDefaultRichTextRun(URichTextBlockDecorator* Decorator, const TSharedRef<FTextLayout>& TextLayout, const FRunInfo& InRunInfo, const TSharedRef< const FString >& InText, const FTextBlockStyle& InStyle, const FTextRange& InRange)
+	FDefaultRichTextRun(URichTextBlockDecorator* InDecorator, const TSharedRef<FTextLayout>& InTextLayout, const FRunInfo& InRunInfo, const TSharedRef< const FString >& InText, const FTextBlockStyle& InStyle, const FTextRange& InRange)
 		: FSlateTextRun(InRunInfo, InText, InStyle, InRange)
-		, TextLayout(TextLayout)
-		, Decorator(Decorator)
+		, TextLayout(InTextLayout)
+		, Decorator(InDecorator)
 	{
 	}
 
@@ -87,10 +87,10 @@ private:
 /////////////////////////////////////////////////////
 // FDefaultRichTextDecorator
 
-FDefaultRichTextDecorator::FDefaultRichTextDecorator(URichTextBlockDecorator* Decorator, const FSlateFontInfo& DefaultFont, const FLinearColor& DefaultColor)
-	: DefaultFont(DefaultFont)
-	, DefaultColor(DefaultColor)
-	, Decorator(Decorator)
+FDefaultRichTextDecorator::FDefaultRichTextDecorator(URichTextBlockDecorator* InDecorator, const FSlateFontInfo& InDefaultFont, const FLinearColor& InDefaultColor)
+	: DefaultFont(InDefaultFont)
+	, DefaultColor(InDefaultColor)
+	, Decorator(InDecorator)
 {
 }
 
