Fix shadowed variables

diff --git a/Engine/Source/Developer/GameplayDebugger/Private/GameplayDebuggingControllerComponent.cpp b/Engine/Source/Developer/GameplayDebugger/Private/GameplayDebuggingControllerComponent.cpp
--- a/Engine/Source/Developer/GameplayDebugger/Private/GameplayDebuggingControllerComponent.cpp
+++ b/Engine/Source/Developer/GameplayDebugger/Private/GameplayDebuggingControllerComponent.cpp
@@ -444,11 +444,13 @@ void UGameplayDebuggingControllerComponent::ToggleDebugCamera()
 		// and ungly @HACK to be able to spawn camera in game world rather then
 		// in editor world (if running PIE). Hate it, but it works, and 
 		// this is a debugging tool		
-		FActorSpawnParameters SpawnInfo;
-		SpawnInfo.bNoCollisionFail = true;
-		SpawnInfo.Owner = PlayerOwner->GetWorldSettings();
-		SpawnInfo.Instigator = PlayerOwner->Instigator;
-		DebugCameraController = GetWorld()->SpawnActor<ADebugCameraController>(SpawnInfo);
+		{
+			FActorSpawnParameters SpawnInfo;
+			SpawnInfo.bNoCollisionFail = true;
+			SpawnInfo.Owner = PlayerOwner->GetWorldSettings();
+			SpawnInfo.Instigator = PlayerOwner->Instigator;
+			DebugCameraController = GetWorld()->SpawnActor<ADebugCameraController>(SpawnInfo);
+		}
 
 		if (DebugCameraController.IsValid())
 		{
diff --git a/Engine/Source/Developer/MeshUtilities/Private/LayoutUV.cpp b/Engine/Source/Developer/MeshUtilities/Private/LayoutUV.cpp
--- a/Engine/Source/Developer/MeshUtilities/Private/LayoutUV.cpp
+++ b/Engine/Source/Developer/MeshUtilities/Private/LayoutUV.cpp
@@ -384,7 +384,10 @@ void FLayoutUV::FindCharts( const TMultiMap<int32,int32>& OverlappingCorners )
 			}
 		}
 
+		PRAGMA_DISABLE_SHADOW_VARIABLE_WARNINGS
 		uint32 NumTris = 0;
+		PRAGMA_POP
+
 		for( int32 i = 0; i < Charts.Num(); i++ )
 		{
 			FMeshChart& Chart = Charts[i];
@@ -604,9 +607,9 @@ void FLayoutUV::ScaleCharts( float UVScale )
 	{
 		uint32 NumMaxedOut = 0;
 		float ScaledUVArea = 0.0f;
-		for( int32 i = 0; i < Charts.Num(); i++ )
+		for( int32 ChartIndex = 0; ChartIndex < Charts.Num(); ChartIndex++ )
 		{
-			FMeshChart& Chart = Charts[i];
+			FMeshChart& Chart = Charts[ChartIndex];
 
 			FVector2D ChartSize	= Chart.MaxUV - Chart.MinUV;
 			FVector2D ChartSizeScaled = ChartSize * Chart.UVScale * UniformScale;
@@ -661,21 +664,21 @@ void FLayoutUV::ScaleCharts( float UVScale )
 	{
 		uint32 NumMaxedOut = 0;
 		float ScaledUVArea = 0.0f;
-		for( int32 i = 0; i < Charts.Num(); i++ )
+		for( int32 ChartIndex = 0; ChartIndex < Charts.Num(); ChartIndex++ )
 		{
-			FMeshChart& Chart = Charts[i];
+			FMeshChart& Chart = Charts[ChartIndex];
 
 			for( int k = 0; k < 2; k++ )
 			{
-				const float MaxChartSize = TextureResolution - 1.0f;
+				const float MaximumChartSize = TextureResolution - 1.0f;
 				const float ChartSize = Chart.MaxUV[k] - Chart.MinUV[k];
 				const float ChartSizeScaled = ChartSize * Chart.UVScale[k] * NonuniformScale;
 
 				const float Epsilon = 0.01f;
-				if( ChartSizeScaled + Epsilon > MaxChartSize )
+				if( ChartSizeScaled + Epsilon > MaximumChartSize )
 				{
 					// Scale oversized charts to max size
-					Chart.UVScale[k] = MaxChartSize / ChartSize;
+					Chart.UVScale[k] = MaximumChartSize / ChartSize;
 					NumMaxedOut++;
 				}
 				else
diff --git a/Engine/Source/Developer/ShaderFormatOpenGL/Private/GlslBackend.cpp b/Engine/Source/Developer/ShaderFormatOpenGL/Private/GlslBackend.cpp
--- a/Engine/Source/Developer/ShaderFormatOpenGL/Private/GlslBackend.cpp
+++ b/Engine/Source/Developer/ShaderFormatOpenGL/Private/GlslBackend.cpp
@@ -3406,7 +3406,6 @@ static ir_rvalue* GenShaderInputSemantic(
 	bool& ApplyClipSpaceAdjustment
 	)
 {
-	ir_variable* Variable = NULL;
 	if (Semantic && FCStringAnsi::Strnicmp(Semantic, "SV_", 3) == 0)
 	{
 		FSystemValue* SystemValues = SystemValueTable[Frequency];
@@ -3420,7 +3419,7 @@ static ir_rvalue* GenShaderInputSemantic(
 				{
 					// Built-in array variable. Like gl_in[x].gl_Position.
 					// The variable for it has already been created in GenShaderInput().
-					/*ir_variable**/ Variable = ParseState->symbols->get_variable("gl_in");
+					ir_variable* Variable = ParseState->symbols->get_variable("gl_in");
 					check(Variable);
 					ir_dereference_variable* ArrayDeref = new(ParseState)ir_dereference_variable(Variable);
 					ir_dereference_array* StructDeref = new(ParseState)ir_dereference_array(
@@ -3504,6 +3503,7 @@ static ir_rvalue* GenShaderInputSemantic(
 		}
 	}
 
+	ir_variable* Variable = NULL;
 	if (Variable == NULL && Frequency == HSF_DomainShader)
 	{
 		const int PrefixLength = 13;
@@ -3579,7 +3579,7 @@ static ir_rvalue* GenShaderInputSemantic(
 		if (ParseState->bGenerateES && Type->is_integer())
 		{
 			// Convert integer attributes to floats
-			ir_variable* Variable = new(ParseState)ir_variable(
+			Variable = new(ParseState)ir_variable(
 				Type,
 				ralloc_asprintf(ParseState, "%s_%s_I", Prefix, Semantic),
 				ir_var_temporary
@@ -3653,7 +3653,7 @@ static ir_rvalue* GenShaderInputSemantic(
 			VariableType = glsl_type::get_array_instance(VariableType, SemanticArraySize);
 		}
 
-		ir_variable* Variable = new(ParseState)ir_variable(VariableType, ralloc_asprintf(ParseState, "in_%s", Semantic), ir_var_in);
+		Variable = new(ParseState)ir_variable(VariableType, ralloc_asprintf(ParseState, "in_%s", Semantic), ir_var_in);
 		Variable->read_only = true;
 		Variable->is_interface_block = true;
 		Variable->centroid = InputQualifier.Fields.bCentroid;
diff --git a/Engine/Source/Runtime/Renderer/Private/ShadowRendering.cpp b/Engine/Source/Runtime/Renderer/Private/ShadowRendering.cpp
--- a/Engine/Source/Runtime/Renderer/Private/ShadowRendering.cpp
+++ b/Engine/Source/Runtime/Renderer/Private/ShadowRendering.cpp
@@ -704,7 +704,7 @@ void OverrideWithDefaultMaterialForShadowDepth(
 -----------------------------------------------------------------------------*/
 
 template <bool bRenderingReflectiveShadowMaps>
-void FShadowDepthDrawingPolicy<bRenderingReflectiveShadowMaps>::UpdateElementState(FShadowStaticMeshElement& State, ERHIFeatureLevel::Type FeatureLevel)
+void FShadowDepthDrawingPolicy<bRenderingReflectiveShadowMaps>::UpdateElementState(FShadowStaticMeshElement& State, ERHIFeatureLevel::Type InFeatureLevel)
 {
 	// can be optimized
 	*this = FShadowDepthDrawingPolicy(
@@ -712,7 +712,7 @@ void FShadowDepthDrawingPolicy<bRenderingReflectiveShadowMaps>::UpdateElementSta
 		bDirectionalLight,
 		bOnePassPointLightShadow,
 		bPreShadow,
-		FeatureLevel,
+		InFeatureLevel,
 		State.Mesh->VertexFactory,
 		State.RenderProxy,
 		State.bIsTwoSided,
@@ -1829,12 +1829,14 @@ void FProjectedShadowInfo::RenderProjection(FRHICommandListImmediate& RHICmdList
 
 	// Find the shadow's view relevance.
 	const FVisibleLightViewInfo& VisibleLightViewInfo = View->VisibleLightInfos[LightSceneInfo->Id];
-	FPrimitiveViewRelevance ViewRelevance = VisibleLightViewInfo.ProjectedShadowViewRelevanceMap[ShadowId];
-
-	// Don't render shadows for subjects which aren't view relevant.
-	if (ViewRelevance.bShadowRelevance == false)
 	{
-		return;
+		FPrimitiveViewRelevance ViewRelevance = VisibleLightViewInfo.ProjectedShadowViewRelevanceMap[ShadowId];
+
+		// Don't render shadows for subjects which aren't view relevant.
+		if (ViewRelevance.bShadowRelevance == false)
+		{
+			return;
+		}
 	}
 
 	if (CascadeSettings.bRayTracedDistanceField)
@@ -1854,22 +1856,22 @@ void FProjectedShadowInfo::RenderProjection(FRHICommandListImmediate& RHICmdList
 		const FVector PreShadowToPreViewTranslation(View->ViewMatrices.PreViewTranslation - PreShadowTranslation);
 
 		// fill out the frustum vertices (this is only needed in the non-whole scene case)
-		for(uint32 Z = 0;Z < 2;Z++)
+		for(uint32 vZ = 0;vZ < 2;vZ++)
 		{
-			for(uint32 Y = 0;Y < 2;Y++)
+			for(uint32 vY = 0;vY < 2;vY++)
 			{
-				for(uint32 X = 0;X < 2;X++)
+				for(uint32 vX = 0;vX < 2;vX++)
 				{
 					const FVector4 UnprojectedVertex = InvReceiverMatrix.TransformFVector4(
 						FVector4(
-							(X ? -1.0f : 1.0f),
-							(Y ? -1.0f : 1.0f),
-							(Z ?  0.0f : 1.0f),
+							(vX ? -1.0f : 1.0f),
+							(vY ? -1.0f : 1.0f),
+							(vZ ?  0.0f : 1.0f),
 							1.0f
 							)
 						);
 					const FVector ProjectedVertex = UnprojectedVertex / UnprojectedVertex.W + PreShadowToPreViewTranslation;
-					FrustumVertices[GetCubeVertexIndex(X,Y,Z)] = FVector4(ProjectedVertex, 0);
+					FrustumVertices[GetCubeVertexIndex(vX,vY,vZ)] = FVector4(ProjectedVertex, 0);
 				}
 			}
 		}
@@ -1914,9 +1916,6 @@ void FProjectedShadowInfo::RenderProjection(FRHICommandListImmediate& RHICmdList
 			BottomDistance > DistanceThreshold;
 	}
 
-	// Find the projection shaders.
-	TShaderMapRef<FShadowProjectionVS> VertexShader(View->ShaderMap);
-
 	// Depth test wo/ writes, no color writing.
 	RHICmdList.SetDepthStencilState(TStaticDepthStencilState<false, CF_DepthNearOrEqual>::GetRHI());
 	RHICmdList.SetBlendState(TStaticBlendState<CW_NONE>::GetRHI());
@@ -2662,10 +2661,10 @@ bool FDeferredShadingSceneRenderer::RenderOnePassPointLightShadows(FRHICommandLi
 				SCOPED_DRAW_EVENT(RHICmdList, ShadowDepthsFromOpaquePointLight);
 
 				bool bPerformClear = true;
-				auto SetShadowRenderTargets = [this, &bPerformClear, ProjectedShadowInfo](FRHICommandList& RHICmdList)
+				auto SetShadowRenderTargets = [this, &bPerformClear, ProjectedShadowInfo](FRHICommandList& InRHICmdList)
 				{
-					GSceneRenderTargets.BeginRenderingCubeShadowDepth(RHICmdList, ProjectedShadowInfo->ResolutionX);
-					ProjectedShadowInfo->ClearDepth(RHICmdList, this, bPerformClear);
+					GSceneRenderTargets.BeginRenderingCubeShadowDepth(InRHICmdList, ProjectedShadowInfo->ResolutionX);
+					ProjectedShadowInfo->ClearDepth(InRHICmdList, this, bPerformClear);
 				};
 
 				ProjectedShadowInfo->bAllocated = true;
@@ -3075,10 +3074,10 @@ bool FDeferredShadingSceneRenderer::RenderReflectiveShadowMaps(FRHICommandListIm
 				FLightPropagationVolume* LightPropagationVolume = ViewState->GetLightPropagationVolume();
 
 				check(LightPropagationVolume);
-				auto SetShadowRenderTargets = [this, LightPropagationVolume, ProjectedShadowInfo](FRHICommandList& RHICmdList)
+				auto SetShadowRenderTargets = [this, LightPropagationVolume, ProjectedShadowInfo](FRHICommandList& InRHICmdList)
 				{
-					GSceneRenderTargets.BeginRenderingReflectiveShadowMap(RHICmdList, LightPropagationVolume);
-					ProjectedShadowInfo->ClearDepth(RHICmdList, this, false);
+					GSceneRenderTargets.BeginRenderingReflectiveShadowMap(InRHICmdList, LightPropagationVolume);
+					ProjectedShadowInfo->ClearDepth(InRHICmdList, this, false);
 				};
 				SetShadowRenderTargets(RHICmdList);  // run it now, maybe run it later for parallel command lists				
 
@@ -3252,9 +3251,9 @@ bool FDeferredShadingSceneRenderer::RenderProjectedShadows(FRHICommandListImmedi
 			SCOPED_DRAW_EVENT(RHICmdList, ShadowDepthsFromOpaqueProjected);
 
 			bool bPerformClear = true;
-			auto SetShadowRenderTargets = [&bPerformClear](FRHICommandList& RHICmdList)
+			auto SetShadowRenderTargets = [&bPerformClear](FRHICommandList& InRHICmdList)
 			{
-				GSceneRenderTargets.BeginRenderingShadowDepth(RHICmdList, bPerformClear);
+				GSceneRenderTargets.BeginRenderingShadowDepth(InRHICmdList, bPerformClear);
 			};
 
 			SetShadowRenderTargets(RHICmdList);  // run it now, maybe run it later for parallel command lists
@@ -3417,10 +3416,10 @@ bool FDeferredShadingSceneRenderer::RenderCachedPreshadows(FRHICommandListImmedi
 					
 					//
 					bool bPerformClear = true;;
-					auto SetShadowRenderTargets = [this, &bPerformClear, ProjectedShadowInfo](FRHICommandList& RHICmdList)
+					auto SetShadowRenderTargets = [this, &bPerformClear, ProjectedShadowInfo](FRHICommandList& InRHICmdList)
 					{
-						SetRenderTarget(RHICmdList, FTextureRHIRef(), GSceneRenderTargets.PreShadowCacheDepthZ->GetRenderTargetItem().TargetableTexture);
-						ProjectedShadowInfo->ClearDepth(RHICmdList, this, bPerformClear);
+						SetRenderTarget(InRHICmdList, FTextureRHIRef(), GSceneRenderTargets.PreShadowCacheDepthZ->GetRenderTargetItem().TargetableTexture);
+						ProjectedShadowInfo->ClearDepth(InRHICmdList, this, bPerformClear);
 					};					
 					SetShadowRenderTargets(RHICmdList); // run it now, maybe run it later for parallel command lists
 					bPerformClear = false;
@@ -3530,9 +3529,9 @@ bool FForwardShadingSceneRenderer::RenderShadowDepthMap(FRHICommandListImmediate
 		SCOPED_DRAW_EVENT(RHICmdList, ShadowDepthsFromOpaqueForward);
 
 		bool bFirst = true;
-		auto SetShadowRenderTargets = [&bFirst](FRHICommandList& RHICmdList)
+		auto SetShadowRenderTargets = [&bFirst](FRHICommandList& InRHICmdList)
 		{
-			GSceneRenderTargets.BeginRenderingShadowDepth(RHICmdList, bFirst);
+			GSceneRenderTargets.BeginRenderingShadowDepth(InRHICmdList, bFirst);
 		};
 
 		SetShadowRenderTargets(RHICmdList);  // run it now, maybe run it later for parallel command lists
diff --git a/Engine/Source/Runtime/Renderer/Private/ShadowSetup.cpp b/Engine/Source/Runtime/Renderer/Private/ShadowSetup.cpp
--- a/Engine/Source/Runtime/Renderer/Private/ShadowSetup.cpp
+++ b/Engine/Source/Runtime/Renderer/Private/ShadowSetup.cpp
@@ -1943,7 +1943,6 @@ inline void FSceneRenderer::GatherShadowsForPrimitiveInner(
 					continue;
 				}
 
-				FPrimitiveSceneProxy* PrimitiveProxy = PrimitiveSceneInfoCompact.Proxy;
 				FLightSceneProxy* RESTRICT LightProxy = ProjectedShadowInfo->GetLightSceneInfo().Proxy;
 
 				const FVector LightDirection = LightProxy->GetDirection();
diff --git a/Engine/Source/Runtime/Slate/Private/Widgets/Views/SHeaderRow.cpp b/Engine/Source/Runtime/Slate/Private/Widgets/Views/SHeaderRow.cpp
--- a/Engine/Source/Runtime/Slate/Private/Widgets/Views/SHeaderRow.cpp
+++ b/Engine/Source/Runtime/Slate/Private/Widgets/Views/SHeaderRow.cpp
@@ -612,18 +612,18 @@ void SHeaderRow::RegenerateWidgets()
 					TWeakPtr<SBorder> WeakSizingGrip = SizingGrip;
 					auto SizingGrip_OnMouseButtonDown = [&SomeColumn, WeakSizingGrip](const FGeometry&, const FPointerEvent&) -> FReply
 					{
-						TSharedPtr<SBorder> SizingGrip = WeakSizingGrip.Pin();
-						if ( SizingGrip.IsValid() )
+						TSharedPtr<SBorder> SizingGripPtr = WeakSizingGrip.Pin();
+						if ( SizingGripPtr.IsValid() )
 						{
-							return FReply::Handled().CaptureMouse(SizingGrip.ToSharedRef());
+							return FReply::Handled().CaptureMouse(SizingGripPtr.ToSharedRef());
 						}
 						return FReply::Unhandled();
 					};
 
 					auto SizingGrip_OnMouseButtonUp = [&SomeColumn, WeakSizingGrip](const FGeometry&, const FPointerEvent&) -> FReply
 					{
-						TSharedPtr<SBorder> SizingGrip = WeakSizingGrip.Pin();
-						if ( SizingGrip.IsValid() && SizingGrip->HasMouseCapture() )
+						TSharedPtr<SBorder> SizingGripPtr = WeakSizingGrip.Pin();
+						if ( SizingGripPtr.IsValid() && SizingGripPtr->HasMouseCapture() )
 						{
 							return FReply::Handled().ReleaseMouseCapture();
 						}
@@ -632,8 +632,8 @@ void SHeaderRow::RegenerateWidgets()
 
 					auto SizingGrip_OnMouseMove = [&SomeColumn, WeakSizingGrip](const FGeometry&, const FPointerEvent& InPointerEvent) -> FReply
 					{
-						TSharedPtr<SBorder> SizingGrip = WeakSizingGrip.Pin();
-						if ( SizingGrip.IsValid() && SizingGrip->HasMouseCapture() )
+						TSharedPtr<SBorder> SizingGripPtr = WeakSizingGrip.Pin();
+						if ( SizingGripPtr.IsValid() && SizingGripPtr->HasMouseCapture() )
 						{
 							// The sizing grip has been moved, so update our columns size from the movement delta
 							const float NewWidth = SomeColumn.GetWidth() + InPointerEvent.GetCursorDelta().X;
diff --git a/Engine/Source/Runtime/SlateCore/Public/Stats/SlateStats.h b/Engine/Source/Runtime/SlateCore/Public/Stats/SlateStats.h
--- a/Engine/Source/Runtime/SlateCore/Public/Stats/SlateStats.h
+++ b/Engine/Source/Runtime/SlateCore/Public/Stats/SlateStats.h
@@ -501,27 +501,27 @@ public:
 
 /** Track a Slate Stat, also adding it to the hierarchy if we are capturing hierarchical stats this frame. */
 #define SLATE_CYCLE_COUNTER_SCOPE(StatCounterWes) \
-	FSlateStatCycleCounterScopeHierarchical<(SLATE_STATS_DETAIL_LEVEL != SLATE_STATS_DETAIL_LEVEL_OFF)> TrackStat_##StatCounterWes(StatCounterWes)
+	FSlateStatCycleCounterScopeHierarchical<(SLATE_STATS_DETAIL_LEVEL != SLATE_STATS_DETAIL_LEVEL_OFF)> PREPROCESSOR_JOIN(TrackStat,__LINE__)(StatCounterWes)
 
 /** Track a Slate Stat, also adding it to the hierarchy if we are capturing hierarchical stats this frame. With a given detail level so it can be compiled out. */
 #define SLATE_CYCLE_COUNTER_SCOPE_DETAILED(DetailLevel, StatCounterWes) \
-	FSlateStatCycleCounterScopeHierarchical<(DetailLevel <= SLATE_STATS_DETAIL_LEVEL) && (DetailLevel != SLATE_STATS_DETAIL_LEVEL_OFF)> TrackStat_##StatCounterWes(StatCounterWes)
+	FSlateStatCycleCounterScopeHierarchical<(DetailLevel <= SLATE_STATS_DETAIL_LEVEL) && (DetailLevel != SLATE_STATS_DETAIL_LEVEL_OFF)> PREPROCESSOR_JOIN(TrackStat,__LINE__)(StatCounterWes)
 
 /** Track a Slate Stat, also adding it to the hierarchy if we are capturing hierarchical stats this frame, along with custom context. */
 #define SLATE_CYCLE_COUNTER_SCOPE_CUSTOM(StatCounterWes,CustomName) \
-	FSlateStatCycleCounterScopeHierarchical<(SLATE_STATS_DETAIL_LEVEL != SLATE_STATS_DETAIL_LEVEL_OFF)> TrackStat_##StatCounterWes(StatCounterWes, CustomName)
+	FSlateStatCycleCounterScopeHierarchical<(SLATE_STATS_DETAIL_LEVEL != SLATE_STATS_DETAIL_LEVEL_OFF)> PREPROCESSOR_JOIN(TrackStat,__LINE__)(StatCounterWes, CustomName)
 
 /** Track a Slate Stat, also adding it to the hierarchy if we are capturing hierarchical stats this frame, along with custom context. With a given detail level so it can be compiled out. */
 #define SLATE_CYCLE_COUNTER_SCOPE_CUSTOM_DETAILED(DetailLevel, StatCounterWes,CustomName) \
-	FSlateStatCycleCounterScopeHierarchical<(DetailLevel <= SLATE_STATS_DETAIL_LEVEL) && (DetailLevel != SLATE_STATS_DETAIL_LEVEL_OFF)> TrackStat_##StatCounterWes(StatCounterWes, CustomName)
+	FSlateStatCycleCounterScopeHierarchical<(DetailLevel <= SLATE_STATS_DETAIL_LEVEL) && (DetailLevel != SLATE_STATS_DETAIL_LEVEL_OFF)> PREPROCESSOR_JOIN(TrackStat,__LINE__)(StatCounterWes, CustomName)
 
 /** Track a Slate Stat, only for flattened display. Slightly cheaper than checking dynamically. It doesn't make sense to have a custom version as flattened stats don't gather custom data right now. */
 #define SLATE_CYCLE_COUNTER_SCOPE_FLAT(StatCounterWes) \
-	FSlateStatCycleCounterScopeFlat<(SLATE_STATS_DETAIL_LEVEL != SLATE_STATS_DETAIL_LEVEL_OFF)> TrackStat_##StatCounterWes(StatCounterWes)
+	FSlateStatCycleCounterScopeFlat<(SLATE_STATS_DETAIL_LEVEL != SLATE_STATS_DETAIL_LEVEL_OFF)> PREPROCESSOR_JOIN(TrackStat,__LINE__)(StatCounterWes)
 
 /** Track a Slate Stat, only for flattened display. Slightly cheaper than checking dynamically. With a given detail level so it can be compiled out. */
 #define SLATE_CYCLE_COUNTER_SCOPE_FLAT_DETAILED(DetailLevel, StatCounterWes) \
-	FSlateStatCycleCounterScopeFlat<(DetailLevel <= SLATE_STATS_DETAIL_LEVEL) && (DetailLevel != SLATE_STATS_DETAIL_LEVEL_OFF)> TrackStat_##StatCounterWes(StatCounterWes)
+	FSlateStatCycleCounterScopeFlat<(DetailLevel <= SLATE_STATS_DETAIL_LEVEL) && (DetailLevel != SLATE_STATS_DETAIL_LEVEL_OFF)> PREPROCESSOR_JOIN(TrackStat,__LINE__)(StatCounterWes)
 
 #define SLATE_STATS_END_FRAME(CurrentTime) FSlateStatCycleCounter::EndFrame(CurrentTime)
 
