fixed ShaderComplexity for Decals

diff --git a/Engine/Source/Runtime/Renderer/Private/DeferredShadingRenderer.cpp b/Engine/Source/Runtime/Renderer/Private/DeferredShadingRenderer.cpp
--- a/Engine/Source/Runtime/Renderer/Private/DeferredShadingRenderer.cpp
+++ b/Engine/Source/Runtime/Renderer/Private/DeferredShadingRenderer.cpp
@@ -944,6 +944,22 @@ void FDeferredShadingSceneRenderer::Render(FRHICommandListImmediate& RHICmdList)
 		RHICmdList.SetCurrentStat(GET_STATID(STAT_CLM_AfterVelocity));
 	}
 
+	// Pre-lighting composition lighting stage
+	// e.g. deferred decals
+	if (FeatureLevel >= ERHIFeatureLevel::SM4
+		&& bGBuffer)
+	{
+		QUICK_SCOPE_CYCLE_COUNTER(STAT_FDeferredShadingSceneRenderer_AfterBasePass);
+
+		GRenderTargetPool.AddPhaseEvent(TEXT("AfterBasePass"));
+
+		for (int32 ViewIndex = 0; ViewIndex < Views.Num(); ViewIndex++)
+		{
+			SCOPED_CONDITIONAL_DRAW_EVENTF(RHICmdList, EventView, Views.Num() > 1, TEXT("View%d"), ViewIndex);
+			GCompositionLighting.ProcessAfterBasePass(RHICmdList, Views[ViewIndex]);
+		}
+	}
+
 	// Render lighting.
 	if (ViewFamily.EngineShowFlags.Lighting
 		&& FeatureLevel >= ERHIFeatureLevel::SM4
@@ -955,14 +971,6 @@ void FDeferredShadingSceneRenderer::Render(FRHICommandListImmediate& RHICmdList)
 
 		GRenderTargetPool.AddPhaseEvent(TEXT("Lighting"));
 
-		// Pre-lighting composition lighting stage
-		// e.g. deferred decals, blurred GBuffer
-		for(int32 ViewIndex = 0;ViewIndex < Views.Num();ViewIndex++)
-		{	
-			SCOPED_CONDITIONAL_DRAW_EVENTF(RHICmdList, EventView,Views.Num() > 1, TEXT("View%d"), ViewIndex);
-			GCompositionLighting.ProcessAfterBasePass(RHICmdList, Views[ViewIndex]);
-		}
-
 		// Clear the translucent lighting volumes before we accumulate
 		ClearTranslucentVolumeLighting(RHICmdList);
 
