renamed some class members to be more clear

diff --git a/Engine/Source/Runtime/Engine/Classes/Materials/MaterialInterface.h b/Engine/Source/Runtime/Engine/Classes/Materials/MaterialInterface.h
--- a/Engine/Source/Runtime/Engine/Classes/Materials/MaterialInterface.h
+++ b/Engine/Source/Runtime/Engine/Classes/Materials/MaterialInterface.h
@@ -58,7 +58,7 @@ struct ENGINE_API FMaterialRelevance
 	uint32 bOutputsVelocityInBasePass : 1;
 
 	UPROPERTY()
-	uint16 LightingProfileMask;
+	uint16 ShadingModelMask;
 
 	/** Default constructor. */
 	FMaterialRelevance()
@@ -69,7 +69,7 @@ struct ENGINE_API FMaterialRelevance
 		, bNormalTranslucency(false)
 		, bDisableDepthTest(false)		
 		, bOutputsVelocityInBasePass(true)
-		, LightingProfileMask(0)
+		, ShadingModelMask(0)
 	{}
 
 	/** Bitwise OR operator.  Sets any relevance bits which are present in either FMaterialRelevance. */
@@ -81,7 +81,7 @@ struct ENGINE_API FMaterialRelevance
 		bSeparateTranslucency |= B.bSeparateTranslucency;
 		bNormalTranslucency |= B.bNormalTranslucency;
 		bDisableDepthTest |= B.bDisableDepthTest;
-		LightingProfileMask |= B.LightingProfileMask;
+		ShadingModelMask |= B.ShadingModelMask;
 		bOutputsVelocityInBasePass |= B.bOutputsVelocityInBasePass;
 		return *this;
 	}
diff --git a/Engine/Source/Runtime/Engine/Private/Materials/MaterialInterface.cpp b/Engine/Source/Runtime/Engine/Private/Materials/MaterialInterface.cpp
--- a/Engine/Source/Runtime/Engine/Private/Materials/MaterialInterface.cpp
+++ b/Engine/Source/Runtime/Engine/Private/Materials/MaterialInterface.cpp
@@ -17,7 +17,7 @@ void FMaterialRelevance::SetPrimitiveViewRelevance(FPrimitiveViewRelevance& OutV
 	OutViewRelevance.bDistortionRelevance = bDistortion;
 	OutViewRelevance.bSeparateTranslucencyRelevance = bSeparateTranslucency;
 	OutViewRelevance.bNormalTranslucencyRelevance = bNormalTranslucency;
-	OutViewRelevance.LightingProfileRelevanceMask = LightingProfileMask;
+	OutViewRelevance.ShadingModelMaskRelevance = ShadingModelMask;
 }
 
 //////////////////////////////////////////////////////////////////////////
@@ -50,7 +50,7 @@ FMaterialRelevance UMaterialInterface::GetRelevance_Internal(const UMaterial* Ma
 		const bool bIsLit =  ShadingModel != MSM_Unlit;
 		// Determine the material's view relevance.
 		FMaterialRelevance MaterialRelevance;
-		MaterialRelevance.LightingProfileMask = 1 << ShadingModel;
+		MaterialRelevance.ShadingModelMask = 1 << ShadingModel;
 		MaterialRelevance.bOpaque = !bIsTranslucent;
 		MaterialRelevance.bMasked = IsMasked();
 		MaterialRelevance.bDistortion = Material->bUsesDistortion && bIsTranslucent;
diff --git a/Engine/Source/Runtime/Engine/Public/PrimitiveViewRelevance.h b/Engine/Source/Runtime/Engine/Public/PrimitiveViewRelevance.h
--- a/Engine/Source/Runtime/Engine/Public/PrimitiveViewRelevance.h
+++ b/Engine/Source/Runtime/Engine/Public/PrimitiveViewRelevance.h
@@ -9,7 +9,7 @@
 struct FPrimitiveViewRelevance
 {
 	/* The LightingProfile supported by this primitive, as a bitmask. */
-	uint16 LightingProfileRelevanceMask;
+	uint16 ShadingModelMaskRelevance;
 	/** The primitive's static elements are rendered for the view. */
 	uint32 bStaticRelevance : 1; 
 	/** The primitive's dynamic elements are rendered for the view. */
@@ -56,7 +56,7 @@ struct FPrimitiveViewRelevance
 
 	/** Initialization constructor. */
 	FPrimitiveViewRelevance():
-		LightingProfileRelevanceMask(0),
+		ShadingModelMaskRelevance(0),
 		bStaticRelevance(false),
 		bDynamicRelevance(false),
 		bDrawRelevance(false),
@@ -78,7 +78,7 @@ struct FPrimitiveViewRelevance
 	/** Bitwise OR operator.  Sets any relevance bits which are present in either FPrimitiveViewRelevance. */
 	FPrimitiveViewRelevance& operator|=(const FPrimitiveViewRelevance& B)
 	{
-		LightingProfileRelevanceMask |= B.LightingProfileRelevanceMask;
+		ShadingModelMaskRelevance |= B.ShadingModelMaskRelevance;
 		bStaticRelevance |= B.bStaticRelevance != 0;
 		bDynamicRelevance |= B.bDynamicRelevance != 0;
 		bDrawRelevance |= B.bDrawRelevance != 0;
diff --git a/Engine/Source/Runtime/Renderer/Private/CompositionLighting/CompositionLighting.cpp b/Engine/Source/Runtime/Renderer/Private/CompositionLighting/CompositionLighting.cpp
--- a/Engine/Source/Runtime/Renderer/Private/CompositionLighting/CompositionLighting.cpp
+++ b/Engine/Source/Runtime/Renderer/Private/CompositionLighting/CompositionLighting.cpp
@@ -357,7 +357,7 @@ void FCompositionLighting::ProcessLighting(FRHICommandListImmediate& RHICmdList,
 
 			bool bSimpleDynamicLighting = IsSimpleDynamicLightingEnabled();
 
-			bool bScreenSpaceSubsurfacePassNeeded = (View.LightingProfilesActiveInView & (1 << MSM_SubsurfaceProfile)) != 0;
+			bool bScreenSpaceSubsurfacePassNeeded = (View.ShadingModelMaskInView & (1 << MSM_SubsurfaceProfile)) != 0;
 			if (bScreenSpaceSubsurfacePassNeeded && Radius > 0 && !bSimpleDynamicLighting && View.Family->EngineShowFlags.SubsurfaceScattering &&
 				//@todo-rco: Remove this when we fix the cross-compiler
 				!IsOpenGLPlatform(View.GetShaderPlatform()))
diff --git a/Engine/Source/Runtime/Renderer/Private/CompositionLighting/PostProcessAmbient.cpp b/Engine/Source/Runtime/Renderer/Private/CompositionLighting/PostProcessAmbient.cpp
--- a/Engine/Source/Runtime/Renderer/Private/CompositionLighting/PostProcessAmbient.cpp
+++ b/Engine/Source/Runtime/Renderer/Private/CompositionLighting/PostProcessAmbient.cpp
@@ -136,7 +136,7 @@ void FRCPassPostProcessAmbient::Process(FRenderingCompositePassContext& Context)
 	// Ambient cubemap specular will be applied in the reflection environment pass if it is enabled
 	const bool bApplySpecular = View.Family->EngineShowFlags.ReflectionEnvironment == 0 || NumReflectionCaptures == 0;
 
-	bool bClearCoatNeeded = (View.LightingProfilesActiveInView & (1 << MSM_ClearCoat)) != 0;
+	bool bClearCoatNeeded = (View.ShadingModelMaskInView & (1 << MSM_ClearCoat)) != 0;
 	if (bClearCoatNeeded)
 	{
 		Render<true>(Context);
diff --git a/Engine/Source/Runtime/Renderer/Private/LightRendering.cpp b/Engine/Source/Runtime/Renderer/Private/LightRendering.cpp
--- a/Engine/Source/Runtime/Renderer/Private/LightRendering.cpp
+++ b/Engine/Source/Runtime/Renderer/Private/LightRendering.cpp
@@ -799,7 +799,7 @@ void FDeferredShadingSceneRenderer::RenderLight(FRHICommandList& RHICmdList, con
 		// Set the device viewport for the view.
 		RHICmdList.SetViewport(View.ViewRect.Min.X, View.ViewRect.Min.Y, 0.0f, View.ViewRect.Max.X, View.ViewRect.Max.Y, 1.0f);
 
-		bool bClearCoatNeeded = (View.LightingProfilesActiveInView & (1 << MSM_ClearCoat)) != 0;
+		bool bClearCoatNeeded = (View.ShadingModelMaskInView & (1 << MSM_ClearCoat)) != 0;
 		if (LightSceneInfo->Proxy->GetLightType() == LightType_Directional)
 		{
 			TShaderMapRef<TDeferredLightVS<false> > VertexShader(View.ShaderMap);
@@ -990,7 +990,7 @@ void FDeferredShadingSceneRenderer::RenderSimpleLightsStandardDeferred(FRHIComma
 
 			SetBoundingGeometryRasterizerAndDepthState(RHICmdList, View, LightBounds);
 
-			bool bClearCoatNeeded = (View.LightingProfilesActiveInView & (1 << MSM_ClearCoat)) != 0;
+			bool bClearCoatNeeded = (View.ShadingModelMaskInView & (1 << MSM_ClearCoat)) != 0;
 			if (SimpleLight.Exponent == 0)
 			{
 				// inverse squared
diff --git a/Engine/Source/Runtime/Renderer/Private/ReflectionEnvironment.cpp b/Engine/Source/Runtime/Renderer/Private/ReflectionEnvironment.cpp
--- a/Engine/Source/Runtime/Renderer/Private/ReflectionEnvironment.cpp
+++ b/Engine/Source/Runtime/Renderer/Private/ReflectionEnvironment.cpp
@@ -877,7 +877,7 @@ void FDeferredShadingSceneRenderer::RenderTiledDeferredImageBasedReflections(FRH
 
 			bool bHasBoxCaptures = (NumBoxCaptures > 0);
 			bool bHasSphereCaptures = (NumSphereCaptures > 0);
-			bool bNeedsClearCoat = (View.LightingProfilesActiveInView & (1 << MSM_ClearCoat)) != 0;
+			bool bNeedsClearCoat = (View.ShadingModelMaskInView & (1 << MSM_ClearCoat)) != 0;
 			ComputeShader = SelectReflectionEnvironmentTiledDeferredCS(View.ShaderMap, bUseLightmaps, bNeedsClearCoat, bHasBoxCaptures, bHasSphereCaptures);
 
 			RHICmdList.SetComputeShader(ComputeShader->GetComputeShader());
diff --git a/Engine/Source/Runtime/Renderer/Private/SceneRendering.cpp b/Engine/Source/Runtime/Renderer/Private/SceneRendering.cpp
--- a/Engine/Source/Runtime/Renderer/Private/SceneRendering.cpp
+++ b/Engine/Source/Runtime/Renderer/Private/SceneRendering.cpp
@@ -211,7 +211,7 @@ void FViewInfo::Init()
 	bIgnoreExistingQueries = false;
 	bDisableQuerySubmissions = false;
 	bDisableDistanceBasedFadeTransitions = false;	
-	LightingProfilesActiveInView = 0;
+	ShadingModelMaskInView = 0;
 	NumVisibleStaticMeshElements = 0;
 	PrecomputedVisibilityData = 0;
 
diff --git a/Engine/Source/Runtime/Renderer/Private/SceneRendering.h b/Engine/Source/Runtime/Renderer/Private/SceneRendering.h
--- a/Engine/Source/Runtime/Renderer/Private/SceneRendering.h
+++ b/Engine/Source/Runtime/Renderer/Private/SceneRendering.h
@@ -442,8 +442,8 @@ public:
 	uint32 bDisableQuerySubmissions : 1;
 	/** Whether we should disable distance-based fade transitions for this frame (usually after a large camera movement.) */
 	uint32 bDisableDistanceBasedFadeTransitions : 1;
-	/** Bitmask of all lighting profiles used by primitives in this view */
-	uint16 LightingProfilesActiveInView;
+	/** Bitmask of all shading models used by primitives in this view */
+	uint16 ShadingModelMaskInView;
 
 	FViewMatrices PrevViewMatrices;
 
diff --git a/Engine/Source/Runtime/Renderer/Private/SceneVisibility.cpp b/Engine/Source/Runtime/Renderer/Private/SceneVisibility.cpp
--- a/Engine/Source/Runtime/Renderer/Private/SceneVisibility.cpp
+++ b/Engine/Source/Runtime/Renderer/Private/SceneVisibility.cpp
@@ -795,7 +795,7 @@ static void ComputeRelevanceForView(
 			}
 		}
 		
-		View.LightingProfilesActiveInView |= ViewRelevance.LightingProfileRelevanceMask;
+		View.ShadingModelMaskInView |= ViewRelevance.ShadingModelMaskRelevance;
 		
 		if (ViewRelevance.bRenderCustomDepth)
 		{
@@ -934,7 +934,7 @@ struct FRelevancePacket
 	FRelevancePrimSet<FPrimitiveSceneProxy*> CustomDepthSet;
 	FRelevancePrimSet<FPrimitiveSceneInfo*> UpdateStaticMeshes;
 	FRelevancePrimSet<FPrimitiveSceneInfo*> VisibleEditorPrimitives;
-	uint16 CombinedLightingProfileMask;
+	uint16 CombinedShadingModelMask;
 
 	FRelevancePacket(
 		FRHICommandListImmediate& InRHICmdList,
@@ -956,7 +956,7 @@ struct FRelevancePacket
 		, OutHasDynamicMeshElementsMasks(InOutHasDynamicMeshElementsMasks)
 		, OutHasDynamicEditorMeshElementsMasks(InOutHasDynamicEditorMeshElementsMasks)
 		, MarkMasks(InMarkMasks)
-		, CombinedLightingProfileMask(0)
+		, CombinedShadingModelMask(0)
 	{
 	}
 
@@ -968,7 +968,7 @@ struct FRelevancePacket
 
 	void ComputeRelevance()
 	{
-		CombinedLightingProfileMask = 0;
+		CombinedShadingModelMask = 0;
 		SCOPE_CYCLE_COUNTER(STAT_ComputeViewRelevance);
 		for (int32 Index = 0; Index < Input.NumPrims; Index++)
 		{
@@ -1028,7 +1028,7 @@ struct FRelevancePacket
 				}
 			}
 
-			CombinedLightingProfileMask |= ViewRelevance.LightingProfileRelevanceMask;			
+			CombinedShadingModelMask |= ViewRelevance.ShadingModelMaskRelevance;			
 
 			if (ViewRelevance.bRenderCustomDepth)
 			{
@@ -1152,7 +1152,7 @@ struct FRelevancePacket
 		{
 			WriteView.PrimitiveVisibilityMap[NotDrawRelevant.Prims[Index]] = false;
 		}
-		WriteView.LightingProfilesActiveInView |= CombinedLightingProfileMask;
+		WriteView.ShadingModelMaskInView |= CombinedShadingModelMask;
 		VisibleEditorPrimitives.AppendTo(WriteView.VisibleEditorPrimitives);
 		VisibleDynamicPrimitives.AppendTo(WriteView.VisibleDynamicPrimitives);
 		WriteView.TranslucentPrimSet.AppendScenePrimitives(SortedTranslucencyPrims.Prims, SortedTranslucencyPrims.NumPrims, SortedSeparateTranslucencyPrims.Prims, SortedSeparateTranslucencyPrims.NumPrims);
