fixed ShaderComplexity for DBuffer decals

diff --git a/Engine/Source/Runtime/Renderer/Private/CompositionLighting/CompositionLighting.cpp b/Engine/Source/Runtime/Renderer/Private/CompositionLighting/CompositionLighting.cpp
--- a/Engine/Source/Runtime/Renderer/Private/CompositionLighting/CompositionLighting.cpp
+++ b/Engine/Source/Runtime/Renderer/Private/CompositionLighting/CompositionLighting.cpp
@@ -251,7 +251,9 @@ void FCompositionLighting::ProcessBeforeBasePass(FRHICommandListImmediate& RHICm
 		// Add the passes we want to add to the graph (commenting a line means the pass is not inserted into the graph) ----------
 
 		// decals are before AmbientOcclusion so the decal can output a normal that AO is affected by
-		if(Context.View.Family->EngineShowFlags.Decals && IsDBufferEnabled()) 
+		if (Context.View.Family->EngineShowFlags.Decals &&
+			!Context.View.Family->EngineShowFlags.ShaderComplexity &&
+			IsDBufferEnabled()) 
 		{
 			AddDeferredDecalsBeforeBasePass(Context);
 		}
diff --git a/Engine/Source/Runtime/Renderer/Private/CompositionLighting/PostProcessDeferredDecals.cpp b/Engine/Source/Runtime/Renderer/Private/CompositionLighting/PostProcessDeferredDecals.cpp
--- a/Engine/Source/Runtime/Renderer/Private/CompositionLighting/PostProcessDeferredDecals.cpp
+++ b/Engine/Source/Runtime/Renderer/Private/CompositionLighting/PostProcessDeferredDecals.cpp
@@ -12,6 +12,9 @@
 #include "ScreenRendering.h"
 #include "SceneUtils.h"
 
+#define DBUFFER_DONT_USE_STENCIL_YET 1
+
+
 static TAutoConsoleVariable<float> CVarStencilSizeThreshold(
 	TEXT("r.Decal.StencilSizeThreshold"),
 	0.1f,
@@ -722,7 +725,7 @@ void FRCPassPostProcessDeferredDecals::Process(FRenderingCompositePassContext& C
 
 	const bool bShaderComplexity = Context.View.Family->EngineShowFlags.ShaderComplexity;
 	const bool bDBuffer = IsDBufferEnabled();
-	const bool bDecalPreStencil = CVarStencilSizeThreshold.GetValueOnRenderThread() >= 0;
+	const bool bStencilSizeThreshold = CVarStencilSizeThreshold.GetValueOnRenderThread() >= 0;
 
 	SCOPED_DRAW_EVENT(RHICmdList, PostProcessDeferredDecals);
 
@@ -828,7 +831,7 @@ void FRCPassPostProcessDeferredDecals::Process(FRenderingCompositePassContext& C
 			uint32 DecalRenderStage = ComputeRenderStage(Data.DecalBlendMode);
 
 			// we could do this test earlier to avoid the decal intersection but getting DecalBlendMode also costs
-			if(RenderStage == DecalRenderStage)
+			if (Context.View.Family->EngineShowFlags.ShaderComplexity || RenderStage == DecalRenderStage)
 			{
 				SortedDecals.Add(Data);
 			}
@@ -840,8 +843,14 @@ void FRCPassPostProcessDeferredDecals::Process(FRenderingCompositePassContext& C
 		FIntRect SrcRect = View.ViewRect;
 		FIntRect DestRect = View.ViewRect;
 
-		// later we also optimize RenderStage == 0 but we would need to output different stencil depending on Stencil Mask
-		bool bStencilDecals = (RenderStage == 1);
+		bool bStencilDecals = true;
+
+#if DBUFFER_DONT_USE_STENCIL_YET
+		if(RenderStage == 0)
+		{
+			bStencilDecals = false;
+		}
+#endif
 
 		// Setup a stencil mask to prevent certain pixels from receiving deferred decals
 		if(bStencilDecals)
@@ -904,7 +913,7 @@ void FRCPassPostProcessDeferredDecals::Process(FRenderingCompositePassContext& C
 	
 		FTextureRHIParamRef TargetsToResolve[ResolveBufferMax] = { nullptr };
 
-		for (int32 DecalIndex = 0; DecalIndex < SortedDecals.Num(); DecalIndex++)
+		for (int32 DecalIndex = 0, DecalCount = SortedDecals.Num(); DecalIndex < DecalCount; DecalIndex++)
 		{
 			const FTransientDecalRenderData& DecalData = SortedDecals[DecalIndex];
 			const FDeferredDecalProxy& DecalProxy = *DecalData.DecalProxy;
@@ -925,6 +934,15 @@ void FRCPassPostProcessDeferredDecals::Process(FRenderingCompositePassContext& C
 
 			EDecalBlendMode DecalBlendMode = DecalData.DecalBlendMode;
 
+			bool bStencilThisDecal = bStencilDecals;
+			
+#if DBUFFER_DONT_USE_STENCIL_YET
+			if(ComputeRenderStage(DecalBlendMode) == 0)
+			{
+				bStencilThisDecal = false;
+			}
+#endif				
+
 			ERenderTargetMode CurrentRenderTargetMode = ComputeRenderTargetMode(DecalBlendMode);
 
 			if(bShaderComplexity)
@@ -983,9 +1001,9 @@ void FRCPassPostProcessDeferredDecals::Process(FRenderingCompositePassContext& C
 
 			bool bThisDecalUsesStencil = false;
 
-			if (bStencilDecals)
+			if (bStencilThisDecal)
 			{
-				if (bDecalPreStencil)
+				if (bStencilSizeThreshold)
 				{
 					// note this is after a SetStreamSource (in if CurrentRenderTargetMode != LastRenderTargetMode) call as it needs to get the VB input
 					bThisDecalUsesStencil = RenderPreStencil(Context, MaterialShaderMap, ComponentToWorldMatrix, FrustumComponentToClip);
@@ -995,15 +1013,15 @@ void FRCPassPostProcessDeferredDecals::Process(FRenderingCompositePassContext& C
 				}
 			}
 
-			const bool bBlendStateChange = DecalData.DecalBlendMode != LastDecalBlendMode;// Has decal mode changed.
+			const bool bBlendStateChange = DecalBlendMode != LastDecalBlendMode;// Has decal mode changed.
 			const bool bDecalNormalChanged = GSupportsSeparateRenderTargetBlendState && // has normal changed for SM5 stain/translucent decals?
-							(DecalData.DecalBlendMode == DBM_Translucent || DecalData.DecalBlendMode == DBM_Stain) &&
+							(DecalBlendMode == DBM_Translucent || DecalBlendMode == DBM_Stain) &&
 							(int32)DecalData.bHasNormal != LastDecalHasNormal;
 
 			// fewer blend state changes if possible
 			if (bBlendStateChange || bDecalNormalChanged)
 			{
-				LastDecalBlendMode = DecalData.DecalBlendMode;
+				LastDecalBlendMode = DecalBlendMode;
 				LastDecalHasNormal = (int32)DecalData.bHasNormal;
 
 				SetDecalBlendState(RHICmdList, SMFeatureLevel, RenderStage, (EDecalBlendMode)LastDecalBlendMode, DecalData.bHasNormal);
