Fix crash when 'r.MotionBlurNew' is active and FeatureLevel < ERHIFeatureLevel::SM5

diff --git a/Engine/Source/Runtime/Renderer/Private/PostProcess/PostProcessing.cpp b/Engine/Source/Runtime/Renderer/Private/PostProcess/PostProcessing.cpp
--- a/Engine/Source/Runtime/Renderer/Private/PostProcess/PostProcessing.cpp
+++ b/Engine/Source/Runtime/Renderer/Private/PostProcess/PostProcessing.cpp
@@ -1131,7 +1131,7 @@ void FPostProcessing::Process(FRHICommandListImmediate& RHICmdList, FViewInfo& V
 			{
 				// Motion blur
 
-				if( CVarMotionBlurNew.GetValueOnRenderThread() )
+				if( CVarMotionBlurNew.GetValueOnRenderThread() && FeatureLevel >= ERHIFeatureLevel::SM5 )
 				{
 					FRenderingCompositeOutputRef MaxTileVelocity;
 					FRenderingCompositeOutputRef SceneDepth( Context.SceneDepth );
