removed GRenderingRealtimeClock.GetGamePaused(), added FSceneViewFamily.bWorldIsPaused to support multiple worlds and for cleaner code

diff --git a/Engine/Source/Editor/UnrealEd/Private/Editor.cpp b/Engine/Source/Editor/UnrealEd/Private/Editor.cpp
--- a/Engine/Source/Editor/UnrealEd/Private/Editor.cpp
+++ b/Engine/Source/Editor/UnrealEd/Private/Editor.cpp
@@ -1450,16 +1450,6 @@ void UEditorEngine::Tick( float DeltaSeconds, bool bIdleMode )
 	// Commit changes to the BSP model.
 	EditorContext.World()->CommitModelSurfaces();
 	EditorContext.World()->SendAllEndOfFrameUpdates();
-	
-	{
-		// tell renderer about EditorContext.World()->IsPaused(), before rendering
-		ENQUEUE_UNIQUE_RENDER_COMMAND_ONEPARAMETER(
-			SetPaused,
-			bool, bGamePaused, EditorContext.World()->IsPaused(),
-		{
-			GRenderingRealtimeClock.SetGamePaused(bGamePaused);
-		});
-	}
 
 	bool bUpdateLinkedOrthoViewports = false;
 	/////////////////////////////
diff --git a/Engine/Source/Runtime/Engine/Private/GameEngine.cpp b/Engine/Source/Runtime/Engine/Private/GameEngine.cpp
--- a/Engine/Source/Runtime/Engine/Private/GameEngine.cpp
+++ b/Engine/Source/Runtime/Engine/Private/GameEngine.cpp
@@ -930,16 +930,6 @@ void UGameEngine::Tick( float DeltaSeconds, bool bIdleMode )
 			}
 		}
 
-		// tell renderer about GWorld->IsPaused(), before rendering
-		{
-			ENQUEUE_UNIQUE_RENDER_COMMAND_ONEPARAMETER(
-				SetPaused,
-				bool, bGamePaused, Context.World()->IsPaused(),
-			{
-				GRenderingRealtimeClock.SetGamePaused(bGamePaused);
-			});
-		}
-
 		if (!bIdleMode && !IsRunningDedicatedServer() && !IsRunningCommandlet())
 		{
 			// Render everything.
diff --git a/Engine/Source/Runtime/Engine/Private/SceneView.cpp b/Engine/Source/Runtime/Engine/Private/SceneView.cpp
--- a/Engine/Source/Runtime/Engine/Private/SceneView.cpp
+++ b/Engine/Source/Runtime/Engine/Private/SceneView.cpp
@@ -1473,6 +1473,7 @@ FSceneViewFamily::FSceneViewFamily( const ConstructionValues& CVS )
 	bRealtimeUpdate(CVS.bRealtimeUpdate),
 	bDeferClear(CVS.bDeferClear),
 	bResolveScene(CVS.bResolveScene),
+	bWorldIsPaused(false),
 	GammaCorrection(CVS.GammaCorrection)
 {
 	// If we do not pass a valid scene pointer then SetWorldTimes must be called to initialized with valid times.
@@ -1497,11 +1498,21 @@ FSceneViewFamily::FSceneViewFamily( const ConstructionValues& CVS )
 
 	// instead of checking IsGameWorld on rendering thread to see if we allow this flag to be disabled 
 	// we force it on in the game thread.
-	if(IsInGameThread())
+	if(IsInGameThread() && Scene)
 	{
-		if ( Scene && Scene->GetWorld() && Scene->GetWorld()->IsGameWorld() )
+		UWorld* World = Scene->GetWorld();
+
+		if (World)
 		{
-			EngineShowFlags.LOD = 1;
+			if (World->IsGameWorld())
+			{
+				EngineShowFlags.LOD = 1;
+			}
+
+			if (World->IsPaused())
+			{
+				bWorldIsPaused = true;
+			}
 		}
 	}
 
diff --git a/Engine/Source/Runtime/Engine/Public/SceneView.h b/Engine/Source/Runtime/Engine/Public/SceneView.h
--- a/Engine/Source/Runtime/Engine/Public/SceneView.h
+++ b/Engine/Source/Runtime/Engine/Public/SceneView.h
@@ -805,6 +805,9 @@ public:
 	/** if true then results of scene rendering are copied/resolved to the RenderTarget. */
 	bool bResolveScene;
 
+	/** from GetWorld->IsPaused() */
+	bool bWorldIsPaused;
+
 	/** Gamma correction used when rendering this family. Default is 1.0 */
 	float GammaCorrection;
 	
diff --git a/Engine/Source/Runtime/RenderCore/Public/RenderCore.h b/Engine/Source/Runtime/RenderCore/Public/RenderCore.h
--- a/Engine/Source/Runtime/RenderCore/Public/RenderCore.h
+++ b/Engine/Source/Runtime/RenderCore/Public/RenderCore.h
@@ -137,9 +137,8 @@ public:
 	 *	Constructor
 	 **/
 	FTimer()
-	:	CurrentDeltaTime(0.0f)
-	,	CurrentTime(0.0f)
-	,	bGamePaused(false)
+		: CurrentDeltaTime(0.0f)
+		, CurrentTime(0.0f)
 	{
 	}
 
@@ -147,30 +146,20 @@ public:
 	 *	Returns the current time, in seconds.
 	 *	@return Current time, in seconds
 	 */
-	float	GetCurrentTime() const
+	float GetCurrentTime() const
 	{
 		return CurrentTime;
 	}
 
-	bool GetGamePaused() const
-	{
-		return bGamePaused;
-	}
-
 	/**
 	 *	Returns the current delta time.
 	 *	@return Current delta time (number of seconds that passed between the last two tick)
 	 */
-	float	GetCurrentDeltaTime() const
+	float GetCurrentDeltaTime() const
 	{
 		return CurrentDeltaTime;
 	}
 
-	void SetGamePaused(bool bInGamePaused)
-	{
-		bGamePaused = bInGamePaused;
-	}
-
 	/**
 	 *	Updates the timer.
 	 *	@param DeltaTime	Number of seconds that have passed since the last tick
@@ -186,8 +175,6 @@ protected:
 	float CurrentDeltaTime;
 	/** Current time, in seconds. */
 	float CurrentTime;
-
-	bool bGamePaused;
 };
 
 /** Whether to pause the global realtime clock for the rendering thread (read and write only on main thread). */
diff --git a/Engine/Source/Runtime/Renderer/Private/PostProcess/PostProcessTemporalAA.cpp b/Engine/Source/Runtime/Renderer/Private/PostProcess/PostProcessTemporalAA.cpp
--- a/Engine/Source/Runtime/Renderer/Private/PostProcess/PostProcessTemporalAA.cpp
+++ b/Engine/Source/Runtime/Renderer/Private/PostProcess/PostProcessTemporalAA.cpp
@@ -699,7 +699,7 @@ void FRCPassPostProcessTemporalAA::Process(FRenderingCompositePassContext& Conte
 
 	Context.RHICmdList.CopyToResolveTarget(DestRenderTarget.TargetableTexture, DestRenderTarget.ShaderResourceTexture, false, FResolveParams());
 
-	if( !GRenderingRealtimeClock.GetGamePaused() )
+	if( !View.Family->bWorldIsPaused )
 	{
 		ViewState->TemporalAAHistoryRT = PassOutputs[0].PooledRenderTarget;
 		check( ViewState->TemporalAAHistoryRT );
diff --git a/Engine/Source/Runtime/Renderer/Private/ScenePrivate.h b/Engine/Source/Runtime/Renderer/Private/ScenePrivate.h
--- a/Engine/Source/Runtime/Renderer/Private/ScenePrivate.h
+++ b/Engine/Source/Runtime/Renderer/Private/ScenePrivate.h
@@ -552,7 +552,7 @@ public:
 	}
 
 	// @param SampleCount 0 or 1 for no TemporalAA 
-	void SetupTemporalAA(uint32 SampleCount)
+	void SetupTemporalAA(uint32 SampleCount, const FSceneViewFamily& Family)
 	{
 		if(!SampleCount)
 		{
@@ -561,7 +561,7 @@ public:
 
 		TemporalAASampleCount = FMath::Min(SampleCount, (uint32)255);
 		
-		if( !GRenderingRealtimeClock.GetGamePaused() )
+		if (!Family.bWorldIsPaused)
 		{
 			TemporalAASampleIndex++;
 		}
diff --git a/Engine/Source/Runtime/Renderer/Private/SceneVisibility.cpp b/Engine/Source/Runtime/Renderer/Private/SceneVisibility.cpp
--- a/Engine/Source/Runtime/Renderer/Private/SceneVisibility.cpp
+++ b/Engine/Source/Runtime/Renderer/Private/SceneVisibility.cpp
@@ -1678,7 +1678,7 @@ void FSceneRenderer::PreVisibilityFrameSetup(FRHICommandListImmediate& RHICmdLis
 						float SamplesX[] = { -8.0f/16.0f, 0.0/16.0f };
 						float SamplesY[] = { /* - */ 0.0f/16.0f, 8.0/16.0f };
 					#endif
-					ViewState->SetupTemporalAA(ARRAY_COUNT(SamplesX));
+					ViewState->SetupTemporalAA(ARRAY_COUNT(SamplesX), ViewFamily);
 					uint32 Index = ViewState->GetCurrentTemporalAASampleIndex();
 					SampleX = SamplesX[ Index ];
 					SampleY = SamplesY[ Index ];
@@ -1692,7 +1692,7 @@ void FSceneRenderer::PreVisibilityFrameSetup(FRHICommandListImmediate& RHICmdLis
 					// Rolling circle pattern (A,B,C).
 					float SamplesX[] = { -2.0f/3.0f,  2.0/3.0f,  0.0/3.0f };
 					float SamplesY[] = { -2.0f/3.0f,  0.0/3.0f,  2.0/3.0f };
-					ViewState->SetupTemporalAA(ARRAY_COUNT(SamplesX));
+					ViewState->SetupTemporalAA(ARRAY_COUNT(SamplesX), ViewFamily);
 					uint32 Index = ViewState->GetCurrentTemporalAASampleIndex();
 					SampleX = SamplesX[ Index ];
 					SampleY = SamplesY[ Index ];
@@ -1708,7 +1708,7 @@ void FSceneRenderer::PreVisibilityFrameSetup(FRHICommandListImmediate& RHICmdLis
 					// Rolling circle pattern (N,E,S,W).
 					float SamplesX[] = { -2.0f/16.0f,  6.0/16.0f, 2.0/16.0f, -6.0/16.0f };
 					float SamplesY[] = { -6.0f/16.0f, -2.0/16.0f, 6.0/16.0f,  2.0/16.0f };
-					ViewState->SetupTemporalAA(ARRAY_COUNT(SamplesX));
+					ViewState->SetupTemporalAA(ARRAY_COUNT(SamplesX), ViewFamily);
 					uint32 Index = ViewState->GetCurrentTemporalAASampleIndex();
 					SampleX = SamplesX[ Index ];
 					SampleY = SamplesY[ Index ];
@@ -1723,7 +1723,7 @@ void FSceneRenderer::PreVisibilityFrameSetup(FRHICommandListImmediate& RHICmdLis
 					// Rolling circle pattern (N,E,S,W).
 					float SamplesX[] = {  0.0f/2.0f,  1.0/2.0f,  0.0/2.0f, -1.0/2.0f };
 					float SamplesY[] = { -1.0f/2.0f,  0.0/2.0f,  1.0/2.0f,  0.0/2.0f };
-					ViewState->SetupTemporalAA(ARRAY_COUNT(SamplesX));
+					ViewState->SetupTemporalAA(ARRAY_COUNT(SamplesX), ViewFamily);
 					uint32 Index = ViewState->GetCurrentTemporalAASampleIndex();
 					SampleX = SamplesX[ Index ];
 					SampleY = SamplesY[ Index ];
@@ -1731,7 +1731,7 @@ void FSceneRenderer::PreVisibilityFrameSetup(FRHICommandListImmediate& RHICmdLis
 				else if( TemporalAASamples == 8 )
 				{
 					// This works better than various orderings of 8xMSAA.
-					ViewState->SetupTemporalAA(8);
+					ViewState->SetupTemporalAA(8, ViewFamily);
 					uint32 Index = ViewState->GetCurrentTemporalAASampleIndex();
 					SampleX = Halton( Index, 2 ) - 0.5f;
 					SampleY = Halton( Index, 3 ) - 0.5f;
@@ -1739,7 +1739,7 @@ void FSceneRenderer::PreVisibilityFrameSetup(FRHICommandListImmediate& RHICmdLis
 				else
 				{
 					// More than 8 samples can improve quality.
-					ViewState->SetupTemporalAA(TemporalAASamples);
+					ViewState->SetupTemporalAA(TemporalAASamples, ViewFamily);
 					uint32 Index = ViewState->GetCurrentTemporalAASampleIndex();
 					SampleX = Halton( Index, 2 ) - 0.5f;
 					SampleY = Halton( Index, 3 ) - 0.5f;
@@ -1769,7 +1769,7 @@ void FSceneRenderer::PreVisibilityFrameSetup(FRHICommandListImmediate& RHICmdLis
 		else if(ViewState)
 		{
 			// no TemporalAA
-			ViewState->SetupTemporalAA(1);
+			ViewState->SetupTemporalAA(1, ViewFamily);
 		}
 
 		if ( ViewState )
@@ -1825,7 +1825,7 @@ void FSceneRenderer::PreVisibilityFrameSetup(FRHICommandListImmediate& RHICmdLis
 					ViewState->PrevViewMatrices = ViewState->PendingPrevViewMatrices;
 
 					// check for pause so we can keep motion blur in paused mode (doesn't work in editor)
-					if(!GRenderingRealtimeClock.GetGamePaused())
+					if(!ViewFamily.bWorldIsPaused)
 					{
 						// pending is needed as we are in init view and still need to render.
 						ViewState->PendingPrevViewMatrices = View.ViewMatrices;
@@ -1833,7 +1833,7 @@ void FSceneRenderer::PreVisibilityFrameSetup(FRHICommandListImmediate& RHICmdLis
 				}
 				// we don't use DeltaTime as it can be 0 (in editor) and is computed by subtracting floats (loses precision over time)
 				// Clamp DeltaWorldTime to reasonable values for the purposes of motion blur, things like TimeDilation can make it very small
-				if(!GRenderingRealtimeClock.GetGamePaused())
+				if (!ViewFamily.bWorldIsPaused)
 				{
 					ViewState->MotionBlurTimeScale			= bEnableTimeScale ? (1.0f / (FMath::Max(View.Family->DeltaWorldTime, .00833f) * 30.0f)) : 1.0f;
 				}
@@ -2347,7 +2347,7 @@ void FDeferredShadingSceneRenderer::InitViews(FRHICommandListImmediate& RHICmdLi
 		InitDynamicShadows(RHICmdList);
 	}
 
-	if(ViewFamily.EngineShowFlags.MotionBlur && GRenderingRealtimeClock.GetGamePaused())
+	if(ViewFamily.EngineShowFlags.MotionBlur && ViewFamily.bWorldIsPaused)
 	{
 		// so we can keep motion blur in paused mode
 
