Bug UE-6289: Only let the console var r.ClearSceneMethod to not clear the g-buffer in game mode.

diff --git a/Engine/Source/Runtime/Core/Private/HAL/ConsoleManager.cpp b/Engine/Source/Runtime/Core/Private/HAL/ConsoleManager.cpp
--- a/Engine/Source/Runtime/Core/Private/HAL/ConsoleManager.cpp
+++ b/Engine/Source/Runtime/Core/Private/HAL/ConsoleManager.cpp
@@ -1617,7 +1617,7 @@ static TAutoConsoleVariable<int32> CVarMobileReduceLoadedMips(
 static TAutoConsoleVariable<int32> CVarSetClearSceneMethod(
 	TEXT("r.ClearSceneMethod"),
 	1,
-	TEXT("Select how scene rendertarget clears are handled\n")
+	TEXT("Select how the g-buffer is cleared in game mode (only affects deferred shading).\n")
 	TEXT(" 0: No clear\n")
 	TEXT(" 1: RHIClear (default)\n")
 	TEXT(" 2: Quad at max z"),
diff --git a/Engine/Source/Runtime/Engine/Classes/Engine/RendererSettings.h b/Engine/Source/Runtime/Engine/Classes/Engine/RendererSettings.h
--- a/Engine/Source/Runtime/Engine/Classes/Engine/RendererSettings.h
+++ b/Engine/Source/Runtime/Engine/Classes/Engine/RendererSettings.h
@@ -238,7 +238,7 @@ class ENGINE_API URendererSettings
 
 	UPROPERTY(config, EditAnywhere, Category=Optimizations, meta=(
 		ConsoleVariable="r.ClearSceneMethod",DisplayName="Clear Scene",
-		ToolTip="Control how the scene is cleared before rendering"))
+		ToolTip="Select how the g-buffer is cleared in game mode (only affects deferred shading)."))
 	TEnumAsByte<EClearSceneOptions::Type> ClearSceneMethod;
 
 	UPROPERTY(config, EditAnywhere, Category=Optimizations, meta=(
diff --git a/Engine/Source/Runtime/Renderer/Private/DeferredShadingRenderer.cpp b/Engine/Source/Runtime/Renderer/Private/DeferredShadingRenderer.cpp
--- a/Engine/Source/Runtime/Renderer/Private/DeferredShadingRenderer.cpp
+++ b/Engine/Source/Runtime/Renderer/Private/DeferredShadingRenderer.cpp
@@ -735,7 +735,15 @@ void FDeferredShadingSceneRenderer::Render(FRHICommandListImmediate& RHICmdList)
 
 	if (ClearMethodCVar)
 	{
-		switch (ClearMethodCVar->GetValueOnRenderThread())
+		int32 clearMethod = ClearMethodCVar->GetValueOnRenderThread();
+
+		if (clearMethod == 0 && !ViewFamily.EngineShowFlags.Game)
+		{
+			// Do not clear the scene only if the view family is in game mode.
+			clearMethod = 1;
+		}
+
+		switch (clearMethod)
 		{
 		case 0: // No clear
 			{
