minor optimization

diff --git a/Engine/Source/Runtime/Renderer/Private/PostProcess/VisualizeTexture.cpp b/Engine/Source/Runtime/Renderer/Private/PostProcess/VisualizeTexture.cpp
--- a/Engine/Source/Runtime/Renderer/Private/PostProcess/VisualizeTexture.cpp
+++ b/Engine/Source/Runtime/Renderer/Private/PostProcess/VisualizeTexture.cpp
@@ -653,7 +653,8 @@ void FVisualizeTexture::SetCheckPoint(FRHICommandList& RHICmdList, const IPooled
 	}
 
 	// is this is the name we are observing with visualize texture?
-	if(ObservedDebugName == DebugName)
+	// First check if we need to find anything to avoid string the comparison
+	if (!ObservedDebugName.IsEmpty() && ObservedDebugName == DebugName)
 	{
 		// if multiple times reused during the frame, is that the one we want to look at?
 		if(*UsageCountPtr == ObservedDebugNameReusedGoal || ObservedDebugNameReusedGoal == 0xffffffff)
