minor optimization

diff --git a/Engine/Source/Runtime/Renderer/Private/VelocityRendering.cpp b/Engine/Source/Runtime/Renderer/Private/VelocityRendering.cpp
--- a/Engine/Source/Runtime/Renderer/Private/VelocityRendering.cpp
+++ b/Engine/Source/Runtime/Renderer/Private/VelocityRendering.cpp
@@ -458,7 +458,7 @@ bool FVelocityDrawingPolicyFactory::DrawDynamicMesh(
 		{			
 			RHICmdList.BuildAndSetLocalBoundShaderState(DrawingPolicy.GetBoundShaderStateInput(FeatureLevel));
 			DrawingPolicy.SetSharedState(RHICmdList, &View, FVelocityDrawingPolicy::ContextDataType());
-			for (int32 BatchElementIndex = 0; BatchElementIndex < Mesh.Elements.Num(); ++BatchElementIndex)
+			for (int32 BatchElementIndex = 0, BatchElementCount = Mesh.Elements.Num(); BatchElementIndex < BatchElementCount; ++BatchElementIndex)
 			{
 				DrawingPolicy.SetMeshRenderState(RHICmdList, View, PrimitiveSceneProxy, Mesh, BatchElementIndex, bBackFace, FMeshDrawingPolicy::ElementDataType(), FVelocityDrawingPolicy::ContextDataType());
 				DrawingPolicy.DrawMesh(RHICmdList, Mesh, BatchElementIndex);
