minor cleanup

diff --git a/Engine/Source/Runtime/Renderer/Private/VelocityRendering.cpp b/Engine/Source/Runtime/Renderer/Private/VelocityRendering.cpp
--- a/Engine/Source/Runtime/Renderer/Private/VelocityRendering.cpp
+++ b/Engine/Source/Runtime/Renderer/Private/VelocityRendering.cpp
@@ -13,8 +13,8 @@ static TAutoConsoleVariable<int32> CVarBasePassOutputsVelocity(
 	TEXT("r.BasePassOutputsVelocity"),
 	0,
 	TEXT("Enables rendering WPO velocities on the base pass.\n") \
-	TEXT("0 - Renders in a separate pass/rendertarget, all movable static meshes + dynamic.\n") \
-	TEXT("1 - Renders during the regular base pass adding an extra GBuffer, but allowing motion blur on materials with Time-based WPO.\n"),
+	TEXT(" 0: Renders in a separate pass/rendertarget, all movable static meshes + dynamic.\n") \
+	TEXT(" 1: Renders during the regular base pass adding an extra GBuffer, but allowing motion blur on materials with Time-based WPO."),
 	ECVF_ReadOnly | ECVF_RenderThreadSafe);
 
 //=============================================================================
@@ -738,26 +738,28 @@ void FDeferredShadingSceneRenderer::RenderVelocities(FRHICommandListImmediate& R
 	FPooledRenderTargetDesc Desc = FVelocityRendering::GetRenderTargetDesc();
 	GRenderTargetPool.FindFreeElement(Desc, VelocityRT, TEXT("Velocity"));
 
-	GPrevPerBoneMotionBlur.StartAppend();
-
-	BeginVelocityRendering(RHICmdList, VelocityRT, true);
-	
-	if (GRHICommandList.UseParallelAlgorithms() && CVarParallelVelocity.GetValueOnRenderThread())
-	{
-		RenderVelocitiesInnerParallel(RHICmdList, VelocityRT);
-	}
-	else
 	{
-		RenderVelocitiesInner(RHICmdList, VelocityRT);
-	}
+		GPrevPerBoneMotionBlur.StartAppend(ViewFamily.bWorldIsPaused);
 
-	RHICmdList.CopyToResolveTarget(VelocityRT->GetRenderTargetItem().TargetableTexture, VelocityRT->GetRenderTargetItem().ShaderResourceTexture, false, FResolveParams());
+		BeginVelocityRendering(RHICmdList, VelocityRT, true);
+
+		if (GRHICommandList.UseParallelAlgorithms() && CVarParallelVelocity.GetValueOnRenderThread())
+		{
+			RenderVelocitiesInnerParallel(RHICmdList, VelocityRT);
+		}
+		else
+		{
+			RenderVelocitiesInner(RHICmdList, VelocityRT);
+		}
+
+		RHICmdList.CopyToResolveTarget(VelocityRT->GetRenderTargetItem().TargetableTexture, VelocityRT->GetRenderTargetItem().ShaderResourceTexture, false, FResolveParams());
+
+		GPrevPerBoneMotionBlur.EndAppend();
+	}
 
 	// restore any color write state changes
 	RHICmdList.SetBlendState(TStaticBlendState<>::GetRHI());
 	RHICmdList.SetRasterizerState(TStaticRasterizerState<>::GetRHI());
-	
-	GPrevPerBoneMotionBlur.EndAppend();
 
 	// to be able to observe results with VisualizeTexture
 	GRenderTargetPool.VisualizeTexture.SetCheckPoint(RHICmdList, VelocityRT);
