maintain velocity in pause for TemporalAA as well (fixes very rare case)

diff --git a/Engine/Source/Runtime/Renderer/Private/SceneVisibility.cpp b/Engine/Source/Runtime/Renderer/Private/SceneVisibility.cpp
--- a/Engine/Source/Runtime/Renderer/Private/SceneVisibility.cpp
+++ b/Engine/Source/Runtime/Renderer/Private/SceneVisibility.cpp
@@ -2344,13 +2344,13 @@ void FDeferredShadingSceneRenderer::InitViews(FRHICommandListImmediate& RHICmdLi
 		InitDynamicShadows(RHICmdList);
 	}
 
-	if(ViewFamily.EngineShowFlags.MotionBlur && ViewFamily.bWorldIsPaused)
+	if (ViewFamily.bWorldIsPaused)
 	{
-		// so we can keep motion blur in paused mode
+		// so we can freeze motion blur and TemporalAA in paused mode
 
-		// per object motion blur
+		// per object velocity (static meshes)
 		Scene->MotionBlurInfoData.RestoreForPausedMotionBlur();
-		// per bone motion blur
+		// per bone velocity (skeletal meshes)
 		GPrevPerBoneMotionBlur.RestoreForPausedMotionBlur();
 	}
 
