Platforms that don't support distance field lighting will skip primitive delta tracking

diff --git a/Engine/Source/Runtime/Renderer/Private/RendererScene.cpp b/Engine/Source/Runtime/Renderer/Private/RendererScene.cpp
--- a/Engine/Source/Runtime/Renderer/Private/RendererScene.cpp
+++ b/Engine/Source/Runtime/Renderer/Private/RendererScene.cpp
@@ -132,14 +132,14 @@ FSceneViewState::~FSceneViewState()
 	DestroyLightPropagationVolume();
 }
 
-FDistanceFieldSceneData::FDistanceFieldSceneData() 
+FDistanceFieldSceneData::FDistanceFieldSceneData(EShaderPlatform ShaderPlatform) 
 	: NumObjectsInBuffer(0)
 	, ObjectBuffers(NULL)
 	, SurfelBuffers(NULL)
 	, InstancedSurfelBuffers(NULL)
 	, AtlasGeneration(0)
 {
-
+	bTrackPrimitives = DoesPlatformSupportDistanceFieldAO(ShaderPlatform) || DoesPlatformSupportDistanceFieldShadowing(ShaderPlatform);
 }
 
 FDistanceFieldSceneData::~FDistanceFieldSceneData() 
@@ -151,7 +151,8 @@ void FDistanceFieldSceneData::AddPrimitive(FPrimitiveSceneInfo* InPrimitive)
 {
 	const FPrimitiveSceneProxy* Proxy = InPrimitive->Proxy;
 
-	if (Proxy->CastsDynamicShadow()
+	if (bTrackPrimitives 
+		&& Proxy->CastsDynamicShadow()
 		&& Proxy->AffectsDistanceFieldLighting())
 	{
 		if (Proxy->SupportsHeightfieldRepresentation())
@@ -172,7 +173,8 @@ void FDistanceFieldSceneData::UpdatePrimitive(FPrimitiveSceneInfo* InPrimitive)
 {
 	const FPrimitiveSceneProxy* Proxy = InPrimitive->Proxy;
 
-	if (Proxy->CastsDynamicShadow() 
+	if (bTrackPrimitives 
+		&& Proxy->CastsDynamicShadow() 
 		&& Proxy->AffectsDistanceFieldLighting()
 		&& Proxy->SupportsDistanceFieldRepresentation() 
 		&& !PendingAddOperations.Contains(InPrimitive)
@@ -187,7 +189,7 @@ void FDistanceFieldSceneData::RemovePrimitive(FPrimitiveSceneInfo* InPrimitive)
 {
 	const FPrimitiveSceneProxy* Proxy = InPrimitive->Proxy;
 
-	if (Proxy->AffectsDistanceFieldLighting())
+	if (bTrackPrimitives && Proxy->AffectsDistanceFieldLighting())
 	{
 		if (Proxy->SupportsDistanceFieldRepresentation())
 		{
@@ -375,6 +377,7 @@ FScene::FScene(UWorld* InWorld, bool bInRequiresHitProxies, bool bInIsEditorScen
 ,	ReflectionSceneData(InFeatureLevel)
 ,	IndirectLightingCache(InFeatureLevel)
 ,	SurfaceCacheResources(NULL)
+,	DistanceFieldSceneData(GShaderPlatformForFeatureLevel[InFeatureLevel])
 ,	PreshadowCacheLayout(0, 0, 0, 0, false, false)
 ,	AtmosphericFog(NULL)
 ,	PrecomputedVisibilityHandler(NULL)
diff --git a/Engine/Source/Runtime/Renderer/Private/ScenePrivate.h b/Engine/Source/Runtime/Renderer/Private/ScenePrivate.h
--- a/Engine/Source/Runtime/Renderer/Private/ScenePrivate.h
+++ b/Engine/Source/Runtime/Renderer/Private/ScenePrivate.h
@@ -983,7 +983,7 @@ class FDistanceFieldSceneData
 {
 public:
 
-	FDistanceFieldSceneData();
+	FDistanceFieldSceneData(EShaderPlatform ShaderPlatform);
 	~FDistanceFieldSceneData();
 
 	void AddPrimitive(FPrimitiveSceneInfo* InPrimitive);
@@ -1030,6 +1030,8 @@ public:
 
 	/** Used to detect atlas reallocations, since objects store UVs into the atlas and need to be updated when it changes. */
 	int32 AtlasGeneration;
+
+	bool bTrackPrimitives;
 };
 
 /** Stores data for an allocation in the FIndirectLightingCache. */
