Fixed DFAO fading out too early

diff --git a/Engine/Source/Runtime/Renderer/Private/DistanceFieldSurfaceCacheLighting.cpp b/Engine/Source/Runtime/Renderer/Private/DistanceFieldSurfaceCacheLighting.cpp
--- a/Engine/Source/Runtime/Renderer/Private/DistanceFieldSurfaceCacheLighting.cpp
+++ b/Engine/Source/Runtime/Renderer/Private/DistanceFieldSurfaceCacheLighting.cpp
@@ -183,7 +183,7 @@ FAutoConsoleVariableRef CVarAOMaxViewDistance(
 	ECVF_Cheat | ECVF_RenderThreadSafe
 	);
 
-float GAOViewFadeDistanceScale = .2f;
+float GAOViewFadeDistanceScale = .7f;
 FAutoConsoleVariableRef CVarAOViewFadeDistanceScale(
 	TEXT("r.AOViewFadeDistanceScale"),
 	GAOViewFadeDistanceScale,
