Distance field primitive updates no longer infinitely accumulate when distance field features are disabled
Distance field scene updates are no longer tracked when the project setting r.GenerateMeshDistanceFields is disabled

diff --git a/Engine/Source/Runtime/Renderer/Private/DistanceFieldObjectManagement.cpp b/Engine/Source/Runtime/Renderer/Private/DistanceFieldObjectManagement.cpp
--- a/Engine/Source/Runtime/Renderer/Private/DistanceFieldObjectManagement.cpp
+++ b/Engine/Source/Runtime/Renderer/Private/DistanceFieldObjectManagement.cpp
@@ -689,6 +689,163 @@ void UpdateGlobalDistanceFieldObjectRemoves(FRHICommandListImmediate& RHICmdList
 	}
 }
 
+/** Gathers the information needed to represent a single object's distance field and appends it to the upload buffers. */
+void ProcessPrimitiveUpdate(
+	bool bIsAddOperation,
+	FRHICommandListImmediate& RHICmdList, 
+	FSceneRenderer& SceneRenderer, 
+	FPrimitiveSceneInfo* PrimitiveSceneInfo, 
+	int32 OriginalNumObjects,
+	FVector InvTextureDim,
+	bool bPrepareForDistanceFieldGI, 
+	TArray<FMatrix>& ObjectLocalToWorldTransforms,
+	TArray<uint32>& UploadObjectIndices,
+	TArray<FVector4>& UploadObjectData)
+{
+	FScene* Scene = SceneRenderer.Scene;
+	FDistanceFieldSceneData& DistanceFieldSceneData = Scene->DistanceFieldSceneData;
+
+	ObjectLocalToWorldTransforms.Reset();
+
+	FBox LocalVolumeBounds;
+	FIntVector BlockMin;
+	FIntVector BlockSize;
+	bool bBuiltAsIfTwoSided;
+	bool bMeshWasPlane;
+	PrimitiveSceneInfo->Proxy->GetDistancefieldAtlasData(LocalVolumeBounds, BlockMin, BlockSize, bBuiltAsIfTwoSided, bMeshWasPlane, ObjectLocalToWorldTransforms);
+
+	if (BlockMin.X >= 0 
+		&& BlockMin.Y >= 0 
+		&& BlockMin.Z >= 0 
+		&& ObjectLocalToWorldTransforms.Num() > 0)
+	{
+		const float BoundingRadius = PrimitiveSceneInfo->Proxy->GetBounds().SphereRadius;
+
+		// Proxy bounds are only useful if single instance
+		if (ObjectLocalToWorldTransforms.Num() > 1 || BoundingRadius < GAOMaxObjectBoundingRadius)
+		{
+			FPrimitiveSurfelAllocation Allocation;
+			FPrimitiveSurfelAllocation InstancedAllocation;
+						
+			if (bPrepareForDistanceFieldGI)
+			{
+				const FPrimitiveSurfelAllocation* AllocationPtr = Scene->DistanceFieldSceneData.SurfelAllocations.FindAllocation(PrimitiveSceneInfo);
+				const FPrimitiveSurfelAllocation* InstancedAllocationPtr = Scene->DistanceFieldSceneData.InstancedSurfelAllocations.FindAllocation(PrimitiveSceneInfo);
+
+				if (AllocationPtr)
+				{
+					checkSlow(InstancedAllocationPtr && InstancedAllocationPtr->NumInstances == ObjectLocalToWorldTransforms.Num());
+					Allocation = *AllocationPtr;
+					InstancedAllocation = *InstancedAllocationPtr;
+
+					extern void GenerateSurfelRepresentation(FRHICommandListImmediate& RHICmdList, FSceneRenderer& Renderer, FViewInfo& View, FPrimitiveSceneInfo* PrimitiveSceneInfo, const FMatrix& Instance0Transform, FPrimitiveSurfelAllocation& Allocation);
+					// @todo - support surfel generation without a view
+					GenerateSurfelRepresentation(RHICmdList, SceneRenderer, SceneRenderer.Views[0], PrimitiveSceneInfo, ObjectLocalToWorldTransforms[0], Allocation);
+
+					if (Allocation.NumSurfels == 0)
+					{
+						InstancedAllocation.NumSurfels = 0;
+						InstancedAllocation.NumInstances = 0;
+						InstancedAllocation.NumLOD0 = 0;
+					}
+				}
+			}
+
+			if (bIsAddOperation)
+			{
+				DistanceFieldSceneData.NumObjectsInBuffer += ObjectLocalToWorldTransforms.Num();
+				PrimitiveSceneInfo->DistanceFieldInstanceIndices.Empty(ObjectLocalToWorldTransforms.Num());
+			}
+
+			for (int32 TransformIndex = 0; TransformIndex < ObjectLocalToWorldTransforms.Num(); TransformIndex++)
+			{
+				const uint32 UploadIndex = bIsAddOperation ? OriginalNumObjects + UploadObjectIndices.Num() : PrimitiveSceneInfo->DistanceFieldInstanceIndices[TransformIndex];
+
+				if (bIsAddOperation)
+				{
+					const int32 AddIndex = OriginalNumObjects + UploadObjectIndices.Num();
+					DistanceFieldSceneData.PrimitiveInstanceMapping.Add(FPrimitiveAndInstance(PrimitiveSceneInfo, TransformIndex));
+					PrimitiveSceneInfo->DistanceFieldInstanceIndices.Add(AddIndex);
+				}
+
+				UploadObjectIndices.Add(UploadIndex);
+
+				FMatrix LocalToWorld = ObjectLocalToWorldTransforms[TransformIndex];
+
+				if (bMeshWasPlane)
+				{
+					FVector LocalScales = LocalToWorld.GetScaleVector();
+					FVector AbsLocalScales(FMath::Abs(LocalScales.X), FMath::Abs(LocalScales.Y), FMath::Abs(LocalScales.Z));
+					float MidScale = FMath::Min(AbsLocalScales.X, AbsLocalScales.Y);
+					float ScaleAdjust = FMath::Sign(LocalScales.Z) * MidScale / AbsLocalScales.Z;
+					// The mesh was determined to be a plane flat in Z during the build process, so we can change the Z scale
+					// Helps in cases with modular ground pieces with scales of (10, 10, 1) and some triangles just above Z=0
+					LocalToWorld.SetAxis(2, LocalToWorld.GetScaledAxis(EAxis::Z) * ScaleAdjust);
+				}
+
+				const FMatrix VolumeToWorld = FScaleMatrix(LocalVolumeBounds.GetExtent()) 
+					* FTranslationMatrix(LocalVolumeBounds.GetCenter())
+					* LocalToWorld;
+
+				const FVector4 ObjectBoundingSphere(VolumeToWorld.GetOrigin(), VolumeToWorld.GetScaleVector().Size());
+
+				UploadObjectData.Add(ObjectBoundingSphere);
+
+				const float MaxExtent = LocalVolumeBounds.GetExtent().GetMax();
+
+				const FMatrix UniformScaleVolumeToWorld = FScaleMatrix(MaxExtent) 
+					* FTranslationMatrix(LocalVolumeBounds.GetCenter())
+					* LocalToWorld;
+
+				const FVector InvBlockSize(1.0f / BlockSize.X, 1.0f / BlockSize.Y, 1.0f / BlockSize.Z);
+
+				//float3 VolumeUV = (VolumePosition / LocalPositionExtent * .5f * UVScale + .5f * UVScale + UVAdd;
+				const FVector LocalPositionExtent = LocalVolumeBounds.GetExtent() / FVector(MaxExtent);
+				const FVector UVScale = FVector(BlockSize) * InvTextureDim;
+				const float VolumeScale = UniformScaleVolumeToWorld.GetMaximumAxisScale();
+
+				const FMatrix WorldToVolume = UniformScaleVolumeToWorld.Inverse();
+				// WorldToVolume
+				UploadObjectData.Add(*(FVector4*)&WorldToVolume.M[0]);
+				UploadObjectData.Add(*(FVector4*)&WorldToVolume.M[1]);
+				UploadObjectData.Add(*(FVector4*)&WorldToVolume.M[2]);
+				UploadObjectData.Add(*(FVector4*)&WorldToVolume.M[3]);
+
+				// Clamp to texel center by subtracting a half texel in the [-1,1] position space
+				// LocalPositionExtent
+				UploadObjectData.Add(FVector4(LocalPositionExtent - InvBlockSize, LocalVolumeBounds.Min.X));
+
+				// UVScale, VolumeScale and sign gives bGeneratedAsTwoSided
+				const float WSign = bBuiltAsIfTwoSided ? -1 : 1;
+				UploadObjectData.Add(FVector4(FVector(BlockSize) * InvTextureDim * .5f / LocalPositionExtent, WSign * VolumeScale));
+
+				// UVAdd
+				UploadObjectData.Add(FVector4(FVector(BlockMin) * InvTextureDim + .5f * UVScale, LocalVolumeBounds.Min.Y));
+
+				// Box bounds
+				UploadObjectData.Add(FVector4(LocalVolumeBounds.Max, LocalVolumeBounds.Min.Z));
+
+				UploadObjectData.Add(*(FVector4*)&UniformScaleVolumeToWorld.M[0]);
+				UploadObjectData.Add(*(FVector4*)&UniformScaleVolumeToWorld.M[1]);
+				UploadObjectData.Add(*(FVector4*)&UniformScaleVolumeToWorld.M[2]);
+
+				UploadObjectData.Add(*(FVector4*)&LocalToWorld.M[0]);
+				UploadObjectData.Add(*(FVector4*)&LocalToWorld.M[1]);
+				UploadObjectData.Add(*(FVector4*)&LocalToWorld.M[2]);
+				UploadObjectData.Add(*(FVector4*)&LocalToWorld.M[3]);
+
+				UploadObjectData.Add(FVector4(Allocation.Offset, Allocation.NumLOD0, Allocation.NumSurfels, InstancedAllocation.Offset + InstancedAllocation.NumSurfels * TransformIndex));
+
+				checkSlow(UploadObjectData.Num() % UploadObjectDataStride == 0);
+			}
+		}
+		else
+		{
+			UE_LOG(LogDistanceField,Log,TEXT("Primitive %s %s excluded due to bounding radius %f"), *PrimitiveSceneInfo->Proxy->GetOwnerName().ToString(), *PrimitiveSceneInfo->Proxy->GetResourceName().ToString(), BoundingRadius);
+		}
+	}
+}
+
 void FDeferredShadingSceneRenderer::UpdateGlobalDistanceFieldObjectBuffers(FRHICommandListImmediate& RHICmdList) 
 {
 	FDistanceFieldSceneData& DistanceFieldSceneData = Scene->DistanceFieldSceneData;
@@ -848,156 +1005,42 @@ void FDeferredShadingSceneRenderer::UpdateGlobalDistanceFieldObjectBuffers(FRHIC
 				}
 			}
 
-			for (int32 UploadPrimitiveIndex = 0; UploadPrimitiveIndex < NumUploadOperations; UploadPrimitiveIndex++)
+			for (int32 UploadPrimitiveIndex = 0; UploadPrimitiveIndex < DistanceFieldSceneData.PendingAddOperations.Num(); UploadPrimitiveIndex++)
 			{
-				const bool bIsAddOperation = UploadPrimitiveIndex < DistanceFieldSceneData.PendingAddOperations.Num();
-
-				FPrimitiveSceneInfo* PrimitiveSceneInfo = bIsAddOperation 
-					? DistanceFieldSceneData.PendingAddOperations[UploadPrimitiveIndex] 
-					: DistanceFieldSceneData.PendingUpdateOperations[UploadPrimitiveIndex - DistanceFieldSceneData.PendingAddOperations.Num()];
-
-				ObjectLocalToWorldTransforms.Reset();
-
-				FBox LocalVolumeBounds;
-				FIntVector BlockMin;
-				FIntVector BlockSize;
-				bool bBuiltAsIfTwoSided;
-				bool bMeshWasPlane;
-				PrimitiveSceneInfo->Proxy->GetDistancefieldAtlasData(LocalVolumeBounds, BlockMin, BlockSize, bBuiltAsIfTwoSided, bMeshWasPlane, ObjectLocalToWorldTransforms);
-
-				if (BlockMin.X >= 0 
-					&& BlockMin.Y >= 0 
-					&& BlockMin.Z >= 0 
-					&& ObjectLocalToWorldTransforms.Num() > 0)
-				{
-					const float BoundingRadius = PrimitiveSceneInfo->Proxy->GetBounds().SphereRadius;
-
-					// Proxy bounds are only useful if single instance
-					if (ObjectLocalToWorldTransforms.Num() > 1 || BoundingRadius < GAOMaxObjectBoundingRadius)
-					{
-						FPrimitiveSurfelAllocation Allocation;
-						FPrimitiveSurfelAllocation InstancedAllocation;
-						
-						if (bPrepareForDistanceFieldGI)
-						{
-							const FPrimitiveSurfelAllocation* AllocationPtr = Scene->DistanceFieldSceneData.SurfelAllocations.FindAllocation(PrimitiveSceneInfo);
-							const FPrimitiveSurfelAllocation* InstancedAllocationPtr = Scene->DistanceFieldSceneData.InstancedSurfelAllocations.FindAllocation(PrimitiveSceneInfo);
-
-							if (AllocationPtr)
-							{
-								checkSlow(InstancedAllocationPtr && InstancedAllocationPtr->NumInstances == ObjectLocalToWorldTransforms.Num());
-								Allocation = *AllocationPtr;
-								InstancedAllocation = *InstancedAllocationPtr;
-
-								extern void GenerateSurfelRepresentation(FRHICommandListImmediate& RHICmdList, FSceneRenderer& Renderer, FViewInfo& View, FPrimitiveSceneInfo* PrimitiveSceneInfo, const FMatrix& Instance0Transform, FPrimitiveSurfelAllocation& Allocation);
-								GenerateSurfelRepresentation(RHICmdList, *this, Views[0], PrimitiveSceneInfo, ObjectLocalToWorldTransforms[0], Allocation);
-
-								if (Allocation.NumSurfels == 0)
-								{
-									InstancedAllocation.NumSurfels = 0;
-									InstancedAllocation.NumInstances = 0;
-									InstancedAllocation.NumLOD0 = 0;
-								}
-							}
-						}
-
-						if (bIsAddOperation)
-						{
-							DistanceFieldSceneData.NumObjectsInBuffer += ObjectLocalToWorldTransforms.Num();
-							PrimitiveSceneInfo->DistanceFieldInstanceIndices.Empty(ObjectLocalToWorldTransforms.Num());
-						}
-
-						for (int32 TransformIndex = 0; TransformIndex < ObjectLocalToWorldTransforms.Num(); TransformIndex++)
-						{
-							const uint32 UploadIndex = bIsAddOperation ? OriginalNumObjects + UploadObjectIndices.Num() : PrimitiveSceneInfo->DistanceFieldInstanceIndices[TransformIndex];
-
-							if (bIsAddOperation)
-							{
-								const int32 AddIndex = OriginalNumObjects + UploadObjectIndices.Num();
-								DistanceFieldSceneData.PrimitiveInstanceMapping.Add(FPrimitiveAndInstance(PrimitiveSceneInfo, TransformIndex));
-								PrimitiveSceneInfo->DistanceFieldInstanceIndices.Add(AddIndex);
-							}
-
-							UploadObjectIndices.Add(UploadIndex);
-
-							FMatrix LocalToWorld = ObjectLocalToWorldTransforms[TransformIndex];
-
-							if (bMeshWasPlane)
-							{
-								FVector LocalScales = LocalToWorld.GetScaleVector();
-								FVector AbsLocalScales(FMath::Abs(LocalScales.X), FMath::Abs(LocalScales.Y), FMath::Abs(LocalScales.Z));
-								float MidScale = FMath::Min(AbsLocalScales.X, AbsLocalScales.Y);
-								float ScaleAdjust = FMath::Sign(LocalScales.Z) * MidScale / AbsLocalScales.Z;
-								// The mesh was determined to be a plane flat in Z during the build process, so we can change the Z scale
-								// Helps in cases with modular ground pieces with scales of (10, 10, 1) and some triangles just above Z=0
-								LocalToWorld.SetAxis(2, LocalToWorld.GetScaledAxis(EAxis::Z) * ScaleAdjust);
-							}
-
-							const FMatrix VolumeToWorld = FScaleMatrix(LocalVolumeBounds.GetExtent()) 
-								* FTranslationMatrix(LocalVolumeBounds.GetCenter())
-								* LocalToWorld;
-
-							const FVector4 ObjectBoundingSphere(VolumeToWorld.GetOrigin(), VolumeToWorld.GetScaleVector().Size());
-
-							UploadObjectData.Add(ObjectBoundingSphere);
-
-							const float MaxExtent = LocalVolumeBounds.GetExtent().GetMax();
-
-							const FMatrix UniformScaleVolumeToWorld = FScaleMatrix(MaxExtent) 
-								* FTranslationMatrix(LocalVolumeBounds.GetCenter())
-								* LocalToWorld;
-
-							const FVector InvBlockSize(1.0f / BlockSize.X, 1.0f / BlockSize.Y, 1.0f / BlockSize.Z);
-
-							//float3 VolumeUV = (VolumePosition / LocalPositionExtent * .5f * UVScale + .5f * UVScale + UVAdd;
-							const FVector LocalPositionExtent = LocalVolumeBounds.GetExtent() / FVector(MaxExtent);
-							const FVector UVScale = FVector(BlockSize) * InvTextureDim;
-							const float VolumeScale = UniformScaleVolumeToWorld.GetMaximumAxisScale();
-
-							const FMatrix WorldToVolume = UniformScaleVolumeToWorld.Inverse();
-							// WorldToVolume
-							UploadObjectData.Add(*(FVector4*)&WorldToVolume.M[0]);
-							UploadObjectData.Add(*(FVector4*)&WorldToVolume.M[1]);
-							UploadObjectData.Add(*(FVector4*)&WorldToVolume.M[2]);
-							UploadObjectData.Add(*(FVector4*)&WorldToVolume.M[3]);
-
-							// Clamp to texel center by subtracting a half texel in the [-1,1] position space
-							// LocalPositionExtent
-							UploadObjectData.Add(FVector4(LocalPositionExtent - InvBlockSize, LocalVolumeBounds.Min.X));
-
-							// UVScale, VolumeScale and sign gives bGeneratedAsTwoSided
-							const float WSign = bBuiltAsIfTwoSided ? -1 : 1;
-							UploadObjectData.Add(FVector4(FVector(BlockSize) * InvTextureDim * .5f / LocalPositionExtent, WSign * VolumeScale));
-
-							// UVAdd
-							UploadObjectData.Add(FVector4(FVector(BlockMin) * InvTextureDim + .5f * UVScale, LocalVolumeBounds.Min.Y));
-
-							// Box bounds
-							UploadObjectData.Add(FVector4(LocalVolumeBounds.Max, LocalVolumeBounds.Min.Z));
-
-							UploadObjectData.Add(*(FVector4*)&UniformScaleVolumeToWorld.M[0]);
-							UploadObjectData.Add(*(FVector4*)&UniformScaleVolumeToWorld.M[1]);
-							UploadObjectData.Add(*(FVector4*)&UniformScaleVolumeToWorld.M[2]);
-
-							UploadObjectData.Add(*(FVector4*)&LocalToWorld.M[0]);
-							UploadObjectData.Add(*(FVector4*)&LocalToWorld.M[1]);
-							UploadObjectData.Add(*(FVector4*)&LocalToWorld.M[2]);
-							UploadObjectData.Add(*(FVector4*)&LocalToWorld.M[3]);
-
-							UploadObjectData.Add(FVector4(Allocation.Offset, Allocation.NumLOD0, Allocation.NumSurfels, InstancedAllocation.Offset + InstancedAllocation.NumSurfels * TransformIndex));
+				FPrimitiveSceneInfo* PrimitiveSceneInfo = DistanceFieldSceneData.PendingAddOperations[UploadPrimitiveIndex];
+
+				ProcessPrimitiveUpdate(
+					true,
+					RHICmdList, 
+					*this, 
+					PrimitiveSceneInfo, 
+					OriginalNumObjects, 
+					InvTextureDim, 
+					bPrepareForDistanceFieldGI, 
+					ObjectLocalToWorldTransforms, 
+					UploadObjectIndices, 
+					UploadObjectData);
+			}
 
-							checkSlow(UploadObjectData.Num() % UploadObjectDataStride == 0);
-						}
-					}
-					else
-					{
-						UE_LOG(LogDistanceField,Log,TEXT("Primitive %s %s excluded due to bounding radius %f"), *PrimitiveSceneInfo->Proxy->GetOwnerName().ToString(), *PrimitiveSceneInfo->Proxy->GetResourceName().ToString(), BoundingRadius);
-					}
-				}
+			for (TSet<FPrimitiveSceneInfo*>::TIterator It(DistanceFieldSceneData.PendingUpdateOperations); It; ++It)
+			{
+				FPrimitiveSceneInfo* PrimitiveSceneInfo = *It;
+
+				ProcessPrimitiveUpdate(
+					false,
+					RHICmdList, 
+					*this, 
+					PrimitiveSceneInfo, 
+					OriginalNumObjects, 
+					InvTextureDim, 
+					bPrepareForDistanceFieldGI, 
+					ObjectLocalToWorldTransforms, 
+					UploadObjectIndices, 
+					UploadObjectData);
 			}
 
 			DistanceFieldSceneData.PendingAddOperations.Reset();
-			DistanceFieldSceneData.PendingUpdateOperations.Reset();
+			DistanceFieldSceneData.PendingUpdateOperations.Empty();
 
 			if (DistanceFieldSceneData.ObjectBuffers->MaxObjects < DistanceFieldSceneData.NumObjectsInBuffer)
 			{
diff --git a/Engine/Source/Runtime/Renderer/Private/RendererScene.cpp b/Engine/Source/Runtime/Renderer/Private/RendererScene.cpp
--- a/Engine/Source/Runtime/Renderer/Private/RendererScene.cpp
+++ b/Engine/Source/Runtime/Renderer/Private/RendererScene.cpp
@@ -139,7 +139,9 @@ FDistanceFieldSceneData::FDistanceFieldSceneData(EShaderPlatform ShaderPlatform)
 	, InstancedSurfelBuffers(NULL)
 	, AtlasGeneration(0)
 {
-	bTrackPrimitives = DoesPlatformSupportDistanceFieldAO(ShaderPlatform) || DoesPlatformSupportDistanceFieldShadowing(ShaderPlatform);
+	static const auto CVar = IConsoleManager::Get().FindTConsoleVariableDataInt(TEXT("r.GenerateMeshDistanceFields"));
+
+	bTrackPrimitives = (DoesPlatformSupportDistanceFieldAO(ShaderPlatform) || DoesPlatformSupportDistanceFieldShadowing(ShaderPlatform)) && CVar->GetValueOnGameThread() != 0;
 }
 
 FDistanceFieldSceneData::~FDistanceFieldSceneData() 
@@ -178,6 +180,8 @@ void FDistanceFieldSceneData::UpdatePrimitive(FPrimitiveSceneInfo* InPrimitive)
 		&& Proxy->AffectsDistanceFieldLighting()
 		&& Proxy->SupportsDistanceFieldRepresentation() 
 		&& !PendingAddOperations.Contains(InPrimitive)
+		// This is needed to prevent infinite buildup when DF features are off such that the pending operations don't get consumed
+		&& !PendingUpdateOperations.Contains(InPrimitive)
 		// This can happen when the primitive fails to allocate from the SDF atlas
 		&& InPrimitive->DistanceFieldInstanceIndices.Num() > 0)
 	{
diff --git a/Engine/Source/Runtime/Renderer/Private/ScenePrivate.h b/Engine/Source/Runtime/Renderer/Private/ScenePrivate.h
--- a/Engine/Source/Runtime/Renderer/Private/ScenePrivate.h
+++ b/Engine/Source/Runtime/Renderer/Private/ScenePrivate.h
@@ -1025,7 +1025,7 @@ public:
 
 	/** Pending operations on the object buffers to be processed next frame. */
 	TArray<FPrimitiveSceneInfo*> PendingAddOperations;
-	TArray<FPrimitiveSceneInfo*> PendingUpdateOperations;
+	TSet<FPrimitiveSceneInfo*> PendingUpdateOperations;
 	TArray<FPrimitiveRemoveInfo> PendingRemoveOperations;
 
 	/** Used to detect atlas reallocations, since objects store UVs into the atlas and need to be updated when it changes. */
