Clamp light depth bounds to be valid

diff --git a/Engine/Source/Runtime/Renderer/Private/LightRendering.cpp b/Engine/Source/Runtime/Renderer/Private/LightRendering.cpp
--- a/Engine/Source/Runtime/Renderer/Private/LightRendering.cpp
+++ b/Engine/Source/Runtime/Renderer/Private/LightRendering.cpp
@@ -929,6 +929,12 @@ void FDeferredShadingSceneRenderer::RenderLight(FRHICommandList& RHICmdList, con
 				float FarDepth = 0.f;
 				CalculateLightNearFarDepthFromBounds(View,LightBounds,NearDepth,FarDepth);
 
+				if (NearDepth <= FarDepth)
+				{
+					NearDepth = 1.0f;
+					FarDepth = 0.0f;
+				}
+
 				// UE4 uses reversed depth, so far < near
 				RHICmdList.EnableDepthBoundsTest(true,FarDepth,NearDepth);
 			}
