Fixed PostprocessMaterial StaticSwitch functionality while avoiding warning that appeared in some reare cases

diff --git a/Engine/Source/Runtime/Engine/Private/Materials/Material.cpp b/Engine/Source/Runtime/Engine/Private/Materials/Material.cpp
--- a/Engine/Source/Runtime/Engine/Private/Materials/Material.cpp
+++ b/Engine/Source/Runtime/Engine/Private/Materials/Material.cpp
@@ -1132,7 +1132,7 @@ void UMaterialInterface::OverrideBlendableSettings(class FSceneView& View, float
 	}
 	else
 	{
-		UMaterialInstanceDynamic* MID = View.State->GetReusableMID((UMaterialInterface*)Material);//, (UMaterialInterface*)this);
+		UMaterialInstanceDynamic* MID = View.State->GetReusableMID((UMaterialInterface*)this);
 
 		if(MID)
 		{
diff --git a/Engine/Source/Runtime/Engine/Private/Materials/MaterialInstance.cpp b/Engine/Source/Runtime/Engine/Private/Materials/MaterialInstance.cpp
--- a/Engine/Source/Runtime/Engine/Private/Materials/MaterialInstance.cpp
+++ b/Engine/Source/Runtime/Engine/Private/Materials/MaterialInstance.cpp
@@ -2315,7 +2315,7 @@ void UMaterialInstance::OverrideBlendableSettings(class FSceneView& View, float
 		{
 			UMaterial* Base = Material->GetBaseMaterial();
 
-			UMaterialInstanceDynamic* MID = View.State->GetReusableMID((UMaterialInterface*)Base);//, (UMaterialInterface*)this);
+			UMaterialInstanceDynamic* MID = View.State->GetReusableMID((UMaterialInterface*)this);
 
 			if(MID)
 			{
@@ -2346,7 +2346,7 @@ void UMaterialInstance::OverrideBlendableSettings(class FSceneView& View, float
 	}
 	else
 	{
-		UMaterialInstanceDynamic* MID = View.State->GetReusableMID((UMaterialInterface*)Material);
+		UMaterialInstanceDynamic* MID = View.State->GetReusableMID((UMaterialInterface*)this);
 
 		if(MID)
 		{
diff --git a/Engine/Source/Runtime/Renderer/Private/ScenePrivate.h b/Engine/Source/Runtime/Renderer/Private/ScenePrivate.h
--- a/Engine/Source/Runtime/Renderer/Private/ScenePrivate.h
+++ b/Engine/Source/Runtime/Renderer/Private/ScenePrivate.h
@@ -731,9 +731,18 @@ public:
 	}
 
 	// Note: OnStartPostProcessing() needs to be called each frame for each view
-	virtual UMaterialInstanceDynamic* GetReusableMID(class UMaterialInterface* ParentMaterial) override
+	virtual UMaterialInstanceDynamic* GetReusableMID(class UMaterialInterface* InParentMaterial) override
 	{		
 		check(IsInGameThread());
+		check(InParentMaterial);
+
+		auto ParentAsMaterialInstance = Cast<UMaterialInstanceDynamic>(InParentMaterial);
+
+		// fixup MID parents as this is not allowed, take the next MIC or Material.
+		UMaterialInterface* ParentMaterial = ParentAsMaterialInstance ? ParentAsMaterialInstance->Parent : InParentMaterial;
+
+		// this is not allowed and would cause an error later in the code
+		check(!ParentMaterial->IsA(UMaterialInstanceDynamic::StaticClass()));
 
 		if(MIDUsedCount < (uint32)MIDPool.Num())
 		{
