fixed DBuffer left eye in Stereo rendering
 see https://udn.unrealengine.com/questions/242621/dbuffer-decal-left-eye-corruption-on-oculus-rift.html

diff --git a/Engine/Source/Runtime/Renderer/Private/DeferredShadingRenderer.cpp b/Engine/Source/Runtime/Renderer/Private/DeferredShadingRenderer.cpp
--- a/Engine/Source/Runtime/Renderer/Private/DeferredShadingRenderer.cpp
+++ b/Engine/Source/Runtime/Renderer/Private/DeferredShadingRenderer.cpp
@@ -851,10 +851,11 @@ void FDeferredShadingSceneRenderer::Render(FRHICommandListImmediate& RHICmdList)
 		GSceneRenderTargets.ResolveSceneDepthTexture(RHICmdList);
 		GSceneRenderTargets.ResolveSceneDepthToAuxiliaryTexture(RHICmdList);
 
-		// e.g. ambient cubemaps, ambient occlusion, deferred decals
+		// e.g. DBuffer deferred decals
 		for(int32 ViewIndex = 0;ViewIndex < Views.Num();ViewIndex++)
 		{	
 			SCOPED_CONDITIONAL_DRAW_EVENTF(RHICmdList, EventView,Views.Num() > 1, TEXT("View%d"), ViewIndex);
+
 			GCompositionLighting.ProcessBeforeBasePass(RHICmdList, Views[ViewIndex]);
 		}
 	}
@@ -1199,7 +1200,7 @@ static void SetupPrePassView(FRHICommandList& RHICmdList, const FIntRect& ViewRe
 	RHICmdList.SetDepthStencilState(TStaticDepthStencilState<true,CF_DepthNearOrEqual>::GetRHI());
 	RHICmdList.SetViewport(ViewRect.Min.X, ViewRect.Min.Y, 0, ViewRect.Max.X, ViewRect.Max.Y, 1);
 	RHICmdList.SetRasterizerState(TStaticRasterizerState<FM_Solid, CM_None>::GetRHI());
-	RHICmdList.SetScissorRect(true, ViewRect.Min.X, ViewRect.Min.Y, ViewRect.Max.X, ViewRect.Max.Y);
+	RHICmdList.SetScissorRect(false, 0, 0, 0, 0);
 }
 
 bool FDeferredShadingSceneRenderer::RenderPrePassView(FRHICommandList& RHICmdList, const FViewInfo& View)
