UE4 - Remove bogus condition

diff --git a/Engine/Source/Runtime/Renderer/Private/SceneHitProxyRendering.cpp b/Engine/Source/Runtime/Renderer/Private/SceneHitProxyRendering.cpp
--- a/Engine/Source/Runtime/Renderer/Private/SceneHitProxyRendering.cpp
+++ b/Engine/Source/Runtime/Renderer/Private/SceneHitProxyRendering.cpp
@@ -37,7 +37,6 @@ public:
 	{
 		// Only compile the hit proxy vertex shader on PC
 		return IsPCPlatform(Platform)
-			&& (IsFeatureLevelSupported(Platform, ERHIFeatureLevel::SM4) || IsPCPlatform(Platform))
 			// and only compile for the default material or materials that are masked.
 			&& (Material->IsSpecialEngineMaterial() || Material->IsMasked() || Material->MaterialMayModifyMeshPosition() || Material->IsTwoSided());
 	}
