cleanup

diff --git a/Engine/Source/Runtime/Renderer/Private/CompositionLighting/PostProcessDeferredDecals.cpp b/Engine/Source/Runtime/Renderer/Private/CompositionLighting/PostProcessDeferredDecals.cpp
--- a/Engine/Source/Runtime/Renderer/Private/CompositionLighting/PostProcessDeferredDecals.cpp
+++ b/Engine/Source/Runtime/Renderer/Private/CompositionLighting/PostProcessDeferredDecals.cpp
@@ -692,7 +692,6 @@ bool RenderPreStencil(FRenderingCompositePassContext& Context, const FMaterialSh
 	>::GetRHI() );
 
 	// Carmack's reverse on the bounds
-	//@todo-martinm
 	Context.RHICmdList.SetDepthStencilState(TStaticDepthStencilState<
 		false,CF_LessEqual,
 		true,CF_Equal,SO_Keep,SO_Keep,SO_Increment,
