display BoneBuffers in VisualizeMotionBlur

diff --git a/Engine/Shaders/PostProcessMotionBlur.usf b/Engine/Shaders/PostProcessMotionBlur.usf
--- a/Engine/Shaders/PostProcessMotionBlur.usf
+++ b/Engine/Shaders/PostProcessMotionBlur.usf
@@ -87,6 +87,13 @@ float4 TextureViewMad;
 // xy:IndividualVelocityScale.xy zw:unused, from postprocess settings
 float4 MotionBlurParameters;
 
+// BoneMatrices and PreviousBoneMatrices for visualization
+
+// The bone matrix buffer stored as 4x3 (3 float4 texels behind each other), all chunks of a skeletal mesh in one
+Buffer<float4> BoneMatrices;
+// buffer with all old bone matrices stored as 4x3 (3 float4 texels behind each other), all chunks of a skeletal mesh in one
+Buffer<float4> PreviousBoneMatrices;
+
 // ------------------------------------------
 
 
@@ -607,6 +614,32 @@ void MainPS(float4 UVAndScreenPos : TEXCOORD0, out float4 OutColor : SV_Target0)
 		OutColor.rgb = lerp(lerp(WorldCheckerboard, TintColor, 0.5f), LineColor, LineMask);
 		OutColor.rgb *= lerp(1.0f, 0.9f, DiskMask);
 	}
+	// visualize BoneBuffers
+	{
+		const int2 LeftTop = int2(16, 90);
+		const int ElementsPerLine = 32;
+		const int Scale = 3;
+		const int LinesPerBuffer = 48;
+		const int GapBetweenBuffers = 8;
+		const int FadeRegion = 16;
+
+		float2 PixelPos = ComputePixelPosCenter(UVAndScreenPos.zw, true);
+		int2 PixelPosInt = ((uint2)PixelPos - LeftTop) / Scale; 
+
+		// 0  and 1 are valid buffers
+		int LineInBuffer = PixelPosInt.y % (LinesPerBuffer + GapBetweenBuffers);
+		int WhichBuffer = PixelPosInt.y / (LinesPerBuffer + GapBetweenBuffers);
+		float Alpha = saturate((LinesPerBuffer - LineInBuffer) / (float)FadeRegion);
+
+		if(PixelPosInt.x >= 0 && PixelPosInt.x < ElementsPerLine && PixelPosInt.y >= 0 && WhichBuffer / 2 == 0 && Alpha > 0)
+		{
+			uint Index = PixelPosInt.x + LineInBuffer * ElementsPerLine;
+
+			float4 Value = WhichBuffer ? BoneMatrices[Index] : PreviousBoneMatrices[Index];
+
+			OutColor.rgb = lerp(OutColor.rgb, Value.rgb, Alpha);
+		}
+	}
 #endif
 }
 
diff --git a/Engine/Source/Runtime/Engine/Private/GPUSkinVertexFactory.cpp b/Engine/Source/Runtime/Engine/Private/GPUSkinVertexFactory.cpp
--- a/Engine/Source/Runtime/Engine/Private/GPUSkinVertexFactory.cpp
+++ b/Engine/Source/Runtime/Engine/Private/GPUSkinVertexFactory.cpp
@@ -448,7 +448,10 @@ public:
 				// we tell the shader where to pickup the data (always, even if we don't have bone data, to avoid false binding)
 				if(PreviousBoneMatrices.IsBound())
 				{
-					RHICmdList.SetShaderResourceViewParameter(Shader->GetVertexShader(), PreviousBoneMatrices.GetBaseIndex(), GPrevPerBoneMotionBlur.GetReadData()->BoneBuffer.VertexBufferSRV);
+					RHICmdList.SetShaderResourceViewParameter(
+						Shader->GetVertexShader(),
+						PreviousBoneMatrices.GetBaseIndex(),
+						GPrevPerBoneMotionBlur.GetBoneDataVertexBuffer(GPrevPerBoneMotionBlur.GetReadBufferIndex())->BoneBuffer.VertexBufferSRV);
 				}
 
 				if(bLocalPerBoneMotionBlur)
diff --git a/Engine/Source/Runtime/Engine/Private/SkeletalRenderGPUSkin.cpp b/Engine/Source/Runtime/Engine/Private/SkeletalRenderGPUSkin.cpp
--- a/Engine/Source/Runtime/Engine/Private/SkeletalRenderGPUSkin.cpp
+++ b/Engine/Source/Runtime/Engine/Private/SkeletalRenderGPUSkin.cpp
@@ -1221,32 +1221,11 @@ void FPreviousPerBoneMotionBlur::EndAppend()
 	}
 }
 
-FBoneDataVertexBuffer* FPreviousPerBoneMotionBlur::GetReadData()
-{
-	return &PerChunkBoneMatricesTexture[GetReadBufferIndex()];
-}
-
 FString FPreviousPerBoneMotionBlur::GetDebugString() const
 {
 	return FString::Printf(TEXT("BufferIndex=%d Pos=%d"), GPrevPerBoneMotionBlur.GetWriteBufferIndex(), GPrevPerBoneMotionBlur.LockedTexelPosition.GetValue());
 }
 
-uint32 FPreviousPerBoneMotionBlur::GetReadBufferIndex() const
-{
-	return BufferIndex;
-}
-
-uint32 FPreviousPerBoneMotionBlur::GetWriteBufferIndex() const
-{
-	uint32 ret = BufferIndex + 1;
-
-	if(ret >= PER_BONE_BUFFER_COUNT)
-	{
-		ret = 0;
-	}
-	return ret;
-}
-
 void FPreviousPerBoneMotionBlur::AdvanceBufferIndex()
 {
 	++BufferIndex;
diff --git a/Engine/Source/Runtime/Engine/Private/SkeletalRenderGPUSkin.h b/Engine/Source/Runtime/Engine/Private/SkeletalRenderGPUSkin.h
--- a/Engine/Source/Runtime/Engine/Private/SkeletalRenderGPUSkin.h
+++ b/Engine/Source/Runtime/Engine/Private/SkeletalRenderGPUSkin.h
@@ -204,11 +204,38 @@ public:
 	/** only call if StartAppend(), if the append wasn't started it silently ignores the call */
 	ENGINE_API void EndAppend();
 
-	/** @return 0 if there should be no bone based motion blur (no previous data available or it's not active) */
-	FBoneDataVertexBuffer* GetReadData();
+	/**
+	 * @param Index 0 .. PER_BONE_BUFFER_COUNT - 1, usually from GetReadBufferIndex()
+	 * @return 0 if there should be no bone based motion blur (no previous data available or it's not active)
+	 */
+	FBoneDataVertexBuffer* GetBoneDataVertexBuffer(uint32 Index)
+	{
+		checkSlow(Index < PER_BONE_BUFFER_COUNT);
+
+		return &PerChunkBoneMatricesTexture[Index];
+	}
 
 	FString GetDebugString() const;
 
+	/** @return 0 .. PER_BONE_BUFFER_COUNT-1 */
+	uint32 GetReadBufferIndex() const
+	{
+		return BufferIndex;
+	}
+
+	/** @return 0 .. PER_BONE_BUFFER_COUNT-1 */
+	uint32 GetWriteBufferIndex() const
+	{
+		uint32 ret = BufferIndex + 1;
+
+		if(ret >= PER_BONE_BUFFER_COUNT)
+		{
+			ret = 0;
+		}
+		return ret;
+	}
+
+
 private:
 
 	/** Stores the bone information with one frame delay, required for per bone motion blur. Buffered to avoid stalls on texture Lock() */
@@ -224,11 +251,6 @@ private:
 
 	bool bWarningBufferSizeExceeded;
 
-	/** @return 0 .. PER_BONE_BUFFER_COUNT-1 */
-	uint32 GetReadBufferIndex() const;
-
-	/** @return 0 .. PER_BONE_BUFFER_COUNT-1 */
-	uint32 GetWriteBufferIndex() const;
 
 	/** to cycle the internal buffer counter */
 	void AdvanceBufferIndex();
diff --git a/Engine/Source/Runtime/Renderer/Private/PostProcess/PostProcessMotionBlur.cpp b/Engine/Source/Runtime/Renderer/Private/PostProcess/PostProcessMotionBlur.cpp
--- a/Engine/Source/Runtime/Renderer/Private/PostProcess/PostProcessMotionBlur.cpp
+++ b/Engine/Source/Runtime/Renderer/Private/PostProcess/PostProcessMotionBlur.cpp
@@ -12,6 +12,8 @@
 #include "PostProcessAmbientOcclusion.h"
 #include "PostProcessing.h"
 #include "SceneUtils.h"
+#include "GPUSkinVertexFactory.h"
+#include "../../Engine/Private/SkeletalRenderGPUSkin.h"
 
 #if !(UE_BUILD_SHIPPING || UE_BUILD_TEST)
 static TAutoConsoleVariable<int32> CVarMotionBlurFiltering(
@@ -293,6 +295,8 @@ public:
 	FShaderParameter PrevViewProjMatrix;
 	FShaderParameter TextureViewMad;
 	FShaderParameter MotionBlurParameters;
+	FShaderResourceParameter BoneMatrices0;
+	FShaderResourceParameter BoneMatrices1;
 
 	/** Initialization constructor. */
 	FPostProcessMotionBlurPS(const ShaderMetaType::CompiledShaderInitializerType& Initializer)
@@ -303,13 +307,15 @@ public:
 		PrevViewProjMatrix.Bind(Initializer.ParameterMap, TEXT("PrevViewProjMatrix"));
 		TextureViewMad.Bind(Initializer.ParameterMap, TEXT("TextureViewMad"));
 		MotionBlurParameters.Bind(Initializer.ParameterMap, TEXT("MotionBlurParameters"));
+		BoneMatrices0.Bind(Initializer.ParameterMap,TEXT("BoneMatrices0"));
+		BoneMatrices1.Bind(Initializer.ParameterMap,TEXT("BoneMatrices1"));
 	}
 
 	// FShader interface.
 	virtual bool Serialize(FArchive& Ar) override
 	{
 		bool bShaderHasOutdatedParameters = FGlobalShader::Serialize(Ar);
-		Ar << PostprocessParameter << DeferredParameters << PrevViewProjMatrix << TextureViewMad << MotionBlurParameters;
+		Ar << PostprocessParameter << DeferredParameters << PrevViewProjMatrix << TextureViewMad << MotionBlurParameters << BoneMatrices0 << BoneMatrices1;
 		return bShaderHasOutdatedParameters;
 	}
 
@@ -401,6 +407,9 @@ public:
 				- MaxVelocity * 2 * AspectRatio);
 			SetShaderValue(Context.RHICmdList, ShaderRHI, MotionBlurParameters, MotionBlurParametersValue);
 		}
+
+		SetSRVParameter(Context.RHICmdList, ShaderRHI, BoneMatrices0, GPrevPerBoneMotionBlur.GetBoneDataVertexBuffer(0)->BoneBuffer.VertexBufferSRV);
+		SetSRVParameter(Context.RHICmdList, ShaderRHI, BoneMatrices1, GPrevPerBoneMotionBlur.GetBoneDataVertexBuffer(1)->BoneBuffer.VertexBufferSRV);
 	}
 
 	static const TCHAR* GetSourceFilename()
@@ -432,7 +441,6 @@ static void SetMotionBlurShaderTempl(const FRenderingCompositePassContext& Conte
 
 	static FGlobalBoundShaderState BoundShaderState;
 	
-
 	SetGlobalBoundShaderState(Context.RHICmdList, Context.GetFeatureLevel(), BoundShaderState, GFilterVertexDeclaration.VertexDeclarationRHI, *VertexShader, *PixelShader);
 
 	VertexShader->SetParameters(Context);
@@ -569,13 +577,20 @@ public:
 		return bShaderHasOutdatedParameters;
 	}
 
-	void SetParameters(const FRenderingCompositePassContext& Context)
+	void SetParameters(const FRenderingCompositePassContext& Context, bool bBilinear)
 	{
 		const FPixelShaderRHIParamRef ShaderRHI = GetPixelShader();
 
 		FGlobalShader::SetParameters(Context.RHICmdList, ShaderRHI, Context.View);
 
-		PostprocessParameter.SetPS(ShaderRHI, Context, TStaticSamplerState<SF_Bilinear,AM_Border,AM_Border,AM_Clamp>::GetRHI());
+		if(bBilinear)
+		{
+			PostprocessParameter.SetPS(ShaderRHI, Context, TStaticSamplerState<SF_Bilinear, AM_Border, AM_Border, AM_Clamp>::GetRHI());
+		}
+		else
+		{
+			PostprocessParameter.SetPS(ShaderRHI, Context, TStaticSamplerState<SF_Point, AM_Border, AM_Border, AM_Clamp>::GetRHI());
+		}
 	}
 };
 
@@ -628,11 +643,13 @@ void FRCPassPostProcessMotionBlurRecombine::Process(FRenderingCompositePassConte
 
 	static FGlobalBoundShaderState BoundShaderState;
 	
-
 	SetGlobalBoundShaderState(Context.RHICmdList, Context.GetFeatureLevel(), BoundShaderState, GFilterVertexDeclaration.VertexDeclarationRHI, *VertexShader, *PixelShader);
 
+	// with point filtering we can better debug in Visualize MotionBlur
+	const bool bBilinear = !View.Family->EngineShowFlags.VisualizeMotionBlur;
+
 	VertexShader->SetParameters(Context);
-	PixelShader->SetParameters(Context);
+	PixelShader->SetParameters(Context, bBilinear);
 
 	// Draw a quad mapping scene color to the view's render target
 	DrawRectangle(
@@ -1494,15 +1511,33 @@ void FRCPassPostProcessVisualizeMotionBlur::Process(FRenderingCompositePassConte
 	FString Line;
 
 	Line = FString::Printf(TEXT("Visualize MotionBlur"));
-	Canvas.DrawShadowedString(X, Y += YStep, *Line, GetStatsFont(), FLinearColor(1, 1, 1));
+	Canvas.DrawShadowedString(X, Y += YStep, *Line, GetStatsFont(), FLinearColor(1, 1, 0));
 	
 	Line = FString::Printf(TEXT("%d"), ViewFamily.FrameNumber);
-	Canvas.DrawShadowedString(X, Y += YStep, TEXT("FrameNumber:"), GetStatsFont(), FLinearColor(1, 1, 1));
-	Canvas.DrawShadowedString(X + ColumnWidth, Y, *Line, GetStatsFont(), FLinearColor(1, 1, 1));
+	Canvas.DrawShadowedString(X, Y += YStep, TEXT("FrameNumber:"), GetStatsFont(), FLinearColor(1, 1, 0));
+	Canvas.DrawShadowedString(X + ColumnWidth, Y, *Line, GetStatsFont(), FLinearColor(1, 1, 0));
 
 	Line = FString::Printf(TEXT("%d"), ViewFamily.bWorldIsPaused ? 1 : 0);
-	Canvas.DrawShadowedString(X, Y += YStep, TEXT("WorldIsPaused:"), GetStatsFont(), FLinearColor(1, 1, 1));
-	Canvas.DrawShadowedString(X + ColumnWidth, Y, *Line, GetStatsFont(), FLinearColor(1, 1, 1));
+	Canvas.DrawShadowedString(X, Y += YStep, TEXT("WorldIsPaused:"), GetStatsFont(), FLinearColor(1, 1, 0));
+	Canvas.DrawShadowedString(X + ColumnWidth, Y, *Line, GetStatsFont(), FLinearColor(1, 1, 0));
+
+	for(uint32 BufferId = 0; BufferId < 2; ++BufferId)
+	{
+		const TCHAR* Usage = TEXT("unused");
+
+		if(BufferId == GPrevPerBoneMotionBlur.GetReadBufferIndex())
+		{
+			Usage = TEXT("read"); 
+		}
+		else if(BufferId == GPrevPerBoneMotionBlur.GetWriteBufferIndex())
+		{
+			Usage = TEXT("write"); 
+		}
+
+		Line = FString::Printf(TEXT("BoneBuffer %d: %s"), BufferId, Usage);
+		// LeftTop.y + (LinesPerBuffer + GapBetweenBuffers) * Scale
+		Canvas.DrawShadowedString(4, 74 + BufferId * (48 + 8) * 3, *Line, GetStatsFont(), FLinearColor(1, 1, 0));
+	}
 
 	Canvas.Flush_RenderThread(Context.RHICmdList);
 
