UE4 - Help track down a crash while rendering per object shadows (UE-14193)

diff --git a/Engine/Source/Runtime/Renderer/Private/ShadowSetup.cpp b/Engine/Source/Runtime/Renderer/Private/ShadowSetup.cpp
--- a/Engine/Source/Runtime/Renderer/Private/ShadowSetup.cpp
+++ b/Engine/Source/Runtime/Renderer/Private/ShadowSetup.cpp
@@ -1313,6 +1313,8 @@ void FDeferredShadingSceneRenderer::CreatePerObjectProjectedShadow(
 		const float ScreenRadius = View.ShadowViewMatrices.ScreenScale *
 			OriginalBounds.SphereRadius /
 			FMath::Max(ScreenPosition.W,1.0f);
+		// Early catch for invalid CalculateShadowFadeAlpha()
+		checkf(ScreenRadius >= 0.0f, TEXT("View.ShadowViewMatrices.ScreenScale %f, OriginalBounds.SphereRadius %f, ScreenPosition.W %f"), View.ShadowViewMatrices.ScreenScale, OriginalBounds.SphereRadius, ScreenPosition.W);
 
 		const float ScreenPercent = FMath::Max(
 			1.0f / 2.0f * View.ShadowViewMatrices.ProjectionScale.X,
