fixed pause motionblur of static meshes (per object motionblur)

diff --git a/Engine/Shaders/PostProcessMotionBlur.usf b/Engine/Shaders/PostProcessMotionBlur.usf
--- a/Engine/Shaders/PostProcessMotionBlur.usf
+++ b/Engine/Shaders/PostProcessMotionBlur.usf
@@ -616,7 +616,7 @@ void MainPS(float4 UVAndScreenPos : TEXCOORD0, out float4 OutColor : SV_Target0)
 	}
 	// visualize BoneBuffers
 	{
-		const int2 LeftTop = int2(16, 90);
+		const int2 LeftTop = int2(16, 100);
 		const int ElementsPerLine = 32;
 		const int Scale = 3;
 		const int LinesPerBuffer = 48;
diff --git a/Engine/Source/Runtime/Engine/Private/SkeletalRenderGPUSkin.cpp b/Engine/Source/Runtime/Engine/Private/SkeletalRenderGPUSkin.cpp
--- a/Engine/Source/Runtime/Engine/Private/SkeletalRenderGPUSkin.cpp
+++ b/Engine/Source/Runtime/Engine/Private/SkeletalRenderGPUSkin.cpp
@@ -1097,14 +1097,6 @@ void FPreviousPerBoneMotionBlur::ReleaseResources()
 	}
 }
 
-void FPreviousPerBoneMotionBlur::RestoreForPausedMotionBlur()
-{
-	if(CVarMotionBlurDebug.GetValueOnRenderThread())
-	{
-		UE_LOG(LogEngine, Log, TEXT("r.MotionBlurDebug: RestoreForPausedMotionBlur"));
-	}
-}
-
 uint32 FPreviousPerBoneMotionBlur::GetSizeX() const
 {
 	return PerChunkBoneMatricesTexture[0].GetSizeX();
diff --git a/Engine/Source/Runtime/Engine/Private/SkeletalRenderGPUSkin.h b/Engine/Source/Runtime/Engine/Private/SkeletalRenderGPUSkin.h
--- a/Engine/Source/Runtime/Engine/Private/SkeletalRenderGPUSkin.h
+++ b/Engine/Source/Runtime/Engine/Private/SkeletalRenderGPUSkin.h
@@ -179,9 +179,6 @@ public:
 	*/
 	void ReleaseResources();
 
-	/** */
-	ENGINE_API void RestoreForPausedMotionBlur();
-
 	void InitIfNeeded();
 
 	/** Returns the width of the texture in pixels. */
diff --git a/Engine/Source/Runtime/Engine/Public/SceneManagement.h b/Engine/Source/Runtime/Engine/Public/SceneManagement.h
--- a/Engine/Source/Runtime/Engine/Public/SceneManagement.h
+++ b/Engine/Source/Runtime/Engine/Public/SceneManagement.h
@@ -1451,9 +1451,6 @@ struct FMotionBlurInfo
 	/**  */
 	void UpdateMotionBlurInfo();
 
-	/** Call if you want to keep the existing motionblur */
-	void RestoreForPausedMotionBlur();
-
 	void SetKeepAndUpdateThisFrame(bool bValue = true)
 	{
 		if(bValue)
@@ -1488,7 +1485,12 @@ struct FMotionBlurInfo
 	void ApplyOffset(FVector InOffset)
 	{
 		PreviousLocalToWorld.SetOrigin(PreviousLocalToWorld.GetOrigin() + InOffset);
-		PausedLocalToWorld.SetOrigin(PausedLocalToWorld.GetOrigin() + InOffset);
+		CurrentLocalToWorld.SetOrigin(CurrentLocalToWorld.GetOrigin() + InOffset);
+	}
+
+	void OnStartFrame()
+	{
+		PreviousLocalToWorld = CurrentLocalToWorld;
 	}
 
 private:
@@ -1498,8 +1500,8 @@ private:
 	FPrimitiveSceneInfo* MBPrimitiveSceneInfo;
 	/** The previous LocalToWorld of the component.	*/
 	FMatrix	PreviousLocalToWorld;
-	/** Used in case when Pause is activate. */
-	FMatrix	PausedLocalToWorld;
+	/** todo */
+	FMatrix	CurrentLocalToWorld;
 	/** if true then the PreviousLocalToWorld has already been updated for the current frame */
 	bool bKeepAndUpdateThisFrame;
 };
@@ -1532,10 +1534,7 @@ public:
 	 */
 	void UpdateMotionBlurCache(class FScene* InScene);
 
-	/**
-	 * Call if you want to keep the existing motionblur
-	 */
-	void RestoreForPausedMotionBlur();
+	void StartFrame(bool bWorldIsPaused);
 
 	/** 
 	 *	Get the primitives motion blur info
@@ -1555,11 +1554,18 @@ public:
 	 */
 	void ApplyOffset(FVector InOffset);
 
+	/**
+	 * Get some debug string for VisualizeMotionBlur
+	 */
+	FString GetDebugString() const;
+
 private:
 	/** The motion blur info entries for the frame. Accessed on Renderthread only! */
 	TMap<FPrimitiveComponentId, FMotionBlurInfo> MotionBlurInfos;
 	/** */
 	bool bShouldClearMotionBlurInfo;
+	/** set in StartFrame() */
+	bool bWorldIsPaused;
 
 	/**
 	 * O(n) with the amount of motion blurred objects but that number should be low
diff --git a/Engine/Source/Runtime/Renderer/Private/DeferredShadingRenderer.cpp b/Engine/Source/Runtime/Renderer/Private/DeferredShadingRenderer.cpp
--- a/Engine/Source/Runtime/Renderer/Private/DeferredShadingRenderer.cpp
+++ b/Engine/Source/Runtime/Renderer/Private/DeferredShadingRenderer.cpp
@@ -799,6 +799,11 @@ void FDeferredShadingSceneRenderer::Render(FRHICommandListImmediate& RHICmdList)
 		FGlobalDynamicIndexBuffer::Get().Commit();
 	}
 
+	{
+		QUICK_SCOPE_CYCLE_COUNTER(STAT_FDeferredShadingSceneRenderer_MotionBlurStartFrame);
+		Scene->MotionBlurInfoData.StartFrame(ViewFamily.bWorldIsPaused);
+	}
+
 	// Notify the FX system that the scene is about to be rendered.
 	if (Scene->FXSystem)
 	{
diff --git a/Engine/Source/Runtime/Renderer/Private/PostProcess/PostProcessMotionBlur.cpp b/Engine/Source/Runtime/Renderer/Private/PostProcess/PostProcessMotionBlur.cpp
--- a/Engine/Source/Runtime/Renderer/Private/PostProcess/PostProcessMotionBlur.cpp
+++ b/Engine/Source/Runtime/Renderer/Private/PostProcess/PostProcessMotionBlur.cpp
@@ -1529,6 +1529,11 @@ void FRCPassPostProcessVisualizeMotionBlur::Process(FRenderingCompositePassConte
 	Canvas.DrawShadowedString(X, Y += YStep, TEXT("WorldIsPaused, r.VelocityTest, Parallel:"), GetStatsFont(), FLinearColor(1, 1, 0));
 	Canvas.DrawShadowedString(X + ColumnWidth, Y, *Line, GetStatsFont(), FLinearColor(1, 1, 0));
 
+	const FScene* Scene = (const FScene*)View.Family->Scene;
+
+	Canvas.DrawShadowedString(X, Y += YStep, TEXT("MotionBlurInfoData (per object):"), GetStatsFont(), FLinearColor(1, 1, 0));
+	Canvas.DrawShadowedString(X + ColumnWidth, Y, *Scene->MotionBlurInfoData.GetDebugString(), GetStatsFont(), FLinearColor(1, 1, 0));	
+
 	for(uint32 BufferId = 0; BufferId < 2; ++BufferId)
 	{
 		const TCHAR* Usage = TEXT("unused");
@@ -1544,7 +1549,7 @@ void FRCPassPostProcessVisualizeMotionBlur::Process(FRenderingCompositePassConte
 
 		Line = FString::Printf(TEXT("BoneBuffer %d: %s"), BufferId, Usage);
 		// LeftTop.y + (LinesPerBuffer + GapBetweenBuffers) * Scale
-		Canvas.DrawShadowedString(4, 74 + BufferId * (48 + 8) * 3, *Line, GetStatsFont(), FLinearColor(1, 1, 0));
+		Canvas.DrawShadowedString(4, 84 + BufferId * (48 + 8) * 3, *Line, GetStatsFont(), FLinearColor(1, 1, 0));
 	}
 
 	Canvas.Flush_RenderThread(Context.RHICmdList);
diff --git a/Engine/Source/Runtime/Renderer/Private/RendererScene.cpp b/Engine/Source/Runtime/Renderer/Private/RendererScene.cpp
--- a/Engine/Source/Runtime/Renderer/Private/RendererScene.cpp
+++ b/Engine/Source/Runtime/Renderer/Private/RendererScene.cpp
@@ -2617,8 +2617,8 @@ TStaticMeshDrawList<TBasePassForForwardShadingDrawingPolicy<FMovableDirectionalL
 
 FMotionBlurInfoData::FMotionBlurInfoData()
 	: bShouldClearMotionBlurInfo(false)
+	, bWorldIsPaused(false)
 {
-
 }
 
 void FMotionBlurInfoData::UpdatePrimitiveMotionBlur(FPrimitiveSceneInfo* PrimitiveSceneInfo)
@@ -2673,31 +2673,28 @@ void FMotionBlurInfo::UpdateMotionBlurInfo()
 {
 	if(MBPrimitiveSceneInfo && MBPrimitiveSceneInfo->Proxy)
 	{
-		PausedLocalToWorld = PreviousLocalToWorld;
 		// only if the proxy is still there
-		PreviousLocalToWorld = MBPrimitiveSceneInfo->Proxy->GetLocalToWorld();
+		CurrentLocalToWorld = MBPrimitiveSceneInfo->Proxy->GetLocalToWorld();
 	}
 
 	bKeepAndUpdateThisFrame = false;
 }
 
-void FMotionBlurInfo::RestoreForPausedMotionBlur()
-{
-	PreviousLocalToWorld = PausedLocalToWorld;
-}
-
 // Doxygen has trouble parsing these functions because the header declaring them is in Engine, not Renderer
 #if !UE_BUILD_DOCS
 
-void FMotionBlurInfoData::RestoreForPausedMotionBlur()
+void FMotionBlurInfoData::StartFrame(bool bInWorldIsPaused)
 {
-	check(IsInRenderingThread());
+	bWorldIsPaused = bInWorldIsPaused;
 
-	for (TMap<FPrimitiveComponentId, FMotionBlurInfo>::TIterator It(MotionBlurInfos); It; ++It)
+	if(!bWorldIsPaused)
 	{
-		FMotionBlurInfo& MotionBlurInfo = It.Value();
+		for (TMap<FPrimitiveComponentId, FMotionBlurInfo>::TIterator It(MotionBlurInfos); It; ++It)
+		{
+			FMotionBlurInfo& MotionBlurInfo = It.Value();
 
-		MotionBlurInfo.RestoreForPausedMotionBlur();
+			MotionBlurInfo.OnStartFrame();
+		}
 	}
 }
 
@@ -2705,6 +2702,11 @@ void FMotionBlurInfoData::UpdateMotionBlurCache(FScene* InScene)
 {
 	check(InScene && IsInRenderingThread());
 
+	if(bWorldIsPaused)
+	{
+		return;
+	}
+
 	if (InScene->GetFeatureLevel() >= ERHIFeatureLevel::SM4)
 	{
 		if(bShouldClearMotionBlurInfo)
@@ -2745,6 +2747,11 @@ void FMotionBlurInfoData::ApplyOffset(FVector InOffset)
 	}
 }
 
+FString FMotionBlurInfoData::GetDebugString() const
+{
+	return FString::Printf(TEXT("Num=%d Clear=%d"), MotionBlurInfos.Num(), bShouldClearMotionBlurInfo);
+}
+
 const FMotionBlurInfo* FMotionBlurInfoData::FindMBInfoIndex(FPrimitiveComponentId ComponentId) const
 {
 	return MotionBlurInfos.Find(ComponentId);
diff --git a/Engine/Source/Runtime/Renderer/Private/SceneVisibility.cpp b/Engine/Source/Runtime/Renderer/Private/SceneVisibility.cpp
--- a/Engine/Source/Runtime/Renderer/Private/SceneVisibility.cpp
+++ b/Engine/Source/Runtime/Renderer/Private/SceneVisibility.cpp
@@ -2148,16 +2148,6 @@ void FDeferredShadingSceneRenderer::InitViews(FRHICommandListImmediate& RHICmdLi
 		InitDynamicShadows(RHICmdList);
 	}
 
-	if (ViewFamily.bWorldIsPaused)
-	{
-		// so we can freeze motion blur and TemporalAA in paused mode
-
-		// per object velocity (static meshes)
-		Scene->MotionBlurInfoData.RestoreForPausedMotionBlur();
-		// per bone velocity (skeletal meshes)
-		GPrevPerBoneMotionBlur.RestoreForPausedMotionBlur();
-	}
-
 	// initialize per-view uniform buffer.
 	for (int32 ViewIndex = 0; ViewIndex < Views.Num(); ViewIndex++)
 	{
