fixed pause of MotionBlur for camera motion

diff --git a/Engine/Shaders/PostProcessMotionBlur.usf b/Engine/Shaders/PostProcessMotionBlur.usf
--- a/Engine/Shaders/PostProcessMotionBlur.usf
+++ b/Engine/Shaders/PostProcessMotionBlur.usf
@@ -616,7 +616,7 @@ void MainPS(float4 UVAndScreenPos : TEXCOORD0, out float4 OutColor : SV_Target0)
 	}
 	// visualize BoneBuffers
 	{
-		const int2 LeftTop = int2(16, 100);
+		const int2 LeftTop = int2(16, 114);
 		const int ElementsPerLine = 32;
 		const int Scale = 3;
 		const int LinesPerBuffer = 48;
diff --git a/Engine/Source/Runtime/Core/Private/Math/UnrealMath.cpp b/Engine/Source/Runtime/Core/Private/Math/UnrealMath.cpp
--- a/Engine/Source/Runtime/Core/Private/Math/UnrealMath.cpp
+++ b/Engine/Source/Runtime/Core/Private/Math/UnrealMath.cpp
@@ -411,6 +411,20 @@ void FMatrix::DebugPrint() const
 	UE_LOG(LogUnrealMath, Log, TEXT("%s"), *ToString());
 }
 
+uint32 FMatrix::ComputeHash() const
+{
+	uint32 Ret = 0;
+
+	const uint32* Data = (uint32*)this;
+
+	for(uint32 i = 0; i < 16; ++i)
+	{
+		Ret ^= Data[i] + i;
+	}
+
+	return Ret;
+}
+
 //////////////////////////////////////////////////////////////////////////
 // FQuat
 
diff --git a/Engine/Source/Runtime/Core/Public/Math/Matrix.h b/Engine/Source/Runtime/Core/Public/Math/Matrix.h
--- a/Engine/Source/Runtime/Core/Public/Math/Matrix.h
+++ b/Engine/Source/Runtime/Core/Public/Math/Matrix.h
@@ -309,6 +309,9 @@ public:
 	/** Output ToString */
 	void DebugPrint() const;
 
+	/** For debugging purpose, could be changed */
+	CORE_API uint32 ComputeHash() const; 
+
 	/**
 	 * Serializes the Matrix.
 	 *
diff --git a/Engine/Source/Runtime/Renderer/Private/PostProcess/PostProcessMotionBlur.cpp b/Engine/Source/Runtime/Renderer/Private/PostProcess/PostProcessMotionBlur.cpp
--- a/Engine/Source/Runtime/Renderer/Private/PostProcess/PostProcessMotionBlur.cpp
+++ b/Engine/Source/Runtime/Renderer/Private/PostProcess/PostProcessMotionBlur.cpp
@@ -1503,10 +1503,10 @@ void FRCPassPostProcessVisualizeMotionBlur::Process(FRenderingCompositePassConte
 
 	FCanvas Canvas(&TempRenderTarget, NULL, ViewFamily.CurrentRealTime, ViewFamily.CurrentWorldTime, ViewFamily.DeltaWorldTime, Context.GetFeatureLevel());
 
-	float X = 30;
+	float X = 20;
 	float Y = 8;
 	const float YStep = 14;
-	const float ColumnWidth = 250;
+	const float ColumnWidth = 200;
 
 	FString Line;
 
@@ -1517,7 +1517,7 @@ void FRCPassPostProcessVisualizeMotionBlur::Process(FRenderingCompositePassConte
 	const int32 MotionBlurDebug = MotionBlurDebugVar ? MotionBlurDebugVar->GetValueOnRenderThread() : 0;
 
 	Line = FString::Printf(TEXT("%d, %d"), ViewFamily.FrameNumber, MotionBlurDebug);
-	Canvas.DrawShadowedString(X, Y += YStep, TEXT("FrameNumber, r.MotionBlurDebug:"), GetStatsFont(), FLinearColor(1, 1, 0));
+	Canvas.DrawShadowedString(X, Y += YStep, TEXT("FrameNo, r.MotionBlurDebug:"), GetStatsFont(), FLinearColor(1, 1, 0));
 	Canvas.DrawShadowedString(X + ColumnWidth, Y, *Line, GetStatsFont(), FLinearColor(1, 1, 0));
 
 	static const auto VelocityTestVar = IConsoleManager::Get().FindTConsoleVariableDataInt(TEXT("r.VelocityTest"));
@@ -1525,8 +1525,8 @@ void FRCPassPostProcessVisualizeMotionBlur::Process(FRenderingCompositePassConte
 	
 	extern bool IsParallelVelocity();
 
-	Line = FString::Printf(TEXT("%d, %d, %d"), ViewFamily.bWorldIsPaused ? 1 : 0, VelocityTest, IsParallelVelocity());
-	Canvas.DrawShadowedString(X, Y += YStep, TEXT("WorldIsPaused, r.VelocityTest, Parallel:"), GetStatsFont(), FLinearColor(1, 1, 0));
+	Line = FString::Printf(TEXT("%d, %d, %d"), ViewFamily.bWorldIsPaused, VelocityTest, IsParallelVelocity());
+	Canvas.DrawShadowedString(X, Y += YStep, TEXT("Paused, r.VelocityTest, Parallel:"), GetStatsFont(), FLinearColor(1, 1, 0));
 	Canvas.DrawShadowedString(X + ColumnWidth, Y, *Line, GetStatsFont(), FLinearColor(1, 1, 0));
 
 	const FScene* Scene = (const FScene*)View.Family->Scene;
@@ -1534,6 +1534,14 @@ void FRCPassPostProcessVisualizeMotionBlur::Process(FRenderingCompositePassConte
 	Canvas.DrawShadowedString(X, Y += YStep, TEXT("MotionBlurInfoData (per object):"), GetStatsFont(), FLinearColor(1, 1, 0));
 	Canvas.DrawShadowedString(X + ColumnWidth, Y, *Scene->MotionBlurInfoData.GetDebugString(), GetStatsFont(), FLinearColor(1, 1, 0));	
 
+	const FSceneViewState *SceneViewState = (const FSceneViewState*)View.State;
+
+	Line = FString::Printf(TEXT("View=%.4x PrevView=%.4x"),
+		View.ViewMatrices.ViewMatrix.ComputeHash() & 0xffff,
+		SceneViewState->PrevViewMatrices.ViewMatrix.ComputeHash() & 0xffff);
+	Canvas.DrawShadowedString(X, Y += YStep, TEXT("ViewMatrix:"), GetStatsFont(), FLinearColor(1, 1, 0));
+	Canvas.DrawShadowedString(X + ColumnWidth, Y, *Line, GetStatsFont(), FLinearColor(1, 1, 0));
+
 	for(uint32 BufferId = 0; BufferId < 2; ++BufferId)
 	{
 		const TCHAR* Usage = TEXT("unused");
@@ -1549,7 +1557,7 @@ void FRCPassPostProcessVisualizeMotionBlur::Process(FRenderingCompositePassConte
 
 		Line = FString::Printf(TEXT("BoneBuffer %d: %s"), BufferId, Usage);
 		// LeftTop.y + (LinesPerBuffer + GapBetweenBuffers) * Scale
-		Canvas.DrawShadowedString(4, 84 + BufferId * (48 + 8) * 3, *Line, GetStatsFont(), FLinearColor(1, 1, 0));
+		Canvas.DrawShadowedString(4, 98 + BufferId * (48 + 8) * 3, *Line, GetStatsFont(), FLinearColor(1, 1, 0));
 	}
 
 	Canvas.Flush_RenderThread(Context.RHICmdList);
