Do not bind the GBuffer in hit proxy rendering, avoiding a check assertion failure at start-up time with r.RenderTargetPoolTest=1

diff --git a/Engine/Shaders/Common.usf b/Engine/Shaders/Common.usf
--- a/Engine/Shaders/Common.usf
+++ b/Engine/Shaders/Common.usf
@@ -697,10 +697,15 @@ float2 CalcScreenUVFromOffsetFraction(float4 ScreenPosition, float2 OffsetFracti
 /** Applies an offset to the scene texture lookup and decodes the HDR linear space color. */
 float3 DecodeSceneColorForMaterialNode(float2 ScreenUV)
 {
+#if HIT_PROXY_SHADER
+	// Hit proxies rendering pass doesn't have access to valide render buffers
+	return float3(0.0f, 0.0f, 0.0f);
+#else
 	float4 EncodedSceneColor = Texture2DSample(LightAttenuationTexture, LightAttenuationTextureSampler, ScreenUV);
 
 	// Undo the function in EncodeSceneColorForMaterialNode
 	return pow(EncodedSceneColor.rgb, 4) * 10;
+#endif
 }
 
 // Converts depth buffer Z / W into W
diff --git a/Engine/Shaders/HitProxyPixelShader.usf b/Engine/Shaders/HitProxyPixelShader.usf
--- a/Engine/Shaders/HitProxyPixelShader.usf
+++ b/Engine/Shaders/HitProxyPixelShader.usf
@@ -4,6 +4,9 @@
 	HitProxyPixelShader.hlsl: Pixel shader for rendering scene hit proxies.
 =============================================================================*/
 
+// Some input nodes can't compute their output value at hit proxy rendering time, and so their implementation changes.
+#define HIT_PROXY_SHADER 1
+
 #include "Common.usf"
 #include "Material.usf"
 #include "VertexFactory.usf"
diff --git a/Engine/Shaders/HitProxyVertexShader.usf b/Engine/Shaders/HitProxyVertexShader.usf
--- a/Engine/Shaders/HitProxyVertexShader.usf
+++ b/Engine/Shaders/HitProxyVertexShader.usf
@@ -4,6 +4,9 @@
 	HitProxyVertexShader.hlsl: Vertex shader for rendering hit proxies.
 =============================================================================*/
 
+// Some input nodes can't compute their output value at hit proxy rendering time, and so their implementation changes.
+#define HIT_PROXY_SHADER 1
+
 #include "Common.usf"
 #include "Material.usf"
 #include "VertexFactory.usf"
diff --git a/Engine/Source/Runtime/Renderer/Private/SceneHitProxyRendering.cpp b/Engine/Source/Runtime/Renderer/Private/SceneHitProxyRendering.cpp
--- a/Engine/Source/Runtime/Renderer/Private/SceneHitProxyRendering.cpp
+++ b/Engine/Source/Runtime/Renderer/Private/SceneHitProxyRendering.cpp
@@ -380,8 +380,6 @@ void RenderHitProxies(FRHICommandListImmediate& RHICmdList, const FSceneRenderer
 
 	const bool bNeedToSwitchVerticalAxis = RHINeedsToSwitchVerticalAxis(GShaderPlatformForFeatureLevel[SceneRenderer->FeatureLevel]);
 
-	GSceneRenderTargets.AdjustGBufferRefCount(1);
-
 	for (int32 ViewIndex = 0; ViewIndex < Views.Num(); ViewIndex++)
 	{
 		const FViewInfo& View = Views[ViewIndex];
@@ -459,8 +457,6 @@ void RenderHitProxies(FRHICommandListImmediate& RHICmdList, const FSceneRenderer
 		View.TopBatchedViewElements.Draw(RHICmdList, FeatureLevel, bNeedToSwitchVerticalAxis, View.ViewProjectionMatrix, View.ViewRect.Width(), View.ViewRect.Height(), true);
 	}
 
-	GSceneRenderTargets.AdjustGBufferRefCount(-1);
-
 	// Finish drawing to the hit proxy render target.
 	RHICmdList.CopyToResolveTarget(HitProxyRT->GetRenderTargetItem().TargetableTexture, HitProxyRT->GetRenderTargetItem().ShaderResourceTexture, false, FResolveParams());
 	// to be able to observe results with VisualizeTexture
