New tone mapper. ACES support. Color grading. White balance. Baked into 32x32x32 LUT.

diff --git a/Engine/Shaders/ACES.usf b/Engine/Shaders/ACES.usf
new file mode 100644
--- /dev/null
+++ b/Engine/Shaders/ACES.usf
@@ -0,0 +1,699 @@
+// Copyright 1998-2015 Epic Games, Inc. All Rights Reserved.
+
+/*
+=============================================================================
+	ACES: Academy Color Encoding System
+	[REDACTED_PROFILE_URL]/aces-dev/tree/v1.0
+	
+	License Terms for Academy Color Encoding System Components
+
+	Academy Color Encoding System (ACES) software and tools are provided by the Academy under 
+	the following terms and conditions: A worldwide, royalty-free, non-exclusive right to copy, modify, create
+	derivatives, and use, in source and binary forms, is hereby granted, subject to acceptance of this license.
+
+	Copyright © 2013 Academy of Motion Picture Arts and Sciences (A.M.P.A.S.). Portions contributed by
+	others as indicated. All rights reserved.
+
+	Performance of any of the aforementioned acts indicates acceptance to be bound by the following 
+	terms and conditions:
+
+	 *	Copies of source code, in whole or in part, must retain the above copyright 
+		notice, this list of conditions and the Disclaimer of Warranty.
+	 *	Use in binary form must retain the above copyright notice, this list of 
+		conditions and the Disclaimer of Warranty in the documentation and/or other 
+		materials provided with the distribution.
+	 *	Nothing in this license shall be deemed to grant any rights to trademarks,
+		copyrights, patents, trade secrets or any other intellectual property of 
+		A.M.P.A.S. or any contributors, except as expressly stated herein.
+	 *	Neither the name "A.M.P.A.S." nor the name of any other contributors to this 
+		software may be used to endorse or promote products derivative of or based on
+		this software without express prior written permission of A.M.P.A.S. or the
+		contributors, as appropriate.
+
+	This license shall be construed pursuant to the laws of the State of California,
+	and any disputes related thereto shall be subject to the jurisdiction of the courts therein.
+
+	Disclaimer of Warranty: THIS SOFTWARE IS PROVIDED BY A.M.P.A.S. AND CONTRIBUTORS "AS
+	IS" AND ANY EXPRESS OR IMPLIED WARRANTIES, INCLUDING, BUT NOT LIMITED TO, THE
+	IMPLIED WARRANTIES OF MERCHANTABILITY, FITNESS FOR A PARTICULAR PURPOSE, AND 
+	NON-INFRINGEMENT ARE DISCLAIMED. IN NO EVENT SHALL A.M.P.A.S., OR ANY 
+	CONTRIBUTORS OR DISTRIBUTORS, BE LIABLE FOR ANY DIRECT, INDIRECT, INCIDENTAL,
+	SPECIAL, EXEMPLARY, RESITUTIONARY, OR CONSEQUENTIAL DAMAGES (INCLUDING, BUT 
+	NOT LIMITED TO, PROCUREMENT OF SUBSTITUTE GOODS OR SERVICES; LOSS OF USE, 
+	DATA, OR PROFITS; OR BUSINESS INTERRUPTION) HOWEVER CAUSED AND ON ANY THEORY
+	OF LIABILITY, WHETHER IN CONTRACT, STRICT LIABILITY, OR TORT (INCLUDING 
+	NEGLIGENCE OR OTHERWISE) ARISING IN ANY WAY OUT OF THE USE OF THIS SOFTWARE,
+	EVEN IF ADVISED OF THE POSSIBILITY OF SUCH DAMAGE.
+
+	WITHOUT LIMITING THE GENERALITY OF THE FOREGOING, THE ACADEMY SPECIFICALLY 
+	DISCLAIMS ANY REPRESENTATIONS OR WARRANTIES WHATSOEVER RELATED TO PATENT OR 
+	OTHER INTELLECTUAL PROPERTY RIGHTS IN THE ACADEMY COLOR ENCODING SYSTEM, OR 
+	APPLICATIONS THEREOF, HELD BY PARTIES OTHER THAN A.M.P.A.S.,WHETHER DISCLOSED
+	OR UNDISCLOSED.
+=============================================================================
+*/
+
+static const float3x3 AP0_2_XYZ_MAT =
+{
+	0.9525523959, 0.0000000000, 0.0000936786,
+	0.3439664498, 0.7281660966,-0.0721325464,
+	0.0000000000, 0.0000000000, 1.0088251844,
+};
+
+static const float3x3 XYZ_2_AP0_MAT =
+{
+	 1.0498110175, 0.0000000000,-0.0000974845,
+	-0.4959030231, 1.3733130458, 0.0982400361,
+	 0.0000000000, 0.0000000000, 0.9912520182,
+};
+
+static const float3x3 AP1_2_XYZ_MAT = 
+{
+	 0.6624541811, 0.1340042065, 0.1561876870,
+	 0.2722287168, 0.6740817658, 0.0536895174,
+	-0.0055746495, 0.0040607335, 1.0103391003,
+};
+
+static const float3x3 XYZ_2_AP1_MAT =
+{
+	 1.6410233797, -0.3248032942, -0.2364246952,
+	-0.6636628587,  1.6153315917,  0.0167563477,
+	 0.0117218943, -0.0082844420,  0.9883948585,
+};
+
+static const float3x3 AP0_2_AP1_MAT = //mul( AP0_2_XYZ_MAT, XYZ_2_AP1_MAT );
+{
+	 1.4514393161, -0.2365107469, -0.2149285693,
+	-0.0765537734,  1.1762296998, -0.0996759264,
+	 0.0083161484, -0.0060324498,  0.9977163014,
+};
+
+static const float3x3 AP1_2_AP0_MAT = //mul( AP1_2_XYZ_MAT, XYZ_2_AP0_MAT );
+{
+	 0.6954522414,  0.1406786965,  0.1638690622,
+	 0.0447945634,  0.8596711185,  0.0955343182,
+	-0.0055258826,  0.0040252103,  1.0015006723,
+};
+
+static const float3 AP1_RGB2Y =
+{
+	0.2722287168, //AP1_2_XYZ_MAT[0][1],
+	0.6740817658, //AP1_2_XYZ_MAT[1][1],
+	0.0536895174, //AP1_2_XYZ_MAT[2][1]
+};
+
+// REC 709 primaries
+static const float3x3 XYZ_2_sRGB_MAT =
+{
+	 3.2409699419, -1.5373831776, -0.4986107603,
+	-0.9692436363,  1.8759675015,  0.0415550574,
+	 0.0556300797, -0.2039769589,  1.0569715142,
+};
+
+static const float3x3 sRGB_2_XYZ_MAT =
+{
+	0.4124564, 0.3575761, 0.1804375,
+	0.2126729, 0.7151522, 0.0721750,
+	0.0193339, 0.1191920, 0.9503041,
+};
+
+// Bradford chromatic adaptation transforms between ACES white point (D60) and sRGB white point (D65)
+static const float3x3 D65_2_D60_CAT =
+{
+	 1.01303,    0.00610531, -0.014971,
+	 0.00769823, 0.998165,   -0.00503203,
+	-0.00284131, 0.00468516,  0.924507,
+};
+
+static const float3x3 D60_2_D65_CAT =
+{
+	 0.987224,   -0.00611327, 0.0159533,
+	-0.00759836,  1.00186,    0.00533002,
+	 0.00307257, -0.00509595, 1.08168,
+};
+
+float rgb_2_saturation( float3 rgb )
+{
+	float minrgb = min( min(rgb.r, rgb.g ), rgb.b );
+	float maxrgb = max( max(rgb.r, rgb.g ), rgb.b );
+	return ( max( maxrgb, 1e-10 ) - max( minrgb, 1e-10 ) ) / max( maxrgb, 1e-2 );
+}
+
+float glow_fwd( float ycIn, float glowGainIn, float glowMid)
+{
+   float glowGainOut;
+
+   if (ycIn <= 2./3. * glowMid) {
+     glowGainOut = glowGainIn;
+   } else if ( ycIn >= 2 * glowMid) {
+     glowGainOut = 0;
+   } else {
+     glowGainOut = glowGainIn * (glowMid / ycIn - 0.5);
+   }
+
+   return glowGainOut;
+}
+
+float sigmoid_shaper( float x)
+{
+	// Sigmoid function in the range 0 to 1 spanning -2 to +2.
+
+	float t = max( 1 - abs( 0.5 * x ), 0 );
+	float y = 1 + sign(x) * (1 - t*t);
+	return 0.5 * y;
+}
+
+
+// ------- Red modifier functions
+float cubic_basis_shaper
+( 
+  float x, 
+  float w   // full base width of the shaper function (in degrees)
+)
+{
+	//return Square( smoothstep( 0, 1, 1 - abs( 2 * x/w ) ) );
+
+	float M[4][4] =
+	{
+		{ -1./6,  3./6, -3./6,  1./6 },
+		{  3./6, -6./6,  3./6,  0./6 },
+		{ -3./6,  0./6,  3./6,  0./6 },
+		{  1./6,  4./6,  1./6,  0./6 }
+	};
+  
+	float knots[5] = { -0.5 * w, -0.25 * w, 0, 0.25 * w, 0.5 * w };
+  
+	float y = 0;
+	if ((x > knots[0]) && (x < knots[4]))
+	{  
+		float knot_coord = (x - knots[0]) * 4.0 / w;
+		int j = knot_coord;
+		float t = knot_coord - j;
+      
+		float monomials[4] = { t*t*t, t*t, t, 1. };
+
+		// (if/else structure required for compatibility with CTL < v1.5.)
+		if ( j == 3) {
+			y = monomials[0] * M[0][0] + monomials[1] * M[1][0] + 
+				monomials[2] * M[2][0] + monomials[3] * M[3][0];
+		} else if ( j == 2) {
+			y = monomials[0] * M[0][1] + monomials[1] * M[1][1] + 
+				monomials[2] * M[2][1] + monomials[3] * M[3][1];
+		} else if ( j == 1) {
+			y = monomials[0] * M[0][2] + monomials[1] * M[1][2] + 
+				monomials[2] * M[2][2] + monomials[3] * M[3][2];
+		} else if ( j == 0) {
+			y = monomials[0] * M[0][3] + monomials[1] * M[1][3] + 
+				monomials[2] * M[2][3] + monomials[3] * M[3][3];
+		} else {
+			y = 0.0;
+		}
+	}
+  
+	return y * 1.5;
+}
+
+float center_hue( float hue, float centerH)
+{
+	float hueCentered = hue - centerH;
+	if (hueCentered < -180.)
+		hueCentered += 360;
+	else if (hueCentered > 180.)
+		hueCentered -= 360;
+	return hueCentered;
+}
+
+
+// Textbook monomial to basis-function conversion matrix.
+static const float3x3 M =
+{
+	{  0.5, -1.0, 0.5 },
+	{ -1.0,  1.0, 0.5 },
+	{  0.5,  0.0, 0.0 }
+};
+
+struct SegmentedSplineParams_c5
+{
+	float coefsLow[6];    // coefs for B-spline between minPoint and midPoint (units of log luminance)
+	float coefsHigh[6];   // coefs for B-spline between midPoint and maxPoint (units of log luminance)
+	float2 minPoint; // {luminance, luminance} linear extension below this
+	float2 midPoint; // {luminance, luminance} 
+	float2 maxPoint; // {luminance, luminance} linear extension above this
+	float slopeLow;       // log-log slope of low linear extension
+	float slopeHigh;      // log-log slope of high linear extension
+};
+
+struct SegmentedSplineParams_c9
+{
+	float coefsLow[10];    // coefs for B-spline between minPoint and midPoint (units of log luminance)
+	float coefsHigh[10];   // coefs for B-spline between midPoint and maxPoint (units of log luminance)
+	float2 minPoint; // {luminance, luminance} linear extension below this
+	float2 midPoint; // {luminance, luminance} 
+	float2 maxPoint; // {luminance, luminance} linear extension above this
+	float slopeLow;       // log-log slope of low linear extension
+	float slopeHigh;      // log-log slope of high linear extension
+};
+
+
+float segmented_spline_c5_fwd( float x )
+{
+	// RRT_PARAMS
+	const SegmentedSplineParams_c5 C =
+	{
+		// coefsLow[6]
+		{ -4.0000000000, -4.0000000000, -3.1573765773, -0.4852499958, 1.8477324706, 1.8477324706 },
+		// coefsHigh[6]
+		{ -0.7185482425, 2.0810307172, 3.6681241237, 4.0000000000, 4.0000000000, 4.0000000000 },
+		{ 0.18*exp2(-15), 0.0001},    // minPoint
+		{ 0.18,              4.8},    // midPoint
+		{ 0.18*exp2(18),  10000.},    // maxPoint
+		0.0,  // slopeLow
+		0.0   // slopeHigh
+	};
+
+	const int N_KNOTS_LOW = 4;
+	const int N_KNOTS_HIGH = 4;
+
+	// Check for negatives or zero before taking the log. If negative or zero,
+	// set to ACESMIN.1
+	float xCheck = x <= 0 ? exp2(-14) : x;
+
+	float logx = log10( xCheck);
+	float logy;
+
+	if ( logx <= log10(C.minPoint.x) )
+	{ 
+		logy = logx * C.slopeLow + ( log10(C.minPoint.y) - C.slopeLow * log10(C.minPoint.x) );
+	}
+	else if (( logx > log10(C.minPoint.x) ) && ( logx < log10(C.midPoint.x) ))
+	{
+		float knot_coord = (N_KNOTS_LOW-1) * (logx-log10(C.minPoint.x))/(log10(C.midPoint.x)-log10(C.minPoint.x));
+		int j = knot_coord;
+		float t = knot_coord - j;
+
+		float3 cf = { C.coefsLow[ j], C.coefsLow[ j + 1], C.coefsLow[ j + 2]};
+		/*if( j <= 0 ) {
+			cf[0] = C.coefsLow[0];  cf[1] = C.coefsLow[1];  cf[2] = C.coefsLow[2];
+		} else if( j == 1 ) {
+			cf[0] = C.coefsLow[1];  cf[1] = C.coefsLow[2];  cf[2] = C.coefsLow[3];
+		} else if( j == 2 ) {
+			cf[0] = C.coefsLow[2];  cf[1] = C.coefsLow[3];  cf[2] = C.coefsLow[4];
+		} else if( j == 3 ) {
+			cf[0] = C.coefsLow[3];  cf[1] = C.coefsLow[4];  cf[2] = C.coefsLow[5];
+		} else if( j == 4 ) {
+			cf[0] = C.coefsLow[4];  cf[1] = C.coefsLow[5];  cf[2] = C.coefsLow[6];
+		} else if( j == 5 ) {
+			cf[0] = C.coefsLow[5];  cf[1] = C.coefsLow[6];  cf[2] = C.coefsLow[7];
+		} else if( j == 6 ) {
+			cf[0] = C.coefsLow[6];  cf[1] = C.coefsLow[7];  cf[2] = C.coefsLow[8];
+		}*/
+    
+		float3 monomials = { t * t, t, 1 };
+		logy = dot( monomials, mul( cf, M));
+	}
+	else if (( logx >= log10(C.midPoint.x) ) && ( logx < log10(C.maxPoint.x) ))
+	{
+		float knot_coord = (N_KNOTS_HIGH-1) * (logx-log10(C.midPoint.x))/(log10(C.maxPoint.x)-log10(C.midPoint.x));
+		int j = knot_coord;
+		float t = knot_coord - j;
+
+		float3 cf = { C.coefsHigh[ j], C.coefsHigh[ j + 1], C.coefsHigh[ j + 2]};
+		/*if ( j <= 0) {
+			cf[ 0] = C.coefsHigh[0];  cf[ 1] = C.coefsHigh[1];  cf[ 2] = C.coefsHigh[2];
+		} else if ( j == 1) {
+			cf[ 0] = C.coefsHigh[1];  cf[ 1] = C.coefsHigh[2];  cf[ 2] = C.coefsHigh[3];
+		} else if ( j == 2) {
+			cf[ 0] = C.coefsHigh[2];  cf[ 1] = C.coefsHigh[3];  cf[ 2] = C.coefsHigh[4];
+		} else if ( j == 3) {
+			cf[ 0] = C.coefsHigh[3];  cf[ 1] = C.coefsHigh[4];  cf[ 2] = C.coefsHigh[5];
+		} else if ( j == 4) {
+			cf[ 0] = C.coefsHigh[4];  cf[ 1] = C.coefsHigh[5];  cf[ 2] = C.coefsHigh[6];
+		} else if ( j == 5) {
+			cf[ 0] = C.coefsHigh[5];  cf[ 1] = C.coefsHigh[6];  cf[ 2] = C.coefsHigh[7];
+		} else if ( j == 6) {
+			cf[ 0] = C.coefsHigh[6];  cf[ 1] = C.coefsHigh[7];  cf[ 2] = C.coefsHigh[8];
+		} */
+
+		float3 monomials = { t * t, t, 1 };
+		logy = dot( monomials, mul( cf, M));
+	}
+	else
+	{ //if ( logIn >= log10(C.maxPoint.x) ) { 
+		logy = logx * C.slopeHigh + ( log10(C.maxPoint.y) - C.slopeHigh * log10(C.maxPoint.x) );
+	}
+
+	return pow( 10, logy );
+}
+
+
+float segmented_spline_c9_fwd( float x )
+{
+	// ODT_48nits
+	const SegmentedSplineParams_c9 C =
+	{
+		// coefsLow[10]
+		{ -1.6989700043, -1.6989700043, -1.4779000000, -1.2291000000, -0.8648000000, -0.4480000000, 0.0051800000, 0.4511080334, 0.9113744414, 0.9113744414},
+		// coefsHigh[10]
+		{ 0.5154386965, 0.8470437783, 1.1358000000, 1.3802000000, 1.5197000000, 1.5985000000, 1.6467000000, 1.6746091357, 1.6878733390, 1.6878733390 },
+		{segmented_spline_c5_fwd( 0.18*exp2(-6.5) ),  0.02},    // minPoint
+		{segmented_spline_c5_fwd( 0.18 ),              4.8},    // midPoint  
+		{segmented_spline_c5_fwd( 0.18*exp2(6.5) ),   48.0},    // maxPoint
+		0.0,  // slopeLow
+		0.04  // slopeHigh
+	};
+
+	const int N_KNOTS_LOW = 8;
+	const int N_KNOTS_HIGH = 8;
+
+	// Check for negatives or zero before taking the log. If negative or zero,
+	// set to OCESMIN.
+	float xCheck = x <= 0 ? 1e-4 : x;
+
+	float logx = log10(xCheck);
+	float logy;
+
+	if (logx <= log10(C.minPoint.x))
+	{
+		logy = logx * C.slopeLow + (log10(C.minPoint.y) - C.slopeLow * log10(C.minPoint.x));
+	}
+	else if ((logx > log10(C.minPoint.x)) && (logx < log10(C.midPoint.x)))
+	{
+		float knot_coord = (N_KNOTS_LOW - 1) * (logx - log10(C.minPoint.x)) / (log10(C.midPoint.x) - log10(C.minPoint.x));
+		int j = knot_coord;
+		float t = knot_coord - j;
+
+		float3 cf = { C.coefsLow[j], C.coefsLow[j + 1], C.coefsLow[j + 2] };
+
+		float3 monomials = { t * t, t, 1 };
+		logy = dot(monomials, mul(cf, M));
+	}
+	else if ((logx >= log10(C.midPoint.x)) && (logx < log10(C.maxPoint.x)))
+	{
+		float knot_coord = (N_KNOTS_HIGH - 1) * (logx - log10(C.midPoint.x)) / (log10(C.maxPoint.x) - log10(C.midPoint.x));
+		int j = knot_coord;
+		float t = knot_coord - j;
+
+		float3 cf = { C.coefsHigh[j], C.coefsHigh[j + 1], C.coefsHigh[j + 2] };
+
+		float3 monomials = { t * t, t, 1 };
+		logy = dot(monomials, mul(cf, M));
+	}
+	else//if ( logIn >= log10(C.maxPoint.x) )
+	{
+		logy = logx * C.slopeHigh + (log10(C.maxPoint.y) - C.slopeHigh * log10(C.maxPoint.x));
+	}
+
+	return pow( 10, logy );
+}
+
+
+// Transformations from RGB to other color representations
+float rgb_2_hue( float3 rgb ) 
+{
+	// Returns a geometric hue angle in degrees (0-360) based on RGB values.
+	// For neutral colors, hue is undefined and the function will return a quiet NaN value.
+	float hue;
+	if (rgb[0] == rgb[1] && rgb[1] == rgb[2])
+	{
+		//hue = FLT_NAN; // RGB triplets where RGB are equal have an undefined hue
+		hue = 0;
+	}
+	else
+	{
+		hue = (180. / PI) * atan2(sqrt(3.0)*(rgb[1] - rgb[2]), 2 * rgb[0] - rgb[1] - rgb[2]);
+	}
+
+	if (hue < 0.)
+		hue = hue + 360;
+
+	return clamp( hue, 0, 360 );
+}
+
+float rgb_2_yc( float3 rgb, float ycRadiusWeight = 1.75)
+{
+	// Converts RGB to a luminance proxy, here called YC
+	// YC is ~ Y + K * Chroma
+	// Constant YC is a cone-shaped surface in RGB space, with the tip on the 
+	// neutral axis, towards white.
+	// YC is normalized: RGB 1 1 1 maps to YC = 1
+	//
+	// ycRadiusWeight defaults to 1.75, although can be overridden in function 
+	// call to rgb_2_yc
+	// ycRadiusWeight = 1 -> YC for pure cyan, magenta, yellow == YC for neutral 
+	// of same value
+	// ycRadiusWeight = 2 -> YC for pure red, green, blue  == YC for  neutral of 
+	// same value.
+
+	float r = rgb[0]; 
+	float g = rgb[1]; 
+	float b = rgb[2];
+  
+	float chroma = sqrt(b*(b-g)+g*(g-r)+r*(r-b));
+
+	return ( b + g + r + ycRadiusWeight * chroma) / 3.;
+}
+
+// 
+// Reference Rendering Transform (RRT)
+//
+//   Input is ACES
+//   Output is OCES
+//
+float3 RRT( float3 aces )
+{
+	// "Glow" module constants
+	const float RRT_GLOW_GAIN = 0.05;
+	const float RRT_GLOW_MID = 0.08;
+
+	float saturation = rgb_2_saturation( aces );
+	float ycIn = rgb_2_yc( aces );
+	float s = sigmoid_shaper( (saturation - 0.4) / 0.2);
+	float addedGlow = 1 + glow_fwd( ycIn, RRT_GLOW_GAIN * s, RRT_GLOW_MID);
+	aces *= addedGlow;
+  
+	// --- Red modifier --- //
+	const float RRT_RED_SCALE = 0.82;
+	const float RRT_RED_PIVOT = 0.03;
+	const float RRT_RED_HUE = 0;
+	const float RRT_RED_WIDTH = 135;
+	float hue = rgb_2_hue( aces );
+	float centeredHue = center_hue( hue, RRT_RED_HUE );
+	float hueWeight = cubic_basis_shaper( centeredHue, RRT_RED_WIDTH );
+		
+	aces.r += hueWeight * saturation * (RRT_RED_PIVOT - aces.r) * (1. - RRT_RED_SCALE);
+
+	// --- ACES to RGB rendering space --- //
+	aces = clamp( aces, 0, 65535 );  // avoids saturated negative colors from becoming positive in the matrix
+
+	float3 rgbPre = mul( AP0_2_AP1_MAT, aces );
+
+	rgbPre = clamp( rgbPre, 0, 65535 );
+
+	// --- Global desaturation --- //
+	const float RRT_SAT_FACTOR = 0.96;
+	rgbPre = lerp( dot( rgbPre, AP1_RGB2Y ), rgbPre, RRT_SAT_FACTOR );
+
+	// --- Apply the tonescale independently in rendering-space RGB --- //
+	float3 rgbPost;
+	rgbPost[0] = segmented_spline_c5_fwd( rgbPre[0] );
+	rgbPost[1] = segmented_spline_c5_fwd( rgbPre[1] );
+	rgbPost[2] = segmented_spline_c5_fwd( rgbPre[2] );
+
+	// --- RGB rendering space to OCES --- //
+	float3 rgbOces = mul( AP1_2_AP0_MAT, rgbPost );
+	return rgbOces;
+}
+
+
+// Transformations between CIE XYZ tristimulus values and CIE x,y 
+// chromaticity coordinates
+float3 XYZ_2_xyY( float3 XYZ )
+{  
+	float3 xyY;
+	float divisor = (XYZ[0] + XYZ[1] + XYZ[2]);
+	if (divisor == 0.) divisor = 1e-10;
+	xyY[0] = XYZ[0] / divisor;
+	xyY[1] = XYZ[1] / divisor;  
+	xyY[2] = XYZ[1];
+  
+	return xyY;
+}
+
+float3 xyY_2_XYZ( float3 xyY )
+{
+	float3 XYZ;
+	XYZ[0] = xyY[0] * xyY[2] / max( xyY[1], 1e-10);
+	XYZ[1] = xyY[2];  
+	XYZ[2] = (1.0 - xyY[0] - xyY[1]) * xyY[2] / max( xyY[1], 1e-10);
+
+	return XYZ;
+}
+
+
+float3x3 ChromaticAdaptation( float2 src_xy, float2 dst_xy )
+{
+	// Von Kries chromatic adaptation 
+
+	// Bradford
+	const float3x3 ConeResponse =
+	{
+		 0.8951,  0.2664, -0.1614,
+		-0.7502,  1.7135,  0.0367,
+		 0.0389, -0.0685,  1.0296,
+	};
+	const float3x3 InvConeResponse =
+	{
+		 0.9869929, -0.1470543,  0.1599627,
+		 0.4323053,  0.5183603,  0.0492912,
+		-0.0085287,  0.0400428,  0.9684867,
+	};
+
+	float3 src_XYZ = xyY_2_XYZ( float3( src_xy, 1 ) );
+	float3 dst_XYZ = xyY_2_XYZ( float3( dst_xy, 1 ) );
+
+	float3 src_coneResp = mul( ConeResponse, src_XYZ );
+	float3 dst_coneResp = mul( ConeResponse, dst_XYZ );
+
+	float3x3 VonKriesMat =
+	{
+		{ dst_coneResp[0] / src_coneResp[0], 0.0, 0.0 },
+		{ 0.0, dst_coneResp[1] / src_coneResp[1], 0.0 },
+		{ 0.0, 0.0, dst_coneResp[2] / src_coneResp[2] }
+	};
+
+	return mul( InvConeResponse, mul( VonKriesMat, ConeResponse ) );
+}
+
+
+float Y_2_linCV( float Y, float Ymax, float Ymin) 
+{
+  return (Y - Ymin) / (Ymax - Ymin);
+}
+
+// Gamma compensation factor
+static const float DIM_SURROUND_GAMMA = 0.9811;
+
+float3 darkSurround_to_dimSurround( float3 linearCV )
+{
+	float3 XYZ = mul( AP1_2_XYZ_MAT, linearCV );
+
+	float3 xyY = XYZ_2_xyY(XYZ);
+	xyY[2] = clamp( xyY[2], 0, 65535 );
+	xyY[2] = pow( xyY[2], DIM_SURROUND_GAMMA );
+	XYZ = xyY_2_XYZ(xyY);
+
+	return mul( XYZ_2_AP1_MAT, XYZ );
+}
+
+
+// 
+// Output Device Transform - RGB computer monitor (D60 simulation)
+//
+
+//
+// Summary :
+//  This transform is intended for mapping OCES onto a desktop computer monitor 
+//  typical of those used in motion picture visual effects production used to 
+//  simulate the image appearance produced by odt_p3d60. These monitors may
+//  occasionally be referred to as "sRGB" displays, however, the monitor for
+//  which this transform is designed does not exactly match the specifications
+//  in IEC 61966-2-1:1999.
+// 
+//  The assumed observer adapted white is D60, and the viewing environment is 
+//  that of a dim surround. 
+//
+//  The monitor specified is intended to be more typical of those found in 
+//  visual effects production.
+//
+// Device Primaries : 
+//  Primaries are those specified in Rec. ITU-R BT.709
+//  CIE 1931 chromaticities:  x         y         Y
+//              Red:          0.64      0.33
+//              Green:        0.3       0.6
+//              Blue:         0.15      0.06
+//              White:        0.3217    0.329     100 cd/m^2
+//
+// Display EOTF :
+//  The reference electro-optical transfer function specified in 
+//  IEC 61966-2-1:1999.
+//
+// Signal Range:
+//    This tranform outputs full range code values.
+//
+// Assumed observer adapted white point:
+//         CIE 1931 chromaticities:    x            y
+//                                     0.32168      0.33767
+//
+// Viewing Environment:
+//   This ODT has a compensation for viewing environment variables more typical 
+//   of those associated with video mastering.
+//
+
+float3 ODT( float3 oces )
+{
+	// OCES to RGB rendering space
+	float3 rgbPre = mul( AP0_2_AP1_MAT, oces );
+
+	// Apply the tonescale independently in rendering-space RGB
+	float3 rgbPost;
+	rgbPost[0] = segmented_spline_c9_fwd( rgbPre[0] );
+	rgbPost[1] = segmented_spline_c9_fwd( rgbPre[1] );
+	rgbPost[2] = segmented_spline_c9_fwd( rgbPre[2] );
+
+	// Target white and black points for cinema system tonescale
+	static const float CINEMA_WHITE = 48.0;
+	static const float CINEMA_BLACK = 0.02; // CINEMA_WHITE / 2400.
+
+	// Scale luminance to linear code value
+	float3 linearCV;
+	linearCV[0] = Y_2_linCV( rgbPost[0], CINEMA_WHITE, CINEMA_BLACK );
+	linearCV[1] = Y_2_linCV( rgbPost[1], CINEMA_WHITE, CINEMA_BLACK );
+	linearCV[2] = Y_2_linCV( rgbPost[2], CINEMA_WHITE, CINEMA_BLACK );    
+
+	// --- Compensate for different white point being darker  --- //
+	// This adjustment is to correct an issue that exists in ODTs where the device 
+	// is calibrated to a white chromaticity other than D60. In order to simulate 
+	// D60 on such devices, unequal code values are sent to the display to achieve 
+	// neutrals at D60. In order to produce D60 on a device calibrated to the DCI 
+	// white point (i.e. equal code values yield CIE x,y chromaticities of 0.314, 
+	// 0.351) the red channel is higher than green and blue to compensate for the 
+	// "greenish" DCI white. This is the correct behavior but it means that as 
+	// highlight increase, the red channel will hit the device maximum first and 
+	// clip, resulting in a chromaticity shift as the green and blue channels 
+	// continue to increase.
+	// To avoid this clipping error, a slight scale factor is applied to allow the 
+	// ODTs to simulate D60 within the D65 calibration white point. 
+
+	// Scale and clamp white to avoid casted highlights due to D60 simulation
+	const float SCALE = 0.955;
+	linearCV = min( linearCV, 1 ) * SCALE;
+
+	// Apply gamma adjustment to compensate for dim surround
+	linearCV = darkSurround_to_dimSurround( linearCV );
+
+	// Apply desaturation to compensate for luminance difference
+	const float ODT_SAT_FACTOR = 0.93;
+	linearCV = lerp( dot( linearCV, AP1_RGB2Y ), linearCV, ODT_SAT_FACTOR );
+
+	// Convert to display primary encoding
+	// Rendering space RGB to XYZ
+	float3 XYZ = mul( AP1_2_XYZ_MAT, linearCV );
+	
+	// Apply CAT from ACES white point to assumed observer adapted white point
+	const float3x3 D60_2_D65_CAT =
+	{
+		 0.987224,   -0.00611327, 0.0159533,
+		-0.00759836,  1.00186,    0.00533002,
+		 0.00307257, -0.00509595, 1.08168,
+	};
+	//XYZ = mul( D60_2_D65_CAT, XYZ );
+
+	// CIE XYZ to display primaries
+	linearCV = mul( XYZ_2_sRGB_MAT, XYZ );
+
+	// Handle out-of-gamut values
+	// Clip values < 0 or > 1 (i.e. projecting outside the display primaries)
+	linearCV = saturate( linearCV );
+
+	return linearCV;
+}
\ No newline at end of file
diff --git a/Engine/Shaders/PostProcessCombineLUTs.usf b/Engine/Shaders/PostProcessCombineLUTs.usf
--- a/Engine/Shaders/PostProcessCombineLUTs.usf
+++ b/Engine/Shaders/PostProcessCombineLUTs.usf
@@ -4,9 +4,14 @@
 	LUTBlender.usf: Filter pixel shader source.
 =============================================================================*/
 
+#define USE_COLOR_MATRIX	1
+#define USE_SHADOW_TINT		1
+#define USE_CONTRAST		1
+
 #include "Common.usf"
 #include "PostProcessCommon.usf"
 #include "TonemapCommon.usf"
+#include "GammaCorrectionCommon.usf"
 
 // ---------------------------------------------------
 
@@ -24,7 +29,124 @@ float LUTWeights[5];
 half3 ColorScale;
 half4 OverlayColor;
 
+static const float LUTSize = 32;
+
+// Accurate for 1000K < Temp < 15000K
+// [Krystek 1985, "An algorithm to calculate correlated colour temperature"]
+float2 PlanckianLocusChromaticity( float Temp )
+{
+	float u = ( 0.860117757f + 1.54118254e-4f * Temp + 1.28641212e-7f * Temp*Temp ) / ( 1.0f + 8.42420235e-4f * Temp + 7.08145163e-7f * Temp*Temp );
+	float v = ( 0.317398726f + 4.22806245e-5f * Temp + 4.20481691e-8f * Temp*Temp ) / ( 1.0f - 2.89741816e-5f * Temp + 1.61456053e-7f * Temp*Temp );
+
+	float x = 3*u / ( 2*u - 8*v + 4 );
+	float y = 2*v / ( 2*u - 8*v + 4 );
+
+	return float2(x,y);
+}
+
+// Accurate for 4000K < Temp < 25000K
+// in: correlated color temperature
+// out: CIE 1931 chromaticity
+float2 D_IlluminantChromaticity( float Temp )
+{
+	// Correct for revision of Plank's law
+	// This makes 6500 == D65
+	Temp *= 1.4388 / 1.438;
+
+	float x =	Temp <= 7000 ?
+				0.244063 + ( 0.09911e3 + ( 2.9678e6 - 4.6070e9 / Temp ) / Temp ) / Temp :
+				0.237040 + ( 0.24748e3 + ( 1.9018e6 - 2.0064e9 / Temp ) / Temp ) / Temp;
+	
+	float y = -3 * x*x + 2.87 * x - 0.275;
+
+	return float2(x,y);
+}
+
+// Find closest color temperature to chromaticity
+// [McCamy 1992, "Correlated color temperature as an explicit function of chromaticity coordinates"]
+float CorrelatedColorTemperature( float x, float y )
+{
+	float n = (x - 0.3320) / (0.1858 - y);
+	return -449 * n*n*n + 3525 * n*n - 6823.3 * n + 5520.33;
+}
+
+float2 PlanckianIsothermal( float Temp, float Tint )
+{
+	float u = ( 0.860117757f + 1.54118254e-4f * Temp + 1.28641212e-7f * Temp*Temp ) / ( 1.0f + 8.42420235e-4f * Temp + 7.08145163e-7f * Temp*Temp );
+	float v = ( 0.317398726f + 4.22806245e-5f * Temp + 4.20481691e-8f * Temp*Temp ) / ( 1.0f - 2.89741816e-5f * Temp + 1.61456053e-7f * Temp*Temp );
+
+	float ud = ( -1.13758118e9f - 1.91615621e6f * Temp - 1.53177f * Temp*Temp ) / Square( 1.41213984e6f + 1189.62f * Temp + Temp*Temp );
+	float vd = (  1.97471536e9f - 705674.0f * Temp - 308.607f * Temp*Temp ) / Square( 6.19363586e6f - 179.456f * Temp + Temp*Temp );
+
+	float2 uvd = normalize( float2( u, v ) );
+
+	// Correlated color temperature is meaningful within +/- 0.05
+	u += -uvd.y * Tint * 0.05;
+	v +=  uvd.x * Tint * 0.05;
+	
+	float x = 3*u / ( 2*u - 8*v + 4 );
+	float y = 2*v / ( 2*u - 8*v + 4 );
+
+	return float2(x,y);
+}
+
+float WhiteTemp;
+float WhiteTint;
+
+float3 WhiteBalance( float3 LinearColor )
+{
+	float2 SrcWhiteDaylight = D_IlluminantChromaticity( WhiteTemp );
+	float2 SrcWhitePlankian = PlanckianLocusChromaticity( WhiteTemp );
+
+	float2 SrcWhite = WhiteTemp < 4000 ? SrcWhitePlankian : SrcWhiteDaylight;
+	float2 D65White = float2( 0.31270,  0.32900 );
+
+	{
+		// Offset along isotherm
+		float2 Isothermal = PlanckianIsothermal( WhiteTemp, WhiteTint ) - SrcWhitePlankian;
+		SrcWhite += Isothermal;
+	}
+
+	float3x3 WhiteBalanceMat = ChromaticAdaptation( SrcWhite, D65White );
+	WhiteBalanceMat = mul( XYZ_2_sRGB_MAT, mul( WhiteBalanceMat, sRGB_2_XYZ_MAT ) );
+
+	return mul( WhiteBalanceMat, LinearColor );
+}
+
+float3 ColorSaturation;
+float3 ColorContrast;
+float3 ColorGamma;
+float3 ColorGain;
+float3 ColorOffset;
+
+float3 ColorGrade( float3 LinearColor )
+{
+	// ACEScg working space
+	const float3x3 sRGB_2_AP1 = mul( XYZ_2_AP1_MAT, mul( D65_2_D60_CAT, sRGB_2_XYZ_MAT ) );
+	const float3x3 AP1_2_sRGB = mul( XYZ_2_sRGB_MAT, mul( D60_2_D65_CAT, AP1_2_XYZ_MAT ) );
+
+	float3 WorkingColor = mul( sRGB_2_AP1, LinearColor );
 
+	// TODO optimize
+	float Luma = dot( WorkingColor, AP1_RGB2Y );
+	WorkingColor = max( 0, lerp( Luma.xxx, WorkingColor, ColorSaturation ) );
+	WorkingColor = pow( WorkingColor * (1.0 / 0.18), ColorContrast ) * 0.18;
+	WorkingColor = pow( WorkingColor, 1.0 / ColorGamma );
+	WorkingColor = WorkingColor * ColorGain + ColorOffset;
+
+#if 0
+	{
+		// Split tone
+		float ShadowWeight = 1 - smoothstep( 0, 0.18, Luma );
+		float HighlightWeight = smoothstep( 0.5, 1, Luma );
+		float MidtoneWeight = 1 - ShadowWeight - HighlightWeight;
+
+		WorkingColor = ShadowWeight * ShadowColor + MidtoneWeight * GradedColor + HighlightWeight * HighlightColor;
+	}
+#endif
+
+	return mul( AP1_2_sRGB, WorkingColor );
+}
 
 // todo: Weight[0] should be used for neutral, Texture* name should start with 1, color correction should apply on top of that
 #if USE_VOLUME_LUT == 1
@@ -33,21 +155,9 @@ void MainPS(FWriteToSliceGeometryOutput Input, out float4 OutColor : SV_Target0)
 	// construct the neutral color from a 3d position volume texture	
 	float4 Neutral;
 	{
-		float2 UV = Input.Vertex.UV - float2(0.5f / 16.0f, 0.5f / 16.0f);
-
-		Neutral = float4(UV * (16.0f / 15.0f), Input.LayerIndex / 15.0f, 0);
-	}
-
-	// Transform Neutral color to 256x16 UV
-	float2 InUV;
-	{
-		float Scale = 15.0f / 16.0f; 
+		float2 UV = Input.Vertex.UV - float2(0.5f / LUTSize, 0.5f / LUTSize);
 
-		float IntB = floor(Neutral.b * 15.9999f) / 16.0f;
-
-		InUV = float2(Neutral.r * Scale / 16.0f + IntB, Neutral.g * Scale);
-		
-		InUV += float2(0.5f / 256.0f, 0.5f / 16.0f);
+		Neutral = float4(UV * LUTSize / (LUTSize - 1), Input.LayerIndex / (LUTSize - 1), 0);
 	}
 #else
 void MainPS(float4 InUV : TEXCOORD0, out float4 OutColor : SV_Target0)
@@ -57,53 +167,110 @@ void MainPS(float4 InUV : TEXCOORD0, out float4 OutColor : SV_Target0)
 	{ 
 		float2 UV = InUV.xy;
 
-		UV -= float2(0.5f / 256.0f, 0.5f / 16.0f);
+		UV -= float2(0.5f / (LUTSize * LUTSize), 0.5f / LUTSize);
 
-		float Scale = 16.0f / 15.0f;
+		float Scale = LUTSize / (LUTSize - 1);
 
 		float3 RGB;
 		
-		RGB.r = frac(UV.x * 16.0f);
-		RGB.b = UV.x - RGB.r / 16.0f;
+		RGB.r = frac(UV.x * LUTSize);
+		RGB.b = UV.x - RGB.r / LUTSize;
 		RGB.g = UV.y;
 
 		Neutral = float4(RGB * Scale, 0);
 	}
 #endif
 
-    OutColor = LUTWeights[0] * Neutral;
+	float3 LogColor = Neutral.rgb;
+	
+	const float LinearRange = 14;
+	const float LinearGrey = 0.18;
+	const float ExposureGrey = 444;
 
-	// BLENDCOUNT is the number of LUT that are blended together including the neutral one
-#if BLENDCOUNT >= 2
-	OutColor += LUTWeights[1] * Texture2DSample(Texture1, Texture1Sampler, InUV);
-#endif
+	float3 LinearColor = exp2( ( LogColor - ExposureGrey / 1023.0 ) * LinearRange ) * LinearGrey;
+	//float3 LinearColor = 2 * ( pow(10.0, ((LogColor - 0.616596 - 0.03) / 0.432699)) - 0.037584 );	// SLog
+	//float3 LinearColor = ( pow( 10, ( 1023 * LogColor - 685 ) / 300) - .0108 ) / (1 - .0108);	// Cineon
+	//LinearColor = max( 0, LinearColor );
 
-#if BLENDCOUNT >= 3
-	OutColor += LUTWeights[2] * Texture2DSample(Texture2, Texture2Sampler, InUV);
-#endif
+	float3 BalancedColor = WhiteBalance( LinearColor );
+	float3 GradedColor = ColorGrade( BalancedColor );
+	float3 FilmColor = FilmToneMap( GradedColor );
+	
+	BRANCH
+	if( ColorShadow_Tint2.a == 0 )
+	{
+		// Legacy tone mapper
+		FilmColor = FilmPostProcess( GradedColor );
+	}
 
-#if BLENDCOUNT >= 4
-	OutColor += LUTWeights[3] * Texture2DSample(Texture3, Texture3Sampler, InUV);
-#endif
+#if BLENDCOUNT > 1
+	{
+		// Legacy LDR LUT color grading
+
+		// FIXME LUTs are in gamma space
+		half3 GammaColor = pow( FilmColor, 1.0 / 2.2 );
+
+		// Transform color to 256x16 UV
+		float Scale = 15.0f / 16.0f;
+
+		float IntB = floor(Neutral.b * 15.9999f) / 16.0f;
+
+		float2 UV = float2(Neutral.r * Scale / 16.0f + IntB, Neutral.g * Scale);
+		UV += float2(0.5f / 256.0f, 0.5f / 16.0f);
+
+		GammaColor = LUTWeights[0] * GammaColor;
+
+		// BLENDCOUNT is the number of LUT that are blended together including the neutral one
+	#if BLENDCOUNT >= 2
+		GammaColor += LUTWeights[1] * Texture2DSample(Texture1, Texture1Sampler, UV).rgb;
+	#endif
 
-#if BLENDCOUNT >= 5
-	OutColor += LUTWeights[4] * Texture2DSample(Texture4, Texture4Sampler, InUV);
+	#if BLENDCOUNT >= 3
+		GammaColor += LUTWeights[2] * Texture2DSample(Texture2, Texture2Sampler, UV).rgb;
+	#endif
+
+	#if BLENDCOUNT >= 4
+		GammaColor += LUTWeights[3] * Texture2DSample(Texture3, Texture3Sampler, UV).rgb;
+	#endif
+
+	#if BLENDCOUNT >= 5
+		GammaColor += LUTWeights[4] * Texture2DSample(Texture4, Texture4Sampler, UV).rgb;
+	#endif
+
+		// Back to linear space
+		FilmColor = pow( GammaColor, 2.2 );
+	}
 #endif
+	
+	// apply math color correction on top to texture based solution
+	FilmColor = ColorCorrection( FilmColor );
 
 	// blend with custom LDR color, used for Fade track in Matinee
-	OutColor.rgb = lerp(OutColor.rgb * ColorScale, OverlayColor.rgb, OverlayColor.a);
+	FilmColor = lerp( FilmColor * ColorScale, OverlayColor.rgb, OverlayColor.a );
 
-	// apply math color correction on top ot texture based solution
-	// (faster than doing it in the full screen pass)
-	OutColor.rgb = ColorCorrection(OutColor.rgb);
-
-//	static float SmallPowConst = 0.00000001f;
+#if 0
+	// Apply "gamma" curve adjustment.
+	FilmColor = pow( FilmColor, InverseGamma.y );
+		
+	#if USE_GAMMA
+		// This is different than the prior "gamma" curve adjustment (but reusing the variable).
+		// For displays set to a gamma colorspace.
+		// Note, MacOSX native output is raw gamma 2.2 not sRGB!
+		half3 OutDeviceColor = pow( FilmColor, InverseGamma.z );
+	#elif USE_709
+		// Didn't profile yet if the branching version would be faster (different linear segment).
+		half3 OutDeviceColor = LinearTo709Branchless( FilmColor );
+	#else
+		// Apply conversion to sRGB (this must be an exact sRGB conversion else darks are bad).
+		// Branching is faster than branchless on AMD on PC.
+		half3 OutDeviceColor = LinearToSrgbBranching( FilmColor );
+	#endif
+#else
+	half3 OutDeviceColor = sqrt( FilmColor );
+#endif
 
-	// needed for d3d11 and if custom gamma is specified
-	// could be integrated into the uberpostprocess but here it is much cheaper
-	// max() clamp to get nice behavior for negative colors (see level sm2test)
-//	OutColor = UnClampedPow(max(float4(SmallPowConst,SmallPowConst,SmallPowConst,SmallPowConst), OutColor), 2.2f * GammaInverse);
-//	OutColor = pow(max(float4(SmallPowConst,SmallPowConst,SmallPowConst,SmallPowConst), OutColor), 2.2f * GammaInverse);
+	// Better to saturate(lerp(a,b,t)) than lerp(saturate(a),saturate(b),t)
+	OutColor.rgb = OutDeviceColor / 1.05;
 }
 
 
diff --git a/Engine/Shaders/PostProcessCommon.usf b/Engine/Shaders/PostProcessCommon.usf
--- a/Engine/Shaders/PostProcessCommon.usf
+++ b/Engine/Shaders/PostProcessCommon.usf
@@ -129,17 +129,12 @@ half3 ColorCorrection(half3 InLDRColor)
 // @param ScreenSpacePos -1..1
 float ComputeVignetteMask(float2 ScreenSpacePos, float Intensity)
 {
-	return lerp(1, DiscMask(ScreenSpacePos / 1.3f), Intensity);
-
-//	NormalizedUV = ScreenSpacePos * 0.5f + 0.5f;
-	// alternative method (pillow like shape):
-	// 0..1 -> 0 .. 0.25f .. 1.0f
-//	float2 Temp = NormalizedUV * (1 - NormalizedUV);
-	// ShapeFactor 1:fast, 0.5:more bright areas, 0.25:even more bright area
-//	float ShapeFactor = 1.0f;
-	// (0 .. 0.25f .. 1.0f, 0 .. 0.25f .. 1.0f) -> (0:border .. 1:center)
-//	float Mask = Temp.x * Temp.y * 16.0f;
-//	return lerp(1, pow(Mask, ShapeFactor), Intensity);
+	// Natural vignetting
+	// cosine-fourth law
+	ScreenSpacePos *= Intensity;
+	float Tan2Angle = dot( ScreenSpacePos, ScreenSpacePos );
+	float Cos4Angle = Square( rcp( Tan2Angle + 1 ) );
+	return Cos4Angle;
 }
 
 // Scale {-1 to 1} space to vignette circle space.
diff --git a/Engine/Shaders/PostProcessMobile.usf b/Engine/Shaders/PostProcessMobile.usf
--- a/Engine/Shaders/PostProcessMobile.usf
+++ b/Engine/Shaders/PostProcessMobile.usf
@@ -899,7 +899,6 @@ void SunMergeVS_ES2(
 #undef SUN_M
 
 float4 SunColorVignetteIntensity;
-float3 VignetteColor;
 float3 BloomColor;
 
 void SunMergePS_ES2(
@@ -952,7 +951,6 @@ void SunMergePS_ES2(
 
 	OutColor.a = Vignette;
 	OutColor.rgb *= OutColor.a;
-	OutColor.rgb += VignetteColor * (1.0 - Vignette);
 }
 
 #undef SUN_P0
@@ -1041,7 +1039,6 @@ void SunMergeSmallPS_ES2(
 
 	OutColor.a = Vignette;
 	OutColor.rgb *= OutColor.a;
-	OutColor.rgb += VignetteColor * (1.0 - Vignette);
 }
 
 
diff --git a/Engine/Shaders/PostProcessTonemap.usf b/Engine/Shaders/PostProcessTonemap.usf
--- a/Engine/Shaders/PostProcessTonemap.usf
+++ b/Engine/Shaders/PostProcessTonemap.usf
@@ -23,51 +23,62 @@ float4 ColorScale1;
 float2 TexScale;
 
 // from the postprocess settings
-float4 VignetteColorIntensity;
+float VignetteIntensity;
 
 // Fine film grain
 float3 GrainRandomFull;
 float3 GrainScaleBiasJitter;
 
-float4 OverlayColor;
-
 half GrainFromUV(float2 GrainUV) 
 {
 	half Grain = frac(sin(GrainUV.x + GrainUV.y * 543.31) *  493013.0);
 	return Grain;
 }
 
-
-
-// 256x16 color LUT for color grading
+// LUT for color grading
 #if USE_VOLUME_LUT == 1
-	Texture3D ColorGradingLUT;
-	SamplerState ColorGradingLUTSampler;
+	Texture3D		ColorGradingLUT;
+	SamplerState	ColorGradingLUTSampler;
 #endif
 
-//  @param InLDRColor in gamma space, has to be in 0..1 range
-half3 ColorLookupTable(half3 InLDRColor)
+static const float LUTSize = 32;
+
+half4 UnwrappedTexture3DSample( Texture2D Texture, SamplerState Sampler, float3 UVW, float Size )
 {
-#if USE_VOLUME_LUT == 1
-	return Texture3DSample(ColorGradingLUT, ColorGradingLUTSampler, InLDRColor * (15.0f / 16.0f) + (0.5f / 16.0f)).xyz;
-#else
-	// requires a volume texture 16x16x16 unwrapped in a 2d texture 256x16
-	// can be optimized by using a volume texture
-	float2 Offset = float2(0.5f / 256.0f, 0.5f / 16.0f);
-	float Scale = 15.0f / 16.0f; 
+	// a volume texture 16x16x16 would be unwrapped to a 2d texture 256x16
 
-	// Also consider blur value in the blur buffer written by translucency
-	float IntB = floor(InLDRColor.b * 14.9999f) / 16.0f;
-	half FracB = InLDRColor.b * 15.0f - IntB * 16.0f;
+	float IntW = floor( UVW.z * Size - 0.5 );
+	half FracW = UVW.z * Size - 0.5 - IntW;
 
-	float U = IntB + InLDRColor.r * Scale / 16.0f;
-	float V = InLDRColor.g * Scale;
+	float U = UVW.x + IntW / Size;
+	float V = UVW.y;
 
-	half3 RG0 = Texture2DSample(PostprocessInput3, PostprocessInput3Sampler, Offset + float2(U             , V)).rgb;
-	half3 RG1 = Texture2DSample(PostprocessInput3, PostprocessInput3Sampler, Offset + float2(U + 1.0f / 16.0f, V)).rgb;
+	half4 RG0 = Texture.Sample( Sampler, float2(U, V) );
+	half4 RG1 = Texture.Sample( Sampler, float2(U + 1.0f / Size, V) );
+
+	return lerp(RG0, RG1, FracW);
+}
 
-	return lerp(RG0, RG1, FracB);
+half3 ColorLookupTable( half3 LinearColor )
+{
+	float LinearRange = 14;
+	float LinearGrey = 0.18;
+	float ExposureGrey = 444;
+
+	float3 LogColor = log2( LinearColor / LinearGrey ) / LinearRange + ExposureGrey / 1023.0;
+	//float3 LogColor = (0.432699 * log10(0.5 * LinearColor + 0.037584) + 0.616596) + 0.03;	// SLog
+	//float3 LogColor = ( 300 * log10( LinearColor * (1 - .0108) + .0108 ) + 685 ) / 1023;	// Cineon
+	LogColor = saturate( LogColor );
+	
+	float3 UVW = LogColor * ((LUTSize - 1) / LUTSize) + (0.5f / LUTSize);
+
+#if USE_VOLUME_LUT == 1
+	half3 OutDeviceColor = ColorGradingLUT.Sample( ColorGradingLUTSampler, UVW ).rgb;
+#else
+	half3 OutDeviceColor = UnwrappedTexture3DSample( PostprocessInput3, PostprocessInput3Sampler, UVW, LUTSize ).rgb;
 #endif
+	
+	return OutDeviceColor * 1.05;
 }
 
 // xy should be a integer position (e.g. pixel position on the screen)
@@ -167,8 +178,10 @@ float3 Graph( float2 ScreenSpacePos )
 	float3 Color;
 	Color  = float3( 1, 0, 0 ) * LineShade( F0(p.x), p.y, ( F0(p.x + LineWidth) - F0(p.x - LineWidth) ) / (2 * LineWidth), LineWidth );
 	Color += float3( 0, 1, 0 ) * LineShade( F1(p.x), p.y, ( F1(p.x + LineWidth) - F1(p.x - LineWidth) ) / (2 * LineWidth), LineWidth );
-	//Color += float3( 1, 1, 0 ) * LineShade( F2(Graph.x), Graph.y, ( F2(p.x + LineWidth) - F2(p.x - LineWidth) ) / (2 * LineWidth), LineWidth );
-	//Color += float3( 0, 1, 1 ) * LineShade( F3(Graph.x), Graph.y, ( F3(p.x + LineWidth) - F3(p.x - LineWidth) ) / (2 * LineWidth), LineWidth );
+	Color += float3( 0, 0, 1 ) * LineShade( F2(p.x), p.y, ( F2(p.x + LineWidth) - F2(p.x - LineWidth) ) / (2 * LineWidth), LineWidth );
+	Color += float3( 1, 1, 0 ) * LineShade( F3(p.x), p.y, ( F3(p.x + LineWidth) - F3(p.x - LineWidth) ) / (2 * LineWidth), LineWidth );
+	//Color += float3( 0, 1, 1 ) * LineShade( F4(p.x), p.y, ( F4(p.x + LineWidth) - F4(p.x - LineWidth) ) / (2 * LineWidth), LineWidth );
+	//Color += float3( 1, 1, 1 ) * LineShade( F5(p.x), p.y, ( F5(p.x + LineWidth) - F5(p.x - LineWidth) ) / (2 * LineWidth), LineWidth );
 	return Color;
 }
 
@@ -232,11 +245,7 @@ void MainPS(
 		LinearColor *= ExposureScale;
 
 		#if USE_VIGNETTE
-			#if USE_VIGNETTE_COLOR 
-				LinearColor.rgb = lerp(VignetteColorIntensity.rgb, LinearColor.rgb, ComputeVignetteMask(InExposureScaleVignette.yz, VignetteColorIntensity.a));
-			#else
-				LinearColor.rgb *= ComputeVignetteMask(InExposureScaleVignette.yz, VignetteColorIntensity.a);
-			#endif
+			LinearColor.rgb *= ComputeVignetteMask( InExposureScaleVignette.yz, VignetteIntensity );
 		#endif
 
 		#if USE_GRAIN_INTENSITY
@@ -245,68 +254,49 @@ void MainPS(
 			LinearColor.rgb *= GrainMul;
 		#endif
 
-		#if USE_PHOTO
-			half3 FilmColor = FilmPostProcessPhoto(LinearColor.rgb);
-		#else
-			half3 FilmColor = FilmPostProcess(LinearColor.rgb);
-		#endif
+		half3 OutDeviceColor = ColorLookupTable( LinearColor );
 		
+#if 1	// TODO remove all of this. Do full output device conversion in LUT
+		// remove fake gamma
+		OutDeviceColor *= OutDeviceColor;
+
 		// Apply "gamma" curve adjustment.
-		FilmColor.rgb = pow(FilmColor.rgb, InverseGamma.y);
+		OutDeviceColor.rgb = pow(OutDeviceColor.rgb, InverseGamma.y);
 		
 		#if MAC
 			// Note, MacOSX native output is raw gamma 2.2 not sRGB!
-			half3 TonemappedColor = pow(FilmColor.rgb, 1.0/2.2);
+			OutDeviceColor = pow(OutDeviceColor.rgb, 1.0/2.2);
 		#else
 			#if USE_709
 				// Didn't profile yet if the branching version would be faster (different linear segment).
-				half3 TonemappedColor = LinearTo709Branchless(FilmColor);
+				OutDeviceColor = LinearTo709Branchless(OutDeviceColor);
 			#else
 				#if USE_GAMMA
 					// This is different than the prior "gamma" curve adjustment (but reusing the variable).
 					// For displays set to a gamma colorspace.
-					half3 TonemappedColor = pow(FilmColor.rgb, InverseGamma.z);
+					OutDeviceColor = pow(OutDeviceColor.rgb, InverseGamma.z);
 				#else
 					// Apply conversion to sRGB (this must be an exact sRGB conversion else darks are bad).
 					// Branching is faster than branchless on AMD on PC.
-					half3 TonemappedColor = LinearToSrgbBranching(FilmColor);
+					OutDeviceColor = LinearToSrgbBranching(OutDeviceColor);
 				#endif
 			#endif
 		#endif
+#endif
 
-		//TonemappedColor = rsqrt( ( 0.8 + 0.8 * LinearColor ) / LinearColor );		// 1 vmad, 3 rcp, 1 vmul, 3 rsqrt
-		//TonemappedColor = sqrt( 1 - exp( -LinearColor ) );						// 1 vmul, 3 exp2, 1 vadd, 3 sqrt
-		//TonemappedColor = ( LinearColor * (3.85 * LinearColor + 0.128) ) / ( LinearColor * (3.74 * LinearColor + 1.14) + 0.02 );	// 3 rcp, 3 vmads, 2 vmuls
-		//TonemappedColor = pow( TonemappedColor, 2.0 / 2.2 );
-
-		//TonemappedColor = Graph( ScreenSpacePos );
-
-		#if USE_PHOTO
-			// Blend with custom LDR color, used for Fade track in Matinee.
-			// This is the 101% wrong way to do this,
-			//  - It adds an extra redundant lerp.
-			//  - It is not going to work with the future-forward ES3 fast path of sRGB output.
-			//  - And it does super ugly non-linear blending.
-			// The right way is to adjust exposure instead.
-			TonemappedColor = lerp(TonemappedColor.rgb, OverlayColor.rgb, OverlayColor.a);
-		#endif
-
-		half LuminanceForPostProcessAA  = dot(TonemappedColor, half3(0.299f, 0.587f, 0.114f));
+		//OutDeviceColor = Graph( ScreenSpacePos );
+	
+		half LuminanceForPostProcessAA  = dot(OutDeviceColor, half3(0.299f, 0.587f, 0.114f));
 
 		#if USE_GRAIN_QUANTIZATION
 			// Needs to go after tonemapping.
 			half GrainQuantization = 1.0/256.0;
 			half GrainAdd = (Grain * GrainQuantization) + (-0.5 * GrainQuantization);
-			TonemappedColor.rgb += GrainAdd;
-		#endif
-
-		#if USE_COLOR_GRADING
-			// apply color grading
-			TonemappedColor = ColorLookupTable(TonemappedColor);
+			OutDeviceColor.rgb += GrainAdd;
 		#endif
 
 		// RETURN_COLOR not needed unless writing to SceneColor
-		OutColor = float4(TonemappedColor, LuminanceForPostProcessAA);
+		OutColor = float4(OutDeviceColor, LuminanceForPostProcessAA);
 	#endif
 }
 
@@ -393,6 +383,8 @@ half2 CocBlendScaleBiasFine()
 	return ScaleBias;
 }
 
+float4 OverlayColor;
+
 void MainPS_ES2(
 	in float4 UVAndScreenPos : TEXCOORD0,
 	in float4 FineDofGrain : TEXCOORD1,
@@ -484,18 +476,12 @@ void MainPS_ES2(
 		// Match PC naming.
 		half3 LinearColor = SceneColor.rgb;
 
-		#if USE_HDR_MOSAIC
-			#if USE_VIGNETTE
-				#if USE_VIGNETTE_COLOR 
-					LinearColor.rgb = lerp(VignetteColorIntensity.rgb, LinearColor.rgb, ComputeVignetteMask(InVignette.xy, VignetteColorIntensity.a));
-				#else
-					LinearColor.rgb *= ComputeVignetteMask(InVignette.xy, VignetteColorIntensity.a);
-				#endif
-			#endif
+		#if USE_HDR_MOSAIC && USE_VIGNETTE
+			LinearColor.rgb *= ComputeVignetteMask( InVignette, VignetteIntensity );
 		#endif
 
 		// It is faster to do vignette as a texture lookup + mad because this is an ALU bound shader.
-		#if (!USE_HDR_MOSAIC) && (USE_BLOOM || USE_LIGHT_SHAFTS || USE_VIGNETTE || USE_VIGNETTE_COLOR)
+		#if (!USE_HDR_MOSAIC) && (USE_BLOOM || USE_LIGHT_SHAFTS || USE_VIGNETTE)
 			half4 CombinedBloomSunVignette = Texture2DSample(PostprocessInput1, PostprocessInput1Sampler, FullViewUV.xy);
 			LinearColor.rgb = LinearColor.rgb * CombinedBloomSunVignette.a + CombinedBloomSunVignette.rgb;
 		#endif
diff --git a/Engine/Shaders/TonemapCommon.usf b/Engine/Shaders/TonemapCommon.usf
--- a/Engine/Shaders/TonemapCommon.usf
+++ b/Engine/Shaders/TonemapCommon.usf
@@ -4,6 +4,8 @@
 	TonemapCommon.usf: PostProcessing tone mapping common
 =============================================================================*/
 
+#include "ACES.usf"
+
 // usually 1/2.2, the .y is used for inverse gamma when "gamma only" mode is not used
 half3 InverseGamma;
 
@@ -39,11 +41,6 @@ half3 FilmPostProcess(half3 LinearColor)
 			MatrixColor = LinearColor * ColorMatrixB_ColorCurveCm2.rgb;
 		#endif
 	#endif
-
-	// DON'T CHECK INTO MAIN: NEEDS A PERMUTATION TO AVOID NON-BACKWARDS COMPATIBLE CHANGES
-	// This corrects so the linear segment is linear in log space!
-	MatrixColor = log2(1.0+MatrixColor);
-
 	// Apply color curve (includes tonemapping).
 	#if USE_CONTRAST == 1
 		// Full path.
@@ -65,10 +62,6 @@ half3 FilmPostProcess(half3 LinearColor)
 		half3 CurveColor = (MatrixColorH*ColorCurve_Ch1_Ch2.xxx + ColorCurve_Ch1_Ch2.yyy) * rcp(MatrixColorH + ColorCurve_Cm0Cd0_Cd2_Ch0Cm1_Ch3.www) + MatrixColorM;
 	#endif
 
-	// DON'T CHECK INTO MAIN: NEEDS A PERMUTATION TO AVOID NON-BACKWARDS COMPATIBLE CHANGES
-	// This corrects so the linear segment is linear in log space!
-	CurveColor = exp2(CurveColor)-1.0;
-
 	return CurveColor;
 }
 
@@ -78,18 +71,7 @@ half3 TonemapAndGammaCorrect(half3 LinearColor)
 	// This displays negative colors as black
 	LinearColor = max(LinearColor, 0);
 
-	half TonemapperRange = 20;
-	half TonemapperScaleClamped = 1;
-	half A = 0.22 / TonemapperScaleClamped;
-	half B = (TonemapperRange + A) / TonemapperRange;
-
-	half LinearSteepness = 1;
-	// value below we have linear adjustments to the sigmoid curve that we don't to above the value
-	half FunctionSplitPos = sqrt(A * B / LinearSteepness) - A;
-
-	half3 GammaColor;
-
-	GammaColor = pow(LinearColor, InverseGamma.x);
+	half3 GammaColor = pow(LinearColor, InverseGamma.x);
 	
 	// in all cases it's good to clamp into the 0..1 range (e.g for LUT color grading)
 	GammaColor = saturate(GammaColor);
@@ -97,36 +79,138 @@ half3 TonemapAndGammaCorrect(half3 LinearColor)
 	return GammaColor;
 }
 
-//
-// Improved Film Post
-//
 
-half4 PhotoMatrixR_Lift;
-half4 PhotoMatrixG_LiftScale;
-half4 PhotoMatrixB_Gamma;
-half4 PhotoTint_Over;
-half4 PhotoOverTint_OneMinusMid;
+/*
+============================================
+// Uncharted settings
+Slope = 0.63;
+Toe = 0.55;
+Shoulder = 0.47;
+BlackClip= 0;
+WhiteClip = 0.01;
+
+// HP settings
+Slope = 0.65;
+Toe = 0.63;
+Shoulder = 0.45;
+BlackClip = 0;
+WhiteClip = 0;
+
+// Weta settings
+Slope = 0.6;
+Toe = 0.72;
+Shoulder = 0.55;
+BlackClip = 0;
+WhiteClip = 0.02;
+
+// Legacy settings
+Slope = 0.98;
+Toe = 0.3;
+Shoulder = 0.22;
+BlackClip = 0;
+WhiteClip = 0.025;
+
+// ACES settings
+Slope = 0.88;
+Toe = 0.55;
+Shoulder = 0.26;
+BlackClip = 0;
+WhiteClip = 0.04;
+===========================================
+*/
+
+float FilmSlope;
+float FilmToe;
+float FilmShoulder;
+float FilmBlackClip;
+float FilmWhiteClip;
+
+half3 FilmToneMap( half3 LinearColor ) 
+{
+	const float3x3 sRGB_2_AP0 = mul( XYZ_2_AP0_MAT, mul( D65_2_D60_CAT, sRGB_2_XYZ_MAT ) );
+	const float3x3 sRGB_2_AP1 = mul( XYZ_2_AP1_MAT, mul( D65_2_D60_CAT, sRGB_2_XYZ_MAT ) );
+	const float3x3 AP1_2_sRGB = mul( XYZ_2_sRGB_MAT, mul( D60_2_D65_CAT, AP1_2_XYZ_MAT ) );
+	
+#if 1
+	float3 ACESColor = mul( sRGB_2_AP0, LinearColor );
+
+	// --- Red modifier --- //
+	const float RRT_RED_SCALE = 0.82;
+	const float RRT_RED_PIVOT = 0.03;
+	const float RRT_RED_HUE = 0;
+	const float RRT_RED_WIDTH = 135;
+
+	float saturation = rgb_2_saturation( ACESColor );
+	float hue = rgb_2_hue( ACESColor );
+	float centeredHue = center_hue( hue, RRT_RED_HUE );
+	float hueWeight = Square( smoothstep( 0, 1, 1 - abs( 2 * centeredHue / RRT_RED_WIDTH ) ) );
+		
+	ACESColor.r += hueWeight * saturation * (RRT_RED_PIVOT - ACESColor.r) * (1. - RRT_RED_SCALE);
+
+	// Use ACEScg primaries as working space
+	float3 WorkingColor = mul( AP0_2_AP1_MAT, ACESColor );
+#else
+	// Use ACEScg primaries as working space
+	float3 WorkingColor = mul( sRGB_2_AP1, LinearColor );
+#endif
+
+	WorkingColor = max( 0, WorkingColor );
+
+	// Pre desaturate
+	WorkingColor = lerp( dot( WorkingColor, AP1_RGB2Y ), WorkingColor, 0.96 );
+	
+	const half ToeScale			= 1 + FilmBlackClip - FilmToe;
+	const half ShoulderScale	= 1 + FilmWhiteClip - FilmShoulder;
+	
+	const float InMatch = 0.18;
+	const float OutMatch = 0.18;
+
+	float ToeMatch;
+	if( FilmToe > 0.8 )
+	{
+		// 0.18 will be on straight segment
+		ToeMatch = ( 1 - FilmToe  - OutMatch ) / FilmSlope + log10( InMatch );
+	}
+	else
+	{
+		// 0.18 will be on toe segment
+
+		// Solve for ToeMatch such that input of InMatch gives output of OutMatch.
+		const float bt = ( OutMatch + FilmBlackClip ) / ToeScale - 1;
+		ToeMatch = log10( InMatch ) - 0.5 * log( (1+bt)/(1-bt) ) * (ToeScale / FilmSlope);
+	}
+
+	float StraightMatch = ( 1 - FilmToe ) / FilmSlope - ToeMatch;
+	float ShoulderMatch = FilmShoulder / FilmSlope - StraightMatch;
+	
+	half3 LogColor = log10( WorkingColor );
+	half3 StraightColor = FilmSlope * ( LogColor + StraightMatch );
+	
+	half3 ToeColor		= (    -FilmBlackClip ) + (2 *      ToeScale) / ( 1 + exp( (-2 * FilmSlope /      ToeScale) * ( LogColor -      ToeMatch ) ) );
+	half3 ShoulderColor	= ( 1 + FilmWhiteClip ) - (2 * ShoulderScale) / ( 1 + exp( ( 2 * FilmSlope / ShoulderScale) * ( LogColor - ShoulderMatch ) ) );
+
+	ToeColor		= LogColor <      ToeMatch ?      ToeColor : StraightColor;
+	ShoulderColor	= LogColor > ShoulderMatch ? ShoulderColor : StraightColor;
+
+	half3 t = saturate( ( LogColor - ToeMatch ) / ( ShoulderMatch - ToeMatch ) );
+	t = ShoulderMatch < ToeMatch ? 1 - t : t;
+	t = (3-2*t)*t*t;
+	half3 ToneColor = lerp( ToeColor, ShoulderColor, t );
+
+	// Post desaturate
+	ToneColor = lerp( dot( ToneColor, AP1_RGB2Y ), ToneColor, 0.93 );
 
-half3 FilmPostProcessPhoto(half3 LinearColor) 
+	ToneColor = mul( AP1_2_sRGB, ToneColor );
+
+	//return saturate( ToneColor );
+	return max( 0, ToneColor );
+}
+
+float3 ACES( float3 LinearColor )
 {
-	// Color and exposure control.
-	half3 PhotoColor = LinearColor;
-	// Apply color matrix.
-	PhotoColor.r=dot(LinearColor, PhotoMatrixR_Lift.rgb);
-	PhotoColor.g=dot(LinearColor, PhotoMatrixG_LiftScale.rgb);
-	PhotoColor.b=dot(LinearColor, PhotoMatrixB_Gamma.rgb);
-	// Lift, gamma, tint.
-	PhotoColor = pow(PhotoColor*PhotoMatrixG_LiftScale.a+PhotoMatrixR_Lift.a, PhotoMatrixB_Gamma.a)*PhotoTint_Over.rgb;
-	// Remove negatives.    
-	PhotoColor = max(PhotoColor, 1.0/32768.0);
-	// Move over pure hues towards "white" on over-exposure.
-	// This actually moves over to PhotoOverTint, 
-	// which can be a highly colored value, which later turns to white via highlight compression.
-	half OverLuma = dot(PhotoColor, half3(0.2126, 0.7152, 0.0722)); // BT.709 luma seems to work best.
-	half Over = (PhotoTint_Over.a*OverLuma) / (PhotoOverTint_OneMinusMid.a+OverLuma);
-	PhotoColor = lerp(PhotoColor, PhotoOverTint_OneMinusMid.rgb*OverLuma, Over);
-	// Highlight compression without change of mid-tone.
-	const float WhiteScale = 1.1;
-	PhotoColor *= WhiteScale/(PhotoOverTint_OneMinusMid.a*WhiteScale+PhotoColor);    
-	return PhotoColor;
+	const float3x3 sRGB_2_AP0 = mul( XYZ_2_AP0_MAT, mul( D65_2_D60_CAT, sRGB_2_XYZ_MAT ) );
+
+	float3 aces = mul( sRGB_2_AP0, LinearColor * 1.5 );
+	float3 oces = RRT( aces );
+	return ODT( oces );
 }
diff --git a/Engine/Source/Runtime/Core/Private/HAL/ConsoleManager.cpp b/Engine/Source/Runtime/Core/Private/HAL/ConsoleManager.cpp
--- a/Engine/Source/Runtime/Core/Private/HAL/ConsoleManager.cpp
+++ b/Engine/Source/Runtime/Core/Private/HAL/ConsoleManager.cpp
@@ -1999,13 +1999,6 @@ static TAutoConsoleVariable<int32> CVarTonemapperQuality(
 	TEXT("1: high (default, with high frequency pixel pattern to fight 8 bit color quantization)"),
 	ECVF_Scalability | ECVF_RenderThreadSafe);
 
-static TAutoConsoleVariable<int32> CVarTonemapperPhoto(
-	TEXT("r.TonemapperPhoto"),
-	0,
-	TEXT("0: use old film post\n")
-	TEXT("1: use new photo post"),
-	ECVF_Scalability | ECVF_RenderThreadSafe);
-
 static TAutoConsoleVariable<float> CVarTonemapperHDR(
 	TEXT("r.TonemapperHDR"),
 	1,
diff --git a/Engine/Source/Runtime/Engine/Classes/Engine/Scene.h b/Engine/Source/Runtime/Engine/Classes/Engine/Scene.h
--- a/Engine/Source/Runtime/Engine/Classes/Engine/Scene.h
+++ b/Engine/Source/Runtime/Engine/Classes/Engine/Scene.h
@@ -38,6 +38,22 @@ struct FPostProcessSettings
 
 	// first all bOverride_... as they get grouped together into bitfields
 
+	UPROPERTY(BlueprintReadWrite, Category=Overrides, meta=(PinHiddenByDefault))
+	uint32 bOverride_WhiteTemp:1;
+	UPROPERTY(BlueprintReadWrite, Category=Overrides, meta=(PinHiddenByDefault))
+	uint32 bOverride_WhiteTint:1;
+
+	UPROPERTY(BlueprintReadWrite, Category=Overrides, meta=(PinHiddenByDefault))
+	uint32 bOverride_ColorSaturation:1;
+	UPROPERTY(BlueprintReadWrite, Category=Overrides, meta=(PinHiddenByDefault))
+	uint32 bOverride_ColorContrast:1;
+	UPROPERTY(BlueprintReadWrite, Category=Overrides, meta=(PinHiddenByDefault))
+	uint32 bOverride_ColorGamma:1;
+	UPROPERTY(BlueprintReadWrite, Category=Overrides, meta=(PinHiddenByDefault))
+	uint32 bOverride_ColorGain:1;
+	UPROPERTY(BlueprintReadWrite, Category=Overrides, meta=(PinHiddenByDefault))
+	uint32 bOverride_ColorOffset:1;
+
 	UPROPERTY(BlueprintReadWrite, Category=Overrides, meta=(PinHiddenByDefault))
 	uint32 bOverride_FilmWhitePoint:1;
 
@@ -75,40 +91,18 @@ struct FPostProcessSettings
 	uint32 bOverride_FilmShadowTintAmount:1;
 
 	UPROPERTY(BlueprintReadWrite, Category=Overrides, meta=(PinHiddenByDefault))
-	uint32 bOverride_PhotoExposure:1;
-
-	UPROPERTY(BlueprintReadWrite, Category=Overrides, meta=(PinHiddenByDefault))
-	uint32 bOverride_PhotoMid:1;
+	uint32 bOverride_FilmEnable:1;
 
 	UPROPERTY(BlueprintReadWrite, Category=Overrides, meta=(PinHiddenByDefault))
-	uint32 bOverride_PhotoWhite:1;
-
+	uint32 bOverride_FilmSlope:1;
 	UPROPERTY(BlueprintReadWrite, Category=Overrides, meta=(PinHiddenByDefault))
-	uint32 bOverride_PhotoSaturation:1;
-
+	uint32 bOverride_FilmToe:1;
 	UPROPERTY(BlueprintReadWrite, Category=Overrides, meta=(PinHiddenByDefault))
-	uint32 bOverride_PhotoChannelMixerRed:1;
-
+	uint32 bOverride_FilmShoulder:1;
 	UPROPERTY(BlueprintReadWrite, Category=Overrides, meta=(PinHiddenByDefault))
-	uint32 bOverride_PhotoChannelMixerGreen:1;
-
+	uint32 bOverride_FilmBlackClip:1;
 	UPROPERTY(BlueprintReadWrite, Category=Overrides, meta=(PinHiddenByDefault))
-	uint32 bOverride_PhotoChannelMixerBlue:1;
-
-	UPROPERTY(BlueprintReadWrite, Category=Overrides, meta=(PinHiddenByDefault))
-	uint32 bOverride_PhotoBlackOut:1;
-
-	UPROPERTY(BlueprintReadWrite, Category=Overrides, meta=(PinHiddenByDefault))
-	uint32 bOverride_PhotoContrast:1;
-
-	UPROPERTY(BlueprintReadWrite, Category=Overrides, meta=(PinHiddenByDefault))
-	uint32 bOverride_PhotoTint:1;
-
-	UPROPERTY(BlueprintReadWrite, Category=Overrides, meta=(PinHiddenByDefault))
-	uint32 bOverride_PhotoOver:1;
-
-	UPROPERTY(BlueprintReadWrite, Category=Overrides, meta=(PinHiddenByDefault))
-	uint32 bOverride_PhotoOverTint:1;
+	uint32 bOverride_FilmWhiteClip:1;
 
 	UPROPERTY(BlueprintReadWrite, Category=Overrides, meta=(PinHiddenByDefault))
 	uint32 bOverride_SceneColorTint:1;
@@ -227,9 +221,6 @@ struct FPostProcessSettings
 	UPROPERTY(BlueprintReadWrite, Category=Overrides, meta=(PinHiddenByDefault))
 	uint32 bOverride_VignetteIntensity:1;
 
-	UPROPERTY(BlueprintReadWrite, Category=Overrides, meta=(PinHiddenByDefault))
-	uint32 bOverride_VignetteColor:1;
-
 	UPROPERTY(BlueprintReadWrite, Category=Overrides, meta=(PinHiddenByDefault))
 	uint32 bOverride_GrainIntensity:1;
 
@@ -394,6 +385,21 @@ struct FPostProcessSettings
 
 	// -----------------------------------------------------------------------
 
+	UPROPERTY(interp, BlueprintReadWrite, Category=WhiteBalance, meta=(UIMin = "1500.0", UIMax = "15000.0", editcondition = "bOverride_WhiteTemp", DisplayName = "Temp"))
+	float WhiteTemp;
+	UPROPERTY(interp, BlueprintReadWrite, Category=WhiteBalance, meta=(UIMin = "-1.0", UIMax = "1.0", editcondition = "bOverride_WhiteTint", DisplayName = "Tint"))
+	float WhiteTint;
+
+	UPROPERTY(interp, BlueprintReadWrite, Category=ColorGrading, meta=(UIMin = "0.0", UIMax = "2.0", editcondition = "bOverride_ColorSaturation", DisplayName = "Saturation"))
+	FVector ColorSaturation;
+	UPROPERTY(interp, BlueprintReadWrite, Category=ColorGrading, meta=(UIMin = "0.0", UIMax = "2.0", editcondition = "bOverride_ColorContrast", DisplayName = "Contrast"))
+	FVector ColorContrast;
+	UPROPERTY(interp, BlueprintReadWrite, Category=ColorGrading, meta=(UIMin = "0.0", UIMax = "2.0", editcondition = "bOverride_ColorGamma", DisplayName = "Gamma"))
+	FVector ColorGamma;
+	UPROPERTY(interp, BlueprintReadWrite, Category=ColorGrading, meta=(UIMin = "0.0", UIMax = "2.0", editcondition = "bOverride_ColorGain", DisplayName = "Gain"))
+	FVector ColorGain;
+	UPROPERTY(interp, BlueprintReadWrite, Category=ColorGrading, meta=(UIMin = "-1.0", UIMax = "1.0", editcondition = "bOverride_ColorOffset", DisplayName = "Offset"))
+	FVector ColorOffset;
 
 	UPROPERTY(interp, BlueprintReadWrite, Category=Film, meta=(editcondition = "bOverride_FilmWhitePoint", DisplayName = "Tint", HideAlphaChannel))
 	FLinearColor FilmWhitePoint;
@@ -422,30 +428,19 @@ struct FPostProcessSettings
 	UPROPERTY(interp, BlueprintReadWrite, Category=Film, AdvancedDisplay, meta=(UIMin = "1.0", UIMax = "4.0", editcondition = "bOverride_FilmDynamicRange", DisplayName = "Dynamic Range"))
 	float FilmDynamicRange;
 
-	UPROPERTY(interp, BlueprintReadWrite, Category=Photo, AdvancedDisplay, meta=(UIMin = "0.0", UIMax = "1.0", editcondition = "bOverride_PhotoMid", DisplayName = "Mid-point (linear)"))
-	float PhotoMid;
-	UPROPERTY(interp, BlueprintReadWrite, Category=Photo, AdvancedDisplay, meta=(UIMin = "-16.0", UIMax = "16.0", editcondition = "bOverride_PhotoExposure", DisplayName = "Exposure (+/- stops)"))
-	float PhotoExposure;
-	UPROPERTY(interp, BlueprintReadWrite, Category=Photo, AdvancedDisplay, meta=(editcondition = "bOverride_PhotoWhite", DisplayName = "White (pre-tonemap)", HideAlphaChannel))
-	FLinearColor PhotoWhite;
-	UPROPERTY(interp, BlueprintReadWrite, Category=Photo, AdvancedDisplay, meta=(UIMin = "0.0", UIMax = "2.0", editcondition = "bOverride_PhotoSaturation", DisplayName = "Saturation"))
-	float PhotoSaturation;
-	UPROPERTY(interp, BlueprintReadWrite, Category=Photo, AdvancedDisplay, meta=(editcondition = "bOverride_PhotoChannelMixerRed", DisplayName = "Channel Mixer Red", HideAlphaChannel))
-	FLinearColor PhotoChannelMixerRed;
-	UPROPERTY(interp, BlueprintReadWrite, Category=Photo, AdvancedDisplay, meta=(editcondition = "bOverride_PhotoChannelMixerGreen", DisplayName = "Channel Mixer Green", HideAlphaChannel))
-	FLinearColor PhotoChannelMixerGreen;
-	UPROPERTY(interp, BlueprintReadWrite, Category=Photo, AdvancedDisplay, meta=(editcondition = "bOverride_PhotoChannelMixerBlue", DisplayName = " Channel Mixer Blue", HideAlphaChannel))
-	FLinearColor PhotoChannelMixerBlue;
-	UPROPERTY(interp, BlueprintReadWrite, Category=Photo, AdvancedDisplay, meta=(UIMin = "-6.0", UIMax = "2.0", editcondition = "bOverride_PhotoBlackOut", DisplayName = "Lift (+/- stops from mid^3)"))
-	float PhotoBlackOut;
-	UPROPERTY(interp, BlueprintReadWrite, Category=Photo, AdvancedDisplay, meta=(UIMin = "1.0", UIMax = "4.0", editcondition = "bOverride_PhotoContrast", DisplayName = "Contrast (gamma)"))
-	float PhotoContrast;
-	UPROPERTY(interp, BlueprintReadWrite, Category=Photo, AdvancedDisplay, meta=(editcondition = "bOverride_PhotoTint", DisplayName = "Tint (post-tonemap)", HideAlphaChannel))
-	FLinearColor PhotoTint;
-	UPROPERTY(interp, BlueprintReadWrite, Category=Photo, AdvancedDisplay, meta=(UIMin = "0.0", UIMax = "1.0", editcondition = "bOverride_PhotoOver", DisplayName = "Over Exposure Mix"))
-	float PhotoOver;
-	UPROPERTY(interp, BlueprintReadWrite, Category=Photo, AdvancedDisplay, meta=(editcondition = "bOverride_PhotoOverTint", DisplayName = "Over Exposure Tint", HideAlphaChannel))
-	FLinearColor PhotoOverTint;
+	UPROPERTY(EditAnywhere, BlueprintReadWrite, Category=Film, meta=(editcondition = "bOverride_FilmEnable", DisplayName = "Use film tone mapper"))
+	uint32 FilmEnable:1;
+
+	UPROPERTY(interp, BlueprintReadWrite, Category=Film, AdvancedDisplay, meta=(UIMin = "0.0", UIMax = "1.0", editcondition = "bOverride_FilmSlope", DisplayName = "Slope"))
+	float FilmSlope;
+	UPROPERTY(interp, BlueprintReadWrite, Category=Film, AdvancedDisplay, meta=(UIMin = "0.0", UIMax = "1.0", editcondition = "bOverride_FilmToe", DisplayName = "Toe"))
+	float FilmToe;
+	UPROPERTY(interp, BlueprintReadWrite, Category=Film, AdvancedDisplay, meta=(UIMin = "0.0", UIMax = "1.0", editcondition = "bOverride_FilmShoulder", DisplayName = "Shoulder"))
+	float FilmShoulder;
+	UPROPERTY(interp, BlueprintReadWrite, Category=Film, AdvancedDisplay, meta=(UIMin = "0.0", UIMax = "1.0", editcondition = "bOverride_FilmBlackClip", DisplayName = "Black clip"))
+	float FilmBlackClip;
+	UPROPERTY(interp, BlueprintReadWrite, Category=Film, AdvancedDisplay, meta=(UIMin = "0.0", UIMax = "1.0", editcondition = "bOverride_FilmWhiteClip", DisplayName = "White clip"))
+	float FilmWhiteClip;
 	
 	/** Scene tint color */
 	UPROPERTY(interp, BlueprintReadWrite, Category=SceneColor, AdvancedDisplay, meta=(editcondition = "bOverride_SceneColorTint", HideAlphaChannel))
@@ -691,10 +686,6 @@ struct FPostProcessSettings
 	UPROPERTY(interp, BlueprintReadWrite, Category=SceneColor, meta=(UIMin = "0.0", UIMax = "1.0", editcondition = "bOverride_VignetteIntensity"))
 	float VignetteIntensity;
 
-	/** Vignette color. */
-	UPROPERTY(interp, BlueprintReadWrite, Category=SceneColor, AdvancedDisplay, meta=(editcondition = "bOverride_VignetteColor", DisplayName = "Vignette Color", HideAlphaChannel))
-	FLinearColor VignetteColor;
-
 	/** 0..1 grain jitter */
 	UPROPERTY(interp, BlueprintReadWrite, Category = SceneColor, AdvancedDisplay, meta=(UIMin = "0.0", UIMax = "1.0", editcondition = "bOverride_GrainJitter"))
 	float GrainJitter;
@@ -892,6 +883,15 @@ struct FPostProcessSettings
 		// to set all bOverride_.. by default to false
 		FMemory::Memzero(this, sizeof(FPostProcessSettings));
 
+		WhiteTemp = 6500.0f;
+		WhiteTint = 0.0f;
+
+		ColorSaturation = FVector( 1.0f, 1.0f, 1.0f );
+		ColorContrast = FVector( 1.0f, 1.0f, 1.0f );
+		ColorGamma = FVector( 1.0f, 1.0f, 1.0f );
+		ColorGain = FVector( 1.0f, 1.0f, 1.0f );
+		ColorOffset = FVector( 0.0f, 0.0f, 0.0f );
+
 		// default values:
 		FilmWhitePoint = FLinearColor(1.0f,1.0f,1.0f);
 		FilmSaturation = 1.0f;
@@ -906,18 +906,14 @@ struct FPostProcessSettings
 		FilmShadowTintBlend = 0.5;
 		FilmShadowTintAmount = 0.0;
 
-		PhotoExposure = 0.0f;
-		PhotoMid = 0.18f;
-		PhotoWhite = FLinearColor(1.0f,1.0f,1.0f);
-		PhotoSaturation = 1.0f;
-		PhotoChannelMixerRed = FLinearColor(1.0f,0.0f,0.0f);
-		PhotoChannelMixerGreen = FLinearColor(0.0f,1.0f,0.0f);
-		PhotoChannelMixerBlue = FLinearColor(0.0f,0.0f,1.0f);
-		PhotoBlackOut = -1.0f;
-		PhotoContrast = 1.2f;
-		PhotoTint = FLinearColor(1.0f,1.0f,1.0f);
-		PhotoOver = 1.0f/32.0f;
-		PhotoOverTint = FLinearColor(16.0f,12.0f,1.0f);
+		FilmEnable = 0;
+
+		// ACES settings
+		FilmSlope = 0.88f;
+		FilmToe = 0.55f;
+		FilmShoulder = 0.26f;
+		FilmBlackClip = 0.0f;
+		FilmWhiteClip = 0.04f;
 
 		SceneColorTint = FLinearColor(1, 1, 1);
 		SceneFringeIntensity = 0.0f;
@@ -968,8 +964,7 @@ struct FPostProcessSettings
 		LensFlareTint = FLinearColor(1.0f, 1.0f, 1.0f);
 		LensFlareBokehSize = 3.0f;
 		LensFlareThreshold = 8.0f;
-		VignetteIntensity = 0.0f;
-		VignetteColor = FLinearColor(0.0f, 0.0f, 0.0f);
+		VignetteIntensity = 0.4f;
 		GrainIntensity = 0.0f;
 		GrainJitter = 0.0f;
 		// next value might get overwritten by r.DefaultFeature.AmbientOcclusion
diff --git a/Engine/Source/Runtime/Engine/Private/SceneView.cpp b/Engine/Source/Runtime/Engine/Private/SceneView.cpp
--- a/Engine/Source/Runtime/Engine/Private/SceneView.cpp
+++ b/Engine/Source/Runtime/Engine/Private/SceneView.cpp
@@ -751,6 +751,15 @@ void FSceneView::OverridePostProcessSettings(const FPostProcessSettings& Src, fl
 	FFinalPostProcessSettings& Dest = FinalPostProcessSettings;
 
 	// The following code needs to be adjusted when settings in FPostProcessSettings change.
+	LERP_PP(WhiteTemp);
+	LERP_PP(WhiteTint);
+	
+	LERP_PP(ColorSaturation);
+	LERP_PP(ColorContrast);
+	LERP_PP(ColorGamma);
+	LERP_PP(ColorGain);
+	LERP_PP(ColorOffset);
+
 	LERP_PP(FilmWhitePoint);
 	LERP_PP(FilmSaturation);
 	LERP_PP(FilmChannelMixerRed);
@@ -764,18 +773,11 @@ void FSceneView::OverridePostProcessSettings(const FPostProcessSettings& Src, fl
 	LERP_PP(FilmShadowTintBlend);
 	LERP_PP(FilmShadowTintAmount);
 
-	LERP_PP(PhotoExposure);
-	LERP_PP(PhotoMid);
-	LERP_PP(PhotoWhite);
-	LERP_PP(PhotoSaturation);
-	LERP_PP(PhotoChannelMixerRed);
-	LERP_PP(PhotoChannelMixerGreen);
-	LERP_PP(PhotoChannelMixerBlue);
-	LERP_PP(PhotoBlackOut);
-	LERP_PP(PhotoContrast);
-	LERP_PP(PhotoTint);
-	LERP_PP(PhotoOver);
-	LERP_PP(PhotoOverTint);
+	LERP_PP(FilmSlope);
+	LERP_PP(FilmToe);
+	LERP_PP(FilmShoulder);
+	LERP_PP(FilmBlackClip);
+	LERP_PP(FilmWhiteClip);
 
 	LERP_PP(SceneColorTint);
 	LERP_PP(SceneFringeIntensity);
@@ -816,7 +818,6 @@ void FSceneView::OverridePostProcessSettings(const FPostProcessSettings& Src, fl
 	LERP_PP(LensFlareBokehSize);
 	LERP_PP(LensFlareThreshold);
 	LERP_PP(VignetteIntensity);
-	LERP_PP(VignetteColor);
 	LERP_PP(GrainIntensity);
 	LERP_PP(GrainJitter);
 	LERP_PP(AmbientOcclusionIntensity);
@@ -914,6 +915,11 @@ void FSceneView::OverridePostProcessSettings(const FPostProcessSettings& Src, fl
 		}
 	}
 
+	if(Src.bOverride_FilmEnable)
+	{
+		Dest.FilmEnable = Src.FilmEnable;
+	}
+
 	if(Src.bOverride_DepthOfFieldMethod)
 	{
 		Dest.DepthOfFieldMethod = Src.DepthOfFieldMethod;
@@ -1127,7 +1133,6 @@ void FSceneView::EndFinalPostprocessSettings(const FSceneViewInitOptions& ViewIn
 	if(!Family->EngineShowFlags.Vignette)
 	{
 		FinalPostProcessSettings.VignetteIntensity = 0;
-		FinalPostProcessSettings.VignetteColor = FLinearColor(0.0f, 0.0f, 0.0f);
 	}
 
 	if(!Family->EngineShowFlags.Grain)
diff --git a/Engine/Source/Runtime/Renderer/Private/PostProcess/PostProcessCombineLUTs.cpp b/Engine/Source/Runtime/Renderer/Private/PostProcess/PostProcessCombineLUTs.cpp
--- a/Engine/Source/Runtime/Renderer/Private/PostProcess/PostProcessCombineLUTs.cpp
+++ b/Engine/Source/Runtime/Renderer/Private/PostProcess/PostProcessCombineLUTs.cpp
@@ -35,6 +35,14 @@ static TAutoConsoleVariable<float> CVarColorMax(
 	TEXT("Value should be around 1, smaller values darken the highlights, larger values move more colors towards white, Default: 1"),
 	ECVF_RenderThreadSafe);
 
+int32 GLUTSize = 32;
+static FAutoConsoleVariableRef CVarLUTSize(
+	TEXT("r.LUT.Size"),
+	GLUTSize,
+	TEXT("Size of film LUT"),
+	ECVF_RenderThreadSafe
+	);
+
 // false:use 256x16 texture / true:use volume texture (faster, requires geometry shader)
 // USE_VOLUME_LUT: needs to be the same for C++ and HLSL
 static bool UseVolumeTextureLUT(EShaderPlatform Platform) 
@@ -131,11 +139,39 @@ public:
 		WeightsParameter.Bind(Initializer.ParameterMap, TEXT("LUTWeights"));
 		ColorScale.Bind(Initializer.ParameterMap,TEXT("ColorScale"));
 		OverlayColor.Bind(Initializer.ParameterMap,TEXT("OverlayColor"));
+		InverseGamma.Bind(Initializer.ParameterMap,TEXT("InverseGamma"));
+
+		WhiteTemp.Bind( Initializer.ParameterMap,TEXT("WhiteTemp") );
+		WhiteTint.Bind( Initializer.ParameterMap,TEXT("WhiteTint") );
+
+		ColorSaturation.Bind(	Initializer.ParameterMap,TEXT("ColorSaturation") );
+		ColorContrast.Bind(		Initializer.ParameterMap,TEXT("ColorContrast") );
+		ColorGamma.Bind(		Initializer.ParameterMap,TEXT("ColorGamma") );
+		ColorGain.Bind(			Initializer.ParameterMap,TEXT("ColorGain") );
+		ColorOffset.Bind(		Initializer.ParameterMap,TEXT("ColorOffset") );
+
+		FilmSlope.Bind(		Initializer.ParameterMap,TEXT("FilmSlope") );
+		FilmToe.Bind(		Initializer.ParameterMap,TEXT("FilmToe") );
+		FilmShoulder.Bind(	Initializer.ParameterMap,TEXT("FilmShoulder") );
+		FilmBlackClip.Bind(	Initializer.ParameterMap,TEXT("FilmBlackClip") );
+		FilmWhiteClip.Bind(	Initializer.ParameterMap,TEXT("FilmWhiteClip") );
+
+		ColorMatrixR_ColorCurveCd1.Bind(Initializer.ParameterMap, TEXT("ColorMatrixR_ColorCurveCd1"));
+		ColorMatrixG_ColorCurveCd3Cm3.Bind(Initializer.ParameterMap, TEXT("ColorMatrixG_ColorCurveCd3Cm3"));
+		ColorMatrixB_ColorCurveCm2.Bind(Initializer.ParameterMap, TEXT("ColorMatrixB_ColorCurveCm2"));
+		ColorCurve_Cm0Cd0_Cd2_Ch0Cm1_Ch3.Bind(Initializer.ParameterMap, TEXT("ColorCurve_Cm0Cd0_Cd2_Ch0Cm1_Ch3"));
+		ColorCurve_Ch1_Ch2.Bind(Initializer.ParameterMap, TEXT("ColorCurve_Ch1_Ch2"));
+		ColorShadow_Luma.Bind(Initializer.ParameterMap, TEXT("ColorShadow_Luma"));
+		ColorShadow_Tint1.Bind(Initializer.ParameterMap, TEXT("ColorShadow_Tint1"));
+		ColorShadow_Tint2.Bind(Initializer.ParameterMap, TEXT("ColorShadow_Tint2"));
 	}
 	FLUTBlenderPS() {}
 
 	void SetParameters(FRHICommandList& RHICmdList, const FSceneView& View, FTexture* Texture[BlendCount], float Weights[BlendCount])
 	{
+		const FPostProcessSettings& Settings = View.FinalPostProcessSettings;
+		const FSceneViewFamily& ViewFamily = *(View.Family);
+
 		const FPixelShaderRHIParamRef ShaderRHI = GetPixelShader();
 
 		for(uint32 i = 0; i < BlendCount; ++i)
@@ -152,6 +188,166 @@ public:
 		SetShaderValue(RHICmdList, ShaderRHI, ColorScale, View.ColorScale);
 		SetShaderValue(RHICmdList, ShaderRHI, OverlayColor, View.OverlayColor);
 		ColorRemapShaderParameters.Set(RHICmdList, ShaderRHI);
+
+		// White balance
+		SetShaderValue( RHICmdList, ShaderRHI, WhiteTemp, Settings.WhiteTemp );
+		SetShaderValue( RHICmdList, ShaderRHI, WhiteTint, Settings.WhiteTint );
+
+		// Color grade
+		SetShaderValue( RHICmdList, ShaderRHI, ColorSaturation,	Settings.ColorSaturation );
+		SetShaderValue( RHICmdList, ShaderRHI, ColorContrast,	Settings.ColorContrast );
+		SetShaderValue( RHICmdList, ShaderRHI, ColorGamma,		Settings.ColorGamma );
+		SetShaderValue( RHICmdList, ShaderRHI, ColorGain,		Settings.ColorGain );
+		SetShaderValue( RHICmdList, ShaderRHI, ColorOffset,		Settings.ColorOffset );
+
+		// Film
+		SetShaderValue( RHICmdList, ShaderRHI, FilmSlope,		Settings.FilmSlope );
+		SetShaderValue( RHICmdList, ShaderRHI, FilmToe,			Settings.FilmToe );
+		SetShaderValue( RHICmdList, ShaderRHI, FilmShoulder,	Settings.FilmShoulder );
+		SetShaderValue( RHICmdList, ShaderRHI, FilmBlackClip,	Settings.FilmBlackClip );
+		SetShaderValue( RHICmdList, ShaderRHI, FilmWhiteClip,	Settings.FilmWhiteClip );
+
+		{
+			FVector InvDisplayGammaValue;
+			InvDisplayGammaValue.X = 1.0f / ViewFamily.RenderTarget->GetDisplayGamma();
+			InvDisplayGammaValue.Y = 2.2f / ViewFamily.RenderTarget->GetDisplayGamma();
+			{
+				static TConsoleVariableData<float>* CVar = IConsoleManager::Get().FindTConsoleVariableDataFloat(TEXT("r.TonemapperGamma"));
+				float Value = CVar->GetValueOnRenderThread();
+				if(Value < 1.0f)
+				{
+					Value = 1.0f;
+				}
+				InvDisplayGammaValue.Z = 1.0f / Value;
+			}
+			SetShaderValue(RHICmdList, ShaderRHI, InverseGamma, InvDisplayGammaValue);
+		}
+
+		{
+			// Legacy tone mapper
+			// TODO remove
+
+			// Must insure inputs are in correct range (else possible generation of NaNs).
+			float InExposure = 1.0f;
+			FVector InWhitePoint(Settings.FilmWhitePoint);
+			float InSaturation = FMath::Clamp(Settings.FilmSaturation, 0.0f, 2.0f);
+			FVector InLuma = FVector(1.0f/3.0f, 1.0f/3.0f, 1.0f/3.0f);
+			FVector InMatrixR(Settings.FilmChannelMixerRed);
+			FVector InMatrixG(Settings.FilmChannelMixerGreen);
+			FVector InMatrixB(Settings.FilmChannelMixerBlue);
+			float InContrast = FMath::Clamp(Settings.FilmContrast, 0.0f, 1.0f) + 1.0f;
+			float InDynamicRange = powf(2.0f, FMath::Clamp(Settings.FilmDynamicRange, 1.0f, 4.0f));
+			float InToe = (1.0f - FMath::Clamp(Settings.FilmToeAmount, 0.0f, 1.0f)) * 0.18f;
+			InToe = FMath::Clamp(InToe, 0.18f/8.0f, 0.18f * (15.0f/16.0f));
+			float InHeal = 1.0f - (FMath::Max(1.0f/32.0f, 1.0f - FMath::Clamp(Settings.FilmHealAmount, 0.0f, 1.0f)) * (1.0f - 0.18f)); 
+			FVector InShadowTint(Settings.FilmShadowTint);
+			float InShadowTintBlend = FMath::Clamp(Settings.FilmShadowTintBlend, 0.0f, 1.0f) * 64.0f;
+
+			// Shadow tint amount enables turning off shadow tinting.
+			float InShadowTintAmount = FMath::Clamp(Settings.FilmShadowTintAmount, 0.0f, 1.0f);
+			InShadowTint = InWhitePoint + (InShadowTint - InWhitePoint) * InShadowTintAmount;
+
+			// Make sure channel mixer inputs sum to 1 (+ smart dealing with all zeros).
+			InMatrixR.X += 1.0f / (256.0f*256.0f*32.0f);
+			InMatrixG.Y += 1.0f / (256.0f*256.0f*32.0f);
+			InMatrixB.Z += 1.0f / (256.0f*256.0f*32.0f);
+			InMatrixR *= 1.0f / FVector::DotProduct(InMatrixR, FVector(1.0f));
+			InMatrixG *= 1.0f / FVector::DotProduct(InMatrixG, FVector(1.0f));
+			InMatrixB *= 1.0f / FVector::DotProduct(InMatrixB, FVector(1.0f));
+
+			// Conversion from linear rgb to luma (using HDTV coef).
+			FVector LumaWeights = FVector(0.2126f, 0.7152f, 0.0722f);
+
+			// Make sure white point has 1.0 as luma (so adjusting white point doesn't change exposure).
+			// Make sure {0.0,0.0,0.0} inputs do something sane (default to white).
+			InWhitePoint += FVector(1.0f / (256.0f*256.0f*32.0f));
+			InWhitePoint *= 1.0f / FVector::DotProduct(InWhitePoint, LumaWeights);
+			InShadowTint += FVector(1.0f / (256.0f*256.0f*32.0f));
+			InShadowTint *= 1.0f / FVector::DotProduct(InShadowTint, LumaWeights);
+
+			// Grey after color matrix is applied.
+			FVector ColorMatrixLuma = FVector(
+			FVector::DotProduct(InLuma.X * FVector(InMatrixR.X, InMatrixG.X, InMatrixB.X), FVector(1.0f)),
+			FVector::DotProduct(InLuma.Y * FVector(InMatrixR.Y, InMatrixG.Y, InMatrixB.Y), FVector(1.0f)),
+			FVector::DotProduct(InLuma.Z * FVector(InMatrixR.Z, InMatrixG.Z, InMatrixB.Z), FVector(1.0f)));
+
+			FVector OutMatrixR = FVector(0.0f);
+			FVector OutMatrixG = FVector(0.0f);
+			FVector OutMatrixB = FVector(0.0f);
+			FVector OutColorShadow_Luma = LumaWeights * InShadowTintBlend;
+			FVector OutColorShadow_Tint1 = InWhitePoint;
+			FVector OutColorShadow_Tint2 = InShadowTint - InWhitePoint;
+
+			// Final color matrix effected by saturation and exposure.
+			OutMatrixR = (ColorMatrixLuma + ((InMatrixR - ColorMatrixLuma) * InSaturation)) * InExposure;
+			OutMatrixG = (ColorMatrixLuma + ((InMatrixG - ColorMatrixLuma) * InSaturation)) * InExposure;
+			OutMatrixB = (ColorMatrixLuma + ((InMatrixB - ColorMatrixLuma) * InSaturation)) * InExposure;
+
+			// Line for linear section.
+			float FilmLineOffset = 0.18f - 0.18f*InContrast;
+			float FilmXAtY0 = -FilmLineOffset/InContrast;
+			float FilmXAtY1 = (1.0f - FilmLineOffset) / InContrast;
+			float FilmXS = FilmXAtY1 - FilmXAtY0;
+
+			// Coordinates of linear section.
+			float FilmHiX = FilmXAtY0 + InHeal*FilmXS;
+			float FilmHiY = FilmHiX*InContrast + FilmLineOffset;
+			float FilmLoX = FilmXAtY0 + InToe*FilmXS;
+			float FilmLoY = FilmLoX*InContrast + FilmLineOffset;
+			// Supported exposure range before clipping.
+			float FilmHeal = InDynamicRange - FilmHiX;
+			// Intermediates.
+			float FilmMidXS = FilmHiX - FilmLoX;
+			float FilmMidYS = FilmHiY - FilmLoY;
+			float FilmSlope = FilmMidYS / (FilmMidXS);
+			float FilmHiYS = 1.0f - FilmHiY;
+			float FilmLoYS = FilmLoY;
+			float FilmToe = FilmLoX;
+			float FilmHiG = (-FilmHiYS + (FilmSlope*FilmHeal)) / (FilmSlope*FilmHeal);
+			float FilmLoG = (-FilmLoYS + (FilmSlope*FilmToe)) / (FilmSlope*FilmToe);
+
+			// Constants.
+			float OutColorCurveCh1 = FilmHiYS/FilmHiG;
+			float OutColorCurveCh2 = -FilmHiX*(FilmHiYS/FilmHiG);
+			float OutColorCurveCh3 = FilmHiYS/(FilmSlope*FilmHiG) - FilmHiX;
+			float OutColorCurveCh0Cm1 = FilmHiX;
+			float OutColorCurveCm2 = FilmSlope;
+			float OutColorCurveCm0Cd0 = FilmLoX;
+			float OutColorCurveCd3Cm3 = FilmLoY - FilmLoX*FilmSlope;
+			float OutColorCurveCd1 = 0.0f;
+			float OutColorCurveCd2 = 1.0f;
+			// Handle these separate in case of FilmLoG being 0.
+			if(FilmLoG != 0.0f)
+			{
+				OutColorCurveCd1 = -FilmLoYS/FilmLoG;
+				OutColorCurveCd2 = FilmLoYS/(FilmSlope*FilmLoG);
+			}
+			else
+			{
+				// FilmLoG being zero means dark region is a linear segment (so just continue the middle section).
+				OutColorCurveCm0Cd0 = 0.0f;
+				OutColorCurveCd3Cm3 = 0.0f;
+			}
+
+			FVector4 Constants[8];
+			Constants[0] = FVector4(OutMatrixR, OutColorCurveCd1);
+			Constants[1] = FVector4(OutMatrixG, OutColorCurveCd3Cm3);
+			Constants[2] = FVector4(OutMatrixB, OutColorCurveCm2); 
+			Constants[3] = FVector4(OutColorCurveCm0Cd0, OutColorCurveCd2, OutColorCurveCh0Cm1, OutColorCurveCh3); 
+			Constants[4] = FVector4(OutColorCurveCh1, OutColorCurveCh2, 0.0f, 0.0f);
+			Constants[5] = FVector4(OutColorShadow_Luma, 0.0f);
+			Constants[6] = FVector4(OutColorShadow_Tint1, 0.0f);
+			Constants[7] = FVector4(OutColorShadow_Tint2, Settings.FilmEnable);
+
+			SetShaderValue(RHICmdList, ShaderRHI, ColorMatrixR_ColorCurveCd1, Constants[0]);
+			SetShaderValue(RHICmdList, ShaderRHI, ColorMatrixG_ColorCurveCd3Cm3, Constants[1]);
+			SetShaderValue(RHICmdList, ShaderRHI, ColorMatrixB_ColorCurveCm2, Constants[2]); 
+			SetShaderValue(RHICmdList, ShaderRHI, ColorCurve_Cm0Cd0_Cd2_Ch0Cm1_Ch3, Constants[3]); 
+			SetShaderValue(RHICmdList, ShaderRHI, ColorCurve_Ch1_Ch2, Constants[4]);
+			SetShaderValue(RHICmdList, ShaderRHI, ColorShadow_Luma, Constants[5]);
+			SetShaderValue(RHICmdList, ShaderRHI, ColorShadow_Tint1, Constants[6]);
+			SetShaderValue(RHICmdList, ShaderRHI, ColorShadow_Tint2, Constants[7]);
+		}
 	}
 
 	static void ModifyCompilationEnvironment(EShaderPlatform Platform, FShaderCompilerEnvironment& OutEnvironment)
@@ -160,6 +356,8 @@ public:
 
 		OutEnvironment.SetDefine(TEXT("BLENDCOUNT"), BlendCount);
 		OutEnvironment.SetDefine(TEXT("USE_VOLUME_LUT"), UseVolumeTextureLUT(Platform));
+		//OutEnvironment.SetDefine(TEXT("USE_GAMMA"), OutputDevice == 1);
+		//OutEnvironment.SetDefine(TEXT("USE_709"), OutputDevice == 2);
 	}
 
 	virtual bool Serialize(FArchive& Ar) override
@@ -174,6 +372,24 @@ public:
 		Ar << WeightsParameter << ColorScale << OverlayColor;
 		Ar << ColorRemapShaderParameters;
 
+		Ar << WhiteTemp;
+		Ar << WhiteTint;
+
+		Ar << ColorSaturation;
+		Ar << ColorContrast;
+		Ar << ColorGamma;
+		Ar << ColorGain;
+		Ar << ColorOffset;
+
+		Ar << FilmSlope;
+		Ar << FilmToe;
+		Ar << FilmShoulder;
+		Ar << FilmBlackClip;
+		Ar << FilmWhiteClip;
+
+		Ar	<< ColorMatrixR_ColorCurveCd1 << ColorMatrixG_ColorCurveCd3Cm3 << ColorMatrixB_ColorCurveCm2
+			<< ColorCurve_Cm0Cd0_Cd2_Ch0Cm1_Ch3 << ColorCurve_Ch1_Ch2 << ColorShadow_Luma << ColorShadow_Tint1 << ColorShadow_Tint2;
+
 		return bShaderHasOutdatedParameters;
 	}
 
@@ -185,7 +401,33 @@ private: // ---------------------------------------------------
 	FShaderParameter WeightsParameter;
 	FShaderParameter ColorScale;
 	FShaderParameter OverlayColor;
+	FShaderParameter InverseGamma;
 	FColorRemapShaderParameters ColorRemapShaderParameters;
+
+	FShaderParameter WhiteTemp;
+	FShaderParameter WhiteTint;
+
+	FShaderParameter ColorSaturation;
+	FShaderParameter ColorContrast;
+	FShaderParameter ColorGamma;
+	FShaderParameter ColorGain;
+	FShaderParameter ColorOffset;
+
+	FShaderParameter FilmSlope;
+	FShaderParameter FilmToe;
+	FShaderParameter FilmShoulder;
+	FShaderParameter FilmBlackClip;
+	FShaderParameter FilmWhiteClip;
+
+	// Legacy
+	FShaderParameter ColorMatrixR_ColorCurveCd1;
+	FShaderParameter ColorMatrixG_ColorCurveCd3Cm3;
+	FShaderParameter ColorMatrixB_ColorCurveCm2;
+	FShaderParameter ColorCurve_Cm0Cd0_Cd2_Ch0Cm1_Ch3;
+	FShaderParameter ColorCurve_Ch1_Ch2;
+	FShaderParameter ColorShadow_Luma;
+	FShaderParameter ColorShadow_Tint1;
+	FShaderParameter ColorShadow_Tint2;
 };
 
 
@@ -400,7 +642,7 @@ void FRCPassPostProcessCombineLUTs::Process(FRenderingCompositePassContext& Cont
 	}
 
 	// for a 3D texture, the viewport is 16x16 (per slice), for a 2D texture, it's unwrapped to 256x16
-	FIntPoint DestSize(UseVolumeTextureLUT(ShaderPlatform) ? 16 : 256, 16);
+	FIntPoint DestSize(UseVolumeTextureLUT(ShaderPlatform) ? GLUTSize : GLUTSize * GLUTSize, GLUTSize);
 
 	const FSceneRenderTargetItem& DestRenderTarget = PassOutputs[0].RequestSurface(Context);
 
@@ -413,7 +655,7 @@ void FRCPassPostProcessCombineLUTs::Process(FRenderingCompositePassContext& Cont
 	Context.RHICmdList.SetRasterizerState(TStaticRasterizerState<>::GetRHI());
 	Context.RHICmdList.SetDepthStencilState(TStaticDepthStencilState<false, CF_Always>::GetRHI());
 
-	const FVolumeBounds VolumeBounds(16);
+	const FVolumeBounds VolumeBounds(GLUTSize);
 
 	SetLUTBlenderShader(Context, LocalCount, LocalTextures, LocalWeights, VolumeBounds);
 
@@ -429,12 +671,12 @@ void FRCPassPostProcessCombineLUTs::Process(FRenderingCompositePassContext& Cont
 
 		DrawRectangle( 
 			Context.RHICmdList,
-			0, 0,						// XY
-			16 * 16, 16,				// SizeXY
-			0, 0,						// UV
-			16 * 16, 16,				// SizeUV
-			FIntPoint(16 * 16, 16),		// TargetSize
-			FIntPoint(16 * 16, 16),		// TextureSize
+			0, 0,										// XY
+			GLUTSize * GLUTSize, GLUTSize,				// SizeXY
+			0, 0,										// UV
+			GLUTSize * GLUTSize, GLUTSize,				// SizeUV
+			FIntPoint(GLUTSize * GLUTSize, GLUTSize),	// TargetSize
+			FIntPoint(GLUTSize * GLUTSize, GLUTSize),	// TextureSize
 			*VertexShader,
 			EDRF_UseTriangleOptimization);
 	}
@@ -444,66 +686,15 @@ void FRCPassPostProcessCombineLUTs::Process(FRenderingCompositePassContext& Cont
 
 FPooledRenderTargetDesc FRCPassPostProcessCombineLUTs::ComputeOutputDesc(EPassOutputId InPassOutputId) const
 {
-	FPooledRenderTargetDesc Ret = FPooledRenderTargetDesc::Create2DDesc(FIntPoint(256, 16), PF_B8G8R8A8, TexCreate_None, TexCreate_RenderTargetable | TexCreate_ShaderResource, false);
+	FPooledRenderTargetDesc Ret = FPooledRenderTargetDesc::Create2DDesc(FIntPoint(GLUTSize * GLUTSize, GLUTSize), PF_A2B10G10R10, TexCreate_None, TexCreate_RenderTargetable | TexCreate_ShaderResource, false);
 
 	if(UseVolumeTextureLUT(ShaderPlatform))
 	{
-		Ret.Extent = FIntPoint(16, 16);
-		Ret.Depth = 16;
+		Ret.Extent = FIntPoint(GLUTSize, GLUTSize);
+		Ret.Depth = GLUTSize;
 	}
 
 	Ret.DebugName = TEXT("CombineLUTs");
 
 	return Ret;
-}
-
-bool FRCPassPostProcessCombineLUTs::IsColorGradingLUTNeeded(const FViewInfo* RESTRICT View)
-{
-	check(View);
-
-	FColorTransform ColorTransform;
-	ColorTransform.MinValue = FMath::Clamp(CVarColorMin.GetValueOnRenderThread(), -10.0f, 10.0f);
-	ColorTransform.MidValue = FMath::Clamp(CVarColorMid.GetValueOnRenderThread(), -10.0f, 10.0f);
-	ColorTransform.MaxValue = FMath::Clamp(CVarColorMax.GetValueOnRenderThread(), -10.0f, 10.0f);
-
-	if(ColorTransform.MinValue != 0.0f)
-	{
-		return true;
-	}
-	if(ColorTransform.MidValue != 0.5f)
-	{
-		return true;
-	}
-	if(ColorTransform.MaxValue != 1.0f)
-	{
-		return true;
-	}
-
-	if(View->ColorScale.R != 1.0f || View->ColorScale.G != 1.0f || View->ColorScale.B != 1.0f)
-	{
-		return true;
-	}
-
-	if(View->OverlayColor.A > (1.0f/512.0f))
-	{
-		return true;
-	}
-
-	if(View->GetFeatureLevel() >= ERHIFeatureLevel::SM4)
-	{
-		const FFinalPostProcessSettings Settings = View->FinalPostProcessSettings;
-		for(uint32 i = 0; i < (uint32)Settings.ContributingLUTs.Num(); ++i)
-		{
-			UTexture* LUTTexture = Settings.ContributingLUTs[i].LUTTexture;
-			if(LUTTexture == NULL) 
-			{
-				continue;
-			}
-			if(Settings.ContributingLUTs[i].Weight > 1.0f / 512.0f)
-			{
-				return true;
-			}
-		}
-	}
-	return false;
-}
+}
\ No newline at end of file
diff --git a/Engine/Source/Runtime/Renderer/Private/PostProcess/PostProcessCombineLUTs.h b/Engine/Source/Runtime/Renderer/Private/PostProcess/PostProcessCombineLUTs.h
--- a/Engine/Source/Runtime/Renderer/Private/PostProcess/PostProcessCombineLUTs.h
+++ b/Engine/Source/Runtime/Renderer/Private/PostProcess/PostProcessCombineLUTs.h
@@ -28,8 +28,6 @@ public:
 	/** @return 0xffffffff if not found */
 	uint32 FindIndex(const FFinalPostProcessSettings& Settings, UTexture* Tex) const;
 
-	static bool IsColorGradingLUTNeeded(const FViewInfo* RESTRICT View);
-
 private:
 
 	EShaderPlatform ShaderPlatform;
diff --git a/Engine/Source/Runtime/Renderer/Private/PostProcess/PostProcessMobile.cpp b/Engine/Source/Runtime/Renderer/Private/PostProcess/PostProcessMobile.cpp
--- a/Engine/Source/Runtime/Renderer/Private/PostProcess/PostProcessMobile.cpp
+++ b/Engine/Source/Runtime/Renderer/Private/PostProcess/PostProcessMobile.cpp
@@ -1450,8 +1450,6 @@ public:
 		SunColorVignetteIntensityParam.W = Settings.VignetteIntensity;
 		SetShaderValue(Context.RHICmdList, ShaderRHI, SunColorVignetteIntensity, SunColorVignetteIntensityParam);
 
-		SetShaderValue(Context.RHICmdList, ShaderRHI, VignetteColor, Context.View.FinalPostProcessSettings.VignetteColor);
-
 		// Scaling Bloom1 by extra factor to match filter area difference between PC default and mobile.
 		SetShaderValue(Context.RHICmdList, ShaderRHI, BloomColor, Context.View.FinalPostProcessSettings.Bloom1Tint * Context.View.FinalPostProcessSettings.BloomIntensity * 0.5);
 	}
@@ -1459,7 +1457,7 @@ public:
 	virtual bool Serialize(FArchive& Ar) override
 	{
 		bool bShaderHasOutdatedParameters = FGlobalShader::Serialize(Ar);
-		Ar << PostprocessParameter << SunColorVignetteIntensity << VignetteColor << BloomColor;
+		Ar << PostprocessParameter << SunColorVignetteIntensity << BloomColor;
 		return bShaderHasOutdatedParameters;
 	}
 };
@@ -1679,8 +1677,6 @@ public:
 		SunColorVignetteIntensityParam.W = Settings.VignetteIntensity;
 		SetShaderValue(Context.RHICmdList, ShaderRHI, SunColorVignetteIntensity, SunColorVignetteIntensityParam);
 
-		SetShaderValue(Context.RHICmdList, ShaderRHI, VignetteColor, Context.View.FinalPostProcessSettings.VignetteColor);
-
 		// Scaling Bloom1 by extra factor to match filter area difference between PC default and mobile.
 		SetShaderValue(Context.RHICmdList, ShaderRHI, BloomColor, Context.View.FinalPostProcessSettings.Bloom1Tint * Context.View.FinalPostProcessSettings.BloomIntensity * 0.5);
 		SetShaderValue(Context.RHICmdList, ShaderRHI, BloomColor2, Context.View.FinalPostProcessSettings.Bloom2Tint * Context.View.FinalPostProcessSettings.BloomIntensity);
@@ -1689,7 +1685,7 @@ public:
 	virtual bool Serialize(FArchive& Ar) override
 	{
 		bool bShaderHasOutdatedParameters = FGlobalShader::Serialize(Ar);
-		Ar << PostprocessParameter << SunColorVignetteIntensity << VignetteColor << BloomColor << BloomColor2;
+		Ar << PostprocessParameter << SunColorVignetteIntensity << BloomColor << BloomColor2;
 		return bShaderHasOutdatedParameters;
 	}
 };
@@ -2603,14 +2599,16 @@ public:
 		FSceneViewState* ViewState = (FSceneViewState*)Context.View.State;
 		if(ViewState)
 		{
-			FMatrix Proj = Context.View.ViewMatrices.ProjMatrix;
+			const FViewInfo& View = Context.View;
+
+			FMatrix Proj = View.ViewMatrices.ProjMatrix;
 			FMatrix PrevProj = ViewState->PrevViewMatrices.ProjMatrix;
 
 			// Remove jitter
-			Proj.M[2][0] = 0.0f;
-			Proj.M[2][1] = 0.0f;
-			PrevProj.M[2][0] = 0.0f;
-			PrevProj.M[2][1] = 0.0f;
+			Proj.M[2][0] -= View.ViewMatrices.TemporalAASample.X * 2.0f / View.ViewRect.Width();
+			Proj.M[2][1] -= View.ViewMatrices.TemporalAASample.Y * 2.0f / View.ViewRect.Height();
+			PrevProj.M[2][0] -= ViewState->PrevViewMatrices.TemporalAASample.X * 2.0f / View.ViewRect.Width();
+			PrevProj.M[2][1] -= ViewState->PrevViewMatrices.TemporalAASample.Y * 2.0f / View.ViewRect.Height();
 
 			FMatrix ViewProj = ( Context.View.ViewMatrices.ViewMatrix * Proj ).GetTransposed();
 			FMatrix PrevViewProj = ( ViewState->PrevViewMatrices.ViewMatrix * PrevProj ).GetTransposed();
@@ -2786,8 +2784,8 @@ void FRCPassPostProcessAaES2::Process(FRenderingCompositePassContext& Context)
 	if (FSceneRenderer::ShouldCompositeEditorPrimitives(View))
 	{
 		// Remove jitter (ensures editor prims are stable.)
-		View.ViewMatrices.ProjMatrix.M[2][0] = 0.0f;
-		View.ViewMatrices.ProjMatrix.M[2][1] = 0.0f;
+		View.ViewMatrices.ProjMatrix.M[2][0] -= View.ViewMatrices.TemporalAASample.X * 2.0f / View.ViewRect.Width();
+		View.ViewMatrices.ProjMatrix.M[2][1] -= View.ViewMatrices.TemporalAASample.Y * 2.0f / View.ViewRect.Height();
 
 		// Compute the view projection matrix and its inverse.
 		View.ViewProjectionMatrix = View.ViewMatrices.ViewMatrix * View.ViewMatrices.ProjMatrix;
diff --git a/Engine/Source/Runtime/Renderer/Private/PostProcess/PostProcessTonemap.cpp b/Engine/Source/Runtime/Renderer/Private/PostProcess/PostProcessTonemap.cpp
--- a/Engine/Source/Runtime/Renderer/Private/PostProcess/PostProcessTonemap.cpp
+++ b/Engine/Source/Runtime/Renderer/Private/PostProcess/PostProcessTonemap.cpp
@@ -31,15 +31,12 @@ typedef enum {
 	TonemapperBloom             = (1<<7),
 	TonemapperDOF               = (1<<8),
 	TonemapperVignette          = (1<<9),
-	TonemapperVignetteColor     = (1<<10),
-	TonemapperLightShafts       = (1<<11),
-	TonemapperMosaic            = (1<<12),
-	TonemapperColorFringe       = (1<<13),
-	TonemapperColorGrading      = (1<<14),
-	TonemapperMsaa              = (1<<15),
-	TonemapperPhoto             = (1<<16),
-	Tonemapper709               = (1<<17),
-	TonemapperGamma             = (1<<18),
+	TonemapperLightShafts       = (1<<10),
+	TonemapperMosaic            = (1<<11),
+	TonemapperColorFringe       = (1<<12),
+	TonemapperMsaa              = (1<<13),
+	Tonemapper709               = (1<<14),
+	TonemapperGamma             = (1<<15),
 } TonemapperOption;
 
 // Tonemapper option cost (0 = no cost, 255 = max cost).
@@ -56,13 +53,10 @@ static uint8 TonemapperCostTab[] = {
 	1, //TonemapperBloom
 	1, //TonemapperDOF
 	1, //TonemapperVignette
-	1, //TonemapperVignetteColor
 	1, //TonemapperLightShafts
 	1, //TonemapperMosaic
 	1, //TonemapperColorFringe
-	1, //TonemapperColorGrading
 	1, //TonemapperMsaa
-	1, //TonemapperPhoto
 	1, //Tonemapper709
 	1, //TonemapperGamma
 };
@@ -72,206 +66,64 @@ static uint8 TonemapperCostTab[] = {
 // Place most common first (faster when searching in TonemapperFindLeastExpensive()).
 
 // List of configurations compiled for PC.
-static uint32 TonemapperConfBitmaskPC[22] = { 
+static uint32 TonemapperConfBitmaskPC[9] = { 
 
 	TonemapperBloom +
-	TonemapperContrast +
-	TonemapperColorMatrix +
-	TonemapperShadowTint +
 	TonemapperGrainJitter +
 	TonemapperGrainIntensity +
 	TonemapperGrainQuantization +
 	TonemapperVignette +
-	TonemapperVignetteColor +
 	TonemapperColorFringe +
-	TonemapperColorGrading +
 	0,
 
 	TonemapperBloom +
-	TonemapperContrast +
-	TonemapperColorMatrix +
-	TonemapperShadowTint +
 	TonemapperGrainJitter +
 	TonemapperGrainIntensity +
 	TonemapperGrainQuantization +
 	TonemapperVignette +
-	TonemapperVignetteColor +
 	TonemapperColorFringe +
-	TonemapperColorGrading +
 	Tonemapper709 +
 	0,
 
 	TonemapperBloom + 
-	TonemapperContrast + 
-	TonemapperColorMatrix + 
-	TonemapperShadowTint +
 	TonemapperVignette +
 	TonemapperGrainQuantization +
 	TonemapperColorFringe +
-	TonemapperColorGrading +
 	0,
 
 	TonemapperBloom + 
-	TonemapperContrast + 
-	TonemapperColorMatrix +
-	TonemapperShadowTint +
 	TonemapperGrainQuantization +
-	TonemapperColorGrading +
-	0,
-
-	TonemapperBloom + 
-	TonemapperContrast +
-	TonemapperColorMatrix +
-	TonemapperGrainQuantization +
-	TonemapperColorGrading +
 	0,
 
 	// same without TonemapperGrainQuantization
 
 	TonemapperBloom +
-	TonemapperContrast +
-	TonemapperColorMatrix +
-	TonemapperShadowTint +
 	TonemapperGrainJitter +
 	TonemapperGrainIntensity +
 	TonemapperVignette +
-	TonemapperVignetteColor +
 	TonemapperColorFringe +
-	TonemapperColorGrading +
 	0,
 
 	TonemapperBloom +
-	TonemapperContrast +
-	TonemapperColorMatrix +
-	TonemapperShadowTint +
 	TonemapperGrainJitter +
 	TonemapperGrainIntensity +
 	TonemapperVignette +
-	TonemapperVignetteColor +
 	TonemapperColorFringe +
-	TonemapperColorGrading +
 	Tonemapper709 +
 	0,
 
 	TonemapperBloom + 
-	TonemapperContrast + 
-	TonemapperColorMatrix + 
-	TonemapperShadowTint +
 	TonemapperVignette +
 	TonemapperColorFringe +
-	TonemapperColorGrading +
 	0,
 
 	TonemapperBloom + 
-	TonemapperContrast + 
-	TonemapperColorMatrix +
-	TonemapperShadowTint +
-	TonemapperColorGrading +
 	0,
 
-	TonemapperBloom + 
-	TonemapperContrast +
-	TonemapperColorMatrix +
-	TonemapperColorGrading +
-	0,
-	
 	//
 
 	TonemapperGammaOnly +
-	TonemapperColorGrading +
-	0,
-
-// SAME WITHOUT COLORGRADING WITH NEW PHOTO POST
-
-	TonemapperBloom +
-	TonemapperGrainJitter +
-	TonemapperGrainIntensity +
-	TonemapperGrainQuantization +
-	TonemapperVignette +
-	TonemapperVignetteColor +
-	TonemapperColorFringe +
-	TonemapperPhoto +
-	0,
-
-	TonemapperBloom +
-	TonemapperGrainJitter +
-	TonemapperGrainIntensity +
-	TonemapperGrainQuantization +
-	TonemapperVignette +
-	TonemapperVignetteColor +
-	TonemapperPhoto +
-	0,
-
-	TonemapperBloom +
-	TonemapperGrainJitter +
-	TonemapperGrainIntensity +
-	TonemapperGrainQuantization +
-	TonemapperVignette +
-	TonemapperVignetteColor +
-	TonemapperColorFringe +
-	TonemapperPhoto +
-	Tonemapper709 +
 	0,
-
-	TonemapperBloom +
-	TonemapperGrainJitter +
-	TonemapperGrainIntensity +
-	TonemapperGrainQuantization +
-	TonemapperVignette +
-	TonemapperVignetteColor +
-	TonemapperColorFringe +
-	TonemapperPhoto +
-	TonemapperGamma +
-	0,
-
-	TonemapperBloom +
-	TonemapperGrainQuantization +
-	TonemapperPhoto +
-	0,
-
-
-	// similar without TonemapperGrain
-
-	TonemapperBloom +
-	TonemapperVignette +
-	TonemapperVignetteColor +
-	TonemapperColorFringe +
-	TonemapperPhoto +
-	0,
-
-	TonemapperBloom +
-	TonemapperVignette +
-	TonemapperVignetteColor +
-	TonemapperPhoto +
-	0,
-
-	TonemapperBloom +
-	TonemapperVignette +
-	TonemapperVignetteColor +
-	TonemapperColorFringe +
-	TonemapperPhoto +
-	Tonemapper709 +
-	0,
-
-	TonemapperBloom +
-	TonemapperVignette +
-	TonemapperVignetteColor +
-	TonemapperColorFringe +
-	TonemapperPhoto +
-	TonemapperGamma +
-	0,
-
-	TonemapperBloom +
-	TonemapperPhoto +
-	0,
-
-
-	//
-
-	TonemapperGammaOnly +
-	0,
-
-
 };
 
 // List of configurations compiled for Mobile.
@@ -299,7 +151,6 @@ static uint32 TonemapperConfBitmaskMobile[39] = {
 	TonemapperBloom +
 	TonemapperLightShafts +
 	TonemapperVignette +
-	TonemapperVignetteColor +
 	0,
 
 	TonemapperContrast +
@@ -307,7 +158,6 @@ static uint32 TonemapperConfBitmaskMobile[39] = {
 	TonemapperBloom +
 	TonemapperLightShafts +
 	TonemapperVignette +
-	TonemapperVignetteColor +
 	0,
 
 	TonemapperContrast +
@@ -316,7 +166,6 @@ static uint32 TonemapperConfBitmaskMobile[39] = {
 	TonemapperBloom +
 	TonemapperLightShafts +
 	TonemapperVignette +
-	TonemapperVignetteColor +
 	0,
 
 	// DOF enabled.
@@ -325,7 +174,6 @@ static uint32 TonemapperConfBitmaskMobile[39] = {
 	TonemapperBloom +
 	TonemapperLightShafts +
 	TonemapperVignette +
-	TonemapperVignetteColor +
 	TonemapperDOF +
 	0,
 
@@ -335,7 +183,6 @@ static uint32 TonemapperConfBitmaskMobile[39] = {
 	TonemapperBloom +
 	TonemapperLightShafts +
 	TonemapperVignette +
-	TonemapperVignetteColor +
 	TonemapperDOF +
 	0,
 
@@ -354,7 +201,6 @@ static uint32 TonemapperConfBitmaskMobile[39] = {
 	TonemapperBloom +
 	TonemapperLightShafts +
 	TonemapperVignette +
-	TonemapperVignetteColor +
 	TonemapperGrainIntensity +
 	0,
 
@@ -363,7 +209,6 @@ static uint32 TonemapperConfBitmaskMobile[39] = {
 	TonemapperBloom +
 	TonemapperLightShafts +
 	TonemapperVignette +
-	TonemapperVignetteColor +
 	TonemapperGrainIntensity +
 	0,
 
@@ -373,7 +218,6 @@ static uint32 TonemapperConfBitmaskMobile[39] = {
 	TonemapperBloom +
 	TonemapperLightShafts +
 	TonemapperVignette +
-	TonemapperVignetteColor +
 	TonemapperGrainIntensity +
 	0,
 
@@ -383,7 +227,6 @@ static uint32 TonemapperConfBitmaskMobile[39] = {
 	TonemapperBloom +
 	TonemapperLightShafts +
 	TonemapperVignette +
-	TonemapperVignetteColor +
 	TonemapperDOF +
 	TonemapperGrainIntensity +
 	0,
@@ -394,7 +237,6 @@ static uint32 TonemapperConfBitmaskMobile[39] = {
 	TonemapperBloom +
 	TonemapperLightShafts +
 	TonemapperVignette +
-	TonemapperVignetteColor +
 	TonemapperDOF +
 	TonemapperGrainIntensity +
 	0,
@@ -430,14 +272,12 @@ static uint32 TonemapperConfBitmaskMobile[39] = {
 	TonemapperMosaic + 
 	TonemapperContrast +
 	TonemapperVignette +
-	TonemapperVignetteColor +
 	0,
 
 	TonemapperMosaic + 
 	TonemapperContrast +
 	TonemapperColorMatrix +
 	TonemapperVignette +
-	TonemapperVignetteColor +
 	0,
 
 	TonemapperMosaic + 
@@ -445,7 +285,6 @@ static uint32 TonemapperConfBitmaskMobile[39] = {
 	TonemapperColorMatrix +
 	TonemapperShadowTint +
 	TonemapperVignette +
-	TonemapperVignetteColor +
 	0,
 
 	// With grain
@@ -471,7 +310,6 @@ static uint32 TonemapperConfBitmaskMobile[39] = {
 	TonemapperMosaic + 
 	TonemapperContrast +
 	TonemapperVignette +
-	TonemapperVignetteColor +
 	TonemapperGrainIntensity +
 	0,
 
@@ -479,7 +317,6 @@ static uint32 TonemapperConfBitmaskMobile[39] = {
 	TonemapperContrast +
 	TonemapperColorMatrix +
 	TonemapperVignette +
-	TonemapperVignetteColor +
 	TonemapperGrainIntensity +
 	0,
 
@@ -488,7 +325,6 @@ static uint32 TonemapperConfBitmaskMobile[39] = {
 	TonemapperColorMatrix +
 	TonemapperShadowTint +
 	TonemapperVignette +
-	TonemapperVignetteColor +
 	TonemapperGrainIntensity +
 	0,
 
@@ -510,7 +346,6 @@ static uint32 TonemapperConfBitmaskMobile[39] = {
 	TonemapperContrast +
 	TonemapperBloom +
 	TonemapperVignette +
-	TonemapperVignetteColor +
 	0,
 
 	TonemapperMsaa +
@@ -518,7 +353,6 @@ static uint32 TonemapperConfBitmaskMobile[39] = {
 	TonemapperColorMatrix +
 	TonemapperBloom +
 	TonemapperVignette +
-	TonemapperVignetteColor +
 	0,
 
 	TonemapperMsaa +
@@ -527,7 +361,6 @@ static uint32 TonemapperConfBitmaskMobile[39] = {
 	TonemapperShadowTint +
 	TonemapperBloom +
 	TonemapperVignette +
-	TonemapperVignetteColor +
 	0,
 
 	// Same with grain.
@@ -546,7 +379,6 @@ static uint32 TonemapperConfBitmaskMobile[39] = {
 	TonemapperContrast +
 	TonemapperBloom +
 	TonemapperVignette +
-	TonemapperVignetteColor +
 	TonemapperGrainIntensity +
 	0,
 
@@ -555,7 +387,6 @@ static uint32 TonemapperConfBitmaskMobile[39] = {
 	TonemapperColorMatrix +
 	TonemapperBloom +
 	TonemapperVignette +
-	TonemapperVignetteColor +
 	TonemapperGrainIntensity +
 	0,
 
@@ -565,7 +396,6 @@ static uint32 TonemapperConfBitmaskMobile[39] = {
 	TonemapperShadowTint +
 	TonemapperBloom +
 	TonemapperVignette +
-	TonemapperVignetteColor +
 	TonemapperGrainIntensity +
 	0,
 
@@ -586,11 +416,6 @@ static uint32 TonemapperFindLeastExpensive(uint32* RESTRICT Table, uint32 TableE
 	MustNotHaveBitmask += ((RequiredOptionsBitmask & TonemapperMosaic) == 0) ? TonemapperMosaic : 0;
 	MustNotHaveBitmask += ((RequiredOptionsBitmask & TonemapperMsaa) == 0) ? TonemapperMsaa : 0;
 
-	MustNotHaveBitmask += ((RequiredOptionsBitmask & TonemapperPhoto) != 0) ? TonemapperColorMatrix : 0;
-	MustNotHaveBitmask += ((RequiredOptionsBitmask & TonemapperPhoto) != 0) ? TonemapperShadowTint : 0;
-	MustNotHaveBitmask += ((RequiredOptionsBitmask & TonemapperPhoto) != 0) ? TonemapperColorGrading : 0;
-	MustNotHaveBitmask += ((RequiredOptionsBitmask & TonemapperPhoto) == 0) ? TonemapperPhoto : 0;
-
 	// Search for exact match first.
 	uint32 Index;
 	for(Index = 0; Index < TableEntries; ++Index) 
@@ -675,13 +500,11 @@ static uint32 TonemapperGenerateBitmask(const FViewInfo* RESTRICT View, bool bGa
 
 	FVector Tint(Settings->FilmWhitePoint);
 	FVector TintShadow(Settings->FilmShadowTint);
-	FVector VignetteColor(Settings->VignetteColor);
 
 	Bitmask += (Settings->FilmShadowTintAmount > 0.0f) ? TonemapperShadowTint     : 0;	
 	Bitmask += (Settings->FilmContrast > 0.0f)         ? TonemapperContrast       : 0;
 	Bitmask += (Settings->GrainIntensity > 0.0f)       ? TonemapperGrainIntensity : 0;
 	Bitmask += (Settings->VignetteIntensity > 0.0f)    ? TonemapperVignette       : 0;
-	Bitmask += (VignetteColor.GetAbsMax() != 0.0f)     ? TonemapperVignetteColor  : 0;
 
 	return Bitmask;
 }
@@ -701,7 +524,7 @@ static uint32 TonemapperGenerateBitmaskPost(const FViewInfo* RESTRICT View)
 }
 
 // PC only.
-static uint32 TonemapperGenerateBitmaskPC(const FViewInfo* RESTRICT View, bool bGammaOnly, bool bUseLUT)
+static uint32 TonemapperGenerateBitmaskPC(const FViewInfo* RESTRICT View, bool bGammaOnly)
 {
 	uint32 Bitmask = TonemapperGenerateBitmask(View, bGammaOnly, false);
 
@@ -743,23 +566,6 @@ static uint32 TonemapperGenerateBitmaskPC(const FViewInfo* RESTRICT View, bool b
 		Bitmask |= TonemapperColorFringe;
 	}
 
-	if(bUseLUT)
-	{
-		Bitmask |= TonemapperColorGrading;
-	}
-
-	{
-		static TConsoleVariableData<int32>* CVar = IConsoleManager::Get().FindTConsoleVariableDataInt(TEXT("r.TonemapperPhoto"));
-		int32 Value = CVar->GetValueOnRenderThread();
-
-		if(Value > 0)
-		{
-			// Remove settings not compatible.
-			Bitmask &= ~(TonemapperColorMatrix|TonemapperShadowTint|TonemapperContrast|TonemapperColorGrading|TonemapperMsaa);
-			Bitmask |= TonemapperPhoto;
-		}
-	}
-
 	return Bitmask + TonemapperGenerateBitmaskPost(View);
 }
 
@@ -819,10 +625,6 @@ void FilmPostSetConstants(FVector4* RESTRICT const Constants, const uint32 Confi
 	uint32 UseShadowTint  = TonemapperIsDefined(ConfigBitmask, TonemapperShadowTint);
 	uint32 UseContrast    = TonemapperIsDefined(ConfigBitmask, TonemapperContrast);
 
-	// DON'T CHECK INTO MAIN: NEEDS A PERMUTATION TO AVOID NON-BACKWARDS COMPATIBLE CHANGES
-	// This corrects so the linear segment is linear in log space!
-	float In18Grey = FMath::Log2(1.0f+0.18f);
-
 	// Must insure inputs are in correct range (else possible generation of NaNs).
 	float InExposure = 1.0f;
 	FVector InWhitePoint(FinalPostProcessSettings->FilmWhitePoint);
@@ -833,9 +635,9 @@ void FilmPostSetConstants(FVector4* RESTRICT const Constants, const uint32 Confi
 	FVector InMatrixB(FinalPostProcessSettings->FilmChannelMixerBlue);
 	float InContrast = FMath::Clamp(FinalPostProcessSettings->FilmContrast, 0.0f, 1.0f) + 1.0f;
 	float InDynamicRange = powf(2.0f, FMath::Clamp(FinalPostProcessSettings->FilmDynamicRange, 1.0f, 4.0f));
-	float InToe = (1.0f - FMath::Clamp(FinalPostProcessSettings->FilmToeAmount, 0.0f, 1.0f)) * In18Grey;
-	InToe = FMath::Clamp(InToe, In18Grey/8.0f, In18Grey * (15.0f/16.0f));
-	float InHeal = 1.0f - (FMath::Max(1.0f/32.0f, 1.0f - FMath::Clamp(FinalPostProcessSettings->FilmHealAmount, 0.0f, 1.0f)) * (1.0f - In18Grey)); 
+	float InToe = (1.0f - FMath::Clamp(FinalPostProcessSettings->FilmToeAmount, 0.0f, 1.0f)) * 0.18f;
+	InToe = FMath::Clamp(InToe, 0.18f/8.0f, 0.18f * (15.0f/16.0f));
+	float InHeal = 1.0f - (FMath::Max(1.0f/32.0f, 1.0f - FMath::Clamp(FinalPostProcessSettings->FilmHealAmount, 0.0f, 1.0f)) * (1.0f - 0.18f)); 
 	FVector InShadowTint(FinalPostProcessSettings->FilmShadowTint);
 	float InShadowTintBlend = FMath::Clamp(FinalPostProcessSettings->FilmShadowTintBlend, 0.0f, 1.0f) * 64.0f;
 
@@ -915,7 +717,7 @@ void FilmPostSetConstants(FVector4* RESTRICT const Constants, const uint32 Confi
 	float OutColorCurveCm2;
 
 	// Line for linear section.
-	float FilmLineOffset = In18Grey - In18Grey*InContrast;
+	float FilmLineOffset = 0.18f - 0.18f*InContrast;
 	float FilmXAtY0 = -FilmLineOffset/InContrast;
 	float FilmXAtY1 = (1.0f - FilmLineOffset) / InContrast;
 	float FilmXS = FilmXAtY1 - FilmXAtY0;
@@ -987,102 +789,6 @@ void FilmPostSetConstants(FVector4* RESTRICT const Constants, const uint32 Confi
 	Constants[7] = FVector4(OutColorShadow_Tint2, 0.0f);
 }
 
-// Improved film post: setting of constants.
-void FilmPostSetConstantsPhoto(FVector4* RESTRICT const Constants, const FPostProcessSettings* RESTRICT const FinalPostProcessSettings)
-{
-	float PhotoExp = FMath::Pow(2.0f,FMath::Clamp(FinalPostProcessSettings->PhotoExposure, -64.0f, 64.0f));
-	float PhotoMid = 0.18f * FMath::Pow(2.0f,FMath::Clamp(FinalPostProcessSettings->PhotoMid, -64.0f, 64.0f));
-	FVector PhotoWhite(FinalPostProcessSettings->PhotoWhite);
-	float PhotoSat = FMath::Clamp(FinalPostProcessSettings->PhotoSaturation, 0.0f, 2.0f) - 1.0f;
-	FVector PhotoMixR(FinalPostProcessSettings->PhotoChannelMixerRed);
-	FVector PhotoMixG(FinalPostProcessSettings->PhotoChannelMixerGreen);
-	FVector PhotoMixB(FinalPostProcessSettings->PhotoChannelMixerBlue);
-	FVector PhotoOverTint(FinalPostProcessSettings->PhotoOverTint);
-	float PhotoGamma = FMath::Clamp(FinalPostProcessSettings->PhotoContrast, 0.0f, 4.0f);
-	float PhotoBlackIn = (PhotoMid*PhotoMid*PhotoMid);
-	float PhotoBlackOut = (PhotoMid*PhotoMid*PhotoMid)*FMath::Pow(2.0f,FMath::Clamp(FinalPostProcessSettings->PhotoBlackOut, -64.0f, 64.0f));
-	FVector PhotoTint(FinalPostProcessSettings->PhotoTint);
-	float PhotoOver = FMath::Clamp(FinalPostProcessSettings->PhotoOver, 0.0f, 1.0f);
-
-	// Conversion from linear rgb to luma (using HDTV coef).
-	FVector LumaWeights = FVector(0.2126f, 0.7152f, 0.0722f);
-
-	// Make sure not zero.
-	PhotoWhite += FVector(1.0f / (256.0f*256.0f*32.0f));
-	PhotoMixR += FVector(1.0f / (256.0f*256.0f*32.0f));
-	PhotoMixG += FVector(1.0f / (256.0f*256.0f*32.0f));
-	PhotoMixB += FVector(1.0f / (256.0f*256.0f*32.0f));
-
-	// Normalize inputs for luma.
-	FVector PW = PhotoWhite/FVector::DotProduct(PhotoWhite, LumaWeights);
-	FVector PR = PhotoMixR*(FVector::DotProduct(FVector(1.0f,0.0f,0.0f), LumaWeights)/FVector::DotProduct(PhotoMixR, LumaWeights));    
-	FVector PG = PhotoMixG*(FVector::DotProduct(FVector(0.0f,1.0f,0.0f), LumaWeights)/FVector::DotProduct(PhotoMixG, LumaWeights));    
-	FVector PB = PhotoMixB*(FVector::DotProduct(FVector(0.0f,0.0f,1.0f), LumaWeights)/FVector::DotProduct(PhotoMixB, LumaWeights));    
-
-	// Fold exposure into white point.
-	PW *= PhotoExp;
-
-	// Channel mixer matrix 
-	FVector MR = FVector(PR.X,PG.X,PB.X);
-	FVector MG = FVector(PR.Y,PG.Y,PB.Y);
-	FVector MB = FVector(PR.Z,PG.Z,PB.Z);
-
-	// Fold white point exposure into channel mixer matrix.
-	MR *= PW;
-	MG *= PW;
-	MB *= PW;
-
-	// Saturation matrix.
-	FVector SR = FVector(1.0+(1.0-LumaWeights.X)*PhotoSat,(-LumaWeights.Y)*PhotoSat,(-LumaWeights.Z)*PhotoSat);
-	FVector SG = FVector((-LumaWeights.X)*PhotoSat,1.0+(1.0-LumaWeights.Y)*PhotoSat,(-LumaWeights.Z)*PhotoSat);
-	FVector SB = FVector((-LumaWeights.X)*PhotoSat,(-LumaWeights.Y)*PhotoSat,1.0+(1.0-LumaWeights.Z)*PhotoSat);    
-
-	// Final matrix, fold white point exposure mixer into saturation.
-	FVector OutMatrixR = FVector(
-		FVector::DotProduct(FVector(MR.X,MG.X,MB.X),SR),
-		FVector::DotProduct(FVector(MR.Y,MG.Y,MB.Y),SR),
-		FVector::DotProduct(FVector(MR.Z,MG.Z,MB.Z),SR));
-	FVector OutMatrixG = FVector(
-		FVector::DotProduct(FVector(MR.X,MG.X,MB.X),SG),
-		FVector::DotProduct(FVector(MR.Y,MG.Y,MB.Y),SG),
-		FVector::DotProduct(FVector(MR.Z,MG.Z,MB.Z),SG));
-	FVector OutMatrixB = FVector(
-		FVector::DotProduct(FVector(MR.X,MG.X,MB.X),SB),
-		FVector::DotProduct(FVector(MR.Y,MG.Y,MB.Y),SB),
-		FVector::DotProduct(FVector(MR.Z,MG.Z,MB.Z),SB));    
-
-	float OutLift = (PhotoMid*(FMath::Pow(PhotoBlackOut*FMath::Pow(PhotoMid,PhotoGamma-1.0f),1.0f/PhotoGamma)-PhotoBlackIn))/(PhotoMid-PhotoBlackIn);
-	float OutLiftScale = (PhotoMid-OutLift)/PhotoMid;
-
-	FVector OutTint = PhotoTint;
-	OutTint += FVector(1.0f / (256.0f*256.0f*32.0f));
-	OutTint *= 1.0f / FVector::DotProduct(OutTint, LumaWeights);
-	// Factor in the scale to ensure gamma does not change the 18% midtone.
-	OutTint *= PhotoMid/FMath::Pow(PhotoMid, PhotoGamma);
-
-	// HDR display support provides a secondary exposure adjustment 
-	// which is factored into the tint mulitply in the {Lift, gamma, tint} processing.
-	// This is prior to over-exposure crosstalk (desaturate overexposed color),
-	// and prior to highlight compression.
-	// If an LDR display goes from 0-1 this multiply will be 1.
-	// If an HDR display goes from 0-16 this multiply will be 1/16 (the display is 16x brigher).
-	{
-		static TConsoleVariableData<float>* CVar = IConsoleManager::Get().FindTConsoleVariableDataFloat(TEXT("r.TonemapperHDR"));
-		float Value = CVar->GetValueOnRenderThread();
-		if(Value < 1.0f)
-		{
-			Value = 1.0f;
-		}
-		OutTint *= 1.0f/Value;
-	}
-
-	Constants[0] = FVector4(OutMatrixR, OutLift);
-	Constants[1] = FVector4(OutMatrixG, OutLiftScale);
-	Constants[2] = FVector4(OutMatrixB, PhotoGamma); 
-	Constants[3] = FVector4(OutTint, PhotoOver*(1.0f/3.0f)); 
-	Constants[4] = FVector4(PhotoOverTint, 1.0f-PhotoMid);
-}
-
 
 BEGIN_UNIFORM_BUFFER_STRUCT(FBloomDirtMaskParameters,)
 	DECLARE_UNIFORM_BUFFER_STRUCT_MEMBER(FVector4,Tint)
@@ -1126,12 +832,9 @@ class FPostProcessTonemapPS : public FGlobalShader
 		OutEnvironment.SetDefine(TEXT("USE_GRAIN_INTENSITY"),    TonemapperIsDefined(ConfigBitmask, TonemapperGrainIntensity));
 		OutEnvironment.SetDefine(TEXT("USE_GRAIN_QUANTIZATION"), TonemapperIsDefined(ConfigBitmask, TonemapperGrainQuantization));
 		OutEnvironment.SetDefine(TEXT("USE_VIGNETTE"),           TonemapperIsDefined(ConfigBitmask, TonemapperVignette));
-		OutEnvironment.SetDefine(TEXT("USE_VIGNETTE_COLOR"),     TonemapperIsDefined(ConfigBitmask, TonemapperVignetteColor));
 		OutEnvironment.SetDefine(TEXT("USE_COLOR_FRINGE"),		 TonemapperIsDefined(ConfigBitmask, TonemapperColorFringe));
 		// @todo Mac OS X: in order to share precompiled shaders between GL 3.3 & GL 4.1 devices we mustn't use volume-texture rendering as it isn't universally supported.
 		OutEnvironment.SetDefine(TEXT("USE_VOLUME_LUT"),		 (IsFeatureLevelSupported(Platform,ERHIFeatureLevel::SM4) && GSupportsVolumeTextureRendering && Platform != EShaderPlatform::SP_OPENGL_SM4_MAC));
-		OutEnvironment.SetDefine(TEXT("USE_COLOR_GRADING"),		 TonemapperIsDefined(ConfigBitmask, TonemapperColorGrading));
-		OutEnvironment.SetDefine(TEXT("USE_PHOTO"),              TonemapperIsDefined(ConfigBitmask, TonemapperPhoto));
 		OutEnvironment.SetDefine(TEXT("USE_709"),                TonemapperIsDefined(ConfigBitmask, Tonemapper709));
 		OutEnvironment.SetDefine(TEXT("USE_GAMMA"),              TonemapperIsDefined(ConfigBitmask, TonemapperGamma));
 
@@ -1152,7 +855,7 @@ public:
 	FShaderResourceParameter NoiseTexture;
 	FShaderResourceParameter NoiseTextureSampler;
 	FShaderParameter TexScale;
-	FShaderParameter VignetteColorIntensity;
+	FShaderParameter VignetteIntensity;
 	FShaderParameter GrainScaleBiasJitter;
 	FShaderResourceParameter ColorGradingLUT;
 	FShaderResourceParameter ColorGradingLUTSampler;
@@ -1167,12 +870,6 @@ public:
 	FShaderParameter ColorShadow_Tint1;
 	FShaderParameter ColorShadow_Tint2;
 
-	FShaderParameter PhotoMatrixR_Lift;
-	FShaderParameter PhotoMatrixG_LiftScale;
-	FShaderParameter PhotoMatrixB_Gamma;
-	FShaderParameter PhotoTint_Over;
-	FShaderParameter PhotoOverTint_OneMinusMid;
-
 	//@HACK
 	FShaderParameter OverlayColor;
 
@@ -1186,7 +883,7 @@ public:
 		NoiseTexture.Bind(Initializer.ParameterMap,TEXT("NoiseTexture"));
 		NoiseTextureSampler.Bind(Initializer.ParameterMap,TEXT("NoiseTextureSampler"));
 		TexScale.Bind(Initializer.ParameterMap, TEXT("TexScale"));
-		VignetteColorIntensity.Bind(Initializer.ParameterMap, TEXT("VignetteColorIntensity"));
+		VignetteIntensity.Bind(Initializer.ParameterMap, TEXT("VignetteIntensity"));
 		GrainScaleBiasJitter.Bind(Initializer.ParameterMap, TEXT("GrainScaleBiasJitter"));
 		ColorGradingLUT.Bind(Initializer.ParameterMap, TEXT("ColorGradingLUT"));
 		ColorGradingLUTSampler.Bind(Initializer.ParameterMap, TEXT("ColorGradingLUTSampler"));
@@ -1200,12 +897,6 @@ public:
 		ColorShadow_Luma.Bind(Initializer.ParameterMap, TEXT("ColorShadow_Luma"));
 		ColorShadow_Tint1.Bind(Initializer.ParameterMap, TEXT("ColorShadow_Tint1"));
 		ColorShadow_Tint2.Bind(Initializer.ParameterMap, TEXT("ColorShadow_Tint2"));
-
-		PhotoMatrixR_Lift.Bind(Initializer.ParameterMap, TEXT("PhotoMatrixR_Lift"));
-		PhotoMatrixG_LiftScale.Bind(Initializer.ParameterMap, TEXT("PhotoMatrixG_LiftScale"));
-		PhotoMatrixB_Gamma.Bind(Initializer.ParameterMap, TEXT("PhotoMatrixB_Gamma"));
-		PhotoTint_Over.Bind(Initializer.ParameterMap, TEXT("PhotoTint_Over"));
-		PhotoOverTint_OneMinusMid.Bind(Initializer.ParameterMap, TEXT("PhotoOverTint_OneMinusMid"));
 		
 		OverlayColor.Bind(Initializer.ParameterMap, TEXT("OverlayColor"));
 	}
@@ -1215,10 +906,9 @@ public:
 	{
 		bool bShaderHasOutdatedParameters = FGlobalShader::Serialize(Ar);
 		Ar  << PostprocessParameter << ColorScale0 << ColorScale1 << InverseGamma << NoiseTexture << NoiseTextureSampler
-			<< TexScale << VignetteColorIntensity << GrainScaleBiasJitter
+			<< TexScale << VignetteIntensity << GrainScaleBiasJitter
 			<< ColorGradingLUT << ColorGradingLUTSampler
 			<< ColorMatrixR_ColorCurveCd1 << ColorMatrixG_ColorCurveCd3Cm3 << ColorMatrixB_ColorCurveCm2 << ColorCurve_Cm0Cd0_Cd2_Ch0Cm1_Ch3 << ColorCurve_Ch1_Ch2 << ColorShadow_Luma << ColorShadow_Tint1 << ColorShadow_Tint2
-			<< PhotoMatrixR_Lift << PhotoMatrixG_LiftScale << PhotoMatrixB_Gamma << PhotoTint_Over << PhotoOverTint_OneMinusMid
 			<< OverlayColor;
 		
 		return bShaderHasOutdatedParameters;
@@ -1264,12 +954,7 @@ public:
 			SetShaderValue(Context.RHICmdList, ShaderRHI, TexScale, TexScaleValue);
 		}
 		
-		FVector4 VignetteColorIntensityValue;
-		VignetteColorIntensityValue.X = Settings.VignetteColor.R;
-		VignetteColorIntensityValue.Y = Settings.VignetteColor.G;
-		VignetteColorIntensityValue.Z = Settings.VignetteColor.B;
-		VignetteColorIntensityValue.W = Settings.VignetteIntensity;
-		SetShaderValue(Context.RHICmdList, ShaderRHI, VignetteColorIntensity, VignetteColorIntensityValue);
+		SetShaderValue(Context.RHICmdList, ShaderRHI, VignetteIntensity, Settings.VignetteIntensity);
 
 		FVector GrainValue;
 		GrainPostSettings(&GrainValue, &Settings);
@@ -1347,16 +1032,6 @@ public:
 			SetShaderValue(Context.RHICmdList, ShaderRHI, ColorShadow_Tint1, Constants[6]);
 			SetShaderValue(Context.RHICmdList, ShaderRHI, ColorShadow_Tint2, Constants[7]);
 		}
-
-		{
-			FVector4 Constants[5];
-			FilmPostSetConstantsPhoto(Constants, &Context.View.FinalPostProcessSettings);
-			SetShaderValue(Context.RHICmdList, ShaderRHI, PhotoMatrixR_Lift, Constants[0]);
-			SetShaderValue(Context.RHICmdList, ShaderRHI, PhotoMatrixG_LiftScale, Constants[1]);
-			SetShaderValue(Context.RHICmdList, ShaderRHI, PhotoMatrixB_Gamma, Constants[2]); 
-			SetShaderValue(Context.RHICmdList, ShaderRHI, PhotoTint_Over, Constants[3]); 
-			SetShaderValue(Context.RHICmdList, ShaderRHI, PhotoOverTint_OneMinusMid, Constants[4]);
-		}
 	}
 	
 	static const TCHAR* GetSourceFilename()
@@ -1375,11 +1050,7 @@ public:
 	IMPLEMENT_SHADER_TYPE2(FPostProcessTonemapPS##A, SF_Pixel);
 
 	VARIATION1(0)  VARIATION1(1)  VARIATION1(2)  VARIATION1(3)  VARIATION1(4)  VARIATION1(5)
-	VARIATION1(6)  VARIATION1(7)  VARIATION1(8)  VARIATION1(9)  VARIATION1(10)
-
-	VARIATION1(11)  VARIATION1(12)  VARIATION1(13)  VARIATION1(14)  VARIATION1(15)  VARIATION1(16)
-	VARIATION1(17)  VARIATION1(18)  VARIATION1(19)  VARIATION1(20)  VARIATION1(21)
-
+	VARIATION1(6)  VARIATION1(7)  VARIATION1(8)
 
 #undef VARIATION1
 
@@ -1390,17 +1061,10 @@ IMPLEMENT_SHADER_TYPE(,FPostProcessTonemapVS,TEXT("PostProcessTonemap"),TEXT("Ma
 FRCPassPostProcessTonemap::FRCPassPostProcessTonemap(const FViewInfo& View, bool bInDoGammaOnly)
 	: bDoGammaOnly(bInDoGammaOnly)
 {
-	uint32 ConfigBitmask = TonemapperGenerateBitmaskPC(&View, bDoGammaOnly, FRCPassPostProcessCombineLUTs::IsColorGradingLUTNeeded(&View));
+	uint32 ConfigBitmask = TonemapperGenerateBitmaskPC(&View, bDoGammaOnly);
 	ConfigIndexPC = TonemapperFindLeastExpensive(TonemapperConfBitmaskPC, sizeof(TonemapperConfBitmaskPC)/4, TonemapperCostTab, ConfigBitmask);
 }
 
-bool FRCPassPostProcessTonemap::IsLUTNeeded() const
-{
-	uint32 ConfigBitmask = TonemapperConfBitmaskPC[ConfigIndexPC];
-
-	return TonemapperIsDefined(ConfigBitmask, TonemapperColorGrading) != 0;
-}
-
 template <uint32 ConfigIndex>
 static void SetShaderTempl(const FRenderingCompositePassContext& Context)
 {
@@ -1465,19 +1129,6 @@ void FRCPassPostProcessTonemap::Process(FRenderingCompositePassContext& Context)
 		case 6: SetShaderTempl<6>(Context); break;
 		case 7: SetShaderTempl<7>(Context); break;
 		case 8: SetShaderTempl<8>(Context); break;
-		case 9: SetShaderTempl<9>(Context); break;
-		case 10: SetShaderTempl<10>(Context); break;
-		case 11: SetShaderTempl<11>(Context); break;
-		case 12: SetShaderTempl<12>(Context); break;
-		case 13: SetShaderTempl<13>(Context); break;
-		case 14: SetShaderTempl<14>(Context); break;
-		case 15: SetShaderTempl<15>(Context); break;
-		case 16: SetShaderTempl<16>(Context); break;
-		case 17: SetShaderTempl<17>(Context); break;
-		case 18: SetShaderTempl<18>(Context); break;
-		case 19: SetShaderTempl<19>(Context); break;
-		case 20: SetShaderTempl<20>(Context); break;
-		case 21: SetShaderTempl<21>(Context); break;
 		default:
 			check(0);
 	}
@@ -1556,7 +1207,6 @@ class FPostProcessTonemapPS_ES2 : public FGlobalShader
 		OutEnvironment.SetDefine(TEXT("USE_GRAIN_INTENSITY"),    TonemapperIsDefined(ConfigBitmask, TonemapperGrainIntensity));
 		OutEnvironment.SetDefine(TEXT("USE_GRAIN_QUANTIZATION"), TonemapperIsDefined(ConfigBitmask, TonemapperGrainQuantization));
 		OutEnvironment.SetDefine(TEXT("USE_VIGNETTE"),           TonemapperIsDefined(ConfigBitmask, TonemapperVignette));
-		OutEnvironment.SetDefine(TEXT("USE_VIGNETTE_COLOR"),     TonemapperIsDefined(ConfigBitmask, TonemapperVignetteColor));
 		OutEnvironment.SetDefine(TEXT("USE_LIGHT_SHAFTS"),       TonemapperIsDefined(ConfigBitmask, TonemapperLightShafts));
 		OutEnvironment.SetDefine(TEXT("USE_DOF"),                TonemapperIsDefined(ConfigBitmask, TonemapperDOF));
 		OutEnvironment.SetDefine(TEXT("USE_MSAA"),               TonemapperIsDefined(ConfigBitmask, TonemapperMsaa));
@@ -1575,7 +1225,7 @@ public:
 	FShaderParameter TexScale;
 	FShaderParameter GrainScaleBiasJitter;
 	FShaderParameter InverseGamma;
-	FShaderParameter VignetteColorIntensity;
+	FShaderParameter VignetteIntensity;
 
 	FShaderParameter ColorMatrixR_ColorCurveCd1;
 	FShaderParameter ColorMatrixG_ColorCurveCd3Cm3;
@@ -1596,7 +1246,7 @@ public:
 		ColorScale0.Bind(Initializer.ParameterMap, TEXT("ColorScale0"));
 		ColorScale1.Bind(Initializer.ParameterMap, TEXT("ColorScale1"));
 		TexScale.Bind(Initializer.ParameterMap, TEXT("TexScale"));
-		VignetteColorIntensity.Bind(Initializer.ParameterMap, TEXT("VignetteColorIntensity"));
+		VignetteIntensity.Bind(Initializer.ParameterMap, TEXT("VignetteIntensity"));
 		GrainScaleBiasJitter.Bind(Initializer.ParameterMap, TEXT("GrainScaleBiasJitter"));
 		InverseGamma.Bind(Initializer.ParameterMap,TEXT("InverseGamma"));
 
@@ -1617,7 +1267,7 @@ public:
 	{
 		bool bShaderHasOutdatedParameters = FGlobalShader::Serialize(Ar);
 		Ar  << PostprocessParameter << ColorScale0 << ColorScale1 << InverseGamma
-			<< TexScale << GrainScaleBiasJitter << VignetteColorIntensity
+			<< TexScale << GrainScaleBiasJitter << VignetteIntensity
 			<< ColorMatrixR_ColorCurveCd1 << ColorMatrixG_ColorCurveCd3Cm3 << ColorMatrixB_ColorCurveCm2 << ColorCurve_Cm0Cd0_Cd2_Ch0Cm1_Ch3 << ColorCurve_Ch1_Ch2 << ColorShadow_Luma << ColorShadow_Tint1 << ColorShadow_Tint2
 			<< OverlayColor;
 
@@ -1667,12 +1317,7 @@ public:
 			SetShaderValue(Context.RHICmdList, ShaderRHI, TexScale, TexScaleValue);
 		}
 
-		FVector4 VignetteColorIntensityValue;
-		VignetteColorIntensityValue.X = Settings.VignetteColor.R;
-		VignetteColorIntensityValue.Y = Settings.VignetteColor.G;
-		VignetteColorIntensityValue.Z = Settings.VignetteColor.B;
-		VignetteColorIntensityValue.W = Settings.VignetteIntensity;
-		SetShaderValue(Context.RHICmdList, ShaderRHI, VignetteColorIntensity, VignetteColorIntensityValue);
+		SetShaderValue(Context.RHICmdList, ShaderRHI, VignetteIntensity, Settings.VignetteIntensity);
 
 		FVector GrainValue;
 		GrainPostSettings(&GrainValue, &Settings);
diff --git a/Engine/Source/Runtime/Renderer/Private/PostProcess/PostProcessTonemap.h b/Engine/Source/Runtime/Renderer/Private/PostProcess/PostProcessTonemap.h
--- a/Engine/Source/Runtime/Renderer/Private/PostProcess/PostProcessTonemap.h
+++ b/Engine/Source/Runtime/Renderer/Private/PostProcess/PostProcessTonemap.h
@@ -43,8 +43,6 @@ public:
 	// constructor
 	FRCPassPostProcessTonemap(const FViewInfo& View, bool bInDoGammaOnly = false);
 
-	bool IsLUTNeeded() const;
-
 	// interface FRenderingCompositePass ---------
 
 	virtual void Process(FRenderingCompositePassContext& Context) override;
diff --git a/Engine/Source/Runtime/Renderer/Private/PostProcess/PostProcessing.cpp b/Engine/Source/Runtime/Renderer/Private/PostProcess/PostProcessing.cpp
--- a/Engine/Source/Runtime/Renderer/Private/PostProcess/PostProcessing.cpp
+++ b/Engine/Source/Runtime/Renderer/Private/PostProcess/PostProcessing.cpp
@@ -194,18 +194,14 @@ static void AddTonemapper(
 	const FRenderingCompositeOutputRef& BloomOutputCombined,
 	const FRenderingCompositeOutputRef& EyeAdaptation)
 {
+	FRenderingCompositePass* CombinedLUT = Context.Graph.RegisterPass(new(FMemStack::Get()) FRCPassPostProcessCombineLUTs(Context.View.GetShaderPlatform()));
+	
 	FRCPassPostProcessTonemap* PostProcessTonemap = Context.Graph.RegisterPass(new(FMemStack::Get()) FRCPassPostProcessTonemap(Context.View));
 
 	PostProcessTonemap->SetInput(ePId_Input0, Context.FinalOutput);
 	PostProcessTonemap->SetInput(ePId_Input1, BloomOutputCombined);
 	PostProcessTonemap->SetInput(ePId_Input2, EyeAdaptation);
-
-	if(PostProcessTonemap->IsLUTNeeded())
-	{
-		FRenderingCompositePass* CombinedLUT = Context.Graph.RegisterPass(new(FMemStack::Get()) FRCPassPostProcessCombineLUTs(Context.View.GetShaderPlatform()));
-
-		PostProcessTonemap->SetInput(ePId_Input3, CombinedLUT);
-	}
+	PostProcessTonemap->SetInput(ePId_Input3, CombinedLUT);
 
 	Context.FinalOutput = FRenderingCompositeOutputRef(PostProcessTonemap);
 }
