Add missing parameter in Serialize

diff --git a/Engine/Source/Runtime/Renderer/Private/PostProcess/PostProcessCombineLUTs.cpp b/Engine/Source/Runtime/Renderer/Private/PostProcess/PostProcessCombineLUTs.cpp
--- a/Engine/Source/Runtime/Renderer/Private/PostProcess/PostProcessCombineLUTs.cpp
+++ b/Engine/Source/Runtime/Renderer/Private/PostProcess/PostProcessCombineLUTs.cpp
@@ -371,6 +371,7 @@ public:
 		}
 		Ar << WeightsParameter << ColorScale << OverlayColor;
 		Ar << ColorRemapShaderParameters;
+		Ar << InverseGamma;
 
 		Ar << WhiteTemp;
 		Ar << WhiteTint;
