Fixed shadowed variable. Changed new tonemapper enable to cvar.

diff --git a/Engine/Source/Runtime/Engine/Classes/Engine/Scene.h b/Engine/Source/Runtime/Engine/Classes/Engine/Scene.h
--- a/Engine/Source/Runtime/Engine/Classes/Engine/Scene.h
+++ b/Engine/Source/Runtime/Engine/Classes/Engine/Scene.h
@@ -90,9 +90,6 @@ struct FPostProcessSettings
 	UPROPERTY(BlueprintReadWrite, Category=Overrides, meta=(PinHiddenByDefault))
 	uint32 bOverride_FilmShadowTintAmount:1;
 
-	UPROPERTY(BlueprintReadWrite, Category=Overrides, meta=(PinHiddenByDefault))
-	uint32 bOverride_FilmEnable:1;
-
 	UPROPERTY(BlueprintReadWrite, Category=Overrides, meta=(PinHiddenByDefault))
 	uint32 bOverride_FilmSlope:1;
 	UPROPERTY(BlueprintReadWrite, Category=Overrides, meta=(PinHiddenByDefault))
@@ -428,9 +425,6 @@ struct FPostProcessSettings
 	UPROPERTY(interp, BlueprintReadWrite, Category=Film, AdvancedDisplay, meta=(UIMin = "1.0", UIMax = "4.0", editcondition = "bOverride_FilmDynamicRange", DisplayName = "Dynamic Range"))
 	float FilmDynamicRange;
 
-	UPROPERTY(EditAnywhere, BlueprintReadWrite, Category=Film, meta=(editcondition = "bOverride_FilmEnable", DisplayName = "Use film tone mapper"))
-	uint32 FilmEnable:1;
-
 	UPROPERTY(interp, BlueprintReadWrite, Category=Film, AdvancedDisplay, meta=(UIMin = "0.0", UIMax = "1.0", editcondition = "bOverride_FilmSlope", DisplayName = "Slope"))
 	float FilmSlope;
 	UPROPERTY(interp, BlueprintReadWrite, Category=Film, AdvancedDisplay, meta=(UIMin = "0.0", UIMax = "1.0", editcondition = "bOverride_FilmToe", DisplayName = "Toe"))
@@ -906,8 +900,6 @@ struct FPostProcessSettings
 		FilmShadowTintBlend = 0.5;
 		FilmShadowTintAmount = 0.0;
 
-		FilmEnable = 0;
-
 		// ACES settings
 		FilmSlope = 0.88f;
 		FilmToe = 0.55f;
diff --git a/Engine/Source/Runtime/Engine/Private/SceneView.cpp b/Engine/Source/Runtime/Engine/Private/SceneView.cpp
--- a/Engine/Source/Runtime/Engine/Private/SceneView.cpp
+++ b/Engine/Source/Runtime/Engine/Private/SceneView.cpp
@@ -915,11 +915,6 @@ void FSceneView::OverridePostProcessSettings(const FPostProcessSettings& Src, fl
 		}
 	}
 
-	if(Src.bOverride_FilmEnable)
-	{
-		Dest.FilmEnable = Src.FilmEnable;
-	}
-
 	if(Src.bOverride_DepthOfFieldMethod)
 	{
 		Dest.DepthOfFieldMethod = Src.DepthOfFieldMethod;
diff --git a/Engine/Source/Runtime/Renderer/Private/PostProcess/PostProcessCombineLUTs.cpp b/Engine/Source/Runtime/Renderer/Private/PostProcess/PostProcessCombineLUTs.cpp
--- a/Engine/Source/Runtime/Renderer/Private/PostProcess/PostProcessCombineLUTs.cpp
+++ b/Engine/Source/Runtime/Renderer/Private/PostProcess/PostProcessCombineLUTs.cpp
@@ -43,6 +43,13 @@ static FAutoConsoleVariableRef CVarLUTSize(
 	ECVF_RenderThreadSafe
 	);
 
+static TAutoConsoleVariable<int32> CVarTonemapperFilm(
+	TEXT("r.TonemapperFilm"),
+	0,
+	TEXT("Use new film tone mapper"),
+	ECVF_RenderThreadSafe
+	);
+
 // false:use 256x16 texture / true:use volume texture (faster, requires geometry shader)
 // USE_VOLUME_LUT: needs to be the same for C++ and HLSL
 static bool UseVolumeTextureLUT(EShaderPlatform Platform) 
@@ -299,28 +306,28 @@ public:
 			// Intermediates.
 			float FilmMidXS = FilmHiX - FilmLoX;
 			float FilmMidYS = FilmHiY - FilmLoY;
-			float FilmSlope = FilmMidYS / (FilmMidXS);
+			float FilmSlopeS = FilmMidYS / (FilmMidXS);
 			float FilmHiYS = 1.0f - FilmHiY;
 			float FilmLoYS = FilmLoY;
 			float FilmToe = FilmLoX;
-			float FilmHiG = (-FilmHiYS + (FilmSlope*FilmHeal)) / (FilmSlope*FilmHeal);
-			float FilmLoG = (-FilmLoYS + (FilmSlope*FilmToe)) / (FilmSlope*FilmToe);
+			float FilmHiG = (-FilmHiYS + (FilmSlopeS*FilmHeal)) / (FilmSlopeS*FilmHeal);
+			float FilmLoG = (-FilmLoYS + (FilmSlopeS*FilmToe)) / (FilmSlopeS*FilmToe);
 
 			// Constants.
 			float OutColorCurveCh1 = FilmHiYS/FilmHiG;
 			float OutColorCurveCh2 = -FilmHiX*(FilmHiYS/FilmHiG);
-			float OutColorCurveCh3 = FilmHiYS/(FilmSlope*FilmHiG) - FilmHiX;
+			float OutColorCurveCh3 = FilmHiYS/(FilmSlopeS*FilmHiG) - FilmHiX;
 			float OutColorCurveCh0Cm1 = FilmHiX;
-			float OutColorCurveCm2 = FilmSlope;
+			float OutColorCurveCm2 = FilmSlopeS;
 			float OutColorCurveCm0Cd0 = FilmLoX;
-			float OutColorCurveCd3Cm3 = FilmLoY - FilmLoX*FilmSlope;
+			float OutColorCurveCd3Cm3 = FilmLoY - FilmLoX*FilmSlopeS;
 			float OutColorCurveCd1 = 0.0f;
 			float OutColorCurveCd2 = 1.0f;
 			// Handle these separate in case of FilmLoG being 0.
 			if(FilmLoG != 0.0f)
 			{
 				OutColorCurveCd1 = -FilmLoYS/FilmLoG;
-				OutColorCurveCd2 = FilmLoYS/(FilmSlope*FilmLoG);
+				OutColorCurveCd2 = FilmLoYS/(FilmSlopeS*FilmLoG);
 			}
 			else
 			{
@@ -337,7 +344,7 @@ public:
 			Constants[4] = FVector4(OutColorCurveCh1, OutColorCurveCh2, 0.0f, 0.0f);
 			Constants[5] = FVector4(OutColorShadow_Luma, 0.0f);
 			Constants[6] = FVector4(OutColorShadow_Tint1, 0.0f);
-			Constants[7] = FVector4(OutColorShadow_Tint2, Settings.FilmEnable);
+			Constants[7] = FVector4(OutColorShadow_Tint2, CVarTonemapperFilm.GetValueOnRenderThread());
 
 			SetShaderValue(RHICmdList, ShaderRHI, ColorMatrixR_ColorCurveCd1, Constants[0]);
 			SetShaderValue(RHICmdList, ShaderRHI, ColorMatrixG_ColorCurveCd3Cm3, Constants[1]);
