Fix for Linux build and shadow variable FilmToe. Renamed to FilmToeS

diff --git a/Engine/Source/Runtime/Renderer/Private/PostProcess/PostProcessCombineLUTs.cpp b/Engine/Source/Runtime/Renderer/Private/PostProcess/PostProcessCombineLUTs.cpp
--- a/Engine/Source/Runtime/Renderer/Private/PostProcess/PostProcessCombineLUTs.cpp
+++ b/Engine/Source/Runtime/Renderer/Private/PostProcess/PostProcessCombineLUTs.cpp
@@ -309,9 +309,9 @@ public:
 			float FilmSlopeS = FilmMidYS / (FilmMidXS);
 			float FilmHiYS = 1.0f - FilmHiY;
 			float FilmLoYS = FilmLoY;
-			float FilmToe = FilmLoX;
+			float FilmToeS = FilmLoX;
 			float FilmHiG = (-FilmHiYS + (FilmSlopeS*FilmHeal)) / (FilmSlopeS*FilmHeal);
-			float FilmLoG = (-FilmLoYS + (FilmSlopeS*FilmToe)) / (FilmSlopeS*FilmToe);
+			float FilmLoG = (-FilmLoYS + (FilmSlopeS*FilmToeS)) / (FilmSlopeS*FilmToeS);
 
 			// Constants.
 			float OutColorCurveCh1 = FilmHiYS/FilmHiG;
