Move output device to LUT. Fixed legacy color grading LUTs.

diff --git a/Engine/Shaders/GammaCorrectionCommon.usf b/Engine/Shaders/GammaCorrectionCommon.usf
--- a/Engine/Shaders/GammaCorrectionCommon.usf
+++ b/Engine/Shaders/GammaCorrectionCommon.usf
@@ -28,3 +28,9 @@ half3 LinearToSrgbBranching(half3 lin)
 		LinearToSrgbBranchingChannel(lin.g),
 		LinearToSrgbBranchingChannel(lin.b));
 }
+
+half3 sRGBToLinear( half3 Color ) 
+{
+	Color = max(6.10352e-5, Color); // minimum positive non-denormal (fixes black problem on DX11 AMD and NV)
+	return Color > 0.04045 ? pow( Color * (1.0 / 1.055) + 0.0521327, 2.4 ) : Color * (1.0 / 12.92);
+}
\ No newline at end of file
diff --git a/Engine/Shaders/PostProcessCombineLUTs.usf b/Engine/Shaders/PostProcessCombineLUTs.usf
--- a/Engine/Shaders/PostProcessCombineLUTs.usf
+++ b/Engine/Shaders/PostProcessCombineLUTs.usf
@@ -148,6 +148,8 @@ float3 ColorGrade( float3 LinearColor )
 	return mul( AP1_2_sRGB, WorkingColor );
 }
 
+uint OutputDevice;
+
 // todo: Weight[0] should be used for neutral, Texture* name should start with 1, color correction should apply on top of that
 #if USE_VOLUME_LUT == 1
 void MainPS(FWriteToSliceGeometryOutput Input, out float4 OutColor : SV_Target0)
@@ -207,15 +209,15 @@ void MainPS(float4 InUV : TEXCOORD0, out float4 OutColor : SV_Target0)
 	{
 		// Legacy LDR LUT color grading
 
-		// FIXME LUTs are in gamma space
-		half3 GammaColor = pow( FilmColor, 1.0 / 2.2 );
+		// FIXME LUTs are in sRGB space
+		half3 GammaColor = LinearToSrgbBranching( saturate( FilmColor ) );
 
 		// Transform color to 256x16 UV
 		float Scale = 15.0f / 16.0f;
 
-		float IntB = floor(Neutral.b * 15.9999f) / 16.0f;
+		float IntB = floor(GammaColor.b * 15.9999f) / 16.0f;
 
-		float2 UV = float2(Neutral.r * Scale / 16.0f + IntB, Neutral.g * Scale);
+		float2 UV = float2(GammaColor.r * Scale / 16.0f + IntB, GammaColor.g * Scale);
 		UV += float2(0.5f / 256.0f, 0.5f / 16.0f);
 
 		GammaColor = LUTWeights[0] * GammaColor;
@@ -238,7 +240,7 @@ void MainPS(float4 InUV : TEXCOORD0, out float4 OutColor : SV_Target0)
 	#endif
 
 		// Back to linear space
-		FilmColor = pow( GammaColor, 2.2 );
+		FilmColor = sRGBToLinear( GammaColor );
 	}
 #endif
 	
@@ -248,26 +250,29 @@ void MainPS(float4 InUV : TEXCOORD0, out float4 OutColor : SV_Target0)
 	// blend with custom LDR color, used for Fade track in Matinee
 	FilmColor = lerp( FilmColor * ColorScale, OverlayColor.rgb, OverlayColor.a );
 
-#if 0
 	// Apply "gamma" curve adjustment.
 	FilmColor = pow( FilmColor, InverseGamma.y );
 		
-	#if USE_GAMMA
+	half3 OutDeviceColor = 0;
+	BRANCH
+	if( OutputDevice == 0 )
+	{
+		// Apply conversion to sRGB (this must be an exact sRGB conversion else darks are bad).
+		// Branching is faster than branchless on AMD on PC.
+		OutDeviceColor = LinearToSrgbBranching( FilmColor );
+	}
+	else if( OutputDevice == 1 )
+	{
+		// Didn't profile yet if the branching version would be faster (different linear segment).
+		OutDeviceColor = LinearTo709Branchless( FilmColor );
+	}
+	else
+	{
 		// This is different than the prior "gamma" curve adjustment (but reusing the variable).
 		// For displays set to a gamma colorspace.
 		// Note, MacOSX native output is raw gamma 2.2 not sRGB!
-		half3 OutDeviceColor = pow( FilmColor, InverseGamma.z );
-	#elif USE_709
-		// Didn't profile yet if the branching version would be faster (different linear segment).
-		half3 OutDeviceColor = LinearTo709Branchless( FilmColor );
-	#else
-		// Apply conversion to sRGB (this must be an exact sRGB conversion else darks are bad).
-		// Branching is faster than branchless on AMD on PC.
-		half3 OutDeviceColor = LinearToSrgbBranching( FilmColor );
-	#endif
-#else
-	half3 OutDeviceColor = sqrt( FilmColor );
-#endif
+		OutDeviceColor = pow( FilmColor, InverseGamma.z );
+	}
 
 	// Better to saturate(lerp(a,b,t)) than lerp(saturate(a),saturate(b),t)
 	OutColor.rgb = OutDeviceColor / 1.05;
diff --git a/Engine/Shaders/PostProcessTonemap.usf b/Engine/Shaders/PostProcessTonemap.usf
--- a/Engine/Shaders/PostProcessTonemap.usf
+++ b/Engine/Shaders/PostProcessTonemap.usf
@@ -255,34 +255,6 @@ void MainPS(
 		#endif
 
 		half3 OutDeviceColor = ColorLookupTable( LinearColor );
-		
-#if 1	// TODO remove all of this. Do full output device conversion in LUT
-		// remove fake gamma
-		OutDeviceColor *= OutDeviceColor;
-
-		// Apply "gamma" curve adjustment.
-		OutDeviceColor.rgb = pow(OutDeviceColor.rgb, InverseGamma.y);
-		
-		#if MAC
-			// Note, MacOSX native output is raw gamma 2.2 not sRGB!
-			OutDeviceColor = pow(OutDeviceColor.rgb, 1.0/2.2);
-		#else
-			#if USE_709
-				// Didn't profile yet if the branching version would be faster (different linear segment).
-				OutDeviceColor = LinearTo709Branchless(OutDeviceColor);
-			#else
-				#if USE_GAMMA
-					// This is different than the prior "gamma" curve adjustment (but reusing the variable).
-					// For displays set to a gamma colorspace.
-					OutDeviceColor = pow(OutDeviceColor.rgb, InverseGamma.z);
-				#else
-					// Apply conversion to sRGB (this must be an exact sRGB conversion else darks are bad).
-					// Branching is faster than branchless on AMD on PC.
-					OutDeviceColor = LinearToSrgbBranching(OutDeviceColor);
-				#endif
-			#endif
-		#endif
-#endif
 
 		//OutDeviceColor = Graph( ScreenSpacePos );
 	
diff --git a/Engine/Source/Runtime/Renderer/Private/PostProcess/PostProcessCombineLUTs.cpp b/Engine/Source/Runtime/Renderer/Private/PostProcess/PostProcessCombineLUTs.cpp
--- a/Engine/Source/Runtime/Renderer/Private/PostProcess/PostProcessCombineLUTs.cpp
+++ b/Engine/Source/Runtime/Renderer/Private/PostProcess/PostProcessCombineLUTs.cpp
@@ -163,6 +163,8 @@ public:
 		FilmBlackClip.Bind(	Initializer.ParameterMap,TEXT("FilmBlackClip") );
 		FilmWhiteClip.Bind(	Initializer.ParameterMap,TEXT("FilmWhiteClip") );
 
+		OutputDevice.Bind( Initializer.ParameterMap,TEXT("OutputDevice") );
+
 		ColorMatrixR_ColorCurveCd1.Bind(Initializer.ParameterMap, TEXT("ColorMatrixR_ColorCurveCd1"));
 		ColorMatrixG_ColorCurveCd3Cm3.Bind(Initializer.ParameterMap, TEXT("ColorMatrixG_ColorCurveCd3Cm3"));
 		ColorMatrixB_ColorCurveCm2.Bind(Initializer.ParameterMap, TEXT("ColorMatrixB_ColorCurveCm2"));
@@ -215,18 +217,26 @@ public:
 		SetShaderValue( RHICmdList, ShaderRHI, FilmWhiteClip,	Settings.FilmWhiteClip );
 
 		{
+			static TConsoleVariableData<int32>* CVar709		= IConsoleManager::Get().FindTConsoleVariableDataInt(TEXT("r.Tonemapper709"));
+			static TConsoleVariableData<float>* CVarGamma	= IConsoleManager::Get().FindTConsoleVariableDataFloat(TEXT("r.TonemapperGamma"));
+
+			int32 Rec709 = CVar709->GetValueOnRenderThread();
+			float Gamma = CVarGamma->GetValueOnRenderThread();
+			
+			if( PLATFORM_APPLE && Gamma == 0.0f )
+			{
+				Gamma = 2.2f;
+			}
+			
+			int32 Value = 0;						// sRGB
+			Value = Rec709			? 1 : Value;	// Rec709
+			Value = Gamma != 0.0f	? 2 : Value;	// Explicit gamma
+			SetShaderValue( RHICmdList, ShaderRHI, OutputDevice, Value );
+			
 			FVector InvDisplayGammaValue;
 			InvDisplayGammaValue.X = 1.0f / ViewFamily.RenderTarget->GetDisplayGamma();
 			InvDisplayGammaValue.Y = 2.2f / ViewFamily.RenderTarget->GetDisplayGamma();
-			{
-				static TConsoleVariableData<float>* CVar = IConsoleManager::Get().FindTConsoleVariableDataFloat(TEXT("r.TonemapperGamma"));
-				float Value = CVar->GetValueOnRenderThread();
-				if(Value < 1.0f)
-				{
-					Value = 1.0f;
-				}
-				InvDisplayGammaValue.Z = 1.0f / Value;
-			}
+			InvDisplayGammaValue.Z = 1.0f / FMath::Max( Gamma, 1.0f );
 			SetShaderValue(RHICmdList, ShaderRHI, InverseGamma, InvDisplayGammaValue);
 		}
 
@@ -363,8 +373,6 @@ public:
 
 		OutEnvironment.SetDefine(TEXT("BLENDCOUNT"), BlendCount);
 		OutEnvironment.SetDefine(TEXT("USE_VOLUME_LUT"), UseVolumeTextureLUT(Platform));
-		//OutEnvironment.SetDefine(TEXT("USE_GAMMA"), OutputDevice == 1);
-		//OutEnvironment.SetDefine(TEXT("USE_709"), OutputDevice == 2);
 	}
 
 	virtual bool Serialize(FArchive& Ar) override
@@ -389,6 +397,8 @@ public:
 		Ar << ColorGain;
 		Ar << ColorOffset;
 
+		Ar << OutputDevice;
+
 		Ar << FilmSlope;
 		Ar << FilmToe;
 		Ar << FilmShoulder;
@@ -427,6 +437,8 @@ private: // ---------------------------------------------------
 	FShaderParameter FilmBlackClip;
 	FShaderParameter FilmWhiteClip;
 
+	FShaderParameter OutputDevice;
+
 	// Legacy
 	FShaderParameter ColorMatrixR_ColorCurveCd1;
 	FShaderParameter ColorMatrixG_ColorCurveCd3Cm3;
diff --git a/Engine/Source/Runtime/Renderer/Private/PostProcess/PostProcessTonemap.cpp b/Engine/Source/Runtime/Renderer/Private/PostProcess/PostProcessTonemap.cpp
--- a/Engine/Source/Runtime/Renderer/Private/PostProcess/PostProcessTonemap.cpp
+++ b/Engine/Source/Runtime/Renderer/Private/PostProcess/PostProcessTonemap.cpp
@@ -35,8 +35,6 @@ typedef enum {
 	TonemapperMosaic            = (1<<11),
 	TonemapperColorFringe       = (1<<12),
 	TonemapperMsaa              = (1<<13),
-	Tonemapper709               = (1<<14),
-	TonemapperGamma             = (1<<15),
 } TonemapperOption;
 
 // Tonemapper option cost (0 = no cost, 255 = max cost).
@@ -57,8 +55,6 @@ static uint8 TonemapperCostTab[] = {
 	1, //TonemapperMosaic
 	1, //TonemapperColorFringe
 	1, //TonemapperMsaa
-	1, //Tonemapper709
-	1, //TonemapperGamma
 };
 
 // Edit the following to add and remove configurations.
@@ -66,7 +62,7 @@ static uint8 TonemapperCostTab[] = {
 // Place most common first (faster when searching in TonemapperFindLeastExpensive()).
 
 // List of configurations compiled for PC.
-static uint32 TonemapperConfBitmaskPC[9] = { 
+static uint32 TonemapperConfBitmaskPC[7] = { 
 
 	TonemapperBloom +
 	TonemapperGrainJitter +
@@ -76,15 +72,6 @@ static uint32 TonemapperConfBitmaskPC[9] = {
 	TonemapperColorFringe +
 	0,
 
-	TonemapperBloom +
-	TonemapperGrainJitter +
-	TonemapperGrainIntensity +
-	TonemapperGrainQuantization +
-	TonemapperVignette +
-	TonemapperColorFringe +
-	Tonemapper709 +
-	0,
-
 	TonemapperBloom + 
 	TonemapperVignette +
 	TonemapperGrainQuantization +
@@ -104,14 +91,6 @@ static uint32 TonemapperConfBitmaskPC[9] = {
 	TonemapperColorFringe +
 	0,
 
-	TonemapperBloom +
-	TonemapperGrainJitter +
-	TonemapperGrainIntensity +
-	TonemapperVignette +
-	TonemapperColorFringe +
-	Tonemapper709 +
-	0,
-
 	TonemapperBloom + 
 	TonemapperVignette +
 	TonemapperColorFringe +
@@ -544,16 +523,6 @@ static uint32 TonemapperGenerateBitmaskPC(const FViewInfo* RESTRICT View, bool b
 		}
 	}
 
-	{
-		static TConsoleVariableData<int32>* CVar = IConsoleManager::Get().FindTConsoleVariableDataInt(TEXT("r.Tonemapper709"));
-		int32 Value = CVar->GetValueOnRenderThread();
-
-		if(Value > 0)
-		{
-			Bitmask |= Tonemapper709;
-		}
-	}
-
 	static const auto CVar = IConsoleManager::Get().FindTConsoleVariableDataInt(TEXT("r.SceneColorFringeQuality")); 
 
 	int32 FringeQuality = CVar->GetValueOnRenderThread();
@@ -835,8 +804,6 @@ class FPostProcessTonemapPS : public FGlobalShader
 		OutEnvironment.SetDefine(TEXT("USE_COLOR_FRINGE"),		 TonemapperIsDefined(ConfigBitmask, TonemapperColorFringe));
 		// @todo Mac OS X: in order to share precompiled shaders between GL 3.3 & GL 4.1 devices we mustn't use volume-texture rendering as it isn't universally supported.
 		OutEnvironment.SetDefine(TEXT("USE_VOLUME_LUT"),		 (IsFeatureLevelSupported(Platform,ERHIFeatureLevel::SM4) && GSupportsVolumeTextureRendering && Platform != EShaderPlatform::SP_OPENGL_SM4_MAC));
-		OutEnvironment.SetDefine(TEXT("USE_709"),                TonemapperIsDefined(ConfigBitmask, Tonemapper709));
-		OutEnvironment.SetDefine(TEXT("USE_GAMMA"),              TonemapperIsDefined(ConfigBitmask, TonemapperGamma));
 
 		if( !IsFeatureLevelSupported(Platform,ERHIFeatureLevel::SM5) )
 		{
@@ -1049,8 +1016,7 @@ public:
 #define VARIATION1(A) typedef FPostProcessTonemapPS<A> FPostProcessTonemapPS##A; \
 	IMPLEMENT_SHADER_TYPE2(FPostProcessTonemapPS##A, SF_Pixel);
 
-	VARIATION1(0)  VARIATION1(1)  VARIATION1(2)  VARIATION1(3)  VARIATION1(4)  VARIATION1(5)
-	VARIATION1(6)  VARIATION1(7)  VARIATION1(8)
+	VARIATION1(0)  VARIATION1(1)  VARIATION1(2)  VARIATION1(3)  VARIATION1(4)  VARIATION1(5) VARIATION1(6)
 
 #undef VARIATION1
 
@@ -1127,8 +1093,6 @@ void FRCPassPostProcessTonemap::Process(FRenderingCompositePassContext& Context)
 		case 4: SetShaderTempl<4>(Context); break;
 		case 5: SetShaderTempl<5>(Context); break;
 		case 6: SetShaderTempl<6>(Context); break;
-		case 7: SetShaderTempl<7>(Context); break;
-		case 8: SetShaderTempl<8>(Context); break;
 		default:
 			check(0);
 	}
