changed some postprocess passes to use UniformBuffers to allow them to use AsyncCompute in a following change

diff --git a/Engine/Shaders/DistanceFieldLightingPost.usf b/Engine/Shaders/DistanceFieldLightingPost.usf
--- a/Engine/Shaders/DistanceFieldLightingPost.usf
+++ b/Engine/Shaders/DistanceFieldLightingPost.usf
@@ -149,7 +149,6 @@ void FillGapsPS(
 #endif
 }
 
-float4 CameraMotion[5];
 Texture2D BentNormalHistoryTexture;
 SamplerState BentNormalHistorySampler;
 Texture2D IrradianceHistoryTexture;
@@ -182,11 +181,11 @@ void UpdateHistoryPS(
 		xyd.xy = UVAndScreenPos.zw * float2(0.5, -0.5) + 0.5;
 		xyd.z = ConvertToDeviceZ(SceneDepth);
 
-		float scaleM = 1.0 / (dot(xyd, CameraMotion[0].xyz) + CameraMotion[0].w);
+		float scaleM = 1.0 / (dot(xyd, CameraMotion.Value[0].xyz) + CameraMotion.Value[0].w);
 		float2 MotionVector;
 		// Unwound vector-matrix transform with special care for precision taken from Temporal AA
-		MotionVector.x = ((xyd.x * ((CameraMotion[1].x * xyd.y) + (CameraMotion[1].y * xyd.z) + CameraMotion[1].z)) + (CameraMotion[1].w * xyd.y) + (CameraMotion[2].x * xyd.x * xyd.x) + (CameraMotion[2].y * xyd.z) + CameraMotion[2].z) * scaleM;
-		MotionVector.y = ((xyd.y * ((CameraMotion[3].x * xyd.x) + (CameraMotion[3].y * xyd.z) + CameraMotion[3].z)) + (CameraMotion[3].w * xyd.x) + (CameraMotion[4].x * xyd.y * xyd.y) + (CameraMotion[4].y * xyd.z) + CameraMotion[4].z) * scaleM;
+		MotionVector.x = ((xyd.x * ((CameraMotion.Value[1].x * xyd.y) + (CameraMotion.Value[1].y * xyd.z) + CameraMotion.Value[1].z)) + (CameraMotion.Value[1].w * xyd.y) + (CameraMotion.Value[2].x * xyd.x * xyd.x) + (CameraMotion.Value[2].y * xyd.z) + CameraMotion.Value[2].z) * scaleM;
+		MotionVector.y = ((xyd.y * ((CameraMotion.Value[3].x * xyd.x) + (CameraMotion.Value[3].y * xyd.z) + CameraMotion.Value[3].z)) + (CameraMotion.Value[3].w * xyd.x) + (CameraMotion.Value[4].x * xyd.y * xyd.y) + (CameraMotion.Value[4].y * xyd.z) + CameraMotion.Value[4].z) * scaleM;
 
 		// Note: have to sample from one of the high res texels exactly to avoid filtering of the velocity buffer
 		// 0 is stored where camera velocity should be used, which must not be filtered with valid object velocities
diff --git a/Engine/Shaders/PostProcessTemporalAA.usf b/Engine/Shaders/PostProcessTemporalAA.usf
--- a/Engine/Shaders/PostProcessTemporalAA.usf
+++ b/Engine/Shaders/PostProcessTemporalAA.usf
@@ -12,7 +12,6 @@
 #define ENABLE_TEMPORAL_AA	1
 #endif
 
-float4 CameraMotion[5];
 float SampleWeights[9];
 float LowpassWeights[9];
 float PlusWeights[5];
diff --git a/Engine/Shaders/PostProcessTemporalCommon.usf b/Engine/Shaders/PostProcessTemporalCommon.usf
--- a/Engine/Shaders/PostProcessTemporalCommon.usf
+++ b/Engine/Shaders/PostProcessTemporalCommon.usf
@@ -187,10 +187,10 @@ which can be configured via defines for various temporal AA passes.
 #endif // HAS_INVERTED_Z_BUFFER
 	#endif	// AA_CROSS
 	// Camera motion for pixel or nearest pixel (in ScreenPos space).
-	float ScaleM = 1.0 / (dot(PosN, CameraMotion[0].xyz) + CameraMotion[0].w);
+	float ScaleM = 1.0 / (dot(PosN, CameraMotion.Value[0].xyz) + CameraMotion.Value[0].w);
 	float2 BackN;
-	BackN.x = -2.0 * ((PosN.x * ((CameraMotion[1].x * PosN.y) + (CameraMotion[1].y * PosN.z) + CameraMotion[1].z)) + (CameraMotion[1].w * PosN.y) + (CameraMotion[2].x * PosN.x * PosN.x) + (CameraMotion[2].y * PosN.z) + CameraMotion[2].z) * ScaleM;
-	BackN.y =  2.0 * ((PosN.y * ((CameraMotion[3].x * PosN.x) + (CameraMotion[3].y * PosN.z) + CameraMotion[3].z)) + (CameraMotion[3].w * PosN.x) + (CameraMotion[4].x * PosN.y * PosN.y) + (CameraMotion[4].y * PosN.z) + CameraMotion[4].z) * ScaleM;
+	BackN.x = -2.0 * ((PosN.x * ((CameraMotion.Value[1].x * PosN.y) + (CameraMotion.Value[1].y * PosN.z) + CameraMotion.Value[1].z)) + (CameraMotion.Value[1].w * PosN.y) + (CameraMotion.Value[2].x * PosN.x * PosN.x) + (CameraMotion.Value[2].y * PosN.z) + CameraMotion.Value[2].z) * ScaleM;
+	BackN.y =  2.0 * ((PosN.y * ((CameraMotion.Value[3].x * PosN.x) + (CameraMotion.Value[3].y * PosN.z) + CameraMotion.Value[3].z)) + (CameraMotion.Value[3].w * PosN.x) + (CameraMotion.Value[4].x * PosN.y * PosN.y) + (CameraMotion.Value[4].y * PosN.z) + CameraMotion.Value[4].z) * ScaleM;
 
 #if 0
 	//Fix for VR 'wobble' Imprecision in this method makes it not good enough to have on by default.  A proper fix is coming soon, but in the meantime, VR users can put this in locally.
diff --git a/Engine/Shaders/PostProcessVelocityFlatten.usf b/Engine/Shaders/PostProcessVelocityFlatten.usf
--- a/Engine/Shaders/PostProcessVelocityFlatten.usf
+++ b/Engine/Shaders/PostProcessVelocityFlatten.usf
@@ -7,11 +7,6 @@
 
 groupshared float2 Shared[ THREADGROUP_TOTALSIZE ];
 
-float4 CameraMotion[5];
-
-// View rect min in xy, max in zw.
-uint4 ViewDimensions;
-
 RWTexture2D<float4>	OutVelocityFlat;
 RWTexture2D<float4>	OutMaxTileVelocity;
 
@@ -30,8 +25,8 @@ void VelocityFlattenMain(
 	uint GroupIndex : SV_GroupIndex ) 
 {
 	uint ThreadIndex = GroupThreadId.y * THREADGROUP_SIZEX + GroupThreadId.x;
-	uint2 PixelPos = DispatchThreadId.xy * 2 + ViewDimensions.xy;
-	float2 ScreenUV = ( (float2)( DispatchThreadId.xy * 2 ) + 0.5 ) / ( ViewDimensions.zw - ViewDimensions.xy );
+	uint2 PixelPos = DispatchThreadId.xy * 2 + VelocityFlatten.ViewDimensions.xy;
+	float2 ScreenUV = ( (float2)( DispatchThreadId.xy * 2 ) + 0.5 ) / ( VelocityFlatten.ViewDimensions.zw - VelocityFlatten.ViewDimensions.xy );
 	float2 ScreenPosition = float2(2.0f, -2.0f) * ScreenUV + float2(-1.0f, 1.0f);
 
 	// FIXME for sub viewrect
@@ -81,13 +76,13 @@ void VelocityFlattenMain(
 		else
 		{
 			float3 PosN = float3( ScreenUV, DepthSamples[i] );
-			float ScaleM = rcp( dot( PosN, CameraMotion[0].xyz ) + CameraMotion[0].w );
-			Velocity[i].x = -2.0 * ((PosN.x * ((CameraMotion[1].x * PosN.y) + (CameraMotion[1].y * PosN.z) + CameraMotion[1].z)) + (CameraMotion[1].w * PosN.y) + (CameraMotion[2].x * PosN.x * PosN.x) + (CameraMotion[2].y * PosN.z) + CameraMotion[2].z) * ScaleM;
-			Velocity[i].y =  2.0 * ((PosN.y * ((CameraMotion[3].x * PosN.x) + (CameraMotion[3].y * PosN.z) + CameraMotion[3].z)) + (CameraMotion[3].w * PosN.x) + (CameraMotion[4].x * PosN.y * PosN.y) + (CameraMotion[4].y * PosN.z) + CameraMotion[4].z) * ScaleM;
+			float ScaleM = rcp( dot( PosN, CameraMotion.Value[0].xyz ) + CameraMotion.Value[0].w );
+			Velocity[i].x = -2.0 * ((PosN.x * ((CameraMotion.Value[1].x * PosN.y) + (CameraMotion.Value[1].y * PosN.z) + CameraMotion.Value[1].z)) + (CameraMotion.Value[1].w * PosN.y) + (CameraMotion.Value[2].x * PosN.x * PosN.x) + (CameraMotion.Value[2].y * PosN.z) + CameraMotion.Value[2].z) * ScaleM;
+			Velocity[i].y =  2.0 * ((PosN.y * ((CameraMotion.Value[3].x * PosN.x) + (CameraMotion.Value[3].y * PosN.z) + CameraMotion.Value[3].z)) + (CameraMotion.Value[3].w * PosN.x) + (CameraMotion.Value[4].x * PosN.y * PosN.y) + (CameraMotion.Value[4].y * PosN.z) + CameraMotion.Value[4].z) * ScaleM;
 		}
 
 		BRANCH
-		if( all( PixelPos.xy + Offsets[i] < ViewDimensions.zw ) )
+		if( all( PixelPos.xy + Offsets[i] < VelocityFlatten.ViewDimensions.zw ) )
 		{
 			OutVelocityFlat[ PixelPos + Offsets[i] ] = EncodeVelocityToTexture( Velocity[i] ).xyxy;
 			//OutPackedVelocityDepth[ PixelPos + Offsets[i] ] = float2( Velocity[i], DepthSamples[i] ).xyxy;
diff --git a/Engine/Source/Runtime/Renderer/Private/CompositionLighting/PostProcessAmbientOcclusion.cpp b/Engine/Source/Runtime/Renderer/Private/CompositionLighting/PostProcessAmbientOcclusion.cpp
--- a/Engine/Source/Runtime/Renderer/Private/CompositionLighting/PostProcessAmbientOcclusion.cpp
+++ b/Engine/Source/Runtime/Renderer/Private/CompositionLighting/PostProcessAmbientOcclusion.cpp
@@ -20,6 +20,7 @@ static TAutoConsoleVariable<int32> CVarAmbientOcclusionSampleSetQuality(
 	TEXT("1: high sample count (defined in shader, 6 * 2 per pixel)"),
 	ECVF_Scalability | ECVF_RenderThreadSafe);
 
+IMPLEMENT_UNIFORM_BUFFER_STRUCT(FCameraMotionParameters,TEXT("CameraMotion"));
 
 /** Encapsulates the post processing ambient occlusion pixel shader. */
 template <uint32 bInitialPass>
@@ -112,18 +113,6 @@ FArchive& operator<<(FArchive& Ar, FScreenSpaceAOandSSRShaderParameters& This)
 
 // --------------------------------------------------------
 
-void FCameraMotionParameters::Bind(const FShaderParameterMap& ParameterMap)
-{
-	CameraMotion.Bind(ParameterMap, TEXT("CameraMotion"));
-}
-
-FArchive& operator<<(FArchive& Ar, FCameraMotionParameters& This)
-{
-	Ar << This.CameraMotion;
-
-	return Ar;
-}
-
 /**
  * Encapsulates the post processing ambient occlusion pixel shader.
  * @param bAOSetupAsInput true:use AO setup instead of full resolution depth and normal
@@ -580,3 +569,69 @@ FPooledRenderTargetDesc FRCPassPostProcessBasePassAO::ComputeOutputDesc(EPassOut
 
 	return Ret;
 }
+
+TUniformBufferRef<FCameraMotionParameters> CreateCameraMotionParametersUniformBuffer(const FSceneView& View)
+{
+	FSceneViewState* ViewState = (FSceneViewState*)View.State;
+
+	FMatrix Proj = View.ViewMatrices.ProjMatrix;
+	FMatrix PrevProj = ViewState->PrevViewMatrices.ProjMatrix;
+
+	// Remove jitter
+	Proj.M[2][0] -= View.ViewMatrices.TemporalAASample.X * 2.0f / View.ViewRect.Width();
+	Proj.M[2][1] -= View.ViewMatrices.TemporalAASample.Y * 2.0f / View.ViewRect.Height();
+	PrevProj.M[2][0] -= ViewState->PrevViewMatrices.TemporalAASample.X * 2.0f / View.ViewRect.Width();
+	PrevProj.M[2][1] -= ViewState->PrevViewMatrices.TemporalAASample.Y * 2.0f / View.ViewRect.Height();
+
+	FVector DeltaTranslation = ViewState->PrevViewMatrices.PreViewTranslation - View.ViewMatrices.PreViewTranslation;
+	FMatrix ViewProj = ( View.ViewMatrices.TranslatedViewMatrix * Proj ).GetTransposed();
+	FMatrix PrevViewProj = ( FTranslationMatrix(DeltaTranslation) * ViewState->PrevViewMatrices.TranslatedViewMatrix * PrevProj ).GetTransposed();
+
+	double InvViewProj[16];
+	Inverse4x4( InvViewProj, (float*)ViewProj.M );
+
+	const float* p = (float*)PrevViewProj.M;
+
+	const double cxx = InvViewProj[ 0]; const double cxy = InvViewProj[ 1]; const double cxz = InvViewProj[ 2]; const double cxw = InvViewProj[ 3];
+	const double cyx = InvViewProj[ 4]; const double cyy = InvViewProj[ 5]; const double cyz = InvViewProj[ 6]; const double cyw = InvViewProj[ 7];
+	const double czx = InvViewProj[ 8]; const double czy = InvViewProj[ 9]; const double czz = InvViewProj[10]; const double czw = InvViewProj[11];
+	const double cwx = InvViewProj[12]; const double cwy = InvViewProj[13]; const double cwz = InvViewProj[14]; const double cww = InvViewProj[15];
+
+	const double pxx = (double)(p[ 0]); const double pxy = (double)(p[ 1]); const double pxz = (double)(p[ 2]); const double pxw = (double)(p[ 3]);
+	const double pyx = (double)(p[ 4]); const double pyy = (double)(p[ 5]); const double pyz = (double)(p[ 6]); const double pyw = (double)(p[ 7]);
+	const double pwx = (double)(p[12]); const double pwy = (double)(p[13]); const double pwz = (double)(p[14]); const double pww = (double)(p[15]);
+
+	FCameraMotionParameters LocalCameraMotion;
+
+	LocalCameraMotion.Value[0] = FVector4(
+		(float)(4.0*(cwx*pww + cxx*pwx + cyx*pwy + czx*pwz)),
+		(float)((-4.0)*(cwy*pww + cxy*pwx + cyy*pwy + czy*pwz)),
+		(float)(2.0*(cwz*pww + cxz*pwx + cyz*pwy + czz*pwz)),
+		(float)(2.0*(cww*pww - cwx*pww + cwy*pww + (cxw - cxx + cxy)*pwx + (cyw - cyx + cyy)*pwy + (czw - czx + czy)*pwz)));
+
+	LocalCameraMotion.Value[1] = FVector4(
+		(float)(( 4.0)*(cwy*pww + cxy*pwx + cyy*pwy + czy*pwz)),
+		(float)((-2.0)*(cwz*pww + cxz*pwx + cyz*pwy + czz*pwz)),
+		(float)((-2.0)*(cww*pww + cwy*pww + cxw*pwx - 2.0*cxx*pwx + cxy*pwx + cyw*pwy - 2.0*cyx*pwy + cyy*pwy + czw*pwz - 2.0*czx*pwz + czy*pwz - cwx*(2.0*pww + pxw) - cxx*pxx - cyx*pxy - czx*pxz)),
+		(float)(-2.0*(cyy*pwy + czy*pwz + cwy*(pww + pxw) + cxy*(pwx + pxx) + cyy*pxy + czy*pxz)));
+
+	LocalCameraMotion.Value[2] = FVector4(
+		(float)((-4.0)*(cwx*pww + cxx*pwx + cyx*pwy + czx*pwz)),
+		(float)(cyz*pwy + czz*pwz + cwz*(pww + pxw) + cxz*(pwx + pxx) + cyz*pxy + czz*pxz),
+		(float)(cwy*pww + cwy*pxw + cww*(pww + pxw) - cwx*(pww + pxw) + (cxw - cxx + cxy)*(pwx + pxx) + (cyw - cyx + cyy)*(pwy + pxy) + (czw - czx + czy)*(pwz + pxz)),
+		(float)(0));
+
+	LocalCameraMotion.Value[3] = FVector4(
+		(float)((-4.0)*(cwx*pww + cxx*pwx + cyx*pwy + czx*pwz)),
+		(float)((-2.0)*(cwz*pww + cxz*pwx + cyz*pwy + czz*pwz)),
+		(float)(2.0*((-cww)*pww + cwx*pww - 2.0*cwy*pww - cxw*pwx + cxx*pwx - 2.0*cxy*pwx - cyw*pwy + cyx*pwy - 2.0*cyy*pwy - czw*pwz + czx*pwz - 2.0*czy*pwz + cwy*pyw + cxy*pyx + cyy*pyy + czy*pyz)),
+		(float)(2.0*(cyx*pwy + czx*pwz + cwx*(pww - pyw) + cxx*(pwx - pyx) - cyx*pyy - czx*pyz)));
+
+	LocalCameraMotion.Value[4] = FVector4(
+		(float)(4.0*(cwy*pww + cxy*pwx + cyy*pwy + czy*pwz)),
+		(float)(cyz*pwy + czz*pwz + cwz*(pww - pyw) + cxz*(pwx - pyx) - cyz*pyy - czz*pyz),
+		(float)(cwy*pww + cww*(pww - pyw) - cwy*pyw + cwx*((-pww) + pyw) + (cxw - cxx + cxy)*(pwx - pyx) + (cyw - cyx + cyy)*(pwy - pyy) + (czw - czx + czy)*(pwz - pyz)),
+		(float)(0));
+
+	return TUniformBufferRef<FCameraMotionParameters>::CreateUniformBufferImmediate(LocalCameraMotion, UniformBuffer_SingleFrame);
+}
\ No newline at end of file
diff --git a/Engine/Source/Runtime/Renderer/Private/CompositionLighting/PostProcessAmbientOcclusion.h b/Engine/Source/Runtime/Renderer/Private/CompositionLighting/PostProcessAmbientOcclusion.h
--- a/Engine/Source/Runtime/Renderer/Private/CompositionLighting/PostProcessAmbientOcclusion.h
+++ b/Engine/Source/Runtime/Renderer/Private/CompositionLighting/PostProcessAmbientOcclusion.h
@@ -149,86 +149,11 @@ private:
 	FShaderParameter ScreenSpaceAOandSSRShaderParams;
 };
 
-/** Pixel shader parameters needed for screen space TemporalAA and SSAO passes. */
-class FCameraMotionParameters
-{
-public:
-
-	void Bind(const FShaderParameterMap& ParameterMap);
-
-	template< typename ShaderRHIParamRef >
-	void Set(FRHICommandList& RHICmdList, const FSceneView& View, const ShaderRHIParamRef ShaderRHI) const;
 
-	friend FArchive& operator<<(FArchive& Ar, FCameraMotionParameters& This);
 
-private:
-	FShaderParameter CameraMotion;
-};
+/** The uniform shader parameters needed for screen space TemporalAA and SSAO passes. */
+BEGIN_UNIFORM_BUFFER_STRUCT(FCameraMotionParameters,)
+	DECLARE_UNIFORM_BUFFER_STRUCT_MEMBER_ARRAY(FVector4, Value, [5])
+END_UNIFORM_BUFFER_STRUCT(FCameraMotionParameters)
 
-
-template< typename ShaderRHIParamRef >
-void FCameraMotionParameters::Set(FRHICommandList& RHICmdList, const FSceneView& View, const ShaderRHIParamRef ShaderRHI) const
-{
-	FSceneViewState* ViewState = (FSceneViewState*)View.State;
-
-	FMatrix Proj = View.ViewMatrices.ProjMatrix;
-	FMatrix PrevProj = ViewState->PrevViewMatrices.ProjMatrix;
-
-	// Remove jitter
-	Proj.M[2][0] -= View.ViewMatrices.TemporalAASample.X * 2.0f / View.ViewRect.Width();
-	Proj.M[2][1] -= View.ViewMatrices.TemporalAASample.Y * 2.0f / View.ViewRect.Height();
-	PrevProj.M[2][0] -= ViewState->PrevViewMatrices.TemporalAASample.X * 2.0f / View.ViewRect.Width();
-	PrevProj.M[2][1] -= ViewState->PrevViewMatrices.TemporalAASample.Y * 2.0f / View.ViewRect.Height();
-
-	FVector DeltaTranslation = ViewState->PrevViewMatrices.PreViewTranslation - View.ViewMatrices.PreViewTranslation;
-	FMatrix ViewProj = ( View.ViewMatrices.TranslatedViewMatrix * Proj ).GetTransposed();
-	FMatrix PrevViewProj = ( FTranslationMatrix(DeltaTranslation) * ViewState->PrevViewMatrices.TranslatedViewMatrix * PrevProj ).GetTransposed();
-
-	double InvViewProj[16];
-	Inverse4x4( InvViewProj, (float*)ViewProj.M );
-
-	const float* p = (float*)PrevViewProj.M;
-
-	const double cxx = InvViewProj[ 0]; const double cxy = InvViewProj[ 1]; const double cxz = InvViewProj[ 2]; const double cxw = InvViewProj[ 3];
-	const double cyx = InvViewProj[ 4]; const double cyy = InvViewProj[ 5]; const double cyz = InvViewProj[ 6]; const double cyw = InvViewProj[ 7];
-	const double czx = InvViewProj[ 8]; const double czy = InvViewProj[ 9]; const double czz = InvViewProj[10]; const double czw = InvViewProj[11];
-	const double cwx = InvViewProj[12]; const double cwy = InvViewProj[13]; const double cwz = InvViewProj[14]; const double cww = InvViewProj[15];
-
-	const double pxx = (double)(p[ 0]); const double pxy = (double)(p[ 1]); const double pxz = (double)(p[ 2]); const double pxw = (double)(p[ 3]);
-	const double pyx = (double)(p[ 4]); const double pyy = (double)(p[ 5]); const double pyz = (double)(p[ 6]); const double pyw = (double)(p[ 7]);
-	const double pwx = (double)(p[12]); const double pwy = (double)(p[13]); const double pwz = (double)(p[14]); const double pww = (double)(p[15]);
-
-	FVector4 CameraMotionValue[5];
-
-	CameraMotionValue[0] = FVector4(
-		(float)(4.0*(cwx*pww + cxx*pwx + cyx*pwy + czx*pwz)),
-		(float)((-4.0)*(cwy*pww + cxy*pwx + cyy*pwy + czy*pwz)),
-		(float)(2.0*(cwz*pww + cxz*pwx + cyz*pwy + czz*pwz)),
-		(float)(2.0*(cww*pww - cwx*pww + cwy*pww + (cxw - cxx + cxy)*pwx + (cyw - cyx + cyy)*pwy + (czw - czx + czy)*pwz)));
-
-	CameraMotionValue[1] = FVector4(
-		(float)(( 4.0)*(cwy*pww + cxy*pwx + cyy*pwy + czy*pwz)),
-		(float)((-2.0)*(cwz*pww + cxz*pwx + cyz*pwy + czz*pwz)),
-		(float)((-2.0)*(cww*pww + cwy*pww + cxw*pwx - 2.0*cxx*pwx + cxy*pwx + cyw*pwy - 2.0*cyx*pwy + cyy*pwy + czw*pwz - 2.0*czx*pwz + czy*pwz - cwx*(2.0*pww + pxw) - cxx*pxx - cyx*pxy - czx*pxz)),
-		(float)(-2.0*(cyy*pwy + czy*pwz + cwy*(pww + pxw) + cxy*(pwx + pxx) + cyy*pxy + czy*pxz)));
-
-	CameraMotionValue[2] = FVector4(
-		(float)((-4.0)*(cwx*pww + cxx*pwx + cyx*pwy + czx*pwz)),
-		(float)(cyz*pwy + czz*pwz + cwz*(pww + pxw) + cxz*(pwx + pxx) + cyz*pxy + czz*pxz),
-		(float)(cwy*pww + cwy*pxw + cww*(pww + pxw) - cwx*(pww + pxw) + (cxw - cxx + cxy)*(pwx + pxx) + (cyw - cyx + cyy)*(pwy + pxy) + (czw - czx + czy)*(pwz + pxz)),
-		(float)(0));
-
-	CameraMotionValue[3] = FVector4(
-		(float)((-4.0)*(cwx*pww + cxx*pwx + cyx*pwy + czx*pwz)),
-		(float)((-2.0)*(cwz*pww + cxz*pwx + cyz*pwy + czz*pwz)),
-		(float)(2.0*((-cww)*pww + cwx*pww - 2.0*cwy*pww - cxw*pwx + cxx*pwx - 2.0*cxy*pwx - cyw*pwy + cyx*pwy - 2.0*cyy*pwy - czw*pwz + czx*pwz - 2.0*czy*pwz + cwy*pyw + cxy*pyx + cyy*pyy + czy*pyz)),
-		(float)(2.0*(cyx*pwy + czx*pwz + cwx*(pww - pyw) + cxx*(pwx - pyx) - cyx*pyy - czx*pyz)));
-
-	CameraMotionValue[4] = FVector4(
-		(float)(4.0*(cwy*pww + cxy*pwx + cyy*pwy + czy*pwz)),
-		(float)(cyz*pwy + czz*pwz + cwz*(pww - pyw) + cxz*(pwx - pyx) - cyz*pyy - czz*pyz),
-		(float)(cwy*pww + cww*(pww - pyw) - cwy*pyw + cwx*((-pww) + pyw) + (cxw - cxx + cxy)*(pwx - pyx) + (cyw - cyx + cyy)*(pwy - pyy) + (czw - czx + czy)*(pwz - pyz)),
-		(float)(0));
-
-	SetShaderValueArray(RHICmdList, ShaderRHI, CameraMotion, CameraMotionValue, 5);
-}
\ No newline at end of file
+TUniformBufferRef<FCameraMotionParameters> CreateCameraMotionParametersUniformBuffer(const FSceneView& View);
\ No newline at end of file
diff --git a/Engine/Source/Runtime/Renderer/Private/DistanceFieldSurfaceCacheLighting.cpp b/Engine/Source/Runtime/Renderer/Private/DistanceFieldSurfaceCacheLighting.cpp
--- a/Engine/Source/Runtime/Renderer/Private/DistanceFieldSurfaceCacheLighting.cpp
+++ b/Engine/Source/Runtime/Renderer/Private/DistanceFieldSurfaceCacheLighting.cpp
@@ -2340,7 +2340,6 @@ public:
 		IrradianceSampler.Bind(Initializer.ParameterMap, TEXT("IrradianceSampler"));
 		IrradianceHistoryTexture.Bind(Initializer.ParameterMap, TEXT("IrradianceHistoryTexture"));
 		IrradianceHistorySampler.Bind(Initializer.ParameterMap, TEXT("IrradianceHistorySampler"));
-		CameraMotion.Bind(Initializer.ParameterMap);
 		HistoryWeight.Bind(Initializer.ParameterMap, TEXT("HistoryWeight"));
 		HistoryDistanceThreshold.Bind(Initializer.ParameterMap, TEXT("HistoryDistanceThreshold"));
 		UseHistoryFilter.Bind(Initializer.ParameterMap, TEXT("UseHistoryFilter"));
@@ -2404,7 +2403,7 @@ public:
 				);
 		}
 
-		CameraMotion.Set(RHICmdList, View, ShaderRHI);
+		SetUniformBufferParameter(RHICmdList, ShaderRHI, GetUniformBufferParameter<FCameraMotionParameters>(), CreateCameraMotionParametersUniformBuffer(View));
 
 		SetShaderValue(RHICmdList, ShaderRHI, HistoryWeight, GAOHistoryWeight);
 		SetShaderValue(RHICmdList, ShaderRHI, HistoryDistanceThreshold, GAOHistoryDistanceThreshold);
@@ -2432,7 +2431,6 @@ public:
 		Ar << IrradianceSampler;
 		Ar << IrradianceHistoryTexture;
 		Ar << IrradianceHistorySampler;
-		Ar << CameraMotion;
 		Ar << HistoryWeight;
 		Ar << HistoryDistanceThreshold;
 		Ar << UseHistoryFilter;
@@ -2452,7 +2450,6 @@ private:
 	FShaderResourceParameter IrradianceSampler;
 	FShaderResourceParameter IrradianceHistoryTexture;
 	FShaderResourceParameter IrradianceHistorySampler;
-	FCameraMotionParameters CameraMotion;
 	FShaderParameter HistoryWeight;
 	FShaderParameter HistoryDistanceThreshold;
 	FShaderParameter UseHistoryFilter;
diff --git a/Engine/Source/Runtime/Renderer/Private/PostProcess/PostProcessMotionBlur.cpp b/Engine/Source/Runtime/Renderer/Private/PostProcess/PostProcessMotionBlur.cpp
--- a/Engine/Source/Runtime/Renderer/Private/PostProcess/PostProcessMotionBlur.cpp
+++ b/Engine/Source/Runtime/Renderer/Private/PostProcess/PostProcessMotionBlur.cpp
@@ -678,7 +678,12 @@ FPooledRenderTargetDesc FRCPassPostProcessMotionBlurRecombine::ComputeOutputDesc
 	return Ret;
 }
 
+/** The uniform shader parameters used in FPostProcessVelocityFlattenCS. */
+BEGIN_UNIFORM_BUFFER_STRUCT(FVelocityFlattenParameters,)
+	DECLARE_UNIFORM_BUFFER_STRUCT_MEMBER(FIntRect,ViewDimensions)	// View rect min in xy, max in zw.
+END_UNIFORM_BUFFER_STRUCT(FVelocityFlattenParameters)
 
+IMPLEMENT_UNIFORM_BUFFER_STRUCT(FVelocityFlattenParameters,TEXT("VelocityFlatten"));
 
 class FPostProcessVelocityFlattenCS : public FGlobalShader
 {
@@ -700,19 +705,17 @@ class FPostProcessVelocityFlattenCS : public FGlobalShader
 	FPostProcessVelocityFlattenCS() {}
 
 public:
-	FShaderParameter		OutVelocityFlat;
-	FShaderParameter		OutPackedVelocityDepth;
-	FShaderParameter		OutMaxTileVelocity;
-	
+	FShaderParameter		OutVelocityFlat;		// UAV
+	FShaderParameter		OutPackedVelocityDepth;	// unused
+	FShaderParameter		OutMaxTileVelocity;		// UAV
+
 	FPostProcessVelocityFlattenCS(const ShaderMetaType::CompiledShaderInitializerType& Initializer)
 		: FGlobalShader(Initializer)
 	{
 		PostprocessParameter.Bind(Initializer.ParameterMap);
-		CameraMotionParams.Bind(Initializer.ParameterMap);
 		OutVelocityFlat.Bind(Initializer.ParameterMap, TEXT("OutVelocityFlat"));
 		OutPackedVelocityDepth.Bind(Initializer.ParameterMap, TEXT("OutPackedVelocityDepth"));
 		OutMaxTileVelocity.Bind(Initializer.ParameterMap, TEXT("OutMaxTileVelocity"));
-		ViewDimensions.Bind(Initializer.ParameterMap, TEXT("ViewDimensions"));
 	}
 
 	void SetCS( FRHICommandList& RHICmdList, const FRenderingCompositePassContext& Context, const FSceneView& View )
@@ -721,27 +724,31 @@ public:
 
 		FGlobalShader::SetParameters(RHICmdList, ShaderRHI, Context.View);
 		PostprocessParameter.SetCS(ShaderRHI, Context, TStaticSamplerState<SF_Point,AM_Clamp,AM_Clamp,AM_Clamp>::GetRHI());
-		CameraMotionParams.Set(RHICmdList, Context.View, ShaderRHI);
 
-		SetShaderValue(RHICmdList, ShaderRHI, ViewDimensions, View.ViewRect);
+		SetUniformBufferParameter(RHICmdList, ShaderRHI, GetUniformBufferParameter<FCameraMotionParameters>(), CreateCameraMotionParametersUniformBuffer(Context.View));
+		
+		{
+			FVelocityFlattenParameters Data;
+
+			Data.ViewDimensions = View.ViewRect;
+
+			TUniformBufferRef<FVelocityFlattenParameters> UniformBuffer = TUniformBufferRef<FVelocityFlattenParameters>::CreateUniformBufferImmediate(Data, UniformBuffer_SingleFrame);
+			SetUniformBufferParameter(RHICmdList, ShaderRHI, GetUniformBufferParameter<FVelocityFlattenParameters>(), UniformBuffer);
+		}
 	}
 	
 	virtual bool Serialize(FArchive& Ar) override
 	{
 		bool bShaderHasOutdatedParameters = FGlobalShader::Serialize(Ar);
 		Ar << PostprocessParameter;
-		Ar << CameraMotionParams;
 		Ar << OutVelocityFlat;
 		Ar << OutPackedVelocityDepth;
 		Ar << OutMaxTileVelocity;
-		Ar << ViewDimensions;
 		return bShaderHasOutdatedParameters;
 	}
 
 private:
 	FPostProcessPassParameters	PostprocessParameter;
-	FCameraMotionParameters		CameraMotionParams;
-	FShaderParameter			ViewDimensions;
 };
 
 IMPLEMENT_SHADER_TYPE(,FPostProcessVelocityFlattenCS,TEXT("PostProcessVelocityFlatten"),TEXT("VelocityFlattenMain"),SF_Compute);
diff --git a/Engine/Source/Runtime/Renderer/Private/PostProcess/PostProcessTemporalAA.cpp b/Engine/Source/Runtime/Renderer/Private/PostProcess/PostProcessTemporalAA.cpp
--- a/Engine/Source/Runtime/Renderer/Private/PostProcess/PostProcessTemporalAA.cpp
+++ b/Engine/Source/Runtime/Renderer/Private/PostProcess/PostProcessTemporalAA.cpp
@@ -65,7 +65,6 @@ public:
 	FShaderParameter SampleWeights;
 	FShaderParameter LowpassWeights;
 	FShaderParameter PlusWeights;
-	FCameraMotionParameters CameraMotionParams;
 	FShaderParameter RandomOffset;
 
 	/** Initialization constructor. */
@@ -77,7 +76,6 @@ public:
 		SampleWeights.Bind(Initializer.ParameterMap, TEXT("SampleWeights"));
 		LowpassWeights.Bind(Initializer.ParameterMap, TEXT("LowpassWeights"));
 		PlusWeights.Bind(Initializer.ParameterMap, TEXT("PlusWeights"));
-		CameraMotionParams.Bind(Initializer.ParameterMap);
 		RandomOffset.Bind(Initializer.ParameterMap, TEXT("RandomOffset"));
 	}
 
@@ -85,7 +83,7 @@ public:
 	virtual bool Serialize(FArchive& Ar) override
 	{
 		bool bShaderHasOutdatedParameters = FGlobalShader::Serialize(Ar);
-		Ar << PostprocessParameter << DeferredParameters << SampleWeights << LowpassWeights << PlusWeights << CameraMotionParams << RandomOffset;
+		Ar << PostprocessParameter << DeferredParameters << SampleWeights << LowpassWeights << PlusWeights << RandomOffset;
 		return bShaderHasOutdatedParameters;
 	}
 
@@ -186,7 +184,7 @@ public:
 
 		}
 
-		CameraMotionParams.Set(Context.RHICmdList, Context.View, ShaderRHI);
+		SetUniformBufferParameter(Context.RHICmdList, ShaderRHI, GetUniformBufferParameter<FCameraMotionParameters>(), CreateCameraMotionParametersUniformBuffer(Context.View));
 	}
 };
 
