Centralize USE_VOLUME_LUT logic

diff --git a/Engine/Source/Runtime/Renderer/Private/PostProcess/PostProcessCombineLUTs.cpp b/Engine/Source/Runtime/Renderer/Private/PostProcess/PostProcessCombineLUTs.cpp
--- a/Engine/Source/Runtime/Renderer/Private/PostProcess/PostProcessCombineLUTs.cpp
+++ b/Engine/Source/Runtime/Renderer/Private/PostProcess/PostProcessCombineLUTs.cpp
@@ -52,7 +52,7 @@ static TAutoConsoleVariable<int32> CVarTonemapperFilm(
 
 // false:use 256x16 texture / true:use volume texture (faster, requires geometry shader)
 // USE_VOLUME_LUT: needs to be the same for C++ and HLSL
-static bool UseVolumeTextureLUT(EShaderPlatform Platform) 
+bool UseVolumeTextureLUT(EShaderPlatform Platform) 
 {
 	// @todo Mac OS X: in order to share precompiled shaders between GL 3.3 & GL 4.1 devices we mustn't use volume-texture rendering as it isn't universally supported.
 	return (IsFeatureLevelSupported(Platform,ERHIFeatureLevel::SM4) && GSupportsVolumeTextureRendering && Platform != EShaderPlatform::SP_OPENGL_SM4_MAC && RHISupportsGeometryShaders(Platform));
diff --git a/Engine/Source/Runtime/Renderer/Private/PostProcess/PostProcessCombineLUTs.h b/Engine/Source/Runtime/Renderer/Private/PostProcess/PostProcessCombineLUTs.h
--- a/Engine/Source/Runtime/Renderer/Private/PostProcess/PostProcessCombineLUTs.h
+++ b/Engine/Source/Runtime/Renderer/Private/PostProcess/PostProcessCombineLUTs.h
@@ -8,6 +8,8 @@
 
 #include "RenderingCompositionGraph.h"
 
+bool UseVolumeTextureLUT(EShaderPlatform Platform);
+
 // derives from TRenderingCompositePassBase<InputCount, OutputCount>
 class FRCPassPostProcessCombineLUTs : public TRenderingCompositePassBase<0, 1>
 {
diff --git a/Engine/Source/Runtime/Renderer/Private/PostProcess/PostProcessTonemap.cpp b/Engine/Source/Runtime/Renderer/Private/PostProcess/PostProcessTonemap.cpp
--- a/Engine/Source/Runtime/Renderer/Private/PostProcess/PostProcessTonemap.cpp
+++ b/Engine/Source/Runtime/Renderer/Private/PostProcess/PostProcessTonemap.cpp
@@ -802,8 +802,7 @@ class FPostProcessTonemapPS : public FGlobalShader
 		OutEnvironment.SetDefine(TEXT("USE_GRAIN_QUANTIZATION"), TonemapperIsDefined(ConfigBitmask, TonemapperGrainQuantization));
 		OutEnvironment.SetDefine(TEXT("USE_VIGNETTE"),           TonemapperIsDefined(ConfigBitmask, TonemapperVignette));
 		OutEnvironment.SetDefine(TEXT("USE_COLOR_FRINGE"),		 TonemapperIsDefined(ConfigBitmask, TonemapperColorFringe));
-		// @todo Mac OS X: in order to share precompiled shaders between GL 3.3 & GL 4.1 devices we mustn't use volume-texture rendering as it isn't universally supported.
-		OutEnvironment.SetDefine(TEXT("USE_VOLUME_LUT"),		 (IsFeatureLevelSupported(Platform,ERHIFeatureLevel::SM4) && GSupportsVolumeTextureRendering && Platform != EShaderPlatform::SP_OPENGL_SM4_MAC));
+		OutEnvironment.SetDefine(TEXT("USE_VOLUME_LUT"), UseVolumeTextureLUT(Platform));
 
 		if( !IsFeatureLevelSupported(Platform,ERHIFeatureLevel::SM5) )
 		{
