RenderTargets can now opt out of CMASK/HTILE.  CMASK/HTILE enable/disable is now locked at creation time.

diff --git a/Engine/Source/Runtime/Engine/Private/Particles/ParticleCurveTexture.cpp b/Engine/Source/Runtime/Engine/Private/Particles/ParticleCurveTexture.cpp
--- a/Engine/Source/Runtime/Engine/Private/Particles/ParticleCurveTexture.cpp
+++ b/Engine/Source/Runtime/Engine/Private/Particles/ParticleCurveTexture.cpp
@@ -508,7 +508,7 @@ void FParticleCurveTexture::InitRHI()
 		PF_B8G8R8A8,
 		/*NumMips=*/ 1,
 		TexCreate_None,
-		TexCreate_RenderTargetable,
+		TexCreate_RenderTargetable | TexCreate_NoFastClear,
 		/*bForceSeparateTargetAndShaderResource=*/ false,
 		CreateInfo,
 		CurveTextureTargetRHI,
diff --git a/Engine/Source/Runtime/RHI/Public/RHIDefinitions.h b/Engine/Source/Runtime/RHI/Public/RHIDefinitions.h
--- a/Engine/Source/Runtime/RHI/Public/RHIDefinitions.h
+++ b/Engine/Source/Runtime/RHI/Public/RHIDefinitions.h
@@ -529,6 +529,8 @@ enum ETextureCreateFlags
 	TexCreate_TargetArraySlicesIndependently	= 1<<23,
 	// Texture that may be shared with DX9 or other devices
 	TexCreate_Shared = 1 << 24,
+	// RenderTarget will not use full-texture fast clear functionality.
+	TexCreate_NoFastClear = 1 << 25,
 };
 
 /**
diff --git a/Engine/Source/Runtime/Renderer/Private/SceneOcclusion.cpp b/Engine/Source/Runtime/Renderer/Private/SceneOcclusion.cpp
--- a/Engine/Source/Runtime/Renderer/Private/SceneOcclusion.cpp
+++ b/Engine/Source/Runtime/Renderer/Private/SceneOcclusion.cpp
@@ -939,7 +939,7 @@ void BuildHZB( FRHICommandListImmediate& RHICmdList, FViewInfo& View )
 	const FIntPoint HZBSize( 512, 256 );
 	const uint32 NumMips = 8;
 
-	FPooledRenderTargetDesc Desc( FPooledRenderTargetDesc::Create2DDesc( HZBSize, PF_R16F, TexCreate_None, TexCreate_RenderTargetable | TexCreate_ShaderResource, false, NumMips ) );
+	FPooledRenderTargetDesc Desc(FPooledRenderTargetDesc::Create2DDesc(HZBSize, PF_R16F, TexCreate_None, TexCreate_RenderTargetable | TexCreate_ShaderResource | TexCreate_NoFastClear, false, NumMips));
 	Desc.Flags |= TexCreate_FastVRAM;
 	GRenderTargetPool.FindFreeElement( Desc, View.HZB, TEXT("HZB") );
 	
diff --git a/Engine/Source/Runtime/Renderer/Private/SystemTextures.cpp b/Engine/Source/Runtime/Renderer/Private/SystemTextures.cpp
--- a/Engine/Source/Runtime/Renderer/Private/SystemTextures.cpp
+++ b/Engine/Source/Runtime/Renderer/Private/SystemTextures.cpp
@@ -32,7 +32,7 @@ void FSystemTextures::InitializeTextures(FRHICommandListImmediate& RHICmdList, E
 	{
 		// Create a dummy white texture.
 		{
-			FPooledRenderTargetDesc Desc(FPooledRenderTargetDesc::Create2DDesc(FIntPoint(1, 1), PF_B8G8R8A8, TexCreate_HideInVisualizeTexture, TexCreate_RenderTargetable, false));
+			FPooledRenderTargetDesc Desc(FPooledRenderTargetDesc::Create2DDesc(FIntPoint(1, 1), PF_B8G8R8A8, TexCreate_HideInVisualizeTexture, TexCreate_RenderTargetable | TexCreate_NoFastClear, false));
 			GRenderTargetPool.FindFreeElement(Desc, WhiteDummy, TEXT("WhiteDummy"));
 
 			SetRenderTarget(RHICmdList, WhiteDummy->GetRenderTargetItem().TargetableTexture, FTextureRHIRef());
@@ -42,7 +42,7 @@ void FSystemTextures::InitializeTextures(FRHICommandListImmediate& RHICmdList, E
 
 		// Create a dummy black texture.
 		{
-			FPooledRenderTargetDesc Desc(FPooledRenderTargetDesc::Create2DDesc(FIntPoint(1, 1), PF_B8G8R8A8, TexCreate_HideInVisualizeTexture, TexCreate_RenderTargetable, false));
+			FPooledRenderTargetDesc Desc(FPooledRenderTargetDesc::Create2DDesc(FIntPoint(1, 1), PF_B8G8R8A8, TexCreate_HideInVisualizeTexture, TexCreate_RenderTargetable | TexCreate_NoFastClear, false));
 			GRenderTargetPool.FindFreeElement(Desc, BlackDummy, TEXT("BlackDummy"));
 
 			SetRenderTarget(RHICmdList, BlackDummy->GetRenderTargetItem().TargetableTexture, FTextureRHIRef());
@@ -52,7 +52,7 @@ void FSystemTextures::InitializeTextures(FRHICommandListImmediate& RHICmdList, E
 
 		// Create the PerlinNoiseGradient texture
 		{
-			FPooledRenderTargetDesc Desc(FPooledRenderTargetDesc::Create2DDesc(FIntPoint(128, 128), PF_B8G8R8A8, TexCreate_HideInVisualizeTexture, TexCreate_None, false));
+			FPooledRenderTargetDesc Desc(FPooledRenderTargetDesc::Create2DDesc(FIntPoint(128, 128), PF_B8G8R8A8, TexCreate_HideInVisualizeTexture, TexCreate_None | TexCreate_NoFastClear, false));
 			GRenderTargetPool.FindFreeElement(Desc, PerlinNoiseGradient, TEXT("PerlinNoiseGradient"));
 			// Write the contents of the texture.
 			uint32 DestStride;
@@ -83,7 +83,7 @@ void FSystemTextures::InitializeTextures(FRHICommandListImmediate& RHICmdList, E
 
 		if (!GSupportsShaderFramebufferFetch && GPixelFormats[PF_FloatRGBA].Supported)
 		{
-			FPooledRenderTargetDesc Desc(FPooledRenderTargetDesc::Create2DDesc(FIntPoint(1, 1), PF_FloatRGBA, TexCreate_HideInVisualizeTexture, TexCreate_RenderTargetable, false));
+			FPooledRenderTargetDesc Desc(FPooledRenderTargetDesc::Create2DDesc(FIntPoint(1, 1), PF_FloatRGBA, TexCreate_HideInVisualizeTexture, TexCreate_RenderTargetable | TexCreate_NoFastClear, false));
 			GRenderTargetPool.FindFreeElement(Desc, MaxFP16Depth, TEXT("MaxFP16Depth"));
 
 			SetRenderTarget(RHICmdList, MaxFP16Depth->GetRenderTargetItem().TargetableTexture, FTextureRHIRef());
@@ -232,7 +232,7 @@ void FSystemTextures::InitializeTextures(FRHICommandListImmediate& RHICmdList, E
 		}
 
 		{
-			FPooledRenderTargetDesc Desc(FPooledRenderTargetDesc::Create2DDesc(FIntPoint(64, 64), PF_B8G8R8A8, TexCreate_HideInVisualizeTexture, TexCreate_None, false));
+			FPooledRenderTargetDesc Desc(FPooledRenderTargetDesc::Create2DDesc(FIntPoint(64, 64), PF_B8G8R8A8, TexCreate_HideInVisualizeTexture, TexCreate_None | TexCreate_NoFastClear, false));
 			GRenderTargetPool.FindFreeElement(Desc, SSAORandomization, TEXT("SSAORandomization"));
 			// Write the contents of the texture.
 			uint32 DestStride;
