UE-7435 - Don't gather mesh elements for shadows that aren't visible.

diff --git a/Engine/Source/Runtime/Renderer/Private/ShadowSetup.cpp b/Engine/Source/Runtime/Renderer/Private/ShadowSetup.cpp
--- a/Engine/Source/Runtime/Renderer/Private/ShadowSetup.cpp
+++ b/Engine/Source/Runtime/Renderer/Private/ShadowSetup.cpp
@@ -888,60 +888,70 @@ void FProjectedShadowInfo::GatherDynamicMeshElements(FSceneRenderer& Renderer, F
 	{
 		// Choose an arbitrary view where this shadow's subject is relevant.
 		FViewInfo* FoundView = NULL;
+		bool bVisibleInAnyView = false;
+
 		for (int32 ViewIndex = 0;ViewIndex < Renderer.Views.Num();ViewIndex++)
 		{
 			FViewInfo* CheckView = &Renderer.Views[ViewIndex];
 			const FVisibleLightViewInfo& VisibleLightViewInfo = CheckView->VisibleLightInfos[LightSceneInfo->Id];
 			FPrimitiveViewRelevance ViewRel = VisibleLightViewInfo.ProjectedShadowViewRelevanceMap[ShadowId];
+
+			bVisibleInAnyView |= VisibleLightViewInfo.ProjectedShadowVisibilityMap[ShadowId];
+
 			if (ViewRel.bShadowRelevance)
 			{
 				FoundView = CheckView;
 				break;
 			}
 		}
-		check(FoundView && IsInRenderingThread());
 
-		// Backup properties of the view that we will override
-		FMatrix OriginalViewMatrix = FoundView->ViewMatrices.ViewMatrix;
+		// Shadow must be visible in at least one view
+		if (bVisibleInAnyView)
+		{
+			check(FoundView && IsInRenderingThread());
 
-		// Override the view matrix so that billboarding primitives will be aligned to the light
-		FoundView->ViewMatrices.ViewMatrix = ShadowViewMatrix;
+			// Backup properties of the view that we will override
+			FMatrix OriginalViewMatrix = FoundView->ViewMatrices.ViewMatrix;
 
-		// Prevent materials from getting overridden during shadow casting in viewmodes like lighting only
-		// Lighting only should only affect the material used with direct lighting, not the indirect lighting
-		FoundView->bForceShowMaterials = true;
+			// Override the view matrix so that billboarding primitives will be aligned to the light
+			FoundView->ViewMatrices.ViewMatrix = ShadowViewMatrix;
 
-		ReusedViewsArray[0] = FoundView;
+			// Prevent materials from getting overridden during shadow casting in viewmodes like lighting only
+			// Lighting only should only affect the material used with direct lighting, not the indirect lighting
+			FoundView->bForceShowMaterials = true;
 
-		check(!FoundView->ViewMatrices.GetDynamicMeshElementsShadowCullFrustum);
+			ReusedViewsArray[0] = FoundView;
 
-		int32 Disable = 0; //CVarDisableCullShadows.GetValueOnRenderThread();
-		FConvexVolume NoCull;
+			check(!FoundView->ViewMatrices.GetDynamicMeshElementsShadowCullFrustum);
 
-		if (IsWholeSceneDirectionalShadow())
-		{
-			FoundView->ViewMatrices.PreShadowTranslation = FVector(0,0,0);
-			FoundView->ViewMatrices.GetDynamicMeshElementsShadowCullFrustum = (Disable & 1) ? &NoCull : &CascadeSettings.ShadowBoundsAccurate;
-			GatherDynamicMeshElementsArray(FoundView, Renderer, SubjectPrimitives, DynamicSubjectMeshElements, ReusedViewsArray);
-			FoundView->ViewMatrices.PreShadowTranslation = PreShadowTranslation;
-		}
-		else
-		{
-			FoundView->ViewMatrices.PreShadowTranslation = PreShadowTranslation;
-			FoundView->ViewMatrices.GetDynamicMeshElementsShadowCullFrustum = (Disable & 1) ? &NoCull : &CasterFrustum;
-			GatherDynamicMeshElementsArray(FoundView, Renderer, SubjectPrimitives, DynamicSubjectMeshElements, ReusedViewsArray);
-		}
+			int32 Disable = 0; //CVarDisableCullShadows.GetValueOnRenderThread();
+			FConvexVolume NoCull;
+
+			if (IsWholeSceneDirectionalShadow())
+			{
+				FoundView->ViewMatrices.PreShadowTranslation = FVector(0,0,0);
+				FoundView->ViewMatrices.GetDynamicMeshElementsShadowCullFrustum = (Disable & 1) ? &NoCull : &CascadeSettings.ShadowBoundsAccurate;
+				GatherDynamicMeshElementsArray(FoundView, Renderer, SubjectPrimitives, DynamicSubjectMeshElements, ReusedViewsArray);
+				FoundView->ViewMatrices.PreShadowTranslation = PreShadowTranslation;
+			}
+			else
+			{
+				FoundView->ViewMatrices.PreShadowTranslation = PreShadowTranslation;
+				FoundView->ViewMatrices.GetDynamicMeshElementsShadowCullFrustum = (Disable & 1) ? &NoCull : &CasterFrustum;
+				GatherDynamicMeshElementsArray(FoundView, Renderer, SubjectPrimitives, DynamicSubjectMeshElements, ReusedViewsArray);
+			}
 
-		FoundView->ViewMatrices.GetDynamicMeshElementsShadowCullFrustum = (Disable & 2) ? &NoCull : &ReceiverFrustum;
-		GatherDynamicMeshElementsArray(FoundView, Renderer, ReceiverPrimitives, DynamicReceiverMeshElements, ReusedViewsArray);
+			FoundView->ViewMatrices.GetDynamicMeshElementsShadowCullFrustum = (Disable & 2) ? &NoCull : &ReceiverFrustum;
+			GatherDynamicMeshElementsArray(FoundView, Renderer, ReceiverPrimitives, DynamicReceiverMeshElements, ReusedViewsArray);
 
-		FoundView->ViewMatrices.GetDynamicMeshElementsShadowCullFrustum = (Disable & 4) ? &NoCull : &CasterFrustum;
-		GatherDynamicMeshElementsArray(FoundView, Renderer, SubjectTranslucentPrimitives, DynamicSubjectTranslucentMeshElements, ReusedViewsArray);
+			FoundView->ViewMatrices.GetDynamicMeshElementsShadowCullFrustum = (Disable & 4) ? &NoCull : &CasterFrustum;
+			GatherDynamicMeshElementsArray(FoundView, Renderer, SubjectTranslucentPrimitives, DynamicSubjectTranslucentMeshElements, ReusedViewsArray);
 
-		FoundView->ViewMatrices.GetDynamicMeshElementsShadowCullFrustum = nullptr;
+			FoundView->ViewMatrices.GetDynamicMeshElementsShadowCullFrustum = nullptr;
 
-		FoundView->bForceShowMaterials = false;
-		FoundView->ViewMatrices.ViewMatrix = OriginalViewMatrix;
+			FoundView->bForceShowMaterials = false;
+			FoundView->ViewMatrices.ViewMatrix = OriginalViewMatrix;
+		}
 	}
 }
 
