UE4 - more careful load balancing on rhicmdlists

diff --git a/Engine/Source/Runtime/Core/Public/Async/ParallelFor.h b/Engine/Source/Runtime/Core/Public/Async/ParallelFor.h
--- a/Engine/Source/Runtime/Core/Public/Async/ParallelFor.h
+++ b/Engine/Source/Runtime/Core/Public/Async/ParallelFor.h
@@ -15,7 +15,7 @@
 	*	@param bForceSingleThread; Mostly used for testing, if true, run single threaded instead.
 	*	Notes: Please add stats around to calls to parallel for and within your lambda as appropriate. Do not clog the task graph with long running tasks or tasks that block.
 **/
-void ParallelFor(int32 Num, TFunctionRef<void(int32)> Body, bool bForceSingleThread = false)
+inline void ParallelFor(int32 Num, TFunctionRef<void(int32)> Body, bool bForceSingleThread = false)
 {
 	// struct to hold the working data; this outlives the ParallelFor call; lifetime is controlled by a shared pointer
 	struct FParallelForData
diff --git a/Engine/Source/Runtime/Core/Public/Async/TaskGraphInterfaces.h b/Engine/Source/Runtime/Core/Public/Async/TaskGraphInterfaces.h
--- a/Engine/Source/Runtime/Core/Public/Async/TaskGraphInterfaces.h
+++ b/Engine/Source/Runtime/Core/Public/Async/TaskGraphInterfaces.h
@@ -591,6 +591,12 @@ public:
 			new ((void *)&Owner->TaskStorage) TTask(Forward<T1>(Arg1), Forward<T2>(Arg2), Forward<T3>(Arg3), Forward<T4>(Arg4), Forward<T5>(Arg5), Forward<T6>(Arg6), Forward<T7>(Arg7), Forward<T8>(Arg8), Forward<T9>(Arg9));
 			return Owner->Setup(Prerequisites, CurrentThreadIfKnown);
 		}		
+		template<typename T1, typename T2, typename T3, typename T4, typename T5, typename T6, typename T7, typename T8, typename T9, typename T10>
+		FGraphEventRef ConstructAndDispatchWhenReady(T1&& Arg1, T2&& Arg2, T3&& Arg3, T4&& Arg4, T5&& Arg5, T6&& Arg6, T7&& Arg7, T8&& Arg8, T9&& Arg9, T10&& Arg10)
+		{
+			new ((void *)&Owner->TaskStorage) TTask(Forward<T1>(Arg1), Forward<T2>(Arg2), Forward<T3>(Arg3), Forward<T4>(Arg4), Forward<T5>(Arg5), Forward<T6>(Arg6), Forward<T7>(Arg7), Forward<T8>(Arg8), Forward<T9>(Arg9), Forward<T10>(Arg10));
+			return Owner->Setup(Prerequisites, CurrentThreadIfKnown);
+		}		
 	#endif
 
 	#if PLATFORM_COMPILER_HAS_VARIADIC_TEMPLATES
diff --git a/Engine/Source/Runtime/RHI/Private/RHICommandList.cpp b/Engine/Source/Runtime/RHI/Private/RHICommandList.cpp
--- a/Engine/Source/Runtime/RHI/Private/RHICommandList.cpp
+++ b/Engine/Source/Runtime/RHI/Private/RHICommandList.cpp
@@ -38,16 +38,18 @@ static TAutoConsoleVariable<int32> CVarRHICmdUseDeferredContexts(
 	TEXT("r.RHICmdUseDeferredContexts"),
 	1,
 	TEXT("True to use deferred contexts to parallelize command list execution. Only available on some RHIs."));
-static TAutoConsoleVariable<int32> CVarRHIDeferredContextWidth(
-	TEXT("r.RHIDeferredContextWidth"),
-	99999,
-	TEXT("Number of pieces to break a parallel translate task into. A very large number is one task per original cmd list. Only available on some RHIs."));
 
 static TAutoConsoleVariable<int32> CVarRHICmdFlushOnQueueParallelSubmit(
 	TEXT("r.RHICmdFlushOnQueueParallelSubmit"),
 	0,
 	TEXT("Wait for completion of parallel commandlists immediately after submitting.  Debug feature. Only available on some RHIs."));
 
+static TAutoConsoleVariable<int32> CVarRHIMergeSmallDeferredContexts(
+	TEXT("r.RHIMergeSmallDeferredContexts"),
+	1,
+	TEXT("When it can be determined, merge small parallel translate tasks based on r.RHICmdMinDrawsPerParallelCmdList."));
+
+
 RHI_API FRHICommandListExecutor GRHICommandList;
 
 static FGraphEventArray AllOutstandingTasks;
@@ -671,7 +673,7 @@ struct FRHICommandWaitForAndSubmitSubList : public FRHICommand<FRHICommandWaitFo
 	}
 };
 
-void FRHICommandListBase::QueueParallelAsyncCommandListSubmit(FGraphEventRef* AnyThreadCompletionEvents, FRHICommandList** CmdLists, int32 Num)
+void FRHICommandListBase::QueueParallelAsyncCommandListSubmit(FGraphEventRef* AnyThreadCompletionEvents, FRHICommandList** CmdLists, int32* NumDrawsIfKnown, int32 Num, int32 MinDrawsPerTranslate, bool bSpewMerge)
 {
 	check(IsInRenderingThread() && IsImmediate() && Num);
 	if (GRHIThread)
@@ -695,61 +697,92 @@ void FRHICommandListBase::QueueParallelAsyncCommandListSubmit(FGraphEventRef* An
 		}
 		if (ContextContainer)
 		{
-			int32 EffectiveThreads = FMath::Min<int32>(Num, CVarRHIDeferredContextWidth.GetValueOnRenderThread());
+			bool bMerge = !!CVarRHIMergeSmallDeferredContexts.GetValueOnRenderThread();
 
+			int32 EffectiveThreads = 0;
 			int32 Start = 0;
-			if (EffectiveThreads)
+			// this is pretty silly but we need to know the number of jobs in advance, so we run the merge logic twice
+			while (Start < Num)
 			{
+				int32 Last = Start;
+				int32 DrawCnt = NumDrawsIfKnown[Start];
 
-				int32 NumPer = Num / EffectiveThreads;
-				int32 Extra = Num - NumPer * EffectiveThreads;
-
-				for (int32 ThreadIndex = 0; ThreadIndex < EffectiveThreads; ThreadIndex++)
+				if (bMerge && DrawCnt >= 0)
 				{
-					int32 Last = Start + (NumPer - 1) + (ThreadIndex < Extra);
-					check(Last >= Start);
-
-					if (!ContextContainer)
+					while (Last < Num - 1 && NumDrawsIfKnown[Last + 1] >= 0 && DrawCnt + NumDrawsIfKnown[Last + 1] <= MinDrawsPerTranslate)
 					{
-						ContextContainer = RHIGetCommandContextContainer();
+						Last++;
+						DrawCnt += NumDrawsIfKnown[Last];
 					}
-					check(ContextContainer);
-
-					FGraphEventArray Prereq;
-					FRHICommandListBase** RHICmdLists = (FRHICommandListBase**)Alloc(sizeof(FRHICommandListBase*) * (1 + Last - Start), ALIGNOF(FRHICommandListBase*));
-					for (int32 Index = Start; Index <= Last; Index++)
+				}
+				check(Last >= Start);
+				Start = Last + 1;
+				EffectiveThreads++;
+			} 
+			Start = 0;
+			int32 ThreadIndex = 0;
+			   
+			while (Start < Num)       
+			{    
+				int32 Last = Start;
+				int32 DrawCnt = NumDrawsIfKnown[Start];
+				int32 TotalMem = bSpewMerge ? CmdLists[Start]->GetUsedMemory() : 0;   // the memory is only accurate if we are spewing because otherwise it isn't done yet!
+
+				if (bMerge && DrawCnt >= 0)
+				{
+					while (Last < Num - 1 && NumDrawsIfKnown[Last + 1] >= 0 && DrawCnt + NumDrawsIfKnown[Last + 1] <= MinDrawsPerTranslate)
 					{
-						FGraphEventRef& AnyThreadCompletionEvent = AnyThreadCompletionEvents[Index];
-						FRHICommandList* CmdList = CmdLists[Index];
-						RHICmdLists[Index - Start] = CmdList;
-						if (AnyThreadCompletionEvent.GetReference())
-						{
-							Prereq.Add(AnyThreadCompletionEvent);
-							AllOutstandingTasks.Add(AnyThreadCompletionEvent);
-							WaitOutstandingTasks.Add(AnyThreadCompletionEvent);
-						}
+						Last++;
+						DrawCnt += NumDrawsIfKnown[Last];
+						TotalMem += bSpewMerge ? CmdLists[Start]->GetUsedMemory() : 0;   // the memory is only accurate if we are spewing because otherwise it isn't done yet!
 					}
-					FGraphEventRef TranslateCompletionEvent = TGraphTask<FParallelTranslateCommandList>::CreateTask(&Prereq, ENamedThreads::RenderThread).ConstructAndDispatchWhenReady(&RHICmdLists[0], 1 + Last - Start, ContextContainer);
+				}
+
+				check(Last >= Start);
+
+				if (!ContextContainer) 
+				{
+					ContextContainer = RHIGetCommandContextContainer();
+				}
+				check(ContextContainer);
 
-					AllOutstandingTasks.Add(TranslateCompletionEvent);
-					new (AllocCommand<FRHICommandWaitForAndSubmitSubListParallel>()) FRHICommandWaitForAndSubmitSubListParallel(TranslateCompletionEvent, ContextContainer, EffectiveThreads, ThreadIndex);
-					if (GRHIThread)
+				FGraphEventArray Prereq;
+				FRHICommandListBase** RHICmdLists = (FRHICommandListBase**)Alloc(sizeof(FRHICommandListBase*) * (1 + Last - Start), ALIGNOF(FRHICommandListBase*));
+				for (int32 Index = Start; Index <= Last; Index++)
+				{
+					FGraphEventRef& AnyThreadCompletionEvent = AnyThreadCompletionEvents[Index];
+					FRHICommandList* CmdList = CmdLists[Index];
+					RHICmdLists[Index - Start] = CmdList;
+					if (AnyThreadCompletionEvent.GetReference())
 					{
-						FRHICommandListExecutor::GetImmediateCommandList().ImmediateFlush(EImmediateFlushType::DispatchToRHIThread); // we don't want stuff after the async cmd list to be bundled with it
+						Prereq.Add(AnyThreadCompletionEvent);
+						AllOutstandingTasks.Add(AnyThreadCompletionEvent);
+						WaitOutstandingTasks.Add(AnyThreadCompletionEvent);
 					}
-
-					ContextContainer = nullptr;
-					Start = Last + 1;
 				}
+				UE_CLOG(bSpewMerge, LogTemp, Display, TEXT("Parallel translate %d->%d    %dKB mem   %d draws (-1 = unknown)"), Start, Last, FMath::DivideAndRoundUp(TotalMem, 1024), DrawCnt);
 
-#if !UE_BUILD_SHIPPING
-				if (CVarRHICmdFlushOnQueueParallelSubmit.GetValueOnRenderThread())
+				FGraphEventRef TranslateCompletionEvent = TGraphTask<FParallelTranslateCommandList>::CreateTask(&Prereq, ENamedThreads::RenderThread).ConstructAndDispatchWhenReady(&RHICmdLists[0], 1 + Last - Start, ContextContainer);
+
+				AllOutstandingTasks.Add(TranslateCompletionEvent);
+				new (AllocCommand<FRHICommandWaitForAndSubmitSubListParallel>()) FRHICommandWaitForAndSubmitSubListParallel(TranslateCompletionEvent, ContextContainer, EffectiveThreads, ThreadIndex++);
+				if (GRHIThread)
 				{
-					FRHICommandListExecutor::GetImmediateCommandList().ImmediateFlush(EImmediateFlushType::FlushRHIThread);
+					FRHICommandListExecutor::GetImmediateCommandList().ImmediateFlush(EImmediateFlushType::DispatchToRHIThread); // we don't want stuff after the async cmd list to be bundled with it
 				}
-#endif
-				return;
+
+				ContextContainer = nullptr;
+				Start = Last + 1;
 			}
+			check(EffectiveThreads == ThreadIndex);
+
+#if !UE_BUILD_SHIPPING
+			if (CVarRHICmdFlushOnQueueParallelSubmit.GetValueOnRenderThread())
+			{
+				FRHICommandListExecutor::GetImmediateCommandList().ImmediateFlush(EImmediateFlushType::FlushRHIThread);
+			}
+#endif
+			return;
 		}
 	}
 	for (int32 Index = 0; Index < Num; Index++)
diff --git a/Engine/Source/Runtime/RHI/Public/RHICommandList.h b/Engine/Source/Runtime/RHI/Public/RHICommandList.h
--- a/Engine/Source/Runtime/RHI/Public/RHICommandList.h
+++ b/Engine/Source/Runtime/RHI/Public/RHICommandList.h
@@ -63,7 +63,7 @@ public:
 	inline bool IsImmediate();
 	const int32 GetUsedMemory() const;
 	void QueueAsyncCommandListSubmit(FGraphEventRef& AnyThreadCompletionEvent, class FRHICommandList* CmdList);
-	void QueueParallelAsyncCommandListSubmit(FGraphEventRef* AnyThreadCompletionEvents, class FRHICommandList** CmdLists, int32 Num);
+	void QueueParallelAsyncCommandListSubmit(FGraphEventRef* AnyThreadCompletionEvents, class FRHICommandList** CmdLists, int32* NumDrawsIfKnown, int32 Num, int32 MinDrawsPerTranslate, bool bSpewMerge);
 	void QueueRenderThreadCommandListSubmit(FGraphEventRef& RenderThreadCompletionEvent, class FRHICommandList* CmdList);
 	void QueueCommandListSubmit(class FRHICommandList* CmdList);
 	void WaitForTasks(bool bKnownToBeComplete = false);
diff --git a/Engine/Source/Runtime/Renderer/Private/DynamicPrimitiveDrawing.inl b/Engine/Source/Runtime/Renderer/Private/DynamicPrimitiveDrawing.inl
--- a/Engine/Source/Runtime/Renderer/Private/DynamicPrimitiveDrawing.inl
+++ b/Engine/Source/Runtime/Renderer/Private/DynamicPrimitiveDrawing.inl
@@ -113,7 +113,7 @@ void DrawViewElementsParallel(
 
 	{
 		int32 NumPrims = ViewMeshElementList.Num();
-		int32 EffectiveThreads = FMath::Min<int32>(NumPrims, ParallelCommandListSet.Width);
+		int32 EffectiveThreads = FMath::Min<int32>(FMath::DivideAndRoundUp(NumPrims, ParallelCommandListSet.MinDrawsPerCommandList), ParallelCommandListSet.Width);
 
 		int32 Start = 0;
 		if (EffectiveThreads)
@@ -134,7 +134,7 @@ void DrawViewElementsParallel(
 					FGraphEventRef AnyThreadCompletionEvent = TGraphTask<FDrawViewElementsAnyThreadTask<DrawingPolicyFactoryType> >::CreateTask(nullptr, ENamedThreads::RenderThread)
 						.ConstructAndDispatchWhenReady(*CmdList, ParallelCommandListSet.View, DrawingContext, DPGIndex, bPreFog, Start, Last);
 
-					ParallelCommandListSet.AddParallelCommandList(CmdList, AnyThreadCompletionEvent);
+					ParallelCommandListSet.AddParallelCommandList(CmdList, AnyThreadCompletionEvent, Last - Start + 1);
 				}
 				Start = Last + 1;
 			}
diff --git a/Engine/Source/Runtime/Renderer/Private/SceneRendering.cpp b/Engine/Source/Runtime/Renderer/Private/SceneRendering.cpp
--- a/Engine/Source/Runtime/Renderer/Private/SceneRendering.cpp
+++ b/Engine/Source/Runtime/Renderer/Private/SceneRendering.cpp
@@ -122,6 +122,26 @@ FRHICommandList* FParallelCommandListSet::AllocCommandList()
 	return new FRHICommandList;
 }
 
+static TAutoConsoleVariable<int32> CVarRHICmdSpewParallelListBalance(
+	TEXT("r.RHICmdSpewParallelListBalance"),
+	0,
+	TEXT("For debugging, spews the size of the parallel command lists. This stalls and otherwise wrecks performance."));
+
+static TAutoConsoleVariable<int32> CVarRHICmdBalanceParallelLists(
+	TEXT("r.RHICmdBalanceParallelLists"),
+	1,
+	TEXT("If >0 preprocess the drawlists to try to balance the load equally among the command lists."));
+
+static TAutoConsoleVariable<int32> CVarRHICmdMinCmdlistForParallelSubmit(
+	TEXT("r.RHICmdMinCmdlistForParallelSubmit"),
+	2,
+	TEXT("Minimum number of parallel translate command lists to submit. If there are fewer than this number, they just run on the RHI thread and immediate context."));
+
+static TAutoConsoleVariable<int32> CVarRHICmdMinDrawsPerParallelCmdList(
+	TEXT("r.RHICmdMinDrawsPerParallelCmdList"),
+	32,
+	TEXT("The minimum number of draws per cmdlist. If the total number of draws is less than this, then no parallel work will be done at all. This can't always be honored or done correctly. More effective with RHICmdBalanceParallelLists."));
+
 FParallelCommandListSet::FParallelCommandListSet(const FViewInfo& InView, FRHICommandList& InParentCmdList, bool* InOutDirty, bool bInParallelExecute)
 	: View(InView)
 	, ParentCmdList(InParentCmdList)
@@ -130,22 +150,63 @@ FParallelCommandListSet::FParallelCommandListSet(const FViewInfo& InView, FRHICo
 	, bParallelExecute(GRHISupportsParallelRHIExecute && bInParallelExecute)
 {
 	Width = CVarRHICmdWidth.GetValueOnRenderThread();
+	MinDrawsPerCommandList = CVarRHICmdMinDrawsPerParallelCmdList.GetValueOnRenderThread();
+	bSpewBalance = !!CVarRHICmdSpewParallelListBalance.GetValueOnRenderThread();
+	bBalanceCommands = !!CVarRHICmdBalanceParallelLists.GetValueOnRenderThread();
 	CommandLists.Reserve(Width * 8);
 	Events.Reserve(Width * 8);
-
+	NumDrawsIfKnown.Reserve(Width * 8);
 }
 
 void FParallelCommandListSet::Dispatch()
 {
 	check(CommandLists.Num() == Events.Num());
-	if (bParallelExecute && CommandLists.Num())
+	if (bSpewBalance)
+	{
+		// finish them all
+		for (auto& Event : Events)
+		{
+			FTaskGraphInterface::Get().WaitUntilTaskCompletes(Event, ENamedThreads::RenderThread_Local);
+		}
+		// spew sizes
+		int32 Index = 0;
+		for (auto CmdList : CommandLists)
+		{
+			UE_LOG(LogTemp, Display, TEXT("CmdList %2d/%2d  : %8dKB"), Index, CommandLists.Num(), (CmdList->GetUsedMemory() + 1023) / 1024);
+			Index++;
+		}
+	}
+	bool bActuallyDoParallelTranslate = bParallelExecute && CommandLists.Num() >= CVarRHICmdMinCmdlistForParallelSubmit.GetValueOnRenderThread();
+	if (bActuallyDoParallelTranslate)
+	{
+		int32 Total = 0;
+		bool bIndeterminate = false;
+		for (int32 Count : NumDrawsIfKnown)
+		{
+			if (Count < 0)
+			{
+				bIndeterminate = true;
+				break; // can't determine how many are in this one; assume we should run parallel translate
+			}
+			Total += Count;
+		}
+		if (!bIndeterminate && Total < MinDrawsPerCommandList)
+		{
+			UE_CLOG(bSpewBalance, LogTemp, Display, TEXT("Disabling parallel translate because the number of draws is known to be small."));
+			bActuallyDoParallelTranslate = false;
+		}
+	}
+
+	if (bActuallyDoParallelTranslate)
 	{
+		UE_CLOG(bSpewBalance, LogTemp, Display, TEXT("%d cmdlists for parallel translate"), CommandLists.Num());
 		check(GRHISupportsParallelRHIExecute);
-		ParentCmdList.QueueParallelAsyncCommandListSubmit(&Events[0], &CommandLists[0], CommandLists.Num());
+		ParentCmdList.QueueParallelAsyncCommandListSubmit(&Events[0], &CommandLists[0], &NumDrawsIfKnown[0], CommandLists.Num(), (MinDrawsPerCommandList * 4) / 3, bSpewBalance);
 		SetStateOnCommandList(ParentCmdList);
 	}
 	else
 	{
+		UE_CLOG(bSpewBalance, LogTemp, Display, TEXT("%d cmdlists (no parallel translate desired)"), CommandLists.Num());
 		for (int32 Index = 0; Index < CommandLists.Num(); Index++)
 		{
 			ParentCmdList.QueueAsyncCommandListSubmit(Events[Index], CommandLists[Index]);
@@ -167,11 +228,12 @@ FRHICommandList* FParallelCommandListSet::NewParallelCommandList()
 	return Result;
 }
 
-void FParallelCommandListSet::AddParallelCommandList(FRHICommandList* CmdList, FGraphEventRef& CompletionEvent)
+void FParallelCommandListSet::AddParallelCommandList(FRHICommandList* CmdList, FGraphEventRef& CompletionEvent, int32 InNumDrawsIfKnown)
 {
 	check(CommandLists.Num() == Events.Num());
 	CommandLists.Add(CmdList);
 	Events.Add(CompletionEvent);
+	NumDrawsIfKnown.Add(InNumDrawsIfKnown);
 }
 
 
diff --git a/Engine/Source/Runtime/Renderer/Private/SceneRendering.h b/Engine/Source/Runtime/Renderer/Private/SceneRendering.h
--- a/Engine/Source/Runtime/Renderer/Private/SceneRendering.h
+++ b/Engine/Source/Runtime/Renderer/Private/SceneRendering.h
@@ -311,14 +311,19 @@ class FParallelCommandListSet
 {
 public:
 	const FViewInfo& View;
-	int32 Width;
 	FRHICommandList& ParentCmdList;
+	int32 Width;
+	int32 MinDrawsPerCommandList;
+	bool bBalanceCommands;
+	bool bSpewBalance;
 private:
 	bool OutDirtyIfIgnored;
 public:
 	bool& OutDirty;
 	TArray<FRHICommandList*,SceneRenderingAllocator> CommandLists;
 	TArray<FGraphEventRef,SceneRenderingAllocator> Events;
+	// number of draws in this commandlist if known, -1 if not known. Overestimates are better than nothing.
+	TArray<int32,SceneRenderingAllocator> NumDrawsIfKnown;
 protected:
 	//this must be called by deriving classes virtual destructor because it calls the virtual SetStateOnCommandList.
 	//C++ will not do dynamic dispatch of virtual calls from destructors so we can't call it in the base class.
@@ -328,8 +333,12 @@ protected:
 public:
 	FParallelCommandListSet(const FViewInfo& InView, FRHICommandList& InParentCmdList, bool* InOutDirty, bool bInParallelExecute);
 	virtual ~FParallelCommandListSet();
+	int32 NumParallelCommandLists() const
+	{
+		return CommandLists.Num();
+	}
 	FRHICommandList* NewParallelCommandList();
-	void AddParallelCommandList(FRHICommandList* CmdList, FGraphEventRef& CompletionEvent);	
+	void AddParallelCommandList(FRHICommandList* CmdList, FGraphEventRef& CompletionEvent, int32 InNumDrawsIfKnown = -1);	
 
 	virtual void SetStateOnCommandList(FRHICommandList& CmdList)
 	{
diff --git a/Engine/Source/Runtime/Renderer/Private/ShadowRendering.cpp b/Engine/Source/Runtime/Renderer/Private/ShadowRendering.cpp
--- a/Engine/Source/Runtime/Renderer/Private/ShadowRendering.cpp
+++ b/Engine/Source/Runtime/Renderer/Private/ShadowRendering.cpp
@@ -1581,15 +1581,16 @@ void FProjectedShadowInfo::RenderDepthInner(FRHICommandList& RHICmdList, FSceneR
 				FGraphEventRef AnyThreadCompletionEvent = TGraphTask<FDrawShadowMeshElementsThreadTask>::CreateTask(nullptr, ENamedThreads::RenderThread)
 					.ConstructAndDispatchWhenReady(*this, *CmdList, *FoundView, bReflectiveShadowmap && !CascadeSettings.bOnePassPointLightShadow);
 
-				ParallelCommandListSet.AddParallelCommandList(CmdList, AnyThreadCompletionEvent);
+				ParallelCommandListSet.AddParallelCommandList(CmdList, AnyThreadCompletionEvent, SubjectMeshElements.Num());
 			}
+			if (DynamicSubjectMeshElements.Num())
 			{
 				FRHICommandList* CmdList = ParallelCommandListSet.NewParallelCommandList();
 
 				FGraphEventRef AnyThreadCompletionEvent = TGraphTask<FRenderDepthDynamicThreadTask>::CreateTask(nullptr, ENamedThreads::RenderThread)
 					.ConstructAndDispatchWhenReady(*this, *CmdList, *FoundView, SceneRenderer);
 
-				ParallelCommandListSet.AddParallelCommandList(CmdList, AnyThreadCompletionEvent);
+				ParallelCommandListSet.AddParallelCommandList(CmdList, AnyThreadCompletionEvent, DynamicSubjectMeshElements.Num());
 			}
 		}
 	}
diff --git a/Engine/Source/Runtime/Renderer/Private/StaticMeshDrawList.inl b/Engine/Source/Runtime/Renderer/Private/StaticMeshDrawList.inl
--- a/Engine/Source/Runtime/Renderer/Private/StaticMeshDrawList.inl
+++ b/Engine/Source/Runtime/Renderer/Private/StaticMeshDrawList.inl
@@ -4,11 +4,11 @@
 	StaticMeshDrawList.inl: Static mesh draw list implementation.
 =============================================================================*/
 
-#ifndef __STATICMESHDRAWLIST_INL__
-#define __STATICMESHDRAWLIST_INL__
+#pragma once
 
 #include "RHICommandList.h"
 #include "EngineStats.h"
+#include "ParallelFor.h"
 
 // Expensive
 #define PER_MESH_DRAW_STATS 0
@@ -221,6 +221,7 @@ bool TStaticMeshDrawList<DrawingPolicyType>::DrawVisible(
 	)
 {
 	bool bDirty = false;
+	STAT(int32 StatInc = 0;)
 	for(typename TArray<FSetElementId>::TConstIterator PolicyIt(OrderedDrawingPolicies); PolicyIt; ++PolicyIt)
 	{
 		FDrawingPolicyLink* DrawingPolicyLink = &DrawingPolicySet(*PolicyIt);
@@ -234,7 +235,7 @@ bool TStaticMeshDrawList<DrawingPolicyType>::DrawVisible(
 			if(StaticMeshVisibilityMap.AccessCorrespondingBit(FRelativeBitReference(CompactElementPtr->MeshId)))
 			{
 				const FElement& Element = DrawingPolicyLink->Elements[ElementIndex];
-				INC_DWORD_STAT_BY(STAT_StaticMeshTriangles,Element.Mesh->GetNumPrimitives());
+				STAT(StatInc +=  Element.Mesh->GetNumPrimitives();)
 				// Avoid the virtual call looking up batch visibility if there is only one element.
 				uint32 BatchElementMask = Element.Mesh->Elements.Num() == 1 ? 1 : Element.Mesh->VertexFactory->GetStaticBatchElementVisibility(View,Element.Mesh);
 				DrawElement(View, PolicyContext, Element, BatchElementMask, DrawingPolicyLink, bDrawnShared);
@@ -242,6 +243,7 @@ bool TStaticMeshDrawList<DrawingPolicyType>::DrawVisible(
 			}
 		}
 	}
+	INC_DWORD_STAT_BY(STAT_StaticMeshTriangles, StatInc);
 	return bDirty;
 }
 
@@ -256,6 +258,7 @@ bool TStaticMeshDrawList<DrawingPolicyType>::DrawVisibleInner(
 	)
 {
 	bool bDirty = false;
+	STAT(int32 StatInc = 0;)
 	for (int32 Index = FirstPolicy; Index <= LastPolicy; Index++)
 	{
 		FDrawingPolicyLink* DrawingPolicyLink = &DrawingPolicySet[OrderedDrawingPolicies[Index]];
@@ -269,7 +272,7 @@ bool TStaticMeshDrawList<DrawingPolicyType>::DrawVisibleInner(
 			if (StaticMeshVisibilityMap.AccessCorrespondingBit(FRelativeBitReference(CompactElementPtr->MeshId)))
 			{
 				const FElement& Element = DrawingPolicyLink->Elements[ElementIndex];
-				INC_DWORD_STAT_BY(STAT_StaticMeshTriangles, Element.Mesh->GetNumPrimitives());
+				STAT(StatInc +=  Element.Mesh->GetNumPrimitives();)
 				// Avoid the cache miss looking up batch visibility if there is only one element.
 				uint64 BatchElementMask = Element.Mesh->Elements.Num() == 1 ? 1 : BatchVisibilityArray[Element.Mesh->Id];
 				DrawElement(RHICmdList, View, PolicyContext, Element, BatchElementMask, DrawingPolicyLink, bDrawnShared);
@@ -277,6 +280,7 @@ bool TStaticMeshDrawList<DrawingPolicyType>::DrawVisibleInner(
 			}
 		}
 	}
+	INC_DWORD_STAT_BY(STAT_StaticMeshTriangles, StatInc);
 	return bDirty;
 }
 
@@ -307,6 +311,7 @@ class FDrawVisibleAnyThreadTask
 	const int32 LastPolicy;
 
 	bool& OutDirty;
+	const TArray<uint16, SceneRenderingAllocator>* PerDrawingPolicyCounts;
 
 public:
 
@@ -319,7 +324,8 @@ public:
 		const TArray<uint64, SceneRenderingAllocator>& InBatchVisibilityArray,
 		int32 InFirstPolicy,
 		int32 InLastPolicy,
-		bool& InOutDirty
+		bool& InOutDirty,
+		const TArray<uint16, SceneRenderingAllocator>* InPerDrawingPolicyCounts
 		)
 		: Caller(InCaller)
 		, RHICmdList(InRHICmdList)
@@ -330,6 +336,7 @@ public:
 		, FirstPolicy(InFirstPolicy)
 		, LastPolicy(InLastPolicy)
 		, OutDirty(InOutDirty)
+		, PerDrawingPolicyCounts(InPerDrawingPolicyCounts)
 	{
 	}
 
@@ -347,9 +354,37 @@ public:
 
 	void DoTask(ENamedThreads::Type CurrentThread, const FGraphEventRef& MyCompletionGraphEvent)
 	{
-		if (this->Caller.DrawVisibleInner(this->RHICmdList, this->View, this->PolicyContext, this->StaticMeshVisibilityMap, this->BatchVisibilityArray, this->FirstPolicy, this->LastPolicy))
+		if (this->PerDrawingPolicyCounts)
 		{
-			this->OutDirty = true;
+			int32 Start = this->FirstPolicy;
+			// we have per policy draw counts, lets not look at the zeros
+			while (Start <= this->LastPolicy)
+			{
+				while (Start <= this->LastPolicy && !(*PerDrawingPolicyCounts)[Start])
+				{
+					Start++;
+				}
+				if (Start <= this->LastPolicy)
+				{
+					int32 BatchEnd = Start;
+					while (BatchEnd < this->LastPolicy && (*PerDrawingPolicyCounts)[BatchEnd + 1])
+					{
+						BatchEnd++;
+					}
+					if (this->Caller.DrawVisibleInner(this->RHICmdList, this->View, this->PolicyContext, this->StaticMeshVisibilityMap, this->BatchVisibilityArray, Start, BatchEnd))
+					{
+						this->OutDirty = true;
+					}
+					Start = BatchEnd + 1;
+				}
+			}
+		}
+		else
+		{
+			if (this->Caller.DrawVisibleInner(this->RHICmdList, this->View, this->PolicyContext, this->StaticMeshVisibilityMap, this->BatchVisibilityArray, this->FirstPolicy, this->LastPolicy))
+			{
+				this->OutDirty = true;
+			}
 		}
 		this->RHICmdList.HandleRTThreadTaskCompletion(MyCompletionGraphEvent);
 	}
@@ -363,35 +398,159 @@ void TStaticMeshDrawList<DrawingPolicyType>::DrawVisibleParallel(
 	FParallelCommandListSet& ParallelCommandListSet
 	)
 {
-
-	int32 EffectiveThreads = FMath::Min<int32>(OrderedDrawingPolicies.Num(), ParallelCommandListSet.Width);
-
-	int32 Start = 0;
+	int32 NumPolicies = OrderedDrawingPolicies.Num();
+	int32 EffectiveThreads = FMath::Min<int32>(NumPolicies, ParallelCommandListSet.Width);
 	if (EffectiveThreads)
 	{
 
-		int32 NumPer = OrderedDrawingPolicies.Num() / EffectiveThreads;
-		int32 Extra = OrderedDrawingPolicies.Num() - NumPer * EffectiveThreads;
+		if (ParallelCommandListSet.bBalanceCommands)
+		{
+			TArray<uint16, SceneRenderingAllocator>& PerDrawingPolicyCounts = *new (FMemStack::Get()) TArray<uint16, SceneRenderingAllocator>();
+			PerDrawingPolicyCounts.AddZeroed(NumPolicies);
 
+			ParallelFor(NumPolicies, 
+				[this, &PerDrawingPolicyCounts, &StaticMeshVisibilityMap, &BatchVisibilityArray](int32 Index)
+				{
+					int32 Count = 0;
+
+					FDrawingPolicyLink* DrawingPolicyLink = &DrawingPolicySet[OrderedDrawingPolicies[Index]];
+					const FElementCompact* CompactElementPtr = DrawingPolicyLink->CompactElements.GetData();
+					FPlatformMisc::Prefetch(CompactElementPtr);
+					const int32 NumElements = DrawingPolicyLink->CompactElements.Num();
+					FPlatformMisc::Prefetch(&DrawingPolicyLink->CompactElements.GetData()->MeshId);
+					for (int32 ElementIndex = 0; ElementIndex < NumElements; ElementIndex++, CompactElementPtr++)
+					{
+						if (StaticMeshVisibilityMap.AccessCorrespondingBit(FRelativeBitReference(CompactElementPtr->MeshId)))
+						{
+							const FElement& Element = DrawingPolicyLink->Elements[ElementIndex];
+							Count += Element.Mesh->Elements.Num();
+						}
+					}
+					if (Count)
+					{
+						// this is unlikely to overflow, but I don't think it would matter much if it did
+						PerDrawingPolicyCounts[Index] = (uint16)FMath::Min<int32>(Count, MAX_uint16);
+					}
+				}
+			);
 
-		for (int32 ThreadIndex = 0; ThreadIndex < EffectiveThreads; ThreadIndex++)
-		{
-			int32 Last = Start + (NumPer - 1) + (ThreadIndex < Extra);
-			check(Last >= Start);
+			int32 Total = 0;
+			for (uint16 Count : PerDrawingPolicyCounts)
+			{
+				Total += Count;
+			}
+			if (Total == 0)
+			{
+				return;
+			}
+			UE_CLOG(ParallelCommandListSet.bSpewBalance, LogTemp, Display, TEXT("Total Draws %d"), Total);
+
+			EffectiveThreads = FMath::Min<int32>(EffectiveThreads, FMath::DivideAndRoundUp(Total, ParallelCommandListSet.MinDrawsPerCommandList));
+			check(EffectiveThreads > 0 && EffectiveThreads <= ParallelCommandListSet.Width);
 
+			int32 DrawsPerCmdList = FMath::DivideAndRoundUp(Total, EffectiveThreads);
+			int32 DrawsPerCmdListMergeLimit = FMath::DivideAndRoundUp(DrawsPerCmdList, 3); // if the last list would be small, we just merge it into the previous one
+
+			int32 Start = 0;
+
+			int32 PreviousBatchStart = -1;
+			int32 PreviousBatchEnd = -2;
+			int32 PreviousBatchDraws = 0;
+
+			while (Start < NumPolicies)
 			{
+				// skip zeros
+				while (Start < NumPolicies && !PerDrawingPolicyCounts[Start])
+				{
+					Start++;
+				}
+				if (Start < NumPolicies)
+				{
+					int32 BatchCount = PerDrawingPolicyCounts[Start];
+					int32 BatchEnd = Start;
+					int32 LastNonZeroPolicy = Start;
+					while (BatchEnd < NumPolicies - 1 && BatchCount < DrawsPerCmdList)
+					{
+						BatchEnd++;
+						if (PerDrawingPolicyCounts[BatchEnd])
+						{
+							BatchCount += PerDrawingPolicyCounts[BatchEnd];
+							LastNonZeroPolicy = BatchEnd;
+						}
+					}
+					if (BatchCount)
+					{
+						if (PreviousBatchStart <= PreviousBatchEnd)
+						{
+							FRHICommandList* CmdList = ParallelCommandListSet.NewParallelCommandList();
+							if (BatchCount < DrawsPerCmdListMergeLimit)
+							{
+								// this is the last batch and it is small, let merge it
+								UE_CLOG(ParallelCommandListSet.bSpewBalance, LogTemp, Display, TEXT("    Index %d  BatchCount %d    (last merge)"), ParallelCommandListSet.NumParallelCommandLists(), PreviousBatchDraws + BatchCount);
+								FGraphEventRef AnyThreadCompletionEvent = TGraphTask<FDrawVisibleAnyThreadTask<DrawingPolicyType> >::CreateTask(nullptr, ENamedThreads::RenderThread)
+									.ConstructAndDispatchWhenReady(*this, *CmdList, ParallelCommandListSet.View, PolicyContext, StaticMeshVisibilityMap, BatchVisibilityArray, PreviousBatchStart, LastNonZeroPolicy, ParallelCommandListSet.OutDirty, &PerDrawingPolicyCounts);
+								ParallelCommandListSet.AddParallelCommandList(CmdList, AnyThreadCompletionEvent, PreviousBatchDraws + BatchCount);
+								PreviousBatchStart = -1;
+								PreviousBatchEnd = -2;
+								PreviousBatchDraws = 0;
+							}
+							else
+							{
+								// this is a decent sized batch, emit the previous batch and save this one for possible merging
+								UE_CLOG(ParallelCommandListSet.bSpewBalance, LogTemp, Display, TEXT("    Index %d  BatchCount %d    "), ParallelCommandListSet.NumParallelCommandLists(), PreviousBatchDraws);
+								FGraphEventRef AnyThreadCompletionEvent = TGraphTask<FDrawVisibleAnyThreadTask<DrawingPolicyType> >::CreateTask(nullptr, ENamedThreads::RenderThread)
+									.ConstructAndDispatchWhenReady(*this, *CmdList, ParallelCommandListSet.View, PolicyContext, StaticMeshVisibilityMap, BatchVisibilityArray, PreviousBatchStart, PreviousBatchEnd, ParallelCommandListSet.OutDirty, &PerDrawingPolicyCounts);
+								ParallelCommandListSet.AddParallelCommandList(CmdList, AnyThreadCompletionEvent, PreviousBatchDraws);
+								PreviousBatchStart = Start;
+								PreviousBatchEnd = LastNonZeroPolicy;
+								PreviousBatchDraws = BatchCount;
+							}
+						}
+						else
+						{
+							// no batch yet, save this one
+							PreviousBatchStart = Start;
+							PreviousBatchEnd = LastNonZeroPolicy;
+							PreviousBatchDraws = BatchCount;
+						}
+					}
+					Start = BatchEnd + 1;
+				}
+			}
+			// we didn't merge the last batch, emit now
+			if (PreviousBatchStart <= PreviousBatchEnd)
+			{
+				UE_CLOG(ParallelCommandListSet.bSpewBalance, LogTemp, Display, TEXT("    Index %d  BatchCount %d    (last)"), ParallelCommandListSet.NumParallelCommandLists(), PreviousBatchDraws);
 				FRHICommandList* CmdList = ParallelCommandListSet.NewParallelCommandList();
-
 				FGraphEventRef AnyThreadCompletionEvent = TGraphTask<FDrawVisibleAnyThreadTask<DrawingPolicyType> >::CreateTask(nullptr, ENamedThreads::RenderThread)
-					.ConstructAndDispatchWhenReady(*this, *CmdList, ParallelCommandListSet.View, PolicyContext, StaticMeshVisibilityMap, BatchVisibilityArray, Start, Last, ParallelCommandListSet.OutDirty);
-
-				ParallelCommandListSet.AddParallelCommandList(CmdList, AnyThreadCompletionEvent);
+					.ConstructAndDispatchWhenReady(*this, *CmdList, ParallelCommandListSet.View, PolicyContext, StaticMeshVisibilityMap, BatchVisibilityArray, PreviousBatchStart, PreviousBatchEnd, ParallelCommandListSet.OutDirty, &PerDrawingPolicyCounts);
+				ParallelCommandListSet.AddParallelCommandList(CmdList, AnyThreadCompletionEvent, PreviousBatchDraws);
 			}
+		}
+		else
+		{
+			int32 NumPer = OrderedDrawingPolicies.Num() / EffectiveThreads;
+			int32 Extra = OrderedDrawingPolicies.Num() - NumPer * EffectiveThreads;
+			int32 Start = 0;
+			for (int32 ThreadIndex = 0; ThreadIndex < EffectiveThreads; ThreadIndex++)
+			{
+				int32 Last = Start + (NumPer - 1) + (ThreadIndex < Extra);
+				check(Last >= Start);
+
+				{
+					FRHICommandList* CmdList = ParallelCommandListSet.NewParallelCommandList();
 
-			Start = Last + 1;
+					FGraphEventRef AnyThreadCompletionEvent = TGraphTask<FDrawVisibleAnyThreadTask<DrawingPolicyType> >::CreateTask(nullptr, ENamedThreads::RenderThread)
+						.ConstructAndDispatchWhenReady(*this, *CmdList, ParallelCommandListSet.View, PolicyContext, StaticMeshVisibilityMap, BatchVisibilityArray, Start, Last, ParallelCommandListSet.OutDirty, nullptr);
+
+					ParallelCommandListSet.AddParallelCommandList(CmdList, AnyThreadCompletionEvent);
+				}
+
+				Start = Last + 1;
+			}
+			check(Start == OrderedDrawingPolicies.Num());
 		}
 	}
-	check(Start == OrderedDrawingPolicies.Num());
 }
 
 template<typename DrawingPolicyType>
@@ -435,6 +594,7 @@ int32 TStaticMeshDrawList<DrawingPolicyType>::DrawVisibleFrontToBack(
 	int32 LastDrawingPolicyIndex = INDEX_NONE;
 	FDrawingPolicyLink* DrawingPolicyLink = NULL;
 	bool bDrawnShared = false;
+	STAT(int32 StatInc = 0;)
 	for (int32 SortedIndex = 0, NumSorted = FMath::Min(SortKeys.Num(),MaxToDraw); SortedIndex < NumSorted; ++SortedIndex)
 	{
 		int32 DrawingPolicyIndex = SortKeys[SortedIndex].Fields.DrawingPolicyIndex;
@@ -447,12 +607,13 @@ int32 TStaticMeshDrawList<DrawingPolicyType>::DrawVisibleFrontToBack(
 		}
 
 		const FElement& Element = DrawingPolicyLink->Elements[ElementIndex];
-		INC_DWORD_STAT_BY(STAT_StaticMeshTriangles,Element.Mesh->GetNumPrimitives());
+		STAT(StatInc +=  Element.Mesh->GetNumPrimitives();)
 		// Avoid the cache miss looking up batch visibility if there is only one element.
 		uint64 BatchElementMask = Element.Mesh->Elements.Num() == 1 ? 1 : BatchVisibilityArray[Element.Mesh->Id];
 		DrawElement(RHICmdList, View, PolicyContext, Element, BatchElementMask, DrawingPolicyLink, bDrawnShared);
 		NumDraws++;
 	}
+	INC_DWORD_STAT_BY(STAT_StaticMeshTriangles, StatInc);
 
 	return NumDraws;
 }
@@ -597,4 +758,3 @@ void TStaticMeshDrawList<DrawingPolicyType>::ApplyWorldOffset(FVector InOffset)
 	}
 }
 
-#endif
diff --git a/Engine/Source/Runtime/Renderer/Private/TranslucentRendering.cpp b/Engine/Source/Runtime/Renderer/Private/TranslucentRendering.cpp
--- a/Engine/Source/Runtime/Renderer/Private/TranslucentRendering.cpp
+++ b/Engine/Source/Runtime/Renderer/Private/TranslucentRendering.cpp
@@ -927,7 +927,7 @@ void FDeferredShadingSceneRenderer::RenderTranslucencyParallel(FRHICommandListIm
 
 			{
 				int32 NumPrims = View.TranslucentPrimSet.NumPrims() - View.TranslucentPrimSet.NumSeparateTranslucencyPrims();
-				int32 EffectiveThreads = FMath::Min<int32>(NumPrims, ParallelCommandListSet.Width);
+				int32 EffectiveThreads = FMath::Min<int32>(FMath::DivideAndRoundUp(NumPrims, ParallelCommandListSet.MinDrawsPerCommandList), ParallelCommandListSet.Width);
 
 				int32 Start = 0;
 				if (EffectiveThreads)
@@ -983,7 +983,7 @@ void FDeferredShadingSceneRenderer::RenderTranslucencyParallel(FRHICommandListIm
 				if (View.TranslucentPrimSet.NumSeparateTranslucencyPrims() > 0)
 				{
 					int32 NumPrims = View.TranslucentPrimSet.NumSeparateTranslucencyPrims();
-					int32 EffectiveThreads = FMath::Min<int32>(NumPrims, ParallelCommandListSet.Width);
+					int32 EffectiveThreads = FMath::Min<int32>(FMath::DivideAndRoundUp(NumPrims, ParallelCommandListSet.MinDrawsPerCommandList), ParallelCommandListSet.Width);
 
 					int32 Start = 0;
 					check(EffectiveThreads);
diff --git a/Engine/Source/Runtime/Renderer/Private/VelocityRendering.cpp b/Engine/Source/Runtime/Renderer/Private/VelocityRendering.cpp
--- a/Engine/Source/Runtime/Renderer/Private/VelocityRendering.cpp
+++ b/Engine/Source/Runtime/Renderer/Private/VelocityRendering.cpp
@@ -659,7 +659,7 @@ void FDeferredShadingSceneRenderer::RenderVelocitiesInnerParallel(FRHICommandLis
 		Scene->VelocityDrawList.DrawVisibleParallel(View.StaticMeshVelocityMap, View.StaticMeshBatchVisibility, ParallelCommandListSet);
 
 		int32 NumPrims = View.DynamicMeshElements.Num();
-		int32 EffectiveThreads = FMath::Min<int32>(NumPrims, ParallelCommandListSet.Width);
+		int32 EffectiveThreads = FMath::Min<int32>(FMath::DivideAndRoundUp(NumPrims, ParallelCommandListSet.MinDrawsPerCommandList), ParallelCommandListSet.Width);
 
 		int32 Start = 0;
 		if (EffectiveThreads)
