Don't use IndirectLightingCache shaders when there are no precomputed lighting volumes.  This optimization does not apply to static draw lists as it is registration order dependent.
#PR 792

diff --git a/Engine/Source/Runtime/Renderer/Private/BasePassRendering.cpp b/Engine/Source/Runtime/Renderer/Private/BasePassRendering.cpp
--- a/Engine/Source/Runtime/Renderer/Private/BasePassRendering.cpp
+++ b/Engine/Source/Runtime/Renderer/Private/BasePassRendering.cpp
@@ -178,6 +178,8 @@ public:
 
 	bool AllowIndirectLightingCache() const 
 	{ 
+		// Note: can't disallow based on presence of PrecomputedLightVolumes in the scene as this is registration time
+		// Unless extra handling is added to recreate static draw lists when new volumes are added
 		return true; 
 	}
 
@@ -286,7 +288,8 @@ public:
 
 	bool AllowIndirectLightingCache() const 
 	{ 
-		return View.Family->EngineShowFlags.IndirectLightingCache; 
+		const FScene* Scene = (const FScene*)View.Family->Scene;
+		return View.Family->EngineShowFlags.IndirectLightingCache && Scene && Scene->PrecomputedLightVolumes.Num() > 0;
 	}
 
 	bool AllowIndirectLightingCacheVolumeTexture() const
diff --git a/Engine/Source/Runtime/Renderer/Private/TranslucentRendering.cpp b/Engine/Source/Runtime/Renderer/Private/TranslucentRendering.cpp
--- a/Engine/Source/Runtime/Renderer/Private/TranslucentRendering.cpp
+++ b/Engine/Source/Runtime/Renderer/Private/TranslucentRendering.cpp
@@ -254,7 +254,8 @@ public:
 
 	bool AllowIndirectLightingCache() const
 	{
-		return View.Family->EngineShowFlags.IndirectLightingCache;
+		const FScene* Scene = (const FScene*)View.Family->Scene;
+		return View.Family->EngineShowFlags.IndirectLightingCache && Scene && Scene->PrecomputedLightVolumes.Num() > 0;
 	}
 
 	bool AllowIndirectLightingCacheVolumeTexture() const
