Early out of ILC interpolation if there are no samples
#PR 793

diff --git a/Engine/Source/Runtime/Renderer/Private/IndirectLightingCache.cpp b/Engine/Source/Runtime/Renderer/Private/IndirectLightingCache.cpp
--- a/Engine/Source/Runtime/Renderer/Private/IndirectLightingCache.cpp
+++ b/Engine/Source/Runtime/Renderer/Private/IndirectLightingCache.cpp
@@ -351,7 +351,7 @@ FIndirectLightingCacheAllocation* FIndirectLightingCache::FindPrimitiveAllocatio
 
 void FIndirectLightingCache::UpdateCache(FScene* Scene, FSceneRenderer& Renderer, bool bAllowUnbuiltPreview)
 {
-	if (IsIndirectLightingCacheAllowed(GetFeatureLevel()))
+	if (IsIndirectLightingCacheAllowed(GetFeatureLevel()) && Scene->PrecomputedLightVolumes.Num() > 0)
 	{
 		bool bAnyViewAllowsIndirectLightingCache = false;
 
