added comment

diff --git a/Engine/Source/Runtime/Renderer/Private/PostProcess/PostProcessing.h b/Engine/Source/Runtime/Renderer/Private/PostProcess/PostProcessing.h
--- a/Engine/Source/Runtime/Renderer/Private/PostProcess/PostProcessing.h
+++ b/Engine/Source/Runtime/Renderer/Private/PostProcess/PostProcessing.h
@@ -22,6 +22,7 @@ public:
 
 	// 0 if there was no scene color available at constructor call time
 	FRenderingCompositePass* SceneColor;
+	// never 0
 	FRenderingCompositePass* SceneDepth;
 
 	FRenderingCompositeOutputRef FinalOutput;
