Minor refactor to the CompositingGraph, allowing parts of the DAG/Graph to be executed before other parts which is useful for upcoming AsyncCompute or to control execution order in general.

API change: Dependency to root is is no longer needed, call Process() with the node directly.

diff --git a/Engine/Source/Runtime/Renderer/Private/CompositionLighting/CompositionLighting.cpp b/Engine/Source/Runtime/Renderer/Private/CompositionLighting/CompositionLighting.cpp
--- a/Engine/Source/Runtime/Renderer/Private/CompositionLighting/CompositionLighting.cpp
+++ b/Engine/Source/Runtime/Renderer/Private/CompositionLighting/CompositionLighting.cpp
@@ -262,10 +262,7 @@ void FCompositionLighting::ProcessBeforeBasePass(FRHICommandListImmediate& RHICm
 
 		SCOPED_DRAW_EVENT(RHICmdList, CompositionBeforeBasePass);
 
-		// you can add multiple dependencies
-		CompositeContext.Root->AddDependency(Context.FinalOutput);
-
-		CompositeContext.Process(TEXT("Composition_BeforeBasePass"));
+		CompositeContext.Process(Context.FinalOutput.GetPass(), TEXT("Composition_BeforeBasePass"));
 	}
 }
 
@@ -323,10 +320,7 @@ void FCompositionLighting::ProcessAfterBasePass(FRHICommandListImmediate& RHICmd
 		Context.FinalOutput.GetOutput()->RenderTargetDesc = SceneColor->GetDesc();
 		Context.FinalOutput.GetOutput()->PooledRenderTarget = SceneColor;
 
-		// you can add multiple dependencies
-		CompositeContext.Root->AddDependency(Context.FinalOutput);
-
-		CompositeContext.Process(TEXT("CompositionLighting_AfterBasePass"));
+		CompositeContext.Process(Context.FinalOutput.GetPass(), TEXT("CompositionLighting_AfterBasePass"));
 	}
 }
 
@@ -382,7 +376,10 @@ void FCompositionLighting::ProcessLighting(FRHICommandListImmediate& RHICmdList,
 				FRenderingCompositePass* RecombinePass = Context.Graph.RegisterPass(new(FMemStack::Get()) FRCPassPostProcessSubsurfaceRecombine());
 				RecombinePass->SetInput(ePId_Input0, Pass1);
 				RecombinePass->SetInput(ePId_Input1, PassExtractSpecular);
-				CompositeContext.Root->AddDependency(RecombinePass);
+				
+				SCOPED_DRAW_EVENT(RHICmdList, CompositionLighting_SSSSS);
+
+				CompositeContext.Process(RecombinePass, TEXT("CompositionLighting_SSSSS"));
 			}
 		}
 
@@ -392,10 +389,7 @@ void FCompositionLighting::ProcessLighting(FRHICommandListImmediate& RHICmdList,
 
 		// we don't replace the final element with the scenecolor because this is what those passes should do by themself
 
-		// you can add multiple dependencies
-		CompositeContext.Root->AddDependency(Context.FinalOutput);
-
-		CompositeContext.Process(TEXT("CompositionLighting_Lighting"));
+		CompositeContext.Process(Context.FinalOutput.GetPass(), TEXT("CompositionLighting"));
 	}
 
 	// We only release the after the last view was processed (SplitScreen)
diff --git a/Engine/Source/Runtime/Renderer/Private/ForwardShadingRenderer.cpp b/Engine/Source/Runtime/Renderer/Private/ForwardShadingRenderer.cpp
--- a/Engine/Source/Runtime/Renderer/Private/ForwardShadingRenderer.cpp
+++ b/Engine/Source/Runtime/Renderer/Private/ForwardShadingRenderer.cpp
@@ -180,8 +180,7 @@ void FForwardShadingSceneRenderer::Render(FRHICommandListImmediate& RHICmdList)
 		FRenderingCompositePassContext CompositeContext(RHICmdList, View);
 
 		FRenderingCompositePass* PostProcessSunMask = CompositeContext.Graph.RegisterPass(new(FMemStack::Get()) FRCPassPostProcessSunMaskES2(PrePostSourceViewportSize, true));
-		CompositeContext.Root->AddDependency(FRenderingCompositeOutputRef(PostProcessSunMask));
-		CompositeContext.Process(TEXT("OnChipAlphaTransform"));
+		CompositeContext.Process(PostProcessSunMask, TEXT("OnChipAlphaTransform"));
 	}
 
 	if (!bGammaSpace || bRenderToScene)
@@ -270,6 +269,5 @@ void FForwardShadingSceneRenderer::BasicPostProcess(FRHICommandListImmediate& RH
 	Context.FinalOutput.GetOutput()->PooledRenderTarget = Temp;
 	Context.FinalOutput.GetOutput()->RenderTargetDesc = Desc;
 
-	CompositeContext.Root->AddDependency(Context.FinalOutput);
-	CompositeContext.Process(TEXT("ES2BasicPostProcess"));
+	CompositeContext.Process(Context.FinalOutput.GetPass(), TEXT("ES2BasicPostProcess"));
 }
diff --git a/Engine/Source/Runtime/Renderer/Private/LightShaftRendering.cpp b/Engine/Source/Runtime/Renderer/Private/LightShaftRendering.cpp
--- a/Engine/Source/Runtime/Renderer/Private/LightShaftRendering.cpp
+++ b/Engine/Source/Runtime/Renderer/Private/LightShaftRendering.cpp
@@ -557,10 +557,8 @@ void ApplyTemporalAA(
 			Context.FinalOutput.GetOutput()->RenderTargetDesc = NewHistory->GetDesc();
 			Context.FinalOutput.GetOutput()->PooledRenderTarget = NewHistory;
 
-			CompositeContext.Root->AddDependency(Context.FinalOutput);
-
 			// Execute Temporal AA
-			CompositeContext.Process(TEXT("LightShaftTemporalAA"));
+			CompositeContext.Process(Context.FinalOutput.GetPass(), TEXT("LightShaftTemporalAA"));
 
 			// Update the view state's render target reference with the new history
 			*HistoryState = NewHistory;
diff --git a/Engine/Source/Runtime/Renderer/Private/PostProcess/PostProcessing.cpp b/Engine/Source/Runtime/Renderer/Private/PostProcess/PostProcessing.cpp
--- a/Engine/Source/Runtime/Renderer/Private/PostProcess/PostProcessing.cpp
+++ b/Engine/Source/Runtime/Renderer/Private/PostProcess/PostProcessing.cpp
@@ -1614,10 +1614,7 @@ void FPostProcessing::Process(FRHICommandListImmediate& RHICmdList, FViewInfo& V
 
 			OverrideRenderTarget(Context.FinalOutput, Temp, Desc);
 
-			// you can add multiple dependencies
-			CompositeContext.Root->AddDependency(Context.FinalOutput);
-
-			CompositeContext.Process(TEXT("PostProcessing"));
+			CompositeContext.Process(Context.FinalOutput.GetPass(), TEXT("PostProcessing"));
 		}
 	}
 
@@ -2000,10 +1997,7 @@ void FPostProcessing::ProcessES2(FRHICommandListImmediate& RHICmdList, FViewInfo
 
 			OverrideRenderTarget(Context.FinalOutput, Temp, Desc);
 
-			// you can add multiple dependencies
-			CompositeContext.Root->AddDependency(Context.FinalOutput);
-
-			CompositeContext.Process(TEXT("PostProcessingES2"));
+			CompositeContext.Process(Context.FinalOutput.GetPass(), TEXT("PostProcessingES2"));
 		}
 	}
 }
diff --git a/Engine/Source/Runtime/Renderer/Private/PostProcess/RenderingCompositionGraph.cpp b/Engine/Source/Runtime/Renderer/Private/PostProcess/RenderingCompositionGraph.cpp
--- a/Engine/Source/Runtime/Renderer/Private/PostProcess/RenderingCompositionGraph.cpp
+++ b/Engine/Source/Runtime/Renderer/Private/PostProcess/RenderingCompositionGraph.cpp
@@ -106,7 +106,6 @@ FAutoConsoleCommand CmdCompositionGraphDebug(
 FRenderingCompositePassContext::FRenderingCompositePassContext(FRHICommandListImmediate& InRHICmdList, FViewInfo& InView/*, const FSceneRenderTargetItem& InRenderTargetItem*/)
 	: View(InView)
 	, ViewState((FSceneViewState*)InView.State)
-	//		, CompositingOutputRTItem(InRenderTargetItem)
 	, Pass(0)
 	, RHICmdList(InRHICmdList)
 	, ViewPortRect(0, 0, 0 ,0)
@@ -114,10 +113,6 @@ FRenderingCompositePassContext::FRenderingCompositePassContext(FRHICommandListIm
 	, ShaderMap(InView.ShaderMap)
 {
 	check(!IsViewportValid());
-
-	Root = Graph.RegisterPass(new(FMemStack::Get()) FRCPassPostProcessRoot());
-
-	Graph.ResetDependencies();
 }
 
 FRenderingCompositePassContext::~FRenderingCompositePassContext()
@@ -125,7 +120,7 @@ FRenderingCompositePassContext::~FRenderingCompositePassContext()
 	Graph.Free();
 }
 
-void FRenderingCompositePassContext::Process(const TCHAR *GraphDebugName)
+void FRenderingCompositePassContext::Process(FRenderingCompositePass* Root, const TCHAR *GraphDebugName)
 {
 	if(Root)
 	{
@@ -194,12 +189,12 @@ void FRenderingCompositionGraph::RecursivelyGatherDependencies(const FRenderingC
 {
 	FRenderingCompositePass *Pass = InOutputRef.GetPass();
 
-	if(Pass->bGraphMarker)
+	if(Pass->bComputeOutputDescWasCalled)
 	{
 		// already processed
 		return;
 	}
-	Pass->bGraphMarker = true;
+	Pass->bComputeOutputDescWasCalled = true;
 
 	// iterate through all inputs and additional dependencies of this pass
 	uint32 Index = 0;
@@ -311,12 +306,12 @@ void FRenderingCompositionGraph::RecursivelyProcess(const FRenderingCompositeOut
 	check(Pass);
 	check(Output);
 
-	if(!Pass->bGraphMarker)
+	if(Pass->bProcessWasCalled)
 	{
 		// already processed
 		return;
 	}
-	Pass->bGraphMarker = false;
+	Pass->bProcessWasCalled = true;
 
 	// iterate through all inputs and additional dependencies of this pass
 	{
@@ -337,7 +332,7 @@ void FRenderingCompositionGraph::RecursivelyProcess(const FRenderingCompositeOut
 				// to track down an issue, should never happen
 				check(OutputRefIt->GetPass());
 
-				if(GRenderTargetPool.IsEventRecordingEnabled() && Pass != Context.Root)
+				if(GRenderTargetPool.IsEventRecordingEnabled())
 				{
 					GRenderTargetPool.AddPhaseEvent(*Pass->ConstructDebugName());
 				}
@@ -620,24 +615,6 @@ int32 FRenderingCompositionGraph::ComputeUniqueOutputId(FRenderingCompositePass*
 	return -1;
 }
 
-
-void FRenderingCompositionGraph::ResetDependencies()
-{
-	for(uint32 i = 0; i < (uint32)Nodes.Num(); ++i)
-	{
-		FRenderingCompositePass *Element = Nodes[i];
-
-		Element->bGraphMarker = false;
-
-		uint32 OutputId = 0;
-
-		while(FRenderingCompositeOutput* Output = Element->GetOutput((EPassOutputId)(OutputId++)))
-		{
-			Output->ResetDependency();
-		}
-	}
-}
-
 FRenderingCompositeOutput *FRenderingCompositeOutputRef::GetOutput() const
 {
 	if(Source == 0)
diff --git a/Engine/Source/Runtime/Renderer/Private/PostProcess/RenderingCompositionGraph.h b/Engine/Source/Runtime/Renderer/Private/PostProcess/RenderingCompositionGraph.h
--- a/Engine/Source/Runtime/Renderer/Private/PostProcess/RenderingCompositionGraph.h
+++ b/Engine/Source/Runtime/Renderer/Private/PostProcess/RenderingCompositionGraph.h
@@ -48,9 +48,6 @@ private:
 	/** */
 	void Free();
 
-	/** calls ResetDependency() in OutputId in each pass */
-	void ResetDependencies();
-
 	/** */
 	void ProcessGatherDependency(const FRenderingCompositeOutputRef* OutputRefIt);
 
@@ -84,7 +81,7 @@ struct FRenderingCompositePassContext
 	~FRenderingCompositePassContext();
 
 	// @param GraphDebugName must not be 0
-	void Process(const TCHAR *GraphDebugName);
+	void Process(FRenderingCompositePass* Root, const TCHAR *GraphDebugName);
 
 	//
 	FViewInfo& View;
@@ -135,8 +132,6 @@ struct FRenderingCompositePassContext
 	EShaderPlatform GetShaderPlatform() const { return GShaderPlatformForFeatureLevel[FeatureLevel]; }
 	TShaderMap<FGlobalShaderType>* GetShaderMap() const { check(ShaderMap); return ShaderMap; }
 
-	FRenderingCompositePass* Root;
-
 	FRenderingCompositionGraph Graph;
 
 	FRHICommandListImmediate& RHICmdList;
@@ -154,7 +149,9 @@ private:
 struct FRenderingCompositePass
 {
 	/** constructor */
-	FRenderingCompositePass() : bGraphMarker(false)
+	FRenderingCompositePass() 
+		: bComputeOutputDescWasCalled(false)
+		, bProcessWasCalled(false)
 	{
 	}
 
@@ -163,6 +160,12 @@ struct FRenderingCompositePass
 	/** @return 0 if outside the range */
 	virtual FRenderingCompositeOutputRef* GetInput(EPassInputId InPassInputId) = 0;
 
+	/**
+	 * const version of GetInput()
+	 * @return 0 if outside the range
+	 */
+	virtual const FRenderingCompositeOutputRef* GetInput(EPassInputId InPassInputId) const = 0;
+
 	/**
 	 * Each input is a dependency and will be processed before the node itself (don't generate cycles)
 	 * The index allows to access the input in Process() and on the shader side
@@ -242,10 +245,12 @@ struct FRenderingCompositePass
 	/** */
 	virtual void Release() = 0;
 
-private:
+protected:
 
-	/** to allow the graph to mark already processed nodes */
-	bool bGraphMarker;
+	/** to avoid wasteful recomputation and to support graph/DAG traversal, if ComputeOutputDesc() was called */
+	bool bComputeOutputDescWasCalled;
+	/** to allows reuse and to support graph/DAG traversal, if Process() was called */
+	bool bProcessWasCalled;
 
 	friend class FRenderingCompositionGraph;
 };
@@ -367,6 +372,17 @@ struct TRenderingCompositePassBase :public FRenderingCompositePass
 		return 0;
 	}
 	
+	// const version of GetInput()
+	virtual const FRenderingCompositeOutputRef* GetInput(EPassInputId InPassInputId) const
+	{
+		if((int32)InPassInputId < InputCount)
+		{
+			return &PassInputs[InPassInputId];
+		}
+
+		return 0;
+	}
+
 	virtual void SetInput(EPassInputId InPassInputId, const FRenderingCompositeOutputRef& VirtualBuffer)
 	{
 		if((int32)InPassInputId < InputCount)
@@ -468,15 +484,4 @@ protected:
 	TArray<FRenderingCompositeOutputRef> AdditionalDependencies;
 };
 
-
-// derives from TRenderingCompositePassBase<InputCount, OutputCount>
-class FRCPassPostProcessRoot : public TRenderingCompositePassBase<0, 1>
-{
-public:
-	// interface FRenderingCompositePass ---------
-	virtual void Process(FRenderingCompositePassContext& Context) override {}
-	virtual void Release() override { delete this; }
-	FPooledRenderTargetDesc ComputeOutputDesc(EPassOutputId InPassOutputId) const override { FPooledRenderTargetDesc Desc; Desc.DebugName = TEXT("Root"); return Desc; }
-};
-
 void CompositionGraph_OnStartFrame();
diff --git a/Engine/Source/Runtime/Renderer/Private/PostProcess/ScreenSpaceReflections.cpp b/Engine/Source/Runtime/Renderer/Private/PostProcess/ScreenSpaceReflections.cpp
--- a/Engine/Source/Runtime/Renderer/Private/PostProcess/ScreenSpaceReflections.cpp
+++ b/Engine/Source/Runtime/Renderer/Private/PostProcess/ScreenSpaceReflections.cpp
@@ -322,6 +322,5 @@ void ScreenSpaceReflections(FRHICommandListImmediate& RHICmdList, FViewInfo& Vie
 		Context.FinalOutput = FRenderingCompositeOutputRef( ReflectionOutput );
 	}
 
-	CompositeContext.Root->AddDependency( Context.FinalOutput );
-	CompositeContext.Process(TEXT("ReflectionEnvironments"));
+	CompositeContext.Process(Context.FinalOutput.GetPass(), TEXT("ReflectionEnvironments"));
 }
diff --git a/Engine/Source/Runtime/Renderer/Private/SceneRendering.cpp b/Engine/Source/Runtime/Renderer/Private/SceneRendering.cpp
--- a/Engine/Source/Runtime/Renderer/Private/SceneRendering.cpp
+++ b/Engine/Source/Runtime/Renderer/Private/SceneRendering.cpp
@@ -1109,10 +1109,8 @@ void FSceneRenderer::RenderFinish(FRHICommandListImmediate& RHICmdList)
 		
 		if(BusyWait.IsValid())
 		{
-			CompositeContext.Root->AddDependency(BusyWait);
+			CompositeContext.Process(BusyWait.GetPass(), TEXT("RenderFinish"));
 		}
-
-		CompositeContext.Process(TEXT("RenderFinish"));
 	}
 	
 #if !(UE_BUILD_SHIPPING || UE_BUILD_TEST)
