Added first use of AsyncCompute of VelocityFlatten (Disabled by default, was enabled for GDC demo)

diff --git a/Engine/Source/Runtime/Renderer/Private/PostProcess/PostProcessMotionBlur.cpp b/Engine/Source/Runtime/Renderer/Private/PostProcess/PostProcessMotionBlur.cpp
--- a/Engine/Source/Runtime/Renderer/Private/PostProcess/PostProcessMotionBlur.cpp
+++ b/Engine/Source/Runtime/Renderer/Private/PostProcess/PostProcessMotionBlur.cpp
@@ -759,6 +759,15 @@ void FRCPassPostProcessVelocityFlatten::Process(FRenderingCompositePassContext&
 	SCOPED_DRAW_EVENT(Context.RHICmdList, PostProcessVelocityFlatten);
 	const FPooledRenderTargetDesc* InputDesc = GetInputDesc(ePId_Input0);
 
+	if (AsyncJobFenceID != -1)
+	{
+		// If we run with AsyncCompute we use the same node twice, once to start the AsyncCompute and once to wait for the result.
+		// The later one is happening here.
+		Context.RHICmdList.GraphicsWaitOnAsyncComputeJob(AsyncJobFenceID);
+		AsyncJobFenceID = -1;
+		return;
+	}
+
 	if(!InputDesc)
 	{
 		// input is not hooked up correctly
@@ -774,27 +783,63 @@ void FRCPassPostProcessVelocityFlatten::Process(FRenderingCompositePassContext&
 
 	TShaderMapRef< FPostProcessVelocityFlattenCS > ComputeShader( Context.GetShaderMap() );
 
-	Context.RHICmdList.SetComputeShader(ComputeShader->GetComputeShader());
+	const bool bAsyncComputeEnabled = IsAsyncComputeEnabled();
+
 	SetRenderTarget(Context.RHICmdList, FTextureRHIRef(), FTextureRHIRef());
 
+	if(bAsyncComputeEnabled)
+	{
+		// If AsyncCompute is enabled we start recording the commands for that
+		Context.RHICmdList.BeginAsyncComputeJob_DrawThread(AsyncComputePriority_Default);
+	}
+
+	Context.RHICmdList.SetComputeShader(ComputeShader->GetComputeShader());
+
 	// set destination
-	Context.RHICmdList.SetUAVParameter( ComputeShader->GetComputeShader(), ComputeShader->OutVelocityFlat.GetBaseIndex(), DestRenderTarget0.UAV );
-	//Context.RHICmdList.SetUAVParameter( ComputeShader->GetComputeShader(), ComputeShader->OutPackedVelocityDepth.GetBaseIndex(), DestRenderTarget1.UAV );
-	Context.RHICmdList.SetUAVParameter( ComputeShader->GetComputeShader(), ComputeShader->OutMaxTileVelocity.GetBaseIndex(), DestRenderTarget1.UAV );
+	Context.RHICmdList.SetUAVParameter(ComputeShader->GetComputeShader(), ComputeShader->OutVelocityFlat.GetBaseIndex(), DestRenderTarget0.UAV);
+	//Context.RHICmdList.SetUAVParameter(ComputeShader->GetComputeShader(), ComputeShader->OutPackedVelocityDepth.GetBaseIndex(), DestRenderTarget1.UAV);
+	Context.RHICmdList.SetUAVParameter(ComputeShader->GetComputeShader(), ComputeShader->OutMaxTileVelocity.GetBaseIndex(), DestRenderTarget1.UAV);
 
-	ComputeShader->SetCS(Context.RHICmdList, Context, View );
-	
-	FIntPoint ThreadGroupCountValue = ComputeThreadGroupCount( View.ViewRect.Size() );
+	ComputeShader->SetCS(Context.RHICmdList, Context, View);
+
+	FIntPoint ThreadGroupCountValue = ComputeThreadGroupCount(View.ViewRect.Size());
 	DispatchComputeShader(Context.RHICmdList, *ComputeShader, ThreadGroupCountValue.X, ThreadGroupCountValue.Y, 1);
 
+#if USE_ASYNC_COMPUTE_CONTEXT
+	if ( AsyncJobFenceID != -1 )
+	{
+		Context.RHICmdList.GraphicsWaitOnAsyncComputeJob(AsyncJobFenceID);
+		AsyncJobFenceID = -1;
+	}
+#endif
+
+	//	void FD3D11DynamicRHI::RHIGraphicsWaitOnAsyncComputeJob( uint32 FenceIndex )
+	Context.RHICmdList.FlushComputeShaderCache();
+
 	// un-set destination
-	Context.RHICmdList.SetUAVParameter( ComputeShader->GetComputeShader(), ComputeShader->OutVelocityFlat.GetBaseIndex(), NULL );
-	//Context.RHICmdList.SetUAVParameter( ComputeShader->GetComputeShader(), ComputeShader->OutPackedVelocityDepth.GetBaseIndex(), NULL );
-	Context.RHICmdList.SetUAVParameter( ComputeShader->GetComputeShader(), ComputeShader->OutMaxTileVelocity.GetBaseIndex(), NULL );
+	Context.RHICmdList.SetUAVParameter(ComputeShader->GetComputeShader(), ComputeShader->OutVelocityFlat.GetBaseIndex(), NULL);
+	//Context.RHICmdList.SetUAVParameter(ComputeShader->GetComputeShader(), ComputeShader->OutPackedVelocityDepth.GetBaseIndex(), NULL);
+	Context.RHICmdList.SetUAVParameter(ComputeShader->GetComputeShader(), ComputeShader->OutMaxTileVelocity.GetBaseIndex(), NULL);
 
 	Context.RHICmdList.CopyToResolveTarget(DestRenderTarget0.TargetableTexture, DestRenderTarget0.ShaderResourceTexture, false, FResolveParams());
 	Context.RHICmdList.CopyToResolveTarget(DestRenderTarget1.TargetableTexture, DestRenderTarget1.ShaderResourceTexture, false, FResolveParams());
 	//Context.RHICmdList.CopyToResolveTarget(DestRenderTarget1.TargetableTexture, DestRenderTarget2.ShaderResourceTexture, false, FResolveParams());
+
+	// we want this pass to be executed another time
+	// so we do this: CompositeContext.Process(VelocityFlattenPass, TEXT("VelocityFlattenPass"));
+	// and this: bProcessWasCalled = false;
+	if(bAsyncComputeEnabled)
+	{
+		// we end recording the commands for AsyncCompute and get back a number to wait for it with GraphicsWaitOnAsyncComputeJob()
+		AsyncJobFenceID = Context.RHICmdList.EndAsyncComputeJob_DrawThread();
+
+		// mark it not processed so it gets executed a second time
+		bProcessWasCalled = false;
+
+		// we run this Process() two times and want to not get the dependencies processed twice (could be moved into the ComposingGraph iteration, slower but easier coding)
+		SetInput(ePId_Input0, FRenderingCompositeOutputRef());
+		SetInput(ePId_Input1, FRenderingCompositeOutputRef());
+	}
 }
 
 FIntPoint FRCPassPostProcessVelocityFlatten::ComputeThreadGroupCount(FIntPoint PixelExtent)
@@ -976,6 +1021,10 @@ public:
 IMPLEMENT_SHADER_TYPE(,FPostProcessVelocityScatterVS,TEXT("PostProcessMotionBlur"),TEXT("VelocityScatterVS"),SF_Vertex);
 IMPLEMENT_SHADER_TYPE(,FPostProcessVelocityScatterPS,TEXT("PostProcessMotionBlur"),TEXT("VelocityScatterPS"),SF_Pixel);
 
+FRCPassPostProcessVelocityFlatten::FRCPassPostProcessVelocityFlatten()
+	: AsyncJobFenceID(-1)
+{
+}
 
 void FRCPassPostProcessVelocityScatter::Process(FRenderingCompositePassContext& Context)
 {
diff --git a/Engine/Source/Runtime/Renderer/Private/PostProcess/PostProcessMotionBlur.h b/Engine/Source/Runtime/Renderer/Private/PostProcess/PostProcessMotionBlur.h
--- a/Engine/Source/Runtime/Renderer/Private/PostProcess/PostProcessMotionBlur.h
+++ b/Engine/Source/Runtime/Renderer/Private/PostProcess/PostProcessMotionBlur.h
@@ -78,6 +78,8 @@ public:
 class FRCPassPostProcessVelocityFlatten : public TRenderingCompositePassBase<2, 2>
 {
 public:
+	FRCPassPostProcessVelocityFlatten();
+
 	// interface FRenderingCompositePass ---------
 	virtual void Process(FRenderingCompositePassContext& Context) override;
 	virtual void Release() override { delete this; }
@@ -85,6 +87,9 @@ public:
 
 	// -------------------------------------------
 
+	// -1 if not set yet (AsyncCompute wasn't started, we need another Process to finish it)
+	uint32 AsyncJobFenceID;
+
 	static const uint32 TileSizeX = 16;
 	static const uint32 TileSizeY = 16;
 	
diff --git a/Engine/Source/Runtime/Renderer/Private/PostProcess/PostProcessing.cpp b/Engine/Source/Runtime/Renderer/Private/PostProcess/PostProcessing.cpp
--- a/Engine/Source/Runtime/Renderer/Private/PostProcess/PostProcessing.cpp
+++ b/Engine/Source/Runtime/Renderer/Private/PostProcess/PostProcessing.cpp
@@ -982,6 +982,8 @@ void FPostProcessing::Process(FRHICommandListImmediate& RHICmdList, FViewInfo& V
 		FRenderingCompositeOutputRef SeparateTranslucency;
 		// optional
 		FRenderingCompositeOutputRef BloomOutputCombined;
+		// not always valid
+		FRenderingCompositePass* VelocityFlattenPass = 0;
 
 		bool bAllowTonemapper = true;
 		EStereoscopicPass StereoPass = View.StereoPass;
@@ -1133,7 +1135,8 @@ void FPostProcessing::Process(FRHICommandListImmediate& RHICmdList, FViewInfo& V
 					FRenderingCompositeOutputRef SceneDepth( Context.SceneDepth );
 
 					{
-						FRenderingCompositePass* VelocityFlattenPass = Context.Graph.RegisterPass( new(FMemStack::Get()) FRCPassPostProcessVelocityFlatten() );
+						check(!VelocityFlattenPass);
+						VelocityFlattenPass = Context.Graph.RegisterPass( new(FMemStack::Get()) FRCPassPostProcessVelocityFlatten() );
 						VelocityFlattenPass->SetInput( ePId_Input0, VelocityInput );
 						VelocityFlattenPass->SetInput( ePId_Input1, SceneDepth );
 
@@ -1614,6 +1617,9 @@ void FPostProcessing::Process(FRHICommandListImmediate& RHICmdList, FViewInfo& V
 
 			OverrideRenderTarget(Context.FinalOutput, Temp, Desc);
 
+			// We want this executed early as it uses AsyncCompute (where available e.g. XboxOne).
+			CompositeContext.Process(VelocityFlattenPass, TEXT("VelocityFlattenPass"));
+			// Now we process the remaining part of the graph.
 			CompositeContext.Process(Context.FinalOutput.GetPass(), TEXT("PostProcessing"));
 		}
 	}
