Extended VisualizeTetxure/Vis to also have Red, Green and Blue channel modes (by [REDACTED_PERSON])

diff --git a/Engine/Shaders/VisualizeTexture.usf b/Engine/Shaders/VisualizeTexture.usf
--- a/Engine/Shaders/VisualizeTexture.usf
+++ b/Engine/Shaders/VisualizeTexture.usf
@@ -26,9 +26,9 @@ TextureCubeArray VisualizeTextureCubeArray;
 SamplerState VisualizeTextureCubeArraySampler;
 #endif
 
-// [0]: r:RGBAdd, g:AMul, b:Add, a:FracScale
+// [0]: r:RGBAdd, g:SingleChannelMul, b:Add, a:FracScale
 // [1]: r:BlinkState=0/1, g:Frac when 0, saturate when 1, b:array index, a:custom mip
-// [2]: r:0=normal 1=depth  gba:unused
+// [2]: r:0=normal 1=depth, g:unwrapcubemap, b:singlechannel(0-3), a:unused
 float4 VisualizeParam[3];
 // x, y (0 for 2D textures), z(only !=0 if 3d texture), w unused
 float4 TextureExtent;
@@ -118,7 +118,19 @@ void VisualizeTexturePS(
 #endif
 	}	
 
-	float Add = TexLookup.a * VisualizeParam[0].y + VisualizeParam[0].z;
+	float4 RGBAMask = 0;
+	int SingleChannel = (int)VisualizeParam[2].z;
+	switch (SingleChannel)
+	{
+		case 0: RGBAMask = float4(1,0,0,0); break;
+		case 1: RGBAMask = float4(0,1,0,0); break;
+		case 2: RGBAMask = float4(0,0,1,0); break;
+		case 3: RGBAMask = float4(0,0,0,1); break;
+		default: RGBAMask = float4(0,0,0,0); break;
+	}
+	float Add = dot( TexLookup, RGBAMask ) * VisualizeParam[0].y + VisualizeParam[0].z;
+
+
 	OutColor.rgb = TexLookup.rgb * VisualizeParam[0].xxx + Add.xxx;
 
 	float3 ScaledColor = OutColor.rgb * VisualizeParam[0].w;
diff --git a/Engine/Source/Runtime/Renderer/Private/PostProcess/VisualizeTexture.cpp b/Engine/Source/Runtime/Renderer/Private/PostProcess/VisualizeTexture.cpp
--- a/Engine/Source/Runtime/Renderer/Private/PostProcess/VisualizeTexture.cpp
+++ b/Engine/Source/Runtime/Renderer/Private/PostProcess/VisualizeTexture.cpp
@@ -90,9 +90,9 @@ public:
 			float FracScale = 1.0f;
 
 			// w * almost_1 to avoid frac(1) => 0
-			VisualizeParamValue[0] = FVector4(Data.RGBMul, Data.AMul, Add, FracScale * 0.9999f);
+			VisualizeParamValue[0] = FVector4(Data.RGBMul, Data.SingleChannelMul, Add, FracScale * 0.9999f);
 			VisualizeParamValue[1] = FVector4(BlinkState, Data.bSaturateInsteadOfFrac ? 1.0f : 0.0f, Data.ArrayIndex, Data.CustomMip);
-			VisualizeParamValue[2] = FVector4(Data.InputValueMapping, 0, 0,0);
+			VisualizeParamValue[2] = FVector4(Data.InputValueMapping, 0.0f, Data.SingleChannel );
 
 			SetShaderValueArray(RHICmdList, ShaderRHI, VisualizeParam, VisualizeParamValue, 3);
 		}
@@ -261,6 +261,8 @@ FVisualizeTexture::FVisualizeTexture()
 {
 	Mode = 0;
 	RGBMul = 1.0f;
+	SingleChannelMul = 0.0f;
+	SingleChannel = -1;
 	AMul = 0.0f;
 	UVInputMapping = 3;
 	Flags = 0;
@@ -406,6 +408,8 @@ void FVisualizeTexture::GenerateContent(FRHICommandListImmediate& RHICmdList, co
 	bool bDepthTexture = (Desc.TargetableFlags & TexCreate_DepthStencilTargetable) != 0;
 	
 	VisualizeTextureData.RGBMul = RGBMul;
+	VisualizeTextureData.SingleChannelMul = SingleChannelMul;
+	VisualizeTextureData.SingleChannel = SingleChannel;
 	VisualizeTextureData.AMul = AMul;
 	VisualizeTextureData.Tex00 = Tex00;
 	VisualizeTextureData.Tex11 = Tex11;
@@ -580,12 +584,23 @@ void FVisualizeTexture::PresentContent(FRHICommandListImmediate& RHICmdList, con
 			ExtendedName = FString::Printf(TEXT("%s"), Desc.DebugName);
 		}
 
-		FString Line = FString::Printf(TEXT("VisualizeTexture: %d \"%s\" RGB*%g+A*%g UV%d"),
+		FString Channels = TEXT("RGB");
+		switch( SingleChannel )
+		{
+			case 0: Channels = TEXT("R"); break;
+			case 1: Channels = TEXT("G"); break;
+			case 2: Channels = TEXT("B"); break;
+			case 3: Channels = TEXT("A"); break;
+		} 
+		float Multiplier = ( SingleChannel == -1 ) ? RGBMul : SingleChannelMul;
+
+		FString Line = FString::Printf(TEXT("VisualizeTexture: %d \"%s\" %s*%g UV%d"),
 			Mode,
 			*ExtendedName,
-			RGBMul,
-			AMul,
+			*Channels,
+			Multiplier,
 			UVInputMapping);
+
 		Canvas.DrawShadowedString( X, Y += YStep, *Line, GetStatsFont(), FLinearColor(1, 1, 1));
 	}
 	{
diff --git a/Engine/Source/Runtime/Renderer/Private/PostProcess/VisualizeTexture.h b/Engine/Source/Runtime/Renderer/Private/PostProcess/VisualizeTexture.h
--- a/Engine/Source/Runtime/Renderer/Private/PostProcess/VisualizeTexture.h
+++ b/Engine/Source/Runtime/Renderer/Private/PostProcess/VisualizeTexture.h
@@ -41,6 +41,13 @@ public:
 	int32 Mode;
 	//
 	float RGBMul;
+
+	// -1=off, 0=R, 1=G, 2=B, 3=A
+	int32 SingleChannel;
+
+	// Multiplier for the single channel
+	float SingleChannelMul;
+
 	//
 	float AMul;
 	// 0=view in left top, 1=whole texture, 2=pixel perfect centered, 3=Picture in Picture
@@ -112,6 +119,8 @@ struct FVisualizeTextureData
 	const FPooledRenderTargetDesc& Desc;
 	TRefCountPtr<FRHIShaderResourceView> StencilSRV;
 	float RGBMul;
+	float SingleChannelMul;
+	int32 SingleChannel;
 	float AMul;
 	FVector2D Tex00;
 	FVector2D Tex11;
diff --git a/Engine/Source/Runtime/Renderer/Private/Renderer.cpp b/Engine/Source/Runtime/Renderer/Private/Renderer.cpp
--- a/Engine/Source/Runtime/Renderer/Private/Renderer.cpp
+++ b/Engine/Source/Runtime/Renderer/Private/Renderer.cpp
@@ -218,6 +218,8 @@ static void VisualizeTextureExec( const TCHAR* Cmd, FOutputDevice &Ar )
 		{
 			// Init
 			GRenderTargetPool.VisualizeTexture.RGBMul = 1;
+			GRenderTargetPool.VisualizeTexture.SingleChannelMul = 0.0f;
+			GRenderTargetPool.VisualizeTexture.SingleChannel = -1;
 			GRenderTargetPool.VisualizeTexture.AMul = 0;
 			GRenderTargetPool.VisualizeTexture.UVInputMapping = 3;
 			GRenderTargetPool.VisualizeTexture.Flags = 0;
@@ -313,18 +315,28 @@ static void VisualizeTextureExec( const TCHAR* Cmd, FOutputDevice &Ar )
 		// e.g. RGB*6, A, *22, /2.7, A*7
 		else if(Parameter.Left(3) == TEXT("rgb")
 			|| Parameter.Left(1) == TEXT("a")
+			|| Parameter.Left(1) == TEXT("r")
+			|| Parameter.Left(1) == TEXT("g")
+			|| Parameter.Left(1) == TEXT("b")
 			|| Parameter.Left(1) == TEXT("*")
 			|| Parameter.Left(1) == TEXT("/"))
 		{
+			int SingleChannel = -1;
+
 			if(Parameter.Left(3) == TEXT("rgb"))
 			{
 				Parameter = Parameter.Right(Parameter.Len() - 3);
 			}
-			else if(Parameter.Left(1) == TEXT("a"))
+			else if(Parameter.Left(1) == TEXT("r")) SingleChannel = 0;
+			else if(Parameter.Left(1) == TEXT("g")) SingleChannel = 1;
+			else if(Parameter.Left(1) == TEXT("b")) SingleChannel = 2;
+			else if(Parameter.Left(1) == TEXT("a")) SingleChannel = 3;
+			if ( SingleChannel >= 0 )
 			{
 				Parameter = Parameter.Right(Parameter.Len() - 1);
+				GRenderTargetPool.VisualizeTexture.SingleChannel = SingleChannel;
+				GRenderTargetPool.VisualizeTexture.SingleChannelMul = 1;
 				GRenderTargetPool.VisualizeTexture.RGBMul = 0;
-				GRenderTargetPool.VisualizeTexture.AMul = 1;
 			}
 
 			float Mul = 1.0f;
@@ -341,6 +353,7 @@ static void VisualizeTextureExec( const TCHAR* Cmd, FOutputDevice &Ar )
 				Mul = 1.0f / FCString::Atof(*Parameter);
 			}
 			GRenderTargetPool.VisualizeTexture.RGBMul *= Mul;
+			GRenderTargetPool.VisualizeTexture.SingleChannelMul *= Mul;
 			GRenderTargetPool.VisualizeTexture.AMul *= Mul;
 		}
 		else
@@ -360,6 +373,9 @@ static void VisualizeTextureExec( const TCHAR* Cmd, FOutputDevice &Ar )
 		Ar.Logf(TEXT("  RGB      = RGB in range 0..1 (default)"));
 		Ar.Logf(TEXT("  *8       = RGB * 8"));
 		Ar.Logf(TEXT("  A        = alpha channel in range 0..1"));
+		Ar.Logf(TEXT("  R        = red channel in range 0..1"));
+		Ar.Logf(TEXT("  G        = green channel in range 0..1"));
+		Ar.Logf(TEXT("  B        = blue channel in range 0..1"));
 		Ar.Logf(TEXT("  A*16     = Alpha * 16"));
 		Ar.Logf(TEXT("  RGB/2    = RGB / 2"));
 		Ar.Logf(TEXT("SubResource:"));
