fix a bug where only emissive objects would render into the LPV, caused by CL 2289763

diff --git a/Engine/Source/Runtime/Renderer/Private/ShadowRendering.cpp b/Engine/Source/Runtime/Renderer/Private/ShadowRendering.cpp
--- a/Engine/Source/Runtime/Renderer/Private/ShadowRendering.cpp
+++ b/Engine/Source/Runtime/Renderer/Private/ShadowRendering.cpp
@@ -1075,9 +1075,10 @@ bool FShadowDepthDrawingPolicyFactory::DrawDynamicMesh(
 		const EMaterialShadingModel ShadingModel = Material->GetShadingModel();
 
 		const bool bLocalOnePassPointLightShadow = Context.ShadowInfo->CascadeSettings.bOnePassPointLightShadow;
-		const bool bReflectiveShadowmap = Context.ShadowInfo->bReflectiveShadowmap && !bLocalOnePassPointLightShadow && Material->ShouldInjectEmissiveIntoLPV();
+		const bool bReflectiveShadowmap = Context.ShadowInfo->bReflectiveShadowmap && !bOnePassPointLightShadow;
 
-		if ( (!IsTranslucentBlendMode(BlendMode) || bReflectiveShadowmap ) && ShadingModel != MSM_Unlit )
+		if ( (!IsTranslucentBlendMode(BlendMode) || bReflectiveShadowmap ) && ShadingModel != MSM_Unlit
+		|| (Material->ShouldInjectEmissiveIntoLPV() && bReflectiveShadowmap) )
 		{
 			const bool bLocalDirectionalLight = Context.ShadowInfo->bDirectionalLight;
 			const bool bPreShadow = Context.ShadowInfo->bPreShadow;
