Fixed compile error

diff --git a/Engine/Source/Runtime/Renderer/Private/ShadowRendering.cpp b/Engine/Source/Runtime/Renderer/Private/ShadowRendering.cpp
--- a/Engine/Source/Runtime/Renderer/Private/ShadowRendering.cpp
+++ b/Engine/Source/Runtime/Renderer/Private/ShadowRendering.cpp
@@ -1075,10 +1075,10 @@ bool FShadowDepthDrawingPolicyFactory::DrawDynamicMesh(
 		const EMaterialShadingModel ShadingModel = Material->GetShadingModel();
 
 		const bool bLocalOnePassPointLightShadow = Context.ShadowInfo->CascadeSettings.bOnePassPointLightShadow;
-		const bool bReflectiveShadowmap = Context.ShadowInfo->bReflectiveShadowmap && !bOnePassPointLightShadow;
+		const bool bReflectiveShadowmap = Context.ShadowInfo->bReflectiveShadowmap && !bLocalOnePassPointLightShadow;
 
 		if ( (!IsTranslucentBlendMode(BlendMode) || bReflectiveShadowmap ) && ShadingModel != MSM_Unlit
-		|| (Material->ShouldInjectEmissiveIntoLPV() && bReflectiveShadowmap) )
+			|| (Material->ShouldInjectEmissiveIntoLPV() && bReflectiveShadowmap) )
 		{
 			const bool bLocalDirectionalLight = Context.ShadowInfo->bDirectionalLight;
 			const bool bPreShadow = Context.ShadowInfo->bPreShadow;
