Fix compile error (|| and && in the same scope).

- Also fixed an unrelated initialization order warning.

diff --git a/Engine/Source/Runtime/CoreUObject/Private/UObject/SavePackage.cpp b/Engine/Source/Runtime/CoreUObject/Private/UObject/SavePackage.cpp
--- a/Engine/Source/Runtime/CoreUObject/Private/UObject/SavePackage.cpp
+++ b/Engine/Source/Runtime/CoreUObject/Private/UObject/SavePackage.cpp
@@ -1809,7 +1809,7 @@ public:
 	 * Constructor
 	 */
 	FExportReferenceSorter()
-		: FArchiveUObject(), CurrentInsertIndex(INDEX_NONE), CurrentClass(NULL), bIgnoreFieldReferences(false), CoreReferencesOffset(INDEX_NONE)
+		: FArchiveUObject(), CurrentInsertIndex(INDEX_NONE), CoreReferencesOffset(INDEX_NONE), bIgnoreFieldReferences(false), CurrentClass(nullptr)
 	{
 		ArIsObjectReferenceCollector = true;
 		ArIsPersistent = true;
diff --git a/Engine/Source/Runtime/Renderer/Private/ShadowRendering.cpp b/Engine/Source/Runtime/Renderer/Private/ShadowRendering.cpp
--- a/Engine/Source/Runtime/Renderer/Private/ShadowRendering.cpp
+++ b/Engine/Source/Runtime/Renderer/Private/ShadowRendering.cpp
@@ -1077,7 +1077,7 @@ bool FShadowDepthDrawingPolicyFactory::DrawDynamicMesh(
 		const bool bLocalOnePassPointLightShadow = Context.ShadowInfo->CascadeSettings.bOnePassPointLightShadow;
 		const bool bReflectiveShadowmap = Context.ShadowInfo->bReflectiveShadowmap && !bLocalOnePassPointLightShadow;
 
-		if ( (!IsTranslucentBlendMode(BlendMode) || bReflectiveShadowmap ) && ShadingModel != MSM_Unlit
+		if ( ((!IsTranslucentBlendMode(BlendMode) || bReflectiveShadowmap ) && ShadingModel != MSM_Unlit)
 			|| (Material->ShouldInjectEmissiveIntoLPV() && bReflectiveShadowmap) )
 		{
 			const bool bLocalDirectionalLight = Context.ShadowInfo->bDirectionalLight;
