Distance Field AO is now evaluated on a screen grid and using a global distance field for far traces

diff --git a/Engine/Source/Runtime/Core/Public/Math/IntVector.h b/Engine/Source/Runtime/Core/Public/Math/IntVector.h
--- a/Engine/Source/Runtime/Core/Public/Math/IntVector.h
+++ b/Engine/Source/Runtime/Core/Public/Math/IntVector.h
@@ -73,6 +73,22 @@ public:
 	 */
 	int32& operator()( int32 ComponentIndex );
 
+	/**
+	 * Gets specific component of a point.
+	 *
+	 * @param ComponentIndex Index of point component.
+	 * @return const reference to component.
+	 */
+	const int32& operator[]( int32 ComponentIndex ) const;
+
+	/**
+	 * Gets specific component of a point.
+	 *
+	 * @param ComponentIndex Index of point component.
+	 * @return reference to component.
+	 */
+	int32& operator[]( int32 ComponentIndex );
+
 	/**
 	 * Compares points for equality.
 	 *
@@ -265,6 +281,17 @@ FORCEINLINE int32& FIntVector::operator()( int32 ComponentIndex )
 }
 
 
+FORCEINLINE const int32& FIntVector::operator[]( int32 ComponentIndex ) const
+{
+	return (&X)[ComponentIndex];
+}
+
+
+FORCEINLINE int32& FIntVector::operator[]( int32 ComponentIndex )
+{
+	return (&X)[ComponentIndex];
+}
+
 FORCEINLINE bool FIntVector::operator==( const FIntVector& Other ) const
 {
 	return X==Other.X && Y==Other.Y && Z==Other.Z;
diff --git a/Engine/Source/Runtime/Renderer/Private/DeferredShadingRenderer.cpp b/Engine/Source/Runtime/Renderer/Private/DeferredShadingRenderer.cpp
--- a/Engine/Source/Runtime/Renderer/Private/DeferredShadingRenderer.cpp
+++ b/Engine/Source/Runtime/Renderer/Private/DeferredShadingRenderer.cpp
@@ -19,6 +19,7 @@
 #include "DeferredShadingRenderer.h"
 #include "SceneUtils.h"
 #include "DistanceFieldSurfaceCacheLighting.h"
+#include "GlobalDistanceField.h"
 #include "PostProcess/PostProcessing.h"
 #include "DistanceFieldAtlas.h"
 #include "../../Engine/Private/SkeletalRenderGPUSkin.h"		// GPrevPerBoneMotionBlur
@@ -723,6 +724,13 @@ void FDeferredShadingSceneRenderer::Render(FRHICommandListImmediate& RHICmdList)
 		for (int32 ViewIndex = 0; ViewIndex < Views.Num(); ViewIndex++)
 		{
 			Views[ViewIndex].HeightfieldLightingViewInfo.SetupVisibleHeightfields(Views[ViewIndex], RHICmdList);
+
+			if (UseGlobalDistanceField())
+			{
+				// Use the skylight's max distance if there is one
+				const float OcclusionMaxDistance = Scene->SkyLight && !Scene->SkyLight->bWantsStaticShadowing ? Scene->SkyLight->OcclusionMaxDistance : GDefaultDFAOMaxOcclusionDistance;
+				UpdateGlobalDistanceFieldVolume(RHICmdList, Views[ViewIndex], Scene, OcclusionMaxDistance, Views[ViewIndex].GlobalDistanceFieldInfo);
+			}
 		}	
 	}
 
@@ -1109,14 +1117,14 @@ void FDeferredShadingSceneRenderer::Render(FRHICommandListImmediate& RHICmdList)
 	if (ViewFamily.EngineShowFlags.VisualizeDistanceFieldAO || ViewFamily.EngineShowFlags.VisualizeDistanceFieldGI)
 	{
 		// Use the skylight's max distance if there is one, to be consistent with DFAO shadowing on the skylight
-		const float OcclusionMaxDistance = Scene->SkyLight && !Scene->SkyLight->bWantsStaticShadowing ? Scene->SkyLight->OcclusionMaxDistance : 600;
+		const float OcclusionMaxDistance = Scene->SkyLight && !Scene->SkyLight->bWantsStaticShadowing ? Scene->SkyLight->OcclusionMaxDistance : GDefaultDFAOMaxOcclusionDistance;
 		TRefCountPtr<IPooledRenderTarget> DummyOutput;
-		RenderDistanceFieldAOSurfaceCache(RHICmdList, FDistanceFieldAOParameters(OcclusionMaxDistance), VelocityRT, DummyOutput, DummyOutput, ViewFamily.EngineShowFlags.VisualizeDistanceFieldAO, ViewFamily.EngineShowFlags.VisualizeDistanceFieldGI);
+		RenderDistanceFieldLighting(RHICmdList, FDistanceFieldAOParameters(OcclusionMaxDistance), VelocityRT, DummyOutput, DummyOutput, ViewFamily.EngineShowFlags.VisualizeDistanceFieldAO, ViewFamily.EngineShowFlags.VisualizeDistanceFieldGI);
 	}
 
 	if (ViewFamily.EngineShowFlags.VisualizeMeshDistanceFields)
 	{
-		RenderMeshDistanceFieldVisualization(RHICmdList, FDistanceFieldAOParameters(600));
+		RenderMeshDistanceFieldVisualization(RHICmdList, FDistanceFieldAOParameters(GDefaultDFAOMaxOcclusionDistance));
 	}
 
 	// Resolve the scene color for post processing.
diff --git a/Engine/Source/Runtime/Renderer/Private/DeferredShadingRenderer.h b/Engine/Source/Runtime/Renderer/Private/DeferredShadingRenderer.h
--- a/Engine/Source/Runtime/Renderer/Private/DeferredShadingRenderer.h
+++ b/Engine/Source/Runtime/Renderer/Private/DeferredShadingRenderer.h
@@ -191,7 +191,7 @@ private:
 	void RenderDynamicSkyLighting(FRHICommandListImmediate& RHICmdList, const TRefCountPtr<IPooledRenderTarget>& VelocityTexture, TRefCountPtr<IPooledRenderTarget>& DynamicBentNormalAO);
 
 	/** Render Ambient Occlusion using mesh distance fields and the surface cache, which supports dynamic rigid meshes. */
-	bool RenderDistanceFieldAOSurfaceCache(
+	bool RenderDistanceFieldLighting(
 		FRHICommandListImmediate& RHICmdList, 
 		const class FDistanceFieldAOParameters& Parameters, 
 		const TRefCountPtr<IPooledRenderTarget>& VelocityTexture,
@@ -200,6 +200,17 @@ private:
 		bool bVisualizeAmbientOcclusion,
 		bool bVisualizeGlobalIllumination);
 
+	/** Render Ambient Occlusion using mesh distance fields on a screen based grid. */
+	void RenderDistanceFieldAOScreenGrid(
+		FRHICommandListImmediate& RHICmdList, 
+		const FViewInfo& View,
+		FIntPoint TileListGroupSize,
+		const FDistanceFieldAOParameters& Parameters, 
+		const TRefCountPtr<IPooledRenderTarget>& VelocityTexture,
+		const TRefCountPtr<IPooledRenderTarget>& DistanceFieldNormal, 
+		TRefCountPtr<IPooledRenderTarget>& OutDynamicBentNormalAO, 
+		TRefCountPtr<IPooledRenderTarget>& OutDynamicIrradiance);
+
 	void RenderMeshDistanceFieldVisualization(FRHICommandListImmediate& RHICmdList, const FDistanceFieldAOParameters& Parameters);
 
 	/** Whether tiled deferred is supported and can be used at all. */
diff --git a/Engine/Source/Runtime/Renderer/Private/DistanceFieldGlobalIllumination.cpp b/Engine/Source/Runtime/Renderer/Private/DistanceFieldGlobalIllumination.cpp
--- a/Engine/Source/Runtime/Renderer/Private/DistanceFieldGlobalIllumination.cpp
+++ b/Engine/Source/Runtime/Renderer/Private/DistanceFieldGlobalIllumination.cpp
@@ -798,19 +798,19 @@ void UpdateVPLs(
 					FConvexVolume FrustumVolume;
 
 					{
-						const float ConeExpandDistance = Parameters.OcclusionMaxDistance;
+						const float ConeExpandDistance = Parameters.ObjectMaxOcclusionDistance;
 						const float TanHalfFOV = 1.0f / View.ViewMatrices.ProjMatrix.M[0][0];
 						const float VertexPullbackLength = ConeExpandDistance / TanHalfFOV;
 
 						// Pull back cone vertex to contain VPLs outside of the view
 						const FVector ViewConeVertex = View.ViewMatrices.ViewOrigin - View.GetViewDirection() * VertexPullbackLength;
 
-						//@todo - expand by AOMaxDistance
+						//@todo - expand by AOObjectMaxDistance
 						ShadowBounds = DirectionalLightProxy->GetShadowSplitBoundsDepthRange(
 							View, 
 							ViewConeVertex,
 							View.NearClippingDistance, 
-							GetMaxAOViewDistance() + VertexPullbackLength + Parameters.OcclusionMaxDistance, 
+							GetMaxAOViewDistance() + VertexPullbackLength + Parameters.ObjectMaxOcclusionDistance, 
 							&CascadeSettings); 
 
 						FSphere SubjectBounds(FVector::ZeroVector, ShadowBounds.W);
@@ -961,12 +961,6 @@ public:
 		return IsFeatureLevelSupported(Platform, ERHIFeatureLevel::SM5) && DoesPlatformSupportDistanceFieldGI(Platform);
 	}
 
-	static void ModifyCompilationEnvironment(EShaderPlatform Platform, FShaderCompilerEnvironment& OutEnvironment)
-	{
-		FGlobalShader::ModifyCompilationEnvironment(Platform,OutEnvironment);
-		OutEnvironment.SetDefine(TEXT("NUM_CONE_DIRECTIONS"), NumConeSampleDirections);
-	}
-
 	FComputeStepBentNormalCS(const ShaderMetaType::CompiledShaderInitializerType& Initializer)
 		: FGlobalShader(Initializer)
 	{
@@ -1074,7 +1068,6 @@ public:
 	static void ModifyCompilationEnvironment(EShaderPlatform Platform, FShaderCompilerEnvironment& OutEnvironment)
 	{
 		FGlobalShader::ModifyCompilationEnvironment(Platform,OutEnvironment);
-		OutEnvironment.SetDefine(TEXT("NUM_CONE_DIRECTIONS"), NumConeSampleDirections);
 		OutEnvironment.SetDefine(TEXT("IRRADIANCE_FROM_SURFELS"), VPLMode == VPLMode_PlaceFromSurfels ? TEXT("1") : TEXT("0"));
 
 		// To reduce shader compile time of compute shaders with shared memory, doesn't have an impact on generated code with current compiler (June 2010 DX SDK)
@@ -1102,6 +1095,7 @@ public:
 		TileListGroupSize.Bind(Initializer.ParameterMap, TEXT("TileListGroupSize"));
 		TileHeadDataUnpacked.Bind(Initializer.ParameterMap, TEXT("TileHeadDataUnpacked"));
 		TileArrayData.Bind(Initializer.ParameterMap, TEXT("TileArrayData"));
+		TileConeDepthRanges.Bind(Initializer.ParameterMap, TEXT("TileConeDepthRanges"));
 		StepBentNormalBuffer.Bind(Initializer.ParameterMap, TEXT("StepBentNormalBuffer"));
 	}
 
@@ -1141,12 +1135,13 @@ public:
 		SetSRVParameter(RHICmdList, ShaderRHI, VPLParameterBuffer, SourceVPLResources.VPLParameterBuffer.SRV);
 		SetSRVParameter(RHICmdList, ShaderRHI, VPLData, SourceVPLResources.VPLData.SRV);
 
-		SetShaderValue(RHICmdList, ShaderRHI, VPLGatherRadius, Parameters.OcclusionMaxDistance);
+		SetShaderValue(RHICmdList, ShaderRHI, VPLGatherRadius, Parameters.ObjectMaxOcclusionDistance);
 
 		FTileIntersectionResources* TileIntersectionResources = ((FSceneViewState*)View.State)->AOTileIntersectionResources;
 
 		SetSRVParameter(RHICmdList, ShaderRHI, TileHeadDataUnpacked, TileIntersectionResources->TileHeadDataUnpacked.SRV);
 		SetSRVParameter(RHICmdList, ShaderRHI, TileArrayData, TileIntersectionResources->TileArrayData.SRV);
+		SetSRVParameter(RHICmdList, ShaderRHI, TileConeDepthRanges, TileIntersectionResources->TileConeDepthRanges.SRV);
 		SetShaderValue(RHICmdList, ShaderRHI, TileListGroupSize, TileIntersectionResources->TileDimensions);
 
 		SetSRVParameter(RHICmdList, ShaderRHI, StepBentNormalBuffer, TemporaryIrradianceCacheResources->StepBentNormal.SRV);
@@ -1178,6 +1173,7 @@ public:
 		Ar << TileListGroupSize;
 		Ar << TileHeadDataUnpacked;
 		Ar << TileArrayData;
+		Ar << TileConeDepthRanges;
 		Ar << StepBentNormalBuffer;
 		return bShaderHasOutdatedParameters;
 	}
@@ -1202,6 +1198,7 @@ private:
 	FShaderParameter TileListGroupSize;
 	FShaderResourceParameter TileHeadDataUnpacked;
 	FShaderResourceParameter TileArrayData;
+	FShaderResourceParameter TileConeDepthRanges;
 	FShaderResourceParameter StepBentNormalBuffer;
 };
 
@@ -1378,6 +1375,359 @@ void ComputeIrradianceForSamples(
 	}
 }
 
+
+const int32 GScreenGridIrradianceThreadGroupSizeX = 8;
+
+class FComputeStepBentNormalScreenGridCS : public FGlobalShader
+{
+	DECLARE_SHADER_TYPE(FComputeStepBentNormalScreenGridCS,Global)
+public:
+
+	static bool ShouldCache(EShaderPlatform Platform)
+	{
+		return IsFeatureLevelSupported(Platform, ERHIFeatureLevel::SM5) && DoesPlatformSupportDistanceFieldAO(Platform);
+	}
+
+	static void ModifyCompilationEnvironment(EShaderPlatform Platform, FShaderCompilerEnvironment& OutEnvironment)
+	{
+		FGlobalShader::ModifyCompilationEnvironment(Platform,OutEnvironment);
+		OutEnvironment.SetDefine(TEXT("SCREEN_GRID_IRRADIANCE_THREADGROUP_SIZE_X"), GScreenGridIrradianceThreadGroupSizeX);
+		extern int32 GConeTraceDownsampleFactor;
+		OutEnvironment.SetDefine(TEXT("TRACE_DOWNSAMPLE_FACTOR"), GConeTraceDownsampleFactor);
+
+		// To reduce shader compile time of compute shaders with shared memory, doesn't have an impact on generated code with current compiler (June 2010 DX SDK)
+		OutEnvironment.CompilerFlags.Add(CFLAG_StandardOptimization);
+	}
+
+	FComputeStepBentNormalScreenGridCS(const ShaderMetaType::CompiledShaderInitializerType& Initializer)
+		: FGlobalShader(Initializer)
+	{
+		ScreenGridParameters.Bind(Initializer.ParameterMap);
+		BentNormalNormalizeFactor.Bind(Initializer.ParameterMap, TEXT("BentNormalNormalizeFactor"));
+		ConeDepthVisibilityFunction.Bind(Initializer.ParameterMap, TEXT("ConeDepthVisibilityFunction"));
+		StepBentNormal.Bind(Initializer.ParameterMap, TEXT("StepBentNormal"));
+	}
+
+	FComputeStepBentNormalScreenGridCS()
+	{
+	}
+
+	void SetParameters(
+		FRHICommandList& RHICmdList, 
+		const FViewInfo& View, 
+		FSceneRenderTargetItem& DistanceFieldNormal, 
+		const FAOScreenGridResources& ScreenGridResources)
+	{
+		FComputeShaderRHIParamRef ShaderRHI = GetComputeShader();
+		FGlobalShader::SetParameters(RHICmdList, ShaderRHI, View);
+		ScreenGridParameters.Set(RHICmdList, ShaderRHI, View, DistanceFieldNormal);
+
+		FAOSampleData2 AOSampleData;
+
+		TArray<FVector, TInlineAllocator<9> > SampleDirections;
+		GetSpacedVectors(SampleDirections);
+
+		for (int32 SampleIndex = 0; SampleIndex < NumConeSampleDirections; SampleIndex++)
+		{
+			AOSampleData.SampleDirections[SampleIndex] = FVector4(SampleDirections[SampleIndex]);
+		}
+
+		SetUniformBufferParameterImmediate(RHICmdList, ShaderRHI, GetUniformBufferParameter<FAOSampleData2>(), AOSampleData);
+
+		FVector UnoccludedVector(0);
+
+		for (int32 SampleIndex = 0; SampleIndex < NumConeSampleDirections; SampleIndex++)
+		{
+			UnoccludedVector += SampleDirections[SampleIndex];
+		}
+
+		float BentNormalNormalizeFactorValue = 1.0f / (UnoccludedVector / NumConeSampleDirections).Size();
+		SetShaderValue(RHICmdList, ShaderRHI, BentNormalNormalizeFactor, BentNormalNormalizeFactorValue);
+
+		SetSRVParameter(RHICmdList, ShaderRHI, ConeDepthVisibilityFunction, ScreenGridResources.ConeDepthVisibilityFunction.SRV);
+
+		StepBentNormal.SetBuffer(RHICmdList, ShaderRHI, ScreenGridResources.StepBentNormal);
+	}
+
+	void UnsetParameters(FRHICommandList& RHICmdList)
+	{
+		StepBentNormal.UnsetUAV(RHICmdList, GetComputeShader());
+	}
+
+	virtual bool Serialize(FArchive& Ar) override
+	{		
+		bool bShaderHasOutdatedParameters = FGlobalShader::Serialize(Ar);
+		Ar << ScreenGridParameters;
+		Ar << BentNormalNormalizeFactor;
+		Ar << ConeDepthVisibilityFunction;
+		Ar << StepBentNormal;
+		return bShaderHasOutdatedParameters;
+	}
+
+private:
+
+	FScreenGridParameters ScreenGridParameters;
+	FShaderParameter BentNormalNormalizeFactor;
+	FShaderResourceParameter ConeDepthVisibilityFunction;
+	FRWShaderParameter StepBentNormal;
+};
+
+IMPLEMENT_SHADER_TYPE(,FComputeStepBentNormalScreenGridCS,TEXT("DistanceFieldGlobalIllumination"),TEXT("ComputeStepBentNormalScreenGridCS"),SF_Compute);
+
+
+class FComputeIrradianceScreenGridCS : public FGlobalShader
+{
+	DECLARE_SHADER_TYPE(FComputeIrradianceScreenGridCS,Global)
+public:
+
+	static bool ShouldCache(EShaderPlatform Platform)
+	{
+		return IsFeatureLevelSupported(Platform, ERHIFeatureLevel::SM5) && DoesPlatformSupportDistanceFieldAO(Platform);
+	}
+
+	static void ModifyCompilationEnvironment(EShaderPlatform Platform, FShaderCompilerEnvironment& OutEnvironment)
+	{
+		FGlobalShader::ModifyCompilationEnvironment(Platform,OutEnvironment);
+		OutEnvironment.SetDefine(TEXT("CULLED_TILE_SIZEX"), GDistanceFieldAOTileSizeX);
+		extern int32 GConeTraceDownsampleFactor;
+		OutEnvironment.SetDefine(TEXT("TRACE_DOWNSAMPLE_FACTOR"), GConeTraceDownsampleFactor);
+		OutEnvironment.SetDefine(TEXT("IRRADIANCE_FROM_SURFELS"), TEXT("1"));
+
+		// To reduce shader compile time of compute shaders with shared memory, doesn't have an impact on generated code with current compiler (June 2010 DX SDK)
+		OutEnvironment.CompilerFlags.Add(CFLAG_StandardOptimization);
+	}
+
+	FComputeIrradianceScreenGridCS(const ShaderMetaType::CompiledShaderInitializerType& Initializer)
+		: FGlobalShader(Initializer)
+	{
+		DeferredParameters.Bind(Initializer.ParameterMap);
+		ObjectParameters.Bind(Initializer.ParameterMap);
+		AOParameters.Bind(Initializer.ParameterMap);
+		ScreenGridParameters.Bind(Initializer.ParameterMap);
+		SurfelParameters.Bind(Initializer.ParameterMap);
+		TileHeadData.Bind(Initializer.ParameterMap, TEXT("TileHeadData"));
+		TileHeadDataUnpacked.Bind(Initializer.ParameterMap, TEXT("TileHeadDataUnpacked"));
+		TileArrayData.Bind(Initializer.ParameterMap, TEXT("TileArrayData"));
+		TileConeDepthRanges.Bind(Initializer.ParameterMap, TEXT("TileConeDepthRanges"));
+		TileListGroupSize.Bind(Initializer.ParameterMap, TEXT("TileListGroupSize"));
+		VPLGatherRadius.Bind(Initializer.ParameterMap, TEXT("VPLGatherRadius"));
+		StepBentNormalBuffer.Bind(Initializer.ParameterMap, TEXT("StepBentNormalBuffer"));
+		SurfelIrradiance.Bind(Initializer.ParameterMap, TEXT("SurfelIrradiance"));
+	}
+
+	FComputeIrradianceScreenGridCS()
+	{
+	}
+
+	void SetParameters(
+		FRHICommandList& RHICmdList, 
+		const FViewInfo& View, 
+		FSceneRenderTargetItem& DistanceFieldNormal, 
+		const FDistanceFieldAOParameters& Parameters)
+	{
+		FComputeShaderRHIParamRef ShaderRHI = GetComputeShader();
+		FGlobalShader::SetParameters(RHICmdList, ShaderRHI, View);
+		DeferredParameters.Set(RHICmdList, ShaderRHI, View);
+
+		extern TGlobalResource<FDistanceFieldObjectBufferResource> GAOCulledObjectBuffers;
+		ObjectParameters.Set(RHICmdList, ShaderRHI, GAOCulledObjectBuffers.Buffers);
+		AOParameters.Set(RHICmdList, ShaderRHI, Parameters);
+		ScreenGridParameters.Set(RHICmdList, ShaderRHI, View, DistanceFieldNormal);
+
+		const FScene* Scene = (const FScene*)View.Family->Scene;
+		SurfelParameters.Set(RHICmdList, ShaderRHI, *Scene->DistanceFieldSceneData.SurfelBuffers, *Scene->DistanceFieldSceneData.InstancedSurfelBuffers);
+
+		FTileIntersectionResources* TileIntersectionResources = View.ViewState->AOTileIntersectionResources;
+
+		SetSRVParameter(RHICmdList, ShaderRHI, TileHeadDataUnpacked, TileIntersectionResources->TileHeadDataUnpacked.SRV);
+		SetSRVParameter(RHICmdList, ShaderRHI, TileHeadData, TileIntersectionResources->TileHeadData.SRV);
+		SetSRVParameter(RHICmdList, ShaderRHI, TileArrayData, TileIntersectionResources->TileArrayData.SRV);
+		SetSRVParameter(RHICmdList, ShaderRHI, TileConeDepthRanges, TileIntersectionResources->TileConeDepthRanges.SRV);
+		SetShaderValue(RHICmdList, ShaderRHI, TileListGroupSize, TileIntersectionResources->TileDimensions);
+
+		SetShaderValue(RHICmdList, ShaderRHI, VPLGatherRadius, Parameters.ObjectMaxOcclusionDistance);
+
+		FAOScreenGridResources* ScreenGridResources = View.ViewState->AOScreenGridResources;
+
+		SetSRVParameter(RHICmdList, ShaderRHI, StepBentNormalBuffer, ScreenGridResources->StepBentNormal.SRV);
+		SurfelIrradiance.SetBuffer(RHICmdList, ShaderRHI, ScreenGridResources->SurfelIrradiance);
+	}
+
+	void UnsetParameters(FRHICommandList& RHICmdList)
+	{
+		SurfelIrradiance.UnsetUAV(RHICmdList, GetComputeShader());
+	}
+
+	virtual bool Serialize(FArchive& Ar) override
+	{		
+		bool bShaderHasOutdatedParameters = FGlobalShader::Serialize(Ar);
+		Ar << DeferredParameters;
+		Ar << ObjectParameters;
+		Ar << AOParameters;
+		Ar << ScreenGridParameters;
+		Ar << SurfelParameters;
+		Ar << TileHeadDataUnpacked;
+		Ar << TileHeadData;
+		Ar << TileArrayData;
+		Ar << TileConeDepthRanges;
+		Ar << TileListGroupSize;
+		Ar << VPLGatherRadius;
+		Ar << StepBentNormalBuffer;
+		Ar << SurfelIrradiance;
+		return bShaderHasOutdatedParameters;
+	}
+
+private:
+
+	FDeferredPixelShaderParameters DeferredParameters;
+	FDistanceFieldCulledObjectBufferParameters ObjectParameters;
+	FAOParameters AOParameters;
+	FScreenGridParameters ScreenGridParameters;
+	FSurfelBufferParameters SurfelParameters;
+	FShaderResourceParameter TileHeadDataUnpacked;
+	FShaderResourceParameter TileHeadData;
+	FShaderResourceParameter TileArrayData;
+	FShaderResourceParameter TileConeDepthRanges;
+	FShaderParameter TileListGroupSize;
+	FShaderParameter VPLGatherRadius;
+	FShaderResourceParameter StepBentNormalBuffer;
+	FRWShaderParameter SurfelIrradiance;
+};
+
+IMPLEMENT_SHADER_TYPE(,FComputeIrradianceScreenGridCS,TEXT("DistanceFieldGlobalIllumination"),TEXT("ComputeIrradianceScreenGridCS"),SF_Compute);
+
+
+class FCombineIrradianceScreenGridCS : public FGlobalShader
+{
+	DECLARE_SHADER_TYPE(FCombineIrradianceScreenGridCS,Global)
+public:
+
+	static bool ShouldCache(EShaderPlatform Platform)
+	{
+		return IsFeatureLevelSupported(Platform, ERHIFeatureLevel::SM5) && DoesPlatformSupportDistanceFieldGI(Platform);
+	}
+
+	static void ModifyCompilationEnvironment(EShaderPlatform Platform, FShaderCompilerEnvironment& OutEnvironment)
+	{
+		OutEnvironment.SetDefine(TEXT("SCREEN_GRID_IRRADIANCE_THREADGROUP_SIZE_X"), GScreenGridIrradianceThreadGroupSizeX);
+		FGlobalShader::ModifyCompilationEnvironment(Platform,OutEnvironment);
+	}
+
+	FCombineIrradianceScreenGridCS(const ShaderMetaType::CompiledShaderInitializerType& Initializer)
+		: FGlobalShader(Initializer)
+	{
+		IrradianceTexture.Bind(Initializer.ParameterMap, TEXT("IrradianceTexture"));
+		SurfelIrradiance.Bind(Initializer.ParameterMap, TEXT("SurfelIrradiance"));
+		HeightfieldIrradiance.Bind(Initializer.ParameterMap, TEXT("HeightfieldIrradiance"));
+		ScreenGridConeVisibilitySize.Bind(Initializer.ParameterMap, TEXT("ScreenGridConeVisibilitySize"));
+	}
+
+	FCombineIrradianceScreenGridCS()
+	{
+	}
+
+	void SetParameters(
+		FRHICommandList& RHICmdList, 
+		const FSceneView& View,
+		const FAOScreenGridResources& ScreenGridResources,
+		FSceneRenderTargetItem& IrradianceTextureValue)
+	{
+		FComputeShaderRHIParamRef ShaderRHI = GetComputeShader();
+		FGlobalShader::SetParameters(RHICmdList, ShaderRHI, View);
+
+		const FScene* Scene = (const FScene*)View.Family->Scene;
+		FSurfaceCacheResources& SurfaceCacheResources = *Scene->SurfaceCacheResources;
+
+		SetSRVParameter(RHICmdList, ShaderRHI, SurfelIrradiance, ScreenGridResources.SurfelIrradiance.SRV);
+		SetSRVParameter(RHICmdList, ShaderRHI, HeightfieldIrradiance, ScreenGridResources.HeightfieldIrradiance.SRV);
+
+		IrradianceTexture.SetTexture(RHICmdList, ShaderRHI, IrradianceTextureValue.ShaderResourceTexture, IrradianceTextureValue.UAV);
+
+		SetShaderValue(RHICmdList, ShaderRHI, ScreenGridConeVisibilitySize, ScreenGridResources.ScreenGridDimensions);
+	}
+
+	void UnsetParameters(FRHICommandList& RHICmdList)
+	{
+		IrradianceTexture.UnsetUAV(RHICmdList, GetComputeShader());
+	}
+
+	virtual bool Serialize(FArchive& Ar) override
+	{		
+		bool bShaderHasOutdatedParameters = FGlobalShader::Serialize(Ar);
+		Ar << IrradianceTexture;
+		Ar << SurfelIrradiance;
+		Ar << HeightfieldIrradiance;
+		Ar << ScreenGridConeVisibilitySize;
+		return bShaderHasOutdatedParameters;
+	}
+
+private:
+
+	FRWShaderParameter IrradianceTexture;
+	FShaderResourceParameter SurfelIrradiance;
+	FShaderResourceParameter HeightfieldIrradiance;
+	FShaderParameter ScreenGridConeVisibilitySize;
+};
+
+IMPLEMENT_SHADER_TYPE(,FCombineIrradianceScreenGridCS,TEXT("DistanceFieldGlobalIllumination"),TEXT("CombineIrradianceScreenGridCS"),SF_Compute);
+
+
+void ComputeIrradianceForScreenGrid(
+	FRHICommandListImmediate& RHICmdList,
+	const FViewInfo& View,
+	const FScene* Scene,
+	const FDistanceFieldAOParameters& Parameters,
+	FSceneRenderTargetItem& DistanceFieldNormal, 
+	const FAOScreenGridResources& ScreenGridResources,
+	FSceneRenderTargetItem& IrradianceTexture)
+{
+	const uint32 GroupSizeX = FMath::DivideAndRoundUp(View.ViewRect.Size().X / GAODownsampleFactor, GScreenGridIrradianceThreadGroupSizeX);
+	const uint32 GroupSizeY = FMath::DivideAndRoundUp(View.ViewRect.Size().Y / GAODownsampleFactor, GScreenGridIrradianceThreadGroupSizeX);
+
+	uint32 ClearValues[4] = { 0 };
+	RHICmdList.ClearUAV(ScreenGridResources.HeightfieldIrradiance.UAV, ClearValues);
+	RHICmdList.ClearUAV(ScreenGridResources.SurfelIrradiance.UAV, ClearValues);
+
+	View.HeightfieldLightingViewInfo.
+		ComputeIrradianceForScreenGrid(View, RHICmdList, DistanceFieldNormal, ScreenGridResources, Parameters);
+	
+	if (GVPLMeshGlobalIllumination)
+	{
+		{
+			SCOPED_DRAW_EVENT(RHICmdList, ComputeStepBentNormal);
+
+			TShaderMapRef<FComputeStepBentNormalScreenGridCS> ComputeShader(View.ShaderMap);
+			RHICmdList.SetComputeShader(ComputeShader->GetComputeShader());
+			ComputeShader->SetParameters(RHICmdList, View, DistanceFieldNormal, ScreenGridResources);
+			DispatchComputeShader(RHICmdList, *ComputeShader, GroupSizeX, GroupSizeY, 1);
+			ComputeShader->UnsetParameters(RHICmdList);
+		}
+
+		if (GVPLSurfelRepresentation)
+		{
+			SCOPED_DRAW_EVENT(RHICmdList, MeshIrradiance);
+
+			TShaderMapRef<FComputeIrradianceScreenGridCS> ComputeShader(View.ShaderMap);
+			RHICmdList.SetComputeShader(ComputeShader->GetComputeShader());
+			ComputeShader->SetParameters(RHICmdList, View, DistanceFieldNormal, Parameters);
+
+			const uint32 ComputeIrradianceGroupSizeX = View.ViewRect.Size().X / GAODownsampleFactor;
+			const uint32 ComputeIrradianceGroupSizeY = View.ViewRect.Size().Y / GAODownsampleFactor;
+			DispatchComputeShader(RHICmdList, *ComputeShader, ComputeIrradianceGroupSizeX, ComputeIrradianceGroupSizeY, 1);
+			ComputeShader->UnsetParameters(RHICmdList);
+		}
+	}
+		
+	{
+		TShaderMapRef<FCombineIrradianceScreenGridCS> ComputeShader(View.ShaderMap);
+		RHICmdList.SetComputeShader(ComputeShader->GetComputeShader());
+		ComputeShader->SetParameters(RHICmdList, View, ScreenGridResources, IrradianceTexture);
+		DispatchComputeShader(RHICmdList, *ComputeShader, GroupSizeX, GroupSizeY, 1);
+		ComputeShader->UnsetParameters(RHICmdList);
+	}
+}
+
+
 void ListDistanceFieldGIMemory(const FViewInfo& View)
 {
 	const FScene* Scene = (const FScene*)View.Family->Scene;
diff --git a/Engine/Source/Runtime/Renderer/Private/DistanceFieldLightingPost.cpp b/Engine/Source/Runtime/Renderer/Private/DistanceFieldLightingPost.cpp
new file mode 100644
--- /dev/null
+++ b/Engine/Source/Runtime/Renderer/Private/DistanceFieldLightingPost.cpp
@@ -0,0 +1,1100 @@
+// Copyright 1998-2015 Epic Games, Inc. All Rights Reserved.
+
+/*=============================================================================
+	DistanceFieldLightingPost.cpp
+=============================================================================*/
+
+#include "RendererPrivate.h"
+#include "ScenePrivate.h"
+#include "UniformBuffer.h"
+#include "ShaderParameters.h"
+#include "PostProcessing.h"
+#include "SceneFilterRendering.h"
+#include "DistanceFieldLightingShared.h"
+#include "DistanceFieldSurfaceCacheLighting.h"
+#include "DistanceFieldGlobalIllumination.h"
+#include "PostProcessAmbientOcclusion.h"
+#include "RHICommandList.h"
+#include "SceneUtils.h"
+#include "OneColorShader.h"
+#include "BasePassRendering.h"
+#include "HeightfieldLighting.h"
+
+int32 GAOFillGaps = 1;
+FAutoConsoleVariableRef CVarAOFillGaps(
+	TEXT("r.AOFillGaps"),
+	GAOFillGaps,
+	TEXT("Whether to fill in pixels using a screen space filter that had no valid world space interpolation weight from surface cache samples.\n")
+	TEXT("This is needed whenever r.AOMinLevel is not 0."),
+	ECVF_Cheat | ECVF_RenderThreadSafe
+	);
+
+int32 GAOFillGapsHighQuality = 1;
+FAutoConsoleVariableRef CVarAOFillGapsHighQuality(
+	TEXT("r.AOFillGapsHighQuality"),
+	GAOFillGapsHighQuality,
+	TEXT("Whether to use the higher quality gap filling method that does a 5x5 gather, or the cheaper method that just looks up the 4 grid samples."),
+	ECVF_Cheat | ECVF_RenderThreadSafe
+	);
+
+int32 GAOUseHistory = 1;
+FAutoConsoleVariableRef CVarAOUseHistory(
+	TEXT("r.AOUseHistory"),
+	GAOUseHistory,
+	TEXT("Whether to apply a temporal filter to the distance field AO, which reduces flickering but also adds trails when occluders are moving."),
+	ECVF_Cheat | ECVF_RenderThreadSafe
+	);
+
+int32 GAOHistoryStabilityPass = 1;
+FAutoConsoleVariableRef CVarAOHistoryStabilityPass(
+	TEXT("r.AOHistoryStabilityPass"),
+	GAOHistoryStabilityPass,
+	TEXT("Whether to gather stable results to fill in holes in the temporal reprojection.  Adds some GPU cost but improves temporal stability with foliage."),
+	ECVF_Cheat | ECVF_RenderThreadSafe
+	);
+
+float GAOHistoryWeight = .85f;
+FAutoConsoleVariableRef CVarAOHistoryWeight(
+	TEXT("r.AOHistoryWeight"),
+	GAOHistoryWeight,
+	TEXT("Amount of last frame's AO to lerp into the final result.  Higher values increase stability, lower values have less streaking under occluder movement."),
+	ECVF_Cheat | ECVF_RenderThreadSafe
+	);
+
+float GAOHistoryDistanceThreshold = 30;
+FAutoConsoleVariableRef CVarAOHistoryDistanceThreshold(
+	TEXT("r.AOHistoryDistanceThreshold"),
+	GAOHistoryDistanceThreshold,
+	TEXT("World space distance threshold needed to discard last frame's DFAO results.  Lower values reduce ghosting from characters when near a wall but increase flickering artifacts."),
+	ECVF_Cheat | ECVF_RenderThreadSafe
+	);
+
+float GAOViewFadeDistanceScale = .7f;
+FAutoConsoleVariableRef CVarAOViewFadeDistanceScale(
+	TEXT("r.AOViewFadeDistanceScale"),
+	GAOViewFadeDistanceScale,
+	TEXT("Distance over which AO will fade out as it approaches r.AOMaxViewDistance, as a fraction of r.AOMaxViewDistance."),
+	ECVF_Cheat | ECVF_RenderThreadSafe
+	);
+
+template<bool bSupportIrradiance>
+class TDistanceFieldAOCombinePS : public FGlobalShader
+{
+	DECLARE_SHADER_TYPE(TDistanceFieldAOCombinePS, Global);
+public:
+
+	static bool ShouldCache(EShaderPlatform Platform)
+	{
+		return IsFeatureLevelSupported(Platform, ERHIFeatureLevel::SM5) && DoesPlatformSupportDistanceFieldAO(Platform);
+	}
+
+	static void ModifyCompilationEnvironment(EShaderPlatform Platform, FShaderCompilerEnvironment& OutEnvironment)
+	{
+		OutEnvironment.SetDefine(TEXT("DOWNSAMPLE_FACTOR"), GAODownsampleFactor);
+		OutEnvironment.SetDefine(TEXT("SUPPORT_IRRADIANCE"), bSupportIrradiance ? TEXT("1") : TEXT("0"));
+	}
+
+	/** Default constructor. */
+	TDistanceFieldAOCombinePS() {}
+
+	/** Initialization constructor. */
+	TDistanceFieldAOCombinePS(const ShaderMetaType::CompiledShaderInitializerType& Initializer)
+		: FGlobalShader(Initializer)
+	{
+		DeferredParameters.Bind(Initializer.ParameterMap);
+		AOParameters.Bind(Initializer.ParameterMap);
+		AOLevelParameters.Bind(Initializer.ParameterMap);
+		BentNormalAOTexture.Bind(Initializer.ParameterMap,TEXT("BentNormalAOTexture"));
+		BentNormalAOSampler.Bind(Initializer.ParameterMap,TEXT("BentNormalAOSampler"));
+		IrradianceTexture.Bind(Initializer.ParameterMap,TEXT("IrradianceTexture"));
+		IrradianceSampler.Bind(Initializer.ParameterMap,TEXT("IrradianceSampler"));
+		DistanceFieldNormalTexture.Bind(Initializer.ParameterMap, TEXT("DistanceFieldNormalTexture"));
+		DistanceFieldNormalSampler.Bind(Initializer.ParameterMap, TEXT("DistanceFieldNormalSampler"));
+		DistanceFadeScale.Bind(Initializer.ParameterMap, TEXT("DistanceFadeScale"));
+		SelfOcclusionReplacement.Bind(Initializer.ParameterMap, TEXT("SelfOcclusionReplacement"));
+	}
+
+	void SetParameters(
+		FRHICommandList& RHICmdList, 
+		const FSceneView& View, 
+		FSceneRenderTargetItem& InBentNormalTexture, 
+		IPooledRenderTarget* InIrradianceTexture, 
+		FSceneRenderTargetItem& DistanceFieldNormal, 
+		const FDistanceFieldAOParameters& Parameters)
+	{
+		const FPixelShaderRHIParamRef ShaderRHI = GetPixelShader();
+
+		FGlobalShader::SetParameters(RHICmdList, ShaderRHI, View);
+		DeferredParameters.Set(RHICmdList, ShaderRHI, View);
+		AOParameters.Set(RHICmdList, ShaderRHI, Parameters);
+		AOLevelParameters.Set(RHICmdList, ShaderRHI, View, GAODownsampleFactor);
+
+		SetTextureParameter(RHICmdList, ShaderRHI, BentNormalAOTexture, BentNormalAOSampler, TStaticSamplerState<SF_Point>::GetRHI(), InBentNormalTexture.ShaderResourceTexture);
+
+		if (IrradianceTexture.IsBound())
+		{
+			SetTextureParameter(RHICmdList, ShaderRHI, IrradianceTexture, IrradianceSampler, TStaticSamplerState<SF_Point>::GetRHI(), InIrradianceTexture->GetRenderTargetItem().ShaderResourceTexture);
+		}
+
+		SetTextureParameter(
+			RHICmdList,
+			ShaderRHI,
+			DistanceFieldNormalTexture,
+			DistanceFieldNormalSampler,
+			TStaticSamplerState<SF_Point,AM_Clamp,AM_Clamp,AM_Clamp>::GetRHI(),
+			DistanceFieldNormal.ShaderResourceTexture
+			);
+
+		const float DistanceFadeScaleValue = 1.0f / ((1.0f - GAOViewFadeDistanceScale) * GetMaxAOViewDistance());
+		SetShaderValue(RHICmdList, ShaderRHI, DistanceFadeScale, DistanceFadeScaleValue);
+
+		extern float GVPLSelfOcclusionReplacement;
+		SetShaderValue(RHICmdList, ShaderRHI, SelfOcclusionReplacement, GVPLSelfOcclusionReplacement);
+	}
+
+	// FShader interface.
+	virtual bool Serialize(FArchive& Ar) override
+	{
+		bool bShaderHasOutdatedParameters = FGlobalShader::Serialize(Ar);
+		Ar << DeferredParameters;
+		Ar << AOParameters;
+		Ar << AOLevelParameters;
+		Ar << BentNormalAOTexture;
+		Ar << BentNormalAOSampler;
+		Ar << IrradianceTexture;
+		Ar << IrradianceSampler;
+		Ar << DistanceFieldNormalTexture;
+		Ar << DistanceFieldNormalSampler;
+		Ar << DistanceFadeScale;
+		Ar << SelfOcclusionReplacement;
+		return bShaderHasOutdatedParameters;
+	}
+
+private:
+
+	FDeferredPixelShaderParameters DeferredParameters;
+	FAOParameters AOParameters;
+	FAOLevelParameters AOLevelParameters;
+	FShaderResourceParameter BentNormalAOTexture;
+	FShaderResourceParameter BentNormalAOSampler;
+	FShaderResourceParameter IrradianceTexture;
+	FShaderResourceParameter IrradianceSampler;
+	FShaderResourceParameter DistanceFieldNormalTexture;
+	FShaderResourceParameter DistanceFieldNormalSampler;
+	FShaderParameter DistanceFadeScale;
+	FShaderParameter SelfOcclusionReplacement;
+};
+
+IMPLEMENT_SHADER_TYPE(template<>,TDistanceFieldAOCombinePS<true>,TEXT("DistanceFieldLightingPost"),TEXT("AOCombinePS"),SF_Pixel);
+IMPLEMENT_SHADER_TYPE(template<>,TDistanceFieldAOCombinePS<false>,TEXT("DistanceFieldLightingPost"),TEXT("AOCombinePS"),SF_Pixel);
+
+template<bool bSupportIrradiance, bool bHighQuality>
+class TFillGapsPS : public FGlobalShader
+{
+	DECLARE_SHADER_TYPE(TFillGapsPS, Global);
+public:
+
+	static bool ShouldCache(EShaderPlatform Platform)
+	{
+		return IsFeatureLevelSupported(Platform, ERHIFeatureLevel::SM5) && DoesPlatformSupportDistanceFieldAO(Platform);
+	}
+
+	static void ModifyCompilationEnvironment(EShaderPlatform Platform, FShaderCompilerEnvironment& OutEnvironment)
+	{
+		OutEnvironment.SetDefine(TEXT("DOWNSAMPLE_FACTOR"), GAODownsampleFactor);
+		OutEnvironment.SetDefine(TEXT("SUPPORT_IRRADIANCE"), bSupportIrradiance ? TEXT("1") : TEXT("0"));
+		OutEnvironment.SetDefine(TEXT("HIGH_QUALITY_FILL_GAPS"), bHighQuality ? TEXT("1") : TEXT("0"));
+	}
+
+	/** Default constructor. */
+	TFillGapsPS() {}
+
+	/** Initialization constructor. */
+	TFillGapsPS(const ShaderMetaType::CompiledShaderInitializerType& Initializer)
+		: FGlobalShader(Initializer)
+	{
+		BentNormalAOTexture.Bind(Initializer.ParameterMap, TEXT("BentNormalAOTexture"));
+		BentNormalAOSampler.Bind(Initializer.ParameterMap, TEXT("BentNormalAOSampler"));
+		IrradianceTexture.Bind(Initializer.ParameterMap, TEXT("IrradianceTexture"));
+		IrradianceSampler.Bind(Initializer.ParameterMap, TEXT("IrradianceSampler"));
+		BentNormalAOTexelSize.Bind(Initializer.ParameterMap, TEXT("BentNormalAOTexelSize"));
+		MinDownsampleFactorToBaseLevel.Bind(Initializer.ParameterMap, TEXT("MinDownsampleFactorToBaseLevel"));
+	}
+
+	void SetParameters(FRHICommandList& RHICmdList, const FSceneView& View, FSceneRenderTargetItem& DistanceFieldAOBentNormal2, IPooledRenderTarget* DistanceFieldIrradiance2)
+	{
+		const FPixelShaderRHIParamRef ShaderRHI = GetPixelShader();
+
+		FGlobalShader::SetParameters(RHICmdList, ShaderRHI, View);
+
+		SetTextureParameter(
+			RHICmdList,
+			ShaderRHI,
+			BentNormalAOTexture,
+			BentNormalAOSampler,
+			TStaticSamplerState<SF_Bilinear,AM_Clamp,AM_Clamp,AM_Clamp>::GetRHI(),
+			DistanceFieldAOBentNormal2.ShaderResourceTexture
+			);
+
+		if (IrradianceTexture.IsBound())
+		{
+			SetTextureParameter(
+				RHICmdList,
+				ShaderRHI,
+				IrradianceTexture,
+				IrradianceSampler,
+				TStaticSamplerState<SF_Bilinear,AM_Clamp,AM_Clamp,AM_Clamp>::GetRHI(),
+				DistanceFieldIrradiance2->GetRenderTargetItem().ShaderResourceTexture
+				);
+		}
+
+		const FIntPoint DownsampledBufferSize(GSceneRenderTargets.GetBufferSizeXY() / FIntPoint(GAODownsampleFactor, GAODownsampleFactor));
+		const FVector2D BaseLevelTexelSizeValue(1.0f / DownsampledBufferSize.X, 1.0f / DownsampledBufferSize.Y);
+		SetShaderValue(RHICmdList, ShaderRHI, BentNormalAOTexelSize, BaseLevelTexelSizeValue);
+
+		extern int32 GAOMinLevel;
+		extern int32 GAOPowerOfTwoBetweenLevels;
+		const float MinDownsampleFactor = 1 << (GAOMinLevel * GAOPowerOfTwoBetweenLevels); 
+		SetShaderValue(RHICmdList, ShaderRHI, MinDownsampleFactorToBaseLevel, MinDownsampleFactor);
+	}
+	// FShader interface.
+	virtual bool Serialize(FArchive& Ar) override
+	{
+		bool bShaderHasOutdatedParameters = FGlobalShader::Serialize(Ar);
+		Ar << BentNormalAOTexture;
+		Ar << BentNormalAOSampler;
+		Ar << IrradianceTexture;
+		Ar << IrradianceSampler;
+		Ar << BentNormalAOTexelSize;
+		Ar << MinDownsampleFactorToBaseLevel;
+		return bShaderHasOutdatedParameters;
+	}
+
+private:
+
+	FShaderResourceParameter BentNormalAOTexture;
+	FShaderResourceParameter BentNormalAOSampler;
+	FShaderResourceParameter IrradianceTexture;
+	FShaderResourceParameter IrradianceSampler;
+	FShaderParameter BentNormalAOTexelSize;
+	FShaderParameter MinDownsampleFactorToBaseLevel;
+};
+
+#define IMPLEMENT_FILLGAPS_PS_TYPE(bSupportIrradiance,bHighQuality) \
+	typedef TFillGapsPS<bSupportIrradiance, bHighQuality> TFillGapsPS##bSupportIrradiance##bHighQuality; \
+	IMPLEMENT_SHADER_TYPE(template<>,TFillGapsPS##bSupportIrradiance##bHighQuality,TEXT("DistanceFieldLightingPost"),TEXT("FillGapsPS"),SF_Pixel); 
+
+IMPLEMENT_FILLGAPS_PS_TYPE(true, true);
+IMPLEMENT_FILLGAPS_PS_TYPE(true, false);
+IMPLEMENT_FILLGAPS_PS_TYPE(false, true);
+IMPLEMENT_FILLGAPS_PS_TYPE(false, false);
+
+template<bool bSupportIrradiance>
+class TUpdateHistoryPS : public FGlobalShader
+{
+	DECLARE_SHADER_TYPE(TUpdateHistoryPS, Global);
+public:
+
+	static bool ShouldCache(EShaderPlatform Platform)
+	{
+		return IsFeatureLevelSupported(Platform, ERHIFeatureLevel::SM5) && DoesPlatformSupportDistanceFieldAO(Platform);
+	}
+
+	static void ModifyCompilationEnvironment(EShaderPlatform Platform, FShaderCompilerEnvironment& OutEnvironment)
+	{
+		OutEnvironment.SetDefine(TEXT("DOWNSAMPLE_FACTOR"), GAODownsampleFactor);
+		OutEnvironment.SetDefine(TEXT("SUPPORT_IRRADIANCE"), bSupportIrradiance ? TEXT("1") : TEXT("0"));
+	}
+
+	/** Default constructor. */
+	TUpdateHistoryPS() {}
+
+	/** Initialization constructor. */
+	TUpdateHistoryPS(const ShaderMetaType::CompiledShaderInitializerType& Initializer)
+		: FGlobalShader(Initializer)
+	{
+		DeferredParameters.Bind(Initializer.ParameterMap);
+		BentNormalAOTexture.Bind(Initializer.ParameterMap, TEXT("BentNormalAOTexture"));
+		BentNormalAOSampler.Bind(Initializer.ParameterMap, TEXT("BentNormalAOSampler"));
+		BentNormalHistoryTexture.Bind(Initializer.ParameterMap, TEXT("BentNormalHistoryTexture"));
+		BentNormalHistorySampler.Bind(Initializer.ParameterMap, TEXT("BentNormalHistorySampler"));
+		IrradianceTexture.Bind(Initializer.ParameterMap, TEXT("IrradianceTexture"));
+		IrradianceSampler.Bind(Initializer.ParameterMap, TEXT("IrradianceSampler"));
+		IrradianceHistoryTexture.Bind(Initializer.ParameterMap, TEXT("IrradianceHistoryTexture"));
+		IrradianceHistorySampler.Bind(Initializer.ParameterMap, TEXT("IrradianceHistorySampler"));
+		HistoryWeight.Bind(Initializer.ParameterMap, TEXT("HistoryWeight"));
+		HistoryDistanceThreshold.Bind(Initializer.ParameterMap, TEXT("HistoryDistanceThreshold"));
+		UseHistoryFilter.Bind(Initializer.ParameterMap, TEXT("UseHistoryFilter"));
+		VelocityTexture.Bind(Initializer.ParameterMap, TEXT("VelocityTexture"));
+		VelocityTextureSampler.Bind(Initializer.ParameterMap, TEXT("VelocityTextureSampler"));
+	}
+
+	void SetParameters(
+		FRHICommandList& RHICmdList, 
+		const FSceneView& View, 
+		FSceneRenderTargetItem& BentNormalHistoryTextureValue, 
+		TRefCountPtr<IPooledRenderTarget>* IrradianceHistoryRT, 
+		FSceneRenderTargetItem& DistanceFieldAOBentNormal, 
+		IPooledRenderTarget* DistanceFieldIrradiance,
+		FSceneRenderTargetItem& VelocityTextureValue)
+	{
+		const FPixelShaderRHIParamRef ShaderRHI = GetPixelShader();
+
+		FGlobalShader::SetParameters(RHICmdList, ShaderRHI, View);
+		DeferredParameters.Set(RHICmdList, ShaderRHI, View);
+
+		SetTextureParameter(
+			RHICmdList,
+			ShaderRHI,
+			BentNormalAOTexture,
+			BentNormalAOSampler,
+			TStaticSamplerState<SF_Bilinear,AM_Clamp,AM_Clamp,AM_Clamp>::GetRHI(),
+			DistanceFieldAOBentNormal.ShaderResourceTexture
+			);
+
+		SetTextureParameter(
+			RHICmdList,
+			ShaderRHI,
+			BentNormalHistoryTexture,
+			BentNormalHistorySampler,
+			TStaticSamplerState<SF_Bilinear,AM_Clamp,AM_Clamp,AM_Clamp>::GetRHI(),
+			BentNormalHistoryTextureValue.ShaderResourceTexture
+			);
+
+		if (IrradianceTexture.IsBound())
+		{
+			SetTextureParameter(
+				RHICmdList,
+				ShaderRHI,
+				IrradianceTexture,
+				IrradianceSampler,
+				TStaticSamplerState<SF_Bilinear,AM_Clamp,AM_Clamp,AM_Clamp>::GetRHI(),
+				DistanceFieldIrradiance->GetRenderTargetItem().ShaderResourceTexture
+				);
+		}
+		
+		if (IrradianceHistoryTexture.IsBound())
+		{
+			SetTextureParameter(
+				RHICmdList,
+				ShaderRHI,
+				IrradianceHistoryTexture,
+				IrradianceHistorySampler,
+				TStaticSamplerState<SF_Bilinear,AM_Clamp,AM_Clamp,AM_Clamp>::GetRHI(),
+				(*IrradianceHistoryRT)->GetRenderTargetItem().ShaderResourceTexture
+				);
+		}
+
+		SetUniformBufferParameter(RHICmdList, ShaderRHI, GetUniformBufferParameter<FCameraMotionParameters>(), CreateCameraMotionParametersUniformBuffer(View));
+
+		SetShaderValue(RHICmdList, ShaderRHI, HistoryWeight, GAOHistoryWeight);
+		SetShaderValue(RHICmdList, ShaderRHI, HistoryDistanceThreshold, GAOHistoryDistanceThreshold);
+		SetShaderValue(RHICmdList, ShaderRHI, UseHistoryFilter, (GAOHistoryStabilityPass ? 1.0f : 0.0f));
+
+		SetTextureParameter(
+			RHICmdList,
+			ShaderRHI,
+			VelocityTexture,
+			VelocityTextureSampler,
+			TStaticSamplerState<SF_Bilinear,AM_Clamp,AM_Clamp,AM_Clamp>::GetRHI(),
+			VelocityTextureValue.ShaderResourceTexture
+			);
+	}
+	// FShader interface.
+	virtual bool Serialize(FArchive& Ar)
+	{
+		bool bShaderHasOutdatedParameters = FGlobalShader::Serialize(Ar);
+		Ar << DeferredParameters;
+		Ar << BentNormalAOTexture;
+		Ar << BentNormalAOSampler;
+		Ar << BentNormalHistoryTexture;
+		Ar << BentNormalHistorySampler;
+		Ar << IrradianceTexture;
+		Ar << IrradianceSampler;
+		Ar << IrradianceHistoryTexture;
+		Ar << IrradianceHistorySampler;
+		Ar << HistoryWeight;
+		Ar << HistoryDistanceThreshold;
+		Ar << UseHistoryFilter;
+		Ar << VelocityTexture;
+		Ar << VelocityTextureSampler;
+		return bShaderHasOutdatedParameters;
+	}
+
+private:
+
+	FDeferredPixelShaderParameters DeferredParameters;
+	FShaderResourceParameter BentNormalAOTexture;
+	FShaderResourceParameter BentNormalAOSampler;
+	FShaderResourceParameter BentNormalHistoryTexture;
+	FShaderResourceParameter BentNormalHistorySampler;
+	FShaderResourceParameter IrradianceTexture;
+	FShaderResourceParameter IrradianceSampler;
+	FShaderResourceParameter IrradianceHistoryTexture;
+	FShaderResourceParameter IrradianceHistorySampler;
+	FShaderParameter HistoryWeight;
+	FShaderParameter HistoryDistanceThreshold;
+	FShaderParameter UseHistoryFilter;
+	FShaderResourceParameter VelocityTexture;
+	FShaderResourceParameter VelocityTextureSampler;
+};
+
+IMPLEMENT_SHADER_TYPE(template<>,TUpdateHistoryPS<true>,TEXT("DistanceFieldLightingPost"),TEXT("UpdateHistoryPS"),SF_Pixel);
+IMPLEMENT_SHADER_TYPE(template<>,TUpdateHistoryPS<false>,TEXT("DistanceFieldLightingPost"),TEXT("UpdateHistoryPS"),SF_Pixel);
+
+
+template<bool bSupportIrradiance>
+class TFilterHistoryPS : public FGlobalShader
+{
+	DECLARE_SHADER_TYPE(TFilterHistoryPS, Global);
+public:
+
+	static bool ShouldCache(EShaderPlatform Platform)
+	{
+		return IsFeatureLevelSupported(Platform, ERHIFeatureLevel::SM5) && DoesPlatformSupportDistanceFieldAO(Platform);
+	}
+
+	static void ModifyCompilationEnvironment(EShaderPlatform Platform, FShaderCompilerEnvironment& OutEnvironment)
+	{
+		OutEnvironment.SetDefine(TEXT("DOWNSAMPLE_FACTOR"), GAODownsampleFactor);
+		OutEnvironment.SetDefine(TEXT("SUPPORT_IRRADIANCE"), bSupportIrradiance ? TEXT("1") : TEXT("0"));
+	}
+
+	/** Default constructor. */
+	TFilterHistoryPS() {}
+
+	/** Initialization constructor. */
+	TFilterHistoryPS(const ShaderMetaType::CompiledShaderInitializerType& Initializer)
+		: FGlobalShader(Initializer)
+	{
+		BentNormalAOTexture.Bind(Initializer.ParameterMap, TEXT("BentNormalAOTexture"));
+		BentNormalAOSampler.Bind(Initializer.ParameterMap, TEXT("BentNormalAOSampler"));
+		IrradianceTexture.Bind(Initializer.ParameterMap, TEXT("IrradianceTexture"));
+		IrradianceSampler.Bind(Initializer.ParameterMap, TEXT("IrradianceSampler"));
+		HistoryWeight.Bind(Initializer.ParameterMap, TEXT("HistoryWeight"));
+		BentNormalAOTexelSize.Bind(Initializer.ParameterMap, TEXT("BentNormalAOTexelSize"));
+	}
+
+	void SetParameters(
+		FRHICommandList& RHICmdList, 
+		const FSceneView& View, 
+		FSceneRenderTargetItem& BentNormalHistoryTextureValue, 
+		IPooledRenderTarget* IrradianceHistoryRT)
+	{
+		const FPixelShaderRHIParamRef ShaderRHI = GetPixelShader();
+
+		FGlobalShader::SetParameters(RHICmdList, ShaderRHI, View);
+
+		SetTextureParameter(
+			RHICmdList,
+			ShaderRHI,
+			BentNormalAOTexture,
+			BentNormalAOSampler,
+			TStaticSamplerState<SF_Bilinear,AM_Clamp,AM_Clamp,AM_Clamp>::GetRHI(),
+			BentNormalHistoryTextureValue.ShaderResourceTexture
+			);
+
+		if (IrradianceTexture.IsBound())
+		{
+			SetTextureParameter(
+				RHICmdList,
+				ShaderRHI,
+				IrradianceTexture,
+				IrradianceSampler,
+				TStaticSamplerState<SF_Bilinear,AM_Clamp,AM_Clamp,AM_Clamp>::GetRHI(),
+				IrradianceHistoryRT->GetRenderTargetItem().ShaderResourceTexture
+				);
+		}
+
+		SetShaderValue(RHICmdList, ShaderRHI, HistoryWeight, GAOHistoryWeight);
+		
+		const FIntPoint DownsampledBufferSize(GSceneRenderTargets.GetBufferSizeXY() / FIntPoint(GAODownsampleFactor, GAODownsampleFactor));
+		const FVector2D BaseLevelTexelSizeValue(1.0f / DownsampledBufferSize.X, 1.0f / DownsampledBufferSize.Y);
+		SetShaderValue(RHICmdList, ShaderRHI, BentNormalAOTexelSize, BaseLevelTexelSizeValue);
+	}
+
+	// FShader interface.
+	virtual bool Serialize(FArchive& Ar)
+	{
+		bool bShaderHasOutdatedParameters = FGlobalShader::Serialize(Ar);
+		Ar << BentNormalAOTexture;
+		Ar << BentNormalAOSampler;
+		Ar << IrradianceTexture;
+		Ar << IrradianceSampler;
+		Ar << HistoryWeight;
+		Ar << BentNormalAOTexelSize;
+		return bShaderHasOutdatedParameters;
+	}
+
+private:
+
+	FShaderResourceParameter BentNormalAOTexture;
+	FShaderResourceParameter BentNormalAOSampler;
+	FShaderResourceParameter IrradianceTexture;
+	FShaderResourceParameter IrradianceSampler;
+	FShaderParameter HistoryWeight;
+	FShaderParameter BentNormalAOTexelSize;
+};
+
+IMPLEMENT_SHADER_TYPE(template<>,TFilterHistoryPS<true>,TEXT("DistanceFieldLightingPost"),TEXT("FilterHistoryPS"),SF_Pixel);
+IMPLEMENT_SHADER_TYPE(template<>,TFilterHistoryPS<false>,TEXT("DistanceFieldLightingPost"),TEXT("FilterHistoryPS"),SF_Pixel);
+
+
+void AllocateOrReuseAORenderTarget(TRefCountPtr<IPooledRenderTarget>& Target, const TCHAR* Name, bool bWithAlpha)
+{
+	if (!Target)
+	{
+		FIntPoint BufferSize = GetBufferSizeForAO();
+
+		FPooledRenderTargetDesc Desc(FPooledRenderTargetDesc::Create2DDesc(BufferSize, bWithAlpha ? PF_FloatRGBA : PF_FloatRGB, TexCreate_None, TexCreate_RenderTargetable | TexCreate_UAV, false));
+		GRenderTargetPool.FindFreeElement(Desc, Target, Name);
+	}
+}
+
+void UpdateHistory(
+	FRHICommandList& RHICmdList,
+	const FViewInfo& View, 
+	const TCHAR* BentNormalHistoryRTName,
+	const TCHAR* IrradianceHistoryRTName,
+	FSceneRenderTargetItem& VelocityTexture,
+	/** Contains last frame's history, if non-NULL.  This will be updated with the new frame's history. */
+	TRefCountPtr<IPooledRenderTarget>* BentNormalHistoryState,
+	TRefCountPtr<IPooledRenderTarget>* IrradianceHistoryState,
+	/** Source */
+	TRefCountPtr<IPooledRenderTarget>& BentNormalSource, 
+	TRefCountPtr<IPooledRenderTarget>& IrradianceSource, 
+	/** Output of Temporal Reprojection for the next step in the pipeline. */
+	TRefCountPtr<IPooledRenderTarget>& BentNormalHistoryOutput,
+	TRefCountPtr<IPooledRenderTarget>& IrradianceHistoryOutput)
+{
+	if (BentNormalHistoryState)
+	{
+		const bool bUseDistanceFieldGI = IsDistanceFieldGIAllowed(View);
+
+		if (*BentNormalHistoryState 
+			&& !View.bCameraCut 
+			&& !View.bPrevTransformsReset 
+			&& (!bUseDistanceFieldGI || (IrradianceHistoryState && *IrradianceHistoryState)))
+		{
+			// Reuse a render target from the pool with a consistent name, for vis purposes
+			TRefCountPtr<IPooledRenderTarget> NewBentNormalHistory;
+			AllocateOrReuseAORenderTarget(NewBentNormalHistory, BentNormalHistoryRTName, true);
+
+			TRefCountPtr<IPooledRenderTarget> NewIrradianceHistory;
+
+			if (bUseDistanceFieldGI)
+			{
+				AllocateOrReuseAORenderTarget(NewIrradianceHistory, IrradianceHistoryRTName, false);
+			}
+
+			SCOPED_DRAW_EVENT(RHICmdList, UpdateHistory);
+
+			{
+				FTextureRHIParamRef RenderTargets[2] =
+				{
+					NewBentNormalHistory->GetRenderTargetItem().TargetableTexture,
+					bUseDistanceFieldGI ? NewIrradianceHistory->GetRenderTargetItem().TargetableTexture : NULL,
+				};
+
+				SetRenderTargets(RHICmdList, ARRAY_COUNT(RenderTargets) - (bUseDistanceFieldGI ? 0 : 1), RenderTargets, FTextureRHIParamRef(), 0, NULL);
+
+				RHICmdList.SetViewport(0, 0, 0.0f, View.ViewRect.Width() / GAODownsampleFactor, View.ViewRect.Height() / GAODownsampleFactor, 1.0f);
+				RHICmdList.SetRasterizerState(TStaticRasterizerState<FM_Solid, CM_None>::GetRHI());
+				RHICmdList.SetDepthStencilState(TStaticDepthStencilState<false, CF_Always>::GetRHI());
+				RHICmdList.SetBlendState(TStaticBlendState<>::GetRHI());
+
+				TShaderMapRef<FPostProcessVS> VertexShader(View.ShaderMap);
+
+				if (bUseDistanceFieldGI)
+				{
+					TShaderMapRef<TUpdateHistoryPS<true> > PixelShader(View.ShaderMap);
+
+					static FGlobalBoundShaderState BoundShaderState;
+					SetGlobalBoundShaderState(RHICmdList, View.GetFeatureLevel(), BoundShaderState, GFilterVertexDeclaration.VertexDeclarationRHI, *VertexShader, *PixelShader);
+					PixelShader->SetParameters(RHICmdList, View, (*BentNormalHistoryState)->GetRenderTargetItem(), IrradianceHistoryState, BentNormalSource->GetRenderTargetItem(), IrradianceSource, VelocityTexture);
+				}
+				else
+				{
+					TShaderMapRef<TUpdateHistoryPS<false> > PixelShader(View.ShaderMap);
+
+					static FGlobalBoundShaderState BoundShaderState;
+					SetGlobalBoundShaderState(RHICmdList, View.GetFeatureLevel(), BoundShaderState, GFilterVertexDeclaration.VertexDeclarationRHI, *VertexShader, *PixelShader);
+					PixelShader->SetParameters(RHICmdList, View, (*BentNormalHistoryState)->GetRenderTargetItem(), IrradianceHistoryState, BentNormalSource->GetRenderTargetItem(), IrradianceSource, VelocityTexture);
+				}
+
+				VertexShader->SetParameters(RHICmdList, View);
+
+				DrawRectangle( 
+					RHICmdList,
+					0, 0, 
+					View.ViewRect.Width() / GAODownsampleFactor, View.ViewRect.Height() / GAODownsampleFactor,
+					View.ViewRect.Min.X / GAODownsampleFactor, View.ViewRect.Min.Y / GAODownsampleFactor,
+					View.ViewRect.Width() / GAODownsampleFactor, View.ViewRect.Height() / GAODownsampleFactor,
+					FIntPoint(View.ViewRect.Width() / GAODownsampleFactor, View.ViewRect.Height() / GAODownsampleFactor),
+					GSceneRenderTargets.GetBufferSizeXY() / FIntPoint(GAODownsampleFactor, GAODownsampleFactor),
+					*VertexShader);
+
+				RHICmdList.CopyToResolveTarget(NewBentNormalHistory->GetRenderTargetItem().TargetableTexture, NewBentNormalHistory->GetRenderTargetItem().ShaderResourceTexture, false, FResolveParams());
+			
+				if (bUseDistanceFieldGI)
+				{
+					RHICmdList.CopyToResolveTarget(NewIrradianceHistory->GetRenderTargetItem().TargetableTexture, NewIrradianceHistory->GetRenderTargetItem().ShaderResourceTexture, false, FResolveParams());
+				}
+			}
+
+			if (GAOHistoryStabilityPass)
+			{
+				const FPooledRenderTargetDesc& HistoryDesc = (*BentNormalHistoryState)->GetDesc();
+
+				// Reallocate history if buffer sizes have changed
+				if (HistoryDesc.Extent != GSceneRenderTargets.GetBufferSizeXY() / FIntPoint(GAODownsampleFactor, GAODownsampleFactor))
+				{
+					GRenderTargetPool.FreeUnusedResource(*BentNormalHistoryState);
+					*BentNormalHistoryState = NULL;
+					// Update the view state's render target reference with the new history
+					AllocateOrReuseAORenderTarget(*BentNormalHistoryState, BentNormalHistoryRTName, true);
+
+					if (bUseDistanceFieldGI)
+					{
+						GRenderTargetPool.FreeUnusedResource(*IrradianceHistoryState);
+						*IrradianceHistoryState = NULL;
+						AllocateOrReuseAORenderTarget(*IrradianceHistoryState, IrradianceHistoryRTName, false);
+					}
+				}
+
+				{
+					FTextureRHIParamRef RenderTargets[2] =
+					{
+						(*BentNormalHistoryState)->GetRenderTargetItem().TargetableTexture,
+						bUseDistanceFieldGI ? (*IrradianceHistoryState)->GetRenderTargetItem().TargetableTexture : NULL,
+					};
+
+					SetRenderTargets(RHICmdList, ARRAY_COUNT(RenderTargets) - (bUseDistanceFieldGI ? 0 : 1), RenderTargets, FTextureRHIParamRef(), 0, NULL);
+
+					RHICmdList.SetViewport(0, 0, 0.0f, View.ViewRect.Width() / GAODownsampleFactor, View.ViewRect.Height() / GAODownsampleFactor, 1.0f);
+					RHICmdList.SetRasterizerState(TStaticRasterizerState<FM_Solid, CM_None>::GetRHI());
+					RHICmdList.SetDepthStencilState(TStaticDepthStencilState<false, CF_Always>::GetRHI());
+					RHICmdList.SetBlendState(TStaticBlendState<>::GetRHI());
+
+					TShaderMapRef<FPostProcessVS> VertexShader(View.ShaderMap);
+
+					if (bUseDistanceFieldGI)
+					{
+						TShaderMapRef<TFilterHistoryPS<true> > PixelShader(View.ShaderMap);
+
+						static FGlobalBoundShaderState BoundShaderState;
+						SetGlobalBoundShaderState(RHICmdList, View.GetFeatureLevel(), BoundShaderState, GFilterVertexDeclaration.VertexDeclarationRHI, *VertexShader, *PixelShader);
+						PixelShader->SetParameters(RHICmdList, View, NewBentNormalHistory->GetRenderTargetItem(), NewIrradianceHistory);
+					}
+					else
+					{
+						TShaderMapRef<TFilterHistoryPS<false> > PixelShader(View.ShaderMap);
+
+						static FGlobalBoundShaderState BoundShaderState;
+						SetGlobalBoundShaderState(RHICmdList, View.GetFeatureLevel(), BoundShaderState, GFilterVertexDeclaration.VertexDeclarationRHI, *VertexShader, *PixelShader);
+						PixelShader->SetParameters(RHICmdList, View, NewBentNormalHistory->GetRenderTargetItem(), NewIrradianceHistory);
+					}
+
+					VertexShader->SetParameters(RHICmdList, View);
+
+					DrawRectangle( 
+						RHICmdList,
+						0, 0, 
+						View.ViewRect.Width() / GAODownsampleFactor, View.ViewRect.Height() / GAODownsampleFactor,
+						0, 0,
+						View.ViewRect.Width() / GAODownsampleFactor, View.ViewRect.Height() / GAODownsampleFactor,
+						FIntPoint(View.ViewRect.Width() / GAODownsampleFactor, View.ViewRect.Height() / GAODownsampleFactor),
+						GSceneRenderTargets.GetBufferSizeXY() / FIntPoint(GAODownsampleFactor, GAODownsampleFactor),
+						*VertexShader);
+
+					RHICmdList.CopyToResolveTarget((*BentNormalHistoryState)->GetRenderTargetItem().TargetableTexture, (*BentNormalHistoryState)->GetRenderTargetItem().ShaderResourceTexture, false, FResolveParams());
+			
+					if (bUseDistanceFieldGI)
+					{
+						RHICmdList.CopyToResolveTarget((*IrradianceHistoryState)->GetRenderTargetItem().TargetableTexture, (*IrradianceHistoryState)->GetRenderTargetItem().ShaderResourceTexture, false, FResolveParams());
+					}
+				}
+
+				BentNormalHistoryOutput = *BentNormalHistoryState;
+				IrradianceHistoryOutput = *IrradianceHistoryState;
+			}
+			else
+			{
+				// Update the view state's render target reference with the new history
+				*BentNormalHistoryState = NewBentNormalHistory;
+				BentNormalHistoryOutput = NewBentNormalHistory;
+
+				*IrradianceHistoryState = NewIrradianceHistory;
+				IrradianceHistoryOutput = NewIrradianceHistory;
+			}
+		}
+		else
+		{
+			// Use the current frame's mask for next frame's history
+			*BentNormalHistoryState = BentNormalSource;
+			BentNormalHistoryOutput = BentNormalSource;
+			BentNormalSource = NULL;
+
+			*IrradianceHistoryState = IrradianceSource;
+			IrradianceHistoryOutput = IrradianceSource;
+			IrradianceSource = NULL;
+		}
+	}
+	else
+	{
+		// Temporal reprojection is disabled or there is no view state - pass through
+		BentNormalHistoryOutput = BentNormalSource;
+		IrradianceHistoryOutput = IrradianceSource;
+	}
+}
+
+void PostProcessBentNormalAOSurfaceCache(
+	FRHICommandList& RHICmdList, 
+	const FDistanceFieldAOParameters& Parameters, 
+	const FViewInfo& View, 
+	FSceneRenderTargetItem& VelocityTexture,
+	FSceneRenderTargetItem& BentNormalInterpolation, 
+	IPooledRenderTarget* IrradianceInterpolation,
+	FSceneRenderTargetItem& DistanceFieldNormal,
+	TRefCountPtr<IPooledRenderTarget>& BentNormalOutput,
+	TRefCountPtr<IPooledRenderTarget>& IrradianceOutput)
+{
+	const bool bUseDistanceFieldGI = IsDistanceFieldGIAllowed(View);
+
+	TRefCountPtr<IPooledRenderTarget> DistanceFieldAOBentNormal;
+	TRefCountPtr<IPooledRenderTarget> DistanceFieldIrradiance;
+	AllocateOrReuseAORenderTarget(DistanceFieldAOBentNormal, TEXT("DistanceFieldBentNormalAO"), true);
+
+	if (bUseDistanceFieldGI)
+	{
+		AllocateOrReuseAORenderTarget(DistanceFieldIrradiance, TEXT("DistanceFieldIrradiance"), false);
+	}
+
+	TRefCountPtr<IPooledRenderTarget> DistanceFieldAOBentNormal2;
+	TRefCountPtr<IPooledRenderTarget> DistanceFieldIrradiance2;
+
+	if (GAOFillGaps)
+	{
+		AllocateOrReuseAORenderTarget(DistanceFieldAOBentNormal2, TEXT("DistanceFieldBentNormalAO2"), true);
+
+		if (bUseDistanceFieldGI)
+		{
+			AllocateOrReuseAORenderTarget(DistanceFieldIrradiance2, TEXT("DistanceFieldIrradiance2"), false);
+		}
+	}
+
+	{
+		SCOPED_DRAW_EVENT(RHICmdList, AOCombine);
+
+		TRefCountPtr<IPooledRenderTarget> EffectiveDiffuseIrradianceTarget;
+
+		if (bUseDistanceFieldGI)
+		{
+			EffectiveDiffuseIrradianceTarget = GAOFillGaps ? DistanceFieldIrradiance2 : DistanceFieldIrradiance;
+		}
+
+		FSceneRenderTargetItem& EffectiveBentNormalTarget = GAOFillGaps ? DistanceFieldAOBentNormal2->GetRenderTargetItem() : DistanceFieldAOBentNormal->GetRenderTargetItem();
+
+		FTextureRHIParamRef RenderTargets[2] =
+		{
+			EffectiveBentNormalTarget.TargetableTexture,
+			EffectiveDiffuseIrradianceTarget ? EffectiveDiffuseIrradianceTarget->GetRenderTargetItem().TargetableTexture : NULL,
+		};
+
+		SetRenderTargets(RHICmdList, ARRAY_COUNT(RenderTargets) - (bUseDistanceFieldGI ? 0 : 1), RenderTargets, FTextureRHIParamRef(), 0, NULL);
+
+		{
+			RHICmdList.SetViewport(0, 0, 0.0f, View.ViewRect.Width() / GAODownsampleFactor, View.ViewRect.Height() / GAODownsampleFactor, 1.0f);
+			RHICmdList.SetRasterizerState(TStaticRasterizerState<FM_Solid, CM_None>::GetRHI());
+			RHICmdList.SetDepthStencilState(TStaticDepthStencilState<false, CF_Always>::GetRHI());
+			RHICmdList.SetBlendState(TStaticBlendState<>::GetRHI());
+
+			TShaderMapRef<FPostProcessVS> VertexShader(View.ShaderMap);
+
+			if (bUseDistanceFieldGI)
+			{
+				TShaderMapRef<TDistanceFieldAOCombinePS<true> > PixelShader(View.ShaderMap);
+				static FGlobalBoundShaderState BoundShaderState;
+				SetGlobalBoundShaderState(RHICmdList, View.GetFeatureLevel(), BoundShaderState, GFilterVertexDeclaration.VertexDeclarationRHI, *VertexShader, *PixelShader);
+				PixelShader->SetParameters(RHICmdList, View, BentNormalInterpolation, IrradianceInterpolation, DistanceFieldNormal, Parameters);
+			}
+			else
+			{
+				TShaderMapRef<TDistanceFieldAOCombinePS<false> > PixelShader(View.ShaderMap);
+				static FGlobalBoundShaderState BoundShaderState;
+				SetGlobalBoundShaderState(RHICmdList, View.GetFeatureLevel(), BoundShaderState, GFilterVertexDeclaration.VertexDeclarationRHI, *VertexShader, *PixelShader);
+				PixelShader->SetParameters(RHICmdList, View, BentNormalInterpolation, IrradianceInterpolation, DistanceFieldNormal, Parameters);
+			}
+
+			DrawRectangle( 
+				RHICmdList,
+				0, 0, 
+				View.ViewRect.Width() / GAODownsampleFactor, View.ViewRect.Height() / GAODownsampleFactor,
+				0, 0, 
+				View.ViewRect.Width() / GAODownsampleFactor, View.ViewRect.Height() / GAODownsampleFactor,
+				FIntPoint(View.ViewRect.Width() / GAODownsampleFactor, View.ViewRect.Height() / GAODownsampleFactor),
+				GetBufferSizeForAO(),
+				*VertexShader);
+		}
+
+		RHICmdList.CopyToResolveTarget(EffectiveBentNormalTarget.TargetableTexture, EffectiveBentNormalTarget.ShaderResourceTexture, false, FResolveParams());
+			
+		if (bUseDistanceFieldGI)
+		{
+			RHICmdList.CopyToResolveTarget(EffectiveDiffuseIrradianceTarget->GetRenderTargetItem().TargetableTexture, EffectiveDiffuseIrradianceTarget->GetRenderTargetItem().ShaderResourceTexture, false, FResolveParams());
+		}
+	}
+
+	if (GAOFillGaps)
+	{
+		SCOPED_DRAW_EVENT(RHICmdList, FillGaps);
+
+		FTextureRHIParamRef RenderTargets[2] =
+		{
+			DistanceFieldAOBentNormal->GetRenderTargetItem().TargetableTexture,
+			bUseDistanceFieldGI ? DistanceFieldIrradiance->GetRenderTargetItem().TargetableTexture : NULL,
+		};
+
+		SetRenderTargets(RHICmdList, ARRAY_COUNT(RenderTargets) - (bUseDistanceFieldGI ? 0 : 1), RenderTargets, FTextureRHIParamRef(), 0, NULL);
+
+		{
+			RHICmdList.SetViewport(0, 0, 0.0f, View.ViewRect.Width() / GAODownsampleFactor, View.ViewRect.Height() / GAODownsampleFactor, 1.0f);
+			RHICmdList.SetRasterizerState(TStaticRasterizerState<FM_Solid, CM_None>::GetRHI());
+			RHICmdList.SetDepthStencilState(TStaticDepthStencilState<false, CF_Always>::GetRHI());
+			RHICmdList.SetBlendState(TStaticBlendState<>::GetRHI());
+
+			TShaderMapRef<FPostProcessVS> VertexShader(View.ShaderMap);
+
+			if (GAOFillGapsHighQuality)
+			{
+				if (bUseDistanceFieldGI)
+				{
+					TShaderMapRef<TFillGapsPS<true, true> > PixelShader(View.ShaderMap);
+
+					static FGlobalBoundShaderState BoundShaderState;
+					SetGlobalBoundShaderState(RHICmdList, View.GetFeatureLevel(), BoundShaderState, GFilterVertexDeclaration.VertexDeclarationRHI, *VertexShader, *PixelShader);
+					PixelShader->SetParameters(RHICmdList, View, DistanceFieldAOBentNormal2->GetRenderTargetItem(), DistanceFieldIrradiance2);
+				}
+				else
+				{
+					TShaderMapRef<TFillGapsPS<false, true> > PixelShader(View.ShaderMap);
+
+					static FGlobalBoundShaderState BoundShaderState;
+					SetGlobalBoundShaderState(RHICmdList, View.GetFeatureLevel(), BoundShaderState, GFilterVertexDeclaration.VertexDeclarationRHI, *VertexShader, *PixelShader);
+					PixelShader->SetParameters(RHICmdList, View, DistanceFieldAOBentNormal2->GetRenderTargetItem(), DistanceFieldIrradiance2);
+				}
+			}
+			else
+			{
+				if (bUseDistanceFieldGI)
+				{
+					TShaderMapRef<TFillGapsPS<true, false> > PixelShader(View.ShaderMap);
+
+					static FGlobalBoundShaderState BoundShaderState;
+					SetGlobalBoundShaderState(RHICmdList, View.GetFeatureLevel(), BoundShaderState, GFilterVertexDeclaration.VertexDeclarationRHI, *VertexShader, *PixelShader);
+					PixelShader->SetParameters(RHICmdList, View, DistanceFieldAOBentNormal2->GetRenderTargetItem(), DistanceFieldIrradiance2);
+				}
+				else
+				{
+					TShaderMapRef<TFillGapsPS<false, false> > PixelShader(View.ShaderMap);
+
+					static FGlobalBoundShaderState BoundShaderState;
+					SetGlobalBoundShaderState(RHICmdList, View.GetFeatureLevel(), BoundShaderState, GFilterVertexDeclaration.VertexDeclarationRHI, *VertexShader, *PixelShader);
+					PixelShader->SetParameters(RHICmdList, View, DistanceFieldAOBentNormal2->GetRenderTargetItem(), DistanceFieldIrradiance2);
+				}
+			}
+
+			VertexShader->SetParameters(RHICmdList, View);
+
+			DrawRectangle( 
+				RHICmdList,
+				0, 0, 
+				View.ViewRect.Width() / GAODownsampleFactor, View.ViewRect.Height() / GAODownsampleFactor,
+				0, 0, 
+				View.ViewRect.Width() / GAODownsampleFactor, View.ViewRect.Height() / GAODownsampleFactor,
+				FIntPoint(View.ViewRect.Width() / GAODownsampleFactor, View.ViewRect.Height() / GAODownsampleFactor),
+				GSceneRenderTargets.GetBufferSizeXY() / FIntPoint(GAODownsampleFactor, GAODownsampleFactor),
+				*VertexShader);
+		}
+
+		RHICmdList.CopyToResolveTarget(DistanceFieldAOBentNormal->GetRenderTargetItem().TargetableTexture, DistanceFieldAOBentNormal->GetRenderTargetItem().ShaderResourceTexture, false, FResolveParams());
+			
+		if (bUseDistanceFieldGI)
+		{
+			RHICmdList.CopyToResolveTarget(DistanceFieldIrradiance->GetRenderTargetItem().TargetableTexture, DistanceFieldIrradiance->GetRenderTargetItem().ShaderResourceTexture, false, FResolveParams());
+		}
+	}
+
+	FSceneViewState* ViewState = View.ViewState;
+	TRefCountPtr<IPooledRenderTarget>* BentNormalHistoryState = ViewState ? &ViewState->DistanceFieldAOHistoryRT : NULL;
+	TRefCountPtr<IPooledRenderTarget>* IrradianceHistoryState = ViewState ? &ViewState->DistanceFieldIrradianceHistoryRT : NULL;
+	BentNormalOutput = DistanceFieldAOBentNormal;
+	IrradianceOutput = DistanceFieldIrradiance;
+
+	if (GAOUseHistory)
+	{
+		UpdateHistory(
+			RHICmdList,
+			View, 
+			TEXT("DistanceFieldAOHistory"),
+			TEXT("DistanceFieldIrradianceHistory"),
+			VelocityTexture,
+			BentNormalHistoryState,
+			IrradianceHistoryState,
+			DistanceFieldAOBentNormal, 
+			DistanceFieldIrradiance,
+			BentNormalOutput,
+			IrradianceOutput);
+	}
+}
+
+enum EAOUpsampleType
+{
+	AOUpsample_OutputBentNormal,
+	AOUpsample_OutputAO,
+	AOUpsample_OutputBentNormalAndIrradiance,
+	AOUpsample_OutputIrradiance
+};
+
+template<EAOUpsampleType UpsampleType>
+class TDistanceFieldAOUpsamplePS : public FGlobalShader
+{
+	DECLARE_SHADER_TYPE(TDistanceFieldAOUpsamplePS, Global);
+public:
+
+	static bool ShouldCache(EShaderPlatform Platform)
+	{
+		return IsFeatureLevelSupported(Platform, ERHIFeatureLevel::SM5) && DoesPlatformSupportDistanceFieldAO(Platform);
+	}
+
+	static void ModifyCompilationEnvironment(EShaderPlatform Platform, FShaderCompilerEnvironment& OutEnvironment)
+	{
+		OutEnvironment.SetDefine(TEXT("DOWNSAMPLE_FACTOR"), GAODownsampleFactor);
+		OutEnvironment.SetDefine(TEXT("OUTPUT_BENT_NORMAL"), (UpsampleType == AOUpsample_OutputBentNormal || UpsampleType == AOUpsample_OutputBentNormalAndIrradiance) ? TEXT("1") : TEXT("0"));
+		OutEnvironment.SetDefine(TEXT("SUPPORT_IRRADIANCE"), (UpsampleType == AOUpsample_OutputIrradiance || UpsampleType == AOUpsample_OutputBentNormalAndIrradiance) ? TEXT("1") : TEXT("0"));
+		OutEnvironment.SetDefine(TEXT("OUTPUT_AO"), (UpsampleType == AOUpsample_OutputAO) ? TEXT("1") : TEXT("0"));
+	}
+
+	/** Default constructor. */
+	TDistanceFieldAOUpsamplePS() {}
+
+	/** Initialization constructor. */
+	TDistanceFieldAOUpsamplePS(const ShaderMetaType::CompiledShaderInitializerType& Initializer)
+		: FGlobalShader(Initializer)
+	{
+		DeferredParameters.Bind(Initializer.ParameterMap);
+		BentNormalAOTexture.Bind(Initializer.ParameterMap,TEXT("BentNormalAOTexture"));
+		BentNormalAOSampler.Bind(Initializer.ParameterMap,TEXT("BentNormalAOSampler"));
+		IrradianceTexture.Bind(Initializer.ParameterMap,TEXT("IrradianceTexture"));
+		IrradianceSampler.Bind(Initializer.ParameterMap,TEXT("IrradianceSampler"));
+	}
+
+	void SetParameters(FRHICommandList& RHICmdList, const FSceneView& View, TRefCountPtr<IPooledRenderTarget>& DistanceFieldAOBentNormal, IPooledRenderTarget* DistanceFieldIrradiance)
+	{
+		const FPixelShaderRHIParamRef ShaderRHI = GetPixelShader();
+
+		FGlobalShader::SetParameters(RHICmdList, ShaderRHI, View);
+		DeferredParameters.Set(RHICmdList, ShaderRHI, View);
+
+		SetTextureParameter(RHICmdList, ShaderRHI, BentNormalAOTexture, BentNormalAOSampler, TStaticSamplerState<SF_Bilinear>::GetRHI(), DistanceFieldAOBentNormal->GetRenderTargetItem().ShaderResourceTexture);
+
+		if (IrradianceTexture.IsBound())
+		{
+			SetTextureParameter(RHICmdList, ShaderRHI, IrradianceTexture, IrradianceSampler, TStaticSamplerState<SF_Bilinear>::GetRHI(), DistanceFieldIrradiance->GetRenderTargetItem().ShaderResourceTexture);
+		}
+	}
+	// FShader interface.
+	virtual bool Serialize(FArchive& Ar) override
+	{
+		bool bShaderHasOutdatedParameters = FGlobalShader::Serialize(Ar);
+		Ar << DeferredParameters;
+		Ar << BentNormalAOTexture;
+		Ar << BentNormalAOSampler;
+		Ar << IrradianceTexture;
+		Ar << IrradianceSampler;
+		return bShaderHasOutdatedParameters;
+	}
+
+private:
+
+	FDeferredPixelShaderParameters DeferredParameters;
+	FShaderResourceParameter BentNormalAOTexture;
+	FShaderResourceParameter BentNormalAOSampler;
+	FShaderResourceParameter IrradianceTexture;
+	FShaderResourceParameter IrradianceSampler;
+};
+
+#define IMPLEMENT_UPSAMPLE_PS_TYPE(UpsampleType) \
+	typedef TDistanceFieldAOUpsamplePS<UpsampleType> TDistanceFieldAOUpsamplePS##UpsampleType; \
+	IMPLEMENT_SHADER_TYPE(template<>,TDistanceFieldAOUpsamplePS##UpsampleType,TEXT("DistanceFieldLightingPost"),TEXT("AOUpsamplePS"),SF_Pixel);
+
+IMPLEMENT_UPSAMPLE_PS_TYPE(AOUpsample_OutputBentNormal)
+IMPLEMENT_UPSAMPLE_PS_TYPE(AOUpsample_OutputAO)
+IMPLEMENT_UPSAMPLE_PS_TYPE(AOUpsample_OutputBentNormalAndIrradiance)
+IMPLEMENT_UPSAMPLE_PS_TYPE(AOUpsample_OutputIrradiance)
+
+void UpsampleBentNormalAO(
+	FRHICommandList& RHICmdList, 
+	const TArray<FViewInfo>& Views, 
+	TRefCountPtr<IPooledRenderTarget>& DistanceFieldAOBentNormal, 
+	TRefCountPtr<IPooledRenderTarget>& DistanceFieldIrradiance,
+	bool bVisualizeAmbientOcclusion,
+	bool bVisualizeGlobalIllumination)
+{
+	for (int32 ViewIndex = 0; ViewIndex < Views.Num(); ViewIndex++)
+	{
+		const FViewInfo& View = Views[ViewIndex];
+		const bool bUseDistanceFieldGI = IsDistanceFieldGIAllowed(View);
+
+		SCOPED_DRAW_EVENT(RHICmdList, UpsampleAO);
+
+		RHICmdList.SetViewport(View.ViewRect.Min.X, View.ViewRect.Min.Y, 0.0f, View.ViewRect.Max.X, View.ViewRect.Max.Y, 1.0f);
+		RHICmdList.SetRasterizerState(TStaticRasterizerState<FM_Solid, CM_None>::GetRHI());
+		RHICmdList.SetDepthStencilState(TStaticDepthStencilState<false, CF_Always>::GetRHI());
+		RHICmdList.SetBlendState(TStaticBlendState<>::GetRHI());
+
+		TShaderMapRef<FPostProcessVS> VertexShader(View.ShaderMap);
+
+		if (bVisualizeAmbientOcclusion)
+		{
+			TShaderMapRef<TDistanceFieldAOUpsamplePS<AOUpsample_OutputAO> > PixelShader(View.ShaderMap);
+			static FGlobalBoundShaderState BoundShaderState;
+			SetGlobalBoundShaderState(RHICmdList, View.GetFeatureLevel(), BoundShaderState, GFilterVertexDeclaration.VertexDeclarationRHI, *VertexShader, *PixelShader);
+			PixelShader->SetParameters(RHICmdList, View, DistanceFieldAOBentNormal, DistanceFieldIrradiance);
+		}
+		else if (bVisualizeGlobalIllumination && bUseDistanceFieldGI)
+		{
+			TShaderMapRef<TDistanceFieldAOUpsamplePS<AOUpsample_OutputIrradiance> > PixelShader(View.ShaderMap);
+			static FGlobalBoundShaderState BoundShaderState;
+			SetGlobalBoundShaderState(RHICmdList, View.GetFeatureLevel(), BoundShaderState, GFilterVertexDeclaration.VertexDeclarationRHI, *VertexShader, *PixelShader);
+			PixelShader->SetParameters(RHICmdList, View, DistanceFieldAOBentNormal, DistanceFieldIrradiance);
+		}
+		else
+		{
+			if (bUseDistanceFieldGI)
+			{
+				TShaderMapRef<TDistanceFieldAOUpsamplePS<AOUpsample_OutputBentNormalAndIrradiance> > PixelShader(View.ShaderMap);
+				static FGlobalBoundShaderState BoundShaderState;
+				SetGlobalBoundShaderState(RHICmdList, View.GetFeatureLevel(), BoundShaderState, GFilterVertexDeclaration.VertexDeclarationRHI, *VertexShader, *PixelShader);
+				PixelShader->SetParameters(RHICmdList, View, DistanceFieldAOBentNormal, DistanceFieldIrradiance);
+			}
+			else
+			{
+				TShaderMapRef<TDistanceFieldAOUpsamplePS<AOUpsample_OutputBentNormal> > PixelShader(View.ShaderMap);
+				static FGlobalBoundShaderState BoundShaderState;
+				SetGlobalBoundShaderState(RHICmdList, View.GetFeatureLevel(), BoundShaderState, GFilterVertexDeclaration.VertexDeclarationRHI, *VertexShader, *PixelShader);
+				PixelShader->SetParameters(RHICmdList, View, DistanceFieldAOBentNormal, DistanceFieldIrradiance);
+			}
+		}
+
+		DrawRectangle( 
+			RHICmdList,
+			0, 0, 
+			View.ViewRect.Width(), View.ViewRect.Height(),
+			View.ViewRect.Min.X / GAODownsampleFactor, View.ViewRect.Min.Y / GAODownsampleFactor, 
+			View.ViewRect.Width() / GAODownsampleFactor, View.ViewRect.Height() / GAODownsampleFactor,
+			FIntPoint(View.ViewRect.Width(), View.ViewRect.Height()),
+			GetBufferSizeForAO(),
+			*VertexShader);
+	}
+}
diff --git a/Engine/Source/Runtime/Renderer/Private/DistanceFieldLightingPost.h b/Engine/Source/Runtime/Renderer/Private/DistanceFieldLightingPost.h
new file mode 100644
--- /dev/null
+++ b/Engine/Source/Runtime/Renderer/Private/DistanceFieldLightingPost.h
@@ -0,0 +1,44 @@
+// Copyright 1998-2015 Epic Games, Inc. All Rights Reserved.
+
+/*=============================================================================
+	DistanceFieldLightingPost.h
+=============================================================================*/
+
+#pragma once
+
+extern void AllocateOrReuseAORenderTarget(TRefCountPtr<IPooledRenderTarget>& Target, const TCHAR* Name, bool bWithAlpha);
+
+extern void UpdateHistory(
+	FRHICommandList& RHICmdList,
+	const FViewInfo& View, 
+	const TCHAR* BentNormalHistoryRTName,
+	const TCHAR* IrradianceHistoryRTName,
+	FSceneRenderTargetItem& VelocityTexture,
+	/** Contains last frame's history, if non-NULL.  This will be updated with the new frame's history. */
+	TRefCountPtr<IPooledRenderTarget>* BentNormalHistoryState,
+	TRefCountPtr<IPooledRenderTarget>* IrradianceHistoryState,
+	/** Source */
+	TRefCountPtr<IPooledRenderTarget>& BentNormalSource, 
+	TRefCountPtr<IPooledRenderTarget>& IrradianceSource, 
+	/** Output of Temporal Reprojection for the next step in the pipeline. */
+	TRefCountPtr<IPooledRenderTarget>& BentNormalHistoryOutput,
+	TRefCountPtr<IPooledRenderTarget>& IrradianceHistoryOutput);
+
+extern void PostProcessBentNormalAOSurfaceCache(
+	FRHICommandList& RHICmdList, 
+	const FDistanceFieldAOParameters& Parameters, 
+	const FViewInfo& View, 
+	FSceneRenderTargetItem& VelocityTexture,
+	FSceneRenderTargetItem& BentNormalInterpolation, 
+	IPooledRenderTarget* IrradianceInterpolation,
+	FSceneRenderTargetItem& DistanceFieldNormal,
+	TRefCountPtr<IPooledRenderTarget>& BentNormalOutput,
+	TRefCountPtr<IPooledRenderTarget>& IrradianceOutput);
+
+extern void UpsampleBentNormalAO(
+	FRHICommandList& RHICmdList, 
+	const TArray<FViewInfo>& Views, 
+	TRefCountPtr<IPooledRenderTarget>& DistanceFieldAOBentNormal, 
+	TRefCountPtr<IPooledRenderTarget>& DistanceFieldIrradiance,
+	bool bVisualizeAmbientOcclusion,
+	bool bVisualizeGlobalIllumination);
\ No newline at end of file
diff --git a/Engine/Source/Runtime/Renderer/Private/DistanceFieldLightingShared.h b/Engine/Source/Runtime/Renderer/Private/DistanceFieldLightingShared.h
--- a/Engine/Source/Runtime/Renderer/Private/DistanceFieldLightingShared.h
+++ b/Engine/Source/Runtime/Renderer/Private/DistanceFieldLightingShared.h
@@ -17,7 +17,6 @@ const int32 GDistanceFieldAOTileSizeY = 16;
 const int32 GAODownsampleFactor = 2;
 /** Must match usf */
 static const int32 GMaxNumObjectsPerTile = 512;
-extern uint32 UpdateObjectsGroupSize;
 
 extern int32 GDistanceFieldGI;
 
@@ -33,6 +32,8 @@ inline bool SupportsDistanceFieldGI(ERHIFeatureLevel::Type FeatureLevel, EShader
 		&& DoesPlatformSupportDistanceFieldGI(ShaderPlatform);
 }
 
+extern bool IsDistanceFieldGIAllowed(const FViewInfo& View);
+
 extern FIntPoint GetBufferSizeForAO();
 
 class FDistanceFieldObjectBuffers
diff --git a/Engine/Source/Runtime/Renderer/Private/DistanceFieldObjectManagement.cpp b/Engine/Source/Runtime/Renderer/Private/DistanceFieldObjectManagement.cpp
--- a/Engine/Source/Runtime/Renderer/Private/DistanceFieldObjectManagement.cpp
+++ b/Engine/Source/Runtime/Renderer/Private/DistanceFieldObjectManagement.cpp
@@ -114,7 +114,7 @@ public:
 
 TGlobalResource<FDistanceFieldRemoveIndicesResource> GDistanceFieldRemoveIndices;
 
-uint32 UpdateObjectsGroupSize = 64;
+const uint32 UpdateObjectsGroupSize = 64;
 
 class FUploadObjectsToBufferCS : public FGlobalShader
 {
@@ -542,7 +542,14 @@ void UpdateGlobalDistanceFieldObjectRemoves(FRHICommandListImmediate& RHICmdList
 			const FPrimitiveSceneInfo* Primitive = DistanceFieldSceneData.PendingRemoveOperations[RemoveIndex].Primitive;
 			DistanceFieldSceneData.SurfelAllocations.RemovePrimitive(Primitive);
 			DistanceFieldSceneData.InstancedSurfelAllocations.RemovePrimitive(Primitive);
-			PendingRemoveOperations.Append(DistanceFieldSceneData.PendingRemoveOperations[RemoveIndex].DistanceFieldInstanceIndices);
+			const TArray<int32, TInlineAllocator<1>>& DistanceFieldInstanceIndices = DistanceFieldSceneData.PendingRemoveOperations[RemoveIndex].DistanceFieldInstanceIndices;
+			PendingRemoveOperations.Append(DistanceFieldInstanceIndices);
+
+			for (int32 RemoveInstanceIndex = 0; RemoveInstanceIndex < DistanceFieldInstanceIndices.Num(); RemoveInstanceIndex++)
+			{
+				const int32 InstanceIndex = DistanceFieldInstanceIndices[RemoveInstanceIndex];
+				DistanceFieldSceneData.PrimitiveModifiedBounds.Add(DistanceFieldSceneData.PrimitiveInstanceMapping[InstanceIndex].BoundingSphere);
+			}
 		}
 
 		check(DistanceFieldSceneData.NumObjectsInBuffer >= PendingRemoveOperations.Num());
@@ -761,13 +768,6 @@ void ProcessPrimitiveUpdate(
 			{
 				const uint32 UploadIndex = bIsAddOperation ? OriginalNumObjects + UploadObjectIndices.Num() : PrimitiveSceneInfo->DistanceFieldInstanceIndices[TransformIndex];
 
-				if (bIsAddOperation)
-				{
-					const int32 AddIndex = OriginalNumObjects + UploadObjectIndices.Num();
-					DistanceFieldSceneData.PrimitiveInstanceMapping.Add(FPrimitiveAndInstance(PrimitiveSceneInfo, TransformIndex));
-					PrimitiveSceneInfo->DistanceFieldInstanceIndices.Add(AddIndex);
-				}
-
 				UploadObjectIndices.Add(UploadIndex);
 
 				FMatrix LocalToWorld = ObjectLocalToWorldTransforms[TransformIndex];
@@ -837,6 +837,19 @@ void ProcessPrimitiveUpdate(
 				UploadObjectData.Add(FVector4(Allocation.Offset, Allocation.NumLOD0, Allocation.NumSurfels, InstancedAllocation.Offset + InstancedAllocation.NumSurfels * TransformIndex));
 
 				checkSlow(UploadObjectData.Num() % UploadObjectDataStride == 0);
+
+				if (bIsAddOperation)
+				{
+					const int32 AddIndex = UploadIndex;
+					DistanceFieldSceneData.PrimitiveInstanceMapping.Add(FPrimitiveAndInstance(ObjectBoundingSphere, PrimitiveSceneInfo, TransformIndex));
+					PrimitiveSceneInfo->DistanceFieldInstanceIndices.Add(AddIndex);
+				}
+				else 
+				{
+					DistanceFieldSceneData.PrimitiveInstanceMapping[PrimitiveSceneInfo->DistanceFieldInstanceIndices[TransformIndex]].BoundingSphere = ObjectBoundingSphere;
+				}
+
+				DistanceFieldSceneData.PrimitiveModifiedBounds.Add(ObjectBoundingSphere);
 			}
 		}
 		else
@@ -850,6 +863,8 @@ void FDeferredShadingSceneRenderer::UpdateGlobalDistanceFieldObjectBuffers(FRHIC
 {
 	FDistanceFieldSceneData& DistanceFieldSceneData = Scene->DistanceFieldSceneData;
 
+	DistanceFieldSceneData.PrimitiveModifiedBounds.Reset();
+
 	if (GDistanceFieldVolumeTextureAtlas.VolumeTextureRHI
 		&& (DistanceFieldSceneData.HasPendingOperations() || DistanceFieldSceneData.AtlasGeneration != GDistanceFieldVolumeTextureAtlas.GetGeneration()))
 	{
@@ -892,6 +907,7 @@ void FDeferredShadingSceneRenderer::UpdateGlobalDistanceFieldObjectBuffers(FRHIC
 
 		// Process removes before adds, as the adds will overwrite primitive allocation info in DistanceFieldSceneData.SurfelAllocations
 		UpdateGlobalDistanceFieldObjectRemoves(RHICmdList, Scene);
+
 		extern int32 GVPLMeshGlobalIllumination;
 		TArray<uint32> UploadObjectIndices;
 		TArray<FVector4> UploadObjectData;
@@ -914,7 +930,6 @@ void FDeferredShadingSceneRenderer::UpdateGlobalDistanceFieldObjectBuffers(FRHIC
 			int32 OriginalNumSurfels = DistanceFieldSceneData.SurfelAllocations.GetNumSurfelsInBuffer();
 			int32 OriginalNumInstancedSurfels = DistanceFieldSceneData.InstancedSurfelAllocations.GetNumSurfelsInBuffer();
 			
-
 			if (bPrepareForDistanceFieldGI)
 			{
 				for (int32 UploadPrimitiveIndex = 0; UploadPrimitiveIndex < DistanceFieldSceneData.PendingAddOperations.Num(); UploadPrimitiveIndex++)
