Distance Field AO is now evaluated on a screen grid and using a global distance field for far traces

diff --git a/Engine/Source/Runtime/Renderer/Private/DistanceFieldScreenGridLighting.cpp b/Engine/Source/Runtime/Renderer/Private/DistanceFieldScreenGridLighting.cpp
new file mode 100644
--- /dev/null
+++ b/Engine/Source/Runtime/Renderer/Private/DistanceFieldScreenGridLighting.cpp
@@ -0,0 +1,911 @@
+// Copyright 1998-2015 Epic Games, Inc. All Rights Reserved.
+
+/*=============================================================================
+	DistanceFieldScreenGridLighting.cpp
+=============================================================================*/
+
+#include "RendererPrivate.h"
+#include "ScenePrivate.h"
+#include "UniformBuffer.h"
+#include "ShaderParameters.h"
+#include "PostProcessing.h"
+#include "SceneFilterRendering.h"
+#include "DistanceFieldLightingShared.h"
+#include "DistanceFieldSurfaceCacheLighting.h"
+#include "DistanceFieldLightingPost.h"
+#include "DistanceFieldGlobalIllumination.h"
+#include "GlobalDistanceField.h"
+#include "PostProcessAmbientOcclusion.h"
+#include "RHICommandList.h"
+#include "SceneUtils.h"
+#include "OneColorShader.h"
+#include "BasePassRendering.h"
+#include "HeightfieldLighting.h"
+
+int32 GAOUseJitter = 1;
+FAutoConsoleVariableRef CVarAOUseJitter(
+	TEXT("r.AOUseJitter"),
+	GAOUseJitter,
+	TEXT("Whether to use 4x temporal supersampling with Screen Grid DFAO.  When jitter is disabled, a shorter history can be used but there will be more spatial aliasing."),
+	ECVF_Cheat | ECVF_RenderThreadSafe
+	);
+
+int32 GConeTraceDownsampleFactor = 4;
+
+FIntPoint GetBufferSizeForConeTracing()
+{
+	return FIntPoint::DivideAndRoundDown(GetBufferSizeForAO(), GConeTraceDownsampleFactor);
+}
+
+FVector2D JitterOffsets[4] = 
+{
+	FVector2D(1, 0),
+	FVector2D(3, 1),
+	FVector2D(2, 3),
+	FVector2D(0, 2)
+};
+
+extern int32 GAOUseHistory;
+
+FVector2D GetJitterOffset(int32 SampleIndex)
+{
+	if (GAOUseJitter && GAOUseHistory)
+	{
+		return JitterOffsets[SampleIndex];
+	}
+
+	return FVector2D(0, 0);
+}
+
+void FAOScreenGridResources::InitDynamicRHI()
+{
+	//@todo - 2d textures
+	ScreenGridConeVisibility.Initialize(sizeof(float), NumConeSampleDirections * ScreenGridDimensions.X * ScreenGridDimensions.Y, PF_R32_FLOAT, BUF_Static);
+
+	if (bAllocateResourceForGI)
+	{
+		ConeDepthVisibilityFunction.Initialize(sizeof(float), NumConeSampleDirections * NumVisibilitySteps * ScreenGridDimensions.X * ScreenGridDimensions.Y, PF_R32_FLOAT, BUF_Static);
+		//@todo - fp16
+		StepBentNormal.Initialize(sizeof(float) * 4, NumVisibilitySteps * ScreenGridDimensions.X * ScreenGridDimensions.Y, PF_A32B32G32R32F, BUF_Static);
+		SurfelIrradiance.Initialize(sizeof(FFloat16Color), ScreenGridDimensions.X * ScreenGridDimensions.Y, PF_FloatRGBA, BUF_Static);
+		HeightfieldIrradiance.Initialize(sizeof(FFloat16Color), ScreenGridDimensions.X * ScreenGridDimensions.Y, PF_FloatRGBA, BUF_Static);
+	}
+}
+
+template<bool bSupportIrradiance, bool bUseGlobalDistanceField>
+class TConeTraceScreenGridObjectOcclusionCS : public FGlobalShader
+{
+	DECLARE_SHADER_TYPE(TConeTraceScreenGridObjectOcclusionCS,Global)
+public:
+
+	static bool ShouldCache(EShaderPlatform Platform)
+	{
+		return IsFeatureLevelSupported(Platform, ERHIFeatureLevel::SM5) && DoesPlatformSupportDistanceFieldAO(Platform);
+	}
+
+	static void ModifyCompilationEnvironment(EShaderPlatform Platform, FShaderCompilerEnvironment& OutEnvironment)
+	{
+		FGlobalShader::ModifyCompilationEnvironment(Platform,OutEnvironment);
+		OutEnvironment.SetDefine(TEXT("SUPPORT_IRRADIANCE"), bSupportIrradiance ? TEXT("1") : TEXT("0"));
+		OutEnvironment.SetDefine(TEXT("CULLED_TILE_SIZEX"), GDistanceFieldAOTileSizeX);
+		OutEnvironment.SetDefine(TEXT("TRACE_DOWNSAMPLE_FACTOR"), GConeTraceDownsampleFactor);
+		OutEnvironment.SetDefine(TEXT("USE_GLOBAL_DISTANCE_FIELD"), bUseGlobalDistanceField ? TEXT("1") : TEXT("0"));
+
+		// To reduce shader compile time of compute shaders with shared memory, doesn't have an impact on generated code with current compiler (June 2010 DX SDK)
+		OutEnvironment.CompilerFlags.Add(CFLAG_StandardOptimization);
+	}
+
+	TConeTraceScreenGridObjectOcclusionCS(const ShaderMetaType::CompiledShaderInitializerType& Initializer)
+		: FGlobalShader(Initializer)
+	{
+		DeferredParameters.Bind(Initializer.ParameterMap);
+		ObjectParameters.Bind(Initializer.ParameterMap);
+		AOParameters.Bind(Initializer.ParameterMap);
+		ScreenGridParameters.Bind(Initializer.ParameterMap);
+		GlobalDistanceFieldParameters.Bind(Initializer.ParameterMap);
+		TileHeadData.Bind(Initializer.ParameterMap, TEXT("TileHeadData"));
+		TileHeadDataUnpacked.Bind(Initializer.ParameterMap, TEXT("TileHeadDataUnpacked"));
+		TileArrayData.Bind(Initializer.ParameterMap, TEXT("TileArrayData"));
+		TileConeDepthRanges.Bind(Initializer.ParameterMap, TEXT("TileConeDepthRanges"));
+		TileListGroupSize.Bind(Initializer.ParameterMap, TEXT("TileListGroupSize"));
+		TanConeHalfAngle.Bind(Initializer.ParameterMap, TEXT("TanConeHalfAngle"));
+		BentNormalNormalizeFactor.Bind(Initializer.ParameterMap, TEXT("BentNormalNormalizeFactor"));
+		ScreenGridConeVisibility.Bind(Initializer.ParameterMap, TEXT("ScreenGridConeVisibility"));
+		ConeDepthVisibilityFunction.Bind(Initializer.ParameterMap, TEXT("ConeDepthVisibilityFunction"));
+	}
+
+	TConeTraceScreenGridObjectOcclusionCS()
+	{
+	}
+
+	void SetParameters(
+		FRHICommandList& RHICmdList, 
+		const FViewInfo& View, 
+		FIntPoint TileListGroupSizeValue, 
+		FSceneRenderTargetItem& DistanceFieldNormal, 
+		const FDistanceFieldAOParameters& Parameters,
+		const FGlobalDistanceFieldInfo& GlobalDistanceFieldInfo)
+	{
+		FComputeShaderRHIParamRef ShaderRHI = GetComputeShader();
+		FGlobalShader::SetParameters(RHICmdList, ShaderRHI, View);
+		DeferredParameters.Set(RHICmdList, ShaderRHI, View);
+		ObjectParameters.Set(RHICmdList, ShaderRHI, GAOCulledObjectBuffers.Buffers);
+		AOParameters.Set(RHICmdList, ShaderRHI, Parameters);
+		ScreenGridParameters.Set(RHICmdList, ShaderRHI, View, DistanceFieldNormal);
+
+		if (bUseGlobalDistanceField)
+		{
+			GlobalDistanceFieldParameters.Set(RHICmdList, ShaderRHI, GlobalDistanceFieldInfo);
+		}
+
+		FAOSampleData2 AOSampleData;
+
+		TArray<FVector, TInlineAllocator<9> > SampleDirections;
+		GetSpacedVectors(SampleDirections);
+
+		for (int32 SampleIndex = 0; SampleIndex < NumConeSampleDirections; SampleIndex++)
+		{
+			AOSampleData.SampleDirections[SampleIndex] = FVector4(SampleDirections[SampleIndex]);
+		}
+
+		SetUniformBufferParameterImmediate(RHICmdList, ShaderRHI, GetUniformBufferParameter<FAOSampleData2>(), AOSampleData);
+
+		FTileIntersectionResources* TileIntersectionResources = View.ViewState->AOTileIntersectionResources;
+
+		SetSRVParameter(RHICmdList, ShaderRHI, TileHeadDataUnpacked, TileIntersectionResources->TileHeadDataUnpacked.SRV);
+		SetSRVParameter(RHICmdList, ShaderRHI, TileHeadData, TileIntersectionResources->TileHeadData.SRV);
+		SetSRVParameter(RHICmdList, ShaderRHI, TileArrayData, TileIntersectionResources->TileArrayData.SRV);
+		SetSRVParameter(RHICmdList, ShaderRHI, TileConeDepthRanges, TileIntersectionResources->TileConeDepthRanges.SRV);
+
+		SetShaderValue(RHICmdList, ShaderRHI, TileListGroupSize, TileListGroupSizeValue);
+
+		extern float GAOConeHalfAngle;
+		SetShaderValue(RHICmdList, ShaderRHI, TanConeHalfAngle, FMath::Tan(GAOConeHalfAngle));
+
+		FVector UnoccludedVector(0);
+
+		for (int32 SampleIndex = 0; SampleIndex < NumConeSampleDirections; SampleIndex++)
+		{
+			UnoccludedVector += SampleDirections[SampleIndex];
+		}
+
+		float BentNormalNormalizeFactorValue = 1.0f / (UnoccludedVector / NumConeSampleDirections).Size();
+		SetShaderValue(RHICmdList, ShaderRHI, BentNormalNormalizeFactor, BentNormalNormalizeFactorValue);
+
+		FAOScreenGridResources* ScreenGridResources = View.ViewState->AOScreenGridResources;
+
+		ScreenGridConeVisibility.SetBuffer(RHICmdList, ShaderRHI, ScreenGridResources->ScreenGridConeVisibility);
+
+		if (bSupportIrradiance)
+		{
+			ConeDepthVisibilityFunction.SetBuffer(RHICmdList, ShaderRHI, ScreenGridResources->ConeDepthVisibilityFunction);
+		}
+	}
+
+	void UnsetParameters(FRHICommandList& RHICmdList)
+	{
+		ScreenGridConeVisibility.UnsetUAV(RHICmdList, GetComputeShader());
+		ConeDepthVisibilityFunction.UnsetUAV(RHICmdList, GetComputeShader());
+	}
+
+	virtual bool Serialize(FArchive& Ar) override
+	{		
+		bool bShaderHasOutdatedParameters = FGlobalShader::Serialize(Ar);
+		Ar << DeferredParameters;
+		Ar << ObjectParameters;
+		Ar << AOParameters;
+		Ar << ScreenGridParameters;
+		Ar << GlobalDistanceFieldParameters;
+		Ar << TileHeadDataUnpacked;
+		Ar << TileHeadData;
+		Ar << TileArrayData;
+		Ar << TileConeDepthRanges;
+		Ar << TileListGroupSize;
+		Ar << TanConeHalfAngle;
+		Ar << BentNormalNormalizeFactor;
+		Ar << ScreenGridConeVisibility;
+		Ar << ConeDepthVisibilityFunction;
+		return bShaderHasOutdatedParameters;
+	}
+
+private:
+
+	FDeferredPixelShaderParameters DeferredParameters;
+	FDistanceFieldCulledObjectBufferParameters ObjectParameters;
+	FAOParameters AOParameters;
+	FScreenGridParameters ScreenGridParameters;
+	FGlobalDistanceFieldParameters GlobalDistanceFieldParameters;
+	FShaderResourceParameter TileHeadDataUnpacked;
+	FShaderResourceParameter TileHeadData;
+	FShaderResourceParameter TileArrayData;
+	FShaderResourceParameter TileConeDepthRanges;
+	FShaderParameter TileListGroupSize;
+	FShaderParameter TanConeHalfAngle;
+	FShaderParameter BentNormalNormalizeFactor;
+	FRWShaderParameter ScreenGridConeVisibility;
+	FRWShaderParameter ConeDepthVisibilityFunction;
+};
+
+#define IMPLEMENT_CONETRACE_CS_TYPE(bSupportIrradiance, bUseGlobalDistanceField) \
+	typedef TConeTraceScreenGridObjectOcclusionCS<bSupportIrradiance, bUseGlobalDistanceField> TConeTraceScreenGridObjectOcclusionCS##bSupportIrradiance##bUseGlobalDistanceField; \
+	IMPLEMENT_SHADER_TYPE(template<>,TConeTraceScreenGridObjectOcclusionCS##bSupportIrradiance##bUseGlobalDistanceField,TEXT("DistanceFieldScreenGridLighting"),TEXT("ConeTraceObjectOcclusionCS"),SF_Compute);
+
+IMPLEMENT_CONETRACE_CS_TYPE(true, true)
+IMPLEMENT_CONETRACE_CS_TYPE(false, true)
+IMPLEMENT_CONETRACE_CS_TYPE(true, false)
+IMPLEMENT_CONETRACE_CS_TYPE(false, false)
+
+const int32 GConeTraceGlobalDFTileSize = 8;
+
+template<bool bSupportIrradiance>
+class TConeTraceScreenGridGlobalOcclusionCS : public FGlobalShader
+{
+	DECLARE_SHADER_TYPE(TConeTraceScreenGridGlobalOcclusionCS,Global)
+public:
+
+	static bool ShouldCache(EShaderPlatform Platform)
+	{
+		return IsFeatureLevelSupported(Platform, ERHIFeatureLevel::SM5) && DoesPlatformSupportDistanceFieldAO(Platform);
+	}
+
+	static void ModifyCompilationEnvironment(EShaderPlatform Platform, FShaderCompilerEnvironment& OutEnvironment)
+	{
+		FGlobalShader::ModifyCompilationEnvironment(Platform,OutEnvironment);
+		OutEnvironment.SetDefine(TEXT("SUPPORT_IRRADIANCE"), bSupportIrradiance ? TEXT("1") : TEXT("0"));
+		OutEnvironment.SetDefine(TEXT("CONE_TRACE_GLOBAL_DISPATCH_SIZEX"), GConeTraceGlobalDFTileSize);
+		OutEnvironment.SetDefine(TEXT("OUTPUT_VISIBILITY_DIRECTLY"), TEXT("1"));
+		OutEnvironment.SetDefine(TEXT("TRACE_DOWNSAMPLE_FACTOR"), GConeTraceDownsampleFactor);
+		OutEnvironment.SetDefine(TEXT("USE_GLOBAL_DISTANCE_FIELD"), TEXT("1"));
+
+		// To reduce shader compile time of compute shaders with shared memory, doesn't have an impact on generated code with current compiler (June 2010 DX SDK)
+		OutEnvironment.CompilerFlags.Add(CFLAG_StandardOptimization);
+	}
+
+	TConeTraceScreenGridGlobalOcclusionCS(const ShaderMetaType::CompiledShaderInitializerType& Initializer)
+		: FGlobalShader(Initializer)
+	{
+		DeferredParameters.Bind(Initializer.ParameterMap);
+		ObjectParameters.Bind(Initializer.ParameterMap);
+		AOParameters.Bind(Initializer.ParameterMap);
+		ScreenGridParameters.Bind(Initializer.ParameterMap);
+		GlobalDistanceFieldParameters.Bind(Initializer.ParameterMap);
+		TileHeadData.Bind(Initializer.ParameterMap, TEXT("TileHeadData"));
+		TileHeadDataUnpacked.Bind(Initializer.ParameterMap, TEXT("TileHeadDataUnpacked"));
+		TileArrayData.Bind(Initializer.ParameterMap, TEXT("TileArrayData"));
+		TileConeDepthRanges.Bind(Initializer.ParameterMap, TEXT("TileConeDepthRanges"));
+		TileListGroupSize.Bind(Initializer.ParameterMap, TEXT("TileListGroupSize"));
+		TanConeHalfAngle.Bind(Initializer.ParameterMap, TEXT("TanConeHalfAngle"));
+		BentNormalNormalizeFactor.Bind(Initializer.ParameterMap, TEXT("BentNormalNormalizeFactor"));
+		ScreenGridConeVisibility.Bind(Initializer.ParameterMap, TEXT("ScreenGridConeVisibility"));
+	}
+
+	TConeTraceScreenGridGlobalOcclusionCS()
+	{
+	}
+
+	void SetParameters(
+		FRHICommandList& RHICmdList, 
+		const FViewInfo& View, 
+		FIntPoint TileListGroupSizeValue, 
+		FSceneRenderTargetItem& DistanceFieldNormal, 
+		const FDistanceFieldAOParameters& Parameters,
+		const FGlobalDistanceFieldInfo& GlobalDistanceFieldInfo)
+	{
+		FComputeShaderRHIParamRef ShaderRHI = GetComputeShader();
+		FGlobalShader::SetParameters(RHICmdList, ShaderRHI, View);
+		DeferredParameters.Set(RHICmdList, ShaderRHI, View);
+		ObjectParameters.Set(RHICmdList, ShaderRHI, GAOCulledObjectBuffers.Buffers);
+		AOParameters.Set(RHICmdList, ShaderRHI, Parameters);
+		ScreenGridParameters.Set(RHICmdList, ShaderRHI, View, DistanceFieldNormal);
+		GlobalDistanceFieldParameters.Set(RHICmdList, ShaderRHI, GlobalDistanceFieldInfo);
+
+		FAOSampleData2 AOSampleData;
+
+		TArray<FVector, TInlineAllocator<9> > SampleDirections;
+		GetSpacedVectors(SampleDirections);
+
+		for (int32 SampleIndex = 0; SampleIndex < NumConeSampleDirections; SampleIndex++)
+		{
+			AOSampleData.SampleDirections[SampleIndex] = FVector4(SampleDirections[SampleIndex]);
+		}
+
+		SetUniformBufferParameterImmediate(RHICmdList, ShaderRHI, GetUniformBufferParameter<FAOSampleData2>(), AOSampleData);
+
+		FTileIntersectionResources* TileIntersectionResources = View.ViewState->AOTileIntersectionResources;
+
+		SetSRVParameter(RHICmdList, ShaderRHI, TileHeadDataUnpacked, TileIntersectionResources->TileHeadDataUnpacked.SRV);
+		SetSRVParameter(RHICmdList, ShaderRHI, TileHeadData, TileIntersectionResources->TileHeadData.SRV);
+		SetSRVParameter(RHICmdList, ShaderRHI, TileArrayData, TileIntersectionResources->TileArrayData.SRV);
+		SetSRVParameter(RHICmdList, ShaderRHI, TileConeDepthRanges, TileIntersectionResources->TileConeDepthRanges.SRV);
+
+		SetShaderValue(RHICmdList, ShaderRHI, TileListGroupSize, TileListGroupSizeValue);
+
+		extern float GAOConeHalfAngle;
+		SetShaderValue(RHICmdList, ShaderRHI, TanConeHalfAngle, FMath::Tan(GAOConeHalfAngle));
+
+		FVector UnoccludedVector(0);
+
+		for (int32 SampleIndex = 0; SampleIndex < NumConeSampleDirections; SampleIndex++)
+		{
+			UnoccludedVector += SampleDirections[SampleIndex];
+		}
+
+		float BentNormalNormalizeFactorValue = 1.0f / (UnoccludedVector / NumConeSampleDirections).Size();
+		SetShaderValue(RHICmdList, ShaderRHI, BentNormalNormalizeFactor, BentNormalNormalizeFactorValue);
+
+		FAOScreenGridResources* ScreenGridResources = View.ViewState->AOScreenGridResources;
+
+		ScreenGridConeVisibility.SetBuffer(RHICmdList, ShaderRHI, ScreenGridResources->ScreenGridConeVisibility);
+	}
+
+	void UnsetParameters(FRHICommandList& RHICmdList)
+	{
+		ScreenGridConeVisibility.UnsetUAV(RHICmdList, GetComputeShader());
+	}
+
+	virtual bool Serialize(FArchive& Ar) override
+	{		
+		bool bShaderHasOutdatedParameters = FGlobalShader::Serialize(Ar);
+		Ar << DeferredParameters;
+		Ar << ObjectParameters;
+		Ar << AOParameters;
+		Ar << ScreenGridParameters;
+		Ar << GlobalDistanceFieldParameters;
+		Ar << TileHeadDataUnpacked;
+		Ar << TileHeadData;
+		Ar << TileArrayData;
+		Ar << TileConeDepthRanges;
+		Ar << TileListGroupSize;
+		Ar << TanConeHalfAngle;
+		Ar << BentNormalNormalizeFactor;
+		Ar << ScreenGridConeVisibility;
+		return bShaderHasOutdatedParameters;
+	}
+
+private:
+
+	FDeferredPixelShaderParameters DeferredParameters;
+	FDistanceFieldCulledObjectBufferParameters ObjectParameters;
+	FAOParameters AOParameters;
+	FScreenGridParameters ScreenGridParameters;
+	FGlobalDistanceFieldParameters GlobalDistanceFieldParameters;
+	FShaderResourceParameter TileHeadDataUnpacked;
+	FShaderResourceParameter TileHeadData;
+	FShaderResourceParameter TileArrayData;
+	FShaderResourceParameter TileConeDepthRanges;
+	FShaderParameter TileListGroupSize;
+	FShaderParameter TanConeHalfAngle;
+	FShaderParameter BentNormalNormalizeFactor;
+	FRWShaderParameter ScreenGridConeVisibility;
+};
+
+#define IMPLEMENT_CONETRACE_GLOBAL_CS_TYPE(bSupportIrradiance) \
+	typedef TConeTraceScreenGridGlobalOcclusionCS<bSupportIrradiance> TConeTraceScreenGridGlobalOcclusionCS##bSupportIrradiance; \
+	IMPLEMENT_SHADER_TYPE(template<>,TConeTraceScreenGridGlobalOcclusionCS##bSupportIrradiance,TEXT("DistanceFieldScreenGridLighting"),TEXT("ConeTraceGlobalOcclusionCS"),SF_Compute);
+
+IMPLEMENT_CONETRACE_GLOBAL_CS_TYPE(true)
+IMPLEMENT_CONETRACE_GLOBAL_CS_TYPE(false)
+
+
+
+const int32 GCombineConesSizeX = 8;
+
+class FCombineConeVisibilityCS : public FGlobalShader
+{
+	DECLARE_SHADER_TYPE(FCombineConeVisibilityCS,Global)
+public:
+
+	static bool ShouldCache(EShaderPlatform Platform)
+	{
+		return IsFeatureLevelSupported(Platform, ERHIFeatureLevel::SM5) && DoesPlatformSupportDistanceFieldAO(Platform);
+	}
+
+	static void ModifyCompilationEnvironment(EShaderPlatform Platform, FShaderCompilerEnvironment& OutEnvironment)
+	{
+		FGlobalShader::ModifyCompilationEnvironment(Platform,OutEnvironment);
+		OutEnvironment.SetDefine(TEXT("COMBINE_CONES_SIZEX"), GCombineConesSizeX);
+		OutEnvironment.SetDefine(TEXT("TRACE_DOWNSAMPLE_FACTOR"), GConeTraceDownsampleFactor);
+
+		// To reduce shader compile time of compute shaders with shared memory, doesn't have an impact on generated code with current compiler (June 2010 DX SDK)
+		OutEnvironment.CompilerFlags.Add(CFLAG_StandardOptimization);
+	}
+
+	FCombineConeVisibilityCS(const ShaderMetaType::CompiledShaderInitializerType& Initializer)
+		: FGlobalShader(Initializer)
+	{
+		ScreenGridParameters.Bind(Initializer.ParameterMap);
+		BentNormalNormalizeFactor.Bind(Initializer.ParameterMap, TEXT("BentNormalNormalizeFactor"));
+		ScreenGridConeVisibility.Bind(Initializer.ParameterMap, TEXT("ScreenGridConeVisibility"));
+		DistanceFieldBentNormal.Bind(Initializer.ParameterMap, TEXT("DistanceFieldBentNormal"));
+	}
+
+	FCombineConeVisibilityCS()
+	{
+	}
+
+	void SetParameters(
+		FRHICommandList& RHICmdList, 
+		const FViewInfo& View, 
+		FSceneRenderTargetItem& DistanceFieldNormal, 
+		FSceneRenderTargetItem& DownsampledBentNormal)
+	{
+		FComputeShaderRHIParamRef ShaderRHI = GetComputeShader();
+		FGlobalShader::SetParameters(RHICmdList, ShaderRHI, View);
+		ScreenGridParameters.Set(RHICmdList, ShaderRHI, View, DistanceFieldNormal);
+
+		FAOSampleData2 AOSampleData;
+
+		TArray<FVector, TInlineAllocator<9> > SampleDirections;
+		GetSpacedVectors(SampleDirections);
+
+		for (int32 SampleIndex = 0; SampleIndex < NumConeSampleDirections; SampleIndex++)
+		{
+			AOSampleData.SampleDirections[SampleIndex] = FVector4(SampleDirections[SampleIndex]);
+		}
+
+		SetUniformBufferParameterImmediate(RHICmdList, ShaderRHI, GetUniformBufferParameter<FAOSampleData2>(), AOSampleData);
+
+		FVector UnoccludedVector(0);
+
+		for (int32 SampleIndex = 0; SampleIndex < NumConeSampleDirections; SampleIndex++)
+		{
+			UnoccludedVector += SampleDirections[SampleIndex];
+		}
+
+		float BentNormalNormalizeFactorValue = 1.0f / (UnoccludedVector / NumConeSampleDirections).Size();
+		SetShaderValue(RHICmdList, ShaderRHI, BentNormalNormalizeFactor, BentNormalNormalizeFactorValue);
+
+		FAOScreenGridResources* ScreenGridResources = View.ViewState->AOScreenGridResources;
+
+		DistanceFieldBentNormal.SetTexture(RHICmdList, ShaderRHI, DownsampledBentNormal.ShaderResourceTexture, DownsampledBentNormal.UAV);
+
+		SetSRVParameter(RHICmdList, ShaderRHI, ScreenGridConeVisibility, ScreenGridResources->ScreenGridConeVisibility.SRV);
+	}
+
+	void UnsetParameters(FRHICommandList& RHICmdList)
+	{
+		DistanceFieldBentNormal.UnsetUAV(RHICmdList, GetComputeShader());
+	}
+
+	virtual bool Serialize(FArchive& Ar) override
+	{		
+		bool bShaderHasOutdatedParameters = FGlobalShader::Serialize(Ar);
+		Ar << ScreenGridParameters;
+		Ar << BentNormalNormalizeFactor;
+		Ar << ScreenGridConeVisibility;
+		Ar << DistanceFieldBentNormal;
+		return bShaderHasOutdatedParameters;
+	}
+
+private:
+
+	FScreenGridParameters ScreenGridParameters;
+	FShaderParameter BentNormalNormalizeFactor;
+	FShaderResourceParameter ScreenGridConeVisibility;
+	FRWShaderParameter DistanceFieldBentNormal;
+};
+
+IMPLEMENT_SHADER_TYPE(,FCombineConeVisibilityCS,TEXT("DistanceFieldScreenGridLighting"),TEXT("CombineConeVisibilityCS"),SF_Compute);
+
+template<bool bSupportIrradiance, bool bHighQuality>
+class TGeometryAwareUpsamplePS : public FGlobalShader
+{
+	DECLARE_SHADER_TYPE(TGeometryAwareUpsamplePS, Global);
+public:
+
+	static bool ShouldCache(EShaderPlatform Platform)
+	{
+		return IsFeatureLevelSupported(Platform, ERHIFeatureLevel::SM5) && DoesPlatformSupportDistanceFieldAO(Platform);
+	}
+
+	static void ModifyCompilationEnvironment(EShaderPlatform Platform, FShaderCompilerEnvironment& OutEnvironment)
+	{
+		OutEnvironment.SetDefine(TEXT("DOWNSAMPLE_FACTOR"), GAODownsampleFactor);
+		OutEnvironment.SetDefine(TEXT("SUPPORT_IRRADIANCE"), bSupportIrradiance ? TEXT("1") : TEXT("0"));
+		OutEnvironment.SetDefine(TEXT("HIGH_QUALITY_FILL_GAPS"), bHighQuality ? TEXT("1") : TEXT("0"));
+	}
+
+	/** Default constructor. */
+	TGeometryAwareUpsamplePS() {}
+
+	/** Initialization constructor. */
+	TGeometryAwareUpsamplePS(const ShaderMetaType::CompiledShaderInitializerType& Initializer)
+		: FGlobalShader(Initializer)
+	{
+		AOParameters.Bind(Initializer.ParameterMap);
+		DistanceFieldNormalTexture.Bind(Initializer.ParameterMap, TEXT("DistanceFieldNormalTexture"));
+		DistanceFieldNormalSampler.Bind(Initializer.ParameterMap, TEXT("DistanceFieldNormalSampler"));
+		BentNormalAOTexture.Bind(Initializer.ParameterMap, TEXT("BentNormalAOTexture"));
+		BentNormalAOSampler.Bind(Initializer.ParameterMap, TEXT("BentNormalAOSampler"));
+		IrradianceTexture.Bind(Initializer.ParameterMap, TEXT("IrradianceTexture"));
+		IrradianceSampler.Bind(Initializer.ParameterMap, TEXT("IrradianceSampler"));
+		DistanceFieldGBufferTexelSize.Bind(Initializer.ParameterMap, TEXT("DistanceFieldGBufferTexelSize"));
+		BentNormalBufferAndTexelSize.Bind(Initializer.ParameterMap, TEXT("BentNormalBufferAndTexelSize"));
+		MinDownsampleFactorToBaseLevel.Bind(Initializer.ParameterMap, TEXT("MinDownsampleFactorToBaseLevel"));
+		DistanceFadeScale.Bind(Initializer.ParameterMap, TEXT("DistanceFadeScale"));
+		JitterOffset.Bind(Initializer.ParameterMap, TEXT("JitterOffset"));
+	}
+
+	void SetParameters(
+		FRHICommandList& RHICmdList, 
+		const FViewInfo& View, 
+		FSceneRenderTargetItem& DistanceFieldNormal, 
+		FSceneRenderTargetItem& DistanceFieldAOBentNormal, 
+		IPooledRenderTarget* DistanceFieldIrradiance, 
+		const FDistanceFieldAOParameters& Parameters)
+	{
+		const FPixelShaderRHIParamRef ShaderRHI = GetPixelShader();
+
+		FGlobalShader::SetParameters(RHICmdList, ShaderRHI, View);
+
+		AOParameters.Set(RHICmdList, ShaderRHI, Parameters);
+
+		SetTextureParameter(
+			RHICmdList,
+			ShaderRHI,
+			DistanceFieldNormalTexture,
+			DistanceFieldNormalSampler,
+			TStaticSamplerState<SF_Bilinear,AM_Clamp,AM_Clamp,AM_Clamp>::GetRHI(),
+			DistanceFieldNormal.ShaderResourceTexture
+			);
+
+		SetTextureParameter(
+			RHICmdList,
+			ShaderRHI,
+			BentNormalAOTexture,
+			BentNormalAOSampler,
+			TStaticSamplerState<SF_Bilinear,AM_Clamp,AM_Clamp,AM_Clamp>::GetRHI(),
+			DistanceFieldAOBentNormal.ShaderResourceTexture
+			);
+
+		if (IrradianceTexture.IsBound())
+		{
+			SetTextureParameter(
+				RHICmdList,
+				ShaderRHI,
+				IrradianceTexture,
+				IrradianceSampler,
+				TStaticSamplerState<SF_Bilinear,AM_Clamp,AM_Clamp,AM_Clamp>::GetRHI(),
+				DistanceFieldIrradiance->GetRenderTargetItem().ShaderResourceTexture
+				);
+		}
+
+		const FIntPoint DownsampledBufferSize = GetBufferSizeForAO();
+		const FVector2D BaseLevelTexelSizeValue(1.0f / DownsampledBufferSize.X, 1.0f / DownsampledBufferSize.Y);
+		SetShaderValue(RHICmdList, ShaderRHI, DistanceFieldGBufferTexelSize, BaseLevelTexelSizeValue);
+
+		const FIntPoint ConeTracingBufferSize = GetBufferSizeForConeTracing();
+		const FVector4 BentNormalBufferAndTexelSizeValue(ConeTracingBufferSize.X, ConeTracingBufferSize.Y, 1.0f / ConeTracingBufferSize.X, 1.0f / ConeTracingBufferSize.Y);
+		SetShaderValue(RHICmdList, ShaderRHI, BentNormalBufferAndTexelSize, BentNormalBufferAndTexelSizeValue);
+
+		const float MinDownsampleFactor = GConeTraceDownsampleFactor; 
+		SetShaderValue(RHICmdList, ShaderRHI, MinDownsampleFactorToBaseLevel, MinDownsampleFactor);
+
+		extern float GAOViewFadeDistanceScale;
+		const float DistanceFadeScaleValue = 1.0f / ((1.0f - GAOViewFadeDistanceScale) * GetMaxAOViewDistance());
+		SetShaderValue(RHICmdList, ShaderRHI, DistanceFadeScale, DistanceFadeScaleValue);
+
+		SetShaderValue(RHICmdList, ShaderRHI, JitterOffset, GetJitterOffset(View.ViewState->GetDistanceFieldTemporalSampleIndex()));
+	}
+
+	// FShader interface.
+	virtual bool Serialize(FArchive& Ar) override
+	{
+		bool bShaderHasOutdatedParameters = FGlobalShader::Serialize(Ar);
+		Ar << AOParameters;
+		Ar << DistanceFieldNormalTexture;
+		Ar << DistanceFieldNormalSampler;
+		Ar << BentNormalAOTexture;
+		Ar << BentNormalAOSampler;
+		Ar << IrradianceTexture;
+		Ar << IrradianceSampler;
+		Ar << DistanceFieldGBufferTexelSize;
+		Ar << BentNormalBufferAndTexelSize;
+		Ar << MinDownsampleFactorToBaseLevel;
+		Ar << DistanceFadeScale;
+		Ar << JitterOffset;
+		return bShaderHasOutdatedParameters;
+	}
+
+private:
+
+	FAOParameters AOParameters;
+	FShaderResourceParameter DistanceFieldNormalTexture;
+	FShaderResourceParameter DistanceFieldNormalSampler;
+	FShaderResourceParameter BentNormalAOTexture;
+	FShaderResourceParameter BentNormalAOSampler;
+	FShaderResourceParameter IrradianceTexture;
+	FShaderResourceParameter IrradianceSampler;
+	FShaderParameter DistanceFieldGBufferTexelSize;
+	FShaderParameter BentNormalBufferAndTexelSize;
+	FShaderParameter MinDownsampleFactorToBaseLevel;
+	FShaderParameter DistanceFadeScale;
+	FShaderParameter JitterOffset;
+};
+
+#define IMPLEMENT_AWARE_UPSAMPLE_PS_TYPE(bSupportIrradiance,bHighQuality) \
+	typedef TGeometryAwareUpsamplePS<bSupportIrradiance, bHighQuality> TGeometryAwareUpsamplePS##bSupportIrradiance##bHighQuality; \
+	IMPLEMENT_SHADER_TYPE(template<>,TGeometryAwareUpsamplePS##bSupportIrradiance##bHighQuality,TEXT("DistanceFieldScreenGridLighting"),TEXT("GeometryAwareUpsamplePS"),SF_Pixel); 
+
+IMPLEMENT_AWARE_UPSAMPLE_PS_TYPE(true, true);
+IMPLEMENT_AWARE_UPSAMPLE_PS_TYPE(true, false);
+IMPLEMENT_AWARE_UPSAMPLE_PS_TYPE(false, true);
+IMPLEMENT_AWARE_UPSAMPLE_PS_TYPE(false, false);
+
+void PostProcessBentNormalAOScreenGrid(
+	FRHICommandList& RHICmdList, 
+	const FDistanceFieldAOParameters& Parameters, 
+	const FViewInfo& View, 
+	FSceneRenderTargetItem& VelocityTexture,
+	FSceneRenderTargetItem& BentNormalInterpolation, 
+	IPooledRenderTarget* IrradianceInterpolation,
+	FSceneRenderTargetItem& DistanceFieldNormal,
+	TRefCountPtr<IPooledRenderTarget>& BentNormalOutput,
+	TRefCountPtr<IPooledRenderTarget>& IrradianceOutput)
+{
+	const bool bUseDistanceFieldGI = IsDistanceFieldGIAllowed(View);
+
+	TRefCountPtr<IPooledRenderTarget> DistanceFieldAOBentNormal;
+	TRefCountPtr<IPooledRenderTarget> DistanceFieldIrradiance;
+	AllocateOrReuseAORenderTarget(DistanceFieldAOBentNormal, TEXT("DistanceFieldBentNormalAO"), true);
+
+	if (bUseDistanceFieldGI)
+	{
+		AllocateOrReuseAORenderTarget(DistanceFieldIrradiance, TEXT("DistanceFieldIrradiance"), false);
+	}
+
+	{
+		SCOPED_DRAW_EVENT(RHICmdList, GeometryAwareUpsample);
+
+		FTextureRHIParamRef RenderTargets[2] =
+		{
+			DistanceFieldAOBentNormal->GetRenderTargetItem().TargetableTexture,
+			bUseDistanceFieldGI ? DistanceFieldIrradiance->GetRenderTargetItem().TargetableTexture : NULL,
+		};
+
+		SetRenderTargets(RHICmdList, ARRAY_COUNT(RenderTargets) - (bUseDistanceFieldGI ? 0 : 1), RenderTargets, FTextureRHIParamRef(), 0, NULL);
+
+		{
+			RHICmdList.SetViewport(0, 0, 0.0f, View.ViewRect.Width() / GAODownsampleFactor, View.ViewRect.Height() / GAODownsampleFactor, 1.0f);
+			RHICmdList.SetRasterizerState(TStaticRasterizerState<FM_Solid, CM_None>::GetRHI());
+			RHICmdList.SetDepthStencilState(TStaticDepthStencilState<false, CF_Always>::GetRHI());
+			RHICmdList.SetBlendState(TStaticBlendState<>::GetRHI());
+
+			TShaderMapRef<FPostProcessVS> VertexShader(View.ShaderMap);
+
+			if (bUseDistanceFieldGI)
+			{
+				TShaderMapRef<TGeometryAwareUpsamplePS<true, false> > PixelShader(View.ShaderMap);
+
+				static FGlobalBoundShaderState BoundShaderState;
+				SetGlobalBoundShaderState(RHICmdList, View.GetFeatureLevel(), BoundShaderState, GFilterVertexDeclaration.VertexDeclarationRHI, *VertexShader, *PixelShader);
+				PixelShader->SetParameters(RHICmdList, View, DistanceFieldNormal, BentNormalInterpolation, IrradianceInterpolation, Parameters);
+			}
+			else
+			{
+				TShaderMapRef<TGeometryAwareUpsamplePS<false, false> > PixelShader(View.ShaderMap);
+
+				static FGlobalBoundShaderState BoundShaderState;
+				SetGlobalBoundShaderState(RHICmdList, View.GetFeatureLevel(), BoundShaderState, GFilterVertexDeclaration.VertexDeclarationRHI, *VertexShader, *PixelShader);
+				PixelShader->SetParameters(RHICmdList, View, DistanceFieldNormal, BentNormalInterpolation, IrradianceInterpolation, Parameters);
+			}
+
+			VertexShader->SetParameters(RHICmdList, View);
+
+			DrawRectangle( 
+				RHICmdList,
+				0, 0, 
+				View.ViewRect.Width() / GAODownsampleFactor, View.ViewRect.Height() / GAODownsampleFactor,
+				0, 0, 
+				View.ViewRect.Width() / GAODownsampleFactor, View.ViewRect.Height() / GAODownsampleFactor,
+				FIntPoint(View.ViewRect.Width() / GAODownsampleFactor, View.ViewRect.Height() / GAODownsampleFactor),
+				GSceneRenderTargets.GetBufferSizeXY() / FIntPoint(GAODownsampleFactor, GAODownsampleFactor),
+				*VertexShader);
+		}
+
+		RHICmdList.CopyToResolveTarget(DistanceFieldAOBentNormal->GetRenderTargetItem().TargetableTexture, DistanceFieldAOBentNormal->GetRenderTargetItem().ShaderResourceTexture, false, FResolveParams());
+			
+		if (bUseDistanceFieldGI)
+		{
+			RHICmdList.CopyToResolveTarget(DistanceFieldIrradiance->GetRenderTargetItem().TargetableTexture, DistanceFieldIrradiance->GetRenderTargetItem().ShaderResourceTexture, false, FResolveParams());
+		}
+	}
+
+	FSceneViewState* ViewState = (FSceneViewState*)View.State;
+	TRefCountPtr<IPooledRenderTarget>* BentNormalHistoryState = ViewState ? &ViewState->DistanceFieldAOHistoryRT : NULL;
+	TRefCountPtr<IPooledRenderTarget>* IrradianceHistoryState = ViewState ? &ViewState->DistanceFieldIrradianceHistoryRT : NULL;
+	BentNormalOutput = DistanceFieldAOBentNormal;
+	IrradianceOutput = DistanceFieldIrradiance;
+
+	if (GAOUseHistory)
+	{
+		UpdateHistory(
+			RHICmdList,
+			View, 
+			TEXT("DistanceFieldAOHistory"),
+			TEXT("DistanceFieldIrradianceHistory"),
+			VelocityTexture,
+			BentNormalHistoryState,
+			IrradianceHistoryState,
+			DistanceFieldAOBentNormal, 
+			DistanceFieldIrradiance,
+			BentNormalOutput,
+			IrradianceOutput);
+	}
+}
+
+void FDeferredShadingSceneRenderer::RenderDistanceFieldAOScreenGrid(
+	FRHICommandListImmediate& RHICmdList, 
+	const FViewInfo& View,
+	FIntPoint TileListGroupSize,
+	const FDistanceFieldAOParameters& Parameters, 
+	const TRefCountPtr<IPooledRenderTarget>& VelocityTexture,
+	const TRefCountPtr<IPooledRenderTarget>& DistanceFieldNormal, 
+	TRefCountPtr<IPooledRenderTarget>& OutDynamicBentNormalAO, 
+	TRefCountPtr<IPooledRenderTarget>& OutDynamicIrradiance)
+{
+	const bool bUseDistanceFieldGI = IsDistanceFieldGIAllowed(View);
+	const bool bUseGlobalDistanceField = UseGlobalDistanceField(Parameters) && Scene->DistanceFieldSceneData.NumObjectsInBuffer > 0;
+	const FIntPoint ConeTraceBufferSize = GetBufferSizeForConeTracing();
+
+	FAOScreenGridResources*& ScreenGridResources = View.ViewState->AOScreenGridResources;
+
+	if (!ScreenGridResources 
+		|| ScreenGridResources->ScreenGridDimensions != ConeTraceBufferSize 
+		|| ScreenGridResources->bAllocateResourceForGI != bUseDistanceFieldGI)
+	{
+		if (ScreenGridResources)
+		{
+			ScreenGridResources->ReleaseResource();
+		}
+		else
+		{
+			ScreenGridResources = new FAOScreenGridResources();
+		}
+
+		ScreenGridResources->bAllocateResourceForGI = bUseDistanceFieldGI;
+		ScreenGridResources->ScreenGridDimensions = ConeTraceBufferSize;
+
+		ScreenGridResources->InitResource();
+	}
+
+	SetRenderTarget(RHICmdList, NULL, NULL);
+
+	{
+		SCOPED_DRAW_EVENT(RHICmdList, ConeTraceObjects);
+
+		const uint32 GroupSizeX = FMath::DivideAndRoundUp(View.ViewRect.Size().X / GAODownsampleFactor, GDistanceFieldAOTileSizeX / GConeTraceDownsampleFactor);
+		const uint32 GroupSizeY = FMath::DivideAndRoundUp(View.ViewRect.Size().Y / GAODownsampleFactor, GDistanceFieldAOTileSizeY / GConeTraceDownsampleFactor);
+
+		if (bUseGlobalDistanceField)
+		{
+			check(View.GlobalDistanceFieldInfo.Clipmaps.Num() > 0);
+
+			if (bUseDistanceFieldGI)
+			{
+				TShaderMapRef<TConeTraceScreenGridObjectOcclusionCS<true, true> > ComputeShader(View.ShaderMap);
+
+				RHICmdList.SetComputeShader(ComputeShader->GetComputeShader());
+				ComputeShader->SetParameters(RHICmdList, View, TileListGroupSize, DistanceFieldNormal->GetRenderTargetItem(), Parameters, View.GlobalDistanceFieldInfo);
+				DispatchComputeShader(RHICmdList, *ComputeShader, GroupSizeX, GroupSizeY, 1);
+				ComputeShader->UnsetParameters(RHICmdList);
+			}
+			else
+			{
+				TShaderMapRef<TConeTraceScreenGridObjectOcclusionCS<false, true> > ComputeShader(View.ShaderMap);
+
+				RHICmdList.SetComputeShader(ComputeShader->GetComputeShader());
+				ComputeShader->SetParameters(RHICmdList, View, TileListGroupSize, DistanceFieldNormal->GetRenderTargetItem(), Parameters, View.GlobalDistanceFieldInfo);
+				DispatchComputeShader(RHICmdList, *ComputeShader, GroupSizeX, GroupSizeY, 1);
+				ComputeShader->UnsetParameters(RHICmdList);
+			}
+		}
+		else
+		{
+			if (bUseDistanceFieldGI)
+			{
+				TShaderMapRef<TConeTraceScreenGridObjectOcclusionCS<true, false> > ComputeShader(View.ShaderMap);
+
+				RHICmdList.SetComputeShader(ComputeShader->GetComputeShader());
+				ComputeShader->SetParameters(RHICmdList, View, TileListGroupSize, DistanceFieldNormal->GetRenderTargetItem(), Parameters, View.GlobalDistanceFieldInfo);
+				DispatchComputeShader(RHICmdList, *ComputeShader, GroupSizeX, GroupSizeY, 1);
+				ComputeShader->UnsetParameters(RHICmdList);
+			}
+			else
+			{
+				TShaderMapRef<TConeTraceScreenGridObjectOcclusionCS<false, false> > ComputeShader(View.ShaderMap);
+
+				RHICmdList.SetComputeShader(ComputeShader->GetComputeShader());
+				ComputeShader->SetParameters(RHICmdList, View, TileListGroupSize, DistanceFieldNormal->GetRenderTargetItem(), Parameters, View.GlobalDistanceFieldInfo);
+				DispatchComputeShader(RHICmdList, *ComputeShader, GroupSizeX, GroupSizeY, 1);
+				ComputeShader->UnsetParameters(RHICmdList);
+			}
+		}
+	}
+
+	if (bUseGlobalDistanceField)
+	{
+		SCOPED_DRAW_EVENT(RHICmdList, ConeTraceGlobal);
+
+		const uint32 GroupSizeX = FMath::DivideAndRoundUp(View.ViewRect.Size().X / GAODownsampleFactor, GConeTraceGlobalDFTileSize);
+		const uint32 GroupSizeY = FMath::DivideAndRoundUp(View.ViewRect.Size().Y / GAODownsampleFactor, GConeTraceGlobalDFTileSize);
+
+		check(View.GlobalDistanceFieldInfo.Clipmaps.Num() > 0);
+
+		if (bUseDistanceFieldGI)
+		{
+			TShaderMapRef<TConeTraceScreenGridGlobalOcclusionCS<true> > ComputeShader(View.ShaderMap);
+
+			RHICmdList.SetComputeShader(ComputeShader->GetComputeShader());
+			ComputeShader->SetParameters(RHICmdList, View, TileListGroupSize, DistanceFieldNormal->GetRenderTargetItem(), Parameters, View.GlobalDistanceFieldInfo);
+			DispatchComputeShader(RHICmdList, *ComputeShader, GroupSizeX, GroupSizeY, 1);
+			ComputeShader->UnsetParameters(RHICmdList);
+		}
+		else
+		{
+			TShaderMapRef<TConeTraceScreenGridGlobalOcclusionCS<false> > ComputeShader(View.ShaderMap);
+
+			RHICmdList.SetComputeShader(ComputeShader->GetComputeShader());
+			ComputeShader->SetParameters(RHICmdList, View, TileListGroupSize, DistanceFieldNormal->GetRenderTargetItem(), Parameters, View.GlobalDistanceFieldInfo);
+			DispatchComputeShader(RHICmdList, *ComputeShader, GroupSizeX, GroupSizeY, 1);
+			ComputeShader->UnsetParameters(RHICmdList);
+		}
+	}
+
+	TRefCountPtr<IPooledRenderTarget> DownsampledIrradiance;
+
+	if (bUseDistanceFieldGI)
+	{
+		{
+			FPooledRenderTargetDesc Desc(FPooledRenderTargetDesc::Create2DDesc(ConeTraceBufferSize, PF_FloatRGBA, TexCreate_None, TexCreate_RenderTargetable | TexCreate_UAV, false));
+			GRenderTargetPool.FindFreeElement(Desc, DownsampledIrradiance, TEXT("DownsampledIrradiance"));
+		}
+
+		extern void ComputeIrradianceForScreenGrid(
+			FRHICommandListImmediate& RHICmdList,
+			const FViewInfo& View,
+			const FScene* Scene,
+			const FDistanceFieldAOParameters& Parameters,
+			FSceneRenderTargetItem& DistanceFieldNormal, 
+			const FAOScreenGridResources& ScreenGridResources,
+			FSceneRenderTargetItem& IrradianceTexture);
+
+		ComputeIrradianceForScreenGrid(RHICmdList, View, Scene, Parameters, DistanceFieldNormal->GetRenderTargetItem(), *ScreenGridResources, DownsampledIrradiance->GetRenderTargetItem());
+	}
+
+	// Compute heightfield occlusion after heightfield GI, otherwise it self-shadows incorrectly
+	View.HeightfieldLightingViewInfo.ComputeOcclusionForScreenGrid(View, RHICmdList, DistanceFieldNormal->GetRenderTargetItem(), *ScreenGridResources, Parameters);
+
+	TRefCountPtr<IPooledRenderTarget> DownsampledBentNormal;
+
+	{
+		FPooledRenderTargetDesc Desc(FPooledRenderTargetDesc::Create2DDesc(ConeTraceBufferSize, PF_FloatRGBA, TexCreate_None, TexCreate_RenderTargetable | TexCreate_UAV, false));
+		GRenderTargetPool.FindFreeElement(Desc, DownsampledBentNormal, TEXT("DownsampledBentNormal"));
+	}
+
+	{
+		SCOPED_DRAW_EVENT(RHICmdList, CombineCones);
+		const uint32 GroupSizeX = FMath::DivideAndRoundUp(View.ViewRect.Size().X / GAODownsampleFactor, GCombineConesSizeX);
+		const uint32 GroupSizeY = FMath::DivideAndRoundUp(View.ViewRect.Size().Y / GAODownsampleFactor, GCombineConesSizeX);
+
+		TShaderMapRef<FCombineConeVisibilityCS> ComputeShader(View.ShaderMap);
+
+		RHICmdList.SetComputeShader(ComputeShader->GetComputeShader());
+		ComputeShader->SetParameters(RHICmdList, View, DistanceFieldNormal->GetRenderTargetItem(), DownsampledBentNormal->GetRenderTargetItem());
+		DispatchComputeShader(RHICmdList, *ComputeShader, GroupSizeX, GroupSizeY, 1);
+		ComputeShader->UnsetParameters(RHICmdList);
+	}
+
+	GRenderTargetPool.VisualizeTexture.SetCheckPoint(RHICmdList, DownsampledBentNormal);
+
+	PostProcessBentNormalAOScreenGrid(
+		RHICmdList, 
+		Parameters, 
+		View, 
+		VelocityTexture->GetRenderTargetItem(), 
+		DownsampledBentNormal->GetRenderTargetItem(), 
+		DownsampledIrradiance, 
+		DistanceFieldNormal->GetRenderTargetItem(), 
+		OutDynamicBentNormalAO, 
+		OutDynamicIrradiance);
+}
